U5 - Adding DebugName, MinNumBatchSize and LLM inheritance to ParallelFor
Also staring to remove some of the legacy TaskBackend Code to reduce the maintenence effort.

diff --git a/Engine/Source/Runtime/Core/Public/Async/ParallelFor.h b/Engine/Source/Runtime/Core/Public/Async/ParallelFor.h
--- a/Engine/Source/Runtime/Core/Public/Async/ParallelFor.h
+++ b/Engine/Source/Runtime/Core/Public/Async/ParallelFor.h
@@ -48,123 +48,6 @@ ENUM_CLASS_FLAGS(EParallelForFlags)
 
 namespace ParallelForImpl
 {
-	// struct to hold the working data; this outlives the ParallelFor call; lifetime is controlled by a shared pointer
-	template<typename FunctionType>
-	struct TParallelForData
-	{
-		int32 Num;
-		int32 BlockSize;
-		int32 LastBlockExtraNum;
-		FunctionType Body;
-		FEvent* Event;
-		FThreadSafeCounter IndexToDo;
-		FThreadSafeCounter NumCompleted;
-#if DO_CHECK
-		std::atomic<bool> bExited{ false };
-#endif
-		bool bTriggered;
-		bool bSaveLastBlockForMaster;
-		TParallelForData(int32 InTotalNum, int32 InNumThreads, bool bInSaveLastBlockForMaster, FunctionType InBody, EParallelForFlags Flags)
-			: Body(InBody)
-			, Event(FPlatformProcess::GetSynchEventFromPool(false))
-			, bTriggered(false)
-			, bSaveLastBlockForMaster(bInSaveLastBlockForMaster)
-		{
-			check(InTotalNum >= InNumThreads);
-
-			if ((Flags & EParallelForFlags::Unbalanced) != EParallelForFlags::None)
-			{
-				BlockSize = 1;
-				Num = InTotalNum;
-			}
-			else
-			{
-				BlockSize = 0;
-				Num = 0;
-				for (int32 Div = 6; Div; Div--)
-				{
-					BlockSize = InTotalNum / (InNumThreads * Div);
-					if (BlockSize)
-					{
-						Num = InTotalNum / BlockSize;
-						if (Num >= InNumThreads + !!bSaveLastBlockForMaster)
-						{
-							break;
-						}
-					}
-				}
-			}
-
-			check(BlockSize && Num);
-			LastBlockExtraNum = InTotalNum - Num * BlockSize;
-			check(LastBlockExtraNum >= 0);
-		}
-		~TParallelForData()
-		{
-			check(IndexToDo.GetValue() >= Num);
-			check(NumCompleted.GetValue() == Num);
-			check(bExited.load(std::memory_order_relaxed));
-			FPlatformProcess::ReturnSynchEventToPool(Event);
-		}
-
-		template<typename ContextType>
-		bool Process(const TArrayView<ContextType>& Contexts, int32 TaskIndex, int32 TasksToSpawn, TSharedRef<TParallelForData, ESPMode::ThreadSafe>& Data, ENamedThreads::Type InDesiredThread, bool bMaster);
-		
-		inline bool Process(int32 TaskIndex, int32 TasksToSpawn, TSharedRef<TParallelForData, ESPMode::ThreadSafe>& Data, ENamedThreads::Type InDesiredThread, bool bMaster)
-		{
-			return Process<nullptr_t>(TArrayView<nullptr_t>(), TaskIndex, TasksToSpawn, Data, InDesiredThread, bMaster);
-		}
-	};
-
-	template<typename FunctionType, typename ContextType = nullptr_t>
-	class TParallelForTask
-	{
-		TArrayView<ContextType> Contexts;
-		TSharedRef<TParallelForData<FunctionType>, ESPMode::ThreadSafe> Data;
-		ENamedThreads::Type DesiredThread;
-		int32 TaskIndex;
-		int32 TasksToSpawn;
-	public:
-		TParallelForTask(const TArrayView<ContextType>& InContexts, int32 InTaskIndex, TSharedRef<TParallelForData<FunctionType>, ESPMode::ThreadSafe>& InData, ENamedThreads::Type InDesiredThread, int32 InTasksToSpawn = 0)
-			: Contexts(InContexts)
-			, Data(InData) 
-			, DesiredThread(InDesiredThread)
-			, TaskIndex(InTaskIndex)
-			, TasksToSpawn(InTasksToSpawn)
-		{
-		}		
-		TParallelForTask(int32 InTaskIndex, TSharedRef<TParallelForData<FunctionType>, ESPMode::ThreadSafe>& InData, ENamedThreads::Type InDesiredThread, int32 InTasksToSpawn = 0)
-			: Data(InData) 
-			, DesiredThread(InDesiredThread)
-			, TaskIndex(InTaskIndex)
-			, TasksToSpawn(InTasksToSpawn)
-		{
-		}
-		static FORCEINLINE TStatId GetStatId()
-		{
-			return GET_STATID(STAT_ParallelForTask);
-		}
-		
-		FORCEINLINE ENamedThreads::Type GetDesiredThread()
-		{
-			return DesiredThread;
-		}
-
-		static FORCEINLINE ESubsequentsMode::Type GetSubsequentsMode() 
-		{ 
-			return ESubsequentsMode::FireAndForget; 
-		}
-		void DoTask(ENamedThreads::Type CurrentThread, const FGraphEventRef& MyCompletionGraphEvent)
-		{
-			FMemMark Mark(FMemStack::Get());
-			if (Data->Process(Contexts, TaskIndex, TasksToSpawn, Data, DesiredThread, false))
-			{
-				checkSlow(!Data->bTriggered);
-				Data->bTriggered = true;
-				Data->Event->Trigger();
-			}
-		}
-	};
 
 	// Helper to call body with context reference
 	template <typename FunctionType, typename ContextType>
@@ -180,204 +63,28 @@ namespace ParallelForImpl
 		Body(Index);
 	}
 
-	template<typename FunctionType> 
-	template<typename ContextType> 
-	inline bool TParallelForData<FunctionType>::Process(const TArrayView<ContextType>& Contexts, int32 TaskIndex, int32 TasksToSpawn, TSharedRef<TParallelForData<FunctionType>, ESPMode::ThreadSafe>& Data, ENamedThreads::Type InDesiredThread, bool bMaster)
-	{
-		TRACE_CPUPROFILER_EVENT_SCOPE(TParallelForData::Process)
-
-		int32 MaybeTasksLeft = Num - IndexToDo.GetValue();
-		if (TasksToSpawn && MaybeTasksLeft > 0)
-		{
-			TasksToSpawn = FMath::Min<int32>(TasksToSpawn, MaybeTasksLeft);
-			TGraphTask<TParallelForTask<FunctionType, ContextType>>::CreateTask().ConstructAndDispatchWhenReady(Contexts, TaskIndex + 1, Data, InDesiredThread, TasksToSpawn - 1);
-		}
-		int32 LocalBlockSize = BlockSize;
-		int32 LocalNum = Num;
-		bool bLocalSaveLastBlockForMaster = bSaveLastBlockForMaster;
-
-		auto Now = [] { return FTimespan::FromSeconds(FPlatformTime::Seconds()); };
-		FTimespan Start = FTimespan::MinValue();
-			
-		bool bIsBackgroundPriority = !bMaster && ((InDesiredThread & ENamedThreads::ThreadPriorityMask) == ENamedThreads::BackgroundThreadPriority);
-
-		FTimespan YieldingThreshold;
-		if (bIsBackgroundPriority)
-		{
-			Start = Now();
-			YieldingThreshold = FTimespan::FromMilliseconds(FMath::Max(0, GParallelForBackgroundYieldingTimeoutMs));
-		}
-
-		while (true)
-		{
-			int32 MyIndex = IndexToDo.Increment() - 1;
-			if (bLocalSaveLastBlockForMaster)
-			{
-				if (!bMaster && MyIndex >= LocalNum - 1)
-				{
-					break; // leave the last block for the master, hoping to avoid an event
-				}
-				else if (bMaster && MyIndex > LocalNum - 1)
-				{
-					MyIndex = LocalNum - 1; // I am the master, I need to take this block, hoping to avoid an event
-				}
-			}
-			if (MyIndex < LocalNum)
-			{
-				check(Contexts.Num() == 0 || TaskIndex < Contexts.Num());
-
-				int32 ThisBlockSize = LocalBlockSize;
-				if (MyIndex == LocalNum - 1)
-				{
-					ThisBlockSize += LastBlockExtraNum;
-				}
-				for (int32 LocalIndex = 0; LocalIndex < ThisBlockSize; LocalIndex++)
-				{
-					CallBody(Body, Contexts, TaskIndex, MyIndex * LocalBlockSize + LocalIndex);
-				}
-				checkSlow(!bExited.load(std::memory_order_relaxed));
-				int32 LocalNumCompleted = NumCompleted.Increment();
-				if (LocalNumCompleted == LocalNum)
-				{
-					return true;
-				}
-				checkSlow(LocalNumCompleted < LocalNum);
-			}
-			if (MyIndex >= LocalNum - 1)
-			{
-				break;
-			}
-
-			if (bIsBackgroundPriority)
-			{
-				auto PassedTime = [Start, &Now]() { return Now() - Start; };
-				if (PassedTime() > YieldingThreshold)
-				{
-					TGraphTask<TParallelForTask<FunctionType, ContextType>>::CreateTask().ConstructAndDispatchWhenReady(Contexts, TaskIndex, Data, InDesiredThread, 0);
-					return false;
-				}
-			}
-		}
-		return false;
-	}
-
-	inline ENamedThreads::Type GetBestDesiredThread(EParallelForFlags Flags)
-	{
-		const bool bBackgroundPriority = (Flags & EParallelForFlags::BackgroundPriority) != EParallelForFlags::None;
-		if (!bBackgroundPriority)
-		{
-			// Anything scheduled by the task graph is latency sensitive because it might impact the frame rate. Anything else is not (i.e. Worker / Background threads).
-			const ETaskTag LatencySensitiveTasks = 
-				ETaskTag::EStaticInit | 
-				ETaskTag::EGameThread | 
-				ETaskTag::ESlateThread | 
-#if !UE_AUDIO_THREAD_AS_PIPE
-				ETaskTag::EAudioThread | 
-#endif
-				ETaskTag::ERenderingThread | 
-				ETaskTag::ERhiThread;
-
-			const bool bIsLatencySensitive = (FTaskTagScope::GetCurrentTag() & LatencySensitiveTasks) != ETaskTag::ENone;
-			if (bIsLatencySensitive)
-			{
-				// Keep the legacy behavior in this case
-				return ENamedThreads::AnyHiPriThreadHiPriTask;
-			}
-			// It's coming from a known worker thread, we'll keep the same task and thread priority in this case
-			else if (FTaskGraphInterface::Get().IsCurrentThreadKnown())
-			{
-				const ENamedThreads::Type CurrentThread = FTaskGraphInterface::Get().GetCurrentThreadIfKnown();
-				const ENamedThreads::Type CurrentTaskPrio = (ENamedThreads::Type)(CurrentThread & ENamedThreads::Type::TaskPriorityMask);
-				const ENamedThreads::Type CurrentThreadPrio = (ENamedThreads::Type)(CurrentThread & ENamedThreads::Type::ThreadPriorityMask);
-				return (ENamedThreads::Type)(ENamedThreads::AnyThread | CurrentTaskPrio | CurrentThreadPrio);
-			}
-		}
-
-		// Either the request comes from a totally unknown thread, or we've specifically been asked to be background
-		return ENamedThreads::AnyBackgroundThreadNormalTask;
-	}
-
-	inline int32 GetNumberOfThreadTasks(int32 Num, EParallelForFlags Flags)
+	inline int32 GetNumberOfThreadTasks(int32 Num, int32 MinBatchSize, EParallelForFlags Flags)
 	{
 		int32 NumThreadTasks = 0;
 		const bool bIsMultithread = FApp::ShouldUseThreadingForPerformance() || FForkProcessHelper::IsForkedMultithreadInstance();
 		if (Num > 1 && (Flags & EParallelForFlags::ForceSingleThread) == EParallelForFlags::None && bIsMultithread)
 		{
-			NumThreadTasks = FMath::Min<int32>(FTaskGraphInterface::Get().GetNumWorkerThreads(), Num - 1);
+			NumThreadTasks = FMath::Min(int32(LowLevelTasks::FScheduler::Get().GetNumWorkers()), (Num + (MinBatchSize/2))/MinBatchSize);
 		}
 
-		return NumThreadTasks;
-	}
-
-	/** 
-		*	General purpose parallel for that uses the taskgraph
-		*	@param Num; number of calls of Body; Body(0), Body(1)....Body(Num - 1)
-		*	@param Body; Function to call from multiple threads
-		*	@param CurrentThreadWorkToDoBeforeHelping; The work is performed on the main thread before it starts helping with the ParallelFor proper
-		*	@param Flags; Used to customize the behavior of the ParallelFor if needed.
-		*   @param Contexts; Optional per thread contexts to accumulate data concurrently.
-		*	Notes: Please add stats around to calls to parallel for and within your lambda as appropriate. Do not clog the task graph with long running tasks or tasks that block.
-	**/
-	template<typename FunctionType, typename PreWorkType, typename ContextType>
-	inline void OldParallelForInternal(int32 Num, FunctionType Body, PreWorkType CurrentThreadWorkToDoBeforeHelping, EParallelForFlags Flags, const TArrayView<ContextType>& Contexts)
-	{
-		SCOPE_CYCLE_COUNTER(STAT_ParallelFor);
-		check(Num >= 0);
-
-		int32 AnyThreadTasks = GetNumberOfThreadTasks(Num, Flags);
-		if (!AnyThreadTasks)
-		{
-			// do the prework
-			CurrentThreadWorkToDoBeforeHelping();
-			// no threads, just do it and return
-			for (int32 Index = 0; Index < Num; Index++)
-			{				
-				CallBody(Body, Contexts, 0, Index);
-			}
-			return;
-		}
-
-		const bool bPumpRenderingThread         = (Flags & EParallelForFlags::PumpRenderingThread) != EParallelForFlags::None;
-		const ENamedThreads::Type DesiredThread = GetBestDesiredThread(Flags);
-
-		TParallelForData<FunctionType>* DataPtr = new TParallelForData<FunctionType>(Num, AnyThreadTasks + 1, (Num > AnyThreadTasks + 1) && bPumpRenderingThread, Body, Flags);
-		TSharedRef<TParallelForData<FunctionType>, ESPMode::ThreadSafe> Data = MakeShareable(DataPtr);
-		TGraphTask<TParallelForTask<FunctionType, ContextType>>::CreateTask().ConstructAndDispatchWhenReady(Contexts, 1, Data, DesiredThread, AnyThreadTasks - 1);
-		// do the prework
-		CurrentThreadWorkToDoBeforeHelping();
-		// this thread can help too and this is important to prevent deadlock on recursion 
-		if (!Data->Process(Contexts, 0, 0, Data, DesiredThread, true))
-		{
-			if (bPumpRenderingThread && IsInActualRenderingThread())
-			{
-				while (!Data->Event->Wait(1))
-				{
-					FTaskGraphInterface::Get().ProcessThreadUntilIdle(ENamedThreads::GetRenderThread_Local());
-				}
-			}
-			else
-			{
-				Data->Event->Wait();
-			}
-			check(Data->bTriggered);
-		}
-		else
+		if (!LowLevelTasks::FScheduler::Get().IsWorkerThread())
 		{
-			check(!Data->bTriggered);
+			NumThreadTasks++; //named threads help with the work
 		}
-		check(Data->NumCompleted.GetValue() == Data->Num);
 
-#if DO_CHECK
-		Data->bExited.store(true, std::memory_order_relaxed);
-#endif
-
-		// DoneEvent waits here if some other thread finishes the last item
-		// Data must live on until all of the tasks are cleared which might be long after this function exits
+		return FMath::Max(NumThreadTasks, 1);
 	}
 
 	/** 
 		*	General purpose parallel for that uses the taskgraph
+		*	@param DebugName; Debugname and Profiling TraceTag
 		*	@param Num; number of calls of Body; Body(0), Body(1)....Body(Num - 1)
+		*	@param MinBatchSize; Minimum size a Batch should have
 		*	@param Body; Function to call from multiple threads
 		*	@param CurrentThreadWorkToDoBeforeHelping; The work is performed on the main thread before it starts helping with the ParallelFor proper
 		*	@param Flags; Used to customize the behavior of the ParallelFor if needed.
@@ -385,14 +92,14 @@ namespace ParallelForImpl
 		*	Notes: Please add stats around to calls to parallel for and within your lambda as appropriate. Do not clog the task graph with long running tasks or tasks that block.
 	**/
 	template<typename BodyType, typename PreWorkType, typename ContextType>
-	inline void NewParallelForInternal(int32 Num, BodyType Body, PreWorkType CurrentThreadWorkToDoBeforeHelping, EParallelForFlags Flags, const TArrayView<ContextType>& Contexts)
+	inline void ParallelForInternal(const TCHAR* DebugName, int32 Num, int32 MinBatchSize, BodyType Body, PreWorkType CurrentThreadWorkToDoBeforeHelping, EParallelForFlags Flags, const TArrayView<ContextType>& Contexts)
 	{
 		SCOPE_CYCLE_COUNTER(STAT_ParallelFor);
 		check(Num >= 0);
 
+		int32 NumWorkers = GetNumberOfThreadTasks(Num, MinBatchSize, Flags);
 		//single threaded mode
-		const bool bIsMultithread = FApp::ShouldUseThreadingForPerformance() || FForkProcessHelper::IsForkedMultithreadInstance();
-		if ((Num <= 1) || ((Flags & EParallelForFlags::ForceSingleThread) == EParallelForFlags::ForceSingleThread) || !bIsMultithread)
+		if (NumWorkers <= 1)
 		{
 			// do the prework
 			CurrentThreadWorkToDoBeforeHelping();
@@ -403,16 +110,7 @@ namespace ParallelForImpl
 			}
 			return;
 		}
-
-		//calculate the number of workers
-		int32 NumWorkers = int32(LowLevelTasks::FScheduler::Get().GetNumWorkers());
-		if (!LowLevelTasks::FScheduler::Get().IsWorkerThread())
-		{
-			NumWorkers++; //named threads help with the work
-		}
-		NumWorkers = FMath::Min(NumWorkers, Num);
-		checkSlow(NumWorkers > 0);
-
+	
 		//calculate the batch sizes
 		int32 BatchSize = 1;
 		int32 NumBatches = Num;
@@ -474,12 +172,21 @@ namespace ParallelForImpl
 			{
 				IncompleteBatches.store(NumBatches, std::memory_order_relaxed);
 				Tasks.AddDefaulted(InNumWorkers);
+
+#if UE_MEMORY_TAGS_TRACE_ENABLED
+				InheritedTraceTag = MemoryTrace_GetActiveTag();
+#endif
+				LLM(InheritedLLMTag = FLowLevelMemTracker::bIsDisabled ? nullptr : FLowLevelMemTracker::Get().GetActiveTagData(ELLMTracker::Default));
 			}
 			std::atomic_int BatchItem  { 0 };
 			std::atomic_int IncompleteBatches { 0 };
 			int32 Num;
 			int32 BatchSize;
 			int32 NumBatches;
+#if UE_MEMORY_TAGS_TRACE_ENABLED
+			int32 InheritedTraceTag;
+#endif
+			LLM(const UE::LLMPrivate::FTagData* InheritedLLMTag);
 			const TArrayView<ContextType>& Contexts;
 			const BodyType& Body;
 			FEventRef& FinishedSignal;
@@ -514,7 +221,7 @@ namespace ParallelForImpl
 			{
 				if (Data.IsValid() && WorkerIndex >= 0)
 				{
-					FParallelExecutor::LaunchTask(MoveTemp(Data), WorkerIndex, Priority);
+					FParallelExecutor::LaunchTask(nullptr, MoveTemp(Data), WorkerIndex, Priority);
 				}
 			}
 
@@ -523,10 +230,19 @@ namespace ParallelForImpl
 				return Data;
 			}
 
-			inline bool operator()(const bool bIsMaster = false) const noexcept
+			inline bool operator()(const bool bIsMaster = false, const TCHAR* DebugName = nullptr) const noexcept
 			{
+				LLM_SCOPE(Data->InheritedLLMTag);
+				UE_MEMSCOPE(Data->InheritedTraceTag);
 				FMemMark Mark(FMemStack::Get());
-				TRACE_CPUPROFILER_EVENT_SCOPE(ParallelFor);
+
+				if(DebugName == nullptr)
+				{
+					LowLevelTasks::FTask& Task = Data->Tasks[WorkerIndex];
+					DebugName = Task.GetDebugName();
+				}
+				TRACE_CPUPROFILER_EVENT_SCOPE_TEXT(DebugName);
+
 
 				auto Now = [] { return FTimespan::FromSeconds(FPlatformTime::Seconds()); };
 				FTimespan Start = FTimespan::MinValue();
@@ -598,10 +314,14 @@ namespace ParallelForImpl
 				}
 			}
 
-			static void LaunchTask(FDataHandle&& InData, int32 InWorkerIndex, LowLevelTasks::ETaskPriority InPriority)
+			static void LaunchTask(const TCHAR* DebugName, FDataHandle&& InData, int32 InWorkerIndex, LowLevelTasks::ETaskPriority InPriority)
 			{
 				LowLevelTasks::FTask& Task = InData->Tasks[InWorkerIndex];
-				Task.Init(TEXT("FParallelExecutor"), InPriority, FParallelExecutor(MoveTemp(InData), InWorkerIndex, InPriority));
+				if(DebugName == nullptr)
+				{
+					DebugName = Task.GetDebugName();
+				}
+				Task.Init(DebugName, InPriority, FParallelExecutor(MoveTemp(InData), InWorkerIndex, InPriority));
 				verify(LowLevelTasks::TryLaunch(Task, LowLevelTasks::EQueuePreference::GlobalQueuePreference));
 			}
 		};
@@ -611,7 +331,7 @@ namespace ParallelForImpl
 		FDataHandle Data = new FParallelForData(Num, BatchSize, NumBatches, NumWorkers, Contexts, Body, FinishedSignal);
 		for (int32 Worker = 0; Worker < NumWorkers; Worker++)
 		{
-			FParallelExecutor::LaunchTask(FDataHandle(Data), Worker, Priority);
+			FParallelExecutor::LaunchTask(DebugName, FDataHandle(Data), Worker, Priority);
 		}
 
 		// do the prework
@@ -619,7 +339,7 @@ namespace ParallelForImpl
 
 		//help with the parallel-for to prevent deadlocks
 		FParallelExecutor LocalExecutor(MoveTemp(Data), NumWorkers, Priority);
-		const bool bFinishedLast = LocalExecutor(true);
+		const bool bFinishedLast = LocalExecutor(true, DebugName);
 
 		if(!bFinishedLast)
 		{
@@ -642,19 +362,6 @@ namespace ParallelForImpl
 		}
 		checkSlow(LocalExecutor.GetData()->BatchItem.load(std::memory_order_relaxed) * LocalExecutor.GetData()->BatchSize >= LocalExecutor.GetData()->Num);
 	}
-
-	template<typename FunctionType, typename PreWorkType, typename ContextType>
-	inline void ParallelForInternal(int32 Num, FunctionType Body, PreWorkType CurrentThreadWorkToDoBeforeHelping, EParallelForFlags Flags, const TArrayView<ContextType>& Contexts)
-	{
-		if (LowLevelTasks::FScheduler::Get().GetNumWorkers())
-		{
-			NewParallelForInternal<FunctionType, PreWorkType, ContextType>(Num, Body, CurrentThreadWorkToDoBeforeHelping, Flags, Contexts);
-		}
-		else
-		{
-			OldParallelForInternal<FunctionType, PreWorkType, ContextType>(Num, Body, CurrentThreadWorkToDoBeforeHelping, Flags, Contexts);
-		}
-	}
 }
 
 /** 
@@ -666,7 +373,7 @@ namespace ParallelForImpl
 **/
 inline void ParallelFor(int32 Num, TFunctionRef<void(int32)> Body, bool bForceSingleThread, bool bPumpRenderingThread=false)
 {
-	ParallelForImpl::ParallelForInternal(Num, Body, [](){},
+	ParallelForImpl::ParallelForInternal(TEXT("ParallelFor"), Num, 1, Body, [](){},
 		(bForceSingleThread ? EParallelForFlags::ForceSingleThread : EParallelForFlags::None) | 
 		(bPumpRenderingThread ? EParallelForFlags::PumpRenderingThread : EParallelForFlags::None), TArrayView<nullptr_t>());
 }
@@ -681,8 +388,24 @@ inline void ParallelFor(int32 Num, TFunctionRef<void(int32)> Body, bool bForceSi
 template<typename FunctionType>
 inline void ParallelForTemplate(int32 Num, const FunctionType& Body, EParallelForFlags Flags = EParallelForFlags::None)
 {
-	ParallelForImpl::ParallelForInternal(Num, Body, [](){}, Flags, TArrayView<nullptr_t>());
+	ParallelForImpl::ParallelForInternal(TEXT("ParallelFor"), Num, 1, Body, [](){}, Flags, TArrayView<nullptr_t>());
 }
+
+/**
+	*	General purpose parallel for that uses the taskgraph
+	*   @param DebugName; ProfilingScope and Debugname
+	*	@param Num; number of calls of Body; Body(0), Body(1)....Body(Num - 1)
+	*   @param MinBatchSize; Minimum Size of a Batch (will only launch DivUp(Num, MinBatchSize) Workers 
+	*	@param Body; Function to call from multiple threads
+	*	@param bForceSingleThread; Mostly used for testing, if true, run single threaded instead.
+	*	Notes: Please add stats around to calls to parallel for and within your lambda as appropriate. Do not clog the task graph with long running tasks or tasks that block.
+**/
+template<typename FunctionType>
+inline void ParallelForTemplate(const TCHAR* DebugName, int32 Num, int32 MinBatchSize, const FunctionType& Body, EParallelForFlags Flags = EParallelForFlags::None)
+{
+	ParallelForImpl::ParallelForInternal(DebugName, Num, MinBatchSize, Body, [](){}, Flags, TArrayView<nullptr_t>());
+}
+
 /** 
 	*	General purpose parallel for that uses the taskgraph for unbalanced tasks
 	*	Offers better work distribution among threads at the cost of a little bit more synchronization.
@@ -695,7 +418,24 @@ inline void ParallelForTemplate(int32 Num, const FunctionType& Body, EParallelFo
 **/
 inline void ParallelFor(int32 Num, TFunctionRef<void(int32)> Body, EParallelForFlags Flags = EParallelForFlags::None)
 {
-	ParallelForImpl::ParallelForInternal(Num, Body, [](){}, Flags, TArrayView<nullptr_t>());
+	ParallelForImpl::ParallelForInternal(TEXT("ParallelFor"), Num, 1, Body, [](){}, Flags, TArrayView<nullptr_t>());
+}
+
+/** 
+	*	General purpose parallel for that uses the taskgraph for unbalanced tasks
+	*	Offers better work distribution among threads at the cost of a little bit more synchronization.
+	*	This should be used for tasks with highly variable computational time.
+	*
+	*   @param DebugName; ProfilingScope and Debugname
+	*	@param Num; number of calls of Body; Body(0), Body(1)....Body(Num - 1)
+	*   @param MinBatchSize; Minimum Size of a Batch (will only launch DivUp(Num, MinBatchSize) Workers 
+	*	@param Body; Function to call from multiple threads
+	*	@param Flags; Used to customize the behavior of the ParallelFor if needed.
+	*	Notes: Please add stats around to calls to parallel for and within your lambda as appropriate. Do not clog the task graph with long running tasks or tasks that block.
+**/
+inline void ParallelFor(const TCHAR* DebugName, int32 Num, int32 MinBatchSize, TFunctionRef<void(int32)> Body, EParallelForFlags Flags = EParallelForFlags::None)
+{
+	ParallelForImpl::ParallelForInternal(DebugName, Num, MinBatchSize, Body, [](){}, Flags, TArrayView<nullptr_t>());
 }
 
 /** 
@@ -708,7 +448,7 @@ inline void ParallelFor(int32 Num, TFunctionRef<void(int32)> Body, EParallelForF
 **/
 inline void ParallelForWithPreWork(int32 Num, TFunctionRef<void(int32)> Body, TFunctionRef<void()> CurrentThreadWorkToDoBeforeHelping, bool bForceSingleThread, bool bPumpRenderingThread = false)
 {
-	ParallelForImpl::ParallelForInternal(Num, Body, CurrentThreadWorkToDoBeforeHelping,
+	ParallelForImpl::ParallelForInternal(TEXT("ParallelFor"), Num, 1, Body, CurrentThreadWorkToDoBeforeHelping,
 		(bForceSingleThread ? EParallelForFlags::ForceSingleThread : EParallelForFlags::None) |
 		(bPumpRenderingThread ? EParallelForFlags::PumpRenderingThread : EParallelForFlags::None), TArrayView<nullptr_t>());
 }
@@ -723,7 +463,22 @@ inline void ParallelForWithPreWork(int32 Num, TFunctionRef<void(int32)> Body, TF
 **/
 inline void ParallelForWithPreWork(int32 Num, TFunctionRef<void(int32)> Body, TFunctionRef<void()> CurrentThreadWorkToDoBeforeHelping, EParallelForFlags Flags = EParallelForFlags::None)
 {
-	ParallelForImpl::ParallelForInternal(Num, Body, CurrentThreadWorkToDoBeforeHelping, Flags, TArrayView<nullptr_t>());
+	ParallelForImpl::ParallelForInternal(TEXT("ParallelFor"), Num, 1, Body, CurrentThreadWorkToDoBeforeHelping, Flags, TArrayView<nullptr_t>());
+}
+
+/** 
+	*	General purpose parallel for that uses the taskgraph
+	*   @param DebugName; ProfilingScope and Debugname
+	*	@param Num; number of calls of Body; Body(0), Body(1)....Body(Num - 1)
+	*   @param MinBatchSize; Minimum Size of a Batch (will only launch DivUp(Num, MinBatchSize) Workers 
+	*	@param Body; Function to call from multiple threads
+	*	@param CurrentThreadWorkToDoBeforeHelping; The work is performed on the main thread before it starts helping with the ParallelFor proper
+	*	@param Flags; Used to customize the behavior of the ParallelFor if needed.
+	*	Notes: Please add stats around to calls to parallel for and within your lambda as appropriate. Do not clog the task graph with long running tasks or tasks that block.
+**/
+inline void ParallelForWithPreWork(const TCHAR* DebugName, int32 Num, int32 MinBatchSize, TFunctionRef<void(int32)> Body, TFunctionRef<void()> CurrentThreadWorkToDoBeforeHelping, EParallelForFlags Flags = EParallelForFlags::None)
+{
+	ParallelForImpl::ParallelForInternal(DebugName, Num, MinBatchSize, Body, CurrentThreadWorkToDoBeforeHelping, Flags, TArrayView<nullptr_t>());
 }
 
 /** 
@@ -743,14 +498,14 @@ inline void ParallelForWithTaskContext(TArray<ContextType, ContextAllocatorType>
 {
 	if (Num > 0)
 	{
-		const int32 NumContexts = ParallelForImpl::GetNumberOfThreadTasks(Num, Flags) + 1;
+		const int32 NumContexts = ParallelForImpl::GetNumberOfThreadTasks(Num, 1, Flags);
 		OutContexts.Reset();
 		OutContexts.AddUninitialized(NumContexts);
 		for (int32 ContextIndex = 0; ContextIndex < NumContexts; ++ContextIndex)
 		{
 			new(&OutContexts[ContextIndex]) ContextType(ContextConstructor(ContextIndex, NumContexts));
 		}
-		ParallelForImpl::ParallelForInternal(Num, Body, [](){}, Flags, TArrayView<ContextType>(OutContexts));
+		ParallelForImpl::ParallelForInternal(TEXT("ParallelFor"), Num, 1, Body, [](){}, Flags, TArrayView<ContextType>(OutContexts));
 	}
 }
 
@@ -769,9 +524,63 @@ inline void ParallelForWithTaskContext(TArray<ContextType, ContextAllocatorType>
 {
 	if (Num > 0)
 	{
-		const int32 NumContexts = ParallelForImpl::GetNumberOfThreadTasks(Num, Flags) + 1;
+		const int32 NumContexts = ParallelForImpl::GetNumberOfThreadTasks(Num, 1, Flags);
+		OutContexts.Reset();
+		OutContexts.AddDefaulted(NumContexts);
+		ParallelForImpl::ParallelForInternal(TEXT("ParallelFor"), Num, 1, Body, [](){}, Flags, TArrayView<ContextType>(OutContexts));
+	}
+}
+
+/** 
+	*	General purpose parallel for that uses the taskgraph. This variant constructs for the caller a user-defined context
+	* 	object for each task that may get spawned to do work, and passes it on to the loop body to give it a task-local
+	*   "workspace" that can be mutated without need for synchronization primitives. For this variant, the user provides a
+	* 	callable to construct each context element.
+	*	@param OutContexts; Array that will hold the user-defined, task-level context objects (allocated per parallel task)
+	*   @param DebugName; ProfilingScope and Debugname
+	*	@param Num; number of calls of Body; Body(0), Body(1)....Body(Num - 1)
+	*   @param MinBatchSize; Minimum Size of a Batch (will only launch DivUp(Num, MinBatchSize) Workers 
+	* 	@param ContextConstructor; Function to call to initialize each task context allocated for the operation
+	*	@param Body; Function to call from multiple threads
+	*	@param Flags; Used to customize the behavior of the ParallelFor if needed.
+	*	Notes: Please add stats around to calls to parallel for and within your lambda as appropriate. Do not clog the task graph with long running tasks or tasks that block.
+**/
+template <typename ContextType, typename ContextAllocatorType, typename ContextConstructorType, typename FunctionType>
+inline void ParallelForWithTaskContext(const TCHAR* DebugName, TArray<ContextType, ContextAllocatorType>& OutContexts, int32 Num, int32 MinBatchSize, const ContextConstructorType& ContextConstructor, const FunctionType& Body, EParallelForFlags Flags = EParallelForFlags::None)
+{
+	if (Num > 0)
+	{
+		const int32 NumContexts = ParallelForImpl::GetNumberOfThreadTasks(Num, MinBatchSize, Flags);
+		OutContexts.Reset();
+		OutContexts.AddUninitialized(NumContexts);
+		for (int32 ContextIndex = 0; ContextIndex < NumContexts; ++ContextIndex)
+		{
+			new(&OutContexts[ContextIndex]) ContextType(ContextConstructor(ContextIndex, NumContexts));
+		}
+		ParallelForImpl::ParallelForInternal(DebugName, Num, MinBatchSize, Body, [](){}, Flags, TArrayView<ContextType>(OutContexts));
+	}
+}
+
+/** 
+	*	General purpose parallel for that uses the taskgraph. This variant constructs for the caller a user-defined context
+	* 	object for each task that may get spawned to do work, and passes it on to the loop body to give it a task-local
+	*   "workspace" that can be mutated without need for synchronization primitives.
+	*	@param OutContexts; Array that will hold the user-defined, task-level context objects (allocated per parallel task)
+	*   @param DebugName; ProfilingScope and Debugname
+	*	@param Num; number of calls of Body; Body(0), Body(1)....Body(Num - 1)
+	*   @param MinBatchSize; Minimum Size of a Batch (will only launch DivUp(Num, MinBatchSize) Workers 
+	*	@param Body; Function to call from multiple threads
+	*	@param Flags; Used to customize the behavior of the ParallelFor if needed.
+	*	Notes: Please add stats around to calls to parallel for and within your lambda as appropriate. Do not clog the task graph with long running tasks or tasks that block.
+**/
+template <typename ContextType, typename ContextAllocatorType, typename FunctionType>
+inline void ParallelForWithTaskContext(const TCHAR* DebugName, TArray<ContextType, ContextAllocatorType>& OutContexts, int32 Num, int32 MinBatchSize, const FunctionType& Body, EParallelForFlags Flags = EParallelForFlags::None)
+{
+	if (Num > 0)
+	{
+		const int32 NumContexts = ParallelForImpl::GetNumberOfThreadTasks(Num, MinBatchSize, Flags);
 		OutContexts.Reset();
 		OutContexts.AddDefaulted(NumContexts);
-		ParallelForImpl::ParallelForInternal(Num, Body, [](){}, Flags, TArrayView<ContextType>(OutContexts));
+		ParallelForImpl::ParallelForInternal(DebugName, Num, MinBatchSize, Body, [](){}, Flags, TArrayView<ContextType>(OutContexts));
 	}
 }
\ No newline at end of file
