Move the command line switches outside of the if/else FORCE_ANSI loop as the command line switches are no longer allocator only. Related to UE-71155, but that jira will be the proper fix to this madness

diff --git a/Engine/Source/Runtime/Core/Private/Unix/UnixPlatformMemory.cpp b/Engine/Source/Runtime/Core/Private/Unix/UnixPlatformMemory.cpp
--- a/Engine/Source/Runtime/Core/Private/Unix/UnixPlatformMemory.cpp
+++ b/Engine/Source/Runtime/Core/Private/Unix/UnixPlatformMemory.cpp
@@ -135,156 +135,156 @@ class FMalloc* FUnixPlatformMemory::BaseAllocator()
 	{
 		AllocatorToUse = EMemoryAllocatorToUse::Binned2;
 	}
-	else 
+	else
 	{
 		AllocatorToUse = EMemoryAllocatorToUse::Binned;
 	}
-	
-	if (FORCE_ANSI_ALLOCATOR)
-	{
-		AllocatorToUse = EMemoryAllocatorToUse::Ansi;
-	}
-	else
-	{
-		// Mimalloc is now the default allocator for editor and programs because it has shown
-		// both great performance and as much as half the memory usage of TBB after
-		// heavy editor workloads. See CL 15887498 description for benchmarks.
+
+	// Mimalloc is now the default allocator for editor and programs because it has shown
+	// both great performance and as much as half the memory usage of TBB after
+	// heavy editor workloads. See CL 15887498 description for benchmarks.
 #if (WITH_EDITORONLY_DATA || IS_PROGRAM) && PLATFORM_SUPPORTS_MIMALLOC && MIMALLOC_ALLOCATOR_ALLOWED
-		AllocatorToUse = EMemoryAllocatorToUse::Mimalloc;
+	AllocatorToUse = EMemoryAllocatorToUse::Mimalloc;
 #endif
 
-		// Allow overriding on the command line.
-		// We get here before main due to global ctors, so need to do some hackery to get command line args
-		if (FILE* CmdLineFile = fopen("/proc/self/cmdline", "r"))
+	// Allow overriding on the command line.
+	// We get here before main due to global ctors, so need to do some hackery to get command line args
+	if (FILE* CmdLineFile = fopen("/proc/self/cmdline", "r"))
+	{
+		char * Arg = nullptr;
+		size_t Size = 0;
+		while(getdelim(&Arg, &Size, 0, CmdLineFile) != -1)
 		{
-			char * Arg = nullptr;
-			size_t Size = 0;
-			while(getdelim(&Arg, &Size, 0, CmdLineFile) != -1)
-			{
 #if PLATFORM_SUPPORTS_JEMALLOC
-				if (FCStringAnsi::Stricmp(Arg, "-jemalloc") == 0)
-				{
-					AllocatorToUse = EMemoryAllocatorToUse::Jemalloc;
-					break;
-				}
+			if (FCStringAnsi::Stricmp(Arg, "-jemalloc") == 0)
+			{
+				AllocatorToUse = EMemoryAllocatorToUse::Jemalloc;
+				break;
+			}
 #endif // PLATFORM_SUPPORTS_JEMALLOC
-				if (FCStringAnsi::Stricmp(Arg, "-ansimalloc") == 0)
-				{
-					// see FPlatformMisc::GetProcessDiagnostics()
-					AllocatorToUse = EMemoryAllocatorToUse::Ansi;
-					break;
-				}
+			if (FCStringAnsi::Stricmp(Arg, "-ansimalloc") == 0)
+			{
+				// see FPlatformMisc::GetProcessDiagnostics()
+				AllocatorToUse = EMemoryAllocatorToUse::Ansi;
+				break;
+			}
 
-				if (FCStringAnsi::Stricmp(Arg, "-binnedmalloc") == 0)
-				{
-					AllocatorToUse = EMemoryAllocatorToUse::Binned;
-					break;
-				}
+			if (FCStringAnsi::Stricmp(Arg, "-binnedmalloc") == 0)
+			{
+				AllocatorToUse = EMemoryAllocatorToUse::Binned;
+				break;
+			}
 
 #if PLATFORM_SUPPORTS_MIMALLOC && MIMALLOC_ALLOCATOR_ALLOWED
-				if (FCStringAnsi::Stricmp(Arg, "-mimalloc") == 0)
-				{
-					AllocatorToUse = EMemoryAllocatorToUse::Mimalloc;
-					break;
-				}
+			if (FCStringAnsi::Stricmp(Arg, "-mimalloc") == 0)
+			{
+				AllocatorToUse = EMemoryAllocatorToUse::Mimalloc;
+				break;
+			}
 #endif
 
-				if (FCStringAnsi::Stricmp(Arg, "-binnedmalloc2") == 0)
-				{
-					AllocatorToUse = EMemoryAllocatorToUse::Binned2;
-					break;
-				}
-
-				if (FCStringAnsi::Stricmp(Arg, "-fullcrashcallstack") == 0)
-				{
-					GFullCrashCallstack = true;
-				}
-
-				if (FCStringAnsi::Stricmp(Arg, "-useksm") == 0)
-				{
-					GUseKSM = true;
-				}
+			if (FCStringAnsi::Stricmp(Arg, "-binnedmalloc2") == 0)
+			{
+				AllocatorToUse = EMemoryAllocatorToUse::Binned2;
+				break;
+			}
 
-				if (FCStringAnsi::Stricmp(Arg, "-ksmmergeall") == 0)
-				{
-					GKSMMergeAllPages = true;
-				}
+			if (FCStringAnsi::Stricmp(Arg, "-fullcrashcallstack") == 0)
+			{
+				GFullCrashCallstack = true;
+			}
 
-				if (FCStringAnsi::Stricmp(Arg, "-noensuretiming") == 0)
-				{
-					GTimeEnsures = false;
-				}
+			if (FCStringAnsi::Stricmp(Arg, "-useksm") == 0)
+			{
+				GUseKSM = true;
+			}
 
-				if (FCStringAnsi::Stricmp(Arg, "-noexclusivelockonwrite") == 0)
-				{
-					GAllowExclusiveLockOnWrite = false;
-				}
+			if (FCStringAnsi::Stricmp(Arg, "-ksmmergeall") == 0)
+			{
+				GKSMMergeAllPages = true;
+			}
 
-				const char SignalToDefaultCmd[] = "-sigdfl=";
-				if (const char* Cmd = FCStringAnsi::Stristr(Arg, SignalToDefaultCmd))
-				{
-					int32 SignalToDefault = FCStringAnsi::Atoi(Cmd + sizeof(SignalToDefaultCmd) - 1);
+			if (FCStringAnsi::Stricmp(Arg, "-noensuretiming") == 0)
+			{
+				GTimeEnsures = false;
+			}
 
-					// Valid signals are only from 1 -> SIGRTMAX
-					if (SignalToDefault > SIGRTMAX)
-					{
-						SignalToDefault = 0;
-					}
+			if (FCStringAnsi::Stricmp(Arg, "-noexclusivelockonwrite") == 0)
+			{
+				GAllowExclusiveLockOnWrite = false;
+			}
 
-					GSignalToDefault = FMath::Max(SignalToDefault, 0);
-				}
+			const char SignalToDefaultCmd[] = "-sigdfl=";
+			if (const char* Cmd = FCStringAnsi::Stristr(Arg, SignalToDefaultCmd))
+			{
+				int32 SignalToDefault = FCStringAnsi::Atoi(Cmd + sizeof(SignalToDefaultCmd) - 1);
 
-				const char CrashHandlerStackSize[] = "-crashhandlerstacksize=";
-				if (const char* Cmd = FCStringAnsi::Stristr(Arg, CrashHandlerStackSize))
+				// Valid signals are only from 1 -> SIGRTMAX
+				if (SignalToDefault > SIGRTMAX)
 				{
-					GCrashHandlerStackSize = FCStringAnsi::Atoi64(Cmd + sizeof(CrashHandlerStackSize) - 1);
+					SignalToDefault = 0;
 				}
 
-				const char FileMapCacheCmd[] = "-filemapcachesize=";
-				if (const char* Cmd = FCStringAnsi::Stristr(Arg, FileMapCacheCmd))
-				{
-					int32 Max = FCStringAnsi::Atoi(Cmd + sizeof(FileMapCacheCmd) - 1);
-					GMaxNumberFileMappingCache = FMath::Clamp(Max, 0, MaximumAllowedMaxNumFileMappingCache);
-				}
+				GSignalToDefault = FMath::Max(SignalToDefault, 0);
+			}
+
+			const char CrashHandlerStackSize[] = "-crashhandlerstacksize=";
+			if (const char* Cmd = FCStringAnsi::Stristr(Arg, CrashHandlerStackSize))
+			{
+				GCrashHandlerStackSize = FCStringAnsi::Atoi64(Cmd + sizeof(CrashHandlerStackSize) - 1);
+			}
+
+			const char FileMapCacheCmd[] = "-filemapcachesize=";
+			if (const char* Cmd = FCStringAnsi::Stristr(Arg, FileMapCacheCmd))
+			{
+				int32 Max = FCStringAnsi::Atoi(Cmd + sizeof(FileMapCacheCmd) - 1);
+				GMaxNumberFileMappingCache = FMath::Clamp(Max, 0, MaximumAllowedMaxNumFileMappingCache);
+			}
 
 #if UE_USE_MALLOC_REPLAY_PROXY
-				if (FCStringAnsi::Stricmp(Arg, "-mallocsavereplay") == 0)
-				{
-					bAddReplayProxy = true;
-				}
+			if (FCStringAnsi::Stricmp(Arg, "-mallocsavereplay") == 0)
+			{
+				bAddReplayProxy = true;
+			}
 #endif // UE_USE_MALLOC_REPLAY_PROXY
 #if WITH_MALLOC_STOMP
-				if (FCStringAnsi::Stricmp(Arg, "-stompmalloc") == 0)
-				{
-					// see FPlatformMisc::GetProcessDiagnostics()
-					AllocatorToUse = EMemoryAllocatorToUse::Stomp;
-					break;
-				}
+			if (FCStringAnsi::Stricmp(Arg, "-stompmalloc") == 0)
+			{
+				// see FPlatformMisc::GetProcessDiagnostics()
+				AllocatorToUse = EMemoryAllocatorToUse::Stomp;
+				break;
+			}
 #endif // WITH_MALLOC_STOMP
 
-				const char VMAPoolScaleSwitch[] = "-vmapoolscale=";
-				if (const char* Cmd = FCStringAnsi::Stristr(Arg, VMAPoolScaleSwitch))
-				{
-					float PoolScale = FCStringAnsi::Atof(Cmd + sizeof(VMAPoolScaleSwitch) - 1);
-					GVMAPoolScale = FMath::Max(PoolScale, 1.0f);
-				}
+			const char VMAPoolScaleSwitch[] = "-vmapoolscale=";
+			if (const char* Cmd = FCStringAnsi::Stristr(Arg, VMAPoolScaleSwitch))
+			{
+				float PoolScale = FCStringAnsi::Atof(Cmd + sizeof(VMAPoolScaleSwitch) - 1);
+				GVMAPoolScale = FMath::Max(PoolScale, 1.0f);
+			}
 
-				if (FCStringAnsi::Stricmp(Arg, "-vmapoolevict") == 0)
-				{
-					GMemoryRangeDecommitIsNoOp = false;
-				}
-				if (FCStringAnsi::Stricmp(Arg, "-novmapoolevict") == 0)
-				{
-					GMemoryRangeDecommitIsNoOp = true;
-				}
-				if (FCStringAnsi::Stricmp(Arg, "-protectforkedpages") == 0)
-				{
-					GEnableProtectForkedPages = true;
-				}
+			if (FCStringAnsi::Stricmp(Arg, "-vmapoolevict") == 0)
+			{
+				GMemoryRangeDecommitIsNoOp = false;
+			}
+			if (FCStringAnsi::Stricmp(Arg, "-novmapoolevict") == 0)
+			{
+				GMemoryRangeDecommitIsNoOp = true;
+			}
+			if (FCStringAnsi::Stricmp(Arg, "-protectforkedpages") == 0)
+			{
+				GEnableProtectForkedPages = true;
 			}
-			free(Arg);
-			fclose(CmdLineFile);
 		}
+		free(Arg);
+		fclose(CmdLineFile);
+	}
+
+	// This was moved to the fact that we aboved the command line statements above to *include* other things besides allocator only switches
+	// Moving here allows the other globals to be set, while we override the ANSI allocator still no matter the command line options
+	if (FORCE_ANSI_ALLOCATOR)
+	{
+		AllocatorToUse = EMemoryAllocatorToUse::Ansi;
 	}
 
 	FMalloc * Allocator = NULL;
