sse_mathfun.h renamed to prevent IncludeTool warning.

diff --git a/Engine/Source/Runtime/Core/Public/Math/UnrealMathSSE.h b/Engine/Source/Runtime/Core/Public/Math/UnrealMathSSE.h
--- a/Engine/Source/Runtime/Core/Public/Math/UnrealMathSSE.h
+++ b/Engine/Source/Runtime/Core/Public/Math/UnrealMathSSE.h
@@ -60,7 +60,7 @@ struct VectorRegisterConstInit {};
 #define UE_SSE_DOUBLE_ALIGNMENT 16
 #endif
 
-#include "Math/sse_mathfun.h"
+#include "Math/sse_mathfun.inl"
 
 // We suppress static analysis warnings for the cast from (double*) to (float*) in VectorLoadFloat2
 // and VectorLoadTwoPairsFloat below:
diff --git a/Engine/Source/Runtime/Core/Public/Math/sse_mathfun.h b/Engine/Source/Runtime/Core/Public/Math/sse_mathfun.h
deleted file mode 100644
diff --git a/Engine/Source/Runtime/Core/Public/Math/sse_mathfun.inl b/Engine/Source/Runtime/Core/Public/Math/sse_mathfun.inl
new file mode 100644
--- /dev/null
+++ b/Engine/Source/Runtime/Core/Public/Math/sse_mathfun.inl
@@ -0,0 +1,730 @@
+// Copyright Epic Games, Inc. All Rights Reserved.
+
+/* SIMD (SSE1+MMX or SSE2) implementation of sin, cos, exp and log
+
+   Inspired by Intel Approximate Math library, and based on the
+   corresponding algorithms of the cephes math library
+
+   The default is to use the SSE1 version. If you define USE_SSE2 the
+   the SSE2 intrinsics will be used in place of the MMX intrinsics. Do
+   not expect any significant performance improvement with SSE2.
+*/
+
+/* Copyright (C) 2007  Julien Pommier
+
+  This software is provided 'as-is', without any express or implied
+  warranty.  In no event will the authors be held liable for any damages
+  arising from the use of this software.
+
+  Permission is granted to anyone to use this software for any purpose,
+  including commercial applications, and to alter it and redistribute it
+  freely, subject to the following restrictions:
+
+  1. The origin of this software must not be misrepresented; you must not
+     claim that you wrote the original software. If you use this software
+     in a product, an acknowledgment in the product documentation would be
+     appreciated but is not required.
+  2. Altered source versions must be plainly marked as such, and must not be
+     misrepresented as being the original software.
+  3. This notice may not be removed or altered from any source distribution.
+
+  (this is the zlib license)
+*/
+
+/*
+  EPIC CHANGES:
+
+  Making a few modifications to this file.
+  1. The original version of this code has #ifdef USE_SSE2. Do not want to pollute the defines
+     and we have no need for the MMX Version so all those are being replace with #if 1.
+  2. It looks like some linked object files have conflicting names (log_ps, exp_ps, etc). So adding the 
+     prefix SseMath_ to all the functions.
+  3. constexpr and alignas transition
+*/
+#pragma once
+
+#include <xmmintrin.h>
+
+/* __m128 is ugly to write */
+typedef __m128 v4sf;  // vector of 4 float (sse1)
+
+//#ifdef USE_SSE2
+#if 1
+# include <emmintrin.h>
+typedef __m128i v4si; // vector of 4 int (sse2)
+#else
+typedef __m64 v2si;   // vector of 2 int (mmx)
+#endif
+
+/* declare some SSE constants -- why can't I figure a better way to do that? */
+#define _PS_CONST(Name, Val)                                            \
+  alignas(16) static constexpr float _ps_##Name[4] = { Val, Val, Val, Val }
+#define _PI32_CONST(Name, Val)                                            \
+  alignas(16) static constexpr int _pi32_##Name[4] = { Val, Val, Val, Val }
+#define _PS_CONST_TYPE(Name, Type, Val)                                 \
+  alignas(16) static constexpr Type _ps_##Name[4] = { Val, Val, Val, Val }
+
+_PS_CONST(1  , 1.0f);
+_PS_CONST(0p5, 0.5f);
+/* the smallest non denormalized float number */
+_PS_CONST_TYPE(min_norm_pos, int, 0x00800000);
+_PS_CONST_TYPE(mant_mask, int, 0x7f800000);
+_PS_CONST_TYPE(inv_mant_mask, int, ~0x7f800000);
+
+_PS_CONST_TYPE(sign_mask, int, (int)0x80000000);
+_PS_CONST_TYPE(inv_sign_mask, int, ~0x80000000);
+
+_PI32_CONST(1, 1);
+_PI32_CONST(inv1, ~1);
+_PI32_CONST(2, 2);
+_PI32_CONST(4, 4);
+_PI32_CONST(0x7f, 0x7f);
+
+_PS_CONST(cephes_SQRTHF, 0.707106781186547524f);
+_PS_CONST(cephes_log_p0, 7.0376836292E-2f);
+_PS_CONST(cephes_log_p1, - 1.1514610310E-1f);
+_PS_CONST(cephes_log_p2, 1.1676998740E-1f);
+_PS_CONST(cephes_log_p3, - 1.2420140846E-1f);
+_PS_CONST(cephes_log_p4, + 1.4249322787E-1f);
+_PS_CONST(cephes_log_p5, - 1.6668057665E-1f);
+_PS_CONST(cephes_log_p6, + 2.0000714765E-1f);
+_PS_CONST(cephes_log_p7, - 2.4999993993E-1f);
+_PS_CONST(cephes_log_p8, + 3.3333331174E-1f);
+_PS_CONST(cephes_log_q1, -2.12194440e-4f);
+_PS_CONST(cephes_log_q2, 0.693359375f);
+
+//#ifdef USE_SSE2
+#if 1
+typedef union xmm_mm_union {
+  __m128 xmm;
+  __m64 mm[2];
+} xmm_mm_union;
+
+#define COPY_XMM_TO_MM(xmm_, mm0_, mm1_) {          \
+    xmm_mm_union u; u.xmm = xmm_;                   \
+    mm0_ = u.mm[0];                                 \
+    mm1_ = u.mm[1];                                 \
+}
+
+#define COPY_MM_TO_XMM(mm0_, mm1_, xmm_) {                         \
+    xmm_mm_union u; u.mm[0]=mm0_; u.mm[1]=mm1_; xmm_ = u.xmm;      \
+  }
+
+#endif // USE_SSE2
+
+/* natural logarithm computed for 4 simultaneous float 
+   return NaN for x <= 0
+*/
+inline v4sf SseMath_log_ps(v4sf x) {
+//#ifdef USE_SSE2
+#if 1
+  v4si emm0;
+#else
+  v2si mm0, mm1;
+#endif
+  v4sf one = *(v4sf*)_ps_1;
+
+  v4sf invalid_mask = _mm_cmple_ps(x, _mm_setzero_ps());
+
+  x = _mm_max_ps(x, *(v4sf*)_ps_min_norm_pos);  /* cut off denormalized stuff */
+
+//#ifndef USE_SSE2
+#if 0
+  /* part 1: x = frexpf(x, &e); */
+  COPY_XMM_TO_MM(x, mm0, mm1);
+  mm0 = _mm_srli_pi32(mm0, 23);
+  mm1 = _mm_srli_pi32(mm1, 23);
+#else
+  emm0 = _mm_srli_epi32(_mm_castps_si128(x), 23);
+#endif
+  /* keep only the fractional part */
+  x = _mm_and_ps(x, *(v4sf*)_ps_inv_mant_mask);
+  x = _mm_or_ps(x, *(v4sf*)_ps_0p5);
+
+//#ifndef USE_SSE2
+#if 0
+  /* now e=mm0:mm1 contain the really base-2 exponent */
+  mm0 = _mm_sub_pi32(mm0, *(v2si*)_pi32_0x7f);
+  mm1 = _mm_sub_pi32(mm1, *(v2si*)_pi32_0x7f);
+  v4sf e = _mm_cvtpi32x2_ps(mm0, mm1);
+  _mm_empty(); /* bye bye mmx */
+#else
+  emm0 = _mm_sub_epi32(emm0, *(v4si*)_pi32_0x7f);
+  v4sf e = _mm_cvtepi32_ps(emm0);
+#endif
+
+  e = _mm_add_ps(e, one);
+
+  /* part2: 
+     if( x < SQRTHF ) {
+       e -= 1;
+       x = x + x - 1.0;
+     } else { x = x - 1.0; }
+  */
+  v4sf mask = _mm_cmplt_ps(x, *(v4sf*)_ps_cephes_SQRTHF);
+  v4sf tmp = _mm_and_ps(x, mask);
+  x = _mm_sub_ps(x, one);
+  e = _mm_sub_ps(e, _mm_and_ps(one, mask));
+  x = _mm_add_ps(x, tmp);
+
+
+  v4sf z = _mm_mul_ps(x,x);
+
+  v4sf y = *(v4sf*)_ps_cephes_log_p0;
+  y = _mm_mul_ps(y, x);
+  y = _mm_add_ps(y, *(v4sf*)_ps_cephes_log_p1);
+  y = _mm_mul_ps(y, x);
+  y = _mm_add_ps(y, *(v4sf*)_ps_cephes_log_p2);
+  y = _mm_mul_ps(y, x);
+  y = _mm_add_ps(y, *(v4sf*)_ps_cephes_log_p3);
+  y = _mm_mul_ps(y, x);
+  y = _mm_add_ps(y, *(v4sf*)_ps_cephes_log_p4);
+  y = _mm_mul_ps(y, x);
+  y = _mm_add_ps(y, *(v4sf*)_ps_cephes_log_p5);
+  y = _mm_mul_ps(y, x);
+  y = _mm_add_ps(y, *(v4sf*)_ps_cephes_log_p6);
+  y = _mm_mul_ps(y, x);
+  y = _mm_add_ps(y, *(v4sf*)_ps_cephes_log_p7);
+  y = _mm_mul_ps(y, x);
+  y = _mm_add_ps(y, *(v4sf*)_ps_cephes_log_p8);
+  y = _mm_mul_ps(y, x);
+
+  y = _mm_mul_ps(y, z);
+  
+
+  tmp = _mm_mul_ps(e, *(v4sf*)_ps_cephes_log_q1);
+  y = _mm_add_ps(y, tmp);
+
+
+  tmp = _mm_mul_ps(z, *(v4sf*)_ps_0p5);
+  y = _mm_sub_ps(y, tmp);
+
+  tmp = _mm_mul_ps(e, *(v4sf*)_ps_cephes_log_q2);
+  x = _mm_add_ps(x, y);
+  x = _mm_add_ps(x, tmp);
+  x = _mm_or_ps(x, invalid_mask); // negative arg will be NAN
+  return x;
+}
+
+_PS_CONST(exp_hi,	88.3762626647949f);
+_PS_CONST(exp_lo,	-88.3762626647949f);
+
+_PS_CONST(cephes_LOG2EF, 1.44269504088896341f);
+_PS_CONST(cephes_exp_C1, 0.693359375f);
+_PS_CONST(cephes_exp_C2, -2.12194440e-4f);
+
+_PS_CONST(cephes_exp_p0, 1.9875691500E-4f);
+_PS_CONST(cephes_exp_p1, 1.3981999507E-3f);
+_PS_CONST(cephes_exp_p2, 8.3334519073E-3f);
+_PS_CONST(cephes_exp_p3, 4.1665795894E-2f);
+_PS_CONST(cephes_exp_p4, 1.6666665459E-1f);
+_PS_CONST(cephes_exp_p5, 5.0000001201E-1f);
+
+inline v4sf SseMath_exp_ps(v4sf x) {
+  v4sf tmp = _mm_setzero_ps(), fx;
+//#ifdef USE_SSE2
+#if 1
+  v4si emm0;
+#else
+  v2si mm0, mm1;
+#endif
+  v4sf one = *(v4sf*)_ps_1;
+
+  x = _mm_min_ps(x, *(v4sf*)_ps_exp_hi);
+  x = _mm_max_ps(x, *(v4sf*)_ps_exp_lo);
+
+  /* express exp(x) as exp(g + n*log(2)) */
+  fx = _mm_mul_ps(x, *(v4sf*)_ps_cephes_LOG2EF);
+  fx = _mm_add_ps(fx, *(v4sf*)_ps_0p5);
+
+  /* how to perform a floorf with SSE: just below */
+//#ifndef USE_SSE2
+#if 0
+  /* step 1 : cast to int */
+  tmp = _mm_movehl_ps(tmp, fx);
+  mm0 = _mm_cvttps_pi32(fx);
+  mm1 = _mm_cvttps_pi32(tmp);
+  /* step 2 : cast back to float */
+  tmp = _mm_cvtpi32x2_ps(mm0, mm1);
+#else
+  emm0 = _mm_cvttps_epi32(fx);
+  tmp  = _mm_cvtepi32_ps(emm0);
+#endif
+  /* if greater, substract 1 */
+  v4sf mask = _mm_cmpgt_ps(tmp, fx);    
+  mask = _mm_and_ps(mask, one);
+  fx = _mm_sub_ps(tmp, mask);
+
+  tmp = _mm_mul_ps(fx, *(v4sf*)_ps_cephes_exp_C1);
+  v4sf z = _mm_mul_ps(fx, *(v4sf*)_ps_cephes_exp_C2);
+  x = _mm_sub_ps(x, tmp);
+  x = _mm_sub_ps(x, z);
+
+  z = _mm_mul_ps(x,x);
+  
+  v4sf y = *(v4sf*)_ps_cephes_exp_p0;
+  y = _mm_mul_ps(y, x);
+  y = _mm_add_ps(y, *(v4sf*)_ps_cephes_exp_p1);
+  y = _mm_mul_ps(y, x);
+  y = _mm_add_ps(y, *(v4sf*)_ps_cephes_exp_p2);
+  y = _mm_mul_ps(y, x);
+  y = _mm_add_ps(y, *(v4sf*)_ps_cephes_exp_p3);
+  y = _mm_mul_ps(y, x);
+  y = _mm_add_ps(y, *(v4sf*)_ps_cephes_exp_p4);
+  y = _mm_mul_ps(y, x);
+  y = _mm_add_ps(y, *(v4sf*)_ps_cephes_exp_p5);
+  y = _mm_mul_ps(y, z);
+  y = _mm_add_ps(y, x);
+  y = _mm_add_ps(y, one);
+
+  /* build 2^n */
+//#ifndef USE_SSE2
+#if 0
+  z = _mm_movehl_ps(z, fx);
+  mm0 = _mm_cvttps_pi32(fx);
+  mm1 = _mm_cvttps_pi32(z);
+  mm0 = _mm_add_pi32(mm0, *(v2si*)_pi32_0x7f);
+  mm1 = _mm_add_pi32(mm1, *(v2si*)_pi32_0x7f);
+  mm0 = _mm_slli_pi32(mm0, 23); 
+  mm1 = _mm_slli_pi32(mm1, 23);
+  
+  v4sf pow2n; 
+  COPY_MM_TO_XMM(mm0, mm1, pow2n);
+  _mm_empty();
+#else
+  emm0 = _mm_cvttps_epi32(fx);
+  emm0 = _mm_add_epi32(emm0, *(v4si*)_pi32_0x7f);
+  emm0 = _mm_slli_epi32(emm0, 23);
+  v4sf pow2n = _mm_castsi128_ps(emm0);
+#endif
+  y = _mm_mul_ps(y, pow2n);
+  return y;
+}
+
+_PS_CONST(minus_cephes_DP1, -0.78515625f);
+_PS_CONST(minus_cephes_DP2, -2.4187564849853515625e-4f);
+_PS_CONST(minus_cephes_DP3, -3.77489497744594108e-8f);
+_PS_CONST(sincof_p0, -1.9515295891E-4f);
+_PS_CONST(sincof_p1,  8.3321608736E-3f);
+_PS_CONST(sincof_p2, -1.6666654611E-1f);
+_PS_CONST(coscof_p0,  2.443315711809948E-005f);
+_PS_CONST(coscof_p1, -1.388731625493765E-003f);
+_PS_CONST(coscof_p2,  4.166664568298827E-002f);
+_PS_CONST(cephes_FOPI, 1.27323954473516f); // 4 / M_PI
+
+
+/* evaluation of 4 sines at onces, using only SSE1+MMX intrinsics so
+   it runs also on old athlons XPs and the pentium III of your grand
+   mother.
+
+   The code is the exact rewriting of the cephes sinf function.
+   Precision is excellent as long as x < 8192 (I did not bother to
+   take into account the special handling they have for greater values
+   -- it does not return garbage for arguments over 8192, though, but
+   the extra precision is missing).
+
+   Note that it is such that sinf((float)M_PI) = 8.74e-8, which is the
+   surprising but correct result.
+
+   Performance is also surprisingly good, 1.33 times faster than the
+   macos vsinf SSE2 function, and 1.5 times faster than the
+   __vrs4_sinf of amd's ACML (which is only available in 64 bits). Not
+   too bad for an SSE1 function (with no special tuning) !
+   However the latter libraries probably have a much better handling of NaN,
+   Inf, denormalized and other special arguments..
+
+   On my core 1 duo, the execution of this function takes approximately 95 cycles.
+
+   From what I have observed on the experiments with Intel AMath lib, switching to an
+   SSE2 version would improve the perf by only 10%.
+
+   Since it is based on SSE intrinsics, it has to be compiled at -O2 to
+   deliver full speed.
+*/
+inline v4sf SseMath_sin_ps(v4sf x) { // any x
+  v4sf xmm1, xmm2 = _mm_setzero_ps(), xmm3, sign_bit, y;
+
+//#ifdef USE_SSE2
+#if 1
+  v4si emm0, emm2;
+#else
+  v2si mm0, mm1, mm2, mm3;
+#endif
+  sign_bit = x;
+  /* take the absolute value */
+  x = _mm_and_ps(x, *(v4sf*)_ps_inv_sign_mask);
+  /* extract the sign bit (upper one) */
+  sign_bit = _mm_and_ps(sign_bit, *(v4sf*)_ps_sign_mask);
+  
+  /* scale by 4/Pi */
+  y = _mm_mul_ps(x, *(v4sf*)_ps_cephes_FOPI);
+
+//#ifdef USE_SSE2
+#if 1
+  /* store the integer part of y in mm0 */
+  emm2 = _mm_cvttps_epi32(y);
+  /* j=(j+1) & (~1) (see the cephes sources) */
+  emm2 = _mm_add_epi32(emm2, *(v4si*)_pi32_1);
+  emm2 = _mm_and_si128(emm2, *(v4si*)_pi32_inv1);
+  y = _mm_cvtepi32_ps(emm2);
+
+  /* get the swap sign flag */
+  emm0 = _mm_and_si128(emm2, *(v4si*)_pi32_4);
+  emm0 = _mm_slli_epi32(emm0, 29);
+  /* get the polynom selection mask 
+     there is one polynom for 0 <= x <= Pi/4
+     and another one for Pi/4<x<=Pi/2
+
+     Both branches will be computed.
+  */
+  emm2 = _mm_and_si128(emm2, *(v4si*)_pi32_2);
+  emm2 = _mm_cmpeq_epi32(emm2, _mm_setzero_si128());
+  
+  v4sf swap_sign_bit = _mm_castsi128_ps(emm0);
+  v4sf poly_mask = _mm_castsi128_ps(emm2);
+  sign_bit = _mm_xor_ps(sign_bit, swap_sign_bit);
+  
+#else
+  /* store the integer part of y in mm0:mm1 */
+  xmm2 = _mm_movehl_ps(xmm2, y);
+  mm2 = _mm_cvttps_pi32(y);
+  mm3 = _mm_cvttps_pi32(xmm2);
+  /* j=(j+1) & (~1) (see the cephes sources) */
+  mm2 = _mm_add_pi32(mm2, *(v2si*)_pi32_1);
+  mm3 = _mm_add_pi32(mm3, *(v2si*)_pi32_1);
+  mm2 = _mm_and_si64(mm2, *(v2si*)_pi32_inv1);
+  mm3 = _mm_and_si64(mm3, *(v2si*)_pi32_inv1);
+  y = _mm_cvtpi32x2_ps(mm2, mm3);
+  /* get the swap sign flag */
+  mm0 = _mm_and_si64(mm2, *(v2si*)_pi32_4);
+  mm1 = _mm_and_si64(mm3, *(v2si*)_pi32_4);
+  mm0 = _mm_slli_pi32(mm0, 29);
+  mm1 = _mm_slli_pi32(mm1, 29);
+  /* get the polynom selection mask */
+  mm2 = _mm_and_si64(mm2, *(v2si*)_pi32_2);
+  mm3 = _mm_and_si64(mm3, *(v2si*)_pi32_2);
+  mm2 = _mm_cmpeq_pi32(mm2, _mm_setzero_si64());
+  mm3 = _mm_cmpeq_pi32(mm3, _mm_setzero_si64());
+  v4sf swap_sign_bit, poly_mask;
+  COPY_MM_TO_XMM(mm0, mm1, swap_sign_bit);
+  COPY_MM_TO_XMM(mm2, mm3, poly_mask);
+  sign_bit = _mm_xor_ps(sign_bit, swap_sign_bit);
+  _mm_empty(); /* good-bye mmx */
+#endif
+  
+  /* The magic pass: "Extended precision modular arithmetic" 
+     x = ((x - y * DP1) - y * DP2) - y * DP3; */
+  xmm1 = *(v4sf*)_ps_minus_cephes_DP1;
+  xmm2 = *(v4sf*)_ps_minus_cephes_DP2;
+  xmm3 = *(v4sf*)_ps_minus_cephes_DP3;
+  xmm1 = _mm_mul_ps(y, xmm1);
+  xmm2 = _mm_mul_ps(y, xmm2);
+  xmm3 = _mm_mul_ps(y, xmm3);
+  x = _mm_add_ps(x, xmm1);
+  x = _mm_add_ps(x, xmm2);
+  x = _mm_add_ps(x, xmm3);
+
+  /* Evaluate the first polynom  (0 <= x <= Pi/4) */
+  y = *(v4sf*)_ps_coscof_p0;
+  v4sf z = _mm_mul_ps(x,x);
+
+  y = _mm_mul_ps(y, z);
+  y = _mm_add_ps(y, *(v4sf*)_ps_coscof_p1);
+  y = _mm_mul_ps(y, z);
+  y = _mm_add_ps(y, *(v4sf*)_ps_coscof_p2);
+  y = _mm_mul_ps(y, z);
+  y = _mm_mul_ps(y, z);
+  v4sf tmp = _mm_mul_ps(z, *(v4sf*)_ps_0p5);
+  y = _mm_sub_ps(y, tmp);
+  y = _mm_add_ps(y, *(v4sf*)_ps_1);
+  
+  /* Evaluate the second polynom  (Pi/4 <= x <= 0) */
+
+  v4sf y2 = *(v4sf*)_ps_sincof_p0;
+  y2 = _mm_mul_ps(y2, z);
+  y2 = _mm_add_ps(y2, *(v4sf*)_ps_sincof_p1);
+  y2 = _mm_mul_ps(y2, z);
+  y2 = _mm_add_ps(y2, *(v4sf*)_ps_sincof_p2);
+  y2 = _mm_mul_ps(y2, z);
+  y2 = _mm_mul_ps(y2, x);
+  y2 = _mm_add_ps(y2, x);
+
+  /* select the correct result from the two polynoms */  
+  xmm3 = poly_mask;
+  y2 = _mm_and_ps(xmm3, y2); //, xmm3);
+  y = _mm_andnot_ps(xmm3, y);
+  y = _mm_add_ps(y,y2);
+  /* update the sign */
+  y = _mm_xor_ps(y, sign_bit);
+  return y;
+}
+
+/* almost the same as sin_ps */
+inline v4sf SseMath_cos_ps(v4sf x) { // any x
+  v4sf xmm1, xmm2 = _mm_setzero_ps(), xmm3, y;
+//#ifdef USE_SSE2
+#if 1
+  v4si emm0, emm2;
+#else
+  v2si mm0, mm1, mm2, mm3;
+#endif
+  /* take the absolute value */
+  x = _mm_and_ps(x, *(v4sf*)_ps_inv_sign_mask);
+  
+  /* scale by 4/Pi */
+  y = _mm_mul_ps(x, *(v4sf*)_ps_cephes_FOPI);
+  
+//#ifdef USE_SSE2
+#if 1
+  /* store the integer part of y in mm0 */
+  emm2 = _mm_cvttps_epi32(y);
+  /* j=(j+1) & (~1) (see the cephes sources) */
+  emm2 = _mm_add_epi32(emm2, *(v4si*)_pi32_1);
+  emm2 = _mm_and_si128(emm2, *(v4si*)_pi32_inv1);
+  y = _mm_cvtepi32_ps(emm2);
+
+  emm2 = _mm_sub_epi32(emm2, *(v4si*)_pi32_2);
+  
+  /* get the swap sign flag */
+  emm0 = _mm_andnot_si128(emm2, *(v4si*)_pi32_4);
+  emm0 = _mm_slli_epi32(emm0, 29);
+  /* get the polynom selection mask */
+  emm2 = _mm_and_si128(emm2, *(v4si*)_pi32_2);
+  emm2 = _mm_cmpeq_epi32(emm2, _mm_setzero_si128());
+  
+  v4sf sign_bit = _mm_castsi128_ps(emm0);
+  v4sf poly_mask = _mm_castsi128_ps(emm2);
+#else
+  /* store the integer part of y in mm0:mm1 */
+  xmm2 = _mm_movehl_ps(xmm2, y);
+  mm2 = _mm_cvttps_pi32(y);
+  mm3 = _mm_cvttps_pi32(xmm2);
+
+  /* j=(j+1) & (~1) (see the cephes sources) */
+  mm2 = _mm_add_pi32(mm2, *(v2si*)_pi32_1);
+  mm3 = _mm_add_pi32(mm3, *(v2si*)_pi32_1);
+  mm2 = _mm_and_si64(mm2, *(v2si*)_pi32_inv1);
+  mm3 = _mm_and_si64(mm3, *(v2si*)_pi32_inv1);
+
+  y = _mm_cvtpi32x2_ps(mm2, mm3);
+
+
+  mm2 = _mm_sub_pi32(mm2, *(v2si*)_pi32_2);
+  mm3 = _mm_sub_pi32(mm3, *(v2si*)_pi32_2);
+
+  /* get the swap sign flag in mm0:mm1 and the 
+     polynom selection mask in mm2:mm3 */
+
+  mm0 = _mm_andnot_si64(mm2, *(v2si*)_pi32_4);
+  mm1 = _mm_andnot_si64(mm3, *(v2si*)_pi32_4);
+  mm0 = _mm_slli_pi32(mm0, 29);
+  mm1 = _mm_slli_pi32(mm1, 29);
+
+  mm2 = _mm_and_si64(mm2, *(v2si*)_pi32_2);
+  mm3 = _mm_and_si64(mm3, *(v2si*)_pi32_2);
+
+  mm2 = _mm_cmpeq_pi32(mm2, _mm_setzero_si64());
+  mm3 = _mm_cmpeq_pi32(mm3, _mm_setzero_si64());
+
+  v4sf sign_bit, poly_mask;
+  COPY_MM_TO_XMM(mm0, mm1, sign_bit);
+  COPY_MM_TO_XMM(mm2, mm3, poly_mask);
+  _mm_empty(); /* good-bye mmx */
+#endif
+  /* The magic pass: "Extended precision modular arithmetic" 
+     x = ((x - y * DP1) - y * DP2) - y * DP3; */
+  xmm1 = *(v4sf*)_ps_minus_cephes_DP1;
+  xmm2 = *(v4sf*)_ps_minus_cephes_DP2;
+  xmm3 = *(v4sf*)_ps_minus_cephes_DP3;
+  xmm1 = _mm_mul_ps(y, xmm1);
+  xmm2 = _mm_mul_ps(y, xmm2);
+  xmm3 = _mm_mul_ps(y, xmm3);
+  x = _mm_add_ps(x, xmm1);
+  x = _mm_add_ps(x, xmm2);
+  x = _mm_add_ps(x, xmm3);
+  
+  /* Evaluate the first polynom  (0 <= x <= Pi/4) */
+  y = *(v4sf*)_ps_coscof_p0;
+  v4sf z = _mm_mul_ps(x,x);
+
+  y = _mm_mul_ps(y, z);
+  y = _mm_add_ps(y, *(v4sf*)_ps_coscof_p1);
+  y = _mm_mul_ps(y, z);
+  y = _mm_add_ps(y, *(v4sf*)_ps_coscof_p2);
+  y = _mm_mul_ps(y, z);
+  y = _mm_mul_ps(y, z);
+  v4sf tmp = _mm_mul_ps(z, *(v4sf*)_ps_0p5);
+  y = _mm_sub_ps(y, tmp);
+  y = _mm_add_ps(y, *(v4sf*)_ps_1);
+  
+  /* Evaluate the second polynom  (Pi/4 <= x <= 0) */
+
+  v4sf y2 = *(v4sf*)_ps_sincof_p0;
+  y2 = _mm_mul_ps(y2, z);
+  y2 = _mm_add_ps(y2, *(v4sf*)_ps_sincof_p1);
+  y2 = _mm_mul_ps(y2, z);
+  y2 = _mm_add_ps(y2, *(v4sf*)_ps_sincof_p2);
+  y2 = _mm_mul_ps(y2, z);
+  y2 = _mm_mul_ps(y2, x);
+  y2 = _mm_add_ps(y2, x);
+
+  /* select the correct result from the two polynoms */  
+  xmm3 = poly_mask;
+  y2 = _mm_and_ps(xmm3, y2); //, xmm3);
+  y = _mm_andnot_ps(xmm3, y);
+  y = _mm_add_ps(y,y2);
+  /* update the sign */
+  y = _mm_xor_ps(y, sign_bit);
+
+  return y;
+}
+
+/* since sin_ps and cos_ps are almost identical, sincos_ps could replace both of them..
+   it is almost as fast, and gives you a free cosine with your sine */
+inline void SseMath_sincos_ps(v4sf x, v4sf *s, v4sf *c) {
+  v4sf xmm1, xmm2, xmm3 = _mm_setzero_ps(), sign_bit_sin, y;
+//#ifdef USE_SSE2
+#if 1
+  v4si emm0, emm2, emm4;
+#else
+  v2si mm0, mm1, mm2, mm3, mm4, mm5;
+#endif
+  sign_bit_sin = x;
+  /* take the absolute value */
+  x = _mm_and_ps(x, *(v4sf*)_ps_inv_sign_mask);
+  /* extract the sign bit (upper one) */
+  sign_bit_sin = _mm_and_ps(sign_bit_sin, *(v4sf*)_ps_sign_mask);
+  
+  /* scale by 4/Pi */
+  y = _mm_mul_ps(x, *(v4sf*)_ps_cephes_FOPI);
+    
+//#ifdef USE_SSE2
+#if 1
+  /* store the integer part of y in emm2 */
+  emm2 = _mm_cvttps_epi32(y);
+
+  /* j=(j+1) & (~1) (see the cephes sources) */
+  emm2 = _mm_add_epi32(emm2, *(v4si*)_pi32_1);
+  emm2 = _mm_and_si128(emm2, *(v4si*)_pi32_inv1);
+  y = _mm_cvtepi32_ps(emm2);
+
+  emm4 = emm2;
+
+  /* get the swap sign flag for the sine */
+  emm0 = _mm_and_si128(emm2, *(v4si*)_pi32_4);
+  emm0 = _mm_slli_epi32(emm0, 29);
+  v4sf swap_sign_bit_sin = _mm_castsi128_ps(emm0);
+
+  /* get the polynom selection mask for the sine*/
+  emm2 = _mm_and_si128(emm2, *(v4si*)_pi32_2);
+  emm2 = _mm_cmpeq_epi32(emm2, _mm_setzero_si128());
+  v4sf poly_mask = _mm_castsi128_ps(emm2);
+#else
+  /* store the integer part of y in mm2:mm3 */
+  xmm3 = _mm_movehl_ps(xmm3, y);
+  mm2 = _mm_cvttps_pi32(y);
+  mm3 = _mm_cvttps_pi32(xmm3);
+
+  /* j=(j+1) & (~1) (see the cephes sources) */
+  mm2 = _mm_add_pi32(mm2, *(v2si*)_pi32_1);
+  mm3 = _mm_add_pi32(mm3, *(v2si*)_pi32_1);
+  mm2 = _mm_and_si64(mm2, *(v2si*)_pi32_inv1);
+  mm3 = _mm_and_si64(mm3, *(v2si*)_pi32_inv1);
+
+  y = _mm_cvtpi32x2_ps(mm2, mm3);
+
+  mm4 = mm2;
+  mm5 = mm3;
+
+  /* get the swap sign flag for the sine */
+  mm0 = _mm_and_si64(mm2, *(v2si*)_pi32_4);
+  mm1 = _mm_and_si64(mm3, *(v2si*)_pi32_4);
+  mm0 = _mm_slli_pi32(mm0, 29);
+  mm1 = _mm_slli_pi32(mm1, 29);
+  v4sf swap_sign_bit_sin;
+  COPY_MM_TO_XMM(mm0, mm1, swap_sign_bit_sin);
+
+  /* get the polynom selection mask for the sine */
+
+  mm2 = _mm_and_si64(mm2, *(v2si*)_pi32_2);
+  mm3 = _mm_and_si64(mm3, *(v2si*)_pi32_2);
+  mm2 = _mm_cmpeq_pi32(mm2, _mm_setzero_si64());
+  mm3 = _mm_cmpeq_pi32(mm3, _mm_setzero_si64());
+  v4sf poly_mask;
+  COPY_MM_TO_XMM(mm2, mm3, poly_mask);
+#endif
+
+  /* The magic pass: "Extended precision modular arithmetic" 
+     x = ((x - y * DP1) - y * DP2) - y * DP3; */
+  xmm1 = *(v4sf*)_ps_minus_cephes_DP1;
+  xmm2 = *(v4sf*)_ps_minus_cephes_DP2;
+  xmm3 = *(v4sf*)_ps_minus_cephes_DP3;
+  xmm1 = _mm_mul_ps(y, xmm1);
+  xmm2 = _mm_mul_ps(y, xmm2);
+  xmm3 = _mm_mul_ps(y, xmm3);
+  x = _mm_add_ps(x, xmm1);
+  x = _mm_add_ps(x, xmm2);
+  x = _mm_add_ps(x, xmm3);
+
+//#ifdef USE_SSE2
+#if 1
+  emm4 = _mm_sub_epi32(emm4, *(v4si*)_pi32_2);
+  emm4 = _mm_andnot_si128(emm4, *(v4si*)_pi32_4);
+  emm4 = _mm_slli_epi32(emm4, 29);
+  v4sf sign_bit_cos = _mm_castsi128_ps(emm4);
+#else
+  /* get the sign flag for the cosine */
+  mm4 = _mm_sub_pi32(mm4, *(v2si*)_pi32_2);
+  mm5 = _mm_sub_pi32(mm5, *(v2si*)_pi32_2);
+  mm4 = _mm_andnot_si64(mm4, *(v2si*)_pi32_4);
+  mm5 = _mm_andnot_si64(mm5, *(v2si*)_pi32_4);
+  mm4 = _mm_slli_pi32(mm4, 29);
+  mm5 = _mm_slli_pi32(mm5, 29);
+  v4sf sign_bit_cos;
+  COPY_MM_TO_XMM(mm4, mm5, sign_bit_cos);
+  _mm_empty(); /* good-bye mmx */
+#endif
+
+  sign_bit_sin = _mm_xor_ps(sign_bit_sin, swap_sign_bit_sin);
+
+  
+  /* Evaluate the first polynom  (0 <= x <= Pi/4) */
+  v4sf z = _mm_mul_ps(x,x);
+  y = *(v4sf*)_ps_coscof_p0;
+
+  y = _mm_mul_ps(y, z);
+  y = _mm_add_ps(y, *(v4sf*)_ps_coscof_p1);
+  y = _mm_mul_ps(y, z);
+  y = _mm_add_ps(y, *(v4sf*)_ps_coscof_p2);
+  y = _mm_mul_ps(y, z);
+  y = _mm_mul_ps(y, z);
+  v4sf tmp = _mm_mul_ps(z, *(v4sf*)_ps_0p5);
+  y = _mm_sub_ps(y, tmp);
+  y = _mm_add_ps(y, *(v4sf*)_ps_1);
+  
+  /* Evaluate the second polynom  (Pi/4 <= x <= 0) */
+
+  v4sf y2 = *(v4sf*)_ps_sincof_p0;
+  y2 = _mm_mul_ps(y2, z);
+  y2 = _mm_add_ps(y2, *(v4sf*)_ps_sincof_p1);
+  y2 = _mm_mul_ps(y2, z);
+  y2 = _mm_add_ps(y2, *(v4sf*)_ps_sincof_p2);
+  y2 = _mm_mul_ps(y2, z);
+  y2 = _mm_mul_ps(y2, x);
+  y2 = _mm_add_ps(y2, x);
+
+  /* select the correct result from the two polynoms */  
+  xmm3 = poly_mask;
+  v4sf ysin2 = _mm_and_ps(xmm3, y2);
+  v4sf ysin1 = _mm_andnot_ps(xmm3, y);
+  y2 = _mm_sub_ps(y2,ysin2);
+  y = _mm_sub_ps(y, ysin1);
+
+  xmm1 = _mm_add_ps(ysin1,ysin2);
+  xmm2 = _mm_add_ps(y,y2);
+ 
+  /* update the sign */
+  *s = _mm_xor_ps(xmm1, sign_bit_sin);
+  *c = _mm_xor_ps(xmm2, sign_bit_cos);
+}
+
