Fix config file threading related crash on startup

diff --git a/Engine/Source/Runtime/Core/Private/Misc/ConfigCacheIni.cpp b/Engine/Source/Runtime/Core/Private/Misc/ConfigCacheIni.cpp
--- a/Engine/Source/Runtime/Core/Private/Misc/ConfigCacheIni.cpp
+++ b/Engine/Source/Runtime/Core/Private/Misc/ConfigCacheIni.cpp
@@ -42,6 +42,16 @@
 #define DISABLE_GENERATED_INI_WHEN_COOKED 0
 #endif
 
+namespace UE
+{
+namespace ConfigCacheIni
+{
+namespace Private
+{
+	static FCriticalSection ExpansionsCriticalSection;
+}
+}
+}?
 const TSet<FString>* FConfigCacheIni::IniCacheSet = nullptr;
 
 namespace
@@ -3888,14 +3898,13 @@ static FString PerformExpansionReplacements(const FConfigLayerExpansion& Expansi
 
 static FString PerformFinalExpansions(const FString InString, const FString& PlatformName, const TCHAR* EngineDir, const TCHAR* ProjectDir)
 {
+	FScopeLock Lock(&UE::ConfigCacheIni::Private::ExpansionsCriticalSection);
+
 	static FString LastPlatform;
 	static FString PlatformExtensionEngineDir;
 	static FString PlatformExtensionProjectDir;
 	static FString ProjectNotForLicenseesDir;
 	static FString ProjectNoRedistDir;
-	static FCriticalSection ExpansionSection;
-
-	FScopeLock Lock(&ExpansionSection);
 
 	if (LastPlatform != PlatformName)
 	{
