#ue5 - Fix a few build issues from 18984296 missed by preflight. Revert unneeded RandRange() overloads which can cause ambiguity.

diff --git a/Engine/Source/Runtime/Core/Private/Misc/DateTime.cpp b/Engine/Source/Runtime/Core/Private/Misc/DateTime.cpp
--- a/Engine/Source/Runtime/Core/Private/Misc/DateTime.cpp
+++ b/Engine/Source/Runtime/Core/Private/Misc/DateTime.cpp
@@ -93,7 +93,7 @@ void FDateTime::GetDate(int32& OutYear, int32& OutMonth, int32& OutDay) const
 
 	int32 i, j, k, l, n;
 
-	l = (int32)(FMath::FloorToInt(GetDate().GetJulianDay() + 0.5) + 68569);
+	l = FMath::FloorToInt32((float)(GetJulianDay() + 0.5)) + 68569;
 	n = 4 * l / 146097;
 	l = l - (146097 * n + 3) / 4;
 	i = 4000 * (l + 1) / 1461001;
diff --git a/Engine/Source/Runtime/Core/Public/Math/UnrealMathUtility.h b/Engine/Source/Runtime/Core/Public/Math/UnrealMathUtility.h
--- a/Engine/Source/Runtime/Core/Public/Math/UnrealMathUtility.h
+++ b/Engine/Source/Runtime/Core/Public/Math/UnrealMathUtility.h
@@ -209,9 +209,6 @@ struct FMath : public FPlatformMath
 		return Min + RandHelper64(Range);
 	}
 
-	UE_NODISCARD static FORCEINLINE int64 RandRange(int32 Min, int64 Max) { return RandRange((int64)Min, Max); }
-	UE_NODISCARD static FORCEINLINE int64 RandRange(int64 Min, int32 Max) { return RandRange(Min, (int64)Max); }
-
 	/** Util to generate a random number in a range. Overloaded to distinguish from int32 version, where passing a float is typically a mistake. */
 	UE_NODISCARD static FORCEINLINE float RandRange(float InMin, float InMax)
 	{
