Avoid contention when calling Rebuild on up-to-date text

The revision data changes rarely, so a RW-lock allows multiple threads to concurrently call Rebuild on up-to-date text instances without incuring a synchronization penalty. Additionally Rebuild is now a no-op for culture invariant text (as it used to be prior to some recent changes to save memory).

diff --git a/Engine/Source/Runtime/Core/Private/Internationalization/TextHistory.cpp b/Engine/Source/Runtime/Core/Private/Internationalization/TextHistory.cpp
--- a/Engine/Source/Runtime/Core/Private/Internationalization/TextHistory.cpp
+++ b/Engine/Source/Runtime/Core/Private/Internationalization/TextHistory.cpp
@@ -682,16 +682,19 @@ void WriteTextFormatToBuffer(FString& Buffer, const FString& TokenMarker, const
 
 void FTextHistory::UpdateDisplayStringIfOutOfDate()
 {
-	FTextLocalizationManager& TLM = FTextLocalizationManager::Get();
-	const uint16 CurrentGlobalRevision = TLM.GetTextRevision();
-	const uint16 CurrentLocalRevision = TLM.GetLocalRevisionForTextId(GetTextId());
-
-	if (GlobalRevision != CurrentGlobalRevision || LocalRevision != CurrentLocalRevision)
+	if (CanUpdateDisplayString())
 	{
-		GlobalRevision = CurrentGlobalRevision;
-		LocalRevision = CurrentLocalRevision;
+		uint16 CurrentGlobalRevision = 0;
+		uint16 CurrentLocalRevision = 0;
+		FTextLocalizationManager::Get().GetTextRevisions(GetTextId(), CurrentGlobalRevision, CurrentLocalRevision);
+
+		if (GlobalRevision != CurrentGlobalRevision || LocalRevision != CurrentLocalRevision)
+		{
+			GlobalRevision = CurrentGlobalRevision;
+			LocalRevision = CurrentLocalRevision;
 
-		UpdateDisplayString();
+			UpdateDisplayString();
+		}
 	}
 }
 
@@ -703,9 +706,15 @@ void FTextHistory::MarkDisplayStringOutOfDate()
 
 void FTextHistory::MarkDisplayStringUpToDate()
 {
-	FTextLocalizationManager& TLM = FTextLocalizationManager::Get();
-	GlobalRevision = TLM.GetTextRevision();
-	LocalRevision = TLM.GetLocalRevisionForTextId(GetTextId());
+	if (CanUpdateDisplayString())
+	{
+		FTextLocalizationManager::Get().GetTextRevisions(GetTextId(), GlobalRevision, LocalRevision);
+	}
+	else
+	{
+		GlobalRevision = 0;
+		LocalRevision = 0;
+	}
 }
 
 ///////////////////////////////////////
@@ -822,6 +831,7 @@ void FTextHistory_Base::Serialize(FStructuredArchive::FRecord Record)
 #endif // WITH_EDITOR
 
 		TextId = FTextId(Namespace, Key);
+		LocalizedString.Reset();
 		MarkDisplayStringOutOfDate();
 	}
 	else if(BaseArchive.IsSaving())
@@ -885,16 +895,15 @@ void FTextHistory_Base::Serialize(FStructuredArchive::FRecord Record)
 	}
 }
 
+bool FTextHistory_Base::CanUpdateDisplayString()
+{
+	return !TextId.IsEmpty();
+}
+
 void FTextHistory_Base::UpdateDisplayString()
 {
-	if (TextId.IsEmpty())
-	{
-		LocalizedString.Reset();
-	}
-	else
-	{
-		LocalizedString = FTextLocalizationManager::Get().GetDisplayString(TextId.GetNamespace(), TextId.GetKey(), &SourceString);
-	}
+	check(!TextId.IsEmpty()); // CanUpdateDisplayString should prevent UpdateDisplayString being called
+	LocalizedString = FTextLocalizationManager::Get().GetDisplayString(TextId.GetNamespace(), TextId.GetKey(), &SourceString);
 }
 
 bool FTextHistory_Base::StaticShouldReadFromBuffer(const TCHAR* Buffer)
@@ -961,6 +970,7 @@ const TCHAR* FTextHistory_Base::ReadFromBuffer(const TCHAR* Buffer, const TCHAR*
 			TextNamespaceUtil::StripPackageNamespaceInline(NamespaceString);
 		}
 		TextId = FTextId(NamespaceString, KeyString);
+		LocalizedString.Reset();
 		MarkDisplayStringOutOfDate();
 
 		return Buffer;
@@ -1015,6 +1025,7 @@ const TCHAR* FTextHistory_Base::ReadFromBuffer(const TCHAR* Buffer, const TCHAR*
 			TextNamespaceUtil::StripPackageNamespaceInline(NamespaceString);
 		}
 		TextId = FTextId(NamespaceString, KeyString);
+		LocalizedString.Reset();
 		MarkDisplayStringOutOfDate();
 
 		return Buffer;
diff --git a/Engine/Source/Runtime/Core/Private/Internationalization/TextHistory.h b/Engine/Source/Runtime/Core/Private/Internationalization/TextHistory.h
--- a/Engine/Source/Runtime/Core/Private/Internationalization/TextHistory.h
+++ b/Engine/Source/Runtime/Core/Private/Internationalization/TextHistory.h
@@ -215,6 +215,9 @@ public:
 	void UpdateDisplayStringIfOutOfDate();
 
 protected:
+	/** True if "UpdateDisplayString" might do something if called, or False if it would be redundant */
+	virtual bool CanUpdateDisplayString() { return true; }
+
 	/** Update the display string when the history is out-of-date */
 	virtual void UpdateDisplayString() = 0;
 
@@ -254,6 +257,7 @@ public:
 	virtual FString BuildInvariantDisplayString() const override;
 	virtual bool IdenticalTo(const FTextHistory& Other, const ETextIdenticalModeFlags CompareModeFlags) const override;
 	virtual void Serialize(FStructuredArchive::FRecord Record) override;
+	virtual bool CanUpdateDisplayString() override;
 	virtual void UpdateDisplayString() override;
 	//~ End FTextHistory Interface
 
diff --git a/Engine/Source/Runtime/Core/Private/Internationalization/TextLocalizationManager.cpp b/Engine/Source/Runtime/Core/Private/Internationalization/TextLocalizationManager.cpp
--- a/Engine/Source/Runtime/Core/Private/Internationalization/TextLocalizationManager.cpp
+++ b/Engine/Source/Runtime/Core/Private/Internationalization/TextLocalizationManager.cpp
@@ -12,6 +12,7 @@
 #include "HAL/LowLevelMemTracker.h"
 #include "Misc/Parse.h"
 #include "Misc/ScopeLock.h"
+#include "Misc/ScopeRWLock.h"
 #include "Misc/CommandLine.h"
 #include "Misc/LazySingleton.h"
 #include "Internationalization/Culture.h"
@@ -483,11 +484,14 @@ FTextLocalizationManager::FTextLocalizationManager()
 
 void FTextLocalizationManager::DumpMemoryInfo() const
 {
-	FScopeLock ScopeLock(&SynchronizationObject);
-
-	UE_LOG(LogTextLocalizationManager, Log, TEXT("DisplayStringLookupTable.GetAllocatedSize()=%d elems=%d"), DisplayStringLookupTable.GetAllocatedSize(), DisplayStringLookupTable.Num());
-
-	UE_LOG(LogTextLocalizationManager, Log, TEXT("LocalTextRevisions.GetAllocatedSize()=%d elems=%d"), LocalTextRevisions.GetAllocatedSize(), LocalTextRevisions.Num());
+	{
+		FScopeLock ScopeLock(&DisplayStringLookupTableCS);
+		UE_LOG(LogTextLocalizationManager, Log, TEXT("DisplayStringLookupTable.GetAllocatedSize()=%d elems=%d"), DisplayStringLookupTable.GetAllocatedSize(), DisplayStringLookupTable.Num());
+	}
+	{
+		FReadScopeLock ScopeLock(TextRevisionRW);
+		UE_LOG(LogTextLocalizationManager, Log, TEXT("LocalTextRevisions.GetAllocatedSize()=%d elems=%d"), LocalTextRevisions.GetAllocatedSize(), LocalTextRevisions.Num());
+	}
 }
 
 void FTextLocalizationManager::CompactDataStructures()
@@ -495,10 +499,15 @@ void FTextLocalizationManager::CompactDataStructures()
 	TRACE_CPUPROFILER_EVENT_SCOPE(FTextLocalizationManager::CompactDataStructures);
 	LLM_SCOPE(ELLMTag::Localization);
 
-	FScopeLock ScopeLock(&SynchronizationObject);
 	double StartTime = FPlatformTime::Seconds();
-	DisplayStringLookupTable.Shrink();
-	LocalTextRevisions.Shrink();
+	{
+		FScopeLock ScopeLock(&DisplayStringLookupTableCS);
+		DisplayStringLookupTable.Shrink();
+	}
+	{
+		FWriteScopeLock ScopeLock(TextRevisionRW);
+		LocalTextRevisions.Shrink();
+	}
 	FTextKey::CompactDataStructures();
 	UE_LOG(LogTextLocalizationManager, Log, TEXT("Compacting localization data took %6.2fms"), 1000.0 * (FPlatformTime::Seconds() - StartTime));
 }
@@ -650,7 +659,7 @@ void FTextLocalizationManager::RegisterPolyglotTextData(TArrayView<const FPolygl
 
 FTextConstDisplayStringPtr FTextLocalizationManager::FindDisplayString(const FTextKey& Namespace, const FTextKey& Key, const FString* const SourceString) const
 {
-	FScopeLock ScopeLock( &SynchronizationObject );
+	FScopeLock ScopeLock(&DisplayStringLookupTableCS);
 
 	const FTextId TextId(Namespace, Key);
 
@@ -669,7 +678,7 @@ FTextConstDisplayStringRef FTextLocalizationManager::GetDisplayString(const FTex
 	TRACE_CPUPROFILER_EVENT_SCOPE(FTextLocalizationManager::GetDisplayString);
 	LLM_SCOPE(ELLMTag::Localization);
 
-	FScopeLock ScopeLock( &SynchronizationObject );
+	FScopeLock ScopeLock(&DisplayStringLookupTableCS);
 
 	// Hack fix for old assets that don't have namespace/key info.
 	if (Namespace.IsEmpty() && Key.IsEmpty())
@@ -805,7 +814,7 @@ FTextConstDisplayStringRef FTextLocalizationManager::GetDisplayString(const FTex
 #if WITH_EDITORONLY_DATA
 bool FTextLocalizationManager::GetLocResID(const FTextKey& Namespace, const FTextKey& Key, FString& OutLocResId) const
 {
-	FScopeLock ScopeLock(&SynchronizationObject);
+	FScopeLock ScopeLock(&DisplayStringLookupTableCS);
 
 	const FTextId TextId(Namespace, Key);
 
@@ -821,11 +830,17 @@ bool FTextLocalizationManager::GetLocResID(const FTextKey& Namespace, const FTex
 }
 #endif
 
+uint16 FTextLocalizationManager::GetTextRevision() const
+{
+	FReadScopeLock ScopeLock(TextRevisionRW);
+	return TextRevisionCounter;
+}
+
 uint16 FTextLocalizationManager::GetLocalRevisionForTextId(const FTextId& InTextId) const
 {
 	if (!InTextId.IsEmpty())
 	{
-		FScopeLock ScopeLock(&SynchronizationObject);
+		FReadScopeLock ScopeLock(TextRevisionRW);
 		if (const uint16* FoundLocalRevision = LocalTextRevisions.Find(InTextId))
 		{
 			return *FoundLocalRevision;
@@ -834,12 +849,27 @@ uint16 FTextLocalizationManager::GetLocalRevisionForTextId(const FTextId& InText
 	return 0;
 }
 
+void FTextLocalizationManager::GetTextRevisions(const FTextId& InTextId, uint16& OutGlobalTextRevision, uint16& OutLocalTextRevision) const
+{
+	FReadScopeLock ScopeLock(TextRevisionRW);
+
+	OutGlobalTextRevision = TextRevisionCounter;
+	if (const uint16* FoundLocalRevision = InTextId.IsEmpty() ? nullptr : LocalTextRevisions.Find(InTextId))
+	{
+		OutLocalTextRevision = *FoundLocalRevision;
+	}
+	else
+	{
+		OutLocalTextRevision = 0;
+	}
+}
+
 bool FTextLocalizationManager::AddDisplayString(const FTextDisplayStringRef& DisplayString, const FTextKey& Namespace, const FTextKey& Key)
 {
 	TRACE_CPUPROFILER_EVENT_SCOPE(FTextLocalizationManager::AddDisplayString);
 	LLM_SCOPE(ELLMTag::Localization);
 
-	FScopeLock ScopeLock( &SynchronizationObject );
+	FScopeLock ScopeLock(&DisplayStringLookupTableCS);
 
 	const FTextId TextId(Namespace, Key);
 
@@ -1108,7 +1138,7 @@ void FTextLocalizationManager::LoadLocalizationResourcesForPrioritizedCultures(T
 	{
 		// Lock while updating the tables
 		{
-			FScopeLock ScopeLock(&SynchronizationObject);
+			FScopeLock ScopeLock(&DisplayStringLookupTableCS);
 
 			for (auto& DisplayStringPair : DisplayStringLookupTable)
 			{
@@ -1144,7 +1174,7 @@ void FTextLocalizationManager::UpdateFromNative(FTextLocalizationResource&& Text
 
 	// Lock while updating the tables
 	{
-		FScopeLock ScopeLock(&SynchronizationObject);
+		FScopeLock ScopeLock(&DisplayStringLookupTableCS);
 
 		DisplayStringLookupTable.Reserve(TextLocalizationResource.Entries.Num());
 
@@ -1253,7 +1283,7 @@ void FTextLocalizationManager::UpdateFromLocalizations(FTextLocalizationResource
 
 	// Lock while updating the tables
 	{
-		FScopeLock ScopeLock(&SynchronizationObject);
+		FScopeLock ScopeLock(&DisplayStringLookupTableCS);
 
 		DisplayStringLookupTable.Reserve(TextLocalizationResource.Entries.Num());
 
@@ -1364,7 +1394,7 @@ void FTextLocalizationManager::DirtyLocalRevisionForTextId(const FTextId& InText
 	TRACE_CPUPROFILER_EVENT_SCOPE(FTextLocalizationManager::DirtyLocalRevisionForTextId);
 	LLM_SCOPE(ELLMTag::Localization);
 
-	FScopeLock ScopeLock(&SynchronizationObject);
+	FWriteScopeLock ScopeLock(TextRevisionRW);
 
 	uint16* FoundLocalRevision = LocalTextRevisions.Find(InTextId);
 	if (FoundLocalRevision)
@@ -1384,7 +1414,7 @@ void FTextLocalizationManager::DirtyTextRevision()
 
 	// Lock while updating the data
 	{
-		FScopeLock ScopeLock(&SynchronizationObject);
+		FWriteScopeLock ScopeLock(TextRevisionRW);
 
 		while (++TextRevisionCounter == 0) {} // Zero is special, don't allow an overflow to stay at zero
 		LocalTextRevisions.Empty();
diff --git a/Engine/Source/Runtime/Core/Public/Internationalization/TextLocalizationManager.h b/Engine/Source/Runtime/Core/Public/Internationalization/TextLocalizationManager.h
--- a/Engine/Source/Runtime/Core/Public/Internationalization/TextLocalizationManager.h
+++ b/Engine/Source/Runtime/Core/Public/Internationalization/TextLocalizationManager.h
@@ -89,8 +89,10 @@ private:
 		return EnumHasAnyFlags(InitializedFlags, ETextLocalizationManagerInitializedFlags::Initializing);
 	}
 
-	mutable FCriticalSection SynchronizationObject;
+	mutable FCriticalSection DisplayStringLookupTableCS;
 	FDisplayStringLookupTable DisplayStringLookupTable;
+
+	mutable FRWLock TextRevisionRW;
 	TMap<FTextId, uint16> LocalTextRevisions;
 	uint16 TextRevisionCounter;
 
@@ -167,13 +169,6 @@ public:
 	UE_DEPRECATED(5.0, "FindNamespaceAndKeyFromDisplayString no longer functions! Use FTextInspector::GetTextId instead.")
 	bool FindNamespaceAndKeyFromDisplayString(const FTextConstDisplayStringPtr& InDisplayString, FTextKey& OutNamespace, FTextKey& OutKey) const { return false; }
 	
-	/**
-	 * Attempts to find a local revision history for the given text ID.
-	 * This will only be set if the display string has been changed since the localization manager version has been changed (eg, if it has been edited while keeping the same key).
-	 * @return The local revision, or 0 if there have been no changes since a global history change.
-	 */
-	uint16 GetLocalRevisionForTextId(const FTextId& InTextId) const;
-
 	/**	Attempts to register the specified display string, associating it with the specified namespace and key.
 	 *	Returns true if the display string has been or was already associated with the namespace and key.
 	 *	Returns false if the display string was already associated with another namespace and key or the namespace and key are already in use by another display string.
@@ -187,9 +182,24 @@ public:
 	/** Reloads resources for the current culture. */
 	void RefreshResources();
 
-	/**	Returns the current text revision number. This value can be cached when caching information from the text localization manager.
-	 *	If the revision does not match, cached information may be invalid and should be recached. */
-	uint16 GetTextRevision() const { return TextRevisionCounter; }
+	/**
+	 * Returns the current text revision number. This value can be cached when caching information from the text localization manager.
+	 * If the revision does not match, cached information may be invalid and should be recached.
+	 */
+	uint16 GetTextRevision() const;
+
+	/**
+	 * Attempts to find a local revision history for the given text ID.
+	 * This will only be set if the display string has been changed since the localization manager version has been changed (eg, if it has been edited while keeping the same key).
+	 * @return The local revision, or 0 if there have been no changes since a global history change.
+	 */
+	uint16 GetLocalRevisionForTextId(const FTextId& InTextId) const;
+
+	/**
+	 * Get both the global and local revision for the given text ID.
+	 * @see GetTextRevision and GetLocalRevisionForTextId.
+	 */
+	void GetTextRevisions(const FTextId& InTextId, uint16& OutGlobalTextRevision, uint16& OutLocalTextRevision) const;
 
 #if WITH_EDITOR
 	/**
