Fix for division by zero in vectorized TTransform<T>::InverseTransformPosition() and TTransform<T>::InverseTransformVector(), causing floating point exceptions.

diff --git a/Engine/Source/Runtime/Core/Public/Math/TransformVectorized.h b/Engine/Source/Runtime/Core/Public/Math/TransformVectorized.h
--- a/Engine/Source/Runtime/Core/Public/Math/TransformVectorized.h
+++ b/Engine/Source/Runtime/Core/Public/Math/TransformVectorized.h
@@ -1787,7 +1787,7 @@ FORCEINLINE TVector<T> TTransform<T>::InverseTransformPosition(const TVector<T>
 	const TransformVectorRegister VR = VectorQuaternionInverseRotateVector(Rotation, TranslatedVec);
 
 	// GetSafeScaleReciprocal(Scale3D);
-	const TransformVectorRegister SafeReciprocal = GetSafeScaleReciprocal(Scale3D);	
+	const TransformVectorRegister SafeReciprocal = GetSafeScaleReciprocal(VectorSet_W1(Scale3D));	
 
 	// ( Rotation.Inverse() * (V-Translation) ) * GetSafeScaleReciprocal(Scale3D);
 	const TransformVectorRegister VResult = VectorMultiply(VR, SafeReciprocal);
@@ -1829,7 +1829,7 @@ FORCEINLINE TVector<T> TTransform<T>::InverseTransformVector(const TVector<T> &V
 	const TransformVectorRegister VR = VectorQuaternionInverseRotateVector(Rotation, InputVector);
 
 	// GetSafeScaleReciprocal(Scale3D);
-	const TransformVectorRegister SafeReciprocal = GetSafeScaleReciprocal(Scale3D);
+	const TransformVectorRegister SafeReciprocal = GetSafeScaleReciprocal(VectorSet_W1(Scale3D));
 
 	// ( Rotation.Inverse() * V) * GetSafeScaleReciprocal(Scale3D);
 	const TransformVectorRegister VResult = VectorMultiply(VR, SafeReciprocal);
