OodleDataCompression OodleScratchBufferCount use one per core in tools
still defaults to 2 for game runtimes
significant speedup in iostore compression at fast compress levels

diff --git a/Engine/Source/Runtime/Core/Private/Compression/OodleDataCompression.cpp b/Engine/Source/Runtime/Core/Private/Compression/OodleDataCompression.cpp
--- a/Engine/Source/Runtime/Core/Private/Compression/OodleDataCompression.cpp
+++ b/Engine/Source/Runtime/Core/Private/Compression/OodleDataCompression.cpp
@@ -3,6 +3,7 @@
 #include "Compression/OodleDataCompression.h"
 #include "Misc/ConfigCacheIni.h"
 #include "Misc/ICompressionFormat.h"
+#include "HAL/PlatformMisc.h"
 
 #include "oodle2.h"
 
@@ -212,20 +213,41 @@ struct OodleScratchBuffers
 		// DecoderMemorySizeNeeded is ~ 460000 , EncoderMemorySizeNeeded is ~ 470000
 		OodleScratchMemorySize = DecoderMemorySizeNeeded > EncoderMemorySizeNeeded ? DecoderMemorySizeNeeded : EncoderMemorySizeNeeded;
 
-		int32 BufferCount = 2;
-
-		// be wary of possible init order problem
-		//  if we init Oodle before config might not exist yet?
-		//  can we just check GConfig vs nullptr ?
-		//  (eg. if Oodle is used to unpak ini filed?)
-		if ( GConfig )
+		int32 BufferCount;
+		
+		if ( FPlatformProperties::RequiresCookedData() )
 		{
-			GConfig->GetInt(TEXT("OodleDataCompressionFormat"), TEXT("PreallocatedBufferCount"), BufferCount, GEngineIni);
-			if (BufferCount < 0)
+			// runtime game or similar environment
+			
+			BufferCount = 2;
+
+			// "PreallocatedBufferCount" is not a great name
+			//	it's actually the max number of static buffers
+			//	 they will only be allocated as needed
+
+			// be wary of possible init order problem
+			//  if we init Oodle before config might not exist yet?
+			//  can we just check GConfig vs nullptr ?
+			//  (eg. if Oodle is used to unpak ini filed?)
+			if ( GConfig )
 			{
-				BufferCount = 0;
+				GConfig->GetInt(TEXT("OodleDataCompressionFormat"), TEXT("PreallocatedBufferCount"), BufferCount, GEngineIni);
+				if (BufferCount < 0)
+				{
+					// negative means one per core
+					BufferCount = FPlatformMisc::NumberOfCores();
+				}
 			}
 		}
+		else
+		{
+			// tools, like UnrealPak or DDC commandlets
+
+			// allow one scratch buffer per core
+			// they will only be allocated on demand if we actually reach that level of parallelism
+			//	 so commandlets that don't use parallel compression don't waste memory
+			BufferCount = FPlatformMisc::NumberOfCores();
+		}
 
 		OodleScratchBufferCount = BufferCount;
 		if (OodleScratchBufferCount)
