Fix buffer overrun

diff --git a/Engine/Source/Runtime/Core/Private/Tests/HAL/FileManagerGenericTest.cpp b/Engine/Source/Runtime/Core/Private/Tests/HAL/FileManagerGenericTest.cpp
--- a/Engine/Source/Runtime/Core/Private/Tests/HAL/FileManagerGenericTest.cpp
+++ b/Engine/Source/Runtime/Core/Private/Tests/HAL/FileManagerGenericTest.cpp
@@ -187,7 +187,7 @@ bool FArchiveFileReaderGenericTest::CreateTestFile()
 	int32 NumInts = 1024 * 128 + 17; // Make sure the file size is not a multiple of buffersize; we need to test that behavior
 	Builder.SetNumUninitialized(NumInts);
 	int32* Data = Builder.GetData();
-	for (int n = 0; n <= Builder.Num(); ++n)
+	for (int n = 0; n < Builder.Num(); ++n)
 	{
 		Data[n] = GetExpectedValue(n);
 	}
