REL - disable FMallocAnsi::GetAllocationSize for Windows because it's not properly supported by Microsoft.

diff --git a/Engine/Source/Runtime/Core/Private/HAL/MallocAnsi.cpp b/Engine/Source/Runtime/Core/Private/HAL/MallocAnsi.cpp
--- a/Engine/Source/Runtime/Core/Private/HAL/MallocAnsi.cpp
+++ b/Engine/Source/Runtime/Core/Private/HAL/MallocAnsi.cpp
@@ -221,8 +221,12 @@ bool FMallocAnsi::GetAllocationSize( void *Original, SIZE_T &SizeOut )
 		return false;
 	}
 
+#if	PLATFORM_USES__ALIGNED_MALLOC
+	return false;
+#else
 	SizeOut = AnsiGetAllocationSize(Original);
 	return true;
+#endif
 }
 
 bool FMallocAnsi::IsInternallyThreadSafe() const
