#ue5 - Resolve bad merge from 18997110, now preserving intended changes from 18841321.

diff --git a/Engine/Source/Runtime/Core/Private/Misc/DateTime.cpp b/Engine/Source/Runtime/Core/Private/Misc/DateTime.cpp
--- a/Engine/Source/Runtime/Core/Private/Misc/DateTime.cpp
+++ b/Engine/Source/Runtime/Core/Private/Misc/DateTime.cpp
@@ -93,7 +93,7 @@ void FDateTime::GetDate(int32& OutYear, int32& OutMonth, int32& OutDay) const
 
 	int32 i, j, k, l, n;
 
-	l = FMath::FloorToInt32((float)(GetJulianDay() + 0.5)) + 68569;
+	l = FMath::FloorToInt32(GetDate().GetJulianDay() + 0.5) + 68569;
 	n = 4 * l / 146097;
 	l = l - (146097 * n + 3) / 4;
 	i = 4000 * (l + 1) / 1461001;
