Optimize code size and overhead for Insights markers generated for scope cycle counters. Down from 213 bytes per marker to 142 bytes (overall program size down ~600KB).

Details: removes duplicate checks on CpuChannel by resolving SpecIds inside the scope event structure and avoiding passing in CpuChannel as the fine-grained channel when we always check CpuChannel at least.

diff --git a/Engine/Source/Runtime/Core/Public/ProfilingDebugging/CpuProfilerTrace.h b/Engine/Source/Runtime/Core/Public/ProfilingDebugging/CpuProfilerTrace.h
--- a/Engine/Source/Runtime/Core/Public/ProfilingDebugging/CpuProfilerTrace.h
+++ b/Engine/Source/Runtime/Core/Public/ProfilingDebugging/CpuProfilerTrace.h
@@ -88,10 +88,55 @@ struct FCpuProfilerTrace
 	 */
 	CORE_API static void OutputEndEvent();
 
-	struct FEventScope
+	class FEventScope
 	{
-		FEventScope(uint32 InSpecId, const UE::Trace::FChannel& Channel, bool bCondition)
-			: bEnabled(bCondition && (Channel | CpuChannel))
+	public:
+		FORCEINLINE FEventScope(uint32 InSpecId, bool bInCondition)
+			: bEnabled(bInCondition && CpuChannel)
+		{
+			BeginEventCommon(InSpecId);
+		}
+
+		FORCEINLINE FEventScope(uint32 InSpecId, const UE::Trace::FChannel& InChannel, bool bInCondition)
+			: bEnabled(bInCondition && (CpuChannel | InChannel))
+		{
+			BeginEventCommon(InSpecId);
+		}
+
+		FORCEINLINE FEventScope(uint32& InOutSpecId, const ANSICHAR* InEventString, bool bInCondition)
+			: bEnabled(bInCondition && CpuChannel)
+		{
+			BeginEventCommon(InOutSpecId, InEventString);
+		}
+
+		FORCEINLINE FEventScope(uint32& InOutSpecId, const ANSICHAR* InEventString, const UE::Trace::FChannel& InChannel, bool bInCondition)
+			: bEnabled(bInCondition && (CpuChannel | InChannel))
+		{
+			BeginEventCommon(InOutSpecId, InEventString);
+		}
+
+		FORCEINLINE FEventScope(uint32& InOutSpecId, const TCHAR* InEventString, bool bInCondition)
+			: bEnabled(bInCondition && CpuChannel)
+		{
+			BeginEventCommon(InOutSpecId, InEventString);
+		}
+
+		FORCEINLINE FEventScope(uint32& InOutSpecId, const TCHAR* InEventString, const UE::Trace::FChannel& InChannel, bool bInCondition)
+			: bEnabled(bInCondition && (CpuChannel | InChannel))
+		{
+			BeginEventCommon(InOutSpecId, InEventString);
+		}
+
+		FORCEINLINE ~FEventScope()
+		{
+			if (bEnabled)
+			{
+				OutputEndEvent();
+			}
+		}
+
+	private:
+		FORCEINLINE void BeginEventCommon(uint32 InSpecId)
 		{
 			if (bEnabled)
 			{
@@ -99,11 +144,27 @@ struct FCpuProfilerTrace
 			}
 		}
 
-		~FEventScope()
+		FORCEINLINE void BeginEventCommon(uint32& InOutSpecId, const ANSICHAR* InEventString)
 		{
 			if (bEnabled)
 			{
-				OutputEndEvent();
+				if (InOutSpecId == 0)
+				{
+					InOutSpecId = FCpuProfilerTrace::OutputEventType(InEventString);
+				}
+				OutputBeginEvent(InOutSpecId);
+			}
+		}
+
+		FORCEINLINE void BeginEventCommon(uint32& InOutSpecId, const TCHAR* InEventString)
+		{
+			if (bEnabled)
+			{
+				if (InOutSpecId == 0)
+				{
+					InOutSpecId = FCpuProfilerTrace::OutputEventType(InEventString);
+				}
+				OutputBeginEvent(InOutSpecId);
 			}
 		}
 
@@ -112,25 +173,43 @@ struct FCpuProfilerTrace
 
 	struct FDynamicEventScope
 	{
-		FDynamicEventScope(const ANSICHAR* InEventName, const UE::Trace::FChannel& Channel, bool Condition, const ANSICHAR* File = nullptr, uint32 Line = 0)
-			: bEnabled(Condition && (Channel | CpuChannel))
+		FORCEINLINE FDynamicEventScope(const ANSICHAR* InEventName, bool bInCondition, const ANSICHAR* InFile = nullptr, uint32 InLine = 0)
+			: bEnabled(bInCondition && CpuChannel)
 		{
 			if (bEnabled)
 			{
-				OutputBeginDynamicEvent(InEventName, File, Line);
+				OutputBeginDynamicEvent(InEventName, InFile, InLine);
 			}
 		}
 
-		FDynamicEventScope(const TCHAR* InEventName, const UE::Trace::FChannel& Channel, bool Condition, const ANSICHAR* File = nullptr, uint32 Line = 0)
-			: bEnabled(Condition && (Channel | CpuChannel))
+		FORCEINLINE FDynamicEventScope(const ANSICHAR* InEventName, const UE::Trace::FChannel& InChannel, bool bInCondition, const ANSICHAR* InFile = nullptr, uint32 InLine = 0)
+			: bEnabled(bInCondition && (CpuChannel | InChannel))
 		{
 			if (bEnabled)
 			{
-				OutputBeginDynamicEvent(InEventName, File, Line);
+				OutputBeginDynamicEvent(InEventName, InFile, InLine);
 			}
 		}
 
-		~FDynamicEventScope()
+		FORCEINLINE FDynamicEventScope(const TCHAR* InEventName, bool bInCondition, const ANSICHAR* InFile = nullptr, uint32 InLine = 0)
+			: bEnabled(bInCondition && CpuChannel)
+		{
+			if (bEnabled)
+			{
+				OutputBeginDynamicEvent(InEventName, InFile, InLine);
+			}
+		}
+
+		FORCEINLINE FDynamicEventScope(const TCHAR* InEventName, const UE::Trace::FChannel& InChannel, bool bInCondition, const ANSICHAR* InFile = nullptr, uint32 InLine = 0)
+			: bEnabled(bInCondition && (CpuChannel | InChannel))
+		{
+			if (bEnabled)
+			{
+				OutputBeginDynamicEvent(InEventName, InFile, InLine);
+			}
+		}
+
+		FORCEINLINE ~FDynamicEventScope()
 		{
 			if (bEnabled)
 			{
@@ -144,24 +223,41 @@ struct FCpuProfilerTrace
 
 // Advanced macro for integrating e.g. stats/named events with cpu trace
 // Declares a cpu timing event for future use, conditionally and with a particular variable name for storing the id
-#define TRACE_CPUPROFILER_EVENT_DECLARE(NameStr, DeclName, Channel, Condition) \
-	static uint32 DeclName; \
-	if ((Condition) && bool(Channel|CpuChannel) && (DeclName == 0)) { \
-		DeclName = FCpuProfilerTrace::OutputEventType(NameStr, __FILE__, __LINE__); \
-	}
+#define TRACE_CPUPROFILER_EVENT_DECLARE(DeclName) \
+	static uint32 DeclName;
+
+// Advanced macro for integrating e.g. stats/named events with cpu trace
+// Traces a scoped event previously declared with TRACE_CPUPROFILER_EVENT_DECLARE, conditionally
+#define TRACE_CPUPROFILER_EVENT_SCOPE_USE(DeclName, NameStr, ScopeName, Condition) \
+	FCpuProfilerTrace::FEventScope ScopeName(DeclName, NameStr, Condition);
 
 // Advanced macro for integrating e.g. stats/named events with cpu trace
 // Traces a scoped event previously declared with TRACE_CPUPROFILER_EVENT_DECLARE, conditionally
-#define TRACE_CPUPROFILER_EVENT_SCOPE_USE(DeclName, ScopeName, Channel, Condition) \
-	FCpuProfilerTrace::FEventScope ScopeName(DeclName, Channel, Condition);
+#define TRACE_CPUPROFILER_EVENT_SCOPE_USE_ON_CHANNEL(DeclName, NameStr, ScopeName, Channel, Condition) \
+	FCpuProfilerTrace::FEventScope ScopeName(DeclName, NameStr, Channel, Condition);
+
+// Trace a scoped cpu timing event providing a static string (const ANSICHAR* or const TCHAR*)
+// as the scope name. It will use the Cpu trace channel.
+// Example: TRACE_CPUPROFILER_EVENT_SCOPE_STR("My Scoped Timer A")
+#define TRACE_CPUPROFILER_EVENT_SCOPE_STR(NameStr) \
+	TRACE_CPUPROFILER_EVENT_DECLARE(PREPROCESSOR_JOIN(__CpuProfilerEventSpecId, __LINE__)); \
+	TRACE_CPUPROFILER_EVENT_SCOPE_USE(PREPROCESSOR_JOIN(__CpuProfilerEventSpecId, __LINE__), NameStr, PREPROCESSOR_JOIN(__CpuProfilerEventScope, __LINE__), true); 
 
 // Trace a scoped cpu timing event providing a static string (const ANSICHAR* or const TCHAR*)
 // as the scope name and a trace channel.
 // Example: TRACE_CPUPROFILER_EVENT_SCOPE_ON_CHANNEL_STR("My Scoped Timer A", CpuChannel)
 // Note: The event will be emitted only if both the given channel and CpuChannel is enabled.
 #define TRACE_CPUPROFILER_EVENT_SCOPE_ON_CHANNEL_STR(NameStr, Channel) \
-	TRACE_CPUPROFILER_EVENT_DECLARE(NameStr, PREPROCESSOR_JOIN(__CpuProfilerEventSpecId, __LINE__), Channel, true); \
-	TRACE_CPUPROFILER_EVENT_SCOPE_USE(PREPROCESSOR_JOIN(__CpuProfilerEventSpecId, __LINE__), PREPROCESSOR_JOIN(__CpuProfilerEventScope, __LINE__), Channel, true); 
+	TRACE_CPUPROFILER_EVENT_DECLARE(PREPROCESSOR_JOIN(__CpuProfilerEventSpecId, __LINE__)); \
+	TRACE_CPUPROFILER_EVENT_SCOPE_USE_ON_CHANNEL(PREPROCESSOR_JOIN(__CpuProfilerEventSpecId, __LINE__), NameStr, PREPROCESSOR_JOIN(__CpuProfilerEventScope, __LINE__), Channel, true); 
+
+// Trace a scoped cpu timing event providing a scope name (plain text).
+// It will use the Cpu trace channel.
+// Example: TRACE_CPUPROFILER_EVENT_SCOPE(MyScopedTimer::A)
+// Note: Do not use this macro with a static string because, in that case, additional quotes will
+//       be added around the event scope name.
+#define TRACE_CPUPROFILER_EVENT_SCOPE(Name) \
+	TRACE_CPUPROFILER_EVENT_SCOPE_STR(#Name)
 
 // Trace a scoped cpu timing event providing a scope name (plain text) and a trace channel.
 // Example: TRACE_CPUPROFILER_EVENT_SCOPE_ON_CHANNEL(MyScopedTimer::A, CpuChannel)
@@ -171,20 +267,6 @@ struct FCpuProfilerTrace
 #define TRACE_CPUPROFILER_EVENT_SCOPE_ON_CHANNEL(Name, Channel) \
 	TRACE_CPUPROFILER_EVENT_SCOPE_ON_CHANNEL_STR(#Name, Channel)
 
-// Trace a scoped cpu timing event providing a static string (const ANSICHAR* or const TCHAR*)
-// as the scope name. It will use the Cpu trace channel.
-// Example: TRACE_CPUPROFILER_EVENT_SCOPE_STR("My Scoped Timer A")
-#define TRACE_CPUPROFILER_EVENT_SCOPE_STR(NameStr) \
-	TRACE_CPUPROFILER_EVENT_SCOPE_ON_CHANNEL_STR(NameStr, CpuChannel)
-
-// Trace a scoped cpu timing event providing a scope name (plain text).
-// It will use the Cpu trace channel.
-// Example: TRACE_CPUPROFILER_EVENT_SCOPE(MyScopedTimer::A)
-// Note: Do not use this macro with a static string because, in that case, additional quotes will
-//       be added around the event scope name.
-#define TRACE_CPUPROFILER_EVENT_SCOPE(Name) \
-	TRACE_CPUPROFILER_EVENT_SCOPE_ON_CHANNEL(Name, CpuChannel)
-
 // Trace a scoped cpu timing event providing a dynamic string (const ANSICHAR* or const TCHAR*)
 // as the scope name and a trace channel.
 // Example: TRACE_CPUPROFILER_EVENT_SCOPE_TEXT_ON_CHANNEL(*MyScopedTimerNameString, CpuChannel)
@@ -199,7 +281,7 @@ struct FCpuProfilerTrace
 // Note: This macro has a larger overhead compared to macro that accepts a plain text name
 //       or a static string. Use it only if scope name really needs to be a dynamic string.
 #define TRACE_CPUPROFILER_EVENT_SCOPE_TEXT(Name) \
-	TRACE_CPUPROFILER_EVENT_SCOPE_TEXT_ON_CHANNEL(Name, CpuChannel)
+	FCpuProfilerTrace::FDynamicEventScope PREPROCESSOR_JOIN(__CpuProfilerEventScope, __LINE__)(Name, true, __FILE__, __LINE__);
 
 #else
 
diff --git a/Engine/Source/Runtime/Core/Public/Stats/Stats.h b/Engine/Source/Runtime/Core/Public/Stats/Stats.h
--- a/Engine/Source/Runtime/Core/Public/Stats/Stats.h
+++ b/Engine/Source/Runtime/Core/Public/Stats/Stats.h
@@ -236,8 +236,8 @@ FORCEINLINE void StatsMasterEnableSubtract(int32 Value = 1)
 #endif
 
 #define SCOPE_CYCLE_COUNTER_TO_TRACE(StatString, StatName, Condition) \
-	TRACE_CPUPROFILER_EVENT_DECLARE(StatString, PREPROCESSOR_JOIN(PREPROCESSOR_JOIN(__Decl_, StatName), __LINE__), CpuChannel, Condition && GCycleStatsShouldEmitNamedEvents>0); \
-	TRACE_CPUPROFILER_EVENT_SCOPE_USE(PREPROCESSOR_JOIN(PREPROCESSOR_JOIN(__Decl_, StatName), __LINE__), PREPROCESSOR_JOIN(PREPROCESSOR_JOIN(__Scope_, StatName), __LINE__), CpuChannel, Condition && GCycleStatsShouldEmitNamedEvents>0);
+	TRACE_CPUPROFILER_EVENT_DECLARE(PREPROCESSOR_JOIN(PREPROCESSOR_JOIN(__Decl_, StatName), __LINE__)); \
+	TRACE_CPUPROFILER_EVENT_SCOPE_USE(PREPROCESSOR_JOIN(PREPROCESSOR_JOIN(__Decl_, StatName), __LINE__), StatString, PREPROCESSOR_JOIN(PREPROCESSOR_JOIN(__Scope_, StatName), __LINE__), Condition && GCycleStatsShouldEmitNamedEvents);
 
 #define DECLARE_SCOPE_CYCLE_COUNTER(CounterName,Stat,GroupId) \
 	FScopeCycleCounter StatNamedEventsScope_##Stat(TStatId(ANSI_TO_PROFILING(#Stat))); \
