Add USE_PROC_SELF_SMAPS_ROLLUP define for FUnixPlatformMemory::GetExtendedStats

Set to 1 if we should try to use /proc/self/smaps_rollup in FUnixPlatformMemory::GetExtendedStats().
There is a potential tradeoff in that smaps_rollup appears to be quite a bit faster in
straight testing, but it also looks like it blocks mmap() calls on other threads until it finishes.

Also move /proc/self/smaps code to use open/read/close and read + parse in 512 byte chunks instead of
fopen to avoid malloc calls in fopen/fgets.

diff --git a/Engine/Source/Programs/TestPAL/Private/Main.cpp b/Engine/Source/Programs/TestPAL/Private/Main.cpp
--- a/Engine/Source/Programs/TestPAL/Private/Main.cpp
+++ b/Engine/Source/Programs/TestPAL/Private/Main.cpp
@@ -6,6 +6,7 @@
 #include "Misc/Guid.h"
 #include "Stats/StatsMisc.h"
 #include "HAL/RunnableThread.h"
+#include "HAL/PlatformMemory.h"
 #include "HAL/PlatformApplicationMisc.h"
 #include "GenericPlatform/GenericApplication.h"
 
@@ -263,6 +264,15 @@ int32 ThreadSingletonTest(const TCHAR* CommandLine)
 	return 0;
 }
 
+#if PLATFORM_LINUX
+
+static float ToMB(uint64 Value)
+{
+	return Value * (1.0 / (1024.0 * 1024.0));
+}
+
+#endif // PLATFORM_LINUX
+
 /**
  * Sysinfo test
  */
@@ -302,6 +312,15 @@ int32 SysInfoTest(const TCHAR* CommandLine)
 
 	FPlatformMemory::DumpStats(*GLog);
 
+#if PLATFORM_LINUX
+	FExtendedPlatformMemoryStats StatsEx = FUnixPlatformMemory::GetExtendedStats();
+
+	UE_LOG(LogTestPAL, Display, TEXT("Shared_Clean:%.2f MB, Shared_Dirty:%.2f MB"),
+		ToMB(StatsEx.Shared_Clean), ToMB(StatsEx.Shared_Dirty));
+	UE_LOG(LogTestPAL, Display, TEXT("Private_Clean:%.2fMB Private_Dirty:%.2fMB"),
+		ToMB(StatsEx.Private_Clean), ToMB(StatsEx.Private_Dirty));
+#endif // PLATFORM_LINUX
+
 	FEngineLoop::AppPreExit();
 	FEngineLoop::AppExit();
 	return 0;
diff --git a/Engine/Source/Runtime/Core/Private/Unix/UnixPlatformMemory.cpp b/Engine/Source/Runtime/Core/Private/Unix/UnixPlatformMemory.cpp
--- a/Engine/Source/Runtime/Core/Private/Unix/UnixPlatformMemory.cpp
+++ b/Engine/Source/Runtime/Core/Private/Unix/UnixPlatformMemory.cpp
@@ -47,6 +47,13 @@
 // Set rather to use BinnedMalloc2 for binned malloc, can be overridden below
 #define USE_MALLOC_BINNED2 (1)
 
+// Set to 1 if we should try to use /proc/self/smaps_rollup in FUnixPlatformMemory::GetExtendedStats().
+// There is a potential tradeoff in that smaps_rollup appears to be quite a bit faster in
+//   straight testing, but it also looks like it blocks mmap() calls on other threads until it finishes.
+#ifndef USE_PROC_SELF_SMAPS_ROLLUP
+#define USE_PROC_SELF_SMAPS_ROLLUP 0
+#endif
+
 // Used in UnixPlatformStackwalk to skip the crash handling callstack frames.
 bool CORE_API GFullCrashCallstack = false;
 
@@ -653,47 +660,89 @@ struct FProcField
 {
 	const ANSICHAR* Name = nullptr;
 	uint64* Addr = nullptr;
+	uint32 NameLen = 0;
+
+	FProcField(const ANSICHAR *NameIn, uint64* AddrIn)
+		: Name(NameIn), Addr(AddrIn)
+	{
+		NameLen = FCStringAnsi::Strlen(Name);
+	}
 };
 
+static uint64 ParseProcFieldsChunk(ANSICHAR *Buffer, uint64 BufferSize, FProcField *ProcFields, uint32 NumFields, uint32 &NumFieldsFound)
+{
+	uint64 ParsePos = 0;
+
+	for (uint64 Idx = 0; Idx < BufferSize; Idx++)
+	{
+		if (Buffer[Idx] == '\n')
+		{
+			const ANSICHAR *Line = Buffer + ParsePos;
+
+			Buffer[Idx] = 0;
+
+			for (uint32 IdxField = 0; IdxField < NumFields; IdxField++)
+			{
+				uint32 NameLen = ProcFields[IdxField].NameLen;
+
+				if (!FCStringAnsi::Strncmp(Line, ProcFields[IdxField].Name, NameLen))
+				{
+					*ProcFields[IdxField].Addr = atoll(Line + NameLen) * 1024ULL;
+					NumFieldsFound++;
+					break;
+				}
+			}
+
+			ParsePos = Idx + 1;
+		}
+	}
+
+	// If we didn't find any linefeeds, skip the entire chunk
+	return ParsePos ? ParsePos : BufferSize;
+}
+
 static uint32 ReadProcFields(const char *FileName, FProcField *ProcFields, uint32 NumFields)
 {
 	uint32 NumFieldsFound = 0;
-	// 12k should be enough for the proc files we're reading for stats, with room for growth.
-	// smaps_rollup is 642 bytes, meminfo is 1428 bytes, and proc/self/status is 1503 bytes.
-	// If we run out of space and fail to find a field we should also see the warning below.
-	ANSICHAR Buffer[12 * 1024];
 
 	int Fd = open(FileName, O_RDONLY);
 	if (Fd >= 0)
 	{
-		ssize_t Bytes;
+		const uint32 ChunkSize = 512;
+		ANSICHAR Buffer[ChunkSize];
+		uint64 BytesAvailableInChunk = 0;
 
-		do
+		for (;;)
 		{
-			Bytes = read(Fd, Buffer, sizeof(Buffer) - 1);
-		} while (Bytes < 0 && errno == EINTR);
-
-		close(Fd);
-		Fd = -1;
+			ssize_t BytesRead;
+			uint64 BytesToRead = ChunkSize - BytesAvailableInChunk;
 
-		if ((Bytes > 0) && (Bytes < UE_ARRAY_COUNT(Buffer)))
-		{
-			Buffer[Bytes] = 0;
+			do
+			{
+				BytesRead = read(Fd, Buffer + BytesAvailableInChunk, BytesToRead);
+			} while (BytesRead < 0 && errno == EINTR);
 
-			for (uint32 Idx = 0; Idx < NumFields; Idx++)
+			if (BytesRead <= 0)
 			{
-				const ANSICHAR* Field = FCStringAnsi::Strstr(Buffer, ProcFields[Idx].Name);
+				break;
+			}
 
-				if (Field)
-				{
-					uint32 NameLen = FCStringAnsi::Strlen(ProcFields[Idx].Name);
+			BytesAvailableInChunk += BytesRead;
 
-					*ProcFields[Idx].Addr = atoll(Field + NameLen) * 1024ULL;
-					NumFieldsFound++;
-				}
+			uint64 BytesParsed = ParseProcFieldsChunk(Buffer, BytesAvailableInChunk, ProcFields, NumFields, NumFieldsFound);
+			checkf(BytesParsed <= BytesAvailableInChunk, TEXT("BytesParsed more than BytesAvailableInChunk %u %u"), BytesParsed, BytesAvailableInChunk);
+
+			if (BytesRead < BytesToRead || NumFieldsFound == NumFields)
+			{
+				break;
 			}
+
+			BytesAvailableInChunk -= BytesParsed;
+			memmove(Buffer, Buffer + BytesParsed, BytesAvailableInChunk);
 		}
 
+		close(Fd);
+
 		if (NumFieldsFound != NumFields)
 		{
 			static bool bLogOnceMissing = false;
@@ -703,7 +752,7 @@ static uint32 ReadProcFields(const char *FileName, FProcField *ProcFields, uint3
 				// allocate memory. This function could get called via FMemory::GCreateMalloc or signal handlers
 				// and we will potentiall crash in these cases calling UE_LOG or TCHAR_TO_UTF8.
 				fprintf(stderr, "Warning: ReadProcFields: %u of %u fields found in %s.\n",
-					NumFieldsFound, NumFields, FileName);
+						NumFieldsFound, NumFields, FileName);
 				fflush(stderr);
 				bLogOnceMissing = true;
 			}
@@ -771,6 +820,37 @@ FPlatformMemoryStats FUnixPlatformMemory::GetStats()
 	return MemoryStats;
 }
 
+static uint64 ParseSMapsFileChunk(ANSICHAR *Buffer, uint64 BufferSize, FProcField *ProcFields, uint32 NumFields)
+{
+	uint64 ParsePos = 0;
+
+	for (uint64 Idx = 0; Idx < BufferSize; Idx++)
+	{
+		if (Buffer[Idx] == '\n')
+		{
+			const ANSICHAR *Line = Buffer + ParsePos;
+
+			Buffer[Idx] = 0;
+
+			for (uint32 IdxField = 0; IdxField < NumFields; IdxField++)
+			{
+				uint32 NameLen = ProcFields[IdxField].NameLen;
+
+				if (!FCStringAnsi::Strncmp(Line, ProcFields[IdxField].Name, NameLen))
+				{
+					*ProcFields[IdxField].Addr += atoll(Line + NameLen) * 1024ULL;
+					break;
+				}
+			}
+
+			ParsePos = Idx + 1;
+		}
+	}
+
+	// If we didn't find any linefeeds, skip the entire chunk
+	return ParsePos ? ParsePos : BufferSize;
+}
+
 FExtendedPlatformMemoryStats FUnixPlatformMemory::GetExtendedStats()
 {
 	const ANSICHAR Shared_CleanStr[]  = "Shared_Clean:";
@@ -788,48 +868,51 @@ FExtendedPlatformMemoryStats FUnixPlatformMemory::GetExtendedStats()
 		{ Private_DirtyStr, (uint64 *)&MemoryStats.Private_Dirty },
 	};
 
+#if USE_PROC_SELF_SMAPS_ROLLUP
 	// ~ 1.06ms per call on my Threadripper 3990X w/ Debian Testing 5.15.0-2-amd64
+	// Note that testing shows us opening smaps_rollup on a thread while another thread is
+	//  calling mmap() can cause large spikes (30+ms) in the mmap.
 	if (!ReadProcFields("/proc/self/smaps_rollup", SMapsFields, UE_ARRAY_COUNT(SMapsFields)))
+#endif
 	{
-		// ~ 6.8ms per call on my Threadripper 3990X w/ Debian Testing 5.15.0-2-amd64 in TestPAL.
-		// Potentially far higher though, like 100s of ms.
-		if (FILE* ProcSMaps = fopen("/proc/self/smaps", "r"))
+		int Fd = open("/proc/self/smaps", O_RDONLY);
+
+		if (Fd >= 0)
 		{
-			const uint32 Shared_CleanStrLen = FCStringAnsi::Strlen(Shared_CleanStr);
-			const uint32 Shared_DirtyStrLen = FCStringAnsi::Strlen(Shared_DirtyStr);
-			const uint32 Private_CleanStrLen = FCStringAnsi::Strlen(Private_CleanStr);
-			const uint32 Private_DirtyStrLen = FCStringAnsi::Strlen(Private_DirtyStr);
+			const uint32 ChunkSize = 512;
+			ANSICHAR Buffer[ChunkSize];
+			uint64 BytesAvailableInChunk = 0;
 
 			for (;;)
 			{
-				ANSICHAR LineBuffer[256];
-				ANSICHAR *Line = fgets(LineBuffer, UE_ARRAY_COUNT(LineBuffer), ProcSMaps);
+				ssize_t BytesRead;
+				uint64 BytesToRead = ChunkSize - BytesAvailableInChunk;
 
-				if (Line == nullptr)
+				do
 				{
-					// eof or error
-					break;
-				}
+					BytesRead = read(Fd, Buffer + BytesAvailableInChunk, BytesToRead);
+				} while (BytesRead < 0 && errno == EINTR);
 
-				if (!FCStringAnsi::Strncmp(Line, Shared_CleanStr, Shared_CleanStrLen))
-				{
-					MemoryStats.Shared_Clean += atoll(Line + Shared_CleanStrLen) * 1024ULL;
-				}
-				else if (!FCStringAnsi::Strncmp(Line, Shared_DirtyStr, Shared_DirtyStrLen))
+				if (BytesRead <= 0)
 				{
-					MemoryStats.Shared_Dirty += atoll(Line + Shared_DirtyStrLen) * 1024ULL;
-				}
-				else if (!FCStringAnsi::Strncmp(Line, Private_CleanStr, Private_CleanStrLen))
-				{
-					MemoryStats.Private_Clean += atoll(Line + Private_CleanStrLen) * 1024ULL;
+					break;
 				}
-				else if (!FCStringAnsi::Strncmp(Line, Private_DirtyStr, Private_DirtyStrLen))
+
+				BytesAvailableInChunk += BytesRead;
+
+				uint64 BytesParsed = ParseSMapsFileChunk(Buffer, BytesAvailableInChunk, SMapsFields, UE_ARRAY_COUNT(SMapsFields));
+				checkf(BytesParsed <= BytesAvailableInChunk, TEXT("BytesParsed more than BytesAvailableInChunk %u %u"), BytesParsed, BytesAvailableInChunk);
+
+				if (BytesRead < BytesToRead)
 				{
-					MemoryStats.Private_Dirty += atoll(Line + Private_DirtyStrLen) * 1024ULL;
+					break;
 				}
+
+				BytesAvailableInChunk -= BytesParsed;
+				memmove(Buffer, Buffer + BytesParsed, BytesAvailableInChunk);
 			}
 
-			fclose(ProcSMaps);
+			close(Fd);
 		}
 	}
 
