Linux: Use PCI-Id information if NVIDIA driver returns Unknown GPU

Seems like v495 and v510 drivers are giving us this:

  $ cat /proc/driver/nvidia/gpus/**/information
  Model: 		 Unknown
  IRQ:   		 150
  GPU UUID: 	 ...
  Video BIOS: 	 90.04.7a.00.9d
  Bus Type: 	 PCIe
  DMA Size: 	 47 bits
  DMA Mask: 	 0x7fffffffffff
  Bus Location: 	 0000:4b:00.0
  Device Minor: 	 0
  GPU Excluded:	 No

diff --git a/Engine/Source/Runtime/Core/Private/Unix/UnixPlatformMisc.cpp b/Engine/Source/Runtime/Core/Private/Unix/UnixPlatformMisc.cpp
--- a/Engine/Source/Runtime/Core/Private/Unix/UnixPlatformMisc.cpp
+++ b/Engine/Source/Runtime/Core/Private/Unix/UnixPlatformMisc.cpp
@@ -1241,6 +1241,28 @@ static FString ReadGPUFile(const char *Filename, const char *Token = nullptr)
 	return Value;
 }
 
+static FString GetVendorDevicePCIIdStr()
+{
+	FString Vendor0 = ReadGPUFile("/sys/class/drm/card0/device/vendor");
+	FString Device0 = ReadGPUFile("/sys/class/drm/card0/device/device");
+	FString SubsystemVendor0 = ReadGPUFile("/sys/class/drm/card0/device/subsystem_vendor");
+	FString SubsystemDevice0 = ReadGPUFile("/sys/class/drm/card0/device/subsystem_device");
+	uint32 VendorId = strtoul(TCHAR_TO_UTF8(*Vendor0), nullptr, 16);
+	const TCHAR* VendorStr = VendorIdToString(VendorId);
+
+	FString ProcGPUBrandStr = VendorStr;
+	if (!Device0.IsEmpty())
+	{
+		ProcGPUBrandStr += FString::Printf(TEXT(" PCI-id: %s-%s"), *Vendor0, *Device0);
+	}
+	if (!SubsystemVendor0.IsEmpty() || !SubsystemDevice0.IsEmpty())
+	{
+		ProcGPUBrandStr += FString::Printf(TEXT(" (%s-%s)"), *SubsystemVendor0, *SubsystemDevice0);
+	}
+
+	return ProcGPUBrandStr;
+}
+
 // Return primary GPU brand string
 //  This should trying to report out what GPU is in use: Nvidia/Type/Driver version.
 FString FUnixPlatformMisc::GetPrimaryGPUBrand()
@@ -1273,27 +1295,21 @@ FString FUnixPlatformMisc::GetPrimaryGPUBrand()
 				globfree(&GlobResults);
 			}
 
-			ProcGPUBrandStr = FString::Printf(TEXT("%s (%s)"), GPUModel.IsEmpty() ? TEXT("NVIDIA") : *GPUModel, *NVIDIAVersion);
+			// It looks like sometimes with v495 and v510 drivers we're getting this:
+			// $ cat /proc/driver/nvidia/gpus/**/information
+			//   Model: Unknown
+			//   ...
+			// If this is the case, let's try to grab the PCI-Id info.
+			if (GPUModel.IsEmpty() || GPUModel == TEXT("Unknown"))
+			{
+				GPUModel = GetVendorDevicePCIIdStr();
+			}
+
+			ProcGPUBrandStr = FString::Printf(TEXT("%s (%s)"), *GPUModel, *NVIDIAVersion);
 		}
 		else
 		{
-			// Returns something ~ like "AMD PCI-id: 1002-731f (1462-3811)"
-			FString Vendor0 = ReadGPUFile("/sys/class/drm/card0/device/vendor");
-			FString Device0 = ReadGPUFile("/sys/class/drm/card0/device/device");
-			FString SubsystemVendor0 = ReadGPUFile("/sys/class/drm/card0/device/subsystem_vendor");
-			FString SubsystemDevice0 = ReadGPUFile("/sys/class/drm/card0/device/subsystem_device");
-			uint32 VendorId = strtoul(TCHAR_TO_UTF8(*Vendor0), nullptr, 16);
-			const TCHAR* VendorStr = VendorIdToString(VendorId);
-
-			ProcGPUBrandStr = VendorStr;
-			if (!Device0.IsEmpty())
-			{
-				ProcGPUBrandStr += FString::Printf(TEXT(" PCI-id: %s-%s"), *Vendor0, *Device0);
-			}
-			if (!SubsystemVendor0.IsEmpty() || !SubsystemDevice0.IsEmpty())
-			{
-				ProcGPUBrandStr += FString::Printf(TEXT(" (%s-%s)"), *SubsystemVendor0, *SubsystemDevice0);
-			}
+			ProcGPUBrandStr = GetVendorDevicePCIIdStr();
 		}
 	}
 
