PR #8863: Fix incorrect return values for FUntilCommand (Contributed by [REDACTED_PERSON])

diff --git a/Engine/Source/Runtime/Core/Public/Misc/AutomationTest.h b/Engine/Source/Runtime/Core/Public/Misc/AutomationTest.h
--- a/Engine/Source/Runtime/Core/Public/Misc/AutomationTest.h
+++ b/Engine/Source/Runtime/Core/Public/Misc/AutomationTest.h
@@ -3455,13 +3455,13 @@ public:
 			if ( NewTime - StartTime >= Timeout )
 			{
 				TimeoutCallback();
-				return false;
+				return true;
 			}
 
-			return true;
+			return false;
 		}
 
-		return false;
+		return true;
 	}
 
 private:
