Fix GetEngineCrypto() not being threadsafe

diff --git a/Engine/Source/Runtime/Core/Private/IO/IoStore.cpp b/Engine/Source/Runtime/Core/Private/IO/IoStore.cpp
--- a/Engine/Source/Runtime/Core/Private/IO/IoStore.cpp
+++ b/Engine/Source/Runtime/Core/Private/IO/IoStore.cpp
@@ -72,7 +72,9 @@ bool WriteArray(IFileHandle* FileHandle, const ArrayType& Array)
 
 static IEngineCrypto* GetEngineCrypto()
 {
+	IModularFeatures::Get().LockModularFeatureList();
 	static TArray<IEngineCrypto*> Features = IModularFeatures::Get().GetModularFeatureImplementations<IEngineCrypto>(IEngineCrypto::GetFeatureName());
+	IModularFeatures::Get().UnlockModularFeatureList();
 	checkf(Features.Num() > 0, TEXT("RSA functionality was used but no modular feature was registered to provide it. Please make sure your project has the PlatformCrypto plugin enabled!"));
 	return Features[0];
 }
