EMT2 - Added Update to the RobinHoodHashTable and used it to fix a recent Bug that would otherwise not Update the RenderData properly.

diff --git a/Engine/Source/Runtime/Core/Public/Experimental/Containers/RobinHoodHashTable.h b/Engine/Source/Runtime/Core/Public/Experimental/Containers/RobinHoodHashTable.h
--- a/Engine/Source/Runtime/Core/Public/Experimental/Containers/RobinHoodHashTable.h
+++ b/Engine/Source/Runtime/Core/Public/Experimental/Containers/RobinHoodHashTable.h
@@ -249,6 +249,12 @@ namespace RobinHoodHashTable_Private
 			return &Pair.Value;
 		}
 
+		template<typename DeducedValueType>
+		inline void Update(DeducedValueType&& InValue)
+		{
+			Pair.Value = Forward<DeducedValueType>(InValue);
+		}
+
 		inline ElementType& GetElement()
 		{
 			return Pair;
@@ -523,9 +529,9 @@ namespace RobinHoodHashTable_Private
 			}
 		}
 
-	protected:
-		template<typename DeducedKeyType, typename DeducedValueType>
-		inline FHashElementId FindOrAddIdByHash(FHashType HashValue, DeducedKeyType&& Key, DeducedValueType&& Val, bool& bIsAlreadyInMap)
+	private:
+		template<bool UpdateValue, typename DeducedKeyType, typename DeducedValueType>
+		inline FHashElementId FindOrUpdateIdByHashInternal(FHashType HashValue, DeducedKeyType&& Key, DeducedValueType&& Val, bool& bIsAlreadyInMap)
 		{
 			CHECK_CONCURRENT_ACCESS(FPlatformAtomics::InterlockedIncrement(&ConcurrentWriters) == 1);
 			CHECK_CONCURRENT_ACCESS(FPlatformAtomics::InterlockedIncrement(&ConcurrentReaders) == 1);
@@ -542,6 +548,10 @@ namespace RobinHoodHashTable_Private
 						CHECK_CONCURRENT_ACCESS(FPlatformAtomics::InterlockedDecrement(&ConcurrentReaders) == 0);
 						CHECK_CONCURRENT_ACCESS(FPlatformAtomics::InterlockedDecrement(&ConcurrentWriters) == 0);
 						bIsAlreadyInMap = true;
+						if constexpr (UpdateValue)
+						{
+							KeyValueData.Get(IndexData[BucketIndex]).Update(Forward<DeducedValueType>(Val));
+						}
 						return IndexData[BucketIndex];
 					}
 				}
@@ -580,6 +590,13 @@ namespace RobinHoodHashTable_Private
 			return FHashElementId(InsertIndex);
 		}
 
+	protected:
+		template<typename DeducedKeyType, typename DeducedValueType>
+		inline FHashElementId FindOrAddIdByHash(FHashType HashValue, DeducedKeyType&& Key, DeducedValueType&& Val, bool& bIsAlreadyInMap)
+		{
+			return FindOrUpdateIdByHashInternal<false>(HashValue, Forward<DeducedKeyType>(Key), Forward<DeducedValueType>(Val), bIsAlreadyInMap);
+		}
+
 		template<typename DeducedKeyType, typename DeducedValueType>
 		inline FHashElementId FindOrAddId(DeducedKeyType&& Key, DeducedValueType&& Val, bool& bIsAlreadyInMap)
 		{
@@ -594,6 +611,26 @@ namespace RobinHoodHashTable_Private
 			return KeyValueData.Get(Id.GetIndex()).FindImpl();
 		}
 
+		template<typename DeducedKeyType, typename DeducedValueType>
+		inline FHashElementId UpdateIdByHash(FHashType HashValue, DeducedKeyType&& Key, DeducedValueType&& Val, bool& bIsAlreadyInMap)
+		{
+			return FindOrUpdateIdByHashInternal<true>(HashValue, Forward<DeducedKeyType>(Key), Forward<DeducedValueType>(Val), bIsAlreadyInMap);
+		}
+
+		template<typename DeducedKeyType, typename DeducedValueType>
+		inline FHashElementId UpdateId(DeducedKeyType&& Key, DeducedValueType&& Val, bool& bIsAlreadyInMap)
+		{
+			FHashType HashValue = ComputeHash(Key);
+			return UpdateIdByHash(HashValue, Forward<DeducedKeyType>(Key), Forward<DeducedValueType>(Val), bIsAlreadyInMap);
+		}
+
+		template<typename DeducedKeyType, typename DeducedValueType>
+		inline FindValueType Update(DeducedKeyType&& Key, DeducedValueType&& Val, bool& bIsAlreadyInMap)
+		{
+			FHashElementId Id = UpdateId(Forward<DeducedKeyType>(Key), Forward<DeducedValueType>(Val), bIsAlreadyInMap);
+			return KeyValueData.Get(Id.GetIndex()).FindImpl();
+		}
+
 		TRobinHoodHashTable()
 		{
 			IndexData.Reserve(1); IndexData.AddUninitialized();
@@ -1234,6 +1271,138 @@ public:
 		bool bIsAlreadyInMap;
 		return Base::FindOrAdd(MoveTemp(Key), MoveTemp(Val), bIsAlreadyInMap);
 	}
+
+	FHashElementId UpdateIdByHash(FHashType HashValue, const KeyType& Key, const ValueType& Val, bool& bIsAlreadyInMap)
+	{
+		return Base::UpdateIdByHash(HashValue, Key, Val, bIsAlreadyInMap);
+	}
+
+	FHashElementId UpdateIdByHash(FHashType HashValue, const KeyType& Key, ValueType&& Val, bool& bIsAlreadyInMap)
+	{
+		return Base::UpdateIdByHash(HashValue, Key, MoveTemp(Val), bIsAlreadyInMap);
+	}
+
+	FHashElementId UpdateIdByHash(FHashType HashValue, KeyType&& Key, const ValueType& Val, bool& bIsAlreadyInMap)
+	{
+		return Base::UpdateIdByHash(HashValue, MoveTemp(Key), Val, bIsAlreadyInMap);
+	}
+
+	FHashElementId UpdateIdByHash(FHashType HashValue, KeyType&& Key, ValueType&& Val, bool& bIsAlreadyInMap)
+	{
+		return Base::UpdateIdByHash(HashValue, MoveTemp(Key), MoveTemp(Val), bIsAlreadyInMap);
+	}
+
+	FHashElementId UpdateId(const KeyType& Key, const ValueType& Val, bool& bIsAlreadyInMap)
+	{
+		return Base::UpdateId(Key, Val, bIsAlreadyInMap);
+	}
+
+	FHashElementId UpdateId(const KeyType& Key, ValueType&& Val, bool& bIsAlreadyInMap)
+	{
+		return Base::UpdateId(Key, MoveTemp(Val), bIsAlreadyInMap);
+	}
+
+	FHashElementId UpdateId(KeyType&& Key, const ValueType& Val, bool& bIsAlreadyInMap)
+	{
+		return Base::UpdateId(MoveTemp(Key), Val);
+	}
+
+	FHashElementId UpdateId(KeyType&& Key, ValueType&& Val, bool& bIsAlreadyInMap)
+	{
+		return Base::UpdateId(MoveTemp(Key), MoveTemp(Val), bIsAlreadyInMap);
+	}
+
+	FindValueType Update(const KeyType& Key, const ValueType& Val, bool& bIsAlreadyInMap)
+	{
+		return Base::Update(Key, Val, bIsAlreadyInMap);
+	}
+
+	FindValueType Update(const KeyType& Key, ValueType&& Val, bool& bIsAlreadyInMap)
+	{
+		return Base::Update(Key, MoveTemp(Val), bIsAlreadyInMap);
+	}
+
+	FindValueType Update(KeyType&& Key, const ValueType& Val, bool& bIsAlreadyInMap)
+	{
+		return Base::Update(MoveTemp(Key), Val, bIsAlreadyInMap);
+	}
+
+	FindValueType Update(KeyType&& Key, ValueType&& Val, bool& bIsAlreadyInMap)
+	{
+		return Base::Update(MoveTemp(Key), MoveTemp(Val), bIsAlreadyInMap);
+	}
+
+	FHashElementId UpdateIdByHash(FHashType HashValue, const KeyType& Key, const ValueType& Val)
+	{
+		bool bIsAlreadyInMap;
+		return Base::UpdateIdByHash(HashValue, Key, Val, bIsAlreadyInMap);
+	}
+
+	FHashElementId UpdateIdByHash(FHashType HashValue, const KeyType& Key, ValueType&& Val)
+	{
+		bool bIsAlreadyInMap;
+		return Base::UpdateIdByHash(HashValue, Key, MoveTemp(Val), bIsAlreadyInMap);
+	}
+
+	FHashElementId UpdateIdByHash(FHashType HashValue, KeyType&& Key, const ValueType& Val)
+	{
+		bool bIsAlreadyInMap;
+		return Base::UpdateIdByHash(HashValue, MoveTemp(Key), Val, bIsAlreadyInMap);
+	}
+
+	FHashElementId UpdateIdByHash(FHashType HashValue, KeyType&& Key, ValueType&& Val)
+	{
+		bool bIsAlreadyInMap;
+		return Base::UpdateIdByHash(HashValue, MoveTemp(Key), MoveTemp(Val), bIsAlreadyInMap);
+	}
+
+	FHashElementId UpdateId(const KeyType& Key, const ValueType& Val)
+	{
+		bool bIsAlreadyInMap;
+		return Base::UpdateId(Key, Val, bIsAlreadyInMap);
+	}
+
+	FHashElementId UpdateId(const KeyType& Key, ValueType&& Val)
+	{
+		bool bIsAlreadyInMap;
+		return Base::UpdateId(Key, MoveTemp(Val), bIsAlreadyInMap);
+	}
+
+	FHashElementId UpdateId(KeyType&& Key, const ValueType& Val)
+	{
+		bool bIsAlreadyInMap;
+		return Base::UpdateId(MoveTemp(Key), Val);
+	}
+
+	FHashElementId UpdateId(KeyType&& Key, ValueType&& Val)
+	{
+		bool bIsAlreadyInMap;
+		return Base::UpdateId(MoveTemp(Key), MoveTemp(Val), bIsAlreadyInMap);
+	}
+
+	FindValueType Update(const KeyType& Key, const ValueType& Val)
+	{
+		bool bIsAlreadyInMap;
+		return Base::Update(Key, Val, bIsAlreadyInMap);
+	}
+
+	FindValueType Update(const KeyType& Key, ValueType&& Val)
+	{
+		bool bIsAlreadyInMap;
+		return Base::Update(Key, MoveTemp(Val), bIsAlreadyInMap);
+	}
+
+	FindValueType Update(KeyType&& Key, const ValueType& Val)
+	{
+		bool bIsAlreadyInMap;
+		return Base::Update(MoveTemp(Key), Val, bIsAlreadyInMap);
+	}
+
+	FindValueType Update(KeyType&& Key, ValueType&& Val)
+	{
+		bool bIsAlreadyInMap;
+		return Base::Update(MoveTemp(Key), MoveTemp(Val), bIsAlreadyInMap);
+	}
 };
 
 template<typename KeyType, typename Hasher = DefaultKeyFuncs<KeyType, false>, typename HashMapAllocator = FDefaultAllocator>
