Fixed CSV profiler thread allocating temporary memory and potentially blocking the memory allocator.

diff --git a/Engine/Source/Runtime/Core/Private/ProfilingDebugging/CsvProfiler.cpp b/Engine/Source/Runtime/Core/Private/ProfilingDebugging/CsvProfiler.cpp
--- a/Engine/Source/Runtime/Core/Private/ProfilingDebugging/CsvProfiler.cpp
+++ b/Engine/Source/Runtime/Core/Private/ProfilingDebugging/CsvProfiler.cpp
@@ -2082,6 +2082,11 @@ public:
 	void Process(FCsvProcessThreadDataStats& OutStats, int32& OutMinFrameNumberProcessed);
 
 private:
+	/** Temporary storage of data collected with every Process() call. */
+	TArray<FCsvTimingMarker> ThreadMarkers;
+	TArray<FCsvCustomStat> CustomStats;
+	TArray<FCsvEvent> Events;
+
 	FCsvStatSeries* FindOrCreateStatSeries(const FCsvStatBase& Stat, FCsvStatSeries::EType SeriesType, bool bIsCountStat)
 	{
 		check(IsInCsvProcessingThread());
@@ -2409,9 +2414,9 @@ void FCsvProfilerThreadDataProcessor::Process(FCsvProcessThreadDataStats& OutSta
 	check(IsInCsvProcessingThread());
 
 	// Read the raw CSV data
-	TArray<FCsvTimingMarker> ThreadMarkers;
-	TArray<FCsvCustomStat> CustomStats;
-	TArray<FCsvEvent> Events;
+	ThreadMarkers.Reset(0);
+	CustomStats.Reset(0);
+	Events.Reset(0);
 	ThreadData->FlushResults(ThreadMarkers, CustomStats, Events);
 
 	OutStats.TimestampCount += ThreadMarkers.Num();
