Windows: Added an option to initialize COM in the MTA (multithreading) mode (it defaults to STA i.e. singlethreading mode as before)

#rnx

diff --git a/Engine/Source/Runtime/Core/Private/Windows/WindowsPlatformMisc.cpp b/Engine/Source/Runtime/Core/Private/Windows/WindowsPlatformMisc.cpp
--- a/Engine/Source/Runtime/Core/Private/Windows/WindowsPlatformMisc.cpp
+++ b/Engine/Source/Runtime/Core/Private/Windows/WindowsPlatformMisc.cpp
@@ -2035,9 +2035,9 @@ void FWindowsPlatformMisc::SetLastError(uint32 ErrorCode)
 	::SetLastError((DWORD)ErrorCode);
 }
 
-bool FWindowsPlatformMisc::CoInitialize()
+bool FWindowsPlatformMisc::CoInitialize(ECOMModel Model)
 {
-	HRESULT hr = ::CoInitialize(NULL);
+	HRESULT hr = ::CoInitializeEx(NULL, (Model == ECOMModel::Singlethreaded) ? COINIT_APARTMENTTHREADED : COINIT_MULTITHREADED);
 	return hr == S_OK || hr == S_FALSE;
 }
 
diff --git a/Engine/Source/Runtime/Core/Public/Windows/WindowsPlatformMisc.h b/Engine/Source/Runtime/Core/Public/Windows/WindowsPlatformMisc.h
--- a/Engine/Source/Runtime/Core/Public/Windows/WindowsPlatformMisc.h
+++ b/Engine/Source/Runtime/Core/Public/Windows/WindowsPlatformMisc.h
@@ -27,6 +27,16 @@ struct CORE_API FWindowsOSVersionHelper
 	static int32 GetOSVersions( FString& OutOSVersion, FString& OutOSSubVersion );
 };
 
+/**
+  * Determines the concurrency model to be set for a thread.
+  * 
+  * @see FWindowsPlatformMisc::CoInitialize
+  */
+enum class ECOMModel : uint8
+{
+	Singlethreaded = 0,		///< Single-Threaded Apartment (STA)
+	Multithreaded,			///< Multi-Threaded Apartment (MTA)
+};
 
 /**
 * Windows implementation of the misc OS functions
@@ -94,7 +104,7 @@ struct CORE_API FWindowsPlatformMisc
 	static bool DeleteStoredValue(const FString& InStoreId, const FString& InSectionName, const FString& InKeyName);
 	static bool DeleteStoredSection(const FString& InStoreId, const FString& InSectionName);
 
-	static bool CoInitialize();
+	static bool CoInitialize(ECOMModel Model = ECOMModel::Singlethreaded);
 	static void CoUninitialize();
 
 	/**
