Log out page size on startup

diff --git a/Engine/Source/Runtime/Core/Private/Unix/UnixPlatformMemory.cpp b/Engine/Source/Runtime/Core/Private/Unix/UnixPlatformMemory.cpp
--- a/Engine/Source/Runtime/Core/Private/Unix/UnixPlatformMemory.cpp
+++ b/Engine/Source/Runtime/Core/Private/Unix/UnixPlatformMemory.cpp
@@ -113,6 +113,8 @@ void FUnixPlatformMemory::Init()
 	UE_LOG(LogInit, Log, TEXT(" - VirtualMemoryAllocator pools will grow at scale %g"), GVMAPoolScale);
 	UE_LOG(LogInit, Log, TEXT(" - MemoryRangeDecommit() will %s"), 
 		GMemoryRangeDecommitIsNoOp ? TEXT("be a no-op (re-run with -vmapoolevict to change)") : TEXT("will evict the memory from RAM (re-run with -novmapoolevict to change)"));
+	UE_LOG(LogInit, Log, TEXT(" - PageSize %z"), MemoryConstants.PageSize);
+	UE_LOG(LogInit, Log, TEXT(" - BinnedPageSize %z"), MemoryConstants.BinnedPageSize);
 }
 
 bool FUnixPlatformMemory::HasForkPageProtectorEnabled()
