#ue5 - Revert FloorToInt(double) etc back to returning int32 instead of int64, and remove overloads of Max()/Min() that were required as part of that change (originally CL 18984296).

diff --git a/Engine/Source/Runtime/Core/Private/Tests/Math/UnrealMathTest.cpp b/Engine/Source/Runtime/Core/Private/Tests/Math/UnrealMathTest.cpp
--- a/Engine/Source/Runtime/Core/Private/Tests/Math/UnrealMathTest.cpp
+++ b/Engine/Source/Runtime/Core/Private/Tests/Math/UnrealMathTest.cpp
@@ -3072,8 +3072,9 @@ bool FVectorRegisterAbstractionTest::RunTest(const FString& Parameters)
 
 		int64 Big1 = 0, Big2 = 0, Big3 = 0;
 		Big3 = FMath::Max(Big1, Big2);
-		Big3 = FMath::Max(Big1, I2);
-		Big3 = FMath::Max(I1, Big2);
+		// Test compilation mixing int32/int64 types
+		//Big3 = FMath::Max(Big1, I2);
+		//Big3 = FMath::Max(I1, Big2);
 	}
 
 	// Sin, Cos, Tan tests
diff --git a/Engine/Source/Runtime/Core/Public/GenericPlatform/GenericPlatformMath.h b/Engine/Source/Runtime/Core/Public/GenericPlatform/GenericPlatformMath.h
--- a/Engine/Source/Runtime/Core/Public/GenericPlatform/GenericPlatformMath.h
+++ b/Engine/Source/Runtime/Core/Public/GenericPlatform/GenericPlatformMath.h
@@ -211,7 +211,7 @@ struct FGenericPlatformMath
 	}
 
 	static CONSTEXPR FORCEINLINE int32 TruncToInt(float F) { return TruncToInt32(F); }
-	static CONSTEXPR FORCEINLINE int64 TruncToInt(double F) { return TruncToInt64(F); }
+	static CONSTEXPR FORCEINLINE int32 TruncToInt(double F) { return TruncToInt32(F); }
 
 	/**
 	 * Converts a float to an integer value with truncation towards zero.
@@ -265,7 +265,7 @@ struct FGenericPlatformMath
 	}
 
 	static FORCEINLINE int32 FloorToInt(float F) { return FloorToInt32(F); }
-	static FORCEINLINE int64 FloorToInt(double F) { return FloorToInt64(F); }
+	static FORCEINLINE int32 FloorToInt(double F) { return FloorToInt32(F); }
 	
 	
 	/**
@@ -314,7 +314,7 @@ struct FGenericPlatformMath
 	}
 	
 	static FORCEINLINE int32 RoundToInt(float F) { return RoundToInt32(F); }
-	static FORCEINLINE int64 RoundToInt(double F) { return RoundToInt64(F); }
+	static FORCEINLINE int32 RoundToInt(double F) { return RoundToInt32(F); }
 
 	/**
 	* Converts a float to the nearest integer. Rounds up when the fraction is .5
@@ -368,7 +368,7 @@ struct FGenericPlatformMath
 	}
 
 	static FORCEINLINE int32 CeilToInt(float F) { return CeilToInt32(F); }
-	static FORCEINLINE int64 CeilToInt(double F) { return CeilToInt64(F); }
+	static FORCEINLINE int32 CeilToInt(double F) { return CeilToInt32(F); }
 
 	/**
 	* Converts a float to the nearest greater or equal integer.
@@ -934,9 +934,11 @@ struct FGenericPlatformMath
 		return (A<=B) ? A : B;
 	}
 
-	// Allow mixing of types to promote to highest precision type
-	MIX_TYPES_2_ARGS_CONSTEXPR(Max);
-	MIX_TYPES_2_ARGS_CONSTEXPR(Min);
+	// Allow mixing of float types to promote to highest precision type
+	static CONSTEXPR FORCEINLINE double Max(const double A, const float B) { return Max<double>(A, B); }
+	static CONSTEXPR FORCEINLINE double Max(const float A, const double B) { return Max<double>(A, B); }
+	static CONSTEXPR FORCEINLINE double Min(const double A, const float B) { return Min<double>(A, B); }
+	static CONSTEXPR FORCEINLINE double Min(const float A, const double B) { return Min<double>(A, B); }
 
 	/**
 	* Min of Array
diff --git a/Engine/Source/Runtime/Core/Public/Math/UnrealPlatformMathSSE.h b/Engine/Source/Runtime/Core/Public/Math/UnrealPlatformMathSSE.h
--- a/Engine/Source/Runtime/Core/Public/Math/UnrealPlatformMathSSE.h
+++ b/Engine/Source/Runtime/Core/Public/Math/UnrealPlatformMathSSE.h
@@ -274,13 +274,13 @@ struct TUnrealPlatformMathSSEBase : public Base
 	//
 
 	static FORCEINLINE int32 TruncToInt(float F) { return TruncToInt32(F); }
-	static FORCEINLINE int64 TruncToInt(double F) { return TruncToInt64(F); }
+	static FORCEINLINE int32 TruncToInt(double F) { return TruncToInt32(F); }
 	static FORCEINLINE int32 FloorToInt(float F) { return FloorToInt32(F); }
-	static FORCEINLINE int64 FloorToInt(double F) { return FloorToInt64(F); }
+	static FORCEINLINE int32 FloorToInt(double F) { return FloorToInt32(F); }
 	static FORCEINLINE int32 RoundToInt(float F) { return RoundToInt32(F); }
-	static FORCEINLINE int64 RoundToInt(double F) { return RoundToInt64(F); }
+	static FORCEINLINE int32 RoundToInt(double F) { return RoundToInt32(F); }
 	static FORCEINLINE int32 CeilToInt(float F) { return CeilToInt32(F); }
-	static FORCEINLINE int64 CeilToInt(double F) { return CeilToInt64(F); }
+	static FORCEINLINE int32 CeilToInt(double F) { return CeilToInt32(F); }
 
 
 	template<typename T>
