- Fixed a "Many failed FindCOnsoleObject" warning, when loading other platform DeviceProfiles - only look at the +CVars lines in the DPs, so we don't treat things like BaseProfileName and TextureLODGroups as cvars that it tries to find (TextureLODGroups was specified in so many DPs that it threw the warning)
- Also fixed a place of frequent "none" named CVar lookups

diff --git a/Engine/Source/Runtime/Core/Private/HAL/ConsoleManager.cpp b/Engine/Source/Runtime/Core/Private/HAL/ConsoleManager.cpp
--- a/Engine/Source/Runtime/Core/Private/HAL/ConsoleManager.cpp
+++ b/Engine/Source/Runtime/Core/Private/HAL/ConsoleManager.cpp
@@ -376,7 +376,9 @@ bool IConsoleManager::VisitPlatformCVarsForEmulation(FName PlatformName, bool bV
 	for (const FSectionPair& SectionPair : Sections)
 	{
 		FConfigSection* Section;
-		if (FCString::Strcmp(SectionPair.Name, DeviceProfileTag) == 0)
+		bool bIsDeviceProfile = FCString::Strcmp(SectionPair.Name, DeviceProfileTag) == 0;
+
+		if (bIsDeviceProfile)
 		{
 			// skip this if we don't want it
 			if (bVisitPlatformDeviceProfile == false)
@@ -400,17 +402,26 @@ bool IConsoleManager::VisitPlatformCVarsForEmulation(FName PlatformName, bool bV
 				FString Value = Pair.Value.GetValue();
 
 				// DPs have +CVars= prefix for the key, and the real cvar KVP is in the Value
-				if (Key.StartsWith(TEXT("CVars")))
+				// and other keys are not CVars so are ignored
+				if (bIsDeviceProfile)
 				{
-					FString ValueCopy = Value;
-					// expect a second = in the value, skip any that don't
-					if (!ValueCopy.Split(TEXT("="), &Key, &Value))
+					if (Key.StartsWith(TEXT("CVars")))
+					{
+						FString ValueCopy = Value;
+						// expect a second = in the value, skip any that don't
+						if (!ValueCopy.Split(TEXT("="), &Key, &Value))
+						{
+							continue;
+						}
+					}
+					else
 					{
 						continue;
 					}
 				}
 
-				IConsoleVariable* CVar = IConsoleManager::Get().FindConsoleVariable(*Key);
+				// don't bother tracking when looking up other platform cvars
+				IConsoleVariable* CVar = IConsoleManager::Get().FindConsoleVariable(*Key, false /* bTrackFrequentCalls */);
 				EConsoleVariableFlags PreviewFlag = (CVar != nullptr) ? (EConsoleVariableFlags)(CVar->GetFlags() & ECVF_Preview) : ECVF_Default;
 
 				if (Key.StartsWith(TEXT("sg.")))
@@ -1424,15 +1435,25 @@ IConsoleObject* FConsoleManager::FindConsoleObject(const TCHAR* Name, bool bTrac
 				UE_LOG(LogConsoleManager, Warning, TEXT("Performance warning: Console object named '%s' shows many (%d) FindConsoleObject() calls (consider caching e.g. using static)"), Name, CVar->FindCallCount);
 			}
 		}
-		else
+		else// if (bEarlyAppPhase || GFrameCounter == 1000)
 		{
 			static uint32 NullFindCallCount = 0;
-		
+			static TMap<FName, uint32> PerNameNullFindCallCount;
+
 			++NullFindCallCount;
+			FName CVarName(Name);
+			if (PerNameNullFindCallCount.FindOrAdd(Name)++ == 30)
+			{
+				UE_LOG(LogConsoleManager, Warning, TEXT("Performance warning: Many (%d) failed FindConsoleObject() for '%s'. "), PerNameNullFindCallCount[Name], Name);
+			}
 
-			if(bEarlyAppPhase && NullFindCallCount == 500)
+			if(NullFindCallCount == 500)
 			{
-				UE_LOG(LogConsoleManager, Warning, TEXT( "Performance warning: Many (%d) failed FindConsoleObject() e.g. '%s' (consider caching, is the name referencing an existing object)"), NullFindCallCount, Name);
+				UE_LOG(LogConsoleManager, Warning, TEXT( "Performance warning: Many (%d) failed FindConsoleObject() across all CVars. Fail counts per name:"), NullFindCallCount, Name);
+				for (TPair<FName, uint32> Entry : PerNameNullFindCallCount)
+				{
+					UE_LOG(LogConsoleManager, Warning, TEXT("   %s : %d"), *Entry.Key.ToString(), Entry.Value);
+				}
 			}
 		}
 	}
@@ -1440,7 +1461,7 @@ IConsoleObject* FConsoleManager::FindConsoleObject(const TCHAR* Name, bool bTrac
 
 	if(CVar && CVar->TestFlags(ECVF_CreatedFromIni))
 	{
-		return 0;
+		return nullptr;
 	}
 
 	return CVar;
diff --git a/Engine/Source/Runtime/Engine/Private/DataDrivenCVars/DataDrivenCVars.cpp b/Engine/Source/Runtime/Engine/Private/DataDrivenCVars/DataDrivenCVars.cpp
--- a/Engine/Source/Runtime/Engine/Private/DataDrivenCVars/DataDrivenCVars.cpp
+++ b/Engine/Source/Runtime/Engine/Private/DataDrivenCVars/DataDrivenCVars.cpp
@@ -6,7 +6,10 @@
 
 FDataDrivenConsoleVariable::~FDataDrivenConsoleVariable()
 {
-	UnRegister();
+	if (!Name.IsEmpty())
+	{
+		UnRegister();
+	}
 }
 
 void FDataDrivenConsoleVariable::Register()
@@ -37,12 +40,16 @@ void FDataDrivenConsoleVariable::Register()
 
 void FDataDrivenConsoleVariable::UnRegister(bool bUseShadowName)
 {
-	IConsoleVariable* CVarToRemove = IConsoleManager::Get().FindConsoleVariable(bUseShadowName  ? *ShadowName : *Name);
-	if (CVarToRemove)
+	const FString& NameToUnregister = bUseShadowName ? ShadowName : Name;
+	if (!NameToUnregister.IsEmpty())
 	{
-		FConsoleVariableDelegate NullCallback;
-		CVarToRemove->SetOnChangedCallback(NullCallback);
-		IConsoleManager::Get().UnregisterConsoleObject(CVarToRemove, false);
+		IConsoleVariable* CVarToRemove = IConsoleManager::Get().FindConsoleVariable(*NameToUnregister);
+		if (CVarToRemove)
+		{
+			FConsoleVariableDelegate NullCallback;
+			CVarToRemove->SetOnChangedCallback(NullCallback);
+			IConsoleManager::Get().UnregisterConsoleObject(CVarToRemove, false);
+		}
 	}
 }
 
