Inlined operator==(FNameEntryId Id, EName Ename) to allow the compiler to optimize out overhead in comparisons to NAME_None.

diff --git a/Engine/Source/Runtime/Core/Private/UObject/UnrealNames.cpp b/Engine/Source/Runtime/Core/Private/UObject/UnrealNames.cpp
--- a/Engine/Source/Runtime/Core/Private/UObject/UnrealNames.cpp
+++ b/Engine/Source/Runtime/Core/Private/UObject/UnrealNames.cpp
@@ -1304,11 +1304,6 @@ static FNamePool& GetNamePoolPostInit()
 	return (FNamePool&)NamePoolData;
 }
 
-bool operator==(FNameEntryId Id, EName Ename)
-{
-	return Ename == NAME_None ? !Id : Id == GetNamePoolPostInit().Find(Ename);
-}
-
 static int32 CompareDifferentIdsAlphabetically(FNameEntryId AId, FNameEntryId BId)
 {
 	checkSlow(AId != BId);
@@ -2800,6 +2795,10 @@ FNameEntryId FNameEntryId::FromValidEName(EName Ename)
 	return GetNamePool().Find(Ename);
 }
 
+FNameEntryId FNameEntryId::FromValidENamePostInit(EName Ename)
+{
+	return GetNamePoolPostInit().Find(Ename);
+}
 
 void FName::TearDown()
 {
diff --git a/Engine/Source/Runtime/Core/Public/UObject/NameTypes.h b/Engine/Source/Runtime/Core/Public/UObject/NameTypes.h
--- a/Engine/Source/Runtime/Core/Public/UObject/NameTypes.h
+++ b/Engine/Source/Runtime/Core/Public/UObject/NameTypes.h
@@ -75,14 +75,19 @@ struct FNameEntryId
 		return Ename == NAME_None ? FNameEntryId() : FromValidEName(Ename);
 	}
 
+	friend inline bool operator==(FNameEntryId Id, EName Ename)
+	{
+		return Ename == NAME_None ? !Id : Id == FromValidENamePostInit(Ename);
+	}
+
 private:
 	uint32 Value;
 
 	CORE_API static FNameEntryId FromValidEName(EName Ename);
+	CORE_API static FNameEntryId FromValidENamePostInit(EName Ename);
 };
 
 CORE_API uint32 GetTypeHash(FNameEntryId Id);
-CORE_API bool operator==(FNameEntryId Id, EName Ename);
 inline bool operator==(EName Ename, FNameEntryId Id) { return Id == Ename; }
 inline bool operator!=(EName Ename, FNameEntryId Id) { return !(Id == Ename); }
 inline bool operator!=(FNameEntryId Id, EName Ename) { return !(Id == Ename); }
