Add support for %c width formatting in FGenericWidePlatformString. Supports + and -  followed by digits as well as * for integer value argument.

diff --git a/Engine/Source/Runtime/Core/Private/GenericPlatform/GenericWidePlatformString.cpp b/Engine/Source/Runtime/Core/Private/GenericPlatform/GenericWidePlatformString.cpp
--- a/Engine/Source/Runtime/Core/Private/GenericPlatform/GenericWidePlatformString.cpp
+++ b/Engine/Source/Runtime/Core/Private/GenericPlatform/GenericWidePlatformString.cpp
@@ -262,6 +262,12 @@ void RunGetVarArgsTests()
 
 	TestGetVarArgs(OutputString, TEXT("Test I|%" UINT64_FMT "|"), MAX_uint64);
 	checkf(FString(OutputString) == FString(TEXT("Test I|18446744073709551615|")), OutputString);
+
+	TestGetVarArgs(OutputString, TEXT("Test J|%*c|"), 10, 'J');
+	checkf(FString(OutputString) == FString(TEXT("Test J|         J|")), OutputString);
+
+	TestGetVarArgs(OutputString, TEXT("Test K|%-5c|"), 'K');
+	checkf(FString(OutputString) == FString(TEXT("Test K|K    |")), OutputString);
 }
 #endif
 
@@ -509,7 +515,17 @@ int32 FGenericWidePlatformString::GetVarArgs( WIDECHAR* Dest, SIZE_T DestSize, c
 			{
 				TCHAR Val = (TCHAR) va_arg(ArgPtr, int);
 				Src++;
-				if (!DestIter.Write(Val))
+				bool bSuccess = true;
+				if (FieldLen > 1)
+				{
+					bSuccess = !!DestIter.Write(TCHAR(' '), FieldLen - 1);
+				}
+				bSuccess = bSuccess && !!DestIter.Write(Val);
+				if (FieldLen < -1)
+				{
+					bSuccess = bSuccess && !!DestIter.Write(TCHAR(' '), FPlatformMath::Abs(FieldLen + 1));
+				}
+				if (!bSuccess)
 				{
 					return -1;
 				}
