Fix for VS2022 compilation errors with most recent compiler (17.1.0)

`operator new` needs to be forward declared

diff --git a/Engine/Source/Runtime/Core/Public/Containers/Array.h b/Engine/Source/Runtime/Core/Public/Containers/Array.h
--- a/Engine/Source/Runtime/Core/Public/Containers/Array.h
+++ b/Engine/Source/Runtime/Core/Public/Containers/Array.h
@@ -29,7 +29,6 @@
 
 #include <type_traits>
 
-
 #if UE_BUILD_SHIPPING || UE_BUILD_TEST
 	#define TARRAY_RANGED_FOR_CHECKS 0
 #else
@@ -41,6 +40,10 @@ template<> struct TCanBulkSerialize<unsigned int> { enum { Value = true }; };
 template<> struct TCanBulkSerialize<unsigned short> { enum { Value = true }; };
 template<> struct TCanBulkSerialize<int> { enum { Value = true }; };
 
+// Forward declarations
+
+template <typename T, typename AllocatorType> inline void* operator new(size_t Size, TArray<T, AllocatorType>& Array);
+template <typename T, typename AllocatorType> inline void* operator new(size_t Size, TArray<T, AllocatorType>& Array, typename TArray<T, AllocatorType>::SizeType Index);
 
 /**
  * Generic iterator which can operate on types that expose the following:
diff --git a/Engine/Source/Runtime/Core/Public/Containers/ChunkedArray.h b/Engine/Source/Runtime/Core/Public/Containers/ChunkedArray.h
--- a/Engine/Source/Runtime/Core/Public/Containers/ChunkedArray.h
+++ b/Engine/Source/Runtime/Core/Public/Containers/ChunkedArray.h
@@ -46,6 +46,14 @@ namespace UE4ChunkedArray_Private
 	};
 }
 
+// Forward declarations
+
+template<typename InElementType, uint32 TargetBytesPerChunk, typename AllocatorType>
+class TChunkedArray;
+
+template <typename T, uint32 TargetBytesPerChunk, typename AllocatorType>
+void* operator new(size_t Size, TChunkedArray<T, TargetBytesPerChunk, AllocatorType>& ChunkedArray);
+
 /** An array that uses multiple allocations to avoid allocation failure due to fragmentation. */
 template<typename InElementType, uint32 TargetBytesPerChunk = 16384, typename AllocatorType = FDefaultAllocator >
 class TChunkedArray
diff --git a/Engine/Source/Runtime/Core/Public/Containers/SparseArray.h b/Engine/Source/Runtime/Core/Public/Containers/SparseArray.h
--- a/Engine/Source/Runtime/Core/Public/Containers/SparseArray.h
+++ b/Engine/Source/Runtime/Core/Public/Containers/SparseArray.h
@@ -26,11 +26,6 @@
 	#define TSPARSEARRAY_RANGED_FOR_CHECKS 1
 #endif
 
-// Forward declarations.
-template<typename ElementType,typename Allocator = FDefaultSparseArrayAllocator >
-class TSparseArray;
-
-
 /** The result of a sparse array allocation. */
 struct FSparseArrayAllocationInfo
 {
@@ -38,6 +33,14 @@ struct FSparseArrayAllocationInfo
 	void* Pointer;
 };
 
+// Forward declarations.
+template<typename ElementType,typename Allocator = FDefaultSparseArrayAllocator >
+class TSparseArray;
+
+template <typename T, typename Allocator> void* operator new(size_t Size, TSparseArray<T, Allocator>& Array);
+template <typename T, typename Allocator> void* operator new(size_t Size, TSparseArray<T, Allocator>& Array, int32 Index);
+inline void* operator new(size_t Size, const FSparseArrayAllocationInfo& Allocation);
+
 /** Allocated elements are overlapped with free element info in the element list. */
 template<typename ElementType>
 union TSparseArrayElementOrFreeListLink
