Adding a pointer accessor to TOptional

diff --git a/Engine/Source/Runtime/Core/Public/Misc/Optional.h b/Engine/Source/Runtime/Core/Public/Misc/Optional.h
--- a/Engine/Source/Runtime/Core/Public/Misc/Optional.h
+++ b/Engine/Source/Runtime/Core/Public/Misc/Optional.h
@@ -200,6 +200,10 @@ public:
 	/** @return The optional value when set; DefaultValue otherwise. */
 	const OptionalType& Get(const OptionalType& DefaultValue) const { return IsSet() ? *(OptionalType*)&Value : DefaultValue; }
 
+	/** @return A pointer to the optional value when set, nullptr otherwise. */
+	OptionalType* GetPtrOrNull() { return IsSet() ? (OptionalType*)&Value : nullptr; }
+	const OptionalType* GetPtrOrNull() const { return IsSet() ? (const OptionalType*)&Value : nullptr; }
+
 private:
 	TTypeCompatibleBytes<OptionalType> Value;
 	bool bIsSet;
