Fix for FGenericWidePlatformString asserting and going into an infinite recursion when setting PLATFORM_USE_GENERIC_STRING_IMPLEMENTATION=1 in Windows.

diff --git a/Engine/Source/Runtime/Core/Private/GenericPlatform/GenericWidePlatformString.cpp b/Engine/Source/Runtime/Core/Private/GenericPlatform/GenericWidePlatformString.cpp
--- a/Engine/Source/Runtime/Core/Private/GenericPlatform/GenericWidePlatformString.cpp
+++ b/Engine/Source/Runtime/Core/Private/GenericPlatform/GenericWidePlatformString.cpp
@@ -79,8 +79,10 @@ WIDECHAR* FGenericWidePlatformString::Strcat(WIDECHAR* Dest, SIZE_T DestCount, c
 
 int32 FGenericWidePlatformString::Strtoi( const WIDECHAR* Start, WIDECHAR** End, int32 Base )
 {
-#if !PLATFORM_TCHAR_IS_CHAR16
+#if PLATFORM_TCHAR_IS_4_BYTES || PLATFORM_TCHAR_IS_UTF8CHAR
 	unimplemented();
+#else
+	static_assert(sizeof(TCHAR) == 2, "TCHAR is expected to be 16-bit");
 #endif
 
 	if (End == nullptr)
@@ -106,8 +108,10 @@ int32 FGenericWidePlatformString::Strtoi( const WIDECHAR* Start, WIDECHAR** End,
 
 int64 FGenericWidePlatformString::Strtoi64( const WIDECHAR* Start, WIDECHAR** End, int32 Base )
 {
-#if !PLATFORM_TCHAR_IS_CHAR16
+#if PLATFORM_TCHAR_IS_4_BYTES || PLATFORM_TCHAR_IS_UTF8CHAR
 	unimplemented();
+#else
+	static_assert(sizeof(TCHAR) == 2, "TCHAR is expected to be 16-bit");
 #endif
 
 	if (End == nullptr)
@@ -394,8 +398,10 @@ namespace
 
 int32 FGenericWidePlatformString::GetVarArgs( WIDECHAR* Dest, SIZE_T DestSize, const WIDECHAR*& Fmt, va_list ArgPtr )
 {
-#if !PLATFORM_TCHAR_IS_CHAR16
+#if PLATFORM_TCHAR_IS_4_BYTES || PLATFORM_TCHAR_IS_UTF8CHAR
 	unimplemented();
+#else
+	static_assert(sizeof(TCHAR) == 2, "TCHAR is expected to be 16-bit");
 #endif
 
 #if !(UE_BUILD_SHIPPING || UE_BUILD_TEST)
