LLM: Handle implicitly-defined parent tags that are not previously declared.
#rnx

diff --git a/Engine/Source/Runtime/Core/Private/HAL/LowLevelMemTracker.cpp b/Engine/Source/Runtime/Core/Private/HAL/LowLevelMemTracker.cpp
--- a/Engine/Source/Runtime/Core/Private/HAL/LowLevelMemTracker.cpp
+++ b/Engine/Source/Runtime/Core/Private/HAL/LowLevelMemTracker.cpp
@@ -1554,6 +1554,7 @@ UE::LLMPrivate::FTagData& FLowLevelMemTracker::RegisterTagData(FName Name, FName
 		// (e.g. LLM_DECLARE_TAG). If a formal registration exists for a tag, it must precede its use
 		// in any LLM_SCOPE calls.
 		if (ReferenceSource != UE::LLMPrivate::ETagReferenceSource::Scope &&
+			ReferenceSource != UE::LLMPrivate::ETagReferenceSource::ImplicitParent &&
 			ReferenceSource != UE::LLMPrivate::ETagReferenceSource::FunctionAPI)
 		{
 			ReportDuplicateTagName(TagDataForName, ReferenceSource);
@@ -1607,7 +1608,7 @@ void FLowLevelMemTracker::ReportDuplicateTagName(UE::LLMPrivate::FTagData* TagDa
 {
 	using namespace UE::LLMPrivate;
 
-	if (ReferenceSource == ETagReferenceSource::FunctionAPI || ReferenceSource == ETagReferenceSource::Scope)
+	if (ReferenceSource == ETagReferenceSource::FunctionAPI || ReferenceSource == ETagReferenceSource::Scope || ReferenceSource == ETagReferenceSource::ImplicitParent)
 	{
 		LLMCheckf(false, TEXT("LLM Error: Unexpected call to RegisterTagData(%s) from LLM_SCOPE or function call when the tag already exists."), *TagData->GetName().ToString());
 	}
@@ -1615,6 +1616,11 @@ void FLowLevelMemTracker::ReportDuplicateTagName(UE::LLMPrivate::FTagData* TagDa
 	{
 		LLMCheckf(false, TEXT("LLM Error: Tag %s parsed from %s after it was already used in an LLM_SCOPE or LLM api call."), *TagData->GetName().ToString(), ToString(ReferenceSource));
 	}
+	else if (TagData->GetReferenceSource() == ETagReferenceSource::ImplicitParent)
+	{
+		LLMCheckf(false, TEXT("LLM Error: Tag %s parsed from %s after it was already used as an implicit parent in another tag. Add LLM_DEFINE_TAG(%s) in cpp, or move it to the same module as the child tag using it, so that it will be defined before the child tag tries to use it."),
+			*TagData->GetName().ToString(), ToString(ReferenceSource), *TagData->GetName().ToString());
+	}
 	else
 	{
 		LLMCheckf(false, TEXT("LLM Error: Multiple occurrences of a unique tag name %s in ELLMTag or LLM_DEFINE_TAG. First occurrence: %s. Second occurrence: %s."),
@@ -1650,11 +1656,15 @@ void FLowLevelMemTracker::FinishConstruct(UE::LLMPrivate::FTagData* TagData, UE:
 			FTagData** ParentDataPtr = TagDataNameMap->Find(ParentName);
 			if (!ParentDataPtr)
 			{
-				const TCHAR* SourceName = ToString(ReferenceSource);
-				// We have to drop the lock so we can allocate strings and call log functions
+				// We have to drop the ReadLock to call RegisterTagData, which takes a WriteLock
 				TagDataLock.ReadUnlock();
-				UE_LOG(LogHAL, Error, TEXT("LLM Parent tag %s was not available when child tag %s was used in %s"), *ParentName.ToString(), *TagData->GetName().ToString(), ToString(ReferenceSource));
+				FTagData* ParentTagData = &RegisterTagData(ParentName, NAME_None, NAME_None, NAME_None, NAME_None, false, ELLMTag::CustomName, false, ETagReferenceSource::ImplicitParent);
 				TagDataLock.ReadLock();
+				if (TagData->IsFinishConstructed())
+				{
+					// Another thread got in and finished construction while we were outside of the lock
+					return;
+				}
 				ParentDataPtr = TagDataNameMap->Find(GetTagName_CustomName());
 				LLMCheck(ParentDataPtr);
 			}
@@ -2310,6 +2320,8 @@ namespace LLMPrivate
 			return TEXT("RegisterPlatformTag/RegisterProjectTag");
 		case ETagReferenceSource::FunctionAPI:
 			return TEXT("DefaultName/InternalCall");
+		case ETagReferenceSource::ImplicitParent:
+			return TEXT("ImplicitParent");
 		default:
 			return TEXT("Invalid");
 		}
diff --git a/Engine/Source/Runtime/Core/Public/HAL/LowLevelMemTracker.h b/Engine/Source/Runtime/Core/Public/HAL/LowLevelMemTracker.h
--- a/Engine/Source/Runtime/Core/Public/HAL/LowLevelMemTracker.h
+++ b/Engine/Source/Runtime/Core/Public/HAL/LowLevelMemTracker.h
@@ -461,7 +461,8 @@ namespace LLMPrivate
 		Declare,
 		EnumTag,
 		CustomEnumTag,
-		FunctionAPI
+		FunctionAPI,
+		ImplicitParent
 	};
 }
 }
