Adding timestamps to child fork signals
Adding optional child process timeout for WAIT_AND_FORK_RESPONSE_SIGNAL
Added logging if we fail to set the child command line

diff --git a/Engine/Source/Runtime/Core/Private/Unix/UnixPlatformProcess.cpp b/Engine/Source/Runtime/Core/Private/Unix/UnixPlatformProcess.cpp
--- a/Engine/Source/Runtime/Core/Private/Unix/UnixPlatformProcess.cpp
+++ b/Engine/Source/Runtime/Core/Private/Unix/UnixPlatformProcess.cpp
@@ -30,6 +30,8 @@
 #include "Misc/CoreDelegates.h"
 #include "Misc/App.h"
 #include "Misc/Fork.h"
+#include "HAL/PlatformTime.h"
+#include "Misc/DateTime.h"
 
 namespace PlatformProcessLimits
 {
@@ -1343,14 +1345,26 @@ FGenericPlatformProcess::EWaitAndForkResult FUnixPlatformProcess::WaitAndFork()
 #ifndef WAIT_AND_FORK_PARENT_SHUTDOWN_EXIT_CODE
 	#define WAIT_AND_FORK_PARENT_SHUTDOWN_EXIT_CODE 0
 #endif
-
+#ifndef WAIT_AND_FORK_RESPONSE_TIMEOUT_EXIT_CODE
+	#define WAIT_AND_FORK_RESPONSE_TIMEOUT_EXIT_CODE 1
+#endif
+	
 	// Only works in -nothreading mode for now (probably best this way)
 	if (FPlatformProcess::SupportsMultithreading())
 	{
 		return EWaitAndForkResult::Error;
 	}
 
-	static TCircularQueue<int32> WaitAndForkSignalQueue(WAIT_AND_FORK_QUEUE_LENGTH);
+	struct FForkSignalData
+	{
+		FForkSignalData() = default;
+		FForkSignalData(int32 InSignal, double InTimeSeconds) : SignalValue(InSignal), TimeSeconds(InTimeSeconds) {}
+
+		int32 SignalValue = 0;
+		double TimeSeconds = 0.0;
+	};
+
+	static TCircularQueue<FForkSignalData> WaitAndForkSignalQueue(WAIT_AND_FORK_QUEUE_LENGTH);
 
 	// If we asked to fork up front without the need to send signals, just push the fork requests on the queue and we will refork them if they close
 	// This is mostly used in cases where there is no external process sending signals to this process to create forks and is a simple way to start or test
@@ -1360,7 +1374,7 @@ FGenericPlatformProcess::EWaitAndForkResult FUnixPlatformProcess::WaitAndFork()
 	{
 		for (int32 ForkIdx = 0; ForkIdx < NumForks; ++ForkIdx)
 		{
-			WaitAndForkSignalQueue.Enqueue(ForkIdx + 1);
+			WaitAndForkSignalQueue.Enqueue(FForkSignalData(ForkIdx + 1, FPlatformTime::Seconds()));
 		}
 	}
 
@@ -1379,6 +1393,13 @@ FGenericPlatformProcess::EWaitAndForkResult FUnixPlatformProcess::WaitAndFork()
 	// If we are asked to wait for a response signal, keep track of that here so we can behave differently in children.
 	const bool bRequireResponseSignal = FParse::Param(FCommandLine::Get(), TEXT("WaitAndForkRequireResponse"));
 
+	double WaitAndForkResponseTimeout = -1.0;
+	FParse::Value(FCommandLine::Get(), TEXT("-WaitAndForkResponseTimeout="), WaitAndForkResponseTimeout);
+	if (WaitAndForkResponseTimeout > 0.0)
+	{
+		UE_LOG(LogHAL, Log, TEXT("WaitAndFork setting WaitAndForkResponseTimeout to %0.2f seconds."), WaitAndForkResponseTimeout);
+	}
+
 	// Set up a signal handler for the signal to fork()
 	{
 		struct sigaction Action;
@@ -1388,7 +1409,7 @@ FGenericPlatformProcess::EWaitAndForkResult FUnixPlatformProcess::WaitAndFork()
 		Action.sa_sigaction = [](int32 Signal, siginfo_t* Info, void* Context) {
 			if (Signal == WAIT_AND_FORK_QUEUE_SIGNAL && Info)
 			{
-				WaitAndForkSignalQueue.Enqueue(Info->si_value.sival_int);
+				WaitAndForkSignalQueue.Enqueue(FForkSignalData(Info->si_value.sival_int, FPlatformTime::Seconds()));
 			}
 		};
 		sigaction(WAIT_AND_FORK_QUEUE_SIGNAL, &Action, nullptr);
@@ -1429,15 +1450,18 @@ FGenericPlatformProcess::EWaitAndForkResult FUnixPlatformProcess::WaitAndFork()
 	{
 		BeginExitIfRequested();
 
-		int32 SignalValue = 0;
-		if (WaitAndForkSignalQueue.Dequeue(SignalValue))
+		FForkSignalData SignalData;
+		if (WaitAndForkSignalQueue.Dequeue(SignalData))
 		{
 			// Sleep for a short while to avoid spamming new processes to the OS all at once
 			FPlatformProcess::Sleep(WAIT_AND_FORK_CHILD_SPAWN_DELAY);
 
-			uint16 Cookie = (SignalValue >> 16) & 0xffff;
-			uint16 ChildIdx = SignalValue & 0xffff;
-			UE_LOG(LogHAL, Log, TEXT("[Parent] WaitAndFork processing child request %04hx-%04hx."), Cookie, ChildIdx);
+			uint16 Cookie = (SignalData.SignalValue >> 16) & 0xffff;
+			uint16 ChildIdx = SignalData.SignalValue & 0xffff;
+
+			FDateTime SignalReceived = FDateTime::FromUnixTimestamp(FMath::FloorToInt64(SignalData.TimeSeconds));
+
+			UE_LOG(LogHAL, Log, TEXT("[Parent] WaitAndFork processing child request %04hx-%04hx received at: %s"), Cookie, ChildIdx, *SignalReceived.ToString());
 
 			FMemoryStatsHolder CurrentMasterMemStats(FPlatformMemory::GetStats());
 			UE_LOG(LogHAL, Log, TEXT("MemoryStats PreFork: AvailablePhysical: %.02fMiB (%+.02fMiB), PeakPhysical: %.02fMiB, PeakVirtual: %.02fMiB"),
@@ -1491,6 +1515,10 @@ FGenericPlatformProcess::EWaitAndForkResult FUnixPlatformProcess::WaitAndFork()
 					{
 						FCommandLine::Set(*NewCmdLine);
 					}
+					else
+					{
+						UE_LOG(LogHAL, Error, TEXT("[Child] WaitAndFork child %04hx-%04hx failed to set command line from: %s"), Cookie, ChildIdx, *CmdLineFilename);
+					}
 				}
 
 				// Start up the log again
@@ -1503,7 +1531,7 @@ FGenericPlatformProcess::EWaitAndForkResult FUnixPlatformProcess::WaitAndFork()
 					if (prctl(PR_SET_NAME, TCHAR_TO_UTF8(*FString::Printf(TEXT("DS-%04hx-%04hx"), Cookie, ChildIdx))) != 0)
 					{
 						int ErrNo = errno;
-						UE_LOG(LogHAL, Fatal, TEXT("WaitAndFork failed to set process name with prctl! error:%d"), ErrNo);
+						UE_LOG(LogHAL, Fatal, TEXT("[Child] WaitAndFork failed to set process name with prctl! error:%d"), ErrNo);
 					}
 				}
 
@@ -1526,10 +1554,20 @@ FGenericPlatformProcess::EWaitAndForkResult FUnixPlatformProcess::WaitAndFork()
 					};
 					sigaction(WAIT_AND_FORK_RESPONSE_SIGNAL, &Action, nullptr);
 
+					const double StartChildWaitSeconds = FPlatformTime::Seconds();
+
 					UE_LOG(LogHAL, Log, TEXT("[Child] WaitAndFork child %04hx-%04hx waiting for signal %d to proceed."), Cookie, ChildIdx, WAIT_AND_FORK_RESPONSE_SIGNAL);
 					while (!IsEngineExitRequested() && !bResponseReceived)
 					{
 						FPlatformProcess::Sleep(1);
+
+						// Check to see how long we've been waiting and if we should time out.
+						if ((WaitAndForkResponseTimeout > 0.0) && ((FPlatformTime::Seconds() - StartChildWaitSeconds) > WaitAndForkResponseTimeout))
+						{
+							UE_LOG(LogHAL, Error, TEXT("[Child] WaitAndFork child %04hx-%04hx has exceeded WAIT_AND_FORK_RESPONSE_SIGNAL timeout"), Cookie, ChildIdx);
+							FPlatformMisc::RequestExitWithStatus(true, WAIT_AND_FORK_RESPONSE_TIMEOUT_EXIT_CODE);
+							break;
+						}
 					}
 
 					FMemory::Memzero(Action);
@@ -1546,7 +1584,7 @@ FGenericPlatformProcess::EWaitAndForkResult FUnixPlatformProcess::WaitAndFork()
 			else
 			{
 				// Parent
-				AllChildren.Emplace(ChildPID, SignalValue);
+				AllChildren.Emplace(ChildPID, SignalData.SignalValue);
 
 				UE_LOG(LogHAL, Log, TEXT("[Parent] WaitAndFork Successfully processed request %04hx-%04hx, made a child with pid %d! Total number of children: %d."), Cookie, ChildIdx, ChildPID, AllChildren.Num());
 			}
@@ -1579,7 +1617,7 @@ FGenericPlatformProcess::EWaitAndForkResult FUnixPlatformProcess::WaitAndFork()
 					else if (NumForks > 0 && ChildPidAndSignal.SignalValue > 0 && ChildPidAndSignal.SignalValue <= NumForks)
 					{
 						UE_LOG(LogHAL, Log, TEXT("[Parent] WaitAndFork child %d missing. This was NumForks child %d. Relaunching..."), ChildPidAndSignal.Pid, ChildPidAndSignal.SignalValue);
-						WaitAndForkSignalQueue.Enqueue(ChildPidAndSignal.SignalValue);
+						WaitAndForkSignalQueue.Enqueue(FForkSignalData(ChildPidAndSignal.SignalValue, FPlatformTime::Seconds()));
 					}
 					else
 					{
