MallocBinned2 fork support.
Clear caches and update a canary value across the fork boundary to avoid touching pages that were allocated pre-fork to minimize unnecessary duplication of shared pages.
Savings depend on pre-fork fragmentation/slack and number of concurrent child processes.

diff --git a/Engine/Source/Runtime/Core/Private/HAL/MallocBinned2.cpp b/Engine/Source/Runtime/Core/Private/HAL/MallocBinned2.cpp
--- a/Engine/Source/Runtime/Core/Private/HAL/MallocBinned2.cpp
+++ b/Engine/Source/Runtime/Core/Private/HAL/MallocBinned2.cpp
@@ -94,7 +94,7 @@ int32 RecursionCounter = 0;
 // Block sizes are based around getting the maximum amount of allocations per pool, with as little alignment waste as possible.
 // Block sizes should be close to even divisors of the system page size, and well distributed.
 // They must be 16-byte aligned as well.
-static uint16 SmallBlockSizes[] =
+static constexpr uint16 SmallBlockSizes[] =
 {
 	16, 32, 48, 64, 80, 96, 128,
 	160, 192, 224, 256, 288, 320, 384, 448,
@@ -525,7 +525,8 @@ struct FMallocBinned2::Private
 				}
 				else
 				{
-					check(NodePool->FirstFreeBlock->Canary == 0 || NodePool->FirstFreeBlock->IsCanaryOk());
+					// If we are freeing memory in this pool it must have the current canary and not the pre-fork one. All caches should have been cleared when forking.
+					check(NodePool->FirstFreeBlock->CanaryAndForkState == EBlockCanary::Zero || NodePool->FirstFreeBlock->CanaryAndForkState == Allocator.CurrentCanary);
 				}
 
 				// Free a pooled allocation.
@@ -533,7 +534,7 @@ struct FMallocBinned2::Private
 				Free->NumFreeBlocks = 1;
 				Free->NextFreeBlock = NodePool->FirstFreeBlock;
 				Free->BlockSize     = InBlockSize;
-				Free->Canary = FFreeBlock::CANARY_VALUE;
+				Free->CanaryAndForkState = Allocator.CurrentCanary;
 				Free->PoolIndex = InPoolIndex;
 				NodePool->FirstFreeBlock   = Free;
 
@@ -593,6 +594,19 @@ struct FMallocBinned2::Private
 #endif
 #if BINNED2_ALLOCATOR_STATS
 		FMallocBinned2::FPerThreadFreeBlockLists::ConsolidatedMemory += FreeBlockLists->AllocatedMemory;
+#endif
+	}
+
+	static void CheckThreadFreeBlockListsForFork()
+	{
+#if BINNED2_FORK_SUPPORT
+		if(GMallocBinned2PerThreadCaches)
+		{
+			FScopeLock Lock(&GetFreeBlockListsRegistrationMutex());
+			TArray<FPerThreadFreeBlockLists*>& List = GetRegisteredFreeBlockLists();
+			UE_CLOG(List.Num() == 1 && List[0] != FPerThreadFreeBlockLists::Get(), LogMemory, Fatal, TEXT("There was a thread-local free list at fork time which did not belong to the main forking thread. No other threads should be alive at fork time. If threads are spawned before forking, they must be killed and FMallocBinned2::ClearAndDisableTLSCachesOnCurrentThread() must be called."));
+			UE_CLOG(List.Num() > 1, LogMemory, Fatal, TEXT("There were multiple thread-local free lists at fork time. No other threads should be alive at fork time. If threads are spawned before forking, they must be killed and FMallocBinned2::ClearAndDisableTLSCachesOnCurrentThread() must be called."));
+		}
 #endif
 	}
 };
@@ -603,6 +617,11 @@ FMallocBinned2::Private::FGlobalRecycler FMallocBinned2::Private::GGlobalRecycle
 int64 FMallocBinned2::FPerThreadFreeBlockLists::ConsolidatedMemory = 0;
 #endif
 
+FORCEINLINE void FMallocBinned2::FPoolList::Clear()
+{
+	Front = nullptr;
+}
+
 FORCEINLINE bool FMallocBinned2::FPoolList::IsEmpty() const
 {
 	return Front == nullptr;
@@ -638,11 +657,11 @@ FMallocBinned2::FPoolInfo& FMallocBinned2::FPoolList::PushNewPoolToFront(FMalloc
 		Private::OutOfMemory(LocalPageSize);
 	}
 #if !UE_USE_VERYLARGEPAGEALLOCATOR || !BINNED2_BOOKKEEPING_AT_THE_END_OF_LARGEBLOCK
-	FFreeBlock* Free = new (FreePtr) FFreeBlock(LocalPageSize, InBlockSize, InPoolIndex);
+	FFreeBlock* Free = new (FreePtr) FFreeBlock(LocalPageSize, InBlockSize, InPoolIndex, Allocator.CurrentCanary);
 	check(IsAligned(Free, LocalPageSize));
 #else
 	FFreeBlock* FreeBlockPtr = GetPoolHeaderFromPointer(FreePtr);
-	FFreeBlock* Free = new (FreeBlockPtr) FFreeBlock(LocalPageSize, InBlockSize, InPoolIndex);
+	FFreeBlock* Free = new (FreeBlockPtr) FFreeBlock(LocalPageSize, InBlockSize, InPoolIndex, Allocator.CurrentCanary);
 #endif
 #if BINNED2_ALLOCATOR_STATS
 	AllocatedOSSmallPoolMemory += (int64)LocalPageSize;
@@ -662,6 +681,9 @@ FMallocBinned2::FMallocBinned2()
 	static bool bOnce = false;
 	check(!bOnce); // this is now a singleton-like thing and you cannot make multiple copies
 	bOnce = true;
+	
+	static_assert(sizeof(FFreeBlock) <= SmallBlockSizes[0], "sizeof(FFreeBlock)  must be fit in smallest allocation size handled by FMallocBinned2.");
+	static_assert(sizeof(FMallocBinned2::FBundleNode) <= SmallBlockSizes[0], "sizeof(FBundleNode) must fit in smallest allocation size handled by FMallocBinned2.");
 
 	for (uint32 Index = 0; Index != BINNED2_SMALL_POOL_COUNT; ++Index)
 	{
@@ -737,6 +759,59 @@ FMallocBinned2::~FMallocBinned2()
 {
 }
 
+void FMallocBinned2::OnPreFork()
+{
+#if BINNED2_FORK_SUPPORT
+	// Trim caches so we don't use them in the child process and cause pages to be copied
+	if (GMallocBinned2PerThreadCaches)
+	{
+		FlushCurrentThreadCache();
+		FMallocBinned2::Private::CheckThreadFreeBlockListsForFork();
+	}
+
+	for (int32 PoolIndex = 0; PoolIndex < UE_ARRAY_COUNT(SmallPoolTables); ++PoolIndex)
+	{
+		while (FBundleNode* Node = FMallocBinned2::Private::GGlobalRecycler.PopBundle(PoolIndex))
+		{
+			Node->NextBundle = nullptr; // We need to override this because it's in a union with Count which is not needed for the freeing work
+			FMallocBinned2::Private::FreeBundles(*this, Node, SmallPoolTables[PoolIndex].BlockSize, PoolIndex);
+		}
+	}
+
+	FScopeLock Lock(&Mutex);
+#if !UE_USE_VERYLARGEPAGEALLOCATOR
+	CachedOSPageAllocator.FreeAll(&Mutex);
+#endif
+
+#endif // BINNED2_FORK_SUPPORT
+}
+
+void FMallocBinned2::OnPostFork()
+{
+#if BINNED2_FORK_SUPPORT
+	if (GMallocBinned2PerThreadCaches)
+	{
+		FlushCurrentThreadCache();
+		FMallocBinned2::Private::CheckThreadFreeBlockListsForFork();
+	}
+
+	FScopeLock Lock(&Mutex);
+
+	// This will be compared against the pool header of existing allocations to turn Free into a no-op for pages shared with the parent process
+	UE_CLOG(CurrentCanary != EBlockCanary::PreFork, LogMemory, Fatal, TEXT("FMallocBinned2 only supports forking once!"));
+
+	OldCanary = CurrentCanary;
+	CurrentCanary = EBlockCanary::PostFork;
+
+	for (FPoolTable& Table : SmallPoolTables)
+	{
+		// Clear our list of partially used pages so we don't dirty them and cause them to become unshared with the parent process
+		Table.ActivePools.Clear();
+		Table.ExhaustedPools.Clear();
+	}
+#endif
+}
+
 bool FMallocBinned2::IsInternallyThreadSafe() const
 { 
 	return true;
@@ -863,10 +938,16 @@ void* FMallocBinned2::ReallocExternal(void* Ptr, SIZE_T NewSize, uint32 Alignmen
 		check(Ptr); // null is 64k aligned so we should not be here
 		// Reallocate to a smaller/bigger pool if necessary
 		FFreeBlock* Free = GetPoolHeaderFromPointer(Ptr);
-		Free->CanaryTest();
+		CanaryTest(Free);
 		uint32 BlockSize = Free->BlockSize;
 		uint32 PoolIndex = Free->PoolIndex;
-		if (
+#if BINNED2_FORK_SUPPORT
+		// If the canary is the pre-fork one, we should not allow this allocation to grow in-place to avoid copying a page from the parent process.
+		const bool bIsOldAlloc = Free->CanaryAndForkState != CurrentCanary;
+#else
+		constexpr bool bIsOldAlloc = false;
+#endif
+		if ( !bIsOldAlloc && 
 			((NewSize <= BlockSize) & (Alignment <= BINNED2_MINIMUM_ALIGNMENT)) && // one branch, not two
 			(PoolIndex == 0 || NewSize > PoolIndexToBlockSize(PoolIndex - 1)))
 		{
@@ -943,7 +1024,15 @@ void FMallocBinned2::FreeExternal(void* Ptr)
 	{
 		check(Ptr); // null is 64k aligned so we should not be here
 		FFreeBlock* BasePtr = GetPoolHeaderFromPointer(Ptr);
-		BasePtr->CanaryTest();
+		CanaryTest(BasePtr);
+
+#if BINNED2_FORK_SUPPORT
+		if (BasePtr->CanaryAndForkState != CurrentCanary)
+		{
+			// This page was allocated before we forked so we want to avoid dirtying it by writing a linked list into it 
+			return;
+		}
+#endif
 		uint32 BlockSize = BasePtr->BlockSize;
 		uint32 PoolIndex = BasePtr->PoolIndex;
 
@@ -1006,7 +1095,7 @@ bool FMallocBinned2::GetAllocationSizeExternal(void* Ptr, SIZE_T& SizeOut)
 	{
 		check(Ptr); // null is 64k aligned so we should not be here
 		const FFreeBlock* Free = GetPoolHeaderFromPointer(Ptr);
-		Free->CanaryTest();
+		CanaryTest(Free);
 		uint32 BlockSize = Free->BlockSize;
 		SizeOut = BlockSize;
 		return true;
@@ -1257,9 +1346,26 @@ void FMallocBinned2::FPerThreadFreeBlockLists::ClearTLS()
 	FPlatformTLS::SetTlsValue(FMallocBinned2::Binned2TlsSlot, nullptr);
 }
 
-void FMallocBinned2::FFreeBlock::CanaryFail() const
+void FMallocBinned2::CanaryTest(const FFreeBlock* Block) const
 {
-	UE_LOG(LogMemory, Fatal, TEXT("FMallocBinned2 Attempt to realloc an unrecognized block %p   canary == 0x%x != 0x%x"), (void*)this, (int32)Canary, (int32)FMallocBinned2::FFreeBlock::CANARY_VALUE);
+#if BINNED2_FORK_SUPPORT
+	// When we support forking there are two valid canary values.
+	if (Block->CanaryAndForkState != CurrentCanary && Block->CanaryAndForkState != OldCanary)
+#else
+	if (Block->CanaryAndForkState != CurrentCanary)
+#endif
+	{
+		CanaryFail(Block);
+	}
+}
+
+void FMallocBinned2::CanaryFail(const FFreeBlock* Block) const
+{
+#if BINNED2_FORK_SUPPORT
+	UE_LOG(LogMemory, Fatal, TEXT("FMallocBinned2 Attempt to realloc an unrecognized block %p   canary == 0x%x != 0x%x or 0x%x "), (void*)Block, (int32)Block->CanaryAndForkState, CurrentCanary, OldCanary);
+#else
+	UE_LOG(LogMemory, Fatal, TEXT("FMallocBinned2 Attempt to realloc an unrecognized block %p   canary == 0x%x != 0x%x"), (void*)Block, (int32)Block->CanaryAndForkState, CurrentCanary);
+#endif
 }
 
 #if BINNED2_ALLOCATOR_STATS
diff --git a/Engine/Source/Runtime/Core/Private/HAL/UnrealMemory.cpp b/Engine/Source/Runtime/Core/Private/HAL/UnrealMemory.cpp
--- a/Engine/Source/Runtime/Core/Private/HAL/UnrealMemory.cpp
+++ b/Engine/Source/Runtime/Core/Private/HAL/UnrealMemory.cpp
@@ -217,6 +217,16 @@ public:
 	{
 		return UsedMalloc->GetDescriptiveName();
 	}
+
+	virtual void OnPreFork() override
+	{
+		UsedMalloc->OnPreFork();
+	}
+
+	virtual void OnPostFork() override
+	{
+		UsedMalloc->OnPostFork();
+	}
 };
 
 void FMemory::EnablePurgatoryTests()
diff --git a/Engine/Source/Runtime/Core/Private/Misc/Fork.cpp b/Engine/Source/Runtime/Core/Private/Misc/Fork.cpp
--- a/Engine/Source/Runtime/Core/Private/Misc/Fork.cpp
+++ b/Engine/Source/Runtime/Core/Private/Misc/Fork.cpp
@@ -128,3 +128,19 @@ bool FForkProcessHelper::SupportsMultithreadingPostFork()
 	return bSupportsMT;
 #endif
 }
+
+void FForkProcessHelper::LowLevelPreFork()
+{
+	GMalloc->OnPreFork();
+}
+
+void FForkProcessHelper::LowLevelPostForkParent()
+{
+	// Currently nothing to do here, just provided for completeness. 
+}
+
+void FForkProcessHelper::LowLevelPostForkChild(uint16 ChildIndex)
+{
+	FForkProcessHelper::SetIsForkedChildProcess(ChildIndex);
+	GMalloc->OnPostFork();
+}
diff --git a/Engine/Source/Runtime/Core/Private/ProfilingDebugging/Microsoft/MemoryTrace.cpp b/Engine/Source/Runtime/Core/Private/ProfilingDebugging/Microsoft/MemoryTrace.cpp
--- a/Engine/Source/Runtime/Core/Private/ProfilingDebugging/Microsoft/MemoryTrace.cpp
+++ b/Engine/Source/Runtime/Core/Private/ProfilingDebugging/Microsoft/MemoryTrace.cpp
@@ -92,6 +92,8 @@ private:
 	virtual bool			ValidateHeap() override										{ return InnerMalloc->ValidateHeap(); }
 	virtual bool			GetAllocationSize(void* Address, SIZE_T &SizeOut) override	{ return InnerMalloc->GetAllocationSize(Address, SizeOut); }
 	virtual void			SetupTLSCachesOnCurrentThread() override					{ return InnerMalloc->SetupTLSCachesOnCurrentThread(); }
+	virtual void			OnPreFork() override										{ InnerMalloc->OnPreFork();}
+	virtual void			OnPostFork() override										{ InnerMalloc->OnPostFork(); }
 
 	FMalloc*				InnerMalloc;
 };
diff --git a/Engine/Source/Runtime/Core/Private/Unix/UnixPlatformProcess.cpp b/Engine/Source/Runtime/Core/Private/Unix/UnixPlatformProcess.cpp
--- a/Engine/Source/Runtime/Core/Private/Unix/UnixPlatformProcess.cpp
+++ b/Engine/Source/Runtime/Core/Private/Unix/UnixPlatformProcess.cpp
@@ -1474,6 +1474,8 @@ FGenericPlatformProcess::EWaitAndForkResult FUnixPlatformProcess::WaitAndFork()
 			// Make sure there are no pending messages in the log.
 			GLog->Flush();
 
+			// This should be the very last thing we do before forking for optimal interaction with GMalloc
+			FForkProcessHelper::LowLevelPreFork();
 			// ******** The fork happens here! ********
 			pid_t ChildPID = fork();
 			// ******** The fork happened! This is now either the parent process or the new child process ********
@@ -1489,8 +1491,8 @@ FGenericPlatformProcess::EWaitAndForkResult FUnixPlatformProcess::WaitAndFork()
 			}
 			else if (ChildPID == 0)
 			{
-				// Child
-				FForkProcessHelper::SetIsForkedChildProcess(ChildIdx);
+				// This should be the very first thing we do after forking for optimal interaction with GMalloc
+				FForkProcessHelper::LowLevelPostForkChild(ChildIdx);
 
 				if (FPlatformMemory::HasForkPageProtectorEnabled())
 				{
@@ -1583,6 +1585,9 @@ FGenericPlatformProcess::EWaitAndForkResult FUnixPlatformProcess::WaitAndFork()
 			}
 			else
 			{
+				// This should be the very first thing we do after forking for optimal interaction with GMalloc
+				FForkProcessHelper::LowLevelPostForkParent();
+
 				// Parent
 				AllChildren.Emplace(ChildPID, SignalData.SignalValue);
 
diff --git a/Engine/Source/Runtime/Core/Public/HAL/MallocBinned2.h b/Engine/Source/Runtime/Core/Public/HAL/MallocBinned2.h
--- a/Engine/Source/Runtime/Core/Public/HAL/MallocBinned2.h
+++ b/Engine/Source/Runtime/Core/Public/HAL/MallocBinned2.h
@@ -51,6 +51,11 @@
 #define BINNED2_BOOKKEEPING_AT_THE_END_OF_LARGEBLOCK 1
 #endif
 
+// If we are emulating forking on a windows server or are a linux server, enable support for avoiding dirtying pages owned by the parent. 
+#ifndef BINNED2_FORK_SUPPORT
+	#define BINNED2_FORK_SUPPORT (UE_SERVER && (PLATFORM_UNIX || DEFAULT_SERVER_FAKE_FORKS))
+#endif
+
 
 #define BINNED2_ALLOW_RUNTIME_TWEAKING 0
 #if BINNED2_ALLOW_RUNTIME_TWEAKING
@@ -97,6 +102,20 @@ extern FCriticalSection ValidationCriticalSection;
 extern int32 RecursionCounter;
 #endif
 
+// Canary value used in FFreeBlock
+// A constant value unless we're compiled with fork support in which case there are two values identifying whether the page
+// was allocated pre- or post-fork
+enum class EBlockCanary : uint8
+{
+	Zero = 0x0, // Not clear why this is needed by FreeBundles
+#if BINNED2_FORK_SUPPORT
+	PreFork = 0xb7,
+	PostFork = 0xca,
+#else
+	Value = 0xe3 
+#endif
+};
+
 
 //
 // Optimized virtual memory allocator.
@@ -112,15 +131,10 @@ class CORE_API FMallocBinned2 : public FMalloc
 	/** Information about a piece of free memory. */
 	struct FFreeBlock
 	{
-		enum
-		{
-			CANARY_VALUE = 0xe3
-		};
-
-		FORCEINLINE FFreeBlock(uint32 InPageSize, uint16 InBlockSize, uint8 InPoolIndex)
+		FORCEINLINE FFreeBlock(uint32 InPageSize, uint16 InBlockSize, uint8 InPoolIndex, EBlockCanary InCanary)
 			: BlockSize(InBlockSize)
 			, PoolIndex(InPoolIndex)
-			, Canary(CANARY_VALUE)
+			, CanaryAndForkState(InCanary)
 			, NextFreeBlock(nullptr)
 		{
 			check(InPoolIndex < MAX_uint8 && InBlockSize <= MAX_uint16);
@@ -136,20 +150,7 @@ class CORE_API FMallocBinned2 : public FMalloc
 		{
 			return NumFreeBlocks;
 		}
-		FORCEINLINE bool IsCanaryOk() const
-		{
-			return Canary == FFreeBlock::CANARY_VALUE;
-		}
 
-		FORCEINLINE void CanaryTest() const
-		{
-			if (!IsCanaryOk())
-			{
-				CanaryFail();
-			}
-			//checkSlow(PoolIndex == BoundSizeToPoolIndex(BlockSize));
-		}
-		void CanaryFail() const;
 
 		FORCEINLINE void* AllocateRegularBlock()
 		{
@@ -173,7 +174,12 @@ class CORE_API FMallocBinned2 : public FMalloc
 
 		uint16 BlockSize;				// Size of the blocks that this list points to
 		uint8 PoolIndex;				// Index of this pool
-		uint8 Canary;					// Constant value of 0xe3
+
+		// Normally this value just functions as a canary to detect invalid memory state.
+		// When process forking is supported, it's still a canary but it has two valid values.
+		// One value is used pre-fork and one post-fork and the value is used to avoid freeing memory in pages shared with the parent process.
+		EBlockCanary CanaryAndForkState; 
+
 		uint32 NumFreeBlocks;          // Number of consecutive free blocks here, at least 1.
 		void*  NextFreeBlock;          // Next free block in another pool
 	};
@@ -182,6 +188,7 @@ class CORE_API FMallocBinned2 : public FMalloc
 	{
 		FPoolList();
 
+		void Clear();
 		bool IsEmpty() const;
 
 		      FPoolInfo& GetFrontPool();
@@ -266,6 +273,12 @@ class CORE_API FMallocBinned2 : public FMalloc
 	PoolHashBucket* HashBuckets;
 	PoolHashBucket* HashBucketFreeList;
 	uint64 NumPoolsPerPage;
+#if BINNED2_FORK_SUPPORT
+	EBlockCanary CurrentCanary = EBlockCanary::PreFork; // The value of the canary for pages we have allocated this side of the fork 
+	EBlockCanary OldCanary = EBlockCanary::PreFork;		// If we have forked, the value canary of old pages we should avoid touching 
+#else 
+	static constexpr EBlockCanary CurrentCanary = EBlockCanary::Value;
+#endif
 
 #if !PLATFORM_UNIX && !PLATFORM_ANDROID
 #if UE_USE_VERYLARGEPAGEALLOCATOR
@@ -288,9 +301,13 @@ class CORE_API FMallocBinned2 : public FMalloc
 #endif
 	}
 
+	// This needs to be small enough to fit inside the smallest allocation handled by MallocBinned2, hence the union.
 	struct FBundleNode
 	{
 		FBundleNode* NextNodeInCurrentBundle;
+
+		// NextBundle ptr is valid when node is stored in FFreeBlockList in a thread-local list of reusable allocations.
+		// Count is valid when node is stored in global recycler and caches the number of nodes in the list formed by NextNodeInCurrentBundle.
 		union
 		{
 			FBundleNode* NextBundle;
@@ -331,7 +348,6 @@ class CORE_API FMallocBinned2 : public FMalloc
 		FBundleNode* Head;
 		uint32       Count;
 	};
-	static_assert(sizeof(FBundleNode) <= BINNED2_MINIMUM_ALIGNMENT, "Bundle nodes must fit into the smallest block size");
 
 	struct FFreeBlockList
 	{
@@ -395,7 +411,8 @@ class CORE_API FMallocBinned2 : public FMalloc
 #if BINNED2_ALLOCATOR_STATS
 			: AllocatedMemory(0) 
 #endif
-		{ }
+		{ 
+		}
 
 		FORCEINLINE void* Malloc(uint32 InPoolIndex)
 		{
@@ -607,7 +624,8 @@ if (NewSize <= BINNED2_MAX_SMALL_POOL_SIZE && Alignment <= BINNED2_MINIMUM_ALIGN
 					FFreeBlock* Free = GetPoolHeaderFromPointer(Ptr);
 					BlockSize = Free->BlockSize;
 					PoolIndex = Free->PoolIndex;
-					bCanFree = Free->IsCanaryOk();
+					// If canary is invalid we will assert in ReallocExternal. Otherwise it's the pre-fork canary and we will allocate new memory without touching this allocation.
+					bCanFree = Free->CanaryAndForkState == CurrentCanary;
 					if (NewSize && bCanFree && NewSize <= BlockSize && (PoolIndex == 0 || NewSize > PoolIndexToBlockSize(PoolIndex - 1)))
 					{
 						return Ptr;
@@ -684,7 +702,8 @@ if (NewSize <= BINNED2_MAX_SMALL_POOL_SIZE && Alignment <= BINNED2_MINIMUM_ALIGN
 			{
 				FFreeBlock* BasePtr = GetPoolHeaderFromPointer(Ptr);
 				int32 BlockSize = BasePtr->BlockSize;
-				if (BasePtr->IsCanaryOk() && Lists->Free(Ptr, BasePtr->PoolIndex, BasePtr->BlockSize))
+				// If canary is invalid we will assert in FreeExternal. Otherwise it's the pre-fork canary and we will turn this free into a no-op.
+				if (BasePtr->CanaryAndForkState == CurrentCanary && Lists->Free(Ptr, BasePtr->PoolIndex, BasePtr->BlockSize))
 				{
 #if BINNED2_ALLOCATOR_STATS
 					Lists->AllocatedMemory -= BasePtr->BlockSize;
@@ -700,7 +719,11 @@ if (NewSize <= BINNED2_MAX_SMALL_POOL_SIZE && Alignment <= BINNED2_MINIMUM_ALIGN
 		if (!IsOSAllocation(Ptr))
 		{
 			const FFreeBlock* Free = GetPoolHeaderFromPointer(Ptr);
-			if (Free->IsCanaryOk())
+#if BINNED2_FORK_SUPPORT
+			if (Free->CanaryAndForkState == CurrentCanary || Free->CanaryAndForkState == OldCanary)
+#else
+			if (Free->CanaryAndForkState == CurrentCanary)
+#endif
 			{
 				SizeOut = Free->BlockSize;
 				return true;
@@ -734,6 +757,8 @@ if (NewSize <= BINNED2_MAX_SMALL_POOL_SIZE && Alignment <= BINNED2_MINIMUM_ALIGN
 	virtual void ClearAndDisableTLSCachesOnCurrentThread() override;
 	virtual const TCHAR* GetDescriptiveName() override;
 	virtual void UpdateStats() override;
+	virtual void OnPreFork() override;
+	virtual void OnPostFork() override;
 	// End FMalloc interface.
 
 	void FlushCurrentThreadCache();
@@ -743,6 +768,9 @@ if (NewSize <= BINNED2_MAX_SMALL_POOL_SIZE && Alignment <= BINNED2_MINIMUM_ALIGN
 	void FreeExternal(void *Ptr);
 	bool GetAllocationSizeExternal(void* Ptr, SIZE_T& SizeOut);
 
+	void CanaryTest(const FFreeBlock* Block) const;
+	void CanaryFail(const FFreeBlock* Block) const;
+
 #if BINNED2_ALLOCATOR_STATS
 	int64 GetTotalAllocatedSmallPoolMemory() const;
 #endif
diff --git a/Engine/Source/Runtime/Core/Public/HAL/MallocCallstackHandler.h b/Engine/Source/Runtime/Core/Public/HAL/MallocCallstackHandler.h
--- a/Engine/Source/Runtime/Core/Public/HAL/MallocCallstackHandler.h
+++ b/Engine/Source/Runtime/Core/Public/HAL/MallocCallstackHandler.h
@@ -131,6 +131,16 @@ public:
 		return UsedMalloc->GetDescriptiveName();
 	}
 
+	virtual void OnPreFork() override
+	{
+		UsedMalloc->OnPreFork();
+	}
+
+	virtual void OnPostFork() override
+	{
+		UsedMalloc->OnPostFork();
+	}
+
 	static constexpr SIZE_T MaxCallStackDepth = 64;
 	static constexpr SIZE_T CallStackEntriesToSkipCount = 2;
 
diff --git a/Engine/Source/Runtime/Core/Public/HAL/MallocPoisonProxy.h b/Engine/Source/Runtime/Core/Public/HAL/MallocPoisonProxy.h
--- a/Engine/Source/Runtime/Core/Public/HAL/MallocPoisonProxy.h
+++ b/Engine/Source/Runtime/Core/Public/HAL/MallocPoisonProxy.h
@@ -144,5 +144,16 @@ public:
 		UsedMalloc->ClearAndDisableTLSCachesOnCurrentThread();
 	}
 
+	virtual void OnPreFork() override
+	{
+		UsedMalloc->OnPreFork();
+	}
+
+
+	virtual void OnPostFork() override
+	{
+		UsedMalloc->OnPostFork();
+	}
+
 	// FMalloc interface end
 };
diff --git a/Engine/Source/Runtime/Core/Public/HAL/MallocThreadSafeProxy.h b/Engine/Source/Runtime/Core/Public/HAL/MallocThreadSafeProxy.h
--- a/Engine/Source/Runtime/Core/Public/HAL/MallocThreadSafeProxy.h
+++ b/Engine/Source/Runtime/Core/Public/HAL/MallocThreadSafeProxy.h
@@ -124,4 +124,14 @@ public:
 	{ 
 		return true; 
 	}
+
+	virtual void OnPreFork() override
+	{
+		UsedMalloc->OnPreFork();
+	}
+
+	virtual void OnPostFork() override
+	{
+		UsedMalloc->OnPostFork();
+	}
 };
diff --git a/Engine/Source/Runtime/Core/Public/HAL/MemoryBase.h b/Engine/Source/Runtime/Core/Public/HAL/MemoryBase.h
--- a/Engine/Source/Runtime/Core/Public/HAL/MemoryBase.h
+++ b/Engine/Source/Runtime/Core/Public/HAL/MemoryBase.h
@@ -205,6 +205,16 @@ public:
 		return TEXT("Unspecified allocator");
 	}
 
+	/**
+	 * Notifies the malloc implementation that the process is about to fork. May be used to trim caches etc.
+	 */
+	virtual void OnPreFork() {}
+
+	/**
+	 * Notifies the malloc implementation that the process has forked so we can try and avoid dirtying pre-fork pages.
+	 */
+	virtual void OnPostFork() {}
+
 protected:
 	friend struct FCurrentFrameCalls;
 
diff --git a/Engine/Source/Runtime/Core/Public/Misc/Fork.h b/Engine/Source/Runtime/Core/Public/Misc/Fork.h
--- a/Engine/Source/Runtime/Core/Public/Misc/Fork.h
+++ b/Engine/Source/Runtime/Core/Public/Misc/Fork.h
@@ -80,6 +80,28 @@ public:
 		uint64 InThreadAffinityMask = FPlatformAffinity::GetNoAffinityMask(),
 		EThreadCreateFlags InCreateFlags = EThreadCreateFlags::None
 	);
+
+	/**
+	 * Performs low-level cross-platform actions that should happen immediately BEFORE forking in a well-specified order.
+	 * Runs after any higher level code like calling into game-level constructs or anything that may allocate memory.
+	 * E.g. notifies GMalloc to optimize for memory sharing across parent/child process
+	 * Note: This will be called multiple times on the parent before each fork. 
+	 */
+	static void LowLevelPreFork();
+
+	/**
+	 * Performs low-level cross-platform actions that should happen immediately AFTER forking in the PARENT process in a well-specified order.
+	 * Runs before any higher level code like calling into game-level constructs.
+	 * E.g. notifies GMalloc to optimize for memory sharing across parent/child process
+	 */
+	static void LowLevelPostForkParent();
+
+	/**
+	 * Performs low-level cross-platform actions that should happen immediately AFTER forking in the CHILD process in a well-specified order.
+	 * Runs before any higher level code like calling into game-level constructs.
+	 * E.g. notifies GMalloc to optimize for memory sharing across parent/child process
+	 */
+	static void LowLevelPostForkChild(uint16 ChildIndex=1);
 };
 
 
diff --git a/Engine/Source/Runtime/Core/Public/ProfilingDebugging/TraceMalloc.h b/Engine/Source/Runtime/Core/Public/ProfilingDebugging/TraceMalloc.h
--- a/Engine/Source/Runtime/Core/Public/ProfilingDebugging/TraceMalloc.h
+++ b/Engine/Source/Runtime/Core/Public/ProfilingDebugging/TraceMalloc.h
@@ -14,6 +14,17 @@ public:
 	virtual void Free(void* Original) override;
 
 	static bool ShouldTrace();
+
+	virtual void OnPreFork() override
+	{
+		WrappedMalloc->OnPreFork();
+	}
+
+	virtual void OnPostFork() override
+	{
+		WrappedMalloc->OnPostFork();
+	}
+
 	FMalloc* WrappedMalloc;
 };
 
diff --git a/Engine/Source/Runtime/Core/Public/Stats/StatsMallocProfilerProxy.h b/Engine/Source/Runtime/Core/Public/Stats/StatsMallocProfilerProxy.h
--- a/Engine/Source/Runtime/Core/Public/Stats/StatsMallocProfilerProxy.h
+++ b/Engine/Source/Runtime/Core/Public/Stats/StatsMallocProfilerProxy.h
@@ -127,6 +127,16 @@ public:
 	{ 
 		return UsedMalloc->GetDescriptiveName(); 
 	}
+
+	virtual void OnPreFork() override
+	{
+		UsedMalloc->OnPreFork();
+	}
+
+	virtual void OnPostFork() override
+	{
+		UsedMalloc->OnPostFork();
+	}
 };
 
 
