Add TargetRules.bFNameOutlineNumber as an option to reduce the size of FName by 4 bytes by storing the number in the name table.
Adds a define UE_FNAME_OUTLINE_NUMBER.
Removes of FName, FMinimalName from memory image support. Adds of FMemoryImageName.
Removal of FMinimalName operator&lt;&lt;, all fields made private, size made variable.
All fields of FScriptName made private.
Added console commands for dumping numbered/unnumbered names and stats.

diff --git a/Engine/Extras/GDBPrinters/UEPrinters.py b/Engine/Extras/GDBPrinters/UEPrinters.py
--- a/Engine/Extras/GDBPrinters/UEPrinters.py
+++ b/Engine/Extras/GDBPrinters/UEPrinters.py
@@ -18,11 +18,13 @@
 #   end
 
 
-import gdb
 import itertools
+import random
 import re
 import sys
 
+import gdb
+
 # ------------------------------------------------------------------------------
 # We make our own base of the iterator to prevent issues between Python 2/3.
 #
@@ -51,7 +53,7 @@ def default_iterator(val):
 # FBitReference
 #
 class FBitReferencePrinter:
-	def __init__(self, typename, val):
+	def __init__(self, val):
 		self.Value = val
 
 	def to_string(self):
@@ -103,7 +105,7 @@ class TBitArrayPrinter:
 
 			return ('[%d]' % self.Counter, (data[self.Counter/32] >> self.Counter) & 1)
 
-	def __init__(self, typename, val):
+	def __init__(self, val):
 		self.Value = val
 		self.NumBits = self.Value['NumBits']
 
@@ -174,7 +176,7 @@ class TChunkedArrayPrinter:
 			Val = gdb.parse_and_eval(Expr)
 			return ('[%d]' % self.Counter, Val)
 
-	def __init__(self, typename, val):
+	def __init__(self, val):
 		self.Value = val
 		self.Typename = typename
 		self.NumElements = self.Value['NumElements']
@@ -252,7 +254,7 @@ class TSparseArrayPrinter:
 				else:
 					return self.__next__()
 
-	def __init__(self, typename, val):
+	def __init__(self, val):
 		self.Value = val
 		self.Typename = typename
 		self.ArrayNum = self.Value['Data']['ArrayNum']
@@ -328,7 +330,7 @@ class TSetPrinter:
 
 			return ('[%s]' % self.Counter, Value.dereference())
 
-	def __init__(self, typename, val):
+	def __init__(self, val):
 		self.Value = val
 		self.typename = typename
 		self.ArrayNum = self.Value["Elements"]["Data"]['ArrayNum']
@@ -353,7 +355,7 @@ class TSetPrinter:
 class TSetElementPrinter:
 	"Print TSetElement"
 
-	def __init__(self, typename, val):
+	def __init__(self, val):
 		self.Value = val
 
 	def to_string(self):
@@ -400,7 +402,7 @@ class TMapPrinter:
 
 			return ('Pairs', self.Pairs)
 
-	def __init__(self, typename, val):
+	def __init__(self, val):
 		self.Value = val
 		self.ArrayNum = self.Value['Pairs']['Elements']['Data']['ArrayNum']
 
@@ -457,7 +459,7 @@ class TWeakObjectPtrPrinter:
 				return ('Object', 'nullptr')
 
 
-	def __init__(self, typename, val):
+	def __init__(self, val):
 		self.Value = val
 
 	def children(self):
@@ -474,7 +476,7 @@ class TWeakObjectPtrPrinter:
 class FStringPrinter:
 	"Print FString"
 
-	def __init__(self, typename, val):
+	def __init__(self, val):
 		self.Value = val
 
 	def to_string(self):
@@ -494,6 +496,29 @@ class FStringPrinter:
 	def display_hint (self):
 		return 'string'
 
+# ------------------------------------------------------------------------------
+# FName shared 
+#
+def lookup_fname_entry(id):
+	expr = '((FNameEntry&)GNameBlocksDebug[%d >> FNameDebugVisualizer::OffsetBits][FNameDebugVisualizer::EntryStride * (%d & FNameDebugVisualizer::OffsetMask)])' % (id, id)
+	return gdb.parse_and_eval(expr)
+
+def get_fname_entry_string(entry):
+	header = entry['Header']
+	len = int(header['Len'].cast(gdb.lookup_type('uint16')))
+	is_wide = header['bIsWide'].cast(gdb.lookup_type('bool'))
+	if is_wide:
+		wide_string = entry['WideName'].cast(gdb.lookup_type('WIDECHAR').pointer())
+		return str(wide_string.string('','',len))
+	else:
+		ansi_string = entry['AnsiName'].cast(gdb.lookup_type('ANSICHAR').pointer())
+		return str(ansi_string.string('','',len))
+
+def get_fname_string(entry, number):
+	if number == 0:
+		return "'%s'" % get_fname_entry_string(entry)
+	else:
+		return "'%s'_%u" % (get_fname_entry_string(entry), number - 1)
 
 # ------------------------------------------------------------------------------
 # FNameEntry
@@ -502,80 +527,86 @@ class FStringPrinter:
 class FNameEntryPrinter:
 	"Print FNameEntry"
 
-	def __init__(self, typename, val):
+	def __init__(self, val):
 		self.Value = val
 
 	def to_string(self):
-		self.Header = self.Value['Header']
-		IsWideString = self.Header['bIsWide'].cast(gdb.lookup_type('bool'))
-		self.AnsiName = self.Value['AnsiName']
-		self.WideName = self.Value['WideName']
-		Len = int(self.Header['Len'].cast(gdb.lookup_type('uint16')))
-		if IsWideString == True:
-			WideString = self.WideName.cast(gdb.lookup_type('WIDECHAR').pointer())
-			return '%s' % WideString.string('','',Len)
+		header = self.Value['Header']
+		len = int(header['Len'].cast(gdb.lookup_type('uint16')))
+		if len == 0:
+			base_entry  = lookup_fname_entry(int(self.Value['NumberedName']['Id']['Value']))
+			number = self.Value['NumberedName']['Number']
+			return get_fname_string(base_entry, number)
 		else:
-			AnsiString = self.AnsiName.cast(gdb.lookup_type('ANSICHAR').pointer())
-			return '%s' % AnsiString.string('','',Len)
-
-	def display_hint (self):
-		return 'string'
+			return get_fname_string(self.Value, 0)
 
 # ------------------------------------------------------------------------------
-# FName
+# FNameEntryId
 #
-class FNamePrinter:
-	"Print FName"
 
-	def __init__(self, typename, val):
+class FNameEntryIdPrinter:
+	"Print FNameEntryId"
+
+	def __init__(self, val):
 		self.Value = val
 
 	def to_string(self):
 		if self.Value.is_optimized_out:
 			return '<optimized out>'
-
-		# ComparisonIndex is an FNameEntryId
-		Index = self.Value['ComparisonIndex']['Value']
-		IndexValue = int(Index)
-
-		UnusedMask = gdb.parse_and_eval('FNameDebugVisualizer::UnusedMask')
-		if (IndexValue & UnusedMask) != 0:
+		id = int(self.Value['Value'])
+		unused_mask = gdb.parse_and_eval('FNameDebugVisualizer::UnusedMask')
+		if (id & unused_mask) != 0:
 			return 'invalid'
-		else:
-			Expr = '((FNameEntry&)GNameBlocksDebug['+str(IndexValue)+' >> FNameDebugVisualizer::OffsetBits][FNameDebugVisualizer::EntryStride * ('+str(IndexValue)+' & FNameDebugVisualizer::OffsetMask)])'
-			NameEntry = gdb.parse_and_eval(Expr)
-			Number = self.Value['Number']
-			NumberValue = int(Number)
-			if NumberValue == 0:
-				return NameEntry
-			else:
-				return str([str(NameEntry), 'Number=' + str(NumberValue)])
+		return lookup_fname_entry(id)
 
 
 # ------------------------------------------------------------------------------
-# FMinimalName
+# FName
 #
-class FMinimalNamePrinter:
-	"Print FMinimalName"
+class FNamePrinter:
+	"Print FName"
 
-	def __init__(self, typename, val):
-		self.Value = val
+	def __init__(self, id, number):
+		self.Id = id
+		self.Number = number
 
 	def to_string(self):
-		Index = self.Value['Index']['Value']
-		IndexValue = int(Index)
+		return get_fname_string(lookup_fname_entry(self.Id), self.Number)
+
+
+def make_fname_printer(val):
+	if val.is_optimized_out:
+		return '<optimized out>'
+	id = int(val['ComparisonIndex']['Value'])
+	unused_mask = gdb.parse_and_eval('FNameDebugVisualizer::UnusedMask')
+	if (id & unused_mask) != 0:
+		return 'invalid'
+	
+	# We need to pick the right printer based on whether FName has a Number member or not
+	if gdb.types.has_field(val.type, "Number"):
+		return FNamePrinter(id, int(val['Number']))
+	else:
+		# look up the id and use the FNameEntry printer
+		return FNameEntryPrinter(lookup_fname_entry(id))
+	
 
-		if IndexValue >= 4194304:
-			return 'invalid'
-		else:
-			Expr = '((FNameEntry&)GNameBlocksDebug['+str(IndexValue)+' >> FNameDebugVisualizer::OffsetBits][FNameDebugVisualizer::EntryStride * ('+str(IndexValue)+' & FNameDebugVisualizer::OffsetMask)])'
-			NameEntry = gdb.parse_and_eval(Expr)
-			Number = self.Value['Number']
-			NumberValue = int(Number)
-			if NumberValue == 0:
-				return NameEntry
-			else:
-				return str([str(NameEntry), 'Number=' + str(NumberValue)])
+# ------------------------------------------------------------------------------
+# FMinimalName
+#
+def make_fminimalname_printer(val):
+	if val.is_optimized_out:
+		return '<optimized out>'
+	id = int(val['Index']['Value'])
+	unused_mask = gdb.parse_and_eval('FNameDebugVisualizer::UnusedMask')
+	if (id & unused_mask) != 0:
+		return 'invalid'
+	
+	# We need to pick the right printer based on whether FName has a Number member or not
+	if gdb.types.has_field(val.type, "Number"):
+		return FNamePrinter(id, int(val['Number']))
+	else:
+		# look up the id and use the FNameEntry printer
+		return FNameEntryIdPrinter(val['Index'])
 
 # ------------------------------------------------------------------------------
 # TTuple
@@ -583,7 +614,7 @@ class FMinimalNamePrinter:
 class TTuplePrinter:
 	"Print TTuple"
 
-	def __init__(self, typename, val):
+	def __init__(self, val):
 		self.Value = val
 
 		try:
@@ -652,7 +683,7 @@ class TArrayPrinter:
 
 			return ('[%d]' % self.Counter, data[self.Counter])
 
-	def __init__(self, typename, val):
+	def __init__(self, val):
 		self.Value = val;
 		self.ArrayNum = self.Value['ArrayNum']
 
@@ -673,54 +704,29 @@ class TArrayPrinter:
 # Register our lookup function. If no objfile is passed use all globally.
 def register_ue_printers(objfile):
 	if objfile == None:
-		objfile = gdb
-
-	objfile.pretty_printers.append(lookup_function)
-
-
-#
-# We need this part which is a definition how pretty printers work.
-#
-def lookup_function (val):
-	"Look-up and return a pretty-printer that can print val."
-
-	# Get the type and check if it points to a reference. We check for both Object and Pointer type.
-	type = val.type;
-	if (type.code == gdb.TYPE_CODE_REF) or (type.code == gdb.TYPE_CODE_PTR):
-		type = type.target ()
-
-	# Get the unqualified type, mean remove const nor volatile and strippe of typedefs.
-	type = type.unqualified ().strip_typedefs ()
-
-	# Get the tag name. The tag name is the name after struct, union, or enum in C and C++
-	tag = type.tag
-	if tag == None:
-		return None
-
-	for function in pretty_printers_dict:
-		if function.search (tag):
-			return pretty_printers_dict[function] (tag, val)
-
-	# Cannot find a pretty printer.  Return None.
-	return None
-
-def build_dictionary ():
-	pretty_printers_dict[re.compile('^FString$')] = lambda typename, val: FStringPrinter(typename, val)
-	pretty_printers_dict[re.compile('^FNameEntry$')] = lambda typename, val: FNameEntryPrinter(typename, val)
-	pretty_printers_dict[re.compile('^FName$')] = lambda typename, val: FNamePrinter(typename, val)
-	pretty_printers_dict[re.compile('^FMinimalName$')] = lambda typename, val: FMinimalNamePrinter(typename, val)
-	pretty_printers_dict[re.compile('^TArray<.+,.+>$')] = lambda typename, val: TArrayPrinter(typename, val)
-	pretty_printers_dict[re.compile('^TBitArray<.+>$')] = lambda typename, val: TBitArrayPrinter(typename, val)
-	pretty_printers_dict[re.compile('^TChunkedArray<.+>$')] = lambda typename, val: TChunkedArrayPrinter(typename, val)
-	pretty_printers_dict[re.compile('^TSparseArray<.+>$')] = lambda typename, val: TSparseArrayPrinter(typename, val)
-	pretty_printers_dict[re.compile('^TSetElement<.+>$')] = lambda typename, val: TSetElementPrinter(typename, val)
-#	pretty_printers_dict[re.compile('^TSet<.+>$')] = lambda typename, val: TSetPrinter(typename, val)
-	pretty_printers_dict[re.compile('^FBitReference$')] = lambda typename, val: FBitReferencePrinter(typename, val)
-#	pretty_printers_dict[re.compile('^TMap<.+,.+,.+>$')] = lambda typename, val: TMapPrinter(typename, val)
-	pretty_printers_dict[re.compile('^TPair<.+,.+>$')] = lambda typename, val: TTuplePrinter(typename, val)
-	pretty_printers_dict[re.compile('^TTuple<.+,.+>$')] = lambda typename, val: TTuplePrinter(typename, val)
-	pretty_printers_dict[re.compile('^TWeakObjectPtr<.+>$')] = lambda typename, val: TWeakObjectPtrPrinter(typename, val)
-
-
-pretty_printers_dict = {}
-build_dictionary ()
+		objfile = gdb.current_objfile()
+	gdb.printing.register_pretty_printer(objfile, build_ue_pretty_printer(), True)
+	print("Registered pretty printers for UE classes")
+
+def build_ue_pretty_printer():
+	# add a random numeric suffix to the printer name so we can reload printers during the same session for iteration
+	pp = gdb.printing.RegexpCollectionPrettyPrinter("UnrealEngine")
+	pp.add_printer("FString", '^FString$', FStringPrinter)
+	pp.add_printer("FNameEntry", '^FNameEntry$', FNameEntryPrinter)
+	pp.add_printer("FNameEntryId", '^FNameEntryId$', FNameEntryIdPrinter)
+	pp.add_printer("FName", '^FName$', make_fname_printer)
+	pp.add_printer("FMinimalName", '^FMinimalName$', make_fminimalname_printer)
+	pp.add_printer("TArray", '^TArray<.+,.+>$', TArrayPrinter)
+	pp.add_printer("TBitArray", '^TBitArray<.+>$', TBitArrayPrinter)
+	pp.add_printer("TChunkedArray", '^TChunkedArray<.+>$', TChunkedArrayPrinter)
+	pp.add_printer("TSparseArray", '^TSparseArray<.+>$', TSparseArrayPrinter)
+	pp.add_printer("TSetElement", '^TSetElement<.+>$', TSetElementPrinter)
+#	pp.add_printer("TSet", '^TSet<.+>$', TSetPrinter)
+	pp.add_printer("FBitReference", '^FBitReference$', FBitReferencePrinter)
+#	pp.add_printer("TMap", '^TMap<.+,.+,.+>$', TMapPrinter)
+	pp.add_printer("TPair", '^TPair<.+,.+>$', TTuplePrinter)
+	pp.add_printer("TTuple", '^TTuple<.+,.+>$', TTuplePrinter)
+	pp.add_printer("TWeakObjectPtr", '^TWeakObjectPtr<.+>$', TWeakObjectPtrPrinter)
+	return pp
+
+register_ue_printers(None)
diff --git a/Engine/Extras/LLDBDataFormatters/UEDataFormatters.py b/Engine/Extras/LLDBDataFormatters/UEDataFormatters.py
--- a/Engine/Extras/LLDBDataFormatters/UEDataFormatters.py
+++ b/Engine/Extras/LLDBDataFormatters/UEDataFormatters.py
@@ -56,19 +56,35 @@ def UEFStringSummaryProvider(valobj,dict):
         Val = str(ValRef.GetSummary())
         return 'string=' + Val
 
+def UEFNameIndexToEntry(EntryId):
+    Index = EntryId.GetChildMemberWithName('Value').GetValueAsUnsigned(0)
+    # FNameDebugVisualizer::OffsetBits = 16
+    # FNameDebugVisualizer::OffsetMask = (1 << OffsetBits) - 1 = 65535
+    NameEntryExpr = '(FNameEntry*)(GNameBlocksDebug['+str(Index)+' >> 16]+((alignof(FNameEntry) * ('+str(Index)+' & 65535))))'
+    NameEntry = EntryId.CreateValueFromExpression('NameEntry', NameEntryExpr)
+    return NameEntry
+
 def UEFNameEntrySummaryProvider(valobj,dict):
     Header = valobj.GetChildMemberWithName('Header')
-    DataPtr = valobj.GetChildAtIndex(1).AddressOf().GetValueAsUnsigned(0)
-    IsWide = Header.GetChildMemberWithName('bIsWide').GetValueAsUnsigned(0)
     Len = min(Header.GetChildMemberWithName('Len').GetValueAsUnsigned(0), 1023)
-    if IsWide:
-        SizeOfTChar = valobj.CreateValueFromExpression('size', 'sizeof(TCHAR))').GetValueAsUnsigned(0)
-        Encoding = "utf-16" if SizeOfTChar == 2 else "utf-32"
-        NumBytes = Len * SizeOfTChar
-        Data = valobj.process.ReadMemory(DataPtr,NumBytes,lldb.SBError())
-        Name = Data.decode(Encoding).encode("utf-8")
+    if Len == 0:
+        UnderlyingId = valobj.GetValueForExpressionPath('.NumberedName.Id')
+        Number = valobj.GetValueForExpressionPath('.NumberedName.Number').GetValueAsUnsigned(0)
+        BaseNameEntry = UEFNameIndexToEntry(UnderlyingId)
+        BaseName = UEFNameEntrySummaryProvider(BaseNameEntry, dict)
+        return '%s_%s' % (BaseName, Number-1) # BaseName will already include name= from recursion 
     else:
-        Name = valobj.process.ReadMemory(DataPtr,Len,lldb.SBError())
+        IsWide = Header.GetChildMemberWithName('bIsWide').GetValueAsUnsigned(0)
+        if IsWide:
+            DataPtr = valobj.GetValueForExpressionPath(".WideName").AddressOf().GetValueAsUnsigned(0)
+            SizeOfTChar = valobj.CreateValueFromExpression('size', 'sizeof(TCHAR))').GetValueAsUnsigned(0)
+            Encoding = "utf-16" if SizeOfTChar == 2 else "utf-32"
+            NumBytes = Len * SizeOfTChar
+            Data = valobj.process.ReadMemory(DataPtr,NumBytes,lldb.SBError())
+            Name = Data.decode(Encoding).encode("utf-8")
+        else:
+            DataPtr = valobj.GetValueForExpressionPath(".AnsiName").AddressOf().GetValueAsUnsigned(0)
+            Name = valobj.process.ReadMemory(DataPtr,Len,lldb.SBError())
     return 'name=%s' % Name
 
 def UEFNameSummaryProvider(valobj,dict):
@@ -77,18 +93,16 @@ def UEFNameSummaryProvider(valobj,dict):
         EntryId = valobj.GetChildMemberWithName('ComparisonIndex')
     if not EntryId.IsValid():
         EntryId = valobj.GetChildMemberWithName('Index')
-    Index = EntryId.GetChildMemberWithName('Value').GetValueAsUnsigned(0)
-    Number = valobj.GetChildMemberWithName('Number').GetValueAsUnsigned(0)
-
-    # FNameDebugVisualizer::OffsetBits = 16
-    # FNameDebugVisualizer::OffsetMask = (1 << OffsetBits) - 1 = 65535
-    NameEntryExpr = '(FNameEntry*)(GNameBlocksDebug['+str(Index)+' >> 16]+((alignof(FNameEntry) * ('+str(Index)+' & 65535))))'
-    NameEntry = valobj.CreateValueFromExpression('NameEntry', NameEntryExpr)
-    NameStr = UEFNameEntrySummaryProvider(NameEntry, dict)
-    if Number != 0:
-        return '%s_%s' % (NameStr, Number-1)
+    NameEntry = UEFNameIndexToEntry(EntryId)
+    Number = valobj.GetChildMemberWithName('Number')
+    if not Number.IsValid():
+        return UEFNameEntrySummaryProvider(NameEntry, dict)
     else:
-        return '%s' % NameStr
+        NameStr = UEFNameEntrySummaryProvider(NameEntry, dict)
+        if Number.GetValueAsUnsigned(0) != 0:
+            return "'%s'_%d" % (NameStr, Number-1)
+        else:
+            return "'%s'" % NameStr
 
 def UEUObjectBaseSummaryProvider(valobj,dict):
     Name = valobj.GetChildMemberWithName('NamePrivate')
diff --git a/Engine/Extras/LLDBDataFormatters/UEDataFormatters_2ByteChars.py b/Engine/Extras/LLDBDataFormatters/UEDataFormatters_2ByteChars.py
--- a/Engine/Extras/LLDBDataFormatters/UEDataFormatters_2ByteChars.py
+++ b/Engine/Extras/LLDBDataFormatters/UEDataFormatters_2ByteChars.py
@@ -77,19 +77,36 @@ def UEFStringSummaryProvider(valobj,dict):
         Val = ValRef.GetSummary()
         return 'string=' + Val
 
+def UEFNameIndexToEntry(EntryId):
+    Index = EntryId.GetChildMemberWithName('Value').GetValueAsUnsigned(0)
+    # FNameDebugVisualizer::OffsetBits = 16
+    # FNameDebugVisualizer::OffsetMask = (1 << OffsetBits) - 1 = 65535
+    NameEntryExpr = '(FNameEntry*)(GNameBlocksDebug['+str(Index)+' >> 16]+((alignof(FNameEntry) * ('+str(Index)+' & 65535))))'
+    NameEntry = EntryId.CreateValueFromExpression('NameEntry', NameEntryExpr)
+    return NameEntry
+
 def UEFNameEntrySummaryProvider(valobj,dict):
     Header = valobj.GetChildMemberWithName('Header')
-    DataPtr = valobj.GetChildAtIndex(1).AddressOf().GetValueAsUnsigned(0)
-    IsWide = Header.GetChildMemberWithName('bIsWide').GetValueAsUnsigned(0)
     Len = min(Header.GetChildMemberWithName('Len').GetValueAsUnsigned(0), 1023)
-    if IsWide:
-        SizeOfTChar = valobj.CreateValueFromExpression('size', 'sizeof(TCHAR))').GetValueAsUnsigned(0)
-        Encoding = "utf-16" if SizeOfTChar == 2 else "utf-32"
-        NumBytes = Len * SizeOfTChar
-        Data = valobj.process.ReadMemory(DataPtr,NumBytes,lldb.SBError())
-        Name = Data.decode(Encoding).encode("utf-8")
+    if Len == 0:
+        UnderlyingId = valobj.GetValueForExpressionPath('.NumberedName.Id')
+        Number = valobj.GetValueForExpressionPath('.NumberedName.Number').GetValueAsUnsigned(0)
+        BaseNameEntry = UEFNameIndexToEntry(UnderlyingId)
+        BaseName = UEFNameEntrySummaryProvider(BaseNameEntry, dict)
+        return '%s_%s' % (BaseName, Number-1) # BaseName will already include name= from recursion 
     else:
-        Name = valobj.process.ReadMemory(DataPtr,Len,lldb.SBError())
+        DataPtr = valobj.GetChildAtIndex(1).AddressOf().GetValueAsUnsigned(0)
+        IsWide = Header.GetChildMemberWithName('bIsWide').GetValueAsUnsigned(0)
+        if IsWide:
+            DataPtr = valobj.GetValueForExpressionPath(".WideName").AddressOf().GetValueAsUnsigned(0)
+            SizeOfTChar = valobj.CreateValueFromExpression('size', 'sizeof(TCHAR))').GetValueAsUnsigned(0)
+            Encoding = "utf-16" if SizeOfTChar == 2 else "utf-32"
+            NumBytes = Len * SizeOfTChar
+            Data = valobj.process.ReadMemory(DataPtr,NumBytes,lldb.SBError())
+            Name = Data.decode(Encoding).encode("utf-8")
+        else:
+            DataPtr = valobj.GetValueForExpressionPath(".AnsiName").AddressOf().GetValueAsUnsigned(0)
+            Name = valobj.process.ReadMemory(DataPtr,Len,lldb.SBError())
     return 'name=%s' % Name
 
 def UEFNameSummaryProvider(valobj,dict):
@@ -98,18 +115,16 @@ def UEFNameSummaryProvider(valobj,dict):
         EntryId = valobj.GetChildMemberWithName('ComparisonIndex')
     if not EntryId.IsValid():
         EntryId = valobj.GetChildMemberWithName('Index')
-    Index = EntryId.GetChildMemberWithName('Value').GetValueAsUnsigned(0)
-    Number = valobj.GetChildMemberWithName('Number').GetValueAsUnsigned(0)
-
-    # FNameDebugVisualizer::OffsetBits = 16
-    # FNameDebugVisualizer::OffsetMask = (1 << OffsetBits) - 1 = 65535
-    NameEntryExpr = '(FNameEntry*)(GNameBlocksDebug[%s >> 16]+((alignof(FNameEntry) * (%s & 65535))))' % (Index, Index)
-    NameEntry = valobj.CreateValueFromExpression('NameEntry', NameEntryExpr)
-    NameStr = UEFNameEntrySummaryProvider(NameEntry, dict)
-    if Number != 0:
-        return '%s_%s' % (NameStr, Number-1)
+    NameEntry = UEFNameIndexToEntry(EntryId)
+    Number = valobj.GetChildMemberWithName('Number')
+    if not Number.IsValid():
+        return UEFNameEntrySummaryProvider(NameEntry, dict)
     else:
-        return '%s' % NameStr
+        NameStr = UEFNameEntrySummaryProvider(NameEntry, dict)
+        if Number.GetValueAsUnsigned(0) != 0:
+            return "'%s'_%d" % (NameStr, Number-1)
+        else:
+            return "'%s'" % NameStr
 
 def UEUObjectBaseSummaryProvider(valobj,dict):
     Name = valobj.GetChildMemberWithName('NamePrivate')
diff --git a/Engine/Plugins/Developer/Concert/ConcertMain/Source/ConcertTransport/Private/IdentifierTable/ConcertIdentifierTable.cpp b/Engine/Plugins/Developer/Concert/ConcertMain/Source/ConcertTransport/Private/IdentifierTable/ConcertIdentifierTable.cpp
--- a/Engine/Plugins/Developer/Concert/ConcertMain/Source/ConcertTransport/Private/IdentifierTable/ConcertIdentifierTable.cpp
+++ b/Engine/Plugins/Developer/Concert/ConcertMain/Source/ConcertTransport/Private/IdentifierTable/ConcertIdentifierTable.cpp
@@ -59,7 +59,7 @@ void FConcertLocalIdentifierTable::SetState(const FConcertLocalIdentifierState&
 	NameToMappedIndex.Reserve(InState.MappedNames.Num());
 	for (const FString& MappedNameStr : InState.MappedNames)
 	{
-		const FName MappedName = FName(*MappedNameStr, NAME_NO_NUMBER_INTERNAL, FNAME_Add, /*bSplitName*/false);
+		const FName MappedName = FName(*MappedNameStr, NAME_NO_NUMBER_INTERNAL, /*bSplitName*/false);
 		const int32 NewIndex = MappedNames.Num();
 		MappedNames.Add(MappedName);
 		NameToMappedIndex.Add(MappedName, NewIndex);
diff --git a/Engine/Plugins/Developer/Concert/ConcertMain/Source/ConcertTransport/Private/IdentifierTable/ConcertTransportArchives.cpp b/Engine/Plugins/Developer/Concert/ConcertMain/Source/ConcertTransport/Private/IdentifierTable/ConcertTransportArchives.cpp
--- a/Engine/Plugins/Developer/Concert/ConcertMain/Source/ConcertTransport/Private/IdentifierTable/ConcertTransportArchives.cpp
+++ b/Engine/Plugins/Developer/Concert/ConcertMain/Source/ConcertTransport/Private/IdentifierTable/ConcertTransportArchives.cpp
@@ -101,7 +101,7 @@ FArchive& FConcertIdentifierReader::operator<<(FName& Name)
 		{
 			FString PlainString;
 			*this << PlainString;
-			Name = FName(*PlainString, NAME_NO_NUMBER_INTERNAL, FNAME_Add, /*bSplitName*/false);
+			Name = FName(*PlainString, NAME_NO_NUMBER_INTERNAL, /*bSplitName*/false);
 		}
 		break;
 
diff --git a/Engine/Plugins/FX/Niagara/Source/Niagara/Classes/NiagaraDataInterfaceGrid3DCollection.h b/Engine/Plugins/FX/Niagara/Source/Niagara/Classes/NiagaraDataInterfaceGrid3DCollection.h
--- a/Engine/Plugins/FX/Niagara/Source/Niagara/Classes/NiagaraDataInterfaceGrid3DCollection.h
+++ b/Engine/Plugins/FX/Niagara/Source/Niagara/Classes/NiagaraDataInterfaceGrid3DCollection.h
@@ -144,7 +144,7 @@ private:
 	LAYOUT_FIELD(FShaderResourceParameter, PerAttributeDataParam);
 
 	LAYOUT_FIELD(FShaderResourceParameter, SamplerParam);
-	LAYOUT_FIELD(TMemoryImageArray<FName>, AttributeNames);
+	LAYOUT_FIELD(TMemoryImageArray<FMemoryImageName>, AttributeNames);
 	LAYOUT_FIELD(TMemoryImageArray<uint32>, AttributeChannelCount);
 };
 
diff --git a/Engine/Plugins/FX/Niagara/Source/Niagara/Private/NiagaraDataInterfaceParticleRead.cpp b/Engine/Plugins/FX/Niagara/Source/Niagara/Private/NiagaraDataInterfaceParticleRead.cpp
--- a/Engine/Plugins/FX/Niagara/Source/Niagara/Private/NiagaraDataInterfaceParticleRead.cpp
+++ b/Engine/Plugins/FX/Niagara/Source/Niagara/Private/NiagaraDataInterfaceParticleRead.cpp
@@ -380,7 +380,7 @@ struct FNiagaraDataInterfaceParametersCS_ParticleRead : public FNiagaraDataInter
 		const TArray<FNiagaraVariableLayoutInfo>& SourceEmitterVariableLayouts = SourceDataSet->GetVariableLayouts();
 		for (int AttrNameIdx = 0; AttrNameIdx < AttributeNames.Num(); ++AttrNameIdx)
 		{
-			const FName& AttrName = AttributeNames[AttrNameIdx];
+			const FMemoryImageName& AttrName = AttributeNames[AttrNameIdx];
 			if (AttrName == NAME_None)
 			{
 				InstanceData->AttributeIndices[AttrNameIdx] = -1;
@@ -614,7 +614,7 @@ private:
 	LAYOUT_FIELD(FShaderParameter, AttributeIndicesParam);
 	LAYOUT_FIELD(FShaderParameter, AttributeCompressedParam);
 	LAYOUT_FIELD(FShaderParameter, AcquireTagRegisterIndexParam);
-	LAYOUT_FIELD(TMemoryImageArray<FName>, AttributeNames);
+	LAYOUT_FIELD(TMemoryImageArray<FMemoryImageName>, AttributeNames);
 	LAYOUT_FIELD(TMemoryImageArray<ENiagaraParticleDataValueType>, AttributeTypes);
 };
 
diff --git a/Engine/Plugins/FX/Niagara/Source/Niagara/Private/NiagaraRendererRibbons.cpp b/Engine/Plugins/FX/Niagara/Source/Niagara/Private/NiagaraRendererRibbons.cpp
--- a/Engine/Plugins/FX/Niagara/Source/Niagara/Private/NiagaraRendererRibbons.cpp
+++ b/Engine/Plugins/FX/Niagara/Source/Niagara/Private/NiagaraRendererRibbons.cpp
@@ -252,7 +252,7 @@ void FNiagaraRendererRibbons::CreateRenderThreadResources()
 		FRayTracingGeometryInitializer Initializer;
 		static const FName DebugName("FNiagaraRendererRibbons");
 		static int32 DebugNumber = 0;
-		Initializer.DebugName = FName(DebugName, DebugNumber++);
+		Initializer.DebugName = FDebugName(DebugName, DebugNumber++);
 		Initializer.IndexBuffer = nullptr;
 		Initializer.TotalPrimitiveCount = 0;
 		Initializer.GeometryType = RTGT_Triangles;
diff --git a/Engine/Plugins/FX/Niagara/Source/Niagara/Private/NiagaraRendererSprites.cpp b/Engine/Plugins/FX/Niagara/Source/Niagara/Private/NiagaraRendererSprites.cpp
--- a/Engine/Plugins/FX/Niagara/Source/Niagara/Private/NiagaraRendererSprites.cpp
+++ b/Engine/Plugins/FX/Niagara/Source/Niagara/Private/NiagaraRendererSprites.cpp
@@ -202,7 +202,7 @@ void FNiagaraRendererSprites::CreateRenderThreadResources()
 		FRayTracingGeometryInitializer Initializer;
 		static const FName DebugName("FNiagaraRendererSprites");
 		static int32 DebugNumber = 0;
-		Initializer.DebugName = FName(DebugName, DebugNumber++);
+		Initializer.DebugName = FDebugName(DebugName, DebugNumber++);
 		Initializer.IndexBuffer = nullptr;
 		Initializer.GeometryType = RTGT_Triangles;
 		Initializer.bFastBuild = true;
diff --git a/Engine/Plugins/Runtime/HairStrands/Source/HairStrandsCore/Private/HairStrandsRendering.cpp b/Engine/Plugins/Runtime/HairStrands/Source/HairStrandsCore/Private/HairStrandsRendering.cpp
--- a/Engine/Plugins/Runtime/HairStrands/Source/HairStrandsCore/Private/HairStrandsRendering.cpp
+++ b/Engine/Plugins/Runtime/HairStrands/Source/HairStrandsCore/Private/HairStrandsRendering.cpp
@@ -1392,9 +1392,9 @@ static void BuildHairAccelerationStructure_Strands(
 	FString DebugName(TEXT("AS_HairStrands_LOD") + FString::FromInt(LODIndex) + TEXT("_") + AssetDebugName);
 	Initializer.DebugName = FName(DebugName);
 #else
-	static const FName DebugName("AS_HairStrands");
+	static const FName DebugName("AS_HairStrands");	
 	static int32 DebugNumber = 0;
-	Initializer.DebugName = FName(DebugName, DebugNumber++);
+	Initializer.DebugName = FDebugName(DebugName, DebugNumber++);
 #endif
 
 	if (bProceduralPrimitive)
@@ -1448,7 +1448,7 @@ static void BuildHairAccelerationStructure_Cards(FRHICommandList& RHICmdList,
 #else
 	static const FName DebugName("AS_HairCards");
 	static int32 DebugNumber = 0;
-	Initializer.DebugName = FName(DebugName, DebugNumber++);
+	Initializer.DebugName = FDebugName(DebugName, DebugNumber++);
 #endif
 	Initializer.IndexBuffer = RestResource->RestIndexBuffer.IndexBufferRHI;
 	Initializer.IndexBufferOffset = 0;
@@ -1489,7 +1489,7 @@ static void BuildHairAccelerationStructure_Meshes(FRHICommandList& RHICmdList,
 #else
 	static const FName DebugName("AS_HairMeshes");
 	static int32 DebugNumber = 0;
-	Initializer.DebugName = FName(DebugName, DebugNumber++);
+	Initializer.DebugName = FDebugName(DebugName, DebugNumber++);
 #endif
 	Initializer.IndexBuffer = RestResource->IndexBuffer.IndexBufferRHI;
 	Initializer.IndexBufferOffset = 0;
diff --git a/Engine/Source/Developer/IoStoreUtilities/Private/IoStoreUtilities.cpp b/Engine/Source/Developer/IoStoreUtilities/Private/IoStoreUtilities.cpp
--- a/Engine/Source/Developer/IoStoreUtilities/Private/IoStoreUtilities.cpp
+++ b/Engine/Source/Developer/IoStoreUtilities/Private/IoStoreUtilities.cpp
@@ -3448,7 +3448,7 @@ int32 Describe(
 	for (int32 ScriptObjectIndex = 0; ScriptObjectIndex < NumScriptObjects; ++ScriptObjectIndex)
 	{
 		const FScriptObjectEntry& ScriptObjectEntry = ScriptObjectEntries[ScriptObjectIndex];
-		const FMappedName& MappedName = FMappedName::FromMinimalName(ScriptObjectEntry.ObjectName);
+		FMappedName MappedName = ScriptObjectEntry.Mapped;
 		check(MappedName.IsGlobal());
 		FScriptObjectDesc& ScriptObjectDesc = ScriptObjectByGlobalIdMap.Add(ScriptObjectEntry.GlobalIndex);
 		ScriptObjectDesc.Name = FName::CreateFromDisplayId(GlobalNameMap[MappedName.GetIndex()], MappedName.GetNumber());
diff --git a/Engine/Source/Developer/IoStoreUtilities/Private/PackageStoreOptimizer.cpp b/Engine/Source/Developer/IoStoreUtilities/Private/PackageStoreOptimizer.cpp
--- a/Engine/Source/Developer/IoStoreUtilities/Private/PackageStoreOptimizer.cpp
+++ b/Engine/Source/Developer/IoStoreUtilities/Private/PackageStoreOptimizer.cpp
@@ -1787,7 +1787,7 @@ FIoBuffer FPackageStoreOptimizer::CreateScriptObjectsBuffer() const
 	{
 		NameMapBuilder.MarkNameAsReferenced(ImportData.ObjectName);
 		FScriptObjectEntry& Entry = ScriptObjectEntries.AddDefaulted_GetRef();
-		Entry.ObjectName = NameMapBuilder.MapName(ImportData.ObjectName).ToUnresolvedMinimalName();
+		Entry.Mapped = NameMapBuilder.MapName(ImportData.ObjectName);
 		Entry.GlobalIndex = ImportData.GlobalIndex;
 		Entry.OuterIndex = ImportData.OuterIndex;
 		Entry.CDOClassIndex = ImportData.CDOClassIndex;
@@ -1814,10 +1814,10 @@ void FPackageStoreOptimizer::LoadScriptObjectsBuffer(const FIoBuffer& ScriptObje
 	ScriptObjectsArchive << NumScriptObjects;
 	for (int32 Index = 0; Index < NumScriptObjects; ++Index)
 	{
-		FScriptObjectEntry Entry;
+		FScriptObjectEntry Entry{};
 		ScriptObjectsArchive << Entry;
 		FScriptObjectData& ImportData = ScriptObjectsMap.Add(Entry.GlobalIndex);
-		const FMappedName& MappedName = FMappedName::FromMinimalName(Entry.ObjectName);
+		FMappedName MappedName = Entry.Mapped;
 		ImportData.ObjectName = FName::CreateFromDisplayId(NameMap[MappedName.GetIndex()], MappedName.GetNumber());
 		ImportData.GlobalIndex = Entry.GlobalIndex;
 		ImportData.OuterIndex = Entry.OuterIndex;
diff --git a/Engine/Source/Editor/KismetCompiler/Private/KismetCompilerVMBackend.cpp b/Engine/Source/Editor/KismetCompiler/Private/KismetCompilerVMBackend.cpp
--- a/Engine/Source/Editor/KismetCompiler/Private/KismetCompilerVMBackend.cpp
+++ b/Engine/Source/Editor/KismetCompiler/Private/KismetCompilerVMBackend.cpp
@@ -73,11 +73,12 @@ public:
 	{
 		FArchive& Ar = *this;
 
-		// We can't call Serialize directly as we need to store the data endian clean.
-		FScriptName ScriptName = NameToScriptName(Name);
-		Ar << ScriptName.ComparisonIndex;
-		Ar << ScriptName.DisplayIndex;
-		Ar << ScriptName.Number;
+		// This must match the format and endianness expected by XFERNAME 
+		FNameEntryId ComparisonIndex = Name.GetComparisonIndex(), DisplayIndex = Name.GetDisplayIndex();
+		uint32 Number = Name.GetNumber();
+		Ar << ComparisonIndex;
+		Ar << DisplayIndex;
+		Ar << Number;
 
 		return Ar;
 	}
diff --git a/Engine/Source/Programs/UnrealBuildTool/Configuration/TargetRules.cs b/Engine/Source/Programs/UnrealBuildTool/Configuration/TargetRules.cs
--- a/Engine/Source/Programs/UnrealBuildTool/Configuration/TargetRules.cs
+++ b/Engine/Source/Programs/UnrealBuildTool/Configuration/TargetRules.cs
@@ -766,6 +766,19 @@ namespace UnrealBuildTool
 		}
 		private bool? bWithServerCodeOverride;
 
+		/// <summary>
+		/// Compile with FName storing the number part in the name table. 
+		/// Saves memory when most names are not numbered and those that are are referenced multiple times.
+		/// The game and engine must ensure they reuse numbered names similarly to name strings to avoid leaking memory.
+		/// </summary>
+		[RequiresUniqueBuildEnvironment]
+		public bool bFNameOutlineNumber
+		{
+			get { return bFNameOutlineNumberOverride ?? false;  }
+			set { bFNameOutlineNumberOverride = value;  }
+		}
+		private bool? bFNameOutlineNumberOverride;
+
 		/// <summary>
 		/// When enabled, Push Model Networking support will be compiled in.
 		/// This can help reduce CPU overhead of networking, at the cost of more memory.
@@ -2424,6 +2437,11 @@ namespace UnrealBuildTool
 			get { return Inner.bWithServerCode; }
 		}
 
+		public bool bFNameOutlineNumber
+		{
+			get { return Inner.bFNameOutlineNumber;  }
+		}
+
 		public bool bWithPushModel
 		{
 			get { return Inner.bWithPushModel; }
diff --git a/Engine/Source/Programs/UnrealBuildTool/Configuration/UEBuildTarget.cs b/Engine/Source/Programs/UnrealBuildTool/Configuration/UEBuildTarget.cs
--- a/Engine/Source/Programs/UnrealBuildTool/Configuration/UEBuildTarget.cs
+++ b/Engine/Source/Programs/UnrealBuildTool/Configuration/UEBuildTarget.cs
@@ -3926,6 +3926,8 @@ namespace UnrealBuildTool
 				GlobalCompileEnvironment.Definitions.Add("WITH_SERVER_CODE=0");
 			}
 
+			GlobalCompileEnvironment.Definitions.Add(String.Format("UE_FNAME_OUTLINE_NUMBER={0}", Rules.bFNameOutlineNumber? 1 : 0));
+
 			// Set the defines for Push Model
 			if (Rules.bWithPushModel)
 			{
diff --git a/Engine/Source/Runtime/Core/Private/Misc/CoreGlobals.cpp b/Engine/Source/Runtime/Core/Private/Misc/CoreGlobals.cpp
--- a/Engine/Source/Runtime/Core/Private/Misc/CoreGlobals.cpp
+++ b/Engine/Source/Runtime/Core/Private/Misc/CoreGlobals.cpp
@@ -240,14 +240,13 @@ float					GNearClippingPlane				= 10.0f;				/* Near clipping plane */
 bool					GExitPurge						= false;
 
 FChunkedFixedUObjectArray* GCoreObjectArrayForDebugVisualizers = nullptr;
-template class CORE_API TArray<FMinimalName, TInlineAllocator<3>>;
 
-TArray<FMinimalName, TInlineAllocator<3>>* GCoreComplexObjectPathDebug = nullptr;
+UE::ObjectPath::Private::FStoredObjectPath* GCoreComplexObjectPathDebug = nullptr;
 FObjectHandlePackageDebugData* GCoreObjectHandlePackageDebug = nullptr;
 #if PLATFORM_UNIX
 uint8** CORE_API GNameBlocksDebug = FNameDebugVisualizer::GetBlocks();
 FChunkedFixedUObjectArray*& CORE_API GObjectArrayForDebugVisualizers = GCoreObjectArrayForDebugVisualizers;
-TArray<FMinimalName, TInlineAllocator<3>>*& GComplexObjectPathDebug = GCoreComplexObjectPathDebug;
+UE::ObjectPath::Private::FStoredObjectPath*& GComplexObjectPathDebug = GCoreComplexObjectPathDebug;
 FObjectHandlePackageDebugData*& CORE_API GObjectHandlePackageDebug = GCoreObjectHandlePackageDebug;
 #endif
 
diff --git a/Engine/Source/Runtime/Core/Private/ProfilingDebugging/CsvProfiler.cpp b/Engine/Source/Runtime/Core/Private/ProfilingDebugging/CsvProfiler.cpp
--- a/Engine/Source/Runtime/Core/Private/ProfilingDebugging/CsvProfiler.cpp
+++ b/Engine/Source/Runtime/Core/Private/ProfilingDebugging/CsvProfiler.cpp
@@ -1774,7 +1774,7 @@ public:
 
 private:
 	static CSV_PROFILER_INLINE uint64 GetStatID(const char* StatName) { return uint64(StatName); }
-	static CSV_PROFILER_INLINE uint64 GetStatID(const FName& StatId) { return StatId.GetComparisonIndex().ToUnstableInt(); }
+	static CSV_PROFILER_INLINE uint64 GetStatID(const FName& StatId) { return StatId.ToUnstableInt(); }
 
 	static FCriticalSection TlsCS;
 	static TArray<FWeakPtr> TlsInstances;
diff --git a/Engine/Source/Runtime/Core/Private/Serialization/MemoryImage.cpp b/Engine/Source/Runtime/Core/Private/Serialization/MemoryImage.cpp
--- a/Engine/Source/Runtime/Core/Private/Serialization/MemoryImage.cpp
+++ b/Engine/Source/Runtime/Core/Private/Serialization/MemoryImage.cpp
@@ -1432,16 +1432,16 @@ void FMemoryImageResult::SaveToArchive(FArchive& Ar) const
 
 
 	TArray<uint32> ScriptNameCounts;
-	TArray<uint32> MinimalNameCounts;
+	TArray<uint32> MemoryImageNameCounts;
 	CountNumNames(ScriptNames, ScriptNameCounts);
-	CountNumNames(MinimalNames, MinimalNameCounts);
+	CountNumNames(MemoryImageNames, MemoryImageNameCounts);
 
 	uint32 NumVTables = VTableCounts.Num();
 	uint32 NumScriptNames = ScriptNameCounts.Num();
-	uint32 NumMinimalNames = MinimalNameCounts.Num();
+	uint32 NumMemoryImageNames = MemoryImageNameCounts.Num();
 	Ar << NumVTables;
 	Ar << NumScriptNames;
-	Ar << NumMinimalNames;
+	Ar << NumMemoryImageNames;
 
 	{
 		int32 VTableIndex = 0;
@@ -1466,7 +1466,7 @@ void FMemoryImageResult::SaveToArchive(FArchive& Ar) const
 	}
 	
 	SerializeNames(ScriptNames, ScriptNameCounts, Ar);
-	SerializeNames(MinimalNames, MinimalNameCounts, Ar);
+	SerializeNames(MemoryImageNames, MemoryImageNameCounts, Ar);
 }
 
 static inline void ApplyVTablePatch(void* FrozenObject, const FTypeLayoutDesc& DerivedType, uint32 VTableOffset, uint32 Offset)
@@ -1482,10 +1482,10 @@ static inline void ApplyScriptNamePatch(void* FrozenObject, const FScriptName& N
 	new(NameDst) FScriptName(Name);
 }
 
-static inline void ApplyMinimalNamePatch(void* FrozenObject, const FMinimalName& Name, uint32 Offset)
+static inline void ApplyMemoryImageNamePatch(void* FrozenObject, const FMemoryImageName& Name, uint32 Offset)
 {
 	void* NameDst = (uint8*)FrozenObject + Offset;
-	new(NameDst) FMinimalName(Name);
+	new(NameDst) FMemoryImageName(Name);
 }
 
 void FMemoryImageResult::ApplyPatches(void* FrozenObject) const
@@ -1502,9 +1502,9 @@ void FMemoryImageResult::ApplyPatches(void* FrozenObject) const
 		ApplyScriptNamePatch(FrozenObject, NameToScriptName(Patch.Name), Patch.Offset);
 	}
 
-	for (const FMemoryImageNamePointer& Patch : MinimalNames)
+	for (const FMemoryImageNamePointer& Patch : MemoryImageNames)
 	{
-		ApplyMinimalNamePatch(FrozenObject, NameToMinimalName(Patch.Name), Patch.Offset);
+		ApplyMemoryImageNamePatch(FrozenObject, FMemoryImageName(Patch.Name), Patch.Offset);
 	}
 }
 
@@ -1523,10 +1523,10 @@ FMemoryImageObject FMemoryImageResult::LoadFromArchive(FArchive& Ar, const FType
 
 	uint32 NumVTables = 0u;
 	uint32 NumScriptNames = 0u;
-	uint32 NumMinimalNames = 0u;
+	uint32 NumMemoryImageNames = 0u;
 	Ar << NumVTables;
 	Ar << NumScriptNames;
-	Ar << NumMinimalNames;
+	Ar << NumMemoryImageNames;
 
 	for (uint32 i = 0u; i < NumVTables; ++i)
 	{
@@ -1563,7 +1563,7 @@ FMemoryImageObject FMemoryImageResult::LoadFromArchive(FArchive& Ar, const FType
 		}
 	}
 
-	for (uint32 i = 0u; i < NumMinimalNames; ++i)
+	for (uint32 i = 0u; i < NumMemoryImageNames; ++i)
 	{
 		FName Name;
 		uint32 NumPatches = 0u;
@@ -1574,7 +1574,7 @@ FMemoryImageObject FMemoryImageResult::LoadFromArchive(FArchive& Ar, const FType
 		{
 			uint32 Offset = 0u;
 			Ar << Offset;
-			ApplyMinimalNamePatch(FrozenObject, NameToMinimalName(Name), Offset);
+			ApplyMemoryImageNamePatch(FrozenObject, FMemoryImageName(Name), Offset);
 		}
 	}
 
@@ -1704,26 +1704,13 @@ uint32 FMemoryImageSection::WriteVTable(const FTypeLayoutDesc& TypeDesc, const F
 	return VTablePointer.Offset;
 }
 
-uint32 FMemoryImageSection::WriteFName(const FName& Name)
+uint32 FMemoryImageSection::WriteFMemoryImageName(TConstArrayView<uint8> InBytes, const FName& Name)
 {
-	const FPlatformTypeLayoutParameters& TargetLayoutParameters = ParentImage->TargetLayoutParameters;
-	if (!TargetLayoutParameters.WithCasePreservingFName())
-	{
-		return WriteFMinimalName(NameToMinimalName(Name));
-	}
-	else
-	{
-		return WriteFScriptName(NameToScriptName(Name));
-	}
-}
-
-uint32 FMemoryImageSection::WriteFMinimalName(const FMinimalName& Name)
-{
-	const uint32 Offset = WriteBytes(FMinimalName());
+	const uint32 Offset = WriteBytes(InBytes);
 	if (!Name.IsNone())
 	{
-		FMemoryImageNamePointer& NamePointer = MinimalNames.AddDefaulted_GetRef();
-		NamePointer.Name = MinimalNameToName(Name);
+		FMemoryImageNamePointer& NamePointer = ScriptNames.AddDefaulted_GetRef();
+		NamePointer.Name = Name;
 		NamePointer.Offset = Offset;
 	}
 	return Offset;
@@ -1731,7 +1718,7 @@ uint32 FMemoryImageSection::WriteFMinimalName(const FMinimalName& Name)
 
 uint32 FMemoryImageSection::WriteFScriptName(const FScriptName& Name)
 {
-	const uint32 Offset = WriteBytes(FScriptName());
+	const uint32 Offset = WriteBytes(Name);
 	if (!Name.IsNone())
 	{
 		FMemoryImageNamePointer& NamePointer = ScriptNames.AddDefaulted_GetRef();
@@ -1762,10 +1749,10 @@ uint32 FMemoryImageSection::Flatten(FMemoryImageResult& OutResult) const
 		ResultName->Offset += AlignedOffset;
 	}
 
-	OutResult.MinimalNames.Reserve(OutResult.MinimalNames.Num() + MinimalNames.Num());
-	for (const FMemoryImageNamePointer& Name : MinimalNames)
+	OutResult.MemoryImageNames.Reserve(OutResult.MemoryImageNames.Num() + MemoryImageNames.Num());
+	for (const FMemoryImageNamePointer& Name : MemoryImageNames)
 	{
-		FMemoryImageNamePointer* ResultName = new(OutResult.MinimalNames) FMemoryImageNamePointer(Name);
+		FMemoryImageNamePointer* ResultName = new(OutResult.MemoryImageNames) FMemoryImageNamePointer(Name);
 		ResultName->Offset += AlignedOffset;
 	}
 
@@ -1793,7 +1780,7 @@ void FMemoryImageSection::ComputeHash()
 	HashState.Update((uint8*)Pointers.GetData(), Pointers.Num() * Pointers.GetTypeSize());
 	HashState.Update((uint8*)VTables.GetData(), VTables.Num() * VTables.GetTypeSize());
 	UpdateHashNamePatches(ScriptNames, HashState);
-	UpdateHashNamePatches(MinimalNames, HashState);
+	UpdateHashNamePatches(MemoryImageNames, HashState);
 	HashState.Final();
 	HashState.GetHash(Hash.Hash);
 }
@@ -1870,7 +1857,7 @@ void FMemoryImage::Flatten(FMemoryImageResult& OutResult, bool bMergeDuplicateSe
 	// Sort to group runs of the same FName/VTable
 	OutResult.VTables.Sort();
 	OutResult.ScriptNames.Sort();
-	OutResult.MinimalNames.Sort();
+	OutResult.MemoryImageNames.Sort();
 }
 
 FMemoryImageWriter::FMemoryImageWriter(FMemoryImage& InImage) : Section(InImage.AllocateSection()) {}
@@ -1965,14 +1952,9 @@ uint32 FMemoryImageWriter::WriteVTable(const FTypeLayoutDesc& TypeDesc, const FT
 	return Section->WriteVTable(TypeDesc, DerivedTypeDesc);
 }
 
-uint32 FMemoryImageWriter::WriteFName(const FName& Name)
-{
-	return Section->WriteFName(Name);
-}
-
-uint32 FMemoryImageWriter::WriteFMinimalName(const FMinimalName& Name)
+uint32 FMemoryImageWriter::WriteFMemoryImageName(TConstArrayView<uint8> InBytes, const FName& Name)
 {
-	return Section->WriteFMinimalName(Name);
+	return Section->WriteFMemoryImageName(InBytes, Name);
 }
 
 uint32 FMemoryImageWriter::WriteFScriptName(const FScriptName& Name)
diff --git a/Engine/Source/Runtime/Core/Private/UObject/UnrealNames.cpp b/Engine/Source/Runtime/Core/Private/UObject/UnrealNames.cpp
--- a/Engine/Source/Runtime/Core/Private/UObject/UnrealNames.cpp
+++ b/Engine/Source/Runtime/Core/Private/UObject/UnrealNames.cpp
@@ -5,6 +5,8 @@
 #include "Async/Async.h"
 #include "Misc/AssertionMacros.h"
 #include "Misc/MessageDialog.h"
+#include "Misc/OutputDeviceHelper.h"
+#include "Misc/Paths.h"
 #include "Math/NumericLimits.h"
 #include "Math/UnrealMathUtility.h"
 #include "HAL/UnrealMemory.h"
@@ -23,6 +25,7 @@
 #include "Internationalization/Text.h"
 #include "Internationalization/Internationalization.h"
 #include "Misc/OutputDeviceRedirector.h"
+#include "HAL/FileManager.h"
 #include "HAL/IConsoleManager.h"
 #include "HAL/LowLevelMemTracker.h"
 #include "Serialization/MemoryImage.h"
@@ -31,6 +34,7 @@
 #include "Hash/Blake3.h"
 #include "Hash/CityHash.h"
 #include "Templates/AlignmentTemplates.h"
+#include "Templates/Greater.h"
 
 PRAGMA_DISABLE_UNSAFE_TYPECAST_WARNINGS
 
@@ -46,6 +50,57 @@ PRAGMA_DISABLE_UNSAFE_TYPECAST_WARNINGS
 
 DEFINE_LOG_CATEGORY_STATIC(LogUnrealNames, Log, All);
 
+// Console command declarations
+namespace UE::Name::Private
+{
+	// List FNames to the output device 
+	void ListFNames(const TArray<FString>&, UWorld*, FOutputDevice&);
+	// Dump FNames to a file
+	void DumpFNames(const TArray<FString>&, UWorld*, FOutputDevice&);
+
+	// List FNames to the output device 
+	void ListNumberedFNames(const TArray<FString>&, UWorld*, FOutputDevice&);
+	// Dump FNames to a file
+	void DumpNumberedFNames(const TArray<FString>&, UWorld*, FOutputDevice&);
+
+	// Write FName stats to the output device 
+	void GetFNameStats(FOutputDevice&);
+
+	// Write csv stats on hash distribution
+	void DumpHashCsv(FOutputDevice&);
+
+	FAutoConsoleCommandWithWorldArgsAndOutputDevice CCListFNames(
+		TEXT("FName.List"),
+		TEXT("List all base FName strings to the output device. Pass -num=n to list the most recent n names."),
+		FConsoleCommandWithWorldArgsAndOutputDeviceDelegate::CreateStatic(&ListFNames)
+	);
+	FAutoConsoleCommandWithWorldArgsAndOutputDevice CCDumpFNames(
+		TEXT("FName.Dump"),
+		TEXT("Dump all base FName strings to a file. Pass -num=n to dump the most recent n names."),
+		FConsoleCommandWithWorldArgsAndOutputDeviceDelegate::CreateStatic(&DumpFNames)
+	);
+	FAutoConsoleCommandWithWorldArgsAndOutputDevice CCListNumberedFNames(
+		TEXT("FName.ListNumbered"),
+		TEXT("List all numbered FNames to the output devicce (only when UE_FNAME_OUTLINE_NUMBER is set). Pass -num=n to list the most recent n names."),
+		FConsoleCommandWithWorldArgsAndOutputDeviceDelegate::CreateStatic(&ListNumberedFNames)
+	);
+	FAutoConsoleCommandWithWorldArgsAndOutputDevice CCDumpNumberedFNames(
+		TEXT("FName.DumpNumbered"),
+		TEXT("Dump all numbered FNames to a file (only when UE_FNAME_OUTLINE_NUMBER is set). Pass -num=n to dump the most recent n names."),
+		FConsoleCommandWithWorldArgsAndOutputDeviceDelegate::CreateStatic(&DumpNumberedFNames)
+	);
+	FAutoConsoleCommandWithOutputDevice CCGetFNameStats(
+		TEXT("FName.Stats"),
+		TEXT("Write FName stats to the output device."),
+		FConsoleCommandWithOutputDeviceDelegate::CreateStatic(&GetFNameStats)
+	);
+	FAutoConsoleCommandWithOutputDevice CCDumpHashCsv(
+		TEXT("FName.HashCsv"),
+		TEXT("Write FName hash stats to a csv file."),
+		FConsoleCommandWithOutputDeviceDelegate::CreateStatic(&DumpHashCsv)
+	);
+}
+
 const TCHAR* LexToString(EName Ename)
 {
 	switch (Ename)
@@ -63,6 +118,11 @@ int32 FNameEntry::GetDataOffset()
 	return STRUCT_OFFSET(FNameEntry, AnsiName);
 }
 
+struct FNameEntryIds
+{
+	FNameEntryId ComparisonId, DisplayId;
+};
+
 /*-----------------------------------------------------------------------------
 	FName helpers. 
 -----------------------------------------------------------------------------*/
@@ -126,6 +186,7 @@ struct FNameStringView
 	bool bIsWide;
 
 	bool IsAnsi() const { return !bIsWide; }
+	bool IsNoneString() const;
 
 	int32 BytesWithTerminator() const
 	{
@@ -181,6 +242,8 @@ static constexpr uint32 FNameMaxBlockBits = 13; // Limit block array a bit, stil
 static constexpr uint32 FNameBlockOffsetBits = 16;
 static constexpr uint32 FNameMaxBlocks = 1 << FNameMaxBlockBits;
 static constexpr uint32 FNameBlockOffsets = 1 << FNameBlockOffsetBits;
+static constexpr uint32 FNameEntryIdBits = FNameBlockOffsetBits + FNameMaxBlockBits;
+static constexpr uint32 FNameEntryIdMask = (1 << FNameEntryIdBits ) - 1;
 
 /** An unpacked FNameEntryId */
 struct FNameEntryHandle
@@ -191,12 +254,16 @@ struct FNameEntryHandle
 	FNameEntryHandle(uint32 InBlock, uint32 InOffset)
 		: Block(InBlock)
 		, Offset(InOffset)
-	{}
+	{
+		checkName(Block < FNameMaxBlocks);
+		checkName(Offset < FNameBlockOffsets);
+	}
 
 	FNameEntryHandle(FNameEntryId Id)
 		: Block(Id.ToUnstableInt() >> FNameBlockOffsetBits)
 		, Offset(Id.ToUnstableInt() & (FNameBlockOffsets - 1))
-	{}
+	{
+	}
 
 	operator FNameEntryId() const
 	{
@@ -358,6 +425,27 @@ public:
 		return Handle;
 	}
 
+#if UE_FNAME_OUTLINE_NUMBER
+	template<class ScopeLock>
+	FNameEntryHandle CreateWithNumber(const FNameEntryId& StringPart, uint32 NumberPart, TOptional<FNameEntryId> ComparisonId, FNameEntryHeader Header)
+	{
+		FPlatformMisc::Prefetch(Blocks[CurrentBlock]);
+		FNameEntryHandle Handle = Allocate<ScopeLock>(FNameEntry::GetDataOffset() + sizeof(FNameEntry::NumberedName));
+		FNameEntry& Entry = Resolve(Handle);
+
+#if WITH_CASE_PRESERVING_NAME
+		Entry.ComparisonId = ComparisonId.IsSet() ? ComparisonId.GetValue() : FNameEntryId(Handle);
+#endif
+
+		Entry.Header = Header;
+
+		Entry.NumberedName.Id = StringPart;
+		Entry.NumberedName.Number = NumberPart;
+
+		return Handle;
+	}
+#endif // UE_FNAME_OUTLINE_NUMBER
+
 	FNameEntry& Resolve(FNameEntryHandle Handle) const
 	{
 		// Lock not needed
@@ -382,6 +470,60 @@ public:
 	
 	uint8** GetBlocksForDebugVisualizer() { return Blocks; }
 
+#if UE_FNAME_OUTLINE_NUMBER
+	void IterateNumberedNames(TFunctionRef<void(FName)> InFunc) const
+	{
+		bool bStop = false;
+		TArray<FNameEntryId> Tmp;
+		Tmp.Reserve(BlockSizeBytes / FNameEntry::GetNumberedEntrySize());
+		for (uint32 BlockIdx = 0; !bStop; ++BlockIdx)
+		{
+			Tmp.Reset();
+			{
+				FRWScopeLock _(Lock, FRWScopeLockType::SLT_ReadOnly);
+
+				uint32 Size = BlockSizeBytes;
+				if(BlockIdx == CurrentBlock)
+				{
+					bStop = true;
+					Size = CurrentByteCursor;
+				}
+
+				const uint8* It = Blocks[BlockIdx];
+				const uint8* End = It + Size - FNameEntry::GetDataOffset();
+				while (It < End)
+				{
+					const FNameEntry* Entry = (const FNameEntry*)It;
+					if (uint32 Len = Entry->Header.Len)
+					{
+						// Skip, this is not a numbered entry 
+						It += FNameEntry::GetSize(Len, !Entry->IsWide());
+					}
+					else // Null-terminator or numbered-entry
+					{
+						if ((End - It) < FNameEntry::GetNumberedEntrySize() || (Entry->NumberedName.Id == FNameEntryId() && Entry->NumberedName.Number == 0))
+						{
+							break;
+						}
+						else
+						{
+							FNameEntryId Id = FNameEntryHandle(BlockIdx, (It - Blocks[BlockIdx]) / Stride);
+							Tmp.Add(Id);
+							It += FNameEntry::GetNumberedEntrySize();
+						}
+					}
+				}
+			}
+
+			// Dispatch callbacks for each block without holding lock
+			for (FNameEntryId Id : Tmp)
+			{
+				InFunc(FName(Id, Id, NAME_NO_NUMBER_INTERNAL));
+			}
+		}
+	}
+#endif
+
 	void DebugDump(TArray<const FNameEntry*>& Out) const
 	{
 		FRWScopeLock _(Lock, FRWScopeLockType::SLT_ReadOnly);
@@ -406,9 +548,20 @@ private:
 				Out.Add(Entry);
 				It += FNameEntry::GetSize(Len, !Entry->IsWide());
 			}
-			else // Null-terminator entry found
+			else // Null-terminator or numbered-entry
 			{
-				break;
+#if UE_FNAME_OUTLINE_NUMBER
+				if ((End - It) < FNameEntry::GetNumberedEntrySize()
+				|| (Entry->NumberedName.Id == FNameEntryId() && Entry->NumberedName.Number == 0))
+				{
+					break;
+				}
+				else
+				{
+					Out.Add(Entry);
+					It += FNameEntry::GetNumberedEntrySize();
+				}
+#endif
 			}
 		}
 	}
@@ -422,11 +575,21 @@ private:
 	{
 		LLM_SCOPE(ELLMTag::FName);
 		// Null-terminate final entry to allow DebugDump() entry iteration
+#if UE_FNAME_OUTLINE_NUMBER
+		if (CurrentByteCursor + FNameEntry::GetNumberedEntrySize() <= BlockSizeBytes)
+		{
+			FNameEntry* Terminator = (FNameEntry*)(Blocks[CurrentBlock] + CurrentByteCursor);
+			Terminator->Header.Len = 0;
+			Terminator->NumberedName.Id = FNameEntryId();
+			Terminator->NumberedName.Number = 0;
+		}
+#else
 		if (CurrentByteCursor + FNameEntry::GetDataOffset() <= BlockSizeBytes)
 		{
 			FNameEntry* Terminator = (FNameEntry*)(Blocks[CurrentBlock] + CurrentByteCursor);
 			Terminator->Header.Len = 0;
 		}
+#endif
 
 #if FNAME_WRITE_PROTECT_PAGES
 		FPlatformMemory::PageProtect(Blocks[CurrentBlock], BlockSizeBytes, /* read */ true, /* write */ false);
@@ -500,13 +663,33 @@ struct FNameHash
 		return CityHash64(reinterpret_cast<const char*>(Str), Len * sizeof(CharType));
 	}
 
+	static uint64 GenerateHash(FNameEntryId StringPart, int32 NumberPart)
+	{
+		return CityHash128to64(Uint128_64(NumberPart, StringPart.ToUnstableInt()));
+	}
+
 	template<class CharType>
 	FNameHash(const CharType* Str, int32 Len)
-		: FNameHash(Str, Len, GenerateHash(Str, Len))
+		: FNameHash(GenerateHash(Str, Len), Len, (bool)IsAnsiNone(Str, Len), sizeof(CharType) == sizeof(WIDECHAR))
 	{}
 
 	template<class CharType>
 	FNameHash(const CharType* Str, int32 Len, uint64 Hash)
+		: FNameHash(Hash, Len, (bool)IsAnsiNone(Str, Len), sizeof(CharType) == sizeof(WIDECHAR))
+	{
+	}
+
+	FNameHash(FNameEntryId StringPart, int32 NumberPart)
+		: FNameHash(GenerateHash(StringPart, NumberPart), 0, false, false)
+	{
+	}
+
+	FNameHash(ENoInit)
+	{
+
+	}
+		
+	FNameHash(uint64 Hash, int32 Len, bool IsNone, bool bIsWide)
 	{
 		uint32 Hi = static_cast<uint32>(Hash >> 32);
 		uint32 Lo = static_cast<uint32>(Hash);
@@ -514,7 +697,7 @@ struct FNameHash
 		// "None" has FNameEntryId with a value of zero
 		// Always set a bit in SlotProbeHash for "None" to distinguish unused slot values from None
 		// @see FNameSlot::Used()
-		uint32 IsNoneBit = IsAnsiNone(Str, Len) << FNameSlot::ProbeHashShift;
+		uint32 IsNoneBit = IsNone << FNameSlot::ProbeHashShift;
 
 		static_assert((ShardMask & FNameSlot::ProbeHashMask) == 0, "Masks overlap");
 
@@ -522,7 +705,7 @@ struct FNameHash
 		UnmaskedSlotIndex = Lo;
 		SlotProbeHash = (Hi & FNameSlot::ProbeHashMask) | IsNoneBit;
 		EntryProbeHeader.Len = Len;
-		EntryProbeHeader.bIsWide = sizeof(CharType) == sizeof(WIDECHAR);
+		EntryProbeHeader.bIsWide = bIsWide;
 
 		// When we always use lowercase hashing, we can store parts of the hash in the entry
 		// to avoid copying and decoding entries needlessly. WITH_CUSTOM_NAME_ENCODING
@@ -620,6 +803,11 @@ FNameHash HashName<ENameCase::CaseSensitive>(FNameStringView Name)
 	return Name.IsAnsi() ? FNameHash(Name.Ansi, Name.Len) : FNameHash(Name.Wide, Name.Len);
 }
 
+bool FNameStringView::IsNoneString() const
+{ 
+	return bool(bIsWide ? FNameHash::IsAnsiNone(Wide, Len) : FNameHash::IsAnsiNone(Ansi, Len));
+}
+
 template<ENameCase Sensitivity>
 struct FNameValue
 {
@@ -645,6 +833,37 @@ struct FNameValue
 #endif
 };
 
+#if UE_FNAME_OUTLINE_NUMBER
+template<ENameCase Sensitivity>
+struct FNumberedNameValue
+{
+	FNumberedNameValue()
+		: Hash(NoInit)
+	{
+	}
+
+	FNumberedNameValue(FNameEntryId InStringPart, int32 InNumberPart)
+		: StringPart(InStringPart)
+		, NumberPart(InNumberPart)
+		, Hash(InStringPart, NumberPart)
+	{
+	}
+
+	FNameEntryId	StringPart;
+	int32			NumberPart;
+	FNameHash		Hash;
+
+#if WITH_CASE_PRESERVING_NAME
+	FNameEntryId ComparisonId;
+#endif
+};
+
+using FNumberedNameComparisonValue = FNumberedNameValue<ENameCase::IgnoreCase>;
+#if WITH_CASE_PRESERVING_NAME
+using FNumberedNameDisplayValue = FNumberedNameValue<ENameCase::CaseSensitive>;
+#endif
+#endif
+
 using FNameComparisonValue = FNameValue<ENameCase::IgnoreCase>;
 #if WITH_CASE_PRESERVING_NAME
 using FNameDisplayValue = FNameValue<ENameCase::CaseSensitive>;
@@ -655,6 +874,13 @@ FORCEINLINE TOptional<FNameEntryId> GetExistingComparisonId(const FNameCompariso
 FORCEINLINE TOptional<FNameEntryId> GetExistingComparisonId(const FNameDisplayValue& Value)		{ return Value.ComparisonId; }
 #endif
 
+#if UE_FNAME_OUTLINE_NUMBER
+FORCEINLINE TOptional<FNameEntryId> GetExistingComparisonId(const FNumberedNameComparisonValue& Value) { return TOptional<FNameEntryId>(); }
+#if WITH_CASE_PRESERVING_NAME
+FORCEINLINE TOptional<FNameEntryId> GetExistingComparisonId(const FNumberedNameDisplayValue& Value) { return Value.ComparisonId; }
+#endif // WITH_CASE_PRESERVING_NAME
+#endif // UE_FNAME_OUTLINE_NUMBER
+
 // For pre-locked batch operations
 struct FNullScopeLock
 {
@@ -704,6 +930,7 @@ public:
 	uint32 Capacity() const	{ return CapacityMask + 1; }
 	uint32 NumCreated() const { return NumCreatedEntries; }
 	uint32 NumCreatedWide() const { return NumCreatedWideEntries; }
+	uint32 NumCreatedWithNumber() const { return NumCreatedWithNumberEntries; }
 
 protected:
 	enum { LoadFactorQuotient = 9, LoadFactorDivisor = 10 }; // I.e. realloc slots when 90% full
@@ -715,24 +942,143 @@ protected:
 	FNameEntryAllocator* Entries = nullptr;
 	uint32 NumCreatedEntries = 0;
 	uint32 NumCreatedWideEntries = 0;
+	uint32 NumCreatedWithNumberEntries = 0;
 
 
 	template<ENameCase Sensitivity>
 	FORCEINLINE static bool EntryEqualsValue(const FNameEntry& Entry, const FNameValue<Sensitivity>& Value)
 	{
+		// UE_FNAME_OUTLINE_NUMBER no special case handling here because entry + number entries have Len == 0 which is not true of any stored strings
+		//	so the headers will compare unequal
 		return Entry.Header == Value.Hash.EntryProbeHeader && EqualsSameDimensions<Sensitivity>(Entry, Value.Name);
 	}
+
+#if UE_FNAME_OUTLINE_NUMBER
+	template<ENameCase Sensitivity>
+	FORCEINLINE static bool EntryEqualsValue(const FNameEntry& Entry, const FNumberedNameValue<Sensitivity>& Value)
+	{
+		// UE_FNAME_OUTLINE_NUMBER no checking against string entries here because entry + number entries have Len == 0 which is not true of any stored strings
+		//	so the headers will compare unequal
+		if( Entry.Header == Value.Hash.EntryProbeHeader )
+		{
+			return Entry.NumberedName.Id == Value.StringPart && Entry.NumberedName.Number == Value.NumberPart;
+		}
+		return false;
+	}
+
+	template<ENameCase Sensitivity>
+	FORCEINLINE static bool RehashNameWithNumber(const FNameEntry& Entry, FNumberedNameValue<Sensitivity>& OutValue)
+	{
+		if (Entry.Header.Len == 0)
+		{
+			OutValue = FNumberedNameValue<Sensitivity>(Entry.NumberedName.Id, Entry.NumberedName.Number);
+			return true;
+		}
+		return false;
+	}
+
+#endif // UE_FNAME_OUTLINE_NUMBER
 };
 
 template<ENameCase Sensitivity>
 class FNamePoolShard : public FNamePoolShardBase
 {
 public:
+#if !UE_BUILD_SHIPPING 
+	void LogCsvHeader(FOutputDevice& Ar) const
+	{
+		Ar.Logf(TEXT(",NumEntries,NumWideEntries,NumNumberedEntries,UsedSlots,Capacity,")
+				TEXT("UniqueUnmaskedIndices,NumUnmaskedIndexCollisions,MaxUnmaskedIndexCollisions,")
+				TEXT("UniqueMaskedIndices,NumMaskedIndexCollisions,MaxMaskedIndexCollisions,")
+				TEXT("UniqueProbeHashes,NumProbeHashCollisions,MaxProbeHashCollisions")
+				);
+	}
+
+	void LogStatsCsv(FOutputDevice& Ar) const
+	{
+		FRWScopeLock _(Lock, FRWScopeLockType::SLT_ReadOnly);
+
+		TMap<uint32, int32> CountByUnmaskedIndex;
+		TMap<uint32, int32> CountByMaskedIndex;
+		TMap<uint32, int32> CountBySlotProbeHash;
+
+		for (uint32 i = 0; i < Capacity(); ++i)
+		{
+			FNameSlot& Slot = Slots[i];
+			if (!Slot.Used()) 
+			{
+				continue; 
+			}
+
+			const FNameEntry& Entry = Entries->Resolve(Slot.GetId());
+#if UE_FNAME_OUTLINE_NUMBER
+			FNumberedNameValue<Sensitivity> Value;
+			if (RehashNameWithNumber(Entry, Value))
+			{
+				CountByUnmaskedIndex.FindOrAdd(Value.Hash.UnmaskedSlotIndex)++;
+				CountByMaskedIndex.FindOrAdd(FNameHash::GetProbeStart(Value.Hash.UnmaskedSlotIndex, CapacityMask))++;
+				CountBySlotProbeHash.FindOrAdd(Value.Hash.SlotProbeHash)++;
+			}
+			else
+#endif
+			{
+				FNameBuffer DecodeBuffer;
+				FNameStringView Name = Entry.MakeView(DecodeBuffer);
+				FNameHash Hash = HashName<Sensitivity>(Name);
+				CountByUnmaskedIndex.FindOrAdd(Hash.UnmaskedSlotIndex)++;
+				CountByMaskedIndex.FindOrAdd(FNameHash::GetProbeStart(Hash.UnmaskedSlotIndex, CapacityMask))++;
+				CountBySlotProbeHash.FindOrAdd(Hash.SlotProbeHash)++;
+			}
+		}
+
+		CountByUnmaskedIndex.ValueSort(TGreater<int32>());
+		CountByMaskedIndex.ValueSort(TGreater<int32>());
+		CountBySlotProbeHash.ValueSort(TGreater<int32>());
+
+
+		const uint32 UniqueUnmaskedIndices = CountByUnmaskedIndex.Num();
+		const uint32 UniqueMaskedIndices = CountByMaskedIndex.Num();
+		const uint32 UniqueProbeHashes = CountBySlotProbeHash.Num();
+
+		int32 MaxUnmaskedIndexCollisions = MIN_int32;
+		int32 NumUnmaskedIndexCollisions = 0;
+		for (const TPair<uint32, int32>& Pair : CountByUnmaskedIndex)
+		{
+			NumUnmaskedIndexCollisions += Pair.Value - 1;
+			MaxUnmaskedIndexCollisions = FMath::Max(MaxUnmaskedIndexCollisions, Pair.Value);
+		}
+
+		int32 MaxMaskedIndexCollisions = MIN_int32;
+		int32 NumMaskedIndexCollisions = 0;
+		for (const TPair<uint32, int32>& Pair : CountByMaskedIndex)
+		{
+			NumMaskedIndexCollisions += Pair.Value - 1;
+			MaxMaskedIndexCollisions = FMath::Max(MaxMaskedIndexCollisions, Pair.Value);
+		}
+
+		int32 MaxProbeHashCollisions = MIN_int32;
+		int32 NumProbeHashCollisions = 0;
+		for (const TPair<uint32, int32>& Pair : CountBySlotProbeHash)
+		{
+			NumProbeHashCollisions += Pair.Value - 1;
+			MaxProbeHashCollisions = FMath::Max(MaxProbeHashCollisions, Pair.Value);
+		}
+
+		Ar.Logf(TEXT(",%u,%u,%u,%u,%u,%u,%d,%d,%u,%d,%d,%u,%d,%d"),
+			NumCreated(), NumCreatedWide(), NumCreatedWithNumber(), UsedSlots, Capacity(),
+			UniqueUnmaskedIndices, NumUnmaskedIndexCollisions, MaxUnmaskedIndexCollisions,
+			UniqueMaskedIndices, NumMaskedIndexCollisions, MaxMaskedIndexCollisions,
+			UniqueProbeHashes, NumProbeHashCollisions, MaxProbeHashCollisions
+		);
+	}
+#endif
+
 	FNameEntryId Find(const FNameValue<Sensitivity>& Value) const
 	{
 		FRWScopeLock _(Lock, FRWScopeLockType::SLT_ReadOnly);
 
-		return Probe(Value).GetId();
+		FNameSlot& Slot = Probe(Value);
+		return Slot.GetId();
 	}
 
 	template<class ScopeLock = FWriteScopeLock>
@@ -755,6 +1101,38 @@ public:
 		InsertExistingEntryImpl<FWriteScopeLock>(Hash, ExistingId);
 	}
 
+#if UE_FNAME_OUTLINE_NUMBER
+	FNameEntryId FindWithNumber(const FNumberedNameValue<Sensitivity>& Value) const
+	{
+		FRWScopeLock _(Lock, FRWScopeLockType::SLT_ReadOnly);
+
+		FNameSlot& Slot = ProbeWithNumber(Value);
+		return Slot.GetId();
+	}
+
+	template<class ScopeLock = FWriteScopeLock>
+	FNameEntryId InsertWithNumber(const FNumberedNameValue<Sensitivity>& Value, bool& bCreatedNewEntry)
+	{
+		ScopeLock _(Lock);
+		FNameSlot& Slot = ProbeWithNumber(Value);
+
+		if (Slot.Used())
+		{
+			return Slot.GetId();
+		}
+
+		bCreatedNewEntry = true;
+		return CreateAndInsertEntryWithNumber<ScopeLock>(Slot, Value);
+	}
+
+#if WITH_CASE_PRESERVING_NAME
+	void InsertExistingEntryWithNumber(FNumberedNameValue<Sensitivity> Value)
+	{
+		InsertExistingEntryWithNumberImpl<FWriteScopeLock>(Value);
+	}
+#endif // WITH_CASE_PRESERVING_NAME
+#endif // UE_FNAME_OUTLINE_NUMBER
+
 private:
 	template<class ShardScopeLock>
 	void InsertExistingEntryImpl(FNameHash Hash, FNameEntryId ExistingId)
@@ -770,6 +1148,23 @@ private:
 		}
 	}
 
+#if UE_FNAME_OUTLINE_NUMBER && WITH_CASE_PRESERVING_NAME
+	template<class ShardScopeLock>
+	void InsertExistingEntryWithNumberImpl(FNumberedNameDisplayValue Value)
+	{
+		FNameSlot NewLookup(Value.ComparisonId, Value.Hash.SlotProbeHash);
+
+		ShardScopeLock _(Lock);
+
+		FNameSlot& Slot = ProbeWithNumber(Value);
+
+		if (!Slot.Used())
+		{
+			ClaimSlot(Slot, NewLookup);
+		}
+	}
+#endif
+
 public:
 	void InsertBatch(const TArrayView<FNameLoad<Sensitivity>> Batch)
 	{
@@ -888,6 +1283,20 @@ private:
 		return NewEntryId;
 	}
 
+#if UE_FNAME_OUTLINE_NUMBER
+	template<class EntryScopeLock>
+	FNameEntryId CreateAndInsertEntryWithNumber(FNameSlot& Slot, const FNumberedNameValue<Sensitivity>& Value)
+	{
+		FNameEntryId NewEntryId = Entries->CreateWithNumber<EntryScopeLock>(Value.StringPart, Value.NumberPart, GetExistingComparisonId(Value), Value.Hash.EntryProbeHeader);
+
+		ClaimSlot(Slot, FNameSlot(NewEntryId, Value.Hash.SlotProbeHash));
+
+		++NumCreatedWithNumberEntries;
+
+		return NewEntryId;
+	}
+#endif // UE_FNAME_OUTLINE_NUMBER
+
 	void Grow()
 	{
 		Grow(Capacity() * 2);
@@ -974,18 +1383,40 @@ private:
 		}
 	}
 
+#if UE_FNAME_OUTLINE_NUMBER
+	/** Find slot containing value or the first free slot that should be used to store it  */
+	FORCEINLINE FNameSlot& ProbeWithNumber(const FNumberedNameValue<Sensitivity>& Value) const
+	{
+		return Probe(Value.Hash.UnmaskedSlotIndex,
+			[&](FNameSlot Slot) { return Slot.GetProbeHash() == Value.Hash.SlotProbeHash &&
+			EntryEqualsValue(Entries->Resolve(Slot.GetId()), Value); });
+	}
+#endif // UE_FNAME_OUTLINE_NUMBER
+
 	FORCENOINLINE // Doesn't impact performance and makes sampling profiles more informative
 	void RehashAndInsert(FNameSlot OldSlot)
 	{
 		check(OldSlot.Used());
 
 		const FNameEntry& Entry = Entries->Resolve(OldSlot.GetId());
-		FNameBuffer DecodeBuffer;
-		FNameStringView Name = Entry.MakeView(DecodeBuffer);
-		FNameHash Hash = HashName<Sensitivity>(Name);
-		FNameSlot& NewSlot = Probe(Hash.UnmaskedSlotIndex, [](FNameSlot Slot) { return false; });
-		NewSlot = OldSlot;
-		++UsedSlots;
+#if UE_FNAME_OUTLINE_NUMBER
+		FNumberedNameValue<Sensitivity> Value;
+		if (RehashNameWithNumber(Entry, Value))
+		{
+			FNameSlot& NewSlot = Probe(Value.Hash.UnmaskedSlotIndex, [](FNameSlot Slot) { return false; });
+			NewSlot = OldSlot;
+			++UsedSlots;
+		}
+		else
+#endif
+		{
+			FNameBuffer DecodeBuffer;
+			FNameStringView Name = Entry.MakeView(DecodeBuffer);
+			FNameHash Hash = HashName<Sensitivity>(Name);
+			FNameSlot& NewSlot = Probe(Hash.UnmaskedSlotIndex, [](FNameSlot Slot) { return false; });
+			NewSlot = OldSlot;
+			++UsedSlots;
+		}
 	}
 };
 
@@ -1001,6 +1432,15 @@ public:
 	FNameEntryId	Find(EName Ename) const;
 	const EName*	FindEName(FNameEntryId Id) const;
 
+#if UE_FNAME_OUTLINE_NUMBER
+	FNameEntryId	StoreWithNumber(FNameEntryIds StringParts, int32 NumberPart);
+	FNameEntryId	FindWithNumber(FNameEntryId StringPart, int32 NumberPart) const;
+
+#if WITH_CASE_PRESERVING_NAME
+	FNameEntryId	StoreValueWithNumber(FNumberedNameDisplayValue Value, bool bReuseComparisonId);
+#endif // WITH_CASE_PRESERVING_NAME
+#endif // UE_FNAME_OUTLINE_NUMBER
+
 	/** @pre !!Handle */
 	FNameEntry&		Resolve(FNameEntryHandle Handle) const { return Entries.Resolve(Handle); }
 
@@ -1018,11 +1458,16 @@ public:
 	uint32			NumEntries() const;
 	uint32			NumAnsiEntries() const;
 	uint32			NumWideEntries() const;
+	uint32			NumNumberedEntries() const;
 	uint32			NumBlocks() const { return Entries.NumBlocks(); }
 	uint32			NumSlots() const;
 	void			LogStats(FOutputDevice& Ar) const;
+	void			LogHashStats(FOutputDevice& Ar) const;
 	uint8**			GetBlocksForDebugVisualizer() { return Entries.GetBlocksForDebugVisualizer(); }
 	TArray<const FNameEntry*> DebugDump() const;
+#if UE_FNAME_OUTLINE_NUMBER
+	void IterateNumberedNames(TFunctionRef<void(FName)> Func) const;
+#endif
 
 private:
 	enum { MaxENames = 512 };
@@ -1107,6 +1552,7 @@ FNameEntryId FNamePool::Find(EName Ename) const
 FNameEntryId FNamePool::Find(FNameStringView Name) const
 {
 #if WITH_CASE_PRESERVING_NAME
+	// Look for an exact match with the right casing first
 	FNameDisplayValue DisplayValue(Name);
 	if (FNameEntryId Existing = DisplayShards[DisplayValue.Hash.ShardIndex].Find(DisplayValue))
 	{
@@ -1121,6 +1567,7 @@ FNameEntryId FNamePool::Find(FNameStringView Name) const
 FNameEntryId FNamePool::Store(FNameStringView Name)
 {
 #if WITH_CASE_PRESERVING_NAME
+	// Look for an exact match with the right casing first
 	FNameDisplayValue DisplayValue(Name);
 	if (FNameEntryId Existing = DisplayShards[DisplayValue.Hash.ShardIndex].Find(DisplayValue))
 	{
@@ -1142,6 +1589,65 @@ FNameEntryId FNamePool::Store(FNameStringView Name)
 #endif
 }
 
+#if UE_FNAME_OUTLINE_NUMBER
+FNameEntryId FNamePool::StoreWithNumber(FNameEntryIds StringParts, int32 NumberPart)
+{
+#if WITH_CASE_PRESERVING_NAME
+	// Look for an exact match with the right casing first
+	FNumberedNameDisplayValue DisplayValue(StringParts.DisplayId, NumberPart);
+	if (FNameEntryId Existing = DisplayShards[DisplayValue.Hash.ShardIndex].FindWithNumber(DisplayValue))
+	{
+		return Existing;
+	}
+#endif
+
+	bool bAdded = false;
+
+	// Insert comparison name first since display value must contain comparison name
+	FNumberedNameComparisonValue ComparisonValue(StringParts.ComparisonId, NumberPart);
+	FNameEntryId ComparisonId = ComparisonShards[ComparisonValue.Hash.ShardIndex].InsertWithNumber(ComparisonValue, bAdded);
+	
+#if WITH_CASE_PRESERVING_NAME
+	DisplayValue.ComparisonId = ComparisonId;
+	return StoreValueWithNumber(DisplayValue, bAdded);
+#else
+	return ComparisonId;
+#endif
+}
+
+FNameEntryId FNamePool::FindWithNumber(FNameEntryId StringPart, int32 NumberPart) const
+{
+#if WITH_CASE_PRESERVING_NAME
+	// Look for an exact match with the right casing first
+	FNumberedNameDisplayValue DisplayValue(StringPart, NumberPart);
+	if (FNameEntryId Existing = DisplayShards[DisplayValue.Hash.ShardIndex].FindWithNumber(DisplayValue))
+	{
+		return Existing;
+	}
+	// Otherwise fall back to a name that compares equal case-insensitively
+#endif
+
+	FNumberedNameComparisonValue ComparisonValue(StringPart, NumberPart);
+	return ComparisonShards[ComparisonValue.Hash.ShardIndex].FindWithNumber(ComparisonValue);
+}
+
+#if WITH_CASE_PRESERVING_NAME
+FNameEntryId FNamePool::StoreValueWithNumber(FNumberedNameDisplayValue Value, bool bAddedComparisonEntry)
+{
+	FNamePoolShard<ENameCase::CaseSensitive>& DisplayShard = DisplayShards[Value.Hash.ShardIndex];
+	if (bAddedComparisonEntry || Value.StringPart == Value.ComparisonId)
+	{
+		DisplayShard.InsertExistingEntryWithNumber(Value);
+		return Value.ComparisonId;
+	}
+	else
+	{
+		return DisplayShard.InsertWithNumber(Value, bAddedComparisonEntry);
+	}
+}
+#endif // WITH_CASE_PRESERVING_NAME
+#endif // UE_FNAME_OUTLINE_NUMBER 
+
 FORCEINLINE FNameEntryId FNamePool::StoreValue(const FNameComparisonValue& ComparisonValue)
 {
 	bool bAdded = false;
@@ -1185,11 +1691,13 @@ uint32 FNamePool::NumEntries() const
 	for (const FNamePoolShardBase& Shard : DisplayShards)
 	{
 		Out += Shard.NumCreated();
+		Out += Shard.NumCreatedWithNumber();
 	}
 #endif
 	for (const FNamePoolShardBase& Shard : ComparisonShards)
 	{
 		Out += Shard.NumCreated();
+		Out += Shard.NumCreatedWithNumber();
 	}
 
 	return Out;
@@ -1197,7 +1705,7 @@ uint32 FNamePool::NumEntries() const
 
 uint32 FNamePool::NumAnsiEntries() const
 {
-	return NumEntries() - NumWideEntries();
+	return NumEntries() - NumWideEntries() - NumNumberedEntries();
 }
 
 uint32 FNamePool::NumWideEntries() const
@@ -1217,6 +1725,25 @@ uint32 FNamePool::NumWideEntries() const
 	return Out;
 }
 
+uint32 FNamePool::NumNumberedEntries() const
+{
+	uint32 Out = 0;
+#if UE_FNAME_OUTLINE_NUMBER
+#if WITH_CASE_PRESERVING_NAME
+	for (const FNamePoolShardBase& Shard : DisplayShards)
+	{
+		Out += Shard.NumCreatedWithNumber();
+	}
+#endif
+	for (const FNamePoolShardBase& Shard : ComparisonShards)
+	{
+		Out += Shard.NumCreatedWithNumber();
+	}
+#endif // UE_FNAME_OUTLINE_NUMBER
+
+	return Out;
+}
+
 uint32 FNamePool::NumSlots() const
 {
 	uint32 SlotCapacity = 0;
@@ -1237,6 +1764,22 @@ uint32 FNamePool::NumSlots() const
 void FNamePool::LogStats(FOutputDevice& Ar) const
 {
 	Ar.Logf(TEXT("%i FNames using in %ikB + %ikB"), NumEntries(), sizeof(FNamePool), Entries.NumBlocks() * FNameEntryAllocator::BlockSizeBytes / 1024);
+	Ar.Logf(TEXT("%d ansi FNames"), NumAnsiEntries());
+	Ar.Logf(TEXT("%d wide FNames"), NumWideEntries());
+#if UE_FNAME_OUTLINE_NUMBER
+	Ar.Logf(TEXT("%d FNames with numbers stored in entries"), NumNumberedEntries());
+#endif
+}
+
+void FNamePool::LogHashStats(FOutputDevice& Ar) const
+{
+#if !UE_BUILD_SHIPPING
+	ComparisonShards[0].LogCsvHeader(Ar);
+	for (int32 i = 0; i < FNamePoolShards; ++i)
+	{
+		ComparisonShards[i].LogStatsCsv(Ar);
+	}
+#endif
 }
 
 TArray<const FNameEntry*> FNamePool::DebugDump() const
@@ -1247,6 +1790,14 @@ TArray<const FNameEntry*> FNamePool::DebugDump() const
 	return Out;
 }
 
+#if UE_FNAME_OUTLINE_NUMBER
+void FNamePool::IterateNumberedNames(TFunctionRef<void(FName)> Func) const
+{
+	Entries.IterateNumberedNames(Func);
+}
+#endif
+
+
 bool FNamePool::IsValid(FNameEntryHandle Handle) const
 {
 	return Handle.Block < Entries.NumBlocks();
@@ -1458,6 +2009,11 @@ FORCEINLINE ANSICHAR const* FNameEntry::GetUnterminatedName(ANSICHAR(&OptionalDe
 
 FORCEINLINE FNameStringView FNameEntry::MakeView(FNameBuffer& OptionalDecodeBuffer) const
 {
+#if UE_FNAME_OUTLINE_NUMBER
+	checkf(Header.Len != 0, TEXT("Cannot make a string view for a name-with-number entry"));
+#else // UE_FNAME_OUTLINE_NUMBER
+	check(Header.Len != 0);
+#endif // UE_FNAME_OUTLINE_NUMBER
 	return IsWide()	? FNameStringView(GetUnterminatedName(OptionalDecodeBuffer.WideName), GetNameLength())
 					: FNameStringView(GetUnterminatedName(OptionalDecodeBuffer.AnsiName), GetNameLength());
 }
@@ -1526,6 +2082,7 @@ static const TCHAR* EntryToCString(const FNameEntry& Entry, FNameBuffer& Temp)
 
 FString FNameEntry::GetPlainNameString() const
 {
+	checkName(Header.Len != 0);
 	FNameBuffer Temp;
 	if (Header.bIsWide)
 	{
@@ -1539,12 +2096,14 @@ FString FNameEntry::GetPlainNameString() const
 
 void FNameEntry::AppendNameToString(FString& Out) const
 {
+	checkName(Header.Len != 0);
 	FNameBuffer Temp;
 	Out.Append(EntryToCString(*this, Temp), Header.Len);
 }
 
 void FNameEntry::AppendNameToString(FStringBuilderBase& Out) const
 {
+	checkName(Header.Len != 0);
 	const int32 Offset = Out.AddUninitialized(Header.Len);
 	TCHAR* OutChars = Out.GetData() + Offset;
 	if (Header.bIsWide)
@@ -1561,6 +2120,7 @@ void FNameEntry::AppendNameToString(FStringBuilderBase& Out) const
 
 void FNameEntry::AppendAnsiNameToString(FAnsiStringBuilderBase& Out) const
 {
+	checkName(Header.Len != 0);
 	check(!IsWide());
 	const int32 Offset = Out.AddUninitialized(Header.Len);
 	CopyUnterminatedName(Out.GetData() + Offset);
@@ -1568,6 +2128,7 @@ void FNameEntry::AppendAnsiNameToString(FAnsiStringBuilderBase& Out) const
 
 void FNameEntry::AppendNameToPathString(FString& Out) const
 {
+	checkName(Header.Len != 0);
 	FNameBuffer Temp;
 	Out.PathAppend(EntryToCString(*this, Temp), Header.Len);
 }
@@ -1583,9 +2144,32 @@ int32 FNameEntry::GetSize(int32 Length, bool bIsPureAnsi)
 	return Align(Bytes, alignof(FNameEntry));
 }
 
+#if UE_FNAME_OUTLINE_NUMBER
+int32 FNameEntry::GetNumberedEntrySize()
+{
+	int32 Bytes = GetDataOffset() + sizeof(FNameEntry::NumberedName);
+	return Align(Bytes, alignof(FNameEntry));
+}
+
+uint32 FNameEntry::GetNumber() const
+{
+	checkName(Header.Len == 0);
+	return NumberedName.Number;
+}
+#endif
+
 int32 FNameEntry::GetSizeInBytes() const
 {
-	return GetSize(GetNameLength(), !IsWide());
+#if UE_FNAME_OUTLINE_NUMBER
+	if (Header.Len == 0)
+	{
+		return GetNumberedEntrySize();
+	}
+	else
+#endif
+	{
+		return GetSize(GetNameLength(), !IsWide());
+	}
 }
 
 FNameEntrySerialized::FNameEntrySerialized(const FNameEntry& NameEntry)
@@ -1644,6 +2228,13 @@ int32 FName::GetNumWideNames()
 	return GetNamePool().NumWideEntries();
 }
 
+#if UE_FNAME_OUTLINE_NUMBER
+int32 FName::GetNumNumberedNames()
+{
+	return GetNamePool().NumNumberedEntries();
+}
+#endif
+
 TArray<const FNameEntry*> FName::DebugDump()
 {
 	return GetNamePool().DebugDump();
@@ -1657,7 +2248,8 @@ FNameEntry const* FName::GetEntry(EName Ename)
 
 FNameEntry const* FName::GetEntry(FNameEntryId Id)
 {
-	return &GetNamePool().Resolve(Id);
+	// Public interface, recurse to the actual string entry if necessary
+	return ResolveEntryRecursive(Id);
 }
 
 FString FName::NameToDisplayString( const FString& InDisplayName, const bool bIsBool )
@@ -1947,37 +2539,6 @@ struct FNameHelper
 		return MakeWithNumber(View, FindType, InternalNumber);
 	}
 
-	template<typename CharType>
-	static uint32 ParseNumber(const CharType* Name, int32& InOutLen)
-	{
-		const int32 Len = InOutLen;
-		int32 Digits = 0;
-		for (const CharType* It = Name + Len - 1; It >= Name && *It >= '0' && *It <= '9'; --It)
-		{
-			++Digits;
-		}
-
-		const CharType* FirstDigit = Name + Len - Digits;
-		static constexpr int32 MaxDigitsInt32 = 10;
-		if (Digits && Digits < Len && *(FirstDigit - 1) == '_' && Digits <= MaxDigitsInt32)
-		{
-			// check for the case where there are multiple digits after the _ and the first one
-			// is a 0 ("Rocket_04"). Can't split this case. (So, we check if the first char
-			// is not 0 or the length of the number is 1 (since ROcket_0 is valid)
-			if (Digits == 1 || *FirstDigit != '0')
-			{
-				int64 Number = Atoi64(Name + Len - Digits, Digits);
-				if (Number < MAX_int32)
-				{
-					InOutLen -= 1 + Digits;
-					return static_cast<uint32>(NAME_EXTERNAL_TO_INTERNAL(Number));
-				}
-			}
-		}
-
-		return NAME_NO_NUMBER_INTERNAL;
-	}
-
 	static FName MakeWithNumber(FNameAnsiStringView	 View, EFindName FindType, int32 InternalNumber)
 	{
 		// Ignore the supplied number if the name string is empty
@@ -1987,7 +2548,7 @@ struct FNameHelper
 			return FName();
 		}
 
-		return Make(FNameStringView(View.Str, View.Len), FindType, InternalNumber);
+		return MakeInternal(FNameStringView(View.Str, View.Len), FindType, InternalNumber);
 	}
 
 	static FName MakeWithNumber(FNameUtf8StringViewWithWidth View, EFindName FindType, int32 InternalNumber)
@@ -2001,12 +2562,12 @@ struct FNameHelper
 
 		if (!View.bIsWide)
 		{
-			return Make(FNameStringView(reinterpret_cast<const ANSICHAR*>(View.Str), View.Len), FindType, InternalNumber);
+			return MakeInternal(FNameStringView(reinterpret_cast<const ANSICHAR*>(View.Str), View.Len), FindType, InternalNumber);
 		}
 		else
 		{
 			TStringConversion<FUTF8ToTCHAR_Convert, NAME_SIZE> WideName(View.Str, View.Len);
-			return Make(FNameStringView(WideName.Get(), WideName.Length()), FindType, InternalNumber);
+			return MakeInternal(FNameStringView(WideName.Get(), WideName.Length()), FindType, InternalNumber);
 		}
 	}
 
@@ -2028,49 +2589,204 @@ struct FNameHelper
 			{
 				AnsiName[I] = View.Str[I];
 			}
-			return Make(FNameStringView(AnsiName, View.Len), FindType, InternalNumber);
+			return MakeInternal(FNameStringView(AnsiName, View.Len), FindType, InternalNumber);
 		}
 		else
 		{
-			return Make(FNameStringView(View.Str, View.Len), FindType, InternalNumber);
+			return MakeInternal(FNameStringView(View.Str, View.Len), FindType, InternalNumber);
 		}
 	}
 
-	static FName Make(FNameStringView View, EFindName FindType, int32 InternalNumber)
+	static FName MakeFromLoaded(const FNameEntrySerialized& LoadedEntry)
+	{
+		FNameStringView View = LoadedEntry.bIsWide
+			? FNameStringView(LoadedEntry.WideName, FCStringWide::Strlen(LoadedEntry.WideName))
+			: FNameStringView(LoadedEntry.AnsiName, FCStringAnsi::Strlen(LoadedEntry.AnsiName));
+
+		return MakeInternal(View, FNAME_Add, NAME_NO_NUMBER_INTERNAL);
+	}
+
+	static FName MakeWithNumber(FNameEntryIds BaseIds, EFindName FindType, int32 InternalNumber)
+	{
+#if UE_FNAME_OUTLINE_NUMBER
+#if WITH_CASE_PRESERVING_NAME
+		// Advanced users must pass in matching display & comparison ids
+		checkName(FName::ResolveEntry(BaseIds.DisplayId)->ComparisonId == BaseIds.ComparisonId);
+#endif
+
+		// If BaseIds are NAME_None, we want to produce a numbered suffix of None. 
+		// If anything searches for BaseIds they need to validate whether they want to pass NAME_None on or not before calling this function.
+		if (InternalNumber == NAME_NO_NUMBER_INTERNAL)
+		{
+			return FinalConstruct(BaseIds);
+		}
+
+		FNamePool& Pool = GetNamePool();
+		if (FindType == FNAME_Add)
+		{
+			FNameEntryId DisplayId = Pool.StoreWithNumber(BaseIds, InternalNumber);
+			return FinalConstruct(FNameEntryIds{ ResolveComparisonId(DisplayId), DisplayId });
+		}
+		else if (FindType == FNAME_Find)
+		{
+			FNameEntryId DisplayId = Pool.FindWithNumber(BaseIds.DisplayId, InternalNumber);
+			if (DisplayId)
+			{
+				return FinalConstruct(FNameEntryIds{ ResolveComparisonId(DisplayId), DisplayId });
+			}
+			else
+			{
+				// Not found
+				return FName();
+			}
+		}
+		else
+		{
+			checkf(false, TEXT("FNAME_Replace_Not_Safe_For_Threading not supported for numbered names"));
+			return FName();
+		}
+#else
+		// Number is just stored in the FName pass it straight on
+		return FinalConstruct(BaseIds, InternalNumber);
+#endif
+	}
+
+	template<class CharType>
+	static bool EqualsString(FName Name, const CharType* Str)
+	{
+		// Make NAME_None == TEXT("") or nullptr consistent with NAME_None == FName(TEXT("")) or FName(nullptr)
+		if (Str == nullptr || Str[0] == '\0')
+		{
+			return Name.IsNone();
+		}
+
+		const FNameEntry& Entry = *Name.GetComparisonNameEntry();
+
+		uint32 NameLen = Entry.Header.Len;
+		FNameBuffer Temp;
+		return Entry.IsWide()
+			? StringAndNumberEqualsString(Entry.GetUnterminatedName(Temp.WideName), NameLen, Name.GetNumber(), Str)
+			: StringAndNumberEqualsString(Entry.GetUnterminatedName(Temp.AnsiName), NameLen, Name.GetNumber(), Str);
+	}
+
+	template<typename CharType>
+	static uint32 ParseNumber(const CharType* Name, int32& InOutLen)
+	{
+		const int32 Len = InOutLen;
+		int32 Digits = 0;
+		for (const CharType* It = Name + Len - 1; It >= Name && *It >= '0' && *It <= '9'; --It)
+		{
+			++Digits;
+		}
+
+		const CharType* FirstDigit = Name + Len - Digits;
+		static constexpr int32 MaxDigitsInt32 = 10;
+		if (Digits && Digits < Len && *(FirstDigit - 1) == '_' && Digits <= MaxDigitsInt32)
+		{
+			// check for the case where there are multiple digits after the _ and the first one
+			// is a 0 ("Rocket_04"). Can't split this case. (So, we check if the first char
+			// is not 0 or the length of the number is 1 (since ROcket_0 is valid)
+			if (Digits == 1 || *FirstDigit != '0')
+			{
+				int64 Number = Atoi64(Name + Len - Digits, Digits);
+				if (Number < MAX_int32)
+				{
+					InOutLen -= 1 + Digits;
+					return static_cast<uint32>(NAME_EXTERNAL_TO_INTERNAL(Number));
+				}
+			}
+		}
+
+		return NAME_NO_NUMBER_INTERNAL;
+	}
+
+	// Internal helpers for Make* functions
+private:
+
+	static FName MakeInternal(FNameStringView View, EFindName FindType, int32 InternalNumber)
+	{
+		FNameEntryIds Ids = FindOrStoreString(View, FindType);
+#if UE_FNAME_OUTLINE_NUMBER
+		if (FindType == FNAME_Find && !Ids.DisplayId)
+		{
+			// We need to disambiguate here between "we found 'None'" and "we didn't find what we were looking for"
+			if (View.IsNoneString())
+			{
+				return MakeWithNumber(Ids, FindType, InternalNumber); // 'None' with suffix
+			}
+			else
+			{
+				return FName(); // Not found
+			}
+		}
+#endif
+		return MakeWithNumber(Ids, FindType, InternalNumber);
+	}
+
+	// Indices have already been numbered if necessary
+	// Not an FName constructor because of implementation details around UE_FNAME_OUTLINE_NUMBER that we want to hide from FName interface
+#if UE_FNAME_OUTLINE_NUMBER
+	static FName FinalConstruct(FNameEntryIds Ids)
+#else
+	static FName FinalConstruct(FNameEntryIds Ids, int32 InternalNumber)
+#endif // UE_FNAME_OUTLINE_NUMBER
+	{
+		FName Out;
+		Out.ComparisonIndex = Ids.ComparisonId;
+#if WITH_CASE_PRESERVING_NAME
+		Out.DisplayIndex = Ids.DisplayId;
+#endif
+#if !UE_FNAME_OUTLINE_NUMBER
+		Out.Number = InternalNumber;
+#endif
+		return Out;
+	}
+
+#if WITH_CASE_PRESERVING_NAME
+	static FNameEntryId ResolveComparisonId(FNameEntryId DisplayId)
+	{
+		if (DisplayId.IsNone()) { return FNameEntryId(); }
+
+		return GetNamePool().Resolve(DisplayId).ComparisonId;
+
+	}
+#else
+	static FNameEntryId ResolveComparisonId(FNameEntryId DisplayId)
+	{
+		return DisplayId;
+	}
+#endif
+
+	// Find or store a plain string name entry, returning both comparison and display id for it.
+	// If not found, returns NAME_None for both indices.
+	static FNameEntryIds FindOrStoreString(FNameStringView View, EFindName FindType)
 	{
 		if (View.Len >= NAME_SIZE)
 		{
 			// If we're doing a find, and the string is too long, then clearly we didn't find it
 			if (FindType == FNAME_Find)
 			{
-				return FName();
+				return {};
 			}
 
 			checkf(false, TEXT("FName's %d max length exceeded. Got %d characters excluding null-terminator:\n%.*s"), 
 				NAME_SIZE - 1, View.Len, NAME_SIZE, View.IsAnsi() ? ANSI_TO_TCHAR(View.Data) : View.Data);
-			return FName("ERROR_NAME_SIZE_EXCEEDED");
+
+			const ANSICHAR* ErrorString = "ERROR_NAME_SIZE_EXCEEDED";
+			return FindOrStoreString(FNameStringView(ErrorString, FCStringAnsi::Strlen(ErrorString), false), FNAME_Add);
 		}
 		
 		FNamePool& Pool = GetNamePool();
 
-		FNameEntryId DisplayId, ComparisonId;
-		if (FindType == FNAME_Add)
-		{
-			DisplayId = Pool.Store(View);
-#if WITH_CASE_PRESERVING_NAME
-			ComparisonId = Pool.Resolve(DisplayId).ComparisonId;
-#else
-			ComparisonId = DisplayId;
-#endif
+		if (FindType == FNAME_Add)
+		{
+			FNameEntryId DisplayId = Pool.Store(View);
+			return FNameEntryIds{ ResolveComparisonId(DisplayId), DisplayId };
 		}
 		else if (FindType == FNAME_Find)
 		{
-			DisplayId = Pool.Find(View);
-#if WITH_CASE_PRESERVING_NAME
-			ComparisonId = DisplayId ? Pool.Resolve(DisplayId).ComparisonId : DisplayId;
-#else
-			ComparisonId = DisplayId;
-#endif
+			FNameEntryId DisplayId = Pool.Find(View);
+			return FNameEntryIds{ ResolveComparisonId(DisplayId), DisplayId };
 		}
 		else
 		{
@@ -2079,44 +2795,30 @@ struct FNameHelper
 #if FNAME_WRITE_PROTECT_PAGES
 			checkf(false, TEXT("FNAME_Replace_Not_Safe_For_Threading can't be used together with page protection."));
 #endif
-			DisplayId = Pool.Store(View);
-#if WITH_CASE_PRESERVING_NAME
-			ComparisonId = Pool.Resolve(DisplayId).ComparisonId;
-#else
-			ComparisonId = DisplayId;
-#endif
+
+			FNameEntryId DisplayId = Pool.Store(View);
+			FNameEntryId ComparisonId = ResolveComparisonId(DisplayId);
 			ReplaceName(Pool.Resolve(ComparisonId), View);
+			return FNameEntryIds{ ComparisonId, DisplayId };
 		}
-
-		return FName(ComparisonId, DisplayId, InternalNumber);
-	}
-
-	static FName MakeFromLoaded(const FNameEntrySerialized& LoadedEntry)
-	{
-		FNameStringView View = LoadedEntry.bIsWide
-			? FNameStringView(LoadedEntry.WideName, FCStringWide::Strlen(LoadedEntry.WideName))
-			: FNameStringView(LoadedEntry.AnsiName, FCStringAnsi::Strlen(LoadedEntry.AnsiName));
-
-		return Make(View, FNAME_Add, NAME_NO_NUMBER_INTERNAL);
 	}
 
-	template<class CharType>
-	static bool EqualsString(FName Name, const CharType* Str)
+#if UE_FNAME_OUTLINE_NUMBER
+	static FName FindNumbered(FNameEntryId InDisplayIndex, int32 InternalNumber)
 	{
-		// Make NAME_None == TEXT("") or nullptr consistent with NAME_None == FName(TEXT("")) or FName(nullptr)
-		if (Str == nullptr || Str[0] == '\0')
+		FNamePool& Pool = GetNamePool();
+		FNameEntryId DisplayId = Pool.FindWithNumber(InDisplayIndex, InternalNumber);
+		if (DisplayId)
 		{
-			return Name.IsNone();
+			FNameEntryId ComparisonId = ResolveComparisonId(DisplayId);
+			return FinalConstruct(FNameEntryIds{ ComparisonId, DisplayId });
+		}
+		else
+		{
+			return FName();
 		}
-
-		const FNameEntry& Entry = *Name.GetComparisonNameEntry();
-
-		uint32 NameLen = Entry.Header.Len;
-		FNameBuffer Temp;
-		return Entry.IsWide()
-			? StringAndNumberEqualsString(Entry.GetUnterminatedName(Temp.WideName), NameLen, Name.GetNumber(), Str)
-			: StringAndNumberEqualsString(Entry.GetUnterminatedName(Temp.AnsiName), NameLen, Name.GetNumber(), Str);
 	}
+#endif // UE_FNAME_OUTLINE_NUMBER
 
 	static void ReplaceName(FNameEntry& Existing, FNameStringView Updated)
 	{
@@ -2142,6 +2844,86 @@ FNameEntryId FName::GetComparisonIdFromDisplayId(FNameEntryId DisplayId)
 }
 #endif
 
+#if UE_FNAME_OUTLINE_NUMBER
+FNameEntryId FName::GetComparisonIndex() const
+{
+	const FNameEntry& Entry = GetNamePool().Resolve(ComparisonIndex);
+	if (Entry.Header.Len == 0)
+	{
+		return Entry.NumberedName.Id;
+	}
+	else
+	{
+		return ComparisonIndex;
+	}
+}
+
+FNameEntryId FName::GetDisplayIndex() const
+{
+#if WITH_CASE_PRESERVING_NAME
+	const FNameEntry& Entry = GetNamePool().Resolve(DisplayIndex);
+	if (Entry.Header.Len == 0)
+	{
+		return Entry.NumberedName.Id;
+	}
+	else
+	{
+		return DisplayIndex;
+	}
+#else
+	return GetComparisonIndex();
+#endif
+}
+
+FName FName::FindNumberedName(FNameEntryId DisplayId, int32 Number)
+{
+	return FNameHelper::MakeWithNumber(FNameEntryIds{FName::GetComparisonIdFromDisplayId(DisplayId),DisplayId}, FNAME_Find , Number);
+}
+
+int32 FName::GetNumber() const 
+{
+	FNameEntryId Id = GetDisplayIndexInternal();
+
+	const FNameEntry* Entry = ResolveEntry(Id);
+	if (Entry->Header.Len == 0)
+	{
+		return Entry->NumberedName.Number;
+	}
+	else
+	{
+		return NAME_NO_NUMBER_INTERNAL;
+	}
+}
+void FName::SetNumber(int32 InNumber) 
+{
+	FName NewName = FNameHelper::MakeWithNumber(FNameEntryIds{GetComparisonIndex(), GetDisplayIndex()}, FNAME_Add, InNumber);
+	*this = NewName;
+}
+#endif // UE_FNAME_OUTLINE_NUMBER
+
+// Resolve the entry directly referred to by LookupId
+const FNameEntry* FName::ResolveEntry(FNameEntryId LookupId)
+{
+	return &GetNamePool().Resolve(LookupId);
+}
+
+// Recursively resolve through the entry referred to by LookupId to reach the allocated string entry, in the case of UE_FNAME_OUTLINE_NUMBER=1
+const FNameEntry* FName::ResolveEntryRecursive(FNameEntryId LookupId)
+{
+	const FNameEntry* Entry = ResolveEntry(LookupId);
+#if UE_FNAME_OUTLINE_NUMBER
+	if (Entry->Header.Len == 0)
+	{
+		return ResolveEntry(Entry->NumberedName.Id); // Should only ever recurse one level
+	}
+	else
+#endif
+	{
+		return Entry;
+	}
+}
+
+
 FName::FName(const WIDECHAR* Name, EFindName FindType)
 	: FName(FNameHelper::MakeDetectNumber(MakeUnconvertedView(Name), FindType))
 {}
@@ -2180,23 +2962,56 @@ FName::FName(const UTF8CHAR* Name, int32 InNumber, EFindName FindType)
 
 FName::FName(int32 Len, const WIDECHAR* Name, int32 InNumber, EFindName FindType)
 	: FName(InNumber != NAME_NO_NUMBER_INTERNAL ? FNameHelper::MakeWithNumber(MakeUnconvertedView(Name, Len), FindType, InNumber)
-												: FNameHelper::MakeDetectNumber(MakeUnconvertedView(Name, Len), FindType))
+		: FNameHelper::MakeDetectNumber(MakeUnconvertedView(Name, Len), FindType))
 {}
 
 FName::FName(int32 Len, const ANSICHAR* Name, int32 InNumber, EFindName FindType)
 	: FName(InNumber != NAME_NO_NUMBER_INTERNAL ? FNameHelper::MakeWithNumber(MakeUnconvertedView(Name, Len), FindType, InNumber)
-												: FNameHelper::MakeDetectNumber(MakeUnconvertedView(Name, Len), FindType))
+		: FNameHelper::MakeDetectNumber(MakeUnconvertedView(Name, Len), FindType))
 {}
 
 FName::FName(int32 Len, const UTF8CHAR* Name, int32 InNumber, EFindName FindType)
 	: FName(InNumber != NAME_NO_NUMBER_INTERNAL ? FNameHelper::MakeWithNumber(MakeUnconvertedView(Name, Len), FindType, InNumber)
-												: FNameHelper::MakeDetectNumber(MakeUnconvertedView(Name, Len), FindType))
+		: FNameHelper::MakeDetectNumber(MakeUnconvertedView(Name, Len), FindType))
 {}
 
 FName::FName(const TCHAR* Name, int32 InNumber, EFindName FindType, bool bSplitName)
-	: FName(InNumber == NAME_NO_NUMBER_INTERNAL && bSplitName 
-			? FNameHelper::MakeDetectNumber(MakeUnconvertedView(Name), FindType)
-			: FNameHelper::MakeWithNumber(MakeUnconvertedView(Name), FindType, InNumber))
+	: FName(InNumber == NAME_NO_NUMBER_INTERNAL && bSplitName
+		? FNameHelper::MakeDetectNumber(MakeUnconvertedView(Name), FindType)
+		: FNameHelper::MakeWithNumber(MakeUnconvertedView(Name), FindType, InNumber))
+{}
+
+FName::FName(const WIDECHAR* Name, int32 InNumber)
+	: FName(FNameHelper::MakeWithNumber(MakeUnconvertedView(Name), FNAME_Add, InNumber))
+{}
+
+FName::FName(const ANSICHAR* Name, int32 InNumber)
+	: FName(FNameHelper::MakeWithNumber(MakeUnconvertedView(Name), FNAME_Add, InNumber))
+{}
+
+FName::FName(const UTF8CHAR* Name, int32 InNumber)
+	: FName(FNameHelper::MakeWithNumber(MakeUnconvertedView(Name), FNAME_Add, InNumber))
+{}
+
+FName::FName(int32 Len, const WIDECHAR* Name, int32 InNumber)
+	: FName(InNumber != NAME_NO_NUMBER_INTERNAL ? FNameHelper::MakeWithNumber(MakeUnconvertedView(Name, Len), FNAME_Add, InNumber)
+		: FNameHelper::MakeDetectNumber(MakeUnconvertedView(Name, Len), FNAME_Add))
+{}
+
+FName::FName(int32 Len, const ANSICHAR* Name, int32 InNumber)
+	: FName(InNumber != NAME_NO_NUMBER_INTERNAL ? FNameHelper::MakeWithNumber(MakeUnconvertedView(Name, Len), FNAME_Add, InNumber)
+		: FNameHelper::MakeDetectNumber(MakeUnconvertedView(Name, Len), FNAME_Add))
+{}
+
+FName::FName(int32 Len, const UTF8CHAR* Name, int32 InNumber)
+	: FName(InNumber != NAME_NO_NUMBER_INTERNAL ? FNameHelper::MakeWithNumber(MakeUnconvertedView(Name, Len), FNAME_Add, InNumber)
+		: FNameHelper::MakeDetectNumber(MakeUnconvertedView(Name, Len), FNAME_Add))
+{}
+
+FName::FName(const TCHAR* Name, int32 InNumber, bool bSplitName)
+	: FName(InNumber == NAME_NO_NUMBER_INTERNAL && bSplitName
+		? FNameHelper::MakeDetectNumber(MakeUnconvertedView(Name), FNAME_Add)
+		: FNameHelper::MakeWithNumber(MakeUnconvertedView(Name), FNAME_Add, InNumber))
 {}
 
 FName::FName(const FNameEntrySerialized& LoadedEntry)
@@ -2215,14 +3030,29 @@ bool FName::operator==(const WIDECHAR* Str) const
 
 int32 FName::Compare( const FName& Other ) const
 {
+#if UE_FNAME_OUTLINE_NUMBER
+	if (GetComparisonIndexInternal() == Other.GetComparisonIndexInternal())
+	{
+		return 0; // Identical
+	}
+	else if (GetComparisonIndex() == Other.GetComparisonIndex())
+	{
+		// Same string, number may be different 
+		return GetNumber() - Other.GetNumber();
+	}
+
+	// Names don't match. This means we don't even need to check numbers.
+	return CompareDifferentIdsAlphabetically(GetComparisonIndex(), Other.GetComparisonIndex());
+#else	// UE_FNAME_OUTLINE_NUMBER
 	// Names match, check whether numbers match.
-	if (ComparisonIndex == Other.ComparisonIndex)
+	if (GetComparisonIndex() == Other.GetComparisonIndex())
 	{
 		return GetNumber() - Other.GetNumber();
 	}
 
 	// Names don't match. This means we don't even need to check numbers.
-	return CompareDifferentIdsAlphabetically(ComparisonIndex, Other.ComparisonIndex);
+	return CompareDifferentIdsAlphabetically(GetComparisonIndex(), Other.GetComparisonIndex());
+#endif // UE_FNAME_OUTLINE_NUMBER
 }
 
 uint32 FName::GetPlainNameString(TCHAR(&OutName)[NAME_SIZE]) const
@@ -2249,21 +3079,24 @@ void FName::GetPlainWIDEString(WIDECHAR(&WideName)[NAME_SIZE]) const
 
 const FNameEntry* FName::GetComparisonNameEntry() const
 {
-	return &GetNamePool().Resolve(GetComparisonIndex());
+	return ResolveEntryRecursive(GetComparisonIndexInternal());
 }
 
 const FNameEntry* FName::GetDisplayNameEntry() const
 {
-	return &GetNamePool().Resolve(GetDisplayIndex());
+	return ResolveEntryRecursive(GetDisplayIndexInternal());
 }
 
 FString FName::ToString() const
 {
+// With UE_FNAME_OUTLINE_NUMBER reading number isn't free so skip this check
+#if !UE_FNAME_OUTLINE_NUMBER
 	if (GetNumber() == NAME_NO_NUMBER_INTERNAL)
 	{
 		// Avoids some extra allocations in non-number case
 		return GetDisplayNameEntry()->GetPlainNameString();
 	}
+#endif // !UE_FNAME_OUTLINE_NUMBER
 	
 	FString Out;	
 	ToString(Out);
@@ -2272,6 +3105,25 @@ FString FName::ToString() const
 
 void FName::ToString(FString& Out) const
 {
+#if UE_FNAME_OUTLINE_NUMBER
+	FNameEntryId Id = GetDisplayIndexInternal();	
+	const FNameEntry* ThisNameEntry = ResolveEntry(Id);
+	if (ThisNameEntry->Header.Len != 0)
+	{
+		// No number
+		Out.Reset(ThisNameEntry->GetNameLength());
+		ThisNameEntry->AppendNameToString(Out);
+	}
+	else
+	{
+		const FNameEntry* BaseEntry = ResolveEntry(ThisNameEntry->NumberedName.Id);
+		Out.Reset(BaseEntry->GetNameLength() + 6);
+		BaseEntry->AppendNameToString(Out);
+
+		Out += TEXT('_');
+		Out.AppendInt(NAME_INTERNAL_TO_EXTERNAL(ThisNameEntry->GetNumber()));
+	}
+#else // UE_FNAME_OUTLINE_NUMBER
 	// A version of ToString that saves at least one string copy
 	const FNameEntry* const NameEntry = GetDisplayNameEntry();
 
@@ -2288,6 +3140,7 @@ void FName::ToString(FString& Out) const
 		Out += TEXT('_');
 		Out.AppendInt(NAME_INTERNAL_TO_EXTERNAL(GetNumber()));
 	}
+#endif // UE_FNAME_OUTLINE_NUMBER
 }
 
 void FName::ToString(FStringBuilderBase& Out) const
@@ -2353,6 +3206,23 @@ void FName::AppendString(FString& Out) const
 
 void FName::AppendString(FStringBuilderBase& Out) const
 {
+#if UE_FNAME_OUTLINE_NUMBER
+	FNameEntryId Id = GetDisplayIndexInternal();	
+	const FNameEntry* ThisNameEntry = ResolveEntry(GetDisplayIndexInternal());
+
+	if (ThisNameEntry->Header.Len != 0)
+	{
+		// No number
+		ThisNameEntry->AppendNameToString(Out);
+	}
+	else
+	{
+		const FNameEntry* BaseEntry = ResolveEntry(ThisNameEntry->NumberedName.Id);
+		BaseEntry->AppendNameToString(Out);
+		Out << TEXT('_') << NAME_INTERNAL_TO_EXTERNAL(ThisNameEntry->GetNumber());
+	}
+
+#else // UE_FNAME_OUTLINE_NUMBER
 	GetDisplayNameEntry()->AppendNameToString(Out);
 
 	const int32 InternalNumber = GetNumber();
@@ -2360,6 +3230,7 @@ void FName::AppendString(FStringBuilderBase& Out) const
 	{
 		Out << TEXT('_') << NAME_INTERNAL_TO_EXTERNAL(InternalNumber);
 	}
+#endif // UE_FNAME_OUTLINE_NUMBER
 }
 
 bool FName::TryAppendAnsiString(FAnsiStringBuilderBase& Out) const
@@ -2405,7 +3276,26 @@ void FName::DisplayHash(FOutputDevice& Ar)
 
 FString FName::SafeString(FNameEntryId InDisplayIndex, int32 InstanceNumber)
 {
+#if UE_FNAME_OUTLINE_NUMBER
+	const FNameEntry* Entry = &GetNamePool().Resolve(InDisplayIndex);
+	if (!Entry)
+	{
+		return FString();
+	}
+
+	if (InstanceNumber == NAME_NO_NUMBER_INTERNAL)
+	{
+		// Avoids some extra allocations in non-number case
+		return Entry->GetPlainNameString();
+	}
+
+	TStringBuilder<FName::StringBufferSize> Builder;
+	Entry->AppendNameToString(Builder);
+	Builder << TEXT("_") << NAME_INTERNAL_TO_EXTERNAL(InstanceNumber);
+	return FString(Builder);
+#else // UE_FNAME_OUTLINE_NUMBER
 	return FName(InDisplayIndex, InDisplayIndex, InstanceNumber).ToString();
+#endif // UE_FNAME_OUTLINE_NUMBER
 }
 
 bool FName::IsValidXName(const FName InName, const FString& InInvalidChars, FText* OutReason, const FText* InErrorCtx)
@@ -2497,10 +3387,19 @@ void FName::AutoTest()
 	const FName AutoTest_2(TEXT("AutoTest_2"));
 	const FName AutoTestB_2(TEXT("AutoTestB_2"));
 
+	check(AutoTest_1.GetNumber() == NAME_EXTERNAL_TO_INTERNAL(1));
+	check(autoTest_1.GetNumber() == NAME_EXTERNAL_TO_INTERNAL(1));
+	check(autoTeSt_1.GetNumber() == NAME_EXTERNAL_TO_INTERNAL(1));
+	check(AutoTest_2.GetNumber() == NAME_EXTERNAL_TO_INTERNAL(2));
+	check(AutoTestB_2.GetNumber() == NAME_EXTERNAL_TO_INTERNAL(2));
+
 	check(AutoTest_1 != AutoTest_2);
 	check(AutoTest_1 == autoTest_1);
 	check(AutoTest_1 == autoTeSt_1);
 
+	check(AutoTest_1.ToUnstableInt() == autoTest_1.ToUnstableInt());
+	check(AutoTest_1.ToUnstableInt() != AutoTest_2.ToUnstableInt());
+
 	TCHAR Buffer[FName::StringBufferSize];
 
 #if WITH_CASE_PRESERVING_NAME
@@ -2522,6 +3421,7 @@ void FName::AutoTest()
 	check(*AutoTestB_2.GetPlainNameString() != *AutoTest_2.GetPlainNameString());
 	check(AutoTestB_2.GetNumber() == AutoTest_2.GetNumber());
 	check(autoTest_1.GetNumber() != AutoTest_2.GetNumber());
+	check(AutoTest_1.GetNumber() == autoTest_1.GetNumber());
 
 	check(FCStringAnsi::Strlen("None") == FName().GetStringLength());
 	check(FCStringAnsi::Strlen("ABC") == FName("ABC").GetStringLength());
@@ -2546,6 +3446,22 @@ void FName::AutoTest()
 	check(*FName().ToEName() == NAME_None);
 	check(NullName.GetComparisonIndex().ToUnstableInt() == 0);
 
+	// Numbered nones
+	FName None_7 = FName("None_7");
+	check(None_7.ToString() == "None_7");
+	check(None_7.GetComparisonIndex() == FNameEntryId::FromEName(EName::None));
+	check(None_7.GetNumber() == NAME_EXTERNAL_TO_INTERNAL(7));
+	check(FName(EName::None, NAME_EXTERNAL_TO_INTERNAL(7)) == None_7);
+	check(FName(FName(), NAME_EXTERNAL_TO_INTERNAL(7)) == None_7);
+	check(FName("None", NAME_EXTERNAL_TO_INTERNAL(7)) == None_7);
+	check(FName(FNameEntryId::FromEName(EName::None), FNameEntryId::FromEName(EName::None), NAME_EXTERNAL_TO_INTERNAL(7)) == None_7);
+
+	// Find existing numbered none 
+PRAGMA_DISABLE_DEPRECATION_WARNINGS
+	check(FName("None_7", FNAME_Find) == None_7);
+	check(FName("None", NAME_EXTERNAL_TO_INTERNAL(7), FNAME_Find) == None_7);
+PRAGMA_ENABLE_DEPRECATION_WARNINGS
+
 	const FName Cylinder(NAME_Cylinder);
 	check(Cylinder == FName("Cylinder"));
 	check(Cylinder.ToEName());
@@ -2591,6 +3507,22 @@ void FName::AutoTest()
 	if (Once)
 	{
 		check(FName("UniqueUnicorn!!", FNAME_Find) == FName());
+#if UE_FNAME_OUTLINE_NUMBER
+PRAGMA_DISABLE_DEPRECATION_WARNINGS
+		check(FName("UniqueUnicorn!!", 17, FNAME_Find) == FName());			// We can't find the suffix version either
+PRAGMA_ENABLE_DEPRECATION_WARNINGS
+#else
+PRAGMA_DISABLE_DEPRECATION_WARNINGS
+		check(FName("UniqueUnicorn!!", 17, FNAME_Find) == FName(FNameEntryId(), FNameEntryId(), 17));	 // We fail to find the string part so we make the name "None_16"
+PRAGMA_ENABLE_DEPRECATION_WARNINGS
+#endif
+
+#if UE_FNAME_OUTLINE_NUMBER
+PRAGMA_DISABLE_DEPRECATION_WARNINGS
+		const FName NumberedNone = FName("None", 17, FNAME_Add);
+		check(FName("UniqueUnicorn!!", 17, FNAME_Find) == FName());		// Still can't find a numbered name just because we added a None with the same number
+PRAGMA_ENABLE_DEPRECATION_WARNINGS
+#endif
 
 		// Check that FNAME_Find can find entries
 		const FName UniqueName("UniqueUnicorn!!", FNAME_Add);
@@ -2600,6 +3532,15 @@ void FName::AutoTest()
 		check(FName(TEXT("UNIQUEUNICORN!!"), FNAME_Find) == UniqueName);
 		check(FName("uniqueunicorn!!", FNAME_Find) == UniqueName);
 
+#if UE_FNAME_OUTLINE_NUMBER
+PRAGMA_DISABLE_DEPRECATION_WARNINGS
+		check(FName("UniqueUnicorn!!", 17, FNAME_Find) == FName());		// Still can't find a numbered name just because we made the base version and a None with the same number
+		check(FName("UniqueUnicorn!!", 17, FNAME_Add) != FName());		// Explicitly add it
+		check(FName("UniqueUnicorn!!", 17, FNAME_Find) != FName());		// Now we can find it
+		check(FName("UniqueUnicorn!!", 127, FNAME_Find) == FName());	// But we can't find one with a different number
+PRAGMA_ENABLE_DEPRECATION_WARNINGS
+#endif
+
 #if !FNAME_WRITE_PROTECT_PAGES
 		// Check FNAME_Replace_Not_Safe_For_Threading updates casing
 		check(0 != UniqueName.GetPlainNameString().Compare("UNIQUEunicorn!!", ESearchCase::CaseSensitive));
@@ -2790,11 +3731,43 @@ FArchive& operator<<(FArchive& Ar, FNameEntrySerialized& E)
 	return Ar;
 }
 
+void FNameEntry::DebugDump(FOutputDevice& Out) const
+{
+#if UE_FNAME_OUTLINE_NUMBER
+	if (IsNumbered())
+	{
+		TCHAR Buffer[NAME_SIZE];
+		const FNameEntry& BaseEntry = GetNamePool().Resolve(NumberedName.Id);
+		BaseEntry.GetName(Buffer);
+		Out.Logf(TEXT("%s_%d"), Buffer, NAME_INTERNAL_TO_EXTERNAL(NumberedName.Number));
+	}
+	else
+#endif
+	{
+		TCHAR Buffer[NAME_SIZE];
+		GetName(Buffer);
+		Out.Logf(Buffer);
+	}
+}
+
 FNameEntryId FNameEntryId::FromValidEName(EName Ename)
 {
 	return GetNamePool().Find(Ename);
 }
 
+#if UE_FNAME_OUTLINE_NUMBER
+FName FName::CreateNumberedName(FNameEntryId ComparisonId, FNameEntryId DisplayId, int32 Number)
+{
+	// Users must only pass in base/plain ids, not numbered ones.
+#if WITH_CASE_PRESERVING_NAME
+	checkName(ResolveEntry(DisplayId)->ComparisonId == ComparisonId);
+#endif
+	checkName(ResolveEntry(ComparisonId)->IsNumbered() == false);
+
+	return FNameHelper::MakeWithNumber(FNameEntryIds{ComparisonId, DisplayId}, FNAME_Add, Number);
+}
+#endif // UE_FNAME_OUTLINE_NUMBER
+
 FNameEntryId FNameEntryId::FromValidENamePostInit(EName Ename)
 {
 	return GetNamePoolPostInit().Find(Ename);
@@ -2819,7 +3792,11 @@ FName FLazyName::Resolve() const
 	if (Copy.IsName())
 	{
 		FNameEntryId Id = Copy.AsName();
+#if UE_FNAME_OUTLINE_NUMBER
+		return FNameHelper::MakeWithNumber(FNameEntryIds{ Id, Id}, FNAME_Add, Number);
+#else // UE_FNAME_OUTLINE_NUMBER
 		return FName(Id, Id, Number);
+#endif // UE_FNAME_OUTLINE_NUMBER
 	}
 
 	// Resolve to FName but throw away the number part
@@ -2829,7 +3806,11 @@ FName FLazyName::Resolve() const
 	// Deliberately unsynchronized write of word-sized int, ok if multiple threads resolve same lazy name
 	Either = FLiteralOrName(Id);
 
+#if UE_FNAME_OUTLINE_NUMBER
+	return FNameHelper::MakeWithNumber(FNameEntryIds{Id, Id}, FNAME_Add, Number);
+#else // UE_FNAME_OUTLINE_NUMBER
 	return FName(Id, Id, Number);		
+#endif // UE_FNAME_OUTLINE_NUMBER
 }
 
 uint32 FLazyName::ParseNumber(const ANSICHAR* Str, int32 Len)
@@ -3966,40 +4947,105 @@ FString FScriptName::ToString() const
 	return ScriptNameToName(*this).ToString();
 }
 
-void Freeze::IntrinsicWriteMemoryImage(FMemoryImageWriter& Writer, const FName& Object, const FTypeLayoutDesc&)
+FString FMemoryImageName::ToString() const
 {
-	Writer.WriteFName(Object);
+	return FName(*this).ToString();
 }
 
-uint32 Freeze::IntrinsicAppendHash(const FName* DummyObject, const FTypeLayoutDesc& TypeDesc, const FPlatformTypeLayoutParameters& LayoutParams, FSHA1& Hasher)
+uint32 Freeze::IntrinsicAppendHash(const FMemoryImageName* DummyObject, const FTypeLayoutDesc& TypeDesc, const FPlatformTypeLayoutParameters& LayoutParams, FSHA1& Hasher)
 {
-	const uint32 SizeFromFields = LayoutParams.WithCasePreservingFName() ? sizeof(FScriptName) : sizeof(FMinimalName);
+	const uint32 SizeFromFields = LayoutParams.WithCasePreservingFName() ? sizeof(TMemoryImageNameLayout<1>) : sizeof(TMemoryImageNameLayout<0>);
 	return Freeze::AppendHashForNameAndSize(TypeDesc.Name, SizeFromFields, Hasher);
 }
 
-void Freeze::IntrinsicWriteMemoryImage(FMemoryImageWriter& Writer, const FMinimalName& Object, const FTypeLayoutDesc&)
+void Freeze::IntrinsicWriteMemoryImage(FMemoryImageWriter& Writer, const FMemoryImageName& Object, const FTypeLayoutDesc&)
 {
-	Writer.WriteFMinimalName(Object);
-}
+	const FPlatformTypeLayoutParameters& TargetLayoutParameters = Writer.GetTargetLayoutParams();
+	if (TargetLayoutParameters.IsCurrentPlatform())
+	{
+		TConstArrayView<uint8> Bytes{ reinterpret_cast<const uint8*>(&Object), sizeof(Object) };
+		Writer.WriteFMemoryImageName(Bytes, FName(Object));
+	}
+	else
+	{
+		ensureMsgf(WITH_CASE_PRESERVING_NAME, TEXT("Cooking memory images containing names when WITH_CASE_PRESERVING_NAME is false can lead to cook determinism issues with string casing."));
 
-void Freeze::IntrinsicWriteMemoryImage(FMemoryImageWriter& Writer, const FScriptName& Object, const FTypeLayoutDesc&)
-{
-	Writer.WriteFScriptName(Object);
+		if (!TargetLayoutParameters.WithCasePreservingFName())
+		{
+			TMemoryImageNameLayout<0> ToWrite;
+			ToWrite.ComparisonIndex = Object.ComparisonIndex;
+#if !UE_FNAME_OUTLINE_NUMBER
+			ToWrite.NumberOrDummy = Object.Number;
+#endif
+
+			TConstArrayView<uint8> Bytes{ reinterpret_cast<const uint8*>(&ToWrite), sizeof(ToWrite) };
+			Writer.WriteFMemoryImageName(Bytes, FName(Object));
+		}
+		else
+		{
+			TMemoryImageNameLayout<1> ToWrite;
+			ToWrite.ComparisonIndex = Object.ComparisonIndex;
+#if !UE_FNAME_OUTLINE_NUMBER
+			ToWrite.NumberOrDummy = Object.Number;
+#endif
+#if WITH_CASE_PRESERVING_NAME
+			ToWrite.DisplayIndex = Object.DisplayIndex;
+#else
+			ToWrite.DisplayIndex = Object.ComparisonIndex;
+#endif
+
+			TConstArrayView<uint8> Bytes{ reinterpret_cast<const uint8*>(&ToWrite), sizeof(ToWrite) };
+			Writer.WriteFMemoryImageName(Bytes, FName(Object));
+		}
+	}
 }
 
-uint32 Freeze::IntrinsicUnfrozenCopy(const FMemoryUnfreezeContent& Context, const FName& Object, void* OutDst)
+uint32 Freeze::IntrinsicUnfrozenCopy(const FMemoryUnfreezeContent& Context, const FMemoryImageName& Object, void* OutDst)
 {
-	if (Context.FrozenLayoutParameters.WithEditorOnly())
+	const FPlatformTypeLayoutParameters& TargetLayoutParameters = Context.FrozenLayoutParameters;
+	if (TargetLayoutParameters.IsCurrentPlatform())
 	{
-		new(OutDst) FName(ScriptNameToName((FScriptName&)Object));
-		return sizeof(FScriptName);
+		new(OutDst) FMemoryImageName(Object);
+		return sizeof(FMemoryImageName);
 	}
 	else
 	{
-		new(OutDst) FName(MinimalNameToName((FMinimalName&)Object));
-		return sizeof(FMinimalName);
+		if (!TargetLayoutParameters.WithCasePreservingFName())
+		{
+			ensureMsgf(!WITH_CASE_PRESERVING_NAME, TEXT("Unfreezing memory images containing non-case preserving names when WITH_CASE_PRESERVING_NAME is true can lead to casing cook determinism."));
+
+			const TMemoryImageNameLayout<0>& Loaded = reinterpret_cast<const TMemoryImageNameLayout<0>&>(Object);
+			FMemoryImageName* Dest = new(OutDst) FMemoryImageName();
+			Dest->ComparisonIndex = Loaded.ComparisonIndex;
+#if !UE_FNAME_OUTLINE_NUMBER
+			Dest->Number = Loaded.NumberOrDummy;
+#endif
+#if WITH_CASE_PRESERVING_NAME
+			Dest->DisplayIndex = Loaded.ComparisonIndex;
+#endif
+
+			return sizeof(Loaded);
+		}
+		else
+		{
+			const TMemoryImageNameLayout<1>& Loaded = reinterpret_cast<const TMemoryImageNameLayout<1>&>(Object);
+			FMemoryImageName* Dest = new(OutDst) FMemoryImageName();
+			Dest->ComparisonIndex = Loaded.ComparisonIndex;
+#if !UE_FNAME_OUTLINE_NUMBER
+			Dest->Number = Loaded.NumberOrDummy;
+#endif
+#if WITH_CASE_PRESERVING_NAME
+			Dest->DisplayIndex = Loaded.DisplayIndex;
+#endif
+
+			return sizeof(Loaded);
+		}
 	}
 }
+void Freeze::IntrinsicWriteMemoryImage(FMemoryImageWriter& Writer, const FScriptName& Object, const FTypeLayoutDesc&)
+{
+	Writer.WriteFScriptName(Object);
+}
 
 bool ShouldReplicateAsInteger(EName Ename, const FName& Name)
 {
@@ -4007,3 +5053,143 @@ bool ShouldReplicateAsInteger(EName Ename, const FName& Name)
 }
 
 PRAGMA_RESTORE_UNSAFE_TYPECAST_WARNINGS
+
+// Console command implementations
+namespace UE::Name::Private
+{
+	// Used instead of FOutputDeviceFile because that output device is disabled by some preprocessor definitions.
+	struct FOutputDeviceWrapper : public FOutputDevice
+	{
+		FOutputDeviceWrapper(FArchive& InArchive)
+			: Ar(InArchive)
+		{
+			bSuppressEventTag = true;
+		}
+
+		virtual void Serialize(const TCHAR* Data, ELogVerbosity::Type Verbosity, const FName& Category)
+		{
+			FOutputDeviceHelper::FormatCastAndSerializeLine(Ar, Data, Verbosity, Category, FPlatformTime::Seconds(), bSuppressEventTag, bAutoEmitLineTerminator);
+		}
+
+	private:
+		FArchive& Ar;
+	};
+
+	void DebugDumpInternal(bool bNumbered, int32 Num, FOutputDevice& Out)
+	{
+		const TArray<const FNameEntry*> Entries = FName::DebugDump();
+		TConstArrayView<const FNameEntry*> IterEntries = Entries;
+
+		if (Num < MAX_int32)
+		{
+			IterEntries.RightInline(Num);
+		}
+
+		for (const FNameEntry* Entry : IterEntries)
+		{
+#if UE_FNAME_OUTLINE_NUMBER
+			if (Entry->IsNumbered() == bNumbered)
+#endif
+			{
+				Entry->DebugDump(Out);
+			}
+		}
+	}
+
+	// List FNames to the output device 
+	void ListFNames(const TArray<FString>& Args, UWorld*, FOutputDevice& Out)
+	{
+		int32 Num = MAX_int32;
+		for (const FString& Arg : Args)
+		{
+			FParse::Value(*Arg, TEXT("num="), Num);
+		}
+
+		DebugDumpInternal(false, Num, Out);
+	}
+
+	// Dump FNames to a file
+	void DumpFNames(const TArray<FString>& Args, UWorld*, FOutputDevice& Out)
+	{
+		int32 Num = MAX_int32;
+		for (const FString& Arg : Args)
+		{
+			FParse::Value(*Arg, TEXT("num="), Num);
+		}
+
+		FString Filename = FPaths::ProfilingDir() + FString::Printf(TEXT("FNames_(%s).txt"), *FDateTime::Now().ToString(TEXT("%Y%m%d_%H%M%S")));
+		TUniquePtr<FArchive> FileAr{IFileManager::Get().CreateFileWriter(*Filename)};
+		if (!FileAr)
+		{
+			Out.Logf(ELogVerbosity::Error, TEXT("Failed to open outfile file to dump FNames."));
+			return;
+		}
+
+		FOutputDeviceWrapper OutDevice(*FileAr);
+		DebugDumpInternal(false, Num, OutDevice);
+		Out.Logf(TEXT("Wrote FNames to %s"), *Filename);
+	}
+
+	// List FNames to the output device 
+	void ListNumberedFNames(const TArray<FString>& Args, UWorld*, FOutputDevice& Out)
+	{
+#if UE_FNAME_OUTLINE_NUMBER
+		int32 Num = MAX_int32;
+		for (const FString& Arg : Args)
+		{
+			FParse::Value(*Arg, TEXT("num="), Num);
+		}
+
+		DebugDumpInternal(true, Num, Out);
+#else
+		Out.Logf(ELogVerbosity::Error, TEXT("Cannot list numbered FNames as UE_FNAME_OUTLINE_NUMBER is not set so numbers are stored in FName instances rather than the name table."));
+#endif
+	}
+
+	// Dump FNames to a file
+	void DumpNumberedFNames(const TArray<FString>& Args, UWorld*, FOutputDevice& Out)
+	{
+#if UE_FNAME_OUTLINE_NUMBER
+		int32 Num = MAX_int32;
+		for (const FString& Arg : Args)
+		{
+			FParse::Value(*Arg, TEXT("num="), Num);
+		}
+
+		FString Filename = FPaths::ProfilingDir() + FString::Printf(TEXT("FNamesNumbered_(%s).txt"), *FDateTime::Now().ToString(TEXT("%Y%m%d_%H%M%S")));
+		TUniquePtr<FArchive> FileAr{ IFileManager::Get().CreateFileWriter(*Filename) };
+		if (!FileAr)
+		{
+			Out.Logf(ELogVerbosity::Error, TEXT("Failed to open outfile file to dump FNames."));
+			return;
+		}
+
+		FOutputDeviceWrapper OutDevice(*FileAr);
+		DebugDumpInternal(true, Num, OutDevice);
+		Out.Logf(TEXT("Wrote FNames to %s"), *Filename);
+#else
+		Out.Logf(ELogVerbosity::Error, TEXT("Cannot dump numbered FNames as UE_FNAME_OUTLINE_NUMBER is not set so numbers are stored in FName instances rather than the name table."));
+#endif
+	}
+
+	// Write FName stats to the output device 
+	void GetFNameStats(FOutputDevice& Out)
+	{
+		FName::DisplayHash(Out);
+	}
+
+	void DumpHashCsv(FOutputDevice& Out)
+	{
+		FString Filename = FPaths::ProfilingDir() + FString::Printf(TEXT("FNameHash_(%s).csv"), *FDateTime::Now().ToString(TEXT("%Y%m%d_%H%M%S")));
+		TUniquePtr<FArchive> FileAr{ IFileManager::Get().CreateFileWriter(*Filename) };
+		if (!FileAr)
+		{
+			Out.Logf(ELogVerbosity::Error, TEXT("Failed to open outfile file to dump FName hash stats."));
+			return;
+		}
+
+		FOutputDeviceWrapper OutDevice(*FileAr);
+		GetNamePoolPostInit().LogHashStats(OutDevice);
+		Out.Logf(TEXT("Wrote FName hash stats to %s"), *Filename);
+	}
+}
\ No newline at end of file
diff --git a/Engine/Source/Runtime/Core/Public/CoreGlobals.h b/Engine/Source/Runtime/Core/Public/CoreGlobals.h
--- a/Engine/Source/Runtime/Core/Public/CoreGlobals.h
+++ b/Engine/Source/Runtime/Core/Public/CoreGlobals.h
@@ -734,8 +734,11 @@ class FChunkedFixedUObjectArray;
 extern CORE_API FChunkedFixedUObjectArray* GCoreObjectArrayForDebugVisualizers;
 
 /** Array to help visualize object paths in the debugger */
-extern template class TArray<FMinimalName, TInlineAllocator<3>>;
-extern CORE_API TArray<FMinimalName, TInlineAllocator<3>>* GCoreComplexObjectPathDebug;
+namespace UE::ObjectPath::Private
+{
+	struct FStoredObjectPath;
+}
+extern CORE_API UE::ObjectPath::Private::FStoredObjectPath* GCoreComplexObjectPathDebug;
 
 /** Array to help visualize object handles in the debugger */
 struct FObjectHandlePackageDebugData;
diff --git a/Engine/Source/Runtime/Core/Public/Modules/Boilerplate/ModuleBoilerplate.h b/Engine/Source/Runtime/Core/Public/Modules/Boilerplate/ModuleBoilerplate.h
--- a/Engine/Source/Runtime/Core/Public/Modules/Boilerplate/ModuleBoilerplate.h
+++ b/Engine/Source/Runtime/Core/Public/Modules/Boilerplate/ModuleBoilerplate.h
@@ -70,7 +70,7 @@ class FChunkedFixedUObjectArray;
 	#define UE4_VISUALIZERS_HELPERS \
 		uint8** GNameBlocksDebug = FNameDebugVisualizer::GetBlocks(); \
 		FChunkedFixedUObjectArray*& GObjectArrayForDebugVisualizers = GCoreObjectArrayForDebugVisualizers; \
-		TArray<FMinimalName, TInlineAllocator<3>>*& GComplexObjectPathDebug = GCoreComplexObjectPathDebug; \
+		UE::ObjectPath::Private::FStoredObjectPath*& GComplexObjectPathDebug = GCoreComplexObjectPathDebug; \
 		FObjectHandlePackageDebugData*& GObjectHandlePackageDebug = GCoreObjectHandlePackageDebug;
 #endif
 
diff --git a/Engine/Source/Runtime/Core/Public/Serialization/MappedName.h b/Engine/Source/Runtime/Core/Public/Serialization/MappedName.h
--- a/Engine/Source/Runtime/Core/Public/Serialization/MappedName.h
+++ b/Engine/Source/Runtime/Core/Public/Serialization/MappedName.h
@@ -139,7 +139,7 @@ public:
 
 		FNameEntryId DisplayId = NameEntries[MappedName.GetIndex()];
 		FNameEntryId ComparisonId = FName::GetComparisonIdFromDisplayId(DisplayId);
-		return FMinimalName(ComparisonId, MappedName.GetNumber());
+		return NameToMinimalName(FName(ComparisonId, DisplayId, MappedName.GetNumber()));
 	}
 
 private:
diff --git a/Engine/Source/Runtime/Core/Public/Serialization/MemoryImage.h b/Engine/Source/Runtime/Core/Public/Serialization/MemoryImage.h
--- a/Engine/Source/Runtime/Core/Public/Serialization/MemoryImage.h
+++ b/Engine/Source/Runtime/Core/Public/Serialization/MemoryImage.h
@@ -184,7 +184,7 @@ struct FMemoryImageResult
 	FPlatformTypeLayoutParameters TargetLayoutParameters;
 	TArray<FMemoryImageVTablePointer> VTables;
 	TArray<FMemoryImageNamePointer> ScriptNames;
-	TArray<FMemoryImageNamePointer> MinimalNames;
+	TArray<FMemoryImageNamePointer> MemoryImageNames;
 
 	CORE_API void SaveToArchive(FArchive& Ar) const;
 	CORE_API void ApplyPatches(void* FrozenObject) const;
@@ -237,8 +237,7 @@ public:
 	FMemoryImageSection* WritePointer(const FTypeLayoutDesc& StaticTypeDesc, const FTypeLayoutDesc& DerivedTypeDesc, uint32* OutOffsetToBase = nullptr);
 	uint32 WriteRawPointerSizedBytes(uint64 PointerValue);
 	uint32 WriteVTable(const FTypeLayoutDesc& TypeDesc, const FTypeLayoutDesc& DerivedTypeDesc);
-	uint32 WriteFName(const FName& Name);
-	uint32 WriteFMinimalName(const FMinimalName& Name);
+	uint32 WriteFMemoryImageName(TConstArrayView<uint8> InBytes, const FName& Name);
 	uint32 WriteFScriptName(const FScriptName& Name);
 	uint32 Flatten(FMemoryImageResult& OutResult) const;
 
@@ -249,7 +248,7 @@ public:
 	TArray<FSectionPointer> Pointers;
 	TArray<FMemoryImageVTablePointer> VTables;
 	TArray<FMemoryImageNamePointer> ScriptNames;
-	TArray<FMemoryImageNamePointer> MinimalNames;
+	TArray<FMemoryImageNamePointer> MemoryImageNames;
 	FSHAHash Hash;
 	uint32 MaxAlignment;
 };
diff --git a/Engine/Source/Runtime/Core/Public/Serialization/MemoryImageWriter.h b/Engine/Source/Runtime/Core/Public/Serialization/MemoryImageWriter.h
--- a/Engine/Source/Runtime/Core/Public/Serialization/MemoryImageWriter.h
+++ b/Engine/Source/Runtime/Core/Public/Serialization/MemoryImageWriter.h
@@ -5,7 +5,6 @@
 #include "CoreTypes.h"
 #include "Serialization/MemoryLayout.h"
 
-struct FMinimalName;
 struct FScriptName;
 class FMemoryImage;
 class FMemoryImageSection;
@@ -42,8 +41,7 @@ public:
 	uint32 WriteNullPointer();
 	uint32 WriteRawPointerSizedBytes(uint64 PointerValue);
 	uint32 WriteVTable(const FTypeLayoutDesc& TypeDesc, const FTypeLayoutDesc& DerivedTypeDesc);
-	uint32 WriteFName(const FName& Name);
-	uint32 WriteFMinimalName(const FMinimalName& Name);
+	uint32 WriteFMemoryImageName(TConstArrayView<uint8> InBytes, const FName& Name);
 	uint32 WriteFScriptName(const FScriptName& Name);
 
 	template<typename T>
diff --git a/Engine/Source/Runtime/Core/Public/Stats/Stats2.h b/Engine/Source/Runtime/Core/Public/Stats/Stats2.h
--- a/Engine/Source/Runtime/Core/Public/Stats/Stats2.h
+++ b/Engine/Source/Runtime/Core/Public/Stats/Stats2.h
@@ -413,9 +413,11 @@ FORCEINLINE uint32 FromPackedCallCountDuration_Duration(int64 Both)
 class FStatNameAndInfo
 {
 	/**
-	 * An FName, but the high bits of the Number are used for other fields.
+	 * Store name and number separately in case UE_FNAME_OUTLINE_NUMBER is set, so the high bits of the Number are used for other fields.
 	 */
-	FMinimalName NameAndInfo;
+	FNameEntryId Index;
+	int32 Number;
+
 public:
 	FORCEINLINE_STATS FStatNameAndInfo()
 	{
@@ -425,15 +427,14 @@ public:
 	 * Build from a raw FName
 	 */
 	FORCEINLINE_STATS FStatNameAndInfo(FName Other, bool bAlreadyHasMeta)
-		: NameAndInfo(NameToMinimalName(Other))
+		: Index(Other.GetComparisonIndex())
+		, Number(Other.GetNumber())
 	{
 		if (!bAlreadyHasMeta)
 		{
-			int32 Number = NameAndInfo.Number;
 			// ok, you can't have numbered stat FNames too large
 			checkStats(!(Number >> EStatAllFields::StartShift));
 			Number |= EStatMetaFlags::DummyAlwaysOne << (EStatMetaFlags::Shift + EStatAllFields::StartShift);
-			NameAndInfo.Number = Number;
 		}
 		CheckInvariants();
 	}
@@ -442,13 +443,14 @@ public:
 	 * Build with stat metadata
 	 */
 	FORCEINLINE_STATS FStatNameAndInfo(FName InStatName, char const* InGroup, char const* InCategory, TCHAR const* InDescription, EStatDataType::Type InStatType, bool bShouldClearEveryFrame, bool bCycleStat, bool bSortByName, FPlatformMemory::EMemoryCounterRegion MemoryRegion = FPlatformMemory::MCR_Invalid)
-		: NameAndInfo(NameToMinimalName(ToLongName(InStatName, InGroup, InCategory, InDescription, bSortByName)))
 	{
-		int32 Number = NameAndInfo.Number;
+		FName LongName = ToLongName(InStatName, InGroup, InCategory, InDescription, bSortByName);
+		Index = LongName.GetComparisonIndex();
+		Number = LongName.GetNumber();
+
 		// ok, you can't have numbered stat FNames too large
 		checkStats(!(Number >> EStatAllFields::StartShift));
 		Number |= (EStatMetaFlags::DummyAlwaysOne | EStatMetaFlags::HasLongNameAndMetaInfo) << (EStatMetaFlags::Shift + EStatAllFields::StartShift);
-		NameAndInfo.Number = Number;
 
 		SetField<EStatDataType>(InStatType);
 		SetFlag(EStatMetaFlags::ShouldClearEveryFrame, bShouldClearEveryFrame);
@@ -465,9 +467,9 @@ public:
 	/**
 	 * Internal use, used by the deserializer
 	 */
-	FORCEINLINE_STATS void SetNumberDirect(int32 Number)
+	FORCEINLINE_STATS void SetNumberDirect(int32 InNumber)
 	{
-		NameAndInfo.Number = Number;
+		Number = InNumber;
 	}
 
 	/**
@@ -476,7 +478,7 @@ public:
 	FORCEINLINE_STATS int32 GetRawNumber() const
 	{
 		CheckInvariants();
-		return NameAndInfo.Number;
+		return Number;
 	}
 
 	/**
@@ -487,10 +489,10 @@ public:
 		// ok, you can't have numbered stat FNames too large
 		checkStats(!(RawName.GetNumber() >> EStatAllFields::StartShift));
 		CheckInvariants();
-		int32 Number = NameAndInfo.Number;
-		Number &= ~((1 << EStatAllFields::StartShift) - 1);
-		NameAndInfo = NameToMinimalName(RawName);
-		NameAndInfo.Number = (Number | RawName.GetNumber());
+		int32 LocalNumber = Number;
+		LocalNumber &= ~((1 << EStatAllFields::StartShift) - 1);
+		Index = RawName.GetComparisonIndex();
+		Number = (LocalNumber | RawName.GetNumber());
 	}
 
 	/**
@@ -500,11 +502,7 @@ public:
 	FORCEINLINE_STATS FName GetRawName() const
 	{
 		CheckInvariants();
-		FMinimalName Result(NameAndInfo);
-		int32 Number = NameAndInfo.Number;
-		Number &= ((1 << EStatAllFields::StartShift) - 1);
-		Result.Number = Number;
-		return MinimalNameToName(Result);
+		return FName(Index, Index, Number & ((1 << EStatAllFields::StartShift) - 1));
 	}
 
 	/**
@@ -514,7 +512,7 @@ public:
 	FORCEINLINE_STATS FName GetEncodedName() const
 	{
 		CheckInvariants();
-		return MinimalNameToName(NameAndInfo);
+		return FName(Index, Index, Number);
 	}
 
 	/**
@@ -567,8 +565,8 @@ public:
 	 */
 	FORCEINLINE_STATS void CheckInvariants() const
 	{
-		checkStats((NameAndInfo.Number & (EStatMetaFlags::DummyAlwaysOne << (EStatAllFields::StartShift + EStatMetaFlags::Shift)))
-			&& NameAndInfo.Index);
+		checkStats((Number & (EStatMetaFlags::DummyAlwaysOne << (EStatAllFields::StartShift + EStatMetaFlags::Shift)))
+			&& Index);
 	}
 
 	/**
@@ -579,10 +577,10 @@ public:
 	typename TField::Type GetField() const
 	{
 		CheckInvariants();
-		int32 Number = NameAndInfo.Number;
-		Number = (Number >> (EStatAllFields::StartShift + TField::Shift)) & TField::Mask;
-		checkStats(Number != TField::Invalid && Number < TField::Num);
-		return typename TField::Type(Number);
+		int32 LocalNumber = Number;
+		LocalNumber = (LocalNumber >> (EStatAllFields::StartShift + TField::Shift)) & TField::Mask;
+		checkStats(LocalNumber != TField::Invalid && LocalNumber < TField::Num);
+		return typename TField::Type(LocalNumber);
 	}
 
 	/**
@@ -592,12 +590,12 @@ public:
 	template<typename TField>
 	void SetField(typename TField::Type Value)
 	{
-		int32 Number = NameAndInfo.Number;
+		int32 LocalNumber = Number;
 		CheckInvariants();
 		checkStats(Value < TField::Num && Value != TField::Invalid);
-		Number &= ~(TField::Mask << (EStatAllFields::StartShift + TField::Shift));
-		Number |= Value << (EStatAllFields::StartShift + TField::Shift);
-		NameAndInfo.Number = Number;
+		LocalNumber &= ~(TField::Mask << (EStatAllFields::StartShift + TField::Shift));
+		LocalNumber |= Value << (EStatAllFields::StartShift + TField::Shift);
+		Number = LocalNumber;
 		CheckInvariants();
 	}
 
@@ -607,10 +605,10 @@ public:
 	 */
 	bool GetFlag(EStatMetaFlags::Type Bit) const
 	{
-		int32 Number = NameAndInfo.Number;
+		int32 LocalNumber = Number;
 		CheckInvariants();
 		checkStats(Bit < EStatMetaFlags::Num && Bit != EStatMetaFlags::Invalid);
-		return !!((Number >> (EStatAllFields::StartShift + EStatMetaFlags::Shift)) & Bit);
+		return !!((LocalNumber >> (EStatAllFields::StartShift + EStatMetaFlags::Shift)) & Bit);
 	}
 
 	/**
@@ -620,18 +618,18 @@ public:
 	 */
 	void SetFlag(EStatMetaFlags::Type Bit, bool Value)
 	{
-		int32 Number = NameAndInfo.Number;
+		int32 LocalNumber = Number;
 		CheckInvariants();
 		checkStats(Bit < EStatMetaFlags::Num && Bit != EStatMetaFlags::Invalid);
 		if (Value)
 		{
-			Number |= (Bit << (EStatAllFields::StartShift + EStatMetaFlags::Shift));
+			LocalNumber |= (Bit << (EStatAllFields::StartShift + EStatMetaFlags::Shift));
 		}
 		else
 		{
-			Number &= ~(Bit << (EStatAllFields::StartShift + EStatMetaFlags::Shift));
+			LocalNumber &= ~(Bit << (EStatAllFields::StartShift + EStatMetaFlags::Shift));
 		}
-		NameAndInfo.Number = Number;
+		Number = LocalNumber;
 		CheckInvariants();
 	}
 
diff --git a/Engine/Source/Runtime/Core/Public/UObject/NameTypes.h b/Engine/Source/Runtime/Core/Public/UObject/NameTypes.h
--- a/Engine/Source/Runtime/Core/Public/UObject/NameTypes.h
+++ b/Engine/Source/Runtime/Core/Public/UObject/NameTypes.h
@@ -30,17 +30,33 @@
 	#define WITH_CASE_PRESERVING_NAME WITH_EDITORONLY_DATA
 #endif
 
+// Should the number part of the fname be stored in the name table or in the FName instance?
+// Storing it in the name table may save memory overall but new number suffixes will cause the name table to grow like unique strings do.
+#ifndef UE_FNAME_OUTLINE_NUMBER
+	#define UE_FNAME_OUTLINE_NUMBER 0
+#endif // UE_FNAME_OUTLINE_NUMBER
+
 class FText;
 
 /** Maximum size of name, including the null terminator. */
 enum {NAME_SIZE	= 1024};
 
+struct FMinimalName;
+struct FScriptName;
+class CORE_API FName;
+
 /** Opaque id to a deduplicated name */
 struct FNameEntryId
 {
+	// Default initialize to be equal to NAME_None
 	FNameEntryId() : Value(0) {}
 	FNameEntryId(ENoInit) {}
 
+	bool IsNone() const
+	{
+		return Value == 0;
+	}
+
 	/** Slow alphabetical order that is stable / deterministic over process runs */
 	CORE_API int32 CompareLexical(FNameEntryId Rhs) const;
 	bool LexicalLess(FNameEntryId Rhs) const { return CompareLexical(Rhs) < 0; }
@@ -57,6 +73,7 @@ struct FNameEntryId
 	bool operator==(FNameEntryId Rhs) const { return Value == Rhs.Value; }
 	bool operator!=(FNameEntryId Rhs) const { return Value != Rhs.Value; }
 
+	// Returns true if this FNameEntryId is not equivalent to NAME_None
 	explicit operator bool() const { return Value != 0; }
 
 	UE_DEPRECATED(4.23, "NAME_INDEX is replaced by FNameEntryId, which is no longer a contiguous integer. "
@@ -150,7 +167,11 @@ enum ELinkerNameTableConstructor    {ENAME_LinkerConstructor};
 /** Enumeration for finding name. */
 enum EFindName
 {
-	/** Find a name; return 0 if it doesn't exist. */
+	/** 
+	* Find a name; return 0/NAME_None/FName() if it doesn't exist.
+	* When UE_FNAME_OUTLINE_NUMBER is set, we search for the exact name including the number suffix.
+	* Otherwise we search only for the string part.
+	*/
 	FNAME_Find,
 
 	/** Find a name or add it if it doesn't exist. */
@@ -193,6 +214,16 @@ private:
 	{
 		ANSICHAR	AnsiName[NAME_SIZE];
 		WIDECHAR	WideName[NAME_SIZE];
+#if UE_FNAME_OUTLINE_NUMBER
+		// These fields are valid when Header.Len == 0.
+		// Stores a (string, number) fname pair to construct the string "string_(number-1)" on demand.
+		// Id is a reference to another entry with Header.Len != 0 (no recursion).
+		struct 
+		{
+			FNameEntryId	Id;
+			uint32			Number;
+		} NumberedName;
+#endif // UE_FNAME_OUTLINE_NUMBER
 	};
 
 	FNameEntry(const FNameEntry&) = delete;
@@ -212,6 +243,15 @@ public:
 		return Header.Len;
 	}
 
+	FORCEINLINE bool IsNumbered() const 
+	{
+#if UE_FNAME_OUTLINE_NUMBER
+		return Header.Len == 0;
+#else
+		return false;
+#endif
+	}
+
 	/**
 	 * Copy unterminated name to TCHAR buffer without allocating.
 	 *
@@ -243,6 +283,7 @@ public:
 	/** Appends name to string with path separator using FString::PathAppend(). */
 	CORE_API void AppendNameToPathString(FString& OutString) const;
 
+	CORE_API void DebugDump(FOutputDevice& Out) const;
 
 	/**
 	 * Returns the size in bytes for FNameEntry structure. This is != sizeof(FNameEntry) as we only allocated as needed.
@@ -278,6 +319,11 @@ private:
 	void CopyAndConvertUnterminatedName(TCHAR* OutName) const;
 	const ANSICHAR* GetUnterminatedName(ANSICHAR(&OptionalDecodeBuffer)[NAME_SIZE]) const;
 	const WIDECHAR* GetUnterminatedName(WIDECHAR(&OptionalDecodeBuffer)[NAME_SIZE]) const;
+
+#if UE_FNAME_OUTLINE_NUMBER
+	static int32 GetNumberedEntrySize();
+	uint32 GetNumber() const;
+#endif // UE_FNAME_OUTLINE_NUMBER
 };
 
 /**
@@ -333,10 +379,16 @@ struct FNameEntrySerialized
 
 /**
  * The minimum amount of data required to reconstruct a name
- * This is smaller than FName, but you lose the case-preserving behavior
+ * This is smaller than FName when WITH_CASE_PRESERVING_NAME is set, but you lose the case-preserving behavior.
+ * The size of this type is not portable across different platforms and configurations, as with FName itself.
  */
 struct FMinimalName
 {
+	friend uint32 GetTypeHash(FMinimalName Name);
+	friend bool operator==(FMinimalName Lhs, FMinimalName Rhs);
+	friend bool operator==(FName Lhs, FMinimalName Rhs);
+	friend FName;
+
 	FMinimalName() {}
 	
 	FMinimalName(EName N)
@@ -344,39 +396,33 @@ struct FMinimalName
 	{
 	}
 
-	FMinimalName(FNameEntryId InIndex, int32 InNumber)
-		: Index(InIndex)
-		, Number(InNumber)
-	{
-	}
-
-	FORCEINLINE bool IsNone() const
-	{
-		return !Index && Number == NAME_NO_NUMBER_INTERNAL;
-	}
-
-	friend FORCEINLINE FArchive& operator<<(FArchive& Ar, FMinimalName& E)
-	{
-		Ar << E.Index;
-		Ar << E.Number;
-		return Ar;
-	}
-
-	FORCEINLINE bool operator<(FMinimalName Rhs) const { return Index == Rhs.Index ? Number < Rhs.Number : Index < Rhs.Index; };
-
+	FORCEINLINE explicit FMinimalName(const FName& Name);
+	FORCEINLINE bool IsNone() const;
+	FORCEINLINE bool operator<(FMinimalName Rhs) const;
+	
+private:
 	/** Index into the Names array (used to find String portion of the string/number pair) */
 	FNameEntryId	Index;
+#if !UE_FNAME_OUTLINE_NUMBER
 	/** Number portion of the string/number pair (stored internally as 1 more than actual, so zero'd memory will be the default, no-instance case) */
 	int32			Number = NAME_NO_NUMBER_INTERNAL;
+#endif // UE_FNAME_OUTLINE_NUMBE
 };
 
 /**
  * The full amount of data required to reconstruct a case-preserving name
- * This will be the same size as FName when WITH_CASE_PRESERVING_NAME is 1, and is used to store an FName in cases where 
- * the size of FName must be constant between build configurations (eg, blueprint bytecode)
+ * This will be the maximum size of an FName across all values of WITH_CASE_PRESERVING_NAME and UE_FNAME_OUTLINE_NUMBER
+ * and is used to store an FName in cases where  the size of a name must be constant between build configurations (eg, blueprint bytecode)
+ * the layout is not guaranteed to be the same as FName even if the size is the same, so memory cannot be reinterpreted between the two.
+ * The layout here must be as expected by FScriptBytecodeWriter and XFER_NAME
  */
 struct FScriptName
 {
+	friend uint32 GetTypeHash(FScriptName Name);
+	friend bool operator==(FScriptName Lhs, FScriptName Rhs);
+	friend bool operator==(FName Lhs, FScriptName Rhs);
+	friend FName;
+
 	FScriptName() {}
 	
 	FScriptName(EName Ename)
@@ -385,26 +431,71 @@ struct FScriptName
 	{
 	}
 
-	FScriptName(FNameEntryId InComparisonIndex, FNameEntryId InDisplayIndex, int32 InNumber)
-		: ComparisonIndex(InComparisonIndex)
-		, DisplayIndex(InDisplayIndex)
-		, Number(InNumber)
-	{
-	}
-
-	FORCEINLINE bool IsNone() const
-	{
-		return !ComparisonIndex && Number == NAME_NO_NUMBER_INTERNAL;
-	}
+	FORCEINLINE explicit FScriptName(const FName& Name);
+	FORCEINLINE bool IsNone() const;
+	inline bool operator==(EName Name) { return *this == FScriptName(Name); }
 
 	CORE_API FString ToString() const;
 
-	/** Index into the Names array (used to find String portion of the string/number pair used for comparison) */
+	// The internal structure of FScriptName is private in order to handle UE_FNAME_OUTLINE_NUMBER
+private: 
+	/** Encoded address of name entry  (used to find String portion of the string/number pair used for comparison) */
 	FNameEntryId	ComparisonIndex;
-	/** Index into the Names array (used to find String portion of the string/number pair used for display) */
+	/** Encoded address of name entry  (used to find String portion of the string/number pair used for display) */
 	FNameEntryId	DisplayIndex;
+#if UE_FNAME_OUTLINE_NUMBER
+	uint32			Dummy = 0; // Dummy to keep the size the same regardless of build configuration, but change the name so trying to use Number is a compile error
+#else // UE_FNAME_OUTLINE_NUMBER
+	/** Number portion of the string/number pair (stored internally as 1 more than actual, so zero'd memory will be the default, no-instance case) */
+	uint32			Number = NAME_NO_NUMBER_INTERNAL;
+#endif // UE_FNAME_OUTLINE_NUMBER
+};
+
+struct FMemoryImageName;
+
+namespace Freeze
+{
+	CORE_API void IntrinsicWriteMemoryImage(FMemoryImageWriter& Writer, const FMemoryImageName& Object, const FTypeLayoutDesc&);
+	CORE_API uint32 IntrinsicUnfrozenCopy(const FMemoryUnfreezeContent& Context, const FMemoryImageName& Object, void* OutDst);
+}
+
+/**
+ * Predictably sized structure for representing an FName in memory images while allowing the size to be smaller than FScriptName
+ * when case-preserving behavior is not required.
+ */
+struct FMemoryImageName
+{
+	friend uint32 GetTypeHash(FMemoryImageName Name);
+	friend bool operator==(FMemoryImageName Lhs, FMemoryImageName Rhs);
+	friend bool operator==(FName Lhs, FMemoryImageName Rhs);
+	friend FName;
+
+	friend CORE_API void Freeze::IntrinsicWriteMemoryImage(FMemoryImageWriter& Writer, const FMemoryImageName& Object, const FTypeLayoutDesc&);
+	friend CORE_API uint32 Freeze::IntrinsicUnfrozenCopy(const FMemoryUnfreezeContent& Context, const FMemoryImageName& Object, void* OutDst);
+
+	FMemoryImageName();
+	FMemoryImageName(EName Name);
+	FORCEINLINE FMemoryImageName(const FName& Name);
+
+	inline bool operator==(EName Name) const { return *this == FMemoryImageName(Name); }
+
+	FORCEINLINE bool IsNone() const;
+	CORE_API FString ToString() const;
+
+	// The internal structure of FMemoryImageName is private in order to handle UE_FNAME_OUTLINE_NUMBER
+private:
+	/** Encoded address of name entry (used to find String portion of the string/number pair used for comparison) */
+	FNameEntryId	ComparisonIndex;
+#if UE_FNAME_OUTLINE_NUMBER
+	uint32			Dummy = 0; // Dummy to keep the size the same regardless of build configuration, but change the name so trying to use Number is a compile error
+#else // UE_FNAME_OUTLINE_NUMBER
 	/** Number portion of the string/number pair (stored internally as 1 more than actual, so zero'd memory will be the default, no-instance case) */
 	uint32			Number = NAME_NO_NUMBER_INTERNAL;
+#endif // UE_FNAME_OUTLINE_NUMBER
+#if WITH_CASE_PRESERVING_NAME
+	/** Encoded address of name entry (used to find String portion of the string/number pair used for display) */
+	FNameEntryId	DisplayIndex;
+#endif
 };
 
 /**
@@ -415,6 +506,10 @@ struct FScriptName
 class CORE_API FName
 {
 public:
+#if UE_FNAME_OUTLINE_NUMBER
+	FNameEntryId GetComparisonIndex() const;
+	FNameEntryId GetDisplayIndex() const;
+#else // UE_FNAME_OUTLINE_NUMBER
 	FORCEINLINE FNameEntryId GetComparisonIndex() const
 	{
 		checkName(IsWithinBounds(ComparisonIndex));
@@ -427,7 +522,12 @@ public:
 		checkName(IsWithinBounds(Index));
 		return Index;
 	}
+#endif //UE_FNAME_OUTLINE_NUMBER
 
+#if UE_FNAME_OUTLINE_NUMBER
+	int32 GetNumber() const;
+	void SetNumber(const int32 NewNumber);
+#else //UE_FNAME_OUTLINE_NUMBER
 	FORCEINLINE int32 GetNumber() const
 	{
 		return Number;
@@ -437,6 +537,7 @@ public:
 	{
 		Number = NewNumber;
 	}
+#endif //UE_FNAME_OUTLINE_NUMBER
 	
 	/** Get name without number part as a dynamically allocated string */
 	FString GetPlainNameString() const;
@@ -525,19 +626,11 @@ public:
 	/**
 	 * Check to see if this FName matches the other FName, potentially also checking for any case variations
 	 */
-	FORCEINLINE bool IsEqual(const FName& Other, const ENameCase CompareMethod = ENameCase::IgnoreCase, const bool bCompareNumber = true ) const
-	{
-		return ((CompareMethod == ENameCase::IgnoreCase) ? ComparisonIndex == Other.ComparisonIndex : GetDisplayIndexFast() == Other.GetDisplayIndexFast())
-			&& (!bCompareNumber || GetNumber() == Other.GetNumber());
-	}
+	FORCEINLINE bool IsEqual(const FName& Other, const ENameCase CompareMethod = ENameCase::IgnoreCase, const bool bCompareNumber = true) const;
 
 	FORCEINLINE bool operator==(FName Other) const
 	{
-#if PLATFORM_64BITS && !WITH_CASE_PRESERVING_NAME
-		return ToComparableInt() == Other.ToComparableInt();
-#else
-		return (ComparisonIndex == Other.ComparisonIndex) & (GetNumber() == Other.GetNumber());
-#endif
+		return ToUnstableInt() == Other.ToUnstableInt();
 	}
 
 	FORCEINLINE bool operator!=(FName Other) const
@@ -545,15 +638,9 @@ public:
 		return !(*this == Other);
 	}
 
-	FORCEINLINE bool operator==(EName Ename) const
-	{
-		return (ComparisonIndex == Ename) & (GetNumber() == 0);
-	}
+	FORCEINLINE bool operator==(EName Ename) const;
 	
-	FORCEINLINE bool operator!=(EName Ename) const
-	{
-		return (ComparisonIndex != Ename) | (GetNumber() != 0);
-	}
+	FORCEINLINE bool operator!=(EName Ename) const;
 
 	UE_DEPRECATED(4.23, "Please use FastLess() / FNameFastLess or LexicalLess() / FNameLexicalLess instead. "
 		"Default lexical sort order is deprecated to avoid unintended expensive sorting. ")
@@ -584,9 +671,9 @@ public:
 	FORCEINLINE bool IsNone() const
 	{
 #if PLATFORM_64BITS && !WITH_CASE_PRESERVING_NAME
-		return ToComparableInt() == 0;
+		return ToUnstableInt() == 0;
 #else
-		return !ComparisonIndex && GetNumber() == NAME_NO_NUMBER_INTERNAL;
+		return ComparisonIndex.IsNone() && GetNumber() == NAME_NO_NUMBER_INTERNAL;
 #endif
 	}
 
@@ -678,7 +765,7 @@ public:
 	int32 Compare( const FName& Other ) const;
 
 	/**
-	 * Fast compares name to passed in one using indexes. Sort is allocation order ascending.
+	 * Fast non-alphabetical order that is only stable during this process' lifetime.
 	 *
 	 * @param	Other	Name to compare this against
 	 * @return	< 0 is this < Other, 0 if this == Other, > 0 if this > Other
@@ -690,8 +777,12 @@ public:
 			return ComparisonDiff;
 		}
 
-			return GetNumber() - Other.GetNumber();
-		}
+#if UE_FNAME_OUTLINE_NUMBER
+		return 0;  // If comparison indices are the same we are the same
+#else //UE_FNAME_OUTLINE_NUMBER
+		return GetNumber() - Other.GetNumber();
+#endif //UE_FNAME_OUTLINE_NUMBER
+	}
 
 	/**
 	 * Create an FName with a hardcoded string index.
@@ -707,11 +798,7 @@ public:
 	 * @param InNumber The hardcoded value for the number portion of the name
 	 */
 	FORCEINLINE FName(EName Ename, int32 InNumber)
-		: ComparisonIndex(FNameEntryId::FromEName(Ename))
-#if WITH_CASE_PRESERVING_NAME
-		, DisplayIndex(ComparisonIndex)
-#endif
-		, Number(InNumber)
+		: FName(CreateNumberedNameIfNecessary(FNameEntryId::FromEName(Ename), InNumber))
 	{
 	}
 
@@ -722,12 +809,8 @@ public:
 	 * @param Other The FName to take the string values from
 	 * @param InNumber The hardcoded value for the number portion of the name
 	 */
-	FORCEINLINE FName( const FName& Other, int32 InNumber )
-		: ComparisonIndex( Other.ComparisonIndex )
-#if WITH_CASE_PRESERVING_NAME
-		, DisplayIndex( Other.DisplayIndex )
-#endif
-		, Number( InNumber )
+	FORCEINLINE FName(FName Other, int32 InNumber)
+		: FName(CreateNumberedNameIfNecessary(Other.GetComparisonIndex(), Other.GetDisplayIndex(), InNumber))
 	{
 	}
 
@@ -735,12 +818,8 @@ public:
 	 * Create an FName from its component parts
 	 * Only call this if you *really* know what you're doing
 	 */
-	FORCEINLINE FName( const FNameEntryId InComparisonIndex, const FNameEntryId InDisplayIndex, const int32 InNumber )
-		: ComparisonIndex( InComparisonIndex )
-#if WITH_CASE_PRESERVING_NAME
-		, DisplayIndex( InDisplayIndex )
-#endif
-		, Number( InNumber )
+	FORCEINLINE FName(FNameEntryId InComparisonIndex, FNameEntryId InDisplayIndex, int32 InNumber)
+		: FName(CreateNumberedNameIfNecessary(InComparisonIndex, InDisplayIndex, InNumber))
 	{
 	}
 
@@ -755,14 +834,21 @@ public:
 	 */
 	static FName CreateFromDisplayId(FNameEntryId DisplayId, int32 Number)
 	{
+		check(ResolveEntry(DisplayId)->IsNumbered() == false); // This id should be unnumbered i.e. returned from GetDisplayIndex on an FName.
 		return FName(GetComparisonIdFromDisplayId(DisplayId), DisplayId, Number);
 	}
 
+#if UE_FNAME_OUTLINE_NUMBER
+	static FName FindNumberedName(FNameEntryId DisplayId, int32 Number);
+#endif //UE_FNAME_OUTLINE_NUMBER
+
 	/**
 	 * Default constructor, initialized to None
 	 */
 	FORCEINLINE FName()
+#if !UE_FNAME_OUTLINE_NUMBER
 		: Number(NAME_NO_NUMBER_INTERNAL)
+#endif //!UE_FNAME_OUTLINE_NUMBER
 	{
 	}
 
@@ -776,9 +862,16 @@ public:
 #endif
 	{}
 
+	FORCEINLINE explicit FName(FMinimalName InName);
+	FORCEINLINE explicit FName(FScriptName InName);
+	FORCEINLINE FName(FMemoryImageName InName);
+
 	/**
-	 * Create an FName. If FindType is FNAME_Find, and the string part of the name 
-	 * doesn't already exist, then the name will be NAME_None
+	 * Create an FName. If FindType is FNAME_Find, and the name 
+	 * doesn't already exist, then the name will be NAME_None.
+	 * The check for existance or not depends on UE_FNAME_OUTLINE_NUMBER.
+	 * When UE_FNAME_OUTLINE_NUMBER is 0, we only check for the string part.
+	 * When UE_FNAME_OUTLINE_NUMBER is 1, we check for whole name including the number.
 	 *
 	 * @param Name			Value for the string portion of the name
 	 * @param FindType		Action to take (see EFindName)
@@ -803,29 +896,52 @@ public:
 	}
 
 	/**
-	 * Create an FName. If FindType is FNAME_Find, and the string part of the name 
-	 * doesn't already exist, then the name will be NAME_None
+	 * Create an FName. Will add the string to the name table if it does not exist.
+	 * When UE_FNAME_OUTLINE_NUMBER is set, will also add the combination of base string and number to the name table if it doesn't exist.
 	 *
 	 * @param Name Value for the string portion of the name
 	 * @param Number Value for the number portion of the name
-	 * @param FindType Action to take (see EFindName)
-	 * @param bSplitName true if the trailing number should be split from the name when Number == NAME_NO_NUMBER_INTERNAL, or false to always use the name as-is
 	 */
-	FName(const WIDECHAR* Name, int32 InNumber, EFindName FindType = FNAME_Add);
-	FName(const ANSICHAR* Name, int32 InNumber, EFindName FindType = FNAME_Add);
-	FName(const UTF8CHAR* Name, int32 InNumber, EFindName FindType = FNAME_Add);
-	FName(int32 Len, const WIDECHAR* Name, int32 Number, EFindName FindType = FNAME_Add);
-	FName(int32 Len, const ANSICHAR* Name, int32 InNumber, EFindName FindType = FNAME_Add);
-	FName(int32 Len, const UTF8CHAR* Name, int32 InNumber, EFindName FindType = FNAME_Add);
+	UE_DEPRECATED(5.1, "EFindName has been removed from constructors taking a Number argument to add clarity around UE_FNAME_OUTLINE_NUMBER.")
+	FName(const WIDECHAR* Name, int32 Number, EFindName FindType);
+	UE_DEPRECATED(5.1, "EFindName has been removed from constructors taking a Number argument to add clarity around UE_FNAME_OUTLINE_NUMBER.")
+	FName(const ANSICHAR* Name, int32 Number, EFindName FindType);
+	UE_DEPRECATED(5.1, "EFindName has been removed from constructors taking a Number argument to add clarity around UE_FNAME_OUTLINE_NUMBER.")
+	FName(const UTF8CHAR* Name, int32 Number, EFindName FindType);
+	UE_DEPRECATED(5.1, "EFindName has been removed from constructors taking a Number argument to add clarity around UE_FNAME_OUTLINE_NUMBER.")
+	FName(int32 Len, const WIDECHAR* Name, int32 Number, EFindName FindType);
+	UE_DEPRECATED(5.1, "EFindName has been removed from constructors taking a Number argument to add clarity around UE_FNAME_OUTLINE_NUMBER.")
+	FName(int32 Len, const ANSICHAR* Name, int32 Number, EFindName FindType);
+	UE_DEPRECATED(5.1, "EFindName has been removed from constructors taking a Number argument to add clarity around UE_FNAME_OUTLINE_NUMBER.")
+	FName(int32 Len, const UTF8CHAR* Name, int32 Number, EFindName FindType);
+	FName(const WIDECHAR* Name, int32 Number);
+	FName(const ANSICHAR* Name, int32 Number);
+	FName(const UTF8CHAR* Name, int32 Number);
+	FName(int32 Len, const WIDECHAR* Name, int32 Number);
+	FName(int32 Len, const ANSICHAR* Name, int32 Number);
+	FName(int32 Len, const UTF8CHAR* Name, int32 Number);
 
 	template <typename CharRangeType,
 		typename CharType = typename TRemoveCV<typename TRemovePointer<decltype(GetData(DeclVal<CharRangeType>()))>::Type>::Type,
 		typename = decltype(ImplicitConv<TStringView<CharType>>(DeclVal<CharRangeType>()))>
-	inline FName(CharRangeType&& Name, int32 InNumber, EFindName FindType = FNAME_Add)
+	UE_DEPRECATED(5.1, "EFindName has been removed from constructors taking a Number argument to add clarity around UE_FNAME_OUTLINE_NUMBER.")
+	inline FName(CharRangeType&& Name, int32 InNumber, EFindName FindType)
 		: FName(NoInit)
 	{
+PRAGMA_DISABLE_DEPRECATION_WARNINGS
 		TStringView<CharType> View = Forward<CharRangeType>(Name);
 		*this = FName(View.Len(), View.GetData(), InNumber, FindType);
+PRAGMA_ENABLE_DEPRECATION_WARNINGS
+	}
+
+	template <typename CharRangeType,
+		typename CharType = typename TRemoveCV<typename TRemovePointer<decltype(GetData(DeclVal<CharRangeType>()))>::Type>::Type,
+		typename = decltype(ImplicitConv<TStringView<CharType>>(DeclVal<CharRangeType>()))>
+	inline FName(CharRangeType&& Name, int32 InNumber)
+		: FName(NoInit)
+	{
+		TStringView<CharType> View = Forward<CharRangeType>(Name);
+		*this = FName(View.Len(), View.GetData(), InNumber);
 	}
 
 	/**
@@ -837,7 +953,9 @@ public:
 	 * @param FindType Action to take (see EFindName)
 	 * @param bSplitName true if the trailing number should be split from the name when Number == NAME_NO_NUMBER_INTERNAL, or false to always use the name as-is
 	 */
-	FName( const TCHAR* Name, int32 InNumber, EFindName FindType, bool bSplitName);
+	UE_DEPRECATED(5.1, "EFindName has been removed from constructors taking a Number argument to add clarity around UE_FNAME_OUTLINE_NUMBER.")
+	FName(const TCHAR* Name, int32 InNumber, EFindName FindType, bool bSplitName);
+	FName(const TCHAR* Name, int32 InNumber, bool bSplitName);
 
 	/**
 	 * Constructor used by FLinkerLoad when loading its name table; Creates an FName with an instance
@@ -890,6 +1008,13 @@ public:
 	 */
 	static int32 GetNumWideNames();
 
+#if UE_FNAME_OUTLINE_NUMBER
+	/**
+	 * @return number of numbered names in name table
+	 */
+	static int32 GetNumNumberedNames();
+#endif
+
 	static TArray<const FNameEntry*> DebugDump();
 
 	static FNameEntry const* GetEntry(EName Ename);
@@ -921,30 +1046,53 @@ public:
 	 */
 	static void TearDown();
 
-private:
+	/** Returns an integer that compares equal in the same way FNames do, only usable within the current process */
+#if UE_FNAME_OUTLINE_NUMBER
+	FORCEINLINE uint64 ToUnstableInt() const
+	{
+		return ComparisonIndex.ToUnstableInt();
+	}
+#else
+	FORCEINLINE uint64 ToUnstableInt() const
+	{
+		static_assert(STRUCT_OFFSET(FName, ComparisonIndex) == 0);
+		static_assert(STRUCT_OFFSET(FName, Number) == 4);
+		static_assert((STRUCT_OFFSET(FName, Number) + sizeof(Number)) == sizeof(uint64));
 
+		uint64 Out = 0;
+		FMemory::Memcpy(&Out, this, sizeof(uint64));
+		return Out;
+	}
+#endif
+
+private:
 	/** Index into the Names array (used to find String portion of the string/number pair used for comparison) */
 	FNameEntryId	ComparisonIndex;
+#if !UE_FNAME_OUTLINE_NUMBER
+	/** Number portion of the string/number pair (stored internally as 1 more than actual, so zero'd memory will be the default, no-instance case) */
+	uint32			Number;
+#endif// ! //UE_FNAME_OUTLINE_NUMBER
 #if WITH_CASE_PRESERVING_NAME
 	/** Index into the Names array (used to find String portion of the string/number pair used for display) */
 	FNameEntryId	DisplayIndex;
 #endif // WITH_CASE_PRESERVING_NAME
-	/** Number portion of the string/number pair (stored internally as 1 more than actual, so zero'd memory will be the default, no-instance case) */
-	uint32			Number;
-
-#if PLATFORM_64BITS && !WITH_CASE_PRESERVING_NAME
-	FORCEINLINE uint64 ToComparableInt() const
-	{
-		static_assert(sizeof(*this) == sizeof(uint64), "");
-		alignas(uint64) FName AlignedCopy = *this;
-		return reinterpret_cast<uint64&>(AlignedCopy);
-	}
-#endif
 
 	friend const TCHAR* DebugFName(int32);
 	friend const TCHAR* DebugFName(int32, int32);
 	friend const TCHAR* DebugFName(FName&);
 
+	friend struct FNameHelper;
+	friend FScriptName NameToScriptName(FName InName);
+	friend FMinimalName NameToMinimalName(FName InName);
+	friend uint32 GetTypeHash(FName Name);
+	friend FMinimalName::FMinimalName(const FName& Name);
+	friend FScriptName::FScriptName(const FName& Name);
+	friend FMemoryImageName::FMemoryImageName(const FName& Name);
+
+	friend bool operator==(FName Lhs, FMinimalName Rhs);
+	friend bool operator==(FName Lhs, FScriptName Rhs);
+	friend bool operator==(FName Lhs, FMemoryImageName Rhs);
+
 	FORCEINLINE FNameEntryId GetDisplayIndexFast() const
 	{
 #if WITH_CASE_PRESERVING_NAME
@@ -954,11 +1102,59 @@ private:
 #endif
 	}
 
+	// Accessor for unmodified comparison index when UE_FNAME_OUTLINE_NUMBER is set
+	FORCEINLINE FNameEntryId GetComparisonIndexInternal() const
+	{
+		return ComparisonIndex;
+	}
 
+	// Accessor for unmodified display index when UE_FNAME_OUTLINE_NUMBER is set
+	FORCEINLINE FNameEntryId GetDisplayIndexInternal() const
+	{
+#if WITH_CASE_PRESERVING_NAME
+		return DisplayIndex;
+#else // WITH_CASE_PRESERVING_NAME
+		return ComparisonIndex;
+#endif // WITH_CASE_PRESERVING_NAME
+	}
 
-
+	// Resolve the entry directly referred to by LookupId
+	static const FNameEntry* ResolveEntry(FNameEntryId LookupId);
+	// Recursively resolve through the entry referred to by LookupId to reach the allocated string entry, in the case of UE_FNAME_OUTLINE_NUMBER=1
+	static const FNameEntry* ResolveEntryRecursive(FNameEntryId LookupId);
 
 	static bool IsWithinBounds(FNameEntryId Id);
+
+	// These FNameEntryIds are passed in from user code so they must be non-numbered if Number != NAME_NO_NUMBER_INTERNAL
+#if UE_FNAME_OUTLINE_NUMBER
+	static FName CreateNumberedName(FNameEntryId ComparisonId, FNameEntryId DisplayId, int32 Number);
+#endif
+
+	FORCEINLINE static FName CreateNumberedNameIfNecessary(FNameEntryId ComparisonId, FNameEntryId DisplayId, int32 Number)
+	{
+#if UE_FNAME_OUTLINE_NUMBER
+		if (Number != NAME_NO_NUMBER_INTERNAL)
+		{
+			// We need to store a new entry in the name table
+			return CreateNumberedName(ComparisonId, DisplayId, Number);
+		}
+		// Otherwise we can just set the index members
+#endif
+		FName Out;
+		Out.ComparisonIndex = ComparisonId;
+#if WITH_CASE_PRESERVING_NAME
+		Out.DisplayIndex = DisplayId;
+#endif
+#if !UE_FNAME_OUTLINE_NUMBER
+		Out.Number = Number;
+#endif
+		return Out;
+	}
+
+	FORCEINLINE static FName CreateNumberedNameIfNecessary(FNameEntryId ComparisonId, int32 Number)
+	{
+		return CreateNumberedNameIfNecessary(ComparisonId, ComparisonId, Number);
+	}
 };
 
 template<> struct TIsZeroConstructType<class FName> { enum { Value = true }; };
@@ -966,60 +1162,410 @@ Expose_TNameOf(FName)
 
 namespace Freeze
 {
-	CORE_API void IntrinsicWriteMemoryImage(FMemoryImageWriter& Writer, const FName& Object, const FTypeLayoutDesc&);
-	CORE_API uint32 IntrinsicAppendHash(const FName* DummyObject, const FTypeLayoutDesc& TypeDesc, const FPlatformTypeLayoutParameters& LayoutParams, FSHA1& Hasher);
-	CORE_API void IntrinsicWriteMemoryImage(FMemoryImageWriter& Writer, const FMinimalName& Object, const FTypeLayoutDesc&);
+	// These structures mirror the layout of FMemoryImageName depending on the value of WITH_CASE_PRESERVING_NAME
+	// for use in memory image writing/unfreezing
+	template<bool bCasePreserving>
+	struct TMemoryImageNameLayout;
+
+	template<>
+	struct TMemoryImageNameLayout<false>
+	{
+		FNameEntryId	ComparisonIndex;
+		uint32			NumberOrDummy = 0;
+	};
+
+	template<>
+	struct TMemoryImageNameLayout<true> : public TMemoryImageNameLayout<false>
+	{
+		FNameEntryId	DisplayIndex;
+	};
+	
+	CORE_API uint32 IntrinsicAppendHash(const FMemoryImageName* DummyObject, const FTypeLayoutDesc& TypeDesc, const FPlatformTypeLayoutParameters& LayoutParams, FSHA1& Hasher);
+	CORE_API void IntrinsicWriteMemoryImage(FMemoryImageWriter& Writer, const FMemoryImageName& Object, const FTypeLayoutDesc&);
+	CORE_API uint32 IntrinsicUnfrozenCopy(const FMemoryUnfreezeContent& Context, const FMemoryImageName& Object, void* OutDst);
 	CORE_API void IntrinsicWriteMemoryImage(FMemoryImageWriter& Writer, const FScriptName& Object, const FTypeLayoutDesc&);
-	CORE_API uint32 IntrinsicUnfrozenCopy(const FMemoryUnfreezeContent& Context, const FName& Object, void* OutDst);
 }
 
-DECLARE_INTRINSIC_TYPE_LAYOUT(FName);
-DECLARE_INTRINSIC_TYPE_LAYOUT(FMinimalName);
+FORCEINLINE FMemoryImageName::FMemoryImageName()
+{
+	// The structure must match the layout of Freeze::TMemoryImageNameLayout for the matching value of WITH_CASE_PRESERVING_NAME
+	static_assert(sizeof(FMemoryImageName) == sizeof(Freeze::TMemoryImageNameLayout<WITH_CASE_PRESERVING_NAME>));
+	static_assert(STRUCT_OFFSET(FMemoryImageName, ComparisonIndex) == STRUCT_OFFSET(Freeze::TMemoryImageNameLayout<WITH_CASE_PRESERVING_NAME>, ComparisonIndex));
+#if UE_FNAME_OUTLINE_NUMBER
+	static_assert(STRUCT_OFFSET(FMemoryImageName, Dummy) == STRUCT_OFFSET(Freeze::TMemoryImageNameLayout<WITH_CASE_PRESERVING_NAME>, NumberOrDummy));
+#else
+	static_assert(STRUCT_OFFSET(FMemoryImageName, Number) == STRUCT_OFFSET(Freeze::TMemoryImageNameLayout<WITH_CASE_PRESERVING_NAME>, NumberOrDummy));
+#endif
+#if WITH_CASE_PRESERVING_NAME
+	static_assert(STRUCT_OFFSET(FMemoryImageName, DisplayIndex) == STRUCT_OFFSET(Freeze::TMemoryImageNameLayout<WITH_CASE_PRESERVING_NAME>, DisplayIndex));
+#endif
+
+}
+
+FORCEINLINE FMemoryImageName::FMemoryImageName(EName Name)
+	: FMemoryImageName(FName(Name))
+{
+}
+
+DECLARE_INTRINSIC_TYPE_LAYOUT(FMemoryImageName);
 DECLARE_INTRINSIC_TYPE_LAYOUT(FScriptName);
 
+#if UE_FNAME_OUTLINE_NUMBER
 FORCEINLINE uint32 GetTypeHash(FName Name)
 {
-	return GetTypeHash(Name.GetComparisonIndex()) + Name.GetNumber();
+	return ::GetTypeHash(Name.GetComparisonIndexInternal());
 }
 
 FORCEINLINE uint32 GetTypeHash(FMinimalName Name)
 {
-	return GetTypeHash(Name.Index) + Name.Number;
+	return ::GetTypeHash(Name.Index);
 }
 
 FORCEINLINE uint32 GetTypeHash(FScriptName Name)
 {
-	return GetTypeHash(Name.ComparisonIndex) + Name.Number;
+	return ::GetTypeHash(Name.ComparisonIndex);
 }
 
-FORCEINLINE FString LexToString(const FName& Name)
+FORCEINLINE uint32 GetTypeHash(FMemoryImageName Name)
 {
-	return Name.ToString();
+	return ::GetTypeHash(Name.ComparisonIndex);
 }
 
-FORCEINLINE void LexFromString(FName& Name, const TCHAR* Str)
+FORCEINLINE FName::FName(FMinimalName InName)
+	: ComparisonIndex(InName.Index)
+#if WITH_CASE_PRESERVING_NAME
+	, DisplayIndex(InName.Index)
+#endif
 {
-	Name = FName(Str);
 }
 
-FORCEINLINE FMinimalName NameToMinimalName(const FName& InName)
+FORCEINLINE FName::FName(FScriptName InName)
+	: ComparisonIndex(InName.ComparisonIndex)
+#if WITH_CASE_PRESERVING_NAME
+	, DisplayIndex(InName.DisplayIndex)
+#endif
+{
+
+}
+
+FORCEINLINE FName::FName(FMemoryImageName InName)
+	: ComparisonIndex(InName.ComparisonIndex)
+#if WITH_CASE_PRESERVING_NAME
+	, DisplayIndex(InName.DisplayIndex)
+#endif
+{
+}
+
+FORCEINLINE FMinimalName::FMinimalName(const FName& Name)
+	: Index(Name.GetComparisonIndexInternal())
+{
+}
+
+FORCEINLINE FScriptName::FScriptName(const FName& Name)
+	: ComparisonIndex(Name.GetComparisonIndexInternal())
+#if WITH_CASE_PRESERVING_NAME
+	, DisplayIndex(Name.GetDisplayIndexInternal())
+#endif
+{
+}
+
+FORCEINLINE FMemoryImageName::FMemoryImageName(const FName& Name)
+	: ComparisonIndex(Name.GetComparisonIndexInternal())
+#if WITH_CASE_PRESERVING_NAME
+	, DisplayIndex(Name.GetDisplayIndexInternal())
+#endif
+{
+}
+
+FORCEINLINE bool FMinimalName::IsNone() const
+{
+	return Index.IsNone();
+}
+
+FORCEINLINE bool FMinimalName::operator<(FMinimalName Rhs) const
+{
+	return Index < Rhs.Index;
+}
+
+FORCEINLINE bool FScriptName::IsNone() const
+{
+	return ComparisonIndex.IsNone();
+}
+
+FORCEINLINE bool FMemoryImageName::IsNone() const
+{
+	return ComparisonIndex.IsNone();
+}
+
+FORCEINLINE bool FName::operator==(EName Ename) const
 {
-	return FMinimalName(InName.GetComparisonIndex(), InName.GetNumber());
+	return GetComparisonIndex() == Ename && GetNumber() == NAME_NO_NUMBER_INTERNAL;
 }
 
-FORCEINLINE FName MinimalNameToName(const FMinimalName& InName)
+FORCEINLINE bool operator==(FMinimalName Lhs, FMinimalName Rhs)
 {
-	return FName(InName.Index, InName.Index, InName.Number);
+	return Lhs.Index == Rhs.Index;
 }
 
-FORCEINLINE FScriptName NameToScriptName(const FName& InName)
+FORCEINLINE bool operator==(FScriptName Lhs, FScriptName Rhs)
 {
-	return FScriptName(InName.GetComparisonIndex(), InName.GetDisplayIndex(), InName.GetNumber());
+	return Lhs.ComparisonIndex == Rhs.ComparisonIndex;
 }
 
-FORCEINLINE FName ScriptNameToName(const FScriptName& InName)
+FORCEINLINE bool operator==(FMemoryImageName Lhs, FMemoryImageName Rhs)
 {
-	return FName(InName.ComparisonIndex, InName.DisplayIndex, InName.Number);
+	return Lhs.ComparisonIndex == Rhs.ComparisonIndex;
+}
+
+FORCEINLINE bool operator==(FName Lhs, FMinimalName Rhs)
+{
+	return Lhs.GetComparisonIndexInternal() == Rhs.Index;
+}
+
+FORCEINLINE bool operator==(FName Lhs, FScriptName Rhs)
+{
+	return Lhs.GetComparisonIndexInternal() == Rhs.ComparisonIndex;
+}
+
+FORCEINLINE bool operator==(FName Lhs, FMemoryImageName Rhs)
+{
+	return Lhs.GetComparisonIndexInternal() == Rhs.ComparisonIndex;
+}
+
+FORCEINLINE bool FName::IsEqual(const FName& Rhs, const ENameCase CompareMethod /*= ENameCase::IgnoreCase*/, const bool bCompareNumber /*= true*/) const
+{
+	return bCompareNumber ?
+		(CompareMethod == ENameCase::IgnoreCase ? GetComparisonIndexInternal() == Rhs.GetComparisonIndexInternal() : GetDisplayIndexInternal() == Rhs.GetDisplayIndexInternal()) :  // Unresolved indices include the number, are stored in instances
+		(CompareMethod == ENameCase::IgnoreCase ? GetComparisonIndex() == Rhs.GetComparisonIndex() : GetDisplayIndex() == Rhs.GetDisplayIndex()); // Resolved indices, have to hit the name table
+}
+
+#else // UE_FNAME_OUTLINE_NUMBER
+FORCEINLINE uint32 GetTypeHash(FName Name)
+{
+	return ::GetTypeHash(Name.GetComparisonIndex()) + Name.GetNumber();
+}
+
+FORCEINLINE uint32 GetTypeHash(FMinimalName Name)
+{
+	return ::GetTypeHash(Name.Index) + Name.Number;
+}
+
+FORCEINLINE uint32 GetTypeHash(FScriptName Name)
+{
+	return ::GetTypeHash(Name.ComparisonIndex) + Name.Number;
+}
+
+FORCEINLINE uint32 GetTypeHash(FMemoryImageName Name)
+{
+	return ::GetTypeHash(Name.ComparisonIndex) + Name.Number;
+}
+
+FORCEINLINE FName::FName(FMinimalName InName)
+	: ComparisonIndex(InName.Index)
+	, Number(InName.Number)
+#if WITH_CASE_PRESERVING_NAME
+	, DisplayIndex(InName.Index)
+#endif
+{
+}
+
+FORCEINLINE FName::FName(FScriptName InName)
+	: ComparisonIndex(InName.ComparisonIndex)
+	, Number(InName.Number)
+#if WITH_CASE_PRESERVING_NAME
+	, DisplayIndex(InName.DisplayIndex)
+#endif
+{
+}
+
+FORCEINLINE FName::FName(FMemoryImageName InName)
+	: ComparisonIndex(InName.ComparisonIndex)
+	, Number(InName.Number)
+#if WITH_CASE_PRESERVING_NAME
+	, DisplayIndex(InName.DisplayIndex)
+#endif
+{
+}
+
+FORCEINLINE FMinimalName::FMinimalName(const FName& Name)
+	: Index(Name.GetComparisonIndexInternal())
+	, Number(Name.GetNumber())
+{
+}
+
+FORCEINLINE FScriptName::FScriptName(const FName& Name)
+	: ComparisonIndex(Name.GetComparisonIndexInternal())
+#if WITH_CASE_PRESERVING_NAME
+	, DisplayIndex(Name.GetDisplayIndexInternal())
+#endif
+	, Number(Name.GetNumber())
+{
+}
+
+FORCEINLINE FMemoryImageName::FMemoryImageName(const FName& Name)
+	: ComparisonIndex(Name.GetComparisonIndex())
+	, Number(Name.GetNumber())
+#if WITH_CASE_PRESERVING_NAME
+	, DisplayIndex(Name.GetDisplayIndex())
+#endif
+{
+}
+
+FORCEINLINE bool FMinimalName::IsNone() const
+{
+	return Index.IsNone() && Number == NAME_NO_NUMBER_INTERNAL;
+}
+
+FORCEINLINE bool FMinimalName::operator<(FMinimalName Rhs) const
+{
+	return Index == Rhs.Index ? Number < Rhs.Number : Index < Rhs.Index;
+}
+
+FORCEINLINE bool FScriptName::IsNone() const
+{
+	return ComparisonIndex.IsNone() && Number == NAME_NO_NUMBER_INTERNAL;
+}
+
+FORCEINLINE bool FMemoryImageName::IsNone() const
+{
+	return ComparisonIndex.IsNone() && Number == NAME_NO_NUMBER_INTERNAL;
+}
+
+FORCEINLINE bool FName::operator==(EName Ename) const
+{
+	return GetComparisonIndex() == Ename && GetNumber() == NAME_NO_NUMBER_INTERNAL;
+}
+
+FORCEINLINE bool operator==(FMinimalName Lhs, FMinimalName Rhs)
+{
+	return Lhs.Index == Rhs.Index && Lhs.Number == Rhs.Number;
+}
+
+FORCEINLINE bool operator==(FScriptName Lhs, FScriptName Rhs)
+{
+	return Lhs.ComparisonIndex == Rhs.ComparisonIndex && Lhs.Number == Rhs.Number;
+}
+
+FORCEINLINE bool operator==(FMemoryImageName Lhs, FMemoryImageName Rhs)
+{
+	return Lhs.ComparisonIndex == Rhs.ComparisonIndex && Lhs.Number == Rhs.Number;
+}
+
+FORCEINLINE bool operator==(FName Lhs, FMinimalName Rhs)
+{
+	return Lhs.GetComparisonIndex() == Rhs.Index && Lhs.GetNumber() == Rhs.Number;
+}
+
+FORCEINLINE bool operator==(FName Lhs, FScriptName Rhs)
+{
+	return Lhs.GetComparisonIndex() == Rhs.ComparisonIndex && Lhs.GetNumber() == Rhs.Number;
+}
+
+FORCEINLINE bool operator==(FName Lhs, FMemoryImageName Rhs)
+{
+	return Lhs.GetComparisonIndex() == Rhs.ComparisonIndex && Lhs.GetNumber() == Rhs.Number;
+}
+
+FORCEINLINE bool FName::IsEqual(const FName& Rhs, const ENameCase CompareMethod /*= ENameCase::IgnoreCase*/, const bool bCompareNumber /*= true*/) const
+{
+	return ((CompareMethod == ENameCase::IgnoreCase) ? GetComparisonIndex() == Rhs.GetComparisonIndex() : GetDisplayIndexFast() == Rhs.GetDisplayIndexFast())
+		&& (!bCompareNumber || GetNumber() == Rhs.GetNumber());
+}
+#endif // UE_FNAME_OUTLINE_NUMBER
+
+FORCEINLINE FName MinimalNameToName(FMinimalName InName)
+{
+	return FName(InName);
+}
+
+FORCEINLINE FName ScriptNameToName(FScriptName InName)
+{
+	return FName(InName);
+}
+
+FORCEINLINE FMinimalName NameToMinimalName(FName InName)
+{
+	return FMinimalName(InName);
+}
+
+FORCEINLINE FScriptName NameToScriptName(FName InName)
+{
+	return FScriptName(InName);
+}
+
+FORCEINLINE bool operator==(FMinimalName Lhs, FName Rhs)
+{
+	return operator==(Rhs, Lhs);
+}
+
+FORCEINLINE bool operator==(FScriptName Lhs, FName Rhs)
+{
+	return operator==(Rhs, Lhs);
+}
+
+FORCEINLINE bool operator==(FMemoryImageName Lhs, FName Rhs)
+{
+	return operator==(Rhs, Lhs);
+}
+
+FORCEINLINE bool operator!=(FMinimalName Lhs, FMinimalName Rhs)
+{
+	return !operator==(Lhs, Rhs);
+}
+
+FORCEINLINE bool operator!=(FMinimalName Lhs, FName Rhs)
+{
+	return !operator==(Lhs, Rhs);
+}
+
+FORCEINLINE bool operator!=(FScriptName Lhs, FScriptName Rhs)
+{
+	return !operator==(Lhs, Rhs);
+}
+
+FORCEINLINE bool operator!=(FScriptName Lhs, FName Rhs)
+{
+	return !operator==(Lhs, Rhs);
+}
+
+FORCEINLINE bool operator!=(FMemoryImageName Lhs, FMemoryImageName Rhs)
+{
+	return !operator==(Lhs, Rhs);
+}
+
+FORCEINLINE bool operator!=(FMemoryImageName Lhs, FName Rhs)
+{
+	return !operator==(Lhs, Rhs);
+}
+
+FORCEINLINE bool operator!=(FName Lhs, FMinimalName Rhs)
+{
+	return !operator==(Lhs, Rhs);
+}
+
+FORCEINLINE bool operator!=(FName Lhs, FScriptName Rhs)
+{
+	return !operator==(Lhs, Rhs);
+}
+
+FORCEINLINE bool operator!=(FName Lhs, FMemoryImageName Rhs)
+{
+	return !operator==(Lhs, Rhs);
+}
+
+FORCEINLINE bool FName::operator!=(EName Ename) const
+{
+	return !(*this == Ename);
+}
+
+FORCEINLINE FString LexToString(const FName& Name)
+{
+	return Name.ToString();
+}
+
+FORCEINLINE void LexFromString(FName& Name, const TCHAR* Str)
+{
+	Name = FName(Str);
 }
 
 inline FStringBuilderBase& operator<<(FStringBuilderBase& Builder, const FName& Name)
@@ -1090,66 +1636,6 @@ struct FNameLexicalLess
 	}
 };
 
-FORCEINLINE bool operator==(const FMinimalName& Lhs, const FMinimalName& Rhs)
-{
-	return Lhs.Number == Rhs.Number && Lhs.Index == Rhs.Index;
-}
-
-FORCEINLINE bool operator!=(const FMinimalName& Lhs, const FMinimalName& Rhs)
-{
-	return !operator==(Lhs, Rhs);
-}
-
-FORCEINLINE bool operator==(const FScriptName& Lhs, const FScriptName& Rhs)
-{
-	return Lhs.Number == Rhs.Number && Lhs.ComparisonIndex == Rhs.ComparisonIndex;
-}
-
-FORCEINLINE bool operator!=(const FScriptName& Lhs, const FScriptName& Rhs)
-{
-	return !operator==(Lhs, Rhs);
-}
-
-FORCEINLINE bool operator==(const FName& Lhs, const FMinimalName& Rhs)
-{
-	return Lhs.GetNumber() == Rhs.Number && Lhs.GetComparisonIndex() == Rhs.Index;
-}
-
-FORCEINLINE bool operator!=(const FName& Lhs, const FMinimalName& Rhs)
-{
-	return !operator==(Lhs, Rhs);
-}
-
-FORCEINLINE bool operator==(const FMinimalName& Lhs, const FName& Rhs)
-{
-	return Lhs.Number == Rhs.GetNumber() && Lhs.Index == Rhs.GetComparisonIndex();
-}
-
-FORCEINLINE bool operator!=(const FMinimalName& Lhs, const FName& Rhs)
-{
-	return !operator==(Lhs, Rhs);
-}
-
-FORCEINLINE bool operator==(const FName& Lhs, const FScriptName& Rhs)
-{
-	return Lhs.GetNumber() == Rhs.Number && Lhs.GetComparisonIndex() == Rhs.ComparisonIndex;
-}
-
-FORCEINLINE bool operator!=(const FName& Lhs, const FScriptName& Rhs)
-{
-	return !operator==(Lhs, Rhs);
-}
-
-FORCEINLINE bool operator==(const FScriptName& Lhs, const FName& Rhs)
-{
-	return Lhs.Number == Rhs.GetNumber() && Lhs.ComparisonIndex == Rhs.GetComparisonIndex();
-}
-
-FORCEINLINE bool operator!=(const FScriptName& Lhs, const FName& Rhs)
-{
-	return !operator==(Lhs, Rhs);
-}
-
 #ifndef WITH_CUSTOM_NAME_ENCODING
 inline void FNameEntry::Encode(ANSICHAR*, uint32) {}
 inline void FNameEntry::Encode(WIDECHAR*, uint32) {}
diff --git a/Engine/Source/Runtime/CoreUObject/Private/Serialization/AsyncLoading2.cpp b/Engine/Source/Runtime/CoreUObject/Private/Serialization/AsyncLoading2.cpp
--- a/Engine/Source/Runtime/CoreUObject/Private/Serialization/AsyncLoading2.cpp
+++ b/Engine/Source/Runtime/CoreUObject/Private/Serialization/AsyncLoading2.cpp
@@ -99,7 +99,7 @@ FArchive& operator<<(FArchive& Ar, FExportBundleHeader& ExportBundleHeader)
 
 FArchive& operator<<(FArchive& Ar, FScriptObjectEntry& ScriptObjectEntry)
 {
-	Ar << ScriptObjectEntry.ObjectName.Index << ScriptObjectEntry.ObjectName.Number;
+	Ar << ScriptObjectEntry.Mapped;
 	Ar << ScriptObjectEntry.GlobalIndex;
 	Ar << ScriptObjectEntry.OuterIndex;
 	Ar << ScriptObjectEntry.CDOClassIndex;
@@ -484,9 +484,8 @@ struct FGlobalImportStore
 			ScriptObjectEntriesMap.Reserve(ScriptObjectEntries.Num());
 			for (FScriptObjectEntry& ScriptObjectEntry : ScriptObjectEntries)
 			{
-				const FMappedName& MappedName = FMappedName::FromMinimalName(ScriptObjectEntry.ObjectName);
-				check(MappedName.IsGlobal());
-				ScriptObjectEntry.ObjectName = NameMap.GetMinimalName(MappedName);
+				check(ScriptObjectEntry.Mapped.IsGlobal());
+				ScriptObjectEntry.ObjectName = NameMap.GetMinimalName(ScriptObjectEntry.Mapped);
 
 				ScriptObjectEntriesMap.Add(ScriptObjectEntry.GlobalIndex, &ScriptObjectEntry);
 			}
diff --git a/Engine/Source/Runtime/CoreUObject/Private/Tests/ObjectPtrTest.cpp b/Engine/Source/Runtime/CoreUObject/Private/Tests/ObjectPtrTest.cpp
--- a/Engine/Source/Runtime/CoreUObject/Private/Tests/ObjectPtrTest.cpp
+++ b/Engine/Source/Runtime/CoreUObject/Private/Tests/ObjectPtrTest.cpp
@@ -296,6 +296,35 @@ bool FObjectPtrTestHashConsistency::RunTest(const FString& Parameters)
 	return true;
 }
 
+IMPLEMENT_CUSTOM_SIMPLE_AUTOMATION_TEST(FObjectPtrTestLongPath, FObjectPtrTestBase, TEST_NAME_ROOT TEXT(".LongPath"), ObjectPtrTestFlags)
+bool FObjectPtrTestLongPath::RunTest(const FString& Parameters)
+{
+	const FName TestPackage1Name(TEXT("/Engine/Test/FObjectPtrTestLongPath/Transient"));
+	UPackage* TestPackage1 = NewObject<UPackage>(nullptr, TestPackage1Name, RF_Transient);
+	TestPackage1->AddToRoot();
+
+	ON_SCOPE_EXIT{
+		TestPackage1->RemoveFromRoot();
+	};
+
+	UObject* TestObject1 = NewObject<UTestDummyObject>(TestPackage1, TEXT("TestObject1"));
+	UObject* TestObject2 = NewObject<UTestDummyObject>(TestObject1, TEXT("TestObject2"));
+	UObject* TestObject3 = NewObject<UTestDummyObject>(TestObject2, TEXT("TestObject3"));
+	UObject* TestObject4 = NewObject<UTestDummyObject>(TestObject3, TEXT("TestObject4"));
+
+	FObjectPathId LongPath(TestObject4);
+	FObjectPathId::ResolvedNameContainerType ResolvedNames;
+	LongPath.Resolve(ResolvedNames);
+
+	TestEqual(TEXT("Resolved path from FObjectPathId should have 4 elements"), ResolvedNames.Num(), 4);
+	TestEqual(TEXT("Resolved path from FObjectPathId should have TestObject1 at element 0"), ResolvedNames[0], TestObject1->GetFName());
+	TestEqual(TEXT("Resolved path from FObjectPathId should have TestObject2 at element 1"), ResolvedNames[1], TestObject2->GetFName());
+	TestEqual(TEXT("Resolved path from FObjectPathId should have TestObject3 at element 2"), ResolvedNames[2], TestObject3->GetFName());
+	TestEqual(TEXT("Resolved path from FObjectPathId should have TestObject4 at element 3"), ResolvedNames[3], TestObject4->GetFName());
+
+	return true;
+}
+
 // @TODO: OBJPTR: We should have a test that ensures that lazy loading of an object with an external package is handled correctly.
 //				  This should also include external packages in the outer chain of the target object.
 // IMPLEMENT_CUSTOM_SIMPLE_AUTOMATION_TEST(FObjectPtrTestExternalPackages, FObjectPtrTestBase, TEST_NAME_ROOT TEXT(".ExternalPackages"), ObjectPtrTestFlags)
diff --git a/Engine/Source/Runtime/CoreUObject/Private/UObject/ObjectPathId.cpp b/Engine/Source/Runtime/CoreUObject/Private/UObject/ObjectPathId.cpp
--- a/Engine/Source/Runtime/CoreUObject/Private/UObject/ObjectPathId.cpp
+++ b/Engine/Source/Runtime/CoreUObject/Private/UObject/ObjectPathId.cpp
@@ -24,8 +24,8 @@ FCriticalSection GComplexPathLock;  // Could be changed to an RWLock later if ne
 //representing shared path elements like just registering the paths as FNames and not having
 //our own storage at all.
 TMultiMap<uint32, uint32> GComplexPathHashToId;
-template class CORE_API TArray<FMinimalName, TInlineAllocator<3>>;
-TArray<TArray<FMinimalName, TInlineAllocator<3>>> GComplexPaths;
+
+TArray<UE::ObjectPath::Private::FStoredObjectPath> GComplexPaths;
 /*
 template <int a, int b>
 struct assert_equality
@@ -37,10 +37,71 @@ struct assert_equality
 assert_equality<sizeof(TArray<FMinimalName, TInlineAllocator<3>>), 0> B;
 */
 
-static inline uint64 NameToUInt64(const FName& Name)
+namespace UE::ObjectPath::Private
 {
-	return (static_cast<uint64>(Name.GetComparisonIndex().ToUnstableInt()) << 32) |
-			static_cast<uint64>(Name.GetNumber());
+	FStoredObjectPath::FStoredObjectPath(TConstArrayView<FMinimalName> InNames)
+	{
+		// Copy into storage in reverse
+		NumElements = InNames.Num();
+
+		FMinimalName* Dest = nullptr;
+
+		if (NumElements > NumInlineElements)
+		{
+			Long = new FMinimalName[NumElements];
+			Dest = Long;
+		}
+		else
+		{
+			Dest = Short;
+		}
+
+		for (int32 i = NumElements - 1; i >= 0; --i)
+		{
+			*Dest = InNames[i];
+			++Dest;
+		}
+	}
+
+	FStoredObjectPath::~FStoredObjectPath()
+	{
+		if (NumElements > NumInlineElements)
+		{
+			delete[] Long;
+		}
+	}
+
+	FStoredObjectPath::FStoredObjectPath(FStoredObjectPath&& Other)
+		: NumElements(Other.NumElements)
+	{
+		Other.NumElements = 0;
+		if (NumElements > NumInlineElements)
+		{
+			Long = Other.Long;
+		}
+		else
+		{
+			FMemory::Memcpy(Short, Other.Short, sizeof(Short));
+		}
+	}
+
+	FStoredObjectPath& FStoredObjectPath::operator=(FStoredObjectPath&& Other)
+	{
+		this->~FStoredObjectPath();
+		new(this) FStoredObjectPath(MoveTemp(Other));
+
+		return *this;
+	}
+
+	TConstArrayView<FMinimalName> FStoredObjectPath::GetView() const
+	{
+		return TConstArrayView<FMinimalName>{ NumElements <= NumInlineElements ? Short : Long, NumElements };
+	}
+
+	static bool operator==(const FStoredObjectPath& A, const FStoredObjectPath& B)
+	{
+		return A.NumElements == B.NumElements && CompareItems(A.GetView().GetData(), B.GetView().GetData(), A.NumElements);
+	}
 }
 
 template <typename NameProducerType>
@@ -70,28 +131,31 @@ void StoreObjectPathId(NameProducerType& NameProducer, uint64 SimplePathFlag, ui
 
 	//Complex path scenario
 	TArray<FMinimalName, TInlineAllocator<3>> MinimalNames;
-	uint32 Key = GetTypeHash(NameToUInt64(Name));
-	MinimalNames.Emplace(Name.GetComparisonIndex(), Name.GetNumber());
+	uint32 Key = GetTypeHash(Name.ToUnstableInt());
+	MinimalNames.Emplace(NameToMinimalName(Name));
 	while (OuterName != NAME_None)
 	{
 		Name = OuterName;
 		OuterName = NameProducer();
-		MinimalNames.Emplace(Name.GetComparisonIndex(), Name.GetNumber());
-		Key = HashCombine(Key, GetTypeHash(NameToUInt64(Name)));
+		MinimalNames.Emplace(NameToMinimalName(Name));
+		Key = HashCombine(Key, GetTypeHash(Name.ToUnstableInt()));
 	}
 
+	// Reverse path so we can compare it to stored paths
+	UE::ObjectPath::Private::FStoredObjectPath PathToStore(MinimalNames);
+
 	FScopeLock ComplexPathScopeLock(&GComplexPathLock);
 	for (typename TMultiMap<uint32, uint32>::TConstKeyIterator It(GComplexPathHashToId, Key); It; ++It)
 	{
 		uint32 PotentialPathId = It.Value();
-		if (GComplexPaths[PotentialPathId-1] == MinimalNames)
+		if (GComplexPaths[PotentialPathId-1] == PathToStore)
 		{
 			OutObjectPathId = static_cast<uint32>(PotentialPathId << 1);
 			return;
 		}
 	}
 
-	uint32 NewId = GComplexPaths.Add(MinimalNames) + 1;
+	uint32 NewId = GComplexPaths.Emplace(MoveTemp(PathToStore)) + 1;
 	GComplexPathHashToId.Add(Key, NewId);
 	OutObjectPathId = static_cast<uint32>(NewId << 1);
 	GCoreComplexObjectPathDebug = GComplexPaths.GetData();
@@ -249,12 +313,10 @@ void FObjectPathId::Resolve(ResolvedNameContainerType& OutContainer) const
 	}
 
 	FScopeLock ComplexPathScopeLock(&GComplexPathLock);
-	//Append reversed path
-	TArray<FMinimalName, TInlineAllocator<3>>& FoundContainer = GComplexPaths[(PathId >> 1) - 1];
-	OutContainer.Reserve(OutContainer.Num() + FoundContainer.Num());
-	for (int32 Index = FoundContainer.Num() - 1; Index >= 0; --Index)
+	const UE::ObjectPath::Private::FStoredObjectPath& FoundContainer = GComplexPaths[(PathId >> 1) - 1];
+	
+	for (FMinimalName Name : FoundContainer.GetView())
 	{
-		FMinimalName& FoundMinimalName = FoundContainer[Index];
-		OutContainer.Emplace(FoundMinimalName.Index, FoundMinimalName.Index, FoundMinimalName.Number);
+		OutContainer.Emplace(MinimalNameToName(Name));
 	}
 }
diff --git a/Engine/Source/Runtime/CoreUObject/Private/UObject/ScriptCore.cpp b/Engine/Source/Runtime/CoreUObject/Private/UObject/ScriptCore.cpp
--- a/Engine/Source/Runtime/CoreUObject/Private/UObject/ScriptCore.cpp
+++ b/Engine/Source/Runtime/CoreUObject/Private/UObject/ScriptCore.cpp
@@ -2323,10 +2323,10 @@ DEFINE_FUNCTION(UObject::execInstrumentation)
 #endif
 	if (EventType == EScriptInstrumentation::InlineEvent)
 	{
-		const FName& EventName = *reinterpret_cast<FName*>(&Stack.Code[1]);
-		FScriptInstrumentationSignal InstrumentationEventInfo(EventType, P_THIS, Stack, EventName);
+		const FScriptName& EventName = *reinterpret_cast<FScriptName*>(&Stack.Code[1]);
+		FScriptInstrumentationSignal InstrumentationEventInfo(EventType, P_THIS, Stack, ScriptNameToName(EventName));
 		FBlueprintCoreDelegates::InstrumentScriptEvent(InstrumentationEventInfo);
-		Stack.SkipCode(sizeof(FName) + 1);
+		Stack.SkipCode(sizeof(FScriptName) + 1);
 	}
 	else
 	{
diff --git a/Engine/Source/Runtime/CoreUObject/Private/UObject/UObjectGlobals.cpp b/Engine/Source/Runtime/CoreUObject/Private/UObject/UObjectGlobals.cpp
--- a/Engine/Source/Runtime/CoreUObject/Private/UObject/UObjectGlobals.cpp
+++ b/Engine/Source/Runtime/CoreUObject/Private/UObject/UObjectGlobals.cpp
@@ -23,6 +23,7 @@
 #include "UObject/UObjectBase.h"
 #include "UObject/UObjectBaseUtility.h"
 #include "UObject/UObjectHash.h"
+#include "UObject/UObjectHashPrivate.h"
 #include "UObject/Object.h"
 #include "UObject/GarbageCollection.h"
 #include "UObject/Class.h"
@@ -1870,16 +1871,330 @@ void EndLoad(FUObjectSerializeContext* LoadContext)
 	Object name functions.
 -----------------------------------------------------------------------------*/
 
+static TAtomic<int32> NameNumberUniqueIndex(MAX_int32 - 1000);
+
+DEFINE_LOG_CATEGORY_STATIC(LogUniqueObjectName, Error, Error);
+
+#define TRY_REUSE_NAMES UE_FNAME_OUTLINE_NUMBER
+
+#if TRY_REUSE_NAMES
+namespace NameReuse
+{
+	static bool GTryReuseNames = true;
+	static FAutoConsoleVariableRef CVar_TryReuseNames(
+		TEXT("UObject.ReuseNames"),
+		GTryReuseNames,
+		TEXT("Try to reuse object names to prevent growth of numbered name table.")
+	);
+
+	static int32 GNameRangeCycleCadence = 8;
+	static FAutoConsoleVariableRef CVar_NameRangeCycleCadence(
+		TEXT("UObject.NameRangeCycleCadence"),
+		GNameRangeCycleCadence,
+		TEXT("When we have created this many new names in the range-based allocator and the reuse range is exhausted, return to the start and try reusing existing names. Must be a power of 2.")
+	);
+
+	static int32 GNameRangeMaxIterations = 8;
+	static FAutoConsoleVariableRef CVar_NameRangeMaxIterations(
+		TEXT("UObject.NameRangeMaxIterations"),
+		GNameRangeMaxIterations,
+		TEXT("Max number of iterations of attempting to reuse old names before bailing and creating a new one.")
+	);
+
+	bool CanUseNameInOuter(UObject* TestParent, FName TestName)
+	{
+		return !DoesObjectPossiblyExist(TestParent, TestName);
+	}
+
+	// Pointers passed into this cache are converted to integers for storage and comparison.
+	// We are ok with false caches misses when the address is reused for another object. 
+	struct FRecentNameCache
+	{
+		static constexpr int32 MaxEntries = 64;
+		static constexpr int32 MaxNamesPerEntry = 16;
+
+		struct FEntry
+		{
+			FNameEntryId BaseName;
+			TArray<TTuple<UPTRINT, FName>> Names;
+
+			FEntry(FNameEntryId InBaseName)
+				: BaseName(InBaseName)
+			{
+				Names.Reserve(MaxNamesPerEntry);
+			}
+
+			FEntry(FEntry&&) = default;
+			FEntry& operator=(FEntry&&) = default;
+
+			void Reset(FNameEntryId InBaseName)
+			{
+				BaseName = InBaseName;
+				Names.Reset();
+			}
+
+			void Store(UPTRINT Parent, FName Name)
+			{
+				int32 Index = Names.IndexOfByPredicate([&](const TTuple<UPTRINT, FName>& Pair) { return Pair.Get<1>() == Name; });
+				if (Index != INDEX_NONE)
+				{
+					Names.RemoveAt(Index, 1, false);
+				}
+				else
+				{
+					if (Names.Num() >= MaxNamesPerEntry)
+					{
+						Names.RemoveAt(0, 1, false);
+					}
+				}
+
+				Names.Emplace(Parent, Name);
+			}
+		};
+
+		TArray<FEntry> Entries;
+
+		FRecentNameCache()
+		{
+		}
+
+		FName Find(UObject* ForParent, FNameEntryId BaseId, FName BaseName)
+		{
+			FEntry* Entry = Entries.FindByPredicate([BaseId](const FEntry& Entry) { return Entry.BaseName == BaseId; });
+			if (Entry) 
+			{
+				for (TTuple<UPTRINT, FName> Pair : Entry->Names)
+				{
+					if (Pair.Get<1>() != BaseName && Pair.Get<0>() != reinterpret_cast<UPTRINT>(ForParent))
+					{
+						if (CanUseNameInOuter(ForParent, Pair.Get<1>()))
+						{
+							return Pair.Get<1>();
+						}
+					}
+				}
+			}
+
+			return FName();
+		}
+
+		void Store(UObject* Parent, FNameEntryId BaseId, FName UsedName)
+		{
+			FEntry* Entry = Entries.FindByPredicate([BaseId](const FEntry& Entry) { return Entry.BaseName == BaseId; });
+			if (!Entry)
+			{
+				if (Entries.Num() < MaxEntries)
+				{
+					// Replace the oldest entry in the cache or add a new one
+					Entry = &Entries.Emplace_GetRef(BaseId);
+				}
+				else
+				{
+					Entry = &Entries[0];
+
+					UE_LOG(LogUniqueObjectName, Log, TEXT("EVICT: %s"), *FName(Entry->BaseName, Entry->BaseName, NAME_NO_NUMBER_INTERNAL).ToString());
+					Entry->Reset(BaseId);
+				}
+
+				UE_LOG(LogUniqueObjectName, Log, TEXT("STORE: %s"), *FName(BaseId, BaseId, NAME_NO_NUMBER_INTERNAL).ToString());
+			}
+
+			Entry->Store(reinterpret_cast<UPTRINT>(Parent), UsedName);
+
+			// Shift this entry to the end of the array as it's now the most recently used
+			int32 Index = Entry - Entries.GetData();
+			if (Index != Entries.Num() - 1)
+			{
+				FEntry Removed = MoveTemp(*Entry);
+				Entries.RemoveAt(Index, 1, false);
+				Entries.Add(MoveTemp(Removed));
+			}
+		}
+	};
+
+	static thread_local FRecentNameCache GRecentNameCache;
+
+	struct FNameRangeCache
+	{
+		struct FNameRangeEntry
+		{
+			static const constexpr int32 FirstNewNumber = MAX_int32 - 1001;
+
+			int32 NextNewNumber;
+			int32 Iterator;
+
+			FNameRangeEntry()
+				: NextNewNumber(FirstNewNumber)
+				, Iterator(FirstNewNumber - GNameRangeCycleCadence)
+			{
+			}
+
+			FName AllocateName(UObject* Parent, FNameEntryId BaseId, FName BaseName)
+			{
+				checkSlow(BaseId == BaseName.GetComparisonIndex());
+				int32 BaseNumber = BaseName.GetNumber();
+
+				// Do we have any old numbers to try and reuse? 
+				if (Iterator > NextNewNumber)
+				{
+					int32 Start = Iterator;
+					int32 End = Start - GNameRangeMaxIterations;
+					// Look for existing (name, number) pairs that are unused 
+					for (int32 i = Start; i > NextNewNumber && i > End; --i)
+					{
+						if (i == BaseNumber)
+						{
+							continue;
+						}
+
+						FName TestName(BaseId, BaseId, i);
+						if (CanUseNameInOuter(Parent, TestName))
+						{
+							UE_LOG(LogUniqueObjectName, Log, TEXT("HIT: %s %s"), *Parent->GetPathName(), *TestName.ToString());
+							Iterator = i - 1;
+							return TestName;
+						}
+					}
+
+					if (Iterator == NextNewNumber)
+					{
+						// Failed to find a name to reuse, reset iterator to where we should reset
+						Iterator = NextNewNumber - GNameRangeCycleCadence;
+					}
+				}
+
+				// Start == Next or we fail to reuse an existing index
+				for (int32 i = NextNewNumber; i > 0; --i)
+				{
+					if (i == BaseNumber)
+					{
+						continue;
+					}
+
+					FName TestName(BaseId, BaseId, i);
+					if (CanUseNameInOuter(Parent, TestName))
+					{
+						NextNewNumber = i - 1;
+						if (NextNewNumber < Iterator)
+						{
+							Iterator = FirstNewNumber;
+						}
+
+						UE_LOG(LogUniqueObjectName, Log, TEXT("MISS: %s %s"), *Parent->GetPathName(), *TestName.ToString());
+						return TestName;
+					}
+				}
+
+				return FName();
+			}
+		};
+
+		FName Find(UObject* Parent, FNameEntryId BaseId, FName BaseName)
+		{
+			Lock.ReadLock();
+			FNameRangeEntry* Entry = Find(BaseId);
+
+			if (Entry)
+			{
+				// already allocated name, we can release the shared lock and work on this object directly
+				Lock.ReadUnlock();
+			}
+			else
+			{
+				Lock.ReadUnlock();
+
+				// The first time we request a name if we've never created one, don't bother adding to the cache just yet.
+				FName TestName = FName::FindNumberedName(BaseId, FNameRangeEntry::FirstNewNumber + 1);
+				if (TestName.IsNone())
+				{
+					return FName(BaseId, BaseId, FNameRangeEntry::FirstNewNumber + 1);
+				}
+
+				FWriteScopeLock _(Lock);
+				// We didn't have a name but we may have been preempted as we acquired the write lock
+				Entry = Find(BaseId);
+				if (!Entry)
+				{
+					// we were not pre-empted, add a new entry for this name 
+					Entry = Add(BaseId);
+				}
+			}
+
+			return Entry->AllocateName(Parent, BaseId, BaseName);
+		}
+	private:
+		FRWLock Lock;
+		TMap<FNameEntryId, TUniquePtr<FNameRangeEntry>> Map;
+
+		FNameRangeEntry* Find(FNameEntryId Id)
+		{
+			TUniquePtr<FNameRangeEntry>* Ptr = Map.Find(Id);
+			if (Ptr) { return Ptr->Get(); }
+			return nullptr;
+		}
+
+		FNameRangeEntry* Add(FNameEntryId Id)
+		{
+			return Map.Emplace(Id, MakeUnique<FNameRangeEntry>()).Get();
+		}
+
+	};
+
+	static FNameRangeCache GNameRangeCache;
+
+	FName MakeUniqueObjectNameReusingNumber(UObject* Parent, FName BaseName)
+	{
+		if (!GTryReuseNames)
+		{
+			return FName();
+		}
+
+		static const FName NamePackage(NAME_Package);
+		if (!Parent || Parent == ANY_PACKAGE || BaseName == NamePackage || FPlatformProperties::HasEditorOnlyData() || !GFastPathUniqueNameGeneration)
+		{
+			return FName();
+		}
+
+		LLM_SCOPE(ELLMTag::UObject);
+
+		FNameEntryId BaseId = BaseName.GetDisplayIndex();
+		FName ReturnName = GRecentNameCache.Find(Parent, BaseId, BaseName);
+
+		if (ReturnName.IsNone())
+		{
+			ReturnName = GNameRangeCache.Find(Parent, BaseId, BaseName);
+		}
+
+		if (ReturnName.IsNone())
+		{
+			// Store this name for reuse 
+			GRecentNameCache.Store(Parent, BaseId, ReturnName);
+		}
+
+		return ReturnName;
+	}
+}
+#endif
+
 FName MakeUniqueObjectName( UObject* Parent, const UClass* Class, FName InBaseName/*=NAME_None*/ )
 {
+	TRACE_CPUPROFILER_EVENT_SCOPE(MakeUniqueObjectName);
+
 	CSV_SCOPED_TIMING_STAT(UObject, MakeUniqueObjectName);
 	check(Class);
-	const FName BaseName = (InBaseName == NAME_None) ? Class->GetFName() : InBaseName;
+	const FName BaseName = InBaseName.IsNone() ? Class->GetFName() : InBaseName;
+
+#if TRY_REUSE_NAMES
+	if (FName Result = NameReuse::MakeUniqueObjectNameReusingNumber(Parent, BaseName); !Result.IsNone())
+	{
+		return Result;
+	}
+#endif
+#undef TRY_REUSE_NAMES
 
 	FName TestName;
 	do
 	{
-		// cache the class's name's index for faster name creation later
 		UObject* ExistingObject;
 
 		do
@@ -1922,8 +2237,7 @@ FName MakeUniqueObjectName( UObject* Parent, const UClass* Class, FName InBaseNa
 						*   could never clash with anything with a different outer. For animation trees, these outers are never saved or loaded, thus clashes are
 						*   impossible.
 						*/
-						static TAtomic<int32> UniqueIndex(MAX_int32 - 1000);
-						NameNumber = --UniqueIndex;
+						NameNumber = --NameNumberUniqueIndex;
 					}
 					else
 					{
diff --git a/Engine/Source/Runtime/CoreUObject/Private/UObject/UObjectHash.cpp b/Engine/Source/Runtime/CoreUObject/Private/UObject/UObjectHash.cpp
--- a/Engine/Source/Runtime/CoreUObject/Private/UObject/UObjectHash.cpp
+++ b/Engine/Source/Runtime/CoreUObject/Private/UObject/UObjectHash.cpp
@@ -811,6 +811,17 @@ UObject* StaticFindObjectFastInternal(const UClass* ObjectClass, const UObject*
 	return Result;
 }
 
+// Approximate search for finding unused object names
+bool DoesObjectPossiblyExist(const UObject* InOuter, FName ObjectName)
+{
+	FUObjectHashTables& ThreadHash = FUObjectHashTables::Get();
+	check(InOuter != nullptr);
+	int32 Hash = GetObjectOuterHash(ObjectName, (PTRINT)InOuter);
+	FHashTableLock HashLock(ThreadHash);
+	// We don't need to iterate the multimap here as we are happy with false positives.
+	return ThreadHash.HashOuter.Contains(Hash);
+}
+
 // Assumes that ThreadHash's critical is already locked
 FORCEINLINE static void AddToOuterMap(FUObjectHashTables& ThreadHash, UObjectBase* Object)
 {
diff --git a/Engine/Source/Runtime/CoreUObject/Private/UObject/UObjectHashPrivate.h b/Engine/Source/Runtime/CoreUObject/Private/UObject/UObjectHashPrivate.h
new file mode 100644
--- /dev/null
+++ b/Engine/Source/Runtime/CoreUObject/Private/UObject/UObjectHashPrivate.h
@@ -0,0 +1,19 @@
+// Copyright Epic Games, Inc. All Rights Reserved.
+
+/*=============================================================================
+	UObjectHashPrivate.h: UObjectHash functions for use only within CoreUObject
+=============================================================================*/
+
+#pragma once
+
+#include "UObject/UObjectHash.h"
+
+/**
+ * Approximate version of StaticFindObjectFastInternal that may return a false positive - i.e. the object does not exist but the function returns true.
+ * Will not return a false negative - i.e. if the object does exist, this function must return true
+ *
+ * @param	InOuter			The direct outer to search within. Must be non-null.
+ * @param	ObjectName		The object name to search for.
+ * @return	Returns a true if the object may possibly exist, false if it definitely does not exist.
+ */
+bool DoesObjectPossiblyExist(const UObject* InOuter, FName ObjectName);
diff --git a/Engine/Source/Runtime/CoreUObject/Public/Serialization/AsyncLoading2.h b/Engine/Source/Runtime/CoreUObject/Public/Serialization/AsyncLoading2.h
--- a/Engine/Source/Runtime/CoreUObject/Public/Serialization/AsyncLoading2.h
+++ b/Engine/Source/Runtime/CoreUObject/Public/Serialization/AsyncLoading2.h
@@ -292,13 +292,21 @@ struct FExportBundleHeader
 
 struct FScriptObjectEntry
 {
-	FMinimalName ObjectName;
+	union 
+	{
+		FMappedName Mapped;
+		FMinimalName ObjectName;
+	};
 	FPackageObjectIndex GlobalIndex;
 	FPackageObjectIndex OuterIndex;
 	FPackageObjectIndex CDOClassIndex;
 
+	FScriptObjectEntry() {}
+
 	COREUOBJECT_API friend FArchive& operator<<(FArchive& Ar, FScriptObjectEntry& ScriptObjectEntry);
 };
+// The sizeof FMinimalName may be variable but FMappedName should always be larger, so FScriptObjectEntry has a fixed size.
+static_assert(sizeof(FMappedName) >= sizeof(FMinimalName));
 
 /**
  * Export map entry.
diff --git a/Engine/Source/Runtime/CoreUObject/Public/UObject/ObjectPathId.h b/Engine/Source/Runtime/CoreUObject/Public/UObject/ObjectPathId.h
--- a/Engine/Source/Runtime/CoreUObject/Public/UObject/ObjectPathId.h
+++ b/Engine/Source/Runtime/CoreUObject/Public/UObject/ObjectPathId.h
@@ -10,6 +10,31 @@ class UObject;
 class FLinkerTables;
 struct FObjectImport;
 
+// Declared in the header so the type exists in debug info for debugger visualization, not a public part of the API.
+namespace UE::ObjectPath::Private
+{
+	struct FStoredObjectPath
+	{
+		static constexpr const int32 NumInlineElements = 3;
+		int32 NumElements;
+
+		union
+		{
+			FMinimalName Short[NumInlineElements];
+			FMinimalName* Long;
+		};
+
+		FStoredObjectPath(TConstArrayView<FMinimalName> InNames);
+		~FStoredObjectPath();
+		FStoredObjectPath(const FStoredObjectPath&) = delete;
+		FStoredObjectPath(FStoredObjectPath&&);
+		FStoredObjectPath& operator=(const FStoredObjectPath&) = delete;
+		FStoredObjectPath& operator=(FStoredObjectPath&&);
+
+		TConstArrayView<FMinimalName> GetView() const;
+	};
+}
+
 // @TODO: OBJPTR: Should this be changed to FObjectImportPathId? It is already written to target the specific patterns in import paths,
 //		 if we go further in the future and do things like store classname/classpackage info, to the path segments, it would be
 //		 even more specific to imports.
diff --git a/Engine/Source/Runtime/D3D12RHI/Private/D3D12RayTracing.cpp b/Engine/Source/Runtime/D3D12RHI/Private/D3D12RayTracing.cpp
--- a/Engine/Source/Runtime/D3D12RHI/Private/D3D12RayTracing.cpp
+++ b/Engine/Source/Runtime/D3D12RHI/Private/D3D12RayTracing.cpp
@@ -265,7 +265,7 @@ static void DumpRayTracingGeometries(EDumpRayTracingGeometryMode Mode, int32 Num
 			if (bCSV)
 			{
 				const FString Row = FString::Printf(TEXT("%s,%.3f,%d,%d,%d,%d,%d\n"),
-					Geometry->DebugName.IsValid() ? *Geometry->DebugName.ToString() : TEXT("*UNKNOWN*"),
+					!Geometry->DebugName.IsNone() ? *Geometry->DebugName.ToString() : TEXT("*UNKNOWN*"),
 					SizeBytes / double(1 << 20),
 					Geometry->Initializer.TotalPrimitiveCount,
 					Geometry->Initializer.Segments.Num(),
@@ -277,7 +277,7 @@ static void DumpRayTracingGeometries(EDumpRayTracingGeometryMode Mode, int32 Num
 			else
 			{
 				BufferedOutput.CategorizedLogf(CategoryName, ELogVerbosity::Log, TEXT("Name: %s - Size: %.3f MB - Prims: %d - Segments: %d -  Compaction: %d - Update: %d"),
-					Geometry->DebugName.IsValid() ? *Geometry->DebugName.ToString() : TEXT("*UNKNOWN*"),
+					!Geometry->DebugName.IsNone() ? *Geometry->DebugName.ToString() : TEXT("*UNKNOWN*"),
 					SizeBytes / double(1 << 20),
 					Geometry->Initializer.TotalPrimitiveCount,
 					Geometry->Initializer.Segments.Num(),
@@ -3257,9 +3257,8 @@ enum class ERayTracingBufferType
 	Scratch
 };
 
-static TRefCountPtr<FD3D12Buffer> CreateRayTracingBuffer(FD3D12Adapter* Adapter, uint32 GPUIndex, uint64 Size, ERayTracingBufferType Type, FName DebugName)
+static TRefCountPtr<FD3D12Buffer> CreateRayTracingBuffer(FD3D12Adapter* Adapter, uint32 GPUIndex, uint64 Size, ERayTracingBufferType Type, const FDebugName& DebugName)
 {
-	check(DebugName.IsValid());
 	check(Size);
 
 	TRefCountPtr<FD3D12Buffer> Result;
@@ -3310,8 +3309,8 @@ FD3D12RayTracingGeometry::FD3D12RayTracingGeometry(FD3D12Adapter* Adapter, const
 {
 	INC_DWORD_STAT(STAT_D3D12RayTracingAllocatedBLAS);
 
-	DebugName = Initializer.DebugName.IsValid() ? Initializer.DebugName : FName(TEXT("BLAS"));
-
+	DebugName = !Initializer.DebugName.IsNone() ? Initializer.DebugName : FDebugName(FName(TEXT("BLAS")));
+	
 	FMemory::Memzero(bHasPendingCompactionRequests);
 	FMemory::Memzero(bRegisteredAsRenameListener);
 
@@ -3690,7 +3689,7 @@ void FD3D12RayTracingGeometry::SetInitializer(const FRayTracingGeometryInitializ
 		UnregisterAsRenameListener(GPUIndex);
 	}	
 
-	DebugName = Initializer.DebugName.IsValid() ? Initializer.DebugName : FName(TEXT("BLAS"));
+	DebugName = !Initializer.DebugName.IsNone() ? Initializer.DebugName : FName(TEXT("BLAS"));
 
 	checkf(Initializer.Segments.Num() > 0, TEXT("Ray tracing geometry must be initialized with at least one segment."));
 
diff --git a/Engine/Source/Runtime/D3D12RHI/Private/D3D12RayTracing.h b/Engine/Source/Runtime/D3D12RHI/Private/D3D12RayTracing.h
--- a/Engine/Source/Runtime/D3D12RHI/Private/D3D12RayTracing.h
+++ b/Engine/Source/Runtime/D3D12RHI/Private/D3D12RayTracing.h
@@ -79,7 +79,7 @@ public:
 	// Hit shader parameters per geometry segment
 	TArray<FHitGroupSystemParameters> HitGroupSystemParameters[MAX_NUM_GPUS];
 
-	FName DebugName;
+	FDebugName DebugName;
 
 	// Array of geometry descriptions, one per segment (single-segment geometry is a common case).
 	// Only references CPU-accessible structures (no GPU resources).
diff --git a/Engine/Source/Runtime/Engine/Private/PackageMapClient.cpp b/Engine/Source/Runtime/Engine/Private/PackageMapClient.cpp
--- a/Engine/Source/Runtime/Engine/Private/PackageMapClient.cpp
+++ b/Engine/Source/Runtime/Engine/Private/PackageMapClient.cpp
@@ -1007,16 +1007,16 @@ FNetworkGUID UPackageMapClient::InternalLoadObject( FArchive & Ar, UObject *& Ob
 
 		FNetworkGUID OuterGUID = InternalLoadObject( Ar, ObjOuter, InternalLoadObjectRecursionCount + 1 );
 
-		FString PathName;
+		FString ObjectName;
 		uint32	NetworkChecksum = 0;
 
-		Ar << PathName;
+		Ar << ObjectName;
 
 		if ( ExportFlags.bHasNetworkChecksum )
 		{
 			Ar << NetworkChecksum;
 
-			UE_LOG(LogNetPackageMap, Verbose, TEXT("%s has network checksum %u"), *PathName, NetworkChecksum);
+			UE_LOG(LogNetPackageMap, Verbose, TEXT("%s has network checksum %u"), *ObjectName, NetworkChecksum);
 		}
 
 		const bool bIsPackage = NetGUID.IsStatic() && !OuterGUID.IsValid();
@@ -1029,7 +1029,7 @@ FNetworkGUID UPackageMapClient::InternalLoadObject( FArchive & Ar, UObject *& Ob
 		}
 
 		// Remap name for PIE
-		GEngine->NetworkRemapPath( Connection, PathName, true );
+		GEngine->NetworkRemapPath( Connection, ObjectName, true );
 
 		if (NetGUID.IsDefault())
 		{
@@ -1042,28 +1042,29 @@ FNetworkGUID UPackageMapClient::InternalLoadObject( FArchive & Ar, UObject *& Ob
 			// relative path name is meaningless. This may happen if the outer has been garbage collected.
 			if (!bIsPackage && OuterGUID.IsValid() && ObjOuter == nullptr)
 			{
-				UE_LOG( LogNetPackageMap, Log, TEXT( "InternalLoadObject: couldn't find outer for non-package object. GUID: %s, PathName: %s" ), *NetGUID.ToString(), *PathName );
+				UE_LOG( LogNetPackageMap, Log, TEXT( "InternalLoadObject: couldn't find outer for non-package object. GUID: %s, ObjectName: %s" ), *NetGUID.ToString(), *ObjectName );
 				Object = nullptr;
 				return NetGUID;
 			}
 
-			Object = StaticFindObject(UObject::StaticClass(), ObjOuter, *PathName, false);
+			Object = StaticFindObject(UObject::StaticClass(), ObjOuter, *ObjectName, false);
 
+			// Try to load package if it wasn't found. Note load package fails if the package is already loaded.
 			if (Object == nullptr && bIsPackage)
 			{
-				// Try to load package if it wasn't found. Note load package fails if the package is already loaded.
-				Object = LoadPackage(NULL, *PathName, LOAD_None);
+				FPackagePath Path = FPackagePath::FromPackageNameChecked(ObjectName);
+				Object = LoadPackage(nullptr, Path, LOAD_None);
 			}
 
 			if ( Object == NULL )
 			{
-				UE_LOG( LogNetPackageMap, Warning, TEXT( "UPackageMapClient::InternalLoadObject: Unable to resolve default guid from client: PathName: %s, ObjOuter: %s " ), *PathName, ObjOuter != NULL ? *ObjOuter->GetPathName() : TEXT( "NULL" ) );
+				UE_LOG( LogNetPackageMap, Warning, TEXT( "UPackageMapClient::InternalLoadObject: Unable to resolve default guid from client: ObjectName: %s, ObjOuter: %s " ), *ObjectName, ObjOuter != NULL ? *ObjOuter->GetPathName() : TEXT( "NULL" ) );
 				return NetGUID;
 			}
 
 			if (!IsValid(Object))
 			{
-				UE_LOG( LogNetPackageMap, Warning, TEXT( "UPackageMapClient::InternalLoadObject: Received reference to invalid object from client: PathName: %s, ObjOuter: %s "), *PathName, ObjOuter != NULL ? *ObjOuter->GetPathName() : TEXT( "NULL" ) );
+				UE_LOG( LogNetPackageMap, Warning, TEXT( "UPackageMapClient::InternalLoadObject: Received reference to invalid object from client: ObjectName: %s, ObjOuter: %s "), *ObjectName, ObjOuter != NULL ? *ObjOuter->GetPathName() : TEXT( "NULL" ) );
 				Object = NULL;
 				return NetGUID;
 			}
@@ -1074,7 +1075,7 @@ FNetworkGUID UPackageMapClient::InternalLoadObject( FArchive & Ar, UObject *& Ob
 
 				if (CompareNetworkChecksum != NetworkChecksum )
 				{
-					FString ErrorStr = FString::Printf(TEXT("UPackageMapClient::InternalLoadObject: Default object package network checksum mismatch! PathName: %s, ObjOuter: %s, GUID1: %u, GUID2: %u "), *PathName, ObjOuter != NULL ? *ObjOuter->GetPathName() : TEXT("NULL"), CompareNetworkChecksum, NetworkChecksum);
+					FString ErrorStr = FString::Printf(TEXT("UPackageMapClient::InternalLoadObject: Default object package network checksum mismatch! ObjectName: %s, ObjOuter: %s, GUID1: %u, GUID2: %u "), *ObjectName, ObjOuter != NULL ? *ObjOuter->GetPathName() : TEXT("NULL"), CompareNetworkChecksum, NetworkChecksum);
 					UE_LOG( LogNetPackageMap, Error, TEXT("%s"), *ErrorStr);
 					Object = NULL;
 
@@ -1089,7 +1090,7 @@ FNetworkGUID UPackageMapClient::InternalLoadObject( FArchive & Ar, UObject *& Ob
 
 				if ( Package == NULL )
 				{
-					UE_LOG( LogNetPackageMap, Error, TEXT( "UPackageMapClient::InternalLoadObject: Default object not a package from client: PathName: %s, ObjOuter: %s " ), *PathName, ObjOuter != NULL ? *ObjOuter->GetPathName() : TEXT( "NULL" ) );
+					UE_LOG( LogNetPackageMap, Error, TEXT( "UPackageMapClient::InternalLoadObject: Default object not a package from client: ObjectName: %s, ObjOuter: %s " ), *ObjectName, ObjOuter != NULL ? *ObjOuter->GetPathName() : TEXT( "NULL" ) );
 					Object = NULL;
 					return NetGUID;
 				}
@@ -1106,14 +1107,14 @@ FNetworkGUID UPackageMapClient::InternalLoadObject( FArchive & Ar, UObject *& Ob
 		else if ( Object != nullptr )
 		{
 			// If we already have the object, just do some sanity checking and return
-			SanityCheckExport( GuidCache.Get(), Object, NetGUID, PathName, ObjOuter, OuterGUID, ExportFlags );
+			SanityCheckExport( GuidCache.Get(), Object, NetGUID, ObjectName, ObjOuter, OuterGUID, ExportFlags );
 			return NetGUID;
 		}
 
 		// If we are the server, we should have found the object by now
 		if ( IsNetGUIDAuthority() )
 		{
-			UE_LOG( LogNetPackageMap, Warning, TEXT( "UPackageMapClient::InternalLoadObject: Server could not resolve non default guid from client. PathName: %s, ObjOuter: %s " ), *PathName, ObjOuter != NULL ? *ObjOuter->GetPathName() : TEXT( "NULL" ) );
+			UE_LOG( LogNetPackageMap, Warning, TEXT( "UPackageMapClient::InternalLoadObject: Server could not resolve non default guid from client. ObjectName: %s, ObjOuter: %s " ), *ObjectName, ObjOuter != NULL ? *ObjOuter->GetPathName() : TEXT( "NULL" ) );
 			return NetGUID;
 		}
 
@@ -1124,14 +1125,14 @@ FNetworkGUID UPackageMapClient::InternalLoadObject( FArchive & Ar, UObject *& Ob
 		const bool bIgnoreWhenMissing = ExportFlags.bNoLoad;
 
 		// Register this path and outer guid combo with the net guid
-		GuidCache->RegisterNetGUIDFromPath_Client( NetGUID, PathName, OuterGUID, NetworkChecksum, ExportFlags.bNoLoad, bIgnoreWhenMissing );
+		GuidCache->RegisterNetGUIDFromPath_Client( NetGUID, ObjectName, OuterGUID, NetworkChecksum, ExportFlags.bNoLoad, bIgnoreWhenMissing );
 
 		// Try again now that we've registered the path
 		Object = GuidCache->GetObjectFromNetGUID( NetGUID, GuidCache->IsExportingNetGUIDBunch );
 
 		if ( Object == NULL && !GuidCache->ShouldIgnoreWhenMissing( NetGUID ) )
 		{
-			UE_LOG( LogNetPackageMap, Warning, TEXT( "InternalLoadObject: Unable to resolve object from path. Path: %s, Outer: %s, NetGUID: %s" ), *PathName, ObjOuter ? *ObjOuter->GetPathName() : TEXT( "NULL" ), *NetGUID.ToString() );
+			UE_LOG( LogNetPackageMap, Warning, TEXT( "InternalLoadObject: Unable to resolve object from path. Path: %s, Outer: %s, NetGUID: %s" ), *ObjectName, ObjOuter ? *ObjOuter->GetPathName() : TEXT( "NULL" ), *NetGUID.ToString() );
 		}
 	}
 	else if ( Object == NULL && !GuidCache->ShouldIgnoreWhenMissing( NetGUID ) )
@@ -3270,7 +3271,7 @@ UObject* FNetGUIDCache::GetObjectFromNetGUID( const FNetworkGUID& NetGUID, const
 	}
 
 	// See if this object is in memory
-	Object = StaticFindObject( UObject::StaticClass(), ObjOuter, *CacheObjectPtr->PathName.ToString(), false );
+	Object = FindObjectFast<UObject>(ObjOuter, CacheObjectPtr->PathName);
 
 	// Assume this is a package if the outer is invalid and this is a static guid
 	const bool bIsPackage = NetGUID.IsStatic() && !CacheObjectPtr->OuterGUID.IsValid();
@@ -3310,7 +3311,8 @@ UObject* FNetGUIDCache::GetObjectFromNetGUID( const FNetworkGUID& NetGUID, const
 			else
 			{
 				// Async loading disabled
-				Object = LoadPackage( NULL, *CacheObjectPtr->PathName.ToString(), LOAD_None );
+				FPackagePath Path = FPackagePath::FromPackageNameChecked(CacheObjectPtr->PathName);
+				Object = LoadPackage(nullptr, Path, LOAD_None);
 				SyncLoadedGUIDs.AddUnique(NetGUID);
 			}
 		}
diff --git a/Engine/Source/Runtime/Engine/Private/PerPlatformProperties.cpp b/Engine/Source/Runtime/Engine/Private/PerPlatformProperties.cpp
--- a/Engine/Source/Runtime/Engine/Private/PerPlatformProperties.cpp
+++ b/Engine/Source/Runtime/Engine/Private/PerPlatformProperties.cpp
@@ -34,7 +34,9 @@ ENGINE_API FArchive& operator<<(FArchive& Ar, TPerPlatformProperty<StructType, V
 #if WITH_EDITORONLY_DATA
 		if (!bCooked)
 		{
-			Ar << This->PerPlatform;
+			using MapType = decltype(This->PerPlatform);
+			using KeyFuncs = typename PerPlatformProperty::Private::KeyFuncs<MapType>;
+			KeyFuncs::SerializePerPlatformMap(Ar, This->PerPlatform);
 		}
 #endif
 	}
@@ -67,7 +69,9 @@ ENGINE_API void operator<<(FStructuredArchive::FSlot Slot, TPerPlatformProperty<
 #if WITH_EDITORONLY_DATA
 		if (!bCooked)
 		{
-			Record << SA_VALUE(TEXT("PerPlatform"), This->PerPlatform);
+			using MapType = decltype(This->PerPlatform);
+			using KeyFuncs = typename PerPlatformProperty::Private::KeyFuncs<MapType>;
+			KeyFuncs::SerializePerPlatformMap(UnderlyingArchive, Record, This->PerPlatform);
 		}
 #endif
 	}
diff --git a/Engine/Source/Runtime/Engine/Private/SkeletalRenderCPUSkin.cpp b/Engine/Source/Runtime/Engine/Private/SkeletalRenderCPUSkin.cpp
--- a/Engine/Source/Runtime/Engine/Private/SkeletalRenderCPUSkin.cpp
+++ b/Engine/Source/Runtime/Engine/Private/SkeletalRenderCPUSkin.cpp
@@ -415,7 +415,7 @@ void FSkeletalMeshObjectCPUSkin::FSkeletalMeshObjectLOD::InitResources(FSkelMesh
 				FRayTracingGeometryInitializer Initializer;
 				static const FName DebugName("FSkeletalMeshObjectCPUSkin");
 				static int32 DebugNumber = 0;
-				Initializer.DebugName = FName(DebugName, DebugNumber++);
+				Initializer.DebugName = FDebugName(DebugName, DebugNumber++);
 				Initializer.IndexBuffer = IndexBufferRHI;
 				Initializer.TotalPrimitiveCount = TrianglesCount;
 				Initializer.GeometryType = RTGT_Triangles;
diff --git a/Engine/Source/Runtime/Engine/Private/SkeletalRenderStatic.cpp b/Engine/Source/Runtime/Engine/Private/SkeletalRenderStatic.cpp
--- a/Engine/Source/Runtime/Engine/Private/SkeletalRenderStatic.cpp
+++ b/Engine/Source/Runtime/Engine/Private/SkeletalRenderStatic.cpp
@@ -75,7 +75,7 @@ void FSkeletalMeshObjectStatic::InitResources(USkinnedMeshComponent* InMeshCompo
 							FRayTracingGeometryInitializer Initializer;
 							static const FName DebugName("FSkeletalMeshObjectLOD");
 							static int32 DebugNumber = 0;
-							Initializer.DebugName = FName(DebugName, DebugNumber++);
+							Initializer.DebugName = FDebugName(DebugName, DebugNumber++);
 							Initializer.IndexBuffer = IndexBufferRHI;
 							Initializer.TotalPrimitiveCount = TrianglesCount;
 							Initializer.GeometryType = RTGT_Triangles;
diff --git a/Engine/Source/Runtime/Engine/Private/UnrealEngine.cpp b/Engine/Source/Runtime/Engine/Private/UnrealEngine.cpp
--- a/Engine/Source/Runtime/Engine/Private/UnrealEngine.cpp
+++ b/Engine/Source/Runtime/Engine/Private/UnrealEngine.cpp
@@ -5328,28 +5328,33 @@ bool UEngine::HandleCountDisabledParticleItemsCommand( const TCHAR* Cmd, FOutput
 bool UEngine::HandleViewnamesCommand( const TCHAR* Cmd, FOutputDevice& Ar )
 {
 	const TArray<const FNameEntry*> Entries = FName::DebugDump();
+	TConstArrayView<const FNameEntry*> IterEntries = Entries;
 
-	int32 NumLast = 0;
-	int32 BeginIdx = FParse::Value(Cmd, TEXT("NUM="), NumLast) ? FMath::Max(Entries.Num() - NumLast, 0) : 0;
-
-	FString Filename;
-	TUniquePtr<FArchive> FileWriter;
-	if (FParse::Value(Cmd, TEXT("FILENAME="), Filename))
-	{
-		const FString OutputFilename = FPaths::ProjectSavedDir() / FString::Printf(TEXT("Names-%s-%s.log"), *Filename, *FDateTime::Now().ToString());
-		FileWriter = TUniquePtr<FArchive>(IFileManager::Get().CreateFileWriter(*OutputFilename));
+	if(int32 NumLast = 0; FParse::Value(Cmd, TEXT("NUM="), NumLast))
+	{ 
+		IterEntries.RightInline(NumLast);
 	}
 
-	for (int32 I = BeginIdx; I < Entries.Num(); ++I)
+#if UE_FNAME_OUTLINE_NUMBER
+	bool bNumbered = FParse::Param(Cmd, TEXT("NUMBERED"));
+	bool bPlain = FParse::Param(Cmd, TEXT("PLAIN"));
+	if (!bNumbered && !bPlain)
 	{
-		const FString PlainNameString = Entries[I]->GetPlainNameString();
-
-		Ar.Log(*PlainNameString);
+		bNumbered = bPlain = true;
+	}
+#endif
 
-		if (FileWriter)
+	for (const FNameEntry* Entry : IterEntries)
+	{	
+#if UE_FNAME_OUTLINE_NUMBER
+		if( Entry->IsNumbered() && bNumbered)
 		{
-			FString LogEntry = FString::Printf(TEXT("%s") LINE_TERMINATOR, *PlainNameString);
-			FileWriter->Serialize(TCHAR_TO_ANSI(*LogEntry), LogEntry.Len());
+			Entry->DebugDump(Ar);
+		}
+		else if (!Entry->IsNumbered() && bPlain)
+#endif
+		{
+			Entry->DebugDump(Ar);
 		}
 	}
 
@@ -9207,6 +9212,11 @@ bool UEngine::HandleObjCommand( const TCHAR* Cmd, FOutputDevice& Ar )
 		LogHashOuterStatistics(Ar, bShowHashBucketCollisionInfo);
 		return true;
 	}
+	else if (FParse::Command(&Cmd, TEXT("HASHFNAME")))
+	{
+		FName::DisplayHash(Ar);
+		return true;
+	}
 	else if (FParse::Command(&Cmd, TEXT("OVERHEAD")))
 	{
 		const bool bShowIndividualStats = FParse::Param(Cmd, TEXT("DETAILED"));
diff --git a/Engine/Source/Runtime/Engine/Public/PerPlatformProperties.h b/Engine/Source/Runtime/Engine/Public/PerPlatformProperties.h
--- a/Engine/Source/Runtime/Engine/Public/PerPlatformProperties.h
+++ b/Engine/Source/Runtime/Engine/Public/PerPlatformProperties.h
@@ -20,6 +20,94 @@ PerPlatformProperties.h: Property types that can be overridden on a per-platform
 
 #include "PerPlatformProperties.generated.h"
 
+namespace PerPlatformProperty::Private
+{
+	template<typename MapType>
+	struct TGetKeyType
+	{
+		using Type = typename TTupleElement<0, typename MapType::ElementType>::Type;
+	};
+
+	template<typename NameType>
+	struct FNameFuncs;
+
+	template<>
+	struct FNameFuncs<FName>
+	{
+		static FName NameToKey(FName Name) { return Name; }
+
+		template<typename ValueType>
+		static void SerializePerPlatformMap(FArchive& Ar, TMap<FName, ValueType>& Map)
+		{
+			Ar << Map;
+		}
+
+		template<typename ValueType>
+		static void SerializePerPlatformMap(FArchive& UnderlyingArchive, FStructuredArchive::FRecord& Record, TMap<FName, ValueType>& Map)
+		{
+			Record << SA_VALUE(TEXT("PerPlatform"), Map);
+		}
+	};
+
+	template<>
+	struct FNameFuncs<FMemoryImageName>
+	{
+		static FMemoryImageName NameToKey(FName Name) { return FMemoryImageName(Name); }
+
+		template<typename ValueType>
+		static void SerializePerPlatformMap(FArchive& Ar, TMemoryImageMap<FMemoryImageName, ValueType>& Map)
+		{
+			if( Ar.IsLoading())
+			{
+				TMemoryImageMap<FName, ValueType> TempMap;
+				Ar << TempMap;
+				Map.Reset();
+				for( TPair<FName, ValueType>& Pair : TempMap)
+				{
+					Map.Add(FMemoryImageName(Pair.Key), Pair.Value);
+				}
+			}
+			else
+			{
+				TMemoryImageMap<FName, ValueType> TempMap;
+				for (TPair<FMemoryImageName, ValueType>& Pair : Map)
+				{
+					TempMap.Add(FName(Pair.Key), Pair.Value);
+				}
+				Ar << TempMap;
+			}
+		}
+
+		template<typename ValueType>
+		static void SerializePerPlatformMap(FArchive& UnderlyingArchive, FStructuredArchive::FRecord& Record, TMemoryImageMap<FMemoryImageName, ValueType>& Map)
+		{
+			if (UnderlyingArchive.IsLoading())
+			{
+				TMemoryImageMap<FName, ValueType> TempMap;
+
+				Record << SA_VALUE(TEXT("PerPlatform"), TempMap);
+				Map.Reset();
+				for (TPair<FName, ValueType>& Pair : TempMap)
+				{
+					Map.Add(FMemoryImageName(Pair.Key), Pair.Value);
+				}
+			}
+			else
+			{
+				TMemoryImageMap<FName, ValueType> TempMap;
+				for (TPair<FMemoryImageName, ValueType>& Pair : Map)
+				{
+					TempMap.Add(FName(Pair.Key), Pair.Value);
+				}
+				Record << SA_VALUE(TEXT("PerPlatform"), TempMap);
+			}
+		}
+	};
+
+	template<typename MapType>
+	using KeyFuncs = FNameFuncs<typename TGetKeyType<MapType>::Type>;
+}
+
 /** TPerPlatformProperty - template parent class for per-platform properties 
  *  Implements Serialize function to replace value at cook time, and 
  *  backwards-compatible loading code for properties converted from simple types.
@@ -28,20 +116,25 @@ template<typename _StructType, typename _ValueType, EName _BasePropertyName>
 struct ENGINE_API TPerPlatformProperty
 {
 	typedef _ValueType ValueType;
+	typedef _StructType StructType;
 
 #if WITH_EDITOR
 	/** Get the value for the given platform (using standard "ini" name, so Windows, not Win64 or WindowsClient), which can be used to lookup the group */
 	_ValueType GetValueForPlatform(FName PlatformName) const
 	{
 		const _StructType* This = StaticCast<const _StructType*>(this);
-		const _ValueType* Ptr = This->PerPlatform.Find(PlatformName);
+
+		using MapType = decltype(This->PerPlatform);
+		using KeyFuncs = typename PerPlatformProperty::Private::KeyFuncs<MapType>;
+
+		const _ValueType* Ptr = This->PerPlatform.Find(KeyFuncs::NameToKey(PlatformName));
 
 		if (Ptr == nullptr)
 		{
 			const FDataDrivenPlatformInfo& Info = FDataDrivenPlatformInfoRegistry::GetPlatformInfo(PlatformName);
 			if (Info.PlatformGroupName != NAME_None)
 			{
-				Ptr = This->PerPlatform.Find(Info.PlatformGroupName);
+				Ptr = This->PerPlatform.Find(KeyFuncs::NameToKey(Info.PlatformGroupName));
 			}
 		}
 
@@ -53,6 +146,9 @@ struct ENGINE_API TPerPlatformProperty
 		_ValueType GetValueForPlatformIdentifiers(FName PlatformGroupName, FName VanillaPlatformName = NAME_None) const
 	{
 		const _StructType* This = StaticCast<const _StructType*>(this);
+
+		using MapType = decltype(This->PerPlatform);
+		using KeyFuncs = typename PerPlatformProperty::Private::KeyFuncs<MapType>;
 		
 		const _ValueType* ValuePtr = [This, VanillaPlatformName, PlatformGroupName]() -> const _ValueType*
 		{
@@ -65,11 +161,11 @@ struct ENGINE_API TPerPlatformProperty
 				{
 					return VanillaPlatformName.ToString().Contains(Name.ToString());
 				});
-				Ptr = MatchedName ? This->PerPlatform.Find(*MatchedName) : nullptr;
+				Ptr = MatchedName ? This->PerPlatform.Find(KeyFuncs::NameToKey(*MatchedName)) : nullptr;
 			}			
 			if (Ptr == nullptr && PlatformGroupName != NAME_None)
 			{				
-				Ptr = This->PerPlatform.Find(PlatformGroupName);
+				Ptr = This->PerPlatform.Find(KeyFuncs::NameToKey(PlatformGroupName));
 			}
 			return Ptr;			
 		}();
@@ -212,7 +308,8 @@ struct ENGINE_API FFreezablePerPlatformInt
 	GENERATED_USTRUCT_BODY()
 
 public:
-	using FPerPlatformMap = TMemoryImageMap<FName, int32>;
+	using KeyType = FMemoryImageName;
+	using FPerPlatformMap = TMemoryImageMap<FMemoryImageName, int32>;
 
 	LAYOUT_FIELD(int32, Default);
 	LAYOUT_FIELD_EDITORONLY(FPerPlatformMap, PerPlatform);
@@ -221,20 +318,28 @@ public:
 	FFreezablePerPlatformInt(int32 InDefaultValue) : Default(InDefaultValue) {}
 	FFreezablePerPlatformInt(const FPerPlatformInt& Other)
 		: Default(Other.Default)
+	{
 #if WITH_EDITORONLY_DATA
-		, PerPlatform(Other.PerPlatform)
+		for (const TPair<FName, int32>& Pair : Other.PerPlatform)
+		{
+			PerPlatform.Add(FMemoryImageName(Pair.Key), Pair.Value);
+		}
 #endif
-	{}
+	}
 
 	FString ToString() const;
 };
 
 inline FPerPlatformInt::FPerPlatformInt(const FFreezablePerPlatformInt& Other)
 	: Default(Other.Default)
+{
 #if WITH_EDITORONLY_DATA
-	, PerPlatform(Other.PerPlatform)
+	for (const TPair<FMemoryImageName, int32>& Pair : Other.PerPlatform)
+	{
+		PerPlatform.Add(FName(Pair.Key), Pair.Value);
+	}
 #endif
-{}
+}
 
 extern template ENGINE_API FArchive& operator<<(FArchive&, TPerPlatformProperty<FPerPlatformInt, int32, NAME_IntProperty>&);
 extern template ENGINE_API void operator<<(FStructuredArchive::FSlot Slot, TPerPlatformProperty<FPerPlatformInt, int32, NAME_IntProperty>&);
@@ -290,7 +395,7 @@ struct ENGINE_API FFreezablePerPlatformFloat
 {
 	DECLARE_TYPE_LAYOUT(FFreezablePerPlatformFloat, NonVirtual);
 public:
-	using FPerPlatformMap = TMemoryImageMap<FName, float>;
+	using FPerPlatformMap = TMemoryImageMap<FMemoryImageName, float>;
 	
 	LAYOUT_FIELD(float, Default);
 	LAYOUT_FIELD_EDITORONLY(FPerPlatformMap, PerPlatform);
@@ -299,18 +404,26 @@ public:
 	FFreezablePerPlatformFloat(float InDefaultValue) : Default(InDefaultValue) {}
 	FFreezablePerPlatformFloat(const FPerPlatformFloat& Other)
 		: Default(Other.Default)
+	{
 #if WITH_EDITORONLY_DATA
-		, PerPlatform(Other.PerPlatform)
+		for (const TPair<FName, float>& Pair : Other.PerPlatform)
+		{
+			PerPlatform.Add(FMemoryImageName(Pair.Key), Pair.Value);
+		}
 #endif
-	{}
+}
 };
 
 inline FPerPlatformFloat::FPerPlatformFloat(const FFreezablePerPlatformFloat& Other)
 	: Default(Other.Default)
+{
 #if WITH_EDITORONLY_DATA
-	, PerPlatform(Other.PerPlatform)
+	for (const TPair<FMemoryImageName, float>& Pair : Other.PerPlatform)
+	{
+		PerPlatform.Add(FName(Pair.Key), Pair.Value);
+	}
 #endif
-{}
+}
 
 template<>
 struct TStructOpsTypeTraits<FPerPlatformFloat>
diff --git a/Engine/Source/Runtime/Experimental/GeometryCollectionEngine/Private/GeometryCollection/GeometryCollectionSceneProxy.cpp b/Engine/Source/Runtime/Experimental/GeometryCollectionEngine/Private/GeometryCollection/GeometryCollectionSceneProxy.cpp
--- a/Engine/Source/Runtime/Experimental/GeometryCollectionEngine/Private/GeometryCollection/GeometryCollectionSceneProxy.cpp
+++ b/Engine/Source/Runtime/Experimental/GeometryCollectionEngine/Private/GeometryCollection/GeometryCollectionSceneProxy.cpp
@@ -228,7 +228,7 @@ FGeometryCollectionSceneProxy::FGeometryCollectionSceneProxy(UGeometryCollection
 			[this](FRHICommandListImmediate& RHICmdList)
 			{
 				FRayTracingGeometryInitializer Initializer;
-				Initializer.DebugName = TEXT("GeometryCollection");
+				Initializer.DebugName = FName(TEXT("GeometryCollection"));
 				Initializer.GeometryType = RTGT_Triangles;
 				Initializer.bFastBuild = true;
 				Initializer.bAllowUpdate = false;
diff --git a/Engine/Source/Runtime/Landscape/Private/LandscapeRender.cpp b/Engine/Source/Runtime/Landscape/Private/LandscapeRender.cpp
--- a/Engine/Source/Runtime/Landscape/Private/LandscapeRender.cpp
+++ b/Engine/Source/Runtime/Landscape/Private/LandscapeRender.cpp
@@ -1573,7 +1573,7 @@ void FLandscapeComponentSceneProxy::CreateRenderThreadResources()
 				FRayTracingGeometryInitializer Initializer;
 				static const FName DebugName("FLandscapeComponentSceneProxy");
 				static int32 DebugNumber = 0;
-				Initializer.DebugName = FName(DebugName, DebugNumber++);
+				Initializer.DebugName = FDebugName(DebugName, DebugNumber++);
 				Initializer.IndexBuffer = nullptr;
 				Initializer.GeometryType = RTGT_Triangles;
 				Initializer.bFastBuild = true;
diff --git a/Engine/Source/Runtime/RHI/Private/RHIResources.cpp b/Engine/Source/Runtime/RHI/Private/RHIResources.cpp
--- a/Engine/Source/Runtime/RHI/Private/RHIResources.cpp
+++ b/Engine/Source/Runtime/RHI/Private/RHIResources.cpp
@@ -1 +1,52 @@
 // Copyright Epic Games, Inc. All Rights Reserved.
+
+FDebugName::FDebugName()
+	: Name()
+	, Number(NAME_NO_NUMBER_INTERNAL)
+{
+}
+
+FDebugName::FDebugName(FName InName)
+	: Name(InName)
+	, Number(NAME_NO_NUMBER_INTERNAL)
+{
+}
+
+FDebugName::FDebugName(FName InName, int32 InNumber)
+	: Name(InName)
+	, Number(InNumber)
+{
+}
+
+FDebugName::FDebugName(FMemoryImageName InName, int32 InNumber)
+	: Name(InName)
+	, Number(InNumber)
+{
+}
+
+FDebugName& FDebugName::operator=(FName Other)
+{
+	Name = FMemoryImageName(Other);
+	Number = NAME_NO_NUMBER_INTERNAL;
+	return *this;
+}
+
+FString FDebugName::ToString() const
+{
+	FString Out;
+	FName(Name).AppendString(Out);
+	if (Number != NAME_NO_NUMBER_INTERNAL)
+	{
+		Out.Appendf(TEXT("_%u"), Number);
+	}
+	return Out;
+}
+
+void FDebugName::AppendString(FStringBuilderBase& Builder) const
+{
+	FName(Name).AppendString(Builder);
+	if (Number != NAME_NO_NUMBER_INTERNAL)
+	{
+		Builder << '_' << Number;
+	}
+}
diff --git a/Engine/Source/Runtime/RHI/Public/RHIResources.h b/Engine/Source/Runtime/RHI/Public/RHIResources.h
--- a/Engine/Source/Runtime/RHI/Public/RHIResources.h
+++ b/Engine/Source/Runtime/RHI/Public/RHIResources.h
@@ -1906,6 +1906,35 @@ typedef TRefCountPtr<FRHIComputePipelineState> FComputePipelineStateRHIRef;
 typedef TRefCountPtr<FRHIRayTracingPipelineState> FRayTracingPipelineStateRHIRef;
 
 
+/**
+ * A type used only for printing a string for debugging/profiling.
+ * Adds Number as a suffix to the printed string even if the base name includes a number, so may prints a string like: Base_1_1
+ * This type will always store a numeric suffix explicitly inside itself and never in the name table so it will always be at least 12 bytes
+ * regardless of the value of UE_FNAME_OUTLINE_NUMBER.
+ * It is not comparable or convertible to other name types to encourage its use only for debugging and avoid using more storage than necessary
+ * for the primary use cases of FName (names of objects, assets etc which are widely used and therefor deduped in the name table).
+ */
+class RHI_API FDebugName
+{
+	DECLARE_INLINE_TYPE_LAYOUT(FDebugName, NonVirtual);
+
+public:
+	FDebugName();
+	FDebugName(FName InName);
+	FDebugName(FName InName, int32 InNumber);
+	FDebugName(FMemoryImageName InName, int32 InNumber);
+
+	FDebugName& operator=(FName Other);
+
+	FString ToString() const;
+	bool IsNone() const { return Name.IsNone() && Number == NAME_NO_NUMBER_INTERNAL; }
+	void AppendString(FStringBuilderBase& Builder) const;
+
+private:
+	LAYOUT_FIELD(FMemoryImageName, Name);
+	LAYOUT_FIELD(uint32, Number);
+};
+
 //
 // Ray tracing resources
 //
@@ -2054,7 +2083,7 @@ public:
 	LAYOUT_FIELD_INITIALIZED(bool, bAllowCompaction, true);
 	LAYOUT_FIELD_INITIALIZED(ERayTracingGeometryInitializerType, Type, ERayTracingGeometryInitializerType::Rendering);
 
-	LAYOUT_FIELD(FName, DebugName);
+	LAYOUT_FIELD(FDebugName, DebugName);
 };
 
 enum ERayTracingSceneLifetime
