StringConv.h support for UTF-8 mode.

diff --git a/Engine/Source/Runtime/Core/Public/Containers/StringConv.h b/Engine/Source/Runtime/Core/Public/Containers/StringConv.h
--- a/Engine/Source/Runtime/Core/Public/Containers/StringConv.h
+++ b/Engine/Source/Runtime/Core/Public/Containers/StringConv.h
@@ -984,6 +984,18 @@ typedef TStringPointer<UTF32CHAR, TCHAR> FUTF32ToTCHAR;
 #define TCHAR_TO_UTF32(str) (UTF32CHAR*)(str)
 #define UTF32_TO_TCHAR(str) (TCHAR*)(str)
 
+#elif PLATFORM_TCHAR_IS_UTF8CHAR
+
+typedef TStringConversion<TStringConvert<TCHAR, UTF16CHAR>> FTCHARToUTF16;
+typedef TStringConversion<TStringConvert<UTF16CHAR, TCHAR>> FUTF16ToTCHAR;
+#define TCHAR_TO_UTF16(str) (UTF16CHAR*)FTCHARToUTF16((const TCHAR*)str).Get()
+#define UTF16_TO_TCHAR(str) (TCHAR*)FUTF16ToTCHAR((const UTF16CHAR*)str).Get()
+
+typedef TStringConversion<TStringConvert<TCHAR, UTF32CHAR>> FTCHARToUTF32;
+typedef TStringConversion<TStringConvert<UTF32CHAR, TCHAR>> FUTF32ToTCHAR;
+#define TCHAR_TO_UTF32(str) (UTF32CHAR*)FTCHARToUTF32((const TCHAR*)str).Get()
+#define UTF32_TO_TCHAR(str) (TCHAR*)FUTF32ToTCHAR((const UTF32CHAR*)str).Get()
+
 #else
 
 static_assert(sizeof(TCHAR) == sizeof(UTF16CHAR), "TCHAR and UTF16CHAR are expected to be the same size for inline conversion! PLATFORM_TCHAR_IS_4_BYTES is not configured correctly for this platform.");
@@ -1005,12 +1017,22 @@ typedef TStringConversion<TUTF16ToUTF32_Convert<TCHAR, wchar_t>> FTCHARToWChar;
 typedef TStringConversion<TUTF32ToUTF16_Convert<wchar_t, TCHAR>> FWCharToTCHAR;
 #define TCHAR_TO_WCHAR(str) (wchar_t*)FTCHARToWChar((const TCHAR*)str).Get()
 #define WCHAR_TO_TCHAR(str) (TCHAR*)FWCharToTCHAR((const wchar_t*)str).Get()
+
+#elif PLATFORM_TCHAR_IS_UTF8CHAR
+
+typedef TStringConversion<TStringConvert<TCHAR, wchar_t>> FTCHARToWChar;
+typedef TStringConversion<TStringConvert<wchar_t, TCHAR>> FWCharToTCHAR;
+#define TCHAR_TO_WCHAR(str) (wchar_t*)FTCHARToUTF32((const TCHAR*)str).Get()
+#define WCHAR_TO_TCHAR(str) (TCHAR*)FWCharToTCHAR((const wchar_t*)str).Get()
+
 #else
+
 static_assert(sizeof(TCHAR) == sizeof(wchar_t), "TCHAR and wchar_t are expected to be the same size for inline conversion! PLATFORM_WCHAR_IS_4_BYTES is not configured correctly for this platform.");
 typedef TStringPointer<TCHAR, wchar_t> FTCHARToWChar;
 typedef TStringPointer<wchar_t, TCHAR> FWCharToTCHAR;
 #define TCHAR_TO_WCHAR(str) (wchar_t*)(str)
 #define WCHAR_TO_TCHAR(str) (TCHAR*)(str)
+
 #endif
 
 /**
