Added a Windows-specific version of FWindowsPlatformMisc::GetMaxRefreshRate

diff --git a/Engine/Source/Runtime/Core/Private/Windows/WindowsPlatformMisc.cpp b/Engine/Source/Runtime/Core/Private/Windows/WindowsPlatformMisc.cpp
--- a/Engine/Source/Runtime/Core/Private/Windows/WindowsPlatformMisc.cpp
+++ b/Engine/Source/Runtime/Core/Private/Windows/WindowsPlatformMisc.cpp
@@ -2979,3 +2979,24 @@ uint64 FWindowsPlatformMisc::GetFileVersion(const FString &FileName)
 	}
 	return 0;
 }
+
+int32 FWindowsPlatformMisc::GetMaxRefreshRate()
+{
+	int32 Result = FGenericPlatformMisc::GetMaxRefreshRate();
+
+#if !UE_SERVER
+	DEVMODE DeviceMode;
+	FMemory::Memzero(DeviceMode);
+	DeviceMode.dmSize = sizeof(DEVMODE);
+
+	if (EnumDisplaySettings(NULL, ENUM_CURRENT_SETTINGS, &DeviceMode) != 0)
+	{
+		// dmDisplayFrequency isn't always useful, the Windows docs say it can
+		// return 0 or 1 to indicate 'default refresh rate', so always assume we
+		// can do at least the generic platform default of 60 Hz
+		Result = FMath::Max(Result, (int32)DeviceMode.dmDisplayFrequency);
+	}
+#endif
+
+	return Result;
+}
diff --git a/Engine/Source/Runtime/Core/Public/Windows/WindowsPlatformMisc.h b/Engine/Source/Runtime/Core/Public/Windows/WindowsPlatformMisc.h
--- a/Engine/Source/Runtime/Core/Public/Windows/WindowsPlatformMisc.h
+++ b/Engine/Source/Runtime/Core/Public/Windows/WindowsPlatformMisc.h
@@ -278,6 +278,7 @@ struct CORE_API FWindowsPlatformMisc
 
 	static uint64 GetFileVersion(const FString &FileName);
 
+	static int32 GetMaxRefreshRate();
 };
 
 
