Fixes for memory image names.
Correctly write to MemoryImageNames array, copy contents of TArrayView rather than the layout of the view object itself, handle name layout during patch application on foreign platforms.

diff --git a/Engine/Plugins/FX/Niagara/Source/NiagaraShader/Private/NiagaraShader.cpp b/Engine/Plugins/FX/Niagara/Source/NiagaraShader/Private/NiagaraShader.cpp
--- a/Engine/Plugins/FX/Niagara/Source/NiagaraShader/Private/NiagaraShader.cpp
+++ b/Engine/Plugins/FX/Niagara/Source/NiagaraShader/Private/NiagaraShader.cpp
@@ -274,6 +274,8 @@ void FNiagaraShaderMapId::AppendKeyString(FString& KeyString) const
 	FString FeatureLevelString;
 	GetFeatureLevelName(FeatureLevel, FeatureLevelString);
 
+	LayoutParams.AppendKeyString(KeyString);
+
 	{
 		const FSHAHash LayoutHash = Freeze::HashLayout(StaticGetTypeLayoutDesc<FNiagaraShaderMapContent>(), LayoutParams);
 		KeyString += TEXT("_");
diff --git a/Engine/Source/Runtime/Core/Private/Serialization/MemoryImage.cpp b/Engine/Source/Runtime/Core/Private/Serialization/MemoryImage.cpp
--- a/Engine/Source/Runtime/Core/Private/Serialization/MemoryImage.cpp
+++ b/Engine/Source/Runtime/Core/Private/Serialization/MemoryImage.cpp
@@ -20,6 +20,9 @@ IMPLEMENT_TYPE_LAYOUT(FMemoryImageString);
 IMPLEMENT_TYPE_LAYOUT(FPlatformTypeLayoutParameters);
 IMPLEMENT_TYPE_LAYOUT(FHashedName);
 
+// Guid for versioning format changes/bug fixes to memory images
+#define MEMORYIMAGE_DERIVEDDATA_VER		TEXT("B6A0A753ED62499D8F30931E61FCAB05")
+
 static const uint32 NumTypeLayoutDescHashBuckets = 4357u;
 static const FTypeLayoutDesc* GTypeLayoutHashBuckets[NumTypeLayoutDescHashBuckets] = { nullptr };
 static uint32 GNumTypeLayoutsRegistered = 0u;
@@ -108,7 +111,7 @@ FArchive& FPlatformTypeLayoutParameters::Serialize(FArchive& Ar)
 
 void FPlatformTypeLayoutParameters::AppendKeyString(FString& KeyString) const
 {
-	KeyString += FString::Printf(TEXT("FL_%08x_MFA_%08x_"), Flags, MaxFieldAlignment);
+	KeyString += FString::Printf(TEXT("FL_%08x_MFA_%08x_V_") MEMORYIMAGE_DERIVEDDATA_VER, Flags, MaxFieldAlignment);
 }
 
 // evaluated during static-initialization, so logging from regular check() macros won't work correctly
@@ -1482,10 +1485,17 @@ static inline void ApplyScriptNamePatch(void* FrozenObject, const FScriptName& N
 	new(NameDst) FScriptName(Name);
 }
 
-static inline void ApplyMemoryImageNamePatch(void* FrozenObject, const FMemoryImageName& Name, uint32 Offset)
+static inline void ApplyMemoryImageNamePatch(void* FrozenObject, const FMemoryImageName& Name, uint32 Offset, const FPlatformTypeLayoutParameters& LayoutParameters)
 {
 	void* NameDst = (uint8*)FrozenObject + Offset;
-	new(NameDst) FMemoryImageName(Name);
+	if (LayoutParameters.IsCurrentPlatform())
+	{
+		new(NameDst) FMemoryImageName(Name);
+	}
+	else
+	{
+		Freeze::ApplyMemoryImageNamePatch(NameDst, Name, LayoutParameters);
+	}
 }
 
 void FMemoryImageResult::ApplyPatches(void* FrozenObject) const
@@ -1504,7 +1514,7 @@ void FMemoryImageResult::ApplyPatches(void* FrozenObject) const
 
 	for (const FMemoryImageNamePointer& Patch : MemoryImageNames)
 	{
-		ApplyMemoryImageNamePatch(FrozenObject, FMemoryImageName(Patch.Name), Patch.Offset);
+		ApplyMemoryImageNamePatch(FrozenObject, FMemoryImageName(Patch.Name), Patch.Offset, TargetLayoutParameters);
 	}
 }
 
@@ -1574,7 +1584,7 @@ FMemoryImageObject FMemoryImageResult::LoadFromArchive(FArchive& Ar, const FType
 		{
 			uint32 Offset = 0u;
 			Ar << Offset;
-			ApplyMemoryImageNamePatch(FrozenObject, FMemoryImageName(Name), Offset);
+			ApplyMemoryImageNamePatch(FrozenObject, FMemoryImageName(Name), Offset, LayoutParameters);
 		}
 	}
 
@@ -1706,10 +1716,10 @@ uint32 FMemoryImageSection::WriteVTable(const FTypeLayoutDesc& TypeDesc, const F
 
 uint32 FMemoryImageSection::WriteFMemoryImageName(TConstArrayView<uint8> InBytes, const FName& Name)
 {
-	const uint32 Offset = WriteBytes(InBytes);
+	const uint32 Offset = WriteBytes(InBytes.GetData(), InBytes.Num());
 	if (!Name.IsNone())
 	{
-		FMemoryImageNamePointer& NamePointer = ScriptNames.AddDefaulted_GetRef();
+		FMemoryImageNamePointer& NamePointer = MemoryImageNames.AddDefaulted_GetRef();
 		NamePointer.Name = Name;
 		NamePointer.Offset = Offset;
 	}
diff --git a/Engine/Source/Runtime/Core/Private/UObject/UnrealNames.cpp b/Engine/Source/Runtime/Core/Private/UObject/UnrealNames.cpp
--- a/Engine/Source/Runtime/Core/Private/UObject/UnrealNames.cpp
+++ b/Engine/Source/Runtime/Core/Private/UObject/UnrealNames.cpp
@@ -4958,7 +4958,32 @@ uint32 Freeze::IntrinsicAppendHash(const FMemoryImageName* DummyObject, const FT
 	return Freeze::AppendHashForNameAndSize(TypeDesc.Name, SizeFromFields, Hasher);
 }
 
-void Freeze::IntrinsicWriteMemoryImage(FMemoryImageWriter& Writer, const FMemoryImageName& Object, const FTypeLayoutDesc&)
+void Freeze::ApplyMemoryImageNamePatch(void* NameDst, const FMemoryImageName& Name, const FPlatformTypeLayoutParameters& LayoutParams)
+{
+	if (!LayoutParams.WithCasePreservingFName())
+	{
+		TMemoryImageNameLayout<0>* ToWrite = new(NameDst) TMemoryImageNameLayout<0>();
+		ToWrite->ComparisonIndex = Name.ComparisonIndex;
+#if !UE_FNAME_OUTLINE_NUMBER
+		ToWrite->NumberOrDummy = Name.Number;
+#endif
+	}
+	else
+	{
+		TMemoryImageNameLayout<1>* ToWrite = new(NameDst) TMemoryImageNameLayout<1>();
+		ToWrite->ComparisonIndex = Name.ComparisonIndex;
+#if !UE_FNAME_OUTLINE_NUMBER
+		ToWrite->NumberOrDummy = Name.Number;
+#endif
+#if WITH_CASE_PRESERVING_NAME
+		ToWrite->DisplayIndex = Name.DisplayIndex;
+#else
+		ToWrite->DisplayIndex = Name.ComparisonIndex;
+#endif
+	}
+}
+
+void Freeze::IntrinsicWriteMemoryImage(FMemoryImageWriter& Writer, const FMemoryImageName& Object, const FTypeLayoutDesc& LayoutDesc)
 {
 	const FPlatformTypeLayoutParameters& TargetLayoutParameters = Writer.GetTargetLayoutParams();
 	if (TargetLayoutParameters.IsCurrentPlatform())
diff --git a/Engine/Source/Runtime/Core/Public/UObject/NameTypes.h b/Engine/Source/Runtime/Core/Public/UObject/NameTypes.h
--- a/Engine/Source/Runtime/Core/Public/UObject/NameTypes.h
+++ b/Engine/Source/Runtime/Core/Public/UObject/NameTypes.h
@@ -455,6 +455,7 @@ struct FMemoryImageName;
 
 namespace Freeze
 {
+	CORE_API void ApplyMemoryImageNamePatch(void* NameDst, const FMemoryImageName& Name, const FPlatformTypeLayoutParameters& LayoutParams);
 	CORE_API void IntrinsicWriteMemoryImage(FMemoryImageWriter& Writer, const FMemoryImageName& Object, const FTypeLayoutDesc&);
 	CORE_API uint32 IntrinsicUnfrozenCopy(const FMemoryUnfreezeContent& Context, const FMemoryImageName& Object, void* OutDst);
 }
@@ -470,6 +471,7 @@ struct FMemoryImageName
 	friend bool operator==(FName Lhs, FMemoryImageName Rhs);
 	friend FName;
 
+	friend CORE_API void Freeze::ApplyMemoryImageNamePatch(void* NameDst, const FMemoryImageName& Name, const FPlatformTypeLayoutParameters& LayoutParams);
 	friend CORE_API void Freeze::IntrinsicWriteMemoryImage(FMemoryImageWriter& Writer, const FMemoryImageName& Object, const FTypeLayoutDesc&);
 	friend CORE_API uint32 Freeze::IntrinsicUnfrozenCopy(const FMemoryUnfreezeContent& Context, const FMemoryImageName& Object, void* OutDst);
 
@@ -1179,7 +1181,9 @@ namespace Freeze
 	{
 		FNameEntryId	DisplayIndex;
 	};
-	
+
+	CORE_API void ApplyMemoryImageNamePatch(void* NameDst, const FMemoryImageName& Name, const FPlatformTypeLayoutParameters& LayoutParams);
+
 	CORE_API uint32 IntrinsicAppendHash(const FMemoryImageName* DummyObject, const FTypeLayoutDesc& TypeDesc, const FPlatformTypeLayoutParameters& LayoutParams, FSHA1& Hasher);
 	CORE_API void IntrinsicWriteMemoryImage(FMemoryImageWriter& Writer, const FMemoryImageName& Object, const FTypeLayoutDesc&);
 	CORE_API uint32 IntrinsicUnfrozenCopy(const FMemoryUnfreezeContent& Context, const FMemoryImageName& Object, void* OutDst);
diff --git a/Engine/Source/Runtime/RenderCore/Private/GlobalShader.cpp b/Engine/Source/Runtime/RenderCore/Private/GlobalShader.cpp
--- a/Engine/Source/Runtime/RenderCore/Private/GlobalShader.cpp
+++ b/Engine/Source/Runtime/RenderCore/Private/GlobalShader.cpp
@@ -449,6 +449,8 @@ void FGlobalShaderMapId::AppendKeyString(FString& KeyString, const TArray<FShade
 {
 #if WITH_EDITOR
 
+	LayoutParams.AppendKeyString(KeyString);
+
 	{
 		const FSHAHash LayoutHash = Freeze::HashLayout(StaticGetTypeLayoutDesc<FGlobalShaderMapContent>(), LayoutParams);
 		KeyString += TEXT("_");
