CookOnTheFly (platform protocol): Game: Added an alternative way of sending cooked data to the game using platform-specific communication protocol.

Added IPlatformHostCommunication and IPlatformHostSocket interfaces allowing a game running on a target device to communicate with processes on the connected host PC. They are exposed via FPlatformMisc::GetPlatformHostCommunication(). The main feature is the ability to open socket-like connections to the host PC (IPlatformHostSocket), but it also allows launching processes on the host PC (on some platforms).

Not all platforms are required to provide this functionality, so IPlatformHostCommunication::Available() should be checked before using it.

Added a new transport protocol implementation (PlatformTransport derived from ITransport) to allow NetworkPlatformFile to use the above interfaces to send cooked data.

This new transport implementation is activated by running the game with '-filehostip=platform://' command-line argument.

The editor-side functionality will be submitted in a separate change.

diff --git a/Engine/Source/Runtime/Core/Private/GenericPlatform/GenericPlatformMisc.cpp b/Engine/Source/Runtime/Core/Private/GenericPlatform/GenericPlatformMisc.cpp
--- a/Engine/Source/Runtime/Core/Private/GenericPlatform/GenericPlatformMisc.cpp
+++ b/Engine/Source/Runtime/Core/Private/GenericPlatform/GenericPlatformMisc.cpp
@@ -22,6 +22,7 @@
 #include "Misc/ConfigCacheIni.h"
 #include "Misc/App.h"
 #include "GenericPlatform/GenericPlatformChunkInstall.h"
+#include "GenericPlatform/GenericPlatformHostCommunication.h"
 #include "HAL/FileManagerGeneric.h"
 #include "Misc/VarargsHelper.h"
 #include "Misc/SecureHash.h"
@@ -1205,6 +1206,12 @@ IPlatformChunkInstall* FGenericPlatformMisc::GetPlatformChunkInstall()
 	return &Singleton;
 }
 
+IPlatformHostCommunication& FGenericPlatformMisc::GetPlatformHostCommunication()
+{
+	static FPlatformHostCommunicationAutoInit<FGenericPlatformHostCommunication> Singleton;
+	return Singleton;
+}
+
 void GenericPlatformMisc_GetProjectFilePathProjectDir(FString& OutGameDir)
 {
 	// Here we derive the game path from the project file location.
diff --git a/Engine/Source/Runtime/Core/Public/GenericPlatform/GenericPlatformHostCommunication.h b/Engine/Source/Runtime/Core/Public/GenericPlatform/GenericPlatformHostCommunication.h
new file mode 100644
--- /dev/null
+++ b/Engine/Source/Runtime/Core/Public/GenericPlatform/GenericPlatformHostCommunication.h
@@ -0,0 +1,153 @@
+// Copyright Epic Games, Inc. All Rights Reserved.
+
+#pragma once
+
+#include "CoreTypes.h"
+#include "Containers/UnrealString.h"
+#include "GenericPlatform/GenericPlatformHostSocket.h"
+
+
+
+/**
+ * Known custom communication protocols used by the engine.
+ * 
+ * These values may be used as ProtocolIndex parameter when calling OpenConnection.
+ * 
+ * @see IPlatformHostCommunication::OpenConnection
+ */
+enum EHostProtocol
+{
+	CookOnTheFly = 0
+};
+
+
+/**
+ * Interface for communication between the game running on the target device and the connected host pc.
+ * 
+ * It represents a custom communication channel and may not be implemented on all platforms.
+ *
+ * It is meant to be used in development ONLY.
+ */
+class IPlatformHostCommunication
+{
+public:
+
+	/**
+	  * Initialize the host communication system (done by IPlatformFeaturesModule).
+	  */
+	virtual void Initialize() = 0;
+
+	/**
+	 * Shutdown the host communication system (done by IPlatformFeaturesModule).
+	 */
+	virtual void Shutdown() = 0;
+
+	/**
+	 * Check if the host communication system is available and can be used.
+	 * 
+	 * @return True is available, false otherwise.
+	 */
+	virtual bool Available() const = 0;
+
+	/**
+	 * Open a communication channel with the host PC.
+	 * 
+	 * The connected PC is determined in a platform-dependent way.
+	 * 
+	 * Note that opening a connection does not necessarily mean that the host PC is ready to communicate.
+	 * Check the state of the socket before attempting to send/receive data. There can be only one
+	 * connection using any given ProtocolIndex. Attempts to use a ProtocolIndex while there is already
+	 * a connection using it will fail.
+	 * 
+	 * @param ProtocolIndex Arbitrary 0-based number indicating the communication channel (the same number needs to be used on the host PC).
+	 * @param DebugName     Name of this communication channel used for diagnostic purposes.
+	 * @param Version       (optional) Version of the communication protocol (ProtocolIndex) used to verify compatibility with the host PC.
+	 * @param MinVersion    (optional) Minimum version of the communication protocol (ProtocolIndex) that is supported.
+	 * @return              Socket object on success, nullptr otherwise.
+	 * 
+	 * @see CloseConnection, IPlatformHostSocket::GetState
+	 */
+	virtual IPlatformHostSocketPtr OpenConnection(uint32 ProtocolIndex, const FString& DebugName, uint32 Version = 0, uint32 MinVersion = 0) = 0;
+
+	/**
+	 * Close a communication channel previously opened using OpenConnection.
+	 * @param Socket Socket object.
+	 */
+	virtual void CloseConnection(IPlatformHostSocketPtr Socket) = 0;
+
+	/**
+	 * Launch an executable on the connected host PC.
+	 * @param BinaryPath Path to the executable on the PC to launch.
+	 * @param CmdLine    (optional) Command line parameters to pass to the executable.
+	 * @return           True on success, otherwise false (e.g. no permissions or incorrect executable path).
+	 */
+	virtual bool LaunchOnHost(const char* BinaryPath, const char* CmdLine = nullptr) = 0;
+
+	/**
+	 * Virtual destructor.
+	 */
+	virtual ~IPlatformHostCommunication()
+	{
+	}
+};
+
+
+/**
+ * Generic implementation of IPlatformHostCommunication.
+ */
+class FGenericPlatformHostCommunication : public IPlatformHostCommunication
+{
+public:
+
+	virtual void Initialize() override
+	{
+	}
+
+	virtual void Shutdown() override
+	{
+	}
+
+	virtual bool Available() const override
+	{
+		return false;
+	}
+
+	virtual IPlatformHostSocketPtr OpenConnection(uint32 ProtocolIndex, const FString& DebugName, uint32 Version = 0, uint32 MinVersion = 0) override
+	{
+		return nullptr;
+	}
+
+	void CloseConnection(IPlatformHostSocketPtr Socket) override
+	{
+	}
+
+	bool LaunchOnHost(const char* BinaryPath, const char* CmdLine = nullptr) override
+	{
+		return false;
+	}
+};
+
+
+/**
+ * Helper RAII type for auto initializing instances derived from IPlatformHostCommunication.
+ */
+template <typename T>
+struct FPlatformHostCommunicationAutoInit
+{
+	FPlatformHostCommunicationAutoInit()
+	{
+		HostCommunication.Initialize();
+	}
+
+	~FPlatformHostCommunicationAutoInit()
+	{
+		HostCommunication.Shutdown();
+	}
+
+	operator T& ()
+	{
+		return HostCommunication;
+	}
+
+	T HostCommunication;
+};
diff --git a/Engine/Source/Runtime/Core/Public/GenericPlatform/GenericPlatformHostSocket.h b/Engine/Source/Runtime/Core/Public/GenericPlatform/GenericPlatformHostSocket.h
new file mode 100644
--- /dev/null
+++ b/Engine/Source/Runtime/Core/Public/GenericPlatform/GenericPlatformHostSocket.h
@@ -0,0 +1,89 @@
+// Copyright Epic Games, Inc. All Rights Reserved.
+
+#pragma once
+
+#include "CoreTypes.h"
+#include "Templates/SharedPointer.h"
+
+
+
+/**
+ * Interface for sockets supporting direct communication between the game running
+ * on the target device and a connected PC.
+ * 
+ * It represents a custom communication channel and may not be implemented on all platforms.
+ * 
+ * It is meant to be used in development ONLY.
+ * 
+ * @see IPlatformHostCommunication
+ */
+class IPlatformHostSocket
+{
+public:
+
+	/**
+	  * Status values returned by Send and Receive members.
+	  * @see Send, Receive
+	  */
+	enum class EResultNet : uint8
+	{
+		Ok,							// Communication successful.
+		ErrorUnknown,				// Unknown error.
+		ErrorInvalidArgument,		// Incorrect parameters provided to a function (shouldn't happen assuming the socket object is valid).
+		ErrorInvalidConnection,		// Incorrect socket id used (shouldn't happen assuming the socket object is valid).
+		ErrorInterrupted,			// Data transfer interrupted e.g. a networking issue.
+		ErrorHostNotConnected		// Host PC is not connected (not connected yet or has already disconnected).
+	};
+
+	/**
+	  * State of the socket determining its ability to send/receive data.
+	  * @see GetState
+	  */
+	enum class EConnectionState : uint8
+	{
+		Unknown,		// Default state (shouldn't be returned).
+		Created,		// Socket has been created by cannot communicate yet (no host pc connected yet).
+		Connected,		// Socket ready for communication.
+		Disconnected,	// Host PC has disconnected (no communication possible, socket should be closed).
+		Closed,			// Socket has already been closed and shouldn't be used.
+	};
+
+public:
+
+	/**
+	 * Send data to the connected host PC (blocking operation).
+	 * 
+	 * @param Buffer      Data to be sent.
+	 * @param BytesToSend The number of bytes to send.
+	 * @return            Status value indicating error or success.
+	 */
+	virtual EResultNet Send(const void* Buffer, uint64 BytesToSend) = 0;
+
+	/**
+	 * Receive data from the connected host PC (blocking operation).
+	 * 
+	 * @param Buffer         Data to be sent.
+	 * @param BytesToReceive The number of bytes to receive (Buffer has to be large enough).
+	 * @return               Status value indicating error or success.
+	 */
+	virtual EResultNet Receive(void* Buffer, uint64 BytesToReceive) = 0;
+
+	/**
+	 * Get the state of the socket (determines if the host pc is connected and communication is possible).
+	 */
+	virtual EConnectionState GetState() const = 0;
+
+	/**
+	 * Destructor.
+	 */
+	virtual ~IPlatformHostSocket()
+	{
+	}
+};
+
+
+// Type definition for shared references to instances of IPlatformHostSocket.
+typedef TSharedRef<IPlatformHostSocket, ESPMode::ThreadSafe> IPlatformHostSocketRef;
+
+// Type definition for shared pointers to instances of IPlatformHostSocket.
+typedef TSharedPtr<IPlatformHostSocket, ESPMode::ThreadSafe> IPlatformHostSocketPtr;
diff --git a/Engine/Source/Runtime/Core/Public/GenericPlatform/GenericPlatformMisc.h b/Engine/Source/Runtime/Core/Public/GenericPlatform/GenericPlatformMisc.h
--- a/Engine/Source/Runtime/Core/Public/GenericPlatform/GenericPlatformMisc.h
+++ b/Engine/Source/Runtime/Core/Public/GenericPlatform/GenericPlatformMisc.h
@@ -15,6 +15,7 @@ class GenericApplication;
 class IPlatformChunkInstall;
 class IInstallBundleManager;
 class IPlatformCompression;
+class IPlatformHostCommunication;
 struct FGenericCrashContext;
 struct FGenericMemoryWarningContext;
 struct FCustomChunk;
@@ -1154,6 +1155,13 @@ public:
 	 */
 	static IPlatformCompression* GetPlatformCompression();
 
+	/**
+	 * Returns the platform specific interface for communication with processes running on the host pc.
+	 * 
+	 * @return Returns the platform specific host communication interface.
+	 */
+	static IPlatformHostCommunication& GetPlatformHostCommunication();
+
 	/**
 	 * Has the OS execute a command and path pair (such as launch a browser)
 	 *
diff --git a/Engine/Source/Runtime/NetworkFile/Private/NetworkPlatformFile.cpp b/Engine/Source/Runtime/NetworkFile/Private/NetworkPlatformFile.cpp
--- a/Engine/Source/Runtime/NetworkFile/Private/NetworkPlatformFile.cpp
+++ b/Engine/Source/Runtime/NetworkFile/Private/NetworkPlatformFile.cpp
@@ -17,11 +17,13 @@
 #include "Misc/PackageName.h"
 #include "Misc/PathViews.h"
 #include "SocketSubsystem.h"
+#include "GenericPlatform/GenericPlatformHostCommunication.h"
 
 #if ENABLE_HTTP_FOR_NETWORK_FILE
 #include "HTTPTransport.h"
 #endif
 #include "TCPTransport.h"
+#include "PlatformTransport.h"
 
 #include "HAL/IPlatformFileModule.h"
 #include "Templates/UniquePtr.h"
@@ -77,6 +79,20 @@ ITransport *CreateTransportForHostAddress(const FString &HostIp )
 #endif
 	}
 
+	if ( HostIp.StartsWith(TEXT("platform://")))
+	{
+		if (FPlatformMisc::GetPlatformHostCommunication().Available())
+		{
+			return new FPlatformTransport(EHostProtocol::CookOnTheFly, "CookOnTheFly");
+		}
+		else
+		{
+			UE_LOG(LogNetworkPlatformFile, Warning, TEXT("Platform transport (platform://) not supported for this platform."));
+
+			return nullptr;
+		}
+	}
+
 	// no transport specified assuming tcp
 	return new FTCPTransport();
 }
diff --git a/Engine/Source/Runtime/NetworkFile/Private/PlatformTransport.cpp b/Engine/Source/Runtime/NetworkFile/Private/PlatformTransport.cpp
new file mode 100644
--- /dev/null
+++ b/Engine/Source/Runtime/NetworkFile/Private/PlatformTransport.cpp
@@ -0,0 +1,154 @@
+// Copyright Epic Games, Inc. All Rights Reserved.
+
+#include "PlatformTransport.h"
+#include "NetworkMessage.h"
+#include "NetworkPlatformFile.h"
+#include "HAL/PlatformProcess.h"
+#include "GenericPlatform/GenericPlatformHostSocket.h"
+#include "GenericPlatform/GenericPlatformHostCommunication.h"
+
+
+namespace
+{
+	/**
+	 * Socket abstraction required by FNFSMessageHeader that uses IPlatformHostSocket.
+	 */
+	class FSimpleAbstractSocket_PlatformProtocol : public FSimpleAbstractSocket
+	{
+	public:
+
+		FSimpleAbstractSocket_PlatformProtocol(IPlatformHostSocketPtr InHostSocket)
+			: HostSocket(InHostSocket)
+		{
+			check(HostSocket != nullptr);
+		}
+
+		virtual bool Receive(uint8* Results, int32 Size) const override
+		{
+			return (HostSocket->Receive(Results, Size) == IPlatformHostSocket::EResultNet::Ok);
+		}
+
+		virtual bool Send(const uint8* Buffer, int32 Size) const override
+		{
+			return (HostSocket->Send(Buffer, Size) == IPlatformHostSocket::EResultNet::Ok);
+		}
+
+		virtual uint32 GetMagic() const override
+		{
+			return 0x9E2B83C7;
+		}
+
+	private:
+		IPlatformHostSocketPtr HostSocket;
+	};
+}
+
+
+FPlatformTransport::FPlatformTransport(int32 InProtocolIndex, const FString& InProtocolName)
+	: ProtocolIndex(InProtocolIndex)
+	, ProtocolName(InProtocolName)
+	, HostSocket(nullptr)
+{
+}
+
+
+bool FPlatformTransport::Initialize(const TCHAR* InHostIp)
+{
+	check(HostSocket == nullptr);
+
+	IPlatformHostCommunication& HostCommunication = FPlatformMisc::GetPlatformHostCommunication();
+
+	if (!HostCommunication.Available())
+	{
+		return false;
+	}
+
+	HostSocket = HostCommunication.OpenConnection(ProtocolIndex, ProtocolName);
+
+	if (!HostSocket)
+	{
+		return false;
+	}
+
+	UE_LOG(LogNetworkPlatformFile, Display, TEXT("Waiting for the server to accept the connection (custom protocol)..."));
+
+	// We need to wait because this transport is a bit different than typical sockets.
+	// In this case, it's the client (game) that enables the communication by opening the socket.
+	// The server (pc) can connect only if it detects this enabled communication protocol.
+	return WaitUntilConnected();
+}
+
+
+bool FPlatformTransport::SendPayloadAndReceiveResponse(TArray<uint8>& In, TArray<uint8>& Out)
+{
+	bool SendResult = false;
+
+	SendResult = FNFSMessageHeader::WrapAndSendPayload(In, FSimpleAbstractSocket_PlatformProtocol(HostSocket));
+
+	if (!SendResult)
+		return false;
+
+	FArrayReader Response;
+	bool RetResult = false;
+	RetResult = FNFSMessageHeader::ReceivePayload(Response, FSimpleAbstractSocket_PlatformProtocol(HostSocket));
+
+	if (RetResult)
+	{
+		Out.Append(Response.GetData(), Response.Num());
+		return true;
+	}
+
+	return false;
+}
+
+
+bool FPlatformTransport::ReceiveResponse(TArray<uint8>& Out)
+{
+	FArrayReader Response;
+	bool RetResult = true;
+	RetResult &= FNFSMessageHeader::ReceivePayload(Response, FSimpleAbstractSocket_PlatformProtocol(HostSocket));
+
+	if (RetResult)
+	{
+		Out.Append(Response.GetData(), Response.Num());
+		return true;
+	}
+
+	return false;
+}
+
+
+bool FPlatformTransport::WaitUntilConnected()
+{
+	while (true)
+	{
+		IPlatformHostSocket::EConnectionState ConnectionState = HostSocket->GetState();
+
+		// We're interested in the Connected state, but break the loop in other non-default states
+		// in case any error has appeared.
+		if (HostSocket->GetState() != IPlatformHostSocket::EConnectionState::Created)
+		{
+			break;
+		}
+
+		FPlatformProcess::Sleep(0.5f);
+	}
+
+	return HostSocket->GetState() == IPlatformHostSocket::EConnectionState::Connected;
+}
+
+
+FPlatformTransport::~FPlatformTransport()
+{
+	if (HostSocket)
+	{
+		IPlatformHostCommunication& HostCommunication = FPlatformMisc::GetPlatformHostCommunication();
+
+		if (HostCommunication.Available())
+		{
+			HostCommunication.CloseConnection(HostSocket);
+
+			HostSocket = nullptr;
+		}
+	}
+}
diff --git a/Engine/Source/Runtime/NetworkFile/Private/PlatformTransport.h b/Engine/Source/Runtime/NetworkFile/Private/PlatformTransport.h
new file mode 100644
--- /dev/null
+++ b/Engine/Source/Runtime/NetworkFile/Private/PlatformTransport.h
@@ -0,0 +1,39 @@
+// Copyright Epic Games, Inc. All Rights Reserved.
+
+#pragma once
+
+#include "CoreMinimal.h"
+#include "Runtime/NetworkFile/Private/ITransport.h"
+#include "GenericPlatform/GenericPlatformHostSocket.h"
+
+
+/**
+ * Implementation of ITransport using IPlatformHostCommunication/IPlatformHostSocket interfaces (custom target and host pc communication).
+ */
+class FPlatformTransport : public ITransport
+{
+
+public:
+
+	FPlatformTransport(int32 InProtocolIndex, const FString& InProtocolName);
+	~FPlatformTransport();
+
+	//~ Begin ITransport interface
+	virtual bool Initialize(const TCHAR* HostIp) override;
+	virtual bool SendPayloadAndReceiveResponse(TArray<uint8>& In, TArray<uint8>& Out) override;
+	virtual bool ReceiveResponse(TArray<uint8>& Out) override;
+	//~ End ITransport interface
+
+private:
+
+	/**
+	 * Wait until HostSocket is in a non-default state (preferably Connected).
+	 * @return True is the socket is Connected, false otherwise (an error or if the host pc immediately disconnects).
+	 */
+	bool WaitUntilConnected();
+
+	int32				   ProtocolIndex;
+	FString				   ProtocolName;
+	IPlatformHostSocketPtr HostSocket;
+
+};
