Added check for running on ChromeOS to FUnixPlatformMisc::GetOSVersions

diff --git a/Engine/Source/Runtime/Core/Private/Unix/UnixPlatformMisc.cpp b/Engine/Source/Runtime/Core/Private/Unix/UnixPlatformMisc.cpp
--- a/Engine/Source/Runtime/Core/Private/Unix/UnixPlatformMisc.cpp
+++ b/Engine/Source/Runtime/Core/Private/Unix/UnixPlatformMisc.cpp
@@ -458,6 +458,27 @@ namespace
 		return FString();
 	}
 
+	bool IsRunningOnChromiumOS()
+	{
+		bool bRet = false;
+		int SysVendorFile = open("/sys/class/dmi/id/sys_vendor", O_RDONLY);
+
+		if (SysVendorFile >= 0)
+		{
+			char LineBuffer[128];
+			ssize_t Length = read(SysVendorFile, LineBuffer, sizeof(LineBuffer));
+
+			if (Length > 0)
+			{
+				bRet = !FCStringAnsi::Strncmp(LineBuffer, "ChromiumOS", 10);
+			}
+
+			close(SysVendorFile);
+		}
+
+		return bRet;
+	}
+
 	/**
 	 * Reads a Linux style configuration file ignoring # as comments
 	 */
@@ -502,6 +523,7 @@ void FUnixPlatformMisc::GetOSVersions(FString& out_OSVersionLabel, FString& out_
 	out_OSVersionLabel = FString(TEXT("GenericLinuxVersion"));
 	// Get Kernel Version
 	out_OSSubVersionLabel = GetKernelVersion();
+
 	// Get PRETTY_NAME/NAME or redhat-release line
 	TMap<FString, FString> OsInfo = ReadConfigurationFile(TEXT("/etc/os-release"));
 	if (OsInfo.Num() > 0)
@@ -533,6 +555,11 @@ void FUnixPlatformMisc::GetOSVersions(FString& out_OSVersionLabel, FString& out_
 			}
 		}
 	}
+
+	if (IsRunningOnChromiumOS())
+	{
+		out_OSVersionLabel += TEXT(" ChromiumOS");
+	}
 }
 
 FString FUnixPlatformMisc::GetOSVersion()
