Vectorized TTransform made trivially-copyable.
Obsolete IMPLEMENT_ASSIGNMENT_OPERATOR_MANUALLY switch removed from the non-vectorized TTransform and other types.

diff --git a/Engine/Source/Runtime/Core/Public/Math/Quat.h b/Engine/Source/Runtime/Core/Public/Math/Quat.h
--- a/Engine/Source/Runtime/Core/Public/Math/Quat.h
+++ b/Engine/Source/Runtime/Core/Public/Math/Quat.h
@@ -131,15 +131,6 @@ public:
 
 public:
 
-#ifdef IMPLEMENT_ASSIGNMENT_OPERATOR_MANUALLY
-	/**
-	 * Copy another TQuat into this one
-	 *
-	 * @return reference to this TQuat
-	 */
-	FORCEINLINE TQuat<T>& operator=(const TQuat<T>& Other);
-#endif
-
 	/**
 	 * Gets the result of adding a Quaternion to this.
 	 * This is a component-wise addition; composing quaternions should be done via multiplication.
@@ -860,19 +851,6 @@ inline bool TQuat<T>::InitFromString(const FString& InSourceString)
 	return bSuccess;
 }
 
-#ifdef IMPLEMENT_ASSIGNMENT_OPERATOR_MANUALLY
-template<typename T>
-FORCEINLINE TQuat<T>& TQuat<T>::operator=(const TQuat<T>& Other)
-{
-	this->X = Other.X;
-	this->Y = Other.Y;
-	this->Z = Other.Z;
-	this->W = Other.W;
-
-	return *this;
-}
-#endif
-
 template<typename T>
 FORCEINLINE TQuat<T>::TQuat(TVector<T> Axis, T AngleRad)
 {
diff --git a/Engine/Source/Runtime/Core/Public/Math/TransformNonVectorized.h b/Engine/Source/Runtime/Core/Public/Math/TransformNonVectorized.h
--- a/Engine/Source/Runtime/Core/Public/Math/TransformNonVectorized.h
+++ b/Engine/Source/Runtime/Core/Public/Math/TransformNonVectorized.h
@@ -234,20 +234,6 @@ public:
 	/** Acceptable form: "%f,%f,%f|%f,%f,%f|%f,%f,%f" */
 	CORE_API bool InitFromString(const FString& InSourceString);
 
-#ifdef IMPLEMENT_ASSIGNMENT_OPERATOR_MANUALLY
-	/**
-	* Copy another Transform into this one
-	*/
-	FORCEINLINE TTransform<T>& operator=(const TTransform<T>& Other)
-	{
-		this->Rotation = Other.Rotation;
-		this->Translation = Other.Translation;
-		this->Scale3D = Other.Scale3D;
-
-		return *this;
-	}
-#endif
-
 	/**
 	* Convert this Transform to a transformation matrix with scaling.
 	*/
diff --git a/Engine/Source/Runtime/Core/Public/Math/TransformVectorized.h b/Engine/Source/Runtime/Core/Public/Math/TransformVectorized.h
--- a/Engine/Source/Runtime/Core/Public/Math/TransformVectorized.h
+++ b/Engine/Source/Runtime/Core/Public/Math/TransformVectorized.h
@@ -294,18 +294,6 @@ public:
 	/** Acceptable form: "%f,%f,%f|%f,%f,%f|%f,%f,%f" */
 	CORE_API bool InitFromString( const FString& InSourceString );
 
-	/**
-	* Copy another Transform into this one
-	*/
-	FORCEINLINE TTransform& operator=(const TTransform& Other)
-	{
-		this->Rotation = Other.Rotation;
-		this->Translation = Other.Translation;
-		this->Scale3D = Other.Scale3D;
-
-		return *this;
-	}
-
 
 	FORCEINLINE TMatrix<T> ToMatrixWithScale() const
 	{
diff --git a/Engine/Source/Runtime/Core/Public/Math/UnrealMathUtility.h b/Engine/Source/Runtime/Core/Public/Math/UnrealMathUtility.h
--- a/Engine/Source/Runtime/Core/Public/Math/UnrealMathUtility.h
+++ b/Engine/Source/Runtime/Core/Public/Math/UnrealMathUtility.h
@@ -9,8 +9,6 @@
 #include "Templates/IsFloatingPoint.h"
 #include "Templates/IsIntegral.h"
 
-//#define IMPLEMENT_ASSIGNMENT_OPERATOR_MANUALLY
-
 // Assert on non finite numbers. Used to track NaNs.
 #ifndef ENABLE_NAN_DIAGNOSTIC
 	#if UE_BUILD_DEBUG
diff --git a/Engine/Source/Runtime/Core/Public/Math/Vector.h b/Engine/Source/Runtime/Core/Public/Math/Vector.h
--- a/Engine/Source/Runtime/Core/Public/Math/Vector.h
+++ b/Engine/Source/Runtime/Core/Public/Math/Vector.h
@@ -206,16 +206,6 @@ public:
      */
     explicit FORCEINLINE TVector(EForceInit);
 
-#ifdef IMPLEMENT_ASSIGNMENT_OPERATOR_MANUALLY
-    /**
-    * Copy another TVector<T> into this one
-    *
-    * @param Other The other vector.
-    * @return Reference to vector after copy.
-    */
-    FORCEINLINE TVector<T>& operator=(const TVector<T>& Other);
-#endif
-
     /**
      * Calculate cross product between this and another vector.
      *
@@ -1399,20 +1389,6 @@ FORCEINLINE TVector<T>::TVector(EForceInit)
     DiagnosticCheckNaN();
 }
 
-#ifdef IMPLEMENT_ASSIGNMENT_OPERATOR_MANUALLY
-template<typename T>
-FORCEINLINE TVector<T>& TVector<T>::operator=(const TVector<T>& Other)
-{
-    this->X = Other.X;
-    this->Y = Other.Y;
-    this->Z = Other.Z;
-
-    DiagnosticCheckNaN();
-
-    return *this;
-}
-#endif
-
 template<typename T>
 FORCEINLINE TVector<T> TVector<T>::operator^(const TVector<T>& V) const
 {
