Chaos Cloth Optimization - Store P and InvM together as a 4 vector.

During the main constraint update, P and InvM for a particle are almost always accessed together. By putting them together in memory, this makes them 64-byte aligned and faster to read/write. Currently only ISPC code makes full use of PandInvM.

I added a new Chaos::Softs::FSolverParticles type which holds the PandInvM array. This made code update easy since Chaos::Softs::FSolverParticles was previously an alias for TPBDParticles<FSolverReal,3>. InvM still has its single array representation -- it is combined with P when P is updated at the beginning of each frame. This choice is mostly to reduce the amount of code that I had to change so that InvM can continue to live in TDynamicParticles. As such, when accessing InvM() by index, you will get the version that lives in TDynamicParticles. You can only access the value that lives with PandInvM  by accessing PandInvM() by index.

Overall per frame simulation for the Owen asset dropped from ~1.6ms per frame to ~1.4ms with this change. This comes from the Constraint updates where PandInvM is heavily leveraged.

Tested with the Owen asset for profiling and the Cloth ContentExamples map in editor for correctness.

diff --git a/Engine/Plugins/Experimental/ChaosCloth/Source/Chaos/Private/ChaosCloth/ChaosClothingSimulationCloth.cpp b/Engine/Plugins/Experimental/ChaosCloth/Source/Chaos/Private/ChaosCloth/ChaosClothingSimulationCloth.cpp
--- a/Engine/Plugins/Experimental/ChaosCloth/Source/Chaos/Private/ChaosCloth/ChaosClothingSimulationCloth.cpp
+++ b/Engine/Plugins/Experimental/ChaosCloth/Source/Chaos/Private/ChaosCloth/ChaosClothingSimulationCloth.cpp
@@ -726,9 +726,9 @@ void FClothingSimulationCloth::Update(FClothingSimulationSolver* Solver)
 				PrevLODIndex,
 				LODIndex,
 				Solver->GetNormals(PrevOffset),
-				Solver->GetParticlePs(PrevOffset),
+				Solver->GetParticlePandInvMs(PrevOffset),
 				Solver->GetParticleVs(PrevOffset),
-				Solver->GetParticlePs(Offset),
+				Solver->GetParticlePandInvMs(Offset),
 				Solver->GetParticleXs(Offset),
 				Solver->GetParticleVs(Offset));
 
@@ -855,14 +855,6 @@ TConstArrayView<Softs::FSolverVec3> FClothingSimulationCloth::GetParticlePositio
 	return TConstArrayView<Softs::FSolverVec3>(Solver->GetParticleXs(GetOffset(Solver, LODIndex)), GetNumParticles(LODIndex));
 }
 
-TConstArrayView<Softs::FSolverVec3> FClothingSimulationCloth::GetParticleOldPositions(const FClothingSimulationSolver* Solver) const
-{
-	check(Solver);
-	const int32 LODIndex = LODIndices.FindChecked(Solver);
-	check(GetOffset(Solver, LODIndex) != INDEX_NONE);
-	return TConstArrayView<Softs::FSolverVec3>(Solver->GetParticlePs(GetOffset(Solver, LODIndex)), GetNumParticles(LODIndex));
-}
-
 TConstArrayView<Softs::FSolverVec3> FClothingSimulationCloth::GetParticleNormals(const FClothingSimulationSolver* Solver) const
 {
 	check(Solver);
diff --git a/Engine/Plugins/Experimental/ChaosCloth/Source/Chaos/Private/ChaosCloth/ChaosClothingSimulationMesh.cpp b/Engine/Plugins/Experimental/ChaosCloth/Source/Chaos/Private/ChaosCloth/ChaosClothingSimulationMesh.cpp
--- a/Engine/Plugins/Experimental/ChaosCloth/Source/Chaos/Private/ChaosCloth/ChaosClothingSimulationMesh.cpp
+++ b/Engine/Plugins/Experimental/ChaosCloth/Source/Chaos/Private/ChaosCloth/ChaosClothingSimulationMesh.cpp
@@ -3,6 +3,7 @@
 #include "ChaosCloth/ChaosClothingSimulationSolver.h"
 #include "ChaosCloth/ChaosWeightMapTarget.h"
 #include "ChaosCloth/ChaosClothPrivate.h"
+#include "Chaos/PBDSoftsSolverParticles.h" // Needed for PAndInvM in WrapDeformLOD
 #include "ClothingSimulation.h"
 #include "ClothingAsset.h"
 #include "Containers/ArrayView.h"
@@ -187,12 +188,12 @@ bool FClothingSimulationMesh::WrapDeformLOD(
 bool FClothingSimulationMesh::WrapDeformLOD(
 	int32 PrevLODIndex,
 	int32 LODIndex,
-	const FSolverVec3* Normals,
-	const FSolverVec3* Positions,
-	const FSolverVec3* Velocities,
-	FSolverVec3* OutPositions0,
-	FSolverVec3* OutPositions1,
-	FSolverVec3* OutVelocities) const
+	const Softs::FSolverVec3* Normals,
+	const Softs::FPAndInvM* PositionAndInvMs,
+	const Softs::FSolverVec3* Velocities,
+	Softs::FPAndInvM* OutPositionAndInvMs0,
+	Softs::FSolverVec3* OutPositions1,
+	Softs::FSolverVec3* OutVelocities) const
 {
 	TRACE_CPUPROFILER_EVENT_SCOPE(FClothingSimulationMesh_WrapDeformLOD);
 	SCOPE_CYCLE_COUNTER(STAT_ChaosClothWrapDeformClothLOD);
@@ -217,10 +218,10 @@ bool FClothingSimulationMesh::WrapDeformLOD(
 		const int32 VertIndex1 = (int32)VertData.SourceMeshVertIndices[1];
 		const int32 VertIndex2 = (int32)VertData.SourceMeshVertIndices[2];
 
-		OutPositions0[Index] = OutPositions1[Index] =
-			Positions[VertIndex0] * (FSolverReal)VertData.PositionBaryCoordsAndDist.X + Normals[VertIndex0] * (FSolverReal)VertData.PositionBaryCoordsAndDist.W +
-			Positions[VertIndex1] * (FSolverReal)VertData.PositionBaryCoordsAndDist.Y + Normals[VertIndex1] * (FSolverReal)VertData.PositionBaryCoordsAndDist.W +
-			Positions[VertIndex2] * (FSolverReal)VertData.PositionBaryCoordsAndDist.Z + Normals[VertIndex2] * (FSolverReal)VertData.PositionBaryCoordsAndDist.W;
+		OutPositionAndInvMs0[Index].P = OutPositions1[Index] =
+			PositionAndInvMs[VertIndex0].P * (FSolverReal)VertData.PositionBaryCoordsAndDist.X + Normals[VertIndex0] * (FSolverReal)VertData.PositionBaryCoordsAndDist.W +
+			PositionAndInvMs[VertIndex1].P * (FSolverReal)VertData.PositionBaryCoordsAndDist.Y + Normals[VertIndex1] * (FSolverReal)VertData.PositionBaryCoordsAndDist.W +
+			PositionAndInvMs[VertIndex2].P * (FSolverReal)VertData.PositionBaryCoordsAndDist.Z + Normals[VertIndex2] * (FSolverReal)VertData.PositionBaryCoordsAndDist.W;
 
 		OutVelocities[Index] = 
 			Velocities[VertIndex0] * (FSolverReal)VertData.PositionBaryCoordsAndDist.X +
diff --git a/Engine/Plugins/Experimental/ChaosCloth/Source/Chaos/Private/ChaosCloth/ChaosClothingSimulationSolver.cpp b/Engine/Plugins/Experimental/ChaosCloth/Source/Chaos/Private/ChaosCloth/ChaosClothingSimulationSolver.cpp
--- a/Engine/Plugins/Experimental/ChaosCloth/Source/Chaos/Private/ChaosCloth/ChaosClothingSimulationSolver.cpp
+++ b/Engine/Plugins/Experimental/ChaosCloth/Source/Chaos/Private/ChaosCloth/ChaosClothingSimulationSolver.cpp
@@ -335,7 +335,7 @@ void FClothingSimulationSolver::EnableParticles(int32 Offset, bool bEnable)
 
 void FClothingSimulationSolver::ResetStartPose(int32 Offset, int32 NumParticles)
 {
-	Softs::FSolverVec3* const Ps = GetParticlePs(Offset);
+	Softs::FPAndInvM* const PandInvMs = GetParticlePandInvMs(Offset);
 	Softs::FSolverVec3* const Xs = GetParticleXs(Offset);
 	Softs::FSolverVec3* const Vs = GetParticleVs(Offset);
 	const Softs::FSolverVec3* const Positions = GetAnimationPositions(Offset);
@@ -343,19 +343,19 @@ void FClothingSimulationSolver::ResetStartPose(int32 Offset, int32 NumParticles)
 
 	for (int32 Index = 0; Index < NumParticles; ++Index)
 	{
-		Ps[Index] = Xs[Index] = OldPositions[Index] = Positions[Index];
+		PandInvMs[Index].P = Xs[Index] = OldPositions[Index] = Positions[Index];
 		Vs[Index] = Softs::FSolverVec3(0.);
 	}
 }
 
-const Softs::FSolverVec3* FClothingSimulationSolver::GetParticlePs(int32 Offset) const
+const Softs::FPAndInvM* FClothingSimulationSolver::GetParticlePandInvMs(int32 Offset) const
 {
-	return &Evolution->Particles().P(Offset);
+	return &Evolution->Particles().PAndInvM(Offset);
 }
 
-Softs::FSolverVec3* FClothingSimulationSolver::GetParticlePs(int32 Offset)
+Softs::FPAndInvM* FClothingSimulationSolver::GetParticlePandInvMs(int32 Offset)
 {
-	return &Evolution->Particles().P(Offset);
+	return &Evolution->Particles().PAndInvM(Offset);
 }
 
 const Softs::FSolverVec3* FClothingSimulationSolver::GetParticleXs(int32 Offset) const
@@ -784,10 +784,11 @@ void FClothingSimulationSolver::ApplyPreSimulationTransforms()
 			if (bRealTypeCompatibleWithISPC && bChaos_PreSimulationTransforms_ISPC_Enabled)  // TODO: Make the ISPC works with both Single and Double depending on the FSolverReal type
 			{
 				ispc::ApplyPreSimulationTransforms(
-					(ispc::FVector3f*)Particles.GetP().GetData(),
+					(ispc::FVector4f*)Particles.GetPAndInvM().GetData(),
 					(ispc::FVector3f*)Particles.GetV().GetData(),
 					(ispc::FVector3f*)Particles.XArray().GetData(),
 					(ispc::FVector3f*)OldAnimationPositions.GetData(),
+					Particles.GetInvM().GetData(),
 					ParticleGroupIds.GetData(),
 					(ispc::FTransform3f*)PreSimulationTransforms.GetData(),
 					(ispc::FVector3f&)DeltaLocalSpaceLocation,
@@ -807,6 +808,9 @@ void FClothingSimulationSolver::ApplyPreSimulationTransforms()
 						Particles.P(Index) = Particles.X(Index) = GroupSpaceTransform.TransformPositionNoScale(Particles.X(Index)) - DeltaLocalSpaceLocation;
 						Particles.V(Index) = GroupSpaceTransform.TransformVector(Particles.V(Index));
 
+						// Copy InvM over to PAndInvM
+						Particles.PAndInvM(Index).InvM = Particles.InvM(Index);
+
 						// Update anim initial state (target updated by skinning)
 						OldAnimationPositions[Index] = GroupSpaceTransform.TransformPositionNoScale(OldAnimationPositions[Index]) - DeltaLocalSpaceLocation;
 					}, RangeSize < ClothSolverMinParallelBatchSize);
diff --git a/Engine/Plugins/Experimental/ChaosCloth/Source/Chaos/Private/ChaosCloth/ChaosClothingSimulationSolver.ispc b/Engine/Plugins/Experimental/ChaosCloth/Source/Chaos/Private/ChaosCloth/ChaosClothingSimulationSolver.ispc
--- a/Engine/Plugins/Experimental/ChaosCloth/Source/Chaos/Private/ChaosCloth/ChaosClothingSimulationSolver.ispc
+++ b/Engine/Plugins/Experimental/ChaosCloth/Source/Chaos/Private/ChaosCloth/ChaosClothingSimulationSolver.ispc
@@ -5,10 +5,11 @@
 #include "Math/Vector.isph"
 #include "Math/Transform.isph"
 
-export void ApplyPreSimulationTransforms(uniform FVector3f ParticlesP[],
+export void ApplyPreSimulationTransforms(uniform FVector4f ParticlesPandInvM[],
 										uniform FVector3f ParticlesV[],
 										uniform FVector3f ParticlesX[],
 										uniform FVector3f OldAnimationPositions[],
+										const uniform float InvM[],
 										const uniform uint32 ParticleGroupIds[],
 										const uniform FTransform3f PreSimulationTransforms[],
 										const uniform FVector3f &DeltaLocalSpaceLocation,
@@ -25,7 +26,8 @@ export void ApplyPreSimulationTransforms(uniform FVector3f ParticlesP[],
 
 		const varying FVector3f ParticleX = VectorLoad(&ParticlesX[Index]);
 		const varying FVector3f PX = TransformPosition(GroupSpaceTransform, ParticleX) - DeltaLocalSpaceLocation;
-		VectorStore(&ParticlesP[Index], PX);
+
+		VectorStore(&ParticlesPandInvM[Index], SetVector4(PX, InvM[i]));
 		VectorStore(&ParticlesX[Index], PX);
 
 		const varying FVector3f ParticleV = VectorLoad(&ParticlesV[Index]);
diff --git a/Engine/Plugins/Experimental/ChaosCloth/Source/Chaos/Public/ChaosCloth/ChaosClothingSimulationCloth.h b/Engine/Plugins/Experimental/ChaosCloth/Source/Chaos/Public/ChaosCloth/ChaosClothingSimulationCloth.h
--- a/Engine/Plugins/Experimental/ChaosCloth/Source/Chaos/Public/ChaosCloth/ChaosClothingSimulationCloth.h
+++ b/Engine/Plugins/Experimental/ChaosCloth/Source/Chaos/Public/ChaosCloth/ChaosClothingSimulationCloth.h
@@ -109,8 +109,6 @@ namespace Chaos
 		TConstArrayView<Softs::FSolverVec3> GetAnimationNormals(const FClothingSimulationSolver* Solver) const;
 		// Return the solver's positions for this cloth current LOD, not thread safe, call must be done right after the solver update.
 		TConstArrayView<Softs::FSolverVec3> GetParticlePositions(const FClothingSimulationSolver* Solver) const;
-		// Return the solver's previous frame positions for this cloth current LOD, not thread safe, call must be done right after the solver update.
-		TConstArrayView<Softs::FSolverVec3> GetParticleOldPositions(const FClothingSimulationSolver* Solver) const;
 		// Return the solver's normals for this cloth current LOD, not thread safe, call must be done right after the solver update.
 		TConstArrayView<Softs::FSolverVec3> GetParticleNormals(const FClothingSimulationSolver* Solver) const;
 		// Return the solver's normals for this cloth current LOD, not thread safe, call must be done right after the solver update.
diff --git a/Engine/Plugins/Experimental/ChaosCloth/Source/Chaos/Public/ChaosCloth/ChaosClothingSimulationMesh.h b/Engine/Plugins/Experimental/ChaosCloth/Source/Chaos/Public/ChaosCloth/ChaosClothingSimulationMesh.h
--- a/Engine/Plugins/Experimental/ChaosCloth/Source/Chaos/Public/ChaosCloth/ChaosClothingSimulationMesh.h
+++ b/Engine/Plugins/Experimental/ChaosCloth/Source/Chaos/Public/ChaosCloth/ChaosClothingSimulationMesh.h
@@ -53,9 +53,9 @@ namespace Chaos
 			int32 PrevLODIndex,
 			int32 LODIndex,
 			const Softs::FSolverVec3* Normals,
-			const Softs::FSolverVec3* Positions,
+			const Softs::FPAndInvM* PositionAndInvMs,
 			const Softs::FSolverVec3* Velocities,
-			Softs::FSolverVec3* OutPositions0,
+			Softs::FPAndInvM* OutPositionAndInvMs0,
 			Softs::FSolverVec3* OutPositions1,
 			Softs::FSolverVec3* OutVelocities) const;
 		// ---- End of the Cloth interface ----
diff --git a/Engine/Plugins/Experimental/ChaosCloth/Source/Chaos/Public/ChaosCloth/ChaosClothingSimulationSolver.h b/Engine/Plugins/Experimental/ChaosCloth/Source/Chaos/Public/ChaosCloth/ChaosClothingSimulationSolver.h
--- a/Engine/Plugins/Experimental/ChaosCloth/Source/Chaos/Public/ChaosCloth/ChaosClothingSimulationSolver.h
+++ b/Engine/Plugins/Experimental/ChaosCloth/Source/Chaos/Public/ChaosCloth/ChaosClothingSimulationSolver.h
@@ -117,8 +117,8 @@ namespace Chaos
 		Softs::FSolverVec3* GetAnimationNormals(int32 Offset) { return AnimationNormals.GetData() + Offset; }
 		const Softs::FSolverVec3* GetNormals(int32 Offset) const { return Normals.GetData() + Offset; }
 		Softs::FSolverVec3* GetNormals(int32 Offset) { return Normals.GetData() + Offset; }
-		const Softs::FSolverVec3* GetParticlePs(int32 Offset) const;
-		Softs::FSolverVec3* GetParticlePs(int32 Offset);
+		const Softs::FPAndInvM* GetParticlePandInvMs(int32 Offset) const;
+		Softs::FPAndInvM* GetParticlePandInvMs(int32 Offset);
 		const Softs::FSolverVec3* GetParticleXs(int32 Offset) const;
 		Softs::FSolverVec3* GetParticleXs(int32 Offset);
 		const Softs::FSolverVec3* GetParticleVs(int32 Offset) const;
diff --git a/Engine/Source/Runtime/Core/Public/Math/Vector.isph b/Engine/Source/Runtime/Core/Public/Math/Vector.isph
--- a/Engine/Source/Runtime/Core/Public/Math/Vector.isph
+++ b/Engine/Source/Runtime/Core/Public/Math/Vector.isph
@@ -437,13 +437,13 @@ inline void VectorStore(uniform FVector3f *uniform DstPtr, const varying FVector
 	}
 }
 
-inline varying FVector4 VectorLoad(const uniform FVector4 *uniform SrcPtr)
+inline varying FVector4d VectorLoad(const uniform FVector4d *uniform SrcPtr)
 {
-	varying FVector4 Result;
+	varying FVector4d Result;
 
 	if(((1<<TARGET_WIDTH)-1 ^ lanemask()) == 0)
 	{
-		aos_to_soa4_ispc((uniform FReal * uniform)SrcPtr, &Result.V[0], &Result.V[1], &Result.V[2], &Result.V[3]);
+		aos_to_soa4_ispc((uniform double * uniform)SrcPtr, &Result.V[0], &Result.V[1], &Result.V[2], &Result.V[3]);
 	}
 	else
 	{
@@ -454,11 +454,41 @@ inline varying FVector4 VectorLoad(const uniform FVector4 *uniform SrcPtr)
 	return Result;
 }
 
-inline void VectorStore(uniform FVector4 *uniform DstPtr, const varying FVector4 &V)
+inline varying FVector4f VectorLoad(const uniform FVector4f *uniform SrcPtr)
+{
+	varying FVector4f Result;
+
+	if(((1<<TARGET_WIDTH)-1 ^ lanemask()) == 0)
+	{
+		aos_to_soa4_ispc((uniform float * uniform)SrcPtr, &Result.V[0], &Result.V[1], &Result.V[2], &Result.V[3]);
+	}
+	else
+	{
+		#pragma ignore warning(perf)
+		Result = SrcPtr[programIndex];
+	}
+
+	return Result;
+}
+
+inline void VectorStore(uniform FVector4d *uniform DstPtr, const varying FVector4d &V)
 {
 	if(((1<<TARGET_WIDTH)-1 ^ lanemask()) == 0)
 	{
-		soa_to_aos4_ispc(V.V[0], V.V[1], V.V[2], V.V[3], (uniform FReal *uniform)DstPtr);
+		soa_to_aos4_ispc(V.V[0], V.V[1], V.V[2], V.V[3], (uniform double *uniform)DstPtr);
+	}
+	else
+	{
+		#pragma ignore warning(perf)
+		DstPtr[programIndex] = V;
+	}
+}
+
+inline void VectorStore(uniform FVector4f *uniform DstPtr, const varying FVector4f &V)
+{
+	if(((1<<TARGET_WIDTH)-1 ^ lanemask()) == 0)
+	{
+		soa_to_aos4_ispc(V.V[0], V.V[1], V.V[2], V.V[3], (uniform float *uniform)DstPtr);
 	}
 	else
 	{
@@ -515,20 +545,69 @@ inline void VectorScatter(uniform FVector3f *varying DstPtr, const varying FVect
 	*DstPtr = V;
 }
 
-inline varying FVector4 VectorGather(const uniform FVector4 *varying SrcPtr)
+inline varying FVector4d VectorGather(const uniform FVector4d *varying SrcPtr)
+{
+	varying FVector4d Result;
+
+	if(((1<<TARGET_WIDTH)-1 ^ lanemask()) == 0)
+	{
+		unmasked
+		{
+			uniform FVector4d *uniform Src[programCount];
+			uniform FVector4d *uniform DstPtr = (uniform FVector4d *uniform)&Result;
+
+			foreach(i = 0 ... programCount)
+			{
+				Src[i] = (uniform FVector4d *varying)&SrcPtr[0];
+			}
+
+			DstPtr[0] = *Src[0];
+			DstPtr[1] = *Src[1];
+			DstPtr[2] = *Src[2];
+			DstPtr[3] = *Src[3];
+#if TARGET_WIDTH == 8 || TARGET_WIDTH == 16
+			DstPtr[4] = *Src[4];
+			DstPtr[5] = *Src[5];
+			DstPtr[6] = *Src[6];
+			DstPtr[7] = *Src[7];
+#endif
+#if TARGET_WIDTH == 16
+			DstPtr[8] = *Src[8];
+			DstPtr[9] = *Src[9];
+			DstPtr[10] = *Src[10];
+			DstPtr[11] = *Src[11];
+			DstPtr[12] = *Src[12];
+			DstPtr[13] = *Src[13];
+			DstPtr[14] = *Src[14];
+			DstPtr[15] = *Src[15];
+#endif
+
+			aos_to_soa4_ispc((uniform double * uniform)&Result, &Result.V[0], &Result.V[1], &Result.V[2], &Result.V[3]);
+		}
+	}
+	else
+	{
+		#pragma ignore warning(perf)
+		Result = *SrcPtr;
+	}
+
+	return Result;
+}
+
+inline varying FVector4f VectorGather(const uniform FVector4f *varying SrcPtr)
 {
-	varying FVector4 Result;
+	varying FVector4f Result;
 
 	if(((1<<TARGET_WIDTH)-1 ^ lanemask()) == 0)
 	{
 		unmasked
 		{
-			uniform FVector4 *uniform Src[programCount];
-			uniform FVector4 *uniform DstPtr = (uniform FVector4 *uniform)&Result;
+			uniform FVector4f *uniform Src[programCount];
+			uniform FVector4f *uniform DstPtr = (uniform FVector4f *uniform)&Result;
 
 			foreach(i = 0 ... programCount)
 			{
-				Src[i] = (uniform FVector4 *varying)&SrcPtr[0];
+				Src[i] = (uniform FVector4f *varying)&SrcPtr[0];
 			}
 
 			DstPtr[0] = *Src[0];
@@ -552,7 +631,7 @@ inline varying FVector4 VectorGather(const uniform FVector4 *varying SrcPtr)
 			DstPtr[15] = *Src[15];
 #endif
 
-			aos_to_soa4_ispc((uniform FReal * uniform)&Result, &Result.V[0], &Result.V[1], &Result.V[2], &Result.V[3]);
+			aos_to_soa4_ispc((uniform float * uniform)&Result, &Result.V[0], &Result.V[1], &Result.V[2], &Result.V[3]);
 		}
 	}
 	else
@@ -564,20 +643,65 @@ inline varying FVector4 VectorGather(const uniform FVector4 *varying SrcPtr)
 	return Result;
 }
 
-inline void VectorScatter(uniform FVector4 *varying DstPtr, const varying FVector4 &V)
+inline void VectorScatter(uniform FVector4d *varying DstPtr, const varying FVector4d &V)
+{
+	if(((1<<TARGET_WIDTH)-1 ^ lanemask()) == 0)
+	{
+		unmasked
+		{
+			uniform FVector4d Result[programCount];
+			soa_to_aos4_ispc(V.V[0], V.V[1], V.V[2], V.V[3], (uniform double *uniform)&Result);
+
+			uniform FVector4d *uniform Dst[programCount];
+
+			foreach(i = 0 ... programCount)
+			{
+				Dst[i] = (uniform FVector4d *varying)&DstPtr[0];
+			}
+
+			*Dst[0] = Result[0];
+			*Dst[1] = Result[1];
+			*Dst[2] = Result[2];
+			*Dst[3] = Result[3];
+#if TARGET_WIDTH == 8 || TARGET_WIDTH == 16
+			*Dst[4] = Result[4];
+			*Dst[5] = Result[5];
+			*Dst[6] = Result[6];
+			*Dst[7] = Result[7];
+#endif
+#if TARGET_WIDTH == 16
+			*Dst[8] = Result[8];
+			*Dst[9] = Result[9];
+			*Dst[10] = Result[10];
+			*Dst[11] = Result[11];
+			*Dst[12] = Result[12];
+			*Dst[13] = Result[13];
+			*Dst[14] = Result[14];
+			*Dst[15] = Result[15];
+#endif
+		}
+	}
+	else
+	{
+		#pragma ignore warning(perf)
+		*DstPtr = V;
+	}
+}
+
+inline void VectorScatter(uniform FVector4f *varying DstPtr, const varying FVector4f &V)
 {
 	if(((1<<TARGET_WIDTH)-1 ^ lanemask()) == 0)
 	{
 		unmasked
 		{
-			uniform FVector4 Result[programCount];
-			soa_to_aos4_ispc(V.V[0], V.V[1], V.V[2], V.V[3], (uniform FReal *uniform)&Result);
+			uniform FVector4f Result[programCount];
+			soa_to_aos4_ispc(V.V[0], V.V[1], V.V[2], V.V[3], (uniform float *uniform)&Result);
 
-			uniform FVector4 *uniform Dst[programCount];
+			uniform FVector4f *uniform Dst[programCount];
 
 			foreach(i = 0 ... programCount)
 			{
-				Dst[i] = (uniform FVector4 *varying)&DstPtr[0];
+				Dst[i] = (uniform FVector4f *varying)&DstPtr[0];
 			}
 
 			*Dst[0] = Result[0];
@@ -1111,9 +1235,15 @@ inline uniform FVector4f SetVector4(const uniform float X, const uniform float Y
 	return Result;
 }
 
-inline FVector4 SetVector4(const FVector &V, const FReal W)
+inline FVector4d SetVector4(const FVector3d &V, const FReal W)
+{
+	const FVector4d Result = {{V.V[0], V.V[1], V.V[2], W}};
+	return Result;
+}
+
+inline FVector4f SetVector4(const FVector3f &V, const FReal W)
 {
-	const FVector4 Result = {{V.V[0], V.V[1], V.V[2], W}};
+	const FVector4f Result = {{V.V[0], V.V[1], V.V[2], W}};
 	return Result;
 }
 
diff --git a/Engine/Source/Runtime/Experimental/Chaos/Private/Chaos/PBDAxialSpringConstraints.cpp b/Engine/Source/Runtime/Experimental/Chaos/Private/Chaos/PBDAxialSpringConstraints.cpp
--- a/Engine/Source/Runtime/Experimental/Chaos/Private/Chaos/PBDAxialSpringConstraints.cpp
+++ b/Engine/Source/Runtime/Experimental/Chaos/Private/Chaos/PBDAxialSpringConstraints.cpp
@@ -13,7 +13,7 @@ DECLARE_CYCLE_STAT(TEXT("Chaos PBD Axial Spring Constraint"), STAT_PBD_AxialSpri
 bool bChaos_AxialSpring_ISPC_Enabled = true;
 FAutoConsoleVariableRef CVarChaosAxialSpringISPCEnabled(TEXT("p.Chaos.AxialSpring.ISPC"), bChaos_AxialSpring_ISPC_Enabled, TEXT("Whether to use ISPC optimizations in AxialSpring constraints"));
 
-static_assert(sizeof(ispc::FVector3f) == sizeof(Chaos::Softs::FSolverVec3), "sizeof(ispc::FVector3f) != sizeof(Chaos::Softs::FSolverVec3");
+static_assert(sizeof(ispc::FVector4f) == sizeof(Chaos::Softs::FPAndInvM), "sizeof(ispc::FVector4f) != sizeof(Chaos::Softs::FPAndInvM");
 static_assert(sizeof(ispc::FIntVector) == sizeof(Chaos::TVec3<int32>), "sizeof(ispc::FIntVector) != sizeof(Chaos::TVec3<int32>");
 #endif
 
@@ -108,9 +108,8 @@ void FPBDAxialSpringConstraints::Apply(FSolverParticles& Particles, const FSolve
 					const int32 ColorStart = ConstraintsPerColorStartIndex[ConstraintColorIndex];
 					const int32 ColorSize = ConstraintsPerColorStartIndex[ConstraintColorIndex + 1] - ColorStart;
 					ispc::ApplyAxialSpringConstraints(
-						(ispc::FVector3f*)&Particles.GetP()[0],
+						(ispc::FVector4f*)Particles.GetPAndInvM().GetData(),
 						(ispc::FIntVector*)&Constraints.GetData()[ColorStart],
-						&Particles.GetInvM().GetData()[0],
 						&Barys.GetData()[ColorStart],
 						&Dists.GetData()[ColorStart],
 						ExpStiffnessValue,
@@ -142,9 +141,8 @@ void FPBDAxialSpringConstraints::Apply(FSolverParticles& Particles, const FSolve
 					const int32 ColorStart = ConstraintsPerColorStartIndex[ConstraintColorIndex];
 					const int32 ColorSize = ConstraintsPerColorStartIndex[ConstraintColorIndex + 1] - ColorStart;
 					ispc::ApplyAxialSpringConstraintsWithWeightMaps(
-						(ispc::FVector3f*)&Particles.GetP()[0],
+						(ispc::FVector4f*)Particles.GetPAndInvM().GetData(),
 						(ispc::FIntVector*)&Constraints.GetData()[ColorStart],
-						&Particles.GetInvM().GetData()[0],
 						&Barys.GetData()[ColorStart],
 						&Dists.GetData()[ColorStart],
 						&Stiffness.GetIndices().GetData()[ColorStart],
diff --git a/Engine/Source/Runtime/Experimental/Chaos/Private/Chaos/PBDAxialSpringConstraints.ispc b/Engine/Source/Runtime/Experimental/Chaos/Private/Chaos/PBDAxialSpringConstraints.ispc
--- a/Engine/Source/Runtime/Experimental/Chaos/Private/Chaos/PBDAxialSpringConstraints.ispc
+++ b/Engine/Source/Runtime/Experimental/Chaos/Private/Chaos/PBDAxialSpringConstraints.ispc
@@ -3,10 +3,10 @@
 #define EXPLICIT_VECTOR4 1
 
 #include "Math/Vector.isph"
+#include "Chaos/PBDSofts.isph"
 
-export void ApplyAxialSpringConstraints(uniform FVector3f P[],
+export void ApplyAxialSpringConstraints(uniform FVector4f PandInvM[],
 										const uniform FIntVector Constraints[],
-										const uniform float InvM[],
 										const uniform float Barys[],
 										const uniform float Dists[],
 										const uniform float Stiffness,
@@ -24,10 +24,16 @@ export void ApplyAxialSpringConstraints(uniform FVector3f P[],
 		const varying int32 i2 = Constraint.V[1];
 		const varying int32 i3 = Constraint.V[2];
 
-		const varying FVector3f P2 = VectorGather(&P[i2]);
-		const varying FVector3f P3 = VectorGather(&P[i3]);
+		const varying FVector4f PandInvM2 = VectorGather(&PandInvM[i2]);
+		const varying FVector4f PandInvM3 = VectorGather(&PandInvM[i3]);
 
-		const varying FVector3f P1 = VectorGather(&P[i1]);
+		varying FVector3f P1, P2, P3;
+		varying float M1, M2, M3;
+		UnzipPandInvM(PandInvM2, P2, M2);
+		UnzipPandInvM(PandInvM3, P3, M3);
+
+		const varying FVector4f PandInvM1 = VectorGather(&PandInvM[i1]);
+		UnzipPandInvM(PandInvM1, P1, M1);
 
 		const FVector3f P0 = (P2 - P3) * Bary + P3;
 		const FVector3f Difference = P1 - P0;
@@ -35,13 +41,6 @@ export void ApplyAxialSpringConstraints(uniform FVector3f P[],
 		const float Distance = VectorSize(Difference);
 		const float Multiplier = 2 / (max(Bary, 1 - Bary) + 1);
 
-		#pragma ignore warning(perf)
-		const varying float M1 = InvM[i1];
-		#pragma ignore warning(perf)
-		const varying float M2 = InvM[i2];
-		#pragma ignore warning(perf)
-		const varying float M3 = InvM[i3];
-
 		const FVector3f Direction = Difference / Distance;
 		const FVector3f InnerDelta = (Distance - Dist) * Direction;
 
@@ -53,22 +52,21 @@ export void ApplyAxialSpringConstraints(uniform FVector3f P[],
 
 		if (M1 > 0)
 		{
-			VectorScatter(&P[i1], P1 - (Multiplier * M1 * Delta));
+			VectorScatter(&PandInvM[i1], SetVector4(P1 - (Multiplier * M1 * Delta), M1));
 		}
 		if (M2)
 		{
-			VectorScatter(&P[i2], P2 + (Multiplier * M2 * Bary * Delta));
+			VectorScatter(&PandInvM[i2], SetVector4(P2 + (Multiplier * M2 * Bary * Delta), M2));
 		}
 		if (M3)
 		{
-			VectorScatter(&P[i3], P3 + (Multiplier * M3 * (1 - Bary) * Delta));
+			VectorScatter(&PandInvM[i3], SetVector4(P3 + (Multiplier * M3 * (1 - Bary) * Delta), M3));
 		}
 	}
 }
 
-export void ApplyAxialSpringConstraintsWithWeightMaps(uniform FVector3f P[],
+export void ApplyAxialSpringConstraintsWithWeightMaps(uniform FVector4f PandInvM[],
 										const uniform FIntVector Constraints[],
-										const uniform float InvM[],
 										const uniform float Barys[],
 										const uniform float Dists[],
 										const uniform uint8 StiffnessIndices[],
@@ -88,10 +86,16 @@ export void ApplyAxialSpringConstraintsWithWeightMaps(uniform FVector3f P[],
 		const varying int32 i2 = Constraint.V[1];
 		const varying int32 i3 = Constraint.V[2];
 
-		const varying FVector3f P2 = VectorGather(&P[i2]);
-		const varying FVector3f P3 = VectorGather(&P[i3]);
+		const varying FVector4f PandInvM2 = VectorGather(&PandInvM[i2]);
+		const varying FVector4f PandInvM3 = VectorGather(&PandInvM[i3]);
 
-		const varying FVector3f P1 = VectorGather(&P[i1]);
+		varying FVector3f P1, P2, P3;
+		varying float M1, M2, M3;
+		UnzipPandInvM(PandInvM2, P2, M2);
+		UnzipPandInvM(PandInvM3, P3, M3);
+
+		const varying FVector4f PandInvM1 = VectorGather(&PandInvM[i1]);
+		UnzipPandInvM(PandInvM1, P1, M1);
 
 		const FVector3f P0 = (P2 - P3) * Bary + P3;
 		const FVector3f Difference = P1 - P0;
@@ -99,13 +103,6 @@ export void ApplyAxialSpringConstraintsWithWeightMaps(uniform FVector3f P[],
 		const float Distance = VectorSize(Difference);
 		const float Multiplier = 2 / (max(Bary, 1 - Bary) + 1);
 
-		#pragma ignore warning(perf)
-		const varying float M1 = InvM[i1];
-		#pragma ignore warning(perf)
-		const varying float M2 = InvM[i2];
-		#pragma ignore warning(perf)
-		const varying float M3 = InvM[i3];
-
 		#pragma ignore warning(perf)
 		const varying float Stiffness = StiffnessTable[StiffnessIndex];
 
@@ -120,15 +117,15 @@ export void ApplyAxialSpringConstraintsWithWeightMaps(uniform FVector3f P[],
 
 		if (M1 > 0)
 		{
-			VectorScatter(&P[i1], P1 - (Multiplier * M1 * Delta));
+			VectorScatter(&PandInvM[i1], SetVector4(P1 - (Multiplier * M1 * Delta), M1));
 		}
 		if (M2)
 		{
-			VectorScatter(&P[i2], P2 + (Multiplier * M2 * Bary * Delta));
+			VectorScatter(&PandInvM[i2], SetVector4(P2 + (Multiplier * M2 * Bary * Delta), M2));
 		}
 		if (M3)
 		{
-			VectorScatter(&P[i3], P3 + (Multiplier * M3 * (1 - Bary) * Delta));
+			VectorScatter(&PandInvM[i3], SetVector4(P3 + (Multiplier * M3 * (1 - Bary) * Delta), M3));
 		}
 	}
 }
diff --git a/Engine/Source/Runtime/Experimental/Chaos/Private/Chaos/PBDCollisionSpringConstraintsBase.cpp b/Engine/Source/Runtime/Experimental/Chaos/Private/Chaos/PBDCollisionSpringConstraintsBase.cpp
--- a/Engine/Source/Runtime/Experimental/Chaos/Private/Chaos/PBDCollisionSpringConstraintsBase.cpp
+++ b/Engine/Source/Runtime/Experimental/Chaos/Private/Chaos/PBDCollisionSpringConstraintsBase.cpp
@@ -3,6 +3,7 @@
 #include "Chaos/Plane.h"
 #include "Chaos/TriangleCollisionPoint.h"
 #include "Chaos/TriangleMesh.h"
+#include "Chaos/PBDSoftsSolverParticles.h"
 
 #if !COMPILE_WITHOUT_UNREAL_SUPPORT
 #include "Chaos/Framework/Parallel.h"
diff --git a/Engine/Source/Runtime/Experimental/Chaos/Private/Chaos/PBDEvolution.cpp b/Engine/Source/Runtime/Experimental/Chaos/Private/Chaos/PBDEvolution.cpp
--- a/Engine/Source/Runtime/Experimental/Chaos/Private/Chaos/PBDEvolution.cpp
+++ b/Engine/Source/Runtime/Experimental/Chaos/Private/Chaos/PBDEvolution.cpp
@@ -203,6 +203,7 @@ void FPBDEvolution::PreIterationUpdate(
 
 	if (bVelocityField)
 	{
+		TRACE_CPUPROFILER_EVENT_SCOPE(ChaosPBDVelocityFieldUpdateForces);
 		SCOPE_CYCLE_COUNTER(STAT_ChaosPBDVelocityFieldUpdateForces);
 		VelocityField.UpdateForces(MParticles, Dt);  // Update force per surface element
 	}
@@ -210,11 +211,13 @@ void FPBDEvolution::PreIterationUpdate(
 	FPerParticleDampVelocity DampVelocityRule(MGroupDampings[ParticleGroupId]);
 	if (bDampVelocityRule)
 	{
+		TRACE_CPUPROFILER_EVENT_SCOPE(ChaosPBDVelocityDampUpdateState);
 		SCOPE_CYCLE_COUNTER(STAT_ChaosPBDVelocityDampUpdateState);
 		DampVelocityRule.UpdatePositionBasedState(MParticles, Offset, Range);
 	}
 
 	{
+		TRACE_CPUPROFILER_EVENT_SCOPE(ChaosClothSolverIntegrate);
 		SCOPE_CYCLE_COUNTER(STAT_ChaosClothSolverIntegrate);
 
 		const int32 RangeSize = Range - Offset;
@@ -459,7 +462,7 @@ void FPBDEvolution::AdvanceOneTimeStep(const FSolverReal Dt, const bool bSmoothD
 						ispc::PostIterationUpdates(
 							(ispc::FVector3f*)Particles.GetV().GetData(),
 							(ispc::FVector3f*)Particles.XArray().GetData(),
-							(const ispc::FVector3f*)Particles.GetP().GetData(),
+							(const ispc::FVector4f*)Particles.GetPAndInvM().GetData(),
 							Dt,
 							Offset,
 							Range);
diff --git a/Engine/Source/Runtime/Experimental/Chaos/Private/Chaos/PBDEvolution.ispc b/Engine/Source/Runtime/Experimental/Chaos/Private/Chaos/PBDEvolution.ispc
--- a/Engine/Source/Runtime/Experimental/Chaos/Private/Chaos/PBDEvolution.ispc
+++ b/Engine/Source/Runtime/Experimental/Chaos/Private/Chaos/PBDEvolution.ispc
@@ -6,43 +6,20 @@
 
 export void PostIterationUpdates(uniform FVector3f ParticlesV[],
 								uniform FVector3f ParticlesX[],
-								const uniform FVector3f ParticlesP[],
+								const uniform FVector4f ParticlesPandInvM[],
 								const uniform float Dt,
 								const uniform int32 Offset,
 								const uniform int32 Range)
 {
-	const uniform int RangeBase = (Range - Offset) & ~(programCount-1);
-	const varying float vDt = Dt;
-
-	for(uniform int32 Index = Offset; Index < Offset + RangeBase; Index+=programCount)
+	foreach(i = Offset ... Range)
 	{
-		uniform float *uniform VPtr = (uniform float *uniform)&ParticlesV[Index];
-		uniform float *uniform XPtr = (uniform float *uniform)&ParticlesX[Index];
-		const uniform float *uniform PPtr = (uniform float *uniform)&ParticlesP[Index];
-
-		const varying float P0 = PPtr[programIndex];
-		const varying float P1 = PPtr[programCount + programIndex];
-		const varying float P2 = PPtr[2*programCount + programIndex];
-
-		const varying float X0 = XPtr[programIndex];
-		const varying float X1 = XPtr[programCount + programIndex];
-		const varying float X2 = XPtr[2*programCount + programIndex];
-
-		VPtr[programIndex] = (P0 - X0) / vDt;
-		VPtr[programCount + programIndex] = (P1 - X1) / vDt;
-		VPtr[2*programCount + programIndex] = (P2 - X2) / vDt;
+		const uniform int32 Index = extract(i, 0);
 
-		XPtr[programIndex] = P0;
-		XPtr[programCount + programIndex] = P1;
-		XPtr[2*programCount + programIndex] = P2;
-	}
-
-	for(uniform int32 Index = Offset + RangeBase; Index < Range; ++Index)
-	{
-		const uniform FVector3f P = ParticlesP[Index];
-		const uniform FVector3f X = ParticlesX[Index];
+		const varying FVector4f PandInvM = VectorLoad(&ParticlesPandInvM[Index]);
+		const varying FVector3f P = {{PandInvM.V[0], PandInvM.V[1], PandInvM.V[2]}};
+		const varying FVector3f X = VectorLoad(&ParticlesX[Index]);
 
-		ParticlesV[Index] = (P - X) / Dt;
-		ParticlesX[Index] = P;
+		VectorStore(&ParticlesV[Index], (P - X) / Dt);
+		VectorStore(&ParticlesX[Index], P);
 	}
 }
diff --git a/Engine/Source/Runtime/Experimental/Chaos/Private/Chaos/PBDLongRangeConstraints.cpp b/Engine/Source/Runtime/Experimental/Chaos/Private/Chaos/PBDLongRangeConstraints.cpp
--- a/Engine/Source/Runtime/Experimental/Chaos/Private/Chaos/PBDLongRangeConstraints.cpp
+++ b/Engine/Source/Runtime/Experimental/Chaos/Private/Chaos/PBDLongRangeConstraints.cpp
@@ -6,7 +6,7 @@
 #endif
 
 #if INTEL_ISPC && !UE_BUILD_SHIPPING
-static_assert(sizeof(ispc::FVector3f) == sizeof(Chaos::Softs::FSolverVec3), "sizeof(ispc::FVector3f) != sizeof(Chaos::Softs::FSolverVec3)");
+static_assert(sizeof(ispc::FVector4f) == sizeof(Chaos::Softs::FPAndInvM), "sizeof(ispc::FVector4f) != sizeof(Chaos::Softs::FPAndInvM)");
 static_assert(sizeof(ispc::FTether) == sizeof(Chaos::Softs::FPBDLongRangeConstraintsBase::FTether), "sizeof(ispc::FTether) != sizeof(Chaos::Softs::FPBDLongRangeConstraintsBase::FTether)");
 
 bool bChaos_LongRange_ISPC_Enabled = true;
@@ -34,7 +34,7 @@ void FPBDLongRangeConstraints::Apply(FSolverParticles& Particles, const FSolverR
 				for (const TConstArrayView<FTether>& TetherBatch : Tethers)
 				{
 					ispc::ApplyLongRangeConstraints(
-						(ispc::FVector3f*)Particles.GetP().GetData(),
+						(ispc::FVector4f*)Particles.GetPAndInvM().GetData(),
 						(const ispc::FTether*)TetherBatch.GetData(),
 						ExpStiffnessValue,
 						ScaleValue,
@@ -65,7 +65,7 @@ void FPBDLongRangeConstraints::Apply(FSolverParticles& Particles, const FSolverR
 				for (const TConstArrayView<FTether>& TetherBatch : Tethers)
 				{
 					ispc::ApplyLongRangeConstraintsScaleWeightmap(
-						(ispc::FVector3f*)Particles.GetP().GetData(),
+						(ispc::FVector4f*)Particles.GetPAndInvM().GetData(),
 						(const ispc::FTether*)TetherBatch.GetData(),
 						ExpStiffnessValue,
 						ScaleIndices.GetData(),
@@ -104,7 +104,7 @@ void FPBDLongRangeConstraints::Apply(FSolverParticles& Particles, const FSolverR
 				for (const TConstArrayView<FTether>& TetherBatch : Tethers)
 				{
 					ispc::ApplyLongRangeConstraintsStiffnessWeightmap(
-						(ispc::FVector3f*)Particles.GetP().GetData(),
+						(ispc::FVector4f*)Particles.GetPAndInvM().GetData(),
 						(const ispc::FTether*)TetherBatch.GetData(),
 						Stiffness.GetIndices().GetData(),
 						Stiffness.GetTable().GetData(),
@@ -138,7 +138,7 @@ void FPBDLongRangeConstraints::Apply(FSolverParticles& Particles, const FSolverR
 				for (const TConstArrayView<FTether>& TetherBatch : Tethers)
 				{
 					ispc::ApplyLongRangeConstraintsStiffnessScaleWeightmaps(
-						(ispc::FVector3f*)Particles.GetP().GetData(),
+						(ispc::FVector4f*)Particles.GetPAndInvM().GetData(),
 						(const ispc::FTether*)TetherBatch.GetData(),
 						Stiffness.GetIndices().GetData(),
 						Stiffness.GetTable().GetData(),
diff --git a/Engine/Source/Runtime/Experimental/Chaos/Private/Chaos/PBDLongRangeConstraints.ispc b/Engine/Source/Runtime/Experimental/Chaos/Private/Chaos/PBDLongRangeConstraints.ispc
--- a/Engine/Source/Runtime/Experimental/Chaos/Private/Chaos/PBDLongRangeConstraints.ispc
+++ b/Engine/Source/Runtime/Experimental/Chaos/Private/Chaos/PBDLongRangeConstraints.ispc
@@ -3,6 +3,7 @@
 #define EXPLICIT_VECTOR4 1
 
 #include "Math/Vector.isph"
+#include "Chaos/PBDSofts.isph"
 
 struct FTether
 {
@@ -21,7 +22,7 @@ static inline varying FTether LoadTetherHelper(const uniform FTether *uniform Sr
 	return Tether;
 }
 
-export void ApplyLongRangeConstraints(uniform FVector3f P[],
+export void ApplyLongRangeConstraints(uniform FVector4f PandInvM[],
 										const uniform FTether Tethers[],
 										const uniform float Stiffness,
 										const uniform float Scale,
@@ -49,8 +50,13 @@ export void ApplyLongRangeConstraints(uniform FVector3f P[],
 		const varying int32 i1 = Start + ParticleOffset;
 		const varying int32 i2 = End + ParticleOffset;
 
-		const varying FVector3f P1 = VectorGather(&P[i1]);
-		const varying FVector3f P2 = VectorGather(&P[i2]);
+		const varying FVector4f PandInvM1 = VectorGather(&PandInvM[i1]);
+		const varying FVector4f PandInvM2 = VectorGather(&PandInvM[i2]);
+
+		varying FVector3f P1, P2;
+		varying float InvM1, InvM2;
+		UnzipPandInvM(PandInvM1, P1, InvM1);
+		UnzipPandInvM(PandInvM2, P2, InvM2);
 
 #if HW_GATHER_SUPPORTED == 1
 		if(i + programCount < NumConstraints)
@@ -68,11 +74,11 @@ export void ApplyLongRangeConstraints(uniform FVector3f P[],
 
 		const FVector3f Delta = VectorSelect(Offset < FLOAT_ZERO, FloatZeroVector, ComputedDelta);
 
-		VectorScatter(&P[i2], P2 + Delta);
+		VectorScatter(&PandInvM[i2], SetVector4(P2 + Delta, InvM2));
 	}
 }
 
-export void ApplyLongRangeConstraintsStiffnessWeightmap(uniform FVector3f P[],
+export void ApplyLongRangeConstraintsStiffnessWeightmap(uniform FVector4f PandInvM[],
 														const uniform FTether Tethers[],
 														const uniform uint8 StiffnessIndices[],
 														const uniform float StiffnessTable[],
@@ -104,8 +110,13 @@ export void ApplyLongRangeConstraintsStiffnessWeightmap(uniform FVector3f P[],
 		const varying int32 i1 = Start + ParticleOffset;
 		const varying int32 i2 = End + ParticleOffset;
 
-		const varying FVector3f P1 = VectorGather(&P[i1]);
-		const varying FVector3f P2 = VectorGather(&P[i2]);
+		const varying FVector4f PandInvM1 = VectorGather(&PandInvM[i1]);
+		const varying FVector4f PandInvM2 = VectorGather(&PandInvM[i2]);
+
+		varying FVector3f P1, P2;
+		varying float InvM1, InvM2;
+		UnzipPandInvM(PandInvM1, P1, InvM1);
+		UnzipPandInvM(PandInvM2, P2, InvM2);
 
 #if HW_GATHER_SUPPORTED == 1
 		if(i + programCount < NumConstraints)
@@ -126,11 +137,11 @@ export void ApplyLongRangeConstraintsStiffnessWeightmap(uniform FVector3f P[],
 
 		const FVector3f Delta = VectorSelect(Offset < FLOAT_ZERO, FloatZeroVector, ComputedDelta);
 
-		VectorScatter(&P[i2], P2 + Delta);
+		VectorScatter(&PandInvM[i2], SetVector4(P2 + Delta, InvM2));
 	}
 }
 
-export void ApplyLongRangeConstraintsScaleWeightmap(uniform FVector3f P[],
+export void ApplyLongRangeConstraintsScaleWeightmap(uniform FVector4f PandInvM[],
 													const uniform FTether Tethers[],
 													const uniform float Stiffness,
 													const uniform uint8 ScaleIndices[],
@@ -162,8 +173,13 @@ export void ApplyLongRangeConstraintsScaleWeightmap(uniform FVector3f P[],
 		const varying int32 i1 = Start + ParticleOffset;
 		const varying int32 i2 = End + ParticleOffset;
 
-		const varying FVector3f P1 = VectorGather(&P[i1]);
-		const varying FVector3f P2 = VectorGather(&P[i2]);
+		const varying FVector4f PandInvM1 = VectorGather(&PandInvM[i1]);
+		const varying FVector4f PandInvM2 = VectorGather(&PandInvM[i2]);
+
+		varying FVector3f P1, P2;
+		varying float InvM1, InvM2;
+		UnzipPandInvM(PandInvM1, P1, InvM1);
+		UnzipPandInvM(PandInvM2, P2, InvM2);
 
 #if HW_GATHER_SUPPORTED == 1
 		if(i + programCount < NumConstraints)
@@ -184,12 +200,12 @@ export void ApplyLongRangeConstraintsScaleWeightmap(uniform FVector3f P[],
 
 		const FVector3f Delta = VectorSelect(Offset < FLOAT_ZERO, FloatZeroVector, ComputedDelta);
 
-		VectorScatter(&P[i2], P2 + Delta);
+		VectorScatter(&PandInvM[i2], SetVector4(P2 + Delta, InvM2));
 	}
 }
 
 
-export void ApplyLongRangeConstraintsStiffnessScaleWeightmaps(uniform FVector3f P[],
+export void ApplyLongRangeConstraintsStiffnessScaleWeightmaps(uniform FVector4f PandInvM[],
 																const uniform FTether Tethers[],
 																const uniform uint8 StiffnessIndices[],
 																const uniform float StiffnessTable[],
@@ -224,8 +240,13 @@ export void ApplyLongRangeConstraintsStiffnessScaleWeightmaps(uniform FVector3f
 		const varying int32 i1 = Start + ParticleOffset;
 		const varying int32 i2 = End + ParticleOffset;
 
-		const varying FVector3f P1 = VectorGather(&P[i1]);
-		const varying FVector3f P2 = VectorGather(&P[i2]);
+		const varying FVector4f PandInvM1 = VectorGather(&PandInvM[i1]);
+		const varying FVector4f PandInvM2 = VectorGather(&PandInvM[i2]);
+
+		varying FVector3f P1, P2;
+		varying float InvM1, InvM2;
+		UnzipPandInvM(PandInvM1, P1, InvM1);
+		UnzipPandInvM(PandInvM2, P2, InvM2);
 
 #if HW_GATHER_SUPPORTED == 1
 		if(i + programCount < NumConstraints)
@@ -248,6 +269,6 @@ export void ApplyLongRangeConstraintsStiffnessScaleWeightmaps(uniform FVector3f
 
 		const FVector3f Delta = VectorSelect(Offset < FLOAT_ZERO, FloatZeroVector, ComputedDelta);
 
-		VectorScatter(&P[i2], P2 + Delta);
+		VectorScatter(&PandInvM[i2], SetVector4(P2 + Delta, InvM2));
 	}
 }
diff --git a/Engine/Source/Runtime/Experimental/Chaos/Private/Chaos/PBDSofts.isph b/Engine/Source/Runtime/Experimental/Chaos/Private/Chaos/PBDSofts.isph
new file mode 100644
--- /dev/null
+++ b/Engine/Source/Runtime/Experimental/Chaos/Private/Chaos/PBDSofts.isph
@@ -0,0 +1,14 @@
+// Copyright Epic Games, Inc. All Rights Reserved.
+
+#define EXPLICIT_VECTOR4 1
+
+#include "Math/Vector.isph"
+
+inline void UnzipPandInvM( const FVector4f& PandInvM, FVector3f& P, float& InvM )
+{
+	P.V[0] = PandInvM.V[0];
+	P.V[1] = PandInvM.V[1];
+	P.V[2] = PandInvM.V[2];
+
+	InvM = PandInvM.V[3];
+}
\ No newline at end of file
diff --git a/Engine/Source/Runtime/Experimental/Chaos/Private/Chaos/PBDSphericalConstraint.cpp b/Engine/Source/Runtime/Experimental/Chaos/Private/Chaos/PBDSphericalConstraint.cpp
--- a/Engine/Source/Runtime/Experimental/Chaos/Private/Chaos/PBDSphericalConstraint.cpp
+++ b/Engine/Source/Runtime/Experimental/Chaos/Private/Chaos/PBDSphericalConstraint.cpp
@@ -13,7 +13,7 @@
 bool bChaos_Spherical_ISPC_Enabled = true;
 FAutoConsoleVariableRef CVarChaosSphericalISPCEnabled(TEXT("p.Chaos.Spherical.ISPC"), bChaos_Spherical_ISPC_Enabled, TEXT("Whether to use ISPC optimizations in spherical constraints"));
 
-static_assert(sizeof(ispc::FVector3f) == sizeof(Chaos::Softs::FSolverVec3), "sizeof(ispc::FVector3f) != sizeof(Chaos::Softs::FSolverVec3)");
+static_assert(sizeof(ispc::FVector4f) == sizeof(Chaos::Softs::FPAndInvM), "sizeof(ispc::FVector4f) != sizeof(Chaos::Softs::FPAndInvM)");
 #endif
 
 namespace Chaos::Softs {
@@ -26,9 +26,8 @@ void FPBDSphericalConstraint::ApplyHelperISPC(FSolverParticles& Particles, const
 
 #if INTEL_ISPC
 	ispc::ApplySphericalConstraints(
-		(ispc::FVector3f*)Particles.GetP().GetData(),
+		(ispc::FVector4f*)Particles.GetPAndInvM().GetData(),
 		(const ispc::FVector3f*)AnimationPositions.GetData(),
-		Particles.GetInvM().GetData(),
 		SphereRadii.GetData(),
 		SphereRadiiMultiplier,
 		ParticleOffset,
@@ -43,10 +42,9 @@ void FPBDSphericalBackstopConstraint::ApplyLegacyHelperISPC(FSolverParticles& Pa
 
 #if INTEL_ISPC
 	ispc::ApplyLegacySphericalBackstopConstraints(
-		(ispc::FVector3f*)Particles.GetP().GetData(),
+		(ispc::FVector4f*)Particles.GetPAndInvM().GetData(),
 		(const ispc::FVector3f*)AnimationPositions.GetData(),
 		(const ispc::FVector3f*)AnimationNormals.GetData(),
-		Particles.GetInvM().GetData(),
 		SphereOffsetDistances.GetData(),
 		SphereRadii.GetData(),
 		SphereRadiiMultiplier,
@@ -62,10 +60,9 @@ void FPBDSphericalBackstopConstraint::ApplyHelperISPC(FSolverParticles& Particle
 
 #if INTEL_ISPC
 	ispc::ApplySphericalBackstopConstraints(
-		(ispc::FVector3f*)Particles.GetP().GetData(),
+		(ispc::FVector4f*)Particles.GetPAndInvM().GetData(),
 		(const ispc::FVector3f*)AnimationPositions.GetData(),
 		(const ispc::FVector3f*)AnimationNormals.GetData(),
-		Particles.GetInvM().GetData(),
 		SphereOffsetDistances.GetData(),
 		SphereRadii.GetData(),
 		SphereRadiiMultiplier,
diff --git a/Engine/Source/Runtime/Experimental/Chaos/Private/Chaos/PBDSphericalConstraint.ispc b/Engine/Source/Runtime/Experimental/Chaos/Private/Chaos/PBDSphericalConstraint.ispc
--- a/Engine/Source/Runtime/Experimental/Chaos/Private/Chaos/PBDSphericalConstraint.ispc
+++ b/Engine/Source/Runtime/Experimental/Chaos/Private/Chaos/PBDSphericalConstraint.ispc
@@ -3,10 +3,10 @@
 #define EXPLICIT_VECTOR4 1
 
 #include "Math/Vector.isph"
+#include "Chaos/PBDSofts.isph"
 
-export void ApplySphericalConstraints(uniform FVector3f P[],
+export void ApplySphericalConstraints(uniform FVector4f ParticlesPandInvM[],
 									const uniform FVector3f AnimationPositions[],
-									const uniform float InvM[],
 									const uniform float SphereRadii[],
 									const uniform float SphereRadiiMultiplier,
 									const uniform int32 ParticleOffset,
@@ -17,17 +17,20 @@ export void ApplySphericalConstraints(uniform FVector3f P[],
 	foreach(Index = 0 ... ParticleCount)
 	{
 		const int32 ParticleIndex = ParticleOffset + Index;
-		const float IM = InvM[ParticleIndex];
+		const uniform int32 Offset = extract(ParticleIndex, 0);
+		const FVector4f PandInvM = VectorLoad(&ParticlesPandInvM[Offset]);
+
+		varying FVector3f Position;
+		varying float IM;
+		UnzipPandInvM(PandInvM, Position, IM);
 
 		if(all(IM == 0))
 		{
 			continue;
 		}
 
-		const uniform int32 Offset = extract(ParticleIndex, 0);
 		
 		const FVector3f Center = VectorLoad(&AnimationPositions[Offset]);
-		const FVector3f Position = VectorLoad(&P[Offset]);
 
 		const FVector3f CenterToParticle = Position - Center;
 		const float DistanceSquared = VectorSizeSquared(CenterToParticle);
@@ -38,15 +41,14 @@ export void ApplySphericalConstraints(uniform FVector3f P[],
 
 		if (IM != 0 && DistanceSquared > ((Radius * Radius) + DeadZoneSquareRadius))
 		{
-			VectorStore(&P[Offset], FinalPosition);
+			VectorStore(&ParticlesPandInvM[Offset], SetVector4(FinalPosition, IM));
 		}
 	}
 }
 
-export void ApplyLegacySphericalBackstopConstraints(uniform FVector3f P[],
+export void ApplyLegacySphericalBackstopConstraints(uniform FVector4f ParticlesPandInvM[],
 													const uniform FVector3f AnimationPositions[],
 													const uniform FVector3f AnimationNormals[],
-													const uniform float InvM[],
 													const uniform float SphereOffsetDistances[],
 													const uniform float SphereRadii[],
 													const uniform float SphereRadiiMultiplier,
@@ -58,7 +60,12 @@ export void ApplyLegacySphericalBackstopConstraints(uniform FVector3f P[],
 	foreach(Index = 0 ... ParticleCount)
 	{
 		const int32 ParticleIndex = ParticleOffset + Index;
-		const float IM = InvM[ParticleIndex];
+		const uniform int32 Offset = extract(ParticleIndex, 0);
+		const FVector4f PandInvM = VectorLoad(&ParticlesPandInvM[Offset]);
+
+		varying FVector3f Position;
+		varying float IM;
+		UnzipPandInvM(PandInvM, Position, IM);
 
 		if(all(IM == 0))
 		{
@@ -72,7 +79,7 @@ export void ApplyLegacySphericalBackstopConstraints(uniform FVector3f P[],
 		const float Radius = SphereRadii[Index] * SphereRadiiMultiplier;
 
 		const FVector3f Center = AnimationPosition - SphereOffsetDistance * AnimationNormal;  // Legacy version already includes the radius within the distance
-		const FVector3f CenterToParticle = VectorLoad(&P[extract(ParticleIndex, 0)]) - Center;
+		const FVector3f CenterToParticle = Position - Center;
 		const float DistanceSquared = VectorSizeSquared(CenterToParticle);
 
 		const FVector3f DistSquaredLessRadiusSquared = ((Radius / sqrt(DistanceSquared)) * CenterToParticle) + Center;
@@ -83,15 +90,14 @@ export void ApplyLegacySphericalBackstopConstraints(uniform FVector3f P[],
 
 		if (IM != 0 && NewP.V[0] != 0)
 		{
-			VectorStore(&P[extract(ParticleIndex, 0)], NewP);
+			VectorStore(&ParticlesPandInvM[Offset], SetVector4(NewP, IM));
 		}
 	}
 }
 
-export void ApplySphericalBackstopConstraints(uniform FVector3f P[],
+export void ApplySphericalBackstopConstraints(uniform FVector4f ParticlesPandInvM[],
 												const uniform FVector3f AnimationPositions[],
 												const uniform FVector3f AnimationNormals[],
-												const uniform float InvM[],
 												const uniform float SphereOffsetDistances[],
 												const uniform float SphereRadii[],
 												const uniform float SphereRadiiMultiplier,
@@ -103,7 +109,12 @@ export void ApplySphericalBackstopConstraints(uniform FVector3f P[],
 	foreach(Index = 0 ... ParticleCount)
 	{
 		const int32 ParticleIndex = ParticleOffset + Index;
-		const float IM = InvM[ParticleIndex];
+		const uniform int32 Offset = extract(ParticleIndex, 0);
+		const FVector4f PandInvM = VectorLoad(&ParticlesPandInvM[Offset]);
+
+		varying FVector3f Position;
+		varying float IM;
+		UnzipPandInvM(PandInvM, Position, IM);
 
 		if(all(IM == 0))
 		{
@@ -117,7 +128,7 @@ export void ApplySphericalBackstopConstraints(uniform FVector3f P[],
 		const float Radius = SphereRadii[Index] * SphereRadiiMultiplier;
 
 		const FVector3f Center = AnimationPosition - (Radius + SphereOffsetDistance) * AnimationNormal;  // Non legacy version adds radius to the distance
-		const FVector3f CenterToParticle = VectorLoad(&P[extract(ParticleIndex, 0)]) - Center;
+		const FVector3f CenterToParticle = Position - Center;
 		const float DistanceSquared = VectorSizeSquared(CenterToParticle);
 
 		const FVector3f DistSquaredLessRadiusSquared = ((Radius / sqrt(DistanceSquared)) * CenterToParticle) + Center;
@@ -128,7 +139,7 @@ export void ApplySphericalBackstopConstraints(uniform FVector3f P[],
 
 		if (IM != 0 && NewP.V[0] != 0)
 		{
-			VectorStore(&P[extract(ParticleIndex, 0)], NewP);
+			VectorStore(&ParticlesPandInvM[Offset], SetVector4(NewP, IM));
 		}
 	}
 }
diff --git a/Engine/Source/Runtime/Experimental/Chaos/Private/Chaos/PBDSpringConstraints.cpp b/Engine/Source/Runtime/Experimental/Chaos/Private/Chaos/PBDSpringConstraints.cpp
--- a/Engine/Source/Runtime/Experimental/Chaos/Private/Chaos/PBDSpringConstraints.cpp
+++ b/Engine/Source/Runtime/Experimental/Chaos/Private/Chaos/PBDSpringConstraints.cpp
@@ -1,6 +1,6 @@
 // Copyright Epic Games, Inc. All Rights Reserved.
 #include "Chaos/PBDSpringConstraints.h"
-#include "Chaos/PBDParticles.h"
+#include "Chaos/PBDSoftsSolverParticles.h"
 #include "Chaos/PBDRigidParticles.h"
 #include "Chaos/GraphColoring.h"
 #include "Chaos/Framework/Parallel.h"
@@ -15,7 +15,7 @@
 DECLARE_CYCLE_STAT(TEXT("Chaos PBD Spring Constraint"), STAT_PBD_Spring, STATGROUP_Chaos);
 
 #if INTEL_ISPC && !UE_BUILD_SHIPPING
-static_assert(sizeof(ispc::FVector3f) == sizeof(Chaos::Softs::FSolverVec3), "sizeof(ispc::FVector3f) != sizeof(Chaos::Softs::FSolverVec3)");
+static_assert(sizeof(ispc::FVector4f) == sizeof(Chaos::Softs::FPAndInvM), "sizeof(ispc::FVector4f) != sizeof(Chaos::Softs::FPAndInvM)");
 static_assert(sizeof(ispc::FIntVector2) == sizeof(Chaos::TVec2<int32>), "sizeof(ispc::FIntVector2) != sizeof(Chaos::TVec2<int32>)");
 
 bool bChaos_Spring_ISPC_Enabled = true;
@@ -106,9 +106,8 @@ void FPBDSpringConstraints::Apply(FSolverParticles& Particles, const FSolverReal
 					const int32 ColorStart = ConstraintsPerColorStartIndex[ConstraintColorIndex];
 					const int32 ColorSize = ConstraintsPerColorStartIndex[ConstraintColorIndex + 1] - ColorStart;
 					ispc::ApplySpringConstraints(
-						(ispc::FVector3f*)&Particles.GetP()[0],
+						(ispc::FVector4f*)Particles.GetPAndInvM().GetData(),
 						(ispc::FIntVector2*)&Constraints.GetData()[ColorStart],
-						&Particles.GetInvM().GetData()[0],
 						&Dists.GetData()[ColorStart],
 						ExpStiffnessValue,
 						ColorSize);
@@ -139,9 +138,8 @@ void FPBDSpringConstraints::Apply(FSolverParticles& Particles, const FSolverReal
 					const int32 ColorStart = ConstraintsPerColorStartIndex[ConstraintColorIndex];
 					const int32 ColorSize = ConstraintsPerColorStartIndex[ConstraintColorIndex + 1] - ColorStart;
 					ispc::ApplySpringConstraintsWithWeightMaps(
-						(ispc::FVector3f*) & Particles.GetP()[0],
+						(ispc::FVector4f*)Particles.GetPAndInvM().GetData(),
 						(ispc::FIntVector2*) & Constraints.GetData()[ColorStart],
-						&Particles.GetInvM().GetData()[0],
 						&Dists.GetData()[ColorStart],
 						&Stiffness.GetIndices().GetData()[ColorStart],
 						&Stiffness.GetTable().GetData()[0],
diff --git a/Engine/Source/Runtime/Experimental/Chaos/Private/Chaos/PBDSpringConstraints.ispc b/Engine/Source/Runtime/Experimental/Chaos/Private/Chaos/PBDSpringConstraints.ispc
--- a/Engine/Source/Runtime/Experimental/Chaos/Private/Chaos/PBDSpringConstraints.ispc
+++ b/Engine/Source/Runtime/Experimental/Chaos/Private/Chaos/PBDSpringConstraints.ispc
@@ -3,6 +3,7 @@
 #define EXPLICIT_VECTOR4 1
 
 #include "Math/Vector.isph"
+#include "Chaos/PBDSofts.isph"
 
 static inline float SafeNormalize(FVector3f &Direction)
 {
@@ -12,9 +13,8 @@ static inline float SafeNormalize(FVector3f &Direction)
 	return (SizeSquared < FLOAT_KINDA_SMALL_NUMBER) ? FLOAT_ZERO : Size;
 }
 
-export void ApplySpringConstraints(uniform FVector3f P[],
+export void ApplySpringConstraints(uniform FVector4f PandInvM[],
 									const uniform FIntVector2 Constraints[],
-									const uniform float InvM[],
 									const uniform float Dists[],
 									const uniform float Stiffness,
 									const uniform int32 NumConstraints)
@@ -25,8 +25,13 @@ export void ApplySpringConstraints(uniform FVector3f P[],
 		const varying int32 i1 = Constraint.V[0];
 		const varying int32 i2 = Constraint.V[1];
 
-		const varying FVector3f P1 = VectorGather(&P[i1]);
-		const varying FVector3f P2 = VectorGather(&P[i2]);
+		const varying FVector4f PandInvM1 = VectorGather(&PandInvM[i1]);
+		const varying FVector4f PandInvM2 = VectorGather(&PandInvM[i2]);
+
+		varying FVector3f P1, P2;
+		varying float InvM1, InvM2;
+		UnzipPandInvM(PandInvM1, P1, InvM1);
+		UnzipPandInvM(PandInvM2, P2, InvM2);
 
 		FVector3f Direction = P1 - P2;
 		const float Distance = SafeNormalize(Direction);
@@ -35,11 +40,6 @@ export void ApplySpringConstraints(uniform FVector3f P[],
 
 		const FVector3f InnerDelta = (Distance - Dist) * Direction;
 
-		#pragma ignore warning(perf)
-		const varying float InvM1 = InvM[i1];
-		#pragma ignore warning(perf)
-		const varying float InvM2 = InvM[i2];
-
 		const float CombinedInvMass = InvM2 + InvM1;
 		const FVector3f ComputedDelta = Stiffness * InnerDelta / CombinedInvMass;
 
@@ -47,18 +47,17 @@ export void ApplySpringConstraints(uniform FVector3f P[],
 
 		if (InvM1 > FLOAT_ZERO)
 		{
-			VectorScatter(&P[i1], P1 - (InvM1 * Delta));
+			VectorScatter(&PandInvM[i1], SetVector4( P1 - (InvM1 * Delta), InvM1 ));
 		}
 		if (InvM2 > FLOAT_ZERO)
 		{
-			VectorScatter(&P[i2], P2 + (InvM2 * Delta));
+			VectorScatter(&PandInvM[i2], SetVector4( P2 + (InvM2 * Delta), InvM2 ));
 		}
 	}
 }
 
-export void ApplySpringConstraintsWithWeightMaps(uniform FVector3f P[],
+export void ApplySpringConstraintsWithWeightMaps(uniform FVector4f PandInvM[],
 									const uniform FIntVector2 Constraints[],
-									const uniform float InvM[],
 									const uniform float Dists[],
 									const uniform uint8 StiffnessIndices[],
 									const uniform float StiffnessTable[],
@@ -71,8 +70,13 @@ export void ApplySpringConstraintsWithWeightMaps(uniform FVector3f P[],
 		const varying int32 i1 = Constraint.V[0];
 		const varying int32 i2 = Constraint.V[1];
 
-		const varying FVector3f P1 = VectorGather(&P[i1]);
-		const varying FVector3f P2 = VectorGather(&P[i2]);
+		const varying FVector4f PandInvM1 = VectorGather(&PandInvM[i1]);
+		const varying FVector4f PandInvM2 = VectorGather(&PandInvM[i2]);
+
+		varying FVector3f P1, P2;
+		varying float InvM1, InvM2;
+		UnzipPandInvM(PandInvM1, P1, InvM1);
+		UnzipPandInvM(PandInvM2, P2, InvM2);
 
 		FVector3f Direction = P1 - P2;
 		const float Distance = SafeNormalize(Direction);
@@ -82,11 +86,6 @@ export void ApplySpringConstraintsWithWeightMaps(uniform FVector3f P[],
 
 		const FVector3f InnerDelta = (Distance - Dist) * Direction;
 
-		#pragma ignore warning(perf)
-		const varying float InvM1 = InvM[i1];
-		#pragma ignore warning(perf)
-		const varying float InvM2 = InvM[i2];
-
 		#pragma ignore warning(perf)
 		const varying float Stiffness = StiffnessTable[StiffnessIndex];
 
@@ -97,11 +96,11 @@ export void ApplySpringConstraintsWithWeightMaps(uniform FVector3f P[],
 
 		if (InvM1 > FLOAT_ZERO)
 		{
-			VectorScatter(&P[i1], P1 - (InvM1 * Delta));
+			VectorScatter(&PandInvM[i1], SetVector4( P1 - (InvM1 * Delta), InvM1 ));
 		}
 		if (InvM2 > FLOAT_ZERO)
 		{
-			VectorScatter(&P[i2], P2 + (InvM2 * Delta));
+			VectorScatter(&PandInvM[i2], SetVector4( P2 + (InvM2 * Delta), InvM2 ));
 		}
 	}
 }
diff --git a/Engine/Source/Runtime/Experimental/Chaos/Private/Chaos/PerParticlePBDCollisionConstraint.cpp b/Engine/Source/Runtime/Experimental/Chaos/Private/Chaos/PerParticlePBDCollisionConstraint.cpp
--- a/Engine/Source/Runtime/Experimental/Chaos/Private/Chaos/PerParticlePBDCollisionConstraint.cpp
+++ b/Engine/Source/Runtime/Experimental/Chaos/Private/Chaos/PerParticlePBDCollisionConstraint.cpp
@@ -76,9 +76,8 @@ void FPerParticlePBDCollisionConstraint::ApplyHelperISPC(FSolverParticles& InPar
 					[this, &InParticles, Dt, BatchBegin, BatchEnd, DynamicGroupId, PerGroupFriction, PerGroupThickness](FSolverRigidParticles& CollisionParticles, int32 CollisionOffset, int32 CollisionRange)
 					{
 						ispc::ApplyPerParticleCollisionFastFriction(
-							(ispc::FVector3f*)InParticles.GetP().GetData(),
+							(ispc::FVector4f*)InParticles.GetPAndInvM().GetData(),
 							(const ispc::FVector3f*)InParticles.XArray().GetData(),
-							InParticles.GetInvM().GetData(),
 							(const ispc::FVector3f*)CollisionParticles.AllV().GetData(),
 							(const ispc::FVector3f*)CollisionParticles.XArray().GetData(),
 							(const ispc::FVector3f*)CollisionParticles.AllW().GetData(),
@@ -113,9 +112,8 @@ void FPerParticlePBDCollisionConstraint::ApplyHelperISPC(FSolverParticles& InPar
 					[this, &InParticles, Dt, BatchBegin, BatchEnd, DynamicGroupId, PerGroupThickness](FSolverRigidParticles& CollisionParticles, int32 CollisionOffset, int32 CollisionRange)
 					{
 						ispc::ApplyPerParticleCollisionNoFriction(
-							(ispc::FVector3f*)InParticles.GetP().GetData(),
+							(ispc::FVector4f*)InParticles.GetPAndInvM().GetData(),
 							(const ispc::FVector3f*)InParticles.XArray().GetData(),
-							InParticles.GetInvM().GetData(),
 							(const ispc::FVector3f*)CollisionParticles.AllV().GetData(),
 							(const ispc::FVector3f*)CollisionParticles.XArray().GetData(),
 							(const ispc::FVector3f*)CollisionParticles.AllW().GetData(),
@@ -150,9 +148,8 @@ void FPerParticlePBDCollisionConstraint::ApplyHelperISPC(FSolverParticles& InPar
 				[this, &InParticles, Dt, BatchBegin, BatchEnd](FSolverRigidParticles& CollisionParticles, int32 CollisionOffset, int32 CollisionRange)
 				{
 					ispc::ApplyPerParticleCollision(
-						(ispc::FVector3f*)InParticles.GetP().GetData(),
+						(ispc::FVector4f*)InParticles.GetPAndInvM().GetData(),
 						(const ispc::FVector3f*)InParticles.XArray().GetData(),
-						InParticles.GetInvM().GetData(),
 						(const ispc::FVector3f*)CollisionParticles.AllV().GetData(),
 						(const ispc::FVector3f*)CollisionParticles.XArray().GetData(),
 						(const ispc::FVector3f*)CollisionParticles.AllW().GetData(),
diff --git a/Engine/Source/Runtime/Experimental/Chaos/Private/Chaos/PerParticlePBDCollisionConstraint.ispc b/Engine/Source/Runtime/Experimental/Chaos/Private/Chaos/PerParticlePBDCollisionConstraint.ispc
--- a/Engine/Source/Runtime/Experimental/Chaos/Private/Chaos/PerParticlePBDCollisionConstraint.ispc
+++ b/Engine/Source/Runtime/Experimental/Chaos/Private/Chaos/PerParticlePBDCollisionConstraint.ispc
@@ -4,6 +4,7 @@
 
 #include "Math/Vector.isph"
 #include "Math/Transform.isph"
+#include "Chaos/PBDSofts.isph"
 
 enum ImplicitObjectType
 {
@@ -546,9 +547,8 @@ static inline void PhiWithNormalConvex(const uniform uint8 *uniform TypedObjectP
 
 extern "C" unmasked void GetPhiWithNormal(const uniform uint8 *uniform CollisionParticles, const uniform float *uniform V, uniform float *uniform Normal, uniform float *uniform Phi, const uniform int i, const uniform int ProgramCount, const uniform int Mask);
 
-export void ApplyPerParticleCollision(uniform FVector3f ParticlesP[],
+export void ApplyPerParticleCollision(uniform FVector4f ParticlesPandInvM[],
 										const uniform FVector3f ParticlesX[],
-										const uniform float ParticlesInvM[],
 										const uniform FVector3f CollisionV[],
 										const uniform FVector3f CollisionX[],
 										const uniform FVector3f CollisionW[],
@@ -596,15 +596,16 @@ export void ApplyPerParticleCollision(uniform FVector3f ParticlesP[],
 
 	foreach(Index = OuterOffset ... BailAt)
 	{
-		const float IM = ParticlesInvM[Index];
+		const FVector4f PandInvM = VectorLoad(&ParticlesPandInvM[extract(Index, 0)]);
+		varying FVector3f P;
+		varying float IM;
+		UnzipPandInvM(PandInvM, P, IM);
 
 		if(all(IM == 0))
 		{
 			continue;
 		}
 
-		FVector3f P = VectorLoad(&ParticlesP[extract(Index, 0)]);
-
 		const uint32 DynamicGroupId = DynamicGroupIds[Index]; // Particle group Id
 
 		for(uniform int i = InnerOffset; i < InnerRange; i++) // SequentialFor
@@ -713,14 +714,13 @@ export void ApplyPerParticleCollision(uniform FVector3f ParticlesP[],
 
 		if(IM != 0)
 		{
-			VectorStore(&ParticlesP[extract(Index, 0)], P);
+			VectorStore(&ParticlesPandInvM[extract(Index, 0)], SetVector4(P, IM));
 		}
 	}
 }
 
-export void ApplyPerParticleCollisionFastFriction(uniform FVector3f ParticlesP[],
+export void ApplyPerParticleCollisionFastFriction(uniform FVector4f ParticlesPandInvM[],
 													const uniform FVector3f ParticlesX[],
-													const uniform float ParticlesInvM[],
 													const uniform FVector3f CollisionV[],
 													const uniform FVector3f CollisionX[],
 													const uniform FVector3f CollisionW[],
@@ -766,15 +766,16 @@ export void ApplyPerParticleCollisionFastFriction(uniform FVector3f ParticlesP[]
 
 	foreach(Index = OuterOffset ... BailAt)
 	{
-		const float IM = ParticlesInvM[Index];
+		const FVector4f PandInvM = VectorLoad(&ParticlesPandInvM[extract(Index, 0)]);
+		varying FVector3f P;
+		varying float IM;
+		UnzipPandInvM(PandInvM, P, IM);
 
 		if(all(IM == 0))
 		{
 			continue;
 		}
 
-		FVector3f P = VectorLoad(&ParticlesP[extract(Index, 0)]);
-
 		for(uniform int i = InnerOffset; i < InnerRange; i++) // SequentialFor
 		{
 			// InverseTransformPosition. Can skip Scale calculation because it is 1 for xyz.
@@ -865,14 +866,13 @@ export void ApplyPerParticleCollisionFastFriction(uniform FVector3f ParticlesP[]
 
 		if(IM != 0)
 		{
-			VectorStore(&ParticlesP[extract(Index, 0)], P);
+			VectorStore(&ParticlesPandInvM[extract(Index, 0)], SetVector4(P, IM));
 		}
 	}
 }
 
-export void ApplyPerParticleCollisionNoFriction(uniform FVector3f ParticlesP[],
+export void ApplyPerParticleCollisionNoFriction(uniform FVector4f ParticlesPandInvM[],
 												const uniform FVector3f ParticlesX[],
-												const uniform float ParticlesInvM[],
 												const uniform FVector3f CollisionV[],
 												const uniform FVector3f CollisionX[],
 												const uniform FVector3f CollisionW[],
@@ -917,15 +917,16 @@ export void ApplyPerParticleCollisionNoFriction(uniform FVector3f ParticlesP[],
 
 	foreach(Index = OuterOffset ... BailAt)
 	{
-		const float IM = ParticlesInvM[Index];
+		const FVector4f PandInvM = VectorLoad(&ParticlesPandInvM[extract(Index, 0)]);
+		varying FVector3f P;
+		varying float IM;
+		UnzipPandInvM(PandInvM, P, IM);
 
 		if(all(IM == 0))
 		{
 			continue;
 		}
 
-		FVector3f P = VectorLoad(&ParticlesP[extract(Index, 0)]);
-
 		for(uniform int i = InnerOffset; i < InnerRange; i++) // SequentialFor
 		{
 			// InverseTransformPosition. Can skip Scale calculation because it is 1 for xyz.
@@ -1006,7 +1007,7 @@ export void ApplyPerParticleCollisionNoFriction(uniform FVector3f ParticlesP[],
 
 		if(IM != 0)
 		{
-			VectorStore(&ParticlesP[extract(Index, 0)], P);
+			VectorStore(&ParticlesPandInvM[extract(Index, 0)], SetVector4(P, IM));
 		}
 	}
 }
diff --git a/Engine/Source/Runtime/Experimental/Chaos/Public/Chaos/PBDAxialSpringConstraintsBase.h b/Engine/Source/Runtime/Experimental/Chaos/Public/Chaos/PBDAxialSpringConstraintsBase.h
--- a/Engine/Source/Runtime/Experimental/Chaos/Public/Chaos/PBDAxialSpringConstraintsBase.h
+++ b/Engine/Source/Runtime/Experimental/Chaos/Public/Chaos/PBDAxialSpringConstraintsBase.h
@@ -2,7 +2,7 @@
 #pragma once
 
 #include "Chaos/PBDSoftsEvolutionFwd.h"
-#include "Chaos/PBDParticles.h"
+#include "Chaos/PBDSoftsSolverParticles.h"
 #include "Chaos/PBDStiffness.h"
 #include "Containers/Array.h"
 
diff --git a/Engine/Source/Runtime/Experimental/Chaos/Public/Chaos/PBDBendingConstraintsBase.h b/Engine/Source/Runtime/Experimental/Chaos/Public/Chaos/PBDBendingConstraintsBase.h
--- a/Engine/Source/Runtime/Experimental/Chaos/Public/Chaos/PBDBendingConstraintsBase.h
+++ b/Engine/Source/Runtime/Experimental/Chaos/Public/Chaos/PBDBendingConstraintsBase.h
@@ -3,7 +3,7 @@
 
 #include "Chaos/PBDSoftsEvolutionFwd.h"
 #include "Chaos/Array.h"
-#include "Chaos/PBDParticles.h"
+#include "Chaos/PBDSoftsSolverParticles.h"
 #include "Chaos/ParticleRule.h"
 #include "GenericPlatform/GenericPlatformMath.h"
 
diff --git a/Engine/Source/Runtime/Experimental/Chaos/Public/Chaos/PBDCollisionSpringConstraintsBase.h b/Engine/Source/Runtime/Experimental/Chaos/Public/Chaos/PBDCollisionSpringConstraintsBase.h
--- a/Engine/Source/Runtime/Experimental/Chaos/Public/Chaos/PBDCollisionSpringConstraintsBase.h
+++ b/Engine/Source/Runtime/Experimental/Chaos/Public/Chaos/PBDCollisionSpringConstraintsBase.h
@@ -4,7 +4,7 @@
 #if !COMPILE_WITHOUT_UNREAL_SUPPORT
 #include "Chaos/Core.h"
 #include "Chaos/PBDSoftsEvolutionFwd.h"
-#include "Chaos/PBDParticles.h"
+#include "Chaos/PBDSoftsSolverParticles.h"
 #include "Containers/Set.h"
 
 
diff --git a/Engine/Source/Runtime/Experimental/Chaos/Public/Chaos/PBDEvolution.h b/Engine/Source/Runtime/Experimental/Chaos/Public/Chaos/PBDEvolution.h
--- a/Engine/Source/Runtime/Experimental/Chaos/Public/Chaos/PBDEvolution.h
+++ b/Engine/Source/Runtime/Experimental/Chaos/Public/Chaos/PBDEvolution.h
@@ -5,7 +5,7 @@
 #include "Chaos/ArrayCollection.h"
 #include "Chaos/PBDActiveView.h"
 #include "Chaos/PBDSoftsEvolutionFwd.h"
-#include "Chaos/PBDParticles.h"
+#include "Chaos/PBDSoftsSolverParticles.h"
 #include "Chaos/KinematicGeometryParticles.h"
 #include "Chaos/VelocityField.h"
 
diff --git a/Engine/Source/Runtime/Experimental/Chaos/Public/Chaos/PBDLongRangeConstraintsBase.h b/Engine/Source/Runtime/Experimental/Chaos/Public/Chaos/PBDLongRangeConstraintsBase.h
--- a/Engine/Source/Runtime/Experimental/Chaos/Public/Chaos/PBDLongRangeConstraintsBase.h
+++ b/Engine/Source/Runtime/Experimental/Chaos/Public/Chaos/PBDLongRangeConstraintsBase.h
@@ -2,7 +2,7 @@
 #pragma once
 
 #include "Chaos/Core.h"
-#include "Chaos/PBDParticles.h"
+#include "Chaos/PBDSoftsSolverParticles.h"
 #include "Chaos/PBDSoftsEvolutionFwd.h"
 #include "Chaos/PBDActiveView.h"
 #include "Chaos/PBDStiffness.h"
diff --git a/Engine/Source/Runtime/Experimental/Chaos/Public/Chaos/PBDSoftsEvolutionFwd.h b/Engine/Source/Runtime/Experimental/Chaos/Public/Chaos/PBDSoftsEvolutionFwd.h
--- a/Engine/Source/Runtime/Experimental/Chaos/Public/Chaos/PBDSoftsEvolutionFwd.h
+++ b/Engine/Source/Runtime/Experimental/Chaos/Public/Chaos/PBDSoftsEvolutionFwd.h
@@ -18,6 +18,7 @@ namespace Chaos
 	template <class T, int d> class TRotation;
 	template <class T, int d> class TRigidTransform;
 	template <class T, int d> class TAABB;
+	template <class T> struct TVector3AndScalar;
 
 	class FPerParticleGravity;
 
@@ -31,12 +32,13 @@ namespace Chaos
 		using FSolverVec3 = TVector<FSolverReal, 3>;
 		using FSolverMatrix33 = PMatrix<FSolverReal, 3, 3>;
 		using FSolverMatrix44 = PMatrix<FSolverReal, 4, 4>;
-		using FSolverParticles = TPBDParticles<FSolverReal, 3>;
 		using FSolverRigidParticles = TKinematicGeometryParticlesImp<FSolverReal, 3, EGeometryParticlesSimType::Other>;
 		using FSolverRotation3 = TRotation<FSolverReal, 3>;
 		using FSolverRigidTransform3 = TRigidTransform<FSolverReal, 3>;
 		using FSolverTransform3 = UE::Math::TTransform<FSolverReal>;
 		using FSolverAABB3 = TAABB<FSolverReal, 3>;
+		class FSolverParticles;
+		struct FPAndInvM;
 
 		// Softs solver class
 		class FPBDEvolution;
diff --git a/Engine/Source/Runtime/Experimental/Chaos/Public/Chaos/PBDSoftsSolverParticles.h b/Engine/Source/Runtime/Experimental/Chaos/Public/Chaos/PBDSoftsSolverParticles.h
new file mode 100644
--- /dev/null
+++ b/Engine/Source/Runtime/Experimental/Chaos/Public/Chaos/PBDSoftsSolverParticles.h
@@ -0,0 +1,50 @@
+// Copyright Epic Games, Inc. All Rights Reserved.
+#pragma once
+
+#include "Chaos/Real.h"
+#include "Chaos/Vector.h"
+#include "Chaos/ArrayCollectionArray.h"
+#include "Chaos/DynamicParticles.h"
+#include "Chaos/PBDSoftsEvolutionFwd.h"
+
+namespace Chaos::Softs
+{
+
+struct FPAndInvM
+{
+	FSolverVec3 P;
+	FSolverReal InvM;
+};
+
+class FSolverParticles : public TDynamicParticles<FSolverReal, 3>
+{
+public:
+	typedef TDynamicParticles<FSolverReal, 3> Base;
+
+	FSolverParticles()
+	    : Base()
+	{
+		TArrayCollection::AddArray(&MPAndInvM);
+	}
+	FSolverParticles(const FSolverParticles& Other) = delete;
+	FSolverParticles(FSolverParticles&& Other)
+	    : Base(MoveTemp(Other)), MPAndInvM(MoveTemp(Other.MPAndInvM))
+	{
+		TArrayCollection::AddArray(&MPAndInvM);
+	}
+
+	const FSolverVec3& P(const int32 index) const { return MPAndInvM[index].P; }
+	FSolverVec3& P(const int32 index) { return MPAndInvM[index].P; }
+
+	const FPAndInvM& PAndInvM(const int32 index) const { return MPAndInvM[index]; }
+	FPAndInvM& PAndInvM(const int32 index) { return MPAndInvM[index]; }
+
+	TArrayCollectionArray<FPAndInvM>& GetPAndInvM() { return MPAndInvM; }
+	const TArrayCollectionArray<FPAndInvM>& GetPAndInvM() const { return MPAndInvM; }
+
+private:
+	// Note that InvM is copied from TDynamicParticles to here during ApplyPreSimulationTransforms (when P is updated)
+	TArrayCollectionArray<FPAndInvM> MPAndInvM;
+};
+
+}
diff --git a/Engine/Source/Runtime/Experimental/Chaos/Public/Chaos/PBDSphericalConstraint.h b/Engine/Source/Runtime/Experimental/Chaos/Public/Chaos/PBDSphericalConstraint.h
--- a/Engine/Source/Runtime/Experimental/Chaos/Public/Chaos/PBDSphericalConstraint.h
+++ b/Engine/Source/Runtime/Experimental/Chaos/Public/Chaos/PBDSphericalConstraint.h
@@ -2,7 +2,7 @@
 #pragma once
 
 #include "Chaos/PBDSoftsEvolutionFwd.h"
-#include "Chaos/PBDParticles.h"
+#include "Chaos/PBDSoftsSolverParticles.h"
 #include "ChaosStats.h"
 #include "Containers/ArrayView.h"
 #include "Containers/ContainersFwd.h"
diff --git a/Engine/Source/Runtime/Experimental/Chaos/Public/Chaos/PBDSpringConstraintsBase.h b/Engine/Source/Runtime/Experimental/Chaos/Public/Chaos/PBDSpringConstraintsBase.h
--- a/Engine/Source/Runtime/Experimental/Chaos/Public/Chaos/PBDSpringConstraintsBase.h
+++ b/Engine/Source/Runtime/Experimental/Chaos/Public/Chaos/PBDSpringConstraintsBase.h
@@ -4,7 +4,7 @@
 #include "Chaos/Core.h"
 #include "Chaos/PBDSoftsEvolutionFwd.h"
 #include "Chaos/Utilities.h"
-#include "Chaos/PBDParticles.h"
+#include "Chaos/PBDSoftsSolverParticles.h"
 #include "Chaos/PBDStiffness.h"
 #include "Containers/Array.h"
 #include "Templates/EnableIf.h"
diff --git a/Engine/Source/Runtime/Experimental/Chaos/Public/Chaos/PerParticleDampVelocity.h b/Engine/Source/Runtime/Experimental/Chaos/Public/Chaos/PerParticleDampVelocity.h
--- a/Engine/Source/Runtime/Experimental/Chaos/Public/Chaos/PerParticleDampVelocity.h
+++ b/Engine/Source/Runtime/Experimental/Chaos/Public/Chaos/PerParticleDampVelocity.h
@@ -3,7 +3,7 @@
 
 #include "Chaos/PBDSoftsEvolutionFwd.h"
 #include "Chaos/Vector.h"
-#include "Chaos/PBDParticles.h"
+#include "Chaos/PBDSoftsSolverParticles.h"
 
 namespace Chaos::Softs
 {
diff --git a/Engine/Source/Runtime/Experimental/Chaos/Public/Chaos/PerParticlePBDCCDCollisionConstraint.h b/Engine/Source/Runtime/Experimental/Chaos/Public/Chaos/PerParticlePBDCCDCollisionConstraint.h
--- a/Engine/Source/Runtime/Experimental/Chaos/Public/Chaos/PerParticlePBDCCDCollisionConstraint.h
+++ b/Engine/Source/Runtime/Experimental/Chaos/Public/Chaos/PerParticlePBDCCDCollisionConstraint.h
@@ -4,6 +4,7 @@
 #include "Chaos/GeometryParticlesfwd.h"
 #include "Chaos/Transform.h"
 #include "Chaos/PBDActiveView.h"
+#include "Chaos/PBDSoftsSolverParticles.h"
 #include "Misc/ScopeLock.h"
 
 namespace Chaos::Softs
diff --git a/Engine/Source/Runtime/Experimental/Chaos/Public/Chaos/PerParticlePBDCollisionConstraint.h b/Engine/Source/Runtime/Experimental/Chaos/Public/Chaos/PerParticlePBDCollisionConstraint.h
--- a/Engine/Source/Runtime/Experimental/Chaos/Public/Chaos/PerParticlePBDCollisionConstraint.h
+++ b/Engine/Source/Runtime/Experimental/Chaos/Public/Chaos/PerParticlePBDCollisionConstraint.h
@@ -4,7 +4,7 @@
 #include "Chaos/Core.h"
 #include "Chaos/PBDActiveView.h"
 #include "Chaos/PBDSoftsEvolutionFwd.h"
-#include "Chaos/PBDParticles.h"
+#include "Chaos/PBDSoftsSolverParticles.h"
 #include "Chaos/KinematicGeometryParticles.h"
 #include "HAL/PlatformMath.h"
 
diff --git a/Engine/Source/Runtime/Experimental/Chaos/Public/Chaos/VelocityField.h b/Engine/Source/Runtime/Experimental/Chaos/Public/Chaos/VelocityField.h
--- a/Engine/Source/Runtime/Experimental/Chaos/Public/Chaos/VelocityField.h
+++ b/Engine/Source/Runtime/Experimental/Chaos/Public/Chaos/VelocityField.h
@@ -2,7 +2,7 @@
 #pragma once
 
 #include "Chaos/PBDSoftsEvolutionFwd.h"
-#include "Chaos/PBDParticles.h"
+#include "Chaos/PBDSoftsSolverParticles.h"
 #include "Chaos/TriangleMesh.h"
 
 namespace Chaos::Softs
