PS4 2 MB GPU page tunings:
- Added a range allocator for allocating small blocks. It is enabled if 2 MB GPU page is used to reduce memory overhead. Perf is on par with FGPUMemoryPage.
- 20-run ReplayRun test data with 2 MB page and the new allocator: -0.12 ms FrameTime Avg, -0.17 ms GameThreadTime Avg, -0.43 ms RenderThreadTime Avg, -0.12 ms RHIThreadTime Avg, +0.6 DynRes % Avg, and +18 MemFreeMB Min.
- Added support for recording and replaying FMemBlock allocations. Courtesy to Ben for providing most of the code.
- Small adjustment to how MemBlockStat_Wasted is calculated. It now takes unused page space into account in addition to alignment paddings.
- 2 MB GPU page is enabled by default but can be disabled through a project define.

diff --git a/Engine/Source/Runtime/Core/Private/ProfilingDebugging/CsvProfiler.cpp b/Engine/Source/Runtime/Core/Private/ProfilingDebugging/CsvProfiler.cpp
--- a/Engine/Source/Runtime/Core/Private/ProfilingDebugging/CsvProfiler.cpp
+++ b/Engine/Source/Runtime/Core/Private/ProfilingDebugging/CsvProfiler.cpp
@@ -2895,6 +2895,8 @@ void FCsvProfiler::EndFrame()
 		CSV_CUSTOM_STAT_GLOBAL(ExtendedUsedMB, UsedExtendedMB, ECsvCustomStatOp::Set);
 		CSV_CUSTOM_STAT_GLOBAL(SystemMaxMB, TotalSystemMB, ECsvCustomStatOp::Set);
 
+		MemoryStats.SetEndFrameCsvStats();
+
 		// If we're single-threaded, process the stat data here
 		if (ProcessingThread == nullptr)
 		{
diff --git a/Engine/Source/Runtime/Core/Public/GenericPlatform/GenericPlatformMemory.h b/Engine/Source/Runtime/Core/Public/GenericPlatform/GenericPlatformMemory.h
--- a/Engine/Source/Runtime/Core/Public/GenericPlatform/GenericPlatformMemory.h
+++ b/Engine/Source/Runtime/Core/Public/GenericPlatform/GenericPlatformMemory.h
@@ -174,6 +174,9 @@ struct CORE_API FGenericPlatformMemoryStats : public FPlatformMemoryConstants
 	TArray<FPlatformSpecificStat> GetPlatformSpecificStats() const;
 
 	uint64 GetAvailablePhysical(bool bExcludeExtraDevMemory) const;
+
+	/** Called by FCsvProfiler::EndFrame to set platform specific CSV stats. */
+	void SetEndFrameCsvStats() const {}
 };
 
 
