Fix use of constant in conditional expression from static analysis

diff --git a/Engine/Source/Runtime/Core/Private/HAL/MallocBinned2.cpp b/Engine/Source/Runtime/Core/Private/HAL/MallocBinned2.cpp
--- a/Engine/Source/Runtime/Core/Private/HAL/MallocBinned2.cpp
+++ b/Engine/Source/Runtime/Core/Private/HAL/MallocBinned2.cpp
@@ -943,11 +943,10 @@ void* FMallocBinned2::ReallocExternal(void* Ptr, SIZE_T NewSize, uint32 Alignmen
 		uint32 PoolIndex = Free->PoolIndex;
 #if BINNED2_FORK_SUPPORT
 		// If the canary is the pre-fork one, we should not allow this allocation to grow in-place to avoid copying a page from the parent process.
-		const bool bIsOldAlloc = Free->CanaryAndForkState != CurrentCanary;
+		if (Free->CanaryAndForkState == CurrentCanary && 
 #else
-		constexpr bool bIsOldAlloc = false;
+		if( 
 #endif
-		if ( !bIsOldAlloc && 
 			((NewSize <= BlockSize) & (Alignment <= BINNED2_MINIMUM_ALIGNMENT)) && // one branch, not two
 			(PoolIndex == 0 || NewSize > PoolIndexToBlockSize(PoolIndex - 1)))
 		{
