Give DEFAULT_SERVER_FAKE_FORKS a default value for all games/builds.

diff --git a/Engine/Source/Runtime/Core/Public/HAL/MallocBinned2.h b/Engine/Source/Runtime/Core/Public/HAL/MallocBinned2.h
--- a/Engine/Source/Runtime/Core/Public/HAL/MallocBinned2.h
+++ b/Engine/Source/Runtime/Core/Public/HAL/MallocBinned2.h
@@ -15,6 +15,7 @@
 #include "HAL/Allocators/PooledVirtualMemoryAllocator.h"
 #include "HAL/PlatformMath.h"
 #include "HAL/LowLevelMemTracker.h"
+#include "Misc/Fork.h"
 #include "Templates/Atomic.h"
 
 #define BINNED2_MAX_CACHED_OS_FREES (64)
diff --git a/Engine/Source/Runtime/Core/Public/Misc/Fork.h b/Engine/Source/Runtime/Core/Public/Misc/Fork.h
--- a/Engine/Source/Runtime/Core/Public/Misc/Fork.h
+++ b/Engine/Source/Runtime/Core/Public/Misc/Fork.h
@@ -2,9 +2,16 @@
 
 #pragma once
 
-#include "CoreMinimal.h"
+#include "CoreTypes.h"
 #include "HAL/PlatformAffinity.h"
 
+#ifndef DEFAULT_SERVER_FAKE_FORKS
+	#define DEFAULT_SERVER_FAKE_FORKS 0
+#endif
+
+class CORE_API FRunnable;
+class CORE_API FRunnableThread;
+
 /**
  * Helper functions for processes that fork in order to share memory pages.
  *
