Added support for using FramePro profiling scopes in MallocBinned2 allocator: added a separate memory allocator for FramePro, added support for TArrays to use a different base malloc, rather than FMemory.

#tests automatic and manual FramePro captures of Squads games on Linux in WSL2 and cloud hardware (m5.4)

diff --git a/Engine/Source/Runtime/Core/Private/FramePro/FrameProProfiler.cpp b/Engine/Source/Runtime/Core/Private/FramePro/FrameProProfiler.cpp
--- a/Engine/Source/Runtime/Core/Private/FramePro/FrameProProfiler.cpp
+++ b/Engine/Source/Runtime/Core/Private/FramePro/FrameProProfiler.cpp
@@ -114,6 +114,52 @@ public:
 	}
 };
 
+
+/** Separate memory allocator for FramePro internal memory. Allows adding profiler traces to the main allocator without causing recursion issues. */
+class FrameProAllocator : public FramePro::Allocator
+{
+public:
+	static FrameProAllocator& Get()
+	{
+		static FrameProAllocator Instance;
+		return Instance;
+	}
+
+	FMalloc* GetBaseMalloc() { return &BaseMalloc; }
+
+	virtual void* Alloc(size_t size) override
+	{
+		return BaseMalloc.Malloc(size, DEFAULT_ALIGNMENT);
+	}
+	virtual void Free(void* p) override
+	{
+		return BaseMalloc.Free(p);
+	}
+
+private:
+	FMallocAnsi BaseMalloc;
+};
+
+/** System memory allocator for TArrays. Redirects base malloc calls to the FrameProAllocator. */
+class FrameProMalloc
+{
+public:
+	static void* Malloc(SIZE_T Count, uint32 Alignment = DEFAULT_ALIGNMENT)
+	{
+		return FrameProAllocator::Get().GetBaseMalloc()->Malloc(Count, Alignment);
+	}
+
+	static void* Realloc(void* Original, SIZE_T Count, uint32 Alignment = DEFAULT_ALIGNMENT)
+	{
+		return FrameProAllocator::Get().GetBaseMalloc()->Realloc(Original, Count, Alignment);
+	}
+
+	static void Free(void* Original)
+	{
+		return FrameProAllocator::Get().GetBaseMalloc()->Free(Original);
+	}
+};
+
 /** TSL storage for per-thread scope stack */
 class FFrameProProfilerContext : public TThreadSingleton<FFrameProProfilerContext>
 {
@@ -130,7 +176,7 @@ class FFrameProProfilerContext : public TThreadSingleton<FFrameProProfilerContex
 	}
 	
 	/** Array that represents thread of scopes */
-	TArray<FFrameProProfilerScope> ProfilerScopes;
+	TArray<FFrameProProfilerScope, TSizedHeapAllocator<32, FrameProMalloc>> ProfilerScopes;
 
 public:
 
@@ -230,6 +276,21 @@ static FAutoConsoleVariableRef CVarFrameProCPUStatsUpdateRate(
 	TEXT("0 to disable."),
 	ECVF_Default);
 
+void FFrameProProfiler::Initialize()
+{
+	FRAMEPRO_SET_ALLOCATOR(&FrameProAllocator::Get());
+}
+
+void FFrameProProfiler::TearDown()
+{
+	FRAMEPRO_SHUTDOWN();
+}
+
+bool FFrameProProfiler::IsThreadContextReady()
+{
+	return FFrameProProfilerContext::TryGet() != nullptr;
+}
+
 void FFrameProProfiler::FrameStart()
 {
 	static bool bFirstFrame = true;
diff --git a/Engine/Source/Runtime/Core/Private/HAL/MallocBinned2.cpp b/Engine/Source/Runtime/Core/Private/HAL/MallocBinned2.cpp
--- a/Engine/Source/Runtime/Core/Private/HAL/MallocBinned2.cpp
+++ b/Engine/Source/Runtime/Core/Private/HAL/MallocBinned2.cpp
@@ -12,10 +12,45 @@
 #include "Misc/App.h"
 #include "HAL/MallocTimer.h"
 #include "ProfilingDebugging/CsvProfiler.h"
+#include "FramePro/FrameProProfiler.h"
+
 #if CSV_PROFILER
 CSV_DEFINE_CATEGORY_MODULE(CORE_API, FMemory, true);
 #endif
 
+#if FRAMEPRO_ENABLED
+/** Pushes a profiler scope if it's safe to do so without any new allocations. */
+class FNoAllocScopeCycleCounter
+{
+public:
+	FORCEINLINE FNoAllocScopeCycleCounter(const ANSICHAR* InStatString)
+		: bPop(false)
+		, StatString(InStatString)
+	{
+		if (FFrameProProfiler::IsThreadContextReady() && GCycleStatsShouldEmitNamedEvents)
+		{
+			bPop = true;
+			FFrameProProfiler::PushEvent(StatString);
+		}
+	}
+
+	FORCEINLINE ~FNoAllocScopeCycleCounter()
+	{
+		if (bPop)
+		{
+			FFrameProProfiler::PopEvent(StatString);
+		}
+	}
+private:
+	bool bPop;
+	const ANSICHAR* StatString;
+};
+#define NOALLOC_SCOPE_CYCLE_COUNTER(Stat) FNoAllocScopeCycleCounter NoAllocCycleCounter_##Stat(#Stat)
+#else
+#define NOALLOC_SCOPE_CYCLE_COUNTER(Stat)
+#endif // FRAMEPRO_ENABLED
+
+
 PRAGMA_DISABLE_UNSAFE_TYPECAST_WARNINGS
 
 #if BINNED2_ALLOW_RUNTIME_TWEAKING
@@ -573,6 +608,7 @@ struct FMallocBinned2::Private
 	}
 	static void RegisterThreadFreeBlockLists( FPerThreadFreeBlockLists* FreeBlockLists )
 	{
+		QUICK_SCOPE_CYCLE_COUNTER(STAT_FMallocBinned2_RegisterThreadFreeBlockLists);
 		FScopeLock Lock(&GetFreeBlockListsRegistrationMutex());
 #if BINNED2_ALLOCATOR_STATS_VALIDATION
 		++RecursionCounter;
@@ -584,6 +620,7 @@ struct FMallocBinned2::Private
 	}
 	static void UnregisterThreadFreeBlockLists( FPerThreadFreeBlockLists* FreeBlockLists )
 	{
+		QUICK_SCOPE_CYCLE_COUNTER(STAT_FMallocBinned2_UnregisterThreadFreeBlockLists);
 		FScopeLock Lock(&GetFreeBlockListsRegistrationMutex());
 #if BINNED2_ALLOCATOR_STATS_VALIDATION
 		++RecursionCounter;
@@ -836,6 +873,8 @@ void* FMallocBinned2::MallocExternalSmall(SIZE_T Size, uint32 Alignment)
 		}
 	}
 
+	NOALLOC_SCOPE_CYCLE_COUNTER(STAT_FMallocBinned2_MallocExternalSmall);
+
 	FScopeLock Lock(&Mutex);
 
 	// Allocate from small object pool.
@@ -896,6 +935,8 @@ void* FMallocBinned2::MallocExternalLarge(SIZE_T Size, uint32 Alignment)
 	FPoolInfo* Pool;
 	void*      Result;
 	{
+		QUICK_SCOPE_CYCLE_COUNTER(STAT_FMallocBinned2_MallocExternalLarge);
+
 		FScopeLock Lock(&Mutex);
 
 		// Use OS for non-pooled allocations.
@@ -965,6 +1006,8 @@ void* FMallocBinned2::ReallocExternal(void* Ptr, SIZE_T NewSize, uint32 Alignmen
 		return Result;
 	}
 
+	NOALLOC_SCOPE_CYCLE_COUNTER(STAT_FMallocBinned2_ReallocExternal);
+
 	// Allocated from OS.
 	Mutex.Lock();
 	FPoolInfo* Pool = Private::FindPoolInfo(*this, Ptr);
@@ -1019,6 +1062,8 @@ void* FMallocBinned2::ReallocExternal(void* Ptr, SIZE_T NewSize, uint32 Alignmen
 
 void FMallocBinned2::FreeExternal(void* Ptr)
 {
+	NOALLOC_SCOPE_CYCLE_COUNTER(STAT_FMallocBinned2_FreeExternal);
+
 	if (!IsOSAllocation(Ptr))
 	{
 		check(Ptr); // null is 64k aligned so we should not be here
@@ -1106,6 +1151,7 @@ bool FMallocBinned2::GetAllocationSizeExternal(void* Ptr, SIZE_T& SizeOut)
 
 	FPoolInfo* Pool;
 	{
+		QUICK_SCOPE_CYCLE_COUNTER(STAT_FMallocBinned2_GetAllocationSizeExternal);
 		FScopeLock Lock(&Mutex);
 		Pool = Private::FindPoolInfo(*this, Ptr);
 	}
@@ -1147,6 +1193,7 @@ void FMallocBinned2::FPoolList::ValidateExhaustedPools()
 
 bool FMallocBinned2::ValidateHeap()
 {
+	QUICK_SCOPE_CYCLE_COUNTER(STAT_FMallocBinned2_ValidateHeap);
 	FScopeLock Lock(&Mutex);
 
 	for (FPoolTable& Table : SmallPoolTables)
@@ -1237,6 +1284,8 @@ void FMallocBinned2::Trim(bool bTrimThreadCaches)
 
 void FMallocBinned2::SetupTLSCachesOnCurrentThread()
 {
+	QUICK_SCOPE_CYCLE_COUNTER(STAT_FMallocBinned2_SetupTLSCachesOnCurrentThread);
+
 	if (!BINNED2_ALLOW_RUNTIME_TWEAKING && !GMallocBinned2PerThreadCaches)
 	{
 		return;
@@ -1251,6 +1300,7 @@ void FMallocBinned2::SetupTLSCachesOnCurrentThread()
 
 void FMallocBinned2::ClearAndDisableTLSCachesOnCurrentThread()
 {
+	QUICK_SCOPE_CYCLE_COUNTER(STAT_FMallocBinned2_ClearTLSCachesOnCurrentThread);
 	FlushCurrentThreadCache();
 	FPerThreadFreeBlockLists::ClearTLS();
 }
@@ -1370,6 +1420,7 @@ void FMallocBinned2::CanaryFail(const FFreeBlock* Block) const
 #if BINNED2_ALLOCATOR_STATS
 int64 FMallocBinned2::GetTotalAllocatedSmallPoolMemory() const
 {
+	QUICK_SCOPE_CYCLE_COUNTER(STAT_FMallocBinned2_GetTotalAllocatedSmallPoolMemory);
 	int64 FreeBlockAllocatedMemory = 0;
 	{
 		FScopeLock Lock(&Private::GetFreeBlockListsRegistrationMutex());
@@ -1387,6 +1438,7 @@ int64 FMallocBinned2::GetTotalAllocatedSmallPoolMemory() const
 void FMallocBinned2::GetAllocatorStats( FGenericMemoryStats& OutStats )
 {
 #if BINNED2_ALLOCATOR_STATS
+	QUICK_SCOPE_CYCLE_COUNTER(STAT_FMallocBinned2_GetAllocatorStats);
 
 	int64  TotalAllocatedSmallPoolMemory           = GetTotalAllocatedSmallPoolMemory();
 	int64  LocalAllocatedOSSmallPoolMemory         = AllocatedOSSmallPoolMemory.Load(EMemoryOrder::Relaxed);
diff --git a/Engine/Source/Runtime/Core/Public/Containers/ContainerAllocationPolicies.h b/Engine/Source/Runtime/Core/Public/Containers/ContainerAllocationPolicies.h
--- a/Engine/Source/Runtime/Core/Public/Containers/ContainerAllocationPolicies.h
+++ b/Engine/Source/Runtime/Core/Public/Containers/ContainerAllocationPolicies.h
@@ -469,7 +469,7 @@ template <> struct TBitsToSizeType<32> { using Type = int32; };
 template <> struct TBitsToSizeType<64> { using Type = int64; };
 
 /** The indirect allocation policy always allocates the elements indirectly. */
-template <int IndexSize>
+template <int IndexSize, typename BaseMallocType = FMemory>
 class TSizedHeapAllocator
 {
 public:
@@ -480,7 +480,7 @@ public:
 
 	class ForAnyElementType
 	{
-		template <int>
+		template <int, typename>
 		friend class TSizedHeapAllocator;
 
 	public:
@@ -502,7 +502,7 @@ public:
 
 			if (Data)
 			{
-				FMemory::Free(Data);
+				BaseMallocType::Free(Data);
 			}
 
 			Data = Other.Data;
@@ -526,7 +526,7 @@ public:
 		{
 			if(Data)
 			{
-				FMemory::Free(Data);
+				BaseMallocType::Free(Data);
 			}
 		}
 
@@ -541,7 +541,7 @@ public:
 			if (Data || NumElements)
 			{
 				//checkSlow(((uint64)NumElements*(uint64)ElementTypeInfo.GetSize() < (uint64)INT_MAX));
-				Data = (FScriptContainerElement*)FMemory::Realloc( Data, NumElements*NumBytesPerElement );
+				Data = (FScriptContainerElement*)BaseMallocType::Realloc( Data, NumElements*NumBytesPerElement );
 			}
 		}
 		FORCEINLINE void ResizeAllocation(SizeType PreviousNumElements, SizeType NumElements, SIZE_T NumBytesPerElement, uint32 AlignmentOfElement)
@@ -550,7 +550,7 @@ public:
 			if (Data || NumElements)
 			{
 				//checkSlow(((uint64)NumElements*(uint64)ElementTypeInfo.GetSize() < (uint64)INT_MAX));
-				Data = (FScriptContainerElement*)FMemory::Realloc( Data, NumElements*NumBytesPerElement, AlignmentOfElement );
+				Data = (FScriptContainerElement*)BaseMallocType::Realloc( Data, NumElements*NumBytesPerElement, AlignmentOfElement );
 			}
 		}
 		FORCEINLINE SizeType CalculateSlackReserve(SizeType NumElements, SIZE_T NumBytesPerElement) const
diff --git a/Engine/Source/Runtime/Core/Public/FramePro/FrameProProfiler.h b/Engine/Source/Runtime/Core/Public/FramePro/FrameProProfiler.h
--- a/Engine/Source/Runtime/Core/Public/FramePro/FrameProProfiler.h
+++ b/Engine/Source/Runtime/Core/Public/FramePro/FrameProProfiler.h
@@ -10,6 +10,9 @@
 class CORE_API FFrameProProfiler
 {
 public:
+	static void Initialize();
+	static void TearDown();
+
 	/** Called to mark the start of each frame  */
 	static void FrameStart();
 
@@ -28,6 +31,7 @@ public:
 	static void StopFrameProRecording();
 
 	static bool IsFrameProRecording();
+	static bool IsThreadContextReady();
 };
 
 #endif // FRAMEPRO_ENABLED
\ No newline at end of file
diff --git a/Engine/Source/Runtime/Launch/Private/LaunchEngineLoop.cpp b/Engine/Source/Runtime/Launch/Private/LaunchEngineLoop.cpp
--- a/Engine/Source/Runtime/Launch/Private/LaunchEngineLoop.cpp
+++ b/Engine/Source/Runtime/Launch/Private/LaunchEngineLoop.cpp
@@ -240,6 +240,7 @@ class FFeedbackContext;
 	#if ENABLE_VISUAL_LOG
 		#include "VisualLogger/VisualLogger.h"
 	#endif
+	#include "FramePro/FrameProProfiler.h"
 	#include "ProfilingDebugging/CsvProfiler.h"
 	#include "ProfilingDebugging/TracingProfiler.h"
 #endif
@@ -2413,6 +2414,9 @@ PRAGMA_DISABLE_DEPRECATION_WARNINGS
 	FTracingProfiler::Get()->Init();
 PRAGMA_ENABLE_DEPRECATION_WARNINGS
 #endif
+#if WITH_ENGINE && FRAMEPRO_ENABLED
+	FFrameProProfiler::Initialize();
+#endif // FRAMEPRO_ENABLED
 
 	// Start the application
 	{
@@ -4597,6 +4601,10 @@ void FEngineLoop::Exit()
 	RHIExit();
 
 	FPlatformMisc::ShutdownTaggedStorage();
+
+#if WITH_ENGINE && FRAMEPRO_ENABLED
+	FFrameProProfiler::TearDown();
+#endif // FRAMEPRO_ENABLED
 }
 
 
