Log fatal IOS errors to the log file, consistent with other platforms

diff --git a/Engine/Source/Runtime/Core/Private/IOS/IOSErrorOutputDevice.cpp b/Engine/Source/Runtime/Core/Private/IOS/IOSErrorOutputDevice.cpp
--- a/Engine/Source/Runtime/Core/Private/IOS/IOSErrorOutputDevice.cpp
+++ b/Engine/Source/Runtime/Core/Private/IOS/IOSErrorOutputDevice.cpp
@@ -45,6 +45,12 @@ void FIOSErrorOutputDevice::HandleError()
 	GIsRunning = 0;
 	GIsCriticalError = 1;
 	GLogConsole = NULL;
+	GErrorHist[UE_ARRAY_COUNT(GErrorHist) - 1] = 0;
+    
+	// Dump the error and flush the log.
+#if !NO_LOGGING
+	FDebug::LogFormattedMessageWithCallstack(LogIOS.GetCategoryName(), __FILE__, __LINE__, TEXT("=== Critical error: ==="), GErrorHist, ELogVerbosity::Error);
+#endif
 
 	GLog->PanicFlushThreadedLogs();
 }
