No longer pass TArrayView by const ref, also add overload for const argument type.

diff --git a/Engine/Source/Runtime/Core/Public/Containers/Set.h b/Engine/Source/Runtime/Core/Public/Containers/Set.h
--- a/Engine/Source/Runtime/Core/Public/Containers/Set.h
+++ b/Engine/Source/Runtime/Core/Public/Containers/Set.h
@@ -757,7 +757,17 @@ public:
 	}
 
 	template<typename ViewSizeType>
-	void Append(const TArrayView<ElementType, ViewSizeType>& InElements)
+	void Append(TArrayView<ElementType, ViewSizeType> InElements)
+	{
+		Reserve(Elements.Num() + InElements.Num());
+		for (const ElementType& Element : InElements)
+		{
+			Add(Element);
+		}
+	}
+
+	template<typename ViewSizeType>
+	void Append(TArrayView<const ElementType, ViewSizeType> InElements)
 	{
 		Reserve(Elements.Num() + InElements.Num());
 		for (const ElementType& Element : InElements)
