Client startup fix for BinnedMalloc2 trace scopes (CL 19091015): use no-allocation scopes in all BinnedMalloc2 functions.

diff --git a/Engine/Source/Runtime/Core/Private/HAL/MallocBinned2.cpp b/Engine/Source/Runtime/Core/Private/HAL/MallocBinned2.cpp
--- a/Engine/Source/Runtime/Core/Private/HAL/MallocBinned2.cpp
+++ b/Engine/Source/Runtime/Core/Private/HAL/MallocBinned2.cpp
@@ -608,7 +608,7 @@ struct FMallocBinned2::Private
 	}
 	static void RegisterThreadFreeBlockLists( FPerThreadFreeBlockLists* FreeBlockLists )
 	{
-		QUICK_SCOPE_CYCLE_COUNTER(STAT_FMallocBinned2_RegisterThreadFreeBlockLists);
+		NOALLOC_SCOPE_CYCLE_COUNTER(STAT_FMallocBinned2_RegisterThreadFreeBlockLists);
 		FScopeLock Lock(&GetFreeBlockListsRegistrationMutex());
 #if BINNED2_ALLOCATOR_STATS_VALIDATION
 		++RecursionCounter;
@@ -620,7 +620,7 @@ struct FMallocBinned2::Private
 	}
 	static void UnregisterThreadFreeBlockLists( FPerThreadFreeBlockLists* FreeBlockLists )
 	{
-		QUICK_SCOPE_CYCLE_COUNTER(STAT_FMallocBinned2_UnregisterThreadFreeBlockLists);
+		NOALLOC_SCOPE_CYCLE_COUNTER(STAT_FMallocBinned2_UnregisterThreadFreeBlockLists);
 		FScopeLock Lock(&GetFreeBlockListsRegistrationMutex());
 #if BINNED2_ALLOCATOR_STATS_VALIDATION
 		++RecursionCounter;
@@ -935,7 +935,7 @@ void* FMallocBinned2::MallocExternalLarge(SIZE_T Size, uint32 Alignment)
 	FPoolInfo* Pool;
 	void*      Result;
 	{
-		QUICK_SCOPE_CYCLE_COUNTER(STAT_FMallocBinned2_MallocExternalLarge);
+		NOALLOC_SCOPE_CYCLE_COUNTER(STAT_FMallocBinned2_MallocExternalLarge);
 
 		FScopeLock Lock(&Mutex);
 
@@ -1151,7 +1151,7 @@ bool FMallocBinned2::GetAllocationSizeExternal(void* Ptr, SIZE_T& SizeOut)
 
 	FPoolInfo* Pool;
 	{
-		QUICK_SCOPE_CYCLE_COUNTER(STAT_FMallocBinned2_GetAllocationSizeExternal);
+		NOALLOC_SCOPE_CYCLE_COUNTER(STAT_FMallocBinned2_GetAllocationSizeExternal);
 		FScopeLock Lock(&Mutex);
 		Pool = Private::FindPoolInfo(*this, Ptr);
 	}
@@ -1193,7 +1193,7 @@ void FMallocBinned2::FPoolList::ValidateExhaustedPools()
 
 bool FMallocBinned2::ValidateHeap()
 {
-	QUICK_SCOPE_CYCLE_COUNTER(STAT_FMallocBinned2_ValidateHeap);
+	NOALLOC_SCOPE_CYCLE_COUNTER(STAT_FMallocBinned2_ValidateHeap);
 	FScopeLock Lock(&Mutex);
 
 	for (FPoolTable& Table : SmallPoolTables)
@@ -1214,7 +1214,7 @@ void FMallocBinned2::FlushCurrentThreadCache()
 {
 	double StartTimeInner = FPlatformTime::Seconds();
 	TRACE_CPUPROFILER_EVENT_SCOPE(FMallocBinned2::FlushCurrentThreadCache);
-	QUICK_SCOPE_CYCLE_COUNTER(STAT_FMallocBinned2_FlushCurrentThreadCache);
+	NOALLOC_SCOPE_CYCLE_COUNTER(STAT_FMallocBinned2_FlushCurrentThreadCache);
 	FPerThreadFreeBlockLists* Lists = FPerThreadFreeBlockLists::Get();
 
 	double WaitForMutexTime = 0.0;
@@ -1250,7 +1250,7 @@ void FMallocBinned2::FlushCurrentThreadCache()
 
 void FMallocBinned2::Trim(bool bTrimThreadCaches)
 {
-	QUICK_SCOPE_CYCLE_COUNTER(STAT_FMallocBinned2_Trim);
+	NOALLOC_SCOPE_CYCLE_COUNTER(STAT_FMallocBinned2_Trim);
 
 	if (GMallocBinned2PerThreadCaches  &&  bTrimThreadCaches)
 	{
@@ -1284,7 +1284,7 @@ void FMallocBinned2::Trim(bool bTrimThreadCaches)
 
 void FMallocBinned2::SetupTLSCachesOnCurrentThread()
 {
-	QUICK_SCOPE_CYCLE_COUNTER(STAT_FMallocBinned2_SetupTLSCachesOnCurrentThread);
+	NOALLOC_SCOPE_CYCLE_COUNTER(STAT_FMallocBinned2_SetupTLSCachesOnCurrentThread);
 
 	if (!BINNED2_ALLOW_RUNTIME_TWEAKING && !GMallocBinned2PerThreadCaches)
 	{
@@ -1300,7 +1300,7 @@ void FMallocBinned2::SetupTLSCachesOnCurrentThread()
 
 void FMallocBinned2::ClearAndDisableTLSCachesOnCurrentThread()
 {
-	QUICK_SCOPE_CYCLE_COUNTER(STAT_FMallocBinned2_ClearTLSCachesOnCurrentThread);
+	NOALLOC_SCOPE_CYCLE_COUNTER(STAT_FMallocBinned2_ClearTLSCachesOnCurrentThread);
 	FlushCurrentThreadCache();
 	FPerThreadFreeBlockLists::ClearTLS();
 }
@@ -1420,7 +1420,7 @@ void FMallocBinned2::CanaryFail(const FFreeBlock* Block) const
 #if BINNED2_ALLOCATOR_STATS
 int64 FMallocBinned2::GetTotalAllocatedSmallPoolMemory() const
 {
-	QUICK_SCOPE_CYCLE_COUNTER(STAT_FMallocBinned2_GetTotalAllocatedSmallPoolMemory);
+	NOALLOC_SCOPE_CYCLE_COUNTER(STAT_FMallocBinned2_GetTotalAllocatedSmallPoolMemory);
 	int64 FreeBlockAllocatedMemory = 0;
 	{
 		FScopeLock Lock(&Private::GetFreeBlockListsRegistrationMutex());
@@ -1438,7 +1438,7 @@ int64 FMallocBinned2::GetTotalAllocatedSmallPoolMemory() const
 void FMallocBinned2::GetAllocatorStats( FGenericMemoryStats& OutStats )
 {
 #if BINNED2_ALLOCATOR_STATS
-	QUICK_SCOPE_CYCLE_COUNTER(STAT_FMallocBinned2_GetAllocatorStats);
+	NOALLOC_SCOPE_CYCLE_COUNTER(STAT_FMallocBinned2_GetAllocatorStats);
 
 	int64  TotalAllocatedSmallPoolMemory           = GetTotalAllocatedSmallPoolMemory();
 	int64  LocalAllocatedOSSmallPoolMemory         = AllocatedOSSmallPoolMemory.Load(EMemoryOrder::Relaxed);
