if VeryLargePageAllocator commit fails using OS large page allocation, it fallbacks to commit using OS small page.

Can be disabled with cvar VeryLargePageAllocator.UseFallback

#[review] [REDACTED_HANDLE], [REDACTED_HANDLE]
#rnx
#[robomerge] deadend

diff --git a/Engine/Source/Runtime/Core/Private/HAL/Allocators/CachedOSVeryLargePageAllocator.cpp b/Engine/Source/Runtime/Core/Private/HAL/Allocators/CachedOSVeryLargePageAllocator.cpp
--- a/Engine/Source/Runtime/Core/Private/HAL/Allocators/CachedOSVeryLargePageAllocator.cpp
+++ b/Engine/Source/Runtime/Core/Private/HAL/Allocators/CachedOSVeryLargePageAllocator.cpp
@@ -8,10 +8,17 @@
 
 #if UE_USE_VERYLARGEPAGEALLOCATOR
 
+static int32 GVeryLargePageAllocatorUseFallback = 1;
+static FAutoConsoleVariableRef CVarVeryLargePageAllocatorUseFallback(
+	TEXT("VeryLargePageAllocator.UseFallback"),
+	GVeryLargePageAllocatorUseFallback,
+	TEXT("Whether very large page allocator to use fallback allocator using small OS pages (default is 1)"));
+
+
 CORE_API bool GEnableVeryLargePageAllocator = true;
-void FCachedOSVeryLargePageAllocator::Init()
+void FCachedOSVeryLargePageAllocator::Init(bool bUseOSLargePages)
 {
-	Block = FPlatformMemory::FPlatformVirtualMemoryBlock::AllocateVirtual(AddressSpaceToReserve);
+	Block = FPlatformMemory::FPlatformVirtualMemoryBlock::AllocateVirtual(AddressSpaceToReserve, FPlatformMemory::FPlatformVirtualMemoryBlock::GetVirtualSizeAlignment(), FPlatformMemory::EVirtualMemoryPlatformType::OrdinaryCPU, false, bUseOSLargePages);
 	AddressSpaceReserved = (uintptr_t)Block.GetVirtualPointer();
 	AddressSpaceReservedEnd = AddressSpaceReserved + AddressSpaceToReserve;
 #if UE_VERYLARGEPAGEALLOCATOR_TAKEONALL64KBALLOCATIONS
@@ -45,6 +52,12 @@ void FCachedOSVeryLargePageAllocator::Init()
 	}
 #endif
 
+	if (bUseOSLargePages)
+	{
+		void* FAllbackAllocatorMemory = FPlatformMemory::BinnedAllocFromOS(sizeof(FCachedOSVeryLargePageAllocator));
+		FallbackAllocator = new (FAllbackAllocatorMemory) FCachedOSVeryLargePageAllocator(false);
+	}
+
 	if (!GEnableVeryLargePageAllocator)
 	{
 		bEnabled = false;
@@ -76,7 +89,34 @@ void* FCachedOSVeryLargePageAllocator::Allocate(SIZE_T Size, uint32 AllocationHi
 						FScopeUnlock FScopeUnlock(Mutex);
 #endif
 						LLM_PLATFORM_SCOPE(ELLMTag::FMalloc);
-						Block.Commit(LargePage->BaseAddress - AddressSpaceReserved, SizeOfLargePage);
+
+						if (GVeryLargePageAllocatorUseFallback != 1)
+						{
+							Block.Commit(LargePage->BaseAddress - AddressSpaceReserved, SizeOfLargePage, true);
+						}
+						else
+						{
+							if (FallbackAllocator == nullptr)
+							{
+								Block.Commit(LargePage->BaseAddress - AddressSpaceReserved, SizeOfLargePage, true);
+							}
+							else if (!Block.Commit(LargePage->BaseAddress - AddressSpaceReserved, SizeOfLargePage, false))
+							{
+								// Allocation failed, use the fallback allocator
+#if UE_ALLOW_OSMEMORYLOCKFREE
+								if (Mutex)
+								{
+									FScopeLock FScopeLock(Mutex);
+
+									LargePage->LinkHead(FreeLargePagesHead[LargePage->AllocationHint]);
+									return FallbackAllocator->Allocate(Size, AllocationHint, Mutex);
+								}
+#endif
+
+								LargePage->LinkHead(FreeLargePagesHead[LargePage->AllocationHint]);
+								return FallbackAllocator->Allocate(Size, AllocationHint, Mutex);
+							}
+						}
 					}
 					LargePage->LinkHead(UsedLargePagesWithSpaceHead[AllocationHint]);
 					CachedFree += SizeOfLargePage;
@@ -193,6 +233,11 @@ void FCachedOSVeryLargePageAllocator::Free(void* Ptr, SIZE_T Size, FCriticalSect
 #endif
 		}
 	}
+	else if (FallbackAllocator != nullptr)
+	{
+		// Allocation might have came from fallback allocator. Free it from there
+		FallbackAllocator->Free(Ptr, Size, Mutex, ThreadIsTimeCritical);
+	}
 	else
 	{
 		CachedOSPageAllocator.Free(Ptr, Size, Mutex, ThreadIsTimeCritical);
diff --git a/Engine/Source/Runtime/Core/Public/HAL/Allocators/CachedOSVeryLargePageAllocator.h b/Engine/Source/Runtime/Core/Public/HAL/Allocators/CachedOSVeryLargePageAllocator.h
--- a/Engine/Source/Runtime/Core/Public/HAL/Allocators/CachedOSVeryLargePageAllocator.h
+++ b/Engine/Source/Runtime/Core/Public/HAL/Allocators/CachedOSVeryLargePageAllocator.h
@@ -48,11 +48,12 @@ class FCachedOSVeryLargePageAllocator
 	static constexpr uint64 NumberOfSubPagesPerLargePage = (SizeOfLargePage / SizeOfSubPage);
 public:
 
-	FCachedOSVeryLargePageAllocator()
+	FCachedOSVeryLargePageAllocator(bool bUseOSLargePages = true)
 		: bEnabled(true)
 		, CachedFree(0)
+		, FallbackAllocator(nullptr)
 	{
-		Init();
+		Init(bUseOSLargePages);
 	}
 
 	~FCachedOSVeryLargePageAllocator()
@@ -77,12 +78,13 @@ public:
 		{
 			return true;
 		}
-		return false;
+
+		return (FallbackAllocator != nullptr) ? FallbackAllocator->IsPartOf(Ptr) : false;
 	}
 
 private:
 
-	void Init();
+	void Init(bool bUseOSLargePages);
 
 	bool bEnabled;
 	uintptr_t	AddressSpaceReserved;
@@ -92,6 +94,8 @@ private:
 
 	FPlatformMemory::FPlatformVirtualMemoryBlock Block;
 
+	FCachedOSVeryLargePageAllocator* FallbackAllocator;
+
 	struct FLargePage : public TIntrusiveLinkedList<FLargePage>
 	{
 		uintptr_t	FreeSubPages[NumberOfSubPagesPerLargePage];
