add new core delegates (OnParentBeginFork, OnParentPreFork, OnPostFork, OnEndFramePostFork) to allow interested parties a chance to take actions before and after the process fork point on dedicated servers

diff --git a/Engine/Source/Runtime/Core/Private/Misc/CoreDelegates.cpp b/Engine/Source/Runtime/Core/Private/Misc/CoreDelegates.cpp
--- a/Engine/Source/Runtime/Core/Private/Misc/CoreDelegates.cpp
+++ b/Engine/Source/Runtime/Core/Private/Misc/CoreDelegates.cpp
@@ -3,6 +3,7 @@
 // Core includes.
 #include "Misc/CoreDelegates.h"
 #include "Math/Vector.h"
+#include "Misc/Fork.h"
 
 
 //////////////////////////////////////////////////////////////////////////
@@ -243,3 +244,7 @@ void RegisterEncryptionKeyCallback(TEncryptionKeyFunc InCallback)
 	});
 }
 
+FSimpleMulticastDelegate FCoreDelegates::OnParentBeginFork;
+FSimpleMulticastDelegate FCoreDelegates::OnParentPreFork;
+FCoreDelegates::FProcessForkDelegate FCoreDelegates::OnPostFork;
+FSimpleMulticastDelegate FCoreDelegates::OnChildEndFramePostFork;
\ No newline at end of file
diff --git a/Engine/Source/Runtime/Core/Private/Unix/UnixPlatformProcess.cpp b/Engine/Source/Runtime/Core/Private/Unix/UnixPlatformProcess.cpp
--- a/Engine/Source/Runtime/Core/Private/Unix/UnixPlatformProcess.cpp
+++ b/Engine/Source/Runtime/Core/Private/Unix/UnixPlatformProcess.cpp
@@ -1323,6 +1323,8 @@ void FUnixPlatformProcess::TerminateProc( FProcHandle & ProcessHandle, bool Kill
 	}
 }
 
+static FDelegateHandle OnEndFrameHandle;
+
 /*
  * WaitAndFork on Unix
  *
@@ -1400,6 +1402,8 @@ FGenericPlatformProcess::EWaitAndForkResult FUnixPlatformProcess::WaitAndFork()
 		UE_LOG(LogHAL, Log, TEXT("WaitAndFork setting WaitAndForkResponseTimeout to %0.2f seconds."), WaitAndForkResponseTimeout);
 	}
 
+	FCoreDelegates::OnParentBeginFork.Broadcast();
+
 	// Set up a signal handler for the signal to fork()
 	{
 		struct sigaction Action;
@@ -1461,6 +1465,8 @@ FGenericPlatformProcess::EWaitAndForkResult FUnixPlatformProcess::WaitAndFork()
 
 			FDateTime SignalReceived = FDateTime::FromUnixTimestamp(FMath::FloorToInt64(SignalData.TimeSeconds));
 
+			FCoreDelegates::OnParentPreFork.Broadcast();
+
 			UE_LOG(LogHAL, Log, TEXT("[Parent] WaitAndFork processing child request %04hx-%04hx received at: %s"), Cookie, ChildIdx, *SignalReceived.ToString());
 
 			FMemoryStatsHolder CurrentMasterMemStats(FPlatformMemory::GetStats());
@@ -1579,6 +1585,9 @@ FGenericPlatformProcess::EWaitAndForkResult FUnixPlatformProcess::WaitAndFork()
 				UE_LOG(LogHAL, Log, TEXT("[Child] WaitAndFork child process %04hx-%04hx has started with pid %d."), Cookie, ChildIdx, GetCurrentProcessId());
 				FApp::PrintStartupLogMessages();
 
+				OnEndFrameHandle = FCoreDelegates::OnEndFrame.AddStatic(FUnixPlatformProcess::OnChildEndFramePostFork);
+				FCoreDelegates::OnPostFork.Broadcast(EForkProcessRole::Child);
+
 				// Children break out of the loop and return
 				RetVal = EWaitAndForkResult::Child;
 				break;
@@ -1591,6 +1600,8 @@ FGenericPlatformProcess::EWaitAndForkResult FUnixPlatformProcess::WaitAndFork()
 				// Parent
 				AllChildren.Emplace(ChildPID, SignalData.SignalValue);
 
+				FCoreDelegates::OnPostFork.Broadcast(EForkProcessRole::Parent);
+
 				UE_LOG(LogHAL, Log, TEXT("[Parent] WaitAndFork Successfully processed request %04hx-%04hx, made a child with pid %d! Total number of children: %d."), Cookie, ChildIdx, ChildPID, AllChildren.Num());
 			}
 		}
@@ -2041,3 +2052,11 @@ void FUnixPlatformProcess::CeaseBeingFirstInstance()
 	}
 #endif
 }
+
+void FUnixPlatformProcess::OnChildEndFramePostFork()
+{
+	FCoreDelegates::OnEndFrame.Remove(OnEndFrameHandle);
+	OnEndFrameHandle.Reset();
+
+	FCoreDelegates::OnChildEndFramePostFork.Broadcast();
+}
\ No newline at end of file
diff --git a/Engine/Source/Runtime/Core/Public/Misc/CoreDelegates.h b/Engine/Source/Runtime/Core/Public/Misc/CoreDelegates.h
--- a/Engine/Source/Runtime/Core/Public/Misc/CoreDelegates.h
+++ b/Engine/Source/Runtime/Core/Public/Misc/CoreDelegates.h
@@ -15,6 +15,7 @@
 class AActor;
 class Error;
 class IPakFile;
+enum class EForkProcessRole : uint8;
 
 // delegates for hotfixes
 namespace EHotfixDelegates
@@ -637,6 +638,17 @@ public:
 	DECLARE_DELEGATE_RetVal(TSharedPtr<class IPackageStore>, FCreatePackageStore);
 	static FCreatePackageStore CreatePackageStore;
 
+	// Called immediately before the parent process will start responding to signals to fork
+	static FSimpleMulticastDelegate OnParentBeginFork;
+	// Called each time immediately before the parent process forks itself
+	static FSimpleMulticastDelegate OnParentPreFork;
+
+	DECLARE_MULTICAST_DELEGATE_OneParam(FProcessForkDelegate, EForkProcessRole /* ProcessRole */);
+	// Called immediately after the process spawned a fork
+	static FProcessForkDelegate OnPostFork;
+	// Called at the end of the frame where the process spawned a fork
+	static FSimpleMulticastDelegate OnChildEndFramePostFork;
+
 private:
 
 	// Callbacks for hotfixes
diff --git a/Engine/Source/Runtime/Core/Public/Misc/Fork.h b/Engine/Source/Runtime/Core/Public/Misc/Fork.h
--- a/Engine/Source/Runtime/Core/Public/Misc/Fork.h
+++ b/Engine/Source/Runtime/Core/Public/Misc/Fork.h
@@ -12,6 +12,12 @@
 class CORE_API FRunnable;
 class CORE_API FRunnableThread;
 
+enum class EForkProcessRole : uint8
+{
+	Parent,
+	Child,
+};
+
 /**
  * Helper functions for processes that fork in order to share memory pages.
  *
diff --git a/Engine/Source/Runtime/Core/Public/Unix/UnixPlatformProcess.h b/Engine/Source/Runtime/Core/Public/Unix/UnixPlatformProcess.h
--- a/Engine/Source/Runtime/Core/Public/Unix/UnixPlatformProcess.h
+++ b/Engine/Source/Runtime/Core/Public/Unix/UnixPlatformProcess.h
@@ -297,6 +297,7 @@ struct CORE_API FUnixPlatformProcess : public FGenericPlatformProcess
 	static void ExploreFolder( const TCHAR* FilePath );
 	static bool LaunchFileInDefaultExternalApplication( const TCHAR* FileName, const TCHAR* Parms = NULL, ELaunchVerb::Type Verb = ELaunchVerb::Open, bool bPromptToOpenOnFailure = true);
 	static bool IsFirstInstance();
+	static void OnChildEndFramePostFork();
 
 	/**
 	 * @brief Releases locks that we held for IsFirstInstance check
