Fixed LargePageAllocator CIS issue

#rnx

diff --git a/Engine/Source/Runtime/Core/Private/HAL/Allocators/CachedOSVeryLargePageAllocator.cpp b/Engine/Source/Runtime/Core/Private/HAL/Allocators/CachedOSVeryLargePageAllocator.cpp
--- a/Engine/Source/Runtime/Core/Private/HAL/Allocators/CachedOSVeryLargePageAllocator.cpp
+++ b/Engine/Source/Runtime/Core/Private/HAL/Allocators/CachedOSVeryLargePageAllocator.cpp
@@ -8,17 +8,10 @@
 
 #if UE_USE_VERYLARGEPAGEALLOCATOR
 
-static int32 GVeryLargePageAllocatorUseFallback = 1;
-static FAutoConsoleVariableRef CVarVeryLargePageAllocatorUseFallback(
-	TEXT("VeryLargePageAllocator.UseFallback"),
-	GVeryLargePageAllocatorUseFallback,
-	TEXT("Whether very large page allocator to use fallback allocator using small OS pages (default is 1)"));
-
-
 CORE_API bool GEnableVeryLargePageAllocator = true;
-void FCachedOSVeryLargePageAllocator::Init(bool bUseOSLargePages)
+void FCachedOSVeryLargePageAllocator::Init()
 {
-	Block = FPlatformMemory::FPlatformVirtualMemoryBlock::AllocateVirtual(AddressSpaceToReserve, FPlatformMemory::FPlatformVirtualMemoryBlock::GetVirtualSizeAlignment(), FPlatformMemory::EVirtualMemoryPlatformType::OrdinaryCPU, false, bUseOSLargePages);
+	Block = FPlatformMemory::FPlatformVirtualMemoryBlock::AllocateVirtual(AddressSpaceToReserve);
 	AddressSpaceReserved = (uintptr_t)Block.GetVirtualPointer();
 	AddressSpaceReservedEnd = AddressSpaceReserved + AddressSpaceToReserve;
 #if UE_VERYLARGEPAGEALLOCATOR_TAKEONALL64KBALLOCATIONS
@@ -35,12 +28,12 @@ void FCachedOSVeryLargePageAllocator::Init(bool bUseOSLargePages)
 		UsedLargePagesHead[i] = nullptr;
 	}
 #if UE_VERYLARGEPAGEALLOCATOR_TAKEONALL64KBALLOCATIONS
-	for (int i = 0; i < NumberOfLargePages/2; i++)
+	for (int i = 0; i < NumberOfLargePages / 2; i++)
 #else
 	for (int i = 0; i < NumberOfLargePages; i++)
 #endif
 	{
-		LargePagesArray[i].Init((void*)((uintptr_t)AddressSpaceReserved + (i*SizeOfLargePage)));
+		LargePagesArray[i].Init((void*)((uintptr_t)AddressSpaceReserved + (i * SizeOfLargePage)));
 		LargePagesArray[i].LinkHead(FreeLargePagesHead[FMemory::AllocationHints::SmallPool]);
 	}
 
@@ -52,12 +45,6 @@ void FCachedOSVeryLargePageAllocator::Init(bool bUseOSLargePages)
 	}
 #endif
 
-	if (bUseOSLargePages)
-	{
-		void* FAllbackAllocatorMemory = FPlatformMemory::BinnedAllocFromOS(sizeof(FCachedOSVeryLargePageAllocator));
-		FallbackAllocator = new (FAllbackAllocatorMemory) FCachedOSVeryLargePageAllocator(false);
-	}
-
 	if (!GEnableVeryLargePageAllocator)
 	{
 		bEnabled = false;
@@ -89,34 +76,7 @@ void* FCachedOSVeryLargePageAllocator::Allocate(SIZE_T Size, uint32 AllocationHi
 						FScopeUnlock FScopeUnlock(Mutex);
 #endif
 						LLM_PLATFORM_SCOPE(ELLMTag::FMalloc);
-
-						if (GVeryLargePageAllocatorUseFallback != 1)
-						{
-							Block.Commit(LargePage->BaseAddress - AddressSpaceReserved, SizeOfLargePage, true);
-						}
-						else
-						{
-							if (FallbackAllocator == nullptr)
-							{
-								Block.Commit(LargePage->BaseAddress - AddressSpaceReserved, SizeOfLargePage, true);
-							}
-							else if (!Block.Commit(LargePage->BaseAddress - AddressSpaceReserved, SizeOfLargePage, false))
-							{
-								// Allocation failed, use the fallback allocator
-#if UE_ALLOW_OSMEMORYLOCKFREE
-								if (Mutex)
-								{
-									FScopeLock FScopeLock(Mutex);
-
-									LargePage->LinkHead(FreeLargePagesHead[LargePage->AllocationHint]);
-									return FallbackAllocator->Allocate(Size, AllocationHint, Mutex);
-								}
-#endif
-
-								LargePage->LinkHead(FreeLargePagesHead[LargePage->AllocationHint]);
-								return FallbackAllocator->Allocate(Size, AllocationHint, Mutex);
-							}
-						}
+						Block.Commit(LargePage->BaseAddress - AddressSpaceReserved, SizeOfLargePage);
 					}
 					LargePage->LinkHead(UsedLargePagesWithSpaceHead[AllocationHint]);
 					CachedFree += SizeOfLargePage;
@@ -233,11 +193,6 @@ void FCachedOSVeryLargePageAllocator::Free(void* Ptr, SIZE_T Size, FCriticalSect
 #endif
 		}
 	}
-	else if (FallbackAllocator != nullptr)
-	{
-		// Allocation might have came from fallback allocator. Free it from there
-		FallbackAllocator->Free(Ptr, Size, Mutex, ThreadIsTimeCritical);
-	}
 	else
 	{
 		CachedOSPageAllocator.Free(Ptr, Size, Mutex, ThreadIsTimeCritical);
diff --git a/Engine/Source/Runtime/Core/Public/HAL/Allocators/CachedOSVeryLargePageAllocator.h b/Engine/Source/Runtime/Core/Public/HAL/Allocators/CachedOSVeryLargePageAllocator.h
--- a/Engine/Source/Runtime/Core/Public/HAL/Allocators/CachedOSVeryLargePageAllocator.h
+++ b/Engine/Source/Runtime/Core/Public/HAL/Allocators/CachedOSVeryLargePageAllocator.h
@@ -37,7 +37,7 @@ class FCachedOSVeryLargePageAllocator
 	// we make the address space twice as big as we need and use the 1st have for small pool allocations, the 2nd half is used for other allocations that are still == SizeOfSubPage
 #if UE_VERYLARGEPAGEALLOCATOR_TAKEONALL64KBALLOCATIONS
 	static constexpr uint64 AddressSpaceToReserve = ((1024LL * 1024LL * 1024LL) * UE_VERYLARGEPAGEALLOCATOR_RESERVED_SIZE_IN_GB * 2LL);
-	static constexpr uint64 AddressSpaceToReserveForSmallPool = AddressSpaceToReserve/2;
+	static constexpr uint64 AddressSpaceToReserveForSmallPool = AddressSpaceToReserve / 2;
 #else
 	static constexpr uint64 AddressSpaceToReserve = ((1024 * 1024 * 1024LL) * UE_VERYLARGEPAGEALLOCATOR_RESERVED_SIZE_IN_GB);
 	static constexpr uint64 AddressSpaceToReserveForSmallPool = AddressSpaceToReserve;
@@ -48,12 +48,11 @@ class FCachedOSVeryLargePageAllocator
 	static constexpr uint64 NumberOfSubPagesPerLargePage = (SizeOfLargePage / SizeOfSubPage);
 public:
 
-	FCachedOSVeryLargePageAllocator(bool bUseOSLargePages = true)
+	FCachedOSVeryLargePageAllocator()
 		: bEnabled(true)
 		, CachedFree(0)
-		, FallbackAllocator(nullptr)
 	{
-		Init(bUseOSLargePages);
+		Init();
 	}
 
 	~FCachedOSVeryLargePageAllocator()
@@ -78,13 +77,12 @@ public:
 		{
 			return true;
 		}
-
-		return (FallbackAllocator != nullptr) ? FallbackAllocator->IsPartOf(Ptr) : false;
+		return false;
 	}
 
 private:
 
-	void Init(bool bUseOSLargePages);
+	void Init();
 
 	bool bEnabled;
 	uintptr_t	AddressSpaceReserved;
@@ -94,8 +92,6 @@ private:
 
 	FPlatformMemory::FPlatformVirtualMemoryBlock Block;
 
-	FCachedOSVeryLargePageAllocator* FallbackAllocator;
-
 	struct FLargePage : public TIntrusiveLinkedList<FLargePage>
 	{
 		uintptr_t	FreeSubPages[NumberOfSubPagesPerLargePage];
@@ -132,11 +128,11 @@ private:
 		}
 	};
 
-	FLargePage*	FreeLargePagesHead[FMemory::AllocationHints::Max];				// no backing store
+	FLargePage* FreeLargePagesHead[FMemory::AllocationHints::Max];				// no backing store
 
-	FLargePage*	UsedLargePagesHead[FMemory::AllocationHints::Max];				// has backing store and is full
+	FLargePage* UsedLargePagesHead[FMemory::AllocationHints::Max];				// has backing store and is full
 
-	FLargePage*	UsedLargePagesWithSpaceHead[FMemory::AllocationHints::Max];	// has backing store and still has room
+	FLargePage* UsedLargePagesWithSpaceHead[FMemory::AllocationHints::Max];	// has backing store and still has room
 
 	FLargePage	LargePagesArray[NumberOfLargePages];
 
