TEXT(""_SV) changed to TEXTVIEW("") for UTF-8 compatibility.

diff --git a/Engine/Source/Runtime/Core/Private/Tests/Serialization/CompactBinaryWriterTest.cpp b/Engine/Source/Runtime/Core/Private/Tests/Serialization/CompactBinaryWriterTest.cpp
--- a/Engine/Source/Runtime/Core/Private/Tests/Serialization/CompactBinaryWriterTest.cpp
+++ b/Engine/Source/Runtime/Core/Private/Tests/Serialization/CompactBinaryWriterTest.cpp
@@ -321,7 +321,7 @@ bool FCbWriterStringTest::RunTest(const FString& Parameters)
 	{
 		Writer.Reset();
 		Writer.SetName("String"_ASV).AddString("Value"_ASV);
-		Writer.SetName("String"_ASV).AddString(TEXT("Value"_SV));
+		Writer.SetName("String"_ASV).AddString(TEXTVIEW("Value"));
 		FCbFieldIterator Fields = Writer.Save();
 		if (TestEqual(TEXT("FCbWriter(String, Basic) Validate"), ValidateCompactBinaryRange(Fields.GetOuterBuffer(), ECbValidateMode::All), ECbValidateError::None))
 		{
@@ -1065,7 +1065,7 @@ bool FCbWriterStreamTest::RunTest(const FString& Parameters)
 		Writer
 			<< "AnsiString"_ASV << "AnsiValue"_ASV
 			<< "AnsiStringLiteral"_ASV << "AnsiValue"
-			<< "WideString"_ASV << TEXT("WideValue"_SV)
+			<< "WideString"_ASV << TEXTVIEW("WideValue")
 			<< "WideStringLiteral"_ASV << TEXT("WideValue");
 		Writer.EndObject();
 
