TEXT() concatenation fixes, for UTF-8 mode.

diff --git a/Engine/Source/Runtime/Core/Private/Containers/StringTest.cpp b/Engine/Source/Runtime/Core/Private/Containers/StringTest.cpp
--- a/Engine/Source/Runtime/Core/Private/Containers/StringTest.cpp
+++ b/Engine/Source/Runtime/Core/Private/Containers/StringTest.cpp
@@ -541,66 +541,66 @@ bool FStringPathConcatCompoundOperatorTest::RunTest(const FString& Parameters)
 	TStringBuilder<128> FilenameStringBuilder; FilenameStringBuilder << Filename;
 	TStringBuilder<128> FilenameWithLeadingSlashStringBuilder; FilenameWithLeadingSlashStringBuilder << FilenameWithLeadingSlash;
 	
-#define TEST_EMPTYPATH_EMPTYFILE(Type, Input)						{ FString EmptyPathString; EmptyPathString /= Input; TestTrue(Type TEXT(": EmptyPath/EmptyFilename to be empty"), EmptyPathString.IsEmpty()); }
-#define TEST_VALIDPATH_EMPTYFILE(Type, Input)						{ FString Result(Path); Result /= Input; TestEqual(Type TEXT(": ValidPath/EmptyFilename result to be"), Result, PathWithTrailingSlash); } \
-																	{ FString Result(PathWithTrailingSlash); Result /= Input; TestEqual(Type TEXT(" (with extra /): ValidPath/EmptyFilename result to be"), Result, PathWithTrailingSlash); }	
-#define TEST_EMPTYPATH_VALIDFILE(Type, Input)						{ FString Result; Result /= Input; TestEqual(Type TEXT(": EmptyPath/ValidFilename"), Result, Filename); }
-#define TEST_VALIDPATH_VALIDFILE(Type, Path, File)					{ FString Result(Path); Result /= File; TestEqual(Type TEXT(": ValidPath/ValidFilename"), Result, CombinedPath); }
-#define TEST_VALIDPATH_VALIDFILE_DOUBLE_SEPARATOR(Type, Path, File)	{ FString Result(Path); Result /= File; TestEqual(Type TEXT(": ValidPath//ValidFilename"), Result, CombinedPathWithDoubleSeparator); }
+#define TEST_EMPTYPATH_EMPTYFILE(Type, Input)						{ FString EmptyPathString; EmptyPathString /= Input; TestTrue(TEXT(Type ": EmptyPath/EmptyFilename to be empty"), EmptyPathString.IsEmpty()); }
+#define TEST_VALIDPATH_EMPTYFILE(Type, Input)						{ FString Result(Path); Result /= Input; TestEqual(TEXT(Type ": ValidPath/EmptyFilename result to be"), Result, PathWithTrailingSlash); } \
+																	{ FString Result(PathWithTrailingSlash); Result /= Input; TestEqual(TEXT(Type " (with extra /): ValidPath/EmptyFilename result to be"), Result, PathWithTrailingSlash); }	
+#define TEST_EMPTYPATH_VALIDFILE(Type, Input)						{ FString Result; Result /= Input; TestEqual(TEXT(Type ": EmptyPath/ValidFilename"), Result, Filename); }
+#define TEST_VALIDPATH_VALIDFILE(Type, Path, File)					{ FString Result(Path); Result /= File; TestEqual(TEXT(Type ": ValidPath/ValidFilename"), Result, CombinedPath); }
+#define TEST_VALIDPATH_VALIDFILE_DOUBLE_SEPARATOR(Type, Path, File)	{ FString Result(Path); Result /= File; TestEqual(TEXT(Type ": ValidPath//ValidFilename"), Result, CombinedPathWithDoubleSeparator); }
 
 	// Test empty path /= empty file
-	TEST_EMPTYPATH_EMPTYFILE(TEXT("NullString"), FString());
-	TEST_EMPTYPATH_EMPTYFILE(TEXT("EmptyString"), FString(TEXT("")));
-	TEST_EMPTYPATH_EMPTYFILE(TEXT("EmptyAnsiLiteralString"), "");
-	TEST_EMPTYPATH_EMPTYFILE(TEXT("EmptyLiteralString"), TEXT(""));
-	TEST_EMPTYPATH_EMPTYFILE(TEXT("NullStringView"), FStringView());
-	TEST_EMPTYPATH_EMPTYFILE(TEXT("EmptyStringView"), FStringView(TEXT("")));
-	TEST_EMPTYPATH_EMPTYFILE(TEXT("EmptyStringBuilder"), EmptyStringBuilder);
+	TEST_EMPTYPATH_EMPTYFILE("NullString", FString());
+	TEST_EMPTYPATH_EMPTYFILE("EmptyString", FString(TEXT("")));
+	TEST_EMPTYPATH_EMPTYFILE("EmptyAnsiLiteralString", "");
+	TEST_EMPTYPATH_EMPTYFILE("EmptyLiteralString", TEXT(""));
+	TEST_EMPTYPATH_EMPTYFILE("NullStringView", FStringView());
+	TEST_EMPTYPATH_EMPTYFILE("EmptyStringView", FStringView(TEXT("")));
+	TEST_EMPTYPATH_EMPTYFILE("EmptyStringBuilder", EmptyStringBuilder);
 
 	// Test valid path /= empty file
-	TEST_VALIDPATH_EMPTYFILE(TEXT("NullString"), FString());
-	TEST_VALIDPATH_EMPTYFILE(TEXT("EmptyString"), FString(TEXT("")));
-	TEST_VALIDPATH_EMPTYFILE(TEXT("EmptyAnsiLiteralString"), "");
-	TEST_VALIDPATH_EMPTYFILE(TEXT("EmptyLiteralString"), TEXT(""));
-	TEST_VALIDPATH_EMPTYFILE(TEXT("NullStringView"), FStringView());
-	TEST_VALIDPATH_EMPTYFILE(TEXT("EmptyStringView"), FStringView(TEXT("")));
-	TEST_VALIDPATH_EMPTYFILE(TEXT("EmptyStringBuilder"), EmptyStringBuilder);
+	TEST_VALIDPATH_EMPTYFILE("NullString", FString());
+	TEST_VALIDPATH_EMPTYFILE("EmptyString", FString(TEXT("")));
+	TEST_VALIDPATH_EMPTYFILE("EmptyAnsiLiteralString", "");
+	TEST_VALIDPATH_EMPTYFILE("EmptyLiteralString", TEXT(""));
+	TEST_VALIDPATH_EMPTYFILE("NullStringView", FStringView());
+	TEST_VALIDPATH_EMPTYFILE("EmptyStringView", FStringView(TEXT("")));
+	TEST_VALIDPATH_EMPTYFILE("EmptyStringBuilder", EmptyStringBuilder);
 	
 	// Test empty path /= valid file
-	TEST_EMPTYPATH_VALIDFILE(TEXT("String"), FString(Filename));
-	TEST_EMPTYPATH_VALIDFILE(TEXT("LiteralString"), Filename);
-	TEST_EMPTYPATH_VALIDFILE(TEXT("LiteralAnsiString"), AnsiFilename);
-	TEST_EMPTYPATH_VALIDFILE(TEXT("StringView"), FStringView(Filename));
+	TEST_EMPTYPATH_VALIDFILE("String", FString(Filename));
+	TEST_EMPTYPATH_VALIDFILE("LiteralString", Filename);
+	TEST_EMPTYPATH_VALIDFILE("LiteralAnsiString", AnsiFilename);
+	TEST_EMPTYPATH_VALIDFILE("StringView", FStringView(Filename));
 	
 	// Test valid path /= valid file
-	TEST_VALIDPATH_VALIDFILE(TEXT("String"), Path, FString(Filename));
-	TEST_VALIDPATH_VALIDFILE(TEXT("LiteralString"), Path, Filename);
-	TEST_VALIDPATH_VALIDFILE(TEXT("LiteralAnsiString"), Path, AnsiFilename);
-	TEST_VALIDPATH_VALIDFILE(TEXT("StringView"), Path, FStringView(Filename));
-	TEST_VALIDPATH_VALIDFILE(TEXT("StringBuilder"), Path, FilenameStringBuilder);
+	TEST_VALIDPATH_VALIDFILE("String", Path, FString(Filename));
+	TEST_VALIDPATH_VALIDFILE("LiteralString", Path, Filename);
+	TEST_VALIDPATH_VALIDFILE("LiteralAnsiString", Path, AnsiFilename);
+	TEST_VALIDPATH_VALIDFILE("StringView", Path, FStringView(Filename));
+	TEST_VALIDPATH_VALIDFILE("StringBuilder", Path, FilenameStringBuilder);
 
 	// Test valid path (ending in /) /= valid file
-	TEST_VALIDPATH_VALIDFILE(TEXT("String (path with extra /)"), PathWithTrailingSlash, FString(Filename));
-	TEST_VALIDPATH_VALIDFILE(TEXT("LiteralString (path with extra /)"), PathWithTrailingSlash, Filename);
-	TEST_VALIDPATH_VALIDFILE(TEXT("LiteralAnsiString (path with extra /)"), PathWithTrailingSlash, AnsiFilename);
-	TEST_VALIDPATH_VALIDFILE(TEXT("StringView (path with extra /)"), PathWithTrailingSlash, FStringView(Filename));
-	TEST_VALIDPATH_VALIDFILE(TEXT("StringBuilder (path with extra /)"), PathWithTrailingSlash, FilenameStringBuilder);
+	TEST_VALIDPATH_VALIDFILE("String (path with extra /)", PathWithTrailingSlash, FString(Filename));
+	TEST_VALIDPATH_VALIDFILE("LiteralString (path with extra /)", PathWithTrailingSlash, Filename);
+	TEST_VALIDPATH_VALIDFILE("LiteralAnsiString (path with extra /)", PathWithTrailingSlash, AnsiFilename);
+	TEST_VALIDPATH_VALIDFILE("StringView (path with extra /)", PathWithTrailingSlash, FStringView(Filename));
+	TEST_VALIDPATH_VALIDFILE("StringBuilder (path with extra /)", PathWithTrailingSlash, FilenameStringBuilder);
 	
 	// Test valid path / valid path + file (starting with /)
-	TEST_VALIDPATH_VALIDFILE(TEXT("String (filename with extra /)"), Path, FString(FilenameWithLeadingSlash));
-	TEST_VALIDPATH_VALIDFILE(TEXT("LiteralString (filename with extra /)"), Path, FilenameWithLeadingSlash);
-	TEST_VALIDPATH_VALIDFILE(TEXT("LiteralAnsiString (filename with extra /)"), Path, AnsiFilenameWithLeadingSlash);
-	TEST_VALIDPATH_VALIDFILE(TEXT("StringView (filename with extra /)"), Path, FStringView(FilenameWithLeadingSlash));
-	TEST_VALIDPATH_VALIDFILE(TEXT("StringBuilder (filename with extra /)"), Path, FilenameWithLeadingSlashStringBuilder);
+	TEST_VALIDPATH_VALIDFILE("String (filename with extra /)", Path, FString(FilenameWithLeadingSlash));
+	TEST_VALIDPATH_VALIDFILE("LiteralString (filename with extra /)", Path, FilenameWithLeadingSlash);
+	TEST_VALIDPATH_VALIDFILE("LiteralAnsiString (filename with extra /)", Path, AnsiFilenameWithLeadingSlash);
+	TEST_VALIDPATH_VALIDFILE("StringView (filename with extra /)", Path, FStringView(FilenameWithLeadingSlash));
+	TEST_VALIDPATH_VALIDFILE("StringBuilder (filename with extra /)", Path, FilenameWithLeadingSlashStringBuilder);
 	
 	// Appending a file name that starts with a / to a directory that ends with a / will not remove the erroneous / and so 
 	// will end up with // in the path, these tests are to show this behavior
 	// For example "path/" /= "/file.txt" will result in "path//file.txt" not "path/file.txt"
-	TEST_VALIDPATH_VALIDFILE_DOUBLE_SEPARATOR(TEXT("String (path and filename with extra /)"), PathWithTrailingSlash, FString(FilenameWithLeadingSlash));
-	TEST_VALIDPATH_VALIDFILE_DOUBLE_SEPARATOR(TEXT("LiteralString (path and filename with extra /)"), PathWithTrailingSlash, FilenameWithLeadingSlash);
-	TEST_VALIDPATH_VALIDFILE_DOUBLE_SEPARATOR(TEXT("LiteralAnsiString (path and filename with extra /)"), PathWithTrailingSlash, AnsiFilenameWithLeadingSlash);
-	TEST_VALIDPATH_VALIDFILE_DOUBLE_SEPARATOR(TEXT("StringView (path and filename with extra /)"), PathWithTrailingSlash, FStringView(FilenameWithLeadingSlash));
-	TEST_VALIDPATH_VALIDFILE_DOUBLE_SEPARATOR(TEXT("StringBuilder (path and filename with extra /)"), PathWithTrailingSlash, FilenameWithLeadingSlashStringBuilder);
+	TEST_VALIDPATH_VALIDFILE_DOUBLE_SEPARATOR("String (path and filename with extra /)", PathWithTrailingSlash, FString(FilenameWithLeadingSlash));
+	TEST_VALIDPATH_VALIDFILE_DOUBLE_SEPARATOR("LiteralString (path and filename with extra /)", PathWithTrailingSlash, FilenameWithLeadingSlash);
+	TEST_VALIDPATH_VALIDFILE_DOUBLE_SEPARATOR("LiteralAnsiString (path and filename with extra /)", PathWithTrailingSlash, AnsiFilenameWithLeadingSlash);
+	TEST_VALIDPATH_VALIDFILE_DOUBLE_SEPARATOR("StringView (path and filename with extra /)", PathWithTrailingSlash, FStringView(FilenameWithLeadingSlash));
+	TEST_VALIDPATH_VALIDFILE_DOUBLE_SEPARATOR("StringBuilder (path and filename with extra /)", PathWithTrailingSlash, FilenameWithLeadingSlashStringBuilder);
 
 #undef TEST_EMPTYPATH_EMPTYFILE
 #undef TEST_VALIDPATH_EMPTYFILE
diff --git a/Engine/Source/Runtime/Core/Public/Memory/MemoryView.h b/Engine/Source/Runtime/Core/Public/Memory/MemoryView.h
--- a/Engine/Source/Runtime/Core/Public/Memory/MemoryView.h
+++ b/Engine/Source/Runtime/Core/Public/Memory/MemoryView.h
@@ -213,8 +213,8 @@ public:
 	/** Copies bytes from the input view into this view, and returns the remainder of this view. */
 	inline TMemoryView CopyFrom(FMemoryView InView) const
 	{
-		checkf(InView.Size <= Size, TEXT("Failed to copy from a view of %" UINT64_FMT " bytes ")
-			TEXT("to a view of %" UINT64_FMT " bytes."), InView.Size, Size);
+		checkf(InView.Size <= Size, TEXT("Failed to copy from a view of %" UINT64_FMT " bytes "
+			"to a view of %" UINT64_FMT " bytes."), InView.Size, Size);
 		FMemory::Memcpy(Data, InView.Data, InView.Size);
 		return RightChop(InView.Size);
 	}
