FConcurrentLinearAllocator tags' TagNames changed to char to avoid problems with constexpr TCHARs in UTF-8 mode.

diff --git a/Engine/Source/Runtime/Core/Public/Async/Fundamental/TaskDelegate.h b/Engine/Source/Runtime/Core/Public/Async/Fundamental/TaskDelegate.h
--- a/Engine/Source/Runtime/Core/Public/Async/Fundamental/TaskDelegate.h
+++ b/Engine/Source/Runtime/Core/Public/Async/Fundamental/TaskDelegate.h
@@ -15,7 +15,7 @@ namespace LowLevelTasks
 		static constexpr bool AllowOversizedBlocks = true;
 		static constexpr bool RequiresAccurateSize = false;
 		static constexpr bool InlineBlockAllocation = true;
-		static constexpr const TCHAR* TagName = TEXT("LowLevelTasksLinear");
+		static constexpr const char* TagName = "LowLevelTasksLinear";
 
 		using Allocator = TBlockAllocationCache<BlockSize, FAlignedAllocator>;
 	};
diff --git a/Engine/Source/Runtime/Core/Public/Async/TaskGraphInterfaces.h b/Engine/Source/Runtime/Core/Public/Async/TaskGraphInterfaces.h
--- a/Engine/Source/Runtime/Core/Public/Async/TaskGraphInterfaces.h
+++ b/Engine/Source/Runtime/Core/Public/Async/TaskGraphInterfaces.h
@@ -456,7 +456,7 @@ struct FTaskGraphBlockAllocationTag : FDefaultBlockAllocationTag
 	static constexpr bool AllowOversizedBlocks = false;
 	static constexpr bool RequiresAccurateSize = false;
 	static constexpr bool InlineBlockAllocation = true;
-	static constexpr const TCHAR* TagName = TEXT("TaskGraphLinear");
+	static constexpr const char* TagName = "TaskGraphLinear";
 
 	using Allocator = TBlockAllocationCache<BlockSize, FAlignedAllocator>;
 };
diff --git a/Engine/Source/Runtime/Core/Public/Experimental/ConcurrentLinearAllocator.h b/Engine/Source/Runtime/Core/Public/Experimental/ConcurrentLinearAllocator.h
--- a/Engine/Source/Runtime/Core/Public/Experimental/ConcurrentLinearAllocator.h
+++ b/Engine/Source/Runtime/Core/Public/Experimental/ConcurrentLinearAllocator.h
@@ -132,7 +132,7 @@ struct FDefaultBlockAllocationTag
 	static constexpr bool AllowOversizedBlocks = true;  // The allocator supports oversized Blocks and will store them in a seperate Block with counter 1
 	static constexpr bool RequiresAccurateSize = true;  // GetAllocationSize returning the accurate size of the allocation otherwise it could be relaxed to return the size to the end of the Block
 	static constexpr bool InlineBlockAllocation = false;  // Inline or Noinline the BlockAllocation which can have an impact on Performance
-	static constexpr const TCHAR* TagName = TEXT("DefaultLinear");
+	static constexpr const char* TagName = "DefaultLinear";
 
 	using Allocator = FOsAllocator;
 };
