Added logging to indicate when log levels have been raised to verbose levels to help track down log spam.

diff --git a/Engine/Source/Runtime/Core/Private/Logging/LogSuppressionInterface.cpp b/Engine/Source/Runtime/Core/Private/Logging/LogSuppressionInterface.cpp
--- a/Engine/Source/Runtime/Core/Private/Logging/LogSuppressionInterface.cpp
+++ b/Engine/Source/Runtime/Core/Private/Logging/LogSuppressionInterface.cpp
@@ -178,6 +178,8 @@ class FLogSuppressionImplementation: public FLogSuppressionInterface, private FS
 					CategoryVerbosities.Add(It.Value());
 				}					
 			}
+
+			const uint8 OriginalVerbosityLevel = Value & ELogVerbosity::VerbosityMask;
 			if (CommandParts.Num() == 1)
 			{
 				// only possibility is the reset and toggle command which is meaningless at boot
@@ -223,7 +225,6 @@ class FLogSuppressionImplementation: public FLogSuppressionInterface, private FS
 			}
 			else
 			{
-
 				// now we have the current value, lets change it!
 				for (int32 PartIndex = 1; PartIndex < CommandParts.Num(); PartIndex++)
 				{
@@ -333,10 +334,17 @@ class FLogSuppressionImplementation: public FLogSuppressionInterface, private FS
 					}
 				}
 				// store off the last non-zero one for toggle
-				if (Value & ELogVerbosity::VerbosityMask)
+				const uint8 VerbosityLevel = Value & ELogVerbosity::VerbosityMask;
+				if (VerbosityLevel)
 				{
 					// currently on, store this in the pending and clear it
-					ToggleAssociations.Add(Category, Value & ELogVerbosity::VerbosityMask);
+					ToggleAssociations.Add(Category, VerbosityLevel);
+
+					// Tattle on configs & other paths that raise the level after boot so we can quickly rule out code defaults for log spam
+					if ((OriginalVerbosityLevel < ELogVerbosity::Verbose) && (VerbosityLevel >= ELogVerbosity::Verbose))
+					{
+						UE_LOG(LogHAL, Log, TEXT("Log category %s verbosity has been raised to %s."), *Category.ToString(), (VerbosityLevel == ELogVerbosity::VeryVerbose ? TEXT("VeryVerbose") : TEXT("Verbose")));
+					}
 				}
 			}
 		}
