Small tweak to use ToString instead of manual enum-string conversion for better future-proofing.

diff --git a/Engine/Source/Runtime/Core/Private/Logging/LogSuppressionInterface.cpp b/Engine/Source/Runtime/Core/Private/Logging/LogSuppressionInterface.cpp
--- a/Engine/Source/Runtime/Core/Private/Logging/LogSuppressionInterface.cpp
+++ b/Engine/Source/Runtime/Core/Private/Logging/LogSuppressionInterface.cpp
@@ -343,7 +343,7 @@ class FLogSuppressionImplementation: public FLogSuppressionInterface, private FS
 					// Tattle on configs & other paths that raise the level after boot so we can quickly rule out code defaults for log spam
 					if ((OriginalVerbosityLevel < ELogVerbosity::Verbose) && (VerbosityLevel >= ELogVerbosity::Verbose))
 					{
-						UE_LOG(LogHAL, Log, TEXT("Log category %s verbosity has been raised to %s."), *Category.ToString(), (VerbosityLevel == ELogVerbosity::VeryVerbose ? TEXT("VeryVerbose") : TEXT("Verbose")));
+						UE_LOG(LogHAL, Log, TEXT("Log category %s verbosity has been raised to %s."), *Category.ToString(), ToString(static_cast<ELogVerbosity::Type>(VerbosityLevel)));
 					}
 				}
 			}
