VulkanRHI: Fix the driver version check in Linux where the "user version" is supplied everywhere. The Windows logic would use the last 6 chars, which worked for most versions tested (such as 511.17)  but not for those with three digits in the minor version (such as 470.103).

diff --git a/Engine/Source/Runtime/Core/Public/GenericPlatform/GenericPlatformDriver.h b/Engine/Source/Runtime/Core/Public/GenericPlatform/GenericPlatformDriver.h
--- a/Engine/Source/Runtime/Core/Public/GenericPlatform/GenericPlatformDriver.h
+++ b/Engine/Source/Runtime/Core/Public/GenericPlatform/GenericPlatformDriver.h
@@ -323,8 +323,11 @@ struct FGPUDriverInfo
 		// we use the internal version, not the user version to avoid problem where the name was altered 
 		const FString& FullVersion = InternalDriverVersion;
 
-		if(IsNVIDIA())
+		if(IsNVIDIA() && (InternalDriverVersion != UserDriverVersion))
 		{
+			// only do this processing on platforms where the InternalDriverVersion is different from the UserDriverVersion
+			// for example Linux reads the version from Vulkan properties directly and uses the UserDriverVersion everywhere.
+
 			// on the internal driver number: https://forums.geforce.com/default/topic/378546/confusion-over-driver-version-numbers/
 			//   The first 7 shows u that is a Vista driver, 6 that is an XP and 4 that is Me
 			// we don't care about the windows version so we don't look at the front part of the driver version
