UE::String::Find*() updated to accept UTF8CHAR views.

diff --git a/Engine/Source/Runtime/Core/Private/String/Find.cpp b/Engine/Source/Runtime/Core/Private/String/Find.cpp
--- a/Engine/Source/Runtime/Core/Private/String/Find.cpp
+++ b/Engine/Source/Runtime/Core/Private/String/Find.cpp
@@ -336,7 +336,7 @@ static inline int32 FindLastOfAnyChar(TStringView<CharType> View, TConstArrayVie
 namespace UE::String
 {
 
-int32 FindFirst(FAnsiStringView View, FAnsiStringView Search, ESearchCase::Type SearchCase)
+int32 FindFirst(FUtf8StringView View, FUtf8StringView Search, ESearchCase::Type SearchCase)
 {
 	return Private::FindFirst(View, Search, SearchCase);
 }
@@ -346,7 +346,7 @@ int32 FindFirst(FWideStringView View, FWideStringView Search, ESearchCase::Type
 	return Private::FindFirst(View, Search, SearchCase);
 }
 
-int32 FindLast(FAnsiStringView View, FAnsiStringView Search, ESearchCase::Type SearchCase)
+int32 FindLast(FUtf8StringView View, FUtf8StringView Search, ESearchCase::Type SearchCase)
 {
 	return Private::FindLast(View, Search, SearchCase);
 }
@@ -356,7 +356,7 @@ int32 FindLast(FWideStringView View, FWideStringView Search, ESearchCase::Type S
 	return Private::FindLast(View, Search, SearchCase);
 }
 
-int32 FindFirstOfAny(FAnsiStringView View, TConstArrayView<FAnsiStringView> Search, ESearchCase::Type SearchCase)
+int32 FindFirstOfAny(FUtf8StringView View, TConstArrayView<FUtf8StringView> Search, ESearchCase::Type SearchCase)
 {
 	return Private::FindFirstOfAny(View, Search, SearchCase);
 }
@@ -366,7 +366,7 @@ int32 FindFirstOfAny(FWideStringView View, TConstArrayView<FWideStringView> Sear
 	return Private::FindFirstOfAny(View, Search, SearchCase);
 }
 
-int32 FindLastOfAny(FAnsiStringView View, TConstArrayView<FAnsiStringView> Search, ESearchCase::Type SearchCase)
+int32 FindLastOfAny(FUtf8StringView View, TConstArrayView<FUtf8StringView> Search, ESearchCase::Type SearchCase)
 {
 	return Private::FindLastOfAny(View, Search, SearchCase);
 }
@@ -376,7 +376,7 @@ int32 FindLastOfAny(FWideStringView View, TConstArrayView<FWideStringView> Searc
 	return Private::FindLastOfAny(View, Search, SearchCase);
 }
 
-int32 FindFirstChar(FAnsiStringView View, ANSICHAR Search, ESearchCase::Type SearchCase)
+int32 FindFirstChar(FUtf8StringView View, UTF8CHAR Search, ESearchCase::Type SearchCase)
 {
 	return Private::FindFirstChar(View, Search, SearchCase);
 }
@@ -386,7 +386,7 @@ int32 FindFirstChar(FWideStringView View, WIDECHAR Search, ESearchCase::Type Sea
 	return Private::FindFirstChar(View, Search, SearchCase);
 }
 
-int32 FindLastChar(FAnsiStringView View, ANSICHAR Search, ESearchCase::Type SearchCase)
+int32 FindLastChar(FUtf8StringView View, UTF8CHAR Search, ESearchCase::Type SearchCase)
 {
 	return Private::FindLastChar(View, Search, SearchCase);
 }
@@ -396,7 +396,7 @@ int32 FindLastChar(FWideStringView View, WIDECHAR Search, ESearchCase::Type Sear
 	return Private::FindLastChar(View, Search, SearchCase);
 }
 
-int32 FindFirstOfAnyChar(FAnsiStringView View, TConstArrayView<ANSICHAR> Search, ESearchCase::Type SearchCase)
+int32 FindFirstOfAnyChar(FUtf8StringView View, TConstArrayView<UTF8CHAR> Search, ESearchCase::Type SearchCase)
 {
 	return Private::FindFirstOfAnyChar(View, Search, SearchCase);
 }
@@ -406,7 +406,7 @@ int32 FindFirstOfAnyChar(FWideStringView View, TConstArrayView<WIDECHAR> Search,
 	return Private::FindFirstOfAnyChar(View, Search, SearchCase);
 }
 
-int32 FindLastOfAnyChar(FAnsiStringView View, TConstArrayView<ANSICHAR> Search, ESearchCase::Type SearchCase)
+int32 FindLastOfAnyChar(FUtf8StringView View, TConstArrayView<UTF8CHAR> Search, ESearchCase::Type SearchCase)
 {
 	return Private::FindLastOfAnyChar(View, Search, SearchCase);
 }
diff --git a/Engine/Source/Runtime/Core/Public/String/Find.h b/Engine/Source/Runtime/Core/Public/String/Find.h
--- a/Engine/Source/Runtime/Core/Public/String/Find.h
+++ b/Engine/Source/Runtime/Core/Public/String/Find.h
@@ -17,7 +17,7 @@ namespace UE::String
  * @param SearchCase Whether the comparison should ignore case.
  * @return The position at which the search string was found, or INDEX_NONE if not found.
  */
-CORE_API int32 FindFirst(FAnsiStringView View, FAnsiStringView Search, ESearchCase::Type SearchCase = ESearchCase::CaseSensitive);
+CORE_API int32 FindFirst(FUtf8StringView View, FUtf8StringView Search, ESearchCase::Type SearchCase = ESearchCase::CaseSensitive);
 CORE_API int32 FindFirst(FWideStringView View, FWideStringView Search, ESearchCase::Type SearchCase = ESearchCase::CaseSensitive);
 
 /**
@@ -28,7 +28,7 @@ CORE_API int32 FindFirst(FWideStringView View, FWideStringView Search, ESearchCa
  * @param SearchCase Whether the comparison should ignore case.
  * @return The position at which the search string was found, or INDEX_NONE if not found.
  */
-CORE_API int32 FindLast(FAnsiStringView View, FAnsiStringView Search, ESearchCase::Type SearchCase = ESearchCase::CaseSensitive);
+CORE_API int32 FindLast(FUtf8StringView View, FUtf8StringView Search, ESearchCase::Type SearchCase = ESearchCase::CaseSensitive);
 CORE_API int32 FindLast(FWideStringView View, FWideStringView Search, ESearchCase::Type SearchCase = ESearchCase::CaseSensitive);
 
 /**
@@ -39,7 +39,7 @@ CORE_API int32 FindLast(FWideStringView View, FWideStringView Search, ESearchCas
  * @param SearchCase Whether the comparison should ignore case.
  * @return The position at which any search string was found, or INDEX_NONE if not found.
  */
-CORE_API int32 FindFirstOfAny(FAnsiStringView View, TConstArrayView<FAnsiStringView> Search, ESearchCase::Type SearchCase = ESearchCase::CaseSensitive);
+CORE_API int32 FindFirstOfAny(FUtf8StringView View, TConstArrayView<FUtf8StringView> Search, ESearchCase::Type SearchCase = ESearchCase::CaseSensitive);
 CORE_API int32 FindFirstOfAny(FWideStringView View, TConstArrayView<FWideStringView> Search, ESearchCase::Type SearchCase = ESearchCase::CaseSensitive);
 
 /**
@@ -50,7 +50,7 @@ CORE_API int32 FindFirstOfAny(FWideStringView View, TConstArrayView<FWideStringV
  * @param SearchCase Whether the comparison should ignore case.
  * @return The position at which any search string was found, or INDEX_NONE if not found.
  */
-CORE_API int32 FindLastOfAny(FAnsiStringView View, TConstArrayView<FAnsiStringView> Search, ESearchCase::Type SearchCase = ESearchCase::CaseSensitive);
+CORE_API int32 FindLastOfAny(FUtf8StringView View, TConstArrayView<FUtf8StringView> Search, ESearchCase::Type SearchCase = ESearchCase::CaseSensitive);
 CORE_API int32 FindLastOfAny(FWideStringView View, TConstArrayView<FWideStringView> Search, ESearchCase::Type SearchCase = ESearchCase::CaseSensitive);
 
 /**
@@ -61,8 +61,12 @@ CORE_API int32 FindLastOfAny(FWideStringView View, TConstArrayView<FWideStringVi
  * @param SearchCase Whether the comparison should ignore case.
  * @return The position at which the search character was found, or INDEX_NONE if not found.
  */
-CORE_API int32 FindFirstChar(FAnsiStringView View, ANSICHAR Search, ESearchCase::Type SearchCase = ESearchCase::CaseSensitive);
+CORE_API int32 FindFirstChar(FUtf8StringView View, UTF8CHAR Search, ESearchCase::Type SearchCase = ESearchCase::CaseSensitive);
 CORE_API int32 FindFirstChar(FWideStringView View, WIDECHAR Search, ESearchCase::Type SearchCase = ESearchCase::CaseSensitive);
+inline int32 FindFirstChar(FUtf8StringView View, ANSICHAR Search, ESearchCase::Type SearchCase = ESearchCase::CaseSensitive)
+{
+	return FindFirstChar(View, (UTF8CHAR)Search, SearchCase);
+}
 
 /**
  * Search the view for the last occurrence of the search character.
@@ -72,8 +76,12 @@ CORE_API int32 FindFirstChar(FWideStringView View, WIDECHAR Search, ESearchCase:
  * @param SearchCase Whether the comparison should ignore case.
  * @return The position at which the search character was found, or INDEX_NONE if not found.
  */
-CORE_API int32 FindLastChar(FAnsiStringView View, ANSICHAR Search, ESearchCase::Type SearchCase = ESearchCase::CaseSensitive);
+CORE_API int32 FindLastChar(FUtf8StringView View, UTF8CHAR Search, ESearchCase::Type SearchCase = ESearchCase::CaseSensitive);
 CORE_API int32 FindLastChar(FWideStringView View, WIDECHAR Search, ESearchCase::Type SearchCase = ESearchCase::CaseSensitive);
+inline int32 FindLastChar(FUtf8StringView View, ANSICHAR Search, ESearchCase::Type SearchCase = ESearchCase::CaseSensitive)
+{
+	return FindLastChar(View, (UTF8CHAR)Search, SearchCase);
+}
 
 /**
  * Search the view for the first occurrence of any search character.
@@ -83,8 +91,12 @@ CORE_API int32 FindLastChar(FWideStringView View, WIDECHAR Search, ESearchCase::
  * @param SearchCase Whether the comparison should ignore case.
  * @return The position at which any search character was found, or INDEX_NONE if not found.
  */
-CORE_API int32 FindFirstOfAnyChar(FAnsiStringView View, TConstArrayView<ANSICHAR> Search, ESearchCase::Type SearchCase = ESearchCase::CaseSensitive);
+CORE_API int32 FindFirstOfAnyChar(FUtf8StringView View, TConstArrayView<UTF8CHAR> Search, ESearchCase::Type SearchCase = ESearchCase::CaseSensitive);
 CORE_API int32 FindFirstOfAnyChar(FWideStringView View, TConstArrayView<WIDECHAR> Search, ESearchCase::Type SearchCase = ESearchCase::CaseSensitive);
+inline int32 FindFirstOfAnyChar(FUtf8StringView View, TConstArrayView<ANSICHAR> Search, ESearchCase::Type SearchCase = ESearchCase::CaseSensitive)
+{
+	return FindFirstOfAnyChar(View, TConstArrayView<UTF8CHAR>((const UTF8CHAR*)Search.GetData(), Search.Num()), SearchCase);
+}
 
 /**
  * Search the view for the last occurrence of any search character.
@@ -94,7 +106,11 @@ CORE_API int32 FindFirstOfAnyChar(FWideStringView View, TConstArrayView<WIDECHAR
  * @param SearchCase Whether the comparison should ignore case.
  * @return The position at which any search character was found, or INDEX_NONE if not found.
  */
-CORE_API int32 FindLastOfAnyChar(FAnsiStringView View, TConstArrayView<ANSICHAR> Search, ESearchCase::Type SearchCase = ESearchCase::CaseSensitive);
+CORE_API int32 FindLastOfAnyChar(FUtf8StringView View, TConstArrayView<UTF8CHAR> Search, ESearchCase::Type SearchCase = ESearchCase::CaseSensitive);
 CORE_API int32 FindLastOfAnyChar(FWideStringView View, TConstArrayView<WIDECHAR> Search, ESearchCase::Type SearchCase = ESearchCase::CaseSensitive);
+inline int32 FindLastOfAnyChar(FUtf8StringView View, TConstArrayView<ANSICHAR> Search, ESearchCase::Type SearchCase = ESearchCase::CaseSensitive)
+{
+	return FindLastOfAnyChar(View, TConstArrayView<UTF8CHAR>((const UTF8CHAR*)Search.GetData(), Search.Num()), SearchCase);
+}
 
 } // UE::String
