Fix for iPhone device names with apostrophes appearing as ???? on Mac

diff --git a/Engine/Source/Runtime/Core/Private/Mac/MacPlatformProcess.cpp b/Engine/Source/Runtime/Core/Private/Mac/MacPlatformProcess.cpp
--- a/Engine/Source/Runtime/Core/Private/Mac/MacPlatformProcess.cpp
+++ b/Engine/Source/Runtime/Core/Private/Mac/MacPlatformProcess.cpp
@@ -11,6 +11,7 @@
 #include "Misc/App.h"
 #include "Misc/CoreDelegates.h"
 #include "Misc/Paths.h"
+#include "Misc/StringBuilder.h"
 #include "HAL/FileManager.h"
 #include "Apple/PreAppleSystemHeaders.h"
 #include <mach-o/dyld.h>
@@ -19,7 +20,6 @@
 #include <libproc.h>
 #include <spawn.h>
 #include "Apple/PostAppleSystemHeaders.h"
-
 #if PLATFORM_MAC_X86
     #include <cpuid.h>
 #endif
@@ -1329,26 +1329,28 @@ FString FMacPlatformProcess::ReadPipe( void* ReadPipe )
 {
 	SCOPED_AUTORELEASE_POOL;
 
-	FString Output;
-
 	const int32 READ_SIZE = 8192;
 	ANSICHAR Buffer[READ_SIZE+1];
 	int32 BytesRead = 0;
 
+	// We don't want to use FUtf8StringBuilderBase here because the Buffer boundary could split a UTF-8 character
+	// and we don't want .Append to attempt to interpert the incoming data.  We will convert the whole string in one go at the end.
+	FAnsiStringBuilderBase StringBuilder;
+
 	if(ReadPipe)
 	{
 		do
 		{
-		BytesRead = read([(NSFileHandle*)ReadPipe fileDescriptor], Buffer, READ_SIZE);
-		if (BytesRead > 0)
-		{
-			Buffer[BytesRead] = '\0';
-			Output += StringCast<TCHAR>(Buffer).Get();
-		}
+			BytesRead = read([(NSFileHandle*)ReadPipe fileDescriptor], Buffer, READ_SIZE);
+			if (BytesRead > 0)
+			{
+				Buffer[BytesRead] = '\0';
+				StringBuilder.Append(Buffer, BytesRead);
+			}
 		} while (BytesRead > 0);
 	}
-
-	return Output;
+	
+	return FString(UTF8_TO_TCHAR(StringBuilder.ToString()));
 }
 
 bool FMacPlatformProcess::ReadPipeToArray(void* ReadPipe, TArray<uint8>& Output)
