Fixed a compile error in String/Find.h

#rnx

diff --git a/Engine/Source/Runtime/Core/Private/String/Find.cpp b/Engine/Source/Runtime/Core/Private/String/Find.cpp
--- a/Engine/Source/Runtime/Core/Private/String/Find.cpp
+++ b/Engine/Source/Runtime/Core/Private/String/Find.cpp
@@ -376,6 +376,11 @@ int32 FindLastOfAny(FWideStringView View, TConstArrayView<FWideStringView> Searc
 	return Private::FindLastOfAny(View, Search, SearchCase);
 }
 
+int32 FindFirstChar(FUtf8StringView View, ANSICHAR Search, ESearchCase::Type SearchCase)
+{
+	return Private::FindFirstChar(View, UTF8CHAR(Search), SearchCase);
+}
+
 int32 FindFirstChar(FUtf8StringView View, UTF8CHAR Search, ESearchCase::Type SearchCase)
 {
 	return Private::FindFirstChar(View, Search, SearchCase);
@@ -386,6 +391,11 @@ int32 FindFirstChar(FWideStringView View, WIDECHAR Search, ESearchCase::Type Sea
 	return Private::FindFirstChar(View, Search, SearchCase);
 }
 
+int32 FindLastChar(FUtf8StringView View, ANSICHAR Search, ESearchCase::Type SearchCase)
+{
+	return Private::FindLastChar(View, UTF8CHAR(Search), SearchCase);
+}
+
 int32 FindLastChar(FUtf8StringView View, UTF8CHAR Search, ESearchCase::Type SearchCase)
 {
 	return Private::FindLastChar(View, Search, SearchCase);
@@ -396,6 +406,11 @@ int32 FindLastChar(FWideStringView View, WIDECHAR Search, ESearchCase::Type Sear
 	return Private::FindLastChar(View, Search, SearchCase);
 }
 
+int32 FindFirstOfAnyChar(FUtf8StringView View, TConstArrayView<ANSICHAR> Search, ESearchCase::Type SearchCase)
+{
+	return Private::FindFirstOfAnyChar(View, MakeArrayView((const UTF8CHAR*)Search.GetData(), Search.Num()), SearchCase);
+}
+
 int32 FindFirstOfAnyChar(FUtf8StringView View, TConstArrayView<UTF8CHAR> Search, ESearchCase::Type SearchCase)
 {
 	return Private::FindFirstOfAnyChar(View, Search, SearchCase);
@@ -406,6 +421,11 @@ int32 FindFirstOfAnyChar(FWideStringView View, TConstArrayView<WIDECHAR> Search,
 	return Private::FindFirstOfAnyChar(View, Search, SearchCase);
 }
 
+int32 FindLastOfAnyChar(FUtf8StringView View, TConstArrayView<ANSICHAR> Search, ESearchCase::Type SearchCase)
+{
+	return Private::FindLastOfAnyChar(View, MakeArrayView((const UTF8CHAR*)Search.GetData(), Search.Num()), SearchCase);
+}
+
 int32 FindLastOfAnyChar(FUtf8StringView View, TConstArrayView<UTF8CHAR> Search, ESearchCase::Type SearchCase)
 {
 	return Private::FindLastOfAnyChar(View, Search, SearchCase);
diff --git a/Engine/Source/Runtime/Core/Public/String/Find.h b/Engine/Source/Runtime/Core/Public/String/Find.h
--- a/Engine/Source/Runtime/Core/Public/String/Find.h
+++ b/Engine/Source/Runtime/Core/Public/String/Find.h
@@ -61,12 +61,9 @@ CORE_API int32 FindLastOfAny(FWideStringView View, TConstArrayView<FWideStringVi
  * @param SearchCase Whether the comparison should ignore case.
  * @return The position at which the search character was found, or INDEX_NONE if not found.
  */
+CORE_API int32 FindFirstChar(FUtf8StringView View, ANSICHAR Search, ESearchCase::Type SearchCase = ESearchCase::CaseSensitive);
 CORE_API int32 FindFirstChar(FUtf8StringView View, UTF8CHAR Search, ESearchCase::Type SearchCase = ESearchCase::CaseSensitive);
 CORE_API int32 FindFirstChar(FWideStringView View, WIDECHAR Search, ESearchCase::Type SearchCase = ESearchCase::CaseSensitive);
-inline int32 FindFirstChar(FUtf8StringView View, ANSICHAR Search, ESearchCase::Type SearchCase = ESearchCase::CaseSensitive)
-{
-	return FindFirstChar(View, (UTF8CHAR)Search, SearchCase);
-}
 
 /**
  * Search the view for the last occurrence of the search character.
@@ -76,12 +73,9 @@ inline int32 FindFirstChar(FUtf8StringView View, ANSICHAR Search, ESearchCase::T
  * @param SearchCase Whether the comparison should ignore case.
  * @return The position at which the search character was found, or INDEX_NONE if not found.
  */
+CORE_API int32 FindLastChar(FUtf8StringView View, ANSICHAR Search, ESearchCase::Type SearchCase = ESearchCase::CaseSensitive);
 CORE_API int32 FindLastChar(FUtf8StringView View, UTF8CHAR Search, ESearchCase::Type SearchCase = ESearchCase::CaseSensitive);
 CORE_API int32 FindLastChar(FWideStringView View, WIDECHAR Search, ESearchCase::Type SearchCase = ESearchCase::CaseSensitive);
-inline int32 FindLastChar(FUtf8StringView View, ANSICHAR Search, ESearchCase::Type SearchCase = ESearchCase::CaseSensitive)
-{
-	return FindLastChar(View, (UTF8CHAR)Search, SearchCase);
-}
 
 /**
  * Search the view for the first occurrence of any search character.
@@ -91,12 +85,9 @@ inline int32 FindLastChar(FUtf8StringView View, ANSICHAR Search, ESearchCase::Ty
  * @param SearchCase Whether the comparison should ignore case.
  * @return The position at which any search character was found, or INDEX_NONE if not found.
  */
+CORE_API int32 FindFirstOfAnyChar(FUtf8StringView View, TConstArrayView<ANSICHAR> Search, ESearchCase::Type SearchCase = ESearchCase::CaseSensitive);
 CORE_API int32 FindFirstOfAnyChar(FUtf8StringView View, TConstArrayView<UTF8CHAR> Search, ESearchCase::Type SearchCase = ESearchCase::CaseSensitive);
 CORE_API int32 FindFirstOfAnyChar(FWideStringView View, TConstArrayView<WIDECHAR> Search, ESearchCase::Type SearchCase = ESearchCase::CaseSensitive);
-inline int32 FindFirstOfAnyChar(FUtf8StringView View, TConstArrayView<ANSICHAR> Search, ESearchCase::Type SearchCase = ESearchCase::CaseSensitive)
-{
-	return FindFirstOfAnyChar(View, TConstArrayView<UTF8CHAR>((const UTF8CHAR*)Search.GetData(), Search.Num()), SearchCase);
-}
 
 /**
  * Search the view for the last occurrence of any search character.
@@ -106,11 +97,8 @@ inline int32 FindFirstOfAnyChar(FUtf8StringView View, TConstArrayView<ANSICHAR>
  * @param SearchCase Whether the comparison should ignore case.
  * @return The position at which any search character was found, or INDEX_NONE if not found.
  */
+CORE_API int32 FindLastOfAnyChar(FUtf8StringView View, TConstArrayView<ANSICHAR> Search, ESearchCase::Type SearchCase = ESearchCase::CaseSensitive);
 CORE_API int32 FindLastOfAnyChar(FUtf8StringView View, TConstArrayView<UTF8CHAR> Search, ESearchCase::Type SearchCase = ESearchCase::CaseSensitive);
 CORE_API int32 FindLastOfAnyChar(FWideStringView View, TConstArrayView<WIDECHAR> Search, ESearchCase::Type SearchCase = ESearchCase::CaseSensitive);
-inline int32 FindLastOfAnyChar(FUtf8StringView View, TConstArrayView<ANSICHAR> Search, ESearchCase::Type SearchCase = ESearchCase::CaseSensitive)
-{
-	return FindLastOfAnyChar(View, TConstArrayView<UTF8CHAR>((const UTF8CHAR*)Search.GetData(), Search.Num()), SearchCase);
-}
 
 } // UE::String
