Added "DumpLLM" console command to capture the current values of all tracked LLM tags. Added functionality to LowLevelMemTracker to support this by adding accessor functions to retreive tags and their properties as well as refer to them by using FTagData pointers.

#rnx
#pf none
#test dumpllm, memreport -csv

diff --git a/Engine/Source/Runtime/Core/Private/HAL/LowLevelMemTracker.cpp b/Engine/Source/Runtime/Core/Private/HAL/LowLevelMemTracker.cpp
--- a/Engine/Source/Runtime/Core/Private/HAL/LowLevelMemTracker.cpp
+++ b/Engine/Source/Runtime/Core/Private/HAL/LowLevelMemTracker.cpp
@@ -73,6 +73,64 @@ FAutoConsoleCommand LLMSnapshot(
 		FLowLevelMemTracker::Get().PublishDataSingleFrame();
 	}));
 
+FAutoConsoleCommand DumpLLM(
+	TEXT("DumpLLM"),
+	TEXT("Logs out the current and peak sizes of all tracked LLM tags"),
+	FConsoleCommandWithWorldArgsAndOutputDeviceDelegate::CreateLambda([&](const TArray<FString>& Args, UWorld* InWorld, FOutputDevice& Ar)
+	{
+		const float InvToMb = 1.0 / (1024 * 1024);
+
+		bool bCSV = false;
+		for (const FString& Arg : Args)
+		{
+			//if (Arg.Equals(TEXT("CSV"), ESearchCase::IgnoreCase))
+			if (FParse::Param(*Arg, TEXT("CSV")))
+			{
+				bCSV = true;
+			}
+		}
+
+		if (FLowLevelMemTracker::Get().IsEnabled())
+		{
+			if (bCSV)
+			{
+				Ar.Logf(TEXT(",TagName,SizeMB,PeakMB,Tracker,PathName"));
+			}
+			else
+			{
+				Ar.Logf(TEXT("%40s %12s %12s  Tracker PathName"), TEXT("TagName"), TEXT("SizeMB"), TEXT("PeakMB"));
+			}
+
+			for (ELLMTracker TrackerType : {ELLMTracker::Default, ELLMTracker::Platform})
+			{
+				const TCHAR* TrackerName = TrackerType == ELLMTracker::Default ? TEXT("Default") : TEXT("Platform");
+
+				for (const UE::LLMPrivate::FTagData* TagData : FLowLevelMemTracker::Get().GetTrackedTags(TrackerType))
+				{
+					if (bCSV)
+					{
+						Ar.Logf(TEXT(",%s,%.3f,%.3f,%s,%s"),
+							*FLowLevelMemTracker::Get().GetTagUniqueName(TagData).ToString(),
+							FLowLevelMemTracker::Get().GetTagAmountForTracker(TrackerType, TagData, false) * InvToMb,
+							FLowLevelMemTracker::Get().GetTagAmountForTracker(TrackerType, TagData, true) * InvToMb,
+							TrackerName,
+							*FLowLevelMemTracker::Get().GetTagDisplayPathName(TagData));
+					}
+					else
+					{
+						Ar.Logf(TEXT("%40s %12.3f %12.3f %8s %s"),
+							*FLowLevelMemTracker::Get().GetTagUniqueName(TagData).ToString(),
+							FLowLevelMemTracker::Get().GetTagAmountForTracker(TrackerType, TagData, false) * InvToMb,
+							FLowLevelMemTracker::Get().GetTagAmountForTracker(TrackerType, TagData, true) * InvToMb,
+							TrackerName,
+							*FLowLevelMemTracker::Get().GetTagDisplayPathName(TagData));
+					}
+				}
+			}
+		}
+	})
+);
+
 DECLARE_LLM_MEMORY_STAT(TEXT("LLM Overhead"), STAT_LLMOverheadTotal, STATGROUP_LLMOverhead);
 
 DEFINE_STAT(STAT_EngineSummaryLLM);
@@ -436,11 +494,12 @@ namespace LLMPrivate
 		void OnTagsResorted(FTagDataArray& OldTagDatas);
 		void LockAllThreadTags(bool bLock);
 
-		int64 GetTagAmount(const FTagData* TagData) const;
+		int64 GetTagAmount(const FTagData* TagData, bool bPeakAmount = false) const;
 		void SetTagAmountExternal(const FTagData* TagData, int64 Amount, bool bAddToTotal);
 		void SetTagAmountInUpdate(const FTagData* TagData, int64 Amount, bool bAddToTotal);
 		const FTagData* GetActiveTagData();
 		const FTagData* FindTagForPtr(void* Ptr);
+		TArray<const FTagData*> GetTagDatas();
 
 		int64 GetAllocTypeAmount(ELLMAllocType AllocType);
 
@@ -1682,6 +1741,36 @@ void FLowLevelMemTracker::FinishConstruct(UE::LLMPrivate::FTagData* TagData, UE:
 	}
 }
 
+TArray<const UE::LLMPrivate::FTagData*> FLowLevelMemTracker::GetTrackedTags()
+{
+	using namespace UE::LLMPrivate;
+
+	if (bIsDisabled)
+	{
+		return TArray<const FTagData*>();
+	}
+
+	BootstrapInitialise();
+
+	FReadScopeLock TagDataScopeLock(TagDataLock);
+	return TArray<const FTagData*>(*TagDatas);
+}
+
+TArray<const UE::LLMPrivate::FTagData*> FLowLevelMemTracker::GetTrackedTags(ELLMTracker Tracker)
+{
+	using namespace UE::LLMPrivate;
+
+	if (bIsDisabled)
+	{
+		return TArray<const FTagData*>();
+	}
+	
+	BootstrapInitialise();
+
+	FScopeLock UpdateScopeLock(&UpdateLock); // uses of TagSizes are guarded by the UpdateLock
+	return GetTracker(Tracker)->GetTagDatas();
+}
+
 bool FLowLevelMemTracker::FindTagByName( const TCHAR* Name, uint64& OutTag ) const
 {
 	using namespace UE::LLMPrivate;
@@ -1773,7 +1862,34 @@ FName FLowLevelMemTracker::FindTagDisplayName(uint64 Tag) const
 	return NAME_None;
 }
 
-int64 FLowLevelMemTracker::GetTagAmountForTracker(ELLMTracker Tracker, ELLMTag Tag)
+FName FLowLevelMemTracker::GetTagDisplayName(const UE::LLMPrivate::FTagData* TagData) const
+{
+	return TagData->GetDisplayName();
+}
+
+FString FLowLevelMemTracker::GetTagDisplayPathName(const UE::LLMPrivate::FTagData* TagData) const
+{
+	return TagData->GetDisplayPath();
+}
+
+FName FLowLevelMemTracker::GetTagUniqueName(const UE::LLMPrivate::FTagData* TagData) const
+{
+	return TagData->GetName();
+}
+
+int64 FLowLevelMemTracker::GetTagAmountForTracker(ELLMTracker Tracker, ELLMTag Tag, bool bPeakAmount/* = false*/)
+{
+	if (bIsDisabled)
+	{
+		return 0;
+	}
+
+	BootstrapInitialise();
+
+	return GetTagAmountForTracker(Tracker, FindTagData(Tag), bPeakAmount);
+}
+
+int64 FLowLevelMemTracker::GetTagAmountForTracker(ELLMTracker Tracker, const UE::LLMPrivate::FTagData* TagData, bool bPeakAmount/* = false*/)
 {
 	using namespace UE::LLMPrivate;
 
@@ -1781,15 +1897,16 @@ int64 FLowLevelMemTracker::GetTagAmountForTracker(ELLMTracker Tracker, ELLMTag T
 	{
 		return 0;
 	}
+
 	BootstrapInitialise();
-	const FTagData* TagData = FindTagData(Tag);
+
 	if (TagData == nullptr)
 	{
 		return 0;
 	}
 
 	FScopeLock UpdateScopeLock(&UpdateLock); // uses of TagSizes are guarded by the UpdateLock
-	return GetTracker(Tracker)->GetTagAmount(TagData);
+	return GetTracker(Tracker)->GetTagAmount(TagData, bPeakAmount);
 }
 
 void FLowLevelMemTracker::SetTagAmountForTracker(ELLMTracker Tracker, ELLMTag Tag, int64 Amount, bool bAddToTotal)
@@ -3415,12 +3532,19 @@ namespace LLMPrivate
 		return AllocInfo.GetTag(LLMRef);
 	}
 
-	int64 FLLMTracker::GetTagAmount(const FTagData* TagData) const
+	TArray<const FTagData*> FLLMTracker::GetTagDatas()
+	{
+		TArray<const FTagData*> FoundTagDatas;
+		TagSizes.GetKeys(FoundTagDatas);
+		return FoundTagDatas;
+	}
+
+	int64 FLLMTracker::GetTagAmount(const FTagData* TagData, bool bPeakAmount/* = false*/) const
 	{
 		const FTrackerTagSizeData* AllocationData = TagSizes.Find(TagData);
 		if (AllocationData)
 		{
-			return AllocationData->Size;
+			return AllocationData->GetSize(bPeakAmount);
 		}
 		else
 		{
diff --git a/Engine/Source/Runtime/Core/Public/HAL/LowLevelMemTracker.h b/Engine/Source/Runtime/Core/Public/HAL/LowLevelMemTracker.h
--- a/Engine/Source/Runtime/Core/Public/HAL/LowLevelMemTracker.h
+++ b/Engine/Source/Runtime/Core/Public/HAL/LowLevelMemTracker.h
@@ -547,22 +547,41 @@ public:
 	void RegisterPlatformTag(int32 Tag, const TCHAR* Name, FName StatName, FName SummaryStatName, int32 ParentTag = -1);
 	void RegisterProjectTag(int32 Tag, const TCHAR* Name, FName StatName, FName SummaryStatName, int32 ParentTag = -1);
     
+	// Get all tags being tracked
+	TArray<const UE::LLMPrivate::FTagData*> GetTrackedTags();
+
+	// Get all tags being tracked by the given tracker
+	TArray<const UE::LLMPrivate::FTagData*> GetTrackedTags(ELLMTracker Tracker);
+
 	// look up the ELLMTag associated with the given display name
-	bool FindTagByName( const TCHAR* Name, uint64& OutTag ) const;
+	bool FindTagByName(const TCHAR* Name, uint64& OutTag) const;
 
 	UE_DEPRECATED(4.27, "Use FindTagDisplayName instead")
 	const TCHAR* FindTagName(uint64 Tag) const;
+
 	// get the display name for the given ELLMTag
 	FName FindTagDisplayName(uint64 Tag) const;
 
+	// Get the display name for the given FTagData
+	FName GetTagDisplayName(const UE::LLMPrivate::FTagData* TagData) const;
+
+	// Get the path name for the given FTagData from a chain of its parents' display names
+	FString GetTagDisplayPathName(const UE::LLMPrivate::FTagData* TagData) const;
+
+	// Get the unique identifier name for the given FTagData
+	FName GetTagUniqueName(const UE::LLMPrivate::FTagData* TagData) const;
+
 	// Get the amount of memory for an ELLMTag from the given tracker
-	int64 GetTagAmountForTracker(ELLMTracker Tracker, ELLMTag Tag);
+	int64 GetTagAmountForTracker(ELLMTracker Tracker, ELLMTag Tag, bool bPeakAmount = false);
+
+	// Get the amount of memory for a FTagData from the given tracker
+	int64 GetTagAmountForTracker(ELLMTracker Tracker, const UE::LLMPrivate::FTagData* TagData, bool bPeakAmount = false);
 
 	// Set the amount of memory for an ELLMTag for a given tracker, optionally updating the total tracked memory too
-	void SetTagAmountForTracker(ELLMTracker Tracker, ELLMTag Tag, int64 Amount, bool bAddToTotal );
+	void SetTagAmountForTracker(ELLMTracker Tracker, ELLMTag Tag, int64 Amount, bool bAddToTotal);
 
 	// Dump the display name of the current TagData for the given tracker to the output
-	uint64 DumpTag( ELLMTracker Tracker, const char* FileName, int LineNumber );
+	uint64 DumpTag(ELLMTracker Tracker, const char* FileName, int LineNumber);
 
 	// Publishes the active LLM stats in the active frame, useful for single targeted LLM snapshots
 	void PublishDataSingleFrame();
