Fix for compile error in LowLevelMemTracker

#rnx
#pf none

diff --git a/Engine/Source/Runtime/Core/Private/HAL/LowLevelMemTracker.cpp b/Engine/Source/Runtime/Core/Private/HAL/LowLevelMemTracker.cpp
--- a/Engine/Source/Runtime/Core/Private/HAL/LowLevelMemTracker.cpp
+++ b/Engine/Source/Runtime/Core/Private/HAL/LowLevelMemTracker.cpp
@@ -76,7 +76,7 @@ FAutoConsoleCommand LLMSnapshot(
 FAutoConsoleCommand DumpLLM(
 	TEXT("DumpLLM"),
 	TEXT("Logs out the current and peak sizes of all tracked LLM tags"),
-	FConsoleCommandWithWorldArgsAndOutputDeviceDelegate::CreateLambda([&](const TArray<FString>& Args, UWorld* InWorld, FOutputDevice& Ar)
+	FConsoleCommandWithWorldArgsAndOutputDeviceDelegate::CreateLambda([](const TArray<FString>& Args, UWorld* InWorld, FOutputDevice& Ar)
 	{
 		const float InvToMb = 1.0 / (1024 * 1024);
 
