Custom FChange objects can now add referenced objects

diff --git a/Engine/Source/Editor/UnrealEd/Private/EditorTransaction.cpp b/Engine/Source/Editor/UnrealEd/Private/EditorTransaction.cpp
--- a/Engine/Source/Editor/UnrealEd/Private/EditorTransaction.cpp
+++ b/Engine/Source/Editor/UnrealEd/Private/EditorTransaction.cpp
@@ -763,6 +763,11 @@ void FTransaction::FObjectRecord::AddReferencedObjects( FReferenceCollector& Col
 	AddSerializedObjectReferences(SerializedObject);
 	AddSerializedObjectReferences(SerializedObjectFlip);
 	AddSerializedObjectReferences(SerializedObjectSnapshot);
+
+	if (CustomChange.IsValid())
+	{
+		CustomChange->AddReferencedObjects(Collector);
+	}
 }
 
 bool FTransaction::FObjectRecord::ContainsPieObject() const
diff --git a/Engine/Source/Runtime/Core/Public/Misc/Change.h b/Engine/Source/Runtime/Core/Public/Misc/Change.h
--- a/Engine/Source/Runtime/Core/Public/Misc/Change.h
+++ b/Engine/Source/Runtime/Core/Public/Misc/Change.h
@@ -5,6 +5,9 @@
 #include "Templates/UniquePtr.h"
 #include "Containers/Array.h"
 
+class FFeedbackContext;
+class FReferenceCollector;
+
 /**
  * FChange modifies a UObject and is meant to be used to implement undo/redo.
  * The change is embedded in an FTransaction which executes it *instead* of the standard 
@@ -44,11 +47,14 @@ public:
 	/** @return true if this Change has Expired, ie it will no longer have any effect and could be skipped by undo/redo */
 	virtual bool HasExpired( UObject* Object ) const { return false; }
 
+	/** Used by GC to collect referenced objects. */
+	virtual void AddReferencedObjects( FReferenceCollector& Collector ) { }
+
 	/** Describes this change (for debugging) */
 	virtual FString ToString() const = 0;
 
 	/** Prints this change to the log, including sub-changes if there are any.  For compound changes, there might be multiple lines.  You should not need to override this function. */
-	virtual void PrintToLog( class FFeedbackContext& FeedbackContext, const int32 IndentLevel = 0 );
+	virtual void PrintToLog( FFeedbackContext& FeedbackContext, const int32 IndentLevel = 0 );
 
 	/** Virtual destructor */
 	virtual ~FChange()
