Fixed int64-to-float narrowing warning

#rnx

diff --git a/Engine/Source/Runtime/Core/Private/HAL/LowLevelMemTracker.cpp b/Engine/Source/Runtime/Core/Private/HAL/LowLevelMemTracker.cpp
--- a/Engine/Source/Runtime/Core/Private/HAL/LowLevelMemTracker.cpp
+++ b/Engine/Source/Runtime/Core/Private/HAL/LowLevelMemTracker.cpp
@@ -111,8 +111,8 @@ FAutoConsoleCommand DumpLLM(
 					{
 						Ar.Logf(TEXT(",%s,%.3f,%.3f,%s,%s"),
 							*FLowLevelMemTracker::Get().GetTagUniqueName(TagData).ToString(),
-							FLowLevelMemTracker::Get().GetTagAmountForTracker(TrackerType, TagData, false) * InvToMb,
-							FLowLevelMemTracker::Get().GetTagAmountForTracker(TrackerType, TagData, true) * InvToMb,
+							float(FLowLevelMemTracker::Get().GetTagAmountForTracker(TrackerType, TagData, false)) * InvToMb,
+							float(FLowLevelMemTracker::Get().GetTagAmountForTracker(TrackerType, TagData, true)) * InvToMb,
 							TrackerName,
 							*FLowLevelMemTracker::Get().GetTagDisplayPathName(TagData));
 					}
@@ -120,8 +120,8 @@ FAutoConsoleCommand DumpLLM(
 					{
 						Ar.Logf(TEXT("%40s %12.3f %12.3f %8s %s"),
 							*FLowLevelMemTracker::Get().GetTagUniqueName(TagData).ToString(),
-							FLowLevelMemTracker::Get().GetTagAmountForTracker(TrackerType, TagData, false) * InvToMb,
-							FLowLevelMemTracker::Get().GetTagAmountForTracker(TrackerType, TagData, true) * InvToMb,
+							float(FLowLevelMemTracker::Get().GetTagAmountForTracker(TrackerType, TagData, false)) * InvToMb,
+							float(FLowLevelMemTracker::Get().GetTagAmountForTracker(TrackerType, TagData, true)) * InvToMb,
 							TrackerName,
 							*FLowLevelMemTracker::Get().GetTagDisplayPathName(TagData));
 					}
