Add FGenericPlatformProcess::TerminateProcTreeWithPredicate() that enables killing a launched process tree while keeping some of the child processes alive.

#rnx

diff --git a/Engine/Source/Runtime/Core/Private/GenericPlatform/GenericPlatformProcess.cpp b/Engine/Source/Runtime/Core/Private/GenericPlatform/GenericPlatformProcess.cpp
--- a/Engine/Source/Runtime/Core/Private/GenericPlatform/GenericPlatformProcess.cpp
+++ b/Engine/Source/Runtime/Core/Private/GenericPlatform/GenericPlatformProcess.cpp
@@ -278,6 +278,13 @@ void FGenericPlatformProcess::TerminateProc( FProcHandle & ProcessHandle, bool K
 	UE_LOG(LogHAL, Fatal, TEXT("FGenericPlatformProcess::TerminateProc not implemented on this platform"));
 }
 
+void FGenericPlatformProcess::TerminateProcTreeWithPredicate(
+	FProcHandle& ProcessHandle,
+	TFunctionRef<bool(uint32 ProcessId, const TCHAR* ApplicationName)> Predicate)
+{
+	UE_LOG(LogHAL, Fatal, TEXT("FGenericPlatformProcess::TerminateProcTreeWithPredicate not implemented on this platform"));
+}
+
 FGenericPlatformProcess::EWaitAndForkResult FGenericPlatformProcess::WaitAndFork()
 {
 	UE_LOG(LogHAL, Fatal, TEXT("FGenericPlatformProcess::WaitAndFork not implemented on this platform"));
diff --git a/Engine/Source/Runtime/Core/Private/Windows/WindowsPlatformProcess.cpp b/Engine/Source/Runtime/Core/Private/Windows/WindowsPlatformProcess.cpp
--- a/Engine/Source/Runtime/Core/Private/Windows/WindowsPlatformProcess.cpp
+++ b/Engine/Source/Runtime/Core/Private/Windows/WindowsPlatformProcess.cpp
@@ -439,35 +439,45 @@ void FWindowsPlatformProcess::CloseProc(FProcHandle & ProcessHandle)
 
 void FWindowsPlatformProcess::TerminateProc( FProcHandle & ProcessHandle, bool KillTree )
 {
-	if (KillTree)
+	TerminateProcTreeWithPredicate(ProcessHandle, [](uint32 ProcessId, const TCHAR* ApplicationName) { return true; });
+}
+
+void FWindowsPlatformProcess::TerminateProcTreeWithPredicate(
+	FProcHandle& ProcessHandle,
+	TFunctionRef<bool(uint32 ProcessId, const TCHAR* ApplicationName)> Predicate)
+{
+	::DWORD ProcessId = ::GetProcessId(ProcessHandle.Get());
+	FString ProcessName = FPlatformProcess::GetApplicationName(ProcessId);
+
+	if (!Predicate(ProcessId, *ProcessName))
 	{
-		HANDLE SnapShot = ::CreateToolhelp32Snapshot(TH32CS_SNAPPROCESS, 0);
+		return;
+	}
 
-		if (SnapShot != INVALID_HANDLE_VALUE)
-		{
-			::DWORD ProcessId = ::GetProcessId(ProcessHandle.Get());
+	HANDLE SnapShot = ::CreateToolhelp32Snapshot(TH32CS_SNAPPROCESS, 0);
 
-			PROCESSENTRY32 Entry;
-			Entry.dwSize = sizeof(PROCESSENTRY32);
+	if (SnapShot != INVALID_HANDLE_VALUE)
+	{
+
+		PROCESSENTRY32 Entry;
+		Entry.dwSize = sizeof(PROCESSENTRY32);
 
-			if (::Process32First(SnapShot, &Entry))
+		if (::Process32First(SnapShot, &Entry))
+		{
+			do
 			{
-				do
+				if (Entry.th32ParentProcessID == ProcessId)
 				{
-					if (Entry.th32ParentProcessID == ProcessId)
+					HANDLE ChildProcHandle = ::OpenProcess(PROCESS_ALL_ACCESS, 0, Entry.th32ProcessID);
+
+					if (ChildProcHandle)
 					{
-						HANDLE ChildProcHandle = ::OpenProcess(PROCESS_ALL_ACCESS, 0, Entry.th32ProcessID);
-
-						if (ChildProcHandle)
-						{
-							FProcHandle ChildHandle(ChildProcHandle);
-							TerminateProc(ChildHandle, KillTree);
-//							::TerminateProcess(ChildProcHandle, 1);
-						}
+						FProcHandle ChildHandle(ChildProcHandle);
+						TerminateProcTreeWithPredicate(ChildHandle, Predicate);
 					}
 				}
-				while(::Process32Next(SnapShot, &Entry));
 			}
+			while(::Process32Next(SnapShot, &Entry));
 		}
 	}
 
diff --git a/Engine/Source/Runtime/Core/Public/GenericPlatform/GenericPlatformProcess.h b/Engine/Source/Runtime/Core/Public/GenericPlatform/GenericPlatformProcess.h
--- a/Engine/Source/Runtime/Core/Public/GenericPlatform/GenericPlatformProcess.h
+++ b/Engine/Source/Runtime/Core/Public/GenericPlatform/GenericPlatformProcess.h
@@ -468,6 +468,16 @@ struct CORE_API FGenericPlatformProcess
 	 */
 	static void TerminateProc( FProcHandle & ProcessHandle, bool KillTree = false );
 
+	/** Terminates a process tree
+	 *
+	 * @param ProcessHandle handle returned from FPlatformProcess::CreateProc
+	 * @param Predicate that returns true if the process identified by ProcessId and ApplicationName
+	 *        should be terminated with its children, else that process and its children will be kept alive
+	 */
+	static void TerminateProcTreeWithPredicate(
+			FProcHandle& ProcessHandle,
+			TFunctionRef<bool(uint32 ProcessId, const TCHAR* ApplicationName)> Predicate);
+
 	enum class EWaitAndForkResult : uint8
 	{
 		Error,
diff --git a/Engine/Source/Runtime/Core/Public/Windows/WindowsPlatformProcess.h b/Engine/Source/Runtime/Core/Public/Windows/WindowsPlatformProcess.h
--- a/Engine/Source/Runtime/Core/Public/Windows/WindowsPlatformProcess.h
+++ b/Engine/Source/Runtime/Core/Public/Windows/WindowsPlatformProcess.h
@@ -162,6 +162,9 @@ public:
 	static void WaitForProc( FProcHandle & ProcessHandle );
 	static void CloseProc( FProcHandle & ProcessHandle );
 	static void TerminateProc( FProcHandle & ProcessHandle, bool KillTree = false );
+	static void TerminateProcTreeWithPredicate(
+			FProcHandle& ProcessHandle,
+			TFunctionRef<bool(uint32 ProcessId, const TCHAR* ApplicationName)> Predicate);
 	static bool GetProcReturnCode( FProcHandle & ProcHandle, int32* ReturnCode );
 	static bool GetApplicationMemoryUsage(uint32 ProcessId, SIZE_T* OutMemoryUsage);
 	static bool GetPerFrameProcessorUsage(uint32 ProcessId, float& ProcessUsageFraction, float& IdleUsageFraction);
