FBlockEncoder::Compress parallel for faster writing of large packages

diff --git a/Engine/Source/Runtime/Core/Private/Compression/CompressedBuffer.cpp b/Engine/Source/Runtime/Core/Private/Compression/CompressedBuffer.cpp
--- a/Engine/Source/Runtime/Core/Private/Compression/CompressedBuffer.cpp
+++ b/Engine/Source/Runtime/Core/Private/Compression/CompressedBuffer.cpp
@@ -10,6 +10,7 @@
 #include "Misc/ByteSwap.h"
 #include "Misc/Crc.h"
 #include "Serialization/Archive.h"
+#include "Async/ParallelFor.h"
 
 THIRD_PARTY_INCLUDES_START
 #include "Compression/lz4.h"
@@ -261,6 +262,8 @@ private:
 
 FCompositeBuffer FBlockEncoder::Compress(const FCompositeBuffer& RawData, const uint64 BlockSize) const
 {
+	TRACE_CPUPROFILER_EVENT_SCOPE(FBlockEncoder::Compress);
+
 	checkf(FMath::IsPowerOfTwo(BlockSize) && BlockSize <= MAX_uint32,
 		TEXT("BlockSize must be a 32-bit power of two but was %" UINT64_FMT "."), BlockSize);
 	const uint64 RawSize = RawData.GetSize();
@@ -270,49 +273,100 @@ FCompositeBuffer FBlockEncoder::Compress(const FCompositeBuffer& RawData, const
 	checkf(BlockCount < MAX_uint32, TEXT("Raw data of size %" UINT64_FMT " with block size %" UINT64_FMT " requires ")
 		TEXT("%" UINT64_FMT " blocks, but the limit is %u."), RawSize, BlockSize, BlockCount, MAX_uint32);
 
+	// Compress the raw data in blocks and store the raw data for incompressible blocks.
+
+	struct CompressedBlockInfo
+	{
+		// in:
+		uint64 RawBlockSize;
+		FUniqueBuffer RawBlockCopy;
+		FMemoryView RawBlock;
+
+		// out:
+		TArray64<uint8> CompressedData;
+	};
+	
+	TArray64<CompressedBlockInfo> CompressedBlockInfos;
+	CompressedBlockInfos.SetNum((int)BlockCount);
+	
+	// fill all CompressedBlockInfos :
+	for (int BlockIndex=0;BlockIndex<(int)BlockCount;BlockIndex++)
+	{
+		const uint64 RawOffset = BlockIndex * BlockSize;
+		const uint64 RawBlockSize = FMath::Min(RawSize - RawOffset, BlockSize);
+		CompressedBlockInfos[BlockIndex].RawBlockSize = RawBlockSize;
+		CompressedBlockInfos[BlockIndex].RawBlock = RawData.ViewOrCopyRange(RawOffset, RawBlockSize, CompressedBlockInfos[BlockIndex].RawBlockCopy);
+
+		// sequential hash update :
+		RawHash.Update(CompressedBlockInfos[BlockIndex].RawBlock); // @todo Oodle: replace with parallel hash
+	}
+
+	uint64 CompressBlockAllocSize = CompressBlockBound(BlockSize);
+	
+	// encode blocks independently, can be parallel :
+	//for (int BlockIndex=0;BlockIndex<(int)BlockCount;BlockIndex++)
+	
+	ParallelFor(TEXT("BlockEncoder.Compress.PF"), (int32)BlockCount,1, [&CompressedBlockInfos,CompressBlockAllocSize,this](int64 BlockIndex)
+	{
+		CompressedBlockInfos[BlockIndex].CompressedData.SetNum(CompressBlockAllocSize);
+		
+		FMutableMemoryView CompressedBlockView(CompressedBlockInfos[BlockIndex].CompressedData.GetData(),CompressedBlockInfos[BlockIndex].CompressedData.Num());
+		
+		if (CompressBlock(CompressedBlockView, CompressedBlockInfos[BlockIndex].RawBlock))
+		{
+			check( CompressedBlockView.GetSize() > 0 );
+			CompressedBlockInfos[BlockIndex].CompressedData.SetNum(CompressedBlockView.GetSize());
+		}
+		else
+		{
+			// size 0 indicates error
+			CompressedBlockInfos[BlockIndex].CompressedData.SetNum(0);
+		}
+	});
+	
 	// Allocate the buffer for the header, metadata, and compressed blocks.
 	const uint64 MetaSize = sizeof(uint32) * BlockCount;
 	const uint64 CompressedDataSize = sizeof(FHeader) + MetaSize + GetCompressedBlocksBound(BlockCount, BlockSize, RawSize);
 	FUniqueBuffer CompressedData = FUniqueBuffer::Alloc(CompressedDataSize);
 
-	// Compress the raw data in blocks and store the raw data for incompressible blocks.
+	// gather results sequentially :
 	TArray64<uint32> CompressedBlockSizes;
-	CompressedBlockSizes.Reserve(uint32(BlockCount));
-	uint64 CompressedSize = 0;
+	CompressedBlockSizes.SetNum((int)BlockCount);
+	uint64 TotalCompressedSize = 0;
+	FMutableMemoryView CompressedBlocksView = CompressedData.GetView() + sizeof(FHeader) + MetaSize;
+	for (int BlockIndex=0;BlockIndex<(int)BlockCount;BlockIndex++)
 	{
-		FUniqueBuffer RawBlockCopy;
-		FMutableMemoryView CompressedBlocksView = CompressedData.GetView() + sizeof(FHeader) + MetaSize;
-		for (uint64 RawOffset = 0; RawOffset < RawSize;)
-		{
-			const uint64 RawBlockSize = FMath::Min(RawSize - RawOffset, BlockSize);
-			const FMemoryView RawBlock = RawData.ViewOrCopyRange(RawOffset, RawBlockSize, RawBlockCopy);
-			RawHash.Update(RawBlock);
+		FMemoryView CompressedBlock(CompressedBlockInfos[BlockIndex].CompressedData.GetData(),CompressedBlockInfos[BlockIndex].CompressedData.Num());
 
-			FMutableMemoryView CompressedBlock = CompressedBlocksView;
-			if (!CompressBlock(CompressedBlock, RawBlock))
-			{
-				return FCompositeBuffer();
-			}
+		uint64 CompressedBlockSize = CompressedBlock.GetSize();
+		if ( CompressedBlockSize == 0 )
+		{
+			// Compressor failed 
+			CompressedData.Reset();
+			return FNoneEncoder().Compress(RawData, BlockSize);			
+		}
+		else if (CompressedBlockInfos[BlockIndex].RawBlockSize <= CompressedBlockSize)
+		{
+			CompressedBlockSize = CompressedBlockInfos[BlockIndex].RawBlockSize;
+			CompressedBlocksView = CompressedBlocksView.CopyFrom(CompressedBlockInfos[BlockIndex].RawBlock);
+		}
+		else
+		{
+			CompressedBlocksView = CompressedBlocksView.CopyFrom(CompressedBlock);
+		}
 
-			uint64 CompressedBlockSize = CompressedBlock.GetSize();
-			if (RawBlockSize <= CompressedBlockSize)
-			{
-				CompressedBlockSize = RawBlockSize;
-				CompressedBlocksView = CompressedBlocksView.CopyFrom(RawBlock);
-			}
-			else
-			{
-				CompressedBlocksView += CompressedBlockSize;
-			}
+		// Can be freed after copied out :
+		CompressedBlockInfos[BlockIndex].CompressedData.Empty();
 
-			CompressedBlockSizes.Add(uint32(CompressedBlockSize));
-			CompressedSize += CompressedBlockSize;
-			RawOffset += RawBlockSize;
-		}
+		CompressedBlockSizes[BlockIndex] = (uint32(CompressedBlockSize));
+		TotalCompressedSize += CompressedBlockSize;
 	}
+	
+	// CompressedBlockInfos not needed anymore :
+	CompressedBlockInfos.Empty();
 
 	// Return an uncompressed buffer if the compressed data is larger than the raw data.
-	if (RawSize <= MetaSize + CompressedSize)
+	if (RawSize <= MetaSize + TotalCompressedSize)
 	{
 		CompressedData.Reset();
 		return FNoneEncoder().Compress(RawData, BlockSize);
@@ -329,7 +383,7 @@ FCompositeBuffer FBlockEncoder::Compress(const FCompositeBuffer& RawData, const
 	Header.BlockSizeExponent = uint8(FMath::FloorLog2_64(BlockSize));
 	Header.BlockCount = uint32(BlockCount);
 	Header.TotalRawSize = RawSize;
-	Header.TotalCompressedSize = sizeof(FHeader) + MetaSize + CompressedSize;
+	Header.TotalCompressedSize = sizeof(FHeader) + MetaSize + TotalCompressedSize;
 	Header.RawHash = RawHash.Finalize();
 	Header.Write(CompressedData.GetView().Left(sizeof(FHeader) + MetaSize));
 
diff --git a/Engine/Source/Runtime/Core/Private/Compression/OodleDataCompression.cpp b/Engine/Source/Runtime/Core/Private/Compression/OodleDataCompression.cpp
--- a/Engine/Source/Runtime/Core/Private/Compression/OodleDataCompression.cpp
+++ b/Engine/Source/Runtime/Core/Private/Compression/OodleDataCompression.cpp
@@ -285,6 +285,8 @@ struct OodleScratchBuffers
 	
 	int64 OodleDecode(const void * InCompBuf, int64 InCompBufSize64, void * OutRawBuf, int64 InRawLen64) 
 	{
+		TRACE_CPUPROFILER_EVENT_SCOPE(Oodle.Decode);
+
 		OO_SINTa InCompBufSize = IntCastChecked<OO_SINTa>(InCompBufSize64);
 		OO_SINTa InRawLen = IntCastChecked<OO_SINTa>(InRawLen64);
 
@@ -370,6 +372,8 @@ struct OodleScratchBuffers
 								ECompressor Compressor,
 								ECompressionLevel Level)
 	{
+		TRACE_CPUPROFILER_EVENT_SCOPE(Oodle.Encode);
+
 		OodleLZ_Compressor LZCompressor = CompressorToOodleLZ_Compressor(Compressor);
 		OodleLZ_CompressionLevel LZLevel = CompressionLevelToOodleLZ_CompressionLevel(Level);
 
