OodleDataCompression VeryVerbose logging and some cleanup

diff --git a/Engine/Source/Runtime/Core/Private/Compression/OodleDataCompression.cpp b/Engine/Source/Runtime/Core/Private/Compression/OodleDataCompression.cpp
--- a/Engine/Source/Runtime/Core/Private/Compression/OodleDataCompression.cpp
+++ b/Engine/Source/Runtime/Core/Private/Compression/OodleDataCompression.cpp
@@ -23,7 +23,7 @@ static struct { ECompressor Compressor; const TCHAR* Name; } CompressorNameMap[]
 	{ECompressor::Leviathan, TEXT("Leviathan")}
 };
 
-bool ECompressorToString(ECompressor InCompressor, const TCHAR** OutName)
+CORE_API bool ECompressorToString(ECompressor InCompressor, const TCHAR** OutName)
 {
 	if ((SIZE_T)InCompressor >= sizeof(CompressorNameMap) / sizeof(CompressorNameMap[0]))
 	{
@@ -34,8 +34,14 @@ bool ECompressorToString(ECompressor InCompressor, const TCHAR** OutName)
 	return true;
 }
 
+CORE_API const TCHAR* ECompressorToString(ECompressor InCompressor)
+{
+	const TCHAR* Ret = nullptr;
+	ECompressorToString(InCompressor,&Ret);
+	return Ret;
+}
 
-bool ECompressorFromString(const FString& InName, ECompressor& OutCompressor)
+CORE_API bool ECompressorFromString(const FString& InName, ECompressor& OutCompressor)
 {
 	for (SIZE_T i = 0; i < sizeof(CompressorNameMap) / sizeof(CompressorNameMap[0]); i++)
 	{
@@ -65,7 +71,7 @@ static struct { ECompressionLevel Level; const TCHAR* Name; } CompressionLevelNa
 	{ECompressionLevel::Optimal4, TEXT("Optimal4")}
 };
 
-bool ECompressionLevelToString(ECompressionLevel InLevel, const TCHAR** OutName)
+CORE_API bool ECompressionLevelToString(ECompressionLevel InLevel, const TCHAR** OutName)
 {
 	for (SIZE_T i = 0; i < sizeof(CompressionLevelNameMap) / sizeof(CompressionLevelNameMap[0]); i++)
 	{
@@ -78,7 +84,14 @@ bool ECompressionLevelToString(ECompressionLevel InLevel, const TCHAR** OutName)
 	return false;
 }
 
-bool ECompressionLevelFromValue(int8 InValue, ECompressionLevel& OutLevel)
+CORE_API const TCHAR* ECompressionLevelToString(ECompressionLevel InLevel)
+{
+	const TCHAR* Ret = nullptr;
+	ECompressionLevelToString(InLevel,&Ret);
+	return Ret;
+}
+
+CORE_API bool ECompressionLevelFromValue(int8 InValue, ECompressionLevel& OutLevel)
 {
 	if (InValue < OodleLZ_CompressionLevel_Min ||
 		InValue > OodleLZ_CompressionLevel_Max)
@@ -286,6 +299,8 @@ struct OodleScratchBuffers
 	int64 OodleDecode(const void * InCompBuf, int64 InCompBufSize64, void * OutRawBuf, int64 InRawLen64) 
 	{
 		TRACE_CPUPROFILER_EVENT_SCOPE(Oodle.Decode);
+		
+		UE_LOG(OodleDataCompression, VeryVerbose, TEXT("OodleDecode: %lld -> %lld"),InCompBufSize64,InRawLen64);
 
 		OO_SINTa InCompBufSize = IntCastChecked<OO_SINTa>(InCompBufSize64);
 		OO_SINTa InRawLen = IntCastChecked<OO_SINTa>(InRawLen64);
@@ -376,7 +391,12 @@ struct OodleScratchBuffers
 
 		OodleLZ_Compressor LZCompressor = CompressorToOodleLZ_Compressor(Compressor);
 		OodleLZ_CompressionLevel LZLevel = CompressionLevelToOodleLZ_CompressionLevel(Level);
-
+		
+		if ( LZCompressor == OodleLZ_Compressor_Invalid || LZLevel == OodleLZ_CompressionLevel_Invalid )
+		{
+			UE_LOG(OodleDataCompression,Error,TEXT("OodleEncode: Compressor or Level Invalid\n"));		
+			return OODLELZ_FAILED;
+		}
 		if ( InCompressedBufferSize < (int64) OodleLZ_GetCompressedBufferSizeNeeded(LZCompressor,IntCastChecked<OO_SINTa>(InUncompressedSize)) )
 		{
 			UE_LOG(OodleDataCompression,Error,TEXT("OodleEncode: OutCompressedSize too small\n"));		
@@ -408,6 +428,10 @@ struct OodleScratchBuffers
 
 				OodleScratches[i].OodleScratchMemoryMutex.Unlock();
 
+				
+				UE_LOG(OodleDataCompression, VeryVerbose, TEXT("OodleEncode: %s (%s): %lld -> %lld"), \
+					ECompressorToString(Compressor),ECompressionLevelToString(Level),InUncompressedSize,Result);
+
 				return (int64) Result;
 			}
 		}
@@ -420,6 +444,9 @@ struct OodleScratchBuffers
 		OO_SINTa Result = OodleLZ_Compress(LZCompressor,InUncompressedData,InUncompressedSize,OutCompressedData,LZLevel,
 													NULL,NULL,NULL,
 													scratchMem,scratchSize);
+													
+		UE_LOG(OodleDataCompression, VeryVerbose, TEXT("OodleEncode: %s (%s): %lld -> %lld"), \
+			ECompressorToString(Compressor),ECompressionLevelToString(Level),InUncompressedSize,Result);
 
 		return (int64) Result;
 	}
diff --git a/Engine/Source/Runtime/Core/Private/Tests/Compression/CompressedBufferTest.cpp b/Engine/Source/Runtime/Core/Private/Tests/Compression/CompressedBufferTest.cpp
--- a/Engine/Source/Runtime/Core/Private/Tests/Compression/CompressedBufferTest.cpp
+++ b/Engine/Source/Runtime/Core/Private/Tests/Compression/CompressedBufferTest.cpp
@@ -166,7 +166,7 @@ bool FCompressedBufferDecompressTest::RunTest(const FString& Parameters)
 		const FCompressedBuffer Compressed = FCompressedBuffer::Compress(
 			FSharedBuffer::MakeView(MakeMemoryView(ExpectedValues)),
 			FOodleDataCompression::ECompressor::Mermaid,
-			FOodleDataCompression::ECompressionLevel::Optimal4,
+			FOodleDataCompression::ECompressionLevel::Fast,
 			BlockSize);
 
 		UncompressAndValidate(Compressed, 0, N, ExpectedValues);
@@ -191,7 +191,7 @@ bool FCompressedBufferDecompressTest::RunTest(const FString& Parameters)
 		const FCompressedBuffer Compressed = FCompressedBuffer::Compress(
 			FSharedBuffer::MakeView(MakeMemoryView(ExpectedValues)),
 			FOodleDataCompression::ECompressor::Mermaid,
-			FOodleDataCompression::ECompressionLevel::Optimal4,
+			FOodleDataCompression::ECompressionLevel::Fast,
 			BlockSize);
 		
 		constexpr uint64 OffsetCount = 150;
@@ -229,7 +229,7 @@ bool FCompressedBufferDecompressTest::RunTest(const FString& Parameters)
 		const FCompressedBuffer Compressed = FCompressedBuffer::Compress(
 			FSharedBuffer::MakeView(MakeMemoryView(ExpectedValues)),
 			FOodleDataCompression::ECompressor::Mermaid,
-			FOodleDataCompression::ECompressionLevel::Optimal4,
+			FOodleDataCompression::ECompressionLevel::Fast,
 			BlockSize);
 
 		constexpr uint64 OffsetCount = 2;
diff --git a/Engine/Source/Runtime/Core/Public/Compression/OodleDataCompression.h b/Engine/Source/Runtime/Core/Public/Compression/OodleDataCompression.h
--- a/Engine/Source/Runtime/Core/Public/Compression/OodleDataCompression.h
+++ b/Engine/Source/Runtime/Core/Public/Compression/OodleDataCompression.h
@@ -54,6 +54,7 @@ enum class ECompressor : uint8
 // If the input is invalid, returns false, and doesn't touch output.
 bool CORE_API ECompressorToString(ECompressor InCompressor, const TCHAR** OutName);
 bool CORE_API ECompressorFromString(const class FString& InName, ECompressor& OutCompressor);
+CORE_API const TCHAR* ECompressorToString(ECompressor InCompressor);
 
 /**
  * ECompressionLevel : Choose the Oodle Compression Level
@@ -106,6 +107,7 @@ enum class ECompressionLevel : int8
 // If the input is invalid, returns false, and doesn't touch output.
 bool CORE_API ECompressionLevelFromValue(int8 InValue, ECompressionLevel& OutLevel);
 bool CORE_API ECompressionLevelToString(ECompressionLevel InLevel, const TCHAR** OutName);
+CORE_API const TCHAR * ECompressionLevelToString(ECompressionLevel InLevel);
 
 enum class ECompressionCommonUsage : uint8
 {
