DepletableMpscQueue: Added an AllocatorType template parameter

The AllocatorType must have static void* Malloc(SIZE_T Size, uint32 Alignment) and static void Free(void*). This can be extended in the future to make the allocator a member if desired, without breaking compatibility with this version.

The default AllocatorType uses FMemory::Malloc and FMemory::Free.

#rnx

diff --git a/Engine/Source/Runtime/Core/Public/Containers/DepletableMpscQueue.h b/Engine/Source/Runtime/Core/Public/Containers/DepletableMpscQueue.h
--- a/Engine/Source/Runtime/Core/Public/Containers/DepletableMpscQueue.h
+++ b/Engine/Source/Runtime/Core/Public/Containers/DepletableMpscQueue.h
@@ -3,8 +3,10 @@
 #pragma once
 
 #include "CoreTypes.h"
+#include "HAL/UnrealMemory.h"
 #include "Templates/TypeCompatibleBytes.h"
 #include <atomic>
+#include <type_traits>
 
 namespace UE
 {
@@ -14,7 +16,7 @@ namespace UE
 	 * A typical use case is when the consumer doesn't stop until the queue is depleted.
 	 * Is faster than traditional MPSC queues, especially for consumer.
 	 */
-	template<typename T>
+	template<typename T, typename AllocatorType = FMemory>
 	class TDepletableMpscQueue final
 	{
 	private:
@@ -40,7 +42,8 @@ namespace UE
 			{
 				DestructItem(Node->Value.GetTypedPtr());
 				FNode* Next = Node->Next.load(std::memory_order_relaxed);
-				delete Node;
+				static_assert(std::is_trivially_destructible_v<FNode>);
+				AllocatorType::Free(Node);
 				Node = Next;
 			}
 		}
@@ -54,7 +57,7 @@ namespace UE
 		template <typename... ArgTypes>
 		bool EnqueueAndReturnWasEmpty(ArgTypes&&... Args)
 		{
-			FNode* New = new FNode;
+			FNode* New = new(AllocatorType::Malloc(sizeof(FNode), alignof(FNode))) FNode;
 			new (&New->Value) T(Forward<ArgTypes>(Args)...);
 
 			// switch `Tail` to the new node and only then link the old tail to the new one. The list is not fully linked between these ops,
@@ -132,7 +135,8 @@ namespace UE
 			{
 				Consume(First);
 				FNode* Next = GetNext(First);
-				delete First;
+				static_assert(std::is_trivially_destructible_v<FNode>);
+				AllocatorType::Free(First);
 				First = Next;
 			}
 
