Logging: Optimized the log redirector to handle logging from multiple threads more efficiently

Timing from a synthetic benchmark that logs 131,072 lines at Log verbosity using ParallelFor executing on 62 worker threads.
- Before: 160ms wall time, 73.2us/log
- After: 6.8ms wall time, 2.9us/log

diff --git a/Engine/Source/Runtime/Core/Private/Misc/OutputDeviceRedirector.cpp b/Engine/Source/Runtime/Core/Private/Misc/OutputDeviceRedirector.cpp
--- a/Engine/Source/Runtime/Core/Private/Misc/OutputDeviceRedirector.cpp
+++ b/Engine/Source/Runtime/Core/Private/Misc/OutputDeviceRedirector.cpp
@@ -2,295 +2,287 @@
 
 #include "Misc/OutputDeviceRedirector.h"
 
-#include "Containers/BitArray.h"
+#include "Containers/DepletableMpscQueue.h"
+#include "Experimental/ConcurrentLinearAllocator.h"
 #include "HAL/PlatformProcess.h"
-#include "Misc/CoreStats.h"
+#include "HAL/PlatformTime.h"
+#include "HAL/PlatformTLS.h"
 #include "Misc/ScopeLock.h"
-#include "Stats/Stats.h"
+#include "Misc/ScopeRWLock.h"
+#include "ProfilingDebugging/CpuProfilerTrace.h"
+#include <atomic>
 
 /*-----------------------------------------------------------------------------
 	FOutputDeviceRedirector.
 -----------------------------------------------------------------------------*/
 
-class FLogAllocator
+FBufferedLine::FBufferedLine(const TCHAR* InData, const FName& InCategory, ELogVerbosity::Type InVerbosity, double InTime)
+	: Category(InCategory)
+	, Time(InTime)
+	, Verbosity(InVerbosity)
 {
-public:
-	bool HasSpace(int32 NumChars) const
-	{
-		return Data[BufferIndex].Num() + NumChars <= BufferSize;
-	}
+	int32 NumChars = FCString::Strlen(InData) + 1;
+	void* Dest = FMemory::Malloc(sizeof(TCHAR) * NumChars);
+	Data = (TCHAR*)FMemory::Memcpy(Dest, InData, sizeof(TCHAR) * NumChars);
+}
 
-	TCHAR* Alloc(int32 NumChars)
-	{
-		check(HasSpace(NumChars));
-		int32 DataIndex = Data[BufferIndex].AddUninitialized(NumChars);
-		return Data[BufferIndex].GetData() + DataIndex;
-	}
+FBufferedLine::~FBufferedLine()
+{
+	FMemory::Free(const_cast<TCHAR*>(Data));
+}
 
-	/**
-	 * Swap which buffer is being used for new allocations.
-	 *
-	 * Buffer remains valid until BufferCount calls are made to this function.
-	 */
-	void SwapBuffers()
+namespace UE::Private
+{
+
+struct FOutputDeviceBlockAllocationTag : FDefaultBlockAllocationTag
+{
+	static constexpr const TCHAR* TagName = TEXT("OutputDeviceLinear");
+
+	struct Allocator
 	{
-		BufferIndex = (BufferIndex + 1) % BufferCount;
-		// A failed check within InternalFlushThreadedLogs can cause a stack overflow. This can
-		// currently fail because InternalFlushThreadedLogs can be called from multiple threads
-		// simultaneously if errors cause SetCurrentThreadAsMasterThread to be called from more
-		// than one thread at once, or while the game thread is in InternalFlushThreadedLogs.
-		//check(!DataLocked[BufferIndex]);
-		Data[BufferIndex].Empty();
-	}
+		static constexpr bool SupportsAlignment = false;
 
-	struct FBufferLock 
-	{ 
-		int32 Index = -1;
+		FORCEINLINE static void* Malloc(SIZE_T Size, uint32 Alignment)
+		{
+			return FMemory::Malloc(Size, Alignment);
+		}
 
-		/** Returns true if this object represents a locked buffer */
-		bool IsValid() const
+		FORCEINLINE static void Free(void* Pointer, SIZE_T Size)
 		{
-			return Index >= 0;
+			return FMemory::Free(Pointer);
 		}
 	};
+};
 
-	FBufferLock LockBuffer()
+struct FOutputDeviceLinearAllocator
+{
+	FORCEINLINE static void* Malloc(SIZE_T Size, uint32 Alignment)
 	{
-		DataLocked[BufferIndex] = true;
-		return {BufferIndex};
+		return TConcurrentLinearAllocator<FOutputDeviceBlockAllocationTag>::Malloc(Size, Alignment);
 	}
 
-	void UnlockBuffer(FBufferLock Lock)
+	FORCEINLINE static void Free(void* Pointer)
+	{
+		TConcurrentLinearAllocator<FOutputDeviceBlockAllocationTag>::Free(Pointer);
+	}
+};
+
+struct FOutputDeviceLine
+{
+	const double Time;
+	const TCHAR* Data;
+	const FName Category;
+	const ELogVerbosity::Type Verbosity;
+
+	FOutputDeviceLine(const FBufferedLine&) = delete;
+	FOutputDeviceLine& operator=(const FBufferedLine&) = delete;
+
+	FORCEINLINE FOutputDeviceLine(const TCHAR* const InData, const FName InCategory, const ELogVerbosity::Type InVerbosity, const double InTime)
+		: Time(InTime)
+		, Data(CopyData(InData))
+		, Category(InCategory)
+		, Verbosity(InVerbosity)
 	{
-		DataLocked[Lock.Index] = false;
+	}
+
+	FORCEINLINE ~FOutputDeviceLine()
+	{
+		FOutputDeviceLinearAllocator::Free(const_cast<TCHAR*>(Data));
 	}
 
 private:
-	static constexpr int32 BufferSize = 4096;
-	// BufferCount can be 2 once the check in SwapBuffers can be safely re-enabled. Using more
-	// buffers in the interim minimizes the likelihood of writes to a buffer while it is still
-	// being flushed by another thread.
-	static constexpr int32 BufferCount = 4;
-	TArray<TCHAR, TInlineAllocator<BufferSize>> Data[BufferCount];
-	TBitArray<TInlineAllocator<1>> DataLocked{false, BufferCount};
-	int32 BufferIndex = 0;
+	FORCEINLINE static const TCHAR* CopyData(const TCHAR* const InData)
+	{
+		const int32 Len = FCString::Strlen(InData) + 1;
+		void* const Dest = FOutputDeviceLinearAllocator::Malloc(sizeof(TCHAR) * Len, alignof(TCHAR));
+		return static_cast<TCHAR*>(FMemory::Memcpy(Dest, InData, sizeof(TCHAR) * Len));
+	}
 };
 
-FBufferedLine::FBufferedLine(const TCHAR* InData, const FName& InCategory, ELogVerbosity::Type InVerbosity, double InTime, FLogAllocator* ExternalAllocator)
-	: FBufferedLine(InData, FLazyName(InCategory), InVerbosity, InTime, ExternalAllocator)
+struct FOutputDeviceRedirectorState
 {
-}
+	FRWLock OutputDevicesLock;
+	std::atomic<uint32> OutputDevicesLockState;
+	uint8 OutputDevicesLockPadding[PLATFORM_CACHE_LINE_SIZE - sizeof(OutputDevicesLock) - sizeof(OutputDevicesLockState)]{};
 
-FBufferedLine::FBufferedLine(const TCHAR* InData, const FLazyName& InCategory, ELogVerbosity::Type InVerbosity, double InTime, FLogAllocator* ExternalAllocator)
-	: Category(InCategory)
-	, Time(InTime)
-	, Verbosity(InVerbosity)
-{
-	int32 NumChars = FCString::Strlen(InData) + 1;
-	bExternalAllocation = ExternalAllocator && ExternalAllocator->HasSpace(NumChars);
-	void* Dest = bExternalAllocation ? ExternalAllocator->Alloc(NumChars) : FMemory::Malloc(sizeof(TCHAR) * NumChars);
-	Data = (TCHAR*)FMemory::Memcpy(Dest, InData, sizeof(TCHAR) * NumChars);
-}
+	/** A FIFO of lines logged by non-master threads. */
+	TDepletableMpscQueue<FOutputDeviceLine, FOutputDeviceLinearAllocator> BufferedLines;
+	uint8 BufferedLinesPadding[PLATFORM_CACHE_LINE_SIZE - sizeof(BufferedLines)]{};
 
-FBufferedLine::~FBufferedLine()
-{
-	if (!bExternalAllocation)
+	/** Array of output devices to redirect to from the master thread. */
+	TArray<FOutputDevice*> BufferedOutputDevices;
+
+	/** Array of output devices to redirect to from the calling thread. */
+	TArray<FOutputDevice*> UnbufferedOutputDevices;
+
+	/** The ID of the master thread. Logging from other threads will be buffered for processing by the master thread. */
+	uint32 MasterThreadID = FPlatformTLS::GetCurrentThreadId();
+
+	/** Objects used for synchronization via a scoped lock. */
+	FCriticalSection SynchronizationObject;
+
+	/** A FIFO backlog of messages logged before the editor had a chance to intercept them. */
+	TArray<FBufferedLine> BacklogLines;
+
+	/** Whether backlogging is enabled. */
+	bool bEnableBacklog = false;
+
+	bool IsInMasterThread() const
 	{
-		FMemory::Free(const_cast<TCHAR*>(Data));
+		return MasterThreadID == FPlatformTLS::GetCurrentThreadId();
 	}
-}
+};
 
-/** Initialization constructor. */
-FOutputDeviceRedirector::FOutputDeviceRedirector(FLogAllocator* Allocator)
-:	MasterThreadID(FPlatformTLS::GetCurrentThreadId())
-,	bEnableBacklog(false)
-,	BufferedLinesAllocator(Allocator)
+struct FOutputDevicesReadScopeLock
 {
-}
+	FORCEINLINE explicit FOutputDevicesReadScopeLock(FOutputDeviceRedirectorState& InState)
+		: State(InState)
+	{
+		// Read locks add/sub by 2 to keep the LSB free for write locks to use.
+		if (State.OutputDevicesLockState.fetch_add(2, std::memory_order_relaxed) & 1)
+		{
+			WaitForWriteLock();
+		}
+	}
 
-static FLogAllocator GLogAllocator;
+	FORCENOINLINE void WaitForWriteLock()
+	{
+		TRACE_CPUPROFILER_EVENT_SCOPE(FOutputDevicesReadScopeLock);
+		// A write lock has set the LSB. Cancel this read lock and wait for the write.
+		State.OutputDevicesLockState.fetch_sub(2, std::memory_order_relaxed);
+		// This read lock will wait until the write lock exits.
+		FReadScopeLock ScopeLock(State.OutputDevicesLock);
+		// Acquire on this read lock because the write may have mutated state that we read.
+		uint32 LockState = State.OutputDevicesLockState.fetch_add(2, std::memory_order_acquire);
+		check((LockState & 1) == 0);
+	}
 
-FOutputDeviceRedirector* FOutputDeviceRedirector::Get()
-{
-	static FOutputDeviceRedirector Singleton(&GLogAllocator);
-	return &Singleton;
-}
+	FORCEINLINE ~FOutputDevicesReadScopeLock()
+	{
+		State.OutputDevicesLockState.fetch_sub(2, std::memory_order_relaxed);
+	}
+
+	FOutputDeviceRedirectorState& State;
+};
 
-/**
- * Adds an output device to the chain of redirections.	
- *
- * @param OutputDevice	output device to add
- */
-void FOutputDeviceRedirector::AddOutputDevice( FOutputDevice* OutputDevice )
+struct FOutputDevicesWriteScopeLock
 {
-	if (OutputDevice)
+	FORCEINLINE explicit FOutputDevicesWriteScopeLock(FOutputDeviceRedirectorState& InState)
+		: State(InState)
 	{
-		bool bAdded = false;
-		do
+		TRACE_CPUPROFILER_EVENT_SCOPE(FOutputDevicesWriteScopeLock);
+		// Take the lock before modifying the state, to avoid contention on the LSB.
+		State.OutputDevicesLock.WriteLock();
+		// Set the LSB to flag to read locks that a write lock is waiting.
+		uint32 LockState = State.OutputDevicesLockState.fetch_or(uint32(1), std::memory_order_relaxed);
+		check((LockState & 1) == 0);
+		if (LockState > 1)
 		{
-			{
-				FScopeLock ScopeLock(&OutputDevicesMutex);
-				if (OutputDevicesLockCounter.GetValue() == 0)
-				{
-					if (OutputDevice->CanBeUsedOnMultipleThreads())
-					{
-						UnbufferedOutputDevices.AddUnique(OutputDevice);
-					}
-					else
-					{
-						BufferedOutputDevices.AddUnique(OutputDevice);
-					}
-					bAdded = true;
-				}
-			}
-			if (!bAdded)
+			// Wait for read locks to be cleared.
+			do
 			{
 				FPlatformProcess::Sleep(0);
+				LockState = State.OutputDevicesLockState.load(std::memory_order_relaxed);
 			}
-		} while (!bAdded);
+			while (LockState > 1);
+		}
 	}
+
+	FORCEINLINE ~FOutputDevicesWriteScopeLock()
+	{
+		// Clear the LSB to allow read locks after the unlock below.
+		uint32 LockState = State.OutputDevicesLockState.fetch_and(~uint32(1), std::memory_order_release);
+		check((LockState & 1) == 1);
+		State.OutputDevicesLock.WriteUnlock();
+	}
+
+	FOutputDeviceRedirectorState& State;
+};
+
+} // UE::Private
+
+FOutputDeviceRedirector::FOutputDeviceRedirector()
+	: State(MakePimpl<UE::Private::FOutputDeviceRedirectorState>())
+{
 }
 
-/**
- * Removes an output device from the chain of redirections.	
- *
- * @param OutputDevice	output device to remove
- */
-void FOutputDeviceRedirector::RemoveOutputDevice( FOutputDevice* OutputDevice )
+FOutputDeviceRedirector* FOutputDeviceRedirector::Get()
 {
-	bool bRemoved = false;
-	do
+	static FOutputDeviceRedirector Singleton;
+	return &Singleton;
+}
+
+void FOutputDeviceRedirector::AddOutputDevice(FOutputDevice* OutputDevice)
+{
+	if (OutputDevice)
 	{
+		UE::Private::FOutputDevicesWriteScopeLock ScopeLock(*State);
+		if (OutputDevice->CanBeUsedOnMultipleThreads())
 		{
-			FScopeLock ScopeLock(&OutputDevicesMutex);
-			if (OutputDevicesLockCounter.GetValue() == 0)
-			{
-				BufferedOutputDevices.Remove( OutputDevice );
-				UnbufferedOutputDevices.Remove(OutputDevice);
-				bRemoved = true;
-			}
+			State->UnbufferedOutputDevices.AddUnique(OutputDevice);
 		}
-		if (!bRemoved)
+		else
 		{
-			FPlatformProcess::Sleep(0);
+			State->BufferedOutputDevices.AddUnique(OutputDevice);
 		}
-	} while (!bRemoved);
+	}
 }
 
-/**
- * Returns whether an output device is currently in the list of redirectors.
- *
- * @param	OutputDevice	output device to check the list against
- * @return	true if messages are currently redirected to the the passed in output device, false otherwise
- */
-bool FOutputDeviceRedirector::IsRedirectingTo( FOutputDevice* OutputDevice )
+void FOutputDeviceRedirector::RemoveOutputDevice(FOutputDevice* OutputDevice)
 {
-	// For performance reasons whe're not using the FOutputDevicesLock here
-	FScopeLock OutputDevicesLock(&OutputDevicesMutex);
-	return BufferedOutputDevices.Contains(OutputDevice) || UnbufferedOutputDevices.Contains(OutputDevice);
+	if (OutputDevice)
+	{
+		UE::Private::FOutputDevicesWriteScopeLock ScopeLock(*State);
+		State->BufferedOutputDevices.Remove(OutputDevice);
+		State->UnbufferedOutputDevices.Remove(OutputDevice);
+	}
 }
 
-void FOutputDeviceRedirector::InternalFlushThreadedLogs(bool bUseAllDevices)
+bool FOutputDeviceRedirector::IsRedirectingTo(FOutputDevice* OutputDevice)
 {
-	TLocalOutputDevicesArray LocalBufferedDevices;
-	TLocalOutputDevicesArray LocalUnbufferedDevices;
-	FOutputDevicesLock OutputDevicesLock(this, LocalBufferedDevices, LocalUnbufferedDevices);
-
-	InternalFlushThreadedLogs(LocalBufferedDevices, LocalUnbufferedDevices, bUseAllDevices);
+	UE::Private::FOutputDevicesReadScopeLock Lock(*State);
+	return State->BufferedOutputDevices.Contains(OutputDevice) || State->UnbufferedOutputDevices.Contains(OutputDevice);
 }
 
-/**
- * The unsynchronized version of FlushThreadedLogs.
- */
-void FOutputDeviceRedirector::InternalFlushThreadedLogs(TLocalOutputDevicesArray& InBufferedDevices, TLocalOutputDevicesArray& InUnbufferedDevices, bool bUseAllDevices)
+void FOutputDeviceRedirector::FlushBufferedLines(TConstArrayView<FOutputDevice*> BufferedDevices, bool bUseAllDevices)
 {	
-	if (BufferedLines.Num())
+	if (State->BufferedLines.IsEmpty())
 	{
-		TRACE_CPUPROFILER_EVENT_SCOPE(FOutputDeviceRedirector::InternalFlushThreadedLogs);
+		return;
+	}
 
-		TArray<FBufferedLine, TInlineAllocator<64>> LocalBufferedLines;
-		FLogAllocator::FBufferLock BufferLock;
-		{
-			FScopeLock ScopeLock(&BufferSynchronizationObject);
-			// Copy the buffered lines only if there's any buffered output devices
-			if (InBufferedDevices.Num())
-			{
-				LocalBufferedLines.AddUninitialized(BufferedLines.Num());
-				for (int32 LineIndex = 0; LineIndex < BufferedLines.Num(); LineIndex++)
-				{
-					new(&LocalBufferedLines[LineIndex]) FBufferedLine(BufferedLines[LineIndex], FBufferedLine::EMoveCtor);
-				}
-			}
+	TRACE_CPUPROFILER_EVENT_SCOPE(FOutputDeviceRedirector::FlushBufferedLines);
 
-			// If there's no output devices to redirect to (assumption is that we haven't added any yet)
-			// don't clear the buffer otherwise its content will be lost (for example when calling SetCurrentThreadAsMasterThread() on init)
-			if (InBufferedDevices.Num() || InUnbufferedDevices.Num())
-			{
-				if (BufferedLinesAllocator)
-				{
-					BufferLock = BufferedLinesAllocator->LockBuffer();
-				}
-				EmptyBufferedLines();
-			}
-		}
+	State->BufferedLines.Deplete([BufferedDevices, bUseAllDevices](UE::Private::FOutputDeviceLine&& Line)
+	{
+		const double Time = Line.Time;
+		const TCHAR* const Data = Line.Data;
+		const FName Category = Line.Category;
+		const ELogVerbosity::Type Verbosity = Line.Verbosity;
 
-		for (FBufferedLine& Line : LocalBufferedLines)
+		for (FOutputDevice* OutputDevice : BufferedDevices)
 		{
-			const TCHAR* Data = Line.Data;
-			const FLazyName Category = Line.Category;
-			const double Time = Line.Time;
-			const ELogVerbosity::Type Verbosity = Line.Verbosity;
-
-			for (FOutputDevice* OutputDevice : InBufferedDevices)
+			if (bUseAllDevices || OutputDevice->CanBeUsedOnAnyThread())
 			{
-				if (OutputDevice->CanBeUsedOnAnyThread() || bUseAllDevices)
-				{
-					OutputDevice->Serialize(Data, Verbosity, Category, Time);
-				}
+				OutputDevice->Serialize(Data, Verbosity, Category, Time);
 			}
 		}
-
-		if (BufferLock.IsValid())
-		{
-			FScopeLock ScopeLock(&BufferSynchronizationObject);
-			BufferedLinesAllocator->UnlockBuffer(BufferLock);
-		}
-	}
+	});
 }
 
-void FOutputDeviceRedirector::EmptyBufferedLines()
-{
-	BufferedLines.Empty();
-
-	if (BufferedLinesAllocator)
-	{
-		BufferedLinesAllocator->SwapBuffers();
-	}
-}
-
-/**
- * Flushes lines buffered by secondary threads.
- */
 void FOutputDeviceRedirector::FlushThreadedLogs()
 {
-	//QUICK_SCOPE_CYCLE_COUNTER(STAT_FlushThreadedLogs);
-	check(IsInGameThread());
-	InternalFlushThreadedLogs(true);
+	check(State->IsInMasterThread());
+	UE::Private::FOutputDevicesReadScopeLock Lock(*State);
+	FlushBufferedLines(State->BufferedOutputDevices, /*bUseAllDevices*/ true);
 }
 
 void FOutputDeviceRedirector::PanicFlushThreadedLogs()
 {
-	//QUICK_SCOPE_CYCLE_COUNTER(STAT_FlushThreadedLogs);
-
-	TLocalOutputDevicesArray LocalBufferedDevices;
-	TLocalOutputDevicesArray LocalUnbufferedDevices;
-	FOutputDevicesLock OutputDevicesLock(this, LocalBufferedDevices, LocalUnbufferedDevices);
-
-	// Flush threaded logs, but use the safe version.
-	InternalFlushThreadedLogs(LocalBufferedDevices, LocalUnbufferedDevices, false);
+	UE::Private::FOutputDevicesReadScopeLock Lock(*State);
+	FlushBufferedLines(State->BufferedOutputDevices, /*bUseAllDevices*/ State->IsInMasterThread());
 
-	// Flush devices.
-	for (FOutputDevice* OutputDevice : LocalBufferedDevices)
+	for (FOutputDevice* OutputDevice : State->BufferedOutputDevices)
 	{
 		if (OutputDevice->CanBeUsedOnAnyThread())
 		{
@@ -298,85 +290,49 @@ void FOutputDeviceRedirector::PanicFlushThreadedLogs()
 		}
 	}
 
-	for (FOutputDevice* OutputDevice : LocalUnbufferedDevices)
+	for (FOutputDevice* OutputDevice : State->UnbufferedOutputDevices)
 	{
 		OutputDevice->Flush();
 	}
 }
 
-/**
- * Serializes the current backlog to the specified output device.
- * @param OutputDevice	- Output device that will receive the current backlog
- */
-void FOutputDeviceRedirector::SerializeBacklog( FOutputDevice* OutputDevice )
+void FOutputDeviceRedirector::SerializeBacklog(FOutputDevice* OutputDevice)
 {
-	FScopeLock ScopeLock( &SynchronizationObject );
-
-	for (int32 LineIndex = 0; LineIndex < BacklogLines.Num(); LineIndex++)
+	FScopeLock ScopeLock(&State->SynchronizationObject);
+	for (const FBufferedLine& BacklogLine : State->BacklogLines)
 	{
-		const FBufferedLine& BacklogLine = BacklogLines[ LineIndex ];
-		OutputDevice->Serialize( BacklogLine.Data, BacklogLine.Verbosity, BacklogLine.Category, BacklogLine.Time );
+		OutputDevice->Serialize(BacklogLine.Data, BacklogLine.Verbosity, BacklogLine.Category, BacklogLine.Time);
 	}
 }
 
-/**
- * Enables or disables the backlog.
- * @param bEnable	- Starts saving a backlog if true, disables and discards any backlog if false
- */
-void FOutputDeviceRedirector::EnableBacklog( bool bEnable )
+void FOutputDeviceRedirector::EnableBacklog(bool bEnable)
 {
-	FScopeLock ScopeLock( &SynchronizationObject );
-
-	bEnableBacklog = bEnable;
-	if ( bEnableBacklog == false )
+	FScopeLock ScopeLock(&State->SynchronizationObject);
+	State->bEnableBacklog = bEnable;
+	if (!bEnable)
 	{
-		BacklogLines.Empty();
+		State->BacklogLines.Empty();
 	}
 }
 
-/**
- * Sets the current thread to be the master thread that prints directly
- * (isn't queued up)
- */
 void FOutputDeviceRedirector::SetCurrentThreadAsMasterThread()
 {
-	InternalFlushThreadedLogs( false );
-
-	// make sure anything queued up is flushed out, this may be called from a background thread, so use the safe version.
-	{
-		FScopeLock ScopeLock(&SynchronizationObject);
-		// set the current thread as the master thread
-		MasterThreadID = FPlatformTLS::GetCurrentThreadId();
-	}
-}
-
-void FOutputDeviceRedirector::LockOutputDevices(TLocalOutputDevicesArray& OutBufferedDevices, TLocalOutputDevicesArray& OutUnbufferedDevices)
-{
-	FScopeLock OutputDevicesLock(&OutputDevicesMutex);
-	OutputDevicesLockCounter.Increment();
-	OutBufferedDevices.Append(BufferedOutputDevices);
-	OutUnbufferedDevices.Append(UnbufferedOutputDevices);
-}
+	UE::Private::FOutputDevicesReadScopeLock Lock(*State);
+	FlushBufferedLines(State->BufferedOutputDevices, /*bUseAllDevices*/ State->IsInMasterThread());
 
-void FOutputDeviceRedirector::UnlockOutputDevices()
-{
-	FScopeLock OutputDevicesLock(&OutputDevicesMutex);
-	int32 LockValue = OutputDevicesLockCounter.Decrement();
-	check(LockValue >= 0);
+	FScopeLock ScopeLock(&State->SynchronizationObject);
+	State->MasterThreadID = FPlatformTLS::GetCurrentThreadId();
 }
 
-template<class T>
-void FOutputDeviceRedirector::SerializeImpl(const TCHAR* Data, ELogVerbosity::Type Verbosity, T& Category, const double Time)
+void FOutputDeviceRedirector::Serialize(const TCHAR* const Data, const ELogVerbosity::Type Verbosity, const FName& Category, const double Time)
 {
 	const double RealTime = Time == -1.0f ? FPlatformTime::Seconds() - GStartTime : Time;
 
-	TLocalOutputDevicesArray LocalBufferedDevices;
-	TLocalOutputDevicesArray LocalUnbufferedDevices;
-	FOutputDevicesLock OutputDevicesLock(this, LocalBufferedDevices, LocalUnbufferedDevices);
+	UE::Private::FOutputDevicesReadScopeLock Lock(*State);
 
 #if PLATFORM_DESKTOP
-	// this is for errors which occur after shutdown we might be able to salvage information from stdout 
-	if ((LocalBufferedDevices.Num() == 0) && IsEngineExitRequested())
+	// this is for errors which occur after shutdown we might be able to salvage information from stdout
+	if (State->BufferedOutputDevices.IsEmpty() && IsEngineExitRequested())
 	{
 #if PLATFORM_WINDOWS
 		_tprintf(_T("%s\n"), Data);
@@ -389,111 +345,89 @@ void FOutputDeviceRedirector::SerializeImpl(const TCHAR* Data, ELogVerbosity::Ty
 #endif
 
 	// Serialize directly to any output devices which don't require buffering
-	for (FOutputDevice* OutputDevice : LocalUnbufferedDevices)
+	for (FOutputDevice* OutputDevice : State->UnbufferedOutputDevices)
 	{
 		OutputDevice->Serialize(Data, Verbosity, Category, RealTime);
 	}
 
-
-	if (bEnableBacklog)
+	if (State->bEnableBacklog)
 	{
-		FScopeLock ScopeLock(&SynchronizationObject);
-		new(BacklogLines)FBufferedLine(Data, Category, Verbosity, RealTime, nullptr);
+		FScopeLock ScopeLock(&State->SynchronizationObject);
+		new(State->BacklogLines)FBufferedLine(Data, Category, Verbosity, RealTime);
 	}
 
-	if (FPlatformTLS::GetCurrentThreadId() != MasterThreadID || LocalBufferedDevices.Num() == 0)
+	if (!State->IsInMasterThread() || State->BufferedOutputDevices.IsEmpty())
 	{
-		FScopeLock ScopeLock(&BufferSynchronizationObject);
-		new(BufferedLines)FBufferedLine(Data, Category, Verbosity, RealTime, BufferedLinesAllocator);
+		State->BufferedLines.Enqueue(Data, Category, Verbosity, RealTime);
 	}
 	else
 	{
-		// Flush previously buffered lines from secondary threads.
-		InternalFlushThreadedLogs(LocalBufferedDevices, LocalUnbufferedDevices, true);
+		FlushBufferedLines(State->BufferedOutputDevices, /*bUseAllDevices*/ true);
 
-		for (FOutputDevice* OutputDevice : LocalBufferedDevices)
+		for (FOutputDevice* OutputDevice : State->BufferedOutputDevices)
 		{
 			OutputDevice->Serialize(Data, Verbosity, Category, RealTime);
 		}
 	}
 }
 
-void FOutputDeviceRedirector::Serialize( const TCHAR* Data, ELogVerbosity::Type Verbosity, const class FName& Category, const double Time )
-{
-	SerializeImpl(Data, Verbosity, Category, Time);
-}
-
-void FOutputDeviceRedirector::Serialize( const TCHAR* Data, ELogVerbosity::Type Verbosity, const class FName& Category )
+void FOutputDeviceRedirector::Serialize(const TCHAR* Data, ELogVerbosity::Type Verbosity, const FName& Category)
 {
-	SerializeImpl( Data, Verbosity, Category, -1.0 );
+	Serialize(Data, Verbosity, Category, -1.0);
 }
 
 void FOutputDeviceRedirector::RedirectLog(const FName& Category, ELogVerbosity::Type Verbosity, const TCHAR* Data)
 {
-	SerializeImpl(Data, Verbosity, Category, -1.0);
+	Serialize(Data, Verbosity, Category, -1.0);
 }
 
 void FOutputDeviceRedirector::RedirectLog(const FLazyName& Category, ELogVerbosity::Type Verbosity, const TCHAR* Data)
 {
-	SerializeImpl(Data, Verbosity, Category, -1.0);
+	Serialize(Data, Verbosity, Category, -1.0);
 }
 
-/**
- * Passes on the flush request to all current output devices.
- */
 void FOutputDeviceRedirector::Flush()
 {
-	TLocalOutputDevicesArray LocalBufferedDevices;
-	TLocalOutputDevicesArray LocalUnbufferedDevices;
-	FOutputDevicesLock OutputDevicesLock(this, LocalBufferedDevices, LocalUnbufferedDevices);
+	UE::Private::FOutputDevicesReadScopeLock Lock(*State);
 
-	if(FPlatformTLS::GetCurrentThreadId() == MasterThreadID)
+	if (State->IsInMasterThread())
 	{
-		// Flush previously buffered lines from secondary threads.
-		InternalFlushThreadedLogs(true);
+		FlushBufferedLines(State->BufferedOutputDevices, /*bUseAllDevices*/ true);
 
-		for (FOutputDevice* OutputDevice : LocalBufferedDevices)
+		for (FOutputDevice* OutputDevice : State->BufferedOutputDevices)
 		{
 			OutputDevice->Flush();
 		}
 	}
 
-	for (FOutputDevice* OutputDevice : LocalUnbufferedDevices)
+	for (FOutputDevice* OutputDevice : State->UnbufferedOutputDevices)
 	{
 		OutputDevice->Flush();
 	}
 }
 
-/**
- * Closes output device and cleans up. This can't happen in the destructor
- * as we might have to call "delete" which cannot be done for static/ global
- * objects.
- */
 void FOutputDeviceRedirector::TearDown()
 {
-	FScopeLock SyncLock(&SynchronizationObject);
-	check(FPlatformTLS::GetCurrentThreadId() == MasterThreadID);
+	FScopeLock SyncLock(&State->SynchronizationObject);
+	check(State->IsInMasterThread());
 
-	TLocalOutputDevicesArray LocalBufferedDevices;
-	TLocalOutputDevicesArray LocalUnbufferedDevices;
+	TArray<FOutputDevice*> LocalBufferedDevices;
+	TArray<FOutputDevice*> LocalUnbufferedDevices;
 
 	{
 		// We need to lock the mutex here so that it gets unlocked after we empty the devices arrays
-		FScopeLock GlobalOutputDevicesLock(&OutputDevicesMutex);
-		LockOutputDevices(LocalBufferedDevices, LocalUnbufferedDevices);
-		BufferedOutputDevices.Empty();
-		UnbufferedOutputDevices.Empty();	
+		UE::Private::FOutputDevicesWriteScopeLock Lock(*State);
+		LocalBufferedDevices = MoveTemp(State->BufferedOutputDevices);
+		LocalUnbufferedDevices = MoveTemp(State->UnbufferedOutputDevices);
+		State->BufferedOutputDevices.Empty();
+		State->UnbufferedOutputDevices.Empty();
 	}
 
-	// Flush previously buffered lines from secondary threads.
-	InternalFlushThreadedLogs(LocalBufferedDevices, LocalUnbufferedDevices, false);
+	FlushBufferedLines(LocalBufferedDevices, /*bUseAllDevices*/ true);
 
 	for (FOutputDevice* OutputDevice : LocalBufferedDevices)
 	{
-		if (OutputDevice->CanBeUsedOnAnyThread())
-		{
-			OutputDevice->Flush();
-		}
+		OutputDevice->Flush();
 		OutputDevice->TearDown();
 	}
 
@@ -502,8 +436,11 @@ void FOutputDeviceRedirector::TearDown()
 		OutputDevice->Flush();
 		OutputDevice->TearDown();
 	}
+}
 
-	UnlockOutputDevices();
+bool FOutputDeviceRedirector::IsBacklogEnabled() const
+{
+	return State->bEnableBacklog;
 }
 
 CORE_API FOutputDeviceRedirector* GetGlobalLogSingleton()
diff --git a/Engine/Source/Runtime/Core/Public/Misc/OutputDeviceRedirector.h b/Engine/Source/Runtime/Core/Public/Misc/OutputDeviceRedirector.h
--- a/Engine/Source/Runtime/Core/Public/Misc/OutputDeviceRedirector.h
+++ b/Engine/Source/Runtime/Core/Public/Misc/OutputDeviceRedirector.h
@@ -3,16 +3,16 @@
 #pragma once
 
 #include "CoreTypes.h"
-#include "Misc/OutputDevice.h"
 #include "Containers/Array.h"
-#include "Containers/UnrealString.h"
+#include "Misc/OutputDevice.h"
+#include "Templates/PimplPtr.h"
 #include "UObject/NameTypes.h"
 
 /*-----------------------------------------------------------------------------
 FOutputDeviceRedirector.
 -----------------------------------------------------------------------------*/
 
-class FLogAllocator;
+namespace UE::Private { struct FOutputDeviceRedirectorState; }
 
 /** The type of lines buffered by secondary threads. */
 struct CORE_API FBufferedLine
@@ -28,8 +28,7 @@ struct CORE_API FBufferedLine
 	const ELogVerbosity::Type Verbosity;
 	bool bExternalAllocation;
 
-	FBufferedLine(const TCHAR* InData, const FName& InCategory, ELogVerbosity::Type InVerbosity, const double InTime = -1, FLogAllocator* ExternalAllocator = nullptr);
-	FBufferedLine(const TCHAR* InData, const FLazyName& InCategory, ELogVerbosity::Type InVerbosity, const double InTime = -1, FLogAllocator* ExternalAllocator = nullptr);
+	FBufferedLine(const TCHAR* InData, const FName& InCategory, ELogVerbosity::Type InVerbosity, const double InTime = -1);
 
 	FBufferedLine(FBufferedLine& InBufferedLine, EBufferedLineInit Unused)
 		: Data(InBufferedLine.Data)
@@ -49,178 +48,110 @@ struct CORE_API FBufferedLine
 };
 
 /**
-* Class used for output redirection to allow logs to show
-*/
-class CORE_API FOutputDeviceRedirector : public FOutputDevice
+ * Class used for output redirection to allow logs to show in multiple output devices.
+ */
+class CORE_API FOutputDeviceRedirector final : public FOutputDevice
 {
 public:
-
-	typedef TArray<FOutputDevice*, TInlineAllocator<16> > TLocalOutputDevicesArray;
-
-private:
-	enum { InlineLogEntries = 16 };
-
-	/** A FIFO of lines logged by non-master threads. */
-	TArray<FBufferedLine, TInlineAllocator<InlineLogEntries>> BufferedLines;
-
-	/** A FIFO backlog of messages logged before the editor had a chance to intercept them. */
-	TArray<FBufferedLine> BacklogLines;
-
-	/** Array of output devices to redirect to using buffering mechanism */
-	TArray<FOutputDevice*> BufferedOutputDevices;
-
-	/** Array of output devices that can redirected to without bufffering */
-	TArray<FOutputDevice*> UnbufferedOutputDevices;
-
-	/** The master thread ID.  Logging from other threads will be buffered for processing by the master thread. */
-	uint32 MasterThreadID;
-
-	/** Whether backlogging is enabled. */
-	bool bEnableBacklog;
-
-	FLogAllocator* BufferedLinesAllocator;
-
-	/** Objects used for synchronization via a scoped lock */
-	FCriticalSection	SynchronizationObject;
-	FCriticalSection	BufferSynchronizationObject;
-	FCriticalSection	OutputDevicesMutex;
-	FThreadSafeCounter	OutputDevicesLockCounter;
-
-	/**
-	* The unsynchronized version of FlushThreadedLogs.
-	* Assumes that the caller holds a lock on SynchronizationObject.
-	* @param bUseAllDevices - if true this method will use all output devices
-	*/
-	void InternalFlushThreadedLogs(TLocalOutputDevicesArray& InBufferedDevices, TLocalOutputDevicesArray& InUnbufferedDevices, bool bUseAllDevices);
-	void InternalFlushThreadedLogs(bool bUseAllDevices);
-
-	/** Locks OutputDevices arrays so that nothing can be added or removed from them */
-	void LockOutputDevices(TLocalOutputDevicesArray& OutBufferedDevices, TLocalOutputDevicesArray& OutUnbufferedDevices);
-
-	/** Unlocks OutputDevices arrays */
-	void UnlockOutputDevices();
-
-	friend struct FOutputDevicesLock;
-
-	/** Helper struct for scope locking OutputDevices arrays */
-	struct FOutputDevicesLock
-	{
-		FOutputDeviceRedirector* Redirector;
-		FOutputDevicesLock(FOutputDeviceRedirector* InRedirector, TLocalOutputDevicesArray& OutBufferedDevices, TLocalOutputDevicesArray& OutUnbufferedDevices)
-			: Redirector(InRedirector)
-		{
-			Redirector->LockOutputDevices(OutBufferedDevices, OutUnbufferedDevices);
-		}
-		~FOutputDevicesLock()
-		{
-			Redirector->UnlockOutputDevices();
-		}
-	};
-
-	void EmptyBufferedLines();
-
-	template<class T>
-	void SerializeImpl(const TCHAR* Data, ELogVerbosity::Type Verbosity, T& CategoryName, double Time);
-
-public:
+	UE_DEPRECATED(5.1, "TLocalOutputDevicesArray is being removed. Use TArray<FOutputDevice*, TInlineAllocator<16>>.")
+	typedef TArray<FOutputDevice*, TInlineAllocator<16>> TLocalOutputDevicesArray;
 
 	/** Initialization constructor. */
-	explicit FOutputDeviceRedirector(FLogAllocator* Allocator = nullptr);
+	FOutputDeviceRedirector();
 
-	/**
-	* Get the GLog singleton
-	*/
+	/** Get the GLog singleton. */
 	static FOutputDeviceRedirector* Get();
 
 	/**
-	* Adds an output device to the chain of redirections.
-	*
-	* @param OutputDevice	output device to add
-	*/
+	 * Adds an output device to the chain of redirections.
+	 *
+	 * @param OutputDevice   Output device to add.
+	 */
 	void AddOutputDevice(FOutputDevice* OutputDevice);
 
 	/**
-	* Removes an output device from the chain of redirections.
-	*
-	* @param OutputDevice	output device to remove
-	*/
+	 * Removes an output device from the chain of redirections.
+	 *
+	 * @param OutputDevice   Output device to remove.
+	 */
 	void RemoveOutputDevice(FOutputDevice* OutputDevice);
 
 	/**
-	* Returns whether an output device is currently in the list of redirectors.
-	*
-	* @param	OutputDevice	output device to check the list against
-	* @return	true if messages are currently redirected to the the passed in output device, false otherwise
-	*/
+	 * Returns whether an output device is in the list of redirections.
+	 *
+	 * @param OutputDevice   Output device to check the list against.
+	 * @return true if messages are currently redirected to the the passed in output device, false otherwise.
+	 */
 	bool IsRedirectingTo(FOutputDevice* OutputDevice);
 
 	/** Flushes lines buffered by secondary threads. */
-	virtual void FlushThreadedLogs();
+	void FlushThreadedLogs();
 
 	/**
-	*	Flushes lines buffered by secondary threads.
-	*	Only used if a background thread crashed and we needed to push the callstack into the log.
-	*/
-	virtual void PanicFlushThreadedLogs();
+	 * Flushes lines buffered by secondary threads.
+	 *
+	 * Only used if a background thread crashed and we needed to push the callstack into the log.
+	 */
+	void PanicFlushThreadedLogs();
 
 	/**
-	* Serializes the current backlog to the specified output device.
-	* @param OutputDevice	- Output device that will receive the current backlog
-	*/
-	virtual void SerializeBacklog(FOutputDevice* OutputDevice);
+	 * Serializes the current backlog to the specified output device.
+	 * @param OutputDevice   Output device that will receive the current backlog.
+	 */
+	void SerializeBacklog(FOutputDevice* OutputDevice);
 
 	/**
-	* Enables or disables the backlog.
-	* @param bEnable	- Starts saving a backlog if true, disables and discards any backlog if false
-	*/
-	virtual void EnableBacklog(bool bEnable);
+	 * Enables or disables the backlog.
+	 *
+	 * @param bEnable   Starts saving a backlog if true, disables and discards any backlog if false.
+	 */
+	void EnableBacklog(bool bEnable);
 
-	/**
-	* Sets the current thread to be the master thread that prints directly
-	* (isn't queued up)
-	*/
-	virtual void SetCurrentThreadAsMasterThread();
+	/** Sets the current thread to be the master thread redirects logs without buffering. */
+	void SetCurrentThreadAsMasterThread();
 
 	/**
-	* Serializes the passed in data via all current output devices.
-	*
-	* @param	Data	Text to log
-	* @param	Event	Event name used for suppression purposes
-	*/
-	virtual void Serialize(const TCHAR* Data, ELogVerbosity::Type Verbosity, const class FName& Category, const double Time) override;
+	 * Serializes the passed in data via all current output devices.
+	 *
+	 * @param Data   Text to log.
+	 * @param Event  Event name used for suppression purposes.
+	 */
+	void Serialize(const TCHAR* Data, ELogVerbosity::Type Verbosity, const FName& Category, const double Time) final;
 
 	/**
-	* Serializes the passed in data via all current output devices.
-	*
-	* @param	Data	Text to log
-	* @param	Event	Event name used for suppression purposes
-	*/
-	virtual void Serialize(const TCHAR* Data, ELogVerbosity::Type Verbosity, const class FName& Category) override;
+	 * Serializes the passed in data via all current output devices.
+	 *
+	 * @param Data   Text to log.
+	 * @param Event  Event name used for suppression purposes.
+	 */
+	void Serialize(const TCHAR* Data, ELogVerbosity::Type Verbosity, const FName& Category) final;
 	
-	/** Same as Serialize() but FName creation. Only needed to support 
-	*
-	*/
+	/** Same as Serialize(). */
+	void RedirectLog(const FName& Category, ELogVerbosity::Type Verbosity, const TCHAR* Data);
 	void RedirectLog(const FLazyName& Category, ELogVerbosity::Type Verbosity, const TCHAR* Data);
 
-	void RedirectLog(const FName& Category, ELogVerbosity::Type Verbosity, const TCHAR* Data);
+	/** Passes on the flush request to all current output devices. */
+	void Flush() final;
 
 	/**
-	* Passes on the flush request to all current output devices.
-	*/
-	void Flush() override;
+	 * Closes output device and cleans up.
+	 *
+	 * This can't happen in the destructor as we might have to call "delete" which cannot be done for static/global objects.
+	 */
+	void TearDown() final;
 
 	/**
-	* Closes output device and cleans up. This can't happen in the destructor
-	* as we might have to call "delete" which cannot be done for static/ global
-	* objects.
-	*/
-	void TearDown() override;
+	 * Determine if the backlog is enabled.
+	 */
+	bool IsBacklogEnabled() const;
 
+private:
 	/**
-	* Determine if backlog is enabled
-	*/
-	bool IsBacklogEnabled() const
-	{
-		return bEnableBacklog;
-	}
+	 * The unsynchronized version of FlushThreadedLogs.
+	 * Assumes that the caller holds a lock on the output devices.
+	 * @param bUseAllDevices   If true, this method will use all buffered output devices.
+	 */
+	void FlushBufferedLines(TConstArrayView<FOutputDevice*> BufferedDevices, bool bUseAllDevices);
+
+	TPimplPtr<UE::Private::FOutputDeviceRedirectorState> State;
 };
