Adding support for querying the iostore TOC off of a writer instead of just a reader. This is to allow the staging process to find package information and write it back to an asset registry.

diff --git a/Engine/Source/Runtime/Core/Private/IO/IoStore.cpp b/Engine/Source/Runtime/Core/Private/IO/IoStore.cpp
--- a/Engine/Source/Runtime/Core/Private/IO/IoStore.cpp
+++ b/Engine/Source/Runtime/Core/Private/IO/IoStore.cpp
@@ -609,6 +609,55 @@
 		return ChunkIdToFileName.Find(ChunkId);
 	}
 
+	FIoStoreTocChunkInfo GetTocChunkInfo(const int32 TocEntryIndex, const TMap<int32, FString>* InChunkFileNamesMap) const
+	{
+		const FIoStoreTocResource& TocResource = GetTocResource();
+		const FIoStoreTocEntryMeta& Meta = TocResource.ChunkMetas[TocEntryIndex];
+		const FIoOffsetAndLength& OffsetLength = TocResource.ChunkOffsetLengths[TocEntryIndex];
+
+		const bool bIsContainerCompressed = EnumHasAnyFlags(TocResource.Header.ContainerFlags, EIoContainerFlags::Compressed);
+
+		FIoStoreTocChunkInfo ChunkInfo;
+		ChunkInfo.Id = TocResource.ChunkIds[TocEntryIndex];
+		ChunkInfo.ChunkType = ChunkInfo.Id.GetChunkType();
+		ChunkInfo.Hash = Meta.ChunkHash;
+		ChunkInfo.bIsCompressed = EnumHasAnyFlags(Meta.Flags, FIoStoreTocEntryMetaFlags::Compressed);
+		ChunkInfo.bIsMemoryMapped = EnumHasAnyFlags(Meta.Flags, FIoStoreTocEntryMetaFlags::MemoryMapped);
+		ChunkInfo.bForceUncompressed = bIsContainerCompressed && !EnumHasAnyFlags(Meta.Flags, FIoStoreTocEntryMetaFlags::Compressed);
+		ChunkInfo.Offset = OffsetLength.GetOffset();
+		ChunkInfo.Size = OffsetLength.GetLength();
+
+		const FString* FindFileName = 0;
+		if (InChunkFileNamesMap &&
+			(FindFileName = InChunkFileNamesMap->Find(TocEntryIndex)) != nullptr)
+		{
+			ChunkInfo.FileName = *FindFileName;
+			ChunkInfo.bHasValidFileName = true;
+		}
+		else
+		{
+			ChunkInfo.FileName = FString::Printf(TEXT("<%s>"), *LexToString(ChunkInfo.ChunkType));
+			ChunkInfo.bHasValidFileName = false;
+		}
+
+		const uint64 CompressionBlockSize = TocResource.Header.CompressionBlockSize;
+		int32 FirstBlockIndex = int32(OffsetLength.GetOffset() / CompressionBlockSize);
+		int32 LastBlockIndex = int32((Align(OffsetLength.GetOffset() + OffsetLength.GetLength(), CompressionBlockSize) - 1) / CompressionBlockSize);
+
+		ChunkInfo.CompressedSize = 0;
+		ChunkInfo.PartitionIndex = -1;
+		for (int32 BlockIndex = FirstBlockIndex; BlockIndex <= LastBlockIndex; ++BlockIndex)
+		{
+			const FIoStoreTocCompressedBlockEntry& CompressionBlock = TocResource.CompressionBlocks[BlockIndex];
+			ChunkInfo.CompressedSize += CompressionBlock.GetCompressedSize();
+			if (ChunkInfo.PartitionIndex < 0)
+			{
+				ChunkInfo.PartitionIndex = int32(CompressionBlock.GetOffset() / TocResource.Header.PartitionSize);
+			}
+		}
+		return ChunkInfo;
+	}
+
 private:
 	TMap<FIoChunkId, int32> ChunkIdToIndex;
 	FIoStoreTocResource Toc;
@@ -624,6 +673,20 @@
 	{
 	}
 
+	void EnumerateChunks(TFunction<bool(const FIoStoreTocChunkInfo&)>&& Callback) const
+	{
+		const FIoStoreTocResource& TocResource = Toc.GetTocResource();
+
+		for (int32 ChunkIndex = 0; ChunkIndex < TocResource.ChunkIds.Num(); ++ChunkIndex)
+		{
+			FIoStoreTocChunkInfo ChunkInfo = Toc.GetTocChunkInfo(ChunkIndex, 0);
+			if (!Callback(ChunkInfo))
+			{
+				break;
+			}
+		}
+	}
+
 	UE_NODISCARD FIoStatus Initialize(FIoStoreWriterContextImpl& InContext, const FIoContainerSettings& InContainerSettings)
 	{
 		WriterContext = &InContext;
@@ -1922,7 +1985,7 @@
 
 		for (int32 ChunkIndex = 0; ChunkIndex < TocResource.ChunkIds.Num(); ++ChunkIndex)
 		{
-			FIoStoreTocChunkInfo ChunkInfo = GetTocChunkInfo(ChunkIndex);
+			FIoStoreTocChunkInfo ChunkInfo = Toc.GetTocChunkInfo(ChunkIndex, &ChunkFileNamesMap);
 			if (!Callback(ChunkInfo))
 			{
 				break;
@@ -1935,7 +1998,7 @@
 		const int32* TocEntryIndex = Toc.GetTocEntryIndex(ChunkId);
 		if (TocEntryIndex)
 		{
-			return GetTocChunkInfo(*TocEntryIndex);
+			return Toc.GetTocChunkInfo(*TocEntryIndex, &ChunkFileNamesMap);
 		}
 		else
 		{
@@ -1949,7 +2012,7 @@
 
 		if (TocEntryIndex < uint32(TocResource.ChunkIds.Num()))
 		{
-			return GetTocChunkInfo(TocEntryIndex);
+			return Toc.GetTocChunkInfo(TocEntryIndex, &ChunkFileNamesMap);
 		}
 		else
 		{
@@ -2323,52 +2386,7 @@
 	}
 
 private:
-	FIoStoreTocChunkInfo GetTocChunkInfo(const int32 TocEntryIndex) const
-	{
-		const FIoStoreTocResource& TocResource = Toc.GetTocResource();
-		const FIoStoreTocEntryMeta& Meta = TocResource.ChunkMetas[TocEntryIndex];
-		const FIoOffsetAndLength& OffsetLength = TocResource.ChunkOffsetLengths[TocEntryIndex];
 
-		const bool bIsContainerCompressed = EnumHasAnyFlags(TocResource.Header.ContainerFlags, EIoContainerFlags::Compressed);
 
-		FIoStoreTocChunkInfo ChunkInfo;
-		ChunkInfo.Id = TocResource.ChunkIds[TocEntryIndex];
-		ChunkInfo.ChunkType = static_cast<EIoChunkType>(ChunkInfo.Id.Id[11]);
-		ChunkInfo.Hash = Meta.ChunkHash;
-		ChunkInfo.bIsCompressed = EnumHasAnyFlags(Meta.Flags, FIoStoreTocEntryMetaFlags::Compressed);
-		ChunkInfo.bIsMemoryMapped = EnumHasAnyFlags(Meta.Flags, FIoStoreTocEntryMetaFlags::MemoryMapped);
-		ChunkInfo.bForceUncompressed = bIsContainerCompressed && !EnumHasAnyFlags(Meta.Flags, FIoStoreTocEntryMetaFlags::Compressed);
-		ChunkInfo.Offset = OffsetLength.GetOffset();
-		ChunkInfo.Size = OffsetLength.GetLength();
-		const FString* FindFileName = ChunkFileNamesMap.Find(TocEntryIndex);
-		if (FindFileName)
-		{
-			ChunkInfo.FileName = *FindFileName;
-			ChunkInfo.bHasValidFileName = true;
-		}
-		else
-		{
-			ChunkInfo.FileName = FString::Printf(TEXT("<%s>"), *LexToString(ChunkInfo.ChunkType));
-			ChunkInfo.bHasValidFileName = false;
-		}
-
-		const uint64 CompressionBlockSize = TocResource.Header.CompressionBlockSize;
-		int32 FirstBlockIndex = int32(OffsetLength.GetOffset() / CompressionBlockSize);
-		int32 LastBlockIndex = int32((Align(OffsetLength.GetOffset() + OffsetLength.GetLength(), CompressionBlockSize) - 1) / CompressionBlockSize);
-
-		ChunkInfo.CompressedSize = 0;
-		ChunkInfo.PartitionIndex = -1;
-		for (int32 BlockIndex = FirstBlockIndex; BlockIndex <= LastBlockIndex; ++BlockIndex)
-		{
-			const FIoStoreTocCompressedBlockEntry& CompressionBlock = TocResource.CompressionBlocks[BlockIndex];
-			ChunkInfo.CompressedSize += CompressionBlock.GetCompressedSize();
-			if (ChunkInfo.PartitionIndex < 0)
-			{
-				ChunkInfo.PartitionIndex = int32(CompressionBlock.GetOffset() / TocResource.Header.PartitionSize);
-			}
-		}
-		return ChunkInfo;
-	}
-
 	FIoStoreToc Toc;
 	FAES::FAESKey DecryptionKey;
diff --git a/Engine/Source/Runtime/Core/Public/IO/IoDispatcher.h b/Engine/Source/Runtime/Core/Public/IO/IoDispatcher.h
--- a/Engine/Source/Runtime/Core/Public/IO/IoDispatcher.h
+++ b/Engine/Source/Runtime/Core/Public/IO/IoDispatcher.h
@@ -1192,16 +1192,6 @@ public:
 	virtual void FreeSourceBuffer() = 0;
 };
 
-class IIoStoreWriter
-{
-public:
-	CORE_API virtual ~IIoStoreWriter() = default;
-
-	CORE_API virtual void EnableDiskLayoutOrdering(const TArray<TUniquePtr<FIoStoreReader>>& PatchSourceReaders = TArray<TUniquePtr<FIoStoreReader>>()) = 0;
-	CORE_API virtual void Append(const FIoChunkId& ChunkId, FIoBuffer Chunk, const FIoWriteOptions& WriteOptions, uint64 OrderHint = MAX_uint64) = 0;
-	CORE_API virtual void Append(const FIoChunkId& ChunkId, IIoStoreWriteRequest* Request, const FIoWriteOptions& WriteOptions) = 0;
-	CORE_API virtual TIoStatusOr<FIoStoreWriterResult> GetResult() = 0;
-};
 
 struct FIoStoreTocChunkInfo
 {
@@ -1242,6 +1232,19 @@ struct FIoStoreCompressedReadResult
 	uint64 UncompressedSize;
 };
 
+
+class IIoStoreWriter
+{
+public:
+	CORE_API virtual ~IIoStoreWriter() = default;
+
+	CORE_API virtual void EnableDiskLayoutOrdering(const TArray<TUniquePtr<FIoStoreReader>>& PatchSourceReaders = TArray<TUniquePtr<FIoStoreReader>>()) = 0;
+	CORE_API virtual void Append(const FIoChunkId& ChunkId, FIoBuffer Chunk, const FIoWriteOptions& WriteOptions, uint64 OrderHint = MAX_uint64) = 0;
+	CORE_API virtual void Append(const FIoChunkId& ChunkId, IIoStoreWriteRequest* Request, const FIoWriteOptions& WriteOptions) = 0;
+	CORE_API virtual TIoStatusOr<FIoStoreWriterResult> GetResult() = 0;
+	CORE_API virtual void EnumerateChunks(TFunction<bool(const FIoStoreTocChunkInfo&)>&& Callback) const = 0;
+};
+
 class FIoStoreReader
 {
 public:
