Logging: Fixed compile errors when FRWLock is the size of a cache line

#rnx

diff --git a/Engine/Source/Runtime/Core/Private/Misc/OutputDeviceRedirector.cpp b/Engine/Source/Runtime/Core/Private/Misc/OutputDeviceRedirector.cpp
--- a/Engine/Source/Runtime/Core/Private/Misc/OutputDeviceRedirector.cpp
+++ b/Engine/Source/Runtime/Core/Private/Misc/OutputDeviceRedirector.cpp
@@ -99,15 +99,20 @@ private:
 	}
 };
 
+static constexpr uint64 CalculateRedirectorCacheLinePadding(const uint64 Size)
+{
+	return PLATFORM_CACHE_LINE_SIZE * FMath::DivideAndRoundUp<uint64>(Size, PLATFORM_CACHE_LINE_SIZE) - Size;
+}
+
 struct FOutputDeviceRedirectorState
 {
 	FRWLock OutputDevicesLock;
 	std::atomic<uint32> OutputDevicesLockState;
-	uint8 OutputDevicesLockPadding[PLATFORM_CACHE_LINE_SIZE - sizeof(OutputDevicesLock) - sizeof(OutputDevicesLockState)]{};
+	uint8 OutputDevicesLockPadding[CalculateRedirectorCacheLinePadding(sizeof(OutputDevicesLock) + sizeof(OutputDevicesLockState))]{};
 
 	/** A FIFO of lines logged by non-master threads. */
 	TDepletableMpscQueue<FOutputDeviceLine, FOutputDeviceLinearAllocator> BufferedLines;
-	uint8 BufferedLinesPadding[PLATFORM_CACHE_LINE_SIZE - sizeof(BufferedLines)]{};
+	uint8 BufferedLinesPadding[CalculateRedirectorCacheLinePadding(sizeof(BufferedLines))]{};
 
 	/** Array of output devices to redirect to from the master thread. */
 	TArray<FOutputDevice*> BufferedOutputDevices;
