FBlockEncoder::Compress parallel enhanced
special case 1 block fast path
in multi-block case, run hash on thread too

diff --git a/Engine/Source/Runtime/Core/Private/Compression/CompressedBuffer.cpp b/Engine/Source/Runtime/Core/Private/Compression/CompressedBuffer.cpp
--- a/Engine/Source/Runtime/Core/Private/Compression/CompressedBuffer.cpp
+++ b/Engine/Source/Runtime/Core/Private/Compression/CompressedBuffer.cpp
@@ -267,121 +267,179 @@ FCompositeBuffer FBlockEncoder::Compress(const FCompositeBuffer& RawData, const
 	checkf(FMath::IsPowerOfTwo(BlockSize) && BlockSize <= MAX_uint32,
 		TEXT("BlockSize must be a 32-bit power of two but was %" UINT64_FMT "."), BlockSize);
 	const uint64 RawSize = RawData.GetSize();
+
+	if ( RawSize == 0 )
+	{
+		return FNoneEncoder().Compress(RawData, BlockSize);
+	}
+
 	FBlake3 RawHash;
 
-	const uint64 BlockCount = FMath::DivideAndRoundUp(RawSize, BlockSize);
-	checkf(BlockCount < MAX_uint32, TEXT("Raw data of size %" UINT64_FMT " with block size %" UINT64_FMT " requires ")
+	const int64 BlockCount = FMath::DivideAndRoundUp(RawSize, BlockSize);
+	check( BlockCount > 0 );
+	checkf(BlockCount < MAX_int32, TEXT("Raw data of size %" UINT64_FMT " with block size %" UINT64_FMT " requires ")
 		TEXT("%" UINT64_FMT " blocks, but the limit is %u."), RawSize, BlockSize, BlockCount, MAX_uint32);
+		
+	uint64 TotalCompressedSize = 0;
+	const uint64 MetaSize = sizeof(uint32) * BlockCount;
+	const uint64 CompressedDataSize = sizeof(FHeader) + MetaSize + GetCompressedBlocksBound(BlockCount, BlockSize, RawSize);
+	FUniqueBuffer CompressedData;
 
 	// Compress the raw data in blocks and store the raw data for incompressible blocks.
 
-	struct CompressedBlockInfo
-	{
-		// in:
-		uint64 RawBlockSize;
-		FUniqueBuffer RawBlockCopy;
-		FMemoryView RawBlock;
-
-		// out:
-		TArray64<uint8> CompressedData;
-	};
-	
-	TArray64<CompressedBlockInfo> CompressedBlockInfos;
-	CompressedBlockInfos.SetNum((int)BlockCount);
-	
-	// fill all CompressedBlockInfos :
-	for (int BlockIndex=0;BlockIndex<(int)BlockCount;BlockIndex++)
-	{
-		const uint64 RawOffset = BlockIndex * BlockSize;
-		const uint64 RawBlockSize = FMath::Min(RawSize - RawOffset, BlockSize);
-		CompressedBlockInfos[BlockIndex].RawBlockSize = RawBlockSize;
-		CompressedBlockInfos[BlockIndex].RawBlock = RawData.ViewOrCopyRange(RawOffset, RawBlockSize, CompressedBlockInfos[BlockIndex].RawBlockCopy);
-
-		// sequential hash update :
-		RawHash.Update(CompressedBlockInfos[BlockIndex].RawBlock); // @todo Oodle: replace with parallel hash
-	}
-
-	uint64 CompressBlockAllocSize = CompressBlockBound(BlockSize);
-	
-	// encode blocks independently, can be parallel :
-	//for (int BlockIndex=0;BlockIndex<(int)BlockCount;BlockIndex++)
-	
-	ParallelFor(TEXT("BlockEncoder.Compress.PF"), (int32)BlockCount,1, [&CompressedBlockInfos,CompressBlockAllocSize,this](int64 BlockIndex)
+	if ( BlockCount == 1 )
 	{
-		CompressedBlockInfos[BlockIndex].CompressedData.SetNum(CompressBlockAllocSize);
+		FUniqueBuffer RawBlockCopy;	
+		FMemoryView RawBlock = RawData.ViewOrCopyRange(0, RawSize, RawBlockCopy);
 		
-		FMutableMemoryView CompressedBlockView(CompressedBlockInfos[BlockIndex].CompressedData.GetData(),CompressedBlockInfos[BlockIndex].CompressedData.Num());
+		RawHash.Update(RawBlock);
 		
-		if (CompressBlock(CompressedBlockView, CompressedBlockInfos[BlockIndex].RawBlock))
-		{
-			check( CompressedBlockView.GetSize() > 0 );
-			CompressedBlockInfos[BlockIndex].CompressedData.SetNum(CompressedBlockView.GetSize());
-		}
-		else
+		// Allocate the buffer for the header, metadata, and compressed blocks.
+		CompressedData = FUniqueBuffer::Alloc(CompressedDataSize);
+		
+		FMutableMemoryView CompressedBlocksView = CompressedData.GetView() + sizeof(FHeader) + MetaSize;
+		
+		FMutableMemoryView CompressedBlock = CompressedBlocksView;
+		if (!CompressBlock(CompressedBlock, RawBlock))
 		{
-			// size 0 indicates error
-			CompressedBlockInfos[BlockIndex].CompressedData.SetNum(0);
+			// compressor error !?
+			return FCompositeBuffer();
 		}
-	});
-	
-	// Allocate the buffer for the header, metadata, and compressed blocks.
-	const uint64 MetaSize = sizeof(uint32) * BlockCount;
-	const uint64 CompressedDataSize = sizeof(FHeader) + MetaSize + GetCompressedBlocksBound(BlockCount, BlockSize, RawSize);
-	FUniqueBuffer CompressedData = FUniqueBuffer::Alloc(CompressedDataSize);
-
-	// gather results sequentially :
-	TArray64<uint32> CompressedBlockSizes;
-	CompressedBlockSizes.SetNum((int)BlockCount);
-	uint64 TotalCompressedSize = 0;
-	FMutableMemoryView CompressedBlocksView = CompressedData.GetView() + sizeof(FHeader) + MetaSize;
-	for (int BlockIndex=0;BlockIndex<(int)BlockCount;BlockIndex++)
-	{
-		FMemoryView CompressedBlock(CompressedBlockInfos[BlockIndex].CompressedData.GetData(),CompressedBlockInfos[BlockIndex].CompressedData.Num());
 
 		uint64 CompressedBlockSize = CompressedBlock.GetSize();
-		if ( CompressedBlockSize == 0 )
+		CompressedBlocksView += CompressedBlockSize;
+		TotalCompressedSize += CompressedBlockSize;
+
+		// Return an uncompressed buffer if the compressed data is larger than the raw data.
+		if (RawSize <= MetaSize + TotalCompressedSize)
 		{
-			// Compressor failed 
 			CompressedData.Reset();
-			return FNoneEncoder().Compress(RawData, BlockSize);			
+			return FNoneEncoder().Compress(RawData, BlockSize);
 		}
-		else if (CompressedBlockInfos[BlockIndex].RawBlockSize <= CompressedBlockSize)
+			
+		// put CompressedBlockSizes after FHeader :
+		uint32 CompressedBlockSizeForHeader = IntCastChecked<uint32>(CompressedBlockSize);
+		CompressedBlockSizeForHeader = NETWORK_ORDER32(CompressedBlockSizeForHeader);
+		CompressedData.GetView().Mid(sizeof(FHeader), MetaSize).CopyFrom(MakeMemoryView(&CompressedBlockSizeForHeader,sizeof(CompressedBlockSizeForHeader)));
+	}
+	else
+	{
+		// BlockCount >= 2
+
+		struct CompressedBlockInfo
 		{
-			CompressedBlockSize = CompressedBlockInfos[BlockIndex].RawBlockSize;
-			CompressedBlocksView = CompressedBlocksView.CopyFrom(CompressedBlockInfos[BlockIndex].RawBlock);
-		}
-		else
+			// in:
+			uint64 RawBlockSize;
+			FUniqueBuffer RawBlockCopy;
+			FMemoryView RawBlock;
+
+			// out:
+			TArray64<uint8> CompressedData;
+		};
+	
+		TArray64<CompressedBlockInfo> CompressedBlockInfos;
+		CompressedBlockInfos.SetNum(BlockCount);
+	
+		// fill all CompressedBlockInfos :
+		for (int64 BlockIndex=0;BlockIndex<BlockCount;BlockIndex++)
 		{
-			CompressedBlocksView = CompressedBlocksView.CopyFrom(CompressedBlock);
+			const uint64 RawOffset = BlockIndex * BlockSize;
+			const uint64 RawBlockSize = FMath::Min(RawSize - RawOffset, BlockSize);
+			CompressedBlockInfos[BlockIndex].RawBlockSize = RawBlockSize;
+			CompressedBlockInfos[BlockIndex].RawBlock = RawData.ViewOrCopyRange(RawOffset, RawBlockSize, CompressedBlockInfos[BlockIndex].RawBlockCopy);
 		}
+		
+		uint64 CompressBlockAllocSize = CompressBlockBound(BlockSize);
 
-		// Can be freed after copied out :
-		CompressedBlockInfos[BlockIndex].CompressedData.Empty();
+		// encode blocks independently, can be parallel :
+		ParallelFor(TEXT("BlockEncoder.Compress.PF"), IntCastChecked<int32>(BlockCount)+1,1, [&CompressedBlockInfos,CompressBlockAllocSize,BlockCount,&RawHash,this](int64 ThreadIndex)
+		{
+			int64 BlockIndex = ThreadIndex-1;
+			if ( BlockIndex < 0 )
+			{
+				// one thread does the sequential hash update for all raw blocks
+				for (int64 HashBlockIndex=0;HashBlockIndex<BlockCount;HashBlockIndex++)
+				{
+					RawHash.Update(CompressedBlockInfos[HashBlockIndex].RawBlock);
+				}
+			}
+			else
+			{
+				CompressedBlockInfos[BlockIndex].CompressedData.SetNum(CompressBlockAllocSize);
+		
+				FMutableMemoryView CompressedBlockView(CompressedBlockInfos[BlockIndex].CompressedData.GetData(),CompressedBlockInfos[BlockIndex].CompressedData.Num());
+		
+				if (CompressBlock(CompressedBlockView, CompressedBlockInfos[BlockIndex].RawBlock))
+				{
+					check( CompressedBlockView.GetSize() > 0 );
+					// false on SetNum to not realloc
+					CompressedBlockInfos[BlockIndex].CompressedData.SetNum(CompressedBlockView.GetSize(),false);
+				}
+				else
+				{
+					// size 0 indicates error
+					CompressedBlockInfos[BlockIndex].CompressedData.Empty();
+				}
+			}
+		});
+	
+		// Allocate the buffer for the header, metadata, and compressed blocks.
+		CompressedData = FUniqueBuffer::Alloc(CompressedDataSize);
+
+		// gather results sequentially :
+		TArray64<uint32> CompressedBlockSizes;
+		CompressedBlockSizes.SetNum(BlockCount);
+		FMutableMemoryView CompressedBlocksView = CompressedData.GetView() + sizeof(FHeader) + MetaSize;
+		for (int64 BlockIndex=0;BlockIndex<BlockCount;BlockIndex++)
+		{
+			FMemoryView CompressedBlock(CompressedBlockInfos[BlockIndex].CompressedData.GetData(),CompressedBlockInfos[BlockIndex].CompressedData.Num());
 
-		CompressedBlockSizes[BlockIndex] = (uint32(CompressedBlockSize));
-		TotalCompressedSize += CompressedBlockSize;
-	}
+			uint64 CompressedBlockSize = CompressedBlock.GetSize();
+			if ( CompressedBlockSize == 0 )
+			{
+				// Compressor failed !?
+				CompressedData.Reset();
+				return FNoneEncoder().Compress(RawData, BlockSize);			
+			}
+			else if (CompressedBlockInfos[BlockIndex].RawBlockSize <= CompressedBlockSize)
+			{
+				CompressedBlockSize = CompressedBlockInfos[BlockIndex].RawBlockSize;
+				CompressedBlocksView = CompressedBlocksView.CopyFrom(CompressedBlockInfos[BlockIndex].RawBlock);
+			}
+			else
+			{
+				CompressedBlocksView = CompressedBlocksView.CopyFrom(CompressedBlock);
+			}
+
+			// Can be freed after copied out :
+			CompressedBlockInfos[BlockIndex].CompressedData.Empty();
+
+			CompressedBlockSizes[BlockIndex] = IntCastChecked<uint32>(CompressedBlockSize);
+			TotalCompressedSize += CompressedBlockSize;
+		}
 	
-	// CompressedBlockInfos not needed anymore :
-	CompressedBlockInfos.Empty();
+		// CompressedBlockInfos not needed anymore :
+		CompressedBlockInfos.Empty();
 
-	// Return an uncompressed buffer if the compressed data is larger than the raw data.
-	if (RawSize <= MetaSize + TotalCompressedSize)
-	{
-		CompressedData.Reset();
-		return FNoneEncoder().Compress(RawData, BlockSize);
+		// Return an uncompressed buffer if the compressed data is larger than the raw data.
+		if (RawSize <= MetaSize + TotalCompressedSize)
+		{
+			CompressedData.Reset();
+			return FNoneEncoder().Compress(RawData, BlockSize);
+		}
+
+		// put CompressedBlockSizes after FHeader :
+		Algo::ForEach(CompressedBlockSizes, [](uint32& Size) { Size = NETWORK_ORDER32(Size); });
+		CompressedData.GetView().Mid(sizeof(FHeader), MetaSize).CopyFrom(MakeMemoryView(CompressedBlockSizes));
 	}
 
 	// Write the header and calculate the CRC-32.
-	Algo::ForEach(CompressedBlockSizes, [](uint32& Size) { Size = NETWORK_ORDER32(Size); });
-	CompressedData.GetView().Mid(sizeof(FHeader), MetaSize).CopyFrom(MakeMemoryView(CompressedBlockSizes));
-
 	FHeader Header;
 	Header.Method = GetMethod();
 	Header.Compressor = GetCompressor();
 	Header.CompressionLevel = GetCompressionLevel();
 	Header.BlockSizeExponent = uint8(FMath::FloorLog2_64(BlockSize));
-	Header.BlockCount = uint32(BlockCount);
+	Header.BlockCount = IntCastChecked<uint32>(BlockCount);
 	Header.TotalRawSize = RawSize;
 	Header.TotalCompressedSize = sizeof(FHeader) + MetaSize + TotalCompressedSize;
 	Header.RawHash = RawHash.Finalize();
