Fix for compile warning.

#rnx

diff --git a/Engine/Source/Runtime/Core/Private/HAL/LowLevelMemTracker.cpp b/Engine/Source/Runtime/Core/Private/HAL/LowLevelMemTracker.cpp
--- a/Engine/Source/Runtime/Core/Private/HAL/LowLevelMemTracker.cpp
+++ b/Engine/Source/Runtime/Core/Private/HAL/LowLevelMemTracker.cpp
@@ -83,7 +83,6 @@ FAutoConsoleCommand DumpLLM(
 		bool bCSV = false;
 		for (const FString& Arg : Args)
 		{
-			//if (Arg.Equals(TEXT("CSV"), ESearchCase::IgnoreCase))
 			if (FParse::Param(*Arg, TEXT("CSV")))
 			{
 				bCSV = true;
@@ -111,8 +110,8 @@ FAutoConsoleCommand DumpLLM(
 					{
 						Ar.Logf(TEXT(",%s,%.3f,%.3f,%s,%s"),
 							*FLowLevelMemTracker::Get().GetTagUniqueName(TagData).ToString(),
-							float(FLowLevelMemTracker::Get().GetTagAmountForTracker(TrackerType, TagData, false)) * InvToMb,
-							float(FLowLevelMemTracker::Get().GetTagAmountForTracker(TrackerType, TagData, true)) * InvToMb,
+							static_cast<float>(FLowLevelMemTracker::Get().GetTagAmountForTracker(TrackerType, TagData, false)) * InvToMb,
+							static_cast<float>(FLowLevelMemTracker::Get().GetTagAmountForTracker(TrackerType, TagData, true)) * InvToMb,
 							TrackerName,
 							*FLowLevelMemTracker::Get().GetTagDisplayPathName(TagData));
 					}
@@ -120,8 +119,8 @@ FAutoConsoleCommand DumpLLM(
 					{
 						Ar.Logf(TEXT("%40s %12.3f %12.3f %8s %s"),
 							*FLowLevelMemTracker::Get().GetTagUniqueName(TagData).ToString(),
-							float(FLowLevelMemTracker::Get().GetTagAmountForTracker(TrackerType, TagData, false)) * InvToMb,
-							float(FLowLevelMemTracker::Get().GetTagAmountForTracker(TrackerType, TagData, true)) * InvToMb,
+							static_cast<float>(FLowLevelMemTracker::Get().GetTagAmountForTracker(TrackerType, TagData, false)) * InvToMb,
+							static_cast<float>(FLowLevelMemTracker::Get().GetTagAmountForTracker(TrackerType, TagData, true)) * InvToMb,
 							TrackerName,
 							*FLowLevelMemTracker::Get().GetTagDisplayPathName(TagData));
 					}
