Add FBinaryHeap move constructor and delete copy constructor & assignment operator to prevent use-after free (e.g., due to failed return value optimization).

- Fixes crash in debug builds in the Nanite builder.

diff --git a/Engine/Source/Runtime/Core/Public/Containers/BinaryHeap.h b/Engine/Source/Runtime/Core/Public/Containers/BinaryHeap.h
--- a/Engine/Source/Runtime/Core/Public/Containers/BinaryHeap.h
+++ b/Engine/Source/Runtime/Core/Public/Containers/BinaryHeap.h
@@ -22,6 +22,10 @@ public:
 				FBinaryHeap( uint32 InHeapSize, uint32 InIndexSize );
 				~FBinaryHeap();
 
+				FBinaryHeap(const FBinaryHeap&) = delete;
+				void operator =(const FBinaryHeap&) = delete;
+				FBinaryHeap(FBinaryHeap&& Other);
+
 	void		Clear();
 	void		Free();
 	void		Resize( uint32 NewHeapSize, uint32 NewIndexSize );
@@ -49,6 +53,11 @@ protected:
 	void		UpHeap( IndexType HeapIndex );
 	void		DownHeap( IndexType HeapIndex );
 
+	/**
+	 * Reset internal variables to a cleared state, does not free data.
+	 */
+	void		ResetInternal();
+
 	uint32		HeapNum;
 	uint32		HeapSize;
 	uint32		IndexSize;
@@ -61,13 +70,9 @@ protected:
 
 template< typename KeyType, typename IndexType >
 FORCEINLINE FBinaryHeap< KeyType, IndexType >::FBinaryHeap()
-	: HeapNum(0)
-	, HeapSize(0)
-	, IndexSize(0)
-	, Heap( nullptr )
-	, Keys( nullptr )
-	, HeapIndexes( nullptr )
-{}
+{
+	ResetInternal();
+}
 
 template< typename KeyType, typename IndexType >
 FORCEINLINE FBinaryHeap< KeyType, IndexType >::FBinaryHeap( uint32 InHeapSize, uint32 InIndexSize )
@@ -82,6 +87,20 @@ FORCEINLINE FBinaryHeap< KeyType, IndexType >::FBinaryHeap( uint32 InHeapSize, u
 	FMemory::Memset( HeapIndexes, 0xff, IndexSize * sizeof( IndexType ) );
 }
 
+template< typename KeyType, typename IndexType >
+FORCEINLINE FBinaryHeap< KeyType, IndexType >::FBinaryHeap(FBinaryHeap< KeyType, IndexType >&& Other)
+{
+	HeapNum = Other.HeapNum;
+	HeapSize = Other.HeapSize;
+	IndexSize = Other.IndexSize;
+
+	Heap = Other.Heap;
+	Keys = Other.Keys;
+	HeapIndexes = Other.HeapIndexes;
+
+	Other.ResetInternal();
+}
+
 template< typename KeyType, typename IndexType >
 FORCEINLINE FBinaryHeap< KeyType, IndexType >::~FBinaryHeap()
 {
@@ -98,17 +117,11 @@ FORCEINLINE void FBinaryHeap< KeyType, IndexType >::Clear()
 template< typename KeyType, typename IndexType >
 FORCEINLINE void FBinaryHeap< KeyType, IndexType >::Free()
 {
-	HeapNum = 0;
-	HeapSize = 0;
-	IndexSize = 0;
-	
 	delete[] Heap;
 	delete[] Keys;
 	delete[] HeapIndexes;
 
-	Heap = nullptr;
-	Keys = nullptr;
-	HeapIndexes = nullptr;
+	ResetInternal();
 }
 
 template< typename KeyType, typename IndexType >
@@ -374,3 +387,14 @@ void FBinaryHeap< KeyType, IndexType >::DownHeap( IndexType HeapIndex )
 		HeapIndexes[ Heap[i] ] = i;
 	}
 }
+
+template< typename KeyType, typename IndexType >
+FORCEINLINE void FBinaryHeap< KeyType, IndexType >::ResetInternal()
+{
+	HeapNum = 0;
+	HeapSize = 0;
+	IndexSize = 0;
+	Heap = nullptr;
+	Keys = nullptr;
+	HeapIndexes = nullptr;
+}
