Allow platforms to disable the use of FAlignedAllocator and thus ANSIMalloc.cpp.  Substitute FOsAllocator in this case.

#rnx

diff --git a/Engine/Source/Runtime/Core/Public/Experimental/ConcurrentLinearAllocator.h b/Engine/Source/Runtime/Core/Public/Experimental/ConcurrentLinearAllocator.h
--- a/Engine/Source/Runtime/Core/Public/Experimental/ConcurrentLinearAllocator.h
+++ b/Engine/Source/Runtime/Core/Public/Experimental/ConcurrentLinearAllocator.h
@@ -50,6 +50,7 @@ struct FOsAllocator
 	}
 };
 
+#if PLATFORM_USES_ANSI_MALLOC
 struct FAlignedAllocator
 {
 	static constexpr bool SupportsAlignment = true;
@@ -71,6 +72,9 @@ struct FAlignedAllocator
 		return AnsiFree(Pointer);
 	}
 };
+#else // #if PLATFORM_USES_ANSI_MALLOC
+using FAlignedAllocator = FOsAllocator;
+#endif // #else // #if PLATFORM_USES_ANSI_MALLOC
 
 template<uint32 BlockSize, typename Allocator = FOsAllocator>
 class TBlockAllocationCache
diff --git a/Engine/Source/Runtime/Core/Public/HAL/Platform.h b/Engine/Source/Runtime/Core/Public/HAL/Platform.h
--- a/Engine/Source/Runtime/Core/Public/HAL/Platform.h
+++ b/Engine/Source/Runtime/Core/Public/HAL/Platform.h
@@ -482,6 +482,10 @@
 	#define PLATFORM_USE_FULL_TASK_GRAPH						1
 #endif
 
+#ifndef PLATFORM_USES_ANSI_MALLOC
+	#define PLATFORM_USES_ANSI_MALLOC							1
+#endif
+
 #ifndef PLATFORM_USE_ANSI_POSIX_MALLOC
 	#define PLATFORM_USE_ANSI_POSIX_MALLOC						0
 #endif
