Creating RoundToNearestTiesToEven platform math function

#rnx
#p4v-preflight-copy 19180977
#skipci

diff --git a/Engine/Source/Runtime/Core/Public/GenericPlatform/GenericPlatformMath.h b/Engine/Source/Runtime/Core/Public/GenericPlatform/GenericPlatformMath.h
--- a/Engine/Source/Runtime/Core/Public/GenericPlatform/GenericPlatformMath.h
+++ b/Engine/Source/Runtime/Core/Public/GenericPlatform/GenericPlatformMath.h
@@ -16,6 +16,10 @@
 #include "Templates/UnrealTypeTraits.h"
 #include "Templates/ResolveTypeAmbiguity.h"
 
+#if PLATFORM_HAS_FENV_H 
+#include <fenv.h>
+#endif 
+
 /**
  * Generic implementation for most platforms
  */
@@ -397,6 +401,32 @@ struct FGenericPlatformMath
 
 	RESOLVE_FLOAT_AMBIGUITY(CeilToFloat);
 
+	/**
+	 * Converts a double to nearest int64 with ties rounding to nearest even
+	 * May incur a performance penalty. Asserts on platforms that do not support this mode.
+	 * @param F		Double precision floating point value to convert
+	 * @return		The 64-bit integer closest to 'F', with ties going to the nearest even number
+	 */
+	static int64 RoundToNearestTiesToEven(double F)
+	{
+#if PLATFORM_HAS_FENV_H == 0
+		ensureAlwaysMsgf(FLT_ROUNDS == 1, TEXT("Platform does not support FE_TONEAREST for double to int64."));
+		int64 result = llrint(F);
+#else
+		int PreviousRoundingMode = fegetround();
+		if (PreviousRoundingMode != FE_TONEAREST)
+		{
+			fesetround(FE_TONEAREST);
+		}
+		int64 result = llrint(F);
+		if (PreviousRoundingMode != FE_TONEAREST)
+		{
+			fesetround(PreviousRoundingMode);
+		}
+#endif
+		return result;
+	}
+
 	/**
 	* Returns signed fractional part of a float.
 	* @param Value	Floating point value to convert
diff --git a/Engine/Source/Runtime/Core/Public/HAL/Platform.h b/Engine/Source/Runtime/Core/Public/HAL/Platform.h
--- a/Engine/Source/Runtime/Core/Public/HAL/Platform.h
+++ b/Engine/Source/Runtime/Core/Public/HAL/Platform.h
@@ -646,6 +646,10 @@
 	#define PLATFORM_USES_UNFAIR_LOCKS 0
 #endif
 
+#ifndef PLATFORM_HAS_FENV_H
+	#define PLATFORM_HAS_FENV_H 1
+#endif
+
 // deprecated, do not use
 #define PLATFORM_HAS_THREADSAFE_RHIGetRenderQueryResult	#
 #define PLATFORM_SUPPORTS_RHI_THREAD #
