Fix for 'TEXT() TEXT()' usage in ConsoleManager.cpp, which doesn't compile in UTF-8 mode.

diff --git a/Engine/Source/Runtime/Core/Private/HAL/ConsoleManager.cpp b/Engine/Source/Runtime/Core/Private/HAL/ConsoleManager.cpp
--- a/Engine/Source/Runtime/Core/Private/HAL/ConsoleManager.cpp
+++ b/Engine/Source/Runtime/Core/Private/HAL/ConsoleManager.cpp
@@ -2326,8 +2326,8 @@ void CreateConsoleVariables()
 	// this registeres to a reference, so we cannot use TAutoConsoleVariable
 	IConsoleManager::Get().RegisterConsoleVariableRef(TEXT("r.DumpingMovie"),
 		GIsDumpingMovie,
-		TEXT("Allows to dump each rendered frame to disk (slow fps, names MovieFrame..).\n")
-		TEXT("<=0:off (default), <0:remains on, >0:remains on for n frames (n is the number specified)"),
+		TEXT("Allows to dump each rendered frame to disk (slow fps, names MovieFrame..).\n"
+			 "<=0:off (default), <0:remains on, >0:remains on for n frames (n is the number specified)"),
 		ECVF_Cheat);
 
 
@@ -2386,23 +2386,23 @@ void CreateConsoleVariables()
 static TAutoConsoleVariable<FString> CVarFreezeAtPosition(
 	TEXT("FreezeAtPosition"),
 	TEXT(""),	// default value is empty
-	TEXT("This console variable stores the position and rotation for the FreezeAt command which allows\n")
-	TEXT("to lock the camera in order to provide more deterministic render profiling.\n")
-	TEXT("The FreezeAtPosition can be set in the ConsoleVariables.ini (start the map with MAPNAME?bTourist=1).\n")
-	TEXT("Also see the FreezeAt command console command.\n")
-	TEXT("The number syntax if the same as the one used by the BugIt command:\n")
-	TEXT(" The first three values define the position, the next three define the rotation.\n")
-	TEXT("Example:\n")
-	TEXT(" FreezeAtPosition 2819.5520 416.2633 75.1500 65378 -25879 0"),
+	TEXT("This console variable stores the position and rotation for the FreezeAt command which allows\n"
+		 "to lock the camera in order to provide more deterministic render profiling.\n"
+		 "The FreezeAtPosition can be set in the ConsoleVariables.ini (start the map with MAPNAME?bTourist=1).\n"
+		 "Also see the FreezeAt command console command.\n"
+		 "The number syntax if the same as the one used by the BugIt command:\n"
+		 " The first three values define the position, the next three define the rotation.\n"
+		 "Example:\n"
+		 " FreezeAtPosition 2819.5520 416.2633 75.1500 65378 -25879 0"),
 	ECVF_Cheat);
 
 static TAutoConsoleVariable<int32> CVarLimitRenderingFeatures(
 	TEXT("r.LimitRenderingFeatures"),
 	0,
-	TEXT("Allows to quickly reduce render feature to increase render performance.\n")
-	TEXT("This is just a quick way to alter multiple show flags and console variables in the game\n")
-	TEXT("Disabled more feature the higher the number\n")
-	TEXT(" <=0:off, order is defined in code (can be documented here when we settled on an order)"),
+	TEXT("Allows to quickly reduce render feature to increase render performance.\n"
+		 "This is just a quick way to alter multiple show flags and console variables in the game\n"
+		 "Disabled more feature the higher the number\n"
+		 " <=0:off, order is defined in code (can be documented here when we settled on an order)"),
 	ECVF_Cheat | ECVF_RenderThreadSafe);
 
 #endif // !(UE_BUILD_SHIPPING || UE_BUILD_TEST)
@@ -2410,31 +2410,31 @@ static TAutoConsoleVariable<int32> CVarLimitRenderingFeatures(
 static TAutoConsoleVariable<int32> CVarUniformBufferPooling(
 	TEXT("r.UniformBufferPooling"),
 	1,
-	TEXT("If we pool object in RHICreateUniformBuffer to have less real API calls to create buffers\n")
-	TEXT(" 0: off (for debugging)\n")
-	TEXT(" 1: on (optimization)"),
+	TEXT("If we pool object in RHICreateUniformBuffer to have less real API calls to create buffers\n"
+		 " 0: off (for debugging)\n"
+		 " 1: on (optimization)"),
 	ECVF_RenderThreadSafe);
 
 static TAutoConsoleVariable<int32> CVarTranslucentSortPolicy(
 	TEXT("r.TranslucentSortPolicy"),
 	0,
-	TEXT("0: Sort based on distance from camera centerpoint to bounding sphere centerpoint. (default, best for 3D games)\n")
-	TEXT("1: Sort based on projected distance to camera.")
-	TEXT("2: Sort based on the projection onto a fixed axis. (best for 2D games)"),
+	TEXT("0: Sort based on distance from camera centerpoint to bounding sphere centerpoint. (default, best for 3D games)\n"
+		 "1: Sort based on projected distance to camera."
+		 "2: Sort based on the projection onto a fixed axis. (best for 2D games)"),
 	ECVF_RenderThreadSafe);
 
 static TAutoConsoleVariable<int32> CVarMobileHDR(
 	TEXT("r.MobileHDR"),
 	1,
-	TEXT("0: Mobile renders in LDR gamma space. (suggested for unlit games targeting low-end phones)\n")
-	TEXT("1: Mobile renders in HDR linear space. (default)"),
+	TEXT("0: Mobile renders in LDR gamma space. (suggested for unlit games targeting low-end phones)\n"
+		 "1: Mobile renders in HDR linear space. (default)"),
 	ECVF_RenderThreadSafe | ECVF_ReadOnly);
 
 static TAutoConsoleVariable<int32> CVarMobileShadingPath(
 	TEXT("r.Mobile.ShadingPath"),
 	0,
-	TEXT("0: Forward shading (default)")
-	TEXT("1: Deferred shading"),
+	TEXT("0: Forward shading (default)"
+		 "1: Deferred shading"),
 	ECVF_RenderThreadSafe | ECVF_ReadOnly);
 
 static TAutoConsoleVariable<int32> CVarMobileNumDynamicPointLights(
@@ -2446,15 +2446,15 @@ static TAutoConsoleVariable<int32> CVarMobileNumDynamicPointLights(
 static TAutoConsoleVariable<int32> CVarMobileEnableStaticAndCSMShadowReceivers(
 	TEXT("r.Mobile.EnableStaticAndCSMShadowReceivers"),
 	1,
-	TEXT("0: Primitives can receive only static shadowing from stationary lights.\n")
-	TEXT("1: Primitives can receive both CSM and static shadowing from stationary lights. (default)"),
+	TEXT("0: Primitives can receive only static shadowing from stationary lights.\n"
+		 "1: Primitives can receive both CSM and static shadowing from stationary lights. (default)"),
 	ECVF_RenderThreadSafe | ECVF_ReadOnly);
 
 static TAutoConsoleVariable<int32> CVarMobileEnableMovableLightCSMShaderCulling(
 	TEXT("r.Mobile.EnableMovableLightCSMShaderCulling"),
 	1,
-	TEXT("0: All primitives lit by movable directional light render with CSM.\n")
-	TEXT("1: Primitives lit by movable directional light render with the CSM shader when determined to be within CSM range. (default)"),
+	TEXT("0: All primitives lit by movable directional light render with CSM.\n"
+		 "1: Primitives lit by movable directional light render with the CSM shader when determined to be within CSM range. (default)"),
 	ECVF_RenderThreadSafe | ECVF_ReadOnly);
 
 static TAutoConsoleVariable<int32> CVarMobileEnableNoPrecomputedLightingCSMShader(
@@ -2467,8 +2467,8 @@ static TAutoConsoleVariable<int32> CVarMobileEnableNoPrecomputedLightingCSMShade
 static TAutoConsoleVariable<int32> CVarMobileUseCSMShaderBranch(
 	TEXT("r.Mobile.UseCSMShaderBranch"),
 	0,
-	TEXT("0: Use two shader permutations for CSM and non-CSM shading. (default)\n")
-	TEXT("1: Use a single shader pemutation with a branch in a shader to apply CSM (only with r.AllowStaticLighting=0)"),
+	TEXT("0: Use two shader permutations for CSM and non-CSM shading. (default)\n"
+		 "1: Use a single shader pemutation with a branch in a shader to apply CSM (only with r.AllowStaticLighting=0)"),
 	ECVF_RenderThreadSafe | ECVF_ReadOnly);
 
 static TAutoConsoleVariable<float> CVarsCSMDebugHint(
@@ -2480,31 +2480,31 @@ static TAutoConsoleVariable<float> CVarsCSMDebugHint(
 static TAutoConsoleVariable<int32> CVarMobileAllowDistanceFieldShadows(
 	TEXT("r.Mobile.AllowDistanceFieldShadows"),
 	1,
-	TEXT("0: Do not generate shader permutations to render distance field shadows from stationary directional lights.\n")
-	TEXT("1: Generate shader permutations to render distance field shadows from stationary directional lights. (default)"),
+	TEXT("0: Do not generate shader permutations to render distance field shadows from stationary directional lights.\n"
+		 "1: Generate shader permutations to render distance field shadows from stationary directional lights. (default)"),
 	ECVF_RenderThreadSafe | ECVF_ReadOnly);
 
 static TAutoConsoleVariable<int32> CVarMobileAllowMovableDirectionalLights(
 	TEXT("r.Mobile.AllowMovableDirectionalLights"),
 	1,
-	TEXT("0: Do not generate shader permutations to render movable directional lights.\n")
-	TEXT("1: Generate shader permutations to render movable directional lights. (default)"),
+	TEXT("0: Do not generate shader permutations to render movable directional lights.\n"
+		 "1: Generate shader permutations to render movable directional lights. (default)"),
 	ECVF_RenderThreadSafe | ECVF_ReadOnly);
 
 static TAutoConsoleVariable<int32> CVarMobileSkyLightPermutation(
 	TEXT("r.Mobile.SkyLightPermutation"),
 	0,
-	TEXT("0: Generate both sky-light and non-skylight permutations. (default)\n")
-	TEXT("1: Generate only non-skylight permutations.\n")
-	TEXT("2: Generate only skylight permutations"),
+	TEXT("0: Generate both sky-light and non-skylight permutations. (default)\n"
+		 "1: Generate only non-skylight permutations.\n"
+		 "2: Generate only skylight permutations"),
 	ECVF_RenderThreadSafe | ECVF_ReadOnly);
 
 static TAutoConsoleVariable<int32> CVarMobileFloatPrecisionMode(
 	TEXT("r.Mobile.FloatPrecisionMode"),
 	0,
-	TEXT("0: Use Half-precision (default)\n")
-	TEXT("1: Half precision, except Full precision for material expressions\n")
-	TEXT("2: Force use of high precision in pixel shaders.\n"),
+	TEXT("0: Use Half-precision (default)\n"
+		 "1: Half precision, except Full precision for material expressions\n"
+		 "2: Force use of high precision in pixel shaders.\n"),
 	ECVF_ReadOnly);
 
 static TAutoConsoleVariable<int32> CVarMobileAllowDitheredLODTransition(
@@ -2539,10 +2539,10 @@ static TAutoConsoleVariable<int32> CVarMobileSupportGPUScene(
 static TAutoConsoleVariable<int32> CVarSetClearSceneMethod(
 	TEXT("r.ClearSceneMethod"),
 	1,
-	TEXT("Select how the g-buffer is cleared in game mode (only affects deferred shading).\n")
-	TEXT(" 0: No clear\n")
-	TEXT(" 1: RHIClear (default)\n")
-	TEXT(" 2: Quad at max z"),
+	TEXT("Select how the g-buffer is cleared in game mode (only affects deferred shading).\n"
+		 " 0: No clear\n"
+		 " 1: RHIClear (default)\n"
+		 " 2: Quad at max z"),
 	ECVF_RenderThreadSafe);
 
 static TAutoConsoleVariable<int32> CVarLocalExposure(
@@ -2554,20 +2554,20 @@ static TAutoConsoleVariable<int32> CVarLocalExposure(
 static TAutoConsoleVariable<int32> CVarBloomQuality(
 	TEXT("r.BloomQuality"),
 	5,
-	TEXT(" 0: off, no performance impact.\n")
-	TEXT(" 1: average quality, least performance impact.\n")
-	TEXT(" 2: average quality, least performance impact.\n")
-	TEXT(" 3: good quality.\n")
-	TEXT(" 4: good quality.\n")
-	TEXT(" 5: Best quality, most significant performance impact. (default)\n")
-	TEXT(">5: force experimental higher quality on mobile (can be quite slow on some hardware)"),
+	TEXT(" 0: off, no performance impact.\n"
+		 " 1: average quality, least performance impact.\n"
+		 " 2: average quality, least performance impact.\n"
+		 " 3: good quality.\n"
+		 " 4: good quality.\n"
+		 " 5: Best quality, most significant performance impact. (default)\n"
+		 ">5: force experimental higher quality on mobile (can be quite slow on some hardware)"),
 	ECVF_Scalability | ECVF_RenderThreadSafe);
 
 static TAutoConsoleVariable<int32> CVarSceneColorFringeQuality(
 	TEXT("r.SceneColorFringeQuality"),
 	1,
-	TEXT(" 0: off but best for performance\n")
-	TEXT(" 1: 3 texture samples (default)"),
+	TEXT(" 0: off but best for performance\n"
+		 " 1: 3 texture samples (default)"),
 	ECVF_Scalability | ECVF_RenderThreadSafe);
 
 
@@ -2576,88 +2576,88 @@ static TAutoConsoleVariable<int32> CVarSceneColorFringeQuality(
 static TAutoConsoleVariable<float> CVarAmbientOcclusionRadiusScale(
 	TEXT("r.AmbientOcclusionRadiusScale"),
 	1.0f,
-	TEXT("Allows to scale the ambient occlusion radius (SSAO).\n")
-	TEXT(" 0:off, 1.0:normal, <1:smaller, >1:larger"),
+	TEXT("Allows to scale the ambient occlusion radius (SSAO).\n"
+		 " 0:off, 1.0:normal, <1:smaller, >1:larger"),
 	ECVF_Scalability | ECVF_RenderThreadSafe);
 
 static TAutoConsoleVariable<float> CVarAmbientOcclusionStaticFraction(TEXT("r.AmbientOcclusionStaticFraction"),
 	-1.0f,
-	TEXT("Allows to override the Ambient Occlusion Static Fraction (see post process volume). Fractions are between 0 and 1.\n")
-	TEXT("<0: use default setting (default -1)\n")
-	TEXT(" 0: no effect on static lighting, 0 is free meaning no extra rendering pass\n")
-	TEXT(" 1: AO affects the stat lighting"),
+	TEXT("Allows to override the Ambient Occlusion Static Fraction (see post process volume). Fractions are between 0 and 1.\n"
+		 "<0: use default setting (default -1)\n"
+		 " 0: no effect on static lighting, 0 is free meaning no extra rendering pass\n"
+		 " 1: AO affects the stat lighting"),
 	ECVF_Default);
 
 static TAutoConsoleVariable<int32> CVarShadowQuality(
 	TEXT("r.ShadowQuality"),
 	5,
-	TEXT("Defines the shadow method which allows to adjust for quality or performance.\n")
-	TEXT(" 0:off, 1:low(unfiltered), 2:low .. 5:max (default)"),
+	TEXT("Defines the shadow method which allows to adjust for quality or performance.\n"
+		 " 0:off, 1:low(unfiltered), 2:low .. 5:max (default)"),
 	ECVF_Scalability | ECVF_RenderThreadSafe);
 
 static TAutoConsoleVariable<int32> CVarMotionBlurQuality(
 	TEXT("r.MotionBlurQuality"),
 	4,
-	TEXT("Defines the motion blur method which allows to adjust for quality or performance.\n")
-	TEXT(" 0:off, 1:low, 2:medium, 3:high (default), 4: very high"),
+	TEXT("Defines the motion blur method which allows to adjust for quality or performance.\n"
+		 " 0:off, 1:low, 2:medium, 3:high (default), 4: very high"),
 	ECVF_Scalability | ECVF_RenderThreadSafe);
 
 static TAutoConsoleVariable<int32> CVarFullscreenMode(
 	TEXT("r.FullScreenMode"),
 	1,
-	TEXT("Defines how we do full screen when requested (e.g. command line option -fullscreen or in ini [SystemSettings] fullscreen=true)\n")
-	TEXT(" 0: normal full screen (renders faster, more control over vsync, less GPU memory, 10bit color if possible)\n")
-	TEXT(" 1: windowed full screen (quick switch between applications and window mode, slight performance loss)\n")
-	TEXT(" any other number behaves like 0"),
+	TEXT("Defines how we do full screen when requested (e.g. command line option -fullscreen or in ini [SystemSettings] fullscreen=true)\n"
+		 " 0: normal full screen (renders faster, more control over vsync, less GPU memory, 10bit color if possible)\n"
+		 " 1: windowed full screen (quick switch between applications and window mode, slight performance loss)\n"
+		 " any other number behaves like 0"),
 	ECVF_Scalability);
 
 static TAutoConsoleVariable<int32> CVarSceneColorFormat(
 	TEXT("r.SceneColorFormat"),
 	4,
-	TEXT("Defines the memory layout (RGBA) used for the scene color\n")
-	TEXT("(affects performance, mostly through bandwidth, quality especially with translucency).\n")
-	TEXT(" 0: PF_B8G8R8A8 32Bit (mostly for testing, likely to unusable with HDR)\n")
-	TEXT(" 1: PF_A2B10G10R10 32Bit\n")
-	TEXT(" 2: PF_FloatR11G11B10 32Bit\n")
-	TEXT(" 3: PF_FloatRGB 32Bit\n")
-	TEXT(" 4: PF_FloatRGBA 64Bit (default, might be overkill, especially if translucency is mostly using SeparateTranslucency)\n")
-	TEXT(" 5: PF_A32B32G32R32F 128Bit (unreasonable but good for testing)"),
+	TEXT("Defines the memory layout (RGBA) used for the scene color\n"
+		 "(affects performance, mostly through bandwidth, quality especially with translucency).\n"
+		 " 0: PF_B8G8R8A8 32Bit (mostly for testing, likely to unusable with HDR)\n"
+		 " 1: PF_A2B10G10R10 32Bit\n"
+		 " 2: PF_FloatR11G11B10 32Bit\n"
+		 " 3: PF_FloatRGB 32Bit\n"
+		 " 4: PF_FloatRGBA 64Bit (default, might be overkill, especially if translucency is mostly using SeparateTranslucency)\n"
+		 " 5: PF_A32B32G32R32F 128Bit (unreasonable but good for testing)"),
 	ECVF_Scalability | ECVF_RenderThreadSafe);
 
 static TAutoConsoleVariable<int32> CVarMobileSceneColorFormat(
 	TEXT("r.Mobile.SceneColorFormat"),
 	0,
-	TEXT("Overrides the memory layout (RGBA) used for the scene color of the mobile renderer.\nUnsupported overridden formats silently use default")
-	TEXT(" 0: (default) Automatically select the appropriate format depending on project settings and device support.\n")
-	TEXT(" 1: PF_FloatRGBA 64Bit \n")
-	TEXT(" 2: PF_FloatR11G11B10 32Bit\n")
-	TEXT(" 3: PF_B8G8R8A8 32Bit"),
+	TEXT("Overrides the memory layout (RGBA) used for the scene color of the mobile renderer.\nUnsupported overridden formats silently use default"
+		 " 0: (default) Automatically select the appropriate format depending on project settings and device support.\n"
+		 " 1: PF_FloatRGBA 64Bit \n"
+		 " 2: PF_FloatR11G11B10 32Bit\n"
+		 " 3: PF_B8G8R8A8 32Bit"),
 	ECVF_Scalability | ECVF_RenderThreadSafe);
 
 static TAutoConsoleVariable<int32> CVarPostProcessingColorFormat(
 	TEXT("r.PostProcessingColorFormat"),
 	0,
-	TEXT("Defines the memory layout (RGBA) used for most of the post processing chain buffers.\n")
-	TEXT(" 0: Default\n")
-	TEXT(" 1: Force PF_A32B32G32R32F 128Bit (unreasonable but good for testing)"),
+	TEXT("Defines the memory layout (RGBA) used for most of the post processing chain buffers.\n"
+		 " 0: Default\n"
+		 " 1: Force PF_A32B32G32R32F 128Bit (unreasonable but good for testing)"),
 	ECVF_Scalability | ECVF_RenderThreadSafe);
 
 static TAutoConsoleVariable<int32> CVarDepthOfFieldQuality(
 	TEXT("r.DepthOfFieldQuality"),
 	2,
-	TEXT("Allows to adjust the depth of field quality. Currently only fully affects BokehDOF. GaussianDOF is either 0 for off, otherwise on.\n")
-	TEXT(" 0: Off\n")
-	TEXT(" 1: Low\n")
-	TEXT(" 2: high quality (default, adaptive, can be 4x slower)\n")
-	TEXT(" 3: very high quality, intended for non realtime cutscenes, CircleDOF only (slow)\n")
-	TEXT(" 4: extremely high quality, intended for non realtime cutscenes, CircleDOF only (very slow)"),
+	TEXT("Allows to adjust the depth of field quality. Currently only fully affects BokehDOF. GaussianDOF is either 0 for off, otherwise on.\n"
+		 " 0: Off\n"
+		 " 1: Low\n"
+		 " 2: high quality (default, adaptive, can be 4x slower)\n"
+		 " 3: very high quality, intended for non realtime cutscenes, CircleDOF only (slow)\n"
+		 " 4: extremely high quality, intended for non realtime cutscenes, CircleDOF only (very slow)"),
 	ECVF_Scalability | ECVF_RenderThreadSafe);
 
 static TAutoConsoleVariable<int32> CVarHighResScreenshotDelay(
 	TEXT("r.HighResScreenshotDelay"),
 	4,
-	TEXT("When high-res screenshots are requested there is a small delay to allow temporal effects to converge.\n")
-	TEXT("Default: 4. Using a value below the default will disable TemporalAA for improved image quality."),
+	TEXT("When high-res screenshots are requested there is a small delay to allow temporal effects to converge.\n"
+		 "Default: 4. Using a value below the default will disable TemporalAA for improved image quality."),
 	ECVF_Default);
 
 static TAutoConsoleVariable<int32> CVarMaterialQualityLevel(
@@ -2669,11 +2669,11 @@ static TAutoConsoleVariable<int32> CVarMaterialQualityLevel(
 static TAutoConsoleVariable<int32> CVarUseDXT5NormalMaps(
 	TEXT("Compat.UseDXT5NormalMaps"),
 	0,
-	TEXT("Whether to use DXT5 for normal maps, otherwise BC5 will be used, which is not supported on all hardware.\n")
-	TEXT("Both formats require the same amount of memory (if driver doesn't emulate the format).\n")
-	TEXT("Changing this will cause normal maps to be recompressed on next load (or when using recompile shaders)\n")
-	TEXT(" 0: Use BC5 texture format (default)\n")
-	TEXT(" 1: Use DXT5 texture format (lower quality)"),
+	TEXT("Whether to use DXT5 for normal maps, otherwise BC5 will be used, which is not supported on all hardware.\n"
+		 "Both formats require the same amount of memory (if driver doesn't emulate the format).\n"
+		 "Changing this will cause normal maps to be recompressed on next load (or when using recompile shaders)\n"
+		 " 0: Use BC5 texture format (default)\n"
+		 " 1: Use DXT5 texture format (lower quality)"),
 	// 
 	// Changing this causes a full shader recompile
 	ECVF_ReadOnly);
@@ -2681,8 +2681,8 @@ static TAutoConsoleVariable<int32> CVarUseDXT5NormalMaps(
 static TAutoConsoleVariable<int32> CVarContactShadows(
 	TEXT("r.ContactShadows"),
 	1,
-	TEXT(" 0: disabled.\n")
-	TEXT(" 1: enabled.\n"),
+	TEXT(" 0: disabled.\n"
+		 " 1: enabled.\n"),
 	ECVF_Scalability | ECVF_RenderThreadSafe);
 
 static TAutoConsoleVariable<float> CVarContactShadowsNonShadowCastingIntensity(
@@ -2695,8 +2695,8 @@ static TAutoConsoleVariable<float> CVarContactShadowsNonShadowCastingIntensity(
 static TAutoConsoleVariable<int32> CVarAllowStaticLighting(
 	TEXT("r.AllowStaticLighting"),
 	1,
-	TEXT("Whether to allow any static lighting to be generated and used, like lightmaps and shadowmaps.\n")
-	TEXT("Games that only use dynamic lighting should set this to 0 to save some static lighting overhead."),
+	TEXT("Whether to allow any static lighting to be generated and used, like lightmaps and shadowmaps.\n"
+		 "Games that only use dynamic lighting should set this to 0 to save some static lighting overhead."),
 	ECVF_ReadOnly | ECVF_RenderThreadSafe);
 
 static TAutoConsoleVariable<int32> CVarNormalMaps(
@@ -2708,32 +2708,32 @@ static TAutoConsoleVariable<int32> CVarNormalMaps(
 static TAutoConsoleVariable<int32> CVarNumBufferedOcclusionQueries(
 	TEXT("r.NumBufferedOcclusionQueries"),
 	1,
-	TEXT("Number of frames to buffer occlusion queries (including the current renderthread frame).\n")
-	TEXT("More frames reduces the chance of stalling the CPU waiting for results, but increases out of date query artifacts."),
+	TEXT("Number of frames to buffer occlusion queries (including the current renderthread frame).\n"
+		 "More frames reduces the chance of stalling the CPU waiting for results, but increases out of date query artifacts."),
 	ECVF_RenderThreadSafe);
 
 static TAutoConsoleVariable<int32> CVarMinLogVerbosity(
 	TEXT("con.MinLogVerbosity"),
 	0,
-	TEXT("Allows to see the log in the in game console (by default deactivated to avoid spam and minor performance loss).\n")
-	TEXT(" 0: no logging other than console response (default)\n")
-	TEXT(" 1: Only fatal errors (no that useful)\n")
-	TEXT(" 2: additionally errors\n")
-	TEXT(" 3: additionally warnings\n")
-	TEXT(" 4: additionally display\n")
-	TEXT(" 5: additionally log\n")
-	TEXT("..\n")
-	TEXT(">=7: all"),
+	TEXT("Allows to see the log in the in game console (by default deactivated to avoid spam and minor performance loss).\n"
+		 " 0: no logging other than console response (default)\n"
+		 " 1: Only fatal errors (no that useful)\n"
+		 " 2: additionally errors\n"
+		 " 3: additionally warnings\n"
+		 " 4: additionally display\n"
+		 " 5: additionally log\n"
+		 "..\n"
+		 ">=7: all"),
 	ECVF_Default);
 
 static TAutoConsoleVariable<int32> CVarMSAACompositingSampleCount(
 	TEXT("r.MSAA.CompositingSampleCount"),
 	4,
-	TEXT("Affects the render quality of the editor 3d objects.\n")
-	TEXT(" 1: no MSAA, lowest quality\n")
-	TEXT(" 2: 2x MSAA, medium quality (medium GPU memory consumption)\n")
-	TEXT(" 4: 4x MSAA, high quality (high GPU memory consumption)\n")
-	TEXT(" 8: 8x MSAA, very high quality (insane GPU memory consumption)"),
+	TEXT("Affects the render quality of the editor 3d objects.\n"
+		 " 1: no MSAA, lowest quality\n"
+		 " 2: 2x MSAA, medium quality (medium GPU memory consumption)\n"
+		 " 4: 4x MSAA, high quality (high GPU memory consumption)\n"
+		 " 8: 8x MSAA, very high quality (insane GPU memory consumption)"),
 	ECVF_Scalability | ECVF_RenderThreadSafe);
 
 
@@ -2752,23 +2752,23 @@ static TAutoConsoleVariable<int32> CVarNetPackageMapDebugAllObjects(
 static TAutoConsoleVariable<FString> CVarNetPackageMapDebugObject(
 	TEXT("net.PackageMap.DebugObject"),
 	TEXT(""),
-	TEXT("Debugs PackageMap serialization of object")
-	TEXT("Partial name of object to debug"),
+	TEXT("Debugs PackageMap serialization of object"
+		 "Partial name of object to debug"),
 	ECVF_Default);
 
 static TAutoConsoleVariable<int32> CVarNetMontageDebug(
 	TEXT("net.Montage.Debug"),
 	0,
-	TEXT("Prints Replication information about AnimMontages\n")
-	TEXT(" 0: no print.\n")
-	TEXT(" 1: Print AnimMontage info on client side as they are played."),
+	TEXT("Prints Replication information about AnimMontages\n"
+		 " 0: no print.\n"
+		 " 1: Print AnimMontage info on client side as they are played."),
 	ECVF_Cheat);
 
 static TAutoConsoleVariable<int32> CVarRenderTargetPoolMin(
 	TEXT("r.RenderTargetPoolMin"),
 	400,
-	TEXT("If the render target pool size (in MB) is below this number there is no deallocation of rendertargets")
-	TEXT("Default is 200 MB."),
+	TEXT("If the render target pool size (in MB) is below this number there is no deallocation of rendertargets"
+		 "Default is 200 MB."),
 	ECVF_Scalability | ECVF_RenderThreadSafe);
 
 static TAutoConsoleVariable<int32> CVarIdleWhenNotForeground(
@@ -2779,24 +2779,24 @@ static TAutoConsoleVariable<int32> CVarIdleWhenNotForeground(
 static TAutoConsoleVariable<int32> CVarSetVSyncEnabled(
 	TEXT("r.VSync"),
 	0,
-	TEXT("0: VSync is disabled.(default)\n")
-	TEXT("1: VSync is enabled."),
+	TEXT("0: VSync is disabled.(default)\n"
+		 "1: VSync is enabled."),
 	ECVF_Scalability | ECVF_RenderThreadSafe);
 
 #if WITH_EDITOR
 static TAutoConsoleVariable<int32> CVarSetVSyncEditorEnabled(
 	TEXT("r.VSyncEditor"),
 	0,
-	TEXT("0: VSync is disabled in editor.(default)\n")
-	TEXT("1: VSync is enabled in editor."),
+	TEXT("0: VSync is disabled in editor.(default)\n"
+		 "1: VSync is enabled in editor."),
 	ECVF_RenderThreadSafe);
 
 static TAutoConsoleVariable<int32> CVarMobileForceRHISwitchVerticalAxis(
 	TEXT("r.Mobile.ForceRHISwitchVerticalAxis"),
 	0,
-	TEXT("Enable RHISwitchVerticalAxis when previewing mobile renderer. (Useful to test GLES y-axis flip codepaths)\n")
-	TEXT("0: RHISwitchVerticalAxis disabled (default).\n")
-	TEXT("1: RHISwitchVerticalAxis enabled.\n"),
+	TEXT("Enable RHISwitchVerticalAxis when previewing mobile renderer. (Useful to test GLES y-axis flip codepaths)\n"
+		 "0: RHISwitchVerticalAxis disabled (default).\n"
+		 "1: RHISwitchVerticalAxis enabled.\n"),
 	ECVF_Scalability | ECVF_RenderThreadSafe);
 #endif
 
@@ -2827,10 +2827,10 @@ static TAutoConsoleVariable<int32> CVarShadowMaxCSMShadowResolution(
 static TAutoConsoleVariable<float> CVarShadowCSMTransitionScale(
 	TEXT("r.Shadow.CSM.TransitionScale"),
 	1.0f,
-	TEXT("Allows to scale the cascaded shadow map transition region. Clamped within 0..2.\n")
-	TEXT("0: no transition (fastest)\n")
-	TEXT("1: as specific in the light settings (default)\n")
-	TEXT("2: 2x larger than what was specified in the light"),
+	TEXT("Allows to scale the cascaded shadow map transition region. Clamped within 0..2.\n"
+		 "0: no transition (fastest)\n"
+		 "1: as specific in the light settings (default)\n"
+		 "2: 2x larger than what was specified in the light"),
 	ECVF_Scalability | ECVF_RenderThreadSafe);
 
 static TAutoConsoleVariable<float> CVarMobileContentScaleFactor(
@@ -2842,81 +2842,81 @@ static TAutoConsoleVariable<float> CVarMobileContentScaleFactor(
 static TAutoConsoleVariable<int32> CVarMobileTonemapperUpscale(
 	TEXT("r.MobileTonemapperUpscale"),
 	0,
-	TEXT("On mobile, whether to allow upscaling as part of the tonemapper or as a separate pass when possible")
-	TEXT("0: separate pass (default)\n")
-	TEXT("1: as part of the tonemapper pass\n"),
+	TEXT("On mobile, whether to allow upscaling as part of the tonemapper or as a separate pass when possible"
+		 "0: separate pass (default)\n"
+		 "1: as part of the tonemapper pass\n"),
 	ECVF_Default);
 
 // this cvar can be removed in shipping to not compile shaders for development (faster)
 static TAutoConsoleVariable<int32> CVarCompileShadersForDevelopment(
 	TEXT("r.CompileShadersForDevelopment"),
 	1,
-	TEXT("Setting this to 0 allows to ship a game with more optimized shaders as some\n")
-	TEXT("editor and development features are not longer compiled into the shaders.\n")
-	TEXT(" Note: This should be done when shipping but it's not done automatically yet (feature need to mature\n")
-	TEXT("       and shaders will compile slower as shader caching from development isn't shared).\n")
-	TEXT("Cannot be changed at runtime - can be put into BaseEngine.ini\n")
-	TEXT(" 0: off, shader can run a bit faster\n")
-	TEXT(" 1: on (Default)"),
+	TEXT("Setting this to 0 allows to ship a game with more optimized shaders as some\n"
+		 "editor and development features are not longer compiled into the shaders.\n"
+		 " Note: This should be done when shipping but it's not done automatically yet (feature need to mature\n"
+		 "       and shaders will compile slower as shader caching from development isn't shared).\n"
+		 "Cannot be changed at runtime - can be put into BaseEngine.ini\n"
+		 " 0: off, shader can run a bit faster\n"
+		 " 1: on (Default)"),
 	ECVF_ReadOnly);
 
 static TAutoConsoleVariable<int32> CVarDontLimitOnBattery(
 	TEXT("r.DontLimitOnBattery"),
 	0,
-	TEXT("0: Limit performance on devices with a battery.(default)\n")
-	TEXT("1: Do not limit performance due to device having a battery."),
+	TEXT("0: Limit performance on devices with a battery.(default)\n"
+		 "1: Do not limit performance due to device having a battery."),
 	ECVF_Scalability | ECVF_RenderThreadSafe);
 
 static TAutoConsoleVariable<float> CVarViewDistanceScale(
 	TEXT("r.ViewDistanceScale"),
 	1.0f,
-	TEXT("Controls the view distance scale. A primitive's MaxDrawDistance is scaled by this value.\n")
-	TEXT("Higher values will increase view distance but at a performance cost.\n")
-	TEXT("Default = 1."),
+	TEXT("Controls the view distance scale. A primitive's MaxDrawDistance is scaled by this value.\n"
+		 "Higher values will increase view distance but at a performance cost.\n"
+		 "Default = 1."),
 	ECVF_Scalability | ECVF_RenderThreadSafe);
 
 static TAutoConsoleVariable<int32> CVarViewDistanceScaleApplySecondaryScale(
 	TEXT("r.ViewDistanceScale.ApplySecondaryScale"),
 	0,
-	TEXT("If true applies the secondary view distance scale to primitive draw distances.\n")
-	TEXT("Default = 0."),
+	TEXT("If true applies the secondary view distance scale to primitive draw distances.\n"
+		 "Default = 0."),
 	ECVF_RenderThreadSafe);
 
 static TAutoConsoleVariable<float> CVarViewDistanceScaleSecondaryScale(
 	TEXT("r.ViewDistanceScale.SecondaryScale"),
 	1.0f,
-	TEXT("Controls the secondary view distance scale, Default = 1.0.\n")
-	TEXT("This is an optional scale intended to allow some features or gamemodes to opt-in.\n"),
+	TEXT("Controls the secondary view distance scale, Default = 1.0.\n"
+		 "This is an optional scale intended to allow some features or gamemodes to opt-in.\n"),
 	ECVF_Scalability | ECVF_RenderThreadSafe);
 
 static TAutoConsoleVariable<float> CVarViewDistanceScale_FieldOfViewMinAngle(
 	TEXT("r.ViewDistanceScale.FieldOfViewMinAngle"),
 	45.0f,
-	TEXT("Scales the scene view distance scale with camera field of view.\n")
-	TEXT("Minimum angle of the blend range.\n")
-	TEXT("Applies the minimum scale when the camera is at or below this angle."),
+	TEXT("Scales the scene view distance scale with camera field of view.\n"
+		 "Minimum angle of the blend range.\n"
+		 "Applies the minimum scale when the camera is at or below this angle."),
 	ECVF_Scalability | ECVF_RenderThreadSafe);
 
 static TAutoConsoleVariable<float> CVarViewDistanceScale_FieldOfViewMinAngleScale(
 	TEXT("r.ViewDistanceScale.FieldOfViewMinAngleScale"),
 	1.0f,
-	TEXT("Scales the scene view distance scale with camera field of view.\n")
-	TEXT("This value is applied when the camera is at or below the minimum angle."),
+	TEXT("Scales the scene view distance scale with camera field of view.\n"
+		 "This value is applied when the camera is at or below the minimum angle."),
 	ECVF_Scalability | ECVF_RenderThreadSafe);
 
 static TAutoConsoleVariable<float> CVarViewDistanceScale_FieldOfViewMaxAngle(
 	TEXT("r.ViewDistanceScale.FieldOfViewMaxAngle"),
 	90.0f,
-	TEXT("Scales the scene view distance scale with camera field of view.\n")
-	TEXT("Maximum angle of the blend range.\n")
-	TEXT("Applies the maximum scale when the camera is at or above this angle."),
+	TEXT("Scales the scene view distance scale with camera field of view.\n"
+		 "Maximum angle of the blend range.\n"
+		 "Applies the maximum scale when the camera is at or above this angle."),
 	ECVF_Scalability | ECVF_RenderThreadSafe);
 
 static TAutoConsoleVariable<float> CVarViewDistanceScale_FieldOfViewMaxAngleScale(
 	TEXT("r.ViewDistanceScale.FieldOfViewMaxAngleScale"),
 	1.0f,
-	TEXT("Scales the scene view distance scale with camera field of view.\n")
-	TEXT("This value is applied when the camera is at or above the maximum angle."),
+	TEXT("Scales the scene view distance scale with camera field of view.\n"
+		 "This value is applied when the camera is at or above the maximum angle."),
 	ECVF_Scalability | ECVF_RenderThreadSafe);
 
 static TAutoConsoleVariable<int32> CVarViewDistanceScale_FieldOfViewAffectsHLOD(
@@ -2928,72 +2928,72 @@ static TAutoConsoleVariable<int32> CVarViewDistanceScale_FieldOfViewAffectsHLOD(
 static TAutoConsoleVariable<int32> CVarLightFunctionQuality(
 	TEXT("r.LightFunctionQuality"),
 	2,
-	TEXT("Defines the light function quality which allows to adjust for quality or performance.\n")
-	TEXT("<=0: off (fastest)\n")
-	TEXT("  1: low quality (e.g. half res with blurring, not yet implemented)\n")
-	TEXT("  2: normal quality (default)\n")
-	TEXT("  3: high quality (e.g. super-sampled or colored, not yet implemented)"),
+	TEXT("Defines the light function quality which allows to adjust for quality or performance.\n"
+		 "<=0: off (fastest)\n"
+		 "  1: low quality (e.g. half res with blurring, not yet implemented)\n"
+		 "  2: normal quality (default)\n"
+		 "  3: high quality (e.g. super-sampled or colored, not yet implemented)"),
 	ECVF_Scalability | ECVF_RenderThreadSafe);
 
 static TAutoConsoleVariable<int32> CVarEyeAdaptationQuality(
 	TEXT("r.EyeAdaptationQuality"),
 	2,
-	TEXT("Defines the eye adaptation quality which allows to adjust for quality or performance.\n")
-	TEXT("<=0: off (fastest)\n")
-	TEXT("  1: low quality (e.g. non histogram based, not yet implemented)\n")
-	TEXT("  2: normal quality (default)\n")
-	TEXT("  3: high quality (e.g. screen position localized, not yet implemented)"),
+	TEXT("Defines the eye adaptation quality which allows to adjust for quality or performance.\n"
+		 "<=0: off (fastest)\n"
+		 "  1: low quality (e.g. non histogram based, not yet implemented)\n"
+		 "  2: normal quality (default)\n"
+		 "  3: high quality (e.g. screen position localized, not yet implemented)"),
 	ECVF_Scalability | ECVF_RenderThreadSafe);
 
 static TAutoConsoleVariable<float> CVarShadowDistanceScale(
 	TEXT("r.Shadow.DistanceScale"),
 	1.0f,
-	TEXT("Scalability option to trade shadow distance versus performance for directional lights (clamped within a reasonable range).\n")
-	TEXT("<1: shorter distance\n")
-	TEXT(" 1: normal (default)\n")
-	TEXT(">1: larger distance"),
+	TEXT("Scalability option to trade shadow distance versus performance for directional lights (clamped within a reasonable range).\n"
+		 "<1: shorter distance\n"
+		 " 1: normal (default)\n"
+		 ">1: larger distance"),
 	ECVF_Scalability | ECVF_RenderThreadSafe);
 
 static TAutoConsoleVariable<int32> CVarFreeSkeletalMeshBuffers(
 	TEXT("r.FreeSkeletalMeshBuffers"),
 	0,
-	TEXT("Controls whether skeletal mesh buffers are kept in CPU memory to support merging of skeletal meshes.\n")
-	TEXT("0: Keep buffers(default)\n")
-	TEXT("1: Free buffers"),
+	TEXT("Controls whether skeletal mesh buffers are kept in CPU memory to support merging of skeletal meshes.\n"
+		 "0: Keep buffers(default)\n"
+		 "1: Free buffers"),
 	ECVF_RenderThreadSafe);
 
 static TAutoConsoleVariable<int32> CVarTonemapperGrainQuantization(
 	TEXT("r.Tonemapper.GrainQuantization"),
 	1,
-	TEXT("0: low (minor performance benefit)\n")
-	TEXT("1: high (default, with high frequency pixel pattern to fight 8 bit color quantization)"),
+	TEXT("0: low (minor performance benefit)\n"
+		 "1: high (default, with high frequency pixel pattern to fight 8 bit color quantization)"),
 	ECVF_Scalability | ECVF_RenderThreadSafe);
 
 static TAutoConsoleVariable<int32> CVarDetailMode(
 	TEXT("r.DetailMode"),
 	2,
-	TEXT("Current detail mode; determines whether components of actors should be updated/ ticked.\n")
-	TEXT(" 0: low, show only object with DetailMode low or higher\n")
-	TEXT(" 1: medium, show all object with DetailMode medium or higher\n")
-	TEXT(" 2: high, show all objects (default)"),
+	TEXT("Current detail mode; determines whether components of actors should be updated/ ticked.\n"
+		 " 0: low, show only object with DetailMode low or higher\n"
+		 " 1: medium, show all object with DetailMode medium or higher\n"
+		 " 2: high, show all objects (default)"),
 	ECVF_Scalability | ECVF_RenderThreadSafe);
 
 static TAutoConsoleVariable<int32> CVarCookOutUnusedDetailModeComponents(
 	TEXT("r.CookOutUnusedDetailModeComponents"),
 	0,
-	TEXT("If set, components which are not relevant for the current detail mode will be cooked out.\n")
-	TEXT(" 0: keep components even if not relevant for the current detail mode.\n")
-	TEXT(" 1: cook out components not relevant for the current detail mode.\n"),
+	TEXT("If set, components which are not relevant for the current detail mode will be cooked out.\n"
+		 " 0: keep components even if not relevant for the current detail mode.\n"
+		 " 1: cook out components not relevant for the current detail mode.\n"),
 	ECVF_RenderThreadSafe);
 
 static TAutoConsoleVariable<int32> CVarDBuffer(
 	TEXT("r.DBuffer"),
 	1,
-	TEXT("Enables DBuffer decal material blend modes.\n")
-	TEXT("DBuffer decals are rendered before the base pass, allowing them to affect static lighting and skylighting correctly. \n")
-	TEXT("When enabled, a full prepass will be forced which adds CPU / GPU cost.  Several texture lookups will be done in the base pass to fetch the decal properties, which adds pixel work.\n")
-	TEXT(" 0: off\n")
-	TEXT(" 1: on (default)"),
+	TEXT("Enables DBuffer decal material blend modes.\n"
+		 "DBuffer decals are rendered before the base pass, allowing them to affect static lighting and skylighting correctly. \n"
+		 "When enabled, a full prepass will be forced which adds CPU / GPU cost.  Several texture lookups will be done in the base pass to fetch the decal properties, which adds pixel work.\n"
+		 " 0: off\n"
+		 " 1: on (default)"),
 	ECVF_RenderThreadSafe | ECVF_ReadOnly);
 
 static TAutoConsoleVariable<float> CVarSkeletalMeshLODRadiusScale(
@@ -3035,9 +3035,9 @@ static TAutoConsoleVariable<int32> CVarFeatureLevelPreview(
 static TAutoConsoleVariable<int32> CVarVerifyPeer(
 	TEXT("n.VerifyPeer"),
 	1,
-	TEXT("Sets libcurl's CURLOPT_SSL_VERIFYPEER option to verify authenticity of the peer's certificate.\n")
-	TEXT("  0 = disable (allows self-signed certificates)\n")
-	TEXT("  1 = enable [default]"),
+	TEXT("Sets libcurl's CURLOPT_SSL_VERIFYPEER option to verify authenticity of the peer's certificate.\n"
+		 "  0 = disable (allows self-signed certificates)\n"
+		 "  1 = enable [default]"),
 	ECVF_ReadOnly);
 
 static TAutoConsoleVariable<float> CVarEmitterSpawnRateScale(
@@ -3055,58 +3055,58 @@ static TAutoConsoleVariable<int32> CVarCheckSRVTransitions(
 static TAutoConsoleVariable<int32> CVarDisableThreadedRendering(
 	TEXT("r.AndroidDisableThreadedRendering"),
 	0,
-	TEXT("Sets whether or not to allow threaded rendering for a particular Android device profile.\n")
-	TEXT("	0 = Allow threaded rendering [default]\n")
-	TEXT("	1 = Disable creation of render thread on startup"),
+	TEXT("Sets whether or not to allow threaded rendering for a particular Android device profile.\n"
+		 "	0 = Allow threaded rendering [default]\n"
+		 "	1 = Disable creation of render thread on startup"),
 	ECVF_ReadOnly);
 
 
 static TAutoConsoleVariable<int32> CVarDisableThreadedRenderingFirstLoad(
 	TEXT("r.AndroidDisableThreadedRenderingFirstLoad"),
 	0,
-	TEXT("Sets whether or not to allow threaded rendering for a particular Android device profile on the initial load.\n")
-	TEXT("	0 = Allow threaded rendering on the initial load [default]\n")
-	TEXT("	1 = Disable threaded rendering on the initial load"),
+	TEXT("Sets whether or not to allow threaded rendering for a particular Android device profile on the initial load.\n"
+		 "	0 = Allow threaded rendering on the initial load [default]\n"
+		 "	1 = Disable threaded rendering on the initial load"),
 	ECVF_ReadOnly);
 
 static TAutoConsoleVariable<int32> CVarDisableVulkanSupport(
 	TEXT("r.Android.DisableVulkanSupport"),
 	0,
-	TEXT("Disable support for vulkan API. (Android Only)\n")
-	TEXT("  0 = vulkan API will be used (providing device and project supports it) [default]\n")
-	TEXT("  1 = vulkan will be disabled, opengl fall back will be used."),
+	TEXT("Disable support for vulkan API. (Android Only)\n"
+		 "  0 = vulkan API will be used (providing device and project supports it) [default]\n"
+		 "  1 = vulkan will be disabled, opengl fall back will be used."),
 	ECVF_ReadOnly);
 
 static TAutoConsoleVariable<int32> CVarDisableVulkanSM5Support(
 	TEXT("r.Android.DisableVulkanSM5Support"),
 	0,
-	TEXT("Disable support for vulkan API. (Android Only)\n")
-	TEXT("  0 = Vulkan SM5 API will be used (providing device and project supports it) [default]\n")
-	TEXT("  1 = Vulkan SM5 will be disabled, Vulkan or OpenGL fall back will be used."),
+	TEXT("Disable support for vulkan API. (Android Only)\n"
+		 "  0 = Vulkan SM5 API will be used (providing device and project supports it) [default]\n"
+		 "  1 = Vulkan SM5 will be disabled, Vulkan or OpenGL fall back will be used."),
 	ECVF_ReadOnly);
 
 static TAutoConsoleVariable<int32> CVarDisableOpenGLES31Support(
 	TEXT("r.Android.DisableOpenGLES31Support"),
 	0,
-	TEXT("Disable support for OpenGLES 3.1 API. (Android Only)\n")
-	TEXT("  0 = OpenGLES 3.1 API will be used (providing device and project supports it) [default]\n")
-	TEXT("  1 = OpenGLES 3.1 will be disabled, Vulkan will be used."),
+	TEXT("Disable support for OpenGLES 3.1 API. (Android Only)\n"
+		 "  0 = OpenGLES 3.1 API will be used (providing device and project supports it) [default]\n"
+		 "  1 = OpenGLES 3.1 will be disabled, Vulkan will be used."),
 	ECVF_ReadOnly);
 
 static TAutoConsoleVariable<int32> CVarDisableAndroidGLASTCSupport(
 	TEXT("r.Android.DisableASTCSupport"),
 	0,
-	TEXT("Disable support for ASTC Texture compression if OpenGL driver supports it. (Android Only)\n")
-	TEXT("  0 = ASTC texture compression will be used if driver supports it [default]\n")
-	TEXT("  1 = ASTC texture compression will not be used."),
+	TEXT("Disable support for ASTC Texture compression if OpenGL driver supports it. (Android Only)\n"
+		 "  0 = ASTC texture compression will be used if driver supports it [default]\n"
+		 "  1 = ASTC texture compression will not be used."),
 	ECVF_ReadOnly);
 
 static TAutoConsoleVariable<int32> CVarDisableOpenGLTextureStreamingSupport(
 	TEXT("r.OpenGL.DisableTextureStreamingSupport"),
 	0,
-	TEXT("Disable support for texture streaming on OpenGL.\n")
-	TEXT("  0 = Texture streaming will be used if device supports it [default]\n")
-	TEXT("  1 = Texture streaming will be disabled."),
+	TEXT("Disable support for texture streaming on OpenGL.\n"
+		 "  0 = Texture streaming will be used if device supports it [default]\n"
+		 "  1 = Texture streaming will be disabled."),
 	ECVF_ReadOnly);
 
 // Moved here from OpenGLRHI module to make sure its always accessible on all platforms
@@ -3120,45 +3120,45 @@ static FAutoConsoleVariable CVarOpenGLUseEmulatedUBs(
 static TAutoConsoleVariable<int32> CVarAndroidOverrideExternalTextureSupport(
 	TEXT("r.Android.OverrideExternalTextureSupport"),
 	0,
-	TEXT("Override external texture support for OpenGLES API. (Android Only)\n")
-	TEXT("  0 = normal detection used [default]\n")
-	TEXT("  1 = disable external texture support\n")
-	TEXT("  2 = force ImageExternal100 (version #100 with GL_OES_EGL_image_external)\n")
-	TEXT("  3 = force ImageExternal300 (version #300 with GL_OES_EGL_image_external)\n")
-	TEXT("  4 = force ImageExternalESSL300 (version #300 with GL_OES_EGL_image_external_essl3)"),
+	TEXT("Override external texture support for OpenGLES API. (Android Only)\n"
+		 "  0 = normal detection used [default]\n"
+		 "  1 = disable external texture support\n"
+		 "  2 = force ImageExternal100 (version #100 with GL_OES_EGL_image_external)\n"
+		 "  3 = force ImageExternal300 (version #300 with GL_OES_EGL_image_external)\n"
+		 "  4 = force ImageExternalESSL300 (version #300 with GL_OES_EGL_image_external_essl3)"),
 	ECVF_ReadOnly);
 
 static TAutoConsoleVariable<FString> CVarCustomUnsafeZones(
 	TEXT("r.CustomUnsafeZones"),
 	TEXT(""),
-	TEXT("Allows you to set custom unsafe zones. Define them based on Portrait (P) or Landscape (L) for a device oriented 'upright'.")
-	TEXT("Unsafe zones may be either fixed or free, depending on if they move along with the rotation of the device.")
-	TEXT("Format is (P:fixed[x1, y1][width, height]), semicolon-separated for each custom unsafe zone. +Values add from 0, -Values subtract from Height or Width"),
+	TEXT("Allows you to set custom unsafe zones. Define them based on Portrait (P) or Landscape (L) for a device oriented 'upright'."
+		 "Unsafe zones may be either fixed or free, depending on if they move along with the rotation of the device."
+		 "Format is (P:fixed[x1, y1][width, height]), semicolon-separated for each custom unsafe zone. +Values add from 0, -Values subtract from Height or Width"),
 	ECVF_Default);
 
 static TAutoConsoleVariable<int32> CVarSkyLightingQuality(
 	TEXT("r.SkyLightingQuality"),
 	1,
-	TEXT("Defines the sky lighting quality which allows to adjust for performance.\n")
-	TEXT("<=0: off (fastest)\n")
-	TEXT("  1: on\n"),
+	TEXT("Defines the sky lighting quality which allows to adjust for performance.\n"
+		 "<=0: off (fastest)\n"
+		 "  1: on\n"),
 	ECVF_Scalability | ECVF_RenderThreadSafe
 );
 
 static TAutoConsoleVariable<int32> CVarMobileDefaultAntiAliasing(
 	TEXT("r.Mobile.AntiAliasing"),
 	1,
-	TEXT("Mobile default AntiAliasingMethod\n")
-	TEXT(" 0: off (no anti-aliasing)\n")
-	TEXT(" 1: FXAA (default, faster than TemporalAA but much more shimmering for non static cases)\n")
-	TEXT(" 2: TemporalAA(it will fallback to FXAA if SupportsGen4TAA is disabled) \n")
-	TEXT(" 3: MSAA"),
+	TEXT("Mobile default AntiAliasingMethod\n"
+		 " 0: off (no anti-aliasing)\n"
+		 " 1: FXAA (default, faster than TemporalAA but much more shimmering for non static cases)\n"
+		 " 2: TemporalAA(it will fallback to FXAA if SupportsGen4TAA is disabled) \n"
+		 " 3: MSAA"),
 	ECVF_RenderThreadSafe);
 
 static TAutoConsoleVariable<int32> CVarMobileSupportsGen4TAA(
 	TEXT("r.Mobile.SupportsGen4TAA"),
 	1,
-	TEXT("Support desktop Gen4 TAA with mobile rendering\n")
-	TEXT("0: Fallback to FXAA")
-	TEXT("1: Support Desktop Gen4 TAA (default)"),
+	TEXT("Support desktop Gen4 TAA with mobile rendering\n"
+		 "0: Fallback to FXAA"
+		 "1: Support Desktop Gen4 TAA (default)"),
 	ECVF_ReadOnly | ECVF_RenderThreadSafe);
\ No newline at end of file
