Add insights traces on pipe tasks for audio and stats

diff --git a/Engine/Source/Runtime/Core/Private/Async/TaskGraph.cpp b/Engine/Source/Runtime/Core/Private/Async/TaskGraph.cpp
--- a/Engine/Source/Runtime/Core/Private/Async/TaskGraph.cpp
+++ b/Engine/Source/Runtime/Core/Private/Async/TaskGraph.cpp
@@ -1421,6 +1421,7 @@ PRAGMA_ENABLE_DEPRECATION_WARNINGS
 			GStatsPipe.Launch(UE_SOURCE_LOCATION,
 				[Task]
 				{
+					TRACE_CPUPROFILER_EVENT_SCOPE(StatsPipeWork);
 					TArray<FBaseGraphTask*> Dummy;
 					Task->Execute(Dummy, ENamedThreads::AnyThread, true);
 				}
@@ -1444,6 +1445,7 @@ PRAGMA_ENABLE_DEPRECATION_WARNINGS
 			GAudioPipe.Launch(UE_SOURCE_LOCATION,
 				[Task]
 				{
+					TRACE_CPUPROFILER_EVENT_SCOPE(AudioPipeWork);
 					TArray<FBaseGraphTask*> Dummy;
 					Task->Execute(Dummy, ENamedThreads::AnyThread, true);
 				}
