ParalleFor doesn't go wider than number of cores

diff --git a/Engine/Source/Runtime/Core/Public/Async/ParallelFor.h b/Engine/Source/Runtime/Core/Public/Async/ParallelFor.h
--- a/Engine/Source/Runtime/Core/Public/Async/ParallelFor.h
+++ b/Engine/Source/Runtime/Core/Public/Async/ParallelFor.h
@@ -77,6 +77,9 @@ namespace ParallelForImpl
 			NumThreadTasks++; //named threads help with the work
 		}
 
+		// don't go wider than number of cores
+		NumThreadTasks = FMath::Min(NumThreadTasks, FPlatformMisc::NumberOfCoresIncludingHyperthreads());
+
 		return FMath::Max(NumThreadTasks, 1);
 	}
 
