Fix calculation of number of partitions for iostore containers to use end of last block rather than start of last block.

#localization none
#tests staged & run windows server locally

diff --git a/Engine/Source/Runtime/Core/Private/IO/IoStore.cpp b/Engine/Source/Runtime/Core/Private/IO/IoStore.cpp
--- a/Engine/Source/Runtime/Core/Private/IO/IoStore.cpp
+++ b/Engine/Source/Runtime/Core/Private/IO/IoStore.cpp
@@ -2711,7 +2711,10 @@ TIoStatusOr<uint64> FIoStoreTocResource::Write(
 	}
 	else if (WriterSettings.MaxPartitionSize)
 	{
-		TocHeader.PartitionCount = uint32(Align(TocResource.CompressionBlocks.Last().GetOffset(), WriterSettings.MaxPartitionSize) / WriterSettings.MaxPartitionSize);
+		const FIoStoreTocCompressedBlockEntry& LastBlock = TocResource.CompressionBlocks.Last();
+		uint64 LastBlockEnd = LastBlock.GetOffset() + LastBlock.GetCompressedSize();
+		TocHeader.PartitionCount = uint32(Align(LastBlockEnd, WriterSettings.MaxPartitionSize) / WriterSettings.MaxPartitionSize);
+		check(TocHeader.PartitionCount > 0);
 		TocHeader.PartitionSize = WriterSettings.MaxPartitionSize;
 	}
 	else
