Update PathPermissionList API to include RemoveDeny/AllowList calls

diff --git a/Engine/Source/Runtime/Core/Private/Misc/PathPermissionList.cpp b/Engine/Source/Runtime/Core/Private/Misc/PathPermissionList.cpp
--- a/Engine/Source/Runtime/Core/Private/Misc/PathPermissionList.cpp
+++ b/Engine/Source/Runtime/Core/Private/Misc/PathPermissionList.cpp
@@ -129,6 +129,38 @@ bool FPathPermissionList::AddDenyListItem(const FName OwnerName, const TCHAR* It
 	return AddDenyListItem(OwnerName, FStringView(Item));
 }
 
+bool FPathPermissionList::RemoveDenyListItem(const FName OwnerName, const FStringView Item)
+{
+	const uint32 ItemHash = GetTypeHash(Item);
+
+	FPermissionListOwners* Owners = DenyList.FindByHash(ItemHash, Item);
+	if (Owners && Owners->Remove(OwnerName) == 1)
+	{
+		if (Owners->Num() == 0)
+		{
+			DenyList.RemoveByHash(ItemHash, Item);
+		}
+
+		if (!bSuppressOnFilterChanged)
+		{
+			OnFilterChanged().Broadcast();
+		}
+		return true;
+	}
+
+	return false;
+}
+
+bool FPathPermissionList::RemoveDenyListItem(const FName OwnerName, const FName Item)
+{
+	return RemoveDenyListItem(OwnerName, FNameBuilder(Item));
+}
+
+bool FPathPermissionList::RemoveDenyListItem(const FName OwnerName, const TCHAR* Item)
+{
+	return RemoveDenyListItem(OwnerName, FStringView(Item));
+}
+
 bool FPathPermissionList::AddAllowListItem(const FName OwnerName, const FStringView Item)
 {
 	const uint32 ItemHash = GetTypeHash(Item);
@@ -174,6 +206,38 @@ bool FPathPermissionList::AddDenyListAll(const FName OwnerName)
 	return bFilterChanged;
 }
 
+bool FPathPermissionList::RemoveAllowListItem(const FName OwnerName, const FStringView Item)
+{
+	const uint32 ItemHash = GetTypeHash(Item);
+
+	FPermissionListOwners* Owners = AllowList.FindByHash(ItemHash, Item);
+	if (Owners && Owners->Remove(OwnerName) == 1)
+	{
+		if (Owners->Num() == 0)
+		{
+			AllowList.RemoveByHash(ItemHash, Item);
+		}
+
+		if (!bSuppressOnFilterChanged)
+		{
+			OnFilterChanged().Broadcast();
+		}
+		return true;
+	}
+
+	return false;
+}
+
+bool FPathPermissionList::RemoveAllowListItem(const FName OwnerName, const FName Item)
+{
+	return RemoveAllowListItem(OwnerName, FNameBuilder(Item));
+}
+
+bool FPathPermissionList::RemoveAllowListItem(const FName OwnerName, const TCHAR* Item)
+{
+	return RemoveAllowListItem(OwnerName, FStringView(Item));
+}
+
 bool FPathPermissionList::HasFiltering() const
 {
 	return DenyList.Num() > 0 || AllowList.Num() > 0 || DenyListAll.Num() > 0;
diff --git a/Engine/Source/Runtime/Core/Public/Misc/NamePermissionList.h b/Engine/Source/Runtime/Core/Public/Misc/NamePermissionList.h
--- a/Engine/Source/Runtime/Core/Public/Misc/NamePermissionList.h
+++ b/Engine/Source/Runtime/Core/Public/Misc/NamePermissionList.h
@@ -151,6 +151,24 @@ public:
 	 */
 	bool AddDenyListItem(const FName OwnerName, const TCHAR* Item);
 
+	/**
+	* Remove item from the DenyList
+	* @return whether the filters changed
+	*/
+	bool RemoveDenyListItem(const FName OwnerName, const FStringView Item);
+
+	/**
+	* Remove item from the DenyList
+	* @return whether the filters changed
+	*/
+	bool RemoveDenyListItem(const FName OwnerName, const FName Item);
+
+	/**
+	* Remove item from the DenyList
+	* @return whether the filters changed
+	*/
+	bool RemoveDenyListItem(const FName OwnerName, const TCHAR* Item);
+
 	/**
 	 * Add item to allowlist after which all items not in the allowlist will be filtered out.
 	 * @return whether the filters changed.
@@ -169,6 +187,24 @@ public:
 	 */
 	bool AddAllowListItem(const FName OwnerName, const TCHAR* Item);
 
+	/**
+	* Remove item from the AllowList
+	* @return whether the filters changed
+	*/
+	bool RemoveAllowListItem(const FName OwnerName, const FStringView Item);
+
+	/**
+	* Remove item from the AllowList
+	* @return whether the filters changed
+	*/
+	bool RemoveAllowListItem(const FName OwnerName, const FName Item);
+
+	/**
+	* Remove item from the AllowList
+	* @return whether the filters changed
+	*/
+	bool RemoveAllowListItem(const FName OwnerName, const TCHAR* Item);
+
 	/**
 	 * Set to filter out all items.
 	 * @return whether the filters changed.
