Fix accidental removal of 'n'
#rnx

diff --git a/Engine/Source/Runtime/Core/Private/UObject/DevObjectVersion.cpp b/Engine/Source/Runtime/Core/Private/UObject/DevObjectVersion.cpp
--- a/Engine/Source/Runtime/Core/Private/UObject/DevObjectVersion.cpp
+++ b/Engine/Source/Runtime/Core/Private/UObject/DevObjectVersion.cpp
@@ -306,8 +306,8 @@ const FGuid FUE5LWCRenderingStreamObjectVersion::GUID(0x30D58BE3, 0x95EA4282, 0x
 // Register MediaFramework custom version with Core
 FDevVersionRegistration GRegisterLWCRenderingStreamObjectVersion(FUE5LWCRenderingStreamObjectVersion::GUID, FUE5LWCRenderingStreamObjectVersion::LatestVersion, TEXT("UE5-Dev-LWCRendering"));
 
-const FGuid FDevSystemGuids::GLOBALSHADERMAP_DERIVEDDATA_VER  (0x7912A706, 0x52B8450A, 0x9CD66FA7, 0xEFBBFA0A);
+const FGuid FDevSystemGuids::GLOBALSHADERMAP_DERIVEDDATA_VER(0x7912A706, 0x52B8450A, 0x9CD66FA7, 0xEFBBFA0A);
 const FGuid FDevSystemGuids::MATERIALSHADERMAP_DERIVEDDATA_VER(0x2579AAFE, 0x8F1D4E4F, 0xB04278C6, 0x35917C17);
-const FGuid FDevSystemGuids::NIAGARASHADERMAP_DERIVEDDATA_VER (0x7BBD9913, 0xC1554D20, 0xADAE9F17, 0xB006299E);
-const FGuid FDevSystemGuids::Niagara_LatestScriptCompileVersio(0x6D32B8EE, 0x909FCA7E, 0xA5CE4F17, 0x066A5F25);
-const FGuid FDevSystemGuids::SkeletalMeshDerivedDataVersion   (0x9B5F4544, 0x76D7481C, 0x9AD3F614, 0xA6C07904);
\ No newline at end of file
+const FGuid FDevSystemGuids::NIAGARASHADERMAP_DERIVEDDATA_VER(0x7BBD9913, 0xC1554D20, 0xADAE9F17, 0xB006299E);
+const FGuid FDevSystemGuids::Niagara_LatestScriptCompileVersion(0x6D32B8EE, 0x909FCA7E, 0xA5CE4F17, 0x066A5F25);
+const FGuid FDevSystemGuids::SkeletalMeshDerivedDataVersion(0x9B5F4544, 0x76D7481C, 0x9AD3F614, 0xA6C07904);
\ No newline at end of file
