Add missing #include
#rnx

diff --git a/Engine/Source/Runtime/Core/Public/UObject/UE5MainStreamObjectVersion.h b/Engine/Source/Runtime/Core/Public/UObject/UE5MainStreamObjectVersion.h
--- a/Engine/Source/Runtime/Core/Public/UObject/UE5MainStreamObjectVersion.h
+++ b/Engine/Source/Runtime/Core/Public/UObject/UE5MainStreamObjectVersion.h
@@ -3,6 +3,7 @@
 
 #include "CoreTypes.h"
 #include "Misc/Guid.h"
+#include "UObject/DevObjectVersion.h"
 
 // Custom serialization version for changes made in //UE5/Main stream
 struct CORE_API FUE5MainStreamObjectVersion
