Resolve static analyzer warning "conflicting declarations of 'FAlignedAllocator'" by converting both declarations to "using".

#rnx

diff --git a/Engine/Source/Runtime/Core/Public/Experimental/ConcurrentLinearAllocator.h b/Engine/Source/Runtime/Core/Public/Experimental/ConcurrentLinearAllocator.h
--- a/Engine/Source/Runtime/Core/Public/Experimental/ConcurrentLinearAllocator.h
+++ b/Engine/Source/Runtime/Core/Public/Experimental/ConcurrentLinearAllocator.h
@@ -51,7 +51,7 @@ struct FOsAllocator
 };
 
 #if PLATFORM_USES_ANSI_MALLOC
-struct FAlignedAllocator
+struct FAlignedAllocatorInternal
 {
 	static constexpr bool SupportsAlignment = true;
 
@@ -72,6 +72,8 @@ struct FAlignedAllocator
 		return AnsiFree(Pointer);
 	}
 };
+
+using FAlignedAllocator = FAlignedAllocatorInternal;
 #else // #if PLATFORM_USES_ANSI_MALLOC
 using FAlignedAllocator = FOsAllocator;
 #endif // #else // #if PLATFORM_USES_ANSI_MALLOC
