Fixed static init order bug with new DevSystemGuids.

diff --git a/Engine/Plugins/FX/Niagara/Source/NiagaraCore/Private/NiagaraCustomVersion.cpp b/Engine/Plugins/FX/Niagara/Source/NiagaraCore/Private/NiagaraCustomVersion.cpp
--- a/Engine/Plugins/FX/Niagara/Source/NiagaraCore/Private/NiagaraCustomVersion.cpp
+++ b/Engine/Plugins/FX/Niagara/Source/NiagaraCore/Private/NiagaraCustomVersion.cpp
@@ -10,5 +10,5 @@ FCustomVersionRegistration GRegisterNiagaraCustomVersion(FNiagaraCustomVersion::
 
 FGuid FNiagaraCustomVersion::GetLatestScriptCompileVersion()
 {
-	return FDevSystemGuids::GetSystemGuid(FDevSystemGuids::Niagara_LatestScriptCompileVersion);
+	return FDevSystemGuids::GetSystemGuid(FDevSystemGuids::Get().Niagara_LatestScriptCompileVersion);
 }
diff --git a/Engine/Plugins/FX/Niagara/Source/NiagaraShader/Private/NiagaraShader.cpp b/Engine/Plugins/FX/Niagara/Source/NiagaraShader/Private/NiagaraShader.cpp
--- a/Engine/Plugins/FX/Niagara/Source/NiagaraShader/Private/NiagaraShader.cpp
+++ b/Engine/Plugins/FX/Niagara/Source/NiagaraShader/Private/NiagaraShader.cpp
@@ -619,7 +619,7 @@ void NiagaraShaderMapAppendKeyString(EShaderPlatform Platform, FString& KeyStrin
 /** Creates a string key for the derived data cache given a shader map id. */
 static FString GetNiagaraShaderMapKeyString(const FNiagaraShaderMapId& ShaderMapId, EShaderPlatform Platform)
 {
-	static const FString NIAGARASHADERMAP_DERIVEDDATA_VER = FDevSystemGuids::GetSystemGuid(FDevSystemGuids::NIAGARASHADERMAP_DERIVEDDATA_VER).ToString(EGuidFormats::DigitsWithHyphens);
+	static const FString NIAGARASHADERMAP_DERIVEDDATA_VER = FDevSystemGuids::GetSystemGuid(FDevSystemGuids::Get().NIAGARASHADERMAP_DERIVEDDATA_VER).ToString(EGuidFormats::DigitsWithHyphens);
 
 	FName Format = LegacyShaderPlatformToShaderFormat(Platform);
 	FString ShaderMapKeyString = Format.ToString() + TEXT("_") + FString(FString::FromInt(GetTargetPlatformManagerRef().ShaderFormatVersion(Format))) + TEXT("_");
diff --git a/Engine/Source/Runtime/Core/Private/UObject/DevObjectVersion.cpp b/Engine/Source/Runtime/Core/Private/UObject/DevObjectVersion.cpp
--- a/Engine/Source/Runtime/Core/Private/UObject/DevObjectVersion.cpp
+++ b/Engine/Source/Runtime/Core/Private/UObject/DevObjectVersion.cpp
@@ -99,6 +99,8 @@ void FDevSystemGuidRegistry::RegisterSystemGuids(const TMap<FGuid, FGuid>& Syste
 
 	for (const TPair<FGuid, FGuid>& SystemGuid : SystemGuids)
 	{
+		check(SystemGuid.Key.IsValid());
+		check(SystemGuid.Value.IsValid());
 		RegisterSystemGuid(SystemGuid.Key, SystemGuid.Value);
 	}
 }
@@ -112,6 +114,21 @@ FGuid FDevSystemGuidRegistry::GetSystemGuid(FGuid System)
 	return RegisteredGuid.Guid;
 }
 
+FDevSystemGuids::FDevSystemGuids()
+	: GLOBALSHADERMAP_DERIVEDDATA_VER(0x7912A706, 0x52B8450A, 0x9CD66FA7, 0xEFBBFA0A)
+	, MATERIALSHADERMAP_DERIVEDDATA_VER(0x2579AAFE, 0x8F1D4E4F, 0xB04278C6, 0x35917C17)
+	, NIAGARASHADERMAP_DERIVEDDATA_VER(0x7BBD9913, 0xC1554D20, 0xADAE9F17, 0xB006299E)
+	, Niagara_LatestScriptCompileVersion(0x6D32B8EE, 0x909FCA7E, 0xA5CE4F17, 0x066A5F25)
+	, SkeletalMeshDerivedDataVersion(0x9B5F4544, 0x76D7481C, 0x9AD3F614, 0xA6C07904)
+{
+}
+
+const FDevSystemGuids& FDevSystemGuids::Get()
+{
+	static FDevSystemGuids Guids;
+	return Guids;
+}
+
 FGuid FDevSystemGuids::GetSystemGuid(FGuid System)
 {
 	return GSystemGuidRegistry.GetSystemGuid(System);
@@ -305,9 +322,3 @@ FDevVersionRegistration GRegisterMediaFrameworkObjectVersion(FMediaFrameworkObje
 const FGuid FUE5LWCRenderingStreamObjectVersion::GUID(0x30D58BE3, 0x95EA4282, 0xA6E3B159, 0xD8EBB06A);
 // Register MediaFramework custom version with Core
 FDevVersionRegistration GRegisterLWCRenderingStreamObjectVersion(FUE5LWCRenderingStreamObjectVersion::GUID, FUE5LWCRenderingStreamObjectVersion::LatestVersion, TEXT("UE5-Dev-LWCRendering"));
-
-const FGuid FDevSystemGuids::GLOBALSHADERMAP_DERIVEDDATA_VER(0x7912A706, 0x52B8450A, 0x9CD66FA7, 0xEFBBFA0A);
-const FGuid FDevSystemGuids::MATERIALSHADERMAP_DERIVEDDATA_VER(0x2579AAFE, 0x8F1D4E4F, 0xB04278C6, 0x35917C17);
-const FGuid FDevSystemGuids::NIAGARASHADERMAP_DERIVEDDATA_VER(0x7BBD9913, 0xC1554D20, 0xADAE9F17, 0xB006299E);
-const FGuid FDevSystemGuids::Niagara_LatestScriptCompileVersion(0x6D32B8EE, 0x909FCA7E, 0xA5CE4F17, 0x066A5F25);
-const FGuid FDevSystemGuids::SkeletalMeshDerivedDataVersion(0x9B5F4544, 0x76D7481C, 0x9AD3F614, 0xA6C07904);
\ No newline at end of file
diff --git a/Engine/Source/Runtime/Core/Public/UObject/DevObjectVersion.h b/Engine/Source/Runtime/Core/Public/UObject/DevObjectVersion.h
--- a/Engine/Source/Runtime/Core/Public/UObject/DevObjectVersion.h
+++ b/Engine/Source/Runtime/Core/Public/UObject/DevObjectVersion.h
@@ -14,13 +14,15 @@ public:
 
 struct CORE_API FDevSystemGuids
 {
+	static const FDevSystemGuids& Get();
 	static FGuid GetSystemGuid(FGuid System);
-
-	static const FGuid GLOBALSHADERMAP_DERIVEDDATA_VER;
-	static const FGuid MATERIALSHADERMAP_DERIVEDDATA_VER;
-	static const FGuid NIAGARASHADERMAP_DERIVEDDATA_VER;
-	static const FGuid Niagara_LatestScriptCompileVersion;
-	static const FGuid SkeletalMeshDerivedDataVersion;
+	
+	FDevSystemGuids();
+	const FGuid GLOBALSHADERMAP_DERIVEDDATA_VER;
+	const FGuid MATERIALSHADERMAP_DERIVEDDATA_VER;
+	const FGuid NIAGARASHADERMAP_DERIVEDDATA_VER;
+	const FGuid Niagara_LatestScriptCompileVersion;
+	const FGuid SkeletalMeshDerivedDataVersion;
 };
 
 class CORE_API FDevVersionRegistration :  public FCustomVersionRegistration
diff --git a/Engine/Source/Runtime/Core/Public/UObject/FortniteMainBranchObjectVersion.h b/Engine/Source/Runtime/Core/Public/UObject/FortniteMainBranchObjectVersion.h
--- a/Engine/Source/Runtime/Core/Public/UObject/FortniteMainBranchObjectVersion.h
+++ b/Engine/Source/Runtime/Core/Public/UObject/FortniteMainBranchObjectVersion.h
@@ -207,12 +207,13 @@ struct CORE_API FFortniteMainBranchObjectVersion
 	static TMap<FGuid, FGuid> GetSystemGuids()
 	{
 		TMap<FGuid, FGuid> SystemGuids;
+		const FDevSystemGuids& DevGuids = FDevSystemGuids::Get();
 
-		SystemGuids.Add(FDevSystemGuids::GLOBALSHADERMAP_DERIVEDDATA_VER,    FGuid("8D3A12924CDF4D9F84C0A900E9445CAD"));
-		SystemGuids.Add(FDevSystemGuids::MATERIALSHADERMAP_DERIVEDDATA_VER,  FGuid("1DA85DA5733C4A489A1BE117CDC5ACCC"));
-		SystemGuids.Add(FDevSystemGuids::NIAGARASHADERMAP_DERIVEDDATA_VER,   FGuid("2289D5116CF94BC0AFEAEC541468E645"));
-		SystemGuids.Add(FDevSystemGuids::Niagara_LatestScriptCompileVersion, FGuid("A366AB6E36266F4CBB2DB87D6C07A6A6"));
-		SystemGuids.Add(FDevSystemGuids::SkeletalMeshDerivedDataVersion,     FGuid("ACF593EAAE354FCBB34CF44F5AA1BFD2"));
+		SystemGuids.Add(DevGuids.GLOBALSHADERMAP_DERIVEDDATA_VER,    FGuid("8D3A12924CDF4D9F84C0A900E9445CAD"));
+		SystemGuids.Add(DevGuids.MATERIALSHADERMAP_DERIVEDDATA_VER,  FGuid("1DA85DA5733C4A489A1BE117CDC5ACCC"));
+		SystemGuids.Add(DevGuids.NIAGARASHADERMAP_DERIVEDDATA_VER,   FGuid("2289D5116CF94BC0AFEAEC541468E645"));
+		SystemGuids.Add(DevGuids.Niagara_LatestScriptCompileVersion, FGuid("A366AB6E36266F4CBB2DB87D6C07A6A6"));
+		SystemGuids.Add(DevGuids.SkeletalMeshDerivedDataVersion,     FGuid("ACF593EAAE354FCBB34CF44F5AA1BFD2"));
 
 		return SystemGuids;
 	}
diff --git a/Engine/Source/Runtime/Engine/Private/ShaderCompiler/ShaderCompiler.cpp b/Engine/Source/Runtime/Engine/Private/ShaderCompiler/ShaderCompiler.cpp
--- a/Engine/Source/Runtime/Engine/Private/ShaderCompiler/ShaderCompiler.cpp
+++ b/Engine/Source/Runtime/Engine/Private/ShaderCompiler/ShaderCompiler.cpp
@@ -227,13 +227,13 @@ namespace GlobalShaderCookStats
 
 const FString& GetGlobalShaderMapDDCKey()
 {
-	static FString GlobalShaderMapDDCKey = FDevSystemGuids::GetSystemGuid(FDevSystemGuids::GLOBALSHADERMAP_DERIVEDDATA_VER).ToString();
+	static FString GlobalShaderMapDDCKey = FDevSystemGuids::GetSystemGuid(FDevSystemGuids::Get().GLOBALSHADERMAP_DERIVEDDATA_VER).ToString();
 	return GlobalShaderMapDDCKey;
 }
 
 const FString& GetMaterialShaderMapDDCKey()
 {
-	static FString MaterialShaderMapDDCKey = FDevSystemGuids::GetSystemGuid(FDevSystemGuids::MATERIALSHADERMAP_DERIVEDDATA_VER).ToString();
+	static FString MaterialShaderMapDDCKey = FDevSystemGuids::GetSystemGuid(FDevSystemGuids::Get().MATERIALSHADERMAP_DERIVEDDATA_VER).ToString();
 	return MaterialShaderMapDDCKey;
 }
 
diff --git a/Engine/Source/Runtime/Engine/Private/SkeletalMeshRenderData.cpp b/Engine/Source/Runtime/Engine/Private/SkeletalMeshRenderData.cpp
--- a/Engine/Source/Runtime/Engine/Private/SkeletalMeshRenderData.cpp
+++ b/Engine/Source/Runtime/Engine/Private/SkeletalMeshRenderData.cpp
@@ -185,7 +185,7 @@ static void SerializeLODInfoForDDC(USkeletalMesh* SkeletalMesh, FString& KeySuff
 
 const FString& GetSkeletalMeshDerivedDataVersion()
 {
-	static FString CachedVersionString = FDevSystemGuids::GetSystemGuid(FDevSystemGuids::SkeletalMeshDerivedDataVersion).ToString();
+	static FString CachedVersionString = FDevSystemGuids::GetSystemGuid(FDevSystemGuids::Get().SkeletalMeshDerivedDataVersion).ToString();
 	return CachedVersionString;
 }
 
