Adapt UE5 Release system guids to static init fix
Actually register the ue5 release system guids
#rnx

diff --git a/Engine/Source/Runtime/Core/Private/UObject/DevObjectVersion.cpp b/Engine/Source/Runtime/Core/Private/UObject/DevObjectVersion.cpp
--- a/Engine/Source/Runtime/Core/Private/UObject/DevObjectVersion.cpp
+++ b/Engine/Source/Runtime/Core/Private/UObject/DevObjectVersion.cpp
@@ -303,6 +303,7 @@ FDevVersionRegistration GRegisterUE5MainStreamObjectVersion(FUE5MainStreamObject
 const FGuid FUE5ReleaseStreamObjectVersion::GUID(0xD89B5E42, 0x24BD4D46, 0x8412ACA8, 0xDF641779);
 // Register UE5 release stream custom version with Core
 FDevVersionRegistration GRegisterUE5ReleaseStreamObjectVersion(FUE5ReleaseStreamObjectVersion::GUID, FUE5ReleaseStreamObjectVersion::LatestVersion, TEXT("UE5-Release"));
+FDevSystemGuidRegistration GRegisterUE5ReleaseBranchSystemGuids(FUE5ReleaseStreamObjectVersion::GetSystemGuids());
 
 // Unique UE5 private frosty version id
 const FGuid FUE5PrivateFrostyStreamObjectVersion::GUID(0x59DA5D52, 0x12324948, 0xB8785978, 0x70B8E98B);
diff --git a/Engine/Source/Runtime/Core/Public/UObject/UE5ReleaseStreamObjectVersion.h b/Engine/Source/Runtime/Core/Public/UObject/UE5ReleaseStreamObjectVersion.h
--- a/Engine/Source/Runtime/Core/Public/UObject/UE5ReleaseStreamObjectVersion.h
+++ b/Engine/Source/Runtime/Core/Public/UObject/UE5ReleaseStreamObjectVersion.h
@@ -123,12 +123,13 @@ struct CORE_API FUE5ReleaseStreamObjectVersion
 	static TMap<FGuid, FGuid> GetSystemGuids()
 	{
 		TMap<FGuid, FGuid> SystemGuids;
+		const FDevSystemGuids& DevGuids = FDevSystemGuids::Get();
 
-		SystemGuids.Add(FDevSystemGuids::GLOBALSHADERMAP_DERIVEDDATA_VER, FGuid("D0BF3452816D46908073DFDD4B855AE5"));
-		SystemGuids.Add(FDevSystemGuids::MATERIALSHADERMAP_DERIVEDDATA_VER, FGuid("79090A66F9D94B5BB285D49E5D39468E"));
-		SystemGuids.Add(FDevSystemGuids::NIAGARASHADERMAP_DERIVEDDATA_VER, FGuid("2289D5116CF94BC0AFEAEC541468E645"));
-		SystemGuids.Add(FDevSystemGuids::Niagara_LatestScriptCompileVersion, FGuid("3B3F2E14B150418481F033E5FAD9D13A"));
-		SystemGuids.Add(FDevSystemGuids::SkeletalMeshDerivedDataVersion, FGuid("25C49E579B3142DDA2A8C14037267679"));
+		SystemGuids.Add(DevGuids.GLOBALSHADERMAP_DERIVEDDATA_VER, FGuid("D0BF3452816D46908073DFDD4B855AE5"));
+		SystemGuids.Add(DevGuids.MATERIALSHADERMAP_DERIVEDDATA_VER, FGuid("79090A66F9D94B5BB285D49E5D39468E"));
+		SystemGuids.Add(DevGuids.NIAGARASHADERMAP_DERIVEDDATA_VER, FGuid("2289D5116CF94BC0AFEAEC541468E645"));
+		SystemGuids.Add(DevGuids.Niagara_LatestScriptCompileVersion, FGuid("3B3F2E14B150418481F033E5FAD9D13A"));
+		SystemGuids.Add(DevGuids.SkeletalMeshDerivedDataVersion, FGuid("25C49E579B3142DDA2A8C14037267679"));
 
 		return SystemGuids;
 	}
