Adapt UE5 Main system guids to static init fix
Actually register the ue5 main  system guids
#rnx

diff --git a/Engine/Source/Runtime/Core/Private/UObject/DevObjectVersion.cpp b/Engine/Source/Runtime/Core/Private/UObject/DevObjectVersion.cpp
--- a/Engine/Source/Runtime/Core/Private/UObject/DevObjectVersion.cpp
+++ b/Engine/Source/Runtime/Core/Private/UObject/DevObjectVersion.cpp
@@ -298,6 +298,7 @@ FDevVersionRegistration GRegisterVirtualProductionObjectVersion(FVirtualProducti
 const FGuid FUE5MainStreamObjectVersion::GUID(0x697DD581, 0xE64f41AB, 0xAA4A51EC, 0xBEB7B628);
 // Register UE5 main stream custom version with Core
 FDevVersionRegistration GRegisterUE5MainStreamObjectVersion(FUE5MainStreamObjectVersion::GUID, FUE5MainStreamObjectVersion::LatestVersion, TEXT("UE5-Main"));
+FDevSystemGuidRegistration GRegisterUE5MainBranchSystemGuids(FUE5MainStreamObjectVersion::GetSystemGuids());
 
 // Unique UE5 release version id
 const FGuid FUE5ReleaseStreamObjectVersion::GUID(0xD89B5E42, 0x24BD4D46, 0x8412ACA8, 0xDF641779);
diff --git a/Engine/Source/Runtime/Core/Public/UObject/UE5MainStreamObjectVersion.h b/Engine/Source/Runtime/Core/Public/UObject/UE5MainStreamObjectVersion.h
--- a/Engine/Source/Runtime/Core/Public/UObject/UE5MainStreamObjectVersion.h
+++ b/Engine/Source/Runtime/Core/Public/UObject/UE5MainStreamObjectVersion.h
@@ -205,12 +205,13 @@ struct CORE_API FUE5MainStreamObjectVersion
 	static TMap<FGuid, FGuid> GetSystemGuids()
 	{
 		TMap<FGuid, FGuid> SystemGuids;
+		const FDevSystemGuids& DevGuids = FDevSystemGuids::Get();
 
-		SystemGuids.Add(FDevSystemGuids::GLOBALSHADERMAP_DERIVEDDATA_VER, FGuid("62F5564D1FED4A2D8864DF300EC5AA2F"));
-		SystemGuids.Add(FDevSystemGuids::MATERIALSHADERMAP_DERIVEDDATA_VER, FGuid("65C01D817C9A4EEFAE9E988D41A1F3DD"));
-		SystemGuids.Add(FDevSystemGuids::NIAGARASHADERMAP_DERIVEDDATA_VER, FGuid("8A37C45D24F2423CBE5F8F371DE33575"));
-		SystemGuids.Add(FDevSystemGuids::Niagara_LatestScriptCompileVersion, FGuid("BFA6C45A89CE2B84BB2C444C78C46E76"));
-		SystemGuids.Add(FDevSystemGuids::SkeletalMeshDerivedDataVersion, FGuid("FB7B2FB546DE49B2A2F88B7F69763CEE"));
+		SystemGuids.Add(DevGuids.GLOBALSHADERMAP_DERIVEDDATA_VER, FGuid("62F5564D1FED4A2D8864DF300EC5AA2F"));
+		SystemGuids.Add(DevGuids.MATERIALSHADERMAP_DERIVEDDATA_VER, FGuid("65C01D817C9A4EEFAE9E988D41A1F3DD"));
+		SystemGuids.Add(DevGuids.NIAGARASHADERMAP_DERIVEDDATA_VER, FGuid("8A37C45D24F2423CBE5F8F371DE33575"));
+		SystemGuids.Add(DevGuids.Niagara_LatestScriptCompileVersion, FGuid("BFA6C45A89CE2B84BB2C444C78C46E76"));
+		SystemGuids.Add(DevGuids.SkeletalMeshDerivedDataVersion, FGuid("FB7B2FB546DE49B2A2F88B7F69763CEE"));
 
 		return SystemGuids;
 	}
