Adding DoesAllocatorContainAddress to generic growable allocator.

#rnx

diff --git a/Engine/Source/Runtime/Core/Public/Misc/GenericGrowableAllocator.h b/Engine/Source/Runtime/Core/Public/Misc/GenericGrowableAllocator.h
--- a/Engine/Source/Runtime/Core/Public/Misc/GenericGrowableAllocator.h
+++ b/Engine/Source/Runtime/Core/Public/Misc/GenericGrowableAllocator.h
@@ -121,6 +121,10 @@ public:
 	 */
 	virtual bool DoesChunkContainAllocation(const FGrowableAllocationBase* Allocation) = 0;
 
+	/**
+	 * Queries the implementation if the given address came from this chunk
+	 */
+	virtual bool DoesChunkContainAddress(const void* Address) = 0;
 
 
 	/**
@@ -709,6 +713,23 @@ public:
 		}
 	}
 
+	bool DoesAllocatorContainAddress(const void* Address)
+	{
+		// multi-thread protection
+		FScopeLock ScopeLock(&CriticalSection);
+
+		// loop through the chunks, query each one to see if they contain the address
+		for (int32 ChunkIndex = 0; ChunkIndex < AllocChunks.Num(); ChunkIndex++)
+		{
+			ChunkAllocatorType* Chunk = AllocChunks[ChunkIndex];
+			if (Chunk && Chunk->DoesChunkContainAddress(Address))
+			{
+				return true;
+			}
+		}
+		return false;
+	}
+
 	void ShowAllocationInfo()
 	{
  		FPlatformMisc::LowLevelOutputDebugStringf(TEXT("   Allocator has %d chunks\n"), AllocChunks.Num());
