Allow TAsyncResult<void> by disabling the const ref constructor for non-referenceable result types

diff --git a/Engine/Source/Runtime/Core/Public/Async/AsyncResult.h b/Engine/Source/Runtime/Core/Public/Async/AsyncResult.h
--- a/Engine/Source/Runtime/Core/Public/Async/AsyncResult.h
+++ b/Engine/Source/Runtime/Core/Public/Async/AsyncResult.h
@@ -23,10 +23,14 @@ public:
 
 	/**
 	 * Creates and initializes a new instance with the given synchronous result value.
+	 * This constructor is only available for referenceable result types, to make 'void' a feasible result type
+	 * for this template class.
 	 *
 	 * @param Result The result value to set.
 	 */
-	TAsyncResult(const ResultType& Result)
+	template <typename ArgType,
+			  typename = std::enable_if_t<std::is_same_v<ArgType, ResultType> && std::is_reference_v<const ArgType&>>>
+	TAsyncResult(const ArgType& Result)
 	{
 		TPromise<ResultType> Promise;
 		Promise.SetValue(Result);
