Fix overflows reported by static analysis (UE-141612)

#rnx

diff --git a/Engine/Source/Runtime/Core/Private/Serialization/MemoryImage.cpp b/Engine/Source/Runtime/Core/Private/Serialization/MemoryImage.cpp
--- a/Engine/Source/Runtime/Core/Private/Serialization/MemoryImage.cpp
+++ b/Engine/Source/Runtime/Core/Private/Serialization/MemoryImage.cpp
@@ -1682,6 +1682,7 @@ FMemoryImageSection* FMemoryImageSection::WritePointer(const FTypeLayoutDesc& St
 	FFrozenMemoryImagePtr FrozenPtr;
 	FrozenPtr.Packed = 0; // dummy value
 	// If derived type matches the static type, store INDEX_NONE to indicate the static type may be used
+	checkf(static_cast<uint64>(TypeDependencyIndex + 1) <= (1ULL << FFrozenMemoryImagePtr::TypeIndexBits), TEXT(""));
 	FrozenPtr.SetTypeIndex(bStaticTypeMatchesDerived ? INDEX_NONE : TypeDependencyIndex);
 
 	SectionPointer.Offset = WriteBytes(FrozenPtr);
@@ -1854,7 +1855,7 @@ void FMemoryImage::Flatten(FMemoryImageResult& OutResult, bool bMergeDuplicateSe
 			const int32 RemapSectionIndex = SectionIndexRemap[Pointer.SectionIndex];
 			FFrozenMemoryImagePtr* FrozenPtr = (FFrozenMemoryImagePtr*)(OutResult.Bytes.GetData() + OffsetToPointer);
 			check(FrozenPtr->GetOffsetFromThis() == 0 && !FrozenPtr->IsFrozen());
-			const int64 OffsetFromPointer = (int64)(SectionOffset[RemapSectionIndex] + Pointer.PointerOffset) - OffsetToPointer;
+			const int64 OffsetFromPointer = static_cast<int64>(SectionOffset[RemapSectionIndex]) + static_cast<int64>(Pointer.PointerOffset) - OffsetToPointer;
 			FrozenPtr->SetOffsetFromThis(OffsetFromPointer);
 			FrozenPtr->SetIsFrozen(true);
 			check(FrozenPtr->GetOffsetFromThis() == OffsetFromPointer);
diff --git a/Engine/Source/Runtime/Core/Public/Serialization/MemoryImage.h b/Engine/Source/Runtime/Core/Public/Serialization/MemoryImage.h
--- a/Engine/Source/Runtime/Core/Public/Serialization/MemoryImage.h
+++ b/Engine/Source/Runtime/Core/Public/Serialization/MemoryImage.h
@@ -342,7 +342,8 @@ struct FFrozenMemoryImagePtr
 	void SetTypeIndex(int32 TypeIndex)
 	{
 		static_assert(INDEX_NONE == -1, "TypeIndex cannot store INDEX_NONE when it's not -1");
-		Packed = (Packed & ~TypeIndexMask) | ((static_cast<uint64>(TypeIndex + 1) << TypeIndexShift) & TypeIndexMask);
+		// PVS warns about a possible overflow in TypeIndex + 1. We don't care as we don't expect 2^31 type indices anyway
+		Packed = (Packed & ~TypeIndexMask) | ((static_cast<uint64>(TypeIndex + 1) << TypeIndexShift) & TypeIndexMask); //-V1028
 	}
 };
 
