Add FVector::RotateAngleAxis method taking radian angle as parameter

diff --git a/Engine/Source/Runtime/Core/Public/Math/Vector.h b/Engine/Source/Runtime/Core/Public/Math/Vector.h
--- a/Engine/Source/Runtime/Core/Public/Math/Vector.h
+++ b/Engine/Source/Runtime/Core/Public/Math/Vector.h
@@ -749,11 +749,20 @@ public:
     /**
      * Rotates around Axis (assumes Axis.Size() == 1).
      *
-     * @param Angle Angle to rotate (in degrees).
+     * @param AngleDeg Angle to rotate (in degrees).
      * @param Axis Axis to rotate around.
      * @return Rotated Vector.
      */
     TVector<T> RotateAngleAxis(const T AngleDeg, const TVector<T>& Axis) const;
+    
+    /**
+     * Rotates around Axis (assumes Axis.Size() == 1).
+     *
+     * @param AngleRad Angle to rotate (in radians).
+     * @param Axis Axis to rotate around.
+     * @return Rotated Vector.
+     */
+    TVector<T> RotateAngleAxisRad(const T AngleRad, const TVector<T>& Axis) const;
 
     /**
      * Returns the cosine of the angle between this vector and another projected onto the XY plane (no Z).
@@ -1212,9 +1221,15 @@ FORCEINLINE TVector<T>::TVector(const TVector2<T> V, T InZ)
 
 template<typename T>
 inline TVector<T> TVector<T>::RotateAngleAxis(const T AngleDeg, const TVector<T>& Axis) const
+{
+	return RotateAngleAxisRad(FMath::DegreesToRadians(AngleDeg), Axis);
+}
+
+template<typename T>
+inline TVector<T> TVector<T>::RotateAngleAxisRad(const T AngleRad, const TVector<T>& Axis) const
 {
     T S, C;
-    FMath::SinCos(&S, &C, FMath::DegreesToRadians(AngleDeg));
+    FMath::SinCos(&S, &C, AngleRad);
 
     const T XX	= Axis.X * Axis.X;
     const T YY	= Axis.Y * Axis.Y;
