Change ECVF_Scalability on console variables and commands to imply ECVF_Preview

diff --git a/Engine/Source/Runtime/Core/Private/HAL/ConsoleManager.cpp b/Engine/Source/Runtime/Core/Private/HAL/ConsoleManager.cpp
--- a/Engine/Source/Runtime/Core/Private/HAL/ConsoleManager.cpp
+++ b/Engine/Source/Runtime/Core/Private/HAL/ConsoleManager.cpp
@@ -81,6 +81,16 @@ public:
 		:Flags(InFlags), bWarnedAboutThreadSafety(false)
 	{
 		SetHelp(InHelp);
+
+		ApplyPreviewIfScalability();
+	}
+
+	void ApplyPreviewIfScalability()
+	{
+		if (((uint32)Flags & (uint32)ECVF_Scalability) != 0)
+		{
+			Flags = (EConsoleVariableFlags)((uint32)Flags | (uint32)ECVF_Preview);
+		}
 	}
 
 	// interface IConsoleVariable -----------------------------------
@@ -105,6 +115,7 @@ public:
 	virtual void SetFlags(const EConsoleVariableFlags Value)
 	{
 		Flags = Value;
+		ApplyPreviewIfScalability();
 	}
 
 	virtual class IConsoleVariable* AsVariable()
@@ -241,6 +252,16 @@ public:
 	{
 		check(InHelp);
 		//check(*Help != 0); for now disabled as there callstack when we crash early during engine init
+
+		ApplyPreviewIfScalability();
+	}
+
+	void ApplyPreviewIfScalability()
+	{
+		if (((uint32)Flags & (uint32)ECVF_Scalability) != 0)
+		{
+			Flags = (EConsoleVariableFlags)((uint32)Flags | (uint32)ECVF_Preview);
+		}
 	}
 
 	// interface IConsoleVariable -----------------------------------
@@ -263,6 +284,7 @@ public:
 	virtual void SetFlags(const EConsoleVariableFlags Value)
 	{
 		Flags = Value;
+		ApplyPreviewIfScalability();
 	}
 
 	virtual struct IConsoleCommand* AsCommand()
