Fix static analysis issues by moving checks to static asserts.

diff --git a/Engine/Source/Runtime/Core/Private/HAL/MallocBinned2.cpp b/Engine/Source/Runtime/Core/Private/HAL/MallocBinned2.cpp
--- a/Engine/Source/Runtime/Core/Private/HAL/MallocBinned2.cpp
+++ b/Engine/Source/Runtime/Core/Private/HAL/MallocBinned2.cpp
@@ -736,9 +736,10 @@ FMallocBinned2::FMallocBinned2()
 	checkf(FMath::IsPowerOfTwo(PageSize), TEXT("OS page size must be a power of two"));
 	checkf(FMath::IsPowerOfTwo(Constants.AddressLimit), TEXT("OS address limit must be a power of two"));
 	checkf(Constants.AddressLimit > PageSize, TEXT("OS address limit must be greater than the page size")); // Check to catch 32 bit overflow in AddressLimit
-	checkf(SmallBlockSizes[BINNED2_SMALL_POOL_COUNT - 1] == BINNED2_MAX_SMALL_POOL_SIZE, TEXT("BINNED2_MAX_SMALL_POOL_SIZE must equal the smallest block size"));
 	checkf(PageSize % BINNED2_LARGE_ALLOC == 0, TEXT("OS page size must be a multiple of BINNED2_LARGE_ALLOC"));
-	checkf(sizeof(FMallocBinned2::FFreeBlock) <= SmallBlockSizes[0], TEXT("Pool header must be able to fit into the smallest block"));
+
+	static_assert(SmallBlockSizes[BINNED2_SMALL_POOL_COUNT - 1] == BINNED2_MAX_SMALL_POOL_SIZE, "BINNED2_MAX_SMALL_POOL_SIZE must equal the smallest block size");
+	static_assert(sizeof(FMallocBinned2::FFreeBlock) <= SmallBlockSizes[0], "Pool header must be able to fit into the smallest block");
 	static_assert(UE_ARRAY_COUNT(SmallBlockSizes) == BINNED2_SMALL_POOL_COUNT, "Small block size array size must match BINNED2_SMALL_POOL_COUNT");
 	static_assert(UE_ARRAY_COUNT(SmallBlockSizes) <= 256, "Small block size array size must fit in a byte");
 	static_assert(sizeof(FFreeBlock) <= BINNED2_MINIMUM_ALIGNMENT, "Free block struct must be small enough to fit into a block.");
