Add ::DoesSectionExist to FConfigFile to match the same method in FConfigCacheIni.

#rnx

diff --git a/Engine/Source/Runtime/Core/Private/Misc/ConfigCacheIni.cpp b/Engine/Source/Runtime/Core/Private/Misc/ConfigCacheIni.cpp
--- a/Engine/Source/Runtime/Core/Private/Misc/ConfigCacheIni.cpp
+++ b/Engine/Source/Runtime/Core/Private/Misc/ConfigCacheIni.cpp
@@ -2054,7 +2054,10 @@ int32 FConfigFile::GetArray(const TCHAR* Section, const TCHAR* Key, TArray<FStri
 	return Value.Num();
 }
 
-
+bool FConfigFile::DoesSectionExist(const TCHAR* Section) const
+{
+	return Find(Section) != nullptr;
+}
 
 void FConfigFile::SetString( const TCHAR* Section, const TCHAR* Key, const TCHAR* Value )
 {
diff --git a/Engine/Source/Runtime/Core/Public/Misc/ConfigCacheIni.h b/Engine/Source/Runtime/Core/Public/Misc/ConfigCacheIni.h
--- a/Engine/Source/Runtime/Core/Public/Misc/ConfigCacheIni.h
+++ b/Engine/Source/Runtime/Core/Public/Misc/ConfigCacheIni.h
@@ -530,6 +530,8 @@ public:
 		return GetArray(Section, Key, Value);
 	}
 
+	CORE_API bool DoesSectionExist(const TCHAR* Section) const;
+
 	CORE_API void SetString(const TCHAR* Section, const TCHAR* Key, const TCHAR* Value);
 	CORE_API void SetText(const TCHAR* Section, const TCHAR* Key, const FText& Value);
 	CORE_API void SetFloat(const TCHAR* Section, const TCHAR* Key, float Value);
