Add a method to the virtualization system to allow it to be shutdown manually.

#rnx

diff --git a/Engine/Source/Runtime/Core/Private/Virtualization/VirtualizationSystem.cpp b/Engine/Source/Runtime/Core/Private/Virtualization/VirtualizationSystem.cpp
--- a/Engine/Source/Runtime/Core/Private/Virtualization/VirtualizationSystem.cpp
+++ b/Engine/Source/Runtime/Core/Private/Virtualization/VirtualizationSystem.cpp
@@ -102,6 +102,12 @@ void Initialize()
 	GVirtualizationSystem = MakeUnique<FNullVirtualizationSystem>();	
 }
 
+void Shutdown()
+{
+	GVirtualizationSystem.Reset();
+	UE_LOG(LogVirtualization, Log, TEXT("UE::Virtualization was shutdown"));
+}
+
 IVirtualizationSystem& IVirtualizationSystem::Get()
 {
 	// For now allow Initialize to be called directly if it was not called explicitly.
diff --git a/Engine/Source/Runtime/Core/Public/Virtualization/VirtualizationSystem.h b/Engine/Source/Runtime/Core/Public/Virtualization/VirtualizationSystem.h
--- a/Engine/Source/Runtime/Core/Public/Virtualization/VirtualizationSystem.h
+++ b/Engine/Source/Runtime/Core/Public/Virtualization/VirtualizationSystem.h
@@ -119,9 +119,15 @@ struct FPushRequest
  */
 CORE_API void Initialize();
 
+/**
+ * Shutdowns the global IVirtualizationSystem if it exists. 
+ * Calling this is optional as the system will shut itself down along with the rest of the engine.
+ */
+CORE_API void Shutdown();
+
 /** 
  * The base interface for the virtualization system. An Epic maintained version can be found in the Virtualization module.
- To implement your own, simply derived from this interface and then use the
+ * To implement your own, simply derived from this interface and then use the
  * UE_REGISTER_VIRTUALIZATION_SYSTEM macro in the cpp to register it as an option. 
  * You can then set the config file option [Core.ContentVirtualization]SystemName=FooBar, where FooBar should be the 
  * SystemName parameter you used when registering with the macro.
