Hopeful static analysis fix: "warning: 'NumBits' operand is less than the 'NumBitsPerWord' operand, the result is always equal to the left operand."

diff --git a/Engine/Source/Runtime/Core/Public/Containers/StaticBitArray.h b/Engine/Source/Runtime/Core/Public/Containers/StaticBitArray.h
--- a/Engine/Source/Runtime/Core/Public/Containers/StaticBitArray.h
+++ b/Engine/Source/Runtime/Core/Public/Containers/StaticBitArray.h
@@ -400,7 +400,7 @@ private:
 	 */
 	void Trim_()
 	{
-		if(NumBits % NumBitsPerWord != 0)
+		if constexpr (NumBits % NumBitsPerWord != 0)
 		{
 			Words[NumWords-1] &= (WordType(1) << (NumBits % NumBitsPerWord)) - 1;
 		}
