Create log files in /Library/Caches for TVOS, which cannot write to /Documents

diff --git a/Engine/Source/Runtime/Core/Private/IOS/IOSPlatformFile.cpp b/Engine/Source/Runtime/Core/Private/IOS/IOSPlatformFile.cpp
--- a/Engine/Source/Runtime/Core/Private/IOS/IOSPlatformFile.cpp
+++ b/Engine/Source/Runtime/Core/Private/IOS/IOSPlatformFile.cpp
@@ -1078,10 +1078,17 @@ FString FIOSPlatformFile::ConvertToIOSPath(const FString& Filename, bool bForWri
 
 	if(bForWrite)
 	{
+#if PLATFORM_TVOS
+		// tvOS cannot write to the Documents directory. All files must be written to Library/Caches
+		static FString PublicWritePathBase = FString([NSSearchPathForDirectoriesInDomains(NSCachesDirectory, NSUserDomainMask, YES) objectAtIndex:0]) + TEXT("/");
+
+		return PublicWritePathBase + Result;
+#else
 		static FString PublicWritePathBase = FString([NSSearchPathForDirectoriesInDomains(NSDocumentDirectory, NSUserDomainMask, YES) objectAtIndex:0]) + TEXT("/");
 		static FString PrivateWritePathBase = FString([NSSearchPathForDirectoriesInDomains(NSLibraryDirectory, NSUserDomainMask, YES) objectAtIndex:0]) + TEXT("/");
 		
 		return (bIsPublicWrite ? PublicWritePathBase : PrivateWritePathBase) + Result;
+#endif
 	}
 	else
 	{
