Add a lower case digits mode for FGuid::ToString

diff --git a/Engine/Source/Runtime/Core/Private/Misc/Guid.cpp b/Engine/Source/Runtime/Core/Private/Misc/Guid.cpp
--- a/Engine/Source/Runtime/Core/Private/Misc/Guid.cpp
+++ b/Engine/Source/Runtime/Core/Private/Misc/Guid.cpp
@@ -114,6 +114,9 @@ FString FGuid::ToString(EGuidFormats Format) const
 		return Result.Reverse();
 	}
 
+	case EGuidFormats::DigitsLower:
+		return FString::Printf(TEXT("%08x%08x%08x%08x"), A, B, C, D);
+
 	default:
 		return FString::Printf(TEXT("%08X%08X%08X%08X"), A, B, C, D);
 	}
@@ -188,6 +191,10 @@ void FGuid::AppendString(FAnsiStringBuilderBase& Builder, EGuidFormats Format) c
 		return;
 	}
 
+	case EGuidFormats::DigitsLower:
+		Builder.Appendf("%08x%08x%08x%08x", A, B, C, D);
+		return;
+
 	default:
 		Builder.Appendf("%08X%08X%08X%08X", A, B, C, D);
 		return;
@@ -337,7 +344,7 @@ bool FGuid::ParseExact(const FString& GuidString, EGuidFormats Format, FGuid& Ou
 
 	NormalizedGuidString.Empty(32);
 
-	if (Format == EGuidFormats::Digits)
+	if (Format == EGuidFormats::Digits || Format == EGuidFormats::DigitsLower)
 	{
 		NormalizedGuidString = GuidString;
 	}
diff --git a/Engine/Source/Runtime/Core/Public/Misc/Guid.h b/Engine/Source/Runtime/Core/Public/Misc/Guid.h
--- a/Engine/Source/Runtime/Core/Public/Misc/Guid.h
+++ b/Engine/Source/Runtime/Core/Public/Misc/Guid.h
@@ -28,6 +28,13 @@ enum class EGuidFormats
 	 */
 	Digits,
 
+	/**
+	 * 32 digits in lowercase
+	 *
+	 * For example: "0123abc456def789abcd123ef4a5b6c7"
+	 */
+	 DigitsLower,
+
 	/**
 	 * 32 digits separated by hyphens.
 	 *
