CIS: Fix warning regarding potentially dereferencing a nullptr

#rnx

- The warning is correct, if we failed to find a valid config file it would be a nullptr and would be incorrectly accessed when setting up the fallback FNullVirtualizationSystem.
- We now log an error if no config file is provided by the caller and we are unable to find one to use.
- When setting up the fallback FNullVirtualizationSystem we pass a reference to an empty FConfigFile on the stack instead since we know that it is not used. This avoids any potential invalid access.

diff --git a/Engine/Source/Runtime/Core/Private/Virtualization/VirtualizationSystem.cpp b/Engine/Source/Runtime/Core/Private/Virtualization/VirtualizationSystem.cpp
--- a/Engine/Source/Runtime/Core/Private/Virtualization/VirtualizationSystem.cpp
+++ b/Engine/Source/Runtime/Core/Private/Virtualization/VirtualizationSystem.cpp
@@ -107,33 +107,37 @@ void Initialize(FConfigFile* ConfigFile)
 			SystemName = FName(RawSystemName);
 			UE_LOG(LogVirtualization, Display, TEXT("VirtualizationSystem name found in ini file: %s"), *RawSystemName);
 		}
-	}
-	
-	if (!SystemName.IsNone())
-	{
-		Private::IVirtualizationSystemFactory* SystemFactory = FindFactory(SystemName);
-		if (SystemFactory != nullptr)
-		{
-			GVirtualizationSystem = SystemFactory->Create();
-			check(GVirtualizationSystem.IsValid()); // It is assumed that create will always return a valid pointer
 
-			if (!GVirtualizationSystem->Initialize(*ConfigFile))
+		if (!SystemName.IsNone())
+		{
+			Private::IVirtualizationSystemFactory* SystemFactory = FindFactory(SystemName);
+			if (SystemFactory != nullptr)
 			{
-				UE_LOG(LogVirtualization, Error, TEXT("Initialization of the virtualization system '%s' failed, falling back to the default implementation"), *SystemName.ToString());
-				GVirtualizationSystem.Reset();
+				GVirtualizationSystem = SystemFactory->Create();
+				check(GVirtualizationSystem.IsValid()); // It is assumed that create will always return a valid pointer
+
+				if (!GVirtualizationSystem->Initialize(*ConfigFile))
+				{
+					UE_LOG(LogVirtualization, Error, TEXT("Initialization of the virtualization system '%s' failed, falling back to the default implementation"), *SystemName.ToString());
+					GVirtualizationSystem.Reset();
+				}
+			}
+			else
+			{
+				UE_LOG(LogVirtualization, Error, TEXT("Unable to find factory to create the virtualization system: %s"), *SystemName.ToString());
 			}
 		}
-		else
-		{
-			UE_LOG(LogVirtualization, Error, TEXT("Unable to find factory to create the virtualization system: %s"), *SystemName.ToString());
-		}
+	}
+	else
+	{
+		UE_LOG(LogVirtualization, Error, TEXT("Unable to find a valid engine config file when trying to create the virtualization system"));
 	}
 
 	if (!GVirtualizationSystem.IsValid())
 	{
 		// We found no system to create so we will use the fallback Null system
 		GVirtualizationSystem = MakeUnique<FNullVirtualizationSystem>();
-		GVirtualizationSystem->Initialize(*ConfigFile); // We know this method does nothing but we call it anyway for sake of form
+		GVirtualizationSystem->Initialize(FConfigFile()); // We know this method does nothing but we call it anyway for sake of form
 	}
 }
 
