The virtualization system may now take the project name as a parameter when it is initialized rather than trying to find it itself.

#rnx

- To make it easier to extend the number of parameters when initializing the virtualization system, the function has been changed to accept a single struct, FInitParams, which will contains all potential parameters.
- Calling the version of UE::Virtualization::Initialize without any parameters will fallback to using the default values.
- The virtualization manager now passes the provided project name onto each backend that it creates.
- The source control backend now stored the provided project name and uses that when creating the submit description for new payloads rather than polling FApp for the current project name.
-- This is required for the stand alone virtualization application which will not have a specific project set.

diff --git a/Engine/Source/Developer/Virtualization/Private/VirtualizationDDCBackend.cpp b/Engine/Source/Developer/Virtualization/Private/VirtualizationDDCBackend.cpp
--- a/Engine/Source/Developer/Virtualization/Private/VirtualizationDDCBackend.cpp
+++ b/Engine/Source/Developer/Virtualization/Private/VirtualizationDDCBackend.cpp
@@ -17,7 +17,7 @@ static UE::DerivedData::FValueId ToDerivedDataValueId(const FIoHash& Id)
 	return UE::DerivedData::FValueId::FromHash(Id);
 }
 
-FDDCBackend::FDDCBackend(FStringView ConfigName, FStringView InDebugName)
+FDDCBackend::FDDCBackend(FStringView ProjectName, FStringView ConfigName, FStringView InDebugName)
 : IVirtualizationBackend(ConfigName, InDebugName, EOperations::Both)
 , BucketName(TEXT("BulkData"))
 , TransferPolicy(UE::DerivedData::ECachePolicy::None)
diff --git a/Engine/Source/Developer/Virtualization/Private/VirtualizationDDCBackend.h b/Engine/Source/Developer/Virtualization/Private/VirtualizationDDCBackend.h
--- a/Engine/Source/Developer/Virtualization/Private/VirtualizationDDCBackend.h
+++ b/Engine/Source/Developer/Virtualization/Private/VirtualizationDDCBackend.h
@@ -28,7 +28,7 @@ namespace UE::Virtualization
 class FDDCBackend final : public IVirtualizationBackend
 {
 public:
-	explicit FDDCBackend(FStringView ConfigName, FStringView InDebugName);
+	explicit FDDCBackend(FStringView ProjectName, FStringView ConfigName, FStringView InDebugName);
 	virtual ~FDDCBackend() = default;
 
 private:
diff --git a/Engine/Source/Developer/Virtualization/Private/VirtualizationFileBackend.cpp b/Engine/Source/Developer/Virtualization/Private/VirtualizationFileBackend.cpp
--- a/Engine/Source/Developer/Virtualization/Private/VirtualizationFileBackend.cpp
+++ b/Engine/Source/Developer/Virtualization/Private/VirtualizationFileBackend.cpp
@@ -11,7 +11,7 @@
 namespace UE::Virtualization
 {
 
-FFileSystemBackend::FFileSystemBackend(FStringView ConfigName, FStringView DebugName)
+FFileSystemBackend::FFileSystemBackend(FStringView ProjectName, FStringView ConfigName, FStringView DebugName)
 	: IVirtualizationBackend(ConfigName, DebugName, EOperations::Both)
 {
 }
diff --git a/Engine/Source/Developer/Virtualization/Private/VirtualizationFileBackend.h b/Engine/Source/Developer/Virtualization/Private/VirtualizationFileBackend.h
--- a/Engine/Source/Developer/Virtualization/Private/VirtualizationFileBackend.h
+++ b/Engine/Source/Developer/Virtualization/Private/VirtualizationFileBackend.h
@@ -31,7 +31,7 @@ namespace UE::Virtualization
 class FFileSystemBackend final : public IVirtualizationBackend
 {
 public:
-	explicit FFileSystemBackend(FStringView ConfigName, FStringView DebugName);
+	explicit FFileSystemBackend(FStringView ProjectName, FStringView ConfigName, FStringView DebugName);
 	virtual ~FFileSystemBackend() = default;
 
 private:
diff --git a/Engine/Source/Developer/Virtualization/Private/VirtualizationHttpBackend.cpp b/Engine/Source/Developer/Virtualization/Private/VirtualizationHttpBackend.cpp
--- a/Engine/Source/Developer/Virtualization/Private/VirtualizationHttpBackend.cpp
+++ b/Engine/Source/Developer/Virtualization/Private/VirtualizationHttpBackend.cpp
@@ -1293,7 +1293,7 @@ private:
 	FString SiteIdentifier;
 };
 
-FHttpBackend::FHttpBackend(FStringView ConfigName, FStringView InDebugName)
+FHttpBackend::FHttpBackend(FStringView ProjectName, FStringView ConfigName, FStringView InDebugName)
 	: IVirtualizationBackend(ConfigName, InDebugName, EOperations::Both)
 	, Namespace(TEXT("mirage"))
 	, Bucket(TEXT("default"))
diff --git a/Engine/Source/Developer/Virtualization/Private/VirtualizationHttpBackend.h b/Engine/Source/Developer/Virtualization/Private/VirtualizationHttpBackend.h
--- a/Engine/Source/Developer/Virtualization/Private/VirtualizationHttpBackend.h
+++ b/Engine/Source/Developer/Virtualization/Private/VirtualizationHttpBackend.h
@@ -32,7 +32,7 @@ namespace UE::Virtualization
 class FHttpBackend final : public IVirtualizationBackend
 {
 public:
-	explicit FHttpBackend(FStringView ConfigName, FStringView DebugName);
+	explicit FHttpBackend(FStringView ProjectName, FStringView ConfigName, FStringView DebugName);
 	virtual ~FHttpBackend() = default;
 
 private:
diff --git a/Engine/Source/Developer/Virtualization/Private/VirtualizationManager.cpp b/Engine/Source/Developer/Virtualization/Private/VirtualizationManager.cpp
--- a/Engine/Source/Developer/Virtualization/Private/VirtualizationManager.cpp
+++ b/Engine/Source/Developer/Virtualization/Private/VirtualizationManager.cpp
@@ -12,10 +12,11 @@
 #include "Misc/Parse.h"
 #include "Misc/Paths.h"
 #include "Misc/ScopeLock.h"
+#include "PackageSubmissionChecks.h"
 #include "ProfilingDebugging/CookStats.h"
-
 #include "VirtualizationFilterSettings.h"
 
+
 namespace UE::Virtualization
 {
 UE_REGISTER_VIRTUALIZATION_SYSTEM(UE::Virtualization::FVirtualizationManager, Default);
@@ -283,17 +284,19 @@ FVirtualizationManager::~FVirtualizationManager()
 	UE_LOG(LogVirtualization, Log, TEXT("Virtualization manager destroyed"));
 }
 
-bool FVirtualizationManager::Initialize(const FConfigFile& ConfigFile)
+bool FVirtualizationManager::Initialize(const FInitParams& InitParams)
 {
 	TRACE_CPUPROFILER_EVENT_SCOPE(FVirtualizationManager::Initialize);
 
-	ApplySettingsFromConfigFiles(ConfigFile);
-	ApplyDebugSettingsFromConfigFiles(ConfigFile);
+	ProjectName = InitParams.ProjectName;
+
+	ApplySettingsFromConfigFiles(InitParams.ConfigFile);
+	ApplyDebugSettingsFromConfigFiles(InitParams.ConfigFile);
 	
 	ApplySettingsFromCmdline();
 	ApplyDebugSettingsFromFromCmdline();
 
-	MountBackends(ConfigFile);
+	MountBackends(InitParams.ConfigFile);
 
 	return true;
 }
@@ -852,7 +855,7 @@ bool FVirtualizationManager::CreateBackend(const FConfigFile& ConfigFile, const
 		if (FactoryPtr != nullptr && *FactoryPtr != nullptr)
 		{
 			IVirtualizationBackendFactory* Factory = *FactoryPtr;
-			TUniquePtr<IVirtualizationBackend> Backend = Factory->CreateInstance(ConfigEntryName);
+			TUniquePtr<IVirtualizationBackend> Backend = Factory->CreateInstance(ProjectName, ConfigEntryName);
 
 			if (Backend == nullptr)
 			{
diff --git a/Engine/Source/Developer/Virtualization/Private/VirtualizationManager.h b/Engine/Source/Developer/Virtualization/Private/VirtualizationManager.h
--- a/Engine/Source/Developer/Virtualization/Private/VirtualizationManager.h
+++ b/Engine/Source/Developer/Virtualization/Private/VirtualizationManager.h
@@ -105,7 +105,7 @@ public:
 private:
 	/* IVirtualizationSystem implementation */
 
-	virtual bool Initialize(const FConfigFile& ConfigFile) override;
+	virtual bool Initialize(const FInitParams& InitParams) override;
 
 	virtual bool IsEnabled() const override;
 	virtual bool IsPushingEnabled(EStorageType StorageType) const override;
@@ -205,6 +205,9 @@ private:
 	 */
 	bool bValidateAfterPushOperation;
 
+	/** The name of the current project */
+	FString ProjectName;
+
 	/** Array of backend names that should have their pull operation disabled */
 	TArray<FString> BackendsToDisablePulls;
 
diff --git a/Engine/Source/Developer/Virtualization/Private/VirtualizationSourceControlBackend.cpp b/Engine/Source/Developer/Virtualization/Private/VirtualizationSourceControlBackend.cpp
--- a/Engine/Source/Developer/Virtualization/Private/VirtualizationSourceControlBackend.cpp
+++ b/Engine/Source/Developer/Virtualization/Private/VirtualizationSourceControlBackend.cpp
@@ -8,7 +8,6 @@
 #include "ISourceControlModule.h"
 #include "ISourceControlProvider.h"
 #include "Logging/MessageLog.h"
-#include "Misc/App.h"
 #include "Misc/FileHelper.h"
 #include "Misc/Parse.h"
 #include "Misc/PathViews.h"
@@ -58,11 +57,11 @@ namespace UE::Virtualization
 }
 
 /** Builds a changelist description to be used when submitting a payload to source control */
-static void CreateDescription(const TArray<const FPushRequest*>& FileRequests, TStringBuilder<512>& OutDescription)
+static void CreateDescription(const FString& ProjectName, const TArray<const FPushRequest*>& FileRequests, TStringBuilder<512>& OutDescription)
 {
 	// TODO: Maybe make writing out the project name an option or allow for a codename to be set via ini file?
 	OutDescription << TEXT("Submitted for project: ");
-	OutDescription << FApp::GetProjectName();
+	OutDescription << ProjectName;
 
 	bool bInitialNewline = false;
 
@@ -81,8 +80,9 @@ static void CreateDescription(const TArray<const FPushRequest*>& FileRequests, T
 	}
 }
 
-FSourceControlBackend::FSourceControlBackend(FStringView ConfigName, FStringView InDebugName)
+FSourceControlBackend::FSourceControlBackend(FStringView ProjectName, FStringView ConfigName, FStringView InDebugName)
 	: IVirtualizationBackend(ConfigName, InDebugName, EOperations::Both)
+	, ProjectName(ProjectName)
 {
 }
 
@@ -474,7 +474,7 @@ bool FSourceControlBackend::PushData(TArrayView<FPushRequest> Requests)
 		TSharedRef<FCheckIn, ESPMode::ThreadSafe> CheckInOperation = ISourceControlOperation::Create<FCheckIn>();
 
 		TStringBuilder<512> Description;
-		CreateDescription(FileRequests, Description);
+		CreateDescription(ProjectName, FileRequests, Description);
 
 		CheckInOperation->SetDescription(FText::FromString(Description.ToString()));
 
diff --git a/Engine/Source/Developer/Virtualization/Private/VirtualizationSourceControlBackend.h b/Engine/Source/Developer/Virtualization/Private/VirtualizationSourceControlBackend.h
--- a/Engine/Source/Developer/Virtualization/Private/VirtualizationSourceControlBackend.h
+++ b/Engine/Source/Developer/Virtualization/Private/VirtualizationSourceControlBackend.h
@@ -6,6 +6,8 @@
 
 #include "Containers/StringView.h"
 
+class ISourceControlProvider;
+
 namespace UE::Virtualization
 {
 
@@ -39,7 +41,7 @@ namespace UE::Virtualization
 class FSourceControlBackend final : public IVirtualizationBackend
 {
 public:
-	explicit FSourceControlBackend(FStringView ConfigName, FStringView InDebugName);
+	explicit FSourceControlBackend(FStringView ProjectName, FStringView ConfigName, FStringView InDebugName);
 	virtual ~FSourceControlBackend() = default;
 	
 private:
@@ -64,7 +66,10 @@ private:
 
 	/** Will display a FMessage notification to the user on the next valid engine tick to try and keep them aware of connection failures */
 	void OnConnectionError();
-	
+
+	/** The name of the current project */
+	FString ProjectName;
+
 	/** The root where the virtualized payloads are stored in source control */
 	FString DepotRoot;
 
diff --git a/Engine/Source/Developer/Virtualization/Public/IVirtualizationBackend.h b/Engine/Source/Developer/Virtualization/Public/IVirtualizationBackend.h
--- a/Engine/Source/Developer/Virtualization/Public/IVirtualizationBackend.h
+++ b/Engine/Source/Developer/Virtualization/Public/IVirtualizationBackend.h
@@ -227,10 +227,11 @@ public:
 	/** 
 	 * Creates a new backend instance.
 	 * 
+	 * @param ProjectName	The name of the current project
 	 * @param ConfigName	The name given to the back end in the config ini file
 	 * @return A new backend instance
 	 */
-	virtual TUniquePtr<IVirtualizationBackend> CreateInstance(FStringView ConfigName) = 0;
+	virtual TUniquePtr<IVirtualizationBackend> CreateInstance(FStringView ProjectName, FStringView ConfigName) = 0;
 
 	/** Returns the name used to identify the type in config ini files */
 	virtual FName GetName() = 0;
@@ -252,7 +253,10 @@ public:
 		BackendClass##Factory() { IModularFeatures::Get().RegisterModularFeature(FName("VirtualizationBackendFactory"), this); }\
 		virtual ~BackendClass##Factory() { IModularFeatures::Get().UnregisterModularFeature(FName("VirtualizationBackendFactory"), this); } \
 	private: \
-		virtual TUniquePtr<IVirtualizationBackend> CreateInstance(FStringView ConfigName) override { return MakeUnique<BackendClass>(ConfigName, WriteToString<256>(#ConfigName, TEXT(" - "), ConfigName).ToString()); } \
+		virtual TUniquePtr<IVirtualizationBackend> CreateInstance(FStringView ProjectName, FStringView ConfigName) override \
+		{ \
+			return MakeUnique<BackendClass>(ProjectName, ConfigName, WriteToString<256>(#ConfigName, TEXT(" - "), ConfigName).ToString()); \
+		} \
 		virtual FName GetName() override { return FName(#ConfigName); } \
 	}; \
 	static BackendClass##Factory BackendClass##Factory##Instance;
diff --git a/Engine/Source/Runtime/Core/Private/Virtualization/VirtualizationSystem.cpp b/Engine/Source/Runtime/Core/Private/Virtualization/VirtualizationSystem.cpp
--- a/Engine/Source/Runtime/Core/Private/Virtualization/VirtualizationSystem.cpp
+++ b/Engine/Source/Runtime/Core/Private/Virtualization/VirtualizationSystem.cpp
@@ -3,6 +3,7 @@
 #include "Virtualization/VirtualizationSystem.h"
 
 #include "CoreGlobals.h"
+#include "Misc/App.h"
 #include "Misc/ConfigCacheIni.h"
 
 namespace UE::Virtualization
@@ -19,7 +20,7 @@ public:
 
 	virtual ~FNullVirtualizationSystem() = default;
 
-	virtual bool Initialize(const FConfigFile& ConfigFile) override
+	virtual bool Initialize(const FInitParams& InitParams) override
 	{
 		return true;
 	}
@@ -90,54 +91,66 @@ Private::IVirtualizationSystemFactory* FindFactory(FName SystemName)
 	return nullptr;
 }
 
-void Initialize(FConfigFile* ConfigFile)
+void Initialize()
+{
+	FApp::GetProjectName();
+
+	const FConfigFile* ConfigFile = GConfig->Find(GEngineIni);
+
+	if (ConfigFile != nullptr)
+	{
+		FInitParams InitParams(FApp::GetProjectName() , *ConfigFile);
+		Initialize(InitParams);
+	}
+	else
+	{
+		UE_LOG(LogVirtualization, Error, TEXT("Unable to find a valid engine config file when trying to create the virtualization system"));
+
+		GVirtualizationSystem = MakeUnique<FNullVirtualizationSystem>();
+
+		FConfigFile EmptyConfigFile;
+		FInitParams DummyParams(TEXT(""), EmptyConfigFile);
+
+		GVirtualizationSystem->Initialize(DummyParams);
+	}	
+}
+
+void Initialize(const FInitParams& InitParams)
 {
 	FName SystemName;
 
-	if (ConfigFile == nullptr)
+	FString RawSystemName;
+	if (InitParams.ConfigFile.GetString(TEXT("Core.ContentVirtualization"), TEXT("SystemName"), RawSystemName))
 	{
-		ConfigFile = GConfig->Find(GEngineIni);
+		SystemName = FName(RawSystemName);
+		UE_LOG(LogVirtualization, Display, TEXT("VirtualizationSystem name found in ini file: %s"), *RawSystemName);
 	}
-	
-	if (ConfigFile != nullptr)
+
+	if (!SystemName.IsNone())
 	{
-		FString RawSystemName;
-		if (ConfigFile->GetString(TEXT("Core.ContentVirtualization"), TEXT("SystemName"), RawSystemName))
+		Private::IVirtualizationSystemFactory* SystemFactory = FindFactory(SystemName);
+		if (SystemFactory != nullptr)
 		{
-			SystemName = FName(RawSystemName);
-			UE_LOG(LogVirtualization, Display, TEXT("VirtualizationSystem name found in ini file: %s"), *RawSystemName);
-		}
+			GVirtualizationSystem = SystemFactory->Create();
+			check(GVirtualizationSystem.IsValid()); // It is assumed that create will always return a valid pointer
 
-		if (!SystemName.IsNone())
-		{
-			Private::IVirtualizationSystemFactory* SystemFactory = FindFactory(SystemName);
-			if (SystemFactory != nullptr)
-			{
-				GVirtualizationSystem = SystemFactory->Create();
-				check(GVirtualizationSystem.IsValid()); // It is assumed that create will always return a valid pointer
-
-				if (!GVirtualizationSystem->Initialize(*ConfigFile))
-				{
-					UE_LOG(LogVirtualization, Error, TEXT("Initialization of the virtualization system '%s' failed, falling back to the default implementation"), *SystemName.ToString());
-					GVirtualizationSystem.Reset();
-				}
-			}
-			else
+			if (!GVirtualizationSystem->Initialize(InitParams))
 			{
-				UE_LOG(LogVirtualization, Error, TEXT("Unable to find factory to create the virtualization system: %s"), *SystemName.ToString());
+				UE_LOG(LogVirtualization, Error, TEXT("Initialization of the virtualization system '%s' failed, falling back to the default implementation"), *SystemName.ToString());
+				GVirtualizationSystem.Reset();
 			}
 		}
-	}
-	else
-	{
-		UE_LOG(LogVirtualization, Error, TEXT("Unable to find a valid engine config file when trying to create the virtualization system"));
+		else
+		{
+			UE_LOG(LogVirtualization, Error, TEXT("Unable to find factory to create the virtualization system: %s"), *SystemName.ToString());
+		}
 	}
 
+	// If we found no system to create so we will use the fallback system
 	if (!GVirtualizationSystem.IsValid())
 	{
-		// We found no system to create so we will use the fallback Null system
 		GVirtualizationSystem = MakeUnique<FNullVirtualizationSystem>();
-		GVirtualizationSystem->Initialize(FConfigFile()); // We know this method does nothing but we call it anyway for sake of form
+		GVirtualizationSystem->Initialize(InitParams);
 	}
 }
 
diff --git a/Engine/Source/Runtime/Core/Public/Virtualization/VirtualizationSystem.h b/Engine/Source/Runtime/Core/Public/Virtualization/VirtualizationSystem.h
--- a/Engine/Source/Runtime/Core/Public/Virtualization/VirtualizationSystem.h
+++ b/Engine/Source/Runtime/Core/Public/Virtualization/VirtualizationSystem.h
@@ -112,12 +112,40 @@ struct FPushRequest
 	EStatus Status = EStatus::Failed;
 };
 
+/** 
+ * The set of parameters to be used when initializing the virtualization system. The 
+ * members must remain valid for the duration of the call to ::Initialize. It is not
+ * expected that any virtualization system will store a reference to the members, if
+ * they want to retain the data then they will make their own copies.
+ */
+struct FInitParams
+{
+	FInitParams(FStringView InProjectName, const FConfigFile& InConfigFile)
+		: ProjectName(InProjectName)
+		, ConfigFile(InConfigFile)
+	{
+
+	}
+
+	/** The name of the current project (will default to FApp::GetProjectName()) */
+	FStringView ProjectName;
+
+	/** The config file to load the settings from (will default to GEngineIni) */
+	const FConfigFile& ConfigFile;
+};
 /**
  * Creates the global IVirtualizationSystem if it has not already been set up. This can be called explicitly
  * during process start up but it will also be called by IVirtualizationSystem::Get if it detects that the
  * IVirtualizationSystem has not yet been set up.
+ * 
+ * This version will use the default values of FInitParams.
  */
-CORE_API void Initialize(FConfigFile* ConfigFile = nullptr);
+CORE_API void Initialize();
+
+/**
+ * This version of ::Initialize takes parameters via the FInitParams structure.
+ */
+CORE_API void Initialize(const FInitParams& InitParams);
 
 /**
  * Shutdowns the global IVirtualizationSystem if it exists. 
@@ -146,18 +174,19 @@ public:
 	virtual ~IVirtualizationSystem() = default;
 
 	/**
-	 * Initialize the system from the given config file. Since the system might be initialized from a custom set
-	 * of specifically loaded config files, this is the only point in the systems life time that you are guaranteed 
-	 * to have the correct config file available so make sure that all settings that are required are parsed and
-	 * stored at this point rather than trying to access the config system directly in the future.
+	 * Initialize the system from the parameters given in the FInitParams structure.
+	 * The system can only rely on the members of FInitParams to be valid for the duration of the method call, so
+	 * if a system needs to retain information longer term then it should make it's own copy of the required data.
+	 * 
+	 * NOTE: Although it is relatively easy to access cached FConfigFiles, systems should use the one provided 
+	 * by InitParams to ensure that the correct settings are parsed.
 	 * 
-	 * @param ConfigFile	The config file to parse the systems settings from. The object is guaranteed to be
-	 *						valid for the life time of the method so do not store it for future use.
+	 * @param InitParam The parameters used to initialize the system
 	 * @return				True if the system was initialized correctly, otherwise false. Note that if the method
-	 *						returns false then the system will be deleted and the default FNullVirtualizationSystem 
+	 *						returns false then the system will be deleted and the default FNullVirtualizationSystem
 	 *						will be used instead.
 	 */
-	virtual bool Initialize(const FConfigFile& ConfigFile) = 0;
+	virtual bool Initialize(const FInitParams& InitParams) = 0;
 
 	/** Gain access to the current virtualization system active for the project */
 	static IVirtualizationSystem& Get();
