Another hopeful fix for a PVS static anaylsis warning.

diff --git a/Engine/Source/Runtime/Core/Public/Containers/StaticBitArray.h b/Engine/Source/Runtime/Core/Public/Containers/StaticBitArray.h
--- a/Engine/Source/Runtime/Core/Public/Containers/StaticBitArray.h
+++ b/Engine/Source/Runtime/Core/Public/Containers/StaticBitArray.h
@@ -400,9 +400,10 @@ private:
 	 */
 	void Trim_()
 	{
-		if constexpr (NumBits % NumBitsPerWord != 0)
+		constexpr uint32 NumOverflowBits = NumBits % NumBitsPerWord; //-V1064 The 'NumBits' operand of the modulo operation is less than the 'NumBitsPerWord' operand. The result is always equal to the left operand.
+		if constexpr (NumOverflowBits != 0)
 		{
-			Words[NumWords-1] &= (WordType(1) << (NumBits % NumBitsPerWord)) - 1;
+			Words[NumWords-1] &= (WordType(1) << NumOverflowBits) - 1;
 		}
 	}
 
