Add support for dynamic thread prioritization
  - This restore proper prioritization support that was lost with the new task backend while keeping a sane number of threads
  - This helps to prevent latency sensitive tasks from going to efficient cores for Alder Lake
  - This fixes performance problems caused by high prio task getting picked up by background threads and being preempted under high-load, causing FPS hiccups.
  - SetThreadPriority takes ~20us on average on Windows System and is only paid for background work
  - Longer SetThreadPriority can be expected when lowering the priority of a task under heavy load, which might cause the caller to be preempted right away, this is expected.
  - Can be A/B tested from command-line with -TaskGraphUseDynamicPrioritization=1/0

diff --git a/Engine/Source/Runtime/Core/Private/Async/Fundamental/ReserveScheduler.cpp b/Engine/Source/Runtime/Core/Private/Async/Fundamental/ReserveScheduler.cpp
--- a/Engine/Source/Runtime/Core/Private/Async/Fundamental/ReserveScheduler.cpp
+++ b/Engine/Source/Runtime/Core/Private/Async/Fundamental/ReserveScheduler.cpp
@@ -3,6 +3,8 @@
 #include "Async/Fundamental/ReserveScheduler.h"
 #include "Misc/ScopeLock.h"
 
+extern CORE_API bool GTaskGraphUseDynamicPrioritization;
+
 namespace LowLevelTasks
 {
 
@@ -65,6 +67,8 @@ void FReserveScheduler::StartWorkers(FScheduler& MainScheduler, uint32 NumWorker
 		NumWorkers = FMath::Min(FPlatformMisc::NumberOfWorkerThreadsToSpawn(), 64);
 	}
 
+	WorkerPriority = GTaskGraphUseDynamicPrioritization ? MainScheduler.GetWorkerPriority() : WorkerPriority;
+
 	uint32 OldActiveWorkers = ActiveWorkers.load(std::memory_order_relaxed);
 	if(OldActiveWorkers == 0 && FPlatformProcess::SupportsMultithreading() && ActiveWorkers.compare_exchange_strong(OldActiveWorkers, NumWorkers, std::memory_order_relaxed))
 	{
diff --git a/Engine/Source/Runtime/Core/Private/Async/Fundamental/Scheduler.cpp b/Engine/Source/Runtime/Core/Private/Async/Fundamental/Scheduler.cpp
--- a/Engine/Source/Runtime/Core/Private/Async/Fundamental/Scheduler.cpp
+++ b/Engine/Source/Runtime/Core/Private/Async/Fundamental/Scheduler.cpp
@@ -4,10 +4,13 @@
 #include "Async/Fundamental/Task.h"
 #include "Async/TaskTrace.h"
 #include "Logging/LogMacros.h"
+#include "HAL/RunnableThread.h"
 #include "Misc/ScopeLock.h"
 #include "Misc/Fork.h"
 #include "CoreGlobals.h"
 
+extern CORE_API bool GTaskGraphUseDynamicPrioritization;
+
 namespace LowLevelTasks
 {
 	DEFINE_LOG_CATEGORY(LowLevelTasks);
@@ -90,6 +93,12 @@ namespace LowLevelTasks
 
 	void FScheduler::StartWorkers(uint32 NumForegroundWorkers, uint32 NumBackgroundWorkers, FThread::EForkable IsForkable, EThreadPriority InWorkerPriority,  EThreadPriority InBackgroundPriority, uint64 InWorkerAffinity, uint64 InBackgroundAffinity)
 	{
+		int32 Value = 0;
+		if (FParse::Value(FCommandLine::Get(), TEXT("TaskGraphUseDynamicPrioritization="), Value))
+		{
+			GTaskGraphUseDynamicPrioritization = Value != 0;
+		}
+
 		if (NumForegroundWorkers == 0 && NumBackgroundWorkers == 0)
 		{
 			NumForegroundWorkers = FMath::Max<int32>(1, FMath::Min<int32>(2, FPlatformMisc::NumberOfWorkerThreadsToSpawn() - 1));
@@ -135,12 +144,55 @@ namespace LowLevelTasks
 			{
 				WorkerEvents.Emplace();
 				WorkerLocalQueues.Emplace(QueueRegistry, ELocalQueueType::EBackground, &WorkerEvents.Last());
-				WorkerThreads.Add(CreateWorker(true, IsForkable, &WorkerEvents.Last(), &WorkerLocalQueues.Last(), BackgroundPriority, BackgroundAffinity));
+				WorkerThreads.Add(CreateWorker(true, IsForkable, &WorkerEvents.Last(), &WorkerLocalQueues.Last(), GTaskGraphUseDynamicPrioritization ? WorkerPriority : BackgroundPriority, BackgroundAffinity));
 			}
 			UE::Trace::ThreadGroupEnd();
 		}
 	}
 
+	void FScheduler::ExecuteTask(FTask& Task)
+	{
+		FTask* ParentTask = FTask::ActiveTask;
+		FTask::ActiveTask = &Task;
+
+		if (!Task.IsBackgroundTask())
+		{
+			TRACE_CPUPROFILER_EVENT_SCOPE(ExecuteForegroundTask);
+			Task.ExecuteTask();
+		}
+		else
+		{
+			// Dynamic priority only enables for root task when we're not inside a named thread (i.e. GT, RT)
+			const bool bSkipPriorityChange = ParentTask || !GTaskGraphUseDynamicPrioritization || !FSchedulerTls::IsWorkerThread();
+
+			FRunnableThread* RunnableThread = nullptr;
+			if (!bSkipPriorityChange)
+			{
+				// We assume all threads executing tasks are RunnableThread and this can't be null or it will crash. 
+				// Which is fine since we want to know about it sooner rather than later.
+				RunnableThread = FRunnableThread::GetRunnableThread();
+
+				checkSlow(RunnableThread && RunnableThread->GetThreadPriority() == WorkerPriority);
+
+				TRACE_CPUPROFILER_EVENT_SCOPE(LowerThreadPriority);
+				RunnableThread->SetThreadPriority(BackgroundPriority);
+			}
+
+			{
+				TRACE_CPUPROFILER_EVENT_SCOPE(ExecuteBackgroundTask);
+				Task.ExecuteTask();
+			}
+
+			if (!bSkipPriorityChange)
+			{
+				TRACE_CPUPROFILER_EVENT_SCOPE(RaiseThreadPriority);
+				RunnableThread->SetThreadPriority(WorkerPriority);
+			}
+		}
+
+		FTask::ActiveTask = ParentTask;
+	}
+
 	void FScheduler::StopWorkers(bool DrainGlobalQueue)
 	{
 		uint32 OldActiveWorkers = ActiveWorkers.load(std::memory_order_relaxed);
@@ -163,7 +215,7 @@ namespace LowLevelTasks
 			{
 				for (FTask* Task = QueueRegistry.Dequeue(); Task != nullptr; Task = QueueRegistry.Dequeue())
 				{
-					Task->ExecuteTask();
+					ExecuteTask(*Task);
 				}
 			}
 		}
@@ -213,7 +265,7 @@ namespace LowLevelTasks
 		}
 		else
 		{
-			Task.ExecuteTask();
+			ExecuteTask(Task);
 		}
 	}
 
@@ -264,20 +316,8 @@ namespace LowLevelTasks
 						WakeUpWorker(!bPermitBackgroundWork);
 					}
 				}
-				FTask* OldTask = FTask::ActiveTask;
-				FTask::ActiveTask = Task;
-				{
-#if UE_TASK_TRACE_ENABLED
-					if (!UE_TRACE_CHANNELEXPR_IS_ENABLED(TaskTrace::TaskChannel))
-					{
-						TRACE_CPUPROFILER_EVENT_SCOPE(ExecuteTask);
-					}
-#else
-					TRACE_CPUPROFILER_EVENT_SCOPE(ExecuteTask);
-#endif
-					Task->ExecuteTask();
-				}
-				FTask::ActiveTask = OldTask;
+
+				ExecuteTask(*Task);
 				return true;
 			}
 			return false;
diff --git a/Engine/Source/Runtime/Core/Private/Async/TaskGraph.cpp b/Engine/Source/Runtime/Core/Private/Async/TaskGraph.cpp
--- a/Engine/Source/Runtime/Core/Private/Async/TaskGraph.cpp
+++ b/Engine/Source/Runtime/Core/Private/Async/TaskGraph.cpp
@@ -101,6 +101,14 @@ static FAutoConsoleVariableRef CVarForkedProcessMaxWorkerThreads(
 	TEXT("Configures the number of worker threads a forked process should spawn if it allows multithreading.")
 );
 
+CORE_API bool GTaskGraphUseDynamicPrioritization = 1;
+static FAutoConsoleVariableRef CVarTaskDynamicPrioritization(
+	TEXT("TaskGraph.UseDynamicPrioritization"),
+	GTaskGraphUseDynamicPrioritization,
+	TEXT("Adjust thread priority per-task so that higher priority tasks running on background threads can't be preempted as easily. Helps a lot under high load."),
+	ECVF_ReadOnly
+);
+
 CORE_API int32 GUseNewTaskBackend = 1;
 int32 GNumForegroundWorkers = 2;
 static FAutoConsoleVariableRef CVarNumForegroundWorkers(
@@ -1905,7 +1913,7 @@ public:
 			check(IsInGameThread()); // otherwise we can have a race on starting reserve workers below
 			if (GConfig == nullptr)
 			{
-				// postpone starting reserve workers tillGConfig is initialized, to know if reserve workers are disabled
+				// postpone starting reserve workers until GConfig is initialized, to know if reserve workers are disabled
 				FCoreDelegates::ConfigReadyForUse.AddRaw(this, &FTaskGraphCompatibilityImplementation::StartReserveWorkers);
 			}
 			else
diff --git a/Engine/Source/Runtime/Core/Private/Windows/WindowsRunnableThread.cpp b/Engine/Source/Runtime/Core/Private/Windows/WindowsRunnableThread.cpp
--- a/Engine/Source/Runtime/Core/Private/Windows/WindowsRunnableThread.cpp
+++ b/Engine/Source/Runtime/Core/Private/Windows/WindowsRunnableThread.cpp
@@ -22,9 +22,12 @@ int FRunnableThreadWin::TranslateThreadPriority(EThreadPriority Priority)
 		case TPri_Highest:				return THREAD_PRIORITY_HIGHEST;
 		case TPri_TimeCritical:			return THREAD_PRIORITY_HIGHEST;
 		case TPri_Lowest:				return THREAD_PRIORITY_LOWEST;
-		case TPri_SlightlyBelowNormal:	return THREAD_PRIORITY_BELOW_NORMAL;
 
-	// Note: previously, the behaviour was:
+		// There is no such things as slightly below normal on Windows.
+		// This can't be below normal since we don't want latency sensitive task to go to efficient cores on Alder Lake.
+		case TPri_SlightlyBelowNormal:	return THREAD_PRIORITY_NORMAL;
+
+	// Note: previously, the behavior was:
 	//
 	//case TPri_AboveNormal:			return THREAD_PRIORITY_HIGHEST;
 	//case TPri_Normal:					return THREAD_PRIORITY_HIGHEST - 1;
@@ -36,12 +39,12 @@ int FRunnableThreadWin::TranslateThreadPriority(EThreadPriority Priority)
 	//
 	// But the change (CL3747560) was not well documented (it didn't describe
 	// the symptoms it was supposed to address) and introduces undesirable 
-	// system behaviour on Windows since it starves out other processes in
+	// system behavior on Windows since it starves out other processes in
 	// the system when UE compiles shaders or otherwise goes wide due to the
 	// inflation in priority (Normal mapped to THREAD_PRIORITY_ABOVE_NORMAL)
 	// I kept the TPri_TimeCritical mapping to THREAD_PRIORITY_HIGHEST however
-	// to avoid introducing poor behaviour since time critical priority is
-	// similarly detrimental to overall system behaviour.
+	// to avoid introducing poor behavior since time critical priority is
+	// similarly detrimental to overall system behavior.
 	//
 	// If we discover thread scheduling issues it would maybe be better to 
 	// adjust actual thread priorities at the source instead of this mapping?
diff --git a/Engine/Source/Runtime/Core/Public/Async/Fundamental/Scheduler.h b/Engine/Source/Runtime/Core/Public/Async/Fundamental/Scheduler.h
--- a/Engine/Source/Runtime/Core/Public/Async/Fundamental/Scheduler.h
+++ b/Engine/Source/Runtime/Core/Public/Async/Fundamental/Scheduler.h
@@ -174,11 +174,17 @@ namespace LowLevelTasks
 		//get the Queue registry to register additonal WorkerQueues for this Scheduler
 		inline FSchedulerTls::FQueueRegistry& GetQueueRegistry();
 
+		//get the worker priority set when workers were started
+		inline EThreadPriority GetWorkerPriority() const { return WorkerPriority; }
+
+		//get the background priority set when workers were started
+		inline EThreadPriority GetBackgroundPriority() const { return BackgroundPriority; }
 	public:
 		FScheduler() = default;
 		~FScheduler();
 
 	private: 
+		void ExecuteTask(FTask& Task);
 		TUniquePtr<FThread> CreateWorker(bool bPermitBackgroundWork = false, FThread::EForkable IsForkable = FThread::NonForkable, FSleepEvent* ExternalWorkerEvent = nullptr, FSchedulerTls::FLocalQueueType* ExternalWorkerLocalQueue = nullptr, EThreadPriority Priority = EThreadPriority::TPri_Normal, uint64 InAffinity = 0);
 		void WorkerMain(struct FSleepEvent* WorkerEvent, FSchedulerTls::FLocalQueueType* ExternalWorkerLocalQueue, uint32 WaitCycles, bool bPermitBackgroundWork);
 		CORE_API void LaunchInternal(FTask& Task, EQueuePreference QueuePreference, bool bWakeUpWorker);
