TMap: Add Array copy-accessor for the key-value pairs.
#rnx

diff --git a/Engine/Source/Runtime/Core/Public/Containers/Map.h b/Engine/Source/Runtime/Core/Public/Containers/Map.h
--- a/Engine/Source/Runtime/Core/Public/Containers/Map.h
+++ b/Engine/Source/Runtime/Core/Public/Containers/Map.h
@@ -691,6 +691,12 @@ public:
 		return Pairs.ContainsByHash(KeyHash, Key);
 	}
 
+	/** Copy the key/value pairs in this map into an array. */
+	TArray<ElementType> Array() const
+	{
+		return Pairs.Array();
+	}
+
 	/**
 	 * Generate an array from the keys in this map.
 	 *
