Fix non-unity build

diff --git a/Engine/Source/Runtime/Core/Private/Async/Fundamental/Scheduler.cpp b/Engine/Source/Runtime/Core/Private/Async/Fundamental/Scheduler.cpp
--- a/Engine/Source/Runtime/Core/Private/Async/Fundamental/Scheduler.cpp
+++ b/Engine/Source/Runtime/Core/Private/Async/Fundamental/Scheduler.cpp
@@ -6,6 +6,7 @@
 #include "Logging/LogMacros.h"
 #include "HAL/RunnableThread.h"
 #include "Misc/ScopeLock.h"
+#include "Misc/CommandLine.h"
 #include "Misc/Fork.h"
 #include "CoreGlobals.h"
 
