String: Added string escaping functions under UE::String

#rnx

diff --git a/Engine/Source/Runtime/Core/Private/String/Escape.cpp b/Engine/Source/Runtime/Core/Private/String/Escape.cpp
new file mode 100644
--- /dev/null
+++ b/Engine/Source/Runtime/Core/Private/String/Escape.cpp
@@ -0,0 +1,64 @@
+// Copyright Epic Games, Inc. All Rights Reserved.
+
+#include "String/Escape.h"
+
+#include "Misc/StringBuilder.h"
+#include "String/BytesToHex.h"
+
+namespace UE::String
+{
+
+void EscapeTo(FStringView Input, FStringBuilderBase& Output)
+{
+	for (const TCHAR Char : Input)
+	{
+		if (Char < 0x20)
+		{
+			switch (Char)
+			{
+			case '\a': Output.Append(TEXT("\\a")); break;
+			case '\b': Output.Append(TEXT("\\b")); break;
+			case '\f': Output.Append(TEXT("\\f")); break;
+			case '\n': Output.Append(TEXT("\\n")); break;
+			case '\r': Output.Append(TEXT("\\r")); break;
+			case '\t': Output.Append(TEXT("\\t")); break;
+			case '\v': Output.Append(TEXT("\\v")); break;
+			default:
+				Output.Append(TEXT("\\x"));
+				BytesToHexLower({uint8(Char)}, Output);
+				break;
+			}
+		}
+		else if (Char < 0x7f)
+		{
+			switch (Char)
+			{
+			case '\"': Output.Append(TEXT("\\\"")); break;
+			case '\'': Output.Append(TEXT("\\\'")); break;
+			case '\\': Output.Append(TEXT("\\\\")); break;
+			default:
+				Output.AppendChar(Char);
+				break;
+			}
+		}
+		else if (Char < 0x80)
+		{
+			Output.Append(TEXT("\\x"));
+			BytesToHexLower({uint8(Char)}, Output);
+		}
+		else
+		{
+			Output.Append(TEXT("\\u"));
+			BytesToHexLower({uint8(Char >> 8), uint8(Char)}, Output);
+		}
+	}
+}
+
+void QuoteEscapeTo(FStringView Input, FStringBuilderBase& Output)
+{
+	Output.AppendChar(TEXT('\"'));
+	EscapeTo(Input, Output);
+	Output.AppendChar(TEXT('\"'));
+}
+
+} // UE::String
diff --git a/Engine/Source/Runtime/Core/Public/String/Escape.h b/Engine/Source/Runtime/Core/Public/String/Escape.h
new file mode 100644
--- /dev/null
+++ b/Engine/Source/Runtime/Core/Public/String/Escape.h
@@ -0,0 +1,65 @@
+// Copyright Epic Games, Inc. All Rights Reserved.
+
+#pragma once
+
+#include "CoreTypes.h"
+#include "Containers/StringFwd.h"
+#include "Containers/StringView.h"
+
+namespace UE::String::Private
+{
+
+struct FEscape { FStringView Input; };
+struct FQuoteEscape { FStringView Input; };
+
+} // UE::String::Private
+
+namespace UE::String
+{
+
+/**
+ * Escape the string into the string builder.
+ *
+ * Example: Escape("Tab is \t and \" is a quote.", Output) -> "Tab is \\t and \\\" is a quote."
+ */
+CORE_API void EscapeTo(FStringView Input, FStringBuilderBase& Output);
+
+/**
+ * Escape the string when appended to a string builder.
+ *
+ * Example: Builder << String::Escape(Value);
+ */
+inline Private::FEscape Escape(FStringView Input) { return {Input}; }
+
+/**
+ * Quote and escape the string into the string builder.
+ *
+ * Example: QuoteEscape("Tab is \t and \" is a quote.", Output) -> "\"Tab is \\t and \\\" is a quote.\""
+ */
+CORE_API void QuoteEscapeTo(FStringView Input, FStringBuilderBase& Output);
+
+/**
+ * Quote and escape the string when appended to a string builder.
+ *
+ * Example: Builder << String::QuoteEscape(Value);
+ */
+inline Private::FQuoteEscape QuoteEscape(FStringView Input) { return {Input}; }
+
+} // UE::String
+
+namespace UE::String::Private
+{
+
+inline FStringBuilderBase& operator<<(FStringBuilderBase& Builder, const FEscape& Adapter)
+{
+	EscapeTo(Adapter.Input, Builder);
+	return Builder;
+}
+
+inline FStringBuilderBase& operator<<(FStringBuilderBase& Builder, const FQuoteEscape& Adapter)
+{
+	QuoteEscapeTo(Adapter.Input, Builder);
+	return Builder;
+}
+
+} // UE::String::Private
diff --git a/Engine/Source/Runtime/Core/Tests/String/EscapeTest.cpp b/Engine/Source/Runtime/Core/Tests/String/EscapeTest.cpp
new file mode 100644
--- /dev/null
+++ b/Engine/Source/Runtime/Core/Tests/String/EscapeTest.cpp
@@ -0,0 +1,37 @@
+// Copyright Epic Games, Inc. All Rights Reserved.
+
+#include "String/Escape.h"
+
+#include "Containers/StringView.h"
+#include "Misc/StringBuilder.h"
+#include "TestHarness.h"
+
+namespace UE::String
+{
+
+TEST_CASE("Core::String::Escape", "[Core][String][Smoke]")
+{
+	SECTION("Escape")
+	{
+		CHECK(WriteToString<128>(Escape(TEXTVIEW("No escaping necessary."))) == TEXTVIEW("No escaping necessary."));
+		CHECK(WriteToString<128>(Escape(TEXTVIEW("Tab is \t & \" is quote."))) == TEXTVIEW("Tab is \\t & \\\" is quote."));
+		CHECK(WriteToString<128>(Escape(TEXTVIEW("\a\b\f\n\r\t\v\"\'\\"))) == TEXTVIEW("\\a\\b\\f\\n\\r\\t\\v\\\"\\\'\\\\"));
+		CHECK(WriteToString<128>(Escape(TEXTVIEW("\x00\x01\x02\x03\x1c\x1d\x1e\x1f\x7f"))) == TEXTVIEW("\\x00\\x01\\x02\\x03\\x1c\\x1d\\x1e\\x1f\\x7f"));
+		CHECK(WriteToString<128>(Escape(TEXTVIEW("\u0080\u0081\u0082\u0083"))) == TEXTVIEW("\\u0080\\u0081\\u0082\\u0083"));
+		CHECK(WriteToString<128>(Escape(TEXTVIEW("\u00a1\u00a2\u00a3\u00a4"))) == TEXTVIEW("\\u00a1\\u00a2\\u00a3\\u00a4"));
+		CHECK(WriteToString<128>(Escape(TEXTVIEW("\u0101\u0102\u0103\u0104"))) == TEXTVIEW("\\u0101\\u0102\\u0103\\u0104"));
+	}
+
+	SECTION("QuoteEscape")
+	{
+		CHECK(WriteToString<128>(QuoteEscape(TEXTVIEW("No escaping necessary."))) == TEXTVIEW("\"No escaping necessary.\""));
+		CHECK(WriteToString<128>(QuoteEscape(TEXTVIEW("Tab is \t & \" is quote."))) == TEXTVIEW("\"Tab is \\t & \\\" is quote.\""));
+		CHECK(WriteToString<128>(QuoteEscape(TEXTVIEW("\a\b\f\n\r\t\v\"\'\\"))) == TEXTVIEW("\"\\a\\b\\f\\n\\r\\t\\v\\\"\\\'\\\\\""));
+		CHECK(WriteToString<128>(QuoteEscape(TEXTVIEW("\x00\x01\x02\x03\x1c\x1d\x1e\x1f\x7f"))) == TEXTVIEW("\"\\x00\\x01\\x02\\x03\\x1c\\x1d\\x1e\\x1f\\x7f\""));
+		CHECK(WriteToString<128>(QuoteEscape(TEXTVIEW("\u0080\u0081\u0082\u0083"))) == TEXTVIEW("\"\\u0080\\u0081\\u0082\\u0083\""));
+		CHECK(WriteToString<128>(QuoteEscape(TEXTVIEW("\u00a1\u00a2\u00a3\u00a4"))) == TEXTVIEW("\"\\u00a1\\u00a2\\u00a3\\u00a4\""));
+		CHECK(WriteToString<128>(QuoteEscape(TEXTVIEW("\u0101\u0102\u0103\u0104"))) == TEXTVIEW("\"\\u0101\\u0102\\u0103\\u0104\""));
+	}
+}
+
+} // UE::String
