PlatformFile: Eliminated the lock in IterateDirectoryRecursively

There was contention on the lock when 60+ worker threads execute the ParallelFor.

#rnx

diff --git a/Engine/Source/Runtime/Core/Private/GenericPlatform/GenericPlatformFile.cpp b/Engine/Source/Runtime/Core/Private/GenericPlatform/GenericPlatformFile.cpp
--- a/Engine/Source/Runtime/Core/Private/GenericPlatform/GenericPlatformFile.cpp
+++ b/Engine/Source/Runtime/Core/Private/GenericPlatform/GenericPlatformFile.cpp
@@ -1,6 +1,7 @@
 // Copyright Epic Games, Inc. All Rights Reserved.
 
 #include "GenericPlatform/GenericPlatformFile.h"
+#include "Algo/Accumulate.h"
 #include "HAL/FileManager.h"
 #include "Misc/Paths.h"
 #include "HAL/PlatformMisc.h"
@@ -560,13 +561,10 @@ bool IPlatformFile::IterateDirectoryRecursively(const TCHAR* Directory, FDirecto
 	class FRecurse : public FDirectoryVisitor
 	{
 	public:
-		IPlatformFile&      PlatformFile;
 		FDirectoryVisitor&  Visitor;
-		FRWLock             DirectoriesLock;
 		TArray<FString>&    Directories;
-		FRecurse(IPlatformFile&	InPlatformFile, FDirectoryVisitor& InVisitor, TArray<FString>& InDirectories)
+		FRecurse(FDirectoryVisitor& InVisitor, TArray<FString>& InDirectories)
 			: FDirectoryVisitor(InVisitor.DirectoryVisitorFlags)
-			, PlatformFile(InPlatformFile)
 			, Visitor(InVisitor)
 			, Directories(InDirectories)
 		{
@@ -576,34 +574,40 @@ bool IPlatformFile::IterateDirectoryRecursively(const TCHAR* Directory, FDirecto
 			bool bResult = Visitor.Visit(FilenameOrDirectory, bIsDirectory);
 			if (bResult && bIsDirectory)
 			{
-				FString Directory(FilenameOrDirectory);
-				FRWScopeLock ScopeLock(DirectoriesLock, SLT_Write);
-				Directories.Emplace(MoveTemp(Directory));
+				Directories.Emplace(FilenameOrDirectory);
 			}
 			return bResult;
 		}
 	};
 
-	TArray<FString> DirectoriesToVisitNext;
-	DirectoriesToVisitNext.Add(Directory);
-
-	const bool IsTaskGraphReady = FTaskGraphInterface::IsRunning();
-	TAtomic<bool> bResult(true);
-	FRecurse Recurse(*this, Visitor, DirectoriesToVisitNext);
-	while (bResult && DirectoriesToVisitNext.Num() > 0)
-	{
-		TArray<FString> DirectoriesToVisit = MoveTemp(DirectoriesToVisitNext);
-		ParallelFor( TEXT("IterateDirectoryRecursively.PF"),
-			DirectoriesToVisit.Num(),1,
-			[this, &DirectoriesToVisit, &Recurse, &bResult](int32 Index)
+	TArray<FString> DirectoriesToVisit;
+	DirectoriesToVisit.Add(Directory);
+
+	constexpr int32 MinBatchSize = 1;
+	const EParallelForFlags ParallelForFlags = FTaskGraphInterface::IsRunning() && Visitor.IsThreadSafe()
+		? EParallelForFlags::Unbalanced : EParallelForFlags::ForceSingleThread;
+	std::atomic<bool> bResult{true};
+	TArray<TArray<FString>> DirectoriesToVisitNext;
+	while (bResult && DirectoriesToVisit.Num() > 0)
+	{
+		ParallelForWithTaskContext(TEXT("IterateDirectoryRecursively.PF"),
+			DirectoriesToVisitNext,
+			DirectoriesToVisit.Num(),
+			MinBatchSize,
+			[this, &Visitor, &DirectoriesToVisit, &bResult](TArray<FString>& Directories, int32 Index)
 			{
-				if (bResult && !IterateDirectory(*DirectoriesToVisit[Index], Recurse))
+				FRecurse Recurse(Visitor, Directories);
+				if (bResult.load(std::memory_order_relaxed) && !IterateDirectory(*DirectoriesToVisit[Index], Recurse))
 				{
-					bResult = false;
+					bResult.store(false, std::memory_order_relaxed);
 				}
 			},
-			(IsTaskGraphReady && Visitor.IsThreadSafe()) ? EParallelForFlags::Unbalanced : EParallelForFlags::ForceSingleThread
-		);
+			ParallelForFlags);
+		DirectoriesToVisit.Reset(Algo::TransformAccumulate(DirectoriesToVisitNext, &TArray<FString>::Num, 0));
+		for (TArray<FString>& Directories : DirectoriesToVisitNext)
+		{
+			DirectoriesToVisit.Append(MoveTemp(Directories));
+		}
 	}
 
 	return bResult;
