Core: Migrated a few tests to Catch

#rnx

diff --git a/Engine/Source/Runtime/Core/Private/Tests/Hash/Blake3Test.cpp b/Engine/Source/Runtime/Core/Private/Tests/Hash/Blake3Test.cpp
deleted file mode 100644
diff --git a/Engine/Source/Runtime/Core/Private/Tests/String/FindTest.cpp b/Engine/Source/Runtime/Core/Private/Tests/String/FindTest.cpp
deleted file mode 100644
diff --git a/Engine/Source/Runtime/Core/Private/Tests/String/ParseLinesTest.cpp b/Engine/Source/Runtime/Core/Private/Tests/String/ParseLinesTest.cpp
deleted file mode 100644
diff --git a/Engine/Source/Runtime/Core/Private/Tests/String/ParseTokensTest.cpp b/Engine/Source/Runtime/Core/Private/Tests/String/ParseTokensTest.cpp
deleted file mode 100644
diff --git a/Engine/Source/Runtime/Core/Tests/Hash/Blake3Test.cpp b/Engine/Source/Runtime/Core/Tests/Hash/Blake3Test.cpp
new file mode 100644
--- /dev/null
+++ b/Engine/Source/Runtime/Core/Tests/Hash/Blake3Test.cpp
@@ -0,0 +1,61 @@
+// Copyright Epic Games, Inc. All Rights Reserved.
+
+#include "Hash/Blake3.h"
+
+#include "Misc/StringBuilder.h"
+#include "TestHarness.h"
+
+TEST_CASE("Core::Hash::Blake3", "[Core][Hash][Smoke]")
+{
+	struct FBlake3TestCase
+	{
+		uint32 InputLength;
+		FStringView Hash;
+	};
+
+	constexpr FBlake3TestCase TestCases[] =
+	{
+		{     0u, TEXTVIEW("af1349b9f5f9a1a6a0404dea36dcc9499bcb25c9adc112b7cc9a93cae41f3262")},
+		{     1u, TEXTVIEW("2d3adedff11b61f14c886e35afa036736dcd87a74d27b5c1510225d0f592e213")},
+		{ 1'023u, TEXTVIEW("10108970eeda3eb932baac1428c7a2163b0e924c9a9e25b35bba72b28f70bd11")},
+		{ 1'024u, TEXTVIEW("42214739f095a406f3fc83deb889744ac00df831c10daa55189b5d121c855af7")},
+		{ 1'025u, TEXTVIEW("d00278ae47eb27b34faecf67b4fe263f82d5412916c1ffd97c8cb7fb814b8444")},
+		{ 2'048u, TEXTVIEW("e776b6028c7cd22a4d0ba182a8bf62205d2ef576467e838ed6f2529b85fba24a")},
+		{ 2'049u, TEXTVIEW("5f4d72f40d7a5f82b15ca2b2e44b1de3c2ef86c426c95c1af0b6879522563030")},
+		{ 3'072u, TEXTVIEW("b98cb0ff3623be03326b373de6b9095218513e64f1ee2edd2525c7ad1e5cffd2")},
+		{ 3'073u, TEXTVIEW("7124b49501012f81cc7f11ca069ec9226cecb8a2c850cfe644e327d22d3e1cd3")},
+		{ 4'096u, TEXTVIEW("015094013f57a5277b59d8475c0501042c0b642e531b0a1c8f58d2163229e969")},
+		{ 4'097u, TEXTVIEW("9b4052b38f1c5fc8b1f9ff7ac7b27cd242487b3d890d15c96a1c25b8aa0fb995")},
+		{ 5'120u, TEXTVIEW("9cadc15fed8b5d854562b26a9536d9707cadeda9b143978f319ab34230535833")},
+		{ 5'121u, TEXTVIEW("628bd2cb2004694adaab7bbd778a25df25c47b9d4155a55f8fbd79f2fe154cff")},
+		{ 6'144u, TEXTVIEW("3e2e5b74e048f3add6d21faab3f83aa44d3b2278afb83b80b3c35164ebeca205")},
+		{ 6'145u, TEXTVIEW("f1323a8631446cc50536a9f705ee5cb619424d46887f3c376c695b70e0f0507f")},
+		{ 7'168u, TEXTVIEW("61da957ec2499a95d6b8023e2b0e604ec7f6b50e80a9678b89d2628e99ada77a")},
+		{ 7'169u, TEXTVIEW("a003fc7a51754a9b3c7fae0367ab3d782dccf28855a03d435f8cfe74605e7817")},
+		{ 8'192u, TEXTVIEW("aae792484c8efe4f19e2ca7d371d8c467ffb10748d8a5a1ae579948f718a2a63")},
+		{ 8'193u, TEXTVIEW("bab6c09cb8ce8cf459261398d2e7aef35700bf488116ceb94a36d0f5f1b7bc3b")},
+		{16'384u, TEXTVIEW("f875d6646de28985646f34ee13be9a576fd515f76b5b0a26bb324735041ddde4")},
+		{31'744u, TEXTVIEW("62b6960e1a44bcc1eb1a611a8d6235b6b4b78f32e7abc4fb4c6cdcce94895c47")},
+	};
+
+	constexpr uint32 MaxInputLength = 31'744u;
+
+	uint8 Input[MaxInputLength];
+	for (uint32 Index = 0; Index < MaxInputLength; ++Index)
+	{
+		Input[Index] = uint8(Index % 251);
+	}
+
+	FBlake3 Hash;
+	uint32 InputIndex = 0;
+	for (const FBlake3TestCase& TestCase : TestCases)
+	{
+		REQUIRE(InputIndex <= TestCase.InputLength);
+		REQUIRE(TestCase.InputLength <= MaxInputLength);
+		Hash.Update(Input + InputIndex, TestCase.InputLength - InputIndex);
+		InputIndex = TestCase.InputLength;
+		TStringBuilder<70> HashHex;
+		HashHex << Hash.Finalize();
+		CHECK(HashHex.ToView() == TestCase.Hash);
+	}
+}
diff --git a/Engine/Source/Runtime/Core/Tests/String/FindTest.cpp b/Engine/Source/Runtime/Core/Tests/String/FindTest.cpp
new file mode 100644
--- /dev/null
+++ b/Engine/Source/Runtime/Core/Tests/String/FindTest.cpp
@@ -0,0 +1,224 @@
+// Copyright Epic Games, Inc. All Rights Reserved.
+
+#include "String/Find.h"
+
+#include "Containers/StringView.h"
+#include "TestHarness.h"
+
+namespace UE::String
+{
+
+TEST_CASE("Core::String::Find", "[Core][String][Smoke]")
+{
+	SECTION("FindFirst")
+	{
+		CHECK(FindFirst(TEXT("AbCABCAbCABC"), TEXT("A")) == 0);
+		CHECK(FindFirst(TEXT("AbCABCAbCABC"), TEXT("a"), ESearchCase::IgnoreCase) == 0);
+		CHECK(FindFirst(TEXT("AbCABCAbCABC"), TEXT("b")) == 1);
+		CHECK(FindFirst(TEXT("AbCABCAbCABC"), TEXT("B")) == 4);
+		CHECK(FindFirst(TEXT("AbCABCAbCABC"), TEXT("B"), ESearchCase::IgnoreCase) == 1);
+		CHECK(FindFirst(TEXT("AbCABCAbCABC"), TEXT("a")) == INDEX_NONE);
+		CHECK(FindFirst(TEXT("AbCABCAbCABC"), TEXT("D"), ESearchCase::IgnoreCase) == INDEX_NONE);
+		CHECK(FindFirst(TEXT("AbCABCAbCABD"), TEXT("D")) == 11);
+		CHECK(FindFirst(TEXT("AbCABCAbCABD"), TEXT("d"), ESearchCase::IgnoreCase) == 11);
+
+		CHECK(FindFirst(TEXT("AbCABCAbCABC"), TEXT("AbC")) == 0);
+		CHECK(FindFirst(TEXT("AbCABCAbCABC"), TEXT("ABC")) == 3);
+		CHECK(FindFirst(TEXT("AbCABCAbCABC"), TEXT("Bc"), ESearchCase::IgnoreCase) == 1);
+		CHECK(FindFirst(TEXT("AbCABCAbCABC"), TEXT("ab")) == INDEX_NONE);
+		CHECK(FindFirst(TEXT("AbCABCAbCABC"), TEXT("CD"), ESearchCase::IgnoreCase) == INDEX_NONE);
+		CHECK(FindFirst(TEXT("AbCABCAbCABD"), TEXT("BD")) == 10);
+		CHECK(FindFirst(TEXT("AbCABCAbCABD"), TEXT("Bd"), ESearchCase::IgnoreCase) == 10);
+
+		CHECK(FindFirst(TEXT(""), TEXT("A")) == INDEX_NONE);
+		CHECK(FindFirst(TEXT("A"), TEXT("A")) == 0);
+		CHECK(FindFirst(TEXT("A"), TEXT("A"), ESearchCase::IgnoreCase) == 0);
+		CHECK(FindFirst(TEXT("ABC"), TEXT("ABC")) == 0);
+		CHECK(FindFirst(TEXT("ABC"), TEXT("abc"), ESearchCase::IgnoreCase) == 0);
+		CHECK(FindFirst(TEXT("AB"), TEXT("ABC")) == INDEX_NONE);
+
+		CHECK(FindFirst("AbCABCAbCABC", "ABC") == 3);
+
+		CHECK(FindFirst(FStringView(nullptr, 0), TEXT("SearchTerm")) == INDEX_NONE);
+		CHECK(FindFirst(FStringView(), TEXT("SearchTerm")) == INDEX_NONE);
+	}
+
+	SECTION("FindLast")
+	{
+		CHECK(FindLast(TEXT("AbCABCAbCABC"), TEXT("b")) == 7);
+		CHECK(FindLast(TEXT("AbCABCAbCABC"), TEXT("B")) == 10);
+		CHECK(FindLast(TEXT("AbCABCAbCABC"), TEXT("b"), ESearchCase::IgnoreCase) == 10);
+		CHECK(FindLast(TEXT("AbCABCAbCABC"), TEXT("a")) == INDEX_NONE);
+		CHECK(FindLast(TEXT("AbCABCAbCABC"), TEXT("D"), ESearchCase::IgnoreCase) == INDEX_NONE);
+		CHECK(FindLast(TEXT("AbCABCAbCABD"), TEXT("D")) == 11);
+		CHECK(FindLast(TEXT("AbCABCAbCABD"), TEXT("d"), ESearchCase::IgnoreCase) == 11);
+
+		CHECK(FindLast(TEXT("AbCABCAbCABC"), TEXT("AbC")) == 6);
+		CHECK(FindLast(TEXT("AbCABCAbCABC"), TEXT("ABC")) == 9);
+		CHECK(FindLast(TEXT("AbCABCAbCABC"), TEXT("Bc"), ESearchCase::IgnoreCase) == 10);
+		CHECK(FindLast(TEXT("AbCABCAbCABC"), TEXT("ab")) == INDEX_NONE);
+		CHECK(FindLast(TEXT("AbCABCAbCABC"), TEXT("CD"), ESearchCase::IgnoreCase) == INDEX_NONE);
+		CHECK(FindLast(TEXT("AbCABCAbCABC"), TEXT("BC")) == 10);
+		CHECK(FindLast(TEXT("AbCABCAbCABC"), TEXT("Bc"), ESearchCase::IgnoreCase) == 10);
+
+		CHECK(FindLast(TEXT(""), TEXT("A")) == INDEX_NONE);
+		CHECK(FindLast(TEXT("A"), TEXT("A")) == 0);
+		CHECK(FindLast(TEXT("A"), TEXT("A"), ESearchCase::IgnoreCase) == 0);
+		CHECK(FindLast(TEXT("ABC"), TEXT("ABC")) == 0);
+		CHECK(FindLast(TEXT("ABC"), TEXT("abc"), ESearchCase::IgnoreCase) == 0);
+		CHECK(FindLast(TEXT("AB"), TEXT("ABC")) == INDEX_NONE);
+
+		CHECK(FindLast("AbCABCAbCABC", "ABC") == 9);
+
+		CHECK(FindLast(FStringView(nullptr, 0), TEXT("SearchTerm")) == INDEX_NONE);
+		CHECK(FindLast(FStringView(), TEXT("SearchTerm")) == INDEX_NONE);
+	}
+
+	SECTION("FindFirstOfAny")
+	{
+		CHECK(FindFirstOfAny(TEXT("AbCABCAbcABC"), {TEXT("A"), TEXT("B")}) == 0);
+		CHECK(FindFirstOfAny(TEXT("AbCABCAbcABC"), {TEXT("a"), TEXT("B")}, ESearchCase::IgnoreCase) == 0);
+		CHECK(FindFirstOfAny(TEXT("AbCABCAbcABC"), {TEXT("b")}) == 1);
+		CHECK(FindFirstOfAny(TEXT("AbCABCAbcABC"), {TEXT("c"), TEXT("B")}) == 4);
+		CHECK(FindFirstOfAny(TEXT("AbCABCAbcABC"), {TEXT("c"), TEXT("B")}, ESearchCase::IgnoreCase) == 1);
+		CHECK(FindFirstOfAny(TEXT("AbCABCAbcABC"), {TEXT("D"), TEXT("a")}) == INDEX_NONE);
+		CHECK(FindFirstOfAny(TEXT("AbCABCAbcABC"), {TEXT("E"), TEXT("D")}, ESearchCase::IgnoreCase) == INDEX_NONE);
+		CHECK(FindFirstOfAny(TEXT("AbCABCAbcABD"), {TEXT("E"), TEXT("D")}) == 11);
+		CHECK(FindFirstOfAny(TEXT("AbCABCAbcABD"), {TEXT("E"), TEXT("d")}, ESearchCase::IgnoreCase) == 11);
+
+		CHECK(FindFirstOfAny(TEXT("AbCABCAbCABC"), {TEXT("AbC")}) == 0);
+		CHECK(FindFirstOfAny(TEXT("AbCABCAbCABC"), {TEXT("CABc"), TEXT("ABC")}) == 3);
+		CHECK(FindFirstOfAny(TEXT("AbCABCAbCABC"), {TEXT("ABD"), TEXT("Bc")}, ESearchCase::IgnoreCase) == 1);
+		CHECK(FindFirstOfAny(TEXT("AbCABCAbCABC"), {TEXT("bc"), TEXT("ab")}) == INDEX_NONE);
+		CHECK(FindFirstOfAny(TEXT("AbCABCAbCABC"), {TEXT("DA"), TEXT("CD")}, ESearchCase::IgnoreCase) == INDEX_NONE);
+		CHECK(FindFirstOfAny(TEXT("AbCABCAbCABD"), {TEXT("BD"), TEXT("CABB")}) == 10);
+		CHECK(FindFirstOfAny(TEXT("AbCABCAbCABD"), {TEXT("Bd"), TEXT("CABB")}, ESearchCase::IgnoreCase) == 10);
+
+		CHECK(FindFirstOfAny(TEXT(""), {TEXT("A"), TEXT("B")}) == INDEX_NONE);
+		CHECK(FindFirstOfAny(TEXT("A"), {TEXT("A"), TEXT("B")}) == 0);
+		CHECK(FindFirstOfAny(TEXT("A"), {TEXT("A"), TEXT("B")}, ESearchCase::IgnoreCase) == 0);
+		CHECK(FindFirstOfAny(TEXT("ABC"), {TEXT("ABC"), TEXT("BC")}) == 0);
+		CHECK(FindFirstOfAny(TEXT("ABC"), {TEXT("abc"), TEXT("bc")}, ESearchCase::IgnoreCase) == 0);
+		CHECK(FindFirstOfAny(TEXT("AB"), {TEXT("ABC"), TEXT("ABD")}) == INDEX_NONE);
+
+		CHECK(FindFirstOfAny("AbCABCAbCABC", {"CABc", "ABC"}) == 3);
+
+		CHECK(FindFirstOfAny(FStringView(nullptr, 0), {TEXT("ABC"), TEXT("ABD")}) == INDEX_NONE);
+		CHECK(FindFirstOfAny(FStringView(), {TEXT("ABC"), TEXT("ABD")}) == INDEX_NONE);
+	}
+
+	SECTION("FindLastOfAny")
+	{
+		CHECK(FindLastOfAny(TEXT("AbCABCAbcABC"), {TEXT("a"), TEXT("b")}) == 7);
+		CHECK(FindLastOfAny(TEXT("AbCABCAbcABC"), {TEXT("a"), TEXT("b")}, ESearchCase::IgnoreCase) == 10);
+		CHECK(FindLastOfAny(TEXT("AbCABCAbcABC"), {TEXT("b")}) == 7);
+		CHECK(FindLastOfAny(TEXT("AbCABCAbcABC"), {TEXT("c"), TEXT("B")}) == 10);
+		CHECK(FindLastOfAny(TEXT("AbCABCAbcABC"), {TEXT("c"), TEXT("B")}, ESearchCase::IgnoreCase) == 11);
+		CHECK(FindLastOfAny(TEXT("AbCABCAbcABC"), {TEXT("D"), TEXT("a")}) == INDEX_NONE);
+		CHECK(FindLastOfAny(TEXT("AbCABCAbcABC"), {TEXT("E"), TEXT("D")}, ESearchCase::IgnoreCase) == INDEX_NONE);
+		CHECK(FindLastOfAny(TEXT("AbCABCAbcABD"), {TEXT("E"), TEXT("D")}) == 11);
+		CHECK(FindLastOfAny(TEXT("AbCABCAbcABD"), {TEXT("E"), TEXT("d")}, ESearchCase::IgnoreCase) == 11);
+
+		CHECK(FindLastOfAny(TEXT("AbCABCAbCABC"), {TEXT("AbC")}) == 6);
+		CHECK(FindLastOfAny(TEXT("AbCABCAbCABC"), {TEXT("CABc"), TEXT("ABC")}) == 9);
+		CHECK(FindLastOfAny(TEXT("AbCABCAbCABC"), {TEXT("ABD"), TEXT("Bc")}, ESearchCase::IgnoreCase) == 10);
+		CHECK(FindLastOfAny(TEXT("AbCABCAbCABC"), {TEXT("bc"), TEXT("ab")}) == INDEX_NONE);
+		CHECK(FindLastOfAny(TEXT("AbCABCAbCABC"), {TEXT("DA"), TEXT("CD")}, ESearchCase::IgnoreCase) == INDEX_NONE);
+		CHECK(FindLastOfAny(TEXT("AbCABCAbCABD"), {TEXT("BD"), TEXT("CABB")}) == 10);
+		CHECK(FindLastOfAny(TEXT("AbCABCAbCABD"), {TEXT("Bd"), TEXT("CABB")}, ESearchCase::IgnoreCase) == 10);
+
+		CHECK(FindLastOfAny(TEXT(""), {TEXT("A"), TEXT("B")}) == INDEX_NONE);
+		CHECK(FindLastOfAny(TEXT("A"), {TEXT("A"), TEXT("B")}) == 0);
+		CHECK(FindLastOfAny(TEXT("A"), {TEXT("A"), TEXT("B")}, ESearchCase::IgnoreCase) == 0);
+		CHECK(FindLastOfAny(TEXT("ABC"), {TEXT("ABC"), TEXT("BC")}) == 1);
+		CHECK(FindLastOfAny(TEXT("ABC"), {TEXT("abc"), TEXT("bc")}, ESearchCase::IgnoreCase) == 1);
+		CHECK(FindLastOfAny(TEXT("AB"), {TEXT("ABC"), TEXT("ABD")}) == INDEX_NONE);
+
+		CHECK(FindLastOfAny("AbCABCAbCABC", {"CABc", "ABC"}) == 9);
+
+		CHECK(FindLastOfAny(FStringView(nullptr, 0), { TEXT("ABC"), TEXT("ABD") }) == INDEX_NONE);
+		CHECK(FindLastOfAny(FStringView(), { TEXT("ABC"), TEXT("ABD") }) == INDEX_NONE);
+	}
+
+	SECTION("FindFirstChar")
+	{
+		CHECK(FindFirstChar(TEXT("AbCABCAbCABC"), TEXT('b')) == 1);
+		CHECK(FindFirstChar(TEXT("AbCABCAbCABC"), TEXT('B')) == 4);
+		CHECK(FindFirstChar(TEXT("AbCABCAbCABC"), TEXT('B'), ESearchCase::IgnoreCase) == 1);
+		CHECK(FindFirstChar(TEXT("AbCABCAbCABC"), TEXT('a')) == INDEX_NONE);
+		CHECK(FindFirstChar(TEXT("AbCABCAbCABC"), TEXT('D'), ESearchCase::IgnoreCase) == INDEX_NONE);
+		CHECK(FindFirstChar(TEXT("AbCABCAbCABD"), TEXT('D')) == 11);
+		CHECK(FindFirstChar(TEXT("AbCABCAbCABD"), TEXT('d'), ESearchCase::IgnoreCase) == 11);
+
+		CHECK(FindFirstChar(TEXT(""), TEXT('A')) == INDEX_NONE);
+		CHECK(FindFirstChar(TEXT("A"), TEXT('A')) == 0);
+		CHECK(FindFirstChar(TEXT("A"), TEXT('A'), ESearchCase::IgnoreCase) == 0);
+
+		CHECK(FindFirstChar("AbCABCAbCABC", 'B') == 4);
+
+		CHECK(FindFirstChar(FStringView(nullptr, 0), TEXT('A')) == INDEX_NONE);
+		CHECK(FindFirstChar(FStringView(), TEXT('A')) == INDEX_NONE);
+	}
+
+	SECTION("FindLastChar")
+	{
+		CHECK(FindLastChar(TEXT("AbCABCAbCABC"), TEXT('b')) == 7);
+		CHECK(FindLastChar(TEXT("AbCABCAbCABC"), TEXT('B')) == 10);
+		CHECK(FindLastChar(TEXT("AbCABCAbCABC"), TEXT('b'), ESearchCase::IgnoreCase) == 10);
+		CHECK(FindLastChar(TEXT("AbCABCAbCABC"), TEXT('a')) == INDEX_NONE);
+		CHECK(FindLastChar(TEXT("AbCABCAbCABC"), TEXT('D'), ESearchCase::IgnoreCase) == INDEX_NONE);
+		CHECK(FindLastChar(TEXT("AbCABCAbCABD"), TEXT('D')) == 11);
+		CHECK(FindLastChar(TEXT("AbCABCAbCABD"), TEXT('d'), ESearchCase::IgnoreCase) == 11);
+
+		CHECK(FindLastChar(TEXT(""), TEXT('A')) == INDEX_NONE);
+		CHECK(FindLastChar(TEXT("A"), TEXT('A')) == 0);
+		CHECK(FindLastChar(TEXT("A"), TEXT('A'), ESearchCase::IgnoreCase) == 0);
+
+		CHECK(FindLastChar("AbCABCAbCABC", 'B') == 10);
+
+		CHECK(FindLastChar(FStringView(nullptr, 0), TEXT('A')) == INDEX_NONE);
+		CHECK(FindLastChar(FStringView(), TEXT('A')) == INDEX_NONE);
+	}
+
+	SECTION("FindFirstOfAnyChar")
+	{
+		CHECK(FindFirstOfAnyChar(TEXT("AbCABCAbcABC"), {TEXT('b')}) == 1);
+		CHECK(FindFirstOfAnyChar(TEXT("AbCABCAbcABC"), {TEXT('c'), TEXT('B')}) == 4);
+		CHECK(FindFirstOfAnyChar(TEXT("AbCABCAbcABC"), {TEXT('c'), TEXT('B')}, ESearchCase::IgnoreCase) == 1);
+		CHECK(FindFirstOfAnyChar(TEXT("AbCABCAbcABC"), {TEXT('D'), TEXT('a')}) == INDEX_NONE);
+		CHECK(FindFirstOfAnyChar(TEXT("AbCABCAbcABC"), {TEXT('E'), TEXT('D')}, ESearchCase::IgnoreCase) == INDEX_NONE);
+		CHECK(FindFirstOfAnyChar(TEXT("AbCABCAbcABD"), {TEXT('E'), TEXT('D')}) == 11);
+		CHECK(FindFirstOfAnyChar(TEXT("AbCABCAbcABD"), {TEXT('E'), TEXT('d')}, ESearchCase::IgnoreCase) == 11);
+
+		CHECK(FindFirstOfAnyChar(TEXT(""), {TEXT('A'), TEXT('B')}) == INDEX_NONE);
+		CHECK(FindFirstOfAnyChar(TEXT("A"), {TEXT('A'), TEXT('B')}) == 0);
+		CHECK(FindFirstOfAnyChar(TEXT("A"), {TEXT('A'), TEXT('B')}, ESearchCase::IgnoreCase) == 0);
+
+		CHECK(FindFirstOfAnyChar("AbCABCAbcABC", {'c', 'B'}) == 4);
+
+		CHECK(FindFirstOfAnyChar(FStringView(nullptr, 0), { TEXT('A'), TEXT('B') }) == INDEX_NONE);
+		CHECK(FindFirstOfAnyChar(FStringView(), { TEXT('A'), TEXT('B') }) == INDEX_NONE);
+	}
+
+	SECTION("FindLastOfAnyChar")
+	{
+		CHECK(FindLastOfAnyChar(TEXT("AbCABCAbcABC"), {TEXT('b')}) == 7);
+		CHECK(FindLastOfAnyChar(TEXT("AbCABCAbcABC"), {TEXT('c'), TEXT('B')}) == 10);
+		CHECK(FindLastOfAnyChar(TEXT("AbCABCAbcABC"), {TEXT('c'), TEXT('B')}, ESearchCase::IgnoreCase) == 11);
+		CHECK(FindLastOfAnyChar(TEXT("AbCABCAbcABC"), {TEXT('D'), TEXT('a')}) == INDEX_NONE);
+		CHECK(FindLastOfAnyChar(TEXT("AbCABCAbcABC"), {TEXT('E'), TEXT('D')}, ESearchCase::IgnoreCase) == INDEX_NONE);
+		CHECK(FindLastOfAnyChar(TEXT("AbCABCAbcABD"), {TEXT('E'), TEXT('D')}) == 11);
+		CHECK(FindLastOfAnyChar(TEXT("AbCABCAbcABD"), {TEXT('E'), TEXT('d')}, ESearchCase::IgnoreCase) == 11);
+
+		CHECK(FindLastOfAnyChar(TEXT(""), {TEXT('A'), TEXT('B')}) == INDEX_NONE);
+		CHECK(FindLastOfAnyChar(TEXT("A"), {TEXT('A'), TEXT('B')}) == 0);
+		CHECK(FindLastOfAnyChar(TEXT("A"), {TEXT('A'), TEXT('B')}, ESearchCase::IgnoreCase) == 0);
+
+		CHECK(FindLastOfAnyChar("AbCABCAbcABC", {'c', 'B'}) == 10);
+
+		CHECK(FindLastOfAnyChar(FStringView(nullptr, 0), { TEXT('A'), TEXT('B') }) == INDEX_NONE);
+		CHECK(FindLastOfAnyChar(FStringView(), { TEXT('A'), TEXT('B') }) == INDEX_NONE);
+	}
+}
+
+} // UE::String
diff --git a/Engine/Source/Runtime/Core/Tests/String/ParseLinesTest.cpp b/Engine/Source/Runtime/Core/Tests/String/ParseLinesTest.cpp
new file mode 100644
--- /dev/null
+++ b/Engine/Source/Runtime/Core/Tests/String/ParseLinesTest.cpp
@@ -0,0 +1,82 @@
+// Copyright Epic Games, Inc. All Rights Reserved.
+
+#include "String/ParseLines.h"
+
+#include "Algo/Compare.h"
+#include "Containers/StringView.h"
+#include "Misc/StringBuilder.h"
+#include "String/Escape.h"
+#include "String/Join.h"
+#include "TestHarness.h"
+
+namespace UE::String
+{
+
+TEST_CASE("Core::String::ParseLines", "[Core][String][Smoke]")
+{
+	using FStringViewArray = TArray<FStringView, TInlineAllocator<8>>;
+
+	constexpr EParseLinesOptions KeepEmpty = EParseLinesOptions::None;
+	constexpr EParseLinesOptions SkipEmpty = EParseLinesOptions::SkipEmpty;
+	constexpr EParseLinesOptions Trim = EParseLinesOptions::Trim;
+
+	const auto [View, ExpectedLines, Options] = GENERATE_COPY(table<FStringView, FStringViewArray, EParseLinesOptions>(
+	{
+		{TEXTVIEW(""), {}, SkipEmpty},
+		{TEXTVIEW(""), {TEXTVIEW("")}, KeepEmpty},
+		{TEXTVIEW("\n"), {}, SkipEmpty},
+		{TEXTVIEW("\n"), {TEXTVIEW("")}, KeepEmpty},
+		{TEXTVIEW("\r"), {}, SkipEmpty},
+		{TEXTVIEW("\r"), {TEXTVIEW("")}, KeepEmpty},
+		{TEXTVIEW("\r\n"), {}, SkipEmpty},
+		{TEXTVIEW("\r\n"), {TEXTVIEW("")}, KeepEmpty},
+		{TEXTVIEW("\n\n"), {}, SkipEmpty},
+		{TEXTVIEW("\n\n"), {TEXTVIEW(""), TEXTVIEW("")}, KeepEmpty},
+		{TEXTVIEW("\r\r"), {}, SkipEmpty},
+		{TEXTVIEW("\r\r"), {TEXTVIEW(""), TEXTVIEW("")}, KeepEmpty},
+		{TEXTVIEW("\r\n\r\n"), {}, SkipEmpty},
+		{TEXTVIEW("\r\n\r\n"), {TEXTVIEW(""), TEXTVIEW("")}, KeepEmpty},
+		{TEXTVIEW("\r\nABC").Left(2), {}, SkipEmpty},
+		{TEXTVIEW("\r\nABC").Left(2), {TEXTVIEW("")}, KeepEmpty},
+		{TEXTVIEW("\r\nABC\r\nDEF").Left(5), {TEXTVIEW("ABC")}, SkipEmpty},
+		{TEXTVIEW("\r\nABC\r\nDEF").Left(5), {TEXTVIEW(""), TEXTVIEW("ABC")}, KeepEmpty},
+		{TEXTVIEW("ABC DEF"), {TEXTVIEW("ABC DEF")}, SkipEmpty},
+		{TEXTVIEW("\nABC DEF\n"), {TEXTVIEW("ABC DEF")}, SkipEmpty},
+		{TEXTVIEW("\nABC DEF\n"), {TEXTVIEW(""), TEXTVIEW("ABC DEF")}, KeepEmpty},
+		{TEXTVIEW("\rABC DEF\r"), {TEXTVIEW("ABC DEF")}, SkipEmpty},
+		{TEXTVIEW("\rABC DEF\r"), {TEXTVIEW(""), TEXTVIEW("ABC DEF")}, KeepEmpty},
+		{TEXTVIEW("\r\nABC DEF\r\n"), {TEXTVIEW("ABC DEF")}, SkipEmpty},
+		{TEXTVIEW("\r\nABC DEF\r\n"), {TEXTVIEW(""), TEXTVIEW("ABC DEF")}, KeepEmpty},
+		{TEXTVIEW("\r\n\r\nABC DEF\r\n\r\n"), {TEXTVIEW("ABC DEF")}, SkipEmpty},
+		{TEXTVIEW("\r\n\r\nABC DEF\r\n\r\n"), {TEXTVIEW(""), TEXTVIEW(""), TEXTVIEW("ABC DEF"), TEXTVIEW("")}, KeepEmpty},
+		{TEXTVIEW("ABC\nDEF"), {TEXTVIEW("ABC"), TEXTVIEW("DEF")}, SkipEmpty},
+		{TEXTVIEW("ABC\rDEF"), {TEXTVIEW("ABC"), TEXTVIEW("DEF")}, SkipEmpty},
+		{TEXTVIEW("\r\nABC\r\nDEF\r\n"), {TEXTVIEW("ABC"), TEXTVIEW("DEF")}, SkipEmpty},
+		{TEXTVIEW("\r\nABC\r\nDEF\r\n"), {TEXTVIEW(""), TEXTVIEW("ABC"), TEXTVIEW("DEF")}, KeepEmpty},
+		{TEXTVIEW("\r\nABC\r\n\r\nDEF\r\n"), {TEXTVIEW("ABC"), TEXTVIEW("DEF")}, SkipEmpty},
+		{TEXTVIEW("\r\nABC\r\n\r\nDEF\r\n"), {TEXTVIEW(""), TEXTVIEW("ABC"), TEXTVIEW(""), TEXTVIEW("DEF")}, KeepEmpty},
+		{TEXTVIEW(" \t\r\n\t ABC \t\r\n\t \t\r\n\t DEF \t\r\n"), {TEXTVIEW(" \t"), TEXTVIEW("\t ABC \t"), TEXTVIEW("\t \t"), TEXTVIEW("\t DEF \t")}, SkipEmpty},
+		{TEXTVIEW(" \t\r\n\t ABC \t\r\n\t \t\r\n\t DEF \t\r\n"), {TEXTVIEW("ABC"), TEXTVIEW("DEF")}, SkipEmpty | Trim},
+		{TEXTVIEW(" \t\r\n\t ABC \t\r\n\t \t\r\n\t DEF \t\r\n"), {TEXTVIEW(" \t"), TEXTVIEW("\t ABC \t"), TEXTVIEW("\t \t"), TEXTVIEW("\t DEF \t")}, KeepEmpty},
+		{TEXTVIEW(" \t\r\n\t ABC \t\r\n\t \t\r\n\t DEF \t\r\n"), {TEXTVIEW(""), TEXTVIEW("ABC"), TEXTVIEW(""), TEXTVIEW("DEF")}, KeepEmpty | Trim},
+	}));
+
+	FStringViewArray ActualLines;
+	ParseLines(View, ActualLines, Options);
+	const bool bEqual = Algo::Compare(ActualLines, ExpectedLines);
+	if (!bEqual)
+	{
+		TStringBuilder<256> Input, Expected, Actual;
+		Input << QuoteEscape(View);
+		Expected << TEXTVIEW("[") << JoinBy(ExpectedLines, QuoteEscape, TEXTVIEW(", ")) << TEXTVIEW("]");
+		Actual << TEXTVIEW("[") << JoinBy(ActualLines, QuoteEscape, TEXTVIEW(", ")) << TEXTVIEW("]");
+		CAPTURE(Input, Expected, Actual);
+		CHECK(bEqual);
+	}
+	else
+	{
+		CHECK(bEqual);
+	}
+}
+
+} // UE::String
diff --git a/Engine/Source/Runtime/Core/Tests/String/ParseTokensTest.cpp b/Engine/Source/Runtime/Core/Tests/String/ParseTokensTest.cpp
new file mode 100644
--- /dev/null
+++ b/Engine/Source/Runtime/Core/Tests/String/ParseTokensTest.cpp
@@ -0,0 +1,164 @@
+// Copyright Epic Games, Inc. All Rights Reserved.
+
+#include "String/ParseTokens.h"
+
+#include "Algo/Compare.h"
+#include "Misc/StringBuilder.h"
+#include "String/Escape.h"
+#include "String/Join.h"
+#include "TestHarness.h"
+
+namespace UE::String
+{
+
+TEST_CASE("Core::String::ParseTokens", "[Core][String][Smoke]")
+{
+	using FCharArray = TArray<TCHAR, TInlineAllocator<8>>;
+	using FStringViewArray = TArray<FStringView, TInlineAllocator<8>>;
+
+	constexpr EParseTokensOptions KeepEmpty = EParseTokensOptions::None;
+	constexpr EParseTokensOptions SkipEmpty = EParseTokensOptions::SkipEmpty;
+	constexpr EParseTokensOptions IgnoreCase = EParseTokensOptions::IgnoreCase;
+	constexpr EParseTokensOptions Trim = EParseTokensOptions::Trim;
+
+	SECTION("String")
+	{
+		const auto [View, DelimiterArray, ExpectedTokens, Options] = GENERATE_COPY(table<FStringView, FStringViewArray, FStringViewArray, EParseTokensOptions>(
+		{
+			{TEXTVIEW(""),                {},                   {},                                                     SkipEmpty},
+			{TEXTVIEW(""),                {},                   {TEXTVIEW("")},                                         KeepEmpty},
+			{TEXTVIEW("ABC"),             {},                   {TEXTVIEW("ABC")},                                      KeepEmpty},
+
+			{TEXTVIEW(""),                {TEXTVIEW(",")},      {},                                                     SkipEmpty},
+			{TEXTVIEW(""),                {TEXTVIEW(",")},      {TEXTVIEW("")},                                         KeepEmpty},
+			{TEXTVIEW(","),               {TEXTVIEW(",")},      {},                                                     SkipEmpty},
+			{TEXTVIEW(","),               {TEXTVIEW(",")},      {TEXTVIEW(""), TEXTVIEW("")},                           KeepEmpty},
+			{TEXTVIEW(",,"),              {TEXTVIEW(",")},      {},                                                     SkipEmpty},
+			{TEXTVIEW(",,"),              {TEXTVIEW(",")},      {TEXTVIEW(""), TEXTVIEW(""), TEXTVIEW("")},             KeepEmpty},
+			{TEXTVIEW(", ,"),             {TEXTVIEW(",")},      {TEXTVIEW(" ")},                                        SkipEmpty},
+			{TEXTVIEW(", ,"),             {TEXTVIEW(",")},      {},                                                     SkipEmpty | Trim},
+			{TEXTVIEW(", ,"),             {TEXTVIEW(",")},      {TEXTVIEW(""), TEXTVIEW(" "), TEXTVIEW("")},            KeepEmpty},
+			{TEXTVIEW(", ,"),             {TEXTVIEW(",")},      {TEXTVIEW(""), TEXTVIEW(""), TEXTVIEW("")},             KeepEmpty | Trim},
+			{TEXTVIEW("ABC"),             {TEXTVIEW(",")},      {TEXTVIEW("ABC")},                                      KeepEmpty},
+			{TEXTVIEW("A,,C"),            {TEXTVIEW(",")},      {TEXTVIEW("A"), TEXTVIEW("C")},                         SkipEmpty},
+			{TEXTVIEW("A,,C"),            {TEXTVIEW(",")},      {TEXTVIEW("A"), TEXTVIEW(""), TEXTVIEW("C")},           KeepEmpty},
+			{TEXTVIEW("A,\tB\t,C"),       {TEXTVIEW(",")},      {TEXTVIEW("A"), TEXTVIEW("\tB\t"), TEXTVIEW("C")},      KeepEmpty},
+			{TEXTVIEW(",A, B ,C,"),       {TEXTVIEW(",")},      {TEXTVIEW("A"), TEXTVIEW(" B "), TEXTVIEW("C")},        SkipEmpty},
+			{TEXTVIEW(",A, B ,C,"),       {TEXTVIEW(",")},      {TEXTVIEW(""), TEXTVIEW("A"), TEXTVIEW(" B "), TEXTVIEW("C"), TEXTVIEW("")}, KeepEmpty},
+			{TEXTVIEW("A\u2022B\u2022C"), {TEXTVIEW("\u2022")}, {TEXTVIEW("A"), TEXTVIEW("B"), TEXTVIEW("C")},          KeepEmpty},
+
+			{TEXTVIEW("ABCDABCD"), {TEXTVIEW("AB")},   {TEXTVIEW("CD"), TEXTVIEW("CD")},               SkipEmpty},
+			{TEXTVIEW("ABCDABCD"), {TEXTVIEW("AB")},   {TEXTVIEW(""), TEXTVIEW("CD"), TEXTVIEW("CD")}, KeepEmpty},
+			{TEXTVIEW("ABCDABCD"), {TEXTVIEW("ABCD")}, {},                                             SkipEmpty},
+			{TEXTVIEW("ABCDABCD"), {TEXTVIEW("ABCD")}, {TEXTVIEW(""), TEXTVIEW(""), TEXTVIEW("")},     KeepEmpty},
+			{TEXTVIEW("ABCDABCD"), {TEXTVIEW("DA")},   {TEXTVIEW("ABC"), TEXTVIEW("BCD")},             KeepEmpty},
+
+			{TEXTVIEW("ABCDABCD"), {TEXTVIEW("B"),  TEXTVIEW("D")},  {TEXTVIEW("A"), TEXTVIEW("C"), TEXTVIEW("A"), TEXTVIEW("C")},               SkipEmpty},
+			{TEXTVIEW("ABCDABCD"), {TEXTVIEW("B"),  TEXTVIEW("D")},  {TEXTVIEW("A"), TEXTVIEW("C"), TEXTVIEW("A"), TEXTVIEW("C"), TEXTVIEW("")}, KeepEmpty},
+			{TEXTVIEW("ABCDABCD"), {TEXTVIEW("BC"), TEXTVIEW("DA")}, {TEXTVIEW("A"), TEXTVIEW("D")},                                             SkipEmpty},
+			{TEXTVIEW("ABCDABCD"), {TEXTVIEW("BC"), TEXTVIEW("DA")}, {TEXTVIEW("A"), TEXTVIEW(""), TEXTVIEW(""), TEXTVIEW("D")},                 KeepEmpty},
+
+			{TEXTVIEW("AbCdaBcDAbCd"), {TEXTVIEW("Bc"), TEXTVIEW("da")}, {TEXTVIEW("AbC"), TEXTVIEW("DAbCd")}, SkipEmpty},
+			{TEXTVIEW("AbCdaBcDAbCd"), {TEXTVIEW("Bc"), TEXTVIEW("da")}, {TEXTVIEW("A"), TEXTVIEW("d")},       SkipEmpty | IgnoreCase},
+
+			{TEXTVIEW("A\u2022\u2022B,,C"), {TEXTVIEW(",,"), TEXTVIEW("\u2022\u2022")},                     {TEXTVIEW("A"), TEXTVIEW("B"), TEXTVIEW("C")}, KeepEmpty},
+			{TEXTVIEW("A\u2022\u2022B\u0085\u0085C"), {TEXTVIEW("\u0085\u0085"), TEXTVIEW("\u2022\u2022")}, {TEXTVIEW("A"), TEXTVIEW("B"), TEXTVIEW("C")}, KeepEmpty},
+		}));
+
+		FStringViewArray ActualTokens;
+		if (GetNum(DelimiterArray) == 1)
+		{
+			ParseTokens(View, GetData(DelimiterArray)[0], ActualTokens, Options);
+		}
+		else
+		{
+			ParseTokensMultiple(View, DelimiterArray, ActualTokens, Options);
+		}
+
+		const bool bEqual = Algo::Compare(ActualTokens, ExpectedTokens);
+		if (!bEqual)
+		{
+			TStringBuilder<256> Input, Delimiters, Expected, Actual;
+			Input << QuoteEscape(View);
+			Delimiters << TEXTVIEW("[") << JoinBy(DelimiterArray, QuoteEscape, TEXTVIEW(", ")) << TEXTVIEW("]");
+			Expected << TEXTVIEW("[") << JoinBy(ExpectedTokens, QuoteEscape, TEXTVIEW(", ")) << TEXTVIEW("]");
+			Actual << TEXTVIEW("[") << JoinBy(ActualTokens, QuoteEscape, TEXTVIEW(", ")) << TEXTVIEW("]");
+			CAPTURE(Input, Delimiters, Expected, Actual);
+			CHECK(bEqual);
+		}
+		else
+		{
+			CHECK(bEqual);
+		}
+	}
+
+	SECTION("Char")
+	{
+		const auto [View, DelimiterArray, ExpectedTokens, Options] = GENERATE_COPY(table<FStringView, FCharArray, FStringViewArray, EParseTokensOptions>(
+		{
+			{TEXTVIEW(""),                {},               {},                                    SkipEmpty},
+			{TEXTVIEW(""),                {},               {TEXT("")},                            KeepEmpty},
+			{TEXTVIEW("ABC"),             {},               {TEXT("ABC")},                         KeepEmpty},
+
+			{TEXTVIEW(""),                {TEXT(',')},      {},                                    SkipEmpty},
+			{TEXTVIEW(""),                {TEXT(',')},      {TEXT("")},                            KeepEmpty},
+			{TEXTVIEW(","),               {TEXT(',')},      {},                                    SkipEmpty},
+			{TEXTVIEW(","),               {TEXT(',')},      {TEXT(""), TEXT("")},                  KeepEmpty},
+			{TEXTVIEW(",,"),              {TEXT(',')},      {},                                    SkipEmpty},
+			{TEXTVIEW(",,"),              {TEXT(',')},      {TEXT(""), TEXT(""), TEXT("")},        KeepEmpty},
+			{TEXTVIEW(", ,"),             {TEXT(',')},      {TEXT(" ")},                           SkipEmpty},
+			{TEXTVIEW(", ,"),             {TEXT(',')},      {},                                    SkipEmpty | Trim},
+			{TEXTVIEW(", ,"),             {TEXT(',')},      {TEXT(""), TEXT(" "), TEXT("")},       KeepEmpty},
+			{TEXTVIEW(", ,"),             {TEXT(',')},      {TEXT(""), TEXT(""), TEXT("")},        KeepEmpty | Trim},
+			{TEXTVIEW("ABC"),             {TEXT(',')},      {TEXT("ABC")},                         KeepEmpty},
+			{TEXTVIEW("A,,C"),            {TEXT(',')},      {TEXT("A"), TEXT("C")},                SkipEmpty},
+			{TEXTVIEW("A,,C"),            {TEXT(',')},      {TEXT("A"), TEXT(""), TEXT("C")},      KeepEmpty},
+			{TEXTVIEW("A,\tB\t,C"),       {TEXT(',')},      {TEXT("A"), TEXT("\tB\t"), TEXT("C")}, KeepEmpty},
+			{TEXTVIEW(",A, B ,C,"),       {TEXT(',')},      {TEXT("A"), TEXT(" B "), TEXT("C")},                     SkipEmpty},
+			{TEXTVIEW(",A, B ,C,"),       {TEXT(',')},      {TEXT(""), TEXT("A"), TEXT(" B "), TEXT("C"), TEXT("")}, KeepEmpty},
+			{TEXTVIEW("A\u2022B\u2022C"), {TEXT('\u2022')}, {TEXT("A"), TEXT("B"), TEXT("C")},     KeepEmpty},
+
+			{TEXTVIEW("ABCDABCD"),        {TEXT('B'), TEXT('D')},           {TEXT("A"), TEXT("C"), TEXT("A"), TEXT("C")},           SkipEmpty},
+			{TEXTVIEW("ABCDABCD"),        {TEXT('B'), TEXT('D')},           {TEXT("A"), TEXT("C"), TEXT("A"), TEXT("C"), TEXT("")}, KeepEmpty},
+			{TEXTVIEW("A\u2022B,C"),      {TEXT(','), TEXT('\u2022')},      {TEXT("A"), TEXT("B"), TEXT("C")},                      KeepEmpty},
+			{TEXTVIEW("A\u2022B\u0085C"), {TEXT('\u0085'), TEXT('\u2022')}, {TEXT("A"), TEXT("B"), TEXT("C")},                      KeepEmpty},
+
+			{TEXTVIEW("ABC"), {TEXT('b')}, {TEXT("ABC")}, SkipEmpty},
+			{TEXTVIEW("ABC"), {TEXT('b')}, {TEXT("A"), TEXT("C")}, SkipEmpty | IgnoreCase},
+
+			{TEXTVIEW("AbCdaBcD"), {TEXT('B'), TEXT('d')}, {TEXT("AbC"), TEXT("A"), TEXT("cD")},         SkipEmpty},
+			{TEXTVIEW("AbCdaBcD"), {TEXT('B'), TEXT('d')}, {TEXT("A"), TEXT("C"), TEXT("a"), TEXT("c")}, SkipEmpty | IgnoreCase},
+
+			{TEXTVIEW("A\u2022B\u2022C"), {TEXT('\u2022'), TEXT('b')}, {TEXT("A"), TEXT("B"), TEXT("C")}, SkipEmpty},
+			{TEXTVIEW("A\u2022B\u2022C"), {TEXT('\u2022'), TEXT('b')}, {TEXT("A"), TEXT("C")},            SkipEmpty | IgnoreCase},
+		}));
+
+		FStringViewArray ActualTokens;
+		if (GetNum(DelimiterArray) == 1)
+		{
+			ParseTokens(View, GetData(DelimiterArray)[0], ActualTokens, Options);
+		}
+		else
+		{
+			ParseTokensMultiple(View, DelimiterArray, ActualTokens, Options);
+		}
+
+		const bool bEqual = Algo::Compare(ActualTokens, ExpectedTokens);
+		if (!bEqual)
+		{
+			TStringBuilder<256> Input, Delimiters, Expected, Actual;
+			Input << QuoteEscape(View);
+			Delimiters << TEXTVIEW("[") << JoinQuoted(DelimiterArray, TEXTVIEW(", "), TEXT('\'')) << TEXTVIEW("]");
+			Expected << TEXTVIEW("[") << JoinBy(ExpectedTokens, QuoteEscape, TEXTVIEW(", ")) << TEXTVIEW("]");
+			Actual << TEXTVIEW("[") << JoinBy(ActualTokens, QuoteEscape, TEXTVIEW(", ")) << TEXTVIEW("]");
+			CAPTURE(Input, Delimiters, Expected, Actual);
+			CHECK(bEqual);
+		}
+		else
+		{
+			CHECK(bEqual);
+		}
+	}
+}
+
+} // UE::String
