[Insights] Set default parameter value for channels in Start

diff --git a/Engine/Source/Runtime/Core/Public/ProfilingDebugging/TraceAuxiliary.h b/Engine/Source/Runtime/Core/Public/ProfilingDebugging/TraceAuxiliary.h
--- a/Engine/Source/Runtime/Core/Public/ProfilingDebugging/TraceAuxiliary.h
+++ b/Engine/Source/Runtime/Core/Public/ProfilingDebugging/TraceAuxiliary.h
@@ -48,12 +48,12 @@ public:
 	 * already active this call does nothing.
 	 * @param Type Type of connection
 	 * @param Target String to use for connection. See /ref EConnectionType for details.
-	 * @param Channels Comma separated list of channels to enable. If the pointer is null the default channels will be active.
+	 * @param Channels Comma separated list of channels to enable. Default set of channels are enabled if argument is not specified. If the pointer is null no channels are enabled.
 	 * @param Options Optional additional tracing options.
 	 * @param LogCategory Log channel to output messages to. Default set to 'Core'.
 	 * @return True when successfully starting the trace, false if the data connection could not be made.
 	 */
-	static bool Start(EConnectionType Type, const TCHAR* Target, const TCHAR* Channels, Options* Options = nullptr, const FLogCategoryAlias& LogCategory = LogCore);
+	static bool Start(EConnectionType Type, const TCHAR* Target, const TCHAR* Channels = TEXT("default"), Options* Options = nullptr, const FLogCategoryAlias& LogCategory = LogCore);
 
 	/**
 	 * Stop tracing.
