Add `GetTypeHash` for TScriptDelegate.

diff --git a/Engine/Source/Runtime/Core/Public/UObject/ScriptDelegates.h b/Engine/Source/Runtime/Core/Public/UObject/ScriptDelegates.h
--- a/Engine/Source/Runtime/Core/Public/UObject/ScriptDelegates.h
+++ b/Engine/Source/Runtime/Core/Public/UObject/ScriptDelegates.h
@@ -247,6 +247,11 @@ public:
 		ObjectPtr->ProcessEvent(Function, Parameters);
 	}
 
+	friend uint32 GetTypeHash(const TScriptDelegate& Delegate)
+	{
+		return HashCombine(GetTypeHash(Delegate.Object), GetTypeHash(Delegate.GetFunctionName()));
+	}
+
 protected:
 
 	/** The object bound to this delegate, or nullptr if no object is bound */
