Remove SFINAE from TNonNullPtr to allow for instantiation with incomplete types

diff --git a/Engine/Source/Runtime/Core/Public/Templates/NonNullPointer.h b/Engine/Source/Runtime/Core/Public/Templates/NonNullPointer.h
--- a/Engine/Source/Runtime/Core/Public/Templates/NonNullPointer.h
+++ b/Engine/Source/Runtime/Core/Public/Templates/NonNullPointer.h
@@ -39,10 +39,7 @@ public:
 	/**
 	 * Constructs a non-null pointer from the provided pointer. Must not be nullptr.
 	 */
-	template <
-		typename OtherObjectType,
-		typename = typename TEnableIf<TPointerIsConvertibleFromTo<OtherObjectType, ObjectType>::Value>::Type
-	>
+	template <typename OtherObjectType>
 	FORCEINLINE TNonNullPtr(OtherObjectType* InObject)
 		: Object(InObject)
 	{
@@ -52,10 +49,7 @@ public:
 	/**
 	 * Constructs a non-null pointer from another non-null pointer
 	 */
-	template <
-		typename OtherObjectType,
-		typename = typename TEnableIf<TPointerIsConvertibleFromTo<OtherObjectType, ObjectType>::Value>::Type
-	>
+	template <typename OtherObjectType>
 	FORCEINLINE TNonNullPtr(TNonNullPtr<OtherObjectType>& Other)
 		: Object(Other.Object)
 	{
@@ -73,10 +67,7 @@ public:
 	/**
 	 * Assignment operator taking a pointer
 	 */
-	template <
-		typename OtherObjectType,
-		typename = typename TEnableIf<TPointerIsConvertibleFromTo<OtherObjectType, ObjectType>::Value>::Type
-	>
+	template <typename OtherObjectType>
 	FORCEINLINE TNonNullPtr& operator=(OtherObjectType* InObject)
 	{
 		ensureMsgf(InObject, TEXT("Tried to assign a null pointer to a TNonNullPtr!"));
@@ -87,10 +78,7 @@ public:
 	/**
 	 * Assignment operator taking another TNonNullPtr
 	 */
-	template <
-		typename OtherObjectType,
-		typename = typename TEnableIf<TPointerIsConvertibleFromTo<OtherObjectType, ObjectType>::Value>::Type
-	>
+	template <typename OtherObjectType>
 	FORCEINLINE TNonNullPtr& operator=(TNonNullPtr<OtherObjectType>& Other)
 	{
 		Object = Other.Object;
