Fix DelegateInstances to allow using CreateWeakLambda in a const function

uses the same cast as TBaseUObjectMethodDelegateInstance

#tests created a weak lambda in a virtual function, did a full compile and stepped through the code

diff --git a/Engine/Source/Runtime/Core/Public/Delegates/DelegateInstancesImpl.h b/Engine/Source/Runtime/Core/Public/Delegates/DelegateInstancesImpl.h
--- a/Engine/Source/Runtime/Core/Public/Delegates/DelegateInstancesImpl.h
+++ b/Engine/Source/Runtime/Core/Public/Delegates/DelegateInstancesImpl.h
@@ -906,7 +906,7 @@ public:
 
 	UObject* GetUObject() const final
 	{
-		return ContextObject.Get();
+		return (UObject*)ContextObject.Get();
 	}
 
 	const void* GetObjectForTimerManager() const final
