Fix for UTF-8 decoding wrongly emitting multiple bogus chars for a correctly parsed codepoint but invalid value.

diff --git a/Engine/Source/Runtime/Core/Private/GenericPlatform/GenericPlatformString.cpp b/Engine/Source/Runtime/Core/Private/GenericPlatform/GenericPlatformString.cpp
--- a/Engine/Source/Runtime/Core/Private/GenericPlatform/GenericPlatformString.cpp
+++ b/Engine/Source/Runtime/Core/Private/GenericPlatform/GenericPlatformString.cpp
@@ -275,14 +275,16 @@ namespace UE::Core::Private
 
 		if (Octet < 128)  // one octet char: 0 to 127
 		{
-			++SourceString;  // skip to next possible start of codepoint.
-			return Octet;
+			++SourceString;  // read all octets - skip to next possible start of codepoint.
+
+			Codepoint = Octet;
 		}
 		else if (Octet < 192)  // bad (starts with 10xxxxxx).
 		{
 			// Apparently each of these is supposed to be flagged as a bogus
 			//  char, instead of just resyncing to the next valid codepoint.
-			++SourceString;  // skip to next possible start of codepoint.
+
+			++SourceString;  // Sequence was not valid UTF-8. Skip the first byte and continue.
 			return UNICODE_BOGUS_CHAR_CODEPOINT;
 		}
 		else if (Octet < 224)  // two octets
@@ -303,11 +305,12 @@ namespace UE::Core::Private
 				return UNICODE_BOGUS_CHAR_CODEPOINT;
 			}
 
+			SourceString += 2;  // read all octets - skip to next possible start of codepoint.
+
 			Codepoint = ((Octet << 6) | (Octet2 - 128));
-			if ((Codepoint >= 0x80) && (Codepoint <= 0x7FF))
+			if (Codepoint < 0x80 || Codepoint > 0x7FF)
 			{
-				SourceString += 2;  // skip to next possible start of codepoint.
-				return Codepoint;
+				return UNICODE_BOGUS_CHAR_CODEPOINT;
 			}
 		}
 		else if (Octet < 240)  // three octets
@@ -335,21 +338,15 @@ namespace UE::Core::Private
 				return UNICODE_BOGUS_CHAR_CODEPOINT;
 			}
 
+			SourceString += 3;  // read all octets - skip to next possible start of codepoint.
+
 			Codepoint = ( ((Octet << 12)) | ((Octet2-128) << 6) | ((Octet3-128)) );
 
-			// UTF-8 characters cannot be in the UTF-16 surrogates range
-			if (IsHighSurrogate(Codepoint) || IsLowSurrogate(Codepoint))
+			// UTF-8 characters cannot be in the UTF-16 surrogates range.  0xFFFE and 0xFFFF are illegal, too, so we check them at the edge.
+			if (Codepoint < 0x800 || Codepoint > 0xFFFD || IsHighSurrogate(Codepoint) || IsLowSurrogate(Codepoint))
 			{
-				++SourceString;  // Sequence was not valid UTF-8. Skip the first byte and continue.
 				return UNICODE_BOGUS_CHAR_CODEPOINT;
 			}
-
-			// 0xFFFE and 0xFFFF are illegal, too, so we check them at the edge.
-			if ((Codepoint >= 0x800) && (Codepoint <= 0xFFFD))
-			{
-				SourceString += 3;  // skip to next possible start of codepoint.
-				return Codepoint;
-			}
 		}
 		else if (Octet < 248)  // four octets
 		{
@@ -383,12 +380,13 @@ namespace UE::Core::Private
 				return UNICODE_BOGUS_CHAR_CODEPOINT;
 			}
 
+			SourceString += 4;  // read all octets - skip to next possible start of codepoint.
+
 			Codepoint = ( ((Octet << 18)) | ((Octet2 - 128) << 12) |
 						((Octet3 - 128) << 6) | ((Octet4 - 128)) );
-			if ((Codepoint >= 0x10000) && (Codepoint <= 0x10FFFF))
+			if (Codepoint < 0x10000 || Codepoint > 0x10FFFF)
 			{
-				SourceString += 4;  // skip to next possible start of codepoint.
-				return Codepoint;
+				return UNICODE_BOGUS_CHAR_CODEPOINT;
 			}
 		}
 		// Five and six octet sequences became illegal in rfc3629.
@@ -432,7 +430,8 @@ namespace UE::Core::Private
 				return UNICODE_BOGUS_CHAR_CODEPOINT;
 			}
 
-			SourceString += 5;  // skip to next possible start of codepoint.
+			SourceString += 5;  // read all octets - skip to next possible start of codepoint.
+
 			return UNICODE_BOGUS_CHAR_CODEPOINT;
 		}
 
@@ -481,12 +480,13 @@ namespace UE::Core::Private
 				return UNICODE_BOGUS_CHAR_CODEPOINT;
 			}
 
-			SourceString += 6;  // skip to next possible start of codepoint.
+			SourceString += 6;  // read all octets - skip to next possible start of codepoint.
+
 			return UNICODE_BOGUS_CHAR_CODEPOINT;
 		}
 
-		++SourceString;  // Sequence was not valid UTF-8. Skip the first byte and continue.
-		return UNICODE_BOGUS_CHAR_CODEPOINT;  // catch everything else.
+		// Should only reach here after parsing a legal codepoint
+		return Codepoint;
 	}
 
 	/**
