Revert removal of SFINAE from TNonNullPtr

diff --git a/Engine/Source/Runtime/Core/Public/Templates/NonNullPointer.h b/Engine/Source/Runtime/Core/Public/Templates/NonNullPointer.h
--- a/Engine/Source/Runtime/Core/Public/Templates/NonNullPointer.h
+++ b/Engine/Source/Runtime/Core/Public/Templates/NonNullPointer.h
@@ -39,7 +39,10 @@ public:
 	/**
 	 * Constructs a non-null pointer from the provided pointer. Must not be nullptr.
 	 */
-	template <typename OtherObjectType>
+	template <
+		typename OtherObjectType,
+		typename = typename TEnableIf<TPointerIsConvertibleFromTo<OtherObjectType, ObjectType>::Value>::Type
+	>
 	FORCEINLINE TNonNullPtr(OtherObjectType* InObject)
 		: Object(InObject)
 	{
@@ -49,7 +52,10 @@ public:
 	/**
 	 * Constructs a non-null pointer from another non-null pointer
 	 */
-	template <typename OtherObjectType>
+	template <
+		typename OtherObjectType,
+		typename = typename TEnableIf<TPointerIsConvertibleFromTo<OtherObjectType, ObjectType>::Value>::Type
+	>
 	FORCEINLINE TNonNullPtr(TNonNullPtr<OtherObjectType>& Other)
 		: Object(Other.Object)
 	{
@@ -67,7 +73,10 @@ public:
 	/**
 	 * Assignment operator taking a pointer
 	 */
-	template <typename OtherObjectType>
+	template <
+		typename OtherObjectType,
+		typename = typename TEnableIf<TPointerIsConvertibleFromTo<OtherObjectType, ObjectType>::Value>::Type
+	>
 	FORCEINLINE TNonNullPtr& operator=(OtherObjectType* InObject)
 	{
 		ensureMsgf(InObject, TEXT("Tried to assign a null pointer to a TNonNullPtr!"));
@@ -78,7 +87,10 @@ public:
 	/**
 	 * Assignment operator taking another TNonNullPtr
 	 */
-	template <typename OtherObjectType>
+	template <
+		typename OtherObjectType,
+		typename = typename TEnableIf<TPointerIsConvertibleFromTo<OtherObjectType, ObjectType>::Value>::Type
+	>
 	FORCEINLINE TNonNullPtr& operator=(TNonNullPtr<OtherObjectType>& Other)
 	{
 		Object = Other.Object;
