Deprecation of PLATFORM_COMPILER_HAS_IF_CONSTEXPR.

diff --git a/Engine/Source/Runtime/Core/Public/Clang/ClangPlatform.h b/Engine/Source/Runtime/Core/Public/Clang/ClangPlatform.h
--- a/Engine/Source/Runtime/Core/Public/Clang/ClangPlatform.h
+++ b/Engine/Source/Runtime/Core/Public/Clang/ClangPlatform.h
@@ -6,9 +6,7 @@
 
 #pragma once
 
-#if defined(__cpp_if_constexpr)
-	#define PLATFORM_COMPILER_HAS_IF_CONSTEXPR 1
-#else
+#if !defined(__cpp_if_constexpr)
 	#error "Compiler is expected to support if constexpr"
 #endif
 
diff --git a/Engine/Source/Runtime/Core/Public/HAL/Platform.h b/Engine/Source/Runtime/Core/Public/HAL/Platform.h
--- a/Engine/Source/Runtime/Core/Public/HAL/Platform.h
+++ b/Engine/Source/Runtime/Core/Public/HAL/Platform.h
@@ -239,8 +239,13 @@
 #ifndef PLATFORM_COMPILER_HAS_DECLTYPE_AUTO
 	#define PLATFORM_COMPILER_HAS_DECLTYPE_AUTO 1
 #endif
-#ifndef PLATFORM_COMPILER_HAS_IF_CONSTEXPR
+#ifdef _MSC_VER
 	#define PLATFORM_COMPILER_HAS_IF_CONSTEXPR 1
+	#ifndef __clang__
+		#pragma deprecated("PLATFORM_COMPILER_HAS_IF_CONSTEXPR")
+	#endif
+#else
+	#define PLATFORM_COMPILER_HAS_IF_CONSTEXPR 1 DEPRECATED_MACRO(5.1, "PLATFORM_COMPILER_HAS_IF_CONSTEXPR has been deprecated and should be replaced with 1.")
 #endif
 #ifndef PLATFORM_COMPILER_HAS_FOLD_EXPRESSIONS
 	#define PLATFORM_COMPILER_HAS_FOLD_EXPRESSIONS 0
diff --git a/Engine/Source/Runtime/Core/Public/MSVC/MSVCPlatform.h b/Engine/Source/Runtime/Core/Public/MSVC/MSVCPlatform.h
--- a/Engine/Source/Runtime/Core/Public/MSVC/MSVCPlatform.h
+++ b/Engine/Source/Runtime/Core/Public/MSVC/MSVCPlatform.h
@@ -6,9 +6,7 @@
 
 #pragma once
 
-#if _MSC_VER >= 1920
-	#define PLATFORM_COMPILER_HAS_IF_CONSTEXPR 1
-#else
+#if _MSC_VER < 1920
 	#error "Compiler is expected to support if constexpr"
 #endif
 
diff --git a/Engine/Source/Runtime/Core/Public/Templates/Function.h b/Engine/Source/Runtime/Core/Public/Templates/Function.h
--- a/Engine/Source/Runtime/Core/Public/Templates/Function.h
+++ b/Engine/Source/Runtime/Core/Public/Templates/Function.h
@@ -255,8 +255,6 @@ namespace UE::Core::Private::Function
 	{
 	};
 
-#if PLATFORM_COMPILER_HAS_IF_CONSTEXPR
-
 	template <typename T>
 	FORCEINLINE bool IsBound(const T& Func)
 	{
@@ -273,25 +271,6 @@ namespace UE::Core::Private::Function
 		}
 	}
 
-#else
-
-	template <typename T>
-	FORCEINLINE typename TEnableIf<TIsNullableBinding<T>::Value, bool>::Type IsBound(const T& Func)
-	{
-		// Function pointers, data member pointers, member function pointers and TFunctions
-		// can all be null/unbound, so test them using their boolean state.
-		return !!Func;
-	}
-
-	template <typename T>
-	FORCEINLINE typename TEnableIf<!TIsNullableBinding<T>::Value, bool>::Type IsBound(const T& Func)
-	{
-		// We can't tell if any other generic callable can be invoked, so just assume they can be.
-		return true;
-	}
-
-#endif
-
 	template <typename FunctorType, bool bUnique, bool bOnHeap>
 	struct TStorageOwnerType;
 
@@ -393,8 +372,6 @@ namespace UE::Core::Private::Function
 		{
 		}
 
-#if PLATFORM_COMPILER_HAS_IF_CONSTEXPR
-
 		template <typename FunctorType>
 		typename TDecay<FunctorType>::Type* Bind(FunctorType&& InFunc)
 		{
@@ -428,57 +405,6 @@ namespace UE::Core::Private::Function
 			auto* NewOwned = new (NewAlloc) OwnedType(Forward<FunctorType>(InFunc));
 			return &NewOwned->Obj;
 		}
-
-#else
-
-		#if TFUNCTION_USES_INLINE_STORAGE
-		template <typename FunctorType>
-		typename TEnableIf<
-			(sizeof(TStorageOwnerTypeT<FunctorType, bUnique, false>) <= TFUNCTION_INLINE_SIZE),
-			typename TDecay<FunctorType>::Type*
-		>::Type Bind(FunctorType&& InFunc)
-		{
-			if (!IsBound(InFunc))
-			{
-				return nullptr;
-			}
-
-			using OwnedType = TStorageOwnerTypeT<FunctorType, bUnique, false>;
-
-			void* NewAlloc = &InlineAllocation;
-			auto* NewOwned = new (NewAlloc) OwnedType(Forward<FunctorType>(InFunc));
-
-			return &NewOwned->Obj;
-		}
-		#endif
-
-		template <typename FunctorType>
-		#if TFUNCTION_USES_INLINE_STORAGE
-		typename TEnableIf<
-			(sizeof(TStorageOwnerTypeT<FunctorType, bUnique, true>) > TFUNCTION_INLINE_SIZE),
-			typename TDecay<FunctorType>::Type*
-		>::Type Bind(FunctorType&& InFunc)
-		#else
-		typename TDecay<FunctorType>::Type* Bind(FunctorType&& InFunc)
-		#endif
-		{
-			if (!IsBound(InFunc))
-			{
-				return nullptr;
-			}
-
-			using OwnedType = TStorageOwnerTypeT<FunctorType, bUnique, true>;
-
-			void* NewAlloc = FMemory::Malloc(sizeof(OwnedType), alignof(OwnedType));
-			CA_ASSUME(NewAlloc);
-			auto* NewOwned = new (NewAlloc) OwnedType(Forward<FunctorType>(InFunc));
-
-			HeapAllocation = NewAlloc;
-
-			return &NewOwned->Obj;
-		}
-
-#endif
 	};
 
 	template <typename T, bool bOnHeap>
diff --git a/Engine/Source/Runtime/Core/Public/Templates/MemoryOps.h b/Engine/Source/Runtime/Core/Public/Templates/MemoryOps.h
--- a/Engine/Source/Runtime/Core/Public/Templates/MemoryOps.h
+++ b/Engine/Source/Runtime/Core/Public/Templates/MemoryOps.h
@@ -4,7 +4,6 @@
 
 #include "CoreTypes.h"
 #include "HAL/UnrealMemory.h"
-#include "Templates/EnableIf.h"
 #include "Templates/AreTypesEqual.h"
 #include "Templates/IsTriviallyCopyConstructible.h"
 #include "Templates/UnrealTypeTraits.h"
@@ -33,43 +32,20 @@ namespace UE::Core::Private::MemoryOps
 	};
 }
 
-#if PLATFORM_COMPILER_HAS_IF_CONSTEXPR
-
-	/**
-	 * Default constructs a range of items in memory.
-	 *
-	 * @param	Elements	The address of the first memory location to construct at.
-	 * @param	Count		The number of elements to destruct.
-	 */
-	template <typename ElementType, typename SizeType>
-	FORCEINLINE void DefaultConstructItems(void* Address, SizeType Count)
+/**
+ * Default constructs a range of items in memory.
+ *
+ * @param	Elements	The address of the first memory location to construct at.
+ * @param	Count		The number of elements to destruct.
+ */
+template <typename ElementType, typename SizeType>
+FORCEINLINE void DefaultConstructItems(void* Address, SizeType Count)
+{
+	if constexpr (TIsZeroConstructType<ElementType>::Value)
 	{
-		if constexpr (TIsZeroConstructType<ElementType>::Value)
-		{
-			FMemory::Memset(Address, 0, sizeof(ElementType) * Count);
-		}
-		else
-		{
-			ElementType* Element = (ElementType*)Address;
-			while (Count)
-			{
-				new (Element) ElementType;
-				++Element;
-				--Count;
-			}
-		}
+		FMemory::Memset(Address, 0, sizeof(ElementType) * Count);
 	}
-
-#else
-
-	/**
-	 * Default constructs a range of items in memory.
-	 *
-	 * @param	Elements	The address of the first memory location to construct at.
-	 * @param	Count		The number of elements to destruct.
-	 */
-	template <typename ElementType, typename SizeType>
-	FORCEINLINE typename TEnableIf<!TIsZeroConstructType<ElementType>::Value>::Type DefaultConstructItems(void* Address, SizeType Count)
+	else
 	{
 		ElementType* Element = (ElementType*)Address;
 		while (Count)
@@ -79,101 +55,39 @@ namespace UE::Core::Private::MemoryOps
 			--Count;
 		}
 	}
+}
 
-
-	template <typename ElementType, typename SizeType>
-	FORCEINLINE typename TEnableIf<TIsZeroConstructType<ElementType>::Value>::Type DefaultConstructItems(void* Elements, SizeType Count)
-	{
-		FMemory::Memset(Elements, 0, sizeof(ElementType) * Count);
-	}
-
-#endif
-
-#if PLATFORM_COMPILER_HAS_IF_CONSTEXPR
-
-	/**
-	 * Destructs a single item in memory.
-	 *
-	 * @param	Elements	A pointer to the item to destruct.
-	 *
-	 * @note: This function is optimized for values of T, and so will not dynamically dispatch destructor calls if T's destructor is virtual.
-	 */
-	template <typename ElementType>
-	FORCEINLINE void DestructItem(ElementType* Element)
-	{
-		if constexpr (!TIsTriviallyDestructible<ElementType>::Value)
-		{
-			// We need a typedef here because VC won't compile the destructor call below if ElementType itself has a member called ElementType
-			typedef ElementType DestructItemsElementTypeTypedef;
-
-			Element->DestructItemsElementTypeTypedef::~DestructItemsElementTypeTypedef();
-		}
-	}
-
-#else
-
-	/**
-	 * Destructs a single item in memory.
-	 *
-	 * @param	Elements	A pointer to the item to destruct.
-	 *
-	 * @note: This function is optimized for values of T, and so will not dynamically dispatch destructor calls if T's destructor is virtual.
-	 */
-	template <typename ElementType>
-	FORCEINLINE typename TEnableIf<!TIsTriviallyDestructible<ElementType>::Value>::Type DestructItem(ElementType* Element)
+/**
+ * Destructs a single item in memory.
+ *
+ * @param	Elements	A pointer to the item to destruct.
+ *
+ * @note: This function is optimized for values of T, and so will not dynamically dispatch destructor calls if T's destructor is virtual.
+ */
+template <typename ElementType>
+FORCEINLINE void DestructItem(ElementType* Element)
+{
+	if constexpr (!TIsTriviallyDestructible<ElementType>::Value)
 	{
 		// We need a typedef here because VC won't compile the destructor call below if ElementType itself has a member called ElementType
 		typedef ElementType DestructItemsElementTypeTypedef;
 
 		Element->DestructItemsElementTypeTypedef::~DestructItemsElementTypeTypedef();
 	}
+}
 
-
-	template <typename ElementType>
-	FORCEINLINE typename TEnableIf<TIsTriviallyDestructible<ElementType>::Value>::Type DestructItem(ElementType* Element)
-	{
-	}
-#endif
-
-#if PLATFORM_COMPILER_HAS_IF_CONSTEXPR
-
-	/**
-	 * Destructs a range of items in memory.
-	 *
-	 * @param	Elements	A pointer to the first item to destruct.
-	 * @param	Count		The number of elements to destruct.
-	 *
-	 * @note: This function is optimized for values of T, and so will not dynamically dispatch destructor calls if T's destructor is virtual.
-	 */
-	template <typename ElementType, typename SizeType>
-	FORCEINLINE void DestructItems(ElementType* Element, SizeType Count)
-	{
-		if constexpr (!TIsTriviallyDestructible<ElementType>::Value)
-		{
-			while (Count)
-			{
-				// We need a typedef here because VC won't compile the destructor call below if ElementType itself has a member called ElementType
-				typedef ElementType DestructItemsElementTypeTypedef;
-
-				Element->DestructItemsElementTypeTypedef::~DestructItemsElementTypeTypedef();
-				++Element;
-				--Count;
-			}
-		}
-	}
-
-#else
-
-	/**
-	 * Destructs a range of items in memory.
-	 *
-	 * @param	Elements	A pointer to the first item to destruct.
-	 * @param	Count		The number of elements to destruct.
-	 *
-	 * @note: This function is optimized for values of T, and so will not dynamically dispatch destructor calls if T's destructor is virtual.
-	 */
-	template <typename ElementType, typename SizeType>
-	FORCEINLINE typename TEnableIf<!TIsTriviallyDestructible<ElementType>::Value>::Type DestructItems(ElementType* Element, SizeType Count)
+/**
+ * Destructs a range of items in memory.
+ *
+ * @param	Elements	A pointer to the first item to destruct.
+ * @param	Count		The number of elements to destruct.
+ *
+ * @note: This function is optimized for values of T, and so will not dynamically dispatch destructor calls if T's destructor is virtual.
+ */
+template <typename ElementType, typename SizeType>
+FORCEINLINE void DestructItems(ElementType* Element, SizeType Count)
+{
+	if constexpr (!TIsTriviallyDestructible<ElementType>::Value)
 	{
 		while (Count)
 		{
@@ -185,54 +99,23 @@ namespace UE::Core::Private::MemoryOps
 			--Count;
 		}
 	}
+}
 
-
-	template <typename ElementType, typename SizeType>
-	FORCEINLINE typename TEnableIf<TIsTriviallyDestructible<ElementType>::Value>::Type DestructItems(ElementType* Elements, SizeType Count)
-	{
-	}
-
-#endif
-
-#if PLATFORM_COMPILER_HAS_IF_CONSTEXPR
-
-	/**
-	 * Constructs a range of items into memory from a set of arguments.  The arguments come from an another array.
-	 *
-	 * @param	Dest		The memory location to start copying into.
-	 * @param	Source		A pointer to the first argument to pass to the constructor.
-	 * @param	Count		The number of elements to copy.
-	 */
-	template <typename DestinationElementType, typename SourceElementType, typename SizeType>
-	FORCEINLINE void ConstructItems(void* Dest, const SourceElementType* Source, SizeType Count)
+/**
+ * Constructs a range of items into memory from a set of arguments.  The arguments come from an another array.
+ *
+ * @param	Dest		The memory location to start copying into.
+ * @param	Source		A pointer to the first argument to pass to the constructor.
+ * @param	Count		The number of elements to copy.
+ */
+template <typename DestinationElementType, typename SourceElementType, typename SizeType>
+FORCEINLINE void ConstructItems(void* Dest, const SourceElementType* Source, SizeType Count)
+{
+	if constexpr (TIsBitwiseConstructible<DestinationElementType, SourceElementType>::Value)
 	{
-		if constexpr (TIsBitwiseConstructible<DestinationElementType, SourceElementType>::Value)
-		{
-			FMemory::Memcpy(Dest, Source, sizeof(SourceElementType) * Count);
-		}
-		else
-		{
-			while (Count)
-			{
-				new (Dest) DestinationElementType(*Source);
-				++(DestinationElementType*&)Dest;
-				++Source;
-				--Count;
-			}
-		}
+		FMemory::Memcpy(Dest, Source, sizeof(SourceElementType) * Count);
 	}
-
-#else
-
-	/**
-	 * Constructs a range of items into memory from a set of arguments.  The arguments come from an another array.
-	 *
-	 * @param	Dest		The memory location to start copying into.
-	 * @param	Source		A pointer to the first argument to pass to the constructor.
-	 * @param	Count		The number of elements to copy.
-	 */
-	template <typename DestinationElementType, typename SourceElementType, typename SizeType>
-	FORCEINLINE typename TEnableIf<!TIsBitwiseConstructible<DestinationElementType, SourceElementType>::Value>::Type ConstructItems(void* Dest, const SourceElementType* Source, SizeType Count)
+	else
 	{
 		while (Count)
 		{
@@ -242,55 +125,23 @@ namespace UE::Core::Private::MemoryOps
 			--Count;
 		}
 	}
+}
 
-
-	template <typename DestinationElementType, typename SourceElementType, typename SizeType>
-	FORCEINLINE typename TEnableIf<TIsBitwiseConstructible<DestinationElementType, SourceElementType>::Value>::Type ConstructItems(void* Dest, const SourceElementType* Source, SizeType Count)
-	{
-		FMemory::Memcpy(Dest, Source, sizeof(SourceElementType) * Count);
-	}
-
-#endif
-
-#if PLATFORM_COMPILER_HAS_IF_CONSTEXPR
-
-	/**
-	 * Copy assigns a range of items.
-	 *
-	 * @param	Dest		The memory location to start assigning to.
-	 * @param	Source		A pointer to the first item to assign.
-	 * @param	Count		The number of elements to assign.
-	 */
-	template <typename ElementType, typename SizeType>
-	FORCEINLINE void CopyAssignItems(ElementType* Dest, const ElementType* Source, SizeType Count)
+/**
+ * Copy assigns a range of items.
+ *
+ * @param	Dest		The memory location to start assigning to.
+ * @param	Source		A pointer to the first item to assign.
+ * @param	Count		The number of elements to assign.
+ */
+template <typename ElementType, typename SizeType>
+FORCEINLINE void CopyAssignItems(ElementType* Dest, const ElementType* Source, SizeType Count)
+{
+	if constexpr (TIsTriviallyCopyAssignable<ElementType>::Value)
 	{
-		if constexpr (TIsTriviallyCopyAssignable<ElementType>::Value)
-		{
-			FMemory::Memcpy(Dest, Source, sizeof(ElementType) * Count);
-		}
-		else
-		{
-			while (Count)
-			{
-				*Dest = *Source;
-				++Dest;
-				++Source;
-				--Count;
-			}
-		}
+		FMemory::Memcpy(Dest, Source, sizeof(ElementType) * Count);
 	}
-
-#else
-
-	/**
-	 * Copy assigns a range of items.
-	 *
-	 * @param	Dest		The memory location to start assigning to.
-	 * @param	Source		A pointer to the first item to assign.
-	 * @param	Count		The number of elements to assign.
-	 */
-	template <typename ElementType, typename SizeType>
-	FORCEINLINE typename TEnableIf<!TIsTriviallyCopyAssignable<ElementType>::Value>::Type CopyAssignItems(ElementType* Dest, const ElementType* Source, SizeType Count)
+	else
 	{
 		while (Count)
 		{
@@ -300,68 +151,32 @@ namespace UE::Core::Private::MemoryOps
 			--Count;
 		}
 	}
+}
 
-
-	template <typename ElementType, typename SizeType>
-	FORCEINLINE typename TEnableIf<TIsTriviallyCopyAssignable<ElementType>::Value>::Type CopyAssignItems(ElementType* Dest, const ElementType* Source, SizeType Count)
-	{
-		FMemory::Memcpy(Dest, Source, sizeof(ElementType) * Count);
-	}
-
-#endif
-
-#if PLATFORM_COMPILER_HAS_IF_CONSTEXPR
-
-	/**
-	 * Relocates a range of items to a new memory location as a new type. This is a so-called 'destructive move' for which
-	 * there is no single operation in C++ but which can be implemented very efficiently in general.
-	 *
-	 * @param	Dest		The memory location to relocate to.
-	 * @param	Source		A pointer to the first item to relocate.
-	 * @param	Count		The number of elements to relocate.
-	 */
-	template <typename DestinationElementType, typename SourceElementType, typename SizeType>
-	FORCEINLINE void RelocateConstructItems(void* Dest, const SourceElementType* Source, SizeType Count)
+/**
+ * Relocates a range of items to a new memory location as a new type. This is a so-called 'destructive move' for which
+ * there is no single operation in C++ but which can be implemented very efficiently in general.
+ *
+ * @param	Dest		The memory location to relocate to.
+ * @param	Source		A pointer to the first item to relocate.
+ * @param	Count		The number of elements to relocate.
+ */
+template <typename DestinationElementType, typename SourceElementType, typename SizeType>
+FORCEINLINE void RelocateConstructItems(void* Dest, const SourceElementType* Source, SizeType Count)
+{
+	if constexpr (UE::Core::Private::MemoryOps::TCanBitwiseRelocate<DestinationElementType, SourceElementType>::Value)
 	{
-		if constexpr (UE::Core::Private::MemoryOps::TCanBitwiseRelocate<DestinationElementType, SourceElementType>::Value)
-		{
-			/* All existing UE containers seem to assume trivial relocatability (i.e. memcpy'able) of their members,
-			 * so we're going to assume that this is safe here.  However, it's not generally possible to assume this
-			 * in general as objects which contain pointers/references to themselves are not safe to be trivially
-			 * relocated.
-			 *
-			 * However, it is not yet possible to automatically infer this at compile time, so we can't enable
-			 * different (i.e. safer) implementations anyway. */
-
-			FMemory::Memmove(Dest, Source, sizeof(SourceElementType) * Count);
-		}
-		else
-		{
-			while (Count)
-			{
-				// We need a typedef here because VC won't compile the destructor call below if SourceElementType itself has a member called SourceElementType
-				typedef SourceElementType RelocateConstructItemsElementTypeTypedef;
+		/* All existing UE containers seem to assume trivial relocatability (i.e. memcpy'able) of their members,
+		 * so we're going to assume that this is safe here.  However, it's not generally possible to assume this
+		 * in general as objects which contain pointers/references to themselves are not safe to be trivially
+		 * relocated.
+		 *
+		 * However, it is not yet possible to automatically infer this at compile time, so we can't enable
+		 * different (i.e. safer) implementations anyway. */
 
-				new (Dest) DestinationElementType(*Source);
-				++(DestinationElementType*&)Dest;
-				(Source++)->RelocateConstructItemsElementTypeTypedef::~RelocateConstructItemsElementTypeTypedef();
-				--Count;
-			}
-		}
+		FMemory::Memmove(Dest, Source, sizeof(SourceElementType) * Count);
 	}
-
-#else
-
-	/**
-	 * Relocates a range of items to a new memory location as a new type. This is a so-called 'destructive move' for which
-	 * there is no single operation in C++ but which can be implemented very efficiently in general.
-	 *
-	 * @param	Dest		The memory location to relocate to.
-	 * @param	Source		A pointer to the first item to relocate.
-	 * @param	Count		The number of elements to relocate.
-	 */
-	template <typename DestinationElementType, typename SourceElementType, typename SizeType>
-	FORCEINLINE typename TEnableIf<!UE::Core::Private::MemoryOps::TCanBitwiseRelocate<DestinationElementType, SourceElementType>::Value>::Type RelocateConstructItems(void* Dest, const SourceElementType* Source, SizeType Count)
+	else
 	{
 		while (Count)
 		{
@@ -374,62 +189,23 @@ namespace UE::Core::Private::MemoryOps
 			--Count;
 		}
 	}
+}
 
-	template <typename DestinationElementType, typename SourceElementType, typename SizeType>
-	FORCEINLINE typename TEnableIf<UE::Core::Private::MemoryOps::TCanBitwiseRelocate<DestinationElementType, SourceElementType>::Value>::Type RelocateConstructItems(void* Dest, const SourceElementType* Source, SizeType Count)
-	{
-		/* All existing UE containers seem to assume trivial relocatability (i.e. memcpy'able) of their members,
-		 * so we're going to assume that this is safe here.  However, it's not generally possible to assume this
-		 * in general as objects which contain pointers/references to themselves are not safe to be trivially
-		 * relocated.
-		 *
-		 * However, it is not yet possible to automatically infer this at compile time, so we can't enable
-		 * different (i.e. safer) implementations anyway. */
-
-		FMemory::Memmove(Dest, Source, sizeof(SourceElementType) * Count);
-	}
-
-#endif
-
-#if PLATFORM_COMPILER_HAS_IF_CONSTEXPR
-
-	/**
-	 * Move constructs a range of items into memory.
-	 *
-	 * @param	Dest		The memory location to start moving into.
-	 * @param	Source		A pointer to the first item to move from.
-	 * @param	Count		The number of elements to move.
-	 */
-	template <typename ElementType, typename SizeType>
-	FORCEINLINE void MoveConstructItems(void* Dest, const ElementType* Source, SizeType Count)
+/**
+ * Move constructs a range of items into memory.
+ *
+ * @param	Dest		The memory location to start moving into.
+ * @param	Source		A pointer to the first item to move from.
+ * @param	Count		The number of elements to move.
+ */
+template <typename ElementType, typename SizeType>
+FORCEINLINE void MoveConstructItems(void* Dest, const ElementType* Source, SizeType Count)
+{
+	if constexpr (TIsTriviallyCopyConstructible<ElementType>::Value)
 	{
-		if constexpr (TIsTriviallyCopyConstructible<ElementType>::Value)
-		{
-			FMemory::Memmove(Dest, Source, sizeof(ElementType) * Count);
-		}
-		else
-		{
-			while (Count)
-			{
-				new (Dest) ElementType((ElementType&&)*Source);
-				++(ElementType*&)Dest;
-				++Source;
-				--Count;
-			}
-		}
+		FMemory::Memmove(Dest, Source, sizeof(ElementType) * Count);
 	}
-
-#else
-
-	/**
-	 * Move constructs a range of items into memory.
-	 *
-	 * @param	Dest		The memory location to start moving into.
-	 * @param	Source		A pointer to the first item to move from.
-	 * @param	Count		The number of elements to move.
-	 */
-	template <typename ElementType, typename SizeType>
-	FORCEINLINE typename TEnableIf<!TIsTriviallyCopyConstructible<ElementType>::Value>::Type MoveConstructItems(void* Dest, const ElementType* Source, SizeType Count)
+	else
 	{
 		while (Count)
 		{
@@ -439,54 +215,23 @@ namespace UE::Core::Private::MemoryOps
 			--Count;
 		}
 	}
+}
 
-	template <typename ElementType, typename SizeType>
-	FORCEINLINE typename TEnableIf<TIsTriviallyCopyConstructible<ElementType>::Value>::Type MoveConstructItems(void* Dest, const ElementType* Source, SizeType Count)
+/**
+ * Move assigns a range of items.
+ *
+ * @param	Dest		The memory location to start move assigning to.
+ * @param	Source		A pointer to the first item to move assign.
+ * @param	Count		The number of elements to move assign.
+ */
+template <typename ElementType, typename SizeType>
+FORCEINLINE void MoveAssignItems(ElementType* Dest, const ElementType* Source, SizeType Count)
+{
+	if constexpr (TIsTriviallyCopyAssignable<ElementType>::Value)
 	{
 		FMemory::Memmove(Dest, Source, sizeof(ElementType) * Count);
 	}
-
-#endif
-
-#if PLATFORM_COMPILER_HAS_IF_CONSTEXPR
-
-	/**
-	 * Move assigns a range of items.
-	 *
-	 * @param	Dest		The memory location to start move assigning to.
-	 * @param	Source		A pointer to the first item to move assign.
-	 * @param	Count		The number of elements to move assign.
-	 */
-	template <typename ElementType, typename SizeType>
-	FORCEINLINE void MoveAssignItems(ElementType* Dest, const ElementType* Source, SizeType Count)
-	{
-		if constexpr (TIsTriviallyCopyAssignable<ElementType>::Value)
-		{
-			FMemory::Memmove(Dest, Source, sizeof(ElementType) * Count);
-		}
-		else
-		{
-			while (Count)
-			{
-				*Dest = (ElementType&&)*Source;
-				++Dest;
-				++Source;
-				--Count;
-			}
-		}
-	}
-
-#else
-
-	/**
-	 * Move assigns a range of items.
-	 *
-	 * @param	Dest		The memory location to start move assigning to.
-	 * @param	Source		A pointer to the first item to move assign.
-	 * @param	Count		The number of elements to move assign.
-	 */
-	template <typename ElementType, typename SizeType>
-	FORCEINLINE typename TEnableIf<!TIsTriviallyCopyAssignable<ElementType>::Value>::Type MoveAssignItems(ElementType* Dest, const ElementType* Source, SizeType Count)
+	else
 	{
 		while (Count)
 		{
@@ -496,53 +241,16 @@ namespace UE::Core::Private::MemoryOps
 			--Count;
 		}
 	}
+}
 
-	template <typename ElementType, typename SizeType>
-	FORCEINLINE typename TEnableIf<TIsTriviallyCopyAssignable<ElementType>::Value>::Type MoveAssignItems(ElementType* Dest, const ElementType* Source, SizeType Count)
-	{
-		FMemory::Memmove(Dest, Source, sizeof(ElementType) * Count);
-	}
-
-#endif
-
-#if PLATFORM_COMPILER_HAS_IF_CONSTEXPR
-
-	template <typename ElementType, typename SizeType>
-	FORCEINLINE bool CompareItems(const ElementType* A, const ElementType* B, SizeType Count)
-	{
-		if constexpr (TTypeTraits<ElementType>::IsBytewiseComparable)
-		{
-			return !FMemory::Memcmp(A, B, sizeof(ElementType) * Count);
-		}
-		else
-		{
-			while (Count)
-			{
-				if (!(*A == *B))
-				{
-					return false;
-				}
-
-				++A;
-				++B;
-				--Count;
-			}
-
-			return true;
-		}
-	}
-
-#else
-
-	template <typename ElementType, typename SizeType>
-	FORCEINLINE typename TEnableIf<TTypeTraits<ElementType>::IsBytewiseComparable, bool>::Type CompareItems(const ElementType* A, const ElementType* B, SizeType Count)
+template <typename ElementType, typename SizeType>
+FORCEINLINE bool CompareItems(const ElementType* A, const ElementType* B, SizeType Count)
+{
+	if constexpr (TTypeTraits<ElementType>::IsBytewiseComparable)
 	{
 		return !FMemory::Memcmp(A, B, sizeof(ElementType) * Count);
 	}
-
-
-	template <typename ElementType, typename SizeType>
-	FORCEINLINE typename TEnableIf<!TTypeTraits<ElementType>::IsBytewiseComparable, bool>::Type CompareItems(const ElementType* A, const ElementType* B, SizeType Count)
+	else
 	{
 		while (Count)
 		{
@@ -558,5 +266,4 @@ namespace UE::Core::Private::MemoryOps
 
 		return true;
 	}
-
-#endif
+}
diff --git a/Engine/Source/Runtime/CoreUObject/Public/UObject/Class.h b/Engine/Source/Runtime/CoreUObject/Public/UObject/Class.h
--- a/Engine/Source/Runtime/CoreUObject/Public/UObject/Class.h
+++ b/Engine/Source/Runtime/CoreUObject/Public/UObject/Class.h
@@ -755,297 +755,6 @@ struct TStructOpsTypeTraits : public TStructOpsTypeTraitsBase2<CPPSTRUCT>
 #endif
 
 
-#if !PLATFORM_COMPILER_HAS_IF_CONSTEXPR
-
-	/**
-	 * Selection of constructor behavior.
-	 */
-	template<class CPPSTRUCT>
-	FORCEINLINE typename TEnableIf<!DISABLE_ABSTRACT_CONSTRUCT && !TStructOpsTypeTraits<CPPSTRUCT>::WithNoInitConstructor>::Type ConstructWithNoInitOrNot(void* Data)
-	{
-		new (Data) CPPSTRUCT();
-	}
-
-	template<class CPPSTRUCT>
-	FORCEINLINE typename TEnableIf<!DISABLE_ABSTRACT_CONSTRUCT && TStructOpsTypeTraits<CPPSTRUCT>::WithNoInitConstructor>::Type ConstructWithNoInitOrNot(void* Data)
-	{
-		new (Data) CPPSTRUCT(ForceInit);
-	}
-
-	template<class CPPSTRUCT>
-	FORCEINLINE typename TEnableIf<DISABLE_ABSTRACT_CONSTRUCT>::Type ConstructWithNoInitOrNot(void* Data)
-	{
-	}
-
-	template<class CPPSTRUCT>
-	FORCEINLINE typename TEnableIf<!DISABLE_ABSTRACT_CONSTRUCT && !TStructOpsTypeTraits<CPPSTRUCT>::WithNoInitConstructor>::Type ConstructForTestsWithNoInitOrNot(void* Data)
-	{
-		new (Data) CPPSTRUCT;
-	}
-
-	template<class CPPSTRUCT>
-	FORCEINLINE typename TEnableIf<!DISABLE_ABSTRACT_CONSTRUCT && TStructOpsTypeTraits<CPPSTRUCT>::WithNoInitConstructor>::Type ConstructForTestsWithNoInitOrNot(void* Data)
-	{
-		new (Data) CPPSTRUCT(ForceInit);
-	}
-
-	template<class CPPSTRUCT>
-	FORCEINLINE typename TEnableIf<DISABLE_ABSTRACT_CONSTRUCT>::Type ConstructForTestsWithNoInitOrNot(void* Data)
-	{
-	}
-
-	/**
-	 * Selection of Serialize call.
-	 */
-	template<class CPPSTRUCT>
-	FORCEINLINE typename TEnableIf<!TStructOpsTypeTraits<CPPSTRUCT>::WithSerializer, bool>::Type SerializeOrNot(FArchive& Ar, CPPSTRUCT *Data)
-	{
-		return false;
-	}
-
-	template<class CPPSTRUCT>
-	FORCEINLINE typename TEnableIf<TStructOpsTypeTraits<CPPSTRUCT>::WithSerializer, bool>::Type SerializeOrNot(FArchive& Ar, CPPSTRUCT *Data)
-	{
-		return Data->Serialize(Ar);
-	}
-
-	/**
-	* Selection of structured Serialize call.
-	*/
-	template<class CPPSTRUCT>
-	FORCEINLINE typename TEnableIf<!TStructOpsTypeTraits<CPPSTRUCT>::WithStructuredSerializer, bool>::Type SerializeOrNot(FStructuredArchive::FSlot Slot, CPPSTRUCT *Data)
-	{
-		return false;
-	}
-
-	template<class CPPSTRUCT>
-	FORCEINLINE typename TEnableIf<TStructOpsTypeTraits<CPPSTRUCT>::WithStructuredSerializer, bool>::Type SerializeOrNot(FStructuredArchive::FSlot Slot, CPPSTRUCT *Data)
-	{
-		return Data->Serialize(Slot);
-	}
-
-
-	/**
-	 * Selection of PostSerialize call.
-	 */
-	template<class CPPSTRUCT>
-	FORCEINLINE typename TEnableIf<!TStructOpsTypeTraits<CPPSTRUCT>::WithPostSerialize>::Type PostSerializeOrNot(const FArchive& Ar, CPPSTRUCT *Data)
-	{
-	}
-
-	template<class CPPSTRUCT>
-	FORCEINLINE typename TEnableIf<TStructOpsTypeTraits<CPPSTRUCT>::WithPostSerialize>::Type PostSerializeOrNot(const FArchive& Ar, CPPSTRUCT *Data)
-	{
-		Data->PostSerialize(Ar);
-	}
-
-
-	/**
-	 * Selection of NetSerialize call.
-	 */
-	template<class CPPSTRUCT>
-	FORCEINLINE typename TEnableIf<!TStructOpsTypeTraits<CPPSTRUCT>::WithNetSerializer, bool>::Type NetSerializeOrNot(FArchive& Ar, UPackageMap* Map, bool& bOutSuccess, CPPSTRUCT *Data)
-	{
-		return false;
-	}
-
-	template<class CPPSTRUCT>
-	FORCEINLINE typename TEnableIf<TStructOpsTypeTraits<CPPSTRUCT>::WithNetSerializer, bool>::Type NetSerializeOrNot(FArchive& Ar, UPackageMap* Map, bool& bOutSuccess, CPPSTRUCT *Data)
-	{
-		return Data->NetSerialize(Ar, Map, bOutSuccess);
-	}
-
-
-	/**
-	 * Selection of NetDeltaSerialize call.
-	 */
-	template<class CPPSTRUCT>
-	FORCEINLINE typename TEnableIf<!TStructOpsTypeTraits<CPPSTRUCT>::WithNetDeltaSerializer, bool>::Type NetDeltaSerializeOrNot(FNetDeltaSerializeInfo & DeltaParms, CPPSTRUCT *Data)
-	{
-		return false;
-	}
-
-	template<class CPPSTRUCT>
-	FORCEINLINE typename TEnableIf<TStructOpsTypeTraits<CPPSTRUCT>::WithNetDeltaSerializer, bool>::Type NetDeltaSerializeOrNot(FNetDeltaSerializeInfo & DeltaParms, CPPSTRUCT *Data)
-	{
-		return Data->NetDeltaSerialize(DeltaParms);
-	}
-
-
-	/**
-	 * Selection of PostScriptConstruct call.
-	 */
-	template<class CPPSTRUCT>
-	FORCEINLINE typename TEnableIf<!TStructOpsTypeTraits<CPPSTRUCT>::WithPostScriptConstruct>::Type PostScriptConstructOrNot(CPPSTRUCT *Data)
-	{
-	}
-
-	template<class CPPSTRUCT>
-	FORCEINLINE typename TEnableIf<TStructOpsTypeTraits<CPPSTRUCT>::WithPostScriptConstruct>::Type PostScriptConstructOrNot(CPPSTRUCT *Data)
-	{
-		Data->PostScriptConstruct();
-	}
-
-
-	/**
-	 * Selection of GetPreloadDependencies call.
-	 */
-	template<class CPPSTRUCT>
-	FORCEINLINE typename TEnableIf<!TStructOpsTypeTraits<CPPSTRUCT>::WithGetPreloadDependencies>::Type GetPreloadDependenciesOrNot(CPPSTRUCT* Data, TArray<UObject*>& OutDeps)
-	{
-	}
-
-	template<class CPPSTRUCT>
-	FORCEINLINE typename TEnableIf<TStructOpsTypeTraits<CPPSTRUCT>::WithGetPreloadDependencies>::Type GetPreloadDependenciesOrNot(CPPSTRUCT* Data, TArray<UObject*>& OutDeps)
-	{
-		Data->GetPreloadDependencies(OutDeps);
-	}
-
-
-	/**
-	 * Selection of Copy behavior.
-	 */
-	template<class CPPSTRUCT>
-	FORCEINLINE typename TEnableIf<!TStructOpsTypeTraits<CPPSTRUCT>::WithCopy, bool>::Type CopyOrNot(CPPSTRUCT* Dest, CPPSTRUCT const* Src, int32 ArrayDim)
-	{
-		return false;
-	}
-
-	template<class CPPSTRUCT>
-	FORCEINLINE typename TEnableIf<TStructOpsTypeTraits<CPPSTRUCT>::WithCopy, bool>::Type CopyOrNot(CPPSTRUCT* Dest, CPPSTRUCT const* Src, int32 ArrayDim)
-	{
-		static_assert((!TIsPODType<CPPSTRUCT>::Value), "You probably don't want custom copy for a POD type.");
-		for (; ArrayDim; --ArrayDim)
-		{
-			*Dest++ = *Src++;
-		}
-		return true;
-	}
-
-
-	/**
-	 * Selection of Identical check.
-	 */
-	template<class CPPSTRUCT>
-	FORCEINLINE typename TEnableIf<TStructOpsTypeTraits<CPPSTRUCT>::WithIdentical && TStructOpsTypeTraits<CPPSTRUCT>::WithIdenticalViaEquality, bool>::Type IdenticalOrNot(const CPPSTRUCT* A, const CPPSTRUCT* B, uint32 PortFlags, bool& bOutResult)
-	{
-		static_assert(sizeof(CPPSTRUCT) == 0, "Should not have both WithIdenticalViaEquality and WithIdentical.");
-	}
-
-	template<class CPPSTRUCT>
-	FORCEINLINE typename TEnableIf<!TStructOpsTypeTraits<CPPSTRUCT>::WithIdentical && !TStructOpsTypeTraits<CPPSTRUCT>::WithIdenticalViaEquality, bool>::Type IdenticalOrNot(const CPPSTRUCT* A, const CPPSTRUCT* B, uint32 PortFlags, bool& bOutResult)
-	{
-		bOutResult = false;
-		return false;
-	}
-
-	template<class CPPSTRUCT>
-	FORCEINLINE typename TEnableIf<TStructOpsTypeTraits<CPPSTRUCT>::WithIdentical && !TStructOpsTypeTraits<CPPSTRUCT>::WithIdenticalViaEquality, bool>::Type IdenticalOrNot(const CPPSTRUCT* A, const CPPSTRUCT* B, uint32 PortFlags, bool& bOutResult)
-	{
-		bOutResult = A->Identical(B, PortFlags);
-		return true;
-	}
-
-	template<class CPPSTRUCT>
-	FORCEINLINE typename TEnableIf<!TStructOpsTypeTraits<CPPSTRUCT>::WithIdentical && TStructOpsTypeTraits<CPPSTRUCT>::WithIdenticalViaEquality, bool>::Type IdenticalOrNot(const CPPSTRUCT* A, const CPPSTRUCT* B, uint32 PortFlags, bool& bOutResult)
-	{
-		bOutResult = (*A == *B);
-		return true;
-	}
-
-
-	/**
-	 * Selection of ExportTextItem call.
-	 */
-	template<class CPPSTRUCT>
-	FORCEINLINE typename TEnableIf<!TStructOpsTypeTraits<CPPSTRUCT>::WithExportTextItem, bool>::Type ExportTextItemOrNot(FString& ValueStr, const CPPSTRUCT* PropertyValue, const CPPSTRUCT* DefaultValue, UObject* Parent, int32 PortFlags, UObject* ExportRootScope)
-	{
-		return false;
-	}
-
-	template<class CPPSTRUCT>
-	FORCEINLINE typename TEnableIf<TStructOpsTypeTraits<CPPSTRUCT>::WithExportTextItem, bool>::Type ExportTextItemOrNot(FString& ValueStr, const CPPSTRUCT* PropertyValue, const CPPSTRUCT* DefaultValue, UObject* Parent, int32 PortFlags, UObject* ExportRootScope)
-	{
-		return PropertyValue->ExportTextItem(ValueStr, *DefaultValue, Parent, PortFlags, ExportRootScope);
-	}
-
-
-	/**
-	 * Selection of ImportTextItem call.
-	 */
-	template<class CPPSTRUCT>
-	FORCEINLINE typename TEnableIf<!TStructOpsTypeTraits<CPPSTRUCT>::WithImportTextItem, bool>::Type ImportTextItemOrNot(const TCHAR*& Buffer, CPPSTRUCT* Data, int32 PortFlags, UObject* OwnerObject, FOutputDevice* ErrorText)
-	{
-		return false;
-	}
-
-	template<class CPPSTRUCT>
-	FORCEINLINE typename TEnableIf<TStructOpsTypeTraits<CPPSTRUCT>::WithImportTextItem, bool>::Type ImportTextItemOrNot(const TCHAR*& Buffer, CPPSTRUCT* Data, int32 PortFlags, UObject* OwnerObject, FOutputDevice* ErrorText)
-	{
-		return Data->ImportTextItem(Buffer, PortFlags, OwnerObject, ErrorText);
-	}
-
-
-	/**
-	 * Selection of SerializeFromMismatchedTag call.
-	 */
-	template<class CPPSTRUCT>
-	FORCEINLINE typename TEnableIf<!TStructOpsTypeTraits<CPPSTRUCT>::WithSerializeFromMismatchedTag, bool>::Type SerializeFromMismatchedTagOrNot(FPropertyTag const& Tag, FArchive& Ar, CPPSTRUCT *Data)
-	{
-		return false;
-	}
-
-	template<class CPPSTRUCT>
-	FORCEINLINE typename TEnableIf<TStructOpsTypeTraits<CPPSTRUCT>::WithSerializeFromMismatchedTag && !TIsUECoreType<CPPSTRUCT>::Value, bool>::Type SerializeFromMismatchedTagOrNot(FPropertyTag const& Tag, FArchive& Ar, CPPSTRUCT *Data)
-	{
-		return Data->SerializeFromMismatchedTag(Tag, Ar);
-	}
-
-	template<class CPPSTRUCT>
-	FORCEINLINE typename TEnableIf<TStructOpsTypeTraits<CPPSTRUCT>::WithSerializeFromMismatchedTag && TIsUECoreType<CPPSTRUCT>::Value, bool>::Type SerializeFromMismatchedTagOrNot(FPropertyTag const& Tag, FArchive& Ar, CPPSTRUCT *Data)
-	{
-		return Data->SerializeFromMismatchedTag(Tag.StructName, Ar);
-	}
-
-	template<class CPPSTRUCT>
-	FORCEINLINE typename TEnableIf<!TStructOpsTypeTraits<CPPSTRUCT>::WithStructuredSerializeFromMismatchedTag, bool>::Type StructuredSerializeFromMismatchedTagOrNot(FPropertyTag const& Tag, FStructuredArchive::FSlot Slot, CPPSTRUCT *Data)
-	{
-		return false;
-	}
-
-	template<class CPPSTRUCT>
-	FORCEINLINE typename TEnableIf<TStructOpsTypeTraits<CPPSTRUCT>::WithStructuredSerializeFromMismatchedTag && !TIsUECoreType<CPPSTRUCT>::Value, bool>::Type StructuredSerializeFromMismatchedTagOrNot(FPropertyTag const& Tag, FStructuredArchive::FSlot Slot, CPPSTRUCT *Data)
-	{
-		return Data->SerializeFromMismatchedTag(Tag, Slot);
-	}
-
-	template<class CPPSTRUCT>
-	FORCEINLINE typename TEnableIf<TStructOpsTypeTraits<CPPSTRUCT>::WithStructuredSerializeFromMismatchedTag && TIsUECoreType<CPPSTRUCT>::Value, bool>::Type StructuredSerializeFromMismatchedTagOrNot(FPropertyTag const& Tag, FStructuredArchive::FSlot Slot, CPPSTRUCT *Data)
-	{
-		return Data->SerializeFromMismatchedTag(Tag.StructName, Slot);
-	}
-
-
-
-	/**
-	 * Selection of GetTypeHash call.
-	 */
-	template<class CPPSTRUCT>
-	FORCEINLINE typename TEnableIf<!TModels<CGetTypeHashable, CPPSTRUCT>::Value, uint32>::Type GetTypeHashOrNot(const CPPSTRUCT *Data)
-	{
-		return 0;
-	}
-
-	template<class CPPSTRUCT>
-	FORCEINLINE typename TEnableIf<TModels<CGetTypeHashable, CPPSTRUCT>::Value, uint32>::Type GetTypeHashOrNot(const CPPSTRUCT *Data)
-	{
-		return GetTypeHash(*Data);
-	}
-
-#endif
-
-#if PLATFORM_COMPILER_HAS_IF_CONSTEXPR
-
 	/**
 	 * Selection of AddStructReferencedObjects check.
 	 */
@@ -1058,24 +767,6 @@ struct TStructOpsTypeTraits : public TStructOpsTypeTraitsBase2<CPPSTRUCT>
 		}
 	}
 
-#else
-
-	/**
-	 * Selection of AddStructReferencedObjects check.
-	 */
-	template<class CPPSTRUCT>
-	FORCEINLINE typename TEnableIf<!TStructOpsTypeTraits<CPPSTRUCT>::WithAddStructReferencedObjects>::Type AddStructReferencedObjectsOrNot(void* A, FReferenceCollector& Collector)
-	{
-	}
-
-	template<class CPPSTRUCT>
-	FORCEINLINE typename TEnableIf<TStructOpsTypeTraits<CPPSTRUCT>::WithAddStructReferencedObjects>::Type AddStructReferencedObjectsOrNot(void* A, FReferenceCollector& Collector)
-	{
-		((CPPSTRUCT*)A)->AddStructReferencedObjects(Collector);
-	}
-
-#endif
-
 /**
  * Reflection data for a standalone structure declared in a header or as a user defined struct
  */
@@ -1260,7 +951,6 @@ public:
 			check(!TTraits::WithZeroConstructor); // don't call this if we have indicated it is not necessary
 			// that could have been an if statement, but we might as well force optimization above the virtual call
 			// could also not attempt to call the constructor for types where this is not possible, but I didn't do that here
-#if PLATFORM_COMPILER_HAS_IF_CONSTEXPR
 #if CHECK_PUREVIRTUALS
 			if constexpr (!TStructOpsTypeTraits<CPPSTRUCT>::WithPureVirtual)
 #endif
@@ -1274,16 +964,12 @@ public:
 					new (Dest) CPPSTRUCT();
 				}
 			}
-#else
-			ConstructWithNoInitOrNot<CPPSTRUCT>(Dest);
-#endif
 		}
 		virtual void ConstructForTests(void* Dest) override
 		{
 			check(!TTraits::WithZeroConstructor); // don't call this if we have indicated it is not necessary
 			// that could have been an if statement, but we might as well force optimization above the virtual call
 			// could also not attempt to call the constructor for types where this is not possible, but I didn't do that here
-#if PLATFORM_COMPILER_HAS_IF_CONSTEXPR
 #if CHECK_PUREVIRTUALS
 			if constexpr (!TStructOpsTypeTraits<CPPSTRUCT>::WithPureVirtual)
 #endif
@@ -1297,9 +983,6 @@ public:
 					new (Dest) CPPSTRUCT;
 				}
 			}
-#else
-			ConstructForTestsWithNoInitOrNot<CPPSTRUCT>(Dest);
-#endif
 		}
 		virtual bool HasDestructor() override
 		{
@@ -1323,7 +1006,6 @@ public:
 		virtual bool Serialize(FArchive& Ar, void *Data) override
 		{
 			check(TTraits::WithSerializer); // don't call this if we have indicated it is not necessary
-#if PLATFORM_COMPILER_HAS_IF_CONSTEXPR
 			if constexpr (TStructOpsTypeTraits<CPPSTRUCT>::WithSerializer)
 			{
 				return ((CPPSTRUCT*)Data)->Serialize(Ar);
@@ -1332,14 +1014,10 @@ public:
 			{
 				return false;
 			}
-#else
-			return SerializeOrNot(Ar, (CPPSTRUCT*)Data);
-#endif
 		}
 		virtual bool Serialize(FStructuredArchive::FSlot Slot, void *Data) override
 		{
 			check(TTraits::WithStructuredSerializer); // don't call this if we have indicated it is not necessary
-#if PLATFORM_COMPILER_HAS_IF_CONSTEXPR
 			if constexpr (TStructOpsTypeTraits<CPPSTRUCT>::WithStructuredSerializer)
 			{
 				return ((CPPSTRUCT*)Data)->Serialize(Slot);
@@ -1348,10 +1026,7 @@ public:
 			{
 				return false;
 			}
-#else
-			return SerializeOrNot(Slot, (CPPSTRUCT*)Data);
-#endif
-			}
+		}
 		virtual bool HasPostSerialize() override
 		{
 			return TTraits::WithPostSerialize;
@@ -1359,14 +1034,10 @@ public:
 		virtual void PostSerialize(const FArchive& Ar, void *Data) override
 		{
 			check(TTraits::WithPostSerialize); // don't call this if we have indicated it is not necessary
-#if PLATFORM_COMPILER_HAS_IF_CONSTEXPR
 			if constexpr (TStructOpsTypeTraits<CPPSTRUCT>::WithPostSerialize)
 			{
 				((CPPSTRUCT*)Data)->PostSerialize(Ar);
 			}
-#else
-			PostSerializeOrNot(Ar, (CPPSTRUCT*)Data);
-#endif
 		}
 		virtual bool HasNetSerializer() override
 		{
@@ -1382,7 +1053,6 @@ public:
 		}
 		virtual bool NetSerialize(FArchive& Ar, class UPackageMap* Map, bool& bOutSuccess, void *Data) override
 		{
-#if PLATFORM_COMPILER_HAS_IF_CONSTEXPR
 			if constexpr (TStructOpsTypeTraits<CPPSTRUCT>::WithNetSerializer)
 			{
 				return ((CPPSTRUCT*)Data)->NetSerialize(Ar, Map, bOutSuccess);
@@ -1391,13 +1061,9 @@ public:
 			{
 				return false;
 			}
-#else
-			return NetSerializeOrNot(Ar, Map, bOutSuccess, (CPPSTRUCT*)Data);
-#endif
 		}
 		virtual bool NetDeltaSerialize(FNetDeltaSerializeInfo & DeltaParms, void *Data) override
 		{
-#if PLATFORM_COMPILER_HAS_IF_CONSTEXPR
 			if constexpr (TStructOpsTypeTraits<CPPSTRUCT>::WithNetDeltaSerializer)
 			{
 				return ((CPPSTRUCT*)Data)->NetDeltaSerialize(DeltaParms);
@@ -1406,9 +1072,6 @@ public:
 			{
 				return false;
 			}
-#else
-			return NetDeltaSerializeOrNot(DeltaParms, (CPPSTRUCT*)Data);
-#endif
 		}
 		virtual bool HasPostScriptConstruct() override
 		{
@@ -1417,25 +1080,17 @@ public:
 		virtual void PostScriptConstruct(void *Data) override
 		{
 			check(TTraits::WithPostScriptConstruct); // don't call this if we have indicated it is not necessary
-#if PLATFORM_COMPILER_HAS_IF_CONSTEXPR
 			if constexpr (TStructOpsTypeTraits<CPPSTRUCT>::WithPostScriptConstruct)
 			{
 				((CPPSTRUCT*)Data)->PostScriptConstruct();
 			}
-#else
-			PostScriptConstructOrNot((CPPSTRUCT*)Data);
-#endif
 		}
 		virtual void GetPreloadDependencies(void* Data, TArray<UObject*>& OutDeps) override
 		{
-#if PLATFORM_COMPILER_HAS_IF_CONSTEXPR
 			if constexpr (TStructOpsTypeTraits<CPPSTRUCT>::WithGetPreloadDependencies)
 			{
 				((CPPSTRUCT*)Data)->GetPreloadDependencies(OutDeps);
 			}
-#else
-			GetPreloadDependenciesOrNot((CPPSTRUCT*)Data, OutDeps);
-#endif
 		}
 		virtual bool IsPlainOldData() override
 		{
@@ -1451,7 +1106,6 @@ public:
 		}
 		virtual bool Copy(void* Dest, void const* Src, int32 ArrayDim) override
 		{
-#if PLATFORM_COMPILER_HAS_IF_CONSTEXPR
 			if constexpr (TStructOpsTypeTraits<CPPSTRUCT>::WithCopy)
 			{
 				static_assert((!TIsPODType<CPPSTRUCT>::Value), "You probably don't want custom copy for a POD type.");
@@ -1469,9 +1123,6 @@ public:
 			{
 				return false;
 			}
-#else
-			return CopyOrNot((CPPSTRUCT*)Dest, (CPPSTRUCT const*)Src, ArrayDim);
-#endif
 		}
 		virtual bool HasIdentical() override
 		{
@@ -1480,7 +1131,6 @@ public:
 		virtual bool Identical(const void* A, const void* B, uint32 PortFlags, bool& bOutResult) override
 		{
 			check((TTraits::WithIdentical || TTraits::WithIdenticalViaEquality)); // don't call this if we have indicated it is not necessary
-#if PLATFORM_COMPILER_HAS_IF_CONSTEXPR
 			if constexpr (TStructOpsTypeTraits<CPPSTRUCT>::WithIdentical)
 			{
 				static_assert(!TStructOpsTypeTraits<CPPSTRUCT>::WithIdenticalViaEquality, "Should not have both WithIdenticalViaEquality and WithIdentical.");
@@ -1498,9 +1148,6 @@ public:
 				bOutResult = false;
 				return false;
 			}
-#else
-			return IdenticalOrNot((const CPPSTRUCT*)A, (const CPPSTRUCT*)B, PortFlags, bOutResult);
-#endif
 		}
 		virtual bool HasExportTextItem() override
 		{
@@ -1509,7 +1156,6 @@ public:
 		virtual bool ExportTextItem(FString& ValueStr, const void* PropertyValue, const void* DefaultValue, class UObject* Parent, int32 PortFlags, class UObject* ExportRootScope) override
 		{
 			check(TTraits::WithExportTextItem); // don't call this if we have indicated it is not necessary
-#if PLATFORM_COMPILER_HAS_IF_CONSTEXPR
 			if constexpr (TStructOpsTypeTraits<CPPSTRUCT>::WithExportTextItem)
 			{
 				return ((const CPPSTRUCT*)PropertyValue)->ExportTextItem(ValueStr, *(const CPPSTRUCT*)DefaultValue, Parent, PortFlags, ExportRootScope);
@@ -1518,9 +1164,6 @@ public:
 			{
 				return false;
 			}
-#else
-			return ExportTextItemOrNot(ValueStr, (const CPPSTRUCT*)PropertyValue, (const CPPSTRUCT*)DefaultValue, Parent, PortFlags, ExportRootScope);
-#endif
 		}
 		virtual bool HasImportTextItem() override
 		{
@@ -1529,7 +1172,6 @@ public:
 		virtual bool ImportTextItem(const TCHAR*& Buffer, void* Data, int32 PortFlags, class UObject* OwnerObject, FOutputDevice* ErrorText) override
 		{
 			check(TTraits::WithImportTextItem); // don't call this if we have indicated it is not necessary
-#if PLATFORM_COMPILER_HAS_IF_CONSTEXPR
 			if constexpr (TStructOpsTypeTraits<CPPSTRUCT>::WithImportTextItem)
 			{
 				return ((CPPSTRUCT*)Data)->ImportTextItem(Buffer, PortFlags, OwnerObject, ErrorText);
@@ -1538,9 +1180,6 @@ public:
 			{
 				return false;
 			}
-#else
-			return ImportTextItemOrNot(Buffer, (CPPSTRUCT*)Data, PortFlags, OwnerObject, ErrorText);
-#endif
 		}
 		virtual bool HasAddStructReferencedObjects() override
 		{
@@ -1558,7 +1197,6 @@ public:
 		virtual bool SerializeFromMismatchedTag(struct FPropertyTag const& Tag, FArchive& Ar, void *Data) override
 		{
 			check(TTraits::WithSerializeFromMismatchedTag); // don't call this if we have indicated it is not allowed
-#if PLATFORM_COMPILER_HAS_IF_CONSTEXPR
 			if constexpr (TStructOpsTypeTraits<CPPSTRUCT>::WithSerializeFromMismatchedTag)
 			{
 				if constexpr (TIsUECoreType<CPPSTRUCT>::Value)
@@ -1575,9 +1213,6 @@ public:
 			{
 				return false;
 			}
-#else
-			return SerializeFromMismatchedTagOrNot(Tag, Ar, (CPPSTRUCT*)Data);
-#endif
 		}
 		virtual bool HasStructuredSerializeFromMismatchedTag() override
 		{
@@ -1586,7 +1221,6 @@ public:
 		virtual bool StructuredSerializeFromMismatchedTag(struct FPropertyTag const& Tag, FStructuredArchive::FSlot Slot, void *Data) override
 		{
 			check(TTraits::WithStructuredSerializeFromMismatchedTag); // don't call this if we have indicated it is not allowed
-#if PLATFORM_COMPILER_HAS_IF_CONSTEXPR
 			if constexpr (TStructOpsTypeTraits<CPPSTRUCT>::WithStructuredSerializeFromMismatchedTag)
 			{
 				if constexpr (TIsUECoreType<CPPSTRUCT>::Value)
@@ -1603,9 +1237,6 @@ public:
 			{
 				return false;
 			}
-#else
-			return StructuredSerializeFromMismatchedTagOrNot(Tag, Slot, (CPPSTRUCT*)Data);
-#endif
 		}
 
 		static_assert(!(TTraits::WithSerializeFromMismatchedTag && TTraits::WithStructuredSerializeFromMismatchedTag), "Structs cannot have both WithSerializeFromMismatchedTag and WithStructuredSerializeFromMismatchedTag set");
@@ -1618,7 +1249,6 @@ public:
 		{
 			ensure(HasGetTypeHash());
 
-#if PLATFORM_COMPILER_HAS_IF_CONSTEXPR
 			if constexpr (TModels<CGetTypeHashable, CPPSTRUCT>::Value)
 			{
 				return GetTypeHash(*(const CPPSTRUCT*)Src);
@@ -1627,9 +1257,6 @@ public:
 			{
 				return 0;
 			}
-#else
-			return GetTypeHashOrNot((const CPPSTRUCT*)Src);
-#endif
 		}
 		virtual EPropertyFlags GetComputedPropertyFlags() const override
 		{
diff --git a/Engine/Source/Runtime/Engine/Public/ShowFlags.h b/Engine/Source/Runtime/Engine/Public/ShowFlags.h
--- a/Engine/Source/Runtime/Engine/Public/ShowFlags.h
+++ b/Engine/Source/Runtime/Engine/Public/ShowFlags.h
@@ -572,15 +572,10 @@ enum class EShowFlagShippingValue
 template<EShowFlagShippingValue ShippingValue = EShowFlagShippingValue::Dynamic>
 struct TCustomShowFlag
 {
-#if PLATFORM_COMPILER_HAS_IF_CONSTEXPR
-#define IF_CONSTEXPR if constexpr
-#else 
-#define IF_CONSTEXPR if
-#endif
 	TCustomShowFlag(const TCHAR* Name, bool DefaultEnabled, EShowFlagGroup Group = SFG_Custom, FText DisplayName = {})
 	{
 #if UE_BUILD_SHIPPING
-		IF_CONSTEXPR(ShippingValue == EShowFlagShippingValue::Dynamic)
+		if constexpr (ShippingValue == EShowFlagShippingValue::Dynamic)
 #endif
 		{
 			FlagIndex = FEngineShowFlags::RegisterCustomShowFlag(Name, DefaultEnabled, Group, DisplayName);
@@ -590,11 +585,11 @@ struct TCustomShowFlag
 	bool IsEnabled(const FEngineShowFlags& ShowFlags)
 	{
 #if UE_BUILD_SHIPPING
-		IF_CONSTEXPR (ShippingValue == EShowFlagShippingValue::ForceDisabled)
+		if constexpr (ShippingValue == EShowFlagShippingValue::ForceDisabled)
 		{
 			return false;
 		}
-		else IF_CONSTEXPR(ShippingValue == EShowFlagShippingValue::ForceEnabled)
+		else if constexpr (ShippingValue == EShowFlagShippingValue::ForceEnabled)
 		{
 			return true;
 		}
@@ -611,7 +606,6 @@ struct TCustomShowFlag
 			}
 		}
 	}
-#undef IF_CONSTEXPR
 
 private:
 	FEngineShowFlags::ECustomShowFlag FlagIndex = FEngineShowFlags::ECustomShowFlag::None;
