Split the default state of FGameplayEventFunction to a separate non-property allocation

This avoids instance overrides preventing the default state from propagating (because structs are considered a single value) and also means that instances no longer have a copy of the default data. Instead the default state is found via the archetype chain when needed at runtime.

#rnx

diff --git a/Engine/Source/Runtime/Core/Public/UObject/FortniteMainBranchObjectVersion.h b/Engine/Source/Runtime/Core/Public/UObject/FortniteMainBranchObjectVersion.h
--- a/Engine/Source/Runtime/Core/Public/UObject/FortniteMainBranchObjectVersion.h
+++ b/Engine/Source/Runtime/Core/Public/UObject/FortniteMainBranchObjectVersion.h
@@ -199,6 +199,9 @@ struct CORE_API FFortniteMainBranchObjectVersion
 		// Add faster damping calculations to the cloth simulation and rename previous Damping parameter to LocalDamping.
 		ChaosClothFasterDamping,
 
+		// Migrated function handlers to the CDO/archetype data
+		MigratedFunctionHandlersToDefaults,
+
 		// -----<new versions can be added above this line>-------------------------------------------------
 		VersionPlusOne,
 		LatestVersion = VersionPlusOne - 1
