Fix driver version check for Windows on NVIDIA.

#rnx

diff --git a/Engine/Source/Runtime/Core/Private/Windows/WindowsPlatformMisc.cpp b/Engine/Source/Runtime/Core/Private/Windows/WindowsPlatformMisc.cpp
--- a/Engine/Source/Runtime/Core/Private/Windows/WindowsPlatformMisc.cpp
+++ b/Engine/Source/Runtime/Core/Private/Windows/WindowsPlatformMisc.cpp
@@ -2463,7 +2463,7 @@ static void GetVideoDriverDetails(const FString& Key, FGPUDriverInfo& Out)
 
 	if(Out.IsNVIDIA())
 	{
-		Out.UserDriverVersion = Out.GetUnifiedDriverVersion();
+		Out.UserDriverVersion = Out.TrimNVIDIAInternalVersion(Out.InternalDriverVersion);
 	}
 	else if(Out.IsAMD())
 	{
diff --git a/Engine/Source/Runtime/Core/Public/GenericPlatform/GenericPlatformDriver.h b/Engine/Source/Runtime/Core/Public/GenericPlatform/GenericPlatformDriver.h
--- a/Engine/Source/Runtime/Core/Public/GenericPlatform/GenericPlatformDriver.h
+++ b/Engine/Source/Runtime/Core/Public/GenericPlatform/GenericPlatformDriver.h
@@ -318,6 +318,25 @@ struct FGPUDriverInfo
 		return true;
 	}
 
+	static FString TrimNVIDIAInternalVersion(const FString& InternalVersion)
+	{
+		// on the internal driver number: https://forums.geforce.com/default/topic/378546/confusion-over-driver-version-numbers/
+		//   The first 7 shows u that is a Vista driver, 6 that is an XP and 4 that is Me
+		// we don't care about the windows version so we don't look at the front part of the driver version
+		// "9.18.13.4788" -> "347.88"
+		// "10.18.13.4788" -> "347.88"
+		// the following code works with the current numbering scheme, if needed we have to update that
+
+		// we don't care about the windows version so we don't look at the front part of the driver version
+		// e.g. 36.143
+		FString RightPart = InternalVersion.Right(6);
+
+		// move the dot
+		RightPart = RightPart.Replace(TEXT("."), TEXT(""));
+		RightPart.InsertAt(3, TEXT("."));
+		return RightPart;
+	}
+
 	FString GetUnifiedDriverVersion() const
 	{
 		// we use the internal version, not the user version to avoid problem where the name was altered 
@@ -325,24 +344,7 @@ struct FGPUDriverInfo
 
 		if(IsNVIDIA() && (InternalDriverVersion != UserDriverVersion))
 		{
-			// only do this processing on platforms where the InternalDriverVersion is different from the UserDriverVersion
-			// for example Linux reads the version from Vulkan properties directly and uses the UserDriverVersion everywhere.
-
-			// on the internal driver number: https://forums.geforce.com/default/topic/378546/confusion-over-driver-version-numbers/
-			//   The first 7 shows u that is a Vista driver, 6 that is an XP and 4 that is Me
-			// we don't care about the windows version so we don't look at the front part of the driver version
-			// "9.18.13.4788" -> "347.88"
-			// "10.18.13.4788" -> "347.88"
-			// the following code works with the current numbering scheme, if needed we have to update that
-
-			// we don't care about the windows version so we don't look at the front part of the driver version
-			// e.g. 36.143
-			FString RightPart = FullVersion.Right(6);
-
-			// move the dot
-			RightPart = RightPart.Replace(TEXT("."), TEXT(""));
-			RightPart.InsertAt(3, TEXT("."));
-			return RightPart;
+			return TrimNVIDIAInternalVersion(FullVersion);
 		}
 		else if(IsAMD())
 		{
