Better SFINAE for TNonNullPtr

diff --git a/Engine/Source/Runtime/Core/Public/Templates/NonNullPointer.h b/Engine/Source/Runtime/Core/Public/Templates/NonNullPointer.h
--- a/Engine/Source/Runtime/Core/Public/Templates/NonNullPointer.h
+++ b/Engine/Source/Runtime/Core/Public/Templates/NonNullPointer.h
@@ -11,6 +11,40 @@ struct TOptional;
 // So we can construct TNonNullPtrs
 enum class EDefaultConstructNonNullPtr { UnsafeDoNotUse };
 
+namespace UE::Core::Private::NonNullPtr {
+template <typename...>
+using TVoid = void;
+
+/**
+ * Version of `::TPointerIsConvertibleFromTo` that produces an incomplete type
+ * when either `From` or `To` are incomplete types
+ */
+template <typename, typename, typename = void>
+struct TPointerIsConvertibleFromTo;
+
+/**
+ * Specialization of
+ * `UE::Core::Private::NonNullPtr::TPointerIsConvertibleFromTo` for complete
+ * non-function types
+ */
+template <typename From, typename To>
+struct TPointerIsConvertibleFromTo<From, To, TVoid<decltype(sizeof(From)), decltype(sizeof(To))>>
+	: ::TPointerIsConvertibleFromTo<From, To>
+{
+};
+
+/**
+ * Specialization of
+ * `UE::Core::Private::NonNullPtr::TPointerIsConvertibleFromTo` for function
+ * types, which are always complete types
+ */
+template <typename Result1, typename... Args1, typename Result2, typename... Args2>
+struct TPointerIsConvertibleFromTo<Result1(Args1...), Result2(Args2...)>
+	: ::TPointerIsConvertibleFromTo<Result1(Args1...), Result2(Args2...)>
+{
+};
+}
+
 /**
  * TNonNullPtr is a non-nullable, non-owning, raw/naked/unsafe pointer.
  */
@@ -39,11 +73,7 @@ public:
 	/**
 	 * Constructs a non-null pointer from the provided pointer. Must not be nullptr.
 	 */
-	template <
-		typename OtherObjectType,
-		typename = typename TEnableIf<TPointerIsConvertibleFromTo<OtherObjectType, ObjectType>::Value>::Type
-	>
-	FORCEINLINE TNonNullPtr(OtherObjectType* InObject)
+	FORCEINLINE TNonNullPtr(ObjectType* InObject)
 		: Object(InObject)
 	{
 		ensureMsgf(InObject, TEXT("Tried to initialize TNonNullPtr with a null pointer!"));
@@ -54,9 +84,9 @@ public:
 	 */
 	template <
 		typename OtherObjectType,
-		typename = typename TEnableIf<TPointerIsConvertibleFromTo<OtherObjectType, ObjectType>::Value>::Type
+		typename = typename TEnableIf<UE::Core::Private::NonNullPtr::TPointerIsConvertibleFromTo<OtherObjectType, ObjectType>::Value>::Type
 	>
-	FORCEINLINE TNonNullPtr(TNonNullPtr<OtherObjectType>& Other)
+	FORCEINLINE TNonNullPtr(const TNonNullPtr<OtherObjectType>& Other)
 		: Object(Other.Object)
 	{
 	}
@@ -73,11 +103,7 @@ public:
 	/**
 	 * Assignment operator taking a pointer
 	 */
-	template <
-		typename OtherObjectType,
-		typename = typename TEnableIf<TPointerIsConvertibleFromTo<OtherObjectType, ObjectType>::Value>::Type
-	>
-	FORCEINLINE TNonNullPtr& operator=(OtherObjectType* InObject)
+	FORCEINLINE TNonNullPtr& operator=(ObjectType* InObject)
 	{
 		ensureMsgf(InObject, TEXT("Tried to assign a null pointer to a TNonNullPtr!"));
 		Object = InObject;
@@ -87,11 +113,8 @@ public:
 	/**
 	 * Assignment operator taking another TNonNullPtr
 	 */
-	template <
-		typename OtherObjectType,
-		typename = typename TEnableIf<TPointerIsConvertibleFromTo<OtherObjectType, ObjectType>::Value>::Type
-	>
-	FORCEINLINE TNonNullPtr& operator=(TNonNullPtr<OtherObjectType>& Other)
+	template <typename OtherObjectType>
+	FORCEINLINE typename TEnableIf<UE::Core::Private::NonNullPtr::TPointerIsConvertibleFromTo<OtherObjectType, ObjectType>::Value, TNonNullPtr&>::Type operator=(const TNonNullPtr<OtherObjectType>& Other)
 	{
 		Object = Other.Object;
 		return *this;
