Structured archive config defines moved into their own header.

diff --git a/Engine/Source/Runtime/Core/Public/Serialization/StructuredArchive.h b/Engine/Source/Runtime/Core/Public/Serialization/StructuredArchive.h
--- a/Engine/Source/Runtime/Core/Public/Serialization/StructuredArchive.h
+++ b/Engine/Source/Runtime/Core/Public/Serialization/StructuredArchive.h
@@ -3,6 +3,7 @@
 #pragma once
 
 #include "CoreTypes.h"
+#include "Serialization/StructuredArchiveDefines.h"
 #include "Serialization/StructuredArchiveFormatter.h"
 #include "Formatters/BinaryArchiveFormatter.h"
 #include "Misc/Optional.h"
@@ -12,24 +13,6 @@
 #include "Serialization/ArchiveProxy.h"
 #include "Templates/UniqueObj.h"
 
-/**
- * DO_STRUCTURED_ARCHIVE_CONTAINER_CHECKS - if set, checks that nested container types are serialized correctly.
- */
-#ifndef DO_STRUCTURED_ARCHIVE_CONTAINER_CHECKS
-	#if DO_GUARD_SLOW
-		#define DO_STRUCTURED_ARCHIVE_CONTAINER_CHECKS 1
-	#else
-		#define DO_STRUCTURED_ARCHIVE_CONTAINER_CHECKS 0
-	#endif
-#endif
-
-/**
- * DO_STRUCTURED_ARCHIVE_UNIQUE_FIELD_NAME_CHECKS - if set, checks that field names are unique within a container.  Requires DO_STRUCTURED_ARCHIVE_CONTAINER_CHECKS.
- */
-#ifndef DO_STRUCTURED_ARCHIVE_UNIQUE_FIELD_NAME_CHECKS
-	#define DO_STRUCTURED_ARCHIVE_UNIQUE_FIELD_NAME_CHECKS 0
-#endif
-
 /**
  * Class to contain a named value for serialization. Intended to be created as a temporary and passed to object serialization methods.
  */
diff --git a/Engine/Source/Runtime/Core/Public/Serialization/StructuredArchiveDefines.h b/Engine/Source/Runtime/Core/Public/Serialization/StructuredArchiveDefines.h
new file mode 100644
--- /dev/null
+++ b/Engine/Source/Runtime/Core/Public/Serialization/StructuredArchiveDefines.h
@@ -0,0 +1,21 @@
+// Copyright Epic Games, Inc. All Rights Reserved.
+
+#pragma once
+
+/**
+ * DO_STRUCTURED_ARCHIVE_CONTAINER_CHECKS - if set, checks that nested container types are serialized correctly.
+ */
+#ifndef DO_STRUCTURED_ARCHIVE_CONTAINER_CHECKS
+	#if DO_GUARD_SLOW
+		#define DO_STRUCTURED_ARCHIVE_CONTAINER_CHECKS 1
+	#else
+		#define DO_STRUCTURED_ARCHIVE_CONTAINER_CHECKS 0
+	#endif
+#endif
+
+/**
+ * DO_STRUCTURED_ARCHIVE_UNIQUE_FIELD_NAME_CHECKS - if set, checks that field names are unique within a container.  Requires DO_STRUCTURED_ARCHIVE_CONTAINER_CHECKS.
+ */
+#ifndef DO_STRUCTURED_ARCHIVE_UNIQUE_FIELD_NAME_CHECKS
+	#define DO_STRUCTURED_ARCHIVE_UNIQUE_FIELD_NAME_CHECKS 0
+#endif
