StructuredArchive namespace fixes.

diff --git a/Engine/Source/Runtime/Core/Private/Serialization/StructuredArchive.cpp b/Engine/Source/Runtime/Core/Private/Serialization/StructuredArchive.cpp
--- a/Engine/Source/Runtime/Core/Private/Serialization/StructuredArchive.cpp
+++ b/Engine/Source/Runtime/Core/Private/Serialization/StructuredArchive.cpp
@@ -86,7 +86,7 @@ FStructuredArchiveSlot FStructuredArchive::Open()
 	check(!CurrentSlotElementId.IsValid());
 
 	RootElementId = ElementIdGenerator.Generate();
-	CurrentScope.Emplace(RootElementId, StructuredArchive_Private::EElementType::Root);
+	CurrentScope.Emplace(RootElementId, UE::StructuredArchive::Private::EElementType::Root);
 
 	CurrentSlotElementId = ElementIdGenerator.Generate();
 
@@ -95,16 +95,16 @@ FStructuredArchiveSlot FStructuredArchive::Open()
 
 void FStructuredArchive::Close()
 {
-	SetScope(StructuredArchive_Private::FSlotPosition(0, RootElementId));
+	SetScope(UE::StructuredArchive::Private::FSlotPosition(0, RootElementId));
 }
 
-void FStructuredArchive::EnterSlot(StructuredArchive_Private::FSlotPosition Slot, bool bEnteringAttributedValue)
+void FStructuredArchive::EnterSlot(UE::StructuredArchive::Private::FSlotPosition Slot, bool bEnteringAttributedValue)
 {
-	int32                                    ParentDepth = Slot.Depth;
-	StructuredArchive_Private::FElementId ElementId   = Slot.ElementId;
+	int32                                      ParentDepth = Slot.Depth;
+	UE::StructuredArchive::Private::FElementId ElementId   = Slot.ElementId;
 
 	// If the slot being entered has attributes, enter the value slot first.
-	if (ParentDepth + 1 < CurrentScope.Num() && CurrentScope[ParentDepth + 1].Id == ElementId && CurrentScope[ParentDepth + 1].Type == StructuredArchive_Private::EElementType::AttributedValue)
+	if (ParentDepth + 1 < CurrentScope.Num() && CurrentScope[ParentDepth + 1].Id == ElementId && CurrentScope[ParentDepth + 1].Type == UE::StructuredArchive::Private::EElementType::AttributedValue)
 	{
 #if DO_STRUCTURED_ARCHIVE_CONTAINER_CHECKS
 		checkf(!CurrentSlotElementId.IsValid() && !CurrentContainer.Top()->bAttributedValueWritten, TEXT("Attempt to serialize data into an invalid slot"));
@@ -113,15 +113,15 @@ void FStructuredArchive::EnterSlot(StructuredArchive_Private::FSlotPosition Slot
 		checkf(!CurrentSlotElementId.IsValid(), TEXT("Attempt to serialize data into an invalid slot"));
 #endif
 
-		SetScope(StructuredArchive_Private::FSlotPosition(ParentDepth + 1, ElementId));
+		SetScope(UE::StructuredArchive::Private::FSlotPosition(ParentDepth + 1, ElementId));
 		Formatter.EnterAttributedValueValue();
 	}
 	else if (!bEnteringAttributedValue && Formatter.TryEnterAttributedValueValue())
 	{
-		int32 NewDepth = EnterSlotAsType(StructuredArchive_Private::FSlotPosition(ParentDepth, ElementId), StructuredArchive_Private::EElementType::AttributedValue);
+		int32 NewDepth = EnterSlotAsType(UE::StructuredArchive::Private::FSlotPosition(ParentDepth, ElementId), UE::StructuredArchive::Private::EElementType::AttributedValue);
 		check(NewDepth == ParentDepth + 1);
-		StructuredArchive_Private::FElementId AttributedValueId = CurrentScope[NewDepth].Id;
-		SetScope(StructuredArchive_Private::FSlotPosition(NewDepth, AttributedValueId));
+		UE::StructuredArchive::Private::FElementId AttributedValueId = CurrentScope[NewDepth].Id;
+		SetScope(UE::StructuredArchive::Private::FSlotPosition(NewDepth, AttributedValueId));
 #if DO_STRUCTURED_ARCHIVE_CONTAINER_CHECKS
 		CurrentContainer.Emplace(0);
 #endif
@@ -132,12 +132,12 @@ void FStructuredArchive::EnterSlot(StructuredArchive_Private::FSlotPosition Slot
 		CurrentSlotElementId.Reset();
 	}
 
-	CurrentEnteringAttributeState = StructuredArchive_Private::EEnteringAttributeState::NotEnteringAttribute;
+	CurrentEnteringAttributeState = UE::StructuredArchive::Private::EEnteringAttributeState::NotEnteringAttribute;
 }
 
-int32 FStructuredArchive::EnterSlotAsType(StructuredArchive_Private::FSlotPosition Slot, StructuredArchive_Private::EElementType ElementType)
+int32 FStructuredArchive::EnterSlotAsType(UE::StructuredArchive::Private::FSlotPosition Slot, UE::StructuredArchive::Private::EElementType ElementType)
 {
-	EnterSlot(Slot, ElementType == StructuredArchive_Private::EElementType::AttributedValue);
+	EnterSlot(Slot, ElementType == UE::StructuredArchive::Private::EElementType::AttributedValue);
 
 	int32 NewSlotDepth = Slot.Depth + 1;
 
@@ -146,8 +146,8 @@ int32 FStructuredArchive::EnterSlotAsType(StructuredArchive_Private::FSlotPositi
 	//
 	// We don't need to do adjust for attributes, because entering the attribute slot will bump the depth anyway.
 	if (NewSlotDepth < CurrentScope.Num() &&
-		CurrentScope[NewSlotDepth].Type == StructuredArchive_Private::EElementType::AttributedValue &&
-		CurrentEnteringAttributeState == StructuredArchive_Private::EEnteringAttributeState::NotEnteringAttribute)
+		CurrentScope[NewSlotDepth].Type == UE::StructuredArchive::Private::EElementType::AttributedValue &&
+		CurrentEnteringAttributeState == UE::StructuredArchive::Private::EEnteringAttributeState::NotEnteringAttribute)
 	{
 		++NewSlotDepth;
 	}
@@ -162,32 +162,32 @@ void FStructuredArchive::LeaveSlot()
 	{
 		switch (CurrentScope.Top().Type)
 		{
-		case StructuredArchive_Private::EElementType::Record:
+		case UE::StructuredArchive::Private::EElementType::Record:
 			Formatter.LeaveField();
 			break;
-		case StructuredArchive_Private::EElementType::Array:
+		case UE::StructuredArchive::Private::EElementType::Array:
 			Formatter.LeaveArrayElement();
 #if DO_STRUCTURED_ARCHIVE_CONTAINER_CHECKS
 			CurrentContainer.Top()->Index++;
 #endif
 			break;
-		case StructuredArchive_Private::EElementType::Stream:
+		case UE::StructuredArchive::Private::EElementType::Stream:
 			Formatter.LeaveStreamElement();
 			break;
-		case StructuredArchive_Private::EElementType::Map:
+		case UE::StructuredArchive::Private::EElementType::Map:
 			Formatter.LeaveMapElement();
 #if DO_STRUCTURED_ARCHIVE_CONTAINER_CHECKS
 			CurrentContainer.Top()->Index++;
 #endif
 			break;
-		case StructuredArchive_Private::EElementType::AttributedValue:
+		case UE::StructuredArchive::Private::EElementType::AttributedValue:
 			Formatter.LeaveAttribute();
 			break;
 		}
 	}
 }
 
-void FStructuredArchive::SetScope(StructuredArchive_Private::FSlotPosition Slot)
+void FStructuredArchive::SetScope(UE::StructuredArchive::Private::FSlotPosition Slot)
 {
 	// Make sure the scope is valid
 	checkf(Slot.Depth < CurrentScope.Num() && CurrentScope[Slot.Depth].Id == Slot.ElementId, TEXT("Invalid scope for writing to archive"));
@@ -202,13 +202,13 @@ void FStructuredArchive::SetScope(StructuredArchive_Private::FSlotPosition Slot)
 			const FElement& Element = CurrentScope[CurrentDepth];
 			switch (Element.Type)
 			{
-			case StructuredArchive_Private::EElementType::Record:
+			case UE::StructuredArchive::Private::EElementType::Record:
 				Formatter.LeaveRecord();
 #if DO_STRUCTURED_ARCHIVE_CONTAINER_CHECKS
 				CurrentContainer.Pop(false);
 #endif
 				break;
-			case StructuredArchive_Private::EElementType::Array:
+			case UE::StructuredArchive::Private::EElementType::Array:
 #if DO_STRUCTURED_ARCHIVE_CONTAINER_CHECKS
 				checkf(GetUnderlyingArchive().IsLoading() || CurrentContainer.Top()->Index == CurrentContainer.Top()->Count, TEXT("Incorrect number of elements serialized in array"));
 #endif
@@ -217,10 +217,10 @@ void FStructuredArchive::SetScope(StructuredArchive_Private::FSlotPosition Slot)
 				CurrentContainer.Pop(false);
 #endif
 				break;
-			case StructuredArchive_Private::EElementType::Stream:
+			case UE::StructuredArchive::Private::EElementType::Stream:
 				Formatter.LeaveStream();
 				break;
-			case StructuredArchive_Private::EElementType::Map:
+			case UE::StructuredArchive::Private::EElementType::Map:
 #if DO_STRUCTURED_ARCHIVE_CONTAINER_CHECKS
 				checkf(CurrentContainer.Top()->Index == CurrentContainer.Top()->Count, TEXT("Incorrect number of elements serialized in map"));
 #endif
@@ -229,7 +229,7 @@ void FStructuredArchive::SetScope(StructuredArchive_Private::FSlotPosition Slot)
 				CurrentContainer.Pop(false);
 #endif
 				break;
-			case StructuredArchive_Private::EElementType::AttributedValue:
+			case UE::StructuredArchive::Private::EElementType::AttributedValue:
 				Formatter.LeaveAttributedValue();
 #if DO_STRUCTURED_ARCHIVE_CONTAINER_CHECKS
 				CurrentContainer.Pop(false);
@@ -255,7 +255,7 @@ void FStructuredArchive::SetScope(StructuredArchive_Private::FSlotPosition Slot)
 
 FStructuredArchiveRecord FStructuredArchiveSlot::EnterRecord()
 {
-	int32 NewDepth = Ar.EnterSlotAsType(*this, StructuredArchive_Private::EElementType::Record);
+	int32 NewDepth = Ar.EnterSlotAsType(*this, UE::StructuredArchive::Private::EElementType::Record);
 
 #if DO_STRUCTURED_ARCHIVE_CONTAINER_CHECKS
 	Ar.CurrentContainer.Emplace(0);
@@ -268,7 +268,7 @@ FStructuredArchiveRecord FStructuredArchiveSlot::EnterRecord()
 
 FStructuredArchiveRecord FStructuredArchiveSlot::EnterRecord_TextOnly(TArray<FString>& OutFieldNames)
 {
-	int32 NewDepth = Ar.EnterSlotAsType(*this, StructuredArchive_Private::EElementType::Record);
+	int32 NewDepth = Ar.EnterSlotAsType(*this, UE::StructuredArchive::Private::EElementType::Record);
 
 #if DO_STRUCTURED_ARCHIVE_CONTAINER_CHECKS
 	Ar.CurrentContainer.Emplace(0);
@@ -281,7 +281,7 @@ FStructuredArchiveRecord FStructuredArchiveSlot::EnterRecord_TextOnly(TArray<FSt
 
 FStructuredArchiveArray FStructuredArchiveSlot::EnterArray(int32& Num)
 {
-	int32 NewDepth = Ar.EnterSlotAsType(*this, StructuredArchive_Private::EElementType::Array);
+	int32 NewDepth = Ar.EnterSlotAsType(*this, UE::StructuredArchive::Private::EElementType::Array);
 
 	Ar.Formatter.EnterArray(Num);
 
@@ -294,7 +294,7 @@ FStructuredArchiveArray FStructuredArchiveSlot::EnterArray(int32& Num)
 
 FStructuredArchiveStream FStructuredArchiveSlot::EnterStream()
 {
-	int32 NewDepth = Ar.EnterSlotAsType(*this, StructuredArchive_Private::EElementType::Stream);
+	int32 NewDepth = Ar.EnterSlotAsType(*this, UE::StructuredArchive::Private::EElementType::Stream);
 
 	Ar.Formatter.EnterStream();
 
@@ -303,7 +303,7 @@ FStructuredArchiveStream FStructuredArchiveSlot::EnterStream()
 
 FStructuredArchiveStream FStructuredArchiveSlot::EnterStream_TextOnly(int32& OutNumElements)
 {
-	int32 NewDepth = Ar.EnterSlotAsType(*this, StructuredArchive_Private::EElementType::Stream);
+	int32 NewDepth = Ar.EnterSlotAsType(*this, UE::StructuredArchive::Private::EElementType::Stream);
 
 	Ar.Formatter.EnterStream_TextOnly(OutNumElements);
 
@@ -312,7 +312,7 @@ FStructuredArchiveStream FStructuredArchiveSlot::EnterStream_TextOnly(int32& Out
 
 FStructuredArchiveMap FStructuredArchiveSlot::EnterMap(int32& Num)
 {
-	int32 NewDepth = Ar.EnterSlotAsType(*this, StructuredArchive_Private::EElementType::Map);
+	int32 NewDepth = Ar.EnterSlotAsType(*this, UE::StructuredArchive::Private::EElementType::Map);
 
 	Ar.Formatter.EnterMap(Num);
 
@@ -328,9 +328,9 @@ FStructuredArchiveSlot FStructuredArchiveSlot::EnterAttribute(FArchiveFieldName
 	check(Ar.CurrentScope.Num() > 0);
 
 	int32 NewDepth = Depth + 1;
-	if (NewDepth >= Ar.CurrentScope.Num() || Ar.CurrentScope[NewDepth].Id != ElementId || Ar.CurrentScope[NewDepth].Type != StructuredArchive_Private::EElementType::AttributedValue)
+	if (NewDepth >= Ar.CurrentScope.Num() || Ar.CurrentScope[NewDepth].Id != ElementId || Ar.CurrentScope[NewDepth].Type != UE::StructuredArchive::Private::EElementType::AttributedValue)
 	{
-		int32 NewDepthCheck = Ar.EnterSlotAsType(*this, StructuredArchive_Private::EElementType::AttributedValue);
+		int32 NewDepthCheck = Ar.EnterSlotAsType(*this, UE::StructuredArchive::Private::EElementType::AttributedValue);
 		checkSlow(NewDepth == NewDepthCheck);
 
 		Ar.Formatter.EnterAttributedValue();
@@ -340,11 +340,11 @@ FStructuredArchiveSlot FStructuredArchiveSlot::EnterAttribute(FArchiveFieldName
 #endif
 	}
 
-	Ar.CurrentEnteringAttributeState = StructuredArchive_Private::EEnteringAttributeState::NotEnteringAttribute;
+	Ar.CurrentEnteringAttributeState = UE::StructuredArchive::Private::EEnteringAttributeState::NotEnteringAttribute;
 
-	StructuredArchive_Private::FElementId AttributedValueId = Ar.CurrentScope[NewDepth].Id;
+	UE::StructuredArchive::Private::FElementId AttributedValueId = Ar.CurrentScope[NewDepth].Id;
 
-	Ar.SetScope(StructuredArchive_Private::FSlotPosition(NewDepth, AttributedValueId));
+	Ar.SetScope(UE::StructuredArchive::Private::FSlotPosition(NewDepth, AttributedValueId));
 
 	Ar.CurrentSlotElementId = Ar.ElementIdGenerator.Generate();
 
@@ -369,9 +369,9 @@ TOptional<FStructuredArchiveSlot> FStructuredArchiveSlot::TryEnterAttribute(FArc
 	check(Ar.CurrentScope.Num() > 0);
 
 	int32 NewDepth = Depth + 1;
-	if (NewDepth >= Ar.CurrentScope.Num() || Ar.CurrentScope[NewDepth].Id != ElementId || Ar.CurrentScope[NewDepth].Type != StructuredArchive_Private::EElementType::AttributedValue)
+	if (NewDepth >= Ar.CurrentScope.Num() || Ar.CurrentScope[NewDepth].Id != ElementId || Ar.CurrentScope[NewDepth].Type != UE::StructuredArchive::Private::EElementType::AttributedValue)
 	{
-		int32 NewDepthCheck = Ar.EnterSlotAsType(*this, StructuredArchive_Private::EElementType::AttributedValue);
+		int32 NewDepthCheck = Ar.EnterSlotAsType(*this, UE::StructuredArchive::Private::EElementType::AttributedValue);
 		checkSlow(NewDepth == NewDepthCheck);
 
 		Ar.Formatter.EnterAttributedValue();
@@ -392,13 +392,13 @@ TOptional<FStructuredArchiveSlot> FStructuredArchiveSlot::TryEnterAttribute(FArc
 #endif
 #endif
 
-	StructuredArchive_Private::FElementId AttributedValueId = Ar.CurrentScope[NewDepth].Id;
+	UE::StructuredArchive::Private::FElementId AttributedValueId = Ar.CurrentScope[NewDepth].Id;
 
-	Ar.SetScope(StructuredArchive_Private::FSlotPosition(NewDepth, AttributedValueId));
+	Ar.SetScope(UE::StructuredArchive::Private::FSlotPosition(NewDepth, AttributedValueId));
 
 	if (Ar.Formatter.TryEnterAttribute(AttributeName, bEnterWhenWriting))
 	{
-		Ar.CurrentEnteringAttributeState = StructuredArchive_Private::EEnteringAttributeState::NotEnteringAttribute;
+		Ar.CurrentEnteringAttributeState = UE::StructuredArchive::Private::EEnteringAttributeState::NotEnteringAttribute;
 
 		Ar.CurrentSlotElementId = Ar.ElementIdGenerator.Generate();
 
diff --git a/Engine/Source/Runtime/Core/Public/Serialization/StructuredArchive.h b/Engine/Source/Runtime/Core/Public/Serialization/StructuredArchive.h
--- a/Engine/Source/Runtime/Core/Public/Serialization/StructuredArchive.h
+++ b/Engine/Source/Runtime/Core/Public/Serialization/StructuredArchive.h
@@ -13,86 +13,67 @@
 #include "Serialization/ArchiveProxy.h"
 #include "Templates/UniqueObj.h"
 
-/**
- * Class to contain a named value for serialization. Intended to be created as a temporary and passed to object serialization methods.
- */
-template<typename T> struct TNamedValue
+namespace UE::StructuredArchive::Private
 {
-	FArchiveFieldName Name;
-	T& Value;
+	/**
+	 * Class to contain a named value for serialization. Intended to be created as a temporary and passed to object serialization methods.
+	 */
+	template <typename T>
+	struct TNamedValue
+	{
+		FArchiveFieldName Name;
+		T& Value;
+	};
 
-	FORCEINLINE TNamedValue(FArchiveFieldName InName, T& InValue)
-		: Name(InName)
-		, Value(InValue)
+	/**
+	 * Class to contain a named attribute for serialization. Intended to be created as a temporary and passed to object serialization methods.
+	 */
+	template <typename T>
+	struct TNamedAttribute
 	{
-	}
-};
+		FArchiveFieldName Name;
+		T& Value;
+	};
 
-/**
- * Class to contain a named attribute for serialization. Intended to be created as a temporary and passed to object serialization methods.
- */
-template<typename T> struct TNamedAttribute
-{
-	FArchiveFieldName Name;
-	T& Value;
+	/**
+	 * Class to contain a named attribute for serialization, with a default. Intended to be created as a temporary and passed to object
+	 * serialization methods, which can choose not to serialize the attribute if it matches the default.
+	 */
+	template <typename T>
+	struct TOptionalNamedAttribute
+	{
+		FArchiveFieldName Name;
+		T& Value;
+		const T& Default;
+	};
 
-	explicit FORCEINLINE TNamedAttribute(FArchiveFieldName InName, T& InValue)
-		: Name(InName)
-		, Value(InValue)
+	template <typename T>
+	FORCEINLINE TNamedValue<T> MakeNamedValue(FArchiveFieldName Name, T& Value)
 	{
+		return TNamedValue<T>{ Name, Value };
 	}
-};
 
-/**
- * Class to contain a named attribute for serialization, with a default. Intended to be created as a temporary and passed to object
- * serialization methods, which can choose not to serialize the attribute if it matches the default.
- */
-template<typename T> struct TOptionalNamedAttribute
-{
-	FArchiveFieldName Name;
-	T& Value;
-	const T& Default;
-
-	explicit FORCEINLINE TOptionalNamedAttribute(FArchiveFieldName InName, T& InValue, const T& InDefault)
-		: Name(InName)
-		, Value(InValue)
-		, Default(InDefault)
+	template <typename T>
+	FORCEINLINE TNamedAttribute<T> MakeNamedAttribute(FArchiveFieldName Name, T& Value)
 	{
+		return TNamedAttribute<T>{ Name, Value };
 	}
-};
 
-/**
- * Helper function to construct a TNamedValue, deducing the value type.
- */
-template<typename T> FORCEINLINE TNamedValue<T> MakeNamedValue(FArchiveFieldName Name, T& Value)
-{
-	return TNamedValue<T>(Name, Value);
-}
-
-/**
- * Helper function to construct a TNamedAttribute, deducing the value type.
- */
-template<typename T> FORCEINLINE TNamedAttribute<T> MakeNamedAttribute(FArchiveFieldName Name, T& Value)
-{
-	return TNamedAttribute<T>(Name, Value);
-}
-
-/**
- * Helper function to construct a TOptionalNamedAttribute, deducing the value type.
- */
-template<typename T> FORCEINLINE TOptionalNamedAttribute<T> MakeOptionalNamedAttribute(FArchiveFieldName Name, T& Value, const typename TIdentity<T>::Type& Default)
-{
-	return TOptionalNamedAttribute<T>(Name, Value, Default);
+	template <typename T>
+	FORCEINLINE TOptionalNamedAttribute<T> MakeOptionalNamedAttribute(FArchiveFieldName Name, T& Value, const typename TIdentity<T>::Type& Default)
+	{
+		return TOptionalNamedAttribute<T>{ Name, Value, Default };
+	}
 }
 
 /** Construct a TNamedValue given an ANSI string and value reference. */
-#define SA_VALUE(Name, Value) MakeNamedValue(FArchiveFieldName(Name), Value)
+#define SA_VALUE(Name, Value) UE::StructuredArchive::Private::MakeNamedValue(FArchiveFieldName(Name), Value)
 
 /** Construct a TNamedAttribute given an ANSI string and value reference. */
-#define SA_ATTRIBUTE(Name, Value) MakeNamedAttribute(FArchiveFieldName(Name), Value)
+#define SA_ATTRIBUTE(Name, Value) UE::StructuredArchive::Private::MakeNamedAttribute(FArchiveFieldName(Name), Value)
 
 /** Construct a TOptionalNamedAttribute given an ANSI string and value reference. */
-#define SA_OPTIONAL_ATTRIBUTE(Name, Value, Default) MakeOptionalNamedAttribute(FArchiveFieldName(Name), Value, Default)
+#define SA_OPTIONAL_ATTRIBUTE(Name, Value, Default) UE::StructuredArchive::Private::MakeOptionalNamedAttribute(FArchiveFieldName(Name), Value, Default)
 
 /** Typedef for which formatter type to support */
 #if WITH_TEXT_ARCHIVE_SUPPORT
@@ -109,7 +90,7 @@ class FStructuredArchiveArray;
 class FStructuredArchiveStream;
 class FStructuredArchiveMap;
 
-namespace StructuredArchive_Private
+namespace UE::StructuredArchive::Private
 {
 	struct FElementId
 	{
@@ -212,7 +193,7 @@ namespace StructuredArchive_Private
  * and can merely have a value serialized into it. That value may be a literal (eg. int, float) or compound object
  * (eg. object, array, map).
  */
-class CORE_API FStructuredArchiveSlot final : public StructuredArchive_Private::FSlotBase
+class CORE_API FStructuredArchiveSlot final : public UE::StructuredArchive::Private::FSlotBase
 {
 public:
 	FStructuredArchiveRecord EnterRecord();
@@ -264,13 +245,13 @@ public:
 	}
 
 	template <typename T>
-	FORCEINLINE void operator<<(TNamedAttribute<T> Item)
+	FORCEINLINE void operator<<(UE::StructuredArchive::Private::TNamedAttribute<T> Item)
 	{
 		EnterAttribute(Item.Name) << Item.Value;
 	}
 
 	template <typename T>
-	FORCEINLINE void operator<<(TOptionalNamedAttribute<T> Item)
+	FORCEINLINE void operator<<(UE::StructuredArchive::Private::TOptionalNamedAttribute<T> Item)
 	{
 		if (TOptional<FStructuredArchiveSlot> Attribute = TryEnterAttribute(Item.Name, Item.Value != Item.Default))
 		{
@@ -296,14 +277,14 @@ private:
 	friend FStructuredArchiveStream;
 	friend FStructuredArchiveMap;
 
-	using StructuredArchive_Private::FSlotBase::FSlotBase;
+	using UE::StructuredArchive::Private::FSlotBase::FSlotBase;
 };
 
 /**
  * Represents a record in the structured archive. An object contains slots that are identified by FArchiveName,
  * which may be compiled out with binary-only archives.
  */
-class CORE_API FStructuredArchiveRecord final : public StructuredArchive_Private::FSlotBase
+class CORE_API FStructuredArchiveRecord final : public UE::StructuredArchive::Private::FSlotBase
 {
 public:
 	FStructuredArchiveSlot EnterField(FArchiveFieldName Name);
@@ -317,7 +298,7 @@ public:
 
 	TOptional<FStructuredArchiveSlot> TryEnterField(FArchiveFieldName Name, bool bEnterForSaving);
 
-	template<typename T> FORCEINLINE FStructuredArchiveRecord& operator<<(TNamedValue<T> Item)
+	template<typename T> FORCEINLINE FStructuredArchiveRecord& operator<<(UE::StructuredArchive::Private::TNamedValue<T> Item)
 	{
 		EnterField(Item.Name) << Item.Value;
 		return *this;
@@ -327,14 +308,14 @@ private:
 	friend FStructuredArchive;
 	friend FStructuredArchiveSlot;
 
-	using StructuredArchive_Private::FSlotBase::FSlotBase;
+	using UE::StructuredArchive::Private::FSlotBase::FSlotBase;
 };
 
 /**
  * Represents an array in the structured archive. An object contains slots that are identified by a FArchiveFieldName,
  * which may be compiled out with binary-only archives.
  */
-class CORE_API FStructuredArchiveArray final : public StructuredArchive_Private::FSlotBase
+class CORE_API FStructuredArchiveArray final : public UE::StructuredArchive::Private::FSlotBase
 {
 public:
 	FStructuredArchiveSlot EnterElement();
@@ -350,13 +331,13 @@ private:
 	friend FStructuredArchive;
 	friend FStructuredArchiveSlot;
 
-	using StructuredArchive_Private::FSlotBase::FSlotBase;
+	using UE::StructuredArchive::Private::FSlotBase::FSlotBase;
 };
 
 /**
  * Represents an unsized sequence of slots in the structured archive (similar to an array, but without a known size).
  */
-class CORE_API FStructuredArchiveStream final : public StructuredArchive_Private::FSlotBase
+class CORE_API FStructuredArchiveStream final : public UE::StructuredArchive::Private::FSlotBase
 {
 public:
 	FStructuredArchiveSlot EnterElement();
@@ -372,14 +353,14 @@ private:
 	friend FStructuredArchive;
 	friend FStructuredArchiveSlot;
 
-	using StructuredArchive_Private::FSlotBase::FSlotBase;
+	using UE::StructuredArchive::Private::FSlotBase::FSlotBase;
 };
 
 /**
  * Represents a map in the structured archive. A map is similar to a record, but keys can be read back out from an archive.
  * (This is an important distinction for binary archives).
  */
-class CORE_API FStructuredArchiveMap final : public StructuredArchive_Private::FSlotBase
+class CORE_API FStructuredArchiveMap final : public UE::StructuredArchive::Private::FSlotBase
 {
 public:
 	FStructuredArchiveSlot EnterElement(FString& Name);
@@ -389,7 +370,7 @@ private:
 	friend FStructuredArchive;
 	friend FStructuredArchiveSlot;
 
-	using StructuredArchive_Private::FSlotBase::FSlotBase;
+	using UE::StructuredArchive::Private::FSlotBase::FSlotBase;
 };
 
 /**
@@ -471,10 +452,10 @@ private:
 
 	struct FElement
 	{
-		StructuredArchive_Private::FElementId Id;
-		StructuredArchive_Private::EElementType Type;
+		UE::StructuredArchive::Private::FElementId Id;
+		UE::StructuredArchive::Private::EElementType Type;
 
-		FElement(StructuredArchive_Private::FElementId InId, StructuredArchive_Private::EElementType InType)
+		FElement(UE::StructuredArchive::Private::FElementId InId, UE::StructuredArchive::Private::EElementType InType)
 			: Id(InId)
 			, Type(InType)
 		{
@@ -483,9 +464,9 @@ private:
 
 	struct FIdGenerator
 	{
-		StructuredArchive_Private::FElementId Generate()
+		UE::StructuredArchive::Private::FElementId Generate()
 		{
-			return StructuredArchive_Private::FElementId(NextId++);
+			return UE::StructuredArchive::Private::FElementId(NextId++);
 		}
 
 	private:
@@ -500,12 +481,12 @@ private:
 	/**
 	 * The ID of the root element.
 	 */
-	StructuredArchive_Private::FElementId RootElementId;
+	UE::StructuredArchive::Private::FElementId RootElementId;
 
 	/**
 	 * The element ID assigned for the current slot. Slots are transient, and only exist as placeholders until something is written into them. This is reset to 0 when something is created in a slot, and the created item can assume the element id.
 	 */
-	StructuredArchive_Private::FElementId CurrentSlotElementId;
+	UE::StructuredArchive::Private::FElementId CurrentSlotElementId;
 
 	/**
 	 * Tracks the current stack of objects being written. Used by SetScope() to ensure that scopes are always closed correctly in the underlying formatter,
@@ -525,19 +506,19 @@ private:
 	/**
 	 * Whether or not we've just entered an attribute
 	 */
-	StructuredArchive_Private::EEnteringAttributeState CurrentEnteringAttributeState = StructuredArchive_Private::EEnteringAttributeState::NotEnteringAttribute;
+	UE::StructuredArchive::Private::EEnteringAttributeState CurrentEnteringAttributeState = UE::StructuredArchive::Private::EEnteringAttributeState::NotEnteringAttribute;
 
 	/**
 	 * Enters the current slot for serializing a value. Asserts if the archive is not in a state about to write to an empty-slot.
 	 */
-	void EnterSlot(StructuredArchive_Private::FSlotPosition Slot, bool bEnteringAttributedValue = false);
+	void EnterSlot(UE::StructuredArchive::Private::FSlotPosition Slot, bool bEnteringAttributedValue = false);
 
 	/**
 	 * Enters the current slot, adding an element onto the stack. Asserts if the archive is not in a state about to write to an empty-slot.
 	 *
 	 * @return  The depth of the newly-entered slot.
 	 */
-	int32 EnterSlotAsType(StructuredArchive_Private::FSlotPosition Slot, StructuredArchive_Private::EElementType ElementType);
+	int32 EnterSlotAsType(UE::StructuredArchive::Private::FSlotPosition Slot, UE::StructuredArchive::Private::EElementType ElementType);
 
 	/**
 	 * Leaves slot at the top of the current scope
@@ -547,16 +528,16 @@ private:
 	/**
 	 * Switches to the scope for the given slot.
 	 */
-	void SetScope(StructuredArchive_Private::FSlotPosition Slot);
+	void SetScope(UE::StructuredArchive::Private::FSlotPosition Slot);
 #endif
 };
 
-FORCEINLINE FArchive& StructuredArchive_Private::FSlotBase::GetUnderlyingArchive() const
+FORCEINLINE FArchive& UE::StructuredArchive::Private::FSlotBase::GetUnderlyingArchive() const
 {
 	return Ar.GetUnderlyingArchive();
 }
 
-FORCEINLINE const FArchiveState& StructuredArchive_Private::FSlotBase::GetArchiveState() const
+FORCEINLINE const FArchiveState& UE::StructuredArchive::Private::FSlotBase::GetArchiveState() const
 {
 	return Ar.GetArchiveState();
 }
@@ -1053,4 +1034,3 @@ typename TEnableIf<
 	}
 
 #endif
-
