Removal of unused CopyFormatted*() and *_TextOnly() functions.

diff --git a/Engine/Source/Runtime/Core/Private/Serialization/Formatters/BinaryArchiveFormatter.cpp b/Engine/Source/Runtime/Core/Private/Serialization/Formatters/BinaryArchiveFormatter.cpp
--- a/Engine/Source/Runtime/Core/Private/Serialization/Formatters/BinaryArchiveFormatter.cpp
+++ b/Engine/Source/Runtime/Core/Private/Serialization/Formatters/BinaryArchiveFormatter.cpp
@@ -16,33 +16,3 @@ bool FBinaryArchiveFormatter::HasDocumentTree() const
 {
 	return false;
 }
-
-void FBinaryArchiveFormatter::EnterRecord_TextOnly(TArray<FString>& OutFieldNames)
-{
-	checkf(false, TEXT("FBinaryArchiveFormatter does not support functions with the _TextOnly() suffix."));
-}
-
-void FBinaryArchiveFormatter::EnterField_TextOnly(FArchiveFieldName Name, EArchiveValueType& OutType)
-{
-	checkf(false, TEXT("FBinaryArchiveFormatter does not support functions with the _TextOnly() suffix."));
-}
-
-void FBinaryArchiveFormatter::EnterArrayElement_TextOnly(EArchiveValueType& OutType)
-{
-	checkf(false, TEXT("FBinaryArchiveFormatter does not support functions with the _TextOnly() suffix."));
-}
-
-void FBinaryArchiveFormatter::EnterStream_TextOnly(int32& NumElements)
-{
-	checkf(false, TEXT("FBinaryArchiveFormatter does not support functions with the _TextOnly() suffix."));
-}
-
-void FBinaryArchiveFormatter::EnterStreamElement_TextOnly(EArchiveValueType& OutType)
-{
-	checkf(false, TEXT("FBinaryArchiveFormatter does not support functions with the _TextOnly() suffix."));
-}
-
-void FBinaryArchiveFormatter::EnterMapElement_TextOnly(FString& Name, EArchiveValueType& OutType)
-{
-	checkf(false, TEXT("FBinaryArchiveFormatter does not support functions with the _TextOnly() suffix."));
-}
diff --git a/Engine/Source/Runtime/Core/Private/Serialization/StructuredArchive.cpp b/Engine/Source/Runtime/Core/Private/Serialization/StructuredArchive.cpp
--- a/Engine/Source/Runtime/Core/Private/Serialization/StructuredArchive.cpp
+++ b/Engine/Source/Runtime/Core/Private/Serialization/StructuredArchive.cpp
@@ -266,19 +266,6 @@ FStructuredArchiveRecord FStructuredArchiveSlot::EnterRecord()
 	return FStructuredArchiveRecord(Ar, NewDepth, ElementId);
 }
 
-FStructuredArchiveRecord FStructuredArchiveSlot::EnterRecord_TextOnly(TArray<FString>& OutFieldNames)
-{
-	int32 NewDepth = Ar.EnterSlotAsType(*this, UE::StructuredArchive::Private::EElementType::Record);
-
-#if DO_STRUCTURED_ARCHIVE_CONTAINER_CHECKS
-	Ar.CurrentContainer.Emplace(0);
-#endif
-
-	Ar.Formatter.EnterRecord_TextOnly(OutFieldNames);
-
-	return FStructuredArchiveRecord(Ar, NewDepth, ElementId);
-}
-
 FStructuredArchiveArray FStructuredArchiveSlot::EnterArray(int32& Num)
 {
 	int32 NewDepth = Ar.EnterSlotAsType(*this, UE::StructuredArchive::Private::EElementType::Array);
@@ -301,15 +288,6 @@ FStructuredArchiveStream FStructuredArchiveSlot::EnterStream()
 	return FStructuredArchiveStream(Ar, NewDepth, ElementId);
 }
 
-FStructuredArchiveStream FStructuredArchiveSlot::EnterStream_TextOnly(int32& OutNumElements)
-{
-	int32 NewDepth = Ar.EnterSlotAsType(*this, UE::StructuredArchive::Private::EElementType::Stream);
-
-	Ar.Formatter.EnterStream_TextOnly(OutNumElements);
-
-	return FStructuredArchiveStream(Ar, NewDepth, ElementId);
-}
-
 FStructuredArchiveMap FStructuredArchiveSlot::EnterMap(int32& Num)
 {
 	int32 NewDepth = Ar.EnterSlotAsType(*this, UE::StructuredArchive::Private::EElementType::Map);
@@ -593,11 +571,6 @@ FStructuredArchiveRecord FStructuredArchiveRecord::EnterRecord(FArchiveFieldName
 	return EnterField(Name).EnterRecord();
 }
 
-FStructuredArchiveRecord FStructuredArchiveRecord::EnterRecord_TextOnly(FArchiveFieldName Name, TArray<FString>& OutFieldNames)
-{
-	return EnterField(Name).EnterRecord_TextOnly(OutFieldNames);
-}
-
 FStructuredArchiveArray FStructuredArchiveRecord::EnterArray(FArchiveFieldName Name, int32& Num)
 {
 	return EnterField(Name).EnterArray(Num);
@@ -608,11 +581,6 @@ FStructuredArchiveStream FStructuredArchiveRecord::EnterStream(FArchiveFieldName
 	return EnterField(Name).EnterStream();
 }
 
-FStructuredArchiveStream FStructuredArchiveRecord::EnterStream_TextOnly(FArchiveFieldName Name, int32& OutNumElements)
-{
-	return EnterField(Name).EnterStream_TextOnly(OutNumElements);
-}
-
 FStructuredArchiveMap FStructuredArchiveRecord::EnterMap(FArchiveFieldName Name, int32& Num)
 {
 	return EnterField(Name).EnterMap(Num);
@@ -661,21 +629,6 @@ FStructuredArchiveSlot FStructuredArchiveArray::EnterElement()
 	return FStructuredArchiveSlot(Ar, Depth, Ar.CurrentSlotElementId);
 }
 
-FStructuredArchiveSlot FStructuredArchiveArray::EnterElement_TextOnly(EArchiveValueType& OutType)
-{
-	Ar.SetScope(*this);
-
-#if DO_STRUCTURED_ARCHIVE_CONTAINER_CHECKS
-	checkf(Ar.CurrentContainer.Top()->Index < Ar.CurrentContainer.Top()->Count, TEXT("Serialized too many array elements"));
-#endif
-
-	Ar.CurrentSlotElementId = Ar.ElementIdGenerator.Generate();
-
-	Ar.Formatter.EnterArrayElement_TextOnly(OutType);
-
-	return FStructuredArchiveSlot(Ar, Depth, Ar.CurrentSlotElementId);
-}
-
 //////////// FStructuredArchiveStream ////////////
 
 FStructuredArchiveSlot FStructuredArchiveStream::EnterElement()
@@ -689,17 +642,6 @@ FStructuredArchiveSlot FStructuredArchiveStream::EnterElement()
 	return FStructuredArchiveSlot(Ar, Depth, Ar.CurrentSlotElementId);
 }
 
-FStructuredArchiveSlot FStructuredArchiveStream::EnterElement_TextOnly(EArchiveValueType& OutType)
-{
-	Ar.SetScope(*this);
-
-	Ar.CurrentSlotElementId = Ar.ElementIdGenerator.Generate();
-
-	Ar.Formatter.EnterStreamElement_TextOnly(OutType);
-
-	return FStructuredArchiveSlot(Ar, Depth, Ar.CurrentSlotElementId);
-}
-
 //////////// FStructuredArchiveMap ////////////
 
 FStructuredArchiveSlot FStructuredArchiveMap::EnterElement(FString& Name)
@@ -739,41 +681,4 @@ FStructuredArchiveSlot FStructuredArchiveMap::EnterElement(FString& Name)
 	return FStructuredArchiveSlot(Ar, Depth, Ar.CurrentSlotElementId);
 }
 
-FStructuredArchiveSlot FStructuredArchiveMap::EnterElement_TextOnly(FString& Name, EArchiveValueType& OutType)
-{
-	Ar.SetScope(*this);
-
-#if DO_STRUCTURED_ARCHIVE_CONTAINER_CHECKS
-	checkf(Ar.CurrentContainer.Top()->Index < Ar.CurrentContainer.Top()->Count, TEXT("Serialized too many map elements"));
-#endif
-
-	Ar.CurrentSlotElementId = Ar.ElementIdGenerator.Generate();
-
-#if DO_STRUCTURED_ARCHIVE_CONTAINER_CHECKS
-#if DO_STRUCTURED_ARCHIVE_UNIQUE_FIELD_NAME_CHECKS
-	if(Ar.GetUnderlyingArchive().IsSaving())
-	{
-		FContainer& Container = *Ar.CurrentContainer.Top();
-		checkf(!Container.KeyNames.Contains(Name), TEXT("Multiple keys called '%s' serialized into record"), *Name);
-		Container.KeyNames.Add(Name);
-	}
-#endif
-#endif
-
-	Ar.Formatter.EnterMapElement_TextOnly(Name, OutType);
-
-#if DO_STRUCTURED_ARCHIVE_CONTAINER_CHECKS
-#if DO_STRUCTURED_ARCHIVE_UNIQUE_FIELD_NAME_CHECKS
-	if(Ar.GetUnderlyingArchive().IsLoading())
-	{
-		FContainer& Container = *Ar.CurrentContainer.Top();
-		checkf(!Container.KeyNames.Contains(Name), TEXT("Multiple keys called '%s' serialized into record"), *Name);
-		Container.KeyNames.Add(Name);
-	}
-#endif
-#endif
-
-	return FStructuredArchiveSlot(Ar, Depth, Ar.CurrentSlotElementId);
-}
-
 #endif
diff --git a/Engine/Source/Runtime/Core/Private/Serialization/StructuredArchiveFormatter.cpp b/Engine/Source/Runtime/Core/Private/Serialization/StructuredArchiveFormatter.cpp
--- a/Engine/Source/Runtime/Core/Private/Serialization/StructuredArchiveFormatter.cpp
+++ b/Engine/Source/Runtime/Core/Private/Serialization/StructuredArchiveFormatter.cpp
@@ -10,174 +10,3 @@
 FStructuredArchiveFormatter::~FStructuredArchiveFormatter()
 {
 }
-
-//////////// Functions ////////////
-
-#if WITH_TEXT_ARCHIVE_SUPPORT
-	static void CopyFormattedSlot(FStructuredArchiveFormatter& InputFormatter, FStructuredArchiveFormatter& OutputFormatter, EArchiveValueType Type);
-
-	template<typename Type> static void CopyFormattedValue(FStructuredArchiveFormatter& InputFormatter, FStructuredArchiveFormatter& OutputFormatter)
-	{
-		Type Value;
-		InputFormatter.Serialize(Value);
-		OutputFormatter.Serialize(Value);
-	}
-
-	static void CopyFormattedRecord(FStructuredArchiveFormatter& InputFormatter, FStructuredArchiveFormatter& OutputFormatter)
-	{
-		TArray<FString> Fields;
-		InputFormatter.EnterRecord_TextOnly(Fields);
-		OutputFormatter.EnterRecord();
-
-		for (const FString& Field : Fields)
-		{
-			EArchiveValueType Type;
-			InputFormatter.EnterField_TextOnly(FArchiveFieldName(*Field), Type);
-			OutputFormatter.EnterField(FArchiveFieldName(*Field));
-
-			CopyFormattedSlot(InputFormatter, OutputFormatter, Type);
-
-			OutputFormatter.LeaveField();
-			InputFormatter.LeaveField();
-		}
-
-		OutputFormatter.LeaveRecord();
-		InputFormatter.LeaveRecord();
-	}
-
-	static void CopyFormattedArray(FStructuredArchiveFormatter& InputFormatter, FStructuredArchiveFormatter& OutputFormatter)
-	{
-		int32 NumElements = 0;
-		InputFormatter.EnterArray(NumElements);
-		OutputFormatter.EnterArray(NumElements);
-
-		for (int32 Idx = 0; Idx < NumElements; Idx++)
-		{
-			EArchiveValueType Type;
-			InputFormatter.EnterArrayElement_TextOnly(Type);
-			OutputFormatter.EnterArrayElement();
-
-			CopyFormattedSlot(InputFormatter, OutputFormatter, Type);
-
-			OutputFormatter.LeaveArrayElement();
-			InputFormatter.LeaveArrayElement();
-		}
-
-		OutputFormatter.LeaveArray();
-		InputFormatter.LeaveArray();
-	}
-
-	static void CopyFormattedStream(FStructuredArchiveFormatter& InputFormatter, FStructuredArchiveFormatter& OutputFormatter)
-	{
-		int32 NumElements = 0;
-		InputFormatter.EnterStream_TextOnly(NumElements);
-		OutputFormatter.EnterStream();
-
-		for (int32 Idx = 0; Idx < NumElements; Idx++)
-		{
-			EArchiveValueType Type;
-			InputFormatter.EnterStreamElement_TextOnly(Type);
-			OutputFormatter.EnterStreamElement();
-
-			CopyFormattedSlot(InputFormatter, OutputFormatter, Type);
-
-			OutputFormatter.LeaveStreamElement();
-			InputFormatter.LeaveStreamElement();
-		}
-
-		OutputFormatter.LeaveStream();
-		InputFormatter.LeaveStream();
-	}
-
-	static void CopyFormattedMap(FStructuredArchiveFormatter& InputFormatter, FStructuredArchiveFormatter& OutputFormatter)
-	{
-		int32 NumElements = 0;
-		InputFormatter.EnterMap(NumElements);
-		OutputFormatter.EnterMap(NumElements);
-
-		for (int32 Idx = 0; Idx < NumElements; Idx++)
-		{
-			FString Key;
-			EArchiveValueType Type;
-			InputFormatter.EnterMapElement_TextOnly(Key, Type);
-			OutputFormatter.EnterMapElement(Key);
-
-			CopyFormattedSlot(InputFormatter, OutputFormatter, Type);
-
-			OutputFormatter.LeaveMapElement();
-			InputFormatter.LeaveMapElement();
-		}
-
-		OutputFormatter.LeaveMap();
-		InputFormatter.LeaveMap();
-	}
-
-	static void CopyFormattedSlot(FStructuredArchiveFormatter& InputFormatter, FStructuredArchiveFormatter& OutputFormatter, EArchiveValueType Type)
-	{
-		switch (Type)
-		{
-		case EArchiveValueType::Record:
-			CopyFormattedRecord(InputFormatter, OutputFormatter);
-			break;
-		case EArchiveValueType::Array:
-			CopyFormattedArray(InputFormatter, OutputFormatter);
-			break;
-		case EArchiveValueType::Stream:
-			CopyFormattedStream(InputFormatter, OutputFormatter);
-			break;
-		case EArchiveValueType::Map:
-			CopyFormattedMap(InputFormatter, OutputFormatter);
-			break;
-		case EArchiveValueType::Int8:
-			CopyFormattedValue<int8>(InputFormatter, OutputFormatter);
-			break;
-		case EArchiveValueType::Int16:
-			CopyFormattedValue<int16>(InputFormatter, OutputFormatter);
-			break;
-		case EArchiveValueType::Int32:
-			CopyFormattedValue<int32>(InputFormatter, OutputFormatter);
-			break;
-		case EArchiveValueType::Int64:
-			CopyFormattedValue<int64>(InputFormatter, OutputFormatter);
-			break;
-		case EArchiveValueType::UInt8:
-			CopyFormattedValue<uint8>(InputFormatter, OutputFormatter);
-			break;
-		case EArchiveValueType::UInt16:
-			CopyFormattedValue<uint16>(InputFormatter, OutputFormatter);
-			break;
-		case EArchiveValueType::UInt32:
-			CopyFormattedValue<uint32>(InputFormatter, OutputFormatter);
-			break;
-		case EArchiveValueType::UInt64:
-			CopyFormattedValue<uint64>(InputFormatter, OutputFormatter);
-			break;
-		case EArchiveValueType::Float:
-			CopyFormattedValue<float>(InputFormatter, OutputFormatter);
-			break;
-		case EArchiveValueType::Double:
-			CopyFormattedValue<double>(InputFormatter, OutputFormatter);
-			break;
-		case EArchiveValueType::Bool:
-			CopyFormattedValue<bool>(InputFormatter, OutputFormatter);
-			break;
-		case EArchiveValueType::String:
-			CopyFormattedValue<FString>(InputFormatter, OutputFormatter);
-			break;
-		case EArchiveValueType::Name:
-			CopyFormattedValue<FName>(InputFormatter, OutputFormatter);
-			break;
-		case EArchiveValueType::Object:
-			checkf(false, TEXT("Converting object types is not currently supported"));
-			break;
-		case EArchiveValueType::RawData:
-			CopyFormattedValue<TArray<uint8>>(InputFormatter, OutputFormatter);
-			break;
-		}
-	}
-
-	CORE_API void CopyFormattedData(FStructuredArchiveFormatter& InputFormatter, FStructuredArchiveFormatter& OutputFormatter)
-	{
-		CopyFormattedRecord(InputFormatter, OutputFormatter);
-	}
-#endif
diff --git a/Engine/Source/Runtime/Core/Public/Serialization/Formatters/BinaryArchiveFormatter.h b/Engine/Source/Runtime/Core/Public/Serialization/Formatters/BinaryArchiveFormatter.h
--- a/Engine/Source/Runtime/Core/Public/Serialization/Formatters/BinaryArchiveFormatter.h
+++ b/Engine/Source/Runtime/Core/Public/Serialization/Formatters/BinaryArchiveFormatter.h
@@ -16,30 +16,24 @@ public:
 	virtual FArchive& GetUnderlyingArchive() override;
 
 	virtual void EnterRecord() override;
-	virtual void EnterRecord_TextOnly(TArray<FString>& OutFieldNames) override;
 	virtual void LeaveRecord() override;
 	virtual void EnterField(FArchiveFieldName Name) override;
-	virtual void EnterField_TextOnly(FArchiveFieldName Name, EArchiveValueType& OutType) override;
 	virtual void LeaveField() override;
 	virtual bool TryEnterField(FArchiveFieldName Name, bool bEnterIfWriting);
 
 	virtual void EnterArray(int32& NumElements) override;
 	virtual void LeaveArray() override;
 	virtual void EnterArrayElement() override;
-	virtual void EnterArrayElement_TextOnly(EArchiveValueType& OutType) override;
 	virtual void LeaveArrayElement() override;
 
 	virtual void EnterStream() override;
-	virtual void EnterStream_TextOnly(int32& NumElements) override;
 	virtual void LeaveStream() override;
 	virtual void EnterStreamElement() override;
-	virtual void EnterStreamElement_TextOnly(EArchiveValueType& OutType) override;
 	virtual void LeaveStreamElement() override;
 
 	virtual void EnterMap(int32& NumElements) override;
 	virtual void LeaveMap() override;
 	virtual void EnterMapElement(FString& Name) override;
-	virtual void EnterMapElement_TextOnly(FString& Name, EArchiveValueType& OutType) override;
 	virtual void LeaveMapElement() override;
 
 	virtual void EnterAttributedValue() override;
diff --git a/Engine/Source/Runtime/Core/Public/Serialization/StructuredArchive.h b/Engine/Source/Runtime/Core/Public/Serialization/StructuredArchive.h
--- a/Engine/Source/Runtime/Core/Public/Serialization/StructuredArchive.h
+++ b/Engine/Source/Runtime/Core/Public/Serialization/StructuredArchive.h
@@ -197,10 +197,8 @@ class CORE_API FStructuredArchiveSlot final : public UE::StructuredArchive::Priv
 {
 public:
 	FStructuredArchiveRecord EnterRecord();
-	FStructuredArchiveRecord EnterRecord_TextOnly(TArray<FString>& OutFieldNames);
 	FStructuredArchiveArray EnterArray(int32& Num);
 	FStructuredArchiveStream EnterStream();
-	FStructuredArchiveStream EnterStream_TextOnly(int32& OutNumElements);
 	FStructuredArchiveMap EnterMap(int32& Num);
 	FStructuredArchiveSlot EnterAttribute(FArchiveFieldName AttributeName);
 	TOptional<FStructuredArchiveSlot> TryEnterAttribute(FArchiveFieldName AttributeName, bool bEnterWhenWriting);
@@ -288,12 +286,9 @@ class CORE_API FStructuredArchiveRecord final : public UE::StructuredArchive::Pr
 {
 public:
 	FStructuredArchiveSlot EnterField(FArchiveFieldName Name);
-	FStructuredArchiveSlot EnterField_TextOnly(FArchiveFieldName Name, EArchiveValueType& OutType);
 	FStructuredArchiveRecord EnterRecord(FArchiveFieldName Name);
-	FStructuredArchiveRecord EnterRecord_TextOnly(FArchiveFieldName Name, TArray<FString>& OutFieldNames);
 	FStructuredArchiveArray EnterArray(FArchiveFieldName Name, int32& Num);
 	FStructuredArchiveStream EnterStream(FArchiveFieldName Name);
-	FStructuredArchiveStream EnterStream_TextOnly(FArchiveFieldName Name, int32& OutNumElements);
 	FStructuredArchiveMap EnterMap(FArchiveFieldName Name, int32& Num);
 
 	TOptional<FStructuredArchiveSlot> TryEnterField(FArchiveFieldName Name, bool bEnterForSaving);
@@ -319,7 +314,6 @@ class CORE_API FStructuredArchiveArray final : public UE::StructuredArchive::Pri
 {
 public:
 	FStructuredArchiveSlot EnterElement();
-	FStructuredArchiveSlot EnterElement_TextOnly(EArchiveValueType& OutType);
 
 	template<typename T> FORCEINLINE FStructuredArchiveArray& operator<<(T& Item)
 	{
@@ -341,7 +335,6 @@ class CORE_API FStructuredArchiveStream final : public UE::StructuredArchive::Pr
 {
 public:
 	FStructuredArchiveSlot EnterElement();
-	FStructuredArchiveSlot EnterElement_TextOnly(EArchiveValueType& OutType);
 
 	template<typename T> FORCEINLINE FStructuredArchiveStream& operator<<(T& Item)
 	{
@@ -364,7 +357,6 @@ class CORE_API FStructuredArchiveMap final : public UE::StructuredArchive::Priva
 {
 public:
 	FStructuredArchiveSlot EnterElement(FString& Name);
-	FStructuredArchiveSlot EnterElement_TextOnly(FString& Name, EArchiveValueType& OutType);
 
 private:
 	friend FStructuredArchive;
@@ -786,12 +778,6 @@ typename TEnableIf<
 		return FStructuredArchiveRecord(Ar);
 	}
 
-	FORCEINLINE FStructuredArchiveRecord FStructuredArchiveSlot::EnterRecord_TextOnly(TArray<FString>& OutFieldNames)
-	{
-		Ar.Formatter.EnterRecord_TextOnly(OutFieldNames);
-		return FStructuredArchiveRecord(Ar);
-	}
-
 	FORCEINLINE FStructuredArchiveArray FStructuredArchiveSlot::EnterArray(int32& Num)
 	{
 		Ar.Formatter.EnterArray(Num);
@@ -803,12 +789,6 @@ typename TEnableIf<
 		return FStructuredArchiveStream(Ar);
 	}
 
-	FORCEINLINE FStructuredArchiveStream FStructuredArchiveSlot::EnterStream_TextOnly(int32& OutNumElements)
-	{
-		Ar.Formatter.EnterStream_TextOnly(OutNumElements);
-		return FStructuredArchiveStream(Ar);
-	}
-
 	FORCEINLINE FStructuredArchiveMap FStructuredArchiveSlot::EnterMap(int32& Num)
 	{
 		Ar.Formatter.EnterMap(Num);
@@ -950,22 +930,11 @@ typename TEnableIf<
 		return FStructuredArchiveSlot(Ar);
 	}
 
-	FORCEINLINE FStructuredArchiveSlot FStructuredArchiveRecord::EnterField_TextOnly(FArchiveFieldName Name, EArchiveValueType& OutType)
-	{
-		Ar.Formatter.EnterField_TextOnly(Name, OutType);
-		return FStructuredArchiveSlot(Ar);
-	}
-
 	FORCEINLINE FStructuredArchiveRecord FStructuredArchiveRecord::EnterRecord(FArchiveFieldName Name)
 	{
 		return EnterField(Name).EnterRecord();
 	}
 
-	FORCEINLINE FStructuredArchiveRecord FStructuredArchiveRecord::EnterRecord_TextOnly(FArchiveFieldName Name, TArray<FString>& OutFieldNames)
-	{
-		return EnterField(Name).EnterRecord_TextOnly(OutFieldNames);
-	}
-
 	FORCEINLINE FStructuredArchiveArray FStructuredArchiveRecord::EnterArray(FArchiveFieldName Name, int32& Num)
 	{
 		return EnterField(Name).EnterArray(Num);
@@ -1000,12 +969,6 @@ typename TEnableIf<
 		return FStructuredArchiveSlot(Ar);
 	}
 
-	FORCEINLINE FStructuredArchiveSlot FStructuredArchiveArray::EnterElement_TextOnly(EArchiveValueType& OutType)
-	{
-		Ar.Formatter.EnterArrayElement_TextOnly(OutType);
-		return FStructuredArchiveSlot(Ar);
-	}
-
 	//////////// FStructuredArchiveStream ////////////
 
 	FORCEINLINE FStructuredArchiveSlot FStructuredArchiveStream::EnterElement()
@@ -1013,12 +976,6 @@ typename TEnableIf<
 		return FStructuredArchiveSlot(Ar);
 	}
 
-	FORCEINLINE FStructuredArchiveSlot FStructuredArchiveStream::EnterElement_TextOnly(EArchiveValueType& OutType)
-	{
-		Ar.Formatter.EnterStreamElement_TextOnly(OutType);
-		return FStructuredArchiveSlot(Ar);
-	}
-
 	//////////// FStructuredArchiveMap ////////////
 
 	FORCEINLINE FStructuredArchiveSlot FStructuredArchiveMap::EnterElement(FString& Name)
@@ -1027,10 +984,4 @@ typename TEnableIf<
 		return FStructuredArchiveSlot(Ar);
 	}
 
-	FORCEINLINE FStructuredArchiveSlot FStructuredArchiveMap::EnterElement_TextOnly(FString& Name, EArchiveValueType& OutType)
-	{
-		Ar.Formatter.EnterMapElement_TextOnly(Name, OutType);
-		return FStructuredArchiveSlot(Ar);
-	}
-
 #endif
diff --git a/Engine/Source/Runtime/Core/Public/Serialization/StructuredArchiveFormatter.h b/Engine/Source/Runtime/Core/Public/Serialization/StructuredArchiveFormatter.h
--- a/Engine/Source/Runtime/Core/Public/Serialization/StructuredArchiveFormatter.h
+++ b/Engine/Source/Runtime/Core/Public/Serialization/StructuredArchiveFormatter.h
@@ -76,30 +76,24 @@ public:
 	virtual bool HasDocumentTree() const = 0;
 
 	virtual void EnterRecord() = 0;
-	virtual void EnterRecord_TextOnly(TArray<FString>& OutFieldNames) = 0;
 	virtual void LeaveRecord() = 0;
 	virtual void EnterField(FArchiveFieldName Name) = 0;
-	virtual void EnterField_TextOnly(FArchiveFieldName Name, EArchiveValueType& OutType) = 0;
 	virtual void LeaveField() = 0;
 	virtual bool TryEnterField(FArchiveFieldName Name, bool bEnterWhenWriting) = 0;
 
 	virtual void EnterArray(int32& NumElements) = 0;
 	virtual void LeaveArray() = 0;
 	virtual void EnterArrayElement() = 0;
-	virtual void EnterArrayElement_TextOnly(EArchiveValueType& OutType) = 0;
 	virtual void LeaveArrayElement() = 0;
 
 	virtual void EnterStream() = 0;
-	virtual void EnterStream_TextOnly(int32& OutNumElements) = 0;
 	virtual void LeaveStream() = 0;
 	virtual void EnterStreamElement() = 0;
-	virtual void EnterStreamElement_TextOnly(EArchiveValueType& OutType) = 0;
 	virtual void LeaveStreamElement() = 0;
 
 	virtual void EnterMap(int32& NumElements) = 0;
 	virtual void LeaveMap() = 0;
 	virtual void EnterMapElement(FString& Name) = 0;
-	virtual void EnterMapElement_TextOnly(FString& Name, EArchiveValueType& OutType) = 0;
 	virtual void LeaveMapElement() = 0;
 
 	virtual void EnterAttributedValue() = 0;
@@ -133,10 +127,3 @@ public:
 	virtual void Serialize(TArray<uint8>& Value) = 0;
 	virtual void Serialize(void* Data, uint64 DataSize) = 0;
 };
-
-#if WITH_TEXT_ARCHIVE_SUPPORT
-	/**
-	 * Copies formatted data from one place to another. Useful for conversion functions or visitor patterns.
-	 */
-	CORE_API void CopyFormattedData(FStructuredArchiveFormatter& InputFormatter, FStructuredArchiveFormatter& OutputFormatter);
-#endif
diff --git a/Engine/Source/Runtime/CoreUObject/Private/Serialization/Formatters/JsonArchiveInputFormatter.cpp b/Engine/Source/Runtime/CoreUObject/Private/Serialization/Formatters/JsonArchiveInputFormatter.cpp
--- a/Engine/Source/Runtime/CoreUObject/Private/Serialization/Formatters/JsonArchiveInputFormatter.cpp
+++ b/Engine/Source/Runtime/CoreUObject/Private/Serialization/Formatters/JsonArchiveInputFormatter.cpp
@@ -62,16 +62,6 @@ void FJsonArchiveInputFormatter::EnterRecord()
 	ObjectStack.Emplace(Value->AsObject(), ValueStack.Num());
 }
 
-void FJsonArchiveInputFormatter::EnterRecord_TextOnly(TArray<FString>& OutFieldNames)
-{
-	EnterRecord();
-	ObjectStack.Top().JsonObject->Values.GetKeys(OutFieldNames);
-	for(int32 Idx = 0; Idx < OutFieldNames.Num(); Idx++)
-	{
-		OutFieldNames[Idx] = UnescapeFieldName(*OutFieldNames[Idx]);
-	}
-}
-
 void FJsonArchiveInputFormatter::LeaveRecord()
 {
 	check(ValueStack.Num() == ObjectStack.Top().ValueCountOnCreation);
@@ -86,12 +76,6 @@ void FJsonArchiveInputFormatter::EnterField(FArchiveFieldName Name)
 	ValueStack.Add(Field);
 }
 
-void FJsonArchiveInputFormatter::EnterField_TextOnly(FArchiveFieldName Name, EArchiveValueType& OutType)
-{
-	EnterField(Name);
-	OutType = GetValueType(*ValueStack.Top());
-}
-
 void FJsonArchiveInputFormatter::LeaveField()
 {
 	ValueStack.Pop();
@@ -142,11 +126,6 @@ void FJsonArchiveInputFormatter::EnterArrayElement()
 	check(ArrayValuesRemainingStack.Top() > 0);
 }
 
-void FJsonArchiveInputFormatter::EnterArrayElement_TextOnly(EArchiveValueType& OutType)
-{
-	OutType = GetValueType(*ValueStack.Top());
-}
-
 void FJsonArchiveInputFormatter::LeaveArrayElement()
 {
 	ValueStack.Pop();
@@ -159,11 +138,6 @@ void FJsonArchiveInputFormatter::EnterStream()
 	EnterArray(NumElements);
 }
 
-void FJsonArchiveInputFormatter::EnterStream_TextOnly(int32& NumElements)
-{
-	EnterArray(NumElements);
-}
-
 void FJsonArchiveInputFormatter::LeaveStream()
 {
 	LeaveArray();
@@ -173,11 +147,6 @@ void FJsonArchiveInputFormatter::EnterStreamElement()
 {
 }
 
-void FJsonArchiveInputFormatter::EnterStreamElement_TextOnly(EArchiveValueType& OutType)
-{
-	OutType = GetValueType(*ValueStack.Top());
-}
-
 void FJsonArchiveInputFormatter::LeaveStreamElement()
 {
 	LeaveArrayElement();
@@ -202,12 +171,6 @@ void FJsonArchiveInputFormatter::EnterMapElement(FString& OutName)
 	ValueStack.Add(MapIteratorStack.Top()->Value);
 }
 
-void FJsonArchiveInputFormatter::EnterMapElement_TextOnly(FString& OutName, EArchiveValueType& OutType)
-{
-	EnterMapElement(OutName);
-	OutType = GetValueType(*ValueStack.Top());
-}
-
 void FJsonArchiveInputFormatter::LeaveMapElement()
 {
 	ValueStack.Pop();
diff --git a/Engine/Source/Runtime/CoreUObject/Private/Serialization/Formatters/JsonArchiveOutputFormatter.cpp b/Engine/Source/Runtime/CoreUObject/Private/Serialization/Formatters/JsonArchiveOutputFormatter.cpp
--- a/Engine/Source/Runtime/CoreUObject/Private/Serialization/Formatters/JsonArchiveOutputFormatter.cpp
+++ b/Engine/Source/Runtime/CoreUObject/Private/Serialization/Formatters/JsonArchiveOutputFormatter.cpp
@@ -44,12 +44,6 @@ void FJsonArchiveOutputFormatter::EnterRecord()
 	TextStartPosStack.Push(Inner.Tell());
 }
 
-void FJsonArchiveOutputFormatter::EnterRecord_TextOnly(TArray<FString>& OutFieldNames)
-{
-	EnterRecord();
-	OutFieldNames.Reset();
-}
-
 void FJsonArchiveOutputFormatter::LeaveRecord()
 {
 	Newline.Pop(false);
@@ -70,12 +64,6 @@ void FJsonArchiveOutputFormatter::EnterField(FArchiveFieldName Name)
 	WriteFieldName(Name.Name);
 }
 
-void FJsonArchiveOutputFormatter::EnterField_TextOnly(FArchiveFieldName Name, EArchiveValueType& OutType)
-{
-	EnterField(Name);
-	OutType = EArchiveValueType::None;
-}
-
 void FJsonArchiveOutputFormatter::LeaveField()
 {
 	bNeedsComma = true;
@@ -106,12 +94,6 @@ void FJsonArchiveOutputFormatter::EnterArrayElement()
 	EnterStreamElement();
 }
 
-void FJsonArchiveOutputFormatter::EnterArrayElement_TextOnly(EArchiveValueType& OutType)
-{
-	EnterArrayElement();
-	OutType = EArchiveValueType::None;
-}
-
 void FJsonArchiveOutputFormatter::LeaveArrayElement()
 {
 	LeaveStreamElement();
@@ -127,12 +109,6 @@ void FJsonArchiveOutputFormatter::EnterStream()
 	TextStartPosStack.Push(Inner.Tell());
 }
 
-void FJsonArchiveOutputFormatter::EnterStream_TextOnly(int32& OutNumElements)
-{
-	EnterStream();
-	OutNumElements = 0;
-}
-
 void FJsonArchiveOutputFormatter::LeaveStream()
 {
 	Newline.Pop(false);
@@ -152,12 +128,6 @@ void FJsonArchiveOutputFormatter::EnterStreamElement()
 	WriteOptionalNewline();
 }
 
-void FJsonArchiveOutputFormatter::EnterStreamElement_TextOnly(EArchiveValueType& OutType)
-{
-	EnterStreamElement();
-	OutType = EArchiveValueType::None;
-}
-
 void FJsonArchiveOutputFormatter::LeaveStreamElement()
 {
 	bNeedsComma = true;
@@ -179,12 +149,6 @@ void FJsonArchiveOutputFormatter::EnterMapElement(FString& Name)
 	EnterField(FArchiveFieldName(*Name));
 }
 
-void FJsonArchiveOutputFormatter::EnterMapElement_TextOnly(FString& Name, EArchiveValueType& OutType)
-{
-	EnterMapElement(Name);
-	OutType = EArchiveValueType::None;
-}
-
 void FJsonArchiveOutputFormatter::LeaveMapElement()
 {
 	LeaveField();
diff --git a/Engine/Source/Runtime/CoreUObject/Public/Serialization/Formatters/JsonArchiveInputFormatter.h b/Engine/Source/Runtime/CoreUObject/Public/Serialization/Formatters/JsonArchiveInputFormatter.h
--- a/Engine/Source/Runtime/CoreUObject/Public/Serialization/Formatters/JsonArchiveInputFormatter.h
+++ b/Engine/Source/Runtime/CoreUObject/Public/Serialization/Formatters/JsonArchiveInputFormatter.h
@@ -23,30 +23,24 @@ public:
 	virtual bool HasDocumentTree() const override;
 
 	virtual void EnterRecord() override;
-	virtual void EnterRecord_TextOnly(TArray<FString>& OutFieldNames) override;
 	virtual void LeaveRecord() override;
 	virtual void EnterField(FArchiveFieldName Name) override;
-	virtual void EnterField_TextOnly(FArchiveFieldName Name, EArchiveValueType& OutType) override;
 	virtual void LeaveField() override;
 	virtual bool TryEnterField(FArchiveFieldName Name, bool bEnterWhenSaving) override;
 
 	virtual void EnterArray(int32& NumElements) override;
 	virtual void LeaveArray() override;
 	virtual void EnterArrayElement() override;
-	virtual void EnterArrayElement_TextOnly(EArchiveValueType& OutType) override;
 	virtual void LeaveArrayElement() override;
 
 	virtual void EnterStream() override;
-	virtual void EnterStream_TextOnly(int32& NumElements) override;
 	virtual void LeaveStream() override;
 	virtual void EnterStreamElement() override;
-	virtual void EnterStreamElement_TextOnly(EArchiveValueType& OutType) override;
 	virtual void LeaveStreamElement() override;
 
 	virtual void EnterMap(int32& NumElements) override;
 	virtual void LeaveMap() override;
 	virtual void EnterMapElement(FString& Name) override;
-	virtual void EnterMapElement_TextOnly(FString& OutName, EArchiveValueType& OutType) override;
 	virtual void LeaveMapElement() override;
 
 	virtual void EnterAttributedValue() override;
diff --git a/Engine/Source/Runtime/CoreUObject/Public/Serialization/Formatters/JsonArchiveOutputFormatter.h b/Engine/Source/Runtime/CoreUObject/Public/Serialization/Formatters/JsonArchiveOutputFormatter.h
--- a/Engine/Source/Runtime/CoreUObject/Public/Serialization/Formatters/JsonArchiveOutputFormatter.h
+++ b/Engine/Source/Runtime/CoreUObject/Public/Serialization/Formatters/JsonArchiveOutputFormatter.h
@@ -19,30 +19,24 @@ public:
 	virtual bool HasDocumentTree() const override;
 
 	virtual void EnterRecord() override;
-	virtual void EnterRecord_TextOnly(TArray<FString>& OutFieldNames) override;
 	virtual void LeaveRecord() override;
 	virtual void EnterField(FArchiveFieldName Name) override;
-	virtual void EnterField_TextOnly(FArchiveFieldName Name, EArchiveValueType& OutType) override;
 	virtual void LeaveField() override;
 	virtual bool TryEnterField(FArchiveFieldName Name, bool bEnterWhenSaving) override;
 
 	virtual void EnterArray(int32& NumElements) override;
 	virtual void LeaveArray() override;
 	virtual void EnterArrayElement() override;
-	virtual void EnterArrayElement_TextOnly(EArchiveValueType& OutType) override;
 	virtual void LeaveArrayElement() override;
 
 	virtual void EnterStream() override;
-	virtual void EnterStream_TextOnly(int32& NumElements) override;
 	virtual void LeaveStream() override;
 	virtual void EnterStreamElement() override;
-	virtual void EnterStreamElement_TextOnly(EArchiveValueType& OutType) override;
 	virtual void LeaveStreamElement() override;
 
 	virtual void EnterMap(int32& NumElements) override;
 	virtual void LeaveMap() override;
 	virtual void EnterMapElement(FString& Name) override;
-	virtual void EnterMapElement_TextOnly(FString& Name, EArchiveValueType& OutType) override;
 	virtual void LeaveMapElement() override;
 
 	virtual void EnterAttributedValue() override;
