(Editgrate) Fixing memstomp in HLSL translator for custom hlsl node

#rnx

diff --git a/Engine/Source/Runtime/Core/Public/UObject/FortniteMainBranchObjectVersion.h b/Engine/Source/Runtime/Core/Public/UObject/FortniteMainBranchObjectVersion.h
--- a/Engine/Source/Runtime/Core/Public/UObject/FortniteMainBranchObjectVersion.h
+++ b/Engine/Source/Runtime/Core/Public/UObject/FortniteMainBranchObjectVersion.h
@@ -219,7 +219,7 @@ struct CORE_API FFortniteMainBranchObjectVersion
 		SystemGuids.Add(DevGuids.MATERIALSHADERMAP_DERIVEDDATA_VER,  FGuid("1DA85DA5733C4A489A1BE117CDC5ACCC"));
 		SystemGuids.Add(DevGuids.NANITE_DERIVEDDATA_VER,             FGuid("49FBF137DBFB4204B6C9E5B91DA15EBD"));
 		SystemGuids.Add(DevGuids.NIAGARASHADERMAP_DERIVEDDATA_VER,   FGuid("2289D5116CF94BC0AFEAEC541468E645"));
-		SystemGuids.Add(DevGuids.Niagara_LatestScriptCompileVersion, FGuid("A366AB6E36266F4CBB2DB87D6C07A6A6"));
+		SystemGuids.Add(DevGuids.Niagara_LatestScriptCompileVersion, FGuid("68A430A98C770B4BB1331E206470CC26"));
 		SystemGuids.Add(DevGuids.SkeletalMeshDerivedDataVersion,     FGuid("ACF593EAAE354FCBB34CF44F5AA1BFD2"));
 		SystemGuids.Add(DevGuids.STATICMESH_DERIVEDDATA_VER,         FGuid("3006d21f867648a68cf4edee185446b7"));
 
