CompressedBuffer: Changed RawHash to FIoHash

- Only the first 20 bytes of the RawHash field in the header will be consumed.
- In a subsequent change, the last 12 bytes of the RawHash field will be set to zero for new compressed buffers.

- Leveraged composite buffers to avoid performing an extra copy of compressed blocks.
- Updated FBlockEncoder::Compress to always use ParallelFor.
- Added ParallelForWithPreWorkWithTaskContext.

#rnx

diff --git a/Engine/Source/Runtime/Core/Private/Compression/CompressedBuffer.cpp b/Engine/Source/Runtime/Core/Private/Compression/CompressedBuffer.cpp
--- a/Engine/Source/Runtime/Core/Private/Compression/CompressedBuffer.cpp
+++ b/Engine/Source/Runtime/Core/Private/Compression/CompressedBuffer.cpp
@@ -3,14 +3,16 @@
 #include "Compression/CompressedBuffer.h"
 
 #include "Algo/Accumulate.h"
+#include "Algo/AnyOf.h"
 #include "Algo/ForEach.h"
+#include "Async/ParallelFor.h"
 #include "Compression/OodleDataCompression.h"
 #include "Hash/Blake3.h"
+#include "IO/IoHash.h"
 #include "Math/UnrealMathUtility.h"
 #include "Misc/ByteSwap.h"
 #include "Misc/Crc.h"
 #include "Serialization/Archive.h"
-#include "Async/ParallelFor.h"
 
 THIRD_PARTY_INCLUDES_START
 #include "Compression/lz4.h"
@@ -266,174 +268,137 @@ FCompositeBuffer FBlockEncoder::Compress(const FCompositeBuffer& RawData, const
 
 	checkf(FMath::IsPowerOfTwo(BlockSize) && BlockSize <= MAX_uint32,
 		TEXT("BlockSize must be a 32-bit power of two but was %" UINT64_FMT "."), BlockSize);
-	const uint64 RawSize = RawData.GetSize();
 
-	if ( RawSize == 0 )
+	const uint64 RawSize = RawData.GetSize();
+	if (RawSize == 0)
 	{
-		return FNoneEncoder().Compress(RawData, BlockSize);
+		return FCompositeBuffer();
 	}
 
-	FBlake3 RawHash;
-
 	const int64 BlockCount = FMath::DivideAndRoundUp(RawSize, BlockSize);
-	check( BlockCount > 0 );
-	checkf(BlockCount < MAX_int32, TEXT("Raw data of size %" UINT64_FMT " with block size %" UINT64_FMT " requires ")
-		TEXT("%" UINT64_FMT " blocks, but the limit is %u."), RawSize, BlockSize, BlockCount, MAX_uint32);
-		
-	uint64 TotalCompressedSize = 0;
+	check(BlockCount > 0);
+	checkf(BlockCount <= MAX_int32, TEXT("Raw data of size %" UINT64_FMT " with block size %" UINT64_FMT " requires ")
+		TEXT("%" UINT64_FMT " blocks, but the limit is %d."), RawSize, BlockSize, BlockCount, MAX_int32);
+
+	FBlake3 RawHash;
 	const uint64 MetaSize = sizeof(uint32) * BlockCount;
-	const uint64 CompressedDataSize = sizeof(FHeader) + MetaSize + GetCompressedBlocksBound(BlockCount, BlockSize, RawSize);
-	FUniqueBuffer CompressedData;
 
 	// Compress the raw data in blocks and store the raw data for incompressible blocks.
+	uint64 CompressedSize = 0;
+	FUniqueBuffer CompressedData;
+	TArray64<uint32> CompressedBlockSizes;
+	TArray<FSharedBuffer> CompressedBlocks;
 
-	if ( BlockCount == 1 )
-	{
-		FUniqueBuffer RawBlockCopy;	
-		FMemoryView RawBlock = RawData.ViewOrCopyRange(0, RawSize, RawBlockCopy);
-		
-		RawHash.Update(RawBlock);
-		
-		// Allocate the buffer for the header, metadata, and compressed blocks.
-		CompressedData = FUniqueBuffer::Alloc(CompressedDataSize);
-		
-		FMutableMemoryView CompressedBlocksView = CompressedData.GetView() + sizeof(FHeader) + MetaSize;
-		
-		FMutableMemoryView CompressedBlock = CompressedBlocksView;
-		if (!CompressBlock(CompressedBlock, RawBlock))
-		{
-			// compressor error !?
-			return FCompositeBuffer();
-		}
-
-		uint64 CompressedBlockSize = CompressedBlock.GetSize();
-		CompressedBlocksView += CompressedBlockSize;
-		TotalCompressedSize += CompressedBlockSize;
-
-		// Return an uncompressed buffer if the compressed data is larger than the raw data.
-		if (RawSize <= MetaSize + TotalCompressedSize)
-		{
-			CompressedData.Reset();
-			return FNoneEncoder().Compress(RawData, BlockSize);
-		}
-			
-		// put CompressedBlockSizes after FHeader :
-		uint32 CompressedBlockSizeForHeader = IntCastChecked<uint32>(CompressedBlockSize);
-		CompressedBlockSizeForHeader = NETWORK_ORDER32(CompressedBlockSizeForHeader);
-		CompressedData.GetView().Mid(sizeof(FHeader), MetaSize).CopyFrom(MakeMemoryView(&CompressedBlockSizeForHeader,sizeof(CompressedBlockSizeForHeader)));
-	}
-	else
+	struct FBlockContext
 	{
-		// BlockCount >= 2
+		FUniqueBuffer RawBlockCopy;
+		FMutableMemoryView CompressedBlocksView;
+	};
 
-		struct CompressedBlockInfo
+	// Allocate and encode blocks in parallel with calculating the raw hash.
+	CompressedBlockSizes.SetNum(BlockCount);
+	TArray<FBlockContext> BlockContexts;
+	ParallelForWithPreWorkWithTaskContext(TEXT("BlockEncoder.Compress.PF"), BlockContexts, int32(BlockCount), 1,
+		// Executed on the calling thread to allocate the header and optionally the contiguous compressed blocks.
+		[this, MetaSize, BlockCount, BlockSize, RawSize, &CompressedData, &CompressedBlocks](const int32 ContextIndex, const int32 ContextCount) -> FBlockContext
 		{
-			// in:
-			uint64 RawBlockSize;
-			FUniqueBuffer RawBlockCopy;
-			FMemoryView RawBlock;
-
-			// out:
-			TArray64<uint8> CompressedData;
-		};
-	
-		TArray64<CompressedBlockInfo> CompressedBlockInfos;
-		CompressedBlockInfos.SetNum(BlockCount);
-	
-		// fill all CompressedBlockInfos :
-		for (int64 BlockIndex=0;BlockIndex<BlockCount;BlockIndex++)
+			if (ContextCount == 1)
+			{
+				// Allocate a contiguous buffer for the header, metadata, and compressed blocks.
+				const uint64 CompressedDataSize = sizeof(FHeader) + MetaSize + GetCompressedBlocksBound(BlockCount, BlockSize, RawSize);
+				CompressedData = FUniqueBuffer::Alloc(CompressedDataSize);
+				return {{}, CompressedData.GetView() + sizeof(FHeader) + MetaSize};
+			}
+			if (ContextIndex == 0)
+			{
+				// Allocate a buffer for the header and metadata. The compressed blocks are allocated by workers.
+				CompressedData = FUniqueBuffer::Alloc(sizeof(FHeader) + MetaSize);
+				// Allocate space to store the non-contiguous compressed blocks from the workers.
+				CompressedBlocks.SetNum(int32(BlockCount));
+			}
+			return {};
+		},
+		// Executed in parallel to encode blocks. Blocks are contiguous when there is only one thread.
+		[this, &RawHash, &RawData, RawSize, BlockSize, &CompressedBlocks, &CompressedBlockSizes](FBlockContext& BlockContext, const int32 BlockIndex)
 		{
 			const uint64 RawOffset = BlockIndex * BlockSize;
 			const uint64 RawBlockSize = FMath::Min(RawSize - RawOffset, BlockSize);
-			CompressedBlockInfos[BlockIndex].RawBlockSize = RawBlockSize;
-			CompressedBlockInfos[BlockIndex].RawBlock = RawData.ViewOrCopyRange(RawOffset, RawBlockSize, CompressedBlockInfos[BlockIndex].RawBlockCopy);
-		}
-		
-		uint64 CompressBlockAllocSize = CompressBlockBound(BlockSize);
+			const FMemoryView RawBlock = RawData.ViewOrCopyRange(RawOffset, RawBlockSize, BlockContext.RawBlockCopy);
 
-		// encode blocks independently, can be parallel :
-		ParallelFor(TEXT("BlockEncoder.Compress.PF"), IntCastChecked<int32>(BlockCount)+1,1, [&CompressedBlockInfos,CompressBlockAllocSize,BlockCount,&RawHash,this](int64 ThreadIndex)
-		{
-			int64 BlockIndex = ThreadIndex-1;
-			if ( BlockIndex < 0 )
+			FUniqueBuffer CompressedBlock;
+			FMutableMemoryView CompressedBlockView = BlockContext.CompressedBlocksView;
+			if (CompressedBlockView.IsEmpty())
 			{
-				// one thread does the sequential hash update for all raw blocks
-				for (int64 HashBlockIndex=0;HashBlockIndex<BlockCount;HashBlockIndex++)
-				{
-					RawHash.Update(CompressedBlockInfos[HashBlockIndex].RawBlock);
-				}
+				// Allocate compressed blocks on worker threads when there is more than one thread.
+				CompressedBlock = FUniqueBuffer::Alloc(CompressBlockBound(RawBlockSize));
+				CompressedBlockView = CompressedBlock;
 			}
 			else
 			{
-				CompressedBlockInfos[BlockIndex].CompressedData.SetNum(CompressBlockAllocSize);
-		
-				FMutableMemoryView CompressedBlockView(CompressedBlockInfos[BlockIndex].CompressedData.GetData(),CompressedBlockInfos[BlockIndex].CompressedData.Num());
-		
-				if (CompressBlock(CompressedBlockView, CompressedBlockInfos[BlockIndex].RawBlock))
-				{
-					check( CompressedBlockView.GetSize() > 0 );
-					// false on SetNum to not realloc
-					CompressedBlockInfos[BlockIndex].CompressedData.SetNum(CompressedBlockView.GetSize(),false);
-				}
-				else
-				{
-					// size 0 indicates error
-					CompressedBlockInfos[BlockIndex].CompressedData.Empty();
-				}
+				// Interleave hashing with encoding when there is only one thread.
+				RawHash.Update(RawBlock);
 			}
-		});
-	
-		// Allocate the buffer for the header, metadata, and compressed blocks.
-		CompressedData = FUniqueBuffer::Alloc(CompressedDataSize);
-
-		// gather results sequentially :
-		TArray64<uint32> CompressedBlockSizes;
-		CompressedBlockSizes.SetNum(BlockCount);
-		FMutableMemoryView CompressedBlocksView = CompressedData.GetView() + sizeof(FHeader) + MetaSize;
-		for (int64 BlockIndex=0;BlockIndex<BlockCount;BlockIndex++)
-		{
-			FMemoryView CompressedBlock(CompressedBlockInfos[BlockIndex].CompressedData.GetData(),CompressedBlockInfos[BlockIndex].CompressedData.Num());
 
-			uint64 CompressedBlockSize = CompressedBlock.GetSize();
-			if ( CompressedBlockSize == 0 )
+			if (!CompressBlock(CompressedBlockView, RawBlock))
 			{
-				// Compressor failed !?
-				CompressedData.Reset();
-				return FNoneEncoder().Compress(RawData, BlockSize);			
+				return;
 			}
-			else if (CompressedBlockInfos[BlockIndex].RawBlockSize <= CompressedBlockSize)
+
+			uint64 CompressedBlockSize = CompressedBlockView.GetSize();
+			if (RawBlockSize <= CompressedBlockSize)
+			{
+				// Encoding did not compress the block. Copy the raw block.
+				CompressedBlockSize = RawBlockSize;
+				CompressedBlockView.CopyFrom(RawBlock);
+				CompressedBlockView.LeftInline(CompressedBlockSize);
+			}
+
+			if (CompressedBlock.IsNull())
 			{
-				CompressedBlockSize = CompressedBlockInfos[BlockIndex].RawBlockSize;
-				CompressedBlocksView = CompressedBlocksView.CopyFrom(CompressedBlockInfos[BlockIndex].RawBlock);
+				// Blocks are contiguous. Update the view to the start of the next block.
+				BlockContext.CompressedBlocksView += CompressedBlockSize;
 			}
 			else
 			{
-				CompressedBlocksView = CompressedBlocksView.CopyFrom(CompressedBlock);
+				// Blocks are non-contiguous. Store the view of this block to gather into a composite buffer.
+				CompressedBlocks[BlockIndex] = FSharedBuffer::MakeView(CompressedBlockView, CompressedBlock.MoveToShared());
 			}
 
-			// Can be freed after copied out :
-			CompressedBlockInfos[BlockIndex].CompressedData.Empty();
-
-			CompressedBlockSizes[BlockIndex] = IntCastChecked<uint32>(CompressedBlockSize);
-			TotalCompressedSize += CompressedBlockSize;
-		}
-	
-		// CompressedBlockInfos not needed anymore :
-		CompressedBlockInfos.Empty();
+			// Record the compressed block size. A size of 0 is an error and is handled by CompressBlock.
+			check(CompressedBlockSize > 0);
+			CompressedBlockSizes[BlockIndex] = uint32(CompressedBlockSize);
+		},
+		// Executed on the calling thread before it joins in with encoding blocks.
+		[&RawHash, &RawData, &BlockContexts]
+		{
+			if (BlockContexts.Num() > 1)
+			{
+				// Hash the raw data in parallel with encoding the blocks when there are multiple threads.
+				TRACE_CPUPROFILER_EVENT_SCOPE(FBlockEncoder::Compress::RawHash);
+				RawHash.Update(RawData);
+			}
+		});
 
-		// Return an uncompressed buffer if the compressed data is larger than the raw data.
-		if (RawSize <= MetaSize + TotalCompressedSize)
+	// Accumulate compressed size and check for failed blocks.
+	for (const uint64 CompressedBlockSize : CompressedBlockSizes)
+	{
+		if (CompressedBlockSize == 0)
 		{
-			CompressedData.Reset();
-			return FNoneEncoder().Compress(RawData, BlockSize);
+			return FCompositeBuffer();
 		}
+		CompressedSize += CompressedBlockSize;
+	}
 
-		// put CompressedBlockSizes after FHeader :
-		Algo::ForEach(CompressedBlockSizes, [](uint32& Size) { Size = NETWORK_ORDER32(Size); });
-		CompressedData.GetView().Mid(sizeof(FHeader), MetaSize).CopyFrom(MakeMemoryView(CompressedBlockSizes));
+	// Do not return compressed data unless it is smaller than the raw data.
+	if (RawSize <= MetaSize + CompressedSize)
+	{
+		return FCompositeBuffer();
 	}
 
 	// Write the header and calculate the CRC-32.
+	Algo::ForEach(CompressedBlockSizes, [](uint32& Size) { Size = NETWORK_ORDER32(Size); });
+	CompressedData.GetView().Mid(sizeof(FHeader), MetaSize).CopyFrom(MakeMemoryView(CompressedBlockSizes));
+
 	FHeader Header;
 	Header.Method = GetMethod();
 	Header.Compressor = GetCompressor();
@@ -441,12 +406,12 @@ FCompositeBuffer FBlockEncoder::Compress(const FCompositeBuffer& RawData, const
 	Header.BlockSizeExponent = uint8(FMath::FloorLog2_64(BlockSize));
 	Header.BlockCount = IntCastChecked<uint32>(BlockCount);
 	Header.TotalRawSize = RawSize;
-	Header.TotalCompressedSize = sizeof(FHeader) + MetaSize + TotalCompressedSize;
+	Header.TotalCompressedSize = sizeof(FHeader) + MetaSize + CompressedSize;
 	Header.RawHash = RawHash.Finalize();
 	Header.Write(CompressedData.GetView().Left(sizeof(FHeader) + MetaSize));
 
 	const FMemoryView CompositeView = CompressedData.GetView().Left(Header.TotalCompressedSize);
-	return FCompositeBuffer(FSharedBuffer::MakeView(CompositeView, CompressedData.MoveToShared()));
+	return FCompositeBuffer(FSharedBuffer::MakeView(CompositeView, CompressedData.MoveToShared()), MoveTemp(CompressedBlocks));
 }
 
 class FBlockDecoder : public FDecoder
@@ -905,13 +870,13 @@ FCompressedBuffer FCompressedBuffer::Compress(
 	}
 
 	FCompressedBuffer Local;
-	if (CompressionLevel == ECompressedBufferCompressionLevel::None)
+	if (CompressionLevel != ECompressedBufferCompressionLevel::None)
 	{
-		Local.CompressedData = FNoneEncoder().Compress(RawData, BlockSize);
+		Local.CompressedData = FOodleEncoder(Compressor, CompressionLevel).Compress(RawData, BlockSize);
 	}
-	else
+	if (Local.CompressedData.IsNull())
 	{
-		Local.CompressedData = FOodleEncoder(Compressor, CompressionLevel).Compress(RawData, BlockSize);
+		Local.CompressedData = FNoneEncoder().Compress(RawData, BlockSize);
 	}
 	return Local;
 }
@@ -997,9 +962,9 @@ uint64 FCompressedBuffer::GetRawSize() const
 	return CompressedData ? UE::CompressedBuffer::Private::FHeader::Read(CompressedData).TotalRawSize : 0;
 }
 
-FBlake3Hash FCompressedBuffer::GetRawHash() const
+FIoHash FCompressedBuffer::GetRawHash() const
 {
-	return CompressedData ? UE::CompressedBuffer::Private::FHeader::Read(CompressedData).RawHash : FBlake3Hash();
+	return CompressedData ? UE::CompressedBuffer::Private::FHeader::Read(CompressedData).RawHash : FIoHash();
 }
 
 bool FCompressedBuffer::TryGetCompressParameters(
@@ -1132,7 +1097,7 @@ uint64 FCompressedBufferReader::GetRawSize()
 	return 0;
 }
 
-FBlake3Hash FCompressedBufferReader::GetRawHash()
+FIoHash FCompressedBufferReader::GetRawHash()
 {
 	using namespace UE::CompressedBuffer::Private;
 	FHeader Header;
@@ -1141,7 +1106,7 @@ FBlake3Hash FCompressedBufferReader::GetRawHash()
 	{
 		return Header.RawHash;
 	}
-	return FBlake3Hash();
+	return FIoHash();
 }
 
 bool FCompressedBufferReader::TryGetCompressParameters(
diff --git a/Engine/Source/Runtime/Core/Private/Tests/Compression/CompressedBufferTest.cpp b/Engine/Source/Runtime/Core/Private/Tests/Compression/CompressedBufferTest.cpp
deleted file mode 100644
diff --git a/Engine/Source/Runtime/Core/Public/Async/ParallelFor.h b/Engine/Source/Runtime/Core/Public/Async/ParallelFor.h
--- a/Engine/Source/Runtime/Core/Public/Async/ParallelFor.h
+++ b/Engine/Source/Runtime/Core/Public/Async/ParallelFor.h
@@ -484,6 +484,71 @@ inline void ParallelForWithPreWork(const TCHAR* DebugName, int32 Num, int32 MinB
 	ParallelForImpl::ParallelForInternal(DebugName, Num, MinBatchSize, Body, CurrentThreadWorkToDoBeforeHelping, Flags, TArrayView<nullptr_t>());
 }
 
+/** 
+ * General purpose parallel for that uses the taskgraph
+ * @param DebugName; ProfilingScope and DebugName
+ * @param OutContexts; Array that will hold the user-defined, task-level context objects (allocated per parallel task)
+ * @param Num; number of calls of Body; Body(0), Body(1), ..., Body(Num - 1)
+ * @param MinBatchSize; Minimum Size of a Batch (will only launch DivUp(Num, MinBatchSize) Workers 
+ * @param ContextConstructor; Function to call to initialize each task context allocated for the operation
+ * @param Body; Function to call from multiple threads
+ * @param CurrentThreadWorkToDoBeforeHelping; The work is performed on the main thread before it starts helping with the ParallelFor proper
+ * @param Flags; Used to customize the behavior of the ParallelFor if needed.
+ * Notes: Please add stats around to calls to parallel for and within your lambda as appropriate. Do not clog the task graph with long running tasks or tasks that block.
+ */
+template <typename ContextType, typename ContextAllocatorType, typename ContextConstructorType, typename BodyType, typename PreWorkType>
+inline void ParallelForWithPreWorkWithTaskContext(
+	const TCHAR* DebugName,
+	TArray<ContextType, ContextAllocatorType>& OutContexts,
+	int32 Num,
+	int32 MinBatchSize,
+	ContextConstructorType&& ContextConstructor,
+	BodyType&& Body,
+	PreWorkType&& CurrentThreadWorkToDoBeforeHelping,
+	EParallelForFlags Flags = EParallelForFlags::None)
+{
+	if (Num > 0)
+	{
+		const int32 NumContexts = ParallelForImpl::GetNumberOfThreadTasks(Num, MinBatchSize, Flags);
+		OutContexts.Reset(NumContexts);
+		for (int32 ContextIndex = 0; ContextIndex < NumContexts; ++ContextIndex)
+		{
+			OutContexts.Emplace(ContextConstructor(ContextIndex, NumContexts));
+		}
+		ParallelForImpl::ParallelForInternal(DebugName, Num, MinBatchSize, Forward<BodyType>(Body), Forward<PreWorkType>(CurrentThreadWorkToDoBeforeHelping), Flags, TArrayView<ContextType>(OutContexts));
+	}
+}
+
+/** 
+ * General purpose parallel for that uses the taskgraph
+ * @param DebugName; ProfilingScope and DebugName
+ * @param OutContexts; Array that will hold the user-defined, task-level context objects (allocated per parallel task)
+ * @param Num; number of calls of Body; Body(0), Body(1), ..., Body(Num - 1)
+ * @param MinBatchSize; Minimum Size of a Batch (will only launch DivUp(Num, MinBatchSize) Workers 
+ * @param Body; Function to call from multiple threads
+ * @param CurrentThreadWorkToDoBeforeHelping; The work is performed on the main thread before it starts helping with the ParallelFor proper
+ * @param Flags; Used to customize the behavior of the ParallelFor if needed.
+ * Notes: Please add stats around to calls to parallel for and within your lambda as appropriate. Do not clog the task graph with long running tasks or tasks that block.
+ */
+template <typename ContextType, typename ContextAllocatorType, typename BodyType, typename PreWorkType>
+inline void ParallelForWithPreWorkWithTaskContext(
+	const TCHAR* DebugName,
+	TArray<ContextType, ContextAllocatorType>& OutContexts,
+	int32 Num,
+	int32 MinBatchSize,
+	BodyType&& Body,
+	PreWorkType&& CurrentThreadWorkToDoBeforeHelping,
+	EParallelForFlags Flags = EParallelForFlags::None)
+{
+	if (Num > 0)
+	{
+		const int32 NumContexts = ParallelForImpl::GetNumberOfThreadTasks(Num, MinBatchSize, Flags);
+		OutContexts.Reset();
+		OutContexts.AddDefaulted(NumContexts);
+		ParallelForImpl::ParallelForInternal(DebugName, Num, MinBatchSize, Forward<BodyType>(Body), Forward<PreWorkType>(CurrentThreadWorkToDoBeforeHelping), Flags, TArrayView<ContextType>(OutContexts));
+	}
+}
+
 /** 
 	*	General purpose parallel for that uses the taskgraph. This variant constructs for the caller a user-defined context
 	* 	object for each task that may get spawned to do work, and passes it on to the loop body to give it a task-local
diff --git a/Engine/Source/Runtime/Core/Public/Compression/CompressedBuffer.h b/Engine/Source/Runtime/Core/Public/Compression/CompressedBuffer.h
--- a/Engine/Source/Runtime/Core/Public/Compression/CompressedBuffer.h
+++ b/Engine/Source/Runtime/Core/Public/Compression/CompressedBuffer.h
@@ -7,7 +7,7 @@
 #include "Memory/SharedBuffer.h"
 
 class FArchive;
-struct FBlake3Hash;
+struct FIoHash;
 
 namespace FOodleDataCompression { enum class ECompressionLevel : int8; }
 namespace FOodleDataCompression { enum class ECompressor : uint8; }
@@ -113,7 +113,7 @@ public:
 	[[nodiscard]] CORE_API uint64 GetRawSize() const;
 
 	/** Returns the hash of the raw data. Zero on error or if this is null. */
-	[[nodiscard]] CORE_API FBlake3Hash GetRawHash() const;
+	[[nodiscard]] CORE_API FIoHash GetRawHash() const;
 
 	/**
 	 * Returns the compressor and compression level used by this buffer.
@@ -222,7 +222,7 @@ public:
 	[[nodiscard]] CORE_API uint64 GetRawSize();
 
 	/** Returns the hash of the raw data. Zero on error. */
-	[[nodiscard]] CORE_API FBlake3Hash GetRawHash();
+	[[nodiscard]] CORE_API FIoHash GetRawHash();
 
 	/**
 	 * Returns the compressor and compression level used by this buffer.
diff --git a/Engine/Source/Runtime/Core/Tests/Compression/CompressedBufferTest.cpp b/Engine/Source/Runtime/Core/Tests/Compression/CompressedBufferTest.cpp
new file mode 100644
--- /dev/null
+++ b/Engine/Source/Runtime/Core/Tests/Compression/CompressedBufferTest.cpp
@@ -0,0 +1,248 @@
+// Copyright Epic Games, Inc. All Rights Reserved.
+
+#include "Compression/CompressedBuffer.h"
+
+#include "Algo/Compare.h"
+#include "Compression/OodleDataCompression.h"
+#include "IO/IoHash.h"
+#include "Serialization/BufferReader.h"
+#include "TestCommon/CoreUtilities.h"
+#include "TestHarness.h"
+
+TEST_CASE("Core::Compression::CompressedBuffer::Compress", "[Core][Compression][Smoke]")
+{
+	InitCommandLine();
+
+	const uint8 ZeroBuffer[1024]{};
+	const FIoHash ZeroBufferHash = FIoHash::HashBuffer(MakeMemoryView(ZeroBuffer));
+
+	SECTION("Null")
+	{
+		const FCompressedBuffer Buffer;
+		CHECK_FALSE(Buffer);
+		CHECK(Buffer.IsNull());
+		CHECK(Buffer.IsOwned());
+		CHECK(Buffer.GetCompressedSize() == 0);
+		CHECK(Buffer.GetRawSize() == 0);
+		CHECK(Buffer.GetRawHash() == FIoHash::Zero);
+		CHECK(Buffer.Decompress().IsNull());
+		CHECK(Buffer.DecompressToComposite().IsNull());
+	}
+
+	SECTION("Method None")
+	{
+		const FCompressedBuffer BufferNew = FCompressedBuffer::Compress(FSharedBuffer::MakeView(MakeMemoryView(ZeroBuffer)),
+			ECompressedBufferCompressor::NotSet, ECompressedBufferCompressionLevel::None);
+		const FCompressedBuffer BufferCopy = FCompressedBuffer::FromCompressed(BufferNew.GetCompressed());
+		const FCompressedBuffer Buffer = GENERATE_REF(BufferNew, BufferCopy);
+
+		CHECK(Buffer);
+		CHECK(!Buffer.IsNull());
+		CHECK(Buffer.IsOwned());
+		CHECK(Buffer.GetCompressedSize() == sizeof(ZeroBuffer) + 64);
+		CHECK(Buffer.GetRawSize() == sizeof(ZeroBuffer));
+		CHECK(Buffer.GetRawHash() == ZeroBufferHash);
+		CHECK(FIoHash::HashBuffer(Buffer.Decompress()) == ZeroBufferHash);
+		CHECK(FIoHash::HashBuffer(Buffer.DecompressToComposite()) == Buffer.GetRawHash());
+
+		ECompressedBufferCompressor Compressor = ECompressedBufferCompressor::Kraken;
+		ECompressedBufferCompressionLevel CompressionLevel = ECompressedBufferCompressionLevel::Normal;
+		uint64 BlockSize = MAX_uint64;
+		CHECK(Buffer.TryGetCompressParameters(Compressor, CompressionLevel, BlockSize));
+		CHECK(Compressor == ECompressedBufferCompressor::NotSet);
+		CHECK(CompressionLevel == ECompressedBufferCompressionLevel::None);
+		CHECK(BlockSize == 0);
+	}
+
+	SECTION("Method Oodle")
+	{
+		const FCompressedBuffer BufferNew = FCompressedBuffer::Compress(FSharedBuffer::MakeView(MakeMemoryView(ZeroBuffer)),
+			ECompressedBufferCompressor::Mermaid, ECompressedBufferCompressionLevel::VeryFast);
+		const FCompressedBuffer BufferCopy = FCompressedBuffer::FromCompressed(BufferNew.GetCompressed());
+		const FCompressedBuffer Buffer = GENERATE_REF(BufferNew, BufferCopy);
+
+		CHECK(Buffer);
+		CHECK(!Buffer.IsNull());
+		CHECK(Buffer.IsOwned());
+		CHECK(Buffer.GetCompressedSize() < sizeof(ZeroBuffer));
+		CHECK(Buffer.GetRawSize() == sizeof(ZeroBuffer));
+		CHECK(Buffer.GetRawHash() == ZeroBufferHash);
+		CHECK(FIoHash::HashBuffer(Buffer.Decompress()) == ZeroBufferHash);
+		CHECK(FIoHash::HashBuffer(Buffer.DecompressToComposite()) == Buffer.GetRawHash());
+
+		ECompressedBufferCompressor Compressor = ECompressedBufferCompressor::Kraken;
+		ECompressedBufferCompressionLevel CompressionLevel = ECompressedBufferCompressionLevel::Normal;
+		uint64 BlockSize = MAX_uint64;
+		CHECK(Buffer.TryGetCompressParameters(Compressor, CompressionLevel, BlockSize));
+		CHECK(Compressor == ECompressedBufferCompressor::Mermaid);
+		CHECK(CompressionLevel == ECompressedBufferCompressionLevel::VeryFast);
+		CHECK(FMath::IsPowerOfTwo(BlockSize));
+	}
+}
+
+TEST_CASE("Core::Compression::CompressedBuffer::Decompress", "[Core][Compression][Smoke]")
+{
+	const auto GenerateData = [](int32 N) -> TArray<uint64>
+	{
+		TArray<uint64> Data;
+		Data.SetNum(N);
+		for (int32 Idx = 0; Idx < Data.Num(); ++Idx)
+		{
+			Data[Idx] = uint64(Idx);
+		}
+		return Data;
+	};
+
+	const auto CastToArrayView = [](FMemoryView View) -> TConstArrayView<uint64>
+	{
+		return MakeArrayView(static_cast<const uint64*>(View.GetData()), static_cast<int32>(View.GetSize() / sizeof(uint64)));
+	};
+
+	FCompressedBufferReader Reader;
+
+	SECTION("Decompress with offset and size.")
+	{
+		const auto UncompressAndValidate = [&Reader, &CastToArrayView](
+			const FCompressedBuffer& Compressed,
+			const int32 OffsetCount,
+			const int32 Count,
+			const TConstArrayView<uint64> ExpectedValues)
+		{
+			Reader.SetSource(Compressed);
+			{
+				const FSharedBuffer Uncompressed = Reader.Decompress(OffsetCount * sizeof(uint64), Count * sizeof(uint64));
+				CHECK(Algo::Compare(CastToArrayView(Uncompressed), ExpectedValues.Mid(OffsetCount, Count)));
+			}
+			{
+				FUniqueBuffer Uncompressed = FUniqueBuffer::Alloc(Count * sizeof(uint64));
+				CHECKED_IF(Reader.TryDecompressTo(Uncompressed, OffsetCount * sizeof(uint64)))
+				{
+					CHECK(Algo::Compare(CastToArrayView(Uncompressed), ExpectedValues.Mid(OffsetCount, Count)));
+				}
+			}
+		};
+
+		constexpr uint64 BlockSize = 64 * sizeof(uint64);
+		constexpr int32 N = 5000;
+		const TArray<uint64> ExpectedValues = GenerateData(N);
+
+		const FCompressedBuffer Compressed = FCompressedBuffer::Compress(
+			FSharedBuffer::MakeView(MakeMemoryView(ExpectedValues)),
+			FOodleDataCompression::ECompressor::Mermaid,
+			FOodleDataCompression::ECompressionLevel::Fast,
+			BlockSize);
+
+		UncompressAndValidate(Compressed, 0, N, ExpectedValues);
+		UncompressAndValidate(Compressed, 1, N - 1, ExpectedValues);
+		UncompressAndValidate(Compressed, N - 1, 1, ExpectedValues);
+		UncompressAndValidate(Compressed, 0, 1, ExpectedValues);
+		UncompressAndValidate(Compressed, 2, 4, ExpectedValues);
+		UncompressAndValidate(Compressed, 0, 512, ExpectedValues);
+		UncompressAndValidate(Compressed, 3, 514, ExpectedValues);
+		UncompressAndValidate(Compressed, 256, 512, ExpectedValues);
+		UncompressAndValidate(Compressed, 512, 512, ExpectedValues);
+		UncompressAndValidate(Compressed, 512, 512, ExpectedValues);
+		UncompressAndValidate(Compressed, 4993, 4, ExpectedValues);
+	}
+
+	SECTION("Decompress with offset only.")
+	{
+		constexpr uint64 BlockSize = 64 * sizeof(uint64);
+		constexpr int32 N = 1000;
+		const TArray<uint64> ExpectedValues = GenerateData(N);
+
+		const FCompressedBuffer Compressed = FCompressedBuffer::Compress(
+			FSharedBuffer::MakeView(MakeMemoryView(ExpectedValues)),
+			FOodleDataCompression::ECompressor::Mermaid,
+			FOodleDataCompression::ECompressionLevel::Fast,
+			BlockSize);
+
+		constexpr uint64 OffsetCount = 150;
+		{
+			FSharedBuffer Buffer = Compressed.GetCompressed().ToShared();
+			FBufferReader Ar(const_cast<void*>(Buffer.GetData()), int64(Buffer.GetSize()), /*bFreeOnClose*/ false, /*bIsPersistent*/ true);
+			FCompressedBufferReaderSourceScope Source(Reader, Ar);
+			const FSharedBuffer Uncompressed = Reader.Decompress(OffsetCount * sizeof(uint64));
+			CHECK(Algo::Compare(CastToArrayView(Uncompressed), MakeArrayView(ExpectedValues).RightChop(OffsetCount)));
+		}
+		{
+			FCompressedBufferReaderSourceScope Source(Reader, Compressed);
+			const FSharedBuffer Uncompressed = Reader.Decompress(OffsetCount * sizeof(uint64));
+			CHECK(Algo::Compare(CastToArrayView(Uncompressed), MakeArrayView(ExpectedValues).RightChop(OffsetCount)));
+		}
+
+		// Short Buffer
+		{
+			const FCompressedBuffer CompressedShort =
+				FCompressedBuffer::FromCompressed(Compressed.GetCompressed().Mid(0, Compressed.GetCompressedSize() - 128));
+			Reader.SetSource(CompressedShort);
+			const FSharedBuffer Uncompressed = Reader.Decompress();
+			CHECK(Uncompressed.IsNull());
+		}
+	}
+
+	SECTION("Decompress only block.")
+	{
+		constexpr uint64 BlockSize = 256 * sizeof(uint64);
+		constexpr int32 N = 100;
+		const TArray<uint64> ExpectedValues = GenerateData(N);
+
+		const FCompressedBuffer Compressed = FCompressedBuffer::Compress(
+			FSharedBuffer::MakeView(MakeMemoryView(ExpectedValues)),
+			FOodleDataCompression::ECompressor::Mermaid,
+			FOodleDataCompression::ECompressionLevel::Fast,
+			BlockSize);
+
+		constexpr uint64 OffsetCount = 2;
+		constexpr uint64 Count = 50;
+		{
+			FSharedBuffer Buffer = Compressed.GetCompressed().ToShared();
+			FBufferReader Ar(const_cast<void*>(Buffer.GetData()), int64(Buffer.GetSize()), /*bFreeOnClose*/ false, /*bIsPersistent*/ true);
+			FCompressedBufferReaderSourceScope Source(Reader, Ar);
+			const FSharedBuffer Uncompressed = Reader.Decompress(OffsetCount * sizeof(uint64), Count * sizeof(uint64));
+			CHECK(Algo::Compare(CastToArrayView(Uncompressed), MakeArrayView(ExpectedValues).Mid(OffsetCount, Count)));
+		}
+		{
+			FCompressedBufferReaderSourceScope Source(Reader, Compressed);
+			const FSharedBuffer Uncompressed = Reader.Decompress(OffsetCount * sizeof(uint64), Count * sizeof(uint64));
+			CHECK(Algo::Compare(CastToArrayView(Uncompressed), MakeArrayView(ExpectedValues).Mid(OffsetCount, Count)));
+		}
+	}
+
+	SECTION("Decompress from an uncompressed buffer.")
+	{
+		constexpr int32 N = 4242;
+		const TArray<uint64> ExpectedValues = GenerateData(N);
+
+		const FCompressedBuffer Compressed = FCompressedBuffer::Compress(
+			FSharedBuffer::MakeView(MakeMemoryView(ExpectedValues)),
+			FOodleDataCompression::ECompressor::NotSet,
+			FOodleDataCompression::ECompressionLevel::None);
+		Reader.SetSource(Compressed);
+
+		{
+			constexpr uint64 OffsetCount = 0;
+			constexpr uint64 Count = N;
+			const FSharedBuffer Uncompressed = Reader.Decompress(OffsetCount * sizeof(uint64), Count * sizeof(uint64));
+			CHECK(Algo::Compare(CastToArrayView(Uncompressed), MakeArrayView(ExpectedValues).Mid(OffsetCount, Count)));
+		}
+
+		{
+			constexpr uint64 OffsetCount = 21;
+			constexpr uint64 Count = 999;
+			const FSharedBuffer Uncompressed = Reader.Decompress(OffsetCount * sizeof(uint64), Count * sizeof(uint64));
+			CHECK(Algo::Compare(CastToArrayView(Uncompressed), MakeArrayView(ExpectedValues).Mid(OffsetCount, Count)));
+		}
+
+		// Short Buffer
+		{
+			const FCompressedBuffer CompressedShort =
+				FCompressedBuffer::FromCompressed(Compressed.GetCompressed().Mid(0, Compressed.GetCompressedSize() - 128));
+			Reader.SetSource(CompressedShort);
+			const FSharedBuffer Uncompressed = Reader.Decompress();
+			CHECK(Uncompressed.IsNull());
+		}
+	}
+
+	CleanupCommandLine();
+}
