- Added FDDPIRegistry::GetAllSortedPlaformInfos which returns fake and disabled platforms, while GetSortedPlatformInfos will now only return "concrete/usable" platforms (independent of SDK status, etc)
- Changed the SupportedPlatforms dialog to show all DDPI platforms, not just the ones that are compiled in currently

diff --git a/Engine/Source/Developer/Localization/Private/LocTextHelper.cpp b/Engine/Source/Developer/Localization/Private/LocTextHelper.cpp
--- a/Engine/Source/Developer/Localization/Private/LocTextHelper.cpp
+++ b/Engine/Source/Developer/Localization/Private/LocTextHelper.cpp
@@ -49,7 +49,7 @@ const TArray<FString>& FLocTextPlatformSplitUtils::GetPlatformsToSplit(const ELo
 			{
 				// the Keys of the FDataDrivenPlatformInfoRegistry platforms are the known ini platform names
 				TArray<FString> TmpArray;
-				for (FName Name : FDataDrivenPlatformInfoRegistry::GetSortedPlatformNames())
+				for (FName Name : FDataDrivenPlatformInfoRegistry::GetSortedPlatformNames(EPlatformInfoType::AllPlatformInfos))
 				{
 					TmpArray.Add(Name.ToString());
 				}
diff --git a/Engine/Source/Editor/ProjectTargetPlatformEditor/Private/Widgets/SProjectTargetPlatformSettings.cpp b/Engine/Source/Editor/ProjectTargetPlatformEditor/Private/Widgets/SProjectTargetPlatformSettings.cpp
--- a/Engine/Source/Editor/ProjectTargetPlatformEditor/Private/Widgets/SProjectTargetPlatformSettings.cpp
+++ b/Engine/Source/Editor/ProjectTargetPlatformEditor/Private/Widgets/SProjectTargetPlatformSettings.cpp
@@ -20,30 +20,15 @@
 BEGIN_SLATE_FUNCTION_BUILD_OPTIMIZATION
 void SProjectTargetPlatformSettings::Construct(const FArguments& InArgs)
 {
-	// Create and sort a list of vanilla platforms that are game targets (sort by display name)
-	// We show all of the platforms regardless of whether we have an SDK installed for them or not
-	for(const PlatformInfo::FTargetPlatformInfo* PlatformInfo : PlatformInfo::GetPlatformInfoArray())
-	{
-		if(PlatformInfo->IsVanilla() && PlatformInfo->PlatformType == EBuildTargetType::Game)
-		{
-			AvailablePlatforms.Add(PlatformInfo);
-		}
-	}
-
-	AvailablePlatforms.Sort([](const PlatformInfo::FTargetPlatformInfo& One, const PlatformInfo::FTargetPlatformInfo& Two) -> bool
-	{
-		return One.DisplayName.CompareTo(Two.DisplayName) < 0;
-	});
-
-	// Generate a widget for each platform
 	TSharedRef<SVerticalBox> PlatformsListBox = SNew(SVerticalBox);
-	for(const PlatformInfo::FTargetPlatformInfo* AvailablePlatform : AvailablePlatforms)
+
+	for (const FDataDrivenPlatformInfo* DDPI : FDataDrivenPlatformInfoRegistry::GetSortedPlatformInfos(EPlatformInfoType::TruePlatformsOnly))
 	{
 		PlatformsListBox->AddSlot()
-		.AutoHeight()
-		[
-			MakePlatformRow(AvailablePlatform->DisplayName, AvailablePlatform->Name, FEditorStyle::GetBrush(AvailablePlatform->GetIconStyleName(EPlatformIconSize::Normal)))
-		];
+			.AutoHeight()
+			[
+				MakePlatformRow(FText::FromName(DDPI->IniPlatformName), DDPI->IniPlatformName, FEditorStyle::GetBrush(DDPI->GetIconStyleName(EPlatformIconSize::Normal)))
+			];
 	}
 
 	ChildSlot
@@ -187,9 +172,9 @@ void SProjectTargetPlatformSettings::HandlePlatformCheckBoxStateChanged(ECheckBo
 		else
 		{
 			// We've deselected "All Platforms", so manually select every available platform
-			for(const PlatformInfo::FTargetPlatformInfo* AvailablePlatform : AvailablePlatforms)
+			for (const FDataDrivenPlatformInfo* DDPI : FDataDrivenPlatformInfoRegistry::GetSortedPlatformInfos(EPlatformInfoType::TruePlatformsOnly))
 			{
-				FGameProjectGenerationModule::Get().UpdateSupportedTargetPlatforms(AvailablePlatform->Name, true);
+				FGameProjectGenerationModule::Get().UpdateSupportedTargetPlatforms(DDPI->IniPlatformName, true);
 			}
 		}
 	}
diff --git a/Engine/Source/Runtime/Core/Private/Misc/DataDrivenPlatformInfoRegistry.cpp b/Engine/Source/Runtime/Core/Private/Misc/DataDrivenPlatformInfoRegistry.cpp
--- a/Engine/Source/Runtime/Core/Private/Misc/DataDrivenPlatformInfoRegistry.cpp
+++ b/Engine/Source/Runtime/Core/Private/Misc/DataDrivenPlatformInfoRegistry.cpp
@@ -14,6 +14,8 @@ namespace
 {
 	TMap<FName, FDataDrivenPlatformInfo> DataDrivenPlatforms;
 	TMap<FName, FName> GlobalPlatformNameAliases;
+	TArray<FName> AllSortedPlatformNames;
+	TArray<const FDataDrivenPlatformInfo*> AllSortedPlatformInfos;
 	TArray<FName> SortedPlatformNames;
 	TArray<const FDataDrivenPlatformInfo*> SortedPlatformInfos;
 #if DDPI_HAS_EXTENDED_PLATFORMINFO_DATA
@@ -430,15 +432,29 @@ const TMap<FName, FDataDrivenPlatformInfo>& FDataDrivenPlatformInfoRegistry::Get
 			}
 		}
 
-		DataDrivenPlatforms.GetKeys(SortedPlatformNames);
+		DataDrivenPlatforms.GetKeys(AllSortedPlatformNames);
 		// now sort them into arrays of keys and values
-		Algo::Sort(SortedPlatformNames, [](FName One, FName Two) -> bool
+		Algo::Sort(AllSortedPlatformNames, [](FName One, FName Two) -> bool
 		{
 			return One.Compare(Two) < 0;
 		});
+		
+		// now remove the invalid platforms (this is not about installed SDKs or anything, just based on ini values)
+		SortedPlatformNames = AllSortedPlatformNames;
+		SortedPlatformNames.RemoveAll([](FName Platform)
+		{
+			return DataDrivenPlatforms[Platform].bIsFakePlatform;
+		});
+
 
 		// now build list of values from the sort
-		SortedPlatformInfos.Empty(SortedPlatformNames.Num());
+		AllSortedPlatformInfos.AddZeroed(AllSortedPlatformNames.Num());
+		for (int Index = 0; Index < AllSortedPlatformInfos.Num(); Index++)
+		{
+			AllSortedPlatformInfos[Index] = &DataDrivenPlatforms[AllSortedPlatformNames[Index]];
+		}
+
+		SortedPlatformInfos.AddZeroed(SortedPlatformNames.Num());
 		for (int Index = 0; Index < SortedPlatformInfos.Num(); Index++)
 		{
 			SortedPlatformInfos[Index] = &DataDrivenPlatforms[SortedPlatformNames[Index]];
@@ -448,20 +464,20 @@ const TMap<FName, FDataDrivenPlatformInfo>& FDataDrivenPlatformInfoRegistry::Get
 	return DataDrivenPlatforms;
 }
 
-const TArray<FName> FDataDrivenPlatformInfoRegistry::GetSortedPlatformNames()
+const TArray<FName> FDataDrivenPlatformInfoRegistry::GetSortedPlatformNames(EPlatformInfoType PlatformType)
 {
 	// make sure we've read in the inis
 	GetAllPlatformInfos();
 
-	return SortedPlatformNames;
+	return PlatformType == EPlatformInfoType::AllPlatformInfos ? AllSortedPlatformNames : SortedPlatformNames;
 }
 
-const TArray<const FDataDrivenPlatformInfo*>& FDataDrivenPlatformInfoRegistry::GetSortedPlatformInfos()
+const TArray<const FDataDrivenPlatformInfo*>& FDataDrivenPlatformInfoRegistry::GetSortedPlatformInfos(EPlatformInfoType PlatformType)
 {
 	// make sure we've read in the inis
 	GetAllPlatformInfos();
 
-	return SortedPlatformInfos;
+	return PlatformType == EPlatformInfoType::AllPlatformInfos ? AllSortedPlatformInfos : SortedPlatformInfos;
 }
 
 
diff --git a/Engine/Source/Runtime/Core/Public/Misc/DataDrivenPlatformInfoRegistry.h b/Engine/Source/Runtime/Core/Public/Misc/DataDrivenPlatformInfoRegistry.h
--- a/Engine/Source/Runtime/Core/Public/Misc/DataDrivenPlatformInfoRegistry.h
+++ b/Engine/Source/Runtime/Core/Public/Misc/DataDrivenPlatformInfoRegistry.h
@@ -12,6 +12,15 @@
 
 #define DDPI_HAS_EXTENDED_PLATFORMINFO_DATA (WITH_EDITOR || IS_PROGRAM)
 
+enum class EPlatformInfoType : uint8
+{
+	/** Only return "real" platforms (no groups or fake platforms) */
+	TruePlatformsOnly,
+
+	/** All known platforms/groups/etc that have a DDPI.ini file */
+	AllPlatformInfos,
+};
+
 #if DDPI_HAS_EXTENDED_PLATFORMINFO_DATA
 
 
@@ -259,9 +268,22 @@ struct CORE_API FDataDrivenPlatformInfoRegistry
 	static const FDataDrivenPlatformInfo& GetPlatformInfo(const char* PlatformName);
 
 
-	// get just names or just infos, 
-	static const TArray<FName> GetSortedPlatformNames();
-	static const TArray<const FDataDrivenPlatformInfo*>& GetSortedPlatformInfos();
+	/**
+	 * Get sorted platorm infos or names. PlatformType is used to choose between true platforms and platform groups (or other fake platforms)
+	 */
+	static const TArray<FName> GetSortedPlatformNames(EPlatformInfoType PlatformType);
+	static const TArray<const FDataDrivenPlatformInfo*>& GetSortedPlatformInfos(EPlatformInfoType PlatformType);
+
+	UE_DEPRECATED(5.1, "Use GetSoprtedPlatformNames that takes a PlatformType parameter")
+	static const TArray<FName> GetSortedPlatformNames()
+	{
+		return GetSortedPlatformNames(EPlatformInfoType::AllPlatformInfos);
+	}
+	UE_DEPRECATED(5.1, "Use GetSortedPlatformInfos that takes a PlatformType parameter")
+	static const TArray<const FDataDrivenPlatformInfo*>& GetSortedPlatformInfos()
+	{
+		return GetSortedPlatformInfos(EPlatformInfoType::AllPlatformInfos);
+	}
 
 	/**
 	 * Gets a list of all known confidential platforms (note these are just the platforms you have access to, so, for example PS4 won't be
diff --git a/Engine/Source/Runtime/Engine/Private/PerQualityLevelProperties.cpp b/Engine/Source/Runtime/Engine/Private/PerQualityLevelProperties.cpp
--- a/Engine/Source/Runtime/Engine/Private/PerQualityLevelProperties.cpp
+++ b/Engine/Source/Runtime/Engine/Private/PerQualityLevelProperties.cpp
@@ -57,13 +57,12 @@ namespace QualityLevelProperty
 		if (bIsPlatformGroup)
 		{
 			// get all the platforms from that group
-			for (const FName& DataDrivenPlatformName : FDataDrivenPlatformInfoRegistry::GetSortedPlatformNames())
+			for (const FDataDrivenPlatformInfo* DataDrivenPlatformInfo : FDataDrivenPlatformInfoRegistry::GetSortedPlatformInfos(EPlatformInfoType::TruePlatformsOnly))
 			{
 				// gather all platform related to the platform group
-				const FDataDrivenPlatformInfo& DataDrivenPlatformInfo = FDataDrivenPlatformInfoRegistry::GetPlatformInfo(DataDrivenPlatformName);
-				if (DataDrivenPlatformInfo.PlatformGroupName == FName(*InPlatformName))
+				if (DataDrivenPlatformInfo->PlatformGroupName == FName(*InPlatformName))
 				{
-					EnginePlatforms.AddUnique(DataDrivenPlatformName);
+					EnginePlatforms.AddUnique(DataDrivenPlatformInfo->IniPlatformName);
 				}
 			}
 		}
diff --git a/Engine/Source/Runtime/Engine/Private/VT/VirtualTextureDataBuilder.cpp b/Engine/Source/Runtime/Engine/Private/VT/VirtualTextureDataBuilder.cpp
--- a/Engine/Source/Runtime/Engine/Private/VT/VirtualTextureDataBuilder.cpp
+++ b/Engine/Source/Runtime/Engine/Private/VT/VirtualTextureDataBuilder.cpp
@@ -837,7 +837,7 @@ static const FName RemovePrefixFromName(FName const& InName, FName& OutPrefix)
 	// Then we detect a non-platform prefix (such as codec name)
 	// and split the result into  explicit FORMAT and PREFIX parts.
 
-	for (FName PlatformName : FDataDrivenPlatformInfoRegistry::GetSortedPlatformNames())
+	for (FName PlatformName : FDataDrivenPlatformInfoRegistry::GetSortedPlatformNames(EPlatformInfoType::AllPlatformInfos))
 	{
 		FString PlatformTextureFormatPrefix = PlatformName.ToString();
 		PlatformTextureFormatPrefix += TEXT('_');
