Add command-line option to disable debug output device which is often useful to avoid significant slowdown due to long stalls in OutputDebugString when VS is attached

#rnx

diff --git a/Engine/Source/Runtime/Core/Private/GenericPlatform/GenericPlatformOutputDevices.cpp b/Engine/Source/Runtime/Core/Private/GenericPlatform/GenericPlatformOutputDevices.cpp
--- a/Engine/Source/Runtime/Core/Private/GenericPlatform/GenericPlatformOutputDevices.cpp
+++ b/Engine/Source/Runtime/Core/Private/GenericPlatform/GenericPlatformOutputDevices.cpp
@@ -48,7 +48,7 @@ void FGenericPlatformOutputDevices::SetupOutputDevices()
 #if USE_DEBUG_LOGGING
 	// If the platform has a separate debug output channel (e.g. OutputDebugString) then add an output device
 	// unless logging is turned off
-	if (FPlatformMisc::HasSeparateChannelForDebugOutput())
+	if (FPlatformMisc::HasSeparateChannelForDebugOutput() && !FParse::Param(FCommandLine::Get(), TEXT("NODEBUGOUTPUT")))
 	{
 		GLog->AddOutputDevice(new FOutputDeviceDebug());
 	}
