Fixed TraceStats crashing when the stat name is a wide char FName.

diff --git a/Engine/Source/Runtime/Core/Private/Stats/Stats2.cpp b/Engine/Source/Runtime/Core/Private/Stats/Stats2.cpp
--- a/Engine/Source/Runtime/Core/Private/Stats/Stats2.cpp
+++ b/Engine/Source/Runtime/Core/Private/Stats/Stats2.cpp
@@ -617,9 +617,17 @@ public:
 #if STATSTRACE_ENABLED
 		if (!bCycleStat && (InStatType == EStatDataType::ST_int64 || InStatType == EStatDataType::ST_double))
 		{
-			ANSICHAR NameBuffer[1024];
-			StatShortName.GetPlainANSIString(NameBuffer);
-			FStatsTrace::DeclareStat(Stat, NameBuffer, *StatDescription, InGroup, InStatType == EStatDataType::ST_double, MemoryRegion != FPlatformMemory::MCR_Invalid, bShouldClearEveryFrame);
+			if (!StatShortName.GetDisplayNameEntry()->IsWide())
+			{
+				ANSICHAR NameBuffer[1024];
+				StatShortName.GetPlainANSIString(NameBuffer);
+				FStatsTrace::DeclareStat(Stat, NameBuffer, *StatDescription, InGroup, InStatType == EStatDataType::ST_double, MemoryRegion != FPlatformMemory::MCR_Invalid, bShouldClearEveryFrame);
+			}
+			else
+			{
+				FString StatShortNameString = StatShortName.GetPlainNameString();
+				FStatsTrace::DeclareStat(Stat, TCHAR_TO_ANSI(*StatShortNameString), *StatDescription, InGroup, InStatType == EStatDataType::ST_double, MemoryRegion != FPlatformMemory::MCR_Invalid, bShouldClearEveryFrame);
+			}
 		}
 #endif
 		return TStatId(Result);
