Fixed ThreadSingleton returning a stale pointer if using TryGet() after its destructor has been run.

This fixes FramePro malloc profiler using an invalid FrameProProfilerContext TLS object when profiling memory deallocation of another ThreadSingleton.

#tests Squads games on WSL, manually starting/stopping FramePro profiling.

diff --git a/Engine/Source/Runtime/Core/Private/Tests/HAL/ThreadTest.cpp b/Engine/Source/Runtime/Core/Private/Tests/HAL/ThreadTest.cpp
--- a/Engine/Source/Runtime/Core/Private/Tests/HAL/ThreadTest.cpp
+++ b/Engine/Source/Runtime/Core/Private/Tests/HAL/ThreadTest.cpp
@@ -9,6 +9,7 @@
 #include "HAL/Thread.h"
 #include "HAL/ThreadSafeBool.h"
 #include "HAL/ThreadSafeCounter.h"
+#include "HAL/ThreadSingleton.h"
 #include "HAL/Event.h"
 #include "Containers/Queue.h"
 #include "Containers/StringConv.h"
@@ -93,6 +94,58 @@ namespace
 		UE_LOG(LogTemp, Log, TEXT("%s completed"), StringCast<TCHAR>(__FUNCTION__).Get());
 	}
 
+	void TestThreadSingleton(FThreadTest& This)
+	{
+		{	// check that ThreadSingleton instances in different threads are isolated from each other
+			class FThreadSingletonTest : public TThreadSingleton<FThreadSingletonTest>
+			{
+			public:
+				void SetTestField(int NewValue) { TestField = NewValue; }
+				int GetTestField() const { return TestField; }
+			private:
+				int TestField{ 0 };
+			};
+			FThreadSingletonTest::Get().SetTestField(1);
+			FThread Thread;
+			TAtomic<bool> bDefaultValuePass{ false };
+			Thread = FThread(TEXT("Test.Thread.TestThreadSingleton"),
+				[&bDefaultValuePass]()
+				{
+					bDefaultValuePass = FThreadSingletonTest::TryGet() == nullptr;
+					bDefaultValuePass = bDefaultValuePass && (FThreadSingletonTest::Get().GetTestField() == 0);
+					FThreadSingletonTest::Get().SetTestField(2);
+				});
+			Thread.Join();
+			This.TestTrue(TEXT("Thread singleton should be uninitialized by default"), bDefaultValuePass);
+			This.TestTrue(TEXT("Thread singletons in different threads should be isolated"), FThreadSingletonTest::Get().GetTestField() == 1);
+		}
+		{	// check that ThreadSingleton entries don't point to invalid memory after cleanup
+			class FThreadSingletonFirst : public TThreadSingleton<FThreadSingletonFirst>
+			{
+			};
+			class FThreadSingletonSecond : public TThreadSingleton<FThreadSingletonSecond>
+			{
+			public:
+				virtual ~FThreadSingletonSecond()
+				{
+					// By the time we reach this destructor, the first singleton's destructor should have been executed already.
+					check(FThreadSingletonFirst::TryGet() == nullptr);
+				}
+			};
+			FThread Thread;
+			Thread = FThread(TEXT("Test.Thread.TestThreadSingleton"),
+				[]()
+				{
+					FThreadSingletonFirst::Get();
+					FThreadSingletonSecond::Get();
+				});
+			Thread.Join();
+			This.TestTrue(TEXT("Thread singletons should be uninitialized by default"), FThreadSingletonFirst::TryGet() == nullptr);
+			This.TestTrue(TEXT("Thread singletons should be uninitialized by default"), FThreadSingletonSecond::TryGet() == nullptr);
+		}
+		UE_LOG(LogTemp, Log, TEXT("%s completed"), StringCast<TCHAR>(__FUNCTION__).Get());
+	}
+
 	void TestMovability(FThreadTest& This)
 	{
 		{	// move constructor with default-constructed thread
@@ -210,6 +263,7 @@ bool FThreadTest::RunTest(const FString& Parameters)
 	TestMovability(*this);
 
 	TestTypicalUseCase(*this);
+	TestThreadSingleton(*this);
 
 	return true;
 }
diff --git a/Engine/Source/Runtime/Core/Public/HAL/ThreadSingleton.h b/Engine/Source/Runtime/Core/Public/HAL/ThreadSingleton.h
--- a/Engine/Source/Runtime/Core/Public/HAL/ThreadSingleton.h
+++ b/Engine/Source/Runtime/Core/Public/HAL/ThreadSingleton.h
@@ -62,6 +62,13 @@ protected:
 		: ThreadId(FPlatformTLS::GetCurrentThreadId())
 	{}
 
+	virtual ~TThreadSingleton()
+	{
+		// Clean the dangling pointer from the TLS.
+		check(GetTlsSlot() != 0xFFFFFFFF);
+		FPlatformTLS::SetTlsValue(GetTlsSlot(), nullptr);
+	}
+
 	/**
 	 * @return a new instance of the thread singleton.
 	 */
