CompressedBuffer: Fixed the tests to avoid a crash depending on execution order

#rnx

diff --git a/Engine/Source/Runtime/Core/Tests/Compression/CompressedBufferTest.cpp b/Engine/Source/Runtime/Core/Tests/Compression/CompressedBufferTest.cpp
--- a/Engine/Source/Runtime/Core/Tests/Compression/CompressedBufferTest.cpp
+++ b/Engine/Source/Runtime/Core/Tests/Compression/CompressedBufferTest.cpp
@@ -78,10 +78,14 @@ TEST_CASE("Core::Compression::CompressedBuffer::Compress", "[Core][Compression][
 		CHECK(CompressionLevel == ECompressedBufferCompressionLevel::VeryFast);
 		CHECK(FMath::IsPowerOfTwo(BlockSize));
 	}
+
+	CleanupCommandLine();
 }
 
 TEST_CASE("Core::Compression::CompressedBuffer::Decompress", "[Core][Compression][Smoke]")
 {
+	InitCommandLine();
+
 	const auto GenerateData = [](int32 N) -> TArray<uint64>
 	{
 		TArray<uint64> Data;
