Overlap IoStore reads with decrypt/decompress.

diff --git a/Engine/Source/Runtime/Core/Private/IO/IoStore.cpp b/Engine/Source/Runtime/Core/Private/IO/IoStore.cpp
--- a/Engine/Source/Runtime/Core/Private/IO/IoStore.cpp
+++ b/Engine/Source/Runtime/Core/Private/IO/IoStore.cpp
@@ -2202,35 +2202,76 @@ public:
 		const FIoStoreTocResource& TocResource = Toc.GetTocResource();
 		const uint64 CompressionBlockSize = TocResource.Header.CompressionBlockSize;
 		FIoBuffer UncompressedBuffer(ResolvedSize);
-		TArray<uint8> CompressedBuffer;
-		CompressedBuffer.SetNumUninitialized(static_cast<int32>(CompressionBlockSize));
+		if (ResolvedSize == 0)
+		{
+			return UncompressedBuffer;
+		}
+
+		// From here on we are reading / decompressing at least one block.
+
+		// We try to overlap the IO for the next block with the decrypt/decompress for the current
+		// block, which requires two IO buffers:
+		TArray<uint8> CompressedBuffers[2];
+
 		int32 FirstBlockIndex = int32(ResolvedOffset / CompressionBlockSize);
 		int32 LastBlockIndex = int32((Align(ResolvedOffset + ResolvedSize, CompressionBlockSize) - 1) / CompressionBlockSize);
+
+		// Lambda to kick off a read with a sufficient output buffer.
+		auto LaunchBlockRead = [&TocResource, this](int32 BlockIndex, TArray<uint8>& DestinationBuffer)
+		{
+			const uint64 CompressionBlockSize = TocResource.Header.CompressionBlockSize;
+			const FIoStoreTocCompressedBlockEntry& CompressionBlock = TocResource.CompressionBlocks[BlockIndex];
+
+			// CompressionBlockSize is technically the _uncompresseed_ block size, however it's a good
+			// size to use for reuse as block compression can vary wildly and we want to be able to
+			// read blocks that happen to be uncompressed.
+			uint32 SizeForDecrypt = Align(CompressionBlock.GetCompressedSize(), FAES::AESBlockSize);
+			uint32 CompressedBufferSizeNeeded = FMath::Max(uint32(CompressionBlockSize), SizeForDecrypt);
+
+			if (uint32(DestinationBuffer.Num()) < CompressedBufferSizeNeeded)
+			{
+				DestinationBuffer.SetNumUninitialized(CompressedBufferSizeNeeded);
+			}
+
+			int32 PartitionIndex = int32(CompressionBlock.GetOffset() / TocResource.Header.PartitionSize);
+			int64 PartitionOffset = int64(CompressionBlock.GetOffset() % TocResource.Header.PartitionSize);
+			return ContainerFileHandles[PartitionIndex]->ReadRequest(PartitionOffset, SizeForDecrypt, AIOP_Normal, nullptr, DestinationBuffer.GetData());
+		};
+
+
+		// Kick off the first async read
+		TUniquePtr<IAsyncReadRequest> NextReadRequest;
+		uint8 NextReadBufferIndex = 0;
+		NextReadRequest.Reset(LaunchBlockRead(FirstBlockIndex, CompressedBuffers[NextReadBufferIndex]));
+
 		uint64 UncompressedDestinationOffset = 0;
 		uint64 OffsetInBlock = ResolvedOffset % CompressionBlockSize;
 		uint64 RemainingSize = ResolvedSize;
 		TArray<uint8> TempBuffer;
 		for (int32 BlockIndex = FirstBlockIndex; BlockIndex <= LastBlockIndex; ++BlockIndex)
 		{
-			const FIoStoreTocCompressedBlockEntry& CompressionBlock = TocResource.CompressionBlocks[BlockIndex];
-			uint32 RawSize = Align(CompressionBlock.GetCompressedSize(), FAES::AESBlockSize);
-			if (uint32(CompressedBuffer.Num()) < RawSize)
+			// Kick off the next block's IO if there is one
+			TUniquePtr<IAsyncReadRequest> ReadRequest(MoveTemp(NextReadRequest));
+			uint8 OurBufferIndex = NextReadBufferIndex;			
+			if (BlockIndex + 1 <= LastBlockIndex)
 			{
-				CompressedBuffer.SetNumUninitialized(RawSize);
+				NextReadBufferIndex = NextReadBufferIndex ^ 1;
+				NextReadRequest.Reset(LaunchBlockRead(BlockIndex + 1, CompressedBuffers[NextReadBufferIndex]));
 			}
-			
-			int32 PartitionIndex = int32(CompressionBlock.GetOffset() / TocResource.Header.PartitionSize);
-			int64 PartitionOffset = int64(CompressionBlock.GetOffset() % TocResource.Header.PartitionSize);
-			TUniquePtr<IAsyncReadRequest> ReadRequest(ContainerFileHandles[PartitionIndex]->ReadRequest(PartitionOffset, RawSize, AIOP_Normal, nullptr, CompressedBuffer.GetData()));
+
+			// Now, wait for _our_ block's IO
 			{
 				TRACE_CPUPROFILER_EVENT_SCOPE(WaitForIo);
 				ReadRequest->WaitCompletion();
 			}
 
+			const FIoStoreTocCompressedBlockEntry& CompressionBlock = TocResource.CompressionBlocks[BlockIndex];
+			uint32 RawSize = Align(CompressionBlock.GetCompressedSize(), FAES::AESBlockSize);
+
 			if (EnumHasAnyFlags(TocResource.Header.ContainerFlags, EIoContainerFlags::Encrypted))
 			{
 				TRACE_CPUPROFILER_EVENT_SCOPE(Decrypt);
-				FAES::DecryptData(CompressedBuffer.GetData(), RawSize, DecryptionKey);
+				FAES::DecryptData(CompressedBuffers[OurBufferIndex].GetData(), RawSize, DecryptionKey);
 			}
 
 			FName CompressionMethod = TocResource.CompressionMethods[CompressionBlock.GetCompressionMethodIndex()];
@@ -2239,22 +2280,24 @@ public:
 			if (CompressionMethod.IsNone())
 			{
 				uint64 CopySize = FMath::Min<uint64>(UncompressedSize - OffsetInBlock, RemainingSize);
-				FMemory::Memcpy(UncompressedDestination, CompressedBuffer.GetData() + OffsetInBlock, CopySize);
+				FMemory::Memcpy(UncompressedDestination, CompressedBuffers[OurBufferIndex].GetData() + OffsetInBlock, CopySize);
 			}
 			else
 			{
 				bool bUncompressed;
 				if (OffsetInBlock || RemainingSize < UncompressedSize)
 				{
+					// If this block is larger than the amount of data actually requested, decompress to a temp
+					// buffer and then copy out. Should never happen when reading the entire file.
 					TempBuffer.SetNumUninitialized(UncompressedSize);
-					bUncompressed = FCompression::UncompressMemory(CompressionMethod, TempBuffer.GetData(), UncompressedSize, CompressedBuffer.GetData(), CompressionBlock.GetCompressedSize());
+					bUncompressed = FCompression::UncompressMemory(CompressionMethod, TempBuffer.GetData(), UncompressedSize, CompressedBuffers[OurBufferIndex].GetData(), CompressionBlock.GetCompressedSize());
 					uint64 CopySize = FMath::Min<uint64>(UncompressedSize - OffsetInBlock, RemainingSize);
 					FMemory::Memcpy(UncompressedDestination, TempBuffer.GetData() + OffsetInBlock, CopySize);
 				}
 				else
 				{
 					check(UncompressedDestination + UncompressedSize <= UncompressedBuffer.Data() + UncompressedBuffer.DataSize());
-					bUncompressed = FCompression::UncompressMemory(CompressionMethod, UncompressedDestination, UncompressedSize, CompressedBuffer.GetData(), CompressionBlock.GetCompressedSize());
+					bUncompressed = FCompression::UncompressMemory(CompressionMethod, UncompressedDestination, UncompressedSize, CompressedBuffers[OurBufferIndex].GetData(), CompressionBlock.GetCompressedSize());
 				}
 				if (!bUncompressed)
 				{
