Add support for Xcode 13.3

- Always ignore -Wunused-but-set-variable (on Apple platforms).
- Ignore -Wunordered-compare-function-pointers on third-party includes (across all clang platforms).

#rnx

diff --git a/Engine/Plugins/Media/ElectraPlayer/Source/ElectraPlayerRuntime/Private/Runtime/Core/MediaFastDelegate.h b/Engine/Plugins/Media/ElectraPlayer/Source/ElectraPlayerRuntime/Private/Runtime/Core/MediaFastDelegate.h
--- a/Engine/Plugins/Media/ElectraPlayer/Source/ElectraPlayerRuntime/Private/Runtime/Core/MediaFastDelegate.h
+++ b/Engine/Plugins/Media/ElectraPlayer/Source/ElectraPlayerRuntime/Private/Runtime/Core/MediaFastDelegate.h
@@ -9,6 +9,7 @@
 
 namespace Electra
 {
+THIRD_PARTY_INCLUDES_START
 #if defined(_MSC_VER)
 #pragma warning(push)
 #pragma warning(disable : 4191)
@@ -20,6 +21,7 @@ namespace Electra
 #if defined(_MSC_VER)
 #pragma warning( pop )
 #endif
+THIRD_PARTY_INCLUDES_END
 
 
 
diff --git a/Engine/Source/Runtime/Core/Public/Android/AndroidPlatformCompilerPreSetup.h b/Engine/Source/Runtime/Core/Public/Android/AndroidPlatformCompilerPreSetup.h
--- a/Engine/Source/Runtime/Core/Public/Android/AndroidPlatformCompilerPreSetup.h
+++ b/Engine/Source/Runtime/Core/Public/Android/AndroidPlatformCompilerPreSetup.h
@@ -13,7 +13,8 @@
 		PRAGMA_DISABLE_MISSING_VIRTUAL_DESTRUCTOR_WARNINGS \
 		PRAGMA_DISABLE_DEPRECATION_WARNINGS \
 		PRAGMA_DISABLE_OVERLOADED_VIRTUAL_WARNINGS \
-		PRAGMA_DISABLE_UNINITIALIZED_CONST_REFERENCE_WARNINGS
+		PRAGMA_DISABLE_UNINITIALIZED_CONST_REFERENCE_WARNINGS \
+		PRAGMA_DISABLE_ORDERED_COMPARE_FUNCTION_POINTERS
 #endif
 
 #ifndef THIRD_PARTY_INCLUDES_END
@@ -24,7 +25,8 @@
 		PRAGMA_ENABLE_MISSING_VIRTUAL_DESTRUCTOR_WARNINGS \
 		PRAGMA_ENABLE_UNDEFINED_IDENTIFIER_WARNINGS \
 		PRAGMA_ENABLE_SHADOW_VARIABLE_WARNINGS \
-		PRAGMA_ENABLE_REORDER_WARNINGS
+		PRAGMA_ENABLE_REORDER_WARNINGS \
+		PRAGMA_ENABLE_ORDERED_COMPARE_FUNCTION_POINTERS
 #endif
 
 // Make certain warnings always be warnings, even despite -Werror.
diff --git a/Engine/Source/Runtime/Core/Public/Apple/ApplePlatformCompilerPreSetup.h b/Engine/Source/Runtime/Core/Public/Apple/ApplePlatformCompilerPreSetup.h
--- a/Engine/Source/Runtime/Core/Public/Apple/ApplePlatformCompilerPreSetup.h
+++ b/Engine/Source/Runtime/Core/Public/Apple/ApplePlatformCompilerPreSetup.h
@@ -13,7 +13,8 @@
 		PRAGMA_DISABLE_MISSING_VIRTUAL_DESTRUCTOR_WARNINGS \
 		PRAGMA_DISABLE_DEPRECATION_WARNINGS \
 		PRAGMA_DISABLE_OVERLOADED_VIRTUAL_WARNINGS \
-		PRAGMA_DISABLE_UNINITIALIZED_CONST_REFERENCE_WARNINGS
+		PRAGMA_DISABLE_UNINITIALIZED_CONST_REFERENCE_WARNINGS \
+		PRAGMA_DISABLE_ORDERED_COMPARE_FUNCTION_POINTERS
 #endif
 
 #ifndef THIRD_PARTY_INCLUDES_END
@@ -24,7 +25,8 @@
 		PRAGMA_ENABLE_MISSING_VIRTUAL_DESTRUCTOR_WARNINGS \
 		PRAGMA_ENABLE_UNDEFINED_IDENTIFIER_WARNINGS \
 		PRAGMA_ENABLE_SHADOW_VARIABLE_WARNINGS \
-		PRAGMA_ENABLE_REORDER_WARNINGS
+		PRAGMA_ENABLE_REORDER_WARNINGS \
+		PRAGMA_ENABLE_ORDERED_COMPARE_FUNCTION_POINTERS
 #endif
 
 // Make certain warnings always be warnings, even despite -Werror.
@@ -57,6 +59,10 @@
 #pragma clang diagnostic ignored "-Wunknown-pragmas"
 #pragma clang diagnostic ignored "-Wundefined-var-template"
 
+#if __has_warning("-Wunused-but-set-variable")
+#pragma clang diagnostic ignored "-Wunused-but-set-variable"
+#endif
+
 // Apple LLVM 9.1.0 (Xcode 9.3)
 #if (__clang_major__ > 9) || (__clang_major__ == 9 && __clang_minor__ >= 1)
 #pragma clang diagnostic ignored "-Wunused-lambda-capture"
diff --git a/Engine/Source/Runtime/Core/Public/Clang/ClangPlatformCompilerPreSetup.h b/Engine/Source/Runtime/Core/Public/Clang/ClangPlatformCompilerPreSetup.h
--- a/Engine/Source/Runtime/Core/Public/Clang/ClangPlatformCompilerPreSetup.h
+++ b/Engine/Source/Runtime/Core/Public/Clang/ClangPlatformCompilerPreSetup.h
@@ -123,7 +123,24 @@
 #ifndef PRAGMA_RESTORE_UNSAFE_TYPECAST_WARNINGS
 	#define PRAGMA_RESTORE_UNSAFE_TYPECAST_WARNINGS \
 	_Pragma("clang diagnostic pop")
-#endif // PRAGMA_ENABLE_UNSAFE_TYPECAST_WARNINGS
+#endif // PRAGMA_RESTORE_UNSAFE_TYPECAST_WARNINGS
+
+#if __has_warning("-Wordered-compare-function-pointers")
+#define DISABLE_ORDERED_COMPARE_FUNCTION_POINTERS _Pragma("clang diagnostic ignored \"-Wordered-compare-function-pointers\"")
+#else
+#define DISABLE_ORDERED_COMPARE_FUNCTION_POINTERS
+#endif
+
+#ifndef PRAGMA_DISABLE_ORDERED_COMPARE_FUNCTION_POINTERS
+	#define PRAGMA_DISABLE_ORDERED_COMPARE_FUNCTION_POINTERS \
+		_Pragma("clang diagnostic push") \
+		DISABLE_ORDERED_COMPARE_FUNCTION_POINTERS
+#endif // PRAGMA_DISABLE_ORDERED_COMPARE_FUNCTION_POINTERS
+
+#ifndef PRAGMA_ENABLE_ORDERED_COMPARE_FUNCTION_POINTERS
+	#define PRAGMA_ENABLE_ORDERED_COMPARE_FUNCTION_POINTERS \
+	_Pragma("clang diagnostic pop")
+#endif // PRAGMA_ENABLE_ORDERED_COMPARE_FUNCTION_POINTERS
 
 #ifndef PRAGMA_DISABLE_UNDEFINED_IDENTIFIER_WARNINGS
 	#define PRAGMA_DISABLE_UNDEFINED_IDENTIFIER_WARNINGS \
diff --git a/Engine/Source/Runtime/Core/Public/Unix/UnixPlatformCompilerPreSetup.h b/Engine/Source/Runtime/Core/Public/Unix/UnixPlatformCompilerPreSetup.h
--- a/Engine/Source/Runtime/Core/Public/Unix/UnixPlatformCompilerPreSetup.h
+++ b/Engine/Source/Runtime/Core/Public/Unix/UnixPlatformCompilerPreSetup.h
@@ -43,7 +43,8 @@
 		PRAGMA_DISABLE_DEPRECATION_WARNINGS \
 		PRAGMA_DISABLE_OVERLOADED_VIRTUAL_WARNINGS \
 		PRAGMA_DISABLE_MACRO_REDEFINED_WARNINGS \
-		PRAGMA_DISABLE_UNINITIALIZED_CONST_REFERENCE_WARNINGS
+		PRAGMA_DISABLE_UNINITIALIZED_CONST_REFERENCE_WARNINGS \
+		PRAGMA_DISABLE_ORDERED_COMPARE_FUNCTION_POINTERS
 #endif
 
 #ifndef THIRD_PARTY_INCLUDES_END
@@ -57,7 +58,8 @@
 		PRAGMA_ENABLE_UNDEFINED_IDENTIFIER_WARNINGS \
 		PRAGMA_ENABLE_NULL_DEREFERENCE_WARNINGS \
 		PRAGMA_ENABLE_SHADOW_VARIABLE_WARNINGS \
-		PRAGMA_ENABLE_REORDER_WARNINGS
+		PRAGMA_ENABLE_REORDER_WARNINGS \
+		PRAGMA_ENABLE_ORDERED_COMPARE_FUNCTION_POINTERS
 #endif
 
 // Make certain warnings always be warnings, even despite -Werror.
diff --git a/Engine/Source/Runtime/Core/Public/Windows/WindowsPlatformCompilerPreSetup.h b/Engine/Source/Runtime/Core/Public/Windows/WindowsPlatformCompilerPreSetup.h
--- a/Engine/Source/Runtime/Core/Public/Windows/WindowsPlatformCompilerPreSetup.h
+++ b/Engine/Source/Runtime/Core/Public/Windows/WindowsPlatformCompilerPreSetup.h
@@ -35,7 +35,8 @@
 			PRAGMA_DISABLE_DEPRECATION_WARNINGS \
 			PRAGMA_DISABLE_OVERLOADED_VIRTUAL_WARNINGS \
 			PRAGMA_DISABLE_MISSING_BRACES_WARNINGS \
-	        PRAGMA_DISABLE_UNINITIALIZED_CONST_REFERENCE_WARNINGS
+	        PRAGMA_DISABLE_UNINITIALIZED_CONST_REFERENCE_WARNINGS \
+			PRAGMA_DISABLE_ORDERED_COMPARE_FUNCTION_POINTERS
 	#endif // THIRD_PARTY_INCLUDES_START
 
 	#ifndef THIRD_PARTY_INCLUDES_END
@@ -46,7 +47,8 @@
 			PRAGMA_ENABLE_DEPRECATION_WARNINGS \
 			PRAGMA_ENABLE_UNDEFINED_IDENTIFIER_WARNINGS \
 			PRAGMA_ENABLE_SHADOW_VARIABLE_WARNINGS \
-			PRAGMA_ENABLE_REORDER_WARNINGS
+			PRAGMA_ENABLE_REORDER_WARNINGS \
+			PRAGMA_ENABLE_ORDERED_COMPARE_FUNCTION_POINTERS
 	#endif // THIRD_PARTY_INCLUDES_END
 #else
 	#ifndef DISABLE_DEPRECATION
