Add base class for FAsyncTask to make it easier to interact with while keeping its templated type an implementation detail

#rnx

diff --git a/Engine/Source/Runtime/Core/Public/Async/AsyncWork.h b/Engine/Source/Runtime/Core/Public/Async/AsyncWork.h
--- a/Engine/Source/Runtime/Core/Public/Async/AsyncWork.h
+++ b/Engine/Source/Runtime/Core/Public/Async/AsyncWork.h
@@ -213,24 +213,24 @@ public:
 		delete Task;
 	}
 **/
-template<typename TTask>
-class FAsyncTask
+
+class FAsyncTaskBase
 	: private IQueuedWork
 {
-	/** User job embedded in this task */ 
-	TTask Task;
 	/** Thread safe counter that indicates WORK completion, no necessarily finalization of the job */
 	FThreadSafeCounter	WorkNotFinishedCounter;
 	/** If we aren't doing the work synchronously, this will hold the completion event */
-	FEvent*				DoneEvent;
+	FEvent*				DoneEvent = nullptr;
 	/** Pool we are queued into, maintained by the calling thread */
-	FQueuedThreadPool*	QueuedPool;
+	FQueuedThreadPool*	QueuedPool = nullptr;
 	/** Current priority */
-	EQueuedWorkPriority Priority;
+	EQueuedWorkPriority Priority = EQueuedWorkPriority::Normal;
 	/** Current flags */
-	EQueuedWorkFlags Flags;
+	EQueuedWorkFlags Flags = EQueuedWorkFlags::None;
 	/** Approximation of the peak memory (in bytes) this task could require during it's execution. */
 	int64 RequiredMemory = -1;
+	/** StatId used for FScopeCycleCounter */
+	TStatId StatId;
 	/** optional LLM tag */
 	LLM(const UE::LLMPrivate::FTagData* InheritedLLMTag);
 	/** Memory trace tag */
@@ -246,7 +246,7 @@ class FAsyncTask
 		DoneEvent = nullptr;
 	}
 
-	EQueuedWorkFlags GetQueuedWorkFlags() const override
+	EQueuedWorkFlags GetQueuedWorkFlags() const final
 	{
 		return Flags;
 	}
@@ -256,7 +256,7 @@ class FAsyncTask
 	**/
 	void Start(bool bForceSynchronous, FQueuedThreadPool* InQueuedPool, EQueuedWorkPriority InQueuedWorkPriority, EQueuedWorkFlags InQueuedWorkFlags, int64 InRequiredMemory)
 	{
-		FScopeCycleCounter Scope( Task.GetStatId(), true );
+		FScopeCycleCounter Scope(StatId, true);
 		DECLARE_SCOPE_CYCLE_COUNTER( TEXT( "FAsyncTask::Start" ), STAT_FAsyncTask_Start, STATGROUP_ThreadPoolAsyncTasks );
 		LLM(InheritedLLMTag = FLowLevelMemTracker::bIsDisabled ? nullptr : FLowLevelMemTracker::Get().GetActiveTagData(ELLMTracker::Default));
 #if UE_MEMORY_TAGS_TRACE_ENABLED
@@ -295,12 +295,12 @@ class FAsyncTask
 	* Tells the user job to do the work, sometimes called synchronously, sometimes from the thread pool. Calls the event tracker.
 	**/
 	void DoWork()
-	{	
+	{
 		LLM_SCOPE(InheritedLLMTag);
 		UE_MEMSCOPE(InheritedTraceTag);
-		FScopeCycleCounter Scope(Task.GetStatId(), true); 
+		FScopeCycleCounter Scope(StatId, true);
 
-		Task.DoWork();		
+		DoTaskWork();
 		check(WorkNotFinishedCounter.GetValue() == 1);
 		WorkNotFinishedCounter.Decrement();
 	}
@@ -313,7 +313,7 @@ class FAsyncTask
 		check(QueuedPool);
 		if (DoneEvent)
 		{
-			FScopeCycleCounter Scope( Task.GetStatId(), true );
+			FScopeCycleCounter Scope(StatId, true);
 			DECLARE_SCOPE_CYCLE_COUNTER( TEXT( "FAsyncTask::FinishThreadedWork" ), STAT_FAsyncTask_FinishThreadedWork, STATGROUP_ThreadPoolAsyncTasks );		
 			DoneEvent->Trigger();
 		}
@@ -322,7 +322,7 @@ class FAsyncTask
 	/** 
 	* Performs the work, this is only called from a pool thread.
 	**/
-	virtual void DoThreadedWork() override
+	void DoThreadedWork() final
 	{
 		DoWork();
 		FinishThreadedWork();
@@ -332,11 +332,10 @@ class FAsyncTask
 	 * Always called from the thread pool. Called if the task is removed from queue before it has started which might happen at exit.
 	 * If the user job can abandon, we do that, otherwise we force the work to be done now (doing nothing would not be safe).
 	 */
-	virtual void Abandon(void) override
+	void Abandon() final
 	{
-		if (Task.CanAbandon())
+		if (TryAbandonTask())
 		{
-			Task.Abandon();
 			check(WorkNotFinishedCounter.GetValue() == 1);
 			WorkNotFinishedCounter.Decrement();
 		}
@@ -347,16 +346,7 @@ class FAsyncTask
 		FinishThreadedWork();
 	}
 
-	/** 
-	* Internal call to assert that we are idle
-	**/
-	void CheckIdle() const
-	{
-		check(WorkNotFinishedCounter.GetValue() == 0);
-		check(!QueuedPool);
-	}
-
-	/** 
+	/**
 	* Internal call to synchronize completion between threads, never called from a pool thread
 	* @param bIsLatencySensitive specifies if waiting for the task should return as soon as possible even if this delays other tasks
 	**/
@@ -367,8 +357,8 @@ class FAsyncTask
 		FPlatformMisc::MemoryBarrier();
 		if (QueuedPool)
 		{
-			FScopeCycleCounter Scope( Task.GetStatId() );
-			DECLARE_SCOPE_CYCLE_COUNTER( TEXT( "FAsyncTask::SyncCompletion" ), STAT_FAsyncTask_SyncCompletion, STATGROUP_ThreadPoolAsyncTasks );
+			FScopeCycleCounter Scope(StatId);
+			DECLARE_SCOPE_CYCLE_COUNTER(TEXT("FAsyncTask::SyncCompletion"), STAT_FAsyncTask_SyncCompletion, STATGROUP_ThreadPoolAsyncTasks);
 
 			if (LowLevelTasks::FScheduler::Get().IsWorkerThread() && !bIsLatencySensitive)
 			{
@@ -382,63 +372,42 @@ class FAsyncTask
 		CheckIdle();
 	}
 
-	/** 
-	* Internal call to initialize internal variables
-	**/
-	void Init()
+protected:
+	void Init(TStatId InStatId)
 	{
-		DoneEvent = 0;
-		QueuedPool = 0;
+		StatId = InStatId;
 	}
 
-public:
-	FAsyncTask()
-		: Task()
+	/** 
+	* Internal call to assert that we are idle
+	**/
+	void CheckIdle() const
 	{
-		// This constructor shouldn't be necessary as the forwarding constructor should handle it, but
-		// we are getting VC internal compiler errors on CIS when creating arrays of FAsyncTask.
-
-		Init();
+		check(WorkNotFinishedCounter.GetValue() == 0);
+		check(!QueuedPool);
 	}
 
-	/** Forwarding constructor. */
-	template <typename Arg0Type, typename... ArgTypes>
-	FAsyncTask(Arg0Type&& Arg0, ArgTypes&&... Args)
-		: Task(Forward<Arg0Type>(Arg0), Forward<ArgTypes>(Args)...)
-	{
-		Init();
-	}
+	/** Perform task's work */
+	virtual void DoTaskWork() = 0;
+
+	/** 
+	* Abandon task if possible, returns true on success, false otherwise.
+	**/
+	virtual bool TryAbandonTask() = 0;
 
+public:
 	/** Destructor, not legal when a task is in process */
-	~FAsyncTask()
+	virtual ~FAsyncTaskBase()
 	{
 		// destroying an unfinished task is a bug
 		CheckIdle();
 		DestroyEvent();
 	}
 
-	/* Retrieve embedded user job, not legal to call while a job is in process
-	* @return reference to embedded user job 
-	**/
-	TTask &GetTask()
-	{
-		CheckIdle();  // can't modify a job without it being completed first
-		return Task;
-	}
-
-	/* Retrieve embedded user job, not legal to call while a job is in process
-	* @return reference to embedded user job 
-	**/
-	const TTask &GetTask() const
-	{
-		CheckIdle();  // could be safe, but I won't allow it anyway because the data could be changed while it is being read
-		return Task;
-	}
-
 	/**
 	 * Returns an approximation of the peak memory (in bytes) this task could require during it's execution.
 	 **/
-	int64 GetRequiredMemory() const override
+	int64 GetRequiredMemory() const final
 	{
 		return RequiredMemory;
 	}
@@ -545,7 +514,7 @@ public:
 		FPlatformMisc::MemoryBarrier();
 		if (QueuedPool)
 		{
-			FScopeCycleCounter Scope(Task.GetStatId());
+			FScopeCycleCounter Scope(StatId);
 			DECLARE_SCOPE_CYCLE_COUNTER(TEXT("FAsyncTask::SyncCompletion"), STAT_FAsyncTask_SyncCompletion, STATGROUP_ThreadPoolAsyncTasks);
 
 			uint32 Ms = uint32(TimeLimitSeconds * 1000.0f) + 1;
@@ -598,12 +567,76 @@ public:
 		return WorkNotFinishedCounter.GetValue() == 0 && QueuedPool == 0;
 	}
 
+	bool SetPriority(EQueuedWorkPriority QueuedWorkPriority)
+	{
+		return Reschedule(QueuedPool, QueuedWorkPriority);
+	}
+
 	EQueuedWorkPriority GetPriority() const
 	{
 		return Priority;
 	}
 };
 
+template<typename TTask>
+class FAsyncTask
+	: public FAsyncTaskBase
+{
+	/** User job embedded in this task */ 
+	TTask Task;
+
+public:
+	FAsyncTask()
+		: Task()
+	{
+		// Cache the StatId to remain backward compatible with TTask that declare GetStatId as non-const.
+		Init(Task.GetStatId());
+	}
+
+	/** Forwarding constructor. */
+	template <typename Arg0Type, typename... ArgTypes>
+	FAsyncTask(Arg0Type&& Arg0, ArgTypes&&... Args)
+		: Task(Forward<Arg0Type>(Arg0), Forward<ArgTypes>(Args)...)
+	{
+		// Cache the StatId to remain backward compatible with TTask that declare GetStatId as non-const.
+		Init(Task.GetStatId());
+	}
+
+	/* Retrieve embedded user job, not legal to call while a job is in process
+	* @return reference to embedded user job
+	**/
+	TTask& GetTask()
+	{
+		CheckIdle();  // can't modify a job without it being completed first
+		return Task;
+	}
+
+	/* Retrieve embedded user job, not legal to call while a job is in process
+	* @return reference to embedded user job
+	**/
+	const TTask& GetTask() const
+	{
+		CheckIdle();  // could be safe, but I won't allow it anyway because the data could be changed while it is being read
+		return Task;
+	}
+
+	bool TryAbandonTask() final
+	{
+		if (Task.CanAbandon())
+		{
+			Task.Abandon();
+			return true;
+		}
+
+		return false;
+	}
+
+	void DoTaskWork() final
+	{
+		Task.DoWork();
+	}
+};
+
 /**
  * Stub class to use a base class for tasks that cannot be abandoned
  */
