Silence static analys warning

#rnx

diff --git a/Engine/Source/Runtime/Core/Private/Serialization/MemoryImage.cpp b/Engine/Source/Runtime/Core/Private/Serialization/MemoryImage.cpp
--- a/Engine/Source/Runtime/Core/Private/Serialization/MemoryImage.cpp
+++ b/Engine/Source/Runtime/Core/Private/Serialization/MemoryImage.cpp
@@ -1682,7 +1682,7 @@ FMemoryImageSection* FMemoryImageSection::WritePointer(const FTypeLayoutDesc& St
 	FFrozenMemoryImagePtr FrozenPtr;
 	FrozenPtr.Packed = 0; // dummy value
 	// If derived type matches the static type, store INDEX_NONE to indicate the static type may be used
-	checkf(static_cast<uint64>(TypeDependencyIndex + 1) <= (1ULL << FFrozenMemoryImagePtr::TypeIndexBits), TEXT(""));
+	checkf(static_cast<uint64>(TypeDependencyIndex) + 1 <= (1ULL << FFrozenMemoryImagePtr::TypeIndexBits), TEXT(""));
 	FrozenPtr.SetTypeIndex(bStaticTypeMatchesDerived ? INDEX_NONE : TypeDependencyIndex);
 
 	SectionPointer.Offset = WriteBytes(FrozenPtr);
