Logging: Added a dedicated master thread for logging

diff --git a/Engine/Source/Runtime/Core/Private/Misc/OutputDeviceRedirector.cpp b/Engine/Source/Runtime/Core/Private/Misc/OutputDeviceRedirector.cpp
--- a/Engine/Source/Runtime/Core/Private/Misc/OutputDeviceRedirector.cpp
+++ b/Engine/Source/Runtime/Core/Private/Misc/OutputDeviceRedirector.cpp
@@ -4,10 +4,12 @@
 
 #include "Containers/DepletableMpscQueue.h"
 #include "Experimental/ConcurrentLinearAllocator.h"
+#include "HAL/Event.h"
 #include "HAL/PlatformProcess.h"
 #include "HAL/PlatformTime.h"
 #include "HAL/PlatformTLS.h"
-#include "Misc/ScopeLock.h"
+#include "HAL/Thread.h"
+#include "Misc/App.h"
 #include "Misc/ScopeRWLock.h"
 #include "ProfilingDebugging/CpuProfilerTrace.h"
 #include <atomic>
@@ -106,11 +108,12 @@ static constexpr uint64 CalculateRedirectorCacheLinePadding(const uint64 Size)
 
 struct FOutputDeviceRedirectorState
 {
+	/** A custom lock to guard access to both buffered and unbuffered output devices. */
 	FRWLock OutputDevicesLock;
-	std::atomic<uint32> OutputDevicesLockState;
+	std::atomic<uint32> OutputDevicesLockState{0};
 	uint8 OutputDevicesLockPadding[CalculateRedirectorCacheLinePadding(sizeof(OutputDevicesLock) + sizeof(OutputDevicesLockState))]{};
 
-	/** A FIFO of lines logged by non-master threads. */
+	/** A queue of lines logged by non-master threads. */
 	TDepletableMpscQueue<FOutputDeviceLine, FOutputDeviceLinearAllocator> BufferedLines;
 	uint8 BufferedLinesPadding[CalculateRedirectorCacheLinePadding(sizeof(BufferedLines))]{};
 
@@ -120,24 +123,49 @@ struct FOutputDeviceRedirectorState
 	/** Array of output devices to redirect to from the calling thread. */
 	TArray<FOutputDevice*> UnbufferedOutputDevices;
 
-	/** The ID of the master thread. Logging from other threads will be buffered for processing by the master thread. */
-	uint32 MasterThreadID = FPlatformTLS::GetCurrentThreadId();
+	/** A queue of lines logged before the editor added its output device. */
+	TArray<FBufferedLine> BacklogLines;
+	FRWLock BacklogLock;
 
-	/** Objects used for synchronization via a scoped lock. */
-	FCriticalSection SynchronizationObject;
+	/** An optional dedicated master thread for logging to buffered output devices. */
+	FThread Thread;
 
-	/** A FIFO backlog of messages logged before the editor had a chance to intercept them. */
-	TArray<FBufferedLine> BacklogLines;
+	/** A lock to synchronize access to the thread. */
+	FRWLock ThreadLock;
+
+	/** An event to wake the dedicated master thread to process buffered lines. */
+	FEventRef ThreadWakeEvent{EEventMode::AutoReset};
+
+	/** A queue of events to trigger when the dedicated master thread is idle. */
+	TDepletableMpscQueue<FEvent*, FOutputDeviceLinearAllocator> ThreadIdleEvents;
+
+	/** The ID of the master thread. Logging from other threads will be buffered for processing by the master thread. */
+	std::atomic<uint32> MasterThreadId{FPlatformTLS::GetCurrentThreadId()};
+
+	/** Whether to keep the thread running. */
+	std::atomic<bool> bThreadActive{false};
 
-	/** Whether backlogging is enabled. */
+	/** Whether the backlog is enabled. */
 	bool bEnableBacklog = false;
 
-	bool IsInMasterThread() const
+	bool IsMasterThread(const uint32 ThreadId) const
 	{
-		return MasterThreadID == FPlatformTLS::GetCurrentThreadId();
+		return ThreadId == MasterThreadId.load(std::memory_order_relaxed);
 	}
+
+	void ThreadLoop();
+
+	void FlushBufferedLines();
 };
 
+/**
+ * A scoped lock for readers of the OutputDevices arrays.
+ *
+ * The read lock:
+ * - Must be locked to read the OutputDevices arrays.
+ * - Must be locked to write to unbuffered output devices.
+ * - Must not be entered when the thread holds a write or master lock.
+ */
 struct FOutputDevicesReadScopeLock
 {
 	FORCEINLINE explicit FOutputDevicesReadScopeLock(FOutputDeviceRedirectorState& InState)
@@ -170,6 +198,13 @@ struct FOutputDevicesReadScopeLock
 	FOutputDeviceRedirectorState& State;
 };
 
+/**
+ * A scoped lock for writers of the OutputDevices arrays.
+ *
+ * The write lock has the same access as the master lock, and:
+ * - Must be locked to add or remove output devices.
+ * - Must not be entered when the thread holds a read, write, or master lock.
+ */
 struct FOutputDevicesWriteScopeLock
 {
 	FORCEINLINE explicit FOutputDevicesWriteScopeLock(FOutputDeviceRedirectorState& InState)
@@ -204,6 +239,75 @@ struct FOutputDevicesWriteScopeLock
 	FOutputDeviceRedirectorState& State;
 };
 
+/**
+ * A scoped lock for readers of the OutputDevices arrays that need to access master thread state.
+ *
+ * The master lock has the same access as the read lock, and:
+ * - Must be locked to write to buffered output devices.
+ * - Must be locked while calling FlushBufferedLines().
+ * - Must not be entered when the thread holds a write or master lock.
+ * - May be locked when the thread holds a read lock.
+ */
+struct FOutputDevicesMasterScopeLock
+{
+	FORCEINLINE explicit FOutputDevicesMasterScopeLock(FOutputDeviceRedirectorState& InState)
+		: State(InState)
+	{
+		State.OutputDevicesLock.WriteLock();
+	}
+
+	FORCEINLINE ~FOutputDevicesMasterScopeLock()
+	{
+		State.OutputDevicesLock.WriteUnlock();
+	}
+
+	FOutputDeviceRedirectorState& State;
+};
+
+void FOutputDeviceRedirectorState::ThreadLoop()
+{
+	const uint32 ThreadId = FPlatformTLS::GetCurrentThreadId();
+	{
+		FOutputDevicesMasterScopeLock Lock(*this);
+		MasterThreadId.store(ThreadId, std::memory_order_relaxed);
+	}
+
+	while (IsMasterThread(ThreadId))
+	{
+		ThreadWakeEvent->Wait();
+		do
+		{
+			FOutputDevicesMasterScopeLock Lock(*this);
+			FlushBufferedLines();
+		}
+		while (!BufferedLines.IsEmpty());
+		ThreadIdleEvents.Deplete([](FEvent* Event) { Event->Trigger(); });
+	}
+}
+
+void FOutputDeviceRedirectorState::FlushBufferedLines()
+{
+	if (BufferedLines.IsEmpty())
+	{
+		return;
+	}
+
+	TRACE_CPUPROFILER_EVENT_SCOPE(FOutputDeviceRedirector::FlushBufferedLines);
+
+	BufferedLines.Deplete([this](UE::Private::FOutputDeviceLine&& Line)
+	{
+		const double Time = Line.Time;
+		const TCHAR* const Data = Line.Data;
+		const FName Category = Line.Category;
+		const ELogVerbosity::Type Verbosity = Line.Verbosity;
+
+		for (FOutputDevice* OutputDevice : BufferedOutputDevices)
+		{
+			OutputDevice->Serialize(Data, Verbosity, Category, Time);
+		}
+	});
+}
+
 } // UE::Private
 
 FOutputDeviceRedirector::FOutputDeviceRedirector()
@@ -249,61 +353,34 @@ bool FOutputDeviceRedirector::IsRedirectingTo(FOutputDevice* OutputDevice)
 	return State->BufferedOutputDevices.Contains(OutputDevice) || State->UnbufferedOutputDevices.Contains(OutputDevice);
 }
 
-void FOutputDeviceRedirector::FlushBufferedLines(TConstArrayView<FOutputDevice*> BufferedDevices, bool bUseAllDevices)
-{	
-	if (State->BufferedLines.IsEmpty())
-	{
-		return;
-	}
-
-	TRACE_CPUPROFILER_EVENT_SCOPE(FOutputDeviceRedirector::FlushBufferedLines);
-
-	State->BufferedLines.Deplete([BufferedDevices, bUseAllDevices](UE::Private::FOutputDeviceLine&& Line)
+void FOutputDeviceRedirector::FlushThreadedLogs(EOutputDeviceRedirectorFlushOptions Options)
+{
+	if (FReadScopeLock Lock(State->ThreadLock); State->bThreadActive.load(std::memory_order_relaxed))
 	{
-		const double Time = Line.Time;
-		const TCHAR* const Data = Line.Data;
-		const FName Category = Line.Category;
-		const ELogVerbosity::Type Verbosity = Line.Verbosity;
-
-		for (FOutputDevice* OutputDevice : BufferedDevices)
+		if (!EnumHasAnyFlags(Options, EOutputDeviceRedirectorFlushOptions::Async))
 		{
-			if (bUseAllDevices || OutputDevice->CanBeUsedOnAnyThread())
+			FEventRef IdleEvent(EEventMode::ManualReset);
+			if (State->ThreadIdleEvents.EnqueueAndReturnWasEmpty(IdleEvent.Get()))
 			{
-				OutputDevice->Serialize(Data, Verbosity, Category, Time);
+				State->ThreadWakeEvent->Trigger();
 			}
+			IdleEvent->Wait();
 		}
-	});
-}
+		return;
+	}
 
-void FOutputDeviceRedirector::FlushThreadedLogs()
-{
-	check(State->IsInMasterThread());
-	UE::Private::FOutputDevicesReadScopeLock Lock(*State);
-	FlushBufferedLines(State->BufferedOutputDevices, /*bUseAllDevices*/ true);
+	UE::Private::FOutputDevicesMasterScopeLock Lock(*State);
+	State->FlushBufferedLines();
 }
 
 void FOutputDeviceRedirector::PanicFlushThreadedLogs()
 {
-	UE::Private::FOutputDevicesReadScopeLock Lock(*State);
-	FlushBufferedLines(State->BufferedOutputDevices, /*bUseAllDevices*/ State->IsInMasterThread());
-
-	for (FOutputDevice* OutputDevice : State->BufferedOutputDevices)
-	{
-		if (OutputDevice->CanBeUsedOnAnyThread())
-		{
-			OutputDevice->Flush();
-		}
-	}
-
-	for (FOutputDevice* OutputDevice : State->UnbufferedOutputDevices)
-	{
-		OutputDevice->Flush();
-	}
+	Flush();
 }
 
 void FOutputDeviceRedirector::SerializeBacklog(FOutputDevice* OutputDevice)
 {
-	FScopeLock ScopeLock(&State->SynchronizationObject);
+	FReadScopeLock ScopeLock(State->BacklogLock);
 	for (const FBufferedLine& BacklogLine : State->BacklogLines)
 	{
 		OutputDevice->Serialize(BacklogLine.Data, BacklogLine.Verbosity, BacklogLine.Category, BacklogLine.Time);
@@ -312,7 +389,7 @@ void FOutputDeviceRedirector::SerializeBacklog(FOutputDevice* OutputDevice)
 
 void FOutputDeviceRedirector::EnableBacklog(bool bEnable)
 {
-	FScopeLock ScopeLock(&State->SynchronizationObject);
+	FWriteScopeLock ScopeLock(State->BacklogLock);
 	State->bEnableBacklog = bEnable;
 	if (!bEnable)
 	{
@@ -322,11 +399,40 @@ void FOutputDeviceRedirector::EnableBacklog(bool bEnable)
 
 void FOutputDeviceRedirector::SetCurrentThreadAsMasterThread()
 {
-	UE::Private::FOutputDevicesReadScopeLock Lock(*State);
-	FlushBufferedLines(State->BufferedOutputDevices, /*bUseAllDevices*/ State->IsInMasterThread());
+	const uint32 ThreadId = FPlatformTLS::GetCurrentThreadId();
+
+	if (UE::Private::FOutputDevicesMasterScopeLock Lock(*State); State->MasterThreadId.load(std::memory_order_relaxed) == ThreadId)
+	{
+		return;
+	}
+	else
+	{
+		State->MasterThreadId.store(ThreadId, std::memory_order_relaxed);
+		State->FlushBufferedLines();
+	}
 
-	FScopeLock ScopeLock(&State->SynchronizationObject);
-	State->MasterThreadID = FPlatformTLS::GetCurrentThreadId();
+	if (FWriteScopeLock Lock(State->ThreadLock); State->bThreadActive.load(std::memory_order_relaxed))
+	{
+		State->bThreadActive.store(false, std::memory_order_relaxed);
+		State->ThreadWakeEvent->Trigger();
+		State->Thread.Join();
+	}
+}
+
+bool FOutputDeviceRedirector::TryStartDedicatedMasterThread()
+{
+	if (!FApp::ShouldUseThreadingForPerformance())
+	{
+		return false;
+	}
+
+	if (FWriteScopeLock Lock(State->ThreadLock); !State->bThreadActive.load(std::memory_order_relaxed))
+	{
+		State->bThreadActive.store(true, std::memory_order_relaxed);
+		State->ThreadWakeEvent->Trigger();
+		State->Thread = FThread(TEXT("OutputDeviceRedirector"), [State = State.Get()] { State->ThreadLoop(); });
+	}
+	return true;
 }
 
 void FOutputDeviceRedirector::Serialize(const TCHAR* const Data, const ELogVerbosity::Type Verbosity, const FName& Category, const double Time)
@@ -357,17 +463,31 @@ void FOutputDeviceRedirector::Serialize(const TCHAR* const Data, const ELogVerbo
 
 	if (State->bEnableBacklog)
 	{
-		FScopeLock ScopeLock(&State->SynchronizationObject);
-		new(State->BacklogLines)FBufferedLine(Data, Category, Verbosity, RealTime);
+		FWriteScopeLock ScopeLock(State->BacklogLock);
+		State->BacklogLines.Emplace(Data, Category, Verbosity, RealTime);
 	}
 
-	if (!State->IsInMasterThread() || State->BufferedOutputDevices.IsEmpty())
+	const auto EnqueueLine = [this, Data, Category, Verbosity, RealTime]
 	{
-		State->BufferedLines.Enqueue(Data, Category, Verbosity, RealTime);
+		if (State->BufferedLines.EnqueueAndReturnWasEmpty(Data, Category, Verbosity, RealTime) &&
+			State->bThreadActive.load(std::memory_order_relaxed))
+		{
+			State->ThreadWakeEvent->Trigger();
+		}
+	};
+
+	const uint32 ThreadId = FPlatformTLS::GetCurrentThreadId();
+	if (!State->IsMasterThread(ThreadId) || State->BufferedOutputDevices.IsEmpty())
+	{
+		EnqueueLine();
+	}
+	else if (UE::Private::FOutputDevicesMasterScopeLock MasterLock(*State); !State->IsMasterThread(ThreadId))
+	{
+		EnqueueLine();
 	}
 	else
 	{
-		FlushBufferedLines(State->BufferedOutputDevices, /*bUseAllDevices*/ true);
+		State->FlushBufferedLines();
 
 		for (FOutputDevice* OutputDevice : State->BufferedOutputDevices)
 		{
@@ -393,16 +513,12 @@ void FOutputDeviceRedirector::RedirectLog(const FLazyName& Category, ELogVerbosi
 
 void FOutputDeviceRedirector::Flush()
 {
-	UE::Private::FOutputDevicesReadScopeLock Lock(*State);
+	UE::Private::FOutputDevicesMasterScopeLock Lock(*State);
+	State->FlushBufferedLines();
 
-	if (State->IsInMasterThread())
+	for (FOutputDevice* OutputDevice : State->BufferedOutputDevices)
 	{
-		FlushBufferedLines(State->BufferedOutputDevices, /*bUseAllDevices*/ true);
-
-		for (FOutputDevice* OutputDevice : State->BufferedOutputDevices)
-		{
-			OutputDevice->Flush();
-		}
+		OutputDevice->Flush();
 	}
 
 	for (FOutputDevice* OutputDevice : State->UnbufferedOutputDevices)
@@ -413,14 +529,14 @@ void FOutputDeviceRedirector::Flush()
 
 void FOutputDeviceRedirector::TearDown()
 {
-	FScopeLock SyncLock(&State->SynchronizationObject);
-	check(State->IsInMasterThread());
+	SetCurrentThreadAsMasterThread();
+
+	Flush();
 
 	TArray<FOutputDevice*> LocalBufferedDevices;
 	TArray<FOutputDevice*> LocalUnbufferedDevices;
 
 	{
-		// We need to lock the mutex here so that it gets unlocked after we empty the devices arrays
 		UE::Private::FOutputDevicesWriteScopeLock Lock(*State);
 		LocalBufferedDevices = MoveTemp(State->BufferedOutputDevices);
 		LocalUnbufferedDevices = MoveTemp(State->UnbufferedOutputDevices);
@@ -428,23 +544,20 @@ void FOutputDeviceRedirector::TearDown()
 		State->UnbufferedOutputDevices.Empty();
 	}
 
-	FlushBufferedLines(LocalBufferedDevices, /*bUseAllDevices*/ true);
-
 	for (FOutputDevice* OutputDevice : LocalBufferedDevices)
 	{
-		OutputDevice->Flush();
 		OutputDevice->TearDown();
 	}
 
 	for (FOutputDevice* OutputDevice : LocalUnbufferedDevices)
 	{
-		OutputDevice->Flush();
 		OutputDevice->TearDown();
 	}
 }
 
 bool FOutputDeviceRedirector::IsBacklogEnabled() const
 {
+	FReadScopeLock Lock(State->BacklogLock);
 	return State->bEnableBacklog;
 }
 
diff --git a/Engine/Source/Runtime/Core/Public/Misc/OutputDeviceRedirector.h b/Engine/Source/Runtime/Core/Public/Misc/OutputDeviceRedirector.h
--- a/Engine/Source/Runtime/Core/Public/Misc/OutputDeviceRedirector.h
+++ b/Engine/Source/Runtime/Core/Public/Misc/OutputDeviceRedirector.h
@@ -4,6 +4,7 @@
 
 #include "CoreTypes.h"
 #include "Containers/Array.h"
+#include "Misc/EnumClassFlags.h"
 #include "Misc/OutputDevice.h"
 #include "Templates/PimplPtr.h"
 #include "UObject/NameTypes.h"
@@ -47,6 +48,21 @@ struct CORE_API FBufferedLine
 	~FBufferedLine();
 };
 
+enum class EOutputDeviceRedirectorFlushOptions : uint32
+{
+	None = 0,
+
+	/**
+		* Flush asynchronously when possible.
+		*
+		* When this flag is set and there is a dedicated master thread, the flush function returns immediately.
+		* Otherwise, the flush function does not return until the requested type of flush is complete.
+		*/
+	Async = 1 << 0,
+};
+
+ENUM_CLASS_FLAGS(EOutputDeviceRedirectorFlushOptions);
+
 /**
  * Class used for output redirection to allow logs to show in multiple output devices.
  */
@@ -85,13 +101,9 @@ public:
 	bool IsRedirectingTo(FOutputDevice* OutputDevice);
 
 	/** Flushes lines buffered by secondary threads. */
-	void FlushThreadedLogs();
+	void FlushThreadedLogs(EOutputDeviceRedirectorFlushOptions Options = EOutputDeviceRedirectorFlushOptions::None);
 
-	/**
-	 * Flushes lines buffered by secondary threads.
-	 *
-	 * Only used if a background thread crashed and we needed to push the callstack into the log.
-	 */
+	/** Same as FlushThreadedLogs(). */
 	void PanicFlushThreadedLogs();
 
 	/**
@@ -110,6 +122,15 @@ public:
 	/** Sets the current thread to be the master thread redirects logs without buffering. */
 	void SetCurrentThreadAsMasterThread();
 
+	/**
+	 * Starts a dedicated master thread that redirects logs that require buffering.
+	 *
+	 * A thread will not be started for certain configurations or platforms, or when threading is disabled.
+	 *
+	 * @return true if a dedicated master thread is running, false otherwise.
+	 */
+	bool TryStartDedicatedMasterThread();
+
 	/**
 	 * Serializes the passed in data via all current output devices.
 	 *
@@ -146,12 +167,5 @@ public:
 	bool IsBacklogEnabled() const;
 
 private:
-	/**
-	 * The unsynchronized version of FlushThreadedLogs.
-	 * Assumes that the caller holds a lock on the output devices.
-	 * @param bUseAllDevices   If true, this method will use all buffered output devices.
-	 */
-	void FlushBufferedLines(TConstArrayView<FOutputDevice*> BufferedDevices, bool bUseAllDevices);
-
 	TPimplPtr<UE::Private::FOutputDeviceRedirectorState> State;
 };
diff --git a/Engine/Source/Runtime/CoreUObject/Private/Serialization/AsyncLoading2.cpp b/Engine/Source/Runtime/CoreUObject/Private/Serialization/AsyncLoading2.cpp
--- a/Engine/Source/Runtime/CoreUObject/Private/Serialization/AsyncLoading2.cpp
+++ b/Engine/Source/Runtime/CoreUObject/Private/Serialization/AsyncLoading2.cpp
@@ -6068,7 +6068,7 @@ void FAsyncLoadingThread2::FlushLoading(int32 RequestId)
 					// Flush logging when runing cook-on-the-fly and waiting for packages
 					if (IsRunningCookOnTheFly() && FPlatformTime::Seconds() - LogFlushTime > 1.0)
 					{
-						GLog->FlushThreadedLogs();
+						GLog->FlushThreadedLogs(EOutputDeviceRedirectorFlushOptions::Async);
 						LogFlushTime = FPlatformTime::Seconds();
 					}
 				}
diff --git a/Engine/Source/Runtime/Engine/Private/ShaderCompiler/ShaderCompiler.cpp b/Engine/Source/Runtime/Engine/Private/ShaderCompiler/ShaderCompiler.cpp
--- a/Engine/Source/Runtime/Engine/Private/ShaderCompiler/ShaderCompiler.cpp
+++ b/Engine/Source/Runtime/Engine/Private/ShaderCompiler/ShaderCompiler.cpp
@@ -4062,7 +4062,7 @@ void FShaderCompilingManager::BlockOnShaderMapCompletion(const TArray<int32>& Sh
 				if (++LogCounter > 50)
 				{
 					LogCounter = 0;
-					GLog->FlushThreadedLogs();
+					GLog->FlushThreadedLogs(EOutputDeviceRedirectorFlushOptions::Async);
 				}
 			}
 		} 
diff --git a/Engine/Source/Runtime/Launch/Private/LaunchEngineLoop.cpp b/Engine/Source/Runtime/Launch/Private/LaunchEngineLoop.cpp
--- a/Engine/Source/Runtime/Launch/Private/LaunchEngineLoop.cpp
+++ b/Engine/Source/Runtime/Launch/Private/LaunchEngineLoop.cpp
@@ -1503,12 +1503,11 @@ int32 FEngineLoop::PreInitPreStartupScreen(const TCHAR* CmdLine)
 	FDelayedAutoRegisterHelper::RunAndClearDelayedAutoRegisterDelegates(EDelayedRegisterRunPhase::StartOfEnginePreInit);
 	SCOPED_BOOT_TIMING("FEngineLoop::PreInitPreStartupScreen");
 
-	// The GLog singleton is lazy initialised and by default will assume that
-	// its "master thread" is the one it was created on. This lazy initialisation
-	// can happen during the dynamic init (e.g. static) of a DLL which modern
-	// Windows does on a worker thread thus makeing its master thread not this one.
-	// So we make it this one and GLog->TearDown() is happy.
-	if (GLog != nullptr)
+	// GLog is initialized lazily and its default master thread is the thread that initialized it.
+	// This lazy initialization can happen during initialization of a DLL, which Windows does on a
+	// worker thread, which makes that worker thread the master thread. Make this the master thread
+	// until initialization is far enough along to try to start a dedicated master thread.
+	if (GLog)
 	{
 		GLog->SetCurrentThreadAsMasterThread();
 	}
@@ -1658,6 +1657,12 @@ int32 FEngineLoop::PreInitPreStartupScreen(const TCHAR* CmdLine)
 	// as soon as we determine whether GIsEditor == false
 	GLog->EnableBacklog(true);
 
+	// Try to start the dedicated master thread now that the command line is available.
+	if (!FParse::Param(FCommandLine::Get(), TEXT("NoLogThread")))
+	{
+		GLog->TryStartDedicatedMasterThread();
+	}
+
 	// Initialize std out device as early as possible if requested in the command line
 #if PLATFORM_DESKTOP
 	// consoles don't typically have stdout, and FOutputDeviceDebug is responsible for echoing logs to the
@@ -5043,7 +5048,7 @@ void FEngineLoop::Tick()
 		// flush debug output which has been buffered by other threads
 		{
 			QUICK_SCOPE_CYCLE_COUNTER(STAT_FEngineLoop_FlushThreadedLogs); 
-			GLog->FlushThreadedLogs();
+			GLog->FlushThreadedLogs(EOutputDeviceRedirectorFlushOptions::Async);
 		}
 
 		// exit if frame limit is reached in benchmark mode, or if time limit is reached
diff --git a/Engine/Source/Runtime/RenderCore/Private/RenderingThread.cpp b/Engine/Source/Runtime/RenderCore/Private/RenderingThread.cpp
--- a/Engine/Source/Runtime/RenderCore/Private/RenderingThread.cpp
+++ b/Engine/Source/Runtime/RenderCore/Private/RenderingThread.cpp
@@ -938,7 +938,7 @@ void CheckRenderingThreadHealth()
 	{
 		if (!GIsCriticalError)
 		{
-			GLog->FlushThreadedLogs();
+			GLog->FlushThreadedLogs(EOutputDeviceRedirectorFlushOptions::Async);
 		}
 #if !(UE_BUILD_SHIPPING || UE_BUILD_TEST)
 		TGuardValue<TAtomic<bool>, bool> GuardMainThreadBlockedOnRenderThread(GMainThreadBlockedOnRenderThread,true);
