M - LowLevel Task extension for Expediting, and Revival of previously cancelled Tasks.
Where Revival tries to revert any Cancellation as if it never happened. If the canceled Task was already executed by the Scheduler the revival will fail
(the Task will not be relaunched if revival fails). And where Expediting tries to fast track a task if the result is needed immediately on the current Context.

diff --git a/Engine/Source/Runtime/Core/Private/Async/Fundamental/Scheduler.cpp b/Engine/Source/Runtime/Core/Private/Async/Fundamental/Scheduler.cpp
--- a/Engine/Source/Runtime/Core/Private/Async/Fundamental/Scheduler.cpp
+++ b/Engine/Source/Runtime/Core/Private/Async/Fundamental/Scheduler.cpp
@@ -304,9 +304,11 @@ namespace LowLevelTasks
 			FTask* Task = (Queue->*DequeueFunction)(bPermitBackgroundWork, bDisableThrottleStealing);
 			if (Task)
 			{	
-				if (bIsBusyWaiting && !Task->AllowBusyWaiting())
+				FTask::FInitData InitData = Task->GetInitData();
+				bool bAllowBusyWaiting = EnumHasAnyFlags(InitData.Flags, ETaskFlags::AllowBusyWaiting);
+				if (bIsBusyWaiting && !bAllowBusyWaiting)
 				{
-					QueueRegistry.Enqueue(Task, uint32(Task->GetPriority()));
+					QueueRegistry.Enqueue(Task, uint32(InitData.Priority));
 					continue;
 				}
 				if (OutOfWork.Stop())
diff --git a/Engine/Source/Runtime/Core/Private/Async/Fundamental/UnitTests.cpp b/Engine/Source/Runtime/Core/Private/Async/Fundamental/UnitTests.cpp
--- a/Engine/Source/Runtime/Core/Private/Async/Fundamental/UnitTests.cpp
+++ b/Engine/Source/Runtime/Core/Private/Async/Fundamental/UnitTests.cpp
@@ -491,9 +491,15 @@ namespace Tasks2Tests
 				{
 					TestValue = 42;
 				});
+				verify(Task.TryCancel(ECancellationFlags::PrelaunchCancellation));
+				verify(Task.TryRevive());
 				TryLaunch(Task);
 
 				bool WasCanceled = TestCancel && Task.TryCancel();
+				if(WasCanceled)
+				{
+					WasCanceled = !Task.TryRevive();
+				}	
 				BusyWaitForTask(Task);
 
 				if (WasCanceled)
@@ -531,6 +537,25 @@ namespace Tasks2Tests
 					verify(TestValue == 42);
 				}
 			}
+
+			//expedite
+			{
+				uint32 TestValue = 1337;
+
+				FTask Task;
+				Task.Init(TEXT("expedite Test"), [&TestValue]()
+				{
+					TestValue = 42;
+				});
+				verify(!Task.TryExpedite());
+				TryLaunch(Task);
+				Task.TryExpedite();
+
+				verify(!Task.TryCancel());
+				BusyWaitForTask(Task);
+
+				verify(TestValue == 42);
+			}
 		}
 
 		//example Awaitable Tasks
diff --git a/Engine/Source/Runtime/Core/Public/Async/Fundamental/Task.h b/Engine/Source/Runtime/Core/Public/Async/Fundamental/Task.h
--- a/Engine/Source/Runtime/Core/Public/Async/Fundamental/Task.h
+++ b/Engine/Source/Runtime/Core/Public/Async/Fundamental/Task.h
@@ -2,6 +2,7 @@
 
 #pragma once
 #include "Logging/LogMacros.h"
+#include "Misc/EnumClassFlags.h"
 #include "TaskDelegate.h"
 #include "HAL/Event.h"
 #include "CoreTypes.h"
@@ -13,7 +14,7 @@ namespace LowLevelTasks
 {
 	DECLARE_LOG_CATEGORY_EXTERN(LowLevelTasks, Log, All);
 
-	enum class ETaskPriority : int32
+	enum class ETaskPriority : int8
 	{
 		High,
 		Normal,
@@ -79,20 +80,74 @@ namespace LowLevelTasks
 		return false;
 	}
 
-	enum class ETaskState : int32
+	enum class ECancellationFlags : int8
 	{
-		CanceledAndCompleted,	//means the task is completed with execution of it's continuation but the runnable was cancelled
-		Completed,				//means the task is completed with execution or the default when we create a handle
-		Ready,					//means the Task is ready to be launched
-		Canceled,				//means the task was canceled and launched and therefore queued for execution by a worker (which already might be executing it's continuation)
-		CanceledAndReady,		//means the task was canceled and is ready to be launched (it still is required to be launched)
-		Scheduled,				//means the task is launched and therefore queued for execution by a worker
-		Running,				//means the task is executing it's runnable and continuation by a worker 
-		Count
+		None					= 0 << 0,
+		TryLaunchOnSuccess		= 1 << 0, // try to launch and the continuation immediately if it was not launched yet (requires PrelaunchCancellation to work)
+		PrelaunchCancellation	= 1 << 1, // allow cancellation before a task has been launched (this also allows the optimization of TryLaunchOnSuccess)
+		DefaultFlags			= TryLaunchOnSuccess | PrelaunchCancellation,
 	};
-	
-	//Generic implementation of a Deleter, it often comes up that one has to call a function to cleanup after a Task finished
-	//this can be done by capturing a TDeleter like this in the lambda of the Continuation: [Deleter(LowLevelTasks::TDeleter<Type, &Type::DeleteFunction>(value))](){}
+	ENUM_CLASS_FLAGS(ECancellationFlags)
+
+	enum class ETaskFlags : int8
+	{
+		AllowNothing		= 0 << 0,
+		AllowBusyWaiting	= 1 << 0,
+		AllowCancellation	= 1 << 1,
+		AllowEverything		= AllowBusyWaiting | AllowCancellation,
+		DefaultFlags		= AllowEverything,
+	};
+	ENUM_CLASS_FLAGS(ETaskFlags)
+
+	/*																	       
+	 * (I)nitThread:                                                        STORE(I)----------------------CAS(C)----------------------    
+	 * (C)ancelingThread:                                                      --->|         Ready        |<-->|   CanceledAndReady   |   
+	 *                                                                              ----------------------      ----------------------    
+	 *                                                                                        |OR(L)                      |OR(L)               
+	 *                                                                                        V                           V               
+	 * (L)aunchingThread:   --------------------------------------------------CAS(E)----------------------CAS(C)----------------------    
+	 * (C)ancelingThread:  |                      Running                     |<---|      Scheduled       |<-->|       Canceled       |   
+	 * (E)xpeditingThread:  --------------------------------------------------      ----------------------      ----------------------    
+	 *                                |OR(E)                      |OR(W)                      |OR(W)                      |OR(W)               
+	 *                                V                           V                           V                           V               
+	 * (W)orkerThread:      ---------------------- OR(E)----------------------      ----------------------      ---------------------- 	   
+	 * (E)xpeditingThread: |      Expedited       |<---|      Expediting      |    |       Running        |    |  CanceledAndRunning  |	   
+	 *                      ----------------------      ----------------------      ----------------------      ----------------------   
+	 *                                |OR(W,E)                                                |OR(W)                      |OR(W)		   
+	 *                                V                                                       V                           V			   
+	 * (W)orkerThread:      ----------------------                                  ----------------------      ----------------------    
+	 * (E)xpeditingThread: |ExpeditedAndCompleted |                                |       Completed      |    | CanceledAndCompleted |   
+	 *                      ----------------------                                  ----------------------      ---------------------- 
+	 */
+	enum class ETaskState : int8
+	{
+		ReadyState		=  0,
+		CanceledFlag	=  1 << 0,
+		ScheduledFlag	=  1 << 1,
+		RunningFlag		=  1 << 2,
+		ExpeditingFlag	=  1 << 3,
+		ExpeditedFlag	=  1 << 4,
+		CompletedFlag	=  1 << 5,										//the default state when we create a handle
+		Count			= (1 << 6) - 1,
+
+		Ready					= ReadyState,							//means the Task is ready to be launched
+		CanceledAndReady		= Ready | CanceledFlag,					//means the task was canceled and is ready to be launched (it still is required to be launched)
+		Scheduled				= Ready | ScheduledFlag, 				//means the task is launched and therefore queued for execution by a worker
+		Canceled				= CanceledAndReady | ScheduledFlag,		//means the task was canceled and launched and therefore queued for execution by a worker (which already might be executing it's continuation)
+		Running					= Scheduled | RunningFlag,				//means the task is executing it's runnable and continuation by a worker
+		CanceledAndRunning		= Canceled | RunningFlag,				//means the task is executing it's continuation  but the runnable was cancelled
+		Expediting				= Running | ExpeditingFlag, 			//means the task is expediting and the scheduler has released it's reference to the expediting thread before that was finished
+		Expedited				= Expediting | ExpeditedFlag, 			//means the task was expedited
+		Completed				= Running | CompletedFlag,				//means the task is completed with execution 
+		ExpeditedAndCompleted	= Expedited | CompletedFlag,			//means the task is completed with execution and the runnable was expedited
+		CanceledAndCompleted	= CanceledAndRunning | CompletedFlag,	//means the task is completed with execution of it's continuation but the runnable was cancelled
+	};
+	ENUM_CLASS_FLAGS(ETaskState)
+
+	/*
+	* Generic implementation of a Deleter, it often comes up that one has to call a function to cleanup after a Task finished
+	* this can be done by capturing a TDeleter like this in the lambda of the Continuation: [Deleter(LowLevelTasks::TDeleter<Type, &Type::DeleteFunction>(value))](){}
+	*/
 	template<typename Type, void (Type::*DeleteFunction)()>
 	class TDeleter
 	{
@@ -123,66 +178,117 @@ namespace LowLevelTasks
 		}
 	};
 
-	//this class is just here to hide some variables away
-	//because we don't want to become too close friends with the FScheduler
+	/*
+	* this class is just here to hide some variables away
+	* because we don't want to become too close friends with the FScheduler
+	*/
 	class FTask;
 	namespace Tasks_Impl
 	{
 	class FTaskBase
 	{
+		class FPackedDataAtomic;
+
 		friend class ::LowLevelTasks::FTask;
 		UE_NONCOPYABLE(FTaskBase); //means non movable
 
-		class FPackedData
+		union FPackedData
 		{
-			uintptr_t DebugName			: 57;
-			uintptr_t Priority			: 3;
-			uintptr_t State				: 3;
-			uintptr_t bAllowBusyWaiting : 1;
-			
-		public:
-			FPackedData() : FPackedData(nullptr,  ETaskPriority::Count, ETaskState::Completed, true)
+			uintptr_t PackedData;
+			struct
+			{
+				uintptr_t State				: 6;
+				uintptr_t DebugName			: 53;
+				uintptr_t Priority			: 3;
+				uintptr_t Flags				: 2;
+			};
+
+		private:
+			friend class FTaskBase::FPackedDataAtomic;
+			FPackedData(uintptr_t InPackedData) : PackedData(InPackedData)
+			{}
+
+			constexpr FPackedData() 				
+				: State((uintptr_t)ETaskState::CompletedFlag)
+				, DebugName(0ull)
+				, Priority((uintptr_t)ETaskPriority::Count)
+				, Flags((uintptr_t)ETaskFlags::DefaultFlags)
 			{
 				static_assert(!PLATFORM_32BITS, "32bit Platforms are not supported");
 				static_assert(uintptr_t(ETaskPriority::Count) <= (1ull << 3), "Not enough bits to store ETaskPriority");
-				static_assert(uintptr_t(ETaskState::Count) <= (1ull << 3), "Not enough bits to store ETaskState");
+				static_assert(uintptr_t(ETaskState::Count) <= (1ull << 6), "Not enough bits to store ETaskState");
+				static_assert(uintptr_t(ETaskFlags::AllowEverything) < (1ull << 2), "Not enough bits to store ETaskFlags");
 			}
 
-			FPackedData(const TCHAR* InDebugName, ETaskPriority InPriority, ETaskState InState, bool bInAllowBusyWaiting)
-				: DebugName(reinterpret_cast<uintptr_t>(InDebugName))
+		public:
+			FPackedData(const TCHAR* InDebugName, ETaskPriority InPriority, ETaskState InState, ETaskFlags InFlags)
+				: State((uintptr_t)InState)
+				, DebugName(reinterpret_cast<uintptr_t>(InDebugName))
 				, Priority((uintptr_t)InPriority)
-				, State((uintptr_t)InState)
-				, bAllowBusyWaiting(bInAllowBusyWaiting)
+				, Flags((uintptr_t)InFlags)
+				
 			{
-				checkSlow(reinterpret_cast<uintptr_t>(InDebugName) < (1ull << 57));
+				checkSlow(reinterpret_cast<uintptr_t>(InDebugName) < (1ull << 53));
 				checkSlow((uintptr_t)InPriority < (1ull << 3));
-				checkSlow((uintptr_t)InState < (1ull << 3));
+				checkSlow((uintptr_t)InState < (1ull << 6));
+				checkSlow((uintptr_t)InFlags < (1ull << 2));
 				static_assert(sizeof(FPackedData) == sizeof(uintptr_t), "Packed data needs to be pointer size");
 			}
 
-			FPackedData(FPackedData& Other, ETaskState State)
-				: FPackedData(Other.GetDebugName(), Other.GetPriority(), State, Other.AllowBusyWaiting())
+			FPackedData(const FPackedData& Other, ETaskState State)
+				: FPackedData(Other.GetDebugName(), Other.GetPriority(), State, Other.GetFlags())
 			{
 			}
 
-			const TCHAR* GetDebugName() const
+			inline const TCHAR* GetDebugName() const
 			{
 				return reinterpret_cast<const TCHAR*>(DebugName);
 			}
 
-			ETaskPriority GetPriority() const
+			inline ETaskPriority GetPriority() const
 			{
 				return ETaskPriority(Priority);
 			}
 
-			ETaskState GetState() const
+			inline ETaskState GetState() const
 			{
 				return ETaskState(State);
 			}
 
-			bool AllowBusyWaiting() const
+			inline ETaskFlags GetFlags() const
 			{
-				return bAllowBusyWaiting;
+				return ETaskFlags(Flags);
+			}
+		};
+
+		class FPackedDataAtomic
+		{
+			std::atomic<uintptr_t> PackedData { FPackedData().PackedData };
+
+		public:
+			ETaskState fetch_or(ETaskState State, std::memory_order Order)
+			{
+				return ETaskState(FPackedData(PackedData.fetch_or(uintptr_t(State), Order)).State);
+			}
+
+			bool compare_exchange_strong(FPackedData& Expected, FPackedData Desired, std::memory_order Success, std::memory_order Failure)
+			{
+				return PackedData.compare_exchange_strong(Expected.PackedData, Desired.PackedData, Success, Failure);
+			}
+
+			bool compare_exchange_strong(FPackedData& Expected, FPackedData Desired, std::memory_order Order)
+			{
+				return PackedData.compare_exchange_strong(Expected.PackedData, Desired.PackedData, Order);
+			}
+
+			FPackedData load(std::memory_order Order) const
+			{
+				return PackedData.load(Order);
+			}
+
+			void store(const FPackedData& Expected, std::memory_order Order)
+			{
+				PackedData.store(Expected.PackedData, Order);
 			}
 		};
 
@@ -190,18 +296,16 @@ namespace LowLevelTasks
 		using FTaskDelegate = TTaskDelegate<void(bool), LOWLEVEL_TASK_SIZE - sizeof(FPackedData) - sizeof(void*)>;
 		FTaskDelegate Runnable;
 		mutable void* UserData = nullptr;
-		std::atomic<FPackedData> PackedData { FPackedData() };
+		FPackedDataAtomic PackedData;
 
 	private:
 		FTaskBase() = default;
-
-	private:
-		inline bool IsCanceled() const { return PackedData.load(std::memory_order_relaxed).GetState() == ETaskState::Canceled; }
-		inline bool IsScheduled() const { return PackedData.load(std::memory_order_relaxed).GetState() == ETaskState::Scheduled; }
 	};
 	}
 
-	//minimal core task interface
+	/*
+	* minimal low level task interface
+	*/
 	class FTask final : private Tasks_Impl::FTaskBase
 	{		
 		friend class FScheduler;
@@ -209,30 +313,48 @@ namespace LowLevelTasks
 
 		static thread_local FTask* ActiveTask;
 
-	public: //Public Interface
-		//means the task is completed and this taskhandle can be recycled
+	public:
+
+		/*
+		* means the task is completed and this taskhandle can be recycled
+		*/
 		inline bool IsCompleted(std::memory_order MemoryOrder = std::memory_order_seq_cst) const
 		{
 			ETaskState State = PackedData.load(MemoryOrder).GetState();
-			return State == ETaskState::CanceledAndCompleted || State == ETaskState::Completed;
+			return EnumHasAnyFlags(State, ETaskState::CompletedFlag);
 		}
 		
-		//means the task was canceled but might still need to be launched 
-		inline bool WasCanceled() 		
+		/*
+		* means the task was canceled but might still need to be launched
+		*/
+		inline bool WasCanceled() const
 		{
 			ETaskState State = PackedData.load(std::memory_order_relaxed).GetState();
-			return State == ETaskState::CanceledAndReady || State == ETaskState::Canceled || State == ETaskState::CanceledAndCompleted;
+			return EnumHasAnyFlags(State, ETaskState::CanceledFlag);
+		}
+
+		/*
+		* means the task was expedited or that it already completed
+		*/
+		inline bool WasExpedited() const
+		{
+			ETaskState State = PackedData.load(std::memory_order_acquire).GetState();
+			return EnumHasAnyFlags(State, ETaskState::ExpeditedFlag | ETaskState::CompletedFlag);
 		}
 
-		//means the task is ready to be launched but might already been canceled 
+		/*
+		* means the task is ready to be launched but might already been canceled 
+		*/
 		inline bool IsReady() const 
 		{
 			ETaskState State = PackedData.load(std::memory_order_relaxed).GetState();
-			return State == ETaskState::Ready || State == ETaskState::CanceledAndReady; 
+			return !EnumHasAnyFlags(State, ~ETaskState::CanceledFlag); 
 		}
 
 #if PLATFORM_DESKTOP
-		//get the currently active task if any
+		/*
+		* get the currently active task if any
+		*/
 		CORE_API static const FTask* GetActiveTask();
 #else
 		FORCEINLINE static const FTask* GetActiveTask()
@@ -241,34 +363,52 @@ namespace LowLevelTasks
 		}
 #endif
 
-		//try to cancel the task if it has not been launched yet and ExecuteTaskOnSuccess is true the continuation will run immediately.
-		inline bool TryCancel(bool ExecuteTaskOnSuccess = true);
-
-		//try to execute the task if it has not been launched yet the task will execute immediately.
+		/*
+		* try to cancel the task if it has not been launched yet.
+		*/
+		inline bool TryCancel(ECancellationFlags CancellationFlags = ECancellationFlags::DefaultFlags);
+
+		/*
+		* try to revive a canceled task (as in reverting the cancellation as if it never happened).
+		* if it had been canceled and the scheduler has not run it yet it succeeds.
+		*/
+		inline bool TryRevive();
+
+		/*
+		* try to expedite the task if succeded it will run immediately 
+		* but it will not set the completed state until the scheduler has executed it, because the scheduler still holds a reference.
+		* to check for completion in the context of expediting use WasExpedited. The TaskHandle canot be reused until IsCompleted returns true.
+		*/
+		inline bool TryExpedite();
+
+		/*
+		* try to execute the task if it has not been launched yet the task will execute immediately.
+		*/
 		inline bool TryExecute();
 
 		template<typename TRunnable, typename TContinuation>
-		inline void Init(const TCHAR* InDebugName, ETaskPriority InPriority, TRunnable&& InRunnable, TContinuation&& InContinuation, bool bAllowBusyWaiting = true);
+		inline void Init(const TCHAR* InDebugName, ETaskPriority InPriority, TRunnable&& InRunnable, TContinuation&& InContinuation, ETaskFlags Flags = ETaskFlags::DefaultFlags);
 
 		template<typename TRunnable>
-		inline void Init(const TCHAR* InDebugName, ETaskPriority InPriority, TRunnable&& InRunnable, bool bAllowBusyWaiting = true);
+		inline void Init(const TCHAR* InDebugName, ETaskPriority InPriority, TRunnable&& InRunnable, ETaskFlags Flags = ETaskFlags::DefaultFlags);
 
 		template<typename TRunnable, typename TContinuation>
-		inline void Init(const TCHAR* InDebugName, TRunnable&& InRunnable, TContinuation&& InContinuation, bool bAllowBusyWaiting = true);
+		inline void Init(const TCHAR* InDebugName, TRunnable&& InRunnable, TContinuation&& InContinuation, ETaskFlags Flags = ETaskFlags::DefaultFlags);
 
 		template<typename TRunnable>
-		inline void Init(const TCHAR* InDebugName, TRunnable&& InRunnable, bool bAllowBusyWaiting = true);
+		inline void Init(const TCHAR* InDebugName, TRunnable&& InRunnable, ETaskFlags Flags = ETaskFlags::DefaultFlags);
 
 		inline const TCHAR* GetDebugName() const;
 		inline ETaskPriority GetPriority() const;
 		inline bool IsBackgroundTask() const;
 		inline bool AllowBusyWaiting() const;
+		inline bool AllowCancellation() const;
 
 		struct FInitData
 		{
 			const TCHAR* DebugName;
 			ETaskPriority Priority;
-			bool bAllowBusyWaiting;
+			ETaskFlags Flags;
 		};
 		inline FInitData GetInitData() const;
 
@@ -287,6 +427,8 @@ namespace LowLevelTasks
 
 		inline bool TryPrepareLaunch();
 		//after calling this function the task can be considered dead
+		template<bool bIsExpeditingThread>
+		inline void TryFinish();
 		inline void ExecuteTask();
 		inline void InheritParentData(ETaskPriority& Priority);
 	};
@@ -322,7 +464,7 @@ namespace LowLevelTasks
 	}
 
 	template<typename TRunnable, typename TContinuation>
-	inline void FTask::Init(const TCHAR* InDebugName, ETaskPriority InPriority, TRunnable&& InRunnable, TContinuation&& InContinuation, bool bAllowBusyWaiting)
+	inline void FTask::Init(const TCHAR* InDebugName, ETaskPriority InPriority, TRunnable&& InRunnable, TContinuation&& InContinuation, ETaskFlags Flags)
 	{
 		checkf(IsCompleted(), TEXT("State: %d"), PackedData.load(std::memory_order_relaxed).GetState());
 		checkSlow(!Runnable.IsSet());
@@ -335,11 +477,11 @@ namespace LowLevelTasks
 			LocalContinuation();
 		};
 		InheritParentData(InPriority);
-		PackedData.store(FPackedData(InDebugName, InPriority, ETaskState::Ready, bAllowBusyWaiting), std::memory_order_release);
+		PackedData.store(FPackedData(InDebugName, InPriority, ETaskState::Ready, Flags), std::memory_order_release);
 	}
 
 	template<typename TRunnable>
-	inline void FTask::Init(const TCHAR* InDebugName, ETaskPriority InPriority, TRunnable&& InRunnable, bool bAllowBusyWaiting)
+	inline void FTask::Init(const TCHAR* InDebugName, ETaskPriority InPriority, TRunnable&& InRunnable, ETaskFlags Flags)
 	{
 		checkf(IsCompleted(), TEXT("State: %d"), PackedData.load(std::memory_order_relaxed).GetState());
 		checkSlow(!Runnable.IsSet());
@@ -351,19 +493,19 @@ namespace LowLevelTasks
 			}
 		};
 		InheritParentData(InPriority);
-		PackedData.store(FPackedData(InDebugName, InPriority, ETaskState::Ready, bAllowBusyWaiting), std::memory_order_release);
+		PackedData.store(FPackedData(InDebugName, InPriority, ETaskState::Ready, Flags), std::memory_order_release);
 	}
 
 	template<typename TRunnable, typename TContinuation>
-	inline void FTask::Init(const TCHAR* InDebugName, TRunnable&& InRunnable, TContinuation&& InContinuation, bool bAllowBusyWaiting)
+	inline void FTask::Init(const TCHAR* InDebugName, TRunnable&& InRunnable, TContinuation&& InContinuation, ETaskFlags Flags)
 	{
-		Init(InDebugName, ETaskPriority::Default, Forward<TRunnable>(InRunnable), Forward<TContinuation>(InContinuation), bAllowBusyWaiting);
+		Init(InDebugName, ETaskPriority::Default, Forward<TRunnable>(InRunnable), Forward<TContinuation>(InContinuation), Flags);
 	}
 
 	template<typename TRunnable>
-	inline void FTask::Init(const TCHAR* InDebugName, TRunnable&& InRunnable, bool bAllowBusyWaiting)
+	inline void FTask::Init(const TCHAR* InDebugName, TRunnable&& InRunnable, ETaskFlags Flags)
 	{
-		Init(InDebugName, ETaskPriority::Default, Forward<TRunnable>(InRunnable), bAllowBusyWaiting);
+		Init(InDebugName, ETaskPriority::Default, Forward<TRunnable>(InRunnable), Flags);
 	}
 
 	inline FTask::~FTask()
@@ -373,24 +515,23 @@ namespace LowLevelTasks
 
 	inline bool FTask::TryPrepareLaunch()
 	{
-		FPackedData LocalPackedData = PackedData.load(std::memory_order_relaxed);
-		FPackedData ReadyState = FPackedData(LocalPackedData, ETaskState::Ready);
-		FPackedData CanceledAndReadyState = FPackedData(LocalPackedData, ETaskState::CanceledAndReady);
-		return PackedData.compare_exchange_strong(ReadyState, FPackedData(LocalPackedData, ETaskState::Scheduled), std::memory_order_acquire, std::memory_order_relaxed)
-			|| PackedData.compare_exchange_strong(CanceledAndReadyState, FPackedData(LocalPackedData, ETaskState::Canceled), std::memory_order_acquire, std::memory_order_relaxed);
+		return !EnumHasAnyFlags(PackedData.fetch_or(ETaskState::ScheduledFlag, std::memory_order_release), ETaskState::ScheduledFlag);
 	}
 
-	inline bool FTask::TryCancel(bool ExecuteTaskOnSuccess)
+	inline bool FTask::TryCancel(ECancellationFlags CancellationFlags)
 	{
+		bool bPrelaunchCancellation = EnumHasAnyFlags(CancellationFlags, ECancellationFlags::PrelaunchCancellation);
+		bool bTryLaunchOnSuccess	= EnumHasAllFlags(CancellationFlags, ECancellationFlags::PrelaunchCancellation | ECancellationFlags::TryLaunchOnSuccess);
+
 		FPackedData LocalPackedData = PackedData.load(std::memory_order_relaxed);
 		FPackedData ReadyState(LocalPackedData, ETaskState::Ready);
 		FPackedData ScheduledState(LocalPackedData, ETaskState::Scheduled);
-		//we can use memory_order_relaxed in this case because, when a task is canceled it does not launch the task
-		//to launch a canceled  task it has to go though TryPrepareLaunch which is doing the memory_order_acquire
-		bool WasCanceled = PackedData.compare_exchange_strong(ReadyState, FPackedData(LocalPackedData, ETaskState::CanceledAndReady), std::memory_order_relaxed)
-			|| PackedData.compare_exchange_strong(ScheduledState, FPackedData(LocalPackedData, ETaskState::Canceled), std::memory_order_relaxed);
+		//to launch a canceled  task it has to go though TryPrepareLaunch which is doing the memory_order_release
+		bool WasCanceled = EnumHasAnyFlags(LocalPackedData.GetFlags(), ETaskFlags::AllowCancellation)
+			&& ((bPrelaunchCancellation && PackedData.compare_exchange_strong(ReadyState, FPackedData(LocalPackedData, ETaskState::CanceledAndReady), std::memory_order_acquire))
+									    || PackedData.compare_exchange_strong(ScheduledState, FPackedData(LocalPackedData, ETaskState::Canceled), std::memory_order_acquire));
 
-		if(ExecuteTaskOnSuccess && WasCanceled && TryPrepareLaunch())
+		if(bTryLaunchOnSuccess && WasCanceled && TryPrepareLaunch())
 		{
 			ExecuteTask();
 			return true;
@@ -398,6 +539,21 @@ namespace LowLevelTasks
 		return WasCanceled;
 	}
 
+	inline bool FTask::TryRevive()
+	{
+		FPackedData LocalPackedData = PackedData.load(std::memory_order_relaxed);
+		checkSlow(EnumHasAnyFlags(LocalPackedData.GetState(), ETaskState::CanceledFlag));
+		if(EnumHasAnyFlags(LocalPackedData.GetState(), ETaskState::RunningFlag))
+		{
+			return false;
+		}
+
+		FPackedData CanceledReadyState(LocalPackedData, ETaskState::CanceledAndReady);
+		FPackedData CanceledState(LocalPackedData, ETaskState::Canceled);
+		return PackedData.compare_exchange_strong(CanceledReadyState, FPackedData(LocalPackedData, ETaskState::Ready), std::memory_order_release)
+			|| PackedData.compare_exchange_strong(CanceledState, FPackedData(LocalPackedData, ETaskState::Scheduled), std::memory_order_release);
+	}
+
 	inline bool FTask::TryExecute()
 	{
 		if(TryPrepareLaunch())
@@ -408,22 +564,57 @@ namespace LowLevelTasks
 		return false;
 	}
 
-	inline void FTask::ExecuteTask()
+	template<bool bIsExpeditingThread>
+	inline void FTask::TryFinish()
 	{
-		checkSlow(Runnable.IsSet());
-		FTaskDelegate LocalRunnable;
+		const ETaskState NextState = bIsExpeditingThread ? ETaskState::ExpeditedFlag | ETaskState::ExpeditingFlag : ETaskState::ExpeditingFlag;
+		ETaskState PreviousState = PackedData.fetch_or(NextState, std::memory_order_acq_rel);
+		if constexpr (bIsExpeditingThread)
+		{
+			checkSlow(PreviousState == ETaskState::Running || PreviousState == ETaskState::Expediting);
+		}
+		if(EnumHasAnyFlags(PreviousState, ETaskState::ExpeditingFlag))
+		{
+			FTaskDelegate LocalRunnable = MoveTemp(Runnable);
+			//do not access the task again after this call
+			//as by defitition the task can be considered dead
+			PreviousState = PackedData.fetch_or(ETaskState::CompletedFlag, std::memory_order_seq_cst);
+			checkSlow(PreviousState == ETaskState::Expedited);
+		}
+	}
 
-		checkSlow(IsScheduled() || IsCanceled());
+	inline bool FTask::TryExpedite()
+	{
 		FPackedData LocalPackedData = PackedData.load(std::memory_order_relaxed);
 		FPackedData ScheduledState(LocalPackedData, ETaskState::Scheduled);
-		const bool NotCanceled = PackedData.compare_exchange_strong(ScheduledState, FPackedData(LocalPackedData, ETaskState::Running), std::memory_order_relaxed);
-		
-		Runnable.CallAndMove(LocalRunnable, NotCanceled);
+		if(PackedData.compare_exchange_strong(ScheduledState, FPackedData(LocalPackedData, ETaskState::Running), std::memory_order_acquire))
+		{
+			Runnable(true);
+			TryFinish<true>();
+			return true;
+		}
+		return false;
+	}
 
-		checkSlow(!Runnable.IsSet());
-		//do not access the task again after this call
-		//as by defitition the task can be considered dead
-		PackedData.store(FPackedData(LocalPackedData, NotCanceled ? ETaskState::Completed : ETaskState::CanceledAndCompleted), std::memory_order_seq_cst);
+	inline void FTask::ExecuteTask()
+	{
+		ETaskState PreviousState = PackedData.fetch_or(ETaskState::RunningFlag, std::memory_order_acquire);
+		checkSlow(EnumHasAnyFlags(PreviousState, ETaskState::ScheduledFlag));
+
+		if(!EnumHasAnyFlags(PreviousState, ETaskState::RunningFlag)) //we are running or canceled
+		{
+			FTaskDelegate LocalRunnable;
+			Runnable.CallAndMove(LocalRunnable, !EnumHasAnyFlags(PreviousState, ETaskState::CanceledFlag));
+			//do not access the task again after this call
+			//as by defitition the task can be considered dead
+			PreviousState = PackedData.fetch_or(ETaskState::CompletedFlag, std::memory_order_seq_cst);
+			checkSlow(PreviousState == ETaskState::Running || PreviousState == ETaskState::CanceledAndRunning);
+		}
+		else // we are expedited
+		{
+			checkSlow(PreviousState == ETaskState::Running || PreviousState == ETaskState::Expediting || PreviousState == ETaskState::Expedited);
+			TryFinish<false>();
+		}
 	}
 
 	inline const TCHAR* FTask::GetDebugName() const
@@ -438,13 +629,18 @@ namespace LowLevelTasks
 
 	inline bool FTask::AllowBusyWaiting() const
 	{
-		return PackedData.load(std::memory_order_relaxed).AllowBusyWaiting();
+		return EnumHasAnyFlags(PackedData.load(std::memory_order_relaxed).GetFlags(), ETaskFlags::AllowBusyWaiting);
+	}
+
+	inline bool FTask::AllowCancellation() const
+	{
+		return EnumHasAnyFlags(PackedData.load(std::memory_order_relaxed).GetFlags(), ETaskFlags::AllowCancellation);
 	}
 
 	inline FTask::FInitData FTask::GetInitData() const
 	{
 		FPackedData LocalPackedData = PackedData.load(std::memory_order_relaxed);
-		return { LocalPackedData.GetDebugName(), LocalPackedData.GetPriority(), LocalPackedData.AllowBusyWaiting() };
+		return { LocalPackedData.GetDebugName(), LocalPackedData.GetPriority(), LocalPackedData.GetFlags() };
 	}
 
 	enum class ESleepState
@@ -455,8 +651,10 @@ namespace LowLevelTasks
 		Sleeping,
 	};
 
-	//the struct is naturally 64 bytes aligned, the extra alignment just 
-	//re-enforces this assumption and will error if it changes in the future
+	/*
+	* the struct is naturally 64 bytes aligned, the extra alignment just 
+	* re-enforces this assumption and will error if it changes in the future
+	*/
 	struct alignas(64) FSleepEvent
 	{
 		FEventRef SleepEvent;
diff --git a/Engine/Source/Runtime/Core/Public/Misc/QueuedThreadPoolWrapper.h b/Engine/Source/Runtime/Core/Public/Misc/QueuedThreadPoolWrapper.h
--- a/Engine/Source/Runtime/Core/Public/Misc/QueuedThreadPoolWrapper.h
+++ b/Engine/Source/Runtime/Core/Public/Misc/QueuedThreadPoolWrapper.h
@@ -329,7 +329,7 @@ private:
 
 		LowLevelTasks::ETaskPriority TaskPriorityMapper[int(EQueuedWorkPriority::Count)] = { LowLevelTasks::ETaskPriority::High, LowLevelTasks::ETaskPriority::High, LowLevelTasks::ETaskPriority::BackgroundHigh, LowLevelTasks::ETaskPriority::BackgroundNormal, LowLevelTasks::ETaskPriority::BackgroundLow, LowLevelTasks::ETaskPriority::BackgroundLow };
 		LowLevelTasks::ETaskPriority TaskPriority = TaskPriorityMapper[int(Priority)];
-		const bool bAllowBusyWaiting = (InQueuedWork->GetQueuedWorkFlags() & EQueuedWorkFlags::DoNotRunInsideBusyWait) == EQueuedWorkFlags::None;
+		LowLevelTasks::ETaskFlags Flags = (InQueuedWork->GetQueuedWorkFlags() & EQueuedWorkFlags::DoNotRunInsideBusyWait) == EQueuedWorkFlags::None ? LowLevelTasks::ETaskFlags::DefaultFlags : (LowLevelTasks::ETaskFlags::DefaultFlags & ~LowLevelTasks::ETaskFlags::AllowCancellation);
 
 		QueuedWorkInternalData->Task.Init(TEXT("FQueuedLowLevelThreadPoolTask"), TaskPriority, [InQueuedWork]
 		{
@@ -342,7 +342,7 @@ private:
 			bool bWakeUpWorker = false;
 			ScheduleTasks(bWakeUpWorker);
 		},
-		bAllowBusyWaiting
+		Flags
 		);
 
 		if (!bIsPaused)
