Fix very large hash bucket in backtrace by explicitly storing -1 for functions with complex unwinds rather than truncating.

diff --git a/Engine/Source/Runtime/Core/Private/ProfilingDebugging/Microsoft/WindowsCallstackTrace.cpp b/Engine/Source/Runtime/Core/Private/ProfilingDebugging/Microsoft/WindowsCallstackTrace.cpp
--- a/Engine/Source/Runtime/Core/Private/ProfilingDebugging/Microsoft/WindowsCallstackTrace.cpp
+++ b/Engine/Source/Runtime/Core/Private/ProfilingDebugging/Microsoft/WindowsCallstackTrace.cpp
@@ -648,8 +648,9 @@ uint32 FBacktracer::GetBacktraceId(void* AddressOfReturnAddress) const
 
 #if BACKTRACE_LOCK_FREE
 		{
+			int32 StoreBias = Function->RspBias < 0 ? -1 : Function->RspBias;
 			TGuardValue<bool> ReentranceGuard(bReentranceCheck, true);
-			FunctionLookups.Emplace(RetAddr, Function->RspBias);
+			FunctionLookups.Emplace(RetAddr, StoreBias);
 		}
 #endif
 
