Fix serialization when switching between LWC variant types

diff --git a/Engine/Source/Runtime/Core/Public/Math/BoxSphereBounds.h b/Engine/Source/Runtime/Core/Public/Math/BoxSphereBounds.h
--- a/Engine/Source/Runtime/Core/Public/Math/BoxSphereBounds.h
+++ b/Engine/Source/Runtime/Core/Public/Math/BoxSphereBounds.h
@@ -498,11 +498,13 @@ template <> struct TIsUECoreVariant<FCompactBoxSphereBounds3d> { enum { Value =
 template<>
 inline bool FBoxSphereBounds3f::SerializeFromMismatchedTag(FName StructTag, FArchive& Ar)
 {
-	return UE_SERIALIZE_VARIANT_FROM_MISMATCHED_TAG(Ar, BoxSphereBounds, BoxSphereBounds3f, BoxSphereBounds3d);
+	// Falls back to UseSerializeItem to convert per property.
+	return false;
 }
 
 template<>
 inline bool FBoxSphereBounds3d::SerializeFromMismatchedTag(FName StructTag, FArchive& Ar)
 {
-	return UE_SERIALIZE_VARIANT_FROM_MISMATCHED_TAG(Ar, BoxSphereBounds, BoxSphereBounds3d, BoxSphereBounds3f);
+	// Falls back to UseSerializeItem to convert per property.
+	return false;
 }
diff --git a/Engine/Source/Runtime/Core/Public/Math/Transform.h b/Engine/Source/Runtime/Core/Public/Math/Transform.h
--- a/Engine/Source/Runtime/Core/Public/Math/Transform.h
+++ b/Engine/Source/Runtime/Core/Public/Math/Transform.h
@@ -27,11 +27,21 @@ DECLARE_INTRINSIC_TYPE_LAYOUT(FTransform3d);
 template<>
 inline bool FTransform3f::SerializeFromMismatchedTag(FName StructTag, FArchive& Ar)
 {
+	// FTransform3f/d are immutable, but FTransform is not and needs a per-property conversion
+	if(StructTag == NAME_Transform)
+	{
+		return false;
+	}
 	return UE_SERIALIZE_VARIANT_FROM_MISMATCHED_TAG(Ar, Transform, Transform3f, Transform3d);
 }
 
 template<>
 inline bool FTransform3d::SerializeFromMismatchedTag(FName StructTag, FArchive& Ar)
 {
+	// FTransform3f/d are immutable, but FTransform is not and needs a per-property conversion
+	if(StructTag == NAME_Transform)
+	{
+		return false;
+	}
 	return UE_SERIALIZE_VARIANT_FROM_MISMATCHED_TAG(Ar, Transform, Transform3d, Transform3f);
 }
\ No newline at end of file
diff --git a/Engine/Source/Runtime/Core/Public/Misc/LargeWorldCoordinatesSerializer.h b/Engine/Source/Runtime/Core/Public/Misc/LargeWorldCoordinatesSerializer.h
--- a/Engine/Source/Runtime/Core/Public/Misc/LargeWorldCoordinatesSerializer.h
+++ b/Engine/Source/Runtime/Core/Public/Misc/LargeWorldCoordinatesSerializer.h
@@ -2,23 +2,27 @@
 
 #pragma once
 #include "UObject/NameTypes.h"
+#include "UObject/ObjectVersion.h"
+#include "Serialization/StructuredArchive.h"
 
 // Serialization helper for core variant types only. DO NOT USE!
 #define UE_SERIALIZE_VARIANT_FROM_MISMATCHED_TAG(AR_OR_SLOT, ALIAS, TYPE, ALT_TYPE) LWCSerializerPrivate::SerializeFromMismatchedTag<F##ALT_TYPE>(*this, StructTag, AR_OR_SLOT, NAME_##ALIAS, NAME_##TYPE, NAME_##ALT_TYPE)
 namespace LWCSerializerPrivate
 {
 
+inline bool IsPreLWC(const FArchive& Ar) { return Ar.UEVer() < EUnrealEngineObjectUE5Version::LARGE_WORLD_COORDINATES; } 
+inline bool IsPreLWC(const FStructuredArchive::FSlot& Slot) { return IsPreLWC(Slot.GetUnderlyingArchive()); } 
+
 // SerializeFromMismatchedTag helper for core type use only. DO NOT USE!
 template<typename FAltType, typename FType, typename FArSlot>
-bool SerializeFromMismatchedTag(FType& Target, FName StructTag, FArSlot& ArSlot, FName BaseTag, FName AltTag, FName ThisTag)
+bool SerializeFromMismatchedTag(FType& Target, FName StructTag, FArSlot& ArSlot, FName BaseTag, FName ThisTag, FName AltTag)
 {
-	if( StructTag == BaseTag || 
-		StructTag == ThisTag)
+	if(StructTag == ThisTag || (StructTag == BaseTag && (TIsUECoreVariant<FType, double>::Value || IsPreLWC(ArSlot))))
 	{
 		// Note: relies on Serialize to handle float/double based on archive version.
 		return Target.Serialize(ArSlot);
 	}
-	else if(StructTag == AltTag)
+	else if(StructTag == AltTag || StructTag == BaseTag)
 	{
 		// Convert from alt type
 		FAltType AsAlt;										// TODO: Could we derive this from FType?
diff --git a/Engine/Source/Runtime/CoreUObject/Private/UObject/PropertyStruct.cpp b/Engine/Source/Runtime/CoreUObject/Private/UObject/PropertyStruct.cpp
--- a/Engine/Source/Runtime/CoreUObject/Private/UObject/PropertyStruct.cpp
+++ b/Engine/Source/Runtime/CoreUObject/Private/UObject/PropertyStruct.cpp
@@ -435,9 +435,20 @@ EConvertFromTypeResult FStructProperty::ConvertFromType(const FPropertyTag& Tag,
 				FArchive& Ar = Adapter.GetArchive();
 				if (CppStructOps->HasSerializeFromMismatchedTag() && CppStructOps->SerializeFromMismatchedTag(Tag, Ar, DestAddress))
 				{
-
 					return EConvertFromTypeResult::Converted;
 				}
+				else if(((Struct->StructFlags & STRUCT_SerializeNative) == 0) && CppStructOps->HasSerializeFromMismatchedTag() && CppStructOps->IsUECoreVariant())
+				{
+					// Special case for Transform, as the f/d variants are immutable whilst the default is not, so we must call SerializeTaggedProperties directly to perform the conversion.
+					if(Tag.StructName == NAME_Transform)
+					{
+						Struct->SerializeTaggedProperties(Slot, (uint8*)DestAddress, Struct, nullptr);
+						return EConvertFromTypeResult::Converted;
+					}
+					// If a core variant without a native serializer returns false from SerializeFromMismatchedTag fall back to standard SerializeItem.
+					// We rely on all properties within the variant supporting SerializeFromMismatchedTag to perform the conversion per property.
+					return EConvertFromTypeResult::UseSerializeItem;
+				}
 				else
 				{
 					UE_LOG(LogClass, Warning, TEXT("SerializeFromMismatchedTag failed: Type mismatch in %s of %s - Previous (%s) Current(StructProperty) for package:  %s"), *Tag.Name.ToString(), *GetName(), *Tag.Type.ToString(), *UnderlyingArchive.GetArchiveName());
diff --git a/Engine/Source/Runtime/CoreUObject/Public/UObject/Class.h b/Engine/Source/Runtime/CoreUObject/Public/UObject/Class.h
--- a/Engine/Source/Runtime/CoreUObject/Public/UObject/Class.h
+++ b/Engine/Source/Runtime/CoreUObject/Public/UObject/Class.h
@@ -858,6 +858,9 @@ public:
 		/** return true if this struct is one of the UE Core types (and is include in CoreMinimal.h) **/
 		virtual bool IsUECoreType() = 0;
 
+		/** return true if this struct is one of the UE Core types (and is include in CoreMinimal.h) **/
+		virtual bool IsUECoreVariant() = 0;		
+
 		/** return true if this struct can copy **/
 		virtual bool HasCopy() = 0;
 		/** 
@@ -1100,6 +1103,10 @@ public:
 		{
 			return TIsUECoreType<CPPSTRUCT>::Value;
 		}
+		virtual bool IsUECoreVariant() override
+		{
+			return TIsUECoreVariant<CPPSTRUCT>::Value;
+		}		
 		virtual bool HasCopy() override
 		{
 			return TTraits::WithCopy;
