Add support for reserving memory when using LargeMemoryWriter.

#rnx

diff --git a/Engine/Source/Runtime/Core/Public/Serialization/LargeMemoryWriter.h b/Engine/Source/Runtime/Core/Public/Serialization/LargeMemoryWriter.h
--- a/Engine/Source/Runtime/Core/Public/Serialization/LargeMemoryWriter.h
+++ b/Engine/Source/Runtime/Core/Public/Serialization/LargeMemoryWriter.h
@@ -61,6 +61,14 @@ public:
 		return Data.ReleaseOwnership();
 	}
 
+	/**
+	 * Reserves memory such that the writer can contain at least Size number of bytes.
+	 */
+	void Reserve(int64 Size)
+	{
+		Data.Reserve(Size);
+	}
+
 private:
 
 	FLargeMemoryData Data;
