TCHAR array init fixes for UTF-8 mode.

diff --git a/Engine/Source/Developer/DerivedDataCache/Private/FileSystemCacheStore.cpp b/Engine/Source/Developer/DerivedDataCache/Private/FileSystemCacheStore.cpp
--- a/Engine/Source/Developer/DerivedDataCache/Private/FileSystemCacheStore.cpp
+++ b/Engine/Source/Developer/DerivedDataCache/Private/FileSystemCacheStore.cpp
@@ -69,8 +69,8 @@ TRACE_DECLARE_INT_COUNTER(FileSystemDDC_BytesWritten, TEXT("FileSystemDDC Bytes
 
 ///////////////////////////////////////////////////////////////////////////////////////////////////////////////////////
 
-static const TCHAR GBucketsDirectoryName[] = TEXT("Buckets");
-static const TCHAR GContentDirectoryName[] = TEXT("Content");
+static const TCHAR* GBucketsDirectoryName = TEXT("Buckets");
+static const TCHAR* GContentDirectoryName = TEXT("Content");
 
 ///////////////////////////////////////////////////////////////////////////////////////////////////////////////////////
 
diff --git a/Engine/Source/Developer/DerivedDataCache/Private/HttpCacheStore.cpp b/Engine/Source/Developer/DerivedDataCache/Private/HttpCacheStore.cpp
--- a/Engine/Source/Developer/DerivedDataCache/Private/HttpCacheStore.cpp
+++ b/Engine/Source/Developer/DerivedDataCache/Private/HttpCacheStore.cpp
@@ -1552,7 +1552,7 @@ private:
 		// The expected data stream is structured accordingly
 		// {"JPTR"} {PayloadCount:uint32} {{"JPEE"} {Name:cstr} {Result:uint8} {Hash:IoHash} {Size:uint64} {Payload...}} ...
 
-		const TCHAR ResponseErrorMessage[] = TEXT("Malformed response from server.");
+		const auto& ResponseErrorMessage = TEXT("Malformed response from server.");
 		const ANSICHAR* ProtocolMagic = "JPTR";
 		const ANSICHAR* PayloadMagic = "JPEE";
 		const uint32 MagicSize = 4;
diff --git a/Engine/Source/Programs/UnrealHeaderTool/Private/Manifest.cpp b/Engine/Source/Programs/UnrealHeaderTool/Private/Manifest.cpp
--- a/Engine/Source/Programs/UnrealHeaderTool/Private/Manifest.cpp
+++ b/Engine/Source/Programs/UnrealHeaderTool/Private/Manifest.cpp
@@ -203,7 +203,7 @@ bool FManifestModule::NeedsRegeneration() const
 		return true;
 	}
 	FString Timestamp;
-	TCHAR TimestampText[] = TEXT("Timestamp");
+	const auto& TimestampText = TEXT("Timestamp");
 	Timestamp.Empty(GeneratedIncludeDirectory.Len() + UE_ARRAY_COUNT(TimestampText));
 	Timestamp += GeneratedIncludeDirectory;
 	Timestamp += TEXT("Timestamp");
diff --git a/Engine/Source/Programs/UnrealHeaderTool/Private/UnrealTypeDefinitionInfo.cpp b/Engine/Source/Programs/UnrealHeaderTool/Private/UnrealTypeDefinitionInfo.cpp
--- a/Engine/Source/Programs/UnrealHeaderTool/Private/UnrealTypeDefinitionInfo.cpp
+++ b/Engine/Source/Programs/UnrealHeaderTool/Private/UnrealTypeDefinitionInfo.cpp
@@ -975,7 +975,7 @@ void FUnrealPropertyDefinitionInfo::ExportCppDeclaration(FOutputDevice& Out, EEx
 		NameCpp = FString(TEXT(", ")) + NameCpp;
 	}
 
-	TCHAR ArrayStr[MAX_SPRINTF] = TEXT("");
+	TCHAR ArrayStr[MAX_SPRINTF] = {};
 	const bool bExportStaticArray = 0 == (CPPF_NoStaticArray & AdditionalExportCPPFlags);
 	if (ArrayDimOverride != nullptr && bExportStaticArray)
 	{
diff --git a/Engine/Source/Runtime/Core/Private/GenericPlatform/GenericPlatformMisc.cpp b/Engine/Source/Runtime/Core/Private/GenericPlatform/GenericPlatformMisc.cpp
--- a/Engine/Source/Runtime/Core/Private/GenericPlatform/GenericPlatformMisc.cpp
+++ b/Engine/Source/Runtime/Core/Private/GenericPlatform/GenericPlatformMisc.cpp
@@ -1373,19 +1373,52 @@ const TCHAR* FGenericPlatformMisc::GetUBTTarget()
 	return TEXT(PREPROCESSOR_TO_STRING(UBT_COMPILED_TARGET));
 }
 
-/** The name of the UBT target that the current executable was built from. Defaults to the UE4 default target for this type to make content only projects work, 
-	but will be overridden by the primary game module if it exists */
-TCHAR GUBTTargetName[128] = TEXT("Unreal" PREPROCESSOR_TO_STRING(UBT_COMPILED_TARGET));
+using FUBTTargetNameArrayType = TCHAR[128];
+
+#if PLATFORM_TCHAR_IS_UTF8CHAR
+
+	// We can't initialize a sized UTF8CHAR static array with a UTF8TEXT until we have char8_t in C++20,
+	// so we use a constructor to copy the value in.
+	static FUBTTargetNameArrayType& GetStaticUBTTargetName()
+	{
+		static struct FInitializer
+		{
+			FInitializer()
+			{
+				/** The name of the UBT target that the current executable was built from. Defaults to the UE default target for this type to make content only projects work,
+					but will be overridden by the primary game module if it exists */
+				FCString::Strcpy(UBTTargetName, TEXT("Unreal" PREPROCESSOR_TO_STRING(UBT_COMPILED_TARGET)));
+			}
+
+			TCHAR UBTTargetName[128];
+		} Initializer;
+
+		return Initializer.UBTTargetName;
+	}
+
+#else
+
+	static FUBTTargetNameArrayType& GetStaticUBTTargetName()
+	{
+		/** The name of the UBT target that the current executable was built from. Defaults to the UE default target for this type to make content only projects work,
+			but will be overridden by the primary game module if it exists */
+		static FUBTTargetNameArrayType GUBTTargetName = TEXT("Unreal" PREPROCESSOR_TO_STRING(UBT_COMPILED_TARGET));
+
+		return GUBTTargetName;
+	}
+
+#endif
 
 void FGenericPlatformMisc::SetUBTTargetName(const TCHAR* InTargetName)
 {
-	check(FCString::Strlen(InTargetName) < (UE_ARRAY_COUNT(GUBTTargetName) - 1));
-	FCString::Strcpy(GUBTTargetName, InTargetName);
+	FUBTTargetNameArrayType& UBTTargetName = GetStaticUBTTargetName();
+	check(FCString::Strlen(InTargetName) < (UE_ARRAY_COUNT(UBTTargetName) - 1));
+	FCString::Strcpy(UBTTargetName, InTargetName);
 }
 
 const TCHAR* FGenericPlatformMisc::GetUBTTargetName()
 {
-	return GUBTTargetName;
+	return GetStaticUBTTargetName();
 }
 
 const TCHAR* FGenericPlatformMisc::GetDefaultDeviceProfileName()
diff --git a/Engine/Source/Runtime/Core/Private/HAL/ExceptionHandling.cpp b/Engine/Source/Runtime/Core/Private/HAL/ExceptionHandling.cpp
--- a/Engine/Source/Runtime/Core/Private/HAL/ExceptionHandling.cpp
+++ b/Engine/Source/Runtime/Core/Private/HAL/ExceptionHandling.cpp
@@ -22,7 +22,7 @@ CORE_API bool GAlwaysReportCrash = false;
 /** Whether to use ClientReportClient rather than the old AutoReporter. */
 CORE_API bool GUseCrashReportClient = true;
 
-CORE_API TCHAR MiniDumpFilenameW[1024] = TEXT("");
+CORE_API TCHAR MiniDumpFilenameW[1024] = {};
 
 
 bool GEnsureShowsCRC = false;
diff --git a/Engine/Source/Runtime/Core/Private/IO/IoDispatcher.cpp b/Engine/Source/Runtime/Core/Private/IO/IoDispatcher.cpp
--- a/Engine/Source/Runtime/Core/Private/IO/IoDispatcher.cpp
+++ b/Engine/Source/Runtime/Core/Private/IO/IoDispatcher.cpp
@@ -35,6 +35,27 @@ TUniquePtr<FIoDispatcher> GIoDispatcher;
 CSV_DEFINE_CATEGORY(IoDispatcher, true);
 CSV_DEFINE_STAT(IoDispatcher, PendingIoRequests);
 
+static const TCHAR* const GetIoErrorText_ErrorCodeTextArray[] =
+{
+	TEXT("OK"),
+	TEXT("Unknown Status"),
+	TEXT("Invalid Code"),
+	TEXT("Cancelled"),
+	TEXT("FileOpen Failed"),
+	TEXT("File Not Open"),
+	TEXT("Read Error"),
+	TEXT("Write Error"),
+	TEXT("Not Found"),
+	TEXT("Corrupt Toc"),
+	TEXT("Unknown ChunkID"),
+	TEXT("Invalid Parameter"),
+	TEXT("Signature Error"),
+	TEXT("Invalid Encryption Key"),
+	TEXT("Compression Error")
+};
+
+CORE_API const TCHAR* const* GetIoErrorText_ErrorCodeText = GetIoErrorText_ErrorCodeTextArray;
+
 #if UE_IODISPATCHER_STATS_ENABLED
 class FIoRequestStats
 {
diff --git a/Engine/Source/Runtime/Core/Private/Internationalization/TextHistory.cpp b/Engine/Source/Runtime/Core/Private/Internationalization/TextHistory.cpp
--- a/Engine/Source/Runtime/Core/Private/Internationalization/TextHistory.cpp
+++ b/Engine/Source/Runtime/Core/Private/Internationalization/TextHistory.cpp
@@ -94,8 +94,8 @@ const TCHAR* SkipWhitespaceAndCharacter(const TCHAR* Buffer, const TCHAR InChar)
 
 const TCHAR* ReadNumberFromBuffer(const TCHAR* Buffer, FFormatArgumentValue& OutValue)
 {
-	static const TCHAR ValidNumericChars[] = TEXT("+-0123456789.ful");
-	static const TCHAR SuffixNumericChars[] = TEXT("ful");
+	static const TCHAR* ValidNumericChars = TEXT("+-0123456789.ful");
+	static const TCHAR* SuffixNumericChars = TEXT("ful");
 
 	FString NumericString;
 	while (*Buffer && FCString::Strchr(ValidNumericChars, *Buffer))
diff --git a/Engine/Source/Runtime/Core/Private/Internationalization/TextHistory.h b/Engine/Source/Runtime/Core/Private/Internationalization/TextHistory.h
--- a/Engine/Source/Runtime/Core/Private/Internationalization/TextHistory.h
+++ b/Engine/Source/Runtime/Core/Private/Internationalization/TextHistory.h
@@ -45,26 +45,26 @@ namespace TextStringificationUtil
 {
 
 #define LOC_DEFINE_REGION
-static const TCHAR TextMarker[] = TEXT("TEXT");
-static const TCHAR InvTextMarker[] = TEXT("INVTEXT");
-static const TCHAR NsLocTextMarker[] = TEXT("NSLOCTEXT");
-static const TCHAR LocTextMarker[] = TEXT("LOCTEXT");
-static const TCHAR LocTableMarker[] = TEXT("LOCTABLE");
-static const TCHAR LocGenNumberMarker[] = TEXT("LOCGEN_NUMBER");
-static const TCHAR LocGenPercentMarker[] = TEXT("LOCGEN_PERCENT");
-static const TCHAR LocGenCurrencyMarker[] = TEXT("LOCGEN_CURRENCY");
-static const TCHAR LocGenDateMarker[] = TEXT("LOCGEN_DATE");
-static const TCHAR LocGenTimeMarker[] = TEXT("LOCGEN_TIME");
-static const TCHAR LocGenDateTimeMarker[] = TEXT("LOCGEN_DATETIME");
-static const TCHAR LocGenToLowerMarker[] = TEXT("LOCGEN_TOLOWER");
-static const TCHAR LocGenToUpperMarker[] = TEXT("LOCGEN_TOUPPER");
-static const TCHAR LocGenFormatOrderedMarker[] = TEXT("LOCGEN_FORMAT_ORDERED");
-static const TCHAR LocGenFormatNamedMarker[] = TEXT("LOCGEN_FORMAT_NAMED");
-static const TCHAR GroupedSuffix[] = TEXT("_GROUPED");
-static const TCHAR UngroupedSuffix[] = TEXT("_UNGROUPED");
-static const TCHAR CustomSuffix[] = TEXT("_CUSTOM");
-static const TCHAR UtcSuffix[] = TEXT("_UTC");
-static const TCHAR LocalSuffix[] = TEXT("_LOCAL");
+static const auto& TextMarker = TEXT("TEXT");
+static const auto& InvTextMarker = TEXT("INVTEXT");
+static const auto& NsLocTextMarker = TEXT("NSLOCTEXT");
+static const auto& LocTextMarker = TEXT("LOCTEXT");
+static const auto& LocTableMarker = TEXT("LOCTABLE");
+static const auto& LocGenNumberMarker = TEXT("LOCGEN_NUMBER");
+static const auto& LocGenPercentMarker = TEXT("LOCGEN_PERCENT");
+static const auto& LocGenCurrencyMarker = TEXT("LOCGEN_CURRENCY");
+static const auto& LocGenDateMarker = TEXT("LOCGEN_DATE");
+static const auto& LocGenTimeMarker = TEXT("LOCGEN_TIME");
+static const auto& LocGenDateTimeMarker = TEXT("LOCGEN_DATETIME");
+static const auto& LocGenToLowerMarker = TEXT("LOCGEN_TOLOWER");
+static const auto& LocGenToUpperMarker = TEXT("LOCGEN_TOUPPER");
+static const auto& LocGenFormatOrderedMarker = TEXT("LOCGEN_FORMAT_ORDERED");
+static const auto& LocGenFormatNamedMarker = TEXT("LOCGEN_FORMAT_NAMED");
+static const auto& GroupedSuffix = TEXT("_GROUPED");
+static const auto& UngroupedSuffix = TEXT("_UNGROUPED");
+static const auto& CustomSuffix = TEXT("_CUSTOM");
+static const auto& UtcSuffix = TEXT("_UTC");
+static const auto& LocalSuffix = TEXT("_LOCAL");
 #undef LOC_DEFINE_REGION
 
 #define TEXT_STRINGIFICATION_FUNC_MODIFY_BUFFER_AND_VALIDATE(Func, ...)		\
diff --git a/Engine/Source/Runtime/Core/Private/Logging/LogSuppressionInterface.cpp b/Engine/Source/Runtime/Core/Private/Logging/LogSuppressionInterface.cpp
--- a/Engine/Source/Runtime/Core/Private/Logging/LogSuppressionInterface.cpp
+++ b/Engine/Source/Runtime/Core/Private/Logging/LogSuppressionInterface.cpp
@@ -132,7 +132,7 @@ class FLogSuppressionImplementation: public FLogSuppressionInterface, private FS
 		Cmds.ParseIntoArray(SubCmds, TEXT(","), true);
 		for (int32 Index = 0; Index < SubCmds.Num(); Index++)
 		{
-			static constexpr TCHAR LogString[] = TEXT("Log ");
+			static const auto& LogString = TEXT("Log ");
 			FString Command = SubCmds[Index].TrimStart();
 			if (Command.StartsWith(LogString))
 			{
diff --git a/Engine/Source/Runtime/Core/Private/Misc/CommandLine.cpp b/Engine/Source/Runtime/Core/Private/Misc/CommandLine.cpp
--- a/Engine/Source/Runtime/Core/Private/Misc/CommandLine.cpp
+++ b/Engine/Source/Runtime/Core/Private/Misc/CommandLine.cpp
@@ -13,10 +13,10 @@
 -----------------------------------------------------------------------------*/
 
 bool FCommandLine::bIsInitialized = false;
-TCHAR FCommandLine::CmdLine[FCommandLine::MaxCommandLineSize] = TEXT("");
-TCHAR FCommandLine::OriginalCmdLine[FCommandLine::MaxCommandLineSize] = TEXT("");
-TCHAR FCommandLine::LoggingCmdLine[FCommandLine::MaxCommandLineSize] = TEXT("");
-TCHAR FCommandLine::LoggingOriginalCmdLine[FCommandLine::MaxCommandLineSize] = TEXT("");
+TCHAR FCommandLine::CmdLine[FCommandLine::MaxCommandLineSize] = {};
+TCHAR FCommandLine::OriginalCmdLine[FCommandLine::MaxCommandLineSize] = {};
+TCHAR FCommandLine::LoggingCmdLine[FCommandLine::MaxCommandLineSize] = {};
+TCHAR FCommandLine::LoggingOriginalCmdLine[FCommandLine::MaxCommandLineSize] = {};
 
 FString& FCommandLine::GetSubprocessCommandLine_Internal()
 {
diff --git a/Engine/Source/Runtime/Core/Private/Misc/ConfigCacheIni.cpp b/Engine/Source/Runtime/Core/Private/Misc/ConfigCacheIni.cpp
--- a/Engine/Source/Runtime/Core/Private/Misc/ConfigCacheIni.cpp
+++ b/Engine/Source/Runtime/Core/Private/Misc/ConfigCacheIni.cpp
@@ -1230,13 +1230,13 @@ void FConfigFile::AppendExportedPropertyLine(FString& Out, const FString& Proper
 namespace CommandlineOverrideSpecifiers
 {
 	// -ini:IniName:[Section1]:Key1=Value1,[Section2]:Key2=Value2
-	const TCHAR	IniFileOverrideIdentifier[] = TEXT("-iniFile=");
-	const TCHAR IniSwitchIdentifier[]     = TEXT("-ini:");
-	const TCHAR IniNameEndIdentifier[]    = TEXT(":[");
-	const TCHAR SectionStartIdentifier[]  = TEXT("[");
-	const TCHAR PropertyStartIdentifier[] = TEXT("]:");
-	const TCHAR PropertySeperator[]       = TEXT(",");
-	const TCHAR CustomConfigIdentifier[] = TEXT("-CustomConfig=");
+	const auto& IniFileOverrideIdentifier = TEXT("-iniFile=");
+	const auto& IniSwitchIdentifier       = TEXT("-ini:");
+	const auto& IniNameEndIdentifier      = TEXT(":[");
+	const auto& SectionStartIdentifier    = TEXT("[");
+	const auto& PropertyStartIdentifier   = TEXT("]:");
+	const auto& PropertySeperator         = TEXT(",");
+	const auto& CustomConfigIdentifier    = TEXT("-CustomConfig=");
 }
 
 #endif
@@ -2100,14 +2100,14 @@ void FConfigFile::SetText( const TCHAR* Section, const TCHAR* Key, const FText&
 
 void FConfigFile::SetFloat(const TCHAR* Section, const TCHAR* Key, float Value)
 {
-	TCHAR Text[MAX_SPRINTF] = TEXT("");
+	TCHAR Text[MAX_SPRINTF];
 	FCString::Sprintf(Text, TEXT("%.*g"), std::numeric_limits<float>::max_digits10, Value);
 	SetString(Section, Key, Text);
 }
 
 void FConfigFile::SetDouble(const TCHAR* Section, const TCHAR* Key, double Value)
 {
-	TCHAR Text[MAX_SPRINTF] = TEXT("");
+	TCHAR Text[MAX_SPRINTF];
 	FCString::Sprintf(Text, TEXT("%.*g"), std::numeric_limits<double>::max_digits10, Value);
 	SetString(Section, Key, Text);
 }
@@ -2119,7 +2119,7 @@ void FConfigFile::SetBool(const TCHAR* Section, const TCHAR* Key, bool Value)
 
 void FConfigFile::SetInt64( const TCHAR* Section, const TCHAR* Key, int64 Value )
 {
-	TCHAR Text[MAX_SPRINTF]=TEXT("");
+	TCHAR Text[MAX_SPRINTF];
 	FCString::Sprintf( Text, TEXT("%lld"), Value );
 	SetString( Section, Key, Text );
 }
@@ -3236,7 +3236,7 @@ void FConfigCacheIni::SetInt
 	const FString&	Filename
 )
 {
-	TCHAR Text[MAX_SPRINTF]=TEXT("");
+	TCHAR Text[MAX_SPRINTF];
 	FCString::Sprintf( Text, TEXT("%i"), Value );
 	SetString( Section, Key, Text, Filename );
 }
diff --git a/Engine/Source/Runtime/Core/Private/Misc/CoreGlobals.cpp b/Engine/Source/Runtime/Core/Private/Misc/CoreGlobals.cpp
--- a/Engine/Source/Runtime/Core/Private/Misc/CoreGlobals.cpp
+++ b/Engine/Source/Runtime/Core/Private/Misc/CoreGlobals.cpp
@@ -48,10 +48,10 @@ CORE_API FMalloc**			GFixedMallocLocationPtr = nullptr;		/* Memory allocator poi
 class FPropertyWindowManager*	GPropertyWindowManager	= nullptr;		/* Manages and tracks property editing windows */
 
 /** For building call stack text dump in guard/unguard mechanism. */
-TCHAR GErrorHist[16384]	= TEXT("");
+TCHAR GErrorHist[16384] = {};
 
 /** For building exception description text dump in guard/unguard mechanism. */
-TCHAR GErrorExceptionDescription[4096] = TEXT( "" );
+TCHAR GErrorExceptionDescription[4096] = {};
 
 // We define our texts like this so that the header only needs to refer to references to FTexts,
 // as FText is only forward-declared there.
diff --git a/Engine/Source/Runtime/Core/Private/Misc/FbxErrors.cpp b/Engine/Source/Runtime/Core/Private/Misc/FbxErrors.cpp
--- a/Engine/Source/Runtime/Core/Private/Misc/FbxErrors.cpp
+++ b/Engine/Source/Runtime/Core/Private/Misc/FbxErrors.cpp
@@ -166,8 +166,8 @@ FLazyName FFbxErrors::Animation_CurveNotFound(TEXT("Animation_CurveNotFound"));
 
 namespace
 {
-	constexpr TCHAR FbxErrorsURLPath[] = TEXT("WorkingWithContent/Importing/FBX/FbxErrors/");
-	constexpr TCHAR FbxErrorsLocalPath[] = TEXT("Shared/Editor/FbxErrors");
+	const TCHAR* FbxErrorsURLPath = TEXT("WorkingWithContent/Importing/FBX/FbxErrors/");
+	const TCHAR* FbxErrorsLocalPath = TEXT("Shared/Editor/FbxErrors");
 }
 
 FFbxErrorToken::FFbxErrorToken(const FName& InErrorName)
diff --git a/Engine/Source/Runtime/Core/Private/Misc/MapErrors.cpp b/Engine/Source/Runtime/Core/Private/Misc/MapErrors.cpp
--- a/Engine/Source/Runtime/Core/Private/Misc/MapErrors.cpp
+++ b/Engine/Source/Runtime/Core/Private/Misc/MapErrors.cpp
@@ -66,8 +66,8 @@ FLazyName FMapErrors::InvalidVirtualTextureUsage(TEXT("InvalidVirtualTextureUsag
 
 namespace
 {
-	constexpr TCHAR MapErrorsURLPath[] = TEXT("BuildingWorlds/LevelEditor/MapErrors");
-	constexpr TCHAR MapErrorsLocalPath[] = TEXT("Shared/Editor/MapErrors");
+	const TCHAR* MapErrorsURLPath = TEXT("BuildingWorlds/LevelEditor/MapErrors");
+	const TCHAR* MapErrorsLocalPath = TEXT("Shared/Editor/MapErrors");
 }
 
 FMapErrorToken::FMapErrorToken(const FName& InErrorName)
diff --git a/Engine/Source/Runtime/Core/Private/Misc/MessageDialog.cpp b/Engine/Source/Runtime/Core/Private/Misc/MessageDialog.cpp
--- a/Engine/Source/Runtime/Core/Private/Misc/MessageDialog.cpp
+++ b/Engine/Source/Runtime/Core/Private/Misc/MessageDialog.cpp
@@ -48,7 +48,7 @@ void FMessageDialog::ShowLastError()
 {
 	uint32 LastError = FPlatformMisc::GetLastError();
 
-	TCHAR TempStr[MAX_SPRINTF]=TEXT("");
+	TCHAR TempStr[MAX_SPRINTF] = {};
 	TCHAR ErrorBuffer[1024];
 	FCString::Sprintf( TempStr, TEXT("GetLastError : %d\n\n%s"), LastError, FPlatformMisc::GetSystemErrorMessage(ErrorBuffer, 1024, 0) );
 	if( FApp::IsUnattended() == true )
diff --git a/Engine/Source/Runtime/Core/Private/Misc/Parse.cpp b/Engine/Source/Runtime/Core/Private/Misc/Parse.cpp
--- a/Engine/Source/Runtime/Core/Private/Misc/Parse.cpp
+++ b/Engine/Source/Runtime/Core/Private/Misc/Parse.cpp
@@ -517,7 +517,8 @@ bool FParse::Value( const TCHAR* Stream, const TCHAR* Match, uint64& Value )
 //
 bool FParse::Value( const TCHAR* Stream, const TCHAR* Match, int64& Value )
 {
-	TCHAR Temp[4096]=TEXT(""), *Ptr=Temp;
+	TCHAR Temp[4096] = {};
+	TCHAR* Ptr = Temp;
 	if( FParse::Value( Stream, Match, Temp, UE_ARRAY_COUNT(Temp) ) )
 	{
 		Value = 0;
diff --git a/Engine/Source/Runtime/Core/Private/Misc/Paths.cpp b/Engine/Source/Runtime/Core/Private/Misc/Paths.cpp
--- a/Engine/Source/Runtime/Core/Private/Misc/Paths.cpp
+++ b/Engine/Source/Runtime/Core/Private/Misc/Paths.cpp
@@ -1089,7 +1089,7 @@ void FPaths::NormalizeDirectoryName(FString& InPath)
 
 bool FPaths::CollapseRelativeDirectories(FString& InPath)
 {
-	const TCHAR ParentDir[] = TEXT("/..");
+	const auto& ParentDir = TEXT("/..");
 	const int32 ParentDirLength = UE_ARRAY_COUNT( ParentDir ) - 1; // To avoid hardcoded values
 
 	for (;;)
@@ -1397,7 +1397,7 @@ FString FPaths::GetInvalidFileSystemChars()
 	// # isn't legal. Used for revision specifiers in P4/SVN, and also not allowed on Windows anyway
 	// @ isn't legal. Used for revision/label specifiers in P4/SVN
 	// ^ isn't legal. While the file-system won't complain about this character, Visual Studio will			
-	static const TCHAR RestrictedChars[] = TEXT("/?:&\\*\"<>|%#@^");
+	static const TCHAR* RestrictedChars = TEXT("/?:&\\*\"<>|%#@^");
 	return RestrictedChars;
 }
 
diff --git a/Engine/Source/Runtime/Core/Private/Serialization/BitWriter.cpp b/Engine/Source/Runtime/Core/Private/Serialization/BitWriter.cpp
--- a/Engine/Source/Runtime/Core/Private/Serialization/BitWriter.cpp
+++ b/Engine/Source/Runtime/Core/Private/Serialization/BitWriter.cpp
@@ -127,7 +127,7 @@ void FBitWriter::SerializeInt(uint32& Value, uint32 ValueMax)
 
 	if (WriteValue >= ValueMax)
 	{
-		const TCHAR Msg[] = TEXT("FBitWriter::SerializeInt(): Value out of bounds (Value: %u, ValueMax: %u)");
+		const auto& Msg = TEXT("FBitWriter::SerializeInt(): Value out of bounds (Value: %u, ValueMax: %u)");
 
 		UE_LOG(LogSerialization, Error, Msg, WriteValue, ValueMax);
 		ensureMsgf(false, Msg, WriteValue, ValueMax);
diff --git a/Engine/Source/Runtime/Core/Private/Serialization/MemoryImage.cpp b/Engine/Source/Runtime/Core/Private/Serialization/MemoryImage.cpp
--- a/Engine/Source/Runtime/Core/Private/Serialization/MemoryImage.cpp
+++ b/Engine/Source/Runtime/Core/Private/Serialization/MemoryImage.cpp
@@ -2004,7 +2004,7 @@ uint8 Freeze::FindFieldNameLength(const TCHAR* Name)
 
 		if (Ch == TEXT('_'))
 		{
-			static TCHAR DeprecatedName[] = TEXT("DEPRECATED");
+			static const TCHAR* DeprecatedName = TEXT("DEPRECATED");
 
 			const TCHAR* DepIter = DeprecatedName;
 			for (;;)
diff --git a/Engine/Source/Runtime/Core/Private/Tests/HAL/PlatformAtomicsTest.cpp b/Engine/Source/Runtime/Core/Private/Tests/HAL/PlatformAtomicsTest.cpp
--- a/Engine/Source/Runtime/Core/Private/Tests/HAL/PlatformAtomicsTest.cpp
+++ b/Engine/Source/Runtime/Core/Private/Tests/HAL/PlatformAtomicsTest.cpp
@@ -39,7 +39,7 @@ static bool TestInterlockedAnd(FAutomationTestBase& Test)
 {
 	bool bSuccess = true;
 
-	constexpr TCHAR FunctionName[] = TEXT("And");
+	const TCHAR* FunctionName = TEXT("And");
 	{
 		// Test and with value where some bits are set and some aren't set in the current value.
 		volatile int8 Value0 = 0x30;
@@ -103,7 +103,7 @@ static bool TestInterlockedOr(FAutomationTestBase& Test)
 {
 	bool bSuccess = true;
 
-	constexpr TCHAR FunctionName[] = TEXT("Or");
+	const TCHAR* FunctionName = TEXT("Or");
 	{
 		// Test or with value where some bits are set and some aren't set in the current value.
 		volatile int8 Value0 = 0x30;
@@ -167,7 +167,7 @@ static bool TestInterlockedXor(FAutomationTestBase& Test)
 {
 	bool bSuccess = true;
 
-	constexpr TCHAR FunctionName[] = TEXT("Xor");
+	const TCHAR* FunctionName = TEXT("Xor");
 	{
 		// Test xor with value where some bits are set and some aren't set in the current value.
 		volatile int8 Value0 = 0x30;
@@ -231,7 +231,7 @@ static bool TestInterlockedAdd(FAutomationTestBase& Test)
 {
 	bool bSuccess = true;
 
-	constexpr TCHAR FunctionName[] = TEXT("Add");
+	const TCHAR* FunctionName = TEXT("Add");
 	{
 		// Test adding positive value
 		volatile int8 Value0 = 0x0F;
@@ -311,7 +311,7 @@ static bool TestInterlockedIncrement(FAutomationTestBase& Test)
 {
 	bool bSuccess = true;
 
-	constexpr TCHAR FunctionName[] = TEXT("Increment");
+	const TCHAR* FunctionName = TEXT("Increment");
 	{
 		volatile int8 Value0 = 0x0F;
 		bSuccess = TestInterlocked<int8>(&Value0, 0x10, 0x10, [](volatile int8* Dest) { return FPlatformAtomics::InterlockedIncrement(Dest); }, Test, FunctionName, TEXT("int8")) && bSuccess;
@@ -355,7 +355,7 @@ static bool TestInterlockedDecrement(FAutomationTestBase& Test)
 {
 	bool bSuccess = true;
 
-	constexpr TCHAR FunctionName[] = TEXT("Decrement");
+	const TCHAR* FunctionName = TEXT("Decrement");
 	{
 		volatile int8 Value0 = 0x10;
 		bSuccess = TestInterlocked<int8>(&Value0, 0x0F, 0x0F, [](volatile int8* Dest) { return FPlatformAtomics::InterlockedDecrement(Dest); }, Test, FunctionName, TEXT("int8")) && bSuccess;
@@ -399,7 +399,7 @@ static bool TestInterlockedExchange(FAutomationTestBase& Test)
 {
 	bool bSuccess = true;
 
-	constexpr TCHAR FunctionName[] = TEXT("Exchange");
+	const TCHAR* FunctionName = TEXT("Exchange");
 	{
 		volatile int8 Value = 0x10;
 		bSuccess = TestInterlocked<int8>(&Value, 0x10, 0x01, [](volatile int8* Dest) { return FPlatformAtomics::InterlockedExchange(Dest, 0x01); }, Test, FunctionName, TEXT("int8")) && bSuccess;
@@ -446,7 +446,7 @@ static bool TestInterlockedCompareExchange(FAutomationTestBase& Test)
 {
 	bool bSuccess = true;
 
-	constexpr TCHAR FunctionName[] = TEXT("CompareExchange");
+	const TCHAR* FunctionName = TEXT("CompareExchange");
 	{
 		volatile int8 Value = 0x10;
 
diff --git a/Engine/Source/Runtime/Core/Private/Tests/HAL/PlatformStringTest.cpp b/Engine/Source/Runtime/Core/Private/Tests/HAL/PlatformStringTest.cpp
--- a/Engine/Source/Runtime/Core/Private/Tests/HAL/PlatformStringTest.cpp
+++ b/Engine/Source/Runtime/Core/Private/Tests/HAL/PlatformStringTest.cpp
@@ -43,7 +43,8 @@ bool FPlatformStringTestStrnlen::RunTest(const FString& Parameters)
 	TestEqual(TEXT("Strnlen(\"1\", 1)"), FPlatformString::Strnlen(TEXT("1"), 1), 1);
 	TestEqual(TEXT("Strnlen(\"1\", 2)"), FPlatformString::Strnlen(TEXT("1"), 2), 1);
 	TestEqual(TEXT("Strnlen(\"123\", 2)"), FPlatformString::Strnlen(TEXT("123"), 2), 2);
-	TCHAR Buffer[128] = TEXT("123456789");
+	TCHAR Buffer[128] = {};
+	FCString::Strcpy(Buffer, TEXT("123456789"));
 	TestEqual(TEXT("Strnlen(PaddedBuffer)"), FPlatformString::Strnlen(Buffer, UE_ARRAY_COUNT(Buffer)), 9);
 
 	return true;
diff --git a/Engine/Source/Runtime/Core/Private/Windows/WindowsCriticalSection.cpp b/Engine/Source/Runtime/Core/Private/Windows/WindowsCriticalSection.cpp
--- a/Engine/Source/Runtime/Core/Private/Windows/WindowsCriticalSection.cpp
+++ b/Engine/Source/Runtime/Core/Private/Windows/WindowsCriticalSection.cpp
@@ -20,7 +20,7 @@ FWindowsSystemWideCriticalSection::FWindowsSystemWideCriticalSection(const FStri
 	FString NormalizedMutexName(InName);
 	NormalizedMutexName.ReplaceInline(TEXT("\\"), TEXT("/"));
 
-	TCHAR MutexName[MAX_PATH] = TEXT("");
+	TCHAR MutexName[MAX_PATH] = {};
 	FCString::Strcpy(MutexName, *NormalizedMutexName);
 
 	// Attempt to create and take ownership of a named mutex
diff --git a/Engine/Source/Runtime/Core/Private/Windows/WindowsPlatformCrashContext.cpp b/Engine/Source/Runtime/Core/Private/Windows/WindowsPlatformCrashContext.cpp
--- a/Engine/Source/Runtime/Core/Private/Windows/WindowsPlatformCrashContext.cpp
+++ b/Engine/Source/Runtime/Core/Private/Windows/WindowsPlatformCrashContext.cpp
@@ -409,13 +409,13 @@ bool CreateCrashReportClientPath(TCHAR* OutClientPath, int32 MaxLength)
 	};
 
 #if WITH_EDITOR
-	const TCHAR CrashReportClientShippingName[] = TEXT("CrashReportClientEditor.exe");
-	const TCHAR CrashReportClientDevelopmentName[] = TEXT("CrashReportClientEditor-Win64-Development.exe");
-	const TCHAR CrashReportClientDebugName[] = TEXT("CrashReportClientEditor-Win64-Debug.exe");
+	const TCHAR* CrashReportClientShippingName = TEXT("CrashReportClientEditor.exe");
+	const TCHAR* CrashReportClientDevelopmentName = TEXT("CrashReportClientEditor-Win64-Development.exe");
+	const TCHAR* CrashReportClientDebugName = TEXT("CrashReportClientEditor-Win64-Debug.exe");
 #else
-	const TCHAR CrashReportClientShippingName[] = TEXT("CrashReportClient.exe");
-	const TCHAR CrashReportClientDevelopmentName[] = TEXT("CrashReportClient-Win64-Development.exe");
-	const TCHAR CrashReportClientDebugName[] = TEXT("CrashReportClient-Win64-Debug.exe");
+	const TCHAR* CrashReportClientShippingName = TEXT("CrashReportClient.exe");
+	const TCHAR* CrashReportClientDevelopmentName = TEXT("CrashReportClient-Win64-Development.exe");
+	const TCHAR* CrashReportClientDebugName = TEXT("CrashReportClient-Win64-Debug.exe");
 #endif
 
 	if (CreateCrashReportClientPathImpl(CrashReportClientShippingName))
diff --git a/Engine/Source/Runtime/Core/Private/Windows/WindowsPlatformMisc.cpp b/Engine/Source/Runtime/Core/Private/Windows/WindowsPlatformMisc.cpp
--- a/Engine/Source/Runtime/Core/Private/Windows/WindowsPlatformMisc.cpp
+++ b/Engine/Source/Runtime/Core/Private/Windows/WindowsPlatformMisc.cpp
@@ -773,7 +773,6 @@ static void HardKillIfAutomatedTesting()
 
 		UE_LOG(LogWindows, Warning, TEXT("Attempting to run KillAllPopUpBlockingWindows"));
 
-		TCHAR KillAllBlockingWindows[] = TEXT("KillAllPopUpBlockingWindows.bat");
 		// .bat files never seem to launch correctly with FPlatformProcess::CreateProc so we just use the FPlatformProcess::LaunchURL which will call ShellExecute
 		// we don't really care about the return code in this case 
 		FPlatformProcess::LaunchURL( TEXT("KillAllPopUpBlockingWindows.bat"), NULL, NULL );
diff --git a/Engine/Source/Runtime/Core/Private/Windows/WindowsPlatformProcess.cpp b/Engine/Source/Runtime/Core/Private/Windows/WindowsPlatformProcess.cpp
--- a/Engine/Source/Runtime/Core/Private/Windows/WindowsPlatformProcess.cpp
+++ b/Engine/Source/Runtime/Core/Private/Windows/WindowsPlatformProcess.cpp
@@ -1112,7 +1112,7 @@ const TCHAR* FWindowsPlatformProcess::ApplicationSettingsDir()
 
 const TCHAR* FWindowsPlatformProcess::ComputerName()
 {
-	static TCHAR Result[256]=TEXT("");
+	static TCHAR Result[256] = {};
 	if( !Result[0] )
 	{
 		uint32 Size=UE_ARRAY_COUNT(Result);
@@ -1123,8 +1123,8 @@ const TCHAR* FWindowsPlatformProcess::ComputerName()
 
 const TCHAR* FWindowsPlatformProcess::UserName(bool bOnlyAlphaNumeric/* = true*/)
 {
-	static TCHAR Result[256]=TEXT("");
-	static TCHAR ResultAlpha[256]=TEXT("");
+	static TCHAR Result[256] = {};
+	static TCHAR ResultAlpha[256] = {};
 	if( bOnlyAlphaNumeric )
 	{
 		if( !ResultAlpha[0] )
@@ -1202,7 +1202,7 @@ const FString FWindowsPlatformProcess::ShaderWorkingDir()
 
 const TCHAR* FWindowsPlatformProcess::ExecutablePath()
 {
-	static TCHAR Result[512]=TEXT("");
+	static TCHAR Result[512] = {};
 	if( !Result[0] )
 	{
 		if ( !GetModuleFileName( hInstance, Result, UE_ARRAY_COUNT(Result) ) )
@@ -1215,8 +1215,8 @@ const TCHAR* FWindowsPlatformProcess::ExecutablePath()
 
 const TCHAR* FWindowsPlatformProcess::ExecutableName(bool bRemoveExtension)
 {
-	static TCHAR Result[512]=TEXT("");
-	static TCHAR ResultWithExt[512]=TEXT("");
+	static TCHAR Result[512] = {};
+	static TCHAR ResultWithExt[512] = {};
 	if( !Result[0] )
 	{
 		// Get complete path for the executable
diff --git a/Engine/Source/Runtime/Core/Public/IO/IoDispatcher.h b/Engine/Source/Runtime/Core/Public/IO/IoDispatcher.h
--- a/Engine/Source/Runtime/Core/Public/IO/IoDispatcher.h
+++ b/Engine/Source/Runtime/Core/Public/IO/IoDispatcher.h
@@ -66,26 +66,9 @@ enum class EIoErrorCode
  */
 static const TCHAR* GetIoErrorText(EIoErrorCode ErrorCode)
 {
-	static constexpr const TCHAR* ErrorCodeText[]
-	{
-		TEXT("OK"),
-		TEXT("Unknown Status"),
-		TEXT("Invalid Code"),
-		TEXT("Cancelled"),
-		TEXT("FileOpen Failed"),
-		TEXT("File Not Open"),
-		TEXT("Read Error"),
-		TEXT("Write Error"),
-		TEXT("Not Found"),
-		TEXT("Corrupt Toc"),
-		TEXT("Unknown ChunkID"),
-		TEXT("Invalid Parameter"),
-		TEXT("Signature Error"),
-		TEXT("Invalid Encryption Key"),
-		TEXT("Compression Error")
-	};
+	extern CORE_API const TCHAR* const* GetIoErrorText_ErrorCodeText;
 
-	return ErrorCodeText[static_cast<uint32>(ErrorCode)];
+	return GetIoErrorText_ErrorCodeText[static_cast<uint32>(ErrorCode)];
 }
 
 /**
diff --git a/Engine/Source/Runtime/CoreUObject/Private/Serialization/Formatters/JsonArchiveInputFormatter.cpp b/Engine/Source/Runtime/CoreUObject/Private/Serialization/Formatters/JsonArchiveInputFormatter.cpp
--- a/Engine/Source/Runtime/CoreUObject/Private/Serialization/Formatters/JsonArchiveInputFormatter.cpp
+++ b/Engine/Source/Runtime/CoreUObject/Private/Serialization/Formatters/JsonArchiveInputFormatter.cpp
@@ -377,7 +377,7 @@ void FJsonArchiveInputFormatter::Serialize(FSoftObjectPtr& Value)
 void FJsonArchiveInputFormatter::Serialize(FSoftObjectPath& Value)
 {
 	FString StringValue;
-	const TCHAR Prefix[] = TEXT("Object:");
+	const auto& Prefix = TEXT("Object:");
 	if (ValueStack.Top()->TryGetString(StringValue) && StringValue.StartsWith(Prefix))
 	{
 		Value.SetPath(*StringValue + UE_ARRAY_COUNT(Prefix) - 1);
@@ -391,7 +391,7 @@ void FJsonArchiveInputFormatter::Serialize(FSoftObjectPath& Value)
 void FJsonArchiveInputFormatter::Serialize(FLazyObjectPtr& Value)
 {
 	FString StringValue;
-	const TCHAR Prefix[] = TEXT("Lazy:");
+	const auto& Prefix = TEXT("Lazy:");
 	if (ValueStack.Top()->TryGetString(StringValue) && StringValue.StartsWith(Prefix))
 	{
 		FUniqueObjectGuid Guid;
diff --git a/Engine/Source/Runtime/CoreUObject/Private/UObject/Obj.cpp b/Engine/Source/Runtime/CoreUObject/Private/UObject/Obj.cpp
--- a/Engine/Source/Runtime/CoreUObject/Private/UObject/Obj.cpp
+++ b/Engine/Source/Runtime/CoreUObject/Private/UObject/Obj.cpp
@@ -2726,7 +2726,7 @@ void UObject::SaveConfig( uint64 Flags, const TCHAR* InFilename, FConfigCacheIni
 			}
 			else
 			{
-				TCHAR TempKey[MAX_SPRINTF]=TEXT("");
+				TCHAR TempKey[MAX_SPRINTF] = {};
 				for( int32 Index=0; Index<Property->ArrayDim; Index++ )
 				{
 					if( Property->ArrayDim!=1 )
@@ -4339,7 +4339,7 @@ bool StaticExec( UWorld* InWorld, const TCHAR* Cmd, FOutputDevice& Ar )
 				// check if we want to ignore references from any packages
 				for( int32 i=0; i<16; i++ )
 				{
-					TCHAR Temp[MAX_SPRINTF]=TEXT("");
+					TCHAR Temp[MAX_SPRINTF] = {};
 					FCString::Sprintf( Temp, TEXT("EXCLUDE%i="), i );
 					FName F;
 					if (FParse::Value(Str, Temp, F))
diff --git a/Engine/Source/Runtime/CoreUObject/Private/UObject/Property.cpp b/Engine/Source/Runtime/CoreUObject/Private/UObject/Property.cpp
--- a/Engine/Source/Runtime/CoreUObject/Private/UObject/Property.cpp
+++ b/Engine/Source/Runtime/CoreUObject/Private/UObject/Property.cpp
@@ -991,7 +991,7 @@ void FProperty::ExportCppDeclaration(FOutputDevice& Out, EExportedDeclaration::T
 		NameCpp = FString(TEXT(", ")) + NameCpp;
 	}
 
-	TCHAR ArrayStr[MAX_SPRINTF]=TEXT("");
+	TCHAR ArrayStr[MAX_SPRINTF] = {};
 	const bool bExportStaticArray = 0 == (CPPF_NoStaticArray & AdditionalExportCPPFlags);
 	if ((ArrayDim != 1) && bExportStaticArray)
 	{
diff --git a/Engine/Source/Runtime/Projects/Private/PluginManager.cpp b/Engine/Source/Runtime/Projects/Private/PluginManager.cpp
--- a/Engine/Source/Runtime/Projects/Private/PluginManager.cpp
+++ b/Engine/Source/Runtime/Projects/Private/PluginManager.cpp
@@ -35,7 +35,7 @@ namespace PluginSystemDefs
 {
 	/** File extension of plugin descriptor files.
 	    NOTE: This constant exists in UnrealBuildTool code as well. */
-	static const TCHAR PluginDescriptorFileExtension[] = TEXT( ".uplugin" );
+	static const TCHAR* PluginDescriptorFileExtension = TEXT( ".uplugin" );
 
 	/**
 	 * Parsing the command line and loads any foreign plugins that were
