Added comment about recent code change.

diff --git a/Engine/Source/Runtime/Core/Private/ProfilingDebugging/Microsoft/WindowsCallstackTrace.cpp b/Engine/Source/Runtime/Core/Private/ProfilingDebugging/Microsoft/WindowsCallstackTrace.cpp
--- a/Engine/Source/Runtime/Core/Private/ProfilingDebugging/Microsoft/WindowsCallstackTrace.cpp
+++ b/Engine/Source/Runtime/Core/Private/ProfilingDebugging/Microsoft/WindowsCallstackTrace.cpp
@@ -648,6 +648,8 @@ uint32 FBacktracer::GetBacktraceId(void* AddressOfReturnAddress) const
 
 #if BACKTRACE_LOCK_FREE
 		{
+			// This conversion improves probing performance for the hash set. Additionally it is critical 
+			// to avoid incorrect values when RspBias is compressed into 16 bits in the hash map.
 			int32 StoreBias = Function->RspBias < 0 ? -1 : Function->RspBias;
 			TGuardValue<bool> ReentranceGuard(bReentranceCheck, true);
 			FunctionLookups.Emplace(RetAddr, StoreBias);
