VectorVM fix for arm64 Neon
-Need return type of MakeVectorRegisterInt64 for Neon to match SSE.

diff --git a/Engine/Source/Runtime/Core/Public/Math/UnrealMathNeon.h b/Engine/Source/Runtime/Core/Public/Math/UnrealMathNeon.h
--- a/Engine/Source/Runtime/Core/Public/Math/UnrealMathNeon.h
+++ b/Engine/Source/Runtime/Core/Public/Math/UnrealMathNeon.h
@@ -373,11 +373,11 @@ FORCEINLINE VectorRegister4Int MakeVectorRegisterInt(int32 X, int32 Y, int32 Z,
 	return Tmp.V;
 }
 
-FORCEINLINE VectorRegister2Int64 MakeVectorRegisterInt64(int64 X, int64 Y)
+FORCEINLINE VectorRegister4Int MakeVectorRegisterInt64(int64 X, int64 Y)
 {
 	union U
 	{
-		VectorRegister2Int64 V; int64 I[2];
+		VectorRegister4Int V; int64 I[2];
 		FORCEINLINE U() : V() {}
 	} Tmp;
 	Tmp.I[0] = X;
