Added additional detail to the TRotator class to make it clear in what order the roll, pitch, and yaw are applied.

diff --git a/Engine/Source/Runtime/Core/Public/Math/Rotator.h b/Engine/Source/Runtime/Core/Public/Math/Rotator.h
--- a/Engine/Source/Runtime/Core/Public/Math/Rotator.h
+++ b/Engine/Source/Runtime/Core/Public/Math/Rotator.h
@@ -21,6 +21,11 @@ namespace Math
  * Implements a container for rotation information.
  *
  * All rotation values are stored in degrees.
+ *
+ * The angles are interpreted as intrinsic rotations applied in the order Roll, then Pitch, then Yaw. I.e., an object would be rotated
+ * first by the specified roll around its forward axis, then pitched around its (new) right axis, and then finally yawed around its 
+ * (new) up axis.
+ * 
  */
 template<typename T>
 struct TRotator
