When building name patches, compare names case sensitively. (Sorting was already done case sensitively.)
Removing ensure from FName about unfreezing names that were created for a case-insensitive target. We rely on the archive used for serialization to save names with their case intact.

diff --git a/Engine/Source/Runtime/Core/Private/Serialization/MemoryImage.cpp b/Engine/Source/Runtime/Core/Private/Serialization/MemoryImage.cpp
--- a/Engine/Source/Runtime/Core/Private/Serialization/MemoryImage.cpp
+++ b/Engine/Source/Runtime/Core/Private/Serialization/MemoryImage.cpp
@@ -1360,7 +1360,7 @@ static void CountNumNames(const TArray<FMemoryImageNamePointer>& Names, TArray<u
 	uint32 CurrentNumPatches = 0u;
 	for (const FMemoryImageNamePointer& Patch : Names)
 	{
-		if (Patch.Name != CurrentName)
+		if (Patch.Name.Compare(CurrentName) != 0)
 		{
 			if (CurrentNumPatches > 0u)
 			{
@@ -1389,7 +1389,7 @@ static void SerializeNames(const TArray<FMemoryImageNamePointer>& Names, const T
 		for (uint32 i = 0; i < Num; ++i)
 		{
 			const FMemoryImageNamePointer& Patch = Names[NameIndex++];
-			check(Patch.Name == Name);
+			check(Patch.Name.Compare(Name) == 0);
 
 			uint32 Offset = Patch.Offset;
 			Ar << Offset;
@@ -1774,7 +1774,7 @@ static void UpdateHashNamePatches(const TArray< FMemoryImageNamePointer>& Names,
 {
 	for (const FMemoryImageNamePointer& NamePatch : Names)
 	{
-		const FNameEntry* NameEntry = NamePatch.Name.GetComparisonNameEntry();
+		const FNameEntry* NameEntry = NamePatch.Name.GetDisplayNameEntry();
 		TCHAR NameBuffer[NAME_SIZE];
 		NameEntry->GetName(NameBuffer);
 		const int32 Number = NamePatch.Name.GetNumber();
diff --git a/Engine/Source/Runtime/Core/Private/UObject/UnrealNames.cpp b/Engine/Source/Runtime/Core/Private/UObject/UnrealNames.cpp
--- a/Engine/Source/Runtime/Core/Private/UObject/UnrealNames.cpp
+++ b/Engine/Source/Runtime/Core/Private/UObject/UnrealNames.cpp
@@ -5037,8 +5037,6 @@ uint32 Freeze::IntrinsicUnfrozenCopy(const FMemoryUnfreezeContent& Context, cons
 	{
 		if (!TargetLayoutParameters.WithCasePreservingFName())
 		{
-			ensureMsgf(!WITH_CASE_PRESERVING_NAME, TEXT("Unfreezing memory images containing non-case preserving names when WITH_CASE_PRESERVING_NAME is true can lead to casing cook determinism."));
-
 			const TMemoryImageNameLayout<0>& Loaded = reinterpret_cast<const TMemoryImageNameLayout<0>&>(Object);
 			FMemoryImageName* Dest = new(OutDst) FMemoryImageName();
 			Dest->ComparisonIndex = Loaded.ComparisonIndex;
diff --git a/Engine/Source/Runtime/Core/Public/Serialization/MemoryImage.h b/Engine/Source/Runtime/Core/Public/Serialization/MemoryImage.h
--- a/Engine/Source/Runtime/Core/Public/Serialization/MemoryImage.h
+++ b/Engine/Source/Runtime/Core/Public/Serialization/MemoryImage.h
@@ -162,7 +162,7 @@ struct FMemoryImageNamePointer
 };
 inline bool operator==(const FMemoryImageNamePointer& Lhs, const FMemoryImageNamePointer& Rhs)
 {
-	return Lhs.Offset == Rhs.Offset && Lhs.Name == Rhs.Name;
+	return Lhs.Offset == Rhs.Offset && Lhs.Name.Compare(Rhs.Name) == 0;
 }
 inline bool operator!=(const FMemoryImageNamePointer& Lhs, const FMemoryImageNamePointer& Rhs)
 {
