ApplePlatformFile: Add string descriptor for errno value.
#rnx

diff --git a/Engine/Source/Runtime/Core/Private/Apple/ApplePlatformFile.cpp b/Engine/Source/Runtime/Core/Private/Apple/ApplePlatformFile.cpp
--- a/Engine/Source/Runtime/Core/Private/Apple/ApplePlatformFile.cpp
+++ b/Engine/Source/Runtime/Core/Private/Apple/ApplePlatformFile.cpp
@@ -95,7 +95,8 @@ public:
 				else
 				{
 					TRACE_PLATFORMFILE_FAIL_CLOSE(FileHandle);
-					UE_LOG(LogInit, Warning, TEXT("Failed to properly close readable file: %s with errno: %d"), *Filename, errno);
+					UE_LOG(LogInit, Warning, TEXT("Failed to properly close readable file: %s with errno: %d: %s"),
+						*Filename, errno, UTF8_TO_TCHAR(strerror(errno)));
 				}
 				ActiveHandles[ HandleSlot ] = nullptr;
 			}
@@ -108,7 +109,8 @@ public:
                 int Result = fsync(FileHandle);
 				if (Result < 0)
 				{
-					UE_LOG(LogInit, Error, TEXT("Failed to properly flush writable file with errno: %d"), errno);
+					UE_LOG(LogInit, Error, TEXT("Failed to properly flush writable file: %s with errno: %d: %s"),
+                           *Filename, errno, UTF8_TO_TCHAR(strerror(errno)));
 				}
             }
 			TRACE_PLATFORMFILE_BEGIN_CLOSE(FileHandle);
@@ -121,7 +123,8 @@ public:
 			else
 			{
 				TRACE_PLATFORMFILE_FAIL_CLOSE(FileHandle);
-				UE_LOG(LogInit, Warning, TEXT("Failed to properly close file with errno: %d"), errno);
+				UE_LOG(LogInit, Warning, TEXT("Failed to properly close file: %s with errno: %d: %s"),
+                       *Filename, errno, UTF8_TO_TCHAR(strerror(errno)));
 			}
 		}
 		FileHandle = -1;
