Disable memory mapped IO for Animation on iOS to prevent crash with encrypted paks

diff --git a/Engine/Source/Runtime/Core/Public/IOS/IOSPlatformProperties.h b/Engine/Source/Runtime/Core/Public/IOS/IOSPlatformProperties.h
--- a/Engine/Source/Runtime/Core/Public/IOS/IOSPlatformProperties.h
+++ b/Engine/Source/Runtime/Core/Public/IOS/IOSPlatformProperties.h
@@ -78,7 +78,7 @@ struct FIOSPlatformProperties
 	}
 	static FORCEINLINE bool SupportsMemoryMappedAnimation()
 	{
-		return true;
+		return false;
 	}
 	static FORCEINLINE int64 GetMemoryMappingAlignment()
 	{
