Fix FLinearColor::ToFColorSRGB handling of large A, Quantize handling of large R/G/B/A.

For large enough floats x (e.g. 1e+7f), x*255 is too large to convert to int. On x86 this turns into INT_MIN which then gets clamped to 0, on ARM it gets saturated and turns into INT_MAX which gets clamped to 255.

Fix it so that over-large values result in 255 (which is the more natural behavior) everywhere and note the conversion rules used for both the sRGB and the linear channels while I'm here. Also update ImageCore optimized kernels to include the QuantizeRound fix.

Confirmed to produce bitwise identical results for textures on AncientGame (i.e. nop on typical content).

diff --git a/Engine/Source/Runtime/Core/Private/Math/Color.cpp b/Engine/Source/Runtime/Core/Private/Math/Color.cpp
--- a/Engine/Source/Runtime/Core/Private/Math/Color.cpp
+++ b/Engine/Source/Runtime/Core/Private/Math/Color.cpp
@@ -88,9 +88,14 @@ FLinearColor FLinearColor::FromPow22Color(const FColor& Color)
 	return LinearColor;
 }
 
-static inline uint8 ClampU8(int Value)
+// For pixel format conversions. Clamps to [0,1], but for consistency with
+// GPU conversions, also define behavior for NaNs and make them map to 0.
+static inline float Clamp01(float Value)
 {
-	return (uint8) FMath::Clamp(Value,0,255);
+	// Write this explicitly instead of using FMath::Clamp so we have
+	// control over what happens with NaNs.
+	const float ClampedLo = (Value > 0.0f) ? Value : 0.0f; // Also turns NaNs into 0.
+	return (ClampedLo < 1.0f) ? ClampedLo : 1.0f;
 }
 
 /**
@@ -190,12 +195,15 @@ static uint8 stbir__linear_to_srgb_uchar_fast(float in)
 /** Quantizes the linear color and returns the result as a FColor with optional sRGB conversion and quality as goal. */
 FColor FLinearColor::ToFColorSRGB() const
 {
+	// The convention used here in all channels is that NaNs
+	// convert to 0, as do negative values, and out-of-range
+	// positive values convert to 255.
 	return FColor(
 		stbir__linear_to_srgb_uchar_fast(R),
 		stbir__linear_to_srgb_uchar_fast(G),
 		stbir__linear_to_srgb_uchar_fast(B),
-		ClampU8( (int)(0.5f + A*255.f) )
-		);
+		(uint8)(0.5f + Clamp01(A)*255.f)
+	);
 }
 
 
@@ -213,10 +221,10 @@ FColor FLinearColor::QuantizeFloor() const
 	// replicates behavior of old Quantize() which is deprecated
 	// restoration should be done to bucket centers (+0.5 on dequantize)
 	return FColor(
-		ClampU8( (int)(R*255.f) ),
-		ClampU8( (int)(G*255.f) ),
-		ClampU8( (int)(B*255.f) ),
-		ClampU8( (int)(A*255.f) )
+		(uint8)(Clamp01(R)*255.f),
+		(uint8)(Clamp01(G)*255.f),
+		(uint8)(Clamp01(B)*255.f),
+		(uint8)(Clamp01(A)*255.f)
 		);
 }
 
@@ -224,10 +232,10 @@ FColor FLinearColor::QuantizeRound() const
 {
 	// do not use FMath::RoundToInt because it call floorf
 	return FColor(
-		ClampU8( (int)(0.5f + R*255.f) ),
-		ClampU8( (int)(0.5f + G*255.f) ),
-		ClampU8( (int)(0.5f + B*255.f) ),
-		ClampU8( (int)(0.5f + A*255.f) )
+		(uint8)(0.5f + Clamp01(R)*255.f),
+		(uint8)(0.5f + Clamp01(G)*255.f),
+		(uint8)(0.5f + Clamp01(B)*255.f),
+		(uint8)(0.5f + Clamp01(A)*255.f)
 		);
 }
 
diff --git a/Engine/Source/Runtime/ImageCore/Private/ImageCore.cpp b/Engine/Source/Runtime/ImageCore/Private/ImageCore.cpp
--- a/Engine/Source/Runtime/ImageCore/Private/ImageCore.cpp
+++ b/Engine/Source/Runtime/ImageCore/Private/ImageCore.cpp
@@ -288,7 +288,7 @@ IMAGECORE_API void FImageCore::CopyImage(const FImageView & SrcImage,const FImag
 							Pixel2 = _mm_shuffle_ps(Pixel2, Pixel2, _MM_SHUFFLE(3, 0, 1, 2));
 							Pixel3 = _mm_shuffle_ps(Pixel3, Pixel3, _MM_SHUFFLE(3, 0, 1, 2));
 
-							// scale to [0,255]
+							// Scale to [0,255]
 							__m128 Mul = _mm_set_ps1(255.f);
 							__m128 Add = _mm_set_ps1(0.5f);
 							Pixel0 = _mm_add_ps(_mm_mul_ps(Pixel0, Mul), Add);
@@ -296,7 +296,17 @@ IMAGECORE_API void FImageCore::CopyImage(const FImageView & SrcImage,const FImag
 							Pixel2 = _mm_add_ps(_mm_mul_ps(Pixel2, Mul), Add);
 							Pixel3 = _mm_add_ps(_mm_mul_ps(Pixel3, Mul), Add);
 
-							// cast float to 32-bit components
+							// Clamp large values at 255. Must be first arg: SSE min_ps(255, NaN) = NaN,
+							// but min_ps(NaN, 255) = 255, and we want the NaNs to turn into 0 not 255,
+							// for consistency with QuantizeRound.
+							Pixel0 = _mm_min_ps(Mul, Pixel0);
+							Pixel1 = _mm_min_ps(Mul, Pixel1);
+							Pixel2 = _mm_min_ps(Mul, Pixel2);
+							Pixel3 = _mm_min_ps(Mul, Pixel3);
+
+							// Convert float to 32-bit integer
+							// values are <=255 so no overflow; NaNs and too-large negative values
+							// convert to INT_MIN, all saturate to 0 during pack (no special clamp necessary)
 							__m128i Pixel0i = _mm_cvttps_epi32(Pixel0);
 							__m128i Pixel1i = _mm_cvttps_epi32(Pixel1);
 							__m128i Pixel2i = _mm_cvttps_epi32(Pixel2);
