fix for hang in FOutputDeviceFile::Flush during crashes.

#rnx

diff --git a/Engine/Source/Runtime/Core/Public/HAL/IPlatformFileManagedStorageWrapper.h b/Engine/Source/Runtime/Core/Public/HAL/IPlatformFileManagedStorageWrapper.h
--- a/Engine/Source/Runtime/Core/Public/HAL/IPlatformFileManagedStorageWrapper.h
+++ b/Engine/Source/Runtime/Core/Public/HAL/IPlatformFileManagedStorageWrapper.h
@@ -340,7 +340,8 @@ public:
 		// FPersistentStorageManager depends on FPaths which depends on the command line being initialized.
 		// FPersistentStorageManager depends on GConfig.
 		// FPersistentStorageManager can't be constructed until its dependencies are ready
-		return GConfig && GConfig->IsReadyForUse() && FCommandLine::IsInitialized();
+		// FPersistentStorageManager will try and allocate memory during a crash but this could hang during log file flushing
+		return GConfig && GConfig->IsReadyForUse() && FCommandLine::IsInitialized() && !GIsCriticalError;
 	}
 
 	/** Singleton access **/
