LowLevelMemTracker: Add a LLM_TAG_NAME macro so that name-based-tags from LLM_DEFINE_TAG can be passed as the default tag into OnLowLevelAlloc.
#rnx

diff --git a/Engine/Source/Runtime/Core/Public/HAL/LowLevelMemTracker.h b/Engine/Source/Runtime/Core/Public/HAL/LowLevelMemTracker.h
--- a/Engine/Source/Runtime/Core/Public/HAL/LowLevelMemTracker.h
+++ b/Engine/Source/Runtime/Core/Public/HAL/LowLevelMemTracker.h
@@ -352,6 +352,9 @@ extern FName LLMGetTagStat(ELLMTag Tag);
 #define LLM_DECLARE_TAG(UniqueNameWithUnderscores) extern FLLMTagDeclaration PREPROCESSOR_JOIN(LLMTagDeclaration_, UniqueNameWithUnderscores)
 #define LLM_DECLARE_TAG_API(UniqueNameWithUnderscores, ModuleAPI) extern ModuleAPI FLLMTagDeclaration PREPROCESSOR_JOIN(LLMTagDeclaration_, UniqueNameWithUnderscores)
 
+/** Get the unique Name of a Tag, suitable for passing to LLM functions such as OnLowLevelAlloc that take a Tag UniqueName. */
+#define LLM_TAG_NAME(UniqueNameWithUnderscores) (PREPROCESSOR_JOIN(LLMTagDeclaration_, UniqueNameWithUnderscores).GetUniqueName())
+
 /**
  * The BootStrap versions of LLM_DEFINE_TAG, LLM_SCOPE_BYTAG, and LLM_DECLARE_TAG support use in scopes during global
  * c++ constructors, before Main. These tags are slightly more expensive (even when LLM is disabled) in all
