FString and string builder appending fixes for UTF-8 mode.

diff --git a/Engine/Source/Runtime/Core/Public/Containers/UnrealString.h b/Engine/Source/Runtime/Core/Public/Containers/UnrealString.h
--- a/Engine/Source/Runtime/Core/Public/Containers/UnrealString.h
+++ b/Engine/Source/Runtime/Core/Public/Containers/UnrealString.h
@@ -398,12 +398,28 @@ public:
 
 	/** Append a string and return a reference to this */
 	template <typename StrType>
-	FORCEINLINE FString& operator+=(StrType&& Str)	{ return Append(Forward<StrType>(Str)); }
+	FORCEINLINE auto operator+=(StrType&& Str) -> decltype(Append(Forward<StrType>(Str)))
+	{
+		return Append(Forward<StrType>(Str));
+	}
 
 	/** Append a single character and return a reference to this */
-	FORCEINLINE FString& operator+=(ANSICHAR Char)	{ return AppendChar(Char); }
-	FORCEINLINE FString& operator+=(WIDECHAR Char)	{ return AppendChar(Char); }
-	FORCEINLINE FString& operator+=(UCS2CHAR Char)	{ return AppendChar(Char); }
+	template <
+		typename AppendedCharType,
+		std::enable_if_t<TIsCharType<AppendedCharType>::Value>* = nullptr
+	>
+	FORCEINLINE FString& operator+=(AppendedCharType Char)
+	{
+		if constexpr (FPlatformString::IsCharEncodingSimplyConvertibleTo<AppendedCharType, TCHAR>())
+		{
+			return AppendChar((TCHAR)Char);
+		}
+		else
+		{
+			AppendChars(&Char, 1);
+			return *this;
+		}
+	}
 
 	void InsertAt(int32 Index, TCHAR Character);
 	void InsertAt(int32 Index, const FString& Characters);
@@ -1901,9 +1917,9 @@ UE_NODISCARD inline TCHAR NibbleToTChar(uint8 Num)
 {
 	if (Num > 9)
 	{
-		return TEXT('A') + TCHAR(Num - 10);
+		return (TCHAR)('A' + (Num - 10));
 	}
-	return TEXT('0') + TCHAR(Num);
+	return (TCHAR)('0' + Num);
 }
 
 /** 
@@ -2226,7 +2242,7 @@ public:
 	 */
 	virtual FStringOutputDeviceCountLines& operator+=(const FStringOutputDeviceCountLines& Other)
 	{
-		FString::operator+=(Other);
+		FString::operator+=(static_cast<const FString&>(Other));
 
 		LineCount += Other.GetLineCount();
 
diff --git a/Engine/Source/Runtime/Core/Public/Misc/StringBuilder.h b/Engine/Source/Runtime/Core/Public/Misc/StringBuilder.h
--- a/Engine/Source/Runtime/Core/Public/Misc/StringBuilder.h
+++ b/Engine/Source/Runtime/Core/Public/Misc/StringBuilder.h
@@ -159,17 +159,37 @@ public:
 		return Append(View.GetData(), View.Len());
 	}
 
-	inline BuilderType& AppendChar(CharType Char)
+	template <
+		typename AppendedCharType,
+		std::enable_if_t<TIsCharType<AppendedCharType>::Value>* = nullptr
+	>
+	inline BuilderType& AppendChar(AppendedCharType Char)
+	{
+		if constexpr (FPlatformString::IsCharEncodingSimplyConvertibleTo<AppendedCharType, CharType>())
 	{
 		EnsureAdditionalCapacity(1);
-		*CurPos++ = Char;
+			*CurPos++ = (CharType)Char;
+		}
+		else
+		{
+			int32 ConvertedLength = FPlatformString::ConvertedLength<CharType>(&Char, 1);
+			EnsureAdditionalCapacity(ConvertedLength);
+			CurPos = FPlatformString::Convert(CurPos, ConvertedLength, &Char, 1);
+		}
+
 		return *this;
 	}
 
+	template <
+		typename AppendedCharType,
+		std::enable_if_t<TIsCharType<AppendedCharType>::Value>* = nullptr
+	>
 	UE_DEPRECATED(5.0, "Use AppendChar instead of Append.")
-	inline BuilderType& Append(CharType Char) { return AppendChar(Char); }
+	inline BuilderType& Append(AppendedCharType Char)
+	{
+		return AppendChar(Char);
+	}
 
-	UE_DEPRECATED(5.0, "Use Append instead of AppendAnsi.")
 	inline BuilderType& AppendAnsi(const FAnsiStringView String) { return Append(String); }
 	UE_DEPRECATED(5.0, "Use Append instead of AppendAnsi.")
 	inline BuilderType& AppendAnsi(const ANSICHAR* const String) { return Append(String); }
@@ -183,7 +203,7 @@ public:
 		check(RemoveLen >= 0);
 		check(Pos + RemoveLen <= Len());
 
-		const int DeltaLen = Str.Len() - RemoveLen;
+		const int DeltaLen = Str.Len() - RemoveLen;		
 		if (DeltaLen < 0)
 		{
 			CurPos += DeltaLen;
@@ -203,7 +223,7 @@ public:
 				*It = *(It - DeltaLen);
 			}
 		}
-
+		
 		FMemory::Memcpy(Base + Pos, Str.GetData(), Str.Len() * sizeof(CharType));
 	}
 
@@ -403,6 +423,7 @@ inline FUtf8StringBuilderBase&		operator<<(FUtf8StringBuilderBase& Builder, WIDE
 // Prefer using << instead of += as operator+= is only intended for mechanical FString -> FStringView replacement.
 inline FStringBuilderBase&			operator+=(FStringBuilderBase& Builder, ANSICHAR Char)								{ return Builder.AppendChar(Char); }
 inline FStringBuilderBase&			operator+=(FStringBuilderBase& Builder, WIDECHAR Char)								{ return Builder.AppendChar(Char); }
+inline FStringBuilderBase&			operator+=(FStringBuilderBase& Builder, UTF8CHAR Char)								{ return Builder.AppendChar(Char); }
 inline FStringBuilderBase&			operator+=(FStringBuilderBase& Builder, FWideStringView Str)						{ return Builder.Append(Str); }
 inline FStringBuilderBase&			operator+=(FStringBuilderBase& Builder, FUtf8StringView Str)						{ return Builder.Append(Str); }
 
@@ -410,15 +431,15 @@ inline FStringBuilderBase&			operator+=(FStringBuilderBase& Builder, FUtf8String
 
 inline FAnsiStringBuilderBase&		operator<<(FAnsiStringBuilderBase& Builder, int32 Value)							{ return Builder.Appendf("%d", Value); }
 inline FAnsiStringBuilderBase&		operator<<(FAnsiStringBuilderBase& Builder, uint32 Value)							{ return Builder.Appendf("%u", Value); }
-inline FWideStringBuilderBase&		operator<<(FWideStringBuilderBase& Builder, int32 Value)							{ return Builder.Appendf(TEXT("%d"), Value); }
-inline FWideStringBuilderBase&		operator<<(FWideStringBuilderBase& Builder, uint32 Value)							{ return Builder.Appendf(TEXT("%u"), Value); }
+inline FWideStringBuilderBase&		operator<<(FWideStringBuilderBase& Builder, int32 Value)							{ return Builder.Appendf(WIDETEXT("%d"), Value); }
+inline FWideStringBuilderBase&		operator<<(FWideStringBuilderBase& Builder, uint32 Value)							{ return Builder.Appendf(WIDETEXT("%u"), Value); }
 inline FUtf8StringBuilderBase&		operator<<(FUtf8StringBuilderBase& Builder, int32 Value)							{ return Builder.Appendf(UTF8TEXT("%d"), Value); }
 inline FUtf8StringBuilderBase&		operator<<(FUtf8StringBuilderBase& Builder, uint32 Value)							{ return Builder.Appendf(UTF8TEXT("%u"), Value); }
 
 inline FAnsiStringBuilderBase&		operator<<(FAnsiStringBuilderBase& Builder, int64 Value)							{ return Builder.Appendf("%" INT64_FMT, Value); }
 inline FAnsiStringBuilderBase&		operator<<(FAnsiStringBuilderBase& Builder, uint64 Value)							{ return Builder.Appendf("%" UINT64_FMT, Value); }
-inline FWideStringBuilderBase&		operator<<(FWideStringBuilderBase& Builder, int64 Value)							{ return Builder.Appendf(TEXT("%" INT64_FMT), Value); }
-inline FWideStringBuilderBase&		operator<<(FWideStringBuilderBase& Builder, uint64 Value)							{ return Builder.Appendf(TEXT("%" UINT64_FMT), Value); }
+inline FWideStringBuilderBase&		operator<<(FWideStringBuilderBase& Builder, int64 Value)							{ return Builder.Appendf(WIDETEXT("%" INT64_FMT), Value); }
+inline FWideStringBuilderBase&		operator<<(FWideStringBuilderBase& Builder, uint64 Value)							{ return Builder.Appendf(WIDETEXT("%" UINT64_FMT), Value); }
 inline FUtf8StringBuilderBase&		operator<<(FUtf8StringBuilderBase& Builder, int64 Value)							{ return Builder.Appendf(UTF8TEXT("%" INT64_FMT), Value); }
 inline FUtf8StringBuilderBase&		operator<<(FUtf8StringBuilderBase& Builder, uint64 Value)							{ return Builder.Appendf(UTF8TEXT("%" UINT64_FMT), Value); }
 
