Logging: Allocate the thread wake event on-demand to avoid asserts during static initialization on certain platforms

#rnx

diff --git a/Engine/Source/Runtime/Core/Private/Misc/OutputDeviceRedirector.cpp b/Engine/Source/Runtime/Core/Private/Misc/OutputDeviceRedirector.cpp
--- a/Engine/Source/Runtime/Core/Private/Misc/OutputDeviceRedirector.cpp
+++ b/Engine/Source/Runtime/Core/Private/Misc/OutputDeviceRedirector.cpp
@@ -134,7 +134,7 @@ struct FOutputDeviceRedirectorState
 	FRWLock ThreadLock;
 
 	/** An event to wake the dedicated master thread to process buffered lines. */
-	FEventRef ThreadWakeEvent{EEventMode::AutoReset};
+	std::atomic<FEvent*> ThreadWakeEvent{nullptr};
 
 	/** A queue of events to trigger when the dedicated master thread is idle. */
 	TDepletableMpscQueue<FEvent*, FOutputDeviceLinearAllocator> ThreadIdleEvents;
@@ -142,9 +142,6 @@ struct FOutputDeviceRedirectorState
 	/** The ID of the master thread. Logging from other threads will be buffered for processing by the master thread. */
 	std::atomic<uint32> MasterThreadId{FPlatformTLS::GetCurrentThreadId()};
 
-	/** Whether to keep the thread running. */
-	std::atomic<bool> bThreadActive{false};
-
 	/** Whether the backlog is enabled. */
 	bool bEnableBacklog = false;
 
@@ -272,16 +269,19 @@ void FOutputDeviceRedirectorState::ThreadLoop()
 		MasterThreadId.store(ThreadId, std::memory_order_relaxed);
 	}
 
-	while (IsMasterThread(ThreadId))
+	if (FEvent* WakeEvent = ThreadWakeEvent.load(std::memory_order_acquire))
 	{
-		ThreadWakeEvent->Wait();
-		do
+		while (IsMasterThread(ThreadId))
 		{
-			FOutputDevicesMasterScopeLock Lock(*this);
-			FlushBufferedLines();
+			WakeEvent->Wait();
+			do
+			{
+				FOutputDevicesMasterScopeLock Lock(*this);
+				FlushBufferedLines();
+			}
+			while (!BufferedLines.IsEmpty());
+			ThreadIdleEvents.Deplete([](FEvent* Event) { Event->Trigger(); });
 		}
-		while (!BufferedLines.IsEmpty());
-		ThreadIdleEvents.Deplete([](FEvent* Event) { Event->Trigger(); });
 	}
 }
 
@@ -355,14 +355,14 @@ bool FOutputDeviceRedirector::IsRedirectingTo(FOutputDevice* OutputDevice)
 
 void FOutputDeviceRedirector::FlushThreadedLogs(EOutputDeviceRedirectorFlushOptions Options)
 {
-	if (FReadScopeLock Lock(State->ThreadLock); State->bThreadActive.load(std::memory_order_relaxed))
+	if (FReadScopeLock ThreadLock(State->ThreadLock); FEvent* WakeEvent = State->ThreadWakeEvent.load(std::memory_order_acquire))
 	{
 		if (!EnumHasAnyFlags(Options, EOutputDeviceRedirectorFlushOptions::Async))
 		{
 			FEventRef IdleEvent(EEventMode::ManualReset);
 			if (State->ThreadIdleEvents.EnqueueAndReturnWasEmpty(IdleEvent.Get()))
 			{
-				State->ThreadWakeEvent->Trigger();
+				WakeEvent->Trigger();
 			}
 			IdleEvent->Wait();
 		}
@@ -411,11 +411,12 @@ void FOutputDeviceRedirector::SetCurrentThreadAsMasterThread()
 		State->FlushBufferedLines();
 	}
 
-	if (FWriteScopeLock Lock(State->ThreadLock); State->bThreadActive.load(std::memory_order_relaxed))
+	if (FWriteScopeLock ThreadLock(State->ThreadLock); FEvent* WakeEvent = State->ThreadWakeEvent.exchange(nullptr, std::memory_order_acquire))
 	{
-		State->bThreadActive.store(false, std::memory_order_relaxed);
-		State->ThreadWakeEvent->Trigger();
+		WakeEvent->Trigger();
 		State->Thread.Join();
+		UE::Private::FOutputDevicesWriteScopeLock Lock(*State);
+		FPlatformProcess::ReturnSynchEventToPool(WakeEvent);
 	}
 }
 
@@ -426,10 +427,11 @@ bool FOutputDeviceRedirector::TryStartDedicatedMasterThread()
 		return false;
 	}
 
-	if (FWriteScopeLock Lock(State->ThreadLock); !State->bThreadActive.load(std::memory_order_relaxed))
+	if (FWriteScopeLock ThreadLock(State->ThreadLock); !State->ThreadWakeEvent.load(std::memory_order_relaxed))
 	{
-		State->bThreadActive.store(true, std::memory_order_relaxed);
-		State->ThreadWakeEvent->Trigger();
+		FEvent* WakeEvent = FPlatformProcess::GetSynchEventFromPool();
+		WakeEvent->Trigger();
+		State->ThreadWakeEvent.store(WakeEvent, std::memory_order_release);
 		State->Thread = FThread(TEXT("OutputDeviceRedirector"), [State = State.Get()] { State->ThreadLoop(); });
 	}
 	return true;
@@ -469,10 +471,12 @@ void FOutputDeviceRedirector::Serialize(const TCHAR* const Data, const ELogVerbo
 
 	const auto EnqueueLine = [this, Data, Category, Verbosity, RealTime]
 	{
-		if (State->BufferedLines.EnqueueAndReturnWasEmpty(Data, Category, Verbosity, RealTime) &&
-			State->bThreadActive.load(std::memory_order_relaxed))
+		if (State->BufferedLines.EnqueueAndReturnWasEmpty(Data, Category, Verbosity, RealTime))
 		{
-			State->ThreadWakeEvent->Trigger();
+			if (FEvent* WakeEvent = State->ThreadWakeEvent.load(std::memory_order_acquire))
+			{
+				WakeEvent->Trigger();
+			}
 		}
 	};
 
