Fixed variables named X64 to avoid collision with preprocessor definition X64.

PR #8989: Minor change: Avoid name collisions with common preprocessor definition (X64) (Contributed by [REDACTED_PERSON])

diff --git a/Engine/Source/Runtime/Core/Public/Math/IntPoint.h b/Engine/Source/Runtime/Core/Public/Math/IntPoint.h
--- a/Engine/Source/Runtime/Core/Public/Math/IntPoint.h
+++ b/Engine/Source/Runtime/Core/Public/Math/IntPoint.h
@@ -573,9 +573,9 @@ FORCEINLINE uint32 GetTypeHash(const FIntPoint& InPoint)
 
 FORCEINLINE int32 FIntPoint::Size() const
 {
-	int64 X64 = (int64)X;
-	int64 Y64 = (int64)Y;
-	return int32(FMath::Sqrt(float(X64 * X64 + Y64 * Y64)));
+	int64 LocalX64 = (int64)X;
+	int64 LocalY64 = (int64)Y;
+	return int32(FMath::Sqrt(float(LocalX64 * LocalX64 + LocalY64 * LocalY64)));
 }
 
 FORCEINLINE int32 FIntPoint::SizeSquared() const
diff --git a/Engine/Source/Runtime/Core/Public/Math/IntVector.h b/Engine/Source/Runtime/Core/Public/Math/IntVector.h
--- a/Engine/Source/Runtime/Core/Public/Math/IntVector.h
+++ b/Engine/Source/Runtime/Core/Public/Math/IntVector.h
@@ -522,10 +522,10 @@ FORCEINLINE int32 FIntVector::Num()
 
 FORCEINLINE int32 FIntVector::Size() const
 {
-	int64 X64 = (int64)X;
-	int64 Y64 = (int64)Y;
-	int64 Z64 = (int64)Z;
-	return int32(FMath::Sqrt(float(X64 * X64 + Y64 * Y64 + Z64 * Z64)));
+	int64 LocalX64 = (int64)X;
+	int64 LocalY64 = (int64)Y;
+	int64 LocalZ64 = (int64)Z;
+	return int32(FMath::Sqrt(float(LocalX64 * LocalX64 + LocalY64 * LocalY64 + LocalZ64 * LocalZ64)));
 }
 
 FORCEINLINE bool FIntVector::IsZero() const
