FLinearColor conversion optimizations and some fixes.

FLinearColor ctor from FColor is mainly three table lookups, make it inline and move to the header.
Quantize, QuantizeFloor, QuantizeRound are similarly trivial, now inline.
Move Clamp01 to FColor as static method so it's namespaced; rename to Clamp01NansTo0 to explain why it's not just FMath::Clamp to [0,1]
SSE2-ified FLinearColor::ToFColorSRGB for x86 targets, bitwise exact same results as before for all input floats (tested exhaustively, test included)
Don't want to make ToFColorSRGB inline, it's a bit large for that, but make an array version that does multiple at once. Mainly try to avoid doing a cross-DLL call per pixel for large batches.
Switch ImageCore Linear->sRGB conv to use array version.
FLinearColor R,G,B,A members union'd with "deprecated" RGBA[] array, Component() updated to use it (matching recent fix for FVector)
ImageCore vector loads from FLinearColors updated to use &Component(0) which gives &RGBA[0], pointer to a 4-elem array
Add tests to UnrealMathTest.cpp to verify that sRGB<->Linear and UNORM<->Linear conversions behave as expected, plus the separate ColorConversionHeavy test that does the exhaustive equivalence check between the SSE2 ToFColorSRGB() and reference (warning: takes a while)

diff --git a/Engine/Source/Runtime/Core/Private/Math/Color.cpp b/Engine/Source/Runtime/Core/Private/Math/Color.cpp
--- a/Engine/Source/Runtime/Core/Private/Math/Color.cpp
+++ b/Engine/Source/Runtime/Core/Private/Math/Color.cpp
@@ -34,16 +34,6 @@ const FColor FColor::Turquoise(26, 188, 156);
 const FColor FColor::Silver(189, 195, 199);
 const FColor FColor::Emerald(46, 204, 113);
 
-// FLinearColor(FColor) does sRGB -> Linear conversion.
-// to get direct conversion use ReinterpretAsLinear
-FLinearColor::FLinearColor(const FColor& Color)
-{
-	R = sRGBToLinearTable[Color.R];
-	G = sRGBToLinearTable[Color.G];
-	B =	sRGBToLinearTable[Color.B];
-	A =	float(Color.A) * (1.0f / 255.0f);
-}
-
 FLinearColor::FLinearColor(const FVector3f& Vector) :
 	R(Vector.X),
 	G(Vector.Y),
@@ -88,16 +78,6 @@ FLinearColor FLinearColor::FromPow22Color(const FColor& Color)
 	return LinearColor;
 }
 
-// For pixel format conversions. Clamps to [0,1], but for consistency with
-// GPU conversions, also define behavior for NaNs and make them map to 0.
-static inline float Clamp01(float Value)
-{
-	// Write this explicitly instead of using FMath::Clamp so we have
-	// control over what happens with NaNs.
-	const float ClampedLo = (Value > 0.0f) ? Value : 0.0f; // Also turns NaNs into 0.
-	return (ClampedLo < 1.0f) ? ClampedLo : 1.0f;
-}
-
 /**
  * Converts from a linear float color to RGBE as outlined in Gregory Ward's Real Pixels article, Graphics Gems II, page 80.
  * Implementation details in https://cbloomrants.blogspot.com/2020/06/widespread-error-in-radiance-hdr-rgbe.html
@@ -192,51 +172,113 @@ static uint8 stbir__linear_to_srgb_uchar_fast(float in)
     return (uint8) ((bias + scale*t) >> 16);
 }
 
+#if PLATFORM_CPU_X86_FAMILY
+
+static FColor ConvertLinearToSRGBSSE2(const FLinearColor& InColor)
+{
+	const VectorRegister4Float InRGBA = VectorLoad(&InColor.Component(0));
+
+	// Clamp to [2^(-13), 1-eps]; these two values map to 0 and 1, respectively.
+	// This clamping logic is carefully written so that NaNs map to 0.
+	//
+	// We do this clamping on all four color channels, even though we later handle
+	// A differently; this does not change the results for A: 2^(-13) rounds to 0 in
+	// U8, and 1-eps rounds to 255 in U8, so these are OK endpoints to use.
+	const VectorRegister4Float AlmostOne = VectorCastIntToFloat(VectorIntSet1(0x3f7fffff)); // 1-eps
+	const VectorRegister4Int MinValInt = VectorIntSet1((127 - 13) << 23);
+	const VectorRegister4Float MinValFlt = VectorCastIntToFloat(MinValInt);
+
+	const VectorRegister4Float InClamped = VectorMin(VectorMax(InRGBA, MinValFlt), AlmostOne);
+
+	// Set up for the table lookup
+	// This computes a 3-vector of table indices. The above clamping
+	// ensures that the values in question are in [0,13*8-1]=[0,103].
+	const VectorRegister4Int TabIndexVec = _mm_srli_epi32(_mm_sub_epi32(_mm_castps_si128(InClamped), MinValInt), 20);
+
+	// Do the 4 table lookups with regular loads. We can use PEXTRW (SSE2)
+	// to grab the 3 indices from lanes 1-3, lane 0 we can just get via MOVD.
+	// The latter gives us a full 32 bits, not 16 like the other ones, but given
+	// our value range either works.
+	const VectorRegister4Int TabValR = _mm_cvtsi32_si128(stb_fp32_to_srgb8_tab4[(uint32)_mm_cvtsi128_si32(TabIndexVec)]);
+	const VectorRegister4Int TabValG = _mm_cvtsi32_si128(stb_fp32_to_srgb8_tab4[(uint32)_mm_extract_epi16(TabIndexVec, 2)]);
+	const VectorRegister4Int TabValB = _mm_cvtsi32_si128(stb_fp32_to_srgb8_tab4[(uint32)_mm_extract_epi16(TabIndexVec, 4)]);
+
+	// Merge the four values we just loaded back into a 3-vector (gather complete!)
+	const VectorRegister4Int TabValRG = _mm_unpacklo_epi32(TabValR, TabValG);
+	const VectorRegister4Int TabValsRGB = _mm_unpacklo_epi64(TabValRG, TabValB); // This leaves A=0, which suits us
+
+	// Grab the mantissa bits into the low 16 bits of each 32b lane, and set up 512 in the high
+	// 16 bits of each 32b lane, which is how the bias values in the table are meant to be scaled.
+	//
+	// We grab mantissa bits [12,19] for the lerp.
+	const VectorRegister4Int MantissaLerpFactor = _mm_and_si128(_mm_srli_epi32(_mm_castps_si128(InClamped), 12), _mm_set1_epi32(0xff));
+	const VectorRegister4Int FinalMultiplier = _mm_or_si128(MantissaLerpFactor, _mm_set1_epi32(512 << 16));
+
+	// In the table:
+	//    (bias>>9) was stored in the high 16 bits
+	//    scale was stored in the low 16 bits
+	//    t = (mantissa >> 12) & 0xff
+	//
+	// then we want ((bias + scale*t) >> 16).
+	// Except for the final shift, that's a single PMADDWD:
+	const VectorRegister4Int InterpolatedRGB = _mm_srli_epi32(_mm_madd_epi16(TabValsRGB, FinalMultiplier), 16);
+
+	// Finally, A gets done directly, via (int)(A * 255.f + 0.5f)
+	// We zero out the non-A channels by multiplying by 0; our clamping earlier
+	// took care of NaNs/infinites, so this is fine
+	const VectorRegister4Float ScaledBiasedA = _mm_add_ps(_mm_mul_ps(InClamped, _mm_setr_ps(0.f, 0.f, 0.f, 255.f)), _mm_set1_ps(0.5f));
+	const VectorRegister4Int FinalA  = _mm_cvttps_epi32(ScaledBiasedA);
+
+	// Merge A into the result, reorder to BGRA, then pack down to bytes and store!
+	// InterpolatedRGB has lane 3=0, and ComputedA has the first three lanes zero,
+	// so we can just OR them together.
+	const VectorRegister4Int FinalRGBA = _mm_or_si128(InterpolatedRGB, FinalA);
+	const VectorRegister4Int FinalBGRA = _mm_shuffle_epi32(FinalRGBA, _MM_SHUFFLE(3, 0, 1, 2));
+
+	const VectorRegister4Int Packed16 = _mm_packs_epi32(FinalBGRA, FinalBGRA);
+	const VectorRegister4Int Packed8 = _mm_packus_epi16(Packed16, Packed16);
+
+	return FColor((uint32)_mm_cvtsi128_si32(Packed8));
+}
+
+#endif
+
 /** Quantizes the linear color and returns the result as a FColor with optional sRGB conversion and quality as goal. */
 FColor FLinearColor::ToFColorSRGB() const
 {
 	// The convention used here in all channels is that NaNs
 	// convert to 0, as do negative values, and out-of-range
 	// positive values convert to 255.
+
+#if PLATFORM_CPU_X86_FAMILY
+	return ConvertLinearToSRGBSSE2(*this);
+#else
 	return FColor(
 		stbir__linear_to_srgb_uchar_fast(R),
 		stbir__linear_to_srgb_uchar_fast(G),
 		stbir__linear_to_srgb_uchar_fast(B),
-		(uint8)(0.5f + Clamp01(A)*255.f)
+		(uint8)(0.5f + Clamp01NansTo0(A)*255.f)
 	);
+#endif
 }
 
-
-FColor FLinearColor::Quantize() const
-{
-	// Quantize is deprecated
-	// almost all callers should use QuantizeRound instead
-	// QuantizeFloor maintains old Quantize behavior
-	return QuantizeFloor();
-}
-
-FColor FLinearColor::QuantizeFloor() const
-{
-	// Do not use this for graphics or textures or images, they should use QuantizeRound
-	// replicates behavior of old Quantize() which is deprecated
-	// restoration should be done to bucket centers (+0.5 on dequantize)
-	return FColor(
-		(uint8)(Clamp01(R)*255.f),
-		(uint8)(Clamp01(G)*255.f),
-		(uint8)(Clamp01(B)*255.f),
-		(uint8)(Clamp01(A)*255.f)
-		);
-}
-
-FColor FLinearColor::QuantizeRound() const
+/**
+ * Convert multiple FLinearColors to sRGB FColor; array version of FLinearColor::ToFColorSRGB.
+ * 
+ * @param	InLinearColors	Pointer to one or more FLinearColors to convert.
+ * @param	OutColorsSRGB	Pointer to one or more FColors that receive the results.
+ * @param	InCount			Number of colors to convert.
+ */
+void ConvertFLinearColorsToFColorSRGB(const FLinearColor* InLinearColors, FColor* OutColorsSRGB, int64 InCount)
 {
-	// do not use FMath::RoundToInt because it call floorf
-	return FColor(
-		(uint8)(0.5f + Clamp01(R)*255.f),
-		(uint8)(0.5f + Clamp01(G)*255.f),
-		(uint8)(0.5f + Clamp01(B)*255.f),
-		(uint8)(0.5f + Clamp01(A)*255.f)
-		);
+	// This function exists because calling FLinearColor::ToFColorSRGB()
+	// from another module is not cheap. However, we're in the same module
+	// as the definition here, so inlining works fine, and we can just
+	// implement this as the straightforward loop.
+	for (int64 Index = 0; Index < InCount; ++Index)
+	{
+		OutColorsSRGB[Index] = InLinearColors[Index].ToFColorSRGB();
+	}
 }
 
 /**
diff --git a/Engine/Source/Runtime/Core/Private/Tests/Math/UnrealMathTest.cpp b/Engine/Source/Runtime/Core/Private/Tests/Math/UnrealMathTest.cpp
--- a/Engine/Source/Runtime/Core/Private/Tests/Math/UnrealMathTest.cpp
+++ b/Engine/Source/Runtime/Core/Private/Tests/Math/UnrealMathTest.cpp
@@ -16,7 +16,10 @@
 #include "Math/RotationMatrix.h"
 #include "Math/Quat.h"
 #include "Math/QuatRotationTranslationMatrix.h"
+#include "Math/Color.h"
 #include "Misc/AutomationTest.h"
+#include "Async/ParallelFor.h"
+#include "Misc/ScopeLock.h"
 #include <limits>
 #include <cmath>
 
@@ -4144,4 +4147,255 @@ bool FInitVectorTest::RunTest(const FString& Parameters)
 	return true;
 }
 
+IMPLEMENT_SIMPLE_AUTOMATION_TEST(FColorConversionTest, "System.Core.Math.ColorConversion", EAutomationTestFlags::ApplicationContextMask | EAutomationTestFlags::EngineFilter)
+bool FColorConversionTest::RunTest(const FString& Parameters)
+{
+	// Round-trip conversion tests check that converting from uint8 color formats
+	// to float and back gives the original value. We not only want to guarantee
+	// this on its own, it also gives us useful coverage of the [0,1] value range.
+
+	// Test that sRGB<->Linear conversions round-trip
+	for (int Index = 0; Index < 256; ++Index)
+	{
+		// Make the inputs in R,G,B not the same in cases channels get swapped or similar
+		// Alpha channel is already special because it gets treated differently
+		FColor OriginalColor((uint8)Index, (uint8)(Index ^ 1), (uint8)(Index ^ 123), (uint8)Index);
+
+		FLinearColor SRGBToLinear = FLinearColor::FromSRGBColor(OriginalColor);
+		FColor SRGBConvertedBack = SRGBToLinear.ToFColorSRGB();
+		TestEqual(FString::Printf(TEXT("sRGB to linear to sRGB round-trip: %d"), Index), OriginalColor, SRGBConvertedBack);
+
+		FLinearColor UNORMToLinear = OriginalColor.ReinterpretAsLinear();
+		FColor UNORMConvertedBack = UNORMToLinear.QuantizeRound();
+		TestEqual(FString::Printf(TEXT("UNORM to linear to UNORM round-trip: %d"), Index), OriginalColor, UNORMConvertedBack);
+	}
+
+	// Test values near breakpoints/bucket boundaries to make sure they end up
+	// on the intended side. Since we are interested in the boundaries between
+	// values, this loop only goes to 255.
+	auto ReferenceSRGBToLinear = [](float InValue) -> float
+	{
+		if (InValue < 0.04045f)
+		{
+			return InValue / 12.92f;
+		}
+		else
+		{
+			return (float)FMath::Pow((InValue + 0.055) / 1.055, 2.4); // Compute in doubles
+		}
+	};
+	
+	for (int Index = 0; Index < 255; ++Index)
+	{
+		float BucketMidpoint = float(Index) + 0.5f;
+
+		// Worst-case error in the integer [0,255] sRGB scale is guaranteed to be below
+		// 0.544403 by the conversion we use. (The minimum reachable is 0.5, because
+		// we're quantizing to integers). That means that as long as we stay about
+		// 0.045f units away from a breakpoint, we should always get the exact value.
+		const float DistanceToMidpoint = 0.045f;
+
+		float BelowBoundaryUNORM = (BucketMidpoint - DistanceToMidpoint) / 255.f;
+		float BelowBoundarySRGB = ReferenceSRGBToLinear(BelowBoundaryUNORM);
+		uint8 IndexU8 = (uint8)Index;
+		FColor BelowExpected(IndexU8, IndexU8, IndexU8, IndexU8);
+		FColor BelowConverted = FLinearColor(BelowBoundarySRGB, BelowBoundarySRGB, BelowBoundarySRGB, BelowBoundaryUNORM).ToFColorSRGB();
+		TestEqual(FString::Printf(TEXT("sRGB Boundary below: %d"), Index), BelowConverted, BelowExpected);
+
+		float AboveBoundaryUNORM = (BucketMidpoint + DistanceToMidpoint) / 255.f;
+		float AboveBoundarySRGB = ReferenceSRGBToLinear(AboveBoundaryUNORM);
+		uint8 IndexPlus1U8 = (uint8)(Index + 1);
+		FColor AboveExpected(IndexPlus1U8, IndexPlus1U8, IndexPlus1U8, IndexPlus1U8);
+		FColor AboveConverted = FLinearColor(AboveBoundarySRGB, AboveBoundarySRGB, AboveBoundarySRGB, AboveBoundaryUNORM).ToFColorSRGB();
+		TestEqual(FString::Printf(TEXT("sRGB Boundary above: %d"), Index), AboveConverted, AboveExpected);
+	}
+
+	// Between the two tests above, we have good coverage of what happens inside [0,1];
+	// for values outside, we expect the sRGB and UNORM paths to give the same results.
+	// Just test a few values in interesting parts of the range.
+	const float NaN = FPlatformMath::AsFloat((uint32)0x7fc00000u);
+	const float PosInf = FPlatformMath::AsFloat((uint32)0x7f800000u);
+	const float PosSubnormal = FPlatformMath::AsFloat((uint32)0x20000u);
+	const float MediumLarge = 10.f; // Outside the range, but not hugely so
+	const float VeryLarge = 1e+7f; // Far outside the range
+
+	auto TestExtremalValue = [this](const TCHAR* InTestName, const float InValue, uint8 InExpected)
+	{
+		FLinearColor LinColor(InValue, InValue, InValue, InValue);
+		FColor ConvertedColor = LinColor.ToFColorSRGB();
+		FColor ExpectedColor(InExpected, InExpected, InExpected, InExpected);
+
+		TestEqual(InTestName, ConvertedColor, ExpectedColor);
+	};
+
+	TestExtremalValue(TEXT("Extremal NaN"), NaN, 0);
+	TestExtremalValue(TEXT("Extremal -Inf"), -PosInf, 0);
+	TestExtremalValue(TEXT("Extremal -MediumLarge"), -MediumLarge, 0);
+	TestExtremalValue(TEXT("Extremal 0"), 0.0f, 0);
+	TestExtremalValue(TEXT("Extremal +subnorm"), PosSubnormal, 0);
+	TestExtremalValue(TEXT("Extremal 1"), 1.0f, 255);
+	TestExtremalValue(TEXT("Extremal Mediumlarge"), MediumLarge, 255);
+	TestExtremalValue(TEXT("Extremal VeryLarge"), VeryLarge, 255);
+	TestExtremalValue(TEXT("Extremal +Inf"), PosInf, 255);
+
+	return true;
+}
+
+// This is repeating the reference (scalar) linear->sRGB conversion from Color.cpp, 
+// because the reference version is not necessarily exposed; we can end up using a different
+// implementation depending on the platform (as of this writing, we have SSE2 implementation
+// too), but we want them all to match.
+
+typedef union
+{
+	uint32 u;
+	float f;
+} stbir__FP32;
+
+static const uint32 stb_fp32_to_srgb8_tab4[104] = {
+	0x0073000d, 0x007a000d, 0x0080000d, 0x0087000d, 0x008d000d, 0x0094000d, 0x009a000d, 0x00a1000d,
+	0x00a7001a, 0x00b4001a, 0x00c1001a, 0x00ce001a, 0x00da001a, 0x00e7001a, 0x00f4001a, 0x0101001a,
+	0x010e0033, 0x01280033, 0x01410033, 0x015b0033, 0x01750033, 0x018f0033, 0x01a80033, 0x01c20033,
+	0x01dc0067, 0x020f0067, 0x02430067, 0x02760067, 0x02aa0067, 0x02dd0067, 0x03110067, 0x03440067,
+	0x037800ce, 0x03df00ce, 0x044600ce, 0x04ad00ce, 0x051400ce, 0x057b00c5, 0x05dd00bc, 0x063b00b5,
+	0x06970158, 0x07420142, 0x07e30130, 0x087b0120, 0x090b0112, 0x09940106, 0x0a1700fc, 0x0a9500f2,
+	0x0b0f01cb, 0x0bf401ae, 0x0ccb0195, 0x0d950180, 0x0e56016e, 0x0f0d015e, 0x0fbc0150, 0x10630143,
+	0x11070264, 0x1238023e, 0x1357021d, 0x14660201, 0x156601e9, 0x165a01d3, 0x174401c0, 0x182401af,
+	0x18fe0331, 0x1a9602fe, 0x1c1502d2, 0x1d7e02ad, 0x1ed4028d, 0x201a0270, 0x21520256, 0x227d0240,
+	0x239f0443, 0x25c003fe, 0x27bf03c4, 0x29a10392, 0x2b6a0367, 0x2d1d0341, 0x2ebe031f, 0x304d0300,
+	0x31d105b0, 0x34a80555, 0x37520507, 0x39d504c5, 0x3c37048b, 0x3e7c0458, 0x40a8042a, 0x42bd0401,
+	0x44c20798, 0x488e071e, 0x4c1c06b6, 0x4f76065d, 0x52a50610, 0x55ac05cc, 0x5892058f, 0x5b590559,
+	0x5e0c0a23, 0x631c0980, 0x67db08f6, 0x6c55087f, 0x70940818, 0x74a007bd, 0x787d076c, 0x7c330723,
+};
+
+static uint8 stbir__linear_to_srgb_uchar_fast(float in)
+{
+	static const stbir__FP32 almostone = { 0x3f7fffff }; // 1-eps
+	static const stbir__FP32 minval = { (127 - 13) << 23 };
+	uint32 tab, bias, scale, t;
+	stbir__FP32 f;
+
+	// Clamp to [2^(-13), 1-eps]; these two values map to 0 and 1, respectively.
+	// The tests are carefully written so that NaNs map to 0, same as in the reference
+	// implementation.
+	if (!(in > minval.f)) // written this way to catch NaNs
+		in = minval.f;
+	if (in > almostone.f)
+		in = almostone.f;
+
+	// Do the table lookup and unpack bias, scale
+	f.f = in;
+
+	tab = stb_fp32_to_srgb8_tab4[(f.u - minval.u) >> 20];
+	bias = (tab >> 16) << 9;
+	scale = tab & 0xffff;
+
+	// Grab next-highest mantissa bits and perform linear interpolation
+	t = (f.u >> 12) & 0xff;
+	return (uint8)((bias + scale * t) >> 16);
+}
+
+IMPLEMENT_SIMPLE_AUTOMATION_TEST(FColorConversionHeavyTest, "System.Core.Math.ColorConversionHeavy", EAutomationTestFlags::ApplicationContextMask | EAutomationTestFlags::EngineFilter)
+bool FColorConversionHeavyTest::RunTest(const FString& Parameters)
+{
+	// WARNING: This test runs for a good while (at time of writing, ~90s when using
+	// a single core, proportionately less on many-core machines) with no user
+	// feedback other than log messages being printed. You've been warned.
+
+	// Chop up the 32-bit value range into pow2-sized buckets so we
+	// give at least regular progress updates in the log
+	const int BucketShift = 8;
+	const int BucketCount = 1 << BucketShift;
+	const uint32 ItemsInBucket = 1u << (32 - BucketShift);
+
+	struct FSharedData
+	{
+		FCriticalSection Lock; // Protects everything in here
+		bool Failed = false;
+		int32 CompletionCounter = 0;
+
+		// These are set at most once (when Failed is first set)
+		uint32 FailedValue = 0; // Value we first noticed a mismatch on
+		FColor FailedConverted;
+		FColor FailedExpected;
+	};
+
+	FSharedData Shared;
+
+	// Test that the platform specialization for ToFColorSRGB agrees with the reference
+	// implementation. As of this writing, we have a vectorized version for SSE2 targets.
+	ParallelFor(BucketCount,
+		[BucketCount, ItemsInBucket, &Shared](int32 BucketIndex)
+		{
+			for (uint32 WithinBucketIndex = 0; WithinBucketIndex < ItemsInBucket; ++WithinBucketIndex)
+			{
+				uint32 CurrentBits = BucketIndex * ItemsInBucket + WithinBucketIndex;
+
+				// Don't put the same value in every color channel; we want to make sure we catch
+				// accidental channel swaps too! Three different inputs only; A is special anyway.
+				float R = FPlatformMath::AsFloat(CurrentBits);
+				float G = FPlatformMath::AsFloat(CurrentBits ^ 1234567u);
+				float B = FPlatformMath::AsFloat(CurrentBits ^ 89376762u);
+				FLinearColor LinearInput(R, G, B, R);
+				FColor Converted = LinearInput.ToFColorSRGB();
+
+				// Reference conversion
+				uint8 RgbResult[3];
+				for (int Channel = 0; Channel < 3; ++Channel)
+				{
+					float Value = LinearInput.Component(Channel);
+					RgbResult[Channel] = stbir__linear_to_srgb_uchar_fast(Value);
+				}
+
+				// Clamp A, mapping NaNs to 0
+				float ClampedALo = (LinearInput.A > 0.0f) ? LinearInput.A : 0.0f;
+				float ClampedA = (ClampedALo < 1.0f) ? ClampedALo : 1.0f;
+				uint8 FinalA = (uint8)(ClampedA * 255.f + 0.5f);
+
+				FColor Expected{ RgbResult[0], RgbResult[1], RgbResult[2], FinalA };
+
+				// Test and note failures (get reported outside)
+				if (Converted != Expected)
+				{
+					FScopeLock LockHolder(&Shared.Lock);
+
+					// If we are the first to fail, set failure info
+					if (!Shared.Failed)
+					{
+						Shared.Failed = true;
+						Shared.FailedValue = CurrentBits;
+						Shared.FailedConverted = Converted;
+						Shared.FailedExpected = Expected;
+					}
+				}
+			}
+
+			// Check status and update count
+			bool Failed = false;
+			int32 Completed = 0;
+
+			// Scope lock access to just grabbing the few values, not the log too
+			{
+				FScopeLock LockHolder(&Shared.Lock);
+				Failed = Shared.Failed;
+				Completed = ++Shared.CompletionCounter;
+			}
+
+			// If another instance failed, stop
+			if (Failed)
+				return;
+
+			UE_LOG(LogUnrealMathTest, Log, TEXT("Conversion heavy sRGB %d/%d buckets completed (%.2f%%)"), Completed, BucketCount, double(Completed) * 100.0 / double(BucketCount));
+		}
+	);
+
+	if (Shared.Failed)
+	{
+		TestEqual(FString::Printf(TEXT("Conversion heavy sRGB CurrentBits=0x%08x"), Shared.FailedValue), Shared.FailedConverted, Shared.FailedExpected);
+		return false;
+	}
+
+	return true;
+}
+
 #endif //WITH_DEV_AUTOMATION_TESTS
\ No newline at end of file
diff --git a/Engine/Source/Runtime/Core/Public/Math/Color.h b/Engine/Source/Runtime/Core/Public/Math/Color.h
--- a/Engine/Source/Runtime/Core/Public/Math/Color.h
+++ b/Engine/Source/Runtime/Core/Public/Math/Color.h
@@ -31,16 +31,25 @@ enum class EGammaSpace : uint8
  */
 struct FLinearColor
 {
-	float	R,
-			G,
-			B,
-			A;
+	union
+	{
+		struct 
+		{
+			float	R,
+					G,
+					B,
+					A;
+		};
+
+		UE_DEPRECATED(all, "For internal use only.")
+		float RGBA[4];
+	};
 
 	/** Static lookup table used for FColor -> FLinearColor conversion. Pow(2.2) */
 	static float Pow22OneOver255Table[256];
 
 	/** Static lookup table used for FColor -> FLinearColor conversion. sRGB */
-	static float sRGBToLinearTable[256];
+	static CORE_API float sRGBToLinearTable[256];
 
 	FORCEINLINE FLinearColor() {}
 	FORCEINLINE explicit FLinearColor(EForceInit)
@@ -53,7 +62,7 @@ struct FLinearColor
 	 * @param Color The sRGB color that needs to be converted into linear space.
 	 * to get direct conversion use ReinterpretAsLinear
 	 */
-	CORE_API FLinearColor(const FColor& Color);
+	FORCEINLINE FLinearColor(const FColor& Color);
 
 	CORE_API FLinearColor(const FVector3f& Vector);
 	CORE_API explicit FLinearColor(const FVector3d& Vector); // Warning: keep this explicit, or FVector4f will be implicitly created from FVector3d via FLinearColor
@@ -85,7 +94,7 @@ struct FLinearColor
 	 * Converts an FColor coming from an observed sRGB output, into a linear color.
 	 * @param Color The sRGB color that needs to be converted into linear space.
 	 */
-	CORE_API static FLinearColor FromSRGBColor(const FColor& Color)
+	FORCEINLINE static FLinearColor FromSRGBColor(const FColor& Color)
 	{
 		return FLinearColor(Color);
 	}
@@ -100,12 +109,16 @@ struct FLinearColor
 
 	FORCEINLINE float& Component(int32 Index)
 	{
-		return (&R)[Index];
+	PRAGMA_DISABLE_DEPRECATION_WARNINGS
+		return RGBA[Index];
+	PRAGMA_ENABLE_DEPRECATION_WARNINGS
 	}
 
 	FORCEINLINE const float& Component(int32 Index) const
 	{
-		return (&R)[Index];
+	PRAGMA_DISABLE_DEPRECATION_WARNINGS
+		return RGBA[Index];
+	PRAGMA_ENABLE_DEPRECATION_WARNINGS
 	}
 
 	FORCEINLINE FLinearColor operator+(const FLinearColor& ColorB) const
@@ -311,22 +324,22 @@ struct FLinearColor
 
 	/** Quantizes the linear color with rounding and returns the result as a FColor.  This bypasses the SRGB conversion. 
 	* QuantizeRound can be dequantized back to linear with FColor::ReinterpretAsLinear (just /255.f)
-	* this matches the GPU U8<->float conversion spec and should be preferred
+	* this matches the GPU UNORM<->float conversion spec and should be preferred
 	*/
-	CORE_API FColor QuantizeRound() const;
+	FORCEINLINE FColor QuantizeRound() const;
 	
 	/** Quantizes the linear color and returns the result as a FColor.  This bypasses the SRGB conversion.
 	* Uses floor quantization, which does not match the GPU standard conversion.
 	* Restoration to float should be done with a +0.5 bias to restore to centered buckets.
 	* Do NOT use this for graphics or textures or images, use QuantizeRound instead.
 	*/
-	CORE_API FColor QuantizeFloor() const;
+	FORCEINLINE FColor QuantizeFloor() const;
 
 	/** backwards compatible Quantize function name, does QuantizeFloor.
 	* @todo deprecate me
 	*/
-	UE_DEPRECATED(5.0,"Most callers of Quantize should have been calling QuantizeRound; to match old behavior use QuantizeFloor")
-	CORE_API FColor Quantize() const;
+	UE_DEPRECATED(5.0, "Most callers of Quantize should have been calling QuantizeRound; to match old behavior use QuantizeFloor")
+	FORCEINLINE FColor Quantize() const;
 
 	/** Quantizes the linear color and returns the result as a FColor with optional sRGB conversion. 
 	* Clamps in [0,1] range before conversion.
@@ -402,6 +415,21 @@ struct FLinearColor
 		return bSuccessful;
 	}
 
+	/**
+	 * Helper for pixel format conversions. Clamps to [0,1], mapping NaNs to 0,
+	 * for consistency with GPU conversions.
+	 * 
+	 * @param InValue The input value.
+	 * @return InValue clamped to [0,1]. NaNs map to 0.
+	 */
+	static FORCEINLINE float Clamp01NansTo0(float InValue)
+	{
+		// Write this explicitly instead of using FMath::Clamp because we're particular
+		// about what happens with NaNs here.
+		const float ClampedLo = (InValue > 0.0f) ? InValue : 0.0f; // Also turns NaNs into 0.
+		return (ClampedLo < 1.0f) ? ClampedLo : 1.0f;
+	}
+
 	// Common colors.	
 	static CORE_API const FLinearColor White;
 	static CORE_API const FLinearColor Gray;
@@ -738,6 +766,39 @@ private:
 };
 DECLARE_INTRINSIC_TYPE_LAYOUT(FColor);
 
+FORCEINLINE FLinearColor::FLinearColor(const FColor& Color)
+{
+	R = sRGBToLinearTable[Color.R];
+	G = sRGBToLinearTable[Color.G];
+	B = sRGBToLinearTable[Color.B];
+	A = float(Color.A) * (1.0f / 255.0f);
+}
+
+FORCEINLINE FColor FLinearColor::QuantizeRound() const
+{
+	// Avoid FMath::RoundToInt because it calls floor()
+	return FColor(
+		(uint8)(0.5f + Clamp01NansTo0(R) * 255.f),
+		(uint8)(0.5f + Clamp01NansTo0(G) * 255.f),
+		(uint8)(0.5f + Clamp01NansTo0(B) * 255.f),
+		(uint8)(0.5f + Clamp01NansTo0(A) * 255.f)
+	);
+}
+
+FORCEINLINE FColor FLinearColor::QuantizeFloor() const
+{
+	return FColor(
+		(uint8)(Clamp01NansTo0(R) * 255.f),
+		(uint8)(Clamp01NansTo0(G) * 255.f),
+		(uint8)(Clamp01NansTo0(B) * 255.f),
+		(uint8)(Clamp01NansTo0(A) * 255.f)
+	);
+}
+
+FORCEINLINE FColor FLinearColor::Quantize() const
+{
+	return QuantizeFloor();
+}
 
 FORCEINLINE FColor FLinearColor::ToFColor(const bool bSRGB) const
 {
@@ -767,6 +828,9 @@ FORCEINLINE uint32 GetTypeHash( const FLinearColor& LinearColor )
 /** Computes a brightness and a fixed point color from a floating point color. */
 extern CORE_API void ComputeAndFixedColorAndIntensity(const FLinearColor& InLinearColor,FColor& OutColor,float& OutIntensity);
 
+/** Convert multiple FLinearColors to sRGB FColor; array version of FLinearColor::ToFColorSRGB. */
+extern CORE_API void ConvertFLinearColorsToFColorSRGB(const FLinearColor* InLinearColors, FColor* OutColorsSRGB, int64 InCount);
+
 // These act like a POD
 template <> struct TIsPODType<FColor> { enum { Value = true }; };
 template <> struct TIsPODType<FLinearColor> { enum { Value = true }; };
diff --git a/Engine/Source/Runtime/ImageCore/Private/ImageCore.cpp b/Engine/Source/Runtime/ImageCore/Private/ImageCore.cpp
--- a/Engine/Source/Runtime/ImageCore/Private/ImageCore.cpp
+++ b/Engine/Source/Runtime/ImageCore/Private/ImageCore.cpp
@@ -263,10 +263,12 @@ IMAGECORE_API void FImageCore::CopyImage(const FImageView & SrcImage,const FImag
 				FColor* DestColors = (FColor *) DestImage.RawData;
 				if (bDestIsGammaCorrected)
 				{
-					ParallelLoop(TEXT("Texture.CopyImage.PF"), NumJobs, TexelsPerJob, NumTexels,
-						[DestColors, SrcColors](int64 TexelIndex)
+					ParallelFor(TEXT("Texture.CopyImage.PF"), NumJobs, 1,
+						[DestColors, SrcColors, TexelsPerJob, NumTexels](int64 JobIndex)
 						{
-							DestColors[TexelIndex] = SrcColors[TexelIndex].ToFColorSRGB();
+							const int64 StartIndex = JobIndex * TexelsPerJob;
+							const int64 Count = FMath::Min(TexelsPerJob, NumTexels - StartIndex);
+							ConvertFLinearColorsToFColorSRGB(&SrcColors[StartIndex], &DestColors[StartIndex], Count);
 						}
 					);
 				}
@@ -277,10 +279,10 @@ IMAGECORE_API void FImageCore::CopyImage(const FImageView & SrcImage,const FImag
 						{
 #if PLATFORM_CPU_X86_FAMILY
 							// load 4x RGBA32F
-							__m128 Pixel0 = _mm_loadu_ps((const float*)&SrcColors[TexelIndex + 0]);
-							__m128 Pixel1 = _mm_loadu_ps((const float*)&SrcColors[TexelIndex + 1]);
-							__m128 Pixel2 = _mm_loadu_ps((const float*)&SrcColors[TexelIndex + 2]);
-							__m128 Pixel3 = _mm_loadu_ps((const float*)&SrcColors[TexelIndex + 3]);
+							__m128 Pixel0 = _mm_loadu_ps(&SrcColors[TexelIndex + 0].Component(0));
+							__m128 Pixel1 = _mm_loadu_ps(&SrcColors[TexelIndex + 1].Component(0));
+							__m128 Pixel2 = _mm_loadu_ps(&SrcColors[TexelIndex + 2].Component(0));
+							__m128 Pixel3 = _mm_loadu_ps(&SrcColors[TexelIndex + 3].Component(0));
 
 							// RGBA -> BGRA
 							Pixel0 = _mm_shuffle_ps(Pixel0, Pixel0, _MM_SHUFFLE(3, 0, 1, 2));
@@ -467,11 +469,11 @@ IMAGECORE_API void FImageCore::CopyImage(const FImageView & SrcImage,const FImag
 							Out2 = _mm_shuffle_ps(Out2, Out2, _MM_SHUFFLE(3, 0, 1, 2));
 							Out3 = _mm_shuffle_ps(Out3, Out3, _MM_SHUFFLE(3, 0, 1, 2));
 
-							// store 4x RGBAF pixels
-							_mm_storeu_ps((float*)&DestColors[TexelIndex + 0], Out0);
-							_mm_storeu_ps((float*)&DestColors[TexelIndex + 1], Out1);
-							_mm_storeu_ps((float*)&DestColors[TexelIndex + 2], Out2);
-							_mm_storeu_ps((float*)&DestColors[TexelIndex + 3], Out3);
+							// Store 4 pixels
+							_mm_storeu_ps(&DestColors[TexelIndex + 0].Component(0), Out0);
+							_mm_storeu_ps(&DestColors[TexelIndex + 1].Component(0), Out1);
+							_mm_storeu_ps(&DestColors[TexelIndex + 2].Component(0), Out2);
+							_mm_storeu_ps(&DestColors[TexelIndex + 3].Component(0), Out3);
 #else
 							check(false); // not supported for other platforms, see ParallelLoop
 #endif
