#ue5 - TransformCalculus converted to support templated types, and converted existing types to use templates.

Currently the old types (eg FTransform2D) are typedef'd to use the float version for backwards compat. After Slate converts its internal usage, this can be revisited and swapped to double by default for consistency with engine types; this will require Slate to start using its own defined type for all slate code for clarity.

diff --git a/Engine/Plugins/FX/Niagara/Source/NiagaraEditorWidgets/Private/SNiagaraScratchPad.cpp b/Engine/Plugins/FX/Niagara/Source/NiagaraEditorWidgets/Private/SNiagaraScratchPad.cpp
--- a/Engine/Plugins/FX/Niagara/Source/NiagaraEditorWidgets/Private/SNiagaraScratchPad.cpp
+++ b/Engine/Plugins/FX/Niagara/Source/NiagaraEditorWidgets/Private/SNiagaraScratchPad.cpp
@@ -118,7 +118,7 @@ private:
 	{
 		return IsPinned.Get(false) 
 			? TOptional<FSlateRenderTransform>()
-			: FSlateRenderTransform(FQuat2D(PI / 2));
+			: TOptional<FSlateRenderTransform>(FSlateRenderTransform(FQuat2D(PI / 2)));
 	}
 
 private:
diff --git a/Engine/Source/Runtime/Core/Public/Math/TransformCalculus.h b/Engine/Source/Runtime/Core/Public/Math/TransformCalculus.h
--- a/Engine/Source/Runtime/Core/Public/Math/TransformCalculus.h
+++ b/Engine/Source/Runtime/Core/Public/Math/TransformCalculus.h
@@ -391,7 +391,7 @@ template <typename TransformType, typename PositionType>
 inline PositionType TransformPoint(const TransformType& Transform, const PositionType& Point)
 {
 	return Transform.TransformPoint(Point);
-}
+} 
 
 /**
  * Generic implementation of TransformVector. Attempts to use a member function of the TransformType.
@@ -402,33 +402,6 @@ inline VectorType TransformVector(const TransformType& Transform, const VectorTy
 	return Transform.TransformVector(Vector);
 }
 
-/**
- * Generic implementation of TransformPoint for 2D vectors with double precision. Attempts to use a member function of the TransformType.
- */
-template<typename TransformType>
-inline FVector2D TransformPoint(const TransformType& Transform, const FVector2D& Point)
-{
-	return Transform.TransformPoint(Point);
-}
-
-/**
- * Generic implementation of TransformPoint for 2D vectors with float precision. Attempts to use a member function of the TransformType.
- */
-template<typename TransformType>
-inline FVector2f TransformPoint(const TransformType& Transform, const FVector2f& Point)
-{
-	return FVector2f(Transform.TransformPoint(FVector2D(Point)));
-}
-
-/**
- * Generic implementation of TransformVector for 2D vectors. Attempts to use a member function of the TransformType.
- */
-template<typename TransformType>
-inline FVector2D TransformVector(const TransformType& Transform, const FVector2D& Vector)
-{
-	return Transform.TransformVector(Vector);
-}
-
 //////////////////////////////////////////////////////////////////////////
 // Overloads for uniform Scale.
 // 
@@ -444,7 +417,8 @@ inline FVector2D TransformVector(const TransformType& Transform, const FVector2D
  * @param RHS Scale that goes from space B to space C.
  * @return a new Scale representing the transformation from the input space of LHS to the output space of RHS.
  */
-inline float Concatenate(float LHS, float RHS)
+template<typename FloatType, TEMPLATE_REQUIRES(TIsFloatingPoint<FloatType>::Value)>
+inline FloatType Concatenate(FloatType LHS, FloatType RHS)
 {
 	return LHS * RHS;
 }
@@ -456,7 +430,8 @@ inline float Concatenate(float LHS, float RHS)
  * @param Transform Input transform from space A to space B.
  * @return Inverted transform from space B to space A.
  */
-inline float Inverse(float Scale)
+template<typename FloatType, TEMPLATE_REQUIRES(TIsFloatingPoint<FloatType>::Value)>
+inline FloatType Inverse(FloatType Scale)
 {
 	return 1.0f / Scale;
 }
@@ -464,7 +439,14 @@ inline float Inverse(float Scale)
 /**
  * Specialization for uniform Scale.
  */
-inline FVector TransformPoint(float Transform, const FVector& Point)
+template<typename PositionType>
+inline UE::Math::TVector<PositionType> TransformPoint(float Transform, const UE::Math::TVector<PositionType>& Point)
+{
+	return Transform * Point;
+}
+
+template<typename PositionType>
+inline UE::Math::TVector<PositionType> TransformPoint(double Transform, const UE::Math::TVector<PositionType>& Point)
 {
 	return Transform * Point;
 }
@@ -472,7 +454,14 @@ inline FVector TransformPoint(float Transform, const FVector& Point)
 /**
  * Specialization for uniform Scale.
  */
-inline FVector TransformVector(float Transform, const FVector& Vector)
+template<typename VectorType>
+inline UE::Math::TVector<VectorType> TransformVector(float Transform, const UE::Math::TVector<VectorType>& Vector)
+{
+	return Transform * Vector;
+}
+
+template<typename VectorType>
+inline UE::Math::TVector<VectorType> TransformVector(double Transform, const UE::Math::TVector<VectorType>& Vector)
 {
 	return Transform * Vector;
 }
diff --git a/Engine/Source/Runtime/Core/Public/Math/TransformCalculus2D.h b/Engine/Source/Runtime/Core/Public/Math/TransformCalculus2D.h
--- a/Engine/Source/Runtime/Core/Public/Math/TransformCalculus2D.h
+++ b/Engine/Source/Runtime/Core/Public/Math/TransformCalculus2D.h
@@ -9,15 +9,13 @@
 #include "Math/Vector2D.h"
 #include "Math/TransformCalculus.h"
 
-class FMatrix2x2;
-
 //////////////////////////////////////////////////////////////////////////
 // Transform calculus for 2D types. UE4 already has a 2D Vector class that we
 // will adapt to be interpreted as a translate transform. The rest we create
 // new classes for.
 //
 // The following types are supported
-// * float           -> represents a uniform scale.
+// * float/double    -> represents a uniform scale.
 // * FScale2D        -> represents a 2D non-uniform scale.
 // * FVector2D       -> represents a 2D translation.
 // * FShear2D        -> represents a "2D shear", interpreted as a shear parallel to the X axis followed by a shear parallel to the Y axis.
@@ -26,11 +24,11 @@ class FMatrix2x2;
 //
 //////////////////////////////////////////////////////////////////////////
 
-class FMatrix2x2;
+template<typename T> class TMatrix2x2;
 
 
 //////////////////////////////////////////////////////////////////////////
-// Adapters for FVector2D. 
+// Adapters for TVector2. 
 // 
 // Since it is an existing UE4 types, we cannot rely on the default
 // template that calls member functions. Instead, we provide direct overloads.
@@ -58,14 +56,16 @@ inline UE::Math::TVector2<T> Inverse(const UE::Math::TVector2<T>& Transform)
 	return -Transform;
 }
 
-/** Specialization for FVector2D Translation. */
-inline FVector2D TransformPoint(const FVector2D& Transform, const FVector2D& Point)
+/** Specialization for TVector2 Translation. */
+template<typename T>
+inline UE::Math::TVector2<T> TransformPoint(const UE::Math::TVector2<T>& Transform, const UE::Math::TVector2<T>& Point)
 {
 	return Transform + Point;
 }
 
 /** Specialization for FVector2D Translation (does nothing). */
-inline const FVector2D& TransformVector(const FVector2D& Transform, const FVector2D& Vector)
+template<typename T>
+inline const UE::Math::TVector2<T>& TransformVector(const UE::Math::TVector2<T>& Transform, const UE::Math::TVector2<T>& Vector)
 {
 	return Vector;
 }
@@ -80,7 +80,14 @@ inline const FVector2D& TransformVector(const FVector2D& Transform, const FVecto
 /**
 * Specialization for uniform Scale.
 */
-inline FVector2D TransformPoint(float Transform, const FVector2D& Point)
+template<typename PositionType>
+inline UE::Math::TVector2<PositionType> TransformPoint(float Transform, const UE::Math::TVector2<PositionType>& Point)
+{
+	return Transform * Point;
+}
+
+template<typename PositionType>
+inline UE::Math::TVector2<PositionType> TransformPoint(double Transform, const UE::Math::TVector2<PositionType>& Point)
 {
 	return Transform * Point;
 }
@@ -88,72 +95,94 @@ inline FVector2D TransformPoint(float Transform, const FVector2D& Point)
 /**
 * Specialization for uniform Scale.
 */
-inline FVector2D TransformVector(float Transform, const FVector2D& Vector)
+template<typename VectorType>
+inline UE::Math::TVector2<VectorType> TransformVector(float Transform, const UE::Math::TVector2<VectorType>& Vector)
 {
 	return Transform * Vector;
 }
 
+template<typename VectorType>
+inline UE::Math::TVector2<VectorType> TransformVector(double Transform, const UE::Math::TVector2<VectorType>& Vector)
+{
+	return Transform * Vector;
+}
 
 /** Represents a 2D non-uniform scale (to disambiguate from an FVector2D, which is used for translation) */
-class FScale2D
+template<typename T>
+class TScale2
 {
+	static_assert(std::is_floating_point_v<T>, "T must be floating point");
+
 public:
+	using FReal = T;
+	using Vector2Type = UE::Math::TVector2<T>;
+
 	/** Ctor. initialize to an identity scale, 1.0. */
-	FScale2D() :Scale(1.0f, 1.0f) {}
+	TScale2() : Scale(1.0f, 1.0f) {}
 	/** Ctor. initialize from a uniform scale. */
-	explicit FScale2D(float InScale) :Scale(InScale, InScale) {}
+	explicit TScale2(T InScale) :Scale(InScale, InScale) {}
 	/** Ctor. initialize from a non-uniform scale. */
-	explicit FScale2D(float InScaleX, float InScaleY) :Scale(InScaleX, InScaleY) {}
-	/** Ctor. initialize from an FVector defining the 3D scale. */
-	explicit FScale2D(const FVector2D& InScale) :Scale((FVector2f)InScale) {}
+	explicit TScale2(T InScaleX, T InScaleY) :Scale(InScaleX, InScaleY) {}
 	/** Ctor. initialize from an FVector defining the 3D scale. */
-	explicit FScale2D(const FVector2f& InScale) :Scale(InScale) {}
+	template<typename ArgType>
+	explicit TScale2(const UE::Math::TVector2<ArgType>& InScale) :Scale(InScale) {}
 	
 	/** Transform 2D Point */
-	FVector2D TransformPoint(const FVector2D& Point) const
+	template<typename ArgType>
+	UE::Math::TVector2<ArgType> TransformPoint(const UE::Math::TVector2<ArgType>& Point) const
 	{
-		return FVector2D(Scale) * Point;
+		return UE::Math::TVector2<ArgType>(Scale) * Point;
 	}
+
 	/** Transform 2D Vector*/
-	FVector2D TransformVector(const FVector2D& Vector) const
+	template<typename ArgType>
+	UE::Math::TVector2<ArgType> TransformVector(const UE::Math::TVector2<ArgType>& Vector) const
 	{
 		return TransformPoint(Vector);
 	}
 
 	/** Concatenate two scales. */
-	FScale2D Concatenate(const FScale2D& RHS) const
+	TScale2 Concatenate(const TScale2& RHS) const
 	{
-		return FScale2D(Scale * RHS.Scale);
+		return TScale2(Scale * RHS.Scale);
 	}
 	/** Invert the scale. */
-	FScale2D Inverse() const
+	TScale2 Inverse() const
 	{
-		return FScale2D(1.0f / Scale.X, 1.0f / Scale.Y);
+		return TScale2(1.0f / Scale.X, 1.0f / Scale.Y);
 	}
 
 	/** Equality. */
-	bool operator==(const FScale2D& Other) const
+	bool operator==(const TScale2& Other) const
 	{
 		return Scale == Other.Scale;
 	}
 	
 	/** Inequality. */
-	bool operator!=(const FScale2D& Other) const
+	bool operator!=(const TScale2& Other) const
 	{
 		return !operator==(Other);
 	}
 
 	/** Access to the underlying FVector2D that stores the scale. */
-	const FVector2f& GetVector() const { return Scale; }
+	const Vector2Type& GetVector() const { return Scale; }
+
 private:
 	/** Underlying storage of the 2D scale. */
-	FVector2f Scale;
+	Vector2Type Scale;
 };
 
+/** Base typedefs */
+typedef TScale2<float> FScale2f;
+typedef TScale2<double> FScale2d;
+typedef FScale2f FScale2D; // Default type (for backwards compat)
+
 /** concatenation rules for 2D scales. */
-template<> struct ConcatenateRules<float          , FScale2D  > { typedef FScale2D ResultType; };
+template<typename T> struct ConcatenateRules<float		, TScale2<T>	> { typedef TScale2<T> ResultType; };
+template<typename T> struct ConcatenateRules<double		, TScale2<T>	> { typedef TScale2<T> ResultType; };
 /** concatenation rules for 2D scales. */
-template<> struct ConcatenateRules<FScale2D  , float          > { typedef FScale2D ResultType; };
+template<typename T> struct ConcatenateRules<TScale2<T>	, float			> { typedef TScale2<T> ResultType; };
+template<typename T> struct ConcatenateRules<TScale2<T>	, double		> { typedef TScale2<T> ResultType; };
 
 /** 
  * Represents a 2D shear:
@@ -161,44 +190,53 @@ template<> struct ConcatenateRules<FScale2D  , float          > { typedef FScale
  *   [XX 1]
  * XX represents a shear parallel to the X axis. YY represents a shear parallel to the Y axis.
  */
-class FShear2D
+template<typename T>
+class TShear2
 {
+	static_assert(std::is_floating_point_v<T>, "T must be floating point");
+
 public:
+	using FReal = T;
+	using Vector2Type = UE::Math::TVector2<T>;
+
 	/** Ctor. initialize to an identity. */
-	FShear2D() :Shear(0,0) {}
+	TShear2() :Shear(0, 0) {}
 	/** Ctor. initialize from a set of shears parallel to the X and Y axis, respectively. */
-	explicit FShear2D(float ShearX, float ShearY) :Shear(ShearX, ShearY) {}
+	explicit TShear2(T ShearX, T ShearY) :Shear(ShearX, ShearY) {}
 	/** Ctor. initialize from a 2D vector representing a set of shears parallel to the X and Y axis, respectively. */
-	explicit FShear2D(const FVector2f& InShear) :Shear(InShear) {}
-	explicit FShear2D(const FVector2d& InShear) :Shear((FVector2f)InShear) {}
+	template<typename VType>
+	explicit TShear2(const UE::Math::TVector2<VType>& InShear) :Shear((Vector2Type)InShear) {}
 
 	/**
 	 * Generates a shear structure based on angles instead of slope.
 	 * @param InShearAngles The angles of shear.
 	 * @return the sheare structure.
 	 */
-	static FShear2D FromShearAngles(const FVector2D& InShearAngles)
+	template<typename VType>
+	static TShear2 FromShearAngles(const UE::Math::TVector2<VType>& InShearAngles)
 	{
 		// Compute the M (Shear Slot) = CoTan(90 - SlopeAngle)
 
 		// 0 is a special case because Tan(90) == infinity
-		float ShearX = InShearAngles.X == 0 ? 0 : ( 1.0f / FMath::Tan(FMath::DegreesToRadians(90 - FMath::Clamp((float)InShearAngles.X, -89.0f, 89.0f))) );
-		float ShearY = InShearAngles.Y == 0 ? 0 : ( 1.0f / FMath::Tan(FMath::DegreesToRadians(90 - FMath::Clamp((float)InShearAngles.Y, -89.0f, 89.0f))) );
+		T ShearX = InShearAngles.X == 0 ? 0 : (1.0f / FMath::Tan(FMath::DegreesToRadians(90 - FMath::Clamp<T>((T)InShearAngles.X, -89.0f, 89.0f))));
+		T ShearY = InShearAngles.Y == 0 ? 0 : (1.0f / FMath::Tan(FMath::DegreesToRadians(90 - FMath::Clamp<T>((T)InShearAngles.Y, -89.0f, 89.0f))));
 
-		return FShear2D(ShearX, ShearY);
+		return TShear2(ShearX, ShearY);
 	}
 
-	/** 
-	 * Transform 2D Point 
+	/**
+	 * Transform 2D Point
 	 * [X Y] * [1 YY] == [X+Y*XX Y+X*YY]
 	 *         [XX 1]
 	 */
-	FVector2D TransformPoint(const FVector2D& Point) const
+	template<typename ArgType>
+	UE::Math::TVector2<ArgType> TransformPoint(const UE::Math::TVector2<ArgType>& Point) const
 	{
-		return Point + FVector2D(Point.Y, Point.X) * FVector2D(Shear);
+		return Point + UE::Math::TVector2<ArgType>(Point.Y, Point.X) * UE::Math::TVector2<ArgType>(Shear);
 	}
 	/** Transform 2D Vector*/
-	FVector2D TransformVector(const FVector2D& Vector) const
+	template<typename ArgType>
+	UE::Math::TVector2<ArgType> TransformVector(const UE::Math::TVector2<ArgType>& Vector) const
 	{
 		return TransformPoint(Vector);
 	}
@@ -209,7 +247,7 @@ public:
 	 * [1 YYA] * [1 YYB] == [1+YYA*XXB YYB*YYA]
 	 * [XXA 1]   [XXB 1]    [XXA+XXB XXA*XXB+1]
 	 */
-	inline FMatrix2x2 Concatenate(const FShear2D& RHS) const;
+	inline TMatrix2x2<T> Concatenate(const TShear2& RHS) const;
 
 	/**
 	 * Invert the shear. The result is NOT a shear, but must be represented by a generalized 2x2 transform.
@@ -217,28 +255,35 @@ public:
 	 * [1 YY]^-1  == 1/(1-YY*XX) * [1 -YY]
 	 * [XX 1]                      [-XX 1]
 	 */
-	FMatrix2x2 Inverse() const;
+	TMatrix2x2<T> Inverse() const;
 
 
 	/** Equality. */
-	bool operator==(const FShear2D& Other) const
+	bool operator==(const TShear2& Other) const
 	{
 		return Shear == Other.Shear;
 	}
-	
+
 	/** Inequality. */
-	bool operator!=(const FShear2D& Other) const
+	bool operator!=(const TShear2& Other) const
 	{
 		return !operator==(Other);
 	}
 
 	/** Access to the underlying FVector2D that stores the scale. */
-	const FVector2f& GetVector() const { return Shear; }
+	const Vector2Type& GetVector() const { return Shear; }
+
 private:
 	/** Underlying storage of the 2D shear. */
-	FVector2f Shear;
+	Vector2Type Shear;
 };
 
+/** Base typedefs */
+typedef TShear2<float> FShear2f;
+typedef TShear2<double> FShear2d;
+typedef FShear2f FShear2D;  // Default type (for backwards compat)
+
+
 /** 
  * Represents a 2D rotation as a complex number (analagous to quaternions). 
  *   Rot(theta) == cos(theta) + i * sin(theta)
@@ -246,16 +291,22 @@ private:
  * Does not use "spinor" notation using theta/2 as we don't need that decomposition for our purposes.
  * This makes the implementation for straightforward and efficient for 2D.
  */
-class FQuat2D
+template<typename T>
+class TQuat2
 {
+	static_assert(std::is_floating_point_v<T>, "T must be floating point");
+
 public:
+	using FReal = T;
+	using Vector2Type = UE::Math::TVector2<T>;
+
 	/** Ctor. initialize to an identity rotation. */
-	FQuat2D() :Rot(1.0f, 0.0f) {}
+	TQuat2() :Rot(1.0f, 0.0f) {}
 	/** Ctor. initialize from a rotation in radians. */
-	explicit FQuat2D(float RotRadians) :Rot(FMath::Cos(RotRadians), FMath::Sin(RotRadians)) {}
+	explicit TQuat2(T RotRadians) :Rot(FMath::Cos(RotRadians), FMath::Sin(RotRadians)) {}
 	/** Ctor. initialize from an FVector2D, representing a complex number. */
-	explicit FQuat2D(const FVector2f& InRot) :Rot(InRot) {}
-	explicit FQuat2D(const FVector2d& InRot) :Rot((FVector2f)InRot) {}
+	template<typename VType>
+	explicit TQuat2(const UE::Math::TVector2<VType>& InRot) :Rot((Vector2Type)InRot) {}
 
 	/**
 	 * Transform a 2D point by the 2D complex number representing the rotation:
@@ -268,16 +319,18 @@ public:
 	 *         
 	 * Looking at the above results, we see the equivalence with matrix multiplication.
 	 */
-	FVector2D TransformPoint(const FVector2D& Point) const
+	template<typename ArgType>
+	UE::Math::TVector2<ArgType> TransformPoint(const UE::Math::TVector2<ArgType>& Point) const
 	{
-		return FVector2D(
-			Point.X * Rot.X - Point.Y * Rot.Y,
-			Point.X * Rot.Y + Point.Y * Rot.X);
+		return UE::Math::TVector2<ArgType>(
+			Point.X * (ArgType)Rot.X - Point.Y * (ArgType)Rot.Y,
+			Point.X * (ArgType)Rot.Y + Point.Y * (ArgType)Rot.X);
 	}
 	/**
 	 * Vector rotation is equivalent to rotating a point.
 	 */
-	FVector2D TransformVector(const FVector2D& Vector) const
+	template<typename ArgType>
+	UE::Math::TVector2<ArgType> TransformVector(const UE::Math::TVector2<ArgType>& Vector) const
 	{
 		return TransformPoint(Vector);
 	}
@@ -295,55 +348,67 @@ public:
 	 * out to [ A B] for a visual proof of the results.
 	 *        [-B A]
 	 */
-	FQuat2D Concatenate(const FQuat2D& RHS) const
+	TQuat2 Concatenate(const TQuat2& RHS) const
 	{
-		return FQuat2D(TransformPoint(FVector2D(RHS.Rot)));
+		return TQuat2(TransformPoint(FVector2D(RHS.Rot)));
 	}
 	/**
 	 * Invert the rotation  defined by complex numbers:
 	 * In imaginary land, an inverse is a complex conjugate, which is equivalent to reflecting about the X axis:
 	 * Conj(A + Bi) == A - Bi
 	 */
-	FQuat2D Inverse() const
+	TQuat2 Inverse() const
 	{
-		return FQuat2D(FVector2f(Rot.X, -Rot.Y));
+		return TQuat2(Vector2Type(Rot.X, -Rot.Y));
 	}
 
 	/** Equality. */
-	bool operator==(const FQuat2D& Other) const
+	bool operator==(const TQuat2& Other) const
 	{
 		return Rot == Other.Rot;
 	}
 	
 	/** Inequality. */
-	bool operator!=(const FQuat2D& Other) const
+	bool operator!=(const TQuat2& Other) const
 	{
 		return !operator==(Other);
 	}
 
 	/** Access to the underlying FVector2D that stores the complex number. */
-	const FVector2f& GetVector() const { return Rot; }
+	const Vector2Type& GetVector() const { return Rot; }
+
 private:
 	/** Underlying storage of the rotation (X = cos(theta), Y = sin(theta). */
-	FVector2f Rot;
+	Vector2Type Rot;
 };
 
+/** Base typedefs */
+typedef TQuat2<float> FQuat2f;
+typedef TQuat2<double> FQuat2d;
+typedef FQuat2f FQuat2D; // Default type (for backwards compat)
+
 /**
  * 2x2 generalized matrix. As FMatrix, we assume row vectors, row major storage:
  *    [X Y] * [m00 m01]
  *            [m10 m11]
  */
-class FMatrix2x2
+template<typename T>
+class TMatrix2x2
 {
+	static_assert(std::is_floating_point_v<T>, "T must be floating point");
+
 public:
+	using FReal = T;
+	using Vector2Type = UE::Math::TVector2<T>;
+
 	/** Ctor. initialize to an identity. */
-	FMatrix2x2()
+	TMatrix2x2()
 	{
 		M[0][0] = 1; M[0][1] = 0;
 		M[1][0] = 0; M[1][1] = 1;
 	}
 
-	FMatrix2x2(float m00, float m01, float m10, float m11)
+	TMatrix2x2(T m00, T m01, T m10, T m11)
 	{
 		M[0][0] = m00; M[0][1] = m01;
 		M[1][0] = m10; M[1][1] = m11;
@@ -351,35 +416,35 @@ public:
 
 
 	/** Ctor. initialize from a scale. */
-	explicit FMatrix2x2(float UniformScale)
+	explicit TMatrix2x2(T UniformScale)
 	{
 		M[0][0] = UniformScale; M[0][1] = 0;
 		M[1][0] = 0; M[1][1] = UniformScale;
 	}
 
 	/** Ctor. initialize from a scale. */
-	explicit FMatrix2x2(const FScale2D& Scale)
+	explicit TMatrix2x2(const TScale2<T>& Scale)
 	{
-		float ScaleX = (float)Scale.GetVector().X;
-		float ScaleY = (float)Scale.GetVector().Y;
+		T ScaleX = (T)Scale.GetVector().X;
+		T ScaleY = (T)Scale.GetVector().Y;
 		M[0][0] = ScaleX; M[0][1] = 0;
 		M[1][0] = 0; M[1][1] = ScaleY;
 	}
 
 	/** Factory function. initialize from a 2D shear. */
-	explicit FMatrix2x2(const FShear2D& Shear)
+	explicit TMatrix2x2(const TShear2<T>& Shear)
 	{
-		float XX = (float)Shear.GetVector().X;
-		float YY = (float)Shear.GetVector().Y;
+		T XX = (T)Shear.GetVector().X;
+		T YY = (T)Shear.GetVector().Y;
 		M[0][0] = 1; M[0][1] =YY;
 		M[1][0] =XX; M[1][1] = 1;
 	}
 
 	/** Ctor. initialize from a rotation. */
-	explicit FMatrix2x2(const FQuat2D& Rotation)
+	explicit TMatrix2x2(const FQuat2D& Rotation)
 	{
-		float CosAngle = (float)Rotation.GetVector().X;
-		float SinAngle = (float)Rotation.GetVector().Y;
+		T CosAngle = (T)Rotation.GetVector().X;
+		T SinAngle = (T)Rotation.GetVector().Y;
 		M[0][0] = CosAngle; M[0][1] = SinAngle;
 		M[1][0] = -SinAngle; M[1][1] = CosAngle;
 	}
@@ -389,16 +454,18 @@ public:
 	 *    [X Y] * [m00 m01]
 	 *            [m10 m11]
 	 */
-	FVector2D TransformPoint(const FVector2D& Point) const
+	template<typename ArgType>
+	UE::Math::TVector2<ArgType> TransformPoint(const UE::Math::TVector2<ArgType>& Point) const
 	{
-		return FVector2D(
-			Point.X * M[0][0] + Point.Y * M[1][0],
-			Point.X * M[0][1] + Point.Y * M[1][1]);
+		return UE::Math::TVector2<ArgType>(
+			Point.X * (ArgType)M[0][0] + Point.Y * (ArgType)M[1][0],
+			Point.X * (ArgType)M[0][1] + Point.Y * (ArgType)M[1][1]);
 	}
 	/**
 	 * Vector transformation is equivalent to point transformation as our matrix is not homogeneous.
 	 */
-	FVector2D TransformVector(const FVector2D& Vector) const
+	template<typename ArgType>
+	UE::Math::TVector2<ArgType> TransformVector(const UE::Math::TVector2<ArgType>& Vector) const
 	{
 		return TransformPoint(Vector);
 	}
@@ -407,35 +474,35 @@ public:
 	 * [A B] * [E F] == [AE+BG AF+BH]
 	 * [C D]   [G H]    [CE+DG CF+DH]
 	 */
-	FMatrix2x2 Concatenate(const FMatrix2x2& RHS) const
+	TMatrix2x2 Concatenate(const TMatrix2x2& RHS) const
 	{
-		float A, B, C, D;
+		T A, B, C, D;
 		GetMatrix(A, B, C, D);
-		float E, F, G, H;
+		T E, F, G, H;
 		RHS.GetMatrix(E, F, G, H);
-		return FMatrix2x2(
+		return TMatrix2x2(
 			A*E + B*G, A*F + B*H,
 			C*E + D*G, C*F + D*H);
 	}
 	/**
 	 * Invert the transform.
 	 */
-	FMatrix2x2 Inverse() const
+	TMatrix2x2 Inverse() const
 	{
-		float A, B, C, D;
+		T A, B, C, D;
 		GetMatrix(A, B, C, D);
-		float InvDet = InverseDeterminant();
-		return FMatrix2x2(
+		T InvDet = InverseDeterminant();
+		return TMatrix2x2(
 			 D*InvDet, -B*InvDet,
 			-C*InvDet,  A*InvDet);
 	}
 
 	/** Equality. */
-	bool operator==(const FMatrix2x2& RHS) const
+	bool operator==(const TMatrix2x2& RHS) const
 	{
-		float A, B, C, D;
+		T A, B, C, D;
 		GetMatrix(A, B, C, D);
-		float E, F, G, H;
+		T E, F, G, H;
 		RHS.GetMatrix(E, F, G, H);
 		return
 			FMath::IsNearlyEqual(A, E, KINDA_SMALL_NUMBER) &&
@@ -445,50 +512,56 @@ public:
 	}
 
 	/** Inequality. */
-	bool operator!=(const FMatrix2x2& Other) const
+	bool operator!=(const TMatrix2x2& Other) const
 	{
 		return !operator==(Other);
 	}
 
-	void GetMatrix(float &A, float &B, float &C, float &D) const
+	void GetMatrix(float& A, float& B, float& C, float& D) const
+	{
+		A = (float)M[0][0]; B = (float)M[0][1];
+		C = (float)M[1][0]; D = (float)M[1][1];
+	}
+
+	void GetMatrix(double& A, double& B, double& C, double& D) const
 	{
-		A = M[0][0]; B = M[0][1];
-		C = M[1][0]; D = M[1][1];
+		A = (double)M[0][0]; B = (double)M[0][1];
+		C = (double)M[1][0]; D = (double)M[1][1];
 	}
 
-	float Determinant() const
+	T Determinant() const
 	{
-		float A, B, C, D;
+		T A, B, C, D;
 		GetMatrix(A, B, C, D);
 		return (A*D - B*C);
 	}
 
-	float InverseDeterminant() const
+	T InverseDeterminant() const
 	{
-		float Det = Determinant();
+		T Det = Determinant();
 		checkSlow(Det != 0.0f);
 		return 1.0f / Det;
 	}
 
 	/** Extracts the squared scale from the matrix (avoids sqrt). */
-	FScale2D GetScaleSquared() const
+	TScale2<T> GetScaleSquared() const
 	{
-		float A, B, C, D;
+		T A, B, C, D;
 		GetMatrix(A, B, C, D);
-		return FScale2D(A*A + B*B, C*C + D*D);
+		return TScale2<T>(A*A + B*B, C*C + D*D);
 	}
 
 	/** Gets the scale from the matrix. */
-	FScale2D GetScale() const
+	TScale2<T> GetScale() const
 	{
-		FScale2D ScaleSquared = GetScaleSquared();
-		return FScale2D(FMath::Sqrt(ScaleSquared.GetVector().X), FMath::Sqrt(ScaleSquared.GetVector().Y));
+		TScale2<T> ScaleSquared = GetScaleSquared();
+		return TScale2<T>(FMath::Sqrt(ScaleSquared.GetVector().X), FMath::Sqrt(ScaleSquared.GetVector().Y));
 	}
 
 	/** Gets the rotation angle of the matrix. */
-	float GetRotationAngle() const
+	T GetRotationAngle() const
 	{
-		float A, B, C, D;
+		T A, B, C, D;
 		GetMatrix(A, B, C, D);
 		return FMath::Atan(C / D);
 	}
@@ -500,7 +573,7 @@ public:
 			&& M[1][0] == 0.0f && M[1][1] == 1.0f;
 	}
 
-	bool IsNearlyIdentity(float ErrorTolerance = KINDA_SMALL_NUMBER) const
+	bool IsNearlyIdentity(T ErrorTolerance = KINDA_SMALL_NUMBER) const
 	{
 		return
 			FMath::IsNearlyEqual(M[0][0], 1.0f, ErrorTolerance) &&
@@ -510,40 +583,55 @@ public:
 	}
 
 private:
-	float M[2][2];
+	T M[2][2];
 };
 
-inline FMatrix2x2 FShear2D::Concatenate(const FShear2D& RHS) const
+/** Base typedefs */
+typedef TMatrix2x2<float> FMatrix2x2f;
+typedef TMatrix2x2<double> FMatrix2x2d;
+typedef FMatrix2x2f FMatrix2x2; // Default type (for backwards compat)
+
+
+template<typename T>
+inline TMatrix2x2<T> TShear2<T>::Concatenate(const TShear2<T>& RHS) const
 {
-	float XXA = (float)Shear.X;
-	float YYA = (float)Shear.Y;
-	float XXB = (float)RHS.Shear.X;
-	float YYB = (float)RHS.Shear.Y;
-	return FMatrix2x2(
+	T XXA = (T)Shear.X;
+	T YYA = (T)Shear.Y;
+	T XXB = (T)RHS.Shear.X;
+	T YYB = (T)RHS.Shear.Y;
+	return TMatrix2x2<T>(
 		1+YYA*XXB, YYB*YYA,
 		XXA+XXB, XXA*XXB+1);
 }
 
-inline FMatrix2x2 FShear2D::Inverse() const
+template<typename T>
+inline TMatrix2x2<T> TShear2<T>::Inverse() const
 {
-	float InvDet = 1.0f / float(1.0f - Shear.X*Shear.Y);
-	return FMatrix2x2(
-		InvDet, float(-Shear.Y * InvDet),
-		float(-Shear.X * InvDet), InvDet);
+	T InvDet = 1.0f / T(1.0f - Shear.X*Shear.Y);
+	return TMatrix2x2<T>(
+		InvDet, T(-Shear.Y * InvDet),
+		T(-Shear.X * InvDet), InvDet);
 }
 
-/** Partial specialization of ConcatenateRules for 2x2 and any other type via Upcast to 2x2 first. Requires a conversion ctor on FMatrix2x2. Funky template logic so we don't hide the default rule for NULL conversions. */
-template<typename TransformType> struct ConcatenateRules<typename TEnableIf<!TAreTypesEqual<FMatrix2x2, TransformType>::Value, FMatrix2x2>::Type, TransformType  > { typedef FMatrix2x2 ResultType; };
-/** Partial specialization of ConcatenateRules for 2x2 and any other type via Upcast to 2x2 first. Requires a conversion ctor on FMatrix2x2. Funky template logic so we don't hide the default rule for NULL conversions. */
-template<typename TransformType> struct ConcatenateRules<TransformType  , typename TEnableIf<!TAreTypesEqual<FMatrix2x2, TransformType>::Value, FMatrix2x2>::Type> { typedef FMatrix2x2 ResultType; };
-
-/** concatenation rules for 2x2 transform types. Convert to 2x2 matrix as the fully decomposed math is not that perf critical right now. */
-template<> struct ConcatenateRules<FScale2D  , FShear2D  > { typedef FMatrix2x2 ResultType; };
-template<> struct ConcatenateRules<FScale2D  , FQuat2D   > { typedef FMatrix2x2 ResultType; };
-template<> struct ConcatenateRules<FShear2D  , FScale2D  > { typedef FMatrix2x2 ResultType; };
-template<> struct ConcatenateRules<FQuat2D   , FScale2D  > { typedef FMatrix2x2 ResultType; };
-template<> struct ConcatenateRules<FShear2D  , FQuat2D   > { typedef FMatrix2x2 ResultType; };
-template<> struct ConcatenateRules<FQuat2D   , FShear2D  > { typedef FMatrix2x2 ResultType; };
+/** Concatenation rules for Matrix2x2 and any other type, */
+template<typename T> struct ConcatenateRules<float			, TMatrix2x2<T>	> { typedef TMatrix2x2<T> ResultType; };
+template<typename T> struct ConcatenateRules<double			, TMatrix2x2<T>	> { typedef TMatrix2x2<T> ResultType; };
+template<typename T> struct ConcatenateRules<TScale2<T>		, TMatrix2x2<T>	> { typedef TMatrix2x2<T> ResultType; };
+template<typename T> struct ConcatenateRules<TShear2<T>		, TMatrix2x2<T>	> { typedef TMatrix2x2<T> ResultType; };
+template<typename T> struct ConcatenateRules<TQuat2<T>		, TMatrix2x2<T>	> { typedef TMatrix2x2<T> ResultType; };
+template<typename T> struct ConcatenateRules<TMatrix2x2<T>	, float			> { typedef TMatrix2x2<T> ResultType; };
+template<typename T> struct ConcatenateRules<TMatrix2x2<T>	, double		> { typedef TMatrix2x2<T> ResultType; };
+template<typename T> struct ConcatenateRules<TMatrix2x2<T>	, TScale2<T>	> { typedef TMatrix2x2<T> ResultType; };
+template<typename T> struct ConcatenateRules<TMatrix2x2<T>	, TShear2<T>	> { typedef TMatrix2x2<T> ResultType; };
+template<typename T> struct ConcatenateRules<TMatrix2x2<T>	, TQuat2<T>		> { typedef TMatrix2x2<T> ResultType; };
+
+/** Concatenation rules for 2x2 transform types. Convert to 2x2 matrix as the fully decomposed math is not that perf critical right now. */
+template<typename T> struct ConcatenateRules<TScale2<T>  , TShear2<T>  > { typedef TMatrix2x2<T> ResultType; };
+template<typename T> struct ConcatenateRules<TScale2<T>  , TQuat2<T>   > { typedef TMatrix2x2<T> ResultType; };
+template<typename T> struct ConcatenateRules<TShear2<T>  , TScale2<T>  > { typedef TMatrix2x2<T> ResultType; };
+template<typename T> struct ConcatenateRules<TQuat2<T>   , TScale2<T>  > { typedef TMatrix2x2<T> ResultType; };
+template<typename T> struct ConcatenateRules<TShear2<T>  , TQuat2<T>   > { typedef TMatrix2x2<T> ResultType; };
+template<typename T> struct ConcatenateRules<TQuat2<T>   , TShear2<T>  > { typedef TMatrix2x2<T> ResultType; };
 
  /**
  * Support for generalized 2D affine transforms. 
@@ -552,59 +640,74 @@ template<> struct ConcatenateRules<FQuat2D   , FShear2D  > { typedef FMatrix2x2
  *   [C D 0]
  *   [X Y 1]
  */
-class FTransform2D
+template<typename T>
+class TTransform2
 {
+	static_assert(std::is_floating_point_v<T>, "T must be floating point");
+
 public:
+	using FReal = T;
+	using Vector2Type = UE::Math::TVector2<T>;
+	using Matrix2Type = TMatrix2x2<T>;
+
 	/** Initialize the transform using an identity matrix and a translation. */
-	FTransform2D(const FVector2D& Translation = FVector2D(0.f,0.f))
-		: Trans(Translation)
+	template<typename VType=float, typename VArg=UE::Math::TVector2<VType>>
+	TTransform2(const VArg& Translation = VArg(0.f, 0.f))
+		: Trans((Vector2Type)Translation)
 	{
 	}
 
 	/** Initialize the transform using a uniform scale and a translation. */
-	explicit FTransform2D(float UniformScale, const FVector2D& Translation = FVector2D(0.f,0.f))
-		: M(FScale2D(UniformScale)), Trans(Translation)
+	template<typename VType=float, typename VArg=UE::Math::TVector2<VType>>
+	explicit TTransform2(T UniformScale, const VArg& Translation = VArg(0.f, 0.f))
+		: M(TScale2<T>(UniformScale)), Trans((Vector2Type)Translation)
 	{
 	}
 
 	/** Initialize the transform using a 2D scale and a translation. */
-	explicit FTransform2D(const FScale2D& Scale, const FVector2D& Translation = FVector2D(0.f,0.f))
-		: M(Scale), Trans(Translation)
+	template<typename VType=float, typename VArg=UE::Math::TVector2<VType>>
+	explicit TTransform2(const TScale2<T>& Scale, const VArg& Translation = VArg(0.f, 0.f))
+		: M(Scale), Trans((Vector2Type)Translation)
 	{
 	}
 
 	/** Initialize the transform using a 2D shear and a translation. */
-	explicit FTransform2D(const FShear2D& Shear, const FVector2D& Translation = FVector2D(0.f,0.f))
-		: M(Shear), Trans(Translation)
+	template<typename VType=float, typename VArg=UE::Math::TVector2<VType>>
+	explicit TTransform2(const TShear2<T>& Shear, const VArg& Translation = VArg(0.f, 0.f))
+		: M(Shear), Trans((Vector2Type)Translation)
 	{
 	}
 
 	/** Initialize the transform using a 2D rotation and a translation. */
-	explicit FTransform2D(const FQuat2D& Rot, const FVector2D& Translation = FVector2D(0.f,0.f))
-		: M(Rot), Trans(Translation)
+	template<typename VType=float, typename VArg=UE::Math::TVector2<VType>>
+	explicit TTransform2(const TQuat2<T>& Rot, const VArg& Translation = VArg(0.f, 0.f))
+		: M(Rot), Trans((Vector2Type)Translation)
 	{
 	}
 
 	/** Initialize the transform using a general 2x2 transform and a translation. */
-	explicit FTransform2D(const FMatrix2x2& Transform, const FVector2D& Translation = FVector2D(0.f,0.f))
-		: M(Transform), Trans(Translation)
+	template<typename VType=float, typename VArg=UE::Math::TVector2<VType>>
+	explicit TTransform2(const TMatrix2x2<T>& Transform, const VArg& Translation = VArg(0.f, 0.f))
+		: M(Transform), Trans((Vector2Type)Translation)
 	{
 	}
 
 	/**
 	 * 2D transformation of a point.  Transforms position, rotation, and scale.
 	 */
-	FVector2D TransformPoint(const FVector2D& Point) const
+	template<typename VType>
+	UE::Math::TVector2<VType> TransformPoint(const UE::Math::TVector2<VType>& Point) const
 	{
-		return ::TransformPoint(GetTranslation(), ::TransformPoint(M, Point));
+		return (UE::Math::TVector2<VType>) ::TransformPoint(Trans, ::TransformPoint(M, (Vector2Type)Point));
 	}
 
 	/**
 	 * 2D transformation of a vector.  Transforms rotation and scale.
 	 */
-	FVector2D TransformVector(const FVector2D& Vector) const
+	template<typename VType>
+	UE::Math::TVector2<VType> TransformVector(const UE::Math::TVector2<VType>& Vector) const
 	{
-		return ::TransformVector(M, Vector);
+		return (UE::Math::TVector2<VType>) ::TransformVector(M, (Vector2Type)Vector);
 	}
 
 	/** 
@@ -614,11 +717,11 @@ public:
 	 *  NewM == MA * MB
 	 *  NewT == TA * MB + TB
 	 */
-	FTransform2D Concatenate(const FTransform2D& RHS) const
+	TTransform2 Concatenate(const TTransform2& RHS) const
 	{
-		return FTransform2D(
+		return TTransform2(
 			::Concatenate(M, RHS.M),
-			::Concatenate(::TransformPoint(RHS.M, FVector2D(Trans)), FVector2D(RHS.Trans)));
+			::Concatenate(::TransformPoint(RHS.M, Trans), RHS.Trans));
 	}
 
 	/** 
@@ -649,62 +752,70 @@ public:
 	 *   M' == Inverse(M)
 	 *   T' == Inverse(Translate) * Inverse(M)
 	 */
-	FTransform2D Inverse() const
+	TTransform2 Inverse() const
 	{
-		FMatrix2x2 InvM = ::Inverse(M);
-		FVector2D InvTrans = ::TransformPoint(InvM, ::Inverse(GetTranslation()));
-		return FTransform2D(InvM, InvTrans);
+		Matrix2Type InvM = ::Inverse(M);
+		Vector2Type InvTrans = ::TransformPoint(InvM, ::Inverse(Trans));
+		return TTransform2(InvM, InvTrans);
 	}
 	
 	/** Equality. */
-	bool operator==(const FTransform2D& Other) const
+	bool operator==(const TTransform2& Other) const
 	{
 		return M == Other.M && Trans == Other.Trans;
 	}
 	
 	/** Inequality. */
-	bool operator!=(const FTransform2D& Other) const
+	bool operator!=(const TTransform2& Other) const
 	{
 		return !operator==(Other);
 	}
 
 	/** Access to the 2x2 transform */
-	const FMatrix2x2& GetMatrix() const { return M; }
+	const Matrix2Type& GetMatrix() const { return M; }
 	/** Access to the translation */
-	const FVector2D GetTranslation() const { return FVector2D(Trans); }
+	//const Vector2Type GetTranslation() const { return Vector2Type(Trans); }
+	const FVector2D GetTranslation() const { return (FVector2D)Trans; } // TODO: should be Vector2Type in general, but retaining FVector2D for now for compilation backwards compat.
 
-	void SetTranslation(const FVector2D& InTrans) { Trans = FVector2f(InTrans); }
+	template<typename VType>
+	void SetTranslation(const UE::Math::TVector2<VType>& InTrans) { Trans = (Vector2Type)InTrans; }
 
 	/**
 	 * Specialized function to determine if a transform is precisely the identity transform. Uses exact float comparison, so rounding error is not considered.
 	 */
 	bool IsIdentity() const
 	{
-		return M.IsIdentity() && Trans == FVector2f::ZeroVector;
+		return M.IsIdentity() && Trans == Vector2Type::ZeroVector;
 	}
 
 	/**
 	 * Converts this affine 2D Transform into an affine 3D transform.
 	 */
-	FMatrix To3DMatrix() const
+	UE::Math::TMatrix<T> To3DMatrix() const
 	{
-		float A, B, C, D;
+		T A, B, C, D;
 		M.GetMatrix(A, B, C, D);
 
-		return FMatrix(
-			FPlane(      A,      B, 0.0f, 0.0f),
-			FPlane(      C,      D, 0.0f, 0.0f),
-			FPlane(   0.0f,   0.0f, 1.0f, 0.0f),
-			FPlane(Trans.X, Trans.Y, 0.0f, 1.0f)
+		return UE::Math::TMatrix<T>(
+			UE::Math::TPlane<T>(      A,      B, 0.0f, 0.0f),
+			UE::Math::TPlane<T>(      C,      D, 0.0f, 0.0f),
+			UE::Math::TPlane<T>(   0.0f,   0.0f, 1.0f, 0.0f),
+			UE::Math::TPlane<T>(Trans.X, Trans.Y, 0.0f, 1.0f)
 		);
 	}
 
 private:
-	FMatrix2x2 M;
-	FVector2f Trans;
+	Matrix2Type M;
+	Vector2Type Trans;
 };
 
-template<> struct TIsPODType<FTransform2D> { enum { Value = true }; };
+/** Core typedefs */
+typedef TTransform2<float> FTransform2f;
+typedef TTransform2<double> FTransform2d;
+typedef FTransform2f FTransform2D; // default type, for backwards compat
+
+template<> struct TIsPODType<FTransform2f> { enum { Value = true }; };
+template<> struct TIsPODType<FTransform2d> { enum { Value = true }; };
 
 //////////////////////////////////////////////////////////////////////////
 // Concatenate overloads. 
@@ -714,70 +825,90 @@ template<> struct TIsPODType<FTransform2D> { enum { Value = true }; };
 //////////////////////////////////////////////////////////////////////////
 
 /** Specialization for concatenating a 2D scale and 2D Translation. */
-inline FTransform2D Concatenate(const FScale2D& Scale, const FVector2D& Translation)
+template<typename T, typename V>
+inline TTransform2<T> Concatenate(const TScale2<T>& Scale, const UE::Math::TVector2<V>& Translation)
 {
-	return FTransform2D(Scale, Translation);
+	return TTransform2<T>(Scale, Translation);
 }
 
 /** Specialization for concatenating a 2D shear and 2D Translation. */
-inline FTransform2D Concatenate(const FShear2D& Shear, const FVector2D& Translation)
+template<typename T, typename V>
+inline TTransform2<T> Concatenate(const TShear2<T>& Shear, const UE::Math::TVector2<V>& Translation)
 {
-	return FTransform2D(Shear, Translation);
+	return TTransform2<T>(Shear, Translation);
 }
 
 /** Specialization for concatenating 2D Rotation and 2D Translation. */
-inline FTransform2D Concatenate(const FQuat2D& Rot, const FVector2D& Translation)
+template<typename T, typename V>
+inline TTransform2<T> Concatenate(const TQuat2<T>& Rot, const UE::Math::TVector2<V>& Translation)
 {
-	return FTransform2D(Rot, Translation);
+	return TTransform2<T>(Rot, Translation);
 }
 
 /** Specialization for concatenating 2D generalized transform and 2D Translation. */
-inline FTransform2D Concatenate(const FMatrix2x2& Transform, const FVector2D& Translation)
+template<typename T, typename V>
+inline TTransform2<T> Concatenate(const TMatrix2x2<T>& Transform, const UE::Math::TVector2<V>& Translation)
 {
-	return FTransform2D(Transform, Translation);
+	return TTransform2<T>(Transform, Translation);
 }
 
 /** Specialization for concatenating transform and 2D Translation. */
-inline FTransform2D Concatenate(const FTransform2D& Transform, const FVector2D& Translation)
+template<typename T, typename V>
+inline TTransform2<T> Concatenate(const TTransform2<T>& Transform, const UE::Math::TVector2<V>& Translation)
 {
-	return FTransform2D(Transform.GetMatrix(), Concatenate(Transform.GetTranslation(), Translation));
+	return TTransform2<T>(Transform.GetMatrix(), Concatenate((UE::Math::TVector2<T>)Transform.GetTranslation(), (UE::Math::TVector2<T>)Translation));
 }
 
 /** Specialization for concatenating a 2D Translation and 2D scale. */
-inline FTransform2D Concatenate(const FVector2D& Translation, const FScale2D& Scale)
+template<typename T, typename V>
+inline TTransform2<T> Concatenate(const UE::Math::TVector2<V>& Translation, const TScale2<T>& Scale)
 {
-	return FTransform2D(Scale, ::TransformPoint(Scale, Translation));
+	return TTransform2<T>(Scale, ::TransformPoint(Scale, Translation));
 }
 
 /** Specialization for concatenating a 2D Translation and 2D shear. */
-inline FTransform2D Concatenate(const FVector2D& Translation, const FShear2D& Shear)
+template<typename T, typename V>
+inline TTransform2<T> Concatenate(const UE::Math::TVector2<V>& Translation, const TShear2<T>& Shear)
 {
-	return FTransform2D(Shear, ::TransformPoint(Shear, Translation));
+	return TTransform2<T>(Shear, ::TransformPoint(Shear, Translation));
 }
 
 /** Specialization for concatenating 2D Translation and 2D Rotation. */
-inline FTransform2D Concatenate(const FVector2D& Translation, const FQuat2D& Rot)
+template<typename T, typename V>
+inline TTransform2<T> Concatenate(const UE::Math::TVector2<V>& Translation, const TQuat2<T>& Rot)
 {
-	return FTransform2D(Rot, ::TransformPoint(Rot, Translation));
+	return TTransform2<T>(Rot, ::TransformPoint(Rot, Translation));
 }
 
-/** Specialization for concatenating 2D Translation and 2D generalized transform. See docs for FTransform2D::Inverse for details on how this math is derived. */
-inline FTransform2D Concatenate(const FVector2D& Translation, const FMatrix2x2& Transform)
+/** Specialization for concatenating 2D Translation and 2D generalized transform. See docs for TTransform2<T>::Inverse for details on how this math is derived. */
+template<typename T, typename V>
+inline TTransform2<T> Concatenate(const UE::Math::TVector2<V>& Translation, const TMatrix2x2<T>& Transform)
 {
-	return FTransform2D(Transform, ::TransformPoint(Transform, Translation));
+	return TTransform2<T>(Transform, ::TransformPoint(Transform, Translation));
 }
 
-/** Specialization for concatenating 2D Translation and transform. See docs for FTransform2D::Inverse for details on how this math is derived. */
-inline FTransform2D Concatenate(const FVector2D& Translation, const FTransform2D& Transform)
+/** Specialization for concatenating 2D Translation and transform. See docs for TTransform2<T>::Inverse for details on how this math is derived. */
+template<typename T, typename V>
+inline TTransform2<T> Concatenate(const UE::Math::TVector2<V>& Translation, const TTransform2<T>& Transform)
 {
-	return FTransform2D(Transform.GetMatrix(), Concatenate(::TransformPoint(Transform.GetMatrix(), Translation), Transform.GetTranslation()));
+	return TTransform2<T>(Transform.GetMatrix(), Concatenate(::TransformPoint(Transform.GetMatrix(), (UE::Math::TVector2<T>)Translation), (UE::Math::TVector2<T>)Transform.GetTranslation()));
 }
 
-/** Partial specialization of ConcatenateRules for FTransform2D and any other type via Upcast to FTransform2D first. Requires a conversion ctor on FTransform2D. Funky template logic so we don't hide the default rule for NULL conversions. */
-template<typename TransformType> struct ConcatenateRules<typename TEnableIf<!TAreTypesEqual<FTransform2D, TransformType>::Value, FTransform2D>::Type, TransformType  > { typedef FTransform2D ResultType; };
-/** Partial specialization of ConcatenateRules for FTransform2D and any other type via Upcast to FTransform2D first. Requires a conversion ctor on FTransform2D. Funky template logic so we don't hide the default rule for NULL conversions. */
-template<typename TransformType> struct ConcatenateRules<TransformType  , typename TEnableIf<!TAreTypesEqual<FTransform2D, TransformType>::Value, FTransform2D>::Type> { typedef FTransform2D ResultType; };
+/** Helper to determine if a type is based on TTransform2 */
+namespace TransformCalculusHelper
+{
+	template<typename T>
+	struct TIsTransform2
+	{
+		enum { Value = false };
+	};
 
-/** Provide a disambiguating overload for 2x2 and FTransform2D, since both try to provide a generic set of ConcatenateRules for all types. */
-template<> struct ConcatenateRules<FMatrix2x2, FTransform2D> { typedef FTransform2D ResultType; };
-template<> struct ConcatenateRules<FTransform2D, FMatrix2x2> { typedef FTransform2D ResultType; };
+	template<> struct TIsTransform2< FTransform2f > { enum { Value = true }; };
+	template<> struct TIsTransform2< FTransform2d > { enum { Value = true }; };
+}
+
+/** Partial specialization of ConcatenateRules for FTransform2D and any other type via Upcast to FTransform2D first. Requires a conversion ctor on FTransform2D. Funky template logic so we don't hide the default rule for NULL conversions. */
+template<typename TransformType> struct ConcatenateRules<typename TEnableIf<!TransformCalculusHelper::TIsTransform2<TransformType>::Value, FTransform2f>::Type, TransformType	> { typedef FTransform2f ResultType; };
+template<typename TransformType> struct ConcatenateRules<typename TEnableIf<!TransformCalculusHelper::TIsTransform2<TransformType>::Value, FTransform2d>::Type, TransformType	> { typedef FTransform2d ResultType; };
+template<typename TransformType> struct ConcatenateRules<TransformType, typename TEnableIf<!TransformCalculusHelper::TIsTransform2<TransformType>::Value, FTransform2f>::Type	> { typedef FTransform2f ResultType; };
+template<typename TransformType> struct ConcatenateRules<TransformType, typename TEnableIf<!TransformCalculusHelper::TIsTransform2<TransformType>::Value, FTransform2d>::Type	> { typedef FTransform2d ResultType; };
diff --git a/Engine/Source/Runtime/Core/Public/Math/TransformCalculus3D.h b/Engine/Source/Runtime/Core/Public/Math/TransformCalculus3D.h
--- a/Engine/Source/Runtime/Core/Public/Math/TransformCalculus3D.h
+++ b/Engine/Source/Runtime/Core/Public/Math/TransformCalculus3D.h
@@ -19,55 +19,70 @@
 // requires by the transform calculus framework.
 //
 // The following types are adapted.
-// * float           -> represents a uniform scale.
-// * FScale          -> represents a 3D non-uniform scale.
-// * FVector         -> represents a 3D translation.
-// * FRotator        -> represents a pure rotation.
-// * FQuat           -> represents a pure rotation.
-// * FMatrix         -> represents a general 3D homogeneous transform.
+// * float/double			-> represents a uniform scale.
+// * TScale<T>				-> represents a 3D non-uniform scale.
+// * UE::Math::TVector<T>	-> represents a 3D translation.
+// * UE::Math::TRotator<T>	-> represents a pure rotation.
+// * UE::Math::TQuat<T>		-> represents a pure rotation.
+// * UE::Math::TMatrix<T>	-> represents a general 3D homogeneous transform.
 //
 //////////////////////////////////////////////////////////////////////////
 
 /**
- * Represents a 3D non-uniform scale (to disambiguate from an FVector, which is used for translation).
+ * Represents a 3D non-uniform scale (to disambiguate from an UE::Math::TVector<T>, which is used for translation).
  * 
  * Serves as a good base example of how to write a class that supports the basic transform calculus
  * operations.
  */
-class FScale
+template<typename T>
+class TScale
 {
+	static_assert(std::is_floating_point_v<T>, "T must be floating point");
+
 public:
+	using Vector3Type = UE::Math::TVector<T>;
+
 	/** Ctor. initialize to an identity scale, 1.0. */
-	FScale() :Scale(1.0f) {}
+	TScale() :Scale(1.0f) {}
 	/** Ctor. initialize from a uniform scale. */
-	explicit FScale(float InScale) :Scale(InScale) {}
-	/** Ctor. initialize from an FVector defining the 3D scale. */
-	explicit FScale(const FVector& InScale) :Scale(InScale) {}
-	/** Access to the underlying FVector that stores the scale. */
-	const FVector& GetVector() const { return Scale; }
+	explicit TScale(T InScale) :Scale(InScale) {}
+	/** Ctor. initialize from an UE::Math::TVector<T> defining the 3D scale. */
+	template<typename VType>
+	explicit TScale(const UE::Math::TVector<VType>& InScale) : Scale((Vector3Type)InScale) {}
+	/** Access to the underlying UE::Math::TVector<T> that stores the scale. */
+	const Vector3Type& GetVector() const { return Scale; }
 	/** Concatenate two scales. */
-	const FScale Concatenate(const FScale& RHS) const
+	const TScale Concatenate(const TScale& RHS) const
 	{
-		return FScale(Scale * RHS.GetVector());
+		return TScale(Scale * RHS.GetVector());
 	}
 	/** Invert the scale. */
-	const FScale Inverse() const
+	const TScale Inverse() const
 	{
-		return FScale(FVector(1.0f / Scale.X, 1.0f / Scale.Y, 1.0f / Scale.Z));
+		return TScale(Vector3Type(1.0f / Scale.X, 1.0f / Scale.Y, 1.0f / Scale.Z));
 	}
 private:
 	/** Underlying storage of the 3D scale. */
-	FVector Scale;
+	Vector3Type Scale;
 };
 
+/** Core typedefs */
+typedef TScale<float> FScale3f;
+typedef TScale<double> FScale3d;
+typedef FScale3f FScale3D; // Default type, for backwards compat
+
 /** Specialization for converting a FMatrix to an FRotator. It uses a non-standard explicit conversion function. */
-template<> template<> inline FRotator TransformConverter<FRotator>::Convert<FMatrix>(const FMatrix& Transform)
+template<> template<> inline FRotator3f TransformConverter<FRotator3f>::Convert<FMatrix44f>(const FMatrix44f& Transform)
+{
+	return Transform.Rotator();
+}
+template<> template<> inline FRotator3d TransformConverter<FRotator3d>::Convert<FMatrix44d>(const FMatrix44d& Transform)
 {
 	return Transform.Rotator();
 }
 
 //////////////////////////////////////////////////////////////////////////
-// FMatrix Support
+// UE::Math::TMatrix<T> Support
 //////////////////////////////////////////////////////////////////////////
 
 /**
@@ -75,7 +90,7 @@ template<> template<> inline FRotator TransformConverter<FRotator>::Convert<FMat
  * Uses decltype to allow some classes to return const-ref types for efficiency.
  * 
  * @param Transform
- * @return the FMatrix stored by the Transform.
+ * @return the UE::Math::TMatrix<T> stored by the Transform.
  */
 template<typename TransformType>
 inline auto ToMatrix(const TransformType& Transform) -> decltype(Transform.ToMatrix())
@@ -89,7 +104,12 @@ inline auto ToMatrix(const TransformType& Transform) -> decltype(Transform.ToMat
  * @param Scale Uniform Scale
  * @return Matrix that represents the uniform Scale space.
  */
-inline const FMatrix& ToMatrix(const FMatrix& Transform)
+inline const UE::Math::TMatrix<float>& ToMatrix(const UE::Math::TMatrix<float>& Transform)
+{
+	return Transform;
+}
+
+inline const UE::Math::TMatrix<double>& ToMatrix(const UE::Math::TMatrix<double>& Transform)
 {
 	return Transform;
 }
@@ -100,9 +120,14 @@ inline const FMatrix& ToMatrix(const FMatrix& Transform)
  * @param Scale Uniform Scale
  * @return Matrix that represents the uniform Scale space.
  */
-inline FMatrix ToMatrix(float Scale)
+inline UE::Math::TMatrix<float> ToMatrix(float Scale)
 {
-	return FScaleMatrix(Scale);
+	return UE::Math::TScaleMatrix<float>(Scale);
+}
+
+inline UE::Math::TMatrix<double> ToMatrix(double Scale)
+{
+	return UE::Math::TScaleMatrix<double>(Scale);
 }
 
 /**
@@ -111,9 +136,13 @@ inline FMatrix ToMatrix(float Scale)
  * @param Scale Non-uniform Scale
  * @return Matrix that represents the non-uniform Scale space.
  */
-inline FMatrix ToMatrix(const FScale& Scale)
+inline UE::Math::TMatrix<float> ToMatrix(const TScale<float>& Scale)
+{
+	return UE::Math::TScaleMatrix<float>(UE::Math::TVector<float>(Scale.GetVector()));
+}
+inline UE::Math::TMatrix<double> ToMatrix(const TScale<double>& Scale)
 {
-	return FScaleMatrix(Scale.GetVector());
+	return UE::Math::TScaleMatrix<double>(UE::Math::TVector<double>(Scale.GetVector()));
 }
 
 /**
@@ -122,9 +151,13 @@ inline FMatrix ToMatrix(const FScale& Scale)
  * @param Translation Translation
  * @return Matrix that represents the translated space.
  */
-inline FMatrix ToMatrix(const FVector& Translation)
+inline UE::Math::TMatrix<float> ToMatrix(const UE::Math::TVector<float>& Translation)
+{
+	return UE::Math::TTranslationMatrix<float>(UE::Math::TVector<float>(Translation));
+}
+inline UE::Math::TMatrix<double> ToMatrix(const UE::Math::TVector<double>& Translation)
 {
-	return FTranslationMatrix(Translation);
+	return UE::Math::TTranslationMatrix<double>(UE::Math::TVector<double>(Translation));
 }
 
 /**
@@ -133,9 +166,13 @@ inline FMatrix ToMatrix(const FVector& Translation)
  * @param Rotation Rotation
  * @return Matrix that represents the rotated space.
  */
-inline FMatrix ToMatrix(const FRotator& Rotation)
+inline UE::Math::TMatrix<float> ToMatrix(const UE::Math::TRotator<float>& Rotation)
 {
-	return FRotationMatrix(Rotation);
+	return UE::Math::TRotationMatrix<float>(UE::Math::TRotator<float>(Rotation));
+}
+inline UE::Math::TMatrix<double> ToMatrix(const UE::Math::TRotator<double>& Rotation)
+{
+	return UE::Math::TRotationMatrix<double>(UE::Math::TRotator<double>(Rotation));
 }
 
 /**
@@ -144,58 +181,82 @@ inline FMatrix ToMatrix(const FRotator& Rotation)
  * @param Rotation Rotation
  * @return Matrix that represents the rotated space.
  */
-inline FMatrix ToMatrix(const FQuat& Rotation)
+inline UE::Math::TMatrix<float> ToMatrix(const UE::Math::TQuat<float>& Rotation)
+{
+	return UE::Math::TRotationMatrix<float>::Make(UE::Math::TQuat<float>(Rotation));
+}
+inline UE::Math::TMatrix<double> ToMatrix(const UE::Math::TQuat<double>& Rotation)
 {
-	return FRotationMatrix::Make(Rotation);
+	return UE::Math::TRotationMatrix<double>::Make(UE::Math::TQuat<double>(Rotation));
 }
 
 /**
- * Specialization of TransformConverter for FMatrix. Calls ToMatrix() by default.
+ * Specialization of TransformConverter for UE::Math::TMatrix<T>. Calls ToMatrix() by default.
  * Allows custom types to easily provide support via a ToMatrix() overload or a ToMatrix() member function.
- * Uses decltype to support efficient passthrough of classes that can convert to a FMatrix without creating
+ * Uses decltype to support efficient passthrough of classes that can convert to a UE::Math::TMatrix<T> without creating
  * a new instance.
  */
 template<>
-struct TransformConverter<FMatrix>
+struct TransformConverter<FMatrix44f>
+{
+	template<typename OtherTransformType>
+	static auto Convert(const OtherTransformType& Transform) -> decltype((FMatrix44f)ToMatrix(Transform))
+	{
+		return (FMatrix44f)ToMatrix(Transform);
+	}
+};
+
+template<>
+struct TransformConverter<FMatrix44d>
 {
 	template<typename OtherTransformType>
-	static auto Convert(const OtherTransformType& Transform) -> decltype(ToMatrix(Transform))
+	static auto Convert(const OtherTransformType& Transform) -> decltype((FMatrix44d)ToMatrix(Transform))
 	{
-		return ToMatrix(Transform);
+		return (FMatrix44d)ToMatrix(Transform);
 	}
 };
 
 /** concatenation rules for basic UE4 types. */
-template<> struct ConcatenateRules<float        , FScale       > { typedef FScale ResultType; };
-template<> struct ConcatenateRules<FScale       , float        > { typedef FScale ResultType; };
-template<> struct ConcatenateRules<float        , FVector      > { typedef FMatrix ResultType; };
-template<> struct ConcatenateRules<FVector      , float        > { typedef FMatrix ResultType; };
-template<> struct ConcatenateRules<float        , FRotator     > { typedef FMatrix ResultType; };
-template<> struct ConcatenateRules<FRotator     , float        > { typedef FMatrix ResultType; };
-template<> struct ConcatenateRules<float        , FQuat        > { typedef FMatrix ResultType; };
-template<> struct ConcatenateRules<FQuat        , float        > { typedef FMatrix ResultType; };
-template<> struct ConcatenateRules<float        , FMatrix      > { typedef FMatrix ResultType; };
-template<> struct ConcatenateRules<FMatrix      , float        > { typedef FMatrix ResultType; };
-template<> struct ConcatenateRules<FScale       , FVector      > { typedef FMatrix ResultType; };
-template<> struct ConcatenateRules<FVector      , FScale       > { typedef FMatrix ResultType; };
-template<> struct ConcatenateRules<FScale       , FRotator     > { typedef FMatrix ResultType; };
-template<> struct ConcatenateRules<FRotator     , FScale       > { typedef FMatrix ResultType; };
-template<> struct ConcatenateRules<FScale       , FQuat        > { typedef FMatrix ResultType; };
-template<> struct ConcatenateRules<FQuat        , FScale       > { typedef FMatrix ResultType; };
-template<> struct ConcatenateRules<FScale       , FMatrix      > { typedef FMatrix ResultType; };
-template<> struct ConcatenateRules<FMatrix      , FScale       > { typedef FMatrix ResultType; };
-template<> struct ConcatenateRules<FVector      , FRotator     > { typedef FMatrix ResultType; };
-template<> struct ConcatenateRules<FRotator     , FVector      > { typedef FMatrix ResultType; };
-template<> struct ConcatenateRules<FVector      , FQuat        > { typedef FMatrix ResultType; };
-template<> struct ConcatenateRules<FQuat        , FVector      > { typedef FMatrix ResultType; };
-template<> struct ConcatenateRules<FVector      , FMatrix      > { typedef FMatrix ResultType; };
-template<> struct ConcatenateRules<FMatrix      , FVector      > { typedef FMatrix ResultType; };
-template<> struct ConcatenateRules<FRotator     , FQuat        > { typedef FQuat ResultType; };
-template<> struct ConcatenateRules<FQuat        , FRotator     > { typedef FQuat ResultType; };
-template<> struct ConcatenateRules<FRotator     , FMatrix      > { typedef FMatrix ResultType; };
-template<> struct ConcatenateRules<FMatrix      , FRotator     > { typedef FMatrix ResultType; };
-template<> struct ConcatenateRules<FQuat        , FMatrix      > { typedef FMatrix ResultType; };
-template<> struct ConcatenateRules<FMatrix      , FQuat        > { typedef FMatrix ResultType; };
+template<typename T> struct ConcatenateRules<float					, TScale<T>				> { typedef TScale<T> ResultType; };
+template<typename T> struct ConcatenateRules<double					, TScale<T>				> { typedef TScale<T> ResultType; };
+template<typename T> struct ConcatenateRules<TScale<T>				, float					> { typedef TScale<T> ResultType; };
+template<typename T> struct ConcatenateRules<TScale<T>				, double				> { typedef TScale<T> ResultType; };
+template<typename T> struct ConcatenateRules<float					, UE::Math::TVector<T>	> { typedef UE::Math::TMatrix<T> ResultType; };
+template<typename T> struct ConcatenateRules<double					, UE::Math::TVector<T>	> { typedef UE::Math::TMatrix<T> ResultType; };
+template<typename T> struct ConcatenateRules<UE::Math::TVector<T>	, float					> { typedef UE::Math::TMatrix<T> ResultType; };
+template<typename T> struct ConcatenateRules<UE::Math::TVector<T>	, double				> { typedef UE::Math::TMatrix<T> ResultType; };
+template<typename T> struct ConcatenateRules<float					, UE::Math::TRotator<T>	> { typedef UE::Math::TMatrix<T> ResultType; };
+template<typename T> struct ConcatenateRules<double					, UE::Math::TRotator<T>	> { typedef UE::Math::TMatrix<T> ResultType; };
+template<typename T> struct ConcatenateRules<UE::Math::TRotator<T>	, float					> { typedef UE::Math::TMatrix<T> ResultType; };
+template<typename T> struct ConcatenateRules<UE::Math::TRotator<T>	, double				> { typedef UE::Math::TMatrix<T> ResultType; };
+template<typename T> struct ConcatenateRules<float					, UE::Math::TQuat<T>	> { typedef UE::Math::TMatrix<T> ResultType; };
+template<typename T> struct ConcatenateRules<double					, UE::Math::TQuat<T>	> { typedef UE::Math::TMatrix<T> ResultType; };
+template<typename T> struct ConcatenateRules<UE::Math::TQuat<T>		, float					> { typedef UE::Math::TMatrix<T> ResultType; };
+template<typename T> struct ConcatenateRules<UE::Math::TQuat<T>		, double				> { typedef UE::Math::TMatrix<T> ResultType; };
+template<typename T> struct ConcatenateRules<float					, UE::Math::TMatrix<T>	> { typedef UE::Math::TMatrix<T> ResultType; };
+template<typename T> struct ConcatenateRules<double					, UE::Math::TMatrix<T>	> { typedef UE::Math::TMatrix<T> ResultType; };
+template<typename T> struct ConcatenateRules<UE::Math::TMatrix<T>	, float					> { typedef UE::Math::TMatrix<T> ResultType; };
+template<typename T> struct ConcatenateRules<UE::Math::TMatrix<T>	, double				> { typedef UE::Math::TMatrix<T> ResultType; };
+template<typename T> struct ConcatenateRules<TScale<T>				, UE::Math::TVector<T>  > { typedef UE::Math::TMatrix<T> ResultType; };
+template<typename T> struct ConcatenateRules<UE::Math::TVector<T>	, TScale<T>				> { typedef UE::Math::TMatrix<T> ResultType; };
+template<typename T> struct ConcatenateRules<TScale<T>				, UE::Math::TRotator<T> > { typedef UE::Math::TMatrix<T> ResultType; };
+template<typename T> struct ConcatenateRules<UE::Math::TRotator<T>	, TScale<T>				> { typedef UE::Math::TMatrix<T> ResultType; };
+template<typename T> struct ConcatenateRules<TScale<T>				, UE::Math::TQuat<T>    > { typedef UE::Math::TMatrix<T> ResultType; };
+template<typename T> struct ConcatenateRules<UE::Math::TQuat<T>		, TScale<T>				> { typedef UE::Math::TMatrix<T> ResultType; };
+template<typename T> struct ConcatenateRules<TScale<T>				, UE::Math::TMatrix<T>  > { typedef UE::Math::TMatrix<T> ResultType; };
+template<typename T> struct ConcatenateRules<UE::Math::TMatrix<T>	, TScale<T>				> { typedef UE::Math::TMatrix<T> ResultType; };
+template<typename T> struct ConcatenateRules<UE::Math::TVector<T>	, UE::Math::TRotator<T> > { typedef UE::Math::TMatrix<T> ResultType; };
+template<typename T> struct ConcatenateRules<UE::Math::TRotator<T>	, UE::Math::TVector<T>  > { typedef UE::Math::TMatrix<T> ResultType; };
+template<typename T> struct ConcatenateRules<UE::Math::TVector<T>	, UE::Math::TQuat<T>    > { typedef UE::Math::TMatrix<T> ResultType; };
+template<typename T> struct ConcatenateRules<UE::Math::TQuat<T>		, UE::Math::TVector<T>  > { typedef UE::Math::TMatrix<T> ResultType; };
+template<typename T> struct ConcatenateRules<UE::Math::TVector<T>	, UE::Math::TMatrix<T>  > { typedef UE::Math::TMatrix<T> ResultType; };
+template<typename T> struct ConcatenateRules<UE::Math::TMatrix<T>	, UE::Math::TVector<T>  > { typedef UE::Math::TMatrix<T> ResultType; };
+template<typename T> struct ConcatenateRules<UE::Math::TRotator<T>	, UE::Math::TQuat<T>    > { typedef UE::Math::TQuat<T> ResultType; };
+template<typename T> struct ConcatenateRules<UE::Math::TQuat<T>		, UE::Math::TRotator<T> > { typedef UE::Math::TQuat<T> ResultType; };
+template<typename T> struct ConcatenateRules<UE::Math::TRotator<T>	, UE::Math::TMatrix<T>  > { typedef UE::Math::TMatrix<T> ResultType; };
+template<typename T> struct ConcatenateRules<UE::Math::TMatrix<T>	, UE::Math::TRotator<T> > { typedef UE::Math::TMatrix<T> ResultType; };
+template<typename T> struct ConcatenateRules<UE::Math::TQuat<T>		, UE::Math::TMatrix<T>  > { typedef UE::Math::TMatrix<T> ResultType; };
+template<typename T> struct ConcatenateRules<UE::Math::TMatrix<T>	, UE::Math::TQuat<T>    > { typedef UE::Math::TMatrix<T> ResultType; };
 
 //////////////////////////////////////////////////////////////////////////
 // Concatenate overloads. 
@@ -217,7 +278,8 @@ namespace Math
  * @param RHS rotation that goes from space B to space C.
  * @return a new rotation representing the transformation from the input space of LHS to the output space of RHS.
  */
-inline FMatrix Concatenate(const FMatrix& LHS, const FMatrix& RHS)
+template<typename T>
+inline UE::Math::TMatrix<T> Concatenate(const UE::Math::TMatrix<T>& LHS, const UE::Math::TMatrix<T>& RHS)
 {
 	return LHS * RHS;
 }
@@ -230,7 +292,8 @@ inline FMatrix Concatenate(const FMatrix& LHS, const FMatrix& RHS)
 * @param RHS Translation that goes from space B to space C.
 * @return a new Translation representing the transformation from the input space of LHS to the output space of RHS.
 */
-inline FVector Concatenate(const FVector& LHS, const FVector& RHS)
+template<typename T>
+inline UE::Math::TVector<T> Concatenate(const UE::Math::TVector<T>& LHS, const UE::Math::TVector<T>& RHS)
 {
 	return LHS + RHS;
 }
@@ -239,14 +302,15 @@ inline FVector Concatenate(const FVector& LHS, const FVector& RHS)
 /**
  * Specialization for concatenating two rotations.
  *
- * NOTE: FQuat concatenates right to left, opposite of how FMatrix implements it.
+ * NOTE: UE::Math::TQuat<T> concatenates right to left, opposite of how UE::Math::TMatrix<T> implements it.
  *       Confusing, no? That's why we have these high level functions!
  * 
  * @param LHS rotation that goes from space A to space B
  * @param RHS rotation that goes from space B to space C.
  * @return a new rotation representing the transformation from the input space of LHS to the output space of RHS.
  */
-inline FQuat Concatenate(const FQuat& LHS, const FQuat& RHS)
+template<typename T>
+inline UE::Math::TQuat<T> Concatenate(const UE::Math::TQuat<T>& LHS, const UE::Math::TQuat<T>& RHS)
 {
 	return RHS * LHS;
 }
@@ -261,10 +325,11 @@ inline FQuat Concatenate(const FQuat& LHS, const FQuat& RHS)
  * @param RHS rotation that goes from space B to space C.
  * @return a new rotation representing the transformation from the input space of LHS to the output space of RHS.
  */
-inline FRotator Concatenate(const FRotator& LHS, const FRotator& RHS)
+template<typename T>
+inline UE::Math::TRotator<T> Concatenate(const UE::Math::TRotator<T>& LHS, const UE::Math::TRotator<T>& RHS)
 {
 	//@todo implement a more efficient way to do this.
-	return TransformCast<FRotator>(Concatenate(TransformCast<FMatrix>(LHS), TransformCast<FMatrix>(RHS)));
+	return TransformCast<UE::Math::TRotator<T>>(Concatenate(TransformCast<UE::Math::TMatrix<T>>(LHS), TransformCast<UE::Math::TMatrix<T>>(RHS)));
 }
 
 
@@ -277,37 +342,40 @@ inline FRotator Concatenate(const FRotator& LHS, const FRotator& RHS)
 
 /**
  * Inverts a transform from space A to space B so it transforms from space B to space A.
- * Specialization for FMatrix.
+ * Specialization for UE::Math::TMatrix<T>.
  * 
  * @param Transform Input transform from space A to space B.
  * @return Inverted transform from space B to space A.
  */
-inline FMatrix Inverse(const FMatrix& Transform)
+template<typename T>
+inline UE::Math::TMatrix<T> Inverse(const UE::Math::TMatrix<T>& Transform)
 {
 	return Transform.Inverse();
 }
 
 /**
  * Inverts a transform from space A to space B so it transforms from space B to space A.
- * Specialization for FRotator.
+ * Specialization for UE::Math::TRotator<T>.
  * 
  * @param Transform Input transform from space A to space B.
  * @return Inverted transform from space B to space A.
  */
-inline FRotator Inverse(const FRotator& Transform)
+template<typename T>
+inline UE::Math::TRotator<T> Inverse(const UE::Math::TRotator<T>& Transform)
 {
-	FVector EulerAngles = Transform.Euler();
-	return FRotator::MakeFromEuler(FVector(-EulerAngles.Z, -EulerAngles.Y, -EulerAngles.X));
+	UE::Math::TVector<T> EulerAngles = Transform.Euler();
+	return UE::Math::TRotator<T>::MakeFromEuler(UE::Math::TVector<T>(-EulerAngles.Z, -EulerAngles.Y, -EulerAngles.X));
 }
 
 /**
  * Inverts a transform from space A to space B so it transforms from space B to space A.
- * Specialization for FQuat.
+ * Specialization for UE::Math::TQuat<T>.
  * 
  * @param Transform Input transform from space A to space B.
  * @return Inverted transform from space B to space A.
  */
-inline FQuat Inverse(const FQuat& Transform)
+template<typename T>
+inline UE::Math::TQuat<T> Inverse(const UE::Math::TQuat<T>& Transform)
 {
 	return Transform.Inverse();
 }
@@ -319,7 +387,8 @@ inline FQuat Inverse(const FQuat& Transform)
  * @param Transform Input transform from space A to space B.
  * @return Inverted transform from space B to space A.
  */
-inline FVector Inverse(const FVector& Transform)
+template<typename T>
+inline UE::Math::TVector<T> Inverse(const UE::Math::TVector<T>& Transform)
 {
 	return -Transform;
 }
@@ -332,57 +401,64 @@ inline FVector Inverse(const FVector& Transform)
 //////////////////////////////////////////////////////////////////////////
 
 /**
- * Specialization for FMatrix as it's member function is called something slightly different.
+ * Specialization for UE::Math::TMatrix<T> as it's member function is called something slightly different.
  */
-inline FVector TransformPoint(const FMatrix& Transform, const FVector& Point)
+template<typename T>
+inline UE::Math::TVector<T> TransformPoint(const UE::Math::TMatrix<T>& Transform, const UE::Math::TVector<T>& Point)
 {
 	return Transform.TransformPosition(Point);
 }
 
 /**
- * Specialization for FQuat as it's member function is called something slightly different.
+ * Specialization for UE::Math::TQuat<T> as it's member function is called something slightly different.
  */
-inline FVector TransformPoint(const FQuat& Transform, const FVector& Point)
+template<typename T>
+inline UE::Math::TVector<T> TransformPoint(const UE::Math::TQuat<T>& Transform, const UE::Math::TVector<T>& Point)
 {
 	return Transform.RotateVector(Point);
 }
 
 /**
- * Specialization for FQuat as it's member function is called something slightly different.
+ * Specialization for UE::Math::TQuat<T> as it's member function is called something slightly different.
  */
-inline FVector TransformVector(const FQuat& Transform, const FVector& Vector)
+template<typename T>
+inline UE::Math::TVector<T> TransformVector(const UE::Math::TQuat<T>& Transform, const UE::Math::TVector<T>& Vector)
 {
 	return Transform.RotateVector(Vector);
 }
 
 /**
- * Specialization for FRotator as it's member function is called something slightly different.
+ * Specialization for UE::Math::TRotator<T> as it's member function is called something slightly different.
  */
-inline FVector TransformPoint(const FRotator& Transform, const FVector& Point)
+template<typename T>
+inline UE::Math::TVector<T> TransformPoint(const UE::Math::TRotator<T>& Transform, const UE::Math::TVector<T>& Point)
 {
 	return Transform.RotateVector(Point);
 }
 
 /**
- * Specialization for FRotator as it's member function is called something slightly different.
+ * Specialization for UE::Math::TRotator<T> as it's member function is called something slightly different.
  */
-inline FVector TransformVector(const FRotator& Transform, const FVector& Vector)
+template<typename T>
+inline UE::Math::TVector<T> TransformVector(const UE::Math::TRotator<T>& Transform, const UE::Math::TVector<T>& Vector)
 {
 	return Transform.RotateVector(Vector);
 }
 
 /**
- * Specialization for FVector Translation.
+ * Specialization for UE::Math::TVector<T> Translation.
  */
-inline FVector TransformPoint(const FVector& Transform, const FVector& Point)
+template<typename T>
+inline UE::Math::TVector<T> TransformPoint(const UE::Math::TVector<T>& Transform, const UE::Math::TVector<T>& Point)
 {
 	return Transform + Point;
 }
 
 /**
- * Specialization for FVector Translation (does nothing).
+ * Specialization for UE::Math::TVector<T> Translation (does nothing).
  */
-inline const FVector& TransformVector(const FVector& Transform, const FVector& Vector)
+template<typename T>
+inline const UE::Math::TVector<T>& TransformVector(const UE::Math::TVector<T>& Transform, const UE::Math::TVector<T>& Vector)
 {
 	return Vector;
 }
@@ -390,7 +466,8 @@ inline const FVector& TransformVector(const FVector& Transform, const FVector& V
 /**
  * Specialization for Scale.
  */
-inline FVector TransformPoint(const FScale& Transform, const FVector& Point)
+template<typename T>
+inline UE::Math::TVector<T> TransformPoint(const TScale<T>& Transform, const UE::Math::TVector<T>& Point)
 {
 	return Transform.GetVector() * Point;
 }
@@ -398,7 +475,8 @@ inline FVector TransformPoint(const FScale& Transform, const FVector& Point)
 /**
  * Specialization for Scale.
  */
-inline FVector TransformVector(const FScale& Transform, const FVector& Vector)
+template<typename T>
+inline UE::Math::TVector<T> TransformVector(const TScale<T>& Transform, const UE::Math::TVector<T>& Vector)
 {
 	return Transform.GetVector() * Vector;
 }
diff --git a/Engine/Source/Runtime/SlateCore/Public/Rendering/SlateLayoutTransform.h b/Engine/Source/Runtime/SlateCore/Public/Rendering/SlateLayoutTransform.h
--- a/Engine/Source/Runtime/SlateCore/Public/Rendering/SlateLayoutTransform.h
+++ b/Engine/Source/Runtime/SlateCore/Public/Rendering/SlateLayoutTransform.h
@@ -18,17 +18,25 @@
 class FSlateLayoutTransform
 {
 public:
+	FSlateLayoutTransform(float InScale = 1.0f)
+		: Scale(InScale)
+		, Translation(FVector2f(ForceInit))
+	{
+	}
+
 	/** Ctor from a scale followed by translate. Shortcut to Concatenate(InScale, InTranslation). */
-	explicit FSlateLayoutTransform(float InScale = 1.0f, const FVector2D& InTranslation = FVector2D(ForceInit))
+	template<typename VType>
+	explicit FSlateLayoutTransform(float InScale, const UE::Math::TVector2<VType>& InTranslation)
 		:Scale(InScale)
 		,Translation(FVector2f(InTranslation))
 	{
 	}
 
 	/** Ctor from a 2D translation followed by a scale. Shortcut to Concatenate(InTranslation, InScale). While this is the opposite order we internally store them, we can represent this correctly. */
-	explicit FSlateLayoutTransform(const FVector2D& InTranslation/*, float InScale = 1.0f*/)
-		:Scale(1.0f/*InScale*/)
-		,Translation(FVector2f(InTranslation)/**InScale*/)
+	template<typename VType>
+	explicit FSlateLayoutTransform(const UE::Math::TVector2<VType>& InTranslation)
+		:Scale(1.0f)
+		,Translation(FVector2f(InTranslation))
 	{
 	}
 
@@ -53,15 +61,17 @@ public:
 	}
 
 	/** 2D transform support. */
-	FVector2D TransformPoint(const FVector2D& Point) const
+	template<typename VType>
+	UE::Math::TVector2<VType> TransformPoint(const UE::Math::TVector2<VType>& Point) const
 	{
-		return ::TransformPoint(GetTranslation(), ::TransformPoint(Scale, Point));
+		return ::TransformPoint((UE::Math::TVector2<VType>)Translation, ::TransformPoint(Scale, Point));
 	}
 
 	/** 2D transform support. */
-	FVector2D TransformVector(const FVector2D& Vector) const
+	template<typename VType>
+	UE::Math::TVector2<VType> TransformVector(const UE::Math::TVector2<VType>& Vector) const
 	{
-		return ::TransformVector(GetTranslation(), ::TransformVector(Scale, Vector));
+		return ::TransformVector((UE::Math::TVector2<VType>)Translation, ::TransformVector(Scale, Vector));
 	}
 
 	/**
@@ -75,13 +85,13 @@ public:
 	{
 		// New Translation is essentially: RHS.TransformPoint(TransformPoint(FVector2D::ZeroVector))
 		// Since Zero through LHS -> Translation we optimize it slightly to skip the zero multiplies.
-		return FSlateLayoutTransform(::Concatenate(Scale, RHS.Scale), RHS.TransformPoint(GetTranslation()));
+		return FSlateLayoutTransform(::Concatenate(Scale, RHS.Scale), RHS.TransformPoint(Translation));
 	}
 
 	/** Invert the transform/scale. */
 	FSlateLayoutTransform Inverse() const
 	{
-		return FSlateLayoutTransform(::Inverse(Scale), ::Inverse(FVector2D(Translation)) * ::Inverse(Scale));
+		return FSlateLayoutTransform(::Inverse(Scale), ::Inverse(Translation) * ::Inverse(Scale));
 	}
 
 	/** Equality. */
@@ -102,48 +112,63 @@ private:
 };
 
 /** Specialization for concatenating a uniform scale and 2D Translation. */
-inline FSlateLayoutTransform Concatenate(float Scale, const FVector2D& Translation)
+template<typename T>
+inline FSlateLayoutTransform Concatenate(float Scale, const UE::Math::TVector2<T>& Translation)
 {
 	return FSlateLayoutTransform(Scale, Translation);
 }
 
+template<typename T>
+inline FSlateLayoutTransform Concatenate(double Scale, const UE::Math::TVector2<T>& Translation)
+{
+	return FSlateLayoutTransform((float)Scale, Translation);
+}
+
 /** Specialization for concatenating a 2D Translation and uniform scale. */
-inline FSlateLayoutTransform Concatenate(const FVector2D& Translation, float Scale)
+template<typename T>
+inline FSlateLayoutTransform Concatenate(const UE::Math::TVector2<T>& Translation, float Scale)
 {
 	return FSlateLayoutTransform(Scale, TransformPoint(Scale, Translation));
 }
 
-/** concatenation rules for LayoutTransform. */
-template<> struct ConcatenateRules<FSlateLayoutTransform, float                > { typedef FSlateLayoutTransform ResultType; };
-/** concatenation rules for LayoutTransform. */
-template<> struct ConcatenateRules<float                , FSlateLayoutTransform> { typedef FSlateLayoutTransform ResultType; };
-/** concatenation rules for LayoutTransform. */
-template<> struct ConcatenateRules<FSlateLayoutTransform, FVector2D            > { typedef FSlateLayoutTransform ResultType; };
-/** concatenation rules for LayoutTransform. */
-template<> struct ConcatenateRules<FVector2D            , FSlateLayoutTransform> { typedef FSlateLayoutTransform ResultType; };
+template<typename T>
+inline FSlateLayoutTransform Concatenate(const UE::Math::TVector2<T>& Translation, double Scale)
+{
+	return FSlateLayoutTransform((float)Scale, TransformPoint(Scale, Translation));
+}
 
 /** concatenation rules for LayoutTransform. */
-template<> struct ConcatenateRules<FSlateLayoutTransform, FMatrix              > { typedef FMatrix ResultType; };
-/** concatenation rules for LayoutTransform. */
-template<> struct ConcatenateRules<FMatrix              , FSlateLayoutTransform> { typedef FMatrix ResultType; };
+template<> struct ConcatenateRules<FSlateLayoutTransform, float					> { typedef FSlateLayoutTransform ResultType; };
+template<> struct ConcatenateRules<FSlateLayoutTransform, double				> { typedef FSlateLayoutTransform ResultType; };
+template<> struct ConcatenateRules<float				, FSlateLayoutTransform	> { typedef FSlateLayoutTransform ResultType; };
+template<> struct ConcatenateRules<double				, FSlateLayoutTransform	> { typedef FSlateLayoutTransform ResultType; };
+template<typename T> struct ConcatenateRules<FSlateLayoutTransform, UE::Math::TVector2<T>	> { typedef FSlateLayoutTransform ResultType; };
+template<typename T> struct ConcatenateRules<UE::Math::TVector2<T>, FSlateLayoutTransform	> { typedef FSlateLayoutTransform ResultType; };
+template<typename T> struct ConcatenateRules<FSlateLayoutTransform	, UE::Math::TMatrix<T>	> { typedef UE::Math::TMatrix<T> ResultType; };
+template<typename T> struct ConcatenateRules<UE::Math::TMatrix<T>	, FSlateLayoutTransform	> { typedef UE::Math::TMatrix<T> ResultType; };
 
 /** concatenation rules for layout transforms and 2x2 generalized transforms. Need to be upcast to FTransform2D. */
-template<> struct ConcatenateRules<FScale2D             , FSlateLayoutTransform> { typedef FTransform2D ResultType; };
-template<> struct ConcatenateRules<FShear2D             , FSlateLayoutTransform> { typedef FTransform2D ResultType; };
-template<> struct ConcatenateRules<FQuat2D              , FSlateLayoutTransform> { typedef FTransform2D ResultType; };
-template<> struct ConcatenateRules<FMatrix2x2           , FSlateLayoutTransform> { typedef FTransform2D ResultType; };
-template<> struct ConcatenateRules<FSlateLayoutTransform, FScale2D             > { typedef FTransform2D ResultType; };
-template<> struct ConcatenateRules<FSlateLayoutTransform, FShear2D             > { typedef FTransform2D ResultType; };
-template<> struct ConcatenateRules<FSlateLayoutTransform, FQuat2D              > { typedef FTransform2D ResultType; };
-template<> struct ConcatenateRules<FSlateLayoutTransform, FMatrix2x2           > { typedef FTransform2D ResultType; };
+template<typename T> struct ConcatenateRules<TScale2<T>				, FSlateLayoutTransform	> { typedef TTransform2<T> ResultType; };
+template<typename T> struct ConcatenateRules<TShear2<T>				, FSlateLayoutTransform	> { typedef TTransform2<T> ResultType; };
+template<typename T> struct ConcatenateRules<TQuat2<T>				, FSlateLayoutTransform	> { typedef TTransform2<T> ResultType; };
+template<typename T> struct ConcatenateRules<TMatrix2x2<T>			, FSlateLayoutTransform	> { typedef TTransform2<T> ResultType; };
+template<typename T> struct ConcatenateRules<FSlateLayoutTransform	, TScale2<T>			> { typedef TTransform2<T> ResultType; };
+template<typename T> struct ConcatenateRules<FSlateLayoutTransform	, TShear2<T>			> { typedef TTransform2<T> ResultType; };
+template<typename T> struct ConcatenateRules<FSlateLayoutTransform	, TQuat2<T>				> { typedef TTransform2<T> ResultType; };
+template<typename T> struct ConcatenateRules<FSlateLayoutTransform	, TMatrix2x2<T>			> { typedef TTransform2<T> ResultType; };
 
 //////////////////////////////////////////////////////////////////////////
 // FSlateLayoutTransform adapters.
 // 
 // Adapt FTransform2D to accept FSlateLayoutTransforms as well.
 //////////////////////////////////////////////////////////////////////////
-template<> template<> inline FTransform2D TransformConverter<FTransform2D>::Convert<FSlateLayoutTransform>(const FSlateLayoutTransform& Transform)
+template<> template<> inline TTransform2<float> TransformConverter<TTransform2<float>>::Convert<FSlateLayoutTransform>(const FSlateLayoutTransform& Transform)
+{
+	return TTransform2<float>(TScale2<float>((float)Transform.GetScale()), Transform.GetTranslation());
+}
+
+template<> template<> inline TTransform2<double> TransformConverter<TTransform2<double>>::Convert<FSlateLayoutTransform>(const FSlateLayoutTransform& Transform)
 {
-	return FTransform2D(FScale2D(Transform.GetScale()), Transform.GetTranslation());
+	return TTransform2<double>(TScale2<double>(Transform.GetScale()), Transform.GetTranslation());
 }
 
