bModularFeatureListLocked should be an atomic bool

diff --git a/Engine/Source/Runtime/Core/Private/Features/ModularFeatures.h b/Engine/Source/Runtime/Core/Private/Features/ModularFeatures.h
--- a/Engine/Source/Runtime/Core/Private/Features/ModularFeatures.h
+++ b/Engine/Source/Runtime/Core/Private/Features/ModularFeatures.h
@@ -38,7 +38,7 @@ private:
 	TMultiMap< FName, class IModularFeature* > ModularFeaturesMap;
 
 	/** Used to fire ensure when ModularFeaturesMap is inspected cross-thread without being locked properly */
-	bool bModularFeatureListLocked = false;
+	std::atomic<bool> bModularFeatureListLocked = false;
 
 	/** Lock modular features map so it can be used across threads */
 	FCriticalSection ModularFeaturesMapCriticalSection;
