Added trace profiler scopes to UnixPlatformMemory GetStats and GetExtendedStats, to trace the effect of system calls on other threads.

diff --git a/Engine/Source/Runtime/Core/Private/Unix/UnixPlatformMemory.cpp b/Engine/Source/Runtime/Core/Private/Unix/UnixPlatformMemory.cpp
--- a/Engine/Source/Runtime/Core/Private/Unix/UnixPlatformMemory.cpp
+++ b/Engine/Source/Runtime/Core/Private/Unix/UnixPlatformMemory.cpp
@@ -721,6 +721,7 @@ static uint32 ReadProcFields(const char *FileName, FProcField *ProcFields, uint3
 
 			do
 			{
+				QUICK_SCOPE_CYCLE_COUNTER(STAT_ReadProcFields_FileRead);
 				BytesRead = read(Fd, Buffer + BytesAvailableInChunk, BytesToRead);
 			} while (BytesRead < 0 && errno == EINTR);
 
@@ -783,6 +784,8 @@ static uint32 ReadProcFields(const char *FileName, FProcField *ProcFields, uint3
 //   uint64 PeakUsedVirtual;   /** The peak amount of virtual memory used by the process. */				VmPeak
 FPlatformMemoryStats FUnixPlatformMemory::GetStats()
 {
+	QUICK_SCOPE_CYCLE_COUNTER(STAT_GetStats);
+	
 	uint64 MemFree = 0;
 	uint64 Cached = 0;
 	FPlatformMemoryStats MemoryStats;
@@ -855,6 +858,8 @@ static uint64 ParseSMapsFileChunk(ANSICHAR *Buffer, uint64 BufferSize, FProcFiel
 
 FExtendedPlatformMemoryStats FUnixPlatformMemory::GetExtendedStats()
 {
+	QUICK_SCOPE_CYCLE_COUNTER(STAT_GetExtendedStats);
+	
 	const ANSICHAR Shared_CleanStr[]  = "Shared_Clean:";
 	const ANSICHAR Shared_DirtyStr[]  = "Shared_Dirty:";
 	const ANSICHAR Private_CleanStr[] = "Private_Clean:";
@@ -892,6 +897,7 @@ FExtendedPlatformMemoryStats FUnixPlatformMemory::GetExtendedStats()
 
 				do
 				{
+					QUICK_SCOPE_CYCLE_COUNTER(STAT_GetExtendedStats_FileRead);
 					BytesRead = read(Fd, Buffer + BytesAvailableInChunk, BytesToRead);
 				} while (BytesRead < 0 && errno == EINTR);
 
