ValueOrError: Added specializations for void value and error types

Migrated tests to Catch.

diff --git a/Engine/Source/Runtime/Core/Private/Tests/Templates/ValueOrErrorTest.cpp b/Engine/Source/Runtime/Core/Private/Tests/Templates/ValueOrErrorTest.cpp
deleted file mode 100644
diff --git a/Engine/Source/Runtime/Core/Public/Templates/ValueOrError.h b/Engine/Source/Runtime/Core/Public/Templates/ValueOrError.h
--- a/Engine/Source/Runtime/Core/Public/Templates/ValueOrError.h
+++ b/Engine/Source/Runtime/Core/Public/Templates/ValueOrError.h
@@ -3,6 +3,7 @@
 #pragma once
 
 #include "CoreTypes.h"
+#include "Misc/Optional.h"
 #include "Misc/TVariant.h"
 #include "Templates/Tuple.h"
 #include "Templates/UnrealTemplate.h"
@@ -146,3 +147,127 @@ public:
 private:
 	TVariant<ValueType, FWrapErrorType, FEmptyType> Variant;
 };
+
+template <typename ValueType>
+class TValueOrError<ValueType, void>
+{
+	template <typename... ArgTypes, uint32... ArgIndices>
+	inline TValueOrError(TValueOrError_ValueProxy<ArgTypes...>&& Proxy, TIntegerSequence<uint32, ArgIndices...>)
+		: Value(InPlace, MoveTemp(Proxy.Args).template Get<ArgIndices>()...)
+	{
+	}
+
+public:
+	/** Construct the value from a proxy from MakeValue. */
+	template <typename... ArgTypes>
+	inline TValueOrError(TValueOrError_ValueProxy<ArgTypes...>&& Proxy)
+		: TValueOrError(MoveTemp(Proxy), TMakeIntegerSequence<uint32, sizeof...(ArgTypes)>())
+	{
+	}
+
+	/** Construct the error from a proxy from MakeError. */
+	inline TValueOrError(TValueOrError_ErrorProxy<>&& Proxy)
+	{
+	}
+
+	/** Whether the error is set. An error does imply no value. */
+	inline bool HasError() const { return !Value.IsSet(); }
+
+	/** Whether the value is set. A value does imply no error. */
+	inline bool HasValue() const { return Value.IsSet(); }
+
+	/** Access the value. Asserts if this does not have a value. */
+	inline       ValueType& GetValue()       &  { return Value.GetValue(); }
+	inline const ValueType& GetValue() const &  { return Value.GetValue(); }
+	inline       ValueType  GetValue()       && { return MoveTemp(Value.GetValue()); }
+
+	/** Access the value if it is set. */
+	inline       ValueType* TryGetValue()       { return Value.GetPtrOrNull(); }
+	inline const ValueType* TryGetValue() const { return Value.GetPtrOrNull(); }
+
+	/** Steal the value. Asserts if this does not have a value. This causes the value to be unset. */
+	inline ValueType StealValue()
+	{
+		ValueType Temp = MoveTemp(GetValue());
+		Value.Reset();
+		return Temp;
+	}
+
+private:
+	TOptional<ValueType> Value;
+};
+
+template <typename ErrorType>
+class TValueOrError<void, ErrorType>
+{
+	template <typename... ArgTypes, uint32... ArgIndices>
+	inline TValueOrError(TValueOrError_ErrorProxy<ArgTypes...>&& Proxy, TIntegerSequence<uint32, ArgIndices...>)
+		: Error(InPlace, MoveTemp(Proxy.Args).template Get<ArgIndices>()...)
+	{
+	}
+
+public:
+	/** Construct the value from a proxy from MakeValue. */
+	inline TValueOrError(TValueOrError_ValueProxy<>&& Proxy)
+	{
+	}
+
+	/** Construct the error from a proxy from MakeError. */
+	template <typename... ArgTypes>
+	inline TValueOrError(TValueOrError_ErrorProxy<ArgTypes...>&& Proxy)
+		: TValueOrError(MoveTemp(Proxy), TMakeIntegerSequence<uint32, sizeof...(ArgTypes)>())
+	{
+	}
+
+	/** Whether the value is set. A value does imply no error. */
+	inline bool HasValue() const { return !Error.IsSet(); }
+
+	/** Whether the error is set. An error does imply no value. */
+	inline bool HasError() const { return Error.IsSet(); }
+
+	/** Access the error. Asserts if this does not have an error. */
+	inline       ErrorType& GetError()       &  { return Error.GetValue(); }
+	inline const ErrorType& GetError() const &  { return Error.GetValue(); }
+	inline       ErrorType  GetError()       && { return MoveTemp(Error.GetValue()); }
+
+	/** Access the error if it is set. */
+	inline       ErrorType* TryGetError()       { return Error.GetPtrOrNull(); }
+	inline const ErrorType* TryGetError() const { return Error.GetPtrOrNull(); }
+
+	/** Steal the error. Asserts if this does not have an error. This causes the error to be unset. */
+	inline ErrorType StealError()
+	{
+		ErrorType Temp = MoveTemp(GetError());
+		Error.Reset();
+		return Temp;
+	}
+
+private:
+	TOptional<ErrorType> Error;
+};
+
+template <>
+class TValueOrError<void, void>
+{
+public:
+	/** Construct the value from a proxy from MakeValue. */
+	inline TValueOrError(TValueOrError_ValueProxy<>&& Proxy)
+		: bValue(true)
+	{
+	}
+
+	/** Construct the error from a proxy from MakeError. */
+	inline TValueOrError(TValueOrError_ErrorProxy<>&& Proxy)
+		: bValue(false)
+	{
+	}
+
+	/** Whether the error is set. An error does imply no value. */
+	inline bool HasError() const { return !bValue; }
+
+	/** Whether the value is set. A value does imply no error. */
+	inline bool HasValue() const { return bValue; }
+
+private:
+	bool bValue;
+};
diff --git a/Engine/Source/Runtime/Core/Tests/Templates/ValueOrErrorTest.cpp b/Engine/Source/Runtime/Core/Tests/Templates/ValueOrErrorTest.cpp
new file mode 100644
--- /dev/null
+++ b/Engine/Source/Runtime/Core/Tests/Templates/ValueOrErrorTest.cpp
@@ -0,0 +1,231 @@
+// Copyright Epic Games, Inc. All Rights Reserved.
+
+#include "Templates/ValueOrError.h"
+
+#include "TestHarness.h"
+
+#include <type_traits>
+
+TEST_CASE("Core::Templates::ValueOrError", "[Core][Templates][Smoke]")
+{
+	SECTION("Static")
+	{
+		STATIC_REQUIRE(!std::is_constructible<TValueOrError<int, int>>::value);
+
+		STATIC_REQUIRE(std::is_copy_constructible<TValueOrError<int, int>>::value);
+		STATIC_REQUIRE(std::is_move_constructible<TValueOrError<int, int>>::value);
+
+		STATIC_REQUIRE(std::is_copy_assignable<TValueOrError<int, int>>::value);
+		STATIC_REQUIRE(std::is_move_assignable<TValueOrError<int, int>>::value);
+	}
+
+	static int ValueCount = 0;
+	static int ErrorCount = 0;
+
+	struct FTestValue
+	{
+		FTestValue() { Value = ++ValueCount; }
+		FTestValue(const FTestValue&) { Value = ++ValueCount; }
+		FTestValue(FTestValue&& Other) { Value = Other.Value; ++ValueCount; }
+		FTestValue(int InValue1, int InValue2, int InValue3) { Value = InValue1 + InValue2 + InValue3; ++ValueCount; }
+		~FTestValue() { --ValueCount; }
+		FTestValue& operator=(const FTestValue& Other) = delete;
+		FTestValue& operator=(FTestValue&& Other) = delete;
+		int Value;
+	};
+
+	struct FTestError
+	{
+		FTestError() { Error = ++ErrorCount; }
+		FTestError(const FTestError&) { Error = ++ErrorCount; }
+		FTestError(FTestError&& Other) { Error = Other.Error; ++ErrorCount; }
+		FTestError(int InError1, int InError2) { Error = InError1 + InError2; ++ErrorCount; }
+		~FTestError() { --ErrorCount; }
+		FTestError& operator=(const FTestError& Other) = delete;
+		FTestError& operator=(FTestError&& Other) = delete;
+		int Error;
+	};
+
+	struct FTestMoveOnly
+	{
+		FTestMoveOnly() = default;
+		FTestMoveOnly(const FTestMoveOnly&) = delete;
+		FTestMoveOnly(FTestMoveOnly&&) = default;
+		int Value = 0;
+	};
+
+	using FTestType = TValueOrError<FTestValue, FTestError>;
+
+	SECTION("MakeValue Move")
+	{
+		FTestType ValueOrError = MakeValue(FTestValue());
+		CHECK(ValueCount == 1);
+		CHECK(ValueOrError.TryGetValue() == &ValueOrError.GetValue());
+		CHECK(ValueOrError.GetValue().Value == 1);
+		CHECK(ValueOrError.HasValue());
+		CHECK(ValueOrError.TryGetError() == nullptr);
+		CHECK_FALSE(ValueOrError.HasError());
+	}
+
+	SECTION("MakeValue Proxy")
+	{
+		FTestType ValueOrError = MakeValue(2, 6, 8);
+		CHECK(ValueCount == 1);
+		CHECK(ValueOrError.TryGetValue() == &ValueOrError.GetValue());
+		CHECK(ValueOrError.GetValue().Value == 16);
+		CHECK(ValueOrError.HasValue());
+		CHECK(ValueOrError.TryGetError() == nullptr);
+		CHECK_FALSE(ValueOrError.HasError());
+	}
+
+	SECTION("StealValue")
+	{
+		FTestType ValueOrError = MakeValue(FTestValue());
+		FTestValue Value = ValueOrError.StealValue();
+		CHECK(ValueCount == 1);
+		CHECK(Value.Value == 1);
+		CHECK_FALSE(ValueOrError.HasError());
+		CHECK_FALSE(ValueOrError.HasValue());
+	}
+
+	SECTION("MakeError Move")
+	{
+		FTestType ValueOrError = MakeError(FTestError());
+		CHECK(ErrorCount == 1);
+		CHECK(ValueOrError.TryGetError() == &ValueOrError.GetError());
+		CHECK(ValueOrError.GetError().Error == 1);
+		CHECK(ValueOrError.HasError());
+		CHECK(ValueOrError.TryGetValue() == nullptr);
+		CHECK_FALSE(ValueOrError.HasValue());
+	}
+
+	SECTION("MakeError Proxy")
+	{
+		FTestType ValueOrError = MakeError(4, 12);
+		CHECK(ErrorCount == 1);
+		CHECK(ValueOrError.TryGetError() == &ValueOrError.GetError());
+		CHECK(ValueOrError.GetError().Error == 16);
+		CHECK(ValueOrError.HasError());
+		CHECK(ValueOrError.TryGetValue() == nullptr);
+		CHECK_FALSE(ValueOrError.HasValue());
+	}
+
+	SECTION("StealError")
+	{
+		FTestType ValueOrError = MakeError();
+		FTestError Error = ValueOrError.StealError();
+		CHECK(ErrorCount == 1);
+		CHECK(Error.Error == 1);
+		CHECK_FALSE(ValueOrError.HasValue());
+		CHECK_FALSE(ValueOrError.HasError());
+	}
+
+	SECTION("Assignment")
+	{
+		FTestType ValueOrError = MakeValue();
+		ValueOrError = MakeValue();
+		CHECK(ValueCount == 1);
+		CHECK(ValueOrError.GetValue().Value == 2);
+		CHECK(ErrorCount == 0);
+		ValueOrError = MakeError();
+		CHECK(ValueCount == 0);
+		CHECK(ErrorCount == 1);
+		ValueOrError = MakeError();
+		CHECK(ValueCount == 0);
+		CHECK(ValueOrError.GetError().Error == 2);
+		CHECK(ErrorCount == 1);
+		ValueOrError = MakeValue();
+		CHECK(ValueCount == 1);
+		CHECK(ErrorCount == 0);
+		FTestType UnsetValueOrError = MakeValue();
+		UnsetValueOrError.StealValue();
+		ValueOrError = MoveTemp(UnsetValueOrError);
+		CHECK(ValueCount == 0);
+		CHECK(ErrorCount == 0);
+		CHECK_FALSE(ValueOrError.HasValue());
+		CHECK_FALSE(ValueOrError.HasError());
+	}
+
+	SECTION("Move-Only Value/Error")
+	{
+		TValueOrError<FTestMoveOnly, FTestMoveOnly> Value = MakeValue(FTestMoveOnly{1});
+		TValueOrError<FTestMoveOnly, FTestMoveOnly> Error = MakeError(FTestMoveOnly{1});
+		FTestMoveOnly MovedValue = MoveTemp(Value).GetValue();
+		FTestMoveOnly MovedError = MoveTemp(Error).GetError();
+		CHECK(MovedValue.Value == 1);
+		CHECK(MovedError.Value == 1);
+	}
+
+	SECTION("Integer Value/Error")
+	{
+		TValueOrError<int32, int32> ValueOrError = MakeValue();
+		CHECK(ValueOrError.GetValue() == 0);
+		ValueOrError = MakeValue(1);
+		CHECK(ValueOrError.GetValue() == 1);
+		ValueOrError = MakeError();
+		CHECK(ValueOrError.GetError() == 0);
+		ValueOrError = MakeError(1);
+		CHECK(ValueOrError.GetError() == 1);
+	}
+
+	SECTION("Value/Void Value")
+	{
+		TValueOrError<FTestValue, void> ValueOrError = MakeValue(1, 2, 3);
+		CHECK(ValueCount == 1);
+		CHECK(ValueOrError.HasValue());
+		CHECK_FALSE(ValueOrError.HasError());
+		CHECK(ValueOrError.GetValue().Value == 6);
+		CHECK(ValueOrError.TryGetValue() == &ValueOrError.GetValue());
+		FTestValue Value = ValueOrError.StealValue();
+		CHECK_FALSE(ValueOrError.HasValue());
+		CHECK(Value.Value == 6);
+	}
+
+	SECTION("Value/Void Error")
+	{
+		TValueOrError<FTestValue, void> ValueOrError = MakeError();
+		CHECK(ValueCount == 0);
+		CHECK(ValueOrError.HasError());
+		CHECK_FALSE(ValueOrError.HasValue());
+		CHECK(ValueOrError.TryGetValue() == nullptr);
+	}
+
+	SECTION("Void/Error Value")
+	{
+		TValueOrError<void, FTestError> ValueOrError = MakeValue();
+		CHECK(ErrorCount == 0);
+		CHECK(ValueOrError.HasValue());
+		CHECK_FALSE(ValueOrError.HasError());
+		CHECK(ValueOrError.TryGetError() == nullptr);
+	}
+
+	SECTION("Void/Error Error")
+	{
+		TValueOrError<void, FTestError> ValueOrError = MakeError(1, 2);
+		CHECK(ErrorCount == 1);
+		CHECK(ValueOrError.HasError());
+		CHECK_FALSE(ValueOrError.HasValue());
+		CHECK(ValueOrError.GetError().Error == 3);
+		CHECK(ValueOrError.TryGetError() == &ValueOrError.GetError());
+		FTestError Error = ValueOrError.StealError();
+		CHECK_FALSE(ValueOrError.HasError());
+		CHECK(Error.Error == 3);
+	}
+
+	SECTION("Void/Void Value")
+	{
+		TValueOrError<void, void> ValueOrError = MakeValue();
+		CHECK(ValueOrError.HasValue());
+		CHECK_FALSE(ValueOrError.HasError());
+	}
+
+	SECTION("Void/Void Error")
+	{
+		TValueOrError<void, void> ValueOrError = MakeError();
+		CHECK(ValueOrError.HasError());
+		CHECK_FALSE(ValueOrError.HasValue());
+	}
+
+	CHECK(ValueCount == 0);
+	CHECK(ErrorCount == 0);
+}
