IoBuffer - additional constructors from memory view

#rnx

diff --git a/Engine/Source/Runtime/Core/Private/IO/IoBuffer.cpp b/Engine/Source/Runtime/Core/Private/IO/IoBuffer.cpp
--- a/Engine/Source/Runtime/Core/Private/IO/IoBuffer.cpp
+++ b/Engine/Source/Runtime/Core/Private/IO/IoBuffer.cpp
@@ -117,21 +117,41 @@ FIoBuffer::FIoBuffer(const void* Data, uint64 InSize, const FIoBuffer& OuterBuff
 {
 }
 
+FIoBuffer::FIoBuffer(FMemoryView Memory, const FIoBuffer& OuterBuffer)
+: FIoBuffer(Memory.GetData(), Memory.GetSize(), OuterBuffer)
+{
+}
+
 FIoBuffer::FIoBuffer(FIoBuffer::EWrapTag, const void* Data, uint64 InSize)
 :	CorePtr(new BufCore((uint8*)Data, InSize, /* ownership */ false))
 {
 }
 
+FIoBuffer::FIoBuffer(FIoBuffer::EWrapTag, FMemoryView Memory)
+: FIoBuffer(FIoBuffer::Wrap, Memory.GetData(), Memory.GetSize())
+{
+}
+
 FIoBuffer::FIoBuffer(FIoBuffer::EAssumeOwnershipTag, const void* Data, uint64 InSize)
 :	CorePtr(new BufCore((uint8*)Data, InSize, /* ownership */ true))
 {
 }
 
+FIoBuffer::FIoBuffer(FIoBuffer::EAssumeOwnershipTag, FMemoryView Memory)
+: FIoBuffer(FIoBuffer::AssumeOwnership, Memory.GetData(), Memory.GetSize())
+{
+}
+
 FIoBuffer::FIoBuffer(FIoBuffer::ECloneTag, const void* Data, uint64 InSize)
 :	CorePtr(new BufCore(Clone, (uint8*)Data, InSize))
 {
 }
 
+FIoBuffer::FIoBuffer(FIoBuffer::ECloneTag, FMemoryView Memory)
+: FIoBuffer(FIoBuffer::Clone, Memory.GetData(), Memory.GetSize())
+{
+}
+
 void		
 FIoBuffer::MakeOwned() const
 {
diff --git a/Engine/Source/Runtime/Core/Public/IO/IoDispatcher.h b/Engine/Source/Runtime/Core/Public/IO/IoDispatcher.h
--- a/Engine/Source/Runtime/Core/Public/IO/IoDispatcher.h
+++ b/Engine/Source/Runtime/Core/Public/IO/IoDispatcher.h
@@ -428,10 +428,14 @@ public:
 	CORE_API			FIoBuffer();
 	CORE_API explicit	FIoBuffer(uint64 InSize);
 	CORE_API			FIoBuffer(const void* Data, uint64 InSize, const FIoBuffer& OuterBuffer);
+	CORE_API			FIoBuffer(FMemoryView Memory, const FIoBuffer& OuterBuffer);
 
 	CORE_API			FIoBuffer(EAssumeOwnershipTag,	const void* Data, uint64 InSize);
+	CORE_API			FIoBuffer(EAssumeOwnershipTag,	FMemoryView Memory);
 	CORE_API			FIoBuffer(ECloneTag,			const void* Data, uint64 InSize);
+	CORE_API			FIoBuffer(ECloneTag,			FMemoryView Memory);
 	CORE_API			FIoBuffer(EWrapTag,				const void* Data, uint64 InSize);
+	CORE_API			FIoBuffer(EWrapTag,				FMemoryView Memory);
 
 	// Note: we currently rely on implicit move constructor, thus we do not declare any
 	//		 destructor or copy/assignment operators or copy constructors
