Forking
* SupportsMultithreadingPostFork() will not assert if the cmdline hasn't been set when called.
* Instead it will be return true/false based on the DEFAULT define and hope the cmdline override won't break anything later.

diff --git a/Engine/Source/Runtime/Core/Private/Misc/Fork.cpp b/Engine/Source/Runtime/Core/Private/Misc/Fork.cpp
--- a/Engine/Source/Runtime/Core/Private/Misc/Fork.cpp
+++ b/Engine/Source/Runtime/Core/Private/Misc/Fork.cpp
@@ -117,7 +117,12 @@ bool FForkProcessHelper::IsForkedMultithreadInstance()
 
 bool FForkProcessHelper::SupportsMultithreadingPostFork()
 {
-	check(FCommandLine::IsInitialized());
+	if (!FCommandLine::IsInitialized())
+	{
+		// Return the default setting if the cmdline isn't set yet.
+		return DEFAULT_FORK_PROCESS_MULTITHREAD;
+	}
+
 #if DEFAULT_FORK_PROCESS_MULTITHREAD
 	// Always multi thread unless manually turned off via command line
 	static bool bSupportsMT = FParse::Param(FCommandLine::Get(), TEXT("DisablePostForkThreading")) == false;
