Partially revert 19458254: don't use profiling scopes in GetStats(), as it's used in static initialization during LLM profiling.

diff --git a/Engine/Source/Runtime/Core/Private/Unix/UnixPlatformMemory.cpp b/Engine/Source/Runtime/Core/Private/Unix/UnixPlatformMemory.cpp
--- a/Engine/Source/Runtime/Core/Private/Unix/UnixPlatformMemory.cpp
+++ b/Engine/Source/Runtime/Core/Private/Unix/UnixPlatformMemory.cpp
@@ -721,7 +721,6 @@ static uint32 ReadProcFields(const char *FileName, FProcField *ProcFields, uint3
 
 			do
 			{
-				QUICK_SCOPE_CYCLE_COUNTER(STAT_ReadProcFields_FileRead);
 				BytesRead = read(Fd, Buffer + BytesAvailableInChunk, BytesToRead);
 			} while (BytesRead < 0 && errno == EINTR);
 
@@ -784,8 +783,6 @@ static uint32 ReadProcFields(const char *FileName, FProcField *ProcFields, uint3
 //   uint64 PeakUsedVirtual;   /** The peak amount of virtual memory used by the process. */				VmPeak
 FPlatformMemoryStats FUnixPlatformMemory::GetStats()
 {
-	QUICK_SCOPE_CYCLE_COUNTER(STAT_GetStats);
-	
 	uint64 MemFree = 0;
 	uint64 Cached = 0;
 	FPlatformMemoryStats MemoryStats;
