Converted FIntPoint / FIntVector / FIntRect to a template type and added unsigned variants.

diff --git a/Engine/Plugins/Runtime/HairStrands/Source/HairStrandsCore/Private/HairCardsDatas.cpp b/Engine/Plugins/Runtime/HairStrands/Source/HairStrandsCore/Private/HairCardsDatas.cpp
--- a/Engine/Plugins/Runtime/HairStrands/Source/HairStrandsCore/Private/HairCardsDatas.cpp
+++ b/Engine/Plugins/Runtime/HairStrands/Source/HairStrandsCore/Private/HairCardsDatas.cpp
@@ -7,14 +7,6 @@
 // including the types used in them such as the RenderData structures,
 // CustomVersion will be required to handle the changes
 
-FArchive& operator<<(FArchive& Ar, FUIntPoint& Point)
-{
-	Ar << Point.X;
-	Ar << Point.Y;
-
-	return Ar;
-}
-
 void FHairCardsBulkData::Serialize(FArchive& Ar)
 {
 	Ar << Positions;
diff --git a/Engine/Plugins/Runtime/HairStrands/Source/HairStrandsCore/Public/HairCardsDatas.h b/Engine/Plugins/Runtime/HairStrands/Source/HairStrandsCore/Public/HairCardsDatas.h
--- a/Engine/Plugins/Runtime/HairStrands/Source/HairStrandsCore/Public/HairCardsDatas.h
+++ b/Engine/Plugins/Runtime/HairStrands/Source/HairStrandsCore/Public/HairCardsDatas.h
@@ -11,12 +11,6 @@
 ////////////////////////////////////////////////////////////////////////////////////////////////
 // Utils/Format data
 
-struct FUIntPoint
-{
-	uint32 X;
-	uint32 Y;
-};
-
 struct FHairCardsPositionFormat
 {
 	typedef FVector4f Type;
@@ -103,7 +97,7 @@ struct FHairCardsStrandsPositionFormat
 
 struct FHairCardsOffsetAndCount
 {
-	typedef FUIntPoint Type;
+	typedef FUintPoint Type;
 	static const uint32 ComponentCount = 1;
 	static const uint32 SizeInByte = sizeof(Type);
 	static const EVertexElementType VertexElementType = VET_None;
@@ -322,8 +316,8 @@ struct FHairCardsProceduralGeometry : FHairCardsGeometry
 	TArray<Rect>   Rects;
 	TArray<float>  Lengths;
 	
-	TArray<FUIntPoint> CardIndexToClusterOffsetAndCount;
-	TArray<FUIntPoint> ClusterIndexToVertexOffsetAndCount;
+	TArray<FUintPoint> CardIndexToClusterOffsetAndCount;
+	TArray<FUintPoint> ClusterIndexToVertexOffsetAndCount;
 
 	TArray<FHairOrientedBound> Bounds;
 
diff --git a/Engine/Plugins/Runtime/nDisplay/Source/DisplayCluster/Private/Render/Viewport/RenderTarget/DisplayClusterRenderTargetResource.h b/Engine/Plugins/Runtime/nDisplay/Source/DisplayCluster/Private/Render/Viewport/RenderTarget/DisplayClusterRenderTargetResource.h
--- a/Engine/Plugins/Runtime/nDisplay/Source/DisplayCluster/Private/Render/Viewport/RenderTarget/DisplayClusterRenderTargetResource.h
+++ b/Engine/Plugins/Runtime/nDisplay/Source/DisplayCluster/Private/Render/Viewport/RenderTarget/DisplayClusterRenderTargetResource.h
@@ -15,7 +15,8 @@
 struct FDisplayClusterViewportResourceSettings
 {
 	FDisplayClusterViewportResourceSettings()
-		: Format(EPixelFormat(0))
+		: Size(ForceInit)
+		, Format(EPixelFormat(0))
 		, bShouldUseSRGB(false)
 		, DisplayGamma(1)
 		, bIsRenderTargetable(false)
@@ -23,7 +24,8 @@ struct FDisplayClusterViewportResourceSettings
 	{ }
 
 	FDisplayClusterViewportResourceSettings(const FString& InClusterNodeId, const EPixelFormat InFormat, const bool InbShouldUseSRGB, const float InDisplayGamma)
-		: Format(InFormat)
+		: Size(ForceInit)
+		, Format(InFormat)
 		, bShouldUseSRGB(InbShouldUseSRGB)
 		, DisplayGamma(InDisplayGamma)
 		, bIsRenderTargetable(false)
diff --git a/Engine/Plugins/VirtualProduction/DMX/DMXPixelMapping/Source/DMXPixelMappingRuntime/Private/Components/DMXPixelMappingFixtureGroupItemComponent.cpp b/Engine/Plugins/VirtualProduction/DMX/DMXPixelMapping/Source/DMXPixelMappingRuntime/Private/Components/DMXPixelMappingFixtureGroupItemComponent.cpp
--- a/Engine/Plugins/VirtualProduction/DMX/DMXPixelMapping/Source/DMXPixelMappingRuntime/Private/Components/DMXPixelMappingFixtureGroupItemComponent.cpp
+++ b/Engine/Plugins/VirtualProduction/DMX/DMXPixelMapping/Source/DMXPixelMappingRuntime/Private/Components/DMXPixelMappingFixtureGroupItemComponent.cpp
@@ -213,8 +213,8 @@ void UDMXPixelMappingFixtureGroupItemComponent::QueueDownsample()
 	const FVector2D UVCellSize = UVSize / 2.f;
 	constexpr bool bStaticCalculateUV = true;
 
-	FVector4 ExposeFactor;
-	FIntVector4 InvertFactor;
+	FVector4 ExposeFactor{};
+	FIntVector4 InvertFactor{};
 	if (ColorMode == EDMXColorMode::CM_RGB)
 	{
 		ExposeFactor = FVector4(AttributeRExpose ? 1.f : 0.f, AttributeGExpose ? 1.f : 0.f, AttributeBExpose ? 1.f : 0.f, 1.f);
diff --git a/Engine/Plugins/VirtualProduction/DMX/DMXPixelMapping/Source/DMXPixelMappingRuntime/Private/Components/DMXPixelMappingMatrixCellComponent.cpp b/Engine/Plugins/VirtualProduction/DMX/DMXPixelMapping/Source/DMXPixelMappingRuntime/Private/Components/DMXPixelMappingMatrixCellComponent.cpp
--- a/Engine/Plugins/VirtualProduction/DMX/DMXPixelMapping/Source/DMXPixelMappingRuntime/Private/Components/DMXPixelMappingMatrixCellComponent.cpp
+++ b/Engine/Plugins/VirtualProduction/DMX/DMXPixelMapping/Source/DMXPixelMappingRuntime/Private/Components/DMXPixelMappingMatrixCellComponent.cpp
@@ -238,7 +238,7 @@ void UDMXPixelMappingMatrixCellComponent::QueueDownsample()
 	constexpr bool bStaticCalculateUV = true;
 
 	FVector4 ExposeFactor;
-	FIntVector4 InvertFactor;
+	FIntVector4 InvertFactor{};
 	if (MatrixComponent->ColorMode == EDMXColorMode::CM_RGB)
 	{
 		ExposeFactor = FVector4(MatrixComponent->AttributeRExpose ? 1.f : 0.f, MatrixComponent->AttributeGExpose ? 1.f : 0.f, MatrixComponent->AttributeBExpose ? 1.f : 0.f, 1.f);
diff --git a/Engine/Source/Developer/MaterialUtilities/Public/MaterialUtilities.h b/Engine/Source/Developer/MaterialUtilities/Public/MaterialUtilities.h
--- a/Engine/Source/Developer/MaterialUtilities/Public/MaterialUtilities.h
+++ b/Engine/Source/Developer/MaterialUtilities/Public/MaterialUtilities.h
@@ -121,7 +121,7 @@ private:
 		const uint32 Index = (uint32)Property;
 		check(Index < (uint32)EFlattenMaterialProperties::NumFlattenMaterialProperties);
 
-		static FIntPoint TempSize;
+		static FIntPoint TempSize = FIntPoint::ZeroValue;
 		return Index < (uint32)EFlattenMaterialProperties::NumFlattenMaterialProperties ? PropertySizes[Index] : TempSize;
 	}
 
@@ -139,7 +139,7 @@ private:
 		const uint32 Index = (uint32)Property;
 		check(Index < (uint32)EFlattenMaterialProperties::NumFlattenMaterialProperties);
 
-		static const FIntPoint TempSize;
+		static const FIntPoint TempSize = FIntPoint::ZeroValue;
 		return Index < (uint32)EFlattenMaterialProperties::NumFlattenMaterialProperties ? PropertySizes[Index] : TempSize;
 	}
 
diff --git a/Engine/Source/Editor/LandscapeEditor/Private/LandscapeEditorDetailCustomization_ImportExport.cpp b/Engine/Source/Editor/LandscapeEditor/Private/LandscapeEditorDetailCustomization_ImportExport.cpp
--- a/Engine/Source/Editor/LandscapeEditor/Private/LandscapeEditorDetailCustomization_ImportExport.cpp
+++ b/Engine/Source/Editor/LandscapeEditor/Private/LandscapeEditorDetailCustomization_ImportExport.cpp
@@ -490,7 +490,7 @@ FReply FLandscapeEditorDetailCustomization_ImportExport::OnImportExportButtonCli
 				FIntRect ExportRegion;
 				if (LandscapeInfo->GetLandscapeExtent(LandscapeProxy, ExportRegion))
 				{
-					FIntPoint FileOffset;
+					FIntPoint FileOffset{};
 					if (!LandscapeEdMode->UseSingleFileImport())
 					{
 						ALandscape* Landscape = LandscapeInfo->LandscapeActor.Get();
diff --git a/Engine/Source/Editor/LandscapeEditor/Private/LandscapeImportHelper.cpp b/Engine/Source/Editor/LandscapeEditor/Private/LandscapeImportHelper.cpp
--- a/Engine/Source/Editor/LandscapeEditor/Private/LandscapeImportHelper.cpp
+++ b/Engine/Source/Editor/LandscapeEditor/Private/LandscapeImportHelper.cpp
@@ -122,7 +122,7 @@ ELandscapeImportResult GetImportDescriptorInternal(const FString& FilePath, bool
 		return ELandscapeImportResult::Error;
 	}
 		
-	FIntPoint OutCoord;
+	FIntPoint OutCoord{};
 	FIntPoint MinCoord(INT32_MAX, INT32_MAX); // All coords should be rebased to the min
 	FIntPoint MaxCoord(INT32_MIN, INT32_MIN);
 	FString OutFileImportPattern;
diff --git a/Engine/Source/Editor/UnrealEd/Private/EditorExporters.cpp b/Engine/Source/Editor/UnrealEd/Private/EditorExporters.cpp
--- a/Engine/Source/Editor/UnrealEd/Private/EditorExporters.cpp
+++ b/Engine/Source/Editor/UnrealEd/Private/EditorExporters.cpp
@@ -1035,7 +1035,7 @@ static void ExportMaterialPropertyTexture(const FString& BMPFilename, UMaterialI
 
 	// make the BMP for the diffuse channel
 	TArray<FColor> OutputBMP;
-	FIntPoint OutSize;
+	FIntPoint OutSize{};
 
 	const TEnumAsByte<EBlendMode> BlendMode = Material->GetBlendMode();
 
diff --git a/Engine/Source/Editor/UnrealEd/Public/FbxImporter.h b/Engine/Source/Editor/UnrealEd/Public/FbxImporter.h
--- a/Engine/Source/Editor/UnrealEd/Public/FbxImporter.h
+++ b/Engine/Source/Editor/UnrealEd/Public/FbxImporter.h
@@ -195,7 +195,7 @@ struct FBXImportOptions
 	// Animation option
 	USkeleton* SkeletonForAnimation;
 	EFBXAnimationLengthImportType AnimationLengthImportType;
-	struct FIntPoint AnimationRange;
+	FIntPoint AnimationRange;
 	FString AnimationName;
 	bool	bPreserveLocalTransform;
 	bool	bDeleteExistingMorphTargetCurves;
diff --git a/Engine/Source/Runtime/Core/Private/Math/UnrealMath.cpp b/Engine/Source/Runtime/Core/Private/Math/UnrealMath.cpp
--- a/Engine/Source/Runtime/Core/Private/Math/UnrealMath.cpp
+++ b/Engine/Source/Runtime/Core/Private/Math/UnrealMath.cpp
@@ -70,11 +70,6 @@ CORE_API const uint32 FMath::BitFlag[32] =
 	(1U << 28),	(1U << 29),	(1U << 30),	(1U << 31),
 };
 
-const FIntPoint FIntPoint::ZeroValue(0,0);
-const FIntPoint FIntPoint::NoneValue(INDEX_NONE,INDEX_NONE);
-const FIntVector FIntVector::ZeroValue(0,0,0);
-const FIntVector FIntVector::NoneValue(INDEX_NONE,INDEX_NONE,INDEX_NONE);
-
 template<typename T>
 bool UE::Math::TRotator<T>::NetSerialize(FArchive& Ar, class UPackageMap* Map, bool& bOutSuccess)
 {
diff --git a/Engine/Source/Runtime/Core/Private/Misc/FileHelper.cpp b/Engine/Source/Runtime/Core/Private/Misc/FileHelper.cpp
--- a/Engine/Source/Runtime/Core/Private/Misc/FileHelper.cpp
+++ b/Engine/Source/Runtime/Core/Private/Misc/FileHelper.cpp
@@ -748,7 +748,7 @@ void FFileHelper::GenerateDateTimeBasedBitmapFilename(const FString& Pattern, co
  * @return true if success
  */
 PRAGMA_DISABLE_DEPRECATION_WARNINGS
-bool FFileHelper::CreateBitmap(const TCHAR* Pattern, int32 SourceWidth, int32 SourceHeight, const FColor* Data, struct FIntRect* SubRectangle, IFileManager* FileManager /*= &IFileManager::Get()*/, FString* OutFilename /*= NULL*/, bool bInWriteAlpha /*= false*/, EChannelMask ChannelMask /*= All */)
+bool FFileHelper::CreateBitmap(const TCHAR* Pattern, int32 SourceWidth, int32 SourceHeight, const FColor* Data, FIntRect* SubRectangle, IFileManager* FileManager /*= &IFileManager::Get()*/, FString* OutFilename /*= NULL*/, bool bInWriteAlpha /*= false*/, EChannelMask ChannelMask /*= All */)
 {
 	EColorChannel ColorChannel = EColorChannel::All;
 	if (ChannelMask != EChannelMask::All)
@@ -766,7 +766,7 @@ bool FFileHelper::CreateBitmap(const TCHAR* Pattern, int32 SourceWidth, int32 So
 }
 PRAGMA_ENABLE_DEPRECATION_WARNINGS
 
-bool FFileHelper::CreateBitmap( const TCHAR* Pattern, int32 SourceWidth, int32 SourceHeight, const FColor* Data, struct FIntRect* SubRectangle, IFileManager* FileManager, FString* OutFilename, bool bInWriteAlpha, EColorChannel ColorChannel)
+bool FFileHelper::CreateBitmap( const TCHAR* Pattern, int32 SourceWidth, int32 SourceHeight, const FColor* Data, FIntRect* SubRectangle, IFileManager* FileManager, FString* OutFilename, bool bInWriteAlpha, EColorChannel ColorChannel)
 {
 	FIntRect Src(0, 0, SourceWidth, SourceHeight);
 	if (SubRectangle == NULL || SubRectangle->Area() == 0)
diff --git a/Engine/Source/Runtime/Core/Public/CoreFwd.h b/Engine/Source/Runtime/Core/Public/CoreFwd.h
--- a/Engine/Source/Runtime/Core/Public/CoreFwd.h
+++ b/Engine/Source/Runtime/Core/Public/CoreFwd.h
@@ -29,17 +29,48 @@ UE_DECLARE_LWC_TYPE(Transform, 3);
 UE_DECLARE_LWC_TYPE(Vector2,, FVector2D);
 UE_DECLARE_LWC_TYPE(Vector, 3);
 UE_DECLARE_LWC_TYPE(Vector4);
-struct FIntPoint;
-struct FIntRect;
 
 namespace UE
 {
 namespace Math
 {
+template <typename IntType>
+struct TIntPoint;
+
+template <typename IntType>
+struct TIntRect;
+
+template <typename IntType>
+struct TIntVector2;
+
+template <typename IntType>
+struct TIntVector3;
+
+template <typename IntType>
+struct TIntVector4;
+
 template<typename T, typename TExtent = T>
 struct TBoxSphereBounds;
 }
 }
+
+using FIntPoint = UE::Math::TIntPoint<int32>;
+using FUintPoint = UE::Math::TIntPoint<uint32>;
+
+using FIntRect = UE::Math::TIntRect<int32>;
+using FUintRect = UE::Math::TIntRect<uint32>;
+
+using FIntVector3 = UE::Math::TIntVector3<int32>;
+using FUintVector3 = UE::Math::TIntVector3<uint32>;
+
+using FIntVector = FIntVector3;
+
+using FIntVector2 = UE::Math::TIntVector2<int32>;
+using FUintVector2 = UE::Math::TIntVector2<uint32>;
+
+using FIntVector4 = UE::Math::TIntVector4<int32>;
+using FUintVector4 = UE::Math::TIntVector4<uint32>;
+
 using FBoxSphereBounds3f = UE::Math::TBoxSphereBounds<float, float>;
 using FBoxSphereBounds3d = UE::Math::TBoxSphereBounds<double, double>;
 // FCompactBoxSphereBounds always stores float extents
diff --git a/Engine/Source/Runtime/Core/Public/Math/IntPoint.h b/Engine/Source/Runtime/Core/Public/Math/IntPoint.h
--- a/Engine/Source/Runtime/Core/Public/Math/IntPoint.h
+++ b/Engine/Source/Runtime/Core/Public/Math/IntPoint.h
@@ -9,40 +9,43 @@
 #include "Serialization/StructuredArchive.h"
 #include "Templates/TypeHash.h"
 
+namespace UE::Math
+{
+
 /**
  * Structure for integer points in 2-d space.
  *
  * @todo Docs: The operators need better documentation, i.e. what does it mean to divide a point?
  */
-struct FIntPoint
+template <typename InIntType>
+struct TIntPoint
 {
+	using IntType = InIntType;
+	static_assert(std::is_integral_v<IntType>, "An integer type is required.");
+
 	union
 	{
 		struct
 		{
 			/** Holds the point's x-coordinate. */
-			int32 X;
+			IntType X;
 
 			/** Holds the point's y-coordinate. */
-			int32 Y;
+			IntType Y;
 		};
 
 		UE_DEPRECATED(all, "For internal use only")
-		int32 XY[2];
+		IntType XY[2];
 	};
 
-public:
-
 	/** An integer point with zeroed values. */
-	CORE_API static const FIntPoint ZeroValue;
+	static const TIntPoint ZeroValue;
 
 	/** An integer point with INDEX_NONE values. */
-	CORE_API static const FIntPoint NoneValue;
-
-public:
+	static const TIntPoint NoneValue;
 
 	/** Default constructor (no initialization). */
-	FIntPoint();
+	TIntPoint() = default;
 
 	/**
 	 * Create and initialize a new instance with the specified coordinates.
@@ -50,7 +53,11 @@ public:
 	 * @param InX The x-coordinate.
 	 * @param InY The y-coordinate.
 	 */
-	FIntPoint(int32 InX, int32 InY);
+	TIntPoint(IntType InX, IntType InY)
+		: X(InX)
+		, Y(InY)
+	{
+	}
 
 	/**
 	 * Create and initialize a new instance with a single int.
@@ -58,16 +65,32 @@ public:
 	 *
 	 * @param InXY The x and y-coordinate.
 	 */
-	FIntPoint(int32 InXY);
+	TIntPoint(IntType InXY)
+		: X(InXY)
+		, Y(InXY)
+	{
+	}
 
 	/**
 	 * Create and initialize a new instance to zero.
 	 *
 	 * @param EForceInit Force init enum
 	 */
-	explicit FORCEINLINE FIntPoint(EForceInit);
+	explicit TIntPoint(EForceInit)
+		: X(0)
+		, Y(0)
+	{
+	}
 
-public:
+	/**
+	 * Converts to another int type. Checks that the cast will succeed.
+	 */
+	template <typename OtherIntType>
+	explicit TIntPoint(TIntPoint<OtherIntType> Other)
+		: X(IntCastChecked<OtherIntType>(Other.X))
+		, Y(IntCastChecked<OtherIntType>(Other.Y))
+	{
+	}
 
 	/**
 	 * Get specific component of a point.
@@ -75,7 +98,12 @@ public:
 	 * @param PointIndex Index of point component.
 	 * @return const reference to component.
 	 */
-	const int32& operator()(int32 PointIndex) const;
+	const IntType& operator()(int32 PointIndex) const
+	{
+PRAGMA_DISABLE_DEPRECATION_WARNINGS
+		return XY[PointIndex];
+PRAGMA_ENABLE_DEPRECATION_WARNINGS
+	}
 
 	/**
 	 * Get specific component of a point.
@@ -83,7 +111,12 @@ public:
 	 * @param PointIndex Index of point component
 	 * @return reference to component.
 	 */
-	int32& operator()(int32 PointIndex);
+	IntType& operator()(int32 PointIndex)
+	{
+PRAGMA_DISABLE_DEPRECATION_WARNINGS
+		return XY[PointIndex];
+PRAGMA_ENABLE_DEPRECATION_WARNINGS
+	}
 
 	/**
 	 * Compare two points for equality.
@@ -91,7 +124,10 @@ public:
 	 * @param Other The other int point being compared.
 	 * @return true if the points are equal, false otherwise.
 	 */
-	bool operator==(const FIntPoint& Other) const;
+	bool operator==(const TIntPoint& Other) const
+	{
+		return X == Other.X && Y == Other.Y;
+	}
 
 	/**
 	 * Compare two points for inequality.
@@ -99,7 +135,10 @@ public:
 	 * @param Other The other int point being compared.
 	 * @return true if the points are not equal, false otherwise.
 	 */
-	bool operator!=(const FIntPoint& Other) const;
+	bool operator!=(const TIntPoint& Other) const
+	{
+		return (X != Other.X) || (Y != Other.Y);
+	}
 
 	/**
 	 * Scale this point.
@@ -107,7 +146,13 @@ public:
 	 * @param Scale What to multiply the point by.
 	 * @return Reference to this point after multiplication.
 	 */
-	FIntPoint& operator*=(int32 Scale);
+	TIntPoint& operator*=(IntType Scale)
+	{
+		X *= Scale;
+		Y *= Scale;
+
+		return *this;
+	}
 
 	/**
 	 * Divide this point by a scalar.
@@ -115,7 +160,13 @@ public:
 	 * @param Divisor What to divide the point by.
 	 * @return Reference to this point after division.
 	 */
-	FIntPoint& operator/=(int32 Divisor);
+	TIntPoint& operator/=(IntType Divisor)
+	{
+		X /= Divisor;
+		Y /= Divisor;
+
+		return *this;
+	}
 
 	/**
 	 * Add another point component-wise to this point.
@@ -123,7 +174,13 @@ public:
 	 * @param Other The point to add to this point.
 	 * @return Reference to this point after addition.
 	 */
-	FIntPoint& operator+=(const FIntPoint& Other);
+	TIntPoint& operator+=(const TIntPoint& Other)
+	{
+		X += Other.X;
+		Y += Other.Y;
+
+		return *this;
+	}
 
 	/**
 	 * Multiply another point component-wise from this point.
@@ -131,7 +188,13 @@ public:
 	 * @param Other The point to multiply with this point.
 	 * @return Reference to this point after multiplication.
 	 */
-	FIntPoint& operator*=(const FIntPoint& Other);
+	TIntPoint& operator*=(const TIntPoint& Other)
+	{
+		X *= Other.X;
+		Y *= Other.Y;
+
+		return *this;
+	}
 
 	/**
 	 * Subtract another point component-wise from this point.
@@ -139,7 +202,13 @@ public:
 	 * @param Other The point to subtract from this point.
 	 * @return Reference to this point after subtraction.
 	 */
-	FIntPoint& operator-=(const FIntPoint& Other);
+	TIntPoint& operator-=(const TIntPoint& Other)
+	{
+		X -= Other.X;
+		Y -= Other.Y;
+
+		return *this;
+	}
 
 	/**
 	 * Divide this point component-wise by another point.
@@ -147,7 +216,13 @@ public:
 	 * @param Other The point to divide with.
 	 * @return Reference to this point after division.
 	 */
-	FIntPoint& operator/=(const FIntPoint& Other);
+	TIntPoint& operator/=(const TIntPoint& Other)
+	{
+		X /= Other.X;
+		Y /= Other.Y;
+
+		return *this;
+	}
 
 	/**
 	 * Assign another point to this one.
@@ -155,7 +230,13 @@ public:
 	 * @param Other The point to assign this point from.
 	 * @return Reference to this point after assignment.
 	 */
-	FIntPoint& operator=(const FIntPoint& Other);
+	TIntPoint& operator=(const TIntPoint& Other)
+	{
+		X = Other.X;
+		Y = Other.Y;
+
+		return *this;
+	}
 
 	/**
 	 * Get the result of scaling on this point.
@@ -163,7 +244,10 @@ public:
 	 * @param Scale What to multiply the point by.
 	 * @return A new scaled int point.
 	 */
-	FIntPoint operator*(int32 Scale) const;
+	TIntPoint operator*(IntType Scale) const
+	{
+		return TIntPoint(*this) *= Scale;
+	}
 
 	/**
 	 * Get the result of division on this point.
@@ -171,7 +255,10 @@ public:
 	 * @param Divisor What to divide the point by.
 	 * @return A new divided int point.
 	 */
-	FIntPoint operator/(int32 Divisor) const;
+	TIntPoint operator/(IntType Divisor) const
+	{
+		return TIntPoint(*this) /= Divisor;
+	}
 
 	/**
 	 * Get the result of addition on this point.
@@ -179,7 +266,10 @@ public:
 	 * @param Other The other point to add to this.
 	 * @return A new combined int point.
 	 */
-	FIntPoint operator+(const FIntPoint& Other) const;
+	TIntPoint operator+(const TIntPoint& Other) const
+	{
+		return TIntPoint(*this) += Other;
+	}
 
 	/**
 	 * Get the result of subtraction from this point.
@@ -187,7 +277,10 @@ public:
 	 * @param Other The other point to subtract from this.
 	 * @return A new subtracted int point.
 	 */
-	FIntPoint operator-(const FIntPoint& Other) const;
+	TIntPoint operator-(const TIntPoint& Other) const
+	{
+		return TIntPoint(*this) -= Other;
+	}
 
 	/**
 	 * Get the result of multiplication on this point.
@@ -195,7 +288,10 @@ public:
 	 * @param Other The point to multiply with this point.
 	 * @return A new multiplied int point
 	 */
-	FIntPoint operator*(const FIntPoint& Other) const;
+	TIntPoint operator*(const TIntPoint& Other) const
+	{
+		return TIntPoint(*this) *= Other;
+	}
 
 	/**
 	 * Get the result of division on this point.
@@ -203,7 +299,10 @@ public:
 	 * @param Other The other point to subtract from this.
 	 * @return A new subtracted int point.
 	 */
-	FIntPoint operator/(const FIntPoint& Other) const;
+	TIntPoint operator/(const TIntPoint& Other) const
+	{
+		return TIntPoint(*this) /= Other;
+	}
 
 	/**
 	* Get specific component of the point.
@@ -211,7 +310,11 @@ public:
 	* @param Index the index of point component
 	* @return reference to component.
 	*/
-	int32& operator[](int32 Index);
+	IntType& operator[](IntType Index)
+	{
+		check(Index >= 0 && Index < 2);
+		return ((Index == 0) ? X : Y);
+	}
 
 	/**
 	* Get specific component of the point.
@@ -219,23 +322,31 @@ public:
 	* @param Index the index of point component
 	* @return copy of component value.
 	*/
-	int32 operator[](int32 Index) const;
-
-public:
+	IntType operator[](IntType Index) const
+	{
+		check(Index >= 0 && Index < 2);
+		return ((Index == 0) ? X : Y);
+	}
 
 	/**
 	 * Get the component-wise min of two points.
 	 *
 	 * @see ComponentMax, GetMax
 	 */
-	FORCEINLINE FIntPoint ComponentMin(const FIntPoint& Other) const;
+	TIntPoint ComponentMin(const TIntPoint& Other) const
+	{
+		return TIntPoint(FMath::Min(X, Other.X), FMath::Min(Y, Other.Y));
+	}
 
 	/**
 	 * Get the component-wise max of two points.
 	 *
 	 * @see ComponentMin, GetMin
 	 */
-	FORCEINLINE FIntPoint ComponentMax(const FIntPoint& Other) const;
+	TIntPoint ComponentMax(const TIntPoint& Other) const
+	{
+		return TIntPoint(FMath::Max(X, Other.X), FMath::Max(Y, Other.Y));
+	}
 
 	/**
 	 * Get the larger of the point's two components.
@@ -243,7 +354,10 @@ public:
 	 * @return The maximum component of the point.
 	 * @see GetMin, Size, SizeSquared
 	 */
-	int32 GetMax() const;
+	IntType GetMax() const
+	{
+		return FMath::Max(X, Y);
+	}
 
 	/**
 	 * Get the smaller of the point's two components.
@@ -251,7 +365,10 @@ public:
 	 * @return The minimum component of the point.
 	 * @see GetMax, Size, SizeSquared
 	 */
-	int32 GetMin() const;
+	IntType GetMin() const
+	{
+		return FMath::Min(X, Y);
+	}
 
 	/**
 	 * Get the distance of this point from (0,0).
@@ -259,7 +376,12 @@ public:
 	 * @return The distance of this point from (0,0).
 	 * @see GetMax, GetMin, SizeSquared
 	 */
-	int32 Size() const;
+	IntType Size() const
+	{
+		int64 LocalX64 = (int64)X;
+		int64 LocalY64 = (int64)Y;
+		return IntType(FMath::Sqrt(double(LocalX64 * LocalX64 + LocalY64 * LocalY64)));
+	}
 
 	/**
 	 * Get the squared distance of this point from (0,0).
@@ -267,16 +389,20 @@ public:
 	 * @return The squared distance of this point from (0,0).
 	 * @see GetMax, GetMin, Size
 	 */
-	int32 SizeSquared() const;
-	
+	IntType SizeSquared() const
+	{
+		return X * X + Y * Y;
+	}
+
 	/**
 	 * Get a textual representation of this point.
 	 *
 	 * @return A string describing the point.
 	 */
-	FString ToString() const;
-
-public:
+	FString ToString() const
+	{
+		return FString::Printf(TEXT("X=%s Y=%s"), *LexToString(X), *LexToString(Y));
+	}
 
 	/**
 	 * Divide an int point and round up the result.
@@ -286,8 +412,15 @@ public:
 	 * @return A new divided int point.
 	 * @see DivideAndRoundDown
 	 */
-	static FIntPoint DivideAndRoundUp(FIntPoint lhs, int32 Divisor);
-	static FIntPoint DivideAndRoundUp(FIntPoint lhs, FIntPoint Divisor);
+	static TIntPoint DivideAndRoundUp(TIntPoint lhs, IntType Divisor)
+	{
+		return TIntPoint(FMath::DivideAndRoundUp(lhs.X, Divisor), FMath::DivideAndRoundUp(lhs.Y, Divisor));
+	}
+
+	static TIntPoint DivideAndRoundUp(TIntPoint lhs, TIntPoint Divisor)
+	{
+		return TIntPoint(FMath::DivideAndRoundUp(lhs.X, Divisor.X), FMath::DivideAndRoundUp(lhs.Y, Divisor.Y));
+	}
 
 	/**
 	 * Divide an int point and round down the result.
@@ -297,16 +430,20 @@ public:
 	 * @return A new divided int point.
 	 * @see DivideAndRoundUp
 	 */
-	static FIntPoint DivideAndRoundDown(FIntPoint lhs, int32 Divisor);
+	static TIntPoint DivideAndRoundDown(TIntPoint lhs, IntType Divisor)
+	{
+		return TIntPoint(FMath::DivideAndRoundDown(lhs.X, Divisor), FMath::DivideAndRoundDown(lhs.Y, Divisor));
+	}
 
 	/**
 	 * Get number of components point has.
 	 *
 	 * @return number of components point has.
 	 */
-	static int32 Num();
-
-public:
+	static int32 Num()
+	{
+		return 2;
+	}
 
 	/**
 	 * Serialize the point.
@@ -315,7 +452,7 @@ public:
 	 * @param Point The point to serialize.
 	 * @return Reference to the Archive after serialization.
 	 */
-	friend FArchive& operator<<(FArchive& Ar, FIntPoint& Point)
+	friend FArchive& operator<<(FArchive& Ar, TIntPoint& Point)
 	{
 		return Ar << Point.X << Point.Y;
 	}
@@ -326,7 +463,7 @@ public:
 	 * @param Slot The structured archive slot to serialize into.
 	 * @param Point The point to serialize.
 	 */
-	friend void operator<<(FStructuredArchive::FSlot Slot, FIntPoint& Point)
+	friend void operator<<(FStructuredArchive::FSlot Slot, TIntPoint& Point)
 	{
 		FStructuredArchive::FRecord Record = Slot.EnterRecord();
 		Record << SA_VALUE(TEXT("X"), Point.X) << SA_VALUE(TEXT("Y"), Point.Y);
@@ -357,235 +494,70 @@ public:
 	}
 };
 
+template <typename IntType>
+const TIntPoint<IntType> TIntPoint<IntType>::ZeroValue(0, 0);
 
-/* FIntPoint inline functions
- *****************************************************************************/
-
-FORCEINLINE FIntPoint::FIntPoint() { }
-
-
-FORCEINLINE FIntPoint::FIntPoint(int32 InX, int32 InY)
-	: X(InX)
-	, Y(InY)
-{ }
-
-FORCEINLINE FIntPoint::FIntPoint(int32 InXY)
-	: X(InXY)
-	, Y(InXY)
-{ }
-
-FORCEINLINE FIntPoint::FIntPoint(EForceInit)
-	: X(0)
-	, Y(0)
-{ }
-
-
-FORCEINLINE const int32& FIntPoint::operator()(int32 PointIndex) const
-{
-PRAGMA_DISABLE_DEPRECATION_WARNINGS
-	return XY[PointIndex];
-PRAGMA_ENABLE_DEPRECATION_WARNINGS
-}
-
-
-FORCEINLINE int32& FIntPoint::operator()(int32 PointIndex)
-{
-PRAGMA_DISABLE_DEPRECATION_WARNINGS
-	return XY[PointIndex];
-PRAGMA_ENABLE_DEPRECATION_WARNINGS
-}
-
-
-FORCEINLINE int32 FIntPoint::Num()
-{
-	return 2;
-}
-
-
-FORCEINLINE bool FIntPoint::operator==(const FIntPoint& Other) const
-{
-	return X==Other.X && Y==Other.Y;
-}
-
-
-FORCEINLINE bool FIntPoint::operator!=(const FIntPoint& Other) const
-{
-	return (X != Other.X) || (Y != Other.Y);
-}
-
-
-FORCEINLINE FIntPoint& FIntPoint::operator*=(int32 Scale)
-{
-	X *= Scale;
-	Y *= Scale;
-
-	return *this;
-}
-
+template <typename IntType>
+const TIntPoint<IntType> TIntPoint<IntType>::NoneValue(static_cast<IntType>(INDEX_NONE), static_cast<IntType>(INDEX_NONE));
 
-FORCEINLINE FIntPoint& FIntPoint::operator/=(int32 Divisor)
+template <typename IntType>
+uint32 GetTypeHash(const TIntPoint<IntType>& InPoint)
 {
-	X /= Divisor;
-	Y /= Divisor;
-
-	return *this;
+	return HashCombine(::GetTypeHash(InPoint.X), ::GetTypeHash(InPoint.Y));
 }
 
-
-FORCEINLINE FIntPoint& FIntPoint::operator+=(const FIntPoint& Other)
+template <>
+inline FString TIntPoint<int64>::ToString() const
 {
-	X += Other.X;
-	Y += Other.Y;
-
-	return *this;
+	return FString::Printf(TEXT("X=%lld Y=%lld"), X, Y);
 }
 
-
-FORCEINLINE FIntPoint& FIntPoint::operator-=(const FIntPoint& Other)
+template <>
+inline FString TIntPoint<int32>::ToString() const
 {
-	X -= Other.X;
-	Y -= Other.Y;
-
-	return *this;
-}
-
-
-FORCEINLINE FIntPoint& FIntPoint::operator*=(const FIntPoint& Other)
-{
-	X *= Other.X;
-	Y *= Other.Y;
-
-	return *this;
-}
-
-
-FORCEINLINE FIntPoint& FIntPoint::operator/=(const FIntPoint& Other)
-{
-	X /= Other.X;
-	Y /= Other.Y;
-
-	return *this;
-}
-
-
-FORCEINLINE FIntPoint& FIntPoint::operator=(const FIntPoint& Other)
-{
-	X = Other.X;
-	Y = Other.Y;
-
-	return *this;
-}
-
-
-FORCEINLINE FIntPoint FIntPoint::operator*(int32 Scale) const
-{
-	return FIntPoint(*this) *= Scale;
-}
-
-
-FORCEINLINE FIntPoint FIntPoint::operator/(int32 Divisor) const
-{
-	return FIntPoint(*this) /= Divisor;
-}
-
-
-FORCEINLINE int32& FIntPoint::operator[](int32 Index)
-{
-	check(Index >= 0 && Index < 2);
-	return ((Index == 0) ? X : Y);
-}
-
-
-FORCEINLINE int32 FIntPoint::operator[](int32 Index) const
-{
-	check(Index >= 0 && Index < 2);
-	return ((Index == 0) ? X : Y);
-}
-
-
-FORCEINLINE FIntPoint FIntPoint::ComponentMin(const FIntPoint& Other) const
-{
-	return FIntPoint(FMath::Min(X, Other.X), FMath::Min(Y, Other.Y));
-}
-
-
-FORCEINLINE FIntPoint FIntPoint::ComponentMax(const FIntPoint& Other) const
-{
-	return FIntPoint(FMath::Max(X, Other.X), FMath::Max(Y, Other.Y));
-}
-
-FORCEINLINE FIntPoint FIntPoint::DivideAndRoundUp(FIntPoint lhs, int32 Divisor)
-{
-	return FIntPoint(FMath::DivideAndRoundUp(lhs.X, Divisor), FMath::DivideAndRoundUp(lhs.Y, Divisor));
-}
-
-FORCEINLINE FIntPoint FIntPoint::DivideAndRoundUp(FIntPoint lhs, FIntPoint Divisor)
-{
-	return FIntPoint(FMath::DivideAndRoundUp(lhs.X, Divisor.X), FMath::DivideAndRoundUp(lhs.Y, Divisor.Y));
-}	
-
-FORCEINLINE FIntPoint FIntPoint::DivideAndRoundDown(FIntPoint lhs, int32 Divisor)
-{
-	return FIntPoint(FMath::DivideAndRoundDown(lhs.X, Divisor), FMath::DivideAndRoundDown(lhs.Y, Divisor));
-}	
-
-
-FORCEINLINE FIntPoint FIntPoint::operator+(const FIntPoint& Other) const
-{
-	return FIntPoint(*this) += Other;
-}
-
-
-FORCEINLINE FIntPoint FIntPoint::operator-(const FIntPoint& Other) const
-{
-	return FIntPoint(*this) -= Other;
+	return FString::Printf(TEXT("X=%d Y=%d"), X, Y);
 }
 
-
-FORCEINLINE FIntPoint FIntPoint::operator*(const FIntPoint& Other) const
+template <>
+inline FString TIntPoint<int16>::ToString() const
 {
-	return FIntPoint(*this) *= Other;
+	return FString::Printf(TEXT("X=%d Y=%d"), X, Y);
 }
 
-
-FORCEINLINE FIntPoint FIntPoint::operator/(const FIntPoint& Other) const
+template <>
+inline FString TIntPoint<int8>::ToString() const
 {
-	return FIntPoint(*this) /= Other;
+	return FString::Printf(TEXT("X=%d Y=%d"), X, Y);
 }
 
-
-FORCEINLINE int32 FIntPoint::GetMax() const
+template <>
+inline FString TIntPoint<uint64>::ToString() const
 {
-	return FMath::Max(X, Y);
+	return FString::Printf(TEXT("X=%llu Y=%llu"), X, Y);
 }
 
-
-FORCEINLINE int32 FIntPoint::GetMin() const
+template <>
+inline FString TIntPoint<uint32>::ToString() const
 {
-	return FMath::Min(X,Y);
+	return FString::Printf(TEXT("X=%u Y=%u"), X, Y);
 }
 
-FORCEINLINE uint32 GetTypeHash(const FIntPoint& InPoint)
+template <>
+inline FString TIntPoint<uint16>::ToString() const
 {
-	return HashCombine(GetTypeHash(InPoint.X), GetTypeHash(InPoint.Y));
+	return FString::Printf(TEXT("X=%u Y=%u"), X, Y);
 }
 
-
-FORCEINLINE int32 FIntPoint::Size() const
+template <>
+inline FString TIntPoint<uint8>::ToString() const
 {
-	int64 LocalX64 = (int64)X;
-	int64 LocalY64 = (int64)Y;
-	return int32(FMath::Sqrt(float(LocalX64 * LocalX64 + LocalY64 * LocalY64)));
+	return FString::Printf(TEXT("X=%u Y=%u"), X, Y);
 }
 
-FORCEINLINE int32 FIntPoint::SizeSquared() const
-{
-	return X*X + Y*Y;
-}
+} //! namespace UE::Math
 
-FORCEINLINE FString FIntPoint::ToString() const
-{
-	return FString::Printf(TEXT("X=%d Y=%d"), X, Y);
-}
+template <> struct TIsPODType<FIntPoint>  { enum { Value = true }; };
+template <> struct TIsPODType<FUintPoint> { enum { Value = true }; };
 
-template <> struct TIsPODType<FIntPoint> { enum { Value = true }; };
+template<> struct TIsUECoreType<FIntPoint>  { enum { Value = true }; };
+template<> struct TIsUECoreType<FUintPoint> { enum { Value = true }; };
\ No newline at end of file
diff --git a/Engine/Source/Runtime/Core/Public/Math/IntRect.h b/Engine/Source/Runtime/Core/Public/Math/IntRect.h
--- a/Engine/Source/Runtime/Core/Public/Math/IntRect.h
+++ b/Engine/Source/Runtime/Core/Public/Math/IntRect.h
@@ -8,24 +8,41 @@
 #include "Math/IntPoint.h"
 #include "Math/Vector2D.h"
 
+namespace UE::Math
+{
 
 /**
  * Structure for integer rectangles in 2-d space.
  *
  * @todo Docs: The operators need better documentation, i.e. what does it mean to divide a rectangle?
  */
-struct FIntRect
+template <typename InIntType>
+struct TIntRect
 {
-	/** Holds the first pixel line/row (like in Win32 RECT). */
-	FIntPoint Min;
+	using IntType = InIntType;
+	using IntPointType = TIntPoint<IntType>;
+	static_assert(std::is_integral_v<IntType>, "Only an integer types are supported.");
+
+	union
+	{
+		struct
+		{
+			/** Holds the first pixel line/row (like in Win32 RECT). */
+			IntPointType Min;
 
-	/** Holds the last pixel line/row (like in Win32 RECT). */
-	FIntPoint Max;
+			/** Holds the last pixel line/row (like in Win32 RECT). */
+			IntPointType Max;
+		};
 
-public:
+		UE_DEPRECATED(all, "For internal use only")
+		IntPointType MinMax[2];
+	};
 
 	/** Constructor */
-	FIntRect();
+	TIntRect()
+		: Min(ForceInit)
+		, Max(ForceInit)
+	{}
 
 	/**
 	 * Constructor
@@ -35,7 +52,11 @@ public:
 	 * @param X1 Maximum X coordinate.
 	 * @param Y1 Maximum Y coordinate.
 	 */
-	FIntRect( int32 X0, int32 Y0, int32 X1, int32 Y1 );
+	TIntRect(IntType X0, IntType Y0, IntType X1, IntType Y1)
+		: Min(X0, Y0)
+		, Max(X1, Y1)
+	{
+	}
 
 	/**
 	 * Constructor
@@ -43,9 +64,33 @@ public:
 	 * @param InMin Minimum Point
 	 * @param InMax Maximum Point
 	 */
-	FIntRect( FIntPoint InMin, FIntPoint InMax );
+	TIntRect(IntPointType InMin, IntPointType InMax)
+		: Min(InMin)
+		, Max(InMax)
+	{
+	}
 
-public:
+	TIntRect(const TIntRect& Other)
+	{
+		*this = Other;
+	}
+
+	TIntRect& operator=(const TIntRect& Other)
+	{
+		Min = Other.Min;
+		Max = Other.Max;
+		return *this;
+	}
+
+	/**
+	 * Converts to another int type. Checks that the cast will succeed.
+	 */
+	template <typename OtherIntType>
+	explicit TIntRect(TIntRect<OtherIntType> Other)
+		: Min(IntPointType(Other.Min))
+		, Max(IntPointType(Other.Max))
+	{
+	}
 
 	/**
 	 * Gets a specific point in this rectangle.
@@ -53,7 +98,12 @@ public:
 	 * @param PointIndex Index of Point in rectangle.
 	 * @return Const reference to point in rectangle.
 	 */
-	const FIntPoint& operator()( int32 PointIndex ) const;
+	const TIntRect& operator()(int32 PointIndex) const
+	{
+PRAGMA_DISABLE_DEPRECATION_WARNINGS
+		return MinMax[PointIndex];
+PRAGMA_ENABLE_DEPRECATION_WARNINGS
+	}
 
 	/**
 	 * Gets a specific point in this rectangle.
@@ -61,7 +111,12 @@ public:
 	 * @param PointIndex Index of Point in rectangle.
 	 * @return Reference to point in rectangle.
 	 */
-	FIntPoint& operator()( int32 PointIndex );
+	TIntRect& operator()(int32 PointIndex)
+	{
+PRAGMA_DISABLE_DEPRECATION_WARNINGS
+		return MinMax[PointIndex];
+PRAGMA_ENABLE_DEPRECATION_WARNINGS
+	}
 
 	/**
 	 * Compares Rectangles for equality.
@@ -69,7 +124,10 @@ public:
 	 * @param Other The Other Rectangle for comparison.
 	 * @return true if the rectangles are equal, false otherwise..
 	 */
-	bool operator==( const FIntRect& Other ) const;
+	bool operator==(const TIntRect& Other) const
+	{
+		return Min == Other.Min && Max == Other.Max;
+	}
 
 	/**
 	 * Compares Rectangles for inequality.
@@ -77,7 +135,10 @@ public:
 	 * @param Other The Other Rectangle for comparison.
 	 * @return true if the rectangles are not equal, false otherwise..
 	 */
-	bool operator!=( const FIntRect& Other ) const;
+	bool operator!=(const TIntRect& Other) const
+	{
+		return Min != Other.Min || Max != Other.Max;
+	}
 
 	/**
 	 * Applies scaling to this rectangle.
@@ -85,7 +146,13 @@ public:
 	 * @param Scale What to multiply the rectangle by.
 	 * @return Reference to this rectangle after scaling.
 	 */
-	FIntRect& operator*=( int32 Scale );
+	TIntRect& operator*=(IntType Scale)
+	{
+		Min *= Scale;
+		Max *= Scale;
+
+		return *this;
+	}
 
 	/**
 	 * Adds a point to this rectangle.
@@ -93,7 +160,13 @@ public:
 	 * @param Point The point to add onto both points in the rectangle.
 	 * @return Reference to this rectangle after addition.
 	 */
-	FIntRect& operator+=( const FIntPoint& Point );
+	TIntRect& operator+=(const IntPointType& Point)
+	{
+		Min += Point;
+		Max += Point;
+
+		return *this;
+	}
 
 	/**
 	 * Subtracts a point from this rectangle.
@@ -101,7 +174,13 @@ public:
 	 * @param Point The point to subtract from both points in the rectangle.
 	 * @return Reference to this rectangle after subtraction.
 	 */
-	FIntRect& operator-=( const FIntPoint& Point );
+	TIntRect& operator-=(const IntPointType& Point)
+	{
+		Min -= Point;
+		Max -= Point;
+
+		return *this;
+	}
 
 	/**
 	 * Gets the result of scaling on this rectangle.
@@ -109,7 +188,10 @@ public:
 	 * @param Scale What to multiply this rectangle by.
 	 * @return New scaled rectangle.
 	 */
-	FIntRect operator*( int32 Scale ) const;
+	TIntRect operator*(IntType Scale) const
+	{
+		return TIntRect(Min * Scale, Max * Scale);
+	}
 
 	/**
 	 * Gets the result of division on this rectangle.
@@ -117,7 +199,10 @@ public:
 	 * @param Div What to divide this rectangle by.
 	 * @return New divided rectangle.
 	 */
-	FIntRect operator/( int32 Div ) const;
+	TIntRect operator/(IntType Div) const
+	{
+		return TIntRect(Min / Div, Max / Div);
+	}
 
 	/**
 	 * Gets the result of adding a point to this rectangle.
@@ -125,7 +210,10 @@ public:
 	 * @param Point The point to add to both points in the rectangle.
 	 * @return New rectangle with point added to it.
 	 */
-	FIntRect operator+( const FIntPoint& Point ) const;
+	TIntRect operator+(const IntPointType& Point) const
+	{
+		return TIntRect(Min + Point, Max + Point);
+	}
 
 	/**
 	 * Gets the result of dividing a point with this rectangle.
@@ -133,7 +221,10 @@ public:
 	 * @param Point The point to divide with.
 	 * @return New rectangle with point divided.
 	 */
-	FIntRect operator/( const FIntPoint& Point ) const;
+	TIntRect operator/(const IntPointType& Point) const
+	{
+		return TIntRect(Min / Point, Max / Point);
+	}
 
 	/**
 	 * Gets the result of subtracting a point from this rectangle.
@@ -141,7 +232,10 @@ public:
 	 * @param Point The point to subtract from both points in the rectangle.
 	 * @return New rectangle with point subtracted from it.
 	 */
-	FIntRect operator-( const FIntPoint& Point ) const;
+	TIntRect operator-(const IntPointType& Point) const
+	{
+		return TIntRect(Min - Point, Max - Point);
+	}
 
 	/**
 	 * Gets the result of adding two rectangles together.
@@ -149,7 +243,10 @@ public:
 	 * @param Other The other rectangle to add to this.
 	 * @return New rectangle after both are added together.
 	 */
-	FIntRect operator+( const FIntRect& Other ) const;
+	TIntRect operator+(const TIntRect& Other) const
+	{
+		return TIntRect(Min + Other.Min, Max + Other.Max);
+	}
 
 	/**
 	 * Gets the result of subtracting a rectangle from this one.
@@ -157,16 +254,20 @@ public:
 	 * @param Other The other rectangle to subtract from this.
 	 * @return New rectangle after one is subtracted from this.
 	 */
-	FIntRect operator-( const FIntRect& Other ) const;
-
-public:
+	TIntRect operator-(const TIntRect& Other) const
+	{
+		return TIntRect(Min - Other.Min, Max - Other.Max);
+	}
 
 	/**
 	 * Calculates the area of this rectangle.
 	 *
 	 * @return The area of this rectangle.
 	 */
-	int32 Area() const;
+	IntType Area() const
+	{
+		return (Max.X - Min.X) * (Max.Y - Min.Y);
+	}
 
 	/**
 	 * Creates a rectangle from the bottom part of this rectangle.
@@ -174,20 +275,43 @@ public:
 	 * @param InHeight Height of the new rectangle (<= rectangles original height).
 	 * @return The new rectangle.
 	 */
-	FIntRect Bottom( int32 InHeight ) const;
+
+	TIntRect Bottom(IntType InHeight) const
+	{
+		return TIntRect(Min.X, FMath::Max(Min.Y, Max.Y - InHeight), Max.X, Max.Y);
+	}
 
 	/**
 	 * Clip a rectangle using the bounds of another rectangle.
 	 *
 	 * @param Other The other rectangle to clip against.
 	 */
-	void Clip( const FIntRect& Other );
+	void Clip(const TIntRect& R)
+	{
+		Min.X = FMath::Max<IntType>(Min.X, R.Min.X);
+		Min.Y = FMath::Max<IntType>(Min.Y, R.Min.Y);
+		Max.X = FMath::Min<IntType>(Max.X, R.Max.X);
+		Max.Y = FMath::Min<IntType>(Max.Y, R.Max.Y);
+
+		// return zero area if not overlapping
+		Max.X = FMath::Max<IntType>(Min.X, Max.X);
+		Max.Y = FMath::Max<IntType>(Min.Y, Max.Y);
+	}
 
 	/** Combines the two rectanges. */
-	void Union( const FIntRect& Other);
+	void Union(const TIntRect& R)
+	{
+		Min.X = FMath::Min<IntType>(Min.X, R.Min.X);
+		Min.Y = FMath::Min<IntType>(Min.Y, R.Min.Y);
+		Max.X = FMath::Max<IntType>(Max.X, R.Max.X);
+		Max.Y = FMath::Max<IntType>(Max.Y, R.Max.Y);
+	}
 
 	/** Returns true if the two rects have any overlap. */
-	bool Intersect(const FIntRect& Other) const;
+	bool Intersect(const TIntRect& Other) const
+	{
+		return Other.Min.X < Max.X&& Other.Max.X > Min.X && Other.Min.Y < Max.Y&& Other.Max.Y > Min.Y;
+	}
 
 	/**
 	 * Test whether this rectangle contains a point.
@@ -195,7 +319,10 @@ public:
 	 * @param Point The point to test against.
 	 * @return true if the rectangle contains the specified point,, false otherwise..
 	 */
-	bool Contains( FIntPoint Point ) const;
+	bool Contains(IntPointType P) const
+	{
+		return P.X >= Min.X && P.X < Max.X&& P.Y >= Min.Y && P.Y < Max.Y;
+	}
 
 	/**
 	 * Gets the Center and Extents of this rectangle.
@@ -203,28 +330,50 @@ public:
 	 * @param OutCenter Will contain the center point.
 	 * @param OutExtent Will contain the extent.
 	 */
-	void GetCenterAndExtents( FIntPoint& OutCenter, FIntPoint& OutExtent ) const;
+	void GetCenterAndExtents(IntPointType& OutCenter, IntPointType& OutExtent) const
+	{
+		OutExtent.X = (Max.X - Min.X) / 2;
+		OutExtent.Y = (Max.Y - Min.Y) / 2;
+
+		OutCenter.X = Min.X + OutExtent.X;
+		OutCenter.Y = Min.Y + OutExtent.Y;
+	}
 
 	/**
 	 * Gets the Height of the rectangle.
 	 *
 	 * @return The Height of the rectangle.
 	 */
-	int32 Height() const;
+	IntType Height() const
+	{
+		return (Max.Y - Min.Y);
+	}
+
 
 	/**
 	 * Inflates or deflates the rectangle.
 	 *
 	 * @param Amount The amount to inflate or deflate the rectangle on each side.
 	 */
-	void InflateRect( int32 Amount );
-
+	void InflateRect(IntType Amount)
+	{
+		Min.X -= Amount;
+		Min.Y -= Amount;
+		Max.X += Amount;
+		Max.Y += Amount;
+	}
 	/**
 	 * Adds to this rectangle to include a given point.
 	 *
 	 * @param Point The point to increase the rectangle to.
 	 */
-	void Include( FIntPoint Point );
+	void Include(IntPointType Point)
+	{
+		Min.X = FMath::Min(Min.X, Point.X);
+		Min.Y = FMath::Min(Min.Y, Point.Y);
+		Max.X = FMath::Max(Max.X, Point.X);
+		Max.Y = FMath::Max(Max.Y, Point.Y);
+	}
 
 	/**
 	 * Gets a new rectangle from the inner of this one.
@@ -232,15 +381,20 @@ public:
 	 * @param Shrink How much to remove from each point of this rectangle.
 	 * @return New inner Rectangle.
 	 */
-	FIntRect Inner( FIntPoint Shrink ) const;
-
+	TIntRect Inner(IntPointType Shrink) const
+	{
+		return TIntRect(Min + Shrink, Max - Shrink);
+	}
 	/**
 	 * Creates a rectangle from the right hand side of this rectangle.
 	 *
 	 * @param InWidth Width of the new rectangle (<= rectangles original width).
 	 * @return The new rectangle.
 	 */
-	FIntRect Right( int32 InWidth ) const;
+	TIntRect Right(IntType InWidth) const
+	{
+		return TIntRect(FMath::Max(Min.X, Max.X - InWidth), Min.Y, Max.X, Max.Y);
+	}
 
 	/**
 	 * Scales a rectangle using a floating point number.
@@ -248,37 +402,59 @@ public:
 	 * @param Fraction What to scale the rectangle by
 	 * @return New scaled rectangle.
 	 */
-	FIntRect Scale( float Fraction ) const;
+	TIntRect Scale(double Fraction) const
+	{
+		using Vec2D = UE::Math::TVector2<double>;
+		const Vec2D Min2D = Vec2D((double)Min.X, (double)Min.Y) * Fraction;
+		const Vec2D Max2D = Vec2D((double)Max.X, (double)Max.Y) * Fraction;
+
+		return TIntRect(
+			IntCastChecked<IntType>(FMath::FloorToInt64(Min2D.X)),
+			IntCastChecked<IntType>(FMath::FloorToInt64(Min2D.Y)),
+			IntCastChecked<IntType>(FMath::CeilToInt64(Max2D.X)),
+			IntCastChecked<IntType>(FMath::CeilToInt64(Max2D.Y))
+		);
+	}
 
 	/**
 	 * Gets the distance from one corner of the rectangle to the other.
 	 *
 	 * @return The distance from one corner of the rectangle to the other.
 	 */
-	FIntPoint Size() const;
+	IntPointType Size() const
+	{
+		return IntPointType(Max.X - Min.X, Max.Y - Min.Y);
+	}
 
 	/**
 	 * Get a textual representation of this rectangle.
 	 *
 	 * @return A string describing the rectangle.
 	 */
-	FString ToString() const;
+	FString ToString() const
+	{
+		return FString::Printf(TEXT("Min=(%s) Max=(%s)"), *Min.ToString(), *Max.ToString());
+	}
 
 	/**
 	 * Gets the width of the rectangle.
 	 *
 	 * @return The width of the rectangle.
 	 */
-	int32 Width() const;
+	IntType Width() const
+	{
+		return Max.X - Min.X;
+	}
 
 	/**
 	 * Returns true if the rectangle is 0 x 0.
 	 *
 	 * @return true if the rectangle is 0 x 0.
 	 */
-	bool IsEmpty() const;
-
-public:
+	bool IsEmpty() const
+	{
+		return Width() == 0 && Height() == 0;
+	}
 
 	/**
 	 * Divides a rectangle and rounds up to the nearest integer.
@@ -287,15 +463,25 @@ public:
 	 * @param Div What to divide by.
 	 * @return New divided rectangle.
 	 */
-	static FIntRect DivideAndRoundUp( FIntRect lhs, int32 Div );
-	static FIntRect DivideAndRoundUp( FIntRect lhs, FIntPoint Div );
+	static TIntRect DivideAndRoundUp(TIntRect lhs, IntType Div)
+	{
+		return DivideAndRoundUp(lhs, IntPointType(Div, Div));
+	}
+
+	static TIntRect DivideAndRoundUp(TIntRect lhs, IntPointType Div)
+	{
+		return TIntRect(lhs.Min / Div, IntPointType::DivideAndRoundUp(lhs.Max, Div));
+	}
 
 	/**
 	 * Gets number of points in the Rectangle.
 	 *
 	 * @return Number of points in the Rectangle.
 	 */
-	static int32 Num();
+	static int32 Num()
+	{
+		return 2;
+	}
 
 public:
 
@@ -306,259 +492,19 @@ public:
 	 * @param Rect The rectangle to serialize.
 	 * @return Reference to the Archive after serialization.
 	 */
-	friend FArchive& operator<<( FArchive& Ar, FIntRect& Rect )
+	friend FArchive& operator<<( FArchive& Ar, TIntRect& Rect )
 	{
 		return Ar << Rect.Min.X << Rect.Min.Y << Rect.Max.X << Rect.Max.Y;
 	}
 };
 
-
-FORCEINLINE FIntRect FIntRect::Scale( float Fraction ) const
-{
-	using Vec2D = UE::Math::TVector2<float>;
-	const Vec2D Min2D = Vec2D((float)Min.X, (float)Min.Y) * Fraction;
-	const Vec2D Max2D = Vec2D((float)Max.X, (float)Max.Y) * Fraction;
-
-	return FIntRect(FMath::FloorToInt(Min2D.X), FMath::FloorToInt(Min2D.Y), FMath::CeilToInt(Max2D.X), FMath::CeilToInt(Max2D.Y));
-}
-
-
-/* FIntRect inline functions
- *****************************************************************************/
-
-FORCEINLINE FIntRect::FIntRect()
-	: Min(ForceInit)
-	, Max(ForceInit)
-{ }
-
-
-FORCEINLINE FIntRect::FIntRect( int32 X0, int32 Y0, int32 X1, int32 Y1 )
-	: Min(X0, Y0)
-	, Max(X1, Y1)
-{ }
-
-
-FORCEINLINE FIntRect::FIntRect( FIntPoint InMin, FIntPoint InMax )
-	: Min(InMin)
-	, Max(InMax)
-{ }
-
-
-FORCEINLINE const FIntPoint& FIntRect::operator()( int32 PointIndex ) const
-{
-	return (&Min)[PointIndex];
-}
-
-
-FORCEINLINE FIntPoint& FIntRect::operator()( int32 PointIndex )
-{
-	return (&Min)[PointIndex];
-}
-
-
-FORCEINLINE bool FIntRect::operator==( const FIntRect& Other ) const
-{
-	return Min==Other.Min && Max==Other.Max;
-}
-
-
-FORCEINLINE bool FIntRect::operator!=( const FIntRect& Other ) const
-{
-	return Min!=Other.Min || Max!=Other.Max;
-}
-
-
-FORCEINLINE FIntRect& FIntRect::operator*=( int32 Scale )
-{
-	Min *= Scale;
-	Max *= Scale;
-
-	return *this;
-}
-
-
-FORCEINLINE FIntRect& FIntRect::operator+=( const FIntPoint& Point )
-{
-	Min += Point;
-	Max += Point;
-
-	return *this;
-}
-
-
-FORCEINLINE FIntRect& FIntRect::operator-=( const FIntPoint& Point )
-{
-	Min -= Point;
-	Max -= Point;
-
-	return *this;
-}
-
-
-FORCEINLINE FIntRect FIntRect::operator*( int32 Scale ) const
-{
-	return FIntRect(Min * Scale, Max * Scale);
-}
-
-
-FORCEINLINE FIntRect FIntRect::operator/( int32 Div ) const
-{
-	return FIntRect(Min / Div, Max / Div);
-}
-
-
-FORCEINLINE FIntRect FIntRect::operator+( const FIntPoint& Point ) const
-{
-	return FIntRect(Min + Point, Max + Point);
-}
-
-
-FORCEINLINE FIntRect FIntRect::operator/( const FIntPoint& Point ) const
-{
-	return FIntRect(Min / Point, Max / Point);
-}
-
-
-FORCEINLINE FIntRect FIntRect::operator-( const FIntPoint& Point ) const
-{
-	return FIntRect(Min - Point, Max - Point);
-}
-
-
-FORCEINLINE FIntRect FIntRect::operator+( const FIntRect& Other ) const
-{
-	return FIntRect(Min + Other.Min, Max + Other.Max);
-}
-
-
-FORCEINLINE FIntRect FIntRect::operator-( const FIntRect& Other ) const
-{
-	return FIntRect(Min - Other.Min, Max - Other.Max);
-}
-
-
-FORCEINLINE int32 FIntRect::Area() const
-{
-	return (Max.X - Min.X) * (Max.Y - Min.Y);
-}
-
-
-FORCEINLINE FIntRect FIntRect::Bottom( int32 InHeight ) const
-{
-	return FIntRect( Min.X, FMath::Max(Min.Y,Max.Y-InHeight), Max.X, Max.Y );
-}
-
-
-FORCEINLINE void FIntRect::Clip( const FIntRect& R )
-{
-	Min.X = FMath::Max<int32>(Min.X, R.Min.X);
-	Min.Y = FMath::Max<int32>(Min.Y, R.Min.Y);
-	Max.X = FMath::Min<int32>(Max.X, R.Max.X);
-	Max.Y = FMath::Min<int32>(Max.Y, R.Max.Y);
-
-	// return zero area if not overlapping
-	Max.X = FMath::Max<int32>(Min.X, Max.X);
-	Max.Y = FMath::Max<int32>(Min.Y, Max.Y);
-}
-
-FORCEINLINE void FIntRect::Union( const FIntRect& R )
-{
-	Min.X = FMath::Min<int32>(Min.X, R.Min.X);
-	Min.Y = FMath::Min<int32>(Min.Y, R.Min.Y);
-	Max.X = FMath::Max<int32>(Max.X, R.Max.X);
-	Max.Y = FMath::Max<int32>(Max.Y, R.Max.Y);
-}
-
-FORCEINLINE bool FIntRect::Intersect(const FIntRect& Other) const
-{
-	return Other.Min.X < Max.X && Other.Max.X > Min.X && Other.Min.Y < Max.Y && Other.Max.Y > Min.Y;
-}
-
-FORCEINLINE bool FIntRect::Contains( FIntPoint P ) const
-{
-	return P.X >= Min.X && P.X < Max.X && P.Y >= Min.Y && P.Y < Max.Y;
-}
-
-
-FORCEINLINE FIntRect FIntRect::DivideAndRoundUp( FIntRect lhs, int32 Div )
-{
-	return DivideAndRoundUp(lhs, FIntPoint(Div, Div));
-}
-
-FORCEINLINE FIntRect FIntRect::DivideAndRoundUp( FIntRect lhs, FIntPoint Div )
+template <typename IntType>
+uint32 GetTypeHash(const TIntRect<IntType>& InRect)
 {
-	return FIntRect(lhs.Min / Div, FIntPoint::DivideAndRoundUp(lhs.Max, Div));
-}
-
-FORCEINLINE void FIntRect::GetCenterAndExtents( FIntPoint& OutCenter, FIntPoint& OutExtent ) const
-{
-	OutExtent.X = (Max.X - Min.X) / 2;
-	OutExtent.Y = (Max.Y - Min.Y) / 2;
-
-	OutCenter.X = Min.X + OutExtent.X;
-	OutCenter.Y = Min.Y + OutExtent.Y;
+	return HashCombine(GetTypeHash(InRect.Min), GetTypeHash(InRect.Max));
 }
 
+} //! namespace UE::Math
 
-FORCEINLINE int32 FIntRect::Height() const
-{
-	return (Max.Y - Min.Y);
-}
-
-
-FORCEINLINE void FIntRect::InflateRect( int32 Amount )
-{
-	Min.X -= Amount;
-	Min.Y -= Amount;
-	Max.X += Amount;
-	Max.Y += Amount;
-}
-
-
-FORCEINLINE void FIntRect::Include( FIntPoint Point )
-{
-	Min.X = FMath::Min(Min.X, Point.X);
-	Min.Y = FMath::Min(Min.Y, Point.Y);
-	Max.X = FMath::Max(Max.X, Point.X);
-	Max.Y = FMath::Max(Max.Y, Point.Y);
-}
-
-FORCEINLINE FIntRect FIntRect::Inner( FIntPoint Shrink ) const
-{
-	return FIntRect(Min + Shrink, Max - Shrink);
-}
-
-
-FORCEINLINE int32 FIntRect::Num()
-{
-	return 2;
-}
-
-
-FORCEINLINE FIntRect FIntRect::Right( int32 InWidth ) const
-{
-	return FIntRect( FMath::Max(Min.X, Max.X - InWidth), Min.Y, Max.X, Max.Y );
-}
-
-
-FORCEINLINE FIntPoint FIntRect::Size() const
-{
-	return FIntPoint( Max.X-Min.X, Max.Y-Min.Y );
-}
-
-
-FORCEINLINE FString FIntRect::ToString() const
-{
-	return FString::Printf(TEXT("Min=(%s) Max=(%s)"), *Min.ToString(), *Max.ToString());
-}
-
-
-FORCEINLINE int32 FIntRect::Width() const
-{
-	return Max.X-Min.X;
-}
-
-FORCEINLINE bool FIntRect::IsEmpty() const
-{
-	return Width() == 0 && Height() == 0;
-}
+template <> struct TIsUECoreType<FIntRect>  { enum { Value = true }; };
+template <> struct TIsUECoreType<FUintRect> { enum { Value = true }; };
\ No newline at end of file
diff --git a/Engine/Source/Runtime/Core/Public/Math/IntVector.h b/Engine/Source/Runtime/Core/Public/Math/IntVector.h
--- a/Engine/Source/Runtime/Core/Public/Math/IntVector.h
+++ b/Engine/Source/Runtime/Core/Public/Math/IntVector.h
@@ -8,43 +8,46 @@
 #include "Containers/UnrealString.h"
 #include "Serialization/StructuredArchive.h"
 
+namespace UE::Math
+{
+
 /**
  * Structure for integer vectors in 3-d space.
  */
-struct FIntVector
+template <typename InIntType>
+struct TIntVector3
 {
+	using IntType = InIntType;
+	static_assert(std::is_integral_v<IntType>, "Only an integer types are supported.");
+
 	union
 	{
 		struct
 		{
 			/** Holds the point's x-coordinate. */
-			int32 X;
+			IntType X;
 
 			/** Holds the point's y-coordinate. */
-			int32 Y;
+			IntType Y;
 
 			/**  Holds the point's z-coordinate. */
-			int32 Z;
+			IntType Z;
 		};
 
 		UE_DEPRECATED(all, "For internal use only")
-		int32 XYZ[3];
+		IntType XYZ[3];
 	};
 
-public:
-
 	/** An int point with zeroed values. */
-	CORE_API static const FIntVector ZeroValue;
+	static const TIntVector3 ZeroValue;
 
 	/** An int point with INDEX_NONE values. */
-	CORE_API static const FIntVector NoneValue;
-
-public:
+	static const TIntVector3 NoneValue;
 
 	/**
 	 * Default constructor (no initialization).
 	 */
-	FIntVector();
+	TIntVector3() = default;
 
 	/**
 	 * Creates and initializes a new instance with the specified coordinates.
@@ -53,30 +56,55 @@ public:
 	 * @param InY The y-coordinate.
 	 * @param InZ The z-coordinate.
 	 */
-	FIntVector( int32 InX, int32 InY, int32 InZ );
+	TIntVector3(IntType InX, IntType InY, IntType InZ )
+		: X(InX)
+		, Y(InY)
+		, Z(InZ)
+	{
+	}
 
 	/**
 	 * Constructor
 	 *
 	 * @param InValue replicated to all components
 	 */
-	explicit FIntVector( int32 InValue );
+	explicit TIntVector3(IntType InValue )
+		: X(InValue)
+		, Y(InValue)
+		, Z(InValue)
+	{
+	}
 
 	/**
 	 * Constructor
 	 *
 	 * @param InVector float vector converted to int
 	 */
-	explicit FIntVector( FVector InVector  );
+	explicit TIntVector3( FVector InVector  );
 
 	/**
 	 * Constructor
 	 *
 	 * @param EForceInit Force init enum
 	 */
-	explicit FORCEINLINE FIntVector( EForceInit );
 
-public:
+	explicit TIntVector3( EForceInit )
+		: X(0)
+		, Y(0)
+		, Z(0)
+	{
+	}
+
+	/**
+	 * Converts to another int type. Checks that the cast will succeed.
+	 */
+	template <typename OtherIntType>
+	explicit TIntVector3(TIntVector3<OtherIntType> Other)
+		: X(IntCastChecked<OtherIntType>(Other.X))
+		, Y(IntCastChecked<OtherIntType>(Other.Y))
+		, Z(IntCastChecked<OtherIntType>(Other.Z))
+	{
+	}
 
 	/**
 	 * Gets specific component of a point.
@@ -84,7 +112,12 @@ public:
 	 * @param ComponentIndex Index of point component.
 	 * @return const reference to component.
 	 */
-	const int32& operator()( int32 ComponentIndex ) const;
+	const IntType& operator()(int32 ComponentIndex) const
+	{
+PRAGMA_DISABLE_DEPRECATION_WARNINGS
+		return XYZ[ComponentIndex];
+PRAGMA_ENABLE_DEPRECATION_WARNINGS
+	}
 
 	/**
 	 * Gets specific component of a point.
@@ -92,7 +125,12 @@ public:
 	 * @param ComponentIndex Index of point component.
 	 * @return reference to component.
 	 */
-	int32& operator()( int32 ComponentIndex );
+	IntType& operator()(int32 ComponentIndex)
+	{
+PRAGMA_DISABLE_DEPRECATION_WARNINGS
+		return XYZ[ComponentIndex];
+PRAGMA_ENABLE_DEPRECATION_WARNINGS
+	}
 
 	/**
 	 * Gets specific component of a point.
@@ -100,7 +138,12 @@ public:
 	 * @param ComponentIndex Index of point component.
 	 * @return const reference to component.
 	 */
-	const int32& operator[]( int32 ComponentIndex ) const;
+	const IntType& operator[](int32 ComponentIndex) const
+	{
+PRAGMA_DISABLE_DEPRECATION_WARNINGS
+		return XYZ[ComponentIndex];
+PRAGMA_ENABLE_DEPRECATION_WARNINGS
+	}
 
 	/**
 	 * Gets specific component of a point.
@@ -108,7 +151,12 @@ public:
 	 * @param ComponentIndex Index of point component.
 	 * @return reference to component.
 	 */
-	int32& operator[]( int32 ComponentIndex );
+	IntType& operator[](int32 ComponentIndex)
+	{
+PRAGMA_DISABLE_DEPRECATION_WARNINGS
+		return XYZ[ComponentIndex];
+PRAGMA_ENABLE_DEPRECATION_WARNINGS
+	}
 
 	/**
 	 * Compares points for equality.
@@ -116,7 +164,10 @@ public:
 	 * @param Other The other int point being compared.
 	 * @return true if the points are equal, false otherwise..
 	 */
-	bool operator==( const FIntVector& Other ) const;
+	bool operator==(const TIntVector3& Other) const
+	{
+		return X == Other.X && Y == Other.Y && Z == Other.Z;
+	}
 
 	/**
 	 * Compares points for inequality.
@@ -124,7 +175,10 @@ public:
 	 * @param Other The other int point being compared.
 	 * @return true if the points are not equal, false otherwise..
 	 */
-	bool operator!=( const FIntVector& Other ) const;
+	bool operator!=(const TIntVector3& Other) const
+	{
+		return X != Other.X || Y != Other.Y || Z != Other.Z;
+	}
 
 	/**
 	 * Scales this point.
@@ -132,7 +186,14 @@ public:
 	 * @param Scale What to multiply the point by.
 	 * @return Reference to this point after multiplication.
 	 */
-	FIntVector& operator*=( int32 Scale );
+	TIntVector3& operator*=(IntType Scale)
+	{
+		X *= Scale;
+		Y *= Scale;
+		Z *= Scale;
+
+		return *this;
+	}
 
 	/**
 	 * Divides this point.
@@ -140,7 +201,14 @@ public:
 	 * @param Divisor What to divide the point by.
 	 * @return Reference to this point after division.
 	 */
-	FIntVector& operator/=( int32 Divisor );
+	TIntVector3& operator/=(IntType Divisor)
+	{
+		X /= Divisor;
+		Y /= Divisor;
+		Z /= Divisor;
+
+		return *this;
+	}
 
 	/**
 	 * Adds to this point.
@@ -148,7 +216,14 @@ public:
 	 * @param Other The point to add to this point.
 	 * @return Reference to this point after addition.
 	 */
-	FIntVector& operator+=( const FIntVector& Other );
+	TIntVector3& operator+=(const TIntVector3& Other)
+	{
+		X += Other.X;
+		Y += Other.Y;
+		Z += Other.Z;
+
+		return *this;
+	}
 
 	/**
 	 * Subtracts from this point.
@@ -156,15 +231,14 @@ public:
 	 * @param Other The point to subtract from this point.
 	 * @return Reference to this point after subtraction.
 	 */
-	FIntVector& operator-=( const FIntVector& Other );
+	TIntVector3& operator-=(const TIntVector3& Other)
+	{
+		X -= Other.X;
+		Y -= Other.Y;
+		Z -= Other.Z;
 
-	/**
-	 * Assigns another point to this one.
-	 *
-	 * @param Other The point to assign this point from.
-	 * @return Reference to this point after assignment.
-	 */
-	FIntVector& operator=( const FIntVector& Other );
+		return *this;
+	}
 
 	/**
 	 * Gets the result of scaling on this point.
@@ -172,7 +246,10 @@ public:
 	 * @param Scale What to multiply the point by.
 	 * @return A new scaled int point.
 	 */
-	FIntVector operator*( int32 Scale ) const;
+	TIntVector3 operator*(IntType Scale) const
+	{
+		return TIntVector3(*this) *= Scale;
+	}
 
 	/**
 	 * Gets the result of division on this point.
@@ -180,7 +257,10 @@ public:
 	 * @param Divisor What to divide the point by.
 	 * @return A new divided int point.
 	 */
-	FIntVector operator/( int32 Divisor ) const;
+	TIntVector3 operator/(IntType Divisor) const
+	{
+		return TIntVector3(*this) /= Divisor;
+	}
 
 	/**
 	 * Gets the result of addition on this point.
@@ -188,7 +268,10 @@ public:
 	 * @param Other The other point to add to this.
 	 * @return A new combined int point.
 	 */
-	FIntVector operator+( const FIntVector& Other ) const;
+	TIntVector3 operator+(const TIntVector3& Other) const
+	{
+		return TIntVector3(*this) += Other;
+	}
 
 	/**
 	 * Gets the result of subtraction from this point.
@@ -196,7 +279,10 @@ public:
 	 * @param Other The other point to subtract from this.
 	 * @return A new subtracted int point.
 	 */
-	FIntVector operator-( const FIntVector& Other ) const;
+	TIntVector3 operator-(const TIntVector3& Other) const
+	{
+		return TIntVector3(*this) -= Other;
+	}
 
 	/**
 	 * Shifts all components to the right.
@@ -204,7 +290,10 @@ public:
 	 * @param Shift The number of bits to shift.
 	 * @return A new shifted int point.
 	 */
-	FIntVector operator>>(int32 Shift) const;
+	TIntVector3 operator>>(IntType Shift) const
+	{
+		return TIntVector3(X >> Shift, Y >> Shift, Z >> Shift);
+	}
 
 	/**
 	 * Shifts all components to the left.
@@ -212,7 +301,10 @@ public:
 	 * @param Shift The number of bits to shift.
 	 * @return A new shifted int point.
 	 */
-	FIntVector operator<<(int32 Shift) const;
+	TIntVector3 operator<<(IntType Shift) const
+	{
+		return TIntVector3(X << Shift, Y << Shift, Z << Shift);
+	}
 
 	/**
 	 * Component-wise AND.
@@ -220,7 +312,10 @@ public:
 	 * @param Value Number to AND with the each component.
 	 * @return A new shifted int point.
 	 */
-	FIntVector operator&(int32 Val) const;
+	TIntVector3 operator&(IntType Value) const
+	{
+		return TIntVector3(X & Value, Y & Value, Z & Value);
+	}
 
 	/**
 	 * Component-wise OR.
@@ -228,7 +323,10 @@ public:
 	 * @param Value Number to OR with the each component.
 	 * @return A new shifted int point.
 	 */
-	FIntVector operator|(int32 Value) const;
+	TIntVector3 operator|(IntType Value) const
+	{
+		return TIntVector3(X | Value, Y | Value, Z | Value);
+	}
 
 	/**
 	 * Component-wise XOR.
@@ -236,45 +334,62 @@ public:
 	 * @param Value Number to XOR with the each component.
 	 * @return A new shifted int point.
 	 */
-	FIntVector operator^(int32 Value) const;
+	TIntVector3 operator^(IntType Value) const
+	{
+		return TIntVector3(X ^ Value, Y ^ Value, Z ^ Value);
+	}
 
 	/**
 	 * Is vector equal to zero.
 	 * @return is zero
 	*/
-	bool IsZero() const;
-
-public:
+	bool IsZero() const
+	{
+		return *this == ZeroValue;
+	}
 
 	/**
 	 * Gets the maximum value in the point.
 	 *
 	 * @return The maximum value in the point.
 	 */
-	int32 GetMax() const;
+	IntType GetMax() const
+	{
+		return FMath::Max(FMath::Max(X, Y), Z);
+	}
 
 	/**
 	 * Gets the minimum value in the point.
 	 *
 	 * @return The minimum value in the point.
 	 */
-	int32 GetMin() const;
+	IntType GetMin() const
+	{
+		return FMath::Min(FMath::Min(X, Y), Z);
+	}
 
 	/**
 	 * Gets the distance of this point from (0,0,0).
 	 *
 	 * @return The distance of this point from (0,0,0).
 	 */
-	int32 Size() const;
+	IntType Size() const
+	{
+		int64 LocalX64 = (int64)X;
+		int64 LocalY64 = (int64)Y;
+		int64 LocalZ64 = (int64)Z;
+		return IntType(FMath::Sqrt(double(LocalX64 * LocalX64 + LocalY64 * LocalY64 + LocalZ64 * LocalZ64)));
+	}
 
 	/**
 	 * Get a textual representation of this vector.
 	 *
 	 * @return A string describing the vector.
 	 */
-	FString ToString() const;
-
-public:
+	FString ToString() const
+	{
+		return FString::Printf(TEXT("X=%s Y=%s Z=%s"), *LexToString(X), *LexToString(Y), *LexToString(Z));
+	}
 
 	/**
 	 * Divide an int point and round up the result.
@@ -283,17 +398,25 @@ public:
 	 * @param Divisor What to divide the int point by.
 	 * @return A new divided int point.
 	 */
-	static FIntVector DivideAndRoundUp( FIntVector lhs, int32 Divisor );
-	static FIntVector DivideAndRoundUp( FIntVector lhs, FIntVector Divisor );
+	static TIntVector3 DivideAndRoundUp(TIntVector3 lhs, IntType Divisor)
+	{
+		return TIntVector3(FMath::DivideAndRoundUp(lhs.X, Divisor), FMath::DivideAndRoundUp(lhs.Y, Divisor), FMath::DivideAndRoundUp(lhs.Z, Divisor));
+	}
+
+	static TIntVector3 DivideAndRoundUp(TIntVector3 lhs, TIntVector3 Divisor)
+	{
+		return TIntVector3(FMath::DivideAndRoundUp(lhs.X, Divisor.X), FMath::DivideAndRoundUp(lhs.Y, Divisor.Y), FMath::DivideAndRoundUp(lhs.Z, Divisor.Z));
+	}
 
 	/**
 	 * Gets the number of components a point has.
 	 *
 	 * @return Number of components point has.
 	 */
-	static int32 Num();
-
-public:
+	static int32 Num()
+	{
+		return 3;
+	}
 
 	/**
 	 * Serializes the Rectangle.
@@ -302,12 +425,12 @@ public:
 	 * @param Vector The vector to serialize.
 	 * @return Reference to the Archive after serialization.
 	 */
-	friend FArchive& operator<<( FArchive& Ar, FIntVector& Vector )
+	friend FArchive& operator<<( FArchive& Ar, TIntVector3& Vector )
 	{
 		return Ar << Vector.X << Vector.Y << Vector.Z;
 	}
 
-	friend void operator<<(FStructuredArchive::FSlot Slot, FIntVector& Vector)
+	friend void operator<<(FStructuredArchive::FSlot Slot, TIntVector3& Vector)
 	{
 		FStructuredArchive::FRecord Record = Slot.EnterRecord();
 		Record << SA_VALUE(TEXT("X"), Vector.X);
@@ -322,370 +445,156 @@ public:
 	}
 };
 
-
-/* FIntVector inline functions
- *****************************************************************************/
-
-FORCEINLINE FIntVector::FIntVector()
-{ }
-
-
-FORCEINLINE FIntVector::FIntVector( int32 InX, int32 InY, int32 InZ )
-	: X(InX)
-	, Y(InY)
-	, Z(InZ)
-{ }
-
-
-FORCEINLINE FIntVector::FIntVector( int32 InValue )
-	: X(InValue)
-	, Y(InValue)
-	, Z(InValue)
-{ }
-
-
-FORCEINLINE FIntVector::FIntVector( EForceInit )
-	: X(0)
-	, Y(0)
-	, Z(0)
-{ }
-
-
-FORCEINLINE const int32& FIntVector::operator()( int32 ComponentIndex ) const
-{
-PRAGMA_DISABLE_DEPRECATION_WARNINGS
-	return XYZ[ComponentIndex];
-PRAGMA_ENABLE_DEPRECATION_WARNINGS
-}
-
-
-FORCEINLINE int32& FIntVector::operator()( int32 ComponentIndex )
+template <>
+inline FString TIntVector3<int64>::ToString() const
 {
-PRAGMA_DISABLE_DEPRECATION_WARNINGS
-	return XYZ[ComponentIndex];
-PRAGMA_ENABLE_DEPRECATION_WARNINGS
-}
-
-
-FORCEINLINE const int32& FIntVector::operator[]( int32 ComponentIndex ) const
-{
-PRAGMA_DISABLE_DEPRECATION_WARNINGS
-	return XYZ[ComponentIndex];
-PRAGMA_ENABLE_DEPRECATION_WARNINGS
-}
-
-
-FORCEINLINE int32& FIntVector::operator[]( int32 ComponentIndex )
-{
-PRAGMA_DISABLE_DEPRECATION_WARNINGS
-	return XYZ[ComponentIndex];
-PRAGMA_ENABLE_DEPRECATION_WARNINGS
-}
-
-FORCEINLINE bool FIntVector::operator==( const FIntVector& Other ) const
-{
-	return X==Other.X && Y==Other.Y && Z==Other.Z;
-}
-
-
-FORCEINLINE bool FIntVector::operator!=( const FIntVector& Other ) const
-{
-	return X!=Other.X || Y!=Other.Y || Z!=Other.Z;
-}
-
-
-FORCEINLINE FIntVector& FIntVector::operator*=( int32 Scale )
-{
-	X *= Scale;
-	Y *= Scale;
-	Z *= Scale;
-
-	return *this;
-}
-
-
-FORCEINLINE FIntVector& FIntVector::operator/=( int32 Divisor )
-{
-	X /= Divisor;
-	Y /= Divisor;
-	Z /= Divisor;
-
-	return *this;
-}
-
-
-FORCEINLINE FIntVector& FIntVector::operator+=( const FIntVector& Other )
-{
-	X += Other.X;
-	Y += Other.Y;
-	Z += Other.Z;
-
-	return *this;
-}
-
-
-FORCEINLINE FIntVector& FIntVector::operator-=( const FIntVector& Other )
-{
-	X -= Other.X;
-	Y -= Other.Y;
-	Z -= Other.Z;
-
-	return *this;
-}
-
-
-FORCEINLINE FIntVector& FIntVector::operator=( const FIntVector& Other )
-{
-	X = Other.X;
-	Y = Other.Y;
-	Z = Other.Z;
-
-	return *this;
-}
-
-
-FORCEINLINE FIntVector FIntVector::operator*( int32 Scale ) const
-{
-	return FIntVector(*this) *= Scale;
-}
-
-
-FORCEINLINE FIntVector FIntVector::operator/( int32 Divisor ) const
-{
-	return FIntVector(*this) /= Divisor;
-}
-
-
-FORCEINLINE FIntVector FIntVector::operator+( const FIntVector& Other ) const
-{
-	return FIntVector(*this) += Other;
+	return FString::Printf(TEXT("X=%lld Y=%lld Z=%lld"), X, Y, Z);
 }
 
-FORCEINLINE FIntVector FIntVector::operator-( const FIntVector& Other ) const
+template <>
+inline FString TIntVector3<int32>::ToString() const
 {
-	return FIntVector(*this) -= Other;
-}
-
-FORCEINLINE FIntVector FIntVector::operator>>(int32 Shift) const
-{
-	return FIntVector(X >> Shift, Y >> Shift, Z >> Shift);
-}
-
-FORCEINLINE FIntVector FIntVector::operator<<(int32 Shift) const
-{
-	return FIntVector(X << Shift, Y << Shift, Z << Shift);
-}
-
-FORCEINLINE FIntVector FIntVector::operator&(int32 Value) const
-{
-	return FIntVector(X & Value, Y & Value, Z & Value);
-}
-
-FORCEINLINE FIntVector FIntVector::operator|(int32 Value) const
-{
-	return FIntVector(X | Value, Y | Value, Z | Value);
-}
-
-FORCEINLINE FIntVector FIntVector::operator^(int32 Value) const
-{
-	return FIntVector(X ^ Value, Y ^ Value, Z ^ Value);
+	return FString::Printf(TEXT("X=%d Y=%d Z=%d"), X, Y, Z);
 }
 
-FORCEINLINE FIntVector FIntVector::DivideAndRoundUp( FIntVector lhs, int32 Divisor )
+template <>
+inline FString TIntVector3<int16>::ToString() const
 {
-	return FIntVector(FMath::DivideAndRoundUp(lhs.X, Divisor), FMath::DivideAndRoundUp(lhs.Y, Divisor), FMath::DivideAndRoundUp(lhs.Z, Divisor));
+	return FString::Printf(TEXT("X=%d Y=%d Z=%d"), X, Y, Z);
 }
 
-FORCEINLINE FIntVector FIntVector::DivideAndRoundUp(FIntVector lhs, FIntVector Divisor)
+template <>
+inline FString TIntVector3<int8>::ToString() const
 {
-	return FIntVector(FMath::DivideAndRoundUp(lhs.X, Divisor.X), FMath::DivideAndRoundUp(lhs.Y, Divisor.Y), FMath::DivideAndRoundUp(lhs.Z, Divisor.Z));
+	return FString::Printf(TEXT("X=%d Y=%d Z=%d"), X, Y, Z);
 }
 
-
-FORCEINLINE int32 FIntVector::GetMax() const
+template <>
+inline FString TIntVector3<uint64>::ToString() const
 {
-	return FMath::Max(FMath::Max(X, Y), Z);
+	return FString::Printf(TEXT("X=%llu Y=%llu Z=%llu"), X, Y, Z);
 }
 
-
-FORCEINLINE int32 FIntVector::GetMin() const
+template <>
+inline FString TIntVector3<uint32>::ToString() const
 {
-	return FMath::Min(FMath::Min(X, Y), Z);
+	return FString::Printf(TEXT("X=%u Y=%u Z=%u"), X, Y, Z);
 }
 
-
-FORCEINLINE int32 FIntVector::Num()
+template <>
+inline FString TIntVector3<uint16>::ToString() const
 {
-	return 3;
+	return FString::Printf(TEXT("X=%u Y=%u Z=%u"), X, Y, Z);
 }
 
-
-FORCEINLINE int32 FIntVector::Size() const
+template <>
+inline FString TIntVector3<uint8>::ToString() const
 {
-	int64 LocalX64 = (int64)X;
-	int64 LocalY64 = (int64)Y;
-	int64 LocalZ64 = (int64)Z;
-	return int32(FMath::Sqrt(float(LocalX64 * LocalX64 + LocalY64 * LocalY64 + LocalZ64 * LocalZ64)));
+	return FString::Printf(TEXT("X=%u Y=%u Z=%u"), X, Y, Z);
 }
 
-FORCEINLINE bool FIntVector::IsZero() const
-{
-	return *this == ZeroValue;
-}
+template <typename IntType>
+const TIntVector3<IntType> TIntVector3<IntType>::ZeroValue(0, 0, 0);
 
+template <typename IntType>
+const TIntVector3<IntType> TIntVector3<IntType>::NoneValue(INDEX_NONE, INDEX_NONE, INDEX_NONE);
 
-FORCEINLINE FString FIntVector::ToString() const
-{
-	return FString::Printf(TEXT("X=%d Y=%d Z=%d"), X, Y, Z);
-}
+///////////////////////////////////////////////////////////////////////////////////////////////////
 
-FORCEINLINE uint32 GetTypeHash(const FIntVector& Vector)
+template <typename InIntType>
+struct TIntVector2
 {
-	return FCrc::MemCrc_DEPRECATED(&Vector,sizeof(FIntVector));
-}
+	using IntType = InIntType;
+	static_assert(std::is_integral_v<IntType>, "Only an integer types are supported.");
 
-struct FIntVector2
-{
 	union
 	{
 		struct
 		{
-			int32 X, Y;
+			IntType X, Y;
 		};
 
 		UE_DEPRECATED(all, "For internal use only")
-		int32 XY[2];
+		IntType XY[2];
 	};
 
-	FORCEINLINE FIntVector2()
-	{
-	}
+	TIntVector2() = default;
 
-	FORCEINLINE FIntVector2(int32 InX, int32 InY)
+	TIntVector2(IntType InX, IntType InY)
 		: X(InX)
 		, Y(InY)
 	{
 	}
 
-	FORCEINLINE explicit FIntVector2(int32 InValue)
+	explicit TIntVector2(IntType InValue)
 		: X(InValue)
 		, Y(InValue)
 	{
 	}
 
-	FORCEINLINE FIntVector2(EForceInit)
+	TIntVector2(EForceInit)
 		: X(0)
 		, Y(0)
 	{
 	}
 
-	FORCEINLINE const int32& operator[](int32 ComponentIndex) const
-	{
-PRAGMA_DISABLE_DEPRECATION_WARNINGS
-		return XY[ComponentIndex];
-PRAGMA_ENABLE_DEPRECATION_WARNINGS
-	}
-
-
-	FORCEINLINE int32& operator[](int32 ComponentIndex)
-	{
-PRAGMA_DISABLE_DEPRECATION_WARNINGS
-		return XY[ComponentIndex];
-PRAGMA_ENABLE_DEPRECATION_WARNINGS
-	}
-
-	FORCEINLINE bool operator==(const FIntVector2& Other) const
-	{
-		return X==Other.X && Y==Other.Y;
-	}
-
-
-	FORCEINLINE bool operator!=(const FIntVector2& Other) const
-	{
-		return X!=Other.X || Y!=Other.Y;
-	}
-};
-
-struct FUintVector2
-{
-	union
-	{
-		struct
-		{
-			uint32 X, Y;
-		};
-
-		UE_DEPRECATED(all, "For internal use only")
-		uint32 XY[2];
-	};
-
-	FORCEINLINE FUintVector2()
-	{
-	}
-
-	FORCEINLINE FUintVector2(uint32 InX, uint32 InY)
-		: X(InX)
-		, Y(InY)
-	{
-	}
-
-	FORCEINLINE explicit FUintVector2(uint32 InValue)
-		: X(InValue)
-		, Y(InValue)
-	{
-	}
-
-	FORCEINLINE FUintVector2(EForceInit)
-		: X(0)
-		, Y(0)
+	/**
+	 * Converts to another int type. Checks that the cast will succeed.
+	 */
+	template <typename OtherIntType>
+	explicit TIntVector2(TIntVector2<OtherIntType> Other)
+		: X(IntCastChecked<OtherIntType>(Other.X))
+		, Y(IntCastChecked<OtherIntType>(Other.Y))
 	{
 	}
 
-	FORCEINLINE const uint32& operator[](int32 ComponentIndex) const
+	const IntType& operator[](int32 ComponentIndex) const
 	{
 PRAGMA_DISABLE_DEPRECATION_WARNINGS
 		return XY[ComponentIndex];
 PRAGMA_ENABLE_DEPRECATION_WARNINGS
 	}
 
-
-	FORCEINLINE uint32& operator[](int32 ComponentIndex)
+	IntType& operator[](int32 ComponentIndex)
 	{
 PRAGMA_DISABLE_DEPRECATION_WARNINGS
 		return XY[ComponentIndex];
 PRAGMA_ENABLE_DEPRECATION_WARNINGS
 	}
 
-	FORCEINLINE bool operator==(const FUintVector2& Other) const
+	bool operator==(const TIntVector2& Other) const
 	{
 		return X==Other.X && Y==Other.Y;
 	}
 
-
-	FORCEINLINE bool operator!=(const FUintVector2& Other) const
+	bool operator!=(const TIntVector2& Other) const
 	{
 		return X!=Other.X || Y!=Other.Y;
 	}
 };
 
-struct FIntVector4
+///////////////////////////////////////////////////////////////////////////////////////////////////
+
+template <typename InIntType>
+struct TIntVector4
 {
+	using IntType = InIntType;
+	static_assert(std::is_integral_v<IntType>, "Only an integer types are supported.");
+
 	union
 	{
 		struct
 		{
-			int32 X, Y, Z, W;
+			IntType X, Y, Z, W;
 		};
 
 		UE_DEPRECATED(all, "For internal use only")
-		int32 XYZW[4];
+		IntType XYZW[4];
 	};
 
-	FORCEINLINE FIntVector4()
-	{
-	}
+	TIntVector4() = default;
 
-	FORCEINLINE FIntVector4(int32 InX, int32 InY, int32 InZ, int32 InW)
+	TIntVector4(IntType InX, IntType InY, IntType InZ, IntType InW)
 		: X(InX)
 		, Y(InY)
 		, Z(InZ)
@@ -693,7 +602,7 @@ struct FIntVector4
 	{
 	}
 
-	FORCEINLINE explicit FIntVector4(int32 InValue)
+	explicit TIntVector4(IntType InValue)
 		: X(InValue)
 		, Y(InValue)
 		, Z(InValue)
@@ -701,7 +610,7 @@ struct FIntVector4
 	{
 	}
 
-	FORCEINLINE explicit FIntVector4(const FIntVector& InValue, int32 InW = 0)
+	explicit TIntVector4(const TIntVector3<IntType>& InValue, IntType InW = 0)
 		: X(InValue.X)
 		, Y(InValue.Y)
 		, Z(InValue.Z)
@@ -709,7 +618,7 @@ struct FIntVector4
 	{
 	}
 
-	FORCEINLINE FIntVector4(EForceInit)
+	TIntVector4(EForceInit)
 		: X(0)
 		, Y(0)
 		, Z(0)
@@ -717,13 +626,30 @@ struct FIntVector4
 	{
 	}
 
+	/**
+	 * Converts to another int type. Checks that the cast will succeed.
+	 */
+	template <typename OtherIntType>
+	explicit TIntVector4(TIntVector4<OtherIntType> Other)
+		: X(IntCastChecked<OtherIntType>(Other.X))
+		, Y(IntCastChecked<OtherIntType>(Other.Y))
+		, Z(IntCastChecked<OtherIntType>(Other.Z))
+		, W(IntCastChecked<OtherIntType>(Other.W))
+	{
+	}
+
 	/**
 	 * Gets specific component of a point.
 	 *
 	 * @param ComponentIndex Index of point component.
 	 * @return const reference to component.
 	 */
-	const int32& operator()(int32 ComponentIndex) const;
+	const IntType& operator()(int32 ComponentIndex) const
+	{
+PRAGMA_DISABLE_DEPRECATION_WARNINGS
+		return XYZW[ComponentIndex];
+PRAGMA_ENABLE_DEPRECATION_WARNINGS
+	}
 
 	/**
 	 * Gets specific component of a point.
@@ -731,7 +657,12 @@ struct FIntVector4
 	 * @param ComponentIndex Index of point component.
 	 * @return reference to component.
 	 */
-	int32& operator()(int32 ComponentIndex);
+	IntType& operator()(int32 ComponentIndex)
+	{
+PRAGMA_DISABLE_DEPRECATION_WARNINGS
+		return XYZW[ComponentIndex];
+PRAGMA_ENABLE_DEPRECATION_WARNINGS
+	}
 
 	/**
 	 * Gets specific component of a point.
@@ -739,7 +670,12 @@ struct FIntVector4
 	 * @param ComponentIndex Index of point component.
 	 * @return const reference to component.
 	 */
-	const int32& operator[](int32 ComponentIndex) const;
+	const IntType& operator[](int32 ComponentIndex) const
+	{
+PRAGMA_DISABLE_DEPRECATION_WARNINGS
+		return XYZW[ComponentIndex];
+PRAGMA_ENABLE_DEPRECATION_WARNINGS
+	}
 
 	/**
 	 * Gets specific component of a point.
@@ -747,7 +683,12 @@ struct FIntVector4
 	 * @param ComponentIndex Index of point component.
 	 * @return reference to component.
 	 */
-	int32& operator[](int32 ComponentIndex);
+	IntType& operator[](int32 ComponentIndex)
+	{
+PRAGMA_DISABLE_DEPRECATION_WARNINGS
+		return XYZW[ComponentIndex];
+PRAGMA_ENABLE_DEPRECATION_WARNINGS
+	}
 
 	/**
 	 * Compares points for equality.
@@ -755,7 +696,10 @@ struct FIntVector4
 	 * @param Other The other int point being compared.
 	 * @return true if the points are equal, false otherwise..
 	 */
-	bool operator==(const FIntVector4& Other) const;
+	bool operator==(const TIntVector4& Other) const
+	{
+		return X == Other.X && Y == Other.Y && Z == Other.Z && W == Other.W;
+	}
 
 	/**
 	 * Compares points for inequality.
@@ -763,7 +707,10 @@ struct FIntVector4
 	 * @param Other The other int point being compared.
 	 * @return true if the points are not equal, false otherwise..
 	 */
-	bool operator!=(const FIntVector4& Other) const;
+	bool operator!=(const TIntVector4& Other) const
+	{
+		return X != Other.X || Y != Other.Y || Z != Other.Z || W != Other.W;
+	}
 
 	/**
 	 * Scales this point.
@@ -771,7 +718,15 @@ struct FIntVector4
 	 * @param Scale What to multiply the point by.
 	 * @return Reference to this point after multiplication.
 	 */
-	FIntVector4& operator*=(int32 Scale);
+	TIntVector4& operator*=(IntType Scale)
+	{
+		X *= Scale;
+		Y *= Scale;
+		Z *= Scale;
+		W *= Scale;
+
+		return *this;
+	}
 
 	/**
 	 * Divides this point.
@@ -779,7 +734,15 @@ struct FIntVector4
 	 * @param Divisor What to divide the point by.
 	 * @return Reference to this point after division.
 	 */
-	FIntVector4& operator/=(int32 Divisor);
+	TIntVector4& operator/=(IntType Divisor)
+	{
+		X /= Divisor;
+		Y /= Divisor;
+		Z /= Divisor;
+		W /= Divisor;
+
+		return *this;
+	}
 
 	/**
 	 * Adds to this point.
@@ -787,7 +750,15 @@ struct FIntVector4
 	 * @param Other The point to add to this point.
 	 * @return Reference to this point after addition.
 	 */
-	FIntVector4& operator+=(const FIntVector4& Other);
+	TIntVector4& operator+=(const TIntVector4& Other)
+	{
+		X += Other.X;
+		Y += Other.Y;
+		Z += Other.Z;
+		W += Other.W;
+
+		return *this;
+	}
 
 	/**
 	 * Subtracts from this point.
@@ -795,7 +766,15 @@ struct FIntVector4
 	 * @param Other The point to subtract from this point.
 	 * @return Reference to this point after subtraction.
 	 */
-	FIntVector4& operator-=(const FIntVector4& Other);
+	TIntVector4& operator-=(const TIntVector4& Other)
+	{
+		X -= Other.X;
+		Y -= Other.Y;
+		Z -= Other.Z;
+		W -= Other.W;
+
+		return *this;
+	}
 
 	/**
 	 * Assigns another point to this one.
@@ -803,7 +782,15 @@ struct FIntVector4
 	 * @param Other The point to assign this point from.
 	 * @return Reference to this point after assignment.
 	 */
-	FIntVector4& operator=(const FIntVector4& Other);
+	TIntVector4& operator=(const TIntVector4& Other)
+	{
+		X = Other.X;
+		Y = Other.Y;
+		Z = Other.Z;
+		W = Other.W;
+
+		return *this;
+	}
 
 	/**
 	 * Gets the result of scaling on this point.
@@ -811,7 +798,10 @@ struct FIntVector4
 	 * @param Scale What to multiply the point by.
 	 * @return A new scaled int point.
 	 */
-	FIntVector4 operator*(int32 Scale) const;
+	TIntVector4 operator*(IntType Scale) const
+	{
+		return TIntVector4(*this) *= Scale;
+	}
 
 	/**
 	 * Gets the result of division on this point.
@@ -819,7 +809,10 @@ struct FIntVector4
 	 * @param Divisor What to divide the point by.
 	 * @return A new divided int point.
 	 */
-	FIntVector4 operator/(int32 Divisor) const;
+	TIntVector4 operator/(IntType Divisor) const
+	{
+		return TIntVector4(*this) /= Divisor;
+	}
 
 	/**
 	 * Gets the result of addition on this point.
@@ -827,7 +820,10 @@ struct FIntVector4
 	 * @param Other The other point to add to this.
 	 * @return A new combined int point.
 	 */
-	FIntVector4 operator+(const FIntVector4& Other) const;
+	TIntVector4 operator+(const TIntVector4& Other) const
+	{
+		return TIntVector4(*this) += Other;
+	}
 
 	/**
 	 * Gets the result of subtraction from this point.
@@ -835,7 +831,10 @@ struct FIntVector4
 	 * @param Other The other point to subtract from this.
 	 * @return A new subtracted int point.
 	 */
-	FIntVector4 operator-(const FIntVector4& Other) const;
+	TIntVector4 operator-(const TIntVector4& Other) const
+	{
+		return TIntVector4(*this) -= Other;
+	}
 
 	/**
 	 * Shifts all components to the right.
@@ -843,7 +842,10 @@ struct FIntVector4
 	 * @param Shift The number of bits to shift.
 	 * @return A new shifted int point.
 	 */
-	FIntVector4 operator>>(int32 Shift) const;
+	TIntVector4 operator>>(IntType Shift) const
+	{
+		return TIntVector4(X >> Shift, Y >> Shift, Z >> Shift, W >> Shift);
+	}
 
 	/**
 	 * Shifts all components to the left.
@@ -851,7 +853,10 @@ struct FIntVector4
 	 * @param Shift The number of bits to shift.
 	 * @return A new shifted int point.
 	 */
-	FIntVector4 operator<<(int32 Shift) const;
+	TIntVector4 operator<<(IntType Shift) const
+	{
+		return TIntVector4(X << Shift, Y << Shift, Z << Shift, W << Shift);
+	}
 
 	/**
 	 * Component-wise AND.
@@ -859,7 +864,10 @@ struct FIntVector4
 	 * @param Value Number to AND with the each component.
 	 * @return A new shifted int point.
 	 */
-	FIntVector4 operator&(int32 Val) const;
+	TIntVector4 operator&(IntType Value) const
+	{
+		return TIntVector4(X & Value, Y & Value, Z & Value, W & Value);
+	}
 
 	/**
 	 * Component-wise OR.
@@ -867,7 +875,10 @@ struct FIntVector4
 	 * @param Value Number to OR with the each component.
 	 * @return A new shifted int point.
 	 */
-	FIntVector4 operator|(int32 Value) const;
+	TIntVector4 operator|(IntType Value) const
+	{
+		return TIntVector4(X | Value, Y | Value, Z | Value, W | Value);
+	}
 
 	/**
 	 * Component-wise XOR.
@@ -875,223 +886,37 @@ struct FIntVector4
 	 * @param Value Number to XOR with the each component.
 	 * @return A new shifted int point.
 	 */
-	FIntVector4 operator^(int32 Value) const;
+	TIntVector4 operator^(IntType Value) const
+	{
+		return TIntVector4(X ^ Value, Y ^ Value, Z ^ Value, W ^ Value);
+	}
 };
 
-FORCEINLINE const int32& FIntVector4::operator()(int32 ComponentIndex) const
-{
-PRAGMA_DISABLE_DEPRECATION_WARNINGS
-	return XYZW[ComponentIndex];
-PRAGMA_ENABLE_DEPRECATION_WARNINGS
-}
-
-
-FORCEINLINE int32& FIntVector4::operator()(int32 ComponentIndex)
-{
-PRAGMA_DISABLE_DEPRECATION_WARNINGS
-	return XYZW[ComponentIndex];
-PRAGMA_ENABLE_DEPRECATION_WARNINGS
-}
-
-
-FORCEINLINE const int32& FIntVector4::operator[](int32 ComponentIndex) const
-{
-PRAGMA_DISABLE_DEPRECATION_WARNINGS
-	return XYZW[ComponentIndex];
-PRAGMA_ENABLE_DEPRECATION_WARNINGS
-}
-
-
-FORCEINLINE int32& FIntVector4::operator[](int32 ComponentIndex)
-{
-PRAGMA_DISABLE_DEPRECATION_WARNINGS
-	return XYZW[ComponentIndex];
-PRAGMA_ENABLE_DEPRECATION_WARNINGS
-}
-
-FORCEINLINE bool FIntVector4::operator==(const FIntVector4& Other) const
-{
-	return X == Other.X && Y == Other.Y && Z == Other.Z && W == Other.W;
-}
-
-
-FORCEINLINE bool FIntVector4::operator!=(const FIntVector4& Other) const
-{
-	return X != Other.X || Y != Other.Y || Z != Other.Z || W != Other.W;
-}
-
-
-FORCEINLINE FIntVector4& FIntVector4::operator*=(int32 Scale)
-{
-	X *= Scale;
-	Y *= Scale;
-	Z *= Scale;
-	W *= Scale;
-
-	return *this;
-}
-
-
-FORCEINLINE FIntVector4& FIntVector4::operator/=(int32 Divisor)
-{
-	X /= Divisor;
-	Y /= Divisor;
-	Z /= Divisor;
-	W /= Divisor;
-
-	return *this;
-}
-
-
-FORCEINLINE FIntVector4& FIntVector4::operator+=(const FIntVector4& Other)
-{
-	X += Other.X;
-	Y += Other.Y;
-	Z += Other.Z;
-	W += Other.W;
-
-	return *this;
-}
-
-
-FORCEINLINE FIntVector4& FIntVector4::operator-=(const FIntVector4& Other)
-{
-	X -= Other.X;
-	Y -= Other.Y;
-	Z -= Other.Z;
-	W -= Other.W;
-
-	return *this;
-}
-
-
-FORCEINLINE FIntVector4& FIntVector4::operator=(const FIntVector4& Other)
-{
-	X = Other.X;
-	Y = Other.Y;
-	Z = Other.Z;
-	W = Other.W;
-
-	return *this;
-}
-
-
-FORCEINLINE FIntVector4 FIntVector4::operator*(int32 Scale) const
-{
-	return FIntVector4(*this) *= Scale;
-}
-
-
-FORCEINLINE FIntVector4 FIntVector4::operator/(int32 Divisor) const
-{
-	return FIntVector4(*this) /= Divisor;
-}
-
-
-FORCEINLINE FIntVector4 FIntVector4::operator+(const FIntVector4& Other) const
-{
-	return FIntVector4(*this) += Other;
-}
-
-FORCEINLINE FIntVector4 FIntVector4::operator-(const FIntVector4& Other) const
-{
-	return FIntVector4(*this) -= Other;
-}
-
-FORCEINLINE FIntVector4 FIntVector4::operator>>(int32 Shift) const
-{
-	return FIntVector4(X >> Shift, Y >> Shift, Z >> Shift, W >> Shift);
-}
-
-FORCEINLINE FIntVector4 FIntVector4::operator<<(int32 Shift) const
-{
-	return FIntVector4(X << Shift, Y << Shift, Z << Shift, W << Shift);
-}
-
-FORCEINLINE FIntVector4 FIntVector4::operator&(int32 Value) const
-{
-	return FIntVector4(X & Value, Y & Value, Z & Value, W & Value);
-}
-
-FORCEINLINE FIntVector4 FIntVector4::operator|(int32 Value) const
-{
-	return FIntVector4(X | Value, Y | Value, Z | Value, W | Value);
-}
-
-FORCEINLINE FIntVector4 FIntVector4::operator^(int32 Value) const
+/**
+ * Creates a hash value from an FVector.
+ *
+ * @param Vector the vector to create a hash value for
+ * @return The hash value from the components
+ */
+template<typename T>
+uint32 GetTypeHash(const TIntVector3<T>& Vector)
 {
-	return FIntVector4(X ^ Value, Y ^ Value, Z ^ Value, W ^ Value);
+	// Note: this assumes there's no padding in Vector that could contain uncompared data.
+	return FCrc::MemCrc_DEPRECATED(&Vector, sizeof(Vector));
 }
 
+} //! namespace UE::Math
 
-struct FUintVector4
-{
-	union
-	{
-		struct
-		{
-			uint32 X, Y, Z, W;
-		};
-
-		UE_DEPRECATED(all, "For internal use only")
-		uint32 XYZW[4];
-	};
-
-	FORCEINLINE FUintVector4()
-	{
-	}
-
-	FORCEINLINE FUintVector4(uint32 InX, uint32 InY, uint32 InZ, uint32 InW)
-		: X(InX)
-		, Y(InY)
-		, Z(InZ)
-		, W(InW)
-	{
-	}
-
-	FORCEINLINE explicit FUintVector4(uint32 InValue)
-		: X(InValue)
-		, Y(InValue)
-		, Z(InValue)
-		, W(InValue)
-	{
-	}
-
-	FORCEINLINE FUintVector4(EForceInit)
-		: X(0)
-		, Y(0)
-		, Z(0)
-		, W(0)
-	{
-	}
-
-	FORCEINLINE const uint32& operator[](int32 ComponentIndex) const
-	{
-PRAGMA_DISABLE_DEPRECATION_WARNINGS
-		return XYZW[ComponentIndex];
-PRAGMA_ENABLE_DEPRECATION_WARNINGS
-	}
-
-
-	FORCEINLINE uint32& operator[](int32 ComponentIndex)
-	{
-PRAGMA_DISABLE_DEPRECATION_WARNINGS
-		return XYZW[ComponentIndex];
-PRAGMA_ENABLE_DEPRECATION_WARNINGS
-	}
-
-	FORCEINLINE bool operator==(const FUintVector4& Other) const
-	{
-		return X==Other.X && Y==Other.Y && Z==Other.Z && W==Other.W;
-	}
-
-
-	FORCEINLINE bool operator!=(const FUintVector4& Other) const
-	{
-		return X!=Other.X || Y!=Other.Y || Z!=Other.Z || W!=Other.W;
-	}
-};
-
-template <> struct TIsPODType<FIntVector> { enum { Value = true }; };
-template <> struct TIsPODType<FIntVector4> { enum { Value = true }; };
+template <> struct TIsPODType<FIntVector2>  { enum { Value = true }; };
+template <> struct TIsPODType<FUintVector2> { enum { Value = true }; };
+template <> struct TIsPODType<FIntVector3>  { enum { Value = true }; };
+template <> struct TIsPODType<FUintVector3> { enum { Value = true }; };
+template <> struct TIsPODType<FIntVector4>  { enum { Value = true }; };
 template <> struct TIsPODType<FUintVector4> { enum { Value = true }; };
+
+template <> struct TIsUECoreType<FIntVector2>  { enum { Value = true }; };
+template <> struct TIsUECoreType<FUintVector2> { enum { Value = true }; };
+template <> struct TIsUECoreType<FIntVector3>  { enum { Value = true }; };
+template <> struct TIsUECoreType<FUintVector3> { enum { Value = true }; };
+template <> struct TIsUECoreType<FIntVector4>  { enum { Value = true }; };
+template <> struct TIsUECoreType<FUintVector4> { enum { Value = true }; };
\ No newline at end of file
diff --git a/Engine/Source/Runtime/Core/Public/Math/UnrealMathUtility.h b/Engine/Source/Runtime/Core/Public/Math/UnrealMathUtility.h
--- a/Engine/Source/Runtime/Core/Public/Math/UnrealMathUtility.h
+++ b/Engine/Source/Runtime/Core/Public/Math/UnrealMathUtility.h
@@ -1671,7 +1671,7 @@ public:
 
 	// @parma InOutScissorRect should be set to View.ViewRect before the call
 	// @return 0: light is not visible, 1:use scissor rect, 2: no scissor rect needed
-	UE_NODISCARD static CORE_API uint32 ComputeProjectedSphereScissorRect(struct FIntRect& InOutScissorRect, FVector SphereOrigin, float Radius, FVector ViewOrigin, const FMatrix& ViewMatrix, const FMatrix& ProjMatrix);
+	UE_NODISCARD static CORE_API uint32 ComputeProjectedSphereScissorRect(FIntRect& InOutScissorRect, FVector SphereOrigin, float Radius, FVector ViewOrigin, const FMatrix& ViewMatrix, const FMatrix& ProjMatrix);
 
 	// @param ConeOrigin Cone origin
 	// @param ConeDirection Cone direction
diff --git a/Engine/Source/Runtime/Core/Public/Math/Vector.h b/Engine/Source/Runtime/Core/Public/Math/Vector.h
--- a/Engine/Source/Runtime/Core/Public/Math/Vector.h
+++ b/Engine/Source/Runtime/Core/Public/Math/Vector.h
@@ -190,14 +190,16 @@ public:
      *
      * @param InVector FIntVector to copy from.
      */
-    explicit TVector(FIntVector InVector);
+	template <typename IntType>
+    explicit TVector(TIntVector3<IntType> InVector);
 
     /**
      * Constructs a vector from an FIntPoint.
      *
      * @param A Int Point used to set X and Y coordinates, Z is set to zero.
      */
-    explicit TVector(FIntPoint A);
+	template <typename IntType>
+    explicit TVector(TIntPoint<IntType> A);
 
     /**
      * Constructor which initializes all components to zero.
@@ -1384,14 +1386,16 @@ FORCEINLINE TVector<T>::TVector(const FLinearColor& InColor)
 }
 
 template<typename T>
-FORCEINLINE TVector<T>::TVector(FIntVector InVector)
+template <typename IntType>
+FORCEINLINE TVector<T>::TVector(TIntVector3<IntType> InVector)
     : X((T)InVector.X), Y((T)InVector.Y), Z((T)InVector.Z)
 {
     DiagnosticCheckNaN();
 }
 
 template<typename T>
-FORCEINLINE TVector<T>::TVector(FIntPoint A)
+template <typename IntType>
+FORCEINLINE TVector<T>::TVector(TIntPoint<IntType> A)
     : X((T)A.X), Y((T)A.Y), Z(0.f)
 {
     DiagnosticCheckNaN();
@@ -2603,14 +2607,6 @@ inline FVector FMath::VRand()
     return Result * (1.0f / Sqrt(L));
 }
 
-
-FORCEINLINE FIntVector::FIntVector(FVector InVector)
-	: X(FMath::TruncToInt32(InVector.X))
-	, Y(FMath::TruncToInt32(InVector.Y))
-	, Z(FMath::TruncToInt32(InVector.Z))
-{
-}
-
 template<>
 inline bool FVector3f::SerializeFromMismatchedTag(FName StructTag, FStructuredArchive::FSlot Slot)
 {
@@ -2630,6 +2626,21 @@ inline bool FVector3d::SerializeFromMismatchedTag(FName StructTag, FStructuredAr
 namespace UE {
 namespace Math {
 
+template <>
+FORCEINLINE TIntVector3<int32>::TIntVector3(FVector InVector)
+	: X(FMath::TruncToInt32(InVector.X))
+	, Y(FMath::TruncToInt32(InVector.Y))
+	, Z(FMath::TruncToInt32(InVector.Z))
+{
+}
+
+template <>
+FORCEINLINE TIntVector3<uint32>::TIntVector3(FVector InVector)
+	: X(IntCastChecked<uint32, int64>(FMath::TruncToInt64(InVector.X)))
+	, Y(IntCastChecked<uint32, int64>(FMath::TruncToInt64(InVector.Y)))
+	, Z(IntCastChecked<uint32, int64>(FMath::TruncToInt64(InVector.Z)))
+{
+}
 
 template<typename T>
 FORCEINLINE TVector2<T>::TVector2( const TVector<T>& V )
diff --git a/Engine/Source/Runtime/Core/Public/Math/Vector2D.h b/Engine/Source/Runtime/Core/Public/Math/Vector2D.h
--- a/Engine/Source/Runtime/Core/Public/Math/Vector2D.h
+++ b/Engine/Source/Runtime/Core/Public/Math/Vector2D.h
@@ -97,7 +97,8 @@ public:
 	*
 	* @param InPos Integer point used to set this vector.
 	*/
-	FORCEINLINE TVector2<T>(FIntPoint InPos);
+	template <typename IntType>
+	FORCEINLINE TVector2<T>(TIntPoint<IntType> InPos);
 
 	/**
 	* Constructor which initializes all components to zero.
@@ -764,7 +765,8 @@ FORCEINLINE TVector2<T>::TVector2(T InF)
 { }
 
 template<typename T>
-FORCEINLINE TVector2<T>::TVector2(FIntPoint InPos)
+template<typename IntType>
+FORCEINLINE TVector2<T>::TVector2(TIntPoint<IntType> InPos)
 {
 	X = (T)InPos.X;
 	Y = (T)InPos.Y;
diff --git a/Engine/Source/Runtime/Core/Public/Misc/FileHelper.h b/Engine/Source/Runtime/Core/Public/Misc/FileHelper.h
--- a/Engine/Source/Runtime/Core/Public/Misc/FileHelper.h
+++ b/Engine/Source/Runtime/Core/Public/Misc/FileHelper.h
@@ -166,7 +166,7 @@ struct CORE_API FFileHelper
 	//	BmpImageWrapper now has a good writer, prefer that for general image saving ; use FImageUtils::SaveImage
 	// note this also calls SendDataToPCViaUnrealConsole
 	//   and uses GenerateNextBitmapFilename
-	static bool CreateBitmap( const TCHAR* Pattern, int32 DataWidth, int32 DataHeight, const struct FColor* Data, struct FIntRect* SubRectangle = NULL, IFileManager* FileManager = &IFileManager::Get(), FString* OutFilename = NULL, bool bInWriteAlpha = false, EColorChannel ColorChannel = EColorChannel::All);
+	static bool CreateBitmap( const TCHAR* Pattern, int32 DataWidth, int32 DataHeight, const struct FColor* Data, FIntRect* SubRectangle = NULL, IFileManager* FileManager = &IFileManager::Get(), FString* OutFilename = NULL, bool bInWriteAlpha = false, EColorChannel ColorChannel = EColorChannel::All);
 
 	/**
 	 * Generates the next unique bitmap filename with a specified extension
@@ -223,7 +223,7 @@ struct CORE_API FFileHelper
 
 	PRAGMA_DISABLE_DEPRECATION_WARNINGS
 	UE_DEPRECATED(5.0, "EChannelMask has been deprecated in favor of EColorChannel, please use the other CreateBitmap() method.")
-	static bool CreateBitmap(const TCHAR* Pattern, int32 DataWidth, int32 DataHeight, const struct FColor* Data, struct FIntRect* SubRectangle, IFileManager* FileManager, FString* OutFilename, bool bInWriteAlpha, EChannelMask ChannelMask);
+	static bool CreateBitmap(const TCHAR* Pattern, int32 DataWidth, int32 DataHeight, const struct FColor* Data, FIntRect* SubRectangle, IFileManager* FileManager, FString* OutFilename, bool bInWriteAlpha, EChannelMask ChannelMask);
 	PRAGMA_ENABLE_DEPRECATION_WARNINGS
 };
 
diff --git a/Engine/Source/Runtime/Core/Public/Serialization/Archive.h b/Engine/Source/Runtime/Core/Public/Serialization/Archive.h
--- a/Engine/Source/Runtime/Core/Public/Serialization/Archive.h
+++ b/Engine/Source/Runtime/Core/Public/Serialization/Archive.h
@@ -3,6 +3,7 @@
 #pragma once
 
 #include "CoreTypes.h"
+#include "CoreFwd.h"
 #include "HAL/PlatformProperties.h"
 #include "Internationalization/TextNamespaceFwd.h"
 #include "Misc/AssertionMacros.h"
@@ -1582,7 +1583,7 @@ public:
 	 * @param Ar The archive to serialize from or to.
 	 * @param Value The value to serialize.
 	 */
-	friend FArchive& operator<<(FArchive& Ar, struct FIntRect& Value);
+	friend FArchive& operator<<(FArchive& Ar, FIntRect& Value);
 
 	/**
 	 * Serializes an FString value from or into an archive.
diff --git a/Engine/Source/Runtime/Engine/Classes/Engine/EngineTypes.h b/Engine/Source/Runtime/Engine/Classes/Engine/EngineTypes.h
--- a/Engine/Source/Runtime/Engine/Classes/Engine/EngineTypes.h
+++ b/Engine/Source/Runtime/Engine/Classes/Engine/EngineTypes.h
@@ -3676,10 +3676,10 @@ struct ENGINE_API FRepMovement
 	static int32 EnableMultiplayerWorldOriginRebasing;
 
 	/** Rebase zero-origin position onto local world origin value. */
-	static FVector RebaseOntoLocalOrigin(const FVector& Location, const struct FIntVector& LocalOrigin);
+	static FVector RebaseOntoLocalOrigin(const FVector& Location, const FIntVector& LocalOrigin);
 
 	/** Rebase local-origin position onto zero world origin value. */
-	static FVector RebaseOntoZeroOrigin(const FVector& Location, const struct FIntVector& LocalOrigin);
+	static FVector RebaseOntoZeroOrigin(const FVector& Location, const FIntVector& LocalOrigin);
 
 	/** Rebase zero-origin position onto an Actor's local world origin. */
 	static FVector RebaseOntoLocalOrigin(const FVector& Location, const AActor* const WorldContextActor);
diff --git a/Engine/Source/Runtime/Engine/Private/EngineTypes.cpp b/Engine/Source/Runtime/Engine/Private/EngineTypes.cpp
--- a/Engine/Source/Runtime/Engine/Private/EngineTypes.cpp
+++ b/Engine/Source/Runtime/Engine/Private/EngineTypes.cpp
@@ -673,7 +673,7 @@ bool FRepMovement::NetSerialize(FArchive& Ar, class UPackageMap* Map, bool& bOut
 }
 
 /** Rebase zero-origin position onto local world origin value. */
-FVector FRepMovement::RebaseOntoLocalOrigin(const FVector& Location, const struct FIntVector& LocalOrigin)
+FVector FRepMovement::RebaseOntoLocalOrigin(const FVector& Location, const FIntVector& LocalOrigin)
 {
 	if (EnableMultiplayerWorldOriginRebasing <= 0 || LocalOrigin == FIntVector::ZeroValue)
 	{
@@ -684,7 +684,7 @@ FVector FRepMovement::RebaseOntoLocalOrigin(const FVector& Location, const struc
 }
 
 /** Rebase local-origin position onto zero world origin value. */
-FVector FRepMovement::RebaseOntoZeroOrigin(const FVector& Location, const struct FIntVector& LocalOrigin)
+FVector FRepMovement::RebaseOntoZeroOrigin(const FVector& Location, const FIntVector& LocalOrigin)
 {
 	if (EnableMultiplayerWorldOriginRebasing <= 0 || LocalOrigin == FIntVector::ZeroValue)
 	{
diff --git a/Engine/Source/Runtime/Engine/Public/PrimitiveSceneProxy.h b/Engine/Source/Runtime/Engine/Public/PrimitiveSceneProxy.h
--- a/Engine/Source/Runtime/Engine/Public/PrimitiveSceneProxy.h
+++ b/Engine/Source/Runtime/Engine/Public/PrimitiveSceneProxy.h
@@ -961,7 +961,7 @@ public:
 	 * @param UVChannelIndices (OUT)	The related index for each (array size = TEXSTREAM_MAX_NUM_TEXTURES_PER_MATERIAL / 4)
 	 * @return							Whether scales were computed or not.
 	 */
-	ENGINE_API virtual bool GetMaterialTextureScales(int32 LODIndex, int32 SectionIndex, const class FMaterialRenderProxy* MaterialRenderProxy, FVector4f* OneOverScales, struct FIntVector4* UVChannelIndices) const;
+	ENGINE_API virtual bool GetMaterialTextureScales(int32 LODIndex, int32 SectionIndex, const class FMaterialRenderProxy* MaterialRenderProxy, FVector4f* OneOverScales, FIntVector4* UVChannelIndices) const;
 #endif
 
 	/**
