Logging: Temporary workaround for a hang during crash handling

#rnx

diff --git a/Engine/Source/Runtime/Core/Private/Misc/OutputDeviceRedirector.cpp b/Engine/Source/Runtime/Core/Private/Misc/OutputDeviceRedirector.cpp
--- a/Engine/Source/Runtime/Core/Private/Misc/OutputDeviceRedirector.cpp
+++ b/Engine/Source/Runtime/Core/Private/Misc/OutputDeviceRedirector.cpp
@@ -422,6 +422,11 @@ void FOutputDeviceRedirector::SetCurrentThreadAsMasterThread()
 
 bool FOutputDeviceRedirector::TryStartDedicatedMasterThread()
 {
+	if constexpr (!PLATFORM_WINDOWS)
+	{
+		return false;
+	}
+
 	if (!FApp::ShouldUseThreadingForPerformance())
 	{
 		return false;
