RWLock: Added TryReadLock and TryWriteLock

diff --git a/Engine/Source/Runtime/Core/Public/GenericPlatform/GenericPlatformCriticalSection.h b/Engine/Source/Runtime/Core/Public/GenericPlatform/GenericPlatformCriticalSection.h
--- a/Engine/Source/Runtime/Core/Public/GenericPlatform/GenericPlatformCriticalSection.h
+++ b/Engine/Source/Runtime/Core/Public/GenericPlatform/GenericPlatformCriticalSection.h
@@ -45,31 +45,41 @@ public:
 	FORCEINLINE TGenericPlatformRWLock()
 	{
 	}
-	
+
 	FORCEINLINE ~TGenericPlatformRWLock()
 	{
 	}
-	
+
 	FORCEINLINE void ReadLock()
 	{
 		Mutex.Lock();
 	}
-	
+
 	FORCEINLINE void WriteLock()
 	{
 		Mutex.Lock();
 	}
-	
+
+	FORCEINLINE bool TryReadLock()
+	{
+		return Mutex.TryLock();
+	}
+
+	FORCEINLINE bool TryWriteLock()
+	{
+		return Mutex.TryLock();
+	}
+
 	FORCEINLINE void ReadUnlock()
 	{
 		Mutex.Unlock();
 	}
-	
+
 	FORCEINLINE void WriteUnlock()
 	{
 		Mutex.Unlock();
 	}
-	
+
 private:
 	CriticalSection Mutex;
 };
diff --git a/Engine/Source/Runtime/Core/Public/HAL/PThreadRWLock.h b/Engine/Source/Runtime/Core/Public/HAL/PThreadRWLock.h
--- a/Engine/Source/Runtime/Core/Public/HAL/PThreadRWLock.h
+++ b/Engine/Source/Runtime/Core/Public/HAL/PThreadRWLock.h
@@ -28,25 +28,39 @@ public:
 		int Err = pthread_rwlock_destroy(&Mutex);
 		checkf(Err == 0, TEXT("pthread_rwlock_destroy failed with error: %d"), Err);
 	}
-	
+
 	void ReadLock()
 	{
 		int Err = pthread_rwlock_rdlock(&Mutex);
 		checkf(Err == 0, TEXT("pthread_rwlock_rdlock failed with error: %d"), Err);
 	}
-	
+
 	void WriteLock()
 	{
 		int Err = pthread_rwlock_wrlock(&Mutex);
 		checkf(Err == 0, TEXT("pthread_rwlock_wrlock failed with error: %d"), Err);
 	}
-	
+
+	bool TryReadLock()
+	{
+		int Err = pthread_rwlock_tryrdlock(&Mutex);
+		checkf(Err == 0 || Err == EBUSY, TEXT("pthread_rwlock_tryrdlock failed with error: %d"), Err);
+		return Err == 0;
+	}
+
+	bool TryWriteLock()
+	{
+		int Err = pthread_rwlock_trywrlock(&Mutex);
+		checkf(Err == 0 || Err == EBUSY, TEXT("pthread_rwlock_trywrlock failed with error: %d"), Err);
+		return Err == 0;
+	}
+
 	void ReadUnlock()
 	{
 		int Err = pthread_rwlock_unlock(&Mutex);
 		checkf(Err == 0, TEXT("pthread_rwlock_unlock failed with error: %d"), Err);
 	}
-	
+
 	void WriteUnlock()
 	{
 		int Err = pthread_rwlock_unlock(&Mutex);
diff --git a/Engine/Source/Runtime/Core/Public/HoloLens/HoloLensCriticalSection.h b/Engine/Source/Runtime/Core/Public/HoloLens/HoloLensCriticalSection.h
--- a/Engine/Source/Runtime/Core/Public/HoloLens/HoloLensCriticalSection.h
+++ b/Engine/Source/Runtime/Core/Public/HoloLens/HoloLensCriticalSection.h
@@ -88,31 +88,41 @@ public:
 	{
 		Windows::InitializeSRWLock(&Mutex);
 	}
-	
+
 	FORCEINLINE ~FHoloLensRWLock()
 	{
 	}
-	
+
 	FORCEINLINE void ReadLock()
 	{
 		Windows::AcquireSRWLockShared(&Mutex);
 	}
-	
+
 	FORCEINLINE void WriteLock()
 	{
 		Windows::AcquireSRWLockExclusive(&Mutex);
 	}
-	
+
+	FORCEINLINE bool TryReadLock()
+	{
+		return !!Windows::TryAcquireSRWLockShared(&Mutex);
+	}
+
+	FORCEINLINE bool TryWriteLock()
+	{
+		return !!Windows::TryAcquireSRWLockExclusive(&Mutex);
+	}
+
 	FORCEINLINE void ReadUnlock()
 	{
 		Windows::ReleaseSRWLockShared(&Mutex);
 	}
-	
+
 	FORCEINLINE void WriteUnlock()
 	{
 		Windows::ReleaseSRWLockExclusive(&Mutex);
 	}
-	
+
 private:
 	Windows::SRWLOCK Mutex;
 };
diff --git a/Engine/Source/Runtime/Core/Public/Windows/WindowsCriticalSection.h b/Engine/Source/Runtime/Core/Public/Windows/WindowsCriticalSection.h
--- a/Engine/Source/Runtime/Core/Public/Windows/WindowsCriticalSection.h
+++ b/Engine/Source/Runtime/Core/Public/Windows/WindowsCriticalSection.h
@@ -116,32 +116,42 @@ public:
 	{
 		Windows::InitializeSRWLock(&Mutex);
 	}
-	
+
 	FORCEINLINE ~FWindowsRWLock()
 	{
 		checkf(Windows::TryAcquireSRWLockExclusive(&Mutex), TEXT("Destroying a lock that is still held!"));
 	}
-	
+
 	FORCEINLINE void ReadLock()
 	{
 		Windows::AcquireSRWLockShared(&Mutex);
 	}
-	
+
 	FORCEINLINE void WriteLock()
 	{
 		Windows::AcquireSRWLockExclusive(&Mutex);
 	}
-	
+
+	FORCEINLINE bool TryReadLock()
+	{
+		return !!Windows::TryAcquireSRWLockShared(&Mutex);
+	}
+
+	FORCEINLINE bool TryWriteLock()
+	{
+		return !!Windows::TryAcquireSRWLockExclusive(&Mutex);
+	}
+
 	FORCEINLINE void ReadUnlock()
 	{
 		Windows::ReleaseSRWLockShared(&Mutex);
 	}
-	
+
 	FORCEINLINE void WriteUnlock()
 	{
 		Windows::ReleaseSRWLockExclusive(&Mutex);
 	}
-	
+
 private:
 	Windows::SRWLOCK Mutex;
 };
