UnrealMathTest name conflict resolution.

diff --git a/Engine/Source/Runtime/Core/Private/Tests/Math/UnrealMathTest.cpp b/Engine/Source/Runtime/Core/Private/Tests/Math/UnrealMathTest.cpp
--- a/Engine/Source/Runtime/Core/Private/Tests/Math/UnrealMathTest.cpp
+++ b/Engine/Source/Runtime/Core/Private/Tests/Math/UnrealMathTest.cpp
@@ -4245,6 +4245,10 @@ bool FColorConversionTest::RunTest(const FString& Parameters)
 // because the reference version is not necessarily exposed; we can end up using a different
 // implementation depending on the platform (as of this writing, we have SSE2 implementation
 // too), but we want them all to match.
+// 
+// Wrapping in a namespace to avoid name conflicts but otherwise keep the code as identical
+// as possible.
+namespace UnrealMathTestInternal {
 
 typedef union
 {
@@ -4295,6 +4299,8 @@ static uint8 stbir__linear_to_srgb_uchar_fast(float in)
 	return (uint8)((bias + scale * t) >> 16);
 }
 
+}
+
 IMPLEMENT_SIMPLE_AUTOMATION_TEST(FColorConversionHeavyTest, "System.Core.Math.ColorConversionHeavy", EAutomationTestFlags::ApplicationContextMask | EAutomationTestFlags::EngineFilter)
 bool FColorConversionHeavyTest::RunTest(const FString& Parameters)
 {
@@ -4344,7 +4350,7 @@ bool FColorConversionHeavyTest::RunTest(const FString& Parameters)
 				for (int Channel = 0; Channel < 3; ++Channel)
 				{
 					float Value = LinearInput.Component(Channel);
-					RgbResult[Channel] = stbir__linear_to_srgb_uchar_fast(Value);
+					RgbResult[Channel] = UnrealMathTestInternal::stbir__linear_to_srgb_uchar_fast(Value);
 				}
 
 				// Clamp A, mapping NaNs to 0
