Logging: Added and migrated to GLog->PanicFlush() for crash handling

diff --git a/Engine/Plugins/Experimental/Gauntlet/Source/Gauntlet/Private/GauntletTestController.cpp b/Engine/Plugins/Experimental/Gauntlet/Source/Gauntlet/Private/GauntletTestController.cpp
--- a/Engine/Plugins/Experimental/Gauntlet/Source/Gauntlet/Private/GauntletTestController.cpp
+++ b/Engine/Plugins/Experimental/Gauntlet/Source/Gauntlet/Private/GauntletTestController.cpp
@@ -68,7 +68,7 @@ void UGauntletTestController::EndTest(int32 ExitCode /*= 0*/)
 {
 	UE_LOG(LogGauntlet, Display, TEXT("**** TEST COMPLETE. EXIT CODE: %d ****"), ExitCode);
 	// we flush logs because we don't (currently...) want to treat shutdown errors as failures
-	GLog->PanicFlushThreadedLogs();
+	GLog->FlushThreadedLogs();
 	// force exit only if platform doesn't support quitting
 	FPlatformMisc::RequestExitWithStatus(!FPlatformProperties::SupportsQuit(), static_cast<uint8>(ExitCode));
 }
diff --git a/Engine/Source/Runtime/ApplicationCore/Private/HoloLens/HoloLensErrorOutputDevice.cpp b/Engine/Source/Runtime/ApplicationCore/Private/HoloLens/HoloLensErrorOutputDevice.cpp
--- a/Engine/Source/Runtime/ApplicationCore/Private/HoloLens/HoloLensErrorOutputDevice.cpp
+++ b/Engine/Source/Runtime/ApplicationCore/Private/HoloLens/HoloLensErrorOutputDevice.cpp
@@ -81,7 +81,7 @@ void FHoloLensErrorOutputDevice::HandleError()
 #if !NO_LOGGING
 	FDebug::LogFormattedMessageWithCallstack(LogHoloLens.GetCategoryName(), __FILE__, __LINE__, TEXT("=== Critical error: ==="), GErrorHist, ELogVerbosity::Error);
 #endif
-	GLog->PanicFlushThreadedLogs();
+	GLog->PanicFlush();
 
 	// Copy to clipboard in non-cooked editor builds.
 	FPlatformApplicationMisc::ClipboardCopy(GErrorHist);
diff --git a/Engine/Source/Runtime/ApplicationCore/Private/Mac/MacConsoleOutputDevice.cpp b/Engine/Source/Runtime/ApplicationCore/Private/Mac/MacConsoleOutputDevice.cpp
--- a/Engine/Source/Runtime/ApplicationCore/Private/Mac/MacConsoleOutputDevice.cpp
+++ b/Engine/Source/Runtime/ApplicationCore/Private/Mac/MacConsoleOutputDevice.cpp
@@ -290,7 +290,7 @@ void FMacConsoleOutputDevice::SetDefaultTextColor()
 		UE_LOG(LogCore, Warning, TEXT("*** INTERRUPTED *** : Console Window Closed, Shutting Down"));
 		if (GLog)
 		{
-			GLog->PanicFlushThreadedLogs();
+			GLog->PanicFlush();
 		}
 		_Exit(0);
 	}
diff --git a/Engine/Source/Runtime/Core/Private/Android/AndroidErrorOutputDevice.cpp b/Engine/Source/Runtime/Core/Private/Android/AndroidErrorOutputDevice.cpp
--- a/Engine/Source/Runtime/Core/Private/Android/AndroidErrorOutputDevice.cpp
+++ b/Engine/Source/Runtime/Core/Private/Android/AndroidErrorOutputDevice.cpp
@@ -64,7 +64,7 @@ void FAndroidErrorOutputDevice::HandleError()
 	FDebug::LogFormattedMessageWithCallstack(LogAndroid.GetCategoryName(), __FILE__, __LINE__, TEXT("=== Critical error: ==="), GErrorHist, ELogVerbosity::Error);
 #endif
 	
-	GLog->PanicFlushThreadedLogs();
+	GLog->PanicFlush();
 
 	FCoreDelegates::OnHandleSystemError.Broadcast();
 	FCoreDelegates::OnShutdownAfterError.Broadcast();
diff --git a/Engine/Source/Runtime/Core/Private/Android/AndroidPlatformMisc.cpp b/Engine/Source/Runtime/Core/Private/Android/AndroidPlatformMisc.cpp
--- a/Engine/Source/Runtime/Core/Private/Android/AndroidPlatformMisc.cpp
+++ b/Engine/Source/Runtime/Core/Private/Android/AndroidPlatformMisc.cpp
@@ -1040,8 +1040,7 @@ void DefaultCrashHandler(const FAndroidCrashContext& Context)
 
 		if (GLog)
 		{
-			GLog->SetCurrentThreadAsMasterThread();
-			GLog->Flush();
+			GLog->PanicFlush();
 		}
 
 		if (GWarn)
@@ -1410,8 +1409,7 @@ void FAndroidMisc::TriggerCrashHandler(ECrashContextType InType, const TCHAR* In
 		// we dont flush logs during a fatal signal, malloccrash can cause us to deadlock.
 		if (GLog)
 		{
-			GLog->PanicFlushThreadedLogs();
-			GLog->Flush();
+			GLog->PanicFlush();
 		}
 		if (GWarn)
 		{
diff --git a/Engine/Source/Runtime/Core/Private/HAL/ExceptionHandling.cpp b/Engine/Source/Runtime/Core/Private/HAL/ExceptionHandling.cpp
--- a/Engine/Source/Runtime/Core/Private/HAL/ExceptionHandling.cpp
+++ b/Engine/Source/Runtime/Core/Private/HAL/ExceptionHandling.cpp
@@ -32,7 +32,7 @@ void ReportInteractiveEnsure(const TCHAR* InMessage)
 	GEnsureShowsCRC = true;
 
 #if PLATFORM_USE_REPORT_ENSURE
-	GLog->PanicFlushThreadedLogs();
+	GLog->PanicFlush();
 	ReportEnsure(InMessage, nullptr);
 #endif
 
diff --git a/Engine/Source/Runtime/Core/Private/HAL/ThreadHeartBeat.cpp b/Engine/Source/Runtime/Core/Private/HAL/ThreadHeartBeat.cpp
--- a/Engine/Source/Runtime/Core/Private/HAL/ThreadHeartBeat.cpp
+++ b/Engine/Source/Runtime/Core/Private/HAL/ThreadHeartBeat.cpp
@@ -275,7 +275,7 @@ void FORCENOINLINE FThreadHeartBeat::OnHang(double HangDuration, uint32 ThreadTh
 
 #if PLATFORM_USE_REPORT_ENSURE
 		UE_LOG(LogCore, Error, TEXT("%s"), *ErrorMessage);
-		GLog->PanicFlushThreadedLogs();
+		GLog->PanicFlush();
 
 		// Skip macros and FDebug, we always want this to fire
 		ReportHang(*ErrorMessage, StackFrames, NumStackFrames, ThreadThatHung);
diff --git a/Engine/Source/Runtime/Core/Private/IOS/IOSErrorOutputDevice.cpp b/Engine/Source/Runtime/Core/Private/IOS/IOSErrorOutputDevice.cpp
--- a/Engine/Source/Runtime/Core/Private/IOS/IOSErrorOutputDevice.cpp
+++ b/Engine/Source/Runtime/Core/Private/IOS/IOSErrorOutputDevice.cpp
@@ -52,5 +52,5 @@ void FIOSErrorOutputDevice::HandleError()
 	FDebug::LogFormattedMessageWithCallstack(LogIOS.GetCategoryName(), __FILE__, __LINE__, TEXT("=== Critical error: ==="), GErrorHist, ELogVerbosity::Error);
 #endif
 
-	GLog->PanicFlushThreadedLogs();
+	GLog->PanicFlush();
 }
diff --git a/Engine/Source/Runtime/Core/Private/IOS/IOSPlatformMisc.cpp b/Engine/Source/Runtime/Core/Private/IOS/IOSPlatformMisc.cpp
--- a/Engine/Source/Runtime/Core/Private/IOS/IOSPlatformMisc.cpp
+++ b/Engine/Source/Runtime/Core/Private/IOS/IOSPlatformMisc.cpp
@@ -1362,8 +1362,7 @@ static void DefaultCrashHandler(FIOSCrashContext const& Context)
     Context.ReportCrash();
     if (GLog)
     {
-        GLog->SetCurrentThreadAsMasterThread();
-        GLog->Flush();
+        GLog->PanicFlush();
     }
     if (GWarn)
     {
diff --git a/Engine/Source/Runtime/Core/Private/Mac/MacPlatformMisc.cpp b/Engine/Source/Runtime/Core/Private/Mac/MacPlatformMisc.cpp
--- a/Engine/Source/Runtime/Core/Private/Mac/MacPlatformMisc.cpp
+++ b/Engine/Source/Runtime/Core/Private/Mac/MacPlatformMisc.cpp
@@ -1916,8 +1916,7 @@ static void DefaultCrashHandler(FMacCrashContext const& Context)
 	Context.ReportCrash();
 	if (GLog)
 	{
-		GLog->SetCurrentThreadAsMasterThread();
-		GLog->Flush();
+		GLog->PanicFlush();
 	}
 	if (GWarn)
 	{
diff --git a/Engine/Source/Runtime/Core/Private/Misc/AssertionMacros.cpp b/Engine/Source/Runtime/Core/Private/Misc/AssertionMacros.cpp
--- a/Engine/Source/Runtime/Core/Private/Misc/AssertionMacros.cpp
+++ b/Engine/Source/Runtime/Core/Private/Misc/AssertionMacros.cpp
@@ -515,9 +515,9 @@ void FORCENOINLINE FDebug::CheckVerifyFailedImpl(
 
 	if (GLog)
 	{
-		// Flushing the threaded logs here increases the likelihood that recent messages will be written to the log file, stdout and the debugger console.
+		// Flushing the logs here increases the likelihood that recent messages will be written to the log file, stdout and the debugger console.
 		// Without this, some of the recent messages may not be reported when debugger stops due to an assertion failure.
-		GLog->PanicFlushThreadedLogs();
+		GLog->PanicFlush();
 	}
 
 	if (!FPlatformMisc::IsDebuggerPresent())
diff --git a/Engine/Source/Runtime/Core/Private/Misc/OutputDeviceRedirector.cpp b/Engine/Source/Runtime/Core/Private/Misc/OutputDeviceRedirector.cpp
--- a/Engine/Source/Runtime/Core/Private/Misc/OutputDeviceRedirector.cpp
+++ b/Engine/Source/Runtime/Core/Private/Misc/OutputDeviceRedirector.cpp
@@ -163,8 +163,9 @@ struct FOutputDeviceRedirectorState
  * - Must be locked to write to unbuffered output devices.
  * - Must not be entered when the thread holds a write or master lock.
  */
-struct FOutputDevicesReadScopeLock
+class FOutputDevicesReadScopeLock
 {
+public:
 	FORCEINLINE explicit FOutputDevicesReadScopeLock(FOutputDeviceRedirectorState& InState)
 		: State(InState)
 	{
@@ -192,6 +193,7 @@ struct FOutputDevicesReadScopeLock
 		State.OutputDevicesLockState.fetch_sub(2, std::memory_order_relaxed);
 	}
 
+private:
 	FOutputDeviceRedirectorState& State;
 };
 
@@ -202,8 +204,9 @@ struct FOutputDevicesReadScopeLock
  * - Must be locked to add or remove output devices.
  * - Must not be entered when the thread holds a read, write, or master lock.
  */
-struct FOutputDevicesWriteScopeLock
+class FOutputDevicesWriteScopeLock
 {
+public:
 	FORCEINLINE explicit FOutputDevicesWriteScopeLock(FOutputDeviceRedirectorState& InState)
 		: State(InState)
 	{
@@ -233,6 +236,7 @@ struct FOutputDevicesWriteScopeLock
 		State.OutputDevicesLock.WriteUnlock();
 	}
 
+private:
 	FOutputDeviceRedirectorState& State;
 };
 
@@ -245,8 +249,9 @@ struct FOutputDevicesWriteScopeLock
  * - Must not be entered when the thread holds a write or master lock.
  * - May be locked when the thread holds a read lock.
  */
-struct FOutputDevicesMasterScopeLock
+class FOutputDevicesMasterScopeLock
 {
+public:
 	FORCEINLINE explicit FOutputDevicesMasterScopeLock(FOutputDeviceRedirectorState& InState)
 		: State(InState)
 	{
@@ -258,7 +263,49 @@ struct FOutputDevicesMasterScopeLock
 		State.OutputDevicesLock.WriteUnlock();
 	}
 
+private:
+	FOutputDeviceRedirectorState& State;
+};
+
+/**
+ * A scoped lock for readers of the OutputDevices arrays that need to access master thread state.
+ *
+ * Equivalent to FOutputDevicesMasterScopeLock, except:
+ * - IsLocked() must be checked before accessing the state guarded by the lock.
+ * - The wait to take the lock is bounded. The lock is not taken if the wait time elapses.
+ */
+class FOutputDevicesMasterScopeTryLock
+{
+public:
+	FORCEINLINE explicit FOutputDevicesMasterScopeTryLock(FOutputDeviceRedirectorState& InState, double WaitTime = 0.0)
+		: State(InState)
+		, bLocked(State.OutputDevicesLock.TryWriteLock())
+	{
+		if (!bLocked && WaitTime > 0.0)
+		{
+			TRACE_CPUPROFILER_EVENT_SCOPE(FOutputDevicesWriteScopeTryLock);
+			const double EndTime = FPlatformTime::Seconds() + WaitTime;
+			do
+			{
+				bLocked = State.OutputDevicesLock.TryWriteLock();
+			}
+			while (!bLocked && FPlatformTime::Seconds() < EndTime);
+		}
+	}
+
+	FORCEINLINE ~FOutputDevicesMasterScopeTryLock()
+	{
+		if (bLocked)
+		{
+			State.OutputDevicesLock.WriteUnlock();
+		}
+	}
+
+	FORCEINLINE bool IsLocked() const { return bLocked; }
+
+private:
 	FOutputDeviceRedirectorState& State;
+	bool bLocked = false;
 };
 
 void FOutputDeviceRedirectorState::ThreadLoop()
@@ -373,11 +420,6 @@ void FOutputDeviceRedirector::FlushThreadedLogs(EOutputDeviceRedirectorFlushOpti
 	State->FlushBufferedLines();
 }
 
-void FOutputDeviceRedirector::PanicFlushThreadedLogs()
-{
-	Flush();
-}
-
 void FOutputDeviceRedirector::SerializeBacklog(FOutputDevice* OutputDevice)
 {
 	FReadScopeLock ScopeLock(State->BacklogLock);
@@ -422,11 +464,6 @@ void FOutputDeviceRedirector::SetCurrentThreadAsMasterThread()
 
 bool FOutputDeviceRedirector::TryStartDedicatedMasterThread()
 {
-	if constexpr (!PLATFORM_WINDOWS)
-	{
-		return false;
-	}
-
 	if (!FApp::ShouldUseThreadingForPerformance())
 	{
 		return false;
@@ -536,6 +573,39 @@ void FOutputDeviceRedirector::Flush()
 	}
 }
 
+void FOutputDeviceRedirector::PanicFlush()
+{
+	if (UE::Private::FOutputDevicesMasterScopeTryLock Lock(*State, 1.0); Lock.IsLocked())
+	{
+		State->BufferedLines.Deplete([this](UE::Private::FOutputDeviceLine&& Line)
+		{
+			for (FOutputDevice* OutputDevice : State->BufferedOutputDevices)
+			{
+				if (OutputDevice->CanBeUsedByPanicFlush())
+				{
+					OutputDevice->Serialize(Line.Data, Line.Verbosity, Line.Category, Line.Time);
+				}
+			}
+		});
+
+		for (FOutputDevice* OutputDevice : State->BufferedOutputDevices)
+		{
+			if (OutputDevice->CanBeUsedByPanicFlush())
+			{
+				OutputDevice->Flush();
+			}
+		}
+
+		for (FOutputDevice* OutputDevice : State->UnbufferedOutputDevices)
+		{
+			if (OutputDevice->CanBeUsedByPanicFlush())
+			{
+				OutputDevice->Flush();
+			}
+		}
+	}
+}
+
 void FOutputDeviceRedirector::TearDown()
 {
 	SetCurrentThreadAsMasterThread();
diff --git a/Engine/Source/Runtime/Core/Private/Unix/UnixPlatformCrashContext.cpp b/Engine/Source/Runtime/Core/Private/Unix/UnixPlatformCrashContext.cpp
--- a/Engine/Source/Runtime/Core/Private/Unix/UnixPlatformCrashContext.cpp
+++ b/Engine/Source/Runtime/Core/Private/Unix/UnixPlatformCrashContext.cpp
@@ -847,8 +847,7 @@ void DefaultCrashHandler(const FUnixCrashContext & Context)
 	const_cast<FUnixCrashContext&>(Context).CaptureStackTrace(Context.ErrorFrame);
 	if (GLog)
 	{
-		GLog->SetCurrentThreadAsMasterThread();
-		GLog->Flush();
+		GLog->PanicFlush();
 	}
 	if (GWarn)
 	{
diff --git a/Engine/Source/Runtime/Core/Private/Windows/WindowsErrorOutputDevice.cpp b/Engine/Source/Runtime/Core/Private/Windows/WindowsErrorOutputDevice.cpp
--- a/Engine/Source/Runtime/Core/Private/Windows/WindowsErrorOutputDevice.cpp
+++ b/Engine/Source/Runtime/Core/Private/Windows/WindowsErrorOutputDevice.cpp
@@ -107,7 +107,7 @@ void FWindowsErrorOutputDevice::HandleError()
 #if !NO_LOGGING
 	FDebug::LogFormattedMessageWithCallstack(LogWindows.GetCategoryName(), __FILE__, __LINE__, TEXT("=== Critical error: ==="), GErrorHist, ELogVerbosity::Error);
 #endif
-	GLog->PanicFlushThreadedLogs();
+	GLog->PanicFlush();
 
 	HandleErrorRestoreUI();
 
diff --git a/Engine/Source/Runtime/Core/Private/Windows/WindowsPlatformCrashContext.cpp b/Engine/Source/Runtime/Core/Private/Windows/WindowsPlatformCrashContext.cpp
--- a/Engine/Source/Runtime/Core/Private/Windows/WindowsPlatformCrashContext.cpp
+++ b/Engine/Source/Runtime/Core/Private/Windows/WindowsPlatformCrashContext.cpp
@@ -1410,9 +1410,9 @@ private:
 
 		if (GLog)
 		{
-			//Panic flush the logs to make sure there are no entries queued. This is
-			//not thread safe so it will skip for example editor log.
-			GLog->PanicFlushThreadedLogs();
+			// Panic flush the logs to make sure there are no entries queued.
+			// This is will skip output devices that are not safe to flush during a crash, for example, the editor log.
+			GLog->PanicFlush();
 		}
 		
 		// Get the default settings for the crash context
diff --git a/Engine/Source/Runtime/Core/Public/Misc/OutputDevice.h b/Engine/Source/Runtime/Core/Public/Misc/OutputDevice.h
--- a/Engine/Source/Runtime/Core/Public/Misc/OutputDevice.h
+++ b/Engine/Source/Runtime/Core/Public/Misc/OutputDevice.h
@@ -192,13 +192,21 @@ public:
 	}
 
 	/**
-	* @return whether this output device can be used from multiple threads simultaneously without any locking
-	*/
+	 * @return whether this output device can be used from multiple threads simultaneously without any locking
+	 */
 	virtual bool CanBeUsedOnMultipleThreads() const
 	{
 		return false;
 	}
 
+	/**
+	 * @return whether this output device can be used when flushing after a fatal error.
+	 */
+	virtual bool CanBeUsedByPanicFlush() const
+	{
+		return CanBeUsedOnAnyThread();
+	}
+
 	// Simple text printing.
 	void Log( const TCHAR* S );
 	void Log( ELogVerbosity::Type Verbosity, const TCHAR* S );
diff --git a/Engine/Source/Runtime/Core/Public/Misc/OutputDeviceRedirector.h b/Engine/Source/Runtime/Core/Public/Misc/OutputDeviceRedirector.h
--- a/Engine/Source/Runtime/Core/Public/Misc/OutputDeviceRedirector.h
+++ b/Engine/Source/Runtime/Core/Public/Misc/OutputDeviceRedirector.h
@@ -103,8 +103,9 @@ public:
 	/** Flushes lines buffered by secondary threads. */
 	void FlushThreadedLogs(EOutputDeviceRedirectorFlushOptions Options = EOutputDeviceRedirectorFlushOptions::None);
 
-	/** Same as FlushThreadedLogs(). */
-	void PanicFlushThreadedLogs();
+	/** Same as PanicFlush(). */
+	UE_DEPRECATED(5.1, "Use PanicFlush().")
+	void PanicFlushThreadedLogs() { PanicFlush(); }
 
 	/**
 	 * Serializes the current backlog to the specified output device.
@@ -154,6 +155,9 @@ public:
 	/** Passes on the flush request to all current output devices. */
 	void Flush() final;
 
+	/** Makes an effort to flush threaded logs and flush output devices after a fatal error. */
+	void PanicFlush();
+
 	/**
 	 * Closes output device and cleans up.
 	 *
diff --git a/Engine/Source/Runtime/D3D12RHI/Private/D3D12Util.cpp b/Engine/Source/Runtime/D3D12RHI/Private/D3D12Util.cpp
--- a/Engine/Source/Runtime/D3D12RHI/Private/D3D12Util.cpp
+++ b/Engine/Source/Runtime/D3D12RHI/Private/D3D12Util.cpp
@@ -280,8 +280,7 @@ static bool LogBreadcrumbData(D3D12RHI::FD3DGPUProfiler& GPUProfiler, FD3D12Comm
 		}
 	}
 
-	GLog->PanicFlushThreadedLogs();
-	GLog->Flush();
+	GLog->PanicFlush();
 
 	return true;
 }
@@ -791,8 +790,7 @@ namespace D3D12RHI
 		}
 		
 		// Make sure the log is flushed!
-		GLog->PanicFlushThreadedLogs();
-		GLog->Flush();
+		GLog->PanicFlush();
 
 		// Show message box or trace information
 #if !(UE_BUILD_SHIPPING || UE_BUILD_TEST)
@@ -853,8 +851,7 @@ namespace D3D12RHI
 		}
 
 		// Make sure the log is flushed!
-		GLog->PanicFlushThreadedLogs();
-		GLog->Flush();
+		GLog->PanicFlush();
 
 		UE_LOG(LogD3D12RHI, Fatal, TEXT("%s failed \n at %s:%u \n with error %s\n%s"), ANSI_TO_TCHAR(Code), ANSI_TO_TCHAR(Filename), Line, *ErrorString, *Message);
 
@@ -898,8 +895,7 @@ namespace D3D12RHI
 		}
 
 		// Make sure the log is flushed!
-		GLog->PanicFlushThreadedLogs();
-		GLog->Flush();
+		GLog->PanicFlush();
 
 		UE_LOG(LogD3D12RHI, Fatal,
 			TEXT("%s failed \n at %s:%u \n with error %s, \n Size=%ix%ix%i Format=%s(0x%08X), NumMips=%i, Flags=%s"),
diff --git a/Engine/Source/Runtime/Engine/Private/NetDriver.cpp b/Engine/Source/Runtime/Engine/Private/NetDriver.cpp
--- a/Engine/Source/Runtime/Engine/Private/NetDriver.cpp
+++ b/Engine/Source/Runtime/Engine/Private/NetDriver.cpp
@@ -4962,7 +4962,7 @@ struct FReplicationAutoCapture
 			--KillProcessFrames;
 			if (KillProcessFrames == 0)
 			{
-				GLog->PanicFlushThreadedLogs();
+				GLog->PanicFlush();
 				FPlatformMisc::RequestExit(1);
 			}
 		}
diff --git a/Engine/Source/Runtime/Engine/Private/ShaderCompiler/ShaderCompiler.cpp b/Engine/Source/Runtime/Engine/Private/ShaderCompiler/ShaderCompiler.cpp
--- a/Engine/Source/Runtime/Engine/Private/ShaderCompiler/ShaderCompiler.cpp
+++ b/Engine/Source/Runtime/Engine/Private/ShaderCompiler/ShaderCompiler.cpp
@@ -1117,7 +1117,7 @@ namespace SCWErrorCode
 {
 	void HandleGeneralCrash(const TCHAR* ExceptionInfo, const TCHAR* Callstack)
 	{
-		GLog->PanicFlushThreadedLogs();
+		GLog->PanicFlush();
 		UE_LOG(LogShaderCompilers, Fatal, TEXT("ShaderCompileWorker crashed!\n%s\n\t%s"), ExceptionInfo, Callstack);
 	}
 
@@ -1620,7 +1620,7 @@ static void DumpCompilationJobs(const TArray<FShaderCommonCompileJobPtr>& Queued
 	}
 
 	// Force a log flush so we can track the crash before the cooker potentially crashes before the output shows up
-	GLog->PanicFlushThreadedLogs();
+	GLog->PanicFlush();
 }
 
 // Disable optimization for this crash handler to get full access to the entire stack frame when debugging a crash dump
diff --git a/Engine/Source/Runtime/Launch/Private/Android/AndroidJNI.cpp b/Engine/Source/Runtime/Launch/Private/Android/AndroidJNI.cpp
--- a/Engine/Source/Runtime/Launch/Private/Android/AndroidJNI.cpp
+++ b/Engine/Source/Runtime/Launch/Private/Android/AndroidJNI.cpp
@@ -489,8 +489,7 @@ void EngineCrashHandler(const FGenericCrashContext& GenericContext)
 
 		if (GLog)
 		{
-			GLog->SetCurrentThreadAsMasterThread();
-			GLog->Flush();
+			GLog->PanicFlush();
 		}
 		
 		if (GWarn)
diff --git a/Engine/Source/Runtime/Launch/Private/Mac/LaunchMac.cpp b/Engine/Source/Runtime/Launch/Private/Mac/LaunchMac.cpp
--- a/Engine/Source/Runtime/Launch/Private/Mac/LaunchMac.cpp
+++ b/Engine/Source/Runtime/Launch/Private/Mac/LaunchMac.cpp
@@ -32,8 +32,7 @@ void LogCrashCallstack(const FMacCrashContext& Context)
     Context.ReportCrash();
     if (GLog)
     {
-        GLog->SetCurrentThreadAsMasterThread();
-        GLog->Flush();
+        GLog->PanicFlush();
     }
     if (GWarn)
     {
diff --git a/Engine/Source/Runtime/RHI/Private/RHIValidation.cpp b/Engine/Source/Runtime/RHI/Private/RHIValidation.cpp
--- a/Engine/Source/Runtime/RHI/Private/RHIValidation.cpp
+++ b/Engine/Source/Runtime/RHI/Private/RHIValidation.cpp
@@ -448,7 +448,7 @@ void FValidationRHI::ReportValidationFailure(const TCHAR* InMessage)
 		// Print the message again using the debug output function, because UE_LOG doesn't always reach
 		// the VS output window before the breakpoint is triggered, despite the log flush call below.
 		FPlatformMisc::LowLevelOutputDebugStringf(TEXT("%s\n"), InMessage);
-		GLog->PanicFlushThreadedLogs();
+		GLog->PanicFlush();
 		PLATFORM_BREAK();
 	}
 }
diff --git a/Engine/Source/Runtime/Unix/UnixCommonStartup/Private/UnixCommonStartup.cpp b/Engine/Source/Runtime/Unix/UnixCommonStartup/Private/UnixCommonStartup.cpp
--- a/Engine/Source/Runtime/Unix/UnixCommonStartup/Private/UnixCommonStartup.cpp
+++ b/Engine/Source/Runtime/Unix/UnixCommonStartup/Private/UnixCommonStartup.cpp
@@ -35,8 +35,7 @@ void CommonUnixCrashHandler(const FGenericCrashContext& GenericContext)
 	const_cast< FUnixCrashContext& >(Context).CaptureStackTrace(Context.ErrorFrame);
 	if (GLog)
 	{
-		GLog->SetCurrentThreadAsMasterThread();
-		GLog->Flush();
+		GLog->PanicFlush();
 	}
 	if (GWarn)
 	{
diff --git a/Engine/Source/Runtime/VulkanRHI/Private/VulkanMemory.cpp b/Engine/Source/Runtime/VulkanRHI/Private/VulkanMemory.cpp
--- a/Engine/Source/Runtime/VulkanRHI/Private/VulkanMemory.cpp
+++ b/Engine/Source/Runtime/VulkanRHI/Private/VulkanMemory.cpp
@@ -1022,7 +1022,7 @@ namespace VulkanRHI
 			}
 			DumpRenderTargetPoolMemory(*GLog);
 			Device->GetMemoryManager().DumpMemory();
-			GLog->PanicFlushThreadedLogs();
+			GLog->PanicFlush();
 
 			UE_LOG(LogVulkanRHI, Fatal, TEXT("Out of %s Memory, Requested%.2fKB MemTypeIndex=%d\n"), MemoryType, AllocationSize / 1024.f, MemoryTypeIndex);
 		}
@@ -3624,7 +3624,7 @@ namespace VulkanRHI
 		}
 
 
-		GLog->PanicFlushThreadedLogs();
+		GLog->PanicFlush();
 
 
 #undef VULKAN_LOGMEMORY_PAD
@@ -3643,9 +3643,9 @@ namespace VulkanRHI
 			case VK_ERROR_OUT_OF_DEVICE_MEMORY: MemoryType = TEXT("Local"); break;
 			}
 			DumpMemory();
-			GLog->PanicFlushThreadedLogs();
+			GLog->PanicFlush();
 			DumpRenderTargetPoolMemory(*GLog);
-			GLog->PanicFlushThreadedLogs();
+			GLog->PanicFlush();
 
 			UE_LOG(LogVulkanRHI, Fatal, TEXT("Out of %s Memory, Requested%.2fKB MemTypeIndex=%d\n"), MemoryType, AllocationSize, MemoryTypeIndex);
 		}
diff --git a/Engine/Source/Runtime/VulkanRHI/Private/VulkanUtil.cpp b/Engine/Source/Runtime/VulkanRHI/Private/VulkanUtil.cpp
--- a/Engine/Source/Runtime/VulkanRHI/Private/VulkanUtil.cpp
+++ b/Engine/Source/Runtime/VulkanRHI/Private/VulkanUtil.cpp
@@ -692,8 +692,7 @@ void FVulkanGPUProfiler::DumpCrashMarkers(void* BufferData)
 					const FString* Frame = CachedStrings.Find(Value);
 					UE_LOG(LogVulkanRHI, Error, TEXT("[VK_NV_device_diagnostic_checkpoints] %i: Stage 0x%08x, %s (CRC 0x%x)"), Index, Data[Index].stage, Frame ? *(*Frame) : TEXT("<undefined>"), Value);
 				}
-				GLog->PanicFlushThreadedLogs();
-				GLog->Flush();
+				GLog->PanicFlush();
 			}
 		}
 #endif
@@ -711,8 +710,7 @@ void FVulkanGPUProfiler::DumpCrashMarkers(void* BufferData)
 			UE_LOG(LogVulkanRHI, Warning, TEXT("[gpu_crash_markers] %s"), (CrashMarkers[i] != 0) ? *FrameName : TEXT("unavailable"));
 		}
 
-		GLog->PanicFlushThreadedLogs();
-		GLog->Flush();
+		GLog->PanicFlush();
 	}
 }
 #endif
diff --git a/Engine/Source/Runtime/Windows/D3D11RHI/Private/D3D11RHI.cpp b/Engine/Source/Runtime/Windows/D3D11RHI/Private/D3D11RHI.cpp
--- a/Engine/Source/Runtime/Windows/D3D11RHI/Private/D3D11RHI.cpp
+++ b/Engine/Source/Runtime/Windows/D3D11RHI/Private/D3D11RHI.cpp
@@ -645,7 +645,7 @@ bool FD3DGPUProfiler::CheckGpuHeartbeat(bool bShowActiveStatus) const
 					UE_LOG(LogRHI, Error, TEXT("[Aftermath] Invalid context handle"));
 					NVAFTERMATH_ON_ERROR();
 				}
-				GLog->PanicFlushThreadedLogs();
+				GLog->PanicFlush();
 				return false;
 			}
 		}
