Replacement of DEPRECATED_MACRO with UE_DEPRECATED_MACRO.

diff --git a/Engine/Plugins/Enterprise/VariantManagerContent/Source/VariantManagerContent/Private/PropertyValue.cpp b/Engine/Plugins/Enterprise/VariantManagerContent/Source/VariantManagerContent/Private/PropertyValue.cpp
--- a/Engine/Plugins/Enterprise/VariantManagerContent/Source/VariantManagerContent/Private/PropertyValue.cpp
+++ b/Engine/Plugins/Enterprise/VariantManagerContent/Source/VariantManagerContent/Private/PropertyValue.cpp
@@ -23,9 +23,6 @@
 #include "Editor.h"
 #endif
 
-// WARNING: This should always be the last include in any file that needs it (except .generated.h)
-#include "UObject/UndefineUPropertyMacros.h"
-
 #define LOCTEXT_NAMESPACE "PropertyValue"
 
 DEFINE_LOG_CATEGORY(LogVariantContent);
@@ -1871,5 +1868,3 @@ void UPropertyValueVisibility::Serialize(FArchive& Ar)
 }
 
 #undef LOCTEXT_NAMESPACE
-
-#include "UObject/DefineUPropertyMacros.h"
\ No newline at end of file
diff --git a/Engine/Plugins/ScriptPlugin/Source/ScriptPlugin/Classes/ScriptBlueprintGeneratedClass.h b/Engine/Plugins/ScriptPlugin/Source/ScriptPlugin/Classes/ScriptBlueprintGeneratedClass.h
--- a/Engine/Plugins/ScriptPlugin/Source/ScriptPlugin/Classes/ScriptBlueprintGeneratedClass.h
+++ b/Engine/Plugins/ScriptPlugin/Source/ScriptPlugin/Classes/ScriptBlueprintGeneratedClass.h
@@ -4,9 +4,6 @@
 
 #include "Engine/BlueprintGeneratedClass.h"
 
-// WARNING: This should always be the last include in any file that needs it (except .generated.h)
-#include "UObject/UndefineUPropertyMacros.h"
-
 #include "ScriptBlueprintGeneratedClass.generated.h"
 
 class UScriptBlueprintGeneratedClass;
@@ -185,5 +182,3 @@ public:
 
 	virtual void Serialize(FArchive& Ar) override;
 };
-
-#include "UObject/DefineUPropertyMacros.h"
\ No newline at end of file
diff --git a/Engine/Plugins/ScriptPlugin/Source/ScriptPlugin/Private/ScriptBlueprintGeneratedClass.cpp b/Engine/Plugins/ScriptPlugin/Source/ScriptPlugin/Private/ScriptBlueprintGeneratedClass.cpp
--- a/Engine/Plugins/ScriptPlugin/Source/ScriptPlugin/Private/ScriptBlueprintGeneratedClass.cpp
+++ b/Engine/Plugins/ScriptPlugin/Source/ScriptPlugin/Private/ScriptBlueprintGeneratedClass.cpp
@@ -8,9 +8,6 @@
 #include "LuaIntegration.h"
 #endif
 
-// WARNING: This should always be the last include in any file that needs it (except .generated.h)
-#include "UObject/UndefineUPropertyMacros.h"
-
 /////////////////////////////////////////////////////
 // UScriptBlueprintGeneratedClass
 
@@ -169,5 +166,3 @@ void FScriptContextBase::GetSupportedScriptFileFormats(TArray<FString>& OutForma
 	OutFormats.Add(TEXT("lua;Script"));
 #endif
 }
-
-#include "UObject/DefineUPropertyMacros.h"
\ No newline at end of file
diff --git a/Engine/Source/Programs/Shared/EpicGames.UHT/Parsers/UhtHeaderFileParser.cs b/Engine/Source/Programs/Shared/EpicGames.UHT/Parsers/UhtHeaderFileParser.cs
--- a/Engine/Source/Programs/Shared/EpicGames.UHT/Parsers/UhtHeaderFileParser.cs
+++ b/Engine/Source/Programs/Shared/EpicGames.UHT/Parsers/UhtHeaderFileParser.cs
@@ -691,25 +691,22 @@ namespace EpicGames.UHT.Parsers
 					if (IncludeName.IsConstString())
 					{
 						StringView IncludeNameString = IncludeName.GetUnescapedString(this.HeaderFile);
-						if (IncludeNameString != "UObject/DefineUPropertyMacros.h" && IncludeNameString != "UObject/UndefineUPropertyMacros.h")
+						if (this.bSpottedAutogeneratedHeaderInclude)
 						{
-							if (this.bSpottedAutogeneratedHeaderInclude)
-							{
-								this.HeaderFile.LogError("#include found after .generated.h file - the .generated.h file should always be the last #include in a header");
-							}
-							if (this.HeaderFile.GeneratedHeaderFileName.AsSpan().Equals(IncludeNameString.Span, StringComparison.OrdinalIgnoreCase))
-							{
-								this.bSpottedAutogeneratedHeaderInclude = true;
-							}
-							if (IncludeNameString.Span.Contains(".generated.h", StringComparison.Ordinal))
-							{
-								string RefName = Path.GetFileName(IncludeNameString.ToString()).Replace(".generated.h", ".h");
-								this.HeaderFile.AddReferencedHeader(RefName, true);
-							}
-							else
-							{
-								this.HeaderFile.AddReferencedHeader(IncludeNameString.ToString(), true);
-							}
+							this.HeaderFile.LogError("#include found after .generated.h file - the .generated.h file should always be the last #include in a header");
+						}
+						if (this.HeaderFile.GeneratedHeaderFileName.AsSpan().Equals(IncludeNameString.Span, StringComparison.OrdinalIgnoreCase))
+						{
+							this.bSpottedAutogeneratedHeaderInclude = true;
+						}
+						if (IncludeNameString.Span.Contains(".generated.h", StringComparison.Ordinal))
+						{
+							string RefName = Path.GetFileName(IncludeNameString.ToString()).Replace(".generated.h", ".h");
+							this.HeaderFile.AddReferencedHeader(RefName, true);
+						}
+						else
+						{
+							this.HeaderFile.AddReferencedHeader(IncludeNameString.ToString(), true);
 						}
 					}
 				}
diff --git a/Engine/Source/Programs/UnrealHeaderTool/Private/HeaderParser.cpp b/Engine/Source/Programs/UnrealHeaderTool/Private/HeaderParser.cpp
--- a/Engine/Source/Programs/UnrealHeaderTool/Private/HeaderParser.cpp
+++ b/Engine/Source/Programs/UnrealHeaderTool/Private/HeaderParser.cpp
@@ -8634,24 +8634,21 @@ void FHeaderParser::SimplifiedClassParse(FUnrealSourceFile& SourceFile, const TC
 			// Handle #include directives as if they were 'dependson' keywords.
 			const FString& DependsOnHeaderName = Str;
 
-			if (DependsOnHeaderName != TEXT("\"UObject/DefineUPropertyMacros.h\"") && DependsOnHeaderName != TEXT("\"UObject/UndefineUPropertyMacros.h\""))
+			if (bFoundGeneratedInclude)
 			{
-				if (bFoundGeneratedInclude)
-				{
-					FUHTMessage(SourceFile, CurrentLine).Throwf(TEXT("#include found after .generated.h file - the .generated.h file should always be the last #include in a header"));
-				}
+				FUHTMessage(SourceFile, CurrentLine).Throwf(TEXT("#include found after .generated.h file - the .generated.h file should always be the last #include in a header"));
+			}
 
-				bFoundGeneratedInclude = DependsOnHeaderName.Contains(TEXT(".generated.h"));
-				if (!bFoundGeneratedInclude && DependsOnHeaderName.Len())
-				{
-					bool  bIsQuotedInclude = DependsOnHeaderName[0] == '\"';
-					int32 HeaderFilenameEnd = DependsOnHeaderName.Find(bIsQuotedInclude ? TEXT("\"") : TEXT(">"), ESearchCase::CaseSensitive, ESearchDir::FromStart, 1);
+			bFoundGeneratedInclude = DependsOnHeaderName.Contains(TEXT(".generated.h"));
+			if (!bFoundGeneratedInclude && DependsOnHeaderName.Len())
+			{
+				bool  bIsQuotedInclude = DependsOnHeaderName[0] == '\"';
+				int32 HeaderFilenameEnd = DependsOnHeaderName.Find(bIsQuotedInclude ? TEXT("\"") : TEXT(">"), ESearchCase::CaseSensitive, ESearchDir::FromStart, 1);
 
-					if (HeaderFilenameEnd != INDEX_NONE)
-					{
-						// Include the extension in the name so that we later know where this entry came from.
-						SourceFile.GetIncludes().AddUnique(FHeaderProvider(EHeaderProviderSourceType::FileName, FPaths::GetCleanFilename(DependsOnHeaderName.Mid(1, HeaderFilenameEnd - 1))));
-					}
+				if (HeaderFilenameEnd != INDEX_NONE)
+				{
+					// Include the extension in the name so that we later know where this entry came from.
+					SourceFile.GetIncludes().AddUnique(FHeaderProvider(EHeaderProviderSourceType::FileName, FPaths::GetCleanFilename(DependsOnHeaderName.Mid(1, HeaderFilenameEnd - 1))));
 				}
 			}
 		}
diff --git a/Engine/Source/Runtime/Core/Private/Math/UnrealMath.cpp b/Engine/Source/Runtime/Core/Private/Math/UnrealMath.cpp
--- a/Engine/Source/Runtime/Core/Private/Math/UnrealMath.cpp
+++ b/Engine/Source/Runtime/Core/Private/Math/UnrealMath.cpp
@@ -8,6 +8,7 @@
 #include "Stats/Stats.h"
 #include "Math/RandomStream.h"
 #include "UObject/PropertyPortFlags.h"
+#include "Math/IntRect.h"
 #include "Math/Matrix.h"
 #include "Math/Quat.h"
 #include "Math/Vector.h"
diff --git a/Engine/Source/Runtime/Core/Public/Clang/ClangPlatformCompilerPreSetup.h b/Engine/Source/Runtime/Core/Public/Clang/ClangPlatformCompilerPreSetup.h
--- a/Engine/Source/Runtime/Core/Public/Clang/ClangPlatformCompilerPreSetup.h
+++ b/Engine/Source/Runtime/Core/Public/Clang/ClangPlatformCompilerPreSetup.h
@@ -5,33 +5,6 @@
 #ifndef DISABLE_DEPRECATION
 	#pragma clang diagnostic warning "-Wdeprecated-declarations"
 
-	/**
-	 * Macro for marking up deprecated code, functions and types.
-	 *
-	 * Features that are marked as deprecated are scheduled to be removed from the code base
-	 * in a future release. If you are using a deprecated feature in your code, you should
-	 * replace it before upgrading to the next release. See the Upgrade Notes in the release
-	 * notes for the release in which the feature was marked deprecated.
-	 *
-	 * Sample usage (note the slightly different syntax for classes and structures):
-	 *
-	 *		DEPRECATED(4.xx, "Message")
-	 *		void Function();
-	 *
-	 *		struct DEPRECATED(4.xx, "Message") MODULE_API MyStruct
-	 *		{
-	 *			// StructImplementation
-	 *		};
-	 *		class DEPRECATED(4.xx, "Message") MODULE_API MyClass
-	 *		{
-	 *			// ClassImplementation
-	 *		};
-	 *
-	 * @param VERSION The release number in which the feature was marked deprecated.
-	 * @param MESSAGE A message containing upgrade notes.
-	 */
-	#define DEPRECATED(VERSION, MESSAGE) DEPRECATED_MACRO(4.22, "The DEPRECATED macro has been deprecated in favor of UE_DEPRECATED().") __attribute__((deprecated(MESSAGE " Please update your code to the new API before upgrading to the next release, otherwise your project will no longer compile.")))
-
 	#define PRAGMA_DISABLE_DEPRECATION_WARNINGS \
 		_Pragma("clang diagnostic push") \
 		_Pragma("clang diagnostic ignored \"-Wdeprecated-declarations\"")
@@ -108,7 +81,7 @@
 
 #ifndef PRAGMA_ENABLE_UNSAFE_TYPECAST_WARNINGS
 	#define PRAGMA_ENABLE_UNSAFE_TYPECAST_WARNINGS \
-		DEPRECATED_MACRO(5.0, "The PRAGMA_ENABLE_UNSAFE_TYPECAST_WARNINGS macro has been deprecated in favor of PRAGMA_RESTORE_UNSAFE_TYPECAST_WARNINGS. To force enable warnings use PRAGMA_FORCE_UNSAFE_TYPECAST_WARNINGS.")
+		UE_DEPRECATED_MACRO(5.0, "The PRAGMA_ENABLE_UNSAFE_TYPECAST_WARNINGS macro has been deprecated in favor of PRAGMA_RESTORE_UNSAFE_TYPECAST_WARNINGS. To force enable warnings use PRAGMA_FORCE_UNSAFE_TYPECAST_WARNINGS.")
 #endif // PRAGMA_ENABLE_UNSAFE_TYPECAST_WARNINGS
 
 #ifndef PRAGMA_FORCE_UNSAFE_TYPECAST_WARNINGS
diff --git a/Engine/Source/Runtime/Core/Public/GenericPlatform/GenericPlatformCompilerPreSetup.h b/Engine/Source/Runtime/Core/Public/GenericPlatform/GenericPlatformCompilerPreSetup.h
--- a/Engine/Source/Runtime/Core/Public/GenericPlatform/GenericPlatformCompilerPreSetup.h
+++ b/Engine/Source/Runtime/Core/Public/GenericPlatform/GenericPlatformCompilerPreSetup.h
@@ -3,70 +3,6 @@
 #pragma once
 
 
-#ifndef DEPRECATED
-	/**
-	 * Macro for marking up deprecated code, functions and types.
-	 *
-	 * Features that are marked as deprecated are scheduled to be removed from the code base
-	 * in a future release. If you are using a deprecated feature in your code, you should
-	 * replace it before upgrading to the next release. See the Upgrade Notes in the release
-	 * notes for the release in which the feature was marked deprecated.
-	 *
-	 * Sample usage (note the slightly different syntax for classes and structures):
-	 *
-	 *		DEPRECATED(4.xx, "Message")
-	 *		void Function();
-	 *
-	 *		struct DEPRECATED(4.xx, "Message") MODULE_API MyStruct
-	 *		{
-	 *			// StructImplementation
-	 *		};
-	 *		class DEPRECATED(4.xx, "Message") MODULE_API MyClass
-	 *		{
-	 *			// ClassImplementation
-	 *		};
-	 *
-	 *		Unfortunately, VC++ will complain about using member functions and fields from deprecated
-	 *		class/structs even for class/struct implementation e.g.:
-	 *		class DEPRECATED(4.xx, "") DeprecatedClass
-	 *		{
-	 *		public:
-	 *			DeprecatedClass() {}
-	 *
-	 *			float GetMyFloat()
-	 *			{
-	 *				return MyFloat; // This line will cause warning that deprecated field is used.
-	 *			}
-	 *		private:
-	 *			float MyFloat;
-	 *		};
-	 *
-	 *		To get rid of this warning, place all code not called in class implementation in non-deprecated
-	 *		base class and deprecate only derived one. This may force you to change some access specifiers
-	 *		from private to protected, e.g.:
-	 *
-	 *		class DeprecatedClass_Base_DEPRECATED
-	 *		{
-	 *		protected: // MyFloat is protected now, so DeprecatedClass has access to it.
-	 *			float MyFloat;
-	 *		};
-	 *
-	 *		class DEPRECATED(4.xx, "") DeprecatedClass : DeprecatedClass_Base_DEPRECATED
-	 *		{
-	 *		public:
-	 *			DeprecatedClass() {}
-	 *
-	 *			float GetMyFloat()
-	 *			{
-	 *				return MyFloat;
-	 *			}
-	 *		};
-	 * @param VERSION The release number in which the feature was marked deprecated.
-	 * @param MESSAGE A message text containing additional upgrade notes.
-	 */
-	#define DEPRECATED(VERSION, MESSAGE) DEPRECATED_MACRO(4.22, "The DEPRECATED macro has been deprecated in favor of UE_DEPRECATED().")
-#endif // DEPRECATED
-
 #ifndef PRAGMA_DISABLE_DEPRECATION_WARNINGS
 	#define PRAGMA_DISABLE_DEPRECATION_WARNINGS
 #endif
@@ -84,8 +20,12 @@
 		EMIT_CUSTOM_WARNING_AT_LINE(__LINE__, Warning)
 #endif // EMIT_CUSTOM_WARNING
 
+#ifndef UE_DEPRECATED_MACRO
+	#define UE_DEPRECATED_MACRO(Version, Message) EMIT_CUSTOM_WARNING(Message " Please update your code to the new API before upgrading to the next release, otherwise your project will no longer compile.")
+#endif
+
 #ifndef DEPRECATED_MACRO
-	#define DEPRECATED_MACRO(Version, Message) EMIT_CUSTOM_WARNING(Message " Please update your code to the new API before upgrading to the next release, otherwise your project will no longer compile.")
+	#define DEPRECATED_MACRO(Version, Message) UE_DEPRECATED_MACRO(5.1, "The DEPRECATED_MACRO macro has been deprecated in favor of UE_DEPRECATED_MACRO.") UE_DEPRECATED_MACRO(Version, Message)
 #endif
 
 #ifndef PRAGMA_DEFAULT_VISIBILITY_START
diff --git a/Engine/Source/Runtime/Core/Public/HAL/Platform.h b/Engine/Source/Runtime/Core/Public/HAL/Platform.h
--- a/Engine/Source/Runtime/Core/Public/HAL/Platform.h
+++ b/Engine/Source/Runtime/Core/Public/HAL/Platform.h
@@ -211,25 +211,9 @@
 #ifndef PLATFORM_COMPILER_DISTINGUISHES_INT_AND_LONG
 	#define PLATFORM_COMPILER_DISTINGUISHES_INT_AND_LONG			0
 #endif
-#ifdef _MSC_VER
-	#define PLATFORM_COMPILER_HAS_AUTO_RETURN_TYPES 1
-	#ifndef __clang__
-		#pragma deprecated("PLATFORM_COMPILER_HAS_AUTO_RETURN_TYPES")
-	#endif
-#else
-	#define PLATFORM_COMPILER_HAS_AUTO_RETURN_TYPES 1 DEPRECATED_MACRO(4.19, "PLATFORM_COMPILER_HAS_AUTO_RETURN_TYPES has been deprecated and should be replaced with 1.")
-#endif
 #ifndef PLATFORM_COMPILER_HAS_GENERIC_KEYWORD
 	#define PLATFORM_COMPILER_HAS_GENERIC_KEYWORD	0
 #endif
-#ifdef _MSC_VER
-	#define PLATFORM_COMPILER_HAS_DEFAULTED_FUNCTIONS 1
-	#ifndef __clang__
-		#pragma deprecated("PLATFORM_COMPILER_HAS_DEFAULTED_FUNCTIONS")
-	#endif
-#else
-	#define PLATFORM_COMPILER_HAS_DEFAULTED_FUNCTIONS 1 DEPRECATED_MACRO(4.19, "PLATFORM_COMPILER_HAS_DEFAULTED_FUNCTIONS has been deprecated and should be replaced with 1.")
-#endif
 #ifndef PLATFORM_COMPILER_COMMON_LANGUAGE_RUNTIME_COMPILATION
 	#define PLATFORM_COMPILER_COMMON_LANGUAGE_RUNTIME_COMPILATION 0
 #endif
@@ -245,7 +229,7 @@
 		#pragma deprecated("PLATFORM_COMPILER_HAS_IF_CONSTEXPR")
 	#endif
 #else
-	#define PLATFORM_COMPILER_HAS_IF_CONSTEXPR 1 DEPRECATED_MACRO(5.1, "PLATFORM_COMPILER_HAS_IF_CONSTEXPR has been deprecated and should be replaced with 1.")
+	#define PLATFORM_COMPILER_HAS_IF_CONSTEXPR 1 UE_DEPRECATED_MACRO(5.1, "PLATFORM_COMPILER_HAS_IF_CONSTEXPR has been deprecated and should be replaced with 1.")
 #endif
 #ifndef PLATFORM_COMPILER_HAS_FOLD_EXPRESSIONS
 	#define PLATFORM_COMPILER_HAS_FOLD_EXPRESSIONS 0
@@ -692,26 +676,6 @@
 	#define RESTRICT __restrict						/* no alias hint */
 #endif
 
-/* Wrap a function signature in these to warn that callers should not ignore the return value */
-#ifndef FUNCTION_CHECK_RETURN_START
-	#define FUNCTION_CHECK_RETURN_START \
-		DEPRECATED_MACRO(4.26, "FUNCTION_CHECK_RETURN_START has been deprecated - please use UE_NODISCARD")
-#endif
-#ifndef FUNCTION_CHECK_RETURN_END
-	#define FUNCTION_CHECK_RETURN_END \
-		DEPRECATED_MACRO(4.26, "FUNCTION_CHECK_RETURN_END has been deprecated - please use UE_NODISCARD")
-#endif
-
-/* Wrap a function signature in these to indicate that the function never returns */
-#ifndef FUNCTION_NO_RETURN_START
-	#define FUNCTION_NO_RETURN_START \
-		DEPRECATED_MACRO(4.26, "FUNCTION_NO_RETURN_START has been deprecated - please use UE_NORETURN")
-#endif
-#ifndef FUNCTION_NO_RETURN_END
-	#define FUNCTION_NO_RETURN_END \
-		DEPRECATED_MACRO(4.26, "FUNCTION_NO_RETURN_END has been deprecated - please use UE_NORETURN")
-#endif
-
 /* Use before a function declaration to warn that callers should not ignore the return value */
 #if !defined(UE_NODISCARD) && defined(__has_cpp_attribute)
 	#if __has_cpp_attribute(nodiscard)
@@ -748,10 +712,6 @@
 	#define FUNCTION_NON_NULL_RETURN_END
 #endif
 
-#ifndef FUNCTION_CHECK_RETURN
-	#define FUNCTION_CHECK_RETURN(...) DEPRECATED_MACRO(4.12, "FUNCTION_CHECK_RETURN has been deprecated and should be replaced with FUNCTION_CHECK_RETURN_START and FUNCTION_CHECK_RETURN_END.") FUNCTION_CHECK_RETURN_START __VA_ARGS__ FUNCTION_CHECK_RETURN_END
-#endif
-
 /** Promise expression is true. Compiler can optimize accordingly with undefined behavior if wrong. Static analyzers understand this.  */
 #ifndef UE_ASSUME
 	#if defined(__clang__)
@@ -763,8 +723,6 @@
 	#endif
 #endif
 
-#define ASSUME(x) UE_ASSUME(x) DEPRECATED_MACRO(4.25, "Please use UE_ASSUME instead.")
-
 /** Branch prediction hints */
 #ifndef LIKELY						/* Hints compiler that expression is likely to be true, much softer than UE_ASSUME - allows (penalized by worse performance) expression to be false */
 	#if ( defined(__clang__) || defined(__GNUC__) ) && (PLATFORM_UNIX)	// effect of these on non-Linux platform has not been analyzed as of 2016-03-21
@@ -917,10 +875,6 @@
 	#define FAST_BOOT_HACKS  0
 #endif
 
-#ifndef DEPRECATED_FORGAME
-	#define DEPRECATED_FORGAME(...) DEPRECATED_MACRO(4.22, "The DEPRECATED_FORGAME macro has been deprecated in favor of UE_DEPRECATED_FORGAME().")
-#endif
-
 // Console ANSICHAR/TCHAR command line handling
 #if PLATFORM_COMPILER_HAS_TCHAR_WMAIN
 #define INT32_MAIN_INT32_ARGC_TCHAR_ARGV() int32 wmain(int32 ArgC, TCHAR* ArgV[])
diff --git a/Engine/Source/Runtime/Core/Public/HoloLens/HoloLensPlatformCompilerPreSetup.h b/Engine/Source/Runtime/Core/Public/HoloLens/HoloLensPlatformCompilerPreSetup.h
--- a/Engine/Source/Runtime/Core/Public/HoloLens/HoloLensPlatformCompilerPreSetup.h
+++ b/Engine/Source/Runtime/Core/Public/HoloLens/HoloLensPlatformCompilerPreSetup.h
@@ -3,8 +3,6 @@
 #pragma once
 
 #ifndef DISABLE_DEPRECATION
-	#define DEPRECATED(VERSION, MESSAGE) DEPRECATED_MACRO(4.22, "The DEPRECATED macro has been deprecated in favor of UE_DEPRECATED().") __declspec(deprecated(MESSAGE " Please update your code to the new API before upgrading to the next release, otherwise your project will no longer compile."))
-
 	#define PRAGMA_DISABLE_DEPRECATION_WARNINGS \
 		__pragma (warning(push)) \
 		__pragma (warning(disable: 4995)) /* 'function': name was marked as #pragma deprecated */ \
@@ -38,7 +36,7 @@
 
 #ifndef PRAGMA_ENABLE_UNSAFE_TYPECAST_WARNINGS
 	#define PRAGMA_ENABLE_UNSAFE_TYPECAST_WARNINGS \
-		DEPRECATED_MACRO(5.0, "The PRAGMA_ENABLE_UNSAFE_TYPECAST_WARNINGS macro has been deprecated in favor of PRAGMA_RESTORE_UNSAFE_TYPECAST_WARNINGS. To force enable warnings use PRAGMA_FORCE_UNSAFE_TYPECAST_WARNINGS.")
+		UE_DEPRECATED_MACRO(5.0, "The PRAGMA_ENABLE_UNSAFE_TYPECAST_WARNINGS macro has been deprecated in favor of PRAGMA_RESTORE_UNSAFE_TYPECAST_WARNINGS. To force enable warnings use PRAGMA_FORCE_UNSAFE_TYPECAST_WARNINGS.")
 #endif // PRAGMA_ENABLE_UNSAFE_TYPECAST_WARNINGS
 
 #ifndef PRAGMA_FORCE_UNSAFE_TYPECAST_WARNINGS
diff --git a/Engine/Source/Runtime/Core/Public/Misc/AssertionMacros.h b/Engine/Source/Runtime/Core/Public/Misc/AssertionMacros.h
--- a/Engine/Source/Runtime/Core/Public/Misc/AssertionMacros.h
+++ b/Engine/Source/Runtime/Core/Public/Misc/AssertionMacros.h
@@ -177,7 +177,7 @@ RetType FORCENOINLINE UE_DEBUG_SECTION DispatchCheckVerify(InnerType&& Inner, Ar
 #endif // !UE_BUILD_SHIPPING
 
 #define _DebugBreakAndPromptForRemote() \
-	DEPRECATED_MACRO(5.1, "Use UE_DEBUG_BREAK_AND_PROMPT_FOR_REMOTE.")
+	UE_DEPRECATED_MACRO(5.1, "Use UE_DEBUG_BREAK_AND_PROMPT_FOR_REMOTE.")
 
 #if !UE_BUILD_SHIPPING
 	extern CORE_API bool GIgnoreDebugger;
diff --git a/Engine/Source/Runtime/Core/Public/Misc/CoreMiscDefines.h b/Engine/Source/Runtime/Core/Public/Misc/CoreMiscDefines.h
--- a/Engine/Source/Runtime/Core/Public/Misc/CoreMiscDefines.h
+++ b/Engine/Source/Runtime/Core/Public/Misc/CoreMiscDefines.h
@@ -124,8 +124,8 @@ enum EInPlace {InPlace};
 	#define UE_PUSH_MACRO(name) __pragma(push_macro(name))
 	#define UE_POP_MACRO(name) __pragma(pop_macro(name))
 #endif
-#define PUSH_MACRO(name) DEPRECATED_MACRO(5.0, "PUSH_MACRO is deprecated. Use UE_PUSH_MACRO and pass the macro name as a string.") UE_PUSH_MACRO(PREPROCESSOR_TO_STRING(name))
-#define POP_MACRO(name) DEPRECATED_MACRO(5.0, "POP_MACRO is deprecated. Use UE_POP_MACRO and pass the macro name as a string.") UE_POP_MACRO(PREPROCESSOR_TO_STRING(name))
+#define PUSH_MACRO(name) UE_DEPRECATED_MACRO(5.0, "PUSH_MACRO is deprecated. Use UE_PUSH_MACRO and pass the macro name as a string.") UE_PUSH_MACRO(PREPROCESSOR_TO_STRING(name))
+#define POP_MACRO(name) UE_DEPRECATED_MACRO(5.0, "POP_MACRO is deprecated. Use UE_POP_MACRO and pass the macro name as a string.") UE_POP_MACRO(PREPROCESSOR_TO_STRING(name))
 
 #ifdef __COUNTER__
 	// Created a variable with a unique name
diff --git a/Engine/Source/Runtime/Core/Public/Misc/VarArgs.h b/Engine/Source/Runtime/Core/Public/Misc/VarArgs.h
--- a/Engine/Source/Runtime/Core/Public/Misc/VarArgs.h
+++ b/Engine/Source/Runtime/Core/Public/Misc/VarArgs.h
@@ -77,92 +77,3 @@ Ugly VarArgs type checking (debug builds only).
 #define VARARG_EXTRA(...) __VA_ARGS__,
 #define VARARG_NONE
 #define VARARG_PURE =0
-
-#if PLATFORM_WINDOWS
-
-	static inline uint32			CheckVA(uint32 dw) { return dw; }
-	static inline uint8			CheckVA(uint8 b) { return b; }
-	static inline int32			CheckVA(int32 i) { return i; }
-	static inline uint64			CheckVA(uint64 qw) { return qw; }
-	static inline int64		CheckVA(int64 sqw) { return sqw; }
-	static inline double		CheckVA(double d) { return d; }
-	static inline long		CheckVA(long d) { return d; }
-	static inline long		CheckVA(unsigned long d) { return d; }
-	static inline TCHAR			CheckVA(TCHAR c) { return c; }
-	static inline void*			CheckVA(ANSICHAR* s) { return (void*)s; }
-	static inline bool		CheckVA(bool b) { return b; }
-	template<class T> T*		CheckVA(T* p) { return p; }
-	template<class T> const T*	CheckVA(const T* p) { return p; }
-
-	#define VARARG_DECL( FuncRet, StaticFuncRet, Return, FuncName, Pure, FmtType, ExtraDecl, ExtraCall )	\
-		DEPRECATED_MACRO(4.19, "VARARG_DECL has been deprecated - please use a variadic template instead (see FString::Printf as an example)") \
-		FuncRet FuncName##__VA( ExtraDecl FmtType Fmt, ... ) Pure;  \
-		StaticFuncRet FuncName(ExtraDecl FmtType Fmt) {Return FuncName##__VA(ExtraCall (Fmt));} \
-		template<class T1> \
-		StaticFuncRet FuncName(ExtraDecl FmtType Fmt,T1 V1) {T1 v1=CheckVA(V1);Return FuncName##__VA(ExtraCall (Fmt),(v1));} \
-		template<class T1,class T2> \
-		StaticFuncRet FuncName(ExtraDecl FmtType Fmt,T1 V1,T2 V2) {T1 v1=CheckVA(V1);T2 v2=CheckVA(V2);Return FuncName##__VA(ExtraCall (Fmt),(v1),(v2));} \
-		template<class T1,class T2,class T3> \
-		StaticFuncRet FuncName(ExtraDecl FmtType Fmt,T1 V1,T2 V2,T3 V3) {T1 v1=CheckVA(V1);T2 v2=CheckVA(V2);T3 v3=CheckVA(V3);Return FuncName##__VA(ExtraCall (Fmt),(v1),(v2),(v3));} \
-		template<class T1,class T2,class T3,class T4> \
-		StaticFuncRet FuncName(ExtraDecl FmtType Fmt,T1 V1,T2 V2,T3 V3,T4 V4) {T1 v1=CheckVA(V1);T2 v2=CheckVA(V2);T3 v3=CheckVA(V3);T4 v4=CheckVA(V4);Return FuncName##__VA(ExtraCall (Fmt),(v1),(v2),(v3),(v4));} \
-		template<class T1,class T2,class T3,class T4,class T5> \
-		StaticFuncRet FuncName(ExtraDecl FmtType Fmt,T1 V1,T2 V2,T3 V3,T4 V4,T5 V5) {T1 v1=CheckVA(V1);T2 v2=CheckVA(V2);T3 v3=CheckVA(V3);T4 v4=CheckVA(V4);T5 v5=CheckVA(V5);Return FuncName##__VA(ExtraCall (Fmt),(v1),(v2),(v3),(v4),(v5));} \
-		template<class T1,class T2,class T3,class T4,class T5,class T6> \
-		StaticFuncRet FuncName(ExtraDecl FmtType Fmt,T1 V1,T2 V2,T3 V3,T4 V4,T5 V5,T6 V6) {T1 v1=CheckVA(V1);T2 v2=CheckVA(V2);T3 v3=CheckVA(V3);T4 v4=CheckVA(V4);T5 v5=CheckVA(V5);T6 v6=CheckVA(V6);Return FuncName##__VA(ExtraCall (Fmt),(v1),(v2),(v3),(v4),(v5),(v6));} \
-		template<class T1,class T2,class T3,class T4,class T5,class T6,class T7> \
-		StaticFuncRet FuncName(ExtraDecl FmtType Fmt,T1 V1,T2 V2,T3 V3,T4 V4,T5 V5,T6 V6,T7 V7) {T1 v1=CheckVA(V1);T2 v2=CheckVA(V2);T3 v3=CheckVA(V3);T4 v4=CheckVA(V4);T5 v5=CheckVA(V5);T6 v6=CheckVA(V6);T7 v7=CheckVA(V7);Return FuncName##__VA(ExtraCall (Fmt),(v1),(v2),(v3),(v4),(v5),(v6),(v7));} \
-		template<class T1,class T2,class T3,class T4,class T5,class T6,class T7,class T8> \
-		StaticFuncRet FuncName(ExtraDecl FmtType Fmt,T1 V1,T2 V2,T3 V3,T4 V4,T5 V5,T6 V6,T7 V7,T8 V8) {T1 v1=CheckVA(V1);T2 v2=CheckVA(V2);T3 v3=CheckVA(V3);T4 v4=CheckVA(V4);T5 v5=CheckVA(V5);T6 v6=CheckVA(V6);T7 v7=CheckVA(V7);T8 v8=CheckVA(V8);Return FuncName##__VA(ExtraCall (Fmt),(v1),(v2),(v3),(v4),(v5),(v6),(v7),(v8));} \
-		template<class T1,class T2,class T3,class T4,class T5,class T6,class T7,class T8,class T9> \
-		StaticFuncRet FuncName(ExtraDecl FmtType Fmt,T1 V1,T2 V2,T3 V3,T4 V4,T5 V5,T6 V6,T7 V7,T8 V8,T9 V9) {T1 v1=CheckVA(V1);T2 v2=CheckVA(V2);T3 v3=CheckVA(V3);T4 v4=CheckVA(V4);T5 v5=CheckVA(V5);T6 v6=CheckVA(V6);T7 v7=CheckVA(V7);T8 v8=CheckVA(V8);T9 v9=CheckVA(V9);Return FuncName##__VA(ExtraCall (Fmt),(v1),(v2),(v3),(v4),(v5),(v6),(v7),(v8),(v9));} \
-		template<class T1,class T2,class T3,class T4,class T5,class T6,class T7,class T8,class T9,class T10> \
-		StaticFuncRet FuncName(ExtraDecl FmtType Fmt,T1 V1,T2 V2,T3 V3,T4 V4,T5 V5,T6 V6,T7 V7,T8 V8,T9 V9,T10 V10) {T1 v1=CheckVA(V1);T2 v2=CheckVA(V2);T3 v3=CheckVA(V3);T4 v4=CheckVA(V4);T5 v5=CheckVA(V5);T6 v6=CheckVA(V6);T7 v7=CheckVA(V7);T8 v8=CheckVA(V8);T9 v9=CheckVA(V9);T10 v10=CheckVA(V10);Return FuncName##__VA(ExtraCall (Fmt),(v1),(v2),(v3),(v4),(v5),(v6),(v7),(v8),(v9),(v10));} \
-		template<class T1,class T2,class T3,class T4,class T5,class T6,class T7,class T8,class T9,class T10,class T11> \
-		StaticFuncRet FuncName(ExtraDecl FmtType Fmt,T1 V1,T2 V2,T3 V3,T4 V4,T5 V5,T6 V6,T7 V7,T8 V8,T9 V9,T10 V10,T11 V11) {T1 v1=CheckVA(V1);T2 v2=CheckVA(V2);T3 v3=CheckVA(V3);T4 v4=CheckVA(V4);T5 v5=CheckVA(V5);T6 v6=CheckVA(V6);T7 v7=CheckVA(V7);T8 v8=CheckVA(V8);T9 v9=CheckVA(V9);T10 v10=CheckVA(V10);T11 v11=CheckVA(V11);Return FuncName##__VA(ExtraCall (Fmt),(v1),(v2),(v3),(v4),(v5),(v6),(v7),(v8),(v9),(v10),(v11));} \
-		template<class T1,class T2,class T3,class T4,class T5,class T6,class T7,class T8,class T9,class T10,class T11,class T12> \
-		StaticFuncRet FuncName(ExtraDecl FmtType Fmt,T1 V1,T2 V2,T3 V3,T4 V4,T5 V5,T6 V6,T7 V7,T8 V8,T9 V9,T10 V10,T11 V11,T12 V12) {T1 v1=CheckVA(V1);T2 v2=CheckVA(V2);T3 v3=CheckVA(V3);T4 v4=CheckVA(V4);T5 v5=CheckVA(V5);T6 v6=CheckVA(V6);T7 v7=CheckVA(V7);T8 v8=CheckVA(V8);T9 v9=CheckVA(V9);T10 v10=CheckVA(V10);T11 v11=CheckVA(V11);T12 v12=CheckVA(V12);Return FuncName##__VA(ExtraCall (Fmt),(v1),(v2),(v3),(v4),(v5),(v6),(v7),(v8),(v9),(v10),(v11),(v12));} \
-		template<class T1,class T2,class T3,class T4,class T5,class T6,class T7,class T8,class T9,class T10,class T11,class T12,class T13> \
-		StaticFuncRet FuncName(ExtraDecl FmtType Fmt,T1 V1,T2 V2,T3 V3,T4 V4,T5 V5,T6 V6,T7 V7,T8 V8,T9 V9,T10 V10,T11 V11,T12 V12,T13 V13) {T1 v1=CheckVA(V1);T2 v2=CheckVA(V2);T3 v3=CheckVA(V3);T4 v4=CheckVA(V4);T5 v5=CheckVA(V5);T6 v6=CheckVA(V6);T7 v7=CheckVA(V7);T8 v8=CheckVA(V8);T9 v9=CheckVA(V9);T10 v10=CheckVA(V10);T11 v11=CheckVA(V11);T12 v12=CheckVA(V12);T13 v13=CheckVA(V13);Return FuncName##__VA(ExtraCall (Fmt),(v1),(v2),(v3),(v4),(v5),(v6),(v7),(v8),(v9),(v10),(v11),(v12),(v13));} \
-		template<class T1,class T2,class T3,class T4,class T5,class T6,class T7,class T8,class T9,class T10,class T11,class T12,class T13,class T14> \
-		StaticFuncRet FuncName(ExtraDecl FmtType Fmt,T1 V1,T2 V2,T3 V3,T4 V4,T5 V5,T6 V6,T7 V7,T8 V8,T9 V9,T10 V10,T11 V11,T12 V12,T13 V13,T14 V14) {T1 v1=CheckVA(V1);T2 v2=CheckVA(V2);T3 v3=CheckVA(V3);T4 v4=CheckVA(V4);T5 v5=CheckVA(V5);T6 v6=CheckVA(V6);T7 v7=CheckVA(V7);T8 v8=CheckVA(V8);T9 v9=CheckVA(V9);T10 v10=CheckVA(V10);T11 v11=CheckVA(V11);T12 v12=CheckVA(V12);T13 v13=CheckVA(V13);T14 v14=CheckVA(V14);Return FuncName##__VA(ExtraCall (Fmt),(v1),(v2),(v3),(v4),(v5),(v6),(v7),(v8),(v9),(v10),(v11),(v12),(v13),(v14));} \
-		template<class T1,class T2,class T3,class T4,class T5,class T6,class T7,class T8,class T9,class T10,class T11,class T12,class T13,class T14,class T15> \
-		StaticFuncRet FuncName(ExtraDecl FmtType Fmt,T1 V1,T2 V2,T3 V3,T4 V4,T5 V5,T6 V6,T7 V7,T8 V8,T9 V9,T10 V10,T11 V11,T12 V12,T13 V13,T14 V14,T15 V15) {T1 v1=CheckVA(V1);T2 v2=CheckVA(V2);T3 v3=CheckVA(V3);T4 v4=CheckVA(V4);T5 v5=CheckVA(V5);T6 v6=CheckVA(V6);T7 v7=CheckVA(V7);T8 v8=CheckVA(V8);T9 v9=CheckVA(V9);T10 v10=CheckVA(V10);T11 v11=CheckVA(V11);T12 v12=CheckVA(V12);T13 v13=CheckVA(V13);T14 v14=CheckVA(V14);T15 v15=CheckVA(V15);Return FuncName##__VA(ExtraCall (Fmt),(v1),(v2),(v3),(v4),(v5),(v6),(v7),(v8),(v9),(v10),(v11),(v12),(v13),(v14),(v15));} \
-		template<class T1,class T2,class T3,class T4,class T5,class T6,class T7,class T8,class T9,class T10,class T11,class T12,class T13,class T14,class T15,class T16> \
-		StaticFuncRet FuncName(ExtraDecl FmtType Fmt,T1 V1,T2 V2,T3 V3,T4 V4,T5 V5,T6 V6,T7 V7,T8 V8,T9 V9,T10 V10,T11 V11,T12 V12,T13 V13,T14 V14,T15 V15,T16 V16) {T1 v1=CheckVA(V1);T2 v2=CheckVA(V2);T3 v3=CheckVA(V3);T4 v4=CheckVA(V4);T5 v5=CheckVA(V5);T6 v6=CheckVA(V6);T7 v7=CheckVA(V7);T8 v8=CheckVA(V8);T9 v9=CheckVA(V9);T10 v10=CheckVA(V10);T11 v11=CheckVA(V11);T12 v12=CheckVA(V12);T13 v13=CheckVA(V13);T14 v14=CheckVA(V14);T15 v15=CheckVA(V15);T16 v16=CheckVA(V16);Return FuncName##__VA(ExtraCall (Fmt),(v1),(v2),(v3),(v4),(v5),(v6),(v7),(v8),(v9),(v10),(v11),(v12),(v13),(v14),(v15),(v16));} \
-		template<class T1,class T2,class T3,class T4,class T5,class T6,class T7,class T8,class T9,class T10,class T11,class T12,class T13,class T14,class T15,class T16,class T17> \
-		StaticFuncRet FuncName(ExtraDecl FmtType Fmt,T1 V1,T2 V2,T3 V3,T4 V4,T5 V5,T6 V6,T7 V7,T8 V8,T9 V9,T10 V10,T11 V11,T12 V12,T13 V13,T14 V14,T15 V15,T16 V16,T17 V17) {T1 v1=CheckVA(V1);T2 v2=CheckVA(V2);T3 v3=CheckVA(V3);T4 v4=CheckVA(V4);T5 v5=CheckVA(V5);T6 v6=CheckVA(V6);T7 v7=CheckVA(V7);T8 v8=CheckVA(V8);T9 v9=CheckVA(V9);T10 v10=CheckVA(V10);T11 v11=CheckVA(V11);T12 v12=CheckVA(V12);T13 v13=CheckVA(V13);T14 v14=CheckVA(V14);T15 v15=CheckVA(V15);T16 v16=CheckVA(V16);T17 v17=CheckVA(V17);Return FuncName##__VA(ExtraCall (Fmt),(v1),(v2),(v3),(v4),(v5),(v6),(v7),(v8),(v9),(v10),(v11),(v12),(v13),(v14),(v15),(v16),(v17));} \
-		template<class T1,class T2,class T3,class T4,class T5,class T6,class T7,class T8,class T9,class T10,class T11,class T12,class T13,class T14,class T15,class T16,class T17,class T18> \
-		StaticFuncRet FuncName(ExtraDecl FmtType Fmt,T1 V1,T2 V2,T3 V3,T4 V4,T5 V5,T6 V6,T7 V7,T8 V8,T9 V9,T10 V10,T11 V11,T12 V12,T13 V13,T14 V14,T15 V15,T16 V16,T17 V17,T18 V18) {T1 v1=CheckVA(V1);T2 v2=CheckVA(V2);T3 v3=CheckVA(V3);T4 v4=CheckVA(V4);T5 v5=CheckVA(V5);T6 v6=CheckVA(V6);T7 v7=CheckVA(V7);T8 v8=CheckVA(V8);T9 v9=CheckVA(V9);T10 v10=CheckVA(V10);T11 v11=CheckVA(V11);T12 v12=CheckVA(V12);T13 v13=CheckVA(V13);T14 v14=CheckVA(V14);T15 v15=CheckVA(V15);T16 v16=CheckVA(V16);T17 v17=CheckVA(V17);T18 v18=CheckVA(V18);Return FuncName##__VA(ExtraCall (Fmt),(v1),(v2),(v3),(v4),(v5),(v6),(v7),(v8),(v9),(v10),(v11),(v12),(v13),(v14),(v15),(v16),(v17),(v18));} \
-		template<class T1,class T2,class T3,class T4,class T5,class T6,class T7,class T8,class T9,class T10,class T11,class T12,class T13,class T14,class T15,class T16,class T17,class T18,class T19> \
-		StaticFuncRet FuncName(ExtraDecl FmtType Fmt,T1 V1,T2 V2,T3 V3,T4 V4,T5 V5,T6 V6,T7 V7,T8 V8,T9 V9,T10 V10,T11 V11,T12 V12,T13 V13,T14 V14,T15 V15,T16 V16,T17 V17,T18 V18,T19 V19) {T1 v1=CheckVA(V1);T2 v2=CheckVA(V2);T3 v3=CheckVA(V3);T4 v4=CheckVA(V4);T5 v5=CheckVA(V5);T6 v6=CheckVA(V6);T7 v7=CheckVA(V7);T8 v8=CheckVA(V8);T9 v9=CheckVA(V9);T10 v10=CheckVA(V10);T11 v11=CheckVA(V11);T12 v12=CheckVA(V12);T13 v13=CheckVA(V13);T14 v14=CheckVA(V14);T15 v15=CheckVA(V15);T16 v16=CheckVA(V16);T17 v17=CheckVA(V17);T18 v18=CheckVA(V18);T19 v19=CheckVA(V19);Return FuncName##__VA(ExtraCall (Fmt),(v1),(v2),(v3),(v4),(v5),(v6),(v7),(v8),(v9),(v10),(v11),(v12),(v13),(v14),(v15),(v16),(v17),(v18),(v19));} \
-		template<class T1,class T2,class T3,class T4,class T5,class T6,class T7,class T8,class T9,class T10,class T11,class T12,class T13,class T14,class T15,class T16,class T17,class T18,class T19,class T20> \
-		StaticFuncRet FuncName(ExtraDecl FmtType Fmt,T1 V1,T2 V2,T3 V3,T4 V4,T5 V5,T6 V6,T7 V7,T8 V8,T9 V9,T10 V10,T11 V11,T12 V12,T13 V13,T14 V14,T15 V15,T16 V16,T17 V17,T18 V18,T19 V19,T20 V20) {T1 v1=CheckVA(V1);T2 v2=CheckVA(V2);T3 v3=CheckVA(V3);T4 v4=CheckVA(V4);T5 v5=CheckVA(V5);T6 v6=CheckVA(V6);T7 v7=CheckVA(V7);T8 v8=CheckVA(V8);T9 v9=CheckVA(V9);T10 v10=CheckVA(V10);T11 v11=CheckVA(V11);T12 v12=CheckVA(V12);T13 v13=CheckVA(V13);T14 v14=CheckVA(V14);T15 v15=CheckVA(V15);T16 v16=CheckVA(V16);T17 v17=CheckVA(V17);T18 v18=CheckVA(V18);T19 v19=CheckVA(V19);T20 v20=CheckVA(V20);Return FuncName##__VA(ExtraCall (Fmt),(v1),(v2),(v3),(v4),(v5),(v6),(v7),(v8),(v9),(v10),(v11),(v12),(v13),(v14),(v15),(v16),(v17),(v18),(v19),(v20));} \
-		template<class T1,class T2,class T3,class T4,class T5,class T6,class T7,class T8,class T9,class T10,class T11,class T12,class T13,class T14,class T15,class T16,class T17,class T18,class T19,class T20,class T21> \
-		StaticFuncRet FuncName(ExtraDecl FmtType Fmt,T1 V1,T2 V2,T3 V3,T4 V4,T5 V5,T6 V6,T7 V7,T8 V8,T9 V9,T10 V10,T11 V11,T12 V12,T13 V13,T14 V14,T15 V15,T16 V16,T17 V17,T18 V18,T19 V19,T20 V20, T21 V21) {T1 v1=CheckVA(V1);T2 v2=CheckVA(V2);T3 v3=CheckVA(V3);T4 v4=CheckVA(V4);T5 v5=CheckVA(V5);T6 v6=CheckVA(V6);T7 v7=CheckVA(V7);T8 v8=CheckVA(V8);T9 v9=CheckVA(V9);T10 v10=CheckVA(V10);T11 v11=CheckVA(V11);T12 v12=CheckVA(V12);T13 v13=CheckVA(V13);T14 v14=CheckVA(V14);T15 v15=CheckVA(V15);T16 v16=CheckVA(V16);T17 v17=CheckVA(V17);T18 v18=CheckVA(V18);T19 v19=CheckVA(V19);T20 v20=CheckVA(V20);T21 v21=CheckVA(V21);Return FuncName##__VA(ExtraCall (Fmt),(v1),(v2),(v3),(v4),(v5),(v6),(v7),(v8),(v9),(v10),(v11),(v12),(v13),(v14),(v15),(v16),(v17),(v18),(v19),(v20),(v21));} \
-		template<class T1,class T2,class T3,class T4,class T5,class T6,class T7,class T8,class T9,class T10,class T11,class T12,class T13,class T14,class T15,class T16,class T17,class T18,class T19,class T20,class T21,class T22> \
-		StaticFuncRet FuncName(ExtraDecl FmtType Fmt,T1 V1,T2 V2,T3 V3,T4 V4,T5 V5,T6 V6,T7 V7,T8 V8,T9 V9,T10 V10,T11 V11,T12 V12,T13 V13,T14 V14,T15 V15,T16 V16,T17 V17,T18 V18,T19 V19,T20 V20, T21 V21, T22 V22) {T1 v1=CheckVA(V1);T2 v2=CheckVA(V2);T3 v3=CheckVA(V3);T4 v4=CheckVA(V4);T5 v5=CheckVA(V5);T6 v6=CheckVA(V6);T7 v7=CheckVA(V7);T8 v8=CheckVA(V8);T9 v9=CheckVA(V9);T10 v10=CheckVA(V10);T11 v11=CheckVA(V11);T12 v12=CheckVA(V12);T13 v13=CheckVA(V13);T14 v14=CheckVA(V14);T15 v15=CheckVA(V15);T16 v16=CheckVA(V16);T17 v17=CheckVA(V17);T18 v18=CheckVA(V18);T19 v19=CheckVA(V19);T20 v20=CheckVA(V20);T21 v21=CheckVA(V21);T22 v22=CheckVA(V22);Return FuncName##__VA(ExtraCall (Fmt),(v1),(v2),(v3),(v4),(v5),(v6),(v7),(v8),(v9),(v10),(v11),(v12),(v13),(v14),(v15),(v16),(v17),(v18),(v19),(v20),(v21),(v22));} \
-		template<class T1,class T2,class T3,class T4,class T5,class T6,class T7,class T8,class T9,class T10,class T11,class T12,class T13,class T14,class T15,class T16,class T17,class T18,class T19,class T20,class T21,class T22,class T23> \
-		StaticFuncRet FuncName(ExtraDecl FmtType Fmt,T1 V1,T2 V2,T3 V3,T4 V4,T5 V5,T6 V6,T7 V7,T8 V8,T9 V9,T10 V10,T11 V11,T12 V12,T13 V13,T14 V14,T15 V15,T16 V16,T17 V17,T18 V18,T19 V19,T20 V20, T21 V21, T22 V22, T23 V23) {T1 v1=CheckVA(V1);T2 v2=CheckVA(V2);T3 v3=CheckVA(V3);T4 v4=CheckVA(V4);T5 v5=CheckVA(V5);T6 v6=CheckVA(V6);T7 v7=CheckVA(V7);T8 v8=CheckVA(V8);T9 v9=CheckVA(V9);T10 v10=CheckVA(V10);T11 v11=CheckVA(V11);T12 v12=CheckVA(V12);T13 v13=CheckVA(V13);T14 v14=CheckVA(V14);T15 v15=CheckVA(V15);T16 v16=CheckVA(V16);T17 v17=CheckVA(V17);T18 v18=CheckVA(V18);T19 v19=CheckVA(V19);T20 v20=CheckVA(V20);T21 v21=CheckVA(V21);T22 v22=CheckVA(V22);T23 v23=CheckVA(V23);Return FuncName##__VA(ExtraCall (Fmt),(v1),(v2),(v3),(v4),(v5),(v6),(v7),(v8),(v9),(v10),(v11),(v12),(v13),(v14),(v15),(v16),(v17),(v18),(v19),(v20),(v21),(v22),(v23));} \
-		template<class T1,class T2,class T3,class T4,class T5,class T6,class T7,class T8,class T9,class T10,class T11,class T12,class T13,class T14,class T15,class T16,class T17,class T18,class T19,class T20,class T21,class T22,class T23,class T24> \
-		StaticFuncRet FuncName(ExtraDecl FmtType Fmt,T1 V1,T2 V2,T3 V3,T4 V4,T5 V5,T6 V6,T7 V7,T8 V8,T9 V9,T10 V10,T11 V11,T12 V12,T13 V13,T14 V14,T15 V15,T16 V16,T17 V17,T18 V18,T19 V19,T20 V20, T21 V21, T22 V22, T23 V23, T24 V24) {T1 v1=CheckVA(V1);T2 v2=CheckVA(V2);T3 v3=CheckVA(V3);T4 v4=CheckVA(V4);T5 v5=CheckVA(V5);T6 v6=CheckVA(V6);T7 v7=CheckVA(V7);T8 v8=CheckVA(V8);T9 v9=CheckVA(V9);T10 v10=CheckVA(V10);T11 v11=CheckVA(V11);T12 v12=CheckVA(V12);T13 v13=CheckVA(V13);T14 v14=CheckVA(V14);T15 v15=CheckVA(V15);T16 v16=CheckVA(V16);T17 v17=CheckVA(V17);T18 v18=CheckVA(V18);T19 v19=CheckVA(V19);T20 v20=CheckVA(V20);T21 v21=CheckVA(V21);T22 v22=CheckVA(V22);T23 v23=CheckVA(V23);T24 v24=CheckVA(V24);Return FuncName##__VA(ExtraCall (Fmt),(v1),(v2),(v3),(v4),(v5),(v6),(v7),(v8),(v9),(v10),(v11),(v12),(v13),(v14),(v15),(v16),(v17),(v18),(v19),(v20),(v21),(v22),(v23),(v24));} \
-		template<class T1,class T2,class T3,class T4,class T5,class T6,class T7,class T8,class T9,class T10,class T11,class T12,class T13,class T14,class T15,class T16,class T17,class T18,class T19,class T20,class T21,class T22,class T23,class T24,class T25> \
-		StaticFuncRet FuncName(ExtraDecl FmtType Fmt,T1 V1,T2 V2,T3 V3,T4 V4,T5 V5,T6 V6,T7 V7,T8 V8,T9 V9,T10 V10,T11 V11,T12 V12,T13 V13,T14 V14,T15 V15,T16 V16,T17 V17,T18 V18,T19 V19,T20 V20, T21 V21, T22 V22, T23 V23, T24 V24, T25 V25) {T1 v1=CheckVA(V1);T2 v2=CheckVA(V2);T3 v3=CheckVA(V3);T4 v4=CheckVA(V4);T5 v5=CheckVA(V5);T6 v6=CheckVA(V6);T7 v7=CheckVA(V7);T8 v8=CheckVA(V8);T9 v9=CheckVA(V9);T10 v10=CheckVA(V10);T11 v11=CheckVA(V11);T12 v12=CheckVA(V12);T13 v13=CheckVA(V13);T14 v14=CheckVA(V14);T15 v15=CheckVA(V15);T16 v16=CheckVA(V16);T17 v17=CheckVA(V17);T18 v18=CheckVA(V18);T19 v19=CheckVA(V19);T20 v20=CheckVA(V20);T21 v21=CheckVA(V21);T22 v22=CheckVA(V22);T23 v23=CheckVA(V23);T24 v24=CheckVA(V24);T25 v25=CheckVA(V25);Return FuncName##__VA(ExtraCall (Fmt),(v1),(v2),(v3),(v4),(v5),(v6),(v7),(v8),(v9),(v10),(v11),(v12),(v13),(v14),(v15),(v16),(v17),(v18),(v19),(v20),(v21),(v22),(v23),(v24),(v25));} \
-		template<class T1,class T2,class T3,class T4,class T5,class T6,class T7,class T8,class T9,class T10,class T11,class T12,class T13,class T14,class T15,class T16,class T17,class T18,class T19,class T20,class T21,class T22,class T23,class T24,class T25,class T26> \
-		StaticFuncRet FuncName(ExtraDecl FmtType Fmt,T1 V1,T2 V2,T3 V3,T4 V4,T5 V5,T6 V6,T7 V7,T8 V8,T9 V9,T10 V10,T11 V11,T12 V12,T13 V13,T14 V14,T15 V15,T16 V16,T17 V17,T18 V18,T19 V19,T20 V20, T21 V21, T22 V22, T23 V23, T24 V24, T25 V25, T26 V26) {T1 v1=CheckVA(V1);T2 v2=CheckVA(V2);T3 v3=CheckVA(V3);T4 v4=CheckVA(V4);T5 v5=CheckVA(V5);T6 v6=CheckVA(V6);T7 v7=CheckVA(V7);T8 v8=CheckVA(V8);T9 v9=CheckVA(V9);T10 v10=CheckVA(V10);T11 v11=CheckVA(V11);T12 v12=CheckVA(V12);T13 v13=CheckVA(V13);T14 v14=CheckVA(V14);T15 v15=CheckVA(V15);T16 v16=CheckVA(V16);T17 v17=CheckVA(V17);T18 v18=CheckVA(V18);T19 v19=CheckVA(V19);T20 v20=CheckVA(V20);T21 v21=CheckVA(V21);T22 v22=CheckVA(V22);T23 v23=CheckVA(V23);T24 v24=CheckVA(V24);T25 v25=CheckVA(V25);T26 v26=CheckVA(V26);Return FuncName##__VA(ExtraCall (Fmt),(v1),(v2),(v3),(v4),(v5),(v6),(v7),(v8),(v9),(v10),(v11),(v12),(v13),(v14),(v15),(v16),(v17),(v18),(v19),(v20),(v21),(v22),(v23),(v24),(v25),(v26));}
-
-	#define VARARG_BODY( FuncRet, FuncName, FmtType, ExtraDecl )		\
-		DEPRECATED_MACRO(4.19, "VARARG_BODY has been deprecated - please use a variadic template instead (see FString::PrintfImpl as an example)") \
-		FuncRet FuncName##__VA( ExtraDecl  FmtType Fmt, ... )
-
-#else  // !PLATFORM_WINDOWS
-
-	#define VARARG_DECL( FuncRet, StaticFuncRet, Return, FuncName, Pure, FmtType, ExtraDecl, ExtraCall )	\
-		DEPRECATED_MACRO(4.19, "VARARG_DECL has been deprecated - please use a variadic template instead (see FString::Printf as an example)") \
-		FuncRet FuncName( ExtraDecl FmtType Fmt, ... ) Pure
-
-	#define VARARG_BODY( FuncRet, FuncName, FmtType, ExtraDecl )		\
-		DEPRECATED_MACRO(4.19, "VARARG_BODY has been deprecated - please use a variadic template instead (see FString::PrintfImpl as an example)") \
-		FuncRet FuncName( ExtraDecl FmtType Fmt, ... )
-
-#endif // PLATFORM_WINDOWS
diff --git a/Engine/Source/Runtime/Core/Public/Templates/AlignOf.h b/Engine/Source/Runtime/Core/Public/Templates/AlignOf.h
deleted file mode 100644
diff --git a/Engine/Source/Runtime/Core/Public/Templates/UnrealTemplate.h b/Engine/Source/Runtime/Core/Public/Templates/UnrealTemplate.h
--- a/Engine/Source/Runtime/Core/Public/Templates/UnrealTemplate.h
+++ b/Engine/Source/Runtime/Core/Public/Templates/UnrealTemplate.h
@@ -198,7 +198,6 @@ OutType IntCastChecked(InType In)
 
 // Number of elements in an array.
 #define UE_ARRAY_COUNT( array ) (sizeof(UEArrayCountHelper(array)) - 1)
-#define ARRAY_COUNT( array ) DEPRECATED_MACRO(4.24, "The ARRAY_COUNT macro has been deprecated in favor of UE_ARRAY_COUNT.") UE_ARRAY_COUNT( array )
 
 // Offset of a struct member.
 #ifndef UNREAL_CODE_ANALYZER
diff --git a/Engine/Source/Runtime/Core/Public/Windows/WindowsPlatform.h b/Engine/Source/Runtime/Core/Public/Windows/WindowsPlatform.h
--- a/Engine/Source/Runtime/Core/Public/Windows/WindowsPlatform.h
+++ b/Engine/Source/Runtime/Core/Public/Windows/WindowsPlatform.h
@@ -128,9 +128,6 @@ typedef FWindowsPlatformTypes FPlatformTypes;
 #define STDCALL		__stdcall										/* Standard calling convention */
 #define FORCEINLINE __forceinline									/* Force code to be inline */
 #define FORCENOINLINE __declspec(noinline)							/* Force code to NOT be inline */
-#define FUNCTION_NO_RETURN_START \
-	DEPRECATED_MACRO(4.26, "FUNCTION_NO_RETURN_START has been deprecated - please use UE_NORETURN") \
-	__declspec(noreturn)				/* Indicate that the function never returns. */
 #define FUNCTION_NON_NULL_RETURN_START _Ret_notnull_				/* Indicate that the function never returns nullptr. */
 
 #define DECLARE_UINT64(x)	x
@@ -199,23 +196,6 @@ typedef FWindowsPlatformTypes FPlatformTypes;
 // Include code analysis features
 #include "Windows/WindowsPlatformCodeAnalysis.h"
 
-#if USING_CODE_ANALYSIS && _MSC_VER == 1900
-	// Disable this warning as VC2015 Update 1 produces this warning erroneously when placed on variadic templates:
-	//
-	// warning C28216: The checkReturn annotation only applies to postconditions for function 'Func' _Param_(N).
-	#define FUNCTION_CHECK_RETURN_START	 \
-		DEPRECATED_MACRO(4.26, "FUNCTION_CHECK_RETURN_START has been deprecated - please use UE_NODISCARD") \
-		__pragma(warning(push)) __pragma(warning(disable: 28216)) __declspec("SAL_checkReturn")
-
-	#define FUNCTION_CHECK_RETURN_END	 \
-		DEPRECATED_MACRO(4.26, "FUNCTION_CHECK_RETURN_END has been deprecated - please use UE_NODISCARD") \
-		__pragma(warning(pop))
-#else
-	#define FUNCTION_CHECK_RETURN_START	 \
-		DEPRECATED_MACRO(4.26, "FUNCTION_CHECK_RETURN_START has been deprecated - please use UE_NODISCARD") \
-		__declspec("SAL_checkReturn")	/* Warn that callers should not ignore the return value. */
-#endif
-
 // Other macros
 #ifndef ENABLE_WIN_ALLOC_TRACKING
 #define ENABLE_WIN_ALLOC_TRACKING 0
diff --git a/Engine/Source/Runtime/Core/Public/Windows/WindowsPlatformCompilerPreSetup.h b/Engine/Source/Runtime/Core/Public/Windows/WindowsPlatformCompilerPreSetup.h
--- a/Engine/Source/Runtime/Core/Public/Windows/WindowsPlatformCompilerPreSetup.h
+++ b/Engine/Source/Runtime/Core/Public/Windows/WindowsPlatformCompilerPreSetup.h
@@ -52,8 +52,6 @@
 	#endif // THIRD_PARTY_INCLUDES_END
 #else
 	#ifndef DISABLE_DEPRECATION
-		#define DEPRECATED(VERSION, MESSAGE) DEPRECATED_MACRO(4.22, "The DEPRECATED macro has been deprecated in favor of UE_DEPRECATED().") __declspec(deprecated(MESSAGE " Please update your code to the new API before upgrading to the next release, otherwise your project will no longer compile."))
-
 		#define PRAGMA_DISABLE_DEPRECATION_WARNINGS \
 			__pragma (warning(push)) \
 			__pragma (warning(disable: 4995)) /* 'function': name was marked as #pragma deprecated */ \
@@ -87,7 +85,7 @@
 
 	#ifndef PRAGMA_ENABLE_UNSAFE_TYPECAST_WARNINGS
 		#define PRAGMA_ENABLE_UNSAFE_TYPECAST_WARNINGS \
-			DEPRECATED_MACRO(5.0, "The PRAGMA_ENABLE_UNSAFE_TYPECAST_WARNINGS macro has been deprecated in favor of PRAGMA_RESTORE_UNSAFE_TYPECAST_WARNINGS. To force enable warnings use PRAGMA_FORCE_UNSAFE_TYPECAST_WARNINGS.")
+			UE_DEPRECATED_MACRO(5.0, "The PRAGMA_ENABLE_UNSAFE_TYPECAST_WARNINGS macro has been deprecated in favor of PRAGMA_RESTORE_UNSAFE_TYPECAST_WARNINGS. To force enable warnings use PRAGMA_FORCE_UNSAFE_TYPECAST_WARNINGS.")
 	#endif // PRAGMA_ENABLE_UNSAFE_TYPECAST_WARNINGS
 
 	#ifndef PRAGMA_FORCE_UNSAFE_TYPECAST_WARNINGS
diff --git a/Engine/Source/Runtime/CoreUObject/Private/UObject/Class.cpp b/Engine/Source/Runtime/CoreUObject/Private/UObject/Class.cpp
--- a/Engine/Source/Runtime/CoreUObject/Private/UObject/Class.cpp
+++ b/Engine/Source/Runtime/CoreUObject/Private/UObject/Class.cpp
@@ -52,9 +52,6 @@
 #include "Math/InterpCurvePoint.h"
 #include "UObject/UE5MainStreamObjectVersion.h"
 
-// WARNING: This should always be the last include in any file that needs it (except .generated.h)
-#include "UObject/UndefineUPropertyMacros.h"
-
 // This flag enables some expensive class tree validation that is meant to catch mutations of 
 // the class tree outside of SetSuperStruct. It has been disabled because loading blueprints 
 // does a lot of mutation of the class tree, and the validation checks impact iteration time.
@@ -6472,5 +6469,3 @@ PRAGMA_ENABLE_DEPRECATION_WARNINGS
 		PRAGMA_ENABLE_SHADOW_VARIABLE_WARNINGS
 	#endif
 #endif
-
-#include "UObject/DefineUPropertyMacros.h"
diff --git a/Engine/Source/Runtime/CoreUObject/Private/UObject/EnumProperty.cpp b/Engine/Source/Runtime/CoreUObject/Private/UObject/EnumProperty.cpp
--- a/Engine/Source/Runtime/CoreUObject/Private/UObject/EnumProperty.cpp
+++ b/Engine/Source/Runtime/CoreUObject/Private/UObject/EnumProperty.cpp
@@ -12,9 +12,6 @@
 #include "Misc/NetworkVersion.h"
 #include "Hash/Blake3.h"
 
-// WARNING: This should always be the last include in any file that needs it (except .generated.h)
-#include "UObject/UndefineUPropertyMacros.h"
-
 namespace UEEnumProperty_Private
 {
 	template <typename OldIntType>
@@ -556,5 +553,3 @@ uint64 FEnumProperty::GetMaxNetSerializeBits() const
 	
 	return FMath::Min(DesiredBits, MaxBits);
 }
-
-#include "UObject/DefineUPropertyMacros.h"
\ No newline at end of file
diff --git a/Engine/Source/Runtime/CoreUObject/Private/UObject/Field.cpp b/Engine/Source/Runtime/CoreUObject/Private/UObject/Field.cpp
--- a/Engine/Source/Runtime/CoreUObject/Private/UObject/Field.cpp
+++ b/Engine/Source/Runtime/CoreUObject/Private/UObject/Field.cpp
@@ -44,9 +44,6 @@ Field.cpp: Defines FField property system fundamentals
 #include "UObject/TextProperty.h"
 #include "UObject/EnumProperty.h"
 
-// WARNING: This should always be the last include in any file that needs it (except .generated.h)
-#include "UObject/UndefineUPropertyMacros.h"
-
 FFieldClass::FFieldClass(const TCHAR* InCPPName, uint64 InId, uint64 InCastFlags, FFieldClass* InSuperClass, FField* (*ConstructFnPtr)(const FFieldVariant&, const FName&, EObjectFlags))
 	: Id(InId)
 	, CastFlags(InCastFlags)
@@ -1235,5 +1232,3 @@ bool FFieldCastTest::RunTest(const FString& Parameters)
 }
 
 #endif
-
-#include "UObject/DefineUPropertyMacros.h"
\ No newline at end of file
diff --git a/Engine/Source/Runtime/CoreUObject/Private/UObject/FieldPathProperty.cpp b/Engine/Source/Runtime/CoreUObject/Private/UObject/FieldPathProperty.cpp
--- a/Engine/Source/Runtime/CoreUObject/Private/UObject/FieldPathProperty.cpp
+++ b/Engine/Source/Runtime/CoreUObject/Private/UObject/FieldPathProperty.cpp
@@ -9,9 +9,6 @@
 #include "Misc/Parse.h"
 #include "UObject/PropertyHelper.h"
 
-// WARNING: This should always be the last include in any file that needs it (except .generated.h)
-#include "UObject/UndefineUPropertyMacros.h"
-
 IMPLEMENT_FIELD(FFieldPathProperty)
 
 FFieldPathProperty::FFieldPathProperty(FFieldVariant InOwner, const UECodeGen_Private::FFieldPathPropertyParams& Prop)
@@ -216,5 +213,3 @@ bool FFieldPathProperty::SupportsNetSharedSerialization() const
 {
 	return false;
 }
-
-#include "UObject/DefineUPropertyMacros.h"
\ No newline at end of file
diff --git a/Engine/Source/Runtime/CoreUObject/Private/UObject/LinkerPlaceholderBase.cpp b/Engine/Source/Runtime/CoreUObject/Private/UObject/LinkerPlaceholderBase.cpp
--- a/Engine/Source/Runtime/CoreUObject/Private/UObject/LinkerPlaceholderBase.cpp
+++ b/Engine/Source/Runtime/CoreUObject/Private/UObject/LinkerPlaceholderBase.cpp
@@ -7,9 +7,6 @@
 #include "UObject/UnrealTypePrivate.h"
 #include "Blueprint/BlueprintSupport.h"
 
-// WARNING: This should always be the last include in any file that needs it (except .generated.h)
-#include "UObject/UndefineUPropertyMacros.h"
-
 #if USE_DEFERRED_DEPENDENCY_CHECK_VERIFICATION_TESTS
 	#define DEFERRED_DEPENDENCY_ENSURE(EnsueExpr) ensure(EnsueExpr)
 #else  // USE_DEFERRED_DEPENDENCY_CHECK_VERIFICATION_TESTS
@@ -785,5 +782,3 @@ int32 TLinkerImportPlaceholder<UFunction>::ResolvePropertyReferences(UFunction*
 }
 
 #undef DEFERRED_DEPENDENCY_ENSURE
-
-#include "UObject/DefineUPropertyMacros.h"
\ No newline at end of file
diff --git a/Engine/Source/Runtime/CoreUObject/Private/UObject/Property.cpp b/Engine/Source/Runtime/CoreUObject/Private/UObject/Property.cpp
--- a/Engine/Source/Runtime/CoreUObject/Private/UObject/Property.cpp
+++ b/Engine/Source/Runtime/CoreUObject/Private/UObject/Property.cpp
@@ -24,9 +24,6 @@
 #include "Math/InterpCurvePoint.h"
 #include "UObject/ReleaseObjectVersion.h"
 
-// WARNING: This should always be the last include in any file that needs it (except .generated.h)
-#include "UObject/UndefineUPropertyMacros.h"
-
 DEFINE_LOG_CATEGORY(LogProperty);
 
 // List the core ones here as they have already been included (and can be used without CoreUObject!)
@@ -1933,6 +1930,3 @@ FProperty* UStruct::FindPropertyByName(FName InName) const
 
 	return NULL;
 }
-
-
-#include "UObject/DefineUPropertyMacros.h"
diff --git a/Engine/Source/Runtime/CoreUObject/Private/UObject/PropertyArray.cpp b/Engine/Source/Runtime/CoreUObject/Private/UObject/PropertyArray.cpp
--- a/Engine/Source/Runtime/CoreUObject/Private/UObject/PropertyArray.cpp
+++ b/Engine/Source/Runtime/CoreUObject/Private/UObject/PropertyArray.cpp
@@ -9,9 +9,6 @@
 #include "UObject/LinkerLoad.h"
 #include "UObject/PropertyHelper.h"
 
-// WARNING: This should always be the last include in any file that needs it (except .generated.h)
-#include "UObject/UndefineUPropertyMacros.h"
-
 /*-----------------------------------------------------------------------------
 	FArrayProperty.
 -----------------------------------------------------------------------------*/
@@ -891,5 +888,3 @@ void* FArrayProperty::GetValueAddressAtIndex_Direct(const FProperty* InInner, vo
 
 	return ArrayHelper.GetRawPtr(Index);
 }
-
-#include "UObject/DefineUPropertyMacros.h"
\ No newline at end of file
diff --git a/Engine/Source/Runtime/CoreUObject/Private/UObject/PropertyBaseObject.cpp b/Engine/Source/Runtime/CoreUObject/Private/UObject/PropertyBaseObject.cpp
--- a/Engine/Source/Runtime/CoreUObject/Private/UObject/PropertyBaseObject.cpp
+++ b/Engine/Source/Runtime/CoreUObject/Private/UObject/PropertyBaseObject.cpp
@@ -12,9 +12,6 @@
 #include "UObject/LinkerPlaceholderExportObject.h"
 #include "Misc/StringBuilder.h"
 
-// WARNING: This should always be the last include in any file that needs it (except .generated.h)
-#include "UObject/UndefineUPropertyMacros.h"
-
 /*-----------------------------------------------------------------------------
 	FObjectPropertyBase.
 -----------------------------------------------------------------------------*/
@@ -669,5 +666,3 @@ void FObjectPropertyBase::CopyCompleteValueFromScriptVM( void* Dest, void const*
 		SetObjectPropertyValue(((uint8*)Dest) + Index * ElementSize, ((UObject**)Src)[Index]);
 	}
 }
-
-#include "UObject/DefineUPropertyMacros.h"
\ No newline at end of file
diff --git a/Engine/Source/Runtime/CoreUObject/Private/UObject/PropertyBool.cpp b/Engine/Source/Runtime/CoreUObject/Private/UObject/PropertyBool.cpp
--- a/Engine/Source/Runtime/CoreUObject/Private/UObject/PropertyBool.cpp
+++ b/Engine/Source/Runtime/CoreUObject/Private/UObject/PropertyBool.cpp
@@ -9,9 +9,6 @@
 #include "UObject/UnrealTypePrivate.h"
 #include "UObject/PropertyHelper.h"
 
-// WARNING: This should always be the last include in any file that needs it (except .generated.h)
-#include "UObject/UndefineUPropertyMacros.h"
-
 /*-----------------------------------------------------------------------------
 	FBoolProperty.
 -----------------------------------------------------------------------------*/
@@ -474,5 +471,3 @@ uint32 FBoolProperty::GetValueTypeHashInternal(const void* Src) const
 	uint8* SrcByteValue = (uint8*)Src + ByteOffset;
 	return GetTypeHash(*SrcByteValue & FieldMask);
 }
-
-#include "UObject/DefineUPropertyMacros.h"
\ No newline at end of file
diff --git a/Engine/Source/Runtime/CoreUObject/Private/UObject/PropertyByte.cpp b/Engine/Source/Runtime/CoreUObject/Private/UObject/PropertyByte.cpp
--- a/Engine/Source/Runtime/CoreUObject/Private/UObject/PropertyByte.cpp
+++ b/Engine/Source/Runtime/CoreUObject/Private/UObject/PropertyByte.cpp
@@ -13,9 +13,6 @@
 #include "Misc/NetworkVersion.h"
 #include "Hash/Blake3.h"
 
-// WARNING: This should always be the last include in any file that needs it (except .generated.h)
-#include "UObject/UndefineUPropertyMacros.h"
-
 /*-----------------------------------------------------------------------------
 	FByteProperty.
 -----------------------------------------------------------------------------*/
@@ -529,5 +526,3 @@ uint64 FByteProperty::GetMaxNetSerializeBits() const
 
 	return FMath::Min(DesiredBits, MaxBits);
 }
-
-#include "UObject/DefineUPropertyMacros.h"
\ No newline at end of file
diff --git a/Engine/Source/Runtime/CoreUObject/Private/UObject/PropertyClass.cpp b/Engine/Source/Runtime/CoreUObject/Private/UObject/PropertyClass.cpp
--- a/Engine/Source/Runtime/CoreUObject/Private/UObject/PropertyClass.cpp
+++ b/Engine/Source/Runtime/CoreUObject/Private/UObject/PropertyClass.cpp
@@ -10,9 +10,6 @@
 #include "UObject/PropertyHelper.h"
 #include "Hash/Blake3.h"
 
-// WARNING: This should always be the last include in any file that needs it (except .generated.h)
-#include "UObject/UndefineUPropertyMacros.h"
-
 /*-----------------------------------------------------------------------------
 	FClassProperty.
 -----------------------------------------------------------------------------*/
@@ -211,4 +208,3 @@ void FClassProperty::AppendSchemaHash(FBlake3& Builder, bool bSkipEditorOnly) co
 	}
 }
 #endif
-#include "UObject/DefineUPropertyMacros.h"
\ No newline at end of file
diff --git a/Engine/Source/Runtime/CoreUObject/Private/UObject/PropertyDelegate.cpp b/Engine/Source/Runtime/CoreUObject/Private/UObject/PropertyDelegate.cpp
--- a/Engine/Source/Runtime/CoreUObject/Private/UObject/PropertyDelegate.cpp
+++ b/Engine/Source/Runtime/CoreUObject/Private/UObject/PropertyDelegate.cpp
@@ -10,9 +10,6 @@
 #include "Serialization/ArchiveUObjectFromStructuredArchive.h"
 #include "Hash/Blake3.h"
 
-// WARNING: This should always be the last include in any file that needs it (except .generated.h)
-#include "UObject/UndefineUPropertyMacros.h"
-
 /*-----------------------------------------------------------------------------
 	FDelegateProperty.
 -----------------------------------------------------------------------------*/
@@ -250,5 +247,3 @@ void FDelegateProperty::BeginDestroy()
 
 	Super::BeginDestroy();
 }
-
-#include "UObject/DefineUPropertyMacros.h"
diff --git a/Engine/Source/Runtime/CoreUObject/Private/UObject/PropertyInterface.cpp b/Engine/Source/Runtime/CoreUObject/Private/UObject/PropertyInterface.cpp
--- a/Engine/Source/Runtime/CoreUObject/Private/UObject/PropertyInterface.cpp
+++ b/Engine/Source/Runtime/CoreUObject/Private/UObject/PropertyInterface.cpp
@@ -13,9 +13,6 @@
 #include "UObject/CoreNet.h"
 #include "Misc/NetworkVersion.h"
 
-// WARNING: This should always be the last include in any file that needs it (except .generated.h)
-#include "UObject/UndefineUPropertyMacros.h"
-
 /*-----------------------------------------------------------------------------
 	FInterfaceProperty.
 -----------------------------------------------------------------------------*/
@@ -377,5 +374,3 @@ void FInterfaceProperty::AddReferencedObjects(FReferenceCollector& Collector)
 	Collector.AddReferencedObject(InterfaceClass);
 	Super::AddReferencedObjects(Collector);
 }
-
-#include "UObject/DefineUPropertyMacros.h"
\ No newline at end of file
diff --git a/Engine/Source/Runtime/CoreUObject/Private/UObject/PropertyMap.cpp b/Engine/Source/Runtime/CoreUObject/Private/UObject/PropertyMap.cpp
--- a/Engine/Source/Runtime/CoreUObject/Private/UObject/PropertyMap.cpp
+++ b/Engine/Source/Runtime/CoreUObject/Private/UObject/PropertyMap.cpp
@@ -12,9 +12,6 @@
 #include "Serialization/ArchiveUObjectFromStructuredArchive.h"
 #include "UObject/UObjectThreadContext.h"
 
-// WARNING: This should always be the last include in any file that needs it (except .generated.h)
-#include "UObject/UndefineUPropertyMacros.h"
-
 namespace UEMapProperty_Private
 {
 	/**
@@ -1349,5 +1346,3 @@ void* FMapProperty::GetValueAddressAtIndex_Direct(const FProperty* Inner, void*
 	checkf(false, TEXT("Map element index (%d) out of range"), Index);
 	return nullptr;
 }
-
-#include "UObject/DefineUPropertyMacros.h"
diff --git a/Engine/Source/Runtime/CoreUObject/Private/UObject/PropertyMulticastDelegate.cpp b/Engine/Source/Runtime/CoreUObject/Private/UObject/PropertyMulticastDelegate.cpp
--- a/Engine/Source/Runtime/CoreUObject/Private/UObject/PropertyMulticastDelegate.cpp
+++ b/Engine/Source/Runtime/CoreUObject/Private/UObject/PropertyMulticastDelegate.cpp
@@ -9,9 +9,6 @@
 #include "UObject/LinkerPlaceholderFunction.h"
 #include "Serialization/ArchiveUObjectFromStructuredArchive.h"
 
-// WARNING: This should always be the last include in any file that needs it (except .generated.h)
-#include "UObject/UndefineUPropertyMacros.h"
-
 FMulticastScriptDelegate::FInvocationList FMulticastDelegateProperty::EmptyList;
 
 FMulticastDelegateProperty::FMulticastDelegateProperty(FFieldVariant InOwner, const UECodeGen_Private::FMulticastDelegatePropertyParams& Prop, EPropertyFlags AdditionalPropertyFlags /*= CPF_None*/)
@@ -660,5 +657,3 @@ void FMulticastSparseDelegateProperty::ClearDelegate(UObject* Parent, void* Prop
 }
 
 IMPLEMENT_FIELD(FMulticastSparseDelegateProperty)
-
-#include "UObject/DefineUPropertyMacros.h"
\ No newline at end of file
diff --git a/Engine/Source/Runtime/CoreUObject/Private/UObject/PropertyPrivate.cpp b/Engine/Source/Runtime/CoreUObject/Private/UObject/PropertyPrivate.cpp
--- a/Engine/Source/Runtime/CoreUObject/Private/UObject/PropertyPrivate.cpp
+++ b/Engine/Source/Runtime/CoreUObject/Private/UObject/PropertyPrivate.cpp
@@ -22,9 +22,6 @@
 #include "UObject/LinkerPlaceholderFunction.h"
 #include "UObject/LinkerPlaceholderClass.h"
 
-// WARNING: This should always be the last include in any file that needs it (except .generated.h)
-#include "UObject/UndefineUPropertyMacros.h"
-
 /*-----------------------------------------------------------------------------
 	UProperty implementation.
 -----------------------------------------------------------------------------*/
@@ -925,5 +922,3 @@ IMPLEMENT_CORE_INTRINSIC_CLASS(UMulticastInlineDelegatePropertyWrapper, UMultica
 	{
 	}
 );
-
-#include "UObject/DefineUPropertyMacros.h"
\ No newline at end of file
diff --git a/Engine/Source/Runtime/CoreUObject/Private/UObject/PropertySet.cpp b/Engine/Source/Runtime/CoreUObject/Private/UObject/PropertySet.cpp
--- a/Engine/Source/Runtime/CoreUObject/Private/UObject/PropertySet.cpp
+++ b/Engine/Source/Runtime/CoreUObject/Private/UObject/PropertySet.cpp
@@ -10,9 +10,6 @@
 #include "Misc/ScopeExit.h"
 #include "Serialization/ArchiveUObjectFromStructuredArchive.h"
 
-// WARNING: This should always be the last include in any file that needs it (except .generated.h)
-#include "UObject/UndefineUPropertyMacros.h"
-
 namespace UESetProperty_Private
 {
 	/**
@@ -1120,5 +1117,3 @@ void* FSetProperty::GetValueAddressAtIndex_Direct(const FProperty* Inner, void*
 	checkf(false, TEXT("Set element index (%d) out of range"), Index);
 	return nullptr;
 }
-
-#include "UObject/DefineUPropertyMacros.h"
\ No newline at end of file
diff --git a/Engine/Source/Runtime/CoreUObject/Private/UObject/PropertySoftClassPtr.cpp b/Engine/Source/Runtime/CoreUObject/Private/UObject/PropertySoftClassPtr.cpp
--- a/Engine/Source/Runtime/CoreUObject/Private/UObject/PropertySoftClassPtr.cpp
+++ b/Engine/Source/Runtime/CoreUObject/Private/UObject/PropertySoftClassPtr.cpp
@@ -7,9 +7,6 @@
 #include "UObject/UnrealTypePrivate.h"
 #include "UObject/LinkerPlaceholderClass.h"
 
-// WARNING: This should always be the last include in any file that needs it (except .generated.h)
-#include "UObject/UndefineUPropertyMacros.h"
-
 /*-----------------------------------------------------------------------------
 	FSoftClassProperty.
 -----------------------------------------------------------------------------*/
@@ -125,5 +122,3 @@ bool FSoftClassProperty::SameType(const FProperty* Other) const
 {
 	return Super::SameType(Other) && (MetaClass == ((FSoftClassProperty*)Other)->MetaClass);
 }
-
-#include "UObject/DefineUPropertyMacros.h"
\ No newline at end of file
diff --git a/Engine/Source/Runtime/CoreUObject/Private/UObject/PropertyStruct.cpp b/Engine/Source/Runtime/CoreUObject/Private/UObject/PropertyStruct.cpp
--- a/Engine/Source/Runtime/CoreUObject/Private/UObject/PropertyStruct.cpp
+++ b/Engine/Source/Runtime/CoreUObject/Private/UObject/PropertyStruct.cpp
@@ -12,9 +12,6 @@
 #include "Hash/Blake3.h"
 #include "IO/IoHash.h"
 
-// WARNING: This should always be the last include in any file that needs it (except .generated.h)
-#include "UObject/UndefineUPropertyMacros.h"
-
 static inline void PreloadInnerStructMembers(FStructProperty* StructProperty)
 {
 	if (UseCircularDependencyLoadDeferring())
@@ -478,6 +475,3 @@ void FStructProperty::AppendSchemaHash(FBlake3& Builder, bool bSkipEditorOnly) c
 	}
 }
 #endif
-
-
-#include "UObject/DefineUPropertyMacros.h"
\ No newline at end of file
diff --git a/Engine/Source/Runtime/CoreUObject/Public/Serialization/BulkData.h b/Engine/Source/Runtime/CoreUObject/Public/Serialization/BulkData.h
--- a/Engine/Source/Runtime/CoreUObject/Public/Serialization/BulkData.h
+++ b/Engine/Source/Runtime/CoreUObject/Public/Serialization/BulkData.h
@@ -23,8 +23,8 @@
 #endif
 
 // Enable the following to use the more compact FBulkDataStreamingToken in places where it is implemented
-#define USE_BULKDATA_STREAMING_TOKEN DEPRECATED_MACRO(5.0, "USE_BULKDATA_STREAMING_TOKEN now always evaluates to 0 and will be removed") 0
-#define STREAMINGTOKEN_PARAM(param) DEPRECATED_MACRO(5.0, "STREAMINGTOKEN_PARAM now always evaluates to a NOP")
+#define USE_BULKDATA_STREAMING_TOKEN UE_DEPRECATED_MACRO(5.0, "USE_BULKDATA_STREAMING_TOKEN now always evaluates to 0 and will be removed") 0
+#define STREAMINGTOKEN_PARAM(param) UE_DEPRECATED_MACRO(5.0, "STREAMINGTOKEN_PARAM now always evaluates to a NOP")
 
 class IMappedFileHandle;
 class IMappedFileRegion;
diff --git a/Engine/Source/Runtime/CoreUObject/Public/UObject/DefineUPropertyMacros.h b/Engine/Source/Runtime/CoreUObject/Public/UObject/DefineUPropertyMacros.h
deleted file mode 100644
diff --git a/Engine/Source/Runtime/CoreUObject/Public/UObject/UndefineUPropertyMacros.h b/Engine/Source/Runtime/CoreUObject/Public/UObject/UndefineUPropertyMacros.h
deleted file mode 100644
diff --git a/Engine/Source/Runtime/CoreUObject/Public/UObject/UnrealType.h b/Engine/Source/Runtime/CoreUObject/Public/UObject/UnrealType.h
--- a/Engine/Source/Runtime/CoreUObject/Public/UObject/UnrealType.h
+++ b/Engine/Source/Runtime/CoreUObject/Public/UObject/UnrealType.h
@@ -34,9 +34,6 @@
 #include "UObject/WeakObjectPtr.h"
 #include "UObject/Field.h"
 
-// WARNING: This should always be the last include in any file that needs it (except .generated.h)
-#include "UObject/UndefineUPropertyMacros.h"
-
 class FBlake3;
 class UPropertyWrapper;
 
@@ -7085,5 +7082,3 @@ class COREUOBJECT_API UMulticastInlineDelegatePropertyWrapper : public UMulticas
 {
 	DECLARE_CLASS_INTRINSIC(UMulticastInlineDelegatePropertyWrapper, UMulticastDelegatePropertyWrapper, CLASS_Transient, TEXT("/Script/CoreUObject"));
 };
-
-#include "UObject/DefineUPropertyMacros.h"
diff --git a/Engine/Source/Runtime/CoreUObject/Public/UObject/UnrealTypePrivate.h b/Engine/Source/Runtime/CoreUObject/Public/UObject/UnrealTypePrivate.h
--- a/Engine/Source/Runtime/CoreUObject/Public/UObject/UnrealTypePrivate.h
+++ b/Engine/Source/Runtime/CoreUObject/Public/UObject/UnrealTypePrivate.h
@@ -9,9 +9,6 @@
 #include "CoreMinimal.h"
 #include "UnrealType.h"
 
-// WARNING: This should always be the last include in any file that needs it (except .generated.h)
-#include "UObject/UndefineUPropertyMacros.h"
-
 #define USE_UPROPERTY_LOAD_DEFERRING (USE_CIRCULAR_DEPENDENCY_LOAD_DEFERRING && WITH_EDITORONLY_DATA)
 
 class COREUOBJECT_API UProperty : public UField
@@ -863,6 +860,3 @@ public:
 	{
 	}
 };
-
-// WARNING: This should always be the last include in any file that needs it (except .generated.h)
-#include "UObject/DefineUPropertyMacros.h"
\ No newline at end of file
diff --git a/Engine/Source/Runtime/Engine/Classes/Animation/AnimNodeBase.h b/Engine/Source/Runtime/Engine/Classes/Animation/AnimNodeBase.h
--- a/Engine/Source/Runtime/Engine/Classes/Animation/AnimNodeBase.h
+++ b/Engine/Source/Runtime/Engine/Classes/Animation/AnimNodeBase.h
@@ -19,9 +19,6 @@
 #include "Animation/ExposedValueHandler.h"
 #include "AnimNodeFunctionRef.h"
 
-// WARNING: This should always be the last include in any file that needs it (except .generated.h)
-#include "UObject/UndefineUPropertyMacros.h"
-
 #include "AnimNodeBase.generated.h"
 
 #define DECLARE_SCOPE_HIERARCHICAL_COUNTER_ANIMNODE(Method) \
@@ -1035,5 +1032,3 @@ private:
 // Editor-only way of accessing mutable anim node data but with internal checks
 #define GET_MUTABLE_ANIM_NODE_DATA(Type, Identifier) (GetMutableData<Type>(GET_ANIM_NODE_DATA_ID_INTERNAL(Type, Identifier)))
 #endif
-
-#include "UObject/DefineUPropertyMacros.h"
\ No newline at end of file
diff --git a/Engine/Source/Runtime/Engine/Classes/Engine/BlueprintGeneratedClass.h b/Engine/Source/Runtime/Engine/Classes/Engine/BlueprintGeneratedClass.h
--- a/Engine/Source/Runtime/Engine/Classes/Engine/BlueprintGeneratedClass.h
+++ b/Engine/Source/Runtime/Engine/Classes/Engine/BlueprintGeneratedClass.h
@@ -12,9 +12,6 @@
 #include "EdGraph/EdGraphPin.h"
 #include "Containers/SortedMap.h"
 
-// WARNING: This should always be the last include in any file that needs it (except .generated.h)
-#include "UObject/UndefineUPropertyMacros.h"
-
 #include "BlueprintGeneratedClass.generated.h"
 
 class AActor;
@@ -890,5 +887,3 @@ private:
 	FEditorTags CookedEditorTags;
 #endif // WITH_EDITORONLY_DATA
 };
-
-#include "UObject/DefineUPropertyMacros.h"
\ No newline at end of file
diff --git a/Engine/Source/Runtime/RenderCore/Public/ShaderParameterMacros.h b/Engine/Source/Runtime/RenderCore/Public/ShaderParameterMacros.h
--- a/Engine/Source/Runtime/RenderCore/Public/ShaderParameterMacros.h
+++ b/Engine/Source/Runtime/RenderCore/Public/ShaderParameterMacros.h
@@ -1714,15 +1714,15 @@ const TCHAR* GetShaderParameterMacroName(EUniformBufferBaseType ShaderParameterB
 
 /** Deprecated buffer macros. */
 #define SHADER_PARAMETER_RDG_BUFFER_UPLOAD(MemberName) \
-	DEPRECATED_MACRO(5.0, "SHADER_PARAMETER_RDG_BUFFER_UPLOAD has been deprecated. Use RDG_BUFFER_ACCESS(Buffer, ERHIAccess::CopyDest) instead.") \
+	UE_DEPRECATED_MACRO(5.0, "SHADER_PARAMETER_RDG_BUFFER_UPLOAD has been deprecated. Use RDG_BUFFER_ACCESS(Buffer, ERHIAccess::CopyDest) instead.") \
 	RDG_BUFFER_ACCESS(MemberName, ERHIAccess::CopyDest)
 
 #define SHADER_PARAMETER_RDG_BUFFER(ShaderType,MemberName) \
-	DEPRECATED_MACRO(5.0, "SHADER_PARAMETER_RDG_BUFFER has been deprecated. Use RDG_BUFFER_ACCESS with an explicit RHI state instead.") \
+	UE_DEPRECATED_MACRO(5.0, "SHADER_PARAMETER_RDG_BUFFER has been deprecated. Use RDG_BUFFER_ACCESS with an explicit RHI state instead.") \
 	RDG_BUFFER_ACCESS(MemberName, ERHIAccess::SRVMask | ERHIAccess::IndirectArgs)
 
 #define SHADER_PARAMETER_RDG_BUFFER_ARRAY(ShaderType,MemberName, ArrayDecl) \
-	DEPRECATED_MACRO(5.0, "SHADER_PARAMETER_RDG_BUFFER_ARRAY has been deprecated. Use RDG_BUFFER_ACCESS_ARRAY instead.")
+	UE_DEPRECATED_MACRO(5.0, "SHADER_PARAMETER_RDG_BUFFER_ARRAY has been deprecated. Use RDG_BUFFER_ACCESS_ARRAY instead.")
 
 
 
diff --git a/Engine/Source/Runtime/UMG/Public/Components/Widget.h b/Engine/Source/Runtime/UMG/Public/Components/Widget.h
--- a/Engine/Source/Runtime/UMG/Public/Components/Widget.h
+++ b/Engine/Source/Runtime/UMG/Public/Components/Widget.h
@@ -90,17 +90,6 @@ namespace UMWidget
 
 
 
-/**
- * Helper macro for binding to a delegate or using the constant value when constructing the underlying SWidget
- */
-#define OPTIONAL_BINDING(ReturnType, MemberName)				\
-	DEPRECATED_MACRO(4.17, "OPTIONAL_BINDING macro is deprecated.  Please use PROPERTY_BINDING in place and you'll need to define a PROPERTY_BINDING_IMPLEMENTATION in your header instead.") \
-	( MemberName ## Delegate.IsBound() && !IsDesignTime() )		\
-	?															\
-		TAttribute< ReturnType >::Create(MemberName ## Delegate.GetUObject(), MemberName ## Delegate.GetFunctionName()) \
-	:															\
-		TAttribute< ReturnType >(MemberName)
-
 #if WITH_EDITOR
 
 /**
