RWLock: Fixed assertions in TryReadLock and TryWriteLock

We want to return false for EAGAIN and EDEADLK because these are cases that the caller can reasonably handle. Other errors (EINVAL and EFAULT) are now ignored and return false as well, because they are unlikely when FRWLock is used correctly.

#rnx

diff --git a/Engine/Source/Runtime/Core/Public/HAL/PThreadRWLock.h b/Engine/Source/Runtime/Core/Public/HAL/PThreadRWLock.h
--- a/Engine/Source/Runtime/Core/Public/HAL/PThreadRWLock.h
+++ b/Engine/Source/Runtime/Core/Public/HAL/PThreadRWLock.h
@@ -44,14 +44,12 @@ public:
 	bool TryReadLock()
 	{
 		int Err = pthread_rwlock_tryrdlock(&Mutex);
-		checkf(Err == 0 || Err == EBUSY, TEXT("pthread_rwlock_tryrdlock failed with error: %d"), Err);
 		return Err == 0;
 	}
 
 	bool TryWriteLock()
 	{
 		int Err = pthread_rwlock_trywrlock(&Mutex);
-		checkf(Err == 0 || Err == EBUSY, TEXT("pthread_rwlock_trywrlock failed with error: %d"), Err);
 		return Err == 0;
 	}
 
