FArchiveState::CustomVer doc comment fix

The implementation returns -1 if the version tag is not found, not 0.

diff --git a/Engine/Source/Runtime/Core/Public/Serialization/Archive.h b/Engine/Source/Runtime/Core/Public/Serialization/Archive.h
--- a/Engine/Source/Runtime/Core/Public/Serialization/Archive.h
+++ b/Engine/Source/Runtime/Core/Public/Serialization/Archive.h
@@ -238,7 +238,7 @@ public:
 	 * Queries a custom version from the archive.  If the archive is being used to write, the custom version must have already been registered.
 	 *
 	 * @param Key The guid of the custom version to query.
-	 * @return The version number, or 0 if the custom tag isn't stored in the archive.
+	 * @return The version number, or -1 if the custom tag isn't stored in the archive.
 	 */
 	int32 CustomVer(const struct FGuid& Key) const;
 
