UE_ prefix added to math macros, with switches for turning on deprecation and removal for the old macros.
Fixes for all the macro use in Core.

diff --git a/Engine/Source/Runtime/Core/Core.Build.cs b/Engine/Source/Runtime/Core/Core.Build.cs
--- a/Engine/Source/Runtime/Core/Core.Build.cs
+++ b/Engine/Source/Runtime/Core/Core.Build.cs
@@ -313,6 +313,8 @@ public class Core : ModuleRules
 		PrivateDefinitions.Add("PLATFORM_COMPILER_OPTIMIZATION_PG=" + (Target.bPGOOptimize ? "1" : "0"));
 		PrivateDefinitions.Add("PLATFORM_COMPILER_OPTIMIZATION_PG_PROFILING=" + (Target.bPGOProfile ? "1" : "0"));
 
+		PrivateDefinitions.Add("UE_DEFINE_LEGACY_MATH_CONSTANT_MACRO_NAMES=0");
+
 		UnsafeTypeCastWarningLevel = WarningLevel.Warning;
 	}
 
diff --git a/Engine/Source/Runtime/Core/Private/GenericPlatform/GenericPlatformMath.cpp b/Engine/Source/Runtime/Core/Private/GenericPlatform/GenericPlatformMath.cpp
--- a/Engine/Source/Runtime/Core/Private/GenericPlatform/GenericPlatformMath.cpp
+++ b/Engine/Source/Runtime/Core/Private/GenericPlatform/GenericPlatformMath.cpp
@@ -68,8 +68,8 @@ float FGenericPlatformMath::Atan2(float Y, float X)
 	t0 = t0 * t4 + c[6];
 	t3 = t0 * t3;
 
-	t3 = yAbsBigger ? (0.5f * PI) - t3 : t3;
-	t3 = (X < 0.0f) ? PI - t3 : t3;
+	t3 = yAbsBigger ? (0.5f * UE_PI) - t3 : t3;
+	t3 = (X < 0.0f) ? UE_PI - t3 : t3;
 	t3 = (Y < 0.0f) ? -t3 : t3;
 
 	return t3;
@@ -164,7 +164,7 @@ namespace CompilerHiddenConstants
 	volatile float Twelve = 12.0f;
 	volatile float Sixteen = 16.0f;
 	volatile float MinusOneE37 = -1.0e37f;
-	volatile float FloatMax = MAX_FLT;
+	volatile float FloatMax = UE_MAX_FLT;
 }
 
 template<class MathPlatform>
diff --git a/Engine/Source/Runtime/Core/Private/HAL/ConsoleManager.cpp b/Engine/Source/Runtime/Core/Private/HAL/ConsoleManager.cpp
--- a/Engine/Source/Runtime/Core/Private/HAL/ConsoleManager.cpp
+++ b/Engine/Source/Runtime/Core/Private/HAL/ConsoleManager.cpp
@@ -2168,7 +2168,7 @@ void TestConsoleVariableCallback(IConsoleVariable* Var)
 	check(Var);
 
 	float Value = Var->GetFloat();
-	check(FMath::IsNearlyEqual(Value, 3.1f, KINDA_SMALL_NUMBER));
+	check(FMath::IsNearlyEqual(Value, 3.1f, UE_KINDA_SMALL_NUMBER));
 
 	++GConsoleVariableCallbackTestCounter;
 }
@@ -2231,7 +2231,7 @@ void FConsoleManager::Test()
 		check(VarA->GetString() == FString(TEXT("1")));
 
 		check(VarB->GetInt() == 1);
-		check(FMath::IsNearlyEqual(VarB->GetFloat(), 1.2f, KINDA_SMALL_NUMBER));
+		check(FMath::IsNearlyEqual(VarB->GetFloat(), 1.2f, UE_KINDA_SMALL_NUMBER));
 		check(VarB->GetString() == FString(TEXT("1.2")));
 
 		check(RefD == 2);
@@ -2239,7 +2239,7 @@ void FConsoleManager::Test()
 		check(VarD->GetFloat() == (float)2);
 		check(VarD->GetString() == FString(TEXT("2")));
 
-		check(FMath::IsNearlyEqual(RefE, 2.1f, KINDA_SMALL_NUMBER));
+		check(FMath::IsNearlyEqual(RefE, 2.1f, UE_KINDA_SMALL_NUMBER));
 		check(VarE->GetInt() == (int32)RefE);
 		check(VarE->GetFloat() == RefE);
 		check(VarE->GetString() == FString(TEXT("2.1")));
@@ -2293,7 +2293,7 @@ void FConsoleManager::Test()
 			check(VarC == IConsoleManager::Get().FindConsoleVariable(TEXT("TestNameC")));
 			check(VarC->GetInt() == 1);
 			// note: exact comparison fails in Win32 release
-			check(FMath::IsNearlyEqual(VarC->GetFloat(), 1.23f, KINDA_SMALL_NUMBER));
+			check(FMath::IsNearlyEqual(VarC->GetFloat(), 1.23f, UE_KINDA_SMALL_NUMBER));
 			check(VarC->GetString() == FString(TEXT("1.23")));
 			VarC->Set(TEXT("3.1"), ECVF_SetByConsole);
 			check(VarC->GetString() == FString(TEXT("3.1")));
diff --git a/Engine/Source/Runtime/Core/Private/Internationalization/FastDecimalFormat.cpp b/Engine/Source/Runtime/Core/Private/Internationalization/FastDecimalFormat.cpp
--- a/Engine/Source/Runtime/Core/Private/Internationalization/FastDecimalFormat.cpp
+++ b/Engine/Source/Runtime/Core/Private/Internationalization/FastDecimalFormat.cpp
@@ -529,7 +529,7 @@ void FractionalToString(const double InVal, const FDecimalNumberFormattingRules&
 	// Check for float-> int overflow, fallback on regular lex if occurs
 	// if fractional part overflows then we are losing precession but the number is still valid
 	uint64 IntIntegralPart = static_cast<uint64>(IntegralPart);
-	if (IntegralPart - static_cast<double>(IntIntegralPart) > SMALL_NUMBER && bFastDecimalFormatLargeFloatSupport)
+	if (IntegralPart - static_cast<double>(IntIntegralPart) > UE_SMALL_NUMBER && bFastDecimalFormatLargeFloatSupport)
 	{
 		OutString = LexToSanitizedString(InVal);
 
diff --git a/Engine/Source/Runtime/Core/Private/Math/SHMath.cpp b/Engine/Source/Runtime/Core/Private/Math/SHMath.cpp
--- a/Engine/Source/Runtime/Core/Private/Math/SHMath.cpp
+++ b/Engine/Source/Runtime/Core/Private/Math/SHMath.cpp
@@ -35,7 +35,7 @@ static int32 InitSHTables()
 		BasisM[BasisIndex] = M;
 
 		NormalizationConstants[BasisIndex] = FMath::Sqrt(
-			(float(2 * L + 1) / float(4 * PI)) *
+			(float(2 * L + 1) / float(4 * UE_PI)) *
 			(float(Factorial(L - FMath::Abs(M))) / float(Factorial(L + FMath::Abs(M))))
 			);
 
@@ -57,7 +57,7 @@ static int32 InitDummy = InitSHTables();
 /** So that e.g. LP(1,1,1) which evaluates to -sqrt(1-1^2) is 0.*/
 FORCEINLINE float SafeSqrt(float F)
 {
-	return FMath::Abs(F) > KINDA_SMALL_NUMBER ? FMath::Sqrt(F) : 0.f;
+	return FMath::Abs(F) > UE_KINDA_SMALL_NUMBER ? FMath::Sqrt(F) : 0.f;
 }
 
 /** Evaluates the LegendrePolynomial for L,M at X */
diff --git a/Engine/Source/Runtime/Core/Private/Math/TransformVectorized.cpp b/Engine/Source/Runtime/Core/Private/Math/TransformVectorized.cpp
--- a/Engine/Source/Runtime/Core/Private/Math/TransformVectorized.cpp
+++ b/Engine/Source/Runtime/Core/Private/Math/TransformVectorized.cpp
@@ -163,7 +163,7 @@ void TTransform<T>::SetToRelativeTransform(const TTransform<T>& ParentTransform)
 	checkSlow(ParentTransform.IsRotationNormalized());
 
 	// Scale = S(A)/S(B)	
-	TransformVectorRegister VSafeScale3D	= VectorSet_W0(GetSafeScaleReciprocal(ParentTransform.Scale3D, ScalarRegister(SMALL_NUMBER)));
+	TransformVectorRegister VSafeScale3D	= VectorSet_W0(GetSafeScaleReciprocal(ParentTransform.Scale3D, ScalarRegister(UE_SMALL_NUMBER)));
 	Scale3D = VectorMultiply(Scale3D, VSafeScale3D);
 	
 	//VQTranslation = (  ( T(A).X - T(B).X ),  ( T(A).Y - T(B).Y ), ( T(A).Z - T(B).Z), 0.f );
@@ -195,7 +195,7 @@ void TTransform<T>::GetRelativeTransformUsingMatrixWithScale(TTransform<T>* OutT
 	TMatrix<T> BM = Relative->ToMatrixWithScale();
 	// get combined scale
 	// Scale = S(A)/S(B)
-	static ScalarRegister STolerance(SMALL_NUMBER);
+	static ScalarRegister STolerance(UE_SMALL_NUMBER);
 	TransformVectorRegister VSafeScale3D = VectorSet_W0(GetSafeScaleReciprocal(Relative->Scale3D, STolerance));
 	TransformVectorRegister VScale3D = VectorMultiply(Base->Scale3D, VSafeScale3D);
 	ConstructTransformFromMatrixWithDesiredScale(AM, BM.Inverse(), VScale3D, *OutTransform);
@@ -225,7 +225,7 @@ TTransform<T> TTransform<T>::GetRelativeTransform(const TTransform<T>& Other) co
 	else
 	{
 		// Scale = S(A)/S(B)
-		static ScalarRegister STolerance(SMALL_NUMBER);
+		static ScalarRegister STolerance(UE_SMALL_NUMBER);
 		TransformVectorRegister VSafeScale3D = VectorSet_W0(GetSafeScaleReciprocal(Other.Scale3D, STolerance));
 
 		TransformVectorRegister VScale3D = VectorMultiply(Scale3D, VSafeScale3D);
diff --git a/Engine/Source/Runtime/Core/Private/Math/UnitConversion.cpp b/Engine/Source/Runtime/Core/Private/Math/UnitConversion.cpp
--- a/Engine/Source/Runtime/Core/Private/Math/UnitConversion.cpp
+++ b/Engine/Source/Runtime/Core/Private/Math/UnitConversion.cpp
@@ -559,7 +559,7 @@ namespace UnitConversion
 		// Convert to degrees
 		switch (From)
 		{
-			case EUnit::Radians:			return (180 / PI);
+			case EUnit::Radians:			return (180 / UE_PI);
 			default: 						return 1;
 		}
 	}
diff --git a/Engine/Source/Runtime/Core/Private/Math/UnrealMath.cpp b/Engine/Source/Runtime/Core/Private/Math/UnrealMath.cpp
--- a/Engine/Source/Runtime/Core/Private/Math/UnrealMath.cpp
+++ b/Engine/Source/Runtime/Core/Private/Math/UnrealMath.cpp
@@ -312,7 +312,7 @@ FVector FMath::ClosestPointOnInfiniteLine(const FVector& LineStart, const FVecto
 {
 	const FVector::FReal A = (LineStart - Point) | (LineEnd - LineStart);
 	const FVector::FReal B = (LineEnd - LineStart).SizeSquared();
-	if (B < SMALL_NUMBER)
+	if (B < UE_SMALL_NUMBER)
 	{
 		return LineStart;
 	}
@@ -337,9 +337,9 @@ UE::Math::TVector<T> UE::Math::TRotator<T>::Vector() const
 {
 	// Extremely large but valid values (or invalid values from uninitialized vars) can cause SinCos to return NaN/Inf, so catch that here. Similar to what is done in FRotator::Quaternion().
 #if ENABLE_NAN_DIAGNOSTIC || (DO_CHECK && !UE_BUILD_SHIPPING)
-	if (FMath::Abs(Pitch) > FLOAT_NON_FRACTIONAL ||
-		FMath::Abs(Yaw  ) > FLOAT_NON_FRACTIONAL ||
-		FMath::Abs(Roll ) > FLOAT_NON_FRACTIONAL)
+	if (FMath::Abs(Pitch) > UE_FLOAT_NON_FRACTIONAL ||
+		FMath::Abs(Yaw  ) > UE_FLOAT_NON_FRACTIONAL ||
+		FMath::Abs(Roll ) > UE_FLOAT_NON_FRACTIONAL)
 	{
 		logOrEnsureNanError(TEXT("FRotator::Vector() provided with unreasonably large input values (%s), possible use of uninitialized variable?"), *ToString());
 	}
@@ -410,7 +410,7 @@ FQuat4f FRotator3f::Quaternion() const
 	const VectorRegister4Float Result = VectorAdd(LeftTerm, RightTerm);	
 	FQuat4f RotationQuat = FQuat4f::MakeFromVectorRegister(Result);
 #else
-	const float DEG_TO_RAD = PI/(180.f);
+	const float DEG_TO_RAD = UE_PI/(180.f);
 	const float RADS_DIVIDED_BY_2 = DEG_TO_RAD/2.f;
 	float SP, SY, SR;
 	float CP, CY, CR;
@@ -564,7 +564,7 @@ UE::Math::TRotator<T> UE::Math::TMatrix<T>::Rotator() const
 	const TVector		XAxis	= GetScaledAxis( EAxis::X );
 	const TVector		YAxis	= GetScaledAxis( EAxis::Y );
 	const TVector		ZAxis	= GetScaledAxis( EAxis::Z );
-	const T RadToDeg = T(180.0 / DOUBLE_PI);
+	const T RadToDeg = T(180.0 / UE_DOUBLE_PI);
 
 	TRotator Rotator	= TRotator(
 									FMath::Atan2( XAxis.Z, FMath::Sqrt(FMath::Square(XAxis.X)+FMath::Square(XAxis.Y)) ) * RadToDeg,
@@ -615,7 +615,7 @@ namespace Math
 		// but that isn't the case for us, so I went through different testing, and finally found the case 
 		// where both of world lives happily. 
 		const float SINGULARITY_THRESHOLD = 0.4999995f;
-		const float RAD_TO_DEG = (180.f / PI);
+		const float RAD_TO_DEG = (180.f / UE_PI);
 		float Pitch, Yaw, Roll;
 
 		if (SingularityTest < -SINGULARITY_THRESHOLD)
@@ -659,7 +659,7 @@ namespace Math
 		const double YawX = (1.0 - 2.0 * (FMath::Square(Y) + FMath::Square(Z)));
 
 		const double SINGULARITY_THRESHOLD = 0.4999995;
-		const double RAD_TO_DEG = (180.0 / DOUBLE_PI);
+		const double RAD_TO_DEG = (180.0 / UE_DOUBLE_PI);
 		double Pitch, Yaw, Roll;
 
 		if (SingularityTest < -SINGULARITY_THRESHOLD)
@@ -1038,7 +1038,7 @@ namespace Math
 
 		//UE_LOG(LogUnrealMath, Log,  TEXT("CosAngle: %f Angle: %f"), CosAngle, Angle );
 
-		if ( FMath::Abs(Angle) < T(KINDA_SMALL_NUMBER))
+		if ( FMath::Abs(Angle) < T(UE_KINDA_SMALL_NUMBER))
 		{
 			return quat1;
 		}
@@ -1102,7 +1102,7 @@ namespace Math
 			Q = *this;
 
 			// Make sure we have a non null SquareSum. It shouldn't happen with a quaternion, but better be safe.
-			if (Q.SizeSquared() <= (T)SMALL_NUMBER)
+			if (Q.SizeSquared() <= (T)UE_SMALL_NUMBER)
 			{
 				Q = TQuat<T>::Identity;
 			}
@@ -1225,7 +1225,7 @@ namespace Math
 			const T Angle = FMath::Acos(W);
 			const T SinAngle = FMath::Sin(Angle);
 
-			if (FMath::Abs(SinAngle) >= T(SMALL_NUMBER))
+			if (FMath::Abs(SinAngle) >= T(UE_SMALL_NUMBER))
 			{
 				const T Scale = Angle / SinAngle;
 				Result.X = Scale * X;
@@ -1252,7 +1252,7 @@ namespace Math
 		TQuat<T>  Result;
 		Result.W = FMath::Cos(Angle);
 
-		if (FMath::Abs(SinAngle) >= T(SMALL_NUMBER))
+		if (FMath::Abs(SinAngle) >= T(UE_SMALL_NUMBER))
 		{
 			const T Scale = SinAngle / Angle;
 			Result.X = Scale * X;
@@ -1530,7 +1530,7 @@ struct SegmentDistToSegment_Solver
 		FVector::FReal N1;
 		FVector::FReal N2;
 
-		if (bLinesAreNearlyParallel || D < KINDA_SMALL_NUMBER)
+		if (bLinesAreNearlyParallel || D < UE_KINDA_SMALL_NUMBER)
 		{
 			// the lines are almost parallel
 			N1 = 0.f;	// force using point A on segment S1
@@ -1602,8 +1602,8 @@ struct SegmentDistToSegment_Solver
 		}
 
 		// finally do the division to get the points' location
-		const FVector::FReal T1 = (FMath::Abs(N1) < KINDA_SMALL_NUMBER ? 0.f : N1 / D1);
-		const FVector::FReal T2 = (FMath::Abs(N2) < KINDA_SMALL_NUMBER ? 0.f : N2 / D2);
+		const FVector::FReal T1 = (FMath::Abs(N1) < UE_KINDA_SMALL_NUMBER ? 0.f : N1 / D1);
+		const FVector::FReal T2 = (FMath::Abs(N2) < UE_KINDA_SMALL_NUMBER ? 0.f : N2 / D2);
 
 		// return the closest points
 		OutP1 = A1 + T1 * S1;
@@ -1643,7 +1643,7 @@ void FMath::SegmentDistToSegmentSafe(FVector A1, FVector B1, FVector A2, FVector
 		const	FVector::FReal	Dot12_norm = S1_norm | S2_norm;
 		const	FVector::FReal	D_norm	= Dot11_norm*Dot22_norm - Dot12_norm*Dot12_norm;	// always >= 0
 
-		Solver.bLinesAreNearlyParallel = D_norm < KINDA_SMALL_NUMBER;
+		Solver.bLinesAreNearlyParallel = D_norm < UE_KINDA_SMALL_NUMBER;
 		Solver.Solve(OutP1, OutP2);
 	}
 }
@@ -1662,7 +1662,7 @@ bool FMath::SegmentPlaneIntersection(const FVector& StartPoint, const FVector& E
 {
 	FVector::FReal T = FMath::GetTForSegmentPlaneIntersect(StartPoint, EndPoint, Plane);
 	// If the parameter value is not between 0 and 1, there is no intersection
-	if (T > -KINDA_SMALL_NUMBER && T < 1.f + KINDA_SMALL_NUMBER)
+	if (T > -UE_KINDA_SMALL_NUMBER && T < 1.f + UE_KINDA_SMALL_NUMBER)
 	{
 		out_IntersectionPoint = StartPoint + T * (EndPoint - StartPoint);
 		return true;
@@ -1984,7 +1984,7 @@ FVector FMath::ComputeBaryCentric2D(const FVector& Point, const FVector& A, cons
 	const FVector TriNorm = (B-A) ^ (C-A);
 
 	// Check the size of the triangle is reasonable (TriNorm.Size() will be twice the triangle area)
-	if(TriNorm.SizeSquared() <= SMALL_NUMBER)
+	if(TriNorm.SizeSquared() <= UE_SMALL_NUMBER)
 	{
 		UE_LOG(LogUnrealMath, Warning, TEXT("Small triangle detected in FMath::ComputeBaryCentric2D(), can't compute valid barycentric coordinate."));
 		return FVector(0.0f, 0.0f, 0.0f);
@@ -2017,7 +2017,7 @@ FVector4 FMath::ComputeBaryCentric3D(const FVector& Point, const FVector& A, con
 	const FVector B3 = (D-A);
 
 	//check co-planarity of A,B,C,D
-	check( FMath::Abs(B1 | (B2 ^ B3)) > SMALL_NUMBER && "Coplanar points in FMath::ComputeBaryCentric3D()");
+	check( FMath::Abs(B1 | (B2 ^ B3)) > UE_SMALL_NUMBER && "Coplanar points in FMath::ComputeBaryCentric3D()");
 
 	//Transform Point into this new space
 	const FVector V = (Point - A);
@@ -2130,7 +2130,7 @@ void FMath::SphereDistToLine(FVector SphereOrigin, float SphereRadius, FVector L
 	const FVector::FReal C = LineOriginToSphereOrigin.SizeSquared() - FMath::Square(SphereRadius);
 	const FVector::FReal D	= FMath::Square(B) - 4.f * C;
 
-	if( D <= KINDA_SMALL_NUMBER )
+	if( D <= UE_KINDA_SMALL_NUMBER )
 	{
 		// line is not intersecting sphere (or is tangent at one point if D == 0 )
 		const FVector PointOnLine = LineOrigin + ( -B * 0.5f ) * NormalizedLineDir;
@@ -2265,7 +2265,7 @@ CORE_API FVector FMath::VInterpNormalRotationTo(const FVector& Current, const FV
 	DeltaQuat.ToAxisAndAngle(DeltaAxis, DeltaAngle);
 
 	// Find rotation step for this frame
-	const float RotationStepRadians = RotationSpeedDegrees * (PI / 180) * DeltaTime;
+	const float RotationStepRadians = RotationSpeedDegrees * (UE_PI / 180) * DeltaTime;
 
 	if( FMath::Abs(DeltaAngle) > RotationStepRadians )
 	{
@@ -2310,7 +2310,7 @@ CORE_API FVector FMath::VInterpTo( const FVector& Current, const FVector& Target
 	const FVector Dist = Target - Current;
 
 	// If distance is too small, just set the desired location
-	if( Dist.SizeSquared() < KINDA_SMALL_NUMBER )
+	if( Dist.SizeSquared() < UE_KINDA_SMALL_NUMBER )
 	{
 		return Target;
 	}
@@ -2351,7 +2351,7 @@ CORE_API FVector2D FMath::Vector2DInterpTo( const FVector2D& Current, const FVec
 	}
 
 	const FVector2D Dist = Target - Current;
-	if( Dist.SizeSquared() < KINDA_SMALL_NUMBER )
+	if( Dist.SizeSquared() < UE_KINDA_SMALL_NUMBER )
 	{
 		return Target;
 	}
@@ -2426,7 +2426,7 @@ CORE_API FLinearColor FMath::CInterpTo(const FLinearColor& Current, const FLinea
 	const float Dist = FLinearColor::Dist(Target, Current);
 
 	// If distance is too small, just set the desired color
-	if (Dist < KINDA_SMALL_NUMBER)
+	if (Dist < UE_KINDA_SMALL_NUMBER)
 	{
 		return Target;
 	}
@@ -2453,7 +2453,7 @@ CORE_API UE::Math::TQuat<T> FMath::QInterpConstantTo(const UE::Math::TQuat<T>& C
 	}
 
 	float DeltaInterpSpeed = FMath::Clamp<float>(DeltaTime * InterpSpeed, 0.f, 1.f);
-	float AngularDistance = FMath::Max<float>(SMALL_NUMBER, (float)(Target.AngularDistance(Current)));
+	float AngularDistance = FMath::Max<float>(UE_SMALL_NUMBER, (float)(Target.AngularDistance(Current)));
 	float Alpha = FMath::Clamp<float>(DeltaInterpSpeed / AngularDistance, 0.f, 1.f);
 
 	return UE::Math::TQuat<T>::Slerp(Current, Target, Alpha);
@@ -2489,9 +2489,9 @@ template CORE_API UE::Math::TQuat<double> FMath::QInterpTo<double>(const UE::Mat
 
 CORE_API float ClampFloatTangent( float PrevPointVal, float PrevTime, float CurPointVal, float CurTime, float NextPointVal, float NextTime )
 {
-	const float PrevToNextTimeDiff = FMath::Max< float >( KINDA_SMALL_NUMBER, NextTime - PrevTime );
-	const float PrevToCurTimeDiff = FMath::Max< float >( KINDA_SMALL_NUMBER, CurTime - PrevTime );
-	const float CurToNextTimeDiff = FMath::Max< float >( KINDA_SMALL_NUMBER, NextTime - CurTime );
+	const float PrevToNextTimeDiff = FMath::Max< float >( UE_KINDA_SMALL_NUMBER, NextTime - PrevTime );
+	const float PrevToCurTimeDiff = FMath::Max< float >( UE_KINDA_SMALL_NUMBER, CurTime - PrevTime );
+	const float CurToNextTimeDiff = FMath::Max< float >( UE_KINDA_SMALL_NUMBER, NextTime - CurTime );
 
 	float OutTangentVal = 0.0f;
 
@@ -2581,7 +2581,7 @@ FVector FMath::VRandCone(FVector const& Dir, float ConeHalfAngleRad)
 
 		// Get spherical coords that have an even distribution over the unit sphere
 		// Method described at http://mathworld.wolfram.com/SpherePointPicking.html	
-		float Theta = 2.f * PI * RandU;
+		float Theta = 2.f * UE_PI * RandU;
 		float Phi = FMath::Acos((2.f * RandV) - 1.f);
 
 		// restrict phi to [0, ConeHalfAngleRad]
@@ -2595,8 +2595,8 @@ FVector FMath::VRandCone(FVector const& Dir, float ConeHalfAngleRad)
 		FVector const DirZ = DirMat.GetScaledAxis( EAxis::X );		
 		FVector const DirY = DirMat.GetScaledAxis( EAxis::Y );
 
-		FVector Result = Dir.RotateAngleAxis(Phi * 180.f / PI, DirY);
-		Result = Result.RotateAngleAxis(Theta * 180.f / PI, DirZ);
+		FVector Result = Dir.RotateAngleAxis(Phi * 180.f / UE_PI, DirY);
+		Result = Result.RotateAngleAxis(Theta * 180.f / UE_PI, DirZ);
 
 		// ensure it's a unit vector (might not have been passed in that way)
 		Result = Result.GetSafeNormal();
@@ -2618,7 +2618,7 @@ FVector FMath::VRandCone(FVector const& Dir, float HorizontalConeHalfAngleRad, f
 
 		// Get spherical coords that have an even distribution over the unit sphere
 		// Method described at http://mathworld.wolfram.com/SpherePointPicking.html	
-		float Theta = 2.f * PI * RandU;
+		float Theta = 2.f * UE_PI * RandU;
 		float Phi = FMath::Acos((2.f * RandV) - 1.f);
 
 		// restrict phi to [0, ConeHalfAngleRad]
@@ -2637,8 +2637,8 @@ FVector FMath::VRandCone(FVector const& Dir, float HorizontalConeHalfAngleRad, f
 		FVector const DirZ = DirMat.GetScaledAxis( EAxis::X );		
 		FVector const DirY = DirMat.GetScaledAxis( EAxis::Y );
 
-		FVector Result = Dir.RotateAngleAxis(Phi * 180.f / PI, DirY);
-		Result = Result.RotateAngleAxis(Theta * 180.f / PI, DirZ);
+		FVector Result = Dir.RotateAngleAxis(Phi * 180.f / UE_PI, DirY);
+		Result = Result.RotateAngleAxis(Theta * 180.f / UE_PI, DirZ);
 
 		// ensure it's a unit vector (might not have been passed in that way)
 		Result = Result.GetSafeNormal();
diff --git a/Engine/Source/Runtime/Core/Private/Misc/FrameRate.cpp b/Engine/Source/Runtime/Core/Private/Misc/FrameRate.cpp
--- a/Engine/Source/Runtime/Core/Private/Misc/FrameRate.cpp
+++ b/Engine/Source/Runtime/Core/Private/Misc/FrameRate.cpp
@@ -75,7 +75,7 @@ public:
 		JumpTable.MapPostUnary<FFramesPerSecond>([](double In) -> FExpressionResult { return MakeFrameRateFromFPS(In); });
 		JumpTable.MapPostUnary<FNonDropFrame>([](double In) -> FExpressionResult { return MakeFrameRateFromFPS(In); });
 		JumpTable.MapPostUnary<FDropFrame>([](double In) -> FExpressionResult {
-			if (!FMath::IsNearlyEqual(In, NTSC_30_FPS, DOUBLE_KINDA_SMALL_NUMBER) && !FMath::IsNearlyEqual(In, NTSC_60_FPS, DOUBLE_KINDA_SMALL_NUMBER))
+			if (!FMath::IsNearlyEqual(In, NTSC_30_FPS, UE_DOUBLE_KINDA_SMALL_NUMBER) && !FMath::IsNearlyEqual(In, NTSC_60_FPS, UE_DOUBLE_KINDA_SMALL_NUMBER))
 			{
 				return MakeError(FText::Format(LOCTEXT("InvalidDropFrameFPS", "Drop frame not supported for FPS specified: {0}.\nDrop frame timecode is only supported for NTSC_30 ({1}fps) and NTSC_60 ({2}fps) frame rates."), In, NTSC_30_FPS, NTSC_60_FPS));
 			}
@@ -166,15 +166,15 @@ private:
 		if (RoundedFPS != InFPS)
 		{
 			// Allow creating FFrameRates for fractional FPS values only if they match an NTSC frame rate.
-			if (FMath::IsNearlyEqual(InFPS, NTSC_24_FPS, DOUBLE_KINDA_SMALL_NUMBER))
+			if (FMath::IsNearlyEqual(InFPS, NTSC_24_FPS, UE_DOUBLE_KINDA_SMALL_NUMBER))
 			{
 				return MakeValue(FFrameRate(NTSC_24_Numerator, NTSC_Denominator));
 			}
-			else if (FMath::IsNearlyEqual(InFPS, NTSC_30_FPS, DOUBLE_KINDA_SMALL_NUMBER))
+			else if (FMath::IsNearlyEqual(InFPS, NTSC_30_FPS, UE_DOUBLE_KINDA_SMALL_NUMBER))
 			{
 				return MakeValue(FFrameRate(NTSC_30_Numerator, NTSC_Denominator));
 			}
-			else if (FMath::IsNearlyEqual(InFPS, NTSC_60_FPS, DOUBLE_KINDA_SMALL_NUMBER))
+			else if (FMath::IsNearlyEqual(InFPS, NTSC_60_FPS, UE_DOUBLE_KINDA_SMALL_NUMBER))
 			{
 				return MakeValue(FFrameRate(NTSC_60_Numerator, NTSC_Denominator));
 			}
diff --git a/Engine/Source/Runtime/Core/Private/Stats/StatsHierarchical.cpp b/Engine/Source/Runtime/Core/Private/Stats/StatsHierarchical.cpp
--- a/Engine/Source/Runtime/Core/Private/Stats/StatsHierarchical.cpp
+++ b/Engine/Source/Runtime/Core/Private/Stats/StatsHierarchical.cpp
@@ -248,7 +248,7 @@ void FStatsTreeElement::UpdatePostMeasurement(double InCyclesPerTimerToRemove)
 	};
 
 	// remove the cycles for the timers
-	if (InCyclesPerTimerToRemove > SMALL_NUMBER)
+	if (InCyclesPerTimerToRemove > UE_SMALL_NUMBER)
 	{
 		Local::RemoveCyclesForTimers(this, InCyclesPerTimerToRemove);
 	}
diff --git a/Engine/Source/Runtime/Core/Private/Tests/Math/UnitConversionTests.cpp b/Engine/Source/Runtime/Core/Private/Tests/Math/UnitConversionTests.cpp
--- a/Engine/Source/Runtime/Core/Private/Tests/Math/UnitConversionTests.cpp
+++ b/Engine/Source/Runtime/Core/Private/Tests/Math/UnitConversionTests.cpp
@@ -40,8 +40,8 @@ bool FUnitUnitTests::RunTest(const FString& Parameters)
 		{ 1000,		1.0936,		1e-4,	EUnit::Millimeters, 		EUnit::Yards		 		},
 		{ 2000,		0.0787,		1e-4,	EUnit::Micrometers, 		EUnit::Inches		 		},
 
-		{ 90,		PI/2,		1e-3,	EUnit::Degrees, 			EUnit::Radians		 		},
-		{ PI,		180,		1e-3,	EUnit::Radians, 			EUnit::Degrees		 		},
+		{ 90,		UE_PI/2,	1e-3,	EUnit::Degrees, 			EUnit::Radians		 		},
+		{ UE_PI,	180,		1e-3,	EUnit::Radians, 			EUnit::Degrees		 		},
 
 		{ 12,		43.2,		0.1,	EUnit::MetersPerSecond,		EUnit::KilometersPerHour	},
 		{ 1,		0.6214,		1e-4,	EUnit::KilometersPerHour,	EUnit::MilesPerHour			},
diff --git a/Engine/Source/Runtime/Core/Private/Tests/Math/UnrealMathTest.cpp b/Engine/Source/Runtime/Core/Private/Tests/Math/UnrealMathTest.cpp
--- a/Engine/Source/Runtime/Core/Private/Tests/Math/UnrealMathTest.cpp
+++ b/Engine/Source/Runtime/Core/Private/Tests/Math/UnrealMathTest.cpp
@@ -680,7 +680,7 @@ FORCENOINLINE FRotator3f TestQuaternionToRotator(const FQuat4f& Quat)
 	const float YawX = (1.f-2.f*(FMath::Square(Y) + FMath::Square(Z)));
 	const float SINGULARITY_THRESHOLD = 0.4999995f;
 
-	static const float RAD_TO_DEG = (180.f)/PI;
+	static const float RAD_TO_DEG = (180.f)/UE_PI;
 	FRotator3f RotatorFromQuat;
 
 	// Note: using stock C functions for some trig functions since this is the "reference" implementation
@@ -718,11 +718,11 @@ FORCENOINLINE FQuat4f FindBetween_Old(const FVector3f& vec1, const FVector3f& ve
 	const float crossMag = cross.Size();
 
 	// See if vectors are parallel or anti-parallel
-	if (crossMag < KINDA_SMALL_NUMBER)
+	if (crossMag < UE_KINDA_SMALL_NUMBER)
 	{
 		// If these vectors are parallel - just return identity quaternion (ie no rotation).
 		const float Dot = vec1 | vec2;
-		if (Dot > -KINDA_SMALL_NUMBER)
+		if (Dot > -UE_KINDA_SMALL_NUMBER)
 		{
 			return FQuat4f::Identity; // no rotation
 		}
@@ -746,7 +746,7 @@ FORCENOINLINE FQuat4f FindBetween_Old(const FVector3f& vec1, const FVector3f& ve
 	const float dot = vec1 | vec2;
 	if (dot < 0.0f)
 	{
-		angle = PI - angle;
+		angle = UE_PI - angle;
 	}
 
 	float sinHalfAng, cosHalfAng;
@@ -1920,45 +1920,45 @@ bool RunDoubleVectorTest()
 	V0 = MakeVectorRegister(-1.8, -1.0, -0.8, 0.0);
 	V2 = MakeVectorRegister(-1.0, -1.0, 0.0, 0.0);
 	V3 = VectorTruncate(V0);
-	LogTest<double>(TEXT("VectorTruncate"), TestVectorsEqual(V2, V3, DOUBLE_KINDA_SMALL_NUMBER));
+	LogTest<double>(TEXT("VectorTruncate"), TestVectorsEqual(V2, V3, UE_DOUBLE_KINDA_SMALL_NUMBER));
 
 	V0 = MakeVectorRegister(0.0, 0.8, 1.0, 1.8);
 	V2 = MakeVectorRegister(0.0, 0.0, 1.0, 1.0);
 	V3 = VectorTruncate(V0);
-	LogTest<double>(TEXT("VectorTruncate"), TestVectorsEqual(V2, V3, DOUBLE_KINDA_SMALL_NUMBER));
+	LogTest<double>(TEXT("VectorTruncate"), TestVectorsEqual(V2, V3, UE_DOUBLE_KINDA_SMALL_NUMBER));
 
 	V0 = MakeVectorRegister(0.1, 1.8, 2.4, -2.4);
 	V2 = MakeVectorRegister(0.0, 1.0, 2.0, -2.0);
 	V3 = VectorTruncate(V0);
-	LogTest<double>(TEXT("VectorTruncate"), TestVectorsEqual(V2, V3, DOUBLE_KINDA_SMALL_NUMBER));
+	LogTest<double>(TEXT("VectorTruncate"), TestVectorsEqual(V2, V3, UE_DOUBLE_KINDA_SMALL_NUMBER));
 
 	// VectorFractional
 	V0 = MakeVectorRegister(-1.8, -1.0, -0.8, 0.0);
 	V2 = MakeVectorRegister(-0.8, 0.0, -0.8, 0.0);
 	V3 = VectorFractional(V0);
-	LogTest<double>(TEXT("VectorFractional"), TestVectorsEqual(V2, V3, DOUBLE_KINDA_SMALL_NUMBER));
+	LogTest<double>(TEXT("VectorFractional"), TestVectorsEqual(V2, V3, UE_DOUBLE_KINDA_SMALL_NUMBER));
 
 	V0 = MakeVectorRegister(0.0, 0.8, 1.0, 1.8);
 	V2 = MakeVectorRegister(0.0, 0.8, 0.0, 0.8);
 	V3 = VectorFractional(V0);
-	LogTest<double>(TEXT("VectorFractional"), TestVectorsEqual(V2, V3, DOUBLE_KINDA_SMALL_NUMBER));
+	LogTest<double>(TEXT("VectorFractional"), TestVectorsEqual(V2, V3, UE_DOUBLE_KINDA_SMALL_NUMBER));
 
 	// VectorCeil
 	V0 = MakeVectorRegister(-1.8, -1.0, -0.8, 0.0);
 	V2 = MakeVectorRegister(-1.0, -1.0, -0.0, 0.0);
 	V3 = VectorCeil(V0);
-	LogTest<double>(TEXT("VectorCeil"), TestVectorsEqual(V2, V3, DOUBLE_KINDA_SMALL_NUMBER));
+	LogTest<double>(TEXT("VectorCeil"), TestVectorsEqual(V2, V3, UE_DOUBLE_KINDA_SMALL_NUMBER));
 
 	V0 = MakeVectorRegister(0.0, 0.8, 1.0, 1.8);
 	V2 = MakeVectorRegister(0.0, 1.0, 1.0, 2.0);
 	V3 = VectorCeil(V0);
-	LogTest<double>(TEXT("VectorCeil"), TestVectorsEqual(V2, V3, DOUBLE_KINDA_SMALL_NUMBER));
+	LogTest<double>(TEXT("VectorCeil"), TestVectorsEqual(V2, V3, UE_DOUBLE_KINDA_SMALL_NUMBER));
 
 	// VectorFloor
 	V0 = MakeVectorRegister(-1.8, -1.0, -0.8, 0.0);
 	V2 = MakeVectorRegister(-2.0, -1.0, -1.0, 0.0);
 	V3 = VectorFloor(V0);
-	LogTest<double>(TEXT("VectorFloor"), TestVectorsEqual(V2, V3, DOUBLE_KINDA_SMALL_NUMBER));
+	LogTest<double>(TEXT("VectorFloor"), TestVectorsEqual(V2, V3, UE_DOUBLE_KINDA_SMALL_NUMBER));
 
 	V0 = MakeVectorRegister(0.0, 0.8, 1.0, 1.8);
 	V2 = MakeVectorRegister(0.0, 0.0, 1.0, 1.0);
@@ -2668,40 +2668,40 @@ bool FVectorRegisterAbstractionTest::RunTest(const FString& Parameters)
 	V0 = MakeVectorRegister(-1.8f, -1.0f, -0.8f, 0.0f);
 	V2 = MakeVectorRegister(-1.0f, -1.0f, 0.0f, 0.0f);
 	V3 = VectorTruncate(V0);
-	LogTest<float>(TEXT("VectorTruncate"), TestVectorsEqual(V2, V3, KINDA_SMALL_NUMBER));
+	LogTest<float>(TEXT("VectorTruncate"), TestVectorsEqual(V2, V3, UE_KINDA_SMALL_NUMBER));
 
 	V0 = MakeVectorRegister(0.0f, 0.8f, 1.0f, 1.8f);
 	V2 = MakeVectorRegister(0.0f, 0.0f, 1.0f, 1.0f);
 	V3 = VectorTruncate(V0);
-	LogTest<float>(TEXT("VectorTruncate"), TestVectorsEqual(V2, V3, KINDA_SMALL_NUMBER));
+	LogTest<float>(TEXT("VectorTruncate"), TestVectorsEqual(V2, V3, UE_KINDA_SMALL_NUMBER));
 
 	// VectorFractional
 	V0 = MakeVectorRegister(-1.8f, -1.0f, -0.8f, 0.0f);
 	V2 = MakeVectorRegister(-0.8f, 0.0f, -0.8f, 0.0f);
 	V3 = VectorFractional(V0);
-	LogTest<float>(TEXT("VectorFractional"), TestVectorsEqual(V2, V3, KINDA_SMALL_NUMBER));
+	LogTest<float>(TEXT("VectorFractional"), TestVectorsEqual(V2, V3, UE_KINDA_SMALL_NUMBER));
 
 	V0 = MakeVectorRegister(0.0f, 0.8f, 1.0f, 1.8f);
 	V2 = MakeVectorRegister(0.0f, 0.8f, 0.0f, 0.8f);
 	V3 = VectorFractional(V0);
-	LogTest<float>(TEXT("VectorFractional"), TestVectorsEqual(V2, V3, KINDA_SMALL_NUMBER));
+	LogTest<float>(TEXT("VectorFractional"), TestVectorsEqual(V2, V3, UE_KINDA_SMALL_NUMBER));
 
 	// VectorCeil
 	V0 = MakeVectorRegister(-1.8f, -1.0f, -0.8f, 0.0f);
 	V2 = MakeVectorRegister(-1.0f, -1.0f, -0.0f, 0.0f);
 	V3 = VectorCeil(V0);
-	LogTest<float>(TEXT("VectorCeil"), TestVectorsEqual(V2, V3, KINDA_SMALL_NUMBER));
+	LogTest<float>(TEXT("VectorCeil"), TestVectorsEqual(V2, V3, UE_KINDA_SMALL_NUMBER));
 
 	V0 = MakeVectorRegister(0.0f, 0.8f, 1.0f, 1.8f);
 	V2 = MakeVectorRegister(0.0f, 1.0f, 1.0f, 2.0f);
 	V3 = VectorCeil(V0);
-	LogTest<float>(TEXT("VectorCeil"), TestVectorsEqual(V2, V3, KINDA_SMALL_NUMBER));
+	LogTest<float>(TEXT("VectorCeil"), TestVectorsEqual(V2, V3, UE_KINDA_SMALL_NUMBER));
 
 	// VectorFloor
 	V0 = MakeVectorRegister(-1.8f, -1.0f, -0.8f, 0.0f);
 	V2 = MakeVectorRegister(-2.0f, -1.0f, -1.0f, 0.0f);
 	V3 = VectorFloor(V0);
-	LogTest<float>(TEXT("VectorFloor"), TestVectorsEqual(V2, V3, KINDA_SMALL_NUMBER));
+	LogTest<float>(TEXT("VectorFloor"), TestVectorsEqual(V2, V3, UE_KINDA_SMALL_NUMBER));
 
 	V0 = MakeVectorRegister(0.0f, 0.8f, 1.0f, 1.8f);
 	V2 = MakeVectorRegister(0.0f, 0.0f, 1.0f, 1.0f);
@@ -2803,8 +2803,8 @@ bool FVectorRegisterAbstractionTest::RunTest(const FString& Parameters)
 			{-2.8f,	-0.3f},
 			{-0.4f,	 5.5f},
 			{ 0.4f,	-5.5f},
-			{ 2.8f,	 2.0f + KINDA_SMALL_NUMBER},
-			{-2.8f,	 2.0f - KINDA_SMALL_NUMBER},
+			{ 2.8f,	 2.0f + UE_KINDA_SMALL_NUMBER},
+			{-2.8f,	 2.0f - UE_KINDA_SMALL_NUMBER},
 
 			// Analytically should be zero but floating point precision can cause results close to Y (or erroneously negative) depending on the method used.
 			{55.8f,	 9.3f},
@@ -2813,27 +2813,27 @@ bool FVectorRegisterAbstractionTest::RunTest(const FString& Parameters)
 			// Commonly used for FRotators and angles
 			{725.2f,	360.0f},
 			{179.9f,	90.0f},
-			{ 5.3f * PI,	2.f * PI},
-			{-5.3f * PI,	2.f * PI},
+			{ 5.3f * UE_PI,	2.f * UE_PI},
+			{-5.3f * UE_PI,	2.f * UE_PI},
 
 			// Test extreme ranges
-			{ 1.0f,			 KINDA_SMALL_NUMBER},
-			{ 1.0f,			-KINDA_SMALL_NUMBER},
-			{-SMALL_NUMBER,  SMALL_NUMBER},
-			{ SMALL_NUMBER, -SMALL_NUMBER},
-			{ 1.0f,			 MIN_flt},
-			{ 1.0f,			-MIN_flt},
-			{ MAX_flt,		 MIN_flt},
-			{ MAX_flt,		-MIN_flt},
+			{ 1.0f,				 UE_KINDA_SMALL_NUMBER},
+			{ 1.0f,				-UE_KINDA_SMALL_NUMBER},
+			{-UE_SMALL_NUMBER,   UE_SMALL_NUMBER},
+			{ UE_SMALL_NUMBER,  -UE_SMALL_NUMBER},
+			{ 1.0f,				 MIN_flt},
+			{ 1.0f,				-MIN_flt},
+			{ MAX_flt,			 MIN_flt},
+			{ MAX_flt,			-MIN_flt},
 
 			// Test some tricky cases
 			{8388808.00f, 178.222f},  // FAILS old version
 			{360.0f * 1e10f,		360.0f},
 			{-720.0f * 1000000.0f,	360.0f},
 			{1080.0f * 12345.f,		360.0f},
-			{360.0f,				360.0f - KINDA_SMALL_NUMBER},
-			{1080.0f,				360.0f - KINDA_SMALL_NUMBER},
-			{719.0f,				360.0f - KINDA_SMALL_NUMBER},
+			{360.0f,				360.0f - UE_KINDA_SMALL_NUMBER},
+			{1080.0f,				360.0f - UE_KINDA_SMALL_NUMBER},
+			{719.0f,				360.0f - UE_KINDA_SMALL_NUMBER},
 			
 
 			// We define this to be zero and not NaN.
@@ -2885,8 +2885,8 @@ bool FVectorRegisterAbstractionTest::RunTest(const FString& Parameters)
 			{-2.8,	-0.3},
 			{-0.4,	 5.5},
 			{ 0.4,	-5.5},
-			{ 2.8,	 2.0 + DOUBLE_KINDA_SMALL_NUMBER},
-			{-2.8,	 2.0 - DOUBLE_KINDA_SMALL_NUMBER},
+			{ 2.8,	 2.0 + UE_DOUBLE_KINDA_SMALL_NUMBER},
+			{-2.8,	 2.0 - UE_DOUBLE_KINDA_SMALL_NUMBER},
 
 			// Analytically should be zero but floating point precision can cause results close to Y (or erroneously negative) depending on the method used.
 			{55.8,	 9.3},
@@ -2896,14 +2896,14 @@ bool FVectorRegisterAbstractionTest::RunTest(const FString& Parameters)
 			// Commonly used for FRotators and angles
 			{725.2,	360.0},
 			{179.9,	90.0},
-			{ 5.3 * DOUBLE_PI,	2. * DOUBLE_PI},
-			{-5.3 * DOUBLE_PI,	2. * DOUBLE_PI},
+			{ 5.3 * UE_DOUBLE_PI,	2. * UE_DOUBLE_PI},
+			{-5.3 * UE_DOUBLE_PI,	2. * UE_DOUBLE_PI},
 
 			// Test extreme ranges
-			{ 1.0,			 DOUBLE_KINDA_SMALL_NUMBER},
-			{ 1.0,			-DOUBLE_KINDA_SMALL_NUMBER},
-			{-DOUBLE_SMALL_NUMBER,  DOUBLE_SMALL_NUMBER},
-			{ DOUBLE_SMALL_NUMBER, -DOUBLE_SMALL_NUMBER},
+			{ 1.0,			 UE_DOUBLE_KINDA_SMALL_NUMBER},
+			{ 1.0,			-UE_DOUBLE_KINDA_SMALL_NUMBER},
+			{-UE_DOUBLE_SMALL_NUMBER,  UE_DOUBLE_SMALL_NUMBER},
+			{ UE_DOUBLE_SMALL_NUMBER, -UE_DOUBLE_SMALL_NUMBER},
 			{ 1.0,			 MIN_dbl},
 			{ 1.0,			-MIN_dbl},
 			{ MAX_flt,		 MIN_dbl},
@@ -2914,9 +2914,9 @@ bool FVectorRegisterAbstractionTest::RunTest(const FString& Parameters)
 			{360.0 * 1e10,			360.0},
 			{-720.0 * 1000000.0,	360.0},
 			{1080.0 * 12345.,		360.0},
-			{360.0,					360.0 - DOUBLE_KINDA_SMALL_NUMBER},
-			{1080.0,				360.0 - DOUBLE_KINDA_SMALL_NUMBER},
-			{719.0,					360.0 - DOUBLE_KINDA_SMALL_NUMBER},
+			{360.0,					360.0 - UE_DOUBLE_KINDA_SMALL_NUMBER},
+			{1080.0,				360.0 - UE_DOUBLE_KINDA_SMALL_NUMBER},
+			{719.0,					360.0 - UE_DOUBLE_KINDA_SMALL_NUMBER},
 
 			// We define this to be zero and not NaN.
 			// Disabled since we don't want to trigger an ensure, but left here for testing that logic.
@@ -3120,7 +3120,7 @@ bool FVectorRegisterAbstractionTest::RunTest(const FString& Parameters)
 			LogQuaternionTest(TEXT("FRotator3f::ZeroRotator not != FQuat4f::Identity : Quaternion"), Q0, Q1, !(Q0 != Q1));
 		}
 
-		const float Nudge = KINDA_SMALL_NUMBER * 0.25f;
+		const float Nudge = UE_KINDA_SMALL_NUMBER * 0.25f;
 		const FRotator3f RotArray[] = {
 			FRotator3f(0.f, 0.f, 0.f),
 			FRotator3f(Nudge, -Nudge, Nudge),
@@ -3141,7 +3141,7 @@ bool FVectorRegisterAbstractionTest::RunTest(const FString& Parameters)
 		// FRotator3f tests
 		{
 			// Equality test
-			const float RotTolerance = KINDA_SMALL_NUMBER;
+			const float RotTolerance = UE_KINDA_SMALL_NUMBER;
 			for (auto const& A : RotArray)
 			{
 				for (auto const& B : RotArray)
@@ -3167,7 +3167,7 @@ bool FVectorRegisterAbstractionTest::RunTest(const FString& Parameters)
 
 	// Rotator->Quat->Rotator
 	{
-		const float Nudge = KINDA_SMALL_NUMBER * 0.25f;
+		const float Nudge = UE_KINDA_SMALL_NUMBER * 0.25f;
 		const FRotator3f RotArray[] ={
 			FRotator3f(0.0f, 0.0f, 0.0f),
 			FRotator3f(30.0f, 30.0f, 30.0f),
@@ -3387,9 +3387,9 @@ bool FVectorRegisterAbstractionTest::RunTest(const FString& Parameters)
 					}
 					else
 					{
-						LogTest<float>(TEXT("FindBetween: Old A->B"), TestFVector3Equal(RotAOld, BNorm, KINDA_SMALL_NUMBER));
-						LogTest<float>(TEXT("FindBetween: New A->B (normal)"), TestFVector3Equal(RotANewNormal, BNorm, KINDA_SMALL_NUMBER));
-						LogTest<float>(TEXT("FindBetween: New A->B (vector)"), TestFVector3Equal(RotANewVector, BNorm, KINDA_SMALL_NUMBER));
+						LogTest<float>(TEXT("FindBetween: Old A->B"), TestFVector3Equal(RotAOld, BNorm, UE_KINDA_SMALL_NUMBER));
+						LogTest<float>(TEXT("FindBetween: New A->B (normal)"), TestFVector3Equal(RotANewNormal, BNorm, UE_KINDA_SMALL_NUMBER));
+						LogTest<float>(TEXT("FindBetween: New A->B (vector)"), TestFVector3Equal(RotANewVector, BNorm, UE_KINDA_SMALL_NUMBER));
 					}
 				}
 			}
@@ -3410,8 +3410,8 @@ bool FVectorRegisterAbstractionTest::RunTest(const FString& Parameters)
 				const FVector3f Rotated1 = Q1.RotateVector(FVector3f::ForwardVector);
 				const FVector3f Rotated2 = R0.RotateVector(FVector3f::ForwardVector);
 
-				LogTest<float>(TEXT("V.ToOrientationQuat() rotate"), TestFVector3Equal(Rotated0, Rotated1, KINDA_SMALL_NUMBER));
-				LogTest<float>(TEXT("V.ToOrientationRotator() rotate"), TestFVector3Equal(Rotated0, Rotated2, KINDA_SMALL_NUMBER));
+				LogTest<float>(TEXT("V.ToOrientationQuat() rotate"), TestFVector3Equal(Rotated0, Rotated1, UE_KINDA_SMALL_NUMBER));
+				LogTest<float>(TEXT("V.ToOrientationRotator() rotate"), TestFVector3Equal(Rotated0, Rotated2, UE_KINDA_SMALL_NUMBER));
 			}
 		}
 	}
@@ -4122,7 +4122,7 @@ bool FInitVectorTest::RunTest(const FString& Parameters)
 		const bool bIsInitialized = Actual.InitFromCompactString(InExpected.ToCompactString());
 
 		TestTrue(*(InTestName + " return value"), bIsInitialized);
-		TestEqual(*InTestName, Actual, InExpected, KINDA_SMALL_NUMBER);
+		TestEqual(*InTestName, Actual, InExpected, UE_KINDA_SMALL_NUMBER);
 	};
 
 	TestInitFromCompactString(TEXT("InitFromCompactString Simple"), FVector(1.2f, 2.3f, 3.4f));
diff --git a/Engine/Source/Runtime/Core/Private/Tests/Serialization/CompactBinaryWriterTest.cpp b/Engine/Source/Runtime/Core/Private/Tests/Serialization/CompactBinaryWriterTest.cpp
--- a/Engine/Source/Runtime/Core/Private/Tests/Serialization/CompactBinaryWriterTest.cpp
+++ b/Engine/Source/Runtime/Core/Private/Tests/Serialization/CompactBinaryWriterTest.cpp
@@ -476,7 +476,7 @@ bool FCbWriterFloatTest::RunTest(const FString& Parameters)
 	// Test Float32
 	{
 		Writer.Reset();
-		constexpr float Values[] = { 0.0f, 1.0f, -1.0f, PI };
+		constexpr float Values[] = { 0.0f, 1.0f, -1.0f, UE_PI };
 		for (float Value : Values)
 		{
 			Writer.AddFloat(Value);
@@ -496,7 +496,7 @@ bool FCbWriterFloatTest::RunTest(const FString& Parameters)
 	// Test Float64
 	{
 		Writer.Reset();
-		constexpr double Values[] = { 0.0f, 1.0f, -1.0f, PI, 1.9999998807907104, 1.9999999403953552, 3.4028234663852886e38, 6.8056469327705771e38 };
+		constexpr double Values[] = { 0.0f, 1.0f, -1.0f, UE_PI, 1.9999998807907104, 1.9999999403953552, 3.4028234663852886e38, 6.8056469327705771e38 };
 		for (double Value : Values)
 		{
 			Writer.AddFloat(Value);
diff --git a/Engine/Source/Runtime/Core/Public/Math/Box.h b/Engine/Source/Runtime/Core/Public/Math/Box.h
--- a/Engine/Source/Runtime/Core/Public/Math/Box.h
+++ b/Engine/Source/Runtime/Core/Public/Math/Box.h
@@ -129,7 +129,7 @@ public:
 	 * @param Tolerance Error tolerance.
 	 * @return true if the boxes are equal within tolerance limits, false otherwise.
 	 */
-	bool Equals(const TBox<T>& Other, T Tolerance=KINDA_SMALL_NUMBER) const
+	bool Equals(const TBox<T>& Other, T Tolerance=UE_KINDA_SMALL_NUMBER) const
 	{
 		return Min.Equals(Other.Min, Tolerance) && Max.Equals(Other.Max, Tolerance);
 	}
diff --git a/Engine/Source/Runtime/Core/Public/Math/BoxSphereBounds.h b/Engine/Source/Runtime/Core/Public/Math/BoxSphereBounds.h
--- a/Engine/Source/Runtime/Core/Public/Math/BoxSphereBounds.h
+++ b/Engine/Source/Runtime/Core/Public/Math/BoxSphereBounds.h
@@ -180,7 +180,7 @@ public:
 	 * @param  Tolerance Error tolerance added to test distance.
 	 * @return true if spheres intersect, false otherwise.
 	 */
-	FORCEINLINE static bool SpheresIntersect(const TBoxSphereBounds<T, TExtent>& A, const TBoxSphereBounds<T, TExtent>& B, TExtent Tolerance = KINDA_SMALL_NUMBER)
+	FORCEINLINE static bool SpheresIntersect(const TBoxSphereBounds<T, TExtent>& A, const TBoxSphereBounds<T, TExtent>& B, TExtent Tolerance = UE_KINDA_SMALL_NUMBER)
 	{
 		return (A.Origin - B.Origin).SizeSquared() <= FMath::Square(FMath::Max<TExtent>(0, A.SphereRadius + B.SphereRadius + Tolerance));
 	}
diff --git a/Engine/Source/Runtime/Core/Public/Math/Color.h b/Engine/Source/Runtime/Core/Public/Math/Color.h
--- a/Engine/Source/Runtime/Core/Public/Math/Color.h
+++ b/Engine/Source/Runtime/Core/Public/Math/Color.h
@@ -256,7 +256,7 @@ struct FLinearColor
 	}
 
 	// Error-tolerant comparison.
-	FORCEINLINE bool Equals(const FLinearColor& ColorB, float Tolerance=KINDA_SMALL_NUMBER) const
+	FORCEINLINE bool Equals(const FLinearColor& ColorB, float Tolerance=UE_KINDA_SMALL_NUMBER) const
 	{
 		return FMath::Abs(this->R - ColorB.R) < Tolerance && FMath::Abs(this->G - ColorB.G) < Tolerance && FMath::Abs(this->B - ColorB.B) < Tolerance && FMath::Abs(this->A - ColorB.A) < Tolerance;
 	}
@@ -376,7 +376,7 @@ struct FLinearColor
 	/** useful to detect if a light contribution needs to be rendered */
 	bool IsAlmostBlack() const
 	{
-		return FMath::Square(R) < DELTA && FMath::Square(G) < DELTA && FMath::Square(B) < DELTA;
+		return FMath::Square(R) < UE_DELTA && FMath::Square(G) < UE_DELTA && FMath::Square(B) < UE_DELTA;
 	}
 
 	/**
diff --git a/Engine/Source/Runtime/Core/Public/Math/InterpCurvePoint.h b/Engine/Source/Runtime/Core/Public/Math/InterpCurvePoint.h
--- a/Engine/Source/Runtime/Core/Public/Math/InterpCurvePoint.h
+++ b/Engine/Source/Runtime/Core/Public/Math/InterpCurvePoint.h
@@ -222,7 +222,7 @@ inline void ComputeCurveTangent( float PrevTime, const T& PrevPoint,
 
 	AutoCalcTangent( PrevPoint, CurPoint, NextPoint, Tension, OutTangent );
 
-	const float PrevToNextTimeDiff = FMath::Max< float >( KINDA_SMALL_NUMBER, NextTime - PrevTime );
+	const float PrevToNextTimeDiff = FMath::Max< float >( UE_KINDA_SMALL_NUMBER, NextTime - PrevTime );
 
 	OutTangent /= PrevToNextTimeDiff;
 }
@@ -273,7 +273,7 @@ inline void ComputeClampableFloatVectorCurveTangent( float PrevTime, const T& Pr
 		// No clamping needed
 		AutoCalcTangent( PrevPoint, CurPoint, NextPoint, Tension, OutTangent );
 
-		const float PrevToNextTimeDiff = FMath::Max< float >( KINDA_SMALL_NUMBER, NextTime - PrevTime );
+		const float PrevToNextTimeDiff = FMath::Max< float >( UE_KINDA_SMALL_NUMBER, NextTime - PrevTime );
 
 		OutTangent /= PrevToNextTimeDiff;
 	}
diff --git a/Engine/Source/Runtime/Core/Public/Math/InverseRotationMatrix.h b/Engine/Source/Runtime/Core/Public/Math/InverseRotationMatrix.h
--- a/Engine/Source/Runtime/Core/Public/Math/InverseRotationMatrix.h
+++ b/Engine/Source/Runtime/Core/Public/Math/InverseRotationMatrix.h
@@ -32,19 +32,19 @@ template<typename T>
 FORCEINLINE TInverseRotationMatrix<T>::TInverseRotationMatrix(const TRotator<T>& Rot)
 	: TMatrix<T>(
 		TMatrix<T>( // Yaw
-			TPlane<T>(+FMath::Cos(Rot.Yaw * PI / 180.f), -FMath::Sin(Rot.Yaw * PI / 180.f), 0.0f, 0.0f),
-			TPlane<T>(+FMath::Sin(Rot.Yaw * PI / 180.f), +FMath::Cos(Rot.Yaw * PI / 180.f), 0.0f, 0.0f),
+			TPlane<T>(+FMath::Cos(Rot.Yaw * UE_PI / 180.f), -FMath::Sin(Rot.Yaw * UE_PI / 180.f), 0.0f, 0.0f),
+			TPlane<T>(+FMath::Sin(Rot.Yaw * UE_PI / 180.f), +FMath::Cos(Rot.Yaw * UE_PI / 180.f), 0.0f, 0.0f),
 			TPlane<T>(0.0f, 0.0f, 1.0f, 0.0f),
 			TPlane<T>(0.0f, 0.0f, 0.0f, 1.0f)) *
 		TMatrix<T>( // Pitch
-			TPlane<T>(+FMath::Cos(Rot.Pitch * PI / 180.f), 0.0f, -FMath::Sin(Rot.Pitch * PI / 180.f), 0.0f),
+			TPlane<T>(+FMath::Cos(Rot.Pitch * UE_PI / 180.f), 0.0f, -FMath::Sin(Rot.Pitch * UE_PI / 180.f), 0.0f),
 			TPlane<T>(0.0f, 1.0f, 0.0f, 0.0f),
-			TPlane<T>(+FMath::Sin(Rot.Pitch * PI / 180.f), 0.0f, +FMath::Cos(Rot.Pitch * PI / 180.f), 0.0f),
+			TPlane<T>(+FMath::Sin(Rot.Pitch * UE_PI / 180.f), 0.0f, +FMath::Cos(Rot.Pitch * UE_PI / 180.f), 0.0f),
 			TPlane<T>(0.0f, 0.0f, 0.0f, 1.0f)) *
 		TMatrix<T>( // Roll
 			TPlane<T>(1.0f, 0.0f, 0.0f, 0.0f),
-			TPlane<T>(0.0f, +FMath::Cos(Rot.Roll * PI / 180.f), +FMath::Sin(Rot.Roll * PI / 180.f), 0.0f),
-			TPlane<T>(0.0f, -FMath::Sin(Rot.Roll * PI / 180.f), +FMath::Cos(Rot.Roll * PI / 180.f), 0.0f),
+			TPlane<T>(0.0f, +FMath::Cos(Rot.Roll * UE_PI / 180.f), +FMath::Sin(Rot.Roll * UE_PI / 180.f), 0.0f),
+			TPlane<T>(0.0f, -FMath::Sin(Rot.Roll * UE_PI / 180.f), +FMath::Cos(Rot.Roll * UE_PI / 180.f), 0.0f),
 			TPlane<T>(0.0f, 0.0f, 0.0f, 1.0f))
 	)
 { }
diff --git a/Engine/Source/Runtime/Core/Public/Math/Matrix.h b/Engine/Source/Runtime/Core/Public/Math/Matrix.h
--- a/Engine/Source/Runtime/Core/Public/Math/Matrix.h
+++ b/Engine/Source/Runtime/Core/Public/Math/Matrix.h
@@ -159,7 +159,7 @@ public:
 	 * @param Tolerance Error Tolerance.
 	 * @return true if two Matrix are equal, within specified tolerance, otherwise false.
 	 */
-	inline bool Equals(const TMatrix<T>& Other, T Tolerance = KINDA_SMALL_NUMBER) const;
+	inline bool Equals(const TMatrix<T>& Other, T Tolerance = UE_KINDA_SMALL_NUMBER) const;
 
 	/**
 	 * Checks whether another Matrix is not equal to this, within specified tolerance.
@@ -211,16 +211,16 @@ public:
 
 	// NOTE: There is some compiler optimization issues with WIN64 that cause FORCEINLINE to cause a crash
 	// Remove any scaling from this matrix (ie magnitude of each row is 1) with error Tolerance
-	inline void RemoveScaling(T Tolerance = SMALL_NUMBER);
+	inline void RemoveScaling(T Tolerance = UE_SMALL_NUMBER);
 
 	// Returns matrix after RemoveScaling with error Tolerance
-	inline TMatrix<T> GetMatrixWithoutScale(T Tolerance = SMALL_NUMBER) const;
+	inline TMatrix<T> GetMatrixWithoutScale(T Tolerance = UE_SMALL_NUMBER) const;
 
 	/** Remove any scaling from this matrix (ie magnitude of each row is 1) and return the 3D scale vector that was initially present with error Tolerance */
-	inline TVector<T> ExtractScaling(T Tolerance = SMALL_NUMBER);
+	inline TVector<T> ExtractScaling(T Tolerance = UE_SMALL_NUMBER);
 
 	/** return a 3D scale vector calculated from this matrix (where each component is the magnitude of a row vector) with error Tolerance. */
-	inline TVector<T> GetScaleVector(T Tolerance = SMALL_NUMBER) const;
+	inline TVector<T> GetScaleVector(T Tolerance = UE_SMALL_NUMBER) const;
 
 	// Remove any translation from this matrix
 	inline TMatrix<T> RemoveTranslation() const;
diff --git a/Engine/Source/Runtime/Core/Public/Math/Matrix.inl b/Engine/Source/Runtime/Core/Public/Math/Matrix.inl
--- a/Engine/Source/Runtime/Core/Public/Math/Matrix.inl
+++ b/Engine/Source/Runtime/Core/Public/Math/Matrix.inl
@@ -295,9 +295,9 @@ inline TMatrix<T> TMatrix<T>::InverseFast() const
 	// If we're in non final release, then make sure we're not creating NaNs
 #if !(UE_BUILD_SHIPPING || UE_BUILD_TEST)
 	// Check for zero scale matrix to invert
-	if (GetScaledAxis(EAxis::X).IsNearlyZero(SMALL_NUMBER) &&
-		GetScaledAxis(EAxis::Y).IsNearlyZero(SMALL_NUMBER) &&
-		GetScaledAxis(EAxis::Z).IsNearlyZero(SMALL_NUMBER))
+	if (GetScaledAxis(EAxis::X).IsNearlyZero(UE_SMALL_NUMBER) &&
+		GetScaledAxis(EAxis::Y).IsNearlyZero(UE_SMALL_NUMBER) &&
+		GetScaledAxis(EAxis::Z).IsNearlyZero(UE_SMALL_NUMBER))
 	{
 		ErrorEnsure(TEXT("TMatrix<T>::InverseFast(), trying to invert a NIL matrix, this results in NaNs! Use Inverse() instead."));
 	}
@@ -324,9 +324,9 @@ inline TMatrix<T> TMatrix<T>::Inverse() const
 	TMatrix<T> Result;
 
 	// Check for zero scale matrix to invert
-	if (GetScaledAxis(EAxis::X).IsNearlyZero(SMALL_NUMBER) &&
-		GetScaledAxis(EAxis::Y).IsNearlyZero(SMALL_NUMBER) &&
-		GetScaledAxis(EAxis::Z).IsNearlyZero(SMALL_NUMBER))
+	if (GetScaledAxis(EAxis::X).IsNearlyZero(UE_SMALL_NUMBER) &&
+		GetScaledAxis(EAxis::Y).IsNearlyZero(UE_SMALL_NUMBER) &&
+		GetScaledAxis(EAxis::Z).IsNearlyZero(UE_SMALL_NUMBER))
 	{
 		// just set to zero - avoids unsafe inverse of zero and duplicates what QNANs were resulting in before (scaling away all children)
 		Result = Identity;
@@ -381,7 +381,7 @@ inline TMatrix<T> TMatrix<T>::TransposeAdjoint() const
 // NOTE: There is some compiler optimization issues with WIN64 that cause FORCEINLINE to cause a crash
 // Remove any scaling from this matrix (ie magnitude of each row is 1)
 template<typename T>
-inline void TMatrix<T>::RemoveScaling(T Tolerance/*=SMALL_NUMBER*/)
+inline void TMatrix<T>::RemoveScaling(T Tolerance/*=UE_SMALL_NUMBER*/)
 {
 	// For each row, find magnitude, and if its non-zero re-scale so its unit length.
 	const T SquareSum0 = (M[0][0] * M[0][0]) + (M[0][1] * M[0][1]) + (M[0][2] * M[0][2]);
@@ -404,7 +404,7 @@ inline void TMatrix<T>::RemoveScaling(T Tolerance/*=SMALL_NUMBER*/)
 
 // Returns matrix without scale information
 template<typename T>
-inline TMatrix<T> TMatrix<T>::GetMatrixWithoutScale(T Tolerance/*=SMALL_NUMBER*/) const
+inline TMatrix<T> TMatrix<T>::GetMatrixWithoutScale(T Tolerance/*=UE_SMALL_NUMBER*/) const
 {
 	TMatrix<T> Result = (TMatrix<T>&)*this;
 	Result.RemoveScaling(Tolerance);
@@ -413,7 +413,7 @@ inline TMatrix<T> TMatrix<T>::GetMatrixWithoutScale(T Tolerance/*=SMALL_NUMBER*/
 
 /** Remove any scaling from this matrix (ie magnitude of each row is 1) and return the 3D scale vector that was initially present. */
 template<typename T>
-inline TVector<T> TMatrix<T>::ExtractScaling(T Tolerance/*=SMALL_NUMBER*/)
+inline TVector<T> TMatrix<T>::ExtractScaling(T Tolerance/*=UE_SMALL_NUMBER*/)
 {
 	TVector<T> Scale3D(0, 0, 0);
 
@@ -469,7 +469,7 @@ inline TVector<T> TMatrix<T>::ExtractScaling(T Tolerance/*=SMALL_NUMBER*/)
 
 /** return a 3D scale vector calculated from this matrix (where each component is the magnitude of a row vector). */
 template<typename T>
-inline TVector<T> TMatrix<T>::GetScaleVector(T Tolerance/*=SMALL_NUMBER*/) const
+inline TVector<T> TMatrix<T>::GetScaleVector(T Tolerance/*=UE_SMALL_NUMBER*/) const
 {
 	TVector<T> Scale3D(1, 1, 1);
 
@@ -704,7 +704,7 @@ template <typename T>
 FORCEINLINE bool MakeFrustumPlane(T A, T B, T C, T D, TPlane<T>& OutPlane)
 {
 	const T	LengthSquared = A * A + B * B + C * C;
-	if (LengthSquared > DELTA * DELTA)
+	if (LengthSquared > UE_DELTA * UE_DELTA)
 	{
 		const T	InvLength = FMath::InvSqrt(LengthSquared);
 		OutPlane = TPlane<T>(-A * InvLength, -B * InvLength, -C * InvLength, D * InvLength);
diff --git a/Engine/Source/Runtime/Core/Public/Math/Plane.h b/Engine/Source/Runtime/Core/Public/Math/Plane.h
--- a/Engine/Source/Runtime/Core/Public/Math/Plane.h
+++ b/Engine/Source/Runtime/Core/Public/Math/Plane.h
@@ -148,7 +148,7 @@ public:
 	 * @param Tolerance Minimum squared length of vector for normalization.
 	 * @return true if the plane was normalized correctly, false otherwise.
 	 */
-	bool Normalize(T Tolerance = SMALL_NUMBER);
+	bool Normalize(T Tolerance = UE_SMALL_NUMBER);
 
 	/**
 	 * Get a flipped version of the plane.
@@ -208,7 +208,7 @@ public:
 	 * @param Tolerance Error Tolerance.
 	 * @return true if the two planes are equal within specified tolerance, otherwise false.
 	 */
-	bool Equals(const TPlane<T>& V, T Tolerance = KINDA_SMALL_NUMBER) const;
+	bool Equals(const TPlane<T>& V, T Tolerance = UE_KINDA_SMALL_NUMBER) const;
 
 	/**
 	 * Calculates dot product of two planes.
diff --git a/Engine/Source/Runtime/Core/Public/Math/Quat.h b/Engine/Source/Runtime/Core/Public/Math/Quat.h
--- a/Engine/Source/Runtime/Core/Public/Math/Quat.h
+++ b/Engine/Source/Runtime/Core/Public/Math/Quat.h
@@ -172,7 +172,7 @@ public:
 	 * @param Tolerance Error tolerance for comparison with other Quaternion.
 	 * @return true if two Quaternions are equal, within specified tolerance, otherwise false.
 	 */
-	FORCEINLINE bool Equals(const TQuat<T>& Q, T Tolerance=KINDA_SMALL_NUMBER) const;
+	FORCEINLINE bool Equals(const TQuat<T>& Q, T Tolerance=UE_KINDA_SMALL_NUMBER) const;
 
 	/**
 	 * Checks whether this Quaternion is an Identity Quaternion.
@@ -181,7 +181,7 @@ public:
 	 * @param Tolerance Error tolerance for comparison with Identity Quaternion.
 	 * @return true if Quaternion is a normalized Identity Quaternion.
 	 */
-	FORCEINLINE bool IsIdentity(T Tolerance=SMALL_NUMBER) const;
+	FORCEINLINE bool IsIdentity(T Tolerance=UE_SMALL_NUMBER) const;
 
 	/**
 	 * Subtracts another quaternion from this.
@@ -372,7 +372,7 @@ public:
 	 *
 	 * @param Tolerance Minimum squared length of quaternion for normalization.
 	 */
-	FORCEINLINE void Normalize(T Tolerance = SMALL_NUMBER);
+	FORCEINLINE void Normalize(T Tolerance = UE_SMALL_NUMBER);
 
 	/**
 	 * Get a normalized copy of this quaternion.
@@ -380,7 +380,7 @@ public:
 	 *
 	 * @param Tolerance Minimum squared length of quaternion for normalization.
 	 */
-	FORCEINLINE TQuat<T> GetNormalized(T Tolerance = SMALL_NUMBER) const;
+	FORCEINLINE TQuat<T> GetNormalized(T Tolerance = UE_SMALL_NUMBER) const;
 
 	// Return true if this quaternion is normalized
 	bool IsNormalized() const;
@@ -729,7 +729,7 @@ inline TQuat<T>::TQuat(const UE::Math::TMatrix<T>& M)
 #if !(UE_BUILD_SHIPPING || UE_BUILD_TEST)
 	// Make sure the Rotation part of the Matrix is unit length.
 	// Changed to this (same as RemoveScaling) from RotDeterminant as using two different ways of checking unit length matrix caused inconsistency. 
-	if (!ensure((FMath::Abs(1.f - M.GetScaledAxis(EAxis::X).SizeSquared()) <= KINDA_SMALL_NUMBER) && (FMath::Abs(1.f - M.GetScaledAxis(EAxis::Y).SizeSquared()) <= KINDA_SMALL_NUMBER) && (FMath::Abs(1.f - M.GetScaledAxis(EAxis::Z).SizeSquared()) <= KINDA_SMALL_NUMBER)))
+	if (!ensure((FMath::Abs(1.f - M.GetScaledAxis(EAxis::X).SizeSquared()) <= UE_KINDA_SMALL_NUMBER) && (FMath::Abs(1.f - M.GetScaledAxis(EAxis::Y).SizeSquared()) <= UE_KINDA_SMALL_NUMBER) && (FMath::Abs(1.f - M.GetScaledAxis(EAxis::Z).SizeSquared()) <= UE_KINDA_SMALL_NUMBER)))
 	{
 		*this = TQuat<T>::Identity;
 		return;
@@ -1269,7 +1269,7 @@ template<typename T>
 FORCEINLINE T TQuat<T>::Error(const TQuat<T>& Q1, const TQuat<T>& Q2)
 {
 	const T cosom = FMath::Abs(Q1.X * Q2.X + Q1.Y * Q2.Y + Q1.Z * Q2.Z + Q1.W * Q2.W);
-	return (FMath::Abs(cosom) < 0.9999999f) ? FMath::Acos(cosom)*(1.f / PI) : 0.0f;
+	return (FMath::Abs(cosom) < 0.9999999f) ? FMath::Acos(cosom)*(1.f / UE_PI) : 0.0f;
 }
 
 
diff --git a/Engine/Source/Runtime/Core/Public/Math/RandomStream.h b/Engine/Source/Runtime/Core/Public/Math/RandomStream.h
--- a/Engine/Source/Runtime/Core/Public/Math/RandomStream.h
+++ b/Engine/Source/Runtime/Core/Public/Math/RandomStream.h
@@ -153,7 +153,7 @@ public:
 			Result.Z = GetFraction() * 2.f - 1.f;
 			L = Result.SizeSquared();
 		}
-		while(L > 1.f || L < KINDA_SMALL_NUMBER);
+		while(L > 1.f || L < UE_KINDA_SMALL_NUMBER);
 
 		return Result.GetUnsafeNormal();
 	}
@@ -244,7 +244,7 @@ public:
 
 			// Get spherical coords that have an even distribution over the unit sphere
 			// Method described at http://mathworld.wolfram.com/SpherePointPicking.html	
-			float Theta = 2.f * PI * RandU;
+			float Theta = 2.f * UE_PI * RandU;
 			float Phi = FMath::Acos((2.f * RandV) - 1.f);
 
 			// restrict phi to [0, ConeHalfAngleRad]
@@ -258,8 +258,8 @@ public:
 			FVector const DirZ = DirMat.GetUnitAxis( EAxis::X );		
 			FVector const DirY = DirMat.GetUnitAxis( EAxis::Y );
 
-			FVector Result = Dir.RotateAngleAxis(Phi * 180.f / PI, DirY);
-			Result = Result.RotateAngleAxis(Theta * 180.f / PI, DirZ);
+			FVector Result = Dir.RotateAngleAxis(Phi * 180.f / UE_PI, DirY);
+			Result = Result.RotateAngleAxis(Theta * 180.f / UE_PI, DirZ);
 
 			// ensure it's a unit vector (might not have been passed in that way)
 			Result = Result.GetSafeNormal();
@@ -289,7 +289,7 @@ public:
 
 			// Get spherical coords that have an even distribution over the unit sphere
 			// Method described at http://mathworld.wolfram.com/SpherePointPicking.html	
-			float Theta = 2.f * PI * RandU;
+			float Theta = 2.f * UE_PI * RandU;
 			float Phi = FMath::Acos((2.f * RandV) - 1.f);
 
 			// restrict phi to [0, ConeHalfAngleRad]
@@ -308,8 +308,8 @@ public:
 			FVector const DirZ = DirMat.GetUnitAxis( EAxis::X );		
 			FVector const DirY = DirMat.GetUnitAxis( EAxis::Y );
 
-			FVector Result = Dir.RotateAngleAxis(Phi * 180.f / PI, DirY);
-			Result = Result.RotateAngleAxis(Theta * 180.f / PI, DirZ);
+			FVector Result = Dir.RotateAngleAxis(Phi * 180.f / UE_PI, DirY);
+			Result = Result.RotateAngleAxis(Theta * 180.f / UE_PI, DirZ);
 
 			// ensure it's a unit vector (might not have been passed in that way)
 			Result = Result.GetSafeNormal();
diff --git a/Engine/Source/Runtime/Core/Public/Math/RotationMatrix.h b/Engine/Source/Runtime/Core/Public/Math/RotationMatrix.h
--- a/Engine/Source/Runtime/Core/Public/Math/RotationMatrix.h
+++ b/Engine/Source/Runtime/Core/Public/Math/RotationMatrix.h
@@ -78,7 +78,7 @@ TMatrix<T> TRotationMatrix<T>::MakeFromX(TVector<T> const& XAxis)
 	TVector<T> const NewX = XAxis.GetSafeNormal();
 
 	// try to use up if possible
-	TVector<T> const UpVector = (FMath::Abs(NewX.Z) < (1.f - KINDA_SMALL_NUMBER)) ? TVector<T>(0, 0, 1.f) : TVector<T>(1.f, 0, 0);
+	TVector<T> const UpVector = (FMath::Abs(NewX.Z) < (1.f - UE_KINDA_SMALL_NUMBER)) ? TVector<T>(0, 0, 1.f) : TVector<T>(1.f, 0, 0);
 
 	const TVector<T> NewY = (UpVector ^ NewX).GetSafeNormal();
 	const TVector<T> NewZ = NewX ^ NewY;
@@ -92,7 +92,7 @@ TMatrix<T> TRotationMatrix<T>::MakeFromY(TVector<T> const& YAxis)
 	TVector<T> const NewY = YAxis.GetSafeNormal();
 
 	// try to use up if possible
-	TVector<T> const UpVector = (FMath::Abs(NewY.Z) < (1.f - KINDA_SMALL_NUMBER)) ? TVector<T>(0, 0, 1.f) : TVector<T>(1.f, 0, 0);
+	TVector<T> const UpVector = (FMath::Abs(NewY.Z) < (1.f - UE_KINDA_SMALL_NUMBER)) ? TVector<T>(0, 0, 1.f) : TVector<T>(1.f, 0, 0);
 
 	const TVector<T> NewZ = (UpVector ^ NewY).GetSafeNormal();
 	const TVector<T> NewX = NewY ^ NewZ;
@@ -106,7 +106,7 @@ TMatrix<T> TRotationMatrix<T>::MakeFromZ(TVector<T> const& ZAxis)
 	TVector<T> const NewZ = ZAxis.GetSafeNormal();
 
 	// try to use up if possible
-	TVector<T> const UpVector = (FMath::Abs(NewZ.Z) < (1.f - KINDA_SMALL_NUMBER)) ? TVector<T>(0, 0, 1.f) : TVector<T>(1.f, 0, 0);
+	TVector<T> const UpVector = (FMath::Abs(NewZ.Z) < (1.f - UE_KINDA_SMALL_NUMBER)) ? TVector<T>(0, 0, 1.f) : TVector<T>(1.f, 0, 0);
 
 	const TVector<T> NewX = (UpVector ^ NewZ).GetSafeNormal();
 	const TVector<T> NewY = NewZ ^ NewX;
@@ -124,7 +124,7 @@ TMatrix<T> TRotationMatrix<T>::MakeFromXY(TVector<T> const& XAxis, TVector<T> co
 	if (FMath::IsNearlyEqual(FMath::Abs(NewX | Norm), T(1.f)))
 	{
 		// make sure we don't ever pick the same as NewX
-		Norm = (FMath::Abs(NewX.Z) < (1.f - KINDA_SMALL_NUMBER)) ? TVector<T>(0, 0, 1.f) : TVector<T>(1.f, 0, 0);
+		Norm = (FMath::Abs(NewX.Z) < (1.f - UE_KINDA_SMALL_NUMBER)) ? TVector<T>(0, 0, 1.f) : TVector<T>(1.f, 0, 0);
 	}
 
 	const TVector<T> NewZ = (NewX ^ Norm).GetSafeNormal();
@@ -143,7 +143,7 @@ TMatrix<T> TRotationMatrix<T>::MakeFromXZ(TVector<T> const& XAxis, TVector<T> co
 	if (FMath::IsNearlyEqual(FMath::Abs(NewX | Norm), T(1.f)))
 	{
 		// make sure we don't ever pick the same as NewX
-		Norm = (FMath::Abs(NewX.Z) < (1.f - KINDA_SMALL_NUMBER)) ? TVector<T>(0, 0, 1.f) : TVector<T>(1.f, 0, 0);
+		Norm = (FMath::Abs(NewX.Z) < (1.f - UE_KINDA_SMALL_NUMBER)) ? TVector<T>(0, 0, 1.f) : TVector<T>(1.f, 0, 0);
 	}
 
 	const TVector<T> NewY = (Norm ^ NewX).GetSafeNormal();
@@ -162,7 +162,7 @@ TMatrix<T> TRotationMatrix<T>::MakeFromYX(TVector<T> const& YAxis, TVector<T> co
 	if (FMath::IsNearlyEqual(FMath::Abs(NewY | Norm), T(1.f)))
 	{
 		// make sure we don't ever pick the same as NewX
-		Norm = (FMath::Abs(NewY.Z) < (1.f - KINDA_SMALL_NUMBER)) ? TVector<T>(0, 0, 1.f) : TVector<T>(1.f, 0, 0);
+		Norm = (FMath::Abs(NewY.Z) < (1.f - UE_KINDA_SMALL_NUMBER)) ? TVector<T>(0, 0, 1.f) : TVector<T>(1.f, 0, 0);
 	}
 
 	const TVector<T> NewZ = (Norm ^ NewY).GetSafeNormal();
@@ -181,7 +181,7 @@ TMatrix<T> TRotationMatrix<T>::MakeFromYZ(TVector<T> const& YAxis, TVector<T> co
 	if (FMath::IsNearlyEqual(FMath::Abs(NewY | Norm), T(1.f)))
 	{
 		// make sure we don't ever pick the same as NewX
-		Norm = (FMath::Abs(NewY.Z) < (1.f - KINDA_SMALL_NUMBER)) ? TVector<T>(0, 0, 1.f) : TVector<T>(1.f, 0, 0);
+		Norm = (FMath::Abs(NewY.Z) < (1.f - UE_KINDA_SMALL_NUMBER)) ? TVector<T>(0, 0, 1.f) : TVector<T>(1.f, 0, 0);
 	}
 
 	const TVector<T> NewX = (NewY ^ Norm).GetSafeNormal();
@@ -200,7 +200,7 @@ TMatrix<T> TRotationMatrix<T>::MakeFromZX(TVector<T> const& ZAxis, TVector<T> co
 	if (FMath::IsNearlyEqual(FMath::Abs(NewZ | Norm), T(1.f)))
 	{
 		// make sure we don't ever pick the same as NewX
-		Norm = (FMath::Abs(NewZ.Z) < (1.f - KINDA_SMALL_NUMBER)) ? TVector<T>(0, 0, 1.f) : TVector<T>(1.f, 0, 0);
+		Norm = (FMath::Abs(NewZ.Z) < (1.f - UE_KINDA_SMALL_NUMBER)) ? TVector<T>(0, 0, 1.f) : TVector<T>(1.f, 0, 0);
 	}
 
 	const TVector<T> NewY = (NewZ ^ Norm).GetSafeNormal();
@@ -219,7 +219,7 @@ TMatrix<T> TRotationMatrix<T>::MakeFromZY(TVector<T> const& ZAxis, TVector<T> co
 	if (FMath::IsNearlyEqual(FMath::Abs(NewZ | Norm), T(1.f)))
 	{
 		// make sure we don't ever pick the same as NewX
-		Norm = (FMath::Abs(NewZ.Z) < (1.f - KINDA_SMALL_NUMBER)) ? TVector<T>(0, 0, 1.f) : TVector<T>(1.f, 0, 0);
+		Norm = (FMath::Abs(NewZ.Z) < (1.f - UE_KINDA_SMALL_NUMBER)) ? TVector<T>(0, 0, 1.f) : TVector<T>(1.f, 0, 0);
 	}
 
 	const TVector<T> NewX = (Norm ^ NewZ).GetSafeNormal();
diff --git a/Engine/Source/Runtime/Core/Public/Math/Rotator.h b/Engine/Source/Runtime/Core/Public/Math/Rotator.h
--- a/Engine/Source/Runtime/Core/Public/Math/Rotator.h
+++ b/Engine/Source/Runtime/Core/Public/Math/Rotator.h
@@ -205,7 +205,7 @@ public:
 	 * @param Tolerance Error Tolerance.
 	 * @return true if rotator is nearly zero, within specified tolerance, otherwise false.
 	 */
-	bool IsNearlyZero( T Tolerance = KINDA_SMALL_NUMBER ) const;
+	bool IsNearlyZero( T Tolerance = UE_KINDA_SMALL_NUMBER ) const;
 
 	/**
 	 * Checks whether this has exactly zero rotation, when treated as an orientation.
@@ -223,7 +223,7 @@ public:
 	 * @param Tolerance Error Tolerance.
 	 * @return true if two rotators are equal, within specified tolerance, otherwise false.
 	 */
-	bool Equals( const TRotator<T>& R, T Tolerance = KINDA_SMALL_NUMBER ) const;
+	bool Equals( const TRotator<T>& R, T Tolerance = UE_KINDA_SMALL_NUMBER ) const;
 
 	/**
 	 * Adds to each component of the rotator.
diff --git a/Engine/Source/Runtime/Core/Public/Math/SHMath.h b/Engine/Source/Runtime/Core/Public/Math/SHMath.h
--- a/Engine/Source/Runtime/Core/Public/Math/SHMath.h
+++ b/Engine/Source/Runtime/Core/Public/Math/SHMath.h
@@ -260,7 +260,7 @@ public:
 	void Normalize()
 	{
 		const float Integral = CalcIntegral();
-		if(Integral > DELTA)
+		if(Integral > UE_DELTA)
 		{
 			*this /= Integral;
 		}
@@ -307,9 +307,9 @@ public:
 
 		// These formula are scaling factors for each SH band that convolve a SH with the circularly symmetric function
 		// max(0,cos(theta))
-		float L0 =					PI;
-		float L1 =					2 * PI / 3;
-		float L2 =					PI / 4;
+		float L0 =					UE_PI;
+		float L1 =					2 * UE_PI / 3;
+		float L2 =					UE_PI / 4;
 
 		// Multiply the coefficients in each band with the appropriate band scaling factor.
 		for(int32 BasisIndex = 0;BasisIndex < MaxSHBasis;BasisIndex++)
@@ -349,7 +349,7 @@ public:
 		// Multiply the result by the phi-dependent part of the SH bases.
 		// Skip this for X=0 and Y=0, because atan will be undefined and
 		// we know the Vector will be (0,0,+1) or (0,0,-1).
-		if (FMath::Abs(Vector.X) > KINDA_SMALL_NUMBER || FMath::Abs(Vector.Y) > KINDA_SMALL_NUMBER)
+		if (FMath::Abs(Vector.X) > UE_KINDA_SMALL_NUMBER || FMath::Abs(Vector.Y) > UE_KINDA_SMALL_NUMBER)
 		{
 			const float	Phi = (float)FMath::Atan2(Vector.Y, Vector.X);
 
@@ -379,7 +379,7 @@ public:
 	static TSHVector AmbientFunction()
 	{
 		TSHVector AmbientFunctionSH;
-		AmbientFunctionSH.V[0] = 1.0f / (2.0f * FMath::Sqrt(PI));
+		AmbientFunctionSH.V[0] = 1.0f / (2.0f * FMath::Sqrt(UE_PI));
 		return AmbientFunctionSH;
 	}
 
@@ -440,7 +440,7 @@ public:
 			float delta = -f / fd;
 			Lambda += delta;
 
-			if (FMath::Abs(delta) < KINDA_SMALL_NUMBER)
+			if (FMath::Abs(delta) < UE_KINDA_SMALL_NUMBER)
 			{
 				break;
 			}
diff --git a/Engine/Source/Runtime/Core/Public/Math/Sphere.h b/Engine/Source/Runtime/Core/Public/Math/Sphere.h
--- a/Engine/Source/Runtime/Core/Public/Math/Sphere.h
+++ b/Engine/Source/Runtime/Core/Public/Math/Sphere.h
@@ -90,7 +90,7 @@ public:
 	 * @param Tolerance Error Tolerance.
 	 * @return true if spheres are equal within specified tolerance, otherwise false.
 	 */
-	bool Equals(const TSphere<T>& Sphere, T Tolerance = KINDA_SMALL_NUMBER) const
+	bool Equals(const TSphere<T>& Sphere, T Tolerance = UE_KINDA_SMALL_NUMBER) const
 	{
 		return Center.Equals(Sphere.Center, Tolerance) && FMath::Abs(W - Sphere.W) <= Tolerance;
 	}
@@ -102,7 +102,7 @@ public:
 	 * @param Tolerance Error Tolerance.
 	 * @return true if sphere is inside another, otherwise false.
 	 */
-	bool IsInside(const TSphere<T>& Other, T Tolerance = KINDA_SMALL_NUMBER) const
+	bool IsInside(const TSphere<T>& Other, T Tolerance = UE_KINDA_SMALL_NUMBER) const
 	{
 		if (W > Other.W + Tolerance)
 		{
@@ -118,7 +118,7 @@ public:
 	* @param In The location to test for inside the bounding volume.
 	* @return true if location is inside this volume.
 	*/
-	bool IsInside(const FVector& In, T Tolerance = KINDA_SMALL_NUMBER) const
+	bool IsInside(const FVector& In, T Tolerance = UE_KINDA_SMALL_NUMBER) const
 	{
 		return (Center - In).SizeSquared() <= FMath::Square(W + Tolerance);
 	}
@@ -130,7 +130,7 @@ public:
 	 * @param  Tolerance Error tolerance.
 	 * @return true if spheres intersect, false otherwise.
 	 */
-	FORCEINLINE bool Intersects(const TSphere<T>& Other, T Tolerance = KINDA_SMALL_NUMBER) const
+	FORCEINLINE bool Intersects(const TSphere<T>& Other, T Tolerance = UE_KINDA_SMALL_NUMBER) const
 	{
 		return (Center - Other.Center).SizeSquared() <= FMath::Square(FMath::Max(0.f, Other.W + W + Tolerance));
 	}
@@ -158,7 +158,7 @@ public:
 	 */
 	T GetVolume() const
 	{
-		return (4.f / 3.f) * PI * (W * W * W);
+		return (4.f / 3.f) * UE_PI * (W * W * W);
 	}
 
 	/**
@@ -265,7 +265,7 @@ TSphere<T>& TSphere<T>::operator+=(const TSphere<T>& Other)
 	TVector<T> ToOther = Other.Center - Center;
 	T DistSqr = ToOther.SizeSquared();
 
-	if (FMath::Square(W - Other.W) + KINDA_SMALL_NUMBER >= DistSqr)
+	if (FMath::Square(W - Other.W) + UE_KINDA_SMALL_NUMBER >= DistSqr)
 	{
 		// Pick the smaller
 		if (W < Other.W)
@@ -281,7 +281,7 @@ TSphere<T>& TSphere<T>::operator+=(const TSphere<T>& Other)
 		NewSphere.W = (Dist + Other.W + W) * 0.5f;
 		NewSphere.Center = Center;
 
-		if (Dist > SMALL_NUMBER)
+		if (Dist > UE_SMALL_NUMBER)
 		{
 			NewSphere.Center += ToOther * ((NewSphere.W - W) / Dist);
 		}
diff --git a/Engine/Source/Runtime/Core/Public/Math/TransformVectorized.h b/Engine/Source/Runtime/Core/Public/Math/TransformVectorized.h
--- a/Engine/Source/Runtime/Core/Public/Math/TransformVectorized.h
+++ b/Engine/Source/Runtime/Core/Public/Math/TransformVectorized.h
@@ -532,7 +532,7 @@ public:
 	FORCEINLINE static bool AnyHasNegativeScale(const TVector<T>& InScale3D, const TVector<T>& InOtherScale3D);
 	FORCEINLINE void ScaleTranslation(const TVector<T>& InScale3D);
 	FORCEINLINE void ScaleTranslation(const T& Scale);
-	FORCEINLINE void RemoveScaling(T Tolerance=SMALL_NUMBER);
+	FORCEINLINE void RemoveScaling(T Tolerance=UE_SMALL_NUMBER);
 	FORCEINLINE float GetMaximumAxisScale() const;
 	FORCEINLINE float GetMinimumAxisScale() const;
 	// Inverse does not work well with VQS format(in particular non-uniform), so removing it, but made two below functions to be used instead. 
@@ -588,7 +588,7 @@ public:
 	FORCEINLINE TVector<T>		GetScaledAxis(EAxis::Type InAxis) const;
 	FORCEINLINE TVector<T>		GetUnitAxis(EAxis::Type InAxis) const;
 	FORCEINLINE void		Mirror(EAxis::Type MirrorAxis, EAxis::Type FlipAxis);
-	FORCEINLINE static TVector<T>	GetSafeScaleReciprocal(const TVector<T>& InScale, T Tolerance=SMALL_NUMBER);
+	FORCEINLINE static TVector<T>	GetSafeScaleReciprocal(const TVector<T>& InScale, T Tolerance=UE_SMALL_NUMBER);
 
 
 	FORCEINLINE TVector<T> GetLocation() const
@@ -698,7 +698,7 @@ private:
 		return !!VectorAnyLesserThan(VectorMin(InScale3D, InOtherScale3D), (TransformVectorRegister)GlobalVectorConstants::FloatZero);
 	}
 
-	FORCEINLINE bool Private_RotationEquals( const TransformVectorRegister& InRotation, const FReal ToleranceScalar = KINDA_SMALL_NUMBER) const
+	FORCEINLINE bool Private_RotationEquals( const TransformVectorRegister& InRotation, const FReal ToleranceScalar = UE_KINDA_SMALL_NUMBER) const
 	{			
 		const TransformVectorRegister MyRotation = Rotation; // Load from persistent value, to avoid repeated loads.
 		const TransformVectorRegister Tolerance = VectorLoadFloat1(&ToleranceScalar);
@@ -709,7 +709,7 @@ private:
 		return !VectorAnyGreaterThan(RotationSub, Tolerance) || !VectorAnyGreaterThan(RotationAdd, Tolerance);
 	}
 
-	FORCEINLINE bool Private_TranslationEquals( const TransformVectorRegister& InTranslation, const FReal ToleranceScalar = KINDA_SMALL_NUMBER) const
+	FORCEINLINE bool Private_TranslationEquals( const TransformVectorRegister& InTranslation, const FReal ToleranceScalar = UE_KINDA_SMALL_NUMBER) const
 	{			
 		const TransformVectorRegister Tolerance = VectorLoadFloat1(&ToleranceScalar);
 		// !( (FMath::Abs(X-V.X) > Tolerance) || (FMath::Abs(Y-V.Y) > Tolerance) || (FMath::Abs(Z-V.Z) > Tolerance) )
@@ -717,7 +717,7 @@ private:
 		return !VectorAnyGreaterThan(TranslationDiff, Tolerance);
 	}
 
-	FORCEINLINE bool Private_Scale3DEquals( const TransformVectorRegister& InScale3D, const FReal ToleranceScalar = KINDA_SMALL_NUMBER) const
+	FORCEINLINE bool Private_Scale3DEquals( const TransformVectorRegister& InScale3D, const FReal ToleranceScalar = UE_KINDA_SMALL_NUMBER) const
 	{
 		const TransformVectorRegister Tolerance = VectorLoadFloat1(&ToleranceScalar);
 		// !( (FMath::Abs(X-V.X) > Tolerance) || (FMath::Abs(Y-V.Y) > Tolerance) || (FMath::Abs(Z-V.Z) > Tolerance) )
@@ -728,44 +728,44 @@ private:
 public:
 
 	// Test if A's rotation equals B's rotation, within a tolerance. Preferred over "A.GetRotation().Equals(B.GetRotation())" because it is faster on some platforms.
-	FORCEINLINE static bool AreRotationsEqual(const TTransform<T>& A, const TTransform<T>& B, FReal Tolerance=KINDA_SMALL_NUMBER)
+	FORCEINLINE static bool AreRotationsEqual(const TTransform<T>& A, const TTransform<T>& B, FReal Tolerance=UE_KINDA_SMALL_NUMBER)
 	{
 		return A.Private_RotationEquals(B.Rotation, Tolerance);
 	}
 
 	// Test if A's translation equals B's translation, within a tolerance. Preferred over "A.GetTranslation().Equals(B.GetTranslation())" because it avoids TransformVectorRegister->TVector<T> conversion.
-	FORCEINLINE static bool AreTranslationsEqual(const TTransform<T>& A, const TTransform<T>& B, FReal Tolerance=KINDA_SMALL_NUMBER)
+	FORCEINLINE static bool AreTranslationsEqual(const TTransform<T>& A, const TTransform<T>& B, FReal Tolerance=UE_KINDA_SMALL_NUMBER)
 	{
 		return A.Private_TranslationEquals(B.Translation, Tolerance);
 	}
 
 	// Test if A's scale equals B's scale, within a tolerance. Preferred over "A.GetScale3D().Equals(B.GetScale3D())" because it avoids TransformVectorRegister->TVector<T> conversion.
-	FORCEINLINE static bool AreScale3DsEqual(const TTransform<T>& A, const TTransform<T>& B, FReal Tolerance=KINDA_SMALL_NUMBER)
+	FORCEINLINE static bool AreScale3DsEqual(const TTransform<T>& A, const TTransform<T>& B, FReal Tolerance=UE_KINDA_SMALL_NUMBER)
 	{
 		return A.Private_Scale3DEquals(B.Scale3D, Tolerance);
 	}
 
 
 	// Test if this Transform's rotation equals another's rotation, within a tolerance. Preferred over "GetRotation().Equals(Other.GetRotation())" because it is faster on some platforms.
-	FORCEINLINE bool RotationEquals(const TTransform<T>& Other, FReal Tolerance=KINDA_SMALL_NUMBER) const
+	FORCEINLINE bool RotationEquals(const TTransform<T>& Other, FReal Tolerance=UE_KINDA_SMALL_NUMBER) const
 	{
 		return AreRotationsEqual(*this, Other, Tolerance);
 	}
 
 	// Test if this Transform's translation equals another's translation, within a tolerance. Preferred over "GetTranslation().Equals(Other.GetTranslation())" because it avoids TransformVectorRegister->TVector<T> conversion.
-	FORCEINLINE bool TranslationEquals(const TTransform<T>& Other, FReal Tolerance=KINDA_SMALL_NUMBER) const
+	FORCEINLINE bool TranslationEquals(const TTransform<T>& Other, FReal Tolerance=UE_KINDA_SMALL_NUMBER) const
 	{
 		return AreTranslationsEqual(*this, Other, Tolerance);
 	}
 
 	// Test if this Transform's scale equals another's scale, within a tolerance. Preferred over "GetScale3D().Equals(Other.GetScale3D())" because it avoids TransformVectorRegister->TVector<T> conversion.
-	FORCEINLINE bool Scale3DEquals(const TTransform<T>& Other, FReal Tolerance=KINDA_SMALL_NUMBER) const
+	FORCEINLINE bool Scale3DEquals(const TTransform<T>& Other, FReal Tolerance=UE_KINDA_SMALL_NUMBER) const
 	{
 		return AreScale3DsEqual(*this, Other, Tolerance);
 	}
 
 	// Test if all components of the transforms are equal, within a tolerance.
-	FORCEINLINE bool Equals(const TTransform<T>& Other, FReal Tolerance=KINDA_SMALL_NUMBER) const
+	FORCEINLINE bool Equals(const TTransform<T>& Other, FReal Tolerance=UE_KINDA_SMALL_NUMBER) const
 	{
 		return Private_TranslationEquals(Other.Translation, Tolerance) && Private_RotationEquals(Other.Rotation, Tolerance) && Private_Scale3DEquals(Other.Scale3D, Tolerance);
 	}
@@ -777,7 +777,7 @@ public:
 	}
 
 	// Test if rotation and translation components of the transforms are equal, within a tolerance.
-	FORCEINLINE bool EqualsNoScale(const TTransform<T>& Other, FReal Tolerance=KINDA_SMALL_NUMBER) const
+	FORCEINLINE bool EqualsNoScale(const TTransform<T>& Other, FReal Tolerance=UE_KINDA_SMALL_NUMBER) const
 	{
 		return Private_TranslationEquals(Other.Translation, Tolerance) && Private_RotationEquals(Other.Rotation, Tolerance);
 	}
@@ -1500,7 +1500,7 @@ FORCEINLINE void TTransform<T>::ScaleTranslation(const T& InScale)
 
 // this function is from matrix, and all it does is to normalize rotation portion
 template<typename T>
-FORCEINLINE void TTransform<T>::RemoveScaling(T Tolerance/*=SMALL_NUMBER*/)
+FORCEINLINE void TTransform<T>::RemoveScaling(T Tolerance/*=UE_SMALL_NUMBER*/)
 {
 	Scale3D = GlobalVectorConstants::Float1110;
 	NormalizeRotation();	
diff --git a/Engine/Source/Runtime/Core/Public/Math/TwoVectors.h b/Engine/Source/Runtime/Core/Public/Math/TwoVectors.h
--- a/Engine/Source/Runtime/Core/Public/Math/TwoVectors.h
+++ b/Engine/Source/Runtime/Core/Public/Math/TwoVectors.h
@@ -116,7 +116,7 @@ public:
 	 * @param Tolerance Error Tolerance.
 	 * @return true if two pairs are equal within specified tolerance, false otherwise..
 	 */
-	bool Equals( const FTwoVectors& V, float Tolerance = KINDA_SMALL_NUMBER ) const;
+	bool Equals( const FTwoVectors& V, float Tolerance = UE_KINDA_SMALL_NUMBER ) const;
 
 	// Unary operators.
 
diff --git a/Engine/Source/Runtime/Core/Public/Math/UnrealMathUtility.h b/Engine/Source/Runtime/Core/Public/Math/UnrealMathUtility.h
--- a/Engine/Source/Runtime/Core/Public/Math/UnrealMathUtility.h
+++ b/Engine/Source/Runtime/Core/Public/Math/UnrealMathUtility.h
@@ -43,37 +43,123 @@ class TRange;
 	Floating point constants.
 -----------------------------------------------------------------------------*/
 
-#undef  PI
-#define PI 					(3.1415926535897932f)	/* Extra digits if needed: 3.1415926535897932384626433832795f */
-#define SMALL_NUMBER		(1.e-8f)
-#define KINDA_SMALL_NUMBER	(1.e-4f)
-#define BIG_NUMBER			(3.4e+38f)
-#define EULERS_NUMBER       (2.71828182845904523536f)
-#define UE_GOLDEN_RATIO		(1.6180339887498948482045868343656381f)	/* Also known as divine proportion, golden mean, or golden section - related to the Fibonacci Sequence = (1 + sqrt(5)) / 2 */
-#define FLOAT_NON_FRACTIONAL (8388608.f) /* All single-precision floating point numbers greater than or equal to this have no fractional value. */
-
-
-#define DOUBLE_PI					(3.141592653589793238462643383279502884197169399)
-#define DOUBLE_SMALL_NUMBER			(1.e-8)
-#define DOUBLE_KINDA_SMALL_NUMBER	(1.e-4)
-#define DOUBLE_BIG_NUMBER			(3.4e+38)
-#define DOUBLE_EULERS_NUMBER		(2.7182818284590452353602874713526624977572)
-#define DOUBLE_UE_GOLDEN_RATIO		(1.6180339887498948482045868343656381)	/* Also known as divine proportion, golden mean, or golden section - related to the Fibonacci Sequence = (1 + sqrt(5)) / 2 */
-#define DOUBLE_NON_FRACTIONAL		(4503599627370496.0) /* All double-precision floating point numbers greater than or equal to this have no fractional value. 2^52 */
+// Define this to 1 in a module's .Build.cs to make legacy names issue deprecation warnings.
+#ifndef UE_DEPRECATE_LEGACY_MATH_CONSTANT_MACRO_NAMES
+	#define UE_DEPRECATE_LEGACY_MATH_CONSTANT_MACRO_NAMES 0
+#endif
+
+// Define this to 0 in a module's .Build.cs to stop that module recognizing the old identifiers.
+#ifndef UE_DEFINE_LEGACY_MATH_CONSTANT_MACRO_NAMES
+	#define UE_DEFINE_LEGACY_MATH_CONSTANT_MACRO_NAMES 1
+#endif
+
+// Process for fixing up a module's use of legacy names:
+//
+// - Define UE_DEPRECATE_LEGACY_MATH_CONSTANT_MACRO_NAMES=1 to the module's .Build.cs.
+// - Build and fix all warnings.
+// - Remove UE_DEPRECATE_LEGACY_MATH_CONSTANT_MACRO_NAMES=1 from the module's .Build.cs.
+// - Define UE_DEFINE_LEGACY_MATH_CONSTANT_MACRO_NAMES=0 from the module's .Build.cs to stop new usage of the legacy macros.
+
+#ifdef UE_DEFINE_LEGACY_MATH_CONSTANT_MACRO_NAMES
+	#if UE_DEPRECATE_LEGACY_MATH_CONSTANT_MACRO_NAMES
+		#define UE_PRIVATE_MATH_DEPRECATION(Before, After) UE_DEPRECATED_MACRO(5.1, "The " #Before " macro has been deprecated in favor of " #After ".")
+	#else
+		#define UE_PRIVATE_MATH_DEPRECATION(Before, After)
+	#endif
+
+	#undef  PI
+	#define PI										UE_PRIVATE_MATH_DEPRECATION(PI										, UE_PI										) UE_PI
+	#define SMALL_NUMBER							UE_PRIVATE_MATH_DEPRECATION(SMALL_NUMBER							, UE_SMALL_NUMBER							) UE_SMALL_NUMBER
+	#define KINDA_SMALL_NUMBER						UE_PRIVATE_MATH_DEPRECATION(KINDA_SMALL_NUMBER						, UE_KINDA_SMALL_NUMBER						) UE_KINDA_SMALL_NUMBER
+	#define BIG_NUMBER								UE_PRIVATE_MATH_DEPRECATION(BIG_NUMBER								, UE_BIG_NUMBER								) UE_BIG_NUMBER
+	#define EULERS_NUMBER							UE_PRIVATE_MATH_DEPRECATION(EULERS_NUMBER							, UE_EULERS_NUMBER							) UE_EULERS_NUMBER
+	#define FLOAT_NON_FRACTIONAL					UE_PRIVATE_MATH_DEPRECATION(FLOAT_NON_FRACTIONAL					, UE_FLOAT_NON_FRACTIONAL					) UE_FLOAT_NON_FRACTIONAL
+	#define DOUBLE_PI								UE_PRIVATE_MATH_DEPRECATION(DOUBLE_PI								, UE_DOUBLE_PI								) UE_DOUBLE_PI
+	#define DOUBLE_SMALL_NUMBER						UE_PRIVATE_MATH_DEPRECATION(DOUBLE_SMALL_NUMBER						, UE_DOUBLE_SMALL_NUMBER					) UE_DOUBLE_SMALL_NUMBER
+	#define DOUBLE_KINDA_SMALL_NUMBER				UE_PRIVATE_MATH_DEPRECATION(DOUBLE_KINDA_SMALL_NUMBER				, UE_DOUBLE_KINDA_SMALL_NUMBER				) UE_DOUBLE_KINDA_SMALL_NUMBER
+	#define DOUBLE_BIG_NUMBER						UE_PRIVATE_MATH_DEPRECATION(DOUBLE_BIG_NUMBER						, UE_DOUBLE_BIG_NUMBER						) UE_DOUBLE_BIG_NUMBER
+	#define DOUBLE_EULERS_NUMBER					UE_PRIVATE_MATH_DEPRECATION(DOUBLE_EULERS_NUMBER					, UE_DOUBLE_EULERS_NUMBER					) UE_DOUBLE_EULERS_NUMBER
+	#define DOUBLE_UE_GOLDEN_RATIO					UE_PRIVATE_MATH_DEPRECATION(DOUBLE_UE_GOLDEN_RATIO					, UE_DOUBLE_GOLDEN_RATIO					) UE_DOUBLE_GOLDEN_RATIO
+	#define DOUBLE_NON_FRACTIONAL					UE_PRIVATE_MATH_DEPRECATION(DOUBLE_NON_FRACTIONAL					, UE_DOUBLE_NON_FRACTIONAL					) UE_DOUBLE_NON_FRACTIONAL
+	#define MAX_FLT									UE_PRIVATE_MATH_DEPRECATION(MAX_FLT									, UE_MAX_FLT								) UE_MAX_FLT
+	#define INV_PI									UE_PRIVATE_MATH_DEPRECATION(INV_PI									, UE_INV_PI									) UE_INV_PI
+	#define HALF_PI									UE_PRIVATE_MATH_DEPRECATION(HALF_PI									, UE_HALF_PI								) UE_HALF_PI
+	#define TWO_PI									UE_PRIVATE_MATH_DEPRECATION(TWO_PI									, UE_TWO_PI									) UE_TWO_PI
+	#define PI_SQUARED								UE_PRIVATE_MATH_DEPRECATION(PI_SQUARED								, UE_PI_SQUARED								) UE_PI_SQUARED
+	#define DOUBLE_INV_PI							UE_PRIVATE_MATH_DEPRECATION(DOUBLE_INV_PI							, UE_DOUBLE_INV_PI							) UE_DOUBLE_INV_PI
+	#define DOUBLE_HALF_PI							UE_PRIVATE_MATH_DEPRECATION(DOUBLE_HALF_PI							, UE_DOUBLE_HALF_PI							) UE_DOUBLE_HALF_PI
+	#define DOUBLE_TWO_PI							UE_PRIVATE_MATH_DEPRECATION(DOUBLE_TWO_PI							, UE_DOUBLE_TWO_PI							) UE_DOUBLE_TWO_PI
+	#define DOUBLE_PI_SQUARED						UE_PRIVATE_MATH_DEPRECATION(DOUBLE_PI_SQUARED						, UE_DOUBLE_PI_SQUARED						) UE_DOUBLE_PI_SQUARED
+	#define DOUBLE_UE_SQRT_2						UE_PRIVATE_MATH_DEPRECATION(DOUBLE_UE_SQRT_2						, UE_DOUBLE_SQRT_2							) UE_DOUBLE_SQRT_2
+	#define DOUBLE_UE_SQRT_3						UE_PRIVATE_MATH_DEPRECATION(DOUBLE_UE_SQRT_3						, UE_DOUBLE_SQRT_3							) UE_DOUBLE_SQRT_3
+	#define DOUBLE_UE_INV_SQRT_2					UE_PRIVATE_MATH_DEPRECATION(DOUBLE_UE_INV_SQRT_2					, UE_DOUBLE_INV_SQRT_2						) UE_DOUBLE_INV_SQRT_2
+	#define DOUBLE_UE_INV_SQRT_3					UE_PRIVATE_MATH_DEPRECATION(DOUBLE_UE_INV_SQRT_3					, UE_DOUBLE_INV_SQRT_3						) UE_DOUBLE_INV_SQRT_3
+	#define DOUBLE_UE_HALF_SQRT_2					UE_PRIVATE_MATH_DEPRECATION(DOUBLE_UE_HALF_SQRT_2					, UE_DOUBLE_HALF_SQRT_2						) UE_DOUBLE_HALF_SQRT_2
+	#define DOUBLE_UE_HALF_SQRT_3					UE_PRIVATE_MATH_DEPRECATION(DOUBLE_UE_HALF_SQRT_3					, UE_DOUBLE_HALF_SQRT_3						) UE_DOUBLE_HALF_SQRT_3
+	#define DELTA									UE_PRIVATE_MATH_DEPRECATION(DELTA									, UE_DELTA									) UE_DELTA
+	#define DOUBLE_DELTA							UE_PRIVATE_MATH_DEPRECATION(DOUBLE_DELTA							, UE_DOUBLE_DELTA							) UE_DOUBLE_DELTA
+	#define FLOAT_NORMAL_THRESH						UE_PRIVATE_MATH_DEPRECATION(FLOAT_NORMAL_THRESH						, UE_FLOAT_NORMAL_THRESH					) UE_FLOAT_NORMAL_THRESH
+	#define DOUBLE_NORMAL_THRESH					UE_PRIVATE_MATH_DEPRECATION(DOUBLE_NORMAL_THRESH					, UE_DOUBLE_NORMAL_THRESH					) UE_DOUBLE_NORMAL_THRESH
+	#define THRESH_POINT_ON_PLANE					UE_PRIVATE_MATH_DEPRECATION(THRESH_POINT_ON_PLANE					, UE_THRESH_POINT_ON_PLANE					) UE_THRESH_POINT_ON_PLANE
+	#define THRESH_POINT_ON_SIDE					UE_PRIVATE_MATH_DEPRECATION(THRESH_POINT_ON_SIDE					, UE_THRESH_POINT_ON_SIDE					) UE_THRESH_POINT_ON_SIDE
+	#define THRESH_POINTS_ARE_SAME					UE_PRIVATE_MATH_DEPRECATION(THRESH_POINTS_ARE_SAME					, UE_THRESH_POINTS_ARE_SAME					) UE_THRESH_POINTS_ARE_SAME
+	#define THRESH_POINTS_ARE_NEAR					UE_PRIVATE_MATH_DEPRECATION(THRESH_POINTS_ARE_NEAR					, UE_THRESH_POINTS_ARE_NEAR					) UE_THRESH_POINTS_ARE_NEAR
+	#define THRESH_NORMALS_ARE_SAME					UE_PRIVATE_MATH_DEPRECATION(THRESH_NORMALS_ARE_SAME					, UE_THRESH_NORMALS_ARE_SAME				) UE_THRESH_NORMALS_ARE_SAME
+	#define THRESH_UVS_ARE_SAME						UE_PRIVATE_MATH_DEPRECATION(THRESH_UVS_ARE_SAME						, UE_THRESH_UVS_ARE_SAME					) UE_THRESH_UVS_ARE_SAME
+	#define THRESH_VECTORS_ARE_NEAR					UE_PRIVATE_MATH_DEPRECATION(THRESH_VECTORS_ARE_NEAR					, UE_THRESH_VECTORS_ARE_NEAR				) UE_THRESH_VECTORS_ARE_NEAR
+	#define THRESH_SPLIT_POLY_WITH_PLANE			UE_PRIVATE_MATH_DEPRECATION(THRESH_SPLIT_POLY_WITH_PLANE			, UE_THRESH_SPLIT_POLY_WITH_PLANE			) UE_THRESH_SPLIT_POLY_WITH_PLANE
+	#define THRESH_SPLIT_POLY_PRECISELY				UE_PRIVATE_MATH_DEPRECATION(THRESH_SPLIT_POLY_PRECISELY				, UE_THRESH_SPLIT_POLY_PRECISELY			) UE_THRESH_SPLIT_POLY_PRECISELY
+	#define THRESH_ZERO_NORM_SQUARED				UE_PRIVATE_MATH_DEPRECATION(THRESH_ZERO_NORM_SQUARED				, UE_THRESH_ZERO_NORM_SQUARED				) UE_THRESH_ZERO_NORM_SQUARED
+	#define THRESH_NORMALS_ARE_PARALLEL				UE_PRIVATE_MATH_DEPRECATION(THRESH_NORMALS_ARE_PARALLEL				, UE_THRESH_NORMALS_ARE_PARALLEL			) UE_THRESH_NORMALS_ARE_PARALLEL
+	#define THRESH_NORMALS_ARE_ORTHOGONAL			UE_PRIVATE_MATH_DEPRECATION(THRESH_NORMALS_ARE_ORTHOGONAL			, UE_THRESH_NORMALS_ARE_ORTHOGONAL			) UE_THRESH_NORMALS_ARE_ORTHOGONAL
+	#define THRESH_VECTOR_NORMALIZED				UE_PRIVATE_MATH_DEPRECATION(THRESH_VECTOR_NORMALIZED				, UE_THRESH_VECTOR_NORMALIZED				) UE_THRESH_VECTOR_NORMALIZED
+	#define THRESH_QUAT_NORMALIZED					UE_PRIVATE_MATH_DEPRECATION(THRESH_QUAT_NORMALIZED					, UE_THRESH_QUAT_NORMALIZED					) UE_THRESH_QUAT_NORMALIZED
+	#define DOUBLE_THRESH_POINT_ON_PLANE			UE_PRIVATE_MATH_DEPRECATION(DOUBLE_THRESH_POINT_ON_PLANE			, UE_DOUBLE_THRESH_POINT_ON_PLANE			) UE_DOUBLE_THRESH_POINT_ON_PLANE
+	#define DOUBLE_THRESH_POINT_ON_SIDE				UE_PRIVATE_MATH_DEPRECATION(DOUBLE_THRESH_POINT_ON_SIDE				, UE_DOUBLE_THRESH_POINT_ON_SIDE			) UE_DOUBLE_THRESH_POINT_ON_SIDE
+	#define DOUBLE_THRESH_POINTS_ARE_SAME			UE_PRIVATE_MATH_DEPRECATION(DOUBLE_THRESH_POINTS_ARE_SAME			, UE_DOUBLE_THRESH_POINTS_ARE_SAME			) UE_DOUBLE_THRESH_POINTS_ARE_SAME
+	#define DOUBLE_THRESH_POINTS_ARE_NEAR			UE_PRIVATE_MATH_DEPRECATION(DOUBLE_THRESH_POINTS_ARE_NEAR			, UE_DOUBLE_THRESH_POINTS_ARE_NEAR			) UE_DOUBLE_THRESH_POINTS_ARE_NEAR
+	#define DOUBLE_THRESH_NORMALS_ARE_SAME			UE_PRIVATE_MATH_DEPRECATION(DOUBLE_THRESH_NORMALS_ARE_SAME			, UE_DOUBLE_THRESH_NORMALS_ARE_SAME			) UE_DOUBLE_THRESH_NORMALS_ARE_SAME
+	#define DOUBLE_THRESH_UVS_ARE_SAME				UE_PRIVATE_MATH_DEPRECATION(DOUBLE_THRESH_UVS_ARE_SAME				, UE_DOUBLE_THRESH_UVS_ARE_SAME				) UE_DOUBLE_THRESH_UVS_ARE_SAME
+	#define DOUBLE_THRESH_VECTORS_ARE_NEAR			UE_PRIVATE_MATH_DEPRECATION(DOUBLE_THRESH_VECTORS_ARE_NEAR			, UE_DOUBLE_THRESH_VECTORS_ARE_NEAR			) UE_DOUBLE_THRESH_VECTORS_ARE_NEAR
+	#define DOUBLE_THRESH_SPLIT_POLY_WITH_PLANE		UE_PRIVATE_MATH_DEPRECATION(DOUBLE_THRESH_SPLIT_POLY_WITH_PLANE		, UE_DOUBLE_THRESH_SPLIT_POLY_WITH_PLANE	) UE_DOUBLE_THRESH_SPLIT_POLY_WITH_PLANE
+	#define DOUBLE_THRESH_SPLIT_POLY_PRECISELY		UE_PRIVATE_MATH_DEPRECATION(DOUBLE_THRESH_SPLIT_POLY_PRECISELY		, UE_DOUBLE_THRESH_SPLIT_POLY_PRECISELY		) UE_DOUBLE_THRESH_SPLIT_POLY_PRECISELY
+	#define DOUBLE_THRESH_ZERO_NORM_SQUARED			UE_PRIVATE_MATH_DEPRECATION(DOUBLE_THRESH_ZERO_NORM_SQUARED			, UE_DOUBLE_THRESH_ZERO_NORM_SQUARED		) UE_DOUBLE_THRESH_ZERO_NORM_SQUARED
+	#define DOUBLE_THRESH_NORMALS_ARE_PARALLEL		UE_PRIVATE_MATH_DEPRECATION(DOUBLE_THRESH_NORMALS_ARE_PARALLEL		, UE_DOUBLE_THRESH_NORMALS_ARE_PARALLEL		) UE_DOUBLE_THRESH_NORMALS_ARE_PARALLEL
+	#define DOUBLE_THRESH_NORMALS_ARE_ORTHOGONAL	UE_PRIVATE_MATH_DEPRECATION(DOUBLE_THRESH_NORMALS_ARE_ORTHOGONAL	, UE_DOUBLE_THRESH_NORMALS_ARE_ORTHOGONAL	) UE_DOUBLE_THRESH_NORMALS_ARE_ORTHOGONAL
+	#define DOUBLE_THRESH_VECTOR_NORMALIZED			UE_PRIVATE_MATH_DEPRECATION(DOUBLE_THRESH_VECTOR_NORMALIZED			, UE_DOUBLE_THRESH_VECTOR_NORMALIZED		) UE_DOUBLE_THRESH_VECTOR_NORMALIZED
+	#define DOUBLE_THRESH_QUAT_NORMALIZED			UE_PRIVATE_MATH_DEPRECATION(DOUBLE_THRESH_QUAT_NORMALIZED			, UE_DOUBLE_THRESH_QUAT_NORMALIZED			) UE_DOUBLE_THRESH_QUAT_NORMALIZED
+#endif
+
+#define UE_PI 					(3.1415926535897932f)	/* Extra digits if needed: 3.1415926535897932384626433832795f */
+#define UE_SMALL_NUMBER			(1.e-8f)
+#define UE_KINDA_SMALL_NUMBER	(1.e-4f)
+#define UE_BIG_NUMBER			(3.4e+38f)
+#define UE_EULERS_NUMBER		(2.71828182845904523536f)
+#define UE_GOLDEN_RATIO			(1.6180339887498948482045868343656381f)	/* Also known as divine proportion, golden mean, or golden section - related to the Fibonacci Sequence = (1 + sqrt(5)) / 2 */
+#define UE_FLOAT_NON_FRACTIONAL (8388608.f) /* All single-precision floating point numbers greater than or equal to this have no fractional value. */
+
+
+#define UE_DOUBLE_PI					(3.141592653589793238462643383279502884197169399)
+#define UE_DOUBLE_SMALL_NUMBER			(1.e-8)
+#define UE_DOUBLE_KINDA_SMALL_NUMBER	(1.e-4)
+#define UE_DOUBLE_BIG_NUMBER			(3.4e+38)
+#define UE_DOUBLE_EULERS_NUMBER			(2.7182818284590452353602874713526624977572)
+#define UE_DOUBLE_GOLDEN_RATIO			(1.6180339887498948482045868343656381)	/* Also known as divine proportion, golden mean, or golden section - related to the Fibonacci Sequence = (1 + sqrt(5)) / 2 */
+#define UE_DOUBLE_NON_FRACTIONAL		(4503599627370496.0) /* All double-precision floating point numbers greater than or equal to this have no fractional value. 2^52 */
 
 // Copied from float.h
-#define MAX_FLT 3.402823466e+38F
+#define UE_MAX_FLT 3.402823466e+38F
 
 // Aux constants.
-#define INV_PI			(0.31830988618f)
-#define HALF_PI			(1.57079632679f)
-#define TWO_PI			(6.28318530717f)
-#define PI_SQUARED		(9.86960440108f)
+#define UE_INV_PI			(0.31830988618f)
+#define UE_HALF_PI			(1.57079632679f)
+#define UE_TWO_PI			(6.28318530717f)
+#define UE_PI_SQUARED		(9.86960440108f)
 
-#define DOUBLE_INV_PI		(0.31830988618379067154)
-#define DOUBLE_HALF_PI		(1.57079632679489661923)
-#define DOUBLE_TWO_PI		(6.28318530717958647692)
-#define DOUBLE_PI_SQUARED	(9.86960440108935861883)
+#define UE_DOUBLE_INV_PI		(0.31830988618379067154)
+#define UE_DOUBLE_HALF_PI		(1.57079632679489661923)
+#define UE_DOUBLE_TWO_PI		(6.28318530717958647692)
+#define UE_DOUBLE_PI_SQUARED	(9.86960440108935861883)
 
 // Common square roots
 #define UE_SQRT_2		(1.4142135623730950488016887242097f)
@@ -83,64 +169,64 @@ class TRange;
 #define UE_HALF_SQRT_2	(0.70710678118654752440084436210485f)
 #define UE_HALF_SQRT_3	(0.86602540378443864676372317075294f)
 
-#define DOUBLE_UE_SQRT_2		(1.4142135623730950488016887242097)
-#define DOUBLE_UE_SQRT_3		(1.7320508075688772935274463415059)
-#define DOUBLE_UE_INV_SQRT_2	(0.70710678118654752440084436210485)
-#define DOUBLE_UE_INV_SQRT_3	(0.57735026918962576450914878050196)
-#define DOUBLE_UE_HALF_SQRT_2	(0.70710678118654752440084436210485)
-#define DOUBLE_UE_HALF_SQRT_3	(0.86602540378443864676372317075294)
+#define UE_DOUBLE_SQRT_2		(1.4142135623730950488016887242097)
+#define UE_DOUBLE_SQRT_3		(1.7320508075688772935274463415059)
+#define UE_DOUBLE_INV_SQRT_2	(0.70710678118654752440084436210485)
+#define UE_DOUBLE_INV_SQRT_3	(0.57735026918962576450914878050196)
+#define UE_DOUBLE_HALF_SQRT_2	(0.70710678118654752440084436210485)
+#define UE_DOUBLE_HALF_SQRT_3	(0.86602540378443864676372317075294)
 
 
 // Magic numbers for numerical precision.
-#define DELTA			(0.00001f)
-#define DOUBLE_DELTA	(0.00001 )
+#define UE_DELTA		(0.00001f)
+#define UE_DOUBLE_DELTA	(0.00001 )
 
 /**
  * Lengths of normalized vectors (These are half their maximum values
  * to assure that dot products with normalized vectors don't overflow).
  */
-#define FLOAT_NORMAL_THRESH				(0.0001f)
-#define DOUBLE_NORMAL_THRESH			(0.0001)
+#define UE_FLOAT_NORMAL_THRESH			(0.0001f)
+#define UE_DOUBLE_NORMAL_THRESH			(0.0001)
 
 //
 // Magic numbers for numerical precision.
 //
-#define THRESH_POINT_ON_PLANE			(0.10f)		/* Thickness of plane for front/back/inside test */
-#define THRESH_POINT_ON_SIDE			(0.20f)		/* Thickness of polygon side's side-plane for point-inside/outside/on side test */
-#define THRESH_POINTS_ARE_SAME			(0.00002f)	/* Two points are same if within this distance */
-#define THRESH_POINTS_ARE_NEAR			(0.015f)	/* Two points are near if within this distance and can be combined if imprecise math is ok */
-#define THRESH_NORMALS_ARE_SAME			(0.00002f)	/* Two normal points are same if within this distance */
-#define THRESH_UVS_ARE_SAME			    (0.0009765625f)/* Two UV are same if within this threshold (1.0f/1024f) */
-													/* Making this too large results in incorrect CSG classification and disaster */
-#define THRESH_VECTORS_ARE_NEAR			(0.0004f)	/* Two vectors are near if within this distance and can be combined if imprecise math is ok */
-													/* Making this too large results in lighting problems due to inaccurate texture coordinates */
-#define THRESH_SPLIT_POLY_WITH_PLANE	(0.25f)		/* A plane splits a polygon in half */
-#define THRESH_SPLIT_POLY_PRECISELY		(0.01f)		/* A plane exactly splits a polygon */
-#define THRESH_ZERO_NORM_SQUARED		(0.0001f)	/* Size of a unit normal that is considered "zero", squared */
-#define THRESH_NORMALS_ARE_PARALLEL		(0.999845f)	/* Two unit vectors are parallel if abs(A dot B) is greater than or equal to this. This is roughly cosine(1.0 degrees). */
-#define THRESH_NORMALS_ARE_ORTHOGONAL	(0.017455f)	/* Two unit vectors are orthogonal (perpendicular) if abs(A dot B) is less than or equal this. This is roughly cosine(89.0 degrees). */
-
-#define THRESH_VECTOR_NORMALIZED		(0.01f)		/** Allowed error for a normalized vector (against squared magnitude) */
-#define THRESH_QUAT_NORMALIZED			(0.01f)		/** Allowed error for a normalized quaternion (against squared magnitude) */
+#define UE_THRESH_POINT_ON_PLANE				(0.10f)		/* Thickness of plane for front/back/inside test */
+#define UE_THRESH_POINT_ON_SIDE					(0.20f)		/* Thickness of polygon side's side-plane for point-inside/outside/on side test */
+#define UE_THRESH_POINTS_ARE_SAME				(0.00002f)	/* Two points are same if within this distance */
+#define UE_THRESH_POINTS_ARE_NEAR				(0.015f)	/* Two points are near if within this distance and can be combined if imprecise math is ok */
+#define UE_THRESH_NORMALS_ARE_SAME				(0.00002f)	/* Two normal points are same if within this distance */
+#define UE_THRESH_UVS_ARE_SAME					(0.0009765625f)/* Two UV are same if within this threshold (1.0f/1024f) */
+															/* Making this too large results in incorrect CSG classification and disaster */
+#define UE_THRESH_VECTORS_ARE_NEAR				(0.0004f)	/* Two vectors are near if within this distance and can be combined if imprecise math is ok */
+															/* Making this too large results in lighting problems due to inaccurate texture coordinates */
+#define UE_THRESH_SPLIT_POLY_WITH_PLANE			(0.25f)		/* A plane splits a polygon in half */
+#define UE_THRESH_SPLIT_POLY_PRECISELY			(0.01f)		/* A plane exactly splits a polygon */
+#define UE_THRESH_ZERO_NORM_SQUARED				(0.0001f)	/* Size of a unit normal that is considered "zero", squared */
+#define UE_THRESH_NORMALS_ARE_PARALLEL			(0.999845f)	/* Two unit vectors are parallel if abs(A dot B) is greater than or equal to this. This is roughly cosine(1.0 degrees). */
+#define UE_THRESH_NORMALS_ARE_ORTHOGONAL		(0.017455f)	/* Two unit vectors are orthogonal (perpendicular) if abs(A dot B) is less than or equal this. This is roughly cosine(89.0 degrees). */
+
+#define UE_THRESH_VECTOR_NORMALIZED				(0.01f)		/** Allowed error for a normalized vector (against squared magnitude) */
+#define UE_THRESH_QUAT_NORMALIZED				(0.01f)		/** Allowed error for a normalized quaternion (against squared magnitude) */
 
 // Double precision values
-#define DOUBLE_THRESH_POINT_ON_PLANE			(0.10)		/* Thickness of plane for front/back/inside test */
-#define DOUBLE_THRESH_POINT_ON_SIDE				(0.20)		/* Thickness of polygon side's side-plane for point-inside/outside/on side test */
-#define DOUBLE_THRESH_POINTS_ARE_SAME			(0.00002)	/* Two points are same if within this distance */
-#define DOUBLE_THRESH_POINTS_ARE_NEAR			(0.015)		/* Two points are near if within this distance and can be combined if imprecise math is ok */
-#define DOUBLE_THRESH_NORMALS_ARE_SAME			(0.00002)	/* Two normal points are same if within this distance */
-#define DOUBLE_THRESH_UVS_ARE_SAME			    (0.0009765625)/* Two UV are same if within this threshold (1.0/1024.0) */
+#define UE_DOUBLE_THRESH_POINT_ON_PLANE			(0.10)		/* Thickness of plane for front/back/inside test */
+#define UE_DOUBLE_THRESH_POINT_ON_SIDE			(0.20)		/* Thickness of polygon side's side-plane for point-inside/outside/on side test */
+#define UE_DOUBLE_THRESH_POINTS_ARE_SAME		(0.00002)	/* Two points are same if within this distance */
+#define UE_DOUBLE_THRESH_POINTS_ARE_NEAR		(0.015)		/* Two points are near if within this distance and can be combined if imprecise math is ok */
+#define UE_DOUBLE_THRESH_NORMALS_ARE_SAME		(0.00002)	/* Two normal points are same if within this distance */
+#define UE_DOUBLE_THRESH_UVS_ARE_SAME			(0.0009765625)/* Two UV are same if within this threshold (1.0/1024.0) */
 															/* Making this too large results in incorrect CSG classification and disaster */
-#define DOUBLE_THRESH_VECTORS_ARE_NEAR			(0.0004)	/* Two vectors are near if within this distance and can be combined if imprecise math is ok */
+#define UE_DOUBLE_THRESH_VECTORS_ARE_NEAR		(0.0004)	/* Two vectors are near if within this distance and can be combined if imprecise math is ok */
 															/* Making this too large results in lighting problems due to inaccurate texture coordinates */
-#define DOUBLE_THRESH_SPLIT_POLY_WITH_PLANE		(0.25)		/* A plane splits a polygon in half */
-#define DOUBLE_THRESH_SPLIT_POLY_PRECISELY		(0.01)		/* A plane exactly splits a polygon */
-#define DOUBLE_THRESH_ZERO_NORM_SQUARED			(0.0001)	/* Size of a unit normal that is considered "zero", squared */
-#define DOUBLE_THRESH_NORMALS_ARE_PARALLEL		(0.999845)	/* Two unit vectors are parallel if abs(A dot B) is greater than or equal to this. This is roughly cosine(1.0 degrees). */
-#define DOUBLE_THRESH_NORMALS_ARE_ORTHOGONAL	(0.017455)	/* Two unit vectors are orthogonal (perpendicular) if abs(A dot B) is less than or equal this. This is roughly cosine(89.0 degrees). */
+#define UE_DOUBLE_THRESH_SPLIT_POLY_WITH_PLANE	(0.25)		/* A plane splits a polygon in half */
+#define UE_DOUBLE_THRESH_SPLIT_POLY_PRECISELY	(0.01)		/* A plane exactly splits a polygon */
+#define UE_DOUBLE_THRESH_ZERO_NORM_SQUARED		(0.0001)	/* Size of a unit normal that is considered "zero", squared */
+#define UE_DOUBLE_THRESH_NORMALS_ARE_PARALLEL	(0.999845)	/* Two unit vectors are parallel if abs(A dot B) is greater than or equal to this. This is roughly cosine(1.0 degrees). */
+#define UE_DOUBLE_THRESH_NORMALS_ARE_ORTHOGONAL	(0.017455)	/* Two unit vectors are orthogonal (perpendicular) if abs(A dot B) is less than or equal this. This is roughly cosine(89.0 degrees). */
 
-#define DOUBLE_THRESH_VECTOR_NORMALIZED			(0.01)		/** Allowed error for a normalized vector (against squared magnitude) */
-#define DOUBLE_THRESH_QUAT_NORMALIZED			(0.01)		/** Allowed error for a normalized quaternion (against squared magnitude) */
+#define UE_DOUBLE_THRESH_VECTOR_NORMALIZED		(0.01)		/** Allowed error for a normalized vector (against squared magnitude) */
+#define UE_DOUBLE_THRESH_QUAT_NORMALIZED		(0.01)		/** Allowed error for a normalized quaternion (against squared magnitude) */
 
 /*-----------------------------------------------------------------------------
 	Global functions.
@@ -294,12 +380,12 @@ struct FMath : public FPlatformMath
 	 *	@param ErrorTolerance	Maximum allowed difference for considering them as 'nearly equal'
 	 *	@return					true if A and B are nearly equal
 	 */
-	UE_NODISCARD static FORCEINLINE bool IsNearlyEqual(float A, float B, float ErrorTolerance = SMALL_NUMBER)
+	UE_NODISCARD static FORCEINLINE bool IsNearlyEqual(float A, float B, float ErrorTolerance = UE_SMALL_NUMBER)
 	{
 		return Abs<float>( A - B ) <= ErrorTolerance;
 	}
 
-	UE_NODISCARD static FORCEINLINE bool IsNearlyEqual(double A, double B, double ErrorTolerance = DOUBLE_SMALL_NUMBER)
+	UE_NODISCARD static FORCEINLINE bool IsNearlyEqual(double A, double B, double ErrorTolerance = UE_DOUBLE_SMALL_NUMBER)
 	{
 		return Abs<double>(A - B) <= ErrorTolerance;
 	}
@@ -313,7 +399,7 @@ struct FMath : public FPlatformMath
 	 *	@param ErrorTolerance	Maximum allowed difference for considering Value as 'nearly zero'
 	 *	@return					true if Value is nearly zero
 	 */
-	UE_NODISCARD static FORCEINLINE bool IsNearlyZero(float Value, float ErrorTolerance = SMALL_NUMBER)
+	UE_NODISCARD static FORCEINLINE bool IsNearlyZero(float Value, float ErrorTolerance = UE_SMALL_NUMBER)
 	{
 		return Abs<float>( Value ) <= ErrorTolerance;
 	}
@@ -324,7 +410,7 @@ struct FMath : public FPlatformMath
 	 *	@param ErrorTolerance	Maximum allowed difference for considering Value as 'nearly zero'
 	 *	@return					true if Value is nearly zero
 	 */
-	UE_NODISCARD static FORCEINLINE bool IsNearlyZero(double Value, double ErrorTolerance = DOUBLE_SMALL_NUMBER)
+	UE_NODISCARD static FORCEINLINE bool IsNearlyZero(double Value, double ErrorTolerance = UE_DOUBLE_SMALL_NUMBER)
 	{
 		return Abs<double>( Value ) <= ErrorTolerance;
 	}
@@ -594,7 +680,7 @@ public:
 	static constexpr FORCEINLINE void SinCos(typename TDecay<T>::Type* ScalarSin, typename TDecay<T>::Type* ScalarCos, T  Value )
 	{
 		// Map Value to y in [-pi,pi], x = 2*pi*quotient + remainder.
-		T quotient = (INV_PI*0.5f)*Value;
+		T quotient = (UE_INV_PI*0.5f)*Value;
 		if (Value >= 0.0f)
 		{
 			quotient = (T)((int64)(quotient + 0.5f));
@@ -603,18 +689,18 @@ public:
 		{
 			quotient = (T)((int64)(quotient - 0.5f));
 		}
-		T y = Value - TWO_PI * quotient;
+		T y = Value - UE_TWO_PI * quotient;
 
 		// Map y to [-pi/2,pi/2] with sin(y) = sin(Value).
 		T sign;
-		if (y > HALF_PI)
+		if (y > UE_HALF_PI)
 		{
-			y = PI - y;
+			y = UE_PI - y;
 			sign = -1.0f;
 		}
-		else if (y < -HALF_PI)
+		else if (y < -UE_HALF_PI)
 		{
-			y = -PI - y;
+			y = -UE_PI - y;
 			sign = -1.0f;
 		}
 		else
@@ -691,13 +777,13 @@ public:
 	 * @return					Value in degrees.
 	 */
 	template<class T>
-	UE_NODISCARD static constexpr FORCEINLINE auto RadiansToDegrees(T const& RadVal) -> decltype(RadVal * (180.f / PI))
+	UE_NODISCARD static constexpr FORCEINLINE auto RadiansToDegrees(T const& RadVal) -> decltype(RadVal * (180.f / UE_PI))
 	{
-		return RadVal * (180.f / PI);
+		return RadVal * (180.f / UE_PI);
 	}
 
-	static FORCEINLINE float RadiansToDegrees(float const& RadVal) { return RadVal * (180.f / PI); }
-	static FORCEINLINE double RadiansToDegrees(double const& RadVal) { return RadVal * (180.0 / DOUBLE_PI); }
+	static FORCEINLINE float RadiansToDegrees(float const& RadVal) { return RadVal * (180.f / UE_PI); }
+	static FORCEINLINE double RadiansToDegrees(double const& RadVal) { return RadVal * (180.0 / UE_DOUBLE_PI); }
 
 	/** 
 	 * Converts degrees to radians.
@@ -705,13 +791,13 @@ public:
 	 * @return					Value in radians.
 	 */
 	template<class T>
-	UE_NODISCARD static constexpr FORCEINLINE auto DegreesToRadians(T const& DegVal) -> decltype(DegVal * (PI / 180.f))
+	UE_NODISCARD static constexpr FORCEINLINE auto DegreesToRadians(T const& DegVal) -> decltype(DegVal * (UE_PI / 180.f))
 	{
-		return DegVal * (PI / 180.f);
+		return DegVal * (UE_PI / 180.f);
 	}
 
-	static FORCEINLINE float DegreesToRadians(float const& DegVal) { return DegVal * (PI / 180.f); }
-	static FORCEINLINE double DegreesToRadians(double const& DegVal) { return DegVal * (DOUBLE_PI / 180.0); }
+	static FORCEINLINE float DegreesToRadians(float const& DegVal) { return DegVal * (UE_PI / 180.f); }
+	static FORCEINLINE double DegreesToRadians(double const& DegVal) { return DegVal * (UE_DOUBLE_PI / 180.0); }
 
 	/** 
 	 * Clamps an arbitrary angle to be between the given angles.  Will clamp to nearest boundary.
@@ -754,16 +840,16 @@ public:
 		auto Delta = A2 - A1;
 
 		// If change is larger than PI
-		if (Delta > PI)
+		if (Delta > UE_PI)
 		{
 			// Flip to negative equivalent
-			Delta = Delta - TWO_PI;
+			Delta = Delta - UE_TWO_PI;
 		}
-		else if (Delta < -PI)
+		else if (Delta < -UE_PI)
 		{
 			// Otherwise, if change is smaller than -PI
 			// Flip to positive equivalent
-			Delta = Delta + TWO_PI;
+			Delta = Delta + UE_TWO_PI;
 		}
 
 		// Return delta in [-PI,PI] range
@@ -774,14 +860,14 @@ public:
 	template<typename T, TEMPLATE_REQUIRES(TIsFloatingPoint<T>::Value)>
 	UE_NODISCARD static constexpr T UnwindRadians(T A)
 	{
-		while(A > PI)
+		while(A > UE_PI)
 		{
-			A -= TWO_PI;
+			A -= UE_TWO_PI;
 		}
 
-		while(A < -PI)
+		while(A < -UE_PI)
 		{
-			A += TWO_PI;
+			A += UE_TWO_PI;
 		}
 
 		return A;
@@ -1142,7 +1228,7 @@ public:
 	template< class T >
 	UE_NODISCARD static FORCEINLINE_DEBUGGABLE T InterpSinIn(const T& A, const T& B, float Alpha)
 	{
-		float const ModifiedAlpha = -1.f * Cos(Alpha * HALF_PI) + 1.f;
+		float const ModifiedAlpha = -1.f * Cos(Alpha * UE_HALF_PI) + 1.f;
 		return Lerp<T>(A, B, ModifiedAlpha);
 	}
 	
@@ -1150,7 +1236,7 @@ public:
 	template< class T >
 	UE_NODISCARD static FORCEINLINE_DEBUGGABLE T InterpSinOut(const T& A, const T& B, float Alpha)
 	{
-		float const ModifiedAlpha = Sin(Alpha * HALF_PI);
+		float const ModifiedAlpha = Sin(Alpha * UE_HALF_PI);
 		return Lerp<T>(A, B, ModifiedAlpha);
 	}
 
@@ -1313,7 +1399,7 @@ public:
 		const RetType Dist = Target - Current;
 
 		// If distance is too small, just set the desired location
-		if( FMath::Square(Dist) < SMALL_NUMBER )
+		if( FMath::Square(Dist) < UE_SMALL_NUMBER )
 		{
 			return static_cast<RetType>(Target);
 		}
@@ -1338,7 +1424,7 @@ public:
 		const RetType Dist = Target - Current;
 
 		// If distance is too small, just set the desired location
-		if( FMath::Square(Dist) < SMALL_NUMBER )
+		if( FMath::Square(Dist) < UE_SMALL_NUMBER )
 		{
 			return static_cast<RetType>(Target);
 		}
@@ -1392,7 +1478,7 @@ public:
         const float InDeltaTime,
         const float InSmoothingTime)
 	{
-		if (InSmoothingTime > KINDA_SMALL_NUMBER)
+		if (InSmoothingTime > UE_KINDA_SMALL_NUMBER)
 		{
 			const float A = InDeltaTime / InSmoothingTime;
 			float Exp = InvExpApprox(A);
@@ -1430,7 +1516,7 @@ public:
 		const float InDeltaTime,
 		const float InSmoothingTime)
 	{
-		if (InSmoothingTime > KINDA_SMALL_NUMBER)
+		if (InSmoothingTime > UE_KINDA_SMALL_NUMBER)
 		{
 			// The closed form solution for critically damped motion towards zero is:
 			//
@@ -1500,14 +1586,14 @@ public:
 			return;
 		}
 
-		float W = InUndampedFrequency * TWO_PI;
+		float W = InUndampedFrequency * UE_TWO_PI;
 		// Handle special cases
-		if (W < SMALL_NUMBER) // no strength which means no damping either
+		if (W < UE_SMALL_NUMBER) // no strength which means no damping either
 		{
 			InOutValue += InOutValueRate * InDeltaTime;
 			return;
 		}
-		else if (InDampingRatio < SMALL_NUMBER) // No damping at all
+		else if (InDampingRatio < UE_SMALL_NUMBER) // No damping at all
 		{
 			T Err = InOutValue - InTargetValue;
 			const T B = InOutValueRate / W;
@@ -1594,7 +1680,7 @@ public:
 		const float InSmoothingTime,
 		const float InDampingRatio)
 	{
-		if (InSmoothingTime < SMALL_NUMBER)
+		if (InSmoothingTime < UE_SMALL_NUMBER)
 		{
 			if (InDeltaTime <= 0.0f)
 			{
@@ -1606,7 +1692,7 @@ public:
 		}
 
 		// Undamped frequency
-		float UndampedFrequency = 1.0f / (PI * InSmoothingTime);
+		float UndampedFrequency = 1.0f / (UE_PI * InSmoothingTime);
 		SpringDamper(InOutValue, InOutValueRate, InTargetValue, InTargetValueRate, InDeltaTime, UndampedFrequency, InDampingRatio);
 	}
 
@@ -1621,7 +1707,7 @@ public:
 	 */
 	UE_NODISCARD static float MakePulsatingValue( const double InCurrentTime, const float InPulsesPerSecond, const float InPhase = 0.0f )
 	{
-		return 0.5f + 0.5f * FMath::Sin( ( ( 0.25f + InPhase ) * TWO_PI ) + ( (float)InCurrentTime * TWO_PI ) * InPulsesPerSecond );
+		return 0.5f + 0.5f * FMath::Sin( ( ( 0.25f + InPhase ) * UE_TWO_PI ) + ( (float)InCurrentTime * UE_TWO_PI ) * InPulsesPerSecond );
 	}
 
 	// Geometry intersection 
@@ -1962,8 +2048,8 @@ public:
 	 * @param Tolerance			Maximum allowed difference to 0.5 in order to truncate
 	 * @return					The truncated value
 	 */
-	UE_NODISCARD static CORE_API float TruncateToHalfIfClose(float F, float Tolerance = SMALL_NUMBER);
-	UE_NODISCARD static CORE_API double TruncateToHalfIfClose(double F, double Tolerance = SMALL_NUMBER);
+	UE_NODISCARD static CORE_API float TruncateToHalfIfClose(float F, float Tolerance = UE_SMALL_NUMBER);
+	UE_NODISCARD static CORE_API double TruncateToHalfIfClose(double F, double Tolerance = UE_SMALL_NUMBER);
 
 	/**
 	* Converts a floating point number to the nearest integer, equidistant ties go to the value which is closest to an even value: 1.5 becomes 2, 0.5 becomes 0
diff --git a/Engine/Source/Runtime/Core/Public/Math/UnrealMathVectorConstants.h b/Engine/Source/Runtime/Core/Public/Math/UnrealMathVectorConstants.h
--- a/Engine/Source/Runtime/Core/Public/Math/UnrealMathVectorConstants.h
+++ b/Engine/Source/Runtime/Core/Public/Math/UnrealMathVectorConstants.h
@@ -46,10 +46,10 @@ namespace GlobalVectorConstants
 	inline constexpr VectorRegister4Float FloatMinus1_111= MakeVectorRegisterFloatConstant( -1.f, 1.f, 1.f, 1.f );
 	inline constexpr VectorRegister4Float FloatOneHalf = MakeVectorRegisterFloatConstant( 0.5f, 0.5f, 0.5f, 0.5f );
 	inline constexpr VectorRegister4Float FloatMinusOneHalf = MakeVectorRegisterFloatConstant( -0.5f, -0.5f, -0.5f, -0.5f );
-	inline constexpr VectorRegister4Float KindaSmallNumber = MakeVectorRegisterFloatConstant( KINDA_SMALL_NUMBER, KINDA_SMALL_NUMBER, KINDA_SMALL_NUMBER, KINDA_SMALL_NUMBER );
-	inline constexpr VectorRegister4Float SmallNumber = MakeVectorRegisterFloatConstant( SMALL_NUMBER, SMALL_NUMBER, SMALL_NUMBER, SMALL_NUMBER );
-	inline constexpr VectorRegister4Float ThreshQuatNormalized = MakeVectorRegisterFloatConstant( THRESH_QUAT_NORMALIZED, THRESH_QUAT_NORMALIZED, THRESH_QUAT_NORMALIZED, THRESH_QUAT_NORMALIZED );
-	inline constexpr VectorRegister4Float BigNumber = MakeVectorRegisterFloatConstant(BIG_NUMBER, BIG_NUMBER, BIG_NUMBER, BIG_NUMBER);
+	inline constexpr VectorRegister4Float KindaSmallNumber = MakeVectorRegisterFloatConstant( UE_KINDA_SMALL_NUMBER, UE_KINDA_SMALL_NUMBER, UE_KINDA_SMALL_NUMBER, UE_KINDA_SMALL_NUMBER );
+	inline constexpr VectorRegister4Float SmallNumber = MakeVectorRegisterFloatConstant( UE_SMALL_NUMBER, UE_SMALL_NUMBER, UE_SMALL_NUMBER, UE_SMALL_NUMBER );
+	inline constexpr VectorRegister4Float ThreshQuatNormalized = MakeVectorRegisterFloatConstant( UE_THRESH_QUAT_NORMALIZED, UE_THRESH_QUAT_NORMALIZED, UE_THRESH_QUAT_NORMALIZED, UE_THRESH_QUAT_NORMALIZED );
+	inline constexpr VectorRegister4Float BigNumber = MakeVectorRegisterFloatConstant(UE_BIG_NUMBER, UE_BIG_NUMBER, UE_BIG_NUMBER, UE_BIG_NUMBER);
 
 	inline constexpr VectorRegister2Double DoubleOne2d = MakeVectorRegister2DoubleConstant(1.0, 1.0);
 	inline constexpr VectorRegister4Double DoubleOne = MakeVectorRegisterDoubleConstant(1.0, 1.0, 1.0, 1.0);
@@ -64,10 +64,10 @@ namespace GlobalVectorConstants
 	inline constexpr VectorRegister4Double DoubleMinus1_111 = MakeVectorRegisterDoubleConstant(-1., 1., 1., 1.);
 	inline constexpr VectorRegister4Double DoubleOneHalf = MakeVectorRegisterDoubleConstant(0.5, 0.5, 0.5, 0.5);
 	inline constexpr VectorRegister4Double DoubleMinusOneHalf = MakeVectorRegisterDoubleConstant(-0.5, -0.5, -0.5, -0.5);
-	inline constexpr VectorRegister4Double DoubleKindaSmallNumber = MakeVectorRegisterDoubleConstant(DOUBLE_KINDA_SMALL_NUMBER, DOUBLE_KINDA_SMALL_NUMBER, DOUBLE_KINDA_SMALL_NUMBER, DOUBLE_KINDA_SMALL_NUMBER);
-	inline constexpr VectorRegister4Double DoubleSmallNumber = MakeVectorRegisterDoubleConstant(DOUBLE_SMALL_NUMBER, DOUBLE_SMALL_NUMBER, DOUBLE_SMALL_NUMBER, DOUBLE_SMALL_NUMBER);
-	inline constexpr VectorRegister4Double DoubleThreshQuatNormalized = MakeVectorRegisterDoubleConstant(DOUBLE_THRESH_QUAT_NORMALIZED, DOUBLE_THRESH_QUAT_NORMALIZED, DOUBLE_THRESH_QUAT_NORMALIZED, DOUBLE_THRESH_QUAT_NORMALIZED);
-	inline constexpr VectorRegister4Double DoubleBigNumber = MakeVectorRegisterDoubleConstant(DOUBLE_BIG_NUMBER, DOUBLE_BIG_NUMBER, DOUBLE_BIG_NUMBER, DOUBLE_BIG_NUMBER);
+	inline constexpr VectorRegister4Double DoubleKindaSmallNumber = MakeVectorRegisterDoubleConstant(UE_DOUBLE_KINDA_SMALL_NUMBER, UE_DOUBLE_KINDA_SMALL_NUMBER, UE_DOUBLE_KINDA_SMALL_NUMBER, UE_DOUBLE_KINDA_SMALL_NUMBER);
+	inline constexpr VectorRegister4Double DoubleSmallNumber = MakeVectorRegisterDoubleConstant(UE_DOUBLE_SMALL_NUMBER, UE_DOUBLE_SMALL_NUMBER, UE_DOUBLE_SMALL_NUMBER, UE_DOUBLE_SMALL_NUMBER);
+	inline constexpr VectorRegister4Double DoubleThreshQuatNormalized = MakeVectorRegisterDoubleConstant(UE_DOUBLE_THRESH_QUAT_NORMALIZED, UE_DOUBLE_THRESH_QUAT_NORMALIZED, UE_DOUBLE_THRESH_QUAT_NORMALIZED, UE_DOUBLE_THRESH_QUAT_NORMALIZED);
+	inline constexpr VectorRegister4Double DoubleBigNumber = MakeVectorRegisterDoubleConstant(UE_DOUBLE_BIG_NUMBER, UE_DOUBLE_BIG_NUMBER, UE_DOUBLE_BIG_NUMBER, UE_DOUBLE_BIG_NUMBER);
 
 	inline constexpr VectorRegister Vector0001 = MakeVectorRegisterConstant(0.0, 0.0, 0.0, 1.0);
 	inline constexpr VectorRegister Vector1110 = MakeVectorRegisterConstant(1.0, 1.0, 1.0, 0.0);
@@ -87,12 +87,12 @@ namespace GlobalVectorConstants
 	inline constexpr VectorRegister4Double DOUBLE_QMULTI_SIGN_MASK1 = MakeVectorRegisterDoubleConstant(1., 1., -1., -1.);
 	inline constexpr VectorRegister4Double DOUBLE_QMULTI_SIGN_MASK2 = MakeVectorRegisterDoubleConstant(-1., 1., 1., -1.);
 
-	inline constexpr VectorRegister4Float DEG_TO_RAD = MakeVectorRegisterConstant(PI/(180.f), PI/(180.f), PI/(180.f), PI/(180.f));
-	inline constexpr VectorRegister4Float DEG_TO_RAD_HALF = MakeVectorRegisterConstant((PI/180.f)*0.5f, (PI/180.f)*0.5f, (PI/180.f)*0.5f, (PI/180.f)*0.5f);
-	inline constexpr VectorRegister4Float RAD_TO_DEG = MakeVectorRegisterConstant((180.f)/PI, (180.f)/PI, (180.f)/PI, (180.f)/PI);
-	inline constexpr VectorRegister4Double DOUBLE_DEG_TO_RAD = MakeVectorRegisterConstant(DOUBLE_PI/(180.), DOUBLE_PI/(180.), DOUBLE_PI/(180.), DOUBLE_PI/(180.));
-	inline constexpr VectorRegister4Double DOUBLE_DEG_TO_RAD_HALF = MakeVectorRegisterConstant((DOUBLE_PI/180.) * 0.5, (DOUBLE_PI/180.) * 0.5, (DOUBLE_PI/180.) * 0.5, (DOUBLE_PI/180.) * 0.5);
-	inline constexpr VectorRegister4Double DOUBLE_RAD_TO_DEG = MakeVectorRegisterConstant((180.)/DOUBLE_PI, (180.)/DOUBLE_PI, (180.)/DOUBLE_PI, (180.)/DOUBLE_PI);
+	inline constexpr VectorRegister4Float DEG_TO_RAD = MakeVectorRegisterConstant(UE_PI/(180.f), UE_PI/(180.f), UE_PI/(180.f), UE_PI/(180.f));
+	inline constexpr VectorRegister4Float DEG_TO_RAD_HALF = MakeVectorRegisterConstant((UE_PI/180.f)*0.5f, (UE_PI/180.f)*0.5f, (UE_PI/180.f)*0.5f, (UE_PI/180.f)*0.5f);
+	inline constexpr VectorRegister4Float RAD_TO_DEG = MakeVectorRegisterConstant((180.f)/UE_PI, (180.f)/UE_PI, (180.f)/UE_PI, (180.f)/UE_PI);
+	inline constexpr VectorRegister4Double DOUBLE_DEG_TO_RAD = MakeVectorRegisterConstant(UE_DOUBLE_PI/(180.), UE_DOUBLE_PI/(180.), UE_DOUBLE_PI/(180.), UE_DOUBLE_PI/(180.));
+	inline constexpr VectorRegister4Double DOUBLE_DEG_TO_RAD_HALF = MakeVectorRegisterConstant((UE_DOUBLE_PI/180.) * 0.5, (UE_DOUBLE_PI/180.) * 0.5, (UE_DOUBLE_PI/180.) * 0.5, (UE_DOUBLE_PI/180.) * 0.5);
+	inline constexpr VectorRegister4Double DOUBLE_RAD_TO_DEG = MakeVectorRegisterConstant((180.)/UE_DOUBLE_PI, (180.)/UE_DOUBLE_PI, (180.)/UE_DOUBLE_PI, (180.)/UE_DOUBLE_PI);
 
 	/** Bitmask to AND out the XYZ components in a vector */
 	inline VectorRegister4Float XYZMask()			{ return MakeVectorRegisterFloatMask( ~uint32(0), ~uint32(0), ~uint32(0), uint32(0)); }
@@ -115,19 +115,19 @@ namespace GlobalVectorConstants
 	inline VectorRegister4Float FloatInfinity()		{ return MakeVectorRegisterFloatMask((uint32)0x7F800000, (uint32)0x7F800000, (uint32)0x7F800000, (uint32)0x7F800000); }
 	inline VectorRegister4Double DoubleInfinity()	{ return MakeVectorRegisterDoubleMask((uint64)0x7FF0000000000000, (uint64)0x7FF0000000000000, (uint64)0x7FF0000000000000, (uint64)0x7FF0000000000000); }
 
-	inline constexpr VectorRegister4Float Pi = MakeVectorRegisterConstant(PI, PI, PI, PI);
-	inline constexpr VectorRegister4Float TwoPi = MakeVectorRegisterConstant(2.0f*PI, 2.0f*PI, 2.0f*PI, 2.0f*PI);
-	inline constexpr VectorRegister4Float PiByTwo = MakeVectorRegisterConstant(0.5f*PI, 0.5f*PI, 0.5f*PI, 0.5f*PI);
-	inline constexpr VectorRegister4Float PiByFour = MakeVectorRegisterConstant(0.25f*PI, 0.25f*PI, 0.25f*PI, 0.25f*PI);
-	inline constexpr VectorRegister4Float OneOverPi = MakeVectorRegisterConstant(1.0f / PI, 1.0f / PI, 1.0f / PI, 1.0f / PI);
-	inline constexpr VectorRegister4Float OneOverTwoPi = MakeVectorRegisterConstant(1.0f / (2.0f*PI), 1.0f / (2.0f*PI), 1.0f / (2.0f*PI), 1.0f / (2.0f*PI));
+	inline constexpr VectorRegister4Float Pi = MakeVectorRegisterConstant(UE_PI, UE_PI, UE_PI, UE_PI);
+	inline constexpr VectorRegister4Float TwoPi = MakeVectorRegisterConstant(2.0f*UE_PI, 2.0f*UE_PI, 2.0f*UE_PI, 2.0f*UE_PI);
+	inline constexpr VectorRegister4Float PiByTwo = MakeVectorRegisterConstant(0.5f*UE_PI, 0.5f*UE_PI, 0.5f*UE_PI, 0.5f*UE_PI);
+	inline constexpr VectorRegister4Float PiByFour = MakeVectorRegisterConstant(0.25f*UE_PI, 0.25f*UE_PI, 0.25f*UE_PI, 0.25f*UE_PI);
+	inline constexpr VectorRegister4Float OneOverPi = MakeVectorRegisterConstant(1.0f / UE_PI, 1.0f / UE_PI, 1.0f / UE_PI, 1.0f / UE_PI);
+	inline constexpr VectorRegister4Float OneOverTwoPi = MakeVectorRegisterConstant(1.0f / (2.0f*UE_PI), 1.0f / (2.0f*UE_PI), 1.0f / (2.0f*UE_PI), 1.0f / (2.0f*UE_PI));
 
-	inline constexpr VectorRegister4Double DoublePi = MakeVectorRegisterDoubleConstant(DOUBLE_PI, DOUBLE_PI, DOUBLE_PI, DOUBLE_PI);
-	inline constexpr VectorRegister4Double DoubleTwoPi = MakeVectorRegisterDoubleConstant(2.0 * DOUBLE_PI, 2.0 * DOUBLE_PI, 2.0 * DOUBLE_PI, 2.0 * DOUBLE_PI);
-	inline constexpr VectorRegister4Double DoublePiByTwo = MakeVectorRegisterDoubleConstant(0.5 * DOUBLE_PI, 0.5 * DOUBLE_PI, 0.5 * DOUBLE_PI, 0.5 * DOUBLE_PI);
-	inline constexpr VectorRegister4Double DoublePiByFour = MakeVectorRegisterDoubleConstant(0.25 * DOUBLE_PI, 0.25 * DOUBLE_PI, 0.25 * DOUBLE_PI, 0.25 * DOUBLE_PI);
-	inline constexpr VectorRegister4Double DoubleOneOverPi = MakeVectorRegisterDoubleConstant(1.0 / DOUBLE_PI, 1.0 / DOUBLE_PI, 1.0 / DOUBLE_PI, 1.0 / DOUBLE_PI);
-	inline constexpr VectorRegister4Double DoubleOneOverTwoPi = MakeVectorRegisterDoubleConstant(1.0 / (2.0 * DOUBLE_PI), 1.0 / (2.0 * DOUBLE_PI), 1.0 / (2.0 * DOUBLE_PI), 1.0 / (2.0 * DOUBLE_PI));
+	inline constexpr VectorRegister4Double DoublePi = MakeVectorRegisterDoubleConstant(UE_DOUBLE_PI, UE_DOUBLE_PI, UE_DOUBLE_PI, UE_DOUBLE_PI);
+	inline constexpr VectorRegister4Double DoubleTwoPi = MakeVectorRegisterDoubleConstant(2.0 * UE_DOUBLE_PI, 2.0 * UE_DOUBLE_PI, 2.0 * UE_DOUBLE_PI, 2.0 * UE_DOUBLE_PI);
+	inline constexpr VectorRegister4Double DoublePiByTwo = MakeVectorRegisterDoubleConstant(0.5 * UE_DOUBLE_PI, 0.5 * UE_DOUBLE_PI, 0.5 * UE_DOUBLE_PI, 0.5 * UE_DOUBLE_PI);
+	inline constexpr VectorRegister4Double DoublePiByFour = MakeVectorRegisterDoubleConstant(0.25 * UE_DOUBLE_PI, 0.25 * UE_DOUBLE_PI, 0.25 * UE_DOUBLE_PI, 0.25 * UE_DOUBLE_PI);
+	inline constexpr VectorRegister4Double DoubleOneOverPi = MakeVectorRegisterDoubleConstant(1.0 / UE_DOUBLE_PI, 1.0 / UE_DOUBLE_PI, 1.0 / UE_DOUBLE_PI, 1.0 / UE_DOUBLE_PI);
+	inline constexpr VectorRegister4Double DoubleOneOverTwoPi = MakeVectorRegisterDoubleConstant(1.0 / (2.0 * UE_DOUBLE_PI), 1.0 / (2.0 * UE_DOUBLE_PI), 1.0 / (2.0 * UE_DOUBLE_PI), 1.0 / (2.0 * UE_DOUBLE_PI));
 
 	inline constexpr VectorRegister4Float Float255 = MakeVectorRegisterConstant(255.0f, 255.0f, 255.0f, 255.0f);
 	inline constexpr VectorRegister4Float Float127 = MakeVectorRegisterConstant(127.0f, 127.0f, 127.0f, 127.0f);
@@ -142,8 +142,8 @@ namespace GlobalVectorConstants
 	inline constexpr VectorRegister4Double Double180 = MakeVectorRegisterDoubleConstant(180., 180., 180., 180.);
 
 	// All float numbers greater than or equal to this have no fractional value.
-	inline constexpr VectorRegister4Float FloatNonFractional = MakeVectorRegisterConstant(FLOAT_NON_FRACTIONAL, FLOAT_NON_FRACTIONAL, FLOAT_NON_FRACTIONAL, FLOAT_NON_FRACTIONAL);
-	inline constexpr VectorRegister4Double DoubleNonFractional = MakeVectorRegisterDoubleConstant(DOUBLE_NON_FRACTIONAL, DOUBLE_NON_FRACTIONAL, DOUBLE_NON_FRACTIONAL, DOUBLE_NON_FRACTIONAL);
+	inline constexpr VectorRegister4Float FloatNonFractional = MakeVectorRegisterConstant(UE_FLOAT_NON_FRACTIONAL, UE_FLOAT_NON_FRACTIONAL, UE_FLOAT_NON_FRACTIONAL, UE_FLOAT_NON_FRACTIONAL);
+	inline constexpr VectorRegister4Double DoubleNonFractional = MakeVectorRegisterDoubleConstant(UE_DOUBLE_NON_FRACTIONAL, UE_DOUBLE_NON_FRACTIONAL, UE_DOUBLE_NON_FRACTIONAL, UE_DOUBLE_NON_FRACTIONAL);
 
 	inline constexpr VectorRegister4Float FloatTwo = MakeVectorRegisterConstant(2.0f, 2.0f, 2.0f, 2.0f);
 	inline constexpr uint32 AlmostTwoBits = 0x3fffffff;
diff --git a/Engine/Source/Runtime/Core/Public/Math/Vector.h b/Engine/Source/Runtime/Core/Public/Math/Vector.h
--- a/Engine/Source/Runtime/Core/Public/Math/Vector.h
+++ b/Engine/Source/Runtime/Core/Public/Math/Vector.h
@@ -364,7 +364,7 @@ public:
      * @param Tolerance Error tolerance.
      * @return true if the vectors are equal within tolerance limits, false otherwise.
      */
-    bool Equals(const TVector<T>& V, T Tolerance=KINDA_SMALL_NUMBER) const;
+    bool Equals(const TVector<T>& V, T Tolerance=UE_KINDA_SMALL_NUMBER) const;
 
     /**
      * Checks whether all components of this vector are the same, within a tolerance.
@@ -372,7 +372,7 @@ public:
      * @param Tolerance Error tolerance.
      * @return true if the vectors are equal within tolerance limits, false otherwise.
      */
-    bool AllComponentsEqual(T Tolerance=KINDA_SMALL_NUMBER) const;
+    bool AllComponentsEqual(T Tolerance=UE_KINDA_SMALL_NUMBER) const;
 
     /**
      * Get a negated copy of the vector.
@@ -586,7 +586,7 @@ public:
      * @param Tolerance Error tolerance.
      * @return true if the vector is near to zero, false otherwise.
      */
-    bool IsNearlyZero(T Tolerance=KINDA_SMALL_NUMBER) const;
+    bool IsNearlyZero(T Tolerance=UE_KINDA_SMALL_NUMBER) const;
 
     /**
      * Checks whether all components of the vector are exactly zero.
@@ -601,7 +601,7 @@ public:
      * @param LengthSquaredTolerance Tolerance against squared length.
      * @return true if the vector is a unit vector within the specified tolerance.
      */
-    FORCEINLINE bool IsUnit(T LengthSquaredTolerance = KINDA_SMALL_NUMBER) const;
+    FORCEINLINE bool IsUnit(T LengthSquaredTolerance = UE_KINDA_SMALL_NUMBER) const;
 
     /**
      * Checks whether vector is normalized.
@@ -616,7 +616,7 @@ public:
      * @param Tolerance Minimum squared length of vector for normalization.
      * @return true if the vector was normalized correctly, false otherwise.
      */
-    bool Normalize(T Tolerance=SMALL_NUMBER);
+    bool Normalize(T Tolerance=UE_SMALL_NUMBER);
 
     /**
      * Calculates normalized version of vector without checking for zero length.
@@ -633,7 +633,7 @@ public:
      * @param Tolerance Minimum squared vector length.
      * @return A normalized copy if safe, ResultIfZero otherwise.
      */
-    TVector<T> GetSafeNormal(T Tolerance=SMALL_NUMBER, const TVector<T>& ResultIfZero = ZeroVector) const;
+    TVector<T> GetSafeNormal(T Tolerance=UE_SMALL_NUMBER, const TVector<T>& ResultIfZero = ZeroVector) const;
 
     /**
      * Gets a normalized copy of the 2D components of the vector, checking it is safe to do so. Z is set to zero. 
@@ -642,7 +642,7 @@ public:
      * @param Tolerance Minimum squared vector length.
      * @return Normalized copy if safe, otherwise returns ResultIfZero.
      */
-    TVector<T> GetSafeNormal2D(T Tolerance=SMALL_NUMBER, const TVector<T>& ResultIfZero = ZeroVector) const;
+    TVector<T> GetSafeNormal2D(T Tolerance=UE_SMALL_NUMBER, const TVector<T>& ResultIfZero = ZeroVector) const;
 
     /**
      * Util to convert this vector into a unit direction vector and its original length.
@@ -730,7 +730,7 @@ public:
      * @param Tolerance Specified Tolerance.
      * @return true if X == Y == Z within the specified tolerance.
      */
-    bool IsUniform(T Tolerance=KINDA_SMALL_NUMBER) const;
+    bool IsUniform(T Tolerance=UE_KINDA_SMALL_NUMBER) const;
 
     /**
      * Mirror a vector about a normal vector.
@@ -1039,7 +1039,7 @@ public:
      * @param  ParallelCosineThreshold Normals are parallel if absolute value of dot product (cosine of angle between them) is greater than or equal to this. For example: cos(1.0 degrees).
      * @return true if vectors are nearly parallel, false otherwise.
      */
-    static bool Parallel(const TVector<T>& Normal1, const TVector<T>& Normal2, T ParallelCosineThreshold = THRESH_NORMALS_ARE_PARALLEL);
+    static bool Parallel(const TVector<T>& Normal1, const TVector<T>& Normal2, T ParallelCosineThreshold = UE_THRESH_NORMALS_ARE_PARALLEL);
 
     /**
      * See if two normal vectors are coincident (nearly parallel and point in the same direction).
@@ -1049,7 +1049,7 @@ public:
      * @param  ParallelCosineThreshold Normals are coincident if dot product (cosine of angle between them) is greater than or equal to this. For example: cos(1.0 degrees).
      * @return true if vectors are coincident (nearly parallel and point in the same direction), false otherwise.
      */
-    static bool Coincident(const TVector<T>& Normal1, const TVector<T>& Normal2, T ParallelCosineThreshold = THRESH_NORMALS_ARE_PARALLEL);
+    static bool Coincident(const TVector<T>& Normal1, const TVector<T>& Normal2, T ParallelCosineThreshold = UE_THRESH_NORMALS_ARE_PARALLEL);
 
     /**
      * See if two normal vectors are nearly orthogonal (perpendicular), meaning the angle between them is close to 90 degrees.
@@ -1059,7 +1059,7 @@ public:
      * @param  OrthogonalCosineThreshold Normals are orthogonal if absolute value of dot product (cosine of angle between them) is less than or equal to this. For example: cos(89.0 degrees).
      * @return true if vectors are orthogonal (perpendicular), false otherwise.
      */
-    static bool Orthogonal(const TVector<T>& Normal1, const TVector<T>& Normal2, T OrthogonalCosineThreshold = THRESH_NORMALS_ARE_ORTHOGONAL);
+    static bool Orthogonal(const TVector<T>& Normal1, const TVector<T>& Normal2, T OrthogonalCosineThreshold = UE_THRESH_NORMALS_ARE_ORTHOGONAL);
 
     /**
      * See if two planes are coplanar. They are coplanar if the normals are nearly parallel and the planes include the same set of points.
@@ -1071,7 +1071,7 @@ public:
      * @param ParallelCosineThreshold Normals are parallel if absolute value of dot product is greater than or equal to this.
      * @return true if the planes are coplanar, false otherwise.
      */
-    static bool Coplanar(const TVector<T>& Base1, const TVector<T>& Normal1, const TVector<T>& Base2, const TVector<T>& Normal2, T ParallelCosineThreshold = THRESH_NORMALS_ARE_PARALLEL);
+    static bool Coplanar(const TVector<T>& Base1, const TVector<T>& Normal1, const TVector<T>& Base2, const TVector<T>& Normal2, T ParallelCosineThreshold = UE_THRESH_NORMALS_ARE_PARALLEL);
 
     /**
      * Triple product of three vectors: X dot (Y cross Z).
@@ -1259,13 +1259,13 @@ inline bool TVector<T>::PointsAreSame(const TVector<T>& P, const TVector<T>& Q)
 {
     T Temp;
     Temp=P.X-Q.X;
-    if((Temp > -THRESH_POINTS_ARE_SAME) && (Temp < THRESH_POINTS_ARE_SAME))
+    if((Temp > -UE_THRESH_POINTS_ARE_SAME) && (Temp < UE_THRESH_POINTS_ARE_SAME))
     {
         Temp=P.Y-Q.Y;
-        if((Temp > -THRESH_POINTS_ARE_SAME) && (Temp < THRESH_POINTS_ARE_SAME))
+        if((Temp > -UE_THRESH_POINTS_ARE_SAME) && (Temp < UE_THRESH_POINTS_ARE_SAME))
         {
             Temp=P.Z-Q.Z;
-            if((Temp > -THRESH_POINTS_ARE_SAME) && (Temp < THRESH_POINTS_ARE_SAME))
+            if((Temp > -UE_THRESH_POINTS_ARE_SAME) && (Temp < UE_THRESH_POINTS_ARE_SAME))
             {
                 return true;
             }
@@ -1335,7 +1335,7 @@ template<typename T>
 inline bool TVector<T>::Coplanar(const TVector<T>& Base1, const TVector<T>& Normal1, const TVector<T>& Base2, const TVector<T>& Normal2, T ParallelCosineThreshold)
 {
     if      (!TVector<T>::Parallel(Normal1,Normal2,ParallelCosineThreshold)) return false;
-    else if (FMath::Abs(TVector<T>::PointPlaneDist (Base2,Base1,Normal1)) > THRESH_POINT_ON_PLANE) return false;
+    else if (FMath::Abs(TVector<T>::PointPlaneDist (Base2,Base1,Normal1)) > UE_THRESH_POINT_ON_PLANE) return false;
     else return true;
 }
 
@@ -1351,13 +1351,13 @@ inline T TVector<T>::Triple(const TVector<T>& X, const TVector<T>& Y, const TVec
 template<typename T>
 inline TVector<T> TVector<T>::RadiansToDegrees(const TVector<T>& RadVector)
 {
-    return RadVector * (180.f / PI);
+    return RadVector * (180.f / UE_PI);
 }
 
 template<typename T>
 inline TVector<T> TVector<T>::DegreesToRadians(const TVector<T>& DegVector)
 {
-    return DegVector * (PI / 180.f);
+    return DegVector * (UE_PI / 180.f);
 }
 
 template<typename T>
@@ -1679,14 +1679,14 @@ FORCEINLINE bool TVector<T>::IsUnit(T LengthSquaredTolerance) const
 template<typename T>
 FORCEINLINE bool TVector<T>::IsNormalized() const
 {
-    return (FMath::Abs(1.f - SizeSquared()) < THRESH_VECTOR_NORMALIZED);
+    return (FMath::Abs(1.f - SizeSquared()) < UE_THRESH_VECTOR_NORMALIZED);
 }
 
 template<typename T>
 FORCEINLINE void TVector<T>::ToDirectionAndLength(TVector<T>& OutDir, double& OutLength) const
 {
 	OutLength = Size();
-	if (OutLength > SMALL_NUMBER)
+	if (OutLength > UE_SMALL_NUMBER)
 	{
 		T OneOverLength = T(1.0 / OutLength);
 		OutDir = TVector<T>(X * OneOverLength, Y * OneOverLength, Z * OneOverLength);
@@ -1701,7 +1701,7 @@ template<typename T>
 FORCEINLINE void TVector<T>::ToDirectionAndLength(TVector<T>& OutDir, float& OutLength) const
 {
 	OutLength = (float)Size();
-	if (OutLength > SMALL_NUMBER)
+	if (OutLength > UE_SMALL_NUMBER)
 	{
 		float OneOverLength = 1.0f / OutLength;
 		OutDir = TVector<T>(X * OneOverLength, Y * OneOverLength, Z * OneOverLength);
@@ -1778,7 +1778,7 @@ template<typename T>
 FORCEINLINE TVector<T> TVector<T>::GetClampedToSize(T Min, T Max) const
 {
     T VecSize = Size();
-    const TVector<T> VecDir = (VecSize > SMALL_NUMBER) ? (*this/VecSize) : ZeroVector;
+    const TVector<T> VecDir = (VecSize > UE_SMALL_NUMBER) ? (*this/VecSize) : ZeroVector;
 
     VecSize = FMath::Clamp(VecSize, Min, Max);
 
@@ -1789,7 +1789,7 @@ template<typename T>
 FORCEINLINE TVector<T> TVector<T>::GetClampedToSize2D(T Min, T Max) const
 {
     T VecSize2D = Size2D();
-    const TVector<T> VecDir = (VecSize2D > SMALL_NUMBER) ? (*this/VecSize2D) : ZeroVector;
+    const TVector<T> VecDir = (VecSize2D > UE_SMALL_NUMBER) ? (*this/VecSize2D) : ZeroVector;
 
     VecSize2D = FMath::Clamp(VecSize2D, Min, Max);
 
@@ -1799,7 +1799,7 @@ FORCEINLINE TVector<T> TVector<T>::GetClampedToSize2D(T Min, T Max) const
 template<typename T>
 FORCEINLINE TVector<T> TVector<T>::GetClampedToMaxSize(T MaxSize) const
 {
-    if (MaxSize < KINDA_SMALL_NUMBER)
+    if (MaxSize < UE_KINDA_SMALL_NUMBER)
     {
         return ZeroVector;
     }
@@ -1819,7 +1819,7 @@ FORCEINLINE TVector<T> TVector<T>::GetClampedToMaxSize(T MaxSize) const
 template<typename T>
 FORCEINLINE TVector<T> TVector<T>::GetClampedToMaxSize2D(T MaxSize) const
 {
-    if (MaxSize < KINDA_SMALL_NUMBER)
+    if (MaxSize < UE_KINDA_SMALL_NUMBER)
     {
         return TVector<T>(0.f, 0.f, Z);
     }
@@ -1901,7 +1901,7 @@ FORCEINLINE TVector<T> TVector<T>::Reciprocal() const
     }
     else 
     {
-        RecVector.X = BIG_NUMBER;
+        RecVector.X = UE_BIG_NUMBER;
     }
     if (Y!=0.f)
     {
@@ -1909,7 +1909,7 @@ FORCEINLINE TVector<T> TVector<T>::Reciprocal() const
     }
     else 
     {
-        RecVector.Y = BIG_NUMBER;
+        RecVector.Y = UE_BIG_NUMBER;
     }
     if (Z!=0.f)
     {
@@ -1917,7 +1917,7 @@ FORCEINLINE TVector<T> TVector<T>::Reciprocal() const
     }
     else 
     {
-        RecVector.Z = BIG_NUMBER;
+        RecVector.Z = UE_BIG_NUMBER;
     }
 
     return RecVector;
@@ -2035,7 +2035,7 @@ void TVector<T>::GenerateClusterCenters(TArray<TVector<T>>& Clusters, const TArr
 
 			// Iterate over all clusters to find closes one
 			int32 NearestClusterIndex = INDEX_NONE;
-			T NearestClusterDistSqr = BIG_NUMBER;
+			T NearestClusterDistSqr = UE_BIG_NUMBER;
 			for (int32 j = 0; j < Clusters.Num(); j++)
 			{
 				const T DistSqr = (Pos - Clusters[j]).SizeSquared();
@@ -2133,13 +2133,13 @@ void TVector<T>::CreateOrthonormalBasis(TVector<T>& XAxis, TVector<T>& YAxis, TV
 	YAxis -= (YAxis | ZAxis) / (ZAxis | ZAxis) * ZAxis;
 
 	// If the X axis was parallel to the Z axis, choose a vector which is orthogonal to the Y and Z axes.
-	if (XAxis.SizeSquared() < DELTA * DELTA)
+	if (XAxis.SizeSquared() < UE_DELTA * UE_DELTA)
 	{
 		XAxis = YAxis ^ ZAxis;
 	}
 
 	// If the Y axis was parallel to the Z axis, choose a vector which is orthogonal to the X and Z axes.
-	if (YAxis.SizeSquared() < DELTA * DELTA)
+	if (YAxis.SizeSquared() < UE_DELTA * UE_DELTA)
 	{
 		YAxis = XAxis ^ ZAxis;
 	}
@@ -2602,7 +2602,7 @@ inline FVector FMath::VRand()
         Result.Z = FRand() * 2.f - 1.f;
         L = Result.SizeSquared();
     }
-    while(L > 1.0f || L < KINDA_SMALL_NUMBER);
+    while(L > 1.0f || L < UE_KINDA_SMALL_NUMBER);
 
     return Result * (1.0f / Sqrt(L));
 }
diff --git a/Engine/Source/Runtime/Core/Public/Math/Vector2D.h b/Engine/Source/Runtime/Core/Public/Math/Vector2D.h
--- a/Engine/Source/Runtime/Core/Public/Math/Vector2D.h
+++ b/Engine/Source/Runtime/Core/Public/Math/Vector2D.h
@@ -412,7 +412,7 @@ public:
 	* @param Tolerance Error tolerance.
 	* @return true if the vectors are equal within specified tolerance, otherwise false.
 	*/
-	bool Equals(const TVector2<T>& V, T Tolerance=KINDA_SMALL_NUMBER) const;
+	bool Equals(const TVector2<T>& V, T Tolerance=UE_KINDA_SMALL_NUMBER) const;
 
 	/**
 	* Set the values of the vector directly.
@@ -498,7 +498,7 @@ public:
 	* @param Tolerance Minimum squared length of vector for normalization.
 	* @return A normalized copy of the vector if safe, (0,0) otherwise.
 	*/
-	TVector2<T> GetSafeNormal(T Tolerance=SMALL_NUMBER) const;
+	TVector2<T> GetSafeNormal(T Tolerance=UE_SMALL_NUMBER) const;
 
 	/**
 	* Normalize this vector in-place if it is large enough, set it to (0,0) otherwise.
@@ -506,7 +506,7 @@ public:
 	* @param Tolerance Minimum squared length of vector for normalization.
 	* @see GetSafeNormal()
 	*/
-	void Normalize(T Tolerance=SMALL_NUMBER);
+	void Normalize(T Tolerance=UE_SMALL_NUMBER);
 
 	/**
 	* Checks whether vector is near to zero within a specified tolerance.
@@ -514,7 +514,7 @@ public:
 	* @param Tolerance Error tolerance.
 	* @return true if vector is in tolerance to zero, otherwise false.
 	*/
-	bool IsNearlyZero(T Tolerance=KINDA_SMALL_NUMBER) const;
+	bool IsNearlyZero(T Tolerance=UE_KINDA_SMALL_NUMBER) const;
 
 	/**
 	* Util to convert this vector into a unit direction vector and its original length.
@@ -1069,7 +1069,7 @@ template<typename T>
 FORCEINLINE void TVector2<T>::ToDirectionAndLength(TVector2<T> &OutDir, double &OutLength) const
 {
 	OutLength = Size();
-	if (OutLength > SMALL_NUMBER)
+	if (OutLength > UE_SMALL_NUMBER)
 	{
 		T OneOverLength = 1.0f / OutLength;
 		OutDir = TVector2<T>(X*OneOverLength, Y*OneOverLength);
@@ -1084,7 +1084,7 @@ template<typename T>
 FORCEINLINE void TVector2<T>::ToDirectionAndLength(TVector2<T> &OutDir, float &OutLength) const
 {
 	OutLength = Size();
-	if (OutLength > SMALL_NUMBER)
+	if (OutLength > UE_SMALL_NUMBER)
 	{
 		float OneOverLength = 1.0f / OutLength;
 		OutDir = TVector2<T>(X*OneOverLength, Y*OneOverLength);
diff --git a/Engine/Source/Runtime/Core/Public/Math/Vector4.h b/Engine/Source/Runtime/Core/Public/Math/Vector4.h
--- a/Engine/Source/Runtime/Core/Public/Math/Vector4.h
+++ b/Engine/Source/Runtime/Core/Public/Math/Vector4.h
@@ -369,7 +369,7 @@ public:
 	 * @param Tolerance Error Tolerance.
 	 * @return true if the two vectors are equal within specified tolerance, otherwise false.
 	 */
-	bool Equals(const TVector4<T>& V, T Tolerance=KINDA_SMALL_NUMBER) const;
+	bool Equals(const TVector4<T>& V, T Tolerance=UE_KINDA_SMALL_NUMBER) const;
 
 	/**
 	 * Check if the vector is of unit length, with specified tolerance.
@@ -377,7 +377,7 @@ public:
 	 * @param LengthSquaredTolerance Tolerance against squared length.
 	 * @return true if the vector is a unit vector within the specified tolerance.
 	 */
-	bool IsUnit3(T LengthSquaredTolerance = KINDA_SMALL_NUMBER) const;
+	bool IsUnit3(T LengthSquaredTolerance = UE_KINDA_SMALL_NUMBER) const;
 
 	/**
 	 * Get a textual representation of the vector.
@@ -401,7 +401,7 @@ public:
 	 * @param Tolerance Minimum squared length of vector for normalization.
 	 * @return A normalized copy of the vector or a zero vector.
 	 */
-	FORCEINLINE TVector4 GetSafeNormal(T Tolerance = SMALL_NUMBER) const;
+	FORCEINLINE TVector4 GetSafeNormal(T Tolerance = UE_SMALL_NUMBER) const;
 
 	/**
 	 * Calculates normalized version of vector without checking if it is non-zero.
@@ -477,7 +477,7 @@ public:
 	bool ContainsNaN() const;
 
 	/** Utility to check if all of the components of this vector are nearly zero given the tolerance. */
-	bool IsNearlyZero3(T Tolerance = KINDA_SMALL_NUMBER) const;
+	bool IsNearlyZero3(T Tolerance = UE_KINDA_SMALL_NUMBER) const;
 
 	/** Reflect vector. */
 	TVector4<T> Reflect3(const TVector4<T>& Normal) const;
diff --git a/Engine/Source/Runtime/Core/Public/Math/VectorRegister.h b/Engine/Source/Runtime/Core/Public/Math/VectorRegister.h
--- a/Engine/Source/Runtime/Core/Public/Math/VectorRegister.h
+++ b/Engine/Source/Runtime/Core/Public/Math/VectorRegister.h
@@ -120,6 +120,6 @@ static constexpr VectorRegister VECTOR_INV_255 = MakeVectorRegisterDoubleConstan
 namespace GlobalVectorConstants
 {
 	static constexpr VectorRegister AnimWeightThreshold = MakeVectorRegisterConstant(ZERO_ANIMWEIGHT_THRESH_DOUBLE, ZERO_ANIMWEIGHT_THRESH_DOUBLE, ZERO_ANIMWEIGHT_THRESH_DOUBLE, ZERO_ANIMWEIGHT_THRESH_DOUBLE);
-	static constexpr VectorRegister RotationSignificantThreshold = MakeVectorRegisterConstant(1.0 - DELTA*DELTA, 1.0 - DELTA*DELTA, 1.0 - DELTA*DELTA, 1.0 - DELTA*DELTA);
+	static constexpr VectorRegister RotationSignificantThreshold = MakeVectorRegisterConstant(1.0 - UE_DELTA*UE_DELTA, 1.0 - UE_DELTA*UE_DELTA, 1.0 - UE_DELTA*UE_DELTA, 1.0 - UE_DELTA*UE_DELTA);
 }
 
diff --git a/Engine/Source/Runtime/Core/Public/Misc/AutomationTest.h b/Engine/Source/Runtime/Core/Public/Misc/AutomationTest.h
--- a/Engine/Source/Runtime/Core/Public/Misc/AutomationTest.h
+++ b/Engine/Source/Runtime/Core/Public/Misc/AutomationTest.h
@@ -1447,10 +1447,10 @@ public:
 #if PLATFORM_64BITS
 	bool TestEqual(const TCHAR* What, SIZE_T Actual, SIZE_T Expected);
 #endif
-	bool TestEqual(const TCHAR* What, float Actual, float Expected, float Tolerance = KINDA_SMALL_NUMBER);
-	bool TestEqual(const TCHAR* What, double Actual, double Expected, double Tolerance = KINDA_SMALL_NUMBER);
-	bool TestEqual(const TCHAR* What, FVector Actual, FVector Expected, float Tolerance = KINDA_SMALL_NUMBER);
-	bool TestEqual(const TCHAR* What, FRotator Actual, FRotator Expected, float Tolerance = KINDA_SMALL_NUMBER);
+	bool TestEqual(const TCHAR* What, float Actual, float Expected, float Tolerance = UE_KINDA_SMALL_NUMBER);
+	bool TestEqual(const TCHAR* What, double Actual, double Expected, double Tolerance = UE_KINDA_SMALL_NUMBER);
+	bool TestEqual(const TCHAR* What, FVector Actual, FVector Expected, float Tolerance = UE_KINDA_SMALL_NUMBER);
+	bool TestEqual(const TCHAR* What, FRotator Actual, FRotator Expected, float Tolerance = UE_KINDA_SMALL_NUMBER);
 	bool TestEqual(const TCHAR* What, FColor Actual, FColor Expected);
 	bool TestEqual(const TCHAR* What, FLinearColor Actual, FLinearColor Expected);
 	bool TestEqual(const TCHAR* What, const TCHAR* Actual, const TCHAR* Expected);
@@ -1461,22 +1461,22 @@ public:
 		return TestEqual(*What, Actual, Expected);
 	}
 
-	bool TestEqual(const FString& What, float Actual, float Expected, float Tolerance = KINDA_SMALL_NUMBER)
+	bool TestEqual(const FString& What, float Actual, float Expected, float Tolerance = UE_KINDA_SMALL_NUMBER)
 	{
 		return TestEqual(*What, Actual, Expected, Tolerance);
 	}
 
-	bool TestEqual(const FString& What, double Actual, double Expected, double Tolerance = KINDA_SMALL_NUMBER)
+	bool TestEqual(const FString& What, double Actual, double Expected, double Tolerance = UE_KINDA_SMALL_NUMBER)
 	{
 		return TestEqual(*What, Actual, Expected, Tolerance);
 	}
 
-	bool TestEqual(const FString& What, FVector Actual, FVector Expected, float Tolerance = KINDA_SMALL_NUMBER)
+	bool TestEqual(const FString& What, FVector Actual, FVector Expected, float Tolerance = UE_KINDA_SMALL_NUMBER)
 	{
 		return TestEqual(*What, Actual, Expected, Tolerance);
 	}
 
-	bool TestEqual(const FString& What, FRotator Actual, FRotator Expected, float Tolerance = KINDA_SMALL_NUMBER)
+	bool TestEqual(const FString& What, FRotator Actual, FRotator Expected, float Tolerance = UE_KINDA_SMALL_NUMBER)
 	{
 		return TestEqual(*What, Actual, Expected, Tolerance);
 	}
