Float16 Max constant and comments about clamp behavior change in Set
no change in existing behavior

diff --git a/Engine/Source/Runtime/Core/Public/Math/Float16.h b/Engine/Source/Runtime/Core/Public/Math/Float16.h
--- a/Engine/Source/Runtime/Core/Public/Math/Float16.h
+++ b/Engine/Source/Runtime/Core/Public/Math/Float16.h
@@ -32,6 +32,9 @@ class FFloat16
 {
 public:
 
+	/* Float16 can store values in [-MaxF16Float,MaxF16Float] */
+	constexpr static float MaxF16Float = 65504.f;
+
 	uint16 Encoded;
 
 	/** Default constructor */
@@ -52,10 +55,20 @@ public:
 	/** Convert from Fp16 to Fp32. */
 	operator float() const;
 
-	/** Convert from Fp32 to Fp16. (RTNE) */
+	/** Convert from Fp32 to Fp16, round-to-nearest-even. (RTNE)
+	Stores values out of range as +-Inf */
 	void Set(float FP32Value);
 	
-	/** Convert from Fp32 to Fp16. (backward-compatible truncate conversion) */
+	/*Convert from Fp32 to Fp16, round-to-nearest-even. (RTNE)
+	Clamps values out of range as +-MaxF16Float */
+	void SetClamped(float FP32Value)
+	{
+		Set( FMath::Clamp(FP32Value,-MaxF16Float,MaxF16Float) );
+	}
+
+	/** Convert from Fp32 to Fp16, truncating low bits. 
+	(backward-compatible conversion; was used by Set() previously)
+	Clamps values out of range to [-MaxFloat,MaxFloat] */
 	void SetTruncate(float FP32Value);
 
 	/** Convert from Fp16 to Fp32. */
@@ -123,6 +136,7 @@ FORCEINLINE FFloat16::operator float() const
 }
 
 
+// NOTE: Set() on values out of F16 max range store them as +-Inf
 FORCEINLINE void FFloat16::Set(float FP32Value)
 {
 	// FPlatformMath::StoreHalf follows RTNE (round-to-nearest-even) rounding default convention
@@ -137,6 +151,7 @@ FORCEINLINE float FFloat16::GetFloat() const
 }
 
 
+// NOTE: SetTruncate() on values out of F16 max range store them as +-Inf
 FORCEINLINE void FFloat16::SetTruncate(float FP32Value)
 {
 
