Temporary fix for IncludeTool warnings about inconsistent preprocessor state.

diff --git a/Engine/Source/Runtime/Core/Core.Build.cs b/Engine/Source/Runtime/Core/Core.Build.cs
--- a/Engine/Source/Runtime/Core/Core.Build.cs
+++ b/Engine/Source/Runtime/Core/Core.Build.cs
@@ -313,8 +313,6 @@ public class Core : ModuleRules
 		PrivateDefinitions.Add("PLATFORM_COMPILER_OPTIMIZATION_PG=" + (Target.bPGOOptimize ? "1" : "0"));
 		PrivateDefinitions.Add("PLATFORM_COMPILER_OPTIMIZATION_PG_PROFILING=" + (Target.bPGOProfile ? "1" : "0"));
 
-		PrivateDefinitions.Add("UE_DEFINE_LEGACY_MATH_CONSTANT_MACRO_NAMES=0");
-
 		UnsafeTypeCastWarningLevel = WarningLevel.Warning;
 	}
 
