Tests converted from 18848115 made to work multi-platform.

diff --git a/Engine/Source/Runtime/Core/Tests/Containers/ContainersTest.cpp b/Engine/Source/Runtime/Core/Tests/Containers/ContainersTest.cpp
new file mode 100644
--- /dev/null
+++ b/Engine/Source/Runtime/Core/Tests/Containers/ContainersTest.cpp
@@ -0,0 +1,844 @@
+// Copyright Epic Games, Inc. All Rights Reserved.
+
+#include "CoreTypes.h"
+#include "Misc/AssertionMacros.h"
+#include "Containers/Array.h"
+#include "Containers/Map.h"
+#include "Containers/Set.h"
+#include "Containers/SortedMap.h"
+#include "Containers/ArrayView.h"
+#include "Stats/StatsMisc.h"
+#include "Math/RandomStream.h"
+#include "TestHarness.h"
+
+#define MAX_TEST_OBJECTS      65
+#define MAX_TEST_OBJECTS_STEP 1
+#define RANDOM_SEED 12345
+
+namespace
+{
+	namespace EContainerTestType
+	{
+		enum Type
+		{
+			MovedFrom,
+			Key,
+			Value
+		};
+	}
+
+	struct FContainerTestStats
+	{
+		int32  NextId;
+		int32  ConstructedIDs[MAX_TEST_OBJECTS];
+		int32* End;
+
+		FContainerTestStats()
+		{
+			Reset();
+		}
+
+		void Reset()
+		{
+			NextId = 1;
+			End    = ConstructedIDs;
+		}
+
+		int32 Num() const
+		{
+			return UE_PTRDIFF_TO_INT32(End - ConstructedIDs);
+		}
+
+		int32 Add()
+		{
+			// Ensure we're not constructing too many objects
+			CHECK(Num() < MAX_TEST_OBJECTS);
+
+			// Store ID in array
+			return *End++ = NextId++;
+		}
+
+		void Remove(int32 ObjId)
+		{
+			for (int32* It = ConstructedIDs; It != End; ++It)
+			{
+				if (*It != ObjId)
+					continue;
+
+				// Remove this from the list
+				--End;
+				for (; It != End; ++It)
+				{
+					*It = *(It + 1);
+				}
+				return;
+			}
+
+			// We didn't find an entry for this - an invalid destructor call?
+			CHECK(false);
+		}
+	} ContainerTestStats;
+
+	struct FContainerTestType
+	{
+		FContainerTestType(const TCHAR* InStr, EContainerTestType::Type InType)
+			: Str (InStr)
+			, Type(InType)
+			, Id  (ContainerTestStats.Add())
+		{
+		}
+
+		FContainerTestType(const FContainerTestType& Other)
+			: Str (Other.Str)
+			, Type(Other.Type)
+			, Id  (ContainerTestStats.Add())
+		{
+		}
+
+		FContainerTestType(FContainerTestType&& Other)
+			: Str (Other.Str)
+			, Type(Other.Type)
+			, Id  (ContainerTestStats.Add())
+		{
+			Other.Str  = NULL;
+			Other.Type = EContainerTestType::MovedFrom;
+		}
+
+		FContainerTestType& operator=(FContainerTestType&& Other)
+		{
+			Str  = Other.Str;
+			Type = Other.Type;
+
+			Other.Str  = NULL;
+			Other.Type = EContainerTestType::MovedFrom;
+			
+			return *( static_cast<FContainerTestType*>( this ) );
+		}
+
+		FContainerTestType& operator=(const FContainerTestType& Other)
+		{
+			Str  = Other.Str;
+			Type = Other.Type;
+			
+			return *( static_cast<FContainerTestType*>( this ) );
+		}
+
+		~FContainerTestType()
+		{
+			ContainerTestStats.Remove(Id);
+		}
+
+		friend bool operator==(const FContainerTestType& Lhs, const FContainerTestType& Rhs)
+		{
+			return Lhs.Type == Rhs.Type &&
+			       Lhs.Id   == Rhs.Id   &&
+			       !FCString::Strcmp(Lhs.Str, Rhs.Str);
+		}
+
+		friend bool operator!=(const FContainerTestType& Lhs, const FContainerTestType& Rhs)
+		{
+			return !(Lhs == Rhs);
+		}
+
+		const TCHAR*             Str;
+		EContainerTestType::Type Type;
+		int32                    Id;
+	};
+
+	struct FContainerTestKeyType : FContainerTestType
+	{
+		FContainerTestKeyType()
+			: FContainerTestType(TEXT("<default key>"), EContainerTestType::Key)
+		{
+		}
+
+		explicit FContainerTestKeyType(const TCHAR* InStr)
+			: FContainerTestType(InStr, EContainerTestType::Key)
+		{
+		}
+
+		FContainerTestKeyType(FContainerTestKeyType&& Other)
+			: FContainerTestType(MoveTemp(Other))
+		{
+		}
+
+		FContainerTestKeyType(const FContainerTestKeyType& Other)
+			: FContainerTestType(Other)
+		{
+		}
+
+		FContainerTestKeyType& operator=(FContainerTestKeyType&& Other)
+		{
+			(FContainerTestType&)*this = MoveTemp(Other);
+			
+			return *( static_cast<FContainerTestKeyType*>( this ) );
+		}
+
+		FContainerTestKeyType& operator=(const FContainerTestKeyType& Other)
+		{
+			(FContainerTestType&)*this = Other;
+			
+			return *( static_cast<FContainerTestKeyType*>( this ) );
+		}
+	};
+
+	struct FContainerTestValueType : FContainerTestType
+	{
+		FContainerTestValueType()
+			: FContainerTestType(TEXT("<default value>"), EContainerTestType::Value)
+		{
+		}
+
+		explicit FContainerTestValueType(const TCHAR* InStr)
+			: FContainerTestType(InStr, EContainerTestType::Value)
+		{
+		}
+
+		FContainerTestValueType(FContainerTestValueType&& Other)
+			: FContainerTestType(MoveTemp(Other))
+		{
+		}
+
+		FContainerTestValueType(const FContainerTestValueType& Other)
+			: FContainerTestType(Other)
+		{
+		}
+
+		FContainerTestValueType& operator=(FContainerTestValueType&& Other)
+		{
+			(FContainerTestType&)*this = MoveTemp(Other);
+			
+			return *( static_cast<FContainerTestValueType*>( this ) );
+		}
+
+		FContainerTestValueType& operator=(const FContainerTestValueType& Other)
+		{
+			(FContainerTestType&)*this = Other;
+			
+			return *( static_cast<FContainerTestValueType*>( this ) );
+		}
+	};
+
+	template <typename Container>
+	void CheckContainerElements(Container& Cont)
+	{
+		auto It  = Cont.CreateIterator();
+		auto CIt = Cont.CreateConstIterator();
+		for (auto& E : Cont)
+		{
+			CHECK(*It  == E);
+			CHECK(*CIt == E);
+
+			++It;
+			++CIt;
+		}
+	}
+
+	template <typename Container>
+	void CheckContainerNum(Container& Cont)
+	{
+		int32 Count = 0;
+		for (auto It = Cont.CreateIterator(); It; ++It)
+		{
+			++Count;
+		}
+
+		int32 CCount = 0;
+		for (auto It = Cont.CreateConstIterator(); It; ++It)
+		{
+			++CCount;
+		}
+
+		int32 RCount = 0;
+		for (auto& It : Cont)
+		{
+			++RCount;
+		}
+
+		CHECK(Count  == Cont.Num());
+		CHECK(CCount == Cont.Num());
+		CHECK(RCount == Cont.Num());
+	}
+
+	template <typename Container>
+	void CheckContainerEnds(Container& Cont)
+	{
+		auto Iter  = Cont.CreateIterator();
+		auto CIter = Cont.CreateConstIterator();
+
+		for (int32 Num = Cont.Num(); Num; --Num)
+		{
+			++Iter;
+			++CIter;
+		}
+
+		CHECK(!Iter);
+		CHECK(!CIter);
+	}
+
+	template <typename KeyType>
+	KeyType GenerateTestKey(int32 Input)
+	{
+		return (KeyType)Input;
+	}
+
+	template <>
+	FName GenerateTestKey<FName>(int32 Input)
+	{
+		// Don't use _foo as we want to test the slower compare path
+		return FName(*FString::Printf(TEXT("TestName%d"), Input));
+	}
+
+	template <>
+	FString GenerateTestKey<FString>(int32 Input)
+	{
+		return FString::Printf(TEXT("TestString%d"), Input);
+	}
+
+	template <typename ContainerType, typename KeyType>
+	void RunContainerTests()
+	{
+		ContainerType Cont;
+
+		ContainerTestStats.Reset();
+		// Subtract one to account for temporaries that will be created during an Add
+		for (int32 Count = 0; Count < MAX_TEST_OBJECTS - 1; Count += MAX_TEST_OBJECTS_STEP)
+		{
+			for (int32 N = 0; N != Count; ++N)
+			{
+				Cont.Add(GenerateTestKey<KeyType>(N), FContainerTestValueType(TEXT("New Value")));
+				CheckContainerNum(Cont);
+				CheckContainerEnds(Cont);
+				CheckContainerElements(Cont);
+			}
+
+			for (int32 N = 0; N != Count; ++N)
+			{
+				Cont.Remove(GenerateTestKey<KeyType>(N));
+				CheckContainerNum(Cont);
+				CheckContainerEnds(Cont);
+				CheckContainerElements(Cont);
+			}
+
+			for (int32 N = 0; N != Count; ++N)
+			{
+				Cont.Add(GenerateTestKey<KeyType>((Count - 1) - N), FContainerTestValueType(TEXT("New Value")));
+				CheckContainerNum(Cont);
+				CheckContainerEnds(Cont);
+				CheckContainerElements(Cont);
+			}
+
+			for (int32 N = 0; N != Count; ++N)
+			{
+				Cont.Remove(GenerateTestKey<KeyType>(N));
+				CheckContainerNum(Cont);
+				CheckContainerEnds(Cont);
+				CheckContainerElements(Cont);
+			}
+		}
+	}
+
+	template <typename ContainerType, typename KeyType>
+	void RunPerformanceTest(const FString& Description, int32 NumObjects, int32 NumOperations)
+	{
+		ContainerTestStats.Reset();
+		
+		ContainerType Cont;
+		FRandomStream RandomStream(RANDOM_SEED);
+
+		// Prep keys, not part of performance test
+		TArray<KeyType> KeyArray;
+		KeyArray.Reserve(NumObjects);
+
+		for (int32 i = 0; i < NumObjects; i++)
+		{
+			KeyArray.Add(GenerateTestKey<KeyType>(i));
+		}
+
+		for (int32 i = 0; i < NumObjects; i++)
+		{
+			int32 SwapIndex = RandomStream.RandRange(0, NumObjects - 1);
+			if (i != SwapIndex)
+			{
+				KeyArray.Swap(i, SwapIndex);
+			}
+		}
+
+		FScopeLogTime LogTimePtr(*FString::Printf(TEXT("%s objects=%d count=%d"), *Description, NumObjects, NumOperations), nullptr, FScopeLogTime::ScopeLog_Milliseconds);
+
+		// Add elements in stably randomized order
+		for (int32 i = 0; i < NumObjects; i++)
+		{
+			Cont.Add(KeyArray[i], FString(TEXT("New Value")));
+		}
+		
+		// Now do searches
+		for (int32 i = 0; i < NumOperations; i++)
+		{
+			KeyType& Key = KeyArray[RandomStream.RandRange(0, NumObjects - 1)];
+
+			FString* FoundValue = Cont.Find(Key);
+			CHECK(FoundValue);
+		}
+	}
+
+	template <typename ContainerType, typename KeyType>
+	void RunSetPerformanceTest(const FString& Description, int32 NumObjects, int32 NumOperations)
+	{
+		ContainerTestStats.Reset();
+
+		ContainerType Cont;
+		FRandomStream RandomStream(RANDOM_SEED);
+
+		// Prep keys, not part of performance test
+		TArray<KeyType> KeyArray;
+		KeyArray.Reserve(NumObjects);
+
+		for (int32 i = 0; i < NumObjects; i++)
+		{
+			KeyArray.Add(GenerateTestKey<KeyType>(i));
+		}
+
+		for (int32 i = 0; i < NumObjects; i++)
+		{
+			int32 SwapIndex = RandomStream.RandRange(0, NumObjects - 1);
+			if (i != SwapIndex)
+			{
+				KeyArray.Swap(i, SwapIndex);
+			}
+		}
+
+		FScopeLogTime LogTimePtr(*FString::Printf(TEXT("%s objects=%d count=%d"), *Description, NumObjects, NumOperations), nullptr, FScopeLogTime::ScopeLog_Milliseconds);
+
+		// Add elements in stably randomized order
+		for (int32 i = 0; i < NumObjects; i++)
+		{
+			Cont.Add(KeyArray[i]);
+		}
+
+		// Now do searches
+		for (int32 i = 0; i < NumOperations; i++)
+		{
+			KeyType& Key = KeyArray[RandomStream.RandRange(0, NumObjects - 1)];
+
+			bool FoundValue = Cont.Contains(Key);
+			CHECK(FoundValue);
+		}
+	}
+
+	template <typename ContainerType, typename ElemType, typename GetRefFuncType>
+	void TestGetRef(ContainerType& Container, ElemType&& Elem, GetRefFuncType&& GetRefFunc)
+	{
+		auto& Ref = GetRefFunc(Container, Forward<ElemType>(Elem));
+	}
+
+	template<typename TValueType>
+	struct FCaseSensitiveLookupKeyFuncs : BaseKeyFuncs<TValueType, FString>
+	{
+		static FORCEINLINE const FString& GetSetKey(const TPair<FString, TValueType>& Element)
+		{
+			return Element.Key;
+		}
+		static FORCEINLINE bool Matches(const FString& A, const FString& B)
+		{
+			return A.Equals(B, ESearchCase::CaseSensitive);
+		}
+		static FORCEINLINE uint32 GetKeyHash(const FString& Key)
+		{
+			return FCrc::StrCrc32<TCHAR>(*Key);
+		}
+	};
+
+	void RunGetRefTests()
+	{
+		{
+			TArray<FString> Arr;
+			FString* Str1 = &Arr.AddDefaulted_GetRef();
+			CHECK(Str1->IsEmpty());
+			CHECK(Str1 == &Arr.Last());
+
+			FString* Str2 = &Arr.AddZeroed_GetRef();
+			CHECK(Str2->IsEmpty());
+			CHECK(Str2 == &Arr.Last());
+
+			FString* Str3 = &Arr.Add_GetRef(TEXT("Abc"));
+			CHECK(*Str3 == TEXT("Abc"));
+			CHECK(Str3 == &Arr.Last());
+
+			FString* Str4 = &Arr.Emplace_GetRef(TEXT("Def"));
+			CHECK(*Str4 == TEXT("Def"));
+			CHECK(Str4 == &Arr.Last());
+
+			FString* Str5 = &Arr.EmplaceAt_GetRef(3, TEXT("Ghi"));
+			CHECK(*Str5 == TEXT("Ghi"));
+			CHECK(Str5 == &Arr[3]);
+
+			FString* Str6 = &Arr.InsertDefaulted_GetRef(2);
+			CHECK(Str6->IsEmpty());
+			CHECK(Str6 == &Arr[2]);
+
+			FString* Str7 = &Arr.InsertZeroed_GetRef(4);
+			CHECK(Str7->IsEmpty());
+			CHECK(Str7 == &Arr[4]);
+		}
+	}
+}
+
+TEST_CASE("Core::Containers::TMap::Smoke", "[Core][Containers][Smoke]")
+{
+	RunContainerTests<TMap<int32, FContainerTestValueType>, int32>();
+}
+
+TEST_CASE("Core::Containers::TMap::Full", "[Core][Containers][Full]")
+{
+	RunContainerTests<TMap<int32, FContainerTestValueType>, int32>();
+	RunContainerTests<TMap<FName, FContainerTestValueType>, FName>();
+	RunContainerTests<TMap<FString, FContainerTestValueType>, FString>();
+	RunContainerTests<TMap<int32, FContainerTestValueType, TInlineSetAllocator<64>>, int32>();
+	RunContainerTests<TMap<FString, FContainerTestValueType, FDefaultSetAllocator, FCaseSensitiveLookupKeyFuncs<FContainerTestValueType>>, FString>();
+}
+
+TEST_CASE("Core::Containers::TSortedMap::Full", "[Core][Containers][Full]")
+{
+	RunContainerTests<TSortedMap<int32, FContainerTestValueType>, int32>();
+	RunContainerTests<TSortedMap<FName, FContainerTestValueType, FDefaultAllocator, FNameLexicalLess>, FName>();
+	RunContainerTests<TSortedMap<FString, FContainerTestValueType>, FString>();
+	RunContainerTests<TSortedMap<FString, FContainerTestValueType, TInlineAllocator<64>>, FString>();
+
+	// Verify use of FName index sorter with SortedMap
+
+	TSortedMap<FName, int32, FDefaultAllocator, FNameFastLess> NameMap;
+	NameMap.Add(NAME_NameProperty);
+	NameMap.Add(NAME_FloatProperty);
+	NameMap.Add(NAME_None);
+	NameMap.Add(NAME_IntProperty);
+
+	auto It = NameMap.CreateConstIterator();
+
+	CHECK(It->Key == NAME_None); ++It;
+	CHECK(It->Key == NAME_IntProperty); ++It;
+	CHECK(It->Key == NAME_FloatProperty); ++It;
+	CHECK(It->Key == NAME_NameProperty); ++It;
+	CHECK(!It);
+
+	RunGetRefTests();
+}
+
+TEST_CASE("Core::Containers::TMap::Performance", "[Core][Containers][Perf]")
+{
+	RunPerformanceTest<TMap<int32, FString>, int32>(TEXT("TMap int32"), 1, 1000000);
+	RunPerformanceTest<TMap<int32, FString>, int32>(TEXT("TMap int32"), 10, 1000000);
+	RunPerformanceTest<TMap<int32, FString>, int32>(TEXT("TMap int32"), 100, 1000000);
+	RunPerformanceTest<TMap<int32, FString>, int32>(TEXT("TMap int32"), 1000, 1000000);
+	RunPerformanceTest<TMap<int32, FString>, int32>(TEXT("TMap int32"), 10000, 1000000);
+
+	RunPerformanceTest<TMap<FName, FString>, FName>(TEXT("TMap FName"), 1, 1000000);
+	RunPerformanceTest<TMap<FName, FString>, FName>(TEXT("TMap FName"), 10, 1000000);
+	RunPerformanceTest<TMap<FName, FString>, FName>(TEXT("TMap FName"), 100, 1000000);
+	RunPerformanceTest<TMap<FName, FString>, FName>(TEXT("TMap FName"), 1000, 1000000);
+	RunPerformanceTest<TMap<FName, FString>, FName>(TEXT("TMap FName"), 10000, 1000000);
+
+	RunPerformanceTest<TMap<FString, FString>, FString>(TEXT("TMap FString"), 1, 1000000);
+	RunPerformanceTest<TMap<FString, FString>, FString>(TEXT("TMap FString"), 10, 1000000);
+	RunPerformanceTest<TMap<FString, FString>, FString>(TEXT("TMap FString"), 100, 1000000);
+	RunPerformanceTest<TMap<FString, FString>, FString>(TEXT("TMap FString"), 1000, 1000000);
+	RunPerformanceTest<TMap<FString, FString>, FString>(TEXT("TMap FString"), 10000, 1000000);
+}
+
+TEST_CASE("Core::Containers::TSortedMap::Performance", "[Core][Containers][Perf]")
+{
+	RunPerformanceTest<TSortedMap<int32, FString>, int32>(TEXT("TSortedMap int32"), 1, 1000000);
+	RunPerformanceTest<TSortedMap<int32, FString>, int32>(TEXT("TSortedMap int32"), 10, 1000000);
+	RunPerformanceTest<TSortedMap<int32, FString>, int32>(TEXT("TSortedMap int32"), 100, 1000000);
+	RunPerformanceTest<TSortedMap<int32, FString>, int32>(TEXT("TSortedMap int32"), 1000, 1000000);
+	RunPerformanceTest<TSortedMap<int32, FString>, int32>(TEXT("TSortedMap int32"), 10000, 1000000);
+
+	RunPerformanceTest<TSortedMap<FName, FString, FDefaultAllocator, FNameLexicalLess>, FName>(TEXT("TSortedMap FName"), 1, 1000000);
+	RunPerformanceTest<TSortedMap<FName, FString, FDefaultAllocator, FNameLexicalLess>, FName>(TEXT("TSortedMap FName"), 10, 1000000);
+	RunPerformanceTest<TSortedMap<FName, FString, FDefaultAllocator, FNameLexicalLess>, FName>(TEXT("TSortedMap FName"), 100, 1000000);
+	RunPerformanceTest<TSortedMap<FName, FString, FDefaultAllocator, FNameLexicalLess>, FName>(TEXT("TSortedMap FName"), 1000, 1000000);
+	RunPerformanceTest<TSortedMap<FName, FString, FDefaultAllocator, FNameLexicalLess>, FName>(TEXT("TSortedMap FName"), 10000, 1000000);
+
+	RunPerformanceTest<TSortedMap<FString, FString>, FString>(TEXT("TSortedMap FString"), 1, 1000000);
+	RunPerformanceTest<TSortedMap<FString, FString>, FString>(TEXT("TSortedMap FString"), 10, 1000000);
+	RunPerformanceTest<TSortedMap<FString, FString>, FString>(TEXT("TSortedMap FString"), 100, 1000000);
+	RunPerformanceTest<TSortedMap<FString, FString>, FString>(TEXT("TSortedMap FString"), 1000, 1000000);
+	RunPerformanceTest<TSortedMap<FString, FString>, FString>(TEXT("TSortedMap FString"), 10000, 1000000);
+
+	RunSetPerformanceTest<TSet<FName>, FName>(TEXT("TSet FName"), 1, 1000000);
+	RunSetPerformanceTest<TSet<FName>, FName>(TEXT("TSet FName"), 10, 1000000);
+	RunSetPerformanceTest<TSet<FName>, FName>(TEXT("TSet FName"), 100, 1000000);
+
+	RunSetPerformanceTest<TArray<FName>, FName>(TEXT("TArray FName"), 1, 1000000);
+	RunSetPerformanceTest<TArray<FName>, FName>(TEXT("TArray FName"), 10, 1000000);
+	RunSetPerformanceTest<TArray<FName>, FName>(TEXT("TArray FName"), 100, 1000000);
+}
+
+namespace
+{
+	struct FRecorder
+	{
+		FRecorder(uint32 InKey = 0, uint32 InPayload=0)
+			: Id(NextId++)
+			, Key(InKey)
+			, Payload(InPayload)
+		{
+		}
+		FRecorder(const FRecorder& Other)
+			: Id(NextId++)
+			, Key(Other.Key)
+			, Payload(Other.Payload)
+			, NumCopies(Other.NumCopies+1)
+			, NumMoves(Other.NumMoves)
+		{
+		}
+		FRecorder(FRecorder&& Other)
+			: Id(NextId++)
+			, Key(Other.Key)
+			, Payload(Other.Payload)
+			, NumCopies(Other.NumCopies)
+			, NumMoves(Other.NumMoves+1)
+		{
+		}
+
+		bool operator==(const FRecorder& Other) const
+		{
+			return Key == Other.Key;
+		}
+
+		uint32 Id;
+		uint32 Key;
+		uint32 Payload;
+		uint32 NumCopies = 0;
+		uint32 NumMoves = 0;
+
+		static uint32 NextId;
+	};
+	uint32 FRecorder::NextId = 0;
+
+	uint32 GetTypeHash(const FRecorder& Recorder)
+	{
+		return Recorder.Payload;
+	}
+
+}
+
+TEST_CASE("Core::Containers::TSet::Smoke", "[Core][Containers][Smoke]")
+{
+	enum class EArgType
+	{
+		Copy,
+		Move
+	};
+	enum class EHashType
+	{
+		Internal,
+		PassedIn
+	};
+
+	for (EHashType HashType : { EHashType::Internal, EHashType::PassedIn})
+	{
+		for (EArgType ArgType : {EArgType::Copy, EArgType::Move})
+		{
+			FString FuncName;
+			auto FullText = [HashType,ArgType,&FuncName](const TCHAR* Message)
+			{
+				const TCHAR* HashText = HashType == EHashType::Internal ? TEXT("") : TEXT("ByHash");
+				const TCHAR* ArgText = ArgType == EArgType::Copy ? TEXT("const&") : TEXT("&&");
+				return FString::Printf(TEXT("TSet::%s%s(%s) %s"), *FuncName, HashText, ArgText, Message);
+			};
+
+			// Test TSet::Add(const&), Add(&&), AddByHash(const&), AddByHash(&&)
+			FuncName = TEXT("Add");
+			{
+				TSet<FRecorder> Set;
+				FRecorder First(37, 43);
+				bool bAlreadyInSet = true;
+
+				if (HashType == EHashType::Internal)
+					if (ArgType == EArgType::Copy)
+						Set.Add(First, &bAlreadyInSet);
+					else
+						Set.Add(MoveTemp(First), &bAlreadyInSet);
+				else
+					if (ArgType == EArgType::Copy)
+						Set.AddByHash(GetTypeHash(First), First, &bAlreadyInSet);
+					else
+						Set.AddByHash(GetTypeHash(First), MoveTemp(First), &bAlreadyInSet);
+				TestFalse(FullText(TEXT("returns bAlreadyInSet==false for first add")), bAlreadyInSet);
+
+				FRecorder* Found = Set.Find(First);
+				if (ArgType == EArgType::Copy)
+				{
+					TestTrue(FullText(TEXT("constructs a copy")), Found && Found->Id > First.Id && Found->NumCopies > 0 && Found->Payload == First.Payload);
+				}
+				else
+				{
+					TestTrue(FullText(TEXT("constructs a move")), Found && Found->Id > First.Id && Found->NumCopies == 0 && Found->NumMoves >= 1 && Found->Payload == First.Payload);
+				}
+
+				uint32 FoundId = Found ? Found->Id : 0;
+				Found = Set.Find(First);
+				TestTrue(TEXT("Finding an element returns a reference, no copies"), Found && Found->Id == FoundId);
+
+				FRecorder Second(37, 56);
+				if (HashType == EHashType::Internal)
+					if (ArgType == EArgType::Copy)
+						Set.Add(Second, &bAlreadyInSet);
+					else
+						Set.Add(MoveTemp(Second), &bAlreadyInSet);
+				else
+					if (ArgType == EArgType::Copy)
+						Set.AddByHash(GetTypeHash(Second), Second, &bAlreadyInSet);
+					else
+						Set.AddByHash(GetTypeHash(Second), MoveTemp(Second), &bAlreadyInSet);
+				TestTrue(FullText(TEXT("returns bAlreadyInSet==true for second add")), bAlreadyInSet);
+				Found = Set.Find(First);
+				TestTrue(FullText(TEXT("with a duplicate key constructs a copy of the new key")), Found && Found->Id > Second.Id && Found->Payload == Second.Payload);
+			}
+
+			// Test TSet::FindOrAdd(const&), FindOrAdd(&&), FindOrAddByHash(const&), FindOrAddByHash(&&)
+			FuncName = TEXT("FindOrAdd");
+			{
+				TSet<FRecorder> Set;
+				FRecorder First(37, 43);
+				bool bAlreadyInSet = true;
+
+				FRecorder* FindOrAddResult;
+				if (HashType == EHashType::Internal)
+					if (ArgType == EArgType::Copy)
+						FindOrAddResult = &Set.FindOrAdd(First, &bAlreadyInSet);
+					else
+						FindOrAddResult = &Set.FindOrAdd(MoveTemp(First), &bAlreadyInSet);
+				else
+					if (ArgType == EArgType::Copy)
+						FindOrAddResult = &Set.FindOrAddByHash(GetTypeHash(First), First, &bAlreadyInSet);
+					else
+						FindOrAddResult = &Set.FindOrAddByHash(GetTypeHash(First), MoveTemp(First), &bAlreadyInSet);
+				TestFalse(FullText(TEXT("returns bAlreadyInSet==false for first add")), bAlreadyInSet);
+				if (ArgType == EArgType::Copy)
+				{
+					TestTrue(FullText(TEXT("on the first constructs a copy")), FindOrAddResult->Id > First.Id && FindOrAddResult->NumCopies > 0 && FindOrAddResult->Payload == First.Payload);
+				}
+				else
+				{
+					TestTrue(FullText(TEXT("on the first constructs a move")), FindOrAddResult->Id > First.Id && FindOrAddResult->NumCopies == 0 && FindOrAddResult->NumMoves >= 1 && FindOrAddResult->Payload == First.Payload);
+				}
+				uint32 FoundId = FindOrAddResult->Id;
+
+				FRecorder* Found = Set.Find(First);
+				TestTrue(FullText(TEXT("returns same value as future find")), Found&& Found->Id == FindOrAddResult->Id);
+				Found = Set.Find(First);
+				TestTrue(TEXT("Finding an element returns a reference, no copies"), Found && Found->Id == FoundId);
+
+				FRecorder Second(37, 56);
+				if (HashType == EHashType::Internal)
+					if (ArgType == EArgType::Copy)
+						FindOrAddResult = &Set.FindOrAdd(Second, &bAlreadyInSet);
+					else
+						FindOrAddResult = &Set.FindOrAdd(MoveTemp(Second), &bAlreadyInSet);
+				else
+					if (ArgType == EArgType::Copy)
+						FindOrAddResult = &Set.FindOrAddByHash(GetTypeHash(Second), Second, &bAlreadyInSet);
+					else
+						FindOrAddResult = &Set.FindOrAddByHash(GetTypeHash(Second), MoveTemp(Second), &bAlreadyInSet);
+				TestTrue(FullText(TEXT("returns bAlreadyInSet==true for second add")), bAlreadyInSet);
+				TestTrue(FullText(TEXT("with a duplicate key keeps the original key, returned from FindOrAdd")), FindOrAddResult->Id == FoundId && FindOrAddResult->Payload == First.Payload);
+				Found = Set.Find(First);
+				TestTrue(FullText(TEXT("with a duplicate key keeps the original key, returned from future Finds")), Found->Id == FoundId && Found->Payload == First.Payload);
+			}
+		}
+	}
+}
+
+
+namespace ArrayViewTests
+{
+	// commented out lines shouldn't compile
+
+	struct Base
+	{
+		int32 b;
+	};
+
+	struct Derived : public Base
+	{
+		int32 d;
+	};
+
+	template<typename T>
+	void TestFunction(TArrayView<T>)
+	{
+	}
+
+	template<typename T>
+	void TestFunction64(TArrayView64<T>)
+	{
+	}
+
+	TEST_CASE("Core::Containers::TArrayView::Smoke", "[Core][Containers][Smoke]")
+	{
+		// C array + derived-to-base conversions
+		Derived test1[13];
+		TestFunction<Derived>(test1);
+		//TestFunction<Base>(test1);
+		//TestFunction<const Base>(test1);
+		TestFunction64<Derived>(test1);
+
+		// C array of pointers + derived-to-base conversions
+		Derived* test2[13];
+		TestFunction<const Derived* const>(test2);
+		//TestFunction<const Derived*>(test2);
+		TestFunction<Derived* const>(test2);
+		//TestFunction<const Base* const>(test2);
+		TestFunction64<const Derived* const>(test2);
+		TestFunction64<Derived* const>(test2);
+
+		// TArray + derived-to-base conversions
+		TArray<Derived> test3;
+		TestFunction<Derived>(test3);
+		//TestFunction<Base>(test3);
+		//TestFunction<const Base>(test3);
+		TestFunction64<Derived>(test3);
+
+		// const TArray
+		const TArray<Base> test4;
+		TestFunction<const Base>(test4);
+		//TestFunction<Base>(test4);
+		TestFunction64<const Base>(test4);
+
+		// TArray of const
+		TArray<const Base> test5;
+		TestFunction<const Base>(test5);
+		//TestFunction<Base>(test5);
+		TestFunction64<const Base>(test5);
+
+		// temporary C array
+		struct Test6
+		{
+			Base test[13];
+		};
+		TestFunction<const Base>(Test6().test);
+		//TestFunction<Base>(Test6().test); // shouldn't compile but VS allows it :(
+		TestFunction64<const Base>(Test6().test);
+
+		// TArrayView64 from TArrayView
+		TArrayView<Derived> test7 = test1;
+		TestFunction64<Derived>(test7);
+		//TArrayView64<Derived> test7_64 = test1;
+		//TestFunction<Derived>(test7_64);
+
+		// TArray[64] from TArrayView[64]
+		TArrayView<Derived> test8;
+		TArrayView<Derived> test8_64;
+		TArray<Derived> test8_32from32(test8);
+		TArray<Derived> test8_32from64(test8_64);
+		TArray64<Derived> test8_64from32(test8);
+		TArray64<Derived> test8_64from64(test8_64);
+		test8_32from32 = test8;
+		test8_32from64 = test8_64;
+		test8_64from32 = test8;
+		test8_64from64 = test8_64;
+	}
+};
+
diff --git a/Engine/Source/Runtime/Core/Tests/Containers/RingBufferTest.cpp b/Engine/Source/Runtime/Core/Tests/Containers/RingBufferTest.cpp
new file mode 100644
--- /dev/null
+++ b/Engine/Source/Runtime/Core/Tests/Containers/RingBufferTest.cpp
@@ -0,0 +1,1115 @@
+// Copyright Epic Games, Inc. All Rights Reserved.
+
+#include "Containers/RingBuffer.h"
+#include "TestHarness.h"
+
+class FRingBufferTest
+{
+public:
+	static bool IsIntegerRange(const TRingBuffer<uint32>& Queue, uint32 Start, uint32 End, bool bForward = true)
+	{
+		if (Queue.Num() != End - Start)
+		{
+			return false;
+		}
+
+		// Peek elements in queue at given offset, peek from back to front
+		for (int32 It = 0; It < Queue.Num(); ++It)
+		{
+			uint32 QueueValue = bForward ? Queue[It] : Queue[Queue.Num() - 1 - It];
+			if (QueueValue != Start + It)
+			{
+				return false;
+			}
+		}
+
+		return true;
+	}
+
+	struct Counter
+	{
+		Counter(uint32 InValue = 0x12345)
+			:Value(InValue)
+		{
+			++NumVoid;
+		}
+		Counter(const Counter& Other)
+			:Value(Other.Value)
+		{
+			++NumCopy;
+		}
+		Counter(Counter&& Other)
+			:Value(Other.Value)
+		{
+			++NumMove;
+		}
+		~Counter()
+		{
+			++NumDestruct;
+		}
+
+		operator uint32() const
+		{
+			return Value;
+		}
+
+		bool operator==(const Counter& Other) const
+		{
+			return Value == Other.Value;
+		}
+		static int NumVoid;
+		static int NumCopy;
+		static int NumMove;
+		static int NumDestruct;
+		static void Clear()
+		{
+			NumVoid = NumCopy = NumMove = NumDestruct = 0;
+		}
+
+		uint32 Value;
+	};
+
+	FRingBufferTest()
+	{
+	}
+
+	bool RunTest()
+	{
+		// Test empty
+		{
+			TRingBuffer<uint32> Q(0);
+
+			TestTrue(TEXT("Test empty - IsEmpty"), Q.IsEmpty());
+			TestEqual(TEXT("Test empty - Size"), Q.Num(), 0);
+			TestEqual(TEXT("Test empty - Capacity"), Q.Max(), 0);
+			TestEqual(TEXT("Test empty - Iterator"), Q.begin(), Q.end());
+			TestEqual(TEXT("Test empty - ConvertPointerToIndex"), Q.ConvertPointerToIndex(nullptr), INDEX_NONE);
+			TestEqual(TEXT("Test empty - ConvertPointerToIndex"), Q.ConvertPointerToIndex(reinterpret_cast<uint32*>(this)), INDEX_NONE);
+			Q.Trim();
+			TestEqual(TEXT("Test Trim From empty - Size"), Q.Num(), 0);
+			TestEqual(TEXT("Test Trim From empty - Capacity"), Q.Max(), 0);
+			Q.Reset();
+			TestEqual(TEXT("Test Reset From empty - Size"), Q.Num(), 0);
+			TestEqual(TEXT("Test Reset From empty - Capacity"), Q.Max(), 0);
+			Q.Empty(0);
+			TestEqual(TEXT("Test Empty From empty - Size"), Q.Num(), 0);
+			TestEqual(TEXT("Test Empty From empty - Capacity"), Q.Max(), 0);
+			Q.PopFront(0);
+			Q.Pop(0);
+			TestEqual(TEXT("Test Pop on empty - Size"), Q.Num(), 0);
+			TestEqual(TEXT("Test Pop on empty - Capacity"), Q.Max(), 0);
+			TestEqual(TEXT("Test empty - IsValidIndex"), Q.IsValidIndex(0), false);
+
+
+			const TRingBuffer<uint32> ConstQ(0);
+			TestTrue(TEXT("Test const empty - IsEmpty"), ConstQ.IsEmpty());
+			TestEqual(TEXT("Test const empty - Size"), ConstQ.Num(), 0);
+			TestEqual(TEXT("Test const empty - Capacity"), ConstQ.Max(), 0);
+			TestEqual(TEXT("Test const empty - Iterator"), ConstQ.begin(), ConstQ.end());
+			TestEqual(TEXT("Test const empty - ConvertPointerToIndex"), ConstQ.ConvertPointerToIndex(reinterpret_cast<uint32*>(this)), INDEX_NONE);
+		}
+
+		// Test Adding a sequence of elements
+		{
+			const TRingBuffer<int32>::IndexType FirstSize = 8;
+
+			TRingBuffer<int32> Q(0);
+
+			TestEqual(TEXT("Test AddSequence - Capacity (Implementation Detail)"), Q.Max(), 0);
+			Q.Emplace(0);
+			TestEqual(TEXT("Test AddSequence - Size"), Q.Num(), 1);
+			TestEqual(TEXT("Test AddSequence - Capacity (Implementation Detail)"), Q.Max(), 1);
+			Q.Emplace(1);
+			TestEqual(TEXT("Test AddSequence - Size"), Q.Num(), 2);
+			TestEqual(TEXT("Implementation Detail - These tests expect that growing size will set capacity to successive powers of 2."), Q.Max(), 2);
+			for (int32 It = 2; It < FirstSize; ++It)
+			{
+				Q.Emplace(It);
+				TestEqual(TEXT("Test AddSequence - Size"), Q.Num(), It + 1);
+				TestEqual(TEXT("Test AddSequence - Capacity (Implementation Detail)"), static_cast<uint32>(Q.Max()), FMath::RoundUpToPowerOfTwo(It + 1));
+			}
+
+			for (int32 Index = 0; Index < FirstSize; ++Index)
+			{
+				TestEqual(TEXT("Test AddSequence - Expected values"), Q[Index], Index);
+				TestEqual(TEXT("Test AddSequence const- Expected values"), const_cast<TRingBuffer<int32>&>(Q)[Index], Index);
+			}
+
+			const TRingBuffer<int32>::IndexType SecondSize = 13;
+			for (int32 It = FirstSize; It < SecondSize; ++It)
+			{
+				Q.Emplace(It);
+				TestEqual(TEXT("Test AddSequence non powerof2 - Size"), Q.Num(), It + 1);
+				TestEqual(TEXT("Test AddSequence non powerof2 const - Capacity (Implementation Detail)"), static_cast<uint32>(Q.Max()), FMath::RoundUpToPowerOfTwo(It + 1));
+			}
+
+			for (int32 Index = 0; Index < FirstSize; ++Index)
+			{
+				TestEqual(TEXT("Test AddSequence non powerof2 - Expected values"), Q[Index], Index);
+				TestEqual(TEXT("Test AddSequence non powerof2 const - Expected values"), const_cast<TRingBuffer<int32>&>(Q)[Index], Index);
+			}
+		}
+
+		// Test Add under/over Capacity
+		{
+			const TRingBuffer<int32>::IndexType FirstElementsToAdd = 3;
+			const TRingBuffer<int32>::IndexType InitialCapacity = 8;
+			const TRingBuffer<int32>::IndexType SecondElementsToAdd = 9;
+
+			TRingBuffer<int32> Q(InitialCapacity);
+
+			for (int32 It = 0; It < FirstElementsToAdd; ++It)
+			{
+				Q.Emplace(It);
+			}
+
+			TestEqual(TEXT("Test Add under Capacity - Size"), Q.Num(), FirstElementsToAdd);
+			TestEqual(TEXT("Test Add under Capacity - Capacity"), Q.Max(), InitialCapacity);
+			for (int32 Index = 0; Index < FirstElementsToAdd; ++Index)
+			{
+				TestEqual(TEXT("Test Add under Capacity - Expected values"), Q[Index], Index);
+				TestEqual(TEXT("Test Add under Capacity const - Expected values"), const_cast<TRingBuffer<int32>&>(Q)[Index], Index);
+			}
+
+			for (int32 It = FirstElementsToAdd; It < SecondElementsToAdd; ++It)
+			{
+				Q.Emplace(It);
+			}
+
+			TestEqual(TEXT("Test Add over Capacity - Size"), Q.Num(), SecondElementsToAdd);
+			TestEqual(TEXT("Test Add over Capacity - Capacity (Implementation Detail)"), static_cast<uint32>(Q.Max()), FMath::RoundUpToPowerOfTwo(SecondElementsToAdd));
+			for (int32 Index = 0; Index < SecondElementsToAdd; ++Index)
+			{
+				TestEqual(TEXT("Test Add over Capacity - Expected values"), Q[Index], Index);
+				TestEqual(TEXT("Test Add over Capacity const - Expected values"), const_cast<TRingBuffer<int32>&>(Q)[Index], Index);
+			}
+		}
+
+		// Test Last/First
+		{
+			TRingBuffer<uint32> Q({ 0,1,2,3 });
+			TestEqual(TEXT("Test Last"), 3, Q.Last());
+			Q.Last() = 4;
+			TestEqual(TEXT("Test Last const"), 4, const_cast<TRingBuffer<uint32>&>(Q).Last());
+			TestEqual(TEXT("Test First"), 0, Q.First());
+			Q.First() = 5;
+			TestEqual(TEXT("Test First const"), 5, const_cast<TRingBuffer<uint32>&>(Q).First());
+		}
+
+		// Test PopFrontValue/PopValue
+		{
+			TRingBuffer<Counter> Q({ 31,32,33 });
+			Q.AddFront(30);
+
+			Counter::Clear();
+			Counter C(Q.PopFrontValue());
+			TestEqual(TEXT("PopFrontValue - PoppedValue"), C.Value, 30);
+			TestTrue(TEXT("PopFrontValue - ConstructorCounts"), Counter::NumMove > 0 && Counter::NumCopy == 0);
+			TestEqual(TEXT("PopFrontValue - Remaining Values"), Q, TRingBuffer<Counter>({ 31,32,33 }));
+			Counter::Clear();
+			TestEqual(TEXT("PopFrontValue Inline - PoppedValue"), Q.PopFrontValue().Value, 31);
+			TestTrue(TEXT("PopFrontValue Inline - ConstructorCounts"), Counter::NumCopy == 0);
+			TestEqual(TEXT("PopFrontValue Inline - Remaining Values"), Q, TRingBuffer<Counter>({ 32,33 }));
+
+			Counter::Clear();
+			Counter D(Q.PopValue());
+			TestEqual(TEXT("PopValue - PoppedValue"), D.Value, 33);
+			TestTrue(TEXT("PopValue - ConstructorCounts"), Counter::NumMove > 0 && Counter::NumCopy == 0);
+			TestEqual(TEXT("PopValue - Remaining Values"), Q, TRingBuffer<Counter>({ Counter(32) }));
+			Counter::Clear();
+			TestEqual(TEXT("PopValue Inline - PoppedValue"), Q.PopValue().Value, 32);
+			TestTrue(TEXT("PopValue Inline - ConstructorCounts"), Counter::NumCopy == 0);
+			TestTrue(TEXT("PopValue Inline - Remaining Values"), Q.IsEmpty());
+		}
+
+		// Test Initializer_List
+		{
+			const TRingBuffer<int32>::IndexType InitializerSize = 9;
+			TRingBuffer<int32> Q({ 0, 1, 2, 3, 4, 5, 6, 7, 8 });
+
+			TestEqual(TEXT("Test Initializer_List - Size"), Q.Num(), InitializerSize);
+			TestEqual(TEXT("Test Initializer_List - Capacity (Implementation Detail)"), static_cast<uint32>(Q.Max()), FMath::RoundUpToPowerOfTwo(InitializerSize));
+			for (int32 Index = 0; Index < InitializerSize; ++Index)
+			{
+				TestEqual(TEXT("Test Initializer_List - Expected values"), Q[Index], Index);
+			}
+		}
+
+		// Test RingBuffer's Copy Constructors et al
+		{
+			TRingBuffer<uint32> Original({ 0,1,2,3,4,5,6,7 });
+			TRingBuffer<uint32> Copy(Original);
+			TestEqual(TEXT("Copy Constructor"), Original, Copy);
+			TRingBuffer<uint32> Moved(MoveTemp(Copy));
+			TestEqual(TEXT("Move Constructor"), Original, Moved);
+			TestEqual(TEXT("Move Constructor did in fact move"), Copy.Max(), 0);
+			TRingBuffer<uint32> AssignCopy;
+			AssignCopy = Original;
+			TestEqual(TEXT("Copy Assignment"), Original, AssignCopy);
+			TRingBuffer<uint32> AssignMove;
+			AssignMove = MoveTemp(AssignCopy);
+			TestEqual(TEXT("Move Assignment"), Original, AssignMove);
+			TestEqual(TEXT("Move Assignment did in fact move"), AssignCopy.Max(), 0);
+		}
+
+		// Test Equality 
+		{
+			auto TestEquality = [this](const TCHAR* Message, bool ExpectedEqual, const TRingBuffer<int32>& A, const TRingBuffer<int32>& B)
+			{
+				TestEqual(*FString::Printf(TEXT("Test equality - %s - A == B"), Message), A == B, ExpectedEqual);
+				TestEqual(*FString::Printf(TEXT("Test equality - %s - B == A"), Message), B == A, ExpectedEqual);
+				TestEqual(*FString::Printf(TEXT("Test equality - %s - A != B"), Message), A != B, !ExpectedEqual);
+				TestEqual(*FString::Printf(TEXT("Test equality - %s - B != A"), Message), B != A, !ExpectedEqual);
+			};
+
+			TestEquality(TEXT("empty"), true, TRingBuffer<int32>(0), TRingBuffer<int32>(0));
+			TestEquality(TEXT("empty different capacities"), true, TRingBuffer<int32>(0), TRingBuffer<int32>(8));
+			TestEquality(TEXT("equal nonempty powerof2"), true, TRingBuffer<int32>({ 0, 1, 2, 3 }), TRingBuffer<int32>({ 0, 1, 2, 3 }));
+			TestEquality(TEXT("equal nonempty nonpowerof2"), true, TRingBuffer<int32>({ 0, 1, 2, 3, 4, 5 }), TRingBuffer<int32>({ 0, 1, 2, 3, 4, 5 }));
+			{
+				TRingBuffer<int32> QNum6Cap16(16);
+				for (int32 Index = 0; Index < 6; ++Index)
+				{
+					QNum6Cap16.Add(Index);
+				}
+				TestEquality(TEXT("equal nonempty different capacities"), true, QNum6Cap16, TRingBuffer<int32>({ 0, 1, 2, 3, 4, 5 }));
+			}
+
+			TestEquality(TEXT("empty to nonempty"), false, TRingBuffer<int32>(0), TRingBuffer<int32>({ 0, 1, 2, 3, 4, 5 }));
+			TestEquality(TEXT("smaller size to bigger size"), false, TRingBuffer<int32>({ 0, 1, 2 }), TRingBuffer<int32>({ 0, 1, 2, 3, 4, 5 }));
+			TestEquality(TEXT("same size different elements"), false, TRingBuffer<int32>({ 0, 1, 2 }), TRingBuffer<int32>({ 0, 1, 3 }));
+			TestEquality(TEXT("same elements different order"), false, TRingBuffer<int32>({ 0, 1, 2 }), TRingBuffer<int32>({ 0, 2, 1 }));
+
+			for (int HasPow2 = 0; HasPow2 < 2; ++HasPow2)
+			{
+				const int Count = HasPow2 ? 8 : 7;
+				TRingBuffer<int32> Q0Pop;
+				TRingBuffer<int32> Q1PopFront;
+				TRingBuffer<int32> Q2PopFront;
+				TRingBuffer<int32> Q1Pop;
+				TRingBuffer<int32> Q2Pop;
+				TRingBuffer<int32> Q2PopFront3Pop;
+				Q1PopFront.Add(47);
+				Q2PopFront.Add(576);
+				Q2PopFront.Add(-5);
+				Q2PopFront3Pop.Add(84);
+				Q2PopFront3Pop.Add(1000);
+				for (int Index = 0; Index < Count; ++Index)
+				{
+					Q0Pop.Add(Index);
+					Q1PopFront.Add(Index);
+					Q2PopFront.Add(Index);
+					Q1Pop.Add(Index);
+					Q2Pop.Add(Index);
+					Q2PopFront3Pop.Add(Index);
+				}
+				Q1PopFront.PopFront();
+				Q2PopFront.PopFront();
+				Q2PopFront.PopFront();
+				Q1Pop.Add(-18);
+				Q1Pop.Pop();
+				Q2Pop.Add(105);
+				Q2Pop.Add(219);
+				Q2Pop.Pop();
+				Q2Pop.Pop();
+				Q2PopFront3Pop.Add(456);
+				Q2PopFront3Pop.Add(654);
+				Q2PopFront3Pop.Add(8888888);
+				Q2PopFront3Pop.PopFront();
+				Q2PopFront3Pop.Pop();
+				Q2PopFront3Pop.PopFront();
+				Q2PopFront3Pop.Pop();
+				Q2PopFront3Pop.Pop();
+
+				const TCHAR* Names[] =
+				{
+					TEXT("Q0Pop"),
+					TEXT("Q1PopFront"),
+					TEXT("Q2PopFront"),
+					TEXT("Q1Pop"),
+					TEXT("Q2Pop"),
+					TEXT("Q2PopFront3Pop"),
+				};
+				TRingBuffer<int32>* Pops[] =
+				{
+					&Q0Pop,
+					&Q1PopFront,
+					&Q2PopFront,
+					&Q1Pop,
+					&Q2Pop,
+					&Q2PopFront3Pop
+				};
+
+
+				auto TestThesePops = [this, HasPow2, &TestEquality, &Names, &Pops](int TrialA, int TrialB)
+				{
+					TestEquality(*FString::Printf(TEXT("%s - %s - %s"), Names[TrialA], Names[TrialB], (HasPow2 ? TEXT("powerof2") : TEXT("nonpowerof2"))), true, *Pops[TrialA], *Pops[TrialB]);
+				};
+
+				for (int TrialA = 0; TrialA < UE_ARRAY_COUNT(Names); ++TrialA)
+				{
+					for (int TrialB = TrialA /* test each against itself as well */; TrialB < UE_ARRAY_COUNT(Names); ++TrialB)
+					{
+						TestThesePops(TrialA, TrialB);
+					}
+				}
+			}
+		}
+
+		// Test Add and pop all
+		for (int Direction = 0; Direction < 2; ++Direction)
+		{
+			bool bIsAddBack = Direction == 0;
+			auto GetMessage = [&bIsAddBack](const TCHAR* Message)
+			{
+				return FString::Printf(TEXT("Test %s (%s)"), Message, (bIsAddBack ? TEXT("AddBack") : TEXT("AddFront")));
+			};
+
+			// Test Mixed Adds and Pops
+			{
+				const TRingBuffer<uint32>::IndexType ElementsToAdd = 256;
+				const TRingBuffer<uint32>::IndexType ElementPopMod = 16;
+				const TRingBuffer<uint32>::IndexType ExpectedSize = 256 - ElementPopMod;
+				const TRingBuffer<uint32>::IndexType ExpectedCapacity = 256;
+
+				TRingBuffer<uint32> Q(4);
+
+				uint32 ExpectedPoppedValue = 0;
+				for (uint32 It = 0; It < 256; ++It)
+				{
+					if (bIsAddBack)
+					{
+						Q.Add(It);
+						TestEqual(*GetMessage(TEXT("Add and pop - Add")), It, Q[Q.Num() - 1]);
+					}
+					else
+					{
+						Q.AddFront(It);
+						TestEqual(*GetMessage(TEXT("Add and pop - Add")), It, Q[0]);
+					}
+
+					if (It % ElementPopMod == 0)
+					{
+						uint32 PoppedValue;
+						if (bIsAddBack)
+						{
+							PoppedValue = Q[0];
+							Q.PopFront();
+						}
+						else
+						{
+							PoppedValue = Q[Q.Num() - 1];
+							Q.Pop();
+						}
+						TestEqual(*GetMessage(TEXT("Add and pop - Pop")), ExpectedPoppedValue, PoppedValue);
+						++ExpectedPoppedValue;
+					}
+				}
+
+				TestEqual(*GetMessage(TEXT("Add and pop - Size")), Q.Num(), ExpectedSize);
+				TestEqual(*GetMessage(TEXT("Add and pop - Capacity")), Q.Max(), ExpectedCapacity);
+				TestTrue(*GetMessage(TEXT("Add and pop - IntegerRange")), IsIntegerRange(Q, ExpectedPoppedValue, ExpectedPoppedValue + ExpectedSize, bIsAddBack));
+			}
+
+
+			// Popping down to empty
+			{
+				const TRingBuffer<uint32>::IndexType ElementsToAdd = 256;
+
+				TRingBuffer<uint32> Q(ElementsToAdd);
+
+				TestTrue(*GetMessage(TEXT("Add and pop all - IsEmpty before")), Q.IsEmpty());
+				TestEqual(*GetMessage(TEXT("Add and pop all - Size before")), Q.Num(), 0);
+
+				for (TRingBuffer<int32>::IndexType It = 0; It < ElementsToAdd; ++It)
+				{
+					if (bIsAddBack)
+					{
+						Q.Add(It);
+					}
+					else
+					{
+						Q.AddFront(It);
+					}
+				}
+
+				TestEqual(*GetMessage(TEXT("Add and pop all - Size")), Q.Num(), ElementsToAdd);
+				TestEqual(*GetMessage(TEXT("Add and pop all - Capacity")), Q.Max(), ElementsToAdd);
+				TestTrue(*GetMessage(TEXT("Add and pop all - Expected")), IsIntegerRange(Q, 0, ElementsToAdd, bIsAddBack));
+
+				for (TRingBuffer<int32>::IndexType It = 0; It < ElementsToAdd; ++It)
+				{
+					if (bIsAddBack)
+					{
+						Q.PopFront();
+					}
+					else
+					{
+						Q.Pop();
+					}
+				}
+
+				TestTrue(*GetMessage(TEXT("Add and pop all - IsEmpty after")), Q.IsEmpty());
+				TestEqual(*GetMessage(TEXT("Add and pop all - Size after")), Q.Num(), 0);
+				TestEqual(*GetMessage(TEXT("Add and pop all - Capacity after")), Q.Max(), ElementsToAdd);
+			}
+
+			// Test index wrap
+			{
+				for (int32 Offset : {-12, -8, -5, -1, 0, 2, 7, 8, 15})
+				{
+					const TRingBuffer<uint32>::IndexType ElementsToAdd = 256;
+					const TRingBuffer<uint32>::IndexType ElementPopMod = 16;
+					const TRingBuffer<uint32>::IndexType ExpectedSize = 256 - ElementPopMod;
+					const TRingBuffer<uint32>::IndexType ExpectedCapacity = 256;
+
+					TRingBuffer<uint32> Q(8);
+
+					// Set front and afterback to an arbitrary offset
+					// Note that AfterBack is always exactly equal to Front + Num()
+					Q.Front = Offset;
+					Q.AfterBack = Q.Front;
+
+					TestTrue(*GetMessage(TEXT("index wrap - IsEmpty before")), Q.IsEmpty());
+					TestEqual(*GetMessage(TEXT("index wrap - Size before")), Q.Num(), 0);
+
+					for (TRingBuffer<uint32>::IndexType It = 0; It < ElementsToAdd; ++It)
+					{
+						if (bIsAddBack)
+						{
+							Q.Add(It);
+						}
+						else
+						{
+							Q.AddFront(It);
+						}
+					}
+
+					TestEqual(*GetMessage(TEXT("index wrap - Size")), Q.Num(), ElementsToAdd);
+					TestEqual(*GetMessage(TEXT("index wrap - Capacity")), Q.Max(), ElementsToAdd);
+					TestTrue(*GetMessage(TEXT("index wrap - Expected")), IsIntegerRange(Q, 0, ElementsToAdd, bIsAddBack));
+
+					for (TRingBuffer<int32>::IndexType It = 0; It < ElementsToAdd; ++It)
+					{
+						if (bIsAddBack)
+						{
+							Q.PopFront();
+						}
+						else
+						{
+							Q.Pop();
+						}
+					}
+
+					TestTrue(*GetMessage(TEXT("index wrap - IsEmpty after")), Q.IsEmpty());
+					TestEqual(*GetMessage(TEXT("index wrap - Size after")), Q.Num(), 0);
+					TestEqual(*GetMessage(TEXT("index wrap - Capacity after")), Q.Max(), ElementsToAdd);
+				}
+			}
+		}
+
+		// Test Trim
+		{
+			const TRingBuffer<int32>::IndexType ElementsToAdd = 9;
+			const TRingBuffer<int32>::IndexType ElementsToPop = 5;
+			const TRingBuffer<int32>::IndexType ExpectedCapacity = 16;
+			const TRingBuffer<int32>::IndexType ExpectedCapacityAfterTrim = 4;
+
+			TRingBuffer<uint32> Q(0);
+
+			for (TRingBuffer<int32>::IndexType It = 0; It < ElementsToAdd; ++It)
+			{
+				Q.Add(It);
+			}
+
+			TestEqual(TEXT("Test Trim - Size"), Q.Num(), ElementsToAdd);
+			TestEqual(TEXT("Test Trim - Capacity"), Q.Max(), ExpectedCapacity);
+			TestTrue(TEXT("Test Trim - Expected"), IsIntegerRange(Q, 0, ElementsToAdd));
+
+			for (TRingBuffer<int32>::IndexType It = 0; It < ElementsToPop; ++It)
+			{
+				Q.PopFront();
+			}
+
+			Q.Trim();
+
+			TestEqual(TEXT("Test Trim - Size"), Q.Num(), ElementsToAdd - ElementsToPop);
+			TestEqual(TEXT("Test Trim - Capacity"), Q.Max(), ExpectedCapacityAfterTrim);
+			TestTrue(TEXT("Test Trim - Expected"), IsIntegerRange(Q, ElementsToPop, ElementsToAdd));
+		}
+
+		// Test First and Last acting as two stacks
+		{
+			TRingBuffer<uint32> Q;
+
+			const uint32 ElementsToAdd = 64;
+			const uint32 ElementPopMod = 5;
+
+			for (uint32 It = 0; It < ElementsToAdd; ++It)
+			{
+				Q.Add(It);
+				TestEqual(TEXT("Test TwoStacks - AddBack"), Q.Last(), It);
+				Q.AddFront(It);
+				TestEqual(TEXT("Test TwoStacks - AddFront"), Q.First(), It);
+				if (It % ElementPopMod == 0)
+				{
+					uint32 AddValue = 0xfefefefe;
+					Q.Add(AddValue);
+					TestEqual(TEXT("Test TwoStacks - Sporadic Pop"), Q.Last(), AddValue);
+					Q.Pop();
+					Q.AddFront(AddValue);
+					TestEqual(TEXT("Test TwoStacks - Sporadic PopFront"), Q.First(), AddValue);
+					Q.PopFront();
+				}
+			}
+
+			TestEqual(TEXT("Test TwoStacks - MiddleSize"), Q.Num(), ElementsToAdd * 2);
+			for (uint32 It = 0; It < ElementsToAdd * 2; ++It)
+			{
+				TestEqual(*FString::Printf(TEXT("TwoStacks - Middle value %d"), It), Q[It], (It < ElementsToAdd ? ElementsToAdd - 1 - It : It - ElementsToAdd));
+			}
+
+			for (uint32 It = 0; It < ElementsToAdd; ++It)
+			{
+				TestEqual(TEXT("Test TwoStacks - Final Pop"), Q.Last(), ElementsToAdd - 1 - It);
+				Q.Pop();
+				TestEqual(TEXT("Test TwoStacks - Final PopFront"), Q.First(), ElementsToAdd - 1 - It);
+				Q.PopFront();
+			}
+
+			TestEqual(TEXT("Test TwoStacks - FinalSize"), Q.Num(), 0);
+		}
+
+		// Test adding into space that has been cleared from popping on the other side
+		{
+			for (int Direction = 0; Direction < 2; ++Direction)
+			{
+				bool bIsAddBack = Direction == 0;
+				auto GetMessage = [bIsAddBack](const TCHAR* Message)
+				{
+					return FString::Printf(TEXT("Test AddIntoPop - %s (%s)"), Message, (bIsAddBack ? TEXT("AddBack") : TEXT("AddFront")));
+				};
+				TRingBuffer<uint32> Q({ 0,1,2,3,4,5,6,7 });
+				TRingBuffer<int32>::IndexType InitialSize = 8;
+				TestEqual(*GetMessage(TEXT("InitialSize")), InitialSize, Q.Num());
+				TestEqual(*GetMessage(TEXT("InitialCapacity (Implementation Detail)")), InitialSize, Q.Max());
+
+				if (bIsAddBack)
+				{
+					Q.Pop();
+				}
+				else
+				{
+					Q.PopFront();
+				}
+				TestEqual(*GetMessage(TEXT("PoppedSize")), InitialSize - 1, Q.Num());
+				TestEqual(*GetMessage(TEXT("PoppedCapacity")), InitialSize, Q.Max());
+
+				if (bIsAddBack)
+				{
+					Q.AddFront(8);
+				}
+				else
+				{
+					Q.Add(8);
+				}
+				TestEqual(*GetMessage(TEXT("AddedSize")), InitialSize, Q.Num());
+				TestEqual(*GetMessage(TEXT("AddedCapacity")), InitialSize, Q.Max());
+				if (bIsAddBack)
+				{
+					TestEqual(*GetMessage(TEXT("AddedValues")), Q, TRingBuffer<uint32>({ 8,0,1,2,3,4,5,6 }));
+				}
+				else
+				{
+					TestEqual(*GetMessage(TEXT("AddedValues")), Q, TRingBuffer<uint32>({ 1,2,3,4,5,6,7,8 }));
+				}
+
+				if (bIsAddBack)
+				{
+					Q.AddFront(9);
+				}
+				else
+				{
+					Q.Add(9);
+				}
+				TestEqual(*GetMessage(TEXT("Second AddedSize")), InitialSize + 1, Q.Num());
+				TestEqual(*GetMessage(TEXT("Second AddedCapacity")), static_cast<uint32>(FMath::RoundUpToPowerOfTwo(InitialSize + 1)), Q.Max());
+				if (bIsAddBack)
+				{
+					TestEqual(*GetMessage(TEXT("Second AddedValues")), Q, TRingBuffer<uint32>({ 9,8,0,1,2,3,4,5,6 }));
+				}
+				else
+				{
+					TestEqual(*GetMessage(TEXT("Second AddedValues")), Q, TRingBuffer<uint32>({ 1,2,3,4,5,6,7,8,9 }));
+				}
+			}
+		}
+
+		// Test Empty to a capacity
+		{
+			TRingBuffer<uint32> Q(16);
+			TestEqual(TEXT("Test EmptyToCapacity - InitialCapacity"), 16, Q.Max());
+			Q.Empty(8);
+			TestEqual(TEXT("Test EmptyToCapacity - Lower"), 8, Q.Max());
+			Q.Empty(32);
+			TestEqual(TEXT("Test EmptyToCapacity - Higher"), 32, Q.Max());
+		}
+
+		// Test Different Add constructors
+		{
+			auto Clear = []()
+			{
+				Counter::Clear();
+			};
+			auto TestCounts = [this](const TCHAR* Message, int32 NumVoid, int32 NumCopy, int32 NumMove, int32 NumDestruct)
+			{
+				TestTrue(Message, NumVoid == Counter::NumVoid && NumCopy == Counter::NumCopy && NumMove == Counter::NumMove && NumDestruct == Counter::NumDestruct);
+			};
+
+			Clear();
+			{
+				TRingBuffer<Counter> QEmpty(4);
+				QEmpty.Reserve(8);
+				QEmpty.Empty();
+				TRingBuffer<Counter> QEmpty2(4);
+			}
+			TestCounts(TEXT("Test Add Constructors - Unallocated elements call no constructors/destructors"), 0, 0, 0, 0);
+			{
+				TRingBuffer<Counter> QEmpty(4);
+				QEmpty.Emplace();
+				QEmpty.Pop();
+				Clear();
+			}
+			TestCounts(TEXT("Test Add Constructors - Already removed element calls no destructors"), 0, 0, 0, 0);
+
+
+			uint32 MarkerValue = 0x54321;
+			Counter CounterA(MarkerValue);
+
+			TRingBuffer<Counter> Q(4);
+			Clear();
+			for (int Direction = 0; Direction < 2; ++Direction)
+			{
+				bool bAddBack = Direction == 0;
+				auto TestDirCounts = [this, bAddBack, &TestCounts, &Q, &Clear, MarkerValue](const TCHAR* Message, int32 NumVoid, int32 NumCopy, int32 NumMove, int32 NumDestruct, bool bWasInitialized = true)
+				{
+					const TCHAR* DirectionText = bAddBack ? TEXT("Back") : TEXT("Front");
+					bool bElementExists = Q.Num() == 1;
+					TestTrue(*FString::Printf(TEXT("Test Add Constructors - %s%s ElementExists"), Message, DirectionText), bElementExists);
+					if (bWasInitialized && bElementExists)
+					{
+						TestTrue(*FString::Printf(TEXT("Test Add Constructors - %s%s ValueEquals"), Message, DirectionText), Q.First().Value == MarkerValue);
+					}
+					Q.PopFront();
+					TestCounts(*FString::Printf(TEXT("Test Add Constructors - %s%s CountsEqual"), Message, DirectionText), NumVoid, NumCopy, NumMove, NumDestruct);
+					Clear();
+				};
+
+				if (bAddBack) Q.Add(CounterA); else Q.AddFront(CounterA);
+				TestDirCounts(TEXT("Copy Add"), 0, 1, 0, 1);
+				if (bAddBack) Q.Add_GetRef(CounterA); else Q.AddFront_GetRef(CounterA);
+				TestDirCounts(TEXT("Copy GetRef Add"), 0, 1, 0, 1);
+				if (bAddBack) Q.Add(MoveTemp(CounterA)); else Q.AddFront(MoveTemp(CounterA));
+				TestDirCounts(TEXT("Move Add"), 0, 0, 1, 1);
+				if (bAddBack) Q.Add_GetRef(MoveTemp(CounterA)); else Q.AddFront_GetRef(MoveTemp(CounterA));
+				TestDirCounts(TEXT("Move GetRef Add"), 0, 0, 1, 1);
+				if (bAddBack) Q.Emplace(MarkerValue); else Q.EmplaceFront(MarkerValue);
+				TestDirCounts(TEXT("Emplace"), 1, 0, 0, 1);
+				if (bAddBack) Q.Emplace_GetRef(MarkerValue); else Q.EmplaceFront_GetRef(MarkerValue);
+				TestDirCounts(TEXT("GetRef Emplace"), 1, 0, 0, 1);
+				if (bAddBack) Q.AddUninitialized(); else Q.AddFrontUninitialized();
+				TestDirCounts(TEXT("Uninitialized Add"), 0, 0, 0, 1, false);
+				if (bAddBack) Q.AddUninitialized_GetRef(); else Q.AddFrontUninitialized_GetRef();
+				TestDirCounts(TEXT("Uninitialized GetRef Add"), 0, 0, 0, 1, false);
+			}
+		}
+
+		TestShiftIndex<uint32>();
+		TestShiftIndex<Counter>();
+
+		// Test RemoveAt
+		{
+			{
+				TRingBuffer<uint32> Q{ 0,1,2,3,4,5,6,7 };
+				Q.RemoveAt(2);
+				TestEqual(TEXT("Test RemoveAt Front Closest"), TRingBuffer<uint32>({ 0,1,3,4,5,6,7 }), Q);
+			}
+			{
+				TRingBuffer<uint32> Q{ 0,1,2,3,4,5,6,7 };
+				Q.RemoveAt(5);
+				TestEqual(TEXT("Test RemoveAt Back Closest"), TRingBuffer<uint32>({ 0,1,2,3,4,6,7 }), Q);
+			}
+			{
+				TRingBuffer<uint32> Q{ 0,1,2,3,4,5,6,7 };
+				int32 Offset = 4;
+				Q.Front += Offset;
+				Q.AfterBack += Offset;
+				//Now equal to: TRingBuffer<uint32> Q{ 4,5,6,7,0,1,2,3 };
+				Q.RemoveAt(2);
+				TestEqual(TEXT("Test RemoveAt Front Closest With Offset"), TRingBuffer<uint32>({ 4,5,7,0,1,2,3 }), Q);
+			}
+			{
+				TRingBuffer<uint32> Q{ 0,1,2,3,4,5,6,7 };
+				int32 Offset = 4;
+				Q.Front += Offset;
+				Q.AfterBack += Offset;
+				//Now equal to: TRingBuffer<uint32> Q{ 4,5,6,7,0,1,2,3 };
+				Q.RemoveAt(5);
+				TestEqual(TEXT("Test RemoveAt Back Closest With Offset"), TRingBuffer<uint32>({ 4,5,6,7,0,2,3 }), Q);
+			}
+		}
+
+		// Test Iteration
+		{
+			{
+				TRingBuffer<uint32> Q{ 0,1,2,3,4,5,6,7 };
+				uint32 Counter = 0;
+				for (uint32 Value : Q)
+				{
+					TestEqual(TEXT("Test Iteration - Value"), Counter++, Value);
+				}
+				TestEqual(TEXT("Test Iteration - Num"), Counter, 8);
+			}
+			{
+				TRingBuffer<uint32> Q{ 4,5,6,7,0,1,2,3 };
+				int32 Offset = 4;
+				Q.Front += Offset;
+				Q.AfterBack += Offset;
+				// Now equal to 0,1,2,3,4,5,6,7
+				uint32 Counter = 0;
+				for (uint32 Value : Q)
+				{
+					TestEqual(TEXT("Test Iteration with Offset - Value"), Counter++, Value);
+				}
+				TestEqual(TEXT("Test Iteration with Offset  - Num"), Counter, 8);
+			}
+		}
+
+		// Test ConvertPointerToIndex
+		{
+			{
+				TRingBuffer<uint32> Q{ 4,5,6,7,0,1,2,3 };
+				int32 Offset = 4;
+				Q.Front += Offset;
+				Q.AfterBack += Offset;
+				// Now equal to 0,1,2,3,4,5,6,7
+				TestEqual(TEXT("Test ConvertPointerToIndex - before array"), Q.ConvertPointerToIndex(&Q[0] - 100), INDEX_NONE);
+				TestEqual(TEXT("Test ConvertPointerToIndex - after array"), Q.ConvertPointerToIndex(&Q[0] + 100), INDEX_NONE);
+				for (int32 It = 0; It < 8; ++It)
+				{
+					TestEqual(TEXT("Test ConvertPointerToIndex - Values"), Q.ConvertPointerToIndex(&Q[It]), It);
+				}
+			}
+
+			{
+				TRingBuffer<uint32> Q(16);
+				for (int32 It = 7; It >= 0; --It)
+				{
+					Q.AddFront(It);
+				}
+				Q.Pop();
+				// 8 Invalids, followed by 0,1,2,3,4,5,6, followed by Invalid
+				for (int32 It = 0; It < 7; ++It)
+				{
+					TestEqual(TEXT("Test ConvertPointerToIndex - Cap - Values"), Q.ConvertPointerToIndex(&Q[It]), It);
+				}
+				TestEqual(TEXT("Test ConvertPointerToIndex - Cap - After End"), Q.ConvertPointerToIndex(&Q[6] + 1), INDEX_NONE);
+				TestEqual(TEXT("Test ConvertPointerToIndex - Cap - Before Start"), Q.ConvertPointerToIndex(&Q[0] - 1), INDEX_NONE);
+			}
+		}
+
+		// Test that setting Front to its maximum value and then popping the maximum number of elements does not break the contract that Front < capacity in StorageModulo space
+		{
+			TRingBuffer<uint32> Q(8);
+			Q.AddFront(0);
+			for (uint32 It = 1; It < 8; ++It)
+			{
+				Q.Add(It);
+			}
+			TestTrue(TEXT("Test Front<Capacity - Setup"), (Q.Front & Q.IndexMask) == Q.IndexMask && Q.Num() == Q.Max());
+			Q.PopFront(8);
+			TestTrue(TEXT("Test Front<Capacity - Contract is true"), static_cast<uint32>(Q.Front) < static_cast<uint32>(Q.Max()));
+		}
+
+		// Test IsValidIndex
+		{
+			TRingBuffer<uint32> Q({ 0,1,2,3,4 });
+			for (int32 It = 0; It < Q.Num(); ++It)
+			{
+				TestEqual(TEXT("IsValidIndex - InRange"), Q.IsValidIndex(It), true);
+			}
+			TestEqual(TEXT("IsValidIndex - Negative"), Q.IsValidIndex(-1), false);
+			TestEqual(TEXT("IsValidIndex - Num()"), Q.IsValidIndex(Q.Num() + 1), false);
+			TestEqual(TEXT("IsValidIndex - Capacity"), Q.IsValidIndex(Q.Max()), false);
+			TestEqual(TEXT("IsValidIndex - Capacity + 1"), Q.IsValidIndex(Q.Max()+1), false);
+		}
+
+		// Test Compact
+		{
+			{
+				TRingBuffer<uint32> QEmpty;
+				TestEqual(TEXT("Compact - Empty zero capacity"), QEmpty.Compact().Num(), 0);
+				QEmpty.Add(1);
+				QEmpty.PopFront();
+				TestEqual(TEXT("Compact - Empty non-zero capacity"), QEmpty.Compact().Num(), 0);
+			}
+			{
+				TArrayView<uint32> View;
+				TRingBuffer<uint32> Q(8);
+				Q.AddFront(37);
+				View = Q.Compact();
+				TestTrue(TEXT("Compact - Front at end"), ArrayViewsEqual(View, TArrayView<const uint32>({ 37 })));
+			}
+			{
+				TArrayView<uint32> View;
+				TRingBuffer<uint32> Q(8);
+				for (TRingBuffer<uint32>::IndexType It = 0; It < 6; ++It)
+				{
+					Q.Add(It);
+				}
+				Q.PopFront();
+				TRingBuffer<uint32>::StorageModuloType SavedFront = Q.Front;
+				TestTrue(TEXT("Compact - Front in middle - setup"), SavedFront > 0);
+				View = Q.Compact();
+				TestTrue(TEXT("Compact - Front in middle - values"), ArrayViewsEqual(View, TArrayView<const uint32>({ 1,2,3,4,5 })));
+				TestTrue(TEXT("Compact - Front in middle - no reallocate"), Q.Front == SavedFront);
+			}
+			{
+				TArrayView<uint32> View;
+				TRingBuffer<uint32> Q(8);
+				for (TRingBuffer<uint32>::IndexType It = 1; It < 8; ++It)
+				{
+					Q.Add(It);
+				}
+				Q.AddFront(0);
+				TestTrue(TEXT("Compact - Full array front at end - setup"), (Q.Front & Q.IndexMask) == 7);
+				View = Q.Compact();
+				TestTrue(TEXT("Compact - Full array front at end - values"), ArrayViewsEqual(View, TArrayView<const uint32>({0,1,2,3,4,5,6,7 })));
+				TestTrue(TEXT("Compact - Full array front at end - reallocated"), Q.Front == 0);
+			}
+			{
+				TArrayView<uint32> View;
+				TRingBuffer<uint32> Q(8);
+				for (TRingBuffer<uint32>::IndexType It = 0; It < 8; ++It)
+				{
+					Q.Add(It);
+				}
+				uint32* SavedData = Q.AllocationData;
+				TestTrue(TEXT("Compact - Full array front at start - setup"), Q.Front == 0);
+				View = Q.Compact();
+				TestTrue(TEXT("Compact - Full array front at start - values"), ArrayViewsEqual(View, TArrayView<const uint32>({ 0,1,2,3,4,5,6,7 })));
+				TestTrue(TEXT("Compact - Full array front at start - no reallocate"), Q.AllocationData == SavedData);
+			}
+		}
+
+		// Test Remove
+		{
+			Counter Value;
+			{
+				TRingBuffer<Counter> Q;
+				Value.Value = 2;
+				Counter::Clear();
+				TestEqual(TEXT("Remove - empty"), Q.Remove(Value), 0);
+				TestEqual(TEXT("Remove - empty - destructor count"), Counter::NumDestruct, 0);
+			}
+			{
+				TRingBuffer<Counter> Q({ 0,1,2,3,4 });
+				Value.Value = 5;
+				Counter::Clear();
+				TestEqual(TEXT("Remove - no hits"), Q.Remove(Value), 0);
+				TestEqual(TEXT("Remove - no hits - destructor count"), Counter::NumDestruct, 0);
+				Q.Add(5);
+				TestTrue(TEXT("Remove - no hits - values"), Q == TRingBuffer<Counter>({ 0,1,2,3,4,5 }));
+			}
+			{
+				TRingBuffer<Counter> Q({ 1,2,3,4 });
+				Q.AddFront(0);
+				Value.Value = 0;
+				Counter::Clear();
+				TestEqual(TEXT("Remove - one element at front - num"), Q.Remove(Value), 1);
+				TestEqual(TEXT("Remove - one element at front - destructor count"), Counter::NumDestruct, 5);
+				Q.Add(5);
+				TestTrue(TEXT("Remove - one element at front - values"), Q == TRingBuffer<Counter>({ 1,2,3,4,5 }));
+			}
+			{
+				TRingBuffer<Counter> Q({ 0,1,2,3,4 });
+				Value.Value = 2;
+				Counter::Clear();
+				TestEqual(TEXT("Remove - one element in mid - num"), Q.Remove(Value), 1);
+				TestEqual(TEXT("Remove - one element in mid - destructor count"), Counter::NumDestruct, 3);
+				Q.Add(5);
+				TestTrue(TEXT("Remove - one element in mid - values"), Q == TRingBuffer<Counter>({ 0,1,3,4,5 }));
+			}
+			{
+				TRingBuffer<Counter> Q({ 1,2,3,4 });
+				Q.AddFront(0);
+				Value.Value = 2;
+				Counter::Clear();
+				TestEqual(TEXT("Remove - one element in mid - front at end"), Q.Remove(Value), 1);
+				TestEqual(TEXT("Remove - one element in mid - front at end - destructor count"), Counter::NumDestruct, 3);
+				Q.Add(5);
+				TestTrue(TEXT("Remove - one element in mid - front at end - values"), Q == TRingBuffer<Counter>({ 0,1,3,4,5 }));
+			}
+			{
+				TRingBuffer<Counter> Q({ 0,1,2,3,4 });
+				Value.Value = 4;
+				Counter::Clear();
+				TestEqual(TEXT("Remove - one element - element at end - num"), Q.Remove(Value), 1);
+				TestEqual(TEXT("Remove - one element - element at end - destructor count"), Counter::NumDestruct, 1);
+				Q.Add(5);
+				TestTrue(TEXT("Remove - one element - element at end - values"), Q == TRingBuffer<Counter>({ 0,1,2,3,5 }));
+			}
+			{
+				TRingBuffer<Counter> Q({ 1,2,3,4 });
+				Q.AddFront(4);
+				Value.Value = 4;
+				Counter::Clear();
+				TestEqual(TEXT("Remove - one element at front one at end - num"), Q.Remove(Value), 2);
+				TestEqual(TEXT("Remove - one element at front one at end - destructor count"), Counter::NumDestruct, 5);
+				Q.Add(5);
+				TestTrue(TEXT("Remove - one element at front one at end - values"), Q == TRingBuffer<Counter>({ 1,2,3,5 }));
+			}
+			{
+				TRingBuffer<Counter> Q({ 1,2,3,4 });
+				Q.AddFront(1);
+				Value.Value = 1;
+				Counter::Clear();
+				TestEqual(TEXT("Remove - two elements - front at end - num"), Q.Remove(Value), 2);
+				TestEqual(TEXT("Remove - two elements - front at end - destructor count"), Counter::NumDestruct, 5);
+				Q.Add(5);
+				TestTrue(TEXT("Remove - two elements - front at end - values"), Q == TRingBuffer<Counter>({ 2,3,4,5 }));
+			}
+		}
+
+		return true;
+	}
+
+	template <typename T>
+	void TestShiftIndex()
+	{
+		// Test shifts at specific points
+		{
+			{
+				TRingBuffer<T> Q{ 0,1,2,3,4,5,6,7 };
+				Q.ShiftIndexToFront(5);
+				TestEqual(TEXT("ShiftIndexToFront"), TRingBuffer<T>({ 5,0,1,2,3,4,6,7 }), Q);
+				Q.ShiftIndexToBack(3);
+				TestEqual(TEXT("ShiftIndexToBack"), TRingBuffer<T>({ 5,0,1,3,4,6,7,2 }), Q);
+			}
+
+			{
+				TRingBuffer<T> Q{ 0,1,2,3,4,5,6,7 };
+				int32 Offset = 4;
+				Q.Front += Offset;
+				Q.AfterBack += Offset;
+				//Now equal to: TRingBuffer<uint32> Q{ 4,5,6,7,0,1,2,3 };
+				Q.ShiftIndexToFront(5);
+				TestEqual(TEXT("ShiftIndexToFront With Offset"), TRingBuffer<T>({ 1,4,5,6,7,0,2,3 }), Q);
+				Q.ShiftIndexToBack(3);
+				TestEqual(TEXT("ShiftIndexToBack With Offset"), TRingBuffer<T>({ 1,4,5,7,0,2,3,6 }), Q);
+			}
+
+			{
+				TRingBuffer<T> Q{ 0,1,2,3,4,5,6,7,8 };
+				TestEqual(TEXT("ShiftIndexToFront Cap - Capacity"), Q.Max(), 16);
+				Q.ShiftIndexToFront(5);
+				TestEqual(TEXT("ShiftIndexToFront Cap"), TRingBuffer<T>({ 5,0,1,2,3,4,6,7,8 }), Q);
+				Q.ShiftIndexToBack(3);
+				TestEqual(TEXT("ShiftIndexToBack Cap"), TRingBuffer<T>({ 5,0,1,3,4,6,7,8,2 }), Q);
+			}
+
+			{
+				TRingBuffer<T> Q(16);
+				for (int32 It = 7; It >= 0; --It)
+				{
+					Q.AddFront(It);
+				}
+				Q.Pop();
+				// 8 Invalids, followed by 0,1,2,3,4,5,6, followed by Invalid
+				Q.ShiftIndexToFront(5);
+				TestEqual(TEXT("ShiftIndexToFront Cap With Offset"), TRingBuffer<T>({ 5,0,1,2,3,4,6 }), Q);
+				Q.ShiftIndexToBack(3);
+				TestEqual(TEXT("ShiftIndexToBack Cap With Offset"), TRingBuffer<T>({ 5,0,1,3,4,6,2 }), Q);
+			}
+
+			{
+				TRingBuffer<T> Q(16);
+				for (int32 It = 7; It >= 0; --It)
+				{
+					Q.AddFront(It);
+				}
+				Q.Add(8);
+				// 8, (AfterBack), followed by 7 Invalids, followed by (Start) 0,1,2,3,4,5,6,7
+				Q.ShiftIndexToFront(8);
+				TestEqual(TEXT("ShiftIndexToFront Cap With Wrapped"), TRingBuffer<T>({ 8,0,1,2,3,4,5,6,7 }), Q);
+				Q.ShiftIndexToBack(0);
+				TestEqual(TEXT("ShiftIndexToBack Cap With Wrapped"), TRingBuffer<T>({ 0,1,2,3,4,5,6,7,8 }), Q);
+			}
+		}
+
+		// Test ShiftIndex of each possible index
+		{
+			int32 Count = 8;
+			for (int32 It = 0; It < Count; ++It)
+			{
+				TRingBuffer<T> Q({ 0,1,2,3,4,5,6,7 });
+				Q.ShiftIndexToBack(It);
+				int32 CheckIndex = 0;
+				for (; CheckIndex < It; ++CheckIndex)
+				{
+					TestEqual(*FString::Printf(TEXT("ShiftIndexToBack Entire Array Values (%d,%d)"), It, CheckIndex), CheckIndex, Q[CheckIndex]);
+				}
+				for (; CheckIndex < Count - 1; ++CheckIndex)
+				{
+					TestEqual(*FString::Printf(TEXT("ShiftIndexToBack Entire Array Values (%d,%d)"), It, CheckIndex), CheckIndex + 1, Q[CheckIndex]);
+				}
+				TestEqual(*FString::Printf(TEXT("ShiftIndexToBack Entire Array Values (%d,%d)"), It, Count - 1), It, Q[Count - 1]);
+			}
+			for (int32 It = 0; It < Count; ++It)
+			{
+				TRingBuffer<T> Q({ 0,1,2,3,4,5,6,7 });
+				Q.ShiftIndexToFront(It);
+
+				TestEqual(*FString::Printf(TEXT("ShiftIndexToFront Entire Array Values (%d,%d)"), It, 0), It, Q[0]);
+				int32 CheckIndex = 1;
+				for (; CheckIndex <= It; ++CheckIndex)
+				{
+					TestEqual(*FString::Printf(TEXT("ShiftIndexToFront Entire Array Values (%d,%d)"), It, CheckIndex), CheckIndex - 1, Q[CheckIndex]);
+				}
+				for (; CheckIndex < Count; ++CheckIndex)
+				{
+					TestEqual(*FString::Printf(TEXT("ShiftIndexToFront Entire Array Values (%d,%d)"), It, CheckIndex), CheckIndex, Q[CheckIndex]);
+				}
+			}
+		}
+	}
+
+	template <typename T, typename U>
+	static bool ArrayViewsEqual(const TArrayView<T>& A, const TArrayView<U>& B)
+	{
+		int32 Num = A.Num();
+		if (Num != B.Num())
+		{
+			return false;
+		}
+		for (int It = 0; It < Num; ++It)
+		{
+			if (A[It] != B[It])
+			{
+				return false;
+			}
+		}
+		return true;
+	}
+};
+
+int FRingBufferTest::Counter::NumVoid = 0;
+int FRingBufferTest::Counter::NumCopy = 0;
+int FRingBufferTest::Counter::NumMove = 0;
+int FRingBufferTest::Counter::NumDestruct = 0;
+
+TEST_CASE("Core::Containers::TRingBuffer::Smoke Test", "[Core][Containers][Smoke]")
+{
+	FRingBufferTest RingBufferTest;
+	RingBufferTest.RunTest();
+}
\ No newline at end of file
