Tests converted from 18848115 made to work multi-platform.

diff --git a/Engine/Source/Runtime/Core/Tests/Containers/StringViewTest.cpp b/Engine/Source/Runtime/Core/Tests/Containers/StringViewTest.cpp
new file mode 100644
--- /dev/null
+++ b/Engine/Source/Runtime/Core/Tests/Containers/StringViewTest.cpp
@@ -0,0 +1,954 @@
+// Copyright Epic Games, Inc. All Rights Reserved.
+
+#include "Containers/StringView.h"
+#include "Containers/UnrealString.h"
+#include "Misc/AutomationTest.h"
+#include "Misc/StringBuilder.h"
+#include "TestHarness.h"
+
+static_assert(TIsSame<typename FStringView::ElementType, TCHAR>::Value, "FStringView must use TCHAR.");
+static_assert(TIsSame<typename FAnsiStringView::ElementType, ANSICHAR>::Value, "FAnsiStringView must use ANSICHAR.");
+static_assert(TIsSame<typename FWideStringView::ElementType, WIDECHAR>::Value, "FWideStringView must use WIDECHAR.");
+
+static_assert(TIsSame<FStringView, TStringView<TCHAR>>::Value, "FStringView must be the same as TStringView<TCHAR>.");
+static_assert(TIsSame<FAnsiStringView, TStringView<ANSICHAR>>::Value, "FAnsiStringView must be the same as TStringView<ANSICHAR>.");
+static_assert(TIsSame<FWideStringView, TStringView<WIDECHAR>>::Value, "FWideStringView must be the same as TStringView<WIDECHAR>.");
+
+static_assert(TIsContiguousContainer<FStringView>::Value, "FStringView must be a contiguous container.");
+static_assert(TIsContiguousContainer<FAnsiStringView>::Value, "FAnsiStringView must be a contiguous container.");
+static_assert(TIsContiguousContainer<FWideStringView>::Value, "FWideStringView must be a contiguous container.");
+
+namespace UE::String::Private::TestArgumentDependentLookup
+{
+
+struct FTestType
+{
+	using ElementType = TCHAR;
+};
+
+const TCHAR* GetData(const FTestType&) { return TEXT("ABC"); }
+int32 GetNum(const FTestType&) { return 3; }
+} // UE::String::Private::TestArgumentDependentLookup
+
+template <> struct TIsContiguousContainer<UE::String::Private::TestArgumentDependentLookup::FTestType> { static constexpr bool Value = true; };
+
+constexpr const uint32 TestFlags = EAutomationTestFlags::ApplicationContextMask | EAutomationTestFlags::SmokeFilter;
+
+TEST_CASE("Core::Containers::FStringView::Constructor", "[Core][Containers][Smoke]")
+{
+	// Default View
+	{
+		FStringView View;
+		TestEqual(TEXT(""), View.Len(), 0);
+		TestTrue(TEXT("View.IsEmpty"), View.IsEmpty());
+	}
+
+	// Empty View
+	{
+		FStringView View(TEXT(""));
+		TestEqual(TEXT(""), View.Len(), 0);
+		TestTrue(TEXT("View.IsEmpty"), View.IsEmpty());
+	}
+
+	// Constructing from nullptr is supported; nullptr interpreted as empty string
+	{
+		FStringView View(nullptr);
+		TestEqual(TEXT(""), View.Len(), 0);
+		TestTrue(TEXT("View.IsEmpty"), View.IsEmpty());
+	}
+
+	// Create from a wchar literal
+	{
+		FStringView View(TEXT("Test Ctor"));
+		TestEqual(TEXT("View length"), View.Len(), FCStringWide::Strlen(TEXT("Test Ctor")));
+		TestEqual(TEXT("The result of Strncmp"), FCStringWide::Strncmp(View.GetData(), TEXT("Test Ctor"), View.Len()), 0);
+		TestFalse(TEXT("View.IsEmpty"), View.IsEmpty());
+	}
+
+	// Create from a sub section of a wchar literal
+	{
+		FStringView View(TEXT("Test SubSection Ctor"), 4);
+		TestEqual(TEXT("View length"), View.Len(), 4);
+		TestEqual(TEXT("The result of Strncmp"), FCStringWide::Strncmp(View.GetData(), TEXT("Test"), View.Len()), 0);
+		TestFalse(TEXT("View.IsEmpty"), View.IsEmpty());
+	}
+
+	// Create from a FString
+	{
+		FString String(TEXT("String Object"));
+		FStringView View(String);
+
+		TestEqual(TEXT("View length"), View.Len(), String.Len());
+		TestEqual(TEXT("The result of Strncmp"), FCStringWide::Strncmp(View.GetData(), *String, View.Len()), 0);
+		TestFalse(TEXT("View.IsEmpty"), View.IsEmpty());
+	}
+
+	// Create from an ansi literal
+	{
+		FAnsiStringView View("Test Ctor");
+		TestEqual(TEXT("View length"), View.Len(), FCStringAnsi::Strlen("Test Ctor"));
+		TestEqual(TEXT("The result of Strncmp"), FCStringAnsi::Strncmp(View.GetData(), "Test Ctor", View.Len()), 0);
+		TestFalse(TEXT("View.IsEmpty"), View.IsEmpty());
+	}
+
+	// Create from a sub section of an ansi literal
+	{
+		FAnsiStringView View("Test SubSection Ctor", 4);
+		TestEqual(TEXT("View length"), View.Len(), 4);
+		TestEqual(TEXT("The result of Strncmp"), FCStringAnsi::Strncmp(View.GetData(), "Test", View.Len()), 0);
+		TestFalse(TEXT("View.IsEmpty"), View.IsEmpty());
+	}
+
+	// Create using string view literals
+	{
+		FStringView View = TEXTVIEW("Test");
+		FAnsiStringView ViewAnsi = "Test"_ASV;
+		FWideStringView ViewWide = WIDETEXT("Test"_WSV);
+	}
+
+	// Verify that class template argument deduction is working
+	{
+		TStringView ViewAnsi((ANSICHAR*)"Test");
+		TStringView ViewWide((WIDECHAR*)WIDETEXT("Test"));
+		TStringView ViewUtf8((UTF8CHAR*)UTF8TEXT("Test"));
+	}
+	{
+		TStringView ViewAnsi((const ANSICHAR*)"Test");
+		TStringView ViewWide((const WIDECHAR*)WIDETEXT("Test"));
+		TStringView ViewUtf8((const UTF8CHAR*)UTF8TEXT("Test"));
+	}
+	{
+		TStringView ViewAnsi(WriteToAnsiString<16>("Test"));
+		TStringView ViewWide(WriteToWideString<16>(WIDETEXT("Test")));
+		TStringView ViewUtf8(WriteToUtf8String<16>(UTF8TEXT("Test")));
+	}
+	{
+		FString String(TEXT("Test"));
+		TStringView ViewString(String);
+	}
+
+	// Verify that argument-dependent lookup is working for GetData and GetNum
+	{
+		UE::String::Private::TestArgumentDependentLookup::FTestType Test;
+		FStringView View(Test);
+		TestTrue(TEXT("StringView ADL"), View.Equals(TEXTVIEW("ABC")));
+	}
+}
+
+TEST_CASE("Core::Containers::FStringView::Iterator", "[Core][Containers][Smoke]")
+{
+	// Iterate over a string view
+	{
+		const TCHAR* StringLiteralSrc = TEXT("Iterator!");
+		FStringView View(StringLiteralSrc);
+
+		for (TCHAR C : View)
+		{
+			TestTrue(TEXT("Iterators(0)-Iteration"), C == *StringLiteralSrc++);
+		}
+
+		// Make sure we iterated over the entire string
+		TestTrue(TEXT("Iterators(0-EndCheck"), *StringLiteralSrc == '\0');
+	}
+
+	// Iterate over a partial string view
+	{
+		const TCHAR* StringLiteralSrc = TEXT("Iterator|with extras!");
+		FStringView View(StringLiteralSrc, 8);
+
+		for (TCHAR C : View)
+		{
+			TestTrue(TEXT("Iterators(1)-Iteration"), C == *StringLiteralSrc++);
+		}
+
+		// Make sure we only iterated over the part of the string that the view represents
+		TestTrue(TEXT("Iterators(1)-EndCheck"), *StringLiteralSrc == '|');
+	}
+}
+
+TEST_CASE("Core::Containers::FStringView::Equality Test", "[Core][Containers][Smoke]")
+{
+	const ANSICHAR* AnsiStringLiteralSrc = "String To Test!";
+	const ANSICHAR* AnsiStringLiteralLower = "string to test!";
+	const ANSICHAR* AnsiStringLiteralUpper = "STRING TO TEST!";
+	const TCHAR* WideStringLiteralSrc = TEXT("String To Test!");
+	const TCHAR* WideStringLiteralLower = TEXT("string to test!");
+	const TCHAR* WideStringLiteralUpper = TEXT("STRING TO TEST!");
+	const TCHAR* WideStringLiteralShort = TEXT("String To");
+	const TCHAR* WideStringLiteralLonger = TEXT("String To Test! Extended");
+
+	FStringView WideView(WideStringLiteralSrc);
+
+	TestTrue(TEXT("Equality(0)"), WideView == WideStringLiteralSrc);
+	TestFalse(TEXT("Equality(1)"), WideView != WideStringLiteralSrc);
+	TestTrue(TEXT("Equality(2)"), WideView == WideStringLiteralLower);
+	TestFalse(TEXT("Equality(3)"), WideView != WideStringLiteralLower);
+	TestTrue(TEXT("Equality(4)"), WideView == WideStringLiteralUpper);
+	TestFalse(TEXT("Equality(5)"), WideView != WideStringLiteralUpper);
+	TestFalse(TEXT("Equality(6)"), WideView == WideStringLiteralShort);
+	TestTrue(TEXT("Equality(7)"), WideView != WideStringLiteralShort);
+	TestFalse(TEXT("Equality(8)"), WideView == WideStringLiteralLonger);
+	TestTrue(TEXT("Equality(9)"), WideView != WideStringLiteralLonger);
+
+	TestTrue(TEXT("Equality(10)"), WideStringLiteralSrc == WideView);
+	TestFalse(TEXT("Equality(11)"), WideStringLiteralSrc != WideView);
+	TestTrue(TEXT("Equality(12)"), WideStringLiteralLower == WideView);
+	TestFalse(TEXT("Equality(13)"), WideStringLiteralLower != WideView);
+	TestTrue(TEXT("Equality(14)"), WideStringLiteralUpper == WideView);
+	TestFalse(TEXT("Equality(15)"), WideStringLiteralUpper != WideView);
+	TestFalse(TEXT("Equality(16)"), WideStringLiteralShort == WideView);
+	TestTrue(TEXT("Equality(17)"), WideStringLiteralShort != WideView);
+	TestFalse(TEXT("Equality(18)"), WideStringLiteralLonger == WideView);
+	TestTrue(TEXT("Equality(19)"), WideStringLiteralLonger != WideView);
+
+	FString WideStringSrc = WideStringLiteralSrc;
+	FString WideStringLower = WideStringLiteralLower;
+	FString WideStringUpper = WideStringLiteralUpper;
+	FString WideStringShort = WideStringLiteralShort;
+	FString WideStringLonger = WideStringLiteralLonger;
+
+	TestTrue(TEXT("Equality(20)"), WideView == WideStringSrc);
+	TestFalse(TEXT("Equality(21)"), WideView != WideStringSrc);
+	TestTrue(TEXT("Equality(22)"), WideView == WideStringLower);
+	TestFalse(TEXT("Equality(23)"), WideView != WideStringLower);
+	TestTrue(TEXT("Equality(24)"), WideView == WideStringUpper);
+	TestFalse(TEXT("Equality(25)"), WideView != WideStringUpper);
+	TestFalse(TEXT("Equality(26)"), WideView == WideStringShort);
+	TestTrue(TEXT("Equality(27)"), WideView != WideStringShort);
+	TestFalse(TEXT("Equality(28)"), WideView == WideStringLonger);
+	TestTrue(TEXT("Equality(29)"), WideView != WideStringLonger);
+
+	TestTrue(TEXT("Equality(30)"), WideStringSrc == WideView);
+	TestFalse(TEXT("Equality(31)"), WideStringSrc != WideView);
+	TestTrue(TEXT("Equality(32)"), WideStringLower == WideView);
+	TestFalse(TEXT("Equality(33)"), WideStringLower != WideView);
+	TestTrue(TEXT("Equality(34)"), WideStringUpper == WideView);
+	TestFalse(TEXT("Equality(35)"), WideStringUpper != WideView);
+	TestFalse(TEXT("Equality(36)"), WideStringShort == WideView);
+	TestTrue(TEXT("Equality(37)"), WideStringShort != WideView);
+	TestFalse(TEXT("Equality(38)"), WideStringLonger == WideView);
+	TestTrue(TEXT("Equality(39)"), WideStringLonger != WideView);
+
+	FStringView IdenticalView(WideStringLiteralSrc);
+
+	TestTrue(TEXT("Equality(40a)"), WideView == IdenticalView);
+	TestFalse(TEXT("Equality(40b)"), WideView != IdenticalView);
+	TestTrue(TEXT("Equality(41a)"), IdenticalView == WideView);
+	TestFalse(TEXT("Equality(41b)"), IdenticalView != WideView);
+
+	// Views without null termination
+
+	FStringView ShortViewNoNull = WideView.Left(FStringView(WideStringLiteralShort).Len());
+
+	TestTrue(TEXT("Equality(42)"), ShortViewNoNull == WideStringLiteralShort);
+	TestFalse(TEXT("Equality(43)"), ShortViewNoNull != WideStringLiteralShort);
+	TestTrue(TEXT("Equality(44)"), WideStringLiteralShort == ShortViewNoNull);
+	TestFalse(TEXT("Equality(45)"), WideStringLiteralShort != ShortViewNoNull);
+	TestFalse(TEXT("Equality(46)"), ShortViewNoNull == WideStringLiteralSrc);
+	TestTrue(TEXT("Equality(47)"), ShortViewNoNull != WideStringLiteralSrc);
+	TestFalse(TEXT("Equality(48)"), WideStringLiteralSrc == ShortViewNoNull);
+	TestTrue(TEXT("Equality(49)"), WideStringLiteralSrc != ShortViewNoNull);
+
+	TestTrue(TEXT("Equality(50)"), ShortViewNoNull == WideStringShort);
+	TestFalse(TEXT("Equality(51)"), ShortViewNoNull != WideStringShort);
+	TestTrue(TEXT("Equality(52)"), WideStringShort == ShortViewNoNull);
+	TestFalse(TEXT("Equality(53)"), WideStringShort != ShortViewNoNull);
+	TestFalse(TEXT("Equality(54)"), ShortViewNoNull == WideStringSrc);
+	TestTrue(TEXT("Equality(55)"), ShortViewNoNull != WideStringSrc);
+	TestFalse(TEXT("Equality(56)"), WideStringSrc == ShortViewNoNull);
+	TestTrue(TEXT("Equality(57)"), WideStringSrc != ShortViewNoNull);
+
+	FStringView WideViewNoNull = FStringView(WideStringLiteralLonger).Left(WideView.Len());
+
+	TestTrue(TEXT("Equality(58)"), WideViewNoNull == WideStringLiteralSrc);
+	TestFalse(TEXT("Equality(59)"), WideViewNoNull != WideStringLiteralSrc);
+	TestTrue(TEXT("Equality(60)"), WideStringLiteralSrc == WideViewNoNull);
+	TestFalse(TEXT("Equality(61)"), WideStringLiteralSrc != WideViewNoNull);
+	TestFalse(TEXT("Equality(62)"), WideViewNoNull == WideStringLiteralLonger);
+	TestTrue(TEXT("Equality(63)"), WideViewNoNull != WideStringLiteralLonger);
+	TestFalse(TEXT("Equality(64)"), WideStringLiteralLonger == WideViewNoNull);
+	TestTrue(TEXT("Equality(65)"), WideStringLiteralLonger != WideViewNoNull);
+
+	TestTrue(TEXT("Equality(66)"), WideViewNoNull == WideStringLiteralSrc);
+	TestFalse(TEXT("Equality(67)"), WideViewNoNull != WideStringLiteralSrc);
+	TestTrue(TEXT("Equality(68)"), WideStringLiteralSrc == WideViewNoNull);
+	TestFalse(TEXT("Equality(69)"), WideStringLiteralSrc != WideViewNoNull);
+	TestFalse(TEXT("Equality(70)"), WideViewNoNull == WideStringLiteralLonger);
+	TestTrue(TEXT("Equality(71)"), WideViewNoNull != WideStringLiteralLonger);
+	TestFalse(TEXT("Equality(72)"), WideStringLiteralLonger == WideViewNoNull);
+	TestTrue(TEXT("Equality(73)"), WideStringLiteralLonger != WideViewNoNull);
+
+	// ANSICHAR / TCHAR
+
+	FAnsiStringView AnsiView(AnsiStringLiteralSrc);
+	FAnsiStringView AnsiViewLower(AnsiStringLiteralLower);
+	FAnsiStringView AnsiViewUpper(AnsiStringLiteralUpper);
+
+	TestTrue(TEXT("Equality(74)"), AnsiView.Equals(WideView));
+	TestTrue(TEXT("Equality(75)"), WideView.Equals(AnsiView));
+	TestFalse(TEXT("Equality(76)"), AnsiViewLower.Equals(WideView, ESearchCase::CaseSensitive));
+	TestTrue(TEXT("Equality(77)"), AnsiViewLower.Equals(WideView, ESearchCase::IgnoreCase));
+	TestFalse(TEXT("Equality(78)"), WideView.Equals(AnsiViewLower, ESearchCase::CaseSensitive));
+	TestTrue(TEXT("Equality(79)"), WideView.Equals(AnsiViewLower, ESearchCase::IgnoreCase));
+	TestFalse(TEXT("Equality(80)"), AnsiViewUpper.Equals(WideView, ESearchCase::CaseSensitive));
+	TestTrue(TEXT("Equality(81)"), AnsiViewUpper.Equals(WideView, ESearchCase::IgnoreCase));
+	TestFalse(TEXT("Equality(82)"), WideView.Equals(AnsiViewUpper, ESearchCase::CaseSensitive));
+	TestTrue(TEXT("Equality(83)"), WideView.Equals(AnsiViewUpper, ESearchCase::IgnoreCase));
+
+	TestTrue(TEXT("Equality(84)"), WideView.Equals(AnsiStringLiteralSrc));
+	TestFalse(TEXT("Equality(85)"), WideView.Equals(AnsiStringLiteralLower, ESearchCase::CaseSensitive));
+	TestTrue(TEXT("Equality(86)"), WideView.Equals(AnsiStringLiteralLower, ESearchCase::IgnoreCase));
+	TestFalse(TEXT("Equality(87)"), WideView.Equals(AnsiStringLiteralUpper, ESearchCase::CaseSensitive));
+	TestTrue(TEXT("Equality(88)"), WideView.Equals(AnsiStringLiteralUpper, ESearchCase::IgnoreCase));
+	TestTrue(TEXT("Equality(89)"), AnsiView.Equals(WideStringLiteralSrc));
+	TestFalse(TEXT("Equality(90)"), AnsiViewLower.Equals(WideStringLiteralSrc, ESearchCase::CaseSensitive));
+	TestTrue(TEXT("Equality(91)"), AnsiViewLower.Equals(WideStringLiteralSrc, ESearchCase::IgnoreCase));
+	TestFalse(TEXT("Equality(92)"), AnsiViewUpper.Equals(WideStringLiteralSrc, ESearchCase::CaseSensitive));
+	TestTrue(TEXT("Equality(93)"), AnsiViewUpper.Equals(WideStringLiteralSrc, ESearchCase::IgnoreCase));
+
+	// Test equality of empty strings
+	{
+		const TCHAR* EmptyLiteral = TEXT("");
+		const TCHAR* NonEmptyLiteral = TEXT("ABC");
+		FStringView EmptyView;
+		FStringView NonEmptyView = TEXTVIEW("ABC");
+		TestTrue(TEXT("Equals(94)"), EmptyView.Equals(EmptyLiteral));
+		TestTrue(TEXT("Equals(95)"), !EmptyView.Equals(NonEmptyLiteral));
+		TestTrue(TEXT("Equals(96)"), !NonEmptyView.Equals(EmptyLiteral));
+		TestTrue(TEXT("Equals(97)"), EmptyView.Equals(EmptyView));
+		TestTrue(TEXT("Equals(98)"), !EmptyView.Equals(NonEmptyView));
+		TestTrue(TEXT("Equals(99)"), !NonEmptyView.Equals(EmptyView));
+	}
+
+	// Test types convertible to a string view
+	static_assert(TIsSame<bool, decltype(FAnsiStringView().Equals(FString()))>::Value, "Error with Equals");
+	static_assert(TIsSame<bool, decltype(FWideStringView().Equals(FString()))>::Value, "Error with Equals");
+	static_assert(TIsSame<bool, decltype(FAnsiStringView().Equals(TAnsiStringBuilder<16>()))>::Value, "Error with Equals");
+	static_assert(TIsSame<bool, decltype(FAnsiStringView().Equals(TWideStringBuilder<16>()))>::Value, "Error with Equals");
+	static_assert(TIsSame<bool, decltype(FWideStringView().Equals(TAnsiStringBuilder<16>()))>::Value, "Error with Equals");
+	static_assert(TIsSame<bool, decltype(FWideStringView().Equals(TWideStringBuilder<16>()))>::Value, "Error with Equals");
+}
+
+TEST_CASE("Core::Containers::FStringView::Comparison Case Sensitive", "[Core][Containers][Smoke]")
+{
+	// Basic comparisons involving case
+	{
+		const ANSICHAR* AnsiStringLiteralSrc = "String To Test!";
+		const TCHAR* WideStringLiteralSrc = TEXT("String To Test!");
+		const TCHAR* WideStringLiteralLower = TEXT("string to test!");
+		const TCHAR* WideStringLiteralUpper = TEXT("STRING TO TEST!");
+
+		FStringView WideView(WideStringLiteralSrc);
+
+		TestTrue(TEXT("ComparisonCaseSensitive(0)"), WideView.Compare(WideStringLiteralSrc, ESearchCase::CaseSensitive) == 0);
+		TestFalse(TEXT("ComparisonCaseSensitive(1)"), WideView.Compare(WideStringLiteralLower, ESearchCase::CaseSensitive) > 0);
+		TestFalse(TEXT("ComparisonCaseSensitive(2)"), WideView.Compare(WideStringLiteralUpper, ESearchCase::CaseSensitive) < 0);
+
+		FStringView EmptyView(TEXT(""));
+		TestTrue(TEXT("ComparisonCaseSensitive(3)"), WideView.Compare(EmptyView, ESearchCase::CaseSensitive) > 0);
+
+		FStringView IdenticalView(WideStringLiteralSrc);
+		TestTrue(TEXT("ComparisonCaseSensitive(4)"), WideView.Compare(IdenticalView, ESearchCase::CaseSensitive) == 0);
+
+		FAnsiStringView AnsiView(AnsiStringLiteralSrc);
+		TestTrue(TEXT("ComparisonCaseSensitive(5)"), WideView.Compare(AnsiView, ESearchCase::CaseSensitive) == 0);
+		TestTrue(TEXT("ComparisonCaseSensitive(6)"), WideView.Compare(AnsiStringLiteralSrc, ESearchCase::CaseSensitive) == 0);
+	}
+
+	// Test comparisons of different lengths
+	{
+		const ANSICHAR* AnsiStringLiteralUpper = "ABCDEF";
+		const TCHAR* WideStringLiteralUpper = TEXT("ABCDEF");
+		const TCHAR* WideStringLiteralLower = TEXT("abcdef");
+		const TCHAR* WideStringLiteralLowerShort = TEXT("abc");
+
+		const ANSICHAR* AnsiStringLiteralUpperFirst = "ABCdef";
+		const TCHAR* WideStringLiteralUpperFirst = TEXT("ABCdef");
+		const TCHAR* WideStringLiteralLowerFirst = TEXT("abcDEF");
+
+		FStringView ViewLongUpper(WideStringLiteralUpper);
+		FStringView ViewLongLower(WideStringLiteralLower);
+
+		// Note that the characters after these views are in a different case, this will help catch over read issues
+		FStringView ViewShortUpper(WideStringLiteralUpperFirst, 3);
+		FStringView ViewShortLower(WideStringLiteralLowerFirst, 3);
+
+		// Same length, different cases
+		TestTrue(TEXT("ComparisonCaseSensitive(7)"), ViewLongUpper.Compare(ViewLongLower, ESearchCase::CaseSensitive) < 0);
+		TestTrue(TEXT("ComparisonCaseSensitive(8)"), ViewLongLower.Compare(ViewLongUpper, ESearchCase::CaseSensitive) > 0);
+		TestTrue(TEXT("ComparisonCaseInsensitive(9)"), ViewLongLower.Compare(AnsiStringLiteralUpper, ESearchCase::CaseSensitive) > 0);
+		TestTrue(TEXT("ComparisonCaseInsensitive(10)"), ViewShortUpper.Compare(WideStringLiteralLowerShort, ESearchCase::CaseSensitive) < 0);
+
+		// Same case, different lengths
+		TestTrue(TEXT("ComparisonCaseSensitive(11)"), ViewLongUpper.Compare(ViewShortUpper, ESearchCase::CaseSensitive) > 0);
+		TestTrue(TEXT("ComparisonCaseSensitive(12)"), ViewShortUpper.Compare(ViewLongUpper, ESearchCase::CaseSensitive) < 0);
+		TestTrue(TEXT("ComparisonCaseInsensitive(13)"), ViewShortUpper.Compare(AnsiStringLiteralUpper, ESearchCase::CaseSensitive) < 0);
+		TestTrue(TEXT("ComparisonCaseInsensitive(14)"), ViewLongLower.Compare(WideStringLiteralLowerShort, ESearchCase::CaseSensitive) > 0);
+
+		// Different length, different cases
+		TestTrue(TEXT("ComparisonCaseSensitive(15)"), ViewLongUpper.Compare(ViewShortLower, ESearchCase::CaseSensitive) < 0);
+		TestTrue(TEXT("ComparisonCaseSensitive(16)"), ViewShortLower.Compare(ViewLongUpper, ESearchCase::CaseSensitive) > 0);
+		TestTrue(TEXT("ComparisonCaseInsensitive(17)"), ViewShortLower.Compare(AnsiStringLiteralUpper, ESearchCase::CaseSensitive) > 0);
+		TestTrue(TEXT("ComparisonCaseInsensitive(18)"), ViewLongUpper.Compare(WideStringLiteralLowerShort, ESearchCase::CaseSensitive) < 0);
+	}
+
+	// Test comparisons of empty strings
+	{
+		const TCHAR* EmptyLiteral = TEXT("");
+		const TCHAR* NonEmptyLiteral = TEXT("ABC");
+		FStringView EmptyView;
+		FStringView NonEmptyView = TEXTVIEW("ABC");
+		TestTrue(TEXT("ComparisonEmpty(19)"), EmptyView.Compare(EmptyLiteral) == 0);
+		TestTrue(TEXT("ComparisonEmpty(20)"), EmptyView.Compare(NonEmptyLiteral) < 0);
+		TestTrue(TEXT("ComparisonEmpty(21)"), NonEmptyView.Compare(EmptyLiteral) > 0);
+		TestTrue(TEXT("ComparisonEmpty(22)"), EmptyView.Compare(EmptyView) == 0);
+		TestTrue(TEXT("ComparisonEmpty(23)"), EmptyView.Compare(NonEmptyView) < 0);
+		TestTrue(TEXT("ComparisonEmpty(24)"), NonEmptyView.Compare(EmptyView) > 0);
+	}
+
+	// Test types convertible to a string view
+	static_assert(TIsSame<int32, decltype(FAnsiStringView().Compare(FString()))>::Value, "Error with Compare");
+	static_assert(TIsSame<int32, decltype(FWideStringView().Compare(FString()))>::Value, "Error with Compare");
+	static_assert(TIsSame<int32, decltype(FAnsiStringView().Compare(TAnsiStringBuilder<16>()))>::Value, "Error with Compare");
+	static_assert(TIsSame<int32, decltype(FAnsiStringView().Compare(TWideStringBuilder<16>()))>::Value, "Error with Compare");
+	static_assert(TIsSame<int32, decltype(FWideStringView().Compare(TAnsiStringBuilder<16>()))>::Value, "Error with Compare");
+	static_assert(TIsSame<int32, decltype(FWideStringView().Compare(TWideStringBuilder<16>()))>::Value, "Error with Compare");
+}
+
+TEST_CASE("Core::Containers::FStringView::Comparison Case Insensitive", "[Core][Containers][Smoke]")
+{
+	// Basic comparisons involving case
+	{
+		const ANSICHAR* AnsiStringLiteralSrc = "String To Test!";
+		const TCHAR* WideStringLiteralSrc = TEXT("String To Test!");
+		const TCHAR* WideStringLiteralLower = TEXT("string to test!");
+		const TCHAR* WideStringLiteralUpper = TEXT("STRING TO TEST!");
+
+		FStringView WideView(WideStringLiteralSrc);
+
+		TestTrue(TEXT("ComparisonCaseInsensitive(0)"), WideView.Compare(WideStringLiteralSrc, ESearchCase::IgnoreCase) == 0);
+		TestTrue(TEXT("ComparisonCaseInsensitive(1)"), WideView.Compare(WideStringLiteralLower, ESearchCase::IgnoreCase) == 0);
+		TestTrue(TEXT("ComparisonCaseInsensitive(2)"), WideView.Compare(WideStringLiteralUpper, ESearchCase::IgnoreCase) == 0);
+
+		FStringView EmptyView(TEXT(""));
+		TestTrue(TEXT("ComparisonCaseInsensitive(3)"), WideView.Compare(EmptyView, ESearchCase::IgnoreCase) > 0);
+
+		FStringView IdenticalView(WideStringLiteralSrc);
+		TestTrue(TEXT("ComparisonCaseInsensitive(4)"), WideView.Compare(IdenticalView, ESearchCase::IgnoreCase) == 0);
+
+		FAnsiStringView AnsiView(AnsiStringLiteralSrc);
+		TestTrue(TEXT("ComparisonCaseSensitive(5)"), WideView.Compare(AnsiView, ESearchCase::IgnoreCase) == 0);
+		TestTrue(TEXT("ComparisonCaseSensitive(6)"), WideView.Compare(AnsiStringLiteralSrc, ESearchCase::IgnoreCase) == 0);
+	}
+
+	// Test comparisons of different lengths
+	{
+		const ANSICHAR* AnsiStringLiteralUpper = "ABCDEF";
+		const TCHAR* WideStringLiteralUpper = TEXT("ABCDEF");
+		const TCHAR* WideStringLiteralLower = TEXT("abcdef");
+		const TCHAR* WideStringLiteralLowerShort = TEXT("abc");
+
+		const ANSICHAR* AnsiStringLiteralUpperFirst = "ABCdef";
+		const TCHAR* WideStringLiteralUpperFirst = TEXT("ABCdef");
+		const TCHAR* WideStringLiteralLowerFirst = TEXT("abcDEF");
+
+		FStringView ViewLongUpper(WideStringLiteralUpper);
+		FStringView ViewLongLower(WideStringLiteralLower);
+
+		// Note that the characters after these views are in a different case, this will help catch over read issues
+		FStringView ViewShortUpper(WideStringLiteralUpperFirst, 3);
+		FStringView ViewShortLower(WideStringLiteralLowerFirst, 3);
+
+		// Same length, different cases
+		TestTrue(TEXT("ComparisonCaseInsensitive(7)"), ViewLongUpper.Compare(ViewLongLower, ESearchCase::IgnoreCase) == 0);
+		TestTrue(TEXT("ComparisonCaseInsensitive(8)"), ViewLongLower.Compare(ViewLongUpper, ESearchCase::IgnoreCase) == 0);
+		TestTrue(TEXT("ComparisonCaseInsensitive(9)"), ViewLongLower.Compare(AnsiStringLiteralUpper, ESearchCase::IgnoreCase) == 0);
+		TestTrue(TEXT("ComparisonCaseInsensitive(10)"), ViewShortUpper.Compare(WideStringLiteralLowerShort, ESearchCase::IgnoreCase) == 0);
+
+		// Same case, different lengths
+		TestTrue(TEXT("ComparisonCaseInsensitive(11)"), ViewLongUpper.Compare(ViewShortUpper, ESearchCase::IgnoreCase) > 0);
+		TestTrue(TEXT("ComparisonCaseInsensitive(12)"), ViewShortUpper.Compare(ViewLongUpper, ESearchCase::IgnoreCase) < 0);
+		TestTrue(TEXT("ComparisonCaseInsensitive(13)"), ViewShortUpper.Compare(AnsiStringLiteralUpper, ESearchCase::IgnoreCase) < 0);
+		TestTrue(TEXT("ComparisonCaseInsensitive(14)"), ViewLongLower.Compare(WideStringLiteralLowerShort, ESearchCase::IgnoreCase) > 0);
+
+		// Different length, different cases
+		TestTrue(TEXT("ComparisonCaseInsensitive(15)"), ViewLongUpper.Compare(ViewShortLower, ESearchCase::IgnoreCase) > 0);
+		TestTrue(TEXT("ComparisonCaseInsensitive(16)"), ViewShortLower.Compare(ViewLongUpper, ESearchCase::IgnoreCase) < 0);
+		TestTrue(TEXT("ComparisonCaseInsensitive(17)"), ViewShortLower.Compare(AnsiStringLiteralUpper, ESearchCase::IgnoreCase) < 0);
+		TestTrue(TEXT("ComparisonCaseInsensitive(18)"), ViewLongUpper.Compare(WideStringLiteralLowerShort, ESearchCase::IgnoreCase) > 0);
+	}
+}
+
+TEST_CASE("Core::Containers::FStringView::ArrayAccessor", "[Core][Containers][Smoke]")
+{
+	const TCHAR* SrcString = TEXT("String To Test");
+	FStringView View(SrcString);
+
+	for (int32 i = 0; i < View.Len(); ++i)
+	{
+		TestEqual(TEXT("the character accessed"), View[i], SrcString[i]);
+	}
+}
+
+TEST_CASE("Core::Containers::FStringView::ArrayModifiers", "[Core][Containers][Smoke]")
+{
+	const TCHAR* FullText = TEXT("PrefixSuffix");
+	const TCHAR* Prefix = TEXT("Prefix");
+	const TCHAR* Suffix = TEXT("Suffix");
+
+	// Remove prefix
+	{
+		FStringView View(FullText);
+		View.RemovePrefix(FCStringWide::Strlen(Prefix));
+
+		TestEqual(TEXT("View length"), View.Len(), FCStringWide::Strlen(Suffix));
+		TestEqual(TEXT("The result of Strncmp"), FCStringWide::Strncmp(View.GetData(), Suffix, View.Len()), 0);
+	}
+
+	// Remove suffix
+	{
+		FStringView View(FullText);
+		View.RemoveSuffix(FCStringWide::Strlen(Suffix));
+
+		TestEqual(TEXT("View length"), View.Len(), FCStringWide::Strlen(Prefix));
+		TestEqual(TEXT("The result of Strncmp"), FCStringWide::Strncmp(View.GetData(), Prefix, View.Len()), 0);
+	}
+}
+
+TEST_CASE("Core::Containers::FStringView::StartsWith", "[Core][Containers][Smoke]")
+{
+	// Test an empty view
+	{
+		FStringView View;
+		TestTrue(TEXT(" View.StartsWith"), View.StartsWith(TEXT("")));
+		TestFalse(TEXT(" View.StartsWith"), View.StartsWith(TEXT("Text")));
+		TestFalse(TEXT(" View.StartsWith"), View.StartsWith(TEXT('A')));
+	}
+
+	// Test a valid view with the correct text
+	{
+		FStringView View(TEXT("String to test"));
+		TestTrue(TEXT(" View.StartsWith"), View.StartsWith(TEXT("String")));
+		TestTrue(TEXT(" View.StartsWith"), View.StartsWith(TEXT('S')));
+	}
+
+	// Test a valid view with incorrect text
+	{
+		FStringView View(TEXT("String to test"));
+		TestFalse(TEXT(" View.StartsWith"), View.StartsWith(TEXT("test")));
+		TestFalse(TEXT(" View.StartsWith"), View.StartsWith(TEXT('t')));
+	}
+
+	// Test a valid view with the correct text but with different case
+	{
+		FStringView View(TEXT("String to test"));
+		TestTrue(TEXT(" View.StartsWith"), View.StartsWith(TEXT("sTrInG")));
+
+		// Searching by char is case sensitive to keep compatibility with FString
+		TestFalse(TEXT(" View.StartsWith"), View.StartsWith(TEXT('s')));
+	}
+}
+
+TEST_CASE("Core::Containers::FStringView::EndsWith", "[Core][Containers][Smoke]")
+{
+	// Test an empty view
+	{
+		FStringView View;
+		TestTrue(TEXT(" View.EndsWith"), View.EndsWith(TEXT("")));
+		TestFalse(TEXT(" View.EndsWith"), View.EndsWith(TEXT("Text")));
+		TestFalse(TEXT(" View.EndsWith"), View.EndsWith(TEXT('A')));
+	}
+
+	// Test a valid view with the correct text
+	{
+		FStringView View(TEXT("String to test"));
+		TestTrue(TEXT(" View.EndsWith"), View.EndsWith(TEXT("test")));
+		TestTrue(TEXT(" View.EndsWith"), View.EndsWith(TEXT('t')));
+	}
+
+	// Test a valid view with incorrect text
+	{
+		FStringView View(TEXT("String to test"));
+		TestFalse(TEXT(" View.EndsWith"), View.EndsWith(TEXT("String")));
+		TestFalse(TEXT(" View.EndsWith"), View.EndsWith(TEXT('S')));
+	}
+
+	// Test a valid view with the correct text but with different case
+	{
+		FStringView View(TEXT("String to test"));
+		TestTrue(TEXT(" View.EndsWith"), View.EndsWith(TEXT("TeST")));
+
+		// Searching by char is case sensitive to keep compatibility with FString
+		TestFalse(TEXT(" View.EndsWith"), View.EndsWith(TEXT('T'))); 
+	}
+}
+
+TEST_CASE("Core::Containers::FStringView::SubStr", "[Core][Containers][Smoke]")
+{
+	{
+		FStringView EmptyView;
+		FStringView EmptyResult = EmptyView.SubStr(0, 10);
+		TestTrue(TEXT("FStringView::SubStr(0)"), EmptyResult.IsEmpty());
+
+		// The following line is commented out as it would fail an assert and currently we cannot test for this in unit tests 
+		// FStringView OutofBoundsResult = EmptyView.SubStr(1000, 10000); 
+		FStringView OutofBoundsResult = EmptyView.SubStr(0, 10000);
+		TestTrue(TEXT("FStringView::SubStr(1)"), OutofBoundsResult.IsEmpty());
+	}
+
+	{
+		FStringView View(TEXT("A test string|"));
+		FStringView Word0 = View.SubStr(0, 1);
+		FStringView Word1 = View.SubStr(2, 4);
+		FStringView Word2 = View.SubStr(7, 6);
+		FStringView NullTerminatorResult = View.SubStr(14, 1024);	// We can create a substr that starts at the end of the 
+																	// string since the null terminator is still valid
+		FStringView OutofBoundsResult = View.SubStr(0, 1024);
+
+		TestTrue(TEXT("FStringView::SubStr(2)"), FCString::Strncmp(Word0.GetData(), TEXT("A"), Word0.Len()) == 0);
+		TestTrue(TEXT("FStringView::SubStr(3)"), FCString::Strncmp(Word1.GetData(), TEXT("test"), Word1.Len()) == 0);
+		TestTrue(TEXT("FStringView::SubStr(4)"), FCString::Strncmp(Word2.GetData(), TEXT("string"), Word2.Len()) == 0);
+		TestTrue(TEXT("FStringView::SubStr(5)"), NullTerminatorResult.IsEmpty());
+		TestTrue(TEXT("FStringView::SubStr(6)"), View == OutofBoundsResult);
+	}
+}
+
+TEST_CASE("Core::Containers::FStringView::Left", "[Core][Containers][Smoke]")
+{
+	{
+		FStringView EmptyView;
+		FStringView EmptyResult = EmptyView.Left(0);
+		TestTrue(TEXT("FStringView::Left"), EmptyResult.IsEmpty());
+
+		FStringView OutofBoundsResult = EmptyView.Left(1024);
+		TestTrue(TEXT("FStringView::Left"), OutofBoundsResult.IsEmpty());
+	}
+	
+	{
+		FStringView View(TEXT("A test string padded"), 13); // "A test string" without null termination
+		FStringView Result = View.Left(8);
+
+		TestTrue(TEXT("FStringView::Left"), FCString::Strncmp(Result.GetData(), TEXT("A test s"), Result.Len()) == 0);
+
+		FStringView OutofBoundsResult = View.Left(1024);
+		TestTrue(TEXT("FStringView::Left"), FCString::Strncmp(OutofBoundsResult.GetData(), TEXT("A test string"), OutofBoundsResult.Len()) == 0);
+	}
+}
+
+TEST_CASE("Core::Containers::FStringView::LeftChop", "[Core][Containers][Smoke]")
+{
+	{
+		FStringView EmptyView;
+		FStringView EmptyResult = EmptyView.LeftChop(0);
+		TestTrue(TEXT("FStringView::LeftChop"), EmptyResult.IsEmpty());
+
+		FStringView OutofBoundsResult = EmptyView.LeftChop(1024);
+		TestTrue(TEXT("FStringView::LeftChop"), OutofBoundsResult.IsEmpty());
+	}
+
+	{
+		FStringView View(TEXT("A test string padded"), 13); // "A test string" without null termination
+		FStringView Result = View.LeftChop(5);
+
+		TestTrue(TEXT("FStringView::LeftChop"), FCString::Strncmp(Result.GetData(), TEXT("A test s"), Result.Len()) == 0);
+
+		FStringView OutofBoundsResult = View.LeftChop(1024);
+		TestTrue(TEXT("FStringView::LeftChop"), FCString::Strncmp(OutofBoundsResult.GetData(), TEXT("A test string"), OutofBoundsResult.Len()) == 0);
+	}
+}
+
+TEST_CASE("Core::Containers::FStringView::Right", "[Core][Containers][Smoke]")
+{
+	{
+		FStringView EmptyView;
+		FStringView EmptyResult = EmptyView.Right(0);
+		TestTrue(TEXT("FStringView::Right"), EmptyResult.IsEmpty());
+
+		FStringView OutofBoundsResult = EmptyView.Right(1024);
+		TestTrue(TEXT("FStringView::Right"), OutofBoundsResult.IsEmpty());
+	}
+
+	{
+		FStringView View(TEXT("A test string padded"), 13); // "A test string" without null termination
+		FStringView Result = View.Right(8);
+
+		TestTrue(TEXT("FStringView::Right"), FCString::Strncmp(Result.GetData(), TEXT("t string"), Result.Len()) == 0);
+
+		FStringView OutofBoundsResult = View.Right(1024);
+		TestTrue(TEXT("FStringView::Right"), FCString::Strncmp(OutofBoundsResult.GetData(), TEXT("A test string"), OutofBoundsResult.Len()) == 0);
+	}
+}
+
+TEST_CASE("Core::Containers::FStringView::RightChop", "[Core][Containers][Smoke]")
+{
+	{
+		FStringView EmptyView;
+		FStringView EmptyResult = EmptyView.RightChop(0);
+		TestTrue(TEXT("FStringView::RightChop"), EmptyResult.IsEmpty());
+
+		FStringView OutofBoundsResult = EmptyView.RightChop(1024);
+		TestTrue(TEXT("FStringView::RightChop"), OutofBoundsResult.IsEmpty());
+	}
+
+	{
+		FStringView View(TEXT("A test string padded"), 13); // "A test string" without null termination
+		FStringView Result = View.RightChop(3);
+
+		TestTrue(TEXT("FStringView::RightChop"), FCString::Strncmp(Result.GetData(), TEXT("est string"), Result.Len()) == 0);
+
+		FStringView OutofBoundsResult = View.RightChop(1024);
+		TestTrue(TEXT("FStringView::RightChop"), FCString::Strncmp(OutofBoundsResult.GetData(), TEXT("A test string"), OutofBoundsResult.Len()) == 0);
+	}
+}
+
+TEST_CASE("Core::Containers::FStringView::Mid", "[Core][Containers][Smoke]")
+{
+	{
+		FStringView EmptyView;
+		FStringView EmptyResult = EmptyView.Mid(0, 10);
+		TestTrue(TEXT("FStringView::Mid(0)"), EmptyResult.IsEmpty());
+
+		// The following line is commented out as it would fail an assert and currently we cannot test for this in unit tests 
+		// FStringView OutofBoundsResult = EmptyView.Mid(1000, 10000); 
+		FStringView OutofBoundsResult = EmptyView.Mid(0, 10000);
+		TestTrue(TEXT("FStringView::Mid(1)"), OutofBoundsResult.IsEmpty());
+	}
+
+	{
+		FStringView View(TEXT("A test string|"));
+		FStringView Word0 = View.Mid(0, 1);
+		FStringView Word1 = View.Mid(2, 4);
+		FStringView Word2 = View.Mid(7, 6);
+		FStringView NullTerminatorResult = View.Mid(14, 1024);	// We can call Mid with a position that starts at the end of the 
+																// string since the null terminator is still valid
+		FStringView OutofBoundsResult = View.Mid(0, 1024);
+
+		TestTrue(TEXT("FStringView::Mid(2)"), FCString::Strncmp(Word0.GetData(), TEXT("A"), Word0.Len()) == 0);
+		TestTrue(TEXT("FStringView::Mid(3)"), FCString::Strncmp(Word1.GetData(), TEXT("test"), Word1.Len()) == 0);
+		TestTrue(TEXT("FStringView::Mid(4)"), FCString::Strncmp(Word2.GetData(), TEXT("string"), Word2.Len()) == 0);
+		TestTrue(TEXT("FStringView::Mid(5)"), NullTerminatorResult.IsEmpty());
+		TestTrue(TEXT("FStringView::Mid(6)"), View == OutofBoundsResult);
+		TestTrue(TEXT("FStringView::Mid(7)"), View.Mid(512, 1024).IsEmpty());
+		TestTrue(TEXT("FStringView::Mid(8)"), View.Mid(4, 0).IsEmpty());
+	}
+}
+
+TEST_CASE("Core::Containers::FStringView::TrimStartAndEnd", "[Core][Containers][Smoke]")
+{
+	TestTrue(TEXT("FStringView::TrimStartAndEnd(\"\")"), TEXTVIEW("").TrimStartAndEnd().IsEmpty());
+	TestTrue(TEXT("FStringView::TrimStartAndEnd(\" \")"), TEXTVIEW(" ").TrimStartAndEnd().IsEmpty());
+	TestTrue(TEXT("FStringView::TrimStartAndEnd(\"  \")"), TEXTVIEW("  ").TrimStartAndEnd().IsEmpty());
+	TestTrue(TEXT("FStringView::TrimStartAndEnd(\" \\t\\r\\n\")"), TEXTVIEW(" \t\r\n").TrimStartAndEnd().IsEmpty());
+
+	TestEqual(TEXT("FStringView::TrimStartAndEnd(\"ABC123\")"), TEXTVIEW("ABC123").TrimStartAndEnd(), TEXTVIEW("ABC123"));
+	TestEqual(TEXT("FStringView::TrimStartAndEnd(\"A \\t\\r\\nB\")"), TEXTVIEW("A \t\r\nB").TrimStartAndEnd(), TEXTVIEW("A \t\r\nB"));
+	TestEqual(TEXT("FStringView::TrimStartAndEnd(\" \\t\\r\\nABC123\\n\\r\\t \")"), TEXTVIEW(" \t\r\nABC123\n\r\t ").TrimStartAndEnd(), TEXTVIEW("ABC123"));
+}
+
+TEST_CASE("Core::Containers::FStringView::TrimStart", "[Core][Containers][Smoke]")
+{
+	TestTrue(TEXT("FStringView::TrimStart(\"\")"), TEXTVIEW("").TrimStart().IsEmpty());
+	TestTrue(TEXT("FStringView::TrimStart(\" \")"), TEXTVIEW(" ").TrimStart().IsEmpty());
+	TestTrue(TEXT("FStringView::TrimStart(\"  \")"), TEXTVIEW("  ").TrimStart().IsEmpty());
+	TestTrue(TEXT("FStringView::TrimStart(\" \\t\\r\\n\")"), TEXTVIEW(" \t\r\n").TrimStart().IsEmpty());
+
+	TestEqual(TEXT("FStringView::TrimStart(\"ABC123\")"), TEXTVIEW("ABC123").TrimStart(), TEXTVIEW("ABC123"));
+	TestEqual(TEXT("FStringView::TrimStart(\"A \\t\\r\\nB\")"), TEXTVIEW("A \t\r\nB").TrimStart(), TEXTVIEW("A \t\r\nB"));
+	TestEqual(TEXT("FStringView::TrimStart(\" \\t\\r\\nABC123\\n\\r\\t \")"), TEXTVIEW(" \t\r\nABC123\n\r\t ").TrimStart(), TEXTVIEW("ABC123\n\r\t "));
+}
+
+TEST_CASE("Core::Containers::FStringView::TrimEnd", "[Core][Containers][Smoke]")
+{
+	TestTrue(TEXT("FStringView::TrimEnd(\"\")"), TEXTVIEW("").TrimEnd().IsEmpty());
+	TestTrue(TEXT("FStringView::TrimEnd(\" \")"), TEXTVIEW(" ").TrimEnd().IsEmpty());
+	TestTrue(TEXT("FStringView::TrimEnd(\"  \")"), TEXTVIEW("  ").TrimEnd().IsEmpty());
+	TestTrue(TEXT("FStringView::TrimEnd(\" \\t\\r\\n\")"), TEXTVIEW(" \t\r\n").TrimEnd().IsEmpty());
+
+	TestEqual(TEXT("FStringView::TrimEnd(\"ABC123\")"), TEXTVIEW("ABC123").TrimEnd(), TEXTVIEW("ABC123"));
+	TestEqual(TEXT("FStringView::TrimEnd(\"A \\t\\r\\nB\")"), TEXTVIEW("A \t\r\nB").TrimEnd(), TEXTVIEW("A \t\r\nB"));
+	TestEqual(TEXT("FStringView::TrimEnd(\" \\t\\r\\nABC123\\n\\r\\t \")"), TEXTVIEW(" \t\r\nABC123\n\r\t ").TrimEnd(), TEXTVIEW(" \t\r\nABC123"));
+}
+
+TEST_CASE("Core::Containers::FStringView::FindChar", "[Core][Containers][Smoke]")
+{
+	FStringView EmptyView;
+	FStringView View = TEXT("aBce Fga");
+
+	{
+		int32 Index = INDEX_NONE;
+		TestFalse(TEXT("FStringView::FindChar-Return(0)"), EmptyView.FindChar(TEXT('a'), Index));
+		TestEqual(TEXT("FStringView::FindChar-Index(0)"), Index, INDEX_NONE);
+	}
+
+	{
+		int32 Index = INDEX_NONE;
+		TestTrue(TEXT("FStringView::FindChar-Return(1)"), View.FindChar(TEXT('a'), Index));
+		TestEqual(TEXT("FStringView::FindChar-Index(1)"), Index, 0);
+	}
+
+	{
+		int32 Index = INDEX_NONE;
+		TestTrue(TEXT("FStringView::FindChar-Return(2)"), View.FindChar(TEXT('F'), Index));
+		TestEqual(TEXT("FStringView::FindChar-Index(2)"), Index, 5);
+	}
+
+	{
+		int32 Index = INDEX_NONE;
+		TestFalse(TEXT("FStringView::FindChar-Return(3)"), View.FindChar(TEXT('A'), Index));
+		TestEqual(TEXT("FStringView::FindChar-Index(3)"), Index, INDEX_NONE);
+	}
+
+	{
+		int32 Index = INDEX_NONE;
+		TestFalse(TEXT("FStringView::FindChar-Return(4)"), View.FindChar(TEXT('d'), Index));
+		TestEqual(TEXT("FStringView::FindChar-Index(4)"), Index, INDEX_NONE);
+	}
+
+	{
+		int32 Index = INDEX_NONE;
+		TestTrue(TEXT("FStringView::FindChar-Return(5)"), View.FindChar(TEXT(' '), Index));
+		TestEqual(TEXT("FStringView::FindChar-Index(5)"), Index, 4);
+	}
+}
+
+TEST_CASE("Core::Containers::FStringView::FindLastChar", "[Core][Containers][Smoke]")
+{
+	FStringView EmptyView;
+	FStringView View = TEXT("aBce Fga");
+
+	{
+		int32 Index = INDEX_NONE;
+		TestFalse(TEXT("FStringView::FindChar-Return(0)"), EmptyView.FindLastChar(TEXT('a'), Index));
+		TestEqual(TEXT("FStringView::FindChar-Index(0)"), Index, INDEX_NONE);
+	}
+
+	{
+		int32 Index = INDEX_NONE;
+		TestTrue(TEXT("FStringView::FindChar-Return(1)"), View.FindLastChar(TEXT('a'), Index));
+		TestEqual(TEXT("FStringView::FindChar-Index(1)"), Index, 7);
+	}
+
+	{
+		int32 Index = INDEX_NONE;
+		TestTrue(TEXT("FStringView::FindChar-Return(2)"), View.FindLastChar(TEXT('B'), Index));
+		TestEqual(TEXT("FStringView::FindChar-Index(2)"), Index, 1);
+	}
+
+	{
+		int32 Index = INDEX_NONE;
+		TestFalse(TEXT("FStringView::FindChar-Return(3)"), View.FindLastChar(TEXT('A'), Index));
+		TestEqual(TEXT("FStringView::FindChar-Index(3)"), Index, INDEX_NONE);
+	}
+
+	{
+		int32 Index = INDEX_NONE;
+		TestFalse(TEXT("FStringView::FindChar-Return(4)"), View.FindLastChar(TEXT('d'), Index));
+		TestEqual(TEXT("FStringView::FindChar-Index(4)"), Index, INDEX_NONE);
+	}
+
+	{
+		int32 Index = INDEX_NONE;
+		TestTrue(TEXT("FStringView::FindChar-Return(5)"), View.FindLastChar(TEXT(' '), Index));
+		TestEqual(TEXT("FStringView::FindChar-Index(5)"), Index, 4);
+	}
+}
+
+void TestSlicing(const FString& Str)
+{
+	const FStringView View = Str;
+	const int32       Len  = Str.Len();
+
+	// Left tests
+	{
+		// Try all lengths of the string, including +/- 5 either side
+		for (int32 Index = -5; Index != Len + 5; ++Index)
+		{
+			FString     Substring = Str .Left(Index);
+			FStringView Subview   = View.Left(Index);
+
+			TestEqual(FString::Printf(TEXT("FStringView(\"%s\")::Left(%d)"), *Str, Index), FString(Subview), Substring);
+		}
+	}
+
+	// LeftChop tests
+	{
+		// Try all lengths of the string, including +/- 5 either side
+		for (int32 Index = -5; Index != Len + 5; ++Index)
+		{
+			FString     Substring = Str .LeftChop(Index);
+			FStringView Subview   = View.LeftChop(Index);
+
+			TestEqual(FString::Printf(TEXT("FStringView(\"%s\")::LeftChop(%d)"), *Str, Index), FString(Subview), Substring);
+		}
+	}
+
+	// Right tests
+	{
+		// Try all lengths of the string, including +/- 5 either side
+		for (int32 Index = -5; Index != Len + 5; ++Index)
+		{
+			FString     Substring = Str .Right(Index);
+			FStringView Subview   = View.Right(Index);
+
+			TestEqual(FString::Printf(TEXT("FStringView(\"%s\")::Right(%d)"), *Str, Index), FString(Subview), Substring);
+		}
+	}
+
+	// RightChop tests
+	{
+		// Try all lengths of the string, including +/- 5 either side
+		for (int32 Index = -5; Index != Len + 5; ++Index)
+		{
+			FString     Substring = Str .RightChop(Index);
+			FStringView Subview   = View.RightChop(Index);
+
+			TestEqual(FString::Printf(TEXT("FStringView(\"%s\").RightChop(%d)"), *Str, Index), FString(Subview), Substring);
+		}
+	}
+
+	// Mid tests
+	{
+		// Try all lengths of the string, including +/- 5 either side
+		for (int32 Index = -5; Index != Len + 5; ++Index)
+		{
+			for (int32 Count = -5; Count != Len + 5; ++Count)
+			{
+				FString     Substring = Str .Mid(Index, Count);
+				FStringView Subview   = View.Mid(Index, Count);
+
+				TestEqual(FString::Printf(TEXT("FStringView(\"%s\")::Mid(%d, %d)"), *Str, Index, Count), FString(Subview), Substring);
+			}
+		}
+
+		// Test near limits of int32
+		for (int32 IndexOffset = 0; IndexOffset != Len + 5; ++IndexOffset)
+		{
+			for (int32 CountOffset = 0; CountOffset != Len + 5; ++CountOffset)
+			{
+				int32 Index = MIN_int32 + IndexOffset;
+				int32 Count = MAX_int32 - CountOffset;
+
+				FString     Substring = Str .Mid(Index, Count);
+				FStringView Subview   = View.Mid(Index, Count);
+
+				TestEqual(FString::Printf(TEXT("FStringView(\"%s\")::Mid(%d, %d)"), *Str, Index, Count), FString(Subview), Substring);
+			}
+		}
+	}
+}
+
+TEST_CASE("Core::Containers::FStringView::Slice", "[Core][Containers][Smoke]")
+{
+	// We assume that FString has already passed its tests, and we just want views to be consistent with it
+
+	// Test an aribtrary string
+	TestSlicing(TEXT("Test string"));
+
+	// Test an empty string
+	TestSlicing(FString());
+
+	// Test an null-terminator-only empty string
+	FString TerminatorOnly;
+	TerminatorOnly.GetCharArray().Add(TEXT('\0'));
+	TestSlicing(TerminatorOnly);
+}
diff --git a/Engine/Source/Runtime/Core/Tests/HAL/FileManagerGenericTest.cpp b/Engine/Source/Runtime/Core/Tests/HAL/FileManagerGenericTest.cpp
new file mode 100644
--- /dev/null
+++ b/Engine/Source/Runtime/Core/Tests/HAL/FileManagerGenericTest.cpp
@@ -0,0 +1,297 @@
+// Copyright Epic Games, Inc. All Rights Reserved.
+
+#include "HAL/FileManagerGeneric.h"
+#include "Containers/UnrealString.h"
+#include "HAL/FileManager.h"
+#include "HAL/PlatformFileManager.h"
+#include "Misc/Paths.h"
+#include "Misc/StringBuilder.h"
+#include "TestHarness.h"
+
+class FArchiveFileReaderGenericTest
+{
+public:
+	~FArchiveFileReaderGenericTest();
+	void TestInternalPrecache();
+
+private:
+	bool CreateTestFile();
+	void TestBytesValid(const TCHAR* What, uint8* Data, int64 NumBytes, int64 FileOffset);
+	int32 GetExpectedValue(int64 IntOffset);
+	void SetPosAndBuffer(TUniquePtr<FArchiveFileReaderGeneric>& Reader, int64 Pos, int64 BufferBase, int64 BufferSize);
+
+	FString TestFileName;
+	TUniquePtr<IFileHandle> ReadHandle;
+	int64 FileSize;
+};
+
+FArchiveFileReaderGenericTest::~FArchiveFileReaderGenericTest()
+{
+	ReadHandle.Reset();
+	if (!TestFileName.IsEmpty())
+	{
+		IFileManager& FileManager = IFileManager::Get();
+		FileManager.Delete(*TestFileName, false /* bRequireExists */, true /* bEvenReadOnly */, true /* Quiet */);
+	}
+	TestFileName.Empty();
+}
+
+void FArchiveFileReaderGenericTest::TestInternalPrecache()
+{
+	if (!CreateTestFile())
+	{
+		return;
+	}
+
+	uint32 BufferSize = 1024;
+	TArray<uint8> UnusedBytes;
+	UnusedBytes.SetNumUninitialized(BufferSize);
+	TUniquePtr<FArchiveFileReaderGeneric> Reader(new FArchiveFileReaderGeneric(ReadHandle.Release(), *TestFileName, FileSize, BufferSize));
+	TestEqual(TEXT("Initial Pos should be 0"), Reader->Pos, 0LL);
+	TestEqual(TEXT("Size should be what was passed in"), Reader->Size, FileSize);
+	TestEqual(TEXT("BufferSize should be what was passed in"), Reader->BufferSize, BufferSize);
+
+	// Vanilla InternalPrecache at start of file
+	bool Result = Reader->InternalPrecache(0, BufferSize);
+	TestEqual(TEXT("Vanilla0 - InternalPrecache should succeed"), Result, true);
+	TestEqual(TEXT("Vanilla0 - InternalPrecache should not alter Pos"), Reader->Pos, 0LL);
+	TestEqual(TEXT("Vanilla0 - BufferBase should have be set to PrecacheOffset aka Pos"), Reader->BufferBase, 0LL);
+	TestEqual(TEXT("Vanilla0 - Bytes precached should be set to BufferSize unless it runs out of room"), Reader->BufferArray.Num(), BufferSize);
+	TestBytesValid(TEXT("Vanilla0 - BufferArray should be expected bytes"), Reader->BufferArray.GetData(), Reader->BufferArray.Num(), Reader->BufferBase);
+
+	// InternalPrecache at PrecacheOffset != Pos is ignored, and returns false if the first byte at PrecacheOffset is not cached
+	Result = Reader->InternalPrecache(BufferSize  + BufferSize, BufferSize);
+	TestEqual(TEXT("PrecacheNotAtPos - Should fail"), Result, false);
+	TestEqual(TEXT("PrecacheNotAtPos - BufferBase should not be altered"), Reader->BufferBase, 0LL);
+	TestEqual(TEXT("PrecacheNotAtPos - Buffe should not be altered"), Reader->BufferArray.Num(), BufferSize);
+
+	// InternalPrecache at PrecacheOffset != Pos is ignored, but returns true if the first byte at PrecacheOffset is cached
+	Result = Reader->InternalPrecache(BufferSize - 1, BufferSize);
+	TestEqual(TEXT("PrecacheNotAtPos - Should succeed since the first byte at PrecacheOffset is buffered"), Result, true);
+	TestEqual(TEXT("PrecacheNotAtPos - BufferBase should not be altered"), Reader->BufferBase, 0LL);
+	TestEqual(TEXT("PrecacheNotAtPos - Buffe should not be altered"), Reader->BufferArray.Num(), BufferSize);
+
+	// InternalPrecache partway through the buffer should allocate the full buffer/not allocate anything, depending on bPrecacheAsSoonAsPossible
+	int64 PosStart = BufferSize / 2;
+	SetPosAndBuffer(Reader, PosStart, 0, BufferSize);
+	Result = Reader->InternalPrecache(PosStart, BufferSize);
+	TestEqual(TEXT("PrecachePartwayPosMoreThanBuffer - Should succeed"), Result, true);
+	if (FArchiveFileReaderGeneric::bPrecacheAsSoonAsPossible)
+	{
+		TestEqual(TEXT("PrecachePartwayPosMoreThanBuffer - BufferBase should be updated"), Reader->BufferBase, PosStart);
+	}
+	else
+	{
+		TestEqual(TEXT("PrecachePartwayPosMoreThanBuffer - BufferBase should not be updated"), Reader->BufferBase, 0LL);
+	}
+	TestEqual(TEXT("PrecachePartwayPosMoreThanBuffer - BufferCount should be set to BufferSize"), Reader->BufferArray.Num(), BufferSize);
+	TestBytesValid(TEXT("PrecachePartwayPosMoreThanBuffer - BufferBytes should match BufferBase"), Reader->BufferArray.GetData(), Reader->BufferArray.Num(), Reader->BufferBase);
+
+	// InternalPrecache partway through the buffer should allocate the full buffer/not allocate anything, and should handle the case of BufferArray's allocation being smaller than BufferSize
+	SetPosAndBuffer(Reader, 3*BufferSize/4, BufferSize/2, BufferSize/2);
+	Reader->BufferArray.Shrink();
+	Result = Reader->InternalPrecache(3*BufferSize/4, BufferSize);
+	TestEqual(TEXT("PrecachePartwayPosMoreThanBufferReallocation - Should succeed"), Result, true);
+	if (FArchiveFileReaderGeneric::bPrecacheAsSoonAsPossible)
+	{
+		TestEqual(TEXT("PrecachePartwayPosMoreThanBufferReallocation - BufferBase should be updated"), Reader->BufferBase, 3*BufferSize/4);
+		TestEqual(TEXT("PrecachePartwayPosMoreThanBufferReallocation - BufferCount should be set to BufferSize"), Reader->BufferArray.Num(), BufferSize);
+	}
+	else
+	{
+		TestEqual(TEXT("PrecachePartwayPosMoreThanBufferReallocation - BufferBase should not be updated"), Reader->BufferBase, BufferSize/2);
+		TestEqual(TEXT("PrecachePartwayPosMoreThanBufferReallocation - BufferCount should not be updated"), Reader->BufferArray.Num(), BufferSize/2);
+	}
+	TestBytesValid(TEXT("PrecachePartwayPosMoreThanBufferReallocation - BufferBytes should match BufferBase"), Reader->BufferArray.GetData(), Reader->BufferArray.Num(), Reader->BufferBase);
+
+	// InternalPrecache partway through the buffer, when the requested size is within the remaining buffer, should not change the buffer
+	PosStart = BufferSize / 2;
+	SetPosAndBuffer(Reader, PosStart, 0, BufferSize);
+	Result = Reader->InternalPrecache(PosStart, BufferSize / 4);
+	TestEqual(TEXT("PrecachePartwayPosLessThanBuffer - Should succeed"), Result, true);
+	TestEqual(TEXT("PrecachePartwayPosLessThanBuffer - BufferBase should not be updated"), Reader->BufferBase, 0LL);
+	TestEqual(TEXT("PrecachePartwayPosLessThanBuffer - BufferCount should be set to BufferSize"), Reader->BufferArray.Num(), BufferSize);
+	TestBytesValid(TEXT("PrecachePartwayPosLessThanBuffer - BufferBytes should match BufferBase"), Reader->BufferArray.GetData(), Reader->BufferArray.Num(), Reader->BufferBase);
+
+	// InternalPrecache at the end of the buffer should allocate the full buffer
+	PosStart = BufferSize;
+	SetPosAndBuffer(Reader, PosStart, 0, BufferSize);
+	Result = Reader->InternalPrecache(PosStart, BufferSize);
+	TestEqual(TEXT("PrecachePosEndOfBuffer - Should succeed"), Result, true);
+	TestEqual(TEXT("PrecachePosEndOfBuffer - BufferBase should be updated"), Reader->BufferBase, PosStart);
+	TestEqual(TEXT("PrecachePosEndOfBuffer - BufferCount should be set to BufferSize"), Reader->BufferArray.Num(), BufferSize);
+	TestBytesValid(TEXT("PrecachePosEndOfBuffer - BufferBytes should match BufferBase"), Reader->BufferArray.GetData(), Reader->BufferArray.Num(), Reader->BufferBase);
+
+	// InternalPrecache near the end of the file should clamp BufferSize to Size - Pos
+	PosStart = FileSize - 1;
+	SetPosAndBuffer(Reader, PosStart, 0, 0);
+	Result = Reader->InternalPrecache(PosStart, BufferSize);
+	TestEqual(TEXT("PrecacheNearEndOfFile - Should succeed"), Result, true);
+	TestEqual(TEXT("PrecacheNearEndOfFile - BufferBase should be updated"), Reader->BufferBase, PosStart);
+	TestEqual(TEXT("PrecacheNearEndOfFile - BufferCount should be set Min(BufferSize, Size - Pos)"), Reader->BufferArray.Num(), 1LL);
+	TestBytesValid(TEXT("PrecacheNearEndOfFile - BufferBytes should match BufferBase"), Reader->BufferArray.GetData(), Reader->BufferArray.Num(), Reader->BufferBase);
+
+	//AddExpectedError(TEXT("ReadFile failed"));
+	// InternalPrecache at the end of the file should fail
+	SetPosAndBuffer(Reader, FileSize, 0, 16);
+	Result = Reader->InternalPrecache(FileSize, BufferSize);
+	TestEqual(TEXT("PrecacheEndOfFile - Should fail"), Result, false);
+	TestEqual(TEXT("PrecacheEndOfFile - BufferBase should not be updated"), Reader->BufferBase, 0LL);
+	TestEqual(TEXT("PrecacheEndOfFile - BufferCount should not be updated"), Reader->BufferArray.Num(), 16LL);
+
+	// If ReadLowLevel fails, InternalPrecache should return true if the first byte of PrecacheOffset is in the buffer
+	SetPosAndBuffer(Reader, 8, 0, 16);
+	Reader->SeekLowLevel(FileSize);
+	Result = Reader->InternalPrecache(8, 16);
+	TestEqual(TEXT("PrecacheReadFailsBytesRemain - Should succeed if bytes are left in the buffer"), Result, true);
+	if (FArchiveFileReaderGeneric::bPrecacheAsSoonAsPossible)
+	{
+		TestEqual(TEXT("PrecacheReadFailsBytesRemain - BufferBase should be updated"), Reader->BufferBase, 8LL);
+		TestEqual(TEXT("PrecacheReadFailsBytesRemain - BufferCount should be updated"), Reader->BufferArray.Num(), 8LL);
+	}
+	else
+	{
+		TestEqual(TEXT("PrecacheReadFailsBytesRemain - BufferBase should not be updated"), Reader->BufferBase, 0LL);
+		TestEqual(TEXT("PrecacheReadFailsBytesRemain - BufferCount should not be updated"), Reader->BufferArray.Num(), 16LL);
+	}
+
+	// If ReadLowLevel fails, InternalPrecache should return false if the first byte of PrecacheOffset is not in the buffer
+	SetPosAndBuffer(Reader, 16, 0, 16);
+	Reader->SeekLowLevel(FileSize);
+	Result = Reader->InternalPrecache(16, 16);
+	TestEqual(TEXT("PrecacheReadFails - Should fail if bytes are not left in the buffer"), Result, false);
+	TestEqual(TEXT("PrecacheReadFails - BufferCount should be updated to empty"), Reader->BufferArray.Num(), 0LL);
+
+	// TODO: Call ClearExpectedErrors to clear the AddExpectedError called above when it is implemented. Until then, make sure these failure-testing cases are the last
+	// cases tested in the test
+	//ClearExpectedErrors();
+}
+
+bool FArchiveFileReaderGenericTest::CreateTestFile()
+{
+	IFileManager& FileManager = IFileManager::Get();
+	IPlatformFile& PlatformFile = FPlatformFileManager::Get().GetPlatformFile();;
+	TestFileName = FPaths::CreateTempFilename(FPlatformProcess::UserTempDir(), TEXT("ArchiveFileReaderGenericTest"));
+	TUniquePtr<IFileHandle> WriteHandle(PlatformFile.OpenWrite(*TestFileName, false /* bAppend */, false /* bAllowRead */));
+	if (!WriteHandle)
+	{
+		TestAddError(FString::Printf(TEXT("Could not create test file %s."), *TestFileName));
+		return false;
+	}
+	TArray<int32> Builder;
+	int32 NumInts = 1024 * 128 + 17; // Make sure the file size is not a multiple of buffersize; we need to test that behavior
+	Builder.SetNumUninitialized(NumInts);
+	int32* Data = Builder.GetData();
+	for (int n = 0; n < Builder.Num(); ++n)
+	{
+		Data[n] = GetExpectedValue(n);
+	}
+	FileSize = NumInts * sizeof(int32);
+	WriteHandle->Write(reinterpret_cast<uint8*>(Data), FileSize);
+	WriteHandle.Reset();
+
+	ReadHandle.Reset(PlatformFile.OpenRead(*TestFileName, false /* bAllowWrite */));
+	if (!ReadHandle)
+	{
+		TestAddError(FString::Printf(TEXT("Could not open test file %s."), *TestFileName));
+		return false;
+	}
+	if (ReadHandle->Size() != FileSize)
+	{
+		TestAddError(FString::Printf(TEXT("Received incorrect file size from test file %s. Expected = %lld. Actual = %lld."), *TestFileName, FileSize, ReadHandle->Size()));
+		return false;
+	}
+
+	return true;
+}
+
+void FArchiveFileReaderGenericTest::TestBytesValid(const TCHAR* What, uint8* InData, int64 InNumBytes, int64 InFileOffset)
+{
+	uint8* Data = InData;
+	int64 NumBytes = InNumBytes;
+	int64 FileOffset = InFileOffset;
+
+	int64 StartByte = FileOffset % sizeof(int32);
+	if (StartByte != 0)
+	{
+		int64 FloorIndex = (FileOffset - StartByte)/sizeof(int32);
+		int64 ExpectedValue = GetExpectedValue(FloorIndex);
+		for (int n = 0; n < NumBytes && n < sizeof(int32) - StartByte; ++n)
+		{
+			if (Data[n] != reinterpret_cast<uint8*>(&ExpectedValue)[StartByte + n])
+			{
+				TestEqual(What, Data[n], reinterpret_cast<uint8*>(&ExpectedValue)[StartByte + n]);
+				return;
+			}
+		}
+		if (NumBytes < static_cast<int64>(sizeof(int32) - StartByte))
+		{
+			return;
+		}
+		int32 BytesRead = sizeof(int32) - StartByte;
+		Data += BytesRead;
+		FileOffset += BytesRead;
+		NumBytes -= BytesRead;
+	}
+	int32* IntData = reinterpret_cast<int32*>(Data);
+	int64 IntOffset = FileOffset / sizeof(int32);
+	int64 NumInts = NumBytes / sizeof(int32);
+	for (int n = 0; n < NumInts; ++n)
+	{
+		if (IntData[n] != GetExpectedValue(IntOffset + n))
+		{
+			TestEqual(What, IntData[n], GetExpectedValue(IntOffset + n));
+			return;
+		}
+	}
+
+	int64 NumEndBytes = NumBytes - NumInts * sizeof(int32);
+	if (NumEndBytes != 0)
+	{
+		uint8* DataStart = Data + NumInts * sizeof(int32);
+		int32 ExpectedValue = GetExpectedValue(IntOffset + NumInts);
+		for (int n = 0; n < NumEndBytes; ++n)
+		{
+			if (DataStart[n] != reinterpret_cast<uint8*>(&ExpectedValue)[n])
+			{
+				TestEqual(What, DataStart[n], reinterpret_cast<uint8*>(&ExpectedValue)[n]);
+				return;
+			}
+		}
+	}
+}
+
+void FArchiveFileReaderGenericTest::SetPosAndBuffer(TUniquePtr<FArchiveFileReaderGeneric>& Reader, int64 Pos, int64 BufferBase, int64 BufferSize)
+{
+	Reader->Pos = Pos;
+	Reader->BufferBase = BufferBase;
+	Reader->BufferArray.SetNumUninitialized(BufferSize, false /* bAllowShrink */);
+	if (BufferSize)
+	{
+		Reader->SeekLowLevel(BufferBase);
+		int64 BytesRead;
+		Reader->ReadLowLevel(Reader->BufferArray.GetData(), BufferSize, BytesRead);
+		TestEqual(TEXT("SetPosAndBuffer - ReadLowLevel read the requested bytes"), BytesRead, BufferSize);
+	}
+	// See the contract for Buffer window and LowLevel Pos in the variable comment on BufferArray
+	bool bPosWithinBuffer = BufferBase <= Pos && Pos < BufferBase + BufferSize;
+	if (!bPosWithinBuffer)
+	{
+		Reader->SeekLowLevel(Pos);
+	}
+	TestEqual(TEXT("SetPosAndBuffer - Pos set correctly"), Reader->Pos, Pos);
+	TestEqual(TEXT("SetPosAndBuffer - BufferBase set correctly"), Reader->BufferBase, BufferBase);
+	TestEqual(TEXT("SetPosAndBuffer - BufferSize set correctly"), Reader->BufferArray.Num(), BufferSize);
+	TestBytesValid(TEXT("SetPosAndBuffer - BufferBytes should match BufferBase"), Reader->BufferArray.GetData(), Reader->BufferArray.Num(), Reader->BufferBase);
+}
+
+int32 FArchiveFileReaderGenericTest::GetExpectedValue(int64 IntOffset)
+{
+	return 0xbe000000 + static_cast<int32>(IntOffset);
+}
+
+TEST_CASE("Core::HAL::FArchiveFileReaderGeneric::Smoke Test", "[Core][HAL][Smoke]")
+{
+	FArchiveFileReaderGenericTest ArchiveFileReaderGenericTest;
+	ArchiveFileReaderGenericTest.TestInternalPrecache();
+}
diff --git a/Engine/Source/Runtime/Core/Tests/HAL/FileManagerTest.cpp b/Engine/Source/Runtime/Core/Tests/HAL/FileManagerTest.cpp
new file mode 100644
--- /dev/null
+++ b/Engine/Source/Runtime/Core/Tests/HAL/FileManagerTest.cpp
@@ -0,0 +1,10 @@
+// Copyright Epic Games, Inc. All Rights Reserved.
+
+#include "HAL/FileManager.h"
+
+#include "TestHarness.h"
+
+TEST_CASE("Core::HAL::FileManager::Reader", "[Core][HAL][FileManager]")
+{
+	// TODO: Add tests
+}
\ No newline at end of file
diff --git a/Engine/Source/Runtime/Core/Tests/HAL/PlatformTest.cpp b/Engine/Source/Runtime/Core/Tests/HAL/PlatformTest.cpp
new file mode 100644
--- /dev/null
+++ b/Engine/Source/Runtime/Core/Tests/HAL/PlatformTest.cpp
@@ -0,0 +1,81 @@
+// Copyright Epic Games, Inc. All Rights Reserved.
+
+#include "CoreTypes.h"
+#include "Misc/AssertionMacros.h"
+#include "GenericPlatform/GenericPlatformMath.h"
+#include "Math/UnrealMathUtility.h"
+#include "Templates/UnrealTemplate.h"
+#include "Containers/UnrealString.h"
+#include "UObject/NameTypes.h"
+#include "TestHarness.h"
+
+struct TestA
+{
+	virtual ~TestA() {}
+	virtual void TestAA() 
+	{
+		Space[0] = 1;
+	}
+	uint8 Space[64];
+};
+
+
+struct TestB
+{
+	virtual ~TestB() {}
+	virtual void TestBB() 
+	{
+		Space[5] = 1;
+	}
+	uint8 Space[96];
+};
+
+
+struct TestC : public TestA, TestB
+{
+	int i;
+};
+
+TEST_CASE("Core::HAL::PlatformTest::Smoke Test", "[Core][HAL][Smoke]")
+{
+	PTRINT Offset1 = VTABLE_OFFSET(TestC, TestB);
+	PTRINT Offset2 = VTABLE_OFFSET(TestC, TestA);
+	CHECK(Offset1 == 64 + sizeof(void*));
+	CHECK(Offset2 == 0);
+	int32 Test = 0x12345678;
+#if PLATFORM_LITTLE_ENDIAN
+	CHECK(*(uint8*)&Test == 0x78);
+#else
+	CHECK(*(uint8*)&Test == 0x12);
+#endif
+
+	FPlatformMath::AutoTest();
+
+#if WITH_EDITORONLY_DATA
+	CHECK(FPlatformProperties::HasEditorOnlyData());
+#else
+	CHECK(!FPlatformProperties::HasEditorOnlyData());
+#endif
+
+	CHECK(FPlatformProperties::HasEditorOnlyData() != FPlatformProperties::RequiresCookedData());
+
+#if PLATFORM_LITTLE_ENDIAN
+	CHECK(FPlatformProperties::IsLittleEndian());
+#else
+	CHECK(!FPlatformProperties::IsLittleEndian());
+#endif
+	CHECK(FPlatformProperties::PlatformName());
+
+	CHECK(FString(FPlatformProperties::PlatformName()).Len() > 0); 
+
+	static_assert(alignof(int32) == 4, "Align of int32 is not 4."); //Hmmm, this would be very strange, ok maybe, but strange
+
+	MS_ALIGN(16) struct FTestAlign
+	{
+		uint8 Test;
+	} GCC_ALIGN(16);
+
+	static_assert(alignof(FTestAlign) == 16, "Align of FTestAlign is not 16.");
+
+	FName::AutoTest();
+}
diff --git a/Engine/Source/Runtime/Core/Tests/HAL/ThreadTest.cpp b/Engine/Source/Runtime/Core/Tests/HAL/ThreadTest.cpp
new file mode 100644
--- /dev/null
+++ b/Engine/Source/Runtime/Core/Tests/HAL/ThreadTest.cpp
@@ -0,0 +1,264 @@
+// Copyright Epic Games, Inc. All Rights Reserved.
+
+#include "CoreTypes.h"
+#include "Misc/AssertionMacros.h"
+
+#include "HAL/Thread.h"
+#include "HAL/ThreadSafeBool.h"
+#include "HAL/ThreadSafeCounter.h"
+#include "HAL/ThreadSingleton.h"
+#include "HAL/Event.h"
+#include "Containers/Queue.h"
+#include "Containers/StringConv.h"
+
+#include "TestHarness.h"
+
+namespace
+{
+	void TestIsJoinableAfterCreation()
+	{
+		FThread Thread(TEXT("Test.Thread.TestIsJoinableAfterCreation"), []() { /*NOOP*/ });
+		TestTrue(TEXT("FThread must be joinable after construction"), Thread.IsJoinable());
+		Thread.Join();
+		UE_LOG(LogTemp, Log, TEXT("%s completed"), StringCast<TCHAR>(__FUNCTION__).Get());
+	}
+
+	void TestIsJoinableAfterCompletion()
+	{
+		FThreadSafeBool bDone = false;
+		FThread Thread(TEXT("Test.Thread.TestIsJoinableAfterCompletion"), [&bDone]() { bDone = true; });
+		while (!bDone); // wait for completion //-V529
+		TestTrue(TEXT("FThread must still be joinable after completion"), Thread.IsJoinable());
+		Thread.Join();
+		UE_LOG(LogTemp, Log, TEXT("%s completed"), StringCast<TCHAR>(__FUNCTION__).Get());
+	}
+
+	void TestIsNotJoinableAfterJoining()
+	{
+		FThread Thread(TEXT("Test.Thread.TestIsNotJoinableAfterJoining"), []() { /*NOOP*/ });
+		Thread.Join();
+		TestFalse(TEXT("FThread must not be joinable after joining"), Thread.IsJoinable());
+		UE_LOG(LogTemp, Log, TEXT("%s completed"), StringCast<TCHAR>(__FUNCTION__).Get());
+	}
+
+#if 0 // detaching is not implemented
+	void TestIsNotJoinableAfterDetaching()
+	{
+		// two cases: it's either the calling thread detaches from the thread before the thread is completed
+		{
+			TAtomic<bool> bReady{ false };
+			FThread Thread(TEXT("Test.Thread"), [&bReady]()
+				{
+					while (!bReady) {}
+				});
+			Thread.Detach();
+			bReady = true; // make sure `Detach` is called before thread function exit
+			TestFalse(TEXT("FThread must not be joinable after detaching"), Thread.IsJoinable());
+		}
+		// or thread function is completed fast and `FThreadImpl` releases the reference to itself before
+		// `Detach` call
+		{
+			TAtomic<bool> bReady{ false };
+			FThread Thread(TEXT("Test.Thread"), [&bReady]() { /*NOOP*/});
+			FPlatformProcess::Sleep(0.1); // let the thread exit before detaching
+			Thread.Detach();
+			bReady = true; // make sure `Detach` is called before thread function exit
+			TestFalse(TEXT("FThread must not be joinable after detaching"), Thread.IsJoinable());
+		}
+		UE_LOG(LogTemp, Log, TEXT("%s completed"), StringCast<TCHAR>(__FUNCTION__).Get());
+	}
+#endif
+
+	void TestAssertIfNotJoinedOrDetached()
+	{
+		// this does fails the `check`, but it seems there's no way to test this by UE4 Automation Testing, so commented out
+		FThread Thread(TEXT("Test.Thread.TestAssertIfNotJoinedOrDetached"), []() { /*NOOP*/ });
+		// should assert in the destructor
+	}
+
+	void TestDefaultConstruction()
+	{
+		{
+			FThread Thread;
+			TestFalse(TEXT("Default-constructed FThread must be not joinable"), Thread.IsJoinable());
+		}
+		{	// check that default constructed thread can be "upgraded" to joinable thread
+			FThread Thread;
+			Thread = FThread(TEXT("Test.Thread.TestDefaultConstruction"), []() { /* NOOP */ });
+			TestTrue(TEXT("Move-constructed FThread from joinable thread must be joinable"), Thread.IsJoinable());
+			Thread.Join();
+		}
+		UE_LOG(LogTemp, Log, TEXT("%s completed"), StringCast<TCHAR>(__FUNCTION__).Get());
+	}
+
+	void TestThreadSingleton()
+	{
+		{	// check that ThreadSingleton instances in different threads are isolated from each other
+			class FThreadSingletonTest : public TThreadSingleton<FThreadSingletonTest>
+			{
+			public:
+				void SetTestField(int NewValue) { TestField = NewValue; }
+				int GetTestField() const { return TestField; }
+			private:
+				int TestField{ 0 };
+			};
+			FThreadSingletonTest::Get().SetTestField(1);
+			FThread Thread;
+			TAtomic<bool> bDefaultValuePass{ false };
+			Thread = FThread(TEXT("Test.Thread.TestThreadSingleton"),
+				[&bDefaultValuePass]()
+				{
+					bDefaultValuePass = FThreadSingletonTest::TryGet() == nullptr;
+					bDefaultValuePass = bDefaultValuePass && (FThreadSingletonTest::Get().GetTestField() == 0);
+					FThreadSingletonTest::Get().SetTestField(2);
+				});
+			Thread.Join();
+			TestTrue(TEXT("Thread singleton should be uninitialized by default"), bDefaultValuePass);
+			TestTrue(TEXT("Thread singletons in different threads should be isolated"), FThreadSingletonTest::Get().GetTestField() == 1);
+		}
+		{	// check that ThreadSingleton entries don't point to invalid memory after cleanup
+			class FThreadSingletonFirst : public TThreadSingleton<FThreadSingletonFirst>
+			{
+			};
+			class FThreadSingletonSecond : public TThreadSingleton<FThreadSingletonSecond>
+			{
+			public:
+				virtual ~FThreadSingletonSecond()
+				{
+					// By the time we reach this destructor, the first singleton's destructor should have been executed already.
+					check(FThreadSingletonFirst::TryGet() == nullptr);
+				}
+			};
+			FThread Thread;
+			Thread = FThread(TEXT("Test.Thread.TestThreadSingleton"),
+				[]()
+				{
+					FThreadSingletonFirst::Get();
+					FThreadSingletonSecond::Get();
+				});
+			Thread.Join();
+			TestTrue(TEXT("Thread singletons should be uninitialized by default"), FThreadSingletonFirst::TryGet() == nullptr);
+			TestTrue(TEXT("Thread singletons should be uninitialized by default"), FThreadSingletonSecond::TryGet() == nullptr);
+		}
+		UE_LOG(LogTemp, Log, TEXT("%s completed"), StringCast<TCHAR>(__FUNCTION__).Get());
+	}
+
+	void TestMovability()
+	{
+		{	// move constructor with default-constructed thread
+			FThread Src;
+			FThread Dst(MoveTemp(Src));
+			TestFalse(TEXT("Default-constructed thread must stay not joinable after moving out"), Src.IsJoinable());
+			TestFalse(TEXT("Move-constructed thread from not joinable thread must be not joinable"), Dst.IsJoinable());
+		}
+		{	// move constructor with joinable thread
+			FThread Src(TEXT("Test.Thread.TestMovability.1"), []() { /* NOOP */ });
+			FThread Dst(MoveTemp(Src));
+			TestFalse(TEXT("Moved out thread must be not joinable"), Src.IsJoinable());
+			TestTrue(TEXT("Move-constructed thread from joinable thread must be joinable"), Dst.IsJoinable());
+			Dst.Join();
+		}
+		{	// move assignment operator
+			FThread Src(TEXT("Test.Thread.TestMovability.2"), []() { /* NOOP */ });
+			FThread Dst;
+			Dst = MoveTemp(Src);
+			TestFalse(TEXT("Moved out thread must be not joinable"), Src.IsJoinable());
+			TestTrue(TEXT("Move-assigned thread from joinable thread must be joinable"), Dst.IsJoinable());
+			Dst.Join();
+		}
+		{	// Failure test for move assignment operator of joinable thread
+			//FThread Src(TEXT("Test.Thread"), []() { /* NOOP */ });
+			//FThread Dst(TEXT("Test.Thread"), []() { /* NOOP */ });
+			//Dst = MoveTemp(Src); // must assert that joinable thread wasn't joined before move-assignment, no way to test this
+			//Dst.Join();
+		}
+		{	// Move assignment operator of thread that has been joined
+			FThread Src(TEXT("Test.Thread.TestMovability.3"), []() { /* NOOP */ });
+			FThread Dst(TEXT("Test.Thread.TestMovability.4"), []() { /* NOOP */ });
+			Dst.Join();
+			Dst = MoveTemp(Src);
+			Dst.Join();
+		}
+		UE_LOG(LogTemp, Log, TEXT("%s completed"), StringCast<TCHAR>(__FUNCTION__).Get());
+	}
+
+	// An example of possible implementation of Consumer/Producer idiom
+	void TestTypicalUseCase()
+	{
+		FThreadSafeBool bQuitRequested = false;
+		using FWork = uint32;
+		TQueue<FWork> WorkQueue;
+		FEvent* WorkQueuedEvent = FPlatformProcess::GetSynchEventFromPool();
+
+		FThread WorkerThread(TEXT("Test.Thread.TestTypicalUseCase"), [&bQuitRequested, &WorkQueue, WorkQueuedEvent]() 
+		{
+			while (!bQuitRequested)
+			{
+				// get work
+				FWork Work;
+				if (!WorkQueue.Dequeue(Work))
+				{
+					WorkQueuedEvent->Wait();
+					// first check if quitting requested
+					continue;
+				}
+
+				// do work
+				UE_LOG(LogTemp, Log, TEXT("Work #%d consumed"), Work);
+			}
+
+			UE_LOG(LogTemp, Log, TEXT("Quit"));
+		});
+
+		// produce work
+		const int WorkNum = 3;
+		for (FWork Work = 0; Work != WorkNum; ++Work)
+		{
+			WorkQueue.Enqueue(Work);
+			WorkQueuedEvent->Trigger();
+			UE_LOG(LogTemp, Log, TEXT("Work #%d produced"), Work);
+		}
+
+		UE_LOG(LogTemp, Log, TEXT("Request to quit"));
+		bQuitRequested = true;
+		// the thread can be blocked waiting for work, unblock it
+		WorkQueuedEvent->Trigger();
+		WorkerThread.Join();
+
+		FPlatformProcess::ReturnSynchEventToPool(WorkQueuedEvent);
+
+		// example of output:
+		// Work #0 produced
+		//	Work #0 consumed
+		//	Work #1 produced
+		//	Work #1 consumed
+		//	Work #2 produced
+		//	Work #2 consumed
+		//	Request to quit
+		//	The thread 0x96e0 has exited with code 0 (0x0).
+		//	Quit
+	
+		UE_LOG(LogTemp, Log, TEXT("%s completed"), StringCast<TCHAR>(__FUNCTION__).Get());
+	}
+}
+
+TEST_CASE("Core::HAL::Thread", "[Core][HAL][Thread]")
+{
+	UE_LOG(LogTemp, Log, TEXT("%s"), StringCast<TCHAR>(__FUNCTION__).Get());
+
+	TestIsJoinableAfterCreation();
+	TestIsJoinableAfterCompletion();
+	TestIsNotJoinableAfterJoining();
+	
+#if 0 // detaching is not implemented
+	TestIsNotJoinableAfterDetaching();
+#endif
+
+	//TestAssertIfNotJoinedOrDetached();
+
+	TestDefaultConstruction();
+	TestMovability();
+
+	TestTypicalUseCase();
+	TestThreadSingleton();
+}
\ No newline at end of file
diff --git a/Engine/Source/Runtime/Core/Tests/IO/IoStatusOrTest.cpp b/Engine/Source/Runtime/Core/Tests/IO/IoStatusOrTest.cpp
new file mode 100644
--- /dev/null
+++ b/Engine/Source/Runtime/Core/Tests/IO/IoStatusOrTest.cpp
@@ -0,0 +1,146 @@
+// Copyright Epic Games, Inc. All Rights Reserved.
+
+#include "CoreTypes.h"
+#include "Misc/AssertionMacros.h"
+#include "IO/IoDispatcher.h"
+#include "TestHarness.h"
+
+struct FIoStatusTestType
+{
+	FIoStatusTestType() { }
+	FIoStatusTestType(const FIoStatusTestType& Other)
+		: Text(Other.Text) { }
+	FIoStatusTestType(FIoStatusTestType&&) = default;
+
+	FIoStatusTestType(const FString& InText)
+		: Text(InText) { }
+	FIoStatusTestType(FString&& InText)
+		: Text(MoveTemp(InText)) { }
+
+	FIoStatusTestType& operator=(const FIoStatusTestType& Other) = default;
+	FIoStatusTestType& operator=(FIoStatusTestType&& Other) = default;
+	FIoStatusTestType& operator=(const FString& OtherText)
+	{
+		Text = OtherText;
+		return *this;
+	}
+
+	FString Text;
+};
+
+void TestConstruct()
+{
+	{
+		TIoStatusOr<FIoStatusTestType> Result;
+		TestEqual("Default IoStatus is Unknown", Result.Status(), FIoStatus::Unknown);
+	}
+
+	{
+		const TIoStatusOr<FIoStatusTestType> Other;
+		TIoStatusOr<FIoStatusTestType> Result(Other);
+		TestEqual("Copy construct", Result.Status(), FIoStatus::Unknown);
+	}
+
+	{
+		const FIoStatus IoStatus(EIoErrorCode::InvalidCode);
+		TIoStatusOr<FIoStatusTestType> Result(IoStatus);
+		TestEqual("Construct with status", Result.Status().GetErrorCode(), EIoErrorCode::InvalidCode);
+	}
+
+	{
+		const FString ExpectedText("Unreal");
+		const FIoStatusTestType Type(ExpectedText);
+		TIoStatusOr<FIoStatusTestType> Result(Type);
+		TestEqual("Construct with value", Result.ValueOrDie().Text, ExpectedText);
+	}
+
+	{
+		const FString ExpectedText("Unreal");
+		TIoStatusOr<FIoStatusTestType> Result(FIoStatusTestType("Unreal"));
+		TestEqual("Construct with temporary value", Result.ValueOrDie().Text, ExpectedText);
+	}
+
+	{
+		TIoStatusOr<FIoStatusTestType> Result(FString("Unreal"));
+		TestEqual("Construct with value arguments", Result.ValueOrDie().Text, FString("Unreal"));
+	}
+}
+
+void TestAssignment()
+{
+	{
+		const EIoErrorCode ExpectedErrorCode = EIoErrorCode::InvalidCode;
+		TIoStatusOr<FIoStatusTestType> Other = FIoStatus(ExpectedErrorCode);
+		TIoStatusOr<FIoStatusTestType> Result;
+		Result = Other;
+		TestEqual("Assign IoStatusOr with status", Result.Status().GetErrorCode(), ExpectedErrorCode);
+	}
+
+	{
+		const EIoErrorCode ExpectedErrorCode = EIoErrorCode::InvalidCode;
+		TIoStatusOr<FIoStatusTestType> Result;
+		Result = TIoStatusOr<FIoStatusTestType>(FIoStatus(ExpectedErrorCode));
+		TestEqual("Assign temporary IoStatusOr with status", Result.Status().GetErrorCode(), ExpectedErrorCode);
+	}
+
+	{
+		const FString ExpectedText("Unreal");
+		TIoStatusOr<FIoStatusTestType> Other = FIoStatusTestType(ExpectedText);
+		TIoStatusOr<FIoStatusTestType> Result;
+		Result = Other;
+		TestEqual("Assign IoStatusOr with value", Result.ValueOrDie().Text, ExpectedText);
+	}
+
+	{
+		const FString ExpectedText("Unreal");
+		TIoStatusOr<FIoStatusTestType> Result;
+		Result = TIoStatusOr<FIoStatusTestType>(ExpectedText);
+		TestEqual("Assign temporary IoStatusOr with value", Result.ValueOrDie().Text, ExpectedText);
+	}
+
+	{
+		const EIoErrorCode ExpectedErrorCode = EIoErrorCode::InvalidCode;
+		const FIoStatus IoStatus(ExpectedErrorCode);
+		TIoStatusOr<FIoStatusTestType> Result;
+		Result = IoStatus; 
+		TestEqual("Assign status", Result.Status().GetErrorCode(), ExpectedErrorCode);
+	}
+
+	{
+		const EIoErrorCode ExpectedErrorCode = EIoErrorCode::InvalidCode;
+		TIoStatusOr<FIoStatusTestType> Result;
+		Result = FIoStatus(ExpectedErrorCode);
+		TestEqual("Assign temporary status", Result.Status().GetErrorCode(), ExpectedErrorCode);
+	}
+
+	{
+		const FString ExpectedText("Unreal");
+		const FIoStatusTestType Value(ExpectedText);
+		TIoStatusOr<FIoStatusTestType> Result;
+		Result = Value;
+		TestEqual("Assign value", Result.ValueOrDie().Text, ExpectedText);
+	}
+
+	{
+		const FString ExpectedText("Unreal");
+		TIoStatusOr<FIoStatusTestType> Result;
+		Result = FIoStatusTestType(ExpectedText);
+		TestEqual("Assign temporary value", Result.ValueOrDie().Text, ExpectedText);
+	}
+}
+
+void TestConsumeValue()
+{
+	const FString ExpectedText("Unreal");
+	TIoStatusOr<FIoStatusTestType> Result = FIoStatusTestType(ExpectedText);
+	FIoStatusTestType Value = Result.ConsumeValueOrDie();
+	TestEqual("Consume value or die with valid value", Value.Text, ExpectedText);
+}
+
+TEST_CASE("Core::IO::IoStatusOr::Smoke Test", "[Core][IO][Smoke]")
+{
+	TestConstruct();
+	TestAssignment();
+	TestConsumeValue();
+}
+
diff --git a/Engine/Source/Runtime/Core/Tests/Internationalization/ArchiveTest.cpp b/Engine/Source/Runtime/Core/Tests/Internationalization/ArchiveTest.cpp
new file mode 100644
--- /dev/null
+++ b/Engine/Source/Runtime/Core/Tests/Internationalization/ArchiveTest.cpp
@@ -0,0 +1,201 @@
+// Copyright Epic Games, Inc. All Rights Reserved.
+
+#include "CoreTypes.h"
+#include "Containers/UnrealString.h"
+#include "Internationalization/InternationalizationManifest.h"
+#include "Internationalization/InternationalizationArchive.h"
+#include "Internationalization/InternationalizationMetadata.h"
+#include "TestHarness.h"
+
+TEST_CASE("Core::Internationalization::FInternationalizationArchive::Smoke Test", "[Core][Internationalization][Smoke]")
+{
+	// Key metadata
+	TSharedPtr< FLocMetadataObject > KeyMetadataA = MakeShareable( new FLocMetadataObject );
+	TSharedPtr< FLocMetadataObject > KeyMetadataB = MakeShareable( new FLocMetadataObject );
+
+	// Source metadata
+	TSharedPtr< FLocMetadataObject > SourceMetadataA = MakeShareable( new FLocMetadataObject );
+	TSharedPtr< FLocMetadataObject > SourceMetadataB = MakeShareable( new FLocMetadataObject );
+
+
+	// Setup KeyMetadataA
+	KeyMetadataA->SetStringField( TEXT("Gender"			),	TEXT("Masculine")	);
+	KeyMetadataA->SetStringField( TEXT("Plurality"		),	TEXT("Singular")	);
+	KeyMetadataA->SetStringField( TEXT("TargetGender"	),	TEXT("Masculine")	);
+	KeyMetadataA->SetStringField( TEXT("TargetPlurality"),	TEXT("Singular")	);
+
+	// Setup KeyMetadataB
+	KeyMetadataB->SetStringField( TEXT("Gender"			),	TEXT("Masculine")	);
+	KeyMetadataB->SetStringField( TEXT("Plurality"		),	TEXT("Singular")	);
+	KeyMetadataB->SetStringField( TEXT("TargetGender"	),	TEXT("Feminine")	);
+	KeyMetadataB->SetStringField( TEXT("TargetPlurality"),	TEXT("Singular")	);
+
+	// Setup source metadata
+	SourceMetadataA->SetBoolField( TEXT("*IsMature"), false );
+	SourceMetadataB->SetBoolField( TEXT("*IsMature"), true );
+
+	// Setup source item
+	FLocItem SourceA( TEXT("TextA") );
+	SourceA.MetadataObj = MakeShareable( new FLocMetadataObject(*SourceMetadataA) );
+
+	FLocItem Translation = SourceA;
+	Translation.Text = TEXT("TranslatedTextA");
+
+	FString TestNamespace = TEXT("TestNamespace");
+	FString SourceAKey = TEXT("TextA");
+
+	// Test entry add
+	{
+		bool TestOptionalTrue = true;
+		bool TestOptionalFalse = false;
+
+		// bIsOptional is not used as a key.  We ensure adding entries where the bIsOptional member is the only difference works as expected.
+		FInternationalizationArchive TestArchive;
+		TestArchive.AddEntry( TestNamespace, SourceAKey, SourceA, Translation, NULL, TestOptionalTrue );
+		// Add duplicate entry that differs in bIsOptional flag.  This add should report success because we already have an entry with matching 
+		//  namespace/source/keymetadata.  Differences in bIsOptional are not taken into consideration. 
+		bool bResult = TestArchive.AddEntry( TestNamespace, SourceAKey, SourceA, Translation, NULL, TestOptionalFalse );
+		TestTrue( TEXT("AddEntry result = true"), bResult);
+
+		// We should only have one entry in the archive
+		int32 EntryCount = 0;
+		for(auto Iter = TestArchive.GetEntriesBySourceTextIterator(); Iter; ++Iter, ++EntryCount);
+		TestEqual( TEXT("EntryCount == 1"), EntryCount, 1 );
+
+		// Make sure the original bIsOptional value is not overwritten by our second add.
+		TSharedPtr< FArchiveEntry > FoundEntry = TestArchive.FindEntryByKey( TestNamespace, SourceAKey, NULL );
+		if( !FoundEntry.IsValid() )
+		{
+			TestTrue(TEXT("FArchiveEntry could not find entry using FindEntryByKey."), FoundEntry.IsValid());
+		}
+		else
+		{
+			TestTrue( TEXT("FoundEntry->Namespace == Namespace"), FoundEntry->bIsOptional == TestOptionalTrue );
+		}
+	}
+
+	// Test lookup an entry
+	{
+		{
+			FInternationalizationArchive TestArchive;
+			TestArchive.AddEntry( TestNamespace, SourceAKey, SourceA, Translation, KeyMetadataA, false );
+
+			TSharedPtr< FArchiveEntry > FoundEntry = TestArchive.FindEntryByKey( TestNamespace, SourceAKey, KeyMetadataA );
+			if( !FoundEntry.IsValid() )
+			{
+				TestTrue(TEXT("FArchiveEntry could not find entry using FindEntryByKey."), FoundEntry.IsValid());
+			}
+			else
+			{
+				TestTrue( TEXT("FoundEntry->Namespace == Namespace"), FoundEntry->Namespace == TestNamespace );
+				TestTrue( TEXT("FoundEntry->Source == Source"), FoundEntry->Source == SourceA );
+				TestTrue( TEXT("FoundEntry->Translation == Translation"), FoundEntry->Translation == Translation );
+				if( FoundEntry->KeyMetadataObj == KeyMetadataA )
+				{
+					TestFalse(TEXT("FArchiveEntry KeyMetadataObj is not a unique object."), FoundEntry->KeyMetadataObj == KeyMetadataA );
+				}
+				TestTrue( TEXT("FoundEntry->KeyMetadataObj == KeyMetadataA"), *(FoundEntry->KeyMetadataObj) == *(KeyMetadataA) );
+			}
+
+			// Passing in a mismatched key metadata will fail to find the entry.  Any fallback logic is intended to happen at runtime
+			FoundEntry = TestArchive.FindEntryByKey( TestNamespace, SourceAKey, NULL );
+			TestInvalid( TEXT("!FoundEntry.IsValid()"), FoundEntry);
+
+			FoundEntry = TestArchive.FindEntryByKey( TestNamespace, SourceAKey, MakeShareable( new FLocMetadataObject() ) );
+			TestInvalid( TEXT("!FoundEntry.IsValid()"), FoundEntry);
+
+			FoundEntry = TestArchive.FindEntryByKey( TestNamespace, SourceAKey, KeyMetadataB );
+			TestInvalid( TEXT("!FoundEntry.IsValid()"), FoundEntry);
+		}
+
+		// Ensure we can properly lookup entries with non-null but empty key metadata
+		{
+			FInternationalizationArchive TestArchive;
+			TestArchive.AddEntry( TestNamespace, SourceAKey, SourceA, Translation, MakeShareable( new FLocMetadataObject() ), false );
+
+			TSharedPtr< FArchiveEntry > FoundEntry = TestArchive.FindEntryByKey( TestNamespace, SourceAKey, NULL );
+			
+			bool bFoundEntryValid = FoundEntry.IsValid();
+			TestTrue(TEXT("FArchiveEntry could not find entry using FindEntryByKey."), bFoundEntryValid);
+			if (bFoundEntryValid)
+			{
+				TestTrue( TEXT("FoundEntry->Namespace == Namespace"), FoundEntry->Namespace == TestNamespace );
+				TestTrue( TEXT("FoundEntry->Source == Source"), FoundEntry->Source == SourceA );
+			}
+		}
+
+		// Ensure we can properly lookup entries with null key metadata
+		{
+			FInternationalizationArchive TestArchive;
+			TestArchive.AddEntry( TestNamespace, SourceAKey, SourceA, Translation, NULL, false );
+
+			TSharedPtr< FArchiveEntry > FoundEntry = TestArchive.FindEntryByKey( TestNamespace, SourceAKey, NULL );
+			bool bFoundEntryValid = FoundEntry.IsValid();
+			TestTrue(TEXT("FArchiveEntry could not find entry using FindEntryByKey."), bFoundEntryValid);
+			if (bFoundEntryValid)
+			{
+				TestTrue( TEXT("FoundEntry->Namespace == Namespace"), FoundEntry->Namespace == TestNamespace );
+				TestTrue( TEXT("FoundEntry->Source == Source"), FoundEntry->Source == SourceA );
+			}
+
+			// Test lookup with non-null but empty key metadata
+			FoundEntry = TestArchive.FindEntryByKey( TestNamespace, SourceAKey, MakeShareable( new FLocMetadataObject() ) );
+			bFoundEntryValid = FoundEntry.IsValid();
+			TestTrue(TEXT("FArchiveEntry could not find entry using FindEntryByKey."), bFoundEntryValid);
+			if (bFoundEntryValid)
+			{
+				TestTrue(TEXT("FoundEntry->Namespace == Namespace"), FoundEntry->Namespace == TestNamespace);
+				TestTrue( TEXT("FoundEntry->Source == Source"), FoundEntry->Source == SourceA );
+			}
+		}
+
+		// Ensure lookup where source metadata has * prefixed entries work as expected.  Note: The source metadata object
+		//   supports a * prefix on metadata names which modifies the way we perform metadata comparison(ignores entry type and value, only name is checked)
+		{
+			FLocItem SourceCompare( TEXT("TextA") );
+			SourceCompare.MetadataObj = MakeShareable( new FLocMetadataObject() );
+			SourceCompare.MetadataObj->SetStringField( TEXT("*IsMature"), TEXT("") );
+
+			FInternationalizationArchive TestArchive;
+			// Added entry with String *IsMature entry
+			TestArchive.AddEntry( TestNamespace, SourceAKey, SourceCompare, Translation, KeyMetadataA, false );
+
+			// Finding entry using a source that has Boolean *IsMature
+			TSharedPtr< FArchiveEntry > FoundEntry = TestArchive.FindEntryByKey( TestNamespace, SourceAKey, KeyMetadataA );
+			bool bFoundEntryValid = FoundEntry.IsValid();
+			TestTrue(TEXT("FArchiveEntry could not find entry using FindEntryByKey."), bFoundEntryValid);
+			if (bFoundEntryValid)
+			{
+				TestTrue( TEXT("FoundEntry->Namespace == Namespace"), FoundEntry->Namespace == TestNamespace );
+				TestTrue( TEXT("FoundEntry->Source == Source"), FoundEntry->Source == SourceA );
+			}
+
+			// Attempting to add an entry that only differs by a * prefexed source metadata entry value or type will result in success since
+			//  a 'matching' entry already exists in the archive.  We should, however, only have one entry in the archive
+			bool bResult = TestArchive.AddEntry( TestNamespace, SourceAKey, SourceA, Translation, KeyMetadataA, false );
+
+			TestTrue( TEXT("AddEntry result = true"), bResult);
+
+			// We should only have one entry in the archive
+			int32 EntryCount = 0;
+			for(auto Iter = TestArchive.GetEntriesBySourceTextIterator(); Iter; ++Iter, ++EntryCount);
+			TestEqual( TEXT("EntryCount == 1"), EntryCount, 1 );
+
+			// Check to see that the original type/value of the * prefixed entry did not get modified
+			FoundEntry = TestArchive.FindEntryByKey( TestNamespace, SourceAKey, KeyMetadataA );
+			bFoundEntryValid = FoundEntry.IsValid();
+			TestTrue(TEXT("FArchiveEntry could not find entry using FindEntryByKey."), bFoundEntryValid);
+			if (bFoundEntryValid)
+			{
+				if( FoundEntry->Source.MetadataObj->HasTypedField< ELocMetadataType::String >( TEXT("*IsMature") ) )
+				{
+					TestTrue( TEXT("Metadata Type == String and Value == Empty string"), FoundEntry->Source.MetadataObj->GetStringField( TEXT("*IsMature") ) == TEXT("") );
+				}
+				else
+				{
+					TestTrue(TEXT("FArchiveEntry * prefixed metadata entry on source object was modified unexpectedly."), false);
+				}
+			}
+		}
+	}
+}
