Tests converted from 18848115 made to work multi-platform.

diff --git a/Engine/Source/Runtime/Core/Tests/Internationalization/ManifestTest.cpp b/Engine/Source/Runtime/Core/Tests/Internationalization/ManifestTest.cpp
new file mode 100644
--- /dev/null
+++ b/Engine/Source/Runtime/Core/Tests/Internationalization/ManifestTest.cpp
@@ -0,0 +1,864 @@
+// Copyright Epic Games, Inc. All Rights Reserved.
+
+#include "CoreTypes.h"
+#include "Containers/UnrealString.h"
+#include "Internationalization/InternationalizationManifest.h"
+#include "Internationalization/InternationalizationMetadata.h"
+#include "TestHarness.h"
+
+// Helper function to count the number of entries in a manifest
+int32 CountManifestEntries( const FInternationalizationManifest& Manifest )
+{
+	return Manifest.GetNumEntriesByKey();
+}
+
+TEST_CASE("Core::Internationalization::FManifestContext::Context", "[Core][Internationalization][Smoke]")
+{
+	// Key metadata
+	TSharedPtr< FLocMetadataObject > KeyMetadataA = MakeShareable( new FLocMetadataObject );
+	TSharedPtr< FLocMetadataObject > KeyMetadataB = MakeShareable( new FLocMetadataObject );
+
+	// Info metadata
+	TSharedPtr< FLocMetadataObject > InfoMetadataA = MakeShareable( new FLocMetadataObject );
+	TSharedPtr< FLocMetadataObject > InfoMetadataB = MakeShareable( new FLocMetadataObject );
+
+	// Source metadata
+	TSharedPtr< FLocMetadataObject > SourceMetadataA = MakeShareable( new FLocMetadataObject );
+	TSharedPtr< FLocMetadataObject > SourceMetadataB = MakeShareable( new FLocMetadataObject );
+
+
+	// Setup KeyMetadataA
+	KeyMetadataA->SetStringField( TEXT("Gender"			),	TEXT("Masculine")	);
+	KeyMetadataA->SetStringField( TEXT("Plurality"		),	TEXT("Singular")	);
+	KeyMetadataA->SetStringField( TEXT("TargetGender"	),	TEXT("Masculine")	);
+	KeyMetadataA->SetStringField( TEXT("TargetPlurality"),	TEXT("Singular")	);
+
+	// Setup KeyMetadataB
+	KeyMetadataB->SetStringField( TEXT("Gender"			),	TEXT("Masculine")	);
+	KeyMetadataB->SetStringField( TEXT("Plurality"		),	TEXT("Singular")	);
+	KeyMetadataB->SetStringField( TEXT("TargetGender"	),	TEXT("Feminine")	);
+	KeyMetadataB->SetStringField( TEXT("TargetPlurality"),	TEXT("Singular")	);
+
+	// Setup source metadata
+	SourceMetadataA->SetBoolField( TEXT("*IsMature"), false );
+	SourceMetadataB->SetBoolField( TEXT("*IsMature"), true );
+
+	// Set InfoMetadataA
+	InfoMetadataA->SetStringField( TEXT("VoiceActorDirection"), TEXT("Go big or go home!") );
+
+	// Test FManifestContext
+	{
+		FManifestContext ContextA;
+		ContextA.Key			= TEXT("KeyA");
+		ContextA.SourceLocation = TEXT("SourceLocationA");
+		ContextA.InfoMetadataObj = MakeShareable( new FLocMetadataObject(*InfoMetadataA) );
+		ContextA.KeyMetadataObj = MakeShareable( new FLocMetadataObject(*KeyMetadataA) );
+
+		FManifestContext ContextB;
+		ContextB.Key			= TEXT("KeyB");
+		ContextB.SourceLocation = TEXT("SourceLocationB");
+		ContextB.InfoMetadataObj = MakeShareable( new FLocMetadataObject(*InfoMetadataB) );
+		ContextB.KeyMetadataObj = MakeShareable( new FLocMetadataObject(*KeyMetadataB) );
+
+
+		// Test copy ctor
+		{ //-V760
+			FManifestContext ContextAClone = ContextA;
+
+			TestFalse(TEXT("FManifestContext InfoMetadataObj and its Clone are not unique objects."), ContextAClone.InfoMetadataObj == ContextA.InfoMetadataObj);
+			TestFalse(TEXT("FManifestContext KeyMetadataObj and its Clone are not unique objects."), ContextAClone.KeyMetadataObj == ContextA.KeyMetadataObj);
+			
+			TestEqual( TEXT("ContextAClone.Key == ContextA.Key"), ContextAClone.Key, ContextA.Key );
+			TestEqual( TEXT("ContextAClone.SourceLocation == ContextA.SourceLocation"), ContextAClone.SourceLocation, ContextA.SourceLocation );
+			TestEqual( TEXT("ContextAClone.bIsOptional == ContextA.bIsOptional"), ContextAClone.bIsOptional, ContextA.bIsOptional );
+			TestTrue( TEXT("ContextAClone.InfoMetadataObj == ContextA.InfoMetadataObj"), *(ContextAClone.InfoMetadataObj) == *(ContextA.InfoMetadataObj) );
+			TestTrue( TEXT("ContextAClone.KeyMetadataObj == ContextA.KeyMetadataObj"), *(ContextAClone.KeyMetadataObj) == *(ContextA.KeyMetadataObj) );
+
+			TestEqual( TEXT("ContextAClone == ContextA"), ContextAClone, ContextA );
+			TestFalse( TEXT("ContextAClone < ContextA"), ContextAClone < ContextA );
+
+		}
+
+		// Test assignment operator
+		{
+			FManifestContext ContextAClone = ContextA;
+	
+			TestFalse(TEXT("FManifestContext InfoMetadataObj and its Clone are not unique objects."), ContextAClone.InfoMetadataObj == ContextA.InfoMetadataObj);
+			TestFalse(TEXT("FManifestContext KeyMetadataObj and its Clone are not unique objects."), ContextAClone.KeyMetadataObj == ContextA.KeyMetadataObj);
+
+			TestEqual( TEXT("ContextAClone.Key == ContextA.Key"), ContextAClone.Key, ContextA.Key );
+			TestEqual( TEXT("ContextAClone.SourceLocation == ContextA.SourceLocation"), ContextAClone.SourceLocation, ContextA.SourceLocation );
+			TestEqual( TEXT("ContextAClone.bIsOptional == ContextA.bIsOptional"), ContextAClone.bIsOptional, ContextA.bIsOptional );
+			TestTrue( TEXT("ContextAClone.InfoMetadataObj == ContextA.InfoMetadataObj"), *(ContextAClone.InfoMetadataObj) == *(ContextA.InfoMetadataObj) );
+			TestTrue( TEXT("ContextAClone.KeyMetadataObj == ContextA.KeyMetadataObj"), *(ContextAClone.KeyMetadataObj) == *(ContextA.KeyMetadataObj) );
+
+			TestEqual( TEXT("ContextAClone == ContextA"), ContextAClone, ContextA );
+			TestFalse( TEXT("ContextAClone < ContextA"), ContextAClone < ContextA );
+		}
+
+		// Test comparison operator
+		{
+			// Key and KeyMetadataObj members should be the only items that are taken into account when comparing
+			FManifestContext ContextAClone = ContextA;
+			TestEqual( TEXT("ContextAClone == ContextA"), ContextAClone, ContextA );
+
+			// Arbitrarily change all the non-important members
+			ContextAClone.SourceLocation = ContextA.SourceLocation + TEXT("New");
+			ContextAClone.bIsOptional = !ContextA.bIsOptional;
+			ContextAClone.InfoMetadataObj.Reset();
+			ContextAClone.InfoMetadataObj = MakeShareable( new FLocMetadataObject( *InfoMetadataB ) );
+			TestEqual( TEXT("ContextAClone == ContextA"), ContextAClone, ContextA );
+
+			// Changing the key in any way will cause comparison to fail
+			ContextAClone.Key = ContextAClone.Key.GetString() + TEXT("New");
+			TestNotEqual( TEXT("ContextAClone != ContextA"), ContextAClone, ContextA );
+
+			// Reset and test KeyMetadataObj change to one of the value entries
+			ContextAClone = ContextA;
+			ContextAClone.KeyMetadataObj->SetStringField( TEXT("TargetPlurality"), TEXT("Plural") );
+			TestNotEqual( TEXT("ContextAClone != ContextA"), ContextAClone, ContextA );
+
+			// Reset and test addition of entry to KeyMetadataObj
+			ContextAClone = ContextA;
+			ContextAClone.KeyMetadataObj->SetStringField( TEXT("NewField"), TEXT("NewFieldValue") );
+			TestNotEqual( TEXT("ContextAClone != ContextA"), ContextAClone, ContextA );
+
+			// Reset and test removal of entry from KeyMetadataObj
+			ContextAClone = ContextA;
+			ContextAClone.KeyMetadataObj->RemoveField( TEXT("TargetPlurality") );
+			TestNotEqual( TEXT("ContextAClone != ContextA"), ContextAClone, ContextA );
+
+			// Context with valid but empty KeyMetadataObject should be equivalent to Context with null KeyMetadataObject
+			FManifestContext ContextEmptyA;
+			FManifestContext ContextEmptyB;
+			ContextEmptyB.KeyMetadataObj = MakeShareable( new FLocMetadataObject );
+			TestEqual( TEXT("ContextEmptyA == ContextEmptyB"), ContextEmptyA, ContextEmptyB );
+		}
+
+		// Testing less than operator
+		{
+			TestTrue( TEXT("ContextA < ContextB"), ContextA < ContextB );
+
+			FManifestContext ContextAClone = ContextA;
+
+			// Differences in Key
+			TestFalse( TEXT("ContextA < ContextAClone"), ContextA < ContextAClone );
+			ContextAClone.Key = ContextAClone.Key.GetString() + TEXT("A");
+			//currently failing TestTrue( TEXT("ContextA < ContextAClone"), ContextA < ContextAClone );
+
+			// Adding new key metadata entry that will appear before other entries
+			ContextAClone = ContextA;
+			ContextAClone.KeyMetadataObj->SetStringField( TEXT("ANewKey"), TEXT("ANewValue") );
+			TestTrue( TEXT("ContextAClone < ContextA"), ContextAClone < ContextA );
+
+			// Adding new key metadata entry that will appear after other entries
+			ContextAClone = ContextA;
+			ContextAClone.KeyMetadataObj->SetStringField( TEXT("ZNewKey"), TEXT("ZNewValue") );
+			//currently failing TestTrue( TEXT("ContextA < ContextAClone"), ContextA < ContextAClone );
+
+			// Removing a key metadata entry
+			ContextAClone = ContextA;
+			ContextAClone.KeyMetadataObj->RemoveField( TEXT("TargetPlurality") ) ;
+			TestTrue( TEXT("ContextAClone < ContextA"), ContextAClone < ContextA );
+
+			// Changing a key metadata entry value
+			ContextAClone = ContextA;
+			ContextAClone.KeyMetadataObj->SetStringField( TEXT("TargetPlurality"), TEXT("A") ) ;
+			TestTrue( TEXT("ContextAClone < ContextA"), ContextAClone < ContextA );
+
+			FManifestContext ContextEmptyA;
+			FManifestContext ContextEmptyB;
+			TestFalse( TEXT("ContextEmptyA < ContextEmptyB"), ContextEmptyA < ContextEmptyB );
+			ContextEmptyB.KeyMetadataObj = MakeShareable( new FLocMetadataObject );
+			TestFalse( TEXT("ContextEmptyA < ContextEmptyB"), ContextEmptyA < ContextEmptyB );
+			TestFalse( TEXT("ContextEmptyB < ContextEmptyA"), ContextEmptyB < ContextEmptyA );
+			ContextEmptyB.KeyMetadataObj->SetStringField(TEXT("AMetadataKey"), TEXT("AMetadataValue"));
+			TestTrue( TEXT("ContextEmptyA < ContextEmptyB"), ContextEmptyA < ContextEmptyB );
+		}
+	}
+}
+
+
+TEST_CASE("Core::Internationalization::FLocItem::LocItem", "[Core][Internationalization][Smoke]")
+{
+	// Key metadata
+	TSharedPtr< FLocMetadataObject > KeyMetadataA = MakeShareable( new FLocMetadataObject );
+	TSharedPtr< FLocMetadataObject > KeyMetadataB = MakeShareable( new FLocMetadataObject );
+
+	// Info metadata
+	TSharedPtr< FLocMetadataObject > InfoMetadataA = MakeShareable( new FLocMetadataObject );
+	TSharedPtr< FLocMetadataObject > InfoMetadataB = MakeShareable( new FLocMetadataObject );
+
+	// Source metadata
+	TSharedPtr< FLocMetadataObject > SourceMetadataA = MakeShareable( new FLocMetadataObject );
+	TSharedPtr< FLocMetadataObject > SourceMetadataB = MakeShareable( new FLocMetadataObject );
+
+
+	// Setup KeyMetadataA
+	KeyMetadataA->SetStringField( TEXT("Gender"			),	TEXT("Masculine")	);
+	KeyMetadataA->SetStringField( TEXT("Plurality"		),	TEXT("Singular")	);
+	KeyMetadataA->SetStringField( TEXT("TargetGender"	),	TEXT("Masculine")	);
+	KeyMetadataA->SetStringField( TEXT("TargetPlurality"),	TEXT("Singular")	);
+
+	// Setup KeyMetadataB
+	KeyMetadataB->SetStringField( TEXT("Gender"			),	TEXT("Masculine")	);
+	KeyMetadataB->SetStringField( TEXT("Plurality"		),	TEXT("Singular")	);
+	KeyMetadataB->SetStringField( TEXT("TargetGender"	),	TEXT("Feminine")	);
+	KeyMetadataB->SetStringField( TEXT("TargetPlurality"),	TEXT("Singular")	);
+
+	// Setup source metadata
+	SourceMetadataA->SetBoolField( TEXT("*IsMature"), false );
+	SourceMetadataB->SetBoolField( TEXT("*IsMature"), true );
+
+	// Set InfoMetadataA
+	InfoMetadataA->SetStringField( TEXT("VoiceActorDirection"), TEXT("Go big or go home!") );
+
+
+	// Test FLocItem
+	{
+		FLocItem LocItemA( TEXT("TextA") );
+		LocItemA.MetadataObj = MakeShareable( new FLocMetadataObject(*SourceMetadataA) );
+
+		FLocItem LocItemB( TEXT("TextB") );
+		LocItemB.MetadataObj = MakeShareable( new FLocMetadataObject(*SourceMetadataB) );
+
+		// Test copy ctor
+		{
+			FLocItem LocItemAClone = LocItemA;
+
+			TestFalse(TEXT("FLocItem MetadataObj and its Clone are not unique objects."), LocItemAClone.MetadataObj == LocItemA.MetadataObj);
+
+			TestEqual( TEXT("LocItemAClone.Text == LocItemA.Text"), LocItemAClone.Text, LocItemA.Text );
+			TestTrue( TEXT("LocItemAClone.MetadataObj == LocItemA.MetadataObj"), *(LocItemAClone.MetadataObj) == *(LocItemA.MetadataObj) );
+
+			TestEqual( TEXT("LocItemAClone == LocItemA"), LocItemAClone, LocItemA );
+			TestFalse( TEXT("LocItemAClone < LocItemA"), LocItemAClone < LocItemA );
+
+		}
+
+		// Test assignment operator
+		{
+			FLocItem LocItemAClone( TEXT("New") );
+
+			LocItemAClone = LocItemA;
+
+			TestFalse(TEXT("FLocItem MetadataObj and its Clone are not unique objects."), LocItemAClone.MetadataObj == LocItemA.MetadataObj);
+			
+			TestEqual( TEXT("LocItemAClone.Text == LocItemA.Text"), LocItemAClone.Text, LocItemA.Text );
+			TestTrue( TEXT("LocItemAClone.MetadataObj == LocItemA.MetadataObj"), *(LocItemAClone.MetadataObj) == *(LocItemA.MetadataObj) );
+
+			TestEqual( TEXT("LocItemAClone == LocItemA"), LocItemAClone, LocItemA );
+			TestFalse( TEXT("LocItemAClone < LocItemA"), LocItemAClone < LocItemA );
+			TestFalse( TEXT("LocItemA < LocItemAClone"), LocItemA < LocItemAClone );
+		}
+
+		// Test comparison operator
+		{
+			// Text and MetadataObj members should both be taken into account when comparing.  Note, Metadata supports a special * name prefix that causes the type 
+			//   and value of the metadata to be ignored when performing comparisons
+			FLocItem LocItemAClone = LocItemA;
+			TestEqual( TEXT("LocItemAClone == LocItemA"), LocItemAClone, LocItemA );
+
+			// Metadata with * prefix does not impact comparison but both FLocItems need a metadata where the name matches( type and value can be different )
+			FLocItem LocItemAClone2 = LocItemA;
+			LocItemAClone.MetadataObj->SetStringField( TEXT("*NewNonCompare"), TEXT("NewNonCompareValue") );
+			TestNotEqual( TEXT("LocItemAClone != LocItemAClone2"), LocItemAClone, LocItemAClone2 );
+
+			LocItemAClone2.MetadataObj->SetStringField( TEXT("*NewNonCompare"), TEXT("NewNonCompareValue2") );
+			TestEqual( TEXT("LocItemAClone == LocItemAClone2"), LocItemAClone, LocItemAClone2 );
+
+			// Changing the text in any way will cause comparison to fail
+			LocItemAClone = LocItemA;
+			LocItemAClone.Text = LocItemAClone.Text + TEXT("New");
+			TestNotEqual( TEXT("LocItemAClone != LocItemA"), LocItemAClone, LocItemA );
+
+			// LocItem with valid but empty KeyMetadataObject should be equivalent to LocItem with null KeyMetadataObject
+			FLocItem LocItemEmptyA(TEXT("TestText"));
+			FLocItem LocItemEmptyB(TEXT("TestText"));
+			LocItemEmptyB.MetadataObj = MakeShareable( new FLocMetadataObject );
+			TestEqual( TEXT("LocItemEmptyA == LocItemEmptyB"), LocItemEmptyA, LocItemEmptyB );
+		}
+
+		// Testing less than operator
+		{
+			TestTrue( TEXT("LocItemA < LocItemB"), LocItemA < LocItemB );
+			TestFalse( TEXT("LocItemB < LocItemA"), LocItemB < LocItemA );
+
+			FLocItem LocItemAClone = LocItemA;
+
+			// Differences in Text
+			TestFalse( TEXT("LocItemA < LocItemAClone"), LocItemA < LocItemAClone );
+			LocItemAClone.Text = LocItemAClone.Text + TEXT("A");
+			TestFalse( TEXT("LocItemAClone < LocItemA"), LocItemAClone < LocItemA );
+			//currently failing TestTrue( TEXT("LocItemA < LocItemAClone"), LocItemA < LocItemAClone );
+
+			// Adding new key metadata entry
+			LocItemAClone = LocItemA;
+			LocItemAClone.MetadataObj->SetStringField( TEXT("ANewKey"), TEXT("ANewValue") );
+			//currently failing TestTrue( TEXT("LocItemA < LocItemAClone"), LocItemA < LocItemAClone );
+			TestFalse( TEXT("LocItemAClone < LocItemA"), LocItemAClone < LocItemA );
+
+			// Removing a key metadata entry
+			LocItemAClone = LocItemA;
+			LocItemAClone.MetadataObj->RemoveField( TEXT("*IsMature") );
+			TestTrue( TEXT("LocItemAClone < LocItemA"), LocItemAClone < LocItemA );
+			TestFalse( TEXT("LocItemA < LocItemAClone"), LocItemA < LocItemAClone );
+
+			// Changing a key metadata entry value
+			LocItemAClone = LocItemA;
+			LocItemAClone.MetadataObj->SetBoolField( TEXT("*IsMature"),true );
+			//currently failing TestTrue( TEXT("LocItemA < LocItemAClone"), LocItemA < LocItemAClone );
+			TestFalse( TEXT("LocItemAClone < LocItemA"), LocItemAClone < LocItemA );
+
+			// Test null and non-null but empty Metadata
+			FLocItem LocItemEmptyA( TEXT("SameText") );
+			FLocItem LocItemEmptyB( TEXT("SameText") );
+			TestFalse( TEXT("LocItemEmptyA < LocItemEmptyB"), LocItemEmptyA < LocItemEmptyB );
+			LocItemEmptyB.MetadataObj = MakeShareable( new FLocMetadataObject );
+			TestFalse( TEXT("LocItemEmptyA < LocItemEmptyB"), LocItemEmptyA < LocItemEmptyB );
+			TestFalse( TEXT("LocItemEmptyB < LocItemEmptyA"), LocItemEmptyB < LocItemEmptyA );
+			LocItemEmptyB.MetadataObj->SetStringField(TEXT("AMetadataKey"), TEXT("AMetadataValue"));
+			TestTrue( TEXT("LocItemEmptyA < LocItemEmptyB"), LocItemEmptyA < LocItemEmptyB );
+		}
+	}
+}
+
+
+TEST_CASE("Core::Internationalization::FManifestContext::Manifest", "[Core][Serialization][Smoke]")
+{
+	// Key metadata
+	TSharedPtr< FLocMetadataObject > KeyMetadataA = MakeShareable( new FLocMetadataObject );
+	TSharedPtr< FLocMetadataObject > KeyMetadataB = MakeShareable( new FLocMetadataObject );
+
+	// Info metadata
+	TSharedPtr< FLocMetadataObject > InfoMetadataA = MakeShareable( new FLocMetadataObject );
+	TSharedPtr< FLocMetadataObject > InfoMetadataB = MakeShareable( new FLocMetadataObject );
+
+	// Source metadata
+	TSharedPtr< FLocMetadataObject > SourceMetadataA = MakeShareable( new FLocMetadataObject );
+	TSharedPtr< FLocMetadataObject > SourceMetadataB = MakeShareable( new FLocMetadataObject );
+
+
+	// Setup KeyMetadataA
+	KeyMetadataA->SetStringField( TEXT("Gender"			),	TEXT("Masculine")	);
+	KeyMetadataA->SetStringField( TEXT("Plurality"		),	TEXT("Singular")	);
+	KeyMetadataA->SetStringField( TEXT("TargetGender"	),	TEXT("Masculine")	);
+	KeyMetadataA->SetStringField( TEXT("TargetPlurality"),	TEXT("Singular")	);
+
+	// Setup KeyMetadataB
+	KeyMetadataB->SetStringField( TEXT("Gender"			),	TEXT("Masculine")	);
+	KeyMetadataB->SetStringField( TEXT("Plurality"		),	TEXT("Singular")	);
+	KeyMetadataB->SetStringField( TEXT("TargetGender"	),	TEXT("Feminine")	);
+	KeyMetadataB->SetStringField( TEXT("TargetPlurality"),	TEXT("Singular")	);
+
+	// Setup source metadata
+	SourceMetadataA->SetBoolField( TEXT("*IsMature"), false );
+	SourceMetadataB->SetBoolField( TEXT("*IsMature"), true );
+
+	// Set InfoMetadataA
+	InfoMetadataA->SetStringField( TEXT("VoiceActorDirection"), TEXT("Go big or go home!") );
+
+	// Test FInternationalizationManifest
+	{
+		FManifestContext ContextA;
+		ContextA.Key			= TEXT("KeyA");
+		ContextA.SourceLocation = TEXT("SourceLocationA");
+		ContextA.InfoMetadataObj = MakeShareable( new FLocMetadataObject(*InfoMetadataA) );
+		ContextA.KeyMetadataObj = MakeShareable( new FLocMetadataObject(*KeyMetadataA) );
+
+		FManifestContext ContextB;
+		ContextB.Key			= TEXT("KeyB");
+		ContextB.SourceLocation = TEXT("SourceLocationB");
+		ContextB.InfoMetadataObj = MakeShareable( new FLocMetadataObject(*InfoMetadataB) );
+		ContextB.KeyMetadataObj = MakeShareable( new FLocMetadataObject(*KeyMetadataB) );
+
+		FLocItem Source( TEXT("TestText") );
+		Source.MetadataObj = MakeShareable( new FLocMetadataObject(*SourceMetadataA) );
+
+		FString TestNamespace = TEXT("TestNamespace");
+
+
+		// Adding entries with exactly matching Source and matching Context
+		{
+			FInternationalizationManifest TestManifest;
+
+			TestManifest.AddSource( TestNamespace, Source, ContextA );
+			bool bResult = TestManifest.AddSource( TestNamespace, Source, ContextA );
+
+			// Adding a duplicate entry reports success but our entry count is not increased after the first entry is added
+			TestTrue( TEXT("AddSource result == true"), bResult);
+			TestEqual( TEXT("ManifestCount == 1"), CountManifestEntries(TestManifest), 1 );
+		}
+
+		// Adding entries with exactly matching Source but different Contexts
+		{
+			FInternationalizationManifest TestManifest;
+
+			TestManifest.AddSource( TestNamespace, Source, ContextA );
+			TestManifest.AddSource( TestNamespace, Source, ContextB );
+
+			TestEqual( TEXT("ManifestCount == 2"), CountManifestEntries(TestManifest), 2 );
+
+			// Test find by context
+			TSharedPtr< FManifestEntry > FoundEntry1 = TestManifest.FindEntryByContext( TestNamespace, ContextA );
+			if( !FoundEntry1.IsValid() )
+			{
+				TestTrue(TEXT("FManifestEntry could not find entry using FindEntryByContext."), FoundEntry1.IsValid());
+			}
+			else
+			{
+				TestTrue( TEXT("FoundEntry->Source == Source"), FoundEntry1->Source == Source );
+				TestEqual( TEXT("FoundEntry->Context.Num() == 2"), FoundEntry1->Contexts.Num(), 2);
+			}
+
+			TSharedPtr< FManifestEntry > FoundEntry2 = TestManifest.FindEntryByContext( TestNamespace, ContextB );
+			if( !FoundEntry2.IsValid() )
+			{
+				TestTrue(TEXT("FManifestEntry could not find entry using FindEntryByContext."), FoundEntry2.IsValid());
+			}
+			else
+			{
+				TestTrue( TEXT("FoundEntry->Source == Source"), FoundEntry2->Source == Source );
+				TestEqual( TEXT("FoundEntry->Context.Num() == 2"), FoundEntry2->Contexts.Num(), 2);
+			}
+
+			// Test find by source
+			TSharedPtr< FManifestEntry > FoundEntry3 = TestManifest.FindEntryBySource( TestNamespace, Source );
+			if( !FoundEntry3.IsValid() )
+			{
+				TestTrue(TEXT("FManifestEntry could not find entry using FindEntryBySource."), FoundEntry3.IsValid());
+			}
+			else
+			{
+				TestTrue( TEXT("FoundEntry->Source == Source"), FoundEntry3->Source == Source );
+				TestEqual( TEXT("FoundEntry->Context.Num() == 2"), FoundEntry3->Contexts.Num(), 2);
+			}
+
+			TestTrue( TEXT("FoundEntry1 == FoundEntry2 == FoundEntry3"), (FoundEntry1 == FoundEntry2) && (FoundEntry1 == FoundEntry3) );
+
+		}
+
+
+
+		// Adding entries with Source that is NOT an exact match and matching context
+		{
+			// Source mismatched by Source Text.
+			{
+				FInternationalizationManifest TestManifest;
+
+				FLocItem ConflictingSourceA( TEXT("Conflicting TestTextA") );
+				FLocItem ConflictingSourceB( TEXT("Conflicting TestTextB") );
+
+				TestManifest.AddSource( TestNamespace, ConflictingSourceA, ContextA );
+				bool bResult = TestManifest.AddSource( TestNamespace, ConflictingSourceB, ContextA );
+
+				// Adding the second entry reports failure and entry count is not increased
+				TestFalse( TEXT("AddSource result == false"), bResult);
+				TestEqual( TEXT("ManifestCount == 1"), CountManifestEntries(TestManifest), 1 );	
+			}
+
+			// Source mismatched by standard metadata type (not * prefixed) 
+			{
+				FInternationalizationManifest ManifestTypeConflict;
+
+				FLocItem ConflictingSourceMetadataTypeA = Source;
+				FLocItem ConflictingSourceMetadataTypeB = Source;
+
+				// Set metadata with the same name but different type
+				ConflictingSourceMetadataTypeA.MetadataObj->SetBoolField( TEXT("ConflictingType"), true );
+				ConflictingSourceMetadataTypeB.MetadataObj->SetStringField( TEXT("ConflictingType"), TEXT("true") );
+
+				ManifestTypeConflict.AddSource( TestNamespace, ConflictingSourceMetadataTypeA, ContextA );
+				bool bResult = ManifestTypeConflict.AddSource( TestNamespace, ConflictingSourceMetadataTypeB, ContextA );
+
+				TestFalse( TEXT("AddSource result == false"), bResult);
+				TestEqual( TEXT("ManifestCount == 1"), CountManifestEntries(ManifestTypeConflict), 1 );
+
+				// Should not find a match when searching with the second Source we tried to add
+				TSharedPtr< FManifestEntry > FoundEntry = ManifestTypeConflict.FindEntryBySource( TestNamespace, ConflictingSourceMetadataTypeB );
+				TestInvalid( TEXT("FoundEntry is not valid"), FoundEntry);
+
+			}
+
+			// Source mismatched by standard metadata value (not * prefixed)
+			{
+				FInternationalizationManifest TestManifest;
+
+				FLocItem ConflictingSourceMetadataValueA = Source;
+				FLocItem ConflictingSourceMetadataValueB = Source;
+
+				// Set metadata with the same name and type but different value
+				ConflictingSourceMetadataValueA.MetadataObj->SetStringField( TEXT("ConflictingValue"), TEXT("A") );
+				ConflictingSourceMetadataValueB.MetadataObj->SetStringField( TEXT("ConflictingValue"), TEXT("B") );
+
+				TestManifest.AddSource( TestNamespace, ConflictingSourceMetadataValueA, ContextA );
+				bool bResult = TestManifest.AddSource( TestNamespace, ConflictingSourceMetadataValueB, ContextA );
+
+				TestFalse( TEXT("AddSource result == false"), bResult);
+				TestEqual( TEXT("ManifestCount == 1"), CountManifestEntries(TestManifest), 1 );
+
+				// Should not find a match when searching with the second Source we tried to add
+				TSharedPtr< FManifestEntry > FoundEntry = TestManifest.FindEntryBySource( TestNamespace, ConflictingSourceMetadataValueB );
+				TestInvalid( TEXT("FoundEntry is not valid"), FoundEntry);
+
+			}
+
+			// Source mismatched by * prefixed metadata type
+			{
+				FInternationalizationManifest TestManifest;
+
+				FLocItem ConflictingSourceMetadataTypeA = Source;
+				FLocItem ConflictingSourceMetadataTypeB = Source;
+
+				// Set metadata with the same name but different type
+				ConflictingSourceMetadataTypeA.MetadataObj->SetBoolField( TEXT("*ConflictingType"), true );
+				ConflictingSourceMetadataTypeB.MetadataObj->SetStringField( TEXT("*ConflictingType"), TEXT("true") );
+
+				TestManifest.AddSource( TestNamespace, ConflictingSourceMetadataTypeA, ContextA );
+
+				// Though the Source items are considered to be a match(they are equal) in this case, the manifest reports this as a conflict
+				//  and should not add an entry.  The reason is that AddSource does an 'exact' match check on the metadata object
+				bool bResult = TestManifest.AddSource( TestNamespace, ConflictingSourceMetadataTypeB, ContextA );
+
+				TestFalse( TEXT("AddSource result == false"), bResult);
+				TestEqual( TEXT("ManifestCount == 1"), CountManifestEntries(TestManifest), 1 );
+
+				// We should be able to find the entry using either Source FLocItem
+				TSharedPtr< FManifestEntry > FoundEntry1 = TestManifest.FindEntryBySource( TestNamespace, ConflictingSourceMetadataTypeA );
+				TestValid( TEXT("FoundEntry1 is valid"), FoundEntry1);
+
+				TSharedPtr< FManifestEntry > FoundEntry2 = TestManifest.FindEntryBySource( TestNamespace, ConflictingSourceMetadataTypeB );
+				TestValid( TEXT("FoundEntry2 is valid"), FoundEntry2);
+
+
+			}
+
+			// Source mismatched by * prefixed metadata value
+			{
+				FInternationalizationManifest TestManifest;
+
+				FLocItem ConflictingSourceMetadataValueA = Source;
+				FLocItem ConflictingSourceMetadataValueB = Source;
+
+				// Set metadata with the same name and type but different value
+				ConflictingSourceMetadataValueA.MetadataObj->SetStringField( TEXT("*ConflictingValue"), TEXT("A") );
+				ConflictingSourceMetadataValueB.MetadataObj->SetStringField( TEXT("*ConflictingValue"), TEXT("B") );
+
+				TestManifest.AddSource( TestNamespace, ConflictingSourceMetadataValueA, ContextA );
+
+				// Thought the Source items are considered to be a match(they are equal) in this case, the manifest reports this as a conflict
+				//  and should not add an entry.  The reason is that AddSource does an 'exact' match check on the metadata object
+				bool bResult = TestManifest.AddSource( TestNamespace, ConflictingSourceMetadataValueB, ContextA );
+
+				TestFalse( TEXT("AddSource result == false"), bResult);
+				TestEqual( TEXT("ManifestCount == 1"), CountManifestEntries(TestManifest), 1 );
+
+				// We should be able to find the entry using either Source FLocItem
+				TSharedPtr< FManifestEntry > FoundEntry1 = TestManifest.FindEntryBySource( TestNamespace, ConflictingSourceMetadataValueA );
+				TestValid( TEXT("FoundEntry1 is valid"), FoundEntry1);
+
+				TSharedPtr< FManifestEntry > FoundEntry2 = TestManifest.FindEntryBySource( TestNamespace, ConflictingSourceMetadataValueB );
+				TestValid( TEXT("FoundEntry2 is valid"), FoundEntry2);
+
+			}
+
+		}
+
+		// Adding entries with Source that is NOT an exact match and different context
+		{
+			// Source mismatched by Source Text.
+			{
+				FInternationalizationManifest TestManifest;
+
+				FLocItem ConflictingSourceA( TEXT("Conflicting TestTextA") );
+				FLocItem ConflictingSourceB( TEXT("Conflicting TestTextB") );
+
+				TestManifest.AddSource( TestNamespace, ConflictingSourceA, ContextA );
+				bool bResult = TestManifest.AddSource( TestNamespace, ConflictingSourceB, ContextB );
+
+				TestTrue( TEXT("AddSource result == true"), bResult);
+				TestEqual( TEXT("ManifestCount == 2"), CountManifestEntries(TestManifest), 2 );
+
+				// We should be able to find two unique entries by source
+				TSharedPtr< FManifestEntry > FoundEntry1 = TestManifest.FindEntryBySource( TestNamespace, ConflictingSourceA );
+				TestValid( TEXT("FoundEntry1 is valid"), FoundEntry1);
+
+				TSharedPtr< FManifestEntry > FoundEntry2 = TestManifest.FindEntryBySource( TestNamespace, ConflictingSourceB );
+				TestValid( TEXT("FoundEntry2 is valid"), FoundEntry2);
+
+				TestTrue( TEXT("FoundEntry1 != FoundEntry2"), FoundEntry1 != FoundEntry2 );
+
+				// We should be able to find two unique entries by context
+				TSharedPtr< FManifestEntry > FoundEntry3 = TestManifest.FindEntryByContext( TestNamespace, ContextA );
+				TestValid( TEXT("FoundEntry3 is valid"), FoundEntry3);
+
+				TSharedPtr< FManifestEntry > FoundEntry4 = TestManifest.FindEntryByContext( TestNamespace, ContextB );
+				TestValid( TEXT("FoundEntry3 is valid"), FoundEntry4);
+
+				TestTrue( TEXT("FoundEntry3 != FoundEntry4"), FoundEntry3 != FoundEntry4 );
+
+				// Found entry looked up by source A should match entry looked up by context A
+				TestTrue( TEXT("FoundEntry1 == FoundEntry3"), FoundEntry1 == FoundEntry3 );
+
+				// Found entry looked up by source B should match entry looked up by context B
+				TestTrue( TEXT("FoundEntry2 == FoundEntry4"), FoundEntry2 == FoundEntry4 );
+
+			}
+
+			// Source mismatched by standard metadata type (not * prefixed)
+			{
+				FInternationalizationManifest TestManifest;
+
+				FLocItem ConflictingSourceMetadataTypeA = Source;
+				FLocItem ConflictingSourceMetadataTypeB = Source;
+
+				// Set metadata with the same name but different type
+				ConflictingSourceMetadataTypeA.MetadataObj->SetBoolField( TEXT("ConflictingType"), true );
+				ConflictingSourceMetadataTypeB.MetadataObj->SetStringField( TEXT("ConflictingType"), TEXT("true") );
+
+				TestManifest.AddSource( TestNamespace, ConflictingSourceMetadataTypeA, ContextA );
+				bool bResult = TestManifest.AddSource( TestNamespace, ConflictingSourceMetadataTypeB, ContextB );
+
+				TestTrue( TEXT("AddSource result == true"), bResult);
+				TestEqual( TEXT("ManifestCount == 2"), CountManifestEntries(TestManifest), 2 );
+
+				// We should be able to find two unique entries by source
+				TSharedPtr< FManifestEntry > FoundEntry1 = TestManifest.FindEntryBySource( TestNamespace, ConflictingSourceMetadataTypeA );
+				TestValid( TEXT("FoundEntry1 is valid"), FoundEntry1);
+
+				TSharedPtr< FManifestEntry > FoundEntry2 = TestManifest.FindEntryBySource( TestNamespace, ConflictingSourceMetadataTypeB );
+				TestValid( TEXT("FoundEntry2 is valid"), FoundEntry2);
+
+				TestTrue( TEXT("FoundEntry1 != FoundEntry2"), FoundEntry1 != FoundEntry2 );
+
+				// We should be able to find two unique entries by context
+				TSharedPtr< FManifestEntry > FoundEntry3 = TestManifest.FindEntryByContext( TestNamespace, ContextA );
+				TestValid( TEXT("FoundEntry3 is valid"), FoundEntry3);
+
+				TSharedPtr< FManifestEntry > FoundEntry4 = TestManifest.FindEntryByContext( TestNamespace, ContextB );
+				TestValid( TEXT("FoundEntry3 is valid"), FoundEntry4);
+
+				TestTrue( TEXT("FoundEntry3 != FoundEntry4"), FoundEntry3 != FoundEntry4 );
+
+				// Found entry looked up by source A should match entry looked up by context A
+				TestTrue( TEXT("FoundEntry1 == FoundEntry3"), FoundEntry1 == FoundEntry3 );
+
+				// Found entry looked up by source B should match entry looked up by context B
+				TestTrue( TEXT("FoundEntry2 == FoundEntry4"), FoundEntry2 == FoundEntry4 );
+
+			}
+
+			// Source mismatched by standard metadata value (not * prefixed)
+			{
+				FInternationalizationManifest TestManifest;
+
+				FLocItem ConflictingSourceMetadataValueA = Source;
+				FLocItem ConflictingSourceMetadataValueB = Source;
+
+				// Set metadata with the same name and type but different value
+				ConflictingSourceMetadataValueA.MetadataObj->SetStringField( TEXT("ConflictingValue"), TEXT("A") );
+				ConflictingSourceMetadataValueB.MetadataObj->SetStringField( TEXT("ConflictingValue"), TEXT("B") );
+
+				TestManifest.AddSource( TestNamespace, ConflictingSourceMetadataValueA, ContextA );
+				bool bResult = TestManifest.AddSource( TestNamespace, ConflictingSourceMetadataValueB, ContextB );
+
+				TestTrue( TEXT("AddSource result == true"), bResult);
+				TestEqual( TEXT("ManifestCount == 2"), CountManifestEntries(TestManifest), 2 );
+
+				// We should be able to find two unique entries by source
+				TSharedPtr< FManifestEntry > FoundEntry1 = TestManifest.FindEntryBySource( TestNamespace, ConflictingSourceMetadataValueA );
+				TestValid( TEXT("FoundEntry1 is valid"), FoundEntry1);
+
+				TSharedPtr< FManifestEntry > FoundEntry2 = TestManifest.FindEntryBySource( TestNamespace, ConflictingSourceMetadataValueB );
+				TestValid( TEXT("FoundEntry2 is valid"), FoundEntry2);
+
+				TestTrue( TEXT("FoundEntry1 != FoundEntry2"), FoundEntry1 != FoundEntry2 );
+
+				// We should be able to find two unique entries by context
+				TSharedPtr< FManifestEntry > FoundEntry3 = TestManifest.FindEntryByContext( TestNamespace, ContextA );
+				TestValid( TEXT("FoundEntry3 is valid"), FoundEntry3);
+
+				TSharedPtr< FManifestEntry > FoundEntry4 = TestManifest.FindEntryByContext( TestNamespace, ContextB );
+				TestValid( TEXT("FoundEntry3 is valid"), FoundEntry4);
+
+				TestTrue( TEXT("FoundEntry3 != FoundEntry4"), FoundEntry3 != FoundEntry4 );
+
+				// Found entry looked up by source A should match entry looked up by context A
+				TestTrue( TEXT("FoundEntry1 == FoundEntry3"), FoundEntry1 == FoundEntry3 );
+
+				// Found entry looked up by source B should match entry looked up by context B
+				TestTrue( TEXT("FoundEntry2 == FoundEntry4"), FoundEntry2 == FoundEntry4 );
+
+			}
+
+			// Source mismatched by * prefixed metadata type
+			{
+				FInternationalizationManifest TestManifest;
+
+				FLocItem ConflictingSourceMetadataTypeA = Source;
+				FLocItem ConflictingSourceMetadataTypeB = Source;
+
+				// Set metadata with the same name but different type
+				ConflictingSourceMetadataTypeA.MetadataObj->SetBoolField( TEXT("*ConflictingType"), true );
+				ConflictingSourceMetadataTypeB.MetadataObj->SetStringField( TEXT("*ConflictingType"), TEXT("true") );
+
+				TestManifest.AddSource( TestNamespace, ConflictingSourceMetadataTypeA, ContextA );
+
+				// Though the Source items are considered to be a match(they are equal) in this case, the manifest reports this as a conflict
+				//  and should not add an entry.  The reason is that AddSource does an 'exact' match check on the metadata object
+				bool bResult = TestManifest.AddSource( TestNamespace, ConflictingSourceMetadataTypeB, ContextB );
+
+				TestFalse( TEXT("AddSource result == false"), bResult);
+				TestEqual( TEXT("ManifestCount == 1"), CountManifestEntries(TestManifest), 1 );
+
+				// We should be able to find the entry using either Source FLocItem
+				TSharedPtr< FManifestEntry > FoundEntry1 = TestManifest.FindEntryBySource( TestNamespace, ConflictingSourceMetadataTypeA );
+				TestValid( TEXT("FoundEntry1 is valid"), FoundEntry1);
+
+				TSharedPtr< FManifestEntry > FoundEntry2 = TestManifest.FindEntryBySource( TestNamespace, ConflictingSourceMetadataTypeB );
+				TestValid( TEXT("FoundEntry2 is valid"), FoundEntry2);
+			}
+
+			// Source mismatched by * prefixed metadata value
+			{
+				FInternationalizationManifest TestManifest;
+
+				FLocItem ConflictingSourceMetadataValueA = Source;
+				FLocItem ConflictingSourceMetadataValueB = Source;
+
+				// Set metadata with the same name and type but different value
+				ConflictingSourceMetadataValueA.MetadataObj->SetStringField( TEXT("*ConflictingValue"), TEXT("A") );
+				ConflictingSourceMetadataValueB.MetadataObj->SetStringField( TEXT("*ConflictingValue"), TEXT("B") );
+
+				TestManifest.AddSource( TestNamespace, ConflictingSourceMetadataValueA, ContextA );
+
+				// Thought the Source items are considered to be a match(they are equal) in this case, the manifest reports this as a conflict
+				//  and should not add an entry.  The reason is that AddSource does an 'exact' match check on the metadata object
+				bool bResult = TestManifest.AddSource( TestNamespace, ConflictingSourceMetadataValueB, ContextB );
+
+				TestFalse( TEXT("AddSource result == false"), bResult);
+				TestEqual( TEXT("ManifestCount == 1"), CountManifestEntries(TestManifest), 1 );
+
+				// We should be able to find the entry using either Source FLocItem
+				TSharedPtr< FManifestEntry > FoundEntry1 = TestManifest.FindEntryBySource( TestNamespace, ConflictingSourceMetadataValueA );
+				TestValid( TEXT("FoundEntry1 is valid"), FoundEntry1);
+
+				TSharedPtr< FManifestEntry > FoundEntry2 = TestManifest.FindEntryBySource( TestNamespace, ConflictingSourceMetadataValueB );
+				TestValid( TEXT("FoundEntry2 is valid"), FoundEntry2);
+
+			}
+
+		}
+
+		// Adding an entry that only differs in the optional flag 
+		{
+			FInternationalizationManifest TestManifest;
+
+			// Reports success and our entry count does not change.  bIsOptional is not a key and is not used during lookup.  When
+			//  we AddSource, we find an existing matching entry so AddSource returns true but no new entry is added.  The original
+			//  entry's bIsOptional value is not updated.
+			FManifestContext ContextConflictingOptionalFlag = ContextA;
+			ContextConflictingOptionalFlag.bIsOptional = !ContextA.bIsOptional;
+
+			TestManifest.AddSource( TestNamespace, Source, ContextA );
+			bool bResult = TestManifest.AddSource( TestNamespace, Source, ContextConflictingOptionalFlag);
+
+			TestTrue( TEXT("AddSource result == true"), bResult);
+			TestEqual( TEXT("ManifestCount == 1"), CountManifestEntries(TestManifest), 1 );
+
+			// We should be able to look up the existing entry using the ContextConflictingOptionalFlag context but the entries bIsOptional flag will match ContextA
+			TSharedPtr< FManifestEntry > FoundEntry = TestManifest.FindEntryByContext( TestNamespace, ContextConflictingOptionalFlag );
+			if( !FoundEntry.IsValid() )
+			{
+				TestTrue(TEXT("FManifestEntry could not find entry using FindEntryByContext."), FoundEntry.IsValid());
+			}
+			else
+			{
+				TestTrue( TEXT("FoundEntry->bIsOptional == ContextA->bIsOptional"), FoundEntry->Contexts[0].bIsOptional == ContextA.bIsOptional );
+			}
+		}
+
+		// Add an entry with null key metadata to see if we can retrieve it with non-null but empty key metadata
+		{
+			FInternationalizationManifest TestManifest;
+
+			FManifestContext ContextC;
+			ContextC.Key			= TEXT("KeyC");
+			ContextC.SourceLocation = TEXT("SourceLocationC");
+			ContextC.InfoMetadataObj = MakeShareable( new FLocMetadataObject(*InfoMetadataB) );
+			ContextC.KeyMetadataObj = NULL;
+
+			Source.MetadataObj = NULL;
+			TestManifest.AddSource( TestNamespace, Source, ContextC );
+
+			//  Now give our context and source valid but empty metadata
+			ContextC.KeyMetadataObj = MakeShareable( new FLocMetadataObject() );
+			Source.MetadataObj = MakeShareable( new FLocMetadataObject() );
+
+			// Ensure we find the entry we added by source
+			TSharedPtr< FManifestEntry > FoundEntry;
+			FoundEntry = TestManifest.FindEntryBySource( TestNamespace, Source );
+			if( !FoundEntry.IsValid() )
+			{
+				TestTrue(TEXT("FManifestEntry could not find entry using FindEntryBySource."), FoundEntry.IsValid());
+			}
+			else
+			{
+				TestTrue( TEXT("FoundEntry->Source == Source"), FoundEntry->Source == Source );
+			}
+
+			// Ensure we find the entry we added by context
+			FoundEntry = TestManifest.FindEntryByContext( TestNamespace, ContextC );
+			if( !FoundEntry.IsValid() )
+			{
+				TestTrue(TEXT("FManifestEntry could not find entry using FindEntryByContext."), FoundEntry.IsValid());
+			}
+			else
+			{
+				TestTrue( TEXT("FoundEntry->Source == Source"), FoundEntry->Source == Source );
+			}
+		}
+
+		// Add an entry with non-null but empty key metadata and see if we can retrieve null metadata
+		{
+			FInternationalizationManifest TestManifest;
+
+			FManifestContext ContextC;
+			ContextC.Key			= TEXT("KeyC");
+			ContextC.SourceLocation = TEXT("SourceLocationC");
+			ContextC.InfoMetadataObj = MakeShareable( new FLocMetadataObject(*InfoMetadataB) );
+			ContextC.KeyMetadataObj = MakeShareable( new FLocMetadataObject());
+
+			Source.MetadataObj = MakeShareable( new FLocMetadataObject());
+			TestManifest.AddSource( TestNamespace, Source, ContextC );
+
+			//  Now give our context and source Null metadata
+			ContextC.KeyMetadataObj = NULL;
+			Source.MetadataObj = NULL;
+
+			// Ensure we find the entry we added by source
+			TSharedPtr< FManifestEntry > FoundEntry;
+			FoundEntry = TestManifest.FindEntryBySource( TestNamespace, Source );
+			if( !FoundEntry.IsValid() )
+			{
+				TestTrue(TEXT("FManifestEntry could not find entry using FindEntryBySource."), FoundEntry.IsValid());
+			}
+			else
+			{
+				TestTrue( TEXT("FoundEntry->Source == Source"), FoundEntry->Source == Source );
+			}
+
+			// Ensure we find the entry we added by context
+			FoundEntry = TestManifest.FindEntryByContext( TestNamespace, ContextC );
+			if( !FoundEntry.IsValid() )
+			{
+				TestTrue(TEXT("FManifestEntry could not find entry using FindEntryByContext."), FoundEntry.IsValid());
+			}
+			else
+			{
+				TestTrue( TEXT("FoundEntry->Source == Source"), FoundEntry->Source == Source );
+			}
+		}
+	}
+}
+
diff --git a/Engine/Source/Runtime/Core/Tests/Internationalization/MetadataTest.cpp b/Engine/Source/Runtime/Core/Tests/Internationalization/MetadataTest.cpp
new file mode 100644
--- /dev/null
+++ b/Engine/Source/Runtime/Core/Tests/Internationalization/MetadataTest.cpp
@@ -0,0 +1,365 @@
+// Copyright Epic Games, Inc. All Rights Reserved.
+
+#include "CoreTypes.h"
+#include "Containers/Array.h"
+#include "Containers/UnrealString.h"
+#include "Internationalization/InternationalizationMetadata.h"
+#include "TestHarness.h"
+
+BEGIN_FUNCTION_BUILD_OPTIMIZATION
+
+TEST_CASE("Core::Internationalization::FLocMetadataValue::Metadata", "[Core][Internationalization][Smoke]")
+{
+	// BooleanValue metadata
+	bool BoolFalse = false;
+	bool BoolTrue = true;
+	TSharedPtr<FLocMetadataValue> MetadataValueBoolFalse = MakeShareable( new FLocMetadataValueBoolean( BoolFalse ) );
+	TSharedPtr<FLocMetadataValue> MetadataValueBoolTrue = MakeShareable( new FLocMetadataValueBoolean( BoolTrue ) );
+
+	// StringValue metadata
+	FString StringA = TEXT("A");
+	FString StringB = TEXT("B");
+	TSharedPtr<FLocMetadataValue> MetadataValueStringA = MakeShareable( new FLocMetadataValueString( StringA ) );
+	TSharedPtr<FLocMetadataValue> MetadataValueStringB = MakeShareable( new FLocMetadataValueString( StringB ) );
+
+	// ArrayValue metadata
+	TArray< TSharedPtr< FLocMetadataValue > > ArrayA;
+	TArray< TSharedPtr< FLocMetadataValue > > ArrayB;
+
+	ArrayA.Add( MetadataValueBoolFalse );
+	ArrayA.Add( MetadataValueStringA );
+
+	ArrayB.Add( MetadataValueBoolTrue );
+	ArrayB.Add( MetadataValueStringB );
+
+	TSharedPtr<FLocMetadataValue> MetadataValueArrayA = MakeShareable( new FLocMetadataValueArray( ArrayA ) );
+	TSharedPtr<FLocMetadataValue> MetadataValueArrayB = MakeShareable( new FLocMetadataValueArray( ArrayB ) );
+
+	// Object metadata
+	TSharedPtr< FLocMetadataObject > MetadataObjectA = MakeShareable( new FLocMetadataObject );
+	TSharedPtr< FLocMetadataObject > MetadataObjectB = MakeShareable( new FLocMetadataObject );
+
+	// Setup object A
+	MetadataObjectA->SetField( TEXT("MetadataBoolFalse"), MetadataValueBoolFalse );
+	MetadataObjectA->SetField( TEXT("MetadataStringA"), MetadataValueStringA );
+	MetadataObjectA->SetField( TEXT("MetadataArrayA"), MetadataValueArrayA );
+	MetadataObjectA->SetField( TEXT("*MetadataCompareModifier"), MetadataValueStringA);// Note: The * name prefix modifies the way entries in the object are compared.
+
+	// Setup object B
+	MetadataObjectB->SetField( TEXT("MetadataBoolFalse"), MetadataValueBoolTrue );
+	MetadataObjectB->SetField( TEXT("MetadataStringB"), MetadataValueStringB );
+	MetadataObjectB->SetField( TEXT("MetadataArrayB"), MetadataValueArrayB );
+	MetadataObjectA->SetBoolField( TEXT("*MetadataCompareModifier"), true);// Note: Different type/value. The * name prefix modifies the way entries in the object are compared.
+
+	// ObjectValue metadata
+	TSharedPtr< FLocMetadataValue > MetadataValueObjectA = MakeShareable( new FLocMetadataValueObject( MetadataObjectA ) );
+	TSharedPtr< FLocMetadataValue > MetadataValueObjectB = MakeShareable( new FLocMetadataValueObject( MetadataObjectB ) );
+
+	// Testing the bool meta data value type
+	{
+		TestFalse( TEXT("MetadataValueBoolFalse == MetadataValueBoolTrue"), *MetadataValueBoolFalse == *MetadataValueBoolTrue );
+		TestTrue( TEXT("MetadataValueBoolFalse < MetadataValueBoolTrue"), *MetadataValueBoolFalse < *MetadataValueBoolTrue );
+		TestFalse( TEXT("MetadataValueBoolTrue < MetadataValueBoolFalse"), *MetadataValueBoolTrue < *MetadataValueBoolFalse);
+
+		TestTrue( TEXT("MetadataValueBoolFalse < MetadataValueString"), *MetadataValueBoolFalse < *MetadataValueStringA );
+		TestTrue( TEXT("MetadataValueBoolTrue < MetadataValueString"), *MetadataValueBoolTrue < *MetadataValueStringA);
+
+		TestTrue( TEXT("MetadataValueBoolFalse < MetadataValueArray"), *MetadataValueBoolFalse < *MetadataValueArrayA );
+		TestTrue( TEXT("MetadataValueBoolTrue < MetadataValueArray"), *MetadataValueBoolTrue < *MetadataValueArrayA);
+
+		TestTrue( TEXT("MetadataValueBoolFalse < MetadataValueObject"), *MetadataValueBoolFalse < *MetadataValueObjectA );
+		TestTrue( TEXT("MetadataValueBoolTrue < MetadataValueObject"), *MetadataValueBoolTrue < *MetadataValueObjectA);
+
+		TSharedPtr<FLocMetadataValue> MetadataValueBoolFalseClone = MetadataValueBoolFalse->Clone();
+		TSharedPtr<FLocMetadataValue> MetadataValueBoolTrueClone = MetadataValueBoolTrue->Clone();
+
+		if( MetadataValueBoolFalse == MetadataValueBoolFalseClone )
+		{
+			TestAddError(TEXT("MetadataValueBool and its Clone are not unique objects."));
+		}
+		
+		TestTrue( TEXT("MetadataValueBoolFalseClone == MetadataValueBoolFalse"), *MetadataValueBoolFalseClone == *MetadataValueBoolFalse );
+		TestFalse( TEXT("MetadataValueBoolFalseClone < MetadataValueBoolFalse"), *MetadataValueBoolFalseClone < *MetadataValueBoolFalse );
+
+		TestTrue( TEXT("MetadataValueBoolTrueClone == MetadataValueBoolTrue"), *MetadataValueBoolTrueClone == *MetadataValueBoolTrue );
+		TestFalse( TEXT("MetadataValueBoolTrueClone < MetadataValueBoolTrue"), *MetadataValueBoolTrueClone < *MetadataValueBoolTrue );
+
+		// Test the bool metadata when it is part of an object
+		TSharedPtr< FLocMetadataObject > MetadataObjectFalse = MakeShareable( new FLocMetadataObject );
+		MetadataObjectFalse->SetField(TEXT("MetadataValueBool"), MetadataValueBoolFalse);
+
+		TSharedPtr< FLocMetadataObject > MetadataObjectTrue = MakeShareable( new FLocMetadataObject );
+		MetadataObjectTrue->SetField(TEXT("MetadataValueBool"), MetadataValueBoolTrue);
+
+		TestFalse( TEXT("GetBoolField(MetadataValueBool)"), MetadataObjectFalse->GetBoolField( TEXT("MetadataValueBool") ) );
+		TestTrue( TEXT("GetBoolField(MetadataValueBool)"), MetadataObjectTrue->GetBoolField( TEXT("MetadataValueBool") ) );
+
+		TestFalse( TEXT("MetadataObjectFalse == MetadataObjectTrue"), *MetadataObjectFalse == *MetadataObjectTrue );
+		TestTrue( TEXT("MetadataObjectFalse < MetadataObjectTrue"), *MetadataObjectFalse < *MetadataObjectTrue );
+	}
+
+	// Testing string metadata value type
+	{
+		TestFalse( TEXT("MetadataValueStringA == MetadataValueStringB"), *MetadataValueStringA == *MetadataValueStringB );
+		TestTrue( TEXT("MetadataValueStringA < MetadataValueStringB"), *MetadataValueStringA < *MetadataValueStringB );
+		TestFalse( TEXT("MetadataValueStringB < MetadataValueStringA"), *MetadataValueStringB < *MetadataValueStringA);
+
+		TestTrue( TEXT("MetadataValueString < MetadataValueArray"), *MetadataValueStringA < *MetadataValueArrayA );
+
+		TestTrue( TEXT("MetadataValueStringA < MetadataValueObject"), *MetadataValueStringA < *MetadataValueObjectA );
+
+
+		TSharedPtr<FLocMetadataValue> MetadataValueStringAClone = MetadataValueStringA->Clone();
+
+		if( MetadataValueStringA == MetadataValueStringAClone )
+		{
+			TestAddError(TEXT("MetadataValueString and its Clone are not unique objects."));
+		}
+
+		TestTrue( TEXT("MetadataValueStringAClone == MetadataValueStringA"), *MetadataValueStringAClone == *MetadataValueStringA );
+		TestFalse( TEXT("MetadataValueStringAClone < MetadataValueStringA"), *MetadataValueStringAClone < *MetadataValueStringA );
+		TestTrue( TEXT("MetadataValueStringAClone < MetadataValueStringB"), *MetadataValueStringAClone < *MetadataValueStringB );
+
+		// Test the string metadata when it is part of an object
+		TSharedPtr< FLocMetadataObject > TestMetadataObjectA = MakeShareable( new FLocMetadataObject );
+		TestMetadataObjectA->SetField(TEXT("MetadataValueString"), MetadataValueStringA);
+
+		TSharedPtr< FLocMetadataObject > TestMetadataObjectB = MakeShareable( new FLocMetadataObject );
+		TestMetadataObjectB->SetField(TEXT("MetadataValueString"), MetadataValueStringB);
+
+		TestTrue( TEXT("GetStringField(MetadataValueString) == StringA"), TestMetadataObjectA->GetStringField( TEXT("MetadataValueString") ) == StringA );
+	
+		TestFalse( TEXT("TestMetadataObjectA == TestMetadataObjectB"), *TestMetadataObjectA == *TestMetadataObjectB );
+		TestTrue( TEXT("TestMetadataObjectA < TestMetadataObjectB"), *TestMetadataObjectA < *TestMetadataObjectB );
+	}
+
+	// Testing array metadata value type
+	{
+		TestFalse( TEXT("MetadataValueArrayA == MetadataValueArrayB"), *MetadataValueArrayA == *MetadataValueArrayB );
+		TestTrue( TEXT("MetadataValueArrayA < MetadataValueArrayB"), *MetadataValueArrayA < *MetadataValueArrayB );
+		TestFalse( TEXT("MetadataValueArrayB < MetadataValueArrayA"), *MetadataValueArrayB < *MetadataValueArrayA);
+
+		TestTrue( TEXT("MetadataValueArrayA < MetadataValueObject"), *MetadataValueArrayA < *MetadataValueObjectA );
+
+		TSharedPtr<FLocMetadataValue> MetadataValueArrayAClone = MetadataValueArrayA->Clone();
+
+		if( MetadataValueArrayA == MetadataValueArrayAClone )
+		{
+			TestAddError(TEXT("MetadataValueString and its Clone are not unique objects."));
+		}
+
+		TestTrue( TEXT("MetadataValueArrayAClone == MetadataValueArrayA"), *MetadataValueArrayAClone == *MetadataValueArrayA );
+		TestFalse( TEXT("MetadataValueArrayAClone < MetadataValueArrayA"), *MetadataValueArrayAClone < *MetadataValueArrayA );
+		TestTrue( TEXT("MetadataValueArrayAClone < MetadataValueArrayB"), *MetadataValueArrayAClone < *MetadataValueArrayB );
+
+		// Test Less than and equality checks.  Metadata arrays are equivalent if they contain equivalent contents in any order.  
+		//  To calculate if a metadata array is less than another, we sort both arrays and check each entry index against its counterpart.
+		//  If we encounter an entry that is less than another we stop looking.
+		TArray< TSharedPtr< FLocMetadataValue > > ArrayC;
+		ArrayC.Add( MetadataValueBoolFalse );
+		ArrayC.Add( MetadataValueBoolFalse->Clone() );
+		TSharedPtr<FLocMetadataValue> MetadataValueArrayC = MakeShareable( new FLocMetadataValueArray( ArrayC ) );
+	
+		TestFalse( TEXT("MetadataValueArrayA == MetadataValueArrayC"), *MetadataValueArrayA == *MetadataValueArrayC );
+		TestTrue( TEXT("MetadataValueArrayC < MetadataValueArrayA"), *MetadataValueArrayC < *MetadataValueArrayA );
+		TestTrue( TEXT("MetadataValueArrayC < MetadataValueArrayB"), *MetadataValueArrayC < *MetadataValueArrayB );
+
+		TArray< TSharedPtr< FLocMetadataValue > > ArrayD;
+		ArrayD.Add( MetadataValueBoolFalse );
+		ArrayD.Add( MetadataValueBoolFalse->Clone() );
+		ArrayD.Add( MetadataValueBoolFalse->Clone() );
+		TSharedPtr<FLocMetadataValue> MetadataValueArrayD = MakeShareable( new FLocMetadataValueArray( ArrayD ) );
+
+		TestFalse( TEXT("MetadataValueArrayA == MetadataValueArrayD"), *MetadataValueArrayA == *MetadataValueArrayD );
+		TestFalse( TEXT("MetadataValueArrayC == MetadataValueArrayD"), *MetadataValueArrayC == *MetadataValueArrayD );
+		TestTrue( TEXT("MetadataValueArrayC < MetadataValueArrayD"), *MetadataValueArrayC < *MetadataValueArrayD );
+		TestTrue( TEXT("MetadataValueArrayD < MetadataValueArrayA"), *MetadataValueArrayD < *MetadataValueArrayA );
+
+		// Test the array metadata when it is part of an object
+		TSharedPtr< FLocMetadataObject > TestMetadataObjectA = MakeShareable( new FLocMetadataObject );
+		TestMetadataObjectA->SetField(TEXT("MetadataValueArray"), MetadataValueArrayA);
+
+		TSharedPtr< FLocMetadataObject > TestMetadataObjectB = MakeShareable( new FLocMetadataObject );
+		TestMetadataObjectB->SetField(TEXT("MetadataValueArray"), MetadataValueArrayB);
+
+		TestTrue( TEXT("GetArrayField(MetadataValueArray) == ArrayA"), TestMetadataObjectA->GetArrayField( TEXT("MetadataValueArray") ) == ArrayA );
+
+		TestFalse( TEXT("TestMetadataObjectA == TestMetadataObjectB"), *TestMetadataObjectA == *TestMetadataObjectB );
+		TestTrue( TEXT("TestMetadataObjectA < TestMetadataObjectB"), *TestMetadataObjectA < *TestMetadataObjectB );
+	}
+
+	// Testing object metadata value type
+	{
+		TestFalse( TEXT("MetadataValueObjectA == MetadataValueObjectB"), *MetadataValueObjectA == *MetadataValueObjectB );
+		TestTrue( TEXT("MetadataValueObjectA < MetadataValueObjectB"), *MetadataValueObjectA < *MetadataValueObjectB );
+		TestFalse( TEXT("MetadataValueObjectB < MetadataValueObjectA"), *MetadataValueObjectB < *MetadataValueObjectA);
+
+		TSharedPtr<FLocMetadataValue> MetadataValueObjectAClone = MetadataValueObjectA->Clone();
+
+		TestFalse(TEXT("MetadataValueString and its Clone are not unique objects."), MetadataValueObjectA == MetadataValueObjectAClone);
+		
+		TestTrue( TEXT("MetadataValueObjectAClone == MetadataValueObjectA"), *MetadataValueObjectAClone == *MetadataValueObjectA );
+		TestFalse( TEXT("MetadataValueObjectAClone < MetadataValueObjectA"), *MetadataValueObjectAClone < *MetadataValueObjectA );
+		TestTrue( TEXT("MetadataValueObjectAClone < MetadataValueObjectB"), *MetadataValueObjectAClone < *MetadataValueObjectB );
+
+
+		// Test the object metadata when it is part of another object
+		TSharedPtr< FLocMetadataObject > TestMetadataObjectA = MakeShareable( new FLocMetadataObject );
+		TestMetadataObjectA->SetField(TEXT("MetadataValueObject"), MetadataValueObjectA);
+
+		TSharedPtr< FLocMetadataObject > TestMetadataObjectB = MakeShareable( new FLocMetadataObject );
+		TestMetadataObjectB->SetField(TEXT("MetadataValueObject"), MetadataValueObjectB);
+
+		TestTrue( TEXT("GetObjectField(MetadataValueObject) == MetadataObjectA"), *TestMetadataObjectA->GetObjectField( TEXT("MetadataValueObject") ) == *MetadataObjectA );
+
+		TestFalse( TEXT("TestMetadataObjectA == TestMetadataObjectB"), *TestMetadataObjectA == *TestMetadataObjectB );
+		TestTrue( TEXT("TestMetadataObjectA < TestMetadataObjectB"), *TestMetadataObjectA < *TestMetadataObjectB );
+	}
+
+	// Testing Loc Metadata Object
+	{
+		TestFalse( TEXT("MetadataObjectA == MetadataObjectB"), *MetadataObjectA == *MetadataObjectB );
+		TestTrue( TEXT("MetadataObjectA < MetadataObjectB"), *MetadataObjectA < *MetadataObjectB );
+		TestFalse( TEXT("MetadataObjectB < MetadataObjectA"), *MetadataObjectB < *MetadataObjectA);
+
+		// Test copy ctor
+		{
+			FLocMetadataObject MetadataObjectAClone = *MetadataObjectA;
+			TestTrue( TEXT("MetadataObjectAClone == MetadataObjectA"), MetadataObjectAClone == *MetadataObjectA );
+		}
+		
+		// Test assignment
+		{
+			FLocMetadataObject MetadataObjectAClone = *MetadataObjectB;
+			// PVS-Studio complains about double initialization, but that's something we're testing
+			// so we disable the warning
+			MetadataObjectAClone = *MetadataObjectA; //-V519
+			TestTrue( TEXT("MetadataObjectAClone == MetadataObjectA"), MetadataObjectAClone == *MetadataObjectA );
+			TestFalse( TEXT("MetadataObjectAClone == MetadataObjectB"), MetadataObjectAClone == *MetadataObjectB );
+		}
+		
+		// Test comparison operator  
+		{
+			FLocMetadataObject MetadataObjectAClone = *MetadataObjectA;
+			
+			
+			// Adding standard entry
+			MetadataObjectAClone.SetStringField( TEXT("NewEntry"), TEXT("NewEntryValue") );
+			TestFalse( TEXT("MetadataObjectAClone == MetadataObjectA"), MetadataObjectAClone == *MetadataObjectA );
+			
+			// Adding non-standard entry.  Note metadata with * prefix in the name will ignore value and type when performing comparisons
+			MetadataObjectAClone = *MetadataObjectA;
+			MetadataObjectAClone.SetStringField( TEXT("*NewEntry"), TEXT("*NewEntryValue") );
+			TestFalse( TEXT("MetadataObjectAClone == MetadataObjectA"), MetadataObjectAClone == *MetadataObjectA );
+
+			// Value mismatch on entry with * prefix with same type
+			MetadataObjectAClone = *MetadataObjectA;
+			MetadataObjectAClone.SetStringField( TEXT("*NoCompare"), TEXT("NoCompare") );
+			FLocMetadataObject MetadataObjectAClone2 = *MetadataObjectA;
+			MetadataObjectAClone2.SetStringField( TEXT("*NoCompare"), TEXT("NoCompare2") );
+			TestTrue( TEXT("MetadataObjectAClone == MetadataObjectAClone2"), MetadataObjectAClone == MetadataObjectAClone2 );
+
+			// Value and type mismatch on entry with * prefix
+			MetadataObjectAClone = *MetadataObjectA;
+			MetadataObjectAClone.SetStringField( TEXT("*NoCompare"), TEXT("NoCompare") );
+			MetadataObjectAClone2 = *MetadataObjectA;
+			MetadataObjectAClone2.SetBoolField( TEXT("*NoCompare"), true );
+			TestTrue( TEXT("MetadataObjectAClone == MetadataObjectAClone2"), MetadataObjectAClone == MetadataObjectAClone2 );
+		
+			// Value mismatch on standard entry
+			MetadataObjectAClone = *MetadataObjectA;
+			MetadataObjectAClone.SetStringField( TEXT("DoCompare"), TEXT("DoCompare") );
+			MetadataObjectAClone2 = *MetadataObjectA;
+			MetadataObjectAClone2.SetStringField( TEXT("DoCompare"), TEXT("DoCompare2") );
+			TestFalse( TEXT("MetadataObjectAClone == MetadataObjectAClone2"), MetadataObjectAClone == MetadataObjectAClone2 );
+
+			// Value and type mismatch on standard entry
+			MetadataObjectAClone = *MetadataObjectA;
+			MetadataObjectAClone.SetStringField( TEXT("DoCompare"), TEXT("DoCompare") );
+			MetadataObjectAClone2 = *MetadataObjectA;
+			MetadataObjectAClone2.SetBoolField( TEXT("DoCompare"), true );
+			TestFalse( TEXT("MetadataObjectAClone == MetadataObjectAClone2"), MetadataObjectAClone == MetadataObjectAClone2 );
+		}		
+		
+		// Test IsExactMatch function.  Note: Differs from the == operator which takes into account the COMPARISON_MODIFIER_PREFIX
+		{
+			// This function behaves similar to the comparison operator except that it ensures all metadata entries match exactly.  In other
+			//  words, it will also perform exact match checks on * prefixed metadata items.
+			FLocMetadataObject MetadataObjectAClone = *MetadataObjectA;
+
+			TestTrue( TEXT("MetadataObjectAClone == MetadataObjectA"), MetadataObjectAClone.IsExactMatch( *MetadataObjectA) );
+
+			// Adding standard entry
+			MetadataObjectAClone.SetStringField( TEXT("NewEntry"), TEXT("NewEntryValue") );
+			TestFalse( TEXT("MetadataObjectAClone == MetadataObjectA"), MetadataObjectAClone.IsExactMatch( *MetadataObjectA) );
+
+			// Adding non-standard entry.  Note metadata with * prefix in the name will ignore value and type when performing comparisons
+			MetadataObjectAClone = *MetadataObjectA;
+			MetadataObjectAClone.SetStringField( TEXT("*NewEntry"), TEXT("*NewEntryValue") );
+			TestFalse( TEXT("MetadataObjectAClone == MetadataObjectA"),MetadataObjectAClone.IsExactMatch( *MetadataObjectA)  );
+
+			// Value mismatch on entry with * prefix with same type
+			MetadataObjectAClone = *MetadataObjectA;
+			MetadataObjectAClone.SetStringField( TEXT("*NoCompare"), TEXT("NoCompare") );
+			FLocMetadataObject MetadataObjectAClone2 = *MetadataObjectA;
+			MetadataObjectAClone2.SetStringField( TEXT("*NoCompare"), TEXT("NoCompare2") );
+			TestFalse( TEXT("MetadataObjectAClone == MetadataObjectAClone2"), MetadataObjectAClone.IsExactMatch( MetadataObjectAClone2) );
+
+			// Value and type mismatch on entry with * prefix
+			MetadataObjectAClone = *MetadataObjectA;
+			MetadataObjectAClone.SetStringField( TEXT("*NoCompare"), TEXT("NoCompare") );
+			MetadataObjectAClone2 = *MetadataObjectA;
+			MetadataObjectAClone2.SetBoolField( TEXT("*NoCompare"), true );
+			TestFalse( TEXT("MetadataObjectAClone == MetadataObjectAClone2"),  MetadataObjectAClone.IsExactMatch( MetadataObjectAClone2) );
+
+			// Value mismatch on standard entry
+			MetadataObjectAClone = *MetadataObjectA;
+			MetadataObjectAClone.SetStringField( TEXT("DoCompare"), TEXT("DoCompare") );
+			MetadataObjectAClone2 = *MetadataObjectA;
+			MetadataObjectAClone2.SetStringField( TEXT("DoCompare"), TEXT("DoCompare2") );
+			TestFalse( TEXT("MetadataObjectAClone == MetadataObjectAClone2"),  MetadataObjectAClone.IsExactMatch( MetadataObjectAClone2) );
+
+			// Value and type mismatch on standard entry
+			MetadataObjectAClone = *MetadataObjectA;
+			MetadataObjectAClone.SetStringField( TEXT("DoCompare"), TEXT("DoCompare") );
+			MetadataObjectAClone2 = *MetadataObjectA;
+			MetadataObjectAClone2.SetBoolField( TEXT("DoCompare"), true );
+			TestFalse( TEXT("MetadataObjectAClone == MetadataObjectAClone2"),  MetadataObjectAClone.IsExactMatch( MetadataObjectAClone2) );
+		}
+
+		// Test less than operator
+		{
+			
+			// Adding standard entry that would appear before other entries
+			FLocMetadataObject MetadataObjectAClone = *MetadataObjectA;
+			MetadataObjectAClone.SetStringField( TEXT("ANewEntry"), TEXT("NewEntryValue") );
+			TestTrue( TEXT("MetadataObjectAClone < MetadataObjectA"), MetadataObjectAClone < *MetadataObjectA );
+			
+			// Adding standard entry that would appear after other entries
+			MetadataObjectAClone = *MetadataObjectA;
+			MetadataObjectAClone.SetStringField( TEXT("ZNewEntry"), TEXT("NewEntryValue") );
+			TestTrue( TEXT("MetadataObjectA < MetadataObjectAClone"), *MetadataObjectA < MetadataObjectAClone );
+
+			// Adding non-standard entry that would appear before other entries.  Note metadata with * prefix has no special treatment in less than operator
+			MetadataObjectAClone = *MetadataObjectA;
+			MetadataObjectAClone.SetStringField( TEXT("*NewEntry"), TEXT("NewEntryValue") );
+			TestTrue( TEXT("MetadataObjectAClone < MetadataObjectA"), MetadataObjectAClone < *MetadataObjectA );
+
+			// Value mismatch on entry with * prefix with same type. Note metadata with * prefix has no special treatment in less than operator
+			MetadataObjectAClone = *MetadataObjectA;
+			MetadataObjectAClone.SetStringField( TEXT("*NoCompare"), TEXT("NoCompare") );
+			FLocMetadataObject MetadataObjectAClone2 = *MetadataObjectA;
+			MetadataObjectAClone2.SetStringField( TEXT("*NoCompare"), TEXT("NoCompare2") );
+			TestTrue( TEXT("MetadataObjectAClone < MetadataObjectAClone2"), MetadataObjectAClone < MetadataObjectAClone2 );
+
+			// Value and type mismatch on entry with * prefix.  Note metadata with * prefix has no special treatment in less than operator
+			MetadataObjectAClone = *MetadataObjectA;
+			MetadataObjectAClone.SetBoolField( TEXT("*NoCompare"), true );
+			MetadataObjectAClone2 = *MetadataObjectA;
+			MetadataObjectAClone2.SetStringField( TEXT("*NoCompare"), TEXT("NoCompare") );
+			TestTrue( TEXT("MetadataObjectAClone < MetadataObjectAClone2"), MetadataObjectAClone < MetadataObjectAClone2 );
+		}
+	}
+}
+
+END_FUNCTION_BUILD_OPTIMIZATION
diff --git a/Engine/Source/Runtime/Core/Tests/Math/ColorTests.cpp b/Engine/Source/Runtime/Core/Tests/Math/ColorTests.cpp
new file mode 100644
--- /dev/null
+++ b/Engine/Source/Runtime/Core/Tests/Math/ColorTests.cpp
@@ -0,0 +1,65 @@
+// Copyright Epic Games, Inc. All Rights Reserved.
+#include "CoreMinimal.h"
+#include "Math/Color.h"
+#include "TestHarness.h"
+
+TEST_CASE("Core::Math::FLinearColor::Smoke Test", "[Core][Math][Smoke]")
+{
+	SECTION("FLinearColor")
+	{
+		FLinearColor black(FLinearColor::Black);
+		FLinearColor white(FLinearColor::White);
+		FLinearColor red(FLinearColor::Red);
+
+		REQUIRE(red.R == 1.0);
+		REQUIRE(red.G == 0.0); 
+		REQUIRE(red.B == 0.0);
+		REQUIRE(red.A == 1.0);
+		REQUIRE(red == FLinearColor::Red);
+		REQUIRE(red != FLinearColor::Green);
+
+		REQUIRE(white.IsAlmostBlack() == false);
+		REQUIRE(black.IsAlmostBlack() == true);
+
+		FLinearColor yellow(FColor::FromHex(FString("FFFF00FF")));
+		FLinearColor green(FColor::FromHex(FString("00FF00FF")));
+		FLinearColor blue(FColor::FromHex(FString("0000FFFF")));
+
+		REQUIRE(yellow == FLinearColor::Yellow);
+		REQUIRE(green == FLinearColor::Green);
+		REQUIRE(blue == FLinearColor::Blue);
+		REQUIRE(blue != yellow);
+		REQUIRE(green != yellow);
+		REQUIRE(yellow == yellow);
+	}
+}
+
+TEST_CASE("Core::Math::FColor::Smoke Test", "[Core][Math][Smoke]")
+{
+	SECTION("FColor")
+	{
+		FColor black(FColor::Black);
+		FColor white(FColor::White);
+		FColor red(FColor::Red);
+
+		REQUIRE(red.R == 0xFF);
+		REQUIRE(red.G == 0x00);
+		REQUIRE(red.B == 0x00);
+		REQUIRE(red.A == 0xFF);
+		REQUIRE(red == FColor::Red);
+		REQUIRE(red != FColor::Green);
+
+		FColor yellow(FColor::FromHex(FString("FFFF00FF")));
+		FColor green(FColor::FromHex(FString("00FF00FF")));
+		FColor blue(FColor::FromHex(FString("0000FFFF")));
+
+		REQUIRE(yellow == FColor::Yellow);
+		REQUIRE(green == FColor::Green);
+		REQUIRE(blue == FColor::Blue);
+		REQUIRE(blue != yellow);
+		REQUIRE(green != yellow);
+		REQUIRE(yellow == yellow);
+
+		REQUIRE(1 == 1);
+	}
+}
diff --git a/Engine/Source/Runtime/Core/Tests/Math/RangeBoundTest.cpp b/Engine/Source/Runtime/Core/Tests/Math/RangeBoundTest.cpp
new file mode 100644
--- /dev/null
+++ b/Engine/Source/Runtime/Core/Tests/Math/RangeBoundTest.cpp
@@ -0,0 +1,115 @@
+// Copyright Epic Games, Inc. All Rights Reserved.
+
+#include "CoreTypes.h"
+#include "Math/Range.h"
+#include "TestHarness.h"
+
+TEST_CASE("Core::Math::FFloatRangeBound::Smoke Test", "[Core][Math][Smoke]")
+{
+	// ensure template instantiation works
+	FDateRange DateRange;
+	FDoubleRange DoubleRange;
+	FFloatRange FloatRange;
+	FInt8Range Int8Range;
+	FInt16Range Int16Range;
+	FInt32Range Int32Range;
+	FInt64Range Int64Range;
+
+	// bound types must be correct after construction
+	FFloatRangeBound b1_1 = FFloatRangeBound::Exclusive(2.0f);
+	FFloatRangeBound b1_2 = FFloatRangeBound::Inclusive(2.0f);
+	FFloatRangeBound b1_3 = FFloatRangeBound::Open();
+	FFloatRangeBound b1_4 = 2;
+
+	TestTrue(TEXT("Exclusive bound constructor must create exclusive bound"), b1_1.IsExclusive());
+	TestTrue(TEXT("Exclusive bound constructor must create closed bound"), b1_1.IsClosed());
+	TestFalse(TEXT("Exclusive bound constructor must not create inclusive bound"), b1_1.IsInclusive());
+	TestFalse(TEXT("Exclusive bound constructor must not create open bound"), b1_1.IsOpen());
+	TestEqual(TEXT("Exclusive bound constructor must create the correct value"), b1_1.GetValue(), 2.0f);
+
+	TestTrue(TEXT("Inclusive bound constructor must create exclusive bound"), b1_2.IsInclusive());
+	TestTrue(TEXT("Inclusive bound constructor must create closed bound"), b1_2.IsClosed());
+	TestFalse(TEXT("Inclusive bound constructor must not create exclusive bound"), b1_2.IsExclusive());
+	TestFalse(TEXT("Inclusive bound constructor must not create open bound"), b1_2.IsOpen());
+	TestEqual(TEXT("Inclusive bound constructor must create the correct value"), b1_2.GetValue(), 2.0f);
+
+	TestTrue(TEXT("Open bound constructor must create open bound"), b1_3.IsOpen());
+	TestFalse(TEXT("Open bound constructor must not create closed bound"), b1_3.IsClosed());
+	TestFalse(TEXT("Open bound constructor must not create exclusive bound"), b1_3.IsExclusive());
+	TestFalse(TEXT("Open bound constructor must not create inclusive bound"), b1_3.IsInclusive());
+
+	TestTrue(TEXT("Implicit constructor must create an inclusive bound"), b1_4.IsInclusive());
+	TestEqual(TEXT("Implicit constructor must create the correct value"), b1_4, b1_2);
+
+	// comparisons must be correct
+	FFloatRangeBound b2_1 = FFloatRangeBound::Exclusive(2.0f);
+	FFloatRangeBound b2_2 = FFloatRangeBound::Exclusive(2.0f);
+	FFloatRangeBound b2_3 = FFloatRangeBound::Inclusive(2.0f);
+	FFloatRangeBound b2_4 = FFloatRangeBound::Inclusive(2.0f);
+	FFloatRangeBound b2_5 = FFloatRangeBound::Open();
+	FFloatRangeBound b2_6 = FFloatRangeBound::Open();
+
+	TestTrue(TEXT("Equal exclusive bounds must be equal"), b2_1 == b2_2);
+	TestTrue(TEXT("Equal inclusive bounds must be equal"), b2_3 == b2_4);
+	TestTrue(TEXT("Open bounds must be equal"), b2_5 == b2_6);
+
+	TestFalse(TEXT("Equal exclusive bounds must not be unequal"), b2_1 != b2_2);
+	TestFalse(TEXT("Equal inclusive bounds must not be unequal"), b2_3 != b2_4);
+	TestFalse(TEXT("Open bounds must not be unequal"), b2_5 != b2_6);
+
+	FFloatRangeBound b2_7 = FFloatRangeBound::Exclusive(3.0f);
+	FFloatRangeBound b2_8 = FFloatRangeBound::Inclusive(3.0f);
+
+	TestTrue(TEXT("Unequal exclusive bounds must be unequal"), b2_1 != b2_7);
+	TestTrue(TEXT("Unequal inclusive bounds must be unequal"), b2_2 != b2_8);
+
+	TestFalse(TEXT("Unequal exclusive bounds must not be equal"), b2_1 == b2_7);
+	TestFalse(TEXT("Unequal inclusive bounds must not be equal"), b2_2 == b2_8);
+
+	// min-max comparisons between bounds must be correct
+	FFloatRangeBound b3_1 = FFloatRangeBound::Exclusive(2.0f);
+	FFloatRangeBound b3_2 = FFloatRangeBound::Inclusive(2.0f);
+	FFloatRangeBound b3_3 = FFloatRangeBound::Open();
+
+	TestEqual(TEXT("'[2' must be less than '(2' <1>"), FFloatRangeBound::MinLower(b3_2, b3_1), b3_2);
+	TestEqual(TEXT("'[2' must be less than '(2' <2>"), FFloatRangeBound::MinLower(b3_1, b3_2), b3_2);
+	TestEqual(TEXT("Open lower bound must be less than '(2' <1>"), FFloatRangeBound::MinLower(b3_3, b3_1), b3_3);
+	TestEqual(TEXT("Open lower bound must be less than '(2' <2>"), FFloatRangeBound::MinLower(b3_1, b3_3), b3_3);
+	TestEqual(TEXT("Open lower bound must be less than '[2' <1>"), FFloatRangeBound::MinLower(b3_3, b3_2), b3_3);
+	TestEqual(TEXT("Open lower bound must be less than '[2' <2>"), FFloatRangeBound::MinLower(b3_2, b3_3), b3_3);
+
+	TestEqual(TEXT("'(2' must be greater than '[2' <1>"), FFloatRangeBound::MaxLower(b3_2, b3_1), b3_1);
+	TestEqual(TEXT("'(2' must be greater than '[2' <2>"), FFloatRangeBound::MaxLower(b3_1, b3_2), b3_1);
+	TestEqual(TEXT("'(2' must be greater than open lower bound <1>"), FFloatRangeBound::MaxLower(b3_3, b3_1), b3_1);
+	TestEqual(TEXT("'(2' must be greater than open lower bound <2>"), FFloatRangeBound::MaxLower(b3_1, b3_3), b3_1);
+	TestEqual(TEXT("'[2' must be greater than open lower bound <1>"), FFloatRangeBound::MaxLower(b3_3, b3_2), b3_2);
+	TestEqual(TEXT("'[2' must be greater than open lower bound <2>"), FFloatRangeBound::MaxLower(b3_2, b3_3), b3_2);
+
+	TestEqual(TEXT("'2)' must be less than '2]' <1>"), FFloatRangeBound::MinUpper(b3_2, b3_1), b3_1);
+	TestEqual(TEXT("'2)' must be less than '2]' <2>"), FFloatRangeBound::MinUpper(b3_1, b3_2), b3_1);
+	TestEqual(TEXT("'2)' must be less than open upper bound <1>"), FFloatRangeBound::MinUpper(b3_3, b3_1), b3_1);
+	TestEqual(TEXT("'2)' must be less than open upper bound <2>"), FFloatRangeBound::MinUpper(b3_1, b3_3), b3_1);
+	TestEqual(TEXT("'2]' must be less than open upper bound <1>"), FFloatRangeBound::MinUpper(b3_3, b3_2), b3_2);
+	TestEqual(TEXT("'2]' must be less than open upper bound"), FFloatRangeBound::MinUpper(b3_2, b3_3), b3_2);
+
+	TestEqual(TEXT("'2]' must be greater than '2)' <1>"), FFloatRangeBound::MaxUpper(b3_2, b3_1), b3_2);
+	TestEqual(TEXT("'2]' must be greater than '2)' <2>"), FFloatRangeBound::MaxUpper(b3_1, b3_2), b3_2);
+	TestEqual(TEXT("Open upper bound must be greater than '2)' <1>"), FFloatRangeBound::MaxUpper(b3_3, b3_1), b3_3);
+	TestEqual(TEXT("Open upper bound must be greater than '2)' <2>"), FFloatRangeBound::MaxUpper(b3_1, b3_3), b3_3);
+	TestEqual(TEXT("Open upper bound must be greater than '2]' <1>"), FFloatRangeBound::MaxUpper(b3_3, b3_2), b3_3);
+	TestEqual(TEXT("Open upper bound must be greater than '2]' <2>"), FFloatRangeBound::MaxUpper(b3_2, b3_3), b3_3);
+
+	FFloatRangeBound b4_1 = FFloatRangeBound::Exclusive(3);
+	FFloatRangeBound b4_2 = FFloatRangeBound::Inclusive(3);
+
+	TestEqual(TEXT("'(2' must be less than '[3' <1>"), FFloatRangeBound::MinLower(b3_1, b4_2), b3_1);
+	TestEqual(TEXT("'(2' must be less than '[3' <2>"), FFloatRangeBound::MinLower(b4_2, b3_1), b3_1);
+	TestEqual(TEXT("'[2' must be less than '[3' <1>"), FFloatRangeBound::MinLower(b3_2, b4_2), b3_2);
+	TestEqual(TEXT("'[2' must be less than '[3' <2>"), FFloatRangeBound::MinLower(b4_2, b3_2), b3_2);
+
+	TestEqual(TEXT("'[3' must be greater than '(2' <1>"), FFloatRangeBound::MaxLower(b3_1, b4_2), b4_2);
+	TestEqual(TEXT("'[3' must be greater than '(2' <2>"), FFloatRangeBound::MaxLower(b4_2, b3_1), b4_2);
+	TestEqual(TEXT("'[3' must be greater than '[2' <1>"), FFloatRangeBound::MaxLower(b3_2, b4_2), b4_2);
+	TestEqual(TEXT("'[3' must be greater than '[2' <2>"), FFloatRangeBound::MaxLower(b4_2, b3_2), b4_2);
+
+}
diff --git a/Engine/Source/Runtime/Core/Tests/Math/RangeTest.cpp b/Engine/Source/Runtime/Core/Tests/Math/RangeTest.cpp
new file mode 100644
--- /dev/null
+++ b/Engine/Source/Runtime/Core/Tests/Math/RangeTest.cpp
@@ -0,0 +1,392 @@
+// Copyright Epic Games, Inc. All Rights Reserved.
+
+#include "CoreTypes.h"
+#include "Containers/Array.h"
+#include "Math/Range.h"
+#include "TestHarness.h"
+
+// Disable optimization for RangeTest as it compiles very slowly in development builds
+PRAGMA_DISABLE_OPTIMIZATION
+
+TEST_CASE("Core::Math::FFloatRange::Smoke Test", "[Core][Math][Smoke]")
+{
+	// single element constructor
+	FFloatRange r1_1 = FFloatRange(3.0f);
+
+	TestEqual(TEXT("Single element constructor must create the correct lower bound value"), r1_1.GetLowerBoundValue(), 3.0f);
+	TestEqual(TEXT("Single element constructor must create the correct upper bound value"), r1_1.GetUpperBoundValue(), 3.0f);
+	TestTrue(TEXT("Single element constructor must create an inclusive lower bound"), r1_1.GetLowerBound().IsInclusive());
+	TestTrue(TEXT("Single element constructor must create an inclusive upper bound"), r1_1.GetUpperBound().IsInclusive());
+	TestFalse(TEXT("Single element constructor must create non-empty range"), r1_1.IsEmpty());
+
+	// explicit element pair constructor
+	FFloatRange r2_1 = FFloatRange(1.0f, 4.0f);
+
+	TestEqual(TEXT("Explicit element pair constructor must create the correct lower bound value"), r2_1.GetLowerBoundValue(), 1.0f);
+	TestEqual(TEXT("Explicit element pair constructor must create the correct upper bound value"), r2_1.GetUpperBoundValue(), 4.0f);
+	TestTrue(TEXT("Explicit element pair constructor must create an inclusive lower bound"), r2_1.GetLowerBound().IsInclusive());
+	TestTrue(TEXT("Explicit element pair constructor must create an exclusive upper bound"), r2_1.GetUpperBound().IsExclusive());
+	TestFalse(TEXT("Explicit element pair constructor must create non-empty range"), r2_1.IsEmpty());
+
+	// bound pair constructors
+	FFloatRange r3_1 = FFloatRange(3.0f);
+	FFloatRange r3_2 = FFloatRange(1.0f, 4.0f);
+	FFloatRange r3_3 = FFloatRange(FFloatRangeBound(3.0f), FFloatRangeBound(3.0f));
+	FFloatRange r3_4 = FFloatRange(FFloatRangeBound::Inclusive(1.0f), FFloatRangeBound::Exclusive(4.0f));
+
+	TestEqual(TEXT("Default bound pair constructor must create the correct range"), r3_3, r3_1);
+	TestFalse(TEXT("Default bound pair constructor must create non-empty range"), r3_3.IsEmpty());
+	TestEqual(TEXT("Specific bound pair constructor must create the correct range"), r3_4, r3_2);
+	TestFalse(TEXT("Specific bound pair constructor must create non-empty range"), r3_4.IsEmpty());
+
+	// empty ranges
+	FFloatRange r4_1 = FFloatRange::Empty();
+	FFloatRange r4_2 = FFloatRange(4.0f, 1.0f);
+	FFloatRange r4_3 = FFloatRange(FFloatRangeBound::Inclusive(6.0f), FFloatRangeBound::Exclusive(2.0f));
+	FFloatRange r4_4 = FFloatRange(FFloatRangeBound::Exclusive(4.0f), FFloatRangeBound::Exclusive(4.0f));
+	FFloatRange r4_5 = FFloatRange(FFloatRangeBound::Exclusive(4.0f), FFloatRangeBound::Inclusive(4.0f));
+	FFloatRange r4_6 = FFloatRange(FFloatRangeBound::Inclusive(4.0f), FFloatRangeBound::Exclusive(4.0f));
+
+	TestTrue(TEXT("Empty range must be empty <1>"), r4_1.IsEmpty());
+	TestTrue(TEXT("Empty range must be empty <2>"), r4_2.IsEmpty());
+	TestTrue(TEXT("Empty range must be empty <3>"), r4_3.IsEmpty());
+	TestTrue(TEXT("Empty range must be empty <4>"), r4_4.IsEmpty());
+	TestTrue(TEXT("Empty range must be empty <5>"), r4_5.IsEmpty());
+	TestTrue(TEXT("Empty range must be empty <6>"), r4_6.IsEmpty());
+
+	TestEqual(TEXT("Empty ranges must be equal <1>"), r4_1, r4_2);
+	TestEqual(TEXT("Empty ranges must be equal <2>"), r4_1, r4_3);
+	TestEqual(TEXT("Empty ranges must be equal <3>"), r4_1, r4_4);
+	TestEqual(TEXT("Empty ranges must be equal <4>"), r4_1, r4_5);
+	TestEqual(TEXT("Empty ranges must be equal <5>"), r4_1, r4_6);
+
+	// inclusive-exclusive comparison
+	FFloatRange r5_1 = FFloatRange(FFloatRangeBound::Exclusive(1.0f), FFloatRangeBound::Exclusive(4.0f));
+	FFloatRange r5_2 = FFloatRange(FFloatRangeBound::Exclusive(1.0f), FFloatRangeBound::Inclusive(4.0f));
+	FFloatRange r5_3 = FFloatRange(FFloatRangeBound::Inclusive(1.0f), FFloatRangeBound::Exclusive(4.0f));
+	FFloatRange r5_4 = FFloatRange(FFloatRangeBound::Inclusive(1.0f), FFloatRangeBound::Inclusive(4.0f));
+
+	TestEqual(TEXT("Same inclusive-exclusive ranges must be equal <1>"), r5_1, r5_1);
+	TestEqual(TEXT("Same inclusive-exclusive ranges must be equal <2>"), r5_2, r5_2);
+	TestEqual(TEXT("Same inclusive-exclusive ranges must be equal <3>"), r5_3, r5_3);
+	TestEqual(TEXT("Same inclusive-exclusive ranges must be equal <4>"), r5_4, r5_4);
+
+	TestNotEqual(TEXT("Different inclusive-exclusive ranges must be different <1>"), r5_1, r5_2);
+	TestNotEqual(TEXT("Different inclusive-exclusive ranges must be different <2>"), r5_1, r5_3);
+	TestNotEqual(TEXT("Different inclusive-exclusive ranges must be different <3>"), r5_1, r5_4);
+	TestNotEqual(TEXT("Different inclusive-exclusive ranges must be different <4>"), r5_2, r5_3);
+	TestNotEqual(TEXT("Different inclusive-exclusive ranges must be different <5>"), r5_2, r5_4);
+	TestNotEqual(TEXT("Different inclusive-exclusive ranges must be different <6>"), r5_3, r5_4);
+
+	// adjoined ranges
+	FFloatRange r6_1 = FFloatRange(FFloatRangeBound::Exclusive(1.0f), FFloatRangeBound::Exclusive(4.0f));
+	FFloatRange r6_2 = FFloatRange(FFloatRangeBound::Exclusive(1.0f), FFloatRangeBound::Inclusive(4.0f));
+	FFloatRange r6_3 = FFloatRange(FFloatRangeBound::Exclusive(4.0f), FFloatRangeBound::Inclusive(8.0f));
+	FFloatRange r6_4 = FFloatRange(FFloatRangeBound::Inclusive(4.0f), FFloatRangeBound::Inclusive(8.0f));
+	FFloatRange r6_5 = FFloatRange(FFloatRangeBound::Inclusive(3.0f), FFloatRangeBound::Inclusive(8.0f));
+	FFloatRange r6_6 = FFloatRange(FFloatRangeBound::Inclusive(0.0f), FFloatRangeBound::Inclusive(2.0f));
+
+	TestTrue(TEXT("Adjoined ranges must be adjoined"), r6_1.Adjoins(r6_4));
+	TestTrue(TEXT("Adjoined ranges must be adjoined"), r6_2.Adjoins(r6_3));
+	TestFalse(TEXT("Separated ranges must not be adjoined"), r6_1.Adjoins(r6_3));
+	TestFalse(TEXT("Separated ranges must not be adjoined"), r6_2.Adjoins(r6_4));
+	TestFalse(TEXT("Overlapped ranges must not be adjoined <1>"), r6_1.Adjoins(r6_5));
+	TestFalse(TEXT("Overlapped ranges must not be adjoined <2>"), r6_1.Adjoins(r6_6));
+
+	FFloatRange r6_7 = FFloatRange(FFloatRangeBound::Exclusive(4.0f), FFloatRangeBound::Exclusive(2.0f));
+	FFloatRange r6_8 = FFloatRange(FFloatRangeBound::Inclusive(4.0f), FFloatRangeBound::Inclusive(2.0f));
+
+	TestFalse(TEXT("A non-empty and an empty range must not be adjoined <1>"), r6_1.Adjoins(r6_8));
+	TestFalse(TEXT("A non-empty and an empty range must not be adjoined <2>"), r6_2.Adjoins(r6_7));
+	
+	FFloatRange r6_9 = FFloatRange(FFloatRangeBound::Inclusive(2.0f), FFloatRangeBound::Exclusive(0.0f));
+	FFloatRange r6_10 = FFloatRange(FFloatRangeBound::Exclusive(2.0f), FFloatRangeBound::Exclusive(0.0f));
+
+	TestFalse(TEXT("Two empty ranges must not be adjoined <1>"), r6_7.Adjoins(r6_9));
+	TestFalse(TEXT("Two empty ranges must not be adjoined <2>"), r6_8.Adjoins(r6_10));
+
+	// @todo gmp: conjoined ranges
+
+	// @todo gmp: element containment
+
+	// @todo gmp: range containment
+
+	// @todo gmp: contiguous ranges
+
+	// bounds access
+	FFloatRange r11_1 = FFloatRange(FFloatRangeBound::Inclusive(1.0f), FFloatRangeBound::Exclusive(4.0f));
+	FFloatRange r11_2 = FFloatRange(FFloatRangeBound::Open(), FFloatRangeBound::Open());
+
+	TestTrue(TEXT("A closed range must have a lower bound"), r11_1.HasLowerBound());
+	TestTrue(TEXT("A closed range must have an upper bound"), r11_1.HasUpperBound());
+	TestFalse(TEXT("An open range must not have a lower bound"), r11_2.HasLowerBound());
+	TestFalse(TEXT("An open range must not have an upper bound"), r11_2.HasUpperBound());
+
+	TestEqual(TEXT("The lower bound value of [1, 4) must be 1"), r11_1.GetLowerBoundValue(), 1.0f);
+	TestEqual(TEXT("The upper bound value of [1, 4) must be 4"), r11_1.GetUpperBoundValue(), 4.0f);
+
+	// degenerate ranges
+	FFloatRange r12_1 = FFloatRange(3.0f);
+	FFloatRange r12_2 = FFloatRange(FFloatRangeBound::Inclusive(3.0f), FFloatRangeBound::Inclusive(3.0f));
+
+	TestTrue(TEXT("A range with a single element must be degenerate <1>"), r12_1.IsDegenerate());
+	TestTrue(TEXT("A range with a single element must be degenerate <2>"), r12_2.IsDegenerate());
+
+	// @todo gmp: overlapping ranges
+	FFloatRange r14_1 = FFloatRange(FFloatRangeBound::Exclusive(0.0f), FFloatRangeBound::Exclusive(2.0f));
+	FFloatRange r14_2 = FFloatRange(FFloatRangeBound::Inclusive(0.0f), FFloatRangeBound::Inclusive(2.0f));
+	FFloatRange r14_3 = FFloatRange(FFloatRangeBound::Exclusive(2.0f), FFloatRangeBound::Exclusive(4.0f));
+	FFloatRange r14_4 = FFloatRange(FFloatRangeBound::Inclusive(2.0f), FFloatRangeBound::Inclusive(4.0f));
+	FFloatRange r14_5 = FFloatRange(FFloatRangeBound::Inclusive(3.0f), FFloatRangeBound::Inclusive(5.0f));
+	FFloatRange r14_6 = FFloatRange(FFloatRangeBound::Inclusive(1.0f), FFloatRangeBound::Inclusive(3.0f));
+	FFloatRange r14_7 = FFloatRange(FFloatRangeBound::Exclusive(1.0f), FFloatRangeBound::Inclusive(3.0f));
+	FFloatRange r14_8 = FFloatRange(FFloatRangeBound::Inclusive(1.0f), FFloatRangeBound::Exclusive(3.0f));
+	FFloatRange r14_9 = FFloatRange(FFloatRangeBound::Exclusive(6.0f), FFloatRangeBound::Exclusive(10.0f));
+	FFloatRange r14_10 = FFloatRange(FFloatRangeBound::Inclusive(4.0f), FFloatRangeBound::Open());
+	FFloatRange r14_11 = FFloatRange(FFloatRangeBound::Open(), FFloatRangeBound::Exclusive(3.0f));
+	
+	TestTrue(TEXT("(2, 4) must overlap [3, 5]"), r14_3.Overlaps(r14_5));
+	TestTrue(TEXT("(2, 4) must overlap [1, 3]"), r14_3.Overlaps(r14_6));
+	TestTrue(TEXT("(2, 4) must overlap (1, 3]"), r14_3.Overlaps(r14_7));
+	TestTrue(TEXT("(2, 4) must overlap [1, 3)"), r14_3.Overlaps(r14_8));
+	TestTrue(TEXT("(2, 4) must overlap [2, 4]"), r14_3.Overlaps(r14_4));
+	TestTrue(TEXT("(2, 4) must overlap itself"), r14_3.Overlaps(r14_3));
+	TestTrue(TEXT("(2, 4) must not overlap (0, 2)"), !r14_3.Overlaps(r14_1));
+	TestTrue(TEXT("(2, 4) must not overlap [0, 2]"), !r14_3.Overlaps(r14_2));
+	TestTrue(TEXT("[2, 4] must not overlap (0, 2)"), !r14_4.Overlaps(r14_1));
+	TestTrue(TEXT("[2, 4] must overlap [0, 2]"), r14_4.Overlaps(r14_2));
+	TestTrue(TEXT("[2, 4] must not overlap (6, 10)"), !r14_4.Overlaps(r14_9));
+	TestTrue(TEXT("[2, 4] must overlap [4, inf)"), r14_4.Overlaps(r14_10));
+	TestTrue(TEXT("(2, 4) must not overlap [4, inf)"), !r14_3.Overlaps(r14_10));
+	TestTrue(TEXT("[1, 3] must not overlap [4, inf)"), !r14_6.Overlaps(r14_10));
+	TestTrue(TEXT("[1, 3] must overlap (inf, 3)"), r14_6.Overlaps(r14_11));
+	TestTrue(TEXT("(6, 10) must not overlap (inf, 3)"), !r14_9.Overlaps(r14_11));
+
+	// @todo gmp: range splitting
+
+	// difference
+	FFloatRange r15_1 = FFloatRange(FFloatRangeBound::Exclusive(7.0f), FFloatRangeBound::Exclusive(14.0f));		// X
+	FFloatRange r15_2 = FFloatRange(FFloatRangeBound::Inclusive(7.0f), FFloatRangeBound::Inclusive(14.0f));
+
+	TestEqual(TEXT("The difference between a range and itself must be an empty set <1>"), FFloatRange::Difference(r15_1, r15_1).Num(), 0);
+	TestEqual(TEXT("The difference between a range and itself must be an empty set <2>"), FFloatRange::Difference(r15_2, r15_2).Num(), 0);
+	
+	TArray<FFloatRange > d15_0 = FFloatRange::Difference(r15_2, r15_1);
+	{
+		TestEqual(TEXT("[7, 14] - (7, 14) must result in two ranges"), d15_0.Num(), 2);
+
+		bool b15_1 = (d15_0[0] == FFloatRange(FFloatRangeBound::Inclusive(7.0f), FFloatRangeBound::Inclusive(7.0f)));
+		bool b15_2 = (d15_0[1] == FFloatRange(FFloatRangeBound::Inclusive(14.0f), FFloatRangeBound::Inclusive(14.0f)));
+
+		TestTrue(TEXT("[7, 14] - (7, 14) must be {[7, 7], [14, 14]}"), b15_1 && b15_2);
+	}
+
+	FFloatRange r15_3 = FFloatRange(FFloatRangeBound::Inclusive(2.0f), FFloatRangeBound::Exclusive(9.0f));		// Y overlapping left
+	FFloatRange r15_4 = FFloatRange(FFloatRangeBound::Inclusive(2.0f), FFloatRangeBound::Inclusive(9.0f));
+	FFloatRange r15_5 = FFloatRange(FFloatRangeBound::Exclusive(8.0f), FFloatRangeBound::Exclusive(17.0f));		// Y overlapping right
+	FFloatRange r15_6 = FFloatRange(FFloatRangeBound::Inclusive(8.0f), FFloatRangeBound::Exclusive(17.0f));
+
+	TArray<FFloatRange > d15_1 = FFloatRange::Difference(r15_1, r15_3);
+	{
+		TestEqual(TEXT("(7, 14) - [2, 9) must result in one range"), d15_1.Num(), 1);
+		TestEqual(TEXT("(7, 14) - [2, 9) must be {[9, 14)}"), d15_1[0], FFloatRange(FFloatRangeBound::Inclusive(9.0f), FFloatRangeBound::Exclusive(14.0f)));
+	}
+
+	TArray<FFloatRange > d15_2 = FFloatRange::Difference(r15_1, r15_4);
+	{
+		TestEqual(TEXT("(7, 14) - [2, 9] must result in one range"), d15_2.Num(), 1);
+		TestEqual(TEXT("(7, 14) - [2, 9] must be {(9, 14)}"), d15_2[0], FFloatRange(FFloatRangeBound::Exclusive(9.0f), FFloatRangeBound::Exclusive(14.0f)));
+	}
+
+	TArray<FFloatRange > d15_3 = FFloatRange::Difference(r15_1, r15_5);
+	{
+		TestEqual(TEXT("(7, 14) - (8, 17) must result in one range"), d15_3.Num(), 1);
+		TestEqual(TEXT("(7, 14) - (8, 17) must be {(7, 8]}"), d15_3[0], FFloatRange(FFloatRangeBound::Exclusive(7.0f), FFloatRangeBound::Inclusive(8.0f)));
+	}
+
+	TArray<FFloatRange > d15_4 = FFloatRange::Difference(r15_1, r15_6);
+	{
+		TestEqual(TEXT("(7, 14) - [8, 17) must result in one range"), d15_4.Num(), 1);
+		TestEqual(TEXT("(7, 14) - [8, 17) must be {(7, 8)}"), d15_4[0], FFloatRange(FFloatRangeBound::Exclusive(7.0f), FFloatRangeBound::Exclusive(8.0f)));
+	}
+
+	TArray<FFloatRange > d15_5 = FFloatRange::Difference(r15_2, r15_3);
+	{
+		TestEqual(TEXT("[7, 14] - [2, 9) must result in one range"), d15_5.Num(), 1);
+		TestEqual(TEXT("[7, 14] - [2, 9) must be {[9, 14]}"), d15_5[0], FFloatRange(FFloatRangeBound::Inclusive(9.0f), FFloatRangeBound::Inclusive(14.0f)));
+	}
+
+	TArray<FFloatRange > d15_6 = FFloatRange::Difference(r15_2, r15_4);
+	{
+		TestEqual(TEXT("[7, 14] - [2, 9] must result in one range"), d15_6.Num(), 1);
+		TestEqual(TEXT("[7, 14] - [2, 9] must be {(9, 14]}"), d15_6[0], FFloatRange(FFloatRangeBound::Exclusive(9.0f), FFloatRangeBound::Inclusive(14.0f)));
+	}
+
+	TArray<FFloatRange > d15_7 = FFloatRange::Difference(r15_2, r15_5);
+	{
+		TestEqual(TEXT("[7, 14] - (8, 17) must result in one range"), d15_7.Num(), 1);
+		TestEqual(TEXT("[7, 14] - (8, 17) must be {[7, 8]}"), d15_7[0], FFloatRange(FFloatRangeBound::Inclusive(7.0f), FFloatRangeBound::Inclusive(8.0f)));
+	}
+
+	TArray<FFloatRange > d15_8 = FFloatRange::Difference(r15_2, r15_6);
+	{
+		TestEqual(TEXT("[7, 14] - [8, 17) must result in one range"), d15_8.Num(), 1);
+		TestEqual(TEXT("[7, 14] - [8, 17) must be {[7, 8)}"), d15_8[0], FFloatRange(FFloatRangeBound::Inclusive(7.0f), FFloatRangeBound::Exclusive(8.0f)));
+	}
+
+	FFloatRange r15_7 = FFloatRange(FFloatRangeBound::Inclusive(2.0f), FFloatRangeBound::Exclusive(7.0f));		// Y adjoint left
+	FFloatRange r15_8 = FFloatRange(FFloatRangeBound::Inclusive(2.0f), FFloatRangeBound::Inclusive(7.0f));
+
+	TArray<FFloatRange > d15_9 = FFloatRange::Difference(r15_1, r15_8);
+	{
+		TestEqual(TEXT("(7, 14) - [2, 7] must result in one set"), d15_9.Num(), 1);
+		TestEqual(TEXT("(7, 14) - [2, 7] must be {(7, 14)}"), d15_9[0], FFloatRange(FFloatRangeBound::Exclusive(7.0f), FFloatRangeBound::Exclusive(14.0f)));
+	}
+
+	TArray<FFloatRange > d15_10 = FFloatRange::Difference(r15_2, r15_7);
+	{
+		TestEqual(TEXT("[7, 14] - [2, 7) must result in one set"), d15_10.Num(), 1);
+		TestEqual(TEXT("[7, 14] - [2, 7) must be {[7, 14]}"), d15_10[0], FFloatRange(FFloatRangeBound::Inclusive(7.0f), FFloatRangeBound::Inclusive(14.0f)));
+	}
+
+	TArray<FFloatRange > d15_11 = FFloatRange::Difference(r15_2, r15_8);
+	{
+		TestEqual(TEXT("[7, 14] - [2, 7] must result in one range"), d15_11.Num(), 1);
+		TestEqual(TEXT("[7, 14] - [2, 7] must be {(7, 14]}"), d15_11[0], FFloatRange(FFloatRangeBound::Exclusive(7.0f), FFloatRangeBound::Inclusive(14.0f)));
+	}
+
+	FFloatRange r15_9 = FFloatRange(FFloatRangeBound::Exclusive(14.0f), FFloatRangeBound::Exclusive(16.0f));	// Y adjoint right
+	FFloatRange r15_10 = FFloatRange(FFloatRangeBound::Inclusive(14.0f), FFloatRangeBound::Inclusive(16.0f));
+
+	TArray<FFloatRange > d15_12 = FFloatRange::Difference(r15_1, r15_10);
+	{
+		TestEqual(TEXT("(7, 14) - [14, 18] must result in one range"), d15_12.Num(), 1);
+		TestEqual(TEXT("(7, 14) - [14, 18] must be {(7, 14)}"), d15_12[0], FFloatRange(FFloatRangeBound::Exclusive(7.0f), FFloatRangeBound::Exclusive(14.0f)));
+	}
+
+	TArray<FFloatRange > d15_13 = FFloatRange::Difference(r15_2, r15_9);
+	{
+		TestEqual(TEXT("[7, 14] - (14, 18) must result in one range"), d15_13.Num(), 1);
+		TestEqual(TEXT("[7, 14] - (14, 18) must be {[7, 14]}"), d15_13[0], FFloatRange(FFloatRangeBound::Inclusive(7.0f), FFloatRangeBound::Inclusive(14.0f)));
+	}
+
+	TArray<FFloatRange > d15_14 = FFloatRange::Difference(r15_2, r15_10);
+	{
+		TestEqual(TEXT("[7, 14] - [14, 18] must result in one range"), d15_14.Num(), 1);
+		TestEqual(TEXT("[7, 14] - [14, 18] must be {(7, 14)}"), d15_14[0], FFloatRange(FFloatRangeBound::Inclusive(7.0f), FFloatRangeBound::Exclusive(14.0f)));
+	}
+
+	FFloatRange r15_11 = FFloatRange(FFloatRangeBound::Inclusive(8.0f), FFloatRangeBound::Inclusive(13.0f));	// Y enclosed
+
+	TArray<FFloatRange > d15_15 = FFloatRange::Difference(r15_2, r15_11);
+	{
+		TestEqual(TEXT("[7, 14] - [8, 13] must result in two ranges"), d15_15.Num(), 2);
+
+		bool b15_1 = (d15_15[0] == FFloatRange(FFloatRangeBound::Inclusive(7.0f), FFloatRangeBound::Exclusive(8.0f)));
+		bool b15_2 = (d15_15[1] == FFloatRange(FFloatRangeBound::Exclusive(13.0f), FFloatRangeBound::Inclusive(14.0f)));
+
+		TestTrue(TEXT("[7, 14] - [8, 13] must be {[7, 8), (13, 14]}"), b15_1 && b15_2);
+	}
+	
+	FFloatRange r15_12 = FFloatRange(FFloatRangeBound::Inclusive(2.0f), FFloatRangeBound::Inclusive(4.0f));
+
+	TArray<FFloatRange > d15_16 = FFloatRange::Difference(r15_2, r15_12);
+	{
+		TestEqual(TEXT("[7, 14] - [2, 4) must result in one set"), d15_16.Num(), 1);
+		TestEqual(TEXT("[7, 14] - [2, 4) must be {[7, 14]}"), d15_16[0], FFloatRange(FFloatRangeBound::Inclusive(7.0f), FFloatRangeBound::Inclusive(14.0f)));
+	}
+
+	// range hulls
+	FFloatRange r16_1 = FFloatRange(FFloatRangeBound::Inclusive(7.0f), FFloatRangeBound::Exclusive(9.0f));
+	FFloatRange r16_2 = FFloatRange(FFloatRangeBound::Inclusive(11.0f), FFloatRangeBound::Exclusive(14.0f));
+	FFloatRange r16_3 = FFloatRange(4.0f, 1.0f);
+	FFloatRange r16_4 = FFloatRange(8.0f, 5.0f);
+
+	TestEqual(TEXT("The hull of [7, 9) and [11, 14) must be [7, 14)"), FFloatRange::Hull(r16_1, r16_2), FFloatRange(FFloatRangeBound::Inclusive(7.0f), FFloatRangeBound::Exclusive(14.0f)));
+	TestEqual(TEXT("The hull of [7, 9) and an empty range must be [7, 9)"), FFloatRange::Hull(r16_1, r16_3), FFloatRange(FFloatRangeBound::Inclusive(7.0f), FFloatRangeBound::Exclusive(9.0f)));
+	TestTrue(TEXT("The hull of two empty ranges must be empty"), FFloatRange::Hull(r16_3, r16_4).IsEmpty());
+
+	// intersections
+	FFloatRange r17_1 = FFloatRange(FFloatRangeBound::Inclusive(7.0f), FFloatRangeBound::Exclusive(14.0f));
+	FFloatRange r17_2 = FFloatRange(FFloatRangeBound::Inclusive(2.0f), FFloatRangeBound::Exclusive(8.0f));		// left overlap
+	FFloatRange r17_3 = FFloatRange(FFloatRangeBound::Inclusive(2.0f), FFloatRangeBound::Inclusive(8.0f));
+	FFloatRange r17_4 = FFloatRange(FFloatRangeBound::Exclusive(8.0f), FFloatRangeBound::Exclusive(16.0f));		// right overlap
+	FFloatRange r17_5 = FFloatRange(FFloatRangeBound::Inclusive(8.0f), FFloatRangeBound::Exclusive(16.0f));
+	FFloatRange r17_6 = FFloatRange(4.0f, 1.0f);
+	FFloatRange r17_7 = FFloatRange(8.0f, 5.0f);
+
+	TestEqual(TEXT("The intersection of [7, 14) and [2, 8) must be [7, 8)"), FFloatRange::Intersection(r17_1, r17_2), FFloatRange(FFloatRangeBound::Inclusive(7.0f), FFloatRangeBound::Exclusive(8.0f)));
+	TestEqual(TEXT("The intersection of [7, 14) and [2, 8] must be [7, 8]"), FFloatRange::Intersection(r17_1, r17_3), FFloatRange(FFloatRangeBound::Inclusive(7.0f), FFloatRangeBound::Inclusive(8.0f)));
+	TestEqual(TEXT("The intersection of [7, 14) and (8, 16) must be (8, 14)"), FFloatRange::Intersection(r17_1, r17_4), FFloatRange(FFloatRangeBound::Exclusive(8.0f), FFloatRangeBound::Exclusive(14.0f)));
+	TestEqual(TEXT("The intersection of [7, 14) and [8, 16) must be [8, 14)"), FFloatRange::Intersection(r17_1, r17_5), FFloatRange(FFloatRangeBound::Inclusive(8.0f), FFloatRangeBound::Exclusive(14.0f)));
+	TestTrue(TEXT("The intersection of a non-empty range and an empty range must be empty"), FFloatRange::Intersection(r17_1, r17_6).IsEmpty());
+	TestTrue(TEXT("The intersection of two empty ranges must be empty"), FFloatRange::Intersection(r17_6, r17_7).IsEmpty());
+
+	// union
+	FFloatRange r18_1 = FFloatRange(FFloatRangeBound::Inclusive(7.0f), FFloatRangeBound::Inclusive(14.0f));
+	FFloatRange r18_2 = FFloatRange(FFloatRangeBound::Exclusive(7.0f), FFloatRangeBound::Exclusive(14.0f));
+	FFloatRange r18_3 = FFloatRange(FFloatRangeBound::Inclusive(2.0f), FFloatRangeBound::Inclusive(8.0f));		// left overlap
+	FFloatRange r18_4 = FFloatRange(FFloatRangeBound::Exclusive(15.0f), FFloatRangeBound::Inclusive(16.0f));
+	FFloatRange r18_5 = FFloatRange(FFloatRangeBound::Inclusive(8.0f), FFloatRangeBound::Exclusive(16.0f));		// right overlap
+	FFloatRange r18_6 = FFloatRange(FFloatRangeBound::Inclusive(6.0f), FFloatRangeBound::Exclusive(7.0f));
+	
+	TArray<FFloatRange > u18_0 = FFloatRange::Union(r18_1, r18_1);
+	{
+		TestEqual(TEXT("[7, 14] unioned with itself must result in one range"), u18_0.Num(), 1);
+		TestEqual(TEXT("[7, 14] unioned with itself must result in itself"), u18_0[0], r18_1);
+	}
+
+	TArray<FFloatRange > u18_1 = FFloatRange::Union(r18_1, r18_2);
+	{
+		TestEqual(TEXT("[7, 14] unioned with (7, 14) must result in one range"), u18_1.Num(), 1);
+		TestEqual(TEXT("[7, 14] unioned with (7, 14) must result in {[7, 14]}"), u18_1[0], r18_1);
+	}
+
+	TArray<FFloatRange > u18_2 = FFloatRange::Union(r18_1, r18_3);
+	{
+		TestEqual(TEXT("[7, 14] unioned with [2, 8] must result in one range"), u18_2.Num(), 1);
+		TestEqual(TEXT("[7, 14] unioned with [2, 8] must result in {[2, 14]}"), u18_2[0], FFloatRange(FFloatRangeBound::Inclusive(2.f), FFloatRangeBound::Inclusive(14.f)));
+	}
+
+	TArray<FFloatRange > u18_3 = FFloatRange::Union(r18_1, r18_4);
+	{
+		TestEqual(TEXT("[7, 14] unioned with (15, 16] must result in two ranges"), u18_3.Num(), 2);
+		
+		bool b18_1 = (u18_3[0] == r18_1);
+		bool b18_2 = (u18_3[1] == r18_4);
+
+		TestTrue(TEXT("[7, 14] unioned with (15, 16] must result in {[7, 14], (15, 16]}"), b18_1 && b18_2);
+	}
+
+	TArray<FFloatRange > u18_4 = FFloatRange::Union(r18_1, FFloatRange::Empty());
+	{
+		TestEqual(TEXT("A non-empty range unioned with an empty range must result in one range"), u18_4.Num(), 1);
+		TestEqual(TEXT("A non-empty range unioned with an empty range must result in the non-empty range"), u18_4[0], r18_1);
+	}
+
+	TArray<FFloatRange > u18_5 = FFloatRange::Union(r18_1, r18_5);
+	{
+		TestEqual(TEXT("[7, 14] unioned with [8, 16) must result in one range"), u18_5.Num(), 1);
+		TestEqual(TEXT("[7, 14] unioned with [8, 16) must result in {[7, 16)}"), u18_5[0], FFloatRange(FFloatRangeBound::Inclusive(7.f), FFloatRangeBound::Exclusive(16.f)));
+	}
+
+	TArray<FFloatRange > u18_6 = FFloatRange::Union(r18_1, r18_6);
+	{
+		TestEqual(TEXT("[7, 14] unioned with [6, 7) must result in one range"), u18_6.Num(), 1);
+		TestEqual(TEXT("[7, 14] unioned with [6, 7) must result in {[6, 14]}"), u18_6[0], FFloatRange(FFloatRangeBound::Inclusive(6.f), FFloatRangeBound::Inclusive(14.f)));
+	}
+
+	TArray<FFloatRange > u18_7 = FFloatRange::Union(r18_2, r18_6);
+	{
+		TestEqual(TEXT("(7, 14) unioned with [6, 7) must result in two ranges"), u18_7.Num(), 2);
+		
+		bool b18_1 = (u18_7[0] == r18_2);
+		bool b18_2 = (u18_7[1] == r18_6);
+
+		TestTrue(TEXT("(7, 14) unioned with [6, 7) must result in {(7, 14), [6, 7)}"), b18_1 && b18_2);
+	}
+	
+	TestEqual(TEXT("An pair of empty ranges unioned must yield no ranges."), FFloatRange::Union(FFloatRange::Empty(), FFloatRange::Empty()).Num(), 0);
+}
+
+PRAGMA_ENABLE_OPTIMIZATION
