Tests converted from 18848115 made to work multi-platform.

diff --git a/Engine/Source/Runtime/Core/Tests/Misc/AsciiSetTest.cpp b/Engine/Source/Runtime/Core/Tests/Misc/AsciiSetTest.cpp
new file mode 100644
--- /dev/null
+++ b/Engine/Source/Runtime/Core/Tests/Misc/AsciiSetTest.cpp
@@ -0,0 +1,92 @@
+// Copyright Epic Games, Inc. All Rights Reserved.
+
+
+#include "Containers/StringView.h"
+#include "Misc/AsciiSet.h"
+#include "TestHarness.h"
+
+TEST_CASE("Core::Misc::FAsciiSet::AsciiSet", "[Core][Misc][Smoke]")
+{
+	constexpr FAsciiSet Whitespaces(" \v\f\t\r\n");
+	TestTrue(TEXT("Contains"), Whitespaces.Contains(' '));
+	TestTrue(TEXT("Contains"), Whitespaces.Contains('\n'));
+	TestFalse(TEXT("Contains"), Whitespaces.Contains('a'));
+	TestFalse(TEXT("Contains no extended ASCII"), Whitespaces.Contains('\x80'));
+	TestFalse(TEXT("Contains no extended ASCII"), Whitespaces.Contains('\xA0'));
+	TestFalse(TEXT("Contains no extended ASCII"), Whitespaces.Contains('\xFF'));
+
+	constexpr FAsciiSet Aa("Aa");
+	uint32 ANum = 0;
+	for (char32_t C = 0; C < 512; ++C)
+	{
+		ANum += Aa.Contains(C);
+	}
+	TestEqual("Contains no wide", ANum, 2u);
+
+	constexpr FAsciiSet NonWhitespaces = ~Whitespaces;
+	uint32 WhitespaceNum = 0;
+	for (uint8 Char = 0; Char < 128; ++Char)
+	{
+		WhitespaceNum += !!Whitespaces.Test(Char);
+		TestEqual(TEXT("Inverse"), !!Whitespaces.Test(Char), !NonWhitespaces.Test(Char));
+	}
+	TestEqual(TEXT("Num"), WhitespaceNum, 6);
+
+	// TODO [ML]
+	//TestEqual(TEXT("Skip"), FAsciiSet::Skip(TEXT("  \t\tHello world!"), Whitespaces), TEXT("Hello world!"));
+	//TestEqual(TEXT("Skip"), FAsciiSet::Skip(TEXT("Hello world!"), Whitespaces), TEXT("Hello world!"));
+	//TestEqual(TEXT("Skip to extended ASCII"), FAsciiSet::Skip(" " "\xA0" " abc", Whitespaces), "\xA0" " abc");
+	//TestEqual(TEXT("Skip to wide"), FAsciiSet::Skip(TEXT(" 变 abc"), Whitespaces), TEXT("变 abc"));
+	TestEqual(TEXT("AdvanceToFirst"), *FAsciiSet::FindFirstOrEnd("NonWhitespace\t \nNonWhitespace", Whitespaces), '\t');
+	TestEqual(TEXT("AdvanceToLast"), *FAsciiSet::FindLastOrEnd("NonWhitespace\t \nNonWhitespace", Whitespaces), '\n');
+	TestEqual(TEXT("AdvanceToLast"), *FAsciiSet::FindLastOrEnd("NonWhitespace\t NonWhitespace\n", Whitespaces), '\n');
+	TestEqual(TEXT("AdvanceToFirst"), *FAsciiSet::FindFirstOrEnd("NonWhitespaceNonWhitespace", Whitespaces), '\0');
+	TestEqual(TEXT("AdvanceToLast"), *FAsciiSet::FindLastOrEnd("NonWhitespaceNonWhitespace", Whitespaces), '\0');
+
+	constexpr FAsciiSet Lowercase("abcdefghijklmnopqrstuvwxyz");
+	TestEqual(TEXT("TrimPrefixWithout"), FAsciiSet::TrimPrefixWithout("ABcdEF"_ASV, Lowercase), "cdEF"_ASV);
+	TestEqual(TEXT("FindPrefixWithout"), FAsciiSet::FindPrefixWithout("ABcdEF"_ASV, Lowercase), "AB"_ASV);
+	TestEqual(TEXT("TrimSuffixWithout"), FAsciiSet::TrimSuffixWithout("ABcdEF"_ASV, Lowercase), "ABcd"_ASV);
+	TestEqual(TEXT("FindSuffixWithout"), FAsciiSet::FindSuffixWithout("ABcdEF"_ASV, Lowercase), "EF"_ASV);
+	TestEqual(TEXT("TrimPrefixWithout none"), FAsciiSet::TrimPrefixWithout("same"_ASV, Lowercase), "same"_ASV);
+	TestEqual(TEXT("FindPrefixWithout none"), FAsciiSet::FindPrefixWithout("same"_ASV, Lowercase), ""_ASV);
+	TestEqual(TEXT("TrimSuffixWithout none"), FAsciiSet::TrimSuffixWithout("same"_ASV, Lowercase), "same"_ASV);
+	TestEqual(TEXT("FindSuffixWithout none"), FAsciiSet::FindSuffixWithout("same"_ASV, Lowercase), ""_ASV);
+	TestEqual(TEXT("TrimPrefixWithout empty"), FAsciiSet::TrimPrefixWithout(""_ASV, Lowercase), ""_ASV);
+	TestEqual(TEXT("FindPrefixWithout empty"), FAsciiSet::FindPrefixWithout(""_ASV, Lowercase), ""_ASV);
+	TestEqual(TEXT("TrimSuffixWithout empty"), FAsciiSet::TrimSuffixWithout(""_ASV, Lowercase), ""_ASV);
+	TestEqual(TEXT("FindSuffixWithout empty"), FAsciiSet::FindSuffixWithout(""_ASV, Lowercase), ""_ASV);
+
+
+	auto TestHasFunctions = [&](auto MakeString)
+	{
+		constexpr FAsciiSet XmlEscapeChars("&<>\"'");
+		TestTrue(TEXT("None"), FAsciiSet::HasNone(MakeString("No escape chars"), XmlEscapeChars));
+		TestFalse(TEXT("Any"), FAsciiSet::HasAny(MakeString("No escape chars"), XmlEscapeChars));
+		TestFalse(TEXT("Only"), FAsciiSet::HasOnly(MakeString("No escape chars"), XmlEscapeChars));
+
+		TestTrue(TEXT("None"), FAsciiSet::HasNone(MakeString(""), XmlEscapeChars));
+		TestFalse(TEXT("Any"), FAsciiSet::HasAny(MakeString(""), XmlEscapeChars));
+		TestTrue(TEXT("Only"), FAsciiSet::HasOnly(MakeString(""), XmlEscapeChars));
+
+		TestFalse(TEXT("None"), FAsciiSet::HasNone(MakeString("&<>\"'"), XmlEscapeChars));
+		TestTrue(TEXT("Any"), FAsciiSet::HasAny(MakeString("&<>\"'"), XmlEscapeChars));
+		TestTrue(TEXT("Only"), FAsciiSet::HasOnly(MakeString("&<>\"'"), XmlEscapeChars));
+
+		TestFalse(TEXT("None"), FAsciiSet::HasNone(MakeString("&<>\"' and more"), XmlEscapeChars));
+		TestTrue(TEXT("Any"), FAsciiSet::HasAny(MakeString("&<>\"' and more"), XmlEscapeChars));
+		TestFalse(TEXT("Only"), FAsciiSet::HasOnly(MakeString("&<>\"' and more"), XmlEscapeChars));
+	};
+	TestHasFunctions([](const char* Str) { return Str; });
+	TestHasFunctions([](const char* Str) { return FAnsiStringView(Str); });
+	TestHasFunctions([](const char* Str) { return FString(Str); });
+
+
+	constexpr FAsciiSet Abc("abc");
+	constexpr FAsciiSet Abcd = Abc + 'd';
+	TestTrue(TEXT("Add"), Abcd.Contains('a'));
+	TestTrue(TEXT("Add"), Abcd.Contains('b'));
+	TestTrue(TEXT("Add"), Abcd.Contains('c'));
+	TestTrue(TEXT("Add"), Abcd.Contains('d'));
+	TestFalse(TEXT("Add"), Abcd.Contains('e'));
+}
\ No newline at end of file
diff --git a/Engine/Source/Runtime/Core/Tests/Misc/ByteSwapTest.cpp b/Engine/Source/Runtime/Core/Tests/Misc/ByteSwapTest.cpp
new file mode 100644
--- /dev/null
+++ b/Engine/Source/Runtime/Core/Tests/Misc/ByteSwapTest.cpp
@@ -0,0 +1,93 @@
+// Copyright Epic Games, Inc. All Rights Reserved.
+
+#include "CoreTypes.h"
+#include "Misc/ByteSwap.h"
+#include "TestHarness.h"
+
+/**
+ * Test byte swapping algorithms.
+ */
+
+TEST_CASE("Core::Misc::ByteSwap::Smoke Test", "[Core][Misc][Smoke]")
+{
+	int16  ValS16 = static_cast<int16>(0x1122);
+	uint16 ValU16 = static_cast<uint16>(0x1122);
+	int32  ValS32 = 0xFFEE1122L;
+	uint32 ValU32 = 0xFFEE1122UL;
+	int64  ValS64 = 0xFFEEDDCC11223344ll;
+	uint64 ValU64 = 0xFFEEDDCC11223344ull;
+	float  ValF   = *reinterpret_cast<const float*>(&ValU32);
+	double ValD   = *reinterpret_cast<const double*>(&ValU64);
+	uint32 ExpectedF = ByteSwap(ValU32);
+	uint64 ExpectedD = ByteSwap(ValU64);
+	char16_t ValCh16 = static_cast<char16_t>(0x2233);
+
+	TestTrue("Swapping singned int16 value",  ByteSwap(ValS16) == 0x2211);
+	TestTrue("Swapping unsigned int16 value", ByteSwap(ValU16) == 0x2211);
+	TestTrue("Swapping singned int32 value",  ByteSwap(ValS32) == 0x2211EEFF);
+	TestTrue("Swapping unsigned int32 value", ByteSwap(ValU32) == 0x2211EEFF);
+	TestTrue("Swapping singned int64 value",  ByteSwap(ValS64) == 0x44332211CCDDEEFFll);
+	TestTrue("Swapping unsigned int64 value", ByteSwap(ValU64) == 0x44332211CCDDEEFFull);
+	TestTrue("Swapping float value",          ByteSwap(ValF)   == *reinterpret_cast<const float*>(&ExpectedF));
+	TestTrue("Swapping double value",         ByteSwap(ValD)   == *reinterpret_cast<const double*>(&ExpectedD));
+	TestTrue("Swapping char16_t value",       ByteSwap(ValCh16)== 0x3322);
+} 
+
+// The byte swap benchmarking test are useful to compare the intrinsic implementation vs the generic implementation. Normally, the intrinsic is expected to be faster, but in some
+// cases, the intrinsic vs generic speed is within margin of error in optimized build. So the test result are not always consistent. The code remains there in case a new
+// implementation needs to be testes. Here some observations:
+//     - VC++ 2019 (16.4.3): The intrinsic versions are consistently (and significantly) faster than the generic version. The compiler poorly optimize the generic version.
+//     - Apple Clang (11.0.0): No performance change observed in release. Clang generates the same assembly for the intrinsic and the generic functions.
+//     - Linux Clang (8.0.1): No performance change observed in release. Clang generates the same assembly for the intrinsic and the generic functions.
+#
+template<typename T, typename ByteSwapFn>
+TPair<FTimespan, T> BenchmarkByteswapping(T InitialValue, uint64 LoopCount, ByteSwapFn InByteSwapFn)
+{
+	static_assert(TIsIntegral<T>::Value, "The test only works with integers");
+	static_assert(!TIsSigned<T>::Value, "The test must support overflowing T when making the sum, value cannot be signed.");
+
+	T Sum = 0; // The sum value is not very relevant, but it is an observable result the compiler cannot ignore (prevent to optimize away the entire loop).
+	FDateTime StartTime = FDateTime::UtcNow();
+	for (uint64 I = 0; I < LoopCount; ++I)
+	{
+		Sum += InByteSwapFn(InitialValue + static_cast<T>(I)); // Expect InByteSwapFn to be inlined because the compiler knows the exact function to call.
+	}
+	FTimespan Duration = FDateTime::UtcNow() - StartTime;
+	return MakeTuple(Duration, Sum);
+}
+
+TEST_CASE("Core::Misc::ByteSwap::Perf Test", "[.][Core][Misc][Perf]")
+{
+#ifdef UE_BYTESWAP_INTRINSIC_PRIVATE_16
+	{
+		uint16 InitialValue = 0xF0F0;
+		TPair<FTimespan, uint16> Intrinsic = BenchmarkByteswapping(InitialValue, 1000000000, [](uint16 Val){ return ByteSwap(Val); });
+		TPair<FTimespan, uint16> Generic   = BenchmarkByteswapping(InitialValue, 1000000000, [](uint16 Val){ return Internal::ByteSwapGeneric16(Val); });
+		TestTrue("Swapping uint16 bytes is faster uing the compiler intrinsic than the generic implementation", Intrinsic.Get<0>() <= Generic.Get<0>());
+		TestTrue("Swapping uint16 bytes using intrinsic and generic algorithm produce the same values", Intrinsic.Get<1>() == Generic.Get<1>());
+		GLog->Logf(TEXT("Swapping 2 bytes using intrinsic is %s faster than generic version"), *LexToString(Generic.Get<0>().GetTotalMicroseconds() / Intrinsic.Get<0>().GetTotalMicroseconds()));
+	}
+#endif
+
+#ifdef UE_BYTESWAP_INTRINSIC_PRIVATE_32
+	{
+		uint32 InitialValue = 0xFF00FF00ul;
+		TPair<FTimespan, uint32> Intrinsic = BenchmarkByteswapping(InitialValue, 1000000000, [](uint32 Val){ return ByteSwap(Val); });
+		TPair<FTimespan, uint32> Generic   = BenchmarkByteswapping(InitialValue, 1000000000, [](uint32 Val){ return Internal::ByteSwapGeneric32(Val); });
+		TestTrue("Swapping uint32 bytes is faster uing the compiler intrinsic than the generic implementation", Intrinsic.Get<0>() <= Generic.Get<0>());
+		TestTrue("Swapping uint32 bytes using intrinsic and generic algorithm produce the same values", Intrinsic.Get<1>() == Generic.Get<1>());
+		GLog->Logf(TEXT("Swapping 4 bytes using intrinsic is %s faster than generic version"), *LexToString(Generic.Get<0>().GetTotalMicroseconds() / Intrinsic.Get<0>().GetTotalMicroseconds()));
+	}
+#endif
+
+#ifdef UE_BYTESWAP_INTRINSIC_PRIVATE_64
+	{
+		uint64 InitialValue = 0xFF00FF00FF00FF00ull;
+		TPair<FTimespan, uint64> Intrinsic = BenchmarkByteswapping(InitialValue, 1000000000, [](uint64 Val){ return ByteSwap(Val); });
+		TPair<FTimespan, uint64> Generic   = BenchmarkByteswapping(InitialValue, 1000000000, [](uint64 Val){ return Internal::ByteSwapGeneric64(Val); });
+		TestTrue("Swapping uint64 bytes is faster uing the compiler intrinsic than the generic implementation", Intrinsic.Get<0>() <= Generic.Get<0>());
+		TestTrue("Swapping uint64 bytes using intrinsic and generic algorithm produce the same values", Intrinsic.Get<1>() == Generic.Get<1>());
+		GLog->Logf(TEXT("Swapping 8 bytes using intrinsic is %s faster than generic version"), *LexToString(Generic.Get<0>().GetTotalMicroseconds() / Intrinsic.Get<0>().GetTotalMicroseconds()));
+	}
+#endif
+}
diff --git a/Engine/Source/Runtime/Core/Tests/Misc/CharTest.cpp b/Engine/Source/Runtime/Core/Tests/Misc/CharTest.cpp
new file mode 100644
--- /dev/null
+++ b/Engine/Source/Runtime/Core/Tests/Misc/CharTest.cpp
@@ -0,0 +1,45 @@
+// Copyright Epic Games, Inc. All Rights Reserved.
+
+#include "Misc/Char.h"
+#include <locale.h>
+#include <ctype.h>
+#include <wctype.h>
+#include "TestHarness.h"
+
+namespace crt
+{
+	int tolower(ANSICHAR c) { return ::tolower(c); }
+	int toupper(ANSICHAR c) { return ::toupper(c); }
+
+	int tolower(WIDECHAR c) { return ::towlower(c); }
+	int toupper(WIDECHAR c) { return ::towupper(c); }
+}
+
+template<typename CharType>
+void RunCharTests(uint32 MaxChar)
+{
+	for (uint32 I = 0; I < MaxChar; ++I)
+	{
+		CharType C = (CharType)I;
+		TestEqual("TChar::ToLower()", TChar<CharType>::ToLower(C), crt::tolower(C));
+		TestEqual("TChar::ToUpper()", TChar<CharType>::ToUpper(C), crt::toupper(C));
+	}
+}
+
+TEST_CASE("Core::Misc::Char::Smoke Test", "[Core][Misc][Perf]")
+{
+	const char* CurrentLocale = setlocale(LC_CTYPE, nullptr);
+	if (CurrentLocale == nullptr)
+	{
+		TestFalse(FString::Printf(TEXT("Locale is null but should be \"C\". Did something call setlocale()?")), CurrentLocale == nullptr);
+	}
+	else if (strcmp("C", CurrentLocale))
+	{
+		TestFalse(FString::Printf(TEXT("Locale is null but should be \"C\". Did something call setlocale()?")), strcmp("C", CurrentLocale)!=0);
+	}
+	else
+	{
+		RunCharTests<ANSICHAR>(128);
+		RunCharTests<WIDECHAR>(65536);
+	}
+}
\ No newline at end of file
diff --git a/Engine/Source/Runtime/Core/Tests/Misc/CircularBufferTest.cpp b/Engine/Source/Runtime/Core/Tests/Misc/CircularBufferTest.cpp
new file mode 100644
--- /dev/null
+++ b/Engine/Source/Runtime/Core/Tests/Misc/CircularBufferTest.cpp
@@ -0,0 +1,41 @@
+// Copyright Epic Games, Inc. All Rights Reserved.
+
+#include "CoreTypes.h"
+#include "Containers/UnrealString.h"
+#include "Containers/CircularBuffer.h"
+#include "TestHarness.h"
+
+TEST_CASE("Core::Misc::TCircularBuffer::Smoke Test", "[Core][Misc][Smoke]")
+{
+	// buffer capacity
+	TCircularBuffer<int32> b1_1(127);
+	TCircularBuffer<int32> b1_2(128);
+	TCircularBuffer<int32> b1_3(129);
+
+	TestEqual(TEXT("Buffer capacity of 127 must be rounded up to 128"), b1_1.Capacity(), 128u);
+	TestEqual(TEXT("Buffer capacity of 128 must not change"), b1_2.Capacity(), 128u);
+	TestEqual(TEXT("Buffer capacity of 129 must be rounded up to 256"), b1_3.Capacity(), 256u);
+
+	// initial values
+	TCircularBuffer<int32> b2_1(64, 666);
+
+	for (uint32 Index = 0; Index < b2_1.Capacity(); ++Index)
+	{
+		TestEqual(FString::Printf(TEXT("Initial value must be correct (%i)"), Index), b2_1[Index], 666);
+	}
+
+	// indexing
+	TCircularBuffer<int32> b3_1(64, 0);
+
+	TestEqual(TEXT("Next index from 0 must be 1"), b3_1.GetNextIndex(0), 1u);
+	TestEqual(TEXT("Next index from 63 must be 0"), b3_1.GetNextIndex(63), 0u);
+	TestEqual(TEXT("Next index from 64 must be 1"), b3_1.GetNextIndex(64), 1u);
+
+	b3_1[0] = 42;
+	b3_1[65] = 42;
+
+	TestEqual(TEXT("Index 0 must be written and read correctly"), b3_1[0], 42);
+	TestEqual(TEXT("Index 1 must be written and read correctly"), b3_1[1], 42);
+	TestEqual(TEXT("Index 65 must be written and read correctly"), b3_1[65], 42);
+}
+
diff --git a/Engine/Source/Runtime/Core/Tests/Misc/CircularQueueTest.cpp b/Engine/Source/Runtime/Core/Tests/Misc/CircularQueueTest.cpp
new file mode 100644
--- /dev/null
+++ b/Engine/Source/Runtime/Core/Tests/Misc/CircularQueueTest.cpp
@@ -0,0 +1,129 @@
+// Copyright Epic Games, Inc. All Rights Reserved.
+
+#include "CoreTypes.h"
+#include "Containers/CircularQueue.h"
+#include "TestHarness.h"
+
+TEST_CASE("Core::Misc::TCircularQueue::Smoke Test", "[Core][Misc][Smoke]")
+{
+	const uint32 QueueSize = 8;
+
+	// empty queue
+	{
+		TCircularQueue<int32> Queue(QueueSize);
+
+		TestEqual(TEXT("Newly created queues must have zero elements"), Queue.Count(), 0u);
+		TestTrue(TEXT("Newly created queues must be empty"), Queue.IsEmpty());
+		TestFalse(TEXT("Newly created queues must not be full"), Queue.IsFull());
+
+		int32 Value;
+		TestFalse(TEXT("Peek must fail on an empty queue"), Queue.Peek(Value));
+		TestNull(TEXT("Peek must fail on an empty queue"), Queue.Peek());
+	}
+
+	// partially filled
+	{
+		TCircularQueue<int32> Queue(QueueSize);
+		int32 Value = 0;
+
+		TestTrue(TEXT("Adding to an empty queue must succeed"), Queue.Enqueue(666));
+		TestEqual(TEXT("After adding to an empty queue it must have one element"), Queue.Count(), 1u);
+		TestFalse(TEXT("Partially filled queues must not be empty"), Queue.IsEmpty());
+		TestFalse(TEXT("Partially filled queues must not be full"), Queue.IsFull());
+		TestTrue(TEXT("Peeking at a partially filled queue must succeed"), Queue.Peek(Value));
+		TestEqual(TEXT("The peeked at value must be correct"), Value, 666);
+
+		const int32* PeekValue = Queue.Peek();
+		TestNotNull(TEXT("Peeking at a partially filled queue must succeed"), PeekValue);
+		TestEqual(TEXT("The peeked at value must be correct"), *PeekValue, 666);
+	}
+
+	// full queue
+	for (uint32 PeekType = 0; PeekType < 2; PeekType++)
+	{
+		TCircularQueue<int32> Queue(QueueSize);
+
+		for (int32 Index = 0; Index < QueueSize - 1; ++Index)
+		{
+			TestTrue(TEXT("Adding to non-full queue must succeed"), Queue.Enqueue(Index));
+		}
+
+		TestFalse(TEXT("Full queues must not be empty"), Queue.IsEmpty());
+		TestTrue(TEXT("Full queues must be full"), Queue.IsFull());
+		TestFalse(TEXT("Adding to full queue must fail"), Queue.Enqueue(666));
+
+		int32 Value = 0;
+
+		for (int32 Index = 0; Index < QueueSize - 1; ++Index)
+		{
+			if (PeekType == 0)
+			{
+				TestTrue(TEXT("Peeking at a non-empty queue must succeed"), Queue.Peek(Value));
+				TestEqual(TEXT("The peeked at value must be correct"), Value, Index);
+
+				TestTrue(TEXT("Removing from a non-empty queue must succeed"), Queue.Dequeue(Value));
+				TestEqual(TEXT("The removed value must be correct"), Value, Index);
+			}
+			else
+			{
+				const int32* PeekValue = Queue.Peek();
+				TestNotNull(TEXT("Peeking at a non-empty queue must succeed"), PeekValue);
+				TestEqual(TEXT("The peeked at value must be correct"), *PeekValue, Index);
+
+				TestTrue(TEXT("Removing from a non-empty queue must succeed"), Queue.Dequeue());
+			}
+		}
+
+		TestTrue(TEXT("A queue that had all items removed must be empty"), Queue.IsEmpty());
+		TestFalse(TEXT("A queue that had all items removed must not be full"), Queue.IsFull());
+	}
+
+	// queue with index wrapping around
+	{
+		TCircularQueue<int32> Queue(QueueSize);
+
+		//Fill queue
+		for (int32 Index = 0; Index < QueueSize - 1; ++Index)
+		{
+			TestTrue(TEXT("Adding to non-full queue must succeed"), Queue.Enqueue(Index));
+		}
+
+		int32 Value = 0;
+		const int32 ExpectedSize = QueueSize - 1;
+		for (int32 Index = 0; Index < QueueSize; ++Index)
+		{
+			TestEqual(TEXT("Number of elements must be valid for all permutation of Tail and Head"), Queue.Count(), ExpectedSize);
+			TestTrue(TEXT("Removing from a non-empty queue must succeed"), Queue.Dequeue(Value));
+			TestTrue(TEXT("Adding to non-full queue must succeed"), Queue.Enqueue(Index));
+		}
+	}
+
+	// Non-zero initialization - ensures the backing store constructs and destructs non-POD objects properly.
+	{
+		static uint32 Const;
+		static uint32 Dest;
+		static uint32 CopyConst;
+		struct FNonPOD
+		{
+			FNonPOD() { Const++; }
+			~FNonPOD() { Dest++; }
+			FNonPOD(const FNonPOD&) { CopyConst++; }
+		};
+
+		// The current implementation of TCircularQueue doesn't call the held object constructors on but does call the
+		// destructors. Verify that (somewhat surprising, if typically ok for POD) behavior first.
+		Const = 0;
+		Dest = 0;
+		CopyConst = 0;
+		{
+			TCircularQueue<FNonPOD> Queue(QueueSize);
+
+			TestEqual(TEXT("The constructor should not run"), Const, 0);
+			TestEqual(TEXT("The destructor should not run"), Dest, 0);
+			TestEqual(TEXT("The copy constructor should not run"), CopyConst, 0);
+		}
+		TestEqual(TEXT("The constructor should not run"), Const, 0);
+		TestEqual(TEXT("The destructor should run"), Dest, QueueSize);
+		TestEqual(TEXT("The copy constructor should not run"), CopyConst, 0);
+	}
+}
\ No newline at end of file
diff --git a/Engine/Source/Runtime/Core/Tests/Misc/DateTimeTest.cpp b/Engine/Source/Runtime/Core/Tests/Misc/DateTimeTest.cpp
new file mode 100644
--- /dev/null
+++ b/Engine/Source/Runtime/Core/Tests/Misc/DateTimeTest.cpp
@@ -0,0 +1,135 @@
+// Copyright Epic Games, Inc. All Rights Reserved.
+
+#include "CoreTypes.h"
+#include "Containers/UnrealString.h"
+#include "Misc/DateTime.h"
+#include "TestHarness.h"
+
+#define TestUnixEquivalent(Desc, A, B) TestTrue(FString::Printf(TEXT("%s - A=%d B=%d"), Desc, (A).ToUnixTimestamp(), (B)), (A).ToUnixTimestamp() == (B) );
+#define TestYear(Desc, A, B)  TestTrue(FString::Printf(TEXT("%s - A=%d B=%d"), Desc, (A.GetYear()), (B)), ((A.GetYear()) == (B)));
+#define TestMonth(Desc, A, B) TestTrue(FString::Printf(TEXT("%s - A=%d B=%d"), Desc, (A.GetMonth()), B), (A.GetMonth()) == (B));
+#define TestMonthOfYear(Desc, A, B) TestTrue(FString::Printf(TEXT("%s - A=%d B=%d"), Desc ,(static_cast<int32>(A.GetMonthOfYear())), static_cast<int32>(B)), (A.GetMonthOfYear()) == (B));
+#define TestDay(Desc, A, B) TestTrue(FString::Printf(TEXT("%s - A=%d B=%d"), Desc, (A.GetDay()), (B)), (A.GetDay()==(B)));
+#define TestHour(Desc, A, B) TestTrue(FString::Printf(TEXT("%s - A=%d B=%d"), Desc, (A.GetHour()), (B)), (A.GetHour()==(B)));
+#define TestMinute(Desc, A, B) TestTrue(FString::Printf(TEXT("%s - A=%d B=%d"), Desc, (A.GetMinute()), (B)), (A.GetMinute()) == (B));
+#define TestSecond(Desc, A, B) TestTrue(FString::Printf(TEXT("%s - A=%d B=%d"), Desc, (A.GetSecond()), (B)), (A.GetSecond()) == (B));
+#define TestMillisecond(Desc, A, B) TestTrue(FString::Printf(TEXT("%s - A=%d B=%d"), Desc, (A.GetMillisecond()) ,(B)), (A.GetMillisecond()) == (B));
+
+TEST_CASE("Core::Misc::DateTime::Smoke", "[Core][Misc][Smoke][PlatformFlaky]")
+{
+	const int64 UnixEpochTimestamp = 0;
+	const int64 UnixBillenniumTimestamp = 1000000000;
+	const int64 UnixOnesTimestamp = 1111111111;
+	const int64 UnixDecimalSequenceTimestamp = 1234567890;
+
+	const FDateTime UnixEpoch = FDateTime::FromUnixTimestamp(UnixEpochTimestamp);
+	const FDateTime UnixBillennium = FDateTime::FromUnixTimestamp(UnixBillenniumTimestamp);
+	const FDateTime UnixOnes = FDateTime::FromUnixTimestamp(UnixOnesTimestamp);
+	const FDateTime UnixDecimalSequence = FDateTime::FromUnixTimestamp(UnixDecimalSequenceTimestamp);
+
+	const FDateTime TestDateTime(2013, 8, 14, 12, 34, 56, 789);
+
+	TestUnixEquivalent(TEXT("Testing Unix Epoch Ticks"), UnixEpoch, UnixEpochTimestamp);
+	TestUnixEquivalent(TEXT("Testing Unix Billennium Ticks"), UnixBillennium, UnixBillenniumTimestamp);
+	TestUnixEquivalent(TEXT("Testing Unix Ones Ticks"), UnixOnes, UnixOnesTimestamp);
+	TestUnixEquivalent(TEXT("Testing Unix Decimal Sequence Ticks"), UnixDecimalSequence, UnixDecimalSequenceTimestamp);
+
+	TestYear(TEXT("Testing Unix Epoch Year"), UnixEpoch, 1970);
+	TestMonth(TEXT("Testing Unix Epoch Month"), UnixEpoch, 1);
+	TestMonthOfYear(TEXT("Testing Unix Epoch MonthOfYear"), UnixEpoch, EMonthOfYear::January);
+	TestDay(TEXT("Testing Unix Epoch Day"), UnixEpoch, 1);
+	TestHour(TEXT("Testing Unix Epoch Hour"), UnixEpoch, 0);
+	TestMinute(TEXT("Testing Unix Epoch Minute"), UnixEpoch, 0);
+	TestSecond(TEXT("Testing Unix Epoch Second"), UnixEpoch, 0);
+	TestMillisecond(TEXT("Testing Unix Epoch Millisecond"), UnixEpoch, 0);
+
+	TestYear(TEXT("Testing Unix Billennium Year"), UnixBillennium, 2001);
+	TestMonth(TEXT("Testing Unix Billennium MonthOfYear"), UnixBillennium, 9);
+	TestMonthOfYear(TEXT("Testing Unix Billennium Month"), UnixBillennium, EMonthOfYear::September);
+	TestDay(TEXT("Testing Unix Billennium Day"), UnixBillennium, 9);
+	TestHour(TEXT("Testing Unix Billennium Hour"), UnixBillennium, 1);
+	TestMinute(TEXT("Testing Unix Billennium Minute"), UnixBillennium, 46);
+	TestSecond(TEXT("Testing Unix Billennium Second"), UnixBillennium, 40);
+	TestMillisecond(TEXT("Testing Unix Billennium Millisecond"), UnixBillennium, 0);
+
+	TestYear(TEXT("Testing Unix Ones Year"), UnixOnes, 2005);
+	TestMonth(TEXT("Testing Unix Ones Month"), UnixOnes, 3);
+	TestMonthOfYear(TEXT("Testing Unix Ones MonthOfYear"), UnixOnes, EMonthOfYear::March);
+	TestDay(TEXT("Testing Unix Ones Day"), UnixOnes, 18);
+	TestHour(TEXT("Testing Unix Ones Hour"), UnixOnes, 1);
+	TestMinute(TEXT("Testing Unix Ones Minute"), UnixOnes, 58);
+	TestSecond(TEXT("Testing Unix Ones Second"), UnixOnes, 31);
+	TestMillisecond(TEXT("Testing Unix Ones Millisecond"), UnixOnes, 0);
+
+	TestYear(TEXT("Testing Unix Decimal Sequence Year"), UnixDecimalSequence, 2009);
+	TestMonth(TEXT("Testing Unix Decimal Sequence Month"), UnixDecimalSequence, 2);
+	TestMonthOfYear(TEXT("Testing Unix Decimal Sequence MonthOfYear"), UnixDecimalSequence, EMonthOfYear::February);
+	TestDay(TEXT("Testing Unix Decimal Sequence Day"), UnixDecimalSequence, 13);
+	TestHour(TEXT("Testing Unix Decimal Sequence Hour"), UnixDecimalSequence, 23);
+	TestMinute(TEXT("Testing Unix Decimal Sequence Minute"), UnixDecimalSequence, 31);
+	TestSecond(TEXT("Testing Unix Decimal Sequence Second"), UnixDecimalSequence, 30);
+	TestMillisecond(TEXT("Testing Unix Decimal Sequence Millisecond"), UnixDecimalSequence, 0);
+
+	TestYear(TEXT("Testing Test Date Time Year"), TestDateTime, 2013);
+	TestMonth(TEXT("Testing Test Date Time Month"), TestDateTime, 8);
+	TestMonthOfYear(TEXT("Testing Test Date Time MonthOfYear"), TestDateTime, EMonthOfYear::August);
+	TestDay(TEXT("Testing Test Date Time Day"), TestDateTime, 14);
+	TestHour(TEXT("Testing Test Date Time Hour"), TestDateTime, 12);
+	TestMinute(TEXT("Testing Test Date Time Minute"), TestDateTime, 34);
+	TestSecond(TEXT("Testing Test Date Time Second"), TestDateTime, 56);
+	TestMillisecond(TEXT("Testing Test Date Time Millisecond"), TestDateTime, 789);
+
+	FDateTime ParsedDateTime;
+
+	TestFalse(TEXT("Parsing an empty ISO string must fail"), FDateTime::ParseIso8601(TEXT(""), ParsedDateTime));
+
+	FDateTime::ParseIso8601(TEXT("2019-05-22"), ParsedDateTime);
+	TestEqual(TEXT("Testing ISO 8601 date"), ParsedDateTime, FDateTime(2019, 5, 22));
+
+	FDateTime::ParseIso8601(TEXT("2019-05-20T19:41:38+01:30"), ParsedDateTime);
+	TestEqual(TEXT("Testing ISO 8601 with +hh:mm timezone info"), ParsedDateTime, FDateTime( 2019, 5, 20, 18, 11, 38 ));
+	FDateTime::ParseIso8601(TEXT("2019-05-20T19:41:38-01:30"), ParsedDateTime);
+	TestEqual(TEXT("Testing ISO 8601 with -hh:mm timezone info"), ParsedDateTime, FDateTime( 2019, 5, 20, 21, 11, 38 ));
+	FDateTime::ParseIso8601(TEXT("2019-05-20T19:41:38+0030"), ParsedDateTime);
+	TestEqual(TEXT("Testing ISO 8601 with +hhmm timezone info"), ParsedDateTime, FDateTime( 2019, 5, 20, 19, 11, 38 ));
+	FDateTime::ParseIso8601(TEXT("2019-05-20T19:41:38-01"), ParsedDateTime);
+	TestEqual(TEXT("Testing ISO 8601 with -hh timezone info"), ParsedDateTime, FDateTime( 2019, 5, 20, 20, 41, 38 ));
+
+	for (double JulianDay : {0., 1000., 1721425.5, 1721425.0, FDateTime::MinValue().GetJulianDay(), FDateTime::MaxValue().GetJulianDay()})
+	{
+		TestEqual(TEXT("convertion from/to JulianDay is stable"), FDateTime::FromJulianDay(JulianDay).GetJulianDay(), JulianDay);
+	}
+
+	FDateTime Date{ 2019, 5, 20, 18, 11, 38 };
+	FDateTime DateMidnight = Date.GetDate();
+	TestEqual(TEXT("GetDate returns a FDateTime at midnight"), DateMidnight.GetHour(), 0);
+
+	FDateTime DateNoon = DateMidnight + FTimespan::FromHours(12);
+	TestEqual(TEXT("GetDate returns a FDateTime at midnight"), DateNoon.GetHour(), 12);
+
+	double JulianDayMidnight = DateMidnight.GetJulianDay();
+	double JulianDayNoon = DateNoon.GetJulianDay();
+	TestEqual(TEXT("At midnight, fractionnal part of the JulianDay value should be 0.5"), JulianDayMidnight - (int)JulianDayMidnight, 0.5);
+	TestEqual(TEXT("A 12h timespan adds half a julianday"), JulianDayNoon - JulianDayMidnight, 0.5);
+
+	double OffsetDayCount = 12345.;
+	FDateTime OffsetDate = Date + FTimespan::FromDays(OffsetDayCount);
+	TestEqual(TEXT("An offset by a given numer of days leads to a similar JulianDay offset"), OffsetDate.GetJulianDay() - Date.GetJulianDay(), OffsetDayCount);
+
+	int32 Year, Month, Day;
+	DateMidnight.GetDate(Year, Month, Day);
+	FDateTime PreviousDay(DateMidnight.GetTicks() - 1);
+	int32 YearPrev, MonthPrev, DayPrev;
+	PreviousDay.GetDate(YearPrev, MonthPrev, DayPrev);
+	TestNotEqual(TEXT("One tick before a date at midnight leads to the previous date"), Day, DayPrev);
+}
+
+
+#undef TestUnixEquivalent
+#undef TestYear
+#undef TestMonth
+#undef TestDay
+#undef TestHour
+#undef TestMinute
+#undef TestSecond
+#undef TestMillisecond
diff --git a/Engine/Source/Runtime/Core/Tests/Misc/GuidTest.cpp b/Engine/Source/Runtime/Core/Tests/Misc/GuidTest.cpp
new file mode 100644
--- /dev/null
+++ b/Engine/Source/Runtime/Core/Tests/Misc/GuidTest.cpp
@@ -0,0 +1,90 @@
+// Copyright Epic Games, Inc. All Rights Reserved.
+
+#include "CoreTypes.h"
+#include "Containers/UnrealString.h"
+#include "Misc/Guid.h"
+#include "TestHarness.h"
+
+TEST_CASE("Core::Misc::FGuid::Smoke Test", "[Core][Misc][Smoke]")
+{
+	FGuid g = FGuid(305419896, 2271560481, 305419896, 2271560481);
+
+	// string conversions
+	TestEqual(TEXT("String conversion (Default) must return EGuidFormats::Digits string"), g.ToString(), g.ToString(EGuidFormats::Digits));
+	TestEqual(TEXT("String conversion (EGuidFormats::Digits)"), g.ToString(EGuidFormats::Digits), TEXT("12345678876543211234567887654321"));
+	TestEqual(TEXT("String conversion (EGuidFormats::DigitsWithHyphens)"), g.ToString(EGuidFormats::DigitsWithHyphens), TEXT("12345678-8765-4321-1234-567887654321"));
+	TestEqual(TEXT("String conversion (EGuidFormats::DigitsWithHyphensInBraces)"), g.ToString(EGuidFormats::DigitsWithHyphensInBraces), TEXT("{12345678-8765-4321-1234-567887654321}"));
+	TestEqual(TEXT("String conversion (EGuidFormats::DigitsWithHyphensInParentheses)"), g.ToString(EGuidFormats::DigitsWithHyphensInParentheses), TEXT("(12345678-8765-4321-1234-567887654321)"));
+	TestEqual(TEXT("String conversion (EGuidFormats::HexValuesInBraces)"), g.ToString(EGuidFormats::HexValuesInBraces), TEXT("{0x12345678,0x8765,0x4321,{0x12,0x34,0x56,0x78,0x87,0x65,0x43,0x21}}"));
+	TestEqual(TEXT("String conversion (EGuidFormats::UniqueObjectGuid)"), g.ToString(EGuidFormats::UniqueObjectGuid), TEXT("12345678-87654321-12345678-87654321"));
+	TestEqual(TEXT("String conversion (EGuidFormats::Short)"), g.ToString(EGuidFormats::Short), TEXT("eFY0EiFDZYd4VjQSIUNlhw"));
+	TestEqual(TEXT("String conversion (EGuidFormats::Base36Encoded)"), g.ToString(EGuidFormats::Base36Encoded), TEXT("12SRDE1NOQERISMPEF0H6KUU9"));
+
+	// parsing valid strings (exact)
+	FGuid g2_1;
+
+	TestTrue(TEXT("Parsing valid strings must succeed (EGuidFormats::Digits)"), FGuid::ParseExact(TEXT("12345678876543211234567887654321"), EGuidFormats::Digits, g2_1));
+	TestEqual(TEXT("Parsing valid strings must succeed (EGuidFormats::Digits)"), g2_1, g);
+
+	TestTrue(TEXT("Parsing valid strings must succeed (EGuidFormats::DigitsWithHyphens)"), FGuid::ParseExact(TEXT("12345678-8765-4321-1234-567887654321"), EGuidFormats::DigitsWithHyphens, g2_1));
+	TestEqual(TEXT("Parsing valid strings must succeed (EGuidFormats::DigitsWithHyphens)"), g2_1, g);
+
+	TestTrue(TEXT("Parsing valid strings must succeed (EGuidFormats::DigitsWithHyphensInBraces)"), FGuid::ParseExact(TEXT("{12345678-8765-4321-1234-567887654321}"), EGuidFormats::DigitsWithHyphensInBraces, g2_1));
+	TestEqual(TEXT("Parsing valid strings must succeed (EGuidFormats::DigitsWithHyphensInBraces)"), g2_1, g);
+
+	TestTrue(TEXT("Parsing valid strings must succeed (EGuidFormats::DigitsWithHyphensInParentheses)"), FGuid::ParseExact(TEXT("(12345678-8765-4321-1234-567887654321)"), EGuidFormats::DigitsWithHyphensInParentheses, g2_1));
+	TestEqual(TEXT("Parsing valid strings must succeed (EGuidFormats::DigitsWithHyphensInParentheses)"), g2_1, g);
+
+	TestTrue(TEXT("Parsing valid strings must succeed (EGuidFormats::HexValuesInBraces)"), FGuid::ParseExact(TEXT("{0x12345678,0x8765,0x4321,{0x12,0x34,0x56,0x78,0x87,0x65,0x43,0x21}}"), EGuidFormats::HexValuesInBraces, g2_1));
+	TestEqual(TEXT("Parsing valid strings must succeed (EGuidFormats::HexValuesInBraces)"), g2_1, g);
+
+	TestTrue(TEXT("Parsing valid strings must succeed (EGuidFormats::UniqueObjectGuid)"), FGuid::ParseExact(TEXT("12345678-87654321-12345678-87654321"), EGuidFormats::UniqueObjectGuid, g2_1));
+	TestEqual(TEXT("Parsing valid strings must succeed (EGuidFormats::UniqueObjectGuid)"), g2_1, g);
+
+	TestTrue(TEXT("Parsing valid strings must succeed (EGuidFormats::Short)"), FGuid::ParseExact(TEXT("eFY0EiFDZYd4VjQSIUNlhw"), EGuidFormats::Short, g2_1));
+	TestEqual(TEXT("Parsing valid strings must succeed (EGuidFormats::Short)"), g2_1, g);
+
+	TestTrue(TEXT("Parsing valid strings must succeed (EGuidFormats::Base36Encoded)"), FGuid::ParseExact(TEXT("12SRDE1NOQERISMPEF0H6KUU9"), EGuidFormats::Base36Encoded, g2_1));
+	TestEqual(TEXT("Parsing valid strings must succeed (EGuidFormats::Base36Encoded)"), g2_1, g);
+
+	// parsing invalid strings (exact)
+
+
+	// parsing valid strings (automatic)
+	FGuid g3_1;
+
+	TestTrue(TEXT("Parsing valid strings must succeed (12345678876543211234567887654321)"), FGuid::Parse(TEXT("12345678876543211234567887654321"), g3_1));
+	TestEqual(TEXT("Parsing valid strings must succeed (12345678876543211234567887654321)"), g3_1, g);
+
+	TestTrue(TEXT("Parsing valid strings must succeed (12345678-8765-4321-1234-567887654321)"), FGuid::Parse(TEXT("12345678-8765-4321-1234-567887654321"), g3_1));
+	TestEqual(TEXT("Parsing valid strings must succeed (12345678-8765-4321-1234-567887654321)"), g3_1, g);
+
+	TestTrue(TEXT("Parsing valid strings must succeed ({12345678-8765-4321-1234-567887654321})"), FGuid::Parse(TEXT("{12345678-8765-4321-1234-567887654321}"), g3_1));
+	TestEqual(TEXT("Parsing valid strings must succeed ({12345678-8765-4321-1234-567887654321})"), g3_1, g);
+
+	TestTrue(TEXT("Parsing valid strings must succeed ((12345678-8765-4321-1234-567887654321))"), FGuid::Parse(TEXT("(12345678-8765-4321-1234-567887654321)"), g3_1));
+	TestEqual(TEXT("Parsing valid strings must succeed ((12345678-8765-4321-1234-567887654321))"), g3_1, g);
+
+	TestTrue(TEXT("Parsing valid strings must succeed ({0x12345678,0x8765,0x4321,{0x12,0x34,0x56,0x78,0x87,0x65,0x43,0x21}})"), FGuid::Parse(TEXT("{0x12345678,0x8765,0x4321,{0x12,0x34,0x56,0x78,0x87,0x65,0x43,0x21}}"), g3_1));
+	TestEqual(TEXT("Parsing valid strings must succeed ({0x12345678,0x8765,0x4321,{0x12,0x34,0x56,0x78,0x87,0x65,0x43,0x21}})"), g3_1, g);
+
+	TestTrue(TEXT("Parsing valid strings must succeed (12345678-87654321-12345678-87654321)"), FGuid::Parse(TEXT("12345678-87654321-12345678-87654321"), g3_1));
+	TestEqual(TEXT("Parsing valid strings must succeed (12345678-87654321-12345678-87654321)"), g3_1, g);
+
+	TestTrue(TEXT("Parsing valid strings must succeed (eFY0EiFDZYd4VjQSIUNlhw)"), FGuid::Parse(TEXT("eFY0EiFDZYd4VjQSIUNlhw"), g3_1));
+	TestEqual(TEXT("Parsing valid strings must succeed (eFY0EiFDZYd4VjQSIUNlhw)"), g3_1, g);
+
+	TestTrue(TEXT("Parsing valid strings must succeed (12SRDE1NOQERISMPEF0H6KUU9)"), FGuid::Parse(TEXT("12SRDE1NOQERISMPEF0H6KUU9"), g3_1));
+	TestEqual(TEXT("Parsing valid strings must succeed (12SRDE1NOQERISMPEF0H6KUU9)"), g3_1, g);
+
+	//parsing invalid strings (automatic)
+
+	// GUID validation
+	FGuid g4_1 = FGuid::NewGuid();
+
+	TestTrue(TEXT("New GUIDs must be valid"), g4_1.IsValid());
+	
+	g4_1.Invalidate();
+
+	TestFalse(TEXT("Invalidated GUIDs must be invalid"), g4_1.IsValid());
+}
diff --git a/Engine/Source/Runtime/Core/Tests/Misc/QueueTest.cpp b/Engine/Source/Runtime/Core/Tests/Misc/QueueTest.cpp
new file mode 100644
--- /dev/null
+++ b/Engine/Source/Runtime/Core/Tests/Misc/QueueTest.cpp
@@ -0,0 +1,59 @@
+// Copyright Epic Games, Inc. All Rights Reserved.
+
+#include "CoreTypes.h"
+#include "Containers/Queue.h"
+#include "TestHarness.h"
+
+TEST_CASE("Core::Misc::TQueue::Smoke Test", "[Core][Misc][Smoke]")
+{
+	// empty queues
+	{
+		TQueue<int32> Queue;
+		int32 OutItem = 0;
+
+		TestTrue(TEXT("A new queue must be empty"), Queue.IsEmpty());
+
+		TestFalse(TEXT("A new queue must not dequeue anything"), Queue.Dequeue(OutItem));
+		TestFalse(TEXT("A new queue must not peek anything"), Queue.Peek(OutItem));
+	}
+
+	// insertion & removal
+	{
+		TQueue<int32> Queue;
+		int32 Item1 = 1;
+		int32 Item2 = 2;
+		int32 Item3 = 3;
+		int32 OutItem = 0;
+
+		TestTrue(TEXT("Inserting into a new queue must succeed"), Queue.Enqueue(Item1));
+		TestTrue(TEXT("Peek must succeed on a queue with one item"), Queue.Peek(OutItem));
+		TestEqual(TEXT("Peek must return the first value"), OutItem, Item1);
+
+		TestTrue(TEXT("Inserting into a non-empty queue must succeed"), Queue.Enqueue(Item2));
+		TestTrue(TEXT("Peek must succeed on a queue with two items"), Queue.Peek(OutItem));
+		TestEqual(TEXT("Peek must return the first item"), OutItem, Item1);
+
+		Queue.Enqueue(Item3);
+
+		TestTrue(TEXT("Dequeue must succeed on a queue with three items"), Queue.Dequeue(OutItem));
+		TestEqual(TEXT("Dequeue must return the first item"), OutItem, Item1);
+		TestTrue(TEXT("Dequeue must succeed on a queue with two items"), Queue.Dequeue(OutItem));
+		TestEqual(TEXT("Dequeue must return the second item"), OutItem, Item2);
+		TestTrue(TEXT("Dequeue must succeed on a queue with one item"), Queue.Dequeue(OutItem));
+		TestEqual(TEXT("Dequeue must return the third item"), OutItem, Item3);
+
+		TestTrue(TEXT("After removing all items, the queue must be empty"), Queue.IsEmpty());
+	}
+
+	// emptying
+	{
+		TQueue<int32> Queue;
+
+		Queue.Enqueue(1);
+		Queue.Enqueue(2);
+		Queue.Enqueue(3);
+		Queue.Empty();
+
+		TestTrue(TEXT("An emptied queue must be empty"), Queue.IsEmpty());
+	}
+}
diff --git a/Engine/Source/Runtime/Core/Tests/Misc/TimespanTest.cpp b/Engine/Source/Runtime/Core/Tests/Misc/TimespanTest.cpp
new file mode 100644
--- /dev/null
+++ b/Engine/Source/Runtime/Core/Tests/Misc/TimespanTest.cpp
@@ -0,0 +1,128 @@
+// Copyright Epic Games, Inc. All Rights Reserved.
+
+#include "CoreTypes.h"
+#include "Containers/UnrealString.h"
+#include "Misc/Timespan.h"
+#include "TestHarness.h"
+
+TEST_CASE("Core::Misc::FTimespan::Smoke Test", "[Core][Misc][Smoke]")
+{
+	// constructors must create equal objects
+	{
+		FTimespan ts1_1 = FTimespan(3, 2, 1);
+		FTimespan ts1_2 = FTimespan(0, 3, 2, 1);
+		FTimespan ts1_3 = FTimespan(0, 3, 2, 1, 0);
+
+		TestEqual(TEXT("Constructors must create equal objects (Hours/Minutes/Seconds vs. Days/Hours/Minutes/Seconds)"), ts1_1, ts1_2);
+		TestEqual(TEXT("Constructors must create equal objects (Hours/Minutes/Seconds vs. Days/Hours/Minutes/Seconds/FractionNano)"), ts1_1, ts1_3);
+	}
+
+	// component getters must return correct values
+	{
+		FTimespan ts2_1 = FTimespan(1, 2, 3, 4, 123456789);
+
+		TestEqual(TEXT("Component getters must return correct values (Days)"), ts2_1.GetDays(), 1);
+		TestEqual(TEXT("Component getters must return correct values (Hours)"), ts2_1.GetHours(), 2);
+		TestEqual(TEXT("Component getters must return correct values (Minutes)"), ts2_1.GetMinutes(), 3);
+		TestEqual(TEXT("Component getters must return correct values (Seconds)"), ts2_1.GetSeconds(), 4);
+		TestEqual(TEXT("Component getters must return correct values (FractionMilli)"), ts2_1.GetFractionMilli(), 123);
+		TestEqual(TEXT("Component getters must return correct values (FractionMicro)"), ts2_1.GetFractionMicro(), 123456);
+		TestEqual(TEXT("Component getters must return correct values (FractionNano)"), ts2_1.GetFractionNano(), 123456700);
+	}
+
+	// durations of positive and negative time spans must match
+	{
+		FTimespan ts3_1 = FTimespan(1, 2, 3, 4, 123456789);
+		FTimespan ts3_2 = FTimespan(-1, -2, -3, -4, -123456789);
+
+		TestEqual(TEXT("Durations of positive and negative time spans must match"), ts3_1.GetDuration(), ts3_2.GetDuration());
+	}
+
+	// static constructors must create correct values
+	{
+		TestEqual(TEXT("Static constructors must create correct values (FromDays)"), FTimespan::FromDays(123).GetTotalDays(), 123.0);
+		TestEqual(TEXT("Static constructors must create correct values (FromHours)"), FTimespan::FromHours(123).GetTotalHours(), 123.0);
+		TestEqual(TEXT("Static constructors must create correct values (FromMinutes)"), FTimespan::FromMinutes(123).GetTotalMinutes(), 123.0);
+		TestEqual(TEXT("Static constructors must create correct values (FromSeconds)"), FTimespan::FromSeconds(123).GetTotalSeconds(), 123.0);
+		TestEqual(TEXT("Static constructors must create correct values (FromMilliseconds)"), FTimespan::FromMilliseconds(123).GetTotalMilliseconds(), 123.0);
+		TestEqual(TEXT("Static constructors must create correct values (FromMicroseconds)"), FTimespan::FromMicroseconds(123).GetTotalMicroseconds(), 123.0);
+	}
+
+	// string conversions must return correct strings
+	{
+		FTimespan ts5_1 = FTimespan(1, 2, 3, 4, 123456789);
+
+		TestEqualString(TEXT("String conversion (Default)"), ts5_1.ToString(), TEXT("+1.02:03:04.123"));
+		TestEqualString(TEXT("String conversion (%d.%h:%m:%s.%f)"), ts5_1.ToString(TEXT("%d.%h:%m:%s.%f")), TEXT("+1.02:03:04.123"));
+		TestEqualString(TEXT("String conversion (%d.%h:%m:%s.%u)"), ts5_1.ToString(TEXT("%d.%h:%m:%s.%u")), TEXT("+1.02:03:04.123456"));
+		TestEqualString(TEXT("String conversion (%D.%h:%m:%s.%n)"), ts5_1.ToString(TEXT("%D.%h:%m:%s.%n")), TEXT("+00000001.02:03:04.123456700"));
+	}
+
+	// parsing valid strings must succeed
+	{
+		FTimespan ts6_t;
+
+		FTimespan ts6_1 = FTimespan(1, 2, 3, 4, 123000000);
+		FTimespan ts6_2 = FTimespan(1, 2, 3, 4, 123456000);
+		FTimespan ts6_3 = FTimespan(1, 2, 3, 4, 123456700);
+
+		TestTrue(TEXT("Parsing valid strings must succeed (+1.02:03:04.123)"), FTimespan::Parse(TEXT("+1.02:03:04.123"), ts6_t));
+		TestEqual(TEXT("Parsing valid strings must result in correct values (+1.02:03:04.123)"), ts6_t, ts6_1);
+
+		TestTrue(TEXT("Parsing valid strings must succeed (+1.02:03:04.123456)"), FTimespan::Parse(TEXT("+1.02:03:04.123456"), ts6_t));
+		TestEqual(TEXT("Parsing valid strings must result in correct values (+1.02:03:04.123456)"), ts6_t, ts6_2);
+
+		TestTrue(TEXT("Parsing valid strings must succeed (+1.02:03:04.123456789)"), FTimespan::Parse(TEXT("+1.02:03:04.123456789"), ts6_t));
+		TestEqual(TEXT("Parsing valid strings must result in correct values (+1.02:03:04.123456789)"), ts6_t, ts6_3);
+
+		FTimespan ts6_4 = FTimespan(-1, -2, -3, -4, -123000000);
+		FTimespan ts6_5 = FTimespan(-1, -2, -3, -4, -123456000);
+		FTimespan ts6_6 = FTimespan(-1, -2, -3, -4, -123456700);
+
+		TestTrue(TEXT("Parsing valid strings must succeed (-1.02:03:04.123)"), FTimespan::Parse(TEXT("-1.02:03:04.123"), ts6_t));
+		TestEqual(TEXT("Parsing valid strings must result in correct values (-1.02:03:04.123)"), ts6_t, ts6_4);
+
+		TestTrue(TEXT("Parsing valid strings must succeed (-1.02:03:04.123456)"), FTimespan::Parse(TEXT("-1.02:03:04.123456"), ts6_t));
+		TestEqual(TEXT("Parsing valid strings must result in correct values (-1.02:03:04.123456)"), ts6_t, ts6_5);
+
+		TestTrue(TEXT("Parsing valid strings must succeed (-1.02:03:04.123456789)"), FTimespan::Parse(TEXT("-1.02:03:04.123456789"), ts6_t));
+		TestEqual(TEXT("Parsing valid strings must result in correct values (-1.02:03:04.123456789)"), ts6_t, ts6_6);
+	}
+
+	// parsing invalid strings must fail
+	{
+		FTimespan ts7_1;
+
+		//TestFalse(TEXT("Parsing invalid strings must fail (1,02:03:04.005)"), FTimespan::Parse(TEXT("1,02:03:04.005"), ts7_1));
+		//TestFalse(TEXT("Parsing invalid strings must fail (1.1.02:03:04:005)"), FTimespan::Parse(TEXT("1.1.02:03:04:005"), ts7_1));
+		//TestFalse(TEXT("Parsing invalid strings must fail (04:005)"), FTimespan::Parse(TEXT("04:005"), ts7_1));
+	}
+
+	// `From*` converters must return correct value
+	// test normal and edge cases for polar conversions (FromMicroseconds() and FromDay()) and just normal case for all others
+	{
+		TestEqual(TEXT("FromMicroseconds(0) results in correct value"), FTimespan::FromMicroseconds(0), FTimespan(0));
+		TestEqual(TEXT("FromMicroseconds(1) results in correct value"), FTimespan::FromMicroseconds(1), FTimespan(1 * ETimespan::TicksPerMicrosecond));
+		TestEqual(TEXT("FromMicroseconds(1.1) results in correct value"), FTimespan::FromMicroseconds(1.1), FTimespan(1 * ETimespan::TicksPerMicrosecond + 1));
+		TestEqual(TEXT("FromMicroseconds(1.5) results in correct value"), FTimespan::FromMicroseconds(1.5), FTimespan(1 * ETimespan::TicksPerMicrosecond + 5));
+		TestEqual(TEXT("FromMicroseconds(1.499999999999997) results in 1.5 microsecs of ticks"), FTimespan::FromMicroseconds(1.499999999999997), FTimespan(1 * ETimespan::TicksPerMicrosecond + 5));
+		TestEqual(TEXT("FromMicroseconds(1.50000001) results in  in 1.5 microsecs of ticks"), FTimespan::FromMicroseconds(1.50000001), FTimespan(1 * ETimespan::TicksPerMicrosecond + 5));
+		TestEqual(TEXT("FromMicroseconds(-1) results in correct value"), FTimespan::FromMicroseconds(-1), FTimespan(-1 * ETimespan::TicksPerMicrosecond));
+		TestEqual(TEXT("FromMicroseconds(-1.1) results in correct value"), FTimespan::FromMicroseconds(-1.1), FTimespan(-1 * ETimespan::TicksPerMicrosecond - 1));
+		TestEqual(TEXT("FromMicroseconds(-1.5) results in correct value"), FTimespan::FromMicroseconds(-1.5), FTimespan(-1 * ETimespan::TicksPerMicrosecond - 5));
+
+		TestEqual(TEXT("FromMilliseconds(1) results in correct value"), FTimespan::FromMilliseconds(1), FTimespan(1 * ETimespan::TicksPerMillisecond));
+		TestEqual(TEXT("FromSeconds(1) results in correct value"), FTimespan::FromSeconds(1), FTimespan(1 * ETimespan::TicksPerSecond));
+		TestEqual(TEXT("FromMinutes(1) results in correct value"), FTimespan::FromMinutes(1), FTimespan(1 * ETimespan::TicksPerMinute));
+		TestEqual(TEXT("FromHours(1) results in correct value"), FTimespan::FromHours(1), FTimespan(1 * ETimespan::TicksPerHour));
+
+		TestEqual(TEXT("FromDays(0) results in correct value"), FTimespan::FromDays(0), FTimespan(0));
+		TestEqual(TEXT("FromDays(1) results in correct value"), FTimespan::FromDays(1), FTimespan(1 * ETimespan::TicksPerDay));
+		TestEqual(TEXT("FromDays(1.25) results in correct value (1 day and 6 hours of ticks)"), FTimespan::FromDays(1.25), FTimespan(1 * ETimespan::TicksPerDay + 6 * ETimespan::TicksPerHour));
+		TestEqual(TEXT("FromDays(1.5) results in correct value (1 day and 12 hours of ticks)"), FTimespan::FromDays(1.5), FTimespan(1 * ETimespan::TicksPerDay + 12 * ETimespan::TicksPerHour));
+		TestEqual(TEXT("FromDays(1.499999999999997) results in correct value (1 day and 12 hours of ticks)"), FTimespan::FromDays(1.499999999999997), FTimespan(1 * ETimespan::TicksPerDay + 12 * ETimespan::TicksPerHour));
+		TestEqual(TEXT("FromDays(-1) results in correct value"), FTimespan::FromDays(-1), FTimespan(-1 * ETimespan::TicksPerDay));
+		TestEqual(TEXT("FromDays(-1.25) results in correct value (minus 1 day and 6 hours of ticks)"), FTimespan::FromDays(-1.25), FTimespan(-1 * ETimespan::TicksPerDay - 6 * ETimespan::TicksPerHour));
+		TestEqual(TEXT("FromDays(-1.5) results in correct value (minus 1 day and 12 hours of ticks)"), FTimespan::FromDays(-1.5), FTimespan(-1 * ETimespan::TicksPerDay - 12 * ETimespan::TicksPerHour));
+	}
+}
diff --git a/Engine/Source/Runtime/Core/Tests/Misc/TripleBufferTest.cpp b/Engine/Source/Runtime/Core/Tests/Misc/TripleBufferTest.cpp
new file mode 100644
--- /dev/null
+++ b/Engine/Source/Runtime/Core/Tests/Misc/TripleBufferTest.cpp
@@ -0,0 +1,100 @@
+// Copyright Epic Games, Inc. All Rights Reserved.
+
+#include "CoreTypes.h"
+#include "Containers/UnrealString.h"
+#include "Containers/TripleBuffer.h"
+#include "Math/RandomStream.h"
+#include "TestHarness.h"
+
+TEST_CASE("Core::Misc::TTripleBuffer::Smoke Test", "[Core][Misc][Smoke]")
+{
+	// uninitialized buffer
+	{
+		TTripleBuffer<int32> Buffer(NoInit);
+
+		TestFalse(TEXT("Uninitialized triple buffer must not be dirty"), Buffer.IsDirty());
+	}
+
+	// initialized buffer
+	{
+		TTripleBuffer<int32> Buffer(1);
+
+		TestFalse(TEXT("Initialized triple buffer must not be dirty"), Buffer.IsDirty());
+		TestEqual(TEXT("Initialized triple buffer must have correct read buffer value"), Buffer.Read(), 1);
+
+		Buffer.SwapReadBuffers();
+
+		TestEqual(TEXT("Initialized triple buffer must have correct temp buffer value"), Buffer.Read(), 1);
+
+		Buffer.SwapWriteBuffers();
+
+		TestTrue(TEXT("Write buffer swap must set dirty flag"), Buffer.IsDirty());
+
+		Buffer.SwapReadBuffers();
+
+		TestFalse(TEXT("Read buffer swap must clear dirty flag"), Buffer.IsDirty());
+		TestEqual(TEXT("Initialized triple buffer must have correct temp buffer value"), Buffer.Read(), 1);
+	}
+
+	// pre-set buffer
+	{
+		int32 Array[3] = { 1, 2, 3 };
+		TTripleBuffer<int32> Buffer(Array);
+
+		int32 Read = Buffer.Read();
+		TestEqual(TEXT("Pre-set triple buffer must have correct Read buffer value"), Read, 3);
+
+		Buffer.SwapReadBuffers();
+
+		int32 Temp = Buffer.Read();
+		TestEqual(TEXT("Pre-set triple buffer must have correct Temp buffer value"), Temp, 1);
+
+		Buffer.SwapWriteBuffers();
+		Buffer.SwapReadBuffers();
+
+		int32 Write = Buffer.Read();
+		TestEqual(TEXT("Pre-set triple buffer must have correct Write buffer value"), Write, 2);
+	}
+
+	// operations
+	{
+		TTripleBuffer<int32> Buffer;
+
+		for (int32 Index = 0; Index < 6; ++Index)
+		{
+			int32& Write = Buffer.GetWriteBuffer(); Write = Index; Buffer.SwapWriteBuffers();
+			Buffer.SwapReadBuffers();
+			TestEqual(*FString::Printf(TEXT("Triple buffer must read correct value (%i)"), Index), Buffer.Read(), Index);
+		}
+
+		FRandomStream Rand;
+		int32 LastRead = -1;
+
+		for (int32 Index = 0; Index < 100; ++Index)
+		{
+			int32 Writes = Rand.GetUnsignedInt() % 4;
+
+			while (Writes > 0)
+			{
+				int32& Write = Buffer.GetWriteBuffer(); Write = Index; Buffer.SwapWriteBuffers();
+				--Writes;
+			}
+			
+			int32 Reads = Rand.GetUnsignedInt() % 4;
+
+			while (Reads > 0)
+			{
+				if (!Buffer.IsDirty())
+				{
+					break;
+				}
+
+				Buffer.SwapReadBuffers();
+				int32 Read = Buffer.Read();
+				TestTrue(TEXT("Triple buffer must read in increasing order"), Read > LastRead);
+				LastRead = Read;
+				--Reads;
+			}
+		}
+	}
+}
diff --git a/Engine/Source/Runtime/Core/Tests/Serialization/CompactBinaryWriterTest.cpp b/Engine/Source/Runtime/Core/Tests/Serialization/CompactBinaryWriterTest.cpp
new file mode 100644
--- /dev/null
+++ b/Engine/Source/Runtime/Core/Tests/Serialization/CompactBinaryWriterTest.cpp
@@ -0,0 +1,1057 @@
+// Copyright Epic Games, Inc. All Rights Reserved.
+
+#include "Serialization/CompactBinaryWriter.h"
+
+#include "IO/IoHash.h"
+#include "Misc/DateTime.h"
+#include "Misc/Guid.h"
+#include "Misc/Timespan.h"
+#include "Misc/StringBuilder.h"
+#include "Serialization/BufferArchive.h"
+#include "Serialization/CompactBinaryPackage.h"
+#include "Serialization/CompactBinaryValidation.h"
+#include "TestHarness.h"
+
+TEST_CASE("Core::Serialization::FCbWriter::Object", "[Core][Serialization][Smoke]")
+{
+	TCbWriter<256> Writer;
+
+	// Test Empty Object
+	{
+		Writer.Reset();
+		Writer.BeginObject();
+		Writer.EndObject();
+		FCbField Field = Writer.Save();
+		TestEqual_Conditional(TEXT("FCbWriter(Object, Empty) Validate"), ValidateCompactBinary(Field.GetOuterBuffer(), ECbValidateMode::All), ECbValidateError::None)
+		{
+			TestTrue(TEXT("FCbWriter(Object, Empty).IsObject()"), Field.IsObject());
+			TestFalse(TEXT("FCbWriter(Object, Empty).AsObjectView()"), Field.AsObjectView().CreateViewIterator().HasValue());
+		}
+	}
+
+	// Test Named Empty Object
+	{
+		Writer.Reset();
+		Writer.SetName("Object"_ASV);
+		Writer.BeginObject();
+		Writer.EndObject();
+		FCbField Field = Writer.Save();
+		TestEqual_Conditional(TEXT("FCbWriter(Object, Empty, Name) Validate"), ValidateCompactBinary(Field.GetOuterBuffer(), ECbValidateMode::All), ECbValidateError::None)
+		{
+			TestTrue(TEXT("FCbWriter(Object, Empty, Name).IsObject()"), Field.IsObject());
+			TestFalse(TEXT("FCbWriter(Object, Empty, Name).AsObjectView()"), Field.AsObjectView().CreateViewIterator().HasValue());
+		}
+	}
+
+	// Test Basic Object
+	{
+		Writer.Reset();
+		Writer.BeginObject();
+		Writer.SetName("Integer"_ASV).AddInteger(0);
+		Writer.SetName("Float"_ASV).AddFloat(0.0f);
+		Writer.EndObject();
+		FCbField Field = Writer.Save();
+		TestEqual_Conditional(TEXT("FCbWriter(Object, Basic) Validate"), ValidateCompactBinary(Field.GetOuterBuffer(), ECbValidateMode::All), ECbValidateError::None)
+		{
+			TestTrue(TEXT("FCbWriter(Object, Basic).IsObject()"), Field.IsObject());
+			FCbObjectView Object = Field.AsObjectView();
+			TestTrue(TEXT("FCbWriter(Object, Basic).AsObjectView()[Integer]"), Object["Integer"_ASV].IsInteger());
+			TestTrue(TEXT("FCbWriter(Object, Basic).AsObjectView()[Float]"), Object["Float"_ASV].IsFloat());
+		}
+	}
+
+	// Test Uniform Object
+	{
+		Writer.Reset();
+		Writer.BeginObject();
+		Writer.SetName("Field1"_ASV).AddInteger(0);
+		Writer.SetName("Field2"_ASV).AddInteger(1);
+		Writer.EndObject();
+		FCbField Field = Writer.Save();
+		TestEqual_Conditional(TEXT("FCbWriter(Object, Uniform) Validate"), ValidateCompactBinary(Field.GetOuterBuffer(), ECbValidateMode::All), ECbValidateError::None)
+		{
+			TestTrue(TEXT("FCbWriter(Object, Uniform).IsObject()"), Field.IsObject());
+			FCbObjectView Object = Field.AsObjectView();
+			TestTrue(TEXT("FCbWriter(Object, Uniform).AsObjectView()[Field1]"), Object["Field1"_ASV].IsInteger());
+			TestTrue(TEXT("FCbWriter(Object, Uniform).AsObjectView()[Field2]"), Object["Field2"_ASV].IsInteger());
+		}
+	}
+}
+
+TEST_CASE("Core::Serialization::FCbWriter::Array", "[Core][Serialization][Smoke]")
+{
+	TCbWriter<256> Writer;
+
+	// Test Empty Array
+	{
+		Writer.Reset();
+		Writer.BeginArray();
+		Writer.EndArray();
+		FCbField Field = Writer.Save();
+		TestEqual_Conditional(TEXT("FCbWriter(Array, Empty) Validate"), ValidateCompactBinary(Field.GetOuterBuffer(), ECbValidateMode::All), ECbValidateError::None)
+		{
+			TestTrue(TEXT("FCbWriter(Array, Empty).IsArray()"), Field.IsArray());
+			TestEqual(TEXT("FCbWriter(Array, Empty).AsArrayView()"), Field.AsArrayView().Num(), uint64(0));
+		}
+	}
+
+	// Test Named Empty Array
+	{
+		Writer.Reset();
+		Writer.SetName("Array"_ASV);
+		Writer.BeginArray();
+		Writer.EndArray();
+		FCbField Field = Writer.Save();
+		TestEqual_Conditional(TEXT("FCbWriter(Array, Empty, Name) Validate"), ValidateCompactBinary(Field.GetOuterBuffer(), ECbValidateMode::All), ECbValidateError::None)
+		{
+			TestTrue(TEXT("FCbWriter(Array, Empty, Name).IsArray()"), Field.IsArray());
+			TestEqual(TEXT("FCbWriter(Array, Empty, Name).AsArrayView()"), Field.AsArrayView().Num(), uint64(0));
+		}
+	}
+
+	// Test Basic Array
+	{
+		Writer.Reset();
+		Writer.BeginArray();
+		Writer.AddInteger(0);
+		Writer.AddFloat(0.0f);
+		Writer.EndArray();
+		FCbField Field = Writer.Save();
+		TestEqual_Conditional(TEXT("FCbWriter(Array, Basic) Validate"), ValidateCompactBinary(Field.GetOuterBuffer(), ECbValidateMode::All), ECbValidateError::None)
+		{
+			TestTrue(TEXT("FCbWriter(Array, Basic).IsArray()"), Field.IsArray());
+			FCbFieldViewIterator Iterator = Field.AsArrayView().CreateViewIterator();
+			TestTrue(TEXT("FCbWriter(Array, Basic).AsArrayView()[Integer]"), Iterator.IsInteger());
+			++Iterator;
+			TestTrue(TEXT("FCbWriter(Array, Basic).AsArrayView()[Float]"), Iterator.IsFloat());
+			++Iterator;
+			TestFalse(TEXT("FCbWriter(Array, Basic).AsArrayView()[End]"), Iterator.HasValue());
+		}
+	}
+
+	// Test Uniform Array
+	{
+		Writer.Reset();
+		Writer.BeginArray();
+		Writer.AddInteger(0);
+		Writer.AddInteger(1);
+		Writer.EndArray();
+		FCbField Field = Writer.Save();
+		TestEqual_Conditional(TEXT("FCbWriter(Array, Uniform) Validate"), ValidateCompactBinary(Field.GetOuterBuffer(), ECbValidateMode::All), ECbValidateError::None)
+		{
+			TestTrue(TEXT("FCbWriter(Array, Uniform).IsArray()"), Field.IsArray());
+			FCbFieldViewIterator Iterator = Field.AsArrayView().CreateViewIterator();
+			TestTrue(TEXT("FCbWriter(Array, Basic).AsArrayView()[Field1]"), Iterator.IsInteger());
+			++Iterator;
+			TestTrue(TEXT("FCbWriter(Array, Basic).AsArrayView()[Field2]"), Iterator.IsInteger());
+			++Iterator;
+			TestFalse(TEXT("FCbWriter(Array, Basic).AsArrayView()[End]"), Iterator.HasValue());
+		}
+	}
+}
+
+TEST_CASE("Core::Serialization::FCbWriter::Null", "[Core][Serialization][Smoke]")
+{
+	TCbWriter<256> Writer;
+
+	// Test Null
+	{
+		Writer.Reset();
+		Writer.AddNull();
+		FCbField Field = Writer.Save();
+		TestEqual_Conditional(TEXT("FCbWriter(Null) Validate"), ValidateCompactBinary(Field.GetOuterBuffer(), ECbValidateMode::All), ECbValidateError::None)
+		{
+			TestFalse(TEXT("FCbWriter(Null).HasName()"), Field.HasName());
+			TestTrue(TEXT("FCbWriter(Null).IsNull()"), Field.IsNull());
+		}
+	}
+
+	// Test Null with Name
+	{
+		Writer.Reset();
+		Writer.SetName("Null"_ASV);
+		Writer.AddNull();
+		FCbField Field = Writer.Save();
+		TestEqual_Conditional(TEXT("FCbWriter(Null, Name) Validate"), ValidateCompactBinary(Field.GetOuterBuffer(), ECbValidateMode::All), ECbValidateError::None)
+		{
+			TestEqual(TEXT("FCbWriter(Null, Name).GetName()"), Field.GetName(), "Null"_U8SV);
+			TestTrue(TEXT("FCbWriter(Null, Name).HasName()"), Field.HasName());
+			TestTrue(TEXT("FCbWriter(Null, Name).IsNull()"), Field.IsNull());
+		}
+	}
+
+	// Test Null Array/Object Uniformity
+	{
+		Writer.Reset();
+
+		Writer.BeginArray();
+		Writer.AddNull();
+		Writer.AddNull();
+		Writer.AddNull();
+		Writer.EndArray();
+
+		Writer.BeginObject();
+		Writer.SetName("N1"_ASV).AddNull();
+		Writer.SetName("N2"_ASV).AddNull();
+		Writer.SetName("N3"_ASV).AddNull();
+		Writer.EndObject();
+
+		FCbFieldIterator Fields = Writer.Save();
+		TestEqual(TEXT("FCbWriter(Null, Uniform) Validate"), ValidateCompactBinaryRange(Fields.GetOuterBuffer(), ECbValidateMode::All), ECbValidateError::None);
+	}
+
+	// Test Null with Save(Buffer)
+	{
+		Writer.Reset();
+		constexpr int NullCount = 3;
+		for (int Index = 0; Index < NullCount; ++Index)
+		{
+			Writer.AddNull();
+		}
+		uint8 Buffer[NullCount]{};
+		FCbFieldViewIterator Fields = Writer.Save(MakeMemoryView(Buffer));
+		TestEqual_Conditional(TEXT("FCbWriter(Null, Memory) Validate"), ValidateCompactBinaryRange(MakeMemoryView(Buffer), ECbValidateMode::All), ECbValidateError::None)
+		{
+			for (int Index = 0; Index < NullCount; ++Index)
+			{
+				TestTrue(TEXT("FCbWriter(Null, Memory) IsNull"), Fields.IsNull());
+				++Fields;
+			}
+			TestFalse(TEXT("FCbWriter(Null, Memory) HasValue"), Fields.HasValue());
+		}
+	}
+}
+
+///////////////////////////////////////////////////////////////////////////////////////////////////////////////////////
+
+TEST_CASE("Core::Serialization::FCbWriter::Binary", "[Core][Serialization][Smoke]")
+{
+	TCbWriter<256> Writer;
+
+	// Test Empty Binary
+	{
+		Writer.Reset();
+		Writer.AddBinary(nullptr, 0);
+		FCbField Field = Writer.Save();
+		TestEqual_Conditional(TEXT("FCbWriter(Binary, Empty) Validate"), ValidateCompactBinary(Field.GetOuterBuffer(), ECbValidateMode::All), ECbValidateError::None)
+		{
+			TestFalse(TEXT("FCbWriter(Binary, Empty).HasName()"), Field.HasName());
+			TestTrue(TEXT("FCbWriter(Binary, Empty).IsBinary()"), Field.IsBinary());
+			TestTrue(TEXT("FCbWriter(Binary, Empty).AsBinaryView()"), Field.AsBinaryView().IsEmpty());
+		}
+	}
+
+	// Test Basic Binary
+	{
+		Writer.Reset();
+		const uint8 BinaryValue[] = { 1, 2, 3, 4, 5, 6 };
+		Writer.SetName("Binary"_ASV);
+		Writer.AddBinary(BinaryValue, sizeof(BinaryValue));
+		FCbField Field = Writer.Save();
+		TestEqual_Conditional(TEXT("FCbWriter(Binary, Array) Validate"), ValidateCompactBinary(Field.GetOuterBuffer(), ECbValidateMode::All), ECbValidateError::None)
+		{
+			TestEqual(TEXT("FCbWriter(Binary, Array).GetName()"), Field.GetName(), "Binary"_U8SV);
+			TestTrue(TEXT("FCbWriter(Binary, Array).HasName()"), Field.HasName());
+			TestTrue(TEXT("FCbWriter(Binary, Array).IsBinary()"), Field.IsBinary());
+			TestTrue(TEXT("FCbWriter(Binary, Array).AsBinaryView()"), Field.AsBinaryView().EqualBytes(MakeMemoryView(BinaryValue)));
+		}
+	}
+}
+
+TEST_CASE("Core::Serialization::FCbWriter::String", "[Core][Serialization][Smoke]")
+{
+	TCbWriter<256> Writer;
+
+	// Test Empty Strings
+	{
+		Writer.Reset();
+		Writer.AddString(FAnsiStringView());
+		Writer.AddString(FWideStringView());
+		FCbFieldIterator Fields = Writer.Save();
+		TestEqual_Conditional(TEXT("FCbWriter(String, Empty) Validate"), ValidateCompactBinaryRange(Fields.GetOuterBuffer(), ECbValidateMode::All), ECbValidateError::None)
+		{
+			for (FCbFieldView Field : Fields)
+			{
+				TestFalse(TEXT("FCbWriter(String, Empty).HasName()"), Field.HasName());
+				TestTrue(TEXT("FCbWriter(String, Empty).IsString()"), Field.IsString());
+				TestTrue(TEXT("FCbWriter(String, Empty).AsString()"), Field.AsString().IsEmpty());
+			}
+		}
+	}
+
+	// Test Basic Strings
+	{
+		Writer.Reset();
+		Writer.SetName("String"_ASV).AddString("Value"_ASV);
+		Writer.SetName("String"_ASV).AddString(TEXT("Value"_SV));
+		FCbFieldIterator Fields = Writer.Save();
+		TestEqual_Conditional(TEXT("FCbWriter(String, Basic) Validate"), ValidateCompactBinaryRange(Fields.GetOuterBuffer(), ECbValidateMode::All), ECbValidateError::None)
+		{
+			for (FCbFieldView Field : Fields)
+			{
+				TestEqual(TEXT("FCbWriter(String, Basic).GetName()"), Field.GetName(), "String"_U8SV);
+				TestTrue(TEXT("FCbWriter(String, Basic).HasName()"), Field.HasName());
+				TestTrue(TEXT("FCbWriter(String, Basic).IsString()"), Field.IsString());
+				TestEqual(TEXT("FCbWriter(String, Basic).AsString()"), Field.AsString(), "Value"_U8SV);
+			}
+		}
+	}
+
+	// Test Long Strings
+	{
+		Writer.Reset();
+		constexpr int DotCount = 256;
+		TUtf8StringBuilder<DotCount + 1> Dots;
+		for (int Index = 0; Index < DotCount; ++Index)
+		{
+			Dots << '.';
+		}
+		Writer.AddString(Dots);
+		Writer.AddString(FString::ChrN(DotCount, TEXT('.')));
+		FCbFieldIterator Fields = Writer.Save();
+		TestEqual_Conditional(TEXT("FCbWriter(String, Long) Validate"), ValidateCompactBinaryRange(Fields.GetOuterBuffer(), ECbValidateMode::All), ECbValidateError::None)
+		{
+			for (FCbFieldView Field : Fields)
+			{
+				TestEqual(TEXT("FCbWriter(String, Long).AsString()"), Field.AsString(), Dots.ToView());
+			}
+		}
+	}
+
+	// Test Non-ASCII String
+	{
+		Writer.Reset();
+		WIDECHAR Value[2] = { 0xd83d, 0xde00 };
+		Writer.AddString("\xf0\x9f\x98\x80"_ASV);
+		Writer.AddString(FWideStringView(Value, UE_ARRAY_COUNT(Value)));
+		FCbFieldIterator Fields = Writer.Save();
+		TestEqual_Conditional(TEXT("FCbWriter(String, Unicode) Validate"), ValidateCompactBinaryRange(Fields.GetOuterBuffer(), ECbValidateMode::All), ECbValidateError::None)
+		{
+			for (FCbFieldView Field : Fields)
+			{
+				TestEqual(TEXT("FCbWriter(String, Unicode).AsString()"), Field.AsString(), "\xf0\x9f\x98\x80"_U8SV);
+			}
+		}
+	}
+}
+
+
+TEST_CASE("Core::Serialization::FCbWriter::Integer", "[Core][Serialization][Smoke]")
+{
+	TCbWriter<256> Writer;
+
+	auto TestInt32 = [&Writer](int32 Value)
+	{
+		Writer.Reset();
+		Writer.AddInteger(Value);
+		FCbField Field = Writer.Save();
+		TestEqual_Conditional(TEXT("FCbWriter(Integer, Int32) Validate"), ValidateCompactBinary(Field.GetOuterBuffer(), ECbValidateMode::All), ECbValidateError::None)
+		{
+			TestEqual(TEXT("FCbWriter(Integer, Int32) Value"), Field.AsInt32(), Value);
+			TestFalse(TEXT("FCbWriter(Integer, Int32) Error"), Field.HasError());
+		}
+	};
+
+	auto TestUInt32 = [&Writer](uint32 Value)
+	{
+		Writer.Reset();
+		Writer.AddInteger(Value);
+		FCbField Field = Writer.Save();
+		TestEqual_Conditional(TEXT("FCbWriter(Integer, UInt32) Validate"), ValidateCompactBinary(Field.GetOuterBuffer(), ECbValidateMode::All), ECbValidateError::None)
+		{
+			TestEqual(TEXT("FCbWriter(Integer, UInt32) Value"), Field.AsUInt32(), Value);
+			TestFalse(TEXT("FCbWriter(Integer, UInt32) Error"), Field.HasError());
+		}
+	};
+
+	auto TestInt64 = [&Writer](int64 Value)
+	{
+		Writer.Reset();
+		Writer.AddInteger(Value);
+		FCbField Field = Writer.Save();
+		TestEqual_Conditional(TEXT("FCbWriter(Integer, Int64) Validate"), ValidateCompactBinary(Field.GetOuterBuffer(), ECbValidateMode::All), ECbValidateError::None)
+		{
+			TestEqual(TEXT("FCbWriter(Integer, Int64) Value"), Field.AsInt64(), Value);
+			TestFalse(TEXT("FCbWriter(Integer, Int64) Error"), Field.HasError());
+		}
+	};
+
+	auto TestUInt64 = [&Writer](uint64 Value)
+	{
+		Writer.Reset();
+		Writer.AddInteger(Value);
+		FCbField Field = Writer.Save();
+		TestEqual_Conditional(TEXT("FCbWriter(Integer, UInt64) Validate"), ValidateCompactBinary(Field.GetOuterBuffer(), ECbValidateMode::All), ECbValidateError::None)
+		{
+			TestEqual(TEXT("FCbWriter(Integer, UInt64) Value"), Field.AsUInt64(), Value);
+			TestFalse(TEXT("FCbWriter(Integer, UInt64) Error"), Field.HasError());
+		}
+	};
+
+	TestUInt32(uint32(0x00));
+	TestUInt32(uint32(0x7f));
+	TestUInt32(uint32(0x80));
+	TestUInt32(uint32(0xff));
+	TestUInt32(uint32(0x0100));
+	TestUInt32(uint32(0x7fff));
+	TestUInt32(uint32(0x8000));
+	TestUInt32(uint32(0xffff));
+	TestUInt32(uint32(0x0001'0000));
+	TestUInt32(uint32(0x7fff'ffff));
+	TestUInt32(uint32(0x8000'0000));
+	TestUInt32(uint32(0xffff'ffff));
+
+	TestUInt64(uint64(0x0000'0001'0000'0000));
+	TestUInt64(uint64(0x7fff'ffff'ffff'ffff));
+	TestUInt64(uint64(0x8000'0000'0000'0000));
+	TestUInt64(uint64(0xffff'ffff'ffff'ffff));
+
+	TestInt32(int32(0x01));
+	TestInt32(int32(0x80));
+	TestInt32(int32(0x81));
+	TestInt32(int32(0x8000));
+	TestInt32(int32(0x8001));
+	TestInt32(int32(0x7fff'ffff));
+	TestInt32(int32(0x8000'0000));
+	TestInt32(int32(0x8000'0001));
+
+	TestInt64(int64(0x0000'0001'0000'0000));
+	TestInt64(int64(0x8000'0000'0000'0000));
+	TestInt64(int64(0x7fff'ffff'ffff'ffff));
+	TestInt64(int64(0x8000'0000'0000'0001));
+	TestInt64(int64(0xffff'ffff'ffff'ffff));
+}
+
+TEST_CASE("Core::Serialization::FCbWriter::Float", "[Core][Serialization][Smoke]")
+{
+	TCbWriter<256> Writer;
+
+	// Test Float32
+	{
+		Writer.Reset();
+		constexpr float Values[] = { 0.0f, 1.0f, -1.0f, UE_PI };
+		for (float Value : Values)
+		{
+			Writer.AddFloat(Value);
+		}
+		FCbFieldIterator Fields = Writer.Save();
+		TestEqual_Conditional(TEXT("FCbWriter(Float, Single) Validate"), ValidateCompactBinaryRange(Fields.GetOuterBuffer(), ECbValidateMode::All), ECbValidateError::None)
+		{
+			const float* CheckValue = Values;
+			for (FCbFieldView Field : Fields)
+			{
+				TestEqual(TEXT("FCbWriter(Float, Single).AsFloat()"), Field.AsFloat(), *CheckValue++);
+				TestFalse(TEXT("FCbWriter(Float, Single) Error"), Field.HasError());
+			}
+		}
+	}
+
+	// Test Float64
+	{
+		Writer.Reset();
+		constexpr double Values[] = { 0.0f, 1.0f, -1.0f, UE_PI, 1.9999998807907104, 1.9999999403953552, 3.4028234663852886e38, 6.8056469327705771e38 };
+		for (double Value : Values)
+		{
+			Writer.AddFloat(Value);
+		}
+		FCbFieldIterator Fields = Writer.Save();
+		TestEqual_Conditional(TEXT("FCbWriter(Float, Double) Validate"), ValidateCompactBinaryRange(Fields.GetOuterBuffer(), ECbValidateMode::All), ECbValidateError::None)
+		{
+			const double* CheckValue = Values;
+			for (FCbFieldView Field : Fields)
+			{
+				TestEqual(TEXT("FCbWriter(Float, Double).AsDouble()"), Field.AsDouble(), *CheckValue++);
+				TestFalse(TEXT("FCbWriter(Float, Double) Error"), Field.HasError());
+			}
+		}
+	}
+}
+
+TEST_CASE("Core::Serialization::FCbWriter::Bool", "[Core][Serialization][Smoke]")
+{
+	TCbWriter<256> Writer;
+
+	// Test Bool Values
+	{
+		Writer.AddBool(true);
+		Writer.AddBool(false);
+
+		FCbFieldIterator Fields = Writer.Save();
+		TestEqual_Conditional(TEXT("FCbWriter(Bool) Validate"), ValidateCompactBinaryRange(Fields.GetOuterBuffer(), ECbValidateMode::All), ECbValidateError::None)
+		{
+			TestTrue(TEXT("FCbWriter(Bool).AsBool()"), Fields.AsBool());
+			TestFalse(TEXT("FCbWriter(Bool) Error"), Fields.HasError());
+			++Fields;
+			TestFalse(TEXT("FCbWriter(Bool).AsBool()"), Fields.AsBool());
+			TestFalse(TEXT("FCbWriter(Bool) Error"), Fields.HasError());
+		}
+	}
+
+	// Test Bool Array/Object Uniformity
+	{
+		Writer.Reset();
+
+		Writer.BeginArray();
+		Writer.AddBool(false);
+		Writer.AddBool(false);
+		Writer.AddBool(false);
+		Writer.EndArray();
+
+		Writer.BeginObject();
+		Writer.SetName("B1"_ASV).AddBool(false);
+		Writer.SetName("B2"_ASV).AddBool(false);
+		Writer.SetName("B3"_ASV).AddBool(false);
+		Writer.EndObject();
+
+		FCbFieldIterator Fields = Writer.Save();
+		TestEqual(TEXT("FCbWriter(Bool, Uniform) Validate"), ValidateCompactBinaryRange(Fields.GetOuterBuffer(), ECbValidateMode::All), ECbValidateError::None);
+	}
+}
+
+TEST_CASE("Core::Serialization::FCbWriter::Object Attachment", "[Core][Serialization][Smoke]")
+{
+	TCbWriter<256> Writer;
+
+	const FIoHash::ByteArray ZeroBytes{};
+	const FIoHash::ByteArray SequentialBytes{ 1, 2, 3, 4, 5, 6, 7, 8, 9, 10, 11, 12, 13, 14, 15, 16, 17, 18, 19, 20 };
+
+	const FIoHash Values[] = { FIoHash(ZeroBytes), FIoHash(SequentialBytes) };
+	for (const FIoHash& Value : Values)
+	{
+		Writer.AddObjectAttachment(Value);
+	}
+
+	FCbFieldIterator Fields = Writer.Save();
+	TestEqual_Conditional(TEXT("FCbWriter(ObjectAttachment) Validate"), ValidateCompactBinaryRange(Fields.GetOuterBuffer(), ECbValidateMode::All), ECbValidateError::None)
+	{
+		const FIoHash* CheckValue = Values;
+		for (FCbFieldView Field : Fields)
+		{
+			TestEqual(TEXT("FCbWriter(ObjectAttachment).AsObjectAttachment()"), Field.AsObjectAttachment(), *CheckValue++);
+			TestFalse(TEXT("FCbWriter(ObjectAttachment) Error"), Field.HasError());
+		}
+	}
+}
+
+TEST_CASE("Core::Serialization::FCbWriter::Binary Attachment", "[Core][Serialization][Smoke]")
+{
+	TCbWriter<256> Writer;
+
+	const FIoHash::ByteArray ZeroBytes{};
+	const FIoHash::ByteArray SequentialBytes{ 1, 2, 3, 4, 5, 6, 7, 8, 9, 10, 11, 12, 13, 14, 15, 16, 17, 18, 19, 20 };
+
+	const FIoHash Values[] = { FIoHash(ZeroBytes), FIoHash(SequentialBytes) };
+	for (const FIoHash& Value : Values)
+	{
+		Writer.AddBinaryAttachment(Value);
+	}
+
+	FCbFieldIterator Fields = Writer.Save();
+	TestEqual_Conditional(TEXT("FCbWriter(BinaryAttachment) Validate"), ValidateCompactBinaryRange(Fields.GetOuterBuffer(), ECbValidateMode::All), ECbValidateError::None)
+	{
+		const FIoHash* CheckValue = Values;
+		for (FCbFieldView Field : Fields)
+		{
+			TestEqual(TEXT("FCbWriter(BinaryAttachment).AsBinaryAttachment()"), Field.AsBinaryAttachment(), *CheckValue++);
+			TestFalse(TEXT("FCbWriter(BinaryAttachment) Error"), Field.HasError());
+		}
+	}
+}
+
+
+TEST_CASE("Core::Serialization::FCbWriter::Hash", "[Core][Serialization][Smoke]")
+{
+	TCbWriter<256> Writer;
+
+	const FIoHash::ByteArray ZeroBytes{};
+	const FIoHash::ByteArray SequentialBytes{ 1, 2, 3, 4, 5, 6, 7, 8, 9, 10, 11, 12, 13, 14, 15, 16, 17, 18, 19, 20 };
+
+	const FIoHash Values[] = { FIoHash(ZeroBytes), FIoHash(SequentialBytes) };
+	for (const FIoHash& Value : Values)
+	{
+		Writer.AddHash(Value);
+	}
+
+	FCbFieldIterator Fields = Writer.Save();
+	TestEqual_Conditional(TEXT("FCbWriter(Hash) Validate"), ValidateCompactBinaryRange(Fields.GetOuterBuffer(), ECbValidateMode::All), ECbValidateError::None)
+	{
+		const FIoHash* CheckValue = Values;
+		for (FCbFieldView Field : Fields)
+		{
+			TestEqual(TEXT("FCbWriter(Hash).AsHash()"), Field.AsHash(), *CheckValue++);
+			TestFalse(TEXT("FCbWriter(Hash) Error"), Field.HasError());
+		}
+	}
+}
+
+TEST_CASE("Core::Serialization::FCbWriter::Uuid", "[Core][Serialization][Smoke]")
+{
+	TCbWriter<256> Writer;
+
+	const FGuid Values[] = { FGuid(), FGuid::NewGuid() };
+	for (const FGuid& Value : Values)
+	{
+		Writer.AddUuid(Value);
+	}
+
+	FCbFieldIterator Fields = Writer.Save();
+	TestEqual_Conditional(TEXT("FCbWriter(Uuid) Validate"), ValidateCompactBinaryRange(Fields.GetOuterBuffer(), ECbValidateMode::All), ECbValidateError::None)
+	{
+		const FGuid* CheckValue = Values;
+		for (FCbFieldView Field : Fields)
+		{
+			TestEqual(TEXT("FCbWriter(Uuid).AsUuid()"), Field.AsUuid(), *CheckValue++);
+			TestFalse(TEXT("FCbWriter(Uuid) Error"), Field.HasError());
+		}
+	}
+}
+
+
+TEST_CASE("Core::Serialization::FCbWriter::Date Time", "[Core][Serialization][Smoke]")
+{
+	TCbWriter<256> Writer;
+
+	const FDateTime Values[] = { FDateTime(0), FDateTime(2020, 5, 13, 15, 10) };
+	for (FDateTime Value : Values)
+	{
+		Writer.AddDateTime(Value);
+	}
+
+	FCbFieldIterator Fields = Writer.Save();
+	TestEqual_Conditional(TEXT("FCbWriter(DateTime) Validate"), ValidateCompactBinaryRange(Fields.GetOuterBuffer(), ECbValidateMode::All), ECbValidateError::None)
+	{
+		const FDateTime* CheckValue = Values;
+		for (FCbFieldView Field : Fields)
+		{
+			TestEqual(TEXT("FCbWriter(DateTime).AsDateTime()"), Field.AsDateTime(), *CheckValue++);
+			TestFalse(TEXT("FCbWriter(DateTime) Error"), Field.HasError());
+		}
+	};
+}
+
+TEST_CASE("Core::Serialization::FCbWriter::Time Span", "[Core][Serialization][Smoke]")
+{
+	TCbWriter<256> Writer;
+
+	const FTimespan Values[] = { FTimespan(0), FTimespan(1, 2, 4, 8) };
+	for (FTimespan Value : Values)
+	{
+		Writer.AddTimeSpan(Value);
+	}
+
+	FCbFieldIterator Fields = Writer.Save();
+	TestEqual_Conditional(TEXT("FCbWriter(TimeSpan) Validate"), ValidateCompactBinaryRange(Fields.GetOuterBuffer(), ECbValidateMode::All), ECbValidateError::None)
+	{
+		const FTimespan* CheckValue = Values;
+		for (FCbFieldView Field : Fields)
+		{
+			TestEqual(TEXT("FCbWriter(TimeSpan).AsTimeSpan()"), Field.AsTimeSpan(), *CheckValue++);
+			TestFalse(TEXT("FCbWriter(TimeSpan) Error"), Field.HasError());
+		}
+	}
+}
+
+TEST_CASE("Core::Serialization::FCbWriter::ObjectId", "[Core][Serialization][Smoke]")
+{
+	TCbWriter<256> Writer;
+
+	const FCbObjectId Values[] = { FCbObjectId(), FCbObjectId(MakeMemoryView<uint8>({1, 2, 3, 4, 5, 6, 7, 8, 9, 10, 11, 12})) };
+	for (const FCbObjectId& Value : Values)
+	{
+		Writer.AddObjectId(Value);
+	}
+
+	FCbFieldIterator Fields = Writer.Save();
+	TestEqual_Conditional(TEXT("FCbWriter(ObjectId) Validate"), ValidateCompactBinaryRange(Fields.GetOuterBuffer(), ECbValidateMode::All), ECbValidateError::None)
+	{
+		const FCbObjectId* CheckValue = Values;
+		for (FCbFieldView Field : Fields)
+		{
+			TestEqual(TEXT("FCbWriter(ObjectId).AsObjectId()"), Field.AsObjectId(), *CheckValue++);
+			TestFalse(TEXT("FCbWriter(ObjectId) Error"), Field.HasError());
+		}
+	}
+}
+
+TEST_CASE("Core::Serialization::FCbWriter::Custom By Id", "[Core][Serialization][Smoke]")
+{
+	TCbWriter<256> Writer;
+
+	struct FCustomValue
+	{
+		uint64 Type;
+		TArray<uint8, TInlineAllocator<16>> Bytes;
+	};
+	const FCustomValue Values[] =
+	{
+		{ 1, {1, 2, 3} },
+		{ MAX_uint64, {4, 5, 6} },
+	};
+
+	for (const FCustomValue& Value : Values)
+	{
+		Writer.AddCustom(Value.Type, MakeMemoryView(Value.Bytes));
+	}
+
+	FCbFieldIterator Fields = Writer.Save();
+	TestEqual_Conditional(TEXT("FCbWriter(CustomById) Validate"), ValidateCompactBinaryRange(Fields.GetOuterBuffer(), ECbValidateMode::All), ECbValidateError::None)
+	{
+		const FCustomValue* CheckValue = Values;
+		for (FCbFieldView Field : Fields)
+		{
+			TestTrue(TEXT("FCbWriter(CustomById).AsCustom()"), Field.AsCustom(CheckValue->Type).EqualBytes(MakeMemoryView(CheckValue->Bytes)));
+			TestFalse(TEXT("FCbWriter(CustomById) Error"), Field.HasError());
+			++CheckValue;
+		}
+	}
+}
+
+TEST_CASE("Core::Serialization::FCbWriter::Custom By Name", "[Core][Serialization][Smoke]")
+{
+	TCbWriter<256> Writer;
+
+	struct FCustomValue
+	{
+		FUtf8StringView Type;
+		TArray<uint8, TInlineAllocator<16>> Bytes;
+	};
+	const FCustomValue Values[] =
+	{
+		{ "Type1"_U8SV, {1, 2, 3} },
+		{ "Type2"_U8SV, {4, 5, 6} },
+	};
+
+	for (const FCustomValue& Value : Values)
+	{
+		Writer.AddCustom(Value.Type, MakeMemoryView(Value.Bytes));
+	}
+
+	FCbFieldIterator Fields = Writer.Save();
+	TestEqual_Conditional(TEXT("FCbWriter(CustomByName) Validate"), ValidateCompactBinaryRange(Fields.GetOuterBuffer(), ECbValidateMode::All), ECbValidateError::None)
+	{
+		const FCustomValue* CheckValue = Values;
+		for (FCbFieldView Field : Fields)
+		{
+			TestTrue(TEXT("FCbWriter(CustomByName).AsCustom()"), Field.AsCustom(CheckValue->Type).EqualBytes(MakeMemoryView(CheckValue->Bytes)));
+			TestFalse(TEXT("FCbWriter(CustomByName) Error"), Field.HasError());
+			++CheckValue;
+		}
+	}
+}
+
+TEST_CASE("Core::Serialization::FCbWriter::Complex", "[Core][Serialization][Smoke]")
+{
+	FCbObject Object;
+	FBufferArchive Archive;
+	{
+		FCbWriter Writer;
+		Writer.BeginObject();
+
+		const uint8 LocalField[] = { uint8(ECbFieldType::IntegerPositive | ECbFieldType::HasFieldName), 1, 'I', 42 };
+		Writer.AddField("FieldViewCopy"_ASV, FCbFieldView(LocalField));
+		Writer.AddField("FieldCopy"_ASV, FCbField(FSharedBuffer::Clone(MakeMemoryView(LocalField))));
+
+		const uint8 LocalObject[] = { uint8(ECbFieldType::Object | ECbFieldType::HasFieldName), 1, 'O', 7,
+			uint8(ECbFieldType::IntegerPositive | ECbFieldType::HasFieldName), 1, 'I', 42,
+			uint8(ECbFieldType::Null | ECbFieldType::HasFieldName), 1, 'N' };
+		Writer.AddObject("ObjectViewCopy"_ASV, FCbObjectView(LocalObject));
+		Writer.AddObject("ObjectCopy"_ASV, FCbObject(FSharedBuffer::Clone(MakeMemoryView(LocalObject))));
+
+		const uint8 LocalArray[] = { uint8(ECbFieldType::UniformArray | ECbFieldType::HasFieldName), 1, 'A', 4, 2,
+			uint8(ECbFieldType::IntegerPositive), 42, 21 };
+		Writer.AddArray("ArrayViewCopy"_ASV, FCbArrayView(LocalArray));
+		Writer.AddArray("ArrayCopy"_ASV, FCbArray(FSharedBuffer::Clone(MakeMemoryView(LocalArray))));
+
+		Writer.AddNull("Null"_ASV);
+
+		Writer.BeginObject("Binary"_ASV);
+		{
+			Writer.AddBinary("Empty"_ASV, FMemoryView());
+			Writer.AddBinary("Value"_ASV, MakeMemoryView("BinaryValue"));
+			Writer.AddBinary("LargeViewValue"_ASV, MakeMemoryView(FString::ChrN(256, TEXT('.'))));
+			Writer.AddBinary("LargeValue"_ASV, FSharedBuffer::Clone(MakeMemoryView(FString::ChrN(256, TEXT('!')))));
+		}
+		Writer.EndObject();
+
+		Writer.BeginObject("Strings"_ASV);
+		{
+			Writer.AddString("AnsiString"_ASV, "AnsiValue"_ASV);
+			Writer.AddString("WideString"_ASV, FString::ChrN(256, TEXT('.')));
+			Writer.AddString("EmptyAnsiString"_ASV, FAnsiStringView());
+			Writer.AddString("EmptyWideString"_ASV, FWideStringView());
+			Writer.AddString("EmptyUtf8String"_U8SV, FUtf8StringView());
+			Writer.AddString("AnsiStringLiteral", "AnsiValue");
+			Writer.AddString("WideStringLiteral", TEXT("AnsiValue"));
+		}
+		Writer.EndObject();
+
+		Writer.BeginArray("Integers"_ASV);
+		{
+			Writer.AddInteger(int32(-1));
+			Writer.AddInteger(int64(-1));
+			Writer.AddInteger(uint32(1));
+			Writer.AddInteger(uint64(1));
+			Writer.AddInteger(MIN_int32);
+			Writer.AddInteger(MAX_int32);
+			Writer.AddInteger(MAX_uint32);
+			Writer.AddInteger(MIN_int64);
+			Writer.AddInteger(MAX_int64);
+			Writer.AddInteger(MAX_uint64);
+		}
+		Writer.EndArray();
+
+		Writer.BeginArray("UniformIntegers"_ASV);
+		{
+			Writer.AddInteger(0);
+			Writer.AddInteger(MAX_int32);
+			Writer.AddInteger(MAX_uint32);
+			Writer.AddInteger(MAX_int64);
+			Writer.AddInteger(MAX_uint64);
+		}
+		Writer.EndArray();
+
+		Writer.AddFloat("Float32"_ASV, 1.0f);
+		Writer.AddFloat("Float64as32"_ASV, 2.0);
+		Writer.AddFloat("Float64"_ASV, 3.0e100);
+
+		Writer.AddBool("False"_ASV, false);
+		Writer.AddBool("True"_ASV, true);
+
+		Writer.AddObjectAttachment("ObjectAttachment"_ASV, FIoHash());
+		Writer.AddBinaryAttachment("BinaryAttachment"_ASV, FIoHash());
+		Writer.AddAttachment("Attachment"_ASV, FCbAttachment());
+
+		Writer.AddHash("Hash"_ASV, FIoHash());
+		Writer.AddUuid("Uuid"_ASV, FGuid());
+
+		Writer.AddDateTimeTicks("DateTimeZero"_ASV, 0);
+		Writer.AddDateTime("DateTime2020"_ASV, FDateTime(2020, 5, 13, 15, 10));
+
+		Writer.AddTimeSpanTicks("TimeSpanZero"_ASV, 0);
+		Writer.AddTimeSpan("TimeSpan"_ASV, FTimespan(1, 2, 4, 8));
+
+		Writer.AddObjectId("ObjectIdZero"_ASV, FCbObjectId());
+		Writer.AddObjectId("ObjectId"_ASV, FCbObjectId(MakeMemoryView<uint8>({ 1, 2, 3, 4, 5, 6, 7, 8, 9, 10, 11, 12 })));
+
+		Writer.BeginObject("NestedObjects"_ASV);
+		{
+			Writer.BeginObject("Empty"_ASV);
+			Writer.EndObject();
+
+			Writer.BeginObject("Null"_ASV);
+			Writer.AddNull("Null"_ASV);
+			Writer.EndObject();
+		}
+		Writer.EndObject();
+
+		Writer.BeginArray("NestedArrays"_ASV);
+		{
+			Writer.BeginArray();
+			Writer.EndArray();
+
+			Writer.BeginArray();
+			Writer.AddNull();
+			Writer.AddNull();
+			Writer.AddNull();
+			Writer.EndArray();
+
+			Writer.BeginArray();
+			Writer.AddBool(false);
+			Writer.AddBool(false);
+			Writer.AddBool(false);
+			Writer.EndArray();
+
+			Writer.BeginArray();
+			Writer.AddBool(true);
+			Writer.AddBool(true);
+			Writer.AddBool(true);
+			Writer.EndArray();
+		}
+		Writer.EndArray();
+
+		Writer.BeginArray("ArrayOfObjects"_ASV);
+		{
+			Writer.BeginObject();
+			Writer.EndObject();
+
+			Writer.BeginObject();
+			Writer.AddNull("Null"_ASV);
+			Writer.EndObject();
+		}
+		Writer.EndArray();
+
+		Writer.BeginArray("LargeArray"_ASV);
+		for (int Index = 0; Index < 256; ++Index)
+		{
+			Writer.AddInteger(Index - 128);
+		}
+		Writer.EndArray();
+
+		Writer.BeginArray("LargeUniformArray"_ASV);
+		for (int Index = 0; Index < 256; ++Index)
+		{
+			Writer.AddInteger(Index);
+		}
+		Writer.EndArray();
+
+		Writer.BeginArray("NestedUniformArray"_ASV);
+		for (int Index = 0; Index < 16; ++Index)
+		{
+			Writer.BeginArray();
+			for (int Value = 0; Value < 4; ++Value)
+			{
+				Writer.AddInteger(Value);
+			}
+			Writer.EndArray();
+		}
+		Writer.EndArray();
+
+		Writer.EndObject();
+		Object = Writer.Save().AsObject();
+
+		Writer.Save(Archive);
+		TestEqual(TEXT("FCbWriter(Complex).Save(Ar)->Num()"), uint64(Archive.Num()), Writer.GetSaveSize());
+	}
+
+	TestEqual(TEXT("FCbWriter(Complex).Save()->Validate"),
+		ValidateCompactBinary(Object.GetOuterBuffer(), ECbValidateMode::All), ECbValidateError::None);
+
+	TestEqual(TEXT("FCbWriter(Complex).Save(Ar)->Validate"),
+		ValidateCompactBinary(MakeMemoryView(Archive), ECbValidateMode::All), ECbValidateError::None);
+}
+
+
+TEST_CASE("Core::Serialization::FCbWriter::Owned Read Only", "[Core][Serialization][Smoke]")
+{
+	FCbWriter Writer;
+	Writer.BeginObject();
+	Writer.EndObject();
+	FCbObject Object = Writer.Save().AsObject();
+	TestTrue(TEXT("FCbWriter().Save().IsOwned()"), Object.IsOwned());
+}
+
+TEST_CASE("Core::Serialization::FCbWriter::Stream", "[Core][Serialization][Smoke]")
+{
+	FCbObject Object;
+	{
+		FCbWriter Writer;
+		Writer.BeginObject();
+
+		const uint8 LocalField[] = { uint8(ECbFieldType::IntegerPositive | ECbFieldType::HasFieldName), 1, 'I', 42 };
+		Writer << "FieldCopy"_ASV << FCbFieldView(LocalField);
+
+		const uint8 LocalObject[] = { uint8(ECbFieldType::Object | ECbFieldType::HasFieldName), 1, 'O', 7,
+			uint8(ECbFieldType::IntegerPositive | ECbFieldType::HasFieldName), 1, 'I', 42,
+			uint8(ECbFieldType::Null | ECbFieldType::HasFieldName), 1, 'N' };
+		Writer << "ObjectCopy"_ASV << FCbObjectView(LocalObject);
+
+		const uint8 LocalArray[] = { uint8(ECbFieldType::UniformArray | ECbFieldType::HasFieldName), 1, 'A', 4, 2,
+			uint8(ECbFieldType::IntegerPositive), 42, 21 };
+		Writer << "ArrayCopy"_ASV << FCbArrayView(LocalArray);
+
+		Writer << "Null"_ASV << nullptr;
+
+		Writer << "Strings"_ASV;
+		Writer.BeginObject();
+		Writer
+			<< "AnsiString"_ASV << "AnsiValue"_ASV
+			<< "AnsiStringLiteral"_ASV << "AnsiValue"
+			<< "WideString"_ASV << TEXT("WideValue"_SV)
+			<< "WideStringLiteral"_ASV << TEXT("WideValue");
+		Writer.EndObject();
+
+		Writer << "Integers"_ASV;
+		Writer.BeginArray();
+		Writer << int32(-1) << int64(-1) << uint32(1) << uint64(1);
+		Writer.EndArray();
+
+		Writer << "Float32"_ASV << 1.0f;
+		Writer << "Float64"_ASV << 2.0;
+
+		Writer << "False"_ASV << false << "True"_ASV << true;
+
+		Writer << "Attachment"_ASV << FCbAttachment();
+
+		Writer << "Hash"_ASV << FIoHash();
+		Writer << "Uuid"_ASV << FGuid();
+
+		Writer << "DateTime"_ASV << FDateTime(2020, 5, 13, 15, 10);
+		Writer << "TimeSpan"_ASV << FTimespan(1, 2, 4, 8);
+
+		Writer << "ObjectIdZero"_ASV << FCbObjectId();
+		Writer << "ObjectId"_ASV << FCbObjectId(MakeMemoryView<uint8>({ 1, 2, 3, 4, 5, 6, 7, 8, 9, 10, 11, 12 }));
+
+		Writer << "LiteralName" << nullptr;
+
+		Writer.EndObject();
+		Object = Writer.Save().AsObject();
+	}
+
+	TestEqual(TEXT("FCbWriter(Stream) Validate"), ValidateCompactBinary(Object.GetOuterBuffer(), ECbValidateMode::All), ECbValidateError::None);
+}
+
+TEST_CASE("Core::Serialization::FCbWriter::State", "[Core][Serialization][Smoke]")
+{
+	FCbWriter Writer;
+
+	// Assert on saving an empty writer.
+	//uint8 EmptyField[1];
+	//Writer.Reset();
+	//Writer.Save(MakeMemoryView(EmptyField));
+
+	// Assert on under-sized save buffer.
+	//uint8 ZeroFieldSmall[1];
+	//Writer.Reset();
+	//Writer.AddInteger(0);
+	//Writer.Save(MakeMemoryView(ZeroFieldSmall));
+
+	// Assert on over-sized save buffer.
+	//uint8 ZeroFieldLarge[3];
+	//Writer.Reset();
+	//Writer.AddInteger(0);
+	//Writer.Save(MakeMemoryView(ZeroFieldLarge));
+
+	// Assert on empty name.
+	//Writer.SetName(""_ASV);
+
+	// Assert on name after name.
+	//Writer.SetName("Field"_ASV).SetName("Field"_ASV);
+
+	// Assert on missing name.
+	//Writer.BeginObject();
+	//Writer.AddNull();
+	//Writer.EndObject();
+
+	// Assert on name in array.
+	//Writer.BeginArray();
+	//Writer.SetName("Field"_ASV);
+	//Writer.EndArray();
+
+	// Assert on save in object.
+	//uint8 InvalidObject[1];
+	//Writer.Reset();
+	//Writer.BeginObject();
+	//Writer.Save(MakeMemoryView(InvalidObject));
+	//Writer.EndObject();
+
+	// Assert on save in array.
+	//uint8 InvalidArray[1];
+	//Writer.Reset();
+	//Writer.BeginArray();
+	//Writer.Save(MakeMemoryView(InvalidArray));
+	//Writer.EndArray();
+
+	// Assert on object end with no begin.
+	//Writer.EndObject();
+
+	// Assert on array end with no begin.
+	//Writer.EndArray();
+
+	// Assert on object end after name with no value.
+	//Writer.BeginObject();
+	//Writer.SetName("Field"_ASV);
+	//Writer.EndObject();
+
+	// Assert on writing a field with no value.
+	//Writer.Field(FCbFieldView());
+}
\ No newline at end of file
diff --git a/Engine/Source/Runtime/Core/Tests/Serialization/MemoryReaderTest.cpp b/Engine/Source/Runtime/Core/Tests/Serialization/MemoryReaderTest.cpp
new file mode 100644
--- /dev/null
+++ b/Engine/Source/Runtime/Core/Tests/Serialization/MemoryReaderTest.cpp
@@ -0,0 +1,176 @@
+// Copyright Epic Games, Inc. All Rights Reserved.
+
+#include "CoreMinimal.h"
+#include "Serialization/MemoryReader.h"
+#include "Serialization/MemoryWriter.h"
+#include "TestHarness.h"
+
+TEST_CASE("Core::Serialization::MemoryReader::Memory Reader", "[Core][Serialization][Smoke]")
+{
+	// Test reading uint64 byte swapped, ensuring that bytes really get swapped.
+	{
+		const uint64 WrittenValue  = 0x8877665544332211ull;
+		const uint64 ExpectedValue = 0x1122334455667788ull;
+
+		TArray<uint8> Bytes;
+		Bytes.AddUninitialized(sizeof(uint64));
+		*reinterpret_cast<uint64*>(Bytes.GetData()) = WrittenValue;
+
+		FMemoryReader Reader(Bytes);
+		Reader.SetByteSwapping(true);
+		uint64 ReadValue;
+		Reader << ReadValue;
+		TestTrue("Test reading uint64 in swapped byte order.", ReadValue == ExpectedValue);
+	}
+
+	// Write all supported types and read them back. (Assuming that MemoryWriterTest tests don't fail).
+	{
+		// Keeps the 'official' test value.
+		const uint8  TestValueU8   = 0x12;
+		const int8   TestValueS8   = 0x34;
+		const uint16 TestValueU16  = 0x1122;
+		const int16  TestValueS16  = 0x3344;
+		const uint32 TestValueU32  = 0x11223344;
+		const int32  TestValueS32  = 0x55667788;
+		const uint64 TestValueU64  = 0x1122334455667788;
+		const int64  TestValueS64  = 0x99AABBCCDDEEFF00;
+		const float  TestValueF    = 128.5f;
+		const double TestValueD    = 256.5;
+		const bool   TestValueB    = true;
+		const WIDECHAR TestValueCh = 0xF2;
+		const FString TestAnsiStr  = TEXT("Joe");
+		const FString TestUtf16Str = TEXT("\uC11C\uC6B8\uC0AC\uB78C"); // Must contains at least one non-ANSI codepoint to ensure UTF16 is used.
+
+		// The values to write.
+		TArray<uint8> Bytes;
+		FMemoryWriter Writer(Bytes);
+		uint8  ValueU8   = TestValueU8;
+		int8   ValueS8   = TestValueS8;
+		uint16 ValueU16  = TestValueU16;
+		int16  ValueS16  = TestValueS16;
+		uint32 ValueU32  = TestValueU32;
+		int32  ValueS32  = TestValueS32;
+		uint64 ValueU64  = TestValueU64;
+		int64  ValueS64  = TestValueS64;
+		float  ValueF    = TestValueF;
+		double ValueD    = TestValueD;
+		bool   ValueB    = TestValueB;
+		WIDECHAR ValueCh = TestValueCh;
+		FString AnsiStr  = TestAnsiStr;
+		FString Utf16Str = TestUtf16Str;
+
+		// Try to not align the values in the buffer, so that the test would fail (on platforms not supporting misaligned read) if the reader tries to cast and dereference a pointer type with misaligned address. Doing like: *(uint64*)Array.GetData();
+		Writer << ValueU8;
+		Writer << ValueU16;
+		Writer << ValueU32;
+		Writer << ValueU64;
+		Writer << ValueS8;
+		Writer << ValueS16;
+		Writer << ValueS32;
+		Writer << ValueS64;
+		Writer << ValueF;
+		Writer << ValueD;
+		Writer << ValueB;
+		Writer << ValueCh;
+		Writer << AnsiStr;
+		Writer << Utf16Str;
+
+		// Write again, but write the bytes swapped.
+		Writer.SetByteSwapping(true);
+
+		Writer << ValueU8;
+		Writer << ValueU16;
+		Writer << ValueU32;
+		Writer << ValueU64;
+		Writer << ValueS8;
+		Writer << ValueS16;
+		Writer << ValueS32;
+		Writer << ValueS64;
+		Writer << ValueF;
+		Writer << ValueD;
+		Writer << ValueB;
+		Writer << ValueCh;
+		Writer << AnsiStr;
+		Writer << Utf16Str;
+
+		// The value to reads without swapping the the value swapped back to this platform endianness.
+		uint8  ReadValueU8,  ReadValueSwapU8;
+		int8   ReadValueS8,  ReadValueSwapS8;
+		uint16 ReadValueU16, ReadValueSwapU16;
+		int16  ReadValueS16, ReadValueSwapS16;
+		uint32 ReadValueU32, ReadValueSwapU32;
+		int32  ReadValueS32, ReadValueSwapS32;
+		uint64 ReadValueU64, ReadValueSwapU64;
+		int64  ReadValueS64, ReadValueSwapS64;
+		float  ReadValueF,   ReadValueSwapF;
+		double ReadValueD,   ReadValueSwapD;
+		bool   ReadValueB,   ReadValueSwapB;
+		WIDECHAR ReadValueCh, ReadValueSwapCh;
+		FString ReadAnsiStr,  ReadSwapAnsiStr;
+		FString ReadUtf16Str, ReadSwapUtf16Str;
+
+		// Read the first set of values written in this platform endianness.
+		FMemoryReader Reader(Bytes);
+		Reader << ReadValueU8;
+		Reader << ReadValueU16;
+		Reader << ReadValueU32;
+		Reader << ReadValueU64;
+		Reader << ReadValueS8;
+		Reader << ReadValueS16;
+		Reader << ReadValueS32;
+		Reader << ReadValueS64;
+		Reader << ReadValueF;
+		Reader << ReadValueD;
+		Reader << ReadValueB;
+		Reader << ReadValueCh;
+		Reader << ReadAnsiStr;
+		Reader << ReadUtf16Str;
+
+		// Read the swapped values, swapping them back to their original value (Assuming that MemoryWriterTest tests don't fail)
+		Reader.SetByteSwapping(true);
+		Reader << ReadValueSwapU8;
+		Reader << ReadValueSwapU16;
+		Reader << ReadValueSwapU32;
+		Reader << ReadValueSwapU64;
+		Reader << ReadValueSwapS8;
+		Reader << ReadValueSwapS16;
+		Reader << ReadValueSwapS32;
+		Reader << ReadValueSwapS64;
+		Reader << ReadValueSwapF;
+		Reader << ReadValueSwapD;
+		Reader << ReadValueSwapB;
+		Reader << ReadValueSwapCh;
+		Reader << ReadSwapAnsiStr;
+		Reader << ReadSwapUtf16Str;
+
+		// Validate that the value read are the expected ones.
+		TestTrue("Test reading 'u8' from byte stream",             ReadValueU8      == TestValueU8);
+		TestTrue("Test reading 'u16' from byte stream",            ReadValueU16     == TestValueU16);
+		TestTrue("Test reading 'u32' from byte stream",            ReadValueU32     == TestValueU32);
+		TestTrue("Test reading 'u64' from byte stream",            ReadValueU64     == TestValueU64);
+		TestTrue("Test reading 's8' from byte stream",             ReadValueS8      == TestValueS8);
+		TestTrue("Test reading 's16' from byte stream",            ReadValueS16     == TestValueS16);
+		TestTrue("Test reading 's32' from byte stream",            ReadValueS32     == TestValueS32);
+		TestTrue("Test reading 's64' from byte stream",            ReadValueS64     == TestValueS64);
+		TestTrue("Test reading 'foat' from byte stream",           ReadValueF       == TestValueF);
+		TestTrue("Test reading 'double' from byte stream",         ReadValueD       == TestValueD);
+		TestTrue("Test reading 'bool' from byte stream",           ReadValueB       == TestValueB);
+		TestTrue("Test reading 'wchar' from byte stream",          ReadValueCh      == TestValueCh);
+		TestTrue("Test reading 'ansi str' from byte stream",       ReadAnsiStr      == TestAnsiStr);
+		TestTrue("Test reading 'utf16 str' from byte stream",      ReadUtf16Str     == TestUtf16Str);
+		TestTrue("Test reading 'u8-swapped' from byte stream",     ReadValueSwapU8  == TestValueU8);
+		TestTrue("Test reading 'u16-swapped' from byte stream",    ReadValueSwapU16 == TestValueU16);
+		TestTrue("Test reading 'u32-swapped' from byte stream",    ReadValueSwapU32 == TestValueU32);
+		TestTrue("Test reading 'u64-swapped' from byte stream",    ReadValueSwapU64 == TestValueU64);
+		TestTrue("Test reading 's8-swapped' from byte stream",     ReadValueSwapS8  == TestValueS8);
+		TestTrue("Test reading 's16'-swapped' from byte stream",   ReadValueSwapS16 == TestValueS16);
+		TestTrue("Test reading 's32-swapped' from byte stream",    ReadValueSwapS32 == TestValueS32);
+		TestTrue("Test reading 's64-swapped' from byte stream",    ReadValueSwapS64 == TestValueS64);
+		TestTrue("Test reading 'float-swapped' from byte stream",  ReadValueSwapF   == TestValueF);
+		TestTrue("Test reading 'double-swapped' from byte stream", ReadValueSwapD   == TestValueD);
+		TestTrue("Test reading 'bool-swapped' from byte stream",   ReadValueSwapB   == TestValueB);
+		TestTrue("Test reading 'wchar-swapped' from byte stream",  ReadValueSwapCh  == TestValueCh);
+		TestTrue("Test reading 'ansi-swapped' from byte stream",   ReadSwapAnsiStr  == TestAnsiStr);
+		TestTrue("Test reading 'utf16-swapped' from byte stream",  ReadSwapUtf16Str == TestUtf16Str);
+	}
+}
