Add ResetData to FBitReader that won't reallocate the internal buffer unless it is necessary
Reducing allocations during replay scrub

diff --git a/Engine/Source/Runtime/Core/Private/Serialization/BitReader.cpp b/Engine/Source/Runtime/Core/Private/Serialization/BitReader.cpp
--- a/Engine/Source/Runtime/Core/Private/Serialization/BitReader.cpp
+++ b/Engine/Source/Runtime/Core/Private/Serialization/BitReader.cpp
@@ -188,6 +188,21 @@ void FBitReader::SetData( FBitReader& Src, int64 CountBits )
 	Src.SerializeBits(Buffer.GetData(), CountBits);
 }
 
+void FBitReader::ResetData(FBitReader& Src, int64 CountBits)
+{
+	Num = CountBits;
+	Pos = 0;
+	ClearError();
+
+	// Setup network version
+	this->SetEngineNetVer(Src.EngineNetVer());
+	this->SetGameNetVer(Src.GameNetVer());
+
+	Buffer.Reset();
+	Buffer.AddUninitialized((CountBits + 7) >> 3);
+	Src.SerializeBits(Buffer.GetData(), CountBits);
+}
+
 /** This appends data from another BitReader. It checks that this bit reader is byte-aligned so it can just do a TArray::Append instead of a bitcopy.
  *	It is intended to be used by performance minded code that wants to ensure an appBitCpy is avoided.
  */
diff --git a/Engine/Source/Runtime/Core/Public/Serialization/BitReader.h b/Engine/Source/Runtime/Core/Public/Serialization/BitReader.h
--- a/Engine/Source/Runtime/Core/Public/Serialization/BitReader.h
+++ b/Engine/Source/Runtime/Core/Public/Serialization/BitReader.h
@@ -35,6 +35,9 @@ public:
 	void SetData( uint8* Src, int64 CountBits );
 	void SetData( TArray<uint8>&& Src, int64 CountBits );
 
+	/** Equivalent to SetData (reset position, copy from Src into internal buffer), but uses Reset not Empty to avoid a realloc if possible. */
+	void ResetData(FBitReader& Src, int64 CountBits);
+
 #if defined(_MSC_VER) && PLATFORM_DESKTOP
 #pragma warning( push )
 #pragma warning( disable : 4789 )	// Windows PGO (LTCG) is causing nonsensical errors in certain build environments
diff --git a/Engine/Source/Runtime/Engine/Private/RepLayout.cpp b/Engine/Source/Runtime/Engine/Private/RepLayout.cpp
--- a/Engine/Source/Runtime/Engine/Private/RepLayout.cpp
+++ b/Engine/Source/Runtime/Engine/Private/RepLayout.cpp
@@ -3820,6 +3820,9 @@ bool FRepLayout::ReceiveProperties_BackwardsCompatible_r(
 		}
 	}
 
+	FNetBitReader TempReader;
+	TempReader.PackageMap = Reader.PackageMap;
+
 	while (true)
 	{
 		uint32 NetFieldExportHandle = 0;
@@ -3862,10 +3865,7 @@ bool FRepLayout::ReceiveProperties_BackwardsCompatible_r(
 			return false;
 		}
 
-		FNetBitReader TempReader;
-		
-		TempReader.PackageMap = Reader.PackageMap;
-		TempReader.SetData(Reader, NumBits);
+		TempReader.ResetData(Reader, NumBits);
 
 		if (Reader.IsError())
 		{
diff --git a/Engine/Source/Runtime/Engine/Private/ReplayHelper.cpp b/Engine/Source/Runtime/Engine/Private/ReplayHelper.cpp
--- a/Engine/Source/Runtime/Engine/Private/ReplayHelper.cpp
+++ b/Engine/Source/Runtime/Engine/Private/ReplayHelper.cpp
@@ -1932,16 +1932,16 @@ bool FReplayHelper::ReadDemoFrame(UNetConnection* Connection, FArchive& Ar, TArr
 		// This can occur if the Archive is "double read" due to a level fast forward.
 		const bool bAddToSeenList = bAppendPackets && !bForLevelFastForward;
 
-		FString NameTemp;
 		for (uint32 i = 0; i < NumStreamingLevels; i++)
 		{
+			FString NameTemp;
 			Ar << NameTemp;
 
 			if (bAddToSeenList)
 			{
 				// Add this level to the seen list, but don't actually mark it as being seen.
 				// It will be marked when we have processed packets for it.
-				const FLevelStatus& LevelStatus = FindOrAddLevelStatus(NameTemp);
+				const FLevelStatus& LevelStatus = FindOrAddLevelStatus(MoveTemp(NameTemp));
 				SeenLevelStatuses.Add(LevelStatus.LevelIndex);
 			}
 		}
diff --git a/Engine/Source/Runtime/Engine/Public/ReplayHelper.h b/Engine/Source/Runtime/Engine/Public/ReplayHelper.h
--- a/Engine/Source/Runtime/Engine/Public/ReplayHelper.h
+++ b/Engine/Source/Runtime/Engine/Public/ReplayHelper.h
@@ -482,6 +482,11 @@ private:
 	}
 
 	FLevelStatus& FindOrAddLevelStatus(const FString& LevelPackageName)
+	{
+		return FindOrAddLevelStatus(FString(LevelPackageName));
+	}
+
+	FLevelStatus& FindOrAddLevelStatus(FString&& LevelPackageName)
 	{
 		if (int32* LevelStatusIndex = LevelStatusesByName.Find(LevelPackageName))
 		{
@@ -491,7 +496,7 @@ private:
 		const int32 Index = AllLevelStatuses.Emplace(LevelPackageName);
 		AllLevelStatuses[Index].LevelIndex = Index;
 
-		LevelStatusesByName.Add(LevelPackageName, Index);
+		LevelStatusesByName.Add(MoveTemp(LevelPackageName), Index);
 		NumLevelsAddedThisFrame++;
 
 		return AllLevelStatuses[Index];
