Fix hang on exit due to flush causing other threads to attempt to allocate when crashing

We are now no longer flushing this final callstack into the log, which is a major regression

diff --git a/Engine/Source/Runtime/Core/Private/Unix/UnixErrorOutputDevice.cpp b/Engine/Source/Runtime/Core/Private/Unix/UnixErrorOutputDevice.cpp
--- a/Engine/Source/Runtime/Core/Private/Unix/UnixErrorOutputDevice.cpp
+++ b/Engine/Source/Runtime/Core/Private/Unix/UnixErrorOutputDevice.cpp
@@ -101,7 +101,9 @@ void FUnixErrorOutputDevice::HandleError()
 		UE_LOG(LogCore, Log, TEXT("=== Critical error: ===") LINE_TERMINATOR TEXT("%s") LINE_TERMINATOR, GErrorExceptionDescription);
 		UE_LOG(LogCore, Log, TEXT("%s"), GErrorHist);
 
-		GLog->Flush();
+		// TODO Fix this, as we no longer *output* this callstack most the time now
+		// This fixes a hang vs using GLog->Flush()
+		GLog->PanicFlush();
 
 		HandleErrorRestoreUI();
 
