Kill FGenericWidePlatformString::Strncmp, which reads one character out of bounds for unterminated strings, and fallback to my faster FGenericPlatform::Strncmp algorithm

diff --git a/Engine/Source/Runtime/Core/Public/GenericPlatform/GenericWidePlatformString.h b/Engine/Source/Runtime/Core/Public/GenericPlatform/GenericWidePlatformString.h
--- a/Engine/Source/Runtime/Core/Public/GenericPlatform/GenericWidePlatformString.h
+++ b/Engine/Source/Runtime/Core/Public/GenericPlatform/GenericWidePlatformString.h
@@ -53,20 +53,6 @@ public:
 		return 0;
 	}
 
-	CORE_API static int32 Strncmp( const WIDECHAR* String1, const WIDECHAR* String2, SIZE_T Count )
-	{
-		// walk the strings, comparing them case sensitively, up to a max size
-		for (; (*String1 || *String2) && Count; String1++, String2++, Count--)
-		{
-			TCHAR A = *String1, B = *String2;
-			if (A != B)
-			{
-				return A - B;
-			}
-		}
-		return 0;
-	}
-
 	CORE_API static int32 Strlen( const WIDECHAR* String )
 	{
 		int32 Length = -1;
