Fix for the legacy math macros always being defined even when UE_DEFINE_LEGACY_MATH_CONSTANT_MACRO_NAMES=0.

diff --git a/Engine/Source/Runtime/Core/Public/Math/UnrealMathUtility.h b/Engine/Source/Runtime/Core/Public/Math/UnrealMathUtility.h
--- a/Engine/Source/Runtime/Core/Public/Math/UnrealMathUtility.h
+++ b/Engine/Source/Runtime/Core/Public/Math/UnrealMathUtility.h
@@ -60,7 +60,7 @@ class TRange;
 // - Remove UE_DEPRECATE_LEGACY_MATH_CONSTANT_MACRO_NAMES=1 from the module's .Build.cs.
 // - Define UE_DEFINE_LEGACY_MATH_CONSTANT_MACRO_NAMES=0 from the module's .Build.cs to stop new usage of the legacy macros.
 
-#ifdef UE_DEFINE_LEGACY_MATH_CONSTANT_MACRO_NAMES
+#if UE_DEFINE_LEGACY_MATH_CONSTANT_MACRO_NAMES
 	#if UE_DEPRECATE_LEGACY_MATH_CONSTANT_MACRO_NAMES
 		#define UE_PRIVATE_MATH_DEPRECATION(Before, After) UE_DEPRECATED_MACRO(5.1, "The " #Before " macro has been deprecated in favor of " #After ".")
 	#else
