Math : Added serialization support for IntVector4

diff --git a/Engine/Source/Runtime/Core/Public/Math/IntVector.h b/Engine/Source/Runtime/Core/Public/Math/IntVector.h
--- a/Engine/Source/Runtime/Core/Public/Math/IntVector.h
+++ b/Engine/Source/Runtime/Core/Public/Math/IntVector.h
@@ -419,7 +419,7 @@ PRAGMA_ENABLE_DEPRECATION_WARNINGS
 	}
 
 	/**
-	 * Serializes the Rectangle.
+	 * Serializes the Vector3.
 	 *
 	 * @param Ar The archive to serialize into.
 	 * @param Vector The vector to serialize.
@@ -890,6 +890,33 @@ PRAGMA_ENABLE_DEPRECATION_WARNINGS
 	{
 		return TIntVector4(X ^ Value, Y ^ Value, Z ^ Value, W ^ Value);
 	}
+
+	/**
+	* Serializes the Vector4.
+	*
+	* @param Ar The archive to serialize into.
+	* @param Vector The vector to serialize.
+	* @return Reference to the Archive after serialization.
+	*/
+	friend FArchive& operator<<(FArchive& Ar, TIntVector4& Vector)
+	{
+		return Ar << Vector.X << Vector.Y << Vector.Z << Vector.W;
+	}
+
+	friend void operator<<(FStructuredArchive::FSlot Slot, TIntVector4& Vector)
+	{
+		FStructuredArchive::FRecord Record = Slot.EnterRecord();
+		Record << SA_VALUE(TEXT("X"), Vector.X);
+		Record << SA_VALUE(TEXT("Y"), Vector.Y);
+		Record << SA_VALUE(TEXT("Z"), Vector.Z);
+		Record << SA_VALUE(TEXT("W"), Vector.W);
+	}
+
+	bool Serialize(FArchive& Ar)
+	{
+		Ar << *this;
+		return true;
+	}
 };
 
 /**
