- Fixed a crash with some null memory image types (likely a type #if'd out of shipping)

diff --git a/Engine/Source/Runtime/Core/Private/Serialization/MemoryImage.cpp b/Engine/Source/Runtime/Core/Private/Serialization/MemoryImage.cpp
--- a/Engine/Source/Runtime/Core/Private/Serialization/MemoryImage.cpp
+++ b/Engine/Source/Runtime/Core/Private/Serialization/MemoryImage.cpp
@@ -1314,7 +1314,7 @@ int32 FPointerTableBase::AddTypeDependency(const FTypeLayoutDesc& TypeDesc)
 	{
 		for (int32 Index = 0; Index < TypeDependencies.Num(); ++Index)
 		{
-			if (TypeDependencies[Index]->NameHash == TypeDesc.NameHash)
+			if (TypeDependencies[Index] != nullptr && TypeDependencies[Index]->NameHash == TypeDesc.NameHash)
 			{
 				return Index;
 			}
