Attempt to normalize config paths before adding them to FConfigCacheIni.
This can prevent issues where multiple paths passed to the same config file result in different cache entries, meaning settings can get overwritten or missed.

diff --git a/Engine/Source/Runtime/Core/Private/Misc/ConfigCacheIni.cpp b/Engine/Source/Runtime/Core/Private/Misc/ConfigCacheIni.cpp
--- a/Engine/Source/Runtime/Core/Private/Misc/ConfigCacheIni.cpp
+++ b/Engine/Source/Runtime/Core/Private/Misc/ConfigCacheIni.cpp
@@ -2310,12 +2310,23 @@ FConfigFile* FConfigCacheIni::Find(const FString& Filename)
 	// this is || filesize so we load up .int files if file IO is allowed
 	if (!Result && !bAreFileOperationsDisabled)
 	{
-		Result = &Add(Filename, FConfigFile());
-		UE_LOG(LogConfig, Verbose, TEXT("GConfig::Find is looking for file:  %s"), *Filename);
-		if (DoesConfigFileExistWrapper(*Filename))
+		// Before attempting to add another file, double check that this doesn't exist at a normalized path.
+		const FString UnrealFileName = FPaths::CreateStandardFilename(FPaths::RemoveDuplicateSlashes(Filename));
+		Result = FindConfigFile(UnrealFileName);
+		
+		if (!Result)
+		{
+			Result = &Add(UnrealFileName, FConfigFile());
+			UE_LOG(LogConfig, Verbose, TEXT("GConfig::Find is looking for file:  %s"), *UnrealFileName);
+			if (DoesConfigFileExistWrapper(*UnrealFileName))
+			{
+				Result->Read(UnrealFileName);
+				UE_LOG(LogConfig, Verbose, TEXT("GConfig::Find has loaded file:  %s"), *UnrealFileName);
+			}
+		}
+		else
 		{
-			Result->Read(Filename);
-			UE_LOG(LogConfig, Verbose, TEXT("GConfig::Find has loaded file:  %s"), *Filename);
+			UE_LOG(LogConfig, Warning, TEXT("GConfig::Find attempting to access config with non-normalized path %s"), *Filename);
 		}
 	}
 
