Use mutex before accessing FAutomationTestBase pointer in FAutomationTestOutputDevice

#rnx

diff --git a/Engine/Source/Runtime/Core/Private/Misc/AutomationTest.cpp b/Engine/Source/Runtime/Core/Private/Misc/AutomationTest.cpp
--- a/Engine/Source/Runtime/Core/Private/Misc/AutomationTest.cpp
+++ b/Engine/Source/Runtime/Core/Private/Misc/AutomationTest.cpp
@@ -91,37 +91,38 @@ void FAutomationTestFramework::FAutomationTestOutputDevice::Serialize( const TCH
 	}
 
 	// Ensure there's a valid unit test associated with the context
-	if (CurTest)
+	class FAutomationTestBase* const LocalCurTest = CurTest.load(std::memory_order_relaxed);
+	if (LocalCurTest)
 	{
-		bool CaptureLog = !CurTest->SuppressLogs()
+		bool CaptureLog = !LocalCurTest->SuppressLogs()
 			&& (Verbosity == ELogVerbosity::Error || Verbosity == ELogVerbosity::Warning || Verbosity == ELogVerbosity::Display);
 
 		if (CaptureLog)
 		{
 		
-			ELogVerbosity::Type EffectiveVerbosity = GetAutomationLogLevel(Verbosity, Category, CurTest);
+			ELogVerbosity::Type EffectiveVerbosity = GetAutomationLogLevel(Verbosity, Category, LocalCurTest);
 
 			FString FormattedMsg = FString::Printf(TEXT("%s: %s"), *Category.ToString(), V);
 			
 			// Errors
 			if (EffectiveVerbosity == ELogVerbosity::Error)
 			{
-				if (!LoggedFailureCause.Contains(CurTest))
+				if (!LoggedFailureCause.Contains(LocalCurTest))
 				{
-					CurTest->AddError(FString::Printf(TEXT("%s will be marked as failing due to errors being logged"), *CurTest->GetTestFullName()), STACK_OFFSET);
-					LoggedFailureCause.Add(CurTest);
+					LocalCurTest->AddError(FString::Printf(TEXT("%s will be marked as failing due to errors being logged"), *LocalCurTest->GetTestFullName()), STACK_OFFSET);
+					LoggedFailureCause.Add(LocalCurTest);
 				}
-				CurTest->AddError(FormattedMsg, STACK_OFFSET);
+				LocalCurTest->AddError(FormattedMsg, STACK_OFFSET);
 			}
 			// Warnings
 			else if (EffectiveVerbosity == ELogVerbosity::Warning)
 			{
-				CurTest->AddWarning(FormattedMsg, STACK_OFFSET);
+				LocalCurTest->AddWarning(FormattedMsg, STACK_OFFSET);
 			}
 			// Display
 			else
 			{
-				CurTest->AddInfo(FormattedMsg, STACK_OFFSET);
+				LocalCurTest->AddInfo(FormattedMsg, STACK_OFFSET);
 			}
 		}
 		// Log...etc
@@ -139,11 +140,11 @@ void FAutomationTestFramework::FAutomationTestOutputDevice::Serialize( const TCH
 				//Remove "analytics" from the string
 				LogString.RightInline(LogString.Len() - (AnalyticsString.Len() + 1), false);
 
-				CurTest->AddAnalyticsItem(LogString);
+				LocalCurTest->AddAnalyticsItem(LogString);
 			}
 			//else
 			//{
-			//	CurTest->AddInfo(LogString, STACK_OFFSET);
+			//	LocalCurTest->AddInfo(LogString, STACK_OFFSET);
 			//}
 		}
 	}
@@ -153,11 +154,12 @@ void FAutomationTestFramework::FAutomationTestMessageFilter::Serialize(const TCH
 {
 	// Prevent null dereference if logging happens in async tasks while changing DestinationContext
 	FFeedbackContext* const LocalDestinationContext = DestinationContext.load(std::memory_order_relaxed);
+	class FAutomationTestBase* const LocalCurTest = CurTest.load(std::memory_order_relaxed);
 	if (LocalDestinationContext)
 	{
 		if ((Verbosity == ELogVerbosity::Warning) || (Verbosity == ELogVerbosity::Error))
 		{
-			if (CurTest->IsExpectedError(FString(V)))
+			if (LocalCurTest->IsExpectedError(FString(V)))
 			{
 				Verbosity = ELogVerbosity::Verbose;
 			}
diff --git a/Engine/Source/Runtime/Core/Public/Misc/AutomationTest.h b/Engine/Source/Runtime/Core/Public/Misc/AutomationTest.h
--- a/Engine/Source/Runtime/Core/Public/Misc/AutomationTest.h
+++ b/Engine/Source/Runtime/Core/Public/Misc/AutomationTest.h
@@ -974,11 +974,11 @@ private:
 	{
 	public:
 		FAutomationTestOutputDevice() 
-			: CurTest( NULL ) {}
+			: CurTest( nullptr ) {}
 
 		~FAutomationTestOutputDevice()
 		{
-			CurTest = NULL;
+			CurTest = nullptr;
 		}
 
 		/**
@@ -1002,7 +1002,7 @@ private:
 
 	private:
 		/** Associated automation test; all warnings, errors, etc. are routed to the automation test to track */
-		class FAutomationTestBase* CurTest;
+		std::atomic<class FAutomationTestBase*>CurTest;
 
 		/** Tests that we've logged the failure cause when an error is involved */
 		TSet<FAutomationTestBase*> LoggedFailureCause;
@@ -1054,7 +1054,7 @@ private:
 		 }
 
 	 private:
-		 class FAutomationTestBase* CurTest;
+		 std::atomic<class FAutomationTestBase*> CurTest;
 		 std::atomic<FFeedbackContext*> DestinationContext;
 	 };
 
