Revert basedir path changes for Mac in 19498448
#rnx

diff --git a/Engine/Source/Runtime/Core/Private/Mac/MacPlatformProcess.cpp b/Engine/Source/Runtime/Core/Private/Mac/MacPlatformProcess.cpp
--- a/Engine/Source/Runtime/Core/Private/Mac/MacPlatformProcess.cpp
+++ b/Engine/Source/Runtime/Core/Private/Mac/MacPlatformProcess.cpp
@@ -1044,13 +1044,6 @@ const TCHAR* FMacPlatformProcess::BaseDir()
 
 			FCString::Strcpy(Result, MAC_MAX_PATH, *FString(BasePath));
 			FCString::Strcat(Result, TEXT("/"));
-
-			FString CollapseResult(Result);
-#ifdef UE_RELATIVE_BASE_DIR
-			CollapseResult /= UE_RELATIVE_BASE_DIR;
-#endif
-			FPaths::CollapseRelativeDirectories(CollapseResult);
-			FCString::Strcpy(Result, *CollapseResult);		
 		}
 	}
 	return Result;
