Config (ini): Fixed extraction of values from quoted strings, which was off by one.

For instance,

[REDACTED_HANDLE]=PrimaryAssetType
+PrimaryAssetTypesToScan=(PrimaryAssetType="Map",AssetBaseClass=/Script/Engine.World,bHasBlueprintClasses=False,bIsEditorOnly=True,Directories=((Path="/Game/Maps")),SpecificAssets=,Rules=(Priority=-1,ChunkId=-1,bApplyRecursively=True,CookRule=Unknown))

when processing the second line PrimaryAssetType key would be extracted as "Ma. It worked correctly for non-quoted values.

NOTE: Based on my review, there were no known consequences of this bug, because these extracted values are exclusively used for comparisons (with all compared values being extracted in the same way). It could cause problems if we had values used as keys differing just by the last character.

diff --git a/Engine/Source/Runtime/Core/Private/Misc/ConfigCacheIni.cpp b/Engine/Source/Runtime/Core/Private/Misc/ConfigCacheIni.cpp
--- a/Engine/Source/Runtime/Core/Private/Misc/ConfigCacheIni.cpp
+++ b/Engine/Source/Runtime/Core/Private/Misc/ConfigCacheIni.cpp
@@ -397,7 +397,7 @@ static void ExtractPropertyValue(const FString& FullStructValue, const FString&
 		}
 
 		// pull out the token
-		Out.AppendChars(*FullStructValue + MatchLoc, UE_PTRDIFF_TO_INT32(Travel - Start));
+		Out.AppendChars(Start, UE_PTRDIFF_TO_INT32(Travel - Start));
 	}
 }
 
