Logging: Replaced GLog->PanicFlush() with GLog->Panic()

Panic allows only one thread to become the panic thread. In panic mode, logs are only sent to panic-safe output devices. Each call to Panic flushes buffered logs to panic-safe output devices, and flushes panic-safe output devices.

diff --git a/Engine/Source/Editor/UnrealEd/Private/Commandlets/CookGlobalShaderCommandlet.cpp b/Engine/Source/Editor/UnrealEd/Private/Commandlets/CookGlobalShaderCommandlet.cpp
--- a/Engine/Source/Editor/UnrealEd/Private/Commandlets/CookGlobalShaderCommandlet.cpp
+++ b/Engine/Source/Editor/UnrealEd/Private/Commandlets/CookGlobalShaderCommandlet.cpp
@@ -184,7 +184,7 @@ void UCookGlobalShadersCommandlet::CookGlobalShadersOnDirectoriesChanges()
 		DirectoryWatcherModule.Get()->Tick(FApp::GetDeltaTime());
 		
 		// flush log
-		GLog->FlushThreadedLogs();
+		GLog->FlushThreadedLogs(EOutputDeviceRedirectorFlushOptions::Async);
 
 #if PLATFORM_WINDOWS
 		if (ComWrapperShutdownEvent->Wait(0))
diff --git a/Engine/Source/Editor/UnrealEd/Private/Commandlets/FileServerCommandlet.cpp b/Engine/Source/Editor/UnrealEd/Private/Commandlets/FileServerCommandlet.cpp
--- a/Engine/Source/Editor/UnrealEd/Private/Commandlets/FileServerCommandlet.cpp
+++ b/Engine/Source/Editor/UnrealEd/Private/Commandlets/FileServerCommandlet.cpp
@@ -140,7 +140,7 @@ int32 UFileServerCommandlet::Main( const FString& Params )
 		}
 
 		// flush log
-		GLog->FlushThreadedLogs();
+		GLog->FlushThreadedLogs(EOutputDeviceRedirectorFlushOptions::Async);
 
 #if PLATFORM_WINDOWS
 		if (ComWrapperShutdownEvent->Wait(0))
diff --git a/Engine/Source/Programs/BuildPatchTool/Private/ToolModes/AutomationMode.cpp b/Engine/Source/Programs/BuildPatchTool/Private/ToolModes/AutomationMode.cpp
--- a/Engine/Source/Programs/BuildPatchTool/Private/ToolModes/AutomationMode.cpp
+++ b/Engine/Source/Programs/BuildPatchTool/Private/ToolModes/AutomationMode.cpp
@@ -88,7 +88,7 @@ public:
 			AutomationControllerModule.Tick();
 
 			// Flush threaded logs.
-			GLog->FlushThreadedLogs();
+			GLog->FlushThreadedLogs(EOutputDeviceRedirectorFlushOptions::Async);
 
 			// Throttle frame rate.
 			FPlatformProcess::Sleep(FMath::Max<float>(0.0f, MainsFrameTime - (FPlatformTime::Seconds() - LastTime)));
diff --git a/Engine/Source/Programs/CrashReportClient/Private/Linux/CrashReportClientMainLinux.cpp b/Engine/Source/Programs/CrashReportClient/Private/Linux/CrashReportClientMainLinux.cpp
--- a/Engine/Source/Programs/CrashReportClient/Private/Linux/CrashReportClientMainLinux.cpp
+++ b/Engine/Source/Programs/CrashReportClient/Private/Linux/CrashReportClientMainLinux.cpp
@@ -26,7 +26,7 @@ void CrashReporterCrashHandler(const FGenericCrashContext& GenericContext)
 	const_cast< FUnixCrashContext& >(Context).CaptureStackTrace(Context.FirstCrashHandlerFrame);
 	if (GLog)
 	{
-		GLog->Flush();
+		GLog->Panic();
 	}
 	if (GWarn)
 	{
diff --git a/Engine/Source/Programs/UnrealHeaderTool/Private/CodeGenerator.cpp b/Engine/Source/Programs/UnrealHeaderTool/Private/CodeGenerator.cpp
--- a/Engine/Source/Programs/UnrealHeaderTool/Private/CodeGenerator.cpp
+++ b/Engine/Source/Programs/UnrealHeaderTool/Private/CodeGenerator.cpp
@@ -5911,7 +5911,7 @@ void FNativeClassHeaderGenerator::GenerateSourceFiles(
 
 		// This is strange, but flushing the log can take a long time.  Without an explicit flush, we were getting a long stall in the UE_LOG message
 		// during the detection of script plugins.  By doing it here, the flush should easily complete during code generation.
-		GLog->FlushThreadedLogs();
+		GLog->FlushThreadedLogs(EOutputDeviceRedirectorFlushOptions::Async);
 
 		FTaskGraphInterface::Get().WaitUntilTasksComplete(ExportSourceTasks);
 	}
diff --git a/Engine/Source/Runtime/ApplicationCore/Private/HoloLens/HoloLensErrorOutputDevice.cpp b/Engine/Source/Runtime/ApplicationCore/Private/HoloLens/HoloLensErrorOutputDevice.cpp
--- a/Engine/Source/Runtime/ApplicationCore/Private/HoloLens/HoloLensErrorOutputDevice.cpp
+++ b/Engine/Source/Runtime/ApplicationCore/Private/HoloLens/HoloLensErrorOutputDevice.cpp
@@ -81,7 +81,7 @@ void FHoloLensErrorOutputDevice::HandleError()
 #if !NO_LOGGING
 	FDebug::LogFormattedMessageWithCallstack(LogHoloLens.GetCategoryName(), __FILE__, __LINE__, TEXT("=== Critical error: ==="), GErrorHist, ELogVerbosity::Error);
 #endif
-	GLog->PanicFlush();
+	GLog->Panic();
 
 	// Copy to clipboard in non-cooked editor builds.
 	FPlatformApplicationMisc::ClipboardCopy(GErrorHist);
diff --git a/Engine/Source/Runtime/ApplicationCore/Private/Mac/MacConsoleOutputDevice.cpp b/Engine/Source/Runtime/ApplicationCore/Private/Mac/MacConsoleOutputDevice.cpp
--- a/Engine/Source/Runtime/ApplicationCore/Private/Mac/MacConsoleOutputDevice.cpp
+++ b/Engine/Source/Runtime/ApplicationCore/Private/Mac/MacConsoleOutputDevice.cpp
@@ -290,7 +290,7 @@ void FMacConsoleOutputDevice::SetDefaultTextColor()
 		UE_LOG(LogCore, Warning, TEXT("*** INTERRUPTED *** : Console Window Closed, Shutting Down"));
 		if (GLog)
 		{
-			GLog->PanicFlush();
+			GLog->Panic();
 		}
 		_Exit(0);
 	}
diff --git a/Engine/Source/Runtime/Core/Private/Android/AndroidErrorOutputDevice.cpp b/Engine/Source/Runtime/Core/Private/Android/AndroidErrorOutputDevice.cpp
--- a/Engine/Source/Runtime/Core/Private/Android/AndroidErrorOutputDevice.cpp
+++ b/Engine/Source/Runtime/Core/Private/Android/AndroidErrorOutputDevice.cpp
@@ -64,7 +64,7 @@ void FAndroidErrorOutputDevice::HandleError()
 	FDebug::LogFormattedMessageWithCallstack(LogAndroid.GetCategoryName(), __FILE__, __LINE__, TEXT("=== Critical error: ==="), GErrorHist, ELogVerbosity::Error);
 #endif
 	
-	GLog->PanicFlush();
+	GLog->Panic();
 
 	FCoreDelegates::OnHandleSystemError.Broadcast();
 	FCoreDelegates::OnShutdownAfterError.Broadcast();
diff --git a/Engine/Source/Runtime/Core/Private/Android/AndroidPlatformMisc.cpp b/Engine/Source/Runtime/Core/Private/Android/AndroidPlatformMisc.cpp
--- a/Engine/Source/Runtime/Core/Private/Android/AndroidPlatformMisc.cpp
+++ b/Engine/Source/Runtime/Core/Private/Android/AndroidPlatformMisc.cpp
@@ -231,9 +231,8 @@ void FAndroidMisc::RequestExit( bool Force )
 #endif
 
 	UE_LOG(LogAndroid, Log, TEXT("FAndroidMisc::RequestExit(%i)"), Force);
-	if(GLog)
+	if (GLog)
 	{
-		GLog->FlushThreadedLogs();
 		GLog->Flush();
 	}
 
@@ -1040,7 +1039,7 @@ void DefaultCrashHandler(const FAndroidCrashContext& Context)
 
 		if (GLog)
 		{
-			GLog->PanicFlush();
+			GLog->Panic();
 		}
 
 		if (GWarn)
@@ -1409,7 +1408,7 @@ void FAndroidMisc::TriggerCrashHandler(ECrashContextType InType, const TCHAR* In
 		// we dont flush logs during a fatal signal, malloccrash can cause us to deadlock.
 		if (GLog)
 		{
-			GLog->PanicFlush();
+			GLog->Panic();
 		}
 		if (GWarn)
 		{
diff --git a/Engine/Source/Runtime/Core/Private/HAL/ExceptionHandling.cpp b/Engine/Source/Runtime/Core/Private/HAL/ExceptionHandling.cpp
--- a/Engine/Source/Runtime/Core/Private/HAL/ExceptionHandling.cpp
+++ b/Engine/Source/Runtime/Core/Private/HAL/ExceptionHandling.cpp
@@ -32,7 +32,7 @@ void ReportInteractiveEnsure(const TCHAR* InMessage)
 	GEnsureShowsCRC = true;
 
 #if PLATFORM_USE_REPORT_ENSURE
-	GLog->PanicFlush();
+	GLog->FlushThreadedLogs();
 	ReportEnsure(InMessage, nullptr);
 #endif
 
diff --git a/Engine/Source/Runtime/Core/Private/HAL/ThreadHeartBeat.cpp b/Engine/Source/Runtime/Core/Private/HAL/ThreadHeartBeat.cpp
--- a/Engine/Source/Runtime/Core/Private/HAL/ThreadHeartBeat.cpp
+++ b/Engine/Source/Runtime/Core/Private/HAL/ThreadHeartBeat.cpp
@@ -275,7 +275,15 @@ void FORCENOINLINE FThreadHeartBeat::OnHang(double HangDuration, uint32 ThreadTh
 
 #if PLATFORM_USE_REPORT_ENSURE
 		UE_LOG(LogCore, Error, TEXT("%s"), *ErrorMessage);
-		GLog->PanicFlush();
+
+		if (bHangsAreFatal)
+		{
+			GLog->Panic();
+		}
+		else
+		{
+			GLog->FlushThreadedLogs(EOutputDeviceRedirectorFlushOptions::Async);
+		}
 
 		// Skip macros and FDebug, we always want this to fire
 		ReportHang(*ErrorMessage, StackFrames, NumStackFrames, ThreadThatHung);
diff --git a/Engine/Source/Runtime/Core/Private/HoloLens/HoloLensPlatformCrashContext.cpp b/Engine/Source/Runtime/Core/Private/HoloLens/HoloLensPlatformCrashContext.cpp
--- a/Engine/Source/Runtime/Core/Private/HoloLens/HoloLensPlatformCrashContext.cpp
+++ b/Engine/Source/Runtime/Core/Private/HoloLens/HoloLensPlatformCrashContext.cpp
@@ -254,7 +254,7 @@ enum class EErrorReportUI
 int32 ReportCrashUsingWindowsErrorReporting(EXCEPTION_POINTERS* ExceptionInfo, const TCHAR* ErrorMessage, EErrorReportUI ReportUI)
 {
 	// Flush out the log
-	GLog->Flush();
+	GLog->Panic();
 
 	// Construct the report details
 	WER_REPORT_INFORMATION ReportInformation = {sizeof( WER_REPORT_INFORMATION )};
diff --git a/Engine/Source/Runtime/Core/Private/IOS/IOSErrorOutputDevice.cpp b/Engine/Source/Runtime/Core/Private/IOS/IOSErrorOutputDevice.cpp
--- a/Engine/Source/Runtime/Core/Private/IOS/IOSErrorOutputDevice.cpp
+++ b/Engine/Source/Runtime/Core/Private/IOS/IOSErrorOutputDevice.cpp
@@ -52,5 +52,5 @@ void FIOSErrorOutputDevice::HandleError()
 	FDebug::LogFormattedMessageWithCallstack(LogIOS.GetCategoryName(), __FILE__, __LINE__, TEXT("=== Critical error: ==="), GErrorHist, ELogVerbosity::Error);
 #endif
 
-	GLog->PanicFlush();
+	GLog->Panic();
 }
diff --git a/Engine/Source/Runtime/Core/Private/IOS/IOSPlatformMisc.cpp b/Engine/Source/Runtime/Core/Private/IOS/IOSPlatformMisc.cpp
--- a/Engine/Source/Runtime/Core/Private/IOS/IOSPlatformMisc.cpp
+++ b/Engine/Source/Runtime/Core/Private/IOS/IOSPlatformMisc.cpp
@@ -1362,7 +1362,7 @@ static void DefaultCrashHandler(FIOSCrashContext const& Context)
     Context.ReportCrash();
     if (GLog)
     {
-        GLog->PanicFlush();
+        GLog->Panic();
     }
     if (GWarn)
     {
@@ -1412,7 +1412,7 @@ static void GracefulTerminationHandler(int32 Signal, siginfo_t* Info, void* Cont
     // make sure we write out as much as possible
     if (GLog)
     {
-        GLog->Flush();
+        GLog->Panic();
     }
     if (GWarn)
     {
diff --git a/Engine/Source/Runtime/Core/Private/Mac/MacErrorOutputDevice.cpp b/Engine/Source/Runtime/Core/Private/Mac/MacErrorOutputDevice.cpp
--- a/Engine/Source/Runtime/Core/Private/Mac/MacErrorOutputDevice.cpp
+++ b/Engine/Source/Runtime/Core/Private/Mac/MacErrorOutputDevice.cpp
@@ -91,7 +91,7 @@ void FMacErrorOutputDevice::HandleError()
 	UE_LOG(LogMac, Log, TEXT("=== Critical error: ===") LINE_TERMINATOR TEXT("%s") LINE_TERMINATOR, GErrorExceptionDescription);
 	UE_LOG(LogMac, Log, TEXT("%s"), GErrorHist);
     
-    GLog->Flush();
+    GLog->Panic();
     
 #if PLATFORM_MAC_ARM64
     // stack dumping seems broken on arm, so for the time being spin so we can dump callstacks
diff --git a/Engine/Source/Runtime/Core/Private/Mac/MacPlatformMisc.cpp b/Engine/Source/Runtime/Core/Private/Mac/MacPlatformMisc.cpp
--- a/Engine/Source/Runtime/Core/Private/Mac/MacPlatformMisc.cpp
+++ b/Engine/Source/Runtime/Core/Private/Mac/MacPlatformMisc.cpp
@@ -1916,7 +1916,7 @@ static void DefaultCrashHandler(FMacCrashContext const& Context)
 	Context.ReportCrash();
 	if (GLog)
 	{
-		GLog->PanicFlush();
+		GLog->Panic();
 	}
 	if (GWarn)
 	{
@@ -2020,7 +2020,7 @@ static void GracefulTerminationHandler(int32 Signal, siginfo_t* Info, void* Cont
 	// make sure as much data is written to disk as possible
 	if (GLog)
 	{
-		GLog->Flush();
+		GLog->Panic();
 	}
 	if (GWarn)
 	{
diff --git a/Engine/Source/Runtime/Core/Private/Misc/AssertionMacros.cpp b/Engine/Source/Runtime/Core/Private/Misc/AssertionMacros.cpp
--- a/Engine/Source/Runtime/Core/Private/Misc/AssertionMacros.cpp
+++ b/Engine/Source/Runtime/Core/Private/Misc/AssertionMacros.cpp
@@ -517,7 +517,7 @@ void FORCENOINLINE FDebug::CheckVerifyFailedImpl(
 	{
 		// Flushing the logs here increases the likelihood that recent messages will be written to the log file, stdout and the debugger console.
 		// Without this, some of the recent messages may not be reported when debugger stops due to an assertion failure.
-		GLog->PanicFlush();
+		GLog->FlushThreadedLogs();
 	}
 
 	if (!FPlatformMisc::IsDebuggerPresent())
diff --git a/Engine/Source/Runtime/Core/Private/Misc/OutputDeviceAnsiError.cpp b/Engine/Source/Runtime/Core/Private/Misc/OutputDeviceAnsiError.cpp
--- a/Engine/Source/Runtime/Core/Private/Misc/OutputDeviceAnsiError.cpp
+++ b/Engine/Source/Runtime/Core/Private/Misc/OutputDeviceAnsiError.cpp
@@ -84,7 +84,7 @@ void FOutputDeviceAnsiError::HandleError()
 		UE_LOG( LogHAL, Log, TEXT( "=== Critical error: ===" ) LINE_TERMINATOR LINE_TERMINATOR TEXT( "%s" ) LINE_TERMINATOR, GErrorExceptionDescription );
 		UE_LOG(LogHAL, Log, TEXT("%s"), GErrorHist);
 
-		GLog->Flush();
+		GLog->Panic();
 	}
 	else
 	{
diff --git a/Engine/Source/Runtime/Core/Private/Misc/OutputDeviceRedirector.cpp b/Engine/Source/Runtime/Core/Private/Misc/OutputDeviceRedirector.cpp
--- a/Engine/Source/Runtime/Core/Private/Misc/OutputDeviceRedirector.cpp
+++ b/Engine/Source/Runtime/Core/Private/Misc/OutputDeviceRedirector.cpp
@@ -110,7 +110,7 @@ struct FOutputDeviceRedirectorState
 {
 	/** A custom lock to guard access to both buffered and unbuffered output devices. */
 	FRWLock OutputDevicesLock;
-	std::atomic<uint32> OutputDevicesLockState{0};
+	std::atomic<uint32> OutputDevicesLockState = 0;
 	uint8 OutputDevicesLockPadding[CalculateRedirectorCacheLinePadding(sizeof(OutputDevicesLock) + sizeof(OutputDevicesLockState))]{};
 
 	/** A queue of lines logged by non-master threads. */
@@ -134,13 +134,16 @@ struct FOutputDeviceRedirectorState
 	FRWLock ThreadLock;
 
 	/** An event to wake the dedicated master thread to process buffered lines. */
-	std::atomic<FEvent*> ThreadWakeEvent{nullptr};
+	std::atomic<FEvent*> ThreadWakeEvent = nullptr;
 
 	/** A queue of events to trigger when the dedicated master thread is idle. */
 	TDepletableMpscQueue<FEvent*, FOutputDeviceLinearAllocator> ThreadIdleEvents;
 
 	/** The ID of the master thread. Logging from other threads will be buffered for processing by the master thread. */
-	std::atomic<uint32> MasterThreadId{FPlatformTLS::GetCurrentThreadId()};
+	std::atomic<uint32> MasterThreadId = FPlatformTLS::GetCurrentThreadId();
+
+	/** The ID of the panic thread, which is only set by Panic(). */
+	std::atomic<uint32> PanicThreadId = MAX_uint32;
 
 	/** Whether the backlog is enabled. */
 	bool bEnableBacklog = false;
@@ -150,9 +153,36 @@ struct FOutputDeviceRedirectorState
 		return ThreadId == MasterThreadId.load(std::memory_order_relaxed);
 	}
 
+	bool IsPanicThread(const uint32 ThreadId) const
+	{
+		return ThreadId == PanicThreadId.load(std::memory_order_relaxed);
+	}
+
+	bool CanLockFromThread(const uint32 ThreadId) const
+	{
+		const uint32 LocalPanicThreadId = PanicThreadId.load(std::memory_order_relaxed);
+		return LocalPanicThreadId == MAX_uint32 || LocalPanicThreadId == ThreadId;
+	}
+
+	bool TryStartThread();
+	bool TryStopThread();
+
 	void ThreadLoop();
 
 	void FlushBufferedLines();
+
+	template <typename OutputDevicesType, typename FunctionType, typename... ArgTypes>
+	FORCEINLINE void BroadcastTo(const uint32 ThreadId, const OutputDevicesType& OutputDevices, FunctionType&& Function, ArgTypes&&... Args)
+	{
+		const bool bIsPanicThread = IsPanicThread(ThreadId);
+		for (FOutputDevice* OutputDevice : OutputDevices)
+		{
+			if (!bIsPanicThread || OutputDevice->CanBeUsedOnPanicThread())
+			{
+				Invoke(Function, OutputDevice, Forward<ArgTypes>(Args)...);
+			}
+		}
+	}
 };
 
 /**
@@ -170,7 +200,7 @@ public:
 		: State(InState)
 	{
 		// Read locks add/sub by 2 to keep the LSB free for write locks to use.
-		if (State.OutputDevicesLockState.fetch_add(2, std::memory_order_relaxed) & 1)
+		if (State.OutputDevicesLockState.fetch_add(2, std::memory_order_acquire) & 1)
 		{
 			WaitForWriteLock();
 		}
@@ -244,58 +274,45 @@ private:
  * A scoped lock for readers of the OutputDevices arrays that need to access master thread state.
  *
  * The master lock has the same access as the read lock, and:
+ * - Must not be entered when the thread holds a write lock or master lock.
+ * - Must check IsLocked() before performing restricted operations.
  * - Must be locked to write to buffered output devices.
  * - Must be locked while calling FlushBufferedLines().
- * - Must not be entered when the thread holds a write or master lock.
  * - May be locked when the thread holds a read lock.
+ * - When a panic thread is active, may only be locked from the panic thread.
  */
-class FOutputDevicesMasterScopeLock
-{
-public:
-	FORCEINLINE explicit FOutputDevicesMasterScopeLock(FOutputDeviceRedirectorState& InState)
-		: State(InState)
-	{
-		State.OutputDevicesLock.WriteLock();
-	}
-
-	FORCEINLINE ~FOutputDevicesMasterScopeLock()
-	{
-		State.OutputDevicesLock.WriteUnlock();
-	}
-
-private:
-	FOutputDeviceRedirectorState& State;
-};
-
-/**
- * A scoped lock for readers of the OutputDevices arrays that need to access master thread state.
- *
- * Equivalent to FOutputDevicesMasterScopeLock, except:
- * - IsLocked() must be checked before accessing the state guarded by the lock.
- * - The wait to take the lock is bounded. The lock is not taken if the wait time elapses.
- */
-class FOutputDevicesMasterScopeTryLock
+class FOutputDevicesMasterScope
 {
 public:
-	FORCEINLINE explicit FOutputDevicesMasterScopeTryLock(FOutputDeviceRedirectorState& InState, double WaitTime = 0.0)
+	explicit FOutputDevicesMasterScope(FOutputDeviceRedirectorState& InState)
 		: State(InState)
-		, bLocked(State.OutputDevicesLock.TryWriteLock())
 	{
-		if (!bLocked && WaitTime > 0.0)
+		const uint32 ThreadId = FPlatformTLS::GetCurrentThreadId();
+		if (State.CanLockFromThread(ThreadId))
 		{
-			TRACE_CPUPROFILER_EVENT_SCOPE(FOutputDevicesWriteScopeTryLock);
-			const double EndTime = FPlatformTime::Seconds() + WaitTime;
-			do
+			if (State.IsPanicThread(ThreadId))
 			{
-				bLocked = State.OutputDevicesLock.TryWriteLock();
+				bLocked = true;
+			}
+			else
+			{
+				State.OutputDevicesLock.WriteLock();
+				if (!State.CanLockFromThread(ThreadId))
+				{
+					State.OutputDevicesLock.WriteUnlock();
+				}
+				else
+				{
+					bNeedsUnlock = true;
+					bLocked = true;
+				}
 			}
-			while (!bLocked && FPlatformTime::Seconds() < EndTime);
 		}
 	}
 
-	FORCEINLINE ~FOutputDevicesMasterScopeTryLock()
+	FORCEINLINE ~FOutputDevicesMasterScope()
 	{
-		if (bLocked)
+		if (bNeedsUnlock)
 		{
 			State.OutputDevicesLock.WriteUnlock();
 		}
@@ -305,14 +322,40 @@ public:
 
 private:
 	FOutputDeviceRedirectorState& State;
+	bool bNeedsUnlock = false;
 	bool bLocked = false;
 };
 
+bool FOutputDeviceRedirectorState::TryStartThread()
+{
+	if (FWriteScopeLock ThreadScopeLock(ThreadLock); !ThreadWakeEvent.load(std::memory_order_relaxed))
+	{
+		FEvent* WakeEvent = FPlatformProcess::GetSynchEventFromPool();
+		WakeEvent->Trigger();
+		ThreadWakeEvent.store(WakeEvent, std::memory_order_release);
+		Thread = FThread(TEXT("OutputDeviceRedirector"), [this] { ThreadLoop(); });
+	}
+	return true;
+}
+
+bool FOutputDeviceRedirectorState::TryStopThread()
+{
+	if (FWriteScopeLock ThreadScopeLock(ThreadLock); FEvent* WakeEvent = ThreadWakeEvent.exchange(nullptr, std::memory_order_acquire))
+	{
+		WakeEvent->Trigger();
+		Thread.Join();
+		FOutputDevicesWriteScopeLock Lock(*this);
+		FPlatformProcess::ReturnSynchEventToPool(WakeEvent);
+	}
+	return true;
+}
+
 void FOutputDeviceRedirectorState::ThreadLoop()
 {
 	const uint32 ThreadId = FPlatformTLS::GetCurrentThreadId();
+
+	if (FOutputDevicesMasterScope Lock(*this); Lock.IsLocked())
 	{
-		FOutputDevicesMasterScopeLock Lock(*this);
 		MasterThreadId.store(ThreadId, std::memory_order_relaxed);
 	}
 
@@ -323,8 +366,10 @@ void FOutputDeviceRedirectorState::ThreadLoop()
 			WakeEvent->Wait();
 			do
 			{
-				FOutputDevicesMasterScopeLock Lock(*this);
-				FlushBufferedLines();
+				if (FOutputDevicesMasterScope Lock(*this); Lock.IsLocked())
+				{
+					FlushBufferedLines();
+				}
 			}
 			while (!BufferedLines.IsEmpty());
 			ThreadIdleEvents.Deplete([](FEvent* Event) { Event->Trigger(); });
@@ -341,17 +386,11 @@ void FOutputDeviceRedirectorState::FlushBufferedLines()
 
 	TRACE_CPUPROFILER_EVENT_SCOPE(FOutputDeviceRedirector::FlushBufferedLines);
 
-	BufferedLines.Deplete([this](UE::Private::FOutputDeviceLine&& Line)
+	const uint32 ThreadId = FPlatformTLS::GetCurrentThreadId();
+	BufferedLines.Deplete([this, ThreadId](UE::Private::FOutputDeviceLine&& Line)
 	{
-		const double Time = Line.Time;
-		const TCHAR* const Data = Line.Data;
-		const FName Category = Line.Category;
-		const ELogVerbosity::Type Verbosity = Line.Verbosity;
-
-		for (FOutputDevice* OutputDevice : BufferedOutputDevices)
-		{
-			OutputDevice->Serialize(Data, Verbosity, Category, Time);
-		}
+		BroadcastTo(ThreadId, BufferedOutputDevices, UE_PROJECTION_MEMBER(FOutputDevice, Serialize),
+			Line.Data, Line.Verbosity, Line.Category, Line.Time);
 	});
 }
 
@@ -416,8 +455,10 @@ void FOutputDeviceRedirector::FlushThreadedLogs(EOutputDeviceRedirectorFlushOpti
 		return;
 	}
 
-	UE::Private::FOutputDevicesMasterScopeLock Lock(*State);
-	State->FlushBufferedLines();
+	if (UE::Private::FOutputDevicesMasterScope Lock(*State); Lock.IsLocked())
+	{
+		State->FlushBufferedLines();
+	}
 }
 
 void FOutputDeviceRedirector::SerializeBacklog(FOutputDevice* OutputDevice)
@@ -443,7 +484,7 @@ void FOutputDeviceRedirector::SetCurrentThreadAsMasterThread()
 {
 	const uint32 ThreadId = FPlatformTLS::GetCurrentThreadId();
 
-	if (UE::Private::FOutputDevicesMasterScopeLock Lock(*State); State->MasterThreadId.load(std::memory_order_relaxed) == ThreadId)
+	if (UE::Private::FOutputDevicesMasterScope Lock(*State); !Lock.IsLocked() || State->MasterThreadId.load(std::memory_order_relaxed) == ThreadId)
 	{
 		return;
 	}
@@ -453,30 +494,12 @@ void FOutputDeviceRedirector::SetCurrentThreadAsMasterThread()
 		State->FlushBufferedLines();
 	}
 
-	if (FWriteScopeLock ThreadLock(State->ThreadLock); FEvent* WakeEvent = State->ThreadWakeEvent.exchange(nullptr, std::memory_order_acquire))
-	{
-		WakeEvent->Trigger();
-		State->Thread.Join();
-		UE::Private::FOutputDevicesWriteScopeLock Lock(*State);
-		FPlatformProcess::ReturnSynchEventToPool(WakeEvent);
-	}
+	State->TryStopThread();
 }
 
 bool FOutputDeviceRedirector::TryStartDedicatedMasterThread()
 {
-	if (!FApp::ShouldUseThreadingForPerformance())
-	{
-		return false;
-	}
-
-	if (FWriteScopeLock ThreadLock(State->ThreadLock); !State->ThreadWakeEvent.load(std::memory_order_relaxed))
-	{
-		FEvent* WakeEvent = FPlatformProcess::GetSynchEventFromPool();
-		WakeEvent->Trigger();
-		State->ThreadWakeEvent.store(WakeEvent, std::memory_order_release);
-		State->Thread = FThread(TEXT("OutputDeviceRedirector"), [State = State.Get()] { State->ThreadLoop(); });
-	}
-	return true;
+	return FApp::ShouldUseThreadingForPerformance() && State->TryStartThread();
 }
 
 void FOutputDeviceRedirector::Serialize(const TCHAR* const Data, const ELogVerbosity::Type Verbosity, const FName& Category, const double Time)
@@ -499,11 +522,11 @@ void FOutputDeviceRedirector::Serialize(const TCHAR* const Data, const ELogVerbo
 	}
 #endif
 
+	const uint32 ThreadId = FPlatformTLS::GetCurrentThreadId();
+
 	// Serialize directly to any output devices which don't require buffering
-	for (FOutputDevice* OutputDevice : State->UnbufferedOutputDevices)
-	{
-		OutputDevice->Serialize(Data, Verbosity, Category, RealTime);
-	}
+	State->BroadcastTo(ThreadId, State->UnbufferedOutputDevices, UE_PROJECTION_MEMBER(FOutputDevice, Serialize),
+		Data, Verbosity, Category, RealTime);
 
 	if (State->bEnableBacklog)
 	{
@@ -522,23 +545,26 @@ void FOutputDeviceRedirector::Serialize(const TCHAR* const Data, const ELogVerbo
 		}
 	};
 
-	const uint32 ThreadId = FPlatformTLS::GetCurrentThreadId();
 	if (!State->IsMasterThread(ThreadId) || State->BufferedOutputDevices.IsEmpty())
 	{
 		EnqueueLine();
 	}
-	else if (UE::Private::FOutputDevicesMasterScopeLock MasterLock(*State); !State->IsMasterThread(ThreadId))
+	else if (UE::Private::FOutputDevicesMasterScope MasterLock(*State); MasterLock.IsLocked())
 	{
-		EnqueueLine();
+		if (!State->IsMasterThread(ThreadId) && !State->IsPanicThread(ThreadId))
+		{
+			EnqueueLine();
+		}
+		else
+		{
+			State->FlushBufferedLines();
+			State->BroadcastTo(ThreadId, State->BufferedOutputDevices, UE_PROJECTION_MEMBER(FOutputDevice, Serialize),
+				Data, Verbosity, Category, RealTime);
+		}
 	}
 	else
 	{
-		State->FlushBufferedLines();
-
-		for (FOutputDevice* OutputDevice : State->BufferedOutputDevices)
-		{
-			OutputDevice->Serialize(Data, Verbosity, Category, RealTime);
-		}
+		// Another thread has triggered a panic and this data will be lost.
 	}
 }
 
@@ -559,51 +585,47 @@ void FOutputDeviceRedirector::RedirectLog(const FLazyName& Category, ELogVerbosi
 
 void FOutputDeviceRedirector::Flush()
 {
-	UE::Private::FOutputDevicesMasterScopeLock Lock(*State);
-	State->FlushBufferedLines();
-
-	for (FOutputDevice* OutputDevice : State->BufferedOutputDevices)
+	if (UE::Private::FOutputDevicesMasterScope Lock(*State); Lock.IsLocked())
 	{
-		OutputDevice->Flush();
-	}
-
-	for (FOutputDevice* OutputDevice : State->UnbufferedOutputDevices)
-	{
-		OutputDevice->Flush();
+		State->FlushBufferedLines();
+		const uint32 ThreadId = FPlatformTLS::GetCurrentThreadId();
+		State->BroadcastTo(ThreadId, State->BufferedOutputDevices, &FOutputDevice::Flush);
+		State->BroadcastTo(ThreadId, State->UnbufferedOutputDevices, &FOutputDevice::Flush);
 	}
 }
 
-void FOutputDeviceRedirector::PanicFlush()
+void FOutputDeviceRedirector::Panic()
 {
-	if (UE::Private::FOutputDevicesMasterScopeTryLock Lock(*State, 1.0); Lock.IsLocked())
+	uint32 PreviousThreadId = MAX_uint32;
+	const uint32 ThreadId = FPlatformTLS::GetCurrentThreadId();
+	if (!State->PanicThreadId.compare_exchange_strong(PreviousThreadId, ThreadId, std::memory_order_relaxed))
 	{
-		State->BufferedLines.Deplete([this](UE::Private::FOutputDeviceLine&& Line)
-		{
-			for (FOutputDevice* OutputDevice : State->BufferedOutputDevices)
-			{
-				if (OutputDevice->CanBeUsedByPanicFlush())
-				{
-					OutputDevice->Serialize(Line.Data, Line.Verbosity, Line.Category, Line.Time);
-				}
-			}
-		});
-
-		for (FOutputDevice* OutputDevice : State->BufferedOutputDevices)
-		{
-			if (OutputDevice->CanBeUsedByPanicFlush())
-			{
-				OutputDevice->Flush();
-			}
-		}
-
-		for (FOutputDevice* OutputDevice : State->UnbufferedOutputDevices)
+		return;
+	}
+	else
+	{
+		// Another thread may be holding the lock. Wait a while for it, but avoid waiting forever
+		// because the thread holding the lock may be unable to progress. After the timeout is
+		// reached, assume that it is safe enough to continue on the panic thread. There is a
+		// chance that the thread holding the lock has left an output device in an unusable state
+		// or will resume and crash due to a race with the panic thread. Executing on this thread
+		// and having logging for most panic situations with a chance of a crash is preferable to
+		// the alternative of missing logging in a panic situation.
+		TRACE_CPUPROFILER_EVENT_SCOPE(FOutputDeviceRedirector::PanicWait);
+		constexpr double WaitTime = 1.0;
+		for (const double EndTime = FPlatformTime::Seconds() + WaitTime; FPlatformTime::Seconds() < EndTime;)
 		{
-			if (OutputDevice->CanBeUsedByPanicFlush())
+			if (State->OutputDevicesLock.TryWriteLock())
 			{
-				OutputDevice->Flush();
+				State->OutputDevicesLock.WriteUnlock();
+				break;
 			}
+			FPlatformProcess::Yield();
 		}
+		State->MasterThreadId.exchange(ThreadId, std::memory_order_relaxed);
 	}
+
+	Flush();
 }
 
 void FOutputDeviceRedirector::TearDown()
@@ -612,6 +634,8 @@ void FOutputDeviceRedirector::TearDown()
 
 	Flush();
 
+	State->TryStopThread();
+
 	TArray<FOutputDevice*> LocalBufferedDevices;
 	TArray<FOutputDevice*> LocalUnbufferedDevices;
 
diff --git a/Engine/Source/Runtime/Core/Private/Unix/UnixErrorOutputDevice.cpp b/Engine/Source/Runtime/Core/Private/Unix/UnixErrorOutputDevice.cpp
--- a/Engine/Source/Runtime/Core/Private/Unix/UnixErrorOutputDevice.cpp
+++ b/Engine/Source/Runtime/Core/Private/Unix/UnixErrorOutputDevice.cpp
@@ -101,9 +101,7 @@ void FUnixErrorOutputDevice::HandleError()
 		UE_LOG(LogCore, Log, TEXT("=== Critical error: ===") LINE_TERMINATOR TEXT("%s") LINE_TERMINATOR, GErrorExceptionDescription);
 		UE_LOG(LogCore, Log, TEXT("%s"), GErrorHist);
 
-		// TODO Fix this, as we no longer *output* this callstack most the time now
-		// This fixes a hang vs using GLog->Flush()
-		GLog->PanicFlush();
+		GLog->Panic();
 
 		HandleErrorRestoreUI();
 
diff --git a/Engine/Source/Runtime/Core/Private/Unix/UnixPlatformCrashContext.cpp b/Engine/Source/Runtime/Core/Private/Unix/UnixPlatformCrashContext.cpp
--- a/Engine/Source/Runtime/Core/Private/Unix/UnixPlatformCrashContext.cpp
+++ b/Engine/Source/Runtime/Core/Private/Unix/UnixPlatformCrashContext.cpp
@@ -595,7 +595,7 @@ void FUnixCrashContext::GenerateCrashInfoAndLaunchReporter(bool bReportingNonCra
 			FPaths::NormalizeDirectoryName(LogDstAbsolute);
 
 			// Flush out the log
-			GLog->Flush();
+			GLog->Panic();
 
 #if !NO_LOGGING
 			bool bMemoryOnly = FPlatformOutputDevices::GetLog()->IsMemoryOnly();
@@ -847,7 +847,7 @@ void DefaultCrashHandler(const FUnixCrashContext & Context)
 	const_cast<FUnixCrashContext&>(Context).CaptureStackTrace(Context.ErrorFrame);
 	if (GLog)
 	{
-		GLog->PanicFlush();
+		GLog->Panic();
 	}
 	if (GWarn)
 	{
diff --git a/Engine/Source/Runtime/Core/Private/Windows/WindowsErrorOutputDevice.cpp b/Engine/Source/Runtime/Core/Private/Windows/WindowsErrorOutputDevice.cpp
--- a/Engine/Source/Runtime/Core/Private/Windows/WindowsErrorOutputDevice.cpp
+++ b/Engine/Source/Runtime/Core/Private/Windows/WindowsErrorOutputDevice.cpp
@@ -107,7 +107,7 @@ void FWindowsErrorOutputDevice::HandleError()
 #if !NO_LOGGING
 	FDebug::LogFormattedMessageWithCallstack(LogWindows.GetCategoryName(), __FILE__, __LINE__, TEXT("=== Critical error: ==="), GErrorHist, ELogVerbosity::Error);
 #endif
-	GLog->PanicFlush();
+	GLog->Panic();
 
 	HandleErrorRestoreUI();
 
diff --git a/Engine/Source/Runtime/Core/Private/Windows/WindowsPlatformCrashContext.cpp b/Engine/Source/Runtime/Core/Private/Windows/WindowsPlatformCrashContext.cpp
--- a/Engine/Source/Runtime/Core/Private/Windows/WindowsPlatformCrashContext.cpp
+++ b/Engine/Source/Runtime/Core/Private/Windows/WindowsPlatformCrashContext.cpp
@@ -1412,7 +1412,7 @@ private:
 		{
 			// Panic flush the logs to make sure there are no entries queued.
 			// This is will skip output devices that are not safe to flush during a crash, for example, the editor log.
-			GLog->PanicFlush();
+			GLog->Panic();
 		}
 		
 		// Get the default settings for the crash context
diff --git a/Engine/Source/Runtime/Core/Public/Misc/OutputDevice.h b/Engine/Source/Runtime/Core/Public/Misc/OutputDevice.h
--- a/Engine/Source/Runtime/Core/Public/Misc/OutputDevice.h
+++ b/Engine/Source/Runtime/Core/Public/Misc/OutputDevice.h
@@ -200,9 +200,9 @@ public:
 	}
 
 	/**
-	 * @return whether this output device can be used when flushing after a fatal error.
+	 * @return whether this output device can be used after a panic (crash or fatal error) has been flagged.
 	 */
-	virtual bool CanBeUsedByPanicFlush() const
+	virtual bool CanBeUsedOnPanicThread() const
 	{
 		return CanBeUsedOnAnyThread();
 	}
diff --git a/Engine/Source/Runtime/Core/Public/Misc/OutputDeviceRedirector.h b/Engine/Source/Runtime/Core/Public/Misc/OutputDeviceRedirector.h
--- a/Engine/Source/Runtime/Core/Public/Misc/OutputDeviceRedirector.h
+++ b/Engine/Source/Runtime/Core/Public/Misc/OutputDeviceRedirector.h
@@ -103,9 +103,9 @@ public:
 	/** Flushes lines buffered by secondary threads. */
 	void FlushThreadedLogs(EOutputDeviceRedirectorFlushOptions Options = EOutputDeviceRedirectorFlushOptions::None);
 
-	/** Same as PanicFlush(). */
-	UE_DEPRECATED(5.1, "Use PanicFlush().")
-	void PanicFlushThreadedLogs() { PanicFlush(); }
+	/** See Panic. */
+	UE_DEPRECATED(5.1, "Use Panic() when the caller is handling a crash, otherwise use FlushThreadedLogs().")
+	void PanicFlushThreadedLogs() { Panic(); }
 
 	/**
 	 * Serializes the current backlog to the specified output device.
@@ -155,8 +155,15 @@ public:
 	/** Passes on the flush request to all current output devices. */
 	void Flush() final;
 
-	/** Makes an effort to flush threaded logs and flush output devices after a fatal error. */
-	void PanicFlush();
+	/**
+	 * Attempts to set the calling thread as the panic thread and enable panic mode.
+	 *
+	 * Returns with no side effects if another thread is the panic thread.
+	 * Only redirects logs to panic-safe output devices from this point forward.
+	 * Flushes buffered logs to panic-safe output devices.
+	 * Flushes panic-safe output devices.
+	 */
+	void Panic();
 
 	/**
 	 * Closes output device and cleans up.
diff --git a/Engine/Source/Runtime/D3D12RHI/Private/D3D12Util.cpp b/Engine/Source/Runtime/D3D12RHI/Private/D3D12Util.cpp
--- a/Engine/Source/Runtime/D3D12RHI/Private/D3D12Util.cpp
+++ b/Engine/Source/Runtime/D3D12RHI/Private/D3D12Util.cpp
@@ -280,7 +280,7 @@ static bool LogBreadcrumbData(D3D12RHI::FD3DGPUProfiler& GPUProfiler, FD3D12Comm
 		}
 	}
 
-	GLog->PanicFlush();
+	GLog->Panic();
 
 	return true;
 }
@@ -790,7 +790,7 @@ namespace D3D12RHI
 		}
 		
 		// Make sure the log is flushed!
-		GLog->PanicFlush();
+		GLog->Panic();
 
 		// Show message box or trace information
 #if !(UE_BUILD_SHIPPING || UE_BUILD_TEST)
@@ -851,7 +851,7 @@ namespace D3D12RHI
 		}
 
 		// Make sure the log is flushed!
-		GLog->PanicFlush();
+		GLog->Panic();
 
 		UE_LOG(LogD3D12RHI, Fatal, TEXT("%s failed \n at %s:%u \n with error %s\n%s"), ANSI_TO_TCHAR(Code), ANSI_TO_TCHAR(Filename), Line, *ErrorString, *Message);
 
@@ -895,7 +895,7 @@ namespace D3D12RHI
 		}
 
 		// Make sure the log is flushed!
-		GLog->PanicFlush();
+		GLog->Panic();
 
 		UE_LOG(LogD3D12RHI, Fatal,
 			TEXT("%s failed \n at %s:%u \n with error %s, \n Size=%ix%ix%i Format=%s(0x%08X), NumMips=%i, Flags=%s"),
diff --git a/Engine/Source/Runtime/Engine/Private/NetDriver.cpp b/Engine/Source/Runtime/Engine/Private/NetDriver.cpp
--- a/Engine/Source/Runtime/Engine/Private/NetDriver.cpp
+++ b/Engine/Source/Runtime/Engine/Private/NetDriver.cpp
@@ -4973,7 +4973,7 @@ struct FReplicationAutoCapture
 			--KillProcessFrames;
 			if (KillProcessFrames == 0)
 			{
-				GLog->PanicFlush();
+				GLog->Flush();
 				FPlatformMisc::RequestExit(1);
 			}
 		}
diff --git a/Engine/Source/Runtime/Engine/Private/ShaderCompiler/ShaderCompiler.cpp b/Engine/Source/Runtime/Engine/Private/ShaderCompiler/ShaderCompiler.cpp
--- a/Engine/Source/Runtime/Engine/Private/ShaderCompiler/ShaderCompiler.cpp
+++ b/Engine/Source/Runtime/Engine/Private/ShaderCompiler/ShaderCompiler.cpp
@@ -1140,7 +1140,7 @@ namespace SCWErrorCode
 {
 	void HandleGeneralCrash(const TCHAR* ExceptionInfo, const TCHAR* Callstack)
 	{
-		GLog->PanicFlush();
+		GLog->Panic();
 		UE_LOG(LogShaderCompilers, Fatal, TEXT("ShaderCompileWorker crashed!\n%s\n\t%s"), ExceptionInfo, Callstack);
 	}
 
@@ -1643,7 +1643,7 @@ static void DumpCompilationJobs(const TArray<FShaderCommonCompileJobPtr>& Queued
 	}
 
 	// Force a log flush so we can track the crash before the cooker potentially crashes before the output shows up
-	GLog->PanicFlush();
+	GLog->Flush();
 }
 
 // Disable optimization for this crash handler to get full access to the entire stack frame when debugging a crash dump
diff --git a/Engine/Source/Runtime/Launch/Private/Android/AndroidJNI.cpp b/Engine/Source/Runtime/Launch/Private/Android/AndroidJNI.cpp
--- a/Engine/Source/Runtime/Launch/Private/Android/AndroidJNI.cpp
+++ b/Engine/Source/Runtime/Launch/Private/Android/AndroidJNI.cpp
@@ -489,7 +489,7 @@ void EngineCrashHandler(const FGenericCrashContext& GenericContext)
 
 		if (GLog)
 		{
-			GLog->PanicFlush();
+			GLog->Panic();
 		}
 		
 		if (GWarn)
diff --git a/Engine/Source/Runtime/Launch/Private/Mac/LaunchMac.cpp b/Engine/Source/Runtime/Launch/Private/Mac/LaunchMac.cpp
--- a/Engine/Source/Runtime/Launch/Private/Mac/LaunchMac.cpp
+++ b/Engine/Source/Runtime/Launch/Private/Mac/LaunchMac.cpp
@@ -32,7 +32,7 @@ void LogCrashCallstack(const FMacCrashContext& Context)
     Context.ReportCrash();
     if (GLog)
     {
-        GLog->PanicFlush();
+        GLog->Panic();
     }
     if (GWarn)
     {
diff --git a/Engine/Source/Runtime/RHI/Private/RHIValidation.cpp b/Engine/Source/Runtime/RHI/Private/RHIValidation.cpp
--- a/Engine/Source/Runtime/RHI/Private/RHIValidation.cpp
+++ b/Engine/Source/Runtime/RHI/Private/RHIValidation.cpp
@@ -449,7 +449,7 @@ void FValidationRHI::ReportValidationFailure(const TCHAR* InMessage)
 		// Print the message again using the debug output function, because UE_LOG doesn't always reach
 		// the VS output window before the breakpoint is triggered, despite the log flush call below.
 		FPlatformMisc::LowLevelOutputDebugStringf(TEXT("%s\n"), InMessage);
-		GLog->PanicFlush();
+		GLog->Flush();
 		PLATFORM_BREAK();
 	}
 }
diff --git a/Engine/Source/Runtime/Unix/UnixCommonStartup/Private/UnixCommonStartup.cpp b/Engine/Source/Runtime/Unix/UnixCommonStartup/Private/UnixCommonStartup.cpp
--- a/Engine/Source/Runtime/Unix/UnixCommonStartup/Private/UnixCommonStartup.cpp
+++ b/Engine/Source/Runtime/Unix/UnixCommonStartup/Private/UnixCommonStartup.cpp
@@ -35,7 +35,7 @@ void CommonUnixCrashHandler(const FGenericCrashContext& GenericContext)
 	const_cast< FUnixCrashContext& >(Context).CaptureStackTrace(Context.ErrorFrame);
 	if (GLog)
 	{
-		GLog->PanicFlush();
+		GLog->Panic();
 	}
 	if (GWarn)
 	{
diff --git a/Engine/Source/Runtime/VulkanRHI/Private/VulkanMemory.cpp b/Engine/Source/Runtime/VulkanRHI/Private/VulkanMemory.cpp
--- a/Engine/Source/Runtime/VulkanRHI/Private/VulkanMemory.cpp
+++ b/Engine/Source/Runtime/VulkanRHI/Private/VulkanMemory.cpp
@@ -1022,7 +1022,7 @@ namespace VulkanRHI
 			}
 			DumpRenderTargetPoolMemory(*GLog);
 			Device->GetMemoryManager().DumpMemory();
-			GLog->PanicFlush();
+			GLog->Panic();
 
 			UE_LOG(LogVulkanRHI, Fatal, TEXT("Out of %s Memory, Requested%.2fKB MemTypeIndex=%d\n"), MemoryType, AllocationSize / 1024.f, MemoryTypeIndex);
 		}
@@ -3624,7 +3624,7 @@ namespace VulkanRHI
 		}
 
 
-		GLog->PanicFlush();
+		GLog->Flush();
 
 
 #undef VULKAN_LOGMEMORY_PAD
@@ -3643,9 +3643,9 @@ namespace VulkanRHI
 			case VK_ERROR_OUT_OF_DEVICE_MEMORY: MemoryType = TEXT("Local"); break;
 			}
 			DumpMemory();
-			GLog->PanicFlush();
+			GLog->Panic();
 			DumpRenderTargetPoolMemory(*GLog);
-			GLog->PanicFlush();
+			GLog->Flush();
 
 			UE_LOG(LogVulkanRHI, Fatal, TEXT("Out of %s Memory, Requested%.2fKB MemTypeIndex=%d\n"), MemoryType, AllocationSize, MemoryTypeIndex);
 		}
diff --git a/Engine/Source/Runtime/VulkanRHI/Private/VulkanUtil.cpp b/Engine/Source/Runtime/VulkanRHI/Private/VulkanUtil.cpp
--- a/Engine/Source/Runtime/VulkanRHI/Private/VulkanUtil.cpp
+++ b/Engine/Source/Runtime/VulkanRHI/Private/VulkanUtil.cpp
@@ -692,7 +692,7 @@ void FVulkanGPUProfiler::DumpCrashMarkers(void* BufferData)
 					const FString* Frame = CachedStrings.Find(Value);
 					UE_LOG(LogVulkanRHI, Error, TEXT("[VK_NV_device_diagnostic_checkpoints] %i: Stage 0x%08x, %s (CRC 0x%x)"), Index, Data[Index].stage, Frame ? *(*Frame) : TEXT("<undefined>"), Value);
 				}
-				GLog->PanicFlush();
+				GLog->Panic();
 			}
 		}
 #endif
@@ -710,7 +710,7 @@ void FVulkanGPUProfiler::DumpCrashMarkers(void* BufferData)
 			UE_LOG(LogVulkanRHI, Warning, TEXT("[gpu_crash_markers] %s"), (CrashMarkers[i] != 0) ? *FrameName : TEXT("unavailable"));
 		}
 
-		GLog->PanicFlush();
+		GLog->Panic();
 	}
 }
 #endif
diff --git a/Engine/Source/Runtime/Windows/D3D11RHI/Private/D3D11RHI.cpp b/Engine/Source/Runtime/Windows/D3D11RHI/Private/D3D11RHI.cpp
--- a/Engine/Source/Runtime/Windows/D3D11RHI/Private/D3D11RHI.cpp
+++ b/Engine/Source/Runtime/Windows/D3D11RHI/Private/D3D11RHI.cpp
@@ -645,7 +645,7 @@ bool FD3DGPUProfiler::CheckGpuHeartbeat(bool bShowActiveStatus) const
 					UE_LOG(LogRHI, Error, TEXT("[Aftermath] Invalid context handle"));
 					NVAFTERMATH_ON_ERROR();
 				}
-				GLog->PanicFlush();
+				GLog->Flush();
 				return false;
 			}
 		}
