FileHelper detect sizes too big for 32 bit load and fail rather than just crash

diff --git a/Engine/Source/Runtime/Core/Private/Misc/FileHelper.cpp b/Engine/Source/Runtime/Core/Private/Misc/FileHelper.cpp
--- a/Engine/Source/Runtime/Core/Private/Misc/FileHelper.cpp
+++ b/Engine/Source/Runtime/Core/Private/Misc/FileHelper.cpp
@@ -44,7 +44,16 @@ bool FFileHelper::LoadFileToArray( TArray<uint8>& Result, const TCHAR* Filename,
 		}
 		return false;
 	}
-	int32 TotalSize = (int32)Reader->TotalSize();
+	int64 TotalSize64 = Reader->TotalSize();
+	if ( TotalSize64+2 > MAX_int32 )
+	{
+		if (!(Flags & FILEREAD_Silent))
+		{
+			UE_LOG(LogStreaming,Error,TEXT("File '%s' is too large for 32-bit reader (%lld), use TArray64."),Filename,TotalSize64);
+		}
+		return false;
+	}
+	int32 TotalSize = (int32)TotalSize64;
 	// Allocate slightly larger than file size to avoid re-allocation when caller null terminates file buffer
 	Result.Reset( TotalSize + 2 );
 	Result.AddUninitialized( TotalSize );
