removed unused `NewTasks` from TaskGraph's FGraphEvent::DispatchSubsequents() calls, as a preparation for deprecation

diff --git a/Engine/Plugins/Developer/TraceSourceFiltering/Source/SourceFilteringTrace/Private/SourceFilterManager.cpp b/Engine/Plugins/Developer/TraceSourceFiltering/Source/SourceFilteringTrace/Private/SourceFilterManager.cpp
--- a/Engine/Plugins/Developer/TraceSourceFiltering/Source/SourceFilteringTrace/Private/SourceFilterManager.cpp
+++ b/Engine/Plugins/Developer/TraceSourceFiltering/Source/SourceFilteringTrace/Private/SourceFilterManager.cpp
@@ -116,8 +116,7 @@ FSourceFilterManager::~FSourceFilterManager()
 	if (IsValidRef(DrawTask) && !DrawTask->IsComplete())
 	{
 		// Mark drawing task as finished (would otherwise have been executed on GT)
-		TArray<FBaseGraphTask*> NewTasks;
-		DrawTask->DispatchSubsequents(NewTasks, ENamedThreads::GameThread);
+		DrawTask->DispatchSubsequents(ENamedThreads::GameThread);
 	}
 
 	if (Settings)
@@ -757,9 +756,7 @@ void FSourceFilterManager::WaitForAsyncTasks()
 		if (IsValidRef(AsyncTasks[0]) && !AsyncTasks[0]->IsComplete())
 		{			
 			ApplyGameThreadFilters();
-
-			TArray<FBaseGraphTask*> NewTasks;
-			AsyncTasks[0]->DispatchSubsequents(NewTasks, ENamedThreads::GameThread);
+			AsyncTasks[0]->DispatchSubsequents(ENamedThreads::GameThread);
 		}
 
 		FTaskGraphInterface::Get().WaitUntilTaskCompletes(FinishTask);
diff --git a/Engine/Source/Runtime/Core/Private/FileCache/FileCache.cpp b/Engine/Source/Runtime/Core/Private/FileCache/FileCache.cpp
--- a/Engine/Source/Runtime/Core/Private/FileCache/FileCache.cpp
+++ b/Engine/Source/Runtime/Core/Private/FileCache/FileCache.cpp
@@ -513,8 +513,7 @@ FFileCacheHandle::FFileCacheHandle(IAsyncReadFileHandle* InHandle, int64 InBaseO
 		this->FileSize = Request->GetSizeResults();
 		check(this->FileSize > 0);
 
-		TArray<FBaseGraphTask*> NewTasks;
-		CompletionEvent->DispatchSubsequents(NewTasks);
+		CompletionEvent->DispatchSubsequents();
 		GetCache().PushCompletedRequest(Request);
 	};
 
@@ -589,8 +588,7 @@ IMemoryReadStreamRef FFileCacheHandle::ReadDataUncached(FGraphEventArray& OutCom
 
 	FAsyncFileCallBack ReadCallbackFunction = [CompletionEvent](bool bWasCancelled, IAsyncReadRequest* Request)
 	{
-		TArray<FBaseGraphTask*> NewTasks;
-		CompletionEvent->DispatchSubsequents(NewTasks);
+		CompletionEvent->DispatchSubsequents();
 	};
 
 	OutCompletionEvents.Add(CompletionEvent);
@@ -668,8 +666,7 @@ void FFileCacheHandle::ReadLine(FFileCache& Cache, CacheSlotID SlotID, CacheLine
 	// callback triggered when async read operation is complete, used to signal task graph event
 	FAsyncFileCallBack ReadCallbackFunction = [CompletionEvent](bool bWasCancelled, IAsyncReadRequest* Request)
 	{
-		TArray<FBaseGraphTask*> NewTasks;
-		CompletionEvent->DispatchSubsequents(NewTasks);
+		CompletionEvent->DispatchSubsequents();
 		GetCache().PushCompletedRequest(Request);
 	};
 
diff --git a/Engine/Source/Runtime/Core/Private/IO/IoDispatcher.cpp b/Engine/Source/Runtime/Core/Private/IO/IoDispatcher.cpp
--- a/Engine/Source/Runtime/Core/Private/IO/IoDispatcher.cpp
+++ b/Engine/Source/Runtime/Core/Private/IO/IoDispatcher.cpp
@@ -646,8 +646,7 @@ private:
 		}
 		if (Batch->GraphEvent)
 		{
-			TArray<FBaseGraphTask*> NewTasks;
-			Batch->GraphEvent->DispatchSubsequents(NewTasks);
+			Batch->GraphEvent->DispatchSubsequents();
 		}
 		BatchAllocator.Destroy(Batch);
 	}
diff --git a/Engine/Source/Runtime/Core/Private/ProfilingDebugging/CsvProfiler.cpp b/Engine/Source/Runtime/Core/Private/ProfilingDebugging/CsvProfiler.cpp
--- a/Engine/Source/Runtime/Core/Private/ProfilingDebugging/CsvProfiler.cpp
+++ b/Engine/Source/Runtime/Core/Private/ProfilingDebugging/CsvProfiler.cpp
@@ -3073,8 +3073,7 @@ TSharedFuture<FString> FCsvProfiler::EndCapture(FGraphEventRef EventToSignal)
 	{
 		if (EventToSignal)
 		{
-			TArray<FBaseGraphTask*> Subsequents;
-			EventToSignal->DispatchSubsequents(Subsequents);
+			EventToSignal->DispatchSubsequents();
 		}
 	});
 
diff --git a/Engine/Source/Runtime/Engine/Private/AudioThread.cpp b/Engine/Source/Runtime/Engine/Private/AudioThread.cpp
--- a/Engine/Source/Runtime/Engine/Private/AudioThread.cpp
+++ b/Engine/Source/Runtime/Engine/Private/AudioThread.cpp
@@ -423,8 +423,7 @@ struct FAudioAsyncBatcher
 			FGraphEventRef Dispatch = DispatchEvent[0];
 			TFunction<void()> FlushAudioCommands = [Dispatch]()
 			{
-				TArray<FBaseGraphTask*> NewTasks;
-				Dispatch->DispatchSubsequents(NewTasks);
+				Dispatch->DispatchSubsequents();
 			};
 
 			FFunctionGraphTask::CreateAndDispatchWhenReady(MoveTemp(FlushAudioCommands), TStatId(), nullptr, ENamedThreads::AudioThread);
diff --git a/Engine/Source/Runtime/Engine/Private/Particles/ParticleComponents.cpp b/Engine/Source/Runtime/Engine/Private/Particles/ParticleComponents.cpp
--- a/Engine/Source/Runtime/Engine/Private/Particles/ParticleComponents.cpp
+++ b/Engine/Source/Runtime/Engine/Private/Particles/ParticleComponents.cpp
@@ -4871,8 +4871,7 @@ public:
 				{
 					check(FinalizePrereq.GetReference() && !FinalizePrereq->IsComplete());
 					{
-						TArray<FBaseGraphTask*> NewTasks;
-						FinalizePrereq->DispatchSubsequents(NewTasks);
+						FinalizePrereq->DispatchSubsequents();
 					}
 					delete FinalizeDispatchCounter;
 				}
@@ -4910,8 +4909,7 @@ public:
 	{
 		check(Target.GetReference() && !Target->IsComplete());
 		{
-			TArray<FBaseGraphTask*> NewTasks;
-			Target->DispatchSubsequents(NewTasks);
+			Target->DispatchSubsequents();
 		}
 	}
 };
diff --git a/Engine/Source/Runtime/Engine/Private/PhysicsEngine/PhysScene_PhysX.cpp b/Engine/Source/Runtime/Engine/Private/PhysicsEngine/PhysScene_PhysX.cpp
--- a/Engine/Source/Runtime/Engine/Private/PhysicsEngine/PhysScene_PhysX.cpp
+++ b/Engine/Source/Runtime/Engine/Private/PhysicsEngine/PhysScene_PhysX.cpp
@@ -1175,8 +1175,7 @@ void FPhysScene_PhysX::TickPhysScene(FGraphEventRef& InOutCompletionEvent)
 
 	if (!bTaskOutstanding)
 	{
-		TArray<FBaseGraphTask*> NewTasks;
-		InOutCompletionEvent->DispatchSubsequents(NewTasks, ENamedThreads::AnyThread); // nothing to do, so nothing to wait for
+		InOutCompletionEvent->DispatchSubsequents(); // nothing to do, so nothing to wait for
 	}
 
 	bSubstepping = UPhysicsSettings::Get()->bSubstepping;
diff --git a/Engine/Source/Runtime/Engine/Private/PhysicsEngine/PhysSubstepTasks.h b/Engine/Source/Runtime/Engine/Private/PhysicsEngine/PhysSubstepTasks.h
--- a/Engine/Source/Runtime/Engine/Private/PhysicsEngine/PhysSubstepTasks.h
+++ b/Engine/Source/Runtime/Engine/Private/PhysicsEngine/PhysSubstepTasks.h
@@ -46,8 +46,7 @@ public:
 		FinishSceneStat();
 		if (EventToFire.GetReference())
 		{
-			TArray<FBaseGraphTask*> NewTasks;
-			EventToFire->DispatchSubsequents(NewTasks);
+			EventToFire->DispatchSubsequents();
 		}
 		delete this;
 	}
diff --git a/Engine/Source/Runtime/RHI/Private/RHICommandList.cpp b/Engine/Source/Runtime/RHI/Private/RHICommandList.cpp
--- a/Engine/Source/Runtime/RHI/Private/RHICommandList.cpp
+++ b/Engine/Source/Runtime/RHI/Private/RHICommandList.cpp
@@ -972,8 +972,7 @@ FRHICOMMAND_MACRO(FRHICommandRHIThreadFence)
 	void Execute(FRHICommandListBase& CmdList)
 	{
 		check(IsInRHIThread());
-		static TArray<FBaseGraphTask*> NewTasks;
-		Fence->DispatchSubsequents(NewTasks, IsRunningRHIInDedicatedThread() ? ENamedThreads::RHIThread : ENamedThreads::AnyThread);
+		Fence->DispatchSubsequents(IsRunningRHIInDedicatedThread() ? ENamedThreads::RHIThread : ENamedThreads::AnyThread);
 		Fence = nullptr;
 	}
 };
diff --git a/Engine/Source/Runtime/RHI/Private/RHIUtilities.cpp b/Engine/Source/Runtime/RHI/Private/RHIUtilities.cpp
--- a/Engine/Source/Runtime/RHI/Private/RHIUtilities.cpp
+++ b/Engine/Source/Runtime/RHI/Private/RHIUtilities.cpp
@@ -371,8 +371,7 @@ uint32 FRHIFrameFlipTrackingRunnable::Run()
 			if (Pair.PresentIndex <= SyncFrame)
 			{
 				// "Complete" the task graph event
-				TArray<FBaseGraphTask*> Subsequents;
-				Pair.Event->DispatchSubsequents(Subsequents);
+				Pair.Event->DispatchSubsequents();
 
 				FramePairs.RemoveAtSwap(PairIndex);
 			}
@@ -430,8 +429,7 @@ void FRHIFrameFlipTrackingRunnable::Shutdown()
 	for (auto const& Pair : Singleton.FramePairs)
 	{
 		// "Complete" the task graph event
-		TArray<FBaseGraphTask*> Subsequents;
-		Pair.Event->DispatchSubsequents(Subsequents);
+		Pair.Event->DispatchSubsequents();
 	}
 
 	Singleton.FramePairs.Empty();
@@ -468,9 +466,7 @@ void FRHIFrameFlipTrackingRunnable::CompleteGraphEventOnFlip(uint64 PresentIndex
 	{
 		// Platform does not support flip tracking.
 		// Signal the event now...
-
-		TArray<FBaseGraphTask*> Subsequents;
-		Event->DispatchSubsequents(Subsequents);
+		Event->DispatchSubsequents();
 	}
 }
 
diff --git a/Engine/Source/Runtime/RenderCore/Private/RenderingThread.cpp b/Engine/Source/Runtime/RenderCore/Private/RenderingThread.cpp
--- a/Engine/Source/Runtime/RenderCore/Private/RenderingThread.cpp
+++ b/Engine/Source/Runtime/RenderCore/Private/RenderingThread.cpp
@@ -988,8 +988,7 @@ void StopRenderCommandFenceBundler()
 
 	EndBatchedRelease();
 	check(IsInGameThread() && BundledCompletionEvent.GetReference() && !BundledCompletionEvent->IsComplete() && BundledCompletionEventPrereq.GetReference() && !BundledCompletionEventPrereq->IsComplete()); // can't use this in a nested fashion
-	TArray<FBaseGraphTask*> NewTasks;
-	BundledCompletionEventPrereq->DispatchSubsequents(NewTasks);
+	BundledCompletionEventPrereq->DispatchSubsequents();
 	BundledCompletionEventPrereq = nullptr;
 	BundledCompletionEvent = nullptr;
 }
@@ -1021,8 +1020,7 @@ FRHICOMMAND_MACRO(FRHISyncFrameCommand)
 			// Sync the Game Thread with the RHI Thread
 
 			// "Complete" the graph event
-			TArray<FBaseGraphTask*> Subsequents;
-			GraphEvent->DispatchSubsequents(Subsequents);
+			GraphEvent->DispatchSubsequents();
 		}
 		else
 		{
