Logging: Added more flushing to panic mode

Panic() will always Flush() on the first call, and on subsequent calls from the panic thread. After Panic(), every Serialize() from the panic thread will Flush() too.

diff --git a/Engine/Source/Runtime/Core/Private/Misc/OutputDeviceRedirector.cpp b/Engine/Source/Runtime/Core/Private/Misc/OutputDeviceRedirector.cpp
--- a/Engine/Source/Runtime/Core/Private/Misc/OutputDeviceRedirector.cpp
+++ b/Engine/Source/Runtime/Core/Private/Misc/OutputDeviceRedirector.cpp
@@ -509,15 +509,13 @@ void FOutputDeviceRedirector::Serialize(const TCHAR* const Data, const ELogVerbo
 	UE::Private::FOutputDevicesReadScopeLock Lock(*State);
 
 #if PLATFORM_DESKTOP
-	// this is for errors which occur after shutdown we might be able to salvage information from stdout
-	if (State->BufferedOutputDevices.IsEmpty() && IsEngineExitRequested())
+	// Print anything that arrives after logging has shut down to at least have it in stdout.
+	if (UNLIKELY(State->BufferedOutputDevices.IsEmpty() && IsEngineExitRequested()))
 	{
-#if PLATFORM_WINDOWS
+	#if PLATFORM_WINDOWS
 		_tprintf(_T("%s\n"), Data);
-#else
+	#endif
 		FGenericPlatformMisc::LocalPrint(Data);
-		// printf("%s\n", TCHAR_TO_ANSI(Data));
-#endif
 		return;
 	}
 #endif
@@ -528,44 +526,41 @@ void FOutputDeviceRedirector::Serialize(const TCHAR* const Data, const ELogVerbo
 	State->BroadcastTo(ThreadId, State->UnbufferedOutputDevices, UE_PROJECTION_MEMBER(FOutputDevice, Serialize),
 		Data, Verbosity, Category, RealTime);
 
-	if (State->bEnableBacklog)
+	// Serialize to the backlog when not in panic mode. This will deadlock in panic mode when the
+	// FPlatformMallocCrash allocator has been enabled and logging occurs on a non-panic thread.
+	if (UNLIKELY(State->bEnableBacklog && !State->IsPanicThread(ThreadId)))
 	{
 		FWriteScopeLock ScopeLock(State->BacklogLock);
 		State->BacklogLines.Emplace(Data, Category, Verbosity, RealTime);
 	}
 
-	const auto EnqueueLine = [this, Data, Category, Verbosity, RealTime]
+	// Serialize to buffered output devices from the master thread.
+	// Lines are queued until buffered output devices are added to avoid missing early log lines.
+	if (State->IsMasterThread(ThreadId) && !State->BufferedOutputDevices.IsEmpty())
 	{
-		if (State->BufferedLines.EnqueueAndReturnWasEmpty(Data, Category, Verbosity, RealTime))
+		// Verify that this is the master thread again because another thread may have become
+		// the master thread between the previous check and the lock.
+		if (UE::Private::FOutputDevicesMasterScope MasterLock(*State); MasterLock.IsLocked() && State->IsMasterThread(ThreadId))
 		{
-			if (FEvent* WakeEvent = State->ThreadWakeEvent.load(std::memory_order_acquire))
+			State->FlushBufferedLines();
+			State->BroadcastTo(ThreadId, State->BufferedOutputDevices, UE_PROJECTION_MEMBER(FOutputDevice, Serialize),
+				Data, Verbosity, Category, RealTime);
+			if (UNLIKELY(State->IsPanicThread(ThreadId)))
 			{
-				WakeEvent->Trigger();
+				Flush();
 			}
+			return;
 		}
-	};
-
-	if (!State->IsMasterThread(ThreadId) || State->BufferedOutputDevices.IsEmpty())
-	{
-		EnqueueLine();
 	}
-	else if (UE::Private::FOutputDevicesMasterScope MasterLock(*State); MasterLock.IsLocked())
+
+	// Queue the line to serialize to buffered output devices from the master thread.
+	if (State->BufferedLines.EnqueueAndReturnWasEmpty(Data, Category, Verbosity, RealTime))
 	{
-		if (!State->IsMasterThread(ThreadId) && !State->IsPanicThread(ThreadId))
+		if (FEvent* WakeEvent = State->ThreadWakeEvent.load(std::memory_order_acquire))
 		{
-			EnqueueLine();
-		}
-		else
-		{
-			State->FlushBufferedLines();
-			State->BroadcastTo(ThreadId, State->BufferedOutputDevices, UE_PROJECTION_MEMBER(FOutputDevice, Serialize),
-				Data, Verbosity, Category, RealTime);
+			WakeEvent->Trigger();
 		}
 	}
-	else
-	{
-		// Another thread has triggered a panic and this data will be lost.
-	}
 }
 
 void FOutputDeviceRedirector::Serialize(const TCHAR* Data, ELogVerbosity::Type Verbosity, const FName& Category)
@@ -598,11 +593,7 @@ void FOutputDeviceRedirector::Panic()
 {
 	uint32 PreviousThreadId = MAX_uint32;
 	const uint32 ThreadId = FPlatformTLS::GetCurrentThreadId();
-	if (!State->PanicThreadId.compare_exchange_strong(PreviousThreadId, ThreadId, std::memory_order_relaxed))
-	{
-		return;
-	}
-	else
+	if (State->PanicThreadId.compare_exchange_strong(PreviousThreadId, ThreadId, std::memory_order_relaxed))
 	{
 		// Another thread may be holding the lock. Wait a while for it, but avoid waiting forever
 		// because the thread holding the lock may be unable to progress. After the timeout is
@@ -622,10 +613,18 @@ void FOutputDeviceRedirector::Panic()
 			}
 			FPlatformProcess::Yield();
 		}
+
+		// Make the panic thread the master thread. Neither thread can be changed after this point.
 		State->MasterThreadId.exchange(ThreadId, std::memory_order_relaxed);
-	}
 
-	Flush();
+		// Flush. Every log from the panic thread after this point will also flush.
+		Flush();
+	}
+	else if (PreviousThreadId == ThreadId)
+	{
+		// Calling Panic() multiple times from the panic thread is equivalent to calling Flush().
+		Flush();
+	}
 }
 
 void FOutputDeviceRedirector::TearDown()
diff --git a/Engine/Source/Runtime/Core/Public/Misc/OutputDeviceRedirector.h b/Engine/Source/Runtime/Core/Public/Misc/OutputDeviceRedirector.h
--- a/Engine/Source/Runtime/Core/Public/Misc/OutputDeviceRedirector.h
+++ b/Engine/Source/Runtime/Core/Public/Misc/OutputDeviceRedirector.h
@@ -158,8 +158,9 @@ public:
 	/**
 	 * Attempts to set the calling thread as the panic thread and enable panic mode.
 	 *
-	 * Returns with no side effects if another thread is the panic thread.
+	 * Only one thread can be the panic thread. Subsequent calls from other threads are ignored.
 	 * Only redirects logs to panic-safe output devices from this point forward.
+	 * Makes the calling thread the master thread as well.
 	 * Flushes buffered logs to panic-safe output devices.
 	 * Flushes panic-safe output devices.
 	 */
