Make some basic math routines for log2/clz simpler and faster
make intention and implementation obvious
check that it produces better code
add unit tests

diff --git a/Engine/Source/Runtime/Core/Private/GenericPlatform/GenericPlatformMath.cpp b/Engine/Source/Runtime/Core/Private/GenericPlatform/GenericPlatformMath.cpp
--- a/Engine/Source/Runtime/Core/Private/GenericPlatform/GenericPlatformMath.cpp
+++ b/Engine/Source/Runtime/Core/Private/GenericPlatform/GenericPlatformMath.cpp
@@ -88,51 +88,27 @@ double FGenericPlatformMath::Atan2(double Y, double X)
 float FGenericPlatformMath::Fmod(float X, float Y)
 {
 	const float AbsY = FMath::Abs(Y);
+	// large arbitrary threshold could be removed if possible :
 	if (AbsY <= 1.e-8f)
 	{
 		FmodReportError(X, Y);
 		return 0.0;
 	}
 
-	// Convert to double for better precision, since intermediate rounding can lose enough precision to skew the result.
-	const double DX = double(X);
-	const double DY = double(Y);
-
-	const double Div = (DX / DY);
-	const double IntPortion = DY * TruncToDouble(Div);
-	const double Result = DX - IntPortion;
-	// Convert back to float. This is safe because the result will by definition not exceed the X input.
-	return float(Result);
+	return fmodf(X,Y);
 }
 
 double FGenericPlatformMath::Fmod(double X, double Y)
 {
 	const double AbsY = FMath::Abs(Y);
+	// large arbitrary threshold could be removed if possible :
 	if (AbsY <= 1.e-8)
 	{
 		FmodReportError(X, Y);
 		return 0.0;
 	}
 
-#if 1
-	// Due to the lack of standard support across platforms for `long double`, we can't rely on this to always minimize the intermediate precision loss of division and multiplication.
-	// As a result, due to small precision loss the result here could be different than the std library version of fmod(), but our validation tests in UnrealMathTest should show
-	// that when different the results are within a small delta from either 0 or the value of Y. We could just use the library version here of course, but it is slower and also will not
-	// match the vectorized versions of Fmod (VectorRegisterMod) that use division and truncation.
-	const double Div = (X / Y);
-	const double IntPortion = Y * FMath::TruncToDouble(Div);
-	const double Result = X - IntPortion;
-	return Result;
-#elif 0
-	const long double LDX = (long double)(X);
-	const long double LDY = (long double)(Y);
-	const long double Div = (LDX / LDY);
-	const long double IntPortion = LDY * truncl(Div);
-	const long double Result = LDX - IntPortion;
-	return double(Result);
-#else
 	return fmod(X, Y);
-#endif
 }
 
 void FGenericPlatformMath::FmodReportError(float X, float Y)
@@ -190,11 +166,38 @@ public:
 		check(MathPlatform::IsFinite(One));
 		check(!MathPlatform::IsNaN(MinusOneE37));
 		check(MathPlatform::IsFinite(MinusOneE37));
+
 		check(MathPlatform::FloorLog2((uint32)Zero) == 0);
 		check(MathPlatform::FloorLog2((uint32)One) == 0);
 		check(MathPlatform::FloorLog2((uint32)Two) == 1);
 		check(MathPlatform::FloorLog2((uint32)Twelve) == 3);
 		check(MathPlatform::FloorLog2((uint32)Sixteen) == 4);
+		
+		check(MathPlatform::CeilLogTwo((uint32)Zero) == 0);
+		check(MathPlatform::CeilLogTwo((uint32)One) == 0);
+		check(MathPlatform::CeilLogTwo((uint32)Two) == 1);
+		check(MathPlatform::CeilLogTwo((uint32)Twelve) == 4);
+		check(MathPlatform::CeilLogTwo((uint32)Sixteen) == 4);
+		
+		check(MathPlatform::CountLeadingZeros8((uint8)Zero) == 8);
+		check(MathPlatform::CountLeadingZeros((uint32)Zero) == 32);
+		check(MathPlatform::CountTrailingZeros((uint32)Zero) == 32);
+#if PLATFORM_64BITS
+		check(MathPlatform::CountLeadingZeros64((uint64)Zero) == 64);
+		check(MathPlatform::CountTrailingZeros64((uint64)Zero) == 64);
+		
+		check(MathPlatform::FloorLog2_64((uint32)Zero) == 0);
+		check(MathPlatform::FloorLog2_64((uint32)One) == 0);
+		check(MathPlatform::FloorLog2_64((uint32)Two) == 1);
+		check(MathPlatform::FloorLog2_64((uint32)Twelve) == 3);
+		check(MathPlatform::FloorLog2_64((uint32)Sixteen) == 4);
+		
+		check(MathPlatform::CeilLogTwo64((uint32)Zero) == 0);
+		check(MathPlatform::CeilLogTwo64((uint32)One) == 0);
+		check(MathPlatform::CeilLogTwo64((uint32)Two) == 1);
+		check(MathPlatform::CeilLogTwo64((uint32)Twelve) == 4);
+		check(MathPlatform::CeilLogTwo64((uint32)Sixteen) == 4);
+#endif
 	}
 };
 
diff --git a/Engine/Source/Runtime/Core/Public/GenericPlatform/GenericPlatformMath.h b/Engine/Source/Runtime/Core/Public/GenericPlatform/GenericPlatformMath.h
--- a/Engine/Source/Runtime/Core/Public/GenericPlatform/GenericPlatformMath.h
+++ b/Engine/Source/Runtime/Core/Public/GenericPlatform/GenericPlatformMath.h
@@ -26,6 +26,7 @@
 struct FGenericPlatformMath
 {
 
+	// load half (F16) to float
 	//https://gist.github.com/rygorous/2156668
 	static FORCEINLINE float LoadHalf(const uint16* Ptr)
 	{
@@ -56,7 +57,10 @@ struct FGenericPlatformMath
 		return FP32.f;
 	}
 
-	//https://gist.github.com/rygorous/2156668
+	// store float to half (F16)
+	// converts with RTNE = round to nearest even
+	// values too large for F16 are stored as +-Inf
+	// https://gist.github.com/rygorous/2156668
 	// float_to_half_fast3_rtne
 	static FORCEINLINE void StoreHalf(uint16* Ptr, float Value)
 	{
@@ -599,12 +603,12 @@ struct FGenericPlatformMath
 
 	static FORCEINLINE bool IsNegative(float A)
 	{
-		return ((*(uint32*)&A) >= (uint32)0x80000000); // Detects sign bit.
+		return ( A < 0.f );
 	}
 
 	static FORCEINLINE bool IsNegative(double A)
 	{
-		return ((*(uint64*)&A) >= (uint64)0x8000000000000000); // Detects sign bit.
+		return ( A < 0.0 );
 	}
 
 	/** Returns a random integer between 0 and RAND_MAX, inclusive */
@@ -631,7 +635,7 @@ struct FGenericPlatformMath
 	static CORE_API float SRand();
 
 	/**
-	 * Computes the base 2 logarithm for an integer value that is greater than 0.
+	 * Computes the base 2 logarithm for an integer value.
 	 * The result is rounded down to the nearest integer.
 	 *
 	 * @param Value		The value to compute the log of
@@ -639,72 +643,17 @@ struct FGenericPlatformMath
 	 */	
 	static FORCEINLINE uint32 FloorLog2(uint32 Value) 
 	{
-/*		// reference implementation 
-		// 1500ms on test data
-		uint32 Bit = 32;
-		for (; Bit > 0;)
-		{
-			Bit--;
-			if (Value & (1<<Bit))
-			{
-				break;
-			}
-		}
-		return Bit;
-*/
-		// same output as reference
-
-		// see http://codinggorilla.domemtech.com/?p=81 or http://en.wikipedia.org/wiki/Binary_logarithm but modified to return 0 for a input value of 0
-		// 686ms on test data
 		uint32 pos = 0;
 		if (Value >= 1<<16) { Value >>= 16; pos += 16; }
 		if (Value >= 1<< 8) { Value >>=  8; pos +=  8; }
 		if (Value >= 1<< 4) { Value >>=  4; pos +=  4; }
 		if (Value >= 1<< 2) { Value >>=  2; pos +=  2; }
 		if (Value >= 1<< 1) {				pos +=  1; }
-		return (Value == 0) ? 0 : pos;
-
-		// even faster would be method3 but it can introduce more cache misses and it would need to store the table somewhere
-		// 304ms in test data
-		/*int LogTable256[256];
-
-		void prep()
-		{
-			LogTable256[0] = LogTable256[1] = 0;
-			for (int i = 2; i < 256; i++)
-			{
-				LogTable256[i] = 1 + LogTable256[i / 2];
-			}
-			LogTable256[0] = -1; // if you want log(0) to return -1
-		}
-
-		int _forceinline method3(uint32 v)
-		{
-			int r;     // r will be lg(v)
-			uint32 tt; // temporaries
-
-			if ((tt = v >> 24) != 0)
-			{
-				r = (24 + LogTable256[tt]);
-			}
-			else if ((tt = v >> 16) != 0)
-			{
-				r = (16 + LogTable256[tt]);
-			}
-			else if ((tt = v >> 8 ) != 0)
-			{
-				r = (8 + LogTable256[tt]);
-			}
-			else
-			{
-				r = LogTable256[v];
-			}
-			return r;
-		}*/
+		return pos;
 	}
 
 	/**
-	 * Computes the base 2 logarithm for a 64-bit value that is greater than 0.
+	 * Computes the base 2 logarithm for a 64-bit value.
 	 * The result is rounded down to the nearest integer.
 	 *
 	 * @param Value		The value to compute the log of
@@ -718,8 +667,8 @@ struct FGenericPlatformMath
 		if (Value >= 1ull<< 8) { Value >>=  8; pos +=  8; }
 		if (Value >= 1ull<< 4) { Value >>=  4; pos +=  4; }
 		if (Value >= 1ull<< 2) { Value >>=  2; pos +=  2; }
-		if (Value >= 1ull<< 1) {				pos +=  1; }
-		return (Value == 0) ? 0 : pos;
+		if (Value >= 1ull<< 1) {               pos +=  1; }
+		return pos;
 	}
 
 	/**
@@ -807,18 +756,20 @@ struct FGenericPlatformMath
 
 	/**
 	 * Returns smallest N such that (1<<N)>=Arg.
-	 * Note: CeilLogTwo(0)=0 because (1<<0)=1 >= 0.
+	 * Note: CeilLogTwo(0)=0 
 	 */
 	static FORCEINLINE uint32 CeilLogTwo( uint32 Arg )
 	{
-		int32 Bitmask = ((int32)(CountLeadingZeros(Arg) << 26)) >> 31;
-		return (32 - CountLeadingZeros(Arg - 1)) & (~Bitmask);
+		// if Arg is 0, change it to 1 so that we return 0
+		Arg = Arg ? Arg : 1;
+		return 32 - CountLeadingZeros(Arg - 1);
 	}
 
 	static FORCEINLINE uint64 CeilLogTwo64( uint64 Arg )
 	{
-		int64 Bitmask = ((int64)(CountLeadingZeros64(Arg) << 57)) >> 63;
-		return (64 - CountLeadingZeros64(Arg - 1)) & (~Bitmask);
+		// if Arg is 0, change it to 1 so that we return 0
+		Arg = Arg ? Arg : 1;
+		return 64 - CountLeadingZeros64(Arg - 1);
 	}
 
 	/**
@@ -949,7 +900,7 @@ struct FGenericPlatformMath
 	template< class T > 
 	static CONSTEXPR FORCEINLINE T Abs( const T A )
 	{
-		return (A>=(T)0) ? A : -A;
+		return (A < (T)0) ? -A : A;
 	}
 
 	/** Returns 1, 0, or -1 depending on relation of T to 0 */
diff --git a/Engine/Source/Runtime/Core/Public/Microsoft/MicrosoftPlatformMath.h b/Engine/Source/Runtime/Core/Public/Microsoft/MicrosoftPlatformMath.h
--- a/Engine/Source/Runtime/Core/Public/Microsoft/MicrosoftPlatformMath.h
+++ b/Engine/Source/Runtime/Core/Public/Microsoft/MicrosoftPlatformMath.h
@@ -22,36 +22,29 @@ struct FMicrosoftPlatformMathBase : public TUnrealPlatformMathSSE4Base<FGenericP
 	static FORCEINLINE uint32 FloorLog2(uint32 Value)
 	{
 		// Use BSR to return the log2 of the integer
-		unsigned long Log2;
-		if(_BitScanReverse(&Log2, Value) != 0)
-		{
-			return Log2;
-		}
-
-		return 0;
+		// return 0 if value is 0
+		unsigned long BitIndex;
+		return _BitScanReverse(&BitIndex, Value) ? BitIndex : 0;
 	}
 	static FORCEINLINE uint8 CountLeadingZeros8(uint8 Value)
 	{
-		unsigned long Log2;
-		_BitScanReverse(&Log2, (uint32(Value) << 1) | 1);
-		return uint8(8 - Log2);
+		unsigned long BitIndex;
+		_BitScanReverse(&BitIndex, uint32(Value)*2 + 1);
+		return uint8(8 - BitIndex);
 	}
 
 	static FORCEINLINE uint32 CountTrailingZeros(uint32 Value)
 	{
-		if (Value == 0)
-		{
-			return 32;
-		}
+		// return 32 if value was 0
 		unsigned long BitIndex;	// 0-based, where the LSB is 0 and MSB is 31
-		_BitScanForward( &BitIndex, Value );	// Scans from LSB to MSB
-		return BitIndex;
+		return _BitScanForward( &BitIndex, Value ) ? BitIndex : 32;
 	}
 
 	static FORCEINLINE uint32 CeilLogTwo( uint32 Arg )
 	{
-		int32 Bitmask = ((int32)(CountLeadingZeros(Arg) << 26)) >> 31;
-		return (32 - CountLeadingZeros(Arg - 1)) & (~Bitmask);
+		// if Arg is 0, change it to 1 so that we return 0
+		Arg = Arg ? Arg : 1;
+		return 32 - CountLeadingZeros(Arg - 1);
 	}
 
 	static FORCEINLINE uint32 RoundUpToPowerOfTwo(uint32 Arg)
@@ -70,49 +63,49 @@ struct FMicrosoftPlatformMathBase : public TUnrealPlatformMathSSE4Base<FGenericP
 
 	static FORCEINLINE uint64 FloorLog2_64(uint64 Value)
 	{
-		unsigned long Log2;
-		long Mask = -long(_BitScanReverse64(&Log2, Value) != 0);
-		return Log2 & Mask;
+		unsigned long BitIndex;
+		return _BitScanReverse64(&BitIndex, Value) ? BitIndex : 0;
 	}
 
 	static FORCEINLINE uint64 CeilLogTwo64(uint64 Arg)
 	{
-		int64 Bitmask = ((int64)(CountLeadingZeros64(Arg) << 57)) >> 63;
-		return (64 - CountLeadingZeros64(Arg - 1)) & (~Bitmask);
+		// if Arg is 0, change it to 1 so that we return 0
+		Arg = Arg ? Arg : 1;
+		return 64 - CountLeadingZeros64(Arg - 1);
 	}
 
 	static FORCEINLINE uint64 CountLeadingZeros64(uint64 Value)
 	{
-		unsigned long Log2;
-		long Mask = -long(_BitScanReverse64(&Log2, Value) != 0);
-		return ((63 - Log2) & Mask) | (64 & ~Mask);
+		//https://godbolt.org/z/Ejh5G4vPK	
+		// return 64 if value if was 0
+		unsigned long BitIndex;
+		if ( ! _BitScanReverse64(&BitIndex, Value) ) BitIndex = -1;
+        return 63 - BitIndex;
 	}
 
 	static FORCEINLINE uint64 CountTrailingZeros64(uint64 Value)
 	{
-		if (Value == 0)
-		{
-			return 64;
-		}
-		unsigned long BitIndex;	// 0-based, where the LSB is 0 and MSB is 31
-		_BitScanForward64( &BitIndex, Value );	// Scans from LSB to MSB
-		return BitIndex;
+		// return 64 if Value is 0
+		unsigned long BitIndex;	// 0-based, where the LSB is 0 and MSB is 63
+		return _BitScanForward64( &BitIndex, Value ) ? BitIndex : 64;
 	}
 
 	static FORCEINLINE uint32 CountLeadingZeros(uint32 Value)
 	{
-		unsigned long Log2;
-		_BitScanReverse64(&Log2, (uint64(Value) << 1) | 1);
-		return 32 - Log2;
+		// return 32 if value is zero
+		unsigned long BitIndex;
+		_BitScanReverse64(&BitIndex, uint64(Value)*2 + 1);
+		return 32 - BitIndex;
 	}
 
 #else // 32-bit
 
 	static FORCEINLINE uint32 CountLeadingZeros(uint32 Value)
 	{
-		unsigned long Log2;
-		long Mask = -long(_BitScanReverse(&Log2, Value) != 0);
-		return ((31 - Log2) & Mask) | (32 & ~Mask);
+		// return 32 if value is zero
+		unsigned long BitIndex;
+		if ( ! _BitScanReverse(&BitIndex, Value) ) BitIndex = -1;
+        return 31 - BitIndex;
 	}
 
 #endif
