Followup based on code review feedback.
* avoid using class ClassName*
* CanBeUsedOnMultipleThreads() return true on FAutomationTestOutputDevice to make it unbuffered
* Add lock for OutputDevices and FAutomationTestBase

#rnx

diff --git a/Engine/Source/Runtime/Core/Private/Misc/AutomationTest.cpp b/Engine/Source/Runtime/Core/Private/Misc/AutomationTest.cpp
--- a/Engine/Source/Runtime/Core/Private/Misc/AutomationTest.cpp
+++ b/Engine/Source/Runtime/Core/Private/Misc/AutomationTest.cpp
@@ -91,9 +91,10 @@ void FAutomationTestFramework::FAutomationTestOutputDevice::Serialize( const TCH
 	}
 
 	// Ensure there's a valid unit test associated with the context
-	class FAutomationTestBase* const LocalCurTest = CurTest.load(std::memory_order_relaxed);
+	FAutomationTestBase* const LocalCurTest = CurTest.load(std::memory_order_relaxed);
 	if (LocalCurTest)
 	{
+		FScopeLock Lock(&ActionCS);
 		bool CaptureLog = !LocalCurTest->SuppressLogs()
 			&& (Verbosity == ELogVerbosity::Error || Verbosity == ELogVerbosity::Warning || Verbosity == ELogVerbosity::Display);
 
@@ -154,9 +155,10 @@ void FAutomationTestFramework::FAutomationTestMessageFilter::Serialize(const TCH
 {
 	// Prevent null dereference if logging happens in async tasks while changing DestinationContext
 	FFeedbackContext* const LocalDestinationContext = DestinationContext.load(std::memory_order_relaxed);
-	class FAutomationTestBase* const LocalCurTest = CurTest.load(std::memory_order_relaxed);
+	FAutomationTestBase* const LocalCurTest = CurTest.load(std::memory_order_relaxed);
 	if (LocalDestinationContext)
 	{
+		FScopeLock Lock(&ActionCS);
 		if ((Verbosity == ELogVerbosity::Warning) || (Verbosity == ELogVerbosity::Error))
 		{
 			if (LocalCurTest->IsExpectedError(FString(V)))
@@ -180,7 +182,7 @@ FString FAutomationTestFramework::GetUserAutomationDirectory() const
 	return FString(FPlatformProcess::UserDir()) + DefaultAutomationSubFolder;
 }
 
-bool FAutomationTestFramework::RegisterAutomationTest( const FString& InTestNameToRegister, class FAutomationTestBase* InTestToRegister )
+bool FAutomationTestFramework::RegisterAutomationTest( const FString& InTestNameToRegister, FAutomationTestBase* InTestToRegister )
 {
 	const bool bAlreadyRegistered = AutomationTestClassNameToInstanceMap.Contains( InTestNameToRegister );
 	if ( !bAlreadyRegistered )
@@ -1028,6 +1030,7 @@ void FAutomationTestBase::AddError(const FString& InError, int32 StackOffset)
 {
 	if( !IsExpectedError(InError))
 	{
+		FScopeLock Lock(&ActionCS);
 		ExecutionInfo.AddEvent(FAutomationEvent(EAutomationEventType::Error, InError), StackOffset + 1);
 	}
 }
@@ -1044,6 +1047,7 @@ void FAutomationTestBase::AddErrorS(const FString& InError, const FString& InFil
 {
 	if ( !IsExpectedError(InError))
 	{
+		FScopeLock Lock(&ActionCS);
 		//ExecutionInfo.AddEvent(FAutomationEvent(EAutomationEventType::Error, InError, ExecutionInfo.GetContext(), InFilename, InLineNumber));
 	}
 }
@@ -1052,6 +1056,7 @@ void FAutomationTestBase::AddWarningS(const FString& InWarning, const FString& I
 {
 	if ( !IsExpectedError(InWarning))
 	{
+		FScopeLock Lock(&ActionCS);
 		//ExecutionInfo.AddEvent(FAutomationEvent(EAutomationEventType::Warning, InWarning, ExecutionInfo.GetContext(), InFilename, InLineNumber));
 	}
 }
@@ -1060,27 +1065,32 @@ void FAutomationTestBase::AddWarning( const FString& InWarning, int32 StackOffse
 {
 	if ( !IsExpectedError(InWarning))
 	{
+		FScopeLock Lock(&ActionCS);
 		ExecutionInfo.AddEvent(FAutomationEvent(EAutomationEventType::Warning, InWarning), StackOffset + 1);
 	}
 }
 
 void FAutomationTestBase::AddInfo( const FString& InLogItem, int32 StackOffset )
 {
+	FScopeLock Lock(&ActionCS);
 	ExecutionInfo.AddEvent(FAutomationEvent(EAutomationEventType::Info, InLogItem), StackOffset + 1);
 }
 
 void FAutomationTestBase::AddAnalyticsItem(const FString& InAnalyticsItem)
 {
+	FScopeLock Lock(&ActionCS);
 	ExecutionInfo.AnalyticsItems.Add(InAnalyticsItem);
 }
 
 void FAutomationTestBase::AddTelemetryData(const FString& DataPoint, double Measurement, const FString& Context)
 {
+	FScopeLock Lock(&ActionCS);
 	ExecutionInfo.TelemetryItems.Add(FAutomationTelemetryData(DataPoint, Measurement, Context));
 }
 
 void FAutomationTestBase::AddTelemetryData(const TMap <FString, double>& ValuePairs, const FString& Context)
 {
+	FScopeLock Lock(&ActionCS);
 	for (const TPair<FString, double>& Item : ValuePairs)
 	{
 		ExecutionInfo.TelemetryItems.Add(FAutomationTelemetryData(Item.Key, Item.Value, Context));
@@ -1094,6 +1104,7 @@ void FAutomationTestBase::SetTelemetryStorage(const FString& StorageName)
 
 void FAutomationTestBase::AddEvent(const FAutomationEvent& InEvent, int32 StackOffset)
 {
+	FScopeLock Lock(&ActionCS);
 	ExecutionInfo.AddEvent(InEvent, StackOffset + 1);
 }
 
@@ -1104,6 +1115,7 @@ bool FAutomationTestBase::HasAnyErrors() const
 
 bool FAutomationTestBase::HasMetExpectedErrors()
 {
+	FScopeLock Lock(&ActionCS);
 	bool HasMetAllExpectedErrors = true;
 
 	for (auto& EError : ExpectedErrors)
diff --git a/Engine/Source/Runtime/Core/Public/Misc/AutomationTest.h b/Engine/Source/Runtime/Core/Public/Misc/AutomationTest.h
--- a/Engine/Source/Runtime/Core/Public/Misc/AutomationTest.h
+++ b/Engine/Source/Runtime/Core/Public/Misc/AutomationTest.h
@@ -16,7 +16,7 @@
 #include "HAL/PlatformTime.h"
 #include "HAL/ThreadSafeBool.h" 
 #include "HAL/PlatformStackWalk.h"
-#include "GenericPlatform/GenericPlatformStackWalk.h" 
+#include "GenericPlatform/GenericPlatformStackWalk.h"
 #include "Containers/Queue.h"
 #include "Misc/CString.h"
 #include "Misc/FeedbackContext.h"
@@ -27,6 +27,7 @@
 #include "Math/Color.h"
 #include "Math/Rotator.h"
 #include "HAL/PlatformProcess.h"
+#include "HAL/CriticalSection.h"
 #include "Misc/AutomationEvent.h"
 #include "Internationalization/Regex.h"
 #include <atomic>
@@ -809,7 +810,7 @@ public:
 	 * @return	true if the test was successfully registered; false if a test was already registered under the same
 	 *			name as before
 	 */
-	bool RegisterAutomationTest( const FString& InTestNameToRegister, class FAutomationTestBase* InTestToRegister );
+	bool RegisterAutomationTest( const FString& InTestNameToRegister, FAutomationTestBase* InTestToRegister );
 
 	/**
 	 * Unregister a automation test with the provided name from the framework.
@@ -989,20 +990,33 @@ private:
 		 */
 		virtual void Serialize( const TCHAR* V, ELogVerbosity::Type Verbosity, const class FName& Category ) override;
 
+		/**
+		 * FOutputDevice interface
+		 *
+		 * Make it unbuffered by returning true
+		 */
+		virtual bool CanBeUsedOnMultipleThreads() const override
+		{
+			return true;
+		}
+
 		/**
 		 * Set the automation test associated with the output device. The automation test is where all warnings, errors, etc.
 		 * will be routed to.
 		 *
 		 * @param	InAutomationTest	Automation test to associate with the output device.
 		 */
-		void SetCurrentAutomationTest( class FAutomationTestBase* InAutomationTest )
+		void SetCurrentAutomationTest( FAutomationTestBase* InAutomationTest )
 		{
 			CurTest = InAutomationTest;
 		}
 
 	private:
 		/** Associated automation test; all warnings, errors, etc. are routed to the automation test to track */
-		std::atomic<class FAutomationTestBase*>CurTest;
+		std::atomic<FAutomationTestBase*>CurTest;
+
+		/** Critical section */
+		FCriticalSection ActionCS;
 
 		/** Tests that we've logged the failure cause when an error is involved */
 		TSet<FAutomationTestBase*> LoggedFailureCause;
@@ -1030,6 +1044,16 @@ private:
 		  */
 		 virtual void Serialize(const TCHAR* V, ELogVerbosity::Type Verbosity, const class FName& Category) override;
 
+		 /**
+		  * FOutputDevice interface
+		  *
+		  * Make it unbuffered by returning true
+		  */
+		 virtual bool CanBeUsedOnMultipleThreads() const override
+		 {
+			 return true;
+		 }
+
 		 /**
 		  * Set the automation test associated with the feedback context. The automation test is what will be used
 		  * to determine if a given warning or error is expected and thus should not be treated as a warning or error
@@ -1037,7 +1061,7 @@ private:
 		  *
 		  * @param	InAutomationTest	Automation test to associate with the feedback context.
 		  */
-		 void SetCurrentAutomationTest(class FAutomationTestBase* InAutomationTest)
+		 void SetCurrentAutomationTest(FAutomationTestBase* InAutomationTest)
 		 {
 			 CurTest = InAutomationTest;
 		 }
@@ -1054,8 +1078,9 @@ private:
 		 }
 
 	 private:
-		 std::atomic<class FAutomationTestBase*> CurTest;
+		 std::atomic<FAutomationTestBase*> CurTest;
 		 std::atomic<FFeedbackContext*> DestinationContext;
+		 FCriticalSection ActionCS;;
 	 };
 
 	friend class FAutomationTestOutputDevice;
@@ -1108,7 +1133,7 @@ private:
 	FFeedbackContext* OriginalGWarn = nullptr;
 
 	/** Mapping of automation test names to their respective object instances */
-	TMap<FString, class FAutomationTestBase*> AutomationTestClassNameToInstanceMap;
+	TMap<FString, FAutomationTestBase*> AutomationTestClassNameToInstanceMap;
 
 	/** Queue of deferred commands */
 	TQueue< TSharedPtr<IAutomationLatentCommand> > LatentCommands;
@@ -1797,6 +1822,8 @@ private:
 	/* Errors to be expected while processing this test.*/
 	TArray< FAutomationExpectedError> ExpectedErrors;
 
+	/** Critical section lock */
+	FCriticalSection ActionCS;
 };
 
 class CORE_API FBDDAutomationTestBase : public FAutomationTestBase
