Added std::initializer_list constructor to TSortedMap, for parity with TMap

diff --git a/Engine/Source/Runtime/Core/Public/Containers/SortedMap.h b/Engine/Source/Runtime/Core/Public/Containers/SortedMap.h
--- a/Engine/Source/Runtime/Core/Public/Containers/SortedMap.h
+++ b/Engine/Source/Runtime/Core/Public/Containers/SortedMap.h
@@ -48,6 +48,16 @@ public:
 	{
 	}
 
+	/** Constructor which gets its elements from a native initializer list */
+	TSortedMap(std::initializer_list<TPairInitializer<const KeyType&, const ValueType&>> InitList)
+	{
+		this->Reserve((int32)InitList.size());
+		for (const TPairInitializer<const KeyType&, const ValueType&>& Element : InitList)
+		{
+			this->Add(Element.Key, Element.Value);
+		}
+	}
+
 	/** Assignment operator for moving elements from a TSortedMap with a different ArrayAllocator. */
 	template<typename OtherArrayAllocator>
 	TSortedMap& operator=(TSortedMap<KeyType, ValueType, OtherArrayAllocator, SortPredicate>&& Other)
@@ -64,6 +74,17 @@ public:
 		return *this;
 	}
 
+	/** Assignment operator which gets its elements from a native initializer list */
+	TSortedMap& operator=(std::initializer_list<TPairInitializer<const KeyType&, const ValueType&>> InitList)
+	{
+		this->Empty((int32)InitList.size());
+		for (const TPairInitializer<const KeyType&, const ValueType&>& Element : InitList)
+		{
+			this->Add(Element.Key, Element.Value);
+		}
+		return *this;
+	}
+
 	/** Equality operator. This is efficient because pairs are always sorted. */
 	FORCEINLINE bool operator==(const TSortedMap& Other) const
 	{
