Temporary fix for Live Coding for requested module loads that couldn't immediately be completed.  (Leaving the OriginalFilename blank)

#rnx

diff --git a/Engine/Source/Runtime/Core/Private/Modules/ModuleManager.cpp b/Engine/Source/Runtime/Core/Private/Modules/ModuleManager.cpp
--- a/Engine/Source/Runtime/Core/Private/Modules/ModuleManager.cpp
+++ b/Engine/Source/Runtime/Core/Private/Modules/ModuleManager.cpp
@@ -990,7 +990,7 @@ bool FModuleManager::HasAnyOverridenModuleFilename() const
 	for (const TPair<FName, ModuleInfoRef>& ModuleIt : Modules)
 	{
 		const FModuleInfo& CurModule = *ModuleIt.Value;
-		if(CurModule.Filename != CurModule.OriginalFilename)
+		if(!CurModule.OriginalFilename.IsEmpty() && CurModule.Filename != CurModule.OriginalFilename)
 		{
 			return true;
 		}
