Re-enable the removal of the legacy math macros from Core, removed in CL# 19494088.
A couple of other macro fixes in Core.

diff --git a/Engine/Source/Runtime/Core/Core.Build.cs b/Engine/Source/Runtime/Core/Core.Build.cs
--- a/Engine/Source/Runtime/Core/Core.Build.cs
+++ b/Engine/Source/Runtime/Core/Core.Build.cs
@@ -313,6 +313,8 @@ public class Core : ModuleRules
 		PrivateDefinitions.Add("PLATFORM_COMPILER_OPTIMIZATION_PG=" + (Target.bPGOOptimize ? "1" : "0"));
 		PrivateDefinitions.Add("PLATFORM_COMPILER_OPTIMIZATION_PG_PROFILING=" + (Target.bPGOProfile ? "1" : "0"));
 
+		PrivateDefinitions.Add("UE_DEFINE_LEGACY_MATH_CONSTANT_MACRO_NAMES=0");
+
 		UnsafeTypeCastWarningLevel = WarningLevel.Warning;
 	}
 
diff --git a/Engine/Source/Runtime/Core/Private/Math/Color.cpp b/Engine/Source/Runtime/Core/Private/Math/Color.cpp
--- a/Engine/Source/Runtime/Core/Private/Math/Color.cpp
+++ b/Engine/Source/Runtime/Core/Private/Math/Color.cpp
@@ -530,7 +530,7 @@ FColor FColor::MakeRedToGreenColorFromScalar(float Scalar)
 
 void ComputeAndFixedColorAndIntensity(const FLinearColor& InLinearColor,FColor& OutColor,float& OutIntensity)
 {
-	float MaxComponent = FMath::Max(DELTA,FMath::Max(InLinearColor.R,FMath::Max(InLinearColor.G,InLinearColor.B)));
+	float MaxComponent = FMath::Max(UE_DELTA,FMath::Max(InLinearColor.R,FMath::Max(InLinearColor.G,InLinearColor.B)));
 	OutColor = ( InLinearColor / MaxComponent ).ToFColor(true);
 	OutIntensity = MaxComponent;
 }
diff --git a/Engine/Source/Runtime/Core/Public/Math/TransformCalculus2D.h b/Engine/Source/Runtime/Core/Public/Math/TransformCalculus2D.h
--- a/Engine/Source/Runtime/Core/Public/Math/TransformCalculus2D.h
+++ b/Engine/Source/Runtime/Core/Public/Math/TransformCalculus2D.h
@@ -505,10 +505,10 @@ public:
 		T E, F, G, H;
 		RHS.GetMatrix(E, F, G, H);
 		return
-			FMath::IsNearlyEqual(A, E, KINDA_SMALL_NUMBER) &&
-			FMath::IsNearlyEqual(B, F, KINDA_SMALL_NUMBER) &&
-			FMath::IsNearlyEqual(C, G, KINDA_SMALL_NUMBER) &&
-			FMath::IsNearlyEqual(D, H, KINDA_SMALL_NUMBER);
+			FMath::IsNearlyEqual(A, E, UE_KINDA_SMALL_NUMBER) &&
+			FMath::IsNearlyEqual(B, F, UE_KINDA_SMALL_NUMBER) &&
+			FMath::IsNearlyEqual(C, G, UE_KINDA_SMALL_NUMBER) &&
+			FMath::IsNearlyEqual(D, H, UE_KINDA_SMALL_NUMBER);
 	}
 
 	/** Inequality. */
@@ -573,7 +573,7 @@ public:
 			&& M[1][0] == 0.0f && M[1][1] == 1.0f;
 	}
 
-	bool IsNearlyIdentity(T ErrorTolerance = KINDA_SMALL_NUMBER) const
+	bool IsNearlyIdentity(T ErrorTolerance = UE_KINDA_SMALL_NUMBER) const
 	{
 		return
 			FMath::IsNearlyEqual(M[0][0], 1.0f, ErrorTolerance) &&
