Use TSortedMap for named format arguments

These are typically small data sets, so TSortedMap gives better performance

diff --git a/Engine/Source/Runtime/Core/Public/Internationalization/LocKeyFuncs.h b/Engine/Source/Runtime/Core/Public/Internationalization/LocKeyFuncs.h
--- a/Engine/Source/Runtime/Core/Public/Internationalization/LocKeyFuncs.h
+++ b/Engine/Source/Runtime/Core/Public/Internationalization/LocKeyFuncs.h
@@ -190,3 +190,12 @@ struct FLocKeyMultiMapFuncs : BaseKeyFuncs<ValueType, FString, /*bInAllowDuplica
 		return FLocKey::ProduceHash(Key);
 	}
 };
+
+/** Case sensitive comparison function for TSortedMap */
+struct FLocKeySortedMapLess
+{
+	FORCEINLINE bool operator()(const FString& A, const FString& B) const
+	{
+		return A.Compare(B, ESearchCase::CaseSensitive) < 0;
+	}
+};
diff --git a/Engine/Source/Runtime/Core/Public/Internationalization/Text.h b/Engine/Source/Runtime/Core/Public/Internationalization/Text.h
--- a/Engine/Source/Runtime/Core/Public/Internationalization/Text.h
+++ b/Engine/Source/Runtime/Core/Public/Internationalization/Text.h
@@ -8,7 +8,7 @@
 #include "Templates/UnrealTypeTraits.h"
 #include "Containers/Array.h"
 #include "Containers/UnrealString.h"
-#include "Containers/Map.h"
+#include "Containers/SortedMap.h"
 #include "Containers/EnumAsByte.h"
 #include "Templates/SharedPointer.h"
 #include "Internationalization/TextKey.h"
@@ -146,7 +146,7 @@ namespace EFormatArgumentType
 	};
 }
 
-typedef TMap<FString, FFormatArgumentValue, FDefaultSetAllocator, FLocKeyMapFuncs<FFormatArgumentValue>> FFormatNamedArguments;
+typedef TSortedMap<FString, FFormatArgumentValue, FDefaultAllocator, FLocKeySortedMapLess> FFormatNamedArguments;
 typedef TArray<FFormatArgumentValue> FFormatOrderedArguments;
 
 typedef TSharedRef<FTextFormatPatternDefinition, ESPMode::ThreadSafe> FTextFormatPatternDefinitionRef;
