Some optimizations to FTextFormatter

 - Use TSortedMap for the argument modifier parameters, as these are small sets of data.
 - Use a RW lock to avoid contention on a rarely changing container.
 - Avoid allocating a unique instance in FTextFormat.

diff --git a/Engine/Source/Runtime/Core/Private/Internationalization/TextFormatArgumentModifier.cpp b/Engine/Source/Runtime/Core/Private/Internationalization/TextFormatArgumentModifier.cpp
--- a/Engine/Source/Runtime/Core/Private/Internationalization/TextFormatArgumentModifier.cpp
+++ b/Engine/Source/Runtime/Core/Private/Internationalization/TextFormatArgumentModifier.cpp
@@ -6,7 +6,7 @@
 #include "Internationalization/Internationalization.h"
 #include "Internationalization/TextFormatter.h"
 
-bool ITextFormatArgumentModifier::ParseKeyValueArgs(const FTextFormatString& InArgsString, TMap<FTextFormatString, FTextFormatString>& OutArgKeyValues, const TCHAR InValueSeparator, const TCHAR InArgSeparator)
+bool ITextFormatArgumentModifier::ParseKeyValueArgs(const FTextFormatString& InArgsString, TSortedMap<FTextFormatString, FTextFormatString>& OutArgKeyValues, const TCHAR InValueSeparator, const TCHAR InArgSeparator)
 {
 	const TCHAR* BufferPtr = InArgsString.StringPtr;
 	const TCHAR* BufferEnd = InArgsString.StringPtr + InArgsString.StringLen;
@@ -166,14 +166,14 @@ bool ITextFormatArgumentModifier::ParseValueArgs(const FTextFormatString& InArgs
 
 TSharedPtr<ITextFormatArgumentModifier> FTextFormatArgumentModifier_PluralForm::Create(const ETextPluralType InPluralType, const FTextFormatString& InArgsString, const FTextFormatPatternDefinitionConstRef& InPatternDef)
 {
-	TMap<FTextFormatString, FTextFormatString> ArgKeyValues;
+	TSortedMap<FTextFormatString, FTextFormatString> ArgKeyValues;
 	if (ParseKeyValueArgs(InArgsString, ArgKeyValues))
 	{
 		int32 LongestPluralFormStringLen = 0;
 		bool bDoPluralFormsUseFormatArgs = false;
 
 		// Plural forms may contain format markers, so pre-compile all the variants now so that Evaluate doesn't have to (this also lets us validate the plural form strings and fail if they're not correct)
-		TMap<FTextFormatString, FTextFormat> PluralForms;
+		TSortedMap<FTextFormatString, FTextFormat> PluralForms;
 		PluralForms.Reserve(ArgKeyValues.Num());
 		for (const auto& Pair : ArgKeyValues)
 		{
@@ -339,7 +339,7 @@ void FTextFormatArgumentModifier_PluralForm::EstimateLength(int32& OutLength, bo
 	OutUsesFormatArgs = bDoPluralFormsUseFormatArgs;
 }
 
-FTextFormatArgumentModifier_PluralForm::FTextFormatArgumentModifier_PluralForm(const ETextPluralType InPluralType, const TMap<FTextFormatString, FTextFormat>& InPluralForms, const int32 InLongestPluralFormStringLen, const bool InDoPluralFormsUseFormatArgs)
+FTextFormatArgumentModifier_PluralForm::FTextFormatArgumentModifier_PluralForm(const ETextPluralType InPluralType, const TSortedMap<FTextFormatString, FTextFormat>& InPluralForms, const int32 InLongestPluralFormStringLen, const bool InDoPluralFormsUseFormatArgs)
 	: PluralType(InPluralType)
 	, LongestPluralFormStringLen(InLongestPluralFormStringLen)
 	, bDoPluralFormsUseFormatArgs(InDoPluralFormsUseFormatArgs)
diff --git a/Engine/Source/Runtime/Core/Private/Internationalization/TextFormatArgumentModifier.h b/Engine/Source/Runtime/Core/Private/Internationalization/TextFormatArgumentModifier.h
--- a/Engine/Source/Runtime/Core/Private/Internationalization/TextFormatArgumentModifier.h
+++ b/Engine/Source/Runtime/Core/Private/Internationalization/TextFormatArgumentModifier.h
@@ -9,10 +9,6 @@
 
 struct FPrivateTextFormatArguments;
 
-/// @cond DOXYGEN_WARNINGS
-template<typename KeyType, typename ValueType, typename SetAllocator, typename KeyFuncs > class TMap;
-/// @endcond
-
 /**
  * Plural form argument modifier.
  * Takes a set of key->value arguments, where the key is a valid plural form identifier, and the value is an optionally quoted string that may contain format markers.
@@ -33,7 +29,7 @@ public:
 	virtual void EstimateLength(int32& OutLength, bool& OutUsesFormatArgs) const override;
 
 private:
-	FTextFormatArgumentModifier_PluralForm(const ETextPluralType InPluralType, const TMap<FTextFormatString, FTextFormat>& InPluralForms, const int32 InLongestPluralFormStringLen, const bool InDoPluralFormsUseFormatArgs);
+	FTextFormatArgumentModifier_PluralForm(const ETextPluralType InPluralType, const TSortedMap<FTextFormatString, FTextFormat>& InPluralForms, const int32 InLongestPluralFormStringLen, const bool InDoPluralFormsUseFormatArgs);
 
 	ETextPluralType PluralType;
 	int32 LongestPluralFormStringLen;
diff --git a/Engine/Source/Runtime/Core/Private/Internationalization/TextFormatter.cpp b/Engine/Source/Runtime/Core/Private/Internationalization/TextFormatter.cpp
--- a/Engine/Source/Runtime/Core/Private/Internationalization/TextFormatter.cpp
+++ b/Engine/Source/Runtime/Core/Private/Internationalization/TextFormatter.cpp
@@ -2,6 +2,7 @@
 
 #include "Internationalization/TextFormatter.h"
 #include "Misc/ScopeLock.h"
+#include "Misc/ScopeRWLock.h"
 #include "Internationalization/TextFormatArgumentModifier.h"
 #include "Internationalization/TextHistory.h"
 #include "Misc/ExpressionParser.h"
@@ -336,6 +337,11 @@ struct FPrivateTextFormatArguments
 class FTextFormatData
 {
 public:
+	/**
+	 * Get the common "empty" instance to use for a default constructed FTextFormat.
+	 */
+	static TSharedRef<FTextFormatData, ESPMode::ThreadSafe> GetSharedEmptyInstance();
+
 	/**
 	 * Construct an instance from an FText.
 	 * The text will be immediately compiled. 
@@ -536,22 +542,22 @@ private:
 
 
 FTextFormat::FTextFormat()
-	: TextFormatData(new FTextFormatData(FText(), FTextFormatPatternDefinition::GetDefault()))
+	: TextFormatData(FTextFormatData::GetSharedEmptyInstance())
 {
 }
 
 FTextFormat::FTextFormat(const FText& InText)
-	: TextFormatData(new FTextFormatData(CopyTemp(InText), FTextFormatPatternDefinition::GetDefault()))
+	: TextFormatData(MakeShared<FTextFormatData, ESPMode::ThreadSafe>(CopyTemp(InText), FTextFormatPatternDefinition::GetDefault()))
 {
 }
 
 FTextFormat::FTextFormat(const FText& InText, FTextFormatPatternDefinitionConstRef InCustomPatternDef)
-	: TextFormatData(new FTextFormatData(CopyTemp(InText), MoveTemp(InCustomPatternDef)))
+	: TextFormatData(MakeShared<FTextFormatData, ESPMode::ThreadSafe>(CopyTemp(InText), MoveTemp(InCustomPatternDef)))
 {
 }
 
 FTextFormat::FTextFormat(FString&& InString, FTextFormatPatternDefinitionConstRef InCustomPatternDef)
-	: TextFormatData(new FTextFormatData(MoveTemp(InString), MoveTemp(InCustomPatternDef)))
+	: TextFormatData(MakeShared<FTextFormatData, ESPMode::ThreadSafe>(MoveTemp(InString), MoveTemp(InCustomPatternDef)))
 {
 }
 
@@ -616,6 +622,12 @@ void FTextFormat::GetFormatArgumentNames(TArray<FString>& OutArgumentNames) cons
 }
 
 
+TSharedRef<FTextFormatData, ESPMode::ThreadSafe> FTextFormatData::GetSharedEmptyInstance()
+{
+	static const TSharedRef<FTextFormatData> EmptyInstance = MakeShared<FTextFormatData, ESPMode::ThreadSafe>(FText(), FTextFormatPatternDefinition::GetDefault());
+	return EmptyInstance;
+}
+
 FTextFormatData::FTextFormatData(FText&& InText, FTextFormatPatternDefinitionConstRef InPatternDef)
 	: SourceType(ESourceType::Text)
 	, PatternDef(MoveTemp(InPatternDef))
@@ -930,19 +942,19 @@ FTextFormatter& FTextFormatter::Get()
 
 void FTextFormatter::RegisterTextArgumentModifier(const FTextFormatString& InKeyword, FCompileTextArgumentModifierFuncPtr InCompileFunc)
 {
-	FScopeLock Lock(&TextArgumentModifiersCS);
+	FWriteScopeLock Lock(TextArgumentModifiersRW);
 	TextArgumentModifiers.Add(InKeyword, MoveTemp(InCompileFunc));
 }
 
 void FTextFormatter::UnregisterTextArgumentModifier(const FTextFormatString& InKeyword)
 {
-	FScopeLock Lock(&TextArgumentModifiersCS);
+	FWriteScopeLock Lock(TextArgumentModifiersRW);
 	TextArgumentModifiers.Remove(InKeyword);
 }
 
 FTextFormatter::FCompileTextArgumentModifierFuncPtr FTextFormatter::FindTextArgumentModifier(const FTextFormatString& InKeyword) const
 {
-	FScopeLock Lock(&TextArgumentModifiersCS);
+	FReadScopeLock Lock(TextArgumentModifiersRW);
 	return TextArgumentModifiers.FindRef(InKeyword);
 }
 
diff --git a/Engine/Source/Runtime/Core/Public/Internationalization/ITextFormatArgumentModifier.h b/Engine/Source/Runtime/Core/Public/Internationalization/ITextFormatArgumentModifier.h
--- a/Engine/Source/Runtime/Core/Public/Internationalization/ITextFormatArgumentModifier.h
+++ b/Engine/Source/Runtime/Core/Public/Internationalization/ITextFormatArgumentModifier.h
@@ -29,7 +29,6 @@ public:
 	FTextFormatString()
 		: StringPtr(TEXT(""))
 		, StringLen(0)
-		, StringHash(0)
 		, InternalString()
 	{
 	}
@@ -38,24 +37,20 @@ public:
 	FTextFormatString(FString InStr)
 		: StringPtr(nullptr)
 		, StringLen(0)
-		, StringHash(0)
 		, InternalString(MoveTemp(InStr))
 	{
 		StringPtr = *InternalString;
 		StringLen = InternalString.Len();
-		StringHash = FCrc::MemCrc32(StringPtr, sizeof(TCHAR) * StringLen);
 	}
 
 	/** Construct from the given string (takes a copy, expected to be null terminated) */
 	FTextFormatString(const TCHAR* InStr)
 		: StringPtr(nullptr)
 		, StringLen(0)
-		, StringHash(0)
 		, InternalString(InStr)
 	{
 		StringPtr = *InternalString;
 		StringLen = InternalString.Len();
-		StringHash = FCrc::MemCrc32(StringPtr, sizeof(TCHAR) * StringLen);
 	}
 
 	/** Construct from the given string (takes a reference, expected to be null terminated) */
@@ -73,7 +68,6 @@ public:
 	FTextFormatString(const FTextFormatString& Other)
 		: StringPtr(Other.StringPtr)
 		, StringLen(Other.StringLen)
-		, StringHash(Other.StringHash)
 		, InternalString(Other.InternalString)
 	{
 		if (Other.StringPtr == *Other.InternalString)
@@ -85,7 +79,6 @@ public:
 	FTextFormatString(FTextFormatString&& Other)
 		: StringPtr(Other.StringPtr)
 		, StringLen(Other.StringLen)
-		, StringHash(Other.StringHash)
 		, InternalString(MoveTemp(Other.InternalString))
 	{
 	}
@@ -96,7 +89,6 @@ public:
 		{
 			StringPtr = Other.StringPtr;
 			StringLen = Other.StringLen;
-			StringHash = Other.StringHash;
 			InternalString = Other.InternalString;
 
 			if (Other.StringPtr == *Other.InternalString)
@@ -113,7 +105,6 @@ public:
 		{
 			StringPtr = Other.StringPtr;
 			StringLen = Other.StringLen;
-			StringHash = Other.StringHash;
 			InternalString = MoveTemp(Other.InternalString);
 		}
 		return *this;
@@ -121,7 +112,7 @@ public:
 
 	friend inline uint32 GetTypeHash(const FTextFormatString& InStr)
 	{
-		return InStr.StringHash;
+		return FCrc::MemCrc32(InStr.StringPtr, sizeof(TCHAR) * InStr.StringLen);
 	}
 
 	/** The start of the string */
@@ -131,9 +122,6 @@ public:
 	int32 StringLen;
 
 private:
-	/** Cached hash of the string */
-	uint32 StringHash;
-
 	/** Internal copy if constructed from a string */
 	FString InternalString;
 
@@ -141,7 +129,6 @@ private:
 	FTextFormatString(const TCHAR* InStr, const int32 InLen)
 		: StringPtr(InStr)
 		, StringLen(InLen)
-		, StringHash(FCrc::MemCrc32(StringPtr, sizeof(TCHAR) * StringLen))
 		, InternalString()
 	{
 	}
@@ -159,6 +146,19 @@ inline bool operator!=(const FTextFormatString& LHS, const FTextFormatString& RH
 		|| FCString::Strncmp(LHS.StringPtr, RHS.StringPtr, LHS.StringLen) != 0;
 }
 
+inline bool operator<(const FTextFormatString& LHS, const FTextFormatString& RHS)
+{
+	if (LHS.StringLen < RHS.StringLen)
+	{
+		return true;
+	}
+	if (LHS.StringLen > RHS.StringLen)
+	{
+		return false;
+	}
+	return FCString::Strncmp(LHS.StringPtr, RHS.StringPtr, LHS.StringLen) < 0;
+}
+
 /**
  * Interface for a format argument modifier.
  */
@@ -182,7 +182,7 @@ public:
 
 protected:
 	/** Utility helper to parse out a list of key->value pair arguments. The keys are assumed to only contain valid identifier characters, and the values may be optionally quoted (parsed strings are sub-string references to within the source args string) */
-	CORE_API static bool ParseKeyValueArgs(const FTextFormatString& InArgsString, TMap<FTextFormatString, FTextFormatString>& OutArgKeyValues, const TCHAR InValueSeparator = TEXT('='), const TCHAR InArgSeparator = TEXT(','));
+	CORE_API static bool ParseKeyValueArgs(const FTextFormatString& InArgsString, TSortedMap<FTextFormatString, FTextFormatString>& OutArgKeyValues, const TCHAR InValueSeparator = TEXT('='), const TCHAR InArgSeparator = TEXT(','));
 
 	/** Utility helper to parse out a list of value arguments. The values may be optionally quoted (parsed strings are sub-string references to within the source args string) */
 	CORE_API static bool ParseValueArgs(const FTextFormatString& InArgsString, TArray<FTextFormatString>& OutArgValues, const TCHAR InArgSeparator = TEXT(','));
diff --git a/Engine/Source/Runtime/Core/Public/Internationalization/TextFormatter.h b/Engine/Source/Runtime/Core/Public/Internationalization/TextFormatter.h
--- a/Engine/Source/Runtime/Core/Public/Internationalization/TextFormatter.h
+++ b/Engine/Source/Runtime/Core/Public/Internationalization/TextFormatter.h
@@ -6,7 +6,7 @@
 #include "Containers/Array.h"
 #include "Containers/UnrealString.h"
 #include "Templates/Function.h"
-#include "Containers/Map.h"
+#include "Containers/SortedMap.h"
 #include "Internationalization/Text.h"
 #include "Misc/ExpressionParserTypes.h"
 #include "Internationalization/ITextFormatArgumentModifier.h"
@@ -103,8 +103,8 @@ private:
 	FTextFormatter();
 
 	/** Functions for constructing argument modifier data */
-	TMap<FTextFormatString, FCompileTextArgumentModifierFuncPtr> TextArgumentModifiers;
+	TSortedMap<FTextFormatString, FCompileTextArgumentModifierFuncPtr> TextArgumentModifiers;
 
-	/** Critical section protecting the argument modifiers map from being modified concurrently */
-	mutable FCriticalSection TextArgumentModifiersCS;
+	/** RW lock protecting the argument modifiers map from being modified concurrently */
+	mutable FRWLock TextArgumentModifiersRW;
 };
