Logging: Fixed a deadlock when the logging thread is in its flush loop during a panic

#rnx

diff --git a/Engine/Source/Runtime/Core/Private/Misc/OutputDeviceRedirector.cpp b/Engine/Source/Runtime/Core/Private/Misc/OutputDeviceRedirector.cpp
--- a/Engine/Source/Runtime/Core/Private/Misc/OutputDeviceRedirector.cpp
+++ b/Engine/Source/Runtime/Core/Private/Misc/OutputDeviceRedirector.cpp
@@ -370,6 +370,10 @@ void FOutputDeviceRedirectorState::ThreadLoop()
 				{
 					FlushBufferedLines();
 				}
+				else
+				{
+					break;
+				}
 			}
 			while (!BufferedLines.IsEmpty());
 			ThreadIdleEvents.Deplete([](FEvent* Event) { Event->Trigger(); });
