Add new Unit Tests

diff --git a/Engine/Source/Developer/LowLevelTestsRunner/Private/TestCommon/CoreUObjectUtilities.cpp b/Engine/Source/Developer/LowLevelTestsRunner/Private/TestCommon/CoreUObjectUtilities.cpp
--- a/Engine/Source/Developer/LowLevelTestsRunner/Private/TestCommon/CoreUObjectUtilities.cpp
+++ b/Engine/Source/Developer/LowLevelTestsRunner/Private/TestCommon/CoreUObjectUtilities.cpp
@@ -10,8 +10,17 @@
 #include "UObject/Package.h"
 #include "UObject/PackageResourceManager.h"
 
+bool CoreUObjectInitialised = false;
+
 void InitCoreUObject()
 {
+	if (CoreUObjectInitialised == true)
+		return;
+
+	//Initialize the PackageResourceManager, which is needed to load any (non-script) Packages. It is first used in ProcessNewlyLoadedObjects (due to the loading of asset references in Class Default Objects)
+	// It has to be intialized after the AssetRegistryModule; the editor implementations of PackageResourceManager relies on it
+	IPackageResourceManager::Initialize();
+
 	if (!GetTransientPackage())
 	{
 		FModuleManager::Get().AddExtraBinarySearchPaths();
@@ -22,13 +31,22 @@ void InitCoreUObject()
 		FCoreDelegates::OnInit.Broadcast();
 		ProcessNewlyLoadedUObjects();
 	}
+
+	CoreUObjectInitialised = true;
 }
 
 void CleanupCoreUObject()
 {
+	if (CoreUObjectInitialised == false)
+		return;
+
+	IPackageResourceManager::Shutdown();
+
 	FCoreDelegates::OnPreExit.Broadcast();
 	MALLOC_PROFILER(GMalloc->Exec(nullptr, TEXT("MPROF STOP"), *GLog); );
 	FCoreDelegates::OnExit.Broadcast();
+
+	CoreUObjectInitialised = false;
 }
 
 #endif // WITH_COREUOBJECT
\ No newline at end of file
diff --git a/Engine/Source/Developer/LowLevelTestsRunner/Private/TestCommon/CoreUtilities.cpp b/Engine/Source/Developer/LowLevelTestsRunner/Private/TestCommon/CoreUtilities.cpp
--- a/Engine/Source/Developer/LowLevelTestsRunner/Private/TestCommon/CoreUtilities.cpp
+++ b/Engine/Source/Developer/LowLevelTestsRunner/Private/TestCommon/CoreUtilities.cpp
@@ -102,16 +102,28 @@ void CleanupTaskGraph()
 	FTaskGraphInterface::Shutdown();
 }
 
+bool TaskGraphInitialised = false;
+
 void InitTaskGraphAndDependencies(bool MultiThreaded)
 {
+	if (TaskGraphInitialised == true)
+		return;
+
 	InitTaskGraph(MultiThreaded);
 	InitAllThreadPools(MultiThreaded);
+
+	TaskGraphInitialised = true;
 }
 
 void CleanupTaskGraphAndDependencies()
 {
+	if (TaskGraphInitialised == false)
+		return;
+
 	CleanupAllThreadPools();
 	CleanupTaskGraph();
+
+	TaskGraphInitialised = false;
 }
 
 void CleanupPlatform()
diff --git a/Engine/Source/Developer/LowLevelTestsRunner/Private/TestFixtures/CoreTestFixture.cpp b/Engine/Source/Developer/LowLevelTestsRunner/Private/TestFixtures/CoreTestFixture.cpp
new file mode 100644
--- /dev/null
+++ b/Engine/Source/Developer/LowLevelTestsRunner/Private/TestFixtures/CoreTestFixture.cpp
@@ -0,0 +1,17 @@
+// Copyright Epic Games, Inc. All Rights Reserved.
+
+#include "TestFixtures/CoreTestFixture.h"
+#include "TestCommon/CoreUtilities.h"
+#include "TestCommon/CoreUObjectUtilities.h"
+
+FCoreTestFixture::FCoreTestFixture()
+{
+	InitTaskGraphAndDependencies();
+	InitCoreUObject();
+}
+
+FCoreTestFixture::~FCoreTestFixture()
+{
+	//CleanupCoreUObject();
+	//CleanupTaskGraphAndDependencies();
+}
diff --git a/Engine/Source/Developer/LowLevelTestsRunner/Private/TestRunner.cpp b/Engine/Source/Developer/LowLevelTestsRunner/Private/TestRunner.cpp
--- a/Engine/Source/Developer/LowLevelTestsRunner/Private/TestRunner.cpp
+++ b/Engine/Source/Developer/LowLevelTestsRunner/Private/TestRunner.cpp
@@ -9,6 +9,7 @@
 
 #include "LowLevelTestModule.h"
 #include "TestCommon/CoreUtilities.h"
+#include "TestCommon/CoreUObjectUtilities.h"
 
 #include <set>
 
@@ -209,6 +210,9 @@ int RunTests(int argc, const char* argv[])
 		CatchArgv.Reset();
 	}
 
+	CleanupCoreUObject();
+	CleanupTaskGraphAndDependencies();
+
 	GlobalModuleTeardown();
 	UnloadBaseTestModule(LoadModuleName);
 
diff --git a/Engine/Source/Developer/LowLevelTestsRunner/Public/TestFixtures/CoreTestFixture.h b/Engine/Source/Developer/LowLevelTestsRunner/Public/TestFixtures/CoreTestFixture.h
new file mode 100644
--- /dev/null
+++ b/Engine/Source/Developer/LowLevelTestsRunner/Public/TestFixtures/CoreTestFixture.h
@@ -0,0 +1,12 @@
+// Copyright Epic Games, Inc. All Rights Reserved.
+
+#pragma once
+
+#include "TestHarness.h"
+
+class FCoreTestFixture : public FAutomationTestFixture
+{
+public:
+	FCoreTestFixture();	
+	~FCoreTestFixture();
+};
\ No newline at end of file
diff --git a/Engine/Source/Developer/LowLevelTestsRunner/Public/TestHarness.h b/Engine/Source/Developer/LowLevelTestsRunner/Public/TestHarness.h
--- a/Engine/Source/Developer/LowLevelTestsRunner/Public/TestHarness.h
+++ b/Engine/Source/Developer/LowLevelTestsRunner/Public/TestHarness.h
@@ -91,68 +91,146 @@ struct Catch::StringMaker<TTuple<KeyType, ValueType>>
 	}
 };
 
-// Tests if true and prints message otherwise
-#define TestTrue(a, b) \
-	INFO((a)); \
-	CHECK((b) == true); \
-
-// Quiet version of TestTrue
-#define TestTrue_Q(a) \
-	CHECK((a) == true); \
-
-// Tests if false and prints message otherwise
-#define TestFalse(a, b) \
-	INFO((a)); \
-	CHECK((b) != true); \
-
-// Quiet version of TestFalse
-#define TestFalse_Q(a) \
-	CHECK((a) != true); \
+/** Use with TEST_CASE_METHOD to quickly migrate automation tests to Catch. */
+class FAutomationTestFixture
+{
+public:
 
-// Tests if b == c and prints message otherwise
-#define TestEqual(a, b, c) \
-	INFO((a)); \
-	CHECK((b) == (c)); \
+	template <typename WhatType>
+	inline bool TestTrue(const WhatType What, const bool bValue)
+	{
+		CAPTURE(What);
+		CHECK(bValue);
+		return bValue;
+	}
 
-#define TestEqual_Conditional(a, b, c) \
-	INFO((a)); \
-	CHECKED_IF((b) == (c)) \
+	template <typename WhatType>
+	inline bool TestFalse(const WhatType What, const bool bValue)
+	{
+		CAPTURE(What);
+		CHECK(!bValue);
+		return !bValue;
+	}
 
-// Quiet version of TestEqual
-#define TestEqual_Q(a, b) \
-	CHECK((a) == (b)); \
+	template <typename WhatType, typename ActualType, typename ExpectedType>
+	inline bool TestEqual(const WhatType What, const ActualType& Actual, const ExpectedType& Expected)
+	{
+		CAPTURE(What);
+		CHECK(Actual == Expected);
+		return Actual == Expected;
+	}
 
-#define TestNotEqual(a, b, c) \
-	INFO((a)); \
-	CHECK((b) != (c)); \
+	template <typename WhatType, typename ActualType, typename ExpectedType>
+	inline bool TestNotEqual(const WhatType& What, const ActualType& Actual, const ExpectedType& Expected)
+	{
+		CAPTURE(What);
+		CHECK(Actual != Expected);
+		return Actual != Expected;
+	}
+	template <typename WhatType, typename ActualType>
+	inline bool TestNull(const WhatType& What, const ActualType& Actual)
+	{
+		CAPTURE(What);
+		CHECK(Actual == nullptr);
+		return Actual == nullptr;
+	}
+	template <typename WhatType, typename ActualType>
+	inline bool TestNotNull(const WhatType& What, const ActualType& Actual)
+	{
+		CAPTURE(What);
+		CHECK(Actual != nullptr);
+		return Actual != nullptr;
+	}
+	template <typename WhatType, typename ActualType>
+	inline bool TestValid(const WhatType& What, const ActualType& Actual)
+	{
+		CAPTURE(What);
+		CHECK(Actual.IsValid() == true);
+		return Actual.IsValid() == true;
+	}
+	template <typename WhatType, typename ActualType>
+	inline bool TestInvalid(const WhatType& What, const ActualType& Actual)
+	{
+		CAPTURE(What);
+		CHECK(Actual.IsValid() == false);
+		return Actual.IsValid() == false;
+	}
+	template <typename WhatType, typename ActualType>
+	inline bool Verify(const WhatType What, const ActualType& Actual)
+	{
+		CAPTURE(What);
+		CHECK(Actual == true);
+		return Actual == true;
+	}
+	template <typename WhatType>
+	inline void AddWarning(const WhatType& What)
+	{
+		CAPTURE(What);
+		return;
+	}
+	template <typename WhatType>
+	inline void AddError(const WhatType& What)
+	{
+		CAPTURE(What);
+		return;
+	}
+	
+};
 
-#define TestNull(a, b) \
-	INFO((a)); \
-	CHECK((b) == nullptr); \
+#define CHECK_EQUAL(Actual, Expected)\
+	CHECK(Actual == Expected);
 
-#define TestNotNull(a, b) \
-	INFO((a)); \
-	CHECK((b) != nullptr); \
+#define CHECK_NOT_EQUAL(Actual, Expected)\
+	CHECK(Actual != Expected);
 
-#define TestValid(a, b) \
-	INFO((a)); \
-	CHECK((b).IsValid() == true); \
+#define REQUIRE_EQUAL(Actual, Expected)\
+	REQUIRE(Actual == Expected);
 
-#define TestInvalid(a, b) \
-	INFO((a)); \
-	CHECK((b).IsValid() != true); \
+#define REQUIRE_NOT_EQUAL(Actual, Expected)\
+	REQUIRE(Actual != Expected);
 
-#define TestEqualString(a, b, c) \
-	INFO((a)); \
-	CHECK(FCString::Strcmp(ToCStr((b)), ToCStr((c))) == 0); \
 
-#define TestAddInfo(a) \
-	INFO((a)); \
 
-#define TestAddWarning(a) \
-	WARN((a)); \
+#define TEST_TRUE(What, Value)\
+{\
+	INFO(What);\
+	CHECK((Value)==true);\
+}
+#define TEST_FALSE(What, Value)\
+{\
+	INFO(What);\
+	CHECK((Value)==false);\
+}
+#define TEST_EQUAL(What, Actual, Expected)\
+{\
+	INFO(What);\
+	CHECK((Actual) == (Expected));\
+}
 
-#define TestAddError(a) \
-	FAIL_CHECK((a)); \
+#define TEST_NOT_EQUAL(What, Actual, Expected)\
+{\
+	INFO(What);\
+	CHECK((Actual) != (Expected));\
+}
+#define TEST_NULL(What, Value)\
+{\
+	INFO(What);\
+	CHECK((Value)==nullptr);\
+}
+#define TEST_NOT_NULL(What, Value)\
+{\
+	INFO(What);\
+	CHECK((Value)!=nullptr);\
+}
+#define TEST_VALID(What, Value)\
+{\
+	INFO(What);\
+	CHECK(Value.IsValid()==true);\
+}
+#define TEST_INVALID(What, Value)\
+{\
+	INFO(What);\
+	CHECK(Value.IsValid()==false);\
+}
 
 #endif // UE_ENABLE_TESTHARNESS_ENGINE_SUPPORT
\ No newline at end of file
diff --git a/Engine/Source/Runtime/Cbor/Tests/CborTests.cpp b/Engine/Source/Runtime/Cbor/Tests/CborTests.cpp
--- a/Engine/Source/Runtime/Cbor/Tests/CborTests.cpp
+++ b/Engine/Source/Runtime/Cbor/Tests/CborTests.cpp
@@ -322,7 +322,7 @@ TEST_CASE("CBor::FCbor::Smoke Test", "[CBor][Smoke]")
 	}
 }
 
-TEST_CASE("CBor::FCbor::ByteSwap::Performance", "[CBor][Perf]")
+TEST_CASE("CBor::FCbor::ByteSwap Performance Test", "[CBor][Perf]")
 {
 	//  NOTE: The most expensive value to write in CBOR is the double as it needs to swap 8 bytes all the time. Integers can be encoded in 1, 2, 4 or 8 bytes depending on their value, but not double.
 	//        String are not considered because they are written in UTF8. The test ensure this is not significantly longer to write/read while swapping the bytes.
diff --git a/Engine/Source/Runtime/Core/Public/Async/Async.h b/Engine/Source/Runtime/Core/Public/Async/Async.h
--- a/Engine/Source/Runtime/Core/Public/Async/Async.h
+++ b/Engine/Source/Runtime/Core/Public/Async/Async.h
@@ -347,6 +347,7 @@ auto Async(EAsyncExecution Execution, CallableType&& Callable, TUniqueFunction<v
 	case EAsyncExecution::ThreadPool:
 		if (FPlatformProcess::SupportsMultithreading())
 		{
+			check(GThreadPool != nullptr);
 			GThreadPool->AddQueuedWork(new TAsyncQueuedWork<ResultType>(MoveTemp(Function), MoveTemp(Promise)));
 		}
 		else
@@ -359,6 +360,7 @@ auto Async(EAsyncExecution Execution, CallableType&& Callable, TUniqueFunction<v
 	case EAsyncExecution::LargeThreadPool:
 		if (FPlatformProcess::SupportsMultithreading())
 		{
+			check(GLargeThreadPool != nullptr);
 			GLargeThreadPool->AddQueuedWork(new TAsyncQueuedWork<ResultType>(MoveTemp(Function), MoveTemp(Promise)));
 		}
 		else
diff --git a/Engine/Source/Runtime/Core/Tests/Algorithm/TopologicalSort.cpp b/Engine/Source/Runtime/Core/Tests/Algorithm/TopologicalSort.cpp
--- a/Engine/Source/Runtime/Core/Tests/Algorithm/TopologicalSort.cpp
+++ b/Engine/Source/Runtime/Core/Tests/Algorithm/TopologicalSort.cpp
@@ -3,39 +3,53 @@
 #include "Algo/TopologicalSort.h"
 #include "TestHarness.h"
 
-TEST_CASE("Core::Algorithm::TopologicalSort::Smoke Test", "[Core][Algorithm][Smoke]")
+TEST_CASE_METHOD(FAutomationTestFixture, "Core::Algorithm::TopologicalSort::Smoke Test", "[Core][Algorithm][Smoke]")
 {
 	using namespace Algo;
 
-	const TArray<int32> ExpectedArray{ 1, 2, 3 };
-
 	{
 		TArray<int32> Array{ 1, 2, 3 };
 
 		// Test the sort when each node depends on the previous one
 		bool bHasSucceeded = TopologicalSort(Array, [](int32 Element) { return Element > 1 ? TArray<int32>{ Element - 1 } : TArray<int32>{}; });
-		TestEqual(TEXT("TopologicalSort did not sort correctly"), bHasSucceeded, true);
-		TestEqual(TEXT("TopologicalSort did not sort correctly"), Array, ExpectedArray);
+
+		INFO("TopologicalSort did not sort correctly");
+		CHECK(bHasSucceeded);
+		CHECK(Array==(TArray<int32>{1, 2, 3}));
 	}
 	{
 		TArray<int32> Array{ 1, 2, 3 };
 		
 		// Test the sort when each node depends on the next one
 		bool bHasSucceeded = TopologicalSort(Array, [](int32 Element) { return Element < 3 ? TArray<int32>{ Element + 1 } : TArray<int32>{}; });
-		TestEqual(TEXT("TopologicalSort did not sort correctly"), bHasSucceeded, true);
-		TestEqual(TEXT("TopologicalSort did not sort correctly"), Array, (TArray<int32>{3, 2, 1}));
+
+		INFO("TopologicalSort did not sort correctly");
+		CHECK(bHasSucceeded);
+		CHECK( Array==(TArray<int32>{3, 2, 1}));
 	}
+
 	{
 		TArray<int32> Array{ 1, 2 };
-		const TArray<int32> ExpectedArrayCycle{1, 2};
 		
 		// Test the sort with a cycle between 1 and 2
 		bool bHasSucceeded = TopologicalSort(Array, [](int32 Element) { return TArray<int32>{ 1 + Element % 2 }; }, ETopologicalSort::None);
-		TestEqual(TEXT("TopologicalSort should not have succeeded when a cycle is detected"), bHasSucceeded, false);
-		TestEqual(TEXT("TopologicalSort should not have modified the array when failing"), Array, ExpectedArrayCycle);
+		
+		{
+			INFO("TopologicalSort should not have succeeded when a cycle is detected");
+			CHECK_FALSE(bHasSucceeded);
+		}
 
-		bHasSucceeded = TopologicalSort(Array, [](int32 Element) { return TArray<int32>{ 1 + Element % 2 }; }, ETopologicalSort::AllowCycles);
-		TestEqual(TEXT("TopologicalSort should succeed when a cycle is detected but AllowCycles is specified"), bHasSucceeded, true);
+		{
+			INFO("TopologicalSort should not have modified the array when failing");
+			CHECK(Array == (TArray<int32>{1, 2}));
+		}
+
+		{
+			bHasSucceeded = TopologicalSort(Array, [](int32 Element) { return TArray<int32>{ 1 + Element % 2 }; }, ETopologicalSort::AllowCycles);
+
+			INFO("TopologicalSort should succeed when a cycle is detected but AllowCycles is specified");
+			CHECK(bHasSucceeded);
+		}
 	}
 	{
 		TArray<int32> Array;
@@ -46,8 +60,13 @@ TEST_CASE("Core::Algorithm::TopologicalSort::Smoke Test", "[Core][Algorithm][Smo
 
 		// Make sure node 500 makes it on top if every other node depends on it
 		bool bHasSucceeded = TopologicalSort(Array, [](int32 Element) { return Element == 500 ? TArray<int32>{} : TArray<int32>{ 500 }; });
-		TestEqual(TEXT("TopologicalSort did not sort correctly"), bHasSucceeded, true);
-		TestEqual(TEXT("TopologicalSort did not sort correctly"), Array[0], 500);
+		
+		{
+			INFO("TopologicalSort did not sort correctly")
+			CHECK(bHasSucceeded);
+			CHECK(Array[0]==500);
+		}
+		
 	}
 	{
 		TArray<int32> Array{ 1, 2, 3, 4, 5, 6, 7, 8, 9, 10 };
@@ -76,15 +95,18 @@ TEST_CASE("Core::Algorithm::TopologicalSort::Smoke Test", "[Core][Algorithm][Smo
 		Links.Add(5, 2);
 
 		bool bHasSucceeded = TopologicalSort(Array, [&Links](int32 Element) { TArray<int32> Dependencies; Links.MultiFind(Element, Dependencies); return Dependencies; });
-		TestEqual(TEXT("TopologicalSort did not sort correctly"), bHasSucceeded, true);
+
+		INFO("TopologicalSort did not sort correctly");
+		CHECK(bHasSucceeded);
 
 		// There might be multiple valid answers, so test each condition separately to make sure they are all met.
 		for (const auto& Pair : Links)
 		{
 			int32 ChildIndex = Array.IndexOfByKey(Pair.Key);
 			int32 ParentIndex = Array.IndexOfByKey(Pair.Value);
-			TestEqual(TEXT("TopologicalSort did not sort correctly"), ParentIndex < ChildIndex, true);
+			CHECK(ParentIndex < ChildIndex);
 		}
+
 	}
 
 	{
@@ -99,9 +121,11 @@ TEST_CASE("Core::Algorithm::TopologicalSort::Smoke Test", "[Core][Algorithm][Smo
 		bool bHasSucceeded = TopologicalSort(Array,
 			[&Links](int32 Element) { TArray<int32> Dependencies; Links.MultiFind(Element, Dependencies); return Dependencies; },
 			ETopologicalSort::AllowCycles);
-		TestEqual(TEXT("TopologicalSort did not sort correctly"), bHasSucceeded, true);
-		TestTrue(TEXT("TopologicalSort did not sort correctly"),
-			(Array == TArray<int32>{4, 3, 2, 1} || Array == TArray<int32>{4, 3, 1, 2}));
+
+		INFO("TopologicalSort did not sort correctly");
+
+		CHECK(bHasSucceeded);
+		CHECK((Array == TArray<int32>{4, 3, 2, 1} || Array == TArray<int32>{4, 3, 1, 2}));
 	}
 	{
 		// Test the sort with a cycle in the root and with the root cycle depending on a chain of non-cycle verts, submitted in revese
@@ -116,9 +140,9 @@ TEST_CASE("Core::Algorithm::TopologicalSort::Smoke Test", "[Core][Algorithm][Smo
 			[&Links](int32 Element) { TArray<int32> Dependencies; Links.MultiFind(Element, Dependencies); return Dependencies; },
 			ETopologicalSort::AllowCycles);
 
-		TestEqual(TEXT("TopologicalSort did not sort correctly"), bHasSucceeded, true);
-		TestTrue(TEXT("TopologicalSort did not sort correctly"),
-			(Array == TArray<int32>{1, 2, 3, 4} || Array == TArray<int32>{1, 2, 4, 3}));
+		INFO("TopologicalSort did not sort correctly");
+		CHECK(bHasSucceeded);
+		CHECK((Array == TArray<int32>{1, 2, 3, 4} || Array == TArray<int32>{1, 2, 4, 3}));
 	}
 	{
 		// Test the sort with a cycle at a leaf and a chain from the root depending on that cycle
@@ -133,10 +157,11 @@ TEST_CASE("Core::Algorithm::TopologicalSort::Smoke Test", "[Core][Algorithm][Smo
 			[&Links](int32 Element) { TArray<int32> Dependencies; Links.MultiFind(Element, Dependencies); return Dependencies; },
 			ETopologicalSort::AllowCycles);
 
-		TestEqual(TEXT("TopologicalSort did not sort correctly"), bHasSucceeded, true);
+		INFO("TopologicalSort did not sort correctly");
+		CHECK(bHasSucceeded);
 		// There might be multiple valid answers, so test each condition separately to make sure they are all met.
-		TestTrue(TEXT("TopologicalSort did not sort correctly"), Array.IndexOfByKey(2) < Array.IndexOfByKey(1));
-		TestTrue(TEXT("TopologicalSort did not sort correctly"), Array[0] == 3 || Array[0] == 4);
+		CHECK(Array.IndexOfByKey(2) < Array.IndexOfByKey(1));
+		CHECK((Array[0] == 3 || Array[0] == 4));
 	}
 	{
 		// Test the sort with a cycle in the root and with the root cycle depending on a chain of non-cycle verts, submitted in revese
@@ -151,10 +176,11 @@ TEST_CASE("Core::Algorithm::TopologicalSort::Smoke Test", "[Core][Algorithm][Smo
 			[&Links](int32 Element) { TArray<int32> Dependencies; Links.MultiFind(Element, Dependencies); return Dependencies; },
 			ETopologicalSort::AllowCycles);
 
-		TestEqual(TEXT("TopologicalSort did not sort correctly"), bHasSucceeded, true);
+		INFO("TopologicalSort did not sort correctly");
+		CHECK(bHasSucceeded);
 		// There might be multiple valid answers, so test each condition separately to make sure they are all met.
-		TestTrue(TEXT("TopologicalSort did not sort correctly"), Array.IndexOfByKey(3) < Array.IndexOfByKey(4));
-		TestTrue(TEXT("TopologicalSort did not sort correctly"), Array[0] == 1 || Array[0] == 2);
+		CHECK(Array.IndexOfByKey(3) < Array.IndexOfByKey(4));
+		CHECK((Array[0] == 1 || Array[0] == 2));
 	}
 	{
 		// Verify that when breaking a cycle a member of the cycle is selected rather than an element that depends on but is not in the cycle
@@ -173,8 +199,9 @@ TEST_CASE("Core::Algorithm::TopologicalSort::Smoke Test", "[Core][Algorithm][Smo
 			[&Links](int32 Element) { TArray<int32> Dependencies; Links.MultiFind(Element, Dependencies); return Dependencies; },
 			ETopologicalSort::AllowCycles);
 
-		TestEqual(TEXT("TopologicalSort did not sort correctly"), bHasSucceeded, true);
+		INFO("TopologicalSort did not sort correctly");
+		CHECK(bHasSucceeded);
 		// There might be multiple valid answers, so test each condition separately to make sure they are all met.
-		TestTrue(TEXT("TopologicalSort did not sort correctly"), Array.IndexOfByKey(2) < Array.IndexOfByKey(1));
+		CHECK(Array.IndexOfByKey(2) < Array.IndexOfByKey(1));
 	}
 }
diff --git a/Engine/Source/Runtime/Core/Tests/Algorithm/UniqueTest.cpp b/Engine/Source/Runtime/Core/Tests/Algorithm/UniqueTest.cpp
--- a/Engine/Source/Runtime/Core/Tests/Algorithm/UniqueTest.cpp
+++ b/Engine/Source/Runtime/Core/Tests/Algorithm/UniqueTest.cpp
@@ -5,45 +5,50 @@
 #include "Containers/Array.h"
 #include "TestHarness.h"
 
-TEST_CASE("Core::Algorithm::Unique::Smoke Test", "[Core][Algorithm][Smoke]")
+TEST_CASE_METHOD(FAutomationTestFixture, "Core::Algorithm::Unique::Smoke Test", "[Core][Algorithm][Smoke]")
 {
 	using namespace Algo;
 
-	const TArray<int32> ExpectedArray{1, 2, 3};
-
 	{
 		TArray<int32> Array;
 		int32 RemoveFrom = Unique(Array);
-		TestEqual(TEXT("`Unique` must handle an empty container"), RemoveFrom, 0);
+		INFO("`Unique` must handle an empty container");
+		CHECK(RemoveFrom==0);
 	}
 	{
 		TArray<int32> Array{ 1, 2, 3 };
 		Array.SetNum(Unique(Array));
-		TestEqual(TEXT("Uniqued container with no duplicates must remain unchanged"), Array, ExpectedArray);
+		INFO("Uniqued container with no duplicates must remain unchanged");
+		CHECK(Array==TArray<int32>{1, 2, 3});
 	}
 	{
 		TArray<int32> Array{ 1, 1, 2, 2, 2, 3, 3, 3, 3};
 		Array.SetNum(Unique(Array));
-		TestEqual(TEXT("`Unique` with multiple duplicates must return correct result"), Array, ExpectedArray);
+		INFO("`Unique` with multiple duplicates must return correct result");
+		CHECK(Array==TArray<int32>{1, 2, 3});
 	}
 	{
 		TArray<int32> Array{ 1, 1, 2, 3, 3, 3 };
 		Array.SetNum(Unique(Array));
-		TestEqual(TEXT("`Unique` with duplicates and unique items must return correct result"), Array, ExpectedArray);
+		INFO("`Unique` with duplicates and unique items must return correct result");
+		CHECK(Array==TArray<int32>{1, 2, 3});
 	}
 	{
 		FString Str = TEXT("aa");
 		Str = Str.Mid(0, Unique(Str));
-		TestEqual(TEXT("`Unique` on `FString` as an example of arbitrary random-access container must compile and return correct result"), Str, FString(TEXT("a")));
+		INFO("`Unique` on `FString` as an example of arbitrary random-access container must compile and return correct result");
+		CHECK(Str==FString(TEXT("a")));
 	}
 	{
 		int32 Array[] = {1};
 		int32 NewSize = (int32)Unique(Array);
-		TestEqual(TEXT("`Unique` must support C arrays"), NewSize, 1);
+		INFO("`Unique` must support C arrays");
+		CHECK(NewSize==1);
 	}
 	{
 		TArray<int32> Array = { 1, 1 };
 		int32 NewSize = Unique(MakeArrayView(Array.GetData() + 1, 1));
-		TestEqual(TEXT("`Unique` must support ranges"), NewSize, 1);
+		INFO("`Unique` must support ranges");
+		CHECK(NewSize==1);
 	}
 }
diff --git a/Engine/Source/Runtime/Core/Tests/Async/AsyncTest.cpp b/Engine/Source/Runtime/Core/Tests/Async/AsyncTest.cpp
--- a/Engine/Source/Runtime/Core/Tests/Async/AsyncTest.cpp
+++ b/Engine/Source/Runtime/Core/Tests/Async/AsyncTest.cpp
@@ -3,9 +3,7 @@
 #include "CoreTypes.h"
 #include "Templates/Function.h"
 #include "Async/Async.h"
-#include "TestHarness.h"
-
-#include "TestCommon/CoreUtilities.h"
+#include "TestFixtures/CoreTestFixture.h"
 
 /** Helper methods used in the test cases. */
 namespace AsyncTestUtils
@@ -23,69 +21,51 @@ namespace AsyncTestUtils
 
 
 /** Test that task graph tasks return correctly. */
-TEST_CASE("Core::Async::TaskGraph::Task Graph", "[Core][Async][Smoke]")
+TEST_CASE_METHOD(FCoreTestFixture, "Core::Async::Task::Task Graph", "[Core][Async][Smoke]")
 {
-	InitTaskGraphAndDependencies();
-
 	auto Future = Async(EAsyncExecution::TaskGraph, AsyncTestUtils::Task);
 	int Result = Future.Get();
 
-	TestEqual(TEXT("Task graph task must return expected value"), Result, 123);
-
-	CleanupTaskGraphAndDependencies();
+	CHECK(Result==123);
 }
 
 
 /** Test that threaded tasks return correctly. */
-TEST_CASE("Core::Async::Thread::Thread", "[Core][Async][Smoke]")
+TEST_CASE_METHOD(FCoreTestFixture, "Core::Async::Task::Thread", "[Core][Async][Smoke]")
 {
-	InitTaskGraphAndDependencies();
-
 	auto Future = Async(EAsyncExecution::Thread, AsyncTestUtils::Task);
 	int Result = Future.Get();
 
-	REQUIRE(Result == 123);
-
-	CleanupTaskGraphAndDependencies();
+	CHECK(Result==123);
 }
 
 
 /** Test that threaded pool tasks return correctly. */
-TEST_CASE("Core::Async::ThreadedPool::Threaded Pool", "[Core][Async][Smoke]")
+TEST_CASE_METHOD(FCoreTestFixture, "Core::Async::Task::Threaded Pool", "[Core][Async][Smoke]")
 {
-	InitTaskGraphAndDependencies();
-
 	auto Future = Async(EAsyncExecution::ThreadPool, AsyncTestUtils::Task);
 	int Result = Future.Get();
 
-	TestEqual(TEXT("Thread pool task must return expected value"), Result, 123);
-
-	CleanupTaskGraphAndDependencies();
+	CHECK(Result==123);
 }
 
 
 /** Test that void tasks run without errors or warnings. */
-TEST_CASE("Core::Async::Task::Void Task", "[Core][Async][Smoke]")
+TEST_CASE_METHOD(FCoreTestFixture, "Core::Async::Task::Void Task", "[Core][Async][Smoke]")
 {
-	InitTaskGraphAndDependencies();
-
 	// Reset test variable before running
 	AsyncTestUtils::bHasVoidTaskFinished = false;
 	auto Future = Async(EAsyncExecution::TaskGraph, AsyncTestUtils::VoidTask);
 	Future.Get();
 
 	// Check that the variable state was updated by task
-	TestTrue(TEXT("Void tasks should run"), AsyncTestUtils::bHasVoidTaskFinished);
-
-	CleanupTaskGraphAndDependencies();
+	CHECK(AsyncTestUtils::bHasVoidTaskFinished);
 }
 
 
-/** Test that asynchronous tasks have their completion callback called. */
-TEST_CASE("Core::Async::CompletionCallback::Completion Callback", "[Core][Async][Smoke]")
+/** Test that FCoreTestFixture tasks have their completion callback called. */
+TEST_CASE_METHOD(FCoreTestFixture, "Core::Async::Task::Completion Callback", "[Core][Async][Smoke]")
 {
-	InitTaskGraphAndDependencies();
-
 	bool Completed = false;
 	FEvent* CompletedEvent = FPlatformProcess::GetSynchEventFromPool(true);
 
@@ -102,8 +82,8 @@ TEST_CASE("Core::Async::CompletionCallback::Completion Callback", "[Core][Async]
 	bool CompletedEventTriggered = CompletedEvent->Wait(FTimespan(0 /* hours */, 0 /* minutes */, 5 /* seconds */));
 	FPlatformProcess::ReturnSynchEventToPool(CompletedEvent);
 
-	TestEqual(TEXT("Async Result"), Result, 123);
-	TestTrue(TEXT("Completion callback to be called"), CompletedEventTriggered && Completed);
+	CHECK(Result==123);
+	CHECK(CompletedEventTriggered);
+	CHECK(Completed);
 
-	CleanupTaskGraphAndDependencies();
 }
diff --git a/Engine/Source/Runtime/Core/Tests/Async/AsyncWorkTest.cpp b/Engine/Source/Runtime/Core/Tests/Async/AsyncWorkTest.cpp
--- a/Engine/Source/Runtime/Core/Tests/Async/AsyncWorkTest.cpp
+++ b/Engine/Source/Runtime/Core/Tests/Async/AsyncWorkTest.cpp
@@ -4,9 +4,7 @@
 #include "Misc/AutomationTest.h"
 #include "Async/AsyncWork.h"
 #include "Misc/QueuedThreadPool.h"
-#include "TestHarness.h"
-
-#include "TestCommon/CoreUtilities.h"
+#include "TestFixtures/CoreTestFixture.h"
 
 namespace AsyncWorkTestImpl
 {
@@ -63,7 +61,8 @@ namespace AsyncWorkTestImpl
 		OtherTaskType* OtherTask;
 		void DoWork() 
 		{ 
-			checkf(!bIsOuterTaskRunning, TEXT("We got picked by the busywait of the outer task we are going to wait on, we are now deadlocked"));
+			INFO("We got picked by the busywait of the outer task we are going to wait on, we are now deadlocked");
+			CHECK(!bIsOuterTaskRunning);
 
 			TRACE_CPUPROFILER_EVENT_SCOPE(FCompletionTask::DoWork);
 
@@ -82,10 +81,8 @@ namespace AsyncWorkTestImpl
 	};
 }
 
-TEST_CASE("Core::Async::AsyncWork::Ensure Completion BusyWait Deadlock Test", "[Core][Async][Smoke]")
+TEST_CASE_METHOD(FCoreTestFixture, "Core::Async::AsyncWork::Ensure Completion BusyWait Deadlock Test", "[Core][Async][Smoke]")
 {
-	InitTaskGraphAndDependencies();
-
 	using namespace AsyncWorkTestImpl;
 	TRACE_CPUPROFILER_EVENT_SCOPE(FAsyncWorkEnsureCompletionBusyWaitDeadLockTest::RunTest);
 	
@@ -131,6 +128,5 @@ TEST_CASE("Core::Async::AsyncWork::Ensure Completion BusyWait Deadlock Test", "[
 		CompletionTasks[Index].EnsureCompletion(false);
 	}
 	OuterTask.EnsureCompletion(false);
-
-	CleanupTaskGraphAndDependencies();
+	SUCCEED();
 }
diff --git a/Engine/Source/Runtime/Core/Tests/Async/TaskGraphTest.cpp b/Engine/Source/Runtime/Core/Tests/Async/TaskGraphTest.cpp
--- a/Engine/Source/Runtime/Core/Tests/Async/TaskGraphTest.cpp
+++ b/Engine/Source/Runtime/Core/Tests/Async/TaskGraphTest.cpp
@@ -13,14 +13,11 @@
 #include "Math/RandomStream.h"
 #include "Containers/CircularQueue.h"
 #include "Containers/Queue.h"
-#include "Tests/Benchmark.h"
+#include "Core/Tests/Benchmark.h"
 #include "HAL/Thread.h"
 #include "Async/Fundamental/Scheduler.h"
 #include <atomic>
-#include "TestHarness.h"
-#include "TestRunner.h"
-
-#include "TestCommon/CoreUtilities.h"
+#include "TestFixtures/CoreTestFixture.h"
 
 namespace OldTaskGraphTests
 {
@@ -484,7 +481,7 @@ namespace OldTaskGraphTests
 				Rig.Test2.PopAll(Items);
 				Rig.Test3.PopAll(Items);
 
-				checkf(Items.Num() == 1000, TEXT("Items %d"), Items.Num());
+				CHECK(Items.Num() == 1000);
 
 				for (int32 LookFor = 0; LookFor < 1000; LookFor++)
 				{
@@ -692,19 +689,21 @@ namespace OldTaskGraphTests
 	);
 
 
-	TEST_CASE("Core::Async::TaskGraph::OldBenchmark", "[Core][Async][Perf]")
+	IMPLEMENT_SIMPLE_AUTOMATION_TEST(FTaskGraphOldBenchmark, "System.Core.Async.TaskGraph.OldBenchmark", EAutomationTestFlags::EditorContext | EAutomationTestFlags::ClientContext | EAutomationTestFlags::ServerContext | EAutomationTestFlags::EngineFilter);
+
+	bool FTaskGraphOldBenchmark::RunTest(const FString& Parameters)
 	{
-		InitTaskGraphAndDependencies();
 		TArray<FString> Args;
 		TaskGraphBenchmark(Args);
-		CleanupTaskGraphAndDependencies();
+		return true;
 	}
 
-	TEST_CASE("Core::Async::TaskGraph::LockFree", "[Core][Async][Perf]")
+	IMPLEMENT_SIMPLE_AUTOMATION_TEST(FLockFreeTest, "System.Core.Async.TaskGraph.LockFree", EAutomationTestFlags::ApplicationContextMask | EAutomationTestFlags::EngineFilter);
+
+	bool FLockFreeTest::RunTest(const FString& Parameters)
 	{
-		InitTaskGraphAndDependencies();
 		TestLockFree(3);
-		CleanupTaskGraphAndDependencies();
+		return true;
 	}
 }
 
@@ -712,10 +711,8 @@ extern int32 GNumForegroundWorkers;
 
 namespace TaskGraphTests
 {
-	TEST_CASE("Core::Async::TaskGraph::Graph Event", "[Core][Async][Smoke]")
+	TEST_CASE_METHOD(FCoreTestFixture, "Core::Async::TaskGraph::Graph Event", "[Core][Async][Smoke]")
 	{
-		InitTaskGraphAndDependencies();
-
 		{	// task completes before it's waited for
 			FGraphEventRef Event = FFunctionGraphTask::CreateAndDispatchWhenReady(
 				[]
@@ -818,14 +815,12 @@ namespace TaskGraphTests
 			{}
 			Event->Wait(ENamedThreads::GameThread);
 		}
-
-		CleanupTaskGraphAndDependencies();
 	}
 
-	/* Blocks forever - Disabled */
-	TEST_CASE("Core::Async::TaskGraph::Recursion", "[.][Core][Async][Smoke]")
+
+	TEST_CASE_METHOD(FCoreTestFixture, "Core::Async::TaskGraph::Recursion", "[Core][Async][Smoke]")
 	{
-		InitTaskGraphAndDependencies();
+		return;
 
 		{	// recursive call on game thread
 			FGraphEventRef Event = FFunctionGraphTask::CreateAndDispatchWhenReady(
@@ -849,15 +844,11 @@ namespace TaskGraphTests
 		//	FGraphEventRef Event = FFunctionGraphTask::CreateAndDispatchWhenReady([] {}, TStatId{}, nullptr, ENamedThreads::GameThread_Local);
 		//	Event->Wait(ENamedThreads::GameThread);
 		//}
-
-		CleanupTaskGraphAndDependencies();
 	}
 
 	
-	TEST_CASE("Core::Async::TaskGraph::Basic", "[Core][Async][Smoke]")
+	TEST_CASE_METHOD(FCoreTestFixture, "Core::Async::TaskGraph::Basic", "[Core][Async][Smoke]")
 	{
-		InitTaskGraphAndDependencies();
-
 		// thread and task priorities
 
 		{	// AnyNormalThreadNormalTask
@@ -1033,8 +1024,6 @@ PRAGMA_ENABLE_DEPRECATION_WARNINGS
 			LocalQueueTask->Wait(ENamedThreads::GameThread_Local);
 			CHECK(LocalQueueTask.GetRefCount() == 1);
 		}
-
-		CleanupTaskGraphAndDependencies();
 	}
 
 	// it's fast because tasks are too lightweight and so are executed almost as fast
@@ -1296,10 +1285,8 @@ PRAGMA_ENABLE_DEPRECATION_WARNINGS
 		}
 	}
 
-	TEST_CASE("Core::Async::TaskGraph::Performance", "[Core][Async][Perf]")
+	TEST_CASE_METHOD(FCoreTestFixture, "Core::Async::TaskGraph::Performance", "[.][Core][Async][.Perf]")
 	{
-		InitTaskGraphAndDependencies(true);
-
 		TRACE_CPUPROFILER_EVENT_SCOPE(TaskGraphTests_PerfTest);
 
 		//UE_BENCHMARK(1, TestSpawning<100000>);
@@ -1316,25 +1303,17 @@ PRAGMA_ENABLE_DEPRECATION_WARNINGS
 		UE_BENCHMARK(5, TestWorkStealing<100, 1000>);
 		UE_BENCHMARK(5, TestSpawning<100000>);
 		UE_BENCHMARK(5, TestBatchSpawning<100000>);
-
-		CleanupTaskGraphAndDependencies();
 	}
 
-	TEST_CASE("Core::Async::TaskGraph::Stats Thread Redirection", "[Core][Async][Smoke]")
+	TEST_CASE_METHOD(FCoreTestFixture, "Core::Async::TaskGraph::Stats Thread Redirection", "[Core][Async][Smoke]")
 	{
-		InitTaskGraphAndDependencies();
-
 #if STATS
-
 		bool bExecuted = false;
 PRAGMA_DISABLE_DEPRECATION_WARNINGS
 		FFunctionGraphTask::CreateAndDispatchWhenReady([&bExecuted] { bExecuted = true; }, TStatId{}, nullptr, ENamedThreads::StatsThread)->Wait();
 PRAGMA_ENABLE_DEPRECATION_WARNINGS
 		CHECK(bExecuted);
-
 #endif
-
-		CleanupTaskGraphAndDependencies();
 	}
 
 	template<uint32 Num>
@@ -1362,18 +1341,13 @@ PRAGMA_ENABLE_DEPRECATION_WARNINGS
 					Event->Trigger();
 				}
 			);
-			REQUIRE(Event->Wait(FTimespan::FromSeconds(5.f)));
+			CHECK(Event->Wait(FTimespan::FromSeconds(5.f)));
 		}
 	}
 
-	/* Disabled */
-	TEST_CASE("Core::Async::TaskGraph::Oversubscription", "[.][Core][Async][Smoke]")
+	TEST_CASE_METHOD(FCoreTestFixture, "Core::Async::TaskGraph::Oversubscription", "[Core][Async][.Stress]")
 	{
-		InitTaskGraphAndDependencies(true);
-
 		UE_BENCHMARK(5, OversubscriptionStressTest<10>);
-
-		CleanupTaskGraphAndDependencies();
 	}
 
 	template<uint32 Nujm>
@@ -1400,32 +1374,22 @@ PRAGMA_ENABLE_DEPRECATION_WARNINGS
 				Event->Trigger();
 			}
 		);
-		REQUIRE(Event->Wait(FTimespan::FromSeconds(30.f)));
+		CHECK(Event->Wait(FTimespan::FromSeconds(30.f)));
 	}
 
-	/* Disabled CL 16742068 */
-	TEST_CASE("Core::Async::TaskGraph::Squared Oversubscription", "[.][Core][Async][Smoke]")
+	TEST_CASE_METHOD(FCoreTestFixture, "Core::Async::TaskGraph::Squared Oversubscription", "[Core][Async][.Stress]")
 	{
-		InitTaskGraphAndDependencies(true);
-
 		UE_BENCHMARK(5, SquaredOversubscriptionStressTest<10>);
-
-		CleanupTaskGraphAndDependencies();
 	}
 
 	template<uint32 Num, bool bBackgroundWorkersToo>
 	void BlockingWorkersByFEvent();
 
-	/* Disabled UE-131742 CL 17829717 */
-	TEST_CASE("Core::Async::TaskGraph::Blocking Workers", "[Core][Async][Smoke]")
+	TEST_CASE_METHOD(FCoreTestFixture, "Core::Async::TaskGraph::Blocking Workers", "[Core][Async][.Stress]")
 	{
-		InitTaskGraphAndDependencies();
-
 		FPlatformProcess::Sleep(1.0f);
 		UE_BENCHMARK(10, BlockingWorkersByFEvent<100, false>);
 		UE_BENCHMARK(10, BlockingWorkersByFEvent<100, true>);
-
-		CleanupTaskGraphAndDependencies();
 	}
 
 	template<uint32 Num, bool bBackgroundWorkersToo>
@@ -1448,7 +1412,7 @@ PRAGMA_ENABLE_DEPRECATION_WARNINGS
 				WorkerBlockers[i].Init(TEXT("BlockingWorkers"), LowLevelTasks::ETaskPriority::High,
 					[i, &BlockersBlocker, &AllWorkersBlockedSignal, &NumWorkersNotBlocked]
 					{
-						REQUIRE(LowLevelTasks::FSchedulerTls::GetAffinityIndex() == i);
+						CHECK(LowLevelTasks::FSchedulerTls::GetAffinityIndex() == i);
 						if (--NumWorkersNotBlocked == 0)
 						{
 							AllWorkersBlockedSignal->Trigger();
@@ -1459,10 +1423,10 @@ PRAGMA_ENABLE_DEPRECATION_WARNINGS
 						}
 					});
 
-				REQUIRE(LowLevelTasks::TryLaunchAffinity(WorkerBlockers[i], i));
+				CHECK(LowLevelTasks::TryLaunchAffinity(WorkerBlockers[i], i));
 			}
 
-			REQUIRE(AllWorkersBlockedSignal->Wait(FTimespan::FromSeconds(3)));
+			CHECK(AllWorkersBlockedSignal->Wait(FTimespan::FromSeconds(3)));
 
 			BlockersBlocker->Trigger();
 			for (int i = 0; i != NumWorkers; ++i)
diff --git a/Engine/Source/Runtime/Core/Tests/Compression/CompressedBufferTest.cpp b/Engine/Source/Runtime/Core/Tests/Compression/CompressedBufferTest.cpp
--- a/Engine/Source/Runtime/Core/Tests/Compression/CompressedBufferTest.cpp
+++ b/Engine/Source/Runtime/Core/Tests/Compression/CompressedBufferTest.cpp
@@ -1,76 +1,110 @@
 // Copyright Epic Games, Inc. All Rights Reserved.
 
 #include "Compression/CompressedBuffer.h"
-
-#include "Algo/Compare.h"
 #include "Compression/OodleDataCompression.h"
 #include "IO/IoHash.h"
+#include "Hash/Blake3.h"
 #include "Serialization/BufferReader.h"
-#include "TestCommon/CoreUtilities.h"
-#include "TestHarness.h"
+#include "TestFixtures/CoreTestFixture.h"
 
-TEST_CASE("Core::Compression::CompressedBuffer::Compress", "[Core][Compression][Smoke]")
+TEST_CASE_METHOD(FCoreTestFixture, "Core::Compression::FCompressedBuffer::Compress", "[Core][Math][Smoke]")
 {
 	const uint8 ZeroBuffer[1024]{};
-	const FIoHash ZeroBufferHash = FIoHash::HashBuffer(MakeMemoryView(ZeroBuffer));
+	const FBlake3Hash ZeroBufferHash = FBlake3::HashBuffer(MakeMemoryView(ZeroBuffer));
 
-	SECTION("Null")
+	// Test Null
 	{
+		INFO("Null");
+
 		const FCompressedBuffer Buffer;
-		CHECK_FALSE(Buffer);
+		CHECK_FALSE((bool)Buffer);
 		CHECK(Buffer.IsNull());
 		CHECK(Buffer.IsOwned());
-		CHECK(Buffer.GetCompressedSize() == 0);
-		CHECK(Buffer.GetRawSize() == 0);
-		CHECK(Buffer.GetRawHash() == FIoHash::Zero);
+		CHECK(Buffer.GetCompressedSize()==uint64(0));
+		CHECK(Buffer.GetRawSize()==uint64(0));
+		CHECK(Buffer.GetRawHash()==FBlake3Hash::Zero);
 		CHECK(Buffer.Decompress().IsNull());
 		CHECK(Buffer.DecompressToComposite().IsNull());
 	}
 
-	SECTION("Method None")
+	// Test Method None
 	{
-		const FCompressedBuffer BufferNew = FCompressedBuffer::Compress(FSharedBuffer::MakeView(MakeMemoryView(ZeroBuffer)),
-			ECompressedBufferCompressor::NotSet, ECompressedBufferCompressionLevel::None);
-		const FCompressedBuffer BufferCopy = FCompressedBuffer::FromCompressed(BufferNew.GetCompressed());
-		const FCompressedBuffer Buffer = GENERATE_REF(BufferNew, BufferCopy);
-
-		CHECK(Buffer);
-		CHECK(!Buffer.IsNull());
-		CHECK(Buffer.IsOwned());
-		CHECK(Buffer.GetCompressedSize() == sizeof(ZeroBuffer) + 64);
-		CHECK(Buffer.GetRawSize() == sizeof(ZeroBuffer));
-		CHECK(Buffer.GetRawHash() == ZeroBufferHash);
-		CHECK(FIoHash::HashBuffer(Buffer.Decompress()) == ZeroBufferHash);
-		CHECK(FIoHash::HashBuffer(Buffer.DecompressToComposite()) == Buffer.GetRawHash());
+		INFO("Method None");
 
 		ECompressedBufferCompressor Compressor = ECompressedBufferCompressor::Kraken;
 		ECompressedBufferCompressionLevel CompressionLevel = ECompressedBufferCompressionLevel::Normal;
 		uint64 BlockSize = MAX_uint64;
-		CHECK(Buffer.TryGetCompressParameters(Compressor, CompressionLevel, BlockSize));
+
+		const FCompressedBuffer Buffer = FCompressedBuffer::Compress(FSharedBuffer::MakeView(MakeMemoryView(ZeroBuffer)),
+			ECompressedBufferCompressor::NotSet, ECompressedBufferCompressionLevel::None);
+		CHECK((bool)Buffer);
+		CHECK_FALSE(Buffer.IsNull());
+		CHECK(Buffer.IsOwned());
+		CHECK(Buffer.GetCompressedSize()==sizeof(ZeroBuffer) + 64);
+		CHECK(Buffer.GetRawSize()==sizeof(ZeroBuffer));
+		CHECK(Buffer.GetRawHash()==ZeroBufferHash);
+		CHECK(FBlake3::HashBuffer(Buffer.Decompress())==ZeroBufferHash);
+		CHECK(FBlake3::HashBuffer(Buffer.DecompressToComposite())==Buffer.GetRawHash());
+		CHECK(Buffer.TryGetCompressParameters(Compressor, CompressionLevel, BlockSize));	
 		CHECK(Compressor == ECompressedBufferCompressor::NotSet);
 		CHECK(CompressionLevel == ECompressedBufferCompressionLevel::None);
 		CHECK(BlockSize == 0);
-	}
 
-	SECTION("Method Oodle")
-	{
-		const FCompressedBuffer BufferNew = FCompressedBuffer::Compress(FSharedBuffer::MakeView(MakeMemoryView(ZeroBuffer)),
-			ECompressedBufferCompressor::Mermaid, ECompressedBufferCompressionLevel::VeryFast);
-		const FCompressedBuffer BufferCopy = FCompressedBuffer::FromCompressed(BufferNew.GetCompressed());
-		const FCompressedBuffer Buffer = GENERATE_REF(BufferNew, BufferCopy);
+		Compressor = ECompressedBufferCompressor::Kraken;
+		CompressionLevel = ECompressedBufferCompressionLevel::Normal;
+		BlockSize = MAX_uint64;
 
-		CHECK(Buffer);
-		CHECK(!Buffer.IsNull());
+		const FCompressedBuffer BufferCopy = FCompressedBuffer::FromCompressed(Buffer.GetCompressed());
+		CHECK((bool)Buffer);
+		CHECK_FALSE(Buffer.IsNull());
 		CHECK(Buffer.IsOwned());
-		CHECK(Buffer.GetCompressedSize() < sizeof(ZeroBuffer));
-		CHECK(Buffer.GetRawSize() == sizeof(ZeroBuffer));
-		CHECK(Buffer.GetRawHash() == ZeroBufferHash);
-		CHECK(FIoHash::HashBuffer(Buffer.Decompress()) == ZeroBufferHash);
-		CHECK(FIoHash::HashBuffer(Buffer.DecompressToComposite()) == Buffer.GetRawHash());
+		CHECK(Buffer.GetCompressedSize()==sizeof(ZeroBuffer) + 64);
+		CHECK(Buffer.GetRawSize()==sizeof(ZeroBuffer));
+		CHECK(Buffer.GetRawHash()==ZeroBufferHash);
+		CHECK(FBlake3::HashBuffer(Buffer.Decompress())==ZeroBufferHash);
+		CHECK(FBlake3::HashBuffer(Buffer.DecompressToComposite())==Buffer.GetRawHash());
+		CHECK(Buffer.TryGetCompressParameters(Compressor, CompressionLevel, BlockSize));
+		CHECK(Compressor == ECompressedBufferCompressor::NotSet);
+		CHECK(CompressionLevel ==ECompressedBufferCompressionLevel::None);
+		CHECK(BlockSize == 0);
+	}
+
+	// Test Method Oodle
+	{
+		INFO("Method Oodle");
 
 		ECompressedBufferCompressor Compressor = ECompressedBufferCompressor::Kraken;
 		ECompressedBufferCompressionLevel CompressionLevel = ECompressedBufferCompressionLevel::Normal;
 		uint64 BlockSize = MAX_uint64;
+
+		const FCompressedBuffer Buffer = FCompressedBuffer::Compress(FSharedBuffer::MakeView(MakeMemoryView(ZeroBuffer)),
+			ECompressedBufferCompressor::Mermaid, ECompressedBufferCompressionLevel::VeryFast);
+		CHECK((bool)Buffer);
+		CHECK_FALSE(Buffer.IsNull());
+		CHECK(Buffer.IsOwned());
+		CHECK(Buffer.GetCompressedSize() < sizeof(ZeroBuffer));
+		CHECK(Buffer.GetRawSize()==sizeof(ZeroBuffer));
+		CHECK(Buffer.GetRawHash()==ZeroBufferHash);
+		CHECK(FBlake3::HashBuffer(Buffer.Decompress())==ZeroBufferHash);
+		CHECK(FBlake3::HashBuffer(Buffer.DecompressToComposite())==Buffer.GetRawHash());
+		CHECK(Buffer.TryGetCompressParameters(Compressor, CompressionLevel, BlockSize));
+		CHECK(Compressor == ECompressedBufferCompressor::Mermaid);
+		CHECK(CompressionLevel == ECompressedBufferCompressionLevel::VeryFast);
+		CHECK(FMath::IsPowerOfTwo(BlockSize));
+
+		Compressor = ECompressedBufferCompressor::Kraken;
+		CompressionLevel = ECompressedBufferCompressionLevel::Normal;
+		BlockSize = MAX_uint64;
+
+		const FCompressedBuffer BufferCopy = FCompressedBuffer::FromCompressed(Buffer.GetCompressed());
+		CHECK((bool)Buffer);
+		CHECK_FALSE(Buffer.IsNull());
+		CHECK(Buffer.IsOwned());
+		CHECK(Buffer.GetCompressedSize() < sizeof(ZeroBuffer));
+		CHECK(Buffer.GetRawSize()==sizeof(ZeroBuffer));
+		CHECK(Buffer.GetRawHash()==ZeroBufferHash);
+		CHECK(FBlake3::HashBuffer(Buffer.Decompress())==ZeroBufferHash);
+		CHECK(FBlake3::HashBuffer(Buffer.DecompressToComposite())==Buffer.GetRawHash());
 		CHECK(Buffer.TryGetCompressParameters(Compressor, CompressionLevel, BlockSize));
 		CHECK(Compressor == ECompressedBufferCompressor::Mermaid);
 		CHECK(CompressionLevel == ECompressedBufferCompressionLevel::VeryFast);
@@ -78,7 +112,7 @@ TEST_CASE("Core::Compression::CompressedBuffer::Compress", "[Core][Compression][
 	}
 }
 
-TEST_CASE("Core::Compression::CompressedBuffer::Decompress", "[Core][Compression][Smoke]")
+TEST_CASE_METHOD(FCoreTestFixture, "Core::Compression::FCompressedBuffer::Decompress", "[Core][Math][Smoke]")
 {
 	const auto GenerateData = [](int32 N) -> TArray<uint64>
 	{
@@ -91,6 +125,17 @@ TEST_CASE("Core::Compression::CompressedBuffer::Decompress", "[Core][Compression
 		return Data;
 	};
 
+	const auto ValidateData = [this](TConstArrayView<uint64> Values, TConstArrayView<uint64> ExpectedValues, int32 Offset)
+	{
+		int32 ExpectedIndex = Offset;
+		for (uint64 Value : Values)
+		{
+			const uint64 ExpectedValue = ExpectedValues[ExpectedIndex];
+			CHECK(Value==ExpectedValue);
+			ExpectedIndex++;
+		}
+	};
+
 	const auto CastToArrayView = [](FMemoryView View) -> TConstArrayView<uint64>
 	{
 		return MakeArrayView(static_cast<const uint64*>(View.GetData()), static_cast<int32>(View.GetSize() / sizeof(uint64)));
@@ -98,9 +143,11 @@ TEST_CASE("Core::Compression::CompressedBuffer::Decompress", "[Core][Compression
 
 	FCompressedBufferReader Reader;
 
-	SECTION("Decompress with offset and size.")
+	// Test decompress with offset and size
 	{
-		const auto UncompressAndValidate = [&Reader, &CastToArrayView](
+		INFO("Decompress with offset and size");
+
+		const auto UncompressAndValidate = [this, &Reader, &ValidateData, &CastToArrayView](
 			const FCompressedBuffer& Compressed,
 			const int32 OffsetCount,
 			const int32 Count,
@@ -109,14 +156,16 @@ TEST_CASE("Core::Compression::CompressedBuffer::Decompress", "[Core][Compression
 			Reader.SetSource(Compressed);
 			{
 				const FSharedBuffer Uncompressed = Reader.Decompress(OffsetCount * sizeof(uint64), Count * sizeof(uint64));
-				CHECK(Algo::Compare(CastToArrayView(Uncompressed), ExpectedValues.Mid(OffsetCount, Count)));
+				const TConstArrayView<uint64> Values = CastToArrayView(Uncompressed);
+				CHECK(Values.Num()==Count);
+				ValidateData(Values, ExpectedValues, OffsetCount);
 			}
 			{
 				FUniqueBuffer Uncompressed = FUniqueBuffer::Alloc(Count * sizeof(uint64));
-				CHECKED_IF(Reader.TryDecompressTo(Uncompressed, OffsetCount * sizeof(uint64)))
-				{
-					CHECK(Algo::Compare(CastToArrayView(Uncompressed), ExpectedValues.Mid(OffsetCount, Count)));
-				}
+				const bool bOk = Reader.TryDecompressTo(Uncompressed, OffsetCount * sizeof(uint64));
+				CHECK(bOk);
+				const TConstArrayView<uint64> Values = CastToArrayView(Uncompressed);
+				ValidateData(Values, ExpectedValues, OffsetCount);
 			}
 		};
 
@@ -143,8 +192,10 @@ TEST_CASE("Core::Compression::CompressedBuffer::Decompress", "[Core][Compression
 		UncompressAndValidate(Compressed, 4993, 4, ExpectedValues);
 	}
 
-	SECTION("Decompress with offset only.")
+	// Decompress with offset only
 	{
+		INFO("Decompress with offset only");
+
 		constexpr uint64 BlockSize = 64 * sizeof(uint64);
 		constexpr int32 N = 1000;
 		const TArray<uint64> ExpectedValues = GenerateData(N);
@@ -154,19 +205,21 @@ TEST_CASE("Core::Compression::CompressedBuffer::Decompress", "[Core][Compression
 			FOodleDataCompression::ECompressor::Mermaid,
 			FOodleDataCompression::ECompressionLevel::Fast,
 			BlockSize);
-
+		
 		constexpr uint64 OffsetCount = 150;
 		{
 			FSharedBuffer Buffer = Compressed.GetCompressed().ToShared();
 			FBufferReader Ar(const_cast<void*>(Buffer.GetData()), int64(Buffer.GetSize()), /*bFreeOnClose*/ false, /*bIsPersistent*/ true);
 			FCompressedBufferReaderSourceScope Source(Reader, Ar);
 			const FSharedBuffer Uncompressed = Reader.Decompress(OffsetCount * sizeof(uint64));
-			CHECK(Algo::Compare(CastToArrayView(Uncompressed), MakeArrayView(ExpectedValues).RightChop(OffsetCount)));
+			const TConstArrayView<uint64> Values = CastToArrayView(Uncompressed);
+			ValidateData(Values, ExpectedValues, OffsetCount);
 		}
 		{
 			FCompressedBufferReaderSourceScope Source(Reader, Compressed);
 			const FSharedBuffer Uncompressed = Reader.Decompress(OffsetCount * sizeof(uint64));
-			CHECK(Algo::Compare(CastToArrayView(Uncompressed), MakeArrayView(ExpectedValues).RightChop(OffsetCount)));
+			const TConstArrayView<uint64> Values = CastToArrayView(Uncompressed);
+			ValidateData(Values, ExpectedValues, OffsetCount);
 		}
 
 		// Short Buffer
@@ -179,8 +232,10 @@ TEST_CASE("Core::Compression::CompressedBuffer::Decompress", "[Core][Compression
 		}
 	}
 
-	SECTION("Decompress only block.")
+	// Only one block
 	{
+		INFO("Decompress Only one block");
+
 		constexpr uint64 BlockSize = 256 * sizeof(uint64);
 		constexpr int32 N = 100;
 		const TArray<uint64> ExpectedValues = GenerateData(N);
@@ -198,17 +253,21 @@ TEST_CASE("Core::Compression::CompressedBuffer::Decompress", "[Core][Compression
 			FBufferReader Ar(const_cast<void*>(Buffer.GetData()), int64(Buffer.GetSize()), /*bFreeOnClose*/ false, /*bIsPersistent*/ true);
 			FCompressedBufferReaderSourceScope Source(Reader, Ar);
 			const FSharedBuffer Uncompressed = Reader.Decompress(OffsetCount * sizeof(uint64), Count * sizeof(uint64));
-			CHECK(Algo::Compare(CastToArrayView(Uncompressed), MakeArrayView(ExpectedValues).Mid(OffsetCount, Count)));
+			const TConstArrayView<uint64> Values = CastToArrayView(Uncompressed);
+			ValidateData(Values, ExpectedValues, OffsetCount);
 		}
 		{
 			FCompressedBufferReaderSourceScope Source(Reader, Compressed);
 			const FSharedBuffer Uncompressed = Reader.Decompress(OffsetCount * sizeof(uint64), Count * sizeof(uint64));
-			CHECK(Algo::Compare(CastToArrayView(Uncompressed), MakeArrayView(ExpectedValues).Mid(OffsetCount, Count)));
+			const TConstArrayView<uint64> Values = CastToArrayView(Uncompressed);
+			ValidateData(Values, ExpectedValues, OffsetCount);
 		}
 	}
 
-	SECTION("Decompress from an uncompressed buffer.")
+	// Uncompressed
 	{
+		INFO("Uncompressed");
+
 		constexpr int32 N = 4242;
 		const TArray<uint64> ExpectedValues = GenerateData(N);
 
@@ -222,18 +281,22 @@ TEST_CASE("Core::Compression::CompressedBuffer::Decompress", "[Core][Compression
 			constexpr uint64 OffsetCount = 0;
 			constexpr uint64 Count = N;
 			const FSharedBuffer Uncompressed = Reader.Decompress(OffsetCount * sizeof(uint64), Count * sizeof(uint64));
-			CHECK(Algo::Compare(CastToArrayView(Uncompressed), MakeArrayView(ExpectedValues).Mid(OffsetCount, Count)));
+			const TConstArrayView<uint64> Values = CastToArrayView(Uncompressed);
+			ValidateData(Values, ExpectedValues, OffsetCount);
 		}
 
 		{
 			constexpr uint64 OffsetCount = 21;
 			constexpr uint64 Count = 999;
 			const FSharedBuffer Uncompressed = Reader.Decompress(OffsetCount * sizeof(uint64), Count * sizeof(uint64));
-			CHECK(Algo::Compare(CastToArrayView(Uncompressed), MakeArrayView(ExpectedValues).Mid(OffsetCount, Count)));
+			const TConstArrayView<uint64> Values = CastToArrayView(Uncompressed);
+			ValidateData(Values, ExpectedValues, OffsetCount);
 		}
 
 		// Short Buffer
 		{
+			INFO("Short Buffer");
+
 			const FCompressedBuffer CompressedShort =
 				FCompressedBuffer::FromCompressed(Compressed.GetCompressed().Mid(0, Compressed.GetCompressedSize() - 128));
 			Reader.SetSource(CompressedShort);
