Add new Unit Tests

diff --git a/Engine/Source/Runtime/Core/Tests/Containers/BitArrayTests.cpp b/Engine/Source/Runtime/Core/Tests/Containers/BitArrayTests.cpp
--- a/Engine/Source/Runtime/Core/Tests/Containers/BitArrayTests.cpp
+++ b/Engine/Source/Runtime/Core/Tests/Containers/BitArrayTests.cpp
@@ -17,7 +17,7 @@ namespace BitArrayTest
 
 		for ( ; MaxNum > 0 && *Bits != '\0'; ++Bits)
 		{
-			check(*Bits == ' ' || *Bits == '0' || *Bits == '1');
+			CHECK((*Bits == ' ' || *Bits == '0' || *Bits == '1'));
 
 			// Skip spaces
 			if (*Bits != ' ')
@@ -66,7 +66,7 @@ namespace BitArrayTest
 } // namespace BitArrayTest
 } // namespace UE
 
-class FBitArrayTest 
+class FBitArrayTest : public FAutomationTestFixture
 {
 public:
 	FBitArrayTest():
@@ -95,22 +95,24 @@ public:
 
 	void ConstructAndTestConstructors()
 	{
+		INFO("Construct And Test Constructors");
+
 		ArrTrue.CheckInvariants();
-		TestEqual(TEXT("ArrTrue Size"), ArrTrue.Num(), 10);
+		CHECK(ArrTrue.Num()==10);
 		ArrFalse.CheckInvariants();
-		TestEqual(TEXT("ArrFalse Size"), ArrFalse.Num(), 10);
+		CHECK(ArrFalse.Num()==10);
 		for (int n = 0; n < 10; ++n)
 		{
-			TestEqual(TEXT("ArrTrue elements"), ArrTrue[n], true);
-			TestEqual(TEXT("ArrFalse elements"), ArrFalse[n], false);
+			CHECK(ArrTrue[n]==true);
+			CHECK(ArrFalse[n]==false);
 		}
 
 		ArrGrowingTrue.CheckInvariants();
-		TestEqual(TEXT("ArrGrowingTrue Size"), ArrGrowingTrue.Num(), NumGrowingTrue);
+		CHECK(ArrGrowingTrue.Num()==NumGrowingTrue);
 		for (int n = 0; n < NumGrowingTrue; ++n)
 		{
 			ArrGrowingTrue[n] = bGrowingTrue[n];
-			TestEqual(TEXT("ArrGrowingTrue elements"), ArrGrowingTrue[n], bGrowingTrue[n]);
+			CHECK(ArrGrowingTrue[n]==bGrowingTrue[n]);
 		}
 		ArrSquareWave.CheckInvariants();
 		for (int n = 0; n < NumSquareWave; ++n)
@@ -127,6 +129,8 @@ public:
 
 	void TestEqualityOperator()
 	{
+		INFO("Test Equality Operator");
+
 		// == and != operators
 		TBitArray<> ArrGrowingTrue2(true, NumGrowingTrue);
 		TBitArray<> ArrAlmostGrowingTrue(true, NumGrowingTrue);
@@ -139,14 +143,14 @@ public:
 		}
 		ArrGrowingTrue2[NumGrowingTrue - 1] = bGrowingTrue[NumGrowingTrue - 1];
 		ArrAlmostGrowingTrue[NumGrowingTrue - 1] = !bGrowingTrue[NumGrowingTrue - 1];
-		TestTrue(TEXT("Equality operator on equal arrays"), ArrGrowingTrue == ArrGrowingTrue2);
-		TestFalse(TEXT("Inequality operator on equal arrays"), ArrGrowingTrue != ArrGrowingTrue2);
-		TestFalse(TEXT("Equality operator on nonequal arrays"), ArrGrowingTrue == ArrAlmostGrowingTrue);
-		TestTrue(TEXT("Inequality operator on nonequal arrays"), ArrGrowingTrue != ArrAlmostGrowingTrue);
-		TestFalse(TEXT("Equality operator when lhs is superset of rhs"), ArrGrowingTrue == ArrSubsetGrowingTrue);
-		TestTrue(TEXT("Inequality operator when lhs is superset of rhs"), ArrGrowingTrue != ArrSubsetGrowingTrue);
-		TestFalse(TEXT("Equality operator when lhs is subset of rhs"), ArrSubsetGrowingTrue == ArrGrowingTrue);
-		TestTrue(TEXT("Inequality operator when lhs is subset of rhs"), ArrSubsetGrowingTrue != ArrGrowingTrue);
+		CHECK(ArrGrowingTrue == ArrGrowingTrue2);
+		CHECK_FALSE(ArrGrowingTrue != ArrGrowingTrue2);
+		CHECK_FALSE(ArrGrowingTrue == ArrAlmostGrowingTrue);
+		CHECK(ArrGrowingTrue != ArrAlmostGrowingTrue);
+		CHECK_FALSE(ArrGrowingTrue == ArrSubsetGrowingTrue);
+		CHECK(ArrGrowingTrue != ArrSubsetGrowingTrue);
+		CHECK_FALSE(ArrSubsetGrowingTrue == ArrGrowingTrue);
+		CHECK(ArrSubsetGrowingTrue != ArrGrowingTrue);
 	}
 
 	void TestOtherConstructorAndAssignment()
@@ -174,98 +178,113 @@ public:
 
 		// Move constructor
 		{
+			INFO("Move constructor");
+
 			TBitArray<> ArrVictim(ArrGrowingTrue);
 			ArrVictim.CheckInvariants();
 			TBitArray<> Arr(MoveTemp(ArrVictim));
 			Arr.CheckInvariants();
-			TestTrue(TEXT("Copy Constructor"), Arr == ArrGrowingTrue);
+			CHECK(Arr == ArrGrowingTrue);
 		}
 
 		// Copy constructor
 		{
+			INFO("Copy constructor");
+
 			TBitArray<> Arr(ArrGrowingTrue);
 			Arr.CheckInvariants();
-			TestTrue(TEXT("Copy Constructor"), Arr == ArrGrowingTrue);
+			CHECK(Arr == ArrGrowingTrue);
 		}
 
 		// Assignment operator
 		{
+			INFO("Assignment operator");
+
 			TBitArray<> Arr;
 			Arr.CheckInvariants();
 			Arr = ArrGrowingTrue;
 			Arr.CheckInvariants();
-			TestTrue(TEXT("Assignment operator"), Arr == ArrGrowingTrue);
+			CHECK(Arr == ArrGrowingTrue);
 		}
 
 		// Move Assignment operator
 		{
+			INFO("Move Assignment operator");
+
 			TBitArray<> ArrVictim(ArrGrowingTrue);
 			ArrVictim.CheckInvariants();
 			TBitArray<> Arr;
 			Arr.CheckInvariants();
 			Arr = MoveTemp(ArrVictim);
-			TestTrue(TEXT("Move Assignment operator"), Arr == ArrGrowingTrue);
+			CHECK(Arr == ArrGrowingTrue);
 		}
 	}
 
 	void TestLessThan()
 	{
+		INFO("Test Less Than");
+
 		// operator<
 		TBitArray<> Short(true, 4);
 		TBitArray<> MediumFalse(false, 5);
 		TBitArray<> MediumTrue(true, 5);
 		TBitArray<> Long(false, 6);
-		TestFalse(TEXT("! x < x"), Short.operator<(Short));
-		TestTrue(TEXT("Sorted by length first, so Short < MediumFalse"), Short < MediumFalse);
-		TestFalse(TEXT("Sorted by length first, so !MediumFalse < Short"), MediumFalse < Short);
-		TestTrue(TEXT("Sorted by length first, so Short < MediumTrue"), Short < MediumTrue);
-		TestFalse(TEXT("Sorted by length first, so !MediumTrue < Short"), MediumTrue < Short);
-		TestTrue(TEXT("Sorted by length first, so MediumTrue < Long"), MediumFalse < Long);
-		TestFalse(TEXT("Sorted by length first, so !Long < MediumTrue"), Long < MediumFalse);
-		TestTrue(TEXT("Sorted by length first, so MediumFalse < Long"), MediumTrue < Long);
-		TestFalse(TEXT("Sorted by length first, so !Long < MediumFalse"), Long < MediumTrue);
+		CHECK_FALSE(Short.operator<(Short));
+		CHECK(Short < MediumFalse);
+		CHECK_FALSE(MediumFalse < Short);
+		CHECK(Short < MediumTrue);
+		CHECK_FALSE(MediumTrue < Short);
+		CHECK(MediumFalse < Long);
+		CHECK_FALSE(Long < MediumFalse);
+		CHECK(MediumTrue < Long);
+		CHECK_FALSE(Long < MediumTrue);
 
 		TBitArray<> MediumTrueAtEnd(false, 5);
 		MediumTrueAtEnd[4] = true;
 		TBitArray<> MediumTrueAtStart(false, 5);
 		MediumTrueAtStart[0] = true;
-		TestTrue(TEXT("Sorted lexigraphically second, so MediumFalse < MediumTrueAtEnd"), MediumFalse < MediumTrueAtEnd);
-		TestTrue(TEXT("Sorted lexigraphically second, so MediumTrueAtEnd < MediumTrueAtStart"), MediumFalse < MediumTrueAtEnd);
-		TestTrue(TEXT("Sorted lexigraphically second, so MediumTrueAtStart < MediumTrue"), MediumTrueAtEnd < MediumTrue);
+		CHECK(MediumFalse < MediumTrueAtEnd);
+		CHECK(MediumFalse < MediumTrueAtEnd);
+		CHECK(MediumTrueAtEnd < MediumTrue);
 	}
 
 	void TestRemoveAt()
 	{
+		INFO("Test Remove At");
+
 		{
 			TBitArray<> Arr(ArrSquareWave);
 			Arr.RemoveAt(Arr.Num() - 1, 1);
 			Arr.CheckInvariants();
-			TestEqual(TEXT("RemoveAt from end size"), Arr.Num(), NumSquareWave - 1);
+			INFO("RemoveAt from end");
+			CHECK(Arr.Num()==(NumSquareWave - 1));
 			for (int n = 0; n < NumSquareWave - 1; ++n)
 			{
-				TestEqual(TEXT("RemoveAt from end elements"), Arr[n], ArrSquareWave[n]);
+				CHECK(Arr[n]==ArrSquareWave[n]);
 			}
 
 			Arr = ArrSquareWave;
 			Arr.RemoveAt(0, 1);
 			Arr.CheckInvariants();
-			TestEqual(TEXT("RemoveAt from start size"), Arr.Num(), NumSquareWave - 1);
+			INFO("RemoveAt from start");
+			CHECK(Arr.Num()==(NumSquareWave - 1));
 			for (int n = 0; n < NumSquareWave - 1; ++n)
 			{
-				TestEqual(TEXT("RemoveAt from start elements"), Arr[n], ArrSquareWave[n + 1]);
+				CHECK(Arr[n]==(ArrSquareWave[n + 1]));
 			}
 
 			Arr = ArrSquareWave;
 			Arr.RemoveAt(5, 1);
 			Arr.CheckInvariants();
-			TestEqual(TEXT("RemoveAt from middle size"), Arr.Num(), NumSquareWave - 1);
+			INFO("RemoveAt from middle");
+			CHECK(Arr.Num()==(NumSquareWave - 1));
 			for (int n = 0; n < 5; ++n)
 			{
-				TestEqual(TEXT("RemoveAt from middle elements"), Arr[n], ArrSquareWave[n]);
+				CHECK(Arr[n]==ArrSquareWave[n]);
 			}
 			for (int n = 5; n < NumSquareWave - 1; ++n)
 			{
-				TestEqual(TEXT("RemoveAt from middle elements"), Arr[n], ArrSquareWave[n + 1]);
+				CHECK(Arr[n]==ArrSquareWave[n + 1]);
 			}
 
 			Arr = TBitArray<>(true, 20);
@@ -275,14 +294,15 @@ public:
 			}
 			Arr.RemoveAt(5, 5);
 			Arr.CheckInvariants();
-			TestEqual(TEXT("RemoveAt multiple size"), Arr.Num(), 15);
+			INFO("RemoveAt multiple");
+			CHECK(Arr.Num()==15);
 			for (int n = 0; n < 5; ++n)
 			{
-				TestEqual(TEXT("RemoveAt multiple elements"), Arr[n], true);
+				CHECK(Arr[n]==true);
 			}
 			for (int n = 5; n < 15; ++n)
 			{
-				TestEqual(TEXT("RemoveAt multiple elements"), Arr[n], false);
+				CHECK(Arr[n]==false);
 			}
 		}
 
@@ -303,34 +323,39 @@ public:
 
 	void TestRemoveAtSwap()
 	{
+		INFO("Test Remove At Swap");
+
 		// RemoveAtSwap
 		{
+
+			INFO("Remove At Swap");
+
 			TBitArray<> Arr(ArrSquareWave);
 			Arr.RemoveAtSwap(Arr.Num() - 1, 1);
 			Arr.CheckInvariants();
-			TestEqual(TEXT("RemoveAtSwap from end size"), Arr.Num(), NumSquareWave - 1);
+			TEST_EQUAL(TEXT("RemoveAtSwap from end size"), Arr.Num(), NumSquareWave - 1);
 			for (int n = 0; n < NumSquareWave - 1; ++n)
 			{
-				TestEqual(TEXT("RemoveAtSwap from end elements"), Arr[n], ArrSquareWave[n]);
+				TEST_EQUAL(TEXT("RemoveAtSwap from end elements"), Arr[n], ArrSquareWave[n]);
 			}
 
 			Arr = ArrSquareWave;
 			Arr.RemoveAtSwap(0, 1);
 			Arr.CheckInvariants();
-			TestEqual(TEXT("RemoveAtSwap from start size"), Arr.Num(), NumSquareWave - 1);
-			TestEqual(TEXT("RemoveAtSwap from start elements"), Arr[0], ArrSquareWave[NumSquareWave - 1]);
+			TEST_EQUAL(TEXT("RemoveAtSwap from start size"), Arr.Num(), NumSquareWave - 1);
+			TEST_EQUAL(TEXT("RemoveAtSwap from start elements"), Arr[0], ArrSquareWave[NumSquareWave - 1]);
 			for (int n = 1; n < NumSquareWave - 1; ++n)
 			{
-				TestEqual(TEXT("RemoveAtSwap from start elements"), Arr[n], ArrSquareWave[n]);
+				TEST_EQUAL(TEXT("RemoveAtSwap from start elements"), Arr[n], ArrSquareWave[n]);
 			}
 
 			Arr = ArrSquareWave;
 			Arr.RemoveAtSwap(5, 1);
 			Arr.CheckInvariants();
-			TestEqual(TEXT("RemoveAtSwap from middle size"), Arr.Num(), NumSquareWave - 1);
+			TEST_EQUAL(TEXT("RemoveAtSwap from middle size"), Arr.Num(), NumSquareWave - 1);
 			for (int n = 0; n < NumSquareWave - 1; ++n)
 			{
-				TestEqual(TEXT("RemoveAtSwap from middle elements"), Arr[n], n != 5 ? ArrSquareWave[n] : ArrSquareWave[NumSquareWave - 1]);
+				TEST_EQUAL(TEXT("RemoveAtSwap from middle elements"), Arr[n], n != 5 ? ArrSquareWave[n] : ArrSquareWave[NumSquareWave - 1]);
 			}
 
 			Arr = TBitArray<>(true, 20);
@@ -340,22 +365,22 @@ public:
 			}
 			Arr.RemoveAtSwap(5, 2);
 			Arr.CheckInvariants();
-			TestEqual(TEXT("RemoveAtSwap, multiple, size"), Arr.Num(), 18);
+			TEST_EQUAL(TEXT("RemoveAtSwap, multiple, size"), Arr.Num(), 18);
 			for (int n = 0; n < 5; ++n)
 			{
-				TestEqual(TEXT("RemoveAtSwap, multiple, elements"), Arr[n], true);
+				TEST_EQUAL(TEXT("RemoveAtSwap, multiple, elements"), Arr[n], true);
 			}
 			for (int n = 5; n < 7; ++n)
 			{
-				TestEqual(TEXT("RemoveAtSwap, multiple, elements"), Arr[n], false);
+				TEST_EQUAL(TEXT("RemoveAtSwap, multiple, elements"), Arr[n], false);
 			}
 			for (int n = 7; n < 10; ++n)
 			{
-				TestEqual(TEXT("RemoveAtSwap, multiple, elements"), Arr[n], true);
+				TEST_EQUAL(TEXT("RemoveAtSwap, multiple, elements"), Arr[n], true);
 			}
 			for (int n = 10; n < 18; ++n)
 			{
-				TestEqual(TEXT("RemoveAtSwap, multiple, elements"), Arr[n], false);
+				TEST_EQUAL(TEXT("RemoveAtSwap, multiple, elements"), Arr[n], false);
 			}
 		}
 
@@ -390,9 +415,9 @@ public:
 			Writer << Spacer << ArrEmpty << Spacer << ArrOnes << Spacer << ArrZeroes << Spacer;
 		}
 		// serialize into a loading archive should not modify the array
-		TestTrue(TEXT("Serialize Empty"), ArrEmpty == ArrEmptyOriginal);
-		TestTrue(TEXT("Serialize Ones"), ArrOnes == ArrOnesOriginal);
-		TestTrue(TEXT("Serialize Ones"), ArrZeroes == ArrZeroesOriginal);
+		TEST_TRUE(TEXT("Serialize Empty"), ArrEmpty == ArrEmptyOriginal);
+		TEST_TRUE(TEXT("Serialize Ones"), ArrOnes == ArrOnesOriginal);
+		TEST_TRUE(TEXT("Serialize Ones"), ArrZeroes == ArrZeroesOriginal);
 
 		TBitArray<> ArrEmptyCopy;
 		TBitArray<> ArrZeroesCopy(false, NumBitsPerDWORD + NumBitsPerDWORD / 2);
@@ -402,16 +427,16 @@ public:
 			FMemoryReader Reader(Bytes);
 			Reader << SpacerCopies[0] << ArrEmptyCopy << SpacerCopies[1] << ArrOnesCopy << SpacerCopies[2] << ArrZeroesCopy << SpacerCopies[3];
 		}
-		TestEqual(TEXT("Serialize Empty Underflow"), SpacerCopies[0], Spacer);
-		TestTrue(TEXT("Serialize Empty"), ArrEmpty == ArrEmptyCopy);
+		TEST_EQUAL(TEXT("Serialize Empty Underflow"), SpacerCopies[0], Spacer);
+		TEST_TRUE(TEXT("Serialize Empty"), ArrEmpty == ArrEmptyCopy);
 		ArrEmptyCopy.CheckInvariants();
-		TestEqual(TEXT("Serialize Ones Underflow"), SpacerCopies[1], Spacer);
-		TestTrue(TEXT("Serialize Ones"), ArrOnes == ArrOnesCopy);
+		TEST_EQUAL(TEXT("Serialize Ones Underflow"), SpacerCopies[1], Spacer);
+		TEST_TRUE(TEXT("Serialize Ones"), ArrOnes == ArrOnesCopy);
 		ArrOnesCopy.CheckInvariants();
-		TestEqual(TEXT("Serialize Zeroes Underflow"), SpacerCopies[2], Spacer);
-		TestTrue(TEXT("Serialize Ones"), ArrZeroes == ArrZeroesCopy);
+		TEST_EQUAL(TEXT("Serialize Zeroes Underflow"), SpacerCopies[2], Spacer);
+		TEST_TRUE(TEXT("Serialize Ones"), ArrZeroes == ArrZeroesCopy);
 		ArrZeroesCopy.CheckInvariants();
-		TestEqual(TEXT("Serialize Zeroes Overflow"), SpacerCopies[3], Spacer);
+		TEST_EQUAL(TEXT("Serialize Zeroes Overflow"), SpacerCopies[3], Spacer);
 
 		TArray<uint8> Bytes2;
 		TBitArray<> ArrSmall(true, 16);
@@ -427,8 +452,8 @@ public:
 			FMemoryReader Reader(Bytes2);
 			Reader << Arr;
 		}
-		TestEqual(TEXT("Serialize from a dynamic allocation with an inline allocator sets num down to the size of the loaded array"), Arr.Num(), ArrSmall.Num());
-		TestEqual(TEXT("Serialize from a dynamic allocation with an inline allocator sets max back to the size of the inline allocation"), Arr.Max(), InitialMax);
+		TEST_EQUAL(TEXT("Serialize from a dynamic allocation with an inline allocator sets num down to the size of the loaded array"), Arr.Num(), ArrSmall.Num());
+		TEST_EQUAL(TEXT("Serialize from a dynamic allocation with an inline allocator sets max back to the size of the inline allocation"), Arr.Max(), InitialMax);
 	}
 
 	void TestAdd()
@@ -437,31 +462,31 @@ public:
 		// Add one bit
 		{
 			TBitArray<FDefaultAllocator> Arr;
-			TestEqual(TEXT("With DefaultAllocator MaxBits starts at 0"), Arr.Max(), 0);
+			TEST_EQUAL(TEXT("With DefaultAllocator MaxBits starts at 0"), Arr.Max(), 0);
 			for (int n = 0; n < 10; ++n)
 			{
 				Arr.Add(n % 3 == 0);
 				Arr.CheckInvariants();
 			}
-			TestEqual(TEXT("Add one bit size"), Arr.Num(), 10);
+			TEST_EQUAL(TEXT("Add one bit size"), Arr.Num(), 10);
 			for (int n = 0; n < 10; ++n)
 			{
-				TestEqual(TEXT("Add one bit elements"), Arr[n], n % 3 == 0);
+				TEST_EQUAL(TEXT("Add one bit elements"), Arr[n], n % 3 == 0);
 			}
 
 			Arr.RemoveAt(0, 10);
 			Arr.CheckInvariants();
-			TestEqual(TEXT("Removed all elements leaves size at 0"), Arr.Num(), 0);
-			TestTrue(TEXT("Removed all elements keeps max at original"), Arr.Max() >= 10);
+			TEST_EQUAL(TEXT("Removed all elements leaves size at 0"), Arr.Num(), 0);
+			TEST_TRUE(TEXT("Removed all elements keeps max at original"), Arr.Max() >= 10);
 			for (int n = 0; n < 10; ++n)
 			{
 				Arr.Add(n % 2 == 0);
 				Arr.CheckInvariants();
 			}
-			TestEqual(TEXT("Add one bit no resize size"), Arr.Num(), 10);
+			TEST_EQUAL(TEXT("Add one bit no resize size"), Arr.Num(), 10);
 			for (int n = 0; n < 10; ++n)
 			{
-				TestEqual(TEXT("Add one bit no resize elements"), Arr[n], n % 2 == 0);
+				TEST_EQUAL(TEXT("Add one bit no resize elements"), Arr[n], n % 2 == 0);
 			}
 		}
 
@@ -473,8 +498,8 @@ public:
 			ArrTrueCopy.CheckInvariants();
 			ArrFalseCopy.Add(false, ArrFalse.Num());
 			ArrFalseCopy.CheckInvariants();
-			TestTrue(TEXT("Add multiple true bits from empty into unallocated space"), AreEqual(ArrTrueCopy, ArrTrue));
-			TestTrue(TEXT("Add multiple false bits from empty into unallocated space"), AreEqual(ArrFalseCopy, ArrFalse));
+			TEST_TRUE(TEXT("Add multiple true bits from empty into unallocated space"), AreEqual(ArrTrueCopy, ArrTrue));
+			TEST_TRUE(TEXT("Add multiple false bits from empty into unallocated space"), AreEqual(ArrFalseCopy, ArrFalse));
 
 			ArrTrueCopy.RemoveAt(0, ArrTrueCopy.Num());
 			ArrTrueCopy.CheckInvariants();
@@ -484,8 +509,8 @@ public:
 			ArrFalseCopy.CheckInvariants();
 			ArrFalseCopy.Add(false, ArrFalse.Num());
 			ArrFalseCopy.CheckInvariants();
-			TestTrue(TEXT("Add multiple true bits from empty into previously-allocated space"), AreEqual(ArrTrueCopy, ArrTrue));
-			TestTrue(TEXT("Add multiple false bits from empty into previously-allocated space"), AreEqual(ArrFalseCopy, ArrFalse));
+			TEST_TRUE(TEXT("Add multiple true bits from empty into previously-allocated space"), AreEqual(ArrTrueCopy, ArrTrue));
+			TEST_TRUE(TEXT("Add multiple false bits from empty into previously-allocated space"), AreEqual(ArrFalseCopy, ArrFalse));
 
 			TBitArray<> ArrTrue2(true, 5);
 			TBitArray<> ArrTrue3(true, 10);
@@ -496,8 +521,8 @@ public:
 			ArrTrue2.CheckInvariants();
 			ArrFalse2.Add(false, 5);
 			ArrFalse2.CheckInvariants();
-			TestTrue(TEXT("Add multiple true bits on non-empty"), ArrTrue2 == ArrTrue3);
-			TestTrue(TEXT("Add multiple false bits on non-empty"), ArrFalse2 == ArrFalse3);
+			TEST_TRUE(TEXT("Add multiple true bits on non-empty"), ArrTrue2 == ArrTrue3);
+			TEST_TRUE(TEXT("Add multiple false bits on non-empty"), ArrFalse2 == ArrFalse3);
 		}
 
 		// AddUninitialized
@@ -508,8 +533,8 @@ public:
 			ArrTrueCopy.CheckInvariants();
 			ArrFalseCopy.AddUninitialized(ArrFalse.Num());
 			ArrFalseCopy.CheckInvariants();
-			TestEqual(TEXT("AddUninitialized multiple true bits from empty into unallocated space"), ArrTrueCopy.Num(), ArrTrue.Num());
-			TestEqual(TEXT("AddUninitialized multiple false bits from empty into unallocated space"), ArrFalseCopy.Num(), ArrFalse.Num());
+			TEST_EQUAL(TEXT("AddUninitialized multiple true bits from empty into unallocated space"), ArrTrueCopy.Num(), ArrTrue.Num());
+			TEST_EQUAL(TEXT("AddUninitialized multiple false bits from empty into unallocated space"), ArrFalseCopy.Num(), ArrFalse.Num());
 
 			ArrTrueCopy.RemoveAt(0, ArrTrueCopy.Num());
 			ArrTrueCopy.CheckInvariants();
@@ -519,8 +544,8 @@ public:
 			ArrFalseCopy.CheckInvariants();
 			ArrFalseCopy.AddUninitialized(ArrFalse.Num());
 			ArrFalseCopy.CheckInvariants();
-			TestEqual(TEXT("AddUninitialized multiple true bits from empty into previously-allocated space"), ArrTrueCopy.Num(), ArrTrue.Num());
-			TestEqual(TEXT("AddUninitialized multiple false bits from empty into previously-allocated space"), ArrFalseCopy.Num(), ArrFalse.Num());
+			TEST_EQUAL(TEXT("AddUninitialized multiple true bits from empty into previously-allocated space"), ArrTrueCopy.Num(), ArrTrue.Num());
+			TEST_EQUAL(TEXT("AddUninitialized multiple false bits from empty into previously-allocated space"), ArrFalseCopy.Num(), ArrFalse.Num());
 
 			TBitArray<> ArrTrue2(true, 5);
 			TBitArray<> ArrTrue3(true, 10);
@@ -533,8 +558,8 @@ public:
 			ArrFalse2.CheckInvariants();
 			for (int n = 0; n < 5; ++n)
 			{
-				TestTrue(TEXT("AddUninitialized multiple true bits on non-empty"), ArrTrue2[n]);
-				TestFalse(TEXT("AddUninitialized multiple false bits on non-empty"), ArrFalse2[n]);
+				TEST_TRUE(TEXT("AddUninitialized multiple true bits on non-empty"), ArrTrue2[n]);
+				TEST_FALSE(TEXT("AddUninitialized multiple false bits on non-empty"), ArrFalse2[n]);
 			}
 		}
 	}
@@ -546,11 +571,11 @@ public:
 			TBitArray<> Arr;
 			Arr.AddRange(&GrowingTrueInt, NumGrowingTrue);
 			Arr.CheckInvariants();
-			TestTrue(TEXT("Add from uint32 with ReadOffset 0 from empty"), Arr == ArrGrowingTrue);
+			TEST_TRUE(TEXT("Add from uint32 with ReadOffset 0 from empty"), Arr == ArrGrowingTrue);
 
 			Arr.AddRange(&GrowingTrueInt, NumGrowingTrue);
 			Arr.CheckInvariants();
-			TestTrue(TEXT("Add from uint32 with ReadOffset 0 to nonempty"), Arr == ArrGrowingTrueTwice);
+			TEST_TRUE(TEXT("Add from uint32 with ReadOffset 0 to nonempty"), Arr == ArrGrowingTrueTwice);
 
 			TBitArray<> Arr2;
 			uint32 AllZeroes = 0;
@@ -561,38 +586,38 @@ public:
 			Arr2.CheckInvariants();
 			for (int n = 0; n < 20; ++n)
 			{
-				TestEqual(TEXT("Add from uint32 with ReadOffset 0, Zeroes, Then Ones"), Arr2[n], n >= 10);
+				TEST_EQUAL(TEXT("Add from uint32 with ReadOffset 0, Zeroes, Then Ones"), Arr2[n], n >= 10);
 			}
 
 			TBitArray<> Arr3;
 			uint32 MultipleInts[] = { 0xffff0000, 0x0f0f0f0f };
 			Arr3.AddRange(MultipleInts, 64);
 			Arr3.CheckInvariants();
-			TestEqual(TEXT("Add from uint32 with ReadOffset 0, size"), Arr3.Num(), 64);
+			TEST_EQUAL(TEXT("Add from uint32 with ReadOffset 0, size"), Arr3.Num(), 64);
 			for (int n = 0; n < 32; ++n)
 			{
 				Label.Reset();
-				TestEqual(Label.Appendf(TEXT("Add from uint32 with ReadOffset 0, MultipleInts, %d"), n).ToString(), Arr3[n], n >= 16);
+				TEST_EQUAL(Label.Appendf(TEXT("Add from uint32 with ReadOffset 0, MultipleInts, %d"), n).ToString(), Arr3[n], n >= 16);
 			}
 			for (int n = 32; n < 64; ++n)
 			{
 				Label.Reset();
-				TestEqual(Label.Appendf(TEXT("Add from uint32 with ReadOffset 0, MultipleInts, %d"), n).ToString(), Arr3[n], (n / 4) % 2 == 0);
+				TEST_EQUAL(Label.Appendf(TEXT("Add from uint32 with ReadOffset 0, MultipleInts, %d"), n).ToString(), Arr3[n], (n / 4) % 2 == 0);
 			}
 
 			TBitArray<> Arr4;
 			Arr4.AddRange(MultipleInts, 32, 16);
 			Arr4.CheckInvariants();
-			TestEqual(TEXT("Add from uint32 with ReadOffset 16, MultipleInts, size"), Arr4.Num(), 32);
+			TEST_EQUAL(TEXT("Add from uint32 with ReadOffset 16, MultipleInts, size"), Arr4.Num(), 32);
 			for (int n = 0; n < 16; ++n)
 			{
 				Label.Reset();
-				TestEqual(Label.Appendf(TEXT("Add from uint32 with ReadOffset 16, MultipleInts, %d"), n).ToString(), Arr4[n], true);
+				TEST_EQUAL(Label.Appendf(TEXT("Add from uint32 with ReadOffset 16, MultipleInts, %d"), n).ToString(), Arr4[n], true);
 			}
 			for (int n = 16; n < 32; ++n)
 			{
 				Label.Reset();
-				TestEqual(Label.Appendf(TEXT("Add from uint32 with ReadOffset 16, MultipleInts, %d"), n).ToString(), Arr4[n], (n / 4) % 2 == 0);
+				TEST_EQUAL(Label.Appendf(TEXT("Add from uint32 with ReadOffset 16, MultipleInts, %d"), n).ToString(), Arr4[n], (n / 4) % 2 == 0);
 			}
 		}
 
@@ -601,11 +626,11 @@ public:
 			TBitArray<> Arr;
 			Arr.AddRange(ArrGrowingTrue, NumGrowingTrue);
 			Arr.CheckInvariants();
-			TestTrue(TEXT("Add from BitArray with ReadOffset 0 from empty"), Arr == ArrGrowingTrue);
+			TEST_TRUE(TEXT("Add from BitArray with ReadOffset 0 from empty"), Arr == ArrGrowingTrue);
 
 			Arr.AddRange(ArrGrowingTrue, NumGrowingTrue);
 			Arr.CheckInvariants();
-			TestTrue(TEXT("Add from BitArray with ReadOffset 0 to nonempty"), Arr == ArrGrowingTrueTwice);
+			TEST_TRUE(TEXT("Add from BitArray with ReadOffset 0 to nonempty"), Arr == ArrGrowingTrueTwice);
 
 			TBitArray<> Arr2;
 			Arr2.AddRange(ArrFalse, 10);
@@ -614,7 +639,7 @@ public:
 			Arr2.CheckInvariants();
 			for (int n = 0; n < 20; ++n)
 			{
-				TestEqual(TEXT("Add from BitArray with ReadOffset 0, Zeroes, Then Ones"), Arr2[n], n >= 10);
+				TEST_EQUAL(TEXT("Add from BitArray with ReadOffset 0, Zeroes, Then Ones"), Arr2[n], n >= 10);
 			}
 
 			TBitArray<> Arr3;
@@ -624,21 +649,21 @@ public:
 			ArrMultipleInts.CheckInvariants();
 			Arr3.AddRange(ArrMultipleInts, 64);
 			Arr3.CheckInvariants();
-			TestTrue(TEXT("Add from BitArray with ReadOffset 0"), Arr3 == ArrMultipleInts);
+			TEST_TRUE(TEXT("Add from BitArray with ReadOffset 0"), Arr3 == ArrMultipleInts);
 
 			TBitArray<> Arr4;
 			Arr4.AddRange(ArrMultipleInts, 32, 16);
 			Arr4.CheckInvariants();
-			TestEqual(TEXT("Add from BitArray with ReadOffset 16, MultipleInts, size"), Arr4.Num(), 32);
+			TEST_EQUAL(TEXT("Add from BitArray with ReadOffset 16, MultipleInts, size"), Arr4.Num(), 32);
 			for (int n = 0; n < 16; ++n)
 			{
 				Label.Reset();
-				TestEqual(Label.Appendf(TEXT("Add from BitArray with ReadOffset 16, MultipleInts, %d"), n).ToString(), Arr4[n], true);
+				TEST_EQUAL(Label.Appendf(TEXT("Add from BitArray with ReadOffset 16, MultipleInts, %d"), n).ToString(), Arr4[n], true);
 			}
 			for (int n = 16; n < 32; ++n)
 			{
 				Label.Reset();
-				TestEqual(Label.Appendf(TEXT("Add from BitArray with ReadOffset 16, MultipleInts, %d"), n).ToString(), Arr4[n], (n / 4) % 2 == 0);
+				TEST_EQUAL(Label.Appendf(TEXT("Add from BitArray with ReadOffset 16, MultipleInts, %d"), n).ToString(), Arr4[n], (n / 4) % 2 == 0);
 			}
 		}
 	}
@@ -649,63 +674,63 @@ public:
 		// Insert one bit
 		{
 			TBitArray<FDefaultAllocator> Arr;
-			TestEqual(TEXT("With DefaultAllocator MaxBits starts at 0"), Arr.Max(), 0);
+			TEST_EQUAL(TEXT("With DefaultAllocator MaxBits starts at 0"), Arr.Max(), 0);
 			for (int n = 0; n < 10; ++n)
 			{
 				Arr.Insert(n % 3 == 0, Arr.Num());
 				Arr.CheckInvariants();
 			}
-			TestEqual(TEXT("Insert one bit at end size"), Arr.Num(), 10);
+			TEST_EQUAL(TEXT("Insert one bit at end size"), Arr.Num(), 10);
 			for (int n = 0; n < 10; ++n)
 			{
-				TestEqual(TEXT("Insert one bit at end elements"), Arr[n], n % 3 == 0);
+				TEST_EQUAL(TEXT("Insert one bit at end elements"), Arr[n], n % 3 == 0);
 			}
 
 			Arr.Insert(false, 5);
 			Arr.CheckInvariants();
 			Arr.Insert(true, 5);
 			Arr.CheckInvariants();
-			TestEqual(TEXT("Insert one bit in middle size"), Arr.Num(), 12);
+			TEST_EQUAL(TEXT("Insert one bit in middle size"), Arr.Num(), 12);
 			for (int n = 0; n < 5; ++n)
 			{
-				TestEqual(TEXT("Insert one bit in middle elements"), Arr[n], n % 3 == 0);
+				TEST_EQUAL(TEXT("Insert one bit in middle elements"), Arr[n], n % 3 == 0);
 			}
-			TestEqual(TEXT("Insert one bit in middle elements"), Arr[5], true);
-			TestEqual(TEXT("Insert one bit in middle elements"), Arr[6], false);
+			TEST_EQUAL(TEXT("Insert one bit in middle elements"), Arr[5], true);
+			TEST_EQUAL(TEXT("Insert one bit in middle elements"), Arr[6], false);
 			for (int n = 7; n < 12; ++n)
 			{
-				TestEqual(TEXT("Insert one bit in middle elements"), Arr[n], (n - 2) % 3 == 0);
+				TEST_EQUAL(TEXT("Insert one bit in middle elements"), Arr[n], (n - 2) % 3 == 0);
 			}
 
 			Arr.RemoveAt(0, 12);
 			Arr.CheckInvariants();
-			TestEqual(TEXT("Removed all elements leaves size at 0"), Arr.Num(), 0);
-			TestTrue(TEXT("Removed all elements keeps max at original"), Arr.Max() >= 12);
+			TEST_EQUAL(TEXT("Removed all elements leaves size at 0"), Arr.Num(), 0);
+			TEST_TRUE(TEXT("Removed all elements keeps max at original"), Arr.Max() >= 12);
 			for (int n = 0; n < 10; ++n)
 			{
 				Arr.Insert(n % 2 == 0, Arr.Num());
 				Arr.CheckInvariants();
 			}
-			TestEqual(TEXT("Insert one bit at end no resize size"), Arr.Num(), 10);
+			TEST_EQUAL(TEXT("Insert one bit at end no resize size"), Arr.Num(), 10);
 			for (int n = 0; n < 10; ++n)
 			{
-				TestEqual(TEXT("Insert one bit at end no resize elements"), Arr[n], n % 2 == 0);
+				TEST_EQUAL(TEXT("Insert one bit at end no resize elements"), Arr[n], n % 2 == 0);
 			}
 
 			Arr.Insert(false, 5);
 			Arr.CheckInvariants();
 			Arr.Insert(true, 5);
 			Arr.CheckInvariants();
-			TestEqual(TEXT("Insert one bit in middle no resize size"), Arr.Num(), 12);
+			TEST_EQUAL(TEXT("Insert one bit in middle no resize size"), Arr.Num(), 12);
 			for (int n = 0; n < 5; ++n)
 			{
-				TestEqual(TEXT("Insert one bit in middle elements"), Arr[n], n % 2 == 0);
+				TEST_EQUAL(TEXT("Insert one bit in middle elements"), Arr[n], n % 2 == 0);
 			}
-			TestEqual(TEXT("Insert one bit in middle no resize elements"), Arr[5], true);
-			TestEqual(TEXT("Insert one bit in middle no resize elements"), Arr[6], false);
+			TEST_EQUAL(TEXT("Insert one bit in middle no resize elements"), Arr[5], true);
+			TEST_EQUAL(TEXT("Insert one bit in middle no resize elements"), Arr[6], false);
 			for (int n = 7; n < 12; ++n)
 			{
-				TestEqual(TEXT("Insert one bit in middle no resize elements"), Arr[n], (n - 2) % 2 == 0);
+				TEST_EQUAL(TEXT("Insert one bit in middle no resize elements"), Arr[n], (n - 2) % 2 == 0);
 			}
 		}
 
@@ -717,23 +742,23 @@ public:
 			ArrTrueCopy.CheckInvariants();
 			ArrFalseCopy.Insert(false, 0, ArrFalse.Num());
 			ArrFalseCopy.CheckInvariants();
-			TestTrue(TEXT("Insert multiple true bits at end from empty into unallocated space"), AreEqual(ArrTrueCopy, ArrTrue));
-			TestTrue(TEXT("Insert multiple false bits at end from empty into unallocated space"), AreEqual(ArrFalseCopy, ArrFalse));
+			TEST_TRUE(TEXT("Insert multiple true bits at end from empty into unallocated space"), AreEqual(ArrTrueCopy, ArrTrue));
+			TEST_TRUE(TEXT("Insert multiple false bits at end from empty into unallocated space"), AreEqual(ArrFalseCopy, ArrFalse));
 
 			ArrTrueCopy.Insert(false, 5, 5);
 			ArrTrueCopy.CheckInvariants();
-			TestEqual(TEXT("Insert multiple bits in middle from empty into unallocated space size"), ArrTrueCopy.Num(), ArrTrue.Num() + 5);
+			TEST_EQUAL(TEXT("Insert multiple bits in middle from empty into unallocated space size"), ArrTrueCopy.Num(), ArrTrue.Num() + 5);
 			for (int n = 0; n < 5; ++n)
 			{
-				TestEqual(TEXT("Insert multiple bits in middle from empty into unallocated space elements"), ArrTrueCopy[n], true);
+				TEST_EQUAL(TEXT("Insert multiple bits in middle from empty into unallocated space elements"), ArrTrueCopy[n], true);
 			}
 			for (int n = 5; n < 10; ++n)
 			{
-				TestEqual(TEXT("Insert multiple bits in middle from empty into unallocated space elements"), ArrTrueCopy[n], false);
+				TEST_EQUAL(TEXT("Insert multiple bits in middle from empty into unallocated space elements"), ArrTrueCopy[n], false);
 			}
 			for (int n = 10; n < 5 + ArrTrue.Num(); ++n)
 			{
-				TestEqual(TEXT("Insert multiple bits in middle from empty into unallocated space elements"), ArrTrueCopy[n], true);
+				TEST_EQUAL(TEXT("Insert multiple bits in middle from empty into unallocated space elements"), ArrTrueCopy[n], true);
 			}
 
 			ArrTrueCopy.RemoveAt(0, ArrTrueCopy.Num());
@@ -744,24 +769,24 @@ public:
 			ArrFalseCopy.CheckInvariants();
 			ArrFalseCopy.Insert(false, 0, ArrFalse.Num());
 			ArrFalseCopy.CheckInvariants();
-			TestTrue(TEXT("Insert multiple true bits at end from empty into previously-allocated space"), AreEqual(ArrTrueCopy, ArrTrue));
-			TestTrue(TEXT("Insert multiple false bits at end from empty into previously-allocated space"), AreEqual(ArrFalseCopy, ArrFalse));
+			TEST_TRUE(TEXT("Insert multiple true bits at end from empty into previously-allocated space"), AreEqual(ArrTrueCopy, ArrTrue));
+			TEST_TRUE(TEXT("Insert multiple false bits at end from empty into previously-allocated space"), AreEqual(ArrFalseCopy, ArrFalse));
 			ArrTrueCopy.Insert(true, 5, 3);
 			ArrTrueCopy.CheckInvariants();
 			ArrTrueCopy.Insert(false, 5, 2);
 			ArrTrueCopy.CheckInvariants();
-			TestEqual(TEXT("Insert multiple bits in middle from empty into previously-allocated size"), ArrTrueCopy.Num(), ArrTrue.Num() + 5);
+			TEST_EQUAL(TEXT("Insert multiple bits in middle from empty into previously-allocated size"), ArrTrueCopy.Num(), ArrTrue.Num() + 5);
 			for (int n = 0; n < 5; ++n)
 			{
-				TestEqual(TEXT("Insert multiple bits in middle from empty into previously-allocated elements"), ArrTrueCopy[n], true);
+				TEST_EQUAL(TEXT("Insert multiple bits in middle from empty into previously-allocated elements"), ArrTrueCopy[n], true);
 			}
 			for (int n = 5; n < 7; ++n)
 			{
-				TestEqual(TEXT("Insert multiple bits in middle from empty into previously-allocated elements"), ArrTrueCopy[n], false);
+				TEST_EQUAL(TEXT("Insert multiple bits in middle from empty into previously-allocated elements"), ArrTrueCopy[n], false);
 			}
 			for (int n = 7; n < 5 + ArrTrue.Num(); ++n)
 			{
-				TestEqual(TEXT("Insert multiple bits in middle from empty into previously-allocated elements"), ArrTrueCopy[n], true);
+				TEST_EQUAL(TEXT("Insert multiple bits in middle from empty into previously-allocated elements"), ArrTrueCopy[n], true);
 			}
 
 			TBitArray<> ArrTrue2(true, 5);
@@ -773,8 +798,8 @@ public:
 			ArrTrue2.CheckInvariants();
 			ArrFalse2.Insert(false, ArrFalse2.Num(), 5);
 			ArrFalse2.CheckInvariants();
-			TestTrue(TEXT("Insert multiple true bits at end on non-empty"), ArrTrue2 == ArrTrue3);
-			TestTrue(TEXT("Insert multiple false bits at end on non-empty"), ArrFalse2 == ArrFalse3);
+			TEST_TRUE(TEXT("Insert multiple true bits at end on non-empty"), ArrTrue2 == ArrTrue3);
+			TEST_TRUE(TEXT("Insert multiple false bits at end on non-empty"), ArrFalse2 == ArrFalse3);
 
 			TBitArray<> ArrTrue4(true, 5);
 			TBitArray<> ArrFalse4(false, 5);
@@ -783,8 +808,8 @@ public:
 			ArrTrue4.CheckInvariants();
 			ArrFalse4.Insert(false, 1, 5);
 			ArrFalse4.CheckInvariants();
-			TestTrue(TEXT("Insert multiple true bits at middle on non-empty"), ArrTrue4 == ArrTrue3);
-			TestTrue(TEXT("Insert multiple false bits at middle on non-empty"), ArrFalse4 == ArrFalse3);
+			TEST_TRUE(TEXT("Insert multiple true bits at middle on non-empty"), ArrTrue4 == ArrTrue3);
+			TEST_TRUE(TEXT("Insert multiple false bits at middle on non-empty"), ArrFalse4 == ArrFalse3);
 		}
 
 		// InsertUninitialized
@@ -792,7 +817,7 @@ public:
 			TBitArray<FDefaultAllocator> Arr;
 			Arr.InsertUninitialized(0, ArrTrue.Num());
 			Arr.CheckInvariants();
-			TestEqual(TEXT("InsertUninitialized multiple bits at end from empty into unallocated space"), Arr.Num(), ArrTrue.Num());
+			TEST_EQUAL(TEXT("InsertUninitialized multiple bits at end from empty into unallocated space"), Arr.Num(), ArrTrue.Num());
 
 			TBitArray<FDefaultAllocator> ArrTrueDefaultAlloc(true, ArrTrue.Num());
 			TBitArray<FDefaultAllocator> ArrFalseDefaultAlloc(false, ArrFalse.Num());
@@ -804,16 +829,16 @@ public:
 			ArrTrueCopy.CheckInvariants();
 			ArrFalseCopy.InsertUninitialized(5, 5);
 			ArrFalseCopy.CheckInvariants();
-			TestEqual(TEXT("InsertUninitialized multiple bits in middle from empty into unallocated space size"), ArrTrueCopy.Num(), ArrTrue.Num() + 5);
+			TEST_EQUAL(TEXT("InsertUninitialized multiple bits in middle from empty into unallocated space size"), ArrTrueCopy.Num(), ArrTrue.Num() + 5);
 			for (int n = 0; n < 5; ++n)
 			{
-				TestEqual(TEXT("InsertUninitialized multiple bits in middle from empty into unallocated space elements"), ArrTrueCopy[n], true);
-				TestEqual(TEXT("InsertUninitialized multiple bits in middle from empty into unallocated space elements"), ArrFalseCopy[n], false);
+				TEST_EQUAL(TEXT("InsertUninitialized multiple bits in middle from empty into unallocated space elements"), ArrTrueCopy[n], true);
+				TEST_EQUAL(TEXT("InsertUninitialized multiple bits in middle from empty into unallocated space elements"), ArrFalseCopy[n], false);
 			}
 			for (int n = 10; n < 5 + ArrTrue.Num(); ++n)
 			{
-				TestEqual(TEXT("InsertUninitialized multiple bits in middle from empty into unallocated space elements"), ArrTrueCopy[n], true);
-				TestEqual(TEXT("InsertUninitialized multiple bits in middle from empty into unallocated space elements"), ArrFalseCopy[n], false);
+				TEST_EQUAL(TEXT("InsertUninitialized multiple bits in middle from empty into unallocated space elements"), ArrTrueCopy[n], true);
+				TEST_EQUAL(TEXT("InsertUninitialized multiple bits in middle from empty into unallocated space elements"), ArrFalseCopy[n], false);
 			}
 
 			ArrTrueCopy.RemoveAt(0, ArrTrueCopy.Num());
@@ -824,22 +849,22 @@ public:
 			ArrFalseCopy.CheckInvariants();
 			ArrFalseCopy.InsertUninitialized(0, ArrFalse.Num());
 			ArrFalseCopy.CheckInvariants();
-			TestEqual(TEXT("InsertUninitialized multiple true bits at end from empty into previously-allocated space"), ArrTrueCopy.Num(), ArrTrue.Num());
-			TestEqual(TEXT("InsertUninitialized multiple false bits at end from empty into previously-allocated space"), ArrFalseCopy.Num(), ArrFalse.Num());
+			TEST_EQUAL(TEXT("InsertUninitialized multiple true bits at end from empty into previously-allocated space"), ArrTrueCopy.Num(), ArrTrue.Num());
+			TEST_EQUAL(TEXT("InsertUninitialized multiple false bits at end from empty into previously-allocated space"), ArrFalseCopy.Num(), ArrFalse.Num());
 			ArrTrueCopy.RemoveAt(0, ArrTrueCopy.Num());
 			ArrTrueCopy.CheckInvariants();
 			ArrTrueCopy.Insert(true, 0, ArrTrue.Num());
 			ArrTrueCopy.CheckInvariants();
 			ArrTrueCopy.InsertUninitialized(5, 5);
 			ArrTrueCopy.CheckInvariants();
-			TestEqual(TEXT("InsertUninitialized multiple bits in middle from empty into previously-allocated size"), ArrTrueCopy.Num(), ArrTrue.Num() + 5);
+			TEST_EQUAL(TEXT("InsertUninitialized multiple bits in middle from empty into previously-allocated size"), ArrTrueCopy.Num(), ArrTrue.Num() + 5);
 			for (int n = 0; n < 5; ++n)
 			{
-				TestEqual(TEXT("InsertUninitialized multiple bits in middle from empty into previously-allocated elements"), ArrTrueCopy[n], true);
+				TEST_EQUAL(TEXT("InsertUninitialized multiple bits in middle from empty into previously-allocated elements"), ArrTrueCopy[n], true);
 			}
 			for (int n = 10; n < 5 + ArrTrue.Num(); ++n)
 			{
-				TestEqual(TEXT("InsertUninitialized multiple bits in middle from empty into previously-allocated elements"), ArrTrueCopy[n], true);
+				TEST_EQUAL(TEXT("InsertUninitialized multiple bits in middle from empty into previously-allocated elements"), ArrTrueCopy[n], true);
 			}
 
 			TBitArray<> ArrTrue2(true, 5);
@@ -851,8 +876,8 @@ public:
 			ArrFalse2.CheckInvariants();
 			for (int n = 0; n < 5; ++n)
 			{
-				TestEqual(TEXT("InsertUninitialized multiple true bits at end on non-empty"), ArrTrue2[n], true);
-				TestEqual(TEXT("InsertUninitialized multiple false bits at end on non-empty"), ArrFalse2[n], false);
+				TEST_EQUAL(TEXT("InsertUninitialized multiple true bits at end on non-empty"), ArrTrue2[n], true);
+				TEST_EQUAL(TEXT("InsertUninitialized multiple false bits at end on non-empty"), ArrFalse2[n], false);
 			}
 
 			TBitArray<> ArrTrue4(true, 5);
@@ -866,8 +891,8 @@ public:
 			{
 				if (n < 1 || n >= 6)
 				{
-					TestEqual(TEXT("InsertUninitialized multiple true bits at end on non-empty"), ArrTrue4[n], true);
-					TestEqual(TEXT("InsertUninitialized multiple false bits at end on non-empty"), ArrFalse4[n], false);
+					TEST_EQUAL(TEXT("InsertUninitialized multiple true bits at end on non-empty"), ArrTrue4[n], true);
+					TEST_EQUAL(TEXT("InsertUninitialized multiple false bits at end on non-empty"), ArrFalse4[n], false);
 				}
 			}
 		}
@@ -880,18 +905,18 @@ public:
 			TBitArray<> Arr;
 			Arr.InsertRange(&GrowingTrueInt, 0, NumGrowingTrue);
 			Arr.CheckInvariants();
-			TestTrue(TEXT("Insert from uint32 with ReadOffset 0 at end from empty"), Arr == ArrGrowingTrue);
+			TEST_TRUE(TEXT("Insert from uint32 with ReadOffset 0 at end from empty"), Arr == ArrGrowingTrue);
 
 			Arr.InsertRange(&GrowingTrueInt, Arr.Num(), NumGrowingTrue);
 			Arr.CheckInvariants();
-			TestTrue(TEXT("Insert from uint32 with ReadOffset 0 at end to nonempty"), Arr == ArrGrowingTrueTwice);
+			TEST_TRUE(TEXT("Insert from uint32 with ReadOffset 0 at end to nonempty"), Arr == ArrGrowingTrueTwice);
 
 			TBitArray<> Arr1Insert2;
 			Arr1Insert2.InsertRange(&GrowingTrueInt, 0, NumGrowingTrue);
 			Arr1Insert2.CheckInvariants();
 			Arr1Insert2.InsertRange(&GrowingTrueInt, 0, NumGrowingTrue);
 			Arr1Insert2.CheckInvariants();
-			TestTrue(TEXT("Insert from uint32 with ReadOffset 0 at beginning to nonempty"), Arr1Insert2 == ArrGrowingTrueTwice);
+			TEST_TRUE(TEXT("Insert from uint32 with ReadOffset 0 at beginning to nonempty"), Arr1Insert2 == ArrGrowingTrueTwice);
 
 			uint32 AllZeroes = 0;
 			uint32 AllOnes = 0xffffffff;
@@ -900,18 +925,18 @@ public:
 			Arr1Insert3.CheckInvariants();
 			Arr1Insert3.InsertRange(&AllZeroes, 5, 5);
 			Arr1Insert3.CheckInvariants();
-			TestEqual(TEXT("Insert from uint32 with ReadOffset 0 at middle to nonempty size"), Arr1Insert3.Num(), NumGrowingTrue + 5);
+			TEST_EQUAL(TEXT("Insert from uint32 with ReadOffset 0 at middle to nonempty size"), Arr1Insert3.Num(), NumGrowingTrue + 5);
 			for (int n = 0; n < 5; ++n)
 			{
-				TestEqual(TEXT("Insert from uint32 with ReadOffset 0 at middle to nonempty"), Arr1Insert3[n], ArrGrowingTrue[n]);
+				TEST_EQUAL(TEXT("Insert from uint32 with ReadOffset 0 at middle to nonempty"), Arr1Insert3[n], ArrGrowingTrue[n]);
 			}
 			for (int n = 5; n < 10; ++n)
 			{
-				TestEqual(TEXT("Insert from uint32 with ReadOffset 0 at middle to nonempty"), Arr1Insert3[n], false);
+				TEST_EQUAL(TEXT("Insert from uint32 with ReadOffset 0 at middle to nonempty"), Arr1Insert3[n], false);
 			}
 			for (int n = 10; n < NumGrowingTrue + 5; ++n)
 			{
-				TestEqual(TEXT("Insert from uint32 with ReadOffset 0 at middle to nonempty"), Arr1Insert3[n], ArrGrowingTrue[n - 5]);
+				TEST_EQUAL(TEXT("Insert from uint32 with ReadOffset 0 at middle to nonempty"), Arr1Insert3[n], ArrGrowingTrue[n - 5]);
 			}
 
 			TBitArray<> Arr2;
@@ -921,7 +946,7 @@ public:
 			Arr2.CheckInvariants();
 			for (int n = 0; n < 20; ++n)
 			{
-				TestEqual(TEXT("Insert from uint32 with ReadOffset 0 at end, Zeroes, Then Ones"), Arr2[n], n >= 10);
+				TEST_EQUAL(TEXT("Insert from uint32 with ReadOffset 0 at end, Zeroes, Then Ones"), Arr2[n], n >= 10);
 			}
 
 			TBitArray<> Arr2Insert1;
@@ -931,7 +956,7 @@ public:
 			Arr2Insert1.CheckInvariants();
 			for (int n = 0; n < 20; ++n)
 			{
-				TestEqual(TEXT("Insert from uint32 with ReadOffset 0 at beginning, Zeroes, Then Ones"), Arr2Insert1[n], n < 10);
+				TEST_EQUAL(TEXT("Insert from uint32 with ReadOffset 0 at beginning, Zeroes, Then Ones"), Arr2Insert1[n], n < 10);
 			}
 			TBitArray<> Arr2Insert2;
 			Arr2Insert2.InsertRange(&AllZeroes, 0, 10);
@@ -940,46 +965,46 @@ public:
 			Arr2Insert2.CheckInvariants();
 			for (int n = 0; n < 5; ++n)
 			{
-				TestEqual(TEXT("Insert from uint32 with ReadOffset 0 at middle, Zeroes, Then Ones"), Arr2Insert2[n], false);
+				TEST_EQUAL(TEXT("Insert from uint32 with ReadOffset 0 at middle, Zeroes, Then Ones"), Arr2Insert2[n], false);
 			}
 			for (int n = 5; n < 15; ++n)
 			{
-				TestEqual(TEXT("Insert from uint32 with ReadOffset 0 at middle, Zeroes, Then Ones"), Arr2Insert2[n], true);
+				TEST_EQUAL(TEXT("Insert from uint32 with ReadOffset 0 at middle, Zeroes, Then Ones"), Arr2Insert2[n], true);
 			}
 			for (int n = 15; n < 20; ++n)
 			{
-				TestEqual(TEXT("Insert from uint32 with ReadOffset 0 at middle, Zeroes, Then Ones"), Arr2Insert2[n], false);
+				TEST_EQUAL(TEXT("Insert from uint32 with ReadOffset 0 at middle, Zeroes, Then Ones"), Arr2Insert2[n], false);
 			}
 
 			TBitArray<> Arr3;
 			uint32 MultipleInts[] = { 0xffff0000, 0x0f0f0f0f };
 			Arr3.InsertRange(MultipleInts, 0, 64);
 			Arr3.CheckInvariants();
-			TestEqual(TEXT("Insert from uint32 with ReadOffset 0 at end, MutlipleInts, size"), Arr3.Num(), 64);
+			TEST_EQUAL(TEXT("Insert from uint32 with ReadOffset 0 at end, MutlipleInts, size"), Arr3.Num(), 64);
 			for (int n = 0; n < 32; ++n)
 			{
 				Label.Reset();
-				TestEqual(Label.Appendf(TEXT("Insert from uint32 with ReadOffset 0 at end, MultipleInts, %d"), n).ToString(), Arr3[n], n >= 16);
+				TEST_EQUAL(Label.Appendf(TEXT("Insert from uint32 with ReadOffset 0 at end, MultipleInts, %d"), n).ToString(), Arr3[n], n >= 16);
 			}
 			for (int n = 32; n < 64; ++n)
 			{
 				Label.Reset();
-				TestEqual(Label.Appendf(TEXT("Insert from uint32 with ReadOffset 0 at end, MultipleInts, %d"), n).ToString(), Arr3[n], (n / 4) % 2 == 0);
+				TEST_EQUAL(Label.Appendf(TEXT("Insert from uint32 with ReadOffset 0 at end, MultipleInts, %d"), n).ToString(), Arr3[n], (n / 4) % 2 == 0);
 			}
 
 			TBitArray<> Arr4;
 			Arr4.InsertRange(MultipleInts, 0, 32, 16);
 			Arr4.CheckInvariants();
-			TestEqual(TEXT("Insert from uint32 with ReadOffset 16 at end, MultipleInts, size"), Arr4.Num(), 32);
+			TEST_EQUAL(TEXT("Insert from uint32 with ReadOffset 16 at end, MultipleInts, size"), Arr4.Num(), 32);
 			for (int n = 0; n < 16; ++n)
 			{
 				Label.Reset();
-				TestEqual(Label.Appendf(TEXT("Insert from uint32 with ReadOffset 16 at end, MultipleInts, %d"), n).ToString(), Arr4[n], true);
+				TEST_EQUAL(Label.Appendf(TEXT("Insert from uint32 with ReadOffset 16 at end, MultipleInts, %d"), n).ToString(), Arr4[n], true);
 			}
 			for (int n = 16; n < 32; ++n)
 			{
 				Label.Reset();
-				TestEqual(Label.Appendf(TEXT("Insert from uint32 with ReadOffset 16 at end, MultipleInts, %d"), n).ToString(), Arr4[n], (n / 4) % 2 == 0);
+				TEST_EQUAL(Label.Appendf(TEXT("Insert from uint32 with ReadOffset 16 at end, MultipleInts, %d"), n).ToString(), Arr4[n], (n / 4) % 2 == 0);
 			}
 
 			TBitArray<> Arr5;
@@ -988,26 +1013,26 @@ public:
 			Arr5.CheckInvariants();
 			Arr5.InsertRange(MultipleIntsBackwards, 0, 32, 16);
 			Arr5.CheckInvariants();
-			TestEqual(TEXT("Insert from uint32 with ReadOffset 16 at beginning, MultipleInts, size"), Arr5.Num(), 64);
+			TEST_EQUAL(TEXT("Insert from uint32 with ReadOffset 16 at beginning, MultipleInts, size"), Arr5.Num(), 64);
 			for (int n = 0; n < 16; ++n)
 			{
 				Label.Reset();
-				TestEqual(Label.Appendf(TEXT("Insert from uint32 with ReadOffset 16 at beginning, MultipleInts, %d"), n).ToString(), Arr5[n], (n / 4) % 2 == 0);
+				TEST_EQUAL(Label.Appendf(TEXT("Insert from uint32 with ReadOffset 16 at beginning, MultipleInts, %d"), n).ToString(), Arr5[n], (n / 4) % 2 == 0);
 			}
 			for (int n = 16; n < 32; ++n)
 			{
 				Label.Reset();
-				TestEqual(Label.Appendf(TEXT("Insert from uint32 with ReadOffset 16 at beginning, MultipleInts, %d"), n).ToString(), Arr5[n], false);
+				TEST_EQUAL(Label.Appendf(TEXT("Insert from uint32 with ReadOffset 16 at beginning, MultipleInts, %d"), n).ToString(), Arr5[n], false);
 			}
 			for (int n = 32; n < 48; ++n)
 			{
 				Label.Reset();
-				TestEqual(Label.Appendf(TEXT("Insert from uint32 with ReadOffset 16 at beginning, MultipleInts, %d"), n).ToString(), Arr5[n], true);
+				TEST_EQUAL(Label.Appendf(TEXT("Insert from uint32 with ReadOffset 16 at beginning, MultipleInts, %d"), n).ToString(), Arr5[n], true);
 			}
 			for (int n = 48; n < 64; ++n)
 			{
 				Label.Reset();
-				TestEqual(Label.Appendf(TEXT("Insert from uint32 with ReadOffset 16 at beginning, MultipleInts, %d"), n).ToString(), Arr5[n], (n / 4) % 2 == 0);
+				TEST_EQUAL(Label.Appendf(TEXT("Insert from uint32 with ReadOffset 16 at beginning, MultipleInts, %d"), n).ToString(), Arr5[n], (n / 4) % 2 == 0);
 			}
 		}
 
@@ -1016,36 +1041,36 @@ public:
 			TBitArray<> Arr;
 			Arr.InsertRange(ArrGrowingTrue, 0, NumGrowingTrue);
 			Arr.CheckInvariants();
-			TestTrue(TEXT("Insert from uint32 with ReadOffset 0 at end from empty"), Arr == ArrGrowingTrue);
+			TEST_TRUE(TEXT("Insert from uint32 with ReadOffset 0 at end from empty"), Arr == ArrGrowingTrue);
 
 			Arr.InsertRange(ArrGrowingTrue, Arr.Num(), NumGrowingTrue);
 			Arr.CheckInvariants();
-			TestTrue(TEXT("Insert from uint32 with ReadOffset 0 at end to nonempty"), Arr == ArrGrowingTrueTwice);
+			TEST_TRUE(TEXT("Insert from uint32 with ReadOffset 0 at end to nonempty"), Arr == ArrGrowingTrueTwice);
 
 			TBitArray<> Arr1Insert2;
 			Arr1Insert2.InsertRange(ArrGrowingTrue, 0, NumGrowingTrue);
 			Arr1Insert2.CheckInvariants();
 			Arr1Insert2.InsertRange(ArrGrowingTrue, 0, NumGrowingTrue);
 			Arr1Insert2.CheckInvariants();
-			TestTrue(TEXT("Insert from uint32 with ReadOffset 0 at beginning to nonempty"), Arr1Insert2 == ArrGrowingTrueTwice);
+			TEST_TRUE(TEXT("Insert from uint32 with ReadOffset 0 at beginning to nonempty"), Arr1Insert2 == ArrGrowingTrueTwice);
 
 			TBitArray<> Arr1Insert3;
 			Arr1Insert3.InsertRange(ArrGrowingTrue, 0, NumGrowingTrue);
 			Arr1Insert3.CheckInvariants();
 			Arr1Insert3.InsertRange(ArrFalse, 5, 5);
 			Arr1Insert3.CheckInvariants();
-			TestEqual(TEXT("Insert from uint32 with ReadOffset 0 at middle to nonempty size"), Arr1Insert3.Num(), NumGrowingTrue + 5);
+			TEST_EQUAL(TEXT("Insert from uint32 with ReadOffset 0 at middle to nonempty size"), Arr1Insert3.Num(), NumGrowingTrue + 5);
 			for (int n = 0; n < 5; ++n)
 			{
-				TestEqual(TEXT("Insert from uint32 with ReadOffset 0 at middle to nonempty"), Arr1Insert3[n], ArrGrowingTrue[n]);
+				TEST_EQUAL(TEXT("Insert from uint32 with ReadOffset 0 at middle to nonempty"), Arr1Insert3[n], ArrGrowingTrue[n]);
 			}
 			for (int n = 5; n < 10; ++n)
 			{
-				TestEqual(TEXT("Insert from uint32 with ReadOffset 0 at middle to nonempty"), Arr1Insert3[n], false);
+				TEST_EQUAL(TEXT("Insert from uint32 with ReadOffset 0 at middle to nonempty"), Arr1Insert3[n], false);
 			}
 			for (int n = 10; n < NumGrowingTrue + 5; ++n)
 			{
-				TestEqual(TEXT("Insert from uint32 with ReadOffset 0 at middle to nonempty"), Arr1Insert3[n], ArrGrowingTrue[n - 5]);
+				TEST_EQUAL(TEXT("Insert from uint32 with ReadOffset 0 at middle to nonempty"), Arr1Insert3[n], ArrGrowingTrue[n - 5]);
 			}
 
 			TBitArray<> Arr2;
@@ -1055,7 +1080,7 @@ public:
 			Arr2.CheckInvariants();
 			for (int n = 0; n < 20; ++n)
 			{
-				TestEqual(TEXT("Insert from uint32 with ReadOffset 0 at end, Zeroes, Then Ones"), Arr2[n], n >= 10);
+				TEST_EQUAL(TEXT("Insert from uint32 with ReadOffset 0 at end, Zeroes, Then Ones"), Arr2[n], n >= 10);
 			}
 
 			TBitArray<> Arr2Insert1;
@@ -1065,7 +1090,7 @@ public:
 			Arr2Insert1.CheckInvariants();
 			for (int n = 0; n < 20; ++n)
 			{
-				TestEqual(TEXT("Insert from uint32 with ReadOffset 0 at beginning, Zeroes, Then Ones"), Arr2Insert1[n], n < 10);
+				TEST_EQUAL(TEXT("Insert from uint32 with ReadOffset 0 at beginning, Zeroes, Then Ones"), Arr2Insert1[n], n < 10);
 			}
 			TBitArray<> Arr2Insert2;
 			Arr2Insert2.InsertRange(ArrFalse, 0, 10);
@@ -1074,15 +1099,15 @@ public:
 			Arr2Insert2.CheckInvariants();
 			for (int n = 0; n < 5; ++n)
 			{
-				TestEqual(TEXT("Insert from uint32 with ReadOffset 0 at middle, Zeroes, Then Ones"), Arr2Insert2[n], false);
+				TEST_EQUAL(TEXT("Insert from uint32 with ReadOffset 0 at middle, Zeroes, Then Ones"), Arr2Insert2[n], false);
 			}
 			for (int n = 5; n < 15; ++n)
 			{
-				TestEqual(TEXT("Insert from uint32 with ReadOffset 0 at middle, Zeroes, Then Ones"), Arr2Insert2[n], true);
+				TEST_EQUAL(TEXT("Insert from uint32 with ReadOffset 0 at middle, Zeroes, Then Ones"), Arr2Insert2[n], true);
 			}
 			for (int n = 15; n < 20; ++n)
 			{
-				TestEqual(TEXT("Insert from uint32 with ReadOffset 0 at middle, Zeroes, Then Ones"), Arr2Insert2[n], false);
+				TEST_EQUAL(TEXT("Insert from uint32 with ReadOffset 0 at middle, Zeroes, Then Ones"), Arr2Insert2[n], false);
 			}
 
 			TBitArray<> Arr3;
@@ -1092,21 +1117,21 @@ public:
 			ArrMultipleInts.CheckInvariants();
 			Arr3.InsertRange(ArrMultipleInts, 0, 64);
 			Arr3.CheckInvariants();
-			TestTrue(TEXT("Insert from uint32 with ReadOffset 0 at end, MutlipleInts, size"), Arr3 == ArrMultipleInts);
+			TEST_TRUE(TEXT("Insert from uint32 with ReadOffset 0 at end, MutlipleInts, size"), Arr3 == ArrMultipleInts);
 
 			TBitArray<> Arr4;
 			Arr4.InsertRange(ArrMultipleInts, 0, 32, 16);
 			Arr4.CheckInvariants();
-			TestEqual(TEXT("Insert from uint32 with ReadOffset 16 at end, MultipleInts, size"), Arr4.Num(), 32);
+			TEST_EQUAL(TEXT("Insert from uint32 with ReadOffset 16 at end, MultipleInts, size"), Arr4.Num(), 32);
 			for (int n = 0; n < 16; ++n)
 			{
 				Label.Reset();
-				TestEqual(Label.Appendf(TEXT("Insert from uint32 with ReadOffset 16 at end, MultipleInts, %d"), n).ToString(), Arr4[n], true);
+				TEST_EQUAL(Label.Appendf(TEXT("Insert from uint32 with ReadOffset 16 at end, MultipleInts, %d"), n).ToString(), Arr4[n], true);
 			}
 			for (int n = 16; n < 32; ++n)
 			{
 				Label.Reset();
-				TestEqual(Label.Appendf(TEXT("Insert from uint32 with ReadOffset 16 at end, MultipleInts, %d"), n).ToString(), Arr4[n], (n / 4) % 2 == 0);
+				TEST_EQUAL(Label.Appendf(TEXT("Insert from uint32 with ReadOffset 16 at end, MultipleInts, %d"), n).ToString(), Arr4[n], (n / 4) % 2 == 0);
 			}
 
 			TBitArray<> Arr5;
@@ -1118,26 +1143,26 @@ public:
 			Arr5.CheckInvariants();
 			Arr5.InsertRange(ArrMultipleIntsBackwards, 0, 32, 16);
 			Arr5.CheckInvariants();
-			TestEqual(TEXT("Insert from uint32 with ReadOffset 16 at beginning, MultipleInts, size"), Arr5.Num(), 64);
+			TEST_EQUAL(TEXT("Insert from uint32 with ReadOffset 16 at beginning, MultipleInts, size"), Arr5.Num(), 64);
 			for (int n = 0; n < 16; ++n)
 			{
 				Label.Reset();
-				TestEqual(Label.Appendf(TEXT("Insert from uint32 with ReadOffset 16 at beginning, MultipleInts, %d"), n).ToString(), Arr5[n], (n / 4) % 2 == 0);
+				TEST_EQUAL(Label.Appendf(TEXT("Insert from uint32 with ReadOffset 16 at beginning, MultipleInts, %d"), n).ToString(), Arr5[n], (n / 4) % 2 == 0);
 			}
 			for (int n = 16; n < 32; ++n)
 			{
 				Label.Reset();
-				TestEqual(Label.Appendf(TEXT("Insert from uint32 with ReadOffset 16 at beginning, MultipleInts, %d"), n).ToString(), Arr5[n], false);
+				TEST_EQUAL(Label.Appendf(TEXT("Insert from uint32 with ReadOffset 16 at beginning, MultipleInts, %d"), n).ToString(), Arr5[n], false);
 			}
 			for (int n = 32; n < 48; ++n)
 			{
 				Label.Reset();
-				TestEqual(Label.Appendf(TEXT("Insert from uint32 with ReadOffset 16 at beginning, MultipleInts, %d"), n).ToString(), Arr5[n], true);
+				TEST_EQUAL(Label.Appendf(TEXT("Insert from uint32 with ReadOffset 16 at beginning, MultipleInts, %d"), n).ToString(), Arr5[n], true);
 			}
 			for (int n = 48; n < 64; ++n)
 			{
 				Label.Reset();
-				TestEqual(Label.Appendf(TEXT("Insert from uint32 with ReadOffset 16 at beginning, MultipleInts, %d"), n).ToString(), Arr5[n], (n / 4) % 2 == 0);
+				TEST_EQUAL(Label.Appendf(TEXT("Insert from uint32 with ReadOffset 16 at beginning, MultipleInts, %d"), n).ToString(), Arr5[n], (n / 4) % 2 == 0);
 			}
 		}
 	}
@@ -1201,8 +1226,8 @@ public:
 					ZeroesExpectedText[DisplayIndex] = TEXT('\0');
 					ZeroesText[DisplayIndex] = TEXT('\0');
 			
-					TestTrue(FString::Printf(TEXT("SetRange bool BitWidth=%d WriteOffset=%d, Ones\nExpected=%s\nActual  =%s"), BitWidth, WriteOffset, OnesExpectedText, OnesText), bOnesMatchesExpected);
-					TestTrue(FString::Printf(TEXT("SetRange bool BitWidth=%d WriteOffset=%d, Zeroes\nExpected=%s\nActual  =%s"), BitWidth, WriteOffset, ZeroesExpectedText, ZeroesText), bZeroesMatchesExpected);
+					TEST_TRUE(FString::Printf(TEXT("SetRange bool BitWidth=%d WriteOffset=%d, Ones\nExpected=%s\nActual  =%s"), BitWidth, WriteOffset, OnesExpectedText, OnesText), bOnesMatchesExpected);
+					TEST_TRUE(FString::Printf(TEXT("SetRange bool BitWidth=%d WriteOffset=%d, Zeroes\nExpected=%s\nActual  =%s"), BitWidth, WriteOffset, ZeroesExpectedText, ZeroesText), bZeroesMatchesExpected);
 				}
 			}
 		}
@@ -1230,8 +1255,8 @@ public:
 			ArrOnes.SetRangeFromRange(20, 12, &ZeroesInt, 20);
 			ArrOnes.CheckInvariants();
 
-			TestTrue(TEXT("SetRangeFromRange Zeroes"), ArrZeroes == ArrRefOnes);
-			TestTrue(TEXT("SetRangeFromRange Ones"), ArrOnes == ArrRefZeroes);
+			TEST_TRUE(TEXT("SetRangeFromRange Zeroes"), ArrZeroes == ArrRefOnes);
+			TEST_TRUE(TEXT("SetRangeFromRange Ones"), ArrOnes == ArrRefZeroes);
 		}
 
 		// SetRange that takes a bitarray
@@ -1255,8 +1280,8 @@ public:
 			ArrOnes.SetRangeFromRange(20, 12, ArrRefZeroes, 20);
 			ArrOnes.CheckInvariants();
 
-			TestTrue(TEXT("SetRange bitarray Zeroes"), ArrZeroes == ArrRefOnes);
-			TestTrue(TEXT("SetRange bitarray Ones"), ArrOnes == ArrRefZeroes);
+			TEST_TRUE(TEXT("SetRange bitarray Zeroes"), ArrZeroes == ArrRefOnes);
+			TEST_TRUE(TEXT("SetRange bitarray Ones"), ArrOnes == ArrRefZeroes);
 		}
 	}
 
@@ -1277,8 +1302,8 @@ public:
 		ArrZeroes.GetRange(10, 10, &OnesInt, 10);
 		ArrZeroes.GetRange(20, 12, &OnesInt, 20);
 
-		TestEqual(TEXT("GetRange Zeroes"), ZeroesInt, OnesIntRef);
-		TestEqual(TEXT("GetRange Ones"), OnesInt, ZeroesIntRef);
+		TEST_EQUAL(TEXT("GetRange Zeroes"), ZeroesInt, OnesIntRef);
+		TEST_EQUAL(TEXT("GetRange Ones"), OnesInt, ZeroesIntRef);
 	}
 
 	void TestEmpty()
@@ -1288,16 +1313,16 @@ public:
 		Arr.CheckInvariants();
 		Arr.Empty();
 		Arr.CheckInvariants();
-		TestEqual(TEXT("Empty with no arguments sets num to 0"), Arr.Num(), 0);
-		TestEqual(TEXT("Empty with no arguments sets max to 0"), Arr.Max(), 0);
+		TEST_EQUAL(TEXT("Empty with no arguments sets num to 0"), Arr.Num(), 0);
+		TEST_EQUAL(TEXT("Empty with no arguments sets max to 0"), Arr.Max(), 0);
 
 		TBitArray<FDefaultAllocator> Arr2;
 		Arr2.Add(true, 10);
 		Arr2.CheckInvariants();
 		Arr2.Empty(5);
 		Arr2.CheckInvariants();
-		TestEqual(TEXT("Empty with an arguments sets num to 0"), Arr2.Num(), 0);
-		TestEqual(TEXT("Empty with an argument sets max to rounded up input"), Arr2.Max(), NumBitsPerDWORD);
+		TEST_EQUAL(TEXT("Empty with an arguments sets num to 0"), Arr2.Num(), 0);
+		TEST_EQUAL(TEXT("Empty with an argument sets max to rounded up input"), Arr2.Max(), NumBitsPerDWORD);
 
 		TBitArray<TInlineAllocator<4>> Arr3;
 		int32 InitialMax = Arr3.Max();
@@ -1305,8 +1330,8 @@ public:
 		Arr3.CheckInvariants();
 		Arr3.Empty(0);
 		Arr3.CheckInvariants();
-		TestEqual(TEXT("Empty from a dynamic allocation with an inline allocator sets num to 0"), Arr3.Num(), 0);
-		TestEqual(TEXT("Empty from a dynamic allocation with an inline allocator sets max back to the size of the inline allocation"), Arr3.Max(), InitialMax);
+		TEST_EQUAL(TEXT("Empty from a dynamic allocation with an inline allocator sets num to 0"), Arr3.Num(), 0);
+		TEST_EQUAL(TEXT("Empty from a dynamic allocation with an inline allocator sets max back to the size of the inline allocation"), Arr3.Max(), InitialMax);
 	}
 
 	void TestReserve()
@@ -1314,15 +1339,15 @@ public:
 		TBitArray<FDefaultAllocator> Arr;
 		Arr.Reserve(NumBitsPerDWORD + NumBitsPerDWORD / 2);
 		Arr.CheckInvariants();
-		TestEqual(TEXT("Reserve from empty does not change num"), Arr.Num(), 0);
-		TestTrue(TEXT("Reserve from empty sets max to rounded up request"), Arr.Max() >= NumBitsPerDWORD * 2);
+		TEST_EQUAL(TEXT("Reserve from empty does not change num"), Arr.Num(), 0);
+		TEST_TRUE(TEXT("Reserve from empty sets max to rounded up request"), Arr.Max() >= NumBitsPerDWORD * 2);
 
 		TBitArray<FDefaultAllocator> Arr2Ref(true, NumBitsPerDWORD);
 		TBitArray<FDefaultAllocator> Arr2(true, NumBitsPerDWORD);
 		Arr2.Reserve(NumBitsPerDWORD + NumBitsPerDWORD / 2);
 		Arr2.CheckInvariants();
-		TestTrue(TEXT("Reserve from filled does not size or elements"), Arr2 == Arr2Ref);
-		TestTrue(TEXT("Reserve from filled sets max to rounded up request"), Arr.Max() >= NumBitsPerDWORD * 2);
+		TEST_TRUE(TEXT("Reserve from filled does not size or elements"), Arr2 == Arr2Ref);
+		TEST_TRUE(TEXT("Reserve from filled sets max to rounded up request"), Arr.Max() >= NumBitsPerDWORD * 2);
 	}
 
 	void TestReset()
@@ -1330,14 +1355,14 @@ public:
 		TBitArray<FDefaultAllocator> Arr;
 		Arr.Reset();
 		Arr.CheckInvariants();
-		TestEqual(TEXT("Reset from empty keeps num at 0"), Arr.Num(), 0);
-		TestEqual(TEXT("Reserve from empty keeps max at 0"), Arr.Max(), 0);
+		TEST_EQUAL(TEXT("Reset from empty keeps num at 0"), Arr.Num(), 0);
+		TEST_EQUAL(TEXT("Reserve from empty keeps max at 0"), Arr.Max(), 0);
 
 		TBitArray<FDefaultAllocator> Arr2(true, NumBitsPerDWORD);
 		Arr2.Reset();
 		Arr2.CheckInvariants();
-		TestEqual(TEXT("Reset from filled sets num to 0"), Arr2.Num(), 0);
-		TestEqual(TEXT("Reset from filled sets not change max"), Arr2.Max(), NumBitsPerDWORD);
+		TEST_EQUAL(TEXT("Reset from filled sets num to 0"), Arr2.Num(), 0);
+		TEST_EQUAL(TEXT("Reset from filled sets not change max"), Arr2.Max(), NumBitsPerDWORD);
 	}
 
 	void TestInitAndSetNumUninitialized()
@@ -1347,55 +1372,55 @@ public:
 			TBitArray<FDefaultAllocator> Arr;
 			Arr.SetNumUninitialized(0);
 			Arr.CheckInvariants();
-			TestEqual(TEXT("SetNumUninitialized 0 from empty keeps num at 0"), Arr.Num(), 0);
-			TestEqual(TEXT("SetNumUninitialized 0 from empty keeps max at 0"), Arr.Max(), 0);
+			TEST_EQUAL(TEXT("SetNumUninitialized 0 from empty keeps num at 0"), Arr.Num(), 0);
+			TEST_EQUAL(TEXT("SetNumUninitialized 0 from empty keeps max at 0"), Arr.Max(), 0);
 
 			TBitArray<FDefaultAllocator> Arr2(true, NumBitsPerDWORD);
 			Arr2.SetNumUninitialized(0);
 			Arr2.CheckInvariants();
-			TestEqual(TEXT("SetNumUninitialized 0 from filled sets num to 0"), Arr2.Num(), 0);
-			TestEqual(TEXT("SetNumUninitialized 0 from filled does not change max"), Arr2.Max(), NumBitsPerDWORD);
+			TEST_EQUAL(TEXT("SetNumUninitialized 0 from filled sets num to 0"), Arr2.Num(), 0);
+			TEST_EQUAL(TEXT("SetNumUninitialized 0 from filled does not change max"), Arr2.Max(), NumBitsPerDWORD);
 
 			TBitArray<FDefaultAllocator> Arr3;
 			Arr3.SetNumUninitialized(20);
 			Arr3.CheckInvariants();
-			TestEqual(TEXT("SetNumUninitialized 20 from empty sets num at 20"), Arr3.Num(), 20);
+			TEST_EQUAL(TEXT("SetNumUninitialized 20 from empty sets num at 20"), Arr3.Num(), 20);
 			// Depending on defines, requesting the bitarray's grow to handle a single int might reserve multiple ints. Just confirm the max is a multiple of bitsperdword > num
-			TestTrue(TEXT("SetNumUninitialized 20 sets max at rounded up"), Arr3.Max() >= NumBitsPerDWORD && (Arr3.Max() % NumBitsPerDWORD) == 0);
+			TEST_TRUE(TEXT("SetNumUninitialized 20 sets max at rounded up"), Arr3.Max() >= NumBitsPerDWORD && (Arr3.Max() % NumBitsPerDWORD) == 0);
 
 			TBitArray<FDefaultAllocator> Arr4Ones(true, NumBitsPerDWORD);
 			Arr4Ones.SetNumUninitialized(NumBitsPerDWORD + NumBitsPerDWORD / 2);
 			Arr4Ones.CheckInvariants();
-			TestEqual(TEXT("SetNumUninitialized to a higher number sets num to the higher number"), Arr4Ones.Num(), NumBitsPerDWORD + NumBitsPerDWORD / 2);
+			TEST_EQUAL(TEXT("SetNumUninitialized to a higher number sets num to the higher number"), Arr4Ones.Num(), NumBitsPerDWORD + NumBitsPerDWORD / 2);
 			// Depending on defines, requesting the bitarray's grow to handle a single int might reserve multiple ints. Just confirm the max is a multiple of bitsperdword > num
-			TestTrue(TEXT("SetNumUninitialized to a higher number sets max to rounded up"), Arr4Ones.Max() >= NumBitsPerDWORD * 2 && (Arr4Ones.Max() % NumBitsPerDWORD) == 0);
+			TEST_TRUE(TEXT("SetNumUninitialized to a higher number sets max to rounded up"), Arr4Ones.Max() >= NumBitsPerDWORD * 2 && (Arr4Ones.Max() % NumBitsPerDWORD) == 0);
 			for (int n = 0; n < NumBitsPerDWORD; ++n)
 			{
-				TestEqual(TEXT("SetNumUninitialized to a higher number keeps the old elements - ones"), Arr4Ones[n], true);
+				TEST_EQUAL(TEXT("SetNumUninitialized to a higher number keeps the old elements - ones"), Arr4Ones[n], true);
 			}
 			TBitArray<FDefaultAllocator> Arr4Zeroes(false, NumBitsPerDWORD);
 			Arr4Zeroes.SetNumUninitialized(NumBitsPerDWORD + NumBitsPerDWORD / 2);
 			Arr4Zeroes.CheckInvariants();
 			for (int n = 0; n < NumBitsPerDWORD; ++n)
 			{
-				TestEqual(TEXT("SetNumUninitialized to a higher number keeps the old elements - zeroes"), Arr4Zeroes[n], false);
+				TEST_EQUAL(TEXT("SetNumUninitialized to a higher number keeps the old elements - zeroes"), Arr4Zeroes[n], false);
 			}
 
 			TBitArray<FDefaultAllocator> Arr5Ones(true, NumBitsPerDWORD * 2);
 			Arr5Ones.SetNumUninitialized(NumBitsPerDWORD / 2);
 			Arr5Ones.CheckInvariants();
-			TestEqual(TEXT("SetNumUninitialized to a lower number sets num to the lowernumber"), Arr5Ones.Num(), NumBitsPerDWORD / 2);
-			TestEqual(TEXT("SetNumUninitialized to a lower number does not change max"), Arr5Ones.Max(), NumBitsPerDWORD * 2);
+			TEST_EQUAL(TEXT("SetNumUninitialized to a lower number sets num to the lowernumber"), Arr5Ones.Num(), NumBitsPerDWORD / 2);
+			TEST_EQUAL(TEXT("SetNumUninitialized to a lower number does not change max"), Arr5Ones.Max(), NumBitsPerDWORD * 2);
 			for (int n = 0; n < NumBitsPerDWORD / 2; ++n)
 			{
-				TestEqual(TEXT("SetNumUninitialized to a lower number keeps the old elements below the lower number - ones"), Arr5Ones[n], true);
+				TEST_EQUAL(TEXT("SetNumUninitialized to a lower number keeps the old elements below the lower number - ones"), Arr5Ones[n], true);
 			}
 			TBitArray<FDefaultAllocator> Arr5Zeroes(false, NumBitsPerDWORD * 2);
 			Arr5Zeroes.SetNumUninitialized(NumBitsPerDWORD / 2);
 			Arr5Zeroes.CheckInvariants();
 			for (int n = 0; n < NumBitsPerDWORD / 2; ++n)
 			{
-				TestEqual(TEXT("SetNumUninitialized to a lower number keeps the old elements below the lower number - zeroes"), Arr5Zeroes[n], false);
+				TEST_EQUAL(TEXT("SetNumUninitialized to a lower number keeps the old elements below the lower number - zeroes"), Arr5Zeroes[n], false);
 			}
 		}
 
@@ -1404,14 +1429,14 @@ public:
 			TBitArray<FDefaultAllocator> Arr;
 			Arr.Init(true, 0);
 			Arr.CheckInvariants();
-			TestEqual(TEXT("Init 0 from empty keeps num at 0"), Arr.Num(), 0);
-			TestEqual(TEXT("Init 0 from empty keeps max at 0"), Arr.Max(), 0);
+			TEST_EQUAL(TEXT("Init 0 from empty keeps num at 0"), Arr.Num(), 0);
+			TEST_EQUAL(TEXT("Init 0 from empty keeps max at 0"), Arr.Max(), 0);
 
 			TBitArray<FDefaultAllocator> Arr2(true, NumBitsPerDWORD);
 			Arr2.Init(true, 0);
 			Arr2.CheckInvariants();
-			TestEqual(TEXT("Init 0 from filled sets num to 0"), Arr2.Num(), 0);
-			TestEqual(TEXT("Init 0 from filled does not change max"), Arr2.Max(), NumBitsPerDWORD);
+			TEST_EQUAL(TEXT("Init 0 from filled sets num to 0"), Arr2.Num(), 0);
+			TEST_EQUAL(TEXT("Init 0 from filled does not change max"), Arr2.Max(), NumBitsPerDWORD);
 
 			TBitArray<FDefaultAllocator> Arr3True;
 			TBitArray<FDefaultAllocator> Arr3False;
@@ -1419,49 +1444,49 @@ public:
 			TBitArray<FDefaultAllocator> Arr3FalseRef(false, 20);
 			Arr3True.Init(true, 20);
 			Arr3True.CheckInvariants();
-			TestTrue(TEXT("Init true 20 from empty sets size and elements"), Arr3True == Arr3TrueRef);
-			TestEqual(TEXT("Init true 20 sets max at rounded up"), Arr3True.Max(), NumBitsPerDWORD);
+			TEST_TRUE(TEXT("Init true 20 from empty sets size and elements"), Arr3True == Arr3TrueRef);
+			TEST_EQUAL(TEXT("Init true 20 sets max at rounded up"), Arr3True.Max(), NumBitsPerDWORD);
 			Arr3False.Init(false, 20);
 			Arr3False.CheckInvariants();
-			TestTrue(TEXT("Init false 20 from empty sets size and elements"), Arr3False == Arr3FalseRef);
-			TestEqual(TEXT("Init false 20 sets max at rounded up"), Arr3False.Max(), NumBitsPerDWORD);
+			TEST_TRUE(TEXT("Init false 20 from empty sets size and elements"), Arr3False == Arr3FalseRef);
+			TEST_EQUAL(TEXT("Init false 20 sets max at rounded up"), Arr3False.Max(), NumBitsPerDWORD);
 
 			TBitArray<FDefaultAllocator> Arr4Ones(true, NumBitsPerDWORD);
 			Arr4Ones.Init(false, NumBitsPerDWORD + NumBitsPerDWORD / 2);
 			Arr4Ones.CheckInvariants();
-			TestEqual(TEXT("Init false to a higher number sets num to the higher number"), Arr4Ones.Num(), NumBitsPerDWORD + NumBitsPerDWORD / 2);
-			TestEqual(TEXT("Init false to a higher number sets max to rounded up"), Arr4Ones.Max(), NumBitsPerDWORD * 2);
+			TEST_EQUAL(TEXT("Init false to a higher number sets num to the higher number"), Arr4Ones.Num(), NumBitsPerDWORD + NumBitsPerDWORD / 2);
+			TEST_EQUAL(TEXT("Init false to a higher number sets max to rounded up"), Arr4Ones.Max(), NumBitsPerDWORD * 2);
 			for (int n = 0; n < Arr4Ones.Num(); ++n)
 			{
-				TestEqual(TEXT("Init false to a higher number overwrites all elements"), Arr4Ones[n], false);
+				TEST_EQUAL(TEXT("Init false to a higher number overwrites all elements"), Arr4Ones[n], false);
 			}
 			TBitArray<FDefaultAllocator> Arr4Zeroes(false, NumBitsPerDWORD);
 			Arr4Zeroes.Init(true, NumBitsPerDWORD + NumBitsPerDWORD / 2);
 			Arr4Zeroes.CheckInvariants();
-			TestEqual(TEXT("Init true to a higher number sets num to the higher number"), Arr4Zeroes.Num(), NumBitsPerDWORD + NumBitsPerDWORD / 2);
-			TestEqual(TEXT("Init true to a higher number sets max to rounded up"), Arr4Zeroes.Max(), NumBitsPerDWORD * 2);
+			TEST_EQUAL(TEXT("Init true to a higher number sets num to the higher number"), Arr4Zeroes.Num(), NumBitsPerDWORD + NumBitsPerDWORD / 2);
+			TEST_EQUAL(TEXT("Init true to a higher number sets max to rounded up"), Arr4Zeroes.Max(), NumBitsPerDWORD * 2);
 			for (int n = 0; n < Arr4Zeroes.Num(); ++n)
 			{
-				TestEqual(TEXT("Init true to a higher number overwrites all elements"), Arr4Zeroes[n], true);
+				TEST_EQUAL(TEXT("Init true to a higher number overwrites all elements"), Arr4Zeroes[n], true);
 			}
 
 			TBitArray<FDefaultAllocator> Arr5Ones(true, NumBitsPerDWORD * 2);
 			Arr5Ones.Init(false, NumBitsPerDWORD / 2);
 			Arr5Ones.CheckInvariants();
-			TestEqual(TEXT("Init false to a lower number sets num to the lowernumber"), Arr5Ones.Num(), NumBitsPerDWORD / 2);
-			TestEqual(TEXT("Init false to a lower number does not change max"), Arr5Ones.Max(), NumBitsPerDWORD * 2);
+			TEST_EQUAL(TEXT("Init false to a lower number sets num to the lowernumber"), Arr5Ones.Num(), NumBitsPerDWORD / 2);
+			TEST_EQUAL(TEXT("Init false to a lower number does not change max"), Arr5Ones.Max(), NumBitsPerDWORD * 2);
 			for (int n = 0; n < Arr5Ones.Num(); ++n)
 			{
-				TestEqual(TEXT("Init false to a lower number overwrites all elements"), Arr5Ones[n], false);
+				TEST_EQUAL(TEXT("Init false to a lower number overwrites all elements"), Arr5Ones[n], false);
 			}
 			TBitArray<FDefaultAllocator> Arr5Zeroes(false, NumBitsPerDWORD * 2);
 			Arr5Zeroes.Init(true, NumBitsPerDWORD / 2);
 			Arr5Zeroes.CheckInvariants();
-			TestEqual(TEXT("Init true to a lower number sets num to the lowernumber"), Arr5Zeroes.Num(), NumBitsPerDWORD / 2);
-			TestEqual(TEXT("Init true to a lower number does not change max"), Arr5Zeroes.Max(), NumBitsPerDWORD * 2);
+			TEST_EQUAL(TEXT("Init true to a lower number sets num to the lowernumber"), Arr5Zeroes.Num(), NumBitsPerDWORD / 2);
+			TEST_EQUAL(TEXT("Init true to a lower number does not change max"), Arr5Zeroes.Max(), NumBitsPerDWORD * 2);
 			for (int n = 0; n < Arr5Zeroes.Num(); ++n)
 			{
-				TestEqual(TEXT("Init true to a lower number overwrites all elements"), Arr5Zeroes[n], true);
+				TEST_EQUAL(TEXT("Init true to a lower number overwrites all elements"), Arr5Zeroes[n], true);
 			}
 		}
 	}
@@ -1479,7 +1504,7 @@ public:
 	// TODO: Reverse Iteration
 };
 
-TEST_CASE("Core::Containers::FBitArray::Misc", "[Core][Containers][Smoke]")
+TEST_CASE_METHOD(FAutomationTestFixture, "Core::Containers::FBitArray::Misc", "[Core][Containers][Smoke]")
 {
 	FBitArrayTest BitArrayTest;
 
@@ -1502,7 +1527,7 @@ TEST_CASE("Core::Containers::FBitArray::Misc", "[Core][Containers][Smoke]")
 	BitArrayTest.TestInitAndSetNumUninitialized();
 }
 
-TEST_CASE("Core::Containers::FBitArray::Invariants", "[Core][Containers][Smoke]")
+TEST_CASE_METHOD(FAutomationTestFixture, "Core::Containers::FBitArray::Invariants", "[Core][Containers][Smoke]")
 {
 	using namespace UE::BitArrayTest;
 
@@ -1593,7 +1618,7 @@ TEST_CASE("Core::Containers::FBitArray::Invariants", "[Core][Containers][Smoke]"
 }
 
 
-TEST_CASE("Core::Containers::FBitArray::CountSetBits", "[Core][Containers][Smoke]")
+TEST_CASE_METHOD(FAutomationTestFixture, "Core::Containers::FBitArray::CountSetBits", "[Core][Containers][Smoke]")
 {
 	using namespace UE::BitArrayTest;
 
@@ -1619,7 +1644,7 @@ TEST_CASE("Core::Containers::FBitArray::CountSetBits", "[Core][Containers][Smoke
 
 			const int32 SetBits = Array.CountSetBits();
 			
-			TestEqual(*FString::Printf(TEXT("CountSetBits: Unexpected number of set bits for array %s. Expected: %i, Actual: %i"), *BitArrayToString(Array), Test.Expected, SetBits), SetBits , Test.Expected);	
+			TEST_EQUAL(*FString::Printf(TEXT("CountSetBits: Unexpected number of set bits for array %s. Expected: %i, Actual: %i"), *BitArrayToString(Array), Test.Expected, SetBits), SetBits , Test.Expected);	
 		}
 	}
 
@@ -1643,13 +1668,13 @@ TEST_CASE("Core::Containers::FBitArray::CountSetBits", "[Core][Containers][Smoke
 			const TBitArray<> Array = ConstructBitArray(Test.Bits);
 
 			const int32 SetBits = Array.CountSetBits(Test.StartIndex, Test.EndIndex);
-			TestEqual(*FString::Printf(TEXT("CountSetBits: Unexpected number of set bits for array %s between index %d and %d. Expected: %i, Actual: %i"), *BitArrayToString(Array), Test.StartIndex, Test.EndIndex, Test.Expected, SetBits), SetBits ,Test.Expected);		
+			TEST_EQUAL(*FString::Printf(TEXT("CountSetBits: Unexpected number of set bits for array %s between index %d and %d. Expected: %i, Actual: %i"), *BitArrayToString(Array), Test.StartIndex, Test.EndIndex, Test.Expected, SetBits), SetBits ,Test.Expected);		
 		}
 	}
 }
 
 
-TEST_CASE("Core::Containers::FBitArray::BitwiseNOT", "[Core][Containers][Smoke]")
+TEST_CASE_METHOD(FAutomationTestFixture, "Core::Containers::FBitArray::BitwiseNOT", "[Core][Containers][Smoke]")
 {
 	using namespace UE::BitArrayTest;
 
@@ -1685,11 +1710,11 @@ TEST_CASE("Core::Containers::FBitArray::BitwiseNOT", "[Core][Containers][Smoke]"
 
 		TBitArray<> Result   = Input;
 		Result.BitwiseNOT();
-		TestEqual(*FString::Printf(TEXT("Bitwise NOT: Unexpected result for source %hs. Expected: %hs, Actual: %s"), Test.Input, Test.Expected, *BitArrayToString(Result)), Result, Expected);
+		TEST_EQUAL(*FString::Printf(TEXT("Bitwise NOT: Unexpected result for source %hs. Expected: %hs, Actual: %s"), Test.Input, Test.Expected, *BitArrayToString(Result)), Result, Expected);
 	}
 }
 
-TEST_CASE("Core::Containers::FBitArray::BitwiseAnd", "[Core][Containers][Smoke]")
+TEST_CASE_METHOD(FAutomationTestFixture, "Core::Containers::FBitArray::BitwiseAnd", "[Core][Containers][Smoke]")
 {
 	using namespace UE::BitArrayTest;
 
@@ -1711,7 +1736,7 @@ TEST_CASE("Core::Containers::FBitArray::BitwiseAnd", "[Core][Containers][Smoke]"
 		const char* Expected;
 	};
 
-	auto RunBinaryTestImpl = [](const TCHAR* Description, TArrayView<const FTestInput> Tests, TArrayView<const FTestResult> Results, BinaryCallable BinaryOp)
+	auto RunBinaryTestImpl = [this](const TCHAR* Description, TArrayView<const FTestInput> Tests, TArrayView<const FTestResult> Results, BinaryCallable BinaryOp)
 	{
 		check(Tests.Num() == Results.Num());
 		for (int32 Index = 0; Index < Tests.Num(); ++Index)
@@ -1725,15 +1750,15 @@ TEST_CASE("Core::Containers::FBitArray::BitwiseAnd", "[Core][Containers][Smoke]"
 
 			TBitArray<> Result = BinaryOp(InputA, InputB);
 
-			TestEqual(*FString::Printf(TEXT("%s: Unexpected result for source %hs & %hs. Expected: %hs, Actual: %s"), Description, Test.InputA, Test.InputB, TestResult.Expected, *BitArrayToString(Result)), Result, Expected);
+			TEST_EQUAL(*FString::Printf(TEXT("%s: Unexpected result for source %hs & %hs. Expected: %hs, Actual: %s"), Description, Test.InputA, Test.InputB, TestResult.Expected, *BitArrayToString(Result)), Result, Expected);
 			
 			Result = BinaryOp(InputB, InputA);
 			
-			TestEqual(*FString::Printf(TEXT("%s: Unexpected result for source %hs & %hs. Expected: %hs, Actual: %s"), Description, Test.InputB, Test.InputA, TestResult.Expected, *BitArrayToString(Result)), Result, Expected);
+			TEST_EQUAL(*FString::Printf(TEXT("%s: Unexpected result for source %hs & %hs. Expected: %hs, Actual: %s"), Description, Test.InputB, Test.InputA, TestResult.Expected, *BitArrayToString(Result)), Result, Expected);
 		}
 	};
 
-	auto RunMutatingTestImpl = [](const TCHAR* Description, TArrayView<const FTestInput> Tests, TArrayView<const FTestResult> Results, MutatingCallable MutatingOp)
+	auto RunMutatingTestImpl = [this](const TCHAR* Description, TArrayView<const FTestInput> Tests, TArrayView<const FTestResult> Results, MutatingCallable MutatingOp)
 	{
 		for (int32 Index = 0; Index < Tests.Num(); ++Index)
 		{
@@ -1747,7 +1772,7 @@ TEST_CASE("Core::Containers::FBitArray::BitwiseAnd", "[Core][Containers][Smoke]"
 			TBitArray<> Result = InputA;
 			MutatingOp(Result, InputB);
 
-			TestEqual(*FString::Printf(TEXT("%s: Unexpected result for source %hs & %hs. Expected: %hs, Actual: %s"), Description, Test.InputA, Test.InputB, TestResult.Expected, *BitArrayToString(Result)), Result, Expected);
+			TEST_EQUAL(*FString::Printf(TEXT("%s: Unexpected result for source %hs & %hs. Expected: %hs, Actual: %s"), Description, Test.InputA, Test.InputB, TestResult.Expected, *BitArrayToString(Result)), Result, Expected);
 		}
 	};
 
@@ -1861,7 +1886,7 @@ TEST_CASE("Core::Containers::FBitArray::BitwiseAnd", "[Core][Containers][Smoke]"
 
 
 
-TEST_CASE("Core::Containers::FBitArray::BitwiseOR", "[Core][Containers][Smoke]")
+TEST_CASE_METHOD(FAutomationTestFixture, "Core::Containers::FBitArray::BitwiseOR", "[Core][Containers][Smoke]")
 {
 	using namespace UE::BitArrayTest;
 
@@ -1895,21 +1920,21 @@ TEST_CASE("Core::Containers::FBitArray::BitwiseOR", "[Core][Containers][Smoke]")
 		const TBitArray<> Expected = ConstructBitArray(Test.Expected);
 
 		TBitArray<> Result = TBitArray<>::BitwiseOR(InputA, InputB, EBitwiseOperatorFlags::MaxSize);
-		TestEqual(*FString::Printf(TEXT("BitwiseOR: Unexpected result for source %hs | %hs. Expected: %hs, Actual: %s"), Test.InputA, Test.InputB, Test.Expected, *BitArrayToString(Result)), Result, Expected );
+		TEST_EQUAL(*FString::Printf(TEXT("BitwiseOR: Unexpected result for source %hs | %hs. Expected: %hs, Actual: %s"), Test.InputA, Test.InputB, Test.Expected, *BitArrayToString(Result)), Result, Expected );
 		
 		Result = TBitArray<>::BitwiseOR(InputB, InputA, EBitwiseOperatorFlags::MaxSize);
 		
-		TestEqual(*FString::Printf(TEXT("BitwiseOR: Unexpected result for source %hs | %hs. Expected: %hs, Actual: %s"), Test.InputB, Test.InputA, Test.Expected, *BitArrayToString(Result)), Result, Expected);
+		TEST_EQUAL(*FString::Printf(TEXT("BitwiseOR: Unexpected result for source %hs | %hs. Expected: %hs, Actual: %s"), Test.InputB, Test.InputA, Test.Expected, *BitArrayToString(Result)), Result, Expected);
 	
 		Result = InputA;
 		Result.CombineWithBitwiseOR(InputB, EBitwiseOperatorFlags::MaxSize);
 
-		TestEqual(*FString::Printf(TEXT("CombineWithBitwiseOR: Unexpected result for source %hs | %hs. Expected: %hs, Actual: %s"), Test.InputA, Test.InputB, Test.Expected, *BitArrayToString(Result)), Result, Expected);
+		TEST_EQUAL(*FString::Printf(TEXT("CombineWithBitwiseOR: Unexpected result for source %hs | %hs. Expected: %hs, Actual: %s"), Test.InputA, Test.InputB, Test.Expected, *BitArrayToString(Result)), Result, Expected);
 	}
 }
 
 
-TEST_CASE("Core::Containers::FBitArray::BitwiseXOR", "[Core][Containers][Smoke]")
+TEST_CASE_METHOD(FAutomationTestFixture, "Core::Containers::FBitArray::BitwiseXOR", "[Core][Containers][Smoke]")
 {
 	using namespace UE::BitArrayTest;
 
@@ -1947,21 +1972,21 @@ TEST_CASE("Core::Containers::FBitArray::BitwiseXOR", "[Core][Containers][Smoke]"
 
 		TBitArray<> Result = TBitArray<>::BitwiseXOR(InputA, InputB, EBitwiseOperatorFlags::MaxSize);
 		
-		TestEqual(*FString::Printf(TEXT("BitwiseXOR: Unexpected result for source %hs | %hs. Expected: %hs, Actual: %s"), Test.InputA, Test.InputB, Test.Expected, *BitArrayToString(Result)), Result, Expected);
+		TEST_EQUAL(*FString::Printf(TEXT("BitwiseXOR: Unexpected result for source %hs | %hs. Expected: %hs, Actual: %s"), Test.InputA, Test.InputB, Test.Expected, *BitArrayToString(Result)), Result, Expected);
 		
 		Result = TBitArray<>::BitwiseXOR(InputB, InputA, EBitwiseOperatorFlags::MaxSize);
 
-		TestEqual(*FString::Printf(TEXT("BitwiseXOR: Unexpected result for source %hs | %hs. Expected: %hs, Actual: %s"), Test.InputB, Test.InputA, Test.Expected, *BitArrayToString(Result)), Result, Expected);
+		TEST_EQUAL(*FString::Printf(TEXT("BitwiseXOR: Unexpected result for source %hs | %hs. Expected: %hs, Actual: %s"), Test.InputB, Test.InputA, Test.Expected, *BitArrayToString(Result)), Result, Expected);
 
 		Result = InputA;
 		Result.CombineWithBitwiseXOR(InputB, EBitwiseOperatorFlags::MaxSize);
 
-		TestEqual(*FString::Printf(TEXT("CombineWithBitwiseXOR: Unexpected result for source %hs | %hs. Expected: %hs, Actual: %s"), Test.InputA, Test.InputB, Test.Expected, *BitArrayToString(Result)), Result, Expected);
+		TEST_EQUAL(*FString::Printf(TEXT("CombineWithBitwiseXOR: Unexpected result for source %hs | %hs. Expected: %hs, Actual: %s"), Test.InputA, Test.InputB, Test.Expected, *BitArrayToString(Result)), Result, Expected);
 	}
 }
 
 
-class FBitArrayMemoryTest
+class FBitArrayMemoryTest : public FAutomationTestFixture
 {
 public:
 	bool TestMemmoveBitsWordOrder()
@@ -2126,8 +2151,8 @@ public:
 						ZeroesExpectedText[DisplayIndex] = TEXT('\0');
 						ZeroesText[DisplayIndex] = TEXT('\0');
 
-						TestTrue(FString::Printf(TEXT("MemmoveBitsWordOrder BitWidth=%d ReadOffset=%d, WriteOffset=%d, Overlap=%d, Ones\nExpected=%s\nActual  =%s"), BitWidth, ReadOffset, WriteOffset, Overlap, OnesExpectedText, OnesText), bOnesMatchesExpected);
-						TestTrue(FString::Printf(TEXT("MemmoveBitsWordOrder BitWidth=%d ReadOffset=%d, WriteOffset=%d, Overlap=%d, Zeroes\nExpected=%s\nActual  =%s"), BitWidth, ReadOffset, WriteOffset, Overlap, ZeroesExpectedText, ZeroesText), bZeroesMatchesExpected);
+						TEST_TRUE(FString::Printf(TEXT("MemmoveBitsWordOrder BitWidth=%d ReadOffset=%d, WriteOffset=%d, Overlap=%d, Ones\nExpected=%s\nActual  =%s"), BitWidth, ReadOffset, WriteOffset, Overlap, OnesExpectedText, OnesText), bOnesMatchesExpected);
+						TEST_TRUE(FString::Printf(TEXT("MemmoveBitsWordOrder BitWidth=%d ReadOffset=%d, WriteOffset=%d, Overlap=%d, Zeroes\nExpected=%s\nActual  =%s"), BitWidth, ReadOffset, WriteOffset, Overlap, ZeroesExpectedText, ZeroesText), bZeroesMatchesExpected);
 					
 					}
 				}
@@ -2179,7 +2204,7 @@ public:
 				}
 				Expected[Index] = '\0';
 				Actual[Index] = '\0';
-				TestTrue(FString::Printf(TEXT("MemmoveBitsWordOrder ModularizeWordOffset Offset=%d\nExpected=%s\nActual  =%s"), Offset, ANSI_TO_TCHAR(Expected), ANSI_TO_TCHAR(Actual)), bMatch);	
+				TEST_TRUE(FString::Printf(TEXT("MemmoveBitsWordOrder ModularizeWordOffset Offset=%d\nExpected=%s\nActual  =%s"), Offset, ANSI_TO_TCHAR(Expected), ANSI_TO_TCHAR(Actual)), bMatch);	
 			}
 		}
 		return true;
@@ -2203,7 +2228,7 @@ public:
 				uint32* ExpectedData = BaseData + WordLength;
 				uint32 ExpectedOffset = BitLength;
 				
-				TestFalse(FString::Printf(TEXT("ModularizeWordOffset WordLength=%d, BitLength=%d\nExpected: Data=%d, Offset=%d\nActual:  Data=%d, Offset=%d"), WordLength, BitLength, ExpectedData - BaseData, ExpectedOffset, Data - BaseData, Offset), (Data != ExpectedData || Offset != ExpectedOffset));
+				TEST_FALSE(FString::Printf(TEXT("ModularizeWordOffset WordLength=%d, BitLength=%d\nExpected: Data=%d, Offset=%d\nActual:  Data=%d, Offset=%d"), WordLength, BitLength, ExpectedData - BaseData, ExpectedOffset, Data - BaseData, Offset), (Data != ExpectedData || Offset != ExpectedOffset));
 			}
 		}
 
@@ -2211,7 +2236,7 @@ public:
 	}
 };
 
-TEST_CASE("Core::Containers::FBitArray::MemmoveBitsWordOrder", "[Core][Containers][Smoke]")
+TEST_CASE_METHOD(FAutomationTestFixture, "Core::Containers::FBitArray::MemmoveBitsWordOrder", "[Core][Containers][Smoke]")
 {
 	FBitArrayMemoryTest BitArrayMemoryTest;
 
diff --git a/Engine/Source/Runtime/Core/Tests/Containers/ConcurrentQueuesTest.cpp b/Engine/Source/Runtime/Core/Tests/Containers/ConcurrentQueuesTest.cpp
--- a/Engine/Source/Runtime/Core/Tests/Containers/ConcurrentQueuesTest.cpp
+++ b/Engine/Source/Runtime/Core/Tests/Containers/ConcurrentQueuesTest.cpp
@@ -9,10 +9,8 @@
 #include "Containers/MpscQueue.h"
 #include "Containers/ClosableMpscQueue.h"
 #include "Containers/DepletableMpscQueue.h"
-#include "Tests/Benchmark.h"
-#include "TestHarness.h"
-
-#include "TestCommon/CoreUtilities.h"
+#include "Core/Tests/Benchmark.h"
+#include "TestFixtures/CoreTestFixture.h"
 
 #include <atomic>
 
@@ -245,7 +243,7 @@ namespace UE { namespace ConcurrentQueuesTests
 		CHECK(Produced == NumConcumed);
 	}
 
-	TEST_CASE("Core::Containers::TMpscQueue::ConcurrentQueuesTest", "[Core][Containers][Smoke]")
+	TEST_CASE_METHOD(FAutomationTestFixture, "Core::Containers::TMpscQueue::ConcurrentQueuesTest", "[Core][Containers][Smoke]")
 	{
 		{	// test for support of not default constructible types
 			struct FNonDefaultConstructable
@@ -403,10 +401,8 @@ namespace UE { namespace ConcurrentQueuesTests
 		}
 	}
 
-	TEST_CASE("Core::Containers::TMpscQueue::ConcurrentQueuesPerfTest", "[Core][Containers][Perf]")
+	TEST_CASE_METHOD(FCoreTestFixture, "Core::Containers::TMpscQueue::ConcurrentQueuesPerfTest", "[.][Core][Containers][.Perf]")
 	{
-		InitTaskGraphAndDependencies(true);
-
 		UE_BENCHMARK(5, TestTCircularQueueSingleThread<5'000'000>);
 		UE_BENCHMARK(5, TestQueueSingleThread<5'000'000, TQueueAdapter<TQueue<uint32, EQueueMode::Spsc>>>);
 		UE_BENCHMARK(5, TestQueueSingleThread<5'000'000, TQueueAdapter<TQueue<uint32, EQueueMode::Mpsc>>>);
@@ -427,8 +423,6 @@ namespace UE { namespace ConcurrentQueuesTests
 
 		UE_BENCHMARK(5, TestMpscQueue<1'000'000, TQueueAdapter<TQueue<uint32, EQueueMode::Mpsc>>>);
 		UE_BENCHMARK(5, TestMpscQueue<1'000'000, TMpscQueue<uint32>>);
-
-		CleanupTaskGraphAndDependencies();
 	}
 }
 
@@ -673,7 +667,7 @@ namespace ClosableMpscQueueTests
 		UE_LOG(LogTemp, Display, TEXT("items queued %d"), NumProduced);
 	}
 
-	TEST_CASE("Core::Containers::TClosableMpscQueue::ClosableMpscQueueTest", "[Core][Containers][Smoke]")
+	TEST_CASE_METHOD(FAutomationTestFixture, "Core::Containers::TClosableMpscQueue::ClosableMpscQueueTest", "[Core][Containers][Smoke]")
 	{
 		{
 			TClosableMpscQueue<int> Q;
@@ -727,10 +721,8 @@ namespace ClosableMpscQueueTests
 		}
 	}
 
-	TEST_CASE("Core::Containers::TClosableMpscQueue::ClosableMpscQueuePerfTest", "[Core][Containers][Perf]")
+	TEST_CASE_METHOD(FCoreTestFixture, "Core::Containers::TClosableMpscQueue::ClosableMpscQueuePerfTest", "[.][Core][Containers][.Perf]")
 	{
-		InitTaskGraphAndDependencies(true);
-
 		UE_BENCHMARK(5, TestClosingEmptyQueue<10'000'000, TClosableMpscQueue<void*>>);
 		UE_BENCHMARK(5, TestClosingEmptyQueue<10'000'000, TClosableMpscQueueOnTQueue<void*>>);
 		UE_BENCHMARK(5, TestClosingEmptyQueue<10'000'000, TClosableMpscQueueOnTMpscQueue<void*>>);
@@ -745,8 +737,6 @@ namespace ClosableMpscQueueTests
 		UE_BENCHMARK(5, TestPerf<5'000'000, TClosableMpscQueueOnTQueue<void*>>);
 		UE_BENCHMARK(5, TestPerf<5'000'000, TClosableMpscQueueOnTMpscQueue<void*>>);
 		UE_BENCHMARK(5, TestPerf<5'000'000, TClosableLockFreePointerListUnorderedSingleConsumer_Adapter<void>>);
-
-		CleanupTaskGraphAndDependencies();
 	}
 }
 
@@ -1003,7 +993,7 @@ namespace DepletableMpscQueueTests
 		Queue.Deplete([](void*) { checkNoEntry(); }); // empty
 	}
 
-	TEST_CASE("Core::Containers::TDepletableMpscQueue::DepletableMpscQueueTest", "[Core][Containers][Smoke][PlatformFlaky]")
+	TEST_CASE_METHOD(FAutomationTestFixture, "Core::Containers::TDepletableMpscQueue::DepletableMpscQueueTest", "[Core][Containers][Smoke]")
 	{
 		{
 			TDepletableMpscQueue<int> Q;
@@ -1083,10 +1073,8 @@ namespace DepletableMpscQueueTests
 		}
 	}
 
-	TEST_CASE("Core::Containers::TDepletableMpscQueue::DepletableMpscQueuePerfTest", "[Core][Containers][Perf][PlatformFlaky]")
+	TEST_CASE_METHOD(FCoreTestFixture, "Core::Containers::TDepletableMpscQueue::DepletableMpscQueuePerfTest", "[.][Core][Containers][.Perf]")
 	{
-		InitTaskGraphAndDependencies(true);
-
 		UE_BENCHMARK(5, TestPerf<10'000, TDepletableMpscQueue<void*>>);
 		UE_BENCHMARK(5, TestPerf<10'000, TDepletableMpscQueueOnTQueue<void*>>);
 		UE_BENCHMARK(5, TestPerf<10'000, TDepletableMpscQueueOnTMpscQueue<void*>>);
@@ -1098,7 +1086,5 @@ namespace DepletableMpscQueueTests
 		UE_BENCHMARK(5, TestConsumingEmptyQueue<10'000'000, TDepletableMpscQueueOnTMpscQueue<void*>>);
 		UE_BENCHMARK(5, TestConsumingEmptyQueue<10'000'000, TDepletableMpscQueueOnTClosableLockFreePointerListUnorderedSingleConsumer<void>>);
 		UE_BENCHMARK(5, TestConsumingEmptyQueue<10'000'000, TDepletableMpscQueueOnTClosableMpscQueue<void*>>);
-
-		CleanupTaskGraphAndDependencies();
 	}
 }}
\ No newline at end of file
diff --git a/Engine/Source/Runtime/Core/Tests/Containers/ContainersTest.cpp b/Engine/Source/Runtime/Core/Tests/Containers/ContainersTest.cpp
--- a/Engine/Source/Runtime/Core/Tests/Containers/ContainersTest.cpp
+++ b/Engine/Source/Runtime/Core/Tests/Containers/ContainersTest.cpp
@@ -484,12 +484,12 @@ namespace
 	}
 }
 
-TEST_CASE("Core::Containers::TMap::Smoke", "[Core][Containers][Smoke]")
+TEST_CASE_METHOD(FAutomationTestFixture, "Core::Containers::TMap::Smoke", "[Core][Containers][Smoke]")
 {
 	RunContainerTests<TMap<int32, FContainerTestValueType>, int32>();
 }
 
-TEST_CASE("Core::Containers::TMap::Full", "[Core][Containers][Full]")
+TEST_CASE_METHOD(FAutomationTestFixture, "Core::Containers::TMap::Full", "[Core][Containers][Full]")
 {
 	RunContainerTests<TMap<int32, FContainerTestValueType>, int32>();
 	RunContainerTests<TMap<FName, FContainerTestValueType>, FName>();
@@ -498,7 +498,7 @@ TEST_CASE("Core::Containers::TMap::Full", "[Core][Containers][Full]")
 	RunContainerTests<TMap<FString, FContainerTestValueType, FDefaultSetAllocator, FCaseSensitiveLookupKeyFuncs<FContainerTestValueType>>, FString>();
 }
 
-TEST_CASE("Core::Containers::TSortedMap::Full", "[Core][Containers][Full]")
+TEST_CASE_METHOD(FAutomationTestFixture, "Core::Containers::TSortedMap::Full", "[Core][Containers][Full]")
 {
 	RunContainerTests<TSortedMap<int32, FContainerTestValueType>, int32>();
 	RunContainerTests<TSortedMap<FName, FContainerTestValueType, FDefaultAllocator, FNameLexicalLess>, FName>();
@@ -524,7 +524,7 @@ TEST_CASE("Core::Containers::TSortedMap::Full", "[Core][Containers][Full]")
 	RunGetRefTests();
 }
 
-TEST_CASE("Core::Containers::TMap::Performance", "[Core][Containers][Perf]")
+TEST_CASE_METHOD(FAutomationTestFixture, "Core::Containers::TMap::Performance", "[Core][Containers][.Perf]")
 {
 	RunPerformanceTest<TMap<int32, FString>, int32>(TEXT("TMap int32"), 1, 1000000);
 	RunPerformanceTest<TMap<int32, FString>, int32>(TEXT("TMap int32"), 10, 1000000);
@@ -545,7 +545,7 @@ TEST_CASE("Core::Containers::TMap::Performance", "[Core][Containers][Perf]")
 	RunPerformanceTest<TMap<FString, FString>, FString>(TEXT("TMap FString"), 10000, 1000000);
 }
 
-TEST_CASE("Core::Containers::TSortedMap::Performance", "[Core][Containers][Perf]")
+TEST_CASE_METHOD(FAutomationTestFixture, "Core::Containers::TSortedMap::Performance", "[Core][Containers][.Perf]")
 {
 	RunPerformanceTest<TSortedMap<int32, FString>, int32>(TEXT("TSortedMap int32"), 1, 1000000);
 	RunPerformanceTest<TSortedMap<int32, FString>, int32>(TEXT("TSortedMap int32"), 10, 1000000);
@@ -623,7 +623,7 @@ namespace
 
 }
 
-TEST_CASE("Core::Containers::TSet::Smoke", "[Core][Containers][Smoke]")
+TEST_CASE_METHOD(FAutomationTestFixture, "Core::Containers::TSet::Smoke", "[Core][Containers][Smoke]")
 {
 	enum class EArgType
 	{
@@ -665,17 +665,13 @@ TEST_CASE("Core::Containers::TSet::Smoke", "[Core][Containers][Smoke]")
 						Set.AddByHash(GetTypeHash(First), First, &bAlreadyInSet);
 					else
 						Set.AddByHash(GetTypeHash(First), MoveTemp(First), &bAlreadyInSet);
-				TestFalse(FullText(TEXT("returns bAlreadyInSet==false for first add")), bAlreadyInSet);
+				TEST_FALSE(FullText(TEXT("returns bAlreadyInSet==false for first add")), bAlreadyInSet);
 
 				FRecorder* Found = Set.Find(First);
 				if (ArgType == EArgType::Copy)
-				{
 					TestTrue(FullText(TEXT("constructs a copy")), Found && Found->Id > First.Id && Found->NumCopies > 0 && Found->Payload == First.Payload);
-				}
 				else
-				{
 					TestTrue(FullText(TEXT("constructs a move")), Found && Found->Id > First.Id && Found->NumCopies == 0 && Found->NumMoves >= 1 && Found->Payload == First.Payload);
-				}
 
 				uint32 FoundId = Found ? Found->Id : 0;
 				Found = Set.Find(First);
@@ -715,15 +711,11 @@ TEST_CASE("Core::Containers::TSet::Smoke", "[Core][Containers][Smoke]")
 						FindOrAddResult = &Set.FindOrAddByHash(GetTypeHash(First), First, &bAlreadyInSet);
 					else
 						FindOrAddResult = &Set.FindOrAddByHash(GetTypeHash(First), MoveTemp(First), &bAlreadyInSet);
-				TestFalse(FullText(TEXT("returns bAlreadyInSet==false for first add")), bAlreadyInSet);
+				TEST_FALSE(FullText(TEXT("returns bAlreadyInSet==false for first add")), bAlreadyInSet);
 				if (ArgType == EArgType::Copy)
-				{
 					TestTrue(FullText(TEXT("on the first constructs a copy")), FindOrAddResult->Id > First.Id && FindOrAddResult->NumCopies > 0 && FindOrAddResult->Payload == First.Payload);
-				}
 				else
-				{
 					TestTrue(FullText(TEXT("on the first constructs a move")), FindOrAddResult->Id > First.Id && FindOrAddResult->NumCopies == 0 && FindOrAddResult->NumMoves >= 1 && FindOrAddResult->Payload == First.Payload);
-				}
 				uint32 FoundId = FindOrAddResult->Id;
 
 				FRecorder* Found = Set.Find(First);
@@ -776,7 +768,7 @@ namespace ArrayViewTests
 	{
 	}
 
-	TEST_CASE("Core::Containers::TArrayView::Smoke", "[Core][Containers][Smoke]")
+	TEST_CASE_METHOD(FAutomationTestFixture, "Core::Containers::TArrayView::Smoke", "[Core][Containers][Smoke]")
 	{
 		// C array + derived-to-base conversions
 		Derived test1[13];
