Add new Unit Tests

diff --git a/Engine/Source/Runtime/Core/Tests/Containers/RingBufferTest.cpp b/Engine/Source/Runtime/Core/Tests/Containers/RingBufferTest.cpp
--- a/Engine/Source/Runtime/Core/Tests/Containers/RingBufferTest.cpp
+++ b/Engine/Source/Runtime/Core/Tests/Containers/RingBufferTest.cpp
@@ -3,7 +3,7 @@
 #include "Containers/RingBuffer.h"
 #include "TestHarness.h"
 
-class FRingBufferTest
+class FRingBufferTest : public FAutomationTestFixture
 {
 public:
 	static bool IsIntegerRange(const TRingBuffer<uint32>& Queue, uint32 Start, uint32 End, bool bForward = true)
@@ -79,34 +79,34 @@ public:
 		{
 			TRingBuffer<uint32> Q(0);
 
-			TestTrue(TEXT("Test empty - IsEmpty"), Q.IsEmpty());
-			TestEqual(TEXT("Test empty - Size"), Q.Num(), 0);
-			TestEqual(TEXT("Test empty - Capacity"), Q.Max(), 0);
-			TestEqual(TEXT("Test empty - Iterator"), Q.begin(), Q.end());
-			TestEqual(TEXT("Test empty - ConvertPointerToIndex"), Q.ConvertPointerToIndex(nullptr), INDEX_NONE);
-			TestEqual(TEXT("Test empty - ConvertPointerToIndex"), Q.ConvertPointerToIndex(reinterpret_cast<uint32*>(this)), INDEX_NONE);
+			TEST_TRUE(TEXT("Test empty - IsEmpty"), Q.IsEmpty());
+			TEST_EQUAL(TEXT("Test empty - Size"), Q.Num(), 0);
+			TEST_EQUAL(TEXT("Test empty - Capacity"), Q.Max(), 0);
+			TEST_EQUAL(TEXT("Test empty - Iterator"), Q.begin(), Q.end());
+			TEST_EQUAL(TEXT("Test empty - ConvertPointerToIndex"), Q.ConvertPointerToIndex(nullptr), INDEX_NONE);
+			TEST_EQUAL(TEXT("Test empty - ConvertPointerToIndex"), Q.ConvertPointerToIndex(reinterpret_cast<uint32*>(this)), INDEX_NONE);
 			Q.Trim();
-			TestEqual(TEXT("Test Trim From empty - Size"), Q.Num(), 0);
-			TestEqual(TEXT("Test Trim From empty - Capacity"), Q.Max(), 0);
+			TEST_EQUAL(TEXT("Test Trim From empty - Size"), Q.Num(), 0);
+			TEST_EQUAL(TEXT("Test Trim From empty - Capacity"), Q.Max(), 0);
 			Q.Reset();
-			TestEqual(TEXT("Test Reset From empty - Size"), Q.Num(), 0);
-			TestEqual(TEXT("Test Reset From empty - Capacity"), Q.Max(), 0);
+			TEST_EQUAL(TEXT("Test Reset From empty - Size"), Q.Num(), 0);
+			TEST_EQUAL(TEXT("Test Reset From empty - Capacity"), Q.Max(), 0);
 			Q.Empty(0);
-			TestEqual(TEXT("Test Empty From empty - Size"), Q.Num(), 0);
-			TestEqual(TEXT("Test Empty From empty - Capacity"), Q.Max(), 0);
+			TEST_EQUAL(TEXT("Test Empty From empty - Size"), Q.Num(), 0);
+			TEST_EQUAL(TEXT("Test Empty From empty - Capacity"), Q.Max(), 0);
 			Q.PopFront(0);
 			Q.Pop(0);
-			TestEqual(TEXT("Test Pop on empty - Size"), Q.Num(), 0);
-			TestEqual(TEXT("Test Pop on empty - Capacity"), Q.Max(), 0);
-			TestEqual(TEXT("Test empty - IsValidIndex"), Q.IsValidIndex(0), false);
+			TEST_EQUAL(TEXT("Test Pop on empty - Size"), Q.Num(), 0);
+			TEST_EQUAL(TEXT("Test Pop on empty - Capacity"), Q.Max(), 0);
+			TEST_EQUAL(TEXT("Test empty - IsValidIndex"), Q.IsValidIndex(0), false);
 
 
 			const TRingBuffer<uint32> ConstQ(0);
-			TestTrue(TEXT("Test const empty - IsEmpty"), ConstQ.IsEmpty());
-			TestEqual(TEXT("Test const empty - Size"), ConstQ.Num(), 0);
-			TestEqual(TEXT("Test const empty - Capacity"), ConstQ.Max(), 0);
-			TestEqual(TEXT("Test const empty - Iterator"), ConstQ.begin(), ConstQ.end());
-			TestEqual(TEXT("Test const empty - ConvertPointerToIndex"), ConstQ.ConvertPointerToIndex(reinterpret_cast<uint32*>(this)), INDEX_NONE);
+			TEST_TRUE(TEXT("Test const empty - IsEmpty"), ConstQ.IsEmpty());
+			TEST_EQUAL(TEXT("Test const empty - Size"), ConstQ.Num(), 0);
+			TEST_EQUAL(TEXT("Test const empty - Capacity"), ConstQ.Max(), 0);
+			TEST_EQUAL(TEXT("Test const empty - Iterator"), ConstQ.begin(), ConstQ.end());
+			TEST_EQUAL(TEXT("Test const empty - ConvertPointerToIndex"), ConstQ.ConvertPointerToIndex(reinterpret_cast<uint32*>(this)), INDEX_NONE);
 		}
 
 		// Test Adding a sequence of elements
@@ -115,38 +115,38 @@ public:
 
 			TRingBuffer<int32> Q(0);
 
-			TestEqual(TEXT("Test AddSequence - Capacity (Implementation Detail)"), Q.Max(), 0);
+			TEST_EQUAL(TEXT("Test AddSequence - Capacity (Implementation Detail)"), Q.Max(), 0);
 			Q.Emplace(0);
-			TestEqual(TEXT("Test AddSequence - Size"), Q.Num(), 1);
-			TestEqual(TEXT("Test AddSequence - Capacity (Implementation Detail)"), Q.Max(), 1);
+			TEST_EQUAL(TEXT("Test AddSequence - Size"), Q.Num(), 1);
+			TEST_EQUAL(TEXT("Test AddSequence - Capacity (Implementation Detail)"), Q.Max(), 1);
 			Q.Emplace(1);
-			TestEqual(TEXT("Test AddSequence - Size"), Q.Num(), 2);
-			TestEqual(TEXT("Implementation Detail - These tests expect that growing size will set capacity to successive powers of 2."), Q.Max(), 2);
+			TEST_EQUAL(TEXT("Test AddSequence - Size"), Q.Num(), 2);
+			TEST_EQUAL(TEXT("Implementation Detail - These tests expect that growing size will set capacity to successive powers of 2."), Q.Max(), 2);
 			for (int32 It = 2; It < FirstSize; ++It)
 			{
 				Q.Emplace(It);
-				TestEqual(TEXT("Test AddSequence - Size"), Q.Num(), It + 1);
-				TestEqual(TEXT("Test AddSequence - Capacity (Implementation Detail)"), static_cast<uint32>(Q.Max()), FMath::RoundUpToPowerOfTwo(It + 1));
+				TEST_EQUAL(TEXT("Test AddSequence - Size"), Q.Num(), It + 1);
+				TEST_EQUAL(TEXT("Test AddSequence - Capacity (Implementation Detail)"), static_cast<uint32>(Q.Max()), FMath::RoundUpToPowerOfTwo(It + 1));
 			}
 
 			for (int32 Index = 0; Index < FirstSize; ++Index)
 			{
-				TestEqual(TEXT("Test AddSequence - Expected values"), Q[Index], Index);
-				TestEqual(TEXT("Test AddSequence const- Expected values"), const_cast<TRingBuffer<int32>&>(Q)[Index], Index);
+				TEST_EQUAL(TEXT("Test AddSequence - Expected values"), Q[Index], Index);
+				TEST_EQUAL(TEXT("Test AddSequence const- Expected values"), const_cast<TRingBuffer<int32>&>(Q)[Index], Index);
 			}
 
 			const TRingBuffer<int32>::IndexType SecondSize = 13;
 			for (int32 It = FirstSize; It < SecondSize; ++It)
 			{
 				Q.Emplace(It);
-				TestEqual(TEXT("Test AddSequence non powerof2 - Size"), Q.Num(), It + 1);
-				TestEqual(TEXT("Test AddSequence non powerof2 const - Capacity (Implementation Detail)"), static_cast<uint32>(Q.Max()), FMath::RoundUpToPowerOfTwo(It + 1));
+				TEST_EQUAL(TEXT("Test AddSequence non powerof2 - Size"), Q.Num(), It + 1);
+				TEST_EQUAL(TEXT("Test AddSequence non powerof2 const - Capacity (Implementation Detail)"), static_cast<uint32>(Q.Max()), FMath::RoundUpToPowerOfTwo(It + 1));
 			}
 
 			for (int32 Index = 0; Index < FirstSize; ++Index)
 			{
-				TestEqual(TEXT("Test AddSequence non powerof2 - Expected values"), Q[Index], Index);
-				TestEqual(TEXT("Test AddSequence non powerof2 const - Expected values"), const_cast<TRingBuffer<int32>&>(Q)[Index], Index);
+				TEST_EQUAL(TEXT("Test AddSequence non powerof2 - Expected values"), Q[Index], Index);
+				TEST_EQUAL(TEXT("Test AddSequence non powerof2 const - Expected values"), const_cast<TRingBuffer<int32>&>(Q)[Index], Index);
 			}
 		}
 
@@ -163,12 +163,12 @@ public:
 				Q.Emplace(It);
 			}
 
-			TestEqual(TEXT("Test Add under Capacity - Size"), Q.Num(), FirstElementsToAdd);
-			TestEqual(TEXT("Test Add under Capacity - Capacity"), Q.Max(), InitialCapacity);
+			TEST_EQUAL(TEXT("Test Add under Capacity - Size"), Q.Num(), FirstElementsToAdd);
+			TEST_EQUAL(TEXT("Test Add under Capacity - Capacity"), Q.Max(), InitialCapacity);
 			for (int32 Index = 0; Index < FirstElementsToAdd; ++Index)
 			{
-				TestEqual(TEXT("Test Add under Capacity - Expected values"), Q[Index], Index);
-				TestEqual(TEXT("Test Add under Capacity const - Expected values"), const_cast<TRingBuffer<int32>&>(Q)[Index], Index);
+				TEST_EQUAL(TEXT("Test Add under Capacity - Expected values"), Q[Index], Index);
+				TEST_EQUAL(TEXT("Test Add under Capacity const - Expected values"), const_cast<TRingBuffer<int32>&>(Q)[Index], Index);
 			}
 
 			for (int32 It = FirstElementsToAdd; It < SecondElementsToAdd; ++It)
@@ -176,24 +176,24 @@ public:
 				Q.Emplace(It);
 			}
 
-			TestEqual(TEXT("Test Add over Capacity - Size"), Q.Num(), SecondElementsToAdd);
-			TestEqual(TEXT("Test Add over Capacity - Capacity (Implementation Detail)"), static_cast<uint32>(Q.Max()), FMath::RoundUpToPowerOfTwo(SecondElementsToAdd));
+			TEST_EQUAL(TEXT("Test Add over Capacity - Size"), Q.Num(), SecondElementsToAdd);
+			TEST_EQUAL(TEXT("Test Add over Capacity - Capacity (Implementation Detail)"), static_cast<uint32>(Q.Max()), FMath::RoundUpToPowerOfTwo(SecondElementsToAdd));
 			for (int32 Index = 0; Index < SecondElementsToAdd; ++Index)
 			{
-				TestEqual(TEXT("Test Add over Capacity - Expected values"), Q[Index], Index);
-				TestEqual(TEXT("Test Add over Capacity const - Expected values"), const_cast<TRingBuffer<int32>&>(Q)[Index], Index);
+				TEST_EQUAL(TEXT("Test Add over Capacity - Expected values"), Q[Index], Index);
+				TEST_EQUAL(TEXT("Test Add over Capacity const - Expected values"), const_cast<TRingBuffer<int32>&>(Q)[Index], Index);
 			}
 		}
 
 		// Test Last/First
 		{
 			TRingBuffer<uint32> Q({ 0,1,2,3 });
-			TestEqual(TEXT("Test Last"), 3, Q.Last());
+			TEST_EQUAL(TEXT("Test Last"), 3, Q.Last());
 			Q.Last() = 4;
-			TestEqual(TEXT("Test Last const"), 4, const_cast<TRingBuffer<uint32>&>(Q).Last());
-			TestEqual(TEXT("Test First"), 0, Q.First());
+			TEST_EQUAL(TEXT("Test Last const"), 4, const_cast<TRingBuffer<uint32>&>(Q).Last());
+			TEST_EQUAL(TEXT("Test First"), 0, Q.First());
 			Q.First() = 5;
-			TestEqual(TEXT("Test First const"), 5, const_cast<TRingBuffer<uint32>&>(Q).First());
+			TEST_EQUAL(TEXT("Test First const"), 5, const_cast<TRingBuffer<uint32>&>(Q).First());
 		}
 
 		// Test PopFrontValue/PopValue
@@ -203,23 +203,23 @@ public:
 
 			Counter::Clear();
 			Counter C(Q.PopFrontValue());
-			TestEqual(TEXT("PopFrontValue - PoppedValue"), C.Value, 30);
-			TestTrue(TEXT("PopFrontValue - ConstructorCounts"), Counter::NumMove > 0 && Counter::NumCopy == 0);
-			TestEqual(TEXT("PopFrontValue - Remaining Values"), Q, TRingBuffer<Counter>({ 31,32,33 }));
+			TEST_EQUAL(TEXT("PopFrontValue - PoppedValue"), C.Value, 30);
+			TEST_TRUE(TEXT("PopFrontValue - ConstructorCounts"), Counter::NumMove > 0 && Counter::NumCopy == 0);
+			TEST_EQUAL(TEXT("PopFrontValue - Remaining Values"), Q, TRingBuffer<Counter>({ 31,32,33 }));
 			Counter::Clear();
-			TestEqual(TEXT("PopFrontValue Inline - PoppedValue"), Q.PopFrontValue().Value, 31);
-			TestTrue(TEXT("PopFrontValue Inline - ConstructorCounts"), Counter::NumCopy == 0);
-			TestEqual(TEXT("PopFrontValue Inline - Remaining Values"), Q, TRingBuffer<Counter>({ 32,33 }));
+			TEST_EQUAL(TEXT("PopFrontValue Inline - PoppedValue"), Q.PopFrontValue().Value, 31);
+			TEST_TRUE(TEXT("PopFrontValue Inline - ConstructorCounts"), Counter::NumCopy == 0);
+			TEST_EQUAL(TEXT("PopFrontValue Inline - Remaining Values"), Q, TRingBuffer<Counter>({ 32,33 }));
 
 			Counter::Clear();
 			Counter D(Q.PopValue());
-			TestEqual(TEXT("PopValue - PoppedValue"), D.Value, 33);
-			TestTrue(TEXT("PopValue - ConstructorCounts"), Counter::NumMove > 0 && Counter::NumCopy == 0);
-			TestEqual(TEXT("PopValue - Remaining Values"), Q, TRingBuffer<Counter>({ Counter(32) }));
+			TEST_EQUAL(TEXT("PopValue - PoppedValue"), D.Value, 33);
+			TEST_TRUE(TEXT("PopValue - ConstructorCounts"), Counter::NumMove > 0 && Counter::NumCopy == 0);
+			TEST_EQUAL(TEXT("PopValue - Remaining Values"), Q, TRingBuffer<Counter>({ Counter(32) }));
 			Counter::Clear();
-			TestEqual(TEXT("PopValue Inline - PoppedValue"), Q.PopValue().Value, 32);
-			TestTrue(TEXT("PopValue Inline - ConstructorCounts"), Counter::NumCopy == 0);
-			TestTrue(TEXT("PopValue Inline - Remaining Values"), Q.IsEmpty());
+			TEST_EQUAL(TEXT("PopValue Inline - PoppedValue"), Q.PopValue().Value, 32);
+			TEST_TRUE(TEXT("PopValue Inline - ConstructorCounts"), Counter::NumCopy == 0);
+			TEST_TRUE(TEXT("PopValue Inline - Remaining Values"), Q.IsEmpty());
 		}
 
 		// Test Initializer_List
@@ -227,11 +227,11 @@ public:
 			const TRingBuffer<int32>::IndexType InitializerSize = 9;
 			TRingBuffer<int32> Q({ 0, 1, 2, 3, 4, 5, 6, 7, 8 });
 
-			TestEqual(TEXT("Test Initializer_List - Size"), Q.Num(), InitializerSize);
-			TestEqual(TEXT("Test Initializer_List - Capacity (Implementation Detail)"), static_cast<uint32>(Q.Max()), FMath::RoundUpToPowerOfTwo(InitializerSize));
+			TEST_EQUAL(TEXT("Test Initializer_List - Size"), Q.Num(), InitializerSize);
+			TEST_EQUAL(TEXT("Test Initializer_List - Capacity (Implementation Detail)"), static_cast<uint32>(Q.Max()), FMath::RoundUpToPowerOfTwo(InitializerSize));
 			for (int32 Index = 0; Index < InitializerSize; ++Index)
 			{
-				TestEqual(TEXT("Test Initializer_List - Expected values"), Q[Index], Index);
+				TEST_EQUAL(TEXT("Test Initializer_List - Expected values"), Q[Index], Index);
 			}
 		}
 
@@ -239,27 +239,27 @@ public:
 		{
 			TRingBuffer<uint32> Original({ 0,1,2,3,4,5,6,7 });
 			TRingBuffer<uint32> Copy(Original);
-			TestEqual(TEXT("Copy Constructor"), Original, Copy);
+			TEST_EQUAL(TEXT("Copy Constructor"), Original, Copy);
 			TRingBuffer<uint32> Moved(MoveTemp(Copy));
-			TestEqual(TEXT("Move Constructor"), Original, Moved);
-			TestEqual(TEXT("Move Constructor did in fact move"), Copy.Max(), 0);
+			TEST_EQUAL(TEXT("Move Constructor"), Original, Moved);
+			TEST_EQUAL(TEXT("Move Constructor did in fact move"), Copy.Max(), 0);
 			TRingBuffer<uint32> AssignCopy;
 			AssignCopy = Original;
-			TestEqual(TEXT("Copy Assignment"), Original, AssignCopy);
+			TEST_EQUAL(TEXT("Copy Assignment"), Original, AssignCopy);
 			TRingBuffer<uint32> AssignMove;
 			AssignMove = MoveTemp(AssignCopy);
-			TestEqual(TEXT("Move Assignment"), Original, AssignMove);
-			TestEqual(TEXT("Move Assignment did in fact move"), AssignCopy.Max(), 0);
+			TEST_EQUAL(TEXT("Move Assignment"), Original, AssignMove);
+			TEST_EQUAL(TEXT("Move Assignment did in fact move"), AssignCopy.Max(), 0);
 		}
 
 		// Test Equality 
 		{
 			auto TestEquality = [this](const TCHAR* Message, bool ExpectedEqual, const TRingBuffer<int32>& A, const TRingBuffer<int32>& B)
 			{
-				TestEqual(*FString::Printf(TEXT("Test equality - %s - A == B"), Message), A == B, ExpectedEqual);
-				TestEqual(*FString::Printf(TEXT("Test equality - %s - B == A"), Message), B == A, ExpectedEqual);
-				TestEqual(*FString::Printf(TEXT("Test equality - %s - A != B"), Message), A != B, !ExpectedEqual);
-				TestEqual(*FString::Printf(TEXT("Test equality - %s - B != A"), Message), B != A, !ExpectedEqual);
+				TEST_EQUAL(*FString::Printf(TEXT("Test equality - %s - A == B"), Message), A == B, ExpectedEqual);
+				TEST_EQUAL(*FString::Printf(TEXT("Test equality - %s - B == A"), Message), B == A, ExpectedEqual);
+				TEST_EQUAL(*FString::Printf(TEXT("Test equality - %s - A != B"), Message), A != B, !ExpectedEqual);
+				TEST_EQUAL(*FString::Printf(TEXT("Test equality - %s - B != A"), Message), B != A, !ExpectedEqual);
 			};
 
 			TestEquality(TEXT("empty"), true, TRingBuffer<int32>(0), TRingBuffer<int32>(0));
@@ -380,12 +380,12 @@ public:
 					if (bIsAddBack)
 					{
 						Q.Add(It);
-						TestEqual(*GetMessage(TEXT("Add and pop - Add")), It, Q[Q.Num() - 1]);
+						TEST_EQUAL(*GetMessage(TEXT("Add and pop - Add")), It, Q[Q.Num() - 1]);
 					}
 					else
 					{
 						Q.AddFront(It);
-						TestEqual(*GetMessage(TEXT("Add and pop - Add")), It, Q[0]);
+						TEST_EQUAL(*GetMessage(TEXT("Add and pop - Add")), It, Q[0]);
 					}
 
 					if (It % ElementPopMod == 0)
@@ -401,14 +401,14 @@ public:
 							PoppedValue = Q[Q.Num() - 1];
 							Q.Pop();
 						}
-						TestEqual(*GetMessage(TEXT("Add and pop - Pop")), ExpectedPoppedValue, PoppedValue);
+						TEST_EQUAL(*GetMessage(TEXT("Add and pop - Pop")), ExpectedPoppedValue, PoppedValue);
 						++ExpectedPoppedValue;
 					}
 				}
 
-				TestEqual(*GetMessage(TEXT("Add and pop - Size")), Q.Num(), ExpectedSize);
-				TestEqual(*GetMessage(TEXT("Add and pop - Capacity")), Q.Max(), ExpectedCapacity);
-				TestTrue(*GetMessage(TEXT("Add and pop - IntegerRange")), IsIntegerRange(Q, ExpectedPoppedValue, ExpectedPoppedValue + ExpectedSize, bIsAddBack));
+				TEST_EQUAL(*GetMessage(TEXT("Add and pop - Size")), Q.Num(), ExpectedSize);
+				TEST_EQUAL(*GetMessage(TEXT("Add and pop - Capacity")), Q.Max(), ExpectedCapacity);
+				TEST_TRUE(*GetMessage(TEXT("Add and pop - IntegerRange")), IsIntegerRange(Q, ExpectedPoppedValue, ExpectedPoppedValue + ExpectedSize, bIsAddBack));
 			}
 
 
@@ -418,8 +418,8 @@ public:
 
 				TRingBuffer<uint32> Q(ElementsToAdd);
 
-				TestTrue(*GetMessage(TEXT("Add and pop all - IsEmpty before")), Q.IsEmpty());
-				TestEqual(*GetMessage(TEXT("Add and pop all - Size before")), Q.Num(), 0);
+				TEST_TRUE(*GetMessage(TEXT("Add and pop all - IsEmpty before")), Q.IsEmpty());
+				TEST_EQUAL(*GetMessage(TEXT("Add and pop all - Size before")), Q.Num(), 0);
 
 				for (TRingBuffer<int32>::IndexType It = 0; It < ElementsToAdd; ++It)
 				{
@@ -433,9 +433,9 @@ public:
 					}
 				}
 
-				TestEqual(*GetMessage(TEXT("Add and pop all - Size")), Q.Num(), ElementsToAdd);
-				TestEqual(*GetMessage(TEXT("Add and pop all - Capacity")), Q.Max(), ElementsToAdd);
-				TestTrue(*GetMessage(TEXT("Add and pop all - Expected")), IsIntegerRange(Q, 0, ElementsToAdd, bIsAddBack));
+				TEST_EQUAL(*GetMessage(TEXT("Add and pop all - Size")), Q.Num(), ElementsToAdd);
+				TEST_EQUAL(*GetMessage(TEXT("Add and pop all - Capacity")), Q.Max(), ElementsToAdd);
+				TEST_TRUE(*GetMessage(TEXT("Add and pop all - Expected")), IsIntegerRange(Q, 0, ElementsToAdd, bIsAddBack));
 
 				for (TRingBuffer<int32>::IndexType It = 0; It < ElementsToAdd; ++It)
 				{
@@ -449,9 +449,9 @@ public:
 					}
 				}
 
-				TestTrue(*GetMessage(TEXT("Add and pop all - IsEmpty after")), Q.IsEmpty());
-				TestEqual(*GetMessage(TEXT("Add and pop all - Size after")), Q.Num(), 0);
-				TestEqual(*GetMessage(TEXT("Add and pop all - Capacity after")), Q.Max(), ElementsToAdd);
+				TEST_TRUE(*GetMessage(TEXT("Add and pop all - IsEmpty after")), Q.IsEmpty());
+				TEST_EQUAL(*GetMessage(TEXT("Add and pop all - Size after")), Q.Num(), 0);
+				TEST_EQUAL(*GetMessage(TEXT("Add and pop all - Capacity after")), Q.Max(), ElementsToAdd);
 			}
 
 			// Test index wrap
@@ -470,8 +470,8 @@ public:
 					Q.Front = Offset;
 					Q.AfterBack = Q.Front;
 
-					TestTrue(*GetMessage(TEXT("index wrap - IsEmpty before")), Q.IsEmpty());
-					TestEqual(*GetMessage(TEXT("index wrap - Size before")), Q.Num(), 0);
+					TEST_TRUE(*GetMessage(TEXT("index wrap - IsEmpty before")), Q.IsEmpty());
+					TEST_EQUAL(*GetMessage(TEXT("index wrap - Size before")), Q.Num(), 0);
 
 					for (TRingBuffer<uint32>::IndexType It = 0; It < ElementsToAdd; ++It)
 					{
@@ -485,9 +485,9 @@ public:
 						}
 					}
 
-					TestEqual(*GetMessage(TEXT("index wrap - Size")), Q.Num(), ElementsToAdd);
-					TestEqual(*GetMessage(TEXT("index wrap - Capacity")), Q.Max(), ElementsToAdd);
-					TestTrue(*GetMessage(TEXT("index wrap - Expected")), IsIntegerRange(Q, 0, ElementsToAdd, bIsAddBack));
+					TEST_EQUAL(*GetMessage(TEXT("index wrap - Size")), Q.Num(), ElementsToAdd);
+					TEST_EQUAL(*GetMessage(TEXT("index wrap - Capacity")), Q.Max(), ElementsToAdd);
+					TEST_TRUE(*GetMessage(TEXT("index wrap - Expected")), IsIntegerRange(Q, 0, ElementsToAdd, bIsAddBack));
 
 					for (TRingBuffer<int32>::IndexType It = 0; It < ElementsToAdd; ++It)
 					{
@@ -501,9 +501,9 @@ public:
 						}
 					}
 
-					TestTrue(*GetMessage(TEXT("index wrap - IsEmpty after")), Q.IsEmpty());
-					TestEqual(*GetMessage(TEXT("index wrap - Size after")), Q.Num(), 0);
-					TestEqual(*GetMessage(TEXT("index wrap - Capacity after")), Q.Max(), ElementsToAdd);
+					TEST_TRUE(*GetMessage(TEXT("index wrap - IsEmpty after")), Q.IsEmpty());
+					TEST_EQUAL(*GetMessage(TEXT("index wrap - Size after")), Q.Num(), 0);
+					TEST_EQUAL(*GetMessage(TEXT("index wrap - Capacity after")), Q.Max(), ElementsToAdd);
 				}
 			}
 		}
@@ -522,9 +522,9 @@ public:
 				Q.Add(It);
 			}
 
-			TestEqual(TEXT("Test Trim - Size"), Q.Num(), ElementsToAdd);
-			TestEqual(TEXT("Test Trim - Capacity"), Q.Max(), ExpectedCapacity);
-			TestTrue(TEXT("Test Trim - Expected"), IsIntegerRange(Q, 0, ElementsToAdd));
+			TEST_EQUAL(TEXT("Test Trim - Size"), Q.Num(), ElementsToAdd);
+			TEST_EQUAL(TEXT("Test Trim - Capacity"), Q.Max(), ExpectedCapacity);
+			TEST_TRUE(TEXT("Test Trim - Expected"), IsIntegerRange(Q, 0, ElementsToAdd));
 
 			for (TRingBuffer<int32>::IndexType It = 0; It < ElementsToPop; ++It)
 			{
@@ -533,9 +533,9 @@ public:
 
 			Q.Trim();
 
-			TestEqual(TEXT("Test Trim - Size"), Q.Num(), ElementsToAdd - ElementsToPop);
-			TestEqual(TEXT("Test Trim - Capacity"), Q.Max(), ExpectedCapacityAfterTrim);
-			TestTrue(TEXT("Test Trim - Expected"), IsIntegerRange(Q, ElementsToPop, ElementsToAdd));
+			TEST_EQUAL(TEXT("Test Trim - Size"), Q.Num(), ElementsToAdd - ElementsToPop);
+			TEST_EQUAL(TEXT("Test Trim - Capacity"), Q.Max(), ExpectedCapacityAfterTrim);
+			TEST_TRUE(TEXT("Test Trim - Expected"), IsIntegerRange(Q, ElementsToPop, ElementsToAdd));
 		}
 
 		// Test First and Last acting as two stacks
@@ -548,36 +548,36 @@ public:
 			for (uint32 It = 0; It < ElementsToAdd; ++It)
 			{
 				Q.Add(It);
-				TestEqual(TEXT("Test TwoStacks - AddBack"), Q.Last(), It);
+				TEST_EQUAL(TEXT("Test TwoStacks - AddBack"), Q.Last(), It);
 				Q.AddFront(It);
-				TestEqual(TEXT("Test TwoStacks - AddFront"), Q.First(), It);
+				TEST_EQUAL(TEXT("Test TwoStacks - AddFront"), Q.First(), It);
 				if (It % ElementPopMod == 0)
 				{
 					uint32 AddValue = 0xfefefefe;
 					Q.Add(AddValue);
-					TestEqual(TEXT("Test TwoStacks - Sporadic Pop"), Q.Last(), AddValue);
+					TEST_EQUAL(TEXT("Test TwoStacks - Sporadic Pop"), Q.Last(), AddValue);
 					Q.Pop();
 					Q.AddFront(AddValue);
-					TestEqual(TEXT("Test TwoStacks - Sporadic PopFront"), Q.First(), AddValue);
+					TEST_EQUAL(TEXT("Test TwoStacks - Sporadic PopFront"), Q.First(), AddValue);
 					Q.PopFront();
 				}
 			}
 
-			TestEqual(TEXT("Test TwoStacks - MiddleSize"), Q.Num(), ElementsToAdd * 2);
+			TEST_EQUAL(TEXT("Test TwoStacks - MiddleSize"), Q.Num(), ElementsToAdd * 2);
 			for (uint32 It = 0; It < ElementsToAdd * 2; ++It)
 			{
-				TestEqual(*FString::Printf(TEXT("TwoStacks - Middle value %d"), It), Q[It], (It < ElementsToAdd ? ElementsToAdd - 1 - It : It - ElementsToAdd));
+				TEST_EQUAL(*FString::Printf(TEXT("TwoStacks - Middle value %d"), It), Q[It], (It < ElementsToAdd ? ElementsToAdd - 1 - It : It - ElementsToAdd));
 			}
 
 			for (uint32 It = 0; It < ElementsToAdd; ++It)
 			{
-				TestEqual(TEXT("Test TwoStacks - Final Pop"), Q.Last(), ElementsToAdd - 1 - It);
+				TEST_EQUAL(TEXT("Test TwoStacks - Final Pop"), Q.Last(), ElementsToAdd - 1 - It);
 				Q.Pop();
-				TestEqual(TEXT("Test TwoStacks - Final PopFront"), Q.First(), ElementsToAdd - 1 - It);
+				TEST_EQUAL(TEXT("Test TwoStacks - Final PopFront"), Q.First(), ElementsToAdd - 1 - It);
 				Q.PopFront();
 			}
 
-			TestEqual(TEXT("Test TwoStacks - FinalSize"), Q.Num(), 0);
+			TEST_EQUAL(TEXT("Test TwoStacks - FinalSize"), Q.Num(), 0);
 		}
 
 		// Test adding into space that has been cleared from popping on the other side
@@ -591,8 +591,8 @@ public:
 				};
 				TRingBuffer<uint32> Q({ 0,1,2,3,4,5,6,7 });
 				TRingBuffer<int32>::IndexType InitialSize = 8;
-				TestEqual(*GetMessage(TEXT("InitialSize")), InitialSize, Q.Num());
-				TestEqual(*GetMessage(TEXT("InitialCapacity (Implementation Detail)")), InitialSize, Q.Max());
+				TEST_EQUAL(*GetMessage(TEXT("InitialSize")), InitialSize, Q.Num());
+				TEST_EQUAL(*GetMessage(TEXT("InitialCapacity (Implementation Detail)")), InitialSize, Q.Max());
 
 				if (bIsAddBack)
 				{
@@ -602,8 +602,8 @@ public:
 				{
 					Q.PopFront();
 				}
-				TestEqual(*GetMessage(TEXT("PoppedSize")), InitialSize - 1, Q.Num());
-				TestEqual(*GetMessage(TEXT("PoppedCapacity")), InitialSize, Q.Max());
+				TEST_EQUAL(*GetMessage(TEXT("PoppedSize")), InitialSize - 1, Q.Num());
+				TEST_EQUAL(*GetMessage(TEXT("PoppedCapacity")), InitialSize, Q.Max());
 
 				if (bIsAddBack)
 				{
@@ -613,15 +613,15 @@ public:
 				{
 					Q.Add(8);
 				}
-				TestEqual(*GetMessage(TEXT("AddedSize")), InitialSize, Q.Num());
-				TestEqual(*GetMessage(TEXT("AddedCapacity")), InitialSize, Q.Max());
+				TEST_EQUAL(*GetMessage(TEXT("AddedSize")), InitialSize, Q.Num());
+				TEST_EQUAL(*GetMessage(TEXT("AddedCapacity")), InitialSize, Q.Max());
 				if (bIsAddBack)
 				{
-					TestEqual(*GetMessage(TEXT("AddedValues")), Q, TRingBuffer<uint32>({ 8,0,1,2,3,4,5,6 }));
+					TEST_EQUAL(*GetMessage(TEXT("AddedValues")), Q, TRingBuffer<uint32>({ 8,0,1,2,3,4,5,6 }));
 				}
 				else
 				{
-					TestEqual(*GetMessage(TEXT("AddedValues")), Q, TRingBuffer<uint32>({ 1,2,3,4,5,6,7,8 }));
+					TEST_EQUAL(*GetMessage(TEXT("AddedValues")), Q, TRingBuffer<uint32>({ 1,2,3,4,5,6,7,8 }));
 				}
 
 				if (bIsAddBack)
@@ -632,15 +632,15 @@ public:
 				{
 					Q.Add(9);
 				}
-				TestEqual(*GetMessage(TEXT("Second AddedSize")), InitialSize + 1, Q.Num());
-				TestEqual(*GetMessage(TEXT("Second AddedCapacity")), static_cast<uint32>(FMath::RoundUpToPowerOfTwo(InitialSize + 1)), Q.Max());
+				TEST_EQUAL(*GetMessage(TEXT("Second AddedSize")), InitialSize + 1, Q.Num());
+				TEST_EQUAL(*GetMessage(TEXT("Second AddedCapacity")), static_cast<uint32>(FMath::RoundUpToPowerOfTwo(InitialSize + 1)), Q.Max());
 				if (bIsAddBack)
 				{
-					TestEqual(*GetMessage(TEXT("Second AddedValues")), Q, TRingBuffer<uint32>({ 9,8,0,1,2,3,4,5,6 }));
+					TEST_EQUAL(*GetMessage(TEXT("Second AddedValues")), Q, TRingBuffer<uint32>({ 9,8,0,1,2,3,4,5,6 }));
 				}
 				else
 				{
-					TestEqual(*GetMessage(TEXT("Second AddedValues")), Q, TRingBuffer<uint32>({ 1,2,3,4,5,6,7,8,9 }));
+					TEST_EQUAL(*GetMessage(TEXT("Second AddedValues")), Q, TRingBuffer<uint32>({ 1,2,3,4,5,6,7,8,9 }));
 				}
 			}
 		}
@@ -648,11 +648,11 @@ public:
 		// Test Empty to a capacity
 		{
 			TRingBuffer<uint32> Q(16);
-			TestEqual(TEXT("Test EmptyToCapacity - InitialCapacity"), 16, Q.Max());
+			TEST_EQUAL(TEXT("Test EmptyToCapacity - InitialCapacity"), 16, Q.Max());
 			Q.Empty(8);
-			TestEqual(TEXT("Test EmptyToCapacity - Lower"), 8, Q.Max());
+			TEST_EQUAL(TEXT("Test EmptyToCapacity - Lower"), 8, Q.Max());
 			Q.Empty(32);
-			TestEqual(TEXT("Test EmptyToCapacity - Higher"), 32, Q.Max());
+			TEST_EQUAL(TEXT("Test EmptyToCapacity - Higher"), 32, Q.Max());
 		}
 
 		// Test Different Add constructors
@@ -663,7 +663,7 @@ public:
 			};
 			auto TestCounts = [this](const TCHAR* Message, int32 NumVoid, int32 NumCopy, int32 NumMove, int32 NumDestruct)
 			{
-				TestTrue(Message, NumVoid == Counter::NumVoid && NumCopy == Counter::NumCopy && NumMove == Counter::NumMove && NumDestruct == Counter::NumDestruct);
+				TEST_TRUE(Message, NumVoid == Counter::NumVoid && NumCopy == Counter::NumCopy && NumMove == Counter::NumMove && NumDestruct == Counter::NumDestruct);
 			};
 
 			Clear();
@@ -695,10 +695,10 @@ public:
 				{
 					const TCHAR* DirectionText = bAddBack ? TEXT("Back") : TEXT("Front");
 					bool bElementExists = Q.Num() == 1;
-					TestTrue(*FString::Printf(TEXT("Test Add Constructors - %s%s ElementExists"), Message, DirectionText), bElementExists);
+					TEST_TRUE(*FString::Printf(TEXT("Test Add Constructors - %s%s ElementExists"), Message, DirectionText), bElementExists);
 					if (bWasInitialized && bElementExists)
 					{
-						TestTrue(*FString::Printf(TEXT("Test Add Constructors - %s%s ValueEquals"), Message, DirectionText), Q.First().Value == MarkerValue);
+						TEST_TRUE(*FString::Printf(TEXT("Test Add Constructors - %s%s ValueEquals"), Message, DirectionText), Q.First().Value == MarkerValue);
 					}
 					Q.PopFront();
 					TestCounts(*FString::Printf(TEXT("Test Add Constructors - %s%s CountsEqual"), Message, DirectionText), NumVoid, NumCopy, NumMove, NumDestruct);
@@ -732,12 +732,12 @@ public:
 			{
 				TRingBuffer<uint32> Q{ 0,1,2,3,4,5,6,7 };
 				Q.RemoveAt(2);
-				TestEqual(TEXT("Test RemoveAt Front Closest"), TRingBuffer<uint32>({ 0,1,3,4,5,6,7 }), Q);
+				TEST_EQUAL(TEXT("Test RemoveAt Front Closest"), TRingBuffer<uint32>({ 0,1,3,4,5,6,7 }), Q);
 			}
 			{
 				TRingBuffer<uint32> Q{ 0,1,2,3,4,5,6,7 };
 				Q.RemoveAt(5);
-				TestEqual(TEXT("Test RemoveAt Back Closest"), TRingBuffer<uint32>({ 0,1,2,3,4,6,7 }), Q);
+				TEST_EQUAL(TEXT("Test RemoveAt Back Closest"), TRingBuffer<uint32>({ 0,1,2,3,4,6,7 }), Q);
 			}
 			{
 				TRingBuffer<uint32> Q{ 0,1,2,3,4,5,6,7 };
@@ -746,7 +746,7 @@ public:
 				Q.AfterBack += Offset;
 				//Now equal to: TRingBuffer<uint32> Q{ 4,5,6,7,0,1,2,3 };
 				Q.RemoveAt(2);
-				TestEqual(TEXT("Test RemoveAt Front Closest With Offset"), TRingBuffer<uint32>({ 4,5,7,0,1,2,3 }), Q);
+				TEST_EQUAL(TEXT("Test RemoveAt Front Closest With Offset"), TRingBuffer<uint32>({ 4,5,7,0,1,2,3 }), Q);
 			}
 			{
 				TRingBuffer<uint32> Q{ 0,1,2,3,4,5,6,7 };
@@ -755,7 +755,7 @@ public:
 				Q.AfterBack += Offset;
 				//Now equal to: TRingBuffer<uint32> Q{ 4,5,6,7,0,1,2,3 };
 				Q.RemoveAt(5);
-				TestEqual(TEXT("Test RemoveAt Back Closest With Offset"), TRingBuffer<uint32>({ 4,5,6,7,0,2,3 }), Q);
+				TEST_EQUAL(TEXT("Test RemoveAt Back Closest With Offset"), TRingBuffer<uint32>({ 4,5,6,7,0,2,3 }), Q);
 			}
 		}
 
@@ -766,9 +766,9 @@ public:
 				uint32 Counter = 0;
 				for (uint32 Value : Q)
 				{
-					TestEqual(TEXT("Test Iteration - Value"), Counter++, Value);
+					TEST_EQUAL(TEXT("Test Iteration - Value"), Counter++, Value);
 				}
-				TestEqual(TEXT("Test Iteration - Num"), Counter, 8);
+				TEST_EQUAL(TEXT("Test Iteration - Num"), Counter, 8);
 			}
 			{
 				TRingBuffer<uint32> Q{ 4,5,6,7,0,1,2,3 };
@@ -779,9 +779,9 @@ public:
 				uint32 Counter = 0;
 				for (uint32 Value : Q)
 				{
-					TestEqual(TEXT("Test Iteration with Offset - Value"), Counter++, Value);
+					TEST_EQUAL(TEXT("Test Iteration with Offset - Value"), Counter++, Value);
 				}
-				TestEqual(TEXT("Test Iteration with Offset  - Num"), Counter, 8);
+				TEST_EQUAL(TEXT("Test Iteration with Offset  - Num"), Counter, 8);
 			}
 		}
 
@@ -793,11 +793,11 @@ public:
 				Q.Front += Offset;
 				Q.AfterBack += Offset;
 				// Now equal to 0,1,2,3,4,5,6,7
-				TestEqual(TEXT("Test ConvertPointerToIndex - before array"), Q.ConvertPointerToIndex(&Q[0] - 100), INDEX_NONE);
-				TestEqual(TEXT("Test ConvertPointerToIndex - after array"), Q.ConvertPointerToIndex(&Q[0] + 100), INDEX_NONE);
+				TEST_EQUAL(TEXT("Test ConvertPointerToIndex - before array"), Q.ConvertPointerToIndex(&Q[0] - 100), INDEX_NONE);
+				TEST_EQUAL(TEXT("Test ConvertPointerToIndex - after array"), Q.ConvertPointerToIndex(&Q[0] + 100), INDEX_NONE);
 				for (int32 It = 0; It < 8; ++It)
 				{
-					TestEqual(TEXT("Test ConvertPointerToIndex - Values"), Q.ConvertPointerToIndex(&Q[It]), It);
+					TEST_EQUAL(TEXT("Test ConvertPointerToIndex - Values"), Q.ConvertPointerToIndex(&Q[It]), It);
 				}
 			}
 
@@ -811,10 +811,10 @@ public:
 				// 8 Invalids, followed by 0,1,2,3,4,5,6, followed by Invalid
 				for (int32 It = 0; It < 7; ++It)
 				{
-					TestEqual(TEXT("Test ConvertPointerToIndex - Cap - Values"), Q.ConvertPointerToIndex(&Q[It]), It);
+					TEST_EQUAL(TEXT("Test ConvertPointerToIndex - Cap - Values"), Q.ConvertPointerToIndex(&Q[It]), It);
 				}
-				TestEqual(TEXT("Test ConvertPointerToIndex - Cap - After End"), Q.ConvertPointerToIndex(&Q[6] + 1), INDEX_NONE);
-				TestEqual(TEXT("Test ConvertPointerToIndex - Cap - Before Start"), Q.ConvertPointerToIndex(&Q[0] - 1), INDEX_NONE);
+				TEST_EQUAL(TEXT("Test ConvertPointerToIndex - Cap - After End"), Q.ConvertPointerToIndex(&Q[6] + 1), INDEX_NONE);
+				TEST_EQUAL(TEXT("Test ConvertPointerToIndex - Cap - Before Start"), Q.ConvertPointerToIndex(&Q[0] - 1), INDEX_NONE);
 			}
 		}
 
@@ -826,9 +826,9 @@ public:
 			{
 				Q.Add(It);
 			}
-			TestTrue(TEXT("Test Front<Capacity - Setup"), (Q.Front & Q.IndexMask) == Q.IndexMask && Q.Num() == Q.Max());
+			TEST_TRUE(TEXT("Test Front<Capacity - Setup"), (Q.Front & Q.IndexMask) == Q.IndexMask && Q.Num() == Q.Max());
 			Q.PopFront(8);
-			TestTrue(TEXT("Test Front<Capacity - Contract is true"), static_cast<uint32>(Q.Front) < static_cast<uint32>(Q.Max()));
+			TEST_TRUE(TEXT("Test Front<Capacity - Contract is true"), static_cast<uint32>(Q.Front) < static_cast<uint32>(Q.Max()));
 		}
 
 		// Test IsValidIndex
@@ -836,29 +836,29 @@ public:
 			TRingBuffer<uint32> Q({ 0,1,2,3,4 });
 			for (int32 It = 0; It < Q.Num(); ++It)
 			{
-				TestEqual(TEXT("IsValidIndex - InRange"), Q.IsValidIndex(It), true);
+				TEST_EQUAL(TEXT("IsValidIndex - InRange"), Q.IsValidIndex(It), true);
 			}
-			TestEqual(TEXT("IsValidIndex - Negative"), Q.IsValidIndex(-1), false);
-			TestEqual(TEXT("IsValidIndex - Num()"), Q.IsValidIndex(Q.Num() + 1), false);
-			TestEqual(TEXT("IsValidIndex - Capacity"), Q.IsValidIndex(Q.Max()), false);
-			TestEqual(TEXT("IsValidIndex - Capacity + 1"), Q.IsValidIndex(Q.Max()+1), false);
+			TEST_EQUAL(TEXT("IsValidIndex - Negative"), Q.IsValidIndex(-1), false);
+			TEST_EQUAL(TEXT("IsValidIndex - Num()"), Q.IsValidIndex(Q.Num() + 1), false);
+			TEST_EQUAL(TEXT("IsValidIndex - Capacity"), Q.IsValidIndex(Q.Max()), false);
+			TEST_EQUAL(TEXT("IsValidIndex - Capacity + 1"), Q.IsValidIndex(Q.Max()+1), false);
 		}
 
 		// Test Compact
 		{
 			{
 				TRingBuffer<uint32> QEmpty;
-				TestEqual(TEXT("Compact - Empty zero capacity"), QEmpty.Compact().Num(), 0);
+				TEST_EQUAL(TEXT("Compact - Empty zero capacity"), QEmpty.Compact().Num(), 0);
 				QEmpty.Add(1);
 				QEmpty.PopFront();
-				TestEqual(TEXT("Compact - Empty non-zero capacity"), QEmpty.Compact().Num(), 0);
+				TEST_EQUAL(TEXT("Compact - Empty non-zero capacity"), QEmpty.Compact().Num(), 0);
 			}
 			{
 				TArrayView<uint32> View;
 				TRingBuffer<uint32> Q(8);
 				Q.AddFront(37);
 				View = Q.Compact();
-				TestTrue(TEXT("Compact - Front at end"), ArrayViewsEqual(View, TArrayView<const uint32>({ 37 })));
+				TEST_TRUE(TEXT("Compact - Front at end"), ArrayViewsEqual(View, TArrayView<const uint32>({ 37 })));
 			}
 			{
 				TArrayView<uint32> View;
@@ -869,10 +869,10 @@ public:
 				}
 				Q.PopFront();
 				TRingBuffer<uint32>::StorageModuloType SavedFront = Q.Front;
-				TestTrue(TEXT("Compact - Front in middle - setup"), SavedFront > 0);
+				TEST_TRUE(TEXT("Compact - Front in middle - setup"), SavedFront > 0);
 				View = Q.Compact();
-				TestTrue(TEXT("Compact - Front in middle - values"), ArrayViewsEqual(View, TArrayView<const uint32>({ 1,2,3,4,5 })));
-				TestTrue(TEXT("Compact - Front in middle - no reallocate"), Q.Front == SavedFront);
+				TEST_TRUE(TEXT("Compact - Front in middle - values"), ArrayViewsEqual(View, TArrayView<const uint32>({ 1,2,3,4,5 })));
+				TEST_TRUE(TEXT("Compact - Front in middle - no reallocate"), Q.Front == SavedFront);
 			}
 			{
 				TArrayView<uint32> View;
@@ -882,10 +882,10 @@ public:
 					Q.Add(It);
 				}
 				Q.AddFront(0);
-				TestTrue(TEXT("Compact - Full array front at end - setup"), (Q.Front & Q.IndexMask) == 7);
+				TEST_TRUE(TEXT("Compact - Full array front at end - setup"), (Q.Front & Q.IndexMask) == 7);
 				View = Q.Compact();
-				TestTrue(TEXT("Compact - Full array front at end - values"), ArrayViewsEqual(View, TArrayView<const uint32>({0,1,2,3,4,5,6,7 })));
-				TestTrue(TEXT("Compact - Full array front at end - reallocated"), Q.Front == 0);
+				TEST_TRUE(TEXT("Compact - Full array front at end - values"), ArrayViewsEqual(View, TArrayView<const uint32>({0,1,2,3,4,5,6,7 })));
+				TEST_TRUE(TEXT("Compact - Full array front at end - reallocated"), Q.Front == 0);
 			}
 			{
 				TArrayView<uint32> View;
@@ -895,10 +895,10 @@ public:
 					Q.Add(It);
 				}
 				uint32* SavedData = Q.AllocationData;
-				TestTrue(TEXT("Compact - Full array front at start - setup"), Q.Front == 0);
+				TEST_TRUE(TEXT("Compact - Full array front at start - setup"), Q.Front == 0);
 				View = Q.Compact();
-				TestTrue(TEXT("Compact - Full array front at start - values"), ArrayViewsEqual(View, TArrayView<const uint32>({ 0,1,2,3,4,5,6,7 })));
-				TestTrue(TEXT("Compact - Full array front at start - no reallocate"), Q.AllocationData == SavedData);
+				TEST_TRUE(TEXT("Compact - Full array front at start - values"), ArrayViewsEqual(View, TArrayView<const uint32>({ 0,1,2,3,4,5,6,7 })));
+				TEST_TRUE(TEXT("Compact - Full array front at start - no reallocate"), Q.AllocationData == SavedData);
 			}
 		}
 
@@ -909,75 +909,75 @@ public:
 				TRingBuffer<Counter> Q;
 				Value.Value = 2;
 				Counter::Clear();
-				TestEqual(TEXT("Remove - empty"), Q.Remove(Value), 0);
-				TestEqual(TEXT("Remove - empty - destructor count"), Counter::NumDestruct, 0);
+				TEST_EQUAL(TEXT("Remove - empty"), Q.Remove(Value), 0);
+				TEST_EQUAL(TEXT("Remove - empty - destructor count"), Counter::NumDestruct, 0);
 			}
 			{
 				TRingBuffer<Counter> Q({ 0,1,2,3,4 });
 				Value.Value = 5;
 				Counter::Clear();
-				TestEqual(TEXT("Remove - no hits"), Q.Remove(Value), 0);
-				TestEqual(TEXT("Remove - no hits - destructor count"), Counter::NumDestruct, 0);
+				TEST_EQUAL(TEXT("Remove - no hits"), Q.Remove(Value), 0);
+				TEST_EQUAL(TEXT("Remove - no hits - destructor count"), Counter::NumDestruct, 0);
 				Q.Add(5);
-				TestTrue(TEXT("Remove - no hits - values"), Q == TRingBuffer<Counter>({ 0,1,2,3,4,5 }));
+				TEST_TRUE(TEXT("Remove - no hits - values"), Q == TRingBuffer<Counter>({ 0,1,2,3,4,5 }));
 			}
 			{
 				TRingBuffer<Counter> Q({ 1,2,3,4 });
 				Q.AddFront(0);
 				Value.Value = 0;
 				Counter::Clear();
-				TestEqual(TEXT("Remove - one element at front - num"), Q.Remove(Value), 1);
-				TestEqual(TEXT("Remove - one element at front - destructor count"), Counter::NumDestruct, 5);
+				TEST_EQUAL(TEXT("Remove - one element at front - num"), Q.Remove(Value), 1);
+				TEST_EQUAL(TEXT("Remove - one element at front - destructor count"), Counter::NumDestruct, 5);
 				Q.Add(5);
-				TestTrue(TEXT("Remove - one element at front - values"), Q == TRingBuffer<Counter>({ 1,2,3,4,5 }));
+				TEST_TRUE(TEXT("Remove - one element at front - values"), Q == TRingBuffer<Counter>({ 1,2,3,4,5 }));
 			}
 			{
 				TRingBuffer<Counter> Q({ 0,1,2,3,4 });
 				Value.Value = 2;
 				Counter::Clear();
-				TestEqual(TEXT("Remove - one element in mid - num"), Q.Remove(Value), 1);
-				TestEqual(TEXT("Remove - one element in mid - destructor count"), Counter::NumDestruct, 3);
+				TEST_EQUAL(TEXT("Remove - one element in mid - num"), Q.Remove(Value), 1);
+				TEST_EQUAL(TEXT("Remove - one element in mid - destructor count"), Counter::NumDestruct, 3);
 				Q.Add(5);
-				TestTrue(TEXT("Remove - one element in mid - values"), Q == TRingBuffer<Counter>({ 0,1,3,4,5 }));
+				TEST_TRUE(TEXT("Remove - one element in mid - values"), Q == TRingBuffer<Counter>({ 0,1,3,4,5 }));
 			}
 			{
 				TRingBuffer<Counter> Q({ 1,2,3,4 });
 				Q.AddFront(0);
 				Value.Value = 2;
 				Counter::Clear();
-				TestEqual(TEXT("Remove - one element in mid - front at end"), Q.Remove(Value), 1);
-				TestEqual(TEXT("Remove - one element in mid - front at end - destructor count"), Counter::NumDestruct, 3);
+				TEST_EQUAL(TEXT("Remove - one element in mid - front at end"), Q.Remove(Value), 1);
+				TEST_EQUAL(TEXT("Remove - one element in mid - front at end - destructor count"), Counter::NumDestruct, 3);
 				Q.Add(5);
-				TestTrue(TEXT("Remove - one element in mid - front at end - values"), Q == TRingBuffer<Counter>({ 0,1,3,4,5 }));
+				TEST_TRUE(TEXT("Remove - one element in mid - front at end - values"), Q == TRingBuffer<Counter>({ 0,1,3,4,5 }));
 			}
 			{
 				TRingBuffer<Counter> Q({ 0,1,2,3,4 });
 				Value.Value = 4;
 				Counter::Clear();
-				TestEqual(TEXT("Remove - one element - element at end - num"), Q.Remove(Value), 1);
-				TestEqual(TEXT("Remove - one element - element at end - destructor count"), Counter::NumDestruct, 1);
+				TEST_EQUAL(TEXT("Remove - one element - element at end - num"), Q.Remove(Value), 1);
+				TEST_EQUAL(TEXT("Remove - one element - element at end - destructor count"), Counter::NumDestruct, 1);
 				Q.Add(5);
-				TestTrue(TEXT("Remove - one element - element at end - values"), Q == TRingBuffer<Counter>({ 0,1,2,3,5 }));
+				TEST_TRUE(TEXT("Remove - one element - element at end - values"), Q == TRingBuffer<Counter>({ 0,1,2,3,5 }));
 			}
 			{
 				TRingBuffer<Counter> Q({ 1,2,3,4 });
 				Q.AddFront(4);
 				Value.Value = 4;
 				Counter::Clear();
-				TestEqual(TEXT("Remove - one element at front one at end - num"), Q.Remove(Value), 2);
-				TestEqual(TEXT("Remove - one element at front one at end - destructor count"), Counter::NumDestruct, 5);
+				TEST_EQUAL(TEXT("Remove - one element at front one at end - num"), Q.Remove(Value), 2);
+				TEST_EQUAL(TEXT("Remove - one element at front one at end - destructor count"), Counter::NumDestruct, 5);
 				Q.Add(5);
-				TestTrue(TEXT("Remove - one element at front one at end - values"), Q == TRingBuffer<Counter>({ 1,2,3,5 }));
+				TEST_TRUE(TEXT("Remove - one element at front one at end - values"), Q == TRingBuffer<Counter>({ 1,2,3,5 }));
 			}
 			{
 				TRingBuffer<Counter> Q({ 1,2,3,4 });
 				Q.AddFront(1);
 				Value.Value = 1;
 				Counter::Clear();
-				TestEqual(TEXT("Remove - two elements - front at end - num"), Q.Remove(Value), 2);
-				TestEqual(TEXT("Remove - two elements - front at end - destructor count"), Counter::NumDestruct, 5);
+				TEST_EQUAL(TEXT("Remove - two elements - front at end - num"), Q.Remove(Value), 2);
+				TEST_EQUAL(TEXT("Remove - two elements - front at end - destructor count"), Counter::NumDestruct, 5);
 				Q.Add(5);
-				TestTrue(TEXT("Remove - two elements - front at end - values"), Q == TRingBuffer<Counter>({ 2,3,4,5 }));
+				TEST_TRUE(TEXT("Remove - two elements - front at end - values"), Q == TRingBuffer<Counter>({ 2,3,4,5 }));
 			}
 		}
 
@@ -992,9 +992,9 @@ public:
 			{
 				TRingBuffer<T> Q{ 0,1,2,3,4,5,6,7 };
 				Q.ShiftIndexToFront(5);
-				TestEqual(TEXT("ShiftIndexToFront"), TRingBuffer<T>({ 5,0,1,2,3,4,6,7 }), Q);
+				TEST_EQUAL(TEXT("ShiftIndexToFront"), TRingBuffer<T>({ 5,0,1,2,3,4,6,7 }), Q);
 				Q.ShiftIndexToBack(3);
-				TestEqual(TEXT("ShiftIndexToBack"), TRingBuffer<T>({ 5,0,1,3,4,6,7,2 }), Q);
+				TEST_EQUAL(TEXT("ShiftIndexToBack"), TRingBuffer<T>({ 5,0,1,3,4,6,7,2 }), Q);
 			}
 
 			{
@@ -1004,18 +1004,18 @@ public:
 				Q.AfterBack += Offset;
 				//Now equal to: TRingBuffer<uint32> Q{ 4,5,6,7,0,1,2,3 };
 				Q.ShiftIndexToFront(5);
-				TestEqual(TEXT("ShiftIndexToFront With Offset"), TRingBuffer<T>({ 1,4,5,6,7,0,2,3 }), Q);
+				TEST_EQUAL(TEXT("ShiftIndexToFront With Offset"), TRingBuffer<T>({ 1,4,5,6,7,0,2,3 }), Q);
 				Q.ShiftIndexToBack(3);
-				TestEqual(TEXT("ShiftIndexToBack With Offset"), TRingBuffer<T>({ 1,4,5,7,0,2,3,6 }), Q);
+				TEST_EQUAL(TEXT("ShiftIndexToBack With Offset"), TRingBuffer<T>({ 1,4,5,7,0,2,3,6 }), Q);
 			}
 
 			{
 				TRingBuffer<T> Q{ 0,1,2,3,4,5,6,7,8 };
-				TestEqual(TEXT("ShiftIndexToFront Cap - Capacity"), Q.Max(), 16);
+				TEST_EQUAL(TEXT("ShiftIndexToFront Cap - Capacity"), Q.Max(), 16);
 				Q.ShiftIndexToFront(5);
-				TestEqual(TEXT("ShiftIndexToFront Cap"), TRingBuffer<T>({ 5,0,1,2,3,4,6,7,8 }), Q);
+				TEST_EQUAL(TEXT("ShiftIndexToFront Cap"), TRingBuffer<T>({ 5,0,1,2,3,4,6,7,8 }), Q);
 				Q.ShiftIndexToBack(3);
-				TestEqual(TEXT("ShiftIndexToBack Cap"), TRingBuffer<T>({ 5,0,1,3,4,6,7,8,2 }), Q);
+				TEST_EQUAL(TEXT("ShiftIndexToBack Cap"), TRingBuffer<T>({ 5,0,1,3,4,6,7,8,2 }), Q);
 			}
 
 			{
@@ -1027,9 +1027,9 @@ public:
 				Q.Pop();
 				// 8 Invalids, followed by 0,1,2,3,4,5,6, followed by Invalid
 				Q.ShiftIndexToFront(5);
-				TestEqual(TEXT("ShiftIndexToFront Cap With Offset"), TRingBuffer<T>({ 5,0,1,2,3,4,6 }), Q);
+				TEST_EQUAL(TEXT("ShiftIndexToFront Cap With Offset"), TRingBuffer<T>({ 5,0,1,2,3,4,6 }), Q);
 				Q.ShiftIndexToBack(3);
-				TestEqual(TEXT("ShiftIndexToBack Cap With Offset"), TRingBuffer<T>({ 5,0,1,3,4,6,2 }), Q);
+				TEST_EQUAL(TEXT("ShiftIndexToBack Cap With Offset"), TRingBuffer<T>({ 5,0,1,3,4,6,2 }), Q);
 			}
 
 			{
@@ -1041,9 +1041,9 @@ public:
 				Q.Add(8);
 				// 8, (AfterBack), followed by 7 Invalids, followed by (Start) 0,1,2,3,4,5,6,7
 				Q.ShiftIndexToFront(8);
-				TestEqual(TEXT("ShiftIndexToFront Cap With Wrapped"), TRingBuffer<T>({ 8,0,1,2,3,4,5,6,7 }), Q);
+				TEST_EQUAL(TEXT("ShiftIndexToFront Cap With Wrapped"), TRingBuffer<T>({ 8,0,1,2,3,4,5,6,7 }), Q);
 				Q.ShiftIndexToBack(0);
-				TestEqual(TEXT("ShiftIndexToBack Cap With Wrapped"), TRingBuffer<T>({ 0,1,2,3,4,5,6,7,8 }), Q);
+				TEST_EQUAL(TEXT("ShiftIndexToBack Cap With Wrapped"), TRingBuffer<T>({ 0,1,2,3,4,5,6,7,8 }), Q);
 			}
 		}
 
@@ -1057,28 +1057,28 @@ public:
 				int32 CheckIndex = 0;
 				for (; CheckIndex < It; ++CheckIndex)
 				{
-					TestEqual(*FString::Printf(TEXT("ShiftIndexToBack Entire Array Values (%d,%d)"), It, CheckIndex), CheckIndex, Q[CheckIndex]);
+					TEST_EQUAL(*FString::Printf(TEXT("ShiftIndexToBack Entire Array Values (%d,%d)"), It, CheckIndex), CheckIndex, Q[CheckIndex]);
 				}
 				for (; CheckIndex < Count - 1; ++CheckIndex)
 				{
-					TestEqual(*FString::Printf(TEXT("ShiftIndexToBack Entire Array Values (%d,%d)"), It, CheckIndex), CheckIndex + 1, Q[CheckIndex]);
+					TEST_EQUAL(*FString::Printf(TEXT("ShiftIndexToBack Entire Array Values (%d,%d)"), It, CheckIndex), CheckIndex + 1, Q[CheckIndex]);
 				}
-				TestEqual(*FString::Printf(TEXT("ShiftIndexToBack Entire Array Values (%d,%d)"), It, Count - 1), It, Q[Count - 1]);
+				TEST_EQUAL(*FString::Printf(TEXT("ShiftIndexToBack Entire Array Values (%d,%d)"), It, Count - 1), It, Q[Count - 1]);
 			}
 			for (int32 It = 0; It < Count; ++It)
 			{
 				TRingBuffer<T> Q({ 0,1,2,3,4,5,6,7 });
 				Q.ShiftIndexToFront(It);
 
-				TestEqual(*FString::Printf(TEXT("ShiftIndexToFront Entire Array Values (%d,%d)"), It, 0), It, Q[0]);
+				TEST_EQUAL(*FString::Printf(TEXT("ShiftIndexToFront Entire Array Values (%d,%d)"), It, 0), It, Q[0]);
 				int32 CheckIndex = 1;
 				for (; CheckIndex <= It; ++CheckIndex)
 				{
-					TestEqual(*FString::Printf(TEXT("ShiftIndexToFront Entire Array Values (%d,%d)"), It, CheckIndex), CheckIndex - 1, Q[CheckIndex]);
+					TEST_EQUAL(*FString::Printf(TEXT("ShiftIndexToFront Entire Array Values (%d,%d)"), It, CheckIndex), CheckIndex - 1, Q[CheckIndex]);
 				}
 				for (; CheckIndex < Count; ++CheckIndex)
 				{
-					TestEqual(*FString::Printf(TEXT("ShiftIndexToFront Entire Array Values (%d,%d)"), It, CheckIndex), CheckIndex, Q[CheckIndex]);
+					TEST_EQUAL(*FString::Printf(TEXT("ShiftIndexToFront Entire Array Values (%d,%d)"), It, CheckIndex), CheckIndex, Q[CheckIndex]);
 				}
 			}
 		}
@@ -1108,7 +1108,7 @@ int FRingBufferTest::Counter::NumCopy = 0;
 int FRingBufferTest::Counter::NumMove = 0;
 int FRingBufferTest::Counter::NumDestruct = 0;
 
-TEST_CASE("Core::Containers::TRingBuffer::Smoke Test", "[Core][Containers][Smoke]")
+TEST_CASE_METHOD(FAutomationTestFixture, "Core::Containers::TRingBuffer::Smoke Test", "[Core][Containers][Smoke]")
 {
 	FRingBufferTest RingBufferTest;
 	RingBufferTest.RunTest();
diff --git a/Engine/Source/Runtime/Core/Tests/Containers/StridedViewTest.cpp b/Engine/Source/Runtime/Core/Tests/Containers/StridedViewTest.cpp
--- a/Engine/Source/Runtime/Core/Tests/Containers/StridedViewTest.cpp
+++ b/Engine/Source/Runtime/Core/Tests/Containers/StridedViewTest.cpp
@@ -72,7 +72,7 @@ static_assert(!std::is_convertible_v<TStridedView<const FMyStruct*>, TStridedVie
 } // namespace StridedViewTest
 } // namespace UE
 
-TEST_CASE("Core::Containers::TStridedView::Smoke Test", "[Core][Containers][Smoke]")
+TEST_CASE_METHOD(FAutomationTestFixture, "Core::Containers::TStridedView::Smoke Test", "[Core][Containers][Smoke]")
 {
 	using namespace UE::StridedViewTest;
 
@@ -91,7 +91,7 @@ TEST_CASE("Core::Containers::TStridedView::Smoke Test", "[Core][Containers][Smok
 
 		for (int32 i = 0; i < Num; ++i)
 		{
-			TestEqual(TEXT("Structs[i].Position"), Structs[i].Position, FVector((float)i));
+			TEST_EQUAL(TEXT("Structs[i].Position"), Structs[i].Position, FVector((float)i));
 		}
 
 		const FVector Multiplier = FVector(2.0f, 3.0f, 4.0f);
@@ -102,7 +102,7 @@ TEST_CASE("Core::Containers::TStridedView::Smoke Test", "[Core][Containers][Smok
 
 		for (int32 i = 0; i < Num; ++i)
 		{
-			TestEqual(TEXT("Structs[i].Position"), Structs[i].Position, FVector((float)i) * Multiplier);
+			TEST_EQUAL(TEXT("Structs[i].Position"), Structs[i].Position, FVector((float)i) * Multiplier);
 		}
 	}
 
@@ -123,43 +123,43 @@ TEST_CASE("Core::Containers::TStridedView::Smoke Test", "[Core][Containers][Smok
 
 	{
 		FVector Avg = ComputeMeanPosition(MakeArrayView(Structs));
-		TestEqual(TEXT("ComputeMeanPosition(MakeArrayView(Structs))"), Avg, ExpectedAvg);
+		TEST_EQUAL(TEXT("ComputeMeanPosition(MakeArrayView(Structs))"), Avg, ExpectedAvg);
 	}
 
 	{
 		FVector Avg = ComputeMeanPositionStrided(MakeStridedView(Structs));
-		TestEqual(TEXT("ComputeMeanPositionStrided(MakeStridedView(Structs))"), Avg, ExpectedAvg);
+		TEST_EQUAL(TEXT("ComputeMeanPositionStrided(MakeStridedView(Structs))"), Avg, ExpectedAvg);
 	}
 
 	{
 		FVector Avg = ComputeMean(TStridedView<FVector>(sizeof(FMyStruct), &Structs[0].Position, Num));
-		TestEqual(TEXT("ComputeMean(TStridedView<FVector>(sizeof(FMyStruct), &Structs[0].Position, Num))"), Avg, ExpectedAvg);
+		TEST_EQUAL(TEXT("ComputeMean(TStridedView<FVector>(sizeof(FMyStruct), &Structs[0].Position, Num))"), Avg, ExpectedAvg);
 	}
 
 	{
 		FVector Avg = ComputeMean(TStridedView<const FVector>(sizeof(FMyStruct), &Structs[0].Position, Num));
-		TestEqual(TEXT("ComputeMean(TStridedView<const FVector>(sizeof(FMyStruct), &Structs[0].Position, Num))"), Avg, ExpectedAvg);
+		TEST_EQUAL(TEXT("ComputeMean(TStridedView<const FVector>(sizeof(FMyStruct), &Structs[0].Position, Num))"), Avg, ExpectedAvg);
 	}
 
 	{
 		FVector Avg = ComputeMean(MakeStridedView(sizeof(FMyStruct), &Structs[0].Position, Num));
-		TestEqual(TEXT("ComputeMean(MakeStridedView(sizeof(FMyStruct), &Structs[0].Position, Num))"), Avg, ExpectedAvg);
+		TEST_EQUAL(TEXT("ComputeMean(MakeStridedView(sizeof(FMyStruct), &Structs[0].Position, Num))"), Avg, ExpectedAvg);
 	}
 
 	{
 		FVector Avg = ComputeMean(MakeStridedView(MakeArrayView(Structs), &FMyStruct::Position));
-		TestEqual(TEXT("ComputeMean(MakeStridedView(MakeArrayView(Structs), &FMyStruct::Position))"), Avg, ExpectedAvg);
+		TEST_EQUAL(TEXT("ComputeMean(MakeStridedView(MakeArrayView(Structs), &FMyStruct::Position))"), Avg, ExpectedAvg);
 	}
 
 	{
 		FVector Avg = ComputeMean(MakeStridedView(Structs, &FMyStruct::Position));
-		TestEqual(TEXT("ComputeMean(MakeStridedView(Structs, &FMyStruct::Position))"), Avg, ExpectedAvg);
+		TEST_EQUAL(TEXT("ComputeMean(MakeStridedView(Structs, &FMyStruct::Position))"), Avg, ExpectedAvg);
 	}
 
 	{
 		TArray<FMyStruct> StructsArray = TArray<FMyStruct>(Structs, Num);
 		FVector Avg = ComputeMean(MakeStridedView(StructsArray, &FMyStruct::Position));
-		TestEqual(TEXT("ComputeMean(MakeStridedView(StructsArray, &FMyStruct::Position))"), Avg, ExpectedAvg);
+		TEST_EQUAL(TEXT("ComputeMean(MakeStridedView(StructsArray, &FMyStruct::Position))"), Avg, ExpectedAvg);
 	}
 
 	// Test support for trivial strided views (stride = sizeof(element))
@@ -172,7 +172,7 @@ TEST_CASE("Core::Containers::TStridedView::Smoke Test", "[Core][Containers][Smok
 
 	{
 		FVector Avg = ComputeMean(MakeStridedView(Vecs));
-		TestEqual(TEXT("ComputeMean(MakeStridedView(Vecs))"), Avg, ExpectedAvg);
+		TEST_EQUAL(TEXT("ComputeMean(MakeStridedView(Vecs))"), Avg, ExpectedAvg);
 	}
 
 	{
@@ -180,7 +180,7 @@ TEST_CASE("Core::Containers::TStridedView::Smoke Test", "[Core][Containers][Smok
 		auto VecsStridedView = MakeStridedView(VecsView);
 		static_assert(std::is_const_v<decltype(VecsStridedView)::ElementType>);
 		FVector Avg = ComputeMean(VecsStridedView);
-		TestEqual(TEXT("ComputeMean(MakeStridedView(VecsView))"), Avg, ExpectedAvg);
+		TEST_EQUAL(TEXT("ComputeMean(MakeStridedView(VecsView))"), Avg, ExpectedAvg);
 	}
 
 	{
@@ -188,7 +188,7 @@ TEST_CASE("Core::Containers::TStridedView::Smoke Test", "[Core][Containers][Smok
 		auto VecsStridedView = TestCreateConstArrayView(VecsArray);
 		static_assert(std::is_const_v<decltype(VecsStridedView)::ElementType>);
 		FVector Avg = ComputeMean(VecsStridedView);
-		TestEqual(TEXT("ComputeMean(MakeStridedView(VecsArray))"), Avg, ExpectedAvg);
+		TEST_EQUAL(TEXT("ComputeMean(MakeStridedView(VecsArray))"), Avg, ExpectedAvg);
 	}
 
 	// Test smearing an element
@@ -196,11 +196,11 @@ TEST_CASE("Core::Containers::TStridedView::Smoke Test", "[Core][Containers][Smok
 	{
 		FVector V = FVector(1.0f, 2.0f, 3.0f);
 		FVector Avg = ComputeMean(TStridedView<FVector>(0, &V, Num));
-		TestEqual(TEXT("ComputeMean(TStridedView<FVector>(0, &V, Num))"), Avg, V);
+		TEST_EQUAL(TEXT("ComputeMean(TStridedView<FVector>(0, &V, Num))"), Avg, V);
 
 		FVector Sum = ComputeSum(TStridedView<FVector>(0, &V, Num));
 		FVector ExpectedSum = V * (float)Num;
-		TestEqual(TEXT("ComputeSum(TStridedView<FVector>(0, &V, Num))"), Sum, ExpectedSum);
+		TEST_EQUAL(TEXT("ComputeSum(TStridedView<FVector>(0, &V, Num))"), Sum, ExpectedSum);
 	}
 
 	// Test iterating through derived struct
@@ -222,8 +222,8 @@ TEST_CASE("Core::Containers::TStridedView::Smoke Test", "[Core][Containers][Smok
 
 		for (int32 i = 0; i < PositionsViewDerived.Num(); ++i)
 		{
-			TestEqual(TEXT("StructsDerived[i].Position"), StructsDerived[i].Position, Structs[i].Position);
-			TestEqual(TEXT("StructsDerived[i].Normal"), StructsDerived[i].Normal, FVector(1.0f, 2.0f, 3.0f) * (float)i);
+			TEST_EQUAL(TEXT("StructsDerived[i].Position"), StructsDerived[i].Position, Structs[i].Position);
+			TEST_EQUAL(TEXT("StructsDerived[i].Normal"), StructsDerived[i].Normal, FVector(1.0f, 2.0f, 3.0f) * (float)i);
 		}
 	}
 
@@ -237,17 +237,17 @@ TEST_CASE("Core::Containers::TStridedView::Smoke Test", "[Core][Containers][Smok
 
 	{
 		FVector Avg = ComputeMean(TStridedView<FVector>(sizeof(FMyStructDerived), &StructsDerived[0].Position, Num));
-		TestEqual(TEXT("ComputeMean(TStridedView<FVector>(sizeof(FMyStructDerived), &StructsDerived[0].Position, Num))"), Avg, ExpectedAvg);
+		TEST_EQUAL(TEXT("ComputeMean(TStridedView<FVector>(sizeof(FMyStructDerived), &StructsDerived[0].Position, Num))"), Avg, ExpectedAvg);
 	}
 
 	{
 		FVector Avg = ComputeMean(TStridedView<const FVector>(sizeof(FMyStructDerived), &StructsDerived[0].Position, Num));
-		TestEqual(TEXT("ComputeMean(TStridedView<const FVector>(sizeof(FMyStructDerived), &StructsDerived[0].Position, Num))"), Avg, ExpectedAvg);
+		TEST_EQUAL(TEXT("ComputeMean(TStridedView<const FVector>(sizeof(FMyStructDerived), &StructsDerived[0].Position, Num))"), Avg, ExpectedAvg);
 	}
 
 	{
 		FVector Avg = ComputeMean(MakeStridedView(sizeof(FMyStructDerived), &StructsDerived[0].Position, Num));
-		TestEqual(TEXT("ComputeMean(MakeStridedView(sizeof(FMyStructDerived), &StructsDerived[0].Position, Num))"), Avg, ExpectedAvg);
+		TEST_EQUAL(TEXT("ComputeMean(MakeStridedView(sizeof(FMyStructDerived), &StructsDerived[0].Position, Num))"), Avg, ExpectedAvg);
 	}
 
 	{
@@ -263,45 +263,45 @@ TEST_CASE("Core::Containers::TStridedView::Smoke Test", "[Core][Containers][Smok
 		TStridedView<FMyStruct> ViewDerivedAsBase(sizeof(StructsDerived[0]), &static_cast<FMyStruct&>(StructsDerived[0]), Num);
 		TStridedView<FMyStruct> ViewDerivedAsBase2 = MakeStridedViewOfBase<FMyStruct>(MakeArrayView(StructsDerived));
 
-		TestEqual(TEXT("ViewDerived.GetStride()"), ViewDerived.GetStride(), (int32)sizeof(FMyStructDerived));
-		TestEqual(TEXT("ViewDerivedAsBase.GetStride()"), ViewDerivedAsBase.GetStride(), (int32)sizeof(FMyStructDerived));
-		TestEqual(TEXT("ViewDerivedAsBase2.GetStride()"), ViewDerivedAsBase2.GetStride(), (int32)sizeof(FMyStructDerived));
+		TEST_EQUAL(TEXT("ViewDerived.GetStride()"), ViewDerived.GetStride(), (int32)sizeof(FMyStructDerived));
+		TEST_EQUAL(TEXT("ViewDerivedAsBase.GetStride()"), ViewDerivedAsBase.GetStride(), (int32)sizeof(FMyStructDerived));
+		TEST_EQUAL(TEXT("ViewDerivedAsBase2.GetStride()"), ViewDerivedAsBase2.GetStride(), (int32)sizeof(FMyStructDerived));
 
 		// Incorrect API usage : implicit conversion of derived view to base. Generates compile error (as expected).
 		// FVector Avg = ComputeMeanPositionStrided(ViewDerived);
 
 		FVector AvgDerivedAsBase = ComputeMeanPositionStrided(ViewDerivedAsBase);
-		TestEqual(TEXT("ComputeMeanPositionStrided(ViewDerivedAsBase)"), AvgDerivedAsBase, ExpectedAvg);
+		TEST_EQUAL(TEXT("ComputeMeanPositionStrided(ViewDerivedAsBase)"), AvgDerivedAsBase, ExpectedAvg);
 
 		FVector AvgDerivedAsBase2 = ComputeMeanPositionStrided(ViewDerivedAsBase2);
-		TestEqual(TEXT("ComputeMeanPositionStrided(ViewDerivedAsBase2)"), AvgDerivedAsBase2, ExpectedAvg);
+		TEST_EQUAL(TEXT("ComputeMeanPositionStrided(ViewDerivedAsBase2)"), AvgDerivedAsBase2, ExpectedAvg);
 	}
 
 	{
 		TArrayView<FMyStructDerived> DerivedArrayView = MakeArrayView(StructsDerived);
 		TStridedView<FVector> View = MakeStridedView(DerivedArrayView, &FMyStructDerived::Position);
-		TestEqual(TEXT("MakeStridedView(MakeArrayView(StructsDerived), &FMyStructDerived::Position).GetStride()"), View.GetStride(), (int32)sizeof(StructsDerived[0]));
+		TEST_EQUAL(TEXT("MakeStridedView(MakeArrayView(StructsDerived), &FMyStructDerived::Position).GetStride()"), View.GetStride(), (int32)sizeof(StructsDerived[0]));
 		FVector Avg = ComputeMean(View);
-		TestEqual(TEXT("ComputeMean(MakeStridedView(MakeArrayView(StructsDerived), &FMyStructDerived::Position))"), Avg, ExpectedAvg);
+		TEST_EQUAL(TEXT("ComputeMean(MakeStridedView(MakeArrayView(StructsDerived), &FMyStructDerived::Position))"), Avg, ExpectedAvg);
 	}
 
 	{
 		TArrayView<FMyStructDerived> DerivedArrayView = MakeArrayView(StructsDerived);
 		TStridedView<FVector> View = MakeStridedView(DerivedArrayView, &FMyStruct::Position);
-		TestEqual(TEXT("MakeStridedView(MakeArrayView(StructsDerived), &FMyStruct::Position).GetStride()"), View.GetStride(), (int32)sizeof(StructsDerived[0]));
+		TEST_EQUAL(TEXT("MakeStridedView(MakeArrayView(StructsDerived), &FMyStruct::Position).GetStride()"), View.GetStride(), (int32)sizeof(StructsDerived[0]));
 		FVector Avg = ComputeMean(View);
-		TestEqual(TEXT("ComputeMean(MakeStridedView(MakeArrayView(StructsDerived), &FMyStruct::Position))"), Avg, ExpectedAvg);
+		TEST_EQUAL(TEXT("ComputeMean(MakeStridedView(MakeArrayView(StructsDerived), &FMyStruct::Position))"), Avg, ExpectedAvg);
 	}
 
 	{
 		FVector Avg = ComputeMean(MakeStridedView(StructsDerived, &FMyStructDerived::Position));
-		TestEqual(TEXT("ComputeMean(MakeStridedView(StructsDerived, &FMyStructDerived::Position))"), Avg, ExpectedAvg);
+		TEST_EQUAL(TEXT("ComputeMean(MakeStridedView(StructsDerived, &FMyStructDerived::Position))"), Avg, ExpectedAvg);
 	}
 
 	{
 		TArray<FMyStructDerived> StructsArrayDerived = TArray<FMyStructDerived>(StructsDerived, Num);
 		FVector Avg = ComputeMean(MakeStridedView(StructsArrayDerived, &FMyStructDerived::Position));
-		TestEqual(TEXT("ComputeMean(MakeStridedView(StructsArrayDerived, &FMyStructDerived::Position))"), Avg, ExpectedAvg);
+		TEST_EQUAL(TEXT("ComputeMean(MakeStridedView(StructsArrayDerived, &FMyStructDerived::Position))"), Avg, ExpectedAvg);
 	}
 
 	// Test iterating through nested fields
@@ -314,7 +314,7 @@ TEST_CASE("Core::Containers::TStridedView::Smoke Test", "[Core][Containers][Smok
 			ExpectedSum += Structs[i].Position.Z;
 		}
 		FVector::FReal Sum = Algo::Accumulate(ViewOfZs, FVector::FReal(0));
-		TestEqual(TEXT("Algo::Accumulate(MakeStridedView((int32)sizeof(Structs[0]), &Structs[0].Position.Z, Num), 0.0f)"), Sum, ExpectedSum);
+		TEST_EQUAL(TEXT("Algo::Accumulate(MakeStridedView((int32)sizeof(Structs[0]), &Structs[0].Position.Z, Num), 0.0f)"), Sum, ExpectedSum);
 	}
 }
 
