Add new Unit Tests

diff --git a/Engine/Source/Runtime/Core/Tests/Containers/StringViewTest.cpp b/Engine/Source/Runtime/Core/Tests/Containers/StringViewTest.cpp
--- a/Engine/Source/Runtime/Core/Tests/Containers/StringViewTest.cpp
+++ b/Engine/Source/Runtime/Core/Tests/Containers/StringViewTest.cpp
@@ -34,43 +34,43 @@ template <> struct TIsContiguousContainer<UE::String::Private::TestArgumentDepen
 
 constexpr const uint32 TestFlags = EAutomationTestFlags::ApplicationContextMask | EAutomationTestFlags::SmokeFilter;
 
-TEST_CASE("Core::Containers::FStringView::Constructor", "[Core][Containers][Smoke]")
+TEST_CASE_METHOD(FAutomationTestFixture, "Core::Containers::FStringView::Constructor", "[Core][Containers][Smoke]")
 {
 	// Default View
 	{
 		FStringView View;
-		TestEqual(TEXT(""), View.Len(), 0);
-		TestTrue(TEXT("View.IsEmpty"), View.IsEmpty());
+		TEST_EQUAL(TEXT(""), View.Len(), 0);
+		TEST_TRUE(TEXT("View.IsEmpty"), View.IsEmpty());
 	}
 
 	// Empty View
 	{
 		FStringView View(TEXT(""));
-		TestEqual(TEXT(""), View.Len(), 0);
-		TestTrue(TEXT("View.IsEmpty"), View.IsEmpty());
+		TEST_EQUAL(TEXT(""), View.Len(), 0);
+		TEST_TRUE(TEXT("View.IsEmpty"), View.IsEmpty());
 	}
 
 	// Constructing from nullptr is supported; nullptr interpreted as empty string
 	{
 		FStringView View(nullptr);
-		TestEqual(TEXT(""), View.Len(), 0);
-		TestTrue(TEXT("View.IsEmpty"), View.IsEmpty());
+		TEST_EQUAL(TEXT(""), View.Len(), 0);
+		TEST_TRUE(TEXT("View.IsEmpty"), View.IsEmpty());
 	}
 
 	// Create from a wchar literal
 	{
 		FStringView View(TEXT("Test Ctor"));
-		TestEqual(TEXT("View length"), View.Len(), FCStringWide::Strlen(TEXT("Test Ctor")));
-		TestEqual(TEXT("The result of Strncmp"), FCStringWide::Strncmp(View.GetData(), TEXT("Test Ctor"), View.Len()), 0);
-		TestFalse(TEXT("View.IsEmpty"), View.IsEmpty());
+		TEST_EQUAL(TEXT("View length"), View.Len(), FCStringWide::Strlen(TEXT("Test Ctor")));
+		TEST_EQUAL(TEXT("The result of Strncmp"), FCStringWide::Strncmp(View.GetData(), TEXT("Test Ctor"), View.Len()), 0);
+		TEST_FALSE(TEXT("View.IsEmpty"), View.IsEmpty());
 	}
 
 	// Create from a sub section of a wchar literal
 	{
 		FStringView View(TEXT("Test SubSection Ctor"), 4);
-		TestEqual(TEXT("View length"), View.Len(), 4);
-		TestEqual(TEXT("The result of Strncmp"), FCStringWide::Strncmp(View.GetData(), TEXT("Test"), View.Len()), 0);
-		TestFalse(TEXT("View.IsEmpty"), View.IsEmpty());
+		TEST_EQUAL(TEXT("View length"), View.Len(), 4);
+		TEST_EQUAL(TEXT("The result of Strncmp"), FCStringWide::Strncmp(View.GetData(), TEXT("Test"), View.Len()), 0);
+		TEST_FALSE(TEXT("View.IsEmpty"), View.IsEmpty());
 	}
 
 	// Create from a FString
@@ -78,25 +78,25 @@ TEST_CASE("Core::Containers::FStringView::Constructor", "[Core][Containers][Smok
 		FString String(TEXT("String Object"));
 		FStringView View(String);
 
-		TestEqual(TEXT("View length"), View.Len(), String.Len());
-		TestEqual(TEXT("The result of Strncmp"), FCStringWide::Strncmp(View.GetData(), *String, View.Len()), 0);
-		TestFalse(TEXT("View.IsEmpty"), View.IsEmpty());
+		TEST_EQUAL(TEXT("View length"), View.Len(), String.Len());
+		TEST_EQUAL(TEXT("The result of Strncmp"), FCStringWide::Strncmp(View.GetData(), *String, View.Len()), 0);
+		TEST_FALSE(TEXT("View.IsEmpty"), View.IsEmpty());
 	}
 
 	// Create from an ansi literal
 	{
 		FAnsiStringView View("Test Ctor");
-		TestEqual(TEXT("View length"), View.Len(), FCStringAnsi::Strlen("Test Ctor"));
-		TestEqual(TEXT("The result of Strncmp"), FCStringAnsi::Strncmp(View.GetData(), "Test Ctor", View.Len()), 0);
-		TestFalse(TEXT("View.IsEmpty"), View.IsEmpty());
+		TEST_EQUAL(TEXT("View length"), View.Len(), FCStringAnsi::Strlen("Test Ctor"));
+		TEST_EQUAL(TEXT("The result of Strncmp"), FCStringAnsi::Strncmp(View.GetData(), "Test Ctor", View.Len()), 0);
+		TEST_FALSE(TEXT("View.IsEmpty"), View.IsEmpty());
 	}
 
 	// Create from a sub section of an ansi literal
 	{
 		FAnsiStringView View("Test SubSection Ctor", 4);
-		TestEqual(TEXT("View length"), View.Len(), 4);
-		TestEqual(TEXT("The result of Strncmp"), FCStringAnsi::Strncmp(View.GetData(), "Test", View.Len()), 0);
-		TestFalse(TEXT("View.IsEmpty"), View.IsEmpty());
+		TEST_EQUAL(TEXT("View length"), View.Len(), 4);
+		TEST_EQUAL(TEXT("The result of Strncmp"), FCStringAnsi::Strncmp(View.GetData(), "Test", View.Len()), 0);
+		TEST_FALSE(TEXT("View.IsEmpty"), View.IsEmpty());
 	}
 
 	// Create using string view literals
@@ -131,11 +131,11 @@ TEST_CASE("Core::Containers::FStringView::Constructor", "[Core][Containers][Smok
 	{
 		UE::String::Private::TestArgumentDependentLookup::FTestType Test;
 		FStringView View(Test);
-		TestTrue(TEXT("StringView ADL"), View.Equals(TEXTVIEW("ABC")));
+		TEST_TRUE(TEXT("StringView ADL"), View.Equals(TEXTVIEW("ABC")));
 	}
 }
 
-TEST_CASE("Core::Containers::FStringView::Iterator", "[Core][Containers][Smoke]")
+TEST_CASE_METHOD(FAutomationTestFixture, "Core::Containers::FStringView::Iterator", "[Core][Containers][Smoke]")
 {
 	// Iterate over a string view
 	{
@@ -144,11 +144,11 @@ TEST_CASE("Core::Containers::FStringView::Iterator", "[Core][Containers][Smoke]"
 
 		for (TCHAR C : View)
 		{
-			TestTrue(TEXT("Iterators(0)-Iteration"), C == *StringLiteralSrc++);
+			TEST_TRUE(TEXT("Iterators(0)-Iteration"), C == *StringLiteralSrc++);
 		}
 
 		// Make sure we iterated over the entire string
-		TestTrue(TEXT("Iterators(0-EndCheck"), *StringLiteralSrc == '\0');
+		TEST_TRUE(TEXT("Iterators(0-EndCheck"), *StringLiteralSrc == '\0');
 	}
 
 	// Iterate over a partial string view
@@ -158,15 +158,15 @@ TEST_CASE("Core::Containers::FStringView::Iterator", "[Core][Containers][Smoke]"
 
 		for (TCHAR C : View)
 		{
-			TestTrue(TEXT("Iterators(1)-Iteration"), C == *StringLiteralSrc++);
+			TEST_TRUE(TEXT("Iterators(1)-Iteration"), C == *StringLiteralSrc++);
 		}
 
 		// Make sure we only iterated over the part of the string that the view represents
-		TestTrue(TEXT("Iterators(1)-EndCheck"), *StringLiteralSrc == '|');
+		TEST_TRUE(TEXT("Iterators(1)-EndCheck"), *StringLiteralSrc == '|');
 	}
 }
 
-TEST_CASE("Core::Containers::FStringView::Equality Test", "[Core][Containers][Smoke]")
+TEST_CASE_METHOD(FAutomationTestFixture, "Core::Containers::FStringView::Equality Test", "[Core][Containers][Smoke]")
 {
 	const ANSICHAR* AnsiStringLiteralSrc = "String To Test!";
 	const ANSICHAR* AnsiStringLiteralLower = "string to test!";
@@ -179,27 +179,27 @@ TEST_CASE("Core::Containers::FStringView::Equality Test", "[Core][Containers][Sm
 
 	FStringView WideView(WideStringLiteralSrc);
 
-	TestTrue(TEXT("Equality(0)"), WideView == WideStringLiteralSrc);
-	TestFalse(TEXT("Equality(1)"), WideView != WideStringLiteralSrc);
-	TestTrue(TEXT("Equality(2)"), WideView == WideStringLiteralLower);
-	TestFalse(TEXT("Equality(3)"), WideView != WideStringLiteralLower);
-	TestTrue(TEXT("Equality(4)"), WideView == WideStringLiteralUpper);
-	TestFalse(TEXT("Equality(5)"), WideView != WideStringLiteralUpper);
-	TestFalse(TEXT("Equality(6)"), WideView == WideStringLiteralShort);
-	TestTrue(TEXT("Equality(7)"), WideView != WideStringLiteralShort);
-	TestFalse(TEXT("Equality(8)"), WideView == WideStringLiteralLonger);
-	TestTrue(TEXT("Equality(9)"), WideView != WideStringLiteralLonger);
-
-	TestTrue(TEXT("Equality(10)"), WideStringLiteralSrc == WideView);
-	TestFalse(TEXT("Equality(11)"), WideStringLiteralSrc != WideView);
-	TestTrue(TEXT("Equality(12)"), WideStringLiteralLower == WideView);
-	TestFalse(TEXT("Equality(13)"), WideStringLiteralLower != WideView);
-	TestTrue(TEXT("Equality(14)"), WideStringLiteralUpper == WideView);
-	TestFalse(TEXT("Equality(15)"), WideStringLiteralUpper != WideView);
-	TestFalse(TEXT("Equality(16)"), WideStringLiteralShort == WideView);
-	TestTrue(TEXT("Equality(17)"), WideStringLiteralShort != WideView);
-	TestFalse(TEXT("Equality(18)"), WideStringLiteralLonger == WideView);
-	TestTrue(TEXT("Equality(19)"), WideStringLiteralLonger != WideView);
+	TEST_TRUE(TEXT("Equality(0)"), WideView == WideStringLiteralSrc);
+	TEST_FALSE(TEXT("Equality(1)"), WideView != WideStringLiteralSrc);
+	TEST_TRUE(TEXT("Equality(2)"), WideView == WideStringLiteralLower);
+	TEST_FALSE(TEXT("Equality(3)"), WideView != WideStringLiteralLower);
+	TEST_TRUE(TEXT("Equality(4)"), WideView == WideStringLiteralUpper);
+	TEST_FALSE(TEXT("Equality(5)"), WideView != WideStringLiteralUpper);
+	TEST_FALSE(TEXT("Equality(6)"), WideView == WideStringLiteralShort);
+	TEST_TRUE(TEXT("Equality(7)"), WideView != WideStringLiteralShort);
+	TEST_FALSE(TEXT("Equality(8)"), WideView == WideStringLiteralLonger);
+	TEST_TRUE(TEXT("Equality(9)"), WideView != WideStringLiteralLonger);
+
+	TEST_TRUE(TEXT("Equality(10)"), WideStringLiteralSrc == WideView);
+	TEST_FALSE(TEXT("Equality(11)"), WideStringLiteralSrc != WideView);
+	TEST_TRUE(TEXT("Equality(12)"), WideStringLiteralLower == WideView);
+	TEST_FALSE(TEXT("Equality(13)"), WideStringLiteralLower != WideView);
+	TEST_TRUE(TEXT("Equality(14)"), WideStringLiteralUpper == WideView);
+	TEST_FALSE(TEXT("Equality(15)"), WideStringLiteralUpper != WideView);
+	TEST_FALSE(TEXT("Equality(16)"), WideStringLiteralShort == WideView);
+	TEST_TRUE(TEXT("Equality(17)"), WideStringLiteralShort != WideView);
+	TEST_FALSE(TEXT("Equality(18)"), WideStringLiteralLonger == WideView);
+	TEST_TRUE(TEXT("Equality(19)"), WideStringLiteralLonger != WideView);
 
 	FString WideStringSrc = WideStringLiteralSrc;
 	FString WideStringLower = WideStringLiteralLower;
@@ -207,76 +207,76 @@ TEST_CASE("Core::Containers::FStringView::Equality Test", "[Core][Containers][Sm
 	FString WideStringShort = WideStringLiteralShort;
 	FString WideStringLonger = WideStringLiteralLonger;
 
-	TestTrue(TEXT("Equality(20)"), WideView == WideStringSrc);
-	TestFalse(TEXT("Equality(21)"), WideView != WideStringSrc);
-	TestTrue(TEXT("Equality(22)"), WideView == WideStringLower);
-	TestFalse(TEXT("Equality(23)"), WideView != WideStringLower);
-	TestTrue(TEXT("Equality(24)"), WideView == WideStringUpper);
-	TestFalse(TEXT("Equality(25)"), WideView != WideStringUpper);
-	TestFalse(TEXT("Equality(26)"), WideView == WideStringShort);
-	TestTrue(TEXT("Equality(27)"), WideView != WideStringShort);
-	TestFalse(TEXT("Equality(28)"), WideView == WideStringLonger);
-	TestTrue(TEXT("Equality(29)"), WideView != WideStringLonger);
-
-	TestTrue(TEXT("Equality(30)"), WideStringSrc == WideView);
-	TestFalse(TEXT("Equality(31)"), WideStringSrc != WideView);
-	TestTrue(TEXT("Equality(32)"), WideStringLower == WideView);
-	TestFalse(TEXT("Equality(33)"), WideStringLower != WideView);
-	TestTrue(TEXT("Equality(34)"), WideStringUpper == WideView);
-	TestFalse(TEXT("Equality(35)"), WideStringUpper != WideView);
-	TestFalse(TEXT("Equality(36)"), WideStringShort == WideView);
-	TestTrue(TEXT("Equality(37)"), WideStringShort != WideView);
-	TestFalse(TEXT("Equality(38)"), WideStringLonger == WideView);
-	TestTrue(TEXT("Equality(39)"), WideStringLonger != WideView);
+	TEST_TRUE(TEXT("Equality(20)"), WideView == WideStringSrc);
+	TEST_FALSE(TEXT("Equality(21)"), WideView != WideStringSrc);
+	TEST_TRUE(TEXT("Equality(22)"), WideView == WideStringLower);
+	TEST_FALSE(TEXT("Equality(23)"), WideView != WideStringLower);
+	TEST_TRUE(TEXT("Equality(24)"), WideView == WideStringUpper);
+	TEST_FALSE(TEXT("Equality(25)"), WideView != WideStringUpper);
+	TEST_FALSE(TEXT("Equality(26)"), WideView == WideStringShort);
+	TEST_TRUE(TEXT("Equality(27)"), WideView != WideStringShort);
+	TEST_FALSE(TEXT("Equality(28)"), WideView == WideStringLonger);
+	TEST_TRUE(TEXT("Equality(29)"), WideView != WideStringLonger);
+
+	TEST_TRUE(TEXT("Equality(30)"), WideStringSrc == WideView);
+	TEST_FALSE(TEXT("Equality(31)"), WideStringSrc != WideView);
+	TEST_TRUE(TEXT("Equality(32)"), WideStringLower == WideView);
+	TEST_FALSE(TEXT("Equality(33)"), WideStringLower != WideView);
+	TEST_TRUE(TEXT("Equality(34)"), WideStringUpper == WideView);
+	TEST_FALSE(TEXT("Equality(35)"), WideStringUpper != WideView);
+	TEST_FALSE(TEXT("Equality(36)"), WideStringShort == WideView);
+	TEST_TRUE(TEXT("Equality(37)"), WideStringShort != WideView);
+	TEST_FALSE(TEXT("Equality(38)"), WideStringLonger == WideView);
+	TEST_TRUE(TEXT("Equality(39)"), WideStringLonger != WideView);
 
 	FStringView IdenticalView(WideStringLiteralSrc);
 
-	TestTrue(TEXT("Equality(40a)"), WideView == IdenticalView);
-	TestFalse(TEXT("Equality(40b)"), WideView != IdenticalView);
-	TestTrue(TEXT("Equality(41a)"), IdenticalView == WideView);
-	TestFalse(TEXT("Equality(41b)"), IdenticalView != WideView);
+	TEST_TRUE(TEXT("Equality(40a)"), WideView == IdenticalView);
+	TEST_FALSE(TEXT("Equality(40b)"), WideView != IdenticalView);
+	TEST_TRUE(TEXT("Equality(41a)"), IdenticalView == WideView);
+	TEST_FALSE(TEXT("Equality(41b)"), IdenticalView != WideView);
 
 	// Views without null termination
 
 	FStringView ShortViewNoNull = WideView.Left(FStringView(WideStringLiteralShort).Len());
 
-	TestTrue(TEXT("Equality(42)"), ShortViewNoNull == WideStringLiteralShort);
-	TestFalse(TEXT("Equality(43)"), ShortViewNoNull != WideStringLiteralShort);
-	TestTrue(TEXT("Equality(44)"), WideStringLiteralShort == ShortViewNoNull);
-	TestFalse(TEXT("Equality(45)"), WideStringLiteralShort != ShortViewNoNull);
-	TestFalse(TEXT("Equality(46)"), ShortViewNoNull == WideStringLiteralSrc);
-	TestTrue(TEXT("Equality(47)"), ShortViewNoNull != WideStringLiteralSrc);
-	TestFalse(TEXT("Equality(48)"), WideStringLiteralSrc == ShortViewNoNull);
-	TestTrue(TEXT("Equality(49)"), WideStringLiteralSrc != ShortViewNoNull);
-
-	TestTrue(TEXT("Equality(50)"), ShortViewNoNull == WideStringShort);
-	TestFalse(TEXT("Equality(51)"), ShortViewNoNull != WideStringShort);
-	TestTrue(TEXT("Equality(52)"), WideStringShort == ShortViewNoNull);
-	TestFalse(TEXT("Equality(53)"), WideStringShort != ShortViewNoNull);
-	TestFalse(TEXT("Equality(54)"), ShortViewNoNull == WideStringSrc);
-	TestTrue(TEXT("Equality(55)"), ShortViewNoNull != WideStringSrc);
-	TestFalse(TEXT("Equality(56)"), WideStringSrc == ShortViewNoNull);
-	TestTrue(TEXT("Equality(57)"), WideStringSrc != ShortViewNoNull);
+	TEST_TRUE(TEXT("Equality(42)"), ShortViewNoNull == WideStringLiteralShort);
+	TEST_FALSE(TEXT("Equality(43)"), ShortViewNoNull != WideStringLiteralShort);
+	TEST_TRUE(TEXT("Equality(44)"), WideStringLiteralShort == ShortViewNoNull);
+	TEST_FALSE(TEXT("Equality(45)"), WideStringLiteralShort != ShortViewNoNull);
+	TEST_FALSE(TEXT("Equality(46)"), ShortViewNoNull == WideStringLiteralSrc);
+	TEST_TRUE(TEXT("Equality(47)"), ShortViewNoNull != WideStringLiteralSrc);
+	TEST_FALSE(TEXT("Equality(48)"), WideStringLiteralSrc == ShortViewNoNull);
+	TEST_TRUE(TEXT("Equality(49)"), WideStringLiteralSrc != ShortViewNoNull);
+
+	TEST_TRUE(TEXT("Equality(50)"), ShortViewNoNull == WideStringShort);
+	TEST_FALSE(TEXT("Equality(51)"), ShortViewNoNull != WideStringShort);
+	TEST_TRUE(TEXT("Equality(52)"), WideStringShort == ShortViewNoNull);
+	TEST_FALSE(TEXT("Equality(53)"), WideStringShort != ShortViewNoNull);
+	TEST_FALSE(TEXT("Equality(54)"), ShortViewNoNull == WideStringSrc);
+	TEST_TRUE(TEXT("Equality(55)"), ShortViewNoNull != WideStringSrc);
+	TEST_FALSE(TEXT("Equality(56)"), WideStringSrc == ShortViewNoNull);
+	TEST_TRUE(TEXT("Equality(57)"), WideStringSrc != ShortViewNoNull);
 
 	FStringView WideViewNoNull = FStringView(WideStringLiteralLonger).Left(WideView.Len());
 
-	TestTrue(TEXT("Equality(58)"), WideViewNoNull == WideStringLiteralSrc);
-	TestFalse(TEXT("Equality(59)"), WideViewNoNull != WideStringLiteralSrc);
-	TestTrue(TEXT("Equality(60)"), WideStringLiteralSrc == WideViewNoNull);
-	TestFalse(TEXT("Equality(61)"), WideStringLiteralSrc != WideViewNoNull);
-	TestFalse(TEXT("Equality(62)"), WideViewNoNull == WideStringLiteralLonger);
-	TestTrue(TEXT("Equality(63)"), WideViewNoNull != WideStringLiteralLonger);
-	TestFalse(TEXT("Equality(64)"), WideStringLiteralLonger == WideViewNoNull);
-	TestTrue(TEXT("Equality(65)"), WideStringLiteralLonger != WideViewNoNull);
-
-	TestTrue(TEXT("Equality(66)"), WideViewNoNull == WideStringLiteralSrc);
-	TestFalse(TEXT("Equality(67)"), WideViewNoNull != WideStringLiteralSrc);
-	TestTrue(TEXT("Equality(68)"), WideStringLiteralSrc == WideViewNoNull);
-	TestFalse(TEXT("Equality(69)"), WideStringLiteralSrc != WideViewNoNull);
-	TestFalse(TEXT("Equality(70)"), WideViewNoNull == WideStringLiteralLonger);
-	TestTrue(TEXT("Equality(71)"), WideViewNoNull != WideStringLiteralLonger);
-	TestFalse(TEXT("Equality(72)"), WideStringLiteralLonger == WideViewNoNull);
-	TestTrue(TEXT("Equality(73)"), WideStringLiteralLonger != WideViewNoNull);
+	TEST_TRUE(TEXT("Equality(58)"), WideViewNoNull == WideStringLiteralSrc);
+	TEST_FALSE(TEXT("Equality(59)"), WideViewNoNull != WideStringLiteralSrc);
+	TEST_TRUE(TEXT("Equality(60)"), WideStringLiteralSrc == WideViewNoNull);
+	TEST_FALSE(TEXT("Equality(61)"), WideStringLiteralSrc != WideViewNoNull);
+	TEST_FALSE(TEXT("Equality(62)"), WideViewNoNull == WideStringLiteralLonger);
+	TEST_TRUE(TEXT("Equality(63)"), WideViewNoNull != WideStringLiteralLonger);
+	TEST_FALSE(TEXT("Equality(64)"), WideStringLiteralLonger == WideViewNoNull);
+	TEST_TRUE(TEXT("Equality(65)"), WideStringLiteralLonger != WideViewNoNull);
+
+	TEST_TRUE(TEXT("Equality(66)"), WideViewNoNull == WideStringLiteralSrc);
+	TEST_FALSE(TEXT("Equality(67)"), WideViewNoNull != WideStringLiteralSrc);
+	TEST_TRUE(TEXT("Equality(68)"), WideStringLiteralSrc == WideViewNoNull);
+	TEST_FALSE(TEXT("Equality(69)"), WideStringLiteralSrc != WideViewNoNull);
+	TEST_FALSE(TEXT("Equality(70)"), WideViewNoNull == WideStringLiteralLonger);
+	TEST_TRUE(TEXT("Equality(71)"), WideViewNoNull != WideStringLiteralLonger);
+	TEST_FALSE(TEXT("Equality(72)"), WideStringLiteralLonger == WideViewNoNull);
+	TEST_TRUE(TEXT("Equality(73)"), WideStringLiteralLonger != WideViewNoNull);
 
 	// ANSICHAR / TCHAR
 
@@ -284,27 +284,27 @@ TEST_CASE("Core::Containers::FStringView::Equality Test", "[Core][Containers][Sm
 	FAnsiStringView AnsiViewLower(AnsiStringLiteralLower);
 	FAnsiStringView AnsiViewUpper(AnsiStringLiteralUpper);
 
-	TestTrue(TEXT("Equality(74)"), AnsiView.Equals(WideView));
-	TestTrue(TEXT("Equality(75)"), WideView.Equals(AnsiView));
-	TestFalse(TEXT("Equality(76)"), AnsiViewLower.Equals(WideView, ESearchCase::CaseSensitive));
-	TestTrue(TEXT("Equality(77)"), AnsiViewLower.Equals(WideView, ESearchCase::IgnoreCase));
-	TestFalse(TEXT("Equality(78)"), WideView.Equals(AnsiViewLower, ESearchCase::CaseSensitive));
-	TestTrue(TEXT("Equality(79)"), WideView.Equals(AnsiViewLower, ESearchCase::IgnoreCase));
-	TestFalse(TEXT("Equality(80)"), AnsiViewUpper.Equals(WideView, ESearchCase::CaseSensitive));
-	TestTrue(TEXT("Equality(81)"), AnsiViewUpper.Equals(WideView, ESearchCase::IgnoreCase));
-	TestFalse(TEXT("Equality(82)"), WideView.Equals(AnsiViewUpper, ESearchCase::CaseSensitive));
-	TestTrue(TEXT("Equality(83)"), WideView.Equals(AnsiViewUpper, ESearchCase::IgnoreCase));
-
-	TestTrue(TEXT("Equality(84)"), WideView.Equals(AnsiStringLiteralSrc));
-	TestFalse(TEXT("Equality(85)"), WideView.Equals(AnsiStringLiteralLower, ESearchCase::CaseSensitive));
-	TestTrue(TEXT("Equality(86)"), WideView.Equals(AnsiStringLiteralLower, ESearchCase::IgnoreCase));
-	TestFalse(TEXT("Equality(87)"), WideView.Equals(AnsiStringLiteralUpper, ESearchCase::CaseSensitive));
-	TestTrue(TEXT("Equality(88)"), WideView.Equals(AnsiStringLiteralUpper, ESearchCase::IgnoreCase));
-	TestTrue(TEXT("Equality(89)"), AnsiView.Equals(WideStringLiteralSrc));
-	TestFalse(TEXT("Equality(90)"), AnsiViewLower.Equals(WideStringLiteralSrc, ESearchCase::CaseSensitive));
-	TestTrue(TEXT("Equality(91)"), AnsiViewLower.Equals(WideStringLiteralSrc, ESearchCase::IgnoreCase));
-	TestFalse(TEXT("Equality(92)"), AnsiViewUpper.Equals(WideStringLiteralSrc, ESearchCase::CaseSensitive));
-	TestTrue(TEXT("Equality(93)"), AnsiViewUpper.Equals(WideStringLiteralSrc, ESearchCase::IgnoreCase));
+	TEST_TRUE(TEXT("Equality(74)"), AnsiView.Equals(WideView));
+	TEST_TRUE(TEXT("Equality(75)"), WideView.Equals(AnsiView));
+	TEST_FALSE(TEXT("Equality(76)"), AnsiViewLower.Equals(WideView, ESearchCase::CaseSensitive));
+	TEST_TRUE(TEXT("Equality(77)"), AnsiViewLower.Equals(WideView, ESearchCase::IgnoreCase));
+	TEST_FALSE(TEXT("Equality(78)"), WideView.Equals(AnsiViewLower, ESearchCase::CaseSensitive));
+	TEST_TRUE(TEXT("Equality(79)"), WideView.Equals(AnsiViewLower, ESearchCase::IgnoreCase));
+	TEST_FALSE(TEXT("Equality(80)"), AnsiViewUpper.Equals(WideView, ESearchCase::CaseSensitive));
+	TEST_TRUE(TEXT("Equality(81)"), AnsiViewUpper.Equals(WideView, ESearchCase::IgnoreCase));
+	TEST_FALSE(TEXT("Equality(82)"), WideView.Equals(AnsiViewUpper, ESearchCase::CaseSensitive));
+	TEST_TRUE(TEXT("Equality(83)"), WideView.Equals(AnsiViewUpper, ESearchCase::IgnoreCase));
+
+	TEST_TRUE(TEXT("Equality(84)"), WideView.Equals(AnsiStringLiteralSrc));
+	TEST_FALSE(TEXT("Equality(85)"), WideView.Equals(AnsiStringLiteralLower, ESearchCase::CaseSensitive));
+	TEST_TRUE(TEXT("Equality(86)"), WideView.Equals(AnsiStringLiteralLower, ESearchCase::IgnoreCase));
+	TEST_FALSE(TEXT("Equality(87)"), WideView.Equals(AnsiStringLiteralUpper, ESearchCase::CaseSensitive));
+	TEST_TRUE(TEXT("Equality(88)"), WideView.Equals(AnsiStringLiteralUpper, ESearchCase::IgnoreCase));
+	TEST_TRUE(TEXT("Equality(89)"), AnsiView.Equals(WideStringLiteralSrc));
+	TEST_FALSE(TEXT("Equality(90)"), AnsiViewLower.Equals(WideStringLiteralSrc, ESearchCase::CaseSensitive));
+	TEST_TRUE(TEXT("Equality(91)"), AnsiViewLower.Equals(WideStringLiteralSrc, ESearchCase::IgnoreCase));
+	TEST_FALSE(TEXT("Equality(92)"), AnsiViewUpper.Equals(WideStringLiteralSrc, ESearchCase::CaseSensitive));
+	TEST_TRUE(TEXT("Equality(93)"), AnsiViewUpper.Equals(WideStringLiteralSrc, ESearchCase::IgnoreCase));
 
 	// Test equality of empty strings
 	{
@@ -312,12 +312,12 @@ TEST_CASE("Core::Containers::FStringView::Equality Test", "[Core][Containers][Sm
 		const TCHAR* NonEmptyLiteral = TEXT("ABC");
 		FStringView EmptyView;
 		FStringView NonEmptyView = TEXTVIEW("ABC");
-		TestTrue(TEXT("Equals(94)"), EmptyView.Equals(EmptyLiteral));
-		TestTrue(TEXT("Equals(95)"), !EmptyView.Equals(NonEmptyLiteral));
-		TestTrue(TEXT("Equals(96)"), !NonEmptyView.Equals(EmptyLiteral));
-		TestTrue(TEXT("Equals(97)"), EmptyView.Equals(EmptyView));
-		TestTrue(TEXT("Equals(98)"), !EmptyView.Equals(NonEmptyView));
-		TestTrue(TEXT("Equals(99)"), !NonEmptyView.Equals(EmptyView));
+		TEST_TRUE(TEXT("Equals(94)"), EmptyView.Equals(EmptyLiteral));
+		TEST_TRUE(TEXT("Equals(95)"), !EmptyView.Equals(NonEmptyLiteral));
+		TEST_TRUE(TEXT("Equals(96)"), !NonEmptyView.Equals(EmptyLiteral));
+		TEST_TRUE(TEXT("Equals(97)"), EmptyView.Equals(EmptyView));
+		TEST_TRUE(TEXT("Equals(98)"), !EmptyView.Equals(NonEmptyView));
+		TEST_TRUE(TEXT("Equals(99)"), !NonEmptyView.Equals(EmptyView));
 	}
 
 	// Test types convertible to a string view
@@ -329,7 +329,7 @@ TEST_CASE("Core::Containers::FStringView::Equality Test", "[Core][Containers][Sm
 	static_assert(TIsSame<bool, decltype(FWideStringView().Equals(TWideStringBuilder<16>()))>::Value, "Error with Equals");
 }
 
-TEST_CASE("Core::Containers::FStringView::Comparison Case Sensitive", "[Core][Containers][Smoke]")
+TEST_CASE_METHOD(FAutomationTestFixture, "Core::Containers::FStringView::Comparison Case Sensitive", "[Core][Containers][Smoke]")
 {
 	// Basic comparisons involving case
 	{
@@ -340,19 +340,19 @@ TEST_CASE("Core::Containers::FStringView::Comparison Case Sensitive", "[Core][Co
 
 		FStringView WideView(WideStringLiteralSrc);
 
-		TestTrue(TEXT("ComparisonCaseSensitive(0)"), WideView.Compare(WideStringLiteralSrc, ESearchCase::CaseSensitive) == 0);
-		TestFalse(TEXT("ComparisonCaseSensitive(1)"), WideView.Compare(WideStringLiteralLower, ESearchCase::CaseSensitive) > 0);
-		TestFalse(TEXT("ComparisonCaseSensitive(2)"), WideView.Compare(WideStringLiteralUpper, ESearchCase::CaseSensitive) < 0);
+		TEST_TRUE(TEXT("ComparisonCaseSensitive(0)"), WideView.Compare(WideStringLiteralSrc, ESearchCase::CaseSensitive) == 0);
+		TEST_FALSE(TEXT("ComparisonCaseSensitive(1)"), WideView.Compare(WideStringLiteralLower, ESearchCase::CaseSensitive) > 0);
+		TEST_FALSE(TEXT("ComparisonCaseSensitive(2)"), WideView.Compare(WideStringLiteralUpper, ESearchCase::CaseSensitive) < 0);
 
 		FStringView EmptyView(TEXT(""));
-		TestTrue(TEXT("ComparisonCaseSensitive(3)"), WideView.Compare(EmptyView, ESearchCase::CaseSensitive) > 0);
+		TEST_TRUE(TEXT("ComparisonCaseSensitive(3)"), WideView.Compare(EmptyView, ESearchCase::CaseSensitive) > 0);
 
 		FStringView IdenticalView(WideStringLiteralSrc);
-		TestTrue(TEXT("ComparisonCaseSensitive(4)"), WideView.Compare(IdenticalView, ESearchCase::CaseSensitive) == 0);
+		TEST_TRUE(TEXT("ComparisonCaseSensitive(4)"), WideView.Compare(IdenticalView, ESearchCase::CaseSensitive) == 0);
 
 		FAnsiStringView AnsiView(AnsiStringLiteralSrc);
-		TestTrue(TEXT("ComparisonCaseSensitive(5)"), WideView.Compare(AnsiView, ESearchCase::CaseSensitive) == 0);
-		TestTrue(TEXT("ComparisonCaseSensitive(6)"), WideView.Compare(AnsiStringLiteralSrc, ESearchCase::CaseSensitive) == 0);
+		TEST_TRUE(TEXT("ComparisonCaseSensitive(5)"), WideView.Compare(AnsiView, ESearchCase::CaseSensitive) == 0);
+		TEST_TRUE(TEXT("ComparisonCaseSensitive(6)"), WideView.Compare(AnsiStringLiteralSrc, ESearchCase::CaseSensitive) == 0);
 	}
 
 	// Test comparisons of different lengths
@@ -374,22 +374,22 @@ TEST_CASE("Core::Containers::FStringView::Comparison Case Sensitive", "[Core][Co
 		FStringView ViewShortLower(WideStringLiteralLowerFirst, 3);
 
 		// Same length, different cases
-		TestTrue(TEXT("ComparisonCaseSensitive(7)"), ViewLongUpper.Compare(ViewLongLower, ESearchCase::CaseSensitive) < 0);
-		TestTrue(TEXT("ComparisonCaseSensitive(8)"), ViewLongLower.Compare(ViewLongUpper, ESearchCase::CaseSensitive) > 0);
-		TestTrue(TEXT("ComparisonCaseInsensitive(9)"), ViewLongLower.Compare(AnsiStringLiteralUpper, ESearchCase::CaseSensitive) > 0);
-		TestTrue(TEXT("ComparisonCaseInsensitive(10)"), ViewShortUpper.Compare(WideStringLiteralLowerShort, ESearchCase::CaseSensitive) < 0);
+		TEST_TRUE(TEXT("ComparisonCaseSensitive(7)"), ViewLongUpper.Compare(ViewLongLower, ESearchCase::CaseSensitive) < 0);
+		TEST_TRUE(TEXT("ComparisonCaseSensitive(8)"), ViewLongLower.Compare(ViewLongUpper, ESearchCase::CaseSensitive) > 0);
+		TEST_TRUE(TEXT("ComparisonCaseInsensitive(9)"), ViewLongLower.Compare(AnsiStringLiteralUpper, ESearchCase::CaseSensitive) > 0);
+		TEST_TRUE(TEXT("ComparisonCaseInsensitive(10)"), ViewShortUpper.Compare(WideStringLiteralLowerShort, ESearchCase::CaseSensitive) < 0);
 
 		// Same case, different lengths
-		TestTrue(TEXT("ComparisonCaseSensitive(11)"), ViewLongUpper.Compare(ViewShortUpper, ESearchCase::CaseSensitive) > 0);
-		TestTrue(TEXT("ComparisonCaseSensitive(12)"), ViewShortUpper.Compare(ViewLongUpper, ESearchCase::CaseSensitive) < 0);
-		TestTrue(TEXT("ComparisonCaseInsensitive(13)"), ViewShortUpper.Compare(AnsiStringLiteralUpper, ESearchCase::CaseSensitive) < 0);
-		TestTrue(TEXT("ComparisonCaseInsensitive(14)"), ViewLongLower.Compare(WideStringLiteralLowerShort, ESearchCase::CaseSensitive) > 0);
+		TEST_TRUE(TEXT("ComparisonCaseSensitive(11)"), ViewLongUpper.Compare(ViewShortUpper, ESearchCase::CaseSensitive) > 0);
+		TEST_TRUE(TEXT("ComparisonCaseSensitive(12)"), ViewShortUpper.Compare(ViewLongUpper, ESearchCase::CaseSensitive) < 0);
+		TEST_TRUE(TEXT("ComparisonCaseInsensitive(13)"), ViewShortUpper.Compare(AnsiStringLiteralUpper, ESearchCase::CaseSensitive) < 0);
+		TEST_TRUE(TEXT("ComparisonCaseInsensitive(14)"), ViewLongLower.Compare(WideStringLiteralLowerShort, ESearchCase::CaseSensitive) > 0);
 
 		// Different length, different cases
-		TestTrue(TEXT("ComparisonCaseSensitive(15)"), ViewLongUpper.Compare(ViewShortLower, ESearchCase::CaseSensitive) < 0);
-		TestTrue(TEXT("ComparisonCaseSensitive(16)"), ViewShortLower.Compare(ViewLongUpper, ESearchCase::CaseSensitive) > 0);
-		TestTrue(TEXT("ComparisonCaseInsensitive(17)"), ViewShortLower.Compare(AnsiStringLiteralUpper, ESearchCase::CaseSensitive) > 0);
-		TestTrue(TEXT("ComparisonCaseInsensitive(18)"), ViewLongUpper.Compare(WideStringLiteralLowerShort, ESearchCase::CaseSensitive) < 0);
+		TEST_TRUE(TEXT("ComparisonCaseSensitive(15)"), ViewLongUpper.Compare(ViewShortLower, ESearchCase::CaseSensitive) < 0);
+		TEST_TRUE(TEXT("ComparisonCaseSensitive(16)"), ViewShortLower.Compare(ViewLongUpper, ESearchCase::CaseSensitive) > 0);
+		TEST_TRUE(TEXT("ComparisonCaseInsensitive(17)"), ViewShortLower.Compare(AnsiStringLiteralUpper, ESearchCase::CaseSensitive) > 0);
+		TEST_TRUE(TEXT("ComparisonCaseInsensitive(18)"), ViewLongUpper.Compare(WideStringLiteralLowerShort, ESearchCase::CaseSensitive) < 0);
 	}
 
 	// Test comparisons of empty strings
@@ -398,12 +398,12 @@ TEST_CASE("Core::Containers::FStringView::Comparison Case Sensitive", "[Core][Co
 		const TCHAR* NonEmptyLiteral = TEXT("ABC");
 		FStringView EmptyView;
 		FStringView NonEmptyView = TEXTVIEW("ABC");
-		TestTrue(TEXT("ComparisonEmpty(19)"), EmptyView.Compare(EmptyLiteral) == 0);
-		TestTrue(TEXT("ComparisonEmpty(20)"), EmptyView.Compare(NonEmptyLiteral) < 0);
-		TestTrue(TEXT("ComparisonEmpty(21)"), NonEmptyView.Compare(EmptyLiteral) > 0);
-		TestTrue(TEXT("ComparisonEmpty(22)"), EmptyView.Compare(EmptyView) == 0);
-		TestTrue(TEXT("ComparisonEmpty(23)"), EmptyView.Compare(NonEmptyView) < 0);
-		TestTrue(TEXT("ComparisonEmpty(24)"), NonEmptyView.Compare(EmptyView) > 0);
+		TEST_TRUE(TEXT("ComparisonEmpty(19)"), EmptyView.Compare(EmptyLiteral) == 0);
+		TEST_TRUE(TEXT("ComparisonEmpty(20)"), EmptyView.Compare(NonEmptyLiteral) < 0);
+		TEST_TRUE(TEXT("ComparisonEmpty(21)"), NonEmptyView.Compare(EmptyLiteral) > 0);
+		TEST_TRUE(TEXT("ComparisonEmpty(22)"), EmptyView.Compare(EmptyView) == 0);
+		TEST_TRUE(TEXT("ComparisonEmpty(23)"), EmptyView.Compare(NonEmptyView) < 0);
+		TEST_TRUE(TEXT("ComparisonEmpty(24)"), NonEmptyView.Compare(EmptyView) > 0);
 	}
 
 	// Test types convertible to a string view
@@ -415,7 +415,7 @@ TEST_CASE("Core::Containers::FStringView::Comparison Case Sensitive", "[Core][Co
 	static_assert(TIsSame<int32, decltype(FWideStringView().Compare(TWideStringBuilder<16>()))>::Value, "Error with Compare");
 }
 
-TEST_CASE("Core::Containers::FStringView::Comparison Case Insensitive", "[Core][Containers][Smoke]")
+TEST_CASE_METHOD(FAutomationTestFixture, "Core::Containers::FStringView::Comparison Case Insensitive", "[Core][Containers][Smoke]")
 {
 	// Basic comparisons involving case
 	{
@@ -426,19 +426,19 @@ TEST_CASE("Core::Containers::FStringView::Comparison Case Insensitive", "[Core][
 
 		FStringView WideView(WideStringLiteralSrc);
 
-		TestTrue(TEXT("ComparisonCaseInsensitive(0)"), WideView.Compare(WideStringLiteralSrc, ESearchCase::IgnoreCase) == 0);
-		TestTrue(TEXT("ComparisonCaseInsensitive(1)"), WideView.Compare(WideStringLiteralLower, ESearchCase::IgnoreCase) == 0);
-		TestTrue(TEXT("ComparisonCaseInsensitive(2)"), WideView.Compare(WideStringLiteralUpper, ESearchCase::IgnoreCase) == 0);
+		TEST_TRUE(TEXT("ComparisonCaseInsensitive(0)"), WideView.Compare(WideStringLiteralSrc, ESearchCase::IgnoreCase) == 0);
+		TEST_TRUE(TEXT("ComparisonCaseInsensitive(1)"), WideView.Compare(WideStringLiteralLower, ESearchCase::IgnoreCase) == 0);
+		TEST_TRUE(TEXT("ComparisonCaseInsensitive(2)"), WideView.Compare(WideStringLiteralUpper, ESearchCase::IgnoreCase) == 0);
 
 		FStringView EmptyView(TEXT(""));
-		TestTrue(TEXT("ComparisonCaseInsensitive(3)"), WideView.Compare(EmptyView, ESearchCase::IgnoreCase) > 0);
+		TEST_TRUE(TEXT("ComparisonCaseInsensitive(3)"), WideView.Compare(EmptyView, ESearchCase::IgnoreCase) > 0);
 
 		FStringView IdenticalView(WideStringLiteralSrc);
-		TestTrue(TEXT("ComparisonCaseInsensitive(4)"), WideView.Compare(IdenticalView, ESearchCase::IgnoreCase) == 0);
+		TEST_TRUE(TEXT("ComparisonCaseInsensitive(4)"), WideView.Compare(IdenticalView, ESearchCase::IgnoreCase) == 0);
 
 		FAnsiStringView AnsiView(AnsiStringLiteralSrc);
-		TestTrue(TEXT("ComparisonCaseSensitive(5)"), WideView.Compare(AnsiView, ESearchCase::IgnoreCase) == 0);
-		TestTrue(TEXT("ComparisonCaseSensitive(6)"), WideView.Compare(AnsiStringLiteralSrc, ESearchCase::IgnoreCase) == 0);
+		TEST_TRUE(TEXT("ComparisonCaseSensitive(5)"), WideView.Compare(AnsiView, ESearchCase::IgnoreCase) == 0);
+		TEST_TRUE(TEXT("ComparisonCaseSensitive(6)"), WideView.Compare(AnsiStringLiteralSrc, ESearchCase::IgnoreCase) == 0);
 	}
 
 	// Test comparisons of different lengths
@@ -460,37 +460,37 @@ TEST_CASE("Core::Containers::FStringView::Comparison Case Insensitive", "[Core][
 		FStringView ViewShortLower(WideStringLiteralLowerFirst, 3);
 
 		// Same length, different cases
-		TestTrue(TEXT("ComparisonCaseInsensitive(7)"), ViewLongUpper.Compare(ViewLongLower, ESearchCase::IgnoreCase) == 0);
-		TestTrue(TEXT("ComparisonCaseInsensitive(8)"), ViewLongLower.Compare(ViewLongUpper, ESearchCase::IgnoreCase) == 0);
-		TestTrue(TEXT("ComparisonCaseInsensitive(9)"), ViewLongLower.Compare(AnsiStringLiteralUpper, ESearchCase::IgnoreCase) == 0);
-		TestTrue(TEXT("ComparisonCaseInsensitive(10)"), ViewShortUpper.Compare(WideStringLiteralLowerShort, ESearchCase::IgnoreCase) == 0);
+		TEST_TRUE(TEXT("ComparisonCaseInsensitive(7)"), ViewLongUpper.Compare(ViewLongLower, ESearchCase::IgnoreCase) == 0);
+		TEST_TRUE(TEXT("ComparisonCaseInsensitive(8)"), ViewLongLower.Compare(ViewLongUpper, ESearchCase::IgnoreCase) == 0);
+		TEST_TRUE(TEXT("ComparisonCaseInsensitive(9)"), ViewLongLower.Compare(AnsiStringLiteralUpper, ESearchCase::IgnoreCase) == 0);
+		TEST_TRUE(TEXT("ComparisonCaseInsensitive(10)"), ViewShortUpper.Compare(WideStringLiteralLowerShort, ESearchCase::IgnoreCase) == 0);
 
 		// Same case, different lengths
-		TestTrue(TEXT("ComparisonCaseInsensitive(11)"), ViewLongUpper.Compare(ViewShortUpper, ESearchCase::IgnoreCase) > 0);
-		TestTrue(TEXT("ComparisonCaseInsensitive(12)"), ViewShortUpper.Compare(ViewLongUpper, ESearchCase::IgnoreCase) < 0);
-		TestTrue(TEXT("ComparisonCaseInsensitive(13)"), ViewShortUpper.Compare(AnsiStringLiteralUpper, ESearchCase::IgnoreCase) < 0);
-		TestTrue(TEXT("ComparisonCaseInsensitive(14)"), ViewLongLower.Compare(WideStringLiteralLowerShort, ESearchCase::IgnoreCase) > 0);
+		TEST_TRUE(TEXT("ComparisonCaseInsensitive(11)"), ViewLongUpper.Compare(ViewShortUpper, ESearchCase::IgnoreCase) > 0);
+		TEST_TRUE(TEXT("ComparisonCaseInsensitive(12)"), ViewShortUpper.Compare(ViewLongUpper, ESearchCase::IgnoreCase) < 0);
+		TEST_TRUE(TEXT("ComparisonCaseInsensitive(13)"), ViewShortUpper.Compare(AnsiStringLiteralUpper, ESearchCase::IgnoreCase) < 0);
+		TEST_TRUE(TEXT("ComparisonCaseInsensitive(14)"), ViewLongLower.Compare(WideStringLiteralLowerShort, ESearchCase::IgnoreCase) > 0);
 
 		// Different length, different cases
-		TestTrue(TEXT("ComparisonCaseInsensitive(15)"), ViewLongUpper.Compare(ViewShortLower, ESearchCase::IgnoreCase) > 0);
-		TestTrue(TEXT("ComparisonCaseInsensitive(16)"), ViewShortLower.Compare(ViewLongUpper, ESearchCase::IgnoreCase) < 0);
-		TestTrue(TEXT("ComparisonCaseInsensitive(17)"), ViewShortLower.Compare(AnsiStringLiteralUpper, ESearchCase::IgnoreCase) < 0);
-		TestTrue(TEXT("ComparisonCaseInsensitive(18)"), ViewLongUpper.Compare(WideStringLiteralLowerShort, ESearchCase::IgnoreCase) > 0);
+		TEST_TRUE(TEXT("ComparisonCaseInsensitive(15)"), ViewLongUpper.Compare(ViewShortLower, ESearchCase::IgnoreCase) > 0);
+		TEST_TRUE(TEXT("ComparisonCaseInsensitive(16)"), ViewShortLower.Compare(ViewLongUpper, ESearchCase::IgnoreCase) < 0);
+		TEST_TRUE(TEXT("ComparisonCaseInsensitive(17)"), ViewShortLower.Compare(AnsiStringLiteralUpper, ESearchCase::IgnoreCase) < 0);
+		TEST_TRUE(TEXT("ComparisonCaseInsensitive(18)"), ViewLongUpper.Compare(WideStringLiteralLowerShort, ESearchCase::IgnoreCase) > 0);
 	}
 }
 
-TEST_CASE("Core::Containers::FStringView::ArrayAccessor", "[Core][Containers][Smoke]")
+TEST_CASE_METHOD(FAutomationTestFixture, "Core::Containers::FStringView::ArrayAccessor", "[Core][Containers][Smoke]")
 {
 	const TCHAR* SrcString = TEXT("String To Test");
 	FStringView View(SrcString);
 
 	for (int32 i = 0; i < View.Len(); ++i)
 	{
-		TestEqual(TEXT("the character accessed"), View[i], SrcString[i]);
+		TEST_EQUAL(TEXT("the character accessed"), View[i], SrcString[i]);
 	}
 }
 
-TEST_CASE("Core::Containers::FStringView::ArrayModifiers", "[Core][Containers][Smoke]")
+TEST_CASE_METHOD(FAutomationTestFixture, "Core::Containers::FStringView::ArrayModifiers", "[Core][Containers][Smoke]")
 {
 	const TCHAR* FullText = TEXT("PrefixSuffix");
 	const TCHAR* Prefix = TEXT("Prefix");
@@ -501,8 +501,8 @@ TEST_CASE("Core::Containers::FStringView::ArrayModifiers", "[Core][Containers][S
 		FStringView View(FullText);
 		View.RemovePrefix(FCStringWide::Strlen(Prefix));
 
-		TestEqual(TEXT("View length"), View.Len(), FCStringWide::Strlen(Suffix));
-		TestEqual(TEXT("The result of Strncmp"), FCStringWide::Strncmp(View.GetData(), Suffix, View.Len()), 0);
+		TEST_EQUAL(TEXT("View length"), View.Len(), FCStringWide::Strlen(Suffix));
+		TEST_EQUAL(TEXT("The result of Strncmp"), FCStringWide::Strncmp(View.GetData(), Suffix, View.Len()), 0);
 	}
 
 	// Remove suffix
@@ -510,90 +510,90 @@ TEST_CASE("Core::Containers::FStringView::ArrayModifiers", "[Core][Containers][S
 		FStringView View(FullText);
 		View.RemoveSuffix(FCStringWide::Strlen(Suffix));
 
-		TestEqual(TEXT("View length"), View.Len(), FCStringWide::Strlen(Prefix));
-		TestEqual(TEXT("The result of Strncmp"), FCStringWide::Strncmp(View.GetData(), Prefix, View.Len()), 0);
+		TEST_EQUAL(TEXT("View length"), View.Len(), FCStringWide::Strlen(Prefix));
+		TEST_EQUAL(TEXT("The result of Strncmp"), FCStringWide::Strncmp(View.GetData(), Prefix, View.Len()), 0);
 	}
 }
 
-TEST_CASE("Core::Containers::FStringView::StartsWith", "[Core][Containers][Smoke]")
+TEST_CASE_METHOD(FAutomationTestFixture, "Core::Containers::FStringView::StartsWith", "[Core][Containers][Smoke]")
 {
 	// Test an empty view
 	{
 		FStringView View;
-		TestTrue(TEXT(" View.StartsWith"), View.StartsWith(TEXT("")));
-		TestFalse(TEXT(" View.StartsWith"), View.StartsWith(TEXT("Text")));
-		TestFalse(TEXT(" View.StartsWith"), View.StartsWith(TEXT('A')));
+		TEST_TRUE(TEXT(" View.StartsWith"), View.StartsWith(TEXT("")));
+		TEST_FALSE(TEXT(" View.StartsWith"), View.StartsWith(TEXT("Text")));
+		TEST_FALSE(TEXT(" View.StartsWith"), View.StartsWith(TEXT('A')));
 	}
 
 	// Test a valid view with the correct text
 	{
 		FStringView View(TEXT("String to test"));
-		TestTrue(TEXT(" View.StartsWith"), View.StartsWith(TEXT("String")));
-		TestTrue(TEXT(" View.StartsWith"), View.StartsWith(TEXT('S')));
+		TEST_TRUE(TEXT(" View.StartsWith"), View.StartsWith(TEXT("String")));
+		TEST_TRUE(TEXT(" View.StartsWith"), View.StartsWith(TEXT('S')));
 	}
 
 	// Test a valid view with incorrect text
 	{
 		FStringView View(TEXT("String to test"));
-		TestFalse(TEXT(" View.StartsWith"), View.StartsWith(TEXT("test")));
-		TestFalse(TEXT(" View.StartsWith"), View.StartsWith(TEXT('t')));
+		TEST_FALSE(TEXT(" View.StartsWith"), View.StartsWith(TEXT("test")));
+		TEST_FALSE(TEXT(" View.StartsWith"), View.StartsWith(TEXT('t')));
 	}
 
 	// Test a valid view with the correct text but with different case
 	{
 		FStringView View(TEXT("String to test"));
-		TestTrue(TEXT(" View.StartsWith"), View.StartsWith(TEXT("sTrInG")));
+		TEST_TRUE(TEXT(" View.StartsWith"), View.StartsWith(TEXT("sTrInG")));
 
 		// Searching by char is case sensitive to keep compatibility with FString
-		TestFalse(TEXT(" View.StartsWith"), View.StartsWith(TEXT('s')));
+		TEST_FALSE(TEXT(" View.StartsWith"), View.StartsWith(TEXT('s')));
 	}
 }
 
-TEST_CASE("Core::Containers::FStringView::EndsWith", "[Core][Containers][Smoke]")
+TEST_CASE_METHOD(FAutomationTestFixture, "Core::Containers::FStringView::EndsWith", "[Core][Containers][Smoke]")
 {
 	// Test an empty view
 	{
 		FStringView View;
-		TestTrue(TEXT(" View.EndsWith"), View.EndsWith(TEXT("")));
-		TestFalse(TEXT(" View.EndsWith"), View.EndsWith(TEXT("Text")));
-		TestFalse(TEXT(" View.EndsWith"), View.EndsWith(TEXT('A')));
+		TEST_TRUE(TEXT(" View.EndsWith"), View.EndsWith(TEXT("")));
+		TEST_FALSE(TEXT(" View.EndsWith"), View.EndsWith(TEXT("Text")));
+		TEST_FALSE(TEXT(" View.EndsWith"), View.EndsWith(TEXT('A')));
 	}
 
 	// Test a valid view with the correct text
 	{
 		FStringView View(TEXT("String to test"));
-		TestTrue(TEXT(" View.EndsWith"), View.EndsWith(TEXT("test")));
-		TestTrue(TEXT(" View.EndsWith"), View.EndsWith(TEXT('t')));
+		TEST_TRUE(TEXT(" View.EndsWith"), View.EndsWith(TEXT("test")));
+		TEST_TRUE(TEXT(" View.EndsWith"), View.EndsWith(TEXT('t')));
 	}
 
 	// Test a valid view with incorrect text
 	{
 		FStringView View(TEXT("String to test"));
-		TestFalse(TEXT(" View.EndsWith"), View.EndsWith(TEXT("String")));
-		TestFalse(TEXT(" View.EndsWith"), View.EndsWith(TEXT('S')));
+		TEST_FALSE(TEXT(" View.EndsWith"), View.EndsWith(TEXT("String")));
+		TEST_FALSE(TEXT(" View.EndsWith"), View.EndsWith(TEXT('S')));
 	}
 
 	// Test a valid view with the correct text but with different case
 	{
 		FStringView View(TEXT("String to test"));
-		TestTrue(TEXT(" View.EndsWith"), View.EndsWith(TEXT("TeST")));
+		TEST_TRUE(TEXT(" View.EndsWith"), View.EndsWith(TEXT("TeST")));
 
 		// Searching by char is case sensitive to keep compatibility with FString
-		TestFalse(TEXT(" View.EndsWith"), View.EndsWith(TEXT('T'))); 
+		TEST_FALSE(TEXT(" View.EndsWith"), View.EndsWith(TEXT('T'))); 
 	}
 }
 
-TEST_CASE("Core::Containers::FStringView::SubStr", "[Core][Containers][Smoke]")
+TEST_CASE_METHOD(FAutomationTestFixture, "Core::Containers::FStringView::SubStr", "[Core][Containers][Smoke]")
 {
 	{
 		FStringView EmptyView;
 		FStringView EmptyResult = EmptyView.SubStr(0, 10);
-		TestTrue(TEXT("FStringView::SubStr(0)"), EmptyResult.IsEmpty());
+		TEST_TRUE(TEXT("FStringView::SubStr(0)"), EmptyResult.IsEmpty());
 
 		// The following line is commented out as it would fail an assert and currently we cannot test for this in unit tests 
 		// FStringView OutofBoundsResult = EmptyView.SubStr(1000, 10000); 
 		FStringView OutofBoundsResult = EmptyView.SubStr(0, 10000);
-		TestTrue(TEXT("FStringView::SubStr(1)"), OutofBoundsResult.IsEmpty());
+		TEST_TRUE(TEXT("FStringView::SubStr(1)"), OutofBoundsResult.IsEmpty());
 	}
 
 	{
@@ -605,113 +605,113 @@ TEST_CASE("Core::Containers::FStringView::SubStr", "[Core][Containers][Smoke]")
 																	// string since the null terminator is still valid
 		FStringView OutofBoundsResult = View.SubStr(0, 1024);
 
-		TestTrue(TEXT("FStringView::SubStr(2)"), FCString::Strncmp(Word0.GetData(), TEXT("A"), Word0.Len()) == 0);
-		TestTrue(TEXT("FStringView::SubStr(3)"), FCString::Strncmp(Word1.GetData(), TEXT("test"), Word1.Len()) == 0);
-		TestTrue(TEXT("FStringView::SubStr(4)"), FCString::Strncmp(Word2.GetData(), TEXT("string"), Word2.Len()) == 0);
-		TestTrue(TEXT("FStringView::SubStr(5)"), NullTerminatorResult.IsEmpty());
-		TestTrue(TEXT("FStringView::SubStr(6)"), View == OutofBoundsResult);
+		TEST_TRUE(TEXT("FStringView::SubStr(2)"), FCString::Strncmp(Word0.GetData(), TEXT("A"), Word0.Len()) == 0);
+		TEST_TRUE(TEXT("FStringView::SubStr(3)"), FCString::Strncmp(Word1.GetData(), TEXT("test"), Word1.Len()) == 0);
+		TEST_TRUE(TEXT("FStringView::SubStr(4)"), FCString::Strncmp(Word2.GetData(), TEXT("string"), Word2.Len()) == 0);
+		TEST_TRUE(TEXT("FStringView::SubStr(5)"), NullTerminatorResult.IsEmpty());
+		TEST_TRUE(TEXT("FStringView::SubStr(6)"), View == OutofBoundsResult);
 	}
 }
 
-TEST_CASE("Core::Containers::FStringView::Left", "[Core][Containers][Smoke]")
+TEST_CASE_METHOD(FAutomationTestFixture, "Core::Containers::FStringView::Left", "[Core][Containers][Smoke]")
 {
 	{
 		FStringView EmptyView;
 		FStringView EmptyResult = EmptyView.Left(0);
-		TestTrue(TEXT("FStringView::Left"), EmptyResult.IsEmpty());
+		TEST_TRUE(TEXT("FStringView::Left"), EmptyResult.IsEmpty());
 
 		FStringView OutofBoundsResult = EmptyView.Left(1024);
-		TestTrue(TEXT("FStringView::Left"), OutofBoundsResult.IsEmpty());
+		TEST_TRUE(TEXT("FStringView::Left"), OutofBoundsResult.IsEmpty());
 	}
 	
 	{
 		FStringView View(TEXT("A test string padded"), 13); // "A test string" without null termination
 		FStringView Result = View.Left(8);
 
-		TestTrue(TEXT("FStringView::Left"), FCString::Strncmp(Result.GetData(), TEXT("A test s"), Result.Len()) == 0);
+		TEST_TRUE(TEXT("FStringView::Left"), FCString::Strncmp(Result.GetData(), TEXT("A test s"), Result.Len()) == 0);
 
 		FStringView OutofBoundsResult = View.Left(1024);
-		TestTrue(TEXT("FStringView::Left"), FCString::Strncmp(OutofBoundsResult.GetData(), TEXT("A test string"), OutofBoundsResult.Len()) == 0);
+		TEST_TRUE(TEXT("FStringView::Left"), FCString::Strncmp(OutofBoundsResult.GetData(), TEXT("A test string"), OutofBoundsResult.Len()) == 0);
 	}
 }
 
-TEST_CASE("Core::Containers::FStringView::LeftChop", "[Core][Containers][Smoke]")
+TEST_CASE_METHOD(FAutomationTestFixture, "Core::Containers::FStringView::LeftChop", "[Core][Containers][Smoke]")
 {
 	{
 		FStringView EmptyView;
 		FStringView EmptyResult = EmptyView.LeftChop(0);
-		TestTrue(TEXT("FStringView::LeftChop"), EmptyResult.IsEmpty());
+		TEST_TRUE(TEXT("FStringView::LeftChop"), EmptyResult.IsEmpty());
 
 		FStringView OutofBoundsResult = EmptyView.LeftChop(1024);
-		TestTrue(TEXT("FStringView::LeftChop"), OutofBoundsResult.IsEmpty());
+		TEST_TRUE(TEXT("FStringView::LeftChop"), OutofBoundsResult.IsEmpty());
 	}
 
 	{
 		FStringView View(TEXT("A test string padded"), 13); // "A test string" without null termination
 		FStringView Result = View.LeftChop(5);
 
-		TestTrue(TEXT("FStringView::LeftChop"), FCString::Strncmp(Result.GetData(), TEXT("A test s"), Result.Len()) == 0);
+		TEST_TRUE(TEXT("FStringView::LeftChop"), FCString::Strncmp(Result.GetData(), TEXT("A test s"), Result.Len()) == 0);
 
 		FStringView OutofBoundsResult = View.LeftChop(1024);
-		TestTrue(TEXT("FStringView::LeftChop"), FCString::Strncmp(OutofBoundsResult.GetData(), TEXT("A test string"), OutofBoundsResult.Len()) == 0);
+		TEST_TRUE(TEXT("FStringView::LeftChop"), FCString::Strncmp(OutofBoundsResult.GetData(), TEXT("A test string"), OutofBoundsResult.Len()) == 0);
 	}
 }
 
-TEST_CASE("Core::Containers::FStringView::Right", "[Core][Containers][Smoke]")
+TEST_CASE_METHOD(FAutomationTestFixture, "Core::Containers::FStringView::Right", "[Core][Containers][Smoke]")
 {
 	{
 		FStringView EmptyView;
 		FStringView EmptyResult = EmptyView.Right(0);
-		TestTrue(TEXT("FStringView::Right"), EmptyResult.IsEmpty());
+		TEST_TRUE(TEXT("FStringView::Right"), EmptyResult.IsEmpty());
 
 		FStringView OutofBoundsResult = EmptyView.Right(1024);
-		TestTrue(TEXT("FStringView::Right"), OutofBoundsResult.IsEmpty());
+		TEST_TRUE(TEXT("FStringView::Right"), OutofBoundsResult.IsEmpty());
 	}
 
 	{
 		FStringView View(TEXT("A test string padded"), 13); // "A test string" without null termination
 		FStringView Result = View.Right(8);
 
-		TestTrue(TEXT("FStringView::Right"), FCString::Strncmp(Result.GetData(), TEXT("t string"), Result.Len()) == 0);
+		TEST_TRUE(TEXT("FStringView::Right"), FCString::Strncmp(Result.GetData(), TEXT("t string"), Result.Len()) == 0);
 
 		FStringView OutofBoundsResult = View.Right(1024);
-		TestTrue(TEXT("FStringView::Right"), FCString::Strncmp(OutofBoundsResult.GetData(), TEXT("A test string"), OutofBoundsResult.Len()) == 0);
+		TEST_TRUE(TEXT("FStringView::Right"), FCString::Strncmp(OutofBoundsResult.GetData(), TEXT("A test string"), OutofBoundsResult.Len()) == 0);
 	}
 }
 
-TEST_CASE("Core::Containers::FStringView::RightChop", "[Core][Containers][Smoke]")
+TEST_CASE_METHOD(FAutomationTestFixture, "Core::Containers::FStringView::RightChop", "[Core][Containers][Smoke]")
 {
 	{
 		FStringView EmptyView;
 		FStringView EmptyResult = EmptyView.RightChop(0);
-		TestTrue(TEXT("FStringView::RightChop"), EmptyResult.IsEmpty());
+		TEST_TRUE(TEXT("FStringView::RightChop"), EmptyResult.IsEmpty());
 
 		FStringView OutofBoundsResult = EmptyView.RightChop(1024);
-		TestTrue(TEXT("FStringView::RightChop"), OutofBoundsResult.IsEmpty());
+		TEST_TRUE(TEXT("FStringView::RightChop"), OutofBoundsResult.IsEmpty());
 	}
 
 	{
 		FStringView View(TEXT("A test string padded"), 13); // "A test string" without null termination
 		FStringView Result = View.RightChop(3);
 
-		TestTrue(TEXT("FStringView::RightChop"), FCString::Strncmp(Result.GetData(), TEXT("est string"), Result.Len()) == 0);
+		TEST_TRUE(TEXT("FStringView::RightChop"), FCString::Strncmp(Result.GetData(), TEXT("est string"), Result.Len()) == 0);
 
 		FStringView OutofBoundsResult = View.RightChop(1024);
-		TestTrue(TEXT("FStringView::RightChop"), FCString::Strncmp(OutofBoundsResult.GetData(), TEXT("A test string"), OutofBoundsResult.Len()) == 0);
+		TEST_TRUE(TEXT("FStringView::RightChop"), FCString::Strncmp(OutofBoundsResult.GetData(), TEXT("A test string"), OutofBoundsResult.Len()) == 0);
 	}
 }
 
-TEST_CASE("Core::Containers::FStringView::Mid", "[Core][Containers][Smoke]")
+TEST_CASE_METHOD(FAutomationTestFixture, "Core::Containers::FStringView::Mid", "[Core][Containers][Smoke]")
 {
 	{
 		FStringView EmptyView;
 		FStringView EmptyResult = EmptyView.Mid(0, 10);
-		TestTrue(TEXT("FStringView::Mid(0)"), EmptyResult.IsEmpty());
+		TEST_TRUE(TEXT("FStringView::Mid(0)"), EmptyResult.IsEmpty());
 
 		// The following line is commented out as it would fail an assert and currently we cannot test for this in unit tests 
 		// FStringView OutofBoundsResult = EmptyView.Mid(1000, 10000); 
 		FStringView OutofBoundsResult = EmptyView.Mid(0, 10000);
-		TestTrue(TEXT("FStringView::Mid(1)"), OutofBoundsResult.IsEmpty());
+		TEST_TRUE(TEXT("FStringView::Mid(1)"), OutofBoundsResult.IsEmpty());
 	}
 
 	{
@@ -723,137 +723,137 @@ TEST_CASE("Core::Containers::FStringView::Mid", "[Core][Containers][Smoke]")
 																// string since the null terminator is still valid
 		FStringView OutofBoundsResult = View.Mid(0, 1024);
 
-		TestTrue(TEXT("FStringView::Mid(2)"), FCString::Strncmp(Word0.GetData(), TEXT("A"), Word0.Len()) == 0);
-		TestTrue(TEXT("FStringView::Mid(3)"), FCString::Strncmp(Word1.GetData(), TEXT("test"), Word1.Len()) == 0);
-		TestTrue(TEXT("FStringView::Mid(4)"), FCString::Strncmp(Word2.GetData(), TEXT("string"), Word2.Len()) == 0);
-		TestTrue(TEXT("FStringView::Mid(5)"), NullTerminatorResult.IsEmpty());
-		TestTrue(TEXT("FStringView::Mid(6)"), View == OutofBoundsResult);
-		TestTrue(TEXT("FStringView::Mid(7)"), View.Mid(512, 1024).IsEmpty());
-		TestTrue(TEXT("FStringView::Mid(8)"), View.Mid(4, 0).IsEmpty());
+		TEST_TRUE(TEXT("FStringView::Mid(2)"), FCString::Strncmp(Word0.GetData(), TEXT("A"), Word0.Len()) == 0);
+		TEST_TRUE(TEXT("FStringView::Mid(3)"), FCString::Strncmp(Word1.GetData(), TEXT("test"), Word1.Len()) == 0);
+		TEST_TRUE(TEXT("FStringView::Mid(4)"), FCString::Strncmp(Word2.GetData(), TEXT("string"), Word2.Len()) == 0);
+		TEST_TRUE(TEXT("FStringView::Mid(5)"), NullTerminatorResult.IsEmpty());
+		TEST_TRUE(TEXT("FStringView::Mid(6)"), View == OutofBoundsResult);
+		TEST_TRUE(TEXT("FStringView::Mid(7)"), View.Mid(512, 1024).IsEmpty());
+		TEST_TRUE(TEXT("FStringView::Mid(8)"), View.Mid(4, 0).IsEmpty());
 	}
 }
 
-TEST_CASE("Core::Containers::FStringView::TrimStartAndEnd", "[Core][Containers][Smoke]")
+TEST_CASE_METHOD(FAutomationTestFixture, "Core::Containers::FStringView::TrimStartAndEnd", "[Core][Containers][Smoke]")
 {
-	TestTrue(TEXT("FStringView::TrimStartAndEnd(\"\")"), TEXTVIEW("").TrimStartAndEnd().IsEmpty());
-	TestTrue(TEXT("FStringView::TrimStartAndEnd(\" \")"), TEXTVIEW(" ").TrimStartAndEnd().IsEmpty());
-	TestTrue(TEXT("FStringView::TrimStartAndEnd(\"  \")"), TEXTVIEW("  ").TrimStartAndEnd().IsEmpty());
-	TestTrue(TEXT("FStringView::TrimStartAndEnd(\" \\t\\r\\n\")"), TEXTVIEW(" \t\r\n").TrimStartAndEnd().IsEmpty());
-
-	TestEqual(TEXT("FStringView::TrimStartAndEnd(\"ABC123\")"), TEXTVIEW("ABC123").TrimStartAndEnd(), TEXTVIEW("ABC123"));
-	TestEqual(TEXT("FStringView::TrimStartAndEnd(\"A \\t\\r\\nB\")"), TEXTVIEW("A \t\r\nB").TrimStartAndEnd(), TEXTVIEW("A \t\r\nB"));
-	TestEqual(TEXT("FStringView::TrimStartAndEnd(\" \\t\\r\\nABC123\\n\\r\\t \")"), TEXTVIEW(" \t\r\nABC123\n\r\t ").TrimStartAndEnd(), TEXTVIEW("ABC123"));
+	TEST_TRUE(TEXT("FStringView::TrimStartAndEnd(\"\")"), TEXTVIEW("").TrimStartAndEnd().IsEmpty());
+	TEST_TRUE(TEXT("FStringView::TrimStartAndEnd(\" \")"), TEXTVIEW(" ").TrimStartAndEnd().IsEmpty());
+	TEST_TRUE(TEXT("FStringView::TrimStartAndEnd(\"  \")"), TEXTVIEW("  ").TrimStartAndEnd().IsEmpty());
+	TEST_TRUE(TEXT("FStringView::TrimStartAndEnd(\" \\t\\r\\n\")"), TEXTVIEW(" \t\r\n").TrimStartAndEnd().IsEmpty());
+
+	TEST_EQUAL(TEXT("FStringView::TrimStartAndEnd(\"ABC123\")"), TEXTVIEW("ABC123").TrimStartAndEnd(), TEXTVIEW("ABC123"));
+	TEST_EQUAL(TEXT("FStringView::TrimStartAndEnd(\"A \\t\\r\\nB\")"), TEXTVIEW("A \t\r\nB").TrimStartAndEnd(), TEXTVIEW("A \t\r\nB"));
+	TEST_EQUAL(TEXT("FStringView::TrimStartAndEnd(\" \\t\\r\\nABC123\\n\\r\\t \")"), TEXTVIEW(" \t\r\nABC123\n\r\t ").TrimStartAndEnd(), TEXTVIEW("ABC123"));
 }
 
-TEST_CASE("Core::Containers::FStringView::TrimStart", "[Core][Containers][Smoke]")
+TEST_CASE_METHOD(FAutomationTestFixture, "Core::Containers::FStringView::TrimStart", "[Core][Containers][Smoke]")
 {
-	TestTrue(TEXT("FStringView::TrimStart(\"\")"), TEXTVIEW("").TrimStart().IsEmpty());
-	TestTrue(TEXT("FStringView::TrimStart(\" \")"), TEXTVIEW(" ").TrimStart().IsEmpty());
-	TestTrue(TEXT("FStringView::TrimStart(\"  \")"), TEXTVIEW("  ").TrimStart().IsEmpty());
-	TestTrue(TEXT("FStringView::TrimStart(\" \\t\\r\\n\")"), TEXTVIEW(" \t\r\n").TrimStart().IsEmpty());
-
-	TestEqual(TEXT("FStringView::TrimStart(\"ABC123\")"), TEXTVIEW("ABC123").TrimStart(), TEXTVIEW("ABC123"));
-	TestEqual(TEXT("FStringView::TrimStart(\"A \\t\\r\\nB\")"), TEXTVIEW("A \t\r\nB").TrimStart(), TEXTVIEW("A \t\r\nB"));
-	TestEqual(TEXT("FStringView::TrimStart(\" \\t\\r\\nABC123\\n\\r\\t \")"), TEXTVIEW(" \t\r\nABC123\n\r\t ").TrimStart(), TEXTVIEW("ABC123\n\r\t "));
+	TEST_TRUE(TEXT("FStringView::TrimStart(\"\")"), TEXTVIEW("").TrimStart().IsEmpty());
+	TEST_TRUE(TEXT("FStringView::TrimStart(\" \")"), TEXTVIEW(" ").TrimStart().IsEmpty());
+	TEST_TRUE(TEXT("FStringView::TrimStart(\"  \")"), TEXTVIEW("  ").TrimStart().IsEmpty());
+	TEST_TRUE(TEXT("FStringView::TrimStart(\" \\t\\r\\n\")"), TEXTVIEW(" \t\r\n").TrimStart().IsEmpty());
+
+	TEST_EQUAL(TEXT("FStringView::TrimStart(\"ABC123\")"), TEXTVIEW("ABC123").TrimStart(), TEXTVIEW("ABC123"));
+	TEST_EQUAL(TEXT("FStringView::TrimStart(\"A \\t\\r\\nB\")"), TEXTVIEW("A \t\r\nB").TrimStart(), TEXTVIEW("A \t\r\nB"));
+	TEST_EQUAL(TEXT("FStringView::TrimStart(\" \\t\\r\\nABC123\\n\\r\\t \")"), TEXTVIEW(" \t\r\nABC123\n\r\t ").TrimStart(), TEXTVIEW("ABC123\n\r\t "));
 }
 
-TEST_CASE("Core::Containers::FStringView::TrimEnd", "[Core][Containers][Smoke]")
+TEST_CASE_METHOD(FAutomationTestFixture, "Core::Containers::FStringView::TrimEnd", "[Core][Containers][Smoke]")
 {
-	TestTrue(TEXT("FStringView::TrimEnd(\"\")"), TEXTVIEW("").TrimEnd().IsEmpty());
-	TestTrue(TEXT("FStringView::TrimEnd(\" \")"), TEXTVIEW(" ").TrimEnd().IsEmpty());
-	TestTrue(TEXT("FStringView::TrimEnd(\"  \")"), TEXTVIEW("  ").TrimEnd().IsEmpty());
-	TestTrue(TEXT("FStringView::TrimEnd(\" \\t\\r\\n\")"), TEXTVIEW(" \t\r\n").TrimEnd().IsEmpty());
-
-	TestEqual(TEXT("FStringView::TrimEnd(\"ABC123\")"), TEXTVIEW("ABC123").TrimEnd(), TEXTVIEW("ABC123"));
-	TestEqual(TEXT("FStringView::TrimEnd(\"A \\t\\r\\nB\")"), TEXTVIEW("A \t\r\nB").TrimEnd(), TEXTVIEW("A \t\r\nB"));
-	TestEqual(TEXT("FStringView::TrimEnd(\" \\t\\r\\nABC123\\n\\r\\t \")"), TEXTVIEW(" \t\r\nABC123\n\r\t ").TrimEnd(), TEXTVIEW(" \t\r\nABC123"));
+	TEST_TRUE(TEXT("FStringView::TrimEnd(\"\")"), TEXTVIEW("").TrimEnd().IsEmpty());
+	TEST_TRUE(TEXT("FStringView::TrimEnd(\" \")"), TEXTVIEW(" ").TrimEnd().IsEmpty());
+	TEST_TRUE(TEXT("FStringView::TrimEnd(\"  \")"), TEXTVIEW("  ").TrimEnd().IsEmpty());
+	TEST_TRUE(TEXT("FStringView::TrimEnd(\" \\t\\r\\n\")"), TEXTVIEW(" \t\r\n").TrimEnd().IsEmpty());
+
+	TEST_EQUAL(TEXT("FStringView::TrimEnd(\"ABC123\")"), TEXTVIEW("ABC123").TrimEnd(), TEXTVIEW("ABC123"));
+	TEST_EQUAL(TEXT("FStringView::TrimEnd(\"A \\t\\r\\nB\")"), TEXTVIEW("A \t\r\nB").TrimEnd(), TEXTVIEW("A \t\r\nB"));
+	TEST_EQUAL(TEXT("FStringView::TrimEnd(\" \\t\\r\\nABC123\\n\\r\\t \")"), TEXTVIEW(" \t\r\nABC123\n\r\t ").TrimEnd(), TEXTVIEW(" \t\r\nABC123"));
 }
 
-TEST_CASE("Core::Containers::FStringView::FindChar", "[Core][Containers][Smoke]")
+TEST_CASE_METHOD(FAutomationTestFixture, "Core::Containers::FStringView::FindChar", "[Core][Containers][Smoke]")
 {
 	FStringView EmptyView;
 	FStringView View = TEXT("aBce Fga");
 
 	{
 		int32 Index = INDEX_NONE;
-		TestFalse(TEXT("FStringView::FindChar-Return(0)"), EmptyView.FindChar(TEXT('a'), Index));
-		TestEqual(TEXT("FStringView::FindChar-Index(0)"), Index, INDEX_NONE);
+		TEST_FALSE(TEXT("FStringView::FindChar-Return(0)"), EmptyView.FindChar(TEXT('a'), Index));
+		TEST_EQUAL(TEXT("FStringView::FindChar-Index(0)"), Index, INDEX_NONE);
 	}
 
 	{
 		int32 Index = INDEX_NONE;
-		TestTrue(TEXT("FStringView::FindChar-Return(1)"), View.FindChar(TEXT('a'), Index));
-		TestEqual(TEXT("FStringView::FindChar-Index(1)"), Index, 0);
+		TEST_TRUE(TEXT("FStringView::FindChar-Return(1)"), View.FindChar(TEXT('a'), Index));
+		TEST_EQUAL(TEXT("FStringView::FindChar-Index(1)"), Index, 0);
 	}
 
 	{
 		int32 Index = INDEX_NONE;
-		TestTrue(TEXT("FStringView::FindChar-Return(2)"), View.FindChar(TEXT('F'), Index));
-		TestEqual(TEXT("FStringView::FindChar-Index(2)"), Index, 5);
+		TEST_TRUE(TEXT("FStringView::FindChar-Return(2)"), View.FindChar(TEXT('F'), Index));
+		TEST_EQUAL(TEXT("FStringView::FindChar-Index(2)"), Index, 5);
 	}
 
 	{
 		int32 Index = INDEX_NONE;
-		TestFalse(TEXT("FStringView::FindChar-Return(3)"), View.FindChar(TEXT('A'), Index));
-		TestEqual(TEXT("FStringView::FindChar-Index(3)"), Index, INDEX_NONE);
+		TEST_FALSE(TEXT("FStringView::FindChar-Return(3)"), View.FindChar(TEXT('A'), Index));
+		TEST_EQUAL(TEXT("FStringView::FindChar-Index(3)"), Index, INDEX_NONE);
 	}
 
 	{
 		int32 Index = INDEX_NONE;
-		TestFalse(TEXT("FStringView::FindChar-Return(4)"), View.FindChar(TEXT('d'), Index));
-		TestEqual(TEXT("FStringView::FindChar-Index(4)"), Index, INDEX_NONE);
+		TEST_FALSE(TEXT("FStringView::FindChar-Return(4)"), View.FindChar(TEXT('d'), Index));
+		TEST_EQUAL(TEXT("FStringView::FindChar-Index(4)"), Index, INDEX_NONE);
 	}
 
 	{
 		int32 Index = INDEX_NONE;
-		TestTrue(TEXT("FStringView::FindChar-Return(5)"), View.FindChar(TEXT(' '), Index));
-		TestEqual(TEXT("FStringView::FindChar-Index(5)"), Index, 4);
+		TEST_TRUE(TEXT("FStringView::FindChar-Return(5)"), View.FindChar(TEXT(' '), Index));
+		TEST_EQUAL(TEXT("FStringView::FindChar-Index(5)"), Index, 4);
 	}
 }
 
-TEST_CASE("Core::Containers::FStringView::FindLastChar", "[Core][Containers][Smoke]")
+TEST_CASE_METHOD(FAutomationTestFixture, "Core::Containers::FStringView::FindLastChar", "[Core][Containers][Smoke]")
 {
 	FStringView EmptyView;
 	FStringView View = TEXT("aBce Fga");
 
 	{
 		int32 Index = INDEX_NONE;
-		TestFalse(TEXT("FStringView::FindChar-Return(0)"), EmptyView.FindLastChar(TEXT('a'), Index));
-		TestEqual(TEXT("FStringView::FindChar-Index(0)"), Index, INDEX_NONE);
+		TEST_FALSE(TEXT("FStringView::FindChar-Return(0)"), EmptyView.FindLastChar(TEXT('a'), Index));
+		TEST_EQUAL(TEXT("FStringView::FindChar-Index(0)"), Index, INDEX_NONE);
 	}
 
 	{
 		int32 Index = INDEX_NONE;
-		TestTrue(TEXT("FStringView::FindChar-Return(1)"), View.FindLastChar(TEXT('a'), Index));
-		TestEqual(TEXT("FStringView::FindChar-Index(1)"), Index, 7);
+		TEST_TRUE(TEXT("FStringView::FindChar-Return(1)"), View.FindLastChar(TEXT('a'), Index));
+		TEST_EQUAL(TEXT("FStringView::FindChar-Index(1)"), Index, 7);
 	}
 
 	{
 		int32 Index = INDEX_NONE;
-		TestTrue(TEXT("FStringView::FindChar-Return(2)"), View.FindLastChar(TEXT('B'), Index));
-		TestEqual(TEXT("FStringView::FindChar-Index(2)"), Index, 1);
+		TEST_TRUE(TEXT("FStringView::FindChar-Return(2)"), View.FindLastChar(TEXT('B'), Index));
+		TEST_EQUAL(TEXT("FStringView::FindChar-Index(2)"), Index, 1);
 	}
 
 	{
 		int32 Index = INDEX_NONE;
-		TestFalse(TEXT("FStringView::FindChar-Return(3)"), View.FindLastChar(TEXT('A'), Index));
-		TestEqual(TEXT("FStringView::FindChar-Index(3)"), Index, INDEX_NONE);
+		TEST_FALSE(TEXT("FStringView::FindChar-Return(3)"), View.FindLastChar(TEXT('A'), Index));
+		TEST_EQUAL(TEXT("FStringView::FindChar-Index(3)"), Index, INDEX_NONE);
 	}
 
 	{
 		int32 Index = INDEX_NONE;
-		TestFalse(TEXT("FStringView::FindChar-Return(4)"), View.FindLastChar(TEXT('d'), Index));
-		TestEqual(TEXT("FStringView::FindChar-Index(4)"), Index, INDEX_NONE);
+		TEST_FALSE(TEXT("FStringView::FindChar-Return(4)"), View.FindLastChar(TEXT('d'), Index));
+		TEST_EQUAL(TEXT("FStringView::FindChar-Index(4)"), Index, INDEX_NONE);
 	}
 
 	{
 		int32 Index = INDEX_NONE;
-		TestTrue(TEXT("FStringView::FindChar-Return(5)"), View.FindLastChar(TEXT(' '), Index));
-		TestEqual(TEXT("FStringView::FindChar-Index(5)"), Index, 4);
+		TEST_TRUE(TEXT("FStringView::FindChar-Return(5)"), View.FindLastChar(TEXT(' '), Index));
+		TEST_EQUAL(TEXT("FStringView::FindChar-Index(5)"), Index, 4);
 	}
 }
 
-void TestSlicing(const FString& Str)
+void TestSlicing(FAutomationTestFixture& Test, const FString& Str)
 {
 	const FStringView View = Str;
 	const int32       Len  = Str.Len();
@@ -866,7 +866,7 @@ void TestSlicing(const FString& Str)
 			FString     Substring = Str .Left(Index);
 			FStringView Subview   = View.Left(Index);
 
-			TestEqual(FString::Printf(TEXT("FStringView(\"%s\")::Left(%d)"), *Str, Index), FString(Subview), Substring);
+			TEST_EQUAL(FString::Printf(TEXT("FStringView(\"%s\")::Left(%d)"), *Str, Index), FString(Subview), Substring);
 		}
 	}
 
@@ -878,7 +878,7 @@ void TestSlicing(const FString& Str)
 			FString     Substring = Str .LeftChop(Index);
 			FStringView Subview   = View.LeftChop(Index);
 
-			TestEqual(FString::Printf(TEXT("FStringView(\"%s\")::LeftChop(%d)"), *Str, Index), FString(Subview), Substring);
+			TEST_EQUAL(FString::Printf(TEXT("FStringView(\"%s\")::LeftChop(%d)"), *Str, Index), FString(Subview), Substring);
 		}
 	}
 
@@ -890,7 +890,7 @@ void TestSlicing(const FString& Str)
 			FString     Substring = Str .Right(Index);
 			FStringView Subview   = View.Right(Index);
 
-			TestEqual(FString::Printf(TEXT("FStringView(\"%s\")::Right(%d)"), *Str, Index), FString(Subview), Substring);
+			TEST_EQUAL(FString::Printf(TEXT("FStringView(\"%s\")::Right(%d)"), *Str, Index), FString(Subview), Substring);
 		}
 	}
 
@@ -902,7 +902,7 @@ void TestSlicing(const FString& Str)
 			FString     Substring = Str .RightChop(Index);
 			FStringView Subview   = View.RightChop(Index);
 
-			TestEqual(FString::Printf(TEXT("FStringView(\"%s\").RightChop(%d)"), *Str, Index), FString(Subview), Substring);
+			TEST_EQUAL(FString::Printf(TEXT("FStringView(\"%s\").RightChop(%d)"), *Str, Index), FString(Subview), Substring);
 		}
 	}
 
@@ -916,7 +916,7 @@ void TestSlicing(const FString& Str)
 				FString     Substring = Str .Mid(Index, Count);
 				FStringView Subview   = View.Mid(Index, Count);
 
-				TestEqual(FString::Printf(TEXT("FStringView(\"%s\")::Mid(%d, %d)"), *Str, Index, Count), FString(Subview), Substring);
+				TEST_EQUAL(FString::Printf(TEXT("FStringView(\"%s\")::Mid(%d, %d)"), *Str, Index, Count), FString(Subview), Substring);
 			}
 		}
 
@@ -931,24 +931,24 @@ void TestSlicing(const FString& Str)
 				FString     Substring = Str .Mid(Index, Count);
 				FStringView Subview   = View.Mid(Index, Count);
 
-				TestEqual(FString::Printf(TEXT("FStringView(\"%s\")::Mid(%d, %d)"), *Str, Index, Count), FString(Subview), Substring);
+				TEST_EQUAL(FString::Printf(TEXT("FStringView(\"%s\")::Mid(%d, %d)"), *Str, Index, Count), FString(Subview), Substring);
 			}
 		}
 	}
 }
 
-TEST_CASE("Core::Containers::FStringView::Slice", "[Core][Containers][Smoke]")
+TEST_CASE_METHOD(FAutomationTestFixture, "Core::Containers::FStringView::Slice", "[Core][Containers][Smoke]")
 {
 	// We assume that FString has already passed its tests, and we just want views to be consistent with it
 
 	// Test an aribtrary string
-	TestSlicing(TEXT("Test string"));
+	TestSlicing(*this, TEXT("Test string"));
 
 	// Test an empty string
-	TestSlicing(FString());
+	TestSlicing(*this, FString());
 
 	// Test an null-terminator-only empty string
 	FString TerminatorOnly;
 	TerminatorOnly.GetCharArray().Add(TEXT('\0'));
-	TestSlicing(TerminatorOnly);
+	TestSlicing(*this, TerminatorOnly);
 }
diff --git a/Engine/Source/Runtime/Core/Tests/Containers/TickerTests.cpp b/Engine/Source/Runtime/Core/Tests/Containers/TickerTests.cpp
--- a/Engine/Source/Runtime/Core/Tests/Containers/TickerTests.cpp
+++ b/Engine/Source/Runtime/Core/Tests/Containers/TickerTests.cpp
@@ -1,13 +1,10 @@
 // Copyright Epic Games, Inc. All Rights Reserved.
 
 #include "CoreTypes.h"
-#include "Tests/Benchmark.h"
 #include "Containers/Ticker.h"
 #include "Tasks/Task.h"
-#include "Tests/Benchmark.h"
-#include "TestHarness.h"
-
-#include "TestCommon/CoreUtilities.h"
+#include "Core/Tests/Benchmark.h"
+#include "TestFixtures/CoreTestFixture.h"
 
 template<uint32 NumDelegates, uint32 NumTicks>
 void TickerPerfTest()
@@ -59,9 +56,8 @@ void TSTickerPerfTest()
 	}
 }
 
-TEST_CASE("Core::Containers::FTSTicker::Smoke Test", "[Core][Containers][Smoke]")
+TEST_CASE_METHOD(FCoreTestFixture, "Core::Containers::FTSTicker::Smoke Test", "[Core][Containers][Smoke]")
 {
-	InitTaskGraphAndDependencies();
 	{	// a delegate returning false is executed once
 		FTSTicker Ticker;
 		bool bExecuted = false;
@@ -192,13 +188,10 @@ TEST_CASE("Core::Containers::FTSTicker::Smoke Test", "[Core][Containers][Smoke]"
 		Ticker.Tick(0.0f);
 		CHECK(bTicked);
 	}
+}	
 
-	CleanupTaskGraphAndDependencies();
-}
-
-TEST_CASE("Core::Containers::FTSTicker::Stress Test", "[Core][Containers][Perf]")
+TEST_CASE("Core::Containers::FTSTicker::Stress Test", "[Core][Containers][.Perf]")
 {
-	InitTaskGraphAndDependencies(true);
 	// multithreaded stress test
 	{
 		using namespace UE::Tasks;
@@ -252,6 +245,4 @@ TEST_CASE("Core::Containers::FTSTicker::Stress Test", "[Core][Containers][Perf]"
 
 	UE_BENCHMARK(5, TickerPerfTest<100, 100>);
 	UE_BENCHMARK(5, TSTickerPerfTest<100, 100>);
-
-	CleanupTaskGraphAndDependencies();
 }
\ No newline at end of file
diff --git a/Engine/Source/Runtime/Core/Tests/HAL/FileManagerGenericTest.cpp b/Engine/Source/Runtime/Core/Tests/HAL/FileManagerGenericTest.cpp
--- a/Engine/Source/Runtime/Core/Tests/HAL/FileManagerGenericTest.cpp
+++ b/Engine/Source/Runtime/Core/Tests/HAL/FileManagerGenericTest.cpp
@@ -8,7 +8,7 @@
 #include "Misc/StringBuilder.h"
 #include "TestHarness.h"
 
-class FArchiveFileReaderGenericTest
+class FArchiveFileReaderGenericTest : public FAutomationTestFixture
 {
 public:
 	~FArchiveFileReaderGenericTest();
@@ -47,60 +47,60 @@ void FArchiveFileReaderGenericTest::TestInternalPrecache()
 	TArray<uint8> UnusedBytes;
 	UnusedBytes.SetNumUninitialized(BufferSize);
 	TUniquePtr<FArchiveFileReaderGeneric> Reader(new FArchiveFileReaderGeneric(ReadHandle.Release(), *TestFileName, FileSize, BufferSize));
-	TestEqual(TEXT("Initial Pos should be 0"), Reader->Pos, 0LL);
-	TestEqual(TEXT("Size should be what was passed in"), Reader->Size, FileSize);
-	TestEqual(TEXT("BufferSize should be what was passed in"), Reader->BufferSize, BufferSize);
+	TEST_EQUAL(TEXT("Initial Pos should be 0"), Reader->Pos, 0LL);
+	TEST_EQUAL(TEXT("Size should be what was passed in"), Reader->Size, FileSize);
+	TEST_EQUAL(TEXT("BufferSize should be what was passed in"), Reader->BufferSize, BufferSize);
 
 	// Vanilla InternalPrecache at start of file
 	bool Result = Reader->InternalPrecache(0, BufferSize);
-	TestEqual(TEXT("Vanilla0 - InternalPrecache should succeed"), Result, true);
-	TestEqual(TEXT("Vanilla0 - InternalPrecache should not alter Pos"), Reader->Pos, 0LL);
-	TestEqual(TEXT("Vanilla0 - BufferBase should have be set to PrecacheOffset aka Pos"), Reader->BufferBase, 0LL);
-	TestEqual(TEXT("Vanilla0 - Bytes precached should be set to BufferSize unless it runs out of room"), Reader->BufferArray.Num(), BufferSize);
+	TEST_EQUAL(TEXT("Vanilla0 - InternalPrecache should succeed"), Result, true);
+	TEST_EQUAL(TEXT("Vanilla0 - InternalPrecache should not alter Pos"), Reader->Pos, 0LL);
+	TEST_EQUAL(TEXT("Vanilla0 - BufferBase should have be set to PrecacheOffset aka Pos"), Reader->BufferBase, 0LL);
+	TEST_EQUAL(TEXT("Vanilla0 - Bytes precached should be set to BufferSize unless it runs out of room"), Reader->BufferArray.Num(), BufferSize);
 	TestBytesValid(TEXT("Vanilla0 - BufferArray should be expected bytes"), Reader->BufferArray.GetData(), Reader->BufferArray.Num(), Reader->BufferBase);
 
 	// InternalPrecache at PrecacheOffset != Pos is ignored, and returns false if the first byte at PrecacheOffset is not cached
 	Result = Reader->InternalPrecache(BufferSize  + BufferSize, BufferSize);
-	TestEqual(TEXT("PrecacheNotAtPos - Should fail"), Result, false);
-	TestEqual(TEXT("PrecacheNotAtPos - BufferBase should not be altered"), Reader->BufferBase, 0LL);
-	TestEqual(TEXT("PrecacheNotAtPos - Buffe should not be altered"), Reader->BufferArray.Num(), BufferSize);
+	TEST_EQUAL(TEXT("PrecacheNotAtPos - Should fail"), Result, false);
+	TEST_EQUAL(TEXT("PrecacheNotAtPos - BufferBase should not be altered"), Reader->BufferBase, 0LL);
+	TEST_EQUAL(TEXT("PrecacheNotAtPos - Buffe should not be altered"), Reader->BufferArray.Num(), BufferSize);
 
 	// InternalPrecache at PrecacheOffset != Pos is ignored, but returns true if the first byte at PrecacheOffset is cached
 	Result = Reader->InternalPrecache(BufferSize - 1, BufferSize);
-	TestEqual(TEXT("PrecacheNotAtPos - Should succeed since the first byte at PrecacheOffset is buffered"), Result, true);
-	TestEqual(TEXT("PrecacheNotAtPos - BufferBase should not be altered"), Reader->BufferBase, 0LL);
-	TestEqual(TEXT("PrecacheNotAtPos - Buffe should not be altered"), Reader->BufferArray.Num(), BufferSize);
+	TEST_EQUAL(TEXT("PrecacheNotAtPos - Should succeed since the first byte at PrecacheOffset is buffered"), Result, true);
+	TEST_EQUAL(TEXT("PrecacheNotAtPos - BufferBase should not be altered"), Reader->BufferBase, 0LL);
+	TEST_EQUAL(TEXT("PrecacheNotAtPos - Buffe should not be altered"), Reader->BufferArray.Num(), BufferSize);
 
 	// InternalPrecache partway through the buffer should allocate the full buffer/not allocate anything, depending on bPrecacheAsSoonAsPossible
 	int64 PosStart = BufferSize / 2;
 	SetPosAndBuffer(Reader, PosStart, 0, BufferSize);
 	Result = Reader->InternalPrecache(PosStart, BufferSize);
-	TestEqual(TEXT("PrecachePartwayPosMoreThanBuffer - Should succeed"), Result, true);
+	TEST_EQUAL(TEXT("PrecachePartwayPosMoreThanBuffer - Should succeed"), Result, true);
 	if (FArchiveFileReaderGeneric::bPrecacheAsSoonAsPossible)
 	{
-		TestEqual(TEXT("PrecachePartwayPosMoreThanBuffer - BufferBase should be updated"), Reader->BufferBase, PosStart);
+		TEST_EQUAL(TEXT("PrecachePartwayPosMoreThanBuffer - BufferBase should be updated"), Reader->BufferBase, PosStart);
 	}
 	else
 	{
-		TestEqual(TEXT("PrecachePartwayPosMoreThanBuffer - BufferBase should not be updated"), Reader->BufferBase, 0LL);
+		TEST_EQUAL(TEXT("PrecachePartwayPosMoreThanBuffer - BufferBase should not be updated"), Reader->BufferBase, 0LL);
 	}
-	TestEqual(TEXT("PrecachePartwayPosMoreThanBuffer - BufferCount should be set to BufferSize"), Reader->BufferArray.Num(), BufferSize);
+	TEST_EQUAL(TEXT("PrecachePartwayPosMoreThanBuffer - BufferCount should be set to BufferSize"), Reader->BufferArray.Num(), BufferSize);
 	TestBytesValid(TEXT("PrecachePartwayPosMoreThanBuffer - BufferBytes should match BufferBase"), Reader->BufferArray.GetData(), Reader->BufferArray.Num(), Reader->BufferBase);
 
 	// InternalPrecache partway through the buffer should allocate the full buffer/not allocate anything, and should handle the case of BufferArray's allocation being smaller than BufferSize
 	SetPosAndBuffer(Reader, 3*BufferSize/4, BufferSize/2, BufferSize/2);
 	Reader->BufferArray.Shrink();
 	Result = Reader->InternalPrecache(3*BufferSize/4, BufferSize);
-	TestEqual(TEXT("PrecachePartwayPosMoreThanBufferReallocation - Should succeed"), Result, true);
+	TEST_EQUAL(TEXT("PrecachePartwayPosMoreThanBufferReallocation - Should succeed"), Result, true);
 	if (FArchiveFileReaderGeneric::bPrecacheAsSoonAsPossible)
 	{
-		TestEqual(TEXT("PrecachePartwayPosMoreThanBufferReallocation - BufferBase should be updated"), Reader->BufferBase, 3*BufferSize/4);
-		TestEqual(TEXT("PrecachePartwayPosMoreThanBufferReallocation - BufferCount should be set to BufferSize"), Reader->BufferArray.Num(), BufferSize);
+		TEST_EQUAL(TEXT("PrecachePartwayPosMoreThanBufferReallocation - BufferBase should be updated"), Reader->BufferBase, 3*BufferSize/4);
+		TEST_EQUAL(TEXT("PrecachePartwayPosMoreThanBufferReallocation - BufferCount should be set to BufferSize"), Reader->BufferArray.Num(), BufferSize);
 	}
 	else
 	{
-		TestEqual(TEXT("PrecachePartwayPosMoreThanBufferReallocation - BufferBase should not be updated"), Reader->BufferBase, BufferSize/2);
-		TestEqual(TEXT("PrecachePartwayPosMoreThanBufferReallocation - BufferCount should not be updated"), Reader->BufferArray.Num(), BufferSize/2);
+		TEST_EQUAL(TEXT("PrecachePartwayPosMoreThanBufferReallocation - BufferBase should not be updated"), Reader->BufferBase, BufferSize/2);
+		TEST_EQUAL(TEXT("PrecachePartwayPosMoreThanBufferReallocation - BufferCount should not be updated"), Reader->BufferArray.Num(), BufferSize/2);
 	}
 	TestBytesValid(TEXT("PrecachePartwayPosMoreThanBufferReallocation - BufferBytes should match BufferBase"), Reader->BufferArray.GetData(), Reader->BufferArray.Num(), Reader->BufferBase);
 
@@ -108,59 +108,59 @@ void FArchiveFileReaderGenericTest::TestInternalPrecache()
 	PosStart = BufferSize / 2;
 	SetPosAndBuffer(Reader, PosStart, 0, BufferSize);
 	Result = Reader->InternalPrecache(PosStart, BufferSize / 4);
-	TestEqual(TEXT("PrecachePartwayPosLessThanBuffer - Should succeed"), Result, true);
-	TestEqual(TEXT("PrecachePartwayPosLessThanBuffer - BufferBase should not be updated"), Reader->BufferBase, 0LL);
-	TestEqual(TEXT("PrecachePartwayPosLessThanBuffer - BufferCount should be set to BufferSize"), Reader->BufferArray.Num(), BufferSize);
+	TEST_EQUAL(TEXT("PrecachePartwayPosLessThanBuffer - Should succeed"), Result, true);
+	TEST_EQUAL(TEXT("PrecachePartwayPosLessThanBuffer - BufferBase should not be updated"), Reader->BufferBase, 0LL);
+	TEST_EQUAL(TEXT("PrecachePartwayPosLessThanBuffer - BufferCount should be set to BufferSize"), Reader->BufferArray.Num(), BufferSize);
 	TestBytesValid(TEXT("PrecachePartwayPosLessThanBuffer - BufferBytes should match BufferBase"), Reader->BufferArray.GetData(), Reader->BufferArray.Num(), Reader->BufferBase);
 
 	// InternalPrecache at the end of the buffer should allocate the full buffer
 	PosStart = BufferSize;
 	SetPosAndBuffer(Reader, PosStart, 0, BufferSize);
 	Result = Reader->InternalPrecache(PosStart, BufferSize);
-	TestEqual(TEXT("PrecachePosEndOfBuffer - Should succeed"), Result, true);
-	TestEqual(TEXT("PrecachePosEndOfBuffer - BufferBase should be updated"), Reader->BufferBase, PosStart);
-	TestEqual(TEXT("PrecachePosEndOfBuffer - BufferCount should be set to BufferSize"), Reader->BufferArray.Num(), BufferSize);
+	TEST_EQUAL(TEXT("PrecachePosEndOfBuffer - Should succeed"), Result, true);
+	TEST_EQUAL(TEXT("PrecachePosEndOfBuffer - BufferBase should be updated"), Reader->BufferBase, PosStart);
+	TEST_EQUAL(TEXT("PrecachePosEndOfBuffer - BufferCount should be set to BufferSize"), Reader->BufferArray.Num(), BufferSize);
 	TestBytesValid(TEXT("PrecachePosEndOfBuffer - BufferBytes should match BufferBase"), Reader->BufferArray.GetData(), Reader->BufferArray.Num(), Reader->BufferBase);
 
 	// InternalPrecache near the end of the file should clamp BufferSize to Size - Pos
 	PosStart = FileSize - 1;
 	SetPosAndBuffer(Reader, PosStart, 0, 0);
 	Result = Reader->InternalPrecache(PosStart, BufferSize);
-	TestEqual(TEXT("PrecacheNearEndOfFile - Should succeed"), Result, true);
-	TestEqual(TEXT("PrecacheNearEndOfFile - BufferBase should be updated"), Reader->BufferBase, PosStart);
-	TestEqual(TEXT("PrecacheNearEndOfFile - BufferCount should be set Min(BufferSize, Size - Pos)"), Reader->BufferArray.Num(), 1LL);
+	TEST_EQUAL(TEXT("PrecacheNearEndOfFile - Should succeed"), Result, true);
+	TEST_EQUAL(TEXT("PrecacheNearEndOfFile - BufferBase should be updated"), Reader->BufferBase, PosStart);
+	TEST_EQUAL(TEXT("PrecacheNearEndOfFile - BufferCount should be set Min(BufferSize, Size - Pos)"), Reader->BufferArray.Num(), 1LL);
 	TestBytesValid(TEXT("PrecacheNearEndOfFile - BufferBytes should match BufferBase"), Reader->BufferArray.GetData(), Reader->BufferArray.Num(), Reader->BufferBase);
 
 	//AddExpectedError(TEXT("ReadFile failed"));
 	// InternalPrecache at the end of the file should fail
 	SetPosAndBuffer(Reader, FileSize, 0, 16);
 	Result = Reader->InternalPrecache(FileSize, BufferSize);
-	TestEqual(TEXT("PrecacheEndOfFile - Should fail"), Result, false);
-	TestEqual(TEXT("PrecacheEndOfFile - BufferBase should not be updated"), Reader->BufferBase, 0LL);
-	TestEqual(TEXT("PrecacheEndOfFile - BufferCount should not be updated"), Reader->BufferArray.Num(), 16LL);
+	TEST_EQUAL(TEXT("PrecacheEndOfFile - Should fail"), Result, false);
+	TEST_EQUAL(TEXT("PrecacheEndOfFile - BufferBase should not be updated"), Reader->BufferBase, 0LL);
+	TEST_EQUAL(TEXT("PrecacheEndOfFile - BufferCount should not be updated"), Reader->BufferArray.Num(), 16LL);
 
 	// If ReadLowLevel fails, InternalPrecache should return true if the first byte of PrecacheOffset is in the buffer
 	SetPosAndBuffer(Reader, 8, 0, 16);
 	Reader->SeekLowLevel(FileSize);
 	Result = Reader->InternalPrecache(8, 16);
-	TestEqual(TEXT("PrecacheReadFailsBytesRemain - Should succeed if bytes are left in the buffer"), Result, true);
+	TEST_EQUAL(TEXT("PrecacheReadFailsBytesRemain - Should succeed if bytes are left in the buffer"), Result, true);
 	if (FArchiveFileReaderGeneric::bPrecacheAsSoonAsPossible)
 	{
-		TestEqual(TEXT("PrecacheReadFailsBytesRemain - BufferBase should be updated"), Reader->BufferBase, 8LL);
-		TestEqual(TEXT("PrecacheReadFailsBytesRemain - BufferCount should be updated"), Reader->BufferArray.Num(), 8LL);
+		TEST_EQUAL(TEXT("PrecacheReadFailsBytesRemain - BufferBase should be updated"), Reader->BufferBase, 8LL);
+		TEST_EQUAL(TEXT("PrecacheReadFailsBytesRemain - BufferCount should be updated"), Reader->BufferArray.Num(), 8LL);
 	}
 	else
 	{
-		TestEqual(TEXT("PrecacheReadFailsBytesRemain - BufferBase should not be updated"), Reader->BufferBase, 0LL);
-		TestEqual(TEXT("PrecacheReadFailsBytesRemain - BufferCount should not be updated"), Reader->BufferArray.Num(), 16LL);
+		TEST_EQUAL(TEXT("PrecacheReadFailsBytesRemain - BufferBase should not be updated"), Reader->BufferBase, 0LL);
+		TEST_EQUAL(TEXT("PrecacheReadFailsBytesRemain - BufferCount should not be updated"), Reader->BufferArray.Num(), 16LL);
 	}
 
 	// If ReadLowLevel fails, InternalPrecache should return false if the first byte of PrecacheOffset is not in the buffer
 	SetPosAndBuffer(Reader, 16, 0, 16);
 	Reader->SeekLowLevel(FileSize);
 	Result = Reader->InternalPrecache(16, 16);
-	TestEqual(TEXT("PrecacheReadFails - Should fail if bytes are not left in the buffer"), Result, false);
-	TestEqual(TEXT("PrecacheReadFails - BufferCount should be updated to empty"), Reader->BufferArray.Num(), 0LL);
+	TEST_EQUAL(TEXT("PrecacheReadFails - Should fail if bytes are not left in the buffer"), Result, false);
+	TEST_EQUAL(TEXT("PrecacheReadFails - BufferCount should be updated to empty"), Reader->BufferArray.Num(), 0LL);
 
 	// TODO: Call ClearExpectedErrors to clear the AddExpectedError called above when it is implemented. Until then, make sure these failure-testing cases are the last
 	// cases tested in the test
@@ -175,7 +175,7 @@ bool FArchiveFileReaderGenericTest::CreateTestFile()
 	TUniquePtr<IFileHandle> WriteHandle(PlatformFile.OpenWrite(*TestFileName, false /* bAppend */, false /* bAllowRead */));
 	if (!WriteHandle)
 	{
-		TestAddError(FString::Printf(TEXT("Could not create test file %s."), *TestFileName));
+		FAIL_CHECK(FString::Printf(TEXT("Could not create test file %s."), *TestFileName));
 		return false;
 	}
 	TArray<int32> Builder;
@@ -193,12 +193,12 @@ bool FArchiveFileReaderGenericTest::CreateTestFile()
 	ReadHandle.Reset(PlatformFile.OpenRead(*TestFileName, false /* bAllowWrite */));
 	if (!ReadHandle)
 	{
-		TestAddError(FString::Printf(TEXT("Could not open test file %s."), *TestFileName));
+		FAIL_CHECK(FString::Printf(TEXT("Could not open test file %s."), *TestFileName));
 		return false;
 	}
 	if (ReadHandle->Size() != FileSize)
 	{
-		TestAddError(FString::Printf(TEXT("Received incorrect file size from test file %s. Expected = %lld. Actual = %lld."), *TestFileName, FileSize, ReadHandle->Size()));
+		FAIL_CHECK(FString::Printf(TEXT("Received incorrect file size from test file %s. Expected = %lld. Actual = %lld."), *TestFileName, FileSize, ReadHandle->Size()));
 		return false;
 	}
 
@@ -220,7 +220,7 @@ void FArchiveFileReaderGenericTest::TestBytesValid(const TCHAR* What, uint8* InD
 		{
 			if (Data[n] != reinterpret_cast<uint8*>(&ExpectedValue)[StartByte + n])
 			{
-				TestEqual(What, Data[n], reinterpret_cast<uint8*>(&ExpectedValue)[StartByte + n]);
+				TEST_EQUAL(What, Data[n], reinterpret_cast<uint8*>(&ExpectedValue)[StartByte + n]);
 				return;
 			}
 		}
@@ -240,7 +240,7 @@ void FArchiveFileReaderGenericTest::TestBytesValid(const TCHAR* What, uint8* InD
 	{
 		if (IntData[n] != GetExpectedValue(IntOffset + n))
 		{
-			TestEqual(What, IntData[n], GetExpectedValue(IntOffset + n));
+			TEST_EQUAL(What, IntData[n], GetExpectedValue(IntOffset + n));
 			return;
 		}
 	}
@@ -254,7 +254,7 @@ void FArchiveFileReaderGenericTest::TestBytesValid(const TCHAR* What, uint8* InD
 		{
 			if (DataStart[n] != reinterpret_cast<uint8*>(&ExpectedValue)[n])
 			{
-				TestEqual(What, DataStart[n], reinterpret_cast<uint8*>(&ExpectedValue)[n]);
+				TEST_EQUAL(What, DataStart[n], reinterpret_cast<uint8*>(&ExpectedValue)[n]);
 				return;
 			}
 		}
@@ -271,7 +271,7 @@ void FArchiveFileReaderGenericTest::SetPosAndBuffer(TUniquePtr<FArchiveFileReade
 		Reader->SeekLowLevel(BufferBase);
 		int64 BytesRead;
 		Reader->ReadLowLevel(Reader->BufferArray.GetData(), BufferSize, BytesRead);
-		TestEqual(TEXT("SetPosAndBuffer - ReadLowLevel read the requested bytes"), BytesRead, BufferSize);
+		TEST_EQUAL(TEXT("SetPosAndBuffer - ReadLowLevel read the requested bytes"), BytesRead, BufferSize);
 	}
 	// See the contract for Buffer window and LowLevel Pos in the variable comment on BufferArray
 	bool bPosWithinBuffer = BufferBase <= Pos && Pos < BufferBase + BufferSize;
@@ -279,9 +279,9 @@ void FArchiveFileReaderGenericTest::SetPosAndBuffer(TUniquePtr<FArchiveFileReade
 	{
 		Reader->SeekLowLevel(Pos);
 	}
-	TestEqual(TEXT("SetPosAndBuffer - Pos set correctly"), Reader->Pos, Pos);
-	TestEqual(TEXT("SetPosAndBuffer - BufferBase set correctly"), Reader->BufferBase, BufferBase);
-	TestEqual(TEXT("SetPosAndBuffer - BufferSize set correctly"), Reader->BufferArray.Num(), BufferSize);
+	TEST_EQUAL(TEXT("SetPosAndBuffer - Pos set correctly"), Reader->Pos, Pos);
+	TEST_EQUAL(TEXT("SetPosAndBuffer - BufferBase set correctly"), Reader->BufferBase, BufferBase);
+	TEST_EQUAL(TEXT("SetPosAndBuffer - BufferSize set correctly"), Reader->BufferArray.Num(), BufferSize);
 	TestBytesValid(TEXT("SetPosAndBuffer - BufferBytes should match BufferBase"), Reader->BufferArray.GetData(), Reader->BufferArray.Num(), Reader->BufferBase);
 }
 
@@ -290,7 +290,7 @@ int32 FArchiveFileReaderGenericTest::GetExpectedValue(int64 IntOffset)
 	return 0xbe000000 + static_cast<int32>(IntOffset);
 }
 
-TEST_CASE("Core::HAL::FArchiveFileReaderGeneric::Smoke Test", "[Core][HAL][Smoke]")
+TEST_CASE_METHOD(FAutomationTestFixture, "Core::HAL::FArchiveFileReaderGeneric::Smoke Test", "[Core][HAL][Smoke]")
 {
 	FArchiveFileReaderGenericTest ArchiveFileReaderGenericTest;
 	ArchiveFileReaderGenericTest.TestInternalPrecache();
diff --git a/Engine/Source/Runtime/Core/Tests/HAL/FileManagerTest.cpp b/Engine/Source/Runtime/Core/Tests/HAL/FileManagerTest.cpp
--- a/Engine/Source/Runtime/Core/Tests/HAL/FileManagerTest.cpp
+++ b/Engine/Source/Runtime/Core/Tests/HAL/FileManagerTest.cpp
@@ -2,9 +2,15 @@
 
 #include "HAL/FileManager.h"
 
-#include "TestHarness.h"
+#include "Misc/AutomationTest.h"
 
-TEST_CASE("Core::HAL::FileManager::Reader", "[Core][HAL][FileManager]")
+#if WITH_DEV_AUTOMATION_TESTS
+
+IMPLEMENT_SIMPLE_AUTOMATION_TEST(FFileManagerReaderTest, "System.Core.HAL.FileManager.Reader", EAutomationTestFlags::ApplicationContextMask | EAutomationTestFlags::EngineFilter)
+
+bool FFileManagerReaderTest::RunTest(const FString& Parameters)
 {
-	// TODO: Add tests
-}
\ No newline at end of file
+	return true;
+}
+
+#endif
\ No newline at end of file
diff --git a/Engine/Source/Runtime/Core/Tests/HAL/PlatformAtomicsTest.cpp b/Engine/Source/Runtime/Core/Tests/HAL/PlatformAtomicsTest.cpp
--- a/Engine/Source/Runtime/Core/Tests/HAL/PlatformAtomicsTest.cpp
+++ b/Engine/Source/Runtime/Core/Tests/HAL/PlatformAtomicsTest.cpp
@@ -10,21 +10,24 @@
 #endif
 #include "HAL/PlatformAtomics.h"
 
-template<typename T>
-static bool TestInterlocked(volatile T* Dest, T ExpectedReturnValue, T ExpectedFinalValue, TFunctionRef<T(volatile T*)> InterlockedFunc, const TCHAR* FunctionName, const TCHAR* TypeName)
+class FAtomicsTest : public FAutomationTestFixture
 {
-	const T ReturnValue = InterlockedFunc(Dest);
 
-	bool bReturnValueSuccess = (ReturnValue == ExpectedReturnValue);
-	bool bAtomicReadSuccess = (FPlatformAtomics::AtomicRead(Dest) == ExpectedFinalValue);
+public:
 
-	TestTrue(FString::Printf(TEXT("FPlatformAtomics::Interlocked%s on %s failed"), FunctionName, TypeName), bReturnValueSuccess);
-	TestTrue(FString::Printf(TEXT("FPlatformAtomics::Interlocked%s on %s failed"), FunctionName, TypeName), bAtomicReadSuccess);
+template<typename T>
+bool TestInterlocked(volatile T* Dest, T ExpectedReturnValue, T ExpectedFinalValue, TFunctionRef<T(volatile T*)> InterlockedFunc, const TCHAR* FunctionName, const TCHAR* TypeName)
+{
+	bool bSuccess = true;
+	const T ReturnValue = InterlockedFunc(Dest);
+	
+	TEST_TRUE(FString::Printf(TEXT("FPlatformAtomics::Interlocked%s on %s failed"), FunctionName, TypeName), ReturnValue == ExpectedReturnValue);
+	TEST_TRUE(FString::Printf(TEXT("FPlatformAtomics::Interlocked%s on %s failed"), FunctionName, TypeName), FPlatformAtomics::AtomicRead(Dest) == ExpectedFinalValue );
 
-	return bReturnValueSuccess & bAtomicReadSuccess;
+	return bSuccess;
 }
 
-static bool TestInterlockedAnd()
+bool TestInterlockedAnd()
 {
 	bool bSuccess = true;
 
@@ -88,7 +91,7 @@ static bool TestInterlockedAnd()
 	return bSuccess;
 }
 
-static bool TestInterlockedOr()
+bool TestInterlockedOr()
 {
 	bool bSuccess = true;
 
@@ -152,7 +155,7 @@ static bool TestInterlockedOr()
 	return bSuccess;
 }
 
-static bool TestInterlockedXor()
+bool TestInterlockedXor()
 {
 	bool bSuccess = true;
 
@@ -216,7 +219,7 @@ static bool TestInterlockedXor()
 	return bSuccess;
 }
 
-static bool TestInterlockedAdd()
+bool TestInterlockedAdd()
 {
 	bool bSuccess = true;
 
@@ -296,7 +299,7 @@ static bool TestInterlockedAdd()
 	return bSuccess;
 }
 
-static bool TestInterlockedIncrement()
+bool TestInterlockedIncrement()
 {
 	bool bSuccess = true;
 
@@ -340,7 +343,7 @@ static bool TestInterlockedIncrement()
 	return bSuccess;
 }
 
-static bool TestInterlockedDecrement()
+bool TestInterlockedDecrement()
 {
 	bool bSuccess = true;
 
@@ -384,7 +387,7 @@ static bool TestInterlockedDecrement()
 	return bSuccess;
 }
 
-static bool TestInterlockedExchange()
+bool TestInterlockedExchange()
 {
 	bool bSuccess = true;
 
@@ -412,7 +415,7 @@ static bool TestInterlockedExchange()
 	return bSuccess;
 }
 
-static bool TestInterlockedExchangePtr()
+bool TestInterlockedExchangePtr()
 {
 	int Test;
 
@@ -420,20 +423,20 @@ static bool TestInterlockedExchangePtr()
 	const void* ReturnValue = FPlatformAtomics::InterlockedExchangePtr(&Value, nullptr);
 	if (ReturnValue != &Test)
 	{
-		TestAddError(TEXT("FPlatformAtomics::InterlockedExchangePtr failed"));
+		FAIL_CHECK(TEXT("FPlatformAtomics::InterlockedExchangePtr failed"));
 		return false;
 	}
 
 	if (Value != nullptr)
 	{
-		TestAddError(TEXT("FPlatformAtomics::InterlockedExchangePtr failed"));
+		FAIL_CHECK(TEXT("FPlatformAtomics::InterlockedExchangePtr failed"));
 		return false;
 	}
 
 	return true;
 }
 
-static bool TestInterlockedCompareExchange()
+bool TestInterlockedCompareExchange()
 {
 	bool bSuccess = true;
 
@@ -480,8 +483,9 @@ static bool TestInterlockedCompareExchange()
 
 	return bSuccess;
 }
+};
 
-TEST_CASE("Core::HAL::PlatformAtomics::Smoke Test", "[Core][HAL][Smoke]")
+TEST_CASE_METHOD(FAtomicsTest, "Core::HAL::PlatformAtomics::Smoke Test", "[Core][HAL][Smoke]")
 {
 	TestInterlockedAnd();
 	TestInterlockedOr();
diff --git a/Engine/Source/Runtime/Core/Tests/HAL/PlatformStringTest.cpp b/Engine/Source/Runtime/Core/Tests/HAL/PlatformStringTest.cpp
--- a/Engine/Source/Runtime/Core/Tests/HAL/PlatformStringTest.cpp
+++ b/Engine/Source/Runtime/Core/Tests/HAL/PlatformStringTest.cpp
@@ -12,31 +12,31 @@ static void InvokePlatformStringGetVarArgs(CharType (&Dest)[Size], const CharTyp
 	va_end(ap);
 }
 
-TEST_CASE("Core::HAL::PlatformString::GetVarArgs", "[Core][HAL][Smoke]")
+TEST_CASE_METHOD(FAutomationTestFixture, "Core::HAL::PlatformString::GetVarArgs", "[Core][HAL][Smoke]")
 {
 	TCHAR Buffer[128];
 	InvokePlatformStringGetVarArgs(Buffer, TEXT("A%.*sZ"), 4, TEXT(" to B"));
-	TestEqualString(TEXT("GetVarArgs(%.*s)"), Buffer, TEXT("A to Z"));
+	TEST_TRUE(TEXT("GetVarArgs(%.*s)"), !FCString::Strcmp(Buffer, TEXT("A to Z") ));
 }
 
-TEST_CASE("Core::HAL::PlatformString::Strnlen", "[Core][HAL][Smoke]")
+TEST_CASE_METHOD(FAutomationTestFixture, "Core::HAL::PlatformString::Strnlen", "[Core][HAL][Smoke]")
 {
-	TestEqual(TEXT("Strnlen(nullptr, 0)"), FPlatformString::Strnlen((const ANSICHAR*)nullptr, 0), 0);
-	TestEqual(TEXT("Strnlen(\"\", 0)"), FPlatformString::Strnlen("", 0), 0);
-	TestEqual(TEXT("Strnlen(\"1\", 0)"), FPlatformString::Strnlen("1", 0), 0);
-	TestEqual(TEXT("Strnlen(\"1\", 1)"), FPlatformString::Strnlen("1", 1), 1);
-	TestEqual(TEXT("Strnlen(\"1\", 2)"), FPlatformString::Strnlen("1", 2), 1);
-	TestEqual(TEXT("Strnlen(\"123\", 2)"), FPlatformString::Strnlen("123", 2), 2);
+	TEST_EQUAL(TEXT("Strnlen(nullptr, 0)"), FPlatformString::Strnlen((const ANSICHAR*)nullptr, 0), 0);
+	TEST_EQUAL(TEXT("Strnlen(\"\", 0)"), FPlatformString::Strnlen("", 0), 0);
+	TEST_EQUAL(TEXT("Strnlen(\"1\", 0)"), FPlatformString::Strnlen("1", 0), 0);
+	TEST_EQUAL(TEXT("Strnlen(\"1\", 1)"), FPlatformString::Strnlen("1", 1), 1);
+	TEST_EQUAL(TEXT("Strnlen(\"1\", 2)"), FPlatformString::Strnlen("1", 2), 1);
+	TEST_EQUAL(TEXT("Strnlen(\"123\", 2)"), FPlatformString::Strnlen("123", 2), 2);
 	ANSICHAR AnsiBuffer[128] = "123456789";
-	TestEqual(TEXT("Strnlen(PaddedBuffer)"), FPlatformString::Strnlen(AnsiBuffer, UE_ARRAY_COUNT(AnsiBuffer)), 9);
+	TEST_EQUAL(TEXT("Strnlen(PaddedBuffer)"), FPlatformString::Strnlen(AnsiBuffer, UE_ARRAY_COUNT(AnsiBuffer)), 9);
 
-	TestEqual(TEXT("Strnlen(nullptr, 0)"), FPlatformString::Strnlen((const TCHAR*)nullptr, 0), 0); //-V575
-	TestEqual(TEXT("Strnlen(\"\", 0)"), FPlatformString::Strnlen(TEXT(""), 0), 0); //-V575
-	TestEqual(TEXT("Strnlen(\"1\", 0)"), FPlatformString::Strnlen(TEXT("1"), 0), 0); //-V575
-	TestEqual(TEXT("Strnlen(\"1\", 1)"), FPlatformString::Strnlen(TEXT("1"), 1), 1);
-	TestEqual(TEXT("Strnlen(\"1\", 2)"), FPlatformString::Strnlen(TEXT("1"), 2), 1);
-	TestEqual(TEXT("Strnlen(\"123\", 2)"), FPlatformString::Strnlen(TEXT("123"), 2), 2);
+	TEST_EQUAL(TEXT("Strnlen(nullptr, 0)"), FPlatformString::Strnlen((const TCHAR*)nullptr, 0), 0); //-V575
+	TEST_EQUAL(TEXT("Strnlen(\"\", 0)"), FPlatformString::Strnlen(TEXT(""), 0), 0); //-V575
+	TEST_EQUAL(TEXT("Strnlen(\"1\", 0)"), FPlatformString::Strnlen(TEXT("1"), 0), 0); //-V575
+	TEST_EQUAL(TEXT("Strnlen(\"1\", 1)"), FPlatformString::Strnlen(TEXT("1"), 1), 1);
+	TEST_EQUAL(TEXT("Strnlen(\"1\", 2)"), FPlatformString::Strnlen(TEXT("1"), 2), 1);
+	TEST_EQUAL(TEXT("Strnlen(\"123\", 2)"), FPlatformString::Strnlen(TEXT("123"), 2), 2);
 	TCHAR Buffer[128] = {};
 	FCString::Strcpy(Buffer, TEXT("123456789"));
-	TestEqual(TEXT("Strnlen(PaddedBuffer)"), FPlatformString::Strnlen(Buffer, UE_ARRAY_COUNT(Buffer)), 9);
+	TEST_EQUAL(TEXT("Strnlen(PaddedBuffer)"), FPlatformString::Strnlen(Buffer, UE_ARRAY_COUNT(Buffer)), 9);
 }
