Add new Unit Tests

diff --git a/Engine/Source/Runtime/Core/Tests/HAL/ThreadTest.cpp b/Engine/Source/Runtime/Core/Tests/HAL/ThreadTest.cpp
--- a/Engine/Source/Runtime/Core/Tests/HAL/ThreadTest.cpp
+++ b/Engine/Source/Runtime/Core/Tests/HAL/ThreadTest.cpp
@@ -2,7 +2,6 @@
 
 #include "CoreTypes.h"
 #include "Misc/AssertionMacros.h"
-
 #include "HAL/Thread.h"
 #include "HAL/ThreadSafeBool.h"
 #include "HAL/ThreadSafeCounter.h"
@@ -10,39 +9,38 @@
 #include "HAL/Event.h"
 #include "Containers/Queue.h"
 #include "Containers/StringConv.h"
-
-#include "TestHarness.h"
+#include "TestFixtures/CoreTestFixture.h"
 
 namespace
 {
-	void TestIsJoinableAfterCreation()
+	void TestIsJoinableAfterCreation(FAutomationTestFixture& Test)
 	{
 		FThread Thread(TEXT("Test.Thread.TestIsJoinableAfterCreation"), []() { /*NOOP*/ });
-		TestTrue(TEXT("FThread must be joinable after construction"), Thread.IsJoinable());
+		TEST_TRUE(TEXT("FThread must be joinable after construction"), Thread.IsJoinable());
 		Thread.Join();
 		UE_LOG(LogTemp, Log, TEXT("%s completed"), StringCast<TCHAR>(__FUNCTION__).Get());
 	}
 
-	void TestIsJoinableAfterCompletion()
+	void TestIsJoinableAfterCompletion(FAutomationTestFixture& Test)
 	{
 		FThreadSafeBool bDone = false;
 		FThread Thread(TEXT("Test.Thread.TestIsJoinableAfterCompletion"), [&bDone]() { bDone = true; });
-		while (!bDone); // wait for completion //-V529
-		TestTrue(TEXT("FThread must still be joinable after completion"), Thread.IsJoinable());
+		while (!bDone) {}; // wait for completion //-V529
+		TEST_TRUE(TEXT("FThread must still be joinable after completion"), Thread.IsJoinable());
 		Thread.Join();
 		UE_LOG(LogTemp, Log, TEXT("%s completed"), StringCast<TCHAR>(__FUNCTION__).Get());
 	}
 
-	void TestIsNotJoinableAfterJoining()
+	void TestIsNotJoinableAfterJoining(FAutomationTestFixture& Test)
 	{
 		FThread Thread(TEXT("Test.Thread.TestIsNotJoinableAfterJoining"), []() { /*NOOP*/ });
 		Thread.Join();
-		TestFalse(TEXT("FThread must not be joinable after joining"), Thread.IsJoinable());
+		TEST_FALSE(TEXT("FThread must not be joinable after joining"), Thread.IsJoinable());
 		UE_LOG(LogTemp, Log, TEXT("%s completed"), StringCast<TCHAR>(__FUNCTION__).Get());
 	}
 
 #if 0 // detaching is not implemented
-	void TestIsNotJoinableAfterDetaching()
+	void TestIsNotJoinableAfterDetaching(FAutomationTestFixture& Test)
 	{
 		// two cases: it's either the calling thread detaches from the thread before the thread is completed
 		{
@@ -53,7 +51,7 @@ namespace
 				});
 			Thread.Detach();
 			bReady = true; // make sure `Detach` is called before thread function exit
-			TestFalse(TEXT("FThread must not be joinable after detaching"), Thread.IsJoinable());
+			Test.TEST_FALSE(TEXT("FThread must not be joinable after detaching"), Thread.IsJoinable());
 		}
 		// or thread function is completed fast and `FThreadImpl` releases the reference to itself before
 		// `Detach` call
@@ -63,35 +61,35 @@ namespace
 			FPlatformProcess::Sleep(0.1); // let the thread exit before detaching
 			Thread.Detach();
 			bReady = true; // make sure `Detach` is called before thread function exit
-			TestFalse(TEXT("FThread must not be joinable after detaching"), Thread.IsJoinable());
+			Test.TEST_FALSE(TEXT("FThread must not be joinable after detaching"), Thread.IsJoinable());
 		}
 		UE_LOG(LogTemp, Log, TEXT("%s completed"), StringCast<TCHAR>(__FUNCTION__).Get());
 	}
 #endif
 
-	void TestAssertIfNotJoinedOrDetached()
+	void TestAssertIfNotJoinedOrDetached(FAutomationTestFixture& Test)
 	{
 		// this does fails the `check`, but it seems there's no way to test this by UE4 Automation Testing, so commented out
 		FThread Thread(TEXT("Test.Thread.TestAssertIfNotJoinedOrDetached"), []() { /*NOOP*/ });
 		// should assert in the destructor
 	}
 
-	void TestDefaultConstruction()
+	void TestDefaultConstruction(FAutomationTestFixture& Test)
 	{
 		{
 			FThread Thread;
-			TestFalse(TEXT("Default-constructed FThread must be not joinable"), Thread.IsJoinable());
+			TEST_FALSE(TEXT("Default-constructed FThread must be not joinable"), Thread.IsJoinable());
 		}
 		{	// check that default constructed thread can be "upgraded" to joinable thread
 			FThread Thread;
 			Thread = FThread(TEXT("Test.Thread.TestDefaultConstruction"), []() { /* NOOP */ });
-			TestTrue(TEXT("Move-constructed FThread from joinable thread must be joinable"), Thread.IsJoinable());
+			TEST_TRUE(TEXT("Move-constructed FThread from joinable thread must be joinable"), Thread.IsJoinable());
 			Thread.Join();
 		}
 		UE_LOG(LogTemp, Log, TEXT("%s completed"), StringCast<TCHAR>(__FUNCTION__).Get());
 	}
 
-	void TestThreadSingleton()
+	void TestThreadSingleton(FAutomationTestFixture& Test)
 	{
 		{	// check that ThreadSingleton instances in different threads are isolated from each other
 			class FThreadSingletonTest : public TThreadSingleton<FThreadSingletonTest>
@@ -113,8 +111,8 @@ namespace
 					FThreadSingletonTest::Get().SetTestField(2);
 				});
 			Thread.Join();
-			TestTrue(TEXT("Thread singleton should be uninitialized by default"), bDefaultValuePass);
-			TestTrue(TEXT("Thread singletons in different threads should be isolated"), FThreadSingletonTest::Get().GetTestField() == 1);
+			TEST_TRUE(TEXT("Thread singleton should be uninitialized by default"), bDefaultValuePass);
+			TEST_TRUE(TEXT("Thread singletons in different threads should be isolated"), FThreadSingletonTest::Get().GetTestField() == 1);
 		}
 		{	// check that ThreadSingleton entries don't point to invalid memory after cleanup
 			class FThreadSingletonFirst : public TThreadSingleton<FThreadSingletonFirst>
@@ -137,33 +135,33 @@ namespace
 					FThreadSingletonSecond::Get();
 				});
 			Thread.Join();
-			TestTrue(TEXT("Thread singletons should be uninitialized by default"), FThreadSingletonFirst::TryGet() == nullptr);
-			TestTrue(TEXT("Thread singletons should be uninitialized by default"), FThreadSingletonSecond::TryGet() == nullptr);
+			TEST_TRUE(TEXT("Thread singletons should be uninitialized by default"), FThreadSingletonFirst::TryGet() == nullptr);
+			TEST_TRUE(TEXT("Thread singletons should be uninitialized by default"), FThreadSingletonSecond::TryGet() == nullptr);
 		}
 		UE_LOG(LogTemp, Log, TEXT("%s completed"), StringCast<TCHAR>(__FUNCTION__).Get());
 	}
 
-	void TestMovability()
+	void TestMovability(FAutomationTestFixture& Test)
 	{
 		{	// move constructor with default-constructed thread
 			FThread Src;
 			FThread Dst(MoveTemp(Src));
-			TestFalse(TEXT("Default-constructed thread must stay not joinable after moving out"), Src.IsJoinable());
-			TestFalse(TEXT("Move-constructed thread from not joinable thread must be not joinable"), Dst.IsJoinable());
+			TEST_FALSE(TEXT("Default-constructed thread must stay not joinable after moving out"), Src.IsJoinable());
+			TEST_FALSE(TEXT("Move-constructed thread from not joinable thread must be not joinable"), Dst.IsJoinable());
 		}
 		{	// move constructor with joinable thread
 			FThread Src(TEXT("Test.Thread.TestMovability.1"), []() { /* NOOP */ });
 			FThread Dst(MoveTemp(Src));
-			TestFalse(TEXT("Moved out thread must be not joinable"), Src.IsJoinable());
-			TestTrue(TEXT("Move-constructed thread from joinable thread must be joinable"), Dst.IsJoinable());
+			TEST_FALSE(TEXT("Moved out thread must be not joinable"), Src.IsJoinable());
+			TEST_TRUE(TEXT("Move-constructed thread from joinable thread must be joinable"), Dst.IsJoinable());
 			Dst.Join();
 		}
 		{	// move assignment operator
 			FThread Src(TEXT("Test.Thread.TestMovability.2"), []() { /* NOOP */ });
 			FThread Dst;
 			Dst = MoveTemp(Src);
-			TestFalse(TEXT("Moved out thread must be not joinable"), Src.IsJoinable());
-			TestTrue(TEXT("Move-assigned thread from joinable thread must be joinable"), Dst.IsJoinable());
+			TEST_FALSE(TEXT("Moved out thread must be not joinable"), Src.IsJoinable());
+			TEST_TRUE(TEXT("Move-assigned thread from joinable thread must be joinable"), Dst.IsJoinable());
 			Dst.Join();
 		}
 		{	// Failure test for move assignment operator of joinable thread
@@ -183,7 +181,7 @@ namespace
 	}
 
 	// An example of possible implementation of Consumer/Producer idiom
-	void TestTypicalUseCase()
+	void TestTypicalUseCase(FAutomationTestFixture& Test)
 	{
 		FThreadSafeBool bQuitRequested = false;
 		using FWork = uint32;
@@ -242,23 +240,22 @@ namespace
 	}
 }
 
-TEST_CASE("Core::HAL::Thread", "[Core][HAL][Thread]")
+TEST_CASE_METHOD(FCoreTestFixture, "Core::HAL::Thread::Smoke Test", "[Core][HAL][Smoke]")
 {
 	UE_LOG(LogTemp, Log, TEXT("%s"), StringCast<TCHAR>(__FUNCTION__).Get());
 
-	TestIsJoinableAfterCreation();
-	TestIsJoinableAfterCompletion();
-	TestIsNotJoinableAfterJoining();
+	TestIsJoinableAfterCreation(*this);
+	TestIsJoinableAfterCompletion(*this);
+	TestIsNotJoinableAfterJoining(*this);
 	
 #if 0 // detaching is not implemented
-	TestIsNotJoinableAfterDetaching();
+	TestIsNotJoinableAfterDetaching(*this);
 #endif
 
 	//TestAssertIfNotJoinedOrDetached();
 
-	TestDefaultConstruction();
-	TestMovability();
-
-	TestTypicalUseCase();
-	TestThreadSingleton();
-}
\ No newline at end of file
+	TestDefaultConstruction(*this);
+	TestMovability(*this);
+	TestTypicalUseCase(*this);
+	TestThreadSingleton(*this);
+}
diff --git a/Engine/Source/Runtime/Core/Tests/Hash/Blake3Test.cpp b/Engine/Source/Runtime/Core/Tests/Hash/Blake3Test.cpp
--- a/Engine/Source/Runtime/Core/Tests/Hash/Blake3Test.cpp
+++ b/Engine/Source/Runtime/Core/Tests/Hash/Blake3Test.cpp
@@ -1,45 +1,43 @@
 // Copyright Epic Games, Inc. All Rights Reserved.
 
 #include "Hash/Blake3.h"
-
 #include "Misc/StringBuilder.h"
 #include "TestHarness.h"
 
-TEST_CASE("Core::Hash::Blake3", "[Core][Hash][Smoke]")
+TEST_CASE_METHOD(FAutomationTestFixture, "Core::Hash::FBlake3::Smoke Test", "[Core][Math][Smoke]")
 {
+	constexpr uint32 MaxInputLength = 31'744u;
+
 	struct FBlake3TestCase
 	{
 		uint32 InputLength;
 		FStringView Hash;
 	};
-
-	constexpr FBlake3TestCase TestCases[] =
+	const FBlake3TestCase TestCases[] =
 	{
-		{     0u, TEXTVIEW("af1349b9f5f9a1a6a0404dea36dcc9499bcb25c9adc112b7cc9a93cae41f3262")},
-		{     1u, TEXTVIEW("2d3adedff11b61f14c886e35afa036736dcd87a74d27b5c1510225d0f592e213")},
-		{ 1'023u, TEXTVIEW("10108970eeda3eb932baac1428c7a2163b0e924c9a9e25b35bba72b28f70bd11")},
-		{ 1'024u, TEXTVIEW("42214739f095a406f3fc83deb889744ac00df831c10daa55189b5d121c855af7")},
-		{ 1'025u, TEXTVIEW("d00278ae47eb27b34faecf67b4fe263f82d5412916c1ffd97c8cb7fb814b8444")},
-		{ 2'048u, TEXTVIEW("e776b6028c7cd22a4d0ba182a8bf62205d2ef576467e838ed6f2529b85fba24a")},
-		{ 2'049u, TEXTVIEW("5f4d72f40d7a5f82b15ca2b2e44b1de3c2ef86c426c95c1af0b6879522563030")},
-		{ 3'072u, TEXTVIEW("b98cb0ff3623be03326b373de6b9095218513e64f1ee2edd2525c7ad1e5cffd2")},
-		{ 3'073u, TEXTVIEW("7124b49501012f81cc7f11ca069ec9226cecb8a2c850cfe644e327d22d3e1cd3")},
-		{ 4'096u, TEXTVIEW("015094013f57a5277b59d8475c0501042c0b642e531b0a1c8f58d2163229e969")},
-		{ 4'097u, TEXTVIEW("9b4052b38f1c5fc8b1f9ff7ac7b27cd242487b3d890d15c96a1c25b8aa0fb995")},
-		{ 5'120u, TEXTVIEW("9cadc15fed8b5d854562b26a9536d9707cadeda9b143978f319ab34230535833")},
-		{ 5'121u, TEXTVIEW("628bd2cb2004694adaab7bbd778a25df25c47b9d4155a55f8fbd79f2fe154cff")},
-		{ 6'144u, TEXTVIEW("3e2e5b74e048f3add6d21faab3f83aa44d3b2278afb83b80b3c35164ebeca205")},
-		{ 6'145u, TEXTVIEW("f1323a8631446cc50536a9f705ee5cb619424d46887f3c376c695b70e0f0507f")},
-		{ 7'168u, TEXTVIEW("61da957ec2499a95d6b8023e2b0e604ec7f6b50e80a9678b89d2628e99ada77a")},
-		{ 7'169u, TEXTVIEW("a003fc7a51754a9b3c7fae0367ab3d782dccf28855a03d435f8cfe74605e7817")},
-		{ 8'192u, TEXTVIEW("aae792484c8efe4f19e2ca7d371d8c467ffb10748d8a5a1ae579948f718a2a63")},
-		{ 8'193u, TEXTVIEW("bab6c09cb8ce8cf459261398d2e7aef35700bf488116ceb94a36d0f5f1b7bc3b")},
-		{16'384u, TEXTVIEW("f875d6646de28985646f34ee13be9a576fd515f76b5b0a26bb324735041ddde4")},
-		{31'744u, TEXTVIEW("62b6960e1a44bcc1eb1a611a8d6235b6b4b78f32e7abc4fb4c6cdcce94895c47")},
+		{     0u,        TEXTVIEW("af1349b9f5f9a1a6a0404dea36dcc9499bcb25c9adc112b7cc9a93cae41f3262")},
+		{     1u,        TEXTVIEW("2d3adedff11b61f14c886e35afa036736dcd87a74d27b5c1510225d0f592e213")},
+		{ 1'023u,        TEXTVIEW("10108970eeda3eb932baac1428c7a2163b0e924c9a9e25b35bba72b28f70bd11")},
+		{ 1'024u,        TEXTVIEW("42214739f095a406f3fc83deb889744ac00df831c10daa55189b5d121c855af7")},
+		{ 1'025u,        TEXTVIEW("d00278ae47eb27b34faecf67b4fe263f82d5412916c1ffd97c8cb7fb814b8444")},
+		{ 2'048u,        TEXTVIEW("e776b6028c7cd22a4d0ba182a8bf62205d2ef576467e838ed6f2529b85fba24a")},
+		{ 2'049u,        TEXTVIEW("5f4d72f40d7a5f82b15ca2b2e44b1de3c2ef86c426c95c1af0b6879522563030")},
+		{ 3'072u,        TEXTVIEW("b98cb0ff3623be03326b373de6b9095218513e64f1ee2edd2525c7ad1e5cffd2")},
+		{ 3'073u,        TEXTVIEW("7124b49501012f81cc7f11ca069ec9226cecb8a2c850cfe644e327d22d3e1cd3")},
+		{ 4'096u,        TEXTVIEW("015094013f57a5277b59d8475c0501042c0b642e531b0a1c8f58d2163229e969")},
+		{ 4'097u,        TEXTVIEW("9b4052b38f1c5fc8b1f9ff7ac7b27cd242487b3d890d15c96a1c25b8aa0fb995")},
+		{ 5'120u,        TEXTVIEW("9cadc15fed8b5d854562b26a9536d9707cadeda9b143978f319ab34230535833")},
+		{ 5'121u,        TEXTVIEW("628bd2cb2004694adaab7bbd778a25df25c47b9d4155a55f8fbd79f2fe154cff")},
+		{ 6'144u,        TEXTVIEW("3e2e5b74e048f3add6d21faab3f83aa44d3b2278afb83b80b3c35164ebeca205")},
+		{ 6'145u,        TEXTVIEW("f1323a8631446cc50536a9f705ee5cb619424d46887f3c376c695b70e0f0507f")},
+		{ 7'168u,        TEXTVIEW("61da957ec2499a95d6b8023e2b0e604ec7f6b50e80a9678b89d2628e99ada77a")},
+		{ 7'169u,        TEXTVIEW("a003fc7a51754a9b3c7fae0367ab3d782dccf28855a03d435f8cfe74605e7817")},
+		{ 8'192u,        TEXTVIEW("aae792484c8efe4f19e2ca7d371d8c467ffb10748d8a5a1ae579948f718a2a63")},
+		{ 8'193u,        TEXTVIEW("bab6c09cb8ce8cf459261398d2e7aef35700bf488116ceb94a36d0f5f1b7bc3b")},
+		{16'384u,        TEXTVIEW("f875d6646de28985646f34ee13be9a576fd515f76b5b0a26bb324735041ddde4")},
+		{MaxInputLength, TEXTVIEW("62b6960e1a44bcc1eb1a611a8d6235b6b4b78f32e7abc4fb4c6cdcce94895c47")},
 	};
 
-	constexpr uint32 MaxInputLength = 31'744u;
-
 	uint8 Input[MaxInputLength];
 	for (uint32 Index = 0; Index < MaxInputLength; ++Index)
 	{
@@ -51,11 +49,11 @@ TEST_CASE("Core::Hash::Blake3", "[Core][Hash][Smoke]")
 	for (const FBlake3TestCase& TestCase : TestCases)
 	{
 		REQUIRE(InputIndex <= TestCase.InputLength);
-		REQUIRE(TestCase.InputLength <= MaxInputLength);
 		Hash.Update(Input + InputIndex, TestCase.InputLength - InputIndex);
 		InputIndex = TestCase.InputLength;
 		TStringBuilder<70> HashHex;
 		HashHex << Hash.Finalize();
-		CHECK(HashHex.ToView() == TestCase.Hash);
+		INFO(FString::Printf(TEXT("BLAKE3(%u)"), TestCase.InputLength));
+		CHECK(FStringView(HashHex)==TestCase.Hash);
 	}
 }
diff --git a/Engine/Source/Runtime/Core/Tests/IO/IoStatusOrTest.cpp b/Engine/Source/Runtime/Core/Tests/IO/IoStatusOrTest.cpp
--- a/Engine/Source/Runtime/Core/Tests/IO/IoStatusOrTest.cpp
+++ b/Engine/Source/Runtime/Core/Tests/IO/IoStatusOrTest.cpp
@@ -28,59 +28,59 @@ struct FIoStatusTestType
 	FString Text;
 };
 
-void TestConstruct()
+void TestConstruct(FAutomationTestFixture& Test)
 {
 	{
 		TIoStatusOr<FIoStatusTestType> Result;
-		TestEqual("Default IoStatus is Unknown", Result.Status(), FIoStatus::Unknown);
+		TEST_EQUAL("Default IoStatus is Unknown", Result.Status(), FIoStatus::Unknown);
 	}
 
 	{
 		const TIoStatusOr<FIoStatusTestType> Other;
 		TIoStatusOr<FIoStatusTestType> Result(Other);
-		TestEqual("Copy construct", Result.Status(), FIoStatus::Unknown);
+		TEST_EQUAL("Copy construct", Result.Status(), FIoStatus::Unknown);
 	}
 
 	{
 		const FIoStatus IoStatus(EIoErrorCode::InvalidCode);
 		TIoStatusOr<FIoStatusTestType> Result(IoStatus);
-		TestEqual("Construct with status", Result.Status().GetErrorCode(), EIoErrorCode::InvalidCode);
+		TEST_EQUAL("Construct with status", Result.Status().GetErrorCode(), EIoErrorCode::InvalidCode);
 	}
 
 	{
 		const FString ExpectedText("Unreal");
 		const FIoStatusTestType Type(ExpectedText);
 		TIoStatusOr<FIoStatusTestType> Result(Type);
-		TestEqual("Construct with value", Result.ValueOrDie().Text, ExpectedText);
+		TEST_EQUAL("Construct with value", Result.ValueOrDie().Text, ExpectedText);
 	}
 
 	{
 		const FString ExpectedText("Unreal");
 		TIoStatusOr<FIoStatusTestType> Result(FIoStatusTestType("Unreal"));
-		TestEqual("Construct with temporary value", Result.ValueOrDie().Text, ExpectedText);
+		TEST_EQUAL("Construct with temporary value", Result.ValueOrDie().Text, ExpectedText);
 	}
 
 	{
 		TIoStatusOr<FIoStatusTestType> Result(FString("Unreal"));
-		TestEqual("Construct with value arguments", Result.ValueOrDie().Text, FString("Unreal"));
+		TEST_EQUAL("Construct with value arguments", Result.ValueOrDie().Text, FString("Unreal"));
 	}
 }
 
-void TestAssignment()
+void TestAssignment(FAutomationTestFixture& Test)
 {
 	{
 		const EIoErrorCode ExpectedErrorCode = EIoErrorCode::InvalidCode;
 		TIoStatusOr<FIoStatusTestType> Other = FIoStatus(ExpectedErrorCode);
 		TIoStatusOr<FIoStatusTestType> Result;
 		Result = Other;
-		TestEqual("Assign IoStatusOr with status", Result.Status().GetErrorCode(), ExpectedErrorCode);
+		TEST_EQUAL("Assign IoStatusOr with status", Result.Status().GetErrorCode(), ExpectedErrorCode);
 	}
 
 	{
 		const EIoErrorCode ExpectedErrorCode = EIoErrorCode::InvalidCode;
 		TIoStatusOr<FIoStatusTestType> Result;
 		Result = TIoStatusOr<FIoStatusTestType>(FIoStatus(ExpectedErrorCode));
-		TestEqual("Assign temporary IoStatusOr with status", Result.Status().GetErrorCode(), ExpectedErrorCode);
+		TEST_EQUAL("Assign temporary IoStatusOr with status", Result.Status().GetErrorCode(), ExpectedErrorCode);
 	}
 
 	{
@@ -88,14 +88,14 @@ void TestAssignment()
 		TIoStatusOr<FIoStatusTestType> Other = FIoStatusTestType(ExpectedText);
 		TIoStatusOr<FIoStatusTestType> Result;
 		Result = Other;
-		TestEqual("Assign IoStatusOr with value", Result.ValueOrDie().Text, ExpectedText);
+		TEST_EQUAL("Assign IoStatusOr with value", Result.ValueOrDie().Text, ExpectedText);
 	}
 
 	{
 		const FString ExpectedText("Unreal");
 		TIoStatusOr<FIoStatusTestType> Result;
 		Result = TIoStatusOr<FIoStatusTestType>(ExpectedText);
-		TestEqual("Assign temporary IoStatusOr with value", Result.ValueOrDie().Text, ExpectedText);
+		TEST_EQUAL("Assign temporary IoStatusOr with value", Result.ValueOrDie().Text, ExpectedText);
 	}
 
 	{
@@ -103,14 +103,14 @@ void TestAssignment()
 		const FIoStatus IoStatus(ExpectedErrorCode);
 		TIoStatusOr<FIoStatusTestType> Result;
 		Result = IoStatus; 
-		TestEqual("Assign status", Result.Status().GetErrorCode(), ExpectedErrorCode);
+		TEST_EQUAL("Assign status", Result.Status().GetErrorCode(), ExpectedErrorCode);
 	}
 
 	{
 		const EIoErrorCode ExpectedErrorCode = EIoErrorCode::InvalidCode;
 		TIoStatusOr<FIoStatusTestType> Result;
 		Result = FIoStatus(ExpectedErrorCode);
-		TestEqual("Assign temporary status", Result.Status().GetErrorCode(), ExpectedErrorCode);
+		TEST_EQUAL("Assign temporary status", Result.Status().GetErrorCode(), ExpectedErrorCode);
 	}
 
 	{
@@ -118,29 +118,29 @@ void TestAssignment()
 		const FIoStatusTestType Value(ExpectedText);
 		TIoStatusOr<FIoStatusTestType> Result;
 		Result = Value;
-		TestEqual("Assign value", Result.ValueOrDie().Text, ExpectedText);
+		TEST_EQUAL("Assign value", Result.ValueOrDie().Text, ExpectedText);
 	}
 
 	{
 		const FString ExpectedText("Unreal");
 		TIoStatusOr<FIoStatusTestType> Result;
 		Result = FIoStatusTestType(ExpectedText);
-		TestEqual("Assign temporary value", Result.ValueOrDie().Text, ExpectedText);
+		TEST_EQUAL("Assign temporary value", Result.ValueOrDie().Text, ExpectedText);
 	}
 }
 
-void TestConsumeValue()
+void TestConsumeValue(FAutomationTestFixture& Test)
 {
 	const FString ExpectedText("Unreal");
 	TIoStatusOr<FIoStatusTestType> Result = FIoStatusTestType(ExpectedText);
 	FIoStatusTestType Value = Result.ConsumeValueOrDie();
-	TestEqual("Consume value or die with valid value", Value.Text, ExpectedText);
+	TEST_EQUAL("Consume value or die with valid value", Value.Text, ExpectedText);
 }
 
-TEST_CASE("Core::IO::IoStatusOr::Smoke Test", "[Core][IO][Smoke]")
+TEST_CASE_METHOD(FAutomationTestFixture, "Core::IO::IoStatusOr::Smoke Test", "[Core][IO][Smoke]")
 {
-	TestConstruct();
-	TestAssignment();
-	TestConsumeValue();
+	TestConstruct(*this);
+	TestAssignment(*this);
+	TestConsumeValue(*this);
 }
 
diff --git a/Engine/Source/Runtime/Core/Tests/Internationalization/ArchiveTest.cpp b/Engine/Source/Runtime/Core/Tests/Internationalization/ArchiveTest.cpp
--- a/Engine/Source/Runtime/Core/Tests/Internationalization/ArchiveTest.cpp
+++ b/Engine/Source/Runtime/Core/Tests/Internationalization/ArchiveTest.cpp
@@ -7,7 +7,7 @@
 #include "Internationalization/InternationalizationMetadata.h"
 #include "TestHarness.h"
 
-TEST_CASE("Core::Internationalization::FInternationalizationArchive::Smoke Test", "[Core][Internationalization][Smoke]")
+TEST_CASE_METHOD(FAutomationTestFixture, "Core::Internationalization::FInternationalizationArchive::Smoke Test", "[Core][Internationalization][Smoke]")
 {
 	// Key metadata
 	TSharedPtr< FLocMetadataObject > KeyMetadataA = MakeShareable( new FLocMetadataObject );
@@ -55,22 +55,22 @@ TEST_CASE("Core::Internationalization::FInternationalizationArchive::Smoke Test"
 		// Add duplicate entry that differs in bIsOptional flag.  This add should report success because we already have an entry with matching 
 		//  namespace/source/keymetadata.  Differences in bIsOptional are not taken into consideration. 
 		bool bResult = TestArchive.AddEntry( TestNamespace, SourceAKey, SourceA, Translation, NULL, TestOptionalFalse );
-		TestTrue( TEXT("AddEntry result = true"), bResult);
+		TEST_TRUE( TEXT("AddEntry result = true"), bResult);
 
 		// We should only have one entry in the archive
 		int32 EntryCount = 0;
-		for(auto Iter = TestArchive.GetEntriesBySourceTextIterator(); Iter; ++Iter, ++EntryCount);
-		TestEqual( TEXT("EntryCount == 1"), EntryCount, 1 );
+		for (auto Iter = TestArchive.GetEntriesBySourceTextIterator(); Iter; ++Iter, ++EntryCount) {};
+		TEST_EQUAL( TEXT("EntryCount == 1"), EntryCount, 1 );
 
 		// Make sure the original bIsOptional value is not overwritten by our second add.
 		TSharedPtr< FArchiveEntry > FoundEntry = TestArchive.FindEntryByKey( TestNamespace, SourceAKey, NULL );
 		if( !FoundEntry.IsValid() )
 		{
-			TestTrue(TEXT("FArchiveEntry could not find entry using FindEntryByKey."), FoundEntry.IsValid());
+			TEST_TRUE(TEXT("FArchiveEntry could not find entry using FindEntryByKey."), FoundEntry.IsValid());
 		}
 		else
 		{
-			TestTrue( TEXT("FoundEntry->Namespace == Namespace"), FoundEntry->bIsOptional == TestOptionalTrue );
+			TEST_TRUE( TEXT("FoundEntry->Namespace == Namespace"), FoundEntry->bIsOptional == TestOptionalTrue );
 		}
 	}
 
@@ -83,29 +83,29 @@ TEST_CASE("Core::Internationalization::FInternationalizationArchive::Smoke Test"
 			TSharedPtr< FArchiveEntry > FoundEntry = TestArchive.FindEntryByKey( TestNamespace, SourceAKey, KeyMetadataA );
 			if( !FoundEntry.IsValid() )
 			{
-				TestTrue(TEXT("FArchiveEntry could not find entry using FindEntryByKey."), FoundEntry.IsValid());
+				TEST_TRUE(TEXT("FArchiveEntry could not find entry using FindEntryByKey."), FoundEntry.IsValid());
 			}
 			else
 			{
-				TestTrue( TEXT("FoundEntry->Namespace == Namespace"), FoundEntry->Namespace == TestNamespace );
-				TestTrue( TEXT("FoundEntry->Source == Source"), FoundEntry->Source == SourceA );
-				TestTrue( TEXT("FoundEntry->Translation == Translation"), FoundEntry->Translation == Translation );
+				TEST_TRUE( TEXT("FoundEntry->Namespace == Namespace"), FoundEntry->Namespace == TestNamespace );
+				TEST_TRUE( TEXT("FoundEntry->Source == Source"), FoundEntry->Source == SourceA );
+				TEST_TRUE( TEXT("FoundEntry->Translation == Translation"), FoundEntry->Translation == Translation );
 				if( FoundEntry->KeyMetadataObj == KeyMetadataA )
 				{
-					TestFalse(TEXT("FArchiveEntry KeyMetadataObj is not a unique object."), FoundEntry->KeyMetadataObj == KeyMetadataA );
+					TEST_FALSE(TEXT("FArchiveEntry KeyMetadataObj is not a unique object."), FoundEntry->KeyMetadataObj == KeyMetadataA );
 				}
-				TestTrue( TEXT("FoundEntry->KeyMetadataObj == KeyMetadataA"), *(FoundEntry->KeyMetadataObj) == *(KeyMetadataA) );
+				TEST_TRUE( TEXT("FoundEntry->KeyMetadataObj == KeyMetadataA"), *(FoundEntry->KeyMetadataObj) == *(KeyMetadataA) );
 			}
 
 			// Passing in a mismatched key metadata will fail to find the entry.  Any fallback logic is intended to happen at runtime
 			FoundEntry = TestArchive.FindEntryByKey( TestNamespace, SourceAKey, NULL );
-			TestInvalid( TEXT("!FoundEntry.IsValid()"), FoundEntry);
+			TEST_INVALID( TEXT("!FoundEntry.IsValid()"), FoundEntry);
 
 			FoundEntry = TestArchive.FindEntryByKey( TestNamespace, SourceAKey, MakeShareable( new FLocMetadataObject() ) );
-			TestInvalid( TEXT("!FoundEntry.IsValid()"), FoundEntry);
+			TEST_INVALID( TEXT("!FoundEntry.IsValid()"), FoundEntry);
 
 			FoundEntry = TestArchive.FindEntryByKey( TestNamespace, SourceAKey, KeyMetadataB );
-			TestInvalid( TEXT("!FoundEntry.IsValid()"), FoundEntry);
+			TEST_INVALID( TEXT("!FoundEntry.IsValid()"), FoundEntry);
 		}
 
 		// Ensure we can properly lookup entries with non-null but empty key metadata
@@ -115,12 +115,10 @@ TEST_CASE("Core::Internationalization::FInternationalizationArchive::Smoke Test"
 
 			TSharedPtr< FArchiveEntry > FoundEntry = TestArchive.FindEntryByKey( TestNamespace, SourceAKey, NULL );
 			
-			bool bFoundEntryValid = FoundEntry.IsValid();
-			TestTrue(TEXT("FArchiveEntry could not find entry using FindEntryByKey."), bFoundEntryValid);
-			if (bFoundEntryValid)
+			if (TestTrue(TEXT("FArchiveEntry could not find entry using FindEntryByKey."), FoundEntry.IsValid()))
 			{
-				TestTrue( TEXT("FoundEntry->Namespace == Namespace"), FoundEntry->Namespace == TestNamespace );
-				TestTrue( TEXT("FoundEntry->Source == Source"), FoundEntry->Source == SourceA );
+				TEST_TRUE( TEXT("FoundEntry->Namespace == Namespace"), FoundEntry->Namespace == TestNamespace );
+				TEST_TRUE( TEXT("FoundEntry->Source == Source"), FoundEntry->Source == SourceA );
 			}
 		}
 
@@ -130,22 +128,18 @@ TEST_CASE("Core::Internationalization::FInternationalizationArchive::Smoke Test"
 			TestArchive.AddEntry( TestNamespace, SourceAKey, SourceA, Translation, NULL, false );
 
 			TSharedPtr< FArchiveEntry > FoundEntry = TestArchive.FindEntryByKey( TestNamespace, SourceAKey, NULL );
-			bool bFoundEntryValid = FoundEntry.IsValid();
-			TestTrue(TEXT("FArchiveEntry could not find entry using FindEntryByKey."), bFoundEntryValid);
-			if (bFoundEntryValid)
+			if(TestTrue(TEXT("FArchiveEntry could not find entry using FindEntryByKey."), FoundEntry.IsValid() ) )
 			{
-				TestTrue( TEXT("FoundEntry->Namespace == Namespace"), FoundEntry->Namespace == TestNamespace );
-				TestTrue( TEXT("FoundEntry->Source == Source"), FoundEntry->Source == SourceA );
+				TEST_TRUE( TEXT("FoundEntry->Namespace == Namespace"), FoundEntry->Namespace == TestNamespace );
+				TEST_TRUE( TEXT("FoundEntry->Source == Source"), FoundEntry->Source == SourceA );
 			}
 
 			// Test lookup with non-null but empty key metadata
 			FoundEntry = TestArchive.FindEntryByKey( TestNamespace, SourceAKey, MakeShareable( new FLocMetadataObject() ) );
-			bFoundEntryValid = FoundEntry.IsValid();
-			TestTrue(TEXT("FArchiveEntry could not find entry using FindEntryByKey."), bFoundEntryValid);
-			if (bFoundEntryValid)
+			if (TestTrue(TEXT("FArchiveEntry could not find entry using FindEntryByKey."), FoundEntry.IsValid()))
 			{
-				TestTrue(TEXT("FoundEntry->Namespace == Namespace"), FoundEntry->Namespace == TestNamespace);
-				TestTrue( TEXT("FoundEntry->Source == Source"), FoundEntry->Source == SourceA );
+				TEST_TRUE(TEXT("FoundEntry->Namespace == Namespace"), FoundEntry->Namespace == TestNamespace);
+				TEST_TRUE( TEXT("FoundEntry->Source == Source"), FoundEntry->Source == SourceA );
 			}
 		}
 
@@ -162,38 +156,35 @@ TEST_CASE("Core::Internationalization::FInternationalizationArchive::Smoke Test"
 
 			// Finding entry using a source that has Boolean *IsMature
 			TSharedPtr< FArchiveEntry > FoundEntry = TestArchive.FindEntryByKey( TestNamespace, SourceAKey, KeyMetadataA );
-			bool bFoundEntryValid = FoundEntry.IsValid();
-			TestTrue(TEXT("FArchiveEntry could not find entry using FindEntryByKey."), bFoundEntryValid);
-			if (bFoundEntryValid)
+			if(TestTrue(TEXT("FArchiveEntry could not find entry using FindEntryByKey."), FoundEntry.IsValid() ) )
 			{
-				TestTrue( TEXT("FoundEntry->Namespace == Namespace"), FoundEntry->Namespace == TestNamespace );
-				TestTrue( TEXT("FoundEntry->Source == Source"), FoundEntry->Source == SourceA );
+				TEST_TRUE( TEXT("FoundEntry->Namespace == Namespace"), FoundEntry->Namespace == TestNamespace );
+				TEST_TRUE( TEXT("FoundEntry->Source == Source"), FoundEntry->Source == SourceA );
 			}
 
 			// Attempting to add an entry that only differs by a * prefexed source metadata entry value or type will result in success since
 			//  a 'matching' entry already exists in the archive.  We should, however, only have one entry in the archive
 			bool bResult = TestArchive.AddEntry( TestNamespace, SourceAKey, SourceA, Translation, KeyMetadataA, false );
 
-			TestTrue( TEXT("AddEntry result = true"), bResult);
+			TEST_TRUE( TEXT("AddEntry result = true"), bResult);
 
 			// We should only have one entry in the archive
 			int32 EntryCount = 0;
-			for(auto Iter = TestArchive.GetEntriesBySourceTextIterator(); Iter; ++Iter, ++EntryCount);
-			TestEqual( TEXT("EntryCount == 1"), EntryCount, 1 );
+			for (auto Iter = TestArchive.GetEntriesBySourceTextIterator(); Iter; ++Iter, ++EntryCount) {};
+			TEST_EQUAL( TEXT("EntryCount == 1"), EntryCount, 1 );
 
 			// Check to see that the original type/value of the * prefixed entry did not get modified
 			FoundEntry = TestArchive.FindEntryByKey( TestNamespace, SourceAKey, KeyMetadataA );
-			bFoundEntryValid = FoundEntry.IsValid();
-			TestTrue(TEXT("FArchiveEntry could not find entry using FindEntryByKey."), bFoundEntryValid);
-			if (bFoundEntryValid)
+
+			if (TestTrue(TEXT("FArchiveEntry could not find entry using FindEntryByKey."), FoundEntry.IsValid()) )
 			{
 				if( FoundEntry->Source.MetadataObj->HasTypedField< ELocMetadataType::String >( TEXT("*IsMature") ) )
 				{
-					TestTrue( TEXT("Metadata Type == String and Value == Empty string"), FoundEntry->Source.MetadataObj->GetStringField( TEXT("*IsMature") ) == TEXT("") );
+					TEST_TRUE( TEXT("Metadata Type == String and Value == Empty string"), FoundEntry->Source.MetadataObj->GetStringField( TEXT("*IsMature") ) == TEXT("") );
 				}
 				else
 				{
-					TestTrue(TEXT("FArchiveEntry * prefixed metadata entry on source object was modified unexpectedly."), false);
+					TEST_TRUE(TEXT("FArchiveEntry * prefixed metadata entry on source object was modified unexpectedly."), false);
 				}
 			}
 		}
diff --git a/Engine/Source/Runtime/Core/Tests/Internationalization/ManifestTest.cpp b/Engine/Source/Runtime/Core/Tests/Internationalization/ManifestTest.cpp
--- a/Engine/Source/Runtime/Core/Tests/Internationalization/ManifestTest.cpp
+++ b/Engine/Source/Runtime/Core/Tests/Internationalization/ManifestTest.cpp
@@ -12,7 +12,7 @@ int32 CountManifestEntries( const FInternationalizationManifest& Manifest )
 	return Manifest.GetNumEntriesByKey();
 }
 
-TEST_CASE("Core::Internationalization::FManifestContext::Context", "[Core][Internationalization][Smoke]")
+TEST_CASE_METHOD(FAutomationTestFixture, "Core::Internationalization::FManifestContext::Context", "[Core][Internationalization][Smoke]")
 {
 	// Key metadata
 	TSharedPtr< FLocMetadataObject > KeyMetadataA = MakeShareable( new FLocMetadataObject );
@@ -65,17 +65,17 @@ TEST_CASE("Core::Internationalization::FManifestContext::Context", "[Core][Inter
 		{ //-V760
 			FManifestContext ContextAClone = ContextA;
 
-			TestFalse(TEXT("FManifestContext InfoMetadataObj and its Clone are not unique objects."), ContextAClone.InfoMetadataObj == ContextA.InfoMetadataObj);
-			TestFalse(TEXT("FManifestContext KeyMetadataObj and its Clone are not unique objects."), ContextAClone.KeyMetadataObj == ContextA.KeyMetadataObj);
+			TEST_FALSE(TEXT("FManifestContext InfoMetadataObj and its Clone are not unique objects."), ContextAClone.InfoMetadataObj == ContextA.InfoMetadataObj);
+			TEST_FALSE(TEXT("FManifestContext KeyMetadataObj and its Clone are not unique objects."), ContextAClone.KeyMetadataObj == ContextA.KeyMetadataObj);
 			
-			TestEqual( TEXT("ContextAClone.Key == ContextA.Key"), ContextAClone.Key, ContextA.Key );
-			TestEqual( TEXT("ContextAClone.SourceLocation == ContextA.SourceLocation"), ContextAClone.SourceLocation, ContextA.SourceLocation );
-			TestEqual( TEXT("ContextAClone.bIsOptional == ContextA.bIsOptional"), ContextAClone.bIsOptional, ContextA.bIsOptional );
-			TestTrue( TEXT("ContextAClone.InfoMetadataObj == ContextA.InfoMetadataObj"), *(ContextAClone.InfoMetadataObj) == *(ContextA.InfoMetadataObj) );
-			TestTrue( TEXT("ContextAClone.KeyMetadataObj == ContextA.KeyMetadataObj"), *(ContextAClone.KeyMetadataObj) == *(ContextA.KeyMetadataObj) );
+			TEST_EQUAL( TEXT("ContextAClone.Key == ContextA.Key"), ContextAClone.Key, ContextA.Key );
+			TEST_EQUAL( TEXT("ContextAClone.SourceLocation == ContextA.SourceLocation"), ContextAClone.SourceLocation, ContextA.SourceLocation );
+			TEST_EQUAL( TEXT("ContextAClone.bIsOptional == ContextA.bIsOptional"), ContextAClone.bIsOptional, ContextA.bIsOptional );
+			TEST_TRUE( TEXT("ContextAClone.InfoMetadataObj == ContextA.InfoMetadataObj"), *(ContextAClone.InfoMetadataObj) == *(ContextA.InfoMetadataObj) );
+			TEST_TRUE( TEXT("ContextAClone.KeyMetadataObj == ContextA.KeyMetadataObj"), *(ContextAClone.KeyMetadataObj) == *(ContextA.KeyMetadataObj) );
 
-			TestEqual( TEXT("ContextAClone == ContextA"), ContextAClone, ContextA );
-			TestFalse( TEXT("ContextAClone < ContextA"), ContextAClone < ContextA );
+			TEST_EQUAL( TEXT("ContextAClone == ContextA"), ContextAClone, ContextA );
+			TEST_FALSE( TEXT("ContextAClone < ContextA"), ContextAClone < ContextA );
 
 		}
 
@@ -83,103 +83,103 @@ TEST_CASE("Core::Internationalization::FManifestContext::Context", "[Core][Inter
 		{
 			FManifestContext ContextAClone = ContextA;
 	
-			TestFalse(TEXT("FManifestContext InfoMetadataObj and its Clone are not unique objects."), ContextAClone.InfoMetadataObj == ContextA.InfoMetadataObj);
-			TestFalse(TEXT("FManifestContext KeyMetadataObj and its Clone are not unique objects."), ContextAClone.KeyMetadataObj == ContextA.KeyMetadataObj);
+			TEST_FALSE(TEXT("FManifestContext InfoMetadataObj and its Clone are not unique objects."), ContextAClone.InfoMetadataObj == ContextA.InfoMetadataObj);
+			TEST_FALSE(TEXT("FManifestContext KeyMetadataObj and its Clone are not unique objects."), ContextAClone.KeyMetadataObj == ContextA.KeyMetadataObj);
 
-			TestEqual( TEXT("ContextAClone.Key == ContextA.Key"), ContextAClone.Key, ContextA.Key );
-			TestEqual( TEXT("ContextAClone.SourceLocation == ContextA.SourceLocation"), ContextAClone.SourceLocation, ContextA.SourceLocation );
-			TestEqual( TEXT("ContextAClone.bIsOptional == ContextA.bIsOptional"), ContextAClone.bIsOptional, ContextA.bIsOptional );
-			TestTrue( TEXT("ContextAClone.InfoMetadataObj == ContextA.InfoMetadataObj"), *(ContextAClone.InfoMetadataObj) == *(ContextA.InfoMetadataObj) );
-			TestTrue( TEXT("ContextAClone.KeyMetadataObj == ContextA.KeyMetadataObj"), *(ContextAClone.KeyMetadataObj) == *(ContextA.KeyMetadataObj) );
+			TEST_EQUAL( TEXT("ContextAClone.Key == ContextA.Key"), ContextAClone.Key, ContextA.Key );
+			TEST_EQUAL( TEXT("ContextAClone.SourceLocation == ContextA.SourceLocation"), ContextAClone.SourceLocation, ContextA.SourceLocation );
+			TEST_EQUAL( TEXT("ContextAClone.bIsOptional == ContextA.bIsOptional"), ContextAClone.bIsOptional, ContextA.bIsOptional );
+			TEST_TRUE( TEXT("ContextAClone.InfoMetadataObj == ContextA.InfoMetadataObj"), *(ContextAClone.InfoMetadataObj) == *(ContextA.InfoMetadataObj) );
+			TEST_TRUE( TEXT("ContextAClone.KeyMetadataObj == ContextA.KeyMetadataObj"), *(ContextAClone.KeyMetadataObj) == *(ContextA.KeyMetadataObj) );
 
-			TestEqual( TEXT("ContextAClone == ContextA"), ContextAClone, ContextA );
-			TestFalse( TEXT("ContextAClone < ContextA"), ContextAClone < ContextA );
+			TEST_EQUAL( TEXT("ContextAClone == ContextA"), ContextAClone, ContextA );
+			TEST_FALSE( TEXT("ContextAClone < ContextA"), ContextAClone < ContextA );
 		}
 
 		// Test comparison operator
 		{
 			// Key and KeyMetadataObj members should be the only items that are taken into account when comparing
 			FManifestContext ContextAClone = ContextA;
-			TestEqual( TEXT("ContextAClone == ContextA"), ContextAClone, ContextA );
+			TEST_EQUAL( TEXT("ContextAClone == ContextA"), ContextAClone, ContextA );
 
 			// Arbitrarily change all the non-important members
 			ContextAClone.SourceLocation = ContextA.SourceLocation + TEXT("New");
 			ContextAClone.bIsOptional = !ContextA.bIsOptional;
 			ContextAClone.InfoMetadataObj.Reset();
 			ContextAClone.InfoMetadataObj = MakeShareable( new FLocMetadataObject( *InfoMetadataB ) );
-			TestEqual( TEXT("ContextAClone == ContextA"), ContextAClone, ContextA );
+			TEST_EQUAL( TEXT("ContextAClone == ContextA"), ContextAClone, ContextA );
 
 			// Changing the key in any way will cause comparison to fail
 			ContextAClone.Key = ContextAClone.Key.GetString() + TEXT("New");
-			TestNotEqual( TEXT("ContextAClone != ContextA"), ContextAClone, ContextA );
+			TEST_NOT_EQUAL( TEXT("ContextAClone != ContextA"), ContextAClone, ContextA );
 
 			// Reset and test KeyMetadataObj change to one of the value entries
 			ContextAClone = ContextA;
 			ContextAClone.KeyMetadataObj->SetStringField( TEXT("TargetPlurality"), TEXT("Plural") );
-			TestNotEqual( TEXT("ContextAClone != ContextA"), ContextAClone, ContextA );
+			TEST_NOT_EQUAL( TEXT("ContextAClone != ContextA"), ContextAClone, ContextA );
 
 			// Reset and test addition of entry to KeyMetadataObj
 			ContextAClone = ContextA;
 			ContextAClone.KeyMetadataObj->SetStringField( TEXT("NewField"), TEXT("NewFieldValue") );
-			TestNotEqual( TEXT("ContextAClone != ContextA"), ContextAClone, ContextA );
+			TEST_NOT_EQUAL( TEXT("ContextAClone != ContextA"), ContextAClone, ContextA );
 
 			// Reset and test removal of entry from KeyMetadataObj
 			ContextAClone = ContextA;
 			ContextAClone.KeyMetadataObj->RemoveField( TEXT("TargetPlurality") );
-			TestNotEqual( TEXT("ContextAClone != ContextA"), ContextAClone, ContextA );
+			TEST_NOT_EQUAL( TEXT("ContextAClone != ContextA"), ContextAClone, ContextA );
 
 			// Context with valid but empty KeyMetadataObject should be equivalent to Context with null KeyMetadataObject
 			FManifestContext ContextEmptyA;
 			FManifestContext ContextEmptyB;
 			ContextEmptyB.KeyMetadataObj = MakeShareable( new FLocMetadataObject );
-			TestEqual( TEXT("ContextEmptyA == ContextEmptyB"), ContextEmptyA, ContextEmptyB );
+			TEST_EQUAL( TEXT("ContextEmptyA == ContextEmptyB"), ContextEmptyA, ContextEmptyB );
 		}
 
 		// Testing less than operator
 		{
-			TestTrue( TEXT("ContextA < ContextB"), ContextA < ContextB );
+			TEST_TRUE( TEXT("ContextA < ContextB"), ContextA < ContextB );
 
 			FManifestContext ContextAClone = ContextA;
 
 			// Differences in Key
-			TestFalse( TEXT("ContextA < ContextAClone"), ContextA < ContextAClone );
+			TEST_FALSE( TEXT("ContextA < ContextAClone"), ContextA < ContextAClone );
 			ContextAClone.Key = ContextAClone.Key.GetString() + TEXT("A");
-			//currently failing TestTrue( TEXT("ContextA < ContextAClone"), ContextA < ContextAClone );
+			//currently failing TEST_TRUE( TEXT("ContextA < ContextAClone"), ContextA < ContextAClone );
 
 			// Adding new key metadata entry that will appear before other entries
 			ContextAClone = ContextA;
 			ContextAClone.KeyMetadataObj->SetStringField( TEXT("ANewKey"), TEXT("ANewValue") );
-			TestTrue( TEXT("ContextAClone < ContextA"), ContextAClone < ContextA );
+			TEST_TRUE( TEXT("ContextAClone < ContextA"), ContextAClone < ContextA );
 
 			// Adding new key metadata entry that will appear after other entries
 			ContextAClone = ContextA;
 			ContextAClone.KeyMetadataObj->SetStringField( TEXT("ZNewKey"), TEXT("ZNewValue") );
-			//currently failing TestTrue( TEXT("ContextA < ContextAClone"), ContextA < ContextAClone );
+			//currently failing TEST_TRUE( TEXT("ContextA < ContextAClone"), ContextA < ContextAClone );
 
 			// Removing a key metadata entry
 			ContextAClone = ContextA;
 			ContextAClone.KeyMetadataObj->RemoveField( TEXT("TargetPlurality") ) ;
-			TestTrue( TEXT("ContextAClone < ContextA"), ContextAClone < ContextA );
+			TEST_TRUE( TEXT("ContextAClone < ContextA"), ContextAClone < ContextA );
 
 			// Changing a key metadata entry value
 			ContextAClone = ContextA;
 			ContextAClone.KeyMetadataObj->SetStringField( TEXT("TargetPlurality"), TEXT("A") ) ;
-			TestTrue( TEXT("ContextAClone < ContextA"), ContextAClone < ContextA );
+			TEST_TRUE( TEXT("ContextAClone < ContextA"), ContextAClone < ContextA );
 
 			FManifestContext ContextEmptyA;
 			FManifestContext ContextEmptyB;
-			TestFalse( TEXT("ContextEmptyA < ContextEmptyB"), ContextEmptyA < ContextEmptyB );
+			TEST_FALSE( TEXT("ContextEmptyA < ContextEmptyB"), ContextEmptyA < ContextEmptyB );
 			ContextEmptyB.KeyMetadataObj = MakeShareable( new FLocMetadataObject );
-			TestFalse( TEXT("ContextEmptyA < ContextEmptyB"), ContextEmptyA < ContextEmptyB );
-			TestFalse( TEXT("ContextEmptyB < ContextEmptyA"), ContextEmptyB < ContextEmptyA );
+			TEST_FALSE( TEXT("ContextEmptyA < ContextEmptyB"), ContextEmptyA < ContextEmptyB );
+			TEST_FALSE( TEXT("ContextEmptyB < ContextEmptyA"), ContextEmptyB < ContextEmptyA );
 			ContextEmptyB.KeyMetadataObj->SetStringField(TEXT("AMetadataKey"), TEXT("AMetadataValue"));
-			TestTrue( TEXT("ContextEmptyA < ContextEmptyB"), ContextEmptyA < ContextEmptyB );
+			TEST_TRUE( TEXT("ContextEmptyA < ContextEmptyB"), ContextEmptyA < ContextEmptyB );
 		}
 	}
 }
 
 
-TEST_CASE("Core::Internationalization::FLocItem::LocItem", "[Core][Internationalization][Smoke]")
+TEST_CASE_METHOD(FAutomationTestFixture, "Core::Internationalization::FLocItem::LocItem", "[Core][Internationalization][Smoke]")
 {
 	// Key metadata
 	TSharedPtr< FLocMetadataObject > KeyMetadataA = MakeShareable( new FLocMetadataObject );
@@ -226,13 +226,13 @@ TEST_CASE("Core::Internationalization::FLocItem::LocItem", "[Core][International
 		{
 			FLocItem LocItemAClone = LocItemA;
 
-			TestFalse(TEXT("FLocItem MetadataObj and its Clone are not unique objects."), LocItemAClone.MetadataObj == LocItemA.MetadataObj);
+			TEST_FALSE(TEXT("FLocItem MetadataObj and its Clone are not unique objects."), LocItemAClone.MetadataObj == LocItemA.MetadataObj);
 
-			TestEqual( TEXT("LocItemAClone.Text == LocItemA.Text"), LocItemAClone.Text, LocItemA.Text );
-			TestTrue( TEXT("LocItemAClone.MetadataObj == LocItemA.MetadataObj"), *(LocItemAClone.MetadataObj) == *(LocItemA.MetadataObj) );
+			TEST_EQUAL( TEXT("LocItemAClone.Text == LocItemA.Text"), LocItemAClone.Text, LocItemA.Text );
+			TEST_TRUE( TEXT("LocItemAClone.MetadataObj == LocItemA.MetadataObj"), *(LocItemAClone.MetadataObj) == *(LocItemA.MetadataObj) );
 
-			TestEqual( TEXT("LocItemAClone == LocItemA"), LocItemAClone, LocItemA );
-			TestFalse( TEXT("LocItemAClone < LocItemA"), LocItemAClone < LocItemA );
+			TEST_EQUAL( TEXT("LocItemAClone == LocItemA"), LocItemAClone, LocItemA );
+			TEST_FALSE( TEXT("LocItemAClone < LocItemA"), LocItemAClone < LocItemA );
 
 		}
 
@@ -242,14 +242,14 @@ TEST_CASE("Core::Internationalization::FLocItem::LocItem", "[Core][International
 
 			LocItemAClone = LocItemA;
 
-			TestFalse(TEXT("FLocItem MetadataObj and its Clone are not unique objects."), LocItemAClone.MetadataObj == LocItemA.MetadataObj);
+			TEST_FALSE(TEXT("FLocItem MetadataObj and its Clone are not unique objects."), LocItemAClone.MetadataObj == LocItemA.MetadataObj);
 			
-			TestEqual( TEXT("LocItemAClone.Text == LocItemA.Text"), LocItemAClone.Text, LocItemA.Text );
-			TestTrue( TEXT("LocItemAClone.MetadataObj == LocItemA.MetadataObj"), *(LocItemAClone.MetadataObj) == *(LocItemA.MetadataObj) );
+			TEST_EQUAL( TEXT("LocItemAClone.Text == LocItemA.Text"), LocItemAClone.Text, LocItemA.Text );
+			TEST_TRUE( TEXT("LocItemAClone.MetadataObj == LocItemA.MetadataObj"), *(LocItemAClone.MetadataObj) == *(LocItemA.MetadataObj) );
 
-			TestEqual( TEXT("LocItemAClone == LocItemA"), LocItemAClone, LocItemA );
-			TestFalse( TEXT("LocItemAClone < LocItemA"), LocItemAClone < LocItemA );
-			TestFalse( TEXT("LocItemA < LocItemAClone"), LocItemA < LocItemAClone );
+			TEST_EQUAL( TEXT("LocItemAClone == LocItemA"), LocItemAClone, LocItemA );
+			TEST_FALSE( TEXT("LocItemAClone < LocItemA"), LocItemAClone < LocItemA );
+			TEST_FALSE( TEXT("LocItemA < LocItemAClone"), LocItemA < LocItemAClone );
 		}
 
 		// Test comparison operator
@@ -257,74 +257,74 @@ TEST_CASE("Core::Internationalization::FLocItem::LocItem", "[Core][International
 			// Text and MetadataObj members should both be taken into account when comparing.  Note, Metadata supports a special * name prefix that causes the type 
 			//   and value of the metadata to be ignored when performing comparisons
 			FLocItem LocItemAClone = LocItemA;
-			TestEqual( TEXT("LocItemAClone == LocItemA"), LocItemAClone, LocItemA );
+			TEST_EQUAL( TEXT("LocItemAClone == LocItemA"), LocItemAClone, LocItemA );
 
 			// Metadata with * prefix does not impact comparison but both FLocItems need a metadata where the name matches( type and value can be different )
 			FLocItem LocItemAClone2 = LocItemA;
 			LocItemAClone.MetadataObj->SetStringField( TEXT("*NewNonCompare"), TEXT("NewNonCompareValue") );
-			TestNotEqual( TEXT("LocItemAClone != LocItemAClone2"), LocItemAClone, LocItemAClone2 );
+			TEST_NOT_EQUAL( TEXT("LocItemAClone != LocItemAClone2"), LocItemAClone, LocItemAClone2 );
 
 			LocItemAClone2.MetadataObj->SetStringField( TEXT("*NewNonCompare"), TEXT("NewNonCompareValue2") );
-			TestEqual( TEXT("LocItemAClone == LocItemAClone2"), LocItemAClone, LocItemAClone2 );
+			TEST_EQUAL( TEXT("LocItemAClone == LocItemAClone2"), LocItemAClone, LocItemAClone2 );
 
 			// Changing the text in any way will cause comparison to fail
 			LocItemAClone = LocItemA;
 			LocItemAClone.Text = LocItemAClone.Text + TEXT("New");
-			TestNotEqual( TEXT("LocItemAClone != LocItemA"), LocItemAClone, LocItemA );
+			TEST_NOT_EQUAL( TEXT("LocItemAClone != LocItemA"), LocItemAClone, LocItemA );
 
 			// LocItem with valid but empty KeyMetadataObject should be equivalent to LocItem with null KeyMetadataObject
 			FLocItem LocItemEmptyA(TEXT("TestText"));
 			FLocItem LocItemEmptyB(TEXT("TestText"));
 			LocItemEmptyB.MetadataObj = MakeShareable( new FLocMetadataObject );
-			TestEqual( TEXT("LocItemEmptyA == LocItemEmptyB"), LocItemEmptyA, LocItemEmptyB );
+			TEST_EQUAL( TEXT("LocItemEmptyA == LocItemEmptyB"), LocItemEmptyA, LocItemEmptyB );
 		}
 
 		// Testing less than operator
 		{
-			TestTrue( TEXT("LocItemA < LocItemB"), LocItemA < LocItemB );
-			TestFalse( TEXT("LocItemB < LocItemA"), LocItemB < LocItemA );
+			TEST_TRUE( TEXT("LocItemA < LocItemB"), LocItemA < LocItemB );
+			TEST_FALSE( TEXT("LocItemB < LocItemA"), LocItemB < LocItemA );
 
 			FLocItem LocItemAClone = LocItemA;
 
 			// Differences in Text
-			TestFalse( TEXT("LocItemA < LocItemAClone"), LocItemA < LocItemAClone );
+			TEST_FALSE( TEXT("LocItemA < LocItemAClone"), LocItemA < LocItemAClone );
 			LocItemAClone.Text = LocItemAClone.Text + TEXT("A");
-			TestFalse( TEXT("LocItemAClone < LocItemA"), LocItemAClone < LocItemA );
-			//currently failing TestTrue( TEXT("LocItemA < LocItemAClone"), LocItemA < LocItemAClone );
+			TEST_FALSE( TEXT("LocItemAClone < LocItemA"), LocItemAClone < LocItemA );
+			//currently failing TEST_TRUE( TEXT("LocItemA < LocItemAClone"), LocItemA < LocItemAClone );
 
 			// Adding new key metadata entry
 			LocItemAClone = LocItemA;
 			LocItemAClone.MetadataObj->SetStringField( TEXT("ANewKey"), TEXT("ANewValue") );
-			//currently failing TestTrue( TEXT("LocItemA < LocItemAClone"), LocItemA < LocItemAClone );
-			TestFalse( TEXT("LocItemAClone < LocItemA"), LocItemAClone < LocItemA );
+			//currently failing TEST_TRUE( TEXT("LocItemA < LocItemAClone"), LocItemA < LocItemAClone );
+			TEST_FALSE( TEXT("LocItemAClone < LocItemA"), LocItemAClone < LocItemA );
 
 			// Removing a key metadata entry
 			LocItemAClone = LocItemA;
 			LocItemAClone.MetadataObj->RemoveField( TEXT("*IsMature") );
-			TestTrue( TEXT("LocItemAClone < LocItemA"), LocItemAClone < LocItemA );
-			TestFalse( TEXT("LocItemA < LocItemAClone"), LocItemA < LocItemAClone );
+			TEST_TRUE( TEXT("LocItemAClone < LocItemA"), LocItemAClone < LocItemA );
+			TEST_FALSE( TEXT("LocItemA < LocItemAClone"), LocItemA < LocItemAClone );
 
 			// Changing a key metadata entry value
 			LocItemAClone = LocItemA;
 			LocItemAClone.MetadataObj->SetBoolField( TEXT("*IsMature"),true );
-			//currently failing TestTrue( TEXT("LocItemA < LocItemAClone"), LocItemA < LocItemAClone );
-			TestFalse( TEXT("LocItemAClone < LocItemA"), LocItemAClone < LocItemA );
+			//currently failing TEST_TRUE( TEXT("LocItemA < LocItemAClone"), LocItemA < LocItemAClone );
+			TEST_FALSE( TEXT("LocItemAClone < LocItemA"), LocItemAClone < LocItemA );
 
 			// Test null and non-null but empty Metadata
 			FLocItem LocItemEmptyA( TEXT("SameText") );
 			FLocItem LocItemEmptyB( TEXT("SameText") );
-			TestFalse( TEXT("LocItemEmptyA < LocItemEmptyB"), LocItemEmptyA < LocItemEmptyB );
+			TEST_FALSE( TEXT("LocItemEmptyA < LocItemEmptyB"), LocItemEmptyA < LocItemEmptyB );
 			LocItemEmptyB.MetadataObj = MakeShareable( new FLocMetadataObject );
-			TestFalse( TEXT("LocItemEmptyA < LocItemEmptyB"), LocItemEmptyA < LocItemEmptyB );
-			TestFalse( TEXT("LocItemEmptyB < LocItemEmptyA"), LocItemEmptyB < LocItemEmptyA );
+			TEST_FALSE( TEXT("LocItemEmptyA < LocItemEmptyB"), LocItemEmptyA < LocItemEmptyB );
+			TEST_FALSE( TEXT("LocItemEmptyB < LocItemEmptyA"), LocItemEmptyB < LocItemEmptyA );
 			LocItemEmptyB.MetadataObj->SetStringField(TEXT("AMetadataKey"), TEXT("AMetadataValue"));
-			TestTrue( TEXT("LocItemEmptyA < LocItemEmptyB"), LocItemEmptyA < LocItemEmptyB );
+			TEST_TRUE( TEXT("LocItemEmptyA < LocItemEmptyB"), LocItemEmptyA < LocItemEmptyB );
 		}
 	}
 }
 
 
-TEST_CASE("Core::Internationalization::FManifestContext::Manifest", "[Core][Serialization][Smoke]")
+TEST_CASE_METHOD(FAutomationTestFixture, "Core::Internationalization::FManifestContext::Manifest", "[Core][Serialization][Smoke]")
 {
 	// Key metadata
 	TSharedPtr< FLocMetadataObject > KeyMetadataA = MakeShareable( new FLocMetadataObject );
@@ -386,8 +386,8 @@ TEST_CASE("Core::Internationalization::FManifestContext::Manifest", "[Core][Seri
 			bool bResult = TestManifest.AddSource( TestNamespace, Source, ContextA );
 
 			// Adding a duplicate entry reports success but our entry count is not increased after the first entry is added
-			TestTrue( TEXT("AddSource result == true"), bResult);
-			TestEqual( TEXT("ManifestCount == 1"), CountManifestEntries(TestManifest), 1 );
+			TEST_TRUE( TEXT("AddSource result == true"), bResult);
+			TEST_EQUAL( TEXT("ManifestCount == 1"), CountManifestEntries(TestManifest), 1 );
 		}
 
 		// Adding entries with exactly matching Source but different Contexts
@@ -397,44 +397,44 @@ TEST_CASE("Core::Internationalization::FManifestContext::Manifest", "[Core][Seri
 			TestManifest.AddSource( TestNamespace, Source, ContextA );
 			TestManifest.AddSource( TestNamespace, Source, ContextB );
 
-			TestEqual( TEXT("ManifestCount == 2"), CountManifestEntries(TestManifest), 2 );
+			TEST_EQUAL( TEXT("ManifestCount == 2"), CountManifestEntries(TestManifest), 2 );
 
 			// Test find by context
 			TSharedPtr< FManifestEntry > FoundEntry1 = TestManifest.FindEntryByContext( TestNamespace, ContextA );
 			if( !FoundEntry1.IsValid() )
 			{
-				TestTrue(TEXT("FManifestEntry could not find entry using FindEntryByContext."), FoundEntry1.IsValid());
+				TEST_TRUE(TEXT("FManifestEntry could not find entry using FindEntryByContext."), FoundEntry1.IsValid());
 			}
 			else
 			{
-				TestTrue( TEXT("FoundEntry->Source == Source"), FoundEntry1->Source == Source );
-				TestEqual( TEXT("FoundEntry->Context.Num() == 2"), FoundEntry1->Contexts.Num(), 2);
+				TEST_TRUE( TEXT("FoundEntry->Source == Source"), FoundEntry1->Source == Source );
+				TEST_EQUAL( TEXT("FoundEntry->Context.Num() == 2"), FoundEntry1->Contexts.Num(), 2);
 			}
 
 			TSharedPtr< FManifestEntry > FoundEntry2 = TestManifest.FindEntryByContext( TestNamespace, ContextB );
 			if( !FoundEntry2.IsValid() )
 			{
-				TestTrue(TEXT("FManifestEntry could not find entry using FindEntryByContext."), FoundEntry2.IsValid());
+				TEST_TRUE(TEXT("FManifestEntry could not find entry using FindEntryByContext."), FoundEntry2.IsValid());
 			}
 			else
 			{
-				TestTrue( TEXT("FoundEntry->Source == Source"), FoundEntry2->Source == Source );
-				TestEqual( TEXT("FoundEntry->Context.Num() == 2"), FoundEntry2->Contexts.Num(), 2);
+				TEST_TRUE( TEXT("FoundEntry->Source == Source"), FoundEntry2->Source == Source );
+				TEST_EQUAL( TEXT("FoundEntry->Context.Num() == 2"), FoundEntry2->Contexts.Num(), 2);
 			}
 
 			// Test find by source
 			TSharedPtr< FManifestEntry > FoundEntry3 = TestManifest.FindEntryBySource( TestNamespace, Source );
 			if( !FoundEntry3.IsValid() )
 			{
-				TestTrue(TEXT("FManifestEntry could not find entry using FindEntryBySource."), FoundEntry3.IsValid());
+				TEST_TRUE(TEXT("FManifestEntry could not find entry using FindEntryBySource."), FoundEntry3.IsValid());
 			}
 			else
 			{
-				TestTrue( TEXT("FoundEntry->Source == Source"), FoundEntry3->Source == Source );
-				TestEqual( TEXT("FoundEntry->Context.Num() == 2"), FoundEntry3->Contexts.Num(), 2);
+				TEST_TRUE( TEXT("FoundEntry->Source == Source"), FoundEntry3->Source == Source );
+				TEST_EQUAL( TEXT("FoundEntry->Context.Num() == 2"), FoundEntry3->Contexts.Num(), 2);
 			}
 
-			TestTrue( TEXT("FoundEntry1 == FoundEntry2 == FoundEntry3"), (FoundEntry1 == FoundEntry2) && (FoundEntry1 == FoundEntry3) );
+			TEST_TRUE( TEXT("FoundEntry1 == FoundEntry2 == FoundEntry3"), (FoundEntry1 == FoundEntry2) && (FoundEntry1 == FoundEntry3) );
 
 		}
 
@@ -453,8 +453,8 @@ TEST_CASE("Core::Internationalization::FManifestContext::Manifest", "[Core][Seri
 				bool bResult = TestManifest.AddSource( TestNamespace, ConflictingSourceB, ContextA );
 
 				// Adding the second entry reports failure and entry count is not increased
-				TestFalse( TEXT("AddSource result == false"), bResult);
-				TestEqual( TEXT("ManifestCount == 1"), CountManifestEntries(TestManifest), 1 );	
+				TEST_FALSE( TEXT("AddSource result == false"), bResult);
+				TEST_EQUAL( TEXT("ManifestCount == 1"), CountManifestEntries(TestManifest), 1 );	
 			}
 
 			// Source mismatched by standard metadata type (not * prefixed) 
@@ -471,8 +471,8 @@ TEST_CASE("Core::Internationalization::FManifestContext::Manifest", "[Core][Seri
 				ManifestTypeConflict.AddSource( TestNamespace, ConflictingSourceMetadataTypeA, ContextA );
 				bool bResult = ManifestTypeConflict.AddSource( TestNamespace, ConflictingSourceMetadataTypeB, ContextA );
 
-				TestFalse( TEXT("AddSource result == false"), bResult);
-				TestEqual( TEXT("ManifestCount == 1"), CountManifestEntries(ManifestTypeConflict), 1 );
+				TEST_FALSE( TEXT("AddSource result == false"), bResult);
+				TEST_EQUAL( TEXT("ManifestCount == 1"), CountManifestEntries(ManifestTypeConflict), 1 );
 
 				// Should not find a match when searching with the second Source we tried to add
 				TSharedPtr< FManifestEntry > FoundEntry = ManifestTypeConflict.FindEntryBySource( TestNamespace, ConflictingSourceMetadataTypeB );
@@ -494,8 +494,8 @@ TEST_CASE("Core::Internationalization::FManifestContext::Manifest", "[Core][Seri
 				TestManifest.AddSource( TestNamespace, ConflictingSourceMetadataValueA, ContextA );
 				bool bResult = TestManifest.AddSource( TestNamespace, ConflictingSourceMetadataValueB, ContextA );
 
-				TestFalse( TEXT("AddSource result == false"), bResult);
-				TestEqual( TEXT("ManifestCount == 1"), CountManifestEntries(TestManifest), 1 );
+				TEST_FALSE( TEXT("AddSource result == false"), bResult);
+				TEST_EQUAL( TEXT("ManifestCount == 1"), CountManifestEntries(TestManifest), 1 );
 
 				// Should not find a match when searching with the second Source we tried to add
 				TSharedPtr< FManifestEntry > FoundEntry = TestManifest.FindEntryBySource( TestNamespace, ConflictingSourceMetadataValueB );
@@ -520,8 +520,8 @@ TEST_CASE("Core::Internationalization::FManifestContext::Manifest", "[Core][Seri
 				//  and should not add an entry.  The reason is that AddSource does an 'exact' match check on the metadata object
 				bool bResult = TestManifest.AddSource( TestNamespace, ConflictingSourceMetadataTypeB, ContextA );
 
-				TestFalse( TEXT("AddSource result == false"), bResult);
-				TestEqual( TEXT("ManifestCount == 1"), CountManifestEntries(TestManifest), 1 );
+				TEST_FALSE( TEXT("AddSource result == false"), bResult);
+				TEST_EQUAL( TEXT("ManifestCount == 1"), CountManifestEntries(TestManifest), 1 );
 
 				// We should be able to find the entry using either Source FLocItem
 				TSharedPtr< FManifestEntry > FoundEntry1 = TestManifest.FindEntryBySource( TestNamespace, ConflictingSourceMetadataTypeA );
@@ -550,8 +550,8 @@ TEST_CASE("Core::Internationalization::FManifestContext::Manifest", "[Core][Seri
 				//  and should not add an entry.  The reason is that AddSource does an 'exact' match check on the metadata object
 				bool bResult = TestManifest.AddSource( TestNamespace, ConflictingSourceMetadataValueB, ContextA );
 
-				TestFalse( TEXT("AddSource result == false"), bResult);
-				TestEqual( TEXT("ManifestCount == 1"), CountManifestEntries(TestManifest), 1 );
+				TEST_FALSE( TEXT("AddSource result == false"), bResult);
+				TEST_EQUAL( TEXT("ManifestCount == 1"), CountManifestEntries(TestManifest), 1 );
 
 				// We should be able to find the entry using either Source FLocItem
 				TSharedPtr< FManifestEntry > FoundEntry1 = TestManifest.FindEntryBySource( TestNamespace, ConflictingSourceMetadataValueA );
@@ -576,8 +576,8 @@ TEST_CASE("Core::Internationalization::FManifestContext::Manifest", "[Core][Seri
 				TestManifest.AddSource( TestNamespace, ConflictingSourceA, ContextA );
 				bool bResult = TestManifest.AddSource( TestNamespace, ConflictingSourceB, ContextB );
 
-				TestTrue( TEXT("AddSource result == true"), bResult);
-				TestEqual( TEXT("ManifestCount == 2"), CountManifestEntries(TestManifest), 2 );
+				TEST_TRUE( TEXT("AddSource result == true"), bResult);
+				TEST_EQUAL( TEXT("ManifestCount == 2"), CountManifestEntries(TestManifest), 2 );
 
 				// We should be able to find two unique entries by source
 				TSharedPtr< FManifestEntry > FoundEntry1 = TestManifest.FindEntryBySource( TestNamespace, ConflictingSourceA );
@@ -586,7 +586,7 @@ TEST_CASE("Core::Internationalization::FManifestContext::Manifest", "[Core][Seri
 				TSharedPtr< FManifestEntry > FoundEntry2 = TestManifest.FindEntryBySource( TestNamespace, ConflictingSourceB );
 				TestValid( TEXT("FoundEntry2 is valid"), FoundEntry2);
 
-				TestTrue( TEXT("FoundEntry1 != FoundEntry2"), FoundEntry1 != FoundEntry2 );
+				TEST_TRUE( TEXT("FoundEntry1 != FoundEntry2"), FoundEntry1 != FoundEntry2 );
 
 				// We should be able to find two unique entries by context
 				TSharedPtr< FManifestEntry > FoundEntry3 = TestManifest.FindEntryByContext( TestNamespace, ContextA );
@@ -595,13 +595,13 @@ TEST_CASE("Core::Internationalization::FManifestContext::Manifest", "[Core][Seri
 				TSharedPtr< FManifestEntry > FoundEntry4 = TestManifest.FindEntryByContext( TestNamespace, ContextB );
 				TestValid( TEXT("FoundEntry3 is valid"), FoundEntry4);
 
-				TestTrue( TEXT("FoundEntry3 != FoundEntry4"), FoundEntry3 != FoundEntry4 );
+				TEST_TRUE( TEXT("FoundEntry3 != FoundEntry4"), FoundEntry3 != FoundEntry4 );
 
 				// Found entry looked up by source A should match entry looked up by context A
-				TestTrue( TEXT("FoundEntry1 == FoundEntry3"), FoundEntry1 == FoundEntry3 );
+				TEST_TRUE( TEXT("FoundEntry1 == FoundEntry3"), FoundEntry1 == FoundEntry3 );
 
 				// Found entry looked up by source B should match entry looked up by context B
-				TestTrue( TEXT("FoundEntry2 == FoundEntry4"), FoundEntry2 == FoundEntry4 );
+				TEST_TRUE( TEXT("FoundEntry2 == FoundEntry4"), FoundEntry2 == FoundEntry4 );
 
 			}
 
@@ -619,8 +619,8 @@ TEST_CASE("Core::Internationalization::FManifestContext::Manifest", "[Core][Seri
 				TestManifest.AddSource( TestNamespace, ConflictingSourceMetadataTypeA, ContextA );
 				bool bResult = TestManifest.AddSource( TestNamespace, ConflictingSourceMetadataTypeB, ContextB );
 
-				TestTrue( TEXT("AddSource result == true"), bResult);
-				TestEqual( TEXT("ManifestCount == 2"), CountManifestEntries(TestManifest), 2 );
+				TEST_TRUE( TEXT("AddSource result == true"), bResult);
+				TEST_EQUAL( TEXT("ManifestCount == 2"), CountManifestEntries(TestManifest), 2 );
 
 				// We should be able to find two unique entries by source
 				TSharedPtr< FManifestEntry > FoundEntry1 = TestManifest.FindEntryBySource( TestNamespace, ConflictingSourceMetadataTypeA );
@@ -629,7 +629,7 @@ TEST_CASE("Core::Internationalization::FManifestContext::Manifest", "[Core][Seri
 				TSharedPtr< FManifestEntry > FoundEntry2 = TestManifest.FindEntryBySource( TestNamespace, ConflictingSourceMetadataTypeB );
 				TestValid( TEXT("FoundEntry2 is valid"), FoundEntry2);
 
-				TestTrue( TEXT("FoundEntry1 != FoundEntry2"), FoundEntry1 != FoundEntry2 );
+				TEST_TRUE( TEXT("FoundEntry1 != FoundEntry2"), FoundEntry1 != FoundEntry2 );
 
 				// We should be able to find two unique entries by context
 				TSharedPtr< FManifestEntry > FoundEntry3 = TestManifest.FindEntryByContext( TestNamespace, ContextA );
@@ -638,13 +638,13 @@ TEST_CASE("Core::Internationalization::FManifestContext::Manifest", "[Core][Seri
 				TSharedPtr< FManifestEntry > FoundEntry4 = TestManifest.FindEntryByContext( TestNamespace, ContextB );
 				TestValid( TEXT("FoundEntry3 is valid"), FoundEntry4);
 
-				TestTrue( TEXT("FoundEntry3 != FoundEntry4"), FoundEntry3 != FoundEntry4 );
+				TEST_TRUE( TEXT("FoundEntry3 != FoundEntry4"), FoundEntry3 != FoundEntry4 );
 
 				// Found entry looked up by source A should match entry looked up by context A
-				TestTrue( TEXT("FoundEntry1 == FoundEntry3"), FoundEntry1 == FoundEntry3 );
+				TEST_TRUE( TEXT("FoundEntry1 == FoundEntry3"), FoundEntry1 == FoundEntry3 );
 
 				// Found entry looked up by source B should match entry looked up by context B
-				TestTrue( TEXT("FoundEntry2 == FoundEntry4"), FoundEntry2 == FoundEntry4 );
+				TEST_TRUE( TEXT("FoundEntry2 == FoundEntry4"), FoundEntry2 == FoundEntry4 );
 
 			}
 
@@ -662,8 +662,8 @@ TEST_CASE("Core::Internationalization::FManifestContext::Manifest", "[Core][Seri
 				TestManifest.AddSource( TestNamespace, ConflictingSourceMetadataValueA, ContextA );
 				bool bResult = TestManifest.AddSource( TestNamespace, ConflictingSourceMetadataValueB, ContextB );
 
-				TestTrue( TEXT("AddSource result == true"), bResult);
-				TestEqual( TEXT("ManifestCount == 2"), CountManifestEntries(TestManifest), 2 );
+				TEST_TRUE( TEXT("AddSource result == true"), bResult);
+				TEST_EQUAL( TEXT("ManifestCount == 2"), CountManifestEntries(TestManifest), 2 );
 
 				// We should be able to find two unique entries by source
 				TSharedPtr< FManifestEntry > FoundEntry1 = TestManifest.FindEntryBySource( TestNamespace, ConflictingSourceMetadataValueA );
@@ -672,7 +672,7 @@ TEST_CASE("Core::Internationalization::FManifestContext::Manifest", "[Core][Seri
 				TSharedPtr< FManifestEntry > FoundEntry2 = TestManifest.FindEntryBySource( TestNamespace, ConflictingSourceMetadataValueB );
 				TestValid( TEXT("FoundEntry2 is valid"), FoundEntry2);
 
-				TestTrue( TEXT("FoundEntry1 != FoundEntry2"), FoundEntry1 != FoundEntry2 );
+				TEST_TRUE( TEXT("FoundEntry1 != FoundEntry2"), FoundEntry1 != FoundEntry2 );
 
 				// We should be able to find two unique entries by context
 				TSharedPtr< FManifestEntry > FoundEntry3 = TestManifest.FindEntryByContext( TestNamespace, ContextA );
@@ -681,13 +681,13 @@ TEST_CASE("Core::Internationalization::FManifestContext::Manifest", "[Core][Seri
 				TSharedPtr< FManifestEntry > FoundEntry4 = TestManifest.FindEntryByContext( TestNamespace, ContextB );
 				TestValid( TEXT("FoundEntry3 is valid"), FoundEntry4);
 
-				TestTrue( TEXT("FoundEntry3 != FoundEntry4"), FoundEntry3 != FoundEntry4 );
+				TEST_TRUE( TEXT("FoundEntry3 != FoundEntry4"), FoundEntry3 != FoundEntry4 );
 
 				// Found entry looked up by source A should match entry looked up by context A
-				TestTrue( TEXT("FoundEntry1 == FoundEntry3"), FoundEntry1 == FoundEntry3 );
+				TEST_TRUE( TEXT("FoundEntry1 == FoundEntry3"), FoundEntry1 == FoundEntry3 );
 
 				// Found entry looked up by source B should match entry looked up by context B
-				TestTrue( TEXT("FoundEntry2 == FoundEntry4"), FoundEntry2 == FoundEntry4 );
+				TEST_TRUE( TEXT("FoundEntry2 == FoundEntry4"), FoundEntry2 == FoundEntry4 );
 
 			}
 
@@ -708,8 +708,8 @@ TEST_CASE("Core::Internationalization::FManifestContext::Manifest", "[Core][Seri
 				//  and should not add an entry.  The reason is that AddSource does an 'exact' match check on the metadata object
 				bool bResult = TestManifest.AddSource( TestNamespace, ConflictingSourceMetadataTypeB, ContextB );
 
-				TestFalse( TEXT("AddSource result == false"), bResult);
-				TestEqual( TEXT("ManifestCount == 1"), CountManifestEntries(TestManifest), 1 );
+				TEST_FALSE( TEXT("AddSource result == false"), bResult);
+				TEST_EQUAL( TEXT("ManifestCount == 1"), CountManifestEntries(TestManifest), 1 );
 
 				// We should be able to find the entry using either Source FLocItem
 				TSharedPtr< FManifestEntry > FoundEntry1 = TestManifest.FindEntryBySource( TestNamespace, ConflictingSourceMetadataTypeA );
@@ -736,8 +736,8 @@ TEST_CASE("Core::Internationalization::FManifestContext::Manifest", "[Core][Seri
 				//  and should not add an entry.  The reason is that AddSource does an 'exact' match check on the metadata object
 				bool bResult = TestManifest.AddSource( TestNamespace, ConflictingSourceMetadataValueB, ContextB );
 
-				TestFalse( TEXT("AddSource result == false"), bResult);
-				TestEqual( TEXT("ManifestCount == 1"), CountManifestEntries(TestManifest), 1 );
+				TEST_FALSE( TEXT("AddSource result == false"), bResult);
+				TEST_EQUAL( TEXT("ManifestCount == 1"), CountManifestEntries(TestManifest), 1 );
 
 				// We should be able to find the entry using either Source FLocItem
 				TSharedPtr< FManifestEntry > FoundEntry1 = TestManifest.FindEntryBySource( TestNamespace, ConflictingSourceMetadataValueA );
@@ -763,18 +763,18 @@ TEST_CASE("Core::Internationalization::FManifestContext::Manifest", "[Core][Seri
 			TestManifest.AddSource( TestNamespace, Source, ContextA );
 			bool bResult = TestManifest.AddSource( TestNamespace, Source, ContextConflictingOptionalFlag);
 
-			TestTrue( TEXT("AddSource result == true"), bResult);
-			TestEqual( TEXT("ManifestCount == 1"), CountManifestEntries(TestManifest), 1 );
+			TEST_TRUE( TEXT("AddSource result == true"), bResult);
+			TEST_EQUAL( TEXT("ManifestCount == 1"), CountManifestEntries(TestManifest), 1 );
 
 			// We should be able to look up the existing entry using the ContextConflictingOptionalFlag context but the entries bIsOptional flag will match ContextA
 			TSharedPtr< FManifestEntry > FoundEntry = TestManifest.FindEntryByContext( TestNamespace, ContextConflictingOptionalFlag );
 			if( !FoundEntry.IsValid() )
 			{
-				TestTrue(TEXT("FManifestEntry could not find entry using FindEntryByContext."), FoundEntry.IsValid());
+				TEST_TRUE(TEXT("FManifestEntry could not find entry using FindEntryByContext."), FoundEntry.IsValid());
 			}
 			else
 			{
-				TestTrue( TEXT("FoundEntry->bIsOptional == ContextA->bIsOptional"), FoundEntry->Contexts[0].bIsOptional == ContextA.bIsOptional );
+				TEST_TRUE( TEXT("FoundEntry->bIsOptional == ContextA->bIsOptional"), FoundEntry->Contexts[0].bIsOptional == ContextA.bIsOptional );
 			}
 		}
 
@@ -800,22 +800,22 @@ TEST_CASE("Core::Internationalization::FManifestContext::Manifest", "[Core][Seri
 			FoundEntry = TestManifest.FindEntryBySource( TestNamespace, Source );
 			if( !FoundEntry.IsValid() )
 			{
-				TestTrue(TEXT("FManifestEntry could not find entry using FindEntryBySource."), FoundEntry.IsValid());
+				TEST_TRUE(TEXT("FManifestEntry could not find entry using FindEntryBySource."), FoundEntry.IsValid());
 			}
 			else
 			{
-				TestTrue( TEXT("FoundEntry->Source == Source"), FoundEntry->Source == Source );
+				TEST_TRUE( TEXT("FoundEntry->Source == Source"), FoundEntry->Source == Source );
 			}
 
 			// Ensure we find the entry we added by context
 			FoundEntry = TestManifest.FindEntryByContext( TestNamespace, ContextC );
 			if( !FoundEntry.IsValid() )
 			{
-				TestTrue(TEXT("FManifestEntry could not find entry using FindEntryByContext."), FoundEntry.IsValid());
+				TEST_TRUE(TEXT("FManifestEntry could not find entry using FindEntryByContext."), FoundEntry.IsValid());
 			}
 			else
 			{
-				TestTrue( TEXT("FoundEntry->Source == Source"), FoundEntry->Source == Source );
+				TEST_TRUE( TEXT("FoundEntry->Source == Source"), FoundEntry->Source == Source );
 			}
 		}
 
@@ -841,22 +841,22 @@ TEST_CASE("Core::Internationalization::FManifestContext::Manifest", "[Core][Seri
 			FoundEntry = TestManifest.FindEntryBySource( TestNamespace, Source );
 			if( !FoundEntry.IsValid() )
 			{
-				TestTrue(TEXT("FManifestEntry could not find entry using FindEntryBySource."), FoundEntry.IsValid());
+				TEST_TRUE(TEXT("FManifestEntry could not find entry using FindEntryBySource."), FoundEntry.IsValid());
 			}
 			else
 			{
-				TestTrue( TEXT("FoundEntry->Source == Source"), FoundEntry->Source == Source );
+				TEST_TRUE( TEXT("FoundEntry->Source == Source"), FoundEntry->Source == Source );
 			}
 
 			// Ensure we find the entry we added by context
 			FoundEntry = TestManifest.FindEntryByContext( TestNamespace, ContextC );
 			if( !FoundEntry.IsValid() )
 			{
-				TestTrue(TEXT("FManifestEntry could not find entry using FindEntryByContext."), FoundEntry.IsValid());
+				TEST_TRUE(TEXT("FManifestEntry could not find entry using FindEntryByContext."), FoundEntry.IsValid());
 			}
 			else
 			{
-				TestTrue( TEXT("FoundEntry->Source == Source"), FoundEntry->Source == Source );
+				TEST_TRUE( TEXT("FoundEntry->Source == Source"), FoundEntry->Source == Source );
 			}
 		}
 	}
diff --git a/Engine/Source/Runtime/Core/Tests/Internationalization/MetadataTest.cpp b/Engine/Source/Runtime/Core/Tests/Internationalization/MetadataTest.cpp
--- a/Engine/Source/Runtime/Core/Tests/Internationalization/MetadataTest.cpp
+++ b/Engine/Source/Runtime/Core/Tests/Internationalization/MetadataTest.cpp
@@ -55,309 +55,320 @@ TEST_CASE("Core::Internationalization::FLocMetadataValue::Metadata", "[Core][Int
 	TSharedPtr< FLocMetadataValue > MetadataValueObjectA = MakeShareable( new FLocMetadataValueObject( MetadataObjectA ) );
 	TSharedPtr< FLocMetadataValue > MetadataValueObjectB = MakeShareable( new FLocMetadataValueObject( MetadataObjectB ) );
 
-	// Testing the bool meta data value type
+	SECTION("Testing the bool meta data value type")
 	{
-		TestFalse( TEXT("MetadataValueBoolFalse == MetadataValueBoolTrue"), *MetadataValueBoolFalse == *MetadataValueBoolTrue );
-		TestTrue( TEXT("MetadataValueBoolFalse < MetadataValueBoolTrue"), *MetadataValueBoolFalse < *MetadataValueBoolTrue );
-		TestFalse( TEXT("MetadataValueBoolTrue < MetadataValueBoolFalse"), *MetadataValueBoolTrue < *MetadataValueBoolFalse);
+		CHECK_FALSE((*MetadataValueBoolFalse == *MetadataValueBoolTrue));
+		CHECK((*MetadataValueBoolFalse < *MetadataValueBoolTrue));
+		CHECK_FALSE((* MetadataValueBoolTrue < *MetadataValueBoolFalse));
 
-		TestTrue( TEXT("MetadataValueBoolFalse < MetadataValueString"), *MetadataValueBoolFalse < *MetadataValueStringA );
-		TestTrue( TEXT("MetadataValueBoolTrue < MetadataValueString"), *MetadataValueBoolTrue < *MetadataValueStringA);
+		CHECK((* MetadataValueBoolFalse < *MetadataValueStringA));
+		CHECK((*MetadataValueBoolTrue < *MetadataValueStringA));
 
-		TestTrue( TEXT("MetadataValueBoolFalse < MetadataValueArray"), *MetadataValueBoolFalse < *MetadataValueArrayA );
-		TestTrue( TEXT("MetadataValueBoolTrue < MetadataValueArray"), *MetadataValueBoolTrue < *MetadataValueArrayA);
+		CHECK((*MetadataValueBoolFalse < *MetadataValueArrayA));
+		CHECK((*MetadataValueBoolTrue < *MetadataValueArrayA));
 
-		TestTrue( TEXT("MetadataValueBoolFalse < MetadataValueObject"), *MetadataValueBoolFalse < *MetadataValueObjectA );
-		TestTrue( TEXT("MetadataValueBoolTrue < MetadataValueObject"), *MetadataValueBoolTrue < *MetadataValueObjectA);
+		CHECK((*MetadataValueBoolFalse < *MetadataValueObjectA));
+		CHECK((*MetadataValueBoolTrue < *MetadataValueObjectA));
 
 		TSharedPtr<FLocMetadataValue> MetadataValueBoolFalseClone = MetadataValueBoolFalse->Clone();
 		TSharedPtr<FLocMetadataValue> MetadataValueBoolTrueClone = MetadataValueBoolTrue->Clone();
 
 		if( MetadataValueBoolFalse == MetadataValueBoolFalseClone )
 		{
-			TestAddError(TEXT("MetadataValueBool and its Clone are not unique objects."));
+			FAIL_CHECK(TEXT("MetadataValueBool and its Clone are not unique objects."));
 		}
 		
-		TestTrue( TEXT("MetadataValueBoolFalseClone == MetadataValueBoolFalse"), *MetadataValueBoolFalseClone == *MetadataValueBoolFalse );
-		TestFalse( TEXT("MetadataValueBoolFalseClone < MetadataValueBoolFalse"), *MetadataValueBoolFalseClone < *MetadataValueBoolFalse );
+		CHECK((*MetadataValueBoolFalseClone == *MetadataValueBoolFalse));
+		CHECK_FALSE((*MetadataValueBoolFalseClone < *MetadataValueBoolFalse));
 
-		TestTrue( TEXT("MetadataValueBoolTrueClone == MetadataValueBoolTrue"), *MetadataValueBoolTrueClone == *MetadataValueBoolTrue );
-		TestFalse( TEXT("MetadataValueBoolTrueClone < MetadataValueBoolTrue"), *MetadataValueBoolTrueClone < *MetadataValueBoolTrue );
+		CHECK((* MetadataValueBoolTrueClone == *MetadataValueBoolTrue));
+		CHECK_FALSE((* MetadataValueBoolTrueClone < *MetadataValueBoolTrue));
 
-		// Test the bool metadata when it is part of an object
-		TSharedPtr< FLocMetadataObject > MetadataObjectFalse = MakeShareable( new FLocMetadataObject );
-		MetadataObjectFalse->SetField(TEXT("MetadataValueBool"), MetadataValueBoolFalse);
+		SECTION("Test the bool metadata when it is part of an object")
+		{
+			TSharedPtr< FLocMetadataObject > MetadataObjectFalse = MakeShareable(new FLocMetadataObject);
+			MetadataObjectFalse->SetField(TEXT("MetadataValueBool"), MetadataValueBoolFalse);
 
-		TSharedPtr< FLocMetadataObject > MetadataObjectTrue = MakeShareable( new FLocMetadataObject );
-		MetadataObjectTrue->SetField(TEXT("MetadataValueBool"), MetadataValueBoolTrue);
+			TSharedPtr< FLocMetadataObject > MetadataObjectTrue = MakeShareable(new FLocMetadataObject);
+			MetadataObjectTrue->SetField(TEXT("MetadataValueBool"), MetadataValueBoolTrue);
 
-		TestFalse( TEXT("GetBoolField(MetadataValueBool)"), MetadataObjectFalse->GetBoolField( TEXT("MetadataValueBool") ) );
-		TestTrue( TEXT("GetBoolField(MetadataValueBool)"), MetadataObjectTrue->GetBoolField( TEXT("MetadataValueBool") ) );
+			CHECK_FALSE(MetadataObjectFalse->GetBoolField(TEXT("MetadataValueBool")));
+			CHECK(MetadataObjectTrue->GetBoolField(TEXT("MetadataValueBool")));
 
-		TestFalse( TEXT("MetadataObjectFalse == MetadataObjectTrue"), *MetadataObjectFalse == *MetadataObjectTrue );
-		TestTrue( TEXT("MetadataObjectFalse < MetadataObjectTrue"), *MetadataObjectFalse < *MetadataObjectTrue );
+			CHECK_FALSE(*MetadataObjectFalse == *MetadataObjectTrue);
+			CHECK(*MetadataObjectFalse < *MetadataObjectTrue);
+		}
 	}
 
-	// Testing string metadata value type
+	SECTION("Testing string metadata value type")
 	{
-		TestFalse( TEXT("MetadataValueStringA == MetadataValueStringB"), *MetadataValueStringA == *MetadataValueStringB );
-		TestTrue( TEXT("MetadataValueStringA < MetadataValueStringB"), *MetadataValueStringA < *MetadataValueStringB );
-		TestFalse( TEXT("MetadataValueStringB < MetadataValueStringA"), *MetadataValueStringB < *MetadataValueStringA);
+		CHECK_FALSE((* MetadataValueStringA == *MetadataValueStringB));
+		CHECK((*MetadataValueStringA < *MetadataValueStringB));
+		CHECK_FALSE((* MetadataValueStringB < *MetadataValueStringA));
 
-		TestTrue( TEXT("MetadataValueString < MetadataValueArray"), *MetadataValueStringA < *MetadataValueArrayA );
+		CHECK((*MetadataValueStringA < *MetadataValueArrayA));
 
-		TestTrue( TEXT("MetadataValueStringA < MetadataValueObject"), *MetadataValueStringA < *MetadataValueObjectA );
+		CHECK((*MetadataValueStringA < *MetadataValueObjectA));
 
 
 		TSharedPtr<FLocMetadataValue> MetadataValueStringAClone = MetadataValueStringA->Clone();
 
 		if( MetadataValueStringA == MetadataValueStringAClone )
 		{
-			TestAddError(TEXT("MetadataValueString and its Clone are not unique objects."));
+			FAIL_CHECK("MetadataValueString and its Clone are not unique objects.");
 		}
 
-		TestTrue( TEXT("MetadataValueStringAClone == MetadataValueStringA"), *MetadataValueStringAClone == *MetadataValueStringA );
-		TestFalse( TEXT("MetadataValueStringAClone < MetadataValueStringA"), *MetadataValueStringAClone < *MetadataValueStringA );
-		TestTrue( TEXT("MetadataValueStringAClone < MetadataValueStringB"), *MetadataValueStringAClone < *MetadataValueStringB );
+		CHECK((*MetadataValueStringAClone == *MetadataValueStringA));
+		CHECK_FALSE((* MetadataValueStringAClone < *MetadataValueStringA));
+		CHECK((* MetadataValueStringAClone < *MetadataValueStringB));
+
+		SECTION("Test the string metadata when it is part of an object")
+		{
+			TSharedPtr< FLocMetadataObject > TestMetadataObjectA = MakeShareable(new FLocMetadataObject);
+			TestMetadataObjectA->SetField(TEXT("MetadataValueString"), MetadataValueStringA);
 
-		// Test the string metadata when it is part of an object
-		TSharedPtr< FLocMetadataObject > TestMetadataObjectA = MakeShareable( new FLocMetadataObject );
-		TestMetadataObjectA->SetField(TEXT("MetadataValueString"), MetadataValueStringA);
+			TSharedPtr< FLocMetadataObject > TestMetadataObjectB = MakeShareable(new FLocMetadataObject);
+			TestMetadataObjectB->SetField(TEXT("MetadataValueString"), MetadataValueStringB);
 
-		TSharedPtr< FLocMetadataObject > TestMetadataObjectB = MakeShareable( new FLocMetadataObject );
-		TestMetadataObjectB->SetField(TEXT("MetadataValueString"), MetadataValueStringB);
+			CHECK(TestMetadataObjectA->GetStringField(TEXT("MetadataValueString")) == StringA);
 
-		TestTrue( TEXT("GetStringField(MetadataValueString) == StringA"), TestMetadataObjectA->GetStringField( TEXT("MetadataValueString") ) == StringA );
-	
-		TestFalse( TEXT("TestMetadataObjectA == TestMetadataObjectB"), *TestMetadataObjectA == *TestMetadataObjectB );
-		TestTrue( TEXT("TestMetadataObjectA < TestMetadataObjectB"), *TestMetadataObjectA < *TestMetadataObjectB );
+			CHECK_FALSE(*TestMetadataObjectA == *TestMetadataObjectB);
+			CHECK(*TestMetadataObjectA < *TestMetadataObjectB);
+		}
 	}
 
-	// Testing array metadata value type
+	SECTION("Testing array metadata value type")
 	{
-		TestFalse( TEXT("MetadataValueArrayA == MetadataValueArrayB"), *MetadataValueArrayA == *MetadataValueArrayB );
-		TestTrue( TEXT("MetadataValueArrayA < MetadataValueArrayB"), *MetadataValueArrayA < *MetadataValueArrayB );
-		TestFalse( TEXT("MetadataValueArrayB < MetadataValueArrayA"), *MetadataValueArrayB < *MetadataValueArrayA);
+		CHECK_FALSE((* MetadataValueArrayA == *MetadataValueArrayB));
+		CHECK((*MetadataValueArrayA < *MetadataValueArrayB));
+		CHECK_FALSE((* MetadataValueArrayB < *MetadataValueArrayA));
 
-		TestTrue( TEXT("MetadataValueArrayA < MetadataValueObject"), *MetadataValueArrayA < *MetadataValueObjectA );
+		CHECK((* MetadataValueArrayA < *MetadataValueObjectA));
 
 		TSharedPtr<FLocMetadataValue> MetadataValueArrayAClone = MetadataValueArrayA->Clone();
 
 		if( MetadataValueArrayA == MetadataValueArrayAClone )
 		{
-			TestAddError(TEXT("MetadataValueString and its Clone are not unique objects."));
+			FAIL_CHECK(TEXT("MetadataValueString and its Clone are not unique objects."));
+		}
+
+		CHECK((* MetadataValueArrayAClone == *MetadataValueArrayA));
+		CHECK_FALSE((* MetadataValueArrayAClone < *MetadataValueArrayA));
+		CHECK((* MetadataValueArrayAClone < *MetadataValueArrayB));
+
+		SECTION("Test Less than and equality checks")
+		{
+			//  Metadata arrays are equivalent if they contain equivalent contents in any order.
+			//  To calculate if a metadata array is less than another, we sort both arrays and check each entry index against its counterpart.
+			//  If we encounter an entry that is less than another we stop looking.
+			TArray< TSharedPtr< FLocMetadataValue > > ArrayC;
+			ArrayC.Add(MetadataValueBoolFalse);
+			ArrayC.Add(MetadataValueBoolFalse->Clone());
+			TSharedPtr<FLocMetadataValue> MetadataValueArrayC = MakeShareable(new FLocMetadataValueArray(ArrayC));
+
+			CHECK_FALSE((*MetadataValueArrayA == *MetadataValueArrayC));
+			CHECK((* MetadataValueArrayC < *MetadataValueArrayA));
+			CHECK((* MetadataValueArrayC < *MetadataValueArrayB));
+
+			TArray< TSharedPtr< FLocMetadataValue > > ArrayD;
+			ArrayD.Add(MetadataValueBoolFalse);
+			ArrayD.Add(MetadataValueBoolFalse->Clone());
+			ArrayD.Add(MetadataValueBoolFalse->Clone());
+			TSharedPtr<FLocMetadataValue> MetadataValueArrayD = MakeShareable(new FLocMetadataValueArray(ArrayD));
+
+			CHECK_FALSE((*MetadataValueArrayA == *MetadataValueArrayD));
+			CHECK_FALSE((* MetadataValueArrayC == *MetadataValueArrayD));
+			CHECK((*MetadataValueArrayC < *MetadataValueArrayD));
+			CHECK((*MetadataValueArrayD < *MetadataValueArrayA));
 		}
 
-		TestTrue( TEXT("MetadataValueArrayAClone == MetadataValueArrayA"), *MetadataValueArrayAClone == *MetadataValueArrayA );
-		TestFalse( TEXT("MetadataValueArrayAClone < MetadataValueArrayA"), *MetadataValueArrayAClone < *MetadataValueArrayA );
-		TestTrue( TEXT("MetadataValueArrayAClone < MetadataValueArrayB"), *MetadataValueArrayAClone < *MetadataValueArrayB );
-
-		// Test Less than and equality checks.  Metadata arrays are equivalent if they contain equivalent contents in any order.  
-		//  To calculate if a metadata array is less than another, we sort both arrays and check each entry index against its counterpart.
-		//  If we encounter an entry that is less than another we stop looking.
-		TArray< TSharedPtr< FLocMetadataValue > > ArrayC;
-		ArrayC.Add( MetadataValueBoolFalse );
-		ArrayC.Add( MetadataValueBoolFalse->Clone() );
-		TSharedPtr<FLocMetadataValue> MetadataValueArrayC = MakeShareable( new FLocMetadataValueArray( ArrayC ) );
-	
-		TestFalse( TEXT("MetadataValueArrayA == MetadataValueArrayC"), *MetadataValueArrayA == *MetadataValueArrayC );
-		TestTrue( TEXT("MetadataValueArrayC < MetadataValueArrayA"), *MetadataValueArrayC < *MetadataValueArrayA );
-		TestTrue( TEXT("MetadataValueArrayC < MetadataValueArrayB"), *MetadataValueArrayC < *MetadataValueArrayB );
-
-		TArray< TSharedPtr< FLocMetadataValue > > ArrayD;
-		ArrayD.Add( MetadataValueBoolFalse );
-		ArrayD.Add( MetadataValueBoolFalse->Clone() );
-		ArrayD.Add( MetadataValueBoolFalse->Clone() );
-		TSharedPtr<FLocMetadataValue> MetadataValueArrayD = MakeShareable( new FLocMetadataValueArray( ArrayD ) );
-
-		TestFalse( TEXT("MetadataValueArrayA == MetadataValueArrayD"), *MetadataValueArrayA == *MetadataValueArrayD );
-		TestFalse( TEXT("MetadataValueArrayC == MetadataValueArrayD"), *MetadataValueArrayC == *MetadataValueArrayD );
-		TestTrue( TEXT("MetadataValueArrayC < MetadataValueArrayD"), *MetadataValueArrayC < *MetadataValueArrayD );
-		TestTrue( TEXT("MetadataValueArrayD < MetadataValueArrayA"), *MetadataValueArrayD < *MetadataValueArrayA );
-
-		// Test the array metadata when it is part of an object
-		TSharedPtr< FLocMetadataObject > TestMetadataObjectA = MakeShareable( new FLocMetadataObject );
-		TestMetadataObjectA->SetField(TEXT("MetadataValueArray"), MetadataValueArrayA);
-
-		TSharedPtr< FLocMetadataObject > TestMetadataObjectB = MakeShareable( new FLocMetadataObject );
-		TestMetadataObjectB->SetField(TEXT("MetadataValueArray"), MetadataValueArrayB);
-
-		TestTrue( TEXT("GetArrayField(MetadataValueArray) == ArrayA"), TestMetadataObjectA->GetArrayField( TEXT("MetadataValueArray") ) == ArrayA );
-
-		TestFalse( TEXT("TestMetadataObjectA == TestMetadataObjectB"), *TestMetadataObjectA == *TestMetadataObjectB );
-		TestTrue( TEXT("TestMetadataObjectA < TestMetadataObjectB"), *TestMetadataObjectA < *TestMetadataObjectB );
+		SECTION("Test the array metadata when it is part of an object")
+		{
+			TSharedPtr< FLocMetadataObject > TestMetadataObjectA = MakeShareable(new FLocMetadataObject);
+			TestMetadataObjectA->SetField(TEXT("MetadataValueArray"), MetadataValueArrayA);
+
+			TSharedPtr< FLocMetadataObject > TestMetadataObjectB = MakeShareable(new FLocMetadataObject);
+			TestMetadataObjectB->SetField(TEXT("MetadataValueArray"), MetadataValueArrayB);
+
+			CHECK(TestMetadataObjectA->GetArrayField(TEXT("MetadataValueArray")) == ArrayA);
+
+			CHECK_FALSE(*TestMetadataObjectA == *TestMetadataObjectB);
+			CHECK(*TestMetadataObjectA < *TestMetadataObjectB);
+		}
 	}
 
-	// Testing object metadata value type
+	SECTION("Testing object metadata value type")
 	{
-		TestFalse( TEXT("MetadataValueObjectA == MetadataValueObjectB"), *MetadataValueObjectA == *MetadataValueObjectB );
-		TestTrue( TEXT("MetadataValueObjectA < MetadataValueObjectB"), *MetadataValueObjectA < *MetadataValueObjectB );
-		TestFalse( TEXT("MetadataValueObjectB < MetadataValueObjectA"), *MetadataValueObjectB < *MetadataValueObjectA);
+		CHECK_FALSE((* MetadataValueObjectA == *MetadataValueObjectB));
+		CHECK((* MetadataValueObjectA < *MetadataValueObjectB));
+		CHECK_FALSE((* MetadataValueObjectB < *MetadataValueObjectA));
 
 		TSharedPtr<FLocMetadataValue> MetadataValueObjectAClone = MetadataValueObjectA->Clone();
 
-		TestFalse(TEXT("MetadataValueString and its Clone are not unique objects."), MetadataValueObjectA == MetadataValueObjectAClone);
+		CHECK_FALSE(MetadataValueObjectA == MetadataValueObjectAClone);
 		
-		TestTrue( TEXT("MetadataValueObjectAClone == MetadataValueObjectA"), *MetadataValueObjectAClone == *MetadataValueObjectA );
-		TestFalse( TEXT("MetadataValueObjectAClone < MetadataValueObjectA"), *MetadataValueObjectAClone < *MetadataValueObjectA );
-		TestTrue( TEXT("MetadataValueObjectAClone < MetadataValueObjectB"), *MetadataValueObjectAClone < *MetadataValueObjectB );
+		CHECK((*MetadataValueObjectAClone == *MetadataValueObjectA));
+		CHECK_FALSE((* MetadataValueObjectAClone < *MetadataValueObjectA));
+		CHECK((* MetadataValueObjectAClone < *MetadataValueObjectB));
 
 
-		// Test the object metadata when it is part of another object
-		TSharedPtr< FLocMetadataObject > TestMetadataObjectA = MakeShareable( new FLocMetadataObject );
-		TestMetadataObjectA->SetField(TEXT("MetadataValueObject"), MetadataValueObjectA);
+		SECTION("Test the object metadata when it is part of another object")
+		{
+			TSharedPtr< FLocMetadataObject > TestMetadataObjectA = MakeShareable(new FLocMetadataObject);
+			TestMetadataObjectA->SetField(TEXT("MetadataValueObject"), MetadataValueObjectA);
 
-		TSharedPtr< FLocMetadataObject > TestMetadataObjectB = MakeShareable( new FLocMetadataObject );
-		TestMetadataObjectB->SetField(TEXT("MetadataValueObject"), MetadataValueObjectB);
+			TSharedPtr< FLocMetadataObject > TestMetadataObjectB = MakeShareable(new FLocMetadataObject);
+			TestMetadataObjectB->SetField(TEXT("MetadataValueObject"), MetadataValueObjectB);
 
-		TestTrue( TEXT("GetObjectField(MetadataValueObject) == MetadataObjectA"), *TestMetadataObjectA->GetObjectField( TEXT("MetadataValueObject") ) == *MetadataObjectA );
+			CHECK(*TestMetadataObjectA->GetObjectField(TEXT("MetadataValueObject")) == *MetadataObjectA);
 
-		TestFalse( TEXT("TestMetadataObjectA == TestMetadataObjectB"), *TestMetadataObjectA == *TestMetadataObjectB );
-		TestTrue( TEXT("TestMetadataObjectA < TestMetadataObjectB"), *TestMetadataObjectA < *TestMetadataObjectB );
+			CHECK_FALSE(*TestMetadataObjectA == *TestMetadataObjectB);
+			CHECK(*TestMetadataObjectA < *TestMetadataObjectB);
+		}
 	}
 
-	// Testing Loc Metadata Object
+	SECTION("Testing Loc Metadata Object")
 	{
-		TestFalse( TEXT("MetadataObjectA == MetadataObjectB"), *MetadataObjectA == *MetadataObjectB );
-		TestTrue( TEXT("MetadataObjectA < MetadataObjectB"), *MetadataObjectA < *MetadataObjectB );
-		TestFalse( TEXT("MetadataObjectB < MetadataObjectA"), *MetadataObjectB < *MetadataObjectA);
+		CHECK_FALSE(*MetadataObjectA == *MetadataObjectB );
+		CHECK(*MetadataObjectA < *MetadataObjectB );
+		CHECK_FALSE(*MetadataObjectB < *MetadataObjectA);
 
-		// Test copy ctor
+		SECTION("Test copy ctor")
 		{
 			FLocMetadataObject MetadataObjectAClone = *MetadataObjectA;
-			TestTrue( TEXT("MetadataObjectAClone == MetadataObjectA"), MetadataObjectAClone == *MetadataObjectA );
+			CHECK(MetadataObjectAClone == *MetadataObjectA );
 		}
 		
-		// Test assignment
+		SECTION("Test assignment")
 		{
 			FLocMetadataObject MetadataObjectAClone = *MetadataObjectB;
 			// PVS-Studio complains about double initialization, but that's something we're testing
 			// so we disable the warning
 			MetadataObjectAClone = *MetadataObjectA; //-V519
-			TestTrue( TEXT("MetadataObjectAClone == MetadataObjectA"), MetadataObjectAClone == *MetadataObjectA );
-			TestFalse( TEXT("MetadataObjectAClone == MetadataObjectB"), MetadataObjectAClone == *MetadataObjectB );
+			CHECK(MetadataObjectAClone == *MetadataObjectA );
+			CHECK_FALSE(MetadataObjectAClone == *MetadataObjectB );
 		}
 		
-		// Test comparison operator  
+		SECTION("Test comparison operator")
 		{
 			FLocMetadataObject MetadataObjectAClone = *MetadataObjectA;
 			
 			
 			// Adding standard entry
 			MetadataObjectAClone.SetStringField( TEXT("NewEntry"), TEXT("NewEntryValue") );
-			TestFalse( TEXT("MetadataObjectAClone == MetadataObjectA"), MetadataObjectAClone == *MetadataObjectA );
+			CHECK_FALSE(MetadataObjectAClone == *MetadataObjectA );
 			
 			// Adding non-standard entry.  Note metadata with * prefix in the name will ignore value and type when performing comparisons
 			MetadataObjectAClone = *MetadataObjectA;
 			MetadataObjectAClone.SetStringField( TEXT("*NewEntry"), TEXT("*NewEntryValue") );
-			TestFalse( TEXT("MetadataObjectAClone == MetadataObjectA"), MetadataObjectAClone == *MetadataObjectA );
+			CHECK_FALSE( MetadataObjectAClone == *MetadataObjectA );
 
 			// Value mismatch on entry with * prefix with same type
 			MetadataObjectAClone = *MetadataObjectA;
 			MetadataObjectAClone.SetStringField( TEXT("*NoCompare"), TEXT("NoCompare") );
 			FLocMetadataObject MetadataObjectAClone2 = *MetadataObjectA;
 			MetadataObjectAClone2.SetStringField( TEXT("*NoCompare"), TEXT("NoCompare2") );
-			TestTrue( TEXT("MetadataObjectAClone == MetadataObjectAClone2"), MetadataObjectAClone == MetadataObjectAClone2 );
+			CHECK(MetadataObjectAClone == MetadataObjectAClone2 );
 
 			// Value and type mismatch on entry with * prefix
 			MetadataObjectAClone = *MetadataObjectA;
 			MetadataObjectAClone.SetStringField( TEXT("*NoCompare"), TEXT("NoCompare") );
 			MetadataObjectAClone2 = *MetadataObjectA;
 			MetadataObjectAClone2.SetBoolField( TEXT("*NoCompare"), true );
-			TestTrue( TEXT("MetadataObjectAClone == MetadataObjectAClone2"), MetadataObjectAClone == MetadataObjectAClone2 );
+			CHECK(MetadataObjectAClone == MetadataObjectAClone2 );
 		
 			// Value mismatch on standard entry
 			MetadataObjectAClone = *MetadataObjectA;
 			MetadataObjectAClone.SetStringField( TEXT("DoCompare"), TEXT("DoCompare") );
 			MetadataObjectAClone2 = *MetadataObjectA;
 			MetadataObjectAClone2.SetStringField( TEXT("DoCompare"), TEXT("DoCompare2") );
-			TestFalse( TEXT("MetadataObjectAClone == MetadataObjectAClone2"), MetadataObjectAClone == MetadataObjectAClone2 );
+			CHECK_FALSE(MetadataObjectAClone == MetadataObjectAClone2 );
 
 			// Value and type mismatch on standard entry
 			MetadataObjectAClone = *MetadataObjectA;
 			MetadataObjectAClone.SetStringField( TEXT("DoCompare"), TEXT("DoCompare") );
 			MetadataObjectAClone2 = *MetadataObjectA;
 			MetadataObjectAClone2.SetBoolField( TEXT("DoCompare"), true );
-			TestFalse( TEXT("MetadataObjectAClone == MetadataObjectAClone2"), MetadataObjectAClone == MetadataObjectAClone2 );
+			CHECK_FALSE( MetadataObjectAClone == MetadataObjectAClone2 );
 		}		
 		
-		// Test IsExactMatch function.  Note: Differs from the == operator which takes into account the COMPARISON_MODIFIER_PREFIX
+		SECTION("Test IsExactMatch function.  Note: Differs from the == operator which takes into account the COMPARISON_MODIFIER_PREFIX")
 		{
 			// This function behaves similar to the comparison operator except that it ensures all metadata entries match exactly.  In other
 			//  words, it will also perform exact match checks on * prefixed metadata items.
 			FLocMetadataObject MetadataObjectAClone = *MetadataObjectA;
 
-			TestTrue( TEXT("MetadataObjectAClone == MetadataObjectA"), MetadataObjectAClone.IsExactMatch( *MetadataObjectA) );
+			CHECK(MetadataObjectAClone.IsExactMatch( *MetadataObjectA) );
 
 			// Adding standard entry
 			MetadataObjectAClone.SetStringField( TEXT("NewEntry"), TEXT("NewEntryValue") );
-			TestFalse( TEXT("MetadataObjectAClone == MetadataObjectA"), MetadataObjectAClone.IsExactMatch( *MetadataObjectA) );
+			CHECK_FALSE(MetadataObjectAClone.IsExactMatch( *MetadataObjectA) );
 
 			// Adding non-standard entry.  Note metadata with * prefix in the name will ignore value and type when performing comparisons
 			MetadataObjectAClone = *MetadataObjectA;
 			MetadataObjectAClone.SetStringField( TEXT("*NewEntry"), TEXT("*NewEntryValue") );
-			TestFalse( TEXT("MetadataObjectAClone == MetadataObjectA"),MetadataObjectAClone.IsExactMatch( *MetadataObjectA)  );
+			CHECK_FALSE(MetadataObjectAClone.IsExactMatch( *MetadataObjectA)  );
 
 			// Value mismatch on entry with * prefix with same type
 			MetadataObjectAClone = *MetadataObjectA;
 			MetadataObjectAClone.SetStringField( TEXT("*NoCompare"), TEXT("NoCompare") );
 			FLocMetadataObject MetadataObjectAClone2 = *MetadataObjectA;
 			MetadataObjectAClone2.SetStringField( TEXT("*NoCompare"), TEXT("NoCompare2") );
-			TestFalse( TEXT("MetadataObjectAClone == MetadataObjectAClone2"), MetadataObjectAClone.IsExactMatch( MetadataObjectAClone2) );
+			CHECK_FALSE(MetadataObjectAClone.IsExactMatch( MetadataObjectAClone2) );
 
 			// Value and type mismatch on entry with * prefix
 			MetadataObjectAClone = *MetadataObjectA;
 			MetadataObjectAClone.SetStringField( TEXT("*NoCompare"), TEXT("NoCompare") );
 			MetadataObjectAClone2 = *MetadataObjectA;
 			MetadataObjectAClone2.SetBoolField( TEXT("*NoCompare"), true );
-			TestFalse( TEXT("MetadataObjectAClone == MetadataObjectAClone2"),  MetadataObjectAClone.IsExactMatch( MetadataObjectAClone2) );
+			CHECK_FALSE(MetadataObjectAClone.IsExactMatch( MetadataObjectAClone2) );
 
 			// Value mismatch on standard entry
 			MetadataObjectAClone = *MetadataObjectA;
 			MetadataObjectAClone.SetStringField( TEXT("DoCompare"), TEXT("DoCompare") );
 			MetadataObjectAClone2 = *MetadataObjectA;
 			MetadataObjectAClone2.SetStringField( TEXT("DoCompare"), TEXT("DoCompare2") );
-			TestFalse( TEXT("MetadataObjectAClone == MetadataObjectAClone2"),  MetadataObjectAClone.IsExactMatch( MetadataObjectAClone2) );
+			CHECK_FALSE(MetadataObjectAClone.IsExactMatch( MetadataObjectAClone2) );
 
 			// Value and type mismatch on standard entry
 			MetadataObjectAClone = *MetadataObjectA;
 			MetadataObjectAClone.SetStringField( TEXT("DoCompare"), TEXT("DoCompare") );
 			MetadataObjectAClone2 = *MetadataObjectA;
 			MetadataObjectAClone2.SetBoolField( TEXT("DoCompare"), true );
-			TestFalse( TEXT("MetadataObjectAClone == MetadataObjectAClone2"),  MetadataObjectAClone.IsExactMatch( MetadataObjectAClone2) );
+			CHECK_FALSE(MetadataObjectAClone.IsExactMatch( MetadataObjectAClone2) );
 		}
 
-		// Test less than operator
+		SECTION("Test less than operator")
 		{
 			
 			// Adding standard entry that would appear before other entries
 			FLocMetadataObject MetadataObjectAClone = *MetadataObjectA;
 			MetadataObjectAClone.SetStringField( TEXT("ANewEntry"), TEXT("NewEntryValue") );
-			TestTrue( TEXT("MetadataObjectAClone < MetadataObjectA"), MetadataObjectAClone < *MetadataObjectA );
+			CHECK(MetadataObjectAClone < *MetadataObjectA );
 			
 			// Adding standard entry that would appear after other entries
 			MetadataObjectAClone = *MetadataObjectA;
 			MetadataObjectAClone.SetStringField( TEXT("ZNewEntry"), TEXT("NewEntryValue") );
-			TestTrue( TEXT("MetadataObjectA < MetadataObjectAClone"), *MetadataObjectA < MetadataObjectAClone );
+			CHECK(*MetadataObjectA < MetadataObjectAClone );
 
 			// Adding non-standard entry that would appear before other entries.  Note metadata with * prefix has no special treatment in less than operator
 			MetadataObjectAClone = *MetadataObjectA;
 			MetadataObjectAClone.SetStringField( TEXT("*NewEntry"), TEXT("NewEntryValue") );
-			TestTrue( TEXT("MetadataObjectAClone < MetadataObjectA"), MetadataObjectAClone < *MetadataObjectA );
+			CHECK(MetadataObjectAClone < *MetadataObjectA );
 
 			// Value mismatch on entry with * prefix with same type. Note metadata with * prefix has no special treatment in less than operator
 			MetadataObjectAClone = *MetadataObjectA;
 			MetadataObjectAClone.SetStringField( TEXT("*NoCompare"), TEXT("NoCompare") );
 			FLocMetadataObject MetadataObjectAClone2 = *MetadataObjectA;
 			MetadataObjectAClone2.SetStringField( TEXT("*NoCompare"), TEXT("NoCompare2") );
-			TestTrue( TEXT("MetadataObjectAClone < MetadataObjectAClone2"), MetadataObjectAClone < MetadataObjectAClone2 );
+			CHECK(MetadataObjectAClone < MetadataObjectAClone2 );
 
 			// Value and type mismatch on entry with * prefix.  Note metadata with * prefix has no special treatment in less than operator
 			MetadataObjectAClone = *MetadataObjectA;
 			MetadataObjectAClone.SetBoolField( TEXT("*NoCompare"), true );
 			MetadataObjectAClone2 = *MetadataObjectA;
 			MetadataObjectAClone2.SetStringField( TEXT("*NoCompare"), TEXT("NoCompare") );
-			TestTrue( TEXT("MetadataObjectAClone < MetadataObjectAClone2"), MetadataObjectAClone < MetadataObjectAClone2 );
+			CHECK(MetadataObjectAClone < MetadataObjectAClone2 );
 		}
 	}
 }
