Add new Unit Tests

diff --git a/Engine/Source/Runtime/Core/Tests/Math/RangeBoundTest.cpp b/Engine/Source/Runtime/Core/Tests/Math/RangeBoundTest.cpp
--- a/Engine/Source/Runtime/Core/Tests/Math/RangeBoundTest.cpp
+++ b/Engine/Source/Runtime/Core/Tests/Math/RangeBoundTest.cpp
@@ -21,25 +21,25 @@ TEST_CASE("Core::Math::FFloatRangeBound::Smoke Test", "[Core][Math][Smoke]")
 	FFloatRangeBound b1_3 = FFloatRangeBound::Open();
 	FFloatRangeBound b1_4 = 2;
 
-	TestTrue(TEXT("Exclusive bound constructor must create exclusive bound"), b1_1.IsExclusive());
-	TestTrue(TEXT("Exclusive bound constructor must create closed bound"), b1_1.IsClosed());
-	TestFalse(TEXT("Exclusive bound constructor must not create inclusive bound"), b1_1.IsInclusive());
-	TestFalse(TEXT("Exclusive bound constructor must not create open bound"), b1_1.IsOpen());
-	TestEqual(TEXT("Exclusive bound constructor must create the correct value"), b1_1.GetValue(), 2.0f);
-
-	TestTrue(TEXT("Inclusive bound constructor must create exclusive bound"), b1_2.IsInclusive());
-	TestTrue(TEXT("Inclusive bound constructor must create closed bound"), b1_2.IsClosed());
-	TestFalse(TEXT("Inclusive bound constructor must not create exclusive bound"), b1_2.IsExclusive());
-	TestFalse(TEXT("Inclusive bound constructor must not create open bound"), b1_2.IsOpen());
-	TestEqual(TEXT("Inclusive bound constructor must create the correct value"), b1_2.GetValue(), 2.0f);
-
-	TestTrue(TEXT("Open bound constructor must create open bound"), b1_3.IsOpen());
-	TestFalse(TEXT("Open bound constructor must not create closed bound"), b1_3.IsClosed());
-	TestFalse(TEXT("Open bound constructor must not create exclusive bound"), b1_3.IsExclusive());
-	TestFalse(TEXT("Open bound constructor must not create inclusive bound"), b1_3.IsInclusive());
-
-	TestTrue(TEXT("Implicit constructor must create an inclusive bound"), b1_4.IsInclusive());
-	TestEqual(TEXT("Implicit constructor must create the correct value"), b1_4, b1_2);
+	CHECK(b1_1.IsExclusive());
+	CHECK(b1_1.IsClosed());
+	CHECK_FALSE(b1_1.IsInclusive());
+	CHECK_FALSE(b1_1.IsOpen());
+	CHECK_EQUAL(b1_1.GetValue(), 2.0f);
+
+	CHECK(b1_2.IsInclusive());
+	CHECK(b1_2.IsClosed());
+	CHECK_FALSE(b1_2.IsExclusive());
+	CHECK_FALSE(b1_2.IsOpen());
+	CHECK_EQUAL(b1_2.GetValue(), 2.0f);
+
+	CHECK(b1_3.IsOpen());
+	CHECK_FALSE(b1_3.IsClosed());
+	CHECK_FALSE(b1_3.IsExclusive());
+	CHECK_FALSE(b1_3.IsInclusive());
+
+	CHECK(b1_4.IsInclusive());
+	CHECK_EQUAL(b1_4, b1_2);
 
 	// comparisons must be correct
 	FFloatRangeBound b2_1 = FFloatRangeBound::Exclusive(2.0f);
@@ -49,67 +49,67 @@ TEST_CASE("Core::Math::FFloatRangeBound::Smoke Test", "[Core][Math][Smoke]")
 	FFloatRangeBound b2_5 = FFloatRangeBound::Open();
 	FFloatRangeBound b2_6 = FFloatRangeBound::Open();
 
-	TestTrue(TEXT("Equal exclusive bounds must be equal"), b2_1 == b2_2);
-	TestTrue(TEXT("Equal inclusive bounds must be equal"), b2_3 == b2_4);
-	TestTrue(TEXT("Open bounds must be equal"), b2_5 == b2_6);
+	CHECK(b2_1 == b2_2);
+	CHECK(b2_3 == b2_4);
+	CHECK(b2_5 == b2_6);
 
-	TestFalse(TEXT("Equal exclusive bounds must not be unequal"), b2_1 != b2_2);
-	TestFalse(TEXT("Equal inclusive bounds must not be unequal"), b2_3 != b2_4);
-	TestFalse(TEXT("Open bounds must not be unequal"), b2_5 != b2_6);
+	CHECK_FALSE(b2_1 != b2_2);
+	CHECK_FALSE(b2_3 != b2_4);
+	CHECK_FALSE( b2_5 != b2_6);
 
 	FFloatRangeBound b2_7 = FFloatRangeBound::Exclusive(3.0f);
 	FFloatRangeBound b2_8 = FFloatRangeBound::Inclusive(3.0f);
 
-	TestTrue(TEXT("Unequal exclusive bounds must be unequal"), b2_1 != b2_7);
-	TestTrue(TEXT("Unequal inclusive bounds must be unequal"), b2_2 != b2_8);
+	CHECK(b2_1 != b2_7);
+	CHECK(b2_2 != b2_8);
 
-	TestFalse(TEXT("Unequal exclusive bounds must not be equal"), b2_1 == b2_7);
-	TestFalse(TEXT("Unequal inclusive bounds must not be equal"), b2_2 == b2_8);
+	CHECK_FALSE(b2_1 == b2_7);
+	CHECK_FALSE(b2_2 == b2_8);
 
 	// min-max comparisons between bounds must be correct
 	FFloatRangeBound b3_1 = FFloatRangeBound::Exclusive(2.0f);
 	FFloatRangeBound b3_2 = FFloatRangeBound::Inclusive(2.0f);
 	FFloatRangeBound b3_3 = FFloatRangeBound::Open();
 
-	TestEqual(TEXT("'[2' must be less than '(2' <1>"), FFloatRangeBound::MinLower(b3_2, b3_1), b3_2);
-	TestEqual(TEXT("'[2' must be less than '(2' <2>"), FFloatRangeBound::MinLower(b3_1, b3_2), b3_2);
-	TestEqual(TEXT("Open lower bound must be less than '(2' <1>"), FFloatRangeBound::MinLower(b3_3, b3_1), b3_3);
-	TestEqual(TEXT("Open lower bound must be less than '(2' <2>"), FFloatRangeBound::MinLower(b3_1, b3_3), b3_3);
-	TestEqual(TEXT("Open lower bound must be less than '[2' <1>"), FFloatRangeBound::MinLower(b3_3, b3_2), b3_3);
-	TestEqual(TEXT("Open lower bound must be less than '[2' <2>"), FFloatRangeBound::MinLower(b3_2, b3_3), b3_3);
-
-	TestEqual(TEXT("'(2' must be greater than '[2' <1>"), FFloatRangeBound::MaxLower(b3_2, b3_1), b3_1);
-	TestEqual(TEXT("'(2' must be greater than '[2' <2>"), FFloatRangeBound::MaxLower(b3_1, b3_2), b3_1);
-	TestEqual(TEXT("'(2' must be greater than open lower bound <1>"), FFloatRangeBound::MaxLower(b3_3, b3_1), b3_1);
-	TestEqual(TEXT("'(2' must be greater than open lower bound <2>"), FFloatRangeBound::MaxLower(b3_1, b3_3), b3_1);
-	TestEqual(TEXT("'[2' must be greater than open lower bound <1>"), FFloatRangeBound::MaxLower(b3_3, b3_2), b3_2);
-	TestEqual(TEXT("'[2' must be greater than open lower bound <2>"), FFloatRangeBound::MaxLower(b3_2, b3_3), b3_2);
-
-	TestEqual(TEXT("'2)' must be less than '2]' <1>"), FFloatRangeBound::MinUpper(b3_2, b3_1), b3_1);
-	TestEqual(TEXT("'2)' must be less than '2]' <2>"), FFloatRangeBound::MinUpper(b3_1, b3_2), b3_1);
-	TestEqual(TEXT("'2)' must be less than open upper bound <1>"), FFloatRangeBound::MinUpper(b3_3, b3_1), b3_1);
-	TestEqual(TEXT("'2)' must be less than open upper bound <2>"), FFloatRangeBound::MinUpper(b3_1, b3_3), b3_1);
-	TestEqual(TEXT("'2]' must be less than open upper bound <1>"), FFloatRangeBound::MinUpper(b3_3, b3_2), b3_2);
-	TestEqual(TEXT("'2]' must be less than open upper bound"), FFloatRangeBound::MinUpper(b3_2, b3_3), b3_2);
-
-	TestEqual(TEXT("'2]' must be greater than '2)' <1>"), FFloatRangeBound::MaxUpper(b3_2, b3_1), b3_2);
-	TestEqual(TEXT("'2]' must be greater than '2)' <2>"), FFloatRangeBound::MaxUpper(b3_1, b3_2), b3_2);
-	TestEqual(TEXT("Open upper bound must be greater than '2)' <1>"), FFloatRangeBound::MaxUpper(b3_3, b3_1), b3_3);
-	TestEqual(TEXT("Open upper bound must be greater than '2)' <2>"), FFloatRangeBound::MaxUpper(b3_1, b3_3), b3_3);
-	TestEqual(TEXT("Open upper bound must be greater than '2]' <1>"), FFloatRangeBound::MaxUpper(b3_3, b3_2), b3_3);
-	TestEqual(TEXT("Open upper bound must be greater than '2]' <2>"), FFloatRangeBound::MaxUpper(b3_2, b3_3), b3_3);
+	CHECK_EQUAL(FFloatRangeBound::MinLower(b3_2, b3_1), b3_2);
+	CHECK_EQUAL(FFloatRangeBound::MinLower(b3_1, b3_2), b3_2);
+	CHECK_EQUAL(FFloatRangeBound::MinLower(b3_3, b3_1), b3_3);
+	CHECK_EQUAL(FFloatRangeBound::MinLower(b3_1, b3_3), b3_3);
+	CHECK_EQUAL(FFloatRangeBound::MinLower(b3_3, b3_2), b3_3);
+	CHECK_EQUAL(FFloatRangeBound::MinLower(b3_2, b3_3), b3_3);
+
+	CHECK_EQUAL(FFloatRangeBound::MaxLower(b3_2, b3_1), b3_1);
+	CHECK_EQUAL(FFloatRangeBound::MaxLower(b3_1, b3_2), b3_1);
+	CHECK_EQUAL(FFloatRangeBound::MaxLower(b3_3, b3_1), b3_1);
+	CHECK_EQUAL(FFloatRangeBound::MaxLower(b3_1, b3_3), b3_1);
+	CHECK_EQUAL(FFloatRangeBound::MaxLower(b3_3, b3_2), b3_2);
+	CHECK_EQUAL(FFloatRangeBound::MaxLower(b3_2, b3_3), b3_2);
+
+	CHECK_EQUAL(FFloatRangeBound::MinUpper(b3_2, b3_1), b3_1);
+	CHECK_EQUAL(FFloatRangeBound::MinUpper(b3_1, b3_2), b3_1);
+	CHECK_EQUAL(FFloatRangeBound::MinUpper(b3_3, b3_1), b3_1);
+	CHECK_EQUAL(FFloatRangeBound::MinUpper(b3_1, b3_3), b3_1);
+	CHECK_EQUAL(FFloatRangeBound::MinUpper(b3_3, b3_2), b3_2);
+	CHECK_EQUAL(FFloatRangeBound::MinUpper(b3_2, b3_3), b3_2);
+
+	CHECK_EQUAL(FFloatRangeBound::MaxUpper(b3_2, b3_1), b3_2);
+	CHECK_EQUAL(FFloatRangeBound::MaxUpper(b3_1, b3_2), b3_2);
+	CHECK_EQUAL(FFloatRangeBound::MaxUpper(b3_3, b3_1), b3_3);
+	CHECK_EQUAL(FFloatRangeBound::MaxUpper(b3_1, b3_3), b3_3);
+	CHECK_EQUAL(FFloatRangeBound::MaxUpper(b3_3, b3_2), b3_3);
+	CHECK_EQUAL( FFloatRangeBound::MaxUpper(b3_2, b3_3), b3_3);
 
 	FFloatRangeBound b4_1 = FFloatRangeBound::Exclusive(3);
 	FFloatRangeBound b4_2 = FFloatRangeBound::Inclusive(3);
 
-	TestEqual(TEXT("'(2' must be less than '[3' <1>"), FFloatRangeBound::MinLower(b3_1, b4_2), b3_1);
-	TestEqual(TEXT("'(2' must be less than '[3' <2>"), FFloatRangeBound::MinLower(b4_2, b3_1), b3_1);
-	TestEqual(TEXT("'[2' must be less than '[3' <1>"), FFloatRangeBound::MinLower(b3_2, b4_2), b3_2);
-	TestEqual(TEXT("'[2' must be less than '[3' <2>"), FFloatRangeBound::MinLower(b4_2, b3_2), b3_2);
+	CHECK_EQUAL(FFloatRangeBound::MinLower(b3_1, b4_2), b3_1);
+	CHECK_EQUAL(FFloatRangeBound::MinLower(b4_2, b3_1), b3_1);
+	CHECK_EQUAL(FFloatRangeBound::MinLower(b3_2, b4_2), b3_2);
+	CHECK_EQUAL(FFloatRangeBound::MinLower(b4_2, b3_2), b3_2);
 
-	TestEqual(TEXT("'[3' must be greater than '(2' <1>"), FFloatRangeBound::MaxLower(b3_1, b4_2), b4_2);
-	TestEqual(TEXT("'[3' must be greater than '(2' <2>"), FFloatRangeBound::MaxLower(b4_2, b3_1), b4_2);
-	TestEqual(TEXT("'[3' must be greater than '[2' <1>"), FFloatRangeBound::MaxLower(b3_2, b4_2), b4_2);
-	TestEqual(TEXT("'[3' must be greater than '[2' <2>"), FFloatRangeBound::MaxLower(b4_2, b3_2), b4_2);
+	CHECK_EQUAL(FFloatRangeBound::MaxLower(b3_1, b4_2), b4_2);
+	CHECK_EQUAL(FFloatRangeBound::MaxLower(b4_2, b3_1), b4_2);
+	CHECK_EQUAL(FFloatRangeBound::MaxLower(b3_2, b4_2), b4_2);
+	CHECK_EQUAL(FFloatRangeBound::MaxLower(b4_2, b3_2), b4_2);
 
 }
diff --git a/Engine/Source/Runtime/Core/Tests/Math/RangeTest.cpp b/Engine/Source/Runtime/Core/Tests/Math/RangeTest.cpp
--- a/Engine/Source/Runtime/Core/Tests/Math/RangeTest.cpp
+++ b/Engine/Source/Runtime/Core/Tests/Math/RangeTest.cpp
@@ -8,25 +8,25 @@
 // Disable optimization for RangeTest as it compiles very slowly in development builds
 PRAGMA_DISABLE_OPTIMIZATION
 
-TEST_CASE("Core::Math::FFloatRange::Smoke Test", "[Core][Math][Smoke]")
+TEST_CASE_METHOD(FAutomationTestFixture, "Core::Math::FFloatRange::Smoke Test", "[Core][Math][Smoke]")
 {
 	// single element constructor
 	FFloatRange r1_1 = FFloatRange(3.0f);
 
-	TestEqual(TEXT("Single element constructor must create the correct lower bound value"), r1_1.GetLowerBoundValue(), 3.0f);
-	TestEqual(TEXT("Single element constructor must create the correct upper bound value"), r1_1.GetUpperBoundValue(), 3.0f);
-	TestTrue(TEXT("Single element constructor must create an inclusive lower bound"), r1_1.GetLowerBound().IsInclusive());
-	TestTrue(TEXT("Single element constructor must create an inclusive upper bound"), r1_1.GetUpperBound().IsInclusive());
-	TestFalse(TEXT("Single element constructor must create non-empty range"), r1_1.IsEmpty());
+	TEST_EQUAL(TEXT("Single element constructor must create the correct lower bound value"), r1_1.GetLowerBoundValue(), 3.0f);
+	TEST_EQUAL(TEXT("Single element constructor must create the correct upper bound value"), r1_1.GetUpperBoundValue(), 3.0f);
+	TEST_TRUE(TEXT("Single element constructor must create an inclusive lower bound"), r1_1.GetLowerBound().IsInclusive());
+	TEST_TRUE(TEXT("Single element constructor must create an inclusive upper bound"), r1_1.GetUpperBound().IsInclusive());
+	TEST_FALSE(TEXT("Single element constructor must create non-empty range"), r1_1.IsEmpty());
 
 	// explicit element pair constructor
 	FFloatRange r2_1 = FFloatRange(1.0f, 4.0f);
 
-	TestEqual(TEXT("Explicit element pair constructor must create the correct lower bound value"), r2_1.GetLowerBoundValue(), 1.0f);
-	TestEqual(TEXT("Explicit element pair constructor must create the correct upper bound value"), r2_1.GetUpperBoundValue(), 4.0f);
-	TestTrue(TEXT("Explicit element pair constructor must create an inclusive lower bound"), r2_1.GetLowerBound().IsInclusive());
-	TestTrue(TEXT("Explicit element pair constructor must create an exclusive upper bound"), r2_1.GetUpperBound().IsExclusive());
-	TestFalse(TEXT("Explicit element pair constructor must create non-empty range"), r2_1.IsEmpty());
+	TEST_EQUAL(TEXT("Explicit element pair constructor must create the correct lower bound value"), r2_1.GetLowerBoundValue(), 1.0f);
+	TEST_EQUAL(TEXT("Explicit element pair constructor must create the correct upper bound value"), r2_1.GetUpperBoundValue(), 4.0f);
+	TEST_TRUE(TEXT("Explicit element pair constructor must create an inclusive lower bound"), r2_1.GetLowerBound().IsInclusive());
+	TEST_TRUE(TEXT("Explicit element pair constructor must create an exclusive upper bound"), r2_1.GetUpperBound().IsExclusive());
+	TEST_FALSE(TEXT("Explicit element pair constructor must create non-empty range"), r2_1.IsEmpty());
 
 	// bound pair constructors
 	FFloatRange r3_1 = FFloatRange(3.0f);
@@ -34,10 +34,10 @@ TEST_CASE("Core::Math::FFloatRange::Smoke Test", "[Core][Math][Smoke]")
 	FFloatRange r3_3 = FFloatRange(FFloatRangeBound(3.0f), FFloatRangeBound(3.0f));
 	FFloatRange r3_4 = FFloatRange(FFloatRangeBound::Inclusive(1.0f), FFloatRangeBound::Exclusive(4.0f));
 
-	TestEqual(TEXT("Default bound pair constructor must create the correct range"), r3_3, r3_1);
-	TestFalse(TEXT("Default bound pair constructor must create non-empty range"), r3_3.IsEmpty());
-	TestEqual(TEXT("Specific bound pair constructor must create the correct range"), r3_4, r3_2);
-	TestFalse(TEXT("Specific bound pair constructor must create non-empty range"), r3_4.IsEmpty());
+	TEST_EQUAL(TEXT("Default bound pair constructor must create the correct range"), r3_3, r3_1);
+	TEST_FALSE(TEXT("Default bound pair constructor must create non-empty range"), r3_3.IsEmpty());
+	TEST_EQUAL(TEXT("Specific bound pair constructor must create the correct range"), r3_4, r3_2);
+	TEST_FALSE(TEXT("Specific bound pair constructor must create non-empty range"), r3_4.IsEmpty());
 
 	// empty ranges
 	FFloatRange r4_1 = FFloatRange::Empty();
@@ -47,18 +47,18 @@ TEST_CASE("Core::Math::FFloatRange::Smoke Test", "[Core][Math][Smoke]")
 	FFloatRange r4_5 = FFloatRange(FFloatRangeBound::Exclusive(4.0f), FFloatRangeBound::Inclusive(4.0f));
 	FFloatRange r4_6 = FFloatRange(FFloatRangeBound::Inclusive(4.0f), FFloatRangeBound::Exclusive(4.0f));
 
-	TestTrue(TEXT("Empty range must be empty <1>"), r4_1.IsEmpty());
-	TestTrue(TEXT("Empty range must be empty <2>"), r4_2.IsEmpty());
-	TestTrue(TEXT("Empty range must be empty <3>"), r4_3.IsEmpty());
-	TestTrue(TEXT("Empty range must be empty <4>"), r4_4.IsEmpty());
-	TestTrue(TEXT("Empty range must be empty <5>"), r4_5.IsEmpty());
-	TestTrue(TEXT("Empty range must be empty <6>"), r4_6.IsEmpty());
+	TEST_TRUE(TEXT("Empty range must be empty <1>"), r4_1.IsEmpty());
+	TEST_TRUE(TEXT("Empty range must be empty <2>"), r4_2.IsEmpty());
+	TEST_TRUE(TEXT("Empty range must be empty <3>"), r4_3.IsEmpty());
+	TEST_TRUE(TEXT("Empty range must be empty <4>"), r4_4.IsEmpty());
+	TEST_TRUE(TEXT("Empty range must be empty <5>"), r4_5.IsEmpty());
+	TEST_TRUE(TEXT("Empty range must be empty <6>"), r4_6.IsEmpty());
 
-	TestEqual(TEXT("Empty ranges must be equal <1>"), r4_1, r4_2);
-	TestEqual(TEXT("Empty ranges must be equal <2>"), r4_1, r4_3);
-	TestEqual(TEXT("Empty ranges must be equal <3>"), r4_1, r4_4);
-	TestEqual(TEXT("Empty ranges must be equal <4>"), r4_1, r4_5);
-	TestEqual(TEXT("Empty ranges must be equal <5>"), r4_1, r4_6);
+	TEST_EQUAL(TEXT("Empty ranges must be equal <1>"), r4_1, r4_2);
+	TEST_EQUAL(TEXT("Empty ranges must be equal <2>"), r4_1, r4_3);
+	TEST_EQUAL(TEXT("Empty ranges must be equal <3>"), r4_1, r4_4);
+	TEST_EQUAL(TEXT("Empty ranges must be equal <4>"), r4_1, r4_5);
+	TEST_EQUAL(TEXT("Empty ranges must be equal <5>"), r4_1, r4_6);
 
 	// inclusive-exclusive comparison
 	FFloatRange r5_1 = FFloatRange(FFloatRangeBound::Exclusive(1.0f), FFloatRangeBound::Exclusive(4.0f));
@@ -66,17 +66,17 @@ TEST_CASE("Core::Math::FFloatRange::Smoke Test", "[Core][Math][Smoke]")
 	FFloatRange r5_3 = FFloatRange(FFloatRangeBound::Inclusive(1.0f), FFloatRangeBound::Exclusive(4.0f));
 	FFloatRange r5_4 = FFloatRange(FFloatRangeBound::Inclusive(1.0f), FFloatRangeBound::Inclusive(4.0f));
 
-	TestEqual(TEXT("Same inclusive-exclusive ranges must be equal <1>"), r5_1, r5_1);
-	TestEqual(TEXT("Same inclusive-exclusive ranges must be equal <2>"), r5_2, r5_2);
-	TestEqual(TEXT("Same inclusive-exclusive ranges must be equal <3>"), r5_3, r5_3);
-	TestEqual(TEXT("Same inclusive-exclusive ranges must be equal <4>"), r5_4, r5_4);
+	TEST_EQUAL(TEXT("Same inclusive-exclusive ranges must be equal <1>"), r5_1, r5_1);
+	TEST_EQUAL(TEXT("Same inclusive-exclusive ranges must be equal <2>"), r5_2, r5_2);
+	TEST_EQUAL(TEXT("Same inclusive-exclusive ranges must be equal <3>"), r5_3, r5_3);
+	TEST_EQUAL(TEXT("Same inclusive-exclusive ranges must be equal <4>"), r5_4, r5_4);
 
-	TestNotEqual(TEXT("Different inclusive-exclusive ranges must be different <1>"), r5_1, r5_2);
-	TestNotEqual(TEXT("Different inclusive-exclusive ranges must be different <2>"), r5_1, r5_3);
-	TestNotEqual(TEXT("Different inclusive-exclusive ranges must be different <3>"), r5_1, r5_4);
-	TestNotEqual(TEXT("Different inclusive-exclusive ranges must be different <4>"), r5_2, r5_3);
-	TestNotEqual(TEXT("Different inclusive-exclusive ranges must be different <5>"), r5_2, r5_4);
-	TestNotEqual(TEXT("Different inclusive-exclusive ranges must be different <6>"), r5_3, r5_4);
+	TEST_NOT_EQUAL(TEXT("Different inclusive-exclusive ranges must be different <1>"), r5_1, r5_2);
+	TEST_NOT_EQUAL(TEXT("Different inclusive-exclusive ranges must be different <2>"), r5_1, r5_3);
+	TEST_NOT_EQUAL(TEXT("Different inclusive-exclusive ranges must be different <3>"), r5_1, r5_4);
+	TEST_NOT_EQUAL(TEXT("Different inclusive-exclusive ranges must be different <4>"), r5_2, r5_3);
+	TEST_NOT_EQUAL(TEXT("Different inclusive-exclusive ranges must be different <5>"), r5_2, r5_4);
+	TEST_NOT_EQUAL(TEXT("Different inclusive-exclusive ranges must be different <6>"), r5_3, r5_4);
 
 	// adjoined ranges
 	FFloatRange r6_1 = FFloatRange(FFloatRangeBound::Exclusive(1.0f), FFloatRangeBound::Exclusive(4.0f));
@@ -86,24 +86,24 @@ TEST_CASE("Core::Math::FFloatRange::Smoke Test", "[Core][Math][Smoke]")
 	FFloatRange r6_5 = FFloatRange(FFloatRangeBound::Inclusive(3.0f), FFloatRangeBound::Inclusive(8.0f));
 	FFloatRange r6_6 = FFloatRange(FFloatRangeBound::Inclusive(0.0f), FFloatRangeBound::Inclusive(2.0f));
 
-	TestTrue(TEXT("Adjoined ranges must be adjoined"), r6_1.Adjoins(r6_4));
-	TestTrue(TEXT("Adjoined ranges must be adjoined"), r6_2.Adjoins(r6_3));
-	TestFalse(TEXT("Separated ranges must not be adjoined"), r6_1.Adjoins(r6_3));
-	TestFalse(TEXT("Separated ranges must not be adjoined"), r6_2.Adjoins(r6_4));
-	TestFalse(TEXT("Overlapped ranges must not be adjoined <1>"), r6_1.Adjoins(r6_5));
-	TestFalse(TEXT("Overlapped ranges must not be adjoined <2>"), r6_1.Adjoins(r6_6));
+	TEST_TRUE(TEXT("Adjoined ranges must be adjoined"), r6_1.Adjoins(r6_4));
+	TEST_TRUE(TEXT("Adjoined ranges must be adjoined"), r6_2.Adjoins(r6_3));
+	TEST_FALSE(TEXT("Separated ranges must not be adjoined"), r6_1.Adjoins(r6_3));
+	TEST_FALSE(TEXT("Separated ranges must not be adjoined"), r6_2.Adjoins(r6_4));
+	TEST_FALSE(TEXT("Overlapped ranges must not be adjoined <1>"), r6_1.Adjoins(r6_5));
+	TEST_FALSE(TEXT("Overlapped ranges must not be adjoined <2>"), r6_1.Adjoins(r6_6));
 
 	FFloatRange r6_7 = FFloatRange(FFloatRangeBound::Exclusive(4.0f), FFloatRangeBound::Exclusive(2.0f));
 	FFloatRange r6_8 = FFloatRange(FFloatRangeBound::Inclusive(4.0f), FFloatRangeBound::Inclusive(2.0f));
 
-	TestFalse(TEXT("A non-empty and an empty range must not be adjoined <1>"), r6_1.Adjoins(r6_8));
-	TestFalse(TEXT("A non-empty and an empty range must not be adjoined <2>"), r6_2.Adjoins(r6_7));
+	TEST_FALSE(TEXT("A non-empty and an empty range must not be adjoined <1>"), r6_1.Adjoins(r6_8));
+	TEST_FALSE(TEXT("A non-empty and an empty range must not be adjoined <2>"), r6_2.Adjoins(r6_7));
 	
 	FFloatRange r6_9 = FFloatRange(FFloatRangeBound::Inclusive(2.0f), FFloatRangeBound::Exclusive(0.0f));
 	FFloatRange r6_10 = FFloatRange(FFloatRangeBound::Exclusive(2.0f), FFloatRangeBound::Exclusive(0.0f));
 
-	TestFalse(TEXT("Two empty ranges must not be adjoined <1>"), r6_7.Adjoins(r6_9));
-	TestFalse(TEXT("Two empty ranges must not be adjoined <2>"), r6_8.Adjoins(r6_10));
+	TEST_FALSE(TEXT("Two empty ranges must not be adjoined <1>"), r6_7.Adjoins(r6_9));
+	TEST_FALSE(TEXT("Two empty ranges must not be adjoined <2>"), r6_8.Adjoins(r6_10));
 
 	// @todo gmp: conjoined ranges
 
@@ -117,20 +117,20 @@ TEST_CASE("Core::Math::FFloatRange::Smoke Test", "[Core][Math][Smoke]")
 	FFloatRange r11_1 = FFloatRange(FFloatRangeBound::Inclusive(1.0f), FFloatRangeBound::Exclusive(4.0f));
 	FFloatRange r11_2 = FFloatRange(FFloatRangeBound::Open(), FFloatRangeBound::Open());
 
-	TestTrue(TEXT("A closed range must have a lower bound"), r11_1.HasLowerBound());
-	TestTrue(TEXT("A closed range must have an upper bound"), r11_1.HasUpperBound());
-	TestFalse(TEXT("An open range must not have a lower bound"), r11_2.HasLowerBound());
-	TestFalse(TEXT("An open range must not have an upper bound"), r11_2.HasUpperBound());
+	TEST_TRUE(TEXT("A closed range must have a lower bound"), r11_1.HasLowerBound());
+	TEST_TRUE(TEXT("A closed range must have an upper bound"), r11_1.HasUpperBound());
+	TEST_FALSE(TEXT("An open range must not have a lower bound"), r11_2.HasLowerBound());
+	TEST_FALSE(TEXT("An open range must not have an upper bound"), r11_2.HasUpperBound());
 
-	TestEqual(TEXT("The lower bound value of [1, 4) must be 1"), r11_1.GetLowerBoundValue(), 1.0f);
-	TestEqual(TEXT("The upper bound value of [1, 4) must be 4"), r11_1.GetUpperBoundValue(), 4.0f);
+	TEST_EQUAL(TEXT("The lower bound value of [1, 4) must be 1"), r11_1.GetLowerBoundValue(), 1.0f);
+	TEST_EQUAL(TEXT("The upper bound value of [1, 4) must be 4"), r11_1.GetUpperBoundValue(), 4.0f);
 
 	// degenerate ranges
 	FFloatRange r12_1 = FFloatRange(3.0f);
 	FFloatRange r12_2 = FFloatRange(FFloatRangeBound::Inclusive(3.0f), FFloatRangeBound::Inclusive(3.0f));
 
-	TestTrue(TEXT("A range with a single element must be degenerate <1>"), r12_1.IsDegenerate());
-	TestTrue(TEXT("A range with a single element must be degenerate <2>"), r12_2.IsDegenerate());
+	TEST_TRUE(TEXT("A range with a single element must be degenerate <1>"), r12_1.IsDegenerate());
+	TEST_TRUE(TEXT("A range with a single element must be degenerate <2>"), r12_2.IsDegenerate());
 
 	// @todo gmp: overlapping ranges
 	FFloatRange r14_1 = FFloatRange(FFloatRangeBound::Exclusive(0.0f), FFloatRangeBound::Exclusive(2.0f));
@@ -145,22 +145,22 @@ TEST_CASE("Core::Math::FFloatRange::Smoke Test", "[Core][Math][Smoke]")
 	FFloatRange r14_10 = FFloatRange(FFloatRangeBound::Inclusive(4.0f), FFloatRangeBound::Open());
 	FFloatRange r14_11 = FFloatRange(FFloatRangeBound::Open(), FFloatRangeBound::Exclusive(3.0f));
 	
-	TestTrue(TEXT("(2, 4) must overlap [3, 5]"), r14_3.Overlaps(r14_5));
-	TestTrue(TEXT("(2, 4) must overlap [1, 3]"), r14_3.Overlaps(r14_6));
-	TestTrue(TEXT("(2, 4) must overlap (1, 3]"), r14_3.Overlaps(r14_7));
-	TestTrue(TEXT("(2, 4) must overlap [1, 3)"), r14_3.Overlaps(r14_8));
-	TestTrue(TEXT("(2, 4) must overlap [2, 4]"), r14_3.Overlaps(r14_4));
-	TestTrue(TEXT("(2, 4) must overlap itself"), r14_3.Overlaps(r14_3));
-	TestTrue(TEXT("(2, 4) must not overlap (0, 2)"), !r14_3.Overlaps(r14_1));
-	TestTrue(TEXT("(2, 4) must not overlap [0, 2]"), !r14_3.Overlaps(r14_2));
-	TestTrue(TEXT("[2, 4] must not overlap (0, 2)"), !r14_4.Overlaps(r14_1));
-	TestTrue(TEXT("[2, 4] must overlap [0, 2]"), r14_4.Overlaps(r14_2));
-	TestTrue(TEXT("[2, 4] must not overlap (6, 10)"), !r14_4.Overlaps(r14_9));
-	TestTrue(TEXT("[2, 4] must overlap [4, inf)"), r14_4.Overlaps(r14_10));
-	TestTrue(TEXT("(2, 4) must not overlap [4, inf)"), !r14_3.Overlaps(r14_10));
-	TestTrue(TEXT("[1, 3] must not overlap [4, inf)"), !r14_6.Overlaps(r14_10));
-	TestTrue(TEXT("[1, 3] must overlap (inf, 3)"), r14_6.Overlaps(r14_11));
-	TestTrue(TEXT("(6, 10) must not overlap (inf, 3)"), !r14_9.Overlaps(r14_11));
+	TEST_TRUE(TEXT("(2, 4) must overlap [3, 5]"), r14_3.Overlaps(r14_5));
+	TEST_TRUE(TEXT("(2, 4) must overlap [1, 3]"), r14_3.Overlaps(r14_6));
+	TEST_TRUE(TEXT("(2, 4) must overlap (1, 3]"), r14_3.Overlaps(r14_7));
+	TEST_TRUE(TEXT("(2, 4) must overlap [1, 3)"), r14_3.Overlaps(r14_8));
+	TEST_TRUE(TEXT("(2, 4) must overlap [2, 4]"), r14_3.Overlaps(r14_4));
+	TEST_TRUE(TEXT("(2, 4) must overlap itself"), r14_3.Overlaps(r14_3));
+	TEST_TRUE(TEXT("(2, 4) must not overlap (0, 2)"), !r14_3.Overlaps(r14_1));
+	TEST_TRUE(TEXT("(2, 4) must not overlap [0, 2]"), !r14_3.Overlaps(r14_2));
+	TEST_TRUE(TEXT("[2, 4] must not overlap (0, 2)"), !r14_4.Overlaps(r14_1));
+	TEST_TRUE(TEXT("[2, 4] must overlap [0, 2]"), r14_4.Overlaps(r14_2));
+	TEST_TRUE(TEXT("[2, 4] must not overlap (6, 10)"), !r14_4.Overlaps(r14_9));
+	TEST_TRUE(TEXT("[2, 4] must overlap [4, inf)"), r14_4.Overlaps(r14_10));
+	TEST_TRUE(TEXT("(2, 4) must not overlap [4, inf)"), !r14_3.Overlaps(r14_10));
+	TEST_TRUE(TEXT("[1, 3] must not overlap [4, inf)"), !r14_6.Overlaps(r14_10));
+	TEST_TRUE(TEXT("[1, 3] must overlap (inf, 3)"), r14_6.Overlaps(r14_11));
+	TEST_TRUE(TEXT("(6, 10) must not overlap (inf, 3)"), !r14_9.Overlaps(r14_11));
 
 	// @todo gmp: range splitting
 
@@ -168,17 +168,17 @@ TEST_CASE("Core::Math::FFloatRange::Smoke Test", "[Core][Math][Smoke]")
 	FFloatRange r15_1 = FFloatRange(FFloatRangeBound::Exclusive(7.0f), FFloatRangeBound::Exclusive(14.0f));		// X
 	FFloatRange r15_2 = FFloatRange(FFloatRangeBound::Inclusive(7.0f), FFloatRangeBound::Inclusive(14.0f));
 
-	TestEqual(TEXT("The difference between a range and itself must be an empty set <1>"), FFloatRange::Difference(r15_1, r15_1).Num(), 0);
-	TestEqual(TEXT("The difference between a range and itself must be an empty set <2>"), FFloatRange::Difference(r15_2, r15_2).Num(), 0);
+	TEST_EQUAL(TEXT("The difference between a range and itself must be an empty set <1>"), FFloatRange::Difference(r15_1, r15_1).Num(), 0);
+	TEST_EQUAL(TEXT("The difference between a range and itself must be an empty set <2>"), FFloatRange::Difference(r15_2, r15_2).Num(), 0);
 	
 	TArray<FFloatRange > d15_0 = FFloatRange::Difference(r15_2, r15_1);
 	{
-		TestEqual(TEXT("[7, 14] - (7, 14) must result in two ranges"), d15_0.Num(), 2);
+		TEST_EQUAL(TEXT("[7, 14] - (7, 14) must result in two ranges"), d15_0.Num(), 2);
 
 		bool b15_1 = (d15_0[0] == FFloatRange(FFloatRangeBound::Inclusive(7.0f), FFloatRangeBound::Inclusive(7.0f)));
 		bool b15_2 = (d15_0[1] == FFloatRange(FFloatRangeBound::Inclusive(14.0f), FFloatRangeBound::Inclusive(14.0f)));
 
-		TestTrue(TEXT("[7, 14] - (7, 14) must be {[7, 7], [14, 14]}"), b15_1 && b15_2);
+		TEST_TRUE(TEXT("[7, 14] - (7, 14) must be {[7, 7], [14, 14]}"), b15_1 && b15_2);
 	}
 
 	FFloatRange r15_3 = FFloatRange(FFloatRangeBound::Inclusive(2.0f), FFloatRangeBound::Exclusive(9.0f));		// Y overlapping left
@@ -188,50 +188,50 @@ TEST_CASE("Core::Math::FFloatRange::Smoke Test", "[Core][Math][Smoke]")
 
 	TArray<FFloatRange > d15_1 = FFloatRange::Difference(r15_1, r15_3);
 	{
-		TestEqual(TEXT("(7, 14) - [2, 9) must result in one range"), d15_1.Num(), 1);
-		TestEqual(TEXT("(7, 14) - [2, 9) must be {[9, 14)}"), d15_1[0], FFloatRange(FFloatRangeBound::Inclusive(9.0f), FFloatRangeBound::Exclusive(14.0f)));
+		TEST_EQUAL(TEXT("(7, 14) - [2, 9) must result in one range"), d15_1.Num(), 1);
+		TEST_EQUAL(TEXT("(7, 14) - [2, 9) must be {[9, 14)}"), d15_1[0], FFloatRange(FFloatRangeBound::Inclusive(9.0f), FFloatRangeBound::Exclusive(14.0f)));
 	}
 
 	TArray<FFloatRange > d15_2 = FFloatRange::Difference(r15_1, r15_4);
 	{
-		TestEqual(TEXT("(7, 14) - [2, 9] must result in one range"), d15_2.Num(), 1);
-		TestEqual(TEXT("(7, 14) - [2, 9] must be {(9, 14)}"), d15_2[0], FFloatRange(FFloatRangeBound::Exclusive(9.0f), FFloatRangeBound::Exclusive(14.0f)));
+		TEST_EQUAL(TEXT("(7, 14) - [2, 9] must result in one range"), d15_2.Num(), 1);
+		TEST_EQUAL(TEXT("(7, 14) - [2, 9] must be {(9, 14)}"), d15_2[0], FFloatRange(FFloatRangeBound::Exclusive(9.0f), FFloatRangeBound::Exclusive(14.0f)));
 	}
 
 	TArray<FFloatRange > d15_3 = FFloatRange::Difference(r15_1, r15_5);
 	{
-		TestEqual(TEXT("(7, 14) - (8, 17) must result in one range"), d15_3.Num(), 1);
-		TestEqual(TEXT("(7, 14) - (8, 17) must be {(7, 8]}"), d15_3[0], FFloatRange(FFloatRangeBound::Exclusive(7.0f), FFloatRangeBound::Inclusive(8.0f)));
+		TEST_EQUAL(TEXT("(7, 14) - (8, 17) must result in one range"), d15_3.Num(), 1);
+		TEST_EQUAL(TEXT("(7, 14) - (8, 17) must be {(7, 8]}"), d15_3[0], FFloatRange(FFloatRangeBound::Exclusive(7.0f), FFloatRangeBound::Inclusive(8.0f)));
 	}
 
 	TArray<FFloatRange > d15_4 = FFloatRange::Difference(r15_1, r15_6);
 	{
-		TestEqual(TEXT("(7, 14) - [8, 17) must result in one range"), d15_4.Num(), 1);
-		TestEqual(TEXT("(7, 14) - [8, 17) must be {(7, 8)}"), d15_4[0], FFloatRange(FFloatRangeBound::Exclusive(7.0f), FFloatRangeBound::Exclusive(8.0f)));
+		TEST_EQUAL(TEXT("(7, 14) - [8, 17) must result in one range"), d15_4.Num(), 1);
+		TEST_EQUAL(TEXT("(7, 14) - [8, 17) must be {(7, 8)}"), d15_4[0], FFloatRange(FFloatRangeBound::Exclusive(7.0f), FFloatRangeBound::Exclusive(8.0f)));
 	}
 
 	TArray<FFloatRange > d15_5 = FFloatRange::Difference(r15_2, r15_3);
 	{
-		TestEqual(TEXT("[7, 14] - [2, 9) must result in one range"), d15_5.Num(), 1);
-		TestEqual(TEXT("[7, 14] - [2, 9) must be {[9, 14]}"), d15_5[0], FFloatRange(FFloatRangeBound::Inclusive(9.0f), FFloatRangeBound::Inclusive(14.0f)));
+		TEST_EQUAL(TEXT("[7, 14] - [2, 9) must result in one range"), d15_5.Num(), 1);
+		TEST_EQUAL(TEXT("[7, 14] - [2, 9) must be {[9, 14]}"), d15_5[0], FFloatRange(FFloatRangeBound::Inclusive(9.0f), FFloatRangeBound::Inclusive(14.0f)));
 	}
 
 	TArray<FFloatRange > d15_6 = FFloatRange::Difference(r15_2, r15_4);
 	{
-		TestEqual(TEXT("[7, 14] - [2, 9] must result in one range"), d15_6.Num(), 1);
-		TestEqual(TEXT("[7, 14] - [2, 9] must be {(9, 14]}"), d15_6[0], FFloatRange(FFloatRangeBound::Exclusive(9.0f), FFloatRangeBound::Inclusive(14.0f)));
+		TEST_EQUAL(TEXT("[7, 14] - [2, 9] must result in one range"), d15_6.Num(), 1);
+		TEST_EQUAL(TEXT("[7, 14] - [2, 9] must be {(9, 14]}"), d15_6[0], FFloatRange(FFloatRangeBound::Exclusive(9.0f), FFloatRangeBound::Inclusive(14.0f)));
 	}
 
 	TArray<FFloatRange > d15_7 = FFloatRange::Difference(r15_2, r15_5);
 	{
-		TestEqual(TEXT("[7, 14] - (8, 17) must result in one range"), d15_7.Num(), 1);
-		TestEqual(TEXT("[7, 14] - (8, 17) must be {[7, 8]}"), d15_7[0], FFloatRange(FFloatRangeBound::Inclusive(7.0f), FFloatRangeBound::Inclusive(8.0f)));
+		TEST_EQUAL(TEXT("[7, 14] - (8, 17) must result in one range"), d15_7.Num(), 1);
+		TEST_EQUAL(TEXT("[7, 14] - (8, 17) must be {[7, 8]}"), d15_7[0], FFloatRange(FFloatRangeBound::Inclusive(7.0f), FFloatRangeBound::Inclusive(8.0f)));
 	}
 
 	TArray<FFloatRange > d15_8 = FFloatRange::Difference(r15_2, r15_6);
 	{
-		TestEqual(TEXT("[7, 14] - [8, 17) must result in one range"), d15_8.Num(), 1);
-		TestEqual(TEXT("[7, 14] - [8, 17) must be {[7, 8)}"), d15_8[0], FFloatRange(FFloatRangeBound::Inclusive(7.0f), FFloatRangeBound::Exclusive(8.0f)));
+		TEST_EQUAL(TEXT("[7, 14] - [8, 17) must result in one range"), d15_8.Num(), 1);
+		TEST_EQUAL(TEXT("[7, 14] - [8, 17) must be {[7, 8)}"), d15_8[0], FFloatRange(FFloatRangeBound::Inclusive(7.0f), FFloatRangeBound::Exclusive(8.0f)));
 	}
 
 	FFloatRange r15_7 = FFloatRange(FFloatRangeBound::Inclusive(2.0f), FFloatRangeBound::Exclusive(7.0f));		// Y adjoint left
@@ -239,20 +239,20 @@ TEST_CASE("Core::Math::FFloatRange::Smoke Test", "[Core][Math][Smoke]")
 
 	TArray<FFloatRange > d15_9 = FFloatRange::Difference(r15_1, r15_8);
 	{
-		TestEqual(TEXT("(7, 14) - [2, 7] must result in one set"), d15_9.Num(), 1);
-		TestEqual(TEXT("(7, 14) - [2, 7] must be {(7, 14)}"), d15_9[0], FFloatRange(FFloatRangeBound::Exclusive(7.0f), FFloatRangeBound::Exclusive(14.0f)));
+		TEST_EQUAL(TEXT("(7, 14) - [2, 7] must result in one set"), d15_9.Num(), 1);
+		TEST_EQUAL(TEXT("(7, 14) - [2, 7] must be {(7, 14)}"), d15_9[0], FFloatRange(FFloatRangeBound::Exclusive(7.0f), FFloatRangeBound::Exclusive(14.0f)));
 	}
 
 	TArray<FFloatRange > d15_10 = FFloatRange::Difference(r15_2, r15_7);
 	{
-		TestEqual(TEXT("[7, 14] - [2, 7) must result in one set"), d15_10.Num(), 1);
-		TestEqual(TEXT("[7, 14] - [2, 7) must be {[7, 14]}"), d15_10[0], FFloatRange(FFloatRangeBound::Inclusive(7.0f), FFloatRangeBound::Inclusive(14.0f)));
+		TEST_EQUAL(TEXT("[7, 14] - [2, 7) must result in one set"), d15_10.Num(), 1);
+		TEST_EQUAL(TEXT("[7, 14] - [2, 7) must be {[7, 14]}"), d15_10[0], FFloatRange(FFloatRangeBound::Inclusive(7.0f), FFloatRangeBound::Inclusive(14.0f)));
 	}
 
 	TArray<FFloatRange > d15_11 = FFloatRange::Difference(r15_2, r15_8);
 	{
-		TestEqual(TEXT("[7, 14] - [2, 7] must result in one range"), d15_11.Num(), 1);
-		TestEqual(TEXT("[7, 14] - [2, 7] must be {(7, 14]}"), d15_11[0], FFloatRange(FFloatRangeBound::Exclusive(7.0f), FFloatRangeBound::Inclusive(14.0f)));
+		TEST_EQUAL(TEXT("[7, 14] - [2, 7] must result in one range"), d15_11.Num(), 1);
+		TEST_EQUAL(TEXT("[7, 14] - [2, 7] must be {(7, 14]}"), d15_11[0], FFloatRange(FFloatRangeBound::Exclusive(7.0f), FFloatRangeBound::Inclusive(14.0f)));
 	}
 
 	FFloatRange r15_9 = FFloatRange(FFloatRangeBound::Exclusive(14.0f), FFloatRangeBound::Exclusive(16.0f));	// Y adjoint right
@@ -260,40 +260,40 @@ TEST_CASE("Core::Math::FFloatRange::Smoke Test", "[Core][Math][Smoke]")
 
 	TArray<FFloatRange > d15_12 = FFloatRange::Difference(r15_1, r15_10);
 	{
-		TestEqual(TEXT("(7, 14) - [14, 18] must result in one range"), d15_12.Num(), 1);
-		TestEqual(TEXT("(7, 14) - [14, 18] must be {(7, 14)}"), d15_12[0], FFloatRange(FFloatRangeBound::Exclusive(7.0f), FFloatRangeBound::Exclusive(14.0f)));
+		TEST_EQUAL(TEXT("(7, 14) - [14, 18] must result in one range"), d15_12.Num(), 1);
+		TEST_EQUAL(TEXT("(7, 14) - [14, 18] must be {(7, 14)}"), d15_12[0], FFloatRange(FFloatRangeBound::Exclusive(7.0f), FFloatRangeBound::Exclusive(14.0f)));
 	}
 
 	TArray<FFloatRange > d15_13 = FFloatRange::Difference(r15_2, r15_9);
 	{
-		TestEqual(TEXT("[7, 14] - (14, 18) must result in one range"), d15_13.Num(), 1);
-		TestEqual(TEXT("[7, 14] - (14, 18) must be {[7, 14]}"), d15_13[0], FFloatRange(FFloatRangeBound::Inclusive(7.0f), FFloatRangeBound::Inclusive(14.0f)));
+		TEST_EQUAL(TEXT("[7, 14] - (14, 18) must result in one range"), d15_13.Num(), 1);
+		TEST_EQUAL(TEXT("[7, 14] - (14, 18) must be {[7, 14]}"), d15_13[0], FFloatRange(FFloatRangeBound::Inclusive(7.0f), FFloatRangeBound::Inclusive(14.0f)));
 	}
 
 	TArray<FFloatRange > d15_14 = FFloatRange::Difference(r15_2, r15_10);
 	{
-		TestEqual(TEXT("[7, 14] - [14, 18] must result in one range"), d15_14.Num(), 1);
-		TestEqual(TEXT("[7, 14] - [14, 18] must be {(7, 14)}"), d15_14[0], FFloatRange(FFloatRangeBound::Inclusive(7.0f), FFloatRangeBound::Exclusive(14.0f)));
+		TEST_EQUAL(TEXT("[7, 14] - [14, 18] must result in one range"), d15_14.Num(), 1);
+		TEST_EQUAL(TEXT("[7, 14] - [14, 18] must be {(7, 14)}"), d15_14[0], FFloatRange(FFloatRangeBound::Inclusive(7.0f), FFloatRangeBound::Exclusive(14.0f)));
 	}
 
 	FFloatRange r15_11 = FFloatRange(FFloatRangeBound::Inclusive(8.0f), FFloatRangeBound::Inclusive(13.0f));	// Y enclosed
 
 	TArray<FFloatRange > d15_15 = FFloatRange::Difference(r15_2, r15_11);
 	{
-		TestEqual(TEXT("[7, 14] - [8, 13] must result in two ranges"), d15_15.Num(), 2);
+		TEST_EQUAL(TEXT("[7, 14] - [8, 13] must result in two ranges"), d15_15.Num(), 2);
 
 		bool b15_1 = (d15_15[0] == FFloatRange(FFloatRangeBound::Inclusive(7.0f), FFloatRangeBound::Exclusive(8.0f)));
 		bool b15_2 = (d15_15[1] == FFloatRange(FFloatRangeBound::Exclusive(13.0f), FFloatRangeBound::Inclusive(14.0f)));
 
-		TestTrue(TEXT("[7, 14] - [8, 13] must be {[7, 8), (13, 14]}"), b15_1 && b15_2);
+		TEST_TRUE(TEXT("[7, 14] - [8, 13] must be {[7, 8), (13, 14]}"), b15_1 && b15_2);
 	}
 	
 	FFloatRange r15_12 = FFloatRange(FFloatRangeBound::Inclusive(2.0f), FFloatRangeBound::Inclusive(4.0f));
 
 	TArray<FFloatRange > d15_16 = FFloatRange::Difference(r15_2, r15_12);
 	{
-		TestEqual(TEXT("[7, 14] - [2, 4) must result in one set"), d15_16.Num(), 1);
-		TestEqual(TEXT("[7, 14] - [2, 4) must be {[7, 14]}"), d15_16[0], FFloatRange(FFloatRangeBound::Inclusive(7.0f), FFloatRangeBound::Inclusive(14.0f)));
+		TEST_EQUAL(TEXT("[7, 14] - [2, 4) must result in one set"), d15_16.Num(), 1);
+		TEST_EQUAL(TEXT("[7, 14] - [2, 4) must be {[7, 14]}"), d15_16[0], FFloatRange(FFloatRangeBound::Inclusive(7.0f), FFloatRangeBound::Inclusive(14.0f)));
 	}
 
 	// range hulls
@@ -302,9 +302,9 @@ TEST_CASE("Core::Math::FFloatRange::Smoke Test", "[Core][Math][Smoke]")
 	FFloatRange r16_3 = FFloatRange(4.0f, 1.0f);
 	FFloatRange r16_4 = FFloatRange(8.0f, 5.0f);
 
-	TestEqual(TEXT("The hull of [7, 9) and [11, 14) must be [7, 14)"), FFloatRange::Hull(r16_1, r16_2), FFloatRange(FFloatRangeBound::Inclusive(7.0f), FFloatRangeBound::Exclusive(14.0f)));
-	TestEqual(TEXT("The hull of [7, 9) and an empty range must be [7, 9)"), FFloatRange::Hull(r16_1, r16_3), FFloatRange(FFloatRangeBound::Inclusive(7.0f), FFloatRangeBound::Exclusive(9.0f)));
-	TestTrue(TEXT("The hull of two empty ranges must be empty"), FFloatRange::Hull(r16_3, r16_4).IsEmpty());
+	TEST_EQUAL(TEXT("The hull of [7, 9) and [11, 14) must be [7, 14)"), FFloatRange::Hull(r16_1, r16_2), FFloatRange(FFloatRangeBound::Inclusive(7.0f), FFloatRangeBound::Exclusive(14.0f)));
+	TEST_EQUAL(TEXT("The hull of [7, 9) and an empty range must be [7, 9)"), FFloatRange::Hull(r16_1, r16_3), FFloatRange(FFloatRangeBound::Inclusive(7.0f), FFloatRangeBound::Exclusive(9.0f)));
+	TEST_TRUE(TEXT("The hull of two empty ranges must be empty"), FFloatRange::Hull(r16_3, r16_4).IsEmpty());
 
 	// intersections
 	FFloatRange r17_1 = FFloatRange(FFloatRangeBound::Inclusive(7.0f), FFloatRangeBound::Exclusive(14.0f));
@@ -315,12 +315,12 @@ TEST_CASE("Core::Math::FFloatRange::Smoke Test", "[Core][Math][Smoke]")
 	FFloatRange r17_6 = FFloatRange(4.0f, 1.0f);
 	FFloatRange r17_7 = FFloatRange(8.0f, 5.0f);
 
-	TestEqual(TEXT("The intersection of [7, 14) and [2, 8) must be [7, 8)"), FFloatRange::Intersection(r17_1, r17_2), FFloatRange(FFloatRangeBound::Inclusive(7.0f), FFloatRangeBound::Exclusive(8.0f)));
-	TestEqual(TEXT("The intersection of [7, 14) and [2, 8] must be [7, 8]"), FFloatRange::Intersection(r17_1, r17_3), FFloatRange(FFloatRangeBound::Inclusive(7.0f), FFloatRangeBound::Inclusive(8.0f)));
-	TestEqual(TEXT("The intersection of [7, 14) and (8, 16) must be (8, 14)"), FFloatRange::Intersection(r17_1, r17_4), FFloatRange(FFloatRangeBound::Exclusive(8.0f), FFloatRangeBound::Exclusive(14.0f)));
-	TestEqual(TEXT("The intersection of [7, 14) and [8, 16) must be [8, 14)"), FFloatRange::Intersection(r17_1, r17_5), FFloatRange(FFloatRangeBound::Inclusive(8.0f), FFloatRangeBound::Exclusive(14.0f)));
-	TestTrue(TEXT("The intersection of a non-empty range and an empty range must be empty"), FFloatRange::Intersection(r17_1, r17_6).IsEmpty());
-	TestTrue(TEXT("The intersection of two empty ranges must be empty"), FFloatRange::Intersection(r17_6, r17_7).IsEmpty());
+	TEST_EQUAL(TEXT("The intersection of [7, 14) and [2, 8) must be [7, 8)"), FFloatRange::Intersection(r17_1, r17_2), FFloatRange(FFloatRangeBound::Inclusive(7.0f), FFloatRangeBound::Exclusive(8.0f)));
+	TEST_EQUAL(TEXT("The intersection of [7, 14) and [2, 8] must be [7, 8]"), FFloatRange::Intersection(r17_1, r17_3), FFloatRange(FFloatRangeBound::Inclusive(7.0f), FFloatRangeBound::Inclusive(8.0f)));
+	TEST_EQUAL(TEXT("The intersection of [7, 14) and (8, 16) must be (8, 14)"), FFloatRange::Intersection(r17_1, r17_4), FFloatRange(FFloatRangeBound::Exclusive(8.0f), FFloatRangeBound::Exclusive(14.0f)));
+	TEST_EQUAL(TEXT("The intersection of [7, 14) and [8, 16) must be [8, 14)"), FFloatRange::Intersection(r17_1, r17_5), FFloatRange(FFloatRangeBound::Inclusive(8.0f), FFloatRangeBound::Exclusive(14.0f)));
+	TEST_TRUE(TEXT("The intersection of a non-empty range and an empty range must be empty"), FFloatRange::Intersection(r17_1, r17_6).IsEmpty());
+	TEST_TRUE(TEXT("The intersection of two empty ranges must be empty"), FFloatRange::Intersection(r17_6, r17_7).IsEmpty());
 
 	// union
 	FFloatRange r18_1 = FFloatRange(FFloatRangeBound::Inclusive(7.0f), FFloatRangeBound::Inclusive(14.0f));
@@ -332,61 +332,61 @@ TEST_CASE("Core::Math::FFloatRange::Smoke Test", "[Core][Math][Smoke]")
 	
 	TArray<FFloatRange > u18_0 = FFloatRange::Union(r18_1, r18_1);
 	{
-		TestEqual(TEXT("[7, 14] unioned with itself must result in one range"), u18_0.Num(), 1);
-		TestEqual(TEXT("[7, 14] unioned with itself must result in itself"), u18_0[0], r18_1);
+		TEST_EQUAL(TEXT("[7, 14] unioned with itself must result in one range"), u18_0.Num(), 1);
+		TEST_EQUAL(TEXT("[7, 14] unioned with itself must result in itself"), u18_0[0], r18_1);
 	}
 
 	TArray<FFloatRange > u18_1 = FFloatRange::Union(r18_1, r18_2);
 	{
-		TestEqual(TEXT("[7, 14] unioned with (7, 14) must result in one range"), u18_1.Num(), 1);
-		TestEqual(TEXT("[7, 14] unioned with (7, 14) must result in {[7, 14]}"), u18_1[0], r18_1);
+		TEST_EQUAL(TEXT("[7, 14] unioned with (7, 14) must result in one range"), u18_1.Num(), 1);
+		TEST_EQUAL(TEXT("[7, 14] unioned with (7, 14) must result in {[7, 14]}"), u18_1[0], r18_1);
 	}
 
 	TArray<FFloatRange > u18_2 = FFloatRange::Union(r18_1, r18_3);
 	{
-		TestEqual(TEXT("[7, 14] unioned with [2, 8] must result in one range"), u18_2.Num(), 1);
-		TestEqual(TEXT("[7, 14] unioned with [2, 8] must result in {[2, 14]}"), u18_2[0], FFloatRange(FFloatRangeBound::Inclusive(2.f), FFloatRangeBound::Inclusive(14.f)));
+		TEST_EQUAL(TEXT("[7, 14] unioned with [2, 8] must result in one range"), u18_2.Num(), 1);
+		TEST_EQUAL(TEXT("[7, 14] unioned with [2, 8] must result in {[2, 14]}"), u18_2[0], FFloatRange(FFloatRangeBound::Inclusive(2.f), FFloatRangeBound::Inclusive(14.f)));
 	}
 
 	TArray<FFloatRange > u18_3 = FFloatRange::Union(r18_1, r18_4);
 	{
-		TestEqual(TEXT("[7, 14] unioned with (15, 16] must result in two ranges"), u18_3.Num(), 2);
+		TEST_EQUAL(TEXT("[7, 14] unioned with (15, 16] must result in two ranges"), u18_3.Num(), 2);
 		
 		bool b18_1 = (u18_3[0] == r18_1);
 		bool b18_2 = (u18_3[1] == r18_4);
 
-		TestTrue(TEXT("[7, 14] unioned with (15, 16] must result in {[7, 14], (15, 16]}"), b18_1 && b18_2);
+		TEST_TRUE(TEXT("[7, 14] unioned with (15, 16] must result in {[7, 14], (15, 16]}"), b18_1 && b18_2);
 	}
 
 	TArray<FFloatRange > u18_4 = FFloatRange::Union(r18_1, FFloatRange::Empty());
 	{
-		TestEqual(TEXT("A non-empty range unioned with an empty range must result in one range"), u18_4.Num(), 1);
-		TestEqual(TEXT("A non-empty range unioned with an empty range must result in the non-empty range"), u18_4[0], r18_1);
+		TEST_EQUAL(TEXT("A non-empty range unioned with an empty range must result in one range"), u18_4.Num(), 1);
+		TEST_EQUAL(TEXT("A non-empty range unioned with an empty range must result in the non-empty range"), u18_4[0], r18_1);
 	}
 
 	TArray<FFloatRange > u18_5 = FFloatRange::Union(r18_1, r18_5);
 	{
-		TestEqual(TEXT("[7, 14] unioned with [8, 16) must result in one range"), u18_5.Num(), 1);
-		TestEqual(TEXT("[7, 14] unioned with [8, 16) must result in {[7, 16)}"), u18_5[0], FFloatRange(FFloatRangeBound::Inclusive(7.f), FFloatRangeBound::Exclusive(16.f)));
+		TEST_EQUAL(TEXT("[7, 14] unioned with [8, 16) must result in one range"), u18_5.Num(), 1);
+		TEST_EQUAL(TEXT("[7, 14] unioned with [8, 16) must result in {[7, 16)}"), u18_5[0], FFloatRange(FFloatRangeBound::Inclusive(7.f), FFloatRangeBound::Exclusive(16.f)));
 	}
 
 	TArray<FFloatRange > u18_6 = FFloatRange::Union(r18_1, r18_6);
 	{
-		TestEqual(TEXT("[7, 14] unioned with [6, 7) must result in one range"), u18_6.Num(), 1);
-		TestEqual(TEXT("[7, 14] unioned with [6, 7) must result in {[6, 14]}"), u18_6[0], FFloatRange(FFloatRangeBound::Inclusive(6.f), FFloatRangeBound::Inclusive(14.f)));
+		TEST_EQUAL(TEXT("[7, 14] unioned with [6, 7) must result in one range"), u18_6.Num(), 1);
+		TEST_EQUAL(TEXT("[7, 14] unioned with [6, 7) must result in {[6, 14]}"), u18_6[0], FFloatRange(FFloatRangeBound::Inclusive(6.f), FFloatRangeBound::Inclusive(14.f)));
 	}
 
 	TArray<FFloatRange > u18_7 = FFloatRange::Union(r18_2, r18_6);
 	{
-		TestEqual(TEXT("(7, 14) unioned with [6, 7) must result in two ranges"), u18_7.Num(), 2);
+		TEST_EQUAL(TEXT("(7, 14) unioned with [6, 7) must result in two ranges"), u18_7.Num(), 2);
 		
 		bool b18_1 = (u18_7[0] == r18_2);
 		bool b18_2 = (u18_7[1] == r18_6);
 
-		TestTrue(TEXT("(7, 14) unioned with [6, 7) must result in {(7, 14), [6, 7)}"), b18_1 && b18_2);
+		TEST_TRUE(TEXT("(7, 14) unioned with [6, 7) must result in {(7, 14), [6, 7)}"), b18_1 && b18_2);
 	}
 	
-	TestEqual(TEXT("An pair of empty ranges unioned must yield no ranges."), FFloatRange::Union(FFloatRange::Empty(), FFloatRange::Empty()).Num(), 0);
+	TEST_EQUAL(TEXT("An pair of empty ranges unioned must yield no ranges."), FFloatRange::Union(FFloatRange::Empty(), FFloatRange::Empty()).Num(), 0);
 }
 
 PRAGMA_ENABLE_OPTIMIZATION
diff --git a/Engine/Source/Runtime/Core/Tests/Math/UnitConversionTests.cpp b/Engine/Source/Runtime/Core/Tests/Math/UnitConversionTests.cpp
--- a/Engine/Source/Runtime/Core/Tests/Math/UnitConversionTests.cpp
+++ b/Engine/Source/Runtime/Core/Tests/Math/UnitConversionTests.cpp
@@ -13,7 +13,7 @@ bool IsRoughlyEqual(double One, double Two, float Epsilon)
 	return FMath::Abs(One-Two) <= Epsilon;
 }
 
-TEST_CASE("Core::Math::FUnitConversion::Unit Conversion", "[Core][Math][Smoke]")
+TEST_CASE_METHOD(FAutomationTestFixture, "Core::Math::FUnitConversion::Unit Conversion", "[Core][Math][Smoke]")
 {
 	struct FTestStruct
 	{
@@ -35,7 +35,7 @@ TEST_CASE("Core::Math::FUnitConversion::Unit Conversion", "[Core][Math][Smoke]")
 		{ 1000,		1.0936,		1e-4,	EUnit::Millimeters, 		EUnit::Yards		 		},
 		{ 2000,		0.0787,		1e-4,	EUnit::Micrometers, 		EUnit::Inches		 		},
 
-		{ 90,		UE_PI/2,	1e-3,	EUnit::Degrees, 			EUnit::Radians		 		},
+		{ 90,		UE_PI /2,	1e-3,	EUnit::Degrees, 			EUnit::Radians		 		},
 		{ UE_PI,	180,		1e-3,	EUnit::Radians, 			EUnit::Degrees		 		},
 
 		{ 12,		43.2,		0.1,	EUnit::MetersPerSecond,		EUnit::KilometersPerHour	},
@@ -95,17 +95,11 @@ TEST_CASE("Core::Math::FUnitConversion::Unit Conversion", "[Core][Math][Smoke]")
 		const TCHAR* FromUnitString = FUnitConversion::GetUnitDisplayString(Test.FromUnit);
 		const TCHAR* ToUnitString = FUnitConversion::GetUnitDisplayString(Test.ToUnit);
 
-		TestTrue(FString::Printf(TEXT("Conversion from %s to %s was incorrect. Converting %.10f%s to %s resulted in %.15f%s, expected %.15f%s (threshold = %.15f)"),
-			FromUnitString, ToUnitString,
-			Test.Source, FromUnitString, ToUnitString,
-			ActualResult, ToUnitString,
-			Test.ExpectedResult, ToUnitString,
-			Test.AccuracyEpsilon),
-			IsRoughlyEqual(ActualResult, Test.ExpectedResult, Test.AccuracyEpsilon));
+		CHECK(IsRoughlyEqual(ActualResult, Test.ExpectedResult, Test.AccuracyEpsilon));
 	}
 }
 
-TEST_CASE("Core::Math::FUnitConversion::Parsing Test", "[Core][Math][Smoke]")
+TEST_CASE_METHOD(FAutomationTestFixture, "Core::Math::FUnitConversion::Parsing Test", "[Core][Math][Smoke]")
 {
 	struct FTestCases
 	{
@@ -133,11 +127,7 @@ TEST_CASE("Core::Math::FUnitConversion::Parsing Test", "[Core][Math][Smoke]")
 		TValueOrError<FNumericUnit<double>, FText> Result = UnitConversion::TryParseExpression(Test.Expression, Test.UnderlyingUnit, ExistingValue);
 		if (Result.IsValid())
 		{
-			TestTrue(FString::Printf(TEXT("Parsing of expression \"%s\" failed. Expected %f but got %f."),
-				Test.Expression,
-				Test.ExpectedValue,
-				Result.GetValue().Value), 
-				IsRoughlyEqual(Result.GetValue().ConvertTo(Test.UnderlyingUnit).GetValue().Value, Test.ExpectedValue, 1e-6) );
+			CHECK(IsRoughlyEqual(Result.GetValue().ConvertTo(Test.UnderlyingUnit).GetValue().Value, Test.ExpectedValue, 1e-6) );
 		}
 	}
 }
