Add new Unit Tests

diff --git a/Engine/Source/Runtime/Core/Tests/Math/UnrealMathTest.cpp b/Engine/Source/Runtime/Core/Tests/Math/UnrealMathTest.cpp
--- a/Engine/Source/Runtime/Core/Tests/Math/UnrealMathTest.cpp
+++ b/Engine/Source/Runtime/Core/Tests/Math/UnrealMathTest.cpp
@@ -16,7 +16,6 @@
 #include "Math/RotationMatrix.h"
 #include "Math/Quat.h"
 #include "Math/QuatRotationTranslationMatrix.h"
-#include "Math/Color.h"
 #include "Misc/AutomationTest.h"
 #include "Async/ParallelFor.h"
 #include "Misc/ScopeLock.h"
@@ -34,21 +33,6 @@ static double GSumDouble;
 static bool GPassing;
 
 #define MATHTEST_INLINE FORCENOINLINE // if you want to do performance testing change to FORCEINLINE or FORCEINLINE_DEBUGGABLE
-
-FORCENOINLINE void CheckPassing(bool Passing)
-{
-	// Convenient for setting breakpoints to catch any failures.
-	if (!Passing)
-	{
-		GPassing = false;
-	}
-}
-
-FORCENOINLINE void ResetPassing()
-{
-	GPassing = true;
-}
-
 /**
  * Tests if two vectors (xyzw) are bitwise equal
  *
@@ -65,7 +49,7 @@ FORCENOINLINE bool TestVectorsEqualBitwise( const VectorRegister4Float& Vec0, co
 
 	const bool Passed = (memcmp(GScratch + 0, GScratch + 4, sizeof(float) * 4) == 0);
 
-	CheckPassing(Passed);
+	CHECK(Passed);
 	return Passed;
 }
 
@@ -77,7 +61,7 @@ FORCENOINLINE bool TestVectorsEqualBitwise(const VectorRegister4Double& Vec0, co
 
 	const bool Passed = (memcmp(GScratchDouble + 0, GScratchDouble + 4, sizeof(double) * 4) == 0);
 
-	CheckPassing(Passed);
+	CHECK(Passed);
 	return Passed;
 }
 
@@ -101,7 +85,7 @@ FORCENOINLINE bool TestVectorsEqual(const VectorRegister4Float& Vec0, const Vect
 		float Diff = GScratch[ Component + 0 ] - GScratch[ Component + 4 ];
 		GSum += FMath::Abs(Diff);
 	}
-	CheckPassing(GSum <= Tolerance);
+	CHECK(GSum <= Tolerance);
 	return GSum <= Tolerance;
 }
 
@@ -115,7 +99,7 @@ FORCENOINLINE bool TestVectorsEqual(const VectorRegister4Double& Vec0, const Vec
 		double Diff = GScratchDouble[Component + 0] - GScratchDouble[Component + 4];
 		GSumDouble += FMath::Abs(Diff);
 	}
-	CheckPassing(GSumDouble <= Tolerance);
+	CHECK(GSumDouble <= Tolerance);
 	return GSumDouble <= Tolerance;
 }
 
@@ -135,7 +119,7 @@ FORCENOINLINE bool TestVectorsEqual_ComponentWiseError(const VectorRegister4Floa
 		bPassing &= FMath::IsNearlyZero(Diff, Tolerance);
 		GSum += FMath::Abs(Diff);
 	}
-	CheckPassing(bPassing);
+	CHECK(bPassing);
 	return bPassing;
 }
 
@@ -152,7 +136,7 @@ FORCENOINLINE bool TestVectorsEqual_ComponentWiseError(const VectorRegister4Doub
 		bPassing &= FMath::IsNearlyZero(Diff, Tolerance);
 		GSumDouble += FMath::Abs(Diff);
 	}
-	CheckPassing(bPassing);
+	CHECK(bPassing);
 	return bPassing;
 }
 
@@ -175,7 +159,7 @@ FORCENOINLINE bool TestVectorsEqual3(const VectorRegister4Float& Vec0, const Vec
 	{
 		GSum += FMath::Abs<float>( GScratch[ Component + 0 ] - GScratch[ Component + 4 ] );
 	}
-	CheckPassing(GSum <= Tolerance);
+	CHECK(GSum <= Tolerance);
 	return GSum <= Tolerance;
 }
 
@@ -188,7 +172,7 @@ FORCENOINLINE bool TestVectorsEqual3(const VectorRegister4Double& Vec0, const Ve
 	{
 		GSumDouble += FMath::Abs<double>(GScratchDouble[Component + 0] - GScratchDouble[Component + 4]);
 	}
-	CheckPassing(GSumDouble <= Tolerance);
+	CHECK(GSumDouble <= Tolerance);
 	return GSumDouble <= Tolerance;
 }
 
@@ -217,7 +201,7 @@ FORCENOINLINE bool TestFVector3Equal( const FVector3f& Vec0, const FVector3f& Ve
 	{
 		GSum += FMath::Abs<float>( GScratch[ Component + 0 ] - GScratch[ Component + 4 ] );
 	}
-	CheckPassing(GSum <= Tolerance);
+	CHECK(GSum <= Tolerance);
 	return GSum <= Tolerance;
 }
 
@@ -246,7 +230,7 @@ FORCENOINLINE bool TestFVector3Equal( const FVector3f& Vec0, const FVector3f& Ve
  	{
  		GSumDouble += FMath::Abs<double>(GScratchDouble[Component + 0] - GScratchDouble[Component + 4]);
  	}
- 	CheckPassing(GSumDouble <= Tolerance);
+ 	CHECK(GSumDouble <= Tolerance);
  	return GSumDouble <= Tolerance;
  }
 
@@ -263,7 +247,7 @@ FORCENOINLINE bool TestQuatsEqual(const FQuat4f& Q0, const FQuat4f& Q1, float To
 	GSum = 0.f;
 
 	const bool bEqual = Q0.Equals(Q1, Tolerance);
-	CheckPassing(bEqual);
+	CHECK(bEqual);
 	return bEqual;
 }
 
@@ -280,7 +264,7 @@ FORCENOINLINE bool TestQuatsEqual(const FQuat4d& Q0, const FQuat4d& Q1, double T
 	GSumDouble = 0.f;
 
 	const bool bEqual = Q0.Equals(Q1, Tolerance);
-	CheckPassing(bEqual);
+	CHECK(bEqual);
 	return bEqual;
 }
 
@@ -305,7 +289,7 @@ FORCENOINLINE bool TestFVector3Normalized(const FVector3f& Vec0, float Tolerance
 	GSum = FMath::Sqrt(Vec0.X * Vec0.X + Vec0.Y * Vec0.Y + Vec0.Z * Vec0.Z);
 
 	const bool bNormalized = FMath::IsNearlyEqual(GSum, 1.0f, Tolerance);
-	CheckPassing(bNormalized);
+	CHECK(bNormalized);
 	return bNormalized;
 }
 
@@ -330,7 +314,7 @@ FORCENOINLINE bool TestQuatNormalized(const FQuat4f& Q0, float Tolerance)
 	GSum = FMath::Sqrt(Q0.X*Q0.X + Q0.Y*Q0.Y + Q0.Z*Q0.Z + Q0.W*Q0.W);
 
 	const bool bNormalized = FMath::IsNearlyEqual(GSum, 1.0f, Tolerance);
-	CheckPassing(bNormalized);
+	CHECK(bNormalized);
 	return bNormalized;
 }
 
@@ -355,7 +339,7 @@ FORCENOINLINE bool TestMatricesEqual( const FMatrix44f &Mat0, const FMatrix44f &
 		}
 		if (GSum > Tolerance)
 		{
-			CheckPassing(false);
+			CHECK(false);
 			return false;
 		}
 	}
@@ -374,7 +358,7 @@ FORCENOINLINE bool TestMatricesEqual(const FMatrix44d& Mat0, const FMatrix44d& M
 		}
 		if (GSumDouble > Tolerance)
 		{
-			CheckPassing(false);
+			CHECK(false);
 			return false;
 		}
 	}
@@ -679,7 +663,7 @@ FORCENOINLINE FRotator3f TestQuaternionToRotator(const FQuat4f& Quat)
 	const float YawX = (1.f-2.f*(FMath::Square(Y) + FMath::Square(Z)));
 	const float SINGULARITY_THRESHOLD = 0.4999995f;
 
-	static const float RAD_TO_DEG = (180.f)/UE_PI;
+	static const float RAD_TO_DEG = (180.f)/ UE_PI;
 	FRotator3f RotatorFromQuat;
 
 	// Note: using stock C functions for some trig functions since this is the "reference" implementation
@@ -795,34 +779,21 @@ bool TestRotatorEqual3(const FRotator3f& A, const FRotator3f& B, const float Tol
 // Report an error if bComparison is not equal to bExpected.
 void LogRotatorTest(bool bExpected, const TCHAR* TestName, const FRotator3f& A, const FRotator3f& B, bool bComparison)
 {
-	const bool bHasPassed = (bComparison == bExpected);
-	if (bHasPassed == false)
-	{
-		UE_LOG(LogUnrealMathTest, Log, TEXT("%s: %s"), bHasPassed ? TEXT("PASSED") : TEXT("FAILED"), TestName);
-		UE_LOG(LogUnrealMathTest, Log, TEXT("(%s).Equals(%s) = %d"), *A.ToString(), *B.ToString(), bComparison);
-		CheckPassing(false);
-	}
+	INFO(TestName);
+	CHECK(bComparison == bExpected);
 }
 
 
 void LogRotatorTest(const TCHAR* TestName, const FRotator3f& A, const FRotator3f& B, bool bComparison)
 {
-	if (bComparison == false)
-	{
-		UE_LOG(LogUnrealMathTest, Log, TEXT("%s: %s"), bComparison ? TEXT("PASSED") : TEXT("FAILED"), TestName);
-		UE_LOG(LogUnrealMathTest, Log, TEXT("(%s).Equals(%s) = %d"), *A.ToString(), *B.ToString(), bComparison);
-		CheckPassing(false);
-	}
+	INFO(TestName);
+	CHECK(bComparison);
 }
 
 void LogQuaternionTest(const TCHAR* TestName, const FQuat4f& A, const FQuat4f& B, bool bComparison)
 {
-	if (bComparison == false)
-	{
-		UE_LOG(LogUnrealMathTest, Log, TEXT("%s: %s"), bComparison ? TEXT("PASSED") : TEXT("FAILED"), TestName);
-		UE_LOG(LogUnrealMathTest, Log, TEXT("(%s).Equals(%s) = %d"), *A.ToString(), *B.ToString(), bComparison);
-		CheckPassing(false);
-	}
+	INFO(TestName);
+	CHECK(bComparison);
 }
 
 // Normalize tests
@@ -979,35 +950,25 @@ FORCENOINLINE bool TestVectorFunction2Params(VectorRegisterType& ReferenceResult
 template<typename FloatType>
 FORCENOINLINE void LogTest(const TCHAR* TestName, bool bHasPassed)
 {
-	if (bHasPassed == false)
-	{
-		UE_LOG(LogUnrealMathTest, Error, TEXT("Unimplemented type for LogTest()"));
-		CheckPassing(false);
-	}
+	INFO(TestName);
+	CHECK(bHasPassed);
 }
 
 // Specialization for 'float' so it checks the correct global state filled by the various comparision validation functions
 template<>
 FORCENOINLINE void LogTest<float>(const TCHAR* TestName, bool bHasPassed)
 {
-	if (bHasPassed == false)
-	{
-		UE_LOG(LogUnrealMathTest, Warning, TEXT("%s <float>: %s"), bHasPassed ? TEXT("PASSED") : TEXT("FAILED"), TestName);
-		UE_LOG(LogUnrealMathTest, Warning, TEXT("Bad(%.8f): (%.8f %.8f %.8f %.8f) (%.8f %.8f %.8f %.8f)"), GSum, GScratch[0], GScratch[1], GScratch[2], GScratch[3], GScratch[4], GScratch[5], GScratch[6], GScratch[7]);
-		CheckPassing(false);
-	}
+	INFO(TestName);
+	CHECK(bHasPassed);
+
 }
 
 // Specialization for 'double' so it checks the correct global state filled by the various comparision validation functions
 template<>
 FORCENOINLINE void LogTest<double>(const TCHAR* TestName, bool bHasPassed)
 {
-	if (bHasPassed == false)
-	{
-		UE_LOG(LogUnrealMathTest, Warning, TEXT("%s <double>: %s"), bHasPassed ? TEXT("PASSED") : TEXT("FAILED"), TestName);
-		UE_LOG(LogUnrealMathTest, Warning, TEXT("Bad(%.8f): (%.8f %.8f %.8f %.8f) (%.8f %.8f %.8f %.8f)"), GSumDouble, GScratchDouble[0], GScratchDouble[1], GScratchDouble[2], GScratchDouble[3], GScratchDouble[4], GScratchDouble[5], GScratchDouble[6], GScratchDouble[7]);
-		CheckPassing(false);
-	}
+	INFO(TestName);
+	CHECK(bHasPassed);
 }
 
 /**
@@ -1050,7 +1011,8 @@ FORCENOINLINE void TestVectorReplicate()
 #define ReplicateTest(A, X) \
 		V1 = VectorReplicate(A, X); \
 		V2 = MakeVectorRegister(ArrayV0[X], ArrayV0[X], ArrayV0[X], ArrayV0[X]); \
-		LogTest<RealType>(*FString::Printf(TEXT("VectorReplicate<%d>"), X), TestVectorsEqual(V1, V2));
+		INFO(*FString::Printf(TEXT("VectorReplicate<%d>"), X));\
+		CHECK(TestVectorsEqual(V1, V2));
 
 	ReplicateTest(V0, 0);
 	ReplicateTest(V0, 1);
@@ -1070,8 +1032,9 @@ FORCENOINLINE void TestVectorSwizzle()
 
 #define SwizzleTest(A, X, Y, Z, W) \
 		V1 = VectorSwizzle(A, X, Y, Z, W); \
-		V2 = MakeVectorRegister(ArrayV0[X], ArrayV0[Y], ArrayV0[Z], ArrayV0[W]); \
-		LogTest<RealType>(*FString::Printf(TEXT("VectorSwizzle<%d,%d,%d,%d>"), X, Y, Z, W), TestVectorsEqual(V1, V2));
+		V2 = MakeVectorRegister(ArrayV0[X], ArrayV0[Y], ArrayV0[Z], ArrayV0[W]);\
+		INFO(*FString::Printf(TEXT("VectorSwizzle<%d,%d,%d,%d>"), X, Y, Z, W));\
+		CHECK(TestVectorsEqual(V1, V2));
 
 	// This is not an exhaustive list because it would be 4*4*4*4 = 256 entries, but it tries to test a lot of common permutations.
 	// Unfortunately it can't be done in a loop because it uses a #define and compile-time constants for the VectorSwizzle() 'function'.
@@ -1168,7 +1131,8 @@ FORCENOINLINE void TestVectorShuffle()
 #define ShuffleTest(A, B, X, Y, Z, W) \
 		V2 = VectorShuffle(A, B, X, Y, Z, W); \
 		V3 = MakeVectorRegister(ArrayV0[X], ArrayV0[Y], ArrayV1[Z], ArrayV1[W]); \
-		LogTest<RealType>(*FString::Printf(TEXT("VectorShuffle<%d,%d,%d,%d>"), X, Y, Z, W), TestVectorsEqual(V2, V3));
+		INFO(*FString::Printf(TEXT("VectorShuffle<%d,%d,%d,%d>"), X, Y, Z, W));\
+		CHECK(TestVectorsEqual(V2, V3));
 
 	// This is not an exhaustive list because it would be 4*4*4*4 = 256 entries, but it tries to test a lot of common permutations.
 	// Unfortunately it can't be done in a loop because it uses a #define and compile-time constants for the VectorShuffle() 'function'.
@@ -1427,10 +1391,8 @@ FORCENOINLINE void TestVectorExpLogFunctions()
 	}
 }
 
-TEST_CASE("Core::Math::Vector::Vector Register Abstraction", "[Core][Math][Smoke]")
+TEST_CASE_METHOD(FAutomationTestFixture, "Core::Math::FVector::Vector Register Abstraction", "[Core][Math][Smoke]")
 {
-	ResetPassing();
-
 	double F1 = 1.0;
 	uint64 U1 = *(uint64*)&F1;
 	VectorRegister4Double V0, V1;
@@ -2045,11 +2007,11 @@ TEST_CASE("Core::Math::Vector::Vector Register Abstraction", "[Core][Math][Smoke
 		FVector3d FV0, FV1;
 		FV0 = Rotator0.Vector();
 		FV1 = TRotationMatrix<double>( Rotator0 ).GetScaledAxis( EAxis::X );
-		LogTest<double>( TEXT("Test0 Rotator::Vector()"), TestFVector3Equal(FV0, FV1, 1e-6));
+		LogTest<double>( TEXT("Test0 Rotator::Vector()"), TestFVector3Equal(FV0, FV1, 1e-6f));
 		
 		FV0 = TRotationMatrix<double>( Rotator0 ).GetScaledAxis( EAxis::X );
 		FV1 = TQuatRotationMatrix<double>( FQuat4d(Q0.X, Q0.Y, Q0.Z, Q0.W) ).GetScaledAxis( EAxis::X );
-		LogTest<double>( TEXT("Test0 FQuatRotationMatrix"), TestFVector3Equal(FV0, FV1, 1e-5));
+		LogTest<double>( TEXT("Test0 FQuatRotationMatrix"), TestFVector3Equal(FV0, FV1, 1e-5f));
 
 		Rotator0 = FRotator3d(45.0f,  60.0f, 120.0f);
 		Q0 = FQuat4d(Rotator0);
@@ -2058,15 +2020,15 @@ TEST_CASE("Core::Math::Vector::Vector Register Abstraction", "[Core][Math][Smoke
 
 		FV0 = Rotator0.Vector();
 		FV1 = TRotationMatrix<double>( Rotator0 ).GetScaledAxis( EAxis::X );
-		LogTest<double>( TEXT("Test1 Rotator::Vector()"), TestFVector3Equal(FV0, FV1, 1e-6));
+		LogTest<double>( TEXT("Test1 Rotator::Vector()"), TestFVector3Equal(FV0, FV1, 1e-6f));
 
 		FV0 = TRotationMatrix<double>( Rotator0 ).GetScaledAxis( EAxis::X );
 		FV1 = TQuatRotationMatrix<double>(FQuat4d(Q0.X, Q0.Y, Q0.Z, Q0.W) ).GetScaledAxis( EAxis::X );
-		LogTest<double>(TEXT("Test1 FQuatRotationMatrix"), TestFVector3Equal(FV0, FV1, 1e-5));
+		LogTest<double>(TEXT("Test1 FQuatRotationMatrix"), TestFVector3Equal(FV0, FV1, 1e-5f));
 
 		FV0 = TRotationMatrix<double>( FRotator3d::ZeroRotator ).GetScaledAxis(EAxis::X);
 		FV1 = TQuatRotationMatrix<double>(FQuat4d::Identity ).GetScaledAxis(EAxis::X);
-		LogTest<double>(TEXT("Test2 FQuatRotationMatrix"), TestFVector3Equal(FV0, FV1, 1e-6));
+		LogTest<double>(TEXT("Test2 FQuatRotationMatrix"), TestFVector3Equal(FV0, FV1, 1e-6f));
 	}
 
 	// NaN / Inf tests
@@ -2141,7 +2103,7 @@ TEST_CASE("Core::Math::Vector::Vector Register Abstraction", "[Core][Math][Smoke
  * Run a suite of vector operations to validate vector intrinsics are working on the platform
  */
 
-TEST_CASE("Core::Math::Vector::Vector Intrinsics", "[Core][Math][Smoke]")
+TEST_CASE_METHOD(FAutomationTestFixture, "Core::Math::FVector::Vector Intrinsics", "[Core][Math][Smoke]")
 {
 	float F1 = 1.f;
 	uint32 U1 = *(uint32 *)&F1;
@@ -2159,8 +2121,6 @@ TEST_CASE("Core::Math::Vector::Vector Intrinsics", "[Core][Math][Smoke]")
 	const float NaN = NaNU.FloatNaN;
 	const VectorRegister4Float VectorNaN = MakeVectorRegisterFloat(NaN, NaN, NaN, NaN);
 
-	ResetPassing();
-
 	V0 = MakeVectorRegister( U1, U1, U1, U1);
 	V1 = MakeVectorRegister( F1, F1, F1, F1 );
 	LogTest<float>( TEXT("MakeVectorRegister"), TestVectorsEqual( V0, V1 ) );
@@ -2813,8 +2773,8 @@ TEST_CASE("Core::Math::Vector::Vector Intrinsics", "[Core][Math][Smoke]")
 			// Test extreme ranges
 			{ 1.0f,				 UE_KINDA_SMALL_NUMBER},
 			{ 1.0f,				-UE_KINDA_SMALL_NUMBER},
-			{-UE_SMALL_NUMBER,   UE_SMALL_NUMBER},
-			{ UE_SMALL_NUMBER,  -UE_SMALL_NUMBER},
+			{-UE_SMALL_NUMBER,	 UE_SMALL_NUMBER},
+			{ UE_SMALL_NUMBER,	-UE_SMALL_NUMBER},
 			{ 1.0f,				 MIN_flt},
 			{ 1.0f,				-MIN_flt},
 			{ MAX_flt,			 MIN_flt},
@@ -2856,7 +2816,7 @@ TEST_CASE("Core::Math::Vector::Vector Intrinsics", "[Core][Math][Smoke]")
 			if (Delta > 1e-8f)
 			{
 				UE_LOG(LogUnrealMathTest, Log, TEXT("FMath::Fmod(%f, %f)=%f <-> fmodf(%f, %f)=%f: FAILED"), X, Y, Ours, X, Y, Theirs);
-				CheckPassing(false);
+				CHECK(false);
 			}
 		}
 	}
@@ -2894,10 +2854,10 @@ TEST_CASE("Core::Math::Vector::Vector Intrinsics", "[Core][Math][Smoke]")
 			{-5.3 * UE_DOUBLE_PI,	2. * UE_DOUBLE_PI},
 
 			// Test extreme ranges
-			{ 1.0,			 UE_DOUBLE_KINDA_SMALL_NUMBER},
-			{ 1.0,			-UE_DOUBLE_KINDA_SMALL_NUMBER},
-			{-UE_DOUBLE_SMALL_NUMBER,  UE_DOUBLE_SMALL_NUMBER},
-			{ UE_DOUBLE_SMALL_NUMBER, -UE_DOUBLE_SMALL_NUMBER},
+			{ 1.0,						 UE_DOUBLE_KINDA_SMALL_NUMBER},
+			{ 1.0,						-UE_DOUBLE_KINDA_SMALL_NUMBER},
+			{-UE_DOUBLE_SMALL_NUMBER,	 UE_DOUBLE_SMALL_NUMBER},
+			{ UE_DOUBLE_SMALL_NUMBER,	-UE_DOUBLE_SMALL_NUMBER},
 			{ 1.0,			 MIN_dbl},
 			{ 1.0,			-MIN_dbl},
 			{ MAX_flt,		 MIN_dbl},
@@ -2942,15 +2902,15 @@ TEST_CASE("Core::Math::Vector::Vector Intrinsics", "[Core][Math][Smoke]")
 				if (FractionalDelta > 1e-8)
 				{
 					UE_LOG(LogUnrealMathTest, Log, TEXT("FMath::Fmod(%.12f, %.12f)=%.12f <-> fmod(%.12f, %.12f)=%.12f: FAILED"), X, Y, Ours, X, Y, Theirs);
-					CheckPassing(false);
+					CHECK(false);
 				}				
 			}
 		}
 	}
 
 
-	// Float function compilation with various types. Set this define to 1 to verify that warnings are generated for all code within MATHTEST_CHECK_INVALID_FLOAT_VARIANTS blocks.
-#define MATHTEST_CHECK_INVALID_FLOAT_VARIANTS 0
+	// Float function compilation with various types. Set this define to 1 to verify that warnings are generated for all code within MATHTEST_TEST_INVALID_FLOAT_VARIANTS blocks.
+#define MATHTEST_TEST_INVALID_FLOAT_VARIANTS 0
 	{
 		float F = 1.f, TestFloat;
 		double D = 1.0, TestDouble;
@@ -2962,7 +2922,7 @@ TEST_CASE("Core::Math::Vector::Vector Intrinsics", "[Core][Math][Smoke]")
 		TestFloat = FMath::Fmod(F, I);
 		TestFloat = FMath::Fmod(I, F);
 
-#if MATHTEST_CHECK_INVALID_FLOAT_VARIANTS
+#if MATHTEST_TEST_INVALID_FLOAT_VARIANTS
 		// Expected to generate warnings
 		TestFloat = FMath::Fmod(F, D);
 		TestFloat = FMath::Fmod(D, F);
@@ -2982,12 +2942,12 @@ TEST_CASE("Core::Math::Vector::Vector Intrinsics", "[Core][Math][Smoke]")
 		TestDouble = FMath::Fmod(I, F);
 		TestDouble = FMath::Fmod(I, D);
 
-#if MATHTEST_CHECK_INVALID_FLOAT_VARIANTS
+#if MATHTEST_TEST_INVALID_FLOAT_VARIANTS
 		// Expected to generate warnings
 		TestDouble = FMath::Fmod(I, I); // Should be warned to be deprecated
 #endif
 
-#if MATHTEST_CHECK_INVALID_FLOAT_VARIANTS
+#if MATHTEST_TEST_INVALID_FLOAT_VARIANTS
 		// Expected to generate warnings
 		int TestInt;
 		TestInt = FMath::Fmod(F, F);
@@ -3010,7 +2970,7 @@ TEST_CASE("Core::Math::Vector::Vector Intrinsics", "[Core][Math][Smoke]")
 
 		TestDouble = FMath::TruncToDouble(I); // not currently a warning
 
-#if MATHTEST_CHECK_INVALID_FLOAT_VARIANTS
+#if MATHTEST_TEST_INVALID_FLOAT_VARIANTS
 		// Expected to generate errors
 		TestFloat = FMath::TruncToFloat(I);
 		TestFloat = FMath::TruncToFloat(D);
@@ -3057,14 +3017,14 @@ TEST_CASE("Core::Math::Vector::Vector Intrinsics", "[Core][Math][Smoke]")
 		F = FMath::Clamp(F1, I1, F2);
 		D = FMath::Clamp(F1, D2, I3);
 
-		CheckPassing(FMath::Clamp( 1, 0, 2) == 1);
-		CheckPassing(FMath::Clamp(-1, 0, 2) == 0);
-		CheckPassing(FMath::Clamp( 3, 0, 2) == 2);
+		CHECK(FMath::Clamp( 1, 0, 2) == 1);
+		CHECK(FMath::Clamp(-1, 0, 2) == 0);
+		CHECK(FMath::Clamp( 3, 0, 2) == 2);
 
 		U1 = 0; U2 = 2;
-		CheckPassing(FMath::Clamp<int32>(-1, U1, U2) == 0);
+		CHECK(FMath::Clamp<int32>(-1, U1, U2) == 0);
 
-#if MATHTEST_CHECK_INVALID_FLOAT_VARIANTS
+#if MATHTEST_TEST_INVALID_FLOAT_VARIANTS
 		// Expected to generate errors (ambiguous arguments)
 		U = FMath::Clamp(U1, I1, U3);
 		// Expected to generate errors (double->float truncation)
@@ -3079,10 +3039,10 @@ TEST_CASE("Core::Math::Vector::Vector Intrinsics", "[Core][Math][Smoke]")
 		Big3 = FMath::Max(Big1, Big2);
 
 		U1 = uint32(-1);
-		CheckPassing(FMath::Max<int32>(0, U1) == 0);
+		CHECK(FMath::Max<int32>(0, U1) == 0);
 
 		U1 = uint32(-1);
-		CheckPassing(FMath::Min<int32>(0, U1) == -1);
+		CHECK(FMath::Min<int32>(0, U1) == -1);
 
 		// Test compilation mixing int32/int64 types
 		//Big3 = FMath::Max(Big1, I2);
@@ -3340,7 +3300,8 @@ TEST_CASE("Core::Math::Vector::Vector Intrinsics", "[Core][Math][Smoke]")
 		#define TEST_QUAT_ROTATE(_QIndex, _VIndex, _Quat, _Vec, _Func, _Tolerance) \
 		{ \
 			const FString _TestName = FString::Printf(TEXT("Test Quat%d: Vec%d: %s"), _QIndex, _VIndex, TEXT(#_Func)); \
-			LogTest<float>( *_TestName, TestFVector3Equal(_Quat.RotateVector(_Vec), _Func(_Quat, _Vec), _Tolerance) ); \
+			INFO(*_TestName);\
+			CHECK(TestFVector3Equal(_Quat.RotateVector(_Vec), _Func(_Quat, _Vec), _Tolerance) ); \
 		}
 
 		// Test loop
@@ -3453,7 +3414,7 @@ TEST_CASE("Core::Math::Vector::Vector Intrinsics", "[Core][Math][Smoke]")
 	}
 }
 
-TEST_CASE("Core::Math::Misc::Interpolation Function", "[Core][Math][Smoke]")
+TEST_CASE_METHOD(FAutomationTestFixture, "Core::Math::Misc::Interpolation Function", "[Core][Math][Smoke]")
 {
 	// The purpose of this test is to verify that various combinations of the easing functions are actually equivalent.
 	// It currently only tests the InOut versions over different ranges, because the initial implementation was bad.
@@ -3461,7 +3422,7 @@ TEST_CASE("Core::Math::Misc::Interpolation Function", "[Core][Math][Smoke]")
 	// expansions to this test suite.
 
 	typedef float(*EasingFunc)(float Percent);
-	auto RunInOutTest = [](const TArray< TPair<EasingFunc, FString> >& Functions)
+	auto RunInOutTest = [this](const TArray< TPair<EasingFunc, FString> >& Functions)
 	{
 		for (int32 I = 0; I < 100; ++I)
 		{
@@ -3476,7 +3437,8 @@ TEST_CASE("Core::Math::Misc::Interpolation Function", "[Core][Math][Smoke]")
 			int32 K = 0;
 			for (int32 J = 1; J < Functions.Num(); ++J)
 			{
-				TestTrue(FString::Printf(TEXT("Easing Function tests failed at index %d!"), I), FMath::IsNearlyEqual(Values[K], Values[J], 0.0001f));
+				INFO(FString::Printf(TEXT("Easing Function tests failed at index %d!"), I));
+				CHECK(FMath::IsNearlyEqual(Values[K], Values[J],0.0001f));
 			}
 		}
 	};
@@ -3527,149 +3489,149 @@ TEST_CASE("Core::Math::Misc::Interpolation Function", "[Core][Math][Smoke]")
 	}
 }
 
-TEST_CASE("Core::Math::Misc::Round HalfToZero", "[Core][Math][Smoke]")
+TEST_CASE_METHOD(FAutomationTestFixture, "Core::Math::Misc::Round HalfToZero", "[Core][Math][Smoke]")
 {
-	TestEqual(TEXT("RoundHalfToZero32-Zero"), FMath::RoundHalfToZero(0.0f), 0.0f);
-	TestEqual(TEXT("RoundHalfToZero32-One"), FMath::RoundHalfToZero(1.0f), 1.0f);
-	TestEqual(TEXT("RoundHalfToZero32-LessHalf"), FMath::RoundHalfToZero(1.4f), 1.0f);
-	TestEqual(TEXT("RoundHalfToZero32-NegGreaterHalf"), FMath::RoundHalfToZero(-1.4f), -1.0f);
-	TestEqual(TEXT("RoundHalfToZero32-LessNearHalf"), FMath::RoundHalfToZero(1.4999999f), 1.0f);
-	TestEqual(TEXT("RoundHalfToZero32-NegGreaterNearHalf"), FMath::RoundHalfToZero(-1.4999999f), -1.0f);
-	TestEqual(TEXT("RoundHalfToZero32-Half"), FMath::RoundHalfToZero(1.5f), 1.0f);
-	TestEqual(TEXT("RoundHalfToZero32-NegHalf"), FMath::RoundHalfToZero(-1.5f), -1.0f);
-	TestEqual(TEXT("RoundHalfToZero32-GreaterNearHalf"), FMath::RoundHalfToZero(1.5000001f), 2.0f);
-	TestEqual(TEXT("RoundHalfToZero32-NegLesserNearHalf"), FMath::RoundHalfToZero(-1.5000001f), -2.0f);
-	TestEqual(TEXT("RoundHalfToZero32-GreaterThanHalf"), FMath::RoundHalfToZero(1.6f), 2.0f);
-	TestEqual(TEXT("RoundHalfToZero32-NegLesserThanHalf"), FMath::RoundHalfToZero(-1.6f), -2.0f);
-
-	TestEqual(TEXT("RoundHalfToZero32-TwoToOneBitPrecision"), FMath::RoundHalfToZero(4194303.25f), 4194303.0f);
-	TestEqual(TEXT("RoundHalfToZero32-TwoToOneBitPrecision"), FMath::RoundHalfToZero(4194303.5f), 4194303.0f);
-	TestEqual(TEXT("RoundHalfToZero32-TwoToOneBitPrecision"), FMath::RoundHalfToZero(4194303.75f), 4194304.0f);
-	TestEqual(TEXT("RoundHalfToZero32-TwoToOneBitPrecision"), FMath::RoundHalfToZero(4194304.0f), 4194304.0f);
-	TestEqual(TEXT("RoundHalfToZero32-TwoToOneBitPrecision"), FMath::RoundHalfToZero(4194304.5f), 4194304.0f);
-	TestEqual(TEXT("RoundHalfToZero32-NegTwoToOneBitPrecision"), FMath::RoundHalfToZero(-4194303.25f), -4194303.0f);
-	TestEqual(TEXT("RoundHalfToZero32-NegTwoToOneBitPrecision"), FMath::RoundHalfToZero(-4194303.5f), -4194303.0f);
-	TestEqual(TEXT("RoundHalfToZero32-NegTwoToOneBitPrecision"), FMath::RoundHalfToZero(-4194303.75f), -4194304.0f);
-	TestEqual(TEXT("RoundHalfToZero32-NegTwoToOneBitPrecision"), FMath::RoundHalfToZero(-4194304.0f), -4194304.0f);
-	TestEqual(TEXT("RoundHalfToZero32-NegTwoToOneBitPrecision"), FMath::RoundHalfToZero(-4194304.5f), -4194304.0f);
-
-	TestEqual(TEXT("RoundHalfToZero32-OneToZeroBitPrecision"), FMath::RoundHalfToZero(8388607.0f), 8388607.0f);
-	TestEqual(TEXT("RoundHalfToZero32-OneToZeroBitPrecision"), FMath::RoundHalfToZero(8388607.5f), 8388607.0f);
-	TestEqual(TEXT("RoundHalfToZero32-OneToZeroBitPrecision"), FMath::RoundHalfToZero(8388608.0f), 8388608.0f);
-	TestEqual(TEXT("RoundHalfToZero32-NegOneToZeroBitPrecision"), FMath::RoundHalfToZero(-8388607.0f), -8388607.0f);
-	TestEqual(TEXT("RoundHalfToZero32-NegOneToZeroBitPrecision"), FMath::RoundHalfToZero(-8388607.5f), -8388607.0f);
-	TestEqual(TEXT("RoundHalfToZero32-NegOneToZeroBitPrecision"), FMath::RoundHalfToZero(-8388608.0f), -8388608.0f);
-
-	TestEqual(TEXT("RoundHalfToZero32-ZeroBitPrecision"), FMath::RoundHalfToZero(16777215.0f), 16777215.0f);
-	TestEqual(TEXT("RoundHalfToZero32-NegZeroBitPrecision"), FMath::RoundHalfToZero(-16777215.0f), -16777215.0f);
-
-	TestEqual(TEXT("RoundHalfToZero64-Zero"), FMath::RoundHalfToZero(0.0), 0.0);
-	TestEqual(TEXT("RoundHalfToZero64-One"), FMath::RoundHalfToZero(1.0), 1.0);
-	TestEqual(TEXT("RoundHalfToZero64-LessHalf"), FMath::RoundHalfToZero(1.4), 1.0);
-	TestEqual(TEXT("RoundHalfToZero64-NegGreaterHalf"), FMath::RoundHalfToZero(-1.4), -1.0);
-	TestEqual(TEXT("RoundHalfToZero64-LessNearHalf"), FMath::RoundHalfToZero(1.4999999999999997), 1.0);
-	TestEqual(TEXT("RoundHalfToZero64-NegGreaterNearHalf"), FMath::RoundHalfToZero(-1.4999999999999997), -1.0);
-	TestEqual(TEXT("RoundHalfToZero64-Half"), FMath::RoundHalfToZero(1.5), 1.0);
-	TestEqual(TEXT("RoundHalfToZero64-NegHalf"), FMath::RoundHalfToZero(-1.5), -1.0);
-	TestEqual(TEXT("RoundHalfToZero64-GreaterNearHalf"), FMath::RoundHalfToZero(1.5000000000000002), 2.0);
-	TestEqual(TEXT("RoundHalfToZero64-NegLesserNearHalf"), FMath::RoundHalfToZero(-1.5000000000000002), -2.0);
-	TestEqual(TEXT("RoundHalfToZero64-GreaterThanHalf"), FMath::RoundHalfToZero(1.6), 2.0);
-	TestEqual(TEXT("RoundHalfToZero64-NegLesserThanHalf"), FMath::RoundHalfToZero(-1.6), -2.0);
+	CHECK_EQUAL(FMath::RoundHalfToZero(0.0f), 0.0f);
+	CHECK_EQUAL(FMath::RoundHalfToZero(1.0f), 1.0f);
+	CHECK_EQUAL(FMath::RoundHalfToZero(1.4f), 1.0f);
+	CHECK_EQUAL(FMath::RoundHalfToZero(-1.4f), -1.0f);
+	CHECK_EQUAL(FMath::RoundHalfToZero(1.4999999f), 1.0f);
+	CHECK_EQUAL(FMath::RoundHalfToZero(-1.4999999f), -1.0f);
+	CHECK_EQUAL(FMath::RoundHalfToZero(1.5f), 1.0f);
+	CHECK_EQUAL(FMath::RoundHalfToZero(-1.5f), -1.0f);
+	CHECK_EQUAL(FMath::RoundHalfToZero(1.5000001f), 2.0f);
+	CHECK_EQUAL(FMath::RoundHalfToZero(-1.5000001f), -2.0f);
+	CHECK_EQUAL(FMath::RoundHalfToZero(1.6f), 2.0f);
+	CHECK_EQUAL(FMath::RoundHalfToZero(-1.6f), -2.0f);
+
+	CHECK_EQUAL(FMath::RoundHalfToZero(4194303.25f), 4194303.0f);
+	CHECK_EQUAL(FMath::RoundHalfToZero(4194303.5f), 4194303.0f);
+	CHECK_EQUAL(FMath::RoundHalfToZero(4194303.75f), 4194304.0f);
+	CHECK_EQUAL(FMath::RoundHalfToZero(4194304.0f), 4194304.0f);
+	CHECK_EQUAL(FMath::RoundHalfToZero(4194304.5f), 4194304.0f);
+	CHECK_EQUAL(FMath::RoundHalfToZero(-4194303.25f), -4194303.0f);
+	CHECK_EQUAL(FMath::RoundHalfToZero(-4194303.5f), -4194303.0f);
+	CHECK_EQUAL(FMath::RoundHalfToZero(-4194303.75f), -4194304.0f);
+	CHECK_EQUAL(FMath::RoundHalfToZero(-4194304.0f), -4194304.0f);
+	CHECK_EQUAL(FMath::RoundHalfToZero(-4194304.5f), -4194304.0f);
+
+	CHECK_EQUAL(FMath::RoundHalfToZero(8388607.0f), 8388607.0f);
+	CHECK_EQUAL(FMath::RoundHalfToZero(8388607.5f), 8388607.0f);
+	CHECK_EQUAL(FMath::RoundHalfToZero(8388608.0f), 8388608.0f);
+	CHECK_EQUAL(FMath::RoundHalfToZero(-8388607.0f), -8388607.0f);
+	CHECK_EQUAL(FMath::RoundHalfToZero(-8388607.5f), -8388607.0f);
+	CHECK_EQUAL(FMath::RoundHalfToZero(-8388608.0f), -8388608.0f);
+
+	CHECK_EQUAL(FMath::RoundHalfToZero(16777215.0f), 16777215.0f);
+	CHECK_EQUAL(FMath::RoundHalfToZero(-16777215.0f), -16777215.0f);
+
+	CHECK_EQUAL(FMath::RoundHalfToZero(0.0), 0.0);
+	CHECK_EQUAL(FMath::RoundHalfToZero(1.0), 1.0);
+	CHECK_EQUAL(FMath::RoundHalfToZero(1.4), 1.0);
+	CHECK_EQUAL(FMath::RoundHalfToZero(-1.4), -1.0);
+	CHECK_EQUAL(FMath::RoundHalfToZero(1.4999999999999997), 1.0);
+	CHECK_EQUAL(FMath::RoundHalfToZero(-1.4999999999999997), -1.0);
+	CHECK_EQUAL(FMath::RoundHalfToZero(1.5), 1.0);
+	CHECK_EQUAL(FMath::RoundHalfToZero(-1.5), -1.0);
+	CHECK_EQUAL(FMath::RoundHalfToZero(1.5000000000000002), 2.0);
+	CHECK_EQUAL(FMath::RoundHalfToZero(-1.5000000000000002), -2.0);
+	CHECK_EQUAL(FMath::RoundHalfToZero(1.6), 2.0);
+	CHECK_EQUAL(FMath::RoundHalfToZero(-1.6), -2.0);
 	
-	TestEqual(TEXT("RoundHalfToZero64-TwoToOneBitPrecision"), FMath::RoundHalfToZero(2251799813685247.25), 2251799813685247.0);
-	TestEqual(TEXT("RoundHalfToZero64-TwoToOneBitPrecision"), FMath::RoundHalfToZero(2251799813685247.5), 2251799813685247.0);
-	TestEqual(TEXT("RoundHalfToZero64-TwoToOneBitPrecision"), FMath::RoundHalfToZero(2251799813685247.75), 2251799813685248.0);
-	TestEqual(TEXT("RoundHalfToZero64-TwoToOneBitPrecision"), FMath::RoundHalfToZero(2251799813685248.0), 2251799813685248.0);
-	TestEqual(TEXT("RoundHalfToZero64-TwoToOneBitPrecision"), FMath::RoundHalfToZero(2251799813685248.5), 2251799813685248.0);
-	TestEqual(TEXT("RoundHalfToZero64-NegTwoToOneBitPrecision"), FMath::RoundHalfToZero(-2251799813685247.25), -2251799813685247.0);
-	TestEqual(TEXT("RoundHalfToZero64-NegTwoToOneBitPrecision"), FMath::RoundHalfToZero(-2251799813685247.5), -2251799813685247.0);
-	TestEqual(TEXT("RoundHalfToZero64-NegTwoToOneBitPrecision"), FMath::RoundHalfToZero(-2251799813685247.75), -2251799813685248.0);
-	TestEqual(TEXT("RoundHalfToZero64-NegTwoToOneBitPrecision"), FMath::RoundHalfToZero(-2251799813685248.0), -2251799813685248.0);
-	TestEqual(TEXT("RoundHalfToZero64-NegTwoToOneBitPrecision"), FMath::RoundHalfToZero(-2251799813685248.5), -2251799813685248.0);
-
-	TestEqual(TEXT("RoundHalfToZero64-OneToZeroBitPrecision"), FMath::RoundHalfToZero(4503599627370495.0), 4503599627370495.0);
-	TestEqual(TEXT("RoundHalfToZero64-OneToZeroBitPrecision"), FMath::RoundHalfToZero(4503599627370495.5), 4503599627370495.0);
-	TestEqual(TEXT("RoundHalfToZero64-OneToZeroBitPrecision"), FMath::RoundHalfToZero(4503599627370496.0), 4503599627370496.0);
-	TestEqual(TEXT("RoundHalfToZero64-NegOneToZeroBitPrecision"), FMath::RoundHalfToZero(-4503599627370495.0), -4503599627370495.0);
-	TestEqual(TEXT("RoundHalfToZero64-NegOneToZeroBitPrecision"), FMath::RoundHalfToZero(-4503599627370495.5), -4503599627370495.0);
-	TestEqual(TEXT("RoundHalfToZero64-NegOneToZeroBitPrecision"), FMath::RoundHalfToZero(-4503599627370496.0), -4503599627370496.0);
-
-	TestEqual(TEXT("RoundHalfToZero64-ZeroBitPrecision"), FMath::RoundHalfToZero(9007199254740991.0), 9007199254740991.0);
-	TestEqual(TEXT("RoundHalfToZero64-NegZeroBitPrecision"), FMath::RoundHalfToZero(-9007199254740991.0), -9007199254740991.0);
+	CHECK_EQUAL(FMath::RoundHalfToZero(2251799813685247.25), 2251799813685247.0);
+	CHECK_EQUAL(FMath::RoundHalfToZero(2251799813685247.5), 2251799813685247.0);
+	CHECK_EQUAL(FMath::RoundHalfToZero(2251799813685247.75), 2251799813685248.0);
+	CHECK_EQUAL(FMath::RoundHalfToZero(2251799813685248.0), 2251799813685248.0);
+	CHECK_EQUAL(FMath::RoundHalfToZero(2251799813685248.5), 2251799813685248.0);
+	CHECK_EQUAL(FMath::RoundHalfToZero(-2251799813685247.25), -2251799813685247.0);
+	CHECK_EQUAL(FMath::RoundHalfToZero(-2251799813685247.5), -2251799813685247.0);
+	CHECK_EQUAL(FMath::RoundHalfToZero(-2251799813685247.75), -2251799813685248.0);
+	CHECK_EQUAL(FMath::RoundHalfToZero(-2251799813685248.0), -2251799813685248.0);
+	CHECK_EQUAL(FMath::RoundHalfToZero(-2251799813685248.5), -2251799813685248.0);
+
+	CHECK_EQUAL(FMath::RoundHalfToZero(4503599627370495.0), 4503599627370495.0);
+	CHECK_EQUAL(FMath::RoundHalfToZero(4503599627370495.5), 4503599627370495.0);
+	CHECK_EQUAL(FMath::RoundHalfToZero(4503599627370496.0), 4503599627370496.0);
+	CHECK_EQUAL(FMath::RoundHalfToZero(-4503599627370495.0), -4503599627370495.0);
+	CHECK_EQUAL(FMath::RoundHalfToZero(-4503599627370495.5), -4503599627370495.0);
+	CHECK_EQUAL(FMath::RoundHalfToZero(-4503599627370496.0), -4503599627370496.0);
+
+	CHECK_EQUAL(FMath::RoundHalfToZero(9007199254740991.0), 9007199254740991.0);
+	CHECK_EQUAL(FMath::RoundHalfToZero(-9007199254740991.0), -9007199254740991.0);
 }
 
-TEST_CASE("Core::Math::Misc::Round HalfFromZero", "[Core][Math][Smoke]")
+TEST_CASE_METHOD(FAutomationTestFixture, "Core::Math::Misc::Round HalfFromZero", "[Core][Math][Smoke]")
 {
-	TestEqual(TEXT("RoundHalfFromZero32-Zero"), FMath::RoundHalfFromZero(0.0f), 0.0f);
-	TestEqual(TEXT("RoundHalfFromZero32-One"), FMath::RoundHalfFromZero(1.0f), 1.0f);
-	TestEqual(TEXT("RoundHalfFromZero32-LessHalf"), FMath::RoundHalfFromZero(1.4f), 1.0f);
-	TestEqual(TEXT("RoundHalfFromZero32-NegGreaterHalf"), FMath::RoundHalfFromZero(-1.4f), -1.0f);
-	TestEqual(TEXT("RoundHalfFromZero32-LessNearHalf"), FMath::RoundHalfFromZero(1.4999999f), 1.0f);
-	TestEqual(TEXT("RoundHalfFromZero32-NegGreaterNearHalf"), FMath::RoundHalfFromZero(-1.4999999f), -1.0f);
-	TestEqual(TEXT("RoundHalfFromZero32-Half"), FMath::RoundHalfFromZero(1.5f), 2.0f);
-	TestEqual(TEXT("RoundHalfFromZero32-NegHalf"), FMath::RoundHalfFromZero(-1.5f), -2.0f);
-	TestEqual(TEXT("RoundHalfFromZero32-LessGreaterNearHalf"), FMath::RoundHalfFromZero(1.5000001f), 2.0f);
-	TestEqual(TEXT("RoundHalfFromZero32-NegLesserNearHalf"), FMath::RoundHalfFromZero(-1.5000001f), -2.0f);
-	TestEqual(TEXT("RoundHalfFromZero32-GreaterThanHalf"), FMath::RoundHalfFromZero(1.6f), 2.0f);
-	TestEqual(TEXT("RoundHalfFromZero32-NegLesserThanHalf"), FMath::RoundHalfFromZero(-1.6f), -2.0f);
-
-	TestEqual(TEXT("RoundHalfFromZero32-TwoToOneBitPrecision"), FMath::RoundHalfFromZero(4194303.25f), 4194303.0f);
-	TestEqual(TEXT("RoundHalfFromZero32-TwoToOneBitPrecision"), FMath::RoundHalfFromZero(4194303.5f), 4194304.0f);
-	TestEqual(TEXT("RoundHalfFromZero32-TwoToOneBitPrecision"), FMath::RoundHalfFromZero(4194303.75f), 4194304.0f);
-	TestEqual(TEXT("RoundHalfFromZero32-TwoToOneBitPrecision"), FMath::RoundHalfFromZero(4194304.0f), 4194304.0f);
-	TestEqual(TEXT("RoundHalfFromZero32-TwoToOneBitPrecision"), FMath::RoundHalfFromZero(4194304.5f), 4194305.0f);
-	TestEqual(TEXT("RoundHalfFromZero32-NegTwoToOneBitPrecision"), FMath::RoundHalfFromZero(-4194303.25f), -4194303.0f);
-	TestEqual(TEXT("RoundHalfFromZero32-NegTwoToOneBitPrecision"), FMath::RoundHalfFromZero(-4194303.5f), -4194304.0f);
-	TestEqual(TEXT("RoundHalfFromZero32-NegTwoToOneBitPrecision"), FMath::RoundHalfFromZero(-4194303.75f), -4194304.0f);
-	TestEqual(TEXT("RoundHalfFromZero32-NegTwoToOneBitPrecision"), FMath::RoundHalfFromZero(-4194304.0f), -4194304.0f);
-	TestEqual(TEXT("RoundHalfFromZero32-NegTwoToOneBitPrecision"), FMath::RoundHalfFromZero(-4194304.5f), -4194305.0f);
-
-	TestEqual(TEXT("RoundHalfFromZero32-OneToZeroBitPrecision"), FMath::RoundHalfFromZero(8388607.0f), 8388607.0f);
-	TestEqual(TEXT("RoundHalfFromZero32-OneToZeroBitPrecision"), FMath::RoundHalfFromZero(8388607.5f), 8388608.0f);
-	TestEqual(TEXT("RoundHalfFromZero32-OneToZeroBitPrecision"), FMath::RoundHalfFromZero(8388608.0f), 8388608.0f);
-	TestEqual(TEXT("RoundHalfFromZero32-NegOneToZeroBitPrecision"), FMath::RoundHalfFromZero(-8388607.0f), -8388607.0f);
-	TestEqual(TEXT("RoundHalfFromZero32-NegOneToZeroBitPrecision"), FMath::RoundHalfFromZero(-8388607.5f), -8388608.0f);
-	TestEqual(TEXT("RoundHalfFromZero32-NegOneToZeroBitPrecision"), FMath::RoundHalfFromZero(-8388608.0f), -8388608.0f);
-
-	TestEqual(TEXT("RoundHalfFromZero32-ZeroBitPrecision"), FMath::RoundHalfToZero(16777215.0f), 16777215.0f);
-	TestEqual(TEXT("RoundHalfFromZero32-NegZeroBitPrecision"), FMath::RoundHalfToZero(-16777215.0f), -16777215.0f);
-
-	TestEqual(TEXT("RoundHalfFromZero64-Zero"), FMath::RoundHalfFromZero(0.0), 0.0);
-	TestEqual(TEXT("RoundHalfFromZero64-One"), FMath::RoundHalfFromZero(1.0), 1.0);
-	TestEqual(TEXT("RoundHalfFromZero64-LessHalf"), FMath::RoundHalfFromZero(1.4), 1.0);
-	TestEqual(TEXT("RoundHalfFromZero64-NegGreaterHalf"), FMath::RoundHalfFromZero(-1.4), -1.0);
-	TestEqual(TEXT("RoundHalfFromZero64-LessNearHalf"), FMath::RoundHalfFromZero(1.4999999999999997), 1.0);
-	TestEqual(TEXT("RoundHalfFromZero64-NegGreaterNearHalf"), FMath::RoundHalfFromZero(-1.4999999999999997), -1.0);
-	TestEqual(TEXT("RoundHalfFromZero64-Half"), FMath::RoundHalfFromZero(1.5), 2.0);
-	TestEqual(TEXT("RoundHalfFromZero64-NegHalf"), FMath::RoundHalfFromZero(-1.5), -2.0);
-	TestEqual(TEXT("RoundHalfFromZero64-LessGreaterNearHalf"), FMath::RoundHalfFromZero(1.5000000000000002), 2.0);
-	TestEqual(TEXT("RoundHalfFromZero64-NegLesserNearHalf"), FMath::RoundHalfFromZero(-1.5000000000000002), -2.0);
-	TestEqual(TEXT("RoundHalfFromZero64-GreaterThanHalf"), FMath::RoundHalfFromZero(1.6), 2.0);
-	TestEqual(TEXT("RoundHalfFromZero64-NegLesserThanHalf"), FMath::RoundHalfFromZero(-1.6), -2.0);
-
-	TestEqual(TEXT("RoundHalfFromZero64-TwoToOneBitPrecision"), FMath::RoundHalfFromZero(2251799813685247.25), 2251799813685247.0);
-	TestEqual(TEXT("RoundHalfFromZero64-TwoToOneBitPrecision"), FMath::RoundHalfFromZero(2251799813685247.5), 2251799813685248.0);
-	TestEqual(TEXT("RoundHalfFromZero64-TwoToOneBitPrecision"), FMath::RoundHalfFromZero(2251799813685247.75), 2251799813685248.0);
-	TestEqual(TEXT("RoundHalfFromZero64-TwoToOneBitPrecision"), FMath::RoundHalfFromZero(2251799813685248.0), 2251799813685248.0);
-	TestEqual(TEXT("RoundHalfFromZero64-TwoToOneBitPrecision"), FMath::RoundHalfFromZero(2251799813685248.5), 2251799813685249.0);
-	TestEqual(TEXT("RoundHalfFromZero64-NegTwoToOneBitPrecision"), FMath::RoundHalfFromZero(-2251799813685247.25), -2251799813685247.0);
-	TestEqual(TEXT("RoundHalfFromZero64-NegTwoToOneBitPrecision"), FMath::RoundHalfFromZero(-2251799813685247.5), -2251799813685248.0);
-	TestEqual(TEXT("RoundHalfFromZero64-NegTwoToOneBitPrecision"), FMath::RoundHalfFromZero(-2251799813685247.75), -2251799813685248.0);
-	TestEqual(TEXT("RoundHalfFromZero64-NegTwoToOneBitPrecision"), FMath::RoundHalfFromZero(-2251799813685248.0), -2251799813685248.0);
-	TestEqual(TEXT("RoundHalfFromZero64-NegTwoToOneBitPrecision"), FMath::RoundHalfFromZero(-2251799813685248.5), -2251799813685249.0);
-
-	TestEqual(TEXT("RoundHalfFromZero64-OneToZeroBitPrecision"), FMath::RoundHalfFromZero(4503599627370495.0), 4503599627370495.0);
-	TestEqual(TEXT("RoundHalfFromZero64-OneToZeroBitPrecision"), FMath::RoundHalfFromZero(4503599627370495.5), 4503599627370496.0);
-	TestEqual(TEXT("RoundHalfFromZero64-OneToZeroBitPrecision"), FMath::RoundHalfFromZero(4503599627370496.0), 4503599627370496.0);
-	TestEqual(TEXT("RoundHalfFromZero64-NegOneToZeroBitPrecision"), FMath::RoundHalfFromZero(-4503599627370495.0), -4503599627370495.0);
-	TestEqual(TEXT("RoundHalfFromZero64-NegOneToZeroBitPrecision"), FMath::RoundHalfFromZero(-4503599627370495.5), -4503599627370496.0);
-	TestEqual(TEXT("RoundHalfFromZero64-NegOneToZeroBitPrecision"), FMath::RoundHalfFromZero(-4503599627370496.0), -4503599627370496.0);
-
-	TestEqual(TEXT("RoundHalfFromZero64-ZeroBitPrecision"), FMath::RoundHalfToZero(9007199254740991.0), 9007199254740991.0);
-	TestEqual(TEXT("RoundHalfFromZero64-NegZeroBitPrecision"), FMath::RoundHalfToZero(-9007199254740991.0), -9007199254740991.0);
+	CHECK(FMath::RoundHalfFromZero(0.0f)==0.0f);
+	CHECK(FMath::RoundHalfFromZero(1.0f)==1.0f);
+	CHECK(FMath::RoundHalfFromZero(1.4f)==1.0f);
+	CHECK(FMath::RoundHalfFromZero(-1.4f)==-1.0f);
+	CHECK(FMath::RoundHalfFromZero(1.4999999f)==1.0f);
+	CHECK(FMath::RoundHalfFromZero(-1.4999999f)==-1.0f);
+	CHECK(FMath::RoundHalfFromZero(1.5f)==2.0f);
+	CHECK(FMath::RoundHalfFromZero(-1.5f)==-2.0f);
+	CHECK(FMath::RoundHalfFromZero(1.5000001f)==2.0f);
+	CHECK(FMath::RoundHalfFromZero(-1.5000001f)==-2.0f);
+	CHECK(FMath::RoundHalfFromZero(1.6f)==2.0f);
+	CHECK(FMath::RoundHalfFromZero(-1.6f)==-2.0f);
+
+	CHECK(FMath::RoundHalfFromZero(4194303.25f)==4194303.0f);
+	CHECK(FMath::RoundHalfFromZero(4194303.5f)==4194304.0f);
+	CHECK(FMath::RoundHalfFromZero(4194303.75f)==4194304.0f);
+	CHECK(FMath::RoundHalfFromZero(4194304.0f)==4194304.0f);
+	CHECK(FMath::RoundHalfFromZero(4194304.5f)==4194305.0f);
+	CHECK(FMath::RoundHalfFromZero(-4194303.25f)==-4194303.0f);
+	CHECK(FMath::RoundHalfFromZero(-4194303.5f)==-4194304.0f);
+	CHECK(FMath::RoundHalfFromZero(-4194303.75f)==-4194304.0f);
+	CHECK(FMath::RoundHalfFromZero(-4194304.0f)==-4194304.0f);
+	CHECK(FMath::RoundHalfFromZero(-4194304.5f)==-4194305.0f);
+
+	CHECK(FMath::RoundHalfFromZero(8388607.0f)==8388607.0f);
+	CHECK(FMath::RoundHalfFromZero(8388607.5f)==8388608.0f);
+	CHECK(FMath::RoundHalfFromZero(8388608.0f)==8388608.0f);
+	CHECK(FMath::RoundHalfFromZero(-8388607.0f)==-8388607.0f);
+	CHECK(FMath::RoundHalfFromZero(-8388607.5f)==-8388608.0f);
+	CHECK(FMath::RoundHalfFromZero(-8388608.0f)==-8388608.0f);
+
+	CHECK(FMath::RoundHalfToZero(16777215.0f)== 16777215.0f);
+	CHECK(FMath::RoundHalfToZero(-16777215.0f)==-16777215.0f);
+
+	CHECK(FMath::RoundHalfFromZero(0.0)==0.0);
+	CHECK(FMath::RoundHalfFromZero(1.0)==1.0);
+	CHECK(FMath::RoundHalfFromZero(1.4)==1.0);
+	CHECK(FMath::RoundHalfFromZero(-1.4)==-1.0);
+	CHECK(FMath::RoundHalfFromZero(1.4999999999999997)==1.0);
+	CHECK(FMath::RoundHalfFromZero(-1.4999999999999997)==-1.0);
+	CHECK(FMath::RoundHalfFromZero(1.5)==2.0);
+	CHECK(FMath::RoundHalfFromZero(-1.5)==-2.0);
+	CHECK(FMath::RoundHalfFromZero(1.5000000000000002)==2.0);
+	CHECK(FMath::RoundHalfFromZero(-1.5000000000000002)==-2.0);
+	CHECK(FMath::RoundHalfFromZero(1.6)==2.0);
+	CHECK(FMath::RoundHalfFromZero(-1.6)==-2.0);
+
+	CHECK(FMath::RoundHalfFromZero(2251799813685247.25)==2251799813685247.0);
+	CHECK(FMath::RoundHalfFromZero(2251799813685247.5)==2251799813685248.0);
+	CHECK(FMath::RoundHalfFromZero(2251799813685247.75)==2251799813685248.0);
+	CHECK(FMath::RoundHalfFromZero(2251799813685248.0)==2251799813685248.0);
+	CHECK(FMath::RoundHalfFromZero(2251799813685248.5)==2251799813685249.0);
+	CHECK(FMath::RoundHalfFromZero(-2251799813685247.25)==-2251799813685247.0);
+	CHECK(FMath::RoundHalfFromZero(-2251799813685247.5)==-2251799813685248.0);
+	CHECK(FMath::RoundHalfFromZero(-2251799813685247.75)==-2251799813685248.0);
+	CHECK(FMath::RoundHalfFromZero(-2251799813685248.0)==-2251799813685248.0);
+	CHECK(FMath::RoundHalfFromZero(-2251799813685248.5)==-2251799813685249.0);
+
+	CHECK(FMath::RoundHalfFromZero(4503599627370495.0)==4503599627370495.0);
+	CHECK(FMath::RoundHalfFromZero(4503599627370495.5)==4503599627370496.0);
+	CHECK(FMath::RoundHalfFromZero(4503599627370496.0)==4503599627370496.0);
+	CHECK(FMath::RoundHalfFromZero(-4503599627370495.0)==-4503599627370495.0);
+	CHECK(FMath::RoundHalfFromZero(-4503599627370495.5)==-4503599627370496.0);
+	CHECK(FMath::RoundHalfFromZero(-4503599627370496.0)==-4503599627370496.0);
+
+	CHECK(FMath::RoundHalfToZero(9007199254740991.0)==9007199254740991.0);
+	CHECK(FMath::RoundHalfToZero(-9007199254740991.0)==-9007199254740991.0);
 }
 
-TEST_CASE("Core::Math::Misc::IsNearlyEqualByULP", "[Core][Math][Smoke]")
+TEST_CASE_METHOD(FAutomationTestFixture, "Core::Math::Misc::IsNearlyEqualByULP", "[Core][Math][Smoke]")
 {
 	static const float FloatNan = std::numeric_limits<float>::quiet_NaN();
 	static const double DoubleNan = std::numeric_limits<double>::quiet_NaN();
@@ -3752,17 +3714,19 @@ TEST_CASE("Core::Math::Misc::IsNearlyEqualByULP", "[Core][Math][Smoke]")
 	{
 		if (Item.Predicate)
 		{
-			TestTrue(Item.Name + "-Float", FMath::IsNearlyEqualByULP(Item.F.A, Item.F.B, Item.ULP));
+			INFO(Item.Name + "-Float");
+			CHECK(FMath::IsNearlyEqualByULP(Item.F.A, Item.F.B, Item.ULP));
 		}
 		else
 		{
-			TestFalse(Item.Name + "-Double", FMath::IsNearlyEqualByULP(Item.D.A, Item.D.B, Item.ULP));
+			INFO(Item.Name + "-Double");
+			CHECK_FALSE(FMath::IsNearlyEqualByULP(Item.D.A, Item.D.B, Item.ULP));
 		}
 	}
 }
 
 
-TEST_CASE("Core::Math::Misc::Truncation Functions", "[Core][Math][Smoke]")
+TEST_CASE_METHOD(FAutomationTestFixture, "Core::Math::Misc::Truncation Functions", "[Core][Math][Smoke]")
 {
 	// Float: 1-bit Sign, 8-bit exponent, 23-bit mantissa, implicit 1
 	float FloatTestCases[][5] {
@@ -3814,18 +3778,18 @@ TEST_CASE("Core::Math::Misc::Truncation Functions", "[Core][Math][Smoke]")
 		float* FloatValues = FloatTestCases[TestCaseIndex];
 		float Input = FloatValues[0];
 
-		TestEqual(SubTestName(TEXT("TruncToFloat"), Input), FMath::TruncToFloat(Input), FloatValues[1]);
-		TestEqual(SubTestName(TEXT("CeilToFloat"), Input), FMath::CeilToFloat(Input), FloatValues[2]);
-		TestEqual(SubTestName(TEXT("FloorToFloat"), Input), FMath::FloorToFloat(Input), FloatValues[3]);
-		TestEqual(SubTestName(TEXT("RoundToFloat"), Input), FMath::RoundToFloat(Input), FloatValues[4]);
+		CHECK(FMath::TruncToFloat(Input)==FloatValues[1]);
+		CHECK(FMath::CeilToFloat(Input)==FloatValues[2]);
+		CHECK(FMath::FloorToFloat(Input)==FloatValues[3]);
+		CHECK(FMath::RoundToFloat(Input)==FloatValues[4]);
 
 		int* IntValues = IntTestCases[TestCaseIndex];
 		if ((float)MIN_int32 <= Input && Input <= (float)MAX_int32)
 		{
-			TestEqual(SubTestName(TEXT("TruncToInt"), Input), FMath::TruncToInt(Input), IntValues[0]);
-			TestEqual(SubTestName(TEXT("CeilToInt"), Input), FMath::CeilToInt(Input), IntValues[1]);
-			TestEqual(SubTestName(TEXT("FloorToInt"), Input), FMath::FloorToInt(Input), IntValues[2]);
-			TestEqual(SubTestName(TEXT("RoundToInt"), Input), FMath::RoundToInt(Input), IntValues[3]);
+			CHECK(FMath::TruncToInt(Input)==IntValues[0]);
+			CHECK(FMath::CeilToInt(Input)==IntValues[1]);
+			CHECK(FMath::FloorToInt(Input)==IntValues[2]);
+			CHECK(FMath::RoundToInt(Input)==IntValues[3]);
 		}
 	}
 
@@ -3858,10 +3822,10 @@ TEST_CASE("Core::Math::Misc::Truncation Functions", "[Core][Math][Smoke]")
 		double* DoubleValues = DoubleTestCases[TestCaseIndex];
 		double Input = DoubleValues[0];
 
-		TestEqual(SubTestName(TEXT("TruncToDouble"), Input), FMath::TruncToDouble(Input), DoubleValues[1]);
-		TestEqual(SubTestName(TEXT("CeilToDouble"), Input), FMath::CeilToDouble(Input), DoubleValues[2]);
-		TestEqual(SubTestName(TEXT("FloorToDouble"), Input), FMath::FloorToDouble(Input), DoubleValues[3]);
-		TestEqual(SubTestName(TEXT("RoundToDouble"), Input), FMath::RoundToDouble(Input), DoubleValues[4]);
+		CHECK(FMath::TruncToDouble(Input)==DoubleValues[1]);
+		CHECK(FMath::CeilToDouble(Input)==DoubleValues[2]);
+		CHECK(FMath::FloorToDouble(Input)==DoubleValues[3]);
+		CHECK(FMath::RoundToDouble(Input)==DoubleValues[4]);
 	}
 
 #define MATH_TRUNCATION_SPEED_TEST
@@ -3910,57 +3874,57 @@ TEST_CASE("Core::Math::Misc::Truncation Functions", "[Core][Math][Smoke]")
 
 }
 
-TEST_CASE("Core::Math::Misc::Integer Functions", "[Core][Math][Smoke]")
+TEST_CASE_METHOD(FAutomationTestFixture, "Core::Math::Misc::Integer Functions", "[Core][Math][Smoke]")
 {
 	// Test CountLeadingZeros8
-	TestEqual(TEXT("CountLeadingZeros8(0)"), FMath::CountLeadingZeros8(0), 8);
-	TestEqual(TEXT("CountLeadingZeros8(1)"), FMath::CountLeadingZeros8(1), 7);
-	TestEqual(TEXT("CountLeadingZeros8(2)"), FMath::CountLeadingZeros8(2), 6);
-	TestEqual(TEXT("CountLeadingZeros8(0x7f)"), FMath::CountLeadingZeros8(0x7f), 1);
-	TestEqual(TEXT("CountLeadingZeros8(0x80)"), FMath::CountLeadingZeros8(0x80), 0);
-	TestEqual(TEXT("CountLeadingZeros8(0xff)"), FMath::CountLeadingZeros8(0xff), 0);
+	CHECK(FMath::CountLeadingZeros8(0)==8);
+	CHECK(FMath::CountLeadingZeros8(1)==7);
+	CHECK(FMath::CountLeadingZeros8(2)==6);
+	CHECK(FMath::CountLeadingZeros8(0x7f)==1);
+	CHECK(FMath::CountLeadingZeros8(0x80)==0);
+	CHECK(FMath::CountLeadingZeros8(0xff)==0);
 
 	// Test CountLeadingZeros
-	TestEqual(TEXT("CountLeadingZeros(0)"), FMath::CountLeadingZeros(0), 32);
-	TestEqual(TEXT("CountLeadingZeros(1)"), FMath::CountLeadingZeros(1), 31);
-	TestEqual(TEXT("CountLeadingZeros(2)"), FMath::CountLeadingZeros(2), 30);
-	TestEqual(TEXT("CountLeadingZeros(0x7fffffff)"), FMath::CountLeadingZeros(0x7fffffff), 1);
-	TestEqual(TEXT("CountLeadingZeros(0x80000000)"), FMath::CountLeadingZeros(0x80000000), 0);
-	TestEqual(TEXT("CountLeadingZeros(0xffffffff)"), FMath::CountLeadingZeros(0xffffffff), 0);
+	CHECK(FMath::CountLeadingZeros(0)==32);
+	CHECK(FMath::CountLeadingZeros(1)==31);
+	CHECK(FMath::CountLeadingZeros(2)==30);
+	CHECK(FMath::CountLeadingZeros(0x7fffffff)==1);
+	CHECK(FMath::CountLeadingZeros(0x80000000)==0);
+	CHECK(FMath::CountLeadingZeros(0xffffffff)==0);
 
 	// Test CountLeadingZeros64
-	TestEqual(TEXT("CountLeadingZeros64(0)"), FMath::CountLeadingZeros64(0), uint64(64));
-	TestEqual(TEXT("CountLeadingZeros64(1)"), FMath::CountLeadingZeros64(1), uint64(63));
-	TestEqual(TEXT("CountLeadingZeros64(2)"), FMath::CountLeadingZeros64(2), uint64(62));
-	TestEqual(TEXT("CountLeadingZeros64(0x7fffffff'ffffffff)"), FMath::CountLeadingZeros64(0x7fffffff'ffffffff), uint64(1));
-	TestEqual(TEXT("CountLeadingZeros64(0x80000000'00000000)"), FMath::CountLeadingZeros64(0x80000000'00000000), uint64(0));
-	TestEqual(TEXT("CountLeadingZeros64(0xffffffff'ffffffff)"), FMath::CountLeadingZeros64(0xffffffff'ffffffff), uint64(0));
+	CHECK(FMath::CountLeadingZeros64(0)==uint64(64));
+	CHECK(FMath::CountLeadingZeros64(1)==uint64(63));
+	CHECK(FMath::CountLeadingZeros64(2)==uint64(62));
+	CHECK(FMath::CountLeadingZeros64(0x7fffffff'ffffffff)==uint64(1));
+	CHECK(FMath::CountLeadingZeros64(0x80000000'00000000)==uint64(0));
+	CHECK(FMath::CountLeadingZeros64(0xffffffff'ffffffff)==uint64(0));
 
 	// Test FloorLog2
-	TestEqual(TEXT("FloorLog2(0)"), FMath::FloorLog2(0), 0);
-	TestEqual(TEXT("FloorLog2(1)"), FMath::FloorLog2(1), 0);
-	TestEqual(TEXT("FloorLog2(2)"), FMath::FloorLog2(2), 1);
-	TestEqual(TEXT("FloorLog2(3)"), FMath::FloorLog2(3), 1);
-	TestEqual(TEXT("FloorLog2(4)"), FMath::FloorLog2(4), 2);
-	TestEqual(TEXT("FloorLog2(0x7fffffff)"), FMath::FloorLog2(0x7fffffff), 30);
-	TestEqual(TEXT("FloorLog2(0x80000000)"), FMath::FloorLog2(0x80000000), 31);
-	TestEqual(TEXT("FloorLog2(0xffffffff)"), FMath::FloorLog2(0xffffffff), 31);
+	CHECK(FMath::FloorLog2(0)==0);
+	CHECK(FMath::FloorLog2(1)==0);
+	CHECK(FMath::FloorLog2(2)==1);
+	CHECK(FMath::FloorLog2(3)==1);
+	CHECK(FMath::FloorLog2(4)==2);
+	CHECK(FMath::FloorLog2(0x7fffffff)==30);
+	CHECK(FMath::FloorLog2(0x80000000)==31);
+	CHECK(FMath::FloorLog2(0xffffffff)==31);
 
 	// Test FloorLog2_64
-	TestEqual(TEXT("FloorLog2_64(0)"), FMath::FloorLog2_64(0), uint64(0));
-	TestEqual(TEXT("FloorLog2_64(1)"), FMath::FloorLog2_64(1), uint64(0));
-	TestEqual(TEXT("FloorLog2_64(2)"), FMath::FloorLog2_64(2), uint64(1));
-	TestEqual(TEXT("FloorLog2_64(3)"), FMath::FloorLog2_64(3), uint64(1));
-	TestEqual(TEXT("FloorLog2_64(4)"), FMath::FloorLog2_64(4), uint64(2));
-	TestEqual(TEXT("FloorLog2_64(0x7fffffff)"), FMath::FloorLog2_64(0x7fffffff), uint64(30));
-	TestEqual(TEXT("FloorLog2_64(0x80000000)"), FMath::FloorLog2_64(0x80000000), uint64(31));
-	TestEqual(TEXT("FloorLog2_64(0xffffffff)"), FMath::FloorLog2_64(0xffffffff), uint64(31));
-	TestEqual(TEXT("FloorLog2_64(0x7fffffff'ffffffff)"), FMath::FloorLog2_64(0x7fffffff'ffffffff), uint64(62));
-	TestEqual(TEXT("FloorLog2_64(0x80000000'00000000)"), FMath::FloorLog2_64(0x80000000'00000000), uint64(63));
-	TestEqual(TEXT("FloorLog2_64(0xffffffff'ffffffff)"), FMath::FloorLog2_64(0xffffffff'ffffffff), uint64(63));
+	CHECK(FMath::FloorLog2_64(0)==uint64(0));
+	CHECK(FMath::FloorLog2_64(1)==uint64(0));
+	CHECK(FMath::FloorLog2_64(2)==uint64(1));
+	CHECK(FMath::FloorLog2_64(3)==uint64(1));
+	CHECK(FMath::FloorLog2_64(4)==uint64(2));
+	CHECK(FMath::FloorLog2_64(0x7fffffff)==uint64(30));
+	CHECK(FMath::FloorLog2_64(0x80000000)==uint64(31));
+	CHECK(FMath::FloorLog2_64(0xffffffff)==uint64(31));
+	CHECK(FMath::FloorLog2_64(0x7fffffff'ffffffff)==uint64(62));
+	CHECK(FMath::FloorLog2_64(0x80000000'00000000)==uint64(63));
+	CHECK(FMath::FloorLog2_64(0xffffffff'ffffffff)==uint64(63));
 }
 
-TEST_CASE("Core::Math::Misc::NaN and Inf", "[Core][Math][Smoke]")
+TEST_CASE_METHOD(FAutomationTestFixture, "Core::Math::Misc::NaN and Inf", "[Core][Math][Smoke]")
 {
 	static float FloatNan = FMath::Sqrt(-1.0f);
 	static double DoubleNan = double(FloatNan);
@@ -3977,57 +3941,67 @@ TEST_CASE("Core::Math::Misc::NaN and Inf", "[Core][Math][Smoke]")
 	static double DoubleMax = std::numeric_limits<double>::max();
 	static float FloatMax = std::numeric_limits<float>::max();
 
-	TestTrue(TEXT("HasQuietNaNFloat"), std::numeric_limits<float>::has_quiet_NaN);
-	TestTrue(TEXT("HasQuietNaNDouble"), std::numeric_limits<double>::has_quiet_NaN);
-	TestTrue(TEXT("HasInfinityFloat"), std::numeric_limits<float>::has_infinity);
-	TestTrue(TEXT("HasInfinityDouble"), std::numeric_limits<double>::has_infinity);
-
-	TestTrue(TEXT("SqrtNegOneIsNanFloat"), std::isnan(FloatNan));
-	TestTrue(TEXT("SqrtNegOneIsNanDouble"), std::isnan(DoubleNan));
-	TestTrue(TEXT("OneOverZeroIsInfFloat"), !std::isfinite(FloatInf) && !std::isnan(FloatInf));
-	TestTrue(TEXT("OneOverZeroIsInfDouble"), !std::isfinite(DoubleInf) && !std::isnan(DoubleInf));
-
-	TestTrue(TEXT("UE4IsNanTrueFloat"), FPlatformMath::IsNaN(FloatNan));
-	TestTrue(TEXT("UE4IsNanFalseFloat"), !FPlatformMath::IsNaN(0.0f));
-	TestTrue(TEXT("UE4IsNanTrueDouble"), FPlatformMath::IsNaN(DoubleNan));
-	TestTrue(TEXT("UE4IsNanFalseDouble"), !FPlatformMath::IsNaN(0.0));
-
-	TestTrue(TEXT("UE4IsFiniteTrueFloat"), FPlatformMath::IsFinite(0.0f) && !FPlatformMath::IsNaN(0.0f));
-	TestTrue(TEXT("UE4IsFiniteFalseFloat"), !FPlatformMath::IsFinite(FloatInf) && !FPlatformMath::IsNaN(FloatInf));
-	TestTrue(TEXT("UE4IsFiniteTrueDouble"), FPlatformMath::IsFinite(0.0) && !FPlatformMath::IsNaN(0.0));
-	TestTrue(TEXT("UE4IsFiniteFalseDouble"), !FPlatformMath::IsFinite(DoubleInf) && !FPlatformMath::IsNaN(DoubleInf));
-
-	TestTrue(TEXT("UE4IsNanStdFloat"), FPlatformMath::IsNaN(FloatStdNan));
-	TestTrue(TEXT("UE4IsNanStdDouble"), FPlatformMath::IsNaN(DoubleStdNan));
-
-	TestTrue(TEXT("UE4IsFiniteStdFloat"), !FPlatformMath::IsFinite(FloatStdInf) && !FPlatformMath::IsNaN(FloatStdInf));
-	TestTrue(TEXT("UE4IsFiniteStdDouble"), !FPlatformMath::IsFinite(DoubleStdInf) && !FPlatformMath::IsNaN(DoubleStdInf));
+	CHECK(std::numeric_limits<float>::has_quiet_NaN);
+	CHECK(std::numeric_limits<double>::has_quiet_NaN);
+	CHECK(std::numeric_limits<float>::has_infinity);
+	CHECK(std::numeric_limits<double>::has_infinity);
+
+	CHECK(std::isnan(FloatNan));
+	CHECK(std::isnan(DoubleNan));
+	CHECK(!std::isfinite(FloatInf));
+	CHECK(!std::isnan(FloatInf));
+	CHECK(!std::isfinite(DoubleInf));
+	CHECK(!std::isnan(DoubleInf));
+
+	CHECK(FPlatformMath::IsNaN(FloatNan));
+	CHECK(!FPlatformMath::IsNaN(0.0f));
+	CHECK(FPlatformMath::IsNaN(DoubleNan));
+	CHECK(!FPlatformMath::IsNaN(0.0));
+
+	CHECK(FPlatformMath::IsFinite(0.0f));
+	CHECK(!FPlatformMath::IsNaN(0.0f));
+	CHECK(!FPlatformMath::IsFinite(FloatInf));
+	CHECK(!FPlatformMath::IsNaN(FloatInf));
+	CHECK(FPlatformMath::IsFinite(0.0));
+	CHECK(!FPlatformMath::IsNaN(0.0));
+	CHECK(!FPlatformMath::IsFinite(DoubleInf));
+	CHECK(!FPlatformMath::IsNaN(DoubleInf));
+
+	CHECK(FPlatformMath::IsNaN(FloatStdNan));
+	CHECK(FPlatformMath::IsNaN(DoubleStdNan));
+
+	CHECK(!FPlatformMath::IsFinite(FloatStdInf));
+	CHECK(!FPlatformMath::IsNaN(FloatStdInf));
+	CHECK(!FPlatformMath::IsFinite(DoubleStdInf));
+	CHECK(!FPlatformMath::IsNaN(DoubleStdInf));
 
 	// test for Mac/Linux regression where IsFinite did not have a double equivalent so would downcast to a float and return INF.
-	TestTrue(TEXT("UE4IsFiniteDoubleMax"), FPlatformMath::IsFinite(DoubleMax) && !FPlatformMath::IsNaN(DoubleMax));
-	TestTrue(TEXT("UE4IsFiniteFloatMax"), FPlatformMath::IsFinite(FloatMax) && !FPlatformMath::IsNaN(FloatMax));
+	CHECK(FPlatformMath::IsFinite(DoubleMax));
+	CHECK(!FPlatformMath::IsNaN(DoubleMax));
+	CHECK(FPlatformMath::IsFinite(FloatMax));
+	CHECK(!FPlatformMath::IsNaN(FloatMax));
 }
 
-TEST_CASE("Core::Math::Math::Bitcast", "[Core][Math][Smoke]")
+TEST_CASE_METHOD(FAutomationTestFixture, "Core::Math::Misc::Bitcast", "[Core][Math][Smoke]")
 {
-	TestTrue(TEXT("CastFloatToInt32_0"),  FPlatformMath::AsUInt( 0.0f) == 0x00000000U);
-	TestTrue(TEXT("CastFloatToInt32_P1"), FPlatformMath::AsUInt(+1.0f) == 0x3f800000U);
-	TestTrue(TEXT("CastFloatToInt32_N1"), FPlatformMath::AsUInt(-1.0f) == 0xbf800000U);
+	CHECK(FPlatformMath::AsUInt( 0.0f) == 0x00000000U);
+	CHECK(FPlatformMath::AsUInt(+1.0f) == 0x3f800000U);
+	CHECK(FPlatformMath::AsUInt(-1.0f) == 0xbf800000U);
 
-	TestTrue(TEXT("CastFloatToInt64_0"),  FPlatformMath::AsUInt( 0.0) == 0x0000000000000000ULL);
-	TestTrue(TEXT("CastFloatToInt64_P1"), FPlatformMath::AsUInt(+1.0) == 0x3ff0000000000000ULL);
-	TestTrue(TEXT("CastFloatToInt64_N1"), FPlatformMath::AsUInt(-1.0) == 0xbff0000000000000ULL);
+	CHECK(FPlatformMath::AsUInt( 0.0) == 0x0000000000000000ULL);
+	CHECK(FPlatformMath::AsUInt(+1.0) == 0x3ff0000000000000ULL);
+	CHECK(FPlatformMath::AsUInt(-1.0) == 0xbff0000000000000ULL);
 
-	TestTrue(TEXT("CastIntToFloat32_0"),  FPlatformMath::AsFloat(static_cast<uint32>(0x00000000U)) == 0.0f);
-	TestTrue(TEXT("CastIntToFloat32_P1"), FPlatformMath::AsFloat(static_cast<uint32>(0x3f800000U)) == +1.0f);
-	TestTrue(TEXT("CastIntToFloat32_N1"), FPlatformMath::AsFloat(static_cast<uint32>(0xbf800000U)) == -1.0f);
+	CHECK(FPlatformMath::AsFloat(static_cast<uint32>(0x00000000U)) == 0.0f);
+	CHECK(FPlatformMath::AsFloat(static_cast<uint32>(0x3f800000U)) == +1.0f);
+	CHECK(FPlatformMath::AsFloat(static_cast<uint32>(0xbf800000U)) == -1.0f);
 
-	TestTrue(TEXT("CastIntToFloat64_0"),  FPlatformMath::AsFloat(static_cast<uint64>(0x0000000000000000ULL)) == 0.0);
-	TestTrue(TEXT("CastIntToFloat64_P1"), FPlatformMath::AsFloat(static_cast<uint64>(0x3ff0000000000000ULL)) == +1.0);
-	TestTrue(TEXT("CastIntToFloat64_N1"), FPlatformMath::AsFloat(static_cast<uint64>(0xbff0000000000000ULL)) == -1.0);
+	CHECK(FPlatformMath::AsFloat(static_cast<uint64>(0x0000000000000000ULL)) == 0.0);
+	CHECK(FPlatformMath::AsFloat(static_cast<uint64>(0x3ff0000000000000ULL)) == +1.0);
+	CHECK(FPlatformMath::AsFloat(static_cast<uint64>(0xbff0000000000000ULL)) == -1.0);
 }
 
-TEST_CASE("Core::Math::Misc::Wrap", "[Core][Math][Smoke]")
+TEST_CASE_METHOD(FAutomationTestFixture, "Core::Math::Misc::Wrap", "[Core][Math][Smoke]")
 {
 	for (int Val = -5; Val != 5; ++Val)
 	{
@@ -4037,7 +4011,7 @@ TEST_CASE("Core::Math::Misc::Wrap", "[Core][Math][Smoke]")
 			{
 				int Wrap = FMath::Wrap(Val, Min, Min);
 
-				TestTrue(TEXT("Wrapped value should be in the empty range"), Wrap == Min);
+				CHECK(Wrap == Min);
 			}
 
 			for (int Size = 1; Size != 5; ++Size)
@@ -4045,16 +4019,17 @@ TEST_CASE("Core::Math::Misc::Wrap", "[Core][Math][Smoke]")
 				int Max = Min + Size;
 				int Wrap = FMath::Wrap(Val, Min, Max);
 
-				TestTrue(TEXT("Wrapped value should be in the non-empty range"), Wrap >= Min && Wrap <= Max);
-				TestTrue(FString::Printf(TEXT("Wrapped value should be at a distance which is an exact multiple of the range size: (Val: %d, Min: %d, Max: %d, Wrap: %d, Mod: %d)"), Val, Min, Max, Wrap, (Wrap - Val) % Size), (Wrap - Val) % Size == 0);
+				CHECK(Wrap >= Min);
+				CHECK(Wrap <= Max);
+				CHECK((Wrap - Val) % Size == 0);
 
 				if (Val < Min)
 				{
-					TestNotEqual(TEXT("Wrapping a value from below a non-empty range should never give the max"), Wrap, Max);
+					CHECK(Wrap!=Max);
 				}
 				else if (Val > Max)
 				{
-					TestNotEqual(TEXT("Wrapping a value from above a non-empty range should never give the min"), Wrap, Min);
+					CHECK(Wrap!=Min);
 				}
 			}
 		}
@@ -4062,15 +4037,16 @@ TEST_CASE("Core::Math::Misc::Wrap", "[Core][Math][Smoke]")
 }
 
 
-TEST_CASE("Core::Math::Vector::Init From String", "[Core][Math][Smoke]")
+TEST_CASE_METHOD(FAutomationTestFixture, "Core::Math::FVector::Init From String", "[Core][Math][Smoke]")
 {
-	auto TestInitFromCompactString = [](const FString& InTestName, const FVector& InExpected)
+	auto TestInitFromCompactString = [this](const FString& InTestName, const FVector& InExpected)
 	{
 		FVector Actual(13.37f, 13.37f, 13.37f);
 		const bool bIsInitialized = Actual.InitFromCompactString(InExpected.ToCompactString());
 
-		TestTrue(FString::Printf(TEXT("%s return value"), *InTestName), bIsInitialized);
-		TestTrue(*InTestName, FVector::Dist(Actual, InExpected) < UE_KINDA_SMALL_NUMBER);
+		INFO(InTestName);
+		CHECK(bIsInitialized);
+		CHECK(FVector::Dist(Actual,InExpected)<UE_KINDA_SMALL_NUMBER );
 	};
 
 	TestInitFromCompactString(TEXT("InitFromCompactString Simple"), FVector(1.2f, 2.3f, 3.4f));
@@ -4089,11 +4065,11 @@ TEST_CASE("Core::Math::Vector::Init From String", "[Core][Math][Smoke]")
 	TestInitFromCompactString(TEXT("InitFromCompactString X == 0 && Z == 0"), FVector(0, 61.3f, 0));
 	TestInitFromCompactString(TEXT("InitFromCompactString Y == 0 && Z == 0"), FVector(65.3f, 0, 0));
 
-	TestFalse(TEXT("InitFromCompactString BadString1"), FVector().InitFromCompactString(TEXT("W(0)")));
-	TestFalse(TEXT("InitFromCompactString BadString2"), FVector().InitFromCompactString(TEXT("V(XYZ)")));
+	CHECK_FALSE(FVector().InitFromCompactString(TEXT("W(0)")));
+	CHECK_FALSE(FVector().InitFromCompactString(TEXT("V(XYZ)")));
 }
-
-TEST_CASE("Core::Math::ColorConversion", "[Core][Math]")
+ 
+TEST_CASE_METHOD(FAutomationTestFixture, "Core::Math::FColor::Conversion", "[Core][Math][Smoke]")
 {
 	// Round-trip conversion tests check that converting from uint8 color formats
 	// to float and back gives the original value. We not only want to guarantee
@@ -4108,11 +4084,11 @@ TEST_CASE("Core::Math::ColorConversion", "[Core][Math]")
 
 		FLinearColor SRGBToLinear = FLinearColor::FromSRGBColor(OriginalColor);
 		FColor SRGBConvertedBack = SRGBToLinear.ToFColorSRGB();
-		TestEqual(FString::Printf(TEXT("sRGB to linear to sRGB round-trip: %d"), Index), OriginalColor, SRGBConvertedBack);
+		CHECK(OriginalColor==SRGBConvertedBack);
 
 		FLinearColor UNORMToLinear = OriginalColor.ReinterpretAsLinear();
 		FColor UNORMConvertedBack = UNORMToLinear.QuantizeRound();
-		TestEqual(FString::Printf(TEXT("UNORM to linear to UNORM round-trip: %d"), Index), OriginalColor, UNORMConvertedBack);
+		CHECK(OriginalColor==UNORMConvertedBack);
 	}
 
 	// Test values near breakpoints/bucket boundaries to make sure they end up
@@ -4145,14 +4121,14 @@ TEST_CASE("Core::Math::ColorConversion", "[Core][Math]")
 		uint8 IndexU8 = (uint8)Index;
 		FColor BelowExpected(IndexU8, IndexU8, IndexU8, IndexU8);
 		FColor BelowConverted = FLinearColor(BelowBoundarySRGB, BelowBoundarySRGB, BelowBoundarySRGB, BelowBoundaryUNORM).ToFColorSRGB();
-		TestEqual(FString::Printf(TEXT("sRGB Boundary below: %d"), Index), BelowConverted, BelowExpected);
+		CHECK(BelowConverted==BelowExpected);
 
 		float AboveBoundaryUNORM = (BucketMidpoint + DistanceToMidpoint) / 255.f;
 		float AboveBoundarySRGB = ReferenceSRGBToLinear(AboveBoundaryUNORM);
 		uint8 IndexPlus1U8 = (uint8)(Index + 1);
 		FColor AboveExpected(IndexPlus1U8, IndexPlus1U8, IndexPlus1U8, IndexPlus1U8);
 		FColor AboveConverted = FLinearColor(AboveBoundarySRGB, AboveBoundarySRGB, AboveBoundarySRGB, AboveBoundaryUNORM).ToFColorSRGB();
-		TestEqual(FString::Printf(TEXT("sRGB Boundary above: %d"), Index), AboveConverted, AboveExpected);
+		CHECK(AboveConverted==AboveExpected);
 	}
 
 	// Between the two tests above, we have good coverage of what happens inside [0,1];
@@ -4164,13 +4140,14 @@ TEST_CASE("Core::Math::ColorConversion", "[Core][Math]")
 	const float MediumLarge = 10.f; // Outside the range, but not hugely so
 	const float VeryLarge = 1e+7f; // Far outside the range
 
-	auto TestExtremalValue = [](const TCHAR* InTestName, const float InValue, uint8 InExpected)
+	auto TestExtremalValue = [this](const TCHAR* InTestName, const float InValue, uint8 InExpected)
 	{
 		FLinearColor LinColor(InValue, InValue, InValue, InValue);
 		FColor ConvertedColor = LinColor.ToFColorSRGB();
 		FColor ExpectedColor(InExpected, InExpected, InExpected, InExpected);
 
-		TestEqual(InTestName, ConvertedColor, ExpectedColor);
+		INFO(InTestName);
+		CHECK(ConvertedColor==ExpectedColor);
 	};
 
 	TestExtremalValue(TEXT("Extremal NaN"), NaN, 0);
@@ -4193,57 +4170,56 @@ TEST_CASE("Core::Math::ColorConversion", "[Core][Math]")
 // as possible.
 namespace UnrealMathTestInternal {
 
-	typedef union
-	{
-		uint32 u;
-		float f;
-	} stbir__FP32;
-
-	static const uint32 stb_fp32_to_srgb8_tab4[104] = {
-		0x0073000d, 0x007a000d, 0x0080000d, 0x0087000d, 0x008d000d, 0x0094000d, 0x009a000d, 0x00a1000d,
-		0x00a7001a, 0x00b4001a, 0x00c1001a, 0x00ce001a, 0x00da001a, 0x00e7001a, 0x00f4001a, 0x0101001a,
-		0x010e0033, 0x01280033, 0x01410033, 0x015b0033, 0x01750033, 0x018f0033, 0x01a80033, 0x01c20033,
-		0x01dc0067, 0x020f0067, 0x02430067, 0x02760067, 0x02aa0067, 0x02dd0067, 0x03110067, 0x03440067,
-		0x037800ce, 0x03df00ce, 0x044600ce, 0x04ad00ce, 0x051400ce, 0x057b00c5, 0x05dd00bc, 0x063b00b5,
-		0x06970158, 0x07420142, 0x07e30130, 0x087b0120, 0x090b0112, 0x09940106, 0x0a1700fc, 0x0a9500f2,
-		0x0b0f01cb, 0x0bf401ae, 0x0ccb0195, 0x0d950180, 0x0e56016e, 0x0f0d015e, 0x0fbc0150, 0x10630143,
-		0x11070264, 0x1238023e, 0x1357021d, 0x14660201, 0x156601e9, 0x165a01d3, 0x174401c0, 0x182401af,
-		0x18fe0331, 0x1a9602fe, 0x1c1502d2, 0x1d7e02ad, 0x1ed4028d, 0x201a0270, 0x21520256, 0x227d0240,
-		0x239f0443, 0x25c003fe, 0x27bf03c4, 0x29a10392, 0x2b6a0367, 0x2d1d0341, 0x2ebe031f, 0x304d0300,
-		0x31d105b0, 0x34a80555, 0x37520507, 0x39d504c5, 0x3c37048b, 0x3e7c0458, 0x40a8042a, 0x42bd0401,
-		0x44c20798, 0x488e071e, 0x4c1c06b6, 0x4f76065d, 0x52a50610, 0x55ac05cc, 0x5892058f, 0x5b590559,
-		0x5e0c0a23, 0x631c0980, 0x67db08f6, 0x6c55087f, 0x70940818, 0x74a007bd, 0x787d076c, 0x7c330723,
-	};
-
-	static uint8 stbir__linear_to_srgb_uchar_fast(float in)
-	{
-		static const stbir__FP32 almostone = { 0x3f7fffff }; // 1-eps
-		static const stbir__FP32 minval = { (127 - 13) << 23 };
-		uint32 tab, bias, scale, t;
-		stbir__FP32 f;
-
-		// Clamp to [2^(-13), 1-eps]; these two values map to 0 and 1, respectively.
-		// The tests are carefully written so that NaNs map to 0, same as in the reference
-		// implementation.
-		if (!(in > minval.f)) // written this way to catch NaNs
-			in = minval.f;
-		if (in > almostone.f)
-			in = almostone.f;
-
-		// Do the table lookup and unpack bias, scale
-		f.f = in;
-
-		tab = stb_fp32_to_srgb8_tab4[(f.u - minval.u) >> 20];
-		bias = (tab >> 16) << 9;
-		scale = tab & 0xffff;
-
-		// Grab next-highest mantissa bits and perform linear interpolation
-		t = (f.u >> 12) & 0xff;
-		return (uint8)((bias + scale * t) >> 16);
-	}
+typedef union
+{
+	uint32 u;
+	float f;
+} stbir__FP32;
+
+static const uint32 stb_fp32_to_srgb8_tab4[104] = {
+	0x0073000d, 0x007a000d, 0x0080000d, 0x0087000d, 0x008d000d, 0x0094000d, 0x009a000d, 0x00a1000d,
+	0x00a7001a, 0x00b4001a, 0x00c1001a, 0x00ce001a, 0x00da001a, 0x00e7001a, 0x00f4001a, 0x0101001a,
+	0x010e0033, 0x01280033, 0x01410033, 0x015b0033, 0x01750033, 0x018f0033, 0x01a80033, 0x01c20033,
+	0x01dc0067, 0x020f0067, 0x02430067, 0x02760067, 0x02aa0067, 0x02dd0067, 0x03110067, 0x03440067,
+	0x037800ce, 0x03df00ce, 0x044600ce, 0x04ad00ce, 0x051400ce, 0x057b00c5, 0x05dd00bc, 0x063b00b5,
+	0x06970158, 0x07420142, 0x07e30130, 0x087b0120, 0x090b0112, 0x09940106, 0x0a1700fc, 0x0a9500f2,
+	0x0b0f01cb, 0x0bf401ae, 0x0ccb0195, 0x0d950180, 0x0e56016e, 0x0f0d015e, 0x0fbc0150, 0x10630143,
+	0x11070264, 0x1238023e, 0x1357021d, 0x14660201, 0x156601e9, 0x165a01d3, 0x174401c0, 0x182401af,
+	0x18fe0331, 0x1a9602fe, 0x1c1502d2, 0x1d7e02ad, 0x1ed4028d, 0x201a0270, 0x21520256, 0x227d0240,
+	0x239f0443, 0x25c003fe, 0x27bf03c4, 0x29a10392, 0x2b6a0367, 0x2d1d0341, 0x2ebe031f, 0x304d0300,
+	0x31d105b0, 0x34a80555, 0x37520507, 0x39d504c5, 0x3c37048b, 0x3e7c0458, 0x40a8042a, 0x42bd0401,
+	0x44c20798, 0x488e071e, 0x4c1c06b6, 0x4f76065d, 0x52a50610, 0x55ac05cc, 0x5892058f, 0x5b590559,
+	0x5e0c0a23, 0x631c0980, 0x67db08f6, 0x6c55087f, 0x70940818, 0x74a007bd, 0x787d076c, 0x7c330723,
+};
+
+static uint8 stbir__linear_to_srgb_uchar_fast(float in)
+{
+	static const stbir__FP32 almostone = { 0x3f7fffff }; // 1-eps
+	static const stbir__FP32 minval = { (127 - 13) << 23 };
+	uint32 tab, bias, scale, t;
+	stbir__FP32 f;
+
+	// Clamp to [2^(-13), 1-eps]; these two values map to 0 and 1, respectively.
+	// The tests are carefully written so that NaNs map to 0, same as in the reference
+	// implementation.
+	if (!(in > minval.f)) // written this way to catch NaNs
+		in = minval.f;
+	if (in > almostone.f)
+		in = almostone.f;
+
+	// Do the table lookup and unpack bias, scale
+	f.f = in;
+
+	tab = stb_fp32_to_srgb8_tab4[(f.u - minval.u) >> 20];
+	bias = (tab >> 16) << 9;
+	scale = tab & 0xffff;
+
+	// Grab next-highest mantissa bits and perform linear interpolation
+	t = (f.u >> 12) & 0xff;
+	return (uint8)((bias + scale * t) >> 16);
 }
 
-TEST_CASE("Core::Math::ColorConversionHeavy", "[Core][Math][PlatformFlaky]")
+TEST_CASE_METHOD(FAutomationTestFixture, "Core::Math::FColor::Conversion Heavy", "[Core][Math][.Stress]")
 {
 	// WARNING: This test runs for a good while (at time of writing, ~90s when using
 	// a single core, proportionately less on many-core machines) with no user
@@ -4336,8 +4312,7 @@ TEST_CASE("Core::Math::ColorConversionHeavy", "[Core][Math][PlatformFlaky]")
 		}
 	);
 
-	if (Shared.Failed)
-	{
-		TestEqual(FString::Printf(TEXT("Conversion heavy sRGB CurrentBits=0x%08x"), Shared.FailedValue), Shared.FailedConverted, Shared.FailedExpected);
-	}
+	CHECK(Shared.FailedConverted==Shared.FailedExpected);
+}
+
 }
diff --git a/Engine/Source/Runtime/Core/Tests/Memory/CompositeBufferTest.cpp b/Engine/Source/Runtime/Core/Tests/Memory/CompositeBufferTest.cpp
--- a/Engine/Source/Runtime/Core/Tests/Memory/CompositeBufferTest.cpp
+++ b/Engine/Source/Runtime/Core/Tests/Memory/CompositeBufferTest.cpp
@@ -19,33 +19,33 @@ static_assert(!std::is_constructible_v<FCompositeBuffer, FMemoryView>);
 static_assert(!std::is_constructible_v<FCompositeBuffer, const uint8[4]>);
 
 
-TEST_CASE("Core::Memory::FCompositeBuffer::CompositeBuffer", "[Core][Math][Smoke]")
+TEST_CASE_METHOD(FAutomationTestFixture, "Core::Memory::FCompositeBuffer::Misc", "[Core][Memory][Smoke]")
 {
 	// Test Null
 	{
 		FCompositeBuffer Buffer;
-		TestTrue(TEXT("FCompositeBuffer().IsNull()"), Buffer.IsNull());
-		TestTrue(TEXT("FCompositeBuffer().IsOwned()"), Buffer.IsOwned());
-		TestTrue(TEXT("FCompositeBuffer().MakeOwned().IsNull()"), Buffer.MakeOwned().IsNull());
-		TestTrue(TEXT("FCompositeBuffer().ToShared().IsNull()"), Buffer.ToShared().IsNull());
-		TestTrue(TEXT("FCompositeBuffer().Mid(0, 0).IsNull()"), Buffer.Mid(0, 0).IsNull());
-		TestEqual(TEXT("FCompositeBuffer().GetSize()"), Buffer.GetSize(), uint64(0));
-		TestEqual(TEXT("FCompositeBuffer().GetSegments().Num()"), Buffer.GetSegments().Num(), 0);
+		TEST_TRUE(TEXT("FCompositeBuffer().IsNull()"), Buffer.IsNull());
+		TEST_TRUE(TEXT("FCompositeBuffer().IsOwned()"), Buffer.IsOwned());
+		TEST_TRUE(TEXT("FCompositeBuffer().MakeOwned().IsNull()"), Buffer.MakeOwned().IsNull());
+		TEST_TRUE(TEXT("FCompositeBuffer().ToShared().IsNull()"), Buffer.ToShared().IsNull());
+		TEST_TRUE(TEXT("FCompositeBuffer().Mid(0, 0).IsNull()"), Buffer.Mid(0, 0).IsNull());
+		TEST_EQUAL(TEXT("FCompositeBuffer().GetSize()"), Buffer.GetSize(), uint64(0));
+		TEST_EQUAL(TEXT("FCompositeBuffer().GetSegments().Num()"), Buffer.GetSegments().Num(), 0);
 
 		FUniqueBuffer CopyBuffer;
-		TestTrue(TEXT("FCompositeBuffer().ViewOrCopyRange(0, 0).IsEmpty()"), Buffer.ViewOrCopyRange(0, 0, CopyBuffer).IsEmpty() && CopyBuffer.IsNull());
+		TEST_TRUE(TEXT("FCompositeBuffer().ViewOrCopyRange(0, 0).IsEmpty()"), Buffer.ViewOrCopyRange(0, 0, CopyBuffer).IsEmpty() && CopyBuffer.IsNull());
 
 		FMutableMemoryView CopyView;
 		Buffer.CopyTo(CopyView);
 
 		uint32 VisitCount = 0;
 		Buffer.IterateRange(0, 0, [&VisitCount](FMemoryView) { ++VisitCount; });
-		TestEqual(TEXT("FCompositeBuffer().IterateRange(0, 0)"), VisitCount, uint32(0));
+		TEST_EQUAL(TEXT("FCompositeBuffer().IterateRange(0, 0)"), VisitCount, uint32(0));
 	}
 	{
 		FCompositeBuffer Buffer(FSharedBuffer(), FCompositeBuffer(), TArray<FSharedBuffer>{FSharedBuffer()});
-		TestTrue(TEXT("FCompositeBuffer(Nulls).IsNull()"), Buffer.IsNull());
-		TestEqual(TEXT("FCompositeBuffer(Nulls).GetSegments().Num()"), Buffer.GetSegments().Num(), 0);
+		TEST_TRUE(TEXT("FCompositeBuffer(Nulls).IsNull()"), Buffer.IsNull());
+		TEST_EQUAL(TEXT("FCompositeBuffer(Nulls).GetSegments().Num()"), Buffer.GetSegments().Num(), 0);
 	}
 
 	// Test Empty
@@ -53,46 +53,46 @@ TEST_CASE("Core::Memory::FCompositeBuffer::CompositeBuffer", "[Core][Math][Smoke
 		const uint8 EmptyArray[]{0};
 		const FSharedBuffer EmptyView = FSharedBuffer::MakeView(EmptyArray, 0);
 		FCompositeBuffer Buffer(EmptyView);
-		TestFalse(TEXT("FCompositeBuffer(Empty).IsNull()"), Buffer.IsNull());
-		TestFalse(TEXT("FCompositeBuffer(Empty).IsOwned()"), Buffer.IsOwned());
-		TestFalse(TEXT("FCompositeBuffer(Empty).MakeOwned().IsNull()"), Buffer.MakeOwned().IsNull());
-		TestTrue(TEXT("FCompositeBuffer(Empty).MakeOwned().IsOwned()"), Buffer.MakeOwned().IsOwned());
-		TestEqual(TEXT("FCompositeBuffer(Empty).ToShared()"), Buffer.ToShared(), EmptyView);
-		TestEqual(TEXT("FCompositeBuffer(Empty).Mid(0, 0)"), Buffer.Mid(0, 0).ToShared(), EmptyView);
-		TestEqual(TEXT("FCompositeBuffer(Empty).GetSize()"), Buffer.GetSize(), uint64(0));
-		TestEqual(TEXT("FCompositeBuffer(Empty).GetSegments().Num()"), Buffer.GetSegments().Num(), 1);
-		TestEqual(TEXT("FCompositeBuffer(Empty).GetSegments()[0]"), Buffer.GetSegments()[0], EmptyView);
+		TEST_FALSE(TEXT("FCompositeBuffer(Empty).IsNull()"), Buffer.IsNull());
+		TEST_FALSE(TEXT("FCompositeBuffer(Empty).IsOwned()"), Buffer.IsOwned());
+		TEST_FALSE(TEXT("FCompositeBuffer(Empty).MakeOwned().IsNull()"), Buffer.MakeOwned().IsNull());
+		TEST_TRUE(TEXT("FCompositeBuffer(Empty).MakeOwned().IsOwned()"), Buffer.MakeOwned().IsOwned());
+		TEST_EQUAL(TEXT("FCompositeBuffer(Empty).ToShared()"), Buffer.ToShared(), EmptyView);
+		TEST_EQUAL(TEXT("FCompositeBuffer(Empty).Mid(0, 0)"), Buffer.Mid(0, 0).ToShared(), EmptyView);
+		TEST_EQUAL(TEXT("FCompositeBuffer(Empty).GetSize()"), Buffer.GetSize(), uint64(0));
+		TEST_EQUAL(TEXT("FCompositeBuffer(Empty).GetSegments().Num()"), Buffer.GetSegments().Num(), 1);
+		TEST_EQUAL(TEXT("FCompositeBuffer(Empty).GetSegments()[0]"), Buffer.GetSegments()[0], EmptyView);
 
 		FUniqueBuffer CopyBuffer;
-		TestTrue(TEXT("FCompositeBuffer(Empty).ViewOrCopyRange(0, 0)"), Buffer.ViewOrCopyRange(0, 0, CopyBuffer) == EmptyView.GetView() && CopyBuffer.IsNull());
+		TEST_TRUE(TEXT("FCompositeBuffer(Empty).ViewOrCopyRange(0, 0)"), Buffer.ViewOrCopyRange(0, 0, CopyBuffer) == EmptyView.GetView() && CopyBuffer.IsNull());
 
 		FMutableMemoryView CopyView;
 		Buffer.CopyTo(CopyView);
 
 		uint32 VisitCount = 0;
 		Buffer.IterateRange(0, 0, [&VisitCount](FMemoryView) { ++VisitCount; });
-		TestEqual(TEXT("FCompositeBuffer(Empty).IterateRange(0, 0)"), VisitCount, uint32(1));
+		TEST_EQUAL(TEXT("FCompositeBuffer(Empty).IterateRange(0, 0)"), VisitCount, uint32(1));
 	}
 	{
 		const uint8 EmptyArray[1]{};
 		const FSharedBuffer EmptyView1 = FSharedBuffer::MakeView(EmptyArray, 0);
 		const FSharedBuffer EmptyView2 = FSharedBuffer::MakeView(EmptyArray + 1, 0);
 		FCompositeBuffer Buffer(TArray<FSharedBuffer>{EmptyView1}, FCompositeBuffer(EmptyView2));
-		TestEqual(TEXT("FCompositeBuffer(Empty, Empty).Mid(0, 0)"), Buffer.Mid(0, 0).ToShared(), EmptyView1);
-		TestEqual(TEXT("FCompositeBuffer(Empty, Empty).GetSize()"), Buffer.GetSize(), uint64(0));
-		TestEqual(TEXT("FCompositeBuffer(Empty, Empty).GetSegments().Num()"), Buffer.GetSegments().Num(), 2);
-		TestEqual(TEXT("FCompositeBuffer(Empty, Empty).GetSegments()[0]"), Buffer.GetSegments()[0], EmptyView1);
-		TestEqual(TEXT("FCompositeBuffer(Empty, Empty).GetSegments()[1]"), Buffer.GetSegments()[1], EmptyView2);
+		TEST_EQUAL(TEXT("FCompositeBuffer(Empty, Empty).Mid(0, 0)"), Buffer.Mid(0, 0).ToShared(), EmptyView1);
+		TEST_EQUAL(TEXT("FCompositeBuffer(Empty, Empty).GetSize()"), Buffer.GetSize(), uint64(0));
+		TEST_EQUAL(TEXT("FCompositeBuffer(Empty, Empty).GetSegments().Num()"), Buffer.GetSegments().Num(), 2);
+		TEST_EQUAL(TEXT("FCompositeBuffer(Empty, Empty).GetSegments()[0]"), Buffer.GetSegments()[0], EmptyView1);
+		TEST_EQUAL(TEXT("FCompositeBuffer(Empty, Empty).GetSegments()[1]"), Buffer.GetSegments()[1], EmptyView2);
 
 		FUniqueBuffer CopyBuffer;
-		TestTrue(TEXT("FCompositeBuffer(Empty, Empty).ViewOrCopyRange(0, 0)"), Buffer.ViewOrCopyRange(0, 0, CopyBuffer) == EmptyView1.GetView() && CopyBuffer.IsNull());
+		TEST_TRUE(TEXT("FCompositeBuffer(Empty, Empty).ViewOrCopyRange(0, 0)"), Buffer.ViewOrCopyRange(0, 0, CopyBuffer) == EmptyView1.GetView() && CopyBuffer.IsNull());
 
 		FMutableMemoryView CopyView;
 		Buffer.CopyTo(CopyView);
 
 		uint32 VisitCount = 0;
 		Buffer.IterateRange(0, 0, [&VisitCount](FMemoryView) { ++VisitCount; });
-		TestEqual(TEXT("FCompositeBuffer(Empty, Empty).IterateRange(0, 0)"), VisitCount, uint32(1));
+		TEST_EQUAL(TEXT("FCompositeBuffer(Empty, Empty).IterateRange(0, 0)"), VisitCount, uint32(1));
 	}
 
 	// Test Flat
@@ -100,31 +100,31 @@ TEST_CASE("Core::Memory::FCompositeBuffer::CompositeBuffer", "[Core][Math][Smoke
 		const uint8 FlatArray[]{1, 2, 3, 4, 5, 6, 7, 8};
 		const FSharedBuffer FlatView = FSharedBuffer::Clone(MakeMemoryView(FlatArray));
 		FCompositeBuffer Buffer(FlatView);
-		TestFalse(TEXT("FCompositeBuffer(Flat).IsNull()"), Buffer.IsNull());
-		TestTrue(TEXT("FCompositeBuffer(Flat).IsOwned()"), Buffer.IsOwned());
-		TestEqual(TEXT("FCompositeBuffer(Flat).ToShared()"), Buffer.ToShared(), FlatView);
-		TestEqual(TEXT("FCompositeBuffer(Flat).MakeOwned().ToShared()"), Buffer.MakeOwned().ToShared(), FlatView);
-		TestEqual(TEXT("FCompositeBuffer(Flat).Mid(0)"), Buffer.Mid(0).ToShared(), FlatView);
-		TestEqual(TEXT("FCompositeBuffer(Flat).Mid(4)"), Buffer.Mid(4).ToShared().GetView(), FlatView.GetView().Mid(4));
-		TestEqual(TEXT("FCompositeBuffer(Flat).Mid(8)"), Buffer.Mid(8).ToShared().GetView(), FlatView.GetView().Mid(8));
-		TestEqual(TEXT("FCompositeBuffer(Flat).Mid(4, 2)"), Buffer.Mid(4, 2).ToShared().GetView(), FlatView.GetView().Mid(4, 2));
-		TestEqual(TEXT("FCompositeBuffer(Flat).Mid(8).GetData()"), Buffer.Mid(8).ToShared().GetView().GetData(), FlatView.GetView().Mid(8).GetData());
-		TestEqual(TEXT("FCompositeBuffer(Flat).Mid(4, 2).GetData()"), Buffer.Mid(4, 2).ToShared().GetView().GetData(), FlatView.GetView().Mid(4, 2).GetData());
-		TestEqual(TEXT("FCompositeBuffer(Flat).Mid(8, 0).GetData()"), Buffer.Mid(8, 0).ToShared().GetView().GetData(), FlatView.GetView().Mid(8, 0).GetData());
-		TestEqual(TEXT("FCompositeBuffer(Flat).GetSize()"), Buffer.GetSize(), uint64(sizeof(FlatArray)));
-		TestEqual(TEXT("FCompositeBuffer(Flat).GetSegments().Num()"), Buffer.GetSegments().Num(), 1);
-		TestEqual(TEXT("FCompositeBuffer(Flat).GetSegments()[0]"), Buffer.GetSegments()[0], FlatView);
+		TEST_FALSE(TEXT("FCompositeBuffer(Flat).IsNull()"), Buffer.IsNull());
+		TEST_TRUE(TEXT("FCompositeBuffer(Flat).IsOwned()"), Buffer.IsOwned());
+		TEST_EQUAL(TEXT("FCompositeBuffer(Flat).ToShared()"), Buffer.ToShared(), FlatView);
+		TEST_EQUAL(TEXT("FCompositeBuffer(Flat).MakeOwned().ToShared()"), Buffer.MakeOwned().ToShared(), FlatView);
+		TEST_EQUAL(TEXT("FCompositeBuffer(Flat).Mid(0)"), Buffer.Mid(0).ToShared(), FlatView);
+		TEST_EQUAL(TEXT("FCompositeBuffer(Flat).Mid(4)"), Buffer.Mid(4).ToShared().GetView(), FlatView.GetView().Mid(4));
+		TEST_EQUAL(TEXT("FCompositeBuffer(Flat).Mid(8)"), Buffer.Mid(8).ToShared().GetView(), FlatView.GetView().Mid(8));
+		TEST_EQUAL(TEXT("FCompositeBuffer(Flat).Mid(4, 2)"), Buffer.Mid(4, 2).ToShared().GetView(), FlatView.GetView().Mid(4, 2));
+		TEST_EQUAL(TEXT("FCompositeBuffer(Flat).Mid(8).GetData()"), Buffer.Mid(8).ToShared().GetView().GetData(), FlatView.GetView().Mid(8).GetData());
+		TEST_EQUAL(TEXT("FCompositeBuffer(Flat).Mid(4, 2).GetData()"), Buffer.Mid(4, 2).ToShared().GetView().GetData(), FlatView.GetView().Mid(4, 2).GetData());
+		TEST_EQUAL(TEXT("FCompositeBuffer(Flat).Mid(8, 0).GetData()"), Buffer.Mid(8, 0).ToShared().GetView().GetData(), FlatView.GetView().Mid(8, 0).GetData());
+		TEST_EQUAL(TEXT("FCompositeBuffer(Flat).GetSize()"), Buffer.GetSize(), uint64(sizeof(FlatArray)));
+		TEST_EQUAL(TEXT("FCompositeBuffer(Flat).GetSegments().Num()"), Buffer.GetSegments().Num(), 1);
+		TEST_EQUAL(TEXT("FCompositeBuffer(Flat).GetSegments()[0]"), Buffer.GetSegments()[0], FlatView);
 
 		FUniqueBuffer CopyBuffer;
-		TestTrue(TEXT("FCompositeBuffer(Flat).ViewOrCopyRange(0, 0)"), Buffer.ViewOrCopyRange(0, sizeof(FlatArray), CopyBuffer) == FlatView.GetView() && CopyBuffer.IsNull());
+		TEST_TRUE(TEXT("FCompositeBuffer(Flat).ViewOrCopyRange(0, 0)"), Buffer.ViewOrCopyRange(0, sizeof(FlatArray), CopyBuffer) == FlatView.GetView() && CopyBuffer.IsNull());
 
 		uint8 CopyArray[sizeof(FlatArray) - 3];
 		Buffer.CopyTo(MakeMemoryView(CopyArray), 3);
-		TestTrue(TEXT("FCompositeBuffer(Flat).CopyTo()"), MakeMemoryView(CopyArray).EqualBytes(MakeMemoryView(FlatArray) + 3));
+		TEST_TRUE(TEXT("FCompositeBuffer(Flat).CopyTo()"), MakeMemoryView(CopyArray).EqualBytes(MakeMemoryView(FlatArray) + 3));
 
 		uint32 VisitCount = 0;
 		Buffer.IterateRange(0, sizeof(FlatArray), [&VisitCount](FMemoryView) { ++VisitCount; });
-		TestEqual(TEXT("FCompositeBuffer(Flat).IterateRange(0, N)"), VisitCount, uint32(1));
+		TEST_EQUAL(TEXT("FCompositeBuffer(Flat).IterateRange(0, N)"), VisitCount, uint32(1));
 	}
 
 	// Test Composite
@@ -133,35 +133,35 @@ TEST_CASE("Core::Memory::FCompositeBuffer::CompositeBuffer", "[Core][Math][Smoke
 		const FSharedBuffer FlatView1 = FSharedBuffer::MakeView(MakeMemoryView(FlatArray).Left(4));
 		const FSharedBuffer FlatView2 = FSharedBuffer::MakeView(MakeMemoryView(FlatArray).Right(4));
 		FCompositeBuffer Buffer(FlatView1, FlatView2);
-		TestFalse(TEXT("FCompositeBuffer(Composite).IsNull()"), Buffer.IsNull());
-		TestFalse(TEXT("FCompositeBuffer(Composite).IsOwned()"), Buffer.IsOwned());
-		TestTrue(TEXT("FCompositeBuffer(Composite).ToShared()"), Buffer.ToShared().GetView().EqualBytes(MakeMemoryView(FlatArray)));
-		TestTrue(TEXT("FCompositeBuffer(Composite).Mid(2, 4)"), Buffer.Mid(2, 4).ToShared().GetView().EqualBytes(MakeMemoryView(FlatArray).Mid(2, 4)));
-		TestEqual(TEXT("FCompositeBuffer(Composite).Mid(0, 4)"), Buffer.Mid(0, 4).ToShared(), FlatView1);
-		TestEqual(TEXT("FCompositeBuffer(Composite).Mid(4, 4)"), Buffer.Mid(4, 4).ToShared(), FlatView2);
-		TestEqual(TEXT("FCompositeBuffer(Composite).Mid(4, 0).GetData()"), Buffer.Mid(4, 0).ToShared().GetView().GetData(), static_cast<const void*>(FlatArray + 4));
-		TestEqual(TEXT("FCompositeBuffer(Composite).Mid(8, 0).GetData()"), Buffer.Mid(8, 0).ToShared().GetView().GetData(), static_cast<const void*>(FlatArray + 8));
-		TestEqual(TEXT("FCompositeBuffer(Composite).GetSize()"), Buffer.GetSize(), uint64(sizeof(FlatArray)));
-		TestEqual(TEXT("FCompositeBuffer(Composite).GetSegments().Num()"), Buffer.GetSegments().Num(), 2);
-		TestEqual(TEXT("FCompositeBuffer(Composite).GetSegments()[0]"), Buffer.GetSegments()[0], FlatView1);
-		TestEqual(TEXT("FCompositeBuffer(Composite).GetSegments()[2]"), Buffer.GetSegments()[1], FlatView2);
+		TEST_FALSE(TEXT("FCompositeBuffer(Composite).IsNull()"), Buffer.IsNull());
+		TEST_FALSE(TEXT("FCompositeBuffer(Composite).IsOwned()"), Buffer.IsOwned());
+		TEST_TRUE(TEXT("FCompositeBuffer(Composite).ToShared()"), Buffer.ToShared().GetView().EqualBytes(MakeMemoryView(FlatArray)));
+		TEST_TRUE(TEXT("FCompositeBuffer(Composite).Mid(2, 4)"), Buffer.Mid(2, 4).ToShared().GetView().EqualBytes(MakeMemoryView(FlatArray).Mid(2, 4)));
+		TEST_EQUAL(TEXT("FCompositeBuffer(Composite).Mid(0, 4)"), Buffer.Mid(0, 4).ToShared(), FlatView1);
+		TEST_EQUAL(TEXT("FCompositeBuffer(Composite).Mid(4, 4)"), Buffer.Mid(4, 4).ToShared(), FlatView2);
+		TEST_EQUAL(TEXT("FCompositeBuffer(Composite).Mid(4, 0).GetData()"), Buffer.Mid(4, 0).ToShared().GetView().GetData(), static_cast<const void*>(FlatArray + 4));
+		TEST_EQUAL(TEXT("FCompositeBuffer(Composite).Mid(8, 0).GetData()"), Buffer.Mid(8, 0).ToShared().GetView().GetData(), static_cast<const void*>(FlatArray + 8));
+		TEST_EQUAL(TEXT("FCompositeBuffer(Composite).GetSize()"), Buffer.GetSize(), uint64(sizeof(FlatArray)));
+		TEST_EQUAL(TEXT("FCompositeBuffer(Composite).GetSegments().Num()"), Buffer.GetSegments().Num(), 2);
+		TEST_EQUAL(TEXT("FCompositeBuffer(Composite).GetSegments()[0]"), Buffer.GetSegments()[0], FlatView1);
+		TEST_EQUAL(TEXT("FCompositeBuffer(Composite).GetSegments()[2]"), Buffer.GetSegments()[1], FlatView2);
 
 		FUniqueBuffer CopyBuffer;
-		TestTrue(TEXT("FCompositeBuffer(Composite).ViewOrCopyRange(0, 4)"), Buffer.ViewOrCopyRange(0, 4, CopyBuffer) == FlatView1.GetView() && CopyBuffer.IsNull());
-		TestTrue(TEXT("FCompositeBuffer(Composite).ViewOrCopyRange(4, 4)"), Buffer.ViewOrCopyRange(4, 4, CopyBuffer) == FlatView2.GetView() && CopyBuffer.IsNull());
-		TestTrue(TEXT("FCompositeBuffer(Composite).ViewOrCopyRange(3, 2)"), Buffer.ViewOrCopyRange(3, 2, CopyBuffer).EqualBytes(MakeMemoryView(FlatArray).Mid(3, 2)) && CopyBuffer.GetSize() == 2);
-		TestTrue(TEXT("FCompositeBuffer(Composite).ViewOrCopyRange(1, 6)"), Buffer.ViewOrCopyRange(1, 6, CopyBuffer).EqualBytes(MakeMemoryView(FlatArray).Mid(1, 6)) && CopyBuffer.GetSize() == 6);
-		TestTrue(TEXT("FCompositeBuffer(Composite).ViewOrCopyRange(2, 4)"), Buffer.ViewOrCopyRange(2, 4, CopyBuffer).EqualBytes(MakeMemoryView(FlatArray).Mid(2, 4)) && CopyBuffer.GetSize() == 6);
+		TEST_TRUE(TEXT("FCompositeBuffer(Composite).ViewOrCopyRange(0, 4)"), Buffer.ViewOrCopyRange(0, 4, CopyBuffer) == FlatView1.GetView() && CopyBuffer.IsNull());
+		TEST_TRUE(TEXT("FCompositeBuffer(Composite).ViewOrCopyRange(4, 4)"), Buffer.ViewOrCopyRange(4, 4, CopyBuffer) == FlatView2.GetView() && CopyBuffer.IsNull());
+		TEST_TRUE(TEXT("FCompositeBuffer(Composite).ViewOrCopyRange(3, 2)"), Buffer.ViewOrCopyRange(3, 2, CopyBuffer).EqualBytes(MakeMemoryView(FlatArray).Mid(3, 2)) && CopyBuffer.GetSize() == 2);
+		TEST_TRUE(TEXT("FCompositeBuffer(Composite).ViewOrCopyRange(1, 6)"), Buffer.ViewOrCopyRange(1, 6, CopyBuffer).EqualBytes(MakeMemoryView(FlatArray).Mid(1, 6)) && CopyBuffer.GetSize() == 6);
+		TEST_TRUE(TEXT("FCompositeBuffer(Composite).ViewOrCopyRange(2, 4)"), Buffer.ViewOrCopyRange(2, 4, CopyBuffer).EqualBytes(MakeMemoryView(FlatArray).Mid(2, 4)) && CopyBuffer.GetSize() == 6);
 
 		uint8 CopyArray[4];
 		Buffer.CopyTo(MakeMemoryView(CopyArray), 2);
-		TestTrue(TEXT("FCompositeBuffer(Composite).CopyTo()"), MakeMemoryView(CopyArray).EqualBytes(MakeMemoryView(FlatArray).Mid(2, 4)));
+		TEST_TRUE(TEXT("FCompositeBuffer(Composite).CopyTo()"), MakeMemoryView(CopyArray).EqualBytes(MakeMemoryView(FlatArray).Mid(2, 4)));
 
 		uint32 VisitCount = 0;
 		Buffer.IterateRange(0, sizeof(FlatArray), [&VisitCount](FMemoryView) { ++VisitCount; });
-		TestEqual(TEXT("FCompositeBuffer(Composite).IterateRange(0, N)"), VisitCount, uint32(2));
+		TEST_EQUAL(TEXT("FCompositeBuffer(Composite).IterateRange(0, N)"), VisitCount, uint32(2));
 
-		const auto TestIterateRange = [&Buffer](uint64 Offset, uint64 Size, FMemoryView ExpectedView, const FSharedBuffer& ExpectedViewOuter)
+		const auto TestIterateRange = [this, &Buffer](uint64 Offset, uint64 Size, FMemoryView ExpectedView, const FSharedBuffer& ExpectedViewOuter)
 		{
 			uint32 VisitCount = 0;
 			FMemoryView ActualView;
@@ -172,9 +172,9 @@ TEST_CASE("Core::Memory::FCompositeBuffer::CompositeBuffer", "[Core][Math][Smoke
 				ActualView = View;
 				ActualViewOuter = ViewOuter;
 			});
-			TestEqual(FString::Printf(TEXT("FCompositeBuffer(Composite).IterateRange(%" UINT64_FMT ", %" UINT64_FMT ")->VisitCount"), Offset, Size), VisitCount, uint32(1));
-			TestEqual(FString::Printf(TEXT("FCompositeBuffer(Composite).IterateRange(%" UINT64_FMT ", %" UINT64_FMT ")->View"), Offset, Size), ActualView, ExpectedView);
-			TestEqual(FString::Printf(TEXT("FCompositeBuffer(Composite).IterateRange(%" UINT64_FMT ", %" UINT64_FMT ")->ViewOuter"), Offset, Size), ActualViewOuter, ExpectedViewOuter);
+			TEST_EQUAL(FString::Printf(TEXT("FCompositeBuffer(Composite).IterateRange(%" UINT64_FMT ", %" UINT64_FMT ")->VisitCount"), Offset, Size), VisitCount, uint32(1));
+			TEST_EQUAL(FString::Printf(TEXT("FCompositeBuffer(Composite).IterateRange(%" UINT64_FMT ", %" UINT64_FMT ")->View"), Offset, Size), ActualView, ExpectedView);
+			TEST_EQUAL(FString::Printf(TEXT("FCompositeBuffer(Composite).IterateRange(%" UINT64_FMT ", %" UINT64_FMT ")->ViewOuter"), Offset, Size), ActualViewOuter, ExpectedViewOuter);
 		};
 		TestIterateRange(0, 4, MakeMemoryView(FlatArray).Mid(0, 4), FlatView1);
 		TestIterateRange(4, 0, MakeMemoryView(FlatArray).Mid(4, 0), FlatView1);
