Add new Unit Tests

diff --git a/Engine/Source/Runtime/Core/Tests/Memory/MemoryViewTest.cpp b/Engine/Source/Runtime/Core/Tests/Memory/MemoryViewTest.cpp
--- a/Engine/Source/Runtime/Core/Tests/Memory/MemoryViewTest.cpp
+++ b/Engine/Source/Runtime/Core/Tests/Memory/MemoryViewTest.cpp
@@ -25,22 +25,22 @@ static_assert(!std::is_constructible<FMutableMemoryView, const void*, uint64>::v
 static_assert(std::is_assignable<FMemoryView, const FMutableMemoryView&>::value, "Missing assignment");
 static_assert(!std::is_assignable<FMutableMemoryView, const FMemoryView&>::value, "Invalid assignment");
 
-TEST_CASE("Core::Memory::FMemoryView::MemoryView", "[Core][Math][Smoke]")
+TEST_CASE_METHOD(FAutomationTestFixture, "Core::Memory::FMemoryView::MemoryView", "[Core][Memory][Smoke]")
 {
-	auto TestMemoryView = [](const FMemoryView& View, const void* Data, uint64 Size)
+	auto TestMemoryView = [this](const FMemoryView& View, const void* Data, uint64 Size)
 	{
-		TestEqual(TEXT("MemoryView.GetData()"), View.GetData(), Data);
-		TestEqual(TEXT("MemoryView.GetDataEnd()"), View.GetDataEnd(), static_cast<const void*>(static_cast<const uint8*>(Data) + Size));
-		TestEqual(TEXT("MemoryView.GetSize()"), View.GetSize(), Size);
-		TestEqual(TEXT("MemoryView.IsEmpty()"), View.IsEmpty(), Size == 0);
+		TEST_EQUAL(TEXT("MemoryView.GetData()"), View.GetData(), Data);
+		TEST_EQUAL(TEXT("MemoryView.GetDataEnd()"), View.GetDataEnd(), static_cast<const void*>(static_cast<const uint8*>(Data) + Size));
+		TEST_EQUAL(TEXT("MemoryView.GetSize()"), View.GetSize(), Size);
+		TEST_EQUAL(TEXT("MemoryView.IsEmpty()"), View.IsEmpty(), Size == 0);
 	};
 
-	auto TestMutableMemoryView = [](const FMutableMemoryView& View, void* Data, uint64 Size)
+	auto TestMutableMemoryView = [this](const FMutableMemoryView& View, void* Data, uint64 Size)
 	{
-		TestEqual(TEXT("MemoryView.GetData()"), View.GetData(), Data);
-		TestEqual(TEXT("MemoryView.GetDataEnd()"), View.GetDataEnd(), static_cast<void*>(static_cast<uint8*>(Data) + Size));
-		TestEqual(TEXT("MemoryView.GetSize()"), View.GetSize(), Size);
-		TestEqual(TEXT("MemoryView.IsEmpty()"), View.IsEmpty(), Size == 0);
+		TEST_EQUAL(TEXT("MemoryView.GetData()"), View.GetData(), Data);
+		TEST_EQUAL(TEXT("MemoryView.GetDataEnd()"), View.GetDataEnd(), static_cast<void*>(static_cast<uint8*>(Data) + Size));
+		TEST_EQUAL(TEXT("MemoryView.GetSize()"), View.GetSize(), Size);
+		TEST_EQUAL(TEXT("MemoryView.IsEmpty()"), View.IsEmpty(), Size == 0);
 	};
 
 	struct
@@ -77,7 +77,7 @@ TEST_CASE("Core::Memory::FMemoryView::MemoryView", "[Core][Math][Smoke]")
 	{
 		FMutableMemoryView View = MakeMemoryView(IntArray);
 		View.Reset();
-		TestEqual(TEXT("MemoryView.Reset()"), View, FMutableMemoryView());
+		TEST_EQUAL(TEXT("MemoryView.Reset()"), View, FMutableMemoryView());
 	}
 
 	// Test Left
@@ -95,108 +95,108 @@ TEST_CASE("Core::Memory::FMemoryView::MemoryView", "[Core][Math][Smoke]")
 	static_assert(MakeMemoryView(IntArray).LeftChop(MAX_uint64).IsEmpty(), "Error in LeftChop");
 
 	// Test Right
-	TestEqual(TEXT("MemoryView.Right(0)"), MakeMemoryView(IntArray).Right(0), FMutableMemoryView());
-	TestEqual(TEXT("MemoryView.Right(1)"), MakeMemoryView(IntArray).Right(1), MakeMemoryView(reinterpret_cast<uint8*>(IntArray) + sizeof(IntArray) - 1, 1));
-	TestEqual(TEXT("MemoryView.Right(Size)"), MakeMemoryView(IntArray).Right(sizeof(IntArray)), MakeMemoryView(IntArray));
-	TestEqual(TEXT("MemoryView.Right(Size+1)"), MakeMemoryView(IntArray).Right(sizeof(IntArray) + 1), MakeMemoryView(IntArray));
-	TestEqual(TEXT("MemoryView.Right(MaxSize)"), MakeMemoryView(IntArray).Right(MAX_uint64), MakeMemoryView(IntArray));
+	TEST_EQUAL(TEXT("MemoryView.Right(0)"), MakeMemoryView(IntArray).Right(0), FMutableMemoryView());
+	TEST_EQUAL(TEXT("MemoryView.Right(1)"), MakeMemoryView(IntArray).Right(1), MakeMemoryView(reinterpret_cast<uint8*>(IntArray) + sizeof(IntArray) - 1, 1));
+	TEST_EQUAL(TEXT("MemoryView.Right(Size)"), MakeMemoryView(IntArray).Right(sizeof(IntArray)), MakeMemoryView(IntArray));
+	TEST_EQUAL(TEXT("MemoryView.Right(Size+1)"), MakeMemoryView(IntArray).Right(sizeof(IntArray) + 1), MakeMemoryView(IntArray));
+	TEST_EQUAL(TEXT("MemoryView.Right(MaxSize)"), MakeMemoryView(IntArray).Right(MAX_uint64), MakeMemoryView(IntArray));
 
 	// Test RightChop
-	TestEqual(TEXT("MemoryView.RightChop(0)"), MakeMemoryView(IntArray).RightChop(0), MakeMemoryView(IntArray));
-	TestEqual(TEXT("MemoryView.RightChop(1)"), MakeMemoryView(IntArray).RightChop(1), MakeMemoryView(reinterpret_cast<uint8*>(IntArray) + 1, sizeof(IntArray) - 1));
-	TestEqual(TEXT("MemoryView.RightChop(Size)"), MakeMemoryView(IntArray).RightChop(sizeof(IntArray)), FMutableMemoryView());
-	TestEqual(TEXT("MemoryView.RightChop(Size+1)"), MakeMemoryView(IntArray).RightChop(sizeof(IntArray) + 1), FMutableMemoryView());
-	TestEqual(TEXT("MemoryView.RightChop(MaxSize)"), MakeMemoryView(IntArray).RightChop(MAX_uint64), FMutableMemoryView());
+	TEST_EQUAL(TEXT("MemoryView.RightChop(0)"), MakeMemoryView(IntArray).RightChop(0), MakeMemoryView(IntArray));
+	TEST_EQUAL(TEXT("MemoryView.RightChop(1)"), MakeMemoryView(IntArray).RightChop(1), MakeMemoryView(reinterpret_cast<uint8*>(IntArray) + 1, sizeof(IntArray) - 1));
+	TEST_EQUAL(TEXT("MemoryView.RightChop(Size)"), MakeMemoryView(IntArray).RightChop(sizeof(IntArray)), FMutableMemoryView());
+	TEST_EQUAL(TEXT("MemoryView.RightChop(Size+1)"), MakeMemoryView(IntArray).RightChop(sizeof(IntArray) + 1), FMutableMemoryView());
+	TEST_EQUAL(TEXT("MemoryView.RightChop(MaxSize)"), MakeMemoryView(IntArray).RightChop(MAX_uint64), FMutableMemoryView());
 
 	// Test Mid
-	TestEqual(TEXT("MemoryView.Mid(0)"), MakeMemoryView(IntArray).Mid(0), MakeMemoryView(IntArray));
-	TestEqual(TEXT("MemoryView.Mid(1)"), MakeMemoryView(IntArray).Mid(1), MakeMemoryView(reinterpret_cast<uint8*>(IntArray) + 1, sizeof(IntArray) - 1));
-	TestEqual(TEXT("MemoryView.Mid(Size)"), MakeMemoryView(IntArray).Mid(sizeof(IntArray)), FMutableMemoryView());
-	TestEqual(TEXT("MemoryView.Mid(Size+1)"), MakeMemoryView(IntArray).Mid(sizeof(IntArray) + 1), FMutableMemoryView());
-	TestEqual(TEXT("MemoryView.Mid(MaxSize)"), MakeMemoryView(IntArray).Mid(MAX_uint64), FMutableMemoryView());
-	TestEqual(TEXT("MemoryView.Mid(0,0)"), MakeMemoryView(IntArray).Mid(0, 0), FMutableMemoryView());
-	TestEqual(TEXT("MemoryView.Mid(0,1)"), MakeMemoryView(IntArray).Mid(0, 1), MakeMemoryView(IntArray, 1));
-	TestEqual(TEXT("MemoryView.Mid(1,Size-2)"), MakeMemoryView(IntArray).Mid(1, sizeof(IntArray) - 2), MakeMemoryView(reinterpret_cast<uint8*>(IntArray) + 1, sizeof(IntArray) - 2));
-	TestEqual(TEXT("MemoryView.Mid(1,Size-1)"), MakeMemoryView(IntArray).Mid(1, sizeof(IntArray) - 1), MakeMemoryView(reinterpret_cast<uint8*>(IntArray) + 1, sizeof(IntArray) - 1));
-	TestEqual(TEXT("MemoryView.Mid(1,Size)"), MakeMemoryView(IntArray).Mid(1, sizeof(IntArray)), MakeMemoryView(reinterpret_cast<uint8*>(IntArray) + 1, sizeof(IntArray) - 1));
-	TestEqual(TEXT("MemoryView.Mid(0,MaxSize)"), MakeMemoryView(IntArray).Mid(0, MAX_uint64), MakeMemoryView(IntArray));
-	TestEqual(TEXT("MemoryView.Mid(MaxSize,MaxSize)"), MakeMemoryView(IntArray).Mid(MAX_uint64, MAX_uint64), FMutableMemoryView());
+	TEST_EQUAL(TEXT("MemoryView.Mid(0)"), MakeMemoryView(IntArray).Mid(0), MakeMemoryView(IntArray));
+	TEST_EQUAL(TEXT("MemoryView.Mid(1)"), MakeMemoryView(IntArray).Mid(1), MakeMemoryView(reinterpret_cast<uint8*>(IntArray) + 1, sizeof(IntArray) - 1));
+	TEST_EQUAL(TEXT("MemoryView.Mid(Size)"), MakeMemoryView(IntArray).Mid(sizeof(IntArray)), FMutableMemoryView());
+	TEST_EQUAL(TEXT("MemoryView.Mid(Size+1)"), MakeMemoryView(IntArray).Mid(sizeof(IntArray) + 1), FMutableMemoryView());
+	TEST_EQUAL(TEXT("MemoryView.Mid(MaxSize)"), MakeMemoryView(IntArray).Mid(MAX_uint64), FMutableMemoryView());
+	TEST_EQUAL(TEXT("MemoryView.Mid(0,0)"), MakeMemoryView(IntArray).Mid(0, 0), FMutableMemoryView());
+	TEST_EQUAL(TEXT("MemoryView.Mid(0,1)"), MakeMemoryView(IntArray).Mid(0, 1), MakeMemoryView(IntArray, 1));
+	TEST_EQUAL(TEXT("MemoryView.Mid(1,Size-2)"), MakeMemoryView(IntArray).Mid(1, sizeof(IntArray) - 2), MakeMemoryView(reinterpret_cast<uint8*>(IntArray) + 1, sizeof(IntArray) - 2));
+	TEST_EQUAL(TEXT("MemoryView.Mid(1,Size-1)"), MakeMemoryView(IntArray).Mid(1, sizeof(IntArray) - 1), MakeMemoryView(reinterpret_cast<uint8*>(IntArray) + 1, sizeof(IntArray) - 1));
+	TEST_EQUAL(TEXT("MemoryView.Mid(1,Size)"), MakeMemoryView(IntArray).Mid(1, sizeof(IntArray)), MakeMemoryView(reinterpret_cast<uint8*>(IntArray) + 1, sizeof(IntArray) - 1));
+	TEST_EQUAL(TEXT("MemoryView.Mid(0,MaxSize)"), MakeMemoryView(IntArray).Mid(0, MAX_uint64), MakeMemoryView(IntArray));
+	TEST_EQUAL(TEXT("MemoryView.Mid(MaxSize,MaxSize)"), MakeMemoryView(IntArray).Mid(MAX_uint64, MAX_uint64), FMutableMemoryView());
 
 	// Test Contains
-	TestTrue(TEXT("MemoryView.Contains(Empty)"), FMemoryView().Contains(FMutableMemoryView()));
-	TestTrue(TEXT("MemoryView.Contains(Empty)"), FMutableMemoryView().Contains(FMemoryView()));
-	TestTrue(TEXT("MemoryView.Contains(Equal)"), MakeMemoryView(ByteArray).Contains(MakeMemoryView(ByteArray)));
-	TestTrue(TEXT("MemoryView.Contains(SmallerBy1Left)"), MakeMemoryView(ByteArray).Contains(MakeMemoryView(ByteArray + 1, 15)));
-	TestTrue(TEXT("MemoryView.Contains(SmallerBy1Right)"), MakeMemoryView(ByteArray).Contains(MakeMemoryView(ByteArray, 15)));
-	TestTrue(TEXT("MemoryView.Contains(SmallerBy2Both)"), MakeMemoryView(ByteArray).Contains(MakeMemoryView(ByteArray + 1, 14)));
-	TestTrue(TEXT("MemoryView.Contains(EmptyContained)"), MakeMemoryView(ByteArray).Contains(MakeMemoryView(ByteArray, 0)));
-	TestTrue(TEXT("MemoryView.Contains(EmptyContained)"), MakeMemoryView(ByteArray).Contains(MakeMemoryView(ByteArray + 8, 0)));
-	TestTrue(TEXT("MemoryView.Contains(EmptyContained)"), MakeMemoryView(ByteArray).Contains(MakeMemoryView(ByteArray + 16, 0)));
-	TestFalse(TEXT("MemoryView.Contains(EmptyOutside)"), MakeMemoryView(ByteArray).Contains(MakeMemoryView(ByteArrayContainer.BeforeByteArray + 3, 0)));
-	TestFalse(TEXT("MemoryView.Contains(EmptyOutside)"), MakeMemoryView(ByteArray).Contains(MakeMemoryView(ByteArrayContainer.AfterByteArray + 1, 0)));
-	TestFalse(TEXT("MemoryView.Contains(OutsideBy1Left)"), MakeMemoryView(ByteArray).Contains(MakeMemoryView(ByteArrayContainer.BeforeByteArray + 3, 1)));
-	TestFalse(TEXT("MemoryView.Contains(OutsideBy1Right)"), MakeMemoryView(ByteArray).Contains(MakeMemoryView(ByteArray + 16, 1)));
-	TestFalse(TEXT("MemoryView.Contains(LargerBy1Left)"), MakeMemoryView(ByteArray).Contains(MakeMemoryView(ByteArrayContainer.BeforeByteArray + 3, 17)));
-	TestFalse(TEXT("MemoryView.Contains(LargerBy1Right)"), MakeMemoryView(ByteArray).Contains(MakeMemoryView(ByteArray, 17)));
-	TestFalse(TEXT("MemoryView.Contains(LargerBy2Both)"), MakeMemoryView(ByteArray).Contains(MakeMemoryView(ByteArrayContainer.BeforeByteArray + 3, 18)));
-	TestFalse(TEXT("MemoryView.Contains(SmallerOverlapLeft)"), MakeMemoryView(ByteArray).Contains(MakeMemoryView(ByteArrayContainer.BeforeByteArray + 3, 8)));
-	TestFalse(TEXT("MemoryView.Contains(SmallerOverlapRight)"), MakeMemoryView(ByteArray).Contains(MakeMemoryView(ByteArray + 9, 8)));
+	TEST_TRUE(TEXT("MemoryView.Contains(Empty)"), FMemoryView().Contains(FMutableMemoryView()));
+	TEST_TRUE(TEXT("MemoryView.Contains(Empty)"), FMutableMemoryView().Contains(FMemoryView()));
+	TEST_TRUE(TEXT("MemoryView.Contains(Equal)"), MakeMemoryView(ByteArray).Contains(MakeMemoryView(ByteArray)));
+	TEST_TRUE(TEXT("MemoryView.Contains(SmallerBy1Left)"), MakeMemoryView(ByteArray).Contains(MakeMemoryView(ByteArray + 1, 15)));
+	TEST_TRUE(TEXT("MemoryView.Contains(SmallerBy1Right)"), MakeMemoryView(ByteArray).Contains(MakeMemoryView(ByteArray, 15)));
+	TEST_TRUE(TEXT("MemoryView.Contains(SmallerBy2Both)"), MakeMemoryView(ByteArray).Contains(MakeMemoryView(ByteArray + 1, 14)));
+	TEST_TRUE(TEXT("MemoryView.Contains(EmptyContained)"), MakeMemoryView(ByteArray).Contains(MakeMemoryView(ByteArray, 0)));
+	TEST_TRUE(TEXT("MemoryView.Contains(EmptyContained)"), MakeMemoryView(ByteArray).Contains(MakeMemoryView(ByteArray + 8, 0)));
+	TEST_TRUE(TEXT("MemoryView.Contains(EmptyContained)"), MakeMemoryView(ByteArray).Contains(MakeMemoryView(ByteArray + 16, 0)));
+	TEST_FALSE(TEXT("MemoryView.Contains(EmptyOutside)"), MakeMemoryView(ByteArray).Contains(MakeMemoryView(ByteArrayContainer.BeforeByteArray + 3, 0)));
+	TEST_FALSE(TEXT("MemoryView.Contains(EmptyOutside)"), MakeMemoryView(ByteArray).Contains(MakeMemoryView(ByteArrayContainer.AfterByteArray + 1, 0)));
+	TEST_FALSE(TEXT("MemoryView.Contains(OutsideBy1Left)"), MakeMemoryView(ByteArray).Contains(MakeMemoryView(ByteArrayContainer.BeforeByteArray + 3, 1)));
+	TEST_FALSE(TEXT("MemoryView.Contains(OutsideBy1Right)"), MakeMemoryView(ByteArray).Contains(MakeMemoryView(ByteArray + 16, 1)));
+	TEST_FALSE(TEXT("MemoryView.Contains(LargerBy1Left)"), MakeMemoryView(ByteArray).Contains(MakeMemoryView(ByteArrayContainer.BeforeByteArray + 3, 17)));
+	TEST_FALSE(TEXT("MemoryView.Contains(LargerBy1Right)"), MakeMemoryView(ByteArray).Contains(MakeMemoryView(ByteArray, 17)));
+	TEST_FALSE(TEXT("MemoryView.Contains(LargerBy2Both)"), MakeMemoryView(ByteArray).Contains(MakeMemoryView(ByteArrayContainer.BeforeByteArray + 3, 18)));
+	TEST_FALSE(TEXT("MemoryView.Contains(SmallerOverlapLeft)"), MakeMemoryView(ByteArray).Contains(MakeMemoryView(ByteArrayContainer.BeforeByteArray + 3, 8)));
+	TEST_FALSE(TEXT("MemoryView.Contains(SmallerOverlapRight)"), MakeMemoryView(ByteArray).Contains(MakeMemoryView(ByteArray + 9, 8)));
 
 	// Test Intersects
-	TestTrue(TEXT("MemoryView.Intersects(Equal)"), MakeMemoryView(ByteArray).Intersects(MakeMemoryView(ByteArray)));
-	TestTrue(TEXT("MemoryView.Intersects(SmallerBy1Left)"), MakeMemoryView(ByteArray).Intersects(MakeMemoryView(ByteArray + 1, 15)));
-	TestTrue(TEXT("MemoryView.Intersects(SmallerBy1Right)"), MakeMemoryView(ByteArray).Intersects(MakeMemoryView(ByteArray, 15)));
-	TestTrue(TEXT("MemoryView.Intersects(SmallerBy2Both)"), MakeMemoryView(ByteArray).Intersects(MakeMemoryView(ByteArray + 1, 14)));
-	TestTrue(TEXT("MemoryView.Intersects(SmallerOverlapLeft)"), MakeMemoryView(ByteArray).Intersects(MakeMemoryView(ByteArrayContainer.BeforeByteArray + 3, 8)));
-	TestTrue(TEXT("MemoryView.Intersects(SmallerOverlapRight)"), MakeMemoryView(ByteArray).Intersects(MakeMemoryView(ByteArray + 9, 8)));
-	TestTrue(TEXT("MemoryView.Intersects(LargerBy1Left)"), MakeMemoryView(ByteArray).Intersects(MakeMemoryView(ByteArrayContainer.BeforeByteArray + 3, 17)));
-	TestTrue(TEXT("MemoryView.Intersects(LargerBy1Right)"), MakeMemoryView(ByteArray).Intersects(MakeMemoryView(ByteArray, 17)));
-	TestTrue(TEXT("MemoryView.Intersects(LargerBy2Both)"), MakeMemoryView(ByteArray).Intersects(MakeMemoryView(ByteArrayContainer.BeforeByteArray + 3, 18)));
-	TestTrue(TEXT("MemoryView.Intersects(EmptyMiddle)"), MakeMemoryView(ByteArray).Intersects(MakeMemoryView(ByteArray + 8, 0)));
-	TestFalse(TEXT("MemoryView.Intersects(Empty)"), FMemoryView().Intersects(FMutableMemoryView()));
-	TestFalse(TEXT("MemoryView.Intersects(Empty)"), FMutableMemoryView().Intersects(FMemoryView()));
-	TestFalse(TEXT("MemoryView.Intersects(EmptyLeft)"), MakeMemoryView(ByteArray).Intersects(MakeMemoryView(ByteArray, 0)));
-	TestFalse(TEXT("MemoryView.Intersects(EmptyRight)"), MakeMemoryView(ByteArray).Intersects(MakeMemoryView(ByteArray + 16, 0)));
-	TestFalse(TEXT("MemoryView.Intersects(EmptyOutside)"), MakeMemoryView(ByteArray).Intersects(MakeMemoryView(ByteArrayContainer.BeforeByteArray + 3, 0)));
-	TestFalse(TEXT("MemoryView.Intersects(EmptyOutside)"), MakeMemoryView(ByteArray).Intersects(MakeMemoryView(ByteArrayContainer.AfterByteArray + 1, 0)));
-	TestFalse(TEXT("MemoryView.Intersects(OutsideBy1Left)"), MakeMemoryView(ByteArray).Intersects(MakeMemoryView(ByteArrayContainer.BeforeByteArray + 3, 1)));
-	TestFalse(TEXT("MemoryView.Intersects(OutsideBy1Right)"), MakeMemoryView(ByteArray).Intersects(MakeMemoryView(ByteArray + 16, 1)));
+	TEST_TRUE(TEXT("MemoryView.Intersects(Equal)"), MakeMemoryView(ByteArray).Intersects(MakeMemoryView(ByteArray)));
+	TEST_TRUE(TEXT("MemoryView.Intersects(SmallerBy1Left)"), MakeMemoryView(ByteArray).Intersects(MakeMemoryView(ByteArray + 1, 15)));
+	TEST_TRUE(TEXT("MemoryView.Intersects(SmallerBy1Right)"), MakeMemoryView(ByteArray).Intersects(MakeMemoryView(ByteArray, 15)));
+	TEST_TRUE(TEXT("MemoryView.Intersects(SmallerBy2Both)"), MakeMemoryView(ByteArray).Intersects(MakeMemoryView(ByteArray + 1, 14)));
+	TEST_TRUE(TEXT("MemoryView.Intersects(SmallerOverlapLeft)"), MakeMemoryView(ByteArray).Intersects(MakeMemoryView(ByteArrayContainer.BeforeByteArray + 3, 8)));
+	TEST_TRUE(TEXT("MemoryView.Intersects(SmallerOverlapRight)"), MakeMemoryView(ByteArray).Intersects(MakeMemoryView(ByteArray + 9, 8)));
+	TEST_TRUE(TEXT("MemoryView.Intersects(LargerBy1Left)"), MakeMemoryView(ByteArray).Intersects(MakeMemoryView(ByteArrayContainer.BeforeByteArray + 3, 17)));
+	TEST_TRUE(TEXT("MemoryView.Intersects(LargerBy1Right)"), MakeMemoryView(ByteArray).Intersects(MakeMemoryView(ByteArray, 17)));
+	TEST_TRUE(TEXT("MemoryView.Intersects(LargerBy2Both)"), MakeMemoryView(ByteArray).Intersects(MakeMemoryView(ByteArrayContainer.BeforeByteArray + 3, 18)));
+	TEST_TRUE(TEXT("MemoryView.Intersects(EmptyMiddle)"), MakeMemoryView(ByteArray).Intersects(MakeMemoryView(ByteArray + 8, 0)));
+	TEST_FALSE(TEXT("MemoryView.Intersects(Empty)"), FMemoryView().Intersects(FMutableMemoryView()));
+	TEST_FALSE(TEXT("MemoryView.Intersects(Empty)"), FMutableMemoryView().Intersects(FMemoryView()));
+	TEST_FALSE(TEXT("MemoryView.Intersects(EmptyLeft)"), MakeMemoryView(ByteArray).Intersects(MakeMemoryView(ByteArray, 0)));
+	TEST_FALSE(TEXT("MemoryView.Intersects(EmptyRight)"), MakeMemoryView(ByteArray).Intersects(MakeMemoryView(ByteArray + 16, 0)));
+	TEST_FALSE(TEXT("MemoryView.Intersects(EmptyOutside)"), MakeMemoryView(ByteArray).Intersects(MakeMemoryView(ByteArrayContainer.BeforeByteArray + 3, 0)));
+	TEST_FALSE(TEXT("MemoryView.Intersects(EmptyOutside)"), MakeMemoryView(ByteArray).Intersects(MakeMemoryView(ByteArrayContainer.AfterByteArray + 1, 0)));
+	TEST_FALSE(TEXT("MemoryView.Intersects(OutsideBy1Left)"), MakeMemoryView(ByteArray).Intersects(MakeMemoryView(ByteArrayContainer.BeforeByteArray + 3, 1)));
+	TEST_FALSE(TEXT("MemoryView.Intersects(OutsideBy1Right)"), MakeMemoryView(ByteArray).Intersects(MakeMemoryView(ByteArray + 16, 1)));
 
 	// Test CompareBytes
 	const uint8 CompareBytes[8] = { 5, 4, 6, 2, 4, 7, 1, 3 };
-	TestEqual(TEXT("MemoryView.CompareBytes(Empty)"), FMemoryView().CompareBytes(FMutableMemoryView()), 0);
-	TestEqual(TEXT("MemoryView.CompareBytes(Empty)"), FMutableMemoryView().CompareBytes(FMemoryView()), 0);
-	TestEqual(TEXT("MemoryView.CompareBytes(EqualView)"), MakeMemoryView(ByteArray).CompareBytes(MakeMemoryView(ByteArray)), 0);
-	TestEqual(TEXT("MemoryView.CompareBytes(EqualBytes)"), MakeMemoryView(ByteArray, 8).CompareBytes(MakeMemoryView(ByteArray + 8, 8)), 0);
-	TestTrue(TEXT("MemoryView.CompareBytes(EmptyLeft)"), FMemoryView().CompareBytes(MakeMemoryView(ByteArray)) < 0);
-	TestTrue(TEXT("MemoryView.CompareBytes(EmptyRight)"), MakeMemoryView(ByteArray).CompareBytes(FMemoryView()) > 0);
-	TestTrue(TEXT("MemoryView.CompareBytes(ShorterLeft)"), MakeMemoryView(ByteArray, 8).CompareBytes(MakeMemoryView(ByteArray)) < 0);
-	TestTrue(TEXT("MemoryView.CompareBytes(ShorterRight)"), MakeMemoryView(ByteArray).CompareBytes(MakeMemoryView(ByteArray, 8)) > 0);
-	TestTrue(TEXT("MemoryView.CompareBytes(ShorterLeft)"), MakeMemoryView(IntArray, 8).CompareBytes(MakeMemoryView(ByteArray)) < 0);
-	TestTrue(TEXT("MemoryView.CompareBytes(ShorterRight)"), MakeMemoryView(ByteArray).CompareBytes(MakeMemoryView(IntArray, 8)) > 0);
-	TestTrue(TEXT("MemoryView.CompareBytes(DifferentSize)"), MakeMemoryView(ByteArray, 4).CompareBytes(MakeMemoryView(ByteArray, 8)) < 0);
-	TestTrue(TEXT("MemoryView.CompareBytes(DifferentSize)"), MakeMemoryView(ByteArray, 8).CompareBytes(MakeMemoryView(ByteArray, 4)) > 0);
-	TestTrue(TEXT("MemoryView.CompareBytes(SameSizeLeftLess)"), MakeMemoryView(CompareBytes, 2).CompareBytes(MakeMemoryView(CompareBytes + 2, 2)) < 0);
-	TestTrue(TEXT("MemoryView.CompareBytes(SameSizeLeftGreater)"), MakeMemoryView(CompareBytes, 3).CompareBytes(MakeMemoryView(CompareBytes + 3, 3)) > 0);
+	TEST_EQUAL(TEXT("MemoryView.CompareBytes(Empty)"), FMemoryView().CompareBytes(FMutableMemoryView()), 0);
+	TEST_EQUAL(TEXT("MemoryView.CompareBytes(Empty)"), FMutableMemoryView().CompareBytes(FMemoryView()), 0);
+	TEST_EQUAL(TEXT("MemoryView.CompareBytes(EqualView)"), MakeMemoryView(ByteArray).CompareBytes(MakeMemoryView(ByteArray)), 0);
+	TEST_EQUAL(TEXT("MemoryView.CompareBytes(EqualBytes)"), MakeMemoryView(ByteArray, 8).CompareBytes(MakeMemoryView(ByteArray + 8, 8)), 0);
+	TEST_TRUE(TEXT("MemoryView.CompareBytes(EmptyLeft)"), FMemoryView().CompareBytes(MakeMemoryView(ByteArray)) < 0);
+	TEST_TRUE(TEXT("MemoryView.CompareBytes(EmptyRight)"), MakeMemoryView(ByteArray).CompareBytes(FMemoryView()) > 0);
+	TEST_TRUE(TEXT("MemoryView.CompareBytes(ShorterLeft)"), MakeMemoryView(ByteArray, 8).CompareBytes(MakeMemoryView(ByteArray)) < 0);
+	TEST_TRUE(TEXT("MemoryView.CompareBytes(ShorterRight)"), MakeMemoryView(ByteArray).CompareBytes(MakeMemoryView(ByteArray, 8)) > 0);
+	TEST_TRUE(TEXT("MemoryView.CompareBytes(ShorterLeft)"), MakeMemoryView(IntArray, 8).CompareBytes(MakeMemoryView(ByteArray)) < 0);
+	TEST_TRUE(TEXT("MemoryView.CompareBytes(ShorterRight)"), MakeMemoryView(ByteArray).CompareBytes(MakeMemoryView(IntArray, 8)) > 0);
+	TEST_TRUE(TEXT("MemoryView.CompareBytes(DifferentSize)"), MakeMemoryView(ByteArray, 4).CompareBytes(MakeMemoryView(ByteArray, 8)) < 0);
+	TEST_TRUE(TEXT("MemoryView.CompareBytes(DifferentSize)"), MakeMemoryView(ByteArray, 8).CompareBytes(MakeMemoryView(ByteArray, 4)) > 0);
+	TEST_TRUE(TEXT("MemoryView.CompareBytes(SameSizeLeftLess)"), MakeMemoryView(CompareBytes, 2).CompareBytes(MakeMemoryView(CompareBytes + 2, 2)) < 0);
+	TEST_TRUE(TEXT("MemoryView.CompareBytes(SameSizeLeftGreater)"), MakeMemoryView(CompareBytes, 3).CompareBytes(MakeMemoryView(CompareBytes + 3, 3)) > 0);
 
 	// Test EqualBytes
-	TestTrue(TEXT("MemoryView.EqualBytes(Empty)"), FMemoryView().EqualBytes(FMutableMemoryView()));
-	TestTrue(TEXT("MemoryView.EqualBytes(Empty)"), FMutableMemoryView().EqualBytes(FMemoryView()));
-	TestTrue(TEXT("MemoryView.EqualBytes(EqualView)"), MakeMemoryView(ByteArray).EqualBytes(MakeMemoryView(ByteArray)));
-	TestTrue(TEXT("MemoryView.EqualBytes(EqualBytes)"), MakeMemoryView(ByteArray, 8).EqualBytes(MakeMemoryView(ByteArray + 8, 8)));
-	TestFalse(TEXT("MemoryView.EqualBytes(DifferentSize)"), MakeMemoryView(ByteArray, 8).EqualBytes(MakeMemoryView(ByteArray, 4)));
-	TestFalse(TEXT("MemoryView.EqualBytes(DifferentSize)"), MakeMemoryView(ByteArray, 4).EqualBytes(MakeMemoryView(ByteArray, 8)));
-	TestFalse(TEXT("MemoryView.EqualBytes(DifferentBytes)"), MakeMemoryView(CompareBytes, 4).EqualBytes(MakeMemoryView(CompareBytes + 4, 4)));
+	TEST_TRUE(TEXT("MemoryView.EqualBytes(Empty)"), FMemoryView().EqualBytes(FMutableMemoryView()));
+	TEST_TRUE(TEXT("MemoryView.EqualBytes(Empty)"), FMutableMemoryView().EqualBytes(FMemoryView()));
+	TEST_TRUE(TEXT("MemoryView.EqualBytes(EqualView)"), MakeMemoryView(ByteArray).EqualBytes(MakeMemoryView(ByteArray)));
+	TEST_TRUE(TEXT("MemoryView.EqualBytes(EqualBytes)"), MakeMemoryView(ByteArray, 8).EqualBytes(MakeMemoryView(ByteArray + 8, 8)));
+	TEST_FALSE(TEXT("MemoryView.EqualBytes(DifferentSize)"), MakeMemoryView(ByteArray, 8).EqualBytes(MakeMemoryView(ByteArray, 4)));
+	TEST_FALSE(TEXT("MemoryView.EqualBytes(DifferentSize)"), MakeMemoryView(ByteArray, 4).EqualBytes(MakeMemoryView(ByteArray, 8)));
+	TEST_FALSE(TEXT("MemoryView.EqualBytes(DifferentBytes)"), MakeMemoryView(CompareBytes, 4).EqualBytes(MakeMemoryView(CompareBytes + 4, 4)));
 
 	// Test Equals
-	TestTrue(TEXT("MemoryView.Equals(Empty)"), FMemoryView().Equals(FMemoryView()));
-	TestTrue(TEXT("MemoryView.Equals(Empty)"), FMemoryView().Equals(FMutableMemoryView()));
-	TestTrue(TEXT("MemoryView.Equals(Empty)"), FMutableMemoryView().Equals(FMemoryView()));
-	TestTrue(TEXT("MemoryView.Equals(Empty)"), FMutableMemoryView().Equals(FMutableMemoryView()));
-	TestTrue(TEXT("MemoryView.Equals(Equal)"), MakeMemoryView(IntArray).Equals(MakeMemoryView(AsConst(IntArray))));
-	TestFalse(TEXT("MemoryView.Equals(DataDiff)"), MakeMemoryView(IntArray).Equals(MakeMemoryView(IntArray + 1, sizeof(IntArray) - sizeof(*IntArray))));
-	TestFalse(TEXT("MemoryView.Equals(SizeDiff)"), MakeMemoryView(IntArray).Equals(MakeMemoryView(IntArray, sizeof(*IntArray))));
-	TestFalse(TEXT("MemoryView.Equals(BothDiff)"), MakeMemoryView(IntArray).Equals(FMutableMemoryView()));
+	TEST_TRUE(TEXT("MemoryView.Equals(Empty)"), FMemoryView().Equals(FMemoryView()));
+	TEST_TRUE(TEXT("MemoryView.Equals(Empty)"), FMemoryView().Equals(FMutableMemoryView()));
+	TEST_TRUE(TEXT("MemoryView.Equals(Empty)"), FMutableMemoryView().Equals(FMemoryView()));
+	TEST_TRUE(TEXT("MemoryView.Equals(Empty)"), FMutableMemoryView().Equals(FMutableMemoryView()));
+	TEST_TRUE(TEXT("MemoryView.Equals(Equal)"), MakeMemoryView(IntArray).Equals(MakeMemoryView(AsConst(IntArray))));
+	TEST_FALSE(TEXT("MemoryView.Equals(DataDiff)"), MakeMemoryView(IntArray).Equals(MakeMemoryView(IntArray + 1, sizeof(IntArray) - sizeof(*IntArray))));
+	TEST_FALSE(TEXT("MemoryView.Equals(SizeDiff)"), MakeMemoryView(IntArray).Equals(MakeMemoryView(IntArray, sizeof(*IntArray))));
+	TEST_FALSE(TEXT("MemoryView.Equals(BothDiff)"), MakeMemoryView(IntArray).Equals(FMutableMemoryView()));
 
 	// Test operator==
 	static_assert(MakeMemoryView(ByteArrayContainer.ByteArray) == MakeMemoryView(ByteArrayContainer.ByteArray), "Error in MemoryView == MemoryView"); //-V501
@@ -211,18 +211,18 @@ TEST_CASE("Core::Memory::FMemoryView::MemoryView", "[Core][Math][Smoke]")
 	static_assert(MakeMemoryView(AsConst(ByteArrayContainer.ByteArray)) != MakeMemoryView(AsConst(IntArray)), "Error in MemoryView != MemoryView");
 
 	// Test operator+=
-	TestEqual(TEXT("MemoryView += 0"), MakeMemoryView(ByteArray) += 0, MakeMemoryView(ByteArray));
-	TestEqual(TEXT("MemoryView += Half"), MakeMemoryView(ByteArray) += 8, MakeMemoryView(ByteArray + 8, 8));
-	TestEqual(TEXT("MemoryView += Size"), MakeMemoryView(ByteArray) += 16, MakeMemoryView(ByteArray + 16, 0));
-	TestEqual(TEXT("MemoryView += OutOfBounds"), MakeMemoryView(ByteArray) += 32, MakeMemoryView(ByteArray + 16, 0));
+	TEST_EQUAL(TEXT("MemoryView += 0"), MakeMemoryView(ByteArray) += 0, MakeMemoryView(ByteArray));
+	TEST_EQUAL(TEXT("MemoryView += Half"), MakeMemoryView(ByteArray) += 8, MakeMemoryView(ByteArray + 8, 8));
+	TEST_EQUAL(TEXT("MemoryView += Size"), MakeMemoryView(ByteArray) += 16, MakeMemoryView(ByteArray + 16, 0));
+	TEST_EQUAL(TEXT("MemoryView += OutOfBounds"), MakeMemoryView(ByteArray) += 32, MakeMemoryView(ByteArray + 16, 0));
 
 	// Test operator+
-	TestEqual(TEXT("MemoryView + 0"), MakeMemoryView(ByteArray) + 0, MakeMemoryView(ByteArray));
-	TestEqual(TEXT("0 + MemoryView"), 0 + MakeMemoryView(ByteArray), MakeMemoryView(ByteArray));
-	TestEqual(TEXT("MemoryView + Half"), MakeMemoryView(ByteArray) + 8, MakeMemoryView(ByteArray + 8, 8));
-	TestEqual(TEXT("Half + MemoryView"), 8 + MakeMemoryView(ByteArray), MakeMemoryView(ByteArray + 8, 8));
-	TestEqual(TEXT("MemoryView + Size"), MakeMemoryView(ByteArray) + 16, MakeMemoryView(ByteArray + 16, 0));
-	TestEqual(TEXT("Size + MemoryView"), 16 + MakeMemoryView(ByteArray), MakeMemoryView(ByteArray + 16, 0));
-	TestEqual(TEXT("MemoryView + OutOfBounds"), MakeMemoryView(ByteArray) + 32, MakeMemoryView(ByteArray + 16, 0));
-	TestEqual(TEXT("OutOfBounds + MemoryView"), 32 + MakeMemoryView(ByteArray), MakeMemoryView(ByteArray + 16, 0));
+	TEST_EQUAL(TEXT("MemoryView + 0"), MakeMemoryView(ByteArray) + 0, MakeMemoryView(ByteArray));
+	TEST_EQUAL(TEXT("0 + MemoryView"), 0 + MakeMemoryView(ByteArray), MakeMemoryView(ByteArray));
+	TEST_EQUAL(TEXT("MemoryView + Half"), MakeMemoryView(ByteArray) + 8, MakeMemoryView(ByteArray + 8, 8));
+	TEST_EQUAL(TEXT("Half + MemoryView"), 8 + MakeMemoryView(ByteArray), MakeMemoryView(ByteArray + 8, 8));
+	TEST_EQUAL(TEXT("MemoryView + Size"), MakeMemoryView(ByteArray) + 16, MakeMemoryView(ByteArray + 16, 0));
+	TEST_EQUAL(TEXT("Size + MemoryView"), 16 + MakeMemoryView(ByteArray), MakeMemoryView(ByteArray + 16, 0));
+	TEST_EQUAL(TEXT("MemoryView + OutOfBounds"), MakeMemoryView(ByteArray) + 32, MakeMemoryView(ByteArray + 16, 0));
+	TEST_EQUAL(TEXT("OutOfBounds + MemoryView"), 32 + MakeMemoryView(ByteArray), MakeMemoryView(ByteArray + 16, 0));
 }
diff --git a/Engine/Source/Runtime/Core/Tests/Memory/SharedBufferTest.cpp b/Engine/Source/Runtime/Core/Tests/Memory/SharedBufferTest.cpp
--- a/Engine/Source/Runtime/Core/Tests/Memory/SharedBufferTest.cpp
+++ b/Engine/Source/Runtime/Core/Tests/Memory/SharedBufferTest.cpp
@@ -119,25 +119,25 @@ static_assert(std::is_same<uint32, decltype(GetTypeHash(DeclVal<const FWeakShare
 
 ///////////////////////////////////////////////////////////////////////////////////////////////////////////////////////
 
-TEST_CASE("Core::Memory::FUniqueBuffer::UniqueBuffer", "[Core][Memory][Smoke]")
+TEST_CASE_METHOD(FAutomationTestFixture, "Core::Memory::FUniqueBuffer::UniqueBuffer", "[Core][Memory][Smoke]")
 {
 	// Test Null
 	{
 		FUniqueBuffer Buffer;
-		TestTrue(TEXT("FUniqueBuffer().IsOwned()"), Buffer.IsOwned());
-		TestEqual(TEXT("FUniqueBuffer().GetSize()"), Buffer.GetSize(), uint64(0));
-		TestEqual(TEXT("FUniqueBuffer().GetView().GetSize()"), Buffer.GetView().GetSize(), uint64(0));
-		TestEqual(TEXT("FUniqueBuffer().GetView().GetData()"), Buffer.GetView().GetData(), Buffer.GetData());
+		TEST_TRUE(TEXT("FUniqueBuffer().IsOwned()"), Buffer.IsOwned());
+		TEST_EQUAL(TEXT("FUniqueBuffer().GetSize()"), Buffer.GetSize(), uint64(0));
+		TEST_EQUAL(TEXT("FUniqueBuffer().GetView().GetSize()"), Buffer.GetView().GetSize(), uint64(0));
+		TEST_EQUAL(TEXT("FUniqueBuffer().GetView().GetData()"), Buffer.GetView().GetData(), Buffer.GetData());
 	}
 
 	// Test Alloc
 	{
 		constexpr uint64 Size = 64;
 		FUniqueBuffer Buffer = FUniqueBuffer::Alloc(Size);
-		TestTrue(TEXT("FUniqueBuffer::Alloc().IsOwned()"), Buffer.IsOwned());
-		TestEqual(TEXT("FUniqueBuffer::Alloc().GetSize()"), Buffer.GetSize(), Size);
-		TestEqual(TEXT("FUniqueBuffer::Alloc().GetView().GetSize()"), Buffer.GetView().GetSize(), Size);
-		TestEqual(TEXT("FUniqueBuffer::Alloc().GetView().GetData()"), Buffer.GetView().GetData(), Buffer.GetData());
+		TEST_TRUE(TEXT("FUniqueBuffer::Alloc().IsOwned()"), Buffer.IsOwned());
+		TEST_EQUAL(TEXT("FUniqueBuffer::Alloc().GetSize()"), Buffer.GetSize(), Size);
+		TEST_EQUAL(TEXT("FUniqueBuffer::Alloc().GetView().GetSize()"), Buffer.GetView().GetSize(), Size);
+		TEST_EQUAL(TEXT("FUniqueBuffer::Alloc().GetView().GetData()"), Buffer.GetView().GetData(), Buffer.GetData());
 	}
 
 	// Test Clone
@@ -145,9 +145,9 @@ TEST_CASE("Core::Memory::FUniqueBuffer::UniqueBuffer", "[Core][Memory][Smoke]")
 		constexpr uint64 Size = 64;
 		const uint8 Data[Size]{};
 		FUniqueBuffer Buffer = FUniqueBuffer::Clone(Data, Size);
-		TestTrue(TEXT("FUniqueBuffer::Clone().IsOwned()"), Buffer.IsOwned());
-		TestEqual(TEXT("FUniqueBuffer::Clone().GetSize()"), Buffer.GetSize(), Size);
-		TestNotEqual(TEXT("FUniqueBuffer::Clone().GetData()"), static_cast<const void*>(Buffer.GetData()), static_cast<const void*>(Data));
+		TEST_TRUE(TEXT("FUniqueBuffer::Clone().IsOwned()"), Buffer.IsOwned());
+		TEST_EQUAL(TEXT("FUniqueBuffer::Clone().GetSize()"), Buffer.GetSize(), Size);
+		TEST_NOT_EQUAL(TEXT("FUniqueBuffer::Clone().GetData()"), static_cast<const void*>(Buffer.GetData()), static_cast<const void*>(Data));
 	}
 
 	// Test MakeView
@@ -155,9 +155,9 @@ TEST_CASE("Core::Memory::FUniqueBuffer::UniqueBuffer", "[Core][Memory][Smoke]")
 		constexpr uint64 Size = 64;
 		uint8 Data[Size]{};
 		FUniqueBuffer Buffer = FUniqueBuffer::MakeView(Data, Size);
-		TestFalse(TEXT("FUniqueBuffer::MakeView().IsOwned()"), Buffer.IsOwned());
-		TestEqual(TEXT("FUniqueBuffer::MakeView().GetSize()"), Buffer.GetSize(), Size);
-		TestEqual(TEXT("FUniqueBuffer::MakeView().GetData()"), Buffer.GetData(), static_cast<void*>(Data));
+		TEST_FALSE(TEXT("FUniqueBuffer::MakeView().IsOwned()"), Buffer.IsOwned());
+		TEST_EQUAL(TEXT("FUniqueBuffer::MakeView().GetSize()"), Buffer.GetSize(), Size);
+		TEST_EQUAL(TEXT("FUniqueBuffer::MakeView().GetData()"), Buffer.GetData(), static_cast<void*>(Data));
 	}
 
 	// Test TakeOwnership
@@ -166,11 +166,11 @@ TEST_CASE("Core::Memory::FUniqueBuffer::UniqueBuffer", "[Core][Memory][Smoke]")
 		constexpr uint64 Size = 64;
 		uint8 Data[Size]{};
 		FUniqueBuffer Buffer = FUniqueBuffer::TakeOwnership(Data, Size, [&bDeleted](void*, uint64) { bDeleted = true; });
-		TestTrue(TEXT("FUniqueBuffer::TakeOwnership().IsOwned()"), Buffer.IsOwned());
-		TestEqual(TEXT("FUniqueBuffer::TakeOwnership().GetSize()"), Buffer.GetSize(), Size);
-		TestEqual(TEXT("FUniqueBuffer::TakeOwnership().GetData()"), Buffer.GetData(), static_cast<void*>(Data));
+		TEST_TRUE(TEXT("FUniqueBuffer::TakeOwnership().IsOwned()"), Buffer.IsOwned());
+		TEST_EQUAL(TEXT("FUniqueBuffer::TakeOwnership().GetSize()"), Buffer.GetSize(), Size);
+		TEST_EQUAL(TEXT("FUniqueBuffer::TakeOwnership().GetData()"), Buffer.GetData(), static_cast<void*>(Data));
 		Buffer.Reset();
-		TestTrue(TEXT("FUniqueBuffer::TakeOwnership() Deleted"), bDeleted);
+		TEST_TRUE(TEXT("FUniqueBuffer::TakeOwnership() Deleted"), bDeleted);
 	}
 
 	// Test MakeOwned
@@ -178,27 +178,27 @@ TEST_CASE("Core::Memory::FUniqueBuffer::UniqueBuffer", "[Core][Memory][Smoke]")
 		constexpr uint64 Size = 64;
 		uint8 Data[Size]{};
 		FUniqueBuffer Buffer = FUniqueBuffer::MakeView(Data, Size).MakeOwned();
-		TestTrue(TEXT("FUniqueBuffer::MakeOwned(View).IsOwned()"), Buffer.IsOwned());
-		TestEqual(TEXT("FUniqueBuffer::MakeOwned(View).GetSize()"), Buffer.GetSize(), Size);
-		TestFalse(TEXT("FUniqueBuffer::MakeOwned(View).GetData()"), Buffer.GetData() == Data);
+		TEST_TRUE(TEXT("FUniqueBuffer::MakeOwned(View).IsOwned()"), Buffer.IsOwned());
+		TEST_EQUAL(TEXT("FUniqueBuffer::MakeOwned(View).GetSize()"), Buffer.GetSize(), Size);
+		TEST_FALSE(TEXT("FUniqueBuffer::MakeOwned(View).GetData()"), Buffer.GetData() == Data);
 		void* const OwnedData = Buffer.GetData();
 		Buffer = MoveTemp(Buffer).MakeOwned();
-		TestTrue(TEXT("FUniqueBuffer::MakeOwned(Owned).IsOwned()"), Buffer.IsOwned());
-		TestEqual(TEXT("FUniqueBuffer::MakeOwned(Owned).GetSize()"), Buffer.GetSize(), Size);
-		TestEqual(TEXT("FUniqueBuffer::MakeOwned(Owned).GetData()"), Buffer.GetData(), OwnedData);
+		TEST_TRUE(TEXT("FUniqueBuffer::MakeOwned(Owned).IsOwned()"), Buffer.IsOwned());
+		TEST_EQUAL(TEXT("FUniqueBuffer::MakeOwned(Owned).GetSize()"), Buffer.GetSize(), Size);
+		TEST_EQUAL(TEXT("FUniqueBuffer::MakeOwned(Owned).GetData()"), Buffer.GetData(), OwnedData);
 	}
 }
 
 
-TEST_CASE("Core::Memory::FSharedBuffer::SharedBuffer", "[Core][Memory][Smoke]")
+TEST_CASE_METHOD(FAutomationTestFixture, "Core::Memory::FSharedBuffer::SharedBuffer", "[Core][Memory][Smoke]")
 {
 	// Test Null
 	{
 		FSharedBuffer Buffer;
-		TestTrue(TEXT("FSharedBuffer().IsOwned()"), Buffer.IsOwned());
-		TestEqual(TEXT("FSharedBuffer().GetSize()"), Buffer.GetSize(), uint64(0));
-		TestEqual(TEXT("FSharedBuffer().GetView().GetSize()"), Buffer.GetView().GetSize(), uint64(0));
-		TestEqual(TEXT("FSharedBuffer().GetView().GetData()"), Buffer.GetView().GetData(), Buffer.GetData());
+		TEST_TRUE(TEXT("FSharedBuffer().IsOwned()"), Buffer.IsOwned());
+		TEST_EQUAL(TEXT("FSharedBuffer().GetSize()"), Buffer.GetSize(), uint64(0));
+		TEST_EQUAL(TEXT("FSharedBuffer().GetView().GetSize()"), Buffer.GetView().GetSize(), uint64(0));
+		TEST_EQUAL(TEXT("FSharedBuffer().GetView().GetData()"), Buffer.GetView().GetData(), Buffer.GetData());
 	}
 
 	// Test Clone
@@ -206,9 +206,9 @@ TEST_CASE("Core::Memory::FSharedBuffer::SharedBuffer", "[Core][Memory][Smoke]")
 		constexpr uint64 Size = 64;
 		const uint8 Data[Size]{};
 		FSharedBuffer Buffer = FSharedBuffer::Clone(Data, Size);
-		TestTrue(TEXT("FSharedBuffer::Clone().IsOwned()"), Buffer.IsOwned());
-		TestEqual(TEXT("FSharedBuffer::Clone().GetSize()"), Buffer.GetSize(), Size);
-		TestNotEqual(TEXT("FSharedBuffer::Clone().GetData()"), Buffer.GetData(), static_cast<const void*>(Data));
+		TEST_TRUE(TEXT("FSharedBuffer::Clone().IsOwned()"), Buffer.IsOwned());
+		TEST_EQUAL(TEXT("FSharedBuffer::Clone().GetSize()"), Buffer.GetSize(), Size);
+		TEST_NOT_EQUAL(TEXT("FSharedBuffer::Clone().GetData()"), Buffer.GetData(), static_cast<const void*>(Data));
 	}
 
 	// Test MakeView
@@ -216,9 +216,9 @@ TEST_CASE("Core::Memory::FSharedBuffer::SharedBuffer", "[Core][Memory][Smoke]")
 		constexpr uint64 Size = 64;
 		uint8 Data[Size]{};
 		FSharedBuffer Buffer = FSharedBuffer::MakeView(Data, Size);
-		TestFalse(TEXT("FSharedBuffer::MakeView().IsOwned()"), Buffer.IsOwned());
-		TestEqual(TEXT("FSharedBuffer::MakeView().GetSize()"), Buffer.GetSize(), Size);
-		TestEqual(TEXT("FSharedBuffer::MakeView().GetData()"), Buffer.GetData(), static_cast<const void*>(Data));
+		TEST_FALSE(TEXT("FSharedBuffer::MakeView().IsOwned()"), Buffer.IsOwned());
+		TEST_EQUAL(TEXT("FSharedBuffer::MakeView().GetSize()"), Buffer.GetSize(), Size);
+		TEST_EQUAL(TEXT("FSharedBuffer::MakeView().GetData()"), Buffer.GetData(), static_cast<const void*>(Data));
 	}
 
 	// Test MakeView Outer View
@@ -233,12 +233,12 @@ TEST_CASE("Core::Memory::FSharedBuffer::SharedBuffer", "[Core][Memory][Smoke]")
 				Buffer = FSharedBuffer::MakeView(OuterBuffer.GetData(), OuterBuffer.GetSize() / 2, OuterBuffer);
 				WeakBuffer = OuterBuffer;
 			}
-			TestFalse(TEXT("FSharedBuffer::MakeView(OuterView).IsOwned()"), Buffer.IsOwned());
-			TestEqual(TEXT("FSharedBuffer::MakeView(OuterView).GetSize()"), Buffer.GetSize(), Size / 2);
-			TestEqual(TEXT("FSharedBuffer::MakeView(OuterView).GetData()"), Buffer.GetData(), static_cast<const void*>(Data));
-			TestFalse(TEXT("FSharedBuffer::MakeView(OuterView) Outer Not Null"), WeakBuffer.Pin().IsNull());
+			TEST_FALSE(TEXT("FSharedBuffer::MakeView(OuterView).IsOwned()"), Buffer.IsOwned());
+			TEST_EQUAL(TEXT("FSharedBuffer::MakeView(OuterView).GetSize()"), Buffer.GetSize(), Size / 2);
+			TEST_EQUAL(TEXT("FSharedBuffer::MakeView(OuterView).GetData()"), Buffer.GetData(), static_cast<const void*>(Data));
+			TEST_FALSE(TEXT("FSharedBuffer::MakeView(OuterView) Outer Not Null"), WeakBuffer.Pin().IsNull());
 		}
-		TestTrue(TEXT("FSharedBuffer::MakeView(OuterView) Outer Null"), WeakBuffer.Pin().IsNull());
+		TEST_TRUE(TEXT("FSharedBuffer::MakeView(OuterView) Outer Null"), WeakBuffer.Pin().IsNull());
 	}
 
 	// Test MakeView Outer Owned
@@ -250,16 +250,16 @@ TEST_CASE("Core::Memory::FSharedBuffer::SharedBuffer", "[Core][Memory][Smoke]")
 			{
 				FSharedBuffer OuterBuffer = FUniqueBuffer::Alloc(Size).MoveToShared();
 				Buffer = FSharedBuffer::MakeView(OuterBuffer.GetData(), OuterBuffer.GetSize(), OuterBuffer);
-				TestEqual(TEXT("FSharedBuffer::MakeView(OuterOwned, SameView)"), Buffer, OuterBuffer);
+				TEST_EQUAL(TEXT("FSharedBuffer::MakeView(OuterOwned, SameView)"), Buffer, OuterBuffer);
 				Buffer = FSharedBuffer::MakeView(OuterBuffer.GetData(), OuterBuffer.GetSize() / 2, OuterBuffer);
 				WeakBuffer = OuterBuffer;
 			}
-			TestTrue(TEXT("FSharedBuffer::MakeView(OuterOwned).IsOwned()"), Buffer.IsOwned());
-			TestEqual(TEXT("FSharedBuffer::MakeView(OuterOwned).GetSize()"), Buffer.GetSize(), Size / 2);
-			TestEqual(TEXT("FSharedBuffer::MakeView(OuterOwned).GetData()"), Buffer.GetData(), WeakBuffer.Pin().GetData());
-			TestFalse(TEXT("FSharedBuffer::MakeView(OuterOwned) Outer Not Null"), WeakBuffer.Pin().IsNull());
+			TEST_TRUE(TEXT("FSharedBuffer::MakeView(OuterOwned).IsOwned()"), Buffer.IsOwned());
+			TEST_EQUAL(TEXT("FSharedBuffer::MakeView(OuterOwned).GetSize()"), Buffer.GetSize(), Size / 2);
+			TEST_EQUAL(TEXT("FSharedBuffer::MakeView(OuterOwned).GetData()"), Buffer.GetData(), WeakBuffer.Pin().GetData());
+			TEST_FALSE(TEXT("FSharedBuffer::MakeView(OuterOwned) Outer Not Null"), WeakBuffer.Pin().IsNull());
 		}
-		TestTrue(TEXT("FSharedBuffer::MakeView(OuterOwned) Outer  Null"), WeakBuffer.Pin().IsNull());
+		TEST_TRUE(TEXT("FSharedBuffer::MakeView(OuterOwned) Outer  Null"), WeakBuffer.Pin().IsNull());
 	}
 
 	// Test TakeOwnership
@@ -268,11 +268,11 @@ TEST_CASE("Core::Memory::FSharedBuffer::SharedBuffer", "[Core][Memory][Smoke]")
 		constexpr uint64 Size = 64;
 		const void* Data = new uint8[Size];
 		FSharedBuffer Buffer = FSharedBuffer::TakeOwnership(Data, Size, [&bDeleted](void*, uint64) { bDeleted = true; });
-		TestTrue(TEXT("FSharedBuffer::TakeOwnership().IsOwned()"), Buffer.IsOwned());
-		TestEqual(TEXT("FSharedBuffer::TakeOwnership().GetSize()"), Buffer.GetSize(), Size);
-		TestEqual(TEXT("FSharedBuffer::TakeOwnership().GetData()"), Buffer.GetData(), Data);
+		TEST_TRUE(TEXT("FSharedBuffer::TakeOwnership().IsOwned()"), Buffer.IsOwned());
+		TEST_EQUAL(TEXT("FSharedBuffer::TakeOwnership().GetSize()"), Buffer.GetSize(), Size);
+		TEST_EQUAL(TEXT("FSharedBuffer::TakeOwnership().GetData()"), Buffer.GetData(), Data);
 		Buffer.Reset();
-		TestTrue(TEXT("FSharedBuffer::TakeOwnership() Deleted"), bDeleted);
+		TEST_TRUE(TEXT("FSharedBuffer::TakeOwnership() Deleted"), bDeleted);
 	}
 
 	// Test MakeOwned
@@ -280,28 +280,28 @@ TEST_CASE("Core::Memory::FSharedBuffer::SharedBuffer", "[Core][Memory][Smoke]")
 		constexpr uint64 Size = 64;
 		uint8 Data[Size]{};
 		FSharedBuffer Buffer = FSharedBuffer::MakeView(Data, Size).MakeOwned();
-		TestTrue(TEXT("FSharedBuffer::MakeOwned(View).IsOwned()"), Buffer.IsOwned());
-		TestEqual(TEXT("FSharedBuffer::MakeOwned(View).GetSize()"), Buffer.GetSize(), Size);
-		TestFalse(TEXT("FSharedBuffer::MakeOwned(View).GetData()"), Buffer.GetData() == Data);
+		TEST_TRUE(TEXT("FSharedBuffer::MakeOwned(View).IsOwned()"), Buffer.IsOwned());
+		TEST_EQUAL(TEXT("FSharedBuffer::MakeOwned(View).GetSize()"), Buffer.GetSize(), Size);
+		TEST_FALSE(TEXT("FSharedBuffer::MakeOwned(View).GetData()"), Buffer.GetData() == Data);
 		FSharedBuffer BufferCopy = Buffer.MakeOwned();
-		TestTrue(TEXT("FSharedBuffer::MakeOwned(Owned).IsOwned()"), BufferCopy.IsOwned());
-		TestEqual(TEXT("FSharedBuffer::MakeOwned(Owned).GetSize()"), BufferCopy.GetSize(), Size);
-		TestEqual(TEXT("FSharedBuffer::MakeOwned(Owned).GetData()"), BufferCopy.GetData(), Buffer.GetData());
+		TEST_TRUE(TEXT("FSharedBuffer::MakeOwned(Owned).IsOwned()"), BufferCopy.IsOwned());
+		TEST_EQUAL(TEXT("FSharedBuffer::MakeOwned(Owned).GetSize()"), BufferCopy.GetSize(), Size);
+		TEST_EQUAL(TEXT("FSharedBuffer::MakeOwned(Owned).GetData()"), BufferCopy.GetData(), Buffer.GetData());
 	}
 
 	// Test MoveToUnique
 	{
 		FUniqueBuffer Buffer = FSharedBuffer().MoveToUnique();
-		TestTrue(TEXT("FSharedBuffer::MoveToUnique(Null).IsOwned()"), Buffer.IsOwned());
-		TestEqual(TEXT("FSharedBuffer::MoveToUnique(Null).GetSize()"), Buffer.GetSize(), uint64(0));
+		TEST_TRUE(TEXT("FSharedBuffer::MoveToUnique(Null).IsOwned()"), Buffer.IsOwned());
+		TEST_EQUAL(TEXT("FSharedBuffer::MoveToUnique(Null).GetSize()"), Buffer.GetSize(), uint64(0));
 	}
 	{
 		FSharedBuffer Shared = FSharedBuffer::Clone(MakeMemoryView<uint8>({1, 2, 3, 4}));
 		FUniqueBuffer Buffer = FSharedBuffer(Shared).MoveToUnique();
-		TestTrue(TEXT("FSharedBuffer::MoveToUnique(Shared).IsOwned()"), Buffer.IsOwned());
-		TestEqual(TEXT("FSharedBuffer::MoveToUnique(Shared).GetSize()"), Buffer.GetSize(), Shared.GetSize());
-		TestNotEqual(TEXT("FSharedBuffer::MoveToUnique(Shared).GetData()"), const_cast<const void*>(Buffer.GetData()), Shared.GetData());
-		TestTrue(TEXT("FSharedBuffer::MoveToUnique(Shared)->Equal"), Buffer.GetView().EqualBytes(Shared.GetView()));
+		TEST_TRUE(TEXT("FSharedBuffer::MoveToUnique(Shared).IsOwned()"), Buffer.IsOwned());
+		TEST_EQUAL(TEXT("FSharedBuffer::MoveToUnique(Shared).GetSize()"), Buffer.GetSize(), Shared.GetSize());
+		TEST_NOT_EQUAL(TEXT("FSharedBuffer::MoveToUnique(Shared).GetData()"), const_cast<const void*>(Buffer.GetData()), Shared.GetData());
+		TEST_TRUE(TEXT("FSharedBuffer::MoveToUnique(Shared)->Equal"), Buffer.GetView().EqualBytes(Shared.GetView()));
 	}
 	{
 		bool bDeleted = false;
@@ -309,57 +309,57 @@ TEST_CASE("Core::Memory::FSharedBuffer::SharedBuffer", "[Core][Memory][Smoke]")
 		uint8 Data[Size]{};
 		FSharedBuffer Shared = FUniqueBuffer::TakeOwnership(Data, Size, [&bDeleted](void*) { bDeleted = true; }).MoveToShared();
 		FUniqueBuffer Buffer = Shared.MoveToUnique();
-		TestTrue(TEXT("FSharedBuffer::MoveToUnique(Owned).IsNull()"), Shared.IsNull());
-		TestTrue(TEXT("FSharedBuffer::MoveToUnique(Owned).IsOwned()"), Buffer.IsOwned());
-		TestEqual(TEXT("FSharedBuffer::MoveToUnique(Owned).GetSize()"), Buffer.GetSize(), Size);
-		TestEqual(TEXT("FSharedBuffer::MoveToUnique(Owned).GetData()"), Buffer.GetData(), static_cast<void*>(Data));
+		TEST_TRUE(TEXT("FSharedBuffer::MoveToUnique(Owned).IsNull()"), Shared.IsNull());
+		TEST_TRUE(TEXT("FSharedBuffer::MoveToUnique(Owned).IsOwned()"), Buffer.IsOwned());
+		TEST_EQUAL(TEXT("FSharedBuffer::MoveToUnique(Owned).GetSize()"), Buffer.GetSize(), Size);
+		TEST_EQUAL(TEXT("FSharedBuffer::MoveToUnique(Owned).GetData()"), Buffer.GetData(), static_cast<void*>(Data));
 		Buffer.Reset();
-		TestTrue(TEXT("FSharedBuffer::MoveToUnique(Owned) Deleted"), bDeleted);
+		TEST_TRUE(TEXT("FSharedBuffer::MoveToUnique(Owned) Deleted"), bDeleted);
 	}
 	{
 		constexpr uint64 Size = 64;
 		uint8 Data[Size]{};
 		FSharedBuffer Shared = FSharedBuffer::TakeOwnership(Data, Size, [](void*){});
 		FUniqueBuffer Buffer = Shared.MoveToUnique();
-		TestTrue(TEXT("FSharedBuffer::MoveToUnique(ImmutableOwned).IsNull()"), Shared.IsNull());
-		TestTrue(TEXT("FSharedBuffer::MoveToUnique(ImmutableOwned).IsOwned()"), Buffer.IsOwned());
-		TestEqual(TEXT("FSharedBuffer::MoveToUnique(ImmutableOwned).GetSize()"), Buffer.GetSize(), Size);
-		TestNotEqual(TEXT("FSharedBuffer::MoveToUnique(ImmutableOwned).GetData()"), Buffer.GetData(), static_cast<void*>(Data));
+		TEST_TRUE(TEXT("FSharedBuffer::MoveToUnique(ImmutableOwned).IsNull()"), Shared.IsNull());
+		TEST_TRUE(TEXT("FSharedBuffer::MoveToUnique(ImmutableOwned).IsOwned()"), Buffer.IsOwned());
+		TEST_EQUAL(TEXT("FSharedBuffer::MoveToUnique(ImmutableOwned).GetSize()"), Buffer.GetSize(), Size);
+		TEST_NOT_EQUAL(TEXT("FSharedBuffer::MoveToUnique(ImmutableOwned).GetData()"), Buffer.GetData(), static_cast<void*>(Data));
 	}
 	{
 		constexpr uint64 Size = 64;
 		uint8 Data[Size]{};
 		FSharedBuffer Shared = FUniqueBuffer::MakeView(MakeMemoryView(Data)).MoveToShared();
 		FUniqueBuffer Buffer = Shared.MoveToUnique();
-		TestTrue(TEXT("FSharedBuffer::MoveToUnique(View).IsNull()"), Shared.IsNull());
-		TestTrue(TEXT("FSharedBuffer::MoveToUnique(View).IsOwned()"), Buffer.IsOwned());
-		TestEqual(TEXT("FSharedBuffer::MoveToUnique(View).GetSize()"), Buffer.GetSize(), Size);
-		TestNotEqual(TEXT("FSharedBuffer::MoveToUnique(View).GetData()"), Buffer.GetData(), static_cast<void*>(Data));
+		TEST_TRUE(TEXT("FSharedBuffer::MoveToUnique(View).IsNull()"), Shared.IsNull());
+		TEST_TRUE(TEXT("FSharedBuffer::MoveToUnique(View).IsOwned()"), Buffer.IsOwned());
+		TEST_EQUAL(TEXT("FSharedBuffer::MoveToUnique(View).GetSize()"), Buffer.GetSize(), Size);
+		TEST_NOT_EQUAL(TEXT("FSharedBuffer::MoveToUnique(View).GetData()"), Buffer.GetData(), static_cast<void*>(Data));
 	}
 	{
 		constexpr uint64 Size = 64;
 		uint8 Data[Size]{};
 		FSharedBuffer Shared = FSharedBuffer::MakeView(MakeMemoryView(Data));
 		FUniqueBuffer Buffer = Shared.MoveToUnique();
-		TestTrue(TEXT("FSharedBuffer::MoveToUnique(ImmutableView).IsNull()"), Shared.IsNull());
-		TestTrue(TEXT("FSharedBuffer::MoveToUnique(ImmutableView).IsOwned()"), Buffer.IsOwned());
-		TestEqual(TEXT("FSharedBuffer::MoveToUnique(ImmutableView).GetSize()"), Buffer.GetSize(), Size);
-		TestNotEqual(TEXT("FSharedBuffer::MoveToUnique(ImmutableView).GetData()"), Buffer.GetData(), static_cast<void*>(Data));
+		TEST_TRUE(TEXT("FSharedBuffer::MoveToUnique(ImmutableView).IsNull()"), Shared.IsNull());
+		TEST_TRUE(TEXT("FSharedBuffer::MoveToUnique(ImmutableView).IsOwned()"), Buffer.IsOwned());
+		TEST_EQUAL(TEXT("FSharedBuffer::MoveToUnique(ImmutableView).GetSize()"), Buffer.GetSize(), Size);
+		TEST_NOT_EQUAL(TEXT("FSharedBuffer::MoveToUnique(ImmutableView).GetData()"), Buffer.GetData(), static_cast<void*>(Data));
 	}
 	{
 		constexpr uint64 Size = 64;
 		uint8 Data[Size]{};
 		FSharedBuffer Shared = FUniqueBuffer::TakeOwnership(Data, Size, [](void*){}).MoveToShared();
 		FUniqueBuffer Buffer = FSharedBuffer::MakeView(Shared.GetView().Left(Size / 2), Shared).MoveToUnique();
-		TestFalse(TEXT("FSharedBuffer::MoveToUnique(OwnedView).IsNull()"), Shared.IsNull());
-		TestTrue(TEXT("FSharedBuffer::MoveToUnique(OwnedView).IsOwned()"), Buffer.IsOwned());
-		TestEqual(TEXT("FSharedBuffer::MoveToUnique(OwnedView).GetSize()"), Buffer.GetSize(), Size / 2);
-		TestNotEqual(TEXT("FSharedBuffer::MoveToUnique(OwnedView).GetData()"), Buffer.GetData(), static_cast<void*>(Data));
+		TEST_FALSE(TEXT("FSharedBuffer::MoveToUnique(OwnedView).IsNull()"), Shared.IsNull());
+		TEST_TRUE(TEXT("FSharedBuffer::MoveToUnique(OwnedView).IsOwned()"), Buffer.IsOwned());
+		TEST_EQUAL(TEXT("FSharedBuffer::MoveToUnique(OwnedView).GetSize()"), Buffer.GetSize(), Size / 2);
+		TEST_NOT_EQUAL(TEXT("FSharedBuffer::MoveToUnique(OwnedView).GetData()"), Buffer.GetData(), static_cast<void*>(Data));
 	}
 }
 
 
-TEST_CASE("Core::Memory::FTestBufferOwner::BufferOwner", "[Core][Memory][Smoke]")
+TEST_CASE_METHOD(FAutomationTestFixture, "Core::Memory::FTestBufferOwner::BufferOwner", "[Core][Memory][Smoke]")
 {
 	class FTestBufferOwner final : public FBufferOwner
 	{
@@ -399,50 +399,50 @@ TEST_CASE("Core::Memory::FTestBufferOwner::BufferOwner", "[Core][Memory][Smoke]"
 	bool bDeleted = false;
 
 	FSharedBuffer Buffer(new FTestBufferOwner(bMaterialized, bFreed, bDeleted));
-	TestFalse(TEXT("FSharedBuffer(BufferOwner)->!Materialized"), bMaterialized);
-	TestFalse(TEXT("FSharedBuffer(BufferOwner)->!Freed"), bFreed);
-	TestFalse(TEXT("FSharedBuffer(BufferOwner)->!Deleted"), bDeleted);
+	TEST_FALSE(TEXT("FSharedBuffer(BufferOwner)->!Materialized"), bMaterialized);
+	TEST_FALSE(TEXT("FSharedBuffer(BufferOwner)->!Freed"), bFreed);
+	TEST_FALSE(TEXT("FSharedBuffer(BufferOwner)->!Deleted"), bDeleted);
 
 	(void)Buffer.GetData();
-	TestTrue(TEXT("FSharedBuffer(BufferOwner)->GetData()->Materialized"), bMaterialized);
+	TEST_TRUE(TEXT("FSharedBuffer(BufferOwner)->GetData()->Materialized"), bMaterialized);
 	bMaterialized = false;
 	(void)Buffer.GetData();
-	TestFalse(TEXT("FSharedBuffer(BufferOwner)->GetData()->!Materialized"), bMaterialized);
+	TEST_FALSE(TEXT("FSharedBuffer(BufferOwner)->GetData()->!Materialized"), bMaterialized);
 	bMaterialized = false;
 
 	Buffer = FSharedBuffer(new FTestBufferOwner(bMaterialized, bFreed, bDeleted));
-	TestFalse(TEXT("FSharedBuffer(BufferOwner)->Assign(FSharedBuffer)->!Materialized"), bMaterialized);
-	TestTrue(TEXT("FSharedBuffer(BufferOwner)->Assign(FSharedBuffer)->Freed"), bFreed);
-	TestTrue(TEXT("FSharedBuffer(BufferOwner)->Assign(FSharedBuffer)->Deleted"), bDeleted);
+	TEST_FALSE(TEXT("FSharedBuffer(BufferOwner)->Assign(FSharedBuffer)->!Materialized"), bMaterialized);
+	TEST_TRUE(TEXT("FSharedBuffer(BufferOwner)->Assign(FSharedBuffer)->Freed"), bFreed);
+	TEST_TRUE(TEXT("FSharedBuffer(BufferOwner)->Assign(FSharedBuffer)->Deleted"), bDeleted);
 	bFreed = false;
 	bDeleted = false;
 
 	(void)Buffer.GetSize();
-	TestTrue(TEXT("FSharedBuffer(BufferOwner)->GetSize()->Materialized"), bMaterialized);
+	TEST_TRUE(TEXT("FSharedBuffer(BufferOwner)->GetSize()->Materialized"), bMaterialized);
 	bMaterialized = false;
 	(void)Buffer.GetSize();
-	TestFalse(TEXT("FSharedBuffer(BufferOwner)->GetSize()->!Materialized"), bMaterialized);
+	TEST_FALSE(TEXT("FSharedBuffer(BufferOwner)->GetSize()->!Materialized"), bMaterialized);
 	bMaterialized = false;
 
 	FWeakSharedBuffer WeakBuffer = Buffer;
 	Buffer.Reset();
-	TestTrue(TEXT("FSharedBuffer(BufferOwner)->Reset(Shared)->Freed"), bFreed);
-	TestFalse(TEXT("FSharedBuffer(BufferOwner)->Reset(Shared)->!Deleted"), bDeleted);
+	TEST_TRUE(TEXT("FSharedBuffer(BufferOwner)->Reset(Shared)->Freed"), bFreed);
+	TEST_FALSE(TEXT("FSharedBuffer(BufferOwner)->Reset(Shared)->!Deleted"), bDeleted);
 	WeakBuffer.Reset();
-	TestTrue(TEXT("FSharedBuffer(BufferOwner)->Reset(Weak)->Deleted"), bDeleted);
+	TEST_TRUE(TEXT("FSharedBuffer(BufferOwner)->Reset(Weak)->Deleted"), bDeleted);
 
 	bMaterialized = false;
 	bFreed = false;
 	bDeleted = false;
 	Buffer = FSharedBuffer(new FTestBufferOwner(bMaterialized, bFreed, bDeleted));
 
-	TestFalse(TEXT("FSharedBuffer(BufferOwner)->IsMaterialized()"), Buffer.IsMaterialized());
+	TEST_FALSE(TEXT("FSharedBuffer(BufferOwner)->IsMaterialized()"), Buffer.IsMaterialized());
 	Buffer.Materialize();
-	TestTrue(TEXT("FSharedBuffer(BufferOwner)->Materialize()->IsMaterialized()"), Buffer.IsMaterialized());
-	TestTrue(TEXT("FSharedBuffer(BufferOwner)->Materialize()->Materialized"), bMaterialized);
+	TEST_TRUE(TEXT("FSharedBuffer(BufferOwner)->Materialize()->IsMaterialized()"), Buffer.IsMaterialized());
+	TEST_TRUE(TEXT("FSharedBuffer(BufferOwner)->Materialize()->Materialized"), bMaterialized);
 	bMaterialized = false;
 	Buffer.Materialize();
-	TestFalse(TEXT("FSharedBuffer(BufferOwner)->Materialize()->!Materialized"), bMaterialized);
+	TEST_FALSE(TEXT("FSharedBuffer(BufferOwner)->Materialize()->!Materialized"), bMaterialized);
 	bMaterialized = false;
 	Buffer.Reset();
 
@@ -450,19 +450,19 @@ TEST_CASE("Core::Memory::FTestBufferOwner::BufferOwner", "[Core][Memory][Smoke]"
 	bDeleted = false;
 	Buffer = FSharedBuffer(new FTestBufferOwner(bMaterialized, bFreed, bDeleted));
 	Buffer.Reset();
-	TestFalse(TEXT("FSharedBuffer(BufferOwner)->Reset(!Materialized)->!Materialized"), bMaterialized);
-	TestTrue(TEXT("FSharedBuffer(BufferOwner)->Reset(!Materialized)->Freed"), bFreed);
-	TestTrue(TEXT("FSharedBuffer(BufferOwner)->Reset(!Materialized)->Deleted"), bDeleted);
+	TEST_FALSE(TEXT("FSharedBuffer(BufferOwner)->Reset(!Materialized)->!Materialized"), bMaterialized);
+	TEST_TRUE(TEXT("FSharedBuffer(BufferOwner)->Reset(!Materialized)->Freed"), bFreed);
+	TEST_TRUE(TEXT("FSharedBuffer(BufferOwner)->Reset(!Materialized)->Deleted"), bDeleted);
 }
 
 
-TEST_CASE("Core::Memory::FSharedBuffer::SharedBufferFromArray", "[Core][Memory][Smoke]")
+TEST_CASE_METHOD(FAutomationTestFixture, "Core::Memory::FSharedBuffer::SharedBufferFromArray", "[Core][Memory][Smoke]")
 {
 	TArray<uint16> Array{1, 2, 3, 4, 5, 6, 7, 8};
 	const void* const ExpectedData = Array.GetData();
 	const uint64 ExpectedSize = Array.Num() * sizeof(uint16);
 	const FSharedBuffer Buffer = MakeSharedBufferFromArray(MoveTemp(Array));
-	TestTrue(TEXT("MakeSharedBufferFromArray -> Array.IsEmpty"), Array.IsEmpty());
-	TestEqual(TEXT("MakeSharedBufferFromArray -> Buffer.GetData()"), Buffer.GetData(), ExpectedData);
-	TestEqual(TEXT("MakeSharedBufferFromArray -> Buffer.GetSize()"), Buffer.GetSize(), ExpectedSize);
+	TEST_TRUE(TEXT("MakeSharedBufferFromArray -> Array.IsEmpty"), Array.IsEmpty());
+	TEST_EQUAL(TEXT("MakeSharedBufferFromArray -> Buffer.GetData()"), Buffer.GetData(), ExpectedData);
+	TEST_EQUAL(TEXT("MakeSharedBufferFromArray -> Buffer.GetSize()"), Buffer.GetSize(), ExpectedSize);
 }
diff --git a/Engine/Source/Runtime/Core/Tests/Misc/AsciiSetTest.cpp b/Engine/Source/Runtime/Core/Tests/Misc/AsciiSetTest.cpp
--- a/Engine/Source/Runtime/Core/Tests/Misc/AsciiSetTest.cpp
+++ b/Engine/Source/Runtime/Core/Tests/Misc/AsciiSetTest.cpp
@@ -5,15 +5,15 @@
 #include "Misc/AsciiSet.h"
 #include "TestHarness.h"
 
-TEST_CASE("Core::Misc::FAsciiSet::AsciiSet", "[Core][Misc][Smoke]")
+TEST_CASE_METHOD(FAutomationTestFixture, "Core::Misc::FAsciiSet::AsciiSet", "[Core][Misc][Smoke]")
 {
 	constexpr FAsciiSet Whitespaces(" \v\f\t\r\n");
-	TestTrue(TEXT("Contains"), Whitespaces.Contains(' '));
-	TestTrue(TEXT("Contains"), Whitespaces.Contains('\n'));
-	TestFalse(TEXT("Contains"), Whitespaces.Contains('a'));
-	TestFalse(TEXT("Contains no extended ASCII"), Whitespaces.Contains('\x80'));
-	TestFalse(TEXT("Contains no extended ASCII"), Whitespaces.Contains('\xA0'));
-	TestFalse(TEXT("Contains no extended ASCII"), Whitespaces.Contains('\xFF'));
+	TEST_TRUE(TEXT("Contains"), Whitespaces.Contains(' '));
+	TEST_TRUE(TEXT("Contains"), Whitespaces.Contains('\n'));
+	TEST_FALSE(TEXT("Contains"), Whitespaces.Contains('a'));
+	TEST_FALSE(TEXT("Contains no extended ASCII"), Whitespaces.Contains('\x80'));
+	TEST_FALSE(TEXT("Contains no extended ASCII"), Whitespaces.Contains('\xA0'));
+	TEST_FALSE(TEXT("Contains no extended ASCII"), Whitespaces.Contains('\xFF'));
 
 	constexpr FAsciiSet Aa("Aa");
 	uint32 ANum = 0;
@@ -21,61 +21,61 @@ TEST_CASE("Core::Misc::FAsciiSet::AsciiSet", "[Core][Misc][Smoke]")
 	{
 		ANum += Aa.Contains(C);
 	}
-	TestEqual("Contains no wide", ANum, 2u);
+	TEST_EQUAL("Contains no wide", ANum, 2u);
 
 	constexpr FAsciiSet NonWhitespaces = ~Whitespaces;
 	uint32 WhitespaceNum = 0;
 	for (uint8 Char = 0; Char < 128; ++Char)
 	{
 		WhitespaceNum += !!Whitespaces.Test(Char);
-		TestEqual(TEXT("Inverse"), !!Whitespaces.Test(Char), !NonWhitespaces.Test(Char));
+		TEST_EQUAL(TEXT("Inverse"), !!Whitespaces.Test(Char), !NonWhitespaces.Test(Char));
 	}
-	TestEqual(TEXT("Num"), WhitespaceNum, 6);
+	TEST_EQUAL(TEXT("Num"), WhitespaceNum, 6);
 
 	// TODO [ML]
-	//TestEqual(TEXT("Skip"), FAsciiSet::Skip(TEXT("  \t\tHello world!"), Whitespaces), TEXT("Hello world!"));
-	//TestEqual(TEXT("Skip"), FAsciiSet::Skip(TEXT("Hello world!"), Whitespaces), TEXT("Hello world!"));
-	//TestEqual(TEXT("Skip to extended ASCII"), FAsciiSet::Skip(" " "\xA0" " abc", Whitespaces), "\xA0" " abc");
-	//TestEqual(TEXT("Skip to wide"), FAsciiSet::Skip(TEXT(" 变 abc"), Whitespaces), TEXT("变 abc"));
-	TestEqual(TEXT("AdvanceToFirst"), *FAsciiSet::FindFirstOrEnd("NonWhitespace\t \nNonWhitespace", Whitespaces), '\t');
-	TestEqual(TEXT("AdvanceToLast"), *FAsciiSet::FindLastOrEnd("NonWhitespace\t \nNonWhitespace", Whitespaces), '\n');
-	TestEqual(TEXT("AdvanceToLast"), *FAsciiSet::FindLastOrEnd("NonWhitespace\t NonWhitespace\n", Whitespaces), '\n');
-	TestEqual(TEXT("AdvanceToFirst"), *FAsciiSet::FindFirstOrEnd("NonWhitespaceNonWhitespace", Whitespaces), '\0');
-	TestEqual(TEXT("AdvanceToLast"), *FAsciiSet::FindLastOrEnd("NonWhitespaceNonWhitespace", Whitespaces), '\0');
+	//TEST_EQUAL(TEXT("Skip"), FAsciiSet::Skip(TEXT("  \t\tHello world!"), Whitespaces), TEXT("Hello world!"));
+	//TEST_EQUAL(TEXT("Skip"), FAsciiSet::Skip(TEXT("Hello world!"), Whitespaces), TEXT("Hello world!"));
+	//TEST_EQUAL(TEXT("Skip to extended ASCII"), FAsciiSet::Skip(" " "\xA0" " abc", Whitespaces), "\xA0" " abc");
+	//TEST_EQUAL(TEXT("Skip to wide"), FAsciiSet::Skip(TEXT(" 变 abc"), Whitespaces), TEXT("变 abc"));
+	TEST_EQUAL(TEXT("AdvanceToFirst"), *FAsciiSet::FindFirstOrEnd("NonWhitespace\t \nNonWhitespace", Whitespaces), '\t');
+	TEST_EQUAL(TEXT("AdvanceToLast"), *FAsciiSet::FindLastOrEnd("NonWhitespace\t \nNonWhitespace", Whitespaces), '\n');
+	TEST_EQUAL(TEXT("AdvanceToLast"), *FAsciiSet::FindLastOrEnd("NonWhitespace\t NonWhitespace\n", Whitespaces), '\n');
+	TEST_EQUAL(TEXT("AdvanceToFirst"), *FAsciiSet::FindFirstOrEnd("NonWhitespaceNonWhitespace", Whitespaces), '\0');
+	TEST_EQUAL(TEXT("AdvanceToLast"), *FAsciiSet::FindLastOrEnd("NonWhitespaceNonWhitespace", Whitespaces), '\0');
 
 	constexpr FAsciiSet Lowercase("abcdefghijklmnopqrstuvwxyz");
-	TestEqual(TEXT("TrimPrefixWithout"), FAsciiSet::TrimPrefixWithout("ABcdEF"_ASV, Lowercase), "cdEF"_ASV);
-	TestEqual(TEXT("FindPrefixWithout"), FAsciiSet::FindPrefixWithout("ABcdEF"_ASV, Lowercase), "AB"_ASV);
-	TestEqual(TEXT("TrimSuffixWithout"), FAsciiSet::TrimSuffixWithout("ABcdEF"_ASV, Lowercase), "ABcd"_ASV);
-	TestEqual(TEXT("FindSuffixWithout"), FAsciiSet::FindSuffixWithout("ABcdEF"_ASV, Lowercase), "EF"_ASV);
-	TestEqual(TEXT("TrimPrefixWithout none"), FAsciiSet::TrimPrefixWithout("same"_ASV, Lowercase), "same"_ASV);
-	TestEqual(TEXT("FindPrefixWithout none"), FAsciiSet::FindPrefixWithout("same"_ASV, Lowercase), ""_ASV);
-	TestEqual(TEXT("TrimSuffixWithout none"), FAsciiSet::TrimSuffixWithout("same"_ASV, Lowercase), "same"_ASV);
-	TestEqual(TEXT("FindSuffixWithout none"), FAsciiSet::FindSuffixWithout("same"_ASV, Lowercase), ""_ASV);
-	TestEqual(TEXT("TrimPrefixWithout empty"), FAsciiSet::TrimPrefixWithout(""_ASV, Lowercase), ""_ASV);
-	TestEqual(TEXT("FindPrefixWithout empty"), FAsciiSet::FindPrefixWithout(""_ASV, Lowercase), ""_ASV);
-	TestEqual(TEXT("TrimSuffixWithout empty"), FAsciiSet::TrimSuffixWithout(""_ASV, Lowercase), ""_ASV);
-	TestEqual(TEXT("FindSuffixWithout empty"), FAsciiSet::FindSuffixWithout(""_ASV, Lowercase), ""_ASV);
+	TEST_EQUAL(TEXT("TrimPrefixWithout"), FAsciiSet::TrimPrefixWithout("ABcdEF"_ASV, Lowercase), "cdEF"_ASV);
+	TEST_EQUAL(TEXT("FindPrefixWithout"), FAsciiSet::FindPrefixWithout("ABcdEF"_ASV, Lowercase), "AB"_ASV);
+	TEST_EQUAL(TEXT("TrimSuffixWithout"), FAsciiSet::TrimSuffixWithout("ABcdEF"_ASV, Lowercase), "ABcd"_ASV);
+	TEST_EQUAL(TEXT("FindSuffixWithout"), FAsciiSet::FindSuffixWithout("ABcdEF"_ASV, Lowercase), "EF"_ASV);
+	TEST_EQUAL(TEXT("TrimPrefixWithout none"), FAsciiSet::TrimPrefixWithout("same"_ASV, Lowercase), "same"_ASV);
+	TEST_EQUAL(TEXT("FindPrefixWithout none"), FAsciiSet::FindPrefixWithout("same"_ASV, Lowercase), ""_ASV);
+	TEST_EQUAL(TEXT("TrimSuffixWithout none"), FAsciiSet::TrimSuffixWithout("same"_ASV, Lowercase), "same"_ASV);
+	TEST_EQUAL(TEXT("FindSuffixWithout none"), FAsciiSet::FindSuffixWithout("same"_ASV, Lowercase), ""_ASV);
+	TEST_EQUAL(TEXT("TrimPrefixWithout empty"), FAsciiSet::TrimPrefixWithout(""_ASV, Lowercase), ""_ASV);
+	TEST_EQUAL(TEXT("FindPrefixWithout empty"), FAsciiSet::FindPrefixWithout(""_ASV, Lowercase), ""_ASV);
+	TEST_EQUAL(TEXT("TrimSuffixWithout empty"), FAsciiSet::TrimSuffixWithout(""_ASV, Lowercase), ""_ASV);
+	TEST_EQUAL(TEXT("FindSuffixWithout empty"), FAsciiSet::FindSuffixWithout(""_ASV, Lowercase), ""_ASV);
 
 
 	auto TestHasFunctions = [&](auto MakeString)
 	{
 		constexpr FAsciiSet XmlEscapeChars("&<>\"'");
-		TestTrue(TEXT("None"), FAsciiSet::HasNone(MakeString("No escape chars"), XmlEscapeChars));
-		TestFalse(TEXT("Any"), FAsciiSet::HasAny(MakeString("No escape chars"), XmlEscapeChars));
-		TestFalse(TEXT("Only"), FAsciiSet::HasOnly(MakeString("No escape chars"), XmlEscapeChars));
+		TEST_TRUE(TEXT("None"), FAsciiSet::HasNone(MakeString("No escape chars"), XmlEscapeChars));
+		TEST_FALSE(TEXT("Any"), FAsciiSet::HasAny(MakeString("No escape chars"), XmlEscapeChars));
+		TEST_FALSE(TEXT("Only"), FAsciiSet::HasOnly(MakeString("No escape chars"), XmlEscapeChars));
 
-		TestTrue(TEXT("None"), FAsciiSet::HasNone(MakeString(""), XmlEscapeChars));
-		TestFalse(TEXT("Any"), FAsciiSet::HasAny(MakeString(""), XmlEscapeChars));
-		TestTrue(TEXT("Only"), FAsciiSet::HasOnly(MakeString(""), XmlEscapeChars));
+		TEST_TRUE(TEXT("None"), FAsciiSet::HasNone(MakeString(""), XmlEscapeChars));
+		TEST_FALSE(TEXT("Any"), FAsciiSet::HasAny(MakeString(""), XmlEscapeChars));
+		TEST_TRUE(TEXT("Only"), FAsciiSet::HasOnly(MakeString(""), XmlEscapeChars));
 
-		TestFalse(TEXT("None"), FAsciiSet::HasNone(MakeString("&<>\"'"), XmlEscapeChars));
-		TestTrue(TEXT("Any"), FAsciiSet::HasAny(MakeString("&<>\"'"), XmlEscapeChars));
-		TestTrue(TEXT("Only"), FAsciiSet::HasOnly(MakeString("&<>\"'"), XmlEscapeChars));
+		TEST_FALSE(TEXT("None"), FAsciiSet::HasNone(MakeString("&<>\"'"), XmlEscapeChars));
+		TEST_TRUE(TEXT("Any"), FAsciiSet::HasAny(MakeString("&<>\"'"), XmlEscapeChars));
+		TEST_TRUE(TEXT("Only"), FAsciiSet::HasOnly(MakeString("&<>\"'"), XmlEscapeChars));
 
-		TestFalse(TEXT("None"), FAsciiSet::HasNone(MakeString("&<>\"' and more"), XmlEscapeChars));
-		TestTrue(TEXT("Any"), FAsciiSet::HasAny(MakeString("&<>\"' and more"), XmlEscapeChars));
-		TestFalse(TEXT("Only"), FAsciiSet::HasOnly(MakeString("&<>\"' and more"), XmlEscapeChars));
+		TEST_FALSE(TEXT("None"), FAsciiSet::HasNone(MakeString("&<>\"' and more"), XmlEscapeChars));
+		TEST_TRUE(TEXT("Any"), FAsciiSet::HasAny(MakeString("&<>\"' and more"), XmlEscapeChars));
+		TEST_FALSE(TEXT("Only"), FAsciiSet::HasOnly(MakeString("&<>\"' and more"), XmlEscapeChars));
 	};
 	TestHasFunctions([](const char* Str) { return Str; });
 	TestHasFunctions([](const char* Str) { return FAnsiStringView(Str); });
@@ -84,9 +84,9 @@ TEST_CASE("Core::Misc::FAsciiSet::AsciiSet", "[Core][Misc][Smoke]")
 
 	constexpr FAsciiSet Abc("abc");
 	constexpr FAsciiSet Abcd = Abc + 'd';
-	TestTrue(TEXT("Add"), Abcd.Contains('a'));
-	TestTrue(TEXT("Add"), Abcd.Contains('b'));
-	TestTrue(TEXT("Add"), Abcd.Contains('c'));
-	TestTrue(TEXT("Add"), Abcd.Contains('d'));
-	TestFalse(TEXT("Add"), Abcd.Contains('e'));
+	TEST_TRUE(TEXT("Add"), Abcd.Contains('a'));
+	TEST_TRUE(TEXT("Add"), Abcd.Contains('b'));
+	TEST_TRUE(TEXT("Add"), Abcd.Contains('c'));
+	TEST_TRUE(TEXT("Add"), Abcd.Contains('d'));
+	TEST_FALSE(TEXT("Add"), Abcd.Contains('e'));
 }
\ No newline at end of file
diff --git a/Engine/Source/Runtime/Core/Tests/Misc/ByteSwapTest.cpp b/Engine/Source/Runtime/Core/Tests/Misc/ByteSwapTest.cpp
--- a/Engine/Source/Runtime/Core/Tests/Misc/ByteSwapTest.cpp
+++ b/Engine/Source/Runtime/Core/Tests/Misc/ByteSwapTest.cpp
@@ -8,7 +8,7 @@
  * Test byte swapping algorithms.
  */
 
-TEST_CASE("Core::Misc::ByteSwap::Smoke Test", "[Core][Misc][Smoke]")
+TEST_CASE_METHOD(FAutomationTestFixture, "Core::Misc::ByteSwap::Smoke Test", "[Core][Misc][Smoke]")
 {
 	int16  ValS16 = static_cast<int16>(0x1122);
 	uint16 ValU16 = static_cast<uint16>(0x1122);
@@ -22,15 +22,15 @@ TEST_CASE("Core::Misc::ByteSwap::Smoke Test", "[Core][Misc][Smoke]")
 	uint64 ExpectedD = ByteSwap(ValU64);
 	char16_t ValCh16 = static_cast<char16_t>(0x2233);
 
-	TestTrue("Swapping singned int16 value",  ByteSwap(ValS16) == 0x2211);
-	TestTrue("Swapping unsigned int16 value", ByteSwap(ValU16) == 0x2211);
-	TestTrue("Swapping singned int32 value",  ByteSwap(ValS32) == 0x2211EEFF);
-	TestTrue("Swapping unsigned int32 value", ByteSwap(ValU32) == 0x2211EEFF);
-	TestTrue("Swapping singned int64 value",  ByteSwap(ValS64) == 0x44332211CCDDEEFFll);
-	TestTrue("Swapping unsigned int64 value", ByteSwap(ValU64) == 0x44332211CCDDEEFFull);
-	TestTrue("Swapping float value",          ByteSwap(ValF)   == *reinterpret_cast<const float*>(&ExpectedF));
-	TestTrue("Swapping double value",         ByteSwap(ValD)   == *reinterpret_cast<const double*>(&ExpectedD));
-	TestTrue("Swapping char16_t value",       ByteSwap(ValCh16)== 0x3322);
+	TEST_TRUE("Swapping singned int16 value",  ByteSwap(ValS16) == 0x2211);
+	TEST_TRUE("Swapping unsigned int16 value", ByteSwap(ValU16) == 0x2211);
+	TEST_TRUE("Swapping singned int32 value",  ByteSwap(ValS32) == 0x2211EEFF);
+	TEST_TRUE("Swapping unsigned int32 value", ByteSwap(ValU32) == 0x2211EEFF);
+	TEST_TRUE("Swapping singned int64 value",  ByteSwap(ValS64) == 0x44332211CCDDEEFFll);
+	TEST_TRUE("Swapping unsigned int64 value", ByteSwap(ValU64) == 0x44332211CCDDEEFFull);
+	TEST_TRUE("Swapping float value",          ByteSwap(ValF)   == *reinterpret_cast<const float*>(&ExpectedF));
+	TEST_TRUE("Swapping double value",         ByteSwap(ValD)   == *reinterpret_cast<const double*>(&ExpectedD));
+	TEST_TRUE("Swapping char16_t value",       ByteSwap(ValCh16)== 0x3322);
 } 
 
 // The byte swap benchmarking test are useful to compare the intrinsic implementation vs the generic implementation. Normally, the intrinsic is expected to be faster, but in some
@@ -56,15 +56,15 @@ TPair<FTimespan, T> BenchmarkByteswapping(T InitialValue, uint64 LoopCount, Byte
 	return MakeTuple(Duration, Sum);
 }
 
-TEST_CASE("Core::Misc::ByteSwap::Perf Test", "[.][Core][Misc][Perf]")
+TEST_CASE_METHOD(FAutomationTestFixture, "Core::Misc::ByteSwap::Perf Test", "[.][Core][Misc][.Perf]")
 {
 #ifdef UE_BYTESWAP_INTRINSIC_PRIVATE_16
 	{
 		uint16 InitialValue = 0xF0F0;
 		TPair<FTimespan, uint16> Intrinsic = BenchmarkByteswapping(InitialValue, 1000000000, [](uint16 Val){ return ByteSwap(Val); });
 		TPair<FTimespan, uint16> Generic   = BenchmarkByteswapping(InitialValue, 1000000000, [](uint16 Val){ return Internal::ByteSwapGeneric16(Val); });
-		TestTrue("Swapping uint16 bytes is faster uing the compiler intrinsic than the generic implementation", Intrinsic.Get<0>() <= Generic.Get<0>());
-		TestTrue("Swapping uint16 bytes using intrinsic and generic algorithm produce the same values", Intrinsic.Get<1>() == Generic.Get<1>());
+		TEST_TRUE("Swapping uint16 bytes is faster uing the compiler intrinsic than the generic implementation", Intrinsic.Get<0>() <= Generic.Get<0>());
+		TEST_TRUE("Swapping uint16 bytes using intrinsic and generic algorithm produce the same values", Intrinsic.Get<1>() == Generic.Get<1>());
 		GLog->Logf(TEXT("Swapping 2 bytes using intrinsic is %s faster than generic version"), *LexToString(Generic.Get<0>().GetTotalMicroseconds() / Intrinsic.Get<0>().GetTotalMicroseconds()));
 	}
 #endif
@@ -74,8 +74,8 @@ TEST_CASE("Core::Misc::ByteSwap::Perf Test", "[.][Core][Misc][Perf]")
 		uint32 InitialValue = 0xFF00FF00ul;
 		TPair<FTimespan, uint32> Intrinsic = BenchmarkByteswapping(InitialValue, 1000000000, [](uint32 Val){ return ByteSwap(Val); });
 		TPair<FTimespan, uint32> Generic   = BenchmarkByteswapping(InitialValue, 1000000000, [](uint32 Val){ return Internal::ByteSwapGeneric32(Val); });
-		TestTrue("Swapping uint32 bytes is faster uing the compiler intrinsic than the generic implementation", Intrinsic.Get<0>() <= Generic.Get<0>());
-		TestTrue("Swapping uint32 bytes using intrinsic and generic algorithm produce the same values", Intrinsic.Get<1>() == Generic.Get<1>());
+		TEST_TRUE("Swapping uint32 bytes is faster uing the compiler intrinsic than the generic implementation", Intrinsic.Get<0>() <= Generic.Get<0>());
+		TEST_TRUE("Swapping uint32 bytes using intrinsic and generic algorithm produce the same values", Intrinsic.Get<1>() == Generic.Get<1>());
 		GLog->Logf(TEXT("Swapping 4 bytes using intrinsic is %s faster than generic version"), *LexToString(Generic.Get<0>().GetTotalMicroseconds() / Intrinsic.Get<0>().GetTotalMicroseconds()));
 	}
 #endif
@@ -85,8 +85,8 @@ TEST_CASE("Core::Misc::ByteSwap::Perf Test", "[.][Core][Misc][Perf]")
 		uint64 InitialValue = 0xFF00FF00FF00FF00ull;
 		TPair<FTimespan, uint64> Intrinsic = BenchmarkByteswapping(InitialValue, 1000000000, [](uint64 Val){ return ByteSwap(Val); });
 		TPair<FTimespan, uint64> Generic   = BenchmarkByteswapping(InitialValue, 1000000000, [](uint64 Val){ return Internal::ByteSwapGeneric64(Val); });
-		TestTrue("Swapping uint64 bytes is faster uing the compiler intrinsic than the generic implementation", Intrinsic.Get<0>() <= Generic.Get<0>());
-		TestTrue("Swapping uint64 bytes using intrinsic and generic algorithm produce the same values", Intrinsic.Get<1>() == Generic.Get<1>());
+		TEST_TRUE("Swapping uint64 bytes is faster uing the compiler intrinsic than the generic implementation", Intrinsic.Get<0>() <= Generic.Get<0>());
+		TEST_TRUE("Swapping uint64 bytes using intrinsic and generic algorithm produce the same values", Intrinsic.Get<1>() == Generic.Get<1>());
 		GLog->Logf(TEXT("Swapping 8 bytes using intrinsic is %s faster than generic version"), *LexToString(Generic.Get<0>().GetTotalMicroseconds() / Intrinsic.Get<0>().GetTotalMicroseconds()));
 	}
 #endif
diff --git a/Engine/Source/Runtime/Core/Tests/Misc/CharTest.cpp b/Engine/Source/Runtime/Core/Tests/Misc/CharTest.cpp
--- a/Engine/Source/Runtime/Core/Tests/Misc/CharTest.cpp
+++ b/Engine/Source/Runtime/Core/Tests/Misc/CharTest.cpp
@@ -16,30 +16,30 @@ namespace crt
 }
 
 template<typename CharType>
-void RunCharTests(uint32 MaxChar)
+void RunCharTests(FAutomationTestFixture& Test, uint32 MaxChar)
 {
 	for (uint32 I = 0; I < MaxChar; ++I)
 	{
 		CharType C = (CharType)I;
-		TestEqual("TChar::ToLower()", TChar<CharType>::ToLower(C), crt::tolower(C));
-		TestEqual("TChar::ToUpper()", TChar<CharType>::ToUpper(C), crt::toupper(C));
+		TEST_EQUAL("TChar::ToLower()", TChar<CharType>::ToLower(C), crt::tolower(C));
+		TEST_EQUAL("TChar::ToUpper()", TChar<CharType>::ToUpper(C), crt::toupper(C));
 	}
 }
 
-TEST_CASE("Core::Misc::Char::Smoke Test", "[Core][Misc][Perf]")
+TEST_CASE_METHOD(FAutomationTestFixture, "Core::Misc::Char::Smoke Test", "[Core][Misc][.Perf]")
 {
 	const char* CurrentLocale = setlocale(LC_CTYPE, nullptr);
 	if (CurrentLocale == nullptr)
 	{
-		TestFalse(FString::Printf(TEXT("Locale is null but should be \"C\". Did something call setlocale()?")), CurrentLocale == nullptr);
+		TEST_FALSE(FString::Printf(TEXT("Locale is null but should be \"C\". Did something call setlocale()?")), CurrentLocale == nullptr);
 	}
 	else if (strcmp("C", CurrentLocale))
 	{
-		TestFalse(FString::Printf(TEXT("Locale is null but should be \"C\". Did something call setlocale()?")), strcmp("C", CurrentLocale)!=0);
+		TEST_FALSE(FString::Printf(TEXT("Locale is null but should be \"C\". Did something call setlocale()?")), strcmp("C", CurrentLocale)!=0);
 	}
 	else
 	{
-		RunCharTests<ANSICHAR>(128);
-		RunCharTests<WIDECHAR>(65536);
+		RunCharTests<ANSICHAR>(*this,128);
+		RunCharTests<WIDECHAR>(*this,65536);
 	}
 }
\ No newline at end of file
diff --git a/Engine/Source/Runtime/Core/Tests/Misc/CircularBufferTest.cpp b/Engine/Source/Runtime/Core/Tests/Misc/CircularBufferTest.cpp
--- a/Engine/Source/Runtime/Core/Tests/Misc/CircularBufferTest.cpp
+++ b/Engine/Source/Runtime/Core/Tests/Misc/CircularBufferTest.cpp
@@ -5,37 +5,37 @@
 #include "Containers/CircularBuffer.h"
 #include "TestHarness.h"
 
-TEST_CASE("Core::Misc::TCircularBuffer::Smoke Test", "[Core][Misc][Smoke]")
+TEST_CASE_METHOD(FAutomationTestFixture, "Core::Misc::TCircularBuffer::Smoke Test", "[Core][Misc][Smoke]")
 {
 	// buffer capacity
 	TCircularBuffer<int32> b1_1(127);
 	TCircularBuffer<int32> b1_2(128);
 	TCircularBuffer<int32> b1_3(129);
 
-	TestEqual(TEXT("Buffer capacity of 127 must be rounded up to 128"), b1_1.Capacity(), 128u);
-	TestEqual(TEXT("Buffer capacity of 128 must not change"), b1_2.Capacity(), 128u);
-	TestEqual(TEXT("Buffer capacity of 129 must be rounded up to 256"), b1_3.Capacity(), 256u);
+	TEST_EQUAL(TEXT("Buffer capacity of 127 must be rounded up to 128"), b1_1.Capacity(), 128u);
+	TEST_EQUAL(TEXT("Buffer capacity of 128 must not change"), b1_2.Capacity(), 128u);
+	TEST_EQUAL(TEXT("Buffer capacity of 129 must be rounded up to 256"), b1_3.Capacity(), 256u);
 
 	// initial values
 	TCircularBuffer<int32> b2_1(64, 666);
 
 	for (uint32 Index = 0; Index < b2_1.Capacity(); ++Index)
 	{
-		TestEqual(FString::Printf(TEXT("Initial value must be correct (%i)"), Index), b2_1[Index], 666);
+		TEST_EQUAL(FString::Printf(TEXT("Initial value must be correct (%i)"), Index), b2_1[Index], 666);
 	}
 
 	// indexing
 	TCircularBuffer<int32> b3_1(64, 0);
 
-	TestEqual(TEXT("Next index from 0 must be 1"), b3_1.GetNextIndex(0), 1u);
-	TestEqual(TEXT("Next index from 63 must be 0"), b3_1.GetNextIndex(63), 0u);
-	TestEqual(TEXT("Next index from 64 must be 1"), b3_1.GetNextIndex(64), 1u);
+	TEST_EQUAL(TEXT("Next index from 0 must be 1"), b3_1.GetNextIndex(0), 1u);
+	TEST_EQUAL(TEXT("Next index from 63 must be 0"), b3_1.GetNextIndex(63), 0u);
+	TEST_EQUAL(TEXT("Next index from 64 must be 1"), b3_1.GetNextIndex(64), 1u);
 
 	b3_1[0] = 42;
 	b3_1[65] = 42;
 
-	TestEqual(TEXT("Index 0 must be written and read correctly"), b3_1[0], 42);
-	TestEqual(TEXT("Index 1 must be written and read correctly"), b3_1[1], 42);
-	TestEqual(TEXT("Index 65 must be written and read correctly"), b3_1[65], 42);
+	TEST_EQUAL(TEXT("Index 0 must be written and read correctly"), b3_1[0], 42);
+	TEST_EQUAL(TEXT("Index 1 must be written and read correctly"), b3_1[1], 42);
+	TEST_EQUAL(TEXT("Index 65 must be written and read correctly"), b3_1[65], 42);
 }
 
diff --git a/Engine/Source/Runtime/Core/Tests/Misc/CircularQueueTest.cpp b/Engine/Source/Runtime/Core/Tests/Misc/CircularQueueTest.cpp
--- a/Engine/Source/Runtime/Core/Tests/Misc/CircularQueueTest.cpp
+++ b/Engine/Source/Runtime/Core/Tests/Misc/CircularQueueTest.cpp
@@ -4,7 +4,7 @@
 #include "Containers/CircularQueue.h"
 #include "TestHarness.h"
 
-TEST_CASE("Core::Misc::TCircularQueue::Smoke Test", "[Core][Misc][Smoke]")
+TEST_CASE_METHOD(FAutomationTestFixture, "Core::Misc::TCircularQueue::Smoke Test", "[Core][Misc][Smoke]")
 {
 	const uint32 QueueSize = 8;
 
@@ -12,13 +12,13 @@ TEST_CASE("Core::Misc::TCircularQueue::Smoke Test", "[Core][Misc][Smoke]")
 	{
 		TCircularQueue<int32> Queue(QueueSize);
 
-		TestEqual(TEXT("Newly created queues must have zero elements"), Queue.Count(), 0u);
-		TestTrue(TEXT("Newly created queues must be empty"), Queue.IsEmpty());
-		TestFalse(TEXT("Newly created queues must not be full"), Queue.IsFull());
+		TEST_EQUAL(TEXT("Newly created queues must have zero elements"), Queue.Count(), 0u);
+		TEST_TRUE(TEXT("Newly created queues must be empty"), Queue.IsEmpty());
+		TEST_FALSE(TEXT("Newly created queues must not be full"), Queue.IsFull());
 
 		int32 Value;
-		TestFalse(TEXT("Peek must fail on an empty queue"), Queue.Peek(Value));
-		TestNull(TEXT("Peek must fail on an empty queue"), Queue.Peek());
+		TEST_FALSE(TEXT("Peek must fail on an empty queue"), Queue.Peek(Value));
+		TEST_NULL(TEXT("Peek must fail on an empty queue"), Queue.Peek());
 	}
 
 	// partially filled
@@ -26,16 +26,16 @@ TEST_CASE("Core::Misc::TCircularQueue::Smoke Test", "[Core][Misc][Smoke]")
 		TCircularQueue<int32> Queue(QueueSize);
 		int32 Value = 0;
 
-		TestTrue(TEXT("Adding to an empty queue must succeed"), Queue.Enqueue(666));
-		TestEqual(TEXT("After adding to an empty queue it must have one element"), Queue.Count(), 1u);
-		TestFalse(TEXT("Partially filled queues must not be empty"), Queue.IsEmpty());
-		TestFalse(TEXT("Partially filled queues must not be full"), Queue.IsFull());
-		TestTrue(TEXT("Peeking at a partially filled queue must succeed"), Queue.Peek(Value));
-		TestEqual(TEXT("The peeked at value must be correct"), Value, 666);
+		TEST_TRUE(TEXT("Adding to an empty queue must succeed"), Queue.Enqueue(666));
+		TEST_EQUAL(TEXT("After adding to an empty queue it must have one element"), Queue.Count(), 1u);
+		TEST_FALSE(TEXT("Partially filled queues must not be empty"), Queue.IsEmpty());
+		TEST_FALSE(TEXT("Partially filled queues must not be full"), Queue.IsFull());
+		TEST_TRUE(TEXT("Peeking at a partially filled queue must succeed"), Queue.Peek(Value));
+		TEST_EQUAL(TEXT("The peeked at value must be correct"), Value, 666);
 
 		const int32* PeekValue = Queue.Peek();
-		TestNotNull(TEXT("Peeking at a partially filled queue must succeed"), PeekValue);
-		TestEqual(TEXT("The peeked at value must be correct"), *PeekValue, 666);
+		TEST_NOT_NULL(TEXT("Peeking at a partially filled queue must succeed"), PeekValue);
+		TEST_EQUAL(TEXT("The peeked at value must be correct"), *PeekValue, 666);
 	}
 
 	// full queue
@@ -45,12 +45,12 @@ TEST_CASE("Core::Misc::TCircularQueue::Smoke Test", "[Core][Misc][Smoke]")
 
 		for (int32 Index = 0; Index < QueueSize - 1; ++Index)
 		{
-			TestTrue(TEXT("Adding to non-full queue must succeed"), Queue.Enqueue(Index));
+			TEST_TRUE(TEXT("Adding to non-full queue must succeed"), Queue.Enqueue(Index));
 		}
 
-		TestFalse(TEXT("Full queues must not be empty"), Queue.IsEmpty());
-		TestTrue(TEXT("Full queues must be full"), Queue.IsFull());
-		TestFalse(TEXT("Adding to full queue must fail"), Queue.Enqueue(666));
+		TEST_FALSE(TEXT("Full queues must not be empty"), Queue.IsEmpty());
+		TEST_TRUE(TEXT("Full queues must be full"), Queue.IsFull());
+		TEST_FALSE(TEXT("Adding to full queue must fail"), Queue.Enqueue(666));
 
 		int32 Value = 0;
 
@@ -58,24 +58,24 @@ TEST_CASE("Core::Misc::TCircularQueue::Smoke Test", "[Core][Misc][Smoke]")
 		{
 			if (PeekType == 0)
 			{
-				TestTrue(TEXT("Peeking at a non-empty queue must succeed"), Queue.Peek(Value));
-				TestEqual(TEXT("The peeked at value must be correct"), Value, Index);
+				TEST_TRUE(TEXT("Peeking at a non-empty queue must succeed"), Queue.Peek(Value));
+				TEST_EQUAL(TEXT("The peeked at value must be correct"), Value, Index);
 
-				TestTrue(TEXT("Removing from a non-empty queue must succeed"), Queue.Dequeue(Value));
-				TestEqual(TEXT("The removed value must be correct"), Value, Index);
+				TEST_TRUE(TEXT("Removing from a non-empty queue must succeed"), Queue.Dequeue(Value));
+				TEST_EQUAL(TEXT("The removed value must be correct"), Value, Index);
 			}
 			else
 			{
 				const int32* PeekValue = Queue.Peek();
-				TestNotNull(TEXT("Peeking at a non-empty queue must succeed"), PeekValue);
-				TestEqual(TEXT("The peeked at value must be correct"), *PeekValue, Index);
+				TEST_NOT_NULL(TEXT("Peeking at a non-empty queue must succeed"), PeekValue);
+				TEST_EQUAL(TEXT("The peeked at value must be correct"), *PeekValue, Index);
 
-				TestTrue(TEXT("Removing from a non-empty queue must succeed"), Queue.Dequeue());
+				TEST_TRUE(TEXT("Removing from a non-empty queue must succeed"), Queue.Dequeue());
 			}
 		}
 
-		TestTrue(TEXT("A queue that had all items removed must be empty"), Queue.IsEmpty());
-		TestFalse(TEXT("A queue that had all items removed must not be full"), Queue.IsFull());
+		TEST_TRUE(TEXT("A queue that had all items removed must be empty"), Queue.IsEmpty());
+		TEST_FALSE(TEXT("A queue that had all items removed must not be full"), Queue.IsFull());
 	}
 
 	// queue with index wrapping around
@@ -85,16 +85,16 @@ TEST_CASE("Core::Misc::TCircularQueue::Smoke Test", "[Core][Misc][Smoke]")
 		//Fill queue
 		for (int32 Index = 0; Index < QueueSize - 1; ++Index)
 		{
-			TestTrue(TEXT("Adding to non-full queue must succeed"), Queue.Enqueue(Index));
+			TEST_TRUE(TEXT("Adding to non-full queue must succeed"), Queue.Enqueue(Index));
 		}
 
 		int32 Value = 0;
 		const int32 ExpectedSize = QueueSize - 1;
 		for (int32 Index = 0; Index < QueueSize; ++Index)
 		{
-			TestEqual(TEXT("Number of elements must be valid for all permutation of Tail and Head"), Queue.Count(), ExpectedSize);
-			TestTrue(TEXT("Removing from a non-empty queue must succeed"), Queue.Dequeue(Value));
-			TestTrue(TEXT("Adding to non-full queue must succeed"), Queue.Enqueue(Index));
+			TEST_EQUAL(TEXT("Number of elements must be valid for all permutation of Tail and Head"), Queue.Count(), ExpectedSize);
+			TEST_TRUE(TEXT("Removing from a non-empty queue must succeed"), Queue.Dequeue(Value));
+			TEST_TRUE(TEXT("Adding to non-full queue must succeed"), Queue.Enqueue(Index));
 		}
 	}
 
@@ -118,12 +118,12 @@ TEST_CASE("Core::Misc::TCircularQueue::Smoke Test", "[Core][Misc][Smoke]")
 		{
 			TCircularQueue<FNonPOD> Queue(QueueSize);
 
-			TestEqual(TEXT("The constructor should not run"), Const, 0);
-			TestEqual(TEXT("The destructor should not run"), Dest, 0);
-			TestEqual(TEXT("The copy constructor should not run"), CopyConst, 0);
+			TEST_EQUAL(TEXT("The constructor should not run"), Const, 0);
+			TEST_EQUAL(TEXT("The destructor should not run"), Dest, 0);
+			TEST_EQUAL(TEXT("The copy constructor should not run"), CopyConst, 0);
 		}
-		TestEqual(TEXT("The constructor should not run"), Const, 0);
-		TestEqual(TEXT("The destructor should run"), Dest, QueueSize);
-		TestEqual(TEXT("The copy constructor should not run"), CopyConst, 0);
+		TEST_EQUAL(TEXT("The constructor should not run"), Const, 0);
+		TEST_EQUAL(TEXT("The destructor should run"), Dest, QueueSize);
+		TEST_EQUAL(TEXT("The copy constructor should not run"), CopyConst, 0);
 	}
 }
\ No newline at end of file
diff --git a/Engine/Source/Runtime/Core/Tests/Misc/DateTimeTest.cpp b/Engine/Source/Runtime/Core/Tests/Misc/DateTimeTest.cpp
--- a/Engine/Source/Runtime/Core/Tests/Misc/DateTimeTest.cpp
+++ b/Engine/Source/Runtime/Core/Tests/Misc/DateTimeTest.cpp
@@ -5,17 +5,17 @@
 #include "Misc/DateTime.h"
 #include "TestHarness.h"
 
-#define TestUnixEquivalent(Desc, A, B) TestTrue(FString::Printf(TEXT("%s - A=%d B=%d"), Desc, (A).ToUnixTimestamp(), (B)), (A).ToUnixTimestamp() == (B) );
-#define TestYear(Desc, A, B)  TestTrue(FString::Printf(TEXT("%s - A=%d B=%d"), Desc, (A.GetYear()), (B)), ((A.GetYear()) == (B)));
-#define TestMonth(Desc, A, B) TestTrue(FString::Printf(TEXT("%s - A=%d B=%d"), Desc, (A.GetMonth()), B), (A.GetMonth()) == (B));
-#define TestMonthOfYear(Desc, A, B) TestTrue(FString::Printf(TEXT("%s - A=%d B=%d"), Desc ,(static_cast<int32>(A.GetMonthOfYear())), static_cast<int32>(B)), (A.GetMonthOfYear()) == (B));
-#define TestDay(Desc, A, B) TestTrue(FString::Printf(TEXT("%s - A=%d B=%d"), Desc, (A.GetDay()), (B)), (A.GetDay()==(B)));
-#define TestHour(Desc, A, B) TestTrue(FString::Printf(TEXT("%s - A=%d B=%d"), Desc, (A.GetHour()), (B)), (A.GetHour()==(B)));
-#define TestMinute(Desc, A, B) TestTrue(FString::Printf(TEXT("%s - A=%d B=%d"), Desc, (A.GetMinute()), (B)), (A.GetMinute()) == (B));
-#define TestSecond(Desc, A, B) TestTrue(FString::Printf(TEXT("%s - A=%d B=%d"), Desc, (A.GetSecond()), (B)), (A.GetSecond()) == (B));
-#define TestMillisecond(Desc, A, B) TestTrue(FString::Printf(TEXT("%s - A=%d B=%d"), Desc, (A.GetMillisecond()) ,(B)), (A.GetMillisecond()) == (B));
-
-TEST_CASE("Core::Misc::DateTime::Smoke", "[Core][Misc][Smoke][PlatformFlaky]")
+#define TestUnixEquivalent(Desc, A, B) TEST_TRUE(FString::Printf(TEXT("%s - A=%d B=%d"), Desc, (A).ToUnixTimestamp(), (B)), (A).ToUnixTimestamp() == (B) );
+#define TestYear(Desc, A, B)  TEST_TRUE(FString::Printf(TEXT("%s - A=%d B=%d"), Desc, (A.GetYear()), (B)), ((A.GetYear()) == (B)));
+#define TestMonth(Desc, A, B) TEST_TRUE(FString::Printf(TEXT("%s - A=%d B=%d"), Desc, (A.GetMonth()), B), (A.GetMonth()) == (B));
+#define TestMonthOfYear(Desc, A, B) TEST_TRUE(FString::Printf(TEXT("%s - A=%d B=%d"), Desc ,(static_cast<int32>(A.GetMonthOfYear())), static_cast<int32>(B)), (A.GetMonthOfYear()) == (B));
+#define TestDay(Desc, A, B) TEST_TRUE(FString::Printf(TEXT("%s - A=%d B=%d"), Desc, (A.GetDay()), (B)), (A.GetDay()==(B)));
+#define TestHour(Desc, A, B) TEST_TRUE(FString::Printf(TEXT("%s - A=%d B=%d"), Desc, (A.GetHour()), (B)), (A.GetHour()==(B)));
+#define TestMinute(Desc, A, B) TEST_TRUE(FString::Printf(TEXT("%s - A=%d B=%d"), Desc, (A.GetMinute()), (B)), (A.GetMinute()) == (B));
+#define TestSecond(Desc, A, B) TEST_TRUE(FString::Printf(TEXT("%s - A=%d B=%d"), Desc, (A.GetSecond()), (B)), (A.GetSecond()) == (B));
+#define TestMillisecond(Desc, A, B) TEST_TRUE(FString::Printf(TEXT("%s - A=%d B=%d"), Desc, (A.GetMillisecond()) ,(B)), (A.GetMillisecond()) == (B));
+
+TEST_CASE_METHOD(FAutomationTestFixture, "Core::Misc::DateTime::Smoke Test", "[Core][Misc][Smoke]")
 {
 	const int64 UnixEpochTimestamp = 0;
 	const int64 UnixBillenniumTimestamp = 1000000000;
@@ -81,47 +81,47 @@ TEST_CASE("Core::Misc::DateTime::Smoke", "[Core][Misc][Smoke][PlatformFlaky]")
 
 	FDateTime ParsedDateTime;
 
-	TestFalse(TEXT("Parsing an empty ISO string must fail"), FDateTime::ParseIso8601(TEXT(""), ParsedDateTime));
+	TEST_FALSE(TEXT("Parsing an empty ISO string must fail"), FDateTime::ParseIso8601(TEXT(""), ParsedDateTime));
 
 	FDateTime::ParseIso8601(TEXT("2019-05-22"), ParsedDateTime);
-	TestEqual(TEXT("Testing ISO 8601 date"), ParsedDateTime, FDateTime(2019, 5, 22));
+	TEST_EQUAL(TEXT("Testing ISO 8601 date"), ParsedDateTime, FDateTime(2019, 5, 22));
 
 	FDateTime::ParseIso8601(TEXT("2019-05-20T19:41:38+01:30"), ParsedDateTime);
-	TestEqual(TEXT("Testing ISO 8601 with +hh:mm timezone info"), ParsedDateTime, FDateTime( 2019, 5, 20, 18, 11, 38 ));
+	TEST_EQUAL(TEXT("Testing ISO 8601 with +hh:mm timezone info"), ParsedDateTime, FDateTime( 2019, 5, 20, 18, 11, 38 ));
 	FDateTime::ParseIso8601(TEXT("2019-05-20T19:41:38-01:30"), ParsedDateTime);
-	TestEqual(TEXT("Testing ISO 8601 with -hh:mm timezone info"), ParsedDateTime, FDateTime( 2019, 5, 20, 21, 11, 38 ));
+	TEST_EQUAL(TEXT("Testing ISO 8601 with -hh:mm timezone info"), ParsedDateTime, FDateTime( 2019, 5, 20, 21, 11, 38 ));
 	FDateTime::ParseIso8601(TEXT("2019-05-20T19:41:38+0030"), ParsedDateTime);
-	TestEqual(TEXT("Testing ISO 8601 with +hhmm timezone info"), ParsedDateTime, FDateTime( 2019, 5, 20, 19, 11, 38 ));
+	TEST_EQUAL(TEXT("Testing ISO 8601 with +hhmm timezone info"), ParsedDateTime, FDateTime( 2019, 5, 20, 19, 11, 38 ));
 	FDateTime::ParseIso8601(TEXT("2019-05-20T19:41:38-01"), ParsedDateTime);
-	TestEqual(TEXT("Testing ISO 8601 with -hh timezone info"), ParsedDateTime, FDateTime( 2019, 5, 20, 20, 41, 38 ));
+	TEST_EQUAL(TEXT("Testing ISO 8601 with -hh timezone info"), ParsedDateTime, FDateTime( 2019, 5, 20, 20, 41, 38 ));
 
 	for (double JulianDay : {0., 1000., 1721425.5, 1721425.0, FDateTime::MinValue().GetJulianDay(), FDateTime::MaxValue().GetJulianDay()})
 	{
-		TestEqual(TEXT("convertion from/to JulianDay is stable"), FDateTime::FromJulianDay(JulianDay).GetJulianDay(), JulianDay);
+		TEST_EQUAL(TEXT("convertion from/to JulianDay is stable"), FDateTime::FromJulianDay(JulianDay).GetJulianDay(), JulianDay);
 	}
 
 	FDateTime Date{ 2019, 5, 20, 18, 11, 38 };
 	FDateTime DateMidnight = Date.GetDate();
-	TestEqual(TEXT("GetDate returns a FDateTime at midnight"), DateMidnight.GetHour(), 0);
+	TEST_EQUAL(TEXT("GetDate returns a FDateTime at midnight"), DateMidnight.GetHour(), 0);
 
 	FDateTime DateNoon = DateMidnight + FTimespan::FromHours(12);
-	TestEqual(TEXT("GetDate returns a FDateTime at midnight"), DateNoon.GetHour(), 12);
+	TEST_EQUAL(TEXT("GetDate returns a FDateTime at midnight"), DateNoon.GetHour(), 12);
 
 	double JulianDayMidnight = DateMidnight.GetJulianDay();
 	double JulianDayNoon = DateNoon.GetJulianDay();
-	TestEqual(TEXT("At midnight, fractionnal part of the JulianDay value should be 0.5"), JulianDayMidnight - (int)JulianDayMidnight, 0.5);
-	TestEqual(TEXT("A 12h timespan adds half a julianday"), JulianDayNoon - JulianDayMidnight, 0.5);
+	TEST_EQUAL(TEXT("At midnight, fractionnal part of the JulianDay value should be 0.5"), JulianDayMidnight - (int)JulianDayMidnight, 0.5);
+	TEST_EQUAL(TEXT("A 12h timespan adds half a julianday"), JulianDayNoon - JulianDayMidnight, 0.5);
 
 	double OffsetDayCount = 12345.;
 	FDateTime OffsetDate = Date + FTimespan::FromDays(OffsetDayCount);
-	TestEqual(TEXT("An offset by a given numer of days leads to a similar JulianDay offset"), OffsetDate.GetJulianDay() - Date.GetJulianDay(), OffsetDayCount);
+	TEST_EQUAL(TEXT("An offset by a given numer of days leads to a similar JulianDay offset"), OffsetDate.GetJulianDay() - Date.GetJulianDay(), OffsetDayCount);
 
 	int32 Year, Month, Day;
 	DateMidnight.GetDate(Year, Month, Day);
 	FDateTime PreviousDay(DateMidnight.GetTicks() - 1);
 	int32 YearPrev, MonthPrev, DayPrev;
 	PreviousDay.GetDate(YearPrev, MonthPrev, DayPrev);
-	TestNotEqual(TEXT("One tick before a date at midnight leads to the previous date"), Day, DayPrev);
+	TEST_NOT_EQUAL(TEXT("One tick before a date at midnight leads to the previous date"), Day, DayPrev);
 }
 
 
diff --git a/Engine/Source/Runtime/Core/Tests/Misc/DisplayNameTest.cpp b/Engine/Source/Runtime/Core/Tests/Misc/DisplayNameTest.cpp
--- a/Engine/Source/Runtime/Core/Tests/Misc/DisplayNameTest.cpp
+++ b/Engine/Source/Runtime/Core/Tests/Misc/DisplayNameTest.cpp
@@ -2,30 +2,30 @@
 #include "UObject/UnrealNames.h"
 #include "TestHarness.h"
 
-TEST_CASE("Core::Misc::FName::Display Name", "[Core][Misc][Smoke]")
+TEST_CASE_METHOD(FAutomationTestFixture, "Core::Misc::FName::Display Name", "[Core][Misc][Smoke]")
 {
-	TestEqual(TEXT("Boolean"), FName::NameToDisplayString(TEXT("bTest"), true), TEXT("Test"));
-	TestEqual(TEXT("Boolean Lower"), FName::NameToDisplayString(TEXT("bTwoWords"), true), TEXT("Two Words"));
-	TestEqual(TEXT("Lower Boolean"), FName::NameToDisplayString(TEXT("boolean"), true), TEXT("Boolean"));
-	TestEqual(TEXT("Almost Boolean"), FName::NameToDisplayString(TEXT("bNotBoolean"), false), TEXT("B Not Boolean"));
-	TestEqual(TEXT("Boolean No Prefix"), FName::NameToDisplayString(TEXT("NonprefixBoolean"), true), TEXT("Nonprefix Boolean"));
-	TestEqual(TEXT("Lower Boolean No Prefix"), FName::NameToDisplayString(TEXT("lowerNonprefixBoolean"), true), TEXT("Lower Nonprefix Boolean"));
-	TestEqual(TEXT("Lower Camel Case"), FName::NameToDisplayString(TEXT("lowerCase"), false), TEXT("Lower Case"));
-	TestEqual(TEXT("With Underscores"), FName::NameToDisplayString(TEXT("With_Underscores"), false), TEXT("With Underscores"));
-	TestEqual(TEXT("Lower Underscores"), FName::NameToDisplayString(TEXT("lower_underscores"), false), TEXT("Lower Underscores"));
-	TestEqual(TEXT("Mixed Underscores"), FName::NameToDisplayString(TEXT("mixed_Underscores"), false), TEXT("Mixed Underscores"));
-	TestEqual(TEXT("Mixed Underscores"), FName::NameToDisplayString(TEXT("Mixed_underscores"), false), TEXT("Mixed Underscores"));
-	TestEqual(TEXT("Article in String"), FName::NameToDisplayString(TEXT("ArticleInString"), false), TEXT("Article in String"));
-	TestEqual(TEXT("One or Two"), FName::NameToDisplayString(TEXT("OneOrTwo"), false), TEXT("One or Two"));
-	TestEqual(TEXT("One and Two"), FName::NameToDisplayString(TEXT("OneAndTwo"), false), TEXT("One and Two"));
-	TestEqual(TEXT("-1.5"), FName::NameToDisplayString(TEXT("-1.5"), false), TEXT("-1.5"));
-	TestEqual(TEXT("1234"), FName::NameToDisplayString(TEXT("1234"), false), TEXT("1234"));
-	TestEqual(TEXT("1234.5"), FName::NameToDisplayString(TEXT("1234.5"), false), TEXT("1234.5"));
-	TestEqual(TEXT("-1234.5"), FName::NameToDisplayString(TEXT("-1234.5"), false), TEXT("-1234.5"));
-	TestEqual(TEXT("Text (In Parens)"), FName::NameToDisplayString(TEXT("Text (in parens)"), false), TEXT("Text (In Parens)"));
+	TEST_EQUAL(TEXT("Boolean"), FName::NameToDisplayString(TEXT("bTest"), true), TEXT("Test"));
+	TEST_EQUAL(TEXT("Boolean Lower"), FName::NameToDisplayString(TEXT("bTwoWords"), true), TEXT("Two Words"));
+	TEST_EQUAL(TEXT("Lower Boolean"), FName::NameToDisplayString(TEXT("boolean"), true), TEXT("Boolean"));
+	TEST_EQUAL(TEXT("Almost Boolean"), FName::NameToDisplayString(TEXT("bNotBoolean"), false), TEXT("B Not Boolean"));
+	TEST_EQUAL(TEXT("Boolean No Prefix"), FName::NameToDisplayString(TEXT("NonprefixBoolean"), true), TEXT("Nonprefix Boolean"));
+	TEST_EQUAL(TEXT("Lower Boolean No Prefix"), FName::NameToDisplayString(TEXT("lowerNonprefixBoolean"), true), TEXT("Lower Nonprefix Boolean"));
+	TEST_EQUAL(TEXT("Lower Camel Case"), FName::NameToDisplayString(TEXT("lowerCase"), false), TEXT("Lower Case"));
+	TEST_EQUAL(TEXT("With Underscores"), FName::NameToDisplayString(TEXT("With_Underscores"), false), TEXT("With Underscores"));
+	TEST_EQUAL(TEXT("Lower Underscores"), FName::NameToDisplayString(TEXT("lower_underscores"), false), TEXT("Lower Underscores"));
+	TEST_EQUAL(TEXT("Mixed Underscores"), FName::NameToDisplayString(TEXT("mixed_Underscores"), false), TEXT("Mixed Underscores"));
+	TEST_EQUAL(TEXT("Mixed Underscores"), FName::NameToDisplayString(TEXT("Mixed_underscores"), false), TEXT("Mixed Underscores"));
+	TEST_EQUAL(TEXT("Article in String"), FName::NameToDisplayString(TEXT("ArticleInString"), false), TEXT("Article in String"));
+	TEST_EQUAL(TEXT("One or Two"), FName::NameToDisplayString(TEXT("OneOrTwo"), false), TEXT("One or Two"));
+	TEST_EQUAL(TEXT("One and Two"), FName::NameToDisplayString(TEXT("OneAndTwo"), false), TEXT("One and Two"));
+	TEST_EQUAL(TEXT("-1.5"), FName::NameToDisplayString(TEXT("-1.5"), false), TEXT("-1.5"));
+	TEST_EQUAL(TEXT("1234"), FName::NameToDisplayString(TEXT("1234"), false), TEXT("1234"));
+	TEST_EQUAL(TEXT("1234.5"), FName::NameToDisplayString(TEXT("1234.5"), false), TEXT("1234.5"));
+	TEST_EQUAL(TEXT("-1234.5"), FName::NameToDisplayString(TEXT("-1234.5"), false), TEXT("-1234.5"));
+	TEST_EQUAL(TEXT("Text (In Parens)"), FName::NameToDisplayString(TEXT("Text (in parens)"), false), TEXT("Text (In Parens)"));
 
-	TestEqual(TEXT("Text 3D"), FName::NameToDisplayString(TEXT("Text3D"), false), TEXT("Text 3D"));
-	TestEqual(TEXT("Plural CAPs"), FName::NameToDisplayString(TEXT("PluralCAPs"), false), TEXT("Plural CAPs"));
-	TestEqual(TEXT("FBXEditor"), FName::NameToDisplayString(TEXT("FBXEditor"), false), TEXT("FBXEditor"));
-	TestEqual(TEXT("FBX_Editor"), FName::NameToDisplayString(TEXT("FBX_Editor"), false), TEXT("FBX Editor"));
+	TEST_EQUAL(TEXT("Text 3D"), FName::NameToDisplayString(TEXT("Text3D"), false), TEXT("Text 3D"));
+	TEST_EQUAL(TEXT("Plural CAPs"), FName::NameToDisplayString(TEXT("PluralCAPs"), false), TEXT("Plural CAPs"));
+	TEST_EQUAL(TEXT("FBXEditor"), FName::NameToDisplayString(TEXT("FBXEditor"), false), TEXT("FBXEditor"));
+	TEST_EQUAL(TEXT("FBX_Editor"), FName::NameToDisplayString(TEXT("FBX_Editor"), false), TEXT("FBX Editor"));
 }
\ No newline at end of file
