Add new Unit Tests

diff --git a/Engine/Source/Runtime/Core/Tests/Misc/TripleBufferTest.cpp b/Engine/Source/Runtime/Core/Tests/Misc/TripleBufferTest.cpp
--- a/Engine/Source/Runtime/Core/Tests/Misc/TripleBufferTest.cpp
+++ b/Engine/Source/Runtime/Core/Tests/Misc/TripleBufferTest.cpp
@@ -6,34 +6,34 @@
 #include "Math/RandomStream.h"
 #include "TestHarness.h"
 
-TEST_CASE("Core::Misc::TTripleBuffer::Smoke Test", "[Core][Misc][Smoke]")
+TEST_CASE_METHOD(FAutomationTestFixture, "Core::Misc::TTripleBuffer::Smoke Test", "[Core][Misc][Smoke]")
 {
 	// uninitialized buffer
 	{
 		TTripleBuffer<int32> Buffer(NoInit);
 
-		TestFalse(TEXT("Uninitialized triple buffer must not be dirty"), Buffer.IsDirty());
+		TEST_FALSE(TEXT("Uninitialized triple buffer must not be dirty"), Buffer.IsDirty());
 	}
 
 	// initialized buffer
 	{
 		TTripleBuffer<int32> Buffer(1);
 
-		TestFalse(TEXT("Initialized triple buffer must not be dirty"), Buffer.IsDirty());
-		TestEqual(TEXT("Initialized triple buffer must have correct read buffer value"), Buffer.Read(), 1);
+		TEST_FALSE(TEXT("Initialized triple buffer must not be dirty"), Buffer.IsDirty());
+		TEST_EQUAL(TEXT("Initialized triple buffer must have correct read buffer value"), Buffer.Read(), 1);
 
 		Buffer.SwapReadBuffers();
 
-		TestEqual(TEXT("Initialized triple buffer must have correct temp buffer value"), Buffer.Read(), 1);
+		TEST_EQUAL(TEXT("Initialized triple buffer must have correct temp buffer value"), Buffer.Read(), 1);
 
 		Buffer.SwapWriteBuffers();
 
-		TestTrue(TEXT("Write buffer swap must set dirty flag"), Buffer.IsDirty());
+		TEST_TRUE(TEXT("Write buffer swap must set dirty flag"), Buffer.IsDirty());
 
 		Buffer.SwapReadBuffers();
 
-		TestFalse(TEXT("Read buffer swap must clear dirty flag"), Buffer.IsDirty());
-		TestEqual(TEXT("Initialized triple buffer must have correct temp buffer value"), Buffer.Read(), 1);
+		TEST_FALSE(TEXT("Read buffer swap must clear dirty flag"), Buffer.IsDirty());
+		TEST_EQUAL(TEXT("Initialized triple buffer must have correct temp buffer value"), Buffer.Read(), 1);
 	}
 
 	// pre-set buffer
@@ -42,18 +42,18 @@ TEST_CASE("Core::Misc::TTripleBuffer::Smoke Test", "[Core][Misc][Smoke]")
 		TTripleBuffer<int32> Buffer(Array);
 
 		int32 Read = Buffer.Read();
-		TestEqual(TEXT("Pre-set triple buffer must have correct Read buffer value"), Read, 3);
+		TEST_EQUAL(TEXT("Pre-set triple buffer must have correct Read buffer value"), Read, 3);
 
 		Buffer.SwapReadBuffers();
 
 		int32 Temp = Buffer.Read();
-		TestEqual(TEXT("Pre-set triple buffer must have correct Temp buffer value"), Temp, 1);
+		TEST_EQUAL(TEXT("Pre-set triple buffer must have correct Temp buffer value"), Temp, 1);
 
 		Buffer.SwapWriteBuffers();
 		Buffer.SwapReadBuffers();
 
 		int32 Write = Buffer.Read();
-		TestEqual(TEXT("Pre-set triple buffer must have correct Write buffer value"), Write, 2);
+		TEST_EQUAL(TEXT("Pre-set triple buffer must have correct Write buffer value"), Write, 2);
 	}
 
 	// operations
@@ -64,7 +64,7 @@ TEST_CASE("Core::Misc::TTripleBuffer::Smoke Test", "[Core][Misc][Smoke]")
 		{
 			int32& Write = Buffer.GetWriteBuffer(); Write = Index; Buffer.SwapWriteBuffers();
 			Buffer.SwapReadBuffers();
-			TestEqual(*FString::Printf(TEXT("Triple buffer must read correct value (%i)"), Index), Buffer.Read(), Index);
+			TEST_EQUAL(*FString::Printf(TEXT("Triple buffer must read correct value (%i)"), Index), Buffer.Read(), Index);
 		}
 
 		FRandomStream Rand;
@@ -91,7 +91,7 @@ TEST_CASE("Core::Misc::TTripleBuffer::Smoke Test", "[Core][Misc][Smoke]")
 
 				Buffer.SwapReadBuffers();
 				int32 Read = Buffer.Read();
-				TestTrue(TEXT("Triple buffer must read in increasing order"), Read > LastRead);
+				TEST_TRUE(TEXT("Triple buffer must read in increasing order"), Read > LastRead);
 				LastRead = Read;
 				--Reads;
 			}
diff --git a/Engine/Source/Runtime/Core/Tests/Misc/TypeContainerTest.cpp b/Engine/Source/Runtime/Core/Tests/Misc/TypeContainerTest.cpp
--- a/Engine/Source/Runtime/Core/Tests/Misc/TypeContainerTest.cpp
+++ b/Engine/Source/Runtime/Core/Tests/Misc/TypeContainerTest.cpp
@@ -8,7 +8,7 @@
 #include "Delegates/Delegate.h"
 #include "Async/Async.h"
 #include "Misc/TypeContainer.h"
-#include "TestHarness.h"
+#include "TestFixtures/CoreTestFixture.h"
 
 /* Helpers
  *****************************************************************************/
@@ -71,7 +71,7 @@ Expose_TNameOf(IBerry)
 Expose_TNameOf(IFruit)
 Expose_TNameOf(ISmoothie<>)
 
-TEST_CASE("Core::Misc::TTypeContainer::Smoke Test", "[Core][Misc][Smoke]")
+TEST_CASE_METHOD(FCoreTestFixture, "Core::Misc::TTypeContainer::Smoke Test", "[Core][Misc][Smoke]")
 {
 	// existing instance test
 	{
@@ -86,8 +86,8 @@ TEST_CASE("Core::Misc::TTypeContainer::Smoke Test", "[Core][Misc][Smoke]")
 
 		auto Instance = Container.GetInstance<IFruit>();
 
-		TestEqual(TEXT("Correct instance must be returned"), Instance, Fruit);
-		TestNotEqual(TEXT("Incorrect instance must not be returned"), Instance, StaticCastSharedRef<IFruit>(Berry));
+		TEST_EQUAL(TEXT("Correct instance must be returned"), Instance, Fruit);
+		TEST_NOT_EQUAL(TEXT("Incorrect instance must not be returned"), Instance, StaticCastSharedRef<IFruit>(Berry));
 	}
 
 	// per instance test
@@ -102,9 +102,9 @@ TEST_CASE("Core::Misc::TTypeContainer::Smoke Test", "[Core][Misc][Smoke]")
 		auto Smoothie1 = Container.GetInstance<ISmoothie<>>();
 		auto Smoothie2 = Container.GetInstance<ISmoothie<>>();
 
-		TestNotEqual(TEXT("For per-instances classes, a unique instance must be returned each time"), Smoothie1, Smoothie2);
-		TestNotEqual(TEXT("For per-instances dependencies, a unique instance must be returned each time [1]"), Smoothie1->GetBerry(), Smoothie2->GetBerry());
-		TestNotEqual(TEXT("For per-instances dependencies, a unique instance must be returned each time [2]"), Smoothie1->GetFruit(), Smoothie2->GetFruit());
+		TEST_NOT_EQUAL(TEXT("For per-instances classes, a unique instance must be returned each time"), Smoothie1, Smoothie2);
+		TEST_NOT_EQUAL(TEXT("For per-instances dependencies, a unique instance must be returned each time [1]"), Smoothie1->GetBerry(), Smoothie2->GetBerry());
+		TEST_NOT_EQUAL(TEXT("For per-instances dependencies, a unique instance must be returned each time [2]"), Smoothie1->GetFruit(), Smoothie2->GetFruit());
 	}
 
 	// per thread test
@@ -133,10 +133,10 @@ TEST_CASE("Core::Misc::TTypeContainer::Smoke Test", "[Core][Misc][Smoke]")
 		auto One2 = Smoothies2.Get().One;
 		auto Two2 = Smoothies2.Get().Two;
 
-		TestEqual(TEXT("For per-thread classes, the same instance must be returned from the same thread [1]"), One1, Two1);
-		TestEqual(TEXT("For per-thread classes, the same instance must be returned from the same thread [2]"), One2, Two2);
-		TestNotEqual(TEXT("For per-thread classes, different instances must be returned from different threads [1]"), One1, One2);
-		TestNotEqual(TEXT("For per-thread classes, different instances must be returned from different threads [2]"), Two1, Two2);
+		TEST_EQUAL(TEXT("For per-thread classes, the same instance must be returned from the same thread [1]"), One1, Two1);
+		TEST_EQUAL(TEXT("For per-thread classes, the same instance must be returned from the same thread [2]"), One2, Two2);
+		TEST_NOT_EQUAL(TEXT("For per-thread classes, different instances must be returned from different threads [1]"), One1, One2);
+		TEST_NOT_EQUAL(TEXT("For per-thread classes, different instances must be returned from different threads [2]"), Two1, Two2);
 	}
 
 	// per process test
@@ -165,10 +165,10 @@ TEST_CASE("Core::Misc::TTypeContainer::Smoke Test", "[Core][Misc][Smoke]")
 		auto One2 = Smoothies2.Get().One;
 		auto Two2 = Smoothies2.Get().Two;
 
-		TestEqual(TEXT("For per-process classes, the same instance must be returned from the same thread [1]"), One1, Two1);
-		TestEqual(TEXT("For per-process classes, the same instance must be returned from the same thread [2]"), One2, Two2);
-		TestEqual(TEXT("For per-process classes, the same instance must be returned from different threads [1]"), One1, One2);
-		TestEqual(TEXT("For per-process classes, the same instance must be returned from different threads [1]"), Two1, Two2);
+		TEST_EQUAL(TEXT("For per-process classes, the same instance must be returned from the same thread [1]"), One1, Two1);
+		TEST_EQUAL(TEXT("For per-process classes, the same instance must be returned from the same thread [2]"), One2, Two2);
+		TEST_EQUAL(TEXT("For per-process classes, the same instance must be returned from different threads [1]"), One1, One2);
+		TEST_EQUAL(TEXT("For per-process classes, the same instance must be returned from different threads [1]"), Two1, Two2);
 	}
 
 	// factory test
diff --git a/Engine/Source/Runtime/Core/Tests/Serialization/CompactBinaryPackageTest.cpp b/Engine/Source/Runtime/Core/Tests/Serialization/CompactBinaryPackageTest.cpp
--- a/Engine/Source/Runtime/Core/Tests/Serialization/CompactBinaryPackageTest.cpp
+++ b/Engine/Source/Runtime/Core/Tests/Serialization/CompactBinaryPackageTest.cpp
@@ -7,11 +7,11 @@
 #include "Serialization/CompactBinaryValidation.h"
 #include "Serialization/CompactBinaryWriter.h"
 #include "Serialization/MemoryReader.h"
-#include "TestHarness.h"
+#include "TestFixtures/CoreTestFixture.h"
 
-TEST_CASE("Core::Serialization::FCbAttachment::CbAttachment", "[Core][Serialization][Smoke]")
+TEST_CASE_METHOD(FCoreTestFixture, "Core::Serialization::FCbAttachment::CbAttachment", "[Core][Serialization][Smoke]")
 {
-	const auto TestSaveLoadValidate = [](const TCHAR* Test, const FCbAttachment& Attachment)
+	const auto TestSaveLoadValidate = [this](const TCHAR* Test, const FCbAttachment& Attachment)
 	{
 		TCbWriter<256> Writer;
 		FBufferArchive WriteAr;
@@ -19,69 +19,67 @@ TEST_CASE("Core::Serialization::FCbAttachment::CbAttachment", "[Core][Serializat
 		Attachment.Save(WriteAr);
 		FCbFieldIterator Fields = Writer.Save();
 
-		TestTrue(FString::Printf(TEXT("FCbAttachment(%s).Save()->Equals"), Test),
-			MakeMemoryView(WriteAr).EqualBytes(Fields.GetOuterBuffer().GetView()));
-		TestEqual(FString::Printf(TEXT("FCbAttachment(%s).Save()->ValidateRange"), Test),
-			ValidateCompactBinaryRange(MakeMemoryView(WriteAr), ECbValidateMode::All), ECbValidateError::None);
-		TestEqual(FString::Printf(TEXT("FCbAttachment(%s).Save()->ValidateAttachment"), Test),
-			ValidateCompactBinaryAttachment(MakeMemoryView(WriteAr), ECbValidateMode::All), ECbValidateError::None);
+		INFO(Test);
+		CHECK(MakeMemoryView(WriteAr).EqualBytes(Fields.GetOuterBuffer().GetView()));
+		CHECK_EQUAL(ValidateCompactBinaryRange(MakeMemoryView(WriteAr), ECbValidateMode::All), ECbValidateError::None);
+		CHECK_EQUAL(ValidateCompactBinaryAttachment(MakeMemoryView(WriteAr), ECbValidateMode::All), ECbValidateError::None);
 
 		FCbAttachment FromFields;
 		FromFields.TryLoad(Fields);
-		TestTrue(FString::Printf(TEXT("FCbAttachment(%s).TryLoad(Iterator)->AtEnd"), Test), !bool(Fields));
-		TestEqual(FString::Printf(TEXT("FCbAttachment(%s).TryLoad(Iterator)->Equals"), Test), FromFields, Attachment);
+		CHECK(!bool(Fields));
+		CHECK_EQUAL(FromFields, Attachment);
 
 		FCbAttachment FromArchive;
 		FMemoryReader ReadAr(WriteAr);
 		FromArchive.TryLoad(ReadAr);
-		TestTrue(FString::Printf(TEXT("FCbAttachment(%s).TryLoad(Archive)->AtEnd"), Test), ReadAr.AtEnd());
-		TestEqual(FString::Printf(TEXT("FCbAttachment(%s).TryLoad(Archive)->Equals"), Test), FromArchive, Attachment);
+		CHECK(ReadAr.AtEnd());
+		CHECK_EQUAL(FromArchive, Attachment);
 	};
 
-	// Empty Attachment
+	SECTION("Empty Attachment")
 	{
 		FCbAttachment Attachment;
-		TestTrue(TEXT("FCbAttachment(Null).IsNull()"), Attachment.IsNull());
-		TestFalse(TEXT("FCbAttachment(Null) as bool"), bool(Attachment));
-		TestFalse(TEXT("FCbAttachment(Null).AsBinary()"), bool(Attachment.AsBinary()));
-		TestFalse(TEXT("FCbAttachment(Null).AsObject()"), bool(Attachment.AsObject()));
-		TestFalse(TEXT("FCbAttachment(Null).IsBinary()"), Attachment.IsBinary());
-		TestFalse(TEXT("FCbAttachment(Null).IsCompressedBinary()"), Attachment.IsCompressedBinary());
-		TestFalse(TEXT("FCbAttachment(Null).IsObject()"), Attachment.IsObject());
-		TestEqual(TEXT("FCbAttachment(Null).GetHash()"), Attachment.GetHash(), FIoHash());
+		CHECK(Attachment.IsNull());
+		CHECK_FALSE(bool(Attachment));
+		CHECK_FALSE(bool(Attachment.AsBinary()));
+		CHECK_FALSE(bool(Attachment.AsObject()));
+		CHECK_FALSE(Attachment.IsBinary());
+		CHECK_FALSE(Attachment.IsCompressedBinary());
+		CHECK_FALSE(Attachment.IsObject());
+		CHECK_EQUAL(Attachment.GetHash(), FIoHash());
 	}
 
-	// Binary Attachment
+	SECTION("Binary Attachment")
 	{
 		FSharedBuffer Buffer = FSharedBuffer::Clone(MakeMemoryView<uint8>({0, 1, 2, 3}));
 		FCbAttachment Attachment(Buffer);
-		TestFalse(TEXT("FCbAttachment(Binary).IsNull()"), Attachment.IsNull());
-		TestTrue(TEXT("FCbAttachment(Binary) as bool"), bool(Attachment));
-		TestEqual(TEXT("FCbAttachment(Binary).AsBinary()"), Attachment.AsBinary(), Buffer);
-		TestFalse(TEXT("FCbAttachment(Binary).AsObject()"), bool(Attachment.AsObject()));
-		TestTrue(TEXT("FCbAttachment(Binary).IsBinary()"), Attachment.IsBinary());
-		TestFalse(TEXT("FCbAttachment(Binary).IsCompressedBinary()"), Attachment.IsCompressedBinary());
-		TestFalse(TEXT("FCbAttachment(Binary).IsObject()"), Attachment.IsObject());
-		TestEqual(TEXT("FCbAttachment(Binary).GetHash()"), Attachment.GetHash(), FIoHash::HashBuffer(Buffer));
+		CHECK_FALSE(Attachment.IsNull());
+		CHECK(bool(Attachment));
+		CHECK_EQUAL(Attachment.AsBinary(), Buffer);
+		CHECK_FALSE(bool(Attachment.AsObject()));
+		CHECK(Attachment.IsBinary());
+		CHECK_FALSE(Attachment.IsCompressedBinary());
+		CHECK_FALSE(Attachment.IsObject());
+		CHECK_EQUAL(Attachment.GetHash(), FIoHash::HashBuffer(Buffer));
 		TestSaveLoadValidate(TEXT("Binary"), Attachment);
 	}
 
-	// Compressed Binary Attachment
+	SECTION("Compressed Binary Attachment")
 	{
 		FCompressedBuffer Buffer = FCompressedBuffer::Compress(FSharedBuffer::Clone(MakeMemoryView<uint8>({0, 1, 2, 3})));
 		FCbAttachment Attachment(Buffer);
-		TestFalse(TEXT("FCbAttachment(CompressedBinary).IsNull()"), Attachment.IsNull());
-		TestTrue(TEXT("FCbAttachment(CompressedBinary) as bool"), bool(Attachment));
-		TestTrue(TEXT("FCbAttachment(CompressedBinary).AsCompressedBinary()"), Attachment.AsCompressedBinary().GetCompressed().ToShared().GetView().EqualBytes(Buffer.GetCompressed().ToShared().GetView()));
-		TestFalse(TEXT("FCbAttachment(CompressedBinary).AsObject()"), bool(Attachment.AsObject()));
-		TestFalse(TEXT("FCbAttachment(CompressedBinary).IsBinary()"), Attachment.IsBinary());
-		TestTrue(TEXT("FCbAttachment(CompressedBinary).IsCompressedBinary()"), Attachment.IsCompressedBinary());
-		TestFalse(TEXT("FCbAttachment(CompressedBinary).IsObject()"), Attachment.IsObject());
-		TestEqual(TEXT("FCbAttachment(CompressedBinary).GetHash()"), Attachment.GetHash(), FIoHash(Buffer.GetRawHash()));
+		CHECK_FALSE(Attachment.IsNull());
+		CHECK(bool(Attachment));
+		CHECK(Attachment.AsCompressedBinary().GetCompressed().ToShared().GetView().EqualBytes(Buffer.GetCompressed().ToShared().GetView()));
+		CHECK_FALSE(bool(Attachment.AsObject()));
+		CHECK_FALSE(Attachment.IsBinary());
+		CHECK(Attachment.IsCompressedBinary());
+		CHECK_FALSE(Attachment.IsObject());
+		CHECK_EQUAL(Attachment.GetHash(), FIoHash(Buffer.GetRawHash()));
 		TestSaveLoadValidate(TEXT("CompressedBinary"), Attachment);
 	}
 
-	// Object Attachment
+	SECTION("Object Attachment")
 	{
 		FCbWriter Writer;
 		Writer.BeginObject();
@@ -89,34 +87,34 @@ TEST_CASE("Core::Serialization::FCbAttachment::CbAttachment", "[Core][Serializat
 		Writer.EndObject();
 		FCbObject Object = Writer.Save().AsObject();
 		FCbAttachment Attachment(Object);
-		TestFalse(TEXT("FCbAttachment(Object).IsNull()"), Attachment.IsNull());
-		TestTrue(TEXT("FCbAttachment(Object) as bool"), bool(Attachment));
-		TestEqual(TEXT("FCbAttachment(Object).AsBinary()"), Attachment.AsBinary(), FSharedBuffer());
-		TestTrue(TEXT("FCbAttachment(Object).AsObject()"), Attachment.AsObject().Equals(Object));
-		TestFalse(TEXT("FCbAttachment(Object).IsBinary()"), Attachment.IsBinary());
-		TestFalse(TEXT("FCbAttachment(Object).IsCompressedBinary()"), Attachment.IsCompressedBinary());
-		TestTrue(TEXT("FCbAttachment(Object).IsObject()"), Attachment.IsObject());
-		TestEqual(TEXT("FCbAttachment(Object).GetHash()"), Attachment.GetHash(), FIoHash(Object.GetHash()));
+		CHECK_FALSE(Attachment.IsNull());
+		CHECK(bool(Attachment));
+		CHECK_EQUAL(Attachment.AsBinary(), FSharedBuffer());
+		CHECK(Attachment.AsObject().Equals(Object));
+		CHECK_FALSE(Attachment.IsBinary());
+		CHECK_FALSE(Attachment.IsCompressedBinary());
+		CHECK(Attachment.IsObject());
+		CHECK_EQUAL( Attachment.GetHash(), FIoHash(Object.GetHash()));
 		TestSaveLoadValidate(TEXT("Object"), Attachment);
 	}
 
-	// Binary View
+	SECTION("Binary View")
 	{
 		const uint8 Value[]{0, 1, 2, 3};
 		FSharedBuffer Buffer = FSharedBuffer::MakeView(MakeMemoryView(Value));
 		FCbAttachment Attachment(Buffer);
-		TestFalse(TEXT("FCbAttachment(BinaryView).IsNull()"), Attachment.IsNull());
-		TestTrue(TEXT("FCbAttachment(BinaryView) as bool"), bool(Attachment));
-		TestNotEqual(TEXT("FCbAttachment(BinaryView).AsBinary()"), Attachment.AsBinary(), Buffer);
-		TestTrue(TEXT("FCbAttachment(BinaryView).AsBinary()"), Attachment.AsBinary().GetView().EqualBytes(Buffer.GetView()));
-		TestFalse(TEXT("FCbAttachment(BinaryView).AsObject()"), bool(Attachment.AsObject()));
-		TestTrue(TEXT("FCbAttachment(BinaryView).IsBinary()"), Attachment.IsBinary());
-		TestFalse(TEXT("FCbAttachment(BinaryView).IsCompressedBinary()"), Attachment.IsCompressedBinary());
-		TestFalse(TEXT("FCbAttachment(BinaryView).IsObject()"), Attachment.IsObject());
-		TestEqual(TEXT("FCbAttachment(BinaryView).GetHash()"), Attachment.GetHash(), FIoHash::HashBuffer(Buffer));
+		CHECK_FALSE(Attachment.IsNull());
+		CHECK(bool(Attachment));
+		CHECK_NOT_EQUAL(Attachment.AsBinary(), Buffer);
+		CHECK(Attachment.AsBinary().GetView().EqualBytes(Buffer.GetView()));
+		CHECK_FALSE(bool(Attachment.AsObject()));
+		CHECK(Attachment.IsBinary());
+		CHECK_FALSE(Attachment.IsCompressedBinary());
+		CHECK_FALSE(Attachment.IsObject());
+		CHECK_EQUAL(Attachment.GetHash(), FIoHash::HashBuffer(Buffer));
 	}
 
-	// Object View
+	SECTION("Object View")
 	{
 		FCbWriter Writer;
 		Writer.BeginObject();
@@ -125,16 +123,16 @@ TEST_CASE("Core::Serialization::FCbAttachment::CbAttachment", "[Core][Serializat
 		FCbObject Object = Writer.Save().AsObject();
 		FCbObject ObjectView = FCbObject::MakeView(Object);
 		FCbAttachment Attachment(ObjectView);
-		TestFalse(TEXT("FCbAttachment(ObjectView).IsNull()"), Attachment.IsNull());
-		TestTrue(TEXT("FCbAttachment(ObjectView) as bool"), bool(Attachment));
-		TestTrue(TEXT("FCbAttachment(ObjectView).AsObject()"), Attachment.AsObject().Equals(Object));
-		TestFalse(TEXT("FCbAttachment(ObjectView).IsBinary()"), Attachment.IsBinary());
-		TestFalse(TEXT("FCbAttachment(ObjectView).IsCompressedBinary()"), Attachment.IsCompressedBinary());
-		TestTrue(TEXT("FCbAttachment(ObjectView).IsObject()"), Attachment.IsObject());
-		TestEqual(TEXT("FCbAttachment(ObjectView).GetHash()"), Attachment.GetHash(), FIoHash(Object.GetHash()));
+		CHECK_FALSE(Attachment.IsNull());
+		CHECK(bool(Attachment));
+		CHECK(Attachment.AsObject().Equals(Object));
+		CHECK_FALSE(Attachment.IsBinary());
+		CHECK_FALSE(Attachment.IsCompressedBinary());
+		CHECK(Attachment.IsObject());
+		CHECK_EQUAL( Attachment.GetHash(), FIoHash(Object.GetHash()));
 	}
 
-	// Binary Load from View
+	SECTION("Binary Load from View")
 	{
 		const uint8 Value[]{0, 1, 2, 3};
 		const FSharedBuffer Buffer = FSharedBuffer::MakeView(MakeMemoryView(Value));
@@ -146,21 +144,18 @@ TEST_CASE("Core::Serialization::FCbAttachment::CbAttachment", "[Core][Serializat
 		FCbFieldIterator FieldsView = FCbFieldIterator::MakeRangeView(FCbFieldViewIterator(Fields));
 
 		Attachment.TryLoad(FieldsView);
-		TestFalse(TEXT("FCbAttachment(LoadBinaryView).IsNull()"), Attachment.IsNull());
-		TestTrue(TEXT("FCbAttachment(LoadBinaryView) as bool"), bool(Attachment));
-		TestFalse(TEXT("FCbAttachment(LoadBinaryView).AsBinary()->!InView"),
-			FieldsView.GetOuterBuffer().GetView().Contains(Attachment.AsBinary().GetView()));
-		TestTrue(TEXT("FCbAttachment(LoadBinaryView).AsBinary()->EqualBytes"),
-			Attachment.AsBinary().GetView().EqualBytes(Buffer.GetView()));
-		TestFalse(TEXT("FCbAttachment(LoadBinaryView).AsObject()"), bool(Attachment.AsObject()));
-		TestTrue(TEXT("FCbAttachment(LoadBinaryView).IsBinary()"), Attachment.IsBinary());
-		TestFalse(TEXT("FCbAttachment(LoadBinaryView).IsCompressedBinary()"), Attachment.IsCompressedBinary());
-		TestFalse(TEXT("FCbAttachment(LoadBinaryView).IsObject()"), Attachment.IsObject());
-		TestEqual(TEXT("FCbAttachment(LoadBinaryView).GetHash()"),
-			Attachment.GetHash(), FIoHash::HashBuffer(MakeMemoryView(Value)));
+		CHECK_FALSE(Attachment.IsNull());
+		CHECK( bool(Attachment));
+		CHECK_FALSE(FieldsView.GetOuterBuffer().GetView().Contains(Attachment.AsBinary().GetView()));
+		CHECK(Attachment.AsBinary().GetView().EqualBytes(Buffer.GetView()));
+		CHECK_FALSE(bool(Attachment.AsObject()));
+		CHECK(Attachment.IsBinary());
+		CHECK_FALSE(Attachment.IsCompressedBinary());
+		CHECK_FALSE(Attachment.IsObject());
+		CHECK_EQUAL(Attachment.GetHash(), FIoHash::HashBuffer(MakeMemoryView(Value)));
 	}
 
-	// Compressed Binary Load from View
+	SECTION("Compressed Binary Load from View")
 	{
 		const uint8 Value[]{0, 1, 2, 3};
 		FCompressedBuffer Buffer = FCompressedBuffer::Compress(FSharedBuffer::MakeView(MakeMemoryView(Value)));
@@ -172,29 +167,25 @@ TEST_CASE("Core::Serialization::FCbAttachment::CbAttachment", "[Core][Serializat
 		FCbFieldIterator FieldsView = FCbFieldIterator::MakeRangeView(FCbFieldViewIterator(Fields));
 
 		Attachment.TryLoad(FieldsView);
-		TestFalse(TEXT("FCbAttachment(LoadCompressedBinaryView).IsNull()"), Attachment.IsNull());
-		TestTrue(TEXT("FCbAttachment(LoadCompressedBinaryView) as bool"), bool(Attachment));
-		TestFalse(TEXT("FCbAttachment(LoadCompressedBinaryView).AsBinary()->!InView"),
-			FieldsView.GetOuterBuffer().GetView().Contains(Attachment.AsCompressedBinary().GetCompressed().ToShared().GetView()));
-		TestTrue(TEXT("FCbAttachment(LoadCompressedBinaryView).AsCompressedBinary()->EqualBytes"),
-			Attachment.AsCompressedBinary().GetCompressed().ToShared().GetView().EqualBytes(Buffer.GetCompressed().ToShared().GetView()));
-		TestFalse(TEXT("FCbAttachment(LoadCompressedBinaryView).AsObject()"), bool(Attachment.AsObject()));
-		TestFalse(TEXT("FCbAttachment(LoadCompressedBinaryView).IsBinary()"), Attachment.IsBinary());
-		TestTrue(TEXT("FCbAttachment(LoadCompressedBinaryView).IsCompressedBinary()"), Attachment.IsCompressedBinary());
-		TestFalse(TEXT("FCbAttachment(LoadCompressedBinaryView).IsObject()"), Attachment.IsObject());
-		TestEqual(TEXT("FCbAttachment(LoadCompressedBinaryView).GetHash()"),
-			Attachment.GetHash(), FIoHash::HashBuffer(MakeMemoryView(Value)));
+		CHECK_FALSE(Attachment.IsNull());
+		CHECK(bool(Attachment));
+		CHECK_FALSE(FieldsView.GetOuterBuffer().GetView().Contains(Attachment.AsCompressedBinary().GetCompressed().ToShared().GetView()));
+		CHECK(Attachment.AsCompressedBinary().GetCompressed().ToShared().GetView().EqualBytes(Buffer.GetCompressed().ToShared().GetView()));
+		CHECK_FALSE(bool(Attachment.AsObject()));
+		CHECK_FALSE(Attachment.IsBinary());
+		CHECK(Attachment.IsCompressedBinary());
+		CHECK_FALSE(Attachment.IsObject());
+		CHECK_EQUAL(Attachment.GetHash(), FIoHash::HashBuffer(MakeMemoryView(Value)));
 	}
 
-	// Object Load from View
+	SECTION("Object Load from View")
 	{
 		FCbWriter ValueWriter;
 		ValueWriter.BeginObject();
 		ValueWriter << "Name"_ASV << 42;
 		ValueWriter.EndObject();
 		const FCbObject Value = ValueWriter.Save().AsObject();
-		TestEqual(TEXT("FCbAttachment(LoadObjectView).Validate"),
-			ValidateCompactBinaryRange(Value.GetOuterBuffer(), ECbValidateMode::All), ECbValidateError::None);
+		CHECK_EQUAL(ValidateCompactBinaryRange(Value.GetOuterBuffer(), ECbValidateMode::All), ECbValidateError::None);
 		FCbAttachment Attachment(Value);
 
 		FCbWriter Writer;
@@ -204,62 +195,60 @@ TEST_CASE("Core::Serialization::FCbAttachment::CbAttachment", "[Core][Serializat
 
 		Attachment.TryLoad(FieldsView);
 		FMemoryView View;
-		TestFalse(TEXT("FCbAttachment(LoadObjectView).IsNull()"), Attachment.IsNull());
-		TestTrue(TEXT("FCbAttachment(LoadObjectView) as bool"), bool(Attachment));
-		TestTrue(TEXT("FCbAttachment(LoadObjectView).AsBinary()->EqualBytes"),
-			Attachment.AsBinary().GetView().EqualBytes(FMemoryView()));
-		TestFalse(TEXT("FCbAttachment(LoadObjectView).AsObject()->!InView"),
-			!Attachment.AsObject().TryGetView(View) || FieldsView.GetOuterBuffer().GetView().Contains(View));
-		TestFalse(TEXT("FCbAttachment(LoadObjectView).IsBinary()"), Attachment.IsBinary());
-		TestFalse(TEXT("FCbAttachment(LoadObjectView).IsCompressedBinary()"), Attachment.IsCompressedBinary());
-		TestTrue(TEXT("FCbAttachment(LoadObjectView).IsObject()"), Attachment.IsObject());
-		TestEqual(TEXT("FCbAttachment(LoadObjectView).GetHash()"), Attachment.GetHash(), FIoHash(Value.GetHash()));
+		CHECK_FALSE(Attachment.IsNull());
+		CHECK(bool(Attachment));
+		CHECK(Attachment.AsBinary().GetView().EqualBytes(FMemoryView()));
+		CHECK_FALSE((!Attachment.AsObject().TryGetView(View) || FieldsView.GetOuterBuffer().GetView().Contains(View)));
+		CHECK_FALSE(Attachment.IsBinary());
+		CHECK_FALSE(Attachment.IsCompressedBinary());
+		CHECK(Attachment.IsObject());
+		CHECK_EQUAL(Attachment.GetHash(), FIoHash(Value.GetHash()));
 	}
 
-	// Binary Null
+	SECTION("Binary Null")
 	{
 		const FCbAttachment Attachment(FSharedBuffer{});
-		TestTrue(TEXT("FCbAttachment(BinaryNull).IsNull()"), Attachment.IsNull());
-		TestFalse(TEXT("FCbAttachment(BinaryNull).IsBinary()"), Attachment.IsBinary());
-		TestFalse(TEXT("FCbAttachment(BinaryNull).IsCompressedBinary()"), Attachment.IsCompressedBinary());
-		TestFalse(TEXT("FCbAttachment(BinaryNull).IsObject()"), Attachment.IsObject());
-		TestEqual(TEXT("FCbAttachment(BinaryNull).GetHash()"), Attachment.GetHash(), FIoHash::Zero);
+		CHECK(Attachment.IsNull());
+		CHECK_FALSE(Attachment.IsBinary());
+		CHECK_FALSE(Attachment.IsCompressedBinary());
+		CHECK_FALSE(Attachment.IsObject());
+		CHECK_EQUAL( Attachment.GetHash(), FIoHash::Zero);
 	}
 
-	// Binary Empty
+	SECTION("Binary Empty")
 	{
 		const FCbAttachment Attachment(FUniqueBuffer::Alloc(0).MoveToShared());
-		TestFalse(TEXT("FCbAttachment(BinaryEmpty).IsNull()"), Attachment.IsNull());
-		TestTrue(TEXT("FCbAttachment(BinaryEmpty).IsBinary()"), Attachment.IsBinary());
-		TestFalse(TEXT("FCbAttachment(BinaryEmpty).IsCompressedBinary()"), Attachment.IsCompressedBinary());
-		TestFalse(TEXT("FCbAttachment(BinaryEmpty).IsObject()"), Attachment.IsObject());
-		TestEqual(TEXT("FCbAttachment(BinaryEmpty).GetHash()"), Attachment.GetHash(), FIoHash::HashBuffer(FSharedBuffer{}));
+		CHECK_FALSE(Attachment.IsNull());
+		CHECK(Attachment.IsBinary());
+		CHECK_FALSE(Attachment.IsCompressedBinary());
+		CHECK_FALSE(Attachment.IsObject());
+		CHECK_EQUAL(Attachment.GetHash(), FIoHash::HashBuffer(FSharedBuffer{}));
 	}
 
-	// Compressed Binary Empty
+	SECTION("Compressed Binary Empty")
 	{
 		const FCbAttachment Attachment(FCompressedBuffer::Compress(FUniqueBuffer::Alloc(0).MoveToShared()));
-		TestFalse(TEXT("FCbAttachment(CompressedBinaryEmpty).IsNull()"), Attachment.IsNull());
-		TestFalse(TEXT("FCbAttachment(CompressedBinaryEmpty).IsBinary()"), Attachment.IsBinary());
-		TestTrue(TEXT("FCbAttachment(CompressedBinaryEmpty).IsCompressedBinary()"), Attachment.IsCompressedBinary());
-		TestFalse(TEXT("FCbAttachment(CompressedBinaryEmpty).IsObject()"), Attachment.IsObject());
-		TestEqual(TEXT("FCbAttachment(CompressedBinaryEmpty).GetHash()"), Attachment.GetHash(), FIoHash::HashBuffer(FSharedBuffer{}));
+		CHECK_FALSE(Attachment.IsNull());
+		CHECK_FALSE(Attachment.IsBinary());
+		CHECK(Attachment.IsCompressedBinary());
+		CHECK_FALSE(Attachment.IsObject());
+		CHECK_EQUAL(Attachment.GetHash(), FIoHash::HashBuffer(FSharedBuffer{}));
 	}
 
-	// Object Empty
+	SECTION("Object Empty")
 	{
 		const FCbAttachment Attachment(FCbObject{});
-		TestFalse(TEXT("FCbAttachment(ObjectEmpty).IsNull()"), Attachment.IsNull());
-		TestFalse(TEXT("FCbAttachment(ObjectEmpty).IsBinary()"), Attachment.IsBinary());
-		TestFalse(TEXT("FCbAttachment(ObjectEmpty).IsCompressedBinary()"), Attachment.IsCompressedBinary());
-		TestTrue(TEXT("FCbAttachment(ObjectEmpty).IsObject()"), Attachment.IsObject());
-		TestEqual(TEXT("FCbAttachment(ObjectEmpty).GetHash()"), Attachment.GetHash(), FIoHash(FCbObject().GetHash()));
+		CHECK_FALSE(Attachment.IsNull());
+		CHECK_FALSE(Attachment.IsBinary());
+		CHECK_FALSE(Attachment.IsCompressedBinary());
+		CHECK( Attachment.IsObject());
+		CHECK_EQUAL(Attachment.GetHash(), FIoHash(FCbObject().GetHash()));
 	}
 }
 
-TEST_CASE("Core::Serialization::FCbPackage::CbPackage", "[Core][Serialization][Smoke]")
+TEST_CASE_METHOD(FAutomationTestFixture, "Core::Serialization::FCbPackage::CbPackage", "[Core][Serialization][Smoke]")
 {
-	const auto TestSaveLoadValidate = [](const TCHAR* Test, const FCbPackage& Package)
+	const auto TestSaveLoadValidate = [this](const TCHAR* Test, const FCbPackage& Package)
 	{
 		TCbWriter<256> Writer;
 		FBufferArchive WriteAr;
@@ -267,35 +256,33 @@ TEST_CASE("Core::Serialization::FCbPackage::CbPackage", "[Core][Serialization][S
 		Package.Save(WriteAr);
 		FCbFieldIterator Fields = Writer.Save();
 
-		TestTrue(FString::Printf(TEXT("FCbPackage(%s).Save()->Equals"), Test),
-			MakeMemoryView(WriteAr).EqualBytes(Fields.GetOuterBuffer().GetView()));
-		TestEqual(FString::Printf(TEXT("FCbPackage(%s).Save()->ValidateRange"), Test),
-			ValidateCompactBinaryRange(MakeMemoryView(WriteAr), ECbValidateMode::All), ECbValidateError::None);
-		TestEqual(FString::Printf(TEXT("FCbPackage(%s).Save()->ValidatePackage"), Test),
-			ValidateCompactBinaryPackage(MakeMemoryView(WriteAr), ECbValidateMode::All), ECbValidateError::None);
+		INFO(Test);
+		CHECK(MakeMemoryView(WriteAr).EqualBytes(Fields.GetOuterBuffer().GetView()));
+		CHECK_EQUAL(ValidateCompactBinaryRange(MakeMemoryView(WriteAr), ECbValidateMode::All), ECbValidateError::None);
+		CHECK_EQUAL(ValidateCompactBinaryPackage(MakeMemoryView(WriteAr), ECbValidateMode::All), ECbValidateError::None);
 
 		FCbPackage FromFields;
 		FromFields.TryLoad(Fields);
-		TestFalse(FString::Printf(TEXT("FCbPackage(%s).TryLoad(Iterator)->AtEnd"), Test), bool(Fields));
-		TestEqual(FString::Printf(TEXT("FCbPackage(%s).TryLoad(Iterator)->Equals"), Test), FromFields, Package);
+		CHECK_FALSE(bool(Fields));
+		CHECK_EQUAL(FromFields, Package);
 
 		FCbPackage FromArchive;
 		FMemoryReader ReadAr(WriteAr);
 		FromArchive.TryLoad(ReadAr);
-		TestTrue(FString::Printf(TEXT("FCbPackage(%s).TryLoad(Archive)->AtEnd"), Test), ReadAr.AtEnd());
-		TestEqual(FString::Printf(TEXT("FCbPackage(%s).TryLoad(Archive)->Equals"), Test), FromArchive, Package);
+		CHECK(ReadAr.AtEnd());
+		CHECK_EQUAL(FromArchive, Package);
 	};
 
-	// Empty
+	SECTION("Empty")
 	{
 		FCbPackage Package;
-		TestTrue(TEXT("FCbPackage(Empty).IsNull()"), Package.IsNull());
-		TestFalse(TEXT("FCbPackage(Empty) as bool"), bool(Package));
-		TestEqual(TEXT("FCbPackage(Empty).GetAttachments()"), Package.GetAttachments().Num(), 0);
+		CHECK(Package.IsNull());
+		CHECK_FALSE(bool(Package));
+		CHECK_EQUAL(Package.GetAttachments().Num(), 0);
 		TestSaveLoadValidate(TEXT("Empty"), Package);
 	}
 
-	// Object Only
+	SECTION("Object Only")
 	{
 		TCbWriter<256> Writer;
 		Writer.BeginObject();
@@ -304,16 +291,16 @@ TEST_CASE("Core::Serialization::FCbPackage::CbPackage", "[Core][Serialization][S
 
 		const FCbObject Object = Writer.Save().AsObject();
 		FCbPackage Package(Object);
-		TestFalse(TEXT("FCbPackage(Object).IsNull()"), Package.IsNull());
-		TestTrue(TEXT("FCbPackage(Object) as bool"), bool(Package));
-		TestEqual(TEXT("FCbPackage(Object).GetAttachments()"), Package.GetAttachments().Num(), 0);
-		TestEqual(TEXT("FCbPackage(Object).GetObject()->IsClone"), Package.GetObject().GetOuterBuffer(), Object.GetOuterBuffer());
-		TestEqual(TEXT("FCbPackage(Object).GetObject()"), Package.GetObject()["Field"].AsInt32(), 42);
-		TestEqual(TEXT("FCbPackage(Object).GetObjectHash()"), Package.GetObjectHash(), FIoHash(Package.GetObject().GetHash()));
+		CHECK_FALSE(Package.IsNull());
+		CHECK(bool(Package));
+		CHECK_EQUAL(Package.GetAttachments().Num(), 0);
+		CHECK_EQUAL(Package.GetObject().GetOuterBuffer(), Object.GetOuterBuffer());
+		CHECK_EQUAL(Package.GetObject()["Field"].AsInt32(), 42);
+		CHECK_EQUAL(Package.GetObjectHash(), FIoHash(Package.GetObject().GetHash()));
 		TestSaveLoadValidate(TEXT("Object"), Package);
 	}
 
-	// Object View Only
+	SECTION("Object View Only")
 	{
 		TCbWriter<256> Writer;
 		Writer.BeginObject();
@@ -322,16 +309,16 @@ TEST_CASE("Core::Serialization::FCbPackage::CbPackage", "[Core][Serialization][S
 
 		const FCbObject Object = Writer.Save().AsObject();
 		FCbPackage Package(FCbObject::MakeView(Object));
-		TestFalse(TEXT("FCbPackage(Object).IsNull()"), Package.IsNull());
-		TestTrue(TEXT("FCbPackage(Object) as bool"), bool(Package));
-		TestEqual(TEXT("FCbPackage(Object).GetAttachments()"), Package.GetAttachments().Num(), 0);
-		TestNotEqual(TEXT("FCbPackage(Object).GetObject()->IsClone"), Package.GetObject().GetOuterBuffer(), Object.GetOuterBuffer());
-		TestEqual(TEXT("FCbPackage(Object).GetObject()"), Package.GetObject()["Field"].AsInt32(), 42);
-		TestEqual(TEXT("FCbPackage(Object).GetObjectHash()"), Package.GetObjectHash(), FIoHash(Package.GetObject().GetHash()));
+		CHECK_FALSE(Package.IsNull());
+		CHECK(bool(Package));
+		CHECK_EQUAL(Package.GetAttachments().Num(), 0);
+		CHECK_NOT_EQUAL(Package.GetObject().GetOuterBuffer(), Object.GetOuterBuffer());
+		CHECK_EQUAL(Package.GetObject()["Field"].AsInt32(), 42);
+		CHECK_EQUAL(Package.GetObjectHash(), FIoHash(Package.GetObject().GetHash()));
 		TestSaveLoadValidate(TEXT("Object"), Package);
 	}
 
-	// Attachment Only
+	SECTION("Attachment Only")
 	{
 		FCbObject Object1;
 		{
@@ -354,38 +341,35 @@ TEST_CASE("Core::Serialization::FCbPackage::CbPackage", "[Core][Serialization][S
 		Package.AddAttachment(FCbAttachment(Object1));
 		Package.AddAttachment(FCbAttachment(Object2.GetOuterBuffer()));
 
-		TestFalse(TEXT("FCbPackage(Attachments).IsNull()"), Package.IsNull());
-		TestTrue(TEXT("FCbPackage(Attachments) as bool"), bool(Package));
-		TestEqual(TEXT("FCbPackage(Attachments).GetAttachments()"), Package.GetAttachments().Num(), 2);
-		TestTrue(TEXT("FCbPackage(Attachments).GetObject()"), Package.GetObject().Equals(FCbObject()));
-		TestEqual(TEXT("FCbPackage(Attachments).GetObjectHash()"), Package.GetObjectHash(), FIoHash());
+		CHECK(!Package.IsNull());
+		CHECK(bool(Package));
+		CHECK_EQUAL(Package.GetAttachments().Num(), 2);
+		CHECK(Package.GetObject().Equals(FCbObject()));
+		CHECK_EQUAL(Package.GetObjectHash(), FIoHash());
 		TestSaveLoadValidate(TEXT("Attachments"), Package);
 
 		const FCbAttachment* const Object1Attachment = Package.FindAttachment(Object1.GetHash());
 		const FCbAttachment* const Object2Attachment = Package.FindAttachment(Object2.GetHash());
 
-		TestTrue(TEXT("FCbPackage(Attachments).FindAttachment(Object1)"),
-			Object1Attachment && Object1Attachment->AsObject().Equals(Object1));
-		TestTrue(TEXT("FCbPackage(Attachments).FindAttachment(Object2)"),
-			Object2Attachment && Object2Attachment->AsBinary() == Object2.GetOuterBuffer());
+		REQUIRE(Object1Attachment);
+		CHECK(Object1Attachment->AsObject().Equals(Object1));
+		REQUIRE(Object2Attachment);
+		CHECK(Object2Attachment->AsBinary() == Object2.GetOuterBuffer());
 
 		FSharedBuffer Object1ClonedBuffer = FSharedBuffer::Clone(Object1.GetOuterBuffer());
 		Package.AddAttachment(FCbAttachment(Object1ClonedBuffer));
 		Package.AddAttachment(FCbAttachment(FCbObject::Clone(Object2)));
 
-		TestEqual(TEXT("FCbPackage(Attachments).GetAttachments()"), Package.GetAttachments().Num(), 2);
-		TestEqual(TEXT("FCbPackage(Attachments).FindAttachment(Object1, Re-Add)"),
-			Package.FindAttachment(Object1.GetHash()), Object1Attachment);
-		TestEqual(TEXT("FCbPackage(Attachments).FindAttachment(Object2, Re-Add)"),
-			Package.FindAttachment(Object2.GetHash()), Object2Attachment);
+		CHECK_EQUAL(Package.GetAttachments().Num(), 2);
+		CHECK_EQUAL(Package.FindAttachment(Object1.GetHash()), Object1Attachment);
+		CHECK_EQUAL(Package.FindAttachment(Object2.GetHash()), Object2Attachment);
 
-		TestTrue(TEXT("FCbPackage(Attachments).FindAttachment(ObjectAsBinary)"),
-			Object1Attachment&& Object1Attachment->AsBinary() == Object1ClonedBuffer);
-		TestTrue(TEXT("FCbPackage(Attachments).FindAttachment(FieldAsField)"),
-			Object2Attachment&& Object2Attachment->AsObject().Equals(Object2));
+		REQUIRE(Object1Attachment);
+		CHECK(Object1Attachment->AsBinary() == Object1ClonedBuffer);
+		REQUIRE(Object2Attachment);
+		CHECK(Object2Attachment->AsObject().Equals(Object2));
 
-		TestTrue(TEXT("FCbPackage(Attachments).GetAttachments()->Sorted"),
-			Algo::IsSorted(Package.GetAttachments()));
+		CHECK(Algo::IsSorted(Package.GetAttachments()));
 	}
 
 	// Shared Values
@@ -433,58 +417,52 @@ TEST_CASE("Core::Serialization::FCbPackage::CbPackage", "[Core][Serialization][S
 				FSharedBuffer();
 		};
 
-	// Object + Attachments
+	SECTION("Object + Attachments")
 	{
 		FCbPackage Package;
 		Package.SetObject(Level1, Level1Hash, Resolver);
 
-		TestFalse(TEXT("FCbPackage(Object+Attachments).IsNull()"), Package.IsNull());
-		TestTrue(TEXT("FCbPackage(Object+Attachments) as bool"), bool(Package));
-		TestEqual(TEXT("FCbPackage(Object+Attachments).GetAttachments()"), Package.GetAttachments().Num(), 3);
-		TestTrue(TEXT("FCbPackage(Object+Attachments).GetObject()"),
-			Package.GetObject().GetOuterBuffer() == Level1.GetOuterBuffer());
-		TestEqual(TEXT("FCbPackage(Object+Attachments).GetObjectHash()"), Package.GetObjectHash(), Level1Hash);
+		CHECK_FALSE(Package.IsNull());
+		CHECK_EQUAL(Package.GetAttachments().Num(), 3);
+		CHECK(Package.GetObject().GetOuterBuffer() == Level1.GetOuterBuffer());
+		CHECK_EQUAL(Package.GetObjectHash(), Level1Hash);
 		TestSaveLoadValidate(TEXT("Object+Attachments"), Package);
 
 		const FCbAttachment* const Level2Attachment = Package.FindAttachment(Level2Hash);
 		const FCbAttachment* const Level3Attachment = Package.FindAttachment(Level3Hash);
 		const FCbAttachment* const Level4Attachment = Package.FindAttachment(Level4Hash);
-		TestTrue(TEXT("FCbPackage(Object+Attachments).FindAttachment(Level2)"),
-			Level2Attachment && Level2Attachment->AsObject().Equals(Level2));
-		TestTrue(TEXT("FCbPackage(Object+Attachments).FindAttachment(Level3)"),
-			Level3Attachment && Level3Attachment->AsObject().Equals(Level3));
-		TestTrue(TEXT("FCbPackage(Object+Attachments).FindAttachment(Level4)"),
-			Level4Attachment &&
-			Level4Attachment->AsBinary() != Level4 &&
-			Level4Attachment->AsBinary().GetView().EqualBytes(Level4.GetView()));
-
-		TestTrue(TEXT("FCbPackage(Object+Attachments).GetAttachments()->Sorted"),
-			Algo::IsSorted(Package.GetAttachments()));
+
+		REQUIRE(Level2Attachment);
+		CHECK(Level2Attachment->AsObject().Equals(Level2));
+
+		REQUIRE(Level3Attachment);
+		CHECK(Level3Attachment->AsObject().Equals(Level3));
+
+		REQUIRE(Level4Attachment);
+		CHECK(Level4Attachment->AsBinary() != Level4);
+		CHECK(Level4Attachment->AsBinary().GetView().EqualBytes(Level4.GetView()));
+
+		CHECK(Algo::IsSorted(Package.GetAttachments()));
 
 		const FCbPackage PackageCopy = Package;
-		TestEqual(TEXT("FCbPackage(Object+Attachments).Equals(EqualCopied)"), PackageCopy, Package);
-
-		TestEqual(TEXT("FCbPackage(Object+Attachments).RemoveAttachment(Level1)"),
-			Package.RemoveAttachment(Level1Hash), 0);
-		TestEqual(TEXT("FCbPackage(Object+Attachments).RemoveAttachment(Level2)"),
-			Package.RemoveAttachment(Level2Hash), 1);
-		TestEqual(TEXT("FCbPackage(Object+Attachments).RemoveAttachment(Level3)"),
-			Package.RemoveAttachment(Level3Hash), 1);
-		TestEqual(TEXT("FCbPackage(Object+Attachments).RemoveAttachment(Level4)"),
-			Package.RemoveAttachment(Level4Hash), 1);
-		TestEqual(TEXT("FCbPackage(Object+Attachments).RemoveAttachment(Level4, Again)"),
-			Package.RemoveAttachment(Level4Hash), 0);
-		TestEqual(TEXT("FCbPackage(Object+Attachments).GetAttachments(Removed)"), Package.GetAttachments().Num(), 0);
-
-		TestNotEqual(TEXT("FCbPackage(Object+Attachments).Equals(AttachmentsNotEqual)"), PackageCopy, Package);
+		CHECK_EQUAL(PackageCopy, Package);
+
+		CHECK_EQUAL(Package.RemoveAttachment(Level1Hash), 0);
+		CHECK_EQUAL(Package.RemoveAttachment(Level2Hash), 1);
+		CHECK_EQUAL(Package.RemoveAttachment(Level3Hash), 1);
+		CHECK_EQUAL(Package.RemoveAttachment(Level4Hash), 1);
+		CHECK_EQUAL(Package.RemoveAttachment(Level4Hash), 0);
+		CHECK_EQUAL(Package.GetAttachments().Num(), 0);
+
+		CHECK_NOT_EQUAL(PackageCopy, Package);
 		Package = PackageCopy;
-		TestEqual(TEXT("FCbPackage(Object+Attachments).Equals(EqualAssigned)"), PackageCopy, Package);
+		CHECK_EQUAL(PackageCopy, Package);
 		Package.SetObject(FCbObject());
-		TestNotEqual(TEXT("FCbPackage(Object+Attachments).Equals(ObjectNotEqual)"), PackageCopy, Package);
-		TestEqual(TEXT("FCbPackage(Object+Attachments).GetObjectHash(Null)"), Package.GetObjectHash(), FIoHash());
+		CHECK_NOT_EQUAL(PackageCopy, Package);
+		CHECK_EQUAL(Package.GetObjectHash(), FIoHash());
 	}
 
-	// Out of Order
+	SECTION("Out of Order")
 	{
 		TCbWriter<384> Writer;
 		FCbAttachment Attachment2(Level2, Level2Hash);
@@ -505,27 +483,23 @@ TEST_CASE("Core::Serialization::FCbPackage::CbPackage", "[Core][Serialization][S
 		const FCbAttachment* const Level3Attachment = FromFields.FindAttachment(Level3Hash);
 		const FCbAttachment* const Level4Attachment = FromFields.FindAttachment(Level4Hash);
 
-		TestTrue(TEXT("FCbPackage(OutOfOrder).TryLoad()->Level1"), FromFields.GetObject().Equals(Level1));
-		TestEqual(TEXT("FCbPackage(OutOfOrder).TryLoad()->Level1Buffer"),
-			FromFields.GetObject().GetOuterBuffer(), Fields.GetOuterBuffer());
-		TestEqual(TEXT("FCbPackage(OutOfOrder).TryLoad()->Level1Hash"), FromFields.GetObjectHash(), Level1Hash);
-
-		TestTrue(TEXT("FCbPackage(OutOfOrder).TryLoad()->Level2"),
-			Level2Attachment && Level2Attachment->AsObject().Equals(Level2));
-		TestTrue(TEXT("FCbPackage(OutOfOrder).TryLoad()->Level2Hash"),
-			Level2Attachment && Level2Attachment->GetHash() == Level2Hash);
+		CHECK(FromFields.GetObject().Equals(Level1));
+		CHECK_EQUAL(FromFields.GetObject().GetOuterBuffer(), Fields.GetOuterBuffer());
+		CHECK_EQUAL(FromFields.GetObjectHash(), Level1Hash);
 
-		TestTrue(TEXT("FCbPackage(OutOfOrder).TryLoad()->Level3"),
-			Level3Attachment && Level3Attachment->AsObject().Equals(Level3));
-		TestTrue(TEXT("FCbPackage(OutOfOrder).TryLoad()->Level3Hash"),
-			Level3Attachment && Level3Attachment->GetHash() == Level3Hash);
+		REQUIRE(Level2Attachment);
+		
+		CHECK(Level2Attachment->AsObject().Equals(Level2));
+		CHECK(Level2Attachment->GetHash() == Level2Hash);
+		
+		REQUIRE(Level3Attachment);
+		CHECK(Level3Attachment->AsObject().Equals(Level3));
+		CHECK(Level3Attachment->GetHash() == Level3Hash);
 
-		TestTrue(TEXT("FCbPackage(OutOfOrder).TryLoad()->Level4"),
-			Level4Attachment && Level4Attachment->AsBinary().GetView().EqualBytes(Level4.GetView()));
-		TestTrue(TEXT("FCbPackage(OutOfOrder).TryLoad()->Level4Buffer"),
-			Level4Attachment && Fields.GetOuterBuffer().GetView().Contains(Level4Attachment->AsBinary().GetView()));
-		TestTrue(TEXT("FCbPackage(OutOfOrder).TryLoad()->Level4Hash"),
-			Level4Attachment && Level4Attachment->GetHash() == Level4Hash);
+		REQUIRE(Level4Attachment);
+		CHECK(Level4Attachment->AsBinary().GetView().EqualBytes(Level4.GetView()));
+		CHECK(Fields.GetOuterBuffer().GetView().Contains(Level4Attachment->AsBinary().GetView()));
+		CHECK(Level4Attachment->GetHash() == Level4Hash);
 
 		FBufferArchive WriteAr;
 		Writer.Save(WriteAr);
@@ -537,40 +511,40 @@ TEST_CASE("Core::Serialization::FCbPackage::CbPackage", "[Core][Serialization][S
 		FromArchive.Save(Writer);
 		FCbFieldIterator Saved = Writer.Save();
 		FMemoryView View;
-		TestEqual(TEXT("FCbPackage(OutOfOrder).Save()->Level1Hash"), Saved.AsHash(), Level1Hash);
+		CHECK_EQUAL(Saved.AsHash(), Level1Hash);
 		++Saved;
-		TestTrue(TEXT("FCbPackage(OutOfOrder).Save()->Level1"), Saved.AsObject().Equals(Level1));
+		CHECK(Saved.AsObject().Equals(Level1));
 		++Saved;
-		TestEqual(TEXT("FCbPackage(OutOfOrder).Save()->Level2Hash"), Saved.AsObjectAttachment(), Level2Hash);
+		CHECK_EQUAL(Saved.AsObjectAttachment(), Level2Hash);
 		++Saved;
-		TestTrue(TEXT("FCbPackage(OutOfOrder).Save()->Level2"), Saved.AsObject().Equals(Level2));
+		CHECK(Saved.AsObject().Equals(Level2));
 		++Saved;
-		TestEqual(TEXT("FCbPackage(OutOfOrder).Save()->Level3Hash"), Saved.AsObjectAttachment(), Level3Hash);
+		CHECK_EQUAL(Saved.AsObjectAttachment(), Level3Hash);
 		++Saved;
-		TestTrue(TEXT("FCbPackage(OutOfOrder).Save()->Level3"), Saved.AsObject().Equals(Level3));
+		CHECK(Saved.AsObject().Equals(Level3));
 		++Saved;
-		TestEqual(TEXT("FCbPackage(OutOfOrder).Save()->Level4Hash"), Saved.AsBinaryAttachment(), Level4Hash);
+		CHECK_EQUAL(Saved.AsBinaryAttachment(), Level4Hash);
 		++Saved;
 		FSharedBuffer SavedLevel4Buffer = FSharedBuffer::MakeView(Saved.AsBinaryView());
-		TestTrue(TEXT("FCbPackage(OutOfOrder).Save()->Level4"), SavedLevel4Buffer.GetView().EqualBytes(Level4.GetView()));
+		CHECK(SavedLevel4Buffer.GetView().EqualBytes(Level4.GetView()));
 		++Saved;
-		TestTrue(TEXT("FCbPackage(OutOfOrder).Save()->Null"), Saved.IsNull());
+		CHECK(Saved.IsNull());
 		++Saved;
-		TestTrue(TEXT("FCbPackage(OutOfOrder).Save()->AtEnd"), !Saved);
+		CHECK(!Saved);
 	}
 
-	// Null Attachment
+	SECTION("Null Attachment")
 	{
 		const FCbAttachment NullAttachment;
 		FCbPackage Package;
 		Package.AddAttachment(NullAttachment);
-		TestTrue(TEXT("FCbPackage(NullAttachment).IsNull()"), Package.IsNull());
-		TestFalse(TEXT("FCbPackage(NullAttachment) as bool"), bool(Package));
-		TestEqual(TEXT("FCbPackage(NullAttachment).GetAttachments()"), Package.GetAttachments().Num(), 0);
-		TestTrue(TEXT("FCbPackage(NullAttachment).FindAttachment()"), !Package.FindAttachment(NullAttachment));
+		CHECK(Package.IsNull());
+		CHECK_FALSE(bool(Package));
+		CHECK_EQUAL(Package.GetAttachments().Num(), 0);
+		CHECK(!Package.FindAttachment(NullAttachment));
 	}
 
-	// Resolve After Merge
+	SECTION("Resolve After Merge")
 	{
 		bool bResolved = false;
 		FCbPackage Package;
@@ -581,6 +555,6 @@ TEST_CASE("Core::Serialization::FCbPackage::CbPackage", "[Core][Serialization][S
 				bResolved = true;
 				return FSharedBuffer();
 			});
-		TestTrue(TEXT("FCbPackage(ResolveAfterMerge)->Resolved"), bResolved);
+		CHECK(bResolved);
 	}
 }
\ No newline at end of file
diff --git a/Engine/Source/Runtime/Core/Tests/Serialization/CompactBinarySerializationTest.cpp b/Engine/Source/Runtime/Core/Tests/Serialization/CompactBinarySerializationTest.cpp
--- a/Engine/Source/Runtime/Core/Tests/Serialization/CompactBinarySerializationTest.cpp
+++ b/Engine/Source/Runtime/Core/Tests/Serialization/CompactBinarySerializationTest.cpp
@@ -7,9 +7,9 @@
 #include "Serialization/MemoryReader.h"
 #include "TestHarness.h"
 
-TEST_CASE("Core::Serialization::FCbMeasure::Measure Compact Binary", "[Core][Serialization][Smoke]")
+TEST_CASE_METHOD(FAutomationTestFixture, "Core::Serialization::FCbMeasure::Measure Compact Binary", "[Core][Serialization][Smoke]")
 {
-	const auto TestMeasure = [](
+	const auto TestMeasure = [this](
 		const TCHAR* Test,
 		std::initializer_list<uint8> Data,
 		bool bExpectedValue,
@@ -20,13 +20,15 @@ TEST_CASE("Core::Serialization::FCbMeasure::Measure Compact Binary", "[Core][Ser
 		ECbFieldType ActualType = ECbFieldType(255);
 		uint64 ActualSize = ~uint64(0);
 		const bool bActualValue = TryMeasureCompactBinary(MakeMemoryView(Data), ActualType, ActualSize, ExternalType);
-		TestEqual(FString::Printf(TEXT("TryMeasureCompactBinary(%s)"), Test), bActualValue, bExpectedValue);
-		TestEqual(FString::Printf(TEXT("TryMeasureCompactBinary(%s)->Type"), Test), ActualType, ExpectedType);
-		TestEqual(FString::Printf(TEXT("TryMeasureCompactBinary(%s)->Size"), Test), ActualSize, ExpectedSize);
+		
+		INFO(Test)
+		CHECK_EQUAL(bActualValue, bExpectedValue);
+		CHECK_EQUAL(ActualType, ExpectedType);
+		CHECK_EQUAL(ActualSize, ExpectedSize);
 		if (ActualSize == ExpectedSize)
 		{
 			const uint64 MeasureSize = MeasureCompactBinary(MakeMemoryView(Data), ExternalType);
-			TestEqual(FString::Printf(TEXT("MeasureCompactBinary(%s)"), Test), MeasureSize, bExpectedValue ? ExpectedSize : 0);
+			CHECK_EQUAL(MeasureSize, (bExpectedValue ? ExpectedSize : 0));
 		}
 	};
 
@@ -130,34 +132,34 @@ TEST_CASE("Core::Serialization::FCbMeasure::Measure Compact Binary", "[Core][Ser
 	TestMeasure(TEXT("CustomByName, Size9BShort"), {uint8(EType::CustomByName), 0xff}, false, 10, EType::CustomByName);
 }
 
-TEST_CASE("Core::Serialization::FMemoryView::Save Compact Binary", "[Core][Serialization][Smoke]")
+TEST_CASE_METHOD(FAutomationTestFixture, "Core::Serialization::FMemoryView::Save Compact Binary", "[Core][Serialization][Smoke]")
 {
-	const auto TestSave = [](const TCHAR* Test, auto Value, FMemoryView ExpectedData)
+	const auto TestSave = [this](const TCHAR* Test, auto Value, FMemoryView ExpectedData)
 	{
 		{
 			FBufferArchive WriteAr;
 			SaveCompactBinary(WriteAr, Value);
 
-			TestEqual(FString::Printf(TEXT("SaveCompactBinary(%s)->Size"), Test), uint64(WriteAr.Num()), ExpectedData.GetSize());
-			TestTrue(FString::Printf(TEXT("SaveCompactBinary(%s)->EqualBytes"), Test), ExpectedData.EqualBytes(MakeMemoryView(WriteAr)));
+			TEST_EQUAL(FString::Printf(TEXT("SaveCompactBinary(%s)->Size"), Test), uint64(WriteAr.Num()), ExpectedData.GetSize());
+			TEST_TRUE(FString::Printf(TEXT("SaveCompactBinary(%s)->EqualBytes"), Test), ExpectedData.EqualBytes(MakeMemoryView(WriteAr)));
 
 			{
 				FMemoryReader ReadAr(WriteAr);
 				FCbField Field = LoadCompactBinary(ReadAr);
-				TestTrue(FString::Printf(TEXT("LoadCompactBinary(%s)->EqualBytes"), Test), ExpectedData.EqualBytes(Field.GetOuterBuffer().GetView()));
+				TEST_TRUE(FString::Printf(TEXT("LoadCompactBinary(%s)->EqualBytes"), Test), ExpectedData.EqualBytes(Field.GetOuterBuffer().GetView()));
 			}
 		}
 		{
 			FBufferArchive WriteAr;
 			WriteAr << Value;
 			
-			TestEqual(FString::Printf(TEXT("Ar << CompactBinary Save(%s)->Size"), Test), uint64(WriteAr.Num()), ExpectedData.GetSize());
-			TestTrue(FString::Printf(TEXT("Ar << CompactBinary Save(%s)->EqualBytes"), Test), ExpectedData.EqualBytes(MakeMemoryView(WriteAr)));
+			TEST_EQUAL(FString::Printf(TEXT("Ar << CompactBinary Save(%s)->Size"), Test), uint64(WriteAr.Num()), ExpectedData.GetSize());
+			TEST_TRUE(FString::Printf(TEXT("Ar << CompactBinary Save(%s)->EqualBytes"), Test), ExpectedData.EqualBytes(MakeMemoryView(WriteAr)));
 
 			{
 				FMemoryReader ReadAr(WriteAr);
 				ReadAr << Value;
-				TestTrue(FString::Printf(TEXT("Ar << CompactBinary Load(%s)->EqualBytes"), Test), ExpectedData.EqualBytes(Value.GetOuterBuffer().GetView()));
+				TEST_TRUE(FString::Printf(TEXT("Ar << CompactBinary Load(%s)->EqualBytes"), Test), ExpectedData.EqualBytes(Value.GetOuterBuffer().GetView()));
 			}
 		}
 	};
