Add new Unit Tests

diff --git a/Engine/Source/Runtime/Core/Tests/Serialization/CompactBinaryTest.cpp b/Engine/Source/Runtime/Core/Tests/Serialization/CompactBinaryTest.cpp
--- a/Engine/Source/Runtime/Core/Tests/Serialization/CompactBinaryTest.cpp
+++ b/Engine/Source/Runtime/Core/Tests/Serialization/CompactBinaryTest.cpp
@@ -67,94 +67,94 @@ struct FCbAttachmentAccessors
 	static constexpr FIoHash (FCbFieldView::*AsType)(const FIoHash&) = &FCbFieldView::AsAttachment;
 };
 
-void TestEqualCb(const FString& What, const FCbArrayView& Actual, const FCbArrayView& Expected)
+void TEST_EQUALCb(FAutomationTestFixture& Test, const FString& What, const FCbArrayView& Actual, const FCbArrayView& Expected)
 {
-	TestTrue(What, Actual.Equals(Expected));
+	TEST_TRUE(What, Actual.Equals(Expected));
 }
 
-void TestEqualCb(const FString& What, const FCbObjectView& Actual, const FCbObjectView& Expected)
+void TEST_EQUALCb(FAutomationTestFixture& Test, const FString& What, const FCbObjectView& Actual, const FCbObjectView& Expected)
 {
-	TestTrue(What, Actual.Equals(Expected));
+	TEST_TRUE(What, Actual.Equals(Expected));
 }
 
-void TestEqualCb(const FString& What, const FCbCustomById& Actual, const FCbCustomById& Expected)
+void TEST_EQUALCb(FAutomationTestFixture& Test, const FString& What, const FCbCustomById& Actual, const FCbCustomById& Expected)
 {
-	TestTrue(What, Actual.Id == Expected.Id && Actual.Data == Expected.Data);
+	TEST_TRUE(What, Actual.Id == Expected.Id && Actual.Data == Expected.Data);
 }
 
-void TestEqualCb(const FString& What, const FCbCustomByName& Actual, const FCbCustomByName& Expected)
+void TEST_EQUALCb(FAutomationTestFixture& Test, const FString& What, const FCbCustomByName& Actual, const FCbCustomByName& Expected)
 {
-	TestTrue(What, Actual.Name.Equals(Expected.Name, ESearchCase::CaseSensitive) && Actual.Data == Expected.Data);
+	TEST_TRUE(What, Actual.Name.Equals(Expected.Name, ESearchCase::CaseSensitive) && Actual.Data == Expected.Data);
 }
 
-void TestEqualBytes(const TCHAR* What, FMemoryView Actual, TArrayView<const uint8> Expected)
+void TEST_EQUALBytes(FAutomationTestFixture& Test, const TCHAR* What, FMemoryView Actual, TArrayView<const uint8> Expected)
 {
-	TestTrue(What, Actual.EqualBytes(MakeMemoryView(Expected)));
+	TEST_TRUE(What, Actual.EqualBytes(MakeMemoryView(Expected)));
 }
 
 template <typename A, typename B>
-void TestEqualCb(const FString& What, const A& Actual, const B& Expected)
+void TEST_EQUALCb(FAutomationTestFixture& Test, const FString& What, const A& Actual, const B& Expected)
 {
-	TestEqual(What, Actual, Expected);
+	TEST_EQUAL(What, Actual, Expected);
 }
 
 template <ECbFieldType FieldType, typename T = TCbFieldValueType<FieldType>, typename AccessorsType = TCbFieldTypeAccessors<FieldType>>
-void TestFieldNoClone(const TCHAR* What, FCbFieldView& Field, T ExpectedValue = T(), T DefaultValue = T(), ECbFieldError ExpectedError = ECbFieldError::None, const AccessorsType& Accessors = AccessorsType())
+void TestFieldNoClone(FAutomationTestFixture& Test, const TCHAR* What, FCbFieldView& Field, T ExpectedValue = T(), T DefaultValue = T(), ECbFieldError ExpectedError = ECbFieldError::None, const AccessorsType& Accessors = AccessorsType())
 {
-	TestEqual(FString::Printf(TEXT("FCbFieldView::Is[Type](%s)"), What), Invoke(Accessors.IsType, Field), ExpectedError != ECbFieldError::TypeError);
+	TEST_EQUAL(FString::Printf(TEXT("FCbFieldView::Is[Type](%s)"), What), Invoke(Accessors.IsType, Field), ExpectedError != ECbFieldError::TypeError);
 	if (ExpectedError == ECbFieldError::None && !Field.IsBool())
 	{
-		TestFalse(FString::Printf(TEXT("FCbFieldView::AsBool(%s) == false"), What), Field.AsBool());
-		TestTrue(FString::Printf(TEXT("FCbFieldView::AsBool(%s) -> HasError()"), What), Field.HasError());
-		TestEqual(FString::Printf(TEXT("FCbFieldView::AsBool(%s) -> GetError() == TypeError"), What), Field.GetError(), ECbFieldError::TypeError);
+		TEST_FALSE(FString::Printf(TEXT("FCbFieldView::AsBool(%s) == false"), What), Field.AsBool());
+		TEST_TRUE(FString::Printf(TEXT("FCbFieldView::AsBool(%s) -> HasError()"), What), Field.HasError());
+		TEST_EQUAL(FString::Printf(TEXT("FCbFieldView::AsBool(%s) -> GetError() == TypeError"), What), Field.GetError(), ECbFieldError::TypeError);
 	}
-	TestEqualCb(FString::Printf(TEXT("FCbFieldView::As[Type](%s) -> Equal"), What), Invoke(Accessors.AsType, Field, DefaultValue), ExpectedValue);
-	TestEqualCb(FString::Printf(TEXT("FCbFieldView::As[Type](%s) -> HasError()"), What), Field.HasError(), ExpectedError != ECbFieldError::None);
-	TestEqualCb(FString::Printf(TEXT("FCbFieldView::As[Type](%s) -> GetError()"), What), Field.GetError(), ExpectedError);
+	TEST_EQUALCb(Test, FString::Printf(TEXT("FCbFieldView::As[Type](%s) -> Equal"), What), Invoke(Accessors.AsType, Field, DefaultValue), ExpectedValue);
+	TEST_EQUALCb(Test, FString::Printf(TEXT("FCbFieldView::As[Type](%s) -> HasError()"), What), Field.HasError(), ExpectedError != ECbFieldError::None);
+	TEST_EQUALCb(Test, FString::Printf(TEXT("FCbFieldView::As[Type](%s) -> GetError()"), What), Field.GetError(), ExpectedError);
 }
 
 template <ECbFieldType FieldType, typename T = TCbFieldValueType<FieldType>, typename AccessorsType = TCbFieldTypeAccessors<FieldType>>
-void TestFieldNoClone(const TCHAR* What, TArrayView<const uint8> Value, T ExpectedValue = T(), T DefaultValue = T(), ECbFieldError ExpectedError = ECbFieldError::None, const AccessorsType& Accessors = AccessorsType())
+void TestFieldNoClone(FAutomationTestFixture& Test, const TCHAR* What, TArrayView<const uint8> Value, T ExpectedValue = T(), T DefaultValue = T(), ECbFieldError ExpectedError = ECbFieldError::None, const AccessorsType& Accessors = AccessorsType())
 {
 	FCbFieldView Field(Value.GetData(), FieldType);
-	TestFieldNoClone<FieldType>(What, Field, ExpectedValue, DefaultValue, ExpectedError, Accessors);
+	TestFieldNoClone<FieldType>(Test, What, Field, ExpectedValue, DefaultValue, ExpectedError, Accessors);
 }
 
 template <ECbFieldType FieldType, typename T = TCbFieldValueType<FieldType>, typename AccessorsType = TCbFieldTypeAccessors<FieldType>>
-void TestField(const TCHAR* What, FCbFieldView& Field, T ExpectedValue = T(), T DefaultValue = T(), ECbFieldError ExpectedError = ECbFieldError::None, const AccessorsType& Accessors = AccessorsType())
+void TestField(FAutomationTestFixture& Test, const TCHAR* What, FCbFieldView& Field, T ExpectedValue = T(), T DefaultValue = T(), ECbFieldError ExpectedError = ECbFieldError::None, const AccessorsType& Accessors = AccessorsType())
 {
-	TestFieldNoClone<FieldType>(What, Field, ExpectedValue, DefaultValue, ExpectedError, Accessors);
+	TestFieldNoClone<FieldType>(Test, What, Field, ExpectedValue, DefaultValue, ExpectedError, Accessors);
 	FCbField FieldClone = FCbField::Clone(Field);
-	TestFieldNoClone<FieldType>(*FString::Printf(TEXT("%s, Clone"), What), FieldClone, ExpectedValue, DefaultValue, ExpectedError, Accessors);
-	TestTrue(FString::Printf(TEXT("FCbFieldView::Equals(%s)"), What), Field.Equals(FieldClone));
+	TestFieldNoClone<FieldType>(Test, *FString::Printf(TEXT("%s, Clone"), What), FieldClone, ExpectedValue, DefaultValue, ExpectedError, Accessors);
+	TEST_TRUE(FString::Printf(TEXT("FCbFieldView::Equals(%s)"), What), Field.Equals(FieldClone));
 }
 
 template <ECbFieldType FieldType, typename T = TCbFieldValueType<FieldType>, typename AccessorsType = TCbFieldTypeAccessors<FieldType>>
-void TestField(const TCHAR* What, TArrayView<const uint8> Value, T ExpectedValue = T(), T DefaultValue = T(), ECbFieldError ExpectedError = ECbFieldError::None, const AccessorsType& Accessors = AccessorsType())
+void TestField(FAutomationTestFixture& Test, const TCHAR* What, TArrayView<const uint8> Value, T ExpectedValue = T(), T DefaultValue = T(), ECbFieldError ExpectedError = ECbFieldError::None, const AccessorsType& Accessors = AccessorsType())
 {
 	FCbFieldView Field(Value.GetData(), FieldType);
-	TestEqual(FString::Printf(TEXT("FCbFieldView::GetSize(%s)"), What), Field.GetSize(), uint64(Value.Num() + !FCbFieldType::HasFieldType(FieldType)));
-	TestTrue(FString::Printf(TEXT("FCbFieldView::HasValue(%s)"), What), Field.HasValue());
-	TestFalse(FString::Printf(TEXT("FCbFieldView::HasError(%s) == false"), What), Field.HasError());
-	TestEqual(FString::Printf(TEXT("FCbFieldView::GetError(%s) == None"), What), Field.GetError(), ECbFieldError::None);
-	TestField<FieldType>(What, Field, ExpectedValue, DefaultValue, ExpectedError, Accessors);
+	TEST_EQUAL(FString::Printf(TEXT("FCbFieldView::GetSize(%s)"), What), Field.GetSize(), uint64(Value.Num() + !FCbFieldType::HasFieldType(FieldType)));
+	TEST_TRUE(FString::Printf(TEXT("FCbFieldView::HasValue(%s)"), What), Field.HasValue());
+	TEST_FALSE(FString::Printf(TEXT("FCbFieldView::HasError(%s) == false"), What), Field.HasError());
+	TEST_EQUAL(FString::Printf(TEXT("FCbFieldView::GetError(%s) == None"), What), Field.GetError(), ECbFieldError::None);
+	TestField<FieldType>(Test, What, Field, ExpectedValue, DefaultValue, ExpectedError, Accessors);
 }
 
 template <ECbFieldType FieldType, typename T = TCbFieldValueType<FieldType>, typename AccessorsType = TCbFieldTypeAccessors<FieldType>>
-void TestFieldError(const TCHAR* What, FCbFieldView& Field, ECbFieldError ExpectedError, T ExpectedValue = T(), const AccessorsType& Accessors = AccessorsType())
+void TestFieldError(FAutomationTestFixture& Test, const TCHAR* What, FCbFieldView& Field, ECbFieldError ExpectedError, T ExpectedValue = T(), const AccessorsType& Accessors = AccessorsType())
 {
-	TestFieldNoClone<FieldType>(What, Field, ExpectedValue, ExpectedValue, ExpectedError, Accessors);
+	TestFieldNoClone<FieldType>(Test, What, Field, ExpectedValue, ExpectedValue, ExpectedError, Accessors);
 }
 
 template <ECbFieldType FieldType, typename T = TCbFieldValueType<FieldType>, typename AccessorsType = TCbFieldTypeAccessors<FieldType>>
-void TestFieldError(const TCHAR* What, TArrayView<const uint8> Value, ECbFieldError ExpectedError, T ExpectedValue = T(), const AccessorsType& Accessors = AccessorsType())
+void TestFieldError(FAutomationTestFixture& Test, const TCHAR* What, TArrayView<const uint8> Value, ECbFieldError ExpectedError, T ExpectedValue = T(), const AccessorsType& Accessors = AccessorsType())
 {
 	FCbFieldView Field(Value.GetData(), FieldType);
-	TestFieldError<FieldType>(What, Field, ExpectedError, ExpectedValue, Accessors);
+	TestFieldError<FieldType>(Test, What, Field, ExpectedError, ExpectedValue, Accessors);
 }
 
 
-TEST_CASE("Core::Serialization::FCbField::Ctor Test", "[Core][Serialization][Smoke]")
+TEST_CASE_METHOD(FAutomationTestFixture, "Core::Serialization::FCbField::Ctor Test", "[Core][Serialization][Smoke]")
 {
 	// Test FCbFieldView()
 	{
@@ -163,29 +163,29 @@ TEST_CASE("Core::Serialization::FCbField::Ctor Test", "[Core][Serialization][Smo
 		static_assert(!DefaultField.HasValue(), "Error in HasValue()");
 		static_assert(!DefaultField.HasError(), "Error in HasError()");
 		static_assert(DefaultField.GetError() == ECbFieldError::None, "Error in GetError()");
-		TestEqual(TEXT("FCbFieldView()::GetSize() == 1"), DefaultField.GetSize(), uint64(1));
-		TestEqual(TEXT("FCbFieldView()::GetName().Len() == 0"), DefaultField.GetName().Len(), 0);
-		TestFalse(TEXT("!FCbFieldView()::HasName()"), DefaultField.HasName());
-		TestFalse(TEXT("!FCbFieldView()::HasValue()"), DefaultField.HasValue());
-		TestFalse(TEXT("!FCbFieldView()::HasError()"), DefaultField.HasError());
-		TestEqual(TEXT("FCbFieldView()::GetError() == None"), DefaultField.GetError(), ECbFieldError::None);
-		TestEqual(TEXT("FCbFieldView()::GetHash()"), DefaultField.GetHash(), FIoHash::HashBuffer(MakeMemoryView<uint8>({uint8(ECbFieldType::None)})));
+		TEST_EQUAL(TEXT("FCbFieldView()::GetSize() == 1"), DefaultField.GetSize(), uint64(1));
+		TEST_EQUAL(TEXT("FCbFieldView()::GetName().Len() == 0"), DefaultField.GetName().Len(), 0);
+		TEST_FALSE(TEXT("!FCbFieldView()::HasName()"), DefaultField.HasName());
+		TEST_FALSE(TEXT("!FCbFieldView()::HasValue()"), DefaultField.HasValue());
+		TEST_FALSE(TEXT("!FCbFieldView()::HasError()"), DefaultField.HasError());
+		TEST_EQUAL(TEXT("FCbFieldView()::GetError() == None"), DefaultField.GetError(), ECbFieldError::None);
+		TEST_EQUAL(TEXT("FCbFieldView()::GetHash()"), DefaultField.GetHash(), FIoHash::HashBuffer(MakeMemoryView<uint8>({uint8(ECbFieldType::None)})));
 		FMemoryView View;
-		TestFalse(TEXT("FCbFieldView()::TryGetView()"), DefaultField.TryGetView(View));
+		TEST_FALSE(TEXT("FCbFieldView()::TryGetView()"), DefaultField.TryGetView(View));
 	}
 
 	// Test FCbFieldView(None)
 	{
 		FCbFieldView NoneField(nullptr, ECbFieldType::None);
-		TestEqual(TEXT("FCbFieldView(None)::GetSize() == 1"), NoneField.GetSize(), uint64(1));
-		TestEqual(TEXT("FCbFieldView(None)::GetName().Len() == 0"), NoneField.GetName().Len(), 0);
-		TestFalse(TEXT("!FCbFieldView(None)::HasName()"), NoneField.HasName());
-		TestFalse(TEXT("!FCbFieldView(None)::HasValue()"), NoneField.HasValue());
-		TestFalse(TEXT("!FCbFieldView(None)::HasError()"), NoneField.HasError());
-		TestEqual(TEXT("FCbFieldView(None)::GetError() == None"), NoneField.GetError(), ECbFieldError::None);
-		TestEqual(TEXT("FCbFieldView(None)::GetHash()"), NoneField.GetHash(), FCbFieldView().GetHash());
+		TEST_EQUAL(TEXT("FCbFieldView(None)::GetSize() == 1"), NoneField.GetSize(), uint64(1));
+		TEST_EQUAL(TEXT("FCbFieldView(None)::GetName().Len() == 0"), NoneField.GetName().Len(), 0);
+		TEST_FALSE(TEXT("!FCbFieldView(None)::HasName()"), NoneField.HasName());
+		TEST_FALSE(TEXT("!FCbFieldView(None)::HasValue()"), NoneField.HasValue());
+		TEST_FALSE(TEXT("!FCbFieldView(None)::HasError()"), NoneField.HasError());
+		TEST_EQUAL(TEXT("FCbFieldView(None)::GetError() == None"), NoneField.GetError(), ECbFieldError::None);
+		TEST_EQUAL(TEXT("FCbFieldView(None)::GetHash()"), NoneField.GetHash(), FCbFieldView().GetHash());
 		FMemoryView View;
-		TestFalse(TEXT("FCbFieldView(None)::TryGetView()"), NoneField.TryGetView(View));
+		TEST_FALSE(TEXT("FCbFieldView(None)::TryGetView()"), NoneField.TryGetView(View));
 	}
 
 	// Test FCbFieldView(None|Type|Name)
@@ -193,17 +193,17 @@ TEST_CASE("Core::Serialization::FCbField::Ctor Test", "[Core][Serialization][Smo
 		constexpr ECbFieldType FieldType = ECbFieldType::None | ECbFieldType::HasFieldName;
 		constexpr const ANSICHAR NoneBytes[] = { ANSICHAR(FieldType), 4, 'N', 'a', 'm', 'e' };
 		FCbFieldView NoneField(NoneBytes);
-		TestEqual(TEXT("FCbFieldView(None|Type|Name)::GetSize()"), NoneField.GetSize(), uint64(sizeof(NoneBytes)));
-		TestEqual(TEXT("FCbFieldView(None|Type|Name)::GetName()"), NoneField.GetName(), "Name"_U8SV);
-		TestTrue(TEXT("FCbFieldView(None|Type|Name)::HasName()"), NoneField.HasName());
-		TestFalse(TEXT("!FCbFieldView(None|Type|Name)::HasValue()"), NoneField.HasValue());
-		TestEqual(TEXT("FCbFieldView(None|Type|Name)::GetHash()"), NoneField.GetHash(), FIoHash::HashBuffer(MakeMemoryView(NoneBytes)));
+		TEST_EQUAL(TEXT("FCbFieldView(None|Type|Name)::GetSize()"), NoneField.GetSize(), uint64(sizeof(NoneBytes)));
+		TEST_EQUAL(TEXT("FCbFieldView(None|Type|Name)::GetName()"), NoneField.GetName(), "Name"_U8SV);
+		TEST_TRUE(TEXT("FCbFieldView(None|Type|Name)::HasName()"), NoneField.HasName());
+		TEST_FALSE(TEXT("!FCbFieldView(None|Type|Name)::HasValue()"), NoneField.HasValue());
+		TEST_EQUAL(TEXT("FCbFieldView(None|Type|Name)::GetHash()"), NoneField.GetHash(), FIoHash::HashBuffer(MakeMemoryView(NoneBytes)));
 		FMemoryView View;
-		TestTrue(TEXT("FCbFieldView(None|Type|Name)::TryGetView()"), NoneField.TryGetView(View) && View == MakeMemoryView(NoneBytes));
+		TEST_TRUE(TEXT("FCbFieldView(None|Type|Name)::TryGetView()"), NoneField.TryGetView(View) && View == MakeMemoryView(NoneBytes));
 
 		uint8 CopyBytes[sizeof(NoneBytes)];
 		NoneField.CopyTo(MakeMemoryView(CopyBytes));
-		TestTrue(TEXT("FCbFieldView(None|Type|Name)::CopyTo()"), MakeMemoryView(NoneBytes).EqualBytes(MakeMemoryView(CopyBytes)));
+		TEST_TRUE(TEXT("FCbFieldView(None|Type|Name)::CopyTo()"), MakeMemoryView(NoneBytes).EqualBytes(MakeMemoryView(CopyBytes)));
 	}
 
 	// Test FCbFieldView(None|Type)
@@ -211,13 +211,13 @@ TEST_CASE("Core::Serialization::FCbField::Ctor Test", "[Core][Serialization][Smo
 		constexpr ECbFieldType FieldType = ECbFieldType::None;
 		constexpr const ANSICHAR NoneBytes[] = { ANSICHAR(FieldType) };
 		FCbFieldView NoneField(NoneBytes);
-		TestEqual(TEXT("FCbFieldView(None|Type)::GetSize()"), NoneField.GetSize(), uint64(sizeof(NoneBytes)));
-		TestEqual(TEXT("FCbFieldView(None|Type)::GetName()"), NoneField.GetName().Len(), 0);
-		TestFalse(TEXT("FCbFieldView(None|Type)::HasName()"), NoneField.HasName());
-		TestFalse(TEXT("!FCbFieldView(None|Type)::HasValue()"), NoneField.HasValue());
-		TestEqual(TEXT("FCbFieldView(None|Type)::GetHash()"), NoneField.GetHash(), FCbFieldView().GetHash());
+		TEST_EQUAL(TEXT("FCbFieldView(None|Type)::GetSize()"), NoneField.GetSize(), uint64(sizeof(NoneBytes)));
+		TEST_EQUAL(TEXT("FCbFieldView(None|Type)::GetName()"), NoneField.GetName().Len(), 0);
+		TEST_FALSE(TEXT("FCbFieldView(None|Type)::HasName()"), NoneField.HasName());
+		TEST_FALSE(TEXT("!FCbFieldView(None|Type)::HasValue()"), NoneField.HasValue());
+		TEST_EQUAL(TEXT("FCbFieldView(None|Type)::GetHash()"), NoneField.GetHash(), FCbFieldView().GetHash());
 		FMemoryView View;
-		TestTrue(TEXT("FCbFieldView(None|Type)::TryGetView()"), NoneField.TryGetView(View) && View == MakeMemoryView(NoneBytes));
+		TEST_TRUE(TEXT("FCbFieldView(None|Type)::TryGetView()"), NoneField.TryGetView(View) && View == MakeMemoryView(NoneBytes));
 	}
 
 	// Test FCbFieldView(None|Name)
@@ -225,17 +225,17 @@ TEST_CASE("Core::Serialization::FCbField::Ctor Test", "[Core][Serialization][Smo
 		constexpr ECbFieldType FieldType = ECbFieldType::None | ECbFieldType::HasFieldName;
 		constexpr const ANSICHAR NoneBytes[] = { ANSICHAR(FieldType), 4, 'N', 'a', 'm', 'e' };
 		FCbFieldView NoneField(NoneBytes + 1, FieldType);
-		TestEqual(TEXT("FCbFieldView(None|Name)::GetSize()"), NoneField.GetSize(), uint64(sizeof(NoneBytes)));
-		TestEqual(TEXT("FCbFieldView(None|Name)::GetName()"), NoneField.GetName(), "Name"_U8SV);
-		TestTrue(TEXT("FCbFieldView(None|Name)::HasName()"), NoneField.HasName());
-		TestFalse(TEXT("!FCbFieldView(None|Name)::HasValue()"), NoneField.HasValue());
-		TestEqual(TEXT("FCbFieldView(None|Name)::GetHash()"), NoneField.GetHash(), FIoHash::HashBuffer(MakeMemoryView(NoneBytes)));
+		TEST_EQUAL(TEXT("FCbFieldView(None|Name)::GetSize()"), NoneField.GetSize(), uint64(sizeof(NoneBytes)));
+		TEST_EQUAL(TEXT("FCbFieldView(None|Name)::GetName()"), NoneField.GetName(), "Name"_U8SV);
+		TEST_TRUE(TEXT("FCbFieldView(None|Name)::HasName()"), NoneField.HasName());
+		TEST_FALSE(TEXT("!FCbFieldView(None|Name)::HasValue()"), NoneField.HasValue());
+		TEST_EQUAL(TEXT("FCbFieldView(None|Name)::GetHash()"), NoneField.GetHash(), FIoHash::HashBuffer(MakeMemoryView(NoneBytes)));
 		FMemoryView View;
-		TestFalse(TEXT("FCbFieldView(None|Name)::TryGetView()"), NoneField.TryGetView(View));
+		TEST_FALSE(TEXT("FCbFieldView(None|Name)::TryGetView()"), NoneField.TryGetView(View));
 
 		uint8 CopyBytes[sizeof(NoneBytes)];
 		NoneField.CopyTo(MakeMemoryView(CopyBytes));
-		TestTrue(TEXT("FCbFieldView(None|Name)::CopyTo()"), MakeMemoryView(NoneBytes).EqualBytes(MakeMemoryView(CopyBytes)));
+		TEST_TRUE(TEXT("FCbFieldView(None|Name)::CopyTo()"), MakeMemoryView(NoneBytes).EqualBytes(MakeMemoryView(CopyBytes)));
 	}
 
 	// Test FCbFieldView(None|EmptyName)
@@ -243,77 +243,77 @@ TEST_CASE("Core::Serialization::FCbField::Ctor Test", "[Core][Serialization][Smo
 		constexpr ECbFieldType FieldType = ECbFieldType::None | ECbFieldType::HasFieldName;
 		constexpr const uint8 NoneBytes[] = { uint8(FieldType), 0 };
 		FCbFieldView NoneField(NoneBytes + 1, FieldType);
-		TestEqual(TEXT("FCbFieldView(None|EmptyName)::GetSize()"), NoneField.GetSize(), uint64(sizeof(NoneBytes)));
-		TestEqual(TEXT("FCbFieldView(None|EmptyName)::GetName()"), NoneField.GetName(), ""_U8SV);
-		TestTrue(TEXT("FCbFieldView(None|EmptyName)::HasName()"), NoneField.HasName());
-		TestFalse(TEXT("!FCbFieldView(None|EmptyName)::HasValue()"), NoneField.HasValue());
-		TestEqual(TEXT("FCbFieldView(None|EmptyName)::GetHash()"), NoneField.GetHash(), FIoHash::HashBuffer(MakeMemoryView(NoneBytes)));
+		TEST_EQUAL(TEXT("FCbFieldView(None|EmptyName)::GetSize()"), NoneField.GetSize(), uint64(sizeof(NoneBytes)));
+		TEST_EQUAL(TEXT("FCbFieldView(None|EmptyName)::GetName()"), NoneField.GetName(), ""_U8SV);
+		TEST_TRUE(TEXT("FCbFieldView(None|EmptyName)::HasName()"), NoneField.HasName());
+		TEST_FALSE(TEXT("!FCbFieldView(None|EmptyName)::HasValue()"), NoneField.HasValue());
+		TEST_EQUAL(TEXT("FCbFieldView(None|EmptyName)::GetHash()"), NoneField.GetHash(), FIoHash::HashBuffer(MakeMemoryView(NoneBytes)));
 		FMemoryView View;
-		TestFalse(TEXT("FCbFieldView(None|EmptyName)::TryGetView()"), NoneField.TryGetView(View));
+		TEST_FALSE(TEXT("FCbFieldView(None|EmptyName)::TryGetView()"), NoneField.TryGetView(View));
 	}
 }
 
-TEST_CASE("Core::Serialization::FCbField::Null Test", "[Core][Serialization][Smoke]")
+TEST_CASE_METHOD(FAutomationTestFixture, "Core::Serialization::FCbField::Null Test", "[Core][Serialization][Smoke]")
 {
 	// Test FCbFieldView(Null)
 	{
 		FCbFieldView NullField(nullptr, ECbFieldType::Null);
-		TestEqual(TEXT("FCbFieldView(Null)::GetSize() == 1"), NullField.GetSize(), uint64(1));
-		TestTrue(TEXT("FCbFieldView(Null)::IsNull()"), NullField.IsNull());
-		TestTrue(TEXT("FCbFieldView(Null)::HasValue()"), NullField.HasValue());
-		TestFalse(TEXT("!FCbFieldView(Null)::HasError()"), NullField.HasError());
-		TestEqual(TEXT("FCbFieldView(Null)::GetError() == None"), NullField.GetError(), ECbFieldError::None);
-		TestEqual(TEXT("FCbFieldView(Null)::GetHash()"), NullField.GetHash(), FIoHash::HashBuffer(MakeMemoryView<uint8>({uint8(ECbFieldType::Null)})));
+		TEST_EQUAL(TEXT("FCbFieldView(Null)::GetSize() == 1"), NullField.GetSize(), uint64(1));
+		TEST_TRUE(TEXT("FCbFieldView(Null)::IsNull()"), NullField.IsNull());
+		TEST_TRUE(TEXT("FCbFieldView(Null)::HasValue()"), NullField.HasValue());
+		TEST_FALSE(TEXT("!FCbFieldView(Null)::HasError()"), NullField.HasError());
+		TEST_EQUAL(TEXT("FCbFieldView(Null)::GetError() == None"), NullField.GetError(), ECbFieldError::None);
+		TEST_EQUAL(TEXT("FCbFieldView(Null)::GetHash()"), NullField.GetHash(), FIoHash::HashBuffer(MakeMemoryView<uint8>({uint8(ECbFieldType::Null)})));
 	}
 
 	// Test FCbFieldView(None) as Null
 	{
 		FCbFieldView Field;
-		TestFalse(TEXT("FCbFieldView(None)::IsNull()"), Field.IsNull());
+		TEST_FALSE(TEXT("FCbFieldView(None)::IsNull()"), Field.IsNull());
 	}
 }
 
-TEST_CASE("Core::Serialization::FCbField::Object Test", "[Core][Serialization][Smoke]")
+TEST_CASE_METHOD(FAutomationTestFixture, "Core::Serialization::FCbField::Object Test", "[Core][Serialization][Smoke]")
 {
 	static_assert(!std::is_constructible<FCbFieldView, const FCbObjectView&>::value, "Invalid constructor for FCbFieldView");
 	static_assert(!std::is_assignable<FCbFieldView, const FCbObjectView&>::value, "Invalid assignment for FCbFieldView");
 	static_assert(!std::is_convertible<FCbFieldView, FCbObjectView>::value, "Invalid conversion to FCbObjectView");
 	static_assert(!std::is_assignable<FCbObjectView, const FCbFieldView&>::value, "Invalid assignment for FCbObjectView");
 
-	auto TestIntObject = [](const FCbObjectView& Object, int32 ExpectedNum, uint64 ExpectedValueSize)
+	auto TestIntObject = [this](const FCbObjectView& Object, int32 ExpectedNum, uint64 ExpectedValueSize)
 	{
-		TestEqual(TEXT("FCbFieldView(Object)::AsObjectView().GetSize()"), Object.GetSize(), ExpectedValueSize + sizeof(ECbFieldType));
+		TEST_EQUAL(TEXT("FCbFieldView(Object)::AsObjectView().GetSize()"), Object.GetSize(), ExpectedValueSize + sizeof(ECbFieldType));
 
 		int32 ActualNum = 0;
 		for (FCbFieldViewIterator It = Object.CreateViewIterator(); It; ++It)
 		{
 			++ActualNum;
 			TestNotEqual(TEXT("FCbFieldView(Object) Iterator Name"), It->GetName().Len(), 0);
-			TestEqual(TEXT("FCbFieldView(Object) Iterator"), It->AsInt32(), ActualNum);
+			TEST_EQUAL(TEXT("FCbFieldView(Object) Iterator"), It->AsInt32(), ActualNum);
 		}
-		TestEqual(TEXT("FCbFieldView(Object)::AsObjectView().CreateViewIterator() -> Count"), ActualNum, ExpectedNum);
+		TEST_EQUAL(TEXT("FCbFieldView(Object)::AsObjectView().CreateViewIterator() -> Count"), ActualNum, ExpectedNum);
 
 		ActualNum = 0;
 		for (FCbFieldView Field : Object)
 		{
 			++ActualNum;
 			TestNotEqual(TEXT("FCbFieldView(Object) Iterator Name"), Field.GetName().Len(), 0);
-			TestEqual(TEXT("FCbFieldView(Object) Range"), Field.AsInt32(), ActualNum);
+			TEST_EQUAL(TEXT("FCbFieldView(Object) Range"), Field.AsInt32(), ActualNum);
 		}
-		TestEqual(TEXT("FCbFieldView(Object)::AsObjectView() Range -> Count"), ActualNum, ExpectedNum);
+		TEST_EQUAL(TEXT("FCbFieldView(Object)::AsObjectView() Range -> Count"), ActualNum, ExpectedNum);
 
 		ActualNum = 0;
 		for (FCbFieldView Field : Object.AsFieldView())
 		{
 			++ActualNum;
 			TestNotEqual(TEXT("FCbFieldView(ObjectField) Iterator Name"), Field.GetName().Len(), 0);
-			TestEqual(TEXT("FCbFieldView(ObjectField) Range"), Field.AsInt32(), ActualNum);
+			TEST_EQUAL(TEXT("FCbFieldView(ObjectField) Range"), Field.AsInt32(), ActualNum);
 		}
-		TestEqual(TEXT("FCbFieldView(ObjectField)::AsObjectView() Range -> Count"), ActualNum, ExpectedNum);
+		TEST_EQUAL(TEXT("FCbFieldView(ObjectField)::AsObjectView() Range -> Count"), ActualNum, ExpectedNum);
 	};
 
 	// Test FCbFieldView(Object, Empty)
-	TestField<ECbFieldType::Object>(TEXT("Object, Empty"), {0});
+	TestField<ECbFieldType::Object>(*this, TEXT("Object, Empty"), {0});
 
 	// Test FCbFieldView(Object, Empty)
 	{
@@ -321,21 +321,21 @@ TEST_CASE("Core::Serialization::FCbField::Object Test", "[Core][Serialization][S
 		TestIntObject(Object, 0, 1);
 
 		// Find fields that do not exist.
-		TestFalse(TEXT("FCbObjectView()::Find(Missing)"), Object.FindView("Field"_U8SV).HasValue());
-		TestFalse(TEXT("FCbObjectView()::FindViewIgnoreCase(Missing)"), Object.FindViewIgnoreCase("Field"_U8SV).HasValue());
-		TestFalse(TEXT("FCbObjectView()::operator[](Missing)"), Object["Field"_U8SV].HasValue());
+		TEST_FALSE(TEXT("FCbObjectView()::Find(Missing)"), Object.FindView("Field"_U8SV).HasValue());
+		TEST_FALSE(TEXT("FCbObjectView()::FindViewIgnoreCase(Missing)"), Object.FindViewIgnoreCase("Field"_U8SV).HasValue());
+		TEST_FALSE(TEXT("FCbObjectView()::operator[](Missing)"), Object["Field"_U8SV].HasValue());
 
 		// Advance an iterator past the last field.
 		FCbFieldViewIterator It = Object.CreateViewIterator();
-		TestFalse(TEXT("FCbObjectView()::CreateViewIterator() At End"), bool(It));
-		TestTrue(TEXT("FCbObjectView()::CreateViewIterator() At End"), !It);
+		TEST_FALSE(TEXT("FCbObjectView()::CreateViewIterator() At End"), bool(It));
+		TEST_TRUE(TEXT("FCbObjectView()::CreateViewIterator() At End"), !It);
 		for (int Count = 16; Count > 0; --Count)
 		{
 			++It;
 			It->AsInt32();
 		}
-		TestFalse(TEXT("FCbObjectView()::CreateViewIterator() At End"), bool(It));
-		TestTrue(TEXT("FCbObjectView()::CreateViewIterator() At End"), !It);
+		TEST_FALSE(TEXT("FCbObjectView()::CreateViewIterator() At End"), bool(It));
+		TEST_TRUE(TEXT("FCbObjectView()::CreateViewIterator() At End"), !It);
 	}
 
 	// Test FCbFieldView(Object, NotEmpty)
@@ -343,17 +343,17 @@ TEST_CASE("Core::Serialization::FCbField::Object Test", "[Core][Serialization][S
 		constexpr uint8 IntType = uint8(ECbFieldType::HasFieldName | ECbFieldType::IntegerPositive);
 		const uint8 Value[] = { 12, IntType, 1, 'A', 1, IntType, 1, 'B', 2, IntType, 1, 'C', 3 };
 		FCbFieldView Field(Value, ECbFieldType::Object);
-		TestField<ECbFieldType::Object>(TEXT("Object, NotEmpty"), Field, FCbObjectView(Value, ECbFieldType::Object));
+		TestField<ECbFieldType::Object>(*this, TEXT("Object, NotEmpty"), Field, FCbObjectView(Value, ECbFieldType::Object));
 		FCbObject Object = FCbObject::Clone(Field.AsObjectView());
 		TestIntObject(Object, 3, sizeof(Value));
 		TestIntObject(Field.AsObjectView(), 3, sizeof(Value));
-		TestTrue(TEXT("FCbObjectView::Equals()"), Object.Equals(Field.AsObjectView()));
-		TestEqual(TEXT("FCbObjectView::Find()"), Object.FindView("B"_ASV).AsInt32(), 2);
-		TestEqual(TEXT("FCbObjectView::Find()"), Object.FindView("b"_ASV).AsInt32(4), 4);
-		TestEqual(TEXT("FCbObjectView::FindViewIgnoreCase()"), Object.FindViewIgnoreCase("B"_ASV).AsInt32(), 2);
-		TestEqual(TEXT("FCbObjectView::FindViewIgnoreCase()"), Object.FindViewIgnoreCase("b"_ASV).AsInt32(), 2);
-		TestEqual(TEXT("FCbObjectView::operator[]"), Object["B"_ASV].AsInt32(), 2);
-		TestEqual(TEXT("FCbObjectView::operator[]"), Object["b"_ASV].AsInt32(4), 4);
+		TEST_TRUE(TEXT("FCbObjectView::Equals()"), Object.Equals(Field.AsObjectView()));
+		TEST_EQUAL(TEXT("FCbObjectView::Find()"), Object.FindView("B"_ASV).AsInt32(), 2);
+		TEST_EQUAL(TEXT("FCbObjectView::Find()"), Object.FindView("b"_ASV).AsInt32(4), 4);
+		TEST_EQUAL(TEXT("FCbObjectView::FindViewIgnoreCase()"), Object.FindViewIgnoreCase("B"_ASV).AsInt32(), 2);
+		TEST_EQUAL(TEXT("FCbObjectView::FindViewIgnoreCase()"), Object.FindViewIgnoreCase("b"_ASV).AsInt32(), 2);
+		TEST_EQUAL(TEXT("FCbObjectView::operator[]"), Object["B"_ASV].AsInt32(), 2);
+		TEST_EQUAL(TEXT("FCbObjectView::operator[]"), Object["b"_ASV].AsInt32(4), 4);
 	}
 
 	// Test FCbFieldView(UniformObject, NotEmpty)
@@ -361,54 +361,54 @@ TEST_CASE("Core::Serialization::FCbField::Object Test", "[Core][Serialization][S
 		constexpr uint8 IntType = uint8(ECbFieldType::HasFieldName | ECbFieldType::IntegerPositive);
 		const uint8 Value[] = { 10, IntType, 1, 'A', 1, 1, 'B', 2, 1, 'C', 3 };
 		FCbFieldView Field(Value, ECbFieldType::UniformObject);
-		TestField<ECbFieldType::UniformObject>(TEXT("UniformObject, NotEmpty"), Field, FCbObjectView(Value, ECbFieldType::UniformObject));
+		TestField<ECbFieldType::UniformObject>(*this, TEXT("UniformObject, NotEmpty"), Field, FCbObjectView(Value, ECbFieldType::UniformObject));
 		FCbObject Object = FCbObject::Clone(Field.AsObjectView());
 		TestIntObject(Object, 3, sizeof(Value));
 		TestIntObject(Field.AsObjectView(), 3, sizeof(Value));
-		TestTrue(TEXT("FCbObjectView(Uniform)::Equals()"), Object.Equals(Field.AsObjectView()));
-		TestEqual(TEXT("FCbObjectView(Uniform)::Find()"), Object.FindView("B"_ASV).AsInt32(), 2);
-		TestEqual(TEXT("FCbObjectView(Uniform)::Find()"), Object.Find("B"_ASV).AsInt32(), 2);
-		TestEqual(TEXT("FCbObjectView(Uniform)::Find()"), Object.FindView("b"_ASV).AsInt32(4), 4);
-		TestEqual(TEXT("FCbObjectView(Uniform)::Find()"), Object.Find("b"_ASV).AsInt32(4), 4);
-		TestEqual(TEXT("FCbObjectView(Uniform)::FindViewIgnoreCase()"), Object.FindViewIgnoreCase("B"_ASV).AsInt32(), 2);
-		TestEqual(TEXT("FCbObjectView(Uniform)::FindViewIgnoreCase()"), Object.FindIgnoreCase("B"_ASV).AsInt32(), 2);
-		TestEqual(TEXT("FCbObjectView(Uniform)::FindViewIgnoreCase()"), Object.FindViewIgnoreCase("b"_ASV).AsInt32(), 2);
-		TestEqual(TEXT("FCbObjectView(Uniform)::FindViewIgnoreCase()"), Object.FindIgnoreCase("b"_ASV).AsInt32(), 2);
-		TestEqual(TEXT("FCbObjectView(Uniform)::operator[]"), Object["B"_ASV].AsInt32(), 2);
-		TestEqual(TEXT("FCbObjectView(Uniform)::operator[]"), Object["b"_ASV].AsInt32(4), 4);
-
-		TestTrue(TEXT("FCbObject::AsField()"), Object.GetOuterBuffer() == Object.AsField().AsObject().GetOuterBuffer());
+		TEST_TRUE(TEXT("FCbObjectView(Uniform)::Equals()"), Object.Equals(Field.AsObjectView()));
+		TEST_EQUAL(TEXT("FCbObjectView(Uniform)::Find()"), Object.FindView("B"_ASV).AsInt32(), 2);
+		TEST_EQUAL(TEXT("FCbObjectView(Uniform)::Find()"), Object.Find("B"_ASV).AsInt32(), 2);
+		TEST_EQUAL(TEXT("FCbObjectView(Uniform)::Find()"), Object.FindView("b"_ASV).AsInt32(4), 4);
+		TEST_EQUAL(TEXT("FCbObjectView(Uniform)::Find()"), Object.Find("b"_ASV).AsInt32(4), 4);
+		TEST_EQUAL(TEXT("FCbObjectView(Uniform)::FindViewIgnoreCase()"), Object.FindViewIgnoreCase("B"_ASV).AsInt32(), 2);
+		TEST_EQUAL(TEXT("FCbObjectView(Uniform)::FindViewIgnoreCase()"), Object.FindIgnoreCase("B"_ASV).AsInt32(), 2);
+		TEST_EQUAL(TEXT("FCbObjectView(Uniform)::FindViewIgnoreCase()"), Object.FindViewIgnoreCase("b"_ASV).AsInt32(), 2);
+		TEST_EQUAL(TEXT("FCbObjectView(Uniform)::FindViewIgnoreCase()"), Object.FindIgnoreCase("b"_ASV).AsInt32(), 2);
+		TEST_EQUAL(TEXT("FCbObjectView(Uniform)::operator[]"), Object["B"_ASV].AsInt32(), 2);
+		TEST_EQUAL(TEXT("FCbObjectView(Uniform)::operator[]"), Object["b"_ASV].AsInt32(4), 4);
+
+		TEST_TRUE(TEXT("FCbObject::AsField()"), Object.GetOuterBuffer() == Object.AsField().AsObject().GetOuterBuffer());
 
 		// Equals
 		const uint8 NamedValue[] = { 1, 'O', 10, IntType, 1, 'A', 1, 1, 'B', 2, 1, 'C', 3 };
 		FCbFieldView NamedField(NamedValue, ECbFieldType::UniformObject | ECbFieldType::HasFieldName);
-		TestTrue(TEXT("FCbObjectView::Equals()"), Field.AsObjectView().Equals(NamedField.AsObjectView()));
-		TestTrue(TEXT("FCbObjectView::AsFieldView().Equals()"), Field.Equals(Field.AsObjectView().AsFieldView()));
-		TestTrue(TEXT("FCbObjectView::AsFieldView().Equals()"), NamedField.Equals(NamedField.AsObjectView().AsFieldView()));
+		TEST_TRUE(TEXT("FCbObjectView::Equals()"), Field.AsObjectView().Equals(NamedField.AsObjectView()));
+		TEST_TRUE(TEXT("FCbObjectView::AsFieldView().Equals()"), Field.Equals(Field.AsObjectView().AsFieldView()));
+		TEST_TRUE(TEXT("FCbObjectView::AsFieldView().Equals()"), NamedField.Equals(NamedField.AsObjectView().AsFieldView()));
 
 		// CopyTo
 		uint8 CopyBytes[sizeof(Value) + 1];
 		Field.AsObjectView().CopyTo(MakeMemoryView(CopyBytes));
-		TestTrue(TEXT("FCbObjectView(NoType)::CopyTo()"), MakeMemoryView(Value).EqualBytes(MakeMemoryView(CopyBytes) + 1));
+		TEST_TRUE(TEXT("FCbObjectView(NoType)::CopyTo()"), MakeMemoryView(Value).EqualBytes(MakeMemoryView(CopyBytes) + 1));
 		NamedField.AsObjectView().CopyTo(MakeMemoryView(CopyBytes));
-		TestTrue(TEXT("FCbObjectView(NoType, Name)::CopyTo()"), MakeMemoryView(Value).EqualBytes(MakeMemoryView(CopyBytes) + 1));
+		TEST_TRUE(TEXT("FCbObjectView(NoType, Name)::CopyTo()"), MakeMemoryView(Value).EqualBytes(MakeMemoryView(CopyBytes) + 1));
 
 		// TryGetView
 		FMemoryView View;
-		TestTrue(TEXT("FCbObjectView(Clone)::TryGetView()"), Object.TryGetView(View) && View == Object.GetOuterBuffer().GetView());
-		TestFalse(TEXT("FCbObjectView(NoType)::TryGetView()"), Field.AsObjectView().TryGetView(View));
-		TestFalse(TEXT("FCbObjectView(Name)::TryGetView()"), NamedField.AsObjectView().TryGetView(View));
+		TEST_TRUE(TEXT("FCbObjectView(Clone)::TryGetView()"), Object.TryGetView(View) && View == Object.GetOuterBuffer().GetView());
+		TEST_FALSE(TEXT("FCbObjectView(NoType)::TryGetView()"), Field.AsObjectView().TryGetView(View));
+		TEST_FALSE(TEXT("FCbObjectView(Name)::TryGetView()"), NamedField.AsObjectView().TryGetView(View));
 
 		// GetBuffer
-		TestTrue(TEXT("FCbObjectView(Clone)::GetBuffer()"), Object.GetBuffer().ToShared().GetView() == Object.GetOuterBuffer().GetView());
-		TestTrue(TEXT("FCbObjectView(NoType)::GetBuffer()"), FCbField::MakeView(Field).AsObject().GetBuffer().ToShared().GetView().EqualBytes(Object.GetOuterBuffer().GetView()));
-		TestTrue(TEXT("FCbObjectView(Name)::GetBuffer()"), FCbField::MakeView(NamedField).AsObject().GetBuffer().ToShared().GetView().EqualBytes(Object.GetOuterBuffer().GetView()));
+		TEST_TRUE(TEXT("FCbObjectView(Clone)::GetBuffer()"), Object.GetBuffer().ToShared().GetView() == Object.GetOuterBuffer().GetView());
+		TEST_TRUE(TEXT("FCbObjectView(NoType)::GetBuffer()"), FCbField::MakeView(Field).AsObject().GetBuffer().ToShared().GetView().EqualBytes(Object.GetOuterBuffer().GetView()));
+		TEST_TRUE(TEXT("FCbObjectView(Name)::GetBuffer()"), FCbField::MakeView(NamedField).AsObject().GetBuffer().ToShared().GetView().EqualBytes(Object.GetOuterBuffer().GetView()));
 	}
 
 	// Test FCbFieldView(None) as Object
 	{
 		FCbFieldView Field;
-		TestFieldError<ECbFieldType::Object>(TEXT("Object, None"), Field, ECbFieldError::TypeError);
+		TestFieldError<ECbFieldType::Object>(*this, TEXT("Object, None"), Field, ECbFieldError::TypeError);
 		FCbField::MakeView(Field).AsObject();
 	}
 
@@ -417,17 +417,17 @@ TEST_CASE("Core::Serialization::FCbField::Object Test", "[Core][Serialization][S
 		const uint8 ObjectType = uint8(ECbFieldType::Object | ECbFieldType::HasFieldName);
 		const uint8 Buffer[] = { ObjectType, 3, 'K', 'e', 'y', 4, uint8(ECbFieldType::HasFieldName | ECbFieldType::IntegerPositive), 1, 'F', 8 };
 		const FCbObjectView Object(Buffer);
-		TestEqual(TEXT("FCbObjectView(ObjectWithName)::GetSize()"), Object.GetSize(), uint64(6));
+		TEST_EQUAL(TEXT("FCbObjectView(ObjectWithName)::GetSize()"), Object.GetSize(), uint64(6));
 		const FCbObject ObjectClone = FCbObject::Clone(Object);
-		TestEqual(TEXT("FCbObject(ObjectWithName)::GetSize()"), ObjectClone.GetSize(), uint64(6));
-		TestTrue(TEXT("FCbObjectView::Equals()"), Object.Equals(ObjectClone));
-		TestEqual(TEXT("FCbObjectView::GetHash()"), ObjectClone.GetHash(), Object.GetHash());
+		TEST_EQUAL(TEXT("FCbObject(ObjectWithName)::GetSize()"), ObjectClone.GetSize(), uint64(6));
+		TEST_TRUE(TEXT("FCbObjectView::Equals()"), Object.Equals(ObjectClone));
+		TEST_EQUAL(TEXT("FCbObjectView::GetHash()"), ObjectClone.GetHash(), Object.GetHash());
 		for (FCbFieldIterator It = ObjectClone.CreateIterator(); It; ++It)
 		{
 			FCbField Field = *It;
-			TestEqual(TEXT("FCbObject::CreateIterator().GetName()"), Field.GetName(), "F"_U8SV);
-			TestEqual(TEXT("FCbObject::CreateIterator().AsInt32()"), Field.AsInt32(), 8);
-			TestTrue(TEXT("FCbObject::CreateIterator().IsOwned()"), Field.IsOwned());
+			TEST_EQUAL(TEXT("FCbObject::CreateIterator().GetName()"), Field.GetName(), "F"_U8SV);
+			TEST_EQUAL(TEXT("FCbObject::CreateIterator().AsInt32()"), Field.AsInt32(), 8);
+			TEST_TRUE(TEXT("FCbObject::CreateIterator().IsOwned()"), Field.IsOwned());
 		}
 		for (FCbFieldIterator It = ObjectClone.CreateIterator(), End; It != End; ++It)
 		{
@@ -439,18 +439,18 @@ TEST_CASE("Core::Serialization::FCbField::Object Test", "[Core][Serialization][S
 		// CopyTo
 		uint8 CopyBytes[6];
 		Object.CopyTo(MakeMemoryView(CopyBytes));
-		TestTrue(TEXT("FCbObjectView(Name)::CopyTo()"), ObjectClone.GetOuterBuffer().GetView().EqualBytes(MakeMemoryView(CopyBytes)));
+		TEST_TRUE(TEXT("FCbObjectView(Name)::CopyTo()"), ObjectClone.GetOuterBuffer().GetView().EqualBytes(MakeMemoryView(CopyBytes)));
 		ObjectClone.CopyTo(MakeMemoryView(CopyBytes));
-		TestTrue(TEXT("FCbObjectView()::CopyTo()"), ObjectClone.GetOuterBuffer().GetView().EqualBytes(MakeMemoryView(CopyBytes)));
+		TEST_TRUE(TEXT("FCbObjectView()::CopyTo()"), ObjectClone.GetOuterBuffer().GetView().EqualBytes(MakeMemoryView(CopyBytes)));
 
 		// GetBuffer
-		TestTrue(TEXT("FCbField(ObjectWithName)::GetBuffer()"), FCbField(FSharedBuffer::MakeView(MakeMemoryView(Buffer))).GetBuffer().ToShared().GetView().EqualBytes(MakeMemoryView(Buffer)));
-		TestTrue(TEXT("FCbField(ObjectWithNameNoType)::GetBuffer()"),
+		TEST_TRUE(TEXT("FCbField(ObjectWithName)::GetBuffer()"), FCbField(FSharedBuffer::MakeView(MakeMemoryView(Buffer))).GetBuffer().ToShared().GetView().EqualBytes(MakeMemoryView(Buffer)));
+		TEST_TRUE(TEXT("FCbField(ObjectWithNameNoType)::GetBuffer()"),
 			FCbField(FCbFieldView(Buffer + 1, ECbFieldType(ObjectType)), FSharedBuffer::MakeView(MakeMemoryView(Buffer))).GetBuffer().ToShared().GetView().EqualBytes(MakeMemoryView(Buffer)));
 
 		// Access Missing Field
-		TestFalse(TEXT("FCbObject()[Missing]"), ObjectClone["M"_ASV].HasValue());
-		TestFalse(TEXT("FCbField(Object)[Missing]"), ObjectClone.AsField()["M"_ASV].HasValue());
+		TEST_FALSE(TEXT("FCbObject()[Missing]"), ObjectClone["M"_ASV].HasValue());
+		TEST_FALSE(TEXT("FCbField(Object)[Missing]"), ObjectClone.AsField()["M"_ASV].HasValue());
 	}
 
 	// Test FCbObjectView as FCbFieldViewIterator
@@ -459,10 +459,10 @@ TEST_CASE("Core::Serialization::FCbField::Object Test", "[Core][Serialization][S
 		FCbObjectView Object;
 		for (FCbFieldView Field : FCbFieldViewIterator::MakeSingle(Object.AsFieldView()))
 		{
-			TestTrue(TEXT("FCbObjectView::AsFieldView() as Iterator"), Field.IsObject());
+			TEST_TRUE(TEXT("FCbObjectView::AsFieldView() as Iterator"), Field.IsObject());
 			++Count;
 		}
-		TestEqual(TEXT("FCbObjectView::AsFieldView() as Iterator Count"), Count, 1u);
+		TEST_EQUAL(TEXT("FCbObjectView::AsFieldView() as Iterator Count"), Count, 1u);
 	}
 
 	// Test FCbObject as FCbFieldIterator
@@ -472,10 +472,10 @@ TEST_CASE("Core::Serialization::FCbField::Object Test", "[Core][Serialization][S
 		Object.MakeOwned();
 		for (FCbField Field : FCbFieldIterator::MakeSingle(Object.AsField()))
 		{
-			TestTrue(TEXT("FCbObject::AsFieldView() as Iterator"), Field.IsObject());
+			TEST_TRUE(TEXT("FCbObject::AsFieldView() as Iterator"), Field.IsObject());
 			++Count;
 		}
-		TestEqual(TEXT("FCbObject::AsFieldView() as Iterator Count"), Count, 1u);
+		TEST_EQUAL(TEXT("FCbObject::AsFieldView() as Iterator Count"), Count, 1u);
 	}
 
 	// Test FCbObject(Empty) as FCbFieldIterator
@@ -489,45 +489,45 @@ TEST_CASE("Core::Serialization::FCbField::Object Test", "[Core][Serialization][S
 }
 
 
-TEST_CASE("Core::Serialization::FCbField::Array Test", "[Core][Serialization][Smoke]")
+TEST_CASE_METHOD(FAutomationTestFixture, "Core::Serialization::FCbField::Array Test", "[Core][Serialization][Smoke]")
 {
 	static_assert(!std::is_constructible<FCbFieldView, const FCbArrayView&>::value, "Invalid constructor for FCbFieldView");
 	static_assert(!std::is_assignable<FCbFieldView, const FCbArrayView&>::value, "Invalid assignment for FCbFieldView");
 	static_assert(!std::is_convertible<FCbFieldView, FCbArrayView>::value, "Invalid conversion to FCbArrayView");
 	static_assert(!std::is_assignable<FCbArrayView, const FCbFieldView&>::value, "Invalid assignment for FCbArrayView");
 
-	auto TestIntArray = [](FCbArrayView Array, int32 ExpectedNum, uint64 ExpectedValueSize)
+	auto TestIntArray = [this](FCbArrayView Array, int32 ExpectedNum, uint64 ExpectedValueSize)
 	{
-		TestEqual(TEXT("FCbFieldView(Array)::AsArrayView().GetSize()"), Array.GetSize(), ExpectedValueSize + sizeof(ECbFieldType));
-		TestEqual(TEXT("FCbFieldView(Array)::AsArrayView().Num()"), Array.Num(), uint64(ExpectedNum));
+		TEST_EQUAL(TEXT("FCbFieldView(Array)::AsArrayView().GetSize()"), Array.GetSize(), ExpectedValueSize + sizeof(ECbFieldType));
+		TEST_EQUAL(TEXT("FCbFieldView(Array)::AsArrayView().Num()"), Array.Num(), uint64(ExpectedNum));
 
 		int32 ActualNum = 0;
 		for (FCbFieldViewIterator It = Array.CreateViewIterator(); It; ++It)
 		{
 			++ActualNum;
-			TestEqual(TEXT("FCbFieldView(Array) Iterator"), It->AsInt32(), ActualNum);
+			TEST_EQUAL(TEXT("FCbFieldView(Array) Iterator"), It->AsInt32(), ActualNum);
 		}
-		TestEqual(TEXT("FCbFieldView(Array)::AsArrayView().CreateViewIterator() -> Count"), ActualNum, ExpectedNum);
+		TEST_EQUAL(TEXT("FCbFieldView(Array)::AsArrayView().CreateViewIterator() -> Count"), ActualNum, ExpectedNum);
 
 		ActualNum = 0;
 		for (FCbFieldView Field : Array)
 		{
 			++ActualNum;
-			TestEqual(TEXT("FCbFieldView(Array) Range"), Field.AsInt32(), ActualNum);
+			TEST_EQUAL(TEXT("FCbFieldView(Array) Range"), Field.AsInt32(), ActualNum);
 		}
-		TestEqual(TEXT("FCbFieldView(Array)::AsArrayView() Range -> Count"), ActualNum, ExpectedNum);
+		TEST_EQUAL(TEXT("FCbFieldView(Array)::AsArrayView() Range -> Count"), ActualNum, ExpectedNum);
 
 		ActualNum = 0;
 		for (FCbFieldView Field : Array.AsFieldView())
 		{
 			++ActualNum;
-			TestEqual(TEXT("FCbFieldView(ArrayField) Range"), Field.AsInt32(), ActualNum);
+			TEST_EQUAL(TEXT("FCbFieldView(ArrayField) Range"), Field.AsInt32(), ActualNum);
 		}
-		TestEqual(TEXT("FCbFieldView(ArrayField)::AsArrayView() Range -> Count"), ActualNum, ExpectedNum);
+		TEST_EQUAL(TEXT("FCbFieldView(ArrayField)::AsArrayView() Range -> Count"), ActualNum, ExpectedNum);
 	};
 
 	// Test FCbFieldView(Array, Empty)
-	TestField<ECbFieldType::Array>(TEXT("Array, Empty"), { 1, 0 });
+	TestField<ECbFieldType::Array>(*this, TEXT("Array, Empty"), { 1, 0 });
 
 	// Test FCbFieldView(Array, Empty)
 	{
@@ -536,15 +536,15 @@ TEST_CASE("Core::Serialization::FCbField::Array Test", "[Core][Serialization][Sm
 
 		// Advance an iterator past the last field.
 		FCbFieldViewIterator It = Array.CreateViewIterator();
-		TestFalse(TEXT("FCbArrayView()::CreateViewIterator() At End"), bool(It));
-		TestTrue(TEXT("FCbArrayView()::CreateViewIterator() At End"), !It);
+		TEST_FALSE(TEXT("FCbArrayView()::CreateViewIterator() At End"), bool(It));
+		TEST_TRUE(TEXT("FCbArrayView()::CreateViewIterator() At End"), !It);
 		for (int Count = 16; Count > 0; --Count)
 		{
 			++It;
 			It->AsInt32();
 		}
-		TestFalse(TEXT("FCbArrayView()::CreateViewIterator() At End"), bool(It));
-		TestTrue(TEXT("FCbArrayView()::CreateViewIterator() At End"), !It);
+		TEST_FALSE(TEXT("FCbArrayView()::CreateViewIterator() At End"), bool(It));
+		TEST_TRUE(TEXT("FCbArrayView()::CreateViewIterator() At End"), !It);
 	}
 
 	// Test FCbFieldView(Array, NotEmpty)
@@ -552,11 +552,11 @@ TEST_CASE("Core::Serialization::FCbField::Array Test", "[Core][Serialization][Sm
 		constexpr uint8 IntType = uint8(ECbFieldType::IntegerPositive);
 		const uint8 Value[] = { 7, 3, IntType, 1, IntType, 2, IntType, 3 };
 		FCbFieldView Field(Value, ECbFieldType::Array);
-		TestField<ECbFieldType::Array>(TEXT("Array, NotEmpty"), Field, FCbArrayView(Value, ECbFieldType::Array));
+		TestField<ECbFieldType::Array>(*this, TEXT("Array, NotEmpty"), Field, FCbArrayView(Value, ECbFieldType::Array));
 		FCbArray Array = FCbArray::Clone(Field.AsArrayView());
 		TestIntArray(Array, 3, sizeof(Value));
 		TestIntArray(Field.AsArrayView(), 3, sizeof(Value));
-		TestTrue(TEXT("FCbArrayView::Equals()"), Array.Equals(Field.AsArrayView()));
+		TEST_TRUE(TEXT("FCbArrayView::Equals()"), Array.Equals(Field.AsArrayView()));
 	}
 
 	// Test FCbFieldView(UniformArray)
@@ -564,44 +564,44 @@ TEST_CASE("Core::Serialization::FCbField::Array Test", "[Core][Serialization][Sm
 		constexpr uint8 IntType = uint8(ECbFieldType::IntegerPositive);
 		const uint8 Value[] = { 5, 3, IntType, 1, 2, 3 };
 		FCbFieldView Field(Value, ECbFieldType::UniformArray);
-		TestField<ECbFieldType::UniformArray>(TEXT("UniformArray"), Field, FCbArrayView(Value, ECbFieldType::UniformArray));
+		TestField<ECbFieldType::UniformArray>(*this, TEXT("UniformArray"), Field, FCbArrayView(Value, ECbFieldType::UniformArray));
 		FCbArray Array = FCbArray::Clone(Field.AsArrayView());
 		TestIntArray(Array, 3, sizeof(Value));
 		TestIntArray(Field.AsArrayView(), 3, sizeof(Value));
-		TestTrue(TEXT("FCbArrayView(Uniform)::Equals()"), Array.Equals(Field.AsArrayView()));
+		TEST_TRUE(TEXT("FCbArrayView(Uniform)::Equals()"), Array.Equals(Field.AsArrayView()));
 
-		TestTrue(TEXT("FCbArray::AsField()"), Array.GetOuterBuffer() == Array.AsField().AsArray().GetOuterBuffer());
+		TEST_TRUE(TEXT("FCbArray::AsField()"), Array.GetOuterBuffer() == Array.AsField().AsArray().GetOuterBuffer());
 
 		// Equals
 		const uint8 NamedValue[] = { 1, 'A', 5, 3, IntType, 1, 2, 3 };
 		FCbFieldView NamedField(NamedValue, ECbFieldType::UniformArray | ECbFieldType::HasFieldName);
-		TestTrue(TEXT("FCbArrayView::Equals()"), Field.AsArrayView().Equals(NamedField.AsArrayView()));
-		TestTrue(TEXT("FCbArrayView::AsFieldView().Equals()"), Field.Equals(Field.AsArrayView().AsFieldView()));
-		TestTrue(TEXT("FCbArrayView::AsFieldView().Equals()"), NamedField.Equals(NamedField.AsArrayView().AsFieldView()));
+		TEST_TRUE(TEXT("FCbArrayView::Equals()"), Field.AsArrayView().Equals(NamedField.AsArrayView()));
+		TEST_TRUE(TEXT("FCbArrayView::AsFieldView().Equals()"), Field.Equals(Field.AsArrayView().AsFieldView()));
+		TEST_TRUE(TEXT("FCbArrayView::AsFieldView().Equals()"), NamedField.Equals(NamedField.AsArrayView().AsFieldView()));
 
 		// CopyTo
 		uint8 CopyBytes[sizeof(Value) + 1];
 		Field.AsArrayView().CopyTo(MakeMemoryView(CopyBytes));
-		TestTrue(TEXT("FCbArrayView(NoType)::CopyTo()"), MakeMemoryView(Value).EqualBytes(MakeMemoryView(CopyBytes) + 1));
+		TEST_TRUE(TEXT("FCbArrayView(NoType)::CopyTo()"), MakeMemoryView(Value).EqualBytes(MakeMemoryView(CopyBytes) + 1));
 		NamedField.AsArrayView().CopyTo(MakeMemoryView(CopyBytes));
-		TestTrue(TEXT("FCbArrayView(NoType, Name)::CopyTo()"), MakeMemoryView(Value).EqualBytes(MakeMemoryView(CopyBytes) + 1));
+		TEST_TRUE(TEXT("FCbArrayView(NoType, Name)::CopyTo()"), MakeMemoryView(Value).EqualBytes(MakeMemoryView(CopyBytes) + 1));
 
 		// TryGetView
 		FMemoryView View;
-		TestTrue(TEXT("FCbArrayView(Clone)::TryGetView()"), Array.TryGetView(View) && View == Array.GetOuterBuffer().GetView());
-		TestFalse(TEXT("FCbArrayView(NoType)::TryGetView()"), Field.AsArrayView().TryGetView(View));
-		TestFalse(TEXT("FCbArrayView(Name)::TryGetView()"), NamedField.AsArrayView().TryGetView(View));
+		TEST_TRUE(TEXT("FCbArrayView(Clone)::TryGetView()"), Array.TryGetView(View) && View == Array.GetOuterBuffer().GetView());
+		TEST_FALSE(TEXT("FCbArrayView(NoType)::TryGetView()"), Field.AsArrayView().TryGetView(View));
+		TEST_FALSE(TEXT("FCbArrayView(Name)::TryGetView()"), NamedField.AsArrayView().TryGetView(View));
 
 		// GetBuffer
-		TestTrue(TEXT("FCbArrayView(Clone)::GetBuffer()"), Array.GetBuffer().ToShared().GetView() == Array.GetOuterBuffer().GetView());
-		TestTrue(TEXT("FCbArrayView(NoType)::GetBuffer()"), FCbField::MakeView(Field).AsArray().GetBuffer().ToShared().GetView().EqualBytes(Array.GetOuterBuffer().GetView()));
-		TestTrue(TEXT("FCbArrayView(Name)::GetBuffer()"), FCbField::MakeView(NamedField).AsArray().GetBuffer().ToShared().GetView().EqualBytes(Array.GetOuterBuffer().GetView()));
+		TEST_TRUE(TEXT("FCbArrayView(Clone)::GetBuffer()"), Array.GetBuffer().ToShared().GetView() == Array.GetOuterBuffer().GetView());
+		TEST_TRUE(TEXT("FCbArrayView(NoType)::GetBuffer()"), FCbField::MakeView(Field).AsArray().GetBuffer().ToShared().GetView().EqualBytes(Array.GetOuterBuffer().GetView()));
+		TEST_TRUE(TEXT("FCbArrayView(Name)::GetBuffer()"), FCbField::MakeView(NamedField).AsArray().GetBuffer().ToShared().GetView().EqualBytes(Array.GetOuterBuffer().GetView()));
 	}
 
 	// Test FCbFieldView(None) as Array
 	{
 		FCbFieldView Field;
-		TestFieldError<ECbFieldType::Array>(TEXT("Array, None"), Field, ECbFieldError::TypeError);
+		TestFieldError<ECbFieldType::Array>(*this, TEXT("Array, None"), Field, ECbFieldError::TypeError);
 		FCbField::MakeView(Field).AsArray();
 	}
 
@@ -610,16 +610,16 @@ TEST_CASE("Core::Serialization::FCbField::Array Test", "[Core][Serialization][Sm
 		const uint8 ArrayType = uint8(ECbFieldType::Array | ECbFieldType::HasFieldName);
 		const uint8 Buffer[] = { ArrayType, 3, 'K', 'e', 'y', 3, 1, uint8(ECbFieldType::IntegerPositive), 8 };
 		const FCbArrayView Array(Buffer);
-		TestEqual(TEXT("FCbArrayView(ArrayWithName)::GetSize()"), Array.GetSize(), uint64(5));
+		TEST_EQUAL(TEXT("FCbArrayView(ArrayWithName)::GetSize()"), Array.GetSize(), uint64(5));
 		const FCbArray ArrayClone = FCbArray::Clone(Array);
-		TestEqual(TEXT("FCbArray(ArrayWithName)::GetSize()"), ArrayClone.GetSize(), uint64(5));
-		TestTrue(TEXT("FCbArrayView::Equals()"), Array.Equals(ArrayClone));
-		TestEqual(TEXT("FCbArrayView::GetHash()"), ArrayClone.GetHash(), Array.GetHash());
+		TEST_EQUAL(TEXT("FCbArray(ArrayWithName)::GetSize()"), ArrayClone.GetSize(), uint64(5));
+		TEST_TRUE(TEXT("FCbArrayView::Equals()"), Array.Equals(ArrayClone));
+		TEST_EQUAL(TEXT("FCbArrayView::GetHash()"), ArrayClone.GetHash(), Array.GetHash());
 		for (FCbFieldIterator It = ArrayClone.CreateIterator(); It; ++It)
 		{
 			FCbField Field = *It;
-			TestEqual(TEXT("FCbArray::CreateIterator().AsInt32()"), Field.AsInt32(), 8);
-			TestTrue(TEXT("FCbArray::CreateIterator().IsOwned()"), Field.IsOwned());
+			TEST_EQUAL(TEXT("FCbArray::CreateIterator().AsInt32()"), Field.AsInt32(), 8);
+			TEST_TRUE(TEXT("FCbArray::CreateIterator().IsOwned()"), Field.IsOwned());
 		}
 		for (FCbFieldIterator It = ArrayClone.CreateIterator(), End; It != End; ++It)
 		{
@@ -631,13 +631,13 @@ TEST_CASE("Core::Serialization::FCbField::Array Test", "[Core][Serialization][Sm
 		// CopyTo
 		uint8 CopyBytes[5];
 		Array.CopyTo(MakeMemoryView(CopyBytes));
-		TestTrue(TEXT("FCbArrayView(Name)::CopyTo()"), ArrayClone.GetOuterBuffer().GetView().EqualBytes(MakeMemoryView(CopyBytes)));
+		TEST_TRUE(TEXT("FCbArrayView(Name)::CopyTo()"), ArrayClone.GetOuterBuffer().GetView().EqualBytes(MakeMemoryView(CopyBytes)));
 		ArrayClone.CopyTo(MakeMemoryView(CopyBytes));
-		TestTrue(TEXT("FCbArrayView()::CopyTo()"), ArrayClone.GetOuterBuffer().GetView().EqualBytes(MakeMemoryView(CopyBytes)));
+		TEST_TRUE(TEXT("FCbArrayView()::CopyTo()"), ArrayClone.GetOuterBuffer().GetView().EqualBytes(MakeMemoryView(CopyBytes)));
 
 		// GetBuffer
-		TestTrue(TEXT("FCbField(ArrayWithName)::GetBuffer()"), FCbField(FSharedBuffer::MakeView(MakeMemoryView(Buffer))).GetBuffer().ToShared().GetView().EqualBytes(MakeMemoryView(Buffer)));
-		TestTrue(TEXT("FCbField(ArrayWithNameNoType)::GetBuffer()"),
+		TEST_TRUE(TEXT("FCbField(ArrayWithName)::GetBuffer()"), FCbField(FSharedBuffer::MakeView(MakeMemoryView(Buffer))).GetBuffer().ToShared().GetView().EqualBytes(MakeMemoryView(Buffer)));
+		TEST_TRUE(TEXT("FCbField(ArrayWithNameNoType)::GetBuffer()"),
 			FCbField(FCbFieldView(Buffer + 1, ECbFieldType(ArrayType)), FSharedBuffer::MakeView(MakeMemoryView(Buffer))).GetBuffer().ToShared().GetView().EqualBytes(MakeMemoryView(Buffer)));
 	}
 
@@ -647,10 +647,10 @@ TEST_CASE("Core::Serialization::FCbField::Array Test", "[Core][Serialization][Sm
 		FCbArrayView Array;
 		for (FCbFieldView Field : FCbFieldViewIterator::MakeSingle(Array.AsFieldView()))
 		{
-			TestTrue(TEXT("FCbArrayView::AsFieldView() as Iterator"), Field.IsArray());
+			TEST_TRUE(TEXT("FCbArrayView::AsFieldView() as Iterator"), Field.IsArray());
 			++Count;
 		}
-		TestEqual(TEXT("FCbArrayView::AsFieldView() as Iterator Count"), Count, 1u);
+		TEST_EQUAL(TEXT("FCbArrayView::AsFieldView() as Iterator Count"), Count, 1u);
 	}
 
 	// Test FCbArray as FCbFieldIterator
@@ -660,10 +660,10 @@ TEST_CASE("Core::Serialization::FCbField::Array Test", "[Core][Serialization][Sm
 		Array.MakeOwned();
 		for (FCbField Field : FCbFieldIterator::MakeSingle(Array.AsField()))
 		{
-			TestTrue(TEXT("FCbArray::AsFieldView() as Iterator"), Field.IsArray());
+			TEST_TRUE(TEXT("FCbArray::AsFieldView() as Iterator"), Field.IsArray());
 			++Count;
 		}
-		TestEqual(TEXT("FCbArray::AsFieldView() as Iterator Count"), Count, 1u);
+		TEST_EQUAL(TEXT("FCbArray::AsFieldView() as Iterator Count"), Count, 1u);
 	}
 
 	// Test FCbArray(Empty) as FCbFieldIterator
@@ -676,7 +676,7 @@ TEST_CASE("Core::Serialization::FCbField::Array Test", "[Core][Serialization][Sm
 	}
 }
 
-TEST_CASE("Core::Serialization::FCbField::Binary Test", "[Core][Serialization][Smoke]")
+TEST_CASE_METHOD(FAutomationTestFixture, "Core::Serialization::FCbField::Binary Test", "[Core][Serialization][Smoke]")
 {
 	struct FCbBinaryAccessors
 	{
@@ -686,54 +686,54 @@ TEST_CASE("Core::Serialization::FCbField::Binary Test", "[Core][Serialization][S
 	};
 
 	// Test FCbFieldView(Binary, Empty)
-	TestField<ECbFieldType::Binary>(TEXT("Binary, Empty"), {0});
+	TestField<ECbFieldType::Binary>(*this, TEXT("Binary, Empty"), {0});
 
 	// Test FCbFieldView(Binary, Value)
 	{
 		const uint8 Value[] = { 3, 4, 5, 6 }; // Size: 3, Data: 4/5/6
 		FCbFieldView FieldView(Value, ECbFieldType::Binary);
-		TestFieldNoClone<ECbFieldType::Binary>(TEXT("Binary, Value, View"), FieldView, MakeMemoryView(Value + 1, 3));
+		TestFieldNoClone<ECbFieldType::Binary>(*this, TEXT("Binary, Value, View"), FieldView, MakeMemoryView(Value + 1, 3));
 
 		FCbField Field = FCbField::Clone(FieldView);
 		Field.AsBinary();
-		TestFalse(TEXT("Binary, Value, AsBinary -> GetOuterBuffer()"), Field.GetOuterBuffer().IsNull());
+		TEST_FALSE(TEXT("Binary, Value, AsBinary -> GetOuterBuffer()"), Field.GetOuterBuffer().IsNull());
 		MoveTemp(Field).AsBinary();
-		TestTrue(TEXT("Binary, Value, AsBinary -> GetOuterBuffer()"), Field.GetOuterBuffer().IsNull());
+		TEST_TRUE(TEXT("Binary, Value, AsBinary -> GetOuterBuffer()"), Field.GetOuterBuffer().IsNull());
 	}
 
 	// Test FCbFieldView(None) as Binary
 	{
 		FCbFieldView FieldView;
 		const uint8 Default[] = { 1, 2, 3 };
-		TestFieldError<ECbFieldType::Binary>(TEXT("Binary, None, View"), FieldView, ECbFieldError::TypeError, MakeMemoryView(Default));
+		TestFieldError<ECbFieldType::Binary>(*this, TEXT("Binary, None, View"), FieldView, ECbFieldError::TypeError, MakeMemoryView(Default));
 
 		FCbField Field = FCbField::Clone(FieldView);
-		TestFieldError<ECbFieldType::Binary, FSharedBuffer>(TEXT("Binary, None"), Field, ECbFieldError::TypeError, FSharedBuffer::MakeView(MakeMemoryView(Default)), FCbBinaryAccessors());
+		TestFieldError<ECbFieldType::Binary, FSharedBuffer>(*this, TEXT("Binary, None"), Field, ECbFieldError::TypeError, FSharedBuffer::MakeView(MakeMemoryView(Default)), FCbBinaryAccessors());
 	}
 }
 
-TEST_CASE("Core::Serialization::FCbField::String Test", "[Core][Serialization][Smoke]")
+TEST_CASE_METHOD(FAutomationTestFixture, "Core::Serialization::FCbField::String Test", "[Core][Serialization][Smoke]")
 {
 	// Test FCbFieldView(String, Empty)
-	TestField<ECbFieldType::String>(TEXT("String, Empty"), {0});
+	TestField<ECbFieldType::String>(*this, TEXT("String, Empty"), {0});
 
 	// Test FCbFieldView(String, Value)
 	{
 		const uint8 Value[] = { 3, 'A', 'B', 'C' }; // Size: 3, Data: ABC
-		TestField<ECbFieldType::String>(TEXT("String, Value"), Value, FUtf8StringView(reinterpret_cast<const UTF8CHAR*>(Value) + 1, 3));
+		TestField<ECbFieldType::String>(*this, TEXT("String, Value"), Value, FUtf8StringView(reinterpret_cast<const UTF8CHAR*>(Value) + 1, 3));
 	}
 
 	// Test FCbFieldView(String, OutOfRangeSize)
 	{
 		uint8 Value[9];
 		WriteVarUInt(uint64(1) << 31, Value);
-		TestFieldError<ECbFieldType::String>(TEXT("String, OutOfRangeSize"), Value, ECbFieldError::RangeError, "ABC"_U8SV);
+		TestFieldError<ECbFieldType::String>(*this, TEXT("String, OutOfRangeSize"), Value, ECbFieldError::RangeError, "ABC"_U8SV);
 	}
 
 	// Test FCbFieldView(None) as String
 	{
 		FCbFieldView Field;
-		TestFieldError<ECbFieldType::String>(TEXT("String, None"), Field, ECbFieldError::TypeError, "ABC"_U8SV);
+		TestFieldError<ECbFieldType::String>(*this, TEXT("String, None"), Field, ECbFieldError::TypeError, "ABC"_U8SV);
 	}
 }
 
@@ -772,7 +772,7 @@ struct TCbIntegerAccessors
 	static constexpr T (FCbFieldView::*AsType)(T) = InAsTypeFn;
 };
 
-void TestIntegerField(ECbFieldType FieldType, EIntType ExpectedMask, uint64 Magnitude)
+void TestIntegerField(FAutomationTestFixture& Test, ECbFieldType FieldType, EIntType ExpectedMask, uint64 Magnitude)
 {
 	uint8 Value[9];
 	const bool Negative = bool(uint8(FieldType) & 1);
@@ -780,66 +780,66 @@ void TestIntegerField(ECbFieldType FieldType, EIntType ExpectedMask, uint64 Magn
 	constexpr uint64 DefaultValue = 8;
 	const uint64 ExpectedValue = Negative ? uint64(-int64(Magnitude)) : Magnitude;
 	FCbFieldView Field(Value, FieldType);
-	TestField<ECbFieldType::IntegerNegative>(TEXT("Int8"), Field, int8(EnumHasAnyFlags(ExpectedMask, EIntType::Int8) ? ExpectedValue : DefaultValue),
+	TestField<ECbFieldType::IntegerNegative>(Test, TEXT("Int8"), Field, int8(EnumHasAnyFlags(ExpectedMask, EIntType::Int8) ? ExpectedValue : DefaultValue),
 		int8(DefaultValue), EnumHasAnyFlags(ExpectedMask, EIntType::Int8) ? ECbFieldError::None : ECbFieldError::RangeError, TCbIntegerAccessors<int8, &FCbFieldView::AsInt8>());
-	TestField<ECbFieldType::IntegerNegative>(TEXT("Int16"), Field, int16(EnumHasAnyFlags(ExpectedMask, EIntType::Int16) ? ExpectedValue : DefaultValue),
+	TestField<ECbFieldType::IntegerNegative>(Test, TEXT("Int16"), Field, int16(EnumHasAnyFlags(ExpectedMask, EIntType::Int16) ? ExpectedValue : DefaultValue),
 		int16(DefaultValue), EnumHasAnyFlags(ExpectedMask, EIntType::Int16) ? ECbFieldError::None : ECbFieldError::RangeError, TCbIntegerAccessors<int16, &FCbFieldView::AsInt16>());
-	TestField<ECbFieldType::IntegerNegative>(TEXT("Int32"), Field, int32(EnumHasAnyFlags(ExpectedMask, EIntType::Int32) ? ExpectedValue : DefaultValue),
+	TestField<ECbFieldType::IntegerNegative>(Test, TEXT("Int32"), Field, int32(EnumHasAnyFlags(ExpectedMask, EIntType::Int32) ? ExpectedValue : DefaultValue),
 		int32(DefaultValue), EnumHasAnyFlags(ExpectedMask, EIntType::Int32) ? ECbFieldError::None : ECbFieldError::RangeError, TCbIntegerAccessors<int32, &FCbFieldView::AsInt32>());
-	TestField<ECbFieldType::IntegerNegative>(TEXT("Int64"), Field, int64(EnumHasAnyFlags(ExpectedMask, EIntType::Int64) ? ExpectedValue : DefaultValue),
+	TestField<ECbFieldType::IntegerNegative>(Test, TEXT("Int64"), Field, int64(EnumHasAnyFlags(ExpectedMask, EIntType::Int64) ? ExpectedValue : DefaultValue),
 		int64(DefaultValue), EnumHasAnyFlags(ExpectedMask, EIntType::Int64) ? ECbFieldError::None : ECbFieldError::RangeError, TCbIntegerAccessors<int64, &FCbFieldView::AsInt64>());
-	TestField<ECbFieldType::IntegerPositive>(TEXT("UInt8"), Field, uint8(EnumHasAnyFlags(ExpectedMask, EIntType::UInt8) ? ExpectedValue : DefaultValue),
+	TestField<ECbFieldType::IntegerPositive>(Test, TEXT("UInt8"), Field, uint8(EnumHasAnyFlags(ExpectedMask, EIntType::UInt8) ? ExpectedValue : DefaultValue),
 		uint8(DefaultValue), EnumHasAnyFlags(ExpectedMask, EIntType::UInt8) ? ECbFieldError::None : ECbFieldError::RangeError, TCbIntegerAccessors<uint8, &FCbFieldView::AsUInt8>());
-	TestField<ECbFieldType::IntegerPositive>(TEXT("UInt16"), Field, uint16(EnumHasAnyFlags(ExpectedMask, EIntType::UInt16) ? ExpectedValue : DefaultValue),
+	TestField<ECbFieldType::IntegerPositive>(Test, TEXT("UInt16"), Field, uint16(EnumHasAnyFlags(ExpectedMask, EIntType::UInt16) ? ExpectedValue : DefaultValue),
 		uint16(DefaultValue), EnumHasAnyFlags(ExpectedMask, EIntType::UInt16) ? ECbFieldError::None : ECbFieldError::RangeError, TCbIntegerAccessors<uint16, &FCbFieldView::AsUInt16>());
-	TestField<ECbFieldType::IntegerPositive>(TEXT("UInt32"), Field, uint32(EnumHasAnyFlags(ExpectedMask, EIntType::UInt32) ? ExpectedValue : DefaultValue),
+	TestField<ECbFieldType::IntegerPositive>(Test, TEXT("UInt32"), Field, uint32(EnumHasAnyFlags(ExpectedMask, EIntType::UInt32) ? ExpectedValue : DefaultValue),
 		uint32(DefaultValue), EnumHasAnyFlags(ExpectedMask, EIntType::UInt32) ? ECbFieldError::None : ECbFieldError::RangeError, TCbIntegerAccessors<uint32, &FCbFieldView::AsUInt32>());
-	TestField<ECbFieldType::IntegerPositive>(TEXT("UInt64"), Field, uint64(EnumHasAnyFlags(ExpectedMask, EIntType::UInt64) ? ExpectedValue : DefaultValue),
+	TestField<ECbFieldType::IntegerPositive>(Test, TEXT("UInt64"), Field, uint64(EnumHasAnyFlags(ExpectedMask, EIntType::UInt64) ? ExpectedValue : DefaultValue),
 		uint64(DefaultValue), EnumHasAnyFlags(ExpectedMask, EIntType::UInt64) ? ECbFieldError::None : ECbFieldError::RangeError, TCbIntegerAccessors<uint64, &FCbFieldView::AsUInt64>());
 }
 
-TEST_CASE("Core::Serialization::FCbField::Integer Test", "[Core][Serialization][Smoke]")
+TEST_CASE_METHOD(FAutomationTestFixture, "Core::Serialization::FCbField::Integer Test", "[Core][Serialization][Smoke]")
 {
 	// Test FCbFieldView(IntegerPositive)
-	TestIntegerField(ECbFieldType::IntegerPositive, EIntType::Pos7,  0x00);
-	TestIntegerField(ECbFieldType::IntegerPositive, EIntType::Pos7,  0x7f);
-	TestIntegerField(ECbFieldType::IntegerPositive, EIntType::Pos8,  0x80);
-	TestIntegerField(ECbFieldType::IntegerPositive, EIntType::Pos8,  0xff);
-	TestIntegerField(ECbFieldType::IntegerPositive, EIntType::Pos15, 0x0100);
-	TestIntegerField(ECbFieldType::IntegerPositive, EIntType::Pos15, 0x7fff);
-	TestIntegerField(ECbFieldType::IntegerPositive, EIntType::Pos16, 0x8000);
-	TestIntegerField(ECbFieldType::IntegerPositive, EIntType::Pos16, 0xffff);
-	TestIntegerField(ECbFieldType::IntegerPositive, EIntType::Pos31, 0x0001'0000);
-	TestIntegerField(ECbFieldType::IntegerPositive, EIntType::Pos31, 0x7fff'ffff);
-	TestIntegerField(ECbFieldType::IntegerPositive, EIntType::Pos32, 0x8000'0000);
-	TestIntegerField(ECbFieldType::IntegerPositive, EIntType::Pos32, 0xffff'ffff);
-	TestIntegerField(ECbFieldType::IntegerPositive, EIntType::Pos63, 0x0000'0001'0000'0000);
-	TestIntegerField(ECbFieldType::IntegerPositive, EIntType::Pos63, 0x7fff'ffff'ffff'ffff);
-	TestIntegerField(ECbFieldType::IntegerPositive, EIntType::Pos64, 0x8000'0000'0000'0000);
-	TestIntegerField(ECbFieldType::IntegerPositive, EIntType::Pos64, 0xffff'ffff'ffff'ffff);
+	TestIntegerField(*this, ECbFieldType::IntegerPositive, EIntType::Pos7,  0x00);
+	TestIntegerField(*this, ECbFieldType::IntegerPositive, EIntType::Pos7,  0x7f);
+	TestIntegerField(*this, ECbFieldType::IntegerPositive, EIntType::Pos8,  0x80);
+	TestIntegerField(*this, ECbFieldType::IntegerPositive, EIntType::Pos8,  0xff);
+	TestIntegerField(*this, ECbFieldType::IntegerPositive, EIntType::Pos15, 0x0100);
+	TestIntegerField(*this, ECbFieldType::IntegerPositive, EIntType::Pos15, 0x7fff);
+	TestIntegerField(*this, ECbFieldType::IntegerPositive, EIntType::Pos16, 0x8000);
+	TestIntegerField(*this, ECbFieldType::IntegerPositive, EIntType::Pos16, 0xffff);
+	TestIntegerField(*this, ECbFieldType::IntegerPositive, EIntType::Pos31, 0x0001'0000);
+	TestIntegerField(*this, ECbFieldType::IntegerPositive, EIntType::Pos31, 0x7fff'ffff);
+	TestIntegerField(*this, ECbFieldType::IntegerPositive, EIntType::Pos32, 0x8000'0000);
+	TestIntegerField(*this, ECbFieldType::IntegerPositive, EIntType::Pos32, 0xffff'ffff);
+	TestIntegerField(*this, ECbFieldType::IntegerPositive, EIntType::Pos63, 0x0000'0001'0000'0000);
+	TestIntegerField(*this, ECbFieldType::IntegerPositive, EIntType::Pos63, 0x7fff'ffff'ffff'ffff);
+	TestIntegerField(*this, ECbFieldType::IntegerPositive, EIntType::Pos64, 0x8000'0000'0000'0000);
+	TestIntegerField(*this, ECbFieldType::IntegerPositive, EIntType::Pos64, 0xffff'ffff'ffff'ffff);
 
 	// Test FCbFieldView(IntegerNegative)
-	TestIntegerField(ECbFieldType::IntegerNegative, EIntType::Neg7,  0x01);
-	TestIntegerField(ECbFieldType::IntegerNegative, EIntType::Neg7,  0x80);
-	TestIntegerField(ECbFieldType::IntegerNegative, EIntType::Neg15, 0x81);
-	TestIntegerField(ECbFieldType::IntegerNegative, EIntType::Neg15, 0x8000);
-	TestIntegerField(ECbFieldType::IntegerNegative, EIntType::Neg31, 0x8001);
-	TestIntegerField(ECbFieldType::IntegerNegative, EIntType::Neg31, 0x8000'0000);
-	TestIntegerField(ECbFieldType::IntegerNegative, EIntType::Neg63, 0x8000'0001);
-	TestIntegerField(ECbFieldType::IntegerNegative, EIntType::Neg63, 0x8000'0000'0000'0000);
-	TestIntegerField(ECbFieldType::IntegerNegative, EIntType::None,  0x8000'0000'0000'0001);
-	TestIntegerField(ECbFieldType::IntegerNegative, EIntType::None,  0xffff'ffff'ffff'ffff);
+	TestIntegerField(*this, ECbFieldType::IntegerNegative, EIntType::Neg7,  0x01);
+	TestIntegerField(*this, ECbFieldType::IntegerNegative, EIntType::Neg7,  0x80);
+	TestIntegerField(*this, ECbFieldType::IntegerNegative, EIntType::Neg15, 0x81);
+	TestIntegerField(*this, ECbFieldType::IntegerNegative, EIntType::Neg15, 0x8000);
+	TestIntegerField(*this, ECbFieldType::IntegerNegative, EIntType::Neg31, 0x8001);
+	TestIntegerField(*this, ECbFieldType::IntegerNegative, EIntType::Neg31, 0x8000'0000);
+	TestIntegerField(*this, ECbFieldType::IntegerNegative, EIntType::Neg63, 0x8000'0001);
+	TestIntegerField(*this, ECbFieldType::IntegerNegative, EIntType::Neg63, 0x8000'0000'0000'0000);
+	TestIntegerField(*this, ECbFieldType::IntegerNegative, EIntType::None,  0x8000'0000'0000'0001);
+	TestIntegerField(*this, ECbFieldType::IntegerNegative, EIntType::None,  0xffff'ffff'ffff'ffff);
 
 	// Test FCbFieldView(None) as Integer
 	{
 		FCbFieldView Field;
-		TestFieldError<ECbFieldType::IntegerPositive>(TEXT("Integer+, None"), Field, ECbFieldError::TypeError, uint64(8));
-		TestFieldError<ECbFieldType::IntegerNegative>(TEXT("Integer-, None"), Field, ECbFieldError::TypeError, int64(8));
+		TestFieldError<ECbFieldType::IntegerPositive>(*this, TEXT("Integer+, None"), Field, ECbFieldError::TypeError, uint64(8));
+		TestFieldError<ECbFieldType::IntegerNegative>(*this, TEXT("Integer-, None"), Field, ECbFieldError::TypeError, int64(8));
 	}
 }
 
 
-TEST_CASE("Core::Serialization::FCbField::Float Test", "[Core][Serialization][Smoke]")
+TEST_CASE_METHOD(FAutomationTestFixture, "Core::Serialization::FCbField::Float Test", "[Core][Serialization][Smoke]")
 {
 	//// Test FCbFieldView(Float, 32-bit)
 	//{
@@ -853,10 +853,10 @@ TEST_CASE("Core::Serialization::FCbField::Float Test", "[Core][Serialization][Sm
 	// Test FCbFieldView(Float, 64-bit)
 	{
 		const uint8 Value[] = { 0xc1, 0x23, 0x45, 0x67, 0x89, 0xab, 0xcd, 0xef }; // -631475.76888888876
-		TestField<ECbFieldType::Float64>(TEXT("Float64"), Value, -631475.76888888876);
+		TestField<ECbFieldType::Float64>(*this, TEXT("Float64"), Value, -631475.76888888876);
 
 		FCbFieldView Field(Value, ECbFieldType::Float64);
-		TestFieldError<ECbFieldType::Float32>(TEXT("Float64, AsFloat"), Field, ECbFieldError::RangeError, 8.0f);
+		TestFieldError<ECbFieldType::Float32>(*this, TEXT("Float64, AsFloat"), Field, ECbFieldError::RangeError, 8.0f);
 	}
 
 	// Test FCbFieldView(Integer+, MaxBinary32) as Float
@@ -864,8 +864,8 @@ TEST_CASE("Core::Serialization::FCbField::Float Test", "[Core][Serialization][Sm
 		uint8 Value[9];
 		WriteVarUInt((uint64(1) << 24) - 1, Value); // 16,777,215
 		FCbFieldView Field(Value, ECbFieldType::IntegerPositive);
-		TestField<ECbFieldType::Float32>(TEXT("Integer+, MaxBinary32, AsFloat"), Field, 16'777'215.0f);
-		TestField<ECbFieldType::Float64>(TEXT("Integer+, MaxBinary32, AsDouble"), Field, 16'777'215.0);
+		TestField<ECbFieldType::Float32>(*this, TEXT("Integer+, MaxBinary32, AsFloat"), Field, 16'777'215.0f);
+		TestField<ECbFieldType::Float64>(*this, TEXT("Integer+, MaxBinary32, AsDouble"), Field, 16'777'215.0);
 	}
 
 	// Test FCbFieldView(Integer+, MaxBinary32+1) as Float
@@ -873,8 +873,8 @@ TEST_CASE("Core::Serialization::FCbField::Float Test", "[Core][Serialization][Sm
 		uint8 Value[9];
 		WriteVarUInt(uint64(1) << 24, Value); // 16,777,216
 		FCbFieldView Field(Value, ECbFieldType::IntegerPositive);
-		TestFieldError<ECbFieldType::Float32>(TEXT("Integer+, MaxBinary32+1, AsFloat"), Field, ECbFieldError::RangeError, 8.0f);
-		TestField<ECbFieldType::Float64>(TEXT("Integer+, MaxBinary32+1, AsDouble"), Field, 16'777'216.0);
+		TestFieldError<ECbFieldType::Float32>(*this, TEXT("Integer+, MaxBinary32+1, AsFloat"), Field, ECbFieldError::RangeError, 8.0f);
+		TestField<ECbFieldType::Float64>(*this, TEXT("Integer+, MaxBinary32+1, AsDouble"), Field, 16'777'216.0);
 	}
 
 	// Test FCbFieldView(Integer+, MaxBinary64) as Float
@@ -882,8 +882,8 @@ TEST_CASE("Core::Serialization::FCbField::Float Test", "[Core][Serialization][Sm
 		uint8 Value[9];
 		WriteVarUInt((uint64(1) << 53) - 1, Value); // 9,007,199,254,740,991
 		FCbFieldView Field(Value, ECbFieldType::IntegerPositive);
-		TestFieldError<ECbFieldType::Float32>(TEXT("Integer+, MaxBinary64, AsFloat"), Field, ECbFieldError::RangeError, 8.0f);
-		TestField<ECbFieldType::Float64>(TEXT("Integer+, MaxBinary64, AsDouble"), Field, 9'007'199'254'740'991.0);
+		TestFieldError<ECbFieldType::Float32>(*this, TEXT("Integer+, MaxBinary64, AsFloat"), Field, ECbFieldError::RangeError, 8.0f);
+		TestField<ECbFieldType::Float64>(*this, TEXT("Integer+, MaxBinary64, AsDouble"), Field, 9'007'199'254'740'991.0);
 	}
 
 	// Test FCbFieldView(Integer+, MaxBinary64+1) as Float
@@ -891,8 +891,8 @@ TEST_CASE("Core::Serialization::FCbField::Float Test", "[Core][Serialization][Sm
 		uint8 Value[9];
 		WriteVarUInt(uint64(1) << 53, Value); // 9,007,199,254,740,992
 		FCbFieldView Field(Value, ECbFieldType::IntegerPositive);
-		TestFieldError<ECbFieldType::Float32>(TEXT("Integer+, MaxBinary64+1, AsFloat"), Field, ECbFieldError::RangeError, 8.0f);
-		TestFieldError<ECbFieldType::Float64>(TEXT("Integer+, MaxBinary64+1, AsDouble"), Field, ECbFieldError::RangeError, 8.0);
+		TestFieldError<ECbFieldType::Float32>(*this, TEXT("Integer+, MaxBinary64+1, AsFloat"), Field, ECbFieldError::RangeError, 8.0f);
+		TestFieldError<ECbFieldType::Float64>(*this, TEXT("Integer+, MaxBinary64+1, AsDouble"), Field, ECbFieldError::RangeError, 8.0);
 	}
 
 	// Test FCbFieldView(Integer+, MaxUInt64) as Float
@@ -900,8 +900,8 @@ TEST_CASE("Core::Serialization::FCbField::Float Test", "[Core][Serialization][Sm
 		uint8 Value[9];
 		WriteVarUInt(uint64(-1), Value); // Max uint64
 		FCbFieldView Field(Value, ECbFieldType::IntegerPositive);
-		TestFieldError<ECbFieldType::Float32>(TEXT("Integer+, MaxUInt64, AsFloat"), Field, ECbFieldError::RangeError, 8.0f);
-		TestFieldError<ECbFieldType::Float64>(TEXT("Integer+, MaxUInt64, AsDouble"), Field, ECbFieldError::RangeError, 8.0);
+		TestFieldError<ECbFieldType::Float32>(*this, TEXT("Integer+, MaxUInt64, AsFloat"), Field, ECbFieldError::RangeError, 8.0f);
+		TestFieldError<ECbFieldType::Float64>(*this, TEXT("Integer+, MaxUInt64, AsDouble"), Field, ECbFieldError::RangeError, 8.0);
 	}
 
 	// Test FCbFieldView(Integer-, MaxBinary32) as Float
@@ -909,8 +909,8 @@ TEST_CASE("Core::Serialization::FCbField::Float Test", "[Core][Serialization][Sm
 		uint8 Value[9];
 		WriteVarUInt((uint64(1) << 24) - 2, Value); // -16,777,215
 		FCbFieldView Field(Value, ECbFieldType::IntegerNegative);
-		TestField<ECbFieldType::Float32>(TEXT("Integer-, MaxBinary32, AsFloat"), Field, -16'777'215.0f);
-		TestField<ECbFieldType::Float64>(TEXT("Integer-, MaxBinary32, AsDouble"), Field, -16'777'215.0);
+		TestField<ECbFieldType::Float32>(*this, TEXT("Integer-, MaxBinary32, AsFloat"), Field, -16'777'215.0f);
+		TestField<ECbFieldType::Float64>(*this, TEXT("Integer-, MaxBinary32, AsDouble"), Field, -16'777'215.0);
 	}
 
 	// Test FCbFieldView(Integer-, MaxBinary32+1) as Float
@@ -918,8 +918,8 @@ TEST_CASE("Core::Serialization::FCbField::Float Test", "[Core][Serialization][Sm
 		uint8 Value[9];
 		WriteVarUInt((uint64(1) << 24) - 1, Value); // -16,777,216
 		FCbFieldView Field(Value, ECbFieldType::IntegerNegative);
-		TestFieldError<ECbFieldType::Float32>(TEXT("Integer-, MaxBinary32+1, AsFloat"), Field, ECbFieldError::RangeError, 8.0f);
-		TestField<ECbFieldType::Float64>(TEXT("Integer-, MaxBinary32+1, AsDouble"), Field, -16'777'216.0);
+		TestFieldError<ECbFieldType::Float32>(*this, TEXT("Integer-, MaxBinary32+1, AsFloat"), Field, ECbFieldError::RangeError, 8.0f);
+		TestField<ECbFieldType::Float64>(*this, TEXT("Integer-, MaxBinary32+1, AsDouble"), Field, -16'777'216.0);
 	}
 
 	// Test FCbFieldView(Integer-, MaxBinary64) as Float
@@ -927,8 +927,8 @@ TEST_CASE("Core::Serialization::FCbField::Float Test", "[Core][Serialization][Sm
 		uint8 Value[9];
 		WriteVarUInt((uint64(1) << 53) - 2, Value); // -9,007,199,254,740,991
 		FCbFieldView Field(Value, ECbFieldType::IntegerNegative);
-		TestFieldError<ECbFieldType::Float32>(TEXT("Integer-, MaxBinary64, AsFloat"), Field, ECbFieldError::RangeError, 8.0f);
-		TestField<ECbFieldType::Float64>(TEXT("Integer-, MaxBinary64, AsDouble"), Field, -9'007'199'254'740'991.0);
+		TestFieldError<ECbFieldType::Float32>(*this, TEXT("Integer-, MaxBinary64, AsFloat"), Field, ECbFieldError::RangeError, 8.0f);
+		TestField<ECbFieldType::Float64>(*this, TEXT("Integer-, MaxBinary64, AsDouble"), Field, -9'007'199'254'740'991.0);
 	}
 
 	// Test FCbFieldView(Integer-, MaxBinary64+1) as Float
@@ -936,212 +936,212 @@ TEST_CASE("Core::Serialization::FCbField::Float Test", "[Core][Serialization][Sm
 		uint8 Value[9];
 		WriteVarUInt((uint64(1) << 53) - 1, Value); // -9,007,199,254,740,992
 		FCbFieldView Field(Value, ECbFieldType::IntegerNegative);
-		TestFieldError<ECbFieldType::Float32>(TEXT("Integer-, MaxBinary64+1, AsFloat"), Field, ECbFieldError::RangeError, 8.0f);
-		TestFieldError<ECbFieldType::Float64>(TEXT("Integer-, MaxBinary64+1, AsDouble"), Field, ECbFieldError::RangeError, 8.0);
+		TestFieldError<ECbFieldType::Float32>(*this, TEXT("Integer-, MaxBinary64+1, AsFloat"), Field, ECbFieldError::RangeError, 8.0f);
+		TestFieldError<ECbFieldType::Float64>(*this, TEXT("Integer-, MaxBinary64+1, AsDouble"), Field, ECbFieldError::RangeError, 8.0);
 	}
 
 	// Test FCbFieldView(None) as Float
 	{
 		FCbFieldView Field;
-		TestFieldError<ECbFieldType::Float32>(TEXT("None, AsFloat"), Field, ECbFieldError::TypeError, 8.0f);
-		TestFieldError<ECbFieldType::Float64>(TEXT("None, AsDouble"), Field, ECbFieldError::TypeError, 8.0);
+		TestFieldError<ECbFieldType::Float32>(*this, TEXT("None, AsFloat"), Field, ECbFieldError::TypeError, 8.0f);
+		TestFieldError<ECbFieldType::Float64>(*this, TEXT("None, AsDouble"), Field, ECbFieldError::TypeError, 8.0);
 	}
 }
 
 
-TEST_CASE("Core::Serialization::FCbField::Bool Test", "[Core][Serialization][Smoke]")
+TEST_CASE_METHOD(FAutomationTestFixture, "Core::Serialization::FCbField::Bool Test", "[Core][Serialization][Smoke]")
 {
 	// Test FCbFieldView(Bool, False)
-	TestField<ECbFieldType::BoolFalse>(TEXT("Bool, False"), {}, false, true);
+	TestField<ECbFieldType::BoolFalse>(*this, TEXT("Bool, False"), {}, false, true);
 
 	// Test FCbFieldView(Bool, True)
-	TestField<ECbFieldType::BoolTrue>(TEXT("Bool, True"), {}, true, false);
+	TestField<ECbFieldType::BoolTrue>(*this, TEXT("Bool, True"), {}, true, false);
 
 	// Test FCbFieldView(None) as Bool
 	{
 		FCbFieldView DefaultField;
-		TestFieldError<ECbFieldType::BoolFalse>(TEXT("Bool, False, None"), DefaultField, ECbFieldError::TypeError, false);
-		TestFieldError<ECbFieldType::BoolTrue>(TEXT("Bool, True, None"), DefaultField, ECbFieldError::TypeError, true);
+		TestFieldError<ECbFieldType::BoolFalse>(*this, TEXT("Bool, False, None"), DefaultField, ECbFieldError::TypeError, false);
+		TestFieldError<ECbFieldType::BoolTrue>(*this, TEXT("Bool, True, None"), DefaultField, ECbFieldError::TypeError, true);
 	}
 }
 
 
-TEST_CASE("Core::Serialization::FCbField::Object Attachment Test", "[Core][Serialization][Smoke]")
+TEST_CASE_METHOD(FAutomationTestFixture, "Core::Serialization::FCbField::Object Attachment Test", "[Core][Serialization][Smoke]")
 {
 	const FIoHash::ByteArray ZeroBytes{};
 	const FIoHash::ByteArray SequentialBytes{1, 2, 3, 4, 5, 6, 7, 8, 9, 10, 11, 12, 13, 14, 15, 16, 17, 18, 19, 20};
 
 	// Test FCbFieldView(ObjectAttachment, Zero)
-	TestField<ECbFieldType::ObjectAttachment>(TEXT("ObjectAttachment, Zero"), ZeroBytes);
+	TestField<ECbFieldType::ObjectAttachment>(*this, TEXT("ObjectAttachment, Zero"), ZeroBytes);
 
 	// Test FCbFieldView(ObjectAttachment, NonZero)
-	TestField<ECbFieldType::ObjectAttachment>(TEXT("ObjectAttachment, NonZero"), SequentialBytes, FIoHash(SequentialBytes));
+	TestField<ECbFieldType::ObjectAttachment>(*this, TEXT("ObjectAttachment, NonZero"), SequentialBytes, FIoHash(SequentialBytes));
 
 	// Test FCbFieldView(ObjectAttachment, NonZero) AsAttachment
 	{
 		FCbFieldView Field(SequentialBytes, ECbFieldType::ObjectAttachment);
-		TestField<ECbFieldType::ObjectAttachment>(TEXT("ObjectAttachment, NonZero, AsAttachment"), Field, FIoHash(SequentialBytes), FIoHash(), ECbFieldError::None, FCbAttachmentAccessors());
+		TestField<ECbFieldType::ObjectAttachment>(*this, TEXT("ObjectAttachment, NonZero, AsAttachment"), Field, FIoHash(SequentialBytes), FIoHash(), ECbFieldError::None, FCbAttachmentAccessors());
 	}
 
 	// Test FCbFieldView(None) as ObjectAttachment
 	{
 		FCbFieldView DefaultField;
-		TestFieldError<ECbFieldType::ObjectAttachment>(TEXT("ObjectAttachment, None"), DefaultField, ECbFieldError::TypeError, FIoHash(SequentialBytes));
+		TestFieldError<ECbFieldType::ObjectAttachment>(*this, TEXT("ObjectAttachment, None"), DefaultField, ECbFieldError::TypeError, FIoHash(SequentialBytes));
 	}
 }
 
 
-TEST_CASE("Core::Serialization::FCbField::Binary  Attachment Test", "[Core][Serialization][Smoke]")
+TEST_CASE_METHOD(FAutomationTestFixture, "Core::Serialization::FCbField::Binary  Attachment Test", "[Core][Serialization][Smoke]")
 {
 	const FIoHash::ByteArray ZeroBytes{};
 	const FIoHash::ByteArray SequentialBytes{1, 2, 3, 4, 5, 6, 7, 8, 9, 10, 11, 12, 13, 14, 15, 16, 17, 18, 19, 20};
 
 	// Test FCbFieldView(BinaryAttachment, Zero)
-	TestField<ECbFieldType::BinaryAttachment>(TEXT("BinaryAttachment, Zero"), ZeroBytes);
+	TestField<ECbFieldType::BinaryAttachment>(*this, TEXT("BinaryAttachment, Zero"), ZeroBytes);
 
 	// Test FCbFieldView(BinaryAttachment, NonZero)
-	TestField<ECbFieldType::BinaryAttachment>(TEXT("BinaryAttachment, NonZero"), SequentialBytes, FIoHash(SequentialBytes));
+	TestField<ECbFieldType::BinaryAttachment>(*this, TEXT("BinaryAttachment, NonZero"), SequentialBytes, FIoHash(SequentialBytes));
 
 	// Test FCbFieldView(BinaryAttachment, NonZero) AsAttachment
 	{
 		FCbFieldView Field(SequentialBytes, ECbFieldType::BinaryAttachment);
-		TestField<ECbFieldType::BinaryAttachment>(TEXT("BinaryAttachment, NonZero, AsAttachment"), Field, FIoHash(SequentialBytes), FIoHash(), ECbFieldError::None, FCbAttachmentAccessors());
+		TestField<ECbFieldType::BinaryAttachment>(*this, TEXT("BinaryAttachment, NonZero, AsAttachment"), Field, FIoHash(SequentialBytes), FIoHash(), ECbFieldError::None, FCbAttachmentAccessors());
 	}
 
 	// Test FCbFieldView(None) as BinaryAttachment
 	{
 		FCbFieldView DefaultField;
-		TestFieldError<ECbFieldType::BinaryAttachment>(TEXT("BinaryAttachment, None"), DefaultField, ECbFieldError::TypeError, FIoHash(SequentialBytes));
+		TestFieldError<ECbFieldType::BinaryAttachment>(*this, TEXT("BinaryAttachment, None"), DefaultField, ECbFieldError::TypeError, FIoHash(SequentialBytes));
 	}
 }
 
 
-TEST_CASE("Core::Serialization::FCbField::Hash Test", "[Core][Serialization][Smoke]")
+TEST_CASE_METHOD(FAutomationTestFixture, "Core::Serialization::FCbField::Hash Test", "[Core][Serialization][Smoke]")
 {
 	const FIoHash::ByteArray ZeroBytes{};
 	const FIoHash::ByteArray SequentialBytes{1, 2, 3, 4, 5, 6, 7, 8, 9, 10, 11, 12, 13, 14, 15, 16, 17, 18, 19, 20};
 
 	// Test FCbFieldView(Hash, Zero)
-	TestField<ECbFieldType::Hash>(TEXT("Hash, Zero"), ZeroBytes);
+	TestField<ECbFieldType::Hash>(*this, TEXT("Hash, Zero"), ZeroBytes);
 
 	// Test FCbFieldView(Hash, NonZero)
-	TestField<ECbFieldType::Hash>(TEXT("Hash, NonZero"), SequentialBytes, FIoHash(SequentialBytes));
+	TestField<ECbFieldType::Hash>(*this, TEXT("Hash, NonZero"), SequentialBytes, FIoHash(SequentialBytes));
 
 	// Test FCbFieldView(None) as Hash
 	{
 		FCbFieldView DefaultField;
-		TestFieldError<ECbFieldType::Hash>(TEXT("Hash, None"), DefaultField, ECbFieldError::TypeError, FIoHash(SequentialBytes));
+		TestFieldError<ECbFieldType::Hash>(*this, TEXT("Hash, None"), DefaultField, ECbFieldError::TypeError, FIoHash(SequentialBytes));
 	}
 
 	// Test FCbFieldView(ObjectAttachment) as Hash
 	{
 		FCbFieldView Field(SequentialBytes, ECbFieldType::ObjectAttachment);
-		TestField<ECbFieldType::Hash>(TEXT("ObjectAttachment, NonZero, AsHash"), Field, FIoHash(SequentialBytes));
+		TestField<ECbFieldType::Hash>(*this, TEXT("ObjectAttachment, NonZero, AsHash"), Field, FIoHash(SequentialBytes));
 	}
 
 	// Test FCbFieldView(BinaryAttachment) as Hash
 	{
 		FCbFieldView Field(SequentialBytes, ECbFieldType::BinaryAttachment);
-		TestField<ECbFieldType::Hash>(TEXT("BinaryAttachment, NonZero, AsHash"), Field, FIoHash(SequentialBytes));
+		TestField<ECbFieldType::Hash>(*this, TEXT("BinaryAttachment, NonZero, AsHash"), Field, FIoHash(SequentialBytes));
 	}
 }
 
-TEST_CASE("Core::Serialization::FCbField::Uuid Test", "[Core][Serialization][Smoke]")
+TEST_CASE_METHOD(FAutomationTestFixture, "Core::Serialization::FCbField::Uuid Test", "[Core][Serialization][Smoke]")
 {
 	const uint8 ZeroBytes[]{0, 0, 0, 0, 0, 0, 0, 0, 0, 0, 0, 0, 0, 0, 0, 0};
 	const uint8 SequentialBytes[]{0, 1, 2, 3, 4, 5, 6, 7, 8, 9, 10, 11, 12, 13, 14, 15};
 	const FGuid SequentialGuid(TEXT("00010203-0405-0607-0809-0a0b0c0d0e0f"));
 
 	// Test FCbFieldView(Uuid, Zero)
-	TestField<ECbFieldType::Uuid>(TEXT("Uuid, Zero"), ZeroBytes, FGuid(), SequentialGuid);
+	TestField<ECbFieldType::Uuid>(*this, TEXT("Uuid, Zero"), ZeroBytes, FGuid(), SequentialGuid);
 
 	// Test FCbFieldView(Uuid, NonZero)
-	TestField<ECbFieldType::Uuid>(TEXT("Uuid, NonZero"), SequentialBytes, SequentialGuid, FGuid());
+	TestField<ECbFieldType::Uuid>(*this, TEXT("Uuid, NonZero"), SequentialBytes, SequentialGuid, FGuid());
 
 	// Test FCbFieldView(None) as Uuid
 	{
 		FCbFieldView DefaultField;
-		TestFieldError<ECbFieldType::Uuid>(TEXT("Uuid, None"), DefaultField, ECbFieldError::TypeError, FGuid::NewGuid());
+		TestFieldError<ECbFieldType::Uuid>(*this, TEXT("Uuid, None"), DefaultField, ECbFieldError::TypeError, FGuid::NewGuid());
 	}
 }
 
 
-TEST_CASE("Core::Serialization::FCbField::Date Time", "[Core][Serialization][Smoke]")
+TEST_CASE_METHOD(FAutomationTestFixture, "Core::Serialization::FCbField::Date Time", "[Core][Serialization][Smoke]")
 {
 	// Test FCbFieldView(DateTime, Zero)
-	TestField<ECbFieldType::DateTime>(TEXT("DateTime, Zero"), {0, 0, 0, 0, 0, 0, 0, 0});
+	TestField<ECbFieldType::DateTime>(*this, TEXT("DateTime, Zero"), {0, 0, 0, 0, 0, 0, 0, 0});
 
 	// Test FCbFieldView(DateTime, 0x1020'3040'5060'7080)
-	TestField<ECbFieldType::DateTime>(TEXT("DateTime, NonZero"), {0x10, 0x20, 0x30, 0x40, 0x50, 0x60, 0x70, 0x80}, int64(0x1020'3040'5060'7080));
+	TestField<ECbFieldType::DateTime>(*this, TEXT("DateTime, NonZero"), {0x10, 0x20, 0x30, 0x40, 0x50, 0x60, 0x70, 0x80}, int64(0x1020'3040'5060'7080));
 
 	// Test FCbFieldView(DateTime, Zero) as FDateTime
 	{
 		const uint8 Value[] = {0, 0, 0, 0, 0, 0, 0, 0};
 		FCbFieldView Field(Value, ECbFieldType::DateTime);
-		TestEqual(TEXT("FCbFieldView()::AsDateTime()"), Field.AsDateTime(), FDateTime(0));
+		TEST_EQUAL(TEXT("FCbFieldView()::AsDateTime()"), Field.AsDateTime(), FDateTime(0));
 	}
 
 	// Test FCbFieldView(None) as DateTime
 	{
 		FCbFieldView DefaultField;
-		TestFieldError<ECbFieldType::DateTime>(TEXT("DateTime, None"), DefaultField, ECbFieldError::TypeError);
+		TestFieldError<ECbFieldType::DateTime>(*this, TEXT("DateTime, None"), DefaultField, ECbFieldError::TypeError);
 		const FDateTime DefaultValue(0x1020'3040'5060'7080);
-		TestEqual(TEXT("FCbFieldView()::AsDateTime()"), DefaultField.AsDateTime(DefaultValue), DefaultValue);
+		TEST_EQUAL(TEXT("FCbFieldView()::AsDateTime()"), DefaultField.AsDateTime(DefaultValue), DefaultValue);
 	}
 }
 
-TEST_CASE("Core::Serialization::FCbField::Time Span", "[Core][Serialization][Smoke]")
+TEST_CASE_METHOD(FAutomationTestFixture, "Core::Serialization::FCbField::Time Span", "[Core][Serialization][Smoke]")
 {
 	// Test FCbFieldView(TimeSpan, Zero)
-	TestField<ECbFieldType::TimeSpan>(TEXT("TimeSpan, Zero"), {0, 0, 0, 0, 0, 0, 0, 0});
+	TestField<ECbFieldType::TimeSpan>(*this, TEXT("TimeSpan, Zero"), {0, 0, 0, 0, 0, 0, 0, 0});
 
 	// Test FCbFieldView(TimeSpan, 0x1020'3040'5060'7080)
-	TestField<ECbFieldType::TimeSpan>(TEXT("TimeSpan, NonZero"), {0x10, 0x20, 0x30, 0x40, 0x50, 0x60, 0x70, 0x80}, int64(0x1020'3040'5060'7080));
+	TestField<ECbFieldType::TimeSpan>(*this, TEXT("TimeSpan, NonZero"), {0x10, 0x20, 0x30, 0x40, 0x50, 0x60, 0x70, 0x80}, int64(0x1020'3040'5060'7080));
 
 	// Test FCbFieldView(TimeSpan, Zero) as FTimeSpan
 	{
 		const uint8 Value[] = {0, 0, 0, 0, 0, 0, 0, 0};
 		FCbFieldView Field(Value, ECbFieldType::TimeSpan);
-		TestEqual(TEXT("FCbFieldView()::AsTimeSpan()"), Field.AsTimeSpan(), FTimespan(0));
+		TEST_EQUAL(TEXT("FCbFieldView()::AsTimeSpan()"), Field.AsTimeSpan(), FTimespan(0));
 	}
 
 	// Test FCbFieldView(None) as TimeSpan
 	{
 		FCbFieldView DefaultField;
-		TestFieldError<ECbFieldType::TimeSpan>(TEXT("TimeSpan, None"), DefaultField, ECbFieldError::TypeError);
+		TestFieldError<ECbFieldType::TimeSpan>(*this, TEXT("TimeSpan, None"), DefaultField, ECbFieldError::TypeError);
 		const FTimespan DefaultValue(0x1020'3040'5060'7080);
-		TestEqual(TEXT("FCbFieldView()::AsTimeSpan()"), DefaultField.AsTimeSpan(DefaultValue), DefaultValue);
+		TEST_EQUAL(TEXT("FCbFieldView()::AsTimeSpan()"), DefaultField.AsTimeSpan(DefaultValue), DefaultValue);
 	}
 }
 
 
-TEST_CASE("Core::Serialization::FCbField::ObjectId Span", "[Core][Serialization][Smoke]")
+TEST_CASE_METHOD(FAutomationTestFixture, "Core::Serialization::FCbField::ObjectId Span", "[Core][Serialization][Smoke]")
 {
 	// Test FCbFieldView(ObjectId, Zero)
-	TestField<ECbFieldType::ObjectId>(TEXT("ObjectId, Zero"), {0, 0, 0, 0, 0, 0, 0, 0, 0, 0, 0, 0});
+	TestField<ECbFieldType::ObjectId>(*this, TEXT("ObjectId, Zero"), {0, 0, 0, 0, 0, 0, 0, 0, 0, 0, 0, 0});
 
 	// Test FCbFieldView(ObjectId, 0x102030405060708090A0B0C0)
-	TestField<ECbFieldType::ObjectId>(TEXT("ObjectId, NonZero"), {0x10, 0x20, 0x30, 0x40, 0x50, 0x60, 0x70, 0x80, 0x90, 0xA0, 0xB0, 0xC0},
+	TestField<ECbFieldType::ObjectId>(*this, TEXT("ObjectId, NonZero"), {0x10, 0x20, 0x30, 0x40, 0x50, 0x60, 0x70, 0x80, 0x90, 0xA0, 0xB0, 0xC0},
 		FCbObjectId(MakeMemoryView<uint8>({0x10, 0x20, 0x30, 0x40, 0x50, 0x60, 0x70, 0x80, 0x90, 0xA0, 0xB0, 0xC0})));
 
 	// Test FCbFieldView(ObjectId, Zero) as FCbObjectId
 	{
 		const uint8 Value[] = {0, 0, 0, 0, 0, 0, 0, 0, 0, 0, 0, 0};
 		FCbFieldView Field(Value, ECbFieldType::ObjectId);
-		TestEqual(TEXT("FCbFieldView(ObjectId, Zero)::AsObjectId()"), Field.AsObjectId(), FCbObjectId());
+		TEST_EQUAL(TEXT("FCbFieldView(ObjectId, Zero)::AsObjectId()"), Field.AsObjectId(), FCbObjectId());
 	}
 
 	// Test FCbFieldView(None) as ObjectId
 	{
 		FCbFieldView DefaultField;
-		TestFieldError<ECbFieldType::ObjectId>(TEXT("ObjectId, None"), DefaultField, ECbFieldError::TypeError);
+		TestFieldError<ECbFieldType::ObjectId>(*this, TEXT("ObjectId, None"), DefaultField, ECbFieldError::TypeError);
 		const FCbObjectId DefaultValue(MakeMemoryView<uint8>({0x10, 0x20, 0x30, 0x40, 0x50, 0x60, 0x70, 0x80, 0x90, 0xA0, 0xB0, 0xC0}));
-		TestEqual(TEXT("FCbFieldView(None)::AsObjectId()"), DefaultField.AsObjectId(DefaultValue), DefaultValue);
+		TEST_EQUAL(TEXT("FCbFieldView(None)::AsObjectId()"), DefaultField.AsObjectId(DefaultValue), DefaultValue);
 	}
 }
 
 
-TEST_CASE("Core::Serialization::FCbField::Custom By Id", "[Core][Serialization][Smoke]")
+TEST_CASE_METHOD(FAutomationTestFixture, "Core::Serialization::FCbField::Custom By Id", "[Core][Serialization][Smoke]")
 {
 	struct FCustomByIdAccessor
 	{
@@ -1157,46 +1157,46 @@ TEST_CASE("Core::Serialization::FCbField::Custom By Id", "[Core][Serialization][
 	// Test FCbFieldView(CustomById, MinId, Empty)
 	{
 		const uint8 Value[] = {1, 0};
-		TestField<ECbFieldType::CustomById>(TEXT("CustomById, MinId, Empty"), Value, FCbCustomById{0});
-		TestField<ECbFieldType::CustomById>(TEXT("CustomById, MinId, Empty, View"), Value, FMemoryView(), MakeMemoryView<uint8>({1, 2, 3}), ECbFieldError::None, FCustomByIdAccessor(0));
-		TestFieldError<ECbFieldType::CustomById>(TEXT("CustomById, MinId, Empty, InvalidId"), Value, ECbFieldError::RangeError, MakeMemoryView<uint8>({1, 2, 3}), FCustomByIdAccessor(MAX_uint64));
+		TestField<ECbFieldType::CustomById>(*this, TEXT("CustomById, MinId, Empty"), Value, FCbCustomById{0});
+		TestField<ECbFieldType::CustomById>(*this, TEXT("CustomById, MinId, Empty, View"), Value, FMemoryView(), MakeMemoryView<uint8>({1, 2, 3}), ECbFieldError::None, FCustomByIdAccessor(0));
+		TestFieldError<ECbFieldType::CustomById>(*this, TEXT("CustomById, MinId, Empty, InvalidId"), Value, ECbFieldError::RangeError, MakeMemoryView<uint8>({1, 2, 3}), FCustomByIdAccessor(MAX_uint64));
 	}
 
 	// Test FCbFieldView(CustomById, MinId, Value)
 	{
 		const uint8 Value[] = {5, 0, 1, 2, 3, 4};
-		TestFieldNoClone<ECbFieldType::CustomById>(TEXT("CustomById, MinId, Value"), Value, FCbCustomById{0, MakeMemoryView(Value).Right(4)});
-		TestFieldNoClone<ECbFieldType::CustomById>(TEXT("CustomById, MinId, Value, View"), Value, MakeMemoryView(Value).Right(4), FMemoryView(), ECbFieldError::None, FCustomByIdAccessor(0));
-		TestFieldError<ECbFieldType::CustomById>(TEXT("CustomById, MinId, Value, InvalidId"), Value, ECbFieldError::RangeError, MakeMemoryView<uint8>({1, 2, 3}), FCustomByIdAccessor(MAX_uint64));
+		TestFieldNoClone<ECbFieldType::CustomById>(*this, TEXT("CustomById, MinId, Value"), Value, FCbCustomById{0, MakeMemoryView(Value).Right(4)});
+		TestFieldNoClone<ECbFieldType::CustomById>(*this, TEXT("CustomById, MinId, Value, View"), Value, MakeMemoryView(Value).Right(4), FMemoryView(), ECbFieldError::None, FCustomByIdAccessor(0));
+		TestFieldError<ECbFieldType::CustomById>(*this, TEXT("CustomById, MinId, Value, InvalidId"), Value, ECbFieldError::RangeError, MakeMemoryView<uint8>({1, 2, 3}), FCustomByIdAccessor(MAX_uint64));
 	}
 
 	// Test FCbFieldView(CustomById, MaxId, Empty)
 	{
 		const uint8 Value[] = {9, 0xff, 0xff, 0xff, 0xff, 0xff, 0xff, 0xff, 0xff, 0xff};
-		TestField<ECbFieldType::CustomById>(TEXT("CustomById, MaxId, Empty"), Value, FCbCustomById{MAX_uint64});
-		TestField<ECbFieldType::CustomById>(TEXT("CustomById, MaxId, Empty, View"), Value, FMemoryView(), MakeMemoryView<uint8>({1, 2, 3}), ECbFieldError::None, FCustomByIdAccessor(MAX_uint64));
-		TestFieldError<ECbFieldType::CustomById>(TEXT("CustomById, MaxId, Empty, InvalidId"), Value, ECbFieldError::RangeError, MakeMemoryView<uint8>({1, 2, 3}), FCustomByIdAccessor(0));
+		TestField<ECbFieldType::CustomById>(*this, TEXT("CustomById, MaxId, Empty"), Value, FCbCustomById{MAX_uint64});
+		TestField<ECbFieldType::CustomById>(*this, TEXT("CustomById, MaxId, Empty, View"), Value, FMemoryView(), MakeMemoryView<uint8>({1, 2, 3}), ECbFieldError::None, FCustomByIdAccessor(MAX_uint64));
+		TestFieldError<ECbFieldType::CustomById>(*this, TEXT("CustomById, MaxId, Empty, InvalidId"), Value, ECbFieldError::RangeError, MakeMemoryView<uint8>({1, 2, 3}), FCustomByIdAccessor(0));
 	}
 
 	// Test FCbFieldView(CustomById, MaxId, Value)
 	{
 		const uint8 Value[] = {13, 0xff, 0xff, 0xff, 0xff, 0xff, 0xff, 0xff, 0xff, 0xff, 1, 2, 3, 4};
-		TestFieldNoClone<ECbFieldType::CustomById>(TEXT("CustomById, MaxId, Value"), Value, FCbCustomById{MAX_uint64, MakeMemoryView(Value).Right(4)});
-		TestFieldNoClone<ECbFieldType::CustomById>(TEXT("CustomById, MaxId, Value, View"), Value, MakeMemoryView(Value).Right(4), FMemoryView(), ECbFieldError::None, FCustomByIdAccessor(MAX_uint64));
-		TestFieldError<ECbFieldType::CustomById>(TEXT("CustomById, MaxId, Value, InvalidId"), Value, ECbFieldError::RangeError, MakeMemoryView<uint8>({1, 2, 3}), FCustomByIdAccessor(0));
+		TestFieldNoClone<ECbFieldType::CustomById>(*this, TEXT("CustomById, MaxId, Value"), Value, FCbCustomById{MAX_uint64, MakeMemoryView(Value).Right(4)});
+		TestFieldNoClone<ECbFieldType::CustomById>(*this, TEXT("CustomById, MaxId, Value, View"), Value, MakeMemoryView(Value).Right(4), FMemoryView(), ECbFieldError::None, FCustomByIdAccessor(MAX_uint64));
+		TestFieldError<ECbFieldType::CustomById>(*this, TEXT("CustomById, MaxId, Value, InvalidId"), Value, ECbFieldError::RangeError, MakeMemoryView<uint8>({1, 2, 3}), FCustomByIdAccessor(0));
 	}
 
 	// Test FCbFieldView(None) as CustomById
 	{
 		FCbFieldView DefaultField;
-		TestFieldError<ECbFieldType::CustomById>(TEXT("CustomById, None"), DefaultField, ECbFieldError::TypeError, FCbCustomById{4, MakeMemoryView<uint8>({1, 2, 3})});
-		TestFieldError<ECbFieldType::CustomById>(TEXT("CustomById, None, View"), DefaultField, ECbFieldError::TypeError, MakeMemoryView<uint8>({1, 2, 3}), FCustomByIdAccessor(0));
+		TestFieldError<ECbFieldType::CustomById>(*this, TEXT("CustomById, None"), DefaultField, ECbFieldError::TypeError, FCbCustomById{4, MakeMemoryView<uint8>({1, 2, 3})});
+		TestFieldError<ECbFieldType::CustomById>(*this, TEXT("CustomById, None, View"), DefaultField, ECbFieldError::TypeError, MakeMemoryView<uint8>({1, 2, 3}), FCustomByIdAccessor(0));
 		const uint8 DefaultValue[] = {1, 2, 3};
-		TestEqual(TEXT("FCbFieldView()::AsCustom(Id)"), DefaultField.AsCustom(0, MakeMemoryView(DefaultValue)), MakeMemoryView(DefaultValue));
+		TEST_EQUAL(TEXT("FCbFieldView()::AsCustom(Id)"), DefaultField.AsCustom(0, MakeMemoryView(DefaultValue)), MakeMemoryView(DefaultValue));
 	}
 }
 
-TEST_CASE("Core::Serialization::FCbField::Custom By Name", "[Core][Serialization][Smoke]")
+TEST_CASE_METHOD(FAutomationTestFixture, "Core::Serialization::FCbField::Custom By Name", "[Core][Serialization][Smoke]")
 {
 	struct FCustomByNameAccessor
 	{
@@ -1212,31 +1212,31 @@ TEST_CASE("Core::Serialization::FCbField::Custom By Name", "[Core][Serialization
 	// Test FCbFieldView(CustomByName, ABC, Empty)
 	{
 		const uint8 Value[] = {4, 3, 'A', 'B', 'C'};
-		TestField<ECbFieldType::CustomByName>(TEXT("CustomByName, MinId, Empty"), Value, FCbCustomByName{"ABC"_U8SV});
-		TestField<ECbFieldType::CustomByName>(TEXT("CustomByName, MinId, Empty, View"), Value, FMemoryView(), MakeMemoryView<uint8>({1, 2, 3}), ECbFieldError::None, FCustomByNameAccessor("ABC"_U8SV));
-		TestFieldError<ECbFieldType::CustomByName>(TEXT("CustomByName, MinId, Empty, InvalidCase"), Value, ECbFieldError::RangeError, MakeMemoryView<uint8>({1, 2, 3}), FCustomByNameAccessor("abc"_U8SV));
+		TestField<ECbFieldType::CustomByName>(*this, TEXT("CustomByName, MinId, Empty"), Value, FCbCustomByName{"ABC"_U8SV});
+		TestField<ECbFieldType::CustomByName>(*this, TEXT("CustomByName, MinId, Empty, View"), Value, FMemoryView(), MakeMemoryView<uint8>({1, 2, 3}), ECbFieldError::None, FCustomByNameAccessor("ABC"_U8SV));
+		TestFieldError<ECbFieldType::CustomByName>(*this, TEXT("CustomByName, MinId, Empty, InvalidCase"), Value, ECbFieldError::RangeError, MakeMemoryView<uint8>({1, 2, 3}), FCustomByNameAccessor("abc"_U8SV));
 	}
 
 	// Test FCbFieldView(CustomByName, ABC, Value)
 	{
 		const uint8 Value[] = {8, 3, 'A', 'B', 'C', 1, 2, 3, 4};
-		TestFieldNoClone<ECbFieldType::CustomByName>(TEXT("CustomByName, MinId, Value"), Value, FCbCustomByName{"ABC"_U8SV, MakeMemoryView(Value).Right(4)});
-		TestFieldNoClone<ECbFieldType::CustomByName>(TEXT("CustomByName, MinId, Value, View"), Value, MakeMemoryView(Value).Right(4), FMemoryView(), ECbFieldError::None, FCustomByNameAccessor("ABC"_U8SV));
-		TestFieldError<ECbFieldType::CustomByName>(TEXT("CustomByName, MinId, Value, InvalidCase"), Value, ECbFieldError::RangeError, MakeMemoryView<uint8>({1, 2, 3}), FCustomByNameAccessor("abc"_U8SV));
+		TestFieldNoClone<ECbFieldType::CustomByName>(*this, TEXT("CustomByName, MinId, Value"), Value, FCbCustomByName{"ABC"_U8SV, MakeMemoryView(Value).Right(4)});
+		TestFieldNoClone<ECbFieldType::CustomByName>(*this, TEXT("CustomByName, MinId, Value, View"), Value, MakeMemoryView(Value).Right(4), FMemoryView(), ECbFieldError::None, FCustomByNameAccessor("ABC"_U8SV));
+		TestFieldError<ECbFieldType::CustomByName>(*this, TEXT("CustomByName, MinId, Value, InvalidCase"), Value, ECbFieldError::RangeError, MakeMemoryView<uint8>({1, 2, 3}), FCustomByNameAccessor("abc"_U8SV));
 	}
 
 	// Test FCbFieldView(None) as CustomByName
 	{
 		FCbFieldView DefaultField;
-		TestFieldError<ECbFieldType::CustomByName>(TEXT("CustomByName, None"), DefaultField, ECbFieldError::TypeError, FCbCustomByName{"ABC"_U8SV, MakeMemoryView<uint8>({1, 2, 3})});
-		TestFieldError<ECbFieldType::CustomByName>(TEXT("CustomByName, None, View"), DefaultField, ECbFieldError::TypeError, MakeMemoryView<uint8>({1, 2, 3}), FCustomByNameAccessor("ABC"_U8SV));
+		TestFieldError<ECbFieldType::CustomByName>(*this, TEXT("CustomByName, None"), DefaultField, ECbFieldError::TypeError, FCbCustomByName{"ABC"_U8SV, MakeMemoryView<uint8>({1, 2, 3})});
+		TestFieldError<ECbFieldType::CustomByName>(*this, TEXT("CustomByName, None, View"), DefaultField, ECbFieldError::TypeError, MakeMemoryView<uint8>({1, 2, 3}), FCustomByNameAccessor("ABC"_U8SV));
 		const uint8 DefaultValue[] = {1, 2, 3};
-		TestEqual(TEXT("FCbFieldView()::AsCustom(Name)"), DefaultField.AsCustom("ABC"_U8SV, MakeMemoryView(DefaultValue)), MakeMemoryView(DefaultValue));
+		TEST_EQUAL(TEXT("FCbFieldView()::AsCustom(Name)"), DefaultField.AsCustom("ABC"_U8SV, MakeMemoryView(DefaultValue)), MakeMemoryView(DefaultValue));
 	}
 }
 
 
-TEST_CASE("Core::Serialization::FCbField::Iterate Attachments", "[Core][Serialization][Smoke]")
+TEST_CASE_METHOD(FAutomationTestFixture, "Core::Serialization::FCbField::Iterate Attachments", "[Core][Serialization][Smoke]")
 {
 	const auto MakeTestHash = [](uint32 Index) { return FIoHash::HashBuffer(&Index, sizeof(Index)); };
 
@@ -1395,21 +1395,21 @@ TEST_CASE("Core::Serialization::FCbField::Iterate Attachments", "[Core][Serializ
 		Fields = Writer.Save();
 	}
 
-	TestEqual(TEXT("FCbFieldView::IterateAttachments Validate"), ValidateCompactBinaryRange(Fields.GetOuterBuffer(), ECbValidateMode::All), ECbValidateError::None);
+	TEST_EQUAL(TEXT("FCbFieldView::IterateAttachments Validate"), ValidateCompactBinaryRange(Fields.GetOuterBuffer(), ECbValidateMode::All), ECbValidateError::None);
 
 	uint32 AttachmentIndex = 0;
 
-	Fields.IterateRangeAttachments([&AttachmentIndex, &MakeTestHash](FCbFieldView Field)
+	Fields.IterateRangeAttachments([this, &AttachmentIndex, &MakeTestHash](FCbFieldView Field)
 		{
-			TestTrue(FString::Printf(TEXT("FCbFieldView::IterateAttachments(%u)::IsAttachment"), AttachmentIndex), Field.IsAttachment());
-			TestEqual(FString::Printf(TEXT("FCbFieldView::IterateAttachments(%u)"), AttachmentIndex), Field.AsAttachment(), MakeTestHash(AttachmentIndex));
+			TEST_TRUE(FString::Printf(TEXT("FCbFieldView::IterateAttachments(%u)::IsAttachment"), AttachmentIndex), Field.IsAttachment());
+			TEST_EQUAL(FString::Printf(TEXT("FCbFieldView::IterateAttachments(%u)"), AttachmentIndex), Field.AsAttachment(), MakeTestHash(AttachmentIndex));
 			++AttachmentIndex;
 		});
-	TestEqual(TEXT("FCbFieldView::IterateAttachments"), AttachmentIndex, 28);
+	TEST_EQUAL(TEXT("FCbFieldView::IterateAttachments"), AttachmentIndex, 28);
 }
 
 
-TEST_CASE("Core::Serialization::FCbField::Field Buffer", "[Core][Serialization][Smoke]")
+TEST_CASE_METHOD(FAutomationTestFixture, "Core::Serialization::FCbField::Field Buffer", "[Core][Serialization][Smoke]")
 {
 	static_assert(std::is_constructible<FCbField>::value, "Missing constructor for FCbField");
 	static_assert(std::is_constructible<FCbField, const FCbField&>::value, "Missing constructor for FCbField");
@@ -1433,10 +1433,10 @@ TEST_CASE("Core::Serialization::FCbField::Field Buffer", "[Core][Serialization][
 	// Test FCbField()
 	{
 		FCbField DefaultField;
-		TestFalse(TEXT("FCbField().HasValue()"), DefaultField.HasValue());
-		TestFalse(TEXT("FCbField().IsOwned()"), DefaultField.IsOwned());
+		TEST_FALSE(TEXT("FCbField().HasValue()"), DefaultField.HasValue());
+		TEST_FALSE(TEXT("FCbField().IsOwned()"), DefaultField.IsOwned());
 		DefaultField.MakeOwned();
-		TestTrue(TEXT("FCbField().MakeOwned().IsOwned()"), DefaultField.IsOwned());
+		TEST_TRUE(TEXT("FCbField().MakeOwned().IsOwned()"), DefaultField.IsOwned());
 	}
 
 	// Test Field w/ Type from Shared Buffer
@@ -1458,30 +1458,30 @@ TEST_CASE("Core::Serialization::FCbField::Field Buffer", "[Core][Serialization][
 		//FCbField InvalidOuterBuffer(ImplicitConv<FCbFieldView>(Owned), ViewBuffer);
 		//FCbField InvalidOuterBufferMove(ImplicitConv<FCbFieldView>(Owned), FSharedBuffer(ViewBuffer));
 
-		TestEqual(TEXT("FCbField(ViewBuffer)"), View.AsBinaryView(), ViewBuffer.GetView().Right(3));
-		TestEqual(TEXT("FCbField(ViewBuffer&&)"), ViewMove.AsBinaryView(), View.AsBinaryView());
-		TestEqual(TEXT("FCbField(ViewOuterField)"), ViewOuterField.AsBinaryView(), View.AsBinaryView());
-		TestEqual(TEXT("FCbField(ViewOuterBuffer)"), ViewOuterBuffer.AsBinaryView(), View.AsBinaryView());
-		TestEqual(TEXT("FCbField(OwnedBuffer)"), Owned.AsBinaryView(), OwnedBuffer.GetView().Right(3));
-		TestEqual(TEXT("FCbField(OwnedBuffer&&)"), OwnedMove.AsBinaryView(), Owned.AsBinaryView());
-		TestEqual(TEXT("FCbField(OwnedOuterField)"), OwnedOuterField.AsBinaryView(), Owned.AsBinaryView());
-		TestEqual(TEXT("FCbField(OwnedOuterBuffer)"), OwnedOuterBuffer.AsBinaryView(), Owned.AsBinaryView());
-
-		TestFalse(TEXT("FCbField(ViewBuffer).IsOwned()"), View.IsOwned());
-		TestFalse(TEXT("FCbField(ViewBuffer&&).IsOwned()"), ViewMove.IsOwned());
-		TestFalse(TEXT("FCbField(ViewOuterField).IsOwned()"), ViewOuterField.IsOwned());
-		TestFalse(TEXT("FCbField(ViewOuterBuffer).IsOwned()"), ViewOuterBuffer.IsOwned());
-		TestTrue(TEXT("FCbField(OwnedBuffer).IsOwned()"), Owned.IsOwned());
-		TestTrue(TEXT("FCbField(OwnedBuffer&&).IsOwned()"), OwnedMove.IsOwned());
-		TestTrue(TEXT("FCbField(OwnedOuterField).IsOwned()"), OwnedOuterField.IsOwned());
-		TestTrue(TEXT("FCbField(OwnedOuterBuffer).IsOwned()"), OwnedOuterBuffer.IsOwned());
+		TEST_EQUAL(TEXT("FCbField(ViewBuffer)"), View.AsBinaryView(), ViewBuffer.GetView().Right(3));
+		TEST_EQUAL(TEXT("FCbField(ViewBuffer&&)"), ViewMove.AsBinaryView(), View.AsBinaryView());
+		TEST_EQUAL(TEXT("FCbField(ViewOuterField)"), ViewOuterField.AsBinaryView(), View.AsBinaryView());
+		TEST_EQUAL(TEXT("FCbField(ViewOuterBuffer)"), ViewOuterBuffer.AsBinaryView(), View.AsBinaryView());
+		TEST_EQUAL(TEXT("FCbField(OwnedBuffer)"), Owned.AsBinaryView(), OwnedBuffer.GetView().Right(3));
+		TEST_EQUAL(TEXT("FCbField(OwnedBuffer&&)"), OwnedMove.AsBinaryView(), Owned.AsBinaryView());
+		TEST_EQUAL(TEXT("FCbField(OwnedOuterField)"), OwnedOuterField.AsBinaryView(), Owned.AsBinaryView());
+		TEST_EQUAL(TEXT("FCbField(OwnedOuterBuffer)"), OwnedOuterBuffer.AsBinaryView(), Owned.AsBinaryView());
+
+		TEST_FALSE(TEXT("FCbField(ViewBuffer).IsOwned()"), View.IsOwned());
+		TEST_FALSE(TEXT("FCbField(ViewBuffer&&).IsOwned()"), ViewMove.IsOwned());
+		TEST_FALSE(TEXT("FCbField(ViewOuterField).IsOwned()"), ViewOuterField.IsOwned());
+		TEST_FALSE(TEXT("FCbField(ViewOuterBuffer).IsOwned()"), ViewOuterBuffer.IsOwned());
+		TEST_TRUE(TEXT("FCbField(OwnedBuffer).IsOwned()"), Owned.IsOwned());
+		TEST_TRUE(TEXT("FCbField(OwnedBuffer&&).IsOwned()"), OwnedMove.IsOwned());
+		TEST_TRUE(TEXT("FCbField(OwnedOuterField).IsOwned()"), OwnedOuterField.IsOwned());
+		TEST_TRUE(TEXT("FCbField(OwnedOuterBuffer).IsOwned()"), OwnedOuterBuffer.IsOwned());
 
 		View.MakeOwned();
 		Owned.MakeOwned();
 		TestNotEqual(TEXT("FCbField(View).MakeOwned()"), View.AsBinaryView(), ViewBuffer.GetView().Right(3));
-		TestTrue(TEXT("FCbField(View).MakeOwned().IsOwned()"), View.IsOwned());
-		TestEqual(TEXT("FCbField(Owned).MakeOwned()"), Owned.AsBinaryView(), OwnedBuffer.GetView().Right(3));
-		TestTrue(TEXT("FCbField(Owned).MakeOwned().IsOwned()"), Owned.IsOwned());
+		TEST_TRUE(TEXT("FCbField(View).MakeOwned().IsOwned()"), View.IsOwned());
+		TEST_EQUAL(TEXT("FCbField(Owned).MakeOwned()"), Owned.AsBinaryView(), OwnedBuffer.GetView().Right(3));
+		TEST_TRUE(TEXT("FCbField(Owned).MakeOwned().IsOwned()"), Owned.IsOwned());
 	}
 
 	// Test Field w/ Type
@@ -1495,20 +1495,20 @@ TEST_CASE("Core::Serialization::FCbField::Field Buffer", "[Core][Serialization][
 		FCbField FieldClone = FCbField::Clone(Field);
 		FCbField FieldViewClone = FCbField::Clone(FieldView);
 
-		TestEqual(TEXT("FCbField::MakeView(Void)"), VoidView.AsBinaryView(), MakeMemoryView(Value).Right(3));
+		TEST_EQUAL(TEXT("FCbField::MakeView(Void)"), VoidView.AsBinaryView(), MakeMemoryView(Value).Right(3));
 		TestNotEqual(TEXT("FCbField::Clone(Void)"), VoidClone.AsBinaryView(), MakeMemoryView(Value).Right(3));
-		TestTrue(TEXT("FCbField::Clone(Void)->EqualBytes"), VoidClone.AsBinaryView().EqualBytes(VoidView.AsBinaryView()));
-		TestEqual(TEXT("FCbField::MakeView(Field)"), FieldView.AsBinaryView(), MakeMemoryView(Value).Right(3));
+		TEST_TRUE(TEXT("FCbField::Clone(Void)->EqualBytes"), VoidClone.AsBinaryView().EqualBytes(VoidView.AsBinaryView()));
+		TEST_EQUAL(TEXT("FCbField::MakeView(Field)"), FieldView.AsBinaryView(), MakeMemoryView(Value).Right(3));
 		TestNotEqual(TEXT("FCbField::Clone(Field)"), FieldClone.AsBinaryView(), MakeMemoryView(Value).Right(3));
-		TestTrue(TEXT("FCbField::Clone(Field)->EqualBytes"), FieldClone.AsBinaryView().EqualBytes(VoidView.AsBinaryView()));
+		TEST_TRUE(TEXT("FCbField::Clone(Field)->EqualBytes"), FieldClone.AsBinaryView().EqualBytes(VoidView.AsBinaryView()));
 		TestNotEqual(TEXT("FCbField::Clone(FieldView)"), FieldViewClone.AsBinaryView(), FieldView.AsBinaryView());
-		TestTrue(TEXT("FCbField::Clone(FieldView)->EqualBytes"), FieldViewClone.AsBinaryView().EqualBytes(VoidView.AsBinaryView()));
+		TEST_TRUE(TEXT("FCbField::Clone(FieldView)->EqualBytes"), FieldViewClone.AsBinaryView().EqualBytes(VoidView.AsBinaryView()));
 
-		TestFalse(TEXT("FCbField::MakeView(Void).IsOwned()"), VoidView.IsOwned());
-		TestTrue(TEXT("FCbField::Clone(Void).IsOwned()"), VoidClone.IsOwned());
-		TestFalse(TEXT("FCbField::MakeView(Field).IsOwned()"), FieldView.IsOwned());
-		TestTrue(TEXT("FCbField::Clone(Field).IsOwned()"), FieldClone.IsOwned());
-		TestTrue(TEXT("FCbField::Clone(FieldView).IsOwned()"), FieldViewClone.IsOwned());
+		TEST_FALSE(TEXT("FCbField::MakeView(Void).IsOwned()"), VoidView.IsOwned());
+		TEST_TRUE(TEXT("FCbField::Clone(Void).IsOwned()"), VoidClone.IsOwned());
+		TEST_FALSE(TEXT("FCbField::MakeView(Field).IsOwned()"), FieldView.IsOwned());
+		TEST_TRUE(TEXT("FCbField::Clone(Field).IsOwned()"), FieldClone.IsOwned());
+		TEST_TRUE(TEXT("FCbField::Clone(FieldView).IsOwned()"), FieldViewClone.IsOwned());
 	}
 
 	// Test Field w/o Type
@@ -1520,21 +1520,21 @@ TEST_CASE("Core::Serialization::FCbField::Field Buffer", "[Core][Serialization][
 		FCbField FieldClone = FCbField::Clone(Field);
 		FCbField FieldViewClone = FCbField::Clone(FieldView);
 
-		TestEqual(TEXT("FCbField::MakeView(Field, NoType)"), FieldView.AsBinaryView(), MakeMemoryView(Value).Right(3));
-		TestTrue(TEXT("FCbField::Clone(Field, NoType)"), FieldClone.AsBinaryView().EqualBytes(FieldView.AsBinaryView()));
-		TestTrue(TEXT("FCbField::Clone(FieldView, NoType)"), FieldViewClone.AsBinaryView().EqualBytes(FieldView.AsBinaryView()));
+		TEST_EQUAL(TEXT("FCbField::MakeView(Field, NoType)"), FieldView.AsBinaryView(), MakeMemoryView(Value).Right(3));
+		TEST_TRUE(TEXT("FCbField::Clone(Field, NoType)"), FieldClone.AsBinaryView().EqualBytes(FieldView.AsBinaryView()));
+		TEST_TRUE(TEXT("FCbField::Clone(FieldView, NoType)"), FieldViewClone.AsBinaryView().EqualBytes(FieldView.AsBinaryView()));
 
-		TestFalse(TEXT("FCbField::MakeView(Field, NoType).IsOwned()"), FieldView.IsOwned());
-		TestTrue(TEXT("FCbField::Clone(Field, NoType).IsOwned()"), FieldClone.IsOwned());
-		TestTrue(TEXT("FCbField::Clone(FieldView, NoType).IsOwned()"), FieldViewClone.IsOwned());
+		TEST_FALSE(TEXT("FCbField::MakeView(Field, NoType).IsOwned()"), FieldView.IsOwned());
+		TEST_TRUE(TEXT("FCbField::Clone(Field, NoType).IsOwned()"), FieldClone.IsOwned());
+		TEST_TRUE(TEXT("FCbField::Clone(FieldView, NoType).IsOwned()"), FieldViewClone.IsOwned());
 
 		FieldView.MakeOwned();
-		TestTrue(TEXT("FCbField::MakeView(NoType).MakeOwned()"), FieldView.AsBinaryView().EqualBytes(MakeMemoryView(Value).Right(3)));
-		TestTrue(TEXT("FCbField::MakeView(NoType).MakeOwned().IsOwned()"), FieldView.IsOwned());
+		TEST_TRUE(TEXT("FCbField::MakeView(NoType).MakeOwned()"), FieldView.AsBinaryView().EqualBytes(MakeMemoryView(Value).Right(3)));
+		TEST_TRUE(TEXT("FCbField::MakeView(NoType).MakeOwned().IsOwned()"), FieldView.IsOwned());
 	}
 }
 
-TEST_CASE("Core::Serialization::FCbField::Array Buffer", "[Core][Serialization][Smoke]")
+TEST_CASE_METHOD(FAutomationTestFixture, "Core::Serialization::FCbField::Array Buffer", "[Core][Serialization][Smoke]")
 {
 	static_assert(std::is_constructible<FCbArray>::value, "Missing constructor for FCbArray");
 	static_assert(std::is_constructible<FCbArray, const FCbArray&>::value, "Missing constructor for FCbArray");
@@ -1555,13 +1555,13 @@ TEST_CASE("Core::Serialization::FCbField::Array Buffer", "[Core][Serialization][
 	// Test FCbArray()
 	{
 		FCbArray DefaultArray;
-		TestFalse(TEXT("FCbArray().IsOwned()"), DefaultArray.IsOwned());
+		TEST_FALSE(TEXT("FCbArray().IsOwned()"), DefaultArray.IsOwned());
 		DefaultArray.MakeOwned();
-		TestTrue(TEXT("FCbArray().MakeOwned().IsOwned()"), DefaultArray.IsOwned());
+		TEST_TRUE(TEXT("FCbArray().MakeOwned().IsOwned()"), DefaultArray.IsOwned());
 	}
 }
 
-TEST_CASE("Core::Serialization::FCbField::Object Buffer", "[Core][Serialization][Smoke]")
+TEST_CASE_METHOD(FAutomationTestFixture, "Core::Serialization::FCbField::Object Buffer", "[Core][Serialization][Smoke]")
 {
 	static_assert(std::is_constructible<FCbObject>::value, "Missing constructor for FCbObject");
 	static_assert(std::is_constructible<FCbObject, const FCbObject&&>::value, "Missing constructor for FCbObject");
@@ -1582,13 +1582,13 @@ TEST_CASE("Core::Serialization::FCbField::Object Buffer", "[Core][Serialization]
 	// Test FCbObject()
 	{
 		FCbObject DefaultObject;
-		TestFalse(TEXT("FCbObject().IsOwned()"), DefaultObject.IsOwned());
+		TEST_FALSE(TEXT("FCbObject().IsOwned()"), DefaultObject.IsOwned());
 		DefaultObject.MakeOwned();
-		TestTrue(TEXT("FCbObject().MakeOwned().IsOwned()"), DefaultObject.IsOwned());
+		TEST_TRUE(TEXT("FCbObject().MakeOwned().IsOwned()"), DefaultObject.IsOwned());
 	}
 }
 
-TEST_CASE("Core::Serialization::FCbField::Buffer Iterator", "[Core][Serialization][Smoke]")
+TEST_CASE_METHOD(FAutomationTestFixture, "Core::Serialization::FCbField::Buffer Iterator", "[Core][Serialization][Smoke]")
 {
 	static_assert(std::is_constructible<FCbFieldViewIterator, const FCbFieldIterator&>::value, "Missing constructor for FCbFieldViewIterator");
 	static_assert(std::is_constructible<FCbFieldViewIterator, FCbFieldIterator&&>::value, "Missing constructor for FCbFieldViewIterator");
@@ -1608,8 +1608,8 @@ TEST_CASE("Core::Serialization::FCbField::Buffer Iterator", "[Core][Serializatio
 
 	// Test FCbField[View]Iterator()
 	{
-		TestEqual(TEXT("FCbFieldViewIterator()"), GetCount(FCbFieldViewIterator()), 0);
-		TestEqual(TEXT("FCbFieldIterator()"), GetCount(FCbFieldIterator()), 0);
+		TEST_EQUAL(TEXT("FCbFieldViewIterator()"), GetCount(FCbFieldViewIterator()), 0);
+		TEST_EQUAL(TEXT("FCbFieldIterator()"), GetCount(FCbFieldIterator()), 0);
 	}
 
 	// Test FCbField[View]Iterator(Range)
@@ -1626,89 +1626,89 @@ TEST_CASE("Core::Serialization::FCbField::Buffer Iterator", "[Core][Serializatio
 		const FCbFieldViewIterator FieldViewIt = FCbFieldViewIterator::MakeRange(View);
 		const FCbFieldIterator FieldIt = FCbFieldIterator::MakeRange(View);
 
-		TestEqual(TEXT("FCbFieldViewIterator::GetRangeHash()"), FieldViewIt.GetRangeHash(), FIoHash::HashBuffer(View));
-		TestEqual(TEXT("FCbFieldIterator::GetRangeHash()"), FieldIt.GetRangeHash(), FIoHash::HashBuffer(View));
+		TEST_EQUAL(TEXT("FCbFieldViewIterator::GetRangeHash()"), FieldViewIt.GetRangeHash(), FIoHash::HashBuffer(View));
+		TEST_EQUAL(TEXT("FCbFieldIterator::GetRangeHash()"), FieldIt.GetRangeHash(), FIoHash::HashBuffer(View));
 
 		FMemoryView RangeView;
-		TestTrue(TEXT("FCbFieldViewIterator::TryGetRangeView()"), FieldViewIt.TryGetRangeView(RangeView) && RangeView == MakeMemoryView(Value));
-		TestTrue(TEXT("FCbFieldIterator::TryGetRangeView()"), FieldIt.TryGetRangeView(RangeView) && RangeView == MakeMemoryView(Value));
+		TEST_TRUE(TEXT("FCbFieldViewIterator::TryGetRangeView()"), FieldViewIt.TryGetRangeView(RangeView) && RangeView == MakeMemoryView(Value));
+		TEST_TRUE(TEXT("FCbFieldIterator::TryGetRangeView()"), FieldIt.TryGetRangeView(RangeView) && RangeView == MakeMemoryView(Value));
 
-		TestEqual(TEXT("FCbFieldIterator::CloneRange(EmptyViewIt)"), GetCount(FCbFieldIterator::CloneRange(FCbFieldViewIterator())), 0);
-		TestEqual(TEXT("FCbFieldIterator::CloneRange(EmptyIt)"), GetCount(FCbFieldIterator::CloneRange(FCbFieldIterator())), 0);
+		TEST_EQUAL(TEXT("FCbFieldIterator::CloneRange(EmptyViewIt)"), GetCount(FCbFieldIterator::CloneRange(FCbFieldViewIterator())), 0);
+		TEST_EQUAL(TEXT("FCbFieldIterator::CloneRange(EmptyIt)"), GetCount(FCbFieldIterator::CloneRange(FCbFieldIterator())), 0);
 		const FCbFieldIterator FieldViewItClone = FCbFieldIterator::CloneRange(FieldViewIt);
 		const FCbFieldIterator FieldItClone = FCbFieldIterator::CloneRange(FieldIt);
-		TestEqual(TEXT("FCbFieldIterator::CloneRange(FieldViewIt)"), GetCount(FieldViewItClone), 4);
-		TestEqual(TEXT("FCbFieldIterator::CloneRange(FieldIt)"), GetCount(FieldItClone), 4);
+		TEST_EQUAL(TEXT("FCbFieldIterator::CloneRange(FieldViewIt)"), GetCount(FieldViewItClone), 4);
+		TEST_EQUAL(TEXT("FCbFieldIterator::CloneRange(FieldIt)"), GetCount(FieldItClone), 4);
 		TestNotEqual(TEXT("FCbFieldIterator::CloneRange(FieldViewIt).Equals()"), FieldViewItClone, FieldIt);
 		TestNotEqual(TEXT("FCbFieldIterator::CloneRange(FieldIt).Equals()"), FieldItClone, FieldIt);
 
-		TestEqual(TEXT("FCbFieldViewIterator::MakeRange(EmptyView)"), GetCount(FCbFieldViewIterator::MakeRange(EmptyView)), 0);
-		TestEqual(TEXT("FCbFieldIterator::MakeRange(BufferNullL)"), GetCount(FCbFieldIterator::MakeRange(NullBuffer)), 0);
-		TestEqual(TEXT("FCbFieldIterator::MakeRange(BufferNullR)"), GetCount(FCbFieldIterator::MakeRange(FSharedBuffer(NullBuffer))), 0);
+		TEST_EQUAL(TEXT("FCbFieldViewIterator::MakeRange(EmptyView)"), GetCount(FCbFieldViewIterator::MakeRange(EmptyView)), 0);
+		TEST_EQUAL(TEXT("FCbFieldIterator::MakeRange(BufferNullL)"), GetCount(FCbFieldIterator::MakeRange(NullBuffer)), 0);
+		TEST_EQUAL(TEXT("FCbFieldIterator::MakeRange(BufferNullR)"), GetCount(FCbFieldIterator::MakeRange(FSharedBuffer(NullBuffer))), 0);
 
-		TestEqual(TEXT("FCbFieldViewIterator::MakeRange(BufferView)"), GetCount(FCbFieldViewIterator::MakeRange(MakeMemoryView(Value))), 4);
-		TestEqual(TEXT("FCbFieldIterator::MakeRange(BufferCloneL)"), GetCount(FCbFieldIterator::MakeRange(Clone)), 4);
-		TestEqual(TEXT("FCbFieldIterator::MakeRange(BufferCloneR)"), GetCount(FCbFieldIterator::MakeRange(FSharedBuffer(Clone))), 4);
+		TEST_EQUAL(TEXT("FCbFieldViewIterator::MakeRange(BufferView)"), GetCount(FCbFieldViewIterator::MakeRange(MakeMemoryView(Value))), 4);
+		TEST_EQUAL(TEXT("FCbFieldIterator::MakeRange(BufferCloneL)"), GetCount(FCbFieldIterator::MakeRange(Clone)), 4);
+		TEST_EQUAL(TEXT("FCbFieldIterator::MakeRange(BufferCloneR)"), GetCount(FCbFieldIterator::MakeRange(FSharedBuffer(Clone))), 4);
 
-		TestEqual(TEXT("FCbFieldIterator::MakeRangeView(FieldIt, BufferNullL)"), GetCount(FCbFieldIterator::MakeRangeView(FCbFieldViewIterator::MakeRange(View), NullBuffer)), 4);
-		TestEqual(TEXT("FCbFieldIterator::MakeRangeView(FieldIt, BufferNullR)"), GetCount(FCbFieldIterator::MakeRangeView(FCbFieldViewIterator::MakeRange(View), FSharedBuffer(NullBuffer))), 4);
-		TestEqual(TEXT("FCbFieldIterator::MakeRangeView(FieldIt, BufferViewL)"), GetCount(FCbFieldIterator::MakeRangeView(FCbFieldViewIterator::MakeRange(View), View)), 4);
-		TestEqual(TEXT("FCbFieldIterator::MakeRangeView(FieldIt, BufferViewR)"), GetCount(FCbFieldIterator::MakeRangeView(FCbFieldViewIterator::MakeRange(View), FSharedBuffer(View))), 4);
-		TestEqual(TEXT("FCbFieldIterator::MakeRangeView(FieldIt, BufferCloneL)"), GetCount(FCbFieldIterator::MakeRangeView(FCbFieldViewIterator::MakeRange(Clone), Clone)), 4);
-		TestEqual(TEXT("FCbFieldIterator::MakeRangeView(FieldIt, BufferCloneR)"), GetCount(FCbFieldIterator::MakeRangeView(FCbFieldViewIterator::MakeRange(Clone), FSharedBuffer(Clone))), 4);
+		TEST_EQUAL(TEXT("FCbFieldIterator::MakeRangeView(FieldIt, BufferNullL)"), GetCount(FCbFieldIterator::MakeRangeView(FCbFieldViewIterator::MakeRange(View), NullBuffer)), 4);
+		TEST_EQUAL(TEXT("FCbFieldIterator::MakeRangeView(FieldIt, BufferNullR)"), GetCount(FCbFieldIterator::MakeRangeView(FCbFieldViewIterator::MakeRange(View), FSharedBuffer(NullBuffer))), 4);
+		TEST_EQUAL(TEXT("FCbFieldIterator::MakeRangeView(FieldIt, BufferViewL)"), GetCount(FCbFieldIterator::MakeRangeView(FCbFieldViewIterator::MakeRange(View), View)), 4);
+		TEST_EQUAL(TEXT("FCbFieldIterator::MakeRangeView(FieldIt, BufferViewR)"), GetCount(FCbFieldIterator::MakeRangeView(FCbFieldViewIterator::MakeRange(View), FSharedBuffer(View))), 4);
+		TEST_EQUAL(TEXT("FCbFieldIterator::MakeRangeView(FieldIt, BufferCloneL)"), GetCount(FCbFieldIterator::MakeRangeView(FCbFieldViewIterator::MakeRange(Clone), Clone)), 4);
+		TEST_EQUAL(TEXT("FCbFieldIterator::MakeRangeView(FieldIt, BufferCloneR)"), GetCount(FCbFieldIterator::MakeRangeView(FCbFieldViewIterator::MakeRange(Clone), FSharedBuffer(Clone))), 4);
 
-		TestEqual(TEXT("FCbFieldViewIterator(FieldItL)"), GetCount(FCbFieldViewIterator(FieldIt)), 4);
-		TestEqual(TEXT("FCbFieldViewIterator(FieldItR)"), GetCount(FCbFieldViewIterator(FCbFieldIterator(FieldIt))), 4);
+		TEST_EQUAL(TEXT("FCbFieldViewIterator(FieldItL)"), GetCount(FCbFieldViewIterator(FieldIt)), 4);
+		TEST_EQUAL(TEXT("FCbFieldViewIterator(FieldItR)"), GetCount(FCbFieldViewIterator(FCbFieldIterator(FieldIt))), 4);
 
 		// Uniform
 		const uint8 UniformValue[] = { 0, 1, 2, 3 };
 		const FCbFieldViewIterator UniformFieldViewIt = FCbFieldViewIterator::MakeRange(MakeMemoryView(UniformValue), ECbFieldType::IntegerPositive);
-		TestEqual(TEXT("FCbFieldViewIterator::MakeRange(Uniform).GetRangeHash()"), UniformFieldViewIt.GetRangeHash(), FieldViewIt.GetRangeHash());
-		TestFalse(TEXT("FCbFieldViewIterator::MakeRange(Uniform).TryGetRangeView()"), UniformFieldViewIt.TryGetRangeView(RangeView));
+		TEST_EQUAL(TEXT("FCbFieldViewIterator::MakeRange(Uniform).GetRangeHash()"), UniformFieldViewIt.GetRangeHash(), FieldViewIt.GetRangeHash());
+		TEST_FALSE(TEXT("FCbFieldViewIterator::MakeRange(Uniform).TryGetRangeView()"), UniformFieldViewIt.TryGetRangeView(RangeView));
 		const FSharedBuffer UniformView = FSharedBuffer::MakeView(MakeMemoryView(UniformValue));
 		const FCbFieldIterator UniformFieldIt = FCbFieldIterator::MakeRange(UniformView, ECbFieldType::IntegerPositive);
-		TestEqual(TEXT("FCbFieldIterator::MakeRange(Uniform).GetRangeHash()"), UniformFieldIt.GetRangeHash(), FieldViewIt.GetRangeHash());
-		TestFalse(TEXT("FCbFieldIterator::MakeRange(Uniform).TryGetRangeView()"), UniformFieldIt.TryGetRangeView(RangeView));
+		TEST_EQUAL(TEXT("FCbFieldIterator::MakeRange(Uniform).GetRangeHash()"), UniformFieldIt.GetRangeHash(), FieldViewIt.GetRangeHash());
+		TEST_FALSE(TEXT("FCbFieldIterator::MakeRange(Uniform).TryGetRangeView()"), UniformFieldIt.TryGetRangeView(RangeView));
 
 		// Equals
-		TestTrue(TEXT("FCbFieldViewIterator::Equals(Self)"), FieldViewIt.Equals(AsConst(FieldViewIt)));
-		TestTrue(TEXT("FCbFieldViewIterator::Equals(OtherType)"), FieldViewIt.Equals(AsConst(FieldIt)));
-		TestTrue(TEXT("FCbFieldIterator::Equals(Self)"), FieldIt.Equals(AsConst(FieldIt)));
-		TestTrue(TEXT("FCbFieldIterator::Equals(OtherType)"), FieldIt.Equals(AsConst(FieldViewIt)));
-		TestFalse(TEXT("FCbFieldViewIterator::Equals(OtherRange)"), FieldViewIt.Equals(FieldViewItClone));
-		TestFalse(TEXT("FCbFieldIterator::Equals(OtherRange)"), FieldIt.Equals(FieldItClone));
-		TestTrue(TEXT("FCbFieldViewIterator::Equals(Uniform, Self)"), UniformFieldViewIt.Equals(AsConst(UniformFieldViewIt)));
-		TestTrue(TEXT("FCbFieldViewIterator::Equals(Uniform, OtherType)"), UniformFieldViewIt.Equals(UniformFieldIt));
-		TestTrue(TEXT("FCbFieldIterator::Equals(Uniform, Self)"), UniformFieldIt.Equals(AsConst(UniformFieldIt)));
-		TestTrue(TEXT("FCbFieldIterator::Equals(Uniform, OtherType)"), UniformFieldIt.Equals(UniformFieldViewIt));
-		TestFalse(TEXT("FCbFieldViewIterator::Equals(SameValue, DifferentEnd)"),
+		TEST_TRUE(TEXT("FCbFieldViewIterator::Equals(Self)"), FieldViewIt.Equals(AsConst(FieldViewIt)));
+		TEST_TRUE(TEXT("FCbFieldViewIterator::Equals(OtherType)"), FieldViewIt.Equals(AsConst(FieldIt)));
+		TEST_TRUE(TEXT("FCbFieldIterator::Equals(Self)"), FieldIt.Equals(AsConst(FieldIt)));
+		TEST_TRUE(TEXT("FCbFieldIterator::Equals(OtherType)"), FieldIt.Equals(AsConst(FieldViewIt)));
+		TEST_FALSE(TEXT("FCbFieldViewIterator::Equals(OtherRange)"), FieldViewIt.Equals(FieldViewItClone));
+		TEST_FALSE(TEXT("FCbFieldIterator::Equals(OtherRange)"), FieldIt.Equals(FieldItClone));
+		TEST_TRUE(TEXT("FCbFieldViewIterator::Equals(Uniform, Self)"), UniformFieldViewIt.Equals(AsConst(UniformFieldViewIt)));
+		TEST_TRUE(TEXT("FCbFieldViewIterator::Equals(Uniform, OtherType)"), UniformFieldViewIt.Equals(UniformFieldIt));
+		TEST_TRUE(TEXT("FCbFieldIterator::Equals(Uniform, Self)"), UniformFieldIt.Equals(AsConst(UniformFieldIt)));
+		TEST_TRUE(TEXT("FCbFieldIterator::Equals(Uniform, OtherType)"), UniformFieldIt.Equals(UniformFieldViewIt));
+		TEST_FALSE(TEXT("FCbFieldViewIterator::Equals(SameValue, DifferentEnd)"),
 			FCbFieldViewIterator::MakeRange(MakeMemoryView(UniformValue), ECbFieldType::IntegerPositive)
 				.Equals(FCbFieldViewIterator::MakeRange(MakeMemoryView(UniformValue).LeftChop(1), ECbFieldType::IntegerPositive)));
-		TestFalse(TEXT("FCbFieldViewIterator::Equals(DifferentValue, SameEnd)"),
+		TEST_FALSE(TEXT("FCbFieldViewIterator::Equals(DifferentValue, SameEnd)"),
 			FCbFieldViewIterator::MakeRange(MakeMemoryView(UniformValue), ECbFieldType::IntegerPositive)
 				.Equals(FCbFieldViewIterator::MakeRange(MakeMemoryView(UniformValue).RightChop(1), ECbFieldType::IntegerPositive)));
 
 		// CopyRangeTo
 		uint8 CopyBytes[sizeof(Value)];
 		FieldViewIt.CopyRangeTo(MakeMemoryView(CopyBytes));
-		TestTrue(TEXT("FCbFieldViewIterator::MakeRange().CopyRangeTo()"), MakeMemoryView(CopyBytes).EqualBytes(MakeMemoryView(Value)));
+		TEST_TRUE(TEXT("FCbFieldViewIterator::MakeRange().CopyRangeTo()"), MakeMemoryView(CopyBytes).EqualBytes(MakeMemoryView(Value)));
 		FieldIt.CopyRangeTo(MakeMemoryView(CopyBytes));
-		TestTrue(TEXT("FCbFieldIterator::MakeRange().CopyRangeTo()"), MakeMemoryView(CopyBytes).EqualBytes(MakeMemoryView(Value)));
+		TEST_TRUE(TEXT("FCbFieldIterator::MakeRange().CopyRangeTo()"), MakeMemoryView(CopyBytes).EqualBytes(MakeMemoryView(Value)));
 		UniformFieldViewIt.CopyRangeTo(MakeMemoryView(CopyBytes));
-		TestTrue(TEXT("FCbFieldViewIterator::MakeRange(Uniform).CopyRangeTo()"), MakeMemoryView(CopyBytes).EqualBytes(MakeMemoryView(Value)));
+		TEST_TRUE(TEXT("FCbFieldViewIterator::MakeRange(Uniform).CopyRangeTo()"), MakeMemoryView(CopyBytes).EqualBytes(MakeMemoryView(Value)));
 
 		// MakeRangeOwned
 		FCbFieldIterator OwnedFromView = UniformFieldIt;
 		OwnedFromView.MakeRangeOwned();
-		TestTrue(TEXT("FCbFieldIterator::MakeRangeOwned(View)"), OwnedFromView.TryGetRangeView(RangeView) && RangeView.EqualBytes(MakeMemoryView(Value)));
+		TEST_TRUE(TEXT("FCbFieldIterator::MakeRangeOwned(View)"), OwnedFromView.TryGetRangeView(RangeView) && RangeView.EqualBytes(MakeMemoryView(Value)));
 		FCbFieldIterator OwnedFromOwned = OwnedFromView;
 		OwnedFromOwned.MakeRangeOwned();
-		TestEqual(TEXT("FCbFieldIterator::MakeRangeOwned(Owned)"), OwnedFromOwned, OwnedFromView);
+		TEST_EQUAL(TEXT("FCbFieldIterator::MakeRangeOwned(Owned)"), OwnedFromOwned, OwnedFromView);
 
 		// These lines are expected to assert when uncommented.
 		//const FSharedBuffer ShortView = FSharedBuffer::MakeView(MakeMemoryView(Value).LeftChop(2));
-		//TestEqual(TEXT("FCbFieldIterator::MakeRangeView(FieldIt, InvalidBufferL)"), GetCount(FCbFieldIterator::MakeRangeView(FCbFieldViewIterator::MakeRange(*View), ShortView)), 4);
-		//TestEqual(TEXT("FCbFieldIterator::MakeRangeView(FieldIt, InvalidBufferR)"), GetCount(FCbFieldIterator::MakeRangeView(FCbFieldViewIterator::MakeRange(*View), FSharedBuffer(ShortView))), 4);
+		//TEST_EQUAL(TEXT("FCbFieldIterator::MakeRangeView(FieldIt, InvalidBufferL)"), GetCount(FCbFieldIterator::MakeRangeView(FCbFieldViewIterator::MakeRange(*View), ShortView)), 4);
+		//TEST_EQUAL(TEXT("FCbFieldIterator::MakeRangeView(FieldIt, InvalidBufferR)"), GetCount(FCbFieldIterator::MakeRangeView(FCbFieldViewIterator::MakeRange(*View), FSharedBuffer(ShortView))), 4);
 	}
 
 	// Test FCbField[View]Iterator(Scalar)
@@ -1722,14 +1722,14 @@ TEST_CASE("Core::Serialization::FCbField::Buffer Iterator", "[Core][Serializatio
 		const FCbFieldView FieldView(Value);
 		const FCbField Field(View);
 
-		TestEqual(TEXT("FCbFieldViewIterator::MakeSingle(FieldViewL)"), GetCount(FCbFieldViewIterator::MakeSingle(FieldView)), 1);
-		TestEqual(TEXT("FCbFieldViewIterator::MakeSingle(FieldViewR)"), GetCount(FCbFieldViewIterator::MakeSingle(FCbFieldView(FieldView))), 1);
-		TestEqual(TEXT("FCbFieldIterator::MakeSingle(FieldL)"), GetCount(FCbFieldIterator::MakeSingle(Field)), 1);
-		TestEqual(TEXT("FCbFieldIterator::MakeSingle(FieldR)"), GetCount(FCbFieldIterator::MakeSingle(FCbField(Field))), 1);
+		TEST_EQUAL(TEXT("FCbFieldViewIterator::MakeSingle(FieldViewL)"), GetCount(FCbFieldViewIterator::MakeSingle(FieldView)), 1);
+		TEST_EQUAL(TEXT("FCbFieldViewIterator::MakeSingle(FieldViewR)"), GetCount(FCbFieldViewIterator::MakeSingle(FCbFieldView(FieldView))), 1);
+		TEST_EQUAL(TEXT("FCbFieldIterator::MakeSingle(FieldL)"), GetCount(FCbFieldIterator::MakeSingle(Field)), 1);
+		TEST_EQUAL(TEXT("FCbFieldIterator::MakeSingle(FieldR)"), GetCount(FCbFieldIterator::MakeSingle(FCbField(Field))), 1);
 	}
 }
 
-TEST_CASE("Core::Serialization::FCbField::Parse Test", "[Core][Serialization][Smoke]")
+TEST_CASE_METHOD(FAutomationTestFixture, "Core::Serialization::FCbField::Parse Test", "[Core][Serialization][Smoke]")
 {
 	// Test the optimal object parsing loop because it is expected to be required for high performance.
 	// Under ideal conditions, when the fields are in the expected order and there are no extra fields,
@@ -1775,13 +1775,13 @@ TEST_CASE("Core::Serialization::FCbField::Parse Test", "[Core][Serialization][Sm
 	};
 
 	constexpr uint8 T = uint8(ECbFieldType::IntegerPositive | ECbFieldType::HasFieldName);
-	TestTrue(TEXT("FCbObjectView Parse(None)"), TestParseObject({0}, 0, 0, 0, 0));
-	TestTrue(TEXT("FCbObjectView Parse(ABCD)"), TestParseObject({16, T, 1, 'A', 1, T, 1, 'B', 2, T, 1, 'C', 3, T, 1, 'D', 4}, 1, 2, 3, 4));
-	TestTrue(TEXT("FCbObjectView Parse(BCDA)"), TestParseObject({16, T, 1, 'B', 2, T, 1, 'C', 3, T, 1, 'D', 4, T, 1, 'A', 1}, 1, 2, 3, 4));
-	TestTrue(TEXT("FCbObjectView Parse(BCD)"), TestParseObject({12, T, 1, 'B', 2, T, 1, 'C', 3, T, 1, 'D', 4}, 0, 2, 3, 4));
-	TestTrue(TEXT("FCbObjectView Parse(BC)"), TestParseObject({8, T, 1, 'B', 2, T, 1, 'C', 3}, 0, 2, 3, 0));
-	TestTrue(TEXT("FCbObjectView Parse(ABCDE)"), TestParseObject({20, T, 1, 'A', 1, T, 1, 'B', 2, T, 1, 'C', 3, T, 1, 'D', 4, T, 1, 'E', 5}, 1, 2, 3, 4));
-	TestTrue(TEXT("FCbObjectView Parse(EABCD)"), TestParseObject({20, T, 1, 'E', 5, T, 1, 'A', 1, T, 1, 'B', 2, T, 1, 'C', 3, T, 1, 'D', 4}, 1, 2, 3, 4));
-	TestTrue(TEXT("FCbObjectView Parse(DCBA)"), TestParseObject({16, T, 1, 'D', 4, T, 1, 'C', 3, T, 1, 'B', 2, T, 1, 'A', 1}, 1, 2, 3, 4));
+	TEST_TRUE(TEXT("FCbObjectView Parse(None)"), TestParseObject({0}, 0, 0, 0, 0));
+	TEST_TRUE(TEXT("FCbObjectView Parse(ABCD)"), TestParseObject({16, T, 1, 'A', 1, T, 1, 'B', 2, T, 1, 'C', 3, T, 1, 'D', 4}, 1, 2, 3, 4));
+	TEST_TRUE(TEXT("FCbObjectView Parse(BCDA)"), TestParseObject({16, T, 1, 'B', 2, T, 1, 'C', 3, T, 1, 'D', 4, T, 1, 'A', 1}, 1, 2, 3, 4));
+	TEST_TRUE(TEXT("FCbObjectView Parse(BCD)"), TestParseObject({12, T, 1, 'B', 2, T, 1, 'C', 3, T, 1, 'D', 4}, 0, 2, 3, 4));
+	TEST_TRUE(TEXT("FCbObjectView Parse(BC)"), TestParseObject({8, T, 1, 'B', 2, T, 1, 'C', 3}, 0, 2, 3, 0));
+	TEST_TRUE(TEXT("FCbObjectView Parse(ABCDE)"), TestParseObject({20, T, 1, 'A', 1, T, 1, 'B', 2, T, 1, 'C', 3, T, 1, 'D', 4, T, 1, 'E', 5}, 1, 2, 3, 4));
+	TEST_TRUE(TEXT("FCbObjectView Parse(EABCD)"), TestParseObject({20, T, 1, 'E', 5, T, 1, 'A', 1, T, 1, 'B', 2, T, 1, 'C', 3, T, 1, 'D', 4}, 1, 2, 3, 4));
+	TEST_TRUE(TEXT("FCbObjectView Parse(DCBA)"), TestParseObject({16, T, 1, 'D', 4, T, 1, 'C', 3, T, 1, 'B', 2, T, 1, 'A', 1}, 1, 2, 3, 4));
 
 }
