Add new Unit Tests

diff --git a/Engine/Source/Runtime/Core/Tests/Serialization/CompactBinaryValidationTest.cpp b/Engine/Source/Runtime/Core/Tests/Serialization/CompactBinaryValidationTest.cpp
--- a/Engine/Source/Runtime/Core/Tests/Serialization/CompactBinaryValidationTest.cpp
+++ b/Engine/Source/Runtime/Core/Tests/Serialization/CompactBinaryValidationTest.cpp
@@ -6,9 +6,9 @@
 #include "Serialization/CompactBinary.h"
 #include "Serialization/CompactBinaryPackage.h"
 #include "Serialization/CompactBinaryWriter.h"
-#include "TestHarness.h"
+#include "TestFixtures/CoreTestFixture.h"
 
-TEST_CASE("Core::Serialization::FCbField::Validate Compact Binary", "[Core][Serialization][Smoke]")
+TEST_CASE_METHOD(FAutomationTestFixture, "Core::Serialization::FCbField::Validate Compact Binary", "[Core][Serialization][Smoke]")
 {
 	auto Validate = [](std::initializer_list<uint8> Data, ECbFieldType Type = ECbFieldType::HasFieldType) -> ECbValidateError
 	{
@@ -27,269 +27,269 @@ TEST_CASE("Core::Serialization::FCbField::Validate Compact Binary", "[Core][Seri
 	constexpr uint8 IntWithName = uint8(ECbFieldType::IntegerPositive | ECbFieldType::HasFieldName);
 
 	// Test OutOfBounds
-	TestEqual(TEXT("ValidateCompactBinary(OutOfBounds, Empty)"), Validate({}), ECbValidateError::OutOfBounds);
-	TestEqual(TEXT("ValidateCompactBinary(OutOfBounds, Empty, Mode)"), ValidateMode({}, ECbValidateMode::None), ECbValidateError::None);
-
-	TestEqual(TEXT("ValidateCompactBinary(Valid, Null)"), Validate({NullNoName}), ECbValidateError::None);
-	TestEqual(TEXT("ValidateCompactBinary(Valid, Null, Name)"), Validate({NullWithName, 1, 'N'}), ECbValidateError::None);
-	TestEqual(TEXT("ValidateCompactBinary(OutOfBounds, Null, Name)"), Validate({NullWithName}), ECbValidateError::OutOfBounds);
-	TestEqual(TEXT("ValidateCompactBinary(OutOfBounds, Null, Name)"), Validate({NullWithName, 1}), ECbValidateError::OutOfBounds);
-	TestEqual(TEXT("ValidateCompactBinary(OutOfBounds, Null, Name)"), Validate({NullWithName, 0x80}), ECbValidateError::OutOfBounds);
-	TestEqual(TEXT("ValidateCompactBinary(OutOfBounds, Null, Name)"), Validate({NullWithName, 0x80, 128}), ECbValidateError::OutOfBounds);
-
-	TestEqual(TEXT("ValidateCompactBinary(Valid, Object, Empty)"), Validate({uint8(ECbFieldType::Object), 0}), ECbValidateError::None);
-	TestEqual(TEXT("ValidateCompactBinary(Valid, Object, Empty, NoType)"), Validate({0}, ECbFieldType::Object), ECbValidateError::None);
-	TestEqual(TEXT("ValidateCompactBinary(Valid, Object, Field)"), Validate({uint8(ECbFieldType::Object), 7, NullWithName, 1, 'N', IntWithName, 1, 'I', 0}), ECbValidateError::None);
-	TestEqual(TEXT("ValidateCompactBinary(Valid, Object, Field, NoType)"), Validate({7, NullWithName, 1, 'N', IntWithName, 1, 'I', 0}, ECbFieldType::Object), ECbValidateError::None);
-	TestEqual(TEXT("ValidateCompactBinary(OutOfBounds, Object)"), Validate({uint8(ECbFieldType::Object)}), ECbValidateError::OutOfBounds);
-	TestEqual(TEXT("ValidateCompactBinary(OutOfBounds, Object, NoType)"), Validate({}, ECbFieldType::Object), ECbValidateError::OutOfBounds);
-	TestEqual(TEXT("ValidateCompactBinary(OutOfBounds, Object)"), Validate({uint8(ECbFieldType::Object), 1}), ECbValidateError::OutOfBounds);
-	TestEqual(TEXT("ValidateCompactBinary(OutOfBounds, Object, NoType)"), Validate({1}, ECbFieldType::Object), ECbValidateError::OutOfBounds);
-	TestEqual(TEXT("ValidateCompactBinary(OutOfBounds, Object, OOB Field)"), Validate({uint8(ECbFieldType::Object), 3, AddName(ECbFieldType::Float32), 1, 'N'}), ECbValidateError::OutOfBounds);
-	TestEqual(TEXT("ValidateCompactBinary(OutOfBounds, Object, OOB Field, NoType)"), Validate({3, AddName(ECbFieldType::Float32), 1, 'N'}, ECbFieldType::Object), ECbValidateError::OutOfBounds);
-
-	TestEqual(TEXT("ValidateCompactBinary(Valid, UniformObject, Field)"), Validate({uint8(ECbFieldType::UniformObject), 3, NullWithName, 1, 'N'}), ECbValidateError::None);
-	TestEqual(TEXT("ValidateCompactBinary(Valid, UniformObject, Field, NoType)"), Validate({3, NullWithName, 1, 'N'}, ECbFieldType::UniformObject), ECbValidateError::None);
-	TestEqual(TEXT("ValidateCompactBinary(OutOfBounds, UniformObject)"), Validate({uint8(ECbFieldType::UniformObject)}), ECbValidateError::OutOfBounds);
-	TestEqual(TEXT("ValidateCompactBinary(OutOfBounds, UniformObject, NoType)"), Validate({}, ECbFieldType::UniformObject), ECbValidateError::OutOfBounds);
-	TestEqual(TEXT("ValidateCompactBinary(OutOfBounds, UniformObject)"), Validate({uint8(ECbFieldType::UniformObject), 1}), ECbValidateError::OutOfBounds);
-	TestEqual(TEXT("ValidateCompactBinary(OutOfBounds, UniformObject, NoType)"), Validate({1}, ECbFieldType::UniformObject), ECbValidateError::OutOfBounds);
-	TestEqual(TEXT("ValidateCompactBinary(OutOfBounds, UniformObject, OOB Field)"), Validate({uint8(ECbFieldType::UniformObject), 3, AddName(ECbFieldType::Float32), 1, 'N'}), ECbValidateError::OutOfBounds);
-	TestEqual(TEXT("ValidateCompactBinary(OutOfBounds, UniformObject, OOB Field, NoType)"), Validate({3, AddName(ECbFieldType::Float32), 1, 'N'}, ECbFieldType::UniformObject), ECbValidateError::OutOfBounds);
-
-	TestEqual(TEXT("ValidateCompactBinary(Valid, Array, Empty)"), Validate({uint8(ECbFieldType::Array), 1, 0}), ECbValidateError::None);
-	TestEqual(TEXT("ValidateCompactBinary(Valid, Array, Empty, NoType)"), Validate({1, 0}, ECbFieldType::Array), ECbValidateError::None);
-	TestEqual(TEXT("ValidateCompactBinary(Valid, Array, Field)"), Validate({uint8(ECbFieldType::Array), 4, 2, NullNoName, IntNoName, 0}), ECbValidateError::None);
-	TestEqual(TEXT("ValidateCompactBinary(Valid, Array, Field, NoType)"), Validate({4, 2, NullNoName, IntNoName, 0}, ECbFieldType::Array), ECbValidateError::None);
-	TestEqual(TEXT("ValidateCompactBinary(OutOfBounds, Array)"), Validate({uint8(ECbFieldType::Array)}), ECbValidateError::OutOfBounds);
-	TestEqual(TEXT("ValidateCompactBinary(OutOfBounds, Array, NoType)"), Validate({}, ECbFieldType::Array), ECbValidateError::OutOfBounds);
-	TestEqual(TEXT("ValidateCompactBinary(OutOfBounds, Array)"), Validate({uint8(ECbFieldType::Array), 1}), ECbValidateError::OutOfBounds);
-	TestEqual(TEXT("ValidateCompactBinary(OutOfBounds, Array, NoType)"), Validate({1}, ECbFieldType::Array), ECbValidateError::OutOfBounds);
-	TestEqual(TEXT("ValidateCompactBinary(OutOfBounds, Array, OOB Field)"), Validate({uint8(ECbFieldType::Array), 2, 1, uint8(ECbFieldType::Float32)}), ECbValidateError::OutOfBounds);
-	TestEqual(TEXT("ValidateCompactBinary(OutOfBounds, Array, OOB Field, NoType)"), Validate({2, 1, uint8(ECbFieldType::Float32)}, ECbFieldType::Array), ECbValidateError::OutOfBounds);
-
-	TestEqual(TEXT("ValidateCompactBinary(Valid, UniformArray, Field)"), Validate({uint8(ECbFieldType::UniformArray), 3, 1, IntNoName, 0}), ECbValidateError::None);
-	TestEqual(TEXT("ValidateCompactBinary(Valid, UniformArray, Field, NoType)"), Validate({3, 1, IntNoName, 0}, ECbFieldType::UniformArray), ECbValidateError::None);
-	TestEqual(TEXT("ValidateCompactBinary(OutOfBounds, UniformArray)"), Validate({uint8(ECbFieldType::UniformArray)}), ECbValidateError::OutOfBounds);
-	TestEqual(TEXT("ValidateCompactBinary(OutOfBounds, UniformArray, NoType)"), Validate({}, ECbFieldType::UniformArray), ECbValidateError::OutOfBounds);
-	TestEqual(TEXT("ValidateCompactBinary(OutOfBounds, UniformArray)"), Validate({uint8(ECbFieldType::UniformArray), 1}), ECbValidateError::OutOfBounds);
-	TestEqual(TEXT("ValidateCompactBinary(OutOfBounds, UniformArray, NoType)"), Validate({1}, ECbFieldType::UniformArray), ECbValidateError::OutOfBounds);
-	TestEqual(TEXT("ValidateCompactBinary(OutOfBounds, UniformArray, OOB Field)"), Validate({uint8(ECbFieldType::UniformArray), 2, 1, uint8(ECbFieldType::Float32)}), ECbValidateError::OutOfBounds);
-	TestEqual(TEXT("ValidateCompactBinary(OutOfBounds, UniformArray, OOB Field, NoType)"), Validate({2, 1, uint8(ECbFieldType::Float32)}, ECbFieldType::UniformArray), ECbValidateError::OutOfBounds);
-
-	TestEqual(TEXT("ValidateCompactBinary(Valid, Binary, Empty)"), Validate({uint8(ECbFieldType::Binary), 0}), ECbValidateError::None);
-	TestEqual(TEXT("ValidateCompactBinary(Valid, Binary, Empty, NoType)"), Validate({0}, ECbFieldType::Binary), ECbValidateError::None);
-	TestEqual(TEXT("ValidateCompactBinary(Valid, Binary, Field)"), Validate({uint8(ECbFieldType::Binary), 1, 0}), ECbValidateError::None);
-	TestEqual(TEXT("ValidateCompactBinary(Valid, Binary, Field, NoType)"), Validate({1, 0}, ECbFieldType::Binary), ECbValidateError::None);
-	TestEqual(TEXT("ValidateCompactBinary(OutOfBounds, Binary)"), Validate({uint8(ECbFieldType::Binary)}), ECbValidateError::OutOfBounds);
-	TestEqual(TEXT("ValidateCompactBinary(OutOfBounds, Binary, NoType)"), Validate({}, ECbFieldType::Binary), ECbValidateError::OutOfBounds);
-	TestEqual(TEXT("ValidateCompactBinary(OutOfBounds, Binary)"), Validate({uint8(ECbFieldType::Binary), 1}), ECbValidateError::OutOfBounds);
-	TestEqual(TEXT("ValidateCompactBinary(OutOfBounds, Binary, NoType)"), Validate({1}, ECbFieldType::Binary), ECbValidateError::OutOfBounds);
-
-	TestEqual(TEXT("ValidateCompactBinary(Valid, String, Empty)"), Validate({uint8(ECbFieldType::String), 0}), ECbValidateError::None);
-	TestEqual(TEXT("ValidateCompactBinary(Valid, String, Empty, NoType)"), Validate({0}, ECbFieldType::String), ECbValidateError::None);
-	TestEqual(TEXT("ValidateCompactBinary(Valid, String, Field)"), Validate({uint8(ECbFieldType::String), 1, 'S'}), ECbValidateError::None);
-	TestEqual(TEXT("ValidateCompactBinary(Valid, String, Field, NoType)"), Validate({1, 'S'}, ECbFieldType::String), ECbValidateError::None);
-	TestEqual(TEXT("ValidateCompactBinary(OutOfBounds, String)"), Validate({uint8(ECbFieldType::String)}), ECbValidateError::OutOfBounds);
-	TestEqual(TEXT("ValidateCompactBinary(OutOfBounds, String, NoType)"), Validate({}, ECbFieldType::String), ECbValidateError::OutOfBounds);
-	TestEqual(TEXT("ValidateCompactBinary(OutOfBounds, String)"), Validate({uint8(ECbFieldType::String), 1}), ECbValidateError::OutOfBounds);
-	TestEqual(TEXT("ValidateCompactBinary(OutOfBounds, String, NoType)"), Validate({1}, ECbFieldType::String), ECbValidateError::OutOfBounds);
-
-	TestEqual(TEXT("ValidateCompactBinary(Valid, IntegerPositive, 1-byte)"), Validate({uint8(ECbFieldType::IntegerPositive), 0}), ECbValidateError::None);
-	TestEqual(TEXT("ValidateCompactBinary(Valid, IntegerPositive, 1-byte, NoType)"), Validate({0}, ECbFieldType::IntegerPositive), ECbValidateError::None);
-	TestEqual(TEXT("ValidateCompactBinary(Valid, IntegerPositive, 2-byte)"), Validate({uint8(ECbFieldType::IntegerPositive), 0x80, 0x80}), ECbValidateError::None);
-	TestEqual(TEXT("ValidateCompactBinary(Valid, IntegerPositive, 2-byte, NoType)"), Validate({0x80, 0x80}, ECbFieldType::IntegerPositive), ECbValidateError::None);
-	TestEqual(TEXT("ValidateCompactBinary(OutOfBounds, IntegerPositive, 1-byte)"), Validate({uint8(ECbFieldType::IntegerPositive)}), ECbValidateError::OutOfBounds);
-	TestEqual(TEXT("ValidateCompactBinary(OutOfBounds, IntegerPositive, 1-byte, NoType)"), Validate({}, ECbFieldType::IntegerPositive), ECbValidateError::OutOfBounds);
-	TestEqual(TEXT("ValidateCompactBinary(OutOfBounds, IntegerPositive, 2-byte)"), Validate({uint8(ECbFieldType::IntegerPositive), 0x80}), ECbValidateError::OutOfBounds);
-	TestEqual(TEXT("ValidateCompactBinary(OutOfBounds, IntegerPositive, 2-byte, NoType)"), Validate({0x80}, ECbFieldType::IntegerPositive), ECbValidateError::OutOfBounds);
-	TestEqual(TEXT("ValidateCompactBinary(OutOfBounds, IntegerPositive, 9-byte)"), Validate({uint8(ECbFieldType::IntegerPositive), 0xff, 0, 0, 0, 0, 0, 0, 0}), ECbValidateError::OutOfBounds);
-	TestEqual(TEXT("ValidateCompactBinary(OutOfBounds, IntegerPositive, 9-byte, NoType)"), Validate({0xff, 0, 0, 0, 0, 0, 0, 0}, ECbFieldType::IntegerPositive), ECbValidateError::OutOfBounds);
-
-	TestEqual(TEXT("ValidateCompactBinary(Valid, IntegerNegative, 1-byte)"), Validate({uint8(ECbFieldType::IntegerNegative), 0}), ECbValidateError::None);
-	TestEqual(TEXT("ValidateCompactBinary(Valid, IntegerNegative, 1-byte, NoType)"), Validate({0}, ECbFieldType::IntegerNegative), ECbValidateError::None);
-	TestEqual(TEXT("ValidateCompactBinary(Valid, IntegerNegative, 2-byte)"), Validate({uint8(ECbFieldType::IntegerNegative), 0x80, 0x80}), ECbValidateError::None);
-	TestEqual(TEXT("ValidateCompactBinary(Valid, IntegerNegative, 2-byte, NoType)"), Validate({0x80, 0x80}, ECbFieldType::IntegerNegative), ECbValidateError::None);
-	TestEqual(TEXT("ValidateCompactBinary(OutOfBounds, IntegerNegative, 1-byte)"), Validate({uint8(ECbFieldType::IntegerNegative)}), ECbValidateError::OutOfBounds);
-	TestEqual(TEXT("ValidateCompactBinary(OutOfBounds, IntegerNegative, 1-byte, NoType)"), Validate({}, ECbFieldType::IntegerNegative), ECbValidateError::OutOfBounds);
-	TestEqual(TEXT("ValidateCompactBinary(OutOfBounds, IntegerNegative, 2-byte)"), Validate({uint8(ECbFieldType::IntegerNegative), 0x80}), ECbValidateError::OutOfBounds);
-	TestEqual(TEXT("ValidateCompactBinary(OutOfBounds, IntegerNegative, 2-byte, NoType)"), Validate({0x80}, ECbFieldType::IntegerNegative), ECbValidateError::OutOfBounds);
-	TestEqual(TEXT("ValidateCompactBinary(OutOfBounds, IntegerNegative, 9-byte)"), Validate({uint8(ECbFieldType::IntegerNegative), 0xff, 0, 0, 0, 0, 0, 0, 0}), ECbValidateError::OutOfBounds);
-	TestEqual(TEXT("ValidateCompactBinary(OutOfBounds, IntegerNegative, 9-byte, NoType)"), Validate({0xff, 0, 0, 0, 0, 0, 0, 0}, ECbFieldType::IntegerNegative), ECbValidateError::OutOfBounds);
-
-	TestEqual(TEXT("ValidateCompactBinary(Valid, Float32)"), Validate({uint8(ECbFieldType::Float32), 0, 0, 0, 0}), ECbValidateError::None);
-	TestEqual(TEXT("ValidateCompactBinary(Valid, Float32, NoType)"), Validate({0, 0, 0, 0}, ECbFieldType::Float32), ECbValidateError::None);
-	TestEqual(TEXT("ValidateCompactBinary(OutOfBounds, Float32)"), Validate({uint8(ECbFieldType::Float32), 0, 0, 0}), ECbValidateError::OutOfBounds);
-	TestEqual(TEXT("ValidateCompactBinary(OutOfBounds, Float32, NoType)"), Validate({0, 0, 0}, ECbFieldType::Float32), ECbValidateError::OutOfBounds);
-
-	TestEqual(TEXT("ValidateCompactBinary(Valid, Float64)"), Validate({uint8(ECbFieldType::Float64), 0x3f, 0xff, 0xff, 0xff, 0xf0, 0x00, 0x00, 0x00}), ECbValidateError::None);
-	TestEqual(TEXT("ValidateCompactBinary(Valid, Float64, NoType)"), Validate({0x3f, 0xff, 0xff, 0xff, 0xf0, 0x00, 0x00, 0x00}, ECbFieldType::Float64), ECbValidateError::None);
-	TestEqual(TEXT("ValidateCompactBinary(OutOfBounds, Float64)"), Validate({uint8(ECbFieldType::Float64), 0, 0, 0, 0, 0, 0, 0}), ECbValidateError::OutOfBounds);
-	TestEqual(TEXT("ValidateCompactBinary(OutOfBounds, Float64, NoType)"), Validate({0, 0, 0, 0, 0, 0, 0}, ECbFieldType::Float64), ECbValidateError::OutOfBounds);
-
-	TestEqual(TEXT("ValidateCompactBinary(Valid, BoolFalse)"), Validate({uint8(ECbFieldType::BoolFalse)}), ECbValidateError::None);
-	TestEqual(TEXT("ValidateCompactBinary(Valid, BoolTrue)"), Validate({uint8(ECbFieldType::BoolTrue)}), ECbValidateError::None);
-
-	TestEqual(TEXT("ValidateCompactBinary(Valid, ObjectAttachment)"), Validate({uint8(ECbFieldType::ObjectAttachment), 0, 0, 0, 0, 0, 0, 0, 0, 0, 0, 0, 0, 0, 0, 0, 0, 0, 0, 0, 0}), ECbValidateError::None);
-	TestEqual(TEXT("ValidateCompactBinary(Valid, ObjectAttachment, NoType)"), Validate({0, 0, 0, 0, 0, 0, 0, 0, 0, 0, 0, 0, 0, 0, 0, 0, 0, 0, 0, 0}, ECbFieldType::ObjectAttachment), ECbValidateError::None);
-	TestEqual(TEXT("ValidateCompactBinary(OutOfBounds, ObjectAttachment)"), Validate({uint8(ECbFieldType::ObjectAttachment), 0, 0, 0, 0, 0, 0, 0, 0, 0, 0, 0, 0, 0, 0, 0, 0, 0, 0, 0}), ECbValidateError::OutOfBounds);
-	TestEqual(TEXT("ValidateCompactBinary(OutOfBounds, ObjectAttachment, NoType)"), Validate({0, 0, 0, 0, 0, 0, 0, 0, 0, 0, 0, 0, 0, 0, 0, 0, 0, 0, 0}, ECbFieldType::ObjectAttachment), ECbValidateError::OutOfBounds);
-
-	TestEqual(TEXT("ValidateCompactBinary(Valid, BinaryAttachment)"), Validate({uint8(ECbFieldType::BinaryAttachment), 0, 0, 0, 0, 0, 0, 0, 0, 0, 0, 0, 0, 0, 0, 0, 0, 0, 0, 0, 0}), ECbValidateError::None);
-	TestEqual(TEXT("ValidateCompactBinary(Valid, BinaryAttachment, NoType)"), Validate({0, 0, 0, 0, 0, 0, 0, 0, 0, 0, 0, 0, 0, 0, 0, 0, 0, 0, 0, 0}, ECbFieldType::BinaryAttachment), ECbValidateError::None);
-	TestEqual(TEXT("ValidateCompactBinary(OutOfBounds, BinaryAttachment)"), Validate({uint8(ECbFieldType::BinaryAttachment), 0, 0, 0, 0, 0, 0, 0, 0, 0, 0, 0, 0, 0, 0, 0, 0, 0, 0, 0}), ECbValidateError::OutOfBounds);
-	TestEqual(TEXT("ValidateCompactBinary(OutOfBounds, BinaryAttachment, NoType)"), Validate({0, 0, 0, 0, 0, 0, 0, 0, 0, 0, 0, 0, 0, 0, 0, 0, 0, 0, 0}, ECbFieldType::BinaryAttachment), ECbValidateError::OutOfBounds);
-
-	TestEqual(TEXT("ValidateCompactBinary(Valid, Hash)"), Validate({uint8(ECbFieldType::Hash), 0, 0, 0, 0, 0, 0, 0, 0, 0, 0, 0, 0, 0, 0, 0, 0, 0, 0, 0, 0}), ECbValidateError::None);
-	TestEqual(TEXT("ValidateCompactBinary(Valid, Hash, NoType)"), Validate({0, 0, 0, 0, 0, 0, 0, 0, 0, 0, 0, 0, 0, 0, 0, 0, 0, 0, 0, 0}, ECbFieldType::Hash), ECbValidateError::None);
-	TestEqual(TEXT("ValidateCompactBinary(OutOfBounds, Hash)"), Validate({uint8(ECbFieldType::Hash), 0, 0, 0, 0, 0, 0, 0, 0, 0, 0, 0, 0, 0, 0, 0, 0, 0, 0, 0}), ECbValidateError::OutOfBounds);
-	TestEqual(TEXT("ValidateCompactBinary(OutOfBounds, Hash, NoType)"), Validate({0, 0, 0, 0, 0, 0, 0, 0, 0, 0, 0, 0, 0, 0, 0, 0, 0, 0, 0}, ECbFieldType::Hash), ECbValidateError::OutOfBounds);
-
-	TestEqual(TEXT("ValidateCompactBinary(Valid, Uuid)"), Validate({uint8(ECbFieldType::Uuid), 0, 0, 0, 0, 0, 0, 0, 0, 0, 0, 0, 0, 0, 0, 0, 0}), ECbValidateError::None);
-	TestEqual(TEXT("ValidateCompactBinary(Valid, Uuid, NoType)"), Validate({0, 0, 0, 0, 0, 0, 0, 0, 0, 0, 0, 0, 0, 0, 0, 0}, ECbFieldType::Uuid), ECbValidateError::None);
-	TestEqual(TEXT("ValidateCompactBinary(OutOfBounds, Uuid)"), Validate({uint8(ECbFieldType::Uuid), 0, 0, 0, 0, 0, 0, 0, 0, 0, 0, 0, 0, 0, 0, 0}), ECbValidateError::OutOfBounds);
-	TestEqual(TEXT("ValidateCompactBinary(OutOfBounds, Uuid, NoType)"), Validate({0, 0, 0, 0, 0, 0, 0, 0, 0, 0, 0, 0, 0, 0, 0}, ECbFieldType::Uuid), ECbValidateError::OutOfBounds);
-
-	TestEqual(TEXT("ValidateCompactBinary(Valid, DateTime)"), Validate({uint8(ECbFieldType::DateTime), 0, 0, 0, 0, 0, 0, 0, 0}), ECbValidateError::None);
-	TestEqual(TEXT("ValidateCompactBinary(Valid, DateTime, NoType)"), Validate({0, 0, 0, 0, 0, 0, 0, 0}, ECbFieldType::DateTime), ECbValidateError::None);
-	TestEqual(TEXT("ValidateCompactBinary(OutOfBounds, DateTime)"), Validate({uint8(ECbFieldType::DateTime), 0, 0, 0, 0, 0, 0, 0}), ECbValidateError::OutOfBounds);
-	TestEqual(TEXT("ValidateCompactBinary(OutOfBounds, DateTime, NoType)"), Validate({0, 0, 0, 0, 0, 0, 0}, ECbFieldType::DateTime), ECbValidateError::OutOfBounds);
-
-	TestEqual(TEXT("ValidateCompactBinary(Valid, TimeSpan)"), Validate({uint8(ECbFieldType::TimeSpan), 0, 0, 0, 0, 0, 0, 0, 0}), ECbValidateError::None);
-	TestEqual(TEXT("ValidateCompactBinary(Valid, TimeSpan, NoType)"), Validate({0, 0, 0, 0, 0, 0, 0, 0}, ECbFieldType::TimeSpan), ECbValidateError::None);
-	TestEqual(TEXT("ValidateCompactBinary(OutOfBounds, TimeSpan)"), Validate({uint8(ECbFieldType::TimeSpan), 0, 0, 0, 0, 0, 0, 0}), ECbValidateError::OutOfBounds);
-	TestEqual(TEXT("ValidateCompactBinary(OutOfBounds, TimeSpan, NoType)"), Validate({0, 0, 0, 0, 0, 0, 0}, ECbFieldType::TimeSpan), ECbValidateError::OutOfBounds);
-
-	TestEqual(TEXT("ValidateCompactBinary(Valid, ObjectId)"), Validate({uint8(ECbFieldType::ObjectId), 0, 0, 0, 0, 0, 0, 0, 0, 0, 0, 0, 0}), ECbValidateError::None);
-	TestEqual(TEXT("ValidateCompactBinary(Valid, ObjectId, NoType)"), Validate({0, 0, 0, 0, 0, 0, 0, 0, 0, 0, 0, 0}, ECbFieldType::ObjectId), ECbValidateError::None);
-	TestEqual(TEXT("ValidateCompactBinary(OutOfBounds, ObjectId)"), Validate({uint8(ECbFieldType::ObjectId), 0, 0, 0, 0, 0, 0, 0, 0, 0, 0, 0}), ECbValidateError::OutOfBounds);
-	TestEqual(TEXT("ValidateCompactBinary(OutOfBounds, ObjectId, NoType)"), Validate({0, 0, 0, 0, 0, 0, 0, 0, 0, 0, 0}, ECbFieldType::ObjectId), ECbValidateError::OutOfBounds);
-
-	TestEqual(TEXT("ValidateCompactBinary(Valid, CustomById, Empty)"), Validate({uint8(ECbFieldType::CustomById), 1, 0}), ECbValidateError::None);
-	TestEqual(TEXT("ValidateCompactBinary(Valid, CustomById, Empty, NoType)"), Validate({1, 0}, ECbFieldType::CustomById), ECbValidateError::None);
-	TestEqual(TEXT("ValidateCompactBinary(Valid, CustomById, Value)"), Validate({uint8(ECbFieldType::CustomById), 2, 0, 0}), ECbValidateError::None);
-	TestEqual(TEXT("ValidateCompactBinary(Valid, CustomById, Value, NoType)"), Validate({2, 0, 0}, ECbFieldType::CustomById), ECbValidateError::None);
-	TestEqual(TEXT("ValidateCompactBinary(Valid, CustomById, MaxTypeId)"), Validate({uint8(ECbFieldType::CustomById), 9, 0xff, 0xff, 0xff, 0xff, 0xff, 0xff, 0xff, 0xff, 0xff}), ECbValidateError::None);
-	TestEqual(TEXT("ValidateCompactBinary(Valid, CustomById, MaxTypeId, NoType)"), Validate({9, 0xff, 0xff, 0xff, 0xff, 0xff, 0xff, 0xff, 0xff, 0xff}, ECbFieldType::CustomById), ECbValidateError::None);
-	TestEqual(TEXT("ValidateCompactBinary(Valid, CustomById, MaxTypeId, Value)"), Validate({uint8(ECbFieldType::CustomById), 10, 0xff, 0xff, 0xff, 0xff, 0xff, 0xff, 0xff, 0xff, 0xff, 0}), ECbValidateError::None);
-	TestEqual(TEXT("ValidateCompactBinary(Valid, CustomById, MaxTypeId, Value, NoType)"), Validate({10, 0xff, 0xff, 0xff, 0xff, 0xff, 0xff, 0xff, 0xff, 0xff, 0}, ECbFieldType::CustomById), ECbValidateError::None);
-	TestEqual(TEXT("ValidateCompactBinary(OutOfBounds, CustomById, NoSize)"), Validate({uint8(ECbFieldType::CustomById)}), ECbValidateError::OutOfBounds);
-	TestEqual(TEXT("ValidateCompactBinary(OutOfBounds, CustomById, NoSize, NoType)"), Validate({}, ECbFieldType::CustomById), ECbValidateError::OutOfBounds);
-	TestEqual(TEXT("ValidateCompactBinary(OutOfBounds, CustomById, NoTypeId)"), Validate({uint8(ECbFieldType::CustomById), 0}), ECbValidateError::OutOfBounds);
-	TestEqual(TEXT("ValidateCompactBinary(OutOfBounds, CustomById, NoTypeId, NoType)"), Validate({0}, ECbFieldType::CustomById), ECbValidateError::OutOfBounds);
-	TestEqual(TEXT("ValidateCompactBinary(OutOfBounds, CustomById, ShortTypeId)"), Validate({uint8(ECbFieldType::CustomById), 1, 0x80}), ECbValidateError::OutOfBounds);
-	TestEqual(TEXT("ValidateCompactBinary(OutOfBounds, CustomById, ShortTypeId, NoType)"), Validate({1, 0x80}, ECbFieldType::CustomById), ECbValidateError::OutOfBounds);
-	TestEqual(TEXT("ValidateCompactBinary(OutOfBounds, CustomById, NoValue)"), Validate({uint8(ECbFieldType::CustomById), 2, 0}), ECbValidateError::OutOfBounds);
-	TestEqual(TEXT("ValidateCompactBinary(OutOfBounds, CustomById, NoValue, NoType)"), Validate({2, 0}, ECbFieldType::CustomById), ECbValidateError::OutOfBounds);
-	TestEqual(TEXT("ValidateCompactBinary(OutOfBounds, CustomById, ShortValue)"), Validate({uint8(ECbFieldType::CustomById), 3, 0, 0}), ECbValidateError::OutOfBounds);
-	TestEqual(TEXT("ValidateCompactBinary(OutOfBounds, CustomById, ShortValue, NoType)"), Validate({3, 0, 0}, ECbFieldType::CustomById), ECbValidateError::OutOfBounds);
-
-	TestEqual(TEXT("ValidateCompactBinary(Valid, CustomByName, Empty)"), Validate({uint8(ECbFieldType::CustomByName), 2, 1, 'A'}), ECbValidateError::None);
-	TestEqual(TEXT("ValidateCompactBinary(Valid, CustomByName, Empty, NoType)"), Validate({2, 1, 'A'}, ECbFieldType::CustomByName), ECbValidateError::None);
-	TestEqual(TEXT("ValidateCompactBinary(Valid, CustomByName, Value)"), Validate({uint8(ECbFieldType::CustomByName), 3, 1, 'A', 0}), ECbValidateError::None);
-	TestEqual(TEXT("ValidateCompactBinary(Valid, CustomByName, Value, NoType)"), Validate({3, 1, 'A', 0}, ECbFieldType::CustomByName), ECbValidateError::None);
-	TestEqual(TEXT("ValidateCompactBinary(OutOfBounds, CustomByName, NoSize)"), Validate({uint8(ECbFieldType::CustomByName)}), ECbValidateError::OutOfBounds);
-	TestEqual(TEXT("ValidateCompactBinary(OutOfBounds, CustomByName, NoSize, NoType)"), Validate({}, ECbFieldType::CustomByName), ECbValidateError::OutOfBounds);
-	TestEqual(TEXT("ValidateCompactBinary(OutOfBounds, CustomByName, NoTypeSize)"), Validate({uint8(ECbFieldType::CustomByName), 0}), ECbValidateError::OutOfBounds);
-	TestEqual(TEXT("ValidateCompactBinary(OutOfBounds, CustomByName, NoTypeSize, NoType)"), Validate({0}, ECbFieldType::CustomByName), ECbValidateError::OutOfBounds);
-	TestEqual(TEXT("ValidateCompactBinary(OutOfBounds, CustomByName, ShortTypeSize)"), Validate({uint8(ECbFieldType::CustomByName), 1, 0x80}), ECbValidateError::OutOfBounds);
-	TestEqual(TEXT("ValidateCompactBinary(OutOfBounds, CustomByName, ShortTypeSize, NoType)"), Validate({1, 0x80}, ECbFieldType::CustomByName), ECbValidateError::OutOfBounds);
-	TestEqual(TEXT("ValidateCompactBinary(OutOfBounds, CustomByName, ShortTypeName)"), Validate({uint8(ECbFieldType::CustomByName), 1, 1}), ECbValidateError::OutOfBounds);
-	TestEqual(TEXT("ValidateCompactBinary(OutOfBounds, CustomByName, ShortTypeName, NoType)"), Validate({1, 1}, ECbFieldType::CustomByName), ECbValidateError::OutOfBounds);
-	TestEqual(TEXT("ValidateCompactBinary(OutOfBounds, CustomByName, NoValue)"), Validate({uint8(ECbFieldType::CustomByName), 3, 1, 'A'}), ECbValidateError::OutOfBounds);
-	TestEqual(TEXT("ValidateCompactBinary(OutOfBounds, CustomByName, NoValue, NoType)"), Validate({3, 1, 'A'}, ECbFieldType::CustomByName), ECbValidateError::OutOfBounds);
-	TestEqual(TEXT("ValidateCompactBinary(OutOfBounds, CustomByName, ShortValue)"), Validate({uint8(ECbFieldType::CustomByName), 4, 1, 'A', 0}), ECbValidateError::OutOfBounds);
-	TestEqual(TEXT("ValidateCompactBinary(OutOfBounds, CustomByName, ShortValue, NoType)"), Validate({4, 1, 'A', 0}, ECbFieldType::CustomByName), ECbValidateError::OutOfBounds);
+	TEST_EQUAL(TEXT("ValidateCompactBinary(OutOfBounds, Empty)"), Validate({}), ECbValidateError::OutOfBounds);
+	TEST_EQUAL(TEXT("ValidateCompactBinary(OutOfBounds, Empty, Mode)"), ValidateMode({}, ECbValidateMode::None), ECbValidateError::None);
+
+	TEST_EQUAL(TEXT("ValidateCompactBinary(Valid, Null)"), Validate({NullNoName}), ECbValidateError::None);
+	TEST_EQUAL(TEXT("ValidateCompactBinary(Valid, Null, Name)"), Validate({NullWithName, 1, 'N'}), ECbValidateError::None);
+	TEST_EQUAL(TEXT("ValidateCompactBinary(OutOfBounds, Null, Name)"), Validate({NullWithName}), ECbValidateError::OutOfBounds);
+	TEST_EQUAL(TEXT("ValidateCompactBinary(OutOfBounds, Null, Name)"), Validate({NullWithName, 1}), ECbValidateError::OutOfBounds);
+	TEST_EQUAL(TEXT("ValidateCompactBinary(OutOfBounds, Null, Name)"), Validate({NullWithName, 0x80}), ECbValidateError::OutOfBounds);
+	TEST_EQUAL(TEXT("ValidateCompactBinary(OutOfBounds, Null, Name)"), Validate({NullWithName, 0x80, 128}), ECbValidateError::OutOfBounds);
+
+	TEST_EQUAL(TEXT("ValidateCompactBinary(Valid, Object, Empty)"), Validate({uint8(ECbFieldType::Object), 0}), ECbValidateError::None);
+	TEST_EQUAL(TEXT("ValidateCompactBinary(Valid, Object, Empty, NoType)"), Validate({0}, ECbFieldType::Object), ECbValidateError::None);
+	TEST_EQUAL(TEXT("ValidateCompactBinary(Valid, Object, Field)"), Validate({uint8(ECbFieldType::Object), 7, NullWithName, 1, 'N', IntWithName, 1, 'I', 0}), ECbValidateError::None);
+	TEST_EQUAL(TEXT("ValidateCompactBinary(Valid, Object, Field, NoType)"), Validate({7, NullWithName, 1, 'N', IntWithName, 1, 'I', 0}, ECbFieldType::Object), ECbValidateError::None);
+	TEST_EQUAL(TEXT("ValidateCompactBinary(OutOfBounds, Object)"), Validate({uint8(ECbFieldType::Object)}), ECbValidateError::OutOfBounds);
+	TEST_EQUAL(TEXT("ValidateCompactBinary(OutOfBounds, Object, NoType)"), Validate({}, ECbFieldType::Object), ECbValidateError::OutOfBounds);
+	TEST_EQUAL(TEXT("ValidateCompactBinary(OutOfBounds, Object)"), Validate({uint8(ECbFieldType::Object), 1}), ECbValidateError::OutOfBounds);
+	TEST_EQUAL(TEXT("ValidateCompactBinary(OutOfBounds, Object, NoType)"), Validate({1}, ECbFieldType::Object), ECbValidateError::OutOfBounds);
+	TEST_EQUAL(TEXT("ValidateCompactBinary(OutOfBounds, Object, OOB Field)"), Validate({uint8(ECbFieldType::Object), 3, AddName(ECbFieldType::Float32), 1, 'N'}), ECbValidateError::OutOfBounds);
+	TEST_EQUAL(TEXT("ValidateCompactBinary(OutOfBounds, Object, OOB Field, NoType)"), Validate({3, AddName(ECbFieldType::Float32), 1, 'N'}, ECbFieldType::Object), ECbValidateError::OutOfBounds);
+
+	TEST_EQUAL(TEXT("ValidateCompactBinary(Valid, UniformObject, Field)"), Validate({uint8(ECbFieldType::UniformObject), 3, NullWithName, 1, 'N'}), ECbValidateError::None);
+	TEST_EQUAL(TEXT("ValidateCompactBinary(Valid, UniformObject, Field, NoType)"), Validate({3, NullWithName, 1, 'N'}, ECbFieldType::UniformObject), ECbValidateError::None);
+	TEST_EQUAL(TEXT("ValidateCompactBinary(OutOfBounds, UniformObject)"), Validate({uint8(ECbFieldType::UniformObject)}), ECbValidateError::OutOfBounds);
+	TEST_EQUAL(TEXT("ValidateCompactBinary(OutOfBounds, UniformObject, NoType)"), Validate({}, ECbFieldType::UniformObject), ECbValidateError::OutOfBounds);
+	TEST_EQUAL(TEXT("ValidateCompactBinary(OutOfBounds, UniformObject)"), Validate({uint8(ECbFieldType::UniformObject), 1}), ECbValidateError::OutOfBounds);
+	TEST_EQUAL(TEXT("ValidateCompactBinary(OutOfBounds, UniformObject, NoType)"), Validate({1}, ECbFieldType::UniformObject), ECbValidateError::OutOfBounds);
+	TEST_EQUAL(TEXT("ValidateCompactBinary(OutOfBounds, UniformObject, OOB Field)"), Validate({uint8(ECbFieldType::UniformObject), 3, AddName(ECbFieldType::Float32), 1, 'N'}), ECbValidateError::OutOfBounds);
+	TEST_EQUAL(TEXT("ValidateCompactBinary(OutOfBounds, UniformObject, OOB Field, NoType)"), Validate({3, AddName(ECbFieldType::Float32), 1, 'N'}, ECbFieldType::UniformObject), ECbValidateError::OutOfBounds);
+
+	TEST_EQUAL(TEXT("ValidateCompactBinary(Valid, Array, Empty)"), Validate({uint8(ECbFieldType::Array), 1, 0}), ECbValidateError::None);
+	TEST_EQUAL(TEXT("ValidateCompactBinary(Valid, Array, Empty, NoType)"), Validate({1, 0}, ECbFieldType::Array), ECbValidateError::None);
+	TEST_EQUAL(TEXT("ValidateCompactBinary(Valid, Array, Field)"), Validate({uint8(ECbFieldType::Array), 4, 2, NullNoName, IntNoName, 0}), ECbValidateError::None);
+	TEST_EQUAL(TEXT("ValidateCompactBinary(Valid, Array, Field, NoType)"), Validate({4, 2, NullNoName, IntNoName, 0}, ECbFieldType::Array), ECbValidateError::None);
+	TEST_EQUAL(TEXT("ValidateCompactBinary(OutOfBounds, Array)"), Validate({uint8(ECbFieldType::Array)}), ECbValidateError::OutOfBounds);
+	TEST_EQUAL(TEXT("ValidateCompactBinary(OutOfBounds, Array, NoType)"), Validate({}, ECbFieldType::Array), ECbValidateError::OutOfBounds);
+	TEST_EQUAL(TEXT("ValidateCompactBinary(OutOfBounds, Array)"), Validate({uint8(ECbFieldType::Array), 1}), ECbValidateError::OutOfBounds);
+	TEST_EQUAL(TEXT("ValidateCompactBinary(OutOfBounds, Array, NoType)"), Validate({1}, ECbFieldType::Array), ECbValidateError::OutOfBounds);
+	TEST_EQUAL(TEXT("ValidateCompactBinary(OutOfBounds, Array, OOB Field)"), Validate({uint8(ECbFieldType::Array), 2, 1, uint8(ECbFieldType::Float32)}), ECbValidateError::OutOfBounds);
+	TEST_EQUAL(TEXT("ValidateCompactBinary(OutOfBounds, Array, OOB Field, NoType)"), Validate({2, 1, uint8(ECbFieldType::Float32)}, ECbFieldType::Array), ECbValidateError::OutOfBounds);
+
+	TEST_EQUAL(TEXT("ValidateCompactBinary(Valid, UniformArray, Field)"), Validate({uint8(ECbFieldType::UniformArray), 3, 1, IntNoName, 0}), ECbValidateError::None);
+	TEST_EQUAL(TEXT("ValidateCompactBinary(Valid, UniformArray, Field, NoType)"), Validate({3, 1, IntNoName, 0}, ECbFieldType::UniformArray), ECbValidateError::None);
+	TEST_EQUAL(TEXT("ValidateCompactBinary(OutOfBounds, UniformArray)"), Validate({uint8(ECbFieldType::UniformArray)}), ECbValidateError::OutOfBounds);
+	TEST_EQUAL(TEXT("ValidateCompactBinary(OutOfBounds, UniformArray, NoType)"), Validate({}, ECbFieldType::UniformArray), ECbValidateError::OutOfBounds);
+	TEST_EQUAL(TEXT("ValidateCompactBinary(OutOfBounds, UniformArray)"), Validate({uint8(ECbFieldType::UniformArray), 1}), ECbValidateError::OutOfBounds);
+	TEST_EQUAL(TEXT("ValidateCompactBinary(OutOfBounds, UniformArray, NoType)"), Validate({1}, ECbFieldType::UniformArray), ECbValidateError::OutOfBounds);
+	TEST_EQUAL(TEXT("ValidateCompactBinary(OutOfBounds, UniformArray, OOB Field)"), Validate({uint8(ECbFieldType::UniformArray), 2, 1, uint8(ECbFieldType::Float32)}), ECbValidateError::OutOfBounds);
+	TEST_EQUAL(TEXT("ValidateCompactBinary(OutOfBounds, UniformArray, OOB Field, NoType)"), Validate({2, 1, uint8(ECbFieldType::Float32)}, ECbFieldType::UniformArray), ECbValidateError::OutOfBounds);
+
+	TEST_EQUAL(TEXT("ValidateCompactBinary(Valid, Binary, Empty)"), Validate({uint8(ECbFieldType::Binary), 0}), ECbValidateError::None);
+	TEST_EQUAL(TEXT("ValidateCompactBinary(Valid, Binary, Empty, NoType)"), Validate({0}, ECbFieldType::Binary), ECbValidateError::None);
+	TEST_EQUAL(TEXT("ValidateCompactBinary(Valid, Binary, Field)"), Validate({uint8(ECbFieldType::Binary), 1, 0}), ECbValidateError::None);
+	TEST_EQUAL(TEXT("ValidateCompactBinary(Valid, Binary, Field, NoType)"), Validate({1, 0}, ECbFieldType::Binary), ECbValidateError::None);
+	TEST_EQUAL(TEXT("ValidateCompactBinary(OutOfBounds, Binary)"), Validate({uint8(ECbFieldType::Binary)}), ECbValidateError::OutOfBounds);
+	TEST_EQUAL(TEXT("ValidateCompactBinary(OutOfBounds, Binary, NoType)"), Validate({}, ECbFieldType::Binary), ECbValidateError::OutOfBounds);
+	TEST_EQUAL(TEXT("ValidateCompactBinary(OutOfBounds, Binary)"), Validate({uint8(ECbFieldType::Binary), 1}), ECbValidateError::OutOfBounds);
+	TEST_EQUAL(TEXT("ValidateCompactBinary(OutOfBounds, Binary, NoType)"), Validate({1}, ECbFieldType::Binary), ECbValidateError::OutOfBounds);
+
+	TEST_EQUAL(TEXT("ValidateCompactBinary(Valid, String, Empty)"), Validate({uint8(ECbFieldType::String), 0}), ECbValidateError::None);
+	TEST_EQUAL(TEXT("ValidateCompactBinary(Valid, String, Empty, NoType)"), Validate({0}, ECbFieldType::String), ECbValidateError::None);
+	TEST_EQUAL(TEXT("ValidateCompactBinary(Valid, String, Field)"), Validate({uint8(ECbFieldType::String), 1, 'S'}), ECbValidateError::None);
+	TEST_EQUAL(TEXT("ValidateCompactBinary(Valid, String, Field, NoType)"), Validate({1, 'S'}, ECbFieldType::String), ECbValidateError::None);
+	TEST_EQUAL(TEXT("ValidateCompactBinary(OutOfBounds, String)"), Validate({uint8(ECbFieldType::String)}), ECbValidateError::OutOfBounds);
+	TEST_EQUAL(TEXT("ValidateCompactBinary(OutOfBounds, String, NoType)"), Validate({}, ECbFieldType::String), ECbValidateError::OutOfBounds);
+	TEST_EQUAL(TEXT("ValidateCompactBinary(OutOfBounds, String)"), Validate({uint8(ECbFieldType::String), 1}), ECbValidateError::OutOfBounds);
+	TEST_EQUAL(TEXT("ValidateCompactBinary(OutOfBounds, String, NoType)"), Validate({1}, ECbFieldType::String), ECbValidateError::OutOfBounds);
+
+	TEST_EQUAL(TEXT("ValidateCompactBinary(Valid, IntegerPositive, 1-byte)"), Validate({uint8(ECbFieldType::IntegerPositive), 0}), ECbValidateError::None);
+	TEST_EQUAL(TEXT("ValidateCompactBinary(Valid, IntegerPositive, 1-byte, NoType)"), Validate({0}, ECbFieldType::IntegerPositive), ECbValidateError::None);
+	TEST_EQUAL(TEXT("ValidateCompactBinary(Valid, IntegerPositive, 2-byte)"), Validate({uint8(ECbFieldType::IntegerPositive), 0x80, 0x80}), ECbValidateError::None);
+	TEST_EQUAL(TEXT("ValidateCompactBinary(Valid, IntegerPositive, 2-byte, NoType)"), Validate({0x80, 0x80}, ECbFieldType::IntegerPositive), ECbValidateError::None);
+	TEST_EQUAL(TEXT("ValidateCompactBinary(OutOfBounds, IntegerPositive, 1-byte)"), Validate({uint8(ECbFieldType::IntegerPositive)}), ECbValidateError::OutOfBounds);
+	TEST_EQUAL(TEXT("ValidateCompactBinary(OutOfBounds, IntegerPositive, 1-byte, NoType)"), Validate({}, ECbFieldType::IntegerPositive), ECbValidateError::OutOfBounds);
+	TEST_EQUAL(TEXT("ValidateCompactBinary(OutOfBounds, IntegerPositive, 2-byte)"), Validate({uint8(ECbFieldType::IntegerPositive), 0x80}), ECbValidateError::OutOfBounds);
+	TEST_EQUAL(TEXT("ValidateCompactBinary(OutOfBounds, IntegerPositive, 2-byte, NoType)"), Validate({0x80}, ECbFieldType::IntegerPositive), ECbValidateError::OutOfBounds);
+	TEST_EQUAL(TEXT("ValidateCompactBinary(OutOfBounds, IntegerPositive, 9-byte)"), Validate({uint8(ECbFieldType::IntegerPositive), 0xff, 0, 0, 0, 0, 0, 0, 0}), ECbValidateError::OutOfBounds);
+	TEST_EQUAL(TEXT("ValidateCompactBinary(OutOfBounds, IntegerPositive, 9-byte, NoType)"), Validate({0xff, 0, 0, 0, 0, 0, 0, 0}, ECbFieldType::IntegerPositive), ECbValidateError::OutOfBounds);
+
+	TEST_EQUAL(TEXT("ValidateCompactBinary(Valid, IntegerNegative, 1-byte)"), Validate({uint8(ECbFieldType::IntegerNegative), 0}), ECbValidateError::None);
+	TEST_EQUAL(TEXT("ValidateCompactBinary(Valid, IntegerNegative, 1-byte, NoType)"), Validate({0}, ECbFieldType::IntegerNegative), ECbValidateError::None);
+	TEST_EQUAL(TEXT("ValidateCompactBinary(Valid, IntegerNegative, 2-byte)"), Validate({uint8(ECbFieldType::IntegerNegative), 0x80, 0x80}), ECbValidateError::None);
+	TEST_EQUAL(TEXT("ValidateCompactBinary(Valid, IntegerNegative, 2-byte, NoType)"), Validate({0x80, 0x80}, ECbFieldType::IntegerNegative), ECbValidateError::None);
+	TEST_EQUAL(TEXT("ValidateCompactBinary(OutOfBounds, IntegerNegative, 1-byte)"), Validate({uint8(ECbFieldType::IntegerNegative)}), ECbValidateError::OutOfBounds);
+	TEST_EQUAL(TEXT("ValidateCompactBinary(OutOfBounds, IntegerNegative, 1-byte, NoType)"), Validate({}, ECbFieldType::IntegerNegative), ECbValidateError::OutOfBounds);
+	TEST_EQUAL(TEXT("ValidateCompactBinary(OutOfBounds, IntegerNegative, 2-byte)"), Validate({uint8(ECbFieldType::IntegerNegative), 0x80}), ECbValidateError::OutOfBounds);
+	TEST_EQUAL(TEXT("ValidateCompactBinary(OutOfBounds, IntegerNegative, 2-byte, NoType)"), Validate({0x80}, ECbFieldType::IntegerNegative), ECbValidateError::OutOfBounds);
+	TEST_EQUAL(TEXT("ValidateCompactBinary(OutOfBounds, IntegerNegative, 9-byte)"), Validate({uint8(ECbFieldType::IntegerNegative), 0xff, 0, 0, 0, 0, 0, 0, 0}), ECbValidateError::OutOfBounds);
+	TEST_EQUAL(TEXT("ValidateCompactBinary(OutOfBounds, IntegerNegative, 9-byte, NoType)"), Validate({0xff, 0, 0, 0, 0, 0, 0, 0}, ECbFieldType::IntegerNegative), ECbValidateError::OutOfBounds);
+
+	TEST_EQUAL(TEXT("ValidateCompactBinary(Valid, Float32)"), Validate({uint8(ECbFieldType::Float32), 0, 0, 0, 0}), ECbValidateError::None);
+	TEST_EQUAL(TEXT("ValidateCompactBinary(Valid, Float32, NoType)"), Validate({0, 0, 0, 0}, ECbFieldType::Float32), ECbValidateError::None);
+	TEST_EQUAL(TEXT("ValidateCompactBinary(OutOfBounds, Float32)"), Validate({uint8(ECbFieldType::Float32), 0, 0, 0}), ECbValidateError::OutOfBounds);
+	TEST_EQUAL(TEXT("ValidateCompactBinary(OutOfBounds, Float32, NoType)"), Validate({0, 0, 0}, ECbFieldType::Float32), ECbValidateError::OutOfBounds);
+
+	TEST_EQUAL(TEXT("ValidateCompactBinary(Valid, Float64)"), Validate({uint8(ECbFieldType::Float64), 0x3f, 0xff, 0xff, 0xff, 0xf0, 0x00, 0x00, 0x00}), ECbValidateError::None);
+	TEST_EQUAL(TEXT("ValidateCompactBinary(Valid, Float64, NoType)"), Validate({0x3f, 0xff, 0xff, 0xff, 0xf0, 0x00, 0x00, 0x00}, ECbFieldType::Float64), ECbValidateError::None);
+	TEST_EQUAL(TEXT("ValidateCompactBinary(OutOfBounds, Float64)"), Validate({uint8(ECbFieldType::Float64), 0, 0, 0, 0, 0, 0, 0}), ECbValidateError::OutOfBounds);
+	TEST_EQUAL(TEXT("ValidateCompactBinary(OutOfBounds, Float64, NoType)"), Validate({0, 0, 0, 0, 0, 0, 0}, ECbFieldType::Float64), ECbValidateError::OutOfBounds);
+
+	TEST_EQUAL(TEXT("ValidateCompactBinary(Valid, BoolFalse)"), Validate({uint8(ECbFieldType::BoolFalse)}), ECbValidateError::None);
+	TEST_EQUAL(TEXT("ValidateCompactBinary(Valid, BoolTrue)"), Validate({uint8(ECbFieldType::BoolTrue)}), ECbValidateError::None);
+
+	TEST_EQUAL(TEXT("ValidateCompactBinary(Valid, ObjectAttachment)"), Validate({uint8(ECbFieldType::ObjectAttachment), 0, 0, 0, 0, 0, 0, 0, 0, 0, 0, 0, 0, 0, 0, 0, 0, 0, 0, 0, 0}), ECbValidateError::None);
+	TEST_EQUAL(TEXT("ValidateCompactBinary(Valid, ObjectAttachment, NoType)"), Validate({0, 0, 0, 0, 0, 0, 0, 0, 0, 0, 0, 0, 0, 0, 0, 0, 0, 0, 0, 0}, ECbFieldType::ObjectAttachment), ECbValidateError::None);
+	TEST_EQUAL(TEXT("ValidateCompactBinary(OutOfBounds, ObjectAttachment)"), Validate({uint8(ECbFieldType::ObjectAttachment), 0, 0, 0, 0, 0, 0, 0, 0, 0, 0, 0, 0, 0, 0, 0, 0, 0, 0, 0}), ECbValidateError::OutOfBounds);
+	TEST_EQUAL(TEXT("ValidateCompactBinary(OutOfBounds, ObjectAttachment, NoType)"), Validate({0, 0, 0, 0, 0, 0, 0, 0, 0, 0, 0, 0, 0, 0, 0, 0, 0, 0, 0}, ECbFieldType::ObjectAttachment), ECbValidateError::OutOfBounds);
+
+	TEST_EQUAL(TEXT("ValidateCompactBinary(Valid, BinaryAttachment)"), Validate({uint8(ECbFieldType::BinaryAttachment), 0, 0, 0, 0, 0, 0, 0, 0, 0, 0, 0, 0, 0, 0, 0, 0, 0, 0, 0, 0}), ECbValidateError::None);
+	TEST_EQUAL(TEXT("ValidateCompactBinary(Valid, BinaryAttachment, NoType)"), Validate({0, 0, 0, 0, 0, 0, 0, 0, 0, 0, 0, 0, 0, 0, 0, 0, 0, 0, 0, 0}, ECbFieldType::BinaryAttachment), ECbValidateError::None);
+	TEST_EQUAL(TEXT("ValidateCompactBinary(OutOfBounds, BinaryAttachment)"), Validate({uint8(ECbFieldType::BinaryAttachment), 0, 0, 0, 0, 0, 0, 0, 0, 0, 0, 0, 0, 0, 0, 0, 0, 0, 0, 0}), ECbValidateError::OutOfBounds);
+	TEST_EQUAL(TEXT("ValidateCompactBinary(OutOfBounds, BinaryAttachment, NoType)"), Validate({0, 0, 0, 0, 0, 0, 0, 0, 0, 0, 0, 0, 0, 0, 0, 0, 0, 0, 0}, ECbFieldType::BinaryAttachment), ECbValidateError::OutOfBounds);
+
+	TEST_EQUAL(TEXT("ValidateCompactBinary(Valid, Hash)"), Validate({uint8(ECbFieldType::Hash), 0, 0, 0, 0, 0, 0, 0, 0, 0, 0, 0, 0, 0, 0, 0, 0, 0, 0, 0, 0}), ECbValidateError::None);
+	TEST_EQUAL(TEXT("ValidateCompactBinary(Valid, Hash, NoType)"), Validate({0, 0, 0, 0, 0, 0, 0, 0, 0, 0, 0, 0, 0, 0, 0, 0, 0, 0, 0, 0}, ECbFieldType::Hash), ECbValidateError::None);
+	TEST_EQUAL(TEXT("ValidateCompactBinary(OutOfBounds, Hash)"), Validate({uint8(ECbFieldType::Hash), 0, 0, 0, 0, 0, 0, 0, 0, 0, 0, 0, 0, 0, 0, 0, 0, 0, 0, 0}), ECbValidateError::OutOfBounds);
+	TEST_EQUAL(TEXT("ValidateCompactBinary(OutOfBounds, Hash, NoType)"), Validate({0, 0, 0, 0, 0, 0, 0, 0, 0, 0, 0, 0, 0, 0, 0, 0, 0, 0, 0}, ECbFieldType::Hash), ECbValidateError::OutOfBounds);
+
+	TEST_EQUAL(TEXT("ValidateCompactBinary(Valid, Uuid)"), Validate({uint8(ECbFieldType::Uuid), 0, 0, 0, 0, 0, 0, 0, 0, 0, 0, 0, 0, 0, 0, 0, 0}), ECbValidateError::None);
+	TEST_EQUAL(TEXT("ValidateCompactBinary(Valid, Uuid, NoType)"), Validate({0, 0, 0, 0, 0, 0, 0, 0, 0, 0, 0, 0, 0, 0, 0, 0}, ECbFieldType::Uuid), ECbValidateError::None);
+	TEST_EQUAL(TEXT("ValidateCompactBinary(OutOfBounds, Uuid)"), Validate({uint8(ECbFieldType::Uuid), 0, 0, 0, 0, 0, 0, 0, 0, 0, 0, 0, 0, 0, 0, 0}), ECbValidateError::OutOfBounds);
+	TEST_EQUAL(TEXT("ValidateCompactBinary(OutOfBounds, Uuid, NoType)"), Validate({0, 0, 0, 0, 0, 0, 0, 0, 0, 0, 0, 0, 0, 0, 0}, ECbFieldType::Uuid), ECbValidateError::OutOfBounds);
+
+	TEST_EQUAL(TEXT("ValidateCompactBinary(Valid, DateTime)"), Validate({uint8(ECbFieldType::DateTime), 0, 0, 0, 0, 0, 0, 0, 0}), ECbValidateError::None);
+	TEST_EQUAL(TEXT("ValidateCompactBinary(Valid, DateTime, NoType)"), Validate({0, 0, 0, 0, 0, 0, 0, 0}, ECbFieldType::DateTime), ECbValidateError::None);
+	TEST_EQUAL(TEXT("ValidateCompactBinary(OutOfBounds, DateTime)"), Validate({uint8(ECbFieldType::DateTime), 0, 0, 0, 0, 0, 0, 0}), ECbValidateError::OutOfBounds);
+	TEST_EQUAL(TEXT("ValidateCompactBinary(OutOfBounds, DateTime, NoType)"), Validate({0, 0, 0, 0, 0, 0, 0}, ECbFieldType::DateTime), ECbValidateError::OutOfBounds);
+
+	TEST_EQUAL(TEXT("ValidateCompactBinary(Valid, TimeSpan)"), Validate({uint8(ECbFieldType::TimeSpan), 0, 0, 0, 0, 0, 0, 0, 0}), ECbValidateError::None);
+	TEST_EQUAL(TEXT("ValidateCompactBinary(Valid, TimeSpan, NoType)"), Validate({0, 0, 0, 0, 0, 0, 0, 0}, ECbFieldType::TimeSpan), ECbValidateError::None);
+	TEST_EQUAL(TEXT("ValidateCompactBinary(OutOfBounds, TimeSpan)"), Validate({uint8(ECbFieldType::TimeSpan), 0, 0, 0, 0, 0, 0, 0}), ECbValidateError::OutOfBounds);
+	TEST_EQUAL(TEXT("ValidateCompactBinary(OutOfBounds, TimeSpan, NoType)"), Validate({0, 0, 0, 0, 0, 0, 0}, ECbFieldType::TimeSpan), ECbValidateError::OutOfBounds);
+
+	TEST_EQUAL(TEXT("ValidateCompactBinary(Valid, ObjectId)"), Validate({uint8(ECbFieldType::ObjectId), 0, 0, 0, 0, 0, 0, 0, 0, 0, 0, 0, 0}), ECbValidateError::None);
+	TEST_EQUAL(TEXT("ValidateCompactBinary(Valid, ObjectId, NoType)"), Validate({0, 0, 0, 0, 0, 0, 0, 0, 0, 0, 0, 0}, ECbFieldType::ObjectId), ECbValidateError::None);
+	TEST_EQUAL(TEXT("ValidateCompactBinary(OutOfBounds, ObjectId)"), Validate({uint8(ECbFieldType::ObjectId), 0, 0, 0, 0, 0, 0, 0, 0, 0, 0, 0}), ECbValidateError::OutOfBounds);
+	TEST_EQUAL(TEXT("ValidateCompactBinary(OutOfBounds, ObjectId, NoType)"), Validate({0, 0, 0, 0, 0, 0, 0, 0, 0, 0, 0}, ECbFieldType::ObjectId), ECbValidateError::OutOfBounds);
+
+	TEST_EQUAL(TEXT("ValidateCompactBinary(Valid, CustomById, Empty)"), Validate({uint8(ECbFieldType::CustomById), 1, 0}), ECbValidateError::None);
+	TEST_EQUAL(TEXT("ValidateCompactBinary(Valid, CustomById, Empty, NoType)"), Validate({1, 0}, ECbFieldType::CustomById), ECbValidateError::None);
+	TEST_EQUAL(TEXT("ValidateCompactBinary(Valid, CustomById, Value)"), Validate({uint8(ECbFieldType::CustomById), 2, 0, 0}), ECbValidateError::None);
+	TEST_EQUAL(TEXT("ValidateCompactBinary(Valid, CustomById, Value, NoType)"), Validate({2, 0, 0}, ECbFieldType::CustomById), ECbValidateError::None);
+	TEST_EQUAL(TEXT("ValidateCompactBinary(Valid, CustomById, MaxTypeId)"), Validate({uint8(ECbFieldType::CustomById), 9, 0xff, 0xff, 0xff, 0xff, 0xff, 0xff, 0xff, 0xff, 0xff}), ECbValidateError::None);
+	TEST_EQUAL(TEXT("ValidateCompactBinary(Valid, CustomById, MaxTypeId, NoType)"), Validate({9, 0xff, 0xff, 0xff, 0xff, 0xff, 0xff, 0xff, 0xff, 0xff}, ECbFieldType::CustomById), ECbValidateError::None);
+	TEST_EQUAL(TEXT("ValidateCompactBinary(Valid, CustomById, MaxTypeId, Value)"), Validate({uint8(ECbFieldType::CustomById), 10, 0xff, 0xff, 0xff, 0xff, 0xff, 0xff, 0xff, 0xff, 0xff, 0}), ECbValidateError::None);
+	TEST_EQUAL(TEXT("ValidateCompactBinary(Valid, CustomById, MaxTypeId, Value, NoType)"), Validate({10, 0xff, 0xff, 0xff, 0xff, 0xff, 0xff, 0xff, 0xff, 0xff, 0}, ECbFieldType::CustomById), ECbValidateError::None);
+	TEST_EQUAL(TEXT("ValidateCompactBinary(OutOfBounds, CustomById, NoSize)"), Validate({uint8(ECbFieldType::CustomById)}), ECbValidateError::OutOfBounds);
+	TEST_EQUAL(TEXT("ValidateCompactBinary(OutOfBounds, CustomById, NoSize, NoType)"), Validate({}, ECbFieldType::CustomById), ECbValidateError::OutOfBounds);
+	TEST_EQUAL(TEXT("ValidateCompactBinary(OutOfBounds, CustomById, NoTypeId)"), Validate({uint8(ECbFieldType::CustomById), 0}), ECbValidateError::OutOfBounds);
+	TEST_EQUAL(TEXT("ValidateCompactBinary(OutOfBounds, CustomById, NoTypeId, NoType)"), Validate({0}, ECbFieldType::CustomById), ECbValidateError::OutOfBounds);
+	TEST_EQUAL(TEXT("ValidateCompactBinary(OutOfBounds, CustomById, ShortTypeId)"), Validate({uint8(ECbFieldType::CustomById), 1, 0x80}), ECbValidateError::OutOfBounds);
+	TEST_EQUAL(TEXT("ValidateCompactBinary(OutOfBounds, CustomById, ShortTypeId, NoType)"), Validate({1, 0x80}, ECbFieldType::CustomById), ECbValidateError::OutOfBounds);
+	TEST_EQUAL(TEXT("ValidateCompactBinary(OutOfBounds, CustomById, NoValue)"), Validate({uint8(ECbFieldType::CustomById), 2, 0}), ECbValidateError::OutOfBounds);
+	TEST_EQUAL(TEXT("ValidateCompactBinary(OutOfBounds, CustomById, NoValue, NoType)"), Validate({2, 0}, ECbFieldType::CustomById), ECbValidateError::OutOfBounds);
+	TEST_EQUAL(TEXT("ValidateCompactBinary(OutOfBounds, CustomById, ShortValue)"), Validate({uint8(ECbFieldType::CustomById), 3, 0, 0}), ECbValidateError::OutOfBounds);
+	TEST_EQUAL(TEXT("ValidateCompactBinary(OutOfBounds, CustomById, ShortValue, NoType)"), Validate({3, 0, 0}, ECbFieldType::CustomById), ECbValidateError::OutOfBounds);
+
+	TEST_EQUAL(TEXT("ValidateCompactBinary(Valid, CustomByName, Empty)"), Validate({uint8(ECbFieldType::CustomByName), 2, 1, 'A'}), ECbValidateError::None);
+	TEST_EQUAL(TEXT("ValidateCompactBinary(Valid, CustomByName, Empty, NoType)"), Validate({2, 1, 'A'}, ECbFieldType::CustomByName), ECbValidateError::None);
+	TEST_EQUAL(TEXT("ValidateCompactBinary(Valid, CustomByName, Value)"), Validate({uint8(ECbFieldType::CustomByName), 3, 1, 'A', 0}), ECbValidateError::None);
+	TEST_EQUAL(TEXT("ValidateCompactBinary(Valid, CustomByName, Value, NoType)"), Validate({3, 1, 'A', 0}, ECbFieldType::CustomByName), ECbValidateError::None);
+	TEST_EQUAL(TEXT("ValidateCompactBinary(OutOfBounds, CustomByName, NoSize)"), Validate({uint8(ECbFieldType::CustomByName)}), ECbValidateError::OutOfBounds);
+	TEST_EQUAL(TEXT("ValidateCompactBinary(OutOfBounds, CustomByName, NoSize, NoType)"), Validate({}, ECbFieldType::CustomByName), ECbValidateError::OutOfBounds);
+	TEST_EQUAL(TEXT("ValidateCompactBinary(OutOfBounds, CustomByName, NoTypeSize)"), Validate({uint8(ECbFieldType::CustomByName), 0}), ECbValidateError::OutOfBounds);
+	TEST_EQUAL(TEXT("ValidateCompactBinary(OutOfBounds, CustomByName, NoTypeSize, NoType)"), Validate({0}, ECbFieldType::CustomByName), ECbValidateError::OutOfBounds);
+	TEST_EQUAL(TEXT("ValidateCompactBinary(OutOfBounds, CustomByName, ShortTypeSize)"), Validate({uint8(ECbFieldType::CustomByName), 1, 0x80}), ECbValidateError::OutOfBounds);
+	TEST_EQUAL(TEXT("ValidateCompactBinary(OutOfBounds, CustomByName, ShortTypeSize, NoType)"), Validate({1, 0x80}, ECbFieldType::CustomByName), ECbValidateError::OutOfBounds);
+	TEST_EQUAL(TEXT("ValidateCompactBinary(OutOfBounds, CustomByName, ShortTypeName)"), Validate({uint8(ECbFieldType::CustomByName), 1, 1}), ECbValidateError::OutOfBounds);
+	TEST_EQUAL(TEXT("ValidateCompactBinary(OutOfBounds, CustomByName, ShortTypeName, NoType)"), Validate({1, 1}, ECbFieldType::CustomByName), ECbValidateError::OutOfBounds);
+	TEST_EQUAL(TEXT("ValidateCompactBinary(OutOfBounds, CustomByName, NoValue)"), Validate({uint8(ECbFieldType::CustomByName), 3, 1, 'A'}), ECbValidateError::OutOfBounds);
+	TEST_EQUAL(TEXT("ValidateCompactBinary(OutOfBounds, CustomByName, NoValue, NoType)"), Validate({3, 1, 'A'}, ECbFieldType::CustomByName), ECbValidateError::OutOfBounds);
+	TEST_EQUAL(TEXT("ValidateCompactBinary(OutOfBounds, CustomByName, ShortValue)"), Validate({uint8(ECbFieldType::CustomByName), 4, 1, 'A', 0}), ECbValidateError::OutOfBounds);
+	TEST_EQUAL(TEXT("ValidateCompactBinary(OutOfBounds, CustomByName, ShortValue, NoType)"), Validate({4, 1, 'A', 0}, ECbFieldType::CustomByName), ECbValidateError::OutOfBounds);
 
 	// Test InvalidType
-	TestEqual(TEXT("ValidateCompactBinary(InvalidType, Unknown)"), Validate({uint8(ECbFieldType::ObjectId) + 1}), ECbValidateError::InvalidType);
-	TestEqual(TEXT("ValidateCompactBinary(InvalidType, Unknown)"), Validate({}, ECbFieldType(uint8(ECbFieldType::ObjectId) + 1)), ECbValidateError::InvalidType);
-	TestEqual(TEXT("ValidateCompactBinary(InvalidType, HasFieldType)"), Validate({uint8(ECbFieldType::Null | ECbFieldType::HasFieldType)}), ECbValidateError::InvalidType);
+	TEST_EQUAL(TEXT("ValidateCompactBinary(InvalidType, Unknown)"), Validate({uint8(ECbFieldType::ObjectId) + 1}), ECbValidateError::InvalidType);
+	TEST_EQUAL(TEXT("ValidateCompactBinary(InvalidType, Unknown)"), Validate({}, ECbFieldType(uint8(ECbFieldType::ObjectId) + 1)), ECbValidateError::InvalidType);
+	TEST_EQUAL(TEXT("ValidateCompactBinary(InvalidType, HasFieldType)"), Validate({uint8(ECbFieldType::Null | ECbFieldType::HasFieldType)}), ECbValidateError::InvalidType);
 
-	TestEqual(TEXT("ValidateCompactBinary(InvalidType, ZeroSizeField)"), Validate({}, ECbFieldType::Null), ECbValidateError::InvalidType);
-	TestEqual(TEXT("ValidateCompactBinary(InvalidType, ZeroSizeField, BoolFalse)"), Validate({}, ECbFieldType::BoolFalse), ECbValidateError::InvalidType);
-	TestEqual(TEXT("ValidateCompactBinary(InvalidType, ZeroSizeField, BoolTrue)"), Validate({}, ECbFieldType::BoolTrue), ECbValidateError::InvalidType);
+	TEST_EQUAL(TEXT("ValidateCompactBinary(InvalidType, ZeroSizeField)"), Validate({}, ECbFieldType::Null), ECbValidateError::InvalidType);
+	TEST_EQUAL(TEXT("ValidateCompactBinary(InvalidType, ZeroSizeField, BoolFalse)"), Validate({}, ECbFieldType::BoolFalse), ECbValidateError::InvalidType);
+	TEST_EQUAL(TEXT("ValidateCompactBinary(InvalidType, ZeroSizeField, BoolTrue)"), Validate({}, ECbFieldType::BoolTrue), ECbValidateError::InvalidType);
 
-	TestEqual(TEXT("ValidateCompactBinary(InvalidType, ZeroSizeField, Array)"), Validate({uint8(ECbFieldType::UniformArray), 2, 2, NullNoName}), ECbValidateError::InvalidType);
-	TestEqual(TEXT("ValidateCompactBinary(InvalidType, ZeroSizeField, Object)"), Validate({uint8(ECbFieldType::UniformObject), 2, NullNoName, 0}), ECbValidateError::InvalidType);
+	TEST_EQUAL(TEXT("ValidateCompactBinary(InvalidType, ZeroSizeField, Array)"), Validate({uint8(ECbFieldType::UniformArray), 2, 2, NullNoName}), ECbValidateError::InvalidType);
+	TEST_EQUAL(TEXT("ValidateCompactBinary(InvalidType, ZeroSizeField, Object)"), Validate({uint8(ECbFieldType::UniformObject), 2, NullNoName, 0}), ECbValidateError::InvalidType);
 
-	TestEqual(TEXT("ValidateCompactBinary(InvalidType, CustomByName, EmptyTypeName)"), Validate({uint8(ECbFieldType::CustomByName), 1, 0}), ECbValidateError::InvalidType);
-	TestEqual(TEXT("ValidateCompactBinary(InvalidType, CustomByName, EmptyTypeName, NoType)"), Validate({1, 0}, ECbFieldType::CustomByName), ECbValidateError::InvalidType);
+	TEST_EQUAL(TEXT("ValidateCompactBinary(InvalidType, CustomByName, EmptyTypeName)"), Validate({uint8(ECbFieldType::CustomByName), 1, 0}), ECbValidateError::InvalidType);
+	TEST_EQUAL(TEXT("ValidateCompactBinary(InvalidType, CustomByName, EmptyTypeName, NoType)"), Validate({1, 0}, ECbFieldType::CustomByName), ECbValidateError::InvalidType);
 
 	// Test DuplicateName
-	TestEqual(TEXT("ValidateCompactBinary(DuplicateName)"), Validate({uint8(ECbFieldType::UniformObject), 7, NullWithName, 1, 'A', 1, 'B', 1, 'A'}), ECbValidateError::DuplicateName);
-	TestEqual(TEXT("ValidateCompactBinary(DuplicateName, CaseSensitive)"), Validate({uint8(ECbFieldType::UniformObject), 7, NullWithName, 1, 'A', 1, 'B', 1, 'a'}), ECbValidateError::None);
-	TestEqual(TEXT("ValidateCompactBinary(DuplicateName, Mode)"), ValidateMode({uint8(ECbFieldType::UniformObject), 7, NullWithName, 1, 'A', 1, 'B', 1, 'A'}, ECbValidateMode::All & ~ECbValidateMode::Names), ECbValidateError::None);
+	TEST_EQUAL(TEXT("ValidateCompactBinary(DuplicateName)"), Validate({uint8(ECbFieldType::UniformObject), 7, NullWithName, 1, 'A', 1, 'B', 1, 'A'}), ECbValidateError::DuplicateName);
+	TEST_EQUAL(TEXT("ValidateCompactBinary(DuplicateName, CaseSensitive)"), Validate({uint8(ECbFieldType::UniformObject), 7, NullWithName, 1, 'A', 1, 'B', 1, 'a'}), ECbValidateError::None);
+	TEST_EQUAL(TEXT("ValidateCompactBinary(DuplicateName, Mode)"), ValidateMode({uint8(ECbFieldType::UniformObject), 7, NullWithName, 1, 'A', 1, 'B', 1, 'A'}, ECbValidateMode::All & ~ECbValidateMode::Names), ECbValidateError::None);
 
 	// Test MissingName
-	TestEqual(TEXT("ValidateCompactBinary(MissingName)"), Validate({uint8(ECbFieldType::Object), 3, NullNoName, IntNoName, 0}), ECbValidateError::MissingName);
-	TestEqual(TEXT("ValidateCompactBinary(MissingName, Uniform)"), Validate({uint8(ECbFieldType::UniformObject), 3, IntNoName, 0, 0}), ECbValidateError::MissingName);
-	TestEqual(TEXT("ValidateCompactBinary(MissingName, Mode)"), ValidateMode({uint8(ECbFieldType::Object), 3, NullNoName, IntNoName, 0}, ECbValidateMode::All & ~ECbValidateMode::Names), ECbValidateError::None);
-	TestEqual(TEXT("ValidateCompactBinary(MissingName, Uniform, Mode)"), ValidateMode({uint8(ECbFieldType::UniformObject), 3, IntNoName, 0, 0}, ECbValidateMode::All & ~ECbValidateMode::Names), ECbValidateError::None);
+	TEST_EQUAL(TEXT("ValidateCompactBinary(MissingName)"), Validate({uint8(ECbFieldType::Object), 3, NullNoName, IntNoName, 0}), ECbValidateError::MissingName);
+	TEST_EQUAL(TEXT("ValidateCompactBinary(MissingName, Uniform)"), Validate({uint8(ECbFieldType::UniformObject), 3, IntNoName, 0, 0}), ECbValidateError::MissingName);
+	TEST_EQUAL(TEXT("ValidateCompactBinary(MissingName, Mode)"), ValidateMode({uint8(ECbFieldType::Object), 3, NullNoName, IntNoName, 0}, ECbValidateMode::All & ~ECbValidateMode::Names), ECbValidateError::None);
+	TEST_EQUAL(TEXT("ValidateCompactBinary(MissingName, Uniform, Mode)"), ValidateMode({uint8(ECbFieldType::UniformObject), 3, IntNoName, 0, 0}, ECbValidateMode::All & ~ECbValidateMode::Names), ECbValidateError::None);
 
 	// Test ArrayName
-	TestEqual(TEXT("ValidateCompactBinary(ArrayName)"), Validate({uint8(ECbFieldType::Array), 5, 2, NullNoName, NullWithName, 1, 'F'}), ECbValidateError::ArrayName);
-	TestEqual(TEXT("ValidateCompactBinary(ArrayName, Uniform)"), Validate({uint8(ECbFieldType::UniformArray), 4, 1, NullWithName, 1, 'F'}), ECbValidateError::ArrayName);
-	TestEqual(TEXT("ValidateCompactBinary(ArrayName, Mode)"), ValidateMode({uint8(ECbFieldType::Array), 5, 2, NullNoName, NullWithName, 1, 'F'}, ECbValidateMode::All & ~ECbValidateMode::Names), ECbValidateError::None);
-	TestEqual(TEXT("ValidateCompactBinary(ArrayName, Uniform, Mode)"), ValidateMode({uint8(ECbFieldType::UniformArray), 4, 1, NullWithName, 1, 'F'}, ECbValidateMode::All & ~ECbValidateMode::Names), ECbValidateError::None);
+	TEST_EQUAL(TEXT("ValidateCompactBinary(ArrayName)"), Validate({uint8(ECbFieldType::Array), 5, 2, NullNoName, NullWithName, 1, 'F'}), ECbValidateError::ArrayName);
+	TEST_EQUAL(TEXT("ValidateCompactBinary(ArrayName, Uniform)"), Validate({uint8(ECbFieldType::UniformArray), 4, 1, NullWithName, 1, 'F'}), ECbValidateError::ArrayName);
+	TEST_EQUAL(TEXT("ValidateCompactBinary(ArrayName, Mode)"), ValidateMode({uint8(ECbFieldType::Array), 5, 2, NullNoName, NullWithName, 1, 'F'}, ECbValidateMode::All & ~ECbValidateMode::Names), ECbValidateError::None);
+	TEST_EQUAL(TEXT("ValidateCompactBinary(ArrayName, Uniform, Mode)"), ValidateMode({uint8(ECbFieldType::UniformArray), 4, 1, NullWithName, 1, 'F'}, ECbValidateMode::All & ~ECbValidateMode::Names), ECbValidateError::None);
 
 	// Test InvalidString
 	// Not tested or implemented yet because the engine does not provide enough UTF-8 functionality.
 
 	// Test InvalidInteger
-	TestEqual(TEXT("ValidateCompactBinary(InvalidInteger, NameSize)"), Validate({NullWithName, 0x80, 1, 'N'}), ECbValidateError::InvalidInteger);
-	TestEqual(TEXT("ValidateCompactBinary(InvalidInteger, ObjectSize)"), Validate({uint8(ECbFieldType::Object), 0xc0, 0, 0}), ECbValidateError::InvalidInteger);
-	TestEqual(TEXT("ValidateCompactBinary(InvalidInteger, ArraySize)"), Validate({uint8(ECbFieldType::Array), 0xe0, 0, 0, 1, 0}), ECbValidateError::InvalidInteger);
-	TestEqual(TEXT("ValidateCompactBinary(InvalidInteger, ArrayCount)"), Validate({uint8(ECbFieldType::Array), 5, 0xf0, 0, 0, 0, 0}), ECbValidateError::InvalidInteger);
-	TestEqual(TEXT("ValidateCompactBinary(InvalidInteger, BinarySize)"), Validate({uint8(ECbFieldType::Binary), 0xf8, 0, 0, 0, 0, 0}), ECbValidateError::InvalidInteger);
-	TestEqual(TEXT("ValidateCompactBinary(InvalidInteger, StringSize)"), Validate({uint8(ECbFieldType::String), 0xfc, 0, 0, 0, 0, 0, 0}), ECbValidateError::InvalidInteger);
-	TestEqual(TEXT("ValidateCompactBinary(InvalidInteger, IntegerPositive)"), Validate({uint8(ECbFieldType::IntegerPositive), 0xfe, 0, 0, 0, 0, 0, 0, 0}), ECbValidateError::InvalidInteger);
-	TestEqual(TEXT("ValidateCompactBinary(InvalidInteger, IntegerNegative)"), Validate({uint8(ECbFieldType::IntegerNegative), 0xff, 0, 0, 0, 0, 0, 0, 0, 0}), ECbValidateError::InvalidInteger);
-	TestEqual(TEXT("ValidateCompactBinary(InvalidInteger, ArraySize)"), Validate({uint8(ECbFieldType::Array), 0x80, 1, 0}), ECbValidateError::InvalidInteger);
-	TestEqual(TEXT("ValidateCompactBinary(InvalidInteger, ArrayCount)"), Validate({uint8(ECbFieldType::Array), 3, 0xc0, 0, 0}), ECbValidateError::InvalidInteger);
-	TestEqual(TEXT("ValidateCompactBinary(InvalidInteger, ObjectSize)"), Validate({uint8(ECbFieldType::Object), 0xe0, 0, 0, 0}), ECbValidateError::InvalidInteger);
-	TestEqual(TEXT("ValidateCompactBinary(InvalidInteger, CustomById)"), Validate({uint8(ECbFieldType::CustomById), 2, 0x80, 0}), ECbValidateError::InvalidInteger);
-	TestEqual(TEXT("ValidateCompactBinary(InvalidInteger, CustomByName)"), Validate({uint8(ECbFieldType::CustomByName), 3, 0x80, 1, 'A'}), ECbValidateError::InvalidInteger);
-
-	TestEqual(TEXT("ValidateCompactBinary(InvalidInteger, NameSize, Mode)"), ValidateMode({NullWithName, 0x80, 1, 'N'}, ECbValidateMode::All & ~ECbValidateMode::Format), ECbValidateError::None);
-	TestEqual(TEXT("ValidateCompactBinary(InvalidInteger, ArraySize, Mode)"), ValidateMode({uint8(ECbFieldType::Array), 0xc0, 0, 1, 0}, ECbValidateMode::All & ~ECbValidateMode::Format), ECbValidateError::None);
-	TestEqual(TEXT("ValidateCompactBinary(InvalidInteger, ObjectSize, Mode)"), ValidateMode({uint8(ECbFieldType::Object), 0xe0, 0, 0, 0}, ECbValidateMode::All & ~ECbValidateMode::Format), ECbValidateError::None);
-	TestEqual(TEXT("ValidateCompactBinary(InvalidInteger, CustomById, Mode)"), ValidateMode({uint8(ECbFieldType::CustomById), 2, 0x80, 0}, ECbValidateMode::All & ~ECbValidateMode::Format), ECbValidateError::None);
-	TestEqual(TEXT("ValidateCompactBinary(InvalidInteger, CustomByName, Mode)"), ValidateMode({uint8(ECbFieldType::CustomByName), 3, 0x80, 1, 'A'}, ECbValidateMode::All & ~ECbValidateMode::Format), ECbValidateError::None);
+	TEST_EQUAL(TEXT("ValidateCompactBinary(InvalidInteger, NameSize)"), Validate({NullWithName, 0x80, 1, 'N'}), ECbValidateError::InvalidInteger);
+	TEST_EQUAL(TEXT("ValidateCompactBinary(InvalidInteger, ObjectSize)"), Validate({uint8(ECbFieldType::Object), 0xc0, 0, 0}), ECbValidateError::InvalidInteger);
+	TEST_EQUAL(TEXT("ValidateCompactBinary(InvalidInteger, ArraySize)"), Validate({uint8(ECbFieldType::Array), 0xe0, 0, 0, 1, 0}), ECbValidateError::InvalidInteger);
+	TEST_EQUAL(TEXT("ValidateCompactBinary(InvalidInteger, ArrayCount)"), Validate({uint8(ECbFieldType::Array), 5, 0xf0, 0, 0, 0, 0}), ECbValidateError::InvalidInteger);
+	TEST_EQUAL(TEXT("ValidateCompactBinary(InvalidInteger, BinarySize)"), Validate({uint8(ECbFieldType::Binary), 0xf8, 0, 0, 0, 0, 0}), ECbValidateError::InvalidInteger);
+	TEST_EQUAL(TEXT("ValidateCompactBinary(InvalidInteger, StringSize)"), Validate({uint8(ECbFieldType::String), 0xfc, 0, 0, 0, 0, 0, 0}), ECbValidateError::InvalidInteger);
+	TEST_EQUAL(TEXT("ValidateCompactBinary(InvalidInteger, IntegerPositive)"), Validate({uint8(ECbFieldType::IntegerPositive), 0xfe, 0, 0, 0, 0, 0, 0, 0}), ECbValidateError::InvalidInteger);
+	TEST_EQUAL(TEXT("ValidateCompactBinary(InvalidInteger, IntegerNegative)"), Validate({uint8(ECbFieldType::IntegerNegative), 0xff, 0, 0, 0, 0, 0, 0, 0, 0}), ECbValidateError::InvalidInteger);
+	TEST_EQUAL(TEXT("ValidateCompactBinary(InvalidInteger, ArraySize)"), Validate({uint8(ECbFieldType::Array), 0x80, 1, 0}), ECbValidateError::InvalidInteger);
+	TEST_EQUAL(TEXT("ValidateCompactBinary(InvalidInteger, ArrayCount)"), Validate({uint8(ECbFieldType::Array), 3, 0xc0, 0, 0}), ECbValidateError::InvalidInteger);
+	TEST_EQUAL(TEXT("ValidateCompactBinary(InvalidInteger, ObjectSize)"), Validate({uint8(ECbFieldType::Object), 0xe0, 0, 0, 0}), ECbValidateError::InvalidInteger);
+	TEST_EQUAL(TEXT("ValidateCompactBinary(InvalidInteger, CustomById)"), Validate({uint8(ECbFieldType::CustomById), 2, 0x80, 0}), ECbValidateError::InvalidInteger);
+	TEST_EQUAL(TEXT("ValidateCompactBinary(InvalidInteger, CustomByName)"), Validate({uint8(ECbFieldType::CustomByName), 3, 0x80, 1, 'A'}), ECbValidateError::InvalidInteger);
+
+	TEST_EQUAL(TEXT("ValidateCompactBinary(InvalidInteger, NameSize, Mode)"), ValidateMode({NullWithName, 0x80, 1, 'N'}, ECbValidateMode::All & ~ECbValidateMode::Format), ECbValidateError::None);
+	TEST_EQUAL(TEXT("ValidateCompactBinary(InvalidInteger, ArraySize, Mode)"), ValidateMode({uint8(ECbFieldType::Array), 0xc0, 0, 1, 0}, ECbValidateMode::All & ~ECbValidateMode::Format), ECbValidateError::None);
+	TEST_EQUAL(TEXT("ValidateCompactBinary(InvalidInteger, ObjectSize, Mode)"), ValidateMode({uint8(ECbFieldType::Object), 0xe0, 0, 0, 0}, ECbValidateMode::All & ~ECbValidateMode::Format), ECbValidateError::None);
+	TEST_EQUAL(TEXT("ValidateCompactBinary(InvalidInteger, CustomById, Mode)"), ValidateMode({uint8(ECbFieldType::CustomById), 2, 0x80, 0}, ECbValidateMode::All & ~ECbValidateMode::Format), ECbValidateError::None);
+	TEST_EQUAL(TEXT("ValidateCompactBinary(InvalidInteger, CustomByName, Mode)"), ValidateMode({uint8(ECbFieldType::CustomByName), 3, 0x80, 1, 'A'}, ECbValidateMode::All & ~ECbValidateMode::Format), ECbValidateError::None);
 
 	// Test InvalidFloat
-	TestEqual(TEXT("ValidateCompactBinary(InvalidFloat, MaxSignificant+1)"), Validate({uint8(ECbFieldType::Float64), 0x3f, 0xff, 0xff, 0xff, 0xf0, 0x00, 0x00, 0x00}), ECbValidateError::None); // 1.9999999403953552
-	TestEqual(TEXT("ValidateCompactBinary(InvalidFloat, MaxExponent+1)"), Validate({uint8(ECbFieldType::Float64), 0x47, 0xff, 0xff, 0xff, 0xe0, 0x00, 0x00, 0x00}), ECbValidateError::None); // 6.8056469327705771e38
-	TestEqual(TEXT("ValidateCompactBinary(InvalidFloat, MaxSignificand)"), Validate({uint8(ECbFieldType::Float64), 0x3f, 0xff, 0xff, 0xff, 0xe0, 0x00, 0x00, 0x00}), ECbValidateError::InvalidFloat); // 1.9999998807907104
-	TestEqual(TEXT("ValidateCompactBinary(InvalidFloat, MaxExponent)"), Validate({uint8(ECbFieldType::Float64), 0x47, 0xef, 0xff, 0xff, 0xe0, 0x00, 0x00, 0x00}), ECbValidateError::InvalidFloat); // 3.4028234663852886e38
-	TestEqual(TEXT("ValidateCompactBinary(InvalidFloat, MaxSignificand, Mode)"), ValidateMode({uint8(ECbFieldType::Float64), 0x3f, 0xff, 0xff, 0xff, 0xe0, 0x00, 0x00, 0x00}, ECbValidateMode::All & ~ECbValidateMode::Format), ECbValidateError::None); // 1.9999998807907104
-	TestEqual(TEXT("ValidateCompactBinary(InvalidFloat, MaxExponent, Mode)"), ValidateMode({uint8(ECbFieldType::Float64), 0x47, 0xef, 0xff, 0xff, 0xe0, 0x00, 0x00, 0x00}, ECbValidateMode::All & ~ECbValidateMode::Format), ECbValidateError::None); // 3.4028234663852886e38
+	TEST_EQUAL(TEXT("ValidateCompactBinary(InvalidFloat, MaxSignificant+1)"), Validate({uint8(ECbFieldType::Float64), 0x3f, 0xff, 0xff, 0xff, 0xf0, 0x00, 0x00, 0x00}), ECbValidateError::None); // 1.9999999403953552
+	TEST_EQUAL(TEXT("ValidateCompactBinary(InvalidFloat, MaxExponent+1)"), Validate({uint8(ECbFieldType::Float64), 0x47, 0xff, 0xff, 0xff, 0xe0, 0x00, 0x00, 0x00}), ECbValidateError::None); // 6.8056469327705771e38
+	TEST_EQUAL(TEXT("ValidateCompactBinary(InvalidFloat, MaxSignificand)"), Validate({uint8(ECbFieldType::Float64), 0x3f, 0xff, 0xff, 0xff, 0xe0, 0x00, 0x00, 0x00}), ECbValidateError::InvalidFloat); // 1.9999998807907104
+	TEST_EQUAL(TEXT("ValidateCompactBinary(InvalidFloat, MaxExponent)"), Validate({uint8(ECbFieldType::Float64), 0x47, 0xef, 0xff, 0xff, 0xe0, 0x00, 0x00, 0x00}), ECbValidateError::InvalidFloat); // 3.4028234663852886e38
+	TEST_EQUAL(TEXT("ValidateCompactBinary(InvalidFloat, MaxSignificand, Mode)"), ValidateMode({uint8(ECbFieldType::Float64), 0x3f, 0xff, 0xff, 0xff, 0xe0, 0x00, 0x00, 0x00}, ECbValidateMode::All & ~ECbValidateMode::Format), ECbValidateError::None); // 1.9999998807907104
+	TEST_EQUAL(TEXT("ValidateCompactBinary(InvalidFloat, MaxExponent, Mode)"), ValidateMode({uint8(ECbFieldType::Float64), 0x47, 0xef, 0xff, 0xff, 0xe0, 0x00, 0x00, 0x00}, ECbValidateMode::All & ~ECbValidateMode::Format), ECbValidateError::None); // 3.4028234663852886e38
 
 	// Test NonUniformObject
-	TestEqual(TEXT("ValidateCompactBinary(NonUniformObject)"), Validate({uint8(ECbFieldType::Object), 3, NullWithName, 1, 'A'}), ECbValidateError::NonUniformObject);
-	TestEqual(TEXT("ValidateCompactBinary(NonUniformObject)"), Validate({uint8(ECbFieldType::Object), 6, NullWithName, 1, 'A', NullWithName, 1, 'B'}), ECbValidateError::NonUniformObject);
-	TestEqual(TEXT("ValidateCompactBinary(NonUniformObject, Mode)"), ValidateMode({uint8(ECbFieldType::Object), 3, NullWithName, 1, 'A'}, ECbValidateMode::All & ~ECbValidateMode::Format), ECbValidateError::None);
-	TestEqual(TEXT("ValidateCompactBinary(NonUniformObject, Mode)"), ValidateMode({uint8(ECbFieldType::Object), 6, NullWithName, 1, 'A', NullWithName, 1, 'B'}, ECbValidateMode::All & ~ECbValidateMode::Format), ECbValidateError::None);
+	TEST_EQUAL(TEXT("ValidateCompactBinary(NonUniformObject)"), Validate({uint8(ECbFieldType::Object), 3, NullWithName, 1, 'A'}), ECbValidateError::NonUniformObject);
+	TEST_EQUAL(TEXT("ValidateCompactBinary(NonUniformObject)"), Validate({uint8(ECbFieldType::Object), 6, NullWithName, 1, 'A', NullWithName, 1, 'B'}), ECbValidateError::NonUniformObject);
+	TEST_EQUAL(TEXT("ValidateCompactBinary(NonUniformObject, Mode)"), ValidateMode({uint8(ECbFieldType::Object), 3, NullWithName, 1, 'A'}, ECbValidateMode::All & ~ECbValidateMode::Format), ECbValidateError::None);
+	TEST_EQUAL(TEXT("ValidateCompactBinary(NonUniformObject, Mode)"), ValidateMode({uint8(ECbFieldType::Object), 6, NullWithName, 1, 'A', NullWithName, 1, 'B'}, ECbValidateMode::All & ~ECbValidateMode::Format), ECbValidateError::None);
 
 	// Test NonUniformArray
-	TestEqual(TEXT("ValidateCompactBinary(NonUniformArray)"), Validate({uint8(ECbFieldType::Array), 3, 1, IntNoName, 0}), ECbValidateError::NonUniformArray);
-	TestEqual(TEXT("ValidateCompactBinary(NonUniformArray)"), Validate({uint8(ECbFieldType::Array), 5, 2, IntNoName, 1, IntNoName, 2}), ECbValidateError::NonUniformArray);
-	TestEqual(TEXT("ValidateCompactBinary(NonUniformArray, Null)"), Validate({uint8(ECbFieldType::Array), 3, 2, NullNoName, NullNoName}), ECbValidateError::None);
-	TestEqual(TEXT("ValidateCompactBinary(NonUniformArray, Bool)"), Validate({uint8(ECbFieldType::Array), 3, 2, uint8(ECbFieldType::BoolFalse), uint8(ECbFieldType::BoolFalse)}), ECbValidateError::None);
-	TestEqual(TEXT("ValidateCompactBinary(NonUniformArray, Bool)"), Validate({uint8(ECbFieldType::Array), 3, 2, uint8(ECbFieldType::BoolTrue), uint8(ECbFieldType::BoolTrue)}), ECbValidateError::None);
-	TestEqual(TEXT("ValidateCompactBinary(NonUniformArray, Mode)"), ValidateMode({uint8(ECbFieldType::Array), 3, 1, IntNoName, 0}, ECbValidateMode::All & ~ECbValidateMode::Format), ECbValidateError::None);
-	TestEqual(TEXT("ValidateCompactBinary(NonUniformArray, Mode)"), ValidateMode({uint8(ECbFieldType::Array), 5, 2, IntNoName, 1, IntNoName, 2}, ECbValidateMode::All & ~ECbValidateMode::Format), ECbValidateError::None);
+	TEST_EQUAL(TEXT("ValidateCompactBinary(NonUniformArray)"), Validate({uint8(ECbFieldType::Array), 3, 1, IntNoName, 0}), ECbValidateError::NonUniformArray);
+	TEST_EQUAL(TEXT("ValidateCompactBinary(NonUniformArray)"), Validate({uint8(ECbFieldType::Array), 5, 2, IntNoName, 1, IntNoName, 2}), ECbValidateError::NonUniformArray);
+	TEST_EQUAL(TEXT("ValidateCompactBinary(NonUniformArray, Null)"), Validate({uint8(ECbFieldType::Array), 3, 2, NullNoName, NullNoName}), ECbValidateError::None);
+	TEST_EQUAL(TEXT("ValidateCompactBinary(NonUniformArray, Bool)"), Validate({uint8(ECbFieldType::Array), 3, 2, uint8(ECbFieldType::BoolFalse), uint8(ECbFieldType::BoolFalse)}), ECbValidateError::None);
+	TEST_EQUAL(TEXT("ValidateCompactBinary(NonUniformArray, Bool)"), Validate({uint8(ECbFieldType::Array), 3, 2, uint8(ECbFieldType::BoolTrue), uint8(ECbFieldType::BoolTrue)}), ECbValidateError::None);
+	TEST_EQUAL(TEXT("ValidateCompactBinary(NonUniformArray, Mode)"), ValidateMode({uint8(ECbFieldType::Array), 3, 1, IntNoName, 0}, ECbValidateMode::All & ~ECbValidateMode::Format), ECbValidateError::None);
+	TEST_EQUAL(TEXT("ValidateCompactBinary(NonUniformArray, Mode)"), ValidateMode({uint8(ECbFieldType::Array), 5, 2, IntNoName, 1, IntNoName, 2}, ECbValidateMode::All & ~ECbValidateMode::Format), ECbValidateError::None);
 
 	// Test Padding
-	TestEqual(TEXT("ValidateCompactBinary(Padding)"), Validate({NullNoName, 0}), ECbValidateError::Padding);
-	TestEqual(TEXT("ValidateCompactBinary(Padding)"), Validate({uint8(ECbFieldType::Array), 1, 0, 0}), ECbValidateError::Padding);
-	TestEqual(TEXT("ValidateCompactBinary(Padding)"), Validate({uint8(ECbFieldType::Object), 0, 0}), ECbValidateError::Padding);
-	TestEqual(TEXT("ValidateCompactBinary(Padding, Mode)"), ValidateMode({NullNoName, 0}, ECbValidateMode::All & ~ECbValidateMode::Padding), ECbValidateError::None);
-	TestEqual(TEXT("ValidateCompactBinary(Padding, Mode)"), ValidateMode({uint8(ECbFieldType::Array), 1, 0, 0}, ECbValidateMode::All & ~ECbValidateMode::Padding), ECbValidateError::None);
-	TestEqual(TEXT("ValidateCompactBinary(Padding, Mode)"), ValidateMode({uint8(ECbFieldType::Object), 0, 0}, ECbValidateMode::All & ~ECbValidateMode::Padding), ECbValidateError::None);
+	TEST_EQUAL(TEXT("ValidateCompactBinary(Padding)"), Validate({NullNoName, 0}), ECbValidateError::Padding);
+	TEST_EQUAL(TEXT("ValidateCompactBinary(Padding)"), Validate({uint8(ECbFieldType::Array), 1, 0, 0}), ECbValidateError::Padding);
+	TEST_EQUAL(TEXT("ValidateCompactBinary(Padding)"), Validate({uint8(ECbFieldType::Object), 0, 0}), ECbValidateError::Padding);
+	TEST_EQUAL(TEXT("ValidateCompactBinary(Padding, Mode)"), ValidateMode({NullNoName, 0}, ECbValidateMode::All & ~ECbValidateMode::Padding), ECbValidateError::None);
+	TEST_EQUAL(TEXT("ValidateCompactBinary(Padding, Mode)"), ValidateMode({uint8(ECbFieldType::Array), 1, 0, 0}, ECbValidateMode::All & ~ECbValidateMode::Padding), ECbValidateError::None);
+	TEST_EQUAL(TEXT("ValidateCompactBinary(Padding, Mode)"), ValidateMode({uint8(ECbFieldType::Object), 0, 0}, ECbValidateMode::All & ~ECbValidateMode::Padding), ECbValidateError::None);
 }
 
-TEST_CASE("Core::Serialization::FCbField::Validate Compact Binary Range", "[Core][Serialization][Smoke]")
+TEST_CASE_METHOD(FAutomationTestFixture, "Core::Serialization::FCbField::Validate Compact Binary Range", "[Core][Serialization][Smoke]")
 {
 	auto Validate = [](std::initializer_list<uint8> Data, ECbValidateMode Mode = ECbValidateMode::All) -> ECbValidateError
 	{
@@ -297,19 +297,19 @@ TEST_CASE("Core::Serialization::FCbField::Validate Compact Binary Range", "[Core
 	};
 
 	// Test Empty
-	TestEqual(TEXT("ValidateCompactBinaryRange(Empty)"), Validate({}), ECbValidateError::None);
+	TEST_EQUAL(TEXT("ValidateCompactBinaryRange(Empty)"), Validate({}), ECbValidateError::None);
 
 	// Test Valid
-	TestEqual(TEXT("ValidateCompactBinaryRange(Null x2)"), Validate({uint8(ECbFieldType::Null), uint8(ECbFieldType::Null)}), ECbValidateError::None);
+	TEST_EQUAL(TEXT("ValidateCompactBinaryRange(Null x2)"), Validate({uint8(ECbFieldType::Null), uint8(ECbFieldType::Null)}), ECbValidateError::None);
 
 	// Test Padding
-	TestEqual(TEXT("ValidateCompactBinaryRange(Padding InvalidType)"), Validate({uint8(ECbFieldType::Null), 0}), ECbValidateError::InvalidType);
-	TestEqual(TEXT("ValidateCompactBinaryRange(Padding OutOfBounds)"), Validate({uint8(ECbFieldType::Null), uint8(ECbFieldType::Binary)}), ECbValidateError::OutOfBounds);
-	TestEqual(TEXT("ValidateCompactBinaryRange(Padding OutOfBounds, Mode)"), Validate({uint8(ECbFieldType::Null), uint8(ECbFieldType::Binary)}, ECbValidateMode::None), ECbValidateError::None);
+	TEST_EQUAL(TEXT("ValidateCompactBinaryRange(Padding InvalidType)"), Validate({uint8(ECbFieldType::Null), 0}), ECbValidateError::InvalidType);
+	TEST_EQUAL(TEXT("ValidateCompactBinaryRange(Padding OutOfBounds)"), Validate({uint8(ECbFieldType::Null), uint8(ECbFieldType::Binary)}), ECbValidateError::OutOfBounds);
+	TEST_EQUAL(TEXT("ValidateCompactBinaryRange(Padding OutOfBounds, Mode)"), Validate({uint8(ECbFieldType::Null), uint8(ECbFieldType::Binary)}, ECbValidateMode::None), ECbValidateError::None);
 }
 
 
-TEST_CASE("Core::Serialization::FCbField::Validate Compact Binary Attachment", "[Core][Serialization][Smoke]")
+TEST_CASE_METHOD(FAutomationTestFixture, "Core::Serialization::FCbField::Validate Compact Binary Attachment", "[Core][Serialization][Smoke]")
 {
 	auto Validate = [](std::initializer_list<uint8> Data, ECbValidateMode Mode = ECbValidateMode::All) -> ECbValidateError
 	{
@@ -321,27 +321,27 @@ TEST_CASE("Core::Serialization::FCbField::Validate Compact Binary Attachment", "
 
 	// Test Null
 	{
-		TestEqual(TEXT("ValidateCompactBinaryAttachment(Null)"), Validate({uint8(ECbFieldType::Binary), 0}, ECbValidateMode::All & ~ECbValidateMode::Package), ECbValidateError::None);
-		TestEqual(TEXT("ValidateCompactBinaryAttachment(Null, Padding)"), Validate({uint8(ECbFieldType::Binary), 0, 0}, ECbValidateMode::All & ~ECbValidateMode::Package), ECbValidateError::Padding);
-		TestEqual(TEXT("ValidateCompactBinaryAttachment(Null, Padding, Mode)"), Validate({uint8(ECbFieldType::Binary), 0, 0}, ECbValidateMode::All & ~(ECbValidateMode::Padding | ECbValidateMode::Package)), ECbValidateError::None);
+		TEST_EQUAL(TEXT("ValidateCompactBinaryAttachment(Null)"), Validate({uint8(ECbFieldType::Binary), 0}, ECbValidateMode::All & ~ECbValidateMode::Package), ECbValidateError::None);
+		TEST_EQUAL(TEXT("ValidateCompactBinaryAttachment(Null, Padding)"), Validate({uint8(ECbFieldType::Binary), 0, 0}, ECbValidateMode::All & ~ECbValidateMode::Package), ECbValidateError::Padding);
+		TEST_EQUAL(TEXT("ValidateCompactBinaryAttachment(Null, Padding, Mode)"), Validate({uint8(ECbFieldType::Binary), 0, 0}, ECbValidateMode::All & ~(ECbValidateMode::Padding | ECbValidateMode::Package)), ECbValidateError::None);
 	}
 
 	// Test Binary
 	{
 		FBufferArchive Buffer;
 		FCbAttachment(FSharedBuffer::MakeView(MakeMemoryView<uint8>({0}))).Save(Buffer);
-		TestEqual(TEXT("ValidateCompactBinaryAttachment(Binary)"), ValidateCompactBinaryAttachment(MakeMemoryView(Buffer), ECbValidateMode::All), ECbValidateError::None);
+		TEST_EQUAL(TEXT("ValidateCompactBinaryAttachment(Binary)"), ValidateCompactBinaryAttachment(MakeMemoryView(Buffer), ECbValidateMode::All), ECbValidateError::None);
 		++Buffer[2];
-		TestEqual(TEXT("ValidateCompactBinaryAttachment(Binary, InvalidHash)"), ValidateCompactBinaryAttachment(MakeMemoryView(Buffer), ECbValidateMode::All), ECbValidateError::InvalidPackageHash);
-		TestEqual(TEXT("ValidateCompactBinaryAttachment(Binary, InvalidHash, Mode)"), ValidateCompactBinaryAttachment(MakeMemoryView(Buffer), ECbValidateMode::All & ~ECbValidateMode::PackageHash), ECbValidateError::None);
+		TEST_EQUAL(TEXT("ValidateCompactBinaryAttachment(Binary, InvalidHash)"), ValidateCompactBinaryAttachment(MakeMemoryView(Buffer), ECbValidateMode::All), ECbValidateError::InvalidPackageHash);
+		TEST_EQUAL(TEXT("ValidateCompactBinaryAttachment(Binary, InvalidHash, Mode)"), ValidateCompactBinaryAttachment(MakeMemoryView(Buffer), ECbValidateMode::All & ~ECbValidateMode::PackageHash), ECbValidateError::None);
 		--Buffer[2];
 		Buffer.AddZeroed(1);
-		TestEqual(TEXT("ValidateCompactBinaryAttachment(Binary, Padding)"), ValidateCompactBinaryAttachment(MakeMemoryView(Buffer), ECbValidateMode::All), ECbValidateError::Padding);
-		TestEqual(TEXT("ValidateCompactBinaryAttachment(Binary, Padding, Mode)"), ValidateCompactBinaryAttachment(MakeMemoryView(Buffer), ECbValidateMode::All & ~ECbValidateMode::Padding), ECbValidateError::None);
+		TEST_EQUAL(TEXT("ValidateCompactBinaryAttachment(Binary, Padding)"), ValidateCompactBinaryAttachment(MakeMemoryView(Buffer), ECbValidateMode::All), ECbValidateError::Padding);
+		TEST_EQUAL(TEXT("ValidateCompactBinaryAttachment(Binary, Padding, Mode)"), ValidateCompactBinaryAttachment(MakeMemoryView(Buffer), ECbValidateMode::All & ~ECbValidateMode::Padding), ECbValidateError::None);
 		Buffer.Pop(/*bAllowShrinking*/ false);
-		TestEqual(TEXT("ValidateCompactBinaryAttachment(Binary, OutOfBoundsValue)"), ValidateCompactBinaryAttachment(MakeMemoryView(Buffer).Left(1), ECbValidateMode::All), ECbValidateError::OutOfBounds);
-		TestEqual(TEXT("ValidateCompactBinaryAttachment(Binary, OutOfBoundsHash)"), ValidateCompactBinaryAttachment(MakeMemoryView(Buffer).Left(4), ECbValidateMode::All), ECbValidateError::OutOfBounds);
-		TestEqual(TEXT("ValidateCompactBinaryAttachment(Binary, OutOfBoundsHash, Mode)"), ValidateCompactBinaryAttachment(MakeMemoryView(Buffer).Left(4), ECbValidateMode::None), ECbValidateError::None);
+		TEST_EQUAL(TEXT("ValidateCompactBinaryAttachment(Binary, OutOfBoundsValue)"), ValidateCompactBinaryAttachment(MakeMemoryView(Buffer).Left(1), ECbValidateMode::All), ECbValidateError::OutOfBounds);
+		TEST_EQUAL(TEXT("ValidateCompactBinaryAttachment(Binary, OutOfBoundsHash)"), ValidateCompactBinaryAttachment(MakeMemoryView(Buffer).Left(4), ECbValidateMode::All), ECbValidateError::OutOfBounds);
+		TEST_EQUAL(TEXT("ValidateCompactBinaryAttachment(Binary, OutOfBoundsHash, Mode)"), ValidateCompactBinaryAttachment(MakeMemoryView(Buffer).Left(4), ECbValidateMode::None), ECbValidateError::None);
 	}
 
 	// Test Object
@@ -352,31 +352,31 @@ TEST_CASE("Core::Serialization::FCbField::Validate Compact Binary Attachment", "
 		Writer.EndObject();
 		FBufferArchive Buffer;
 		FCbAttachment(Writer.Save().AsObject()).Save(Buffer);
-		TestEqual(TEXT("ValidateCompactBinaryAttachment(Object)"), ValidateCompactBinaryAttachment(MakeMemoryView(Buffer), ECbValidateMode::All), ECbValidateError::None);
+		TEST_EQUAL(TEXT("ValidateCompactBinaryAttachment(Object)"), ValidateCompactBinaryAttachment(MakeMemoryView(Buffer), ECbValidateMode::All), ECbValidateError::None);
 		++Buffer[2];
-		TestEqual(TEXT("ValidateCompactBinaryAttachment(Object, InvalidHash)"), ValidateCompactBinaryAttachment(MakeMemoryView(Buffer), ECbValidateMode::All), ECbValidateError::InvalidPackageHash);
-		TestEqual(TEXT("ValidateCompactBinaryAttachment(Object, InvalidHash, Mode)"), ValidateCompactBinaryAttachment(MakeMemoryView(Buffer), ECbValidateMode::All & ~ECbValidateMode::PackageHash), ECbValidateError::None);
+		TEST_EQUAL(TEXT("ValidateCompactBinaryAttachment(Object, InvalidHash)"), ValidateCompactBinaryAttachment(MakeMemoryView(Buffer), ECbValidateMode::All), ECbValidateError::InvalidPackageHash);
+		TEST_EQUAL(TEXT("ValidateCompactBinaryAttachment(Object, InvalidHash, Mode)"), ValidateCompactBinaryAttachment(MakeMemoryView(Buffer), ECbValidateMode::All & ~ECbValidateMode::PackageHash), ECbValidateError::None);
 		--Buffer[2];
 		Buffer.AddZeroed(1);
-		TestEqual(TEXT("ValidateCompactBinaryAttachment(Object, Padding)"), ValidateCompactBinaryAttachment(MakeMemoryView(Buffer), ECbValidateMode::All), ECbValidateError::Padding);
-		TestEqual(TEXT("ValidateCompactBinaryAttachment(Object, Padding, Mode)"), ValidateCompactBinaryAttachment(MakeMemoryView(Buffer), ECbValidateMode::All & ~ECbValidateMode::Padding), ECbValidateError::None);
+		TEST_EQUAL(TEXT("ValidateCompactBinaryAttachment(Object, Padding)"), ValidateCompactBinaryAttachment(MakeMemoryView(Buffer), ECbValidateMode::All), ECbValidateError::Padding);
+		TEST_EQUAL(TEXT("ValidateCompactBinaryAttachment(Object, Padding, Mode)"), ValidateCompactBinaryAttachment(MakeMemoryView(Buffer), ECbValidateMode::All & ~ECbValidateMode::Padding), ECbValidateError::None);
 		Buffer.Pop(/*bAllowShrinking*/ false);
-		TestEqual(TEXT("ValidateCompactBinaryAttachment(Object, OutOfBoundsValue)"), ValidateCompactBinaryAttachment(MakeMemoryView(Buffer).Left(1), ECbValidateMode::All), ECbValidateError::OutOfBounds);
-		TestEqual(TEXT("ValidateCompactBinaryAttachment(Object, OutOfBoundsHash)"), ValidateCompactBinaryAttachment(MakeMemoryView(Buffer).Left(8), ECbValidateMode::All), ECbValidateError::OutOfBounds);
-		TestEqual(TEXT("ValidateCompactBinaryAttachment(Object, OutOfBoundsHash, Mode)"), ValidateCompactBinaryAttachment(MakeMemoryView(Buffer).Left(8), ECbValidateMode::None), ECbValidateError::None);
+		TEST_EQUAL(TEXT("ValidateCompactBinaryAttachment(Object, OutOfBoundsValue)"), ValidateCompactBinaryAttachment(MakeMemoryView(Buffer).Left(1), ECbValidateMode::All), ECbValidateError::OutOfBounds);
+		TEST_EQUAL(TEXT("ValidateCompactBinaryAttachment(Object, OutOfBoundsHash)"), ValidateCompactBinaryAttachment(MakeMemoryView(Buffer).Left(8), ECbValidateMode::All), ECbValidateError::OutOfBounds);
+		TEST_EQUAL(TEXT("ValidateCompactBinaryAttachment(Object, OutOfBoundsHash, Mode)"), ValidateCompactBinaryAttachment(MakeMemoryView(Buffer).Left(8), ECbValidateMode::None), ECbValidateError::None);
 	}
 
 	// Test Empty Object
 	{
 		FBufferArchive Buffer;
 		FCbFieldIterator::MakeSingle(FCbObject().AsField()).CopyRangeTo(Buffer);
-		TestEqual(TEXT("ValidateCompactBinaryAttachment(Object)"), ValidateCompactBinaryAttachment(MakeMemoryView(Buffer), ECbValidateMode::All), ECbValidateError::None);
+		TEST_EQUAL(TEXT("ValidateCompactBinaryAttachment(Object)"), ValidateCompactBinaryAttachment(MakeMemoryView(Buffer), ECbValidateMode::All), ECbValidateError::None);
 	}
 
 	// Test InvalidPackageFormat
 	{
-		TestEqual(TEXT("ValidateCompactBinaryAttachment(Empty)"), Validate({}), ECbValidateError::InvalidPackageFormat);
-		TestEqual(TEXT("ValidateCompactBinaryAttachment(Empty, Mode)"), Validate({}, ECbValidateMode::All & ~ECbValidateMode::Package), ECbValidateError::None);
+		TEST_EQUAL(TEXT("ValidateCompactBinaryAttachment(Empty)"), Validate({}), ECbValidateError::InvalidPackageFormat);
+		TEST_EQUAL(TEXT("ValidateCompactBinaryAttachment(Empty, Mode)"), Validate({}, ECbValidateMode::All & ~ECbValidateMode::Package), ECbValidateError::None);
 	}
 	{
 		FBufferArchive Buffer;
@@ -384,8 +384,8 @@ TEST_CASE("Core::Serialization::FCbField::Validate Compact Binary Attachment", "
 		Writer.AddBinaryAttachment(FIoHash::HashBuffer(BinaryView));
 		Writer.SetName("Name"_ASV).AddBinary(BinaryView);
 		Writer.Save(Buffer);
-		TestEqual(TEXT("ValidateCompactBinaryAttachment(NameOnValue)"), ValidateCompactBinaryAttachment(MakeMemoryView(Buffer), ECbValidateMode::All), ECbValidateError::InvalidPackageFormat);
-		TestEqual(TEXT("ValidateCompactBinaryAttachment(NameOnValue, Mode)"), ValidateCompactBinaryAttachment(MakeMemoryView(Buffer), ECbValidateMode::All & ~ECbValidateMode::Package), ECbValidateError::None);
+		TEST_EQUAL(TEXT("ValidateCompactBinaryAttachment(NameOnValue)"), ValidateCompactBinaryAttachment(MakeMemoryView(Buffer), ECbValidateMode::All), ECbValidateError::InvalidPackageFormat);
+		TEST_EQUAL(TEXT("ValidateCompactBinaryAttachment(NameOnValue, Mode)"), ValidateCompactBinaryAttachment(MakeMemoryView(Buffer), ECbValidateMode::All & ~ECbValidateMode::Package), ECbValidateError::None);
 	}
 	{
 		FBufferArchive Buffer;
@@ -393,12 +393,12 @@ TEST_CASE("Core::Serialization::FCbField::Validate Compact Binary Attachment", "
 		Writer.SetName("Name"_ASV).AddBinaryAttachment(FIoHash::HashBuffer(BinaryView));
 		Writer.AddBinary(BinaryView);
 		Writer.Save(Buffer);
-		TestEqual(TEXT("ValidateCompactBinaryAttachment(NameOnHash)"), ValidateCompactBinaryAttachment(MakeMemoryView(Buffer), ECbValidateMode::All), ECbValidateError::InvalidPackageFormat);
-		TestEqual(TEXT("ValidateCompactBinaryAttachment(NameOnHash, Mode)"), ValidateCompactBinaryAttachment(MakeMemoryView(Buffer), ECbValidateMode::All & ~ECbValidateMode::Package), ECbValidateError::None);
+		TEST_EQUAL(TEXT("ValidateCompactBinaryAttachment(NameOnHash)"), ValidateCompactBinaryAttachment(MakeMemoryView(Buffer), ECbValidateMode::All), ECbValidateError::InvalidPackageFormat);
+		TEST_EQUAL(TEXT("ValidateCompactBinaryAttachment(NameOnHash, Mode)"), ValidateCompactBinaryAttachment(MakeMemoryView(Buffer), ECbValidateMode::All & ~ECbValidateMode::Package), ECbValidateError::None);
 	}
 }
 
-TEST_CASE("Core::Serialization::FCbField::Validate Compact Binary Package", "[Core][Serialization][Smoke]")
+TEST_CASE_METHOD(FCoreTestFixture, "Core::Serialization::FCbField::Validate Compact Binary Package", "[Core][Serialization][Smoke]")
 {
 	const uint8 BinaryValue[] = { 0, 1, 2, 3 };
 	const FMemoryView BinaryView = MakeMemoryView(BinaryValue);
@@ -415,9 +415,9 @@ TEST_CASE("Core::Serialization::FCbField::Validate Compact Binary Package", "[Co
 
 	// Test Null
 	{
-		TestEqual(TEXT("ValidateCompactBinaryPackage(Null)"), ValidateCompactBinaryPackage(MakeMemoryView<uint8>({uint8(ECbFieldType::Null)}), ECbValidateMode::All), ECbValidateError::None);
-		TestEqual(TEXT("ValidateCompactBinaryPackage(Null, Padding)"), ValidateCompactBinaryPackage(MakeMemoryView<uint8>({uint8(ECbFieldType::Null), 0}), ECbValidateMode::All), ECbValidateError::Padding);
-		TestEqual(TEXT("ValidateCompactBinaryPackage(Null, Padding, Mode)"), ValidateCompactBinaryPackage(MakeMemoryView<uint8>({uint8(ECbFieldType::Null), 0}), ECbValidateMode::All & ~ECbValidateMode::Padding), ECbValidateError::None);
+		TEST_EQUAL(TEXT("ValidateCompactBinaryPackage(Null)"), ValidateCompactBinaryPackage(MakeMemoryView<uint8>({uint8(ECbFieldType::Null)}), ECbValidateMode::All), ECbValidateError::None);
+		TEST_EQUAL(TEXT("ValidateCompactBinaryPackage(Null, Padding)"), ValidateCompactBinaryPackage(MakeMemoryView<uint8>({uint8(ECbFieldType::Null), 0}), ECbValidateMode::All), ECbValidateError::Padding);
+		TEST_EQUAL(TEXT("ValidateCompactBinaryPackage(Null, Padding, Mode)"), ValidateCompactBinaryPackage(MakeMemoryView<uint8>({uint8(ECbFieldType::Null), 0}), ECbValidateMode::All & ~ECbValidateMode::Padding), ECbValidateError::None);
 	}
 
 	// Test Object
@@ -428,22 +428,22 @@ TEST_CASE("Core::Serialization::FCbField::Validate Compact Binary Package", "[Co
 		Writer.AddObject(CommonRootObject);
 		Writer.AddNull();
 		Writer.Save(Buffer);
-		TestEqual(TEXT("ValidateCompactBinaryPackage(Object)"), ValidateCompactBinaryPackage(MakeMemoryView(Buffer), ECbValidateMode::All), ECbValidateError::None);
-		TestEqual(TEXT("ValidateCompactBinaryPackage(Object, MissingHash)"), ValidateCompactBinaryPackage(MakeMemoryView(Buffer).LeftChop(34), ECbValidateMode::All), ECbValidateError::InvalidPackageFormat);
-		TestEqual(TEXT("ValidateCompactBinaryPackage(Object, MissingHash, Mode)"), ValidateCompactBinaryPackage(MakeMemoryView(Buffer).LeftChop(34), ECbValidateMode::All & ~ECbValidateMode::Package), ECbValidateError::None);
-		TestEqual(TEXT("ValidateCompactBinaryPackage(Object, MissingNull)"), ValidateCompactBinaryPackage(MakeMemoryView(Buffer).LeftChop(1), ECbValidateMode::All), ECbValidateError::InvalidPackageFormat);
-		TestEqual(TEXT("ValidateCompactBinaryPackage(Object, MissingNull, Mode)"), ValidateCompactBinaryPackage(MakeMemoryView(Buffer).LeftChop(1), ECbValidateMode::All & ~ECbValidateMode::Package), ECbValidateError::None);
+		TEST_EQUAL(TEXT("ValidateCompactBinaryPackage(Object)"), ValidateCompactBinaryPackage(MakeMemoryView(Buffer), ECbValidateMode::All), ECbValidateError::None);
+		TEST_EQUAL(TEXT("ValidateCompactBinaryPackage(Object, MissingHash)"), ValidateCompactBinaryPackage(MakeMemoryView(Buffer).LeftChop(34), ECbValidateMode::All), ECbValidateError::InvalidPackageFormat);
+		TEST_EQUAL(TEXT("ValidateCompactBinaryPackage(Object, MissingHash, Mode)"), ValidateCompactBinaryPackage(MakeMemoryView(Buffer).LeftChop(34), ECbValidateMode::All & ~ECbValidateMode::Package), ECbValidateError::None);
+		TEST_EQUAL(TEXT("ValidateCompactBinaryPackage(Object, MissingNull)"), ValidateCompactBinaryPackage(MakeMemoryView(Buffer).LeftChop(1), ECbValidateMode::All), ECbValidateError::InvalidPackageFormat);
+		TEST_EQUAL(TEXT("ValidateCompactBinaryPackage(Object, MissingNull, Mode)"), ValidateCompactBinaryPackage(MakeMemoryView(Buffer).LeftChop(1), ECbValidateMode::All & ~ECbValidateMode::Package), ECbValidateError::None);
 		++Buffer.Last(1);
-		TestEqual(TEXT("ValidateCompactBinaryPackage(Object, InvalidHash)"), ValidateCompactBinaryPackage(MakeMemoryView(Buffer), ECbValidateMode::All), ECbValidateError::InvalidPackageHash);
-		TestEqual(TEXT("ValidateCompactBinaryPackage(Object, InvalidHash, Mode)"), ValidateCompactBinaryPackage(MakeMemoryView(Buffer), ECbValidateMode::All & ~ECbValidateMode::PackageHash), ECbValidateError::None);
+		TEST_EQUAL(TEXT("ValidateCompactBinaryPackage(Object, InvalidHash)"), ValidateCompactBinaryPackage(MakeMemoryView(Buffer), ECbValidateMode::All), ECbValidateError::InvalidPackageHash);
+		TEST_EQUAL(TEXT("ValidateCompactBinaryPackage(Object, InvalidHash, Mode)"), ValidateCompactBinaryPackage(MakeMemoryView(Buffer), ECbValidateMode::All & ~ECbValidateMode::PackageHash), ECbValidateError::None);
 		--Buffer.Last(1);
 		Buffer.AddZeroed(1);
-		TestEqual(TEXT("ValidateCompactBinaryPackage(Object, Padding)"), ValidateCompactBinaryPackage(MakeMemoryView(Buffer), ECbValidateMode::All), ECbValidateError::Padding);
-		TestEqual(TEXT("ValidateCompactBinaryPackage(Object, Padding, Mode)"), ValidateCompactBinaryPackage(MakeMemoryView(Buffer), ECbValidateMode::All & ~ECbValidateMode::Padding), ECbValidateError::None);
+		TEST_EQUAL(TEXT("ValidateCompactBinaryPackage(Object, Padding)"), ValidateCompactBinaryPackage(MakeMemoryView(Buffer), ECbValidateMode::All), ECbValidateError::Padding);
+		TEST_EQUAL(TEXT("ValidateCompactBinaryPackage(Object, Padding, Mode)"), ValidateCompactBinaryPackage(MakeMemoryView(Buffer), ECbValidateMode::All & ~ECbValidateMode::Padding), ECbValidateError::None);
 		Buffer.Pop(/*bAllowShrinking*/ false);
-		TestEqual(TEXT("ValidateCompactBinaryPackage(Object, OutOfBoundsValue)"), ValidateCompactBinaryPackage(MakeMemoryView(Buffer).Left(1), ECbValidateMode::All), ECbValidateError::OutOfBounds);
-		TestEqual(TEXT("ValidateCompactBinaryPackage(Object, OutOfBoundsHash)"), ValidateCompactBinaryPackage(MakeMemoryView(Buffer).Left(11), ECbValidateMode::All), ECbValidateError::OutOfBounds);
-		TestEqual(TEXT("ValidateCompactBinaryPackage(Object, OutOfBoundsHash, Mode)"), ValidateCompactBinaryPackage(MakeMemoryView(Buffer).Left(11), ECbValidateMode::None), ECbValidateError::None);
+		TEST_EQUAL(TEXT("ValidateCompactBinaryPackage(Object, OutOfBoundsValue)"), ValidateCompactBinaryPackage(MakeMemoryView(Buffer).Left(1), ECbValidateMode::All), ECbValidateError::OutOfBounds);
+		TEST_EQUAL(TEXT("ValidateCompactBinaryPackage(Object, OutOfBoundsHash)"), ValidateCompactBinaryPackage(MakeMemoryView(Buffer).Left(11), ECbValidateMode::All), ECbValidateError::OutOfBounds);
+		TEST_EQUAL(TEXT("ValidateCompactBinaryPackage(Object, OutOfBoundsHash, Mode)"), ValidateCompactBinaryPackage(MakeMemoryView(Buffer).Left(11), ECbValidateMode::None), ECbValidateError::None);
 	}
 
 	// Test Object + Attachment
@@ -457,20 +457,20 @@ TEST_CASE("Core::Serialization::FCbField::Validate Compact Binary Package", "[Co
 		Writer.AddBinary(BinaryBuffer);
 		Writer.AddNull();
 		Writer.Save(Buffer);
-		TestEqual(TEXT("ValidateCompactBinaryPackage(Object+Attachment)"), ValidateCompactBinaryPackage(MakeMemoryView(Buffer), ECbValidateMode::All), ECbValidateError::None);
-		TestEqual(TEXT("ValidateCompactBinaryPackage(Object+Attachment, MissingNull)"), ValidateCompactBinaryPackage(MakeMemoryView(Buffer).LeftChop(1), ECbValidateMode::All), ECbValidateError::InvalidPackageFormat);
-		TestEqual(TEXT("ValidateCompactBinaryPackage(Object+Attachment, MissingNull, Mode)"), ValidateCompactBinaryPackage(MakeMemoryView(Buffer).LeftChop(1), ECbValidateMode::All & ~ECbValidateMode::Package), ECbValidateError::None);
+		TEST_EQUAL(TEXT("ValidateCompactBinaryPackage(Object+Attachment)"), ValidateCompactBinaryPackage(MakeMemoryView(Buffer), ECbValidateMode::All), ECbValidateError::None);
+		TEST_EQUAL(TEXT("ValidateCompactBinaryPackage(Object+Attachment, MissingNull)"), ValidateCompactBinaryPackage(MakeMemoryView(Buffer).LeftChop(1), ECbValidateMode::All), ECbValidateError::InvalidPackageFormat);
+		TEST_EQUAL(TEXT("ValidateCompactBinaryPackage(Object+Attachment, MissingNull, Mode)"), ValidateCompactBinaryPackage(MakeMemoryView(Buffer).LeftChop(1), ECbValidateMode::All & ~ECbValidateMode::Package), ECbValidateError::None);
 		++Buffer[2];
-		TestEqual(TEXT("ValidateCompactBinaryPackage(Object+Attachment, InvalidHash)"), ValidateCompactBinaryPackage(MakeMemoryView(Buffer), ECbValidateMode::All), ECbValidateError::InvalidPackageHash);
-		TestEqual(TEXT("ValidateCompactBinaryPackage(Object+Attachment, InvalidHash, Mode)"), ValidateCompactBinaryPackage(MakeMemoryView(Buffer), ECbValidateMode::All & ~ECbValidateMode::PackageHash), ECbValidateError::None);
+		TEST_EQUAL(TEXT("ValidateCompactBinaryPackage(Object+Attachment, InvalidHash)"), ValidateCompactBinaryPackage(MakeMemoryView(Buffer), ECbValidateMode::All), ECbValidateError::InvalidPackageHash);
+		TEST_EQUAL(TEXT("ValidateCompactBinaryPackage(Object+Attachment, InvalidHash, Mode)"), ValidateCompactBinaryPackage(MakeMemoryView(Buffer), ECbValidateMode::All & ~ECbValidateMode::PackageHash), ECbValidateError::None);
 		--Buffer[2];
 		Buffer.AddZeroed(1);
-		TestEqual(TEXT("ValidateCompactBinaryPackage(Object+Attachment, Padding)"), ValidateCompactBinaryPackage(MakeMemoryView(Buffer), ECbValidateMode::All), ECbValidateError::Padding);
-		TestEqual(TEXT("ValidateCompactBinaryPackage(Object+Attachment, Padding, Mode)"), ValidateCompactBinaryPackage(MakeMemoryView(Buffer), ECbValidateMode::All & ~ECbValidateMode::Padding), ECbValidateError::None);
+		TEST_EQUAL(TEXT("ValidateCompactBinaryPackage(Object+Attachment, Padding)"), ValidateCompactBinaryPackage(MakeMemoryView(Buffer), ECbValidateMode::All), ECbValidateError::Padding);
+		TEST_EQUAL(TEXT("ValidateCompactBinaryPackage(Object+Attachment, Padding, Mode)"), ValidateCompactBinaryPackage(MakeMemoryView(Buffer), ECbValidateMode::All & ~ECbValidateMode::Padding), ECbValidateError::None);
 		Buffer.Pop(/*bAllowShrinking*/ false);
-		TestEqual(TEXT("ValidateCompactBinaryPackage(Object+Attachment, OutOfBoundsValue)"), ValidateCompactBinaryPackage(MakeMemoryView(Buffer).Left(AttachmentOffset + 1), ECbValidateMode::All), ECbValidateError::OutOfBounds);
-		TestEqual(TEXT("ValidateCompactBinaryPackage(Object+Attachment, OutOfBoundsHash)"), ValidateCompactBinaryPackage(MakeMemoryView(Buffer).Left(AttachmentOffset + 7), ECbValidateMode::All), ECbValidateError::OutOfBounds);
-		TestEqual(TEXT("ValidateCompactBinaryPackage(Object+Attachment, OutOfBoundsHash, Mode)"), ValidateCompactBinaryPackage(MakeMemoryView(Buffer).Left(AttachmentOffset + 7), ECbValidateMode::None), ECbValidateError::None);
+		TEST_EQUAL(TEXT("ValidateCompactBinaryPackage(Object+Attachment, OutOfBoundsValue)"), ValidateCompactBinaryPackage(MakeMemoryView(Buffer).Left(AttachmentOffset + 1), ECbValidateMode::All), ECbValidateError::OutOfBounds);
+		TEST_EQUAL(TEXT("ValidateCompactBinaryPackage(Object+Attachment, OutOfBoundsHash)"), ValidateCompactBinaryPackage(MakeMemoryView(Buffer).Left(AttachmentOffset + 7), ECbValidateMode::All), ECbValidateError::OutOfBounds);
+		TEST_EQUAL(TEXT("ValidateCompactBinaryPackage(Object+Attachment, OutOfBoundsHash, Mode)"), ValidateCompactBinaryPackage(MakeMemoryView(Buffer).Left(AttachmentOffset + 7), ECbValidateMode::None), ECbValidateError::None);
 	}
 
 	// Test Object + Compressed Binary Attachment
@@ -483,26 +483,26 @@ TEST_CASE("Core::Serialization::FCbField::Validate Compact Binary Package", "[Co
 		Writer.AddBinary(CompressedBinaryBuffer.GetCompressed());
 		Writer.AddNull();
 		Writer.Save(Buffer);
-		TestEqual(TEXT("ValidateCompactBinaryPackage(Object+CompressedBinaryAttachment)"), ValidateCompactBinaryPackage(MakeMemoryView(Buffer), ECbValidateMode::All), ECbValidateError::None);
-		TestEqual(TEXT("ValidateCompactBinaryPackage(Object+CompressedBinaryAttachment, MissingNull)"), ValidateCompactBinaryPackage(MakeMemoryView(Buffer).LeftChop(1), ECbValidateMode::All), ECbValidateError::InvalidPackageFormat);
-		TestEqual(TEXT("ValidateCompactBinaryPackage(Object+CompressedBinaryAttachment, MissingNull, Mode)"), ValidateCompactBinaryPackage(MakeMemoryView(Buffer).LeftChop(1), ECbValidateMode::All & ~ECbValidateMode::Package), ECbValidateError::None);
+		TEST_EQUAL(TEXT("ValidateCompactBinaryPackage(Object+CompressedBinaryAttachment)"), ValidateCompactBinaryPackage(MakeMemoryView(Buffer), ECbValidateMode::All), ECbValidateError::None);
+		TEST_EQUAL(TEXT("ValidateCompactBinaryPackage(Object+CompressedBinaryAttachment, MissingNull)"), ValidateCompactBinaryPackage(MakeMemoryView(Buffer).LeftChop(1), ECbValidateMode::All), ECbValidateError::InvalidPackageFormat);
+		TEST_EQUAL(TEXT("ValidateCompactBinaryPackage(Object+CompressedBinaryAttachment, MissingNull, Mode)"), ValidateCompactBinaryPackage(MakeMemoryView(Buffer).LeftChop(1), ECbValidateMode::All & ~ECbValidateMode::Package), ECbValidateError::None);
 		++Buffer[2];
-		TestEqual(TEXT("ValidateCompactBinaryPackage(Object+CompressedBinaryAttachment, InvalidHash)"), ValidateCompactBinaryPackage(MakeMemoryView(Buffer), ECbValidateMode::All), ECbValidateError::InvalidPackageHash);
-		TestEqual(TEXT("ValidateCompactBinaryPackage(Object+CompressedBinaryAttachment, InvalidHash, Mode)"), ValidateCompactBinaryPackage(MakeMemoryView(Buffer), ECbValidateMode::All & ~ECbValidateMode::PackageHash), ECbValidateError::None);
+		TEST_EQUAL(TEXT("ValidateCompactBinaryPackage(Object+CompressedBinaryAttachment, InvalidHash)"), ValidateCompactBinaryPackage(MakeMemoryView(Buffer), ECbValidateMode::All), ECbValidateError::InvalidPackageHash);
+		TEST_EQUAL(TEXT("ValidateCompactBinaryPackage(Object+CompressedBinaryAttachment, InvalidHash, Mode)"), ValidateCompactBinaryPackage(MakeMemoryView(Buffer), ECbValidateMode::All & ~ECbValidateMode::PackageHash), ECbValidateError::None);
 		--Buffer[2];
 		Buffer.AddZeroed(1);
-		TestEqual(TEXT("ValidateCompactBinaryPackage(Object+CompressedBinaryAttachment, Padding)"), ValidateCompactBinaryPackage(MakeMemoryView(Buffer), ECbValidateMode::All), ECbValidateError::Padding);
-		TestEqual(TEXT("ValidateCompactBinaryPackage(Object+CompressedBinaryAttachment, Padding, Mode)"), ValidateCompactBinaryPackage(MakeMemoryView(Buffer), ECbValidateMode::All & ~ECbValidateMode::Padding), ECbValidateError::None);
+		TEST_EQUAL(TEXT("ValidateCompactBinaryPackage(Object+CompressedBinaryAttachment, Padding)"), ValidateCompactBinaryPackage(MakeMemoryView(Buffer), ECbValidateMode::All), ECbValidateError::Padding);
+		TEST_EQUAL(TEXT("ValidateCompactBinaryPackage(Object+CompressedBinaryAttachment, Padding, Mode)"), ValidateCompactBinaryPackage(MakeMemoryView(Buffer), ECbValidateMode::All & ~ECbValidateMode::Padding), ECbValidateError::None);
 		Buffer.Pop(/*bAllowShrinking*/ false);
-		TestEqual(TEXT("ValidateCompactBinaryPackage(Object+CompressedBinaryAttachment, OutOfBoundsValue)"), ValidateCompactBinaryPackage(MakeMemoryView(Buffer).Left(AttachmentOffset + 1), ECbValidateMode::All), ECbValidateError::OutOfBounds);
-		TestEqual(TEXT("ValidateCompactBinaryPackage(Object+CompressedBinaryAttachment, OutOfBoundsHash)"), ValidateCompactBinaryPackage(MakeMemoryView(Buffer).Left(AttachmentOffset + 7), ECbValidateMode::All), ECbValidateError::OutOfBounds);
-		TestEqual(TEXT("ValidateCompactBinaryPackage(Object+CompressedBinaryAttachment, OutOfBoundsHash, Mode)"), ValidateCompactBinaryPackage(MakeMemoryView(Buffer).Left(AttachmentOffset + 7), ECbValidateMode::None), ECbValidateError::None);
+		TEST_EQUAL(TEXT("ValidateCompactBinaryPackage(Object+CompressedBinaryAttachment, OutOfBoundsValue)"), ValidateCompactBinaryPackage(MakeMemoryView(Buffer).Left(AttachmentOffset + 1), ECbValidateMode::All), ECbValidateError::OutOfBounds);
+		TEST_EQUAL(TEXT("ValidateCompactBinaryPackage(Object+CompressedBinaryAttachment, OutOfBoundsHash)"), ValidateCompactBinaryPackage(MakeMemoryView(Buffer).Left(AttachmentOffset + 7), ECbValidateMode::All), ECbValidateError::OutOfBounds);
+		TEST_EQUAL(TEXT("ValidateCompactBinaryPackage(Object+CompressedBinaryAttachment, OutOfBoundsHash, Mode)"), ValidateCompactBinaryPackage(MakeMemoryView(Buffer).Left(AttachmentOffset + 7), ECbValidateMode::None), ECbValidateError::None);
 	}
 
 	// Test InvalidPackageFormat
 	{
-		TestEqual(TEXT("ValidateCompactBinaryPackage(Empty)"), ValidateCompactBinaryPackage(FMemoryView(), ECbValidateMode::All), ECbValidateError::InvalidPackageFormat);
-		TestEqual(TEXT("ValidateCompactBinaryPackage(Empty, Mode)"), ValidateCompactBinaryPackage(FMemoryView(), ECbValidateMode::All & ~ECbValidateMode::Package), ECbValidateError::None);
+		TEST_EQUAL(TEXT("ValidateCompactBinaryPackage(Empty)"), ValidateCompactBinaryPackage(FMemoryView(), ECbValidateMode::All), ECbValidateError::InvalidPackageFormat);
+		TEST_EQUAL(TEXT("ValidateCompactBinaryPackage(Empty, Mode)"), ValidateCompactBinaryPackage(FMemoryView(), ECbValidateMode::All & ~ECbValidateMode::Package), ECbValidateError::None);
 	}
 	{
 		FBufferArchive Buffer;
@@ -512,8 +512,8 @@ TEST_CASE("Core::Serialization::FCbField::Validate Compact Binary Package", "[Co
 		Writer.AddObject(CommonRootObject);
 		Writer.AddNull();
 		Writer.Save(Buffer);
-		TestEqual(TEXT("ValidateCompactBinaryPackage(NameOnObject)"), ValidateCompactBinaryPackage(MakeMemoryView(Buffer), ECbValidateMode::All), ECbValidateError::InvalidPackageFormat);
-		TestEqual(TEXT("ValidateCompactBinaryPackage(NameOnObject, Mode)"), ValidateCompactBinaryPackage(MakeMemoryView(Buffer), ECbValidateMode::All & ~ECbValidateMode::Package), ECbValidateError::None);
+		TEST_EQUAL(TEXT("ValidateCompactBinaryPackage(NameOnObject)"), ValidateCompactBinaryPackage(MakeMemoryView(Buffer), ECbValidateMode::All), ECbValidateError::InvalidPackageFormat);
+		TEST_EQUAL(TEXT("ValidateCompactBinaryPackage(NameOnObject, Mode)"), ValidateCompactBinaryPackage(MakeMemoryView(Buffer), ECbValidateMode::All & ~ECbValidateMode::Package), ECbValidateError::None);
 	}
 	{
 		FBufferArchive Buffer;
@@ -523,8 +523,8 @@ TEST_CASE("Core::Serialization::FCbField::Validate Compact Binary Package", "[Co
 		Writer.AddObject(CommonRootObject);
 		Writer.AddNull();
 		Writer.Save(Buffer);
-		TestEqual(TEXT("ValidateCompactBinaryPackage(NameOnHash)"), ValidateCompactBinaryPackage(MakeMemoryView(Buffer), ECbValidateMode::All), ECbValidateError::InvalidPackageFormat);
-		TestEqual(TEXT("ValidateCompactBinaryPackage(NameOnHash, Mode)"), ValidateCompactBinaryPackage(MakeMemoryView(Buffer), ECbValidateMode::All & ~ECbValidateMode::Package), ECbValidateError::None);
+		TEST_EQUAL(TEXT("ValidateCompactBinaryPackage(NameOnHash)"), ValidateCompactBinaryPackage(MakeMemoryView(Buffer), ECbValidateMode::All), ECbValidateError::InvalidPackageFormat);
+		TEST_EQUAL(TEXT("ValidateCompactBinaryPackage(NameOnHash, Mode)"), ValidateCompactBinaryPackage(MakeMemoryView(Buffer), ECbValidateMode::All & ~ECbValidateMode::Package), ECbValidateError::None);
 	}
 	{
 		FBufferArchive Buffer;
@@ -535,8 +535,8 @@ TEST_CASE("Core::Serialization::FCbField::Validate Compact Binary Package", "[Co
 		Writer.SetName("Name"_ASV).AddBinary(BinaryBuffer);
 		Writer.AddNull();
 		Writer.Save(Buffer);
-		TestEqual(TEXT("ValidateCompactBinaryPackage(NameOnAttachment)"), ValidateCompactBinaryPackage(MakeMemoryView(Buffer), ECbValidateMode::All), ECbValidateError::InvalidPackageFormat);
-		TestEqual(TEXT("ValidateCompactBinaryPackage(NameOnAttachment, Mode)"), ValidateCompactBinaryPackage(MakeMemoryView(Buffer), ECbValidateMode::All & ~ECbValidateMode::Package), ECbValidateError::None);
+		TEST_EQUAL(TEXT("ValidateCompactBinaryPackage(NameOnAttachment)"), ValidateCompactBinaryPackage(MakeMemoryView(Buffer), ECbValidateMode::All), ECbValidateError::InvalidPackageFormat);
+		TEST_EQUAL(TEXT("ValidateCompactBinaryPackage(NameOnAttachment, Mode)"), ValidateCompactBinaryPackage(MakeMemoryView(Buffer), ECbValidateMode::All & ~ECbValidateMode::Package), ECbValidateError::None);
 	}
 	{
 		FBufferArchive Buffer;
@@ -547,8 +547,8 @@ TEST_CASE("Core::Serialization::FCbField::Validate Compact Binary Package", "[Co
 		Writer.AddObject(CommonRootObject);
 		Writer.AddNull();
 		Writer.Save(Buffer);
-		TestEqual(TEXT("ValidateCompactBinaryPackage(MultipleObjects)"), ValidateCompactBinaryPackage(MakeMemoryView(Buffer), ECbValidateMode::All), ECbValidateError::MultiplePackageObjects);
-		TestEqual(TEXT("ValidateCompactBinaryPackage(MultipleObjects, Mode)"), ValidateCompactBinaryPackage(MakeMemoryView(Buffer), ECbValidateMode::All & ~ECbValidateMode::Package), ECbValidateError::None);
+		TEST_EQUAL(TEXT("ValidateCompactBinaryPackage(MultipleObjects)"), ValidateCompactBinaryPackage(MakeMemoryView(Buffer), ECbValidateMode::All), ECbValidateError::MultiplePackageObjects);
+		TEST_EQUAL(TEXT("ValidateCompactBinaryPackage(MultipleObjects, Mode)"), ValidateCompactBinaryPackage(MakeMemoryView(Buffer), ECbValidateMode::All & ~ECbValidateMode::Package), ECbValidateError::None);
 	}
 	{
 		FBufferArchive Buffer;
@@ -559,8 +559,8 @@ TEST_CASE("Core::Serialization::FCbField::Validate Compact Binary Package", "[Co
 		Writer.AddBinary(BinaryBuffer);
 		Writer.AddNull();
 		Writer.Save(Buffer);
-		TestEqual(TEXT("ValidateCompactBinaryPackage(DuplicateAttachments)"), ValidateCompactBinaryPackage(MakeMemoryView(Buffer), ECbValidateMode::All), ECbValidateError::DuplicateAttachments);
-		TestEqual(TEXT("ValidateCompactBinaryPackage(DuplicateAttachments, Mode)"), ValidateCompactBinaryPackage(MakeMemoryView(Buffer), ECbValidateMode::All & ~ECbValidateMode::Package), ECbValidateError::None);
+		TEST_EQUAL(TEXT("ValidateCompactBinaryPackage(DuplicateAttachments)"), ValidateCompactBinaryPackage(MakeMemoryView(Buffer), ECbValidateMode::All), ECbValidateError::DuplicateAttachments);
+		TEST_EQUAL(TEXT("ValidateCompactBinaryPackage(DuplicateAttachments, Mode)"), ValidateCompactBinaryPackage(MakeMemoryView(Buffer), ECbValidateMode::All & ~ECbValidateMode::Package), ECbValidateError::None);
 	}
 	{
 		FBufferArchive Buffer;
@@ -569,8 +569,8 @@ TEST_CASE("Core::Serialization::FCbField::Validate Compact Binary Package", "[Co
 		Writer.AddBinary(CompressedBinaryBuffer.GetCompressed());
 		Writer.AddNull();
 		Writer.Save(Buffer);
-		TestEqual(TEXT("ValidateCompactBinaryPackage(DuplicateCompressedAttachments)"), ValidateCompactBinaryPackage(MakeMemoryView(Buffer), ECbValidateMode::All), ECbValidateError::DuplicateAttachments);
-		TestEqual(TEXT("ValidateCompactBinaryPackage(DuplicateCompressedAttachments, Mode)"), ValidateCompactBinaryPackage(MakeMemoryView(Buffer), ECbValidateMode::All & ~ECbValidateMode::Package), ECbValidateError::None);
+		TEST_EQUAL(TEXT("ValidateCompactBinaryPackage(DuplicateCompressedAttachments)"), ValidateCompactBinaryPackage(MakeMemoryView(Buffer), ECbValidateMode::All), ECbValidateError::DuplicateAttachments);
+		TEST_EQUAL(TEXT("ValidateCompactBinaryPackage(DuplicateCompressedAttachments, Mode)"), ValidateCompactBinaryPackage(MakeMemoryView(Buffer), ECbValidateMode::All & ~ECbValidateMode::Package), ECbValidateError::None);
 	}
 	{
 		FBufferArchive Buffer;
@@ -580,8 +580,8 @@ TEST_CASE("Core::Serialization::FCbField::Validate Compact Binary Package", "[Co
 		Writer.AddBinary(CompressedBinaryBuffer.GetCompressed());
 		Writer.AddNull();
 		Writer.Save(Buffer);
-		TestEqual(TEXT("ValidateCompactBinaryPackage(DuplicateMixedAttachments)"), ValidateCompactBinaryPackage(MakeMemoryView(Buffer), ECbValidateMode::All), ECbValidateError::DuplicateAttachments);
-		TestEqual(TEXT("ValidateCompactBinaryPackage(DuplicateMixedAttachments, Mode)"), ValidateCompactBinaryPackage(MakeMemoryView(Buffer), ECbValidateMode::All & ~ECbValidateMode::Package), ECbValidateError::None);
+		TEST_EQUAL(TEXT("ValidateCompactBinaryPackage(DuplicateMixedAttachments)"), ValidateCompactBinaryPackage(MakeMemoryView(Buffer), ECbValidateMode::All), ECbValidateError::DuplicateAttachments);
+		TEST_EQUAL(TEXT("ValidateCompactBinaryPackage(DuplicateMixedAttachments, Mode)"), ValidateCompactBinaryPackage(MakeMemoryView(Buffer), ECbValidateMode::All & ~ECbValidateMode::Package), ECbValidateError::None);
 	}
 	{
 		TCbWriter<32> RootObjectWriter;
@@ -596,8 +596,8 @@ TEST_CASE("Core::Serialization::FCbField::Validate Compact Binary Package", "[Co
 		Writer.AddObject(NullRootObject);
 		Writer.AddNull();
 		Writer.Save(Buffer);
-		TestEqual(TEXT("ValidateCompactBinaryPackage(NullObject)"), ValidateCompactBinaryPackage(MakeMemoryView(Buffer), ECbValidateMode::All), ECbValidateError::NullPackageObject);
-		TestEqual(TEXT("ValidateCompactBinaryPackage(NullObject, Mode)"), ValidateCompactBinaryPackage(MakeMemoryView(Buffer), ECbValidateMode::All & ~ECbValidateMode::Package), ECbValidateError::None);
+		TEST_EQUAL(TEXT("ValidateCompactBinaryPackage(NullObject)"), ValidateCompactBinaryPackage(MakeMemoryView(Buffer), ECbValidateMode::All), ECbValidateError::NullPackageObject);
+		TEST_EQUAL(TEXT("ValidateCompactBinaryPackage(NullObject, Mode)"), ValidateCompactBinaryPackage(MakeMemoryView(Buffer), ECbValidateMode::All & ~ECbValidateMode::Package), ECbValidateError::None);
 	}
 	{
 		FBufferArchive Buffer;
@@ -606,7 +606,7 @@ TEST_CASE("Core::Serialization::FCbField::Validate Compact Binary Package", "[Co
 		Writer.AddBinary(FMemoryView());
 		Writer.AddNull();
 		Writer.Save(Buffer);
-		TestEqual(TEXT("ValidateCompactBinaryPackage(NullAttachment)"), ValidateCompactBinaryPackage(MakeMemoryView(Buffer), ECbValidateMode::All), ECbValidateError::NullPackageAttachment);
-		TestEqual(TEXT("ValidateCompactBinaryPackage(NullAttachment, Mode)"), ValidateCompactBinaryPackage(MakeMemoryView(Buffer), ECbValidateMode::All & ~ECbValidateMode::Package), ECbValidateError::None);
+		TEST_EQUAL(TEXT("ValidateCompactBinaryPackage(NullAttachment)"), ValidateCompactBinaryPackage(MakeMemoryView(Buffer), ECbValidateMode::All), ECbValidateError::NullPackageAttachment);
+		TEST_EQUAL(TEXT("ValidateCompactBinaryPackage(NullAttachment, Mode)"), ValidateCompactBinaryPackage(MakeMemoryView(Buffer), ECbValidateMode::All & ~ECbValidateMode::Package), ECbValidateError::None);
 	}
 }
