Add new Unit Tests

diff --git a/Engine/Source/Runtime/Core/Tests/Serialization/CompactBinaryWriterTest.cpp b/Engine/Source/Runtime/Core/Tests/Serialization/CompactBinaryWriterTest.cpp
--- a/Engine/Source/Runtime/Core/Tests/Serialization/CompactBinaryWriterTest.cpp
+++ b/Engine/Source/Runtime/Core/Tests/Serialization/CompactBinaryWriterTest.cpp
@@ -12,7 +12,7 @@
 #include "Serialization/CompactBinaryValidation.h"
 #include "TestHarness.h"
 
-TEST_CASE("Core::Serialization::FCbWriter::Object", "[Core][Serialization][Smoke]")
+TEST_CASE_METHOD(FAutomationTestFixture, "Core::Serialization::FCbWriter::Object", "[Core][Serialization][Smoke]")
 {
 	TCbWriter<256> Writer;
 
@@ -22,10 +22,10 @@ TEST_CASE("Core::Serialization::FCbWriter::Object", "[Core][Serialization][Smoke
 		Writer.BeginObject();
 		Writer.EndObject();
 		FCbField Field = Writer.Save();
-		TestEqual_Conditional(TEXT("FCbWriter(Object, Empty) Validate"), ValidateCompactBinary(Field.GetOuterBuffer(), ECbValidateMode::All), ECbValidateError::None)
+		if (TestEqual(TEXT("FCbWriter(Object, Empty) Validate"), ValidateCompactBinary(Field.GetOuterBuffer(), ECbValidateMode::All), ECbValidateError::None))
 		{
-			TestTrue(TEXT("FCbWriter(Object, Empty).IsObject()"), Field.IsObject());
-			TestFalse(TEXT("FCbWriter(Object, Empty).AsObjectView()"), Field.AsObjectView().CreateViewIterator().HasValue());
+			TEST_TRUE(TEXT("FCbWriter(Object, Empty).IsObject()"), Field.IsObject());
+			TEST_FALSE(TEXT("FCbWriter(Object, Empty).AsObjectView()"), Field.AsObjectView().CreateViewIterator().HasValue());
 		}
 	}
 
@@ -36,10 +36,10 @@ TEST_CASE("Core::Serialization::FCbWriter::Object", "[Core][Serialization][Smoke
 		Writer.BeginObject();
 		Writer.EndObject();
 		FCbField Field = Writer.Save();
-		TestEqual_Conditional(TEXT("FCbWriter(Object, Empty, Name) Validate"), ValidateCompactBinary(Field.GetOuterBuffer(), ECbValidateMode::All), ECbValidateError::None)
+		if (TestEqual(TEXT("FCbWriter(Object, Empty, Name) Validate"), ValidateCompactBinary(Field.GetOuterBuffer(), ECbValidateMode::All), ECbValidateError::None))
 		{
-			TestTrue(TEXT("FCbWriter(Object, Empty, Name).IsObject()"), Field.IsObject());
-			TestFalse(TEXT("FCbWriter(Object, Empty, Name).AsObjectView()"), Field.AsObjectView().CreateViewIterator().HasValue());
+			TEST_TRUE(TEXT("FCbWriter(Object, Empty, Name).IsObject()"), Field.IsObject());
+			TEST_FALSE(TEXT("FCbWriter(Object, Empty, Name).AsObjectView()"), Field.AsObjectView().CreateViewIterator().HasValue());
 		}
 	}
 
@@ -51,12 +51,12 @@ TEST_CASE("Core::Serialization::FCbWriter::Object", "[Core][Serialization][Smoke
 		Writer.SetName("Float"_ASV).AddFloat(0.0f);
 		Writer.EndObject();
 		FCbField Field = Writer.Save();
-		TestEqual_Conditional(TEXT("FCbWriter(Object, Basic) Validate"), ValidateCompactBinary(Field.GetOuterBuffer(), ECbValidateMode::All), ECbValidateError::None)
+		if (TestEqual(TEXT("FCbWriter(Object, Basic) Validate"), ValidateCompactBinary(Field.GetOuterBuffer(), ECbValidateMode::All), ECbValidateError::None))
 		{
-			TestTrue(TEXT("FCbWriter(Object, Basic).IsObject()"), Field.IsObject());
+			TEST_TRUE(TEXT("FCbWriter(Object, Basic).IsObject()"), Field.IsObject());
 			FCbObjectView Object = Field.AsObjectView();
-			TestTrue(TEXT("FCbWriter(Object, Basic).AsObjectView()[Integer]"), Object["Integer"_ASV].IsInteger());
-			TestTrue(TEXT("FCbWriter(Object, Basic).AsObjectView()[Float]"), Object["Float"_ASV].IsFloat());
+			TEST_TRUE(TEXT("FCbWriter(Object, Basic).AsObjectView()[Integer]"), Object["Integer"_ASV].IsInteger());
+			TEST_TRUE(TEXT("FCbWriter(Object, Basic).AsObjectView()[Float]"), Object["Float"_ASV].IsFloat());
 		}
 	}
 
@@ -68,17 +68,17 @@ TEST_CASE("Core::Serialization::FCbWriter::Object", "[Core][Serialization][Smoke
 		Writer.SetName("Field2"_ASV).AddInteger(1);
 		Writer.EndObject();
 		FCbField Field = Writer.Save();
-		TestEqual_Conditional(TEXT("FCbWriter(Object, Uniform) Validate"), ValidateCompactBinary(Field.GetOuterBuffer(), ECbValidateMode::All), ECbValidateError::None)
+		if (TestEqual(TEXT("FCbWriter(Object, Uniform) Validate"), ValidateCompactBinary(Field.GetOuterBuffer(), ECbValidateMode::All), ECbValidateError::None))
 		{
-			TestTrue(TEXT("FCbWriter(Object, Uniform).IsObject()"), Field.IsObject());
+			TEST_TRUE(TEXT("FCbWriter(Object, Uniform).IsObject()"), Field.IsObject());
 			FCbObjectView Object = Field.AsObjectView();
-			TestTrue(TEXT("FCbWriter(Object, Uniform).AsObjectView()[Field1]"), Object["Field1"_ASV].IsInteger());
-			TestTrue(TEXT("FCbWriter(Object, Uniform).AsObjectView()[Field2]"), Object["Field2"_ASV].IsInteger());
+			TEST_TRUE(TEXT("FCbWriter(Object, Uniform).AsObjectView()[Field1]"), Object["Field1"_ASV].IsInteger());
+			TEST_TRUE(TEXT("FCbWriter(Object, Uniform).AsObjectView()[Field2]"), Object["Field2"_ASV].IsInteger());
 		}
 	}
 }
 
-TEST_CASE("Core::Serialization::FCbWriter::Array", "[Core][Serialization][Smoke]")
+TEST_CASE_METHOD(FAutomationTestFixture, "Core::Serialization::FCbWriter::Array", "[Core][Serialization][Smoke]")
 {
 	TCbWriter<256> Writer;
 
@@ -88,10 +88,10 @@ TEST_CASE("Core::Serialization::FCbWriter::Array", "[Core][Serialization][Smoke]
 		Writer.BeginArray();
 		Writer.EndArray();
 		FCbField Field = Writer.Save();
-		TestEqual_Conditional(TEXT("FCbWriter(Array, Empty) Validate"), ValidateCompactBinary(Field.GetOuterBuffer(), ECbValidateMode::All), ECbValidateError::None)
+		if (TestEqual(TEXT("FCbWriter(Array, Empty) Validate"), ValidateCompactBinary(Field.GetOuterBuffer(), ECbValidateMode::All), ECbValidateError::None))
 		{
-			TestTrue(TEXT("FCbWriter(Array, Empty).IsArray()"), Field.IsArray());
-			TestEqual(TEXT("FCbWriter(Array, Empty).AsArrayView()"), Field.AsArrayView().Num(), uint64(0));
+			TEST_TRUE(TEXT("FCbWriter(Array, Empty).IsArray()"), Field.IsArray());
+			TEST_EQUAL(TEXT("FCbWriter(Array, Empty).AsArrayView()"), Field.AsArrayView().Num(), uint64(0));
 		}
 	}
 
@@ -102,10 +102,10 @@ TEST_CASE("Core::Serialization::FCbWriter::Array", "[Core][Serialization][Smoke]
 		Writer.BeginArray();
 		Writer.EndArray();
 		FCbField Field = Writer.Save();
-		TestEqual_Conditional(TEXT("FCbWriter(Array, Empty, Name) Validate"), ValidateCompactBinary(Field.GetOuterBuffer(), ECbValidateMode::All), ECbValidateError::None)
+		if (TestEqual(TEXT("FCbWriter(Array, Empty, Name) Validate"), ValidateCompactBinary(Field.GetOuterBuffer(), ECbValidateMode::All), ECbValidateError::None))
 		{
-			TestTrue(TEXT("FCbWriter(Array, Empty, Name).IsArray()"), Field.IsArray());
-			TestEqual(TEXT("FCbWriter(Array, Empty, Name).AsArrayView()"), Field.AsArrayView().Num(), uint64(0));
+			TEST_TRUE(TEXT("FCbWriter(Array, Empty, Name).IsArray()"), Field.IsArray());
+			TEST_EQUAL(TEXT("FCbWriter(Array, Empty, Name).AsArrayView()"), Field.AsArrayView().Num(), uint64(0));
 		}
 	}
 
@@ -117,15 +117,15 @@ TEST_CASE("Core::Serialization::FCbWriter::Array", "[Core][Serialization][Smoke]
 		Writer.AddFloat(0.0f);
 		Writer.EndArray();
 		FCbField Field = Writer.Save();
-		TestEqual_Conditional(TEXT("FCbWriter(Array, Basic) Validate"), ValidateCompactBinary(Field.GetOuterBuffer(), ECbValidateMode::All), ECbValidateError::None)
+		if (TestEqual(TEXT("FCbWriter(Array, Basic) Validate"), ValidateCompactBinary(Field.GetOuterBuffer(), ECbValidateMode::All), ECbValidateError::None))
 		{
-			TestTrue(TEXT("FCbWriter(Array, Basic).IsArray()"), Field.IsArray());
+			TEST_TRUE(TEXT("FCbWriter(Array, Basic).IsArray()"), Field.IsArray());
 			FCbFieldViewIterator Iterator = Field.AsArrayView().CreateViewIterator();
-			TestTrue(TEXT("FCbWriter(Array, Basic).AsArrayView()[Integer]"), Iterator.IsInteger());
+			TEST_TRUE(TEXT("FCbWriter(Array, Basic).AsArrayView()[Integer]"), Iterator.IsInteger());
 			++Iterator;
-			TestTrue(TEXT("FCbWriter(Array, Basic).AsArrayView()[Float]"), Iterator.IsFloat());
+			TEST_TRUE(TEXT("FCbWriter(Array, Basic).AsArrayView()[Float]"), Iterator.IsFloat());
 			++Iterator;
-			TestFalse(TEXT("FCbWriter(Array, Basic).AsArrayView()[End]"), Iterator.HasValue());
+			TEST_FALSE(TEXT("FCbWriter(Array, Basic).AsArrayView()[End]"), Iterator.HasValue());
 		}
 	}
 
@@ -137,20 +137,20 @@ TEST_CASE("Core::Serialization::FCbWriter::Array", "[Core][Serialization][Smoke]
 		Writer.AddInteger(1);
 		Writer.EndArray();
 		FCbField Field = Writer.Save();
-		TestEqual_Conditional(TEXT("FCbWriter(Array, Uniform) Validate"), ValidateCompactBinary(Field.GetOuterBuffer(), ECbValidateMode::All), ECbValidateError::None)
+		if (TestEqual(TEXT("FCbWriter(Array, Uniform) Validate"), ValidateCompactBinary(Field.GetOuterBuffer(), ECbValidateMode::All), ECbValidateError::None))
 		{
-			TestTrue(TEXT("FCbWriter(Array, Uniform).IsArray()"), Field.IsArray());
+			TEST_TRUE(TEXT("FCbWriter(Array, Uniform).IsArray()"), Field.IsArray());
 			FCbFieldViewIterator Iterator = Field.AsArrayView().CreateViewIterator();
-			TestTrue(TEXT("FCbWriter(Array, Basic).AsArrayView()[Field1]"), Iterator.IsInteger());
+			TEST_TRUE(TEXT("FCbWriter(Array, Basic).AsArrayView()[Field1]"), Iterator.IsInteger());
 			++Iterator;
-			TestTrue(TEXT("FCbWriter(Array, Basic).AsArrayView()[Field2]"), Iterator.IsInteger());
+			TEST_TRUE(TEXT("FCbWriter(Array, Basic).AsArrayView()[Field2]"), Iterator.IsInteger());
 			++Iterator;
-			TestFalse(TEXT("FCbWriter(Array, Basic).AsArrayView()[End]"), Iterator.HasValue());
+			TEST_FALSE(TEXT("FCbWriter(Array, Basic).AsArrayView()[End]"), Iterator.HasValue());
 		}
 	}
 }
 
-TEST_CASE("Core::Serialization::FCbWriter::Null", "[Core][Serialization][Smoke]")
+TEST_CASE_METHOD(FAutomationTestFixture, "Core::Serialization::FCbWriter::Null", "[Core][Serialization][Smoke]")
 {
 	TCbWriter<256> Writer;
 
@@ -159,10 +159,10 @@ TEST_CASE("Core::Serialization::FCbWriter::Null", "[Core][Serialization][Smoke]"
 		Writer.Reset();
 		Writer.AddNull();
 		FCbField Field = Writer.Save();
-		TestEqual_Conditional(TEXT("FCbWriter(Null) Validate"), ValidateCompactBinary(Field.GetOuterBuffer(), ECbValidateMode::All), ECbValidateError::None)
+		if (TestEqual(TEXT("FCbWriter(Null) Validate"), ValidateCompactBinary(Field.GetOuterBuffer(), ECbValidateMode::All), ECbValidateError::None))
 		{
-			TestFalse(TEXT("FCbWriter(Null).HasName()"), Field.HasName());
-			TestTrue(TEXT("FCbWriter(Null).IsNull()"), Field.IsNull());
+			TEST_FALSE(TEXT("FCbWriter(Null).HasName()"), Field.HasName());
+			TEST_TRUE(TEXT("FCbWriter(Null).IsNull()"), Field.IsNull());
 		}
 	}
 
@@ -172,11 +172,11 @@ TEST_CASE("Core::Serialization::FCbWriter::Null", "[Core][Serialization][Smoke]"
 		Writer.SetName("Null"_ASV);
 		Writer.AddNull();
 		FCbField Field = Writer.Save();
-		TestEqual_Conditional(TEXT("FCbWriter(Null, Name) Validate"), ValidateCompactBinary(Field.GetOuterBuffer(), ECbValidateMode::All), ECbValidateError::None)
+		if (TestEqual(TEXT("FCbWriter(Null, Name) Validate"), ValidateCompactBinary(Field.GetOuterBuffer(), ECbValidateMode::All), ECbValidateError::None))
 		{
-			TestEqual(TEXT("FCbWriter(Null, Name).GetName()"), Field.GetName(), "Null"_U8SV);
-			TestTrue(TEXT("FCbWriter(Null, Name).HasName()"), Field.HasName());
-			TestTrue(TEXT("FCbWriter(Null, Name).IsNull()"), Field.IsNull());
+			TEST_EQUAL(TEXT("FCbWriter(Null, Name).GetName()"), Field.GetName(), "Null"_U8SV);
+			TEST_TRUE(TEXT("FCbWriter(Null, Name).HasName()"), Field.HasName());
+			TEST_TRUE(TEXT("FCbWriter(Null, Name).IsNull()"), Field.IsNull());
 		}
 	}
 
@@ -197,7 +197,7 @@ TEST_CASE("Core::Serialization::FCbWriter::Null", "[Core][Serialization][Smoke]"
 		Writer.EndObject();
 
 		FCbFieldIterator Fields = Writer.Save();
-		TestEqual(TEXT("FCbWriter(Null, Uniform) Validate"), ValidateCompactBinaryRange(Fields.GetOuterBuffer(), ECbValidateMode::All), ECbValidateError::None);
+		TEST_EQUAL(TEXT("FCbWriter(Null, Uniform) Validate"), ValidateCompactBinaryRange(Fields.GetOuterBuffer(), ECbValidateMode::All), ECbValidateError::None);
 	}
 
 	// Test Null with Save(Buffer)
@@ -210,21 +210,34 @@ TEST_CASE("Core::Serialization::FCbWriter::Null", "[Core][Serialization][Smoke]"
 		}
 		uint8 Buffer[NullCount]{};
 		FCbFieldViewIterator Fields = Writer.Save(MakeMemoryView(Buffer));
-		TestEqual_Conditional(TEXT("FCbWriter(Null, Memory) Validate"), ValidateCompactBinaryRange(MakeMemoryView(Buffer), ECbValidateMode::All), ECbValidateError::None)
+		if (TestEqual(TEXT("FCbWriter(Null, Memory) Validate"), ValidateCompactBinaryRange(MakeMemoryView(Buffer), ECbValidateMode::All), ECbValidateError::None))
 		{
 			for (int Index = 0; Index < NullCount; ++Index)
 			{
-				TestTrue(TEXT("FCbWriter(Null, Memory) IsNull"), Fields.IsNull());
+				TEST_TRUE(TEXT("FCbWriter(Null, Memory) IsNull"), Fields.IsNull());
 				++Fields;
 			}
-			TestFalse(TEXT("FCbWriter(Null, Memory) HasValue"), Fields.HasValue());
+			TEST_FALSE(TEXT("FCbWriter(Null, Memory) HasValue"), Fields.HasValue());
 		}
 	}
 }
 
 ///////////////////////////////////////////////////////////////////////////////////////////////////////////////////////
 
-TEST_CASE("Core::Serialization::FCbWriter::Binary", "[Core][Serialization][Smoke]")
+class FCbWriterBinaryTestBase : public FAutomationTestBase
+{
+protected:
+	using FAutomationTestBase::FAutomationTestBase;
+	using FAutomationTestBase::TestEqual;
+
+	template <typename T, typename Size>
+	void TestEqual(const TCHAR* What, TArrayView<T, Size> Actual, TArrayView<T, Size> Expected)
+	{
+		TestTrue(What, Actual.Num() == Expected.Num() && CompareItems(Actual.GetData(), Expected.GetData(), Actual.Num()));
+	}
+};
+
+TEST_CASE_METHOD(FAutomationTestFixture, "Core::Serialization::FCbWriter::Binary", "[Core][Serialization][Smoke]")
 {
 	TCbWriter<256> Writer;
 
@@ -233,11 +246,11 @@ TEST_CASE("Core::Serialization::FCbWriter::Binary", "[Core][Serialization][Smoke
 		Writer.Reset();
 		Writer.AddBinary(nullptr, 0);
 		FCbField Field = Writer.Save();
-		TestEqual_Conditional(TEXT("FCbWriter(Binary, Empty) Validate"), ValidateCompactBinary(Field.GetOuterBuffer(), ECbValidateMode::All), ECbValidateError::None)
+		if (TestEqual(TEXT("FCbWriter(Binary, Empty) Validate"), ValidateCompactBinary(Field.GetOuterBuffer(), ECbValidateMode::All), ECbValidateError::None))
 		{
-			TestFalse(TEXT("FCbWriter(Binary, Empty).HasName()"), Field.HasName());
-			TestTrue(TEXT("FCbWriter(Binary, Empty).IsBinary()"), Field.IsBinary());
-			TestTrue(TEXT("FCbWriter(Binary, Empty).AsBinaryView()"), Field.AsBinaryView().IsEmpty());
+			TEST_FALSE(TEXT("FCbWriter(Binary, Empty).HasName()"), Field.HasName());
+			TEST_TRUE(TEXT("FCbWriter(Binary, Empty).IsBinary()"), Field.IsBinary());
+			TEST_TRUE(TEXT("FCbWriter(Binary, Empty).AsBinaryView()"), Field.AsBinaryView().IsEmpty());
 		}
 	}
 
@@ -248,17 +261,17 @@ TEST_CASE("Core::Serialization::FCbWriter::Binary", "[Core][Serialization][Smoke
 		Writer.SetName("Binary"_ASV);
 		Writer.AddBinary(BinaryValue, sizeof(BinaryValue));
 		FCbField Field = Writer.Save();
-		TestEqual_Conditional(TEXT("FCbWriter(Binary, Array) Validate"), ValidateCompactBinary(Field.GetOuterBuffer(), ECbValidateMode::All), ECbValidateError::None)
+		if (TestEqual(TEXT("FCbWriter(Binary, Array) Validate"), ValidateCompactBinary(Field.GetOuterBuffer(), ECbValidateMode::All), ECbValidateError::None))
 		{
-			TestEqual(TEXT("FCbWriter(Binary, Array).GetName()"), Field.GetName(), "Binary"_U8SV);
-			TestTrue(TEXT("FCbWriter(Binary, Array).HasName()"), Field.HasName());
-			TestTrue(TEXT("FCbWriter(Binary, Array).IsBinary()"), Field.IsBinary());
-			TestTrue(TEXT("FCbWriter(Binary, Array).AsBinaryView()"), Field.AsBinaryView().EqualBytes(MakeMemoryView(BinaryValue)));
+			TEST_EQUAL(TEXT("FCbWriter(Binary, Array).GetName()"), Field.GetName(), "Binary"_U8SV);
+			TEST_TRUE(TEXT("FCbWriter(Binary, Array).HasName()"), Field.HasName());
+			TEST_TRUE(TEXT("FCbWriter(Binary, Array).IsBinary()"), Field.IsBinary());
+			TEST_TRUE(TEXT("FCbWriter(Binary, Array).AsBinaryView()"), Field.AsBinaryView().EqualBytes(MakeMemoryView(BinaryValue)));
 		}
 	}
 }
 
-TEST_CASE("Core::Serialization::FCbWriter::String", "[Core][Serialization][Smoke]")
+TEST_CASE_METHOD(FAutomationTestFixture, "Core::Serialization::FCbWriter::String", "[Core][Serialization][Smoke]")
 {
 	TCbWriter<256> Writer;
 
@@ -268,13 +281,13 @@ TEST_CASE("Core::Serialization::FCbWriter::String", "[Core][Serialization][Smoke
 		Writer.AddString(FAnsiStringView());
 		Writer.AddString(FWideStringView());
 		FCbFieldIterator Fields = Writer.Save();
-		TestEqual_Conditional(TEXT("FCbWriter(String, Empty) Validate"), ValidateCompactBinaryRange(Fields.GetOuterBuffer(), ECbValidateMode::All), ECbValidateError::None)
+		if (TestEqual(TEXT("FCbWriter(String, Empty) Validate"), ValidateCompactBinaryRange(Fields.GetOuterBuffer(), ECbValidateMode::All), ECbValidateError::None))
 		{
 			for (FCbFieldView Field : Fields)
 			{
-				TestFalse(TEXT("FCbWriter(String, Empty).HasName()"), Field.HasName());
-				TestTrue(TEXT("FCbWriter(String, Empty).IsString()"), Field.IsString());
-				TestTrue(TEXT("FCbWriter(String, Empty).AsString()"), Field.AsString().IsEmpty());
+				TEST_FALSE(TEXT("FCbWriter(String, Empty).HasName()"), Field.HasName());
+				TEST_TRUE(TEXT("FCbWriter(String, Empty).IsString()"), Field.IsString());
+				TEST_TRUE(TEXT("FCbWriter(String, Empty).AsString()"), Field.AsString().IsEmpty());
 			}
 		}
 	}
@@ -283,16 +296,16 @@ TEST_CASE("Core::Serialization::FCbWriter::String", "[Core][Serialization][Smoke
 	{
 		Writer.Reset();
 		Writer.SetName("String"_ASV).AddString("Value"_ASV);
-		Writer.SetName("String"_ASV).AddString(TEXT("Value"_SV));
+		Writer.SetName("String"_ASV).AddString(TEXTVIEW("Value"));
 		FCbFieldIterator Fields = Writer.Save();
-		TestEqual_Conditional(TEXT("FCbWriter(String, Basic) Validate"), ValidateCompactBinaryRange(Fields.GetOuterBuffer(), ECbValidateMode::All), ECbValidateError::None)
+		if (TestEqual(TEXT("FCbWriter(String, Basic) Validate"), ValidateCompactBinaryRange(Fields.GetOuterBuffer(), ECbValidateMode::All), ECbValidateError::None))
 		{
 			for (FCbFieldView Field : Fields)
 			{
-				TestEqual(TEXT("FCbWriter(String, Basic).GetName()"), Field.GetName(), "String"_U8SV);
-				TestTrue(TEXT("FCbWriter(String, Basic).HasName()"), Field.HasName());
-				TestTrue(TEXT("FCbWriter(String, Basic).IsString()"), Field.IsString());
-				TestEqual(TEXT("FCbWriter(String, Basic).AsString()"), Field.AsString(), "Value"_U8SV);
+				TEST_EQUAL(TEXT("FCbWriter(String, Basic).GetName()"), Field.GetName(), "String"_U8SV);
+				TEST_TRUE(TEXT("FCbWriter(String, Basic).HasName()"), Field.HasName());
+				TEST_TRUE(TEXT("FCbWriter(String, Basic).IsString()"), Field.IsString());
+				TEST_EQUAL(TEXT("FCbWriter(String, Basic).AsString()"), Field.AsString(), "Value"_U8SV);
 			}
 		}
 	}
@@ -309,11 +322,11 @@ TEST_CASE("Core::Serialization::FCbWriter::String", "[Core][Serialization][Smoke
 		Writer.AddString(Dots);
 		Writer.AddString(FString::ChrN(DotCount, TEXT('.')));
 		FCbFieldIterator Fields = Writer.Save();
-		TestEqual_Conditional(TEXT("FCbWriter(String, Long) Validate"), ValidateCompactBinaryRange(Fields.GetOuterBuffer(), ECbValidateMode::All), ECbValidateError::None)
+		if (TestEqual(TEXT("FCbWriter(String, Long) Validate"), ValidateCompactBinaryRange(Fields.GetOuterBuffer(), ECbValidateMode::All), ECbValidateError::None))
 		{
 			for (FCbFieldView Field : Fields)
 			{
-				TestEqual(TEXT("FCbWriter(String, Long).AsString()"), Field.AsString(), Dots.ToView());
+				TEST_EQUAL(TEXT("FCbWriter(String, Long).AsString()"), Field.AsString(), Dots.ToView());
 			}
 		}
 	}
@@ -325,66 +338,66 @@ TEST_CASE("Core::Serialization::FCbWriter::String", "[Core][Serialization][Smoke
 		Writer.AddString("\xf0\x9f\x98\x80"_ASV);
 		Writer.AddString(FWideStringView(Value, UE_ARRAY_COUNT(Value)));
 		FCbFieldIterator Fields = Writer.Save();
-		TestEqual_Conditional(TEXT("FCbWriter(String, Unicode) Validate"), ValidateCompactBinaryRange(Fields.GetOuterBuffer(), ECbValidateMode::All), ECbValidateError::None)
+		if (TestEqual(TEXT("FCbWriter(String, Unicode) Validate"), ValidateCompactBinaryRange(Fields.GetOuterBuffer(), ECbValidateMode::All), ECbValidateError::None))
 		{
 			for (FCbFieldView Field : Fields)
 			{
-				TestEqual(TEXT("FCbWriter(String, Unicode).AsString()"), Field.AsString(), "\xf0\x9f\x98\x80"_U8SV);
+				TEST_EQUAL(TEXT("FCbWriter(String, Unicode).AsString()"), Field.AsString(), "\xf0\x9f\x98\x80"_U8SV);
 			}
 		}
 	}
 }
 
 
-TEST_CASE("Core::Serialization::FCbWriter::Integer", "[Core][Serialization][Smoke]")
+TEST_CASE_METHOD(FAutomationTestFixture, "Core::Serialization::FCbWriter::Integer", "[Core][Serialization][Smoke]")
 {
 	TCbWriter<256> Writer;
 
-	auto TestInt32 = [&Writer](int32 Value)
+	auto TestInt32 = [this,&Writer](int32 Value)
 	{
 		Writer.Reset();
 		Writer.AddInteger(Value);
 		FCbField Field = Writer.Save();
-		TestEqual_Conditional(TEXT("FCbWriter(Integer, Int32) Validate"), ValidateCompactBinary(Field.GetOuterBuffer(), ECbValidateMode::All), ECbValidateError::None)
+		if (TestEqual(TEXT("FCbWriter(Integer, Int32) Validate"), ValidateCompactBinary(Field.GetOuterBuffer(), ECbValidateMode::All), ECbValidateError::None))
 		{
-			TestEqual(TEXT("FCbWriter(Integer, Int32) Value"), Field.AsInt32(), Value);
-			TestFalse(TEXT("FCbWriter(Integer, Int32) Error"), Field.HasError());
+			TEST_EQUAL(TEXT("FCbWriter(Integer, Int32) Value"), Field.AsInt32(), Value);
+			TEST_FALSE(TEXT("FCbWriter(Integer, Int32) Error"), Field.HasError());
 		}
 	};
 
-	auto TestUInt32 = [&Writer](uint32 Value)
+	auto TestUInt32 = [this, &Writer](uint32 Value)
 	{
 		Writer.Reset();
 		Writer.AddInteger(Value);
 		FCbField Field = Writer.Save();
-		TestEqual_Conditional(TEXT("FCbWriter(Integer, UInt32) Validate"), ValidateCompactBinary(Field.GetOuterBuffer(), ECbValidateMode::All), ECbValidateError::None)
+		if (TestEqual(TEXT("FCbWriter(Integer, UInt32) Validate"), ValidateCompactBinary(Field.GetOuterBuffer(), ECbValidateMode::All), ECbValidateError::None))
 		{
-			TestEqual(TEXT("FCbWriter(Integer, UInt32) Value"), Field.AsUInt32(), Value);
-			TestFalse(TEXT("FCbWriter(Integer, UInt32) Error"), Field.HasError());
+			TEST_EQUAL(TEXT("FCbWriter(Integer, UInt32) Value"), Field.AsUInt32(), Value);
+			TEST_FALSE(TEXT("FCbWriter(Integer, UInt32) Error"), Field.HasError());
 		}
 	};
 
-	auto TestInt64 = [&Writer](int64 Value)
+	auto TestInt64 = [this, &Writer](int64 Value)
 	{
 		Writer.Reset();
 		Writer.AddInteger(Value);
 		FCbField Field = Writer.Save();
-		TestEqual_Conditional(TEXT("FCbWriter(Integer, Int64) Validate"), ValidateCompactBinary(Field.GetOuterBuffer(), ECbValidateMode::All), ECbValidateError::None)
+		if (TestEqual(TEXT("FCbWriter(Integer, Int64) Validate"), ValidateCompactBinary(Field.GetOuterBuffer(), ECbValidateMode::All), ECbValidateError::None))
 		{
-			TestEqual(TEXT("FCbWriter(Integer, Int64) Value"), Field.AsInt64(), Value);
-			TestFalse(TEXT("FCbWriter(Integer, Int64) Error"), Field.HasError());
+			TEST_EQUAL(TEXT("FCbWriter(Integer, Int64) Value"), Field.AsInt64(), Value);
+			TEST_FALSE(TEXT("FCbWriter(Integer, Int64) Error"), Field.HasError());
 		}
 	};
 
-	auto TestUInt64 = [&Writer](uint64 Value)
+	auto TestUInt64 = [this, &Writer](uint64 Value)
 	{
 		Writer.Reset();
 		Writer.AddInteger(Value);
 		FCbField Field = Writer.Save();
-		TestEqual_Conditional(TEXT("FCbWriter(Integer, UInt64) Validate"), ValidateCompactBinary(Field.GetOuterBuffer(), ECbValidateMode::All), ECbValidateError::None)
+		if (TestEqual(TEXT("FCbWriter(Integer, UInt64) Validate"), ValidateCompactBinary(Field.GetOuterBuffer(), ECbValidateMode::All), ECbValidateError::None))
 		{
-			TestEqual(TEXT("FCbWriter(Integer, UInt64) Value"), Field.AsUInt64(), Value);
-			TestFalse(TEXT("FCbWriter(Integer, UInt64) Error"), Field.HasError());
+			TEST_EQUAL(TEXT("FCbWriter(Integer, UInt64) Value"), Field.AsUInt64(), Value);
+			TEST_FALSE(TEXT("FCbWriter(Integer, UInt64) Error"), Field.HasError());
 		}
 	};
 
@@ -422,7 +435,7 @@ TEST_CASE("Core::Serialization::FCbWriter::Integer", "[Core][Serialization][Smok
 	TestInt64(int64(0xffff'ffff'ffff'ffff));
 }
 
-TEST_CASE("Core::Serialization::FCbWriter::Float", "[Core][Serialization][Smoke]")
+TEST_CASE_METHOD(FAutomationTestFixture, "Core::Serialization::FCbWriter::Float", "[Core][Serialization][Smoke]")
 {
 	TCbWriter<256> Writer;
 
@@ -435,13 +448,13 @@ TEST_CASE("Core::Serialization::FCbWriter::Float", "[Core][Serialization][Smoke]
 			Writer.AddFloat(Value);
 		}
 		FCbFieldIterator Fields = Writer.Save();
-		TestEqual_Conditional(TEXT("FCbWriter(Float, Single) Validate"), ValidateCompactBinaryRange(Fields.GetOuterBuffer(), ECbValidateMode::All), ECbValidateError::None)
+		if (TestEqual(TEXT("FCbWriter(Float, Single) Validate"), ValidateCompactBinaryRange(Fields.GetOuterBuffer(), ECbValidateMode::All), ECbValidateError::None))
 		{
 			const float* CheckValue = Values;
 			for (FCbFieldView Field : Fields)
 			{
-				TestEqual(TEXT("FCbWriter(Float, Single).AsFloat()"), Field.AsFloat(), *CheckValue++);
-				TestFalse(TEXT("FCbWriter(Float, Single) Error"), Field.HasError());
+				TEST_EQUAL(TEXT("FCbWriter(Float, Single).AsFloat()"), Field.AsFloat(), *CheckValue++);
+				TEST_FALSE(TEXT("FCbWriter(Float, Single) Error"), Field.HasError());
 			}
 		}
 	}
@@ -455,19 +468,19 @@ TEST_CASE("Core::Serialization::FCbWriter::Float", "[Core][Serialization][Smoke]
 			Writer.AddFloat(Value);
 		}
 		FCbFieldIterator Fields = Writer.Save();
-		TestEqual_Conditional(TEXT("FCbWriter(Float, Double) Validate"), ValidateCompactBinaryRange(Fields.GetOuterBuffer(), ECbValidateMode::All), ECbValidateError::None)
+		if (TestEqual(TEXT("FCbWriter(Float, Double) Validate"), ValidateCompactBinaryRange(Fields.GetOuterBuffer(), ECbValidateMode::All), ECbValidateError::None))
 		{
 			const double* CheckValue = Values;
 			for (FCbFieldView Field : Fields)
 			{
-				TestEqual(TEXT("FCbWriter(Float, Double).AsDouble()"), Field.AsDouble(), *CheckValue++);
-				TestFalse(TEXT("FCbWriter(Float, Double) Error"), Field.HasError());
+				TEST_EQUAL(TEXT("FCbWriter(Float, Double).AsDouble()"), Field.AsDouble(), *CheckValue++);
+				TEST_FALSE(TEXT("FCbWriter(Float, Double) Error"), Field.HasError());
 			}
 		}
 	}
 }
 
-TEST_CASE("Core::Serialization::FCbWriter::Bool", "[Core][Serialization][Smoke]")
+TEST_CASE_METHOD(FAutomationTestFixture, "Core::Serialization::FCbWriter::Bool", "[Core][Serialization][Smoke]")
 {
 	TCbWriter<256> Writer;
 
@@ -477,13 +490,13 @@ TEST_CASE("Core::Serialization::FCbWriter::Bool", "[Core][Serialization][Smoke]"
 		Writer.AddBool(false);
 
 		FCbFieldIterator Fields = Writer.Save();
-		TestEqual_Conditional(TEXT("FCbWriter(Bool) Validate"), ValidateCompactBinaryRange(Fields.GetOuterBuffer(), ECbValidateMode::All), ECbValidateError::None)
+		if (TestEqual(TEXT("FCbWriter(Bool) Validate"), ValidateCompactBinaryRange(Fields.GetOuterBuffer(), ECbValidateMode::All), ECbValidateError::None))
 		{
-			TestTrue(TEXT("FCbWriter(Bool).AsBool()"), Fields.AsBool());
-			TestFalse(TEXT("FCbWriter(Bool) Error"), Fields.HasError());
+			TEST_TRUE(TEXT("FCbWriter(Bool).AsBool()"), Fields.AsBool());
+			TEST_FALSE(TEXT("FCbWriter(Bool) Error"), Fields.HasError());
 			++Fields;
-			TestFalse(TEXT("FCbWriter(Bool).AsBool()"), Fields.AsBool());
-			TestFalse(TEXT("FCbWriter(Bool) Error"), Fields.HasError());
+			TEST_FALSE(TEXT("FCbWriter(Bool).AsBool()"), Fields.AsBool());
+			TEST_FALSE(TEXT("FCbWriter(Bool) Error"), Fields.HasError());
 		}
 	}
 
@@ -504,11 +517,11 @@ TEST_CASE("Core::Serialization::FCbWriter::Bool", "[Core][Serialization][Smoke]"
 		Writer.EndObject();
 
 		FCbFieldIterator Fields = Writer.Save();
-		TestEqual(TEXT("FCbWriter(Bool, Uniform) Validate"), ValidateCompactBinaryRange(Fields.GetOuterBuffer(), ECbValidateMode::All), ECbValidateError::None);
+		TEST_EQUAL(TEXT("FCbWriter(Bool, Uniform) Validate"), ValidateCompactBinaryRange(Fields.GetOuterBuffer(), ECbValidateMode::All), ECbValidateError::None);
 	}
 }
 
-TEST_CASE("Core::Serialization::FCbWriter::Object Attachment", "[Core][Serialization][Smoke]")
+TEST_CASE_METHOD(FAutomationTestFixture, "Core::Serialization::FCbWriter::Object Attachment", "[Core][Serialization][Smoke]")
 {
 	TCbWriter<256> Writer;
 
@@ -522,23 +535,23 @@ TEST_CASE("Core::Serialization::FCbWriter::Object Attachment", "[Core][Serializa
 	}
 
 	FCbFieldIterator Fields = Writer.Save();
-	TestEqual_Conditional(TEXT("FCbWriter(ObjectAttachment) Validate"), ValidateCompactBinaryRange(Fields.GetOuterBuffer(), ECbValidateMode::All), ECbValidateError::None)
+	if (TestEqual(TEXT("FCbWriter(ObjectAttachment) Validate"), ValidateCompactBinaryRange(Fields.GetOuterBuffer(), ECbValidateMode::All), ECbValidateError::None))
 	{
 		const FIoHash* CheckValue = Values;
 		for (FCbFieldView Field : Fields)
 		{
-			TestEqual(TEXT("FCbWriter(ObjectAttachment).AsObjectAttachment()"), Field.AsObjectAttachment(), *CheckValue++);
-			TestFalse(TEXT("FCbWriter(ObjectAttachment) Error"), Field.HasError());
+			TEST_EQUAL(TEXT("FCbWriter(ObjectAttachment).AsObjectAttachment()"), Field.AsObjectAttachment(), *CheckValue++);
+			TEST_FALSE(TEXT("FCbWriter(ObjectAttachment) Error"), Field.HasError());
 		}
 	}
 }
 
-TEST_CASE("Core::Serialization::FCbWriter::Binary Attachment", "[Core][Serialization][Smoke]")
+TEST_CASE_METHOD(FAutomationTestFixture, "Core::Serialization::FCbWriter::Binary Attachment", "[Core][Serialization][Smoke]")
 {
 	TCbWriter<256> Writer;
 
 	const FIoHash::ByteArray ZeroBytes{};
-	const FIoHash::ByteArray SequentialBytes{ 1, 2, 3, 4, 5, 6, 7, 8, 9, 10, 11, 12, 13, 14, 15, 16, 17, 18, 19, 20 };
+	const FIoHash::ByteArray SequentialBytes{1, 2, 3, 4, 5, 6, 7, 8, 9, 10, 11, 12, 13, 14, 15, 16, 17, 18, 19, 20};
 
 	const FIoHash Values[] = { FIoHash(ZeroBytes), FIoHash(SequentialBytes) };
 	for (const FIoHash& Value : Values)
@@ -547,24 +560,24 @@ TEST_CASE("Core::Serialization::FCbWriter::Binary Attachment", "[Core][Serializa
 	}
 
 	FCbFieldIterator Fields = Writer.Save();
-	TestEqual_Conditional(TEXT("FCbWriter(BinaryAttachment) Validate"), ValidateCompactBinaryRange(Fields.GetOuterBuffer(), ECbValidateMode::All), ECbValidateError::None)
+	if (TestEqual(TEXT("FCbWriter(BinaryAttachment) Validate"), ValidateCompactBinaryRange(Fields.GetOuterBuffer(), ECbValidateMode::All), ECbValidateError::None))
 	{
 		const FIoHash* CheckValue = Values;
 		for (FCbFieldView Field : Fields)
 		{
-			TestEqual(TEXT("FCbWriter(BinaryAttachment).AsBinaryAttachment()"), Field.AsBinaryAttachment(), *CheckValue++);
-			TestFalse(TEXT("FCbWriter(BinaryAttachment) Error"), Field.HasError());
+			TEST_EQUAL(TEXT("FCbWriter(BinaryAttachment).AsBinaryAttachment()"), Field.AsBinaryAttachment(), *CheckValue++);
+			TEST_FALSE(TEXT("FCbWriter(BinaryAttachment) Error"), Field.HasError());
 		}
 	}
 }
 
 
-TEST_CASE("Core::Serialization::FCbWriter::Hash", "[Core][Serialization][Smoke]")
+TEST_CASE_METHOD(FAutomationTestFixture, "Core::Serialization::FCbWriter::Hash", "[Core][Serialization][Smoke]")
 {
 	TCbWriter<256> Writer;
 
 	const FIoHash::ByteArray ZeroBytes{};
-	const FIoHash::ByteArray SequentialBytes{ 1, 2, 3, 4, 5, 6, 7, 8, 9, 10, 11, 12, 13, 14, 15, 16, 17, 18, 19, 20 };
+	const FIoHash::ByteArray SequentialBytes{1, 2, 3, 4, 5, 6, 7, 8, 9, 10, 11, 12, 13, 14, 15, 16, 17, 18, 19, 20};
 
 	const FIoHash Values[] = { FIoHash(ZeroBytes), FIoHash(SequentialBytes) };
 	for (const FIoHash& Value : Values)
@@ -573,18 +586,18 @@ TEST_CASE("Core::Serialization::FCbWriter::Hash", "[Core][Serialization][Smoke]"
 	}
 
 	FCbFieldIterator Fields = Writer.Save();
-	TestEqual_Conditional(TEXT("FCbWriter(Hash) Validate"), ValidateCompactBinaryRange(Fields.GetOuterBuffer(), ECbValidateMode::All), ECbValidateError::None)
+	if (TestEqual(TEXT("FCbWriter(Hash) Validate"), ValidateCompactBinaryRange(Fields.GetOuterBuffer(), ECbValidateMode::All), ECbValidateError::None))
 	{
 		const FIoHash* CheckValue = Values;
 		for (FCbFieldView Field : Fields)
 		{
-			TestEqual(TEXT("FCbWriter(Hash).AsHash()"), Field.AsHash(), *CheckValue++);
-			TestFalse(TEXT("FCbWriter(Hash) Error"), Field.HasError());
+			TEST_EQUAL(TEXT("FCbWriter(Hash).AsHash()"), Field.AsHash(), *CheckValue++);
+			TEST_FALSE(TEXT("FCbWriter(Hash) Error"), Field.HasError());
 		}
 	}
 }
 
-TEST_CASE("Core::Serialization::FCbWriter::Uuid", "[Core][Serialization][Smoke]")
+TEST_CASE_METHOD(FAutomationTestFixture, "Core::Serialization::FCbWriter::Uuid", "[Core][Serialization][Smoke]")
 {
 	TCbWriter<256> Writer;
 
@@ -595,19 +608,19 @@ TEST_CASE("Core::Serialization::FCbWriter::Uuid", "[Core][Serialization][Smoke]"
 	}
 
 	FCbFieldIterator Fields = Writer.Save();
-	TestEqual_Conditional(TEXT("FCbWriter(Uuid) Validate"), ValidateCompactBinaryRange(Fields.GetOuterBuffer(), ECbValidateMode::All), ECbValidateError::None)
+	if (TestEqual(TEXT("FCbWriter(Uuid) Validate"), ValidateCompactBinaryRange(Fields.GetOuterBuffer(), ECbValidateMode::All), ECbValidateError::None))
 	{
 		const FGuid* CheckValue = Values;
 		for (FCbFieldView Field : Fields)
 		{
-			TestEqual(TEXT("FCbWriter(Uuid).AsUuid()"), Field.AsUuid(), *CheckValue++);
-			TestFalse(TEXT("FCbWriter(Uuid) Error"), Field.HasError());
+			TEST_EQUAL(TEXT("FCbWriter(Uuid).AsUuid()"), Field.AsUuid(), *CheckValue++);
+			TEST_FALSE(TEXT("FCbWriter(Uuid) Error"), Field.HasError());
 		}
 	}
 }
 
 
-TEST_CASE("Core::Serialization::FCbWriter::Date Time", "[Core][Serialization][Smoke]")
+TEST_CASE_METHOD(FAutomationTestFixture, "Core::Serialization::FCbWriter::Date Time", "[Core][Serialization][Smoke]")
 {
 	TCbWriter<256> Writer;
 
@@ -618,18 +631,18 @@ TEST_CASE("Core::Serialization::FCbWriter::Date Time", "[Core][Serialization][Sm
 	}
 
 	FCbFieldIterator Fields = Writer.Save();
-	TestEqual_Conditional(TEXT("FCbWriter(DateTime) Validate"), ValidateCompactBinaryRange(Fields.GetOuterBuffer(), ECbValidateMode::All), ECbValidateError::None)
+	if (TestEqual(TEXT("FCbWriter(DateTime) Validate"), ValidateCompactBinaryRange(Fields.GetOuterBuffer(), ECbValidateMode::All), ECbValidateError::None))
 	{
 		const FDateTime* CheckValue = Values;
 		for (FCbFieldView Field : Fields)
 		{
-			TestEqual(TEXT("FCbWriter(DateTime).AsDateTime()"), Field.AsDateTime(), *CheckValue++);
-			TestFalse(TEXT("FCbWriter(DateTime) Error"), Field.HasError());
+			TEST_EQUAL(TEXT("FCbWriter(DateTime).AsDateTime()"), Field.AsDateTime(), *CheckValue++);
+			TEST_FALSE(TEXT("FCbWriter(DateTime) Error"), Field.HasError());
 		}
 	};
 }
 
-TEST_CASE("Core::Serialization::FCbWriter::Time Span", "[Core][Serialization][Smoke]")
+TEST_CASE_METHOD(FAutomationTestFixture, "Core::Serialization::FCbWriter::Time Span", "[Core][Serialization][Smoke]")
 {
 	TCbWriter<256> Writer;
 
@@ -640,18 +653,18 @@ TEST_CASE("Core::Serialization::FCbWriter::Time Span", "[Core][Serialization][Sm
 	}
 
 	FCbFieldIterator Fields = Writer.Save();
-	TestEqual_Conditional(TEXT("FCbWriter(TimeSpan) Validate"), ValidateCompactBinaryRange(Fields.GetOuterBuffer(), ECbValidateMode::All), ECbValidateError::None)
+	if (TestEqual(TEXT("FCbWriter(TimeSpan) Validate"), ValidateCompactBinaryRange(Fields.GetOuterBuffer(), ECbValidateMode::All), ECbValidateError::None))
 	{
 		const FTimespan* CheckValue = Values;
 		for (FCbFieldView Field : Fields)
 		{
-			TestEqual(TEXT("FCbWriter(TimeSpan).AsTimeSpan()"), Field.AsTimeSpan(), *CheckValue++);
-			TestFalse(TEXT("FCbWriter(TimeSpan) Error"), Field.HasError());
+			TEST_EQUAL(TEXT("FCbWriter(TimeSpan).AsTimeSpan()"), Field.AsTimeSpan(), *CheckValue++);
+			TEST_FALSE(TEXT("FCbWriter(TimeSpan) Error"), Field.HasError());
 		}
 	}
 }
 
-TEST_CASE("Core::Serialization::FCbWriter::ObjectId", "[Core][Serialization][Smoke]")
+TEST_CASE_METHOD(FAutomationTestFixture, "Core::Serialization::FCbWriter::ObjectId", "[Core][Serialization][Smoke]")
 {
 	TCbWriter<256> Writer;
 
@@ -662,18 +675,18 @@ TEST_CASE("Core::Serialization::FCbWriter::ObjectId", "[Core][Serialization][Smo
 	}
 
 	FCbFieldIterator Fields = Writer.Save();
-	TestEqual_Conditional(TEXT("FCbWriter(ObjectId) Validate"), ValidateCompactBinaryRange(Fields.GetOuterBuffer(), ECbValidateMode::All), ECbValidateError::None)
+	if (TestEqual(TEXT("FCbWriter(ObjectId) Validate"), ValidateCompactBinaryRange(Fields.GetOuterBuffer(), ECbValidateMode::All), ECbValidateError::None))
 	{
 		const FCbObjectId* CheckValue = Values;
 		for (FCbFieldView Field : Fields)
 		{
-			TestEqual(TEXT("FCbWriter(ObjectId).AsObjectId()"), Field.AsObjectId(), *CheckValue++);
-			TestFalse(TEXT("FCbWriter(ObjectId) Error"), Field.HasError());
+			TEST_EQUAL(TEXT("FCbWriter(ObjectId).AsObjectId()"), Field.AsObjectId(), *CheckValue++);
+			TEST_FALSE(TEXT("FCbWriter(ObjectId) Error"), Field.HasError());
 		}
 	}
 }
 
-TEST_CASE("Core::Serialization::FCbWriter::Custom By Id", "[Core][Serialization][Smoke]")
+TEST_CASE_METHOD(FAutomationTestFixture, "Core::Serialization::FCbWriter::Custom By Id", "[Core][Serialization][Smoke]")
 {
 	TCbWriter<256> Writer;
 
@@ -694,19 +707,19 @@ TEST_CASE("Core::Serialization::FCbWriter::Custom By Id", "[Core][Serialization]
 	}
 
 	FCbFieldIterator Fields = Writer.Save();
-	TestEqual_Conditional(TEXT("FCbWriter(CustomById) Validate"), ValidateCompactBinaryRange(Fields.GetOuterBuffer(), ECbValidateMode::All), ECbValidateError::None)
+	if (TestEqual(TEXT("FCbWriter(CustomById) Validate"), ValidateCompactBinaryRange(Fields.GetOuterBuffer(), ECbValidateMode::All), ECbValidateError::None))
 	{
 		const FCustomValue* CheckValue = Values;
 		for (FCbFieldView Field : Fields)
 		{
-			TestTrue(TEXT("FCbWriter(CustomById).AsCustom()"), Field.AsCustom(CheckValue->Type).EqualBytes(MakeMemoryView(CheckValue->Bytes)));
-			TestFalse(TEXT("FCbWriter(CustomById) Error"), Field.HasError());
+			TEST_TRUE(TEXT("FCbWriter(CustomById).AsCustom()"), Field.AsCustom(CheckValue->Type).EqualBytes(MakeMemoryView(CheckValue->Bytes)));
+			TEST_FALSE(TEXT("FCbWriter(CustomById) Error"), Field.HasError());
 			++CheckValue;
 		}
 	}
 }
 
-TEST_CASE("Core::Serialization::FCbWriter::Custom By Name", "[Core][Serialization][Smoke]")
+TEST_CASE_METHOD(FAutomationTestFixture, "Core::Serialization::FCbWriter::Custom By Name", "[Core][Serialization][Smoke]")
 {
 	TCbWriter<256> Writer;
 
@@ -727,19 +740,19 @@ TEST_CASE("Core::Serialization::FCbWriter::Custom By Name", "[Core][Serializatio
 	}
 
 	FCbFieldIterator Fields = Writer.Save();
-	TestEqual_Conditional(TEXT("FCbWriter(CustomByName) Validate"), ValidateCompactBinaryRange(Fields.GetOuterBuffer(), ECbValidateMode::All), ECbValidateError::None)
+	if (TestEqual(TEXT("FCbWriter(CustomByName) Validate"), ValidateCompactBinaryRange(Fields.GetOuterBuffer(), ECbValidateMode::All), ECbValidateError::None))
 	{
 		const FCustomValue* CheckValue = Values;
 		for (FCbFieldView Field : Fields)
 		{
-			TestTrue(TEXT("FCbWriter(CustomByName).AsCustom()"), Field.AsCustom(CheckValue->Type).EqualBytes(MakeMemoryView(CheckValue->Bytes)));
-			TestFalse(TEXT("FCbWriter(CustomByName) Error"), Field.HasError());
+			TEST_TRUE(TEXT("FCbWriter(CustomByName).AsCustom()"), Field.AsCustom(CheckValue->Type).EqualBytes(MakeMemoryView(CheckValue->Bytes)));
+			TEST_FALSE(TEXT("FCbWriter(CustomByName) Error"), Field.HasError());
 			++CheckValue;
 		}
 	}
 }
 
-TEST_CASE("Core::Serialization::FCbWriter::Complex", "[Core][Serialization][Smoke]")
+TEST_CASE_METHOD(FAutomationTestFixture, "Core::Serialization::FCbWriter::Complex", "[Core][Serialization][Smoke]")
 {
 	FCbObject Object;
 	FBufferArchive Archive;
@@ -831,7 +844,7 @@ TEST_CASE("Core::Serialization::FCbWriter::Complex", "[Core][Serialization][Smok
 		Writer.AddTimeSpan("TimeSpan"_ASV, FTimespan(1, 2, 4, 8));
 
 		Writer.AddObjectId("ObjectIdZero"_ASV, FCbObjectId());
-		Writer.AddObjectId("ObjectId"_ASV, FCbObjectId(MakeMemoryView<uint8>({ 1, 2, 3, 4, 5, 6, 7, 8, 9, 10, 11, 12 })));
+		Writer.AddObjectId("ObjectId"_ASV, FCbObjectId(MakeMemoryView<uint8>({1, 2, 3, 4, 5, 6, 7, 8, 9, 10, 11, 12})));
 
 		Writer.BeginObject("NestedObjects"_ASV);
 		{
@@ -910,27 +923,27 @@ TEST_CASE("Core::Serialization::FCbWriter::Complex", "[Core][Serialization][Smok
 		Object = Writer.Save().AsObject();
 
 		Writer.Save(Archive);
-		TestEqual(TEXT("FCbWriter(Complex).Save(Ar)->Num()"), uint64(Archive.Num()), Writer.GetSaveSize());
+		TEST_EQUAL(TEXT("FCbWriter(Complex).Save(Ar)->Num()"), uint64(Archive.Num()), Writer.GetSaveSize());
 	}
 
-	TestEqual(TEXT("FCbWriter(Complex).Save()->Validate"),
+	TEST_EQUAL(TEXT("FCbWriter(Complex).Save()->Validate"),
 		ValidateCompactBinary(Object.GetOuterBuffer(), ECbValidateMode::All), ECbValidateError::None);
 
-	TestEqual(TEXT("FCbWriter(Complex).Save(Ar)->Validate"),
+	TEST_EQUAL(TEXT("FCbWriter(Complex).Save(Ar)->Validate"),
 		ValidateCompactBinary(MakeMemoryView(Archive), ECbValidateMode::All), ECbValidateError::None);
 }
 
 
-TEST_CASE("Core::Serialization::FCbWriter::Owned Read Only", "[Core][Serialization][Smoke]")
+TEST_CASE_METHOD(FAutomationTestFixture, "Core::Serialization::FCbWriter::Owned Read Only", "[Core][Serialization][Smoke]")
 {
 	FCbWriter Writer;
 	Writer.BeginObject();
 	Writer.EndObject();
 	FCbObject Object = Writer.Save().AsObject();
-	TestTrue(TEXT("FCbWriter().Save().IsOwned()"), Object.IsOwned());
+	TEST_TRUE(TEXT("FCbWriter().Save().IsOwned()"), Object.IsOwned());
 }
 
-TEST_CASE("Core::Serialization::FCbWriter::Stream", "[Core][Serialization][Smoke]")
+TEST_CASE_METHOD(FAutomationTestFixture, "Core::Serialization::FCbWriter::Stream", "[Core][Serialization][Smoke]")
 {
 	FCbObject Object;
 	{
@@ -956,7 +969,7 @@ TEST_CASE("Core::Serialization::FCbWriter::Stream", "[Core][Serialization][Smoke
 		Writer
 			<< "AnsiString"_ASV << "AnsiValue"_ASV
 			<< "AnsiStringLiteral"_ASV << "AnsiValue"
-			<< "WideString"_ASV << TEXT("WideValue"_SV)
+			<< "WideString"_ASV << TEXTVIEW("WideValue")
 			<< "WideStringLiteral"_ASV << TEXT("WideValue");
 		Writer.EndObject();
 
@@ -979,7 +992,7 @@ TEST_CASE("Core::Serialization::FCbWriter::Stream", "[Core][Serialization][Smoke
 		Writer << "TimeSpan"_ASV << FTimespan(1, 2, 4, 8);
 
 		Writer << "ObjectIdZero"_ASV << FCbObjectId();
-		Writer << "ObjectId"_ASV << FCbObjectId(MakeMemoryView<uint8>({ 1, 2, 3, 4, 5, 6, 7, 8, 9, 10, 11, 12 }));
+		Writer << "ObjectId"_ASV << FCbObjectId(MakeMemoryView<uint8>({1, 2, 3, 4, 5, 6, 7, 8, 9, 10, 11, 12}));
 
 		Writer << "LiteralName" << nullptr;
 
@@ -987,10 +1000,10 @@ TEST_CASE("Core::Serialization::FCbWriter::Stream", "[Core][Serialization][Smoke
 		Object = Writer.Save().AsObject();
 	}
 
-	TestEqual(TEXT("FCbWriter(Stream) Validate"), ValidateCompactBinary(Object.GetOuterBuffer(), ECbValidateMode::All), ECbValidateError::None);
+	TEST_EQUAL(TEXT("FCbWriter(Stream) Validate"), ValidateCompactBinary(Object.GetOuterBuffer(), ECbValidateMode::All), ECbValidateError::None);
 }
 
-TEST_CASE("Core::Serialization::FCbWriter::State", "[Core][Serialization][Smoke]")
+TEST_CASE_METHOD(FAutomationTestFixture, "Core::Serialization::FCbWriter::State", "[Core][Serialization][Smoke]")
 {
 	FCbWriter Writer;
 
diff --git a/Engine/Source/Runtime/Core/Tests/Serialization/MemoryReaderTest.cpp b/Engine/Source/Runtime/Core/Tests/Serialization/MemoryReaderTest.cpp
--- a/Engine/Source/Runtime/Core/Tests/Serialization/MemoryReaderTest.cpp
+++ b/Engine/Source/Runtime/Core/Tests/Serialization/MemoryReaderTest.cpp
@@ -5,7 +5,7 @@
 #include "Serialization/MemoryWriter.h"
 #include "TestHarness.h"
 
-TEST_CASE("Core::Serialization::MemoryReader::Memory Reader", "[Core][Serialization][Smoke]")
+TEST_CASE_METHOD(FAutomationTestFixture, "Core::Serialization::MemoryReader::Memory Reader", "[Core][Serialization][Smoke]")
 {
 	// Test reading uint64 byte swapped, ensuring that bytes really get swapped.
 	{
@@ -20,7 +20,7 @@ TEST_CASE("Core::Serialization::MemoryReader::Memory Reader", "[Core][Serializat
 		Reader.SetByteSwapping(true);
 		uint64 ReadValue;
 		Reader << ReadValue;
-		TestTrue("Test reading uint64 in swapped byte order.", ReadValue == ExpectedValue);
+		TEST_TRUE("Test reading uint64 in swapped byte order.", ReadValue == ExpectedValue);
 	}
 
 	// Write all supported types and read them back. (Assuming that MemoryWriterTest tests don't fail).
@@ -144,33 +144,33 @@ TEST_CASE("Core::Serialization::MemoryReader::Memory Reader", "[Core][Serializat
 		Reader << ReadSwapUtf16Str;
 
 		// Validate that the value read are the expected ones.
-		TestTrue("Test reading 'u8' from byte stream",             ReadValueU8      == TestValueU8);
-		TestTrue("Test reading 'u16' from byte stream",            ReadValueU16     == TestValueU16);
-		TestTrue("Test reading 'u32' from byte stream",            ReadValueU32     == TestValueU32);
-		TestTrue("Test reading 'u64' from byte stream",            ReadValueU64     == TestValueU64);
-		TestTrue("Test reading 's8' from byte stream",             ReadValueS8      == TestValueS8);
-		TestTrue("Test reading 's16' from byte stream",            ReadValueS16     == TestValueS16);
-		TestTrue("Test reading 's32' from byte stream",            ReadValueS32     == TestValueS32);
-		TestTrue("Test reading 's64' from byte stream",            ReadValueS64     == TestValueS64);
-		TestTrue("Test reading 'foat' from byte stream",           ReadValueF       == TestValueF);
-		TestTrue("Test reading 'double' from byte stream",         ReadValueD       == TestValueD);
-		TestTrue("Test reading 'bool' from byte stream",           ReadValueB       == TestValueB);
-		TestTrue("Test reading 'wchar' from byte stream",          ReadValueCh      == TestValueCh);
-		TestTrue("Test reading 'ansi str' from byte stream",       ReadAnsiStr      == TestAnsiStr);
-		TestTrue("Test reading 'utf16 str' from byte stream",      ReadUtf16Str     == TestUtf16Str);
-		TestTrue("Test reading 'u8-swapped' from byte stream",     ReadValueSwapU8  == TestValueU8);
-		TestTrue("Test reading 'u16-swapped' from byte stream",    ReadValueSwapU16 == TestValueU16);
-		TestTrue("Test reading 'u32-swapped' from byte stream",    ReadValueSwapU32 == TestValueU32);
-		TestTrue("Test reading 'u64-swapped' from byte stream",    ReadValueSwapU64 == TestValueU64);
-		TestTrue("Test reading 's8-swapped' from byte stream",     ReadValueSwapS8  == TestValueS8);
-		TestTrue("Test reading 's16'-swapped' from byte stream",   ReadValueSwapS16 == TestValueS16);
-		TestTrue("Test reading 's32-swapped' from byte stream",    ReadValueSwapS32 == TestValueS32);
-		TestTrue("Test reading 's64-swapped' from byte stream",    ReadValueSwapS64 == TestValueS64);
-		TestTrue("Test reading 'float-swapped' from byte stream",  ReadValueSwapF   == TestValueF);
-		TestTrue("Test reading 'double-swapped' from byte stream", ReadValueSwapD   == TestValueD);
-		TestTrue("Test reading 'bool-swapped' from byte stream",   ReadValueSwapB   == TestValueB);
-		TestTrue("Test reading 'wchar-swapped' from byte stream",  ReadValueSwapCh  == TestValueCh);
-		TestTrue("Test reading 'ansi-swapped' from byte stream",   ReadSwapAnsiStr  == TestAnsiStr);
-		TestTrue("Test reading 'utf16-swapped' from byte stream",  ReadSwapUtf16Str == TestUtf16Str);
+		TEST_TRUE("Test reading 'u8' from byte stream",             ReadValueU8      == TestValueU8);
+		TEST_TRUE("Test reading 'u16' from byte stream",            ReadValueU16     == TestValueU16);
+		TEST_TRUE("Test reading 'u32' from byte stream",            ReadValueU32     == TestValueU32);
+		TEST_TRUE("Test reading 'u64' from byte stream",            ReadValueU64     == TestValueU64);
+		TEST_TRUE("Test reading 's8' from byte stream",             ReadValueS8      == TestValueS8);
+		TEST_TRUE("Test reading 's16' from byte stream",            ReadValueS16     == TestValueS16);
+		TEST_TRUE("Test reading 's32' from byte stream",            ReadValueS32     == TestValueS32);
+		TEST_TRUE("Test reading 's64' from byte stream",            ReadValueS64     == TestValueS64);
+		TEST_TRUE("Test reading 'foat' from byte stream",           ReadValueF       == TestValueF);
+		TEST_TRUE("Test reading 'double' from byte stream",         ReadValueD       == TestValueD);
+		TEST_TRUE("Test reading 'bool' from byte stream",           ReadValueB       == TestValueB);
+		TEST_TRUE("Test reading 'wchar' from byte stream",          ReadValueCh      == TestValueCh);
+		TEST_TRUE("Test reading 'ansi str' from byte stream",       ReadAnsiStr      == TestAnsiStr);
+		TEST_TRUE("Test reading 'utf16 str' from byte stream",      ReadUtf16Str     == TestUtf16Str);
+		TEST_TRUE("Test reading 'u8-swapped' from byte stream",     ReadValueSwapU8  == TestValueU8);
+		TEST_TRUE("Test reading 'u16-swapped' from byte stream",    ReadValueSwapU16 == TestValueU16);
+		TEST_TRUE("Test reading 'u32-swapped' from byte stream",    ReadValueSwapU32 == TestValueU32);
+		TEST_TRUE("Test reading 'u64-swapped' from byte stream",    ReadValueSwapU64 == TestValueU64);
+		TEST_TRUE("Test reading 's8-swapped' from byte stream",     ReadValueSwapS8  == TestValueS8);
+		TEST_TRUE("Test reading 's16'-swapped' from byte stream",   ReadValueSwapS16 == TestValueS16);
+		TEST_TRUE("Test reading 's32-swapped' from byte stream",    ReadValueSwapS32 == TestValueS32);
+		TEST_TRUE("Test reading 's64-swapped' from byte stream",    ReadValueSwapS64 == TestValueS64);
+		TEST_TRUE("Test reading 'float-swapped' from byte stream",  ReadValueSwapF   == TestValueF);
+		TEST_TRUE("Test reading 'double-swapped' from byte stream", ReadValueSwapD   == TestValueD);
+		TEST_TRUE("Test reading 'bool-swapped' from byte stream",   ReadValueSwapB   == TestValueB);
+		TEST_TRUE("Test reading 'wchar-swapped' from byte stream",  ReadValueSwapCh  == TestValueCh);
+		TEST_TRUE("Test reading 'ansi-swapped' from byte stream",   ReadSwapAnsiStr  == TestAnsiStr);
+		TEST_TRUE("Test reading 'utf16-swapped' from byte stream",  ReadSwapUtf16Str == TestUtf16Str);
 	}
 }
diff --git a/Engine/Source/Runtime/Core/Tests/Serialization/MemoryWriterTest.cpp b/Engine/Source/Runtime/Core/Tests/Serialization/MemoryWriterTest.cpp
--- a/Engine/Source/Runtime/Core/Tests/Serialization/MemoryWriterTest.cpp
+++ b/Engine/Source/Runtime/Core/Tests/Serialization/MemoryWriterTest.cpp
@@ -11,7 +11,7 @@ namespace MemoryWriterTestUtil
 template <typename T> struct TIsBoolean { enum { Value = false }; };
 template <> struct TIsBoolean<bool> { enum { Value = true }; };
 
-class MemoryWriterTester
+class MemoryWriterTester : public FAutomationTestFixture
 {
 public:
 	MemoryWriterTester()
@@ -25,8 +25,8 @@ public:
 		const T Original = Value;
 		Writer << Value;
 		CheckWrittenByteCount<T>(Bytes.Num());
-		checkf(Value == Original, TEXT("The writer unexpectedly modified the input value")); // Ensure the written value wasn't changed during the write. (Regression check when writing byte swapped)
-		checkf(FMemory::Memcmp(Bytes.GetData(), &Original, sizeof(Original)) == 0, TEXT("The written value doesn't match the expected one."));
+		CHECK(Value == Original); // Ensure the written value wasn't changed during the write. (Regression check when writing byte swapped)
+		CHECK(FMemory::Memcmp(Bytes.GetData(), &Original, sizeof(Original)) == 0);
 	}
 
 	template<typename T>
@@ -36,8 +36,8 @@ public:
 		Writer.SetByteSwapping(true);
 		Writer << Value;
 		CheckWrittenByteCount<T>(Bytes.Num());
-		TestTrue(TEXT("The writer unexpectedly modified the input value"), Value == Original); // Ensure the written value wasn't changed during the write. (Regression check when writing byte swapped)
-		TestTrue(TEXT("The written value doesn't match the swapped value."), FMemory::Memcmp(Bytes.GetData(), &Swapped, sizeof(Swapped)) == 0);
+		CHECK(Value == Original); // Ensure the written value wasn't changed during the write. (Regression check when writing byte swapped)
+		CHECK(FMemory::Memcmp(Bytes.GetData(), &Swapped, sizeof(Swapped)) == 0);
 	}
 
 	template<typename T>
@@ -46,11 +46,11 @@ public:
 		if (TIsBoolean<T>::Value)
 		{
 			// Boolean are written as 4 bytes integer.
-			TestTrue(TEXT("Unexpected number of bytes written by the writer"), WrittenCount == sizeof(int32));
+			CHECK(WrittenCount == sizeof(int32));
 		}
 		else
 		{
-			TestTrue(TEXT("Unexpected number of bytes written by the writer"), WrittenCount == sizeof(T));
+			CHECK(WrittenCount == sizeof(T));
 		}
 	}
 
@@ -61,7 +61,7 @@ public:
 
 }
 
-TEST_CASE("Core::Serialization::MemoryWriter::Memory Writer", "[Core][Serialization][Smoke]")
+TEST_CASE_METHOD(FAutomationTestFixture, "Core::Serialization::MemoryWriter::Memory Writer", "[Core][Serialization][Smoke]")
 {
 	// Keeps the 'official' test value as const, to prevent overwriting them.
 	const uint8  TestValueU8   = 0x12;
@@ -77,119 +77,121 @@ TEST_CASE("Core::Serialization::MemoryWriter::Memory Writer", "[Core][Serializat
 	const bool   TestValueB    = true;
 	const WIDECHAR TestValueCh = 0xF2;
 
-	// Platform endianness tests.
+	SECTION("Platform endianness tests")
 	{
-		MemoryWriterTestUtil::MemoryWriterTester Tester;
-		Tester.TestWritePlatformByteOrder(TestValueU8);
-	}
+		{
+			MemoryWriterTestUtil::MemoryWriterTester Tester;
+			Tester.TestWritePlatformByteOrder(TestValueU8);
+		}
 
-	{
-		MemoryWriterTestUtil::MemoryWriterTester Tester;
-		Tester.TestWritePlatformByteOrder(TestValueS8);
-	}
+		{
+			MemoryWriterTestUtil::MemoryWriterTester Tester;
+			Tester.TestWritePlatformByteOrder(TestValueS8);
+		}
 
-	{
-		MemoryWriterTestUtil::MemoryWriterTester Tester;
-		Tester.TestWritePlatformByteOrder(TestValueU16);
-	}
+		{
+			MemoryWriterTestUtil::MemoryWriterTester Tester;
+			Tester.TestWritePlatformByteOrder(TestValueU16);
+		}
 
-	{
-		MemoryWriterTestUtil::MemoryWriterTester Tester;
-		Tester.TestWritePlatformByteOrder(TestValueS16);
-	}
+		{
+			MemoryWriterTestUtil::MemoryWriterTester Tester;
+			Tester.TestWritePlatformByteOrder(TestValueS16);
+		}
 
-	{
-		MemoryWriterTestUtil::MemoryWriterTester Tester;
-		Tester.TestWritePlatformByteOrder(TestValueU32);
-	}
+		{
+			MemoryWriterTestUtil::MemoryWriterTester Tester;
+			Tester.TestWritePlatformByteOrder(TestValueU32);
+		}
 
-	{
-		MemoryWriterTestUtil::MemoryWriterTester Tester;
-		Tester.TestWritePlatformByteOrder(TestValueS32);
-	}
+		{
+			MemoryWriterTestUtil::MemoryWriterTester Tester;
+			Tester.TestWritePlatformByteOrder(TestValueS32);
+		}
 
-	{
-		MemoryWriterTestUtil::MemoryWriterTester Tester;
-		Tester.TestWritePlatformByteOrder(TestValueU64);
-	}
+		{
+			MemoryWriterTestUtil::MemoryWriterTester Tester;
+			Tester.TestWritePlatformByteOrder(TestValueU64);
+		}
 
-	{
-		MemoryWriterTestUtil::MemoryWriterTester Tester;
-		Tester.TestWritePlatformByteOrder(TestValueS64);
-	}
+		{
+			MemoryWriterTestUtil::MemoryWriterTester Tester;
+			Tester.TestWritePlatformByteOrder(TestValueS64);
+		}
 
-	{
-		MemoryWriterTestUtil::MemoryWriterTester Tester;
-		Tester.TestWritePlatformByteOrder(TestValueF);
-	}
+		{
+			MemoryWriterTestUtil::MemoryWriterTester Tester;
+			Tester.TestWritePlatformByteOrder(TestValueF);
+		}
 
-	{
-		MemoryWriterTestUtil::MemoryWriterTester Tester;
-		Tester.TestWritePlatformByteOrder(TestValueD);
-	}
+		{
+			MemoryWriterTestUtil::MemoryWriterTester Tester;
+			Tester.TestWritePlatformByteOrder(TestValueD);
+		}
 
-	{
-		MemoryWriterTestUtil::MemoryWriterTester Tester;
-		Tester.TestWritePlatformByteOrder(TestValueB);
-	}
+		{
+			MemoryWriterTestUtil::MemoryWriterTester Tester;
+			Tester.TestWritePlatformByteOrder(TestValueB);
+		}
 
-	// Non Platform endianness tests. (Byte swapping)
-	{
-		MemoryWriterTestUtil::MemoryWriterTester Tester;
-		Tester.TestWriteSwappedByteOrder(TestValueU8, TestValueU8);
-	}
+		// Non Platform endianness tests. (Byte swapping)
+		{
+			MemoryWriterTestUtil::MemoryWriterTester Tester;
+			Tester.TestWriteSwappedByteOrder(TestValueU8, TestValueU8);
+		}
 
-	{
-		MemoryWriterTestUtil::MemoryWriterTester Tester;
-		Tester.TestWriteSwappedByteOrder(TestValueS8, TestValueS8);
-	}
+		{
+			MemoryWriterTestUtil::MemoryWriterTester Tester;
+			Tester.TestWriteSwappedByteOrder(TestValueS8, TestValueS8);
+		}
 
-	{
-		MemoryWriterTestUtil::MemoryWriterTester Tester;
-		Tester.TestWriteSwappedByteOrder(TestValueU16, ByteSwap(TestValueU16));
-	}
+		{
+			MemoryWriterTestUtil::MemoryWriterTester Tester;
+			Tester.TestWriteSwappedByteOrder(TestValueU16, ByteSwap(TestValueU16));
+		}
 
-	{
-		MemoryWriterTestUtil::MemoryWriterTester Tester;
-		Tester.TestWriteSwappedByteOrder(TestValueS16, ByteSwap(TestValueS16));
-	}
+		{
+			MemoryWriterTestUtil::MemoryWriterTester Tester;
+			Tester.TestWriteSwappedByteOrder(TestValueS16, ByteSwap(TestValueS16));
+		}
 
-	{
-		MemoryWriterTestUtil::MemoryWriterTester Tester;
-		Tester.TestWriteSwappedByteOrder(TestValueU32, ByteSwap(TestValueU32));
-	}
+		{
+			MemoryWriterTestUtil::MemoryWriterTester Tester;
+			Tester.TestWriteSwappedByteOrder(TestValueU32, ByteSwap(TestValueU32));
+		}
 
-	{
-		MemoryWriterTestUtil::MemoryWriterTester Tester;
-		Tester.TestWriteSwappedByteOrder(TestValueS32, ByteSwap(TestValueS32));
-	}
+		{
+			MemoryWriterTestUtil::MemoryWriterTester Tester;
+			Tester.TestWriteSwappedByteOrder(TestValueS32, ByteSwap(TestValueS32));
+		}
 
-	{
-		MemoryWriterTestUtil::MemoryWriterTester Tester;
-		Tester.TestWriteSwappedByteOrder(TestValueU64, ByteSwap(TestValueU64));
-	}
+		{
+			MemoryWriterTestUtil::MemoryWriterTester Tester;
+			Tester.TestWriteSwappedByteOrder(TestValueU64, ByteSwap(TestValueU64));
+		}
 
-	{
-		MemoryWriterTestUtil::MemoryWriterTester Tester;
-		Tester.TestWriteSwappedByteOrder(TestValueS64, ByteSwap(TestValueS64));
-	}
+		{
+			MemoryWriterTestUtil::MemoryWriterTester Tester;
+			Tester.TestWriteSwappedByteOrder(TestValueS64, ByteSwap(TestValueS64));
+		}
 
-	{
-		MemoryWriterTestUtil::MemoryWriterTester Tester;
-		Tester.TestWriteSwappedByteOrder(TestValueF, ByteSwap(TestValueF));
-	}
+		{
+			MemoryWriterTestUtil::MemoryWriterTester Tester;
+			Tester.TestWriteSwappedByteOrder(TestValueF, ByteSwap(TestValueF));
+		}
 
-	{
-		MemoryWriterTestUtil::MemoryWriterTester Tester;
-		Tester.TestWriteSwappedByteOrder(TestValueD, ByteSwap(TestValueD));
-	}
+		{
+			MemoryWriterTestUtil::MemoryWriterTester Tester;
+			Tester.TestWriteSwappedByteOrder(TestValueD, ByteSwap(TestValueD));
+		}
 
-	{
-		MemoryWriterTestUtil::MemoryWriterTester Tester;
-		Tester.TestWriteSwappedByteOrder(TestValueB, TestValueB);
+		{
+			MemoryWriterTestUtil::MemoryWriterTester Tester;
+			Tester.TestWriteSwappedByteOrder(TestValueB, TestValueB);
+		}
 	}
 
-	// ANSI String
+	SECTION("ANSI String")
 	{
 		const TCHAR* RawStr = TEXT("Joe"); // Only contains ANSI chars -> should be serialized as an ANSI string (1 byte per char)
 		const char* RawStrBytes = "Joe";
@@ -200,7 +202,7 @@ TEST_CASE("Core::Serialization::MemoryWriter::Memory Writer", "[Core][Serializat
 		int32 CharCountInBuffer = CharCount; // A positive count let the deserializer know this an ANSI string rather than a UTF16 string.
 		int32 CharCountInBufferSwapped = ByteSwap(CharCountInBuffer);
 
-		// Platform endianness
+		SECTION("Platform endianness")
 		{
 			TArray<uint8> Bytes;
 			FMemoryWriter Writer(Bytes);
@@ -212,7 +214,7 @@ TEST_CASE("Core::Serialization::MemoryWriter::Memory Writer", "[Core][Serializat
 			CHECK(FMemory::Memcmp(Bytes.GetData() + sizeof(int32), RawStrBytes, StrByteCount) == 0); // Check the string content.
 		}
 	
-		// Swapped endianness
+		SECTION("Swapped endianness")
 		{
 			TArray<uint8> Bytes;
 			FMemoryWriter Writer(Bytes);
@@ -226,7 +228,7 @@ TEST_CASE("Core::Serialization::MemoryWriter::Memory Writer", "[Core][Serializat
 		}
 	}
 
-	// UTF16 Strings
+	SECTION("UTF16 Strings")
 	{
 		// Must not contain at least one non-ANSI char otherwise, the serialization finds out and serialize and ANSI string (1-byte per char).
 		const TCHAR RawStr[] = {TEXT('\u0404'), TEXT('\u0400'), TEXT('\uC0AC'), TEXT('\u0000')};
@@ -241,7 +243,7 @@ TEST_CASE("Core::Serialization::MemoryWriter::Memory Writer", "[Core][Serializat
 		int32 CharCountInBuffer = -CharCount; // A negative count let the deserializer know this an UTF16 string rather than a ANSI string.
 		int32 CharCountInBufferSwapped = ByteSwap(CharCountInBuffer);
 
-		// Platform endianness
+		SECTION("Platform endianness")
 		{
 			TArray<uint8> Bytes;
 			FMemoryWriter Writer(Bytes);
@@ -253,7 +255,7 @@ TEST_CASE("Core::Serialization::MemoryWriter::Memory Writer", "[Core][Serializat
 			CHECK(FMemory::Memcmp(Bytes.GetData() + sizeof(int32), RawStrBytes, StrByteCount) == 0); // Check the string content.
 		}
 	
-		// Swapped endianness
+		SECTION("Swapped endianness")
 		{
 			TArray<uint8> Bytes;
 			FMemoryWriter Writer(Bytes);
diff --git a/Engine/Source/Runtime/Core/Tests/String/FindTest.cpp b/Engine/Source/Runtime/Core/Tests/String/FindTest.cpp
--- a/Engine/Source/Runtime/Core/Tests/String/FindTest.cpp
+++ b/Engine/Source/Runtime/Core/Tests/String/FindTest.cpp
@@ -1,224 +1,226 @@
 // Copyright Epic Games, Inc. All Rights Reserved.
 
 #include "String/Find.h"
-
-#include "Containers/StringView.h"
 #include "TestHarness.h"
 
-namespace UE::String
+TEST_CASE_METHOD(FAutomationTestFixture, "Core::String::Misc::FindFirst", "[Core][String][Smoke]")
+{
+	CHECK_EQUAL(UE::String::FindFirst(TEXT("AbCABCAbCABC"), TEXT("A")), 0);
+	CHECK_EQUAL(UE::String::FindFirst(TEXT("AbCABCAbCABC"), TEXT("a"), ESearchCase::IgnoreCase), 0);
+	CHECK_EQUAL(UE::String::FindFirst(TEXT("AbCABCAbCABC"), TEXT("b")), 1);
+	CHECK_EQUAL(UE::String::FindFirst(TEXT("AbCABCAbCABC"), TEXT("B")), 4);
+	CHECK_EQUAL(UE::String::FindFirst(TEXT("AbCABCAbCABC"), TEXT("B"), ESearchCase::IgnoreCase), 1);
+	CHECK_EQUAL(UE::String::FindFirst(TEXT("AbCABCAbCABC"), TEXT("a")), INDEX_NONE);
+	CHECK_EQUAL(UE::String::FindFirst(TEXT("AbCABCAbCABC"), TEXT("D"), ESearchCase::IgnoreCase), INDEX_NONE);
+	CHECK_EQUAL(UE::String::FindFirst(TEXT("AbCABCAbCABD"), TEXT("D")), 11);
+	CHECK_EQUAL(UE::String::FindFirst(TEXT("AbCABCAbCABD"), TEXT("d"), ESearchCase::IgnoreCase), 11);
+
+	CHECK_EQUAL(UE::String::FindFirst(TEXT("AbCABCAbCABC"), TEXT("AbC")), 0);
+	CHECK_EQUAL(UE::String::FindFirst(TEXT("AbCABCAbCABC"), TEXT("ABC")), 3);
+	CHECK_EQUAL(UE::String::FindFirst(TEXT("AbCABCAbCABC"), TEXT("Bc"), ESearchCase::IgnoreCase), 1);
+	CHECK_EQUAL(UE::String::FindFirst(TEXT("AbCABCAbCABC"), TEXT("ab")), INDEX_NONE);
+	CHECK_EQUAL(UE::String::FindFirst(TEXT("AbCABCAbCABC"), TEXT("CD"), ESearchCase::IgnoreCase), INDEX_NONE);
+	CHECK_EQUAL(UE::String::FindFirst(TEXT("AbCABCAbCABD"), TEXT("BD")), 10);
+	CHECK_EQUAL(UE::String::FindFirst(TEXT("AbCABCAbCABD"), TEXT("Bd"), ESearchCase::IgnoreCase), 10);
+
+	CHECK_EQUAL(UE::String::FindFirst(TEXT(""), TEXT("A")), INDEX_NONE);
+	CHECK_EQUAL(UE::String::FindFirst(TEXT("A"), TEXT("A")), 0);
+	CHECK_EQUAL(UE::String::FindFirst(TEXT("A"), TEXT("A"), ESearchCase::IgnoreCase), 0);
+	CHECK_EQUAL(UE::String::FindFirst(TEXT("ABC"), TEXT("ABC")), 0);
+	CHECK_EQUAL(UE::String::FindFirst(TEXT("ABC"), TEXT("abc"), ESearchCase::IgnoreCase), 0);
+	CHECK_EQUAL(UE::String::FindFirst(TEXT("AB"), TEXT("ABC")), INDEX_NONE);
+
+	CHECK_EQUAL(UE::String::FindFirst("AbCABCAbCABC", "ABC"), 3);
+
+	CHECK_EQUAL(UE::String::FindFirst(FStringView(nullptr, 0), TEXT("SearchTerm")), INDEX_NONE);
+	CHECK_EQUAL(UE::String::FindFirst(FStringView(), TEXT("SearchTerm")), INDEX_NONE);
+	CHECK_EQUAL(UE::String::FindFirst(FString(), TEXT("SearchTerm")), INDEX_NONE);
+}
+
+TEST_CASE_METHOD(FAutomationTestFixture, "Core::String::Misc::FindLast", "[Core][String][Smoke]")
+{
+	CHECK_EQUAL(UE::String::FindLast(TEXT("AbCABCAbCABC"), TEXT("b")), 7);
+	CHECK_EQUAL(UE::String::FindLast(TEXT("AbCABCAbCABC"), TEXT("B")), 10);
+	CHECK_EQUAL(UE::String::FindLast(TEXT("AbCABCAbCABC"), TEXT("b"), ESearchCase::IgnoreCase), 10);
+	CHECK_EQUAL(UE::String::FindLast(TEXT("AbCABCAbCABC"), TEXT("a")), INDEX_NONE);
+	CHECK_EQUAL(UE::String::FindLast(TEXT("AbCABCAbCABC"), TEXT("D"), ESearchCase::IgnoreCase), INDEX_NONE);
+	CHECK_EQUAL(UE::String::FindLast(TEXT("AbCABCAbCABD"), TEXT("D")), 11);
+	CHECK_EQUAL(UE::String::FindLast(TEXT("AbCABCAbCABD"), TEXT("d"), ESearchCase::IgnoreCase), 11);
+
+	CHECK_EQUAL(UE::String::FindLast(TEXT("AbCABCAbCABC"), TEXT("AbC")), 6);
+	CHECK_EQUAL(UE::String::FindLast(TEXT("AbCABCAbCABC"), TEXT("ABC")), 9);
+	CHECK_EQUAL(UE::String::FindLast(TEXT("AbCABCAbCABC"), TEXT("Bc"), ESearchCase::IgnoreCase), 10);
+	CHECK_EQUAL(UE::String::FindLast(TEXT("AbCABCAbCABC"), TEXT("ab")), INDEX_NONE);
+	CHECK_EQUAL(UE::String::FindLast(TEXT("AbCABCAbCABC"), TEXT("CD"), ESearchCase::IgnoreCase), INDEX_NONE);
+	CHECK_EQUAL(UE::String::FindLast(TEXT("AbCABCAbCABC"), TEXT("BC")), 10);
+	CHECK_EQUAL(UE::String::FindLast(TEXT("AbCABCAbCABC"), TEXT("Bc"), ESearchCase::IgnoreCase), 10);
+
+	CHECK_EQUAL(UE::String::FindLast(TEXT(""), TEXT("A")), INDEX_NONE);
+	CHECK_EQUAL(UE::String::FindLast(TEXT("A"), TEXT("A")), 0);
+	CHECK_EQUAL(UE::String::FindLast(TEXT("A"), TEXT("A"), ESearchCase::IgnoreCase), 0);
+	CHECK_EQUAL(UE::String::FindLast(TEXT("ABC"), TEXT("ABC")), 0);
+	CHECK_EQUAL(UE::String::FindLast(TEXT("ABC"), TEXT("abc"), ESearchCase::IgnoreCase), 0);
+	CHECK_EQUAL(UE::String::FindLast(TEXT("AB"), TEXT("ABC")), INDEX_NONE);
+
+	CHECK_EQUAL(UE::String::FindLast("AbCABCAbCABC", "ABC"), 9);
+
+	CHECK_EQUAL(UE::String::FindLast(FStringView(nullptr, 0), TEXT("SearchTerm")), INDEX_NONE);
+	CHECK_EQUAL(UE::String::FindLast(FStringView(), TEXT("SearchTerm")), INDEX_NONE);
+	CHECK_EQUAL(UE::String::FindLast(FString(), TEXT("SearchTerm")), INDEX_NONE);
+}
+
+
+TEST_CASE_METHOD(FAutomationTestFixture, "Core::String::Misc::FindFirstOfAny", "[Core][String][Smoke]")
 {
+	CHECK_EQUAL(UE::String::FindFirstOfAny(TEXT("AbCABCAbcABC"), {TEXT("A"), TEXT("B")}), 0);
+	CHECK_EQUAL(UE::String::FindFirstOfAny(TEXT("AbCABCAbcABC"), {TEXT("a"), TEXT("B")}, ESearchCase::IgnoreCase), 0);
+	CHECK_EQUAL(UE::String::FindFirstOfAny(TEXT("AbCABCAbcABC"), {TEXT("b")}), 1);
+	CHECK_EQUAL(UE::String::FindFirstOfAny(TEXT("AbCABCAbcABC"), {TEXT("c"), TEXT("B")}), 4);
+	CHECK_EQUAL(UE::String::FindFirstOfAny(TEXT("AbCABCAbcABC"), {TEXT("c"), TEXT("B")}, ESearchCase::IgnoreCase), 1);
+	CHECK_EQUAL(UE::String::FindFirstOfAny(TEXT("AbCABCAbcABC"), {TEXT("D"), TEXT("a")}), INDEX_NONE);
+	CHECK_EQUAL(UE::String::FindFirstOfAny(TEXT("AbCABCAbcABC"), {TEXT("E"), TEXT("D")}, ESearchCase::IgnoreCase), INDEX_NONE);
+	CHECK_EQUAL(UE::String::FindFirstOfAny(TEXT("AbCABCAbcABD"), {TEXT("E"), TEXT("D")}), 11);
+	CHECK_EQUAL(UE::String::FindFirstOfAny(TEXT("AbCABCAbcABD"), {TEXT("E"), TEXT("d")}, ESearchCase::IgnoreCase), 11);
+
+	CHECK_EQUAL(UE::String::FindFirstOfAny(TEXT("AbCABCAbCABC"), {TEXT("AbC")}), 0);
+	CHECK_EQUAL(UE::String::FindFirstOfAny(TEXT("AbCABCAbCABC"), {TEXT("CABc"), TEXT("ABC")}), 3);
+	CHECK_EQUAL(UE::String::FindFirstOfAny(TEXT("AbCABCAbCABC"), {TEXT("ABD"), TEXT("Bc")}, ESearchCase::IgnoreCase), 1);
+	CHECK_EQUAL(UE::String::FindFirstOfAny(TEXT("AbCABCAbCABC"), {TEXT("bc"), TEXT("ab")}), INDEX_NONE);
+	CHECK_EQUAL(UE::String::FindFirstOfAny(TEXT("AbCABCAbCABC"), {TEXT("DA"), TEXT("CD")}, ESearchCase::IgnoreCase), INDEX_NONE);
+	CHECK_EQUAL(UE::String::FindFirstOfAny(TEXT("AbCABCAbCABD"), {TEXT("BD"), TEXT("CABB")}), 10);
+	CHECK_EQUAL(UE::String::FindFirstOfAny(TEXT("AbCABCAbCABD"), {TEXT("Bd"), TEXT("CABB")}, ESearchCase::IgnoreCase), 10);
+
+	CHECK_EQUAL(UE::String::FindFirstOfAny(TEXT(""), {TEXT("A"), TEXT("B")}), INDEX_NONE);
+	CHECK_EQUAL(UE::String::FindFirstOfAny(TEXT("A"), {TEXT("A"), TEXT("B")}), 0);
+	CHECK_EQUAL(UE::String::FindFirstOfAny(TEXT("A"), {TEXT("A"), TEXT("B")}, ESearchCase::IgnoreCase), 0);
+	CHECK_EQUAL(UE::String::FindFirstOfAny(TEXT("ABC"), {TEXT("ABC"), TEXT("BC")}), 0);
+	CHECK_EQUAL(UE::String::FindFirstOfAny(TEXT("ABC"), {TEXT("abc"), TEXT("bc")}, ESearchCase::IgnoreCase), 0);
+	CHECK_EQUAL(UE::String::FindFirstOfAny(TEXT("AB"), {TEXT("ABC"), TEXT("ABD")}), INDEX_NONE);
+
+	CHECK_EQUAL(UE::String::FindFirstOfAny("AbCABCAbCABC", {"CABc", "ABC"}), 3);
+
+	CHECK_EQUAL(UE::String::FindFirstOfAny(FStringView(nullptr, 0), {TEXT("ABC"), TEXT("ABD")}), INDEX_NONE);
+	CHECK_EQUAL(UE::String::FindFirstOfAny(FStringView(), {TEXT("ABC"), TEXT("ABD")}), INDEX_NONE);
+	CHECK_EQUAL(UE::String::FindFirstOfAny(FString(), {TEXT("ABC"), TEXT("ABD")}), INDEX_NONE);
+}
 
-TEST_CASE("Core::String::Find", "[Core][String][Smoke]")
+
+TEST_CASE_METHOD(FAutomationTestFixture, "Core::String::Misc::FindLastOfAny", "[Core][String][Smoke]")
 {
-	SECTION("FindFirst")
-	{
-		CHECK(FindFirst(TEXT("AbCABCAbCABC"), TEXT("A")) == 0);
-		CHECK(FindFirst(TEXT("AbCABCAbCABC"), TEXT("a"), ESearchCase::IgnoreCase) == 0);
-		CHECK(FindFirst(TEXT("AbCABCAbCABC"), TEXT("b")) == 1);
-		CHECK(FindFirst(TEXT("AbCABCAbCABC"), TEXT("B")) == 4);
-		CHECK(FindFirst(TEXT("AbCABCAbCABC"), TEXT("B"), ESearchCase::IgnoreCase) == 1);
-		CHECK(FindFirst(TEXT("AbCABCAbCABC"), TEXT("a")) == INDEX_NONE);
-		CHECK(FindFirst(TEXT("AbCABCAbCABC"), TEXT("D"), ESearchCase::IgnoreCase) == INDEX_NONE);
-		CHECK(FindFirst(TEXT("AbCABCAbCABD"), TEXT("D")) == 11);
-		CHECK(FindFirst(TEXT("AbCABCAbCABD"), TEXT("d"), ESearchCase::IgnoreCase) == 11);
-
-		CHECK(FindFirst(TEXT("AbCABCAbCABC"), TEXT("AbC")) == 0);
-		CHECK(FindFirst(TEXT("AbCABCAbCABC"), TEXT("ABC")) == 3);
-		CHECK(FindFirst(TEXT("AbCABCAbCABC"), TEXT("Bc"), ESearchCase::IgnoreCase) == 1);
-		CHECK(FindFirst(TEXT("AbCABCAbCABC"), TEXT("ab")) == INDEX_NONE);
-		CHECK(FindFirst(TEXT("AbCABCAbCABC"), TEXT("CD"), ESearchCase::IgnoreCase) == INDEX_NONE);
-		CHECK(FindFirst(TEXT("AbCABCAbCABD"), TEXT("BD")) == 10);
-		CHECK(FindFirst(TEXT("AbCABCAbCABD"), TEXT("Bd"), ESearchCase::IgnoreCase) == 10);
-
-		CHECK(FindFirst(TEXT(""), TEXT("A")) == INDEX_NONE);
-		CHECK(FindFirst(TEXT("A"), TEXT("A")) == 0);
-		CHECK(FindFirst(TEXT("A"), TEXT("A"), ESearchCase::IgnoreCase) == 0);
-		CHECK(FindFirst(TEXT("ABC"), TEXT("ABC")) == 0);
-		CHECK(FindFirst(TEXT("ABC"), TEXT("abc"), ESearchCase::IgnoreCase) == 0);
-		CHECK(FindFirst(TEXT("AB"), TEXT("ABC")) == INDEX_NONE);
-
-		CHECK(FindFirst("AbCABCAbCABC", "ABC") == 3);
-
-		CHECK(FindFirst(FStringView(nullptr, 0), TEXT("SearchTerm")) == INDEX_NONE);
-		CHECK(FindFirst(FStringView(), TEXT("SearchTerm")) == INDEX_NONE);
-	}
-
-	SECTION("FindLast")
-	{
-		CHECK(FindLast(TEXT("AbCABCAbCABC"), TEXT("b")) == 7);
-		CHECK(FindLast(TEXT("AbCABCAbCABC"), TEXT("B")) == 10);
-		CHECK(FindLast(TEXT("AbCABCAbCABC"), TEXT("b"), ESearchCase::IgnoreCase) == 10);
-		CHECK(FindLast(TEXT("AbCABCAbCABC"), TEXT("a")) == INDEX_NONE);
-		CHECK(FindLast(TEXT("AbCABCAbCABC"), TEXT("D"), ESearchCase::IgnoreCase) == INDEX_NONE);
-		CHECK(FindLast(TEXT("AbCABCAbCABD"), TEXT("D")) == 11);
-		CHECK(FindLast(TEXT("AbCABCAbCABD"), TEXT("d"), ESearchCase::IgnoreCase) == 11);
-
-		CHECK(FindLast(TEXT("AbCABCAbCABC"), TEXT("AbC")) == 6);
-		CHECK(FindLast(TEXT("AbCABCAbCABC"), TEXT("ABC")) == 9);
-		CHECK(FindLast(TEXT("AbCABCAbCABC"), TEXT("Bc"), ESearchCase::IgnoreCase) == 10);
-		CHECK(FindLast(TEXT("AbCABCAbCABC"), TEXT("ab")) == INDEX_NONE);
-		CHECK(FindLast(TEXT("AbCABCAbCABC"), TEXT("CD"), ESearchCase::IgnoreCase) == INDEX_NONE);
-		CHECK(FindLast(TEXT("AbCABCAbCABC"), TEXT("BC")) == 10);
-		CHECK(FindLast(TEXT("AbCABCAbCABC"), TEXT("Bc"), ESearchCase::IgnoreCase) == 10);
-
-		CHECK(FindLast(TEXT(""), TEXT("A")) == INDEX_NONE);
-		CHECK(FindLast(TEXT("A"), TEXT("A")) == 0);
-		CHECK(FindLast(TEXT("A"), TEXT("A"), ESearchCase::IgnoreCase) == 0);
-		CHECK(FindLast(TEXT("ABC"), TEXT("ABC")) == 0);
-		CHECK(FindLast(TEXT("ABC"), TEXT("abc"), ESearchCase::IgnoreCase) == 0);
-		CHECK(FindLast(TEXT("AB"), TEXT("ABC")) == INDEX_NONE);
-
-		CHECK(FindLast("AbCABCAbCABC", "ABC") == 9);
-
-		CHECK(FindLast(FStringView(nullptr, 0), TEXT("SearchTerm")) == INDEX_NONE);
-		CHECK(FindLast(FStringView(), TEXT("SearchTerm")) == INDEX_NONE);
-	}
-
-	SECTION("FindFirstOfAny")
-	{
-		CHECK(FindFirstOfAny(TEXT("AbCABCAbcABC"), {TEXT("A"), TEXT("B")}) == 0);
-		CHECK(FindFirstOfAny(TEXT("AbCABCAbcABC"), {TEXT("a"), TEXT("B")}, ESearchCase::IgnoreCase) == 0);
-		CHECK(FindFirstOfAny(TEXT("AbCABCAbcABC"), {TEXT("b")}) == 1);
-		CHECK(FindFirstOfAny(TEXT("AbCABCAbcABC"), {TEXT("c"), TEXT("B")}) == 4);
-		CHECK(FindFirstOfAny(TEXT("AbCABCAbcABC"), {TEXT("c"), TEXT("B")}, ESearchCase::IgnoreCase) == 1);
-		CHECK(FindFirstOfAny(TEXT("AbCABCAbcABC"), {TEXT("D"), TEXT("a")}) == INDEX_NONE);
-		CHECK(FindFirstOfAny(TEXT("AbCABCAbcABC"), {TEXT("E"), TEXT("D")}, ESearchCase::IgnoreCase) == INDEX_NONE);
-		CHECK(FindFirstOfAny(TEXT("AbCABCAbcABD"), {TEXT("E"), TEXT("D")}) == 11);
-		CHECK(FindFirstOfAny(TEXT("AbCABCAbcABD"), {TEXT("E"), TEXT("d")}, ESearchCase::IgnoreCase) == 11);
-
-		CHECK(FindFirstOfAny(TEXT("AbCABCAbCABC"), {TEXT("AbC")}) == 0);
-		CHECK(FindFirstOfAny(TEXT("AbCABCAbCABC"), {TEXT("CABc"), TEXT("ABC")}) == 3);
-		CHECK(FindFirstOfAny(TEXT("AbCABCAbCABC"), {TEXT("ABD"), TEXT("Bc")}, ESearchCase::IgnoreCase) == 1);
-		CHECK(FindFirstOfAny(TEXT("AbCABCAbCABC"), {TEXT("bc"), TEXT("ab")}) == INDEX_NONE);
-		CHECK(FindFirstOfAny(TEXT("AbCABCAbCABC"), {TEXT("DA"), TEXT("CD")}, ESearchCase::IgnoreCase) == INDEX_NONE);
-		CHECK(FindFirstOfAny(TEXT("AbCABCAbCABD"), {TEXT("BD"), TEXT("CABB")}) == 10);
-		CHECK(FindFirstOfAny(TEXT("AbCABCAbCABD"), {TEXT("Bd"), TEXT("CABB")}, ESearchCase::IgnoreCase) == 10);
-
-		CHECK(FindFirstOfAny(TEXT(""), {TEXT("A"), TEXT("B")}) == INDEX_NONE);
-		CHECK(FindFirstOfAny(TEXT("A"), {TEXT("A"), TEXT("B")}) == 0);
-		CHECK(FindFirstOfAny(TEXT("A"), {TEXT("A"), TEXT("B")}, ESearchCase::IgnoreCase) == 0);
-		CHECK(FindFirstOfAny(TEXT("ABC"), {TEXT("ABC"), TEXT("BC")}) == 0);
-		CHECK(FindFirstOfAny(TEXT("ABC"), {TEXT("abc"), TEXT("bc")}, ESearchCase::IgnoreCase) == 0);
-		CHECK(FindFirstOfAny(TEXT("AB"), {TEXT("ABC"), TEXT("ABD")}) == INDEX_NONE);
-
-		CHECK(FindFirstOfAny("AbCABCAbCABC", {"CABc", "ABC"}) == 3);
-
-		CHECK(FindFirstOfAny(FStringView(nullptr, 0), {TEXT("ABC"), TEXT("ABD")}) == INDEX_NONE);
-		CHECK(FindFirstOfAny(FStringView(), {TEXT("ABC"), TEXT("ABD")}) == INDEX_NONE);
-	}
-
-	SECTION("FindLastOfAny")
-	{
-		CHECK(FindLastOfAny(TEXT("AbCABCAbcABC"), {TEXT("a"), TEXT("b")}) == 7);
-		CHECK(FindLastOfAny(TEXT("AbCABCAbcABC"), {TEXT("a"), TEXT("b")}, ESearchCase::IgnoreCase) == 10);
-		CHECK(FindLastOfAny(TEXT("AbCABCAbcABC"), {TEXT("b")}) == 7);
-		CHECK(FindLastOfAny(TEXT("AbCABCAbcABC"), {TEXT("c"), TEXT("B")}) == 10);
-		CHECK(FindLastOfAny(TEXT("AbCABCAbcABC"), {TEXT("c"), TEXT("B")}, ESearchCase::IgnoreCase) == 11);
-		CHECK(FindLastOfAny(TEXT("AbCABCAbcABC"), {TEXT("D"), TEXT("a")}) == INDEX_NONE);
-		CHECK(FindLastOfAny(TEXT("AbCABCAbcABC"), {TEXT("E"), TEXT("D")}, ESearchCase::IgnoreCase) == INDEX_NONE);
-		CHECK(FindLastOfAny(TEXT("AbCABCAbcABD"), {TEXT("E"), TEXT("D")}) == 11);
-		CHECK(FindLastOfAny(TEXT("AbCABCAbcABD"), {TEXT("E"), TEXT("d")}, ESearchCase::IgnoreCase) == 11);
-
-		CHECK(FindLastOfAny(TEXT("AbCABCAbCABC"), {TEXT("AbC")}) == 6);
-		CHECK(FindLastOfAny(TEXT("AbCABCAbCABC"), {TEXT("CABc"), TEXT("ABC")}) == 9);
-		CHECK(FindLastOfAny(TEXT("AbCABCAbCABC"), {TEXT("ABD"), TEXT("Bc")}, ESearchCase::IgnoreCase) == 10);
-		CHECK(FindLastOfAny(TEXT("AbCABCAbCABC"), {TEXT("bc"), TEXT("ab")}) == INDEX_NONE);
-		CHECK(FindLastOfAny(TEXT("AbCABCAbCABC"), {TEXT("DA"), TEXT("CD")}, ESearchCase::IgnoreCase) == INDEX_NONE);
-		CHECK(FindLastOfAny(TEXT("AbCABCAbCABD"), {TEXT("BD"), TEXT("CABB")}) == 10);
-		CHECK(FindLastOfAny(TEXT("AbCABCAbCABD"), {TEXT("Bd"), TEXT("CABB")}, ESearchCase::IgnoreCase) == 10);
-
-		CHECK(FindLastOfAny(TEXT(""), {TEXT("A"), TEXT("B")}) == INDEX_NONE);
-		CHECK(FindLastOfAny(TEXT("A"), {TEXT("A"), TEXT("B")}) == 0);
-		CHECK(FindLastOfAny(TEXT("A"), {TEXT("A"), TEXT("B")}, ESearchCase::IgnoreCase) == 0);
-		CHECK(FindLastOfAny(TEXT("ABC"), {TEXT("ABC"), TEXT("BC")}) == 1);
-		CHECK(FindLastOfAny(TEXT("ABC"), {TEXT("abc"), TEXT("bc")}, ESearchCase::IgnoreCase) == 1);
-		CHECK(FindLastOfAny(TEXT("AB"), {TEXT("ABC"), TEXT("ABD")}) == INDEX_NONE);
-
-		CHECK(FindLastOfAny("AbCABCAbCABC", {"CABc", "ABC"}) == 9);
-
-		CHECK(FindLastOfAny(FStringView(nullptr, 0), { TEXT("ABC"), TEXT("ABD") }) == INDEX_NONE);
-		CHECK(FindLastOfAny(FStringView(), { TEXT("ABC"), TEXT("ABD") }) == INDEX_NONE);
-	}
-
-	SECTION("FindFirstChar")
-	{
-		CHECK(FindFirstChar(TEXT("AbCABCAbCABC"), TEXT('b')) == 1);
-		CHECK(FindFirstChar(TEXT("AbCABCAbCABC"), TEXT('B')) == 4);
-		CHECK(FindFirstChar(TEXT("AbCABCAbCABC"), TEXT('B'), ESearchCase::IgnoreCase) == 1);
-		CHECK(FindFirstChar(TEXT("AbCABCAbCABC"), TEXT('a')) == INDEX_NONE);
-		CHECK(FindFirstChar(TEXT("AbCABCAbCABC"), TEXT('D'), ESearchCase::IgnoreCase) == INDEX_NONE);
-		CHECK(FindFirstChar(TEXT("AbCABCAbCABD"), TEXT('D')) == 11);
-		CHECK(FindFirstChar(TEXT("AbCABCAbCABD"), TEXT('d'), ESearchCase::IgnoreCase) == 11);
-
-		CHECK(FindFirstChar(TEXT(""), TEXT('A')) == INDEX_NONE);
-		CHECK(FindFirstChar(TEXT("A"), TEXT('A')) == 0);
-		CHECK(FindFirstChar(TEXT("A"), TEXT('A'), ESearchCase::IgnoreCase) == 0);
-
-		CHECK(FindFirstChar("AbCABCAbCABC", 'B') == 4);
-
-		CHECK(FindFirstChar(FStringView(nullptr, 0), TEXT('A')) == INDEX_NONE);
-		CHECK(FindFirstChar(FStringView(), TEXT('A')) == INDEX_NONE);
-	}
-
-	SECTION("FindLastChar")
-	{
-		CHECK(FindLastChar(TEXT("AbCABCAbCABC"), TEXT('b')) == 7);
-		CHECK(FindLastChar(TEXT("AbCABCAbCABC"), TEXT('B')) == 10);
-		CHECK(FindLastChar(TEXT("AbCABCAbCABC"), TEXT('b'), ESearchCase::IgnoreCase) == 10);
-		CHECK(FindLastChar(TEXT("AbCABCAbCABC"), TEXT('a')) == INDEX_NONE);
-		CHECK(FindLastChar(TEXT("AbCABCAbCABC"), TEXT('D'), ESearchCase::IgnoreCase) == INDEX_NONE);
-		CHECK(FindLastChar(TEXT("AbCABCAbCABD"), TEXT('D')) == 11);
-		CHECK(FindLastChar(TEXT("AbCABCAbCABD"), TEXT('d'), ESearchCase::IgnoreCase) == 11);
-
-		CHECK(FindLastChar(TEXT(""), TEXT('A')) == INDEX_NONE);
-		CHECK(FindLastChar(TEXT("A"), TEXT('A')) == 0);
-		CHECK(FindLastChar(TEXT("A"), TEXT('A'), ESearchCase::IgnoreCase) == 0);
-
-		CHECK(FindLastChar("AbCABCAbCABC", 'B') == 10);
-
-		CHECK(FindLastChar(FStringView(nullptr, 0), TEXT('A')) == INDEX_NONE);
-		CHECK(FindLastChar(FStringView(), TEXT('A')) == INDEX_NONE);
-	}
-
-	SECTION("FindFirstOfAnyChar")
-	{
-		CHECK(FindFirstOfAnyChar(TEXT("AbCABCAbcABC"), {TEXT('b')}) == 1);
-		CHECK(FindFirstOfAnyChar(TEXT("AbCABCAbcABC"), {TEXT('c'), TEXT('B')}) == 4);
-		CHECK(FindFirstOfAnyChar(TEXT("AbCABCAbcABC"), {TEXT('c'), TEXT('B')}, ESearchCase::IgnoreCase) == 1);
-		CHECK(FindFirstOfAnyChar(TEXT("AbCABCAbcABC"), {TEXT('D'), TEXT('a')}) == INDEX_NONE);
-		CHECK(FindFirstOfAnyChar(TEXT("AbCABCAbcABC"), {TEXT('E'), TEXT('D')}, ESearchCase::IgnoreCase) == INDEX_NONE);
-		CHECK(FindFirstOfAnyChar(TEXT("AbCABCAbcABD"), {TEXT('E'), TEXT('D')}) == 11);
-		CHECK(FindFirstOfAnyChar(TEXT("AbCABCAbcABD"), {TEXT('E'), TEXT('d')}, ESearchCase::IgnoreCase) == 11);
-
-		CHECK(FindFirstOfAnyChar(TEXT(""), {TEXT('A'), TEXT('B')}) == INDEX_NONE);
-		CHECK(FindFirstOfAnyChar(TEXT("A"), {TEXT('A'), TEXT('B')}) == 0);
-		CHECK(FindFirstOfAnyChar(TEXT("A"), {TEXT('A'), TEXT('B')}, ESearchCase::IgnoreCase) == 0);
-
-		CHECK(FindFirstOfAnyChar("AbCABCAbcABC", {'c', 'B'}) == 4);
-
-		CHECK(FindFirstOfAnyChar(FStringView(nullptr, 0), { TEXT('A'), TEXT('B') }) == INDEX_NONE);
-		CHECK(FindFirstOfAnyChar(FStringView(), { TEXT('A'), TEXT('B') }) == INDEX_NONE);
-	}
-
-	SECTION("FindLastOfAnyChar")
-	{
-		CHECK(FindLastOfAnyChar(TEXT("AbCABCAbcABC"), {TEXT('b')}) == 7);
-		CHECK(FindLastOfAnyChar(TEXT("AbCABCAbcABC"), {TEXT('c'), TEXT('B')}) == 10);
-		CHECK(FindLastOfAnyChar(TEXT("AbCABCAbcABC"), {TEXT('c'), TEXT('B')}, ESearchCase::IgnoreCase) == 11);
-		CHECK(FindLastOfAnyChar(TEXT("AbCABCAbcABC"), {TEXT('D'), TEXT('a')}) == INDEX_NONE);
-		CHECK(FindLastOfAnyChar(TEXT("AbCABCAbcABC"), {TEXT('E'), TEXT('D')}, ESearchCase::IgnoreCase) == INDEX_NONE);
-		CHECK(FindLastOfAnyChar(TEXT("AbCABCAbcABD"), {TEXT('E'), TEXT('D')}) == 11);
-		CHECK(FindLastOfAnyChar(TEXT("AbCABCAbcABD"), {TEXT('E'), TEXT('d')}, ESearchCase::IgnoreCase) == 11);
-
-		CHECK(FindLastOfAnyChar(TEXT(""), {TEXT('A'), TEXT('B')}) == INDEX_NONE);
-		CHECK(FindLastOfAnyChar(TEXT("A"), {TEXT('A'), TEXT('B')}) == 0);
-		CHECK(FindLastOfAnyChar(TEXT("A"), {TEXT('A'), TEXT('B')}, ESearchCase::IgnoreCase) == 0);
-
-		CHECK(FindLastOfAnyChar("AbCABCAbcABC", {'c', 'B'}) == 10);
-
-		CHECK(FindLastOfAnyChar(FStringView(nullptr, 0), { TEXT('A'), TEXT('B') }) == INDEX_NONE);
-		CHECK(FindLastOfAnyChar(FStringView(), { TEXT('A'), TEXT('B') }) == INDEX_NONE);
-	}
+	CHECK_EQUAL(UE::String::FindLastOfAny(TEXT("AbCABCAbcABC"), {TEXT("a"), TEXT("b")}), 7);
+	CHECK_EQUAL(UE::String::FindLastOfAny(TEXT("AbCABCAbcABC"), {TEXT("a"), TEXT("b")}, ESearchCase::IgnoreCase), 10);
+	CHECK_EQUAL(UE::String::FindLastOfAny(TEXT("AbCABCAbcABC"), {TEXT("b")}), 7);
+	CHECK_EQUAL(UE::String::FindLastOfAny(TEXT("AbCABCAbcABC"), {TEXT("c"), TEXT("B")}), 10);
+	CHECK_EQUAL(UE::String::FindLastOfAny(TEXT("AbCABCAbcABC"), {TEXT("c"), TEXT("B")}, ESearchCase::IgnoreCase), 11);
+	CHECK_EQUAL(UE::String::FindLastOfAny(TEXT("AbCABCAbcABC"), {TEXT("D"), TEXT("a")}), INDEX_NONE);
+	CHECK_EQUAL(UE::String::FindLastOfAny(TEXT("AbCABCAbcABC"), {TEXT("E"), TEXT("D")}, ESearchCase::IgnoreCase), INDEX_NONE);
+	CHECK_EQUAL(UE::String::FindLastOfAny(TEXT("AbCABCAbcABD"), {TEXT("E"), TEXT("D")}), 11);
+	CHECK_EQUAL(UE::String::FindLastOfAny(TEXT("AbCABCAbcABD"), {TEXT("E"), TEXT("d")}, ESearchCase::IgnoreCase), 11);
+
+	CHECK_EQUAL(UE::String::FindLastOfAny(TEXT("AbCABCAbCABC"), {TEXT("AbC")}), 6);
+	CHECK_EQUAL(UE::String::FindLastOfAny(TEXT("AbCABCAbCABC"), {TEXT("CABc"), TEXT("ABC")}), 9);
+	CHECK_EQUAL(UE::String::FindLastOfAny(TEXT("AbCABCAbCABC"), {TEXT("ABD"), TEXT("Bc")}, ESearchCase::IgnoreCase), 10);
+	CHECK_EQUAL(UE::String::FindLastOfAny(TEXT("AbCABCAbCABC"), {TEXT("bc"), TEXT("ab")}), INDEX_NONE);
+	CHECK_EQUAL(UE::String::FindLastOfAny(TEXT("AbCABCAbCABC"), {TEXT("DA"), TEXT("CD")}, ESearchCase::IgnoreCase), INDEX_NONE);
+	CHECK_EQUAL(UE::String::FindLastOfAny(TEXT("AbCABCAbCABD"), {TEXT("BD"), TEXT("CABB")}), 10);
+	CHECK_EQUAL(UE::String::FindLastOfAny(TEXT("AbCABCAbCABD"), {TEXT("Bd"), TEXT("CABB")}, ESearchCase::IgnoreCase), 10);
+
+	CHECK_EQUAL(UE::String::FindLastOfAny(TEXT(""), {TEXT("A"), TEXT("B")}), INDEX_NONE);
+	CHECK_EQUAL(UE::String::FindLastOfAny(TEXT("A"), {TEXT("A"), TEXT("B")}), 0);
+	CHECK_EQUAL(UE::String::FindLastOfAny(TEXT("A"), {TEXT("A"), TEXT("B")}, ESearchCase::IgnoreCase), 0);
+	CHECK_EQUAL(UE::String::FindLastOfAny(TEXT("ABC"), {TEXT("ABC"), TEXT("BC")}), 1);
+	CHECK_EQUAL(UE::String::FindLastOfAny(TEXT("ABC"), {TEXT("abc"), TEXT("bc")}, ESearchCase::IgnoreCase), 1);
+	CHECK_EQUAL(UE::String::FindLastOfAny(TEXT("AB"), {TEXT("ABC"), TEXT("ABD")}), INDEX_NONE);
+
+	CHECK_EQUAL(UE::String::FindLastOfAny("AbCABCAbCABC", {"CABc", "ABC"}), 9);
+
+	CHECK_EQUAL(UE::String::FindLastOfAny(FStringView(nullptr, 0), { TEXT("ABC"), TEXT("ABD") }), INDEX_NONE);
+	CHECK_EQUAL(UE::String::FindLastOfAny(FStringView(), { TEXT("ABC"), TEXT("ABD") }), INDEX_NONE);
+	CHECK_EQUAL(UE::String::FindLastOfAny(FString(), { TEXT("ABC"), TEXT("ABD") }), INDEX_NONE);
 }
 
-} // UE::String
+
+TEST_CASE_METHOD(FAutomationTestFixture, "Core::String::Misc::FindFirstChar", "[Core][String][Smoke]")
+{
+	TEST_EQUAL(TEXT("FindFirstChar"), UE::String::FindFirstChar(TEXT("AbCABCAbCABC"), TEXT('b')), 1);
+	TEST_EQUAL(TEXT("FindFirstChar"), UE::String::FindFirstChar(TEXT("AbCABCAbCABC"), TEXT('B')), 4);
+	TEST_EQUAL(TEXT("FindFirstChar"), UE::String::FindFirstChar(TEXT("AbCABCAbCABC"), TEXT('B'), ESearchCase::IgnoreCase), 1);
+	TEST_EQUAL(TEXT("FindFirstChar"), UE::String::FindFirstChar(TEXT("AbCABCAbCABC"), TEXT('a')), INDEX_NONE);
+	TEST_EQUAL(TEXT("FindFirstChar"), UE::String::FindFirstChar(TEXT("AbCABCAbCABC"), TEXT('D'), ESearchCase::IgnoreCase), INDEX_NONE);
+	TEST_EQUAL(TEXT("FindFirstChar"), UE::String::FindFirstChar(TEXT("AbCABCAbCABD"), TEXT('D')), 11);
+	TEST_EQUAL(TEXT("FindFirstChar"), UE::String::FindFirstChar(TEXT("AbCABCAbCABD"), TEXT('d'), ESearchCase::IgnoreCase), 11);
+
+	TEST_EQUAL(TEXT("FindFirstChar"), UE::String::FindFirstChar(TEXT(""), TEXT('A')), INDEX_NONE);
+	TEST_EQUAL(TEXT("FindFirstChar"), UE::String::FindFirstChar(TEXT("A"), TEXT('A')), 0);
+	TEST_EQUAL(TEXT("FindFirstChar"), UE::String::FindFirstChar(TEXT("A"), TEXT('A'), ESearchCase::IgnoreCase), 0);
+
+	TEST_EQUAL(TEXT("FindFirstChar"), UE::String::FindFirstChar("AbCABCAbCABC", 'B'), 4);
+
+	TEST_EQUAL(TEXT("FindFirstChar"), UE::String::FindFirstChar(FStringView(nullptr, 0), TEXT('A')), INDEX_NONE);
+	TEST_EQUAL(TEXT("FindFirstChar"), UE::String::FindFirstChar(FStringView(), TEXT('A')), INDEX_NONE);
+	TEST_EQUAL(TEXT("FindFirstChar"), UE::String::FindFirstChar(FString(), TEXT('A')), INDEX_NONE);
+}
+
+
+TEST_CASE_METHOD(FAutomationTestFixture, "Core::String::Misc::FindLastChar", "[Core][String][Smoke]")
+{
+	TEST_EQUAL(TEXT("FindLastChar"), UE::String::FindLastChar(TEXT("AbCABCAbCABC"), TEXT('b')), 7);
+	TEST_EQUAL(TEXT("FindLastChar"), UE::String::FindLastChar(TEXT("AbCABCAbCABC"), TEXT('B')), 10);
+	TEST_EQUAL(TEXT("FindLastChar"), UE::String::FindLastChar(TEXT("AbCABCAbCABC"), TEXT('b'), ESearchCase::IgnoreCase), 10);
+	TEST_EQUAL(TEXT("FindLastChar"), UE::String::FindLastChar(TEXT("AbCABCAbCABC"), TEXT('a')), INDEX_NONE);
+	TEST_EQUAL(TEXT("FindLastChar"), UE::String::FindLastChar(TEXT("AbCABCAbCABC"), TEXT('D'), ESearchCase::IgnoreCase), INDEX_NONE);
+	TEST_EQUAL(TEXT("FindLastChar"), UE::String::FindLastChar(TEXT("AbCABCAbCABD"), TEXT('D')), 11);
+	TEST_EQUAL(TEXT("FindLastChar"), UE::String::FindLastChar(TEXT("AbCABCAbCABD"), TEXT('d'), ESearchCase::IgnoreCase), 11);
+
+	TEST_EQUAL(TEXT("FindLastChar"), UE::String::FindLastChar(TEXT(""), TEXT('A')), INDEX_NONE);
+	TEST_EQUAL(TEXT("FindLastChar"), UE::String::FindLastChar(TEXT("A"), TEXT('A')), 0);
+	TEST_EQUAL(TEXT("FindLastChar"), UE::String::FindLastChar(TEXT("A"), TEXT('A'), ESearchCase::IgnoreCase), 0);
+
+	TEST_EQUAL(TEXT("FindLastChar"), UE::String::FindLastChar("AbCABCAbCABC", 'B'), 10);
+
+	TEST_EQUAL(TEXT("FindLastChar"), UE::String::FindLastChar(FStringView(nullptr, 0), TEXT('A')), INDEX_NONE);
+	TEST_EQUAL(TEXT("FindLastChar"), UE::String::FindLastChar(FStringView(), TEXT('A')), INDEX_NONE);
+	TEST_EQUAL(TEXT("FindLastChar"), UE::String::FindLastChar(FString(), TEXT('A')), INDEX_NONE);
+}
+
+TEST_CASE_METHOD(FAutomationTestFixture, "Core::String::Misc::FindFirstOfAnyChar", "[Core][String][Smoke]")
+{
+	TEST_EQUAL(TEXT("FindFirstOfAnyChar"), UE::String::FindFirstOfAnyChar(TEXT("AbCABCAbcABC"), {TEXT('b')}), 1);
+	TEST_EQUAL(TEXT("FindFirstOfAnyChar"), UE::String::FindFirstOfAnyChar(TEXT("AbCABCAbcABC"), {TEXT('c'), TEXT('B')}), 4);
+	TEST_EQUAL(TEXT("FindFirstOfAnyChar"), UE::String::FindFirstOfAnyChar(TEXT("AbCABCAbcABC"), {TEXT('c'), TEXT('B')}, ESearchCase::IgnoreCase), 1);
+	TEST_EQUAL(TEXT("FindFirstOfAnyChar"), UE::String::FindFirstOfAnyChar(TEXT("AbCABCAbcABC"), {TEXT('D'), TEXT('a')}), INDEX_NONE);
+	TEST_EQUAL(TEXT("FindFirstOfAnyChar"), UE::String::FindFirstOfAnyChar(TEXT("AbCABCAbcABC"), {TEXT('E'), TEXT('D')}, ESearchCase::IgnoreCase), INDEX_NONE);
+	TEST_EQUAL(TEXT("FindFirstOfAnyChar"), UE::String::FindFirstOfAnyChar(TEXT("AbCABCAbcABD"), {TEXT('E'), TEXT('D')}), 11);
+	TEST_EQUAL(TEXT("FindFirstOfAnyChar"), UE::String::FindFirstOfAnyChar(TEXT("AbCABCAbcABD"), {TEXT('E'), TEXT('d')}, ESearchCase::IgnoreCase), 11);
+
+	TEST_EQUAL(TEXT("FindFirstOfAnyChar"), UE::String::FindFirstOfAnyChar(TEXT(""), {TEXT('A'), TEXT('B')}), INDEX_NONE);
+	TEST_EQUAL(TEXT("FindFirstOfAnyChar"), UE::String::FindFirstOfAnyChar(TEXT("A"), {TEXT('A'), TEXT('B')}), 0);
+	TEST_EQUAL(TEXT("FindFirstOfAnyChar"), UE::String::FindFirstOfAnyChar(TEXT("A"), {TEXT('A'), TEXT('B')}, ESearchCase::IgnoreCase), 0);
+
+	TEST_EQUAL(TEXT("FindFirstOfAnyChar"), UE::String::FindFirstOfAnyChar("AbCABCAbcABC", {'c', 'B'}), 4);
+
+	TEST_EQUAL(TEXT("FindFirstOfAnyChar"), UE::String::FindFirstOfAnyChar(FStringView(nullptr, 0), { TEXT('A'), TEXT('B') }), INDEX_NONE);
+	TEST_EQUAL(TEXT("FindFirstOfAnyChar"), UE::String::FindFirstOfAnyChar(FStringView(), { TEXT('A'), TEXT('B') }), INDEX_NONE);
+	TEST_EQUAL(TEXT("FindFirstOfAnyChar"), UE::String::FindFirstOfAnyChar(FString(), { TEXT('A'), TEXT('B') }), INDEX_NONE);
+}
+
+TEST_CASE_METHOD(FAutomationTestFixture, "Core::String::Misc::FindLastOfAnyChar", "[Core][String][Smoke]")
+{
+	TEST_EQUAL(TEXT("FindLastOfAnyChar"), UE::String::FindLastOfAnyChar(TEXT("AbCABCAbcABC"), {TEXT('b')}), 7);
+	TEST_EQUAL(TEXT("FindLastOfAnyChar"), UE::String::FindLastOfAnyChar(TEXT("AbCABCAbcABC"), {TEXT('c'), TEXT('B')}), 10);
+	TEST_EQUAL(TEXT("FindLastOfAnyChar"), UE::String::FindLastOfAnyChar(TEXT("AbCABCAbcABC"), {TEXT('c'), TEXT('B')}, ESearchCase::IgnoreCase), 11);
+	TEST_EQUAL(TEXT("FindLastOfAnyChar"), UE::String::FindLastOfAnyChar(TEXT("AbCABCAbcABC"), {TEXT('D'), TEXT('a')}), INDEX_NONE);
+	TEST_EQUAL(TEXT("FindLastOfAnyChar"), UE::String::FindLastOfAnyChar(TEXT("AbCABCAbcABC"), {TEXT('E'), TEXT('D')}, ESearchCase::IgnoreCase), INDEX_NONE);
+	TEST_EQUAL(TEXT("FindLastOfAnyChar"), UE::String::FindLastOfAnyChar(TEXT("AbCABCAbcABD"), {TEXT('E'), TEXT('D')}), 11);
+	TEST_EQUAL(TEXT("FindLastOfAnyChar"), UE::String::FindLastOfAnyChar(TEXT("AbCABCAbcABD"), {TEXT('E'), TEXT('d')}, ESearchCase::IgnoreCase), 11);
+
+	TEST_EQUAL(TEXT("FindLastOfAnyChar"), UE::String::FindLastOfAnyChar(TEXT(""), {TEXT('A'), TEXT('B')}), INDEX_NONE);
+	TEST_EQUAL(TEXT("FindLastOfAnyChar"), UE::String::FindLastOfAnyChar(TEXT("A"), {TEXT('A'), TEXT('B')}), 0);
+	TEST_EQUAL(TEXT("FindLastOfAnyChar"), UE::String::FindLastOfAnyChar(TEXT("A"), {TEXT('A'), TEXT('B')}, ESearchCase::IgnoreCase), 0);
+
+	TEST_EQUAL(TEXT("FindLastOfAnyChar"), UE::String::FindLastOfAnyChar("AbCABCAbcABC", {'c', 'B'}), 10);
+
+	TEST_EQUAL(TEXT("FindLastOfAnyChar"), UE::String::FindLastOfAnyChar(FStringView(nullptr, 0), { TEXT('A'), TEXT('B') }), INDEX_NONE);
+	TEST_EQUAL(TEXT("FindLastOfAnyChar"), UE::String::FindLastOfAnyChar(FStringView(), { TEXT('A'), TEXT('B') }), INDEX_NONE);
+	TEST_EQUAL(TEXT("FindLastOfAnyChar"), UE::String::FindLastOfAnyChar(FString(), { TEXT('A'), TEXT('B') }), INDEX_NONE);
+}
diff --git a/Engine/Source/Runtime/Core/Tests/String/ParseLinesTest.cpp b/Engine/Source/Runtime/Core/Tests/String/ParseLinesTest.cpp
--- a/Engine/Source/Runtime/Core/Tests/String/ParseLinesTest.cpp
+++ b/Engine/Source/Runtime/Core/Tests/String/ParseLinesTest.cpp
@@ -1,82 +1,75 @@
 // Copyright Epic Games, Inc. All Rights Reserved.
 
 #include "String/ParseLines.h"
-
 #include "Algo/Compare.h"
 #include "Containers/StringView.h"
 #include "Misc/StringBuilder.h"
-#include "String/Escape.h"
-#include "String/Join.h"
+#include "Templates/EqualTo.h"
 #include "TestHarness.h"
 
-namespace UE::String
+TEST_CASE_METHOD(FAutomationTestFixture, "Core::String::Misc::ParseLines", "[Core][String][Smoke]")
 {
+	using EOptions = UE::String::EParseLinesOptions;
 
-TEST_CASE("Core::String::ParseLines", "[Core][String][Smoke]")
-{
-	using FStringViewArray = TArray<FStringView, TInlineAllocator<8>>;
+	auto RunParseLinesTest = [this](FStringView View, std::initializer_list<FStringView> ExpectedLines, EOptions Options = EOptions::None)
+	{
+		TArray<FStringView, TInlineAllocator<8>> ResultLines;
+		UE::String::ParseLines(View, ResultLines, Options);
 
-	constexpr EParseLinesOptions KeepEmpty = EParseLinesOptions::None;
-	constexpr EParseLinesOptions SkipEmpty = EParseLinesOptions::SkipEmpty;
-	constexpr EParseLinesOptions Trim = EParseLinesOptions::Trim;
+		TEST_TRUE(TEXT("UE::String::ParseLines failed to parse"), Algo::Compare(ResultLines, ExpectedLines));
 
-	const auto [View, ExpectedLines, Options] = GENERATE_COPY(table<FStringView, FStringViewArray, EParseLinesOptions>(
-	{
-		{TEXTVIEW(""), {}, SkipEmpty},
-		{TEXTVIEW(""), {TEXTVIEW("")}, KeepEmpty},
-		{TEXTVIEW("\n"), {}, SkipEmpty},
-		{TEXTVIEW("\n"), {TEXTVIEW("")}, KeepEmpty},
-		{TEXTVIEW("\r"), {}, SkipEmpty},
-		{TEXTVIEW("\r"), {TEXTVIEW("")}, KeepEmpty},
-		{TEXTVIEW("\r\n"), {}, SkipEmpty},
-		{TEXTVIEW("\r\n"), {TEXTVIEW("")}, KeepEmpty},
-		{TEXTVIEW("\n\n"), {}, SkipEmpty},
-		{TEXTVIEW("\n\n"), {TEXTVIEW(""), TEXTVIEW("")}, KeepEmpty},
-		{TEXTVIEW("\r\r"), {}, SkipEmpty},
-		{TEXTVIEW("\r\r"), {TEXTVIEW(""), TEXTVIEW("")}, KeepEmpty},
-		{TEXTVIEW("\r\n\r\n"), {}, SkipEmpty},
-		{TEXTVIEW("\r\n\r\n"), {TEXTVIEW(""), TEXTVIEW("")}, KeepEmpty},
-		{TEXTVIEW("\r\nABC").Left(2), {}, SkipEmpty},
-		{TEXTVIEW("\r\nABC").Left(2), {TEXTVIEW("")}, KeepEmpty},
-		{TEXTVIEW("\r\nABC\r\nDEF").Left(5), {TEXTVIEW("ABC")}, SkipEmpty},
-		{TEXTVIEW("\r\nABC\r\nDEF").Left(5), {TEXTVIEW(""), TEXTVIEW("ABC")}, KeepEmpty},
-		{TEXTVIEW("ABC DEF"), {TEXTVIEW("ABC DEF")}, SkipEmpty},
-		{TEXTVIEW("\nABC DEF\n"), {TEXTVIEW("ABC DEF")}, SkipEmpty},
-		{TEXTVIEW("\nABC DEF\n"), {TEXTVIEW(""), TEXTVIEW("ABC DEF")}, KeepEmpty},
-		{TEXTVIEW("\rABC DEF\r"), {TEXTVIEW("ABC DEF")}, SkipEmpty},
-		{TEXTVIEW("\rABC DEF\r"), {TEXTVIEW(""), TEXTVIEW("ABC DEF")}, KeepEmpty},
-		{TEXTVIEW("\r\nABC DEF\r\n"), {TEXTVIEW("ABC DEF")}, SkipEmpty},
-		{TEXTVIEW("\r\nABC DEF\r\n"), {TEXTVIEW(""), TEXTVIEW("ABC DEF")}, KeepEmpty},
-		{TEXTVIEW("\r\n\r\nABC DEF\r\n\r\n"), {TEXTVIEW("ABC DEF")}, SkipEmpty},
-		{TEXTVIEW("\r\n\r\nABC DEF\r\n\r\n"), {TEXTVIEW(""), TEXTVIEW(""), TEXTVIEW("ABC DEF"), TEXTVIEW("")}, KeepEmpty},
-		{TEXTVIEW("ABC\nDEF"), {TEXTVIEW("ABC"), TEXTVIEW("DEF")}, SkipEmpty},
-		{TEXTVIEW("ABC\rDEF"), {TEXTVIEW("ABC"), TEXTVIEW("DEF")}, SkipEmpty},
-		{TEXTVIEW("\r\nABC\r\nDEF\r\n"), {TEXTVIEW("ABC"), TEXTVIEW("DEF")}, SkipEmpty},
-		{TEXTVIEW("\r\nABC\r\nDEF\r\n"), {TEXTVIEW(""), TEXTVIEW("ABC"), TEXTVIEW("DEF")}, KeepEmpty},
-		{TEXTVIEW("\r\nABC\r\n\r\nDEF\r\n"), {TEXTVIEW("ABC"), TEXTVIEW("DEF")}, SkipEmpty},
-		{TEXTVIEW("\r\nABC\r\n\r\nDEF\r\n"), {TEXTVIEW(""), TEXTVIEW("ABC"), TEXTVIEW(""), TEXTVIEW("DEF")}, KeepEmpty},
-		{TEXTVIEW(" \t\r\n\t ABC \t\r\n\t \t\r\n\t DEF \t\r\n"), {TEXTVIEW(" \t"), TEXTVIEW("\t ABC \t"), TEXTVIEW("\t \t"), TEXTVIEW("\t DEF \t")}, SkipEmpty},
-		{TEXTVIEW(" \t\r\n\t ABC \t\r\n\t \t\r\n\t DEF \t\r\n"), {TEXTVIEW("ABC"), TEXTVIEW("DEF")}, SkipEmpty | Trim},
-		{TEXTVIEW(" \t\r\n\t ABC \t\r\n\t \t\r\n\t DEF \t\r\n"), {TEXTVIEW(" \t"), TEXTVIEW("\t ABC \t"), TEXTVIEW("\t \t"), TEXTVIEW("\t DEF \t")}, KeepEmpty},
-		{TEXTVIEW(" \t\r\n\t ABC \t\r\n\t \t\r\n\t DEF \t\r\n"), {TEXTVIEW(""), TEXTVIEW("ABC"), TEXTVIEW(""), TEXTVIEW("DEF")}, KeepEmpty | Trim},
-	}));
+		/*if (!Algo::Compare(ResultLines, ExpectedLines))
+		{
+			TStringBuilder<512> Error;
+			Error << TEXTVIEW("UE::String::ParseLines failed to parse \"") << FString(View).ReplaceCharWithEscapedChar() << TEXTVIEW("\" result {");
+			Error.JoinQuoted(ResultLines, TEXTVIEW(", "), TEXTVIEW("\""));
+			Error << TEXTVIEW("} expected {");
+			Error.JoinQuoted(ExpectedLines, TEXTVIEW(", "), TEXTVIEW("\""));
+			Error << TEXTVIEW("}");
+			FAIL_CHECK(Error.ToString());
+		}*/
+	};
 
-	FStringViewArray ActualLines;
-	ParseLines(View, ActualLines, Options);
-	const bool bEqual = Algo::Compare(ActualLines, ExpectedLines);
-	if (!bEqual)
-	{
-		TStringBuilder<256> Input, Expected, Actual;
-		Input << QuoteEscape(View);
-		Expected << TEXTVIEW("[") << JoinBy(ExpectedLines, QuoteEscape, TEXTVIEW(", ")) << TEXTVIEW("]");
-		Actual << TEXTVIEW("[") << JoinBy(ActualLines, QuoteEscape, TEXTVIEW(", ")) << TEXTVIEW("]");
-		CAPTURE(Input, Expected, Actual);
-		CHECK(bEqual);
-	}
-	else
-	{
-		CHECK(bEqual);
-	}
+	constexpr EOptions KeepEmpty = EOptions::None;
+	constexpr EOptions SkipEmpty = EOptions::SkipEmpty;
+	constexpr EOptions Trim = EOptions::Trim;
+
+	RunParseLinesTest(TEXTVIEW(""), {}, SkipEmpty);
+	RunParseLinesTest(TEXTVIEW(""), {TEXTVIEW("")}, KeepEmpty);
+	RunParseLinesTest(TEXTVIEW("\n"), {}, SkipEmpty);
+	RunParseLinesTest(TEXTVIEW("\n"), {TEXTVIEW("")}, KeepEmpty);
+	RunParseLinesTest(TEXTVIEW("\r"), {}, SkipEmpty);
+	RunParseLinesTest(TEXTVIEW("\r"), {TEXTVIEW("")}, KeepEmpty);
+	RunParseLinesTest(TEXTVIEW("\r\n"), {}, SkipEmpty);
+	RunParseLinesTest(TEXTVIEW("\r\n"), {TEXTVIEW("")}, KeepEmpty);
+	RunParseLinesTest(TEXTVIEW("\n\n"), {}, SkipEmpty);
+	RunParseLinesTest(TEXTVIEW("\n\n"), {TEXTVIEW(""), TEXTVIEW("")}, KeepEmpty);
+	RunParseLinesTest(TEXTVIEW("\r\r"), {}, SkipEmpty);
+	RunParseLinesTest(TEXTVIEW("\r\r"), {TEXTVIEW(""), TEXTVIEW("")}, KeepEmpty);
+	RunParseLinesTest(TEXTVIEW("\r\n\r\n"), {}, SkipEmpty);
+	RunParseLinesTest(TEXTVIEW("\r\n\r\n"), {TEXTVIEW(""), TEXTVIEW("")}, KeepEmpty);
+	RunParseLinesTest(TEXTVIEW("\r\nABC").Left(2), {}, SkipEmpty);
+	RunParseLinesTest(TEXTVIEW("\r\nABC").Left(2), {TEXTVIEW("")}, KeepEmpty);
+	RunParseLinesTest(TEXTVIEW("\r\nABC\r\nDEF").Left(5), {TEXTVIEW("ABC")}, SkipEmpty);
+	RunParseLinesTest(TEXTVIEW("\r\nABC\r\nDEF").Left(5), {TEXTVIEW(""), TEXTVIEW("ABC")}, KeepEmpty);
+	RunParseLinesTest(TEXTVIEW("ABC DEF"), {TEXTVIEW("ABC DEF")}, SkipEmpty);
+	RunParseLinesTest(TEXTVIEW("\nABC DEF\n"), {TEXTVIEW("ABC DEF")}, SkipEmpty);
+	RunParseLinesTest(TEXTVIEW("\nABC DEF\n"), {TEXTVIEW(""), TEXTVIEW("ABC DEF")}, KeepEmpty);
+	RunParseLinesTest(TEXTVIEW("\rABC DEF\r"), {TEXTVIEW("ABC DEF")}, SkipEmpty);
+	RunParseLinesTest(TEXTVIEW("\rABC DEF\r"), {TEXTVIEW(""), TEXTVIEW("ABC DEF")}, KeepEmpty);
+	RunParseLinesTest(TEXTVIEW("\r\nABC DEF\r\n"), {TEXTVIEW("ABC DEF")}, SkipEmpty);
+	RunParseLinesTest(TEXTVIEW("\r\nABC DEF\r\n"), {TEXTVIEW(""), TEXTVIEW("ABC DEF")}, KeepEmpty);
+	RunParseLinesTest(TEXTVIEW("\r\n\r\nABC DEF\r\n\r\n"), {TEXTVIEW("ABC DEF")}, SkipEmpty);
+	RunParseLinesTest(TEXTVIEW("\r\n\r\nABC DEF\r\n\r\n"), {TEXTVIEW(""), TEXTVIEW(""), TEXTVIEW("ABC DEF"), TEXTVIEW("")}, KeepEmpty);
+	RunParseLinesTest(TEXTVIEW("ABC\nDEF"), {TEXTVIEW("ABC"), TEXTVIEW("DEF")}, SkipEmpty);
+	RunParseLinesTest(TEXTVIEW("ABC\rDEF"), {TEXTVIEW("ABC"), TEXTVIEW("DEF")}, SkipEmpty);
+	RunParseLinesTest(TEXTVIEW("\r\nABC\r\nDEF\r\n"), {TEXTVIEW("ABC"), TEXTVIEW("DEF")}, SkipEmpty);
+	RunParseLinesTest(TEXTVIEW("\r\nABC\r\nDEF\r\n"), {TEXTVIEW(""), TEXTVIEW("ABC"), TEXTVIEW("DEF")}, KeepEmpty);
+	RunParseLinesTest(TEXTVIEW("\r\nABC\r\n\r\nDEF\r\n"), {TEXTVIEW("ABC"), TEXTVIEW("DEF")}, SkipEmpty);
+	RunParseLinesTest(TEXTVIEW("\r\nABC\r\n\r\nDEF\r\n"), {TEXTVIEW(""), TEXTVIEW("ABC"), TEXTVIEW(""), TEXTVIEW("DEF")}, KeepEmpty);
+	RunParseLinesTest(TEXTVIEW(" \t\r\n\t ABC \t\r\n\t \t\r\n\t DEF \t\r\n"), {TEXTVIEW(" \t"), TEXTVIEW("\t ABC \t"), TEXTVIEW("\t \t"), TEXTVIEW("\t DEF \t")}, SkipEmpty);
+	RunParseLinesTest(TEXTVIEW(" \t\r\n\t ABC \t\r\n\t \t\r\n\t DEF \t\r\n"), {TEXTVIEW("ABC"), TEXTVIEW("DEF")}, SkipEmpty | Trim);
+	RunParseLinesTest(TEXTVIEW(" \t\r\n\t ABC \t\r\n\t \t\r\n\t DEF \t\r\n"), {TEXTVIEW(" \t"), TEXTVIEW("\t ABC \t"), TEXTVIEW("\t \t"), TEXTVIEW("\t DEF \t")}, KeepEmpty);
+	RunParseLinesTest(TEXTVIEW(" \t\r\n\t ABC \t\r\n\t \t\r\n\t DEF \t\r\n"), {TEXTVIEW(""), TEXTVIEW("ABC"), TEXTVIEW(""), TEXTVIEW("DEF")}, KeepEmpty | Trim);
 }
 
-} // UE::String
