Add new Unit Tests

diff --git a/Engine/Source/Runtime/Core/Tests/String/ParseTokensTest.cpp b/Engine/Source/Runtime/Core/Tests/String/ParseTokensTest.cpp
--- a/Engine/Source/Runtime/Core/Tests/String/ParseTokensTest.cpp
+++ b/Engine/Source/Runtime/Core/Tests/String/ParseTokensTest.cpp
@@ -4,161 +4,166 @@
 
 #include "Algo/Compare.h"
 #include "Misc/StringBuilder.h"
-#include "String/Escape.h"
-#include "String/Join.h"
+#include "Templates/EqualTo.h"
 #include "TestHarness.h"
 
-namespace UE::String
-{
+#if WITH_DEV_AUTOMATION_TESTS
 
-TEST_CASE("Core::String::ParseTokens", "[Core][String][Smoke]")
+TEST_CASE_METHOD(FAutomationTestFixture, "Core::String::Misc::Parse Tokens By String", "[Core][String][Smoke]")
 {
-	using FCharArray = TArray<TCHAR, TInlineAllocator<8>>;
-	using FStringViewArray = TArray<FStringView, TInlineAllocator<8>>;
-
-	constexpr EParseTokensOptions KeepEmpty = EParseTokensOptions::None;
-	constexpr EParseTokensOptions SkipEmpty = EParseTokensOptions::SkipEmpty;
-	constexpr EParseTokensOptions IgnoreCase = EParseTokensOptions::IgnoreCase;
-	constexpr EParseTokensOptions Trim = EParseTokensOptions::Trim;
+	using EOptions = UE::String::EParseTokensOptions;
 
-	SECTION("String")
+	auto RunParseTokensTest = [this](
+		FStringView View,
+		std::initializer_list<FStringView> Delimiters,
+		std::initializer_list<FStringView> ExpectedTokens,
+		EOptions Options = EOptions::None)
 	{
-		const auto [View, DelimiterArray, ExpectedTokens, Options] = GENERATE_COPY(table<FStringView, FStringViewArray, FStringViewArray, EParseTokensOptions>(
+		TArray<FStringView, TInlineAllocator<8>> ResultTokens;
+		if (GetNum(Delimiters) == 1)
 		{
-			{TEXTVIEW(""),                {},                   {},                                                     SkipEmpty},
-			{TEXTVIEW(""),                {},                   {TEXTVIEW("")},                                         KeepEmpty},
-			{TEXTVIEW("ABC"),             {},                   {TEXTVIEW("ABC")},                                      KeepEmpty},
-
-			{TEXTVIEW(""),                {TEXTVIEW(",")},      {},                                                     SkipEmpty},
-			{TEXTVIEW(""),                {TEXTVIEW(",")},      {TEXTVIEW("")},                                         KeepEmpty},
-			{TEXTVIEW(","),               {TEXTVIEW(",")},      {},                                                     SkipEmpty},
-			{TEXTVIEW(","),               {TEXTVIEW(",")},      {TEXTVIEW(""), TEXTVIEW("")},                           KeepEmpty},
-			{TEXTVIEW(",,"),              {TEXTVIEW(",")},      {},                                                     SkipEmpty},
-			{TEXTVIEW(",,"),              {TEXTVIEW(",")},      {TEXTVIEW(""), TEXTVIEW(""), TEXTVIEW("")},             KeepEmpty},
-			{TEXTVIEW(", ,"),             {TEXTVIEW(",")},      {TEXTVIEW(" ")},                                        SkipEmpty},
-			{TEXTVIEW(", ,"),             {TEXTVIEW(",")},      {},                                                     SkipEmpty | Trim},
-			{TEXTVIEW(", ,"),             {TEXTVIEW(",")},      {TEXTVIEW(""), TEXTVIEW(" "), TEXTVIEW("")},            KeepEmpty},
-			{TEXTVIEW(", ,"),             {TEXTVIEW(",")},      {TEXTVIEW(""), TEXTVIEW(""), TEXTVIEW("")},             KeepEmpty | Trim},
-			{TEXTVIEW("ABC"),             {TEXTVIEW(",")},      {TEXTVIEW("ABC")},                                      KeepEmpty},
-			{TEXTVIEW("A,,C"),            {TEXTVIEW(",")},      {TEXTVIEW("A"), TEXTVIEW("C")},                         SkipEmpty},
-			{TEXTVIEW("A,,C"),            {TEXTVIEW(",")},      {TEXTVIEW("A"), TEXTVIEW(""), TEXTVIEW("C")},           KeepEmpty},
-			{TEXTVIEW("A,\tB\t,C"),       {TEXTVIEW(",")},      {TEXTVIEW("A"), TEXTVIEW("\tB\t"), TEXTVIEW("C")},      KeepEmpty},
-			{TEXTVIEW(",A, B ,C,"),       {TEXTVIEW(",")},      {TEXTVIEW("A"), TEXTVIEW(" B "), TEXTVIEW("C")},        SkipEmpty},
-			{TEXTVIEW(",A, B ,C,"),       {TEXTVIEW(",")},      {TEXTVIEW(""), TEXTVIEW("A"), TEXTVIEW(" B "), TEXTVIEW("C"), TEXTVIEW("")}, KeepEmpty},
-			{TEXTVIEW("A\u2022B\u2022C"), {TEXTVIEW("\u2022")}, {TEXTVIEW("A"), TEXTVIEW("B"), TEXTVIEW("C")},          KeepEmpty},
-
-			{TEXTVIEW("ABCDABCD"), {TEXTVIEW("AB")},   {TEXTVIEW("CD"), TEXTVIEW("CD")},               SkipEmpty},
-			{TEXTVIEW("ABCDABCD"), {TEXTVIEW("AB")},   {TEXTVIEW(""), TEXTVIEW("CD"), TEXTVIEW("CD")}, KeepEmpty},
-			{TEXTVIEW("ABCDABCD"), {TEXTVIEW("ABCD")}, {},                                             SkipEmpty},
-			{TEXTVIEW("ABCDABCD"), {TEXTVIEW("ABCD")}, {TEXTVIEW(""), TEXTVIEW(""), TEXTVIEW("")},     KeepEmpty},
-			{TEXTVIEW("ABCDABCD"), {TEXTVIEW("DA")},   {TEXTVIEW("ABC"), TEXTVIEW("BCD")},             KeepEmpty},
-
-			{TEXTVIEW("ABCDABCD"), {TEXTVIEW("B"),  TEXTVIEW("D")},  {TEXTVIEW("A"), TEXTVIEW("C"), TEXTVIEW("A"), TEXTVIEW("C")},               SkipEmpty},
-			{TEXTVIEW("ABCDABCD"), {TEXTVIEW("B"),  TEXTVIEW("D")},  {TEXTVIEW("A"), TEXTVIEW("C"), TEXTVIEW("A"), TEXTVIEW("C"), TEXTVIEW("")}, KeepEmpty},
-			{TEXTVIEW("ABCDABCD"), {TEXTVIEW("BC"), TEXTVIEW("DA")}, {TEXTVIEW("A"), TEXTVIEW("D")},                                             SkipEmpty},
-			{TEXTVIEW("ABCDABCD"), {TEXTVIEW("BC"), TEXTVIEW("DA")}, {TEXTVIEW("A"), TEXTVIEW(""), TEXTVIEW(""), TEXTVIEW("D")},                 KeepEmpty},
-
-			{TEXTVIEW("AbCdaBcDAbCd"), {TEXTVIEW("Bc"), TEXTVIEW("da")}, {TEXTVIEW("AbC"), TEXTVIEW("DAbCd")}, SkipEmpty},
-			{TEXTVIEW("AbCdaBcDAbCd"), {TEXTVIEW("Bc"), TEXTVIEW("da")}, {TEXTVIEW("A"), TEXTVIEW("d")},       SkipEmpty | IgnoreCase},
-
-			{TEXTVIEW("A\u2022\u2022B,,C"), {TEXTVIEW(",,"), TEXTVIEW("\u2022\u2022")},                     {TEXTVIEW("A"), TEXTVIEW("B"), TEXTVIEW("C")}, KeepEmpty},
-			{TEXTVIEW("A\u2022\u2022B\u0085\u0085C"), {TEXTVIEW("\u0085\u0085"), TEXTVIEW("\u2022\u2022")}, {TEXTVIEW("A"), TEXTVIEW("B"), TEXTVIEW("C")}, KeepEmpty},
-		}));
-
-		FStringViewArray ActualTokens;
-		if (GetNum(DelimiterArray) == 1)
-		{
-			ParseTokens(View, GetData(DelimiterArray)[0], ActualTokens, Options);
+			UE::String::ParseTokens(View, GetData(Delimiters)[0], ResultTokens, Options);
 		}
 		else
 		{
-			ParseTokensMultiple(View, DelimiterArray, ActualTokens, Options);
+			UE::String::ParseTokensMultiple(View, Delimiters, ResultTokens, Options);
 		}
 
-		const bool bEqual = Algo::Compare(ActualTokens, ExpectedTokens);
-		if (!bEqual)
-		{
-			TStringBuilder<256> Input, Delimiters, Expected, Actual;
-			Input << QuoteEscape(View);
-			Delimiters << TEXTVIEW("[") << JoinBy(DelimiterArray, QuoteEscape, TEXTVIEW(", ")) << TEXTVIEW("]");
-			Expected << TEXTVIEW("[") << JoinBy(ExpectedTokens, QuoteEscape, TEXTVIEW(", ")) << TEXTVIEW("]");
-			Actual << TEXTVIEW("[") << JoinBy(ActualTokens, QuoteEscape, TEXTVIEW(", ")) << TEXTVIEW("]");
-			CAPTURE(Input, Delimiters, Expected, Actual);
-			CHECK(bEqual);
-		}
-		else
+		TEST_TRUE(TEXT("UE::String::ParseTokens[Multiple] failed to parse"), Algo::Compare(ResultTokens, ExpectedTokens));
+
+		if (!Algo::Compare(ResultTokens, ExpectedTokens))
 		{
-			CHECK(bEqual);
+			TStringBuilder<512> Error;
+			Error << TEXTVIEW("UE::String::ParseTokens[Multiple] failed to parse \"") << View << TEXTVIEW("\" with delimiters {");
+			Error.JoinQuoted(Delimiters, TEXTVIEW(", "), TEXTVIEW("\""));
+			Error << TEXTVIEW("} result {");
+			Error.JoinQuoted(ResultTokens, TEXTVIEW(", "), TEXTVIEW("\""));
+			Error << TEXTVIEW("} expected {");
+			Error.JoinQuoted(ExpectedTokens, TEXTVIEW(", "), TEXTVIEW("\""));
+			Error << TEXTVIEW("}");
+			FAIL_CHECK(Error.ToString());
 		}
-	}
+	};
+
+	constexpr EOptions KeepEmpty = EOptions::None;
+	constexpr EOptions SkipEmpty = EOptions::SkipEmpty;
+	constexpr EOptions IgnoreCase = EOptions::IgnoreCase;
+	constexpr EOptions Trim = EOptions::Trim;
+
+	RunParseTokensTest(TEXT(""),         {},                       {},                                 SkipEmpty);
+	RunParseTokensTest(TEXT(""),         {},                       {TEXT("")},                         KeepEmpty);
+	RunParseTokensTest(TEXT("ABC"),      {},                       {TEXT("ABC")});
+
+	RunParseTokensTest(TEXT(""),         {TEXT(",")},              {},                                 SkipEmpty);
+	RunParseTokensTest(TEXT(""),         {TEXT(",")},              {TEXT("")},                         KeepEmpty);
+	RunParseTokensTest(TEXT(","),        {TEXT(",")},              {},                                 SkipEmpty);
+	RunParseTokensTest(TEXT(","),        {TEXT(",")},              {TEXT(""), TEXT("")},               KeepEmpty);
+	RunParseTokensTest(TEXT(",,"),       {TEXT(",")},              {},                                 SkipEmpty);
+	RunParseTokensTest(TEXT(",,"),       {TEXT(",")},              {TEXT(""), TEXT(""), TEXT("")},     KeepEmpty);
+	RunParseTokensTest(TEXT(", ,"),      {TEXT(",")},              {TEXT(" ")},                        SkipEmpty);
+	RunParseTokensTest(TEXT(", ,"),      {TEXT(",")},              {},                                 SkipEmpty | Trim);
+	RunParseTokensTest(TEXT(", ,"),      {TEXT(",")},              {TEXT(""), TEXT(" "), TEXT("")},    KeepEmpty);
+	RunParseTokensTest(TEXT(", ,"),      {TEXT(",")},              {TEXT(""), TEXT(""), TEXT("")},    KeepEmpty | Trim);
+	RunParseTokensTest(TEXT("ABC"),      {TEXT(",")},              {TEXT("ABC")});
+	RunParseTokensTest(TEXT("A,,C"),     {TEXT(",")},              {TEXT("A"), TEXT("C")},             SkipEmpty);
+	RunParseTokensTest(TEXT("A,,C"),     {TEXT(",")},              {TEXT("A"), TEXT(""), TEXT("C")},   KeepEmpty);
+	RunParseTokensTest(TEXT("A,\tB\t,C"), {TEXT(",")},             {TEXT("A"), TEXT("\tB\t"), TEXT("C")});
+	RunParseTokensTest(TEXT(",A, B ,C,"), {TEXT(",")},             {TEXT("A"), TEXT(" B "), TEXT("C")},                     SkipEmpty);
+	RunParseTokensTest(TEXT(",A, B ,C,"), {TEXT(",")},             {TEXT(""), TEXT("A"), TEXT(" B "), TEXT("C"), TEXT("")}, KeepEmpty);
+	RunParseTokensTest(TEXT("A\u2022B\u2022C"), {TEXT("\u2022")},  {TEXT("A"), TEXT("B"), TEXT("C")});
+
+	RunParseTokensTest(TEXT("ABCDABCD"), {TEXT("AB")},             {TEXT("CD"), TEXT("CD")},           SkipEmpty);
+	RunParseTokensTest(TEXT("ABCDABCD"), {TEXT("AB")},             {TEXT(""), TEXT("CD"), TEXT("CD")}, KeepEmpty);
+	RunParseTokensTest(TEXT("ABCDABCD"), {TEXT("ABCD")},           {},                                 SkipEmpty);
+	RunParseTokensTest(TEXT("ABCDABCD"), {TEXT("ABCD")},           {TEXT(""), TEXT(""), TEXT("")},     KeepEmpty);
+	RunParseTokensTest(TEXT("ABCDABCD"), {TEXT("DA")},             {TEXT("ABC"), TEXT("BCD")});
+
+	RunParseTokensTest(TEXT("ABCDABCD"), {TEXT("B"),  TEXT("D")},  {TEXT("A"), TEXT("C"), TEXT("A"), TEXT("C")},           SkipEmpty);
+	RunParseTokensTest(TEXT("ABCDABCD"), {TEXT("B"),  TEXT("D")},  {TEXT("A"), TEXT("C"), TEXT("A"), TEXT("C"), TEXT("")}, KeepEmpty);
+	RunParseTokensTest(TEXT("ABCDABCD"), {TEXT("BC"), TEXT("DA")}, {TEXT("A"), TEXT("D")},                                 SkipEmpty);
+	RunParseTokensTest(TEXT("ABCDABCD"), {TEXT("BC"), TEXT("DA")}, {TEXT("A"), TEXT(""), TEXT(""), TEXT("D")},             KeepEmpty);
+
+	RunParseTokensTest(TEXT("AbCdaBcDAbCd"), {TEXT("Bc"), TEXT("da")}, {TEXT("AbC"), TEXT("DAbCd")}, SkipEmpty);
+	RunParseTokensTest(TEXT("AbCdaBcDAbCd"), {TEXT("Bc"), TEXT("da")}, {TEXT("A"), TEXT("d")}, SkipEmpty | IgnoreCase);
+
+	RunParseTokensTest(TEXT("A\u2022\u2022B,,C"), {TEXT(",,"), TEXT("\u2022\u2022")},                      {TEXT("A"), TEXT("B"), TEXT("C")});
+	RunParseTokensTest(TEXT("A\u2022\u2022B\u0085\u0085C"), {TEXT("\u0085\u0085"), TEXT("\u2022\u2022")},  {TEXT("A"), TEXT("B"), TEXT("C")});
+}
 
-	SECTION("Char")
+TEST_CASE_METHOD(FAutomationTestFixture, "Core::String::Misc::Parse Tokens By Token", "[Core][String][Smoke]")
+{
+	using EOptions = UE::String::EParseTokensOptions;
+
+	auto RunParseTokensTest = [this](
+		FStringView View,
+		std::initializer_list<TCHAR> Delimiters,
+		std::initializer_list<FStringView> ExpectedTokens,
+		EOptions Options = EOptions::None)
 	{
-		const auto [View, DelimiterArray, ExpectedTokens, Options] = GENERATE_COPY(table<FStringView, FCharArray, FStringViewArray, EParseTokensOptions>(
+		TArray<FStringView, TInlineAllocator<8>> ResultTokens;
+		if (GetNum(Delimiters) == 1)
 		{
-			{TEXTVIEW(""),                {},               {},                                    SkipEmpty},
-			{TEXTVIEW(""),                {},               {TEXT("")},                            KeepEmpty},
-			{TEXTVIEW("ABC"),             {},               {TEXT("ABC")},                         KeepEmpty},
-
-			{TEXTVIEW(""),                {TEXT(',')},      {},                                    SkipEmpty},
-			{TEXTVIEW(""),                {TEXT(',')},      {TEXT("")},                            KeepEmpty},
-			{TEXTVIEW(","),               {TEXT(',')},      {},                                    SkipEmpty},
-			{TEXTVIEW(","),               {TEXT(',')},      {TEXT(""), TEXT("")},                  KeepEmpty},
-			{TEXTVIEW(",,"),              {TEXT(',')},      {},                                    SkipEmpty},
-			{TEXTVIEW(",,"),              {TEXT(',')},      {TEXT(""), TEXT(""), TEXT("")},        KeepEmpty},
-			{TEXTVIEW(", ,"),             {TEXT(',')},      {TEXT(" ")},                           SkipEmpty},
-			{TEXTVIEW(", ,"),             {TEXT(',')},      {},                                    SkipEmpty | Trim},
-			{TEXTVIEW(", ,"),             {TEXT(',')},      {TEXT(""), TEXT(" "), TEXT("")},       KeepEmpty},
-			{TEXTVIEW(", ,"),             {TEXT(',')},      {TEXT(""), TEXT(""), TEXT("")},        KeepEmpty | Trim},
-			{TEXTVIEW("ABC"),             {TEXT(',')},      {TEXT("ABC")},                         KeepEmpty},
-			{TEXTVIEW("A,,C"),            {TEXT(',')},      {TEXT("A"), TEXT("C")},                SkipEmpty},
-			{TEXTVIEW("A,,C"),            {TEXT(',')},      {TEXT("A"), TEXT(""), TEXT("C")},      KeepEmpty},
-			{TEXTVIEW("A,\tB\t,C"),       {TEXT(',')},      {TEXT("A"), TEXT("\tB\t"), TEXT("C")}, KeepEmpty},
-			{TEXTVIEW(",A, B ,C,"),       {TEXT(',')},      {TEXT("A"), TEXT(" B "), TEXT("C")},                     SkipEmpty},
-			{TEXTVIEW(",A, B ,C,"),       {TEXT(',')},      {TEXT(""), TEXT("A"), TEXT(" B "), TEXT("C"), TEXT("")}, KeepEmpty},
-			{TEXTVIEW("A\u2022B\u2022C"), {TEXT('\u2022')}, {TEXT("A"), TEXT("B"), TEXT("C")},     KeepEmpty},
-
-			{TEXTVIEW("ABCDABCD"),        {TEXT('B'), TEXT('D')},           {TEXT("A"), TEXT("C"), TEXT("A"), TEXT("C")},           SkipEmpty},
-			{TEXTVIEW("ABCDABCD"),        {TEXT('B'), TEXT('D')},           {TEXT("A"), TEXT("C"), TEXT("A"), TEXT("C"), TEXT("")}, KeepEmpty},
-			{TEXTVIEW("A\u2022B,C"),      {TEXT(','), TEXT('\u2022')},      {TEXT("A"), TEXT("B"), TEXT("C")},                      KeepEmpty},
-			{TEXTVIEW("A\u2022B\u0085C"), {TEXT('\u0085'), TEXT('\u2022')}, {TEXT("A"), TEXT("B"), TEXT("C")},                      KeepEmpty},
-
-			{TEXTVIEW("ABC"), {TEXT('b')}, {TEXT("ABC")}, SkipEmpty},
-			{TEXTVIEW("ABC"), {TEXT('b')}, {TEXT("A"), TEXT("C")}, SkipEmpty | IgnoreCase},
-
-			{TEXTVIEW("AbCdaBcD"), {TEXT('B'), TEXT('d')}, {TEXT("AbC"), TEXT("A"), TEXT("cD")},         SkipEmpty},
-			{TEXTVIEW("AbCdaBcD"), {TEXT('B'), TEXT('d')}, {TEXT("A"), TEXT("C"), TEXT("a"), TEXT("c")}, SkipEmpty | IgnoreCase},
-
-			{TEXTVIEW("A\u2022B\u2022C"), {TEXT('\u2022'), TEXT('b')}, {TEXT("A"), TEXT("B"), TEXT("C")}, SkipEmpty},
-			{TEXTVIEW("A\u2022B\u2022C"), {TEXT('\u2022'), TEXT('b')}, {TEXT("A"), TEXT("C")},            SkipEmpty | IgnoreCase},
-		}));
-
-		FStringViewArray ActualTokens;
-		if (GetNum(DelimiterArray) == 1)
-		{
-			ParseTokens(View, GetData(DelimiterArray)[0], ActualTokens, Options);
+			UE::String::ParseTokens(View, GetData(Delimiters)[0], ResultTokens, Options);
 		}
 		else
 		{
-			ParseTokensMultiple(View, DelimiterArray, ActualTokens, Options);
+			UE::String::ParseTokensMultiple(View, Delimiters, ResultTokens, Options);
 		}
 
-		const bool bEqual = Algo::Compare(ActualTokens, ExpectedTokens);
-		if (!bEqual)
-		{
-			TStringBuilder<256> Input, Delimiters, Expected, Actual;
-			Input << QuoteEscape(View);
-			Delimiters << TEXTVIEW("[") << JoinQuoted(DelimiterArray, TEXTVIEW(", "), TEXT('\'')) << TEXTVIEW("]");
-			Expected << TEXTVIEW("[") << JoinBy(ExpectedTokens, QuoteEscape, TEXTVIEW(", ")) << TEXTVIEW("]");
-			Actual << TEXTVIEW("[") << JoinBy(ActualTokens, QuoteEscape, TEXTVIEW(", ")) << TEXTVIEW("]");
-			CAPTURE(Input, Delimiters, Expected, Actual);
-			CHECK(bEqual);
-		}
-		else
+		TEST_TRUE(TEXT("UE::String::ParseTokens[Multiple] failed to parse"), Algo::Compare(ResultTokens, ExpectedTokens));
+
+		if (!Algo::Compare(ResultTokens, ExpectedTokens))
 		{
-			CHECK(bEqual);
+			TStringBuilder<512> Error;
+			Error << TEXTVIEW("UE::String::ParseTokens[Multiple] failed to parse \"") << View << TEXTVIEW("\" with delimiters {");
+			Error.JoinQuoted(Delimiters, TEXTVIEW(", "), TEXTVIEW("'"));
+			Error << TEXTVIEW("} result {");
+			Error.JoinQuoted(ResultTokens, TEXTVIEW(", "), TEXTVIEW("\""));
+			Error << TEXTVIEW("} expected {");
+			Error.JoinQuoted(ExpectedTokens, TEXTVIEW(", "), TEXTVIEW("\""));
+			Error << TEXTVIEW("}");
+			FAIL_CHECK(Error.ToString());
 		}
-	}
+	};
+
+	constexpr EOptions KeepEmpty = EOptions::None;
+	constexpr EOptions SkipEmpty = EOptions::SkipEmpty;
+	constexpr EOptions IgnoreCase = EOptions::IgnoreCase;
+	constexpr EOptions Trim = EOptions::Trim;
+
+	RunParseTokensTest(TEXT(""),         {},                       {},                                 SkipEmpty);
+	RunParseTokensTest(TEXT(""),         {},                       {TEXT("")},                         KeepEmpty);
+	RunParseTokensTest(TEXT("ABC"),      {},                       {TEXT("ABC")});
+
+	RunParseTokensTest(TEXT(""),         {TEXT(',')},              {},                                 SkipEmpty);
+	RunParseTokensTest(TEXT(""),         {TEXT(',')},              {TEXT("")},                         KeepEmpty);
+	RunParseTokensTest(TEXT(","),        {TEXT(',')},              {},                                 SkipEmpty);
+	RunParseTokensTest(TEXT(","),        {TEXT(',')},              {TEXT(""), TEXT("")},               KeepEmpty);
+	RunParseTokensTest(TEXT(",,"),       {TEXT(',')},              {},                                 SkipEmpty);
+	RunParseTokensTest(TEXT(",,"),       {TEXT(',')},              {TEXT(""), TEXT(""), TEXT("")},     KeepEmpty);
+	RunParseTokensTest(TEXT(", ,"),      {TEXT(',')},              {TEXT(" ")},                        SkipEmpty);
+	RunParseTokensTest(TEXT(", ,"),      {TEXT(',')},              {},                                 SkipEmpty | Trim);
+	RunParseTokensTest(TEXT(", ,"),      {TEXT(',')},              {TEXT(""), TEXT(" "), TEXT("")},    KeepEmpty);
+	RunParseTokensTest(TEXT(", ,"),      {TEXT(',')},              {TEXT(""), TEXT(""), TEXT("")},    KeepEmpty | Trim);
+	RunParseTokensTest(TEXT("ABC"),      {TEXT(',')},              {TEXT("ABC")});
+	RunParseTokensTest(TEXT("A,,C"),     {TEXT(',')},              {TEXT("A"), TEXT("C")},             SkipEmpty);
+	RunParseTokensTest(TEXT("A,,C"),     {TEXT(',')},              {TEXT("A"), TEXT(""), TEXT("C")},   KeepEmpty);
+	RunParseTokensTest(TEXT("A,\tB\t,C"), {TEXT(',')},             {TEXT("A"), TEXT("\tB\t"), TEXT("C")});
+	RunParseTokensTest(TEXT(",A, B ,C,"), {TEXT(',')},             {TEXT("A"), TEXT(" B "), TEXT("C")},                     SkipEmpty);
+	RunParseTokensTest(TEXT(",A, B ,C,"), {TEXT(',')},             {TEXT(""), TEXT("A"), TEXT(" B "), TEXT("C"), TEXT("")}, KeepEmpty);
+	RunParseTokensTest(TEXT("A\u2022B\u2022C"), {TEXT('\u2022')},  {TEXT("A"), TEXT("B"), TEXT("C")});
+
+	RunParseTokensTest(TEXT("ABCDABCD"), {TEXT('B'),  TEXT('D')},                  {TEXT("A"), TEXT("C"), TEXT("A"), TEXT("C")},           SkipEmpty);
+	RunParseTokensTest(TEXT("ABCDABCD"), {TEXT('B'),  TEXT('D')},                  {TEXT("A"), TEXT("C"), TEXT("A"), TEXT("C"), TEXT("")}, KeepEmpty);
+	RunParseTokensTest(TEXT("A\u2022B,C"), {TEXT(','), TEXT('\u2022')},            {TEXT("A"), TEXT("B"), TEXT("C")});
+	RunParseTokensTest(TEXT("A\u2022B\u0085C"), {TEXT('\u0085'), TEXT('\u2022')},  {TEXT("A"), TEXT("B"), TEXT("C")});
+
+	RunParseTokensTest(TEXT("ABC"), {TEXT('b')}, {TEXT("ABC")}, SkipEmpty);
+	RunParseTokensTest(TEXT("ABC"), {TEXT('b')}, {TEXT("A"), TEXT("C")}, SkipEmpty | IgnoreCase);
+	RunParseTokensTest(TEXT("AbCdaBcD"), {TEXT('B'),  TEXT('d')}, {TEXT("AbC"), TEXT("A"), TEXT("cD")}, SkipEmpty);
+	RunParseTokensTest(TEXT("AbCdaBcD"), {TEXT('B'),  TEXT('d')}, {TEXT("A"), TEXT("C"), TEXT("a"), TEXT("c")}, SkipEmpty | IgnoreCase);
+	RunParseTokensTest(TEXT("A\u2022B\u2022C"), {TEXT('\u2022'), TEXT('b')}, {TEXT("A"), TEXT("B"), TEXT("C")}, SkipEmpty);
+	RunParseTokensTest(TEXT("A\u2022B\u2022C"), {TEXT('\u2022'), TEXT('b')}, {TEXT("A"), TEXT("C")}, SkipEmpty | IgnoreCase);
 }
 
-} // UE::String
+#endif // WITH_DEV_AUTOMATION_TESTS
diff --git a/Engine/Source/Runtime/Core/Tests/String/StringTest.cpp b/Engine/Source/Runtime/Core/Tests/String/StringTest.cpp
--- a/Engine/Source/Runtime/Core/Tests/String/StringTest.cpp
+++ b/Engine/Source/Runtime/Core/Tests/String/StringTest.cpp
@@ -10,12 +10,12 @@
 #include "Serialization/MemoryWriter.h"
 #include "TestHarness.h"
 
-TEST_CASE("Core::String::FString::SanitizeFloat", "[Core][String][Smoke]")
+TEST_CASE_METHOD(FAutomationTestFixture, "Core::String::FString::SanitizeFloat", "[Core][String][Smoke]")
 {
-	auto DoTest = [](const double InVal, const int32 InMinFractionalDigits, const FString& InExpected)
+	auto DoTest = [this](const double InVal, const int32 InMinFractionalDigits, const FString& InExpected)
 	{
 		const FString Result = FString::SanitizeFloat(InVal, InMinFractionalDigits);
-		TestEqual(FString::Printf(TEXT("%f (%d digits) failure: result '%s' (expected '%s')"), InVal, InMinFractionalDigits, *Result, *InExpected), Result, InExpected);
+		CHECK_EQUAL(Result, InExpected);
 	};
 
 	DoTest(+0.0, 0, TEXT("0"));
@@ -43,11 +43,12 @@ TEST_CASE("Core::String::FString::SanitizeFloat", "[Core][String][Smoke]")
 	DoTest(-100.1010, 4, TEXT("-100.1010"));
 }
 
-TEST_CASE("Core::String::FString::AppendInt", "[Core][String][Smoke]")
+TEST_CASE_METHOD(FAutomationTestFixture, "Core::String::FString::AppendInt", "[Core][String][Smoke]")
 {
-	auto DoTest = [](const TCHAR* Call, const FString& Result, const TCHAR* InExpected)
+	auto DoTest = [this](const TCHAR* Call, const FString& Result, const TCHAR* InExpected)
 	{
-		TestEqual(FString::Printf(TEXT("'%s' failure: result '%s' (expected '%s')"), Call, *Result, InExpected), Result, InExpected);
+		INFO(Call);
+		CHECK_EQUAL(Result, InExpected);
 	};
 
 	{
@@ -77,11 +78,12 @@ TEST_CASE("Core::String::FString::AppendInt", "[Core][String][Smoke]")
 	}
 }
 
-TEST_CASE("Core::String::FString::Unicode", "[Core][String][Smoke]")
+TEST_CASE_METHOD(FAutomationTestFixture, "Core::String::FString::Unicode", "[Core][String][Smoke]")
 {
-	auto DoTest = [](const TCHAR* Call, const FString& Result, const TCHAR* InExpected)
+	auto DoTest = [this](const TCHAR* Call, const FString& Result, const TCHAR* InExpected)
 	{
-		TestEqual(FString::Printf(TEXT("'%s' failure: result '%s' (expected '%s')"), Call, *Result, InExpected), Result, InExpected);
+		INFO(Call);
+		CHECK_EQUAL(Result, InExpected);
 	};
 	
 	// Test data used to verify basic processing of a Unicode character outside the BMP
@@ -117,110 +119,178 @@ TEST_CASE("Core::String::FString::Unicode", "[Core][String][Smoke]")
 	}
 }
 
-TEST_CASE("Core::String::FString::LexTryParseString", "[Core][String][Smoke]")
+TEST_CASE_METHOD(FAutomationTestFixture, "Core::String::FString::LexTryParseString", "[Core][String][Smoke]")
 {
 	// Test that LexFromString can intepret all the numerical formats we expect it to
-	{
-		// Test float values
 
+	SECTION("Test float values")
+	{
 		float Value;
 
-		// Basic numbers
-		TestTrue(TEXT("(float conversion from string) basic numbers"), LexTryParseString(Value, (TEXT("1"))) && Value == 1);
-		TestTrue(TEXT("(float conversion from string) basic numbers"), LexTryParseString(Value, (TEXT("1.0"))) && Value == 1);
-		TestTrue(TEXT("(float conversion from string) basic numbers"), LexTryParseString(Value, (TEXT(".5"))) && Value == 0.5);
-		TestTrue(TEXT("(float conversion from string) basic numbers"), LexTryParseString(Value, (TEXT("1."))) && Value == 1);
-
-		// Variations of 0
-		TestTrue(TEXT("(float conversion from string) variations of 0"), LexTryParseString(Value, (TEXT("0"))) && Value == 0);
-		TestTrue(TEXT("(float conversion from string) variations of 0"), LexTryParseString(Value, (TEXT("-0"))) && Value == 0);
-		TestTrue(TEXT("(float conversion from string) variations of 0"), LexTryParseString(Value, (TEXT("0.0"))) && Value == 0);
-		TestTrue(TEXT("(float conversion from string) variations of 0"), LexTryParseString(Value, (TEXT(".0"))) && Value == 0);
-		TestTrue(TEXT("(float conversion from string) variations of 0"), LexTryParseString(Value, (TEXT("0."))) && Value == 0);
-		TestTrue(TEXT("(float conversion from string) variations of 0"), LexTryParseString(Value, (TEXT("0. 111"))) && Value == 0);
-
-		// Scientific notation
-		TestTrue(TEXT("(float conversion from string) scientific notation"), LexTryParseString(Value, (TEXT("1.0e+10"))) && Value == 1.0e+10f);
-		TestTrue(TEXT("(float conversion from string) scientific notation"), LexTryParseString(Value, (TEXT("1.99999999e-11"))) && Value == 1.99999999e-11f);
-		TestTrue(TEXT("(float conversion from string) scientific notation"), LexTryParseString(Value, (TEXT("1e+10"))) && Value == 1e+10f);
-
-		// Non-finite special numbers
-		TestTrue(TEXT("(float conversion from string) inf"), LexTryParseString(Value, (TEXT("inf"))));
-		TestTrue(TEXT("(float conversion from string) nan"), LexTryParseString(Value, (TEXT("nan"))));
-		TestTrue(TEXT("(float conversion from string) nan(ind)"), LexTryParseString(Value, (TEXT("nan(ind)"))));
-
-		// nan/inf etc. are detected from the start of the string, regardless of any other characters that come afterwards
-		TestTrue(TEXT("(float conversion from string) nananananananana"), LexTryParseString(Value, (TEXT("nananananananana"))));
-		TestTrue(TEXT("(float conversion from string) nan(ind)!"), LexTryParseString(Value, (TEXT("nan(ind)!"))));
-		TestTrue(TEXT("(float conversion from string) infinity"), LexTryParseString(Value, (TEXT("infinity"))));
-
-		// Some numbers with whitespace
-		TestTrue(TEXT("(float conversion from string) whitespace"), LexTryParseString(Value, (TEXT("   2.5   "))) && Value == 2.5);
-		TestTrue(TEXT("(float conversion from string) whitespace"), LexTryParseString(Value, (TEXT("\t3.0\t"))) && Value == 3.0);
-		TestTrue(TEXT("(float conversion from string) whitespace"), LexTryParseString(Value, (TEXT("4.0   \t"))) && Value == 4.0);
-		TestTrue(TEXT("(float conversion from string) whitespace"), LexTryParseString(Value, (TEXT("\r\n5.25"))) && Value == 5.25);
-		TestTrue(TEXT("(float conversion from string) whitespace"), LexTryParseString(Value, (TEXT(" 6 . 2 "))) && Value == 6.0);
-		TestTrue(TEXT("(float conversion from string) whitespace"), LexTryParseString(Value, (TEXT(" 56 . 2 "))) && Value == 56.0);
-		TestTrue(TEXT("(float conversion from string) whitespace"), LexTryParseString(Value, (TEXT(" 5 6 . 2 "))) && Value == 5.0);
-
-		// Failure cases
-		TestFalse(TEXT("(float no conversion from string) not a number"), LexTryParseString(Value, (TEXT("not a number"))));
-		TestFalse(TEXT("(float no conversion from string) <empty string>"), LexTryParseString(Value, (TEXT(""))));
-		TestFalse(TEXT("(float conversion from string) ."), LexTryParseString(Value, (TEXT("."))));
+		SECTION("Basic numbers")
+		{
+			CHECK(LexTryParseString(Value, TEXT("1")));
+			CHECK(Value == 1);
+			CHECK(LexTryParseString(Value, TEXT("1.0")));
+			CHECK(Value == 1);
+			CHECK(LexTryParseString(Value, TEXT(".5")));
+			CHECK(Value == 0.5);
+			CHECK(LexTryParseString(Value, TEXT("1.")));
+			CHECK(Value == 1);
+		}
+
+		SECTION("Variations of 0")
+		{
+			CHECK(LexTryParseString(Value, TEXT("0")));
+			CHECK(Value == 0);
+			CHECK(LexTryParseString(Value, TEXT("-0")));
+			CHECK(Value == 0);
+			CHECK(LexTryParseString(Value, TEXT("0.0")));
+			CHECK(Value == 0);
+			CHECK(LexTryParseString(Value, TEXT(".0")));
+			CHECK(Value == 0);
+			CHECK(LexTryParseString(Value, TEXT("0.")));
+			CHECK(Value == 0);
+			CHECK(LexTryParseString(Value, TEXT("0. 111")));
+			CHECK(Value == 0);
+		}
+
+		SECTION("Scientific notation")
+		{
+			CHECK(LexTryParseString(Value, TEXT("1.0e+10")));
+			CHECK(Value == 1.0e+10f);
+			CHECK(LexTryParseString(Value, TEXT("1.99999999e-11")));
+			CHECK(Value == 1.99999999e-11f);
+			CHECK(LexTryParseString(Value, TEXT("1e+10")));
+			CHECK(Value == 1e+10f);
+		}
+
+		SECTION("Non-finite special numbers")
+		{
+			CHECK(LexTryParseString(Value, TEXT("inf")));
+			CHECK(LexTryParseString(Value, TEXT("nan")));
+			CHECK(LexTryParseString(Value, TEXT("nan(ind)")));
+		}
+
+		SECTION("nan/inf etc. are detected from the start of the string, regardless of any other characters that come afterwards")
+		{
+			CHECK(LexTryParseString(Value, TEXT("nananananananana")));
+			CHECK(LexTryParseString(Value, TEXT("nan(ind)!")));
+			CHECK(LexTryParseString(Value, TEXT("infinity")));
+		}
+
+		SECTION("Some numbers with whitespace")
+		{	
+			CHECK(LexTryParseString(Value, TEXT("   2.5   ")));
+			CHECK(Value == 2.5);
+			CHECK(LexTryParseString(Value, TEXT("\t3.0\t")));
+			CHECK(Value == 3.0);
+			CHECK(LexTryParseString(Value, TEXT("4.0   \t")));
+			CHECK(Value == 4.0);
+			CHECK(LexTryParseString(Value, TEXT("\r\n5.25")));
+			CHECK(Value == 5.25);
+			CHECK(LexTryParseString(Value, TEXT(" 6 . 2 ")));
+			CHECK(Value == 6.0);
+			CHECK(LexTryParseString(Value, TEXT(" 56 . 2 ")));
+			CHECK(Value == 56.0);
+			CHECK(LexTryParseString(Value, TEXT(" 5 6 . 2 ")));
+			CHECK(Value == 5.0);
+		}
+		
+		SECTION("Failure cases")
+		{
+			CHECK_FALSE(LexTryParseString(Value, TEXT("not a number")));
+			CHECK_FALSE(LexTryParseString(Value, TEXT("")));
+			CHECK_FALSE(LexTryParseString(Value, TEXT(".")));
+		}
+		
 	}
 
+	SECTION("Test integer values")
 	{
-		// Test integer values
-
 		int32 Value;
 
-		// Basic numbers
-		TestTrue(TEXT("(int32 conversion from string) basic numbers"), LexTryParseString(Value, (TEXT("1"))) && Value == 1);
-		TestTrue(TEXT("(int32 conversion from string) basic numbers"), LexTryParseString(Value, (TEXT("1.0"))) && Value == 1);
-		TestTrue(TEXT("(int32 conversion from string) basic numbers"), LexTryParseString(Value, (TEXT("3.1"))) && Value == 3);
-		TestTrue(TEXT("(int32 conversion from string) basic numbers"), LexTryParseString(Value, (TEXT("0.5"))) && Value == 0);
-		TestTrue(TEXT("(int32 conversion from string) basic numbers"), LexTryParseString(Value, (TEXT("1."))) && Value == 1);
-
-		// Variations of 0
-		TestTrue(TEXT("(int32 conversion from string) variations of 0"), LexTryParseString(Value, (TEXT("0"))) && Value == 0);
-		TestTrue(TEXT("(int32 conversion from string) variations of 0"), LexTryParseString(Value, (TEXT("0.0"))) && Value == 0);
-		TestFalse(TEXT("(int32 conversion from string) variations of 0"), LexTryParseString(Value, (TEXT(".0"))) && Value == 0);
-		TestTrue(TEXT("(int32 conversion from string) variations of 0"), LexTryParseString(Value, (TEXT("0."))) && Value == 0);
-
-		// Scientific notation
-		TestTrue(TEXT("(int32 conversion from string) scientific notation"), LexTryParseString(Value, (TEXT("1.0e+10"))) && Value == 1);
-		TestTrue(TEXT("(int32 conversion from string) scientific notation"), LexTryParseString(Value, (TEXT("6.0e-10"))) && Value == 6);
-		TestTrue(TEXT("(int32 conversion from string) scientific notation"), LexTryParseString(Value, (TEXT("0.0e+10"))) && Value == 0);
-		TestTrue(TEXT("(int32 conversion from string) scientific notation"), LexTryParseString(Value, (TEXT("0.0e-10"))) && Value == 0);
-		TestTrue(TEXT("(int32 conversion from string) scientific notation"), LexTryParseString(Value, (TEXT("3e+10"))) && Value == 3);
-		TestTrue(TEXT("(int32 conversion from string) scientific notation"), LexTryParseString(Value, (TEXT("4e-10"))) && Value == 4);
-
-		// Some numbers with whitespace
-		TestTrue(TEXT("(int32 conversion from string) whitespace"), LexTryParseString(Value, (TEXT("   2.5   "))) && Value == 2);
-		TestTrue(TEXT("(int32 conversion from string) whitespace"), LexTryParseString(Value, (TEXT("\t3.0\t"))) && Value == 3);
-		TestTrue(TEXT("(int32 conversion from string) whitespace"), LexTryParseString(Value, (TEXT("4.0   \t"))) && Value == 4);
-		TestTrue(TEXT("(int32 conversion from string) whitespace"), LexTryParseString(Value, (TEXT("\r\n5.25"))) && Value == 5);
-		TestTrue(TEXT("(int32 conversion from string) whitespace"), LexTryParseString(Value, (TEXT(" 6 . 2 "))) && Value == 6);
-		TestTrue(TEXT("(int32 conversion from string) whitespace"), LexTryParseString(Value, (TEXT(" 56 . 2 "))) && Value == 56);
-		TestTrue(TEXT("(int32 conversion from string) whitespace"), LexTryParseString(Value, (TEXT(" 5 6 . 2 "))) && Value == 5);
-
-		// Non-finite special numbers. All shouldn't parse into an int
-		TestFalse(TEXT("(int32 no conversion from string) inf"), LexTryParseString(Value, (TEXT("inf"))));
-		TestFalse(TEXT("(int32 no conversion from string) nan"), LexTryParseString(Value, (TEXT("nan"))));
-		TestFalse(TEXT("(int32 no conversion from string) nan(ind)"), LexTryParseString(Value, (TEXT("nan(ind)"))));
-		TestFalse(TEXT("(int32 no conversion from string) nananananananana"), LexTryParseString(Value, (TEXT("nananananananana"))));
-		TestFalse(TEXT("(int32 no conversion from string) nan(ind)!"), LexTryParseString(Value, (TEXT("nan(ind)!"))));
-		TestFalse(TEXT("(int32 no conversion from string) infinity"), LexTryParseString(Value, (TEXT("infinity"))));
-		TestFalse(TEXT("(float no conversion from string) ."), LexTryParseString(Value, (TEXT("."))));
-		TestFalse(TEXT("(float no conversion from string) <empyty string>"), LexTryParseString(Value, (TEXT(""))));
+		SECTION("Basic numbers")
+		{
+			CHECK(LexTryParseString(Value, TEXT("1")));
+			CHECK(Value == 1);
+			CHECK(LexTryParseString(Value, TEXT("1.0")));
+			CHECK(Value == 1);
+			CHECK(LexTryParseString(Value, TEXT("3.1")));
+			CHECK(Value == 3);
+			CHECK(LexTryParseString(Value, TEXT("0.5")));
+			CHECK(Value == 0);
+			CHECK(LexTryParseString(Value, TEXT("1.")));
+			CHECK(Value == 1);
+		}
+
+		SECTION("Variations of 0")
+		{	
+			CHECK(LexTryParseString(Value, TEXT("0")));
+			CHECK(Value == 0);
+			CHECK(LexTryParseString(Value, TEXT("0.0")));
+			CHECK(Value == 0);
+			CHECK_FALSE(LexTryParseString(Value, TEXT(".0")));
+			CHECK(Value == 0);
+			CHECK(LexTryParseString(Value, TEXT("0.")));
+			CHECK(Value == 0);
+		}
+
+		SECTION("Scientific notation")
+		{
+			CHECK(LexTryParseString(Value, TEXT("1.0e+10")));
+			CHECK(Value == 1);
+			CHECK(LexTryParseString(Value, TEXT("6.0e-10")));
+			CHECK(Value == 6);
+			CHECK(LexTryParseString(Value, TEXT("0.0e+10")));
+			CHECK(Value == 0);
+			CHECK(LexTryParseString(Value, TEXT("0.0e-10")));
+			CHECK(Value == 0);
+			CHECK(LexTryParseString(Value, TEXT("3e+10")));
+			CHECK(Value == 3);
+			CHECK(LexTryParseString(Value, TEXT("4e-10")));
+			CHECK(Value == 4);
+		}
+
+		SECTION("Some numbers with whitespace")
+		{
+			CHECK(LexTryParseString(Value, TEXT("   2.5   ")));
+			CHECK(Value == 2);
+			CHECK(LexTryParseString(Value, TEXT("\t3.0\t")));
+			CHECK(Value == 3);
+			CHECK(LexTryParseString(Value, TEXT("4.0   \t")));
+			CHECK(Value == 4);
+			CHECK(LexTryParseString(Value, TEXT("\r\n5.25")));
+			CHECK(Value == 5);
+			CHECK(LexTryParseString(Value, TEXT(" 6 . 2 ")));
+			CHECK(Value == 6);
+			CHECK(LexTryParseString(Value, TEXT(" 56 . 2 ")));
+			CHECK(Value == 56);
+			CHECK(LexTryParseString(Value, TEXT(" 5 6 . 2 ")));
+			CHECK(Value == 5);
+		}
+		
+		SECTION("Non-finite special numbers. All shouldn't parse into an int")
+		{
+			CHECK_FALSE(LexTryParseString(Value, (TEXT("inf"))));
+			CHECK_FALSE(LexTryParseString(Value, (TEXT("nan"))));
+			CHECK_FALSE(LexTryParseString(Value, (TEXT("nan(ind)"))));
+			CHECK_FALSE(LexTryParseString(Value, (TEXT("nananananananana"))));
+			CHECK_FALSE(LexTryParseString(Value, (TEXT("nan(ind)!"))));
+			CHECK_FALSE(LexTryParseString(Value, (TEXT("infinity"))));
+			CHECK_FALSE(LexTryParseString(Value, (TEXT("."))));
+			CHECK_FALSE(LexTryParseString(Value, (TEXT(""))));
+		}
+		
 	}
 }
 
-TEST_CASE("Core::String::FString::Substring", "[Core][String][Smoke]")
+TEST_CASE_METHOD(FAutomationTestFixture, "Core::String::FString::Substring", "[Core][String][Smoke]")
 {
-	auto DoTest = [](const TCHAR* Call, const FString& Result, const TCHAR* InExpected)
+	auto DoTest = [this](const TCHAR* Call, const FString& Result, const TCHAR* InExpected)
 	{
-		TestEqual(FString::Printf(TEXT("'%s' failure: result '%s' (expected '%s')"), Call, *Result, InExpected), Result, InExpected);
+		INFO(Call);
+		CHECK_EQUAL(Result, InExpected);
 	};
 
 	const FString TestString(TEXT("0123456789"));
@@ -233,127 +303,142 @@ TEST_CASE("Core::String::FString::Substring", "[Core][String][Smoke]")
 	Inline##TestName.Operation##Inline(__VA_ARGS__); \
 	DoTest(TEXT("Inline" #TestName), Inline##TestName, ExpectedResult); 
 
-	// Left
-	SUBSTRINGTEST(Left, TEXT("0123"), Left, 4);
-	SUBSTRINGTEST(ExactLengthLeft, *TestString, Left, 10);
-	SUBSTRINGTEST(LongerThanLeft, *TestString, Left, 20);
-	SUBSTRINGTEST(ZeroLeft, TEXT(""), Left, 0);
-	SUBSTRINGTEST(NegativeLeft, TEXT(""), Left, -1);
-
-	// LeftChop
-	SUBSTRINGTEST(LeftChop, TEXT("012345"), LeftChop, 4);
-	SUBSTRINGTEST(ExactLengthLeftChop, TEXT(""), LeftChop, 10);
-	SUBSTRINGTEST(LongerThanLeftChop, TEXT(""), LeftChop, 20);
-	SUBSTRINGTEST(ZeroLeftChop, *TestString, LeftChop, 0);
-	SUBSTRINGTEST(NegativeLeftChop, *TestString, LeftChop, -1);
-
-	// Right
-	SUBSTRINGTEST(Right, TEXT("6789"), Right, 4);
-	SUBSTRINGTEST(ExactLengthRight, *TestString, Right, 10);
-	SUBSTRINGTEST(LongerThanRight, *TestString, Right, 20);
-	SUBSTRINGTEST(ZeroRight, TEXT(""), Right, 0);
-	SUBSTRINGTEST(NegativeRight, TEXT(""), Right, -1);
-
-	// RightChop
-	SUBSTRINGTEST(RightChop, TEXT("456789"), RightChop, 4);
-	SUBSTRINGTEST(ExactLengthRightChop, TEXT(""), RightChop, 10);
-	SUBSTRINGTEST(LongerThanRightChop, TEXT(""), RightChop, 20);
-	SUBSTRINGTEST(ZeroRightChop, *TestString, RightChop, 0);
-	SUBSTRINGTEST(NegativeRightChop, *TestString, RightChop, -1);
-
-	// Mid
-	SUBSTRINGTEST(Mid, TEXT("456789"), Mid, 4);
-	SUBSTRINGTEST(MidCount, TEXT("4567"), Mid, 4, 4);
-	SUBSTRINGTEST(MidCountFullLength, *TestString, Mid, 0, 10);
-	SUBSTRINGTEST(MidCountOffEnd, TEXT("89"), Mid, 8, 4);
-	SUBSTRINGTEST(MidStartAfterEnd, TEXT(""), Mid, 20);
-	SUBSTRINGTEST(MidZeroCount, TEXT(""), Mid, 5, 0);
-	SUBSTRINGTEST(MidNegativeCount, TEXT(""), Mid, 5, -1);
-	SUBSTRINGTEST(MidNegativeStartNegativeEnd, TEXT(""), Mid, -5, 1);
-	SUBSTRINGTEST(MidNegativeStartPositiveEnd, TEXT("012"), Mid, -1, 4);
-	SUBSTRINGTEST(MidNegativeStartBeyondEnd, *TestString, Mid, -1, 15);
+	SECTION("Left")
+	{
+		SUBSTRINGTEST(Left, TEXT("0123"), Left, 4);
+		SUBSTRINGTEST(ExactLengthLeft, *TestString, Left, 10);
+		SUBSTRINGTEST(LongerThanLeft, *TestString, Left, 20);
+		SUBSTRINGTEST(ZeroLeft, TEXT(""), Left, 0);
+		SUBSTRINGTEST(NegativeLeft, TEXT(""), Left, -1);
+	}
+
+	SECTION("LeftChop")
+	{
+		SUBSTRINGTEST(LeftChop, TEXT("012345"), LeftChop, 4);
+		SUBSTRINGTEST(ExactLengthLeftChop, TEXT(""), LeftChop, 10);
+		SUBSTRINGTEST(LongerThanLeftChop, TEXT(""), LeftChop, 20);
+		SUBSTRINGTEST(ZeroLeftChop, *TestString, LeftChop, 0);
+		SUBSTRINGTEST(NegativeLeftChop, *TestString, LeftChop, -1);
+	}
+
+	SECTION("Right")
+	{
+		SUBSTRINGTEST(Right, TEXT("6789"), Right, 4);
+		SUBSTRINGTEST(ExactLengthRight, *TestString, Right, 10);
+		SUBSTRINGTEST(LongerThanRight, *TestString, Right, 20);
+		SUBSTRINGTEST(ZeroRight, TEXT(""), Right, 0);
+		SUBSTRINGTEST(NegativeRight, TEXT(""), Right, -1);
+	}
+	
+	SECTION("RightChop")
+	{
+		SUBSTRINGTEST(RightChop, TEXT("456789"), RightChop, 4);
+		SUBSTRINGTEST(ExactLengthRightChop, TEXT(""), RightChop, 10);
+		SUBSTRINGTEST(LongerThanRightChop, TEXT(""), RightChop, 20);
+		SUBSTRINGTEST(ZeroRightChop, *TestString, RightChop, 0);
+		SUBSTRINGTEST(NegativeRightChop, *TestString, RightChop, -1);
+	}
+
+	SECTION("Mid")
+	{
+		SUBSTRINGTEST(Mid, TEXT("456789"), Mid, 4);
+		SUBSTRINGTEST(MidCount, TEXT("4567"), Mid, 4, 4);
+		SUBSTRINGTEST(MidCountFullLength, *TestString, Mid, 0, 10);
+		SUBSTRINGTEST(MidCountOffEnd, TEXT("89"), Mid, 8, 4);
+		SUBSTRINGTEST(MidStartAfterEnd, TEXT(""), Mid, 20);
+		SUBSTRINGTEST(MidZeroCount, TEXT(""), Mid, 5, 0);
+		SUBSTRINGTEST(MidNegativeCount, TEXT(""), Mid, 5, -1);
+		SUBSTRINGTEST(MidNegativeStartNegativeEnd, TEXT(""), Mid, -5, 1);
+		SUBSTRINGTEST(MidNegativeStartPositiveEnd, TEXT("012"), Mid, -1, 4);
+		SUBSTRINGTEST(MidNegativeStartBeyondEnd, *TestString, Mid, -1, 15);
+	}
 
 #undef SUBSTRINGTEST
 }
 
-TEST_CASE("Core::String::FString::FromStringView", "[Core][String][Smoke]")
+TEST_CASE_METHOD(FAutomationTestFixture, "Core::String::FString::FromStringView", "[Core][String][Smoke]")
 {
-	// Verify basic construction and assignment from a string view.
+	SECTION("Verify basic construction and assignment from a string view.")
 	{
 		const TCHAR* Literal = TEXT("Literal");
 		const ANSICHAR* AnsiLiteral = "Literal";
-		TestEqual(TEXT("String(StringView)"), FString(FStringView(Literal)), Literal);
-		TestEqual(TEXT("String(AnsiStringView)"), FString(FAnsiStringView(AnsiLiteral)), Literal);
-		TestEqual(TEXT("String = StringView"), FString(TEXT("Temp")) = FStringView(Literal), Literal);
+		CAPTURE(Literal,AnsiLiteral);
+		CHECK_EQUAL(FString(FStringView(Literal)), Literal);
+		CHECK_EQUAL(FString(FAnsiStringView(AnsiLiteral)), Literal);
+		CHECK_EQUAL((FString(TEXT("Temp")) = FStringView(Literal)), Literal);
 
 		FStringView EmptyStringView;
 		FString EmptyString(EmptyStringView);
-		TestTrue(TEXT("String(EmptyStringView)"), EmptyString.IsEmpty());
-		TestTrue(TEXT("String(EmptyStringView) (No Allocation)"), EmptyString.GetAllocatedSize() == 0);
+		CHECK(EmptyString.IsEmpty());
+		CHECK(EmptyString.GetAllocatedSize() == 0);
 
 		EmptyString = TEXT("Temp");
 		EmptyString = EmptyStringView;
-		TestTrue(TEXT("String = EmptyStringView"), EmptyString.IsEmpty());
-		TestTrue(TEXT("String = EmptyStringView (No Allocation)"), EmptyString.GetAllocatedSize() == 0);
+		CHECK(EmptyString.IsEmpty());
+		CHECK(EmptyString.GetAllocatedSize() == 0);
 	}
 
-	// Verify assignment from a view of itself.
+	SECTION("Verify assignment from a view of itself.")
 	{
 		FString AssignEntireString(TEXT("AssignEntireString"));
 		AssignEntireString = FStringView(AssignEntireString);
-		TestEqual(TEXT("String = StringView(String)"), AssignEntireString, TEXT("AssignEntireString"));
+		CAPTURE(AssignEntireString);
+		CHECK_EQUAL(AssignEntireString, TEXT("AssignEntireString"));
 
 		FString AssignStartOfString(TEXT("AssignStartOfString"));
 		AssignStartOfString = FStringView(AssignStartOfString).Left(11);
-		TestEqual(TEXT("String = StringView(String).Left"), AssignStartOfString, TEXT("AssignStart"));
+		CAPTURE(AssignStartOfString);
+		CHECK_EQUAL(AssignStartOfString, TEXT("AssignStart"));
 
 		FString AssignEndOfString(TEXT("AssignEndOfString"));
 		AssignEndOfString = FStringView(AssignEndOfString).Right(11);
-		TestEqual(TEXT("String = StringView(String).Right"), AssignEndOfString, TEXT("EndOfString"));
+		CAPTURE(AssignEndOfString);
+		CHECK_EQUAL(AssignEndOfString, TEXT("EndOfString"));
 
 		FString AssignMiddleOfString(TEXT("AssignMiddleOfString"));
 		AssignMiddleOfString = FStringView(AssignMiddleOfString).Mid(6, 6);
-		TestEqual(TEXT("String = StringView(String).Mid"), AssignMiddleOfString, TEXT("Middle"));
+		CAPTURE(AssignMiddleOfString);
+		CHECK_EQUAL(AssignMiddleOfString, TEXT("Middle"));
 	}
 
-	// Verify operators taking string views and character arrays
+	SECTION("Verify operators taking string views and character arrays")
 	{
 		FStringView RhsStringView = FStringView(TEXT("RhsNotSZ"), 3);
 		FString MovePlusSVResult = FString(TEXT("Lhs")) + RhsStringView;
-		TestEqual(TEXT("Move String + StringView"), MovePlusSVResult, TEXT("LhsRhs"));
+		CHECK_EQUAL(MovePlusSVResult, TEXT("LhsRhs"));
 
 		FString CopyLhs(TEXT("Lhs"));
 		FString CopyPlusSVResult = CopyLhs + RhsStringView;
-		TestEqual(TEXT("Copy String + StringView"), CopyPlusSVResult, TEXT("LhsRhs"));
+		CHECK_EQUAL(CopyPlusSVResult, TEXT("LhsRhs"));
 
 		FString MovePlusTCHARsResult = FString(TEXT("Lhs")) + TEXT("Rhs");
-		TestEqual(TEXT("Move String + TCHAR*"), MovePlusTCHARsResult, TEXT("LhsRhs"));
+		CHECK_EQUAL(MovePlusTCHARsResult, TEXT("LhsRhs"));
 
 		FString CopyPlusTCHARsResult = CopyLhs + TEXT("Rhs");
-		TestEqual(TEXT("Copy String + TCHAR*"), CopyPlusTCHARsResult, TEXT("LhsRhs"));
+		CHECK_EQUAL(CopyPlusTCHARsResult, TEXT("LhsRhs"));
 
 		FStringView LhsStringView = FStringView(TEXT("LhsNotSZ"), 3);
 		FString SVPlusMoveResult = LhsStringView + FString(TEXT("Rhs"));
-		TestEqual(TEXT("StringView + Move String"), SVPlusMoveResult, TEXT("LhsRhs"));
+		CHECK_EQUAL(SVPlusMoveResult, TEXT("LhsRhs"));
 
 		FString CopyRhs(TEXT("Rhs"));
 		FString SVPlusCopyResult = LhsStringView + CopyRhs;
-		TestEqual(TEXT("StringView + Copy String"), SVPlusCopyResult, TEXT("LhsRhs"));
+		CHECK_EQUAL(SVPlusCopyResult, TEXT("LhsRhs"));
 
 		FString TCHARsPlusMoveResult = TEXT("Lhs") + FString(TEXT("Rhs"));
-		TestEqual(TEXT("TCHAR* + Move String"), TCHARsPlusMoveResult, TEXT("LhsRhs"));
+		CHECK_EQUAL(TCHARsPlusMoveResult, TEXT("LhsRhs"));
 
 		FString TCHARsPlusCopyResult = TEXT("Lhs") + CopyRhs;
-		TestEqual(TEXT("TCHAR* + Copy String"), TCHARsPlusCopyResult, TEXT("LhsRhs"));
+		CHECK_EQUAL(TCHARsPlusCopyResult, TEXT("LhsRhs"));
 	}
 }
 
-TEST_CASE("Core::String::FString::ConstructWithSlack", "[Core][String][Smoke]")
+TEST_CASE_METHOD(FAutomationTestFixture, "Core::String::FString::ConstructWithSlack", "[Core][String][Smoke]")
 {
 	// Note that the total capacity of a string might be greater than the string length + slack + a null terminator due to
 	// underlying malloc implementations which is why we poll FMemory to see what size of allocation we should be expecting.
 
-	// Test creating from a valid string with various valid slack value
+	SECTION("Test creating from a valid string with various valid slack value")
 	{
 		const TCHAR* TestString = TEXT("FooBar");
 		const char* TestAsciiString = "FooBar";
@@ -364,19 +449,19 @@ TEST_CASE("Core::String::FString::ConstructWithSlack", "[Core][String][Smoke]")
 		const SIZE_T ExpectedCapacity = FMemory::QuantizeSize(NumElements * sizeof(TCHAR));
 
 		FString StringFromTChar(TestString, ExtraSlack);
-		TestEqual(TEXT("(TCHAR: Valid string with valid slack) resulting capacity"), StringFromTChar.GetAllocatedSize(), ExpectedCapacity);
+		CHECK_EQUAL(StringFromTChar.GetAllocatedSize(), ExpectedCapacity);
 
 		FString StringFromAscii(TestAsciiString, ExtraSlack);
-		TestEqual(TEXT("(ASCII: Valid string with valid slack) resulting capacity"), StringFromAscii.GetAllocatedSize(), ExpectedCapacity);
+		CHECK_EQUAL(StringFromAscii.GetAllocatedSize(), ExpectedCapacity);
 
 		FString StringFromFStringView(FStringView(TestString), ExtraSlack);
-		TestEqual(TEXT("(FStringView: Valid string with valid slack) resulting capacity"), StringFromFStringView.GetAllocatedSize(), ExpectedCapacity);
+		CHECK_EQUAL(StringFromFStringView.GetAllocatedSize(), ExpectedCapacity);
 
 		FString StringFromFString(FString(TestString), ExtraSlack);
-		TestEqual(TEXT("(FString: Valid string with valid slack"), StringFromFString.GetAllocatedSize(), ExpectedCapacity);
+		CHECK_EQUAL(StringFromFString.GetAllocatedSize(), ExpectedCapacity);
 	}
 
-	// Test creating from a valid string with a zero slack value
+	SECTION("Test creating from a valid string with a zero slack value")
 	{
 		const TCHAR* TestString = TEXT("FooBar");
 		const char* TestAsciiString = "FooBar";
@@ -387,19 +472,19 @@ TEST_CASE("Core::String::FString::ConstructWithSlack", "[Core][String][Smoke]")
 		const SIZE_T ExpectedCapacity = FMemory::QuantizeSize(NumElements * sizeof(TCHAR));
 
 		FString StringFromTChar(TestString, ExtraSlack);
-		TestEqual(TEXT("(TCHAR: Valid string with zero slack) resulting capacity"), StringFromTChar.GetAllocatedSize(), ExpectedCapacity);
+		CHECK_EQUAL(StringFromTChar.GetAllocatedSize(), ExpectedCapacity);
 
 		FString StringFromAscii(TestAsciiString, ExtraSlack);
-		TestEqual(TEXT("(ASCII: Valid string with zero slack) resulting capacity"), StringFromAscii.GetAllocatedSize(), ExpectedCapacity);
+		CHECK_EQUAL(StringFromAscii.GetAllocatedSize(), ExpectedCapacity);
 
 		FString StringFromFStringView(FStringView(TestString), ExtraSlack);
-		TestEqual(TEXT("(FStringView: Valid string with zero slack) resulting capacity"), StringFromFStringView.GetAllocatedSize(), ExpectedCapacity);
+		CHECK_EQUAL(StringFromFStringView.GetAllocatedSize(), ExpectedCapacity);
 
 		FString StringFromFString(FString(TestString), ExtraSlack);
-		TestEqual(TEXT("(FString: Valid string with zero slack"), StringFromFString.GetAllocatedSize(), ExpectedCapacity);
+		CHECK_EQUAL(StringFromFString.GetAllocatedSize(), ExpectedCapacity);
 	}
 
-	// Test creating from an empty string with a valid slack value
+	SECTION("Test creating from an empty string with a valid slack value")
 	{
 		const TCHAR* TestString = TEXT("");
 		const char* TestAsciiString = "";
@@ -410,19 +495,19 @@ TEST_CASE("Core::String::FString::ConstructWithSlack", "[Core][String][Smoke]")
 		const SIZE_T ExpectedCapacity = FMemory::QuantizeSize(NumElements * sizeof(TCHAR));
 
 		FString StringFromTChar(TestString, ExtraSlack);
-		TestEqual(TEXT("(TCHAR: Empty string with slack) resulting capacity"), StringFromTChar.GetAllocatedSize(), ExpectedCapacity);
+		CHECK_EQUAL(StringFromTChar.GetAllocatedSize(), ExpectedCapacity);
 
 		FString StringFromAscii(TestAsciiString, ExtraSlack);
-		TestEqual(TEXT("(ASCII: Empty string with slack) resulting capacity"), StringFromAscii.GetAllocatedSize(), ExpectedCapacity);
+		CHECK_EQUAL(StringFromAscii.GetAllocatedSize(), ExpectedCapacity);
 
 		FString StringFromFStringView(FStringView(TestString), ExtraSlack);
-		TestEqual(TEXT("(FStringView: Empty string with slack) resulting capacity"), StringFromFStringView.GetAllocatedSize(), ExpectedCapacity);
+		CHECK_EQUAL(StringFromFStringView.GetAllocatedSize(), ExpectedCapacity);
 
 		FString StringFromFString(FString(TestString), ExtraSlack);
-		TestEqual(TEXT("(FString: Empty string with slack) resulting capacity"), StringFromFString.GetAllocatedSize(), ExpectedCapacity);
+		CHECK_EQUAL(StringFromFString.GetAllocatedSize(), ExpectedCapacity);
 	}
 
-	// Test creating from an empty string with a zero slack value
+	SECTION("Test creating from an empty string with a zero slack value")
 	{
 		const TCHAR* TestString = TEXT("");
 		const char* TestAsciiString = "";
@@ -431,39 +516,39 @@ TEST_CASE("Core::String::FString::ConstructWithSlack", "[Core][String][Smoke]")
 		const SIZE_T ExpectedCapacity = 0u;
 
 		FString StringFromTChar(TestString, ExtraSlack);
-		TestEqual(TEXT("(TCHAR: Empty string with zero slack) resulting capacity"), StringFromTChar.GetAllocatedSize(), ExpectedCapacity);
+		CHECK_EQUAL(StringFromTChar.GetAllocatedSize(), ExpectedCapacity);
 
 		FString StringFromAscii(TestAsciiString, ExtraSlack);
-		TestEqual(TEXT("(ASCII: Empty string with zero slack) resulting capacity"), StringFromAscii.GetAllocatedSize(), ExpectedCapacity);
+		CHECK_EQUAL(StringFromAscii.GetAllocatedSize(), ExpectedCapacity);
 
 		FString StringFromFStringView(FStringView(TestString), ExtraSlack);
-		TestEqual(TEXT("(FStringView: Empty string with zero slack) resulting capacity"), StringFromFStringView.GetAllocatedSize(), ExpectedCapacity);
+		CHECK_EQUAL(StringFromFStringView.GetAllocatedSize(), ExpectedCapacity);
 
 		FString StringFromFString(FString(TestString), ExtraSlack);
-		TestEqual(TEXT("(FString: Empty string with zero slack) resulting capacity"), StringFromFString.GetAllocatedSize(), ExpectedCapacity);
+		CHECK_EQUAL(StringFromFString.GetAllocatedSize(), ExpectedCapacity);
 	}
 }
 
-TEST_CASE("Core::String::FString::Equality", "[Core][String][Smoke]")
+TEST_CASE_METHOD(FAutomationTestFixture, "Core::String::FString::Equality", "[Core][String][Smoke]")
 {
-	auto TestSelfEquality = [](const TCHAR* A)
+	auto TestSelfEquality = [this](const TCHAR* A)
 	{
-		TestTrue(TEXT("Self Equality C string"), FString(A) == A);
-		TestTrue(TEXT("Self Equality C string"), A == FString(A));
-		TestTrue(TEXT("Self Equality CaseSensitive"), FString(A).Equals(FString(A), ESearchCase::CaseSensitive));
-		TestTrue(TEXT("Self Equality IgnoreCase"), FString(A).Equals(FString(A), ESearchCase::IgnoreCase));
+		CHECK(FString(A) == A);
+		CHECK(A == FString(A));
+		CHECK(FString(A).Equals(FString(A), ESearchCase::CaseSensitive));
+		CHECK(FString(A).Equals(FString(A), ESearchCase::IgnoreCase));
 
 		FString Slacker(A);
 		Slacker.Reserve(100);
-		TestTrue(TEXT("Self Equality slack"), Slacker == FString(A));
+		CHECK(Slacker == FString(A));
 	};
 
-	auto TestPairEquality = [](const TCHAR* A, const TCHAR* B)
+	auto TestPairEquality = [this](const TCHAR* A, const TCHAR* B)
 	{
-		TestEqual(TEXT("Equals CaseSensitive"), FCString::Strcmp(A, B)  == 0, FString(A).Equals(FString(B), ESearchCase::CaseSensitive));
-		TestEqual(TEXT("Equals CaseSensitive"), FCString::Strcmp(B, A)  == 0, FString(B).Equals(FString(A), ESearchCase::CaseSensitive));
-		TestEqual(TEXT("Equals IgnoreCase"),	FCString::Stricmp(A, B) == 0, FString(A).Equals(FString(B), ESearchCase::IgnoreCase));
-		TestEqual(TEXT("Equals IgnoreCase"),	FCString::Stricmp(B, A) == 0, FString(B).Equals(FString(A), ESearchCase::IgnoreCase));
+		CHECK_EQUAL((FCString::Strcmp(A, B)  == 0), FString(A).Equals(FString(B), ESearchCase::CaseSensitive));
+		CHECK_EQUAL((FCString::Strcmp(B, A)  == 0), FString(B).Equals(FString(A), ESearchCase::CaseSensitive));
+		CHECK_EQUAL((FCString::Stricmp(A, B) == 0), FString(A).Equals(FString(B), ESearchCase::IgnoreCase));
+		CHECK_EQUAL((FCString::Stricmp(B, A) == 0), FString(B).Equals(FString(A), ESearchCase::IgnoreCase));
 	};
 
 	const TCHAR* Pairs[][2] =	{ {TEXT(""),	TEXT(" ")}
@@ -480,7 +565,7 @@ TEST_CASE("Core::String::FString::Equality", "[Core][String][Smoke]")
 	}
 }
 
-TEST_CASE("Core::String::FString::Path Concat Compound Operator", "[Core][String][Smoke]")
+TEST_CASE_METHOD(FAutomationTestFixture, "Core::String::FString::Path Concat Compound Operator", "[Core][String][Smoke]")
 {
 	// No need to test a nullptr TCHAR* as an input parameter as this is expected to cause a crash
 	// No need to test self assignment, clang will catch that was a compiler error (-Wself-assign-overloaded)
@@ -501,57 +586,69 @@ TEST_CASE("Core::String::FString::Path Concat Compound Operator", "[Core][String
 	TStringBuilder<128> FilenameStringBuilder; FilenameStringBuilder << Filename;
 	TStringBuilder<128> FilenameWithLeadingSlashStringBuilder; FilenameWithLeadingSlashStringBuilder << FilenameWithLeadingSlash;
 	
-#define TEST_EMPTYPATH_EMPTYFILE(Type, Input)						{ FString EmptyPathString; EmptyPathString /= Input; TestTrue(TEXT(Type ": EmptyPath/EmptyFilename to be empty"), EmptyPathString.IsEmpty()); }
-#define TEST_VALIDPATH_EMPTYFILE(Type, Input)						{ FString Result(Path); Result /= Input; TestEqual(TEXT(Type ": ValidPath/EmptyFilename result to be"), Result, PathWithTrailingSlash); } \
-																	{ FString Result(PathWithTrailingSlash); Result /= Input; TestEqual(TEXT(Type " (with extra /): ValidPath/EmptyFilename result to be"), Result, PathWithTrailingSlash); }	
-#define TEST_EMPTYPATH_VALIDFILE(Type, Input)						{ FString Result; Result /= Input; TestEqual(TEXT(Type ": EmptyPath/ValidFilename"), Result, Filename); }
-#define TEST_VALIDPATH_VALIDFILE(Type, Path, File)					{ FString Result(Path); Result /= File; TestEqual(TEXT(Type ": ValidPath/ValidFilename"), Result, CombinedPath); }
-#define TEST_VALIDPATH_VALIDFILE_DOUBLE_SEPARATOR(Type, Path, File)	{ FString Result(Path); Result /= File; TestEqual(TEXT(Type ": ValidPath//ValidFilename"), Result, CombinedPathWithDoubleSeparator); }
-
-	// Test empty path /= empty file
-	TEST_EMPTYPATH_EMPTYFILE("NullString", FString());
-	TEST_EMPTYPATH_EMPTYFILE("EmptyString", FString(TEXT("")));
-	TEST_EMPTYPATH_EMPTYFILE("EmptyAnsiLiteralString", "");
-	TEST_EMPTYPATH_EMPTYFILE("EmptyLiteralString", TEXT(""));
-	TEST_EMPTYPATH_EMPTYFILE("NullStringView", FStringView());
-	TEST_EMPTYPATH_EMPTYFILE("EmptyStringView", FStringView(TEXT("")));
-	TEST_EMPTYPATH_EMPTYFILE("EmptyStringBuilder", EmptyStringBuilder);
-
-	// Test valid path /= empty file
-	TEST_VALIDPATH_EMPTYFILE("NullString", FString());
-	TEST_VALIDPATH_EMPTYFILE("EmptyString", FString(TEXT("")));
-	TEST_VALIDPATH_EMPTYFILE("EmptyAnsiLiteralString", "");
-	TEST_VALIDPATH_EMPTYFILE("EmptyLiteralString", TEXT(""));
-	TEST_VALIDPATH_EMPTYFILE("NullStringView", FStringView());
-	TEST_VALIDPATH_EMPTYFILE("EmptyStringView", FStringView(TEXT("")));
-	TEST_VALIDPATH_EMPTYFILE("EmptyStringBuilder", EmptyStringBuilder);
+#define TEST_EMPTYPATH_EMPTYFILE(Type, Input)						{ FString EmptyPathString; EmptyPathString /= Input; TEST_TRUE(TEXT(Type ": EmptyPath/EmptyFilename to be empty"), EmptyPathString.IsEmpty()); }
+#define TEST_VALIDPATH_EMPTYFILE(Type, Input)						{ FString Result(Path); Result /= Input; TEST_EQUAL(TEXT(Type ": ValidPath/EmptyFilename result to be"), Result, PathWithTrailingSlash); } \
+																	{ FString Result(PathWithTrailingSlash); Result /= Input; TEST_EQUAL(TEXT(Type " (with extra /): ValidPath/EmptyFilename result to be"), Result, PathWithTrailingSlash); }	
+#define TEST_EMPTYPATH_VALIDFILE(Type, Input)						{ FString Result; Result /= Input; TEST_EQUAL(TEXT(Type ": EmptyPath/ValidFilename"), Result, Filename); }
+#define TEST_VALIDPATH_VALIDFILE(Type, Path, File)					{ FString Result(Path); Result /= File; TEST_EQUAL(TEXT(Type ": ValidPath/ValidFilename"), Result, CombinedPath); }
+#define TEST_VALIDPATH_VALIDFILE_DOUBLE_SEPARATOR(Type, Path, File)	{ FString Result(Path); Result /= File; TEST_EQUAL(TEXT(Type ": ValidPath//ValidFilename"), Result, CombinedPathWithDoubleSeparator); }
+
+	SECTION("Test empty path /= empty file")
+	{
+		TEST_EMPTYPATH_EMPTYFILE("NullString", FString());
+		TEST_EMPTYPATH_EMPTYFILE("EmptyString", FString(TEXT("")));
+		TEST_EMPTYPATH_EMPTYFILE("EmptyAnsiLiteralString", "");
+		TEST_EMPTYPATH_EMPTYFILE("EmptyLiteralString", TEXT(""));
+		TEST_EMPTYPATH_EMPTYFILE("NullStringView", FStringView());
+		TEST_EMPTYPATH_EMPTYFILE("EmptyStringView", FStringView(TEXT("")));
+		TEST_EMPTYPATH_EMPTYFILE("EmptyStringBuilder", EmptyStringBuilder);
+	}
 	
-	// Test empty path /= valid file
-	TEST_EMPTYPATH_VALIDFILE("String", FString(Filename));
-	TEST_EMPTYPATH_VALIDFILE("LiteralString", Filename);
-	TEST_EMPTYPATH_VALIDFILE("LiteralAnsiString", AnsiFilename);
-	TEST_EMPTYPATH_VALIDFILE("StringView", FStringView(Filename));
+	SECTION("Test valid path /= empty file")
+	{
+		TEST_VALIDPATH_EMPTYFILE("NullString", FString());
+		TEST_VALIDPATH_EMPTYFILE("EmptyString", FString(TEXT("")));
+		TEST_VALIDPATH_EMPTYFILE("EmptyAnsiLiteralString", "");
+		TEST_VALIDPATH_EMPTYFILE("EmptyLiteralString", TEXT(""));
+		TEST_VALIDPATH_EMPTYFILE("NullStringView", FStringView());
+		TEST_VALIDPATH_EMPTYFILE("EmptyStringView", FStringView(TEXT("")));
+		TEST_VALIDPATH_EMPTYFILE("EmptyStringBuilder", EmptyStringBuilder);
+	}
 	
-	// Test valid path /= valid file
-	TEST_VALIDPATH_VALIDFILE("String", Path, FString(Filename));
-	TEST_VALIDPATH_VALIDFILE("LiteralString", Path, Filename);
-	TEST_VALIDPATH_VALIDFILE("LiteralAnsiString", Path, AnsiFilename);
-	TEST_VALIDPATH_VALIDFILE("StringView", Path, FStringView(Filename));
-	TEST_VALIDPATH_VALIDFILE("StringBuilder", Path, FilenameStringBuilder);
-
-	// Test valid path (ending in /) /= valid file
-	TEST_VALIDPATH_VALIDFILE("String (path with extra /)", PathWithTrailingSlash, FString(Filename));
-	TEST_VALIDPATH_VALIDFILE("LiteralString (path with extra /)", PathWithTrailingSlash, Filename);
-	TEST_VALIDPATH_VALIDFILE("LiteralAnsiString (path with extra /)", PathWithTrailingSlash, AnsiFilename);
-	TEST_VALIDPATH_VALIDFILE("StringView (path with extra /)", PathWithTrailingSlash, FStringView(Filename));
-	TEST_VALIDPATH_VALIDFILE("StringBuilder (path with extra /)", PathWithTrailingSlash, FilenameStringBuilder);
+	SECTION("Test empty path /= valid file")
+	{
+		TEST_EMPTYPATH_VALIDFILE("String", FString(Filename));
+		TEST_EMPTYPATH_VALIDFILE("LiteralString", Filename);
+		TEST_EMPTYPATH_VALIDFILE("LiteralAnsiString", AnsiFilename);
+		TEST_EMPTYPATH_VALIDFILE("StringView", FStringView(Filename));
+	}
 	
-	// Test valid path / valid path + file (starting with /)
-	TEST_VALIDPATH_VALIDFILE("String (filename with extra /)", Path, FString(FilenameWithLeadingSlash));
-	TEST_VALIDPATH_VALIDFILE("LiteralString (filename with extra /)", Path, FilenameWithLeadingSlash);
-	TEST_VALIDPATH_VALIDFILE("LiteralAnsiString (filename with extra /)", Path, AnsiFilenameWithLeadingSlash);
-	TEST_VALIDPATH_VALIDFILE("StringView (filename with extra /)", Path, FStringView(FilenameWithLeadingSlash));
-	TEST_VALIDPATH_VALIDFILE("StringBuilder (filename with extra /)", Path, FilenameWithLeadingSlashStringBuilder);
+	SECTION("// Test valid path /= valid file")
+	{
+		TEST_VALIDPATH_VALIDFILE("String", Path, FString(Filename));
+		TEST_VALIDPATH_VALIDFILE("LiteralString", Path, Filename);
+		TEST_VALIDPATH_VALIDFILE("LiteralAnsiString", Path, AnsiFilename);
+		TEST_VALIDPATH_VALIDFILE("StringView", Path, FStringView(Filename));
+		TEST_VALIDPATH_VALIDFILE("StringBuilder", Path, FilenameStringBuilder);
+	}
+	
+	SECTION("Test valid path(ending in / ) /= valid file")
+	{
+		TEST_VALIDPATH_VALIDFILE("String (path with extra /)", PathWithTrailingSlash, FString(Filename));
+		TEST_VALIDPATH_VALIDFILE("LiteralString (path with extra /)", PathWithTrailingSlash, Filename);
+		TEST_VALIDPATH_VALIDFILE("LiteralAnsiString (path with extra /)", PathWithTrailingSlash, AnsiFilename);
+		TEST_VALIDPATH_VALIDFILE("StringView (path with extra /)", PathWithTrailingSlash, FStringView(Filename));
+		TEST_VALIDPATH_VALIDFILE("StringBuilder (path with extra /)", PathWithTrailingSlash, FilenameStringBuilder);
+	}
+	
+	SECTION("Test valid path / valid path + file(starting with / )")
+	{
+		TEST_VALIDPATH_VALIDFILE("String (filename with extra /)", Path, FString(FilenameWithLeadingSlash));
+		TEST_VALIDPATH_VALIDFILE("LiteralString (filename with extra /)", Path, FilenameWithLeadingSlash);
+		TEST_VALIDPATH_VALIDFILE("LiteralAnsiString (filename with extra /)", Path, AnsiFilenameWithLeadingSlash);
+		TEST_VALIDPATH_VALIDFILE("StringView (filename with extra /)", Path, FStringView(FilenameWithLeadingSlash));
+		TEST_VALIDPATH_VALIDFILE("StringBuilder (filename with extra /)", Path, FilenameWithLeadingSlashStringBuilder);
+	}
 	
 	// Appending a file name that starts with a / to a directory that ends with a / will not remove the erroneous / and so 
 	// will end up with // in the path, these tests are to show this behavior
diff --git a/Engine/Source/Runtime/Core/Tests/Tasks/TasksTest.cpp b/Engine/Source/Runtime/Core/Tests/Tasks/TasksTest.cpp
--- a/Engine/Source/Runtime/Core/Tests/Tasks/TasksTest.cpp
+++ b/Engine/Source/Runtime/Core/Tests/Tasks/TasksTest.cpp
@@ -1,13 +1,11 @@
 // Copyright Epic Games, Inc. All Rights Reserved.
 
 #include "CoreTypes.h"
-#include "Tests/Benchmark.h"
+#include "Core/Tests/Benchmark.h"
 #include "Tasks/Pipe.h"
 #include "HAL/Thread.h"
 #include "Async/ParallelFor.h"
-#include "TestHarness.h"
-
-#include "TestCommon/CoreUtilities.h"
+#include "TestFixtures/CoreTestFixture.h"
 
 #include <atomic>
 
@@ -70,10 +68,8 @@ namespace UE { namespace TasksTests
 		CHECK(TasksExecutedNum == TasksNum);
 	}
 
-	TEST_CASE("Core::Async::Task::Basic", "[Core][Async][Smoke]")
+	TEST_CASE_METHOD(FCoreTestFixture, "Core::Async::Task::Basic", "[Core][Async][Smoke]")
 	{
-		InitTaskGraphAndDependencies(true);
-
 		{	// basic example, fire and forget a high-pri task
 			Launch(
 				UE_SOURCE_LOCATION, // debug name
@@ -195,8 +191,8 @@ namespace UE { namespace TasksTests
 #if 0	// unreliable test, destruction can happen on a worker thread, after the task is flagged as completed and so the check can be hit before the destruction
 			uint32 LocalConstructionsNum = ConstructionsNum.load(std::memory_order_relaxed);
 			uint32 LocalDestructionsNum = DestructionsNum.load(std::memory_order_relaxed);
-			checkf(LocalConstructionsNum == 1, TEXT("%d result instances were created but one was expected: the value stored in the task"), LocalConstructionsNum);
-			checkf(LocalConstructionsNum == LocalDestructionsNum, TEXT("Mismatched number of constructions (%d) and destructions (%d)"), LocalConstructionsNum, LocalDestructionsNum);
+			CHECK(LocalConstructionsNum == 1);
+			CHECK(LocalConstructionsNum == LocalDestructionsNum);
 
 			ConstructionsNum = 0;
 			DestructionsNum = 0;
@@ -209,8 +205,8 @@ namespace UE { namespace TasksTests
 #if 0	// unreliable test, destruction can happen on a worker thread, after the task is flagged as completed and so the check can be hit before the destruction
 			LocalConstructionsNum = ConstructionsNum.load(std::memory_order_relaxed);
 			LocalDestructionsNum = DestructionsNum.load(std::memory_order_relaxed);
-			checkf(LocalConstructionsNum == 2, TEXT("%d result instances were created but 2 was expected: the value stored in the task"), LocalConstructionsNum);
-			checkf(LocalConstructionsNum == LocalDestructionsNum, TEXT("Mismatched number of constructions (%d) and destructions (%d)"), LocalConstructionsNum, LocalDestructionsNum);
+			CHECK(LocalConstructionsNum == 2);
+			CHECK(LocalConstructionsNum == LocalDestructionsNum);
 
 			ConstructionsNum = 0;
 			DestructionsNum = 0;
@@ -378,17 +374,13 @@ namespace UE { namespace TasksTests
 		//}
 
 		UE_BENCHMARK(5, BasicStressTest<100, 100>);
-
-		CleanupTaskGraphAndDependencies();
 	}
 
 	template<uint32 SpawnerGroupsNum, uint32 SpawnersPerGroupNum>
 	void PipeStressTest();
 
-	TEST_CASE("Core::Async::Task::Pipe", "[Core][Async][Smoke]")
+	TEST_CASE_METHOD(FCoreTestFixture, "Core::Async::Task::Pipe", "[Core][Async][Smoke]")
 	{
-		InitTaskGraphAndDependencies(true);
-
 		{	// a basic usage example
 			Tasks::FPipe Pipe{ UE_SOURCE_LOCATION }; // a debug name, user-provided or 
 				// `UE_SOURCE_LOCATION` - source file name and line number
@@ -493,9 +485,8 @@ namespace UE { namespace TasksTests
 			Task.Wait();
 		}
 
+		SUCCEED();
 		UE_BENCHMARK(5, PipeStressTest<200, 100>);
-		
-		CleanupTaskGraphAndDependencies();
 	}
 
 	// stress test for named thread tasks, CHECKs spawning a large number of tasks from multiple threads and executing them
@@ -603,14 +594,10 @@ namespace UE { namespace TasksTests
 		TlsValue = Dummy;
 	}
 
-	TEST_CASE("Core::Async::Tls::Tls", "[Core][Async][Smoke]")
+	TEST_CASE_METHOD(FCoreTestFixture, "Core::Tasks::Tls::Tls", "[Core][Async][Smoke]")
 	{
-		InitTaskGraphAndDependencies(true);
-
 		UE_BENCHMARK(5, UeTlsStressTest<10000000>);
 		UE_BENCHMARK(5, ThreadLocalStressTest<10000000>);
-
-		CleanupTaskGraphAndDependencies();
 	}
 
 	template<uint64 NumBranches, uint64 NumLoops, uint64 NumTasks>
@@ -652,11 +639,11 @@ namespace UE { namespace TasksTests
 			BranchTasks.Add(Task.GetResult());
 		}
 		Wait(Branches);
+		SUCCEED();
 	}
 
-	TEST_CASE("Core::Async::Task::Dependencies", "[Core][Async][Smoke]")
+	TEST_CASE_METHOD(FCoreTestFixture, "Core::Async::Task::Dependencies", "[Core][Async][Smoke]")
 	{
-		InitTaskGraphAndDependencies(true);
 
 		{	// a task is not executed until its prerequisite (FTaskEvent) is completed
 			FTaskEvent Prereq{ UE_SOURCE_LOCATION };
@@ -792,9 +779,9 @@ namespace UE { namespace TasksTests
 			Task.Wait();
 		}
 
-		UE_BENCHMARK(5, DependenciesPerfTest<50, 5, 50>);
+		SUCCEED();
 
-		CleanupTaskGraphAndDependencies();
+		UE_BENCHMARK(5, DependenciesPerfTest<50, 5, 50>);
 	}
 
 	// blocks all workers (except reserve workers) until given event is triggered. Returns blocking tasks.
@@ -858,10 +845,8 @@ namespace UE { namespace TasksTests
 			N11.IsCompleted() && N12.IsCompleted() && N21.IsCompleted() && N22.IsCompleted()));
 	}
 
-	TEST_CASE("Core::Async::Task::Deep Retraction", "[Core][Async][Smoke]")
+	TEST_CASE_METHOD(FCoreTestFixture, "Core::Async::Task::Deep Retraction", "[Core][Async][Smoke]")
 	{
-		InitTaskGraphAndDependencies(true);
-
 		FTaskEvent ResumeEvent{ UE_SOURCE_LOCATION };
 		TArray<FTask> WorkerBlockers = BlockWorkers(ResumeEvent);
 
@@ -927,8 +912,6 @@ namespace UE { namespace TasksTests
 
 		ResumeEvent.Trigger();
 		CHECK(Wait(WorkerBlockers, FTimespan::FromSeconds(1)));
-
-		CleanupTaskGraphAndDependencies();
 	}
 
 	template<uint32 Num>
@@ -940,13 +923,9 @@ namespace UE { namespace TasksTests
 		}
 	}
 
-	TEST_CASE("Core::Async::Task::Deep Retraction Strees", "[Core][Async][Stress]")
+	TEST_CASE_METHOD(FCoreTestFixture, "Core::Async::Tasks::Deep Retraction Stress", "[Core][Tasks][.Stress]")
 	{
-		InitTaskGraphAndDependencies(true);
-
 		UE_BENCHMARK(5, DeepRetractionStressTest<1000>);
-
-		CleanupTaskGraphAndDependencies();
 	}
 
 	template<uint64 Num>
@@ -966,13 +945,9 @@ namespace UE { namespace TasksTests
 		}
 	}
 
-	TEST_CASE("Core:::Async::Task::Nested Tasks Strees", "[Core][Async][Stress]")
+	TEST_CASE_METHOD(FCoreTestFixture, "Core:::Async::Task::Nested Tasks Strees", "[Core][Async][.Stress]")
 	{
-		InitTaskGraphAndDependencies(true);
-
 		UE_BENCHMARK(5, NestedTasksStressTest<10000>);
-
-		CleanupTaskGraphAndDependencies();
 	}
 
 	template<int NumTasks>
@@ -1157,10 +1132,8 @@ namespace UE { namespace TasksTests
 		}
 	}
 
-	TEST_CASE("Core::Async::Task::PerfTest", "[Core][Async][Perf]")
+	TEST_CASE_METHOD(FCoreTestFixture, "Core::Async::Tasks::PerfTest", "[Core][Async][.Perf]")
 	{
-		InitTaskGraphAndDependencies(true);
-
 		TRACE_CPUPROFILER_EVENT_SCOPE(TaskGraphTests_PerfTest);
 
 		UE_BENCHMARK(5, TestPerfBasic<10000>);
@@ -1170,8 +1143,6 @@ namespace UE { namespace TasksTests
 		UE_BENCHMARK(5, TestWorkStealing<100, 100>);
 		UE_BENCHMARK(5, TestSpawning<10000>);
 		UE_BENCHMARK(5, TestBatchSpawning<10000>);
-
-		CleanupTaskGraphAndDependencies();
 	}
 }}
 
diff --git a/Engine/Source/Runtime/Core/Tests/Templates/ValueOrErrorTest.cpp b/Engine/Source/Runtime/Core/Tests/Templates/ValueOrErrorTest.cpp
--- a/Engine/Source/Runtime/Core/Tests/Templates/ValueOrErrorTest.cpp
+++ b/Engine/Source/Runtime/Core/Tests/Templates/ValueOrErrorTest.cpp
@@ -1,24 +1,21 @@
 // Copyright Epic Games, Inc. All Rights Reserved.
 
 #include "Templates/ValueOrError.h"
-
 #include "TestHarness.h"
 
+
 #include <type_traits>
 
-TEST_CASE("Core::Templates::ValueOrError", "[Core][Templates][Smoke]")
-{
-	SECTION("Static")
-	{
-		STATIC_REQUIRE(!std::is_constructible<TValueOrError<int, int>>::value);
+static_assert(!std::is_constructible<TValueOrError<int, int>>::value, "Expected no default constructor.");
 
-		STATIC_REQUIRE(std::is_copy_constructible<TValueOrError<int, int>>::value);
-		STATIC_REQUIRE(std::is_move_constructible<TValueOrError<int, int>>::value);
+static_assert(std::is_copy_constructible<TValueOrError<int, int>>::value, "Missing copy constructor.");
+static_assert(std::is_move_constructible<TValueOrError<int, int>>::value, "Missing move constructor.");
 
-		STATIC_REQUIRE(std::is_copy_assignable<TValueOrError<int, int>>::value);
-		STATIC_REQUIRE(std::is_move_assignable<TValueOrError<int, int>>::value);
-	}
+static_assert(std::is_copy_assignable<TValueOrError<int, int>>::value, "Missing copy assignment.");
+static_assert(std::is_move_assignable<TValueOrError<int, int>>::value, "Missing move assignment.");
 
+TEST_CASE_METHOD(FAutomationTestFixture, "Core::Templates::TValueOrError::Smoke Test", "[Core][Templates][Smoke]")
+{
 	static int ValueCount = 0;
 	static int ErrorCount = 0;
 
@@ -56,176 +53,139 @@ TEST_CASE("Core::Templates::ValueOrError", "[Core][Templates][Smoke]")
 
 	using FTestType = TValueOrError<FTestValue, FTestError>;
 
-	SECTION("MakeValue Move")
+	// Test MakeValue Move
 	{
+		INFO("Test MakeValue Move");
 		FTestType ValueOrError = MakeValue(FTestValue());
-		CHECK(ValueCount == 1);
-		CHECK(ValueOrError.TryGetValue() == &ValueOrError.GetValue());
-		CHECK(ValueOrError.GetValue().Value == 1);
+		CHECK_EQUAL(ValueCount, 1);
+		CHECK_EQUAL(ValueOrError.TryGetValue(), &ValueOrError.GetValue());
+		CHECK_EQUAL(ValueOrError.GetValue().Value, 1);
+		CHECK_FALSE(ValueOrError.HasError());
 		CHECK(ValueOrError.HasValue());
 		CHECK(ValueOrError.TryGetError() == nullptr);
-		CHECK_FALSE(ValueOrError.HasError());
 	}
 
-	SECTION("MakeValue Proxy")
+	CHECK_EQUAL(ValueCount, 0);
+
+	// Test MakeValue Proxy
 	{
+		INFO("Test MakeValue Proxy");
 		FTestType ValueOrError = MakeValue(2, 6, 8);
-		CHECK(ValueCount == 1);
-		CHECK(ValueOrError.TryGetValue() == &ValueOrError.GetValue());
-		CHECK(ValueOrError.GetValue().Value == 16);
+		CHECK_EQUAL(ValueCount, 1);
+		CHECK_EQUAL(ValueOrError.TryGetValue(), &ValueOrError.GetValue());
+		CHECK_EQUAL(ValueOrError.GetValue().Value, 16);
+		CHECK_FALSE(ValueOrError.HasError());
 		CHECK(ValueOrError.HasValue());
 		CHECK(ValueOrError.TryGetError() == nullptr);
-		CHECK_FALSE(ValueOrError.HasError());
 	}
+	
+	CHECK_EQUAL(ValueCount, 0);
 
-	SECTION("StealValue")
+	// Test StealValue
 	{
+		INFO("Test StealValue");
 		FTestType ValueOrError = MakeValue(FTestValue());
 		FTestValue Value = ValueOrError.StealValue();
-		CHECK(ValueCount == 1);
-		CHECK(Value.Value == 1);
+		CHECK_EQUAL(ValueCount, 1);
+		CHECK_EQUAL(Value.Value, 1);
 		CHECK_FALSE(ValueOrError.HasError());
 		CHECK_FALSE(ValueOrError.HasValue());
 	}
+	
+	CHECK_EQUAL(ValueCount, 0);
 
-	SECTION("MakeError Move")
+	// Test MakeError Move
 	{
+		INFO("Test MakeError Move");
 		FTestType ValueOrError = MakeError(FTestError());
-		CHECK(ErrorCount == 1);
-		CHECK(ValueOrError.TryGetError() == &ValueOrError.GetError());
-		CHECK(ValueOrError.GetError().Error == 1);
+		CHECK_EQUAL(ErrorCount, 1);
+		CHECK_EQUAL(ValueOrError.TryGetError(), &ValueOrError.GetError());
+		CHECK_EQUAL(ValueOrError.GetError().Error, 1);
+		CHECK_FALSE(ValueOrError.HasValue());
 		CHECK(ValueOrError.HasError());
 		CHECK(ValueOrError.TryGetValue() == nullptr);
-		CHECK_FALSE(ValueOrError.HasValue());
 	}
+	
+	CHECK_EQUAL(ErrorCount, 0);
 
-	SECTION("MakeError Proxy")
+	// Test MakeError Proxy
 	{
+		INFO("Test MakeError Proxy");
 		FTestType ValueOrError = MakeError(4, 12);
-		CHECK(ErrorCount == 1);
-		CHECK(ValueOrError.TryGetError() == &ValueOrError.GetError());
-		CHECK(ValueOrError.GetError().Error == 16);
+		CHECK_EQUAL(ErrorCount, 1);
+		CHECK_EQUAL(ValueOrError.TryGetError(), &ValueOrError.GetError());
+		CHECK_EQUAL(ValueOrError.GetError().Error, 16);
+		CHECK_FALSE(ValueOrError.HasValue());
 		CHECK(ValueOrError.HasError());
 		CHECK(ValueOrError.TryGetValue() == nullptr);
-		CHECK_FALSE(ValueOrError.HasValue());
 	}
+	
+	CHECK_EQUAL(ErrorCount, 0);
 
-	SECTION("StealError")
+	// Test StealError
 	{
+		INFO("Test StealError");
 		FTestType ValueOrError = MakeError();
 		FTestError Error = ValueOrError.StealError();
-		CHECK(ErrorCount == 1);
-		CHECK(Error.Error == 1);
+		CHECK_EQUAL(ErrorCount, 1);
+		CHECK_EQUAL(Error.Error, 1);
 		CHECK_FALSE(ValueOrError.HasValue());
 		CHECK_FALSE(ValueOrError.HasError());
 	}
 
-	SECTION("Assignment")
+	CHECK_EQUAL(ErrorCount, 0);
+
+	// Test Assignment
 	{
+		INFO("Test Assignment");
 		FTestType ValueOrError = MakeValue();
 		ValueOrError = MakeValue();
-		CHECK(ValueCount == 1);
-		CHECK(ValueOrError.GetValue().Value == 2);
-		CHECK(ErrorCount == 0);
+		CHECK_EQUAL(ValueCount, 1);
+		CHECK_EQUAL(ValueOrError.GetValue().Value, 2);
+		CHECK_EQUAL(ErrorCount, 0);
 		ValueOrError = MakeError();
-		CHECK(ValueCount == 0);
-		CHECK(ErrorCount == 1);
+		CHECK_EQUAL(ValueCount, 0);
+		CHECK_EQUAL(ErrorCount, 1);
 		ValueOrError = MakeError();
-		CHECK(ValueCount == 0);
-		CHECK(ValueOrError.GetError().Error == 2);
-		CHECK(ErrorCount == 1);
+		CHECK_EQUAL(ValueCount, 0);
+		CHECK_EQUAL(ValueOrError.GetError().Error, 2);
+		CHECK_EQUAL(ErrorCount, 1);
 		ValueOrError = MakeValue();
-		CHECK(ValueCount == 1);
-		CHECK(ErrorCount == 0);
+		CHECK_EQUAL(ValueCount, 1);
+		CHECK_EQUAL(ErrorCount, 0);
 		FTestType UnsetValueOrError = MakeValue();
 		UnsetValueOrError.StealValue();
 		ValueOrError = MoveTemp(UnsetValueOrError);
-		CHECK(ValueCount == 0);
-		CHECK(ErrorCount == 0);
+		CHECK_EQUAL(ValueCount, 0);
+		CHECK_EQUAL(ErrorCount, 0);
 		CHECK_FALSE(ValueOrError.HasValue());
 		CHECK_FALSE(ValueOrError.HasError());
 	}
 
-	SECTION("Move-Only Value/Error")
+	CHECK_EQUAL(ValueCount, 0);
+	CHECK_EQUAL(ErrorCount, 0);
+
+	// Test Move-Only Value/Error
 	{
-		TValueOrError<FTestMoveOnly, FTestMoveOnly> Value = MakeValue(FTestMoveOnly{1});
-		TValueOrError<FTestMoveOnly, FTestMoveOnly> Error = MakeError(FTestMoveOnly{1});
+		INFO("Move-Only Value/Error");
+		TValueOrError<FTestMoveOnly, FTestMoveOnly> Value = MakeValue(FTestMoveOnly());
+		TValueOrError<FTestMoveOnly, FTestMoveOnly> Error = MakeError(FTestMoveOnly());
 		FTestMoveOnly MovedValue = MoveTemp(Value).GetValue();
 		FTestMoveOnly MovedError = MoveTemp(Error).GetError();
-		CHECK(MovedValue.Value == 1);
-		CHECK(MovedError.Value == 1);
 	}
 
-	SECTION("Integer Value/Error")
+	// Test Integer Value/Error
 	{
+		INFO("Test Integer Value/Error");
 		TValueOrError<int32, int32> ValueOrError = MakeValue();
-		CHECK(ValueOrError.GetValue() == 0);
+		CHECK_EQUAL(ValueOrError.GetValue(), 0);
 		ValueOrError = MakeValue(1);
-		CHECK(ValueOrError.GetValue() == 1);
+		CHECK_EQUAL(ValueOrError.GetValue(), 1);
 		ValueOrError = MakeError();
-		CHECK(ValueOrError.GetError() == 0);
+		CHECK_EQUAL(ValueOrError.GetError(), 0);
 		ValueOrError = MakeError(1);
-		CHECK(ValueOrError.GetError() == 1);
-	}
-
-	SECTION("Value/Void Value")
-	{
-		TValueOrError<FTestValue, void> ValueOrError = MakeValue(1, 2, 3);
-		CHECK(ValueCount == 1);
-		CHECK(ValueOrError.HasValue());
-		CHECK_FALSE(ValueOrError.HasError());
-		CHECK(ValueOrError.GetValue().Value == 6);
-		CHECK(ValueOrError.TryGetValue() == &ValueOrError.GetValue());
-		FTestValue Value = ValueOrError.StealValue();
-		CHECK_FALSE(ValueOrError.HasValue());
-		CHECK(Value.Value == 6);
-	}
-
-	SECTION("Value/Void Error")
-	{
-		TValueOrError<FTestValue, void> ValueOrError = MakeError();
-		CHECK(ValueCount == 0);
-		CHECK(ValueOrError.HasError());
-		CHECK_FALSE(ValueOrError.HasValue());
-		CHECK(ValueOrError.TryGetValue() == nullptr);
-	}
-
-	SECTION("Void/Error Value")
-	{
-		TValueOrError<void, FTestError> ValueOrError = MakeValue();
-		CHECK(ErrorCount == 0);
-		CHECK(ValueOrError.HasValue());
-		CHECK_FALSE(ValueOrError.HasError());
-		CHECK(ValueOrError.TryGetError() == nullptr);
+		CHECK_EQUAL(ValueOrError.GetError(), 1);
 	}
 
-	SECTION("Void/Error Error")
-	{
-		TValueOrError<void, FTestError> ValueOrError = MakeError(1, 2);
-		CHECK(ErrorCount == 1);
-		CHECK(ValueOrError.HasError());
-		CHECK_FALSE(ValueOrError.HasValue());
-		CHECK(ValueOrError.GetError().Error == 3);
-		CHECK(ValueOrError.TryGetError() == &ValueOrError.GetError());
-		FTestError Error = ValueOrError.StealError();
-		CHECK_FALSE(ValueOrError.HasError());
-		CHECK(Error.Error == 3);
-	}
-
-	SECTION("Void/Void Value")
-	{
-		TValueOrError<void, void> ValueOrError = MakeValue();
-		CHECK(ValueOrError.HasValue());
-		CHECK_FALSE(ValueOrError.HasError());
-	}
-
-	SECTION("Void/Void Error")
-	{
-		TValueOrError<void, void> ValueOrError = MakeError();
-		CHECK(ValueOrError.HasError());
-		CHECK_FALSE(ValueOrError.HasValue());
-	}
-
-	CHECK(ValueCount == 0);
-	CHECK(ErrorCount == 0);
 }
+
diff --git a/Engine/Source/Runtime/CoreUObject/Tests/ObjectHandleTest.cpp b/Engine/Source/Runtime/CoreUObject/Tests/ObjectHandleTest.cpp
--- a/Engine/Source/Runtime/CoreUObject/Tests/ObjectHandleTest.cpp
+++ b/Engine/Source/Runtime/CoreUObject/Tests/ObjectHandleTest.cpp
@@ -8,16 +8,10 @@
 
 static_assert(sizeof(FObjectHandle) == sizeof(void*), "FObjectHandle type must always compile to something equivalent to a pointer size.");
 
-#if WITH_DEV_AUTOMATION_TESTS
-
 class FObjectHandleTestBase : public FObjectRefTrackingTestBase
 {
 public:
-	FObjectHandleTestBase(const FString& InName, const bool bInComplexTask)
-	: FObjectRefTrackingTestBase(InName, bInComplexTask)
-	{
-	}
-
+	
 protected:
 	
 	UObject* ResolveHandle(FObjectHandle& TargetHandle)
@@ -34,7 +28,7 @@ protected:
 		}
 	#else
 		// Immediately resolved handles may be null (if the target is invalid) and must be resolved at this point
-		if (!TestTrue(TEXT("Handle to target is not resolved"), IsObjectHandleResolved(TargetHandle)))
+		if (!TestFalse(TEXT("Handle to target is not resolved"), IsObjectHandleResolved(TargetHandle)))
 		{
 			return nullptr;
 		}
@@ -75,7 +69,7 @@ protected:
 
 		if (!ResolvedObject)
 		{
-			AddError(FString::Printf(TEXT("Expected '%s.%s' to resolve to non null."), ANSI_TO_TCHAR(PackageName), ANSI_TO_TCHAR(ObjectName)), 1);
+			FAIL_CHECK(FString::Printf(TEXT("Expected '%s.%s' to resolve to non null."), ANSI_TO_TCHAR(PackageName), ANSI_TO_TCHAR(ObjectName)));
 			return false;
 		}
 		ObjectRefMetrics.TestNumFailedResolves(TEXT("NumFailedResolves should not change after a successful resolve attempt"), 0);
@@ -92,7 +86,7 @@ protected:
 
 		if (ResolvedObject)
 		{
-			AddError(FString::Printf(TEXT("Expected '%s.%s' to resolve to null."), ANSI_TO_TCHAR(PackageName), ANSI_TO_TCHAR(ObjectName)), 1);
+			FAIL_CHECK(FString::Printf(TEXT("Expected '%s.%s' to resolve to null."), ANSI_TO_TCHAR(PackageName), ANSI_TO_TCHAR(ObjectName)));
 			return false;
 		}
 		ObjectRefMetrics.TestNumFailedResolves(TEXT("NumFailedResolves should be incremented by one after a failed resolve attempt"), 1);
@@ -108,7 +102,7 @@ protected:
 
 		if (ResolvedObject)
 		{
-			AddError(FString::Printf(TEXT("Expected PACKEDREF(%" UPTRINT_X_FMT ") to resolve to null."), PackedRef.EncodedRef), 1);
+			FAIL_CHECK(FString::Printf(TEXT("Expected PACKEDREF(%" UPTRINT_X_FMT ") to resolve to null."), PackedRef.EncodedRef));
 			return false;
 		}
 		ObjectRefMetrics.TestNumFailedResolves(TEXT("NumFailedResolves should be incremented by one after a failed resolve attempt"), 1);
@@ -116,51 +110,41 @@ protected:
 	}
 };
 
-#define TEST_NAME_ROOT TEXT("System.CoreUObject.ObjectHandle")
-constexpr const uint32 ObjectHandleTestFlags = EAutomationTestFlags::ApplicationContextMask | EAutomationTestFlags::EngineFilter;
-
-IMPLEMENT_CUSTOM_SIMPLE_AUTOMATION_TEST(FObjectHandleTestNullBehavior, FObjectHandleTestBase, TEST_NAME_ROOT TEXT(".NullBehavior"), ObjectHandleTestFlags)
-bool FObjectHandleTestNullBehavior::RunTest(const FString& Parameters)
+TEST_CASE_METHOD(FObjectHandleTestBase, "CoreUObject::FObjectHandle::Null Behavior", "[CoreUObject][ObjectHandle][Smoke]")
 {
 	FObjectHandle TargetHandle = MakeObjectHandle(nullptr);
 
-	TestTrue(TEXT("Handle to target is null"), IsObjectHandleNull(TargetHandle));
-	TestTrue(TEXT("Handle to target is resolved"), IsObjectHandleResolved(TargetHandle));
+	TEST_TRUE(TEXT("Handle to target is null"), IsObjectHandleNull(TargetHandle));
+	TEST_TRUE(TEXT("Handle to target is resolved"), IsObjectHandleResolved(TargetHandle));
 
 	FSnapshotObjectRefMetrics ObjectRefMetrics(*this);
 	UObject* ResolvedObject = ResolveObjectHandle(TargetHandle);
 
-	TestEqual(TEXT("Resolved object is equal to original object"), (UObject*)nullptr, ResolvedObject);
+	TEST_EQUAL(TEXT("Resolved object is equal to original object"), (UObject*)nullptr, ResolvedObject);
 
 	ObjectRefMetrics.TestNumFailedResolves(TEXT("NumFailedResolves should not change after a resolve attempt on a null handle"), 0);
 	ObjectRefMetrics.TestNumResolves(TEXT("NumResolves should not change after a resolve attempt on a null handle"), 0);
 	ObjectRefMetrics.TestNumReads(TEXT("NumReads should be incremented by one after a resolve attempt on a null handle"), 1);
-
-	return true;
 }
 
-IMPLEMENT_CUSTOM_SIMPLE_AUTOMATION_TEST(FObjectHandleTestPointerBehavior, FObjectHandleTestBase, TEST_NAME_ROOT TEXT(".PointerBehavior"), ObjectHandleTestFlags)
-bool FObjectHandleTestPointerBehavior::RunTest(const FString& Parameters)
+TEST_CASE_METHOD(FObjectHandleTestBase, "CoreUObject::FObjectHandle::Pointer Behavior", "[CoreUObject][ObjectHandle][Smoke]")
 {
 	FObjectHandle TargetHandle = MakeObjectHandle((UObject*)0x0042);
 
-	TestFalse(TEXT("Handle to target is null"), IsObjectHandleNull(TargetHandle));
-	TestTrue(TEXT("Handle to target is resolved"), IsObjectHandleResolved(TargetHandle));
+	TEST_FALSE(TEXT("Handle to target is null"), IsObjectHandleNull(TargetHandle));
+	TEST_TRUE(TEXT("Handle to target is resolved"), IsObjectHandleResolved(TargetHandle));
 
 	FSnapshotObjectRefMetrics ObjectRefMetrics(*this);
 	UObject* ResolvedObject = ResolveObjectHandle(TargetHandle);
 
-	TestEqual(TEXT("Resolved object is equal to original object"), (UObject*)0x0042, ResolvedObject);
+	TEST_EQUAL(TEXT("Resolved object is equal to original object"), (UObject*)0x0042, ResolvedObject);
 
 	ObjectRefMetrics.TestNumResolves(TEXT("NumResolves should not change after a resolve attempt on a pointer handle"), 0);
 	ObjectRefMetrics.TestNumFailedResolves(TEXT("NumFailedResolves should not change after a resolve attempt on a pointer handle"), 0);
 	ObjectRefMetrics.TestNumReads(TEXT("NumReads should be incremented by one after a resolve attempt on a pointer handle"),1);
-
-	return true;
 }
 
-IMPLEMENT_CUSTOM_SIMPLE_AUTOMATION_TEST(FObjectHandleTestResolveEngineContentTarget, FObjectHandleTestBase, TEST_NAME_ROOT TEXT(".ResolveEngineContentTarget"), ObjectHandleTestFlags)
-bool FObjectHandleTestResolveEngineContentTarget::RunTest(const FString& Parameters)
+TEST_CASE_METHOD(FObjectHandleTestBase, "CoreUObject::FObjectHandle::Resolve Engine Content Target", "[CoreUObject][ObjectHandle][.Engine]")
 {
 	// Confirm we successfully resolve a correct reference to engine content
 	TestResolvableNonNull("/Engine/EngineResources/DefaultTexture", "DefaultTexture");
@@ -174,42 +158,28 @@ bool FObjectHandleTestResolveEngineContentTarget::RunTest(const FString& Paramet
 		// Attempt to load something that uses a User Defined Enum
 		TestResolvableNonNull("/Engine/ArtTools/RenderToTexture/Macros/RenderToTextureMacros", "RenderToTextureMacros:Array to HLSL Float Array.K2Node_Select_1", nullptr, nullptr, true);
 	}
-
-	return true;
 }
 
-IMPLEMENT_CUSTOM_SIMPLE_AUTOMATION_TEST(FObjectHandleTestResolveNonExistentTarget, FObjectHandleTestBase, TEST_NAME_ROOT TEXT(".ResolveNonExistentTarget"), ObjectHandleTestFlags)
-bool FObjectHandleTestResolveNonExistentTarget::RunTest(const FString& Parameters)
+TEST_CASE_METHOD(FObjectHandleTestBase, "CoreUObject::FObjectHandle::Resolve Non Existent Target", "[CoreUObject][ObjectHandle][Smoke]")
 {
 	// Confirm we don't successfully resolve an incorrect reference to engine content
 	TestResolveFailure("/Engine/EngineResources/NonExistentPackageName_0", "DefaultTexture");
 	TestResolveFailure("/Engine/EngineResources/DefaultTexture", "NonExistentObject_0");
-
-	return true;
 }
 
-IMPLEMENT_CUSTOM_SIMPLE_AUTOMATION_TEST(FObjectHandleTestResolveScriptTarget, FObjectHandleTestBase, TEST_NAME_ROOT TEXT(".ResolveScriptTarget"), ObjectHandleTestFlags)
-bool FObjectHandleTestResolveScriptTarget::RunTest(const FString& Parameters)
+
+TEST_CASE_METHOD(FObjectHandleTestBase, "CoreUObject::FObjectHandle::Resolve Script Target", "[CoreUObject][ObjectHandle][.Engine]")
 {
 	// Confirm we successfully resolve a correct reference to engine content
 	TestResolvableNonNull("/Script/Engine", "Default__Actor");
-
 	TestResolvableNonNull("/Script/Engine", "DefaultPawn");
-
-	return true;
 }
 
 #if UE_WITH_OBJECT_HANDLE_LATE_RESOLVE
-IMPLEMENT_CUSTOM_SIMPLE_AUTOMATION_TEST(FObjectHandleTestResolveMalformedHandle, FObjectHandleTestBase, TEST_NAME_ROOT TEXT(".ResolveMalformedHandle"), ObjectHandleTestFlags)
-bool FObjectHandleTestResolveMalformedHandle::RunTest(const FString& Parameters)
+
+TEST_CASE_METHOD(FObjectHandleTestBase, "CoreUObject::FObjectHandle::Resolve Malformed Handle", "[CoreUObject][ObjectHandle][Smoke]")
 {
 	TestResolveFailure(FPackedObjectRef { 0xFFFF'FFFF'FFFF'FFFFull });
 	TestResolveFailure(FPackedObjectRef { 0xEFEF'EFEF'EFEF'EFEFull });
-
-	return true;
 }
 #endif // UE_WITH_OBJECT_HANDLE_LATE_RESOLVE
-
-#undef TEST_NAME_ROOT
-
-#endif // WITH_DEV_AUTOMATION_TESTS
