Add missing GetHash functions for TIntVector2 and TIntVector4.

#okforgithuib public

diff --git a/Engine/Source/Runtime/Core/Public/Math/IntVector.h b/Engine/Source/Runtime/Core/Public/Math/IntVector.h
--- a/Engine/Source/Runtime/Core/Public/Math/IntVector.h
+++ b/Engine/Source/Runtime/Core/Public/Math/IntVector.h
@@ -920,7 +920,20 @@ PRAGMA_ENABLE_DEPRECATION_WARNINGS
 };
 
 /**
- * Creates a hash value from an FVector.
+ * Creates a hash value from an IntVector2.
+ *
+ * @param Vector the vector to create a hash value for
+ * @return The hash value from the components
+ */
+template<typename T>
+uint32 GetTypeHash(const TIntVector2<T>& Vector)
+{
+	// Note: this assumes there's no padding in Vector that could contain uncompared data.
+	return FCrc::MemCrc32(&Vector, sizeof(Vector));
+}
+
+/**
+ * Creates a hash value from an IntVector3.
  *
  * @param Vector the vector to create a hash value for
  * @return The hash value from the components
@@ -932,6 +945,19 @@ uint32 GetTypeHash(const TIntVector3<T>& Vector)
 	return FCrc::MemCrc_DEPRECATED(&Vector, sizeof(Vector));
 }
 
+/**
+ * Creates a hash value from an IntVector4.
+ *
+ * @param Vector the vector to create a hash value for
+ * @return The hash value from the components
+ */
+template<typename T>
+uint32 GetTypeHash(const TIntVector4<T>& Vector)
+{
+	// Note: this assumes there's no padding in Vector that could contain uncompared data.
+	return FCrc::MemCrc32(&Vector, sizeof(Vector));
+}
+
 } //! namespace UE::Math
 
 template <> struct TIsPODType<FIntVector2>  { enum { Value = true }; };
