revert GenericPlatformMath.cpp changes from CL 19562378
to preserve legacy bad fmod behavior

diff --git a/Engine/Source/Runtime/Core/Private/GenericPlatform/GenericPlatformMath.cpp b/Engine/Source/Runtime/Core/Private/GenericPlatform/GenericPlatformMath.cpp
--- a/Engine/Source/Runtime/Core/Private/GenericPlatform/GenericPlatformMath.cpp
+++ b/Engine/Source/Runtime/Core/Private/GenericPlatform/GenericPlatformMath.cpp
@@ -88,27 +88,51 @@ double FGenericPlatformMath::Atan2(double Y, double X)
 float FGenericPlatformMath::Fmod(float X, float Y)
 {
 	const float AbsY = FMath::Abs(Y);
-	// large arbitrary threshold could be removed if possible :
 	if (AbsY <= 1.e-8f)
 	{
 		FmodReportError(X, Y);
 		return 0.0;
 	}
 
-	return fmodf(X,Y);
+	// Convert to double for better precision, since intermediate rounding can lose enough precision to skew the result.
+	const double DX = double(X);
+	const double DY = double(Y);
+
+	const double Div = (DX / DY);
+	const double IntPortion = DY * TruncToDouble(Div);
+	const double Result = DX - IntPortion;
+	// Convert back to float. This is safe because the result will by definition not exceed the X input.
+	return float(Result);
 }
 
 double FGenericPlatformMath::Fmod(double X, double Y)
 {
 	const double AbsY = FMath::Abs(Y);
-	// large arbitrary threshold could be removed if possible :
 	if (AbsY <= 1.e-8)
 	{
 		FmodReportError(X, Y);
 		return 0.0;
 	}
 
+#if 1
+	// Due to the lack of standard support across platforms for `long double`, we can't rely on this to always minimize the intermediate precision loss of division and multiplication.
+	// As a result, due to small precision loss the result here could be different than the std library version of fmod(), but our validation tests in UnrealMathTest should show
+	// that when different the results are within a small delta from either 0 or the value of Y. We could just use the library version here of course, but it is slower and also will not
+	// match the vectorized versions of Fmod (VectorRegisterMod) that use division and truncation.
+	const double Div = (X / Y);
+	const double IntPortion = Y * FMath::TruncToDouble(Div);
+	const double Result = X - IntPortion;
+	return Result;
+#elif 0
+	const long double LDX = (long double)(X);
+	const long double LDY = (long double)(Y);
+	const long double Div = (LDX / LDY);
+	const long double IntPortion = LDY * truncl(Div);
+	const long double Result = LDX - IntPortion;
+	return double(Result);
+#else
 	return fmod(X, Y);
+#endif
 }
 
 void FGenericPlatformMath::FmodReportError(float X, float Y)
@@ -166,38 +190,11 @@ public:
 		check(MathPlatform::IsFinite(One));
 		check(!MathPlatform::IsNaN(MinusOneE37));
 		check(MathPlatform::IsFinite(MinusOneE37));
-
 		check(MathPlatform::FloorLog2((uint32)Zero) == 0);
 		check(MathPlatform::FloorLog2((uint32)One) == 0);
 		check(MathPlatform::FloorLog2((uint32)Two) == 1);
 		check(MathPlatform::FloorLog2((uint32)Twelve) == 3);
 		check(MathPlatform::FloorLog2((uint32)Sixteen) == 4);
-		
-		check(MathPlatform::CeilLogTwo((uint32)Zero) == 0);
-		check(MathPlatform::CeilLogTwo((uint32)One) == 0);
-		check(MathPlatform::CeilLogTwo((uint32)Two) == 1);
-		check(MathPlatform::CeilLogTwo((uint32)Twelve) == 4);
-		check(MathPlatform::CeilLogTwo((uint32)Sixteen) == 4);
-		
-		check(MathPlatform::CountLeadingZeros8((uint8)Zero) == 8);
-		check(MathPlatform::CountLeadingZeros((uint32)Zero) == 32);
-		check(MathPlatform::CountTrailingZeros((uint32)Zero) == 32);
-#if PLATFORM_64BITS
-		check(MathPlatform::CountLeadingZeros64((uint64)Zero) == 64);
-		check(MathPlatform::CountTrailingZeros64((uint64)Zero) == 64);
-		
-		check(MathPlatform::FloorLog2_64((uint32)Zero) == 0);
-		check(MathPlatform::FloorLog2_64((uint32)One) == 0);
-		check(MathPlatform::FloorLog2_64((uint32)Two) == 1);
-		check(MathPlatform::FloorLog2_64((uint32)Twelve) == 3);
-		check(MathPlatform::FloorLog2_64((uint32)Sixteen) == 4);
-		
-		check(MathPlatform::CeilLogTwo64((uint32)Zero) == 0);
-		check(MathPlatform::CeilLogTwo64((uint32)One) == 0);
-		check(MathPlatform::CeilLogTwo64((uint32)Two) == 1);
-		check(MathPlatform::CeilLogTwo64((uint32)Twelve) == 4);
-		check(MathPlatform::CeilLogTwo64((uint32)Sixteen) == 4);
-#endif
 	}
 };
 
