Introducing Catch2 disabled test cases, scenarios etc
Motivation: [.] only hides a test and it can still be executed using ~[SomeOtherTag]; It's very easy to fall into the trap of executing an infinite loop test that is decorated with the [.] tag

diff --git a/Engine/Source/Developer/LowLevelTestsRunner/Public/TestHarness.h b/Engine/Source/Developer/LowLevelTestsRunner/Public/TestHarness.h
--- a/Engine/Source/Developer/LowLevelTestsRunner/Public/TestHarness.h
+++ b/Engine/Source/Developer/LowLevelTestsRunner/Public/TestHarness.h
@@ -79,6 +79,13 @@ THIRD_PARTY_INCLUDES_START
 THIRD_PARTY_INCLUDES_END
 #endif // defined(THIRD_PARTY_INCLUDES_START) && defined(THIRD_PARTY_INCLUDES_END)
 
+#define LLT_CONCAT_(x, y) x##y
+#define LLT_CONCAT(x, y) LLT_CONCAT_(x,y)
+
+#define DISABLED_TEST_CASE(...) static void LLT_CONCAT(disabled_test_,__LINE__)()
+#define DISABLED_TEST_CASE_METHOD(...) static void LLT_CONCAT(disabled_test_method_,__LINE__)()
+#define DISABLED_SCENARIO(...) static void LLT_CONCAT(disabled_scenario_,__LINE__)()
+
 // Tell Catch how to print TTuple<KeyType, ValueType>
 template <typename... Types> struct TTuple;
 
diff --git a/Engine/Source/Runtime/Core/Tests/Async/TaskGraphTest.cpp b/Engine/Source/Runtime/Core/Tests/Async/TaskGraphTest.cpp
--- a/Engine/Source/Runtime/Core/Tests/Async/TaskGraphTest.cpp
+++ b/Engine/Source/Runtime/Core/Tests/Async/TaskGraphTest.cpp
@@ -1285,7 +1285,7 @@ PRAGMA_ENABLE_DEPRECATION_WARNINGS
 		}
 	}
 
-	TEST_CASE_METHOD(FCoreTestFixture, "Core::Async::TaskGraph::Performance", "[.][Core][Async][.Perf]")
+	DISABLED_TEST_CASE_METHOD(FCoreTestFixture, "Core::Async::TaskGraph::Performance", "[Core][Async][.Perf]")
 	{
 		TRACE_CPUPROFILER_EVENT_SCOPE(TaskGraphTests_PerfTest);
 
diff --git a/Engine/Source/Runtime/Core/Tests/Containers/ConcurrentQueuesTest.cpp b/Engine/Source/Runtime/Core/Tests/Containers/ConcurrentQueuesTest.cpp
--- a/Engine/Source/Runtime/Core/Tests/Containers/ConcurrentQueuesTest.cpp
+++ b/Engine/Source/Runtime/Core/Tests/Containers/ConcurrentQueuesTest.cpp
@@ -401,7 +401,7 @@ namespace UE { namespace ConcurrentQueuesTests
 		}
 	}
 
-	TEST_CASE_METHOD(FCoreTestFixture, "Core::Containers::TMpscQueue::ConcurrentQueuesPerfTest", "[.][Core][Containers][.Perf]")
+	DISABLED_TEST_CASE_METHOD(FCoreTestFixture, "Core::Containers::TMpscQueue::ConcurrentQueuesPerfTest", "[Core][Containers][.Perf]")
 	{
 		UE_BENCHMARK(5, TestTCircularQueueSingleThread<5'000'000>);
 		UE_BENCHMARK(5, TestQueueSingleThread<5'000'000, TQueueAdapter<TQueue<uint32, EQueueMode::Spsc>>>);
@@ -721,7 +721,7 @@ namespace ClosableMpscQueueTests
 		}
 	}
 
-	TEST_CASE_METHOD(FCoreTestFixture, "Core::Containers::TClosableMpscQueue::ClosableMpscQueuePerfTest", "[.][Core][Containers][.Perf]")
+	DISABLED_TEST_CASE_METHOD(FCoreTestFixture, "Core::Containers::TClosableMpscQueue::ClosableMpscQueuePerfTest", "[Core][Containers][.Perf]")
 	{
 		UE_BENCHMARK(5, TestClosingEmptyQueue<10'000'000, TClosableMpscQueue<void*>>);
 		UE_BENCHMARK(5, TestClosingEmptyQueue<10'000'000, TClosableMpscQueueOnTQueue<void*>>);
@@ -1073,7 +1073,7 @@ namespace DepletableMpscQueueTests
 		}
 	}
 
-	TEST_CASE_METHOD(FCoreTestFixture, "Core::Containers::TDepletableMpscQueue::DepletableMpscQueuePerfTest", "[.][Core][Containers][.Perf]")
+	DISABLED_TEST_CASE_METHOD(FCoreTestFixture, "Core::Containers::TDepletableMpscQueue::DepletableMpscQueuePerfTest", "[Core][Containers][.Perf]")
 	{
 		UE_BENCHMARK(5, TestPerf<10'000, TDepletableMpscQueue<void*>>);
 		UE_BENCHMARK(5, TestPerf<10'000, TDepletableMpscQueueOnTQueue<void*>>);
diff --git a/Engine/Source/Runtime/Core/Tests/Misc/ByteSwapTest.cpp b/Engine/Source/Runtime/Core/Tests/Misc/ByteSwapTest.cpp
--- a/Engine/Source/Runtime/Core/Tests/Misc/ByteSwapTest.cpp
+++ b/Engine/Source/Runtime/Core/Tests/Misc/ByteSwapTest.cpp
@@ -56,7 +56,7 @@ TPair<FTimespan, T> BenchmarkByteswapping(T InitialValue, uint64 LoopCount, Byte
 	return MakeTuple(Duration, Sum);
 }
 
-TEST_CASE_METHOD(FAutomationTestFixture, "Core::Misc::ByteSwap::Perf Test", "[.][Core][Misc][.Perf]")
+DISABLED_TEST_CASE_METHOD(FAutomationTestFixture, "Core::Misc::ByteSwap::Perf Test", "[Core][Misc][.Perf]")
 {
 #ifdef UE_BYTESWAP_INTRINSIC_PRIVATE_16
 	{
diff --git a/Engine/Source/Runtime/CoreUObject/Tests/Serialization/BulkDataTests.cpp b/Engine/Source/Runtime/CoreUObject/Tests/Serialization/BulkDataTests.cpp
--- a/Engine/Source/Runtime/CoreUObject/Tests/Serialization/BulkDataTests.cpp
+++ b/Engine/Source/Runtime/CoreUObject/Tests/Serialization/BulkDataTests.cpp
@@ -1,7 +1,7 @@
 // Copyright Epic Games, Inc. All Rights Reserved.
 
-#include "Serialization/BulkData.h"
 #include "TestHarness.h"
+#include "Serialization/BulkData.h"
 
 namespace BulkDataTest
 {
@@ -61,7 +61,7 @@ namespace BulkDataTest
 
 #if USE_NEW_BULKDATA // The warning only exists for BulkData2
 		// Both the ::Lock and ::GetCopy calls should warn that we cannot load the missing data (as it will still have a valid size)
-		AddExpectedError(TEXT("Attempting to load a BulkData object that cannot be loaded from disk"), EAutomationExpectedErrorFlags::Exact, 2);
+		FAIL(TEXT("Attempting to load a BulkData object that cannot be loaded from disk"));
 #endif //USE_NEW_BULKDATA
 
 		DataPtr = BulkData.Lock(LOCK_READ_WRITE);
