Logging: Fixed a hang in FlushThreadedLogs in panic mode

The logging thread needs to continue running until ThreadWakeEvent is null because it may otherwise miss triggering ThreadIdleEvents before it exits.

#rnx

diff --git a/Engine/Source/Runtime/Core/Private/Misc/OutputDeviceRedirector.cpp b/Engine/Source/Runtime/Core/Private/Misc/OutputDeviceRedirector.cpp
--- a/Engine/Source/Runtime/Core/Private/Misc/OutputDeviceRedirector.cpp
+++ b/Engine/Source/Runtime/Core/Private/Misc/OutputDeviceRedirector.cpp
@@ -359,25 +359,17 @@ void FOutputDeviceRedirectorState::ThreadLoop()
 		MasterThreadId.store(ThreadId, std::memory_order_relaxed);
 	}
 
-	if (FEvent* WakeEvent = ThreadWakeEvent.load(std::memory_order_acquire))
+	while (FEvent* WakeEvent = ThreadWakeEvent.load(std::memory_order_acquire))
 	{
-		while (IsMasterThread(ThreadId))
+		WakeEvent->Wait();
+		while (!BufferedLines.IsEmpty() && IsMasterThread(ThreadId))
 		{
-			WakeEvent->Wait();
-			do
+			if (FOutputDevicesMasterScope Lock(*this); Lock.IsLocked())
 			{
-				if (FOutputDevicesMasterScope Lock(*this); Lock.IsLocked())
-				{
-					FlushBufferedLines();
-				}
-				else
-				{
-					break;
-				}
+				FlushBufferedLines();
 			}
-			while (!BufferedLines.IsEmpty());
-			ThreadIdleEvents.Deplete([](FEvent* Event) { Event->Trigger(); });
 		}
+		ThreadIdleEvents.Deplete([](FEvent* Event) { Event->Trigger(); });
 	}
 }
 
