Partially backed out 19571202 because it stomped changes to several tests

#rnx

diff --git a/Engine/Source/Runtime/Core/Tests/Compression/CompressedBufferTest.cpp b/Engine/Source/Runtime/Core/Tests/Compression/CompressedBufferTest.cpp
--- a/Engine/Source/Runtime/Core/Tests/Compression/CompressedBufferTest.cpp
+++ b/Engine/Source/Runtime/Core/Tests/Compression/CompressedBufferTest.cpp
@@ -1,110 +1,75 @@
 // Copyright Epic Games, Inc. All Rights Reserved.
 
 #include "Compression/CompressedBuffer.h"
+
+#include "Algo/Compare.h"
 #include "Compression/OodleDataCompression.h"
 #include "IO/IoHash.h"
-#include "Hash/Blake3.h"
 #include "Serialization/BufferReader.h"
-#include "TestFixtures/CoreTestFixture.h"
+#include "TestHarness.h"
 
-TEST_CASE_METHOD(FCoreTestFixture, "Core::Compression::FCompressedBuffer::Compress", "[Core][Math][Smoke]")
+TEST_CASE("Core::Compression::CompressedBuffer::Compress", "[Core][Compression][Smoke]")
 {
 	const uint8 ZeroBuffer[1024]{};
-	const FBlake3Hash ZeroBufferHash = FBlake3::HashBuffer(MakeMemoryView(ZeroBuffer));
+	const FIoHash ZeroBufferHash = FIoHash::HashBuffer(MakeMemoryView(ZeroBuffer));
 
-	// Test Null
+	SECTION("Null")
 	{
-		INFO("Null");
-
 		const FCompressedBuffer Buffer;
-		CHECK_FALSE((bool)Buffer);
+		CHECK_FALSE(Buffer);
 		CHECK(Buffer.IsNull());
 		CHECK(Buffer.IsOwned());
-		CHECK(Buffer.GetCompressedSize()==uint64(0));
-		CHECK(Buffer.GetRawSize()==uint64(0));
-		CHECK(Buffer.GetRawHash()==FBlake3Hash::Zero);
+		CHECK(Buffer.GetCompressedSize() == 0);
+		CHECK(Buffer.GetRawSize() == 0);
+		CHECK(Buffer.GetRawHash() == FIoHash::Zero);
 		CHECK(Buffer.Decompress().IsNull());
 		CHECK(Buffer.DecompressToComposite().IsNull());
 	}
 
-	// Test Method None
+	SECTION("Method None")
 	{
-		INFO("Method None");
+		const FCompressedBuffer BufferNew = FCompressedBuffer::Compress(FSharedBuffer::MakeView(MakeMemoryView(ZeroBuffer)),
+			ECompressedBufferCompressor::NotSet, ECompressedBufferCompressionLevel::None);
+		const FCompressedBuffer BufferCopy = FCompressedBuffer::FromCompressed(BufferNew.GetCompressed());
+		const FCompressedBuffer Buffer = GENERATE_REF(BufferNew, BufferCopy);
+
+		CHECK(Buffer);
+		CHECK(!Buffer.IsNull());
+		CHECK(Buffer.IsOwned());
+		CHECK(Buffer.GetCompressedSize() == sizeof(ZeroBuffer) + 64);
+		CHECK(Buffer.GetRawSize() == sizeof(ZeroBuffer));
+		CHECK(Buffer.GetRawHash() == ZeroBufferHash);
+		CHECK(FIoHash::HashBuffer(Buffer.Decompress()) == ZeroBufferHash);
+		CHECK(FIoHash::HashBuffer(Buffer.DecompressToComposite()) == Buffer.GetRawHash());
 
 		ECompressedBufferCompressor Compressor = ECompressedBufferCompressor::Kraken;
 		ECompressedBufferCompressionLevel CompressionLevel = ECompressedBufferCompressionLevel::Normal;
 		uint64 BlockSize = MAX_uint64;
-
-		const FCompressedBuffer Buffer = FCompressedBuffer::Compress(FSharedBuffer::MakeView(MakeMemoryView(ZeroBuffer)),
-			ECompressedBufferCompressor::NotSet, ECompressedBufferCompressionLevel::None);
-		CHECK((bool)Buffer);
-		CHECK_FALSE(Buffer.IsNull());
-		CHECK(Buffer.IsOwned());
-		CHECK(Buffer.GetCompressedSize()==sizeof(ZeroBuffer) + 64);
-		CHECK(Buffer.GetRawSize()==sizeof(ZeroBuffer));
-		CHECK(Buffer.GetRawHash()==ZeroBufferHash);
-		CHECK(FBlake3::HashBuffer(Buffer.Decompress())==ZeroBufferHash);
-		CHECK(FBlake3::HashBuffer(Buffer.DecompressToComposite())==Buffer.GetRawHash());
-		CHECK(Buffer.TryGetCompressParameters(Compressor, CompressionLevel, BlockSize));	
-		CHECK(Compressor == ECompressedBufferCompressor::NotSet);
-		CHECK(CompressionLevel == ECompressedBufferCompressionLevel::None);
-		CHECK(BlockSize == 0);
-
-		Compressor = ECompressedBufferCompressor::Kraken;
-		CompressionLevel = ECompressedBufferCompressionLevel::Normal;
-		BlockSize = MAX_uint64;
-
-		const FCompressedBuffer BufferCopy = FCompressedBuffer::FromCompressed(Buffer.GetCompressed());
-		CHECK((bool)Buffer);
-		CHECK_FALSE(Buffer.IsNull());
-		CHECK(Buffer.IsOwned());
-		CHECK(Buffer.GetCompressedSize()==sizeof(ZeroBuffer) + 64);
-		CHECK(Buffer.GetRawSize()==sizeof(ZeroBuffer));
-		CHECK(Buffer.GetRawHash()==ZeroBufferHash);
-		CHECK(FBlake3::HashBuffer(Buffer.Decompress())==ZeroBufferHash);
-		CHECK(FBlake3::HashBuffer(Buffer.DecompressToComposite())==Buffer.GetRawHash());
 		CHECK(Buffer.TryGetCompressParameters(Compressor, CompressionLevel, BlockSize));
 		CHECK(Compressor == ECompressedBufferCompressor::NotSet);
-		CHECK(CompressionLevel ==ECompressedBufferCompressionLevel::None);
+		CHECK(CompressionLevel == ECompressedBufferCompressionLevel::None);
 		CHECK(BlockSize == 0);
 	}
 
-	// Test Method Oodle
+	SECTION("Method Oodle")
 	{
-		INFO("Method Oodle");
-
-		ECompressedBufferCompressor Compressor = ECompressedBufferCompressor::Kraken;
-		ECompressedBufferCompressionLevel CompressionLevel = ECompressedBufferCompressionLevel::Normal;
-		uint64 BlockSize = MAX_uint64;
-
-		const FCompressedBuffer Buffer = FCompressedBuffer::Compress(FSharedBuffer::MakeView(MakeMemoryView(ZeroBuffer)),
+		const FCompressedBuffer BufferNew = FCompressedBuffer::Compress(FSharedBuffer::MakeView(MakeMemoryView(ZeroBuffer)),
 			ECompressedBufferCompressor::Mermaid, ECompressedBufferCompressionLevel::VeryFast);
-		CHECK((bool)Buffer);
-		CHECK_FALSE(Buffer.IsNull());
-		CHECK(Buffer.IsOwned());
-		CHECK(Buffer.GetCompressedSize() < sizeof(ZeroBuffer));
-		CHECK(Buffer.GetRawSize()==sizeof(ZeroBuffer));
-		CHECK(Buffer.GetRawHash()==ZeroBufferHash);
-		CHECK(FBlake3::HashBuffer(Buffer.Decompress())==ZeroBufferHash);
-		CHECK(FBlake3::HashBuffer(Buffer.DecompressToComposite())==Buffer.GetRawHash());
-		CHECK(Buffer.TryGetCompressParameters(Compressor, CompressionLevel, BlockSize));
-		CHECK(Compressor == ECompressedBufferCompressor::Mermaid);
-		CHECK(CompressionLevel == ECompressedBufferCompressionLevel::VeryFast);
-		CHECK(FMath::IsPowerOfTwo(BlockSize));
-
-		Compressor = ECompressedBufferCompressor::Kraken;
-		CompressionLevel = ECompressedBufferCompressionLevel::Normal;
-		BlockSize = MAX_uint64;
+		const FCompressedBuffer BufferCopy = FCompressedBuffer::FromCompressed(BufferNew.GetCompressed());
+		const FCompressedBuffer Buffer = GENERATE_REF(BufferNew, BufferCopy);
 
-		const FCompressedBuffer BufferCopy = FCompressedBuffer::FromCompressed(Buffer.GetCompressed());
-		CHECK((bool)Buffer);
-		CHECK_FALSE(Buffer.IsNull());
+		CHECK(Buffer);
+		CHECK(!Buffer.IsNull());
 		CHECK(Buffer.IsOwned());
 		CHECK(Buffer.GetCompressedSize() < sizeof(ZeroBuffer));
-		CHECK(Buffer.GetRawSize()==sizeof(ZeroBuffer));
-		CHECK(Buffer.GetRawHash()==ZeroBufferHash);
-		CHECK(FBlake3::HashBuffer(Buffer.Decompress())==ZeroBufferHash);
-		CHECK(FBlake3::HashBuffer(Buffer.DecompressToComposite())==Buffer.GetRawHash());
+		CHECK(Buffer.GetRawSize() == sizeof(ZeroBuffer));
+		CHECK(Buffer.GetRawHash() == ZeroBufferHash);
+		CHECK(FIoHash::HashBuffer(Buffer.Decompress()) == ZeroBufferHash);
+		CHECK(FIoHash::HashBuffer(Buffer.DecompressToComposite()) == Buffer.GetRawHash());
+
+		ECompressedBufferCompressor Compressor = ECompressedBufferCompressor::Kraken;
+		ECompressedBufferCompressionLevel CompressionLevel = ECompressedBufferCompressionLevel::Normal;
+		uint64 BlockSize = MAX_uint64;
 		CHECK(Buffer.TryGetCompressParameters(Compressor, CompressionLevel, BlockSize));
 		CHECK(Compressor == ECompressedBufferCompressor::Mermaid);
 		CHECK(CompressionLevel == ECompressedBufferCompressionLevel::VeryFast);
@@ -112,7 +77,7 @@ TEST_CASE_METHOD(FCoreTestFixture, "Core::Compression::FCompressedBuffer::Compre
 	}
 }
 
-TEST_CASE_METHOD(FCoreTestFixture, "Core::Compression::FCompressedBuffer::Decompress", "[Core][Math][Smoke]")
+TEST_CASE("Core::Compression::CompressedBuffer::Decompress", "[Core][Compression][Smoke]")
 {
 	const auto GenerateData = [](int32 N) -> TArray<uint64>
 	{
@@ -125,17 +90,6 @@ TEST_CASE_METHOD(FCoreTestFixture, "Core::Compression::FCompressedBuffer::Decomp
 		return Data;
 	};
 
-	const auto ValidateData = [this](TConstArrayView<uint64> Values, TConstArrayView<uint64> ExpectedValues, int32 Offset)
-	{
-		int32 ExpectedIndex = Offset;
-		for (uint64 Value : Values)
-		{
-			const uint64 ExpectedValue = ExpectedValues[ExpectedIndex];
-			CHECK(Value==ExpectedValue);
-			ExpectedIndex++;
-		}
-	};
-
 	const auto CastToArrayView = [](FMemoryView View) -> TConstArrayView<uint64>
 	{
 		return MakeArrayView(static_cast<const uint64*>(View.GetData()), static_cast<int32>(View.GetSize() / sizeof(uint64)));
@@ -143,11 +97,9 @@ TEST_CASE_METHOD(FCoreTestFixture, "Core::Compression::FCompressedBuffer::Decomp
 
 	FCompressedBufferReader Reader;
 
-	// Test decompress with offset and size
+	SECTION("Decompress with offset and size.")
 	{
-		INFO("Decompress with offset and size");
-
-		const auto UncompressAndValidate = [this, &Reader, &ValidateData, &CastToArrayView](
+		const auto UncompressAndValidate = [&Reader, &CastToArrayView](
 			const FCompressedBuffer& Compressed,
 			const int32 OffsetCount,
 			const int32 Count,
@@ -156,16 +108,14 @@ TEST_CASE_METHOD(FCoreTestFixture, "Core::Compression::FCompressedBuffer::Decomp
 			Reader.SetSource(Compressed);
 			{
 				const FSharedBuffer Uncompressed = Reader.Decompress(OffsetCount * sizeof(uint64), Count * sizeof(uint64));
-				const TConstArrayView<uint64> Values = CastToArrayView(Uncompressed);
-				CHECK(Values.Num()==Count);
-				ValidateData(Values, ExpectedValues, OffsetCount);
+				CHECK(Algo::Compare(CastToArrayView(Uncompressed), ExpectedValues.Mid(OffsetCount, Count)));
 			}
 			{
 				FUniqueBuffer Uncompressed = FUniqueBuffer::Alloc(Count * sizeof(uint64));
-				const bool bOk = Reader.TryDecompressTo(Uncompressed, OffsetCount * sizeof(uint64));
-				CHECK(bOk);
-				const TConstArrayView<uint64> Values = CastToArrayView(Uncompressed);
-				ValidateData(Values, ExpectedValues, OffsetCount);
+				CHECKED_IF(Reader.TryDecompressTo(Uncompressed, OffsetCount * sizeof(uint64)))
+				{
+					CHECK(Algo::Compare(CastToArrayView(Uncompressed), ExpectedValues.Mid(OffsetCount, Count)));
+				}
 			}
 		};
 
@@ -192,10 +142,8 @@ TEST_CASE_METHOD(FCoreTestFixture, "Core::Compression::FCompressedBuffer::Decomp
 		UncompressAndValidate(Compressed, 4993, 4, ExpectedValues);
 	}
 
-	// Decompress with offset only
+	SECTION("Decompress with offset only.")
 	{
-		INFO("Decompress with offset only");
-
 		constexpr uint64 BlockSize = 64 * sizeof(uint64);
 		constexpr int32 N = 1000;
 		const TArray<uint64> ExpectedValues = GenerateData(N);
@@ -205,21 +153,19 @@ TEST_CASE_METHOD(FCoreTestFixture, "Core::Compression::FCompressedBuffer::Decomp
 			FOodleDataCompression::ECompressor::Mermaid,
 			FOodleDataCompression::ECompressionLevel::Fast,
 			BlockSize);
-		
+
 		constexpr uint64 OffsetCount = 150;
 		{
 			FSharedBuffer Buffer = Compressed.GetCompressed().ToShared();
 			FBufferReader Ar(const_cast<void*>(Buffer.GetData()), int64(Buffer.GetSize()), /*bFreeOnClose*/ false, /*bIsPersistent*/ true);
 			FCompressedBufferReaderSourceScope Source(Reader, Ar);
 			const FSharedBuffer Uncompressed = Reader.Decompress(OffsetCount * sizeof(uint64));
-			const TConstArrayView<uint64> Values = CastToArrayView(Uncompressed);
-			ValidateData(Values, ExpectedValues, OffsetCount);
+			CHECK(Algo::Compare(CastToArrayView(Uncompressed), MakeArrayView(ExpectedValues).RightChop(OffsetCount)));
 		}
 		{
 			FCompressedBufferReaderSourceScope Source(Reader, Compressed);
 			const FSharedBuffer Uncompressed = Reader.Decompress(OffsetCount * sizeof(uint64));
-			const TConstArrayView<uint64> Values = CastToArrayView(Uncompressed);
-			ValidateData(Values, ExpectedValues, OffsetCount);
+			CHECK(Algo::Compare(CastToArrayView(Uncompressed), MakeArrayView(ExpectedValues).RightChop(OffsetCount)));
 		}
 
 		// Short Buffer
@@ -232,10 +178,8 @@ TEST_CASE_METHOD(FCoreTestFixture, "Core::Compression::FCompressedBuffer::Decomp
 		}
 	}
 
-	// Only one block
+	SECTION("Decompress only block.")
 	{
-		INFO("Decompress Only one block");
-
 		constexpr uint64 BlockSize = 256 * sizeof(uint64);
 		constexpr int32 N = 100;
 		const TArray<uint64> ExpectedValues = GenerateData(N);
@@ -253,21 +197,17 @@ TEST_CASE_METHOD(FCoreTestFixture, "Core::Compression::FCompressedBuffer::Decomp
 			FBufferReader Ar(const_cast<void*>(Buffer.GetData()), int64(Buffer.GetSize()), /*bFreeOnClose*/ false, /*bIsPersistent*/ true);
 			FCompressedBufferReaderSourceScope Source(Reader, Ar);
 			const FSharedBuffer Uncompressed = Reader.Decompress(OffsetCount * sizeof(uint64), Count * sizeof(uint64));
-			const TConstArrayView<uint64> Values = CastToArrayView(Uncompressed);
-			ValidateData(Values, ExpectedValues, OffsetCount);
+			CHECK(Algo::Compare(CastToArrayView(Uncompressed), MakeArrayView(ExpectedValues).Mid(OffsetCount, Count)));
 		}
 		{
 			FCompressedBufferReaderSourceScope Source(Reader, Compressed);
 			const FSharedBuffer Uncompressed = Reader.Decompress(OffsetCount * sizeof(uint64), Count * sizeof(uint64));
-			const TConstArrayView<uint64> Values = CastToArrayView(Uncompressed);
-			ValidateData(Values, ExpectedValues, OffsetCount);
+			CHECK(Algo::Compare(CastToArrayView(Uncompressed), MakeArrayView(ExpectedValues).Mid(OffsetCount, Count)));
 		}
 	}
 
-	// Uncompressed
+	SECTION("Decompress from an uncompressed buffer.")
 	{
-		INFO("Uncompressed");
-
 		constexpr int32 N = 4242;
 		const TArray<uint64> ExpectedValues = GenerateData(N);
 
@@ -281,22 +221,18 @@ TEST_CASE_METHOD(FCoreTestFixture, "Core::Compression::FCompressedBuffer::Decomp
 			constexpr uint64 OffsetCount = 0;
 			constexpr uint64 Count = N;
 			const FSharedBuffer Uncompressed = Reader.Decompress(OffsetCount * sizeof(uint64), Count * sizeof(uint64));
-			const TConstArrayView<uint64> Values = CastToArrayView(Uncompressed);
-			ValidateData(Values, ExpectedValues, OffsetCount);
+			CHECK(Algo::Compare(CastToArrayView(Uncompressed), MakeArrayView(ExpectedValues).Mid(OffsetCount, Count)));
 		}
 
 		{
 			constexpr uint64 OffsetCount = 21;
 			constexpr uint64 Count = 999;
 			const FSharedBuffer Uncompressed = Reader.Decompress(OffsetCount * sizeof(uint64), Count * sizeof(uint64));
-			const TConstArrayView<uint64> Values = CastToArrayView(Uncompressed);
-			ValidateData(Values, ExpectedValues, OffsetCount);
+			CHECK(Algo::Compare(CastToArrayView(Uncompressed), MakeArrayView(ExpectedValues).Mid(OffsetCount, Count)));
 		}
 
 		// Short Buffer
 		{
-			INFO("Short Buffer");
-
 			const FCompressedBuffer CompressedShort =
 				FCompressedBuffer::FromCompressed(Compressed.GetCompressed().Mid(0, Compressed.GetCompressedSize() - 128));
 			Reader.SetSource(CompressedShort);
diff --git a/Engine/Source/Runtime/Core/Tests/Hash/Blake3Test.cpp b/Engine/Source/Runtime/Core/Tests/Hash/Blake3Test.cpp
--- a/Engine/Source/Runtime/Core/Tests/Hash/Blake3Test.cpp
+++ b/Engine/Source/Runtime/Core/Tests/Hash/Blake3Test.cpp
@@ -1,43 +1,45 @@
 // Copyright Epic Games, Inc. All Rights Reserved.
 
 #include "Hash/Blake3.h"
+
 #include "Misc/StringBuilder.h"
 #include "TestHarness.h"
 
-TEST_CASE_METHOD(FAutomationTestFixture, "Core::Hash::FBlake3::Smoke Test", "[Core][Math][Smoke]")
+TEST_CASE("Core::Hash::Blake3", "[Core][Hash][Smoke]")
 {
-	constexpr uint32 MaxInputLength = 31'744u;
-
 	struct FBlake3TestCase
 	{
 		uint32 InputLength;
 		FStringView Hash;
 	};
-	const FBlake3TestCase TestCases[] =
+
+	constexpr FBlake3TestCase TestCases[] =
 	{
-		{     0u,        TEXTVIEW("af1349b9f5f9a1a6a0404dea36dcc9499bcb25c9adc112b7cc9a93cae41f3262")},
-		{     1u,        TEXTVIEW("2d3adedff11b61f14c886e35afa036736dcd87a74d27b5c1510225d0f592e213")},
-		{ 1'023u,        TEXTVIEW("10108970eeda3eb932baac1428c7a2163b0e924c9a9e25b35bba72b28f70bd11")},
-		{ 1'024u,        TEXTVIEW("42214739f095a406f3fc83deb889744ac00df831c10daa55189b5d121c855af7")},
-		{ 1'025u,        TEXTVIEW("d00278ae47eb27b34faecf67b4fe263f82d5412916c1ffd97c8cb7fb814b8444")},
-		{ 2'048u,        TEXTVIEW("e776b6028c7cd22a4d0ba182a8bf62205d2ef576467e838ed6f2529b85fba24a")},
-		{ 2'049u,        TEXTVIEW("5f4d72f40d7a5f82b15ca2b2e44b1de3c2ef86c426c95c1af0b6879522563030")},
-		{ 3'072u,        TEXTVIEW("b98cb0ff3623be03326b373de6b9095218513e64f1ee2edd2525c7ad1e5cffd2")},
-		{ 3'073u,        TEXTVIEW("7124b49501012f81cc7f11ca069ec9226cecb8a2c850cfe644e327d22d3e1cd3")},
-		{ 4'096u,        TEXTVIEW("015094013f57a5277b59d8475c0501042c0b642e531b0a1c8f58d2163229e969")},
-		{ 4'097u,        TEXTVIEW("9b4052b38f1c5fc8b1f9ff7ac7b27cd242487b3d890d15c96a1c25b8aa0fb995")},
-		{ 5'120u,        TEXTVIEW("9cadc15fed8b5d854562b26a9536d9707cadeda9b143978f319ab34230535833")},
-		{ 5'121u,        TEXTVIEW("628bd2cb2004694adaab7bbd778a25df25c47b9d4155a55f8fbd79f2fe154cff")},
-		{ 6'144u,        TEXTVIEW("3e2e5b74e048f3add6d21faab3f83aa44d3b2278afb83b80b3c35164ebeca205")},
-		{ 6'145u,        TEXTVIEW("f1323a8631446cc50536a9f705ee5cb619424d46887f3c376c695b70e0f0507f")},
-		{ 7'168u,        TEXTVIEW("61da957ec2499a95d6b8023e2b0e604ec7f6b50e80a9678b89d2628e99ada77a")},
-		{ 7'169u,        TEXTVIEW("a003fc7a51754a9b3c7fae0367ab3d782dccf28855a03d435f8cfe74605e7817")},
-		{ 8'192u,        TEXTVIEW("aae792484c8efe4f19e2ca7d371d8c467ffb10748d8a5a1ae579948f718a2a63")},
-		{ 8'193u,        TEXTVIEW("bab6c09cb8ce8cf459261398d2e7aef35700bf488116ceb94a36d0f5f1b7bc3b")},
-		{16'384u,        TEXTVIEW("f875d6646de28985646f34ee13be9a576fd515f76b5b0a26bb324735041ddde4")},
-		{MaxInputLength, TEXTVIEW("62b6960e1a44bcc1eb1a611a8d6235b6b4b78f32e7abc4fb4c6cdcce94895c47")},
+		{     0u, TEXTVIEW("af1349b9f5f9a1a6a0404dea36dcc9499bcb25c9adc112b7cc9a93cae41f3262")},
+		{     1u, TEXTVIEW("2d3adedff11b61f14c886e35afa036736dcd87a74d27b5c1510225d0f592e213")},
+		{ 1'023u, TEXTVIEW("10108970eeda3eb932baac1428c7a2163b0e924c9a9e25b35bba72b28f70bd11")},
+		{ 1'024u, TEXTVIEW("42214739f095a406f3fc83deb889744ac00df831c10daa55189b5d121c855af7")},
+		{ 1'025u, TEXTVIEW("d00278ae47eb27b34faecf67b4fe263f82d5412916c1ffd97c8cb7fb814b8444")},
+		{ 2'048u, TEXTVIEW("e776b6028c7cd22a4d0ba182a8bf62205d2ef576467e838ed6f2529b85fba24a")},
+		{ 2'049u, TEXTVIEW("5f4d72f40d7a5f82b15ca2b2e44b1de3c2ef86c426c95c1af0b6879522563030")},
+		{ 3'072u, TEXTVIEW("b98cb0ff3623be03326b373de6b9095218513e64f1ee2edd2525c7ad1e5cffd2")},
+		{ 3'073u, TEXTVIEW("7124b49501012f81cc7f11ca069ec9226cecb8a2c850cfe644e327d22d3e1cd3")},
+		{ 4'096u, TEXTVIEW("015094013f57a5277b59d8475c0501042c0b642e531b0a1c8f58d2163229e969")},
+		{ 4'097u, TEXTVIEW("9b4052b38f1c5fc8b1f9ff7ac7b27cd242487b3d890d15c96a1c25b8aa0fb995")},
+		{ 5'120u, TEXTVIEW("9cadc15fed8b5d854562b26a9536d9707cadeda9b143978f319ab34230535833")},
+		{ 5'121u, TEXTVIEW("628bd2cb2004694adaab7bbd778a25df25c47b9d4155a55f8fbd79f2fe154cff")},
+		{ 6'144u, TEXTVIEW("3e2e5b74e048f3add6d21faab3f83aa44d3b2278afb83b80b3c35164ebeca205")},
+		{ 6'145u, TEXTVIEW("f1323a8631446cc50536a9f705ee5cb619424d46887f3c376c695b70e0f0507f")},
+		{ 7'168u, TEXTVIEW("61da957ec2499a95d6b8023e2b0e604ec7f6b50e80a9678b89d2628e99ada77a")},
+		{ 7'169u, TEXTVIEW("a003fc7a51754a9b3c7fae0367ab3d782dccf28855a03d435f8cfe74605e7817")},
+		{ 8'192u, TEXTVIEW("aae792484c8efe4f19e2ca7d371d8c467ffb10748d8a5a1ae579948f718a2a63")},
+		{ 8'193u, TEXTVIEW("bab6c09cb8ce8cf459261398d2e7aef35700bf488116ceb94a36d0f5f1b7bc3b")},
+		{16'384u, TEXTVIEW("f875d6646de28985646f34ee13be9a576fd515f76b5b0a26bb324735041ddde4")},
+		{31'744u, TEXTVIEW("62b6960e1a44bcc1eb1a611a8d6235b6b4b78f32e7abc4fb4c6cdcce94895c47")},
 	};
 
+	constexpr uint32 MaxInputLength = 31'744u;
+
 	uint8 Input[MaxInputLength];
 	for (uint32 Index = 0; Index < MaxInputLength; ++Index)
 	{
@@ -49,11 +51,11 @@ TEST_CASE_METHOD(FAutomationTestFixture, "Core::Hash::FBlake3::Smoke Test", "[Co
 	for (const FBlake3TestCase& TestCase : TestCases)
 	{
 		REQUIRE(InputIndex <= TestCase.InputLength);
+		REQUIRE(TestCase.InputLength <= MaxInputLength);
 		Hash.Update(Input + InputIndex, TestCase.InputLength - InputIndex);
 		InputIndex = TestCase.InputLength;
 		TStringBuilder<70> HashHex;
 		HashHex << Hash.Finalize();
-		INFO(FString::Printf(TEXT("BLAKE3(%u)"), TestCase.InputLength));
-		CHECK(FStringView(HashHex)==TestCase.Hash);
+		CHECK(HashHex.ToView() == TestCase.Hash);
 	}
 }
diff --git a/Engine/Source/Runtime/Core/Tests/Misc/StringBuilderTest.cpp b/Engine/Source/Runtime/Core/Tests/Misc/StringBuilderTest.cpp
--- a/Engine/Source/Runtime/Core/Tests/Misc/StringBuilderTest.cpp
+++ b/Engine/Source/Runtime/Core/Tests/Misc/StringBuilderTest.cpp
@@ -1,134 +1,186 @@
 // Copyright Epic Games, Inc. All Rights Reserved.
-#if 0
+
 #include "Misc/StringBuilder.h"
+
 #include "Containers/UnrealString.h"
 #include "String/Find.h"
 #include "TestHarness.h"
 
-static_assert(TIsSame<typename FStringBuilderBase::ElementType, TCHAR>::Value, "FStringBuilderBase must use TCHAR.");
-static_assert(TIsSame<typename FAnsiStringBuilderBase::ElementType, ANSICHAR>::Value, "FAnsiStringBuilderBase must use ANSICHAR.");
-static_assert(TIsSame<typename FWideStringBuilderBase::ElementType, WIDECHAR>::Value, "FWideStringBuilderBase must use WIDECHAR.");
+TEST_CASE("Core::String::StringBuilder", "[Core][String][Smoke]")
+{
+	SECTION("Static")
+	{
+		STATIC_REQUIRE(TIsSame<typename FStringBuilderBase::ElementType, TCHAR>::Value);
+		STATIC_REQUIRE(TIsSame<typename FAnsiStringBuilderBase::ElementType, ANSICHAR>::Value);
+		STATIC_REQUIRE(TIsSame<typename FWideStringBuilderBase::ElementType, WIDECHAR>::Value);
 
-static_assert(TIsSame<FStringBuilderBase, TStringBuilderBase<TCHAR>>::Value, "FStringBuilderBase must be the same as TStringBuilderBase<TCHAR>.");
-static_assert(TIsSame<FAnsiStringBuilderBase, TStringBuilderBase<ANSICHAR>>::Value, "FAnsiStringBuilderBase must be the same as TStringBuilderBase<ANSICHAR>.");
-static_assert(TIsSame<FWideStringBuilderBase, TStringBuilderBase<WIDECHAR>>::Value, "FWideStringBuilderBase must be the same as TStringBuilderBase<WIDECHAR>.");
+		STATIC_REQUIRE(TIsSame<FStringBuilderBase, TStringBuilderBase<TCHAR>>::Value);
+		STATIC_REQUIRE(TIsSame<FAnsiStringBuilderBase, TStringBuilderBase<ANSICHAR>>::Value);
+		STATIC_REQUIRE(TIsSame<FWideStringBuilderBase, TStringBuilderBase<WIDECHAR>>::Value);
 
-static_assert(TIsContiguousContainer<FStringBuilderBase>::Value, "FStringBuilderBase must be a contiguous container.");
-static_assert(TIsContiguousContainer<FAnsiStringBuilderBase>::Value, "FAnsiStringBuilderBase must be a contiguous container.");
-static_assert(TIsContiguousContainer<FWideStringBuilderBase>::Value, "FWideStringBuilderBase must be a contiguous container.");
+		STATIC_REQUIRE(TIsContiguousContainer<FStringBuilderBase>::Value);
+		STATIC_REQUIRE(TIsContiguousContainer<FAnsiStringBuilderBase>::Value);
+		STATIC_REQUIRE(TIsContiguousContainer<FWideStringBuilderBase>::Value);
 
-static_assert(TIsContiguousContainer<TStringBuilder<128>>::Value, "TStringBuilder<N> must be a contiguous container.");
-static_assert(TIsContiguousContainer<TAnsiStringBuilder<128>>::Value, "TAnsiStringBuilder<N> must be a contiguous container.");
-static_assert(TIsContiguousContainer<TWideStringBuilder<128>>::Value, "TWideStringBuilder<N> must be a contiguous container.");
+		STATIC_REQUIRE(TIsContiguousContainer<TStringBuilder<128>>::Value);
+		STATIC_REQUIRE(TIsContiguousContainer<TAnsiStringBuilder<128>>::Value);
+		STATIC_REQUIRE(TIsContiguousContainer<TWideStringBuilder<128>>::Value);
+	}
 
-TEST_CASE("Core::Misc::FStringBuilderBase::Append String", "[Core][Misc][Smoke]")
-{
-	// Empty Base
+	SECTION("Empty Base")
 	{
 		FStringBuilderBase Builder;
-		TestEqual(TEXT("Empty StringBuilderBase Len"), Builder.Len(), 0);
-		TestEqual(TEXT("Empty StringBuilderBase ToString"), Builder.ToString(), TEXT(""));
-		Builder << TEXT('A');
-		TestEqual(TEXT("Append Char to StringBuilderBase"), Builder.ToString(), TEXT("A"));
+		CHECK(Builder.Len() == 0);
+		CHECK(Builder == TEXTVIEW(""));
 	}
 
-	// Empty With Buffer
+	SECTION("Empty with Buffer")
 	{
-		TStringBuilder<1024> Builder;
-		TestEqual(TEXT("Empty StringBuilderWithBuffer Len"), Builder.Len(), 0);
-		TestEqual(TEXT("Empty StringBuilderWithBuffer ToString"), Builder.ToString(), TEXT(""));
+		TStringBuilder<16> Builder;
+		CHECK(Builder.Len() == 0);
+		CHECK(Builder == TEXTVIEW(""));
 	}
 
-	// Append Char
+	SECTION("Append Char to Base")
+	{
+		FStringBuilderBase Builder;
+		Builder.AppendChar(TEXT('A'));
+		CHECK(Builder.Len() == 1);
+		CHECK(Builder == TEXTVIEW("A"));
+	}
+
+	SECTION("Append Char")
 	{
 		TStringBuilder<7> Builder;
 		Builder << TEXT('A') << TEXT('B') << TEXT('C');
 		Builder << 'D' << 'E' << 'F';
-		TestEqual(TEXT("Append Char"), FStringView(Builder), TEXTVIEW("ABCDEF"));
+		Builder.AppendChar('G').AppendChar('H').AppendChar('I');
+		CHECK(Builder == TEXTVIEW("ABCDEFGHI"));
+	}
 
-		TAnsiStringBuilder<4> AnsiBuilder;
-		AnsiBuilder << 'A' << 'B' << 'C';
-		TestEqual(TEXT("Append AnsiChar"), FAnsiStringView(AnsiBuilder), ANSITEXTVIEW("ABC"));
+	SECTION("Append Char to ANSI")
+	{
+		TAnsiStringBuilder<4> Builder;
+		Builder << 'A' << 'B' << 'C';
+		CHECK(Builder == ANSITEXTVIEW("ABC"));
 	}
 
-	// Append C String
+	SECTION("Append C String")
 	{
 		TStringBuilder<7> Builder;
 		Builder << TEXT("ABC");
 		Builder << "DEF";
-		TestEqual(TEXT("Append C String"), FStringView(Builder), TEXTVIEW("ABCDEF"));
+		CHECK(Builder == TEXTVIEW("ABCDEF"));
+	}
 
-		TAnsiStringBuilder<4> AnsiBuilder;
-		AnsiBuilder << "ABC";
-		TestEqual(TEXT("Append Ansi C String"), FAnsiStringView(AnsiBuilder), ANSITEXTVIEW("ABC"));
+	SECTION("Append C String ANSI")
+	{
+		TAnsiStringBuilder<4> Builder;
+		Builder << "ABC";
+		CHECK(Builder == ANSITEXTVIEW("ABC"));
 	}
 
-	// Append FStringView
+	SECTION("Append FStringView")
 	{
 		TStringBuilder<7> Builder;
 		Builder << TEXTVIEW("ABC");
-		Builder << "DEF"_ASV;
-		TestEqual(TEXT("Append FStringView"), FStringView(Builder), TEXTVIEW("ABCDEF"));
+		Builder << ANSITEXTVIEW("DEF");
+		CHECK(Builder == TEXTVIEW("ABCDEF"));
+	}
 
-		TAnsiStringBuilder<4> AnsiBuilder;
-		AnsiBuilder << "ABC"_ASV;
-		TestEqual(TEXT("Append FAnsiStringView"), FAnsiStringView(AnsiBuilder), "ABC"_ASV);
+	SECTION("Append FAnsiStringView ANSI")
+	{
+		TAnsiStringBuilder<4> Builder;
+		Builder << ANSITEXTVIEW("ABC");
+		CHECK(Builder == ANSITEXTVIEW("ABC"));
 	}
 
-	// Append FStringBuilderBase
+	SECTION("Append FStringBuilderBase")
 	{
 		TStringBuilder<4> Builder;
 		Builder << TEXT("ABC");
 		TStringBuilder<4> BuilderCopy;
 		BuilderCopy << Builder;
-		TestEqual(TEXT("Append FStringBuilderBase"), FStringView(BuilderCopy), TEXTVIEW("ABC"));
+		CHECK(BuilderCopy == TEXTVIEW("ABC"));
+	}
 
-		TAnsiStringBuilder<4> AnsiBuilder;
-		AnsiBuilder << "ABC";
-		TAnsiStringBuilder<4> AnsiBuilderCopy;
-		AnsiBuilderCopy << AnsiBuilder;
-		TestEqual(TEXT("Append FAnsiStringBuilderBase"), FAnsiStringView(AnsiBuilderCopy), "ABC"_ASV);
+	SECTION("Append FStringBuilderBase ANSI")
+	{
+		TAnsiStringBuilder<4> Builder;
+		Builder << "ABC";
+		TAnsiStringBuilder<4> BuilderCopy;
+		BuilderCopy << Builder;
+		CHECK(BuilderCopy == ANSITEXTVIEW("ABC"));
 	}
 
-	// Append FString
+	SECTION("Append FString")
 	{
 		TStringBuilder<4> Builder;
 		Builder << FString(TEXT("ABC"));
-		TestEqual(TEXT("Append FString"), FStringView(Builder), TEXTVIEW("ABC"));
+		CHECK(Builder == TEXTVIEW("ABC"));
 	}
 
-	// Append Char Array
+	SECTION("Append Char Array")
 	{
 		const auto& String = TEXT("ABC");
 		TStringBuilder<4> Builder;
 		Builder << String;
-		TestEqual(TEXT("Append Char Array"), FStringView(Builder), TEXTVIEW("ABC"));
-
-		const ANSICHAR AnsiString[16] = "ABC";
-		TAnsiStringBuilder<4> AnsiBuilder;
-		AnsiBuilder << AnsiString;
-		TestEqual(TEXT("Append Char Array"), FAnsiStringView(AnsiBuilder), "ABC"_ASV);
+		CHECK(Builder == TEXTVIEW("ABC"));
 	}
 
-
-	// Simple ReplaceAt
+	SECTION("Append Char Array ANSI")
 	{
+		const ANSICHAR String[16] = "ABC";
 		TAnsiStringBuilder<4> Builder;
-		Builder.ReplaceAt(0, 0, FAnsiStringView());
-		TestEqual(TEXT("Replace nothing empty"), Builder.ToString(), "");
+		Builder << String;
+		CHECK(Builder == ANSITEXTVIEW("ABC"));
+	}
 
-		Builder << 'a';
+	SECTION("Simple ReplaceAt")
+	{
+		TAnsiStringBuilder<4> Builder;
+		Builder.ReplaceAt(0, 0, ANSITEXTVIEW(""));
+		CHECK(Builder == ANSITEXTVIEW(""));
 
-		Builder.ReplaceAt(0, 0, FAnsiStringView());
-		TestEqual(TEXT("Replace nothing non-empty"), Builder.ToString(), "a");
+		Builder.AppendChar('a');
+		
+		Builder.ReplaceAt(0, 0, ANSITEXTVIEW(""));
+		CHECK(Builder == ANSITEXTVIEW("a"));
 
-		Builder.ReplaceAt(0, 1, FAnsiStringView("b"));
-		TestEqual(TEXT("Replace single char"), Builder.ToString(), "b");
+		Builder.ReplaceAt(0, 1, ANSITEXTVIEW("b"));
+		CHECK(Builder == ANSITEXTVIEW("b"));
 	}
 
-	// Advanced ReplaceAt
-	auto TestReplace = [&](FAnsiStringView Original, FAnsiStringView SearchFor, FAnsiStringView ReplaceWith, FAnsiStringView Expected)
+	SECTION("Advanced ReplaceAt")
 	{
+		const auto [Original, SearchFor, ReplaceWith, Expected] = GENERATE(table<FAnsiStringView, FAnsiStringView, FAnsiStringView, FAnsiStringView>(
+		{
+			// Test single character erase
+			{".foo", ".", "", "foo"},
+			{"f.oo", ".", "", "foo"},
+			{"foo.", ".", "", "foo"},
+		
+			// Test multi character erase
+			{"FooBar", "Bar", "", "Foo"},
+			{"FooBar", "Foo", "", "Bar"},
+			{"FooBar", "Foo", "fOOO", "fOOOBar"},
+
+			// Test replace everything
+			{"Foo", "Foo", "", ""},
+			{"Foo", "Foo", "Bar", "Bar"},
+			{"Foo", "Foo", "0123456789", "0123456789"},
+
+			// Test expanding replace
+			{".foo", ".", "<dot>", "<dot>foo"},
+			{"foo.", ".", "<dot>", "foo<dot>"},
+			{"f.oo", ".", "<dot>", "f<dot>oo"},
+
+			// Test shrinking replace
+			{"aabbcc", "aa", "A", "Abbcc"},
+			{"aabbcc", "bb", "B", "aaBcc"},
+			{"aabbcc", "cc", "C", "aabbC"},
+		}));
+
 		int32 ReplacePos = UE::String::FindFirst(Original, SearchFor);
 		check(ReplacePos != INDEX_NONE);
 
@@ -136,101 +188,94 @@ TEST_CASE("Core::Misc::FStringBuilderBase::Append String", "[Core][Misc][Smoke]"
 		Builder << Original;
 		Builder.ReplaceAt(ReplacePos, SearchFor.Len(), ReplaceWith);
 
-		TestEqual(TEXT("Replace"), FAnsiStringView(Builder), Expected);
-	};
-
-	// Test single character erase
-	TestReplace(".foo", ".", "", "foo");
-	TestReplace("f.oo", ".", "", "foo");
-	TestReplace("foo.", ".", "", "foo");
-
-	// Test multi character erase
-	TestReplace("FooBar", "Bar", "", "Foo");
-	TestReplace("FooBar", "Foo", "", "Bar");
-	TestReplace("FooBar", "Foo", "fOOO", "fOOOBar");
-
-	// Test replace everything
-	TestReplace("Foo", "Foo", "", "");
-	TestReplace("Foo", "Foo", "Bar", "Bar");
-	TestReplace("Foo", "Foo", "0123456789", "0123456789");
-
-	// Test expanding replace
-	TestReplace(".foo", ".", "<dot>", "<dot>foo");
-	TestReplace("foo.", ".", "<dot>", "foo<dot>");
-	TestReplace("f.oo", ".", "<dot>", "f<dot>oo");
-
-	// Test shrinking replace
-	TestReplace("aabbcc", "aa", "A", "Abbcc");
-	TestReplace("aabbcc", "bb", "B", "aaBcc");
-	TestReplace("aabbcc", "cc", "C", "aabbC");
+		CHECK(Builder == Expected);
+	}
 
-	// Prepend
+	SECTION("Prepend")
 	{
 		TAnsiStringBuilder<4> Builder;
-		Builder.Prepend("");
-		TestEqual(TEXT("Prepend nothing to empty"), Builder.Len(), 0);
 
+		// Prepend nothing to empty
+		Builder.Prepend("");
+		CHECK(Builder.Len() == 0);
+		
+		// Prepend single characer
 		Builder.Prepend("e");
-		TestEqual(TEXT("Prepend single characer"), Builder.ToString(), "e");
-
+		CHECK(Builder == ANSITEXTVIEW("e"));
+		
+		// Prepend substring
 		Builder.Prepend("abcd");
-		TestEqual(TEXT("Prepend substring"), Builder.ToString(), "abcde");
+		CHECK(Builder == ANSITEXTVIEW("abcde"));
 
+		// Prepend nothing to non-empty
 		Builder.Prepend("");
-		TestEqual(TEXT("Prepend nothing to non-empty"), Builder.ToString(), "abcde");
+		CHECK(Builder == ANSITEXTVIEW("abcde"));
 	}
-
-	// InsertAt
+	
+	SECTION("InsertAt")
 	{
 		TAnsiStringBuilder<4> Builder;
 		Builder.InsertAt(0, "");
-		TestEqual(TEXT("Insert nothing to empty"), Builder.Len(), 0);
 
-		Builder.InsertAt(0, "d");
-		TestEqual(TEXT("Insert first char"), Builder.ToString(), "d");
+		// Insert nothing to empty
+		CHECK(Builder.Len() == 0);
 
+		// Insert first char
+
+		Builder.InsertAt(0, "d");
+		CHECK(Builder == ANSITEXTVIEW("d"));
+		
+		// Insert single char
 		Builder.InsertAt(0, "c");
 		Builder.InsertAt(0, "a");
 		Builder.InsertAt(1, "b");
 		Builder.InsertAt(4, "e");
-		TestEqual(TEXT("Insert single char"), Builder.ToString(), "abcde");
-
+		CHECK(Builder == ANSITEXTVIEW("abcde"));
+		
+		// Insert substrings
 		Builder.InsertAt(3, "__");
 		Builder.InsertAt(0, "__");
 		Builder.InsertAt(Builder.Len(), "__");
-		TestEqual(TEXT("Insert substrings"), Builder.ToString(), "__abc__de__");
+		CHECK(Builder == ANSITEXTVIEW("__abc__de__"));
 
+		// Insert nothing
 		Builder.InsertAt(Builder.Len(), "");
-		TestEqual(TEXT("Insert nothing"), Builder.ToString(), "__abc__de__");
+		CHECK(Builder == ANSITEXTVIEW("__abc__de__"));
 	}
 
-	// RemoveAt
+	SECTION("RemoveAt")
 	{
 		TAnsiStringBuilder<4> Builder;
 		Builder << "0123456789";
+
+		// Remove nothing
 		Builder.RemoveAt(0, 0);
 		Builder.RemoveAt(Builder.Len(), 0);
 		Builder.RemoveAt(Builder.Len() / 2, 0);
-		TestEqual(TEXT("Remove nothing"), Builder.ToString(), "0123456789");
-
+		CHECK(Builder == ANSITEXTVIEW("0123456789"));
+		
+		// Remove last char
 		Builder.RemoveAt(Builder.Len() - 1, 1);
-		TestEqual(TEXT("Remove last char"), Builder.ToString(), "012345678");
-
+		CHECK(Builder == ANSITEXTVIEW("012345678"));
+		
+		// Remove first char
 		Builder.RemoveAt(0, 1);
-		TestEqual(TEXT("Remove first char"), Builder.ToString(), "12345678");
-
+		CHECK(Builder == ANSITEXTVIEW("12345678"));
+		
+		// Remove middle
 		Builder.RemoveAt(4, 2);
-		TestEqual(TEXT("Remove middle"), Builder.ToString(), "123478");
+		CHECK(Builder == ANSITEXTVIEW("123478"));
 
+		// Remove end
 		Builder.RemoveAt(4, 2);
-		TestEqual(TEXT("Remove end"), Builder.ToString(), "1234");
-
+		CHECK(Builder == ANSITEXTVIEW("1234"));
+		
+		// Remove start
 		Builder.RemoveAt(0, 2);
-		TestEqual(TEXT("Remove start"), Builder.ToString(), "34");
+		CHECK(Builder == ANSITEXTVIEW("34"));
 
+		// Remove start
 		Builder.RemoveAt(0, 2);
-		TestEqual(TEXT("Remove start"), Builder.ToString(), "");
+		CHECK(Builder == ANSITEXTVIEW(""));
 	}
 }
-
-#endif
diff --git a/Engine/Source/Runtime/Core/Tests/String/EscapeTest.cpp b/Engine/Source/Runtime/Core/Tests/String/EscapeTest.cpp
--- a/Engine/Source/Runtime/Core/Tests/String/EscapeTest.cpp
+++ b/Engine/Source/Runtime/Core/Tests/String/EscapeTest.cpp
@@ -9,7 +9,7 @@
 namespace UE::String
 {
 
-TEST_CASE("Core::String::Misc::Escape", "[Core][String][Smoke]")
+TEST_CASE("Core::String::Escape", "[Core][String][Smoke]")
 {
 	SECTION("Escape")
 	{
diff --git a/Engine/Source/Runtime/Core/Tests/String/FindTest.cpp b/Engine/Source/Runtime/Core/Tests/String/FindTest.cpp
--- a/Engine/Source/Runtime/Core/Tests/String/FindTest.cpp
+++ b/Engine/Source/Runtime/Core/Tests/String/FindTest.cpp
@@ -1,226 +1,224 @@
 // Copyright Epic Games, Inc. All Rights Reserved.
 
 #include "String/Find.h"
-#include "TestHarness.h"
-
-TEST_CASE_METHOD(FAutomationTestFixture, "Core::String::Misc::FindFirst", "[Core][String][Smoke]")
-{
-	CHECK_EQUAL(UE::String::FindFirst(TEXT("AbCABCAbCABC"), TEXT("A")), 0);
-	CHECK_EQUAL(UE::String::FindFirst(TEXT("AbCABCAbCABC"), TEXT("a"), ESearchCase::IgnoreCase), 0);
-	CHECK_EQUAL(UE::String::FindFirst(TEXT("AbCABCAbCABC"), TEXT("b")), 1);
-	CHECK_EQUAL(UE::String::FindFirst(TEXT("AbCABCAbCABC"), TEXT("B")), 4);
-	CHECK_EQUAL(UE::String::FindFirst(TEXT("AbCABCAbCABC"), TEXT("B"), ESearchCase::IgnoreCase), 1);
-	CHECK_EQUAL(UE::String::FindFirst(TEXT("AbCABCAbCABC"), TEXT("a")), INDEX_NONE);
-	CHECK_EQUAL(UE::String::FindFirst(TEXT("AbCABCAbCABC"), TEXT("D"), ESearchCase::IgnoreCase), INDEX_NONE);
-	CHECK_EQUAL(UE::String::FindFirst(TEXT("AbCABCAbCABD"), TEXT("D")), 11);
-	CHECK_EQUAL(UE::String::FindFirst(TEXT("AbCABCAbCABD"), TEXT("d"), ESearchCase::IgnoreCase), 11);
-
-	CHECK_EQUAL(UE::String::FindFirst(TEXT("AbCABCAbCABC"), TEXT("AbC")), 0);
-	CHECK_EQUAL(UE::String::FindFirst(TEXT("AbCABCAbCABC"), TEXT("ABC")), 3);
-	CHECK_EQUAL(UE::String::FindFirst(TEXT("AbCABCAbCABC"), TEXT("Bc"), ESearchCase::IgnoreCase), 1);
-	CHECK_EQUAL(UE::String::FindFirst(TEXT("AbCABCAbCABC"), TEXT("ab")), INDEX_NONE);
-	CHECK_EQUAL(UE::String::FindFirst(TEXT("AbCABCAbCABC"), TEXT("CD"), ESearchCase::IgnoreCase), INDEX_NONE);
-	CHECK_EQUAL(UE::String::FindFirst(TEXT("AbCABCAbCABD"), TEXT("BD")), 10);
-	CHECK_EQUAL(UE::String::FindFirst(TEXT("AbCABCAbCABD"), TEXT("Bd"), ESearchCase::IgnoreCase), 10);
-
-	CHECK_EQUAL(UE::String::FindFirst(TEXT(""), TEXT("A")), INDEX_NONE);
-	CHECK_EQUAL(UE::String::FindFirst(TEXT("A"), TEXT("A")), 0);
-	CHECK_EQUAL(UE::String::FindFirst(TEXT("A"), TEXT("A"), ESearchCase::IgnoreCase), 0);
-	CHECK_EQUAL(UE::String::FindFirst(TEXT("ABC"), TEXT("ABC")), 0);
-	CHECK_EQUAL(UE::String::FindFirst(TEXT("ABC"), TEXT("abc"), ESearchCase::IgnoreCase), 0);
-	CHECK_EQUAL(UE::String::FindFirst(TEXT("AB"), TEXT("ABC")), INDEX_NONE);
-
-	CHECK_EQUAL(UE::String::FindFirst("AbCABCAbCABC", "ABC"), 3);
-
-	CHECK_EQUAL(UE::String::FindFirst(FStringView(nullptr, 0), TEXT("SearchTerm")), INDEX_NONE);
-	CHECK_EQUAL(UE::String::FindFirst(FStringView(), TEXT("SearchTerm")), INDEX_NONE);
-	CHECK_EQUAL(UE::String::FindFirst(FString(), TEXT("SearchTerm")), INDEX_NONE);
-}
-
-TEST_CASE_METHOD(FAutomationTestFixture, "Core::String::Misc::FindLast", "[Core][String][Smoke]")
-{
-	CHECK_EQUAL(UE::String::FindLast(TEXT("AbCABCAbCABC"), TEXT("b")), 7);
-	CHECK_EQUAL(UE::String::FindLast(TEXT("AbCABCAbCABC"), TEXT("B")), 10);
-	CHECK_EQUAL(UE::String::FindLast(TEXT("AbCABCAbCABC"), TEXT("b"), ESearchCase::IgnoreCase), 10);
-	CHECK_EQUAL(UE::String::FindLast(TEXT("AbCABCAbCABC"), TEXT("a")), INDEX_NONE);
-	CHECK_EQUAL(UE::String::FindLast(TEXT("AbCABCAbCABC"), TEXT("D"), ESearchCase::IgnoreCase), INDEX_NONE);
-	CHECK_EQUAL(UE::String::FindLast(TEXT("AbCABCAbCABD"), TEXT("D")), 11);
-	CHECK_EQUAL(UE::String::FindLast(TEXT("AbCABCAbCABD"), TEXT("d"), ESearchCase::IgnoreCase), 11);
-
-	CHECK_EQUAL(UE::String::FindLast(TEXT("AbCABCAbCABC"), TEXT("AbC")), 6);
-	CHECK_EQUAL(UE::String::FindLast(TEXT("AbCABCAbCABC"), TEXT("ABC")), 9);
-	CHECK_EQUAL(UE::String::FindLast(TEXT("AbCABCAbCABC"), TEXT("Bc"), ESearchCase::IgnoreCase), 10);
-	CHECK_EQUAL(UE::String::FindLast(TEXT("AbCABCAbCABC"), TEXT("ab")), INDEX_NONE);
-	CHECK_EQUAL(UE::String::FindLast(TEXT("AbCABCAbCABC"), TEXT("CD"), ESearchCase::IgnoreCase), INDEX_NONE);
-	CHECK_EQUAL(UE::String::FindLast(TEXT("AbCABCAbCABC"), TEXT("BC")), 10);
-	CHECK_EQUAL(UE::String::FindLast(TEXT("AbCABCAbCABC"), TEXT("Bc"), ESearchCase::IgnoreCase), 10);
-
-	CHECK_EQUAL(UE::String::FindLast(TEXT(""), TEXT("A")), INDEX_NONE);
-	CHECK_EQUAL(UE::String::FindLast(TEXT("A"), TEXT("A")), 0);
-	CHECK_EQUAL(UE::String::FindLast(TEXT("A"), TEXT("A"), ESearchCase::IgnoreCase), 0);
-	CHECK_EQUAL(UE::String::FindLast(TEXT("ABC"), TEXT("ABC")), 0);
-	CHECK_EQUAL(UE::String::FindLast(TEXT("ABC"), TEXT("abc"), ESearchCase::IgnoreCase), 0);
-	CHECK_EQUAL(UE::String::FindLast(TEXT("AB"), TEXT("ABC")), INDEX_NONE);
-
-	CHECK_EQUAL(UE::String::FindLast("AbCABCAbCABC", "ABC"), 9);
-
-	CHECK_EQUAL(UE::String::FindLast(FStringView(nullptr, 0), TEXT("SearchTerm")), INDEX_NONE);
-	CHECK_EQUAL(UE::String::FindLast(FStringView(), TEXT("SearchTerm")), INDEX_NONE);
-	CHECK_EQUAL(UE::String::FindLast(FString(), TEXT("SearchTerm")), INDEX_NONE);
-}
-
-
-TEST_CASE_METHOD(FAutomationTestFixture, "Core::String::Misc::FindFirstOfAny", "[Core][String][Smoke]")
-{
-	CHECK_EQUAL(UE::String::FindFirstOfAny(TEXT("AbCABCAbcABC"), {TEXT("A"), TEXT("B")}), 0);
-	CHECK_EQUAL(UE::String::FindFirstOfAny(TEXT("AbCABCAbcABC"), {TEXT("a"), TEXT("B")}, ESearchCase::IgnoreCase), 0);
-	CHECK_EQUAL(UE::String::FindFirstOfAny(TEXT("AbCABCAbcABC"), {TEXT("b")}), 1);
-	CHECK_EQUAL(UE::String::FindFirstOfAny(TEXT("AbCABCAbcABC"), {TEXT("c"), TEXT("B")}), 4);
-	CHECK_EQUAL(UE::String::FindFirstOfAny(TEXT("AbCABCAbcABC"), {TEXT("c"), TEXT("B")}, ESearchCase::IgnoreCase), 1);
-	CHECK_EQUAL(UE::String::FindFirstOfAny(TEXT("AbCABCAbcABC"), {TEXT("D"), TEXT("a")}), INDEX_NONE);
-	CHECK_EQUAL(UE::String::FindFirstOfAny(TEXT("AbCABCAbcABC"), {TEXT("E"), TEXT("D")}, ESearchCase::IgnoreCase), INDEX_NONE);
-	CHECK_EQUAL(UE::String::FindFirstOfAny(TEXT("AbCABCAbcABD"), {TEXT("E"), TEXT("D")}), 11);
-	CHECK_EQUAL(UE::String::FindFirstOfAny(TEXT("AbCABCAbcABD"), {TEXT("E"), TEXT("d")}, ESearchCase::IgnoreCase), 11);
-
-	CHECK_EQUAL(UE::String::FindFirstOfAny(TEXT("AbCABCAbCABC"), {TEXT("AbC")}), 0);
-	CHECK_EQUAL(UE::String::FindFirstOfAny(TEXT("AbCABCAbCABC"), {TEXT("CABc"), TEXT("ABC")}), 3);
-	CHECK_EQUAL(UE::String::FindFirstOfAny(TEXT("AbCABCAbCABC"), {TEXT("ABD"), TEXT("Bc")}, ESearchCase::IgnoreCase), 1);
-	CHECK_EQUAL(UE::String::FindFirstOfAny(TEXT("AbCABCAbCABC"), {TEXT("bc"), TEXT("ab")}), INDEX_NONE);
-	CHECK_EQUAL(UE::String::FindFirstOfAny(TEXT("AbCABCAbCABC"), {TEXT("DA"), TEXT("CD")}, ESearchCase::IgnoreCase), INDEX_NONE);
-	CHECK_EQUAL(UE::String::FindFirstOfAny(TEXT("AbCABCAbCABD"), {TEXT("BD"), TEXT("CABB")}), 10);
-	CHECK_EQUAL(UE::String::FindFirstOfAny(TEXT("AbCABCAbCABD"), {TEXT("Bd"), TEXT("CABB")}, ESearchCase::IgnoreCase), 10);
-
-	CHECK_EQUAL(UE::String::FindFirstOfAny(TEXT(""), {TEXT("A"), TEXT("B")}), INDEX_NONE);
-	CHECK_EQUAL(UE::String::FindFirstOfAny(TEXT("A"), {TEXT("A"), TEXT("B")}), 0);
-	CHECK_EQUAL(UE::String::FindFirstOfAny(TEXT("A"), {TEXT("A"), TEXT("B")}, ESearchCase::IgnoreCase), 0);
-	CHECK_EQUAL(UE::String::FindFirstOfAny(TEXT("ABC"), {TEXT("ABC"), TEXT("BC")}), 0);
-	CHECK_EQUAL(UE::String::FindFirstOfAny(TEXT("ABC"), {TEXT("abc"), TEXT("bc")}, ESearchCase::IgnoreCase), 0);
-	CHECK_EQUAL(UE::String::FindFirstOfAny(TEXT("AB"), {TEXT("ABC"), TEXT("ABD")}), INDEX_NONE);
-
-	CHECK_EQUAL(UE::String::FindFirstOfAny("AbCABCAbCABC", {"CABc", "ABC"}), 3);
-
-	CHECK_EQUAL(UE::String::FindFirstOfAny(FStringView(nullptr, 0), {TEXT("ABC"), TEXT("ABD")}), INDEX_NONE);
-	CHECK_EQUAL(UE::String::FindFirstOfAny(FStringView(), {TEXT("ABC"), TEXT("ABD")}), INDEX_NONE);
-	CHECK_EQUAL(UE::String::FindFirstOfAny(FString(), {TEXT("ABC"), TEXT("ABD")}), INDEX_NONE);
-}
-
-
-TEST_CASE_METHOD(FAutomationTestFixture, "Core::String::Misc::FindLastOfAny", "[Core][String][Smoke]")
-{
-	CHECK_EQUAL(UE::String::FindLastOfAny(TEXT("AbCABCAbcABC"), {TEXT("a"), TEXT("b")}), 7);
-	CHECK_EQUAL(UE::String::FindLastOfAny(TEXT("AbCABCAbcABC"), {TEXT("a"), TEXT("b")}, ESearchCase::IgnoreCase), 10);
-	CHECK_EQUAL(UE::String::FindLastOfAny(TEXT("AbCABCAbcABC"), {TEXT("b")}), 7);
-	CHECK_EQUAL(UE::String::FindLastOfAny(TEXT("AbCABCAbcABC"), {TEXT("c"), TEXT("B")}), 10);
-	CHECK_EQUAL(UE::String::FindLastOfAny(TEXT("AbCABCAbcABC"), {TEXT("c"), TEXT("B")}, ESearchCase::IgnoreCase), 11);
-	CHECK_EQUAL(UE::String::FindLastOfAny(TEXT("AbCABCAbcABC"), {TEXT("D"), TEXT("a")}), INDEX_NONE);
-	CHECK_EQUAL(UE::String::FindLastOfAny(TEXT("AbCABCAbcABC"), {TEXT("E"), TEXT("D")}, ESearchCase::IgnoreCase), INDEX_NONE);
-	CHECK_EQUAL(UE::String::FindLastOfAny(TEXT("AbCABCAbcABD"), {TEXT("E"), TEXT("D")}), 11);
-	CHECK_EQUAL(UE::String::FindLastOfAny(TEXT("AbCABCAbcABD"), {TEXT("E"), TEXT("d")}, ESearchCase::IgnoreCase), 11);
-
-	CHECK_EQUAL(UE::String::FindLastOfAny(TEXT("AbCABCAbCABC"), {TEXT("AbC")}), 6);
-	CHECK_EQUAL(UE::String::FindLastOfAny(TEXT("AbCABCAbCABC"), {TEXT("CABc"), TEXT("ABC")}), 9);
-	CHECK_EQUAL(UE::String::FindLastOfAny(TEXT("AbCABCAbCABC"), {TEXT("ABD"), TEXT("Bc")}, ESearchCase::IgnoreCase), 10);
-	CHECK_EQUAL(UE::String::FindLastOfAny(TEXT("AbCABCAbCABC"), {TEXT("bc"), TEXT("ab")}), INDEX_NONE);
-	CHECK_EQUAL(UE::String::FindLastOfAny(TEXT("AbCABCAbCABC"), {TEXT("DA"), TEXT("CD")}, ESearchCase::IgnoreCase), INDEX_NONE);
-	CHECK_EQUAL(UE::String::FindLastOfAny(TEXT("AbCABCAbCABD"), {TEXT("BD"), TEXT("CABB")}), 10);
-	CHECK_EQUAL(UE::String::FindLastOfAny(TEXT("AbCABCAbCABD"), {TEXT("Bd"), TEXT("CABB")}, ESearchCase::IgnoreCase), 10);
-
-	CHECK_EQUAL(UE::String::FindLastOfAny(TEXT(""), {TEXT("A"), TEXT("B")}), INDEX_NONE);
-	CHECK_EQUAL(UE::String::FindLastOfAny(TEXT("A"), {TEXT("A"), TEXT("B")}), 0);
-	CHECK_EQUAL(UE::String::FindLastOfAny(TEXT("A"), {TEXT("A"), TEXT("B")}, ESearchCase::IgnoreCase), 0);
-	CHECK_EQUAL(UE::String::FindLastOfAny(TEXT("ABC"), {TEXT("ABC"), TEXT("BC")}), 1);
-	CHECK_EQUAL(UE::String::FindLastOfAny(TEXT("ABC"), {TEXT("abc"), TEXT("bc")}, ESearchCase::IgnoreCase), 1);
-	CHECK_EQUAL(UE::String::FindLastOfAny(TEXT("AB"), {TEXT("ABC"), TEXT("ABD")}), INDEX_NONE);
-
-	CHECK_EQUAL(UE::String::FindLastOfAny("AbCABCAbCABC", {"CABc", "ABC"}), 9);
-
-	CHECK_EQUAL(UE::String::FindLastOfAny(FStringView(nullptr, 0), { TEXT("ABC"), TEXT("ABD") }), INDEX_NONE);
-	CHECK_EQUAL(UE::String::FindLastOfAny(FStringView(), { TEXT("ABC"), TEXT("ABD") }), INDEX_NONE);
-	CHECK_EQUAL(UE::String::FindLastOfAny(FString(), { TEXT("ABC"), TEXT("ABD") }), INDEX_NONE);
-}
-
-
-TEST_CASE_METHOD(FAutomationTestFixture, "Core::String::Misc::FindFirstChar", "[Core][String][Smoke]")
-{
-	TEST_EQUAL(TEXT("FindFirstChar"), UE::String::FindFirstChar(TEXT("AbCABCAbCABC"), TEXT('b')), 1);
-	TEST_EQUAL(TEXT("FindFirstChar"), UE::String::FindFirstChar(TEXT("AbCABCAbCABC"), TEXT('B')), 4);
-	TEST_EQUAL(TEXT("FindFirstChar"), UE::String::FindFirstChar(TEXT("AbCABCAbCABC"), TEXT('B'), ESearchCase::IgnoreCase), 1);
-	TEST_EQUAL(TEXT("FindFirstChar"), UE::String::FindFirstChar(TEXT("AbCABCAbCABC"), TEXT('a')), INDEX_NONE);
-	TEST_EQUAL(TEXT("FindFirstChar"), UE::String::FindFirstChar(TEXT("AbCABCAbCABC"), TEXT('D'), ESearchCase::IgnoreCase), INDEX_NONE);
-	TEST_EQUAL(TEXT("FindFirstChar"), UE::String::FindFirstChar(TEXT("AbCABCAbCABD"), TEXT('D')), 11);
-	TEST_EQUAL(TEXT("FindFirstChar"), UE::String::FindFirstChar(TEXT("AbCABCAbCABD"), TEXT('d'), ESearchCase::IgnoreCase), 11);
-
-	TEST_EQUAL(TEXT("FindFirstChar"), UE::String::FindFirstChar(TEXT(""), TEXT('A')), INDEX_NONE);
-	TEST_EQUAL(TEXT("FindFirstChar"), UE::String::FindFirstChar(TEXT("A"), TEXT('A')), 0);
-	TEST_EQUAL(TEXT("FindFirstChar"), UE::String::FindFirstChar(TEXT("A"), TEXT('A'), ESearchCase::IgnoreCase), 0);
-
-	TEST_EQUAL(TEXT("FindFirstChar"), UE::String::FindFirstChar("AbCABCAbCABC", 'B'), 4);
-
-	TEST_EQUAL(TEXT("FindFirstChar"), UE::String::FindFirstChar(FStringView(nullptr, 0), TEXT('A')), INDEX_NONE);
-	TEST_EQUAL(TEXT("FindFirstChar"), UE::String::FindFirstChar(FStringView(), TEXT('A')), INDEX_NONE);
-	TEST_EQUAL(TEXT("FindFirstChar"), UE::String::FindFirstChar(FString(), TEXT('A')), INDEX_NONE);
-}
 
+#include "Containers/StringView.h"
+#include "TestHarness.h"
 
-TEST_CASE_METHOD(FAutomationTestFixture, "Core::String::Misc::FindLastChar", "[Core][String][Smoke]")
+namespace UE::String
 {
-	TEST_EQUAL(TEXT("FindLastChar"), UE::String::FindLastChar(TEXT("AbCABCAbCABC"), TEXT('b')), 7);
-	TEST_EQUAL(TEXT("FindLastChar"), UE::String::FindLastChar(TEXT("AbCABCAbCABC"), TEXT('B')), 10);
-	TEST_EQUAL(TEXT("FindLastChar"), UE::String::FindLastChar(TEXT("AbCABCAbCABC"), TEXT('b'), ESearchCase::IgnoreCase), 10);
-	TEST_EQUAL(TEXT("FindLastChar"), UE::String::FindLastChar(TEXT("AbCABCAbCABC"), TEXT('a')), INDEX_NONE);
-	TEST_EQUAL(TEXT("FindLastChar"), UE::String::FindLastChar(TEXT("AbCABCAbCABC"), TEXT('D'), ESearchCase::IgnoreCase), INDEX_NONE);
-	TEST_EQUAL(TEXT("FindLastChar"), UE::String::FindLastChar(TEXT("AbCABCAbCABD"), TEXT('D')), 11);
-	TEST_EQUAL(TEXT("FindLastChar"), UE::String::FindLastChar(TEXT("AbCABCAbCABD"), TEXT('d'), ESearchCase::IgnoreCase), 11);
-
-	TEST_EQUAL(TEXT("FindLastChar"), UE::String::FindLastChar(TEXT(""), TEXT('A')), INDEX_NONE);
-	TEST_EQUAL(TEXT("FindLastChar"), UE::String::FindLastChar(TEXT("A"), TEXT('A')), 0);
-	TEST_EQUAL(TEXT("FindLastChar"), UE::String::FindLastChar(TEXT("A"), TEXT('A'), ESearchCase::IgnoreCase), 0);
-
-	TEST_EQUAL(TEXT("FindLastChar"), UE::String::FindLastChar("AbCABCAbCABC", 'B'), 10);
-
-	TEST_EQUAL(TEXT("FindLastChar"), UE::String::FindLastChar(FStringView(nullptr, 0), TEXT('A')), INDEX_NONE);
-	TEST_EQUAL(TEXT("FindLastChar"), UE::String::FindLastChar(FStringView(), TEXT('A')), INDEX_NONE);
-	TEST_EQUAL(TEXT("FindLastChar"), UE::String::FindLastChar(FString(), TEXT('A')), INDEX_NONE);
-}
 
-TEST_CASE_METHOD(FAutomationTestFixture, "Core::String::Misc::FindFirstOfAnyChar", "[Core][String][Smoke]")
+TEST_CASE("Core::String::Find", "[Core][String][Smoke]")
 {
-	TEST_EQUAL(TEXT("FindFirstOfAnyChar"), UE::String::FindFirstOfAnyChar(TEXT("AbCABCAbcABC"), {TEXT('b')}), 1);
-	TEST_EQUAL(TEXT("FindFirstOfAnyChar"), UE::String::FindFirstOfAnyChar(TEXT("AbCABCAbcABC"), {TEXT('c'), TEXT('B')}), 4);
-	TEST_EQUAL(TEXT("FindFirstOfAnyChar"), UE::String::FindFirstOfAnyChar(TEXT("AbCABCAbcABC"), {TEXT('c'), TEXT('B')}, ESearchCase::IgnoreCase), 1);
-	TEST_EQUAL(TEXT("FindFirstOfAnyChar"), UE::String::FindFirstOfAnyChar(TEXT("AbCABCAbcABC"), {TEXT('D'), TEXT('a')}), INDEX_NONE);
-	TEST_EQUAL(TEXT("FindFirstOfAnyChar"), UE::String::FindFirstOfAnyChar(TEXT("AbCABCAbcABC"), {TEXT('E'), TEXT('D')}, ESearchCase::IgnoreCase), INDEX_NONE);
-	TEST_EQUAL(TEXT("FindFirstOfAnyChar"), UE::String::FindFirstOfAnyChar(TEXT("AbCABCAbcABD"), {TEXT('E'), TEXT('D')}), 11);
-	TEST_EQUAL(TEXT("FindFirstOfAnyChar"), UE::String::FindFirstOfAnyChar(TEXT("AbCABCAbcABD"), {TEXT('E'), TEXT('d')}, ESearchCase::IgnoreCase), 11);
-
-	TEST_EQUAL(TEXT("FindFirstOfAnyChar"), UE::String::FindFirstOfAnyChar(TEXT(""), {TEXT('A'), TEXT('B')}), INDEX_NONE);
-	TEST_EQUAL(TEXT("FindFirstOfAnyChar"), UE::String::FindFirstOfAnyChar(TEXT("A"), {TEXT('A'), TEXT('B')}), 0);
-	TEST_EQUAL(TEXT("FindFirstOfAnyChar"), UE::String::FindFirstOfAnyChar(TEXT("A"), {TEXT('A'), TEXT('B')}, ESearchCase::IgnoreCase), 0);
-
-	TEST_EQUAL(TEXT("FindFirstOfAnyChar"), UE::String::FindFirstOfAnyChar("AbCABCAbcABC", {'c', 'B'}), 4);
-
-	TEST_EQUAL(TEXT("FindFirstOfAnyChar"), UE::String::FindFirstOfAnyChar(FStringView(nullptr, 0), { TEXT('A'), TEXT('B') }), INDEX_NONE);
-	TEST_EQUAL(TEXT("FindFirstOfAnyChar"), UE::String::FindFirstOfAnyChar(FStringView(), { TEXT('A'), TEXT('B') }), INDEX_NONE);
-	TEST_EQUAL(TEXT("FindFirstOfAnyChar"), UE::String::FindFirstOfAnyChar(FString(), { TEXT('A'), TEXT('B') }), INDEX_NONE);
+	SECTION("FindFirst")
+	{
+		CHECK(FindFirst(TEXT("AbCABCAbCABC"), TEXT("A")) == 0);
+		CHECK(FindFirst(TEXT("AbCABCAbCABC"), TEXT("a"), ESearchCase::IgnoreCase) == 0);
+		CHECK(FindFirst(TEXT("AbCABCAbCABC"), TEXT("b")) == 1);
+		CHECK(FindFirst(TEXT("AbCABCAbCABC"), TEXT("B")) == 4);
+		CHECK(FindFirst(TEXT("AbCABCAbCABC"), TEXT("B"), ESearchCase::IgnoreCase) == 1);
+		CHECK(FindFirst(TEXT("AbCABCAbCABC"), TEXT("a")) == INDEX_NONE);
+		CHECK(FindFirst(TEXT("AbCABCAbCABC"), TEXT("D"), ESearchCase::IgnoreCase) == INDEX_NONE);
+		CHECK(FindFirst(TEXT("AbCABCAbCABD"), TEXT("D")) == 11);
+		CHECK(FindFirst(TEXT("AbCABCAbCABD"), TEXT("d"), ESearchCase::IgnoreCase) == 11);
+
+		CHECK(FindFirst(TEXT("AbCABCAbCABC"), TEXT("AbC")) == 0);
+		CHECK(FindFirst(TEXT("AbCABCAbCABC"), TEXT("ABC")) == 3);
+		CHECK(FindFirst(TEXT("AbCABCAbCABC"), TEXT("Bc"), ESearchCase::IgnoreCase) == 1);
+		CHECK(FindFirst(TEXT("AbCABCAbCABC"), TEXT("ab")) == INDEX_NONE);
+		CHECK(FindFirst(TEXT("AbCABCAbCABC"), TEXT("CD"), ESearchCase::IgnoreCase) == INDEX_NONE);
+		CHECK(FindFirst(TEXT("AbCABCAbCABD"), TEXT("BD")) == 10);
+		CHECK(FindFirst(TEXT("AbCABCAbCABD"), TEXT("Bd"), ESearchCase::IgnoreCase) == 10);
+
+		CHECK(FindFirst(TEXT(""), TEXT("A")) == INDEX_NONE);
+		CHECK(FindFirst(TEXT("A"), TEXT("A")) == 0);
+		CHECK(FindFirst(TEXT("A"), TEXT("A"), ESearchCase::IgnoreCase) == 0);
+		CHECK(FindFirst(TEXT("ABC"), TEXT("ABC")) == 0);
+		CHECK(FindFirst(TEXT("ABC"), TEXT("abc"), ESearchCase::IgnoreCase) == 0);
+		CHECK(FindFirst(TEXT("AB"), TEXT("ABC")) == INDEX_NONE);
+
+		CHECK(FindFirst("AbCABCAbCABC", "ABC") == 3);
+
+		CHECK(FindFirst(FStringView(nullptr, 0), TEXT("SearchTerm")) == INDEX_NONE);
+		CHECK(FindFirst(FStringView(), TEXT("SearchTerm")) == INDEX_NONE);
+	}
+
+	SECTION("FindLast")
+	{
+		CHECK(FindLast(TEXT("AbCABCAbCABC"), TEXT("b")) == 7);
+		CHECK(FindLast(TEXT("AbCABCAbCABC"), TEXT("B")) == 10);
+		CHECK(FindLast(TEXT("AbCABCAbCABC"), TEXT("b"), ESearchCase::IgnoreCase) == 10);
+		CHECK(FindLast(TEXT("AbCABCAbCABC"), TEXT("a")) == INDEX_NONE);
+		CHECK(FindLast(TEXT("AbCABCAbCABC"), TEXT("D"), ESearchCase::IgnoreCase) == INDEX_NONE);
+		CHECK(FindLast(TEXT("AbCABCAbCABD"), TEXT("D")) == 11);
+		CHECK(FindLast(TEXT("AbCABCAbCABD"), TEXT("d"), ESearchCase::IgnoreCase) == 11);
+
+		CHECK(FindLast(TEXT("AbCABCAbCABC"), TEXT("AbC")) == 6);
+		CHECK(FindLast(TEXT("AbCABCAbCABC"), TEXT("ABC")) == 9);
+		CHECK(FindLast(TEXT("AbCABCAbCABC"), TEXT("Bc"), ESearchCase::IgnoreCase) == 10);
+		CHECK(FindLast(TEXT("AbCABCAbCABC"), TEXT("ab")) == INDEX_NONE);
+		CHECK(FindLast(TEXT("AbCABCAbCABC"), TEXT("CD"), ESearchCase::IgnoreCase) == INDEX_NONE);
+		CHECK(FindLast(TEXT("AbCABCAbCABC"), TEXT("BC")) == 10);
+		CHECK(FindLast(TEXT("AbCABCAbCABC"), TEXT("Bc"), ESearchCase::IgnoreCase) == 10);
+
+		CHECK(FindLast(TEXT(""), TEXT("A")) == INDEX_NONE);
+		CHECK(FindLast(TEXT("A"), TEXT("A")) == 0);
+		CHECK(FindLast(TEXT("A"), TEXT("A"), ESearchCase::IgnoreCase) == 0);
+		CHECK(FindLast(TEXT("ABC"), TEXT("ABC")) == 0);
+		CHECK(FindLast(TEXT("ABC"), TEXT("abc"), ESearchCase::IgnoreCase) == 0);
+		CHECK(FindLast(TEXT("AB"), TEXT("ABC")) == INDEX_NONE);
+
+		CHECK(FindLast("AbCABCAbCABC", "ABC") == 9);
+
+		CHECK(FindLast(FStringView(nullptr, 0), TEXT("SearchTerm")) == INDEX_NONE);
+		CHECK(FindLast(FStringView(), TEXT("SearchTerm")) == INDEX_NONE);
+	}
+
+	SECTION("FindFirstOfAny")
+	{
+		CHECK(FindFirstOfAny(TEXT("AbCABCAbcABC"), {TEXT("A"), TEXT("B")}) == 0);
+		CHECK(FindFirstOfAny(TEXT("AbCABCAbcABC"), {TEXT("a"), TEXT("B")}, ESearchCase::IgnoreCase) == 0);
+		CHECK(FindFirstOfAny(TEXT("AbCABCAbcABC"), {TEXT("b")}) == 1);
+		CHECK(FindFirstOfAny(TEXT("AbCABCAbcABC"), {TEXT("c"), TEXT("B")}) == 4);
+		CHECK(FindFirstOfAny(TEXT("AbCABCAbcABC"), {TEXT("c"), TEXT("B")}, ESearchCase::IgnoreCase) == 1);
+		CHECK(FindFirstOfAny(TEXT("AbCABCAbcABC"), {TEXT("D"), TEXT("a")}) == INDEX_NONE);
+		CHECK(FindFirstOfAny(TEXT("AbCABCAbcABC"), {TEXT("E"), TEXT("D")}, ESearchCase::IgnoreCase) == INDEX_NONE);
+		CHECK(FindFirstOfAny(TEXT("AbCABCAbcABD"), {TEXT("E"), TEXT("D")}) == 11);
+		CHECK(FindFirstOfAny(TEXT("AbCABCAbcABD"), {TEXT("E"), TEXT("d")}, ESearchCase::IgnoreCase) == 11);
+
+		CHECK(FindFirstOfAny(TEXT("AbCABCAbCABC"), {TEXT("AbC")}) == 0);
+		CHECK(FindFirstOfAny(TEXT("AbCABCAbCABC"), {TEXT("CABc"), TEXT("ABC")}) == 3);
+		CHECK(FindFirstOfAny(TEXT("AbCABCAbCABC"), {TEXT("ABD"), TEXT("Bc")}, ESearchCase::IgnoreCase) == 1);
+		CHECK(FindFirstOfAny(TEXT("AbCABCAbCABC"), {TEXT("bc"), TEXT("ab")}) == INDEX_NONE);
+		CHECK(FindFirstOfAny(TEXT("AbCABCAbCABC"), {TEXT("DA"), TEXT("CD")}, ESearchCase::IgnoreCase) == INDEX_NONE);
+		CHECK(FindFirstOfAny(TEXT("AbCABCAbCABD"), {TEXT("BD"), TEXT("CABB")}) == 10);
+		CHECK(FindFirstOfAny(TEXT("AbCABCAbCABD"), {TEXT("Bd"), TEXT("CABB")}, ESearchCase::IgnoreCase) == 10);
+
+		CHECK(FindFirstOfAny(TEXT(""), {TEXT("A"), TEXT("B")}) == INDEX_NONE);
+		CHECK(FindFirstOfAny(TEXT("A"), {TEXT("A"), TEXT("B")}) == 0);
+		CHECK(FindFirstOfAny(TEXT("A"), {TEXT("A"), TEXT("B")}, ESearchCase::IgnoreCase) == 0);
+		CHECK(FindFirstOfAny(TEXT("ABC"), {TEXT("ABC"), TEXT("BC")}) == 0);
+		CHECK(FindFirstOfAny(TEXT("ABC"), {TEXT("abc"), TEXT("bc")}, ESearchCase::IgnoreCase) == 0);
+		CHECK(FindFirstOfAny(TEXT("AB"), {TEXT("ABC"), TEXT("ABD")}) == INDEX_NONE);
+
+		CHECK(FindFirstOfAny("AbCABCAbCABC", {"CABc", "ABC"}) == 3);
+
+		CHECK(FindFirstOfAny(FStringView(nullptr, 0), {TEXT("ABC"), TEXT("ABD")}) == INDEX_NONE);
+		CHECK(FindFirstOfAny(FStringView(), {TEXT("ABC"), TEXT("ABD")}) == INDEX_NONE);
+	}
+
+	SECTION("FindLastOfAny")
+	{
+		CHECK(FindLastOfAny(TEXT("AbCABCAbcABC"), {TEXT("a"), TEXT("b")}) == 7);
+		CHECK(FindLastOfAny(TEXT("AbCABCAbcABC"), {TEXT("a"), TEXT("b")}, ESearchCase::IgnoreCase) == 10);
+		CHECK(FindLastOfAny(TEXT("AbCABCAbcABC"), {TEXT("b")}) == 7);
+		CHECK(FindLastOfAny(TEXT("AbCABCAbcABC"), {TEXT("c"), TEXT("B")}) == 10);
+		CHECK(FindLastOfAny(TEXT("AbCABCAbcABC"), {TEXT("c"), TEXT("B")}, ESearchCase::IgnoreCase) == 11);
+		CHECK(FindLastOfAny(TEXT("AbCABCAbcABC"), {TEXT("D"), TEXT("a")}) == INDEX_NONE);
+		CHECK(FindLastOfAny(TEXT("AbCABCAbcABC"), {TEXT("E"), TEXT("D")}, ESearchCase::IgnoreCase) == INDEX_NONE);
+		CHECK(FindLastOfAny(TEXT("AbCABCAbcABD"), {TEXT("E"), TEXT("D")}) == 11);
+		CHECK(FindLastOfAny(TEXT("AbCABCAbcABD"), {TEXT("E"), TEXT("d")}, ESearchCase::IgnoreCase) == 11);
+
+		CHECK(FindLastOfAny(TEXT("AbCABCAbCABC"), {TEXT("AbC")}) == 6);
+		CHECK(FindLastOfAny(TEXT("AbCABCAbCABC"), {TEXT("CABc"), TEXT("ABC")}) == 9);
+		CHECK(FindLastOfAny(TEXT("AbCABCAbCABC"), {TEXT("ABD"), TEXT("Bc")}, ESearchCase::IgnoreCase) == 10);
+		CHECK(FindLastOfAny(TEXT("AbCABCAbCABC"), {TEXT("bc"), TEXT("ab")}) == INDEX_NONE);
+		CHECK(FindLastOfAny(TEXT("AbCABCAbCABC"), {TEXT("DA"), TEXT("CD")}, ESearchCase::IgnoreCase) == INDEX_NONE);
+		CHECK(FindLastOfAny(TEXT("AbCABCAbCABD"), {TEXT("BD"), TEXT("CABB")}) == 10);
+		CHECK(FindLastOfAny(TEXT("AbCABCAbCABD"), {TEXT("Bd"), TEXT("CABB")}, ESearchCase::IgnoreCase) == 10);
+
+		CHECK(FindLastOfAny(TEXT(""), {TEXT("A"), TEXT("B")}) == INDEX_NONE);
+		CHECK(FindLastOfAny(TEXT("A"), {TEXT("A"), TEXT("B")}) == 0);
+		CHECK(FindLastOfAny(TEXT("A"), {TEXT("A"), TEXT("B")}, ESearchCase::IgnoreCase) == 0);
+		CHECK(FindLastOfAny(TEXT("ABC"), {TEXT("ABC"), TEXT("BC")}) == 1);
+		CHECK(FindLastOfAny(TEXT("ABC"), {TEXT("abc"), TEXT("bc")}, ESearchCase::IgnoreCase) == 1);
+		CHECK(FindLastOfAny(TEXT("AB"), {TEXT("ABC"), TEXT("ABD")}) == INDEX_NONE);
+
+		CHECK(FindLastOfAny("AbCABCAbCABC", {"CABc", "ABC"}) == 9);
+
+		CHECK(FindLastOfAny(FStringView(nullptr, 0), { TEXT("ABC"), TEXT("ABD") }) == INDEX_NONE);
+		CHECK(FindLastOfAny(FStringView(), { TEXT("ABC"), TEXT("ABD") }) == INDEX_NONE);
+	}
+
+	SECTION("FindFirstChar")
+	{
+		CHECK(FindFirstChar(TEXT("AbCABCAbCABC"), TEXT('b')) == 1);
+		CHECK(FindFirstChar(TEXT("AbCABCAbCABC"), TEXT('B')) == 4);
+		CHECK(FindFirstChar(TEXT("AbCABCAbCABC"), TEXT('B'), ESearchCase::IgnoreCase) == 1);
+		CHECK(FindFirstChar(TEXT("AbCABCAbCABC"), TEXT('a')) == INDEX_NONE);
+		CHECK(FindFirstChar(TEXT("AbCABCAbCABC"), TEXT('D'), ESearchCase::IgnoreCase) == INDEX_NONE);
+		CHECK(FindFirstChar(TEXT("AbCABCAbCABD"), TEXT('D')) == 11);
+		CHECK(FindFirstChar(TEXT("AbCABCAbCABD"), TEXT('d'), ESearchCase::IgnoreCase) == 11);
+
+		CHECK(FindFirstChar(TEXT(""), TEXT('A')) == INDEX_NONE);
+		CHECK(FindFirstChar(TEXT("A"), TEXT('A')) == 0);
+		CHECK(FindFirstChar(TEXT("A"), TEXT('A'), ESearchCase::IgnoreCase) == 0);
+
+		CHECK(FindFirstChar("AbCABCAbCABC", 'B') == 4);
+
+		CHECK(FindFirstChar(FStringView(nullptr, 0), TEXT('A')) == INDEX_NONE);
+		CHECK(FindFirstChar(FStringView(), TEXT('A')) == INDEX_NONE);
+	}
+
+	SECTION("FindLastChar")
+	{
+		CHECK(FindLastChar(TEXT("AbCABCAbCABC"), TEXT('b')) == 7);
+		CHECK(FindLastChar(TEXT("AbCABCAbCABC"), TEXT('B')) == 10);
+		CHECK(FindLastChar(TEXT("AbCABCAbCABC"), TEXT('b'), ESearchCase::IgnoreCase) == 10);
+		CHECK(FindLastChar(TEXT("AbCABCAbCABC"), TEXT('a')) == INDEX_NONE);
+		CHECK(FindLastChar(TEXT("AbCABCAbCABC"), TEXT('D'), ESearchCase::IgnoreCase) == INDEX_NONE);
+		CHECK(FindLastChar(TEXT("AbCABCAbCABD"), TEXT('D')) == 11);
+		CHECK(FindLastChar(TEXT("AbCABCAbCABD"), TEXT('d'), ESearchCase::IgnoreCase) == 11);
+
+		CHECK(FindLastChar(TEXT(""), TEXT('A')) == INDEX_NONE);
+		CHECK(FindLastChar(TEXT("A"), TEXT('A')) == 0);
+		CHECK(FindLastChar(TEXT("A"), TEXT('A'), ESearchCase::IgnoreCase) == 0);
+
+		CHECK(FindLastChar("AbCABCAbCABC", 'B') == 10);
+
+		CHECK(FindLastChar(FStringView(nullptr, 0), TEXT('A')) == INDEX_NONE);
+		CHECK(FindLastChar(FStringView(), TEXT('A')) == INDEX_NONE);
+	}
+
+	SECTION("FindFirstOfAnyChar")
+	{
+		CHECK(FindFirstOfAnyChar(TEXT("AbCABCAbcABC"), {TEXT('b')}) == 1);
+		CHECK(FindFirstOfAnyChar(TEXT("AbCABCAbcABC"), {TEXT('c'), TEXT('B')}) == 4);
+		CHECK(FindFirstOfAnyChar(TEXT("AbCABCAbcABC"), {TEXT('c'), TEXT('B')}, ESearchCase::IgnoreCase) == 1);
+		CHECK(FindFirstOfAnyChar(TEXT("AbCABCAbcABC"), {TEXT('D'), TEXT('a')}) == INDEX_NONE);
+		CHECK(FindFirstOfAnyChar(TEXT("AbCABCAbcABC"), {TEXT('E'), TEXT('D')}, ESearchCase::IgnoreCase) == INDEX_NONE);
+		CHECK(FindFirstOfAnyChar(TEXT("AbCABCAbcABD"), {TEXT('E'), TEXT('D')}) == 11);
+		CHECK(FindFirstOfAnyChar(TEXT("AbCABCAbcABD"), {TEXT('E'), TEXT('d')}, ESearchCase::IgnoreCase) == 11);
+
+		CHECK(FindFirstOfAnyChar(TEXT(""), {TEXT('A'), TEXT('B')}) == INDEX_NONE);
+		CHECK(FindFirstOfAnyChar(TEXT("A"), {TEXT('A'), TEXT('B')}) == 0);
+		CHECK(FindFirstOfAnyChar(TEXT("A"), {TEXT('A'), TEXT('B')}, ESearchCase::IgnoreCase) == 0);
+
+		CHECK(FindFirstOfAnyChar("AbCABCAbcABC", {'c', 'B'}) == 4);
+
+		CHECK(FindFirstOfAnyChar(FStringView(nullptr, 0), { TEXT('A'), TEXT('B') }) == INDEX_NONE);
+		CHECK(FindFirstOfAnyChar(FStringView(), { TEXT('A'), TEXT('B') }) == INDEX_NONE);
+	}
+
+	SECTION("FindLastOfAnyChar")
+	{
+		CHECK(FindLastOfAnyChar(TEXT("AbCABCAbcABC"), {TEXT('b')}) == 7);
+		CHECK(FindLastOfAnyChar(TEXT("AbCABCAbcABC"), {TEXT('c'), TEXT('B')}) == 10);
+		CHECK(FindLastOfAnyChar(TEXT("AbCABCAbcABC"), {TEXT('c'), TEXT('B')}, ESearchCase::IgnoreCase) == 11);
+		CHECK(FindLastOfAnyChar(TEXT("AbCABCAbcABC"), {TEXT('D'), TEXT('a')}) == INDEX_NONE);
+		CHECK(FindLastOfAnyChar(TEXT("AbCABCAbcABC"), {TEXT('E'), TEXT('D')}, ESearchCase::IgnoreCase) == INDEX_NONE);
+		CHECK(FindLastOfAnyChar(TEXT("AbCABCAbcABD"), {TEXT('E'), TEXT('D')}) == 11);
+		CHECK(FindLastOfAnyChar(TEXT("AbCABCAbcABD"), {TEXT('E'), TEXT('d')}, ESearchCase::IgnoreCase) == 11);
+
+		CHECK(FindLastOfAnyChar(TEXT(""), {TEXT('A'), TEXT('B')}) == INDEX_NONE);
+		CHECK(FindLastOfAnyChar(TEXT("A"), {TEXT('A'), TEXT('B')}) == 0);
+		CHECK(FindLastOfAnyChar(TEXT("A"), {TEXT('A'), TEXT('B')}, ESearchCase::IgnoreCase) == 0);
+
+		CHECK(FindLastOfAnyChar("AbCABCAbcABC", {'c', 'B'}) == 10);
+
+		CHECK(FindLastOfAnyChar(FStringView(nullptr, 0), { TEXT('A'), TEXT('B') }) == INDEX_NONE);
+		CHECK(FindLastOfAnyChar(FStringView(), { TEXT('A'), TEXT('B') }) == INDEX_NONE);
+	}
 }
 
-TEST_CASE_METHOD(FAutomationTestFixture, "Core::String::Misc::FindLastOfAnyChar", "[Core][String][Smoke]")
-{
-	TEST_EQUAL(TEXT("FindLastOfAnyChar"), UE::String::FindLastOfAnyChar(TEXT("AbCABCAbcABC"), {TEXT('b')}), 7);
-	TEST_EQUAL(TEXT("FindLastOfAnyChar"), UE::String::FindLastOfAnyChar(TEXT("AbCABCAbcABC"), {TEXT('c'), TEXT('B')}), 10);
-	TEST_EQUAL(TEXT("FindLastOfAnyChar"), UE::String::FindLastOfAnyChar(TEXT("AbCABCAbcABC"), {TEXT('c'), TEXT('B')}, ESearchCase::IgnoreCase), 11);
-	TEST_EQUAL(TEXT("FindLastOfAnyChar"), UE::String::FindLastOfAnyChar(TEXT("AbCABCAbcABC"), {TEXT('D'), TEXT('a')}), INDEX_NONE);
-	TEST_EQUAL(TEXT("FindLastOfAnyChar"), UE::String::FindLastOfAnyChar(TEXT("AbCABCAbcABC"), {TEXT('E'), TEXT('D')}, ESearchCase::IgnoreCase), INDEX_NONE);
-	TEST_EQUAL(TEXT("FindLastOfAnyChar"), UE::String::FindLastOfAnyChar(TEXT("AbCABCAbcABD"), {TEXT('E'), TEXT('D')}), 11);
-	TEST_EQUAL(TEXT("FindLastOfAnyChar"), UE::String::FindLastOfAnyChar(TEXT("AbCABCAbcABD"), {TEXT('E'), TEXT('d')}, ESearchCase::IgnoreCase), 11);
-
-	TEST_EQUAL(TEXT("FindLastOfAnyChar"), UE::String::FindLastOfAnyChar(TEXT(""), {TEXT('A'), TEXT('B')}), INDEX_NONE);
-	TEST_EQUAL(TEXT("FindLastOfAnyChar"), UE::String::FindLastOfAnyChar(TEXT("A"), {TEXT('A'), TEXT('B')}), 0);
-	TEST_EQUAL(TEXT("FindLastOfAnyChar"), UE::String::FindLastOfAnyChar(TEXT("A"), {TEXT('A'), TEXT('B')}, ESearchCase::IgnoreCase), 0);
-
-	TEST_EQUAL(TEXT("FindLastOfAnyChar"), UE::String::FindLastOfAnyChar("AbCABCAbcABC", {'c', 'B'}), 10);
-
-	TEST_EQUAL(TEXT("FindLastOfAnyChar"), UE::String::FindLastOfAnyChar(FStringView(nullptr, 0), { TEXT('A'), TEXT('B') }), INDEX_NONE);
-	TEST_EQUAL(TEXT("FindLastOfAnyChar"), UE::String::FindLastOfAnyChar(FStringView(), { TEXT('A'), TEXT('B') }), INDEX_NONE);
-	TEST_EQUAL(TEXT("FindLastOfAnyChar"), UE::String::FindLastOfAnyChar(FString(), { TEXT('A'), TEXT('B') }), INDEX_NONE);
-}
+} // UE::String
diff --git a/Engine/Source/Runtime/Core/Tests/String/JoinTest.cpp b/Engine/Source/Runtime/Core/Tests/String/JoinTest.cpp
--- a/Engine/Source/Runtime/Core/Tests/String/JoinTest.cpp
+++ b/Engine/Source/Runtime/Core/Tests/String/JoinTest.cpp
@@ -10,7 +10,7 @@
 namespace UE::String
 {
 
-TEST_CASE("Core::String::Misc::Join", "[Core][String][Smoke]")
+TEST_CASE("Core::String::Join", "[Core][String][Smoke]")
 {
 	SECTION("Join")
 	{
diff --git a/Engine/Source/Runtime/Core/Tests/String/ParseLinesTest.cpp b/Engine/Source/Runtime/Core/Tests/String/ParseLinesTest.cpp
--- a/Engine/Source/Runtime/Core/Tests/String/ParseLinesTest.cpp
+++ b/Engine/Source/Runtime/Core/Tests/String/ParseLinesTest.cpp
@@ -1,75 +1,82 @@
 // Copyright Epic Games, Inc. All Rights Reserved.
 
 #include "String/ParseLines.h"
+
 #include "Algo/Compare.h"
 #include "Containers/StringView.h"
 #include "Misc/StringBuilder.h"
-#include "Templates/EqualTo.h"
+#include "String/Escape.h"
+#include "String/Join.h"
 #include "TestHarness.h"
 
-TEST_CASE_METHOD(FAutomationTestFixture, "Core::String::Misc::ParseLines", "[Core][String][Smoke]")
+namespace UE::String
 {
-	using EOptions = UE::String::EParseLinesOptions;
-
-	auto RunParseLinesTest = [this](FStringView View, std::initializer_list<FStringView> ExpectedLines, EOptions Options = EOptions::None)
-	{
-		TArray<FStringView, TInlineAllocator<8>> ResultLines;
-		UE::String::ParseLines(View, ResultLines, Options);
 
-		TEST_TRUE(TEXT("UE::String::ParseLines failed to parse"), Algo::Compare(ResultLines, ExpectedLines));
+TEST_CASE("Core::String::ParseLines", "[Core][String][Smoke]")
+{
+	using FStringViewArray = TArray<FStringView, TInlineAllocator<8>>;
 
-		/*if (!Algo::Compare(ResultLines, ExpectedLines))
-		{
-			TStringBuilder<512> Error;
-			Error << TEXTVIEW("UE::String::ParseLines failed to parse \"") << FString(View).ReplaceCharWithEscapedChar() << TEXTVIEW("\" result {");
-			Error.JoinQuoted(ResultLines, TEXTVIEW(", "), TEXTVIEW("\""));
-			Error << TEXTVIEW("} expected {");
-			Error.JoinQuoted(ExpectedLines, TEXTVIEW(", "), TEXTVIEW("\""));
-			Error << TEXTVIEW("}");
-			FAIL_CHECK(Error.ToString());
-		}*/
-	};
+	constexpr EParseLinesOptions KeepEmpty = EParseLinesOptions::None;
+	constexpr EParseLinesOptions SkipEmpty = EParseLinesOptions::SkipEmpty;
+	constexpr EParseLinesOptions Trim = EParseLinesOptions::Trim;
 
-	constexpr EOptions KeepEmpty = EOptions::None;
-	constexpr EOptions SkipEmpty = EOptions::SkipEmpty;
-	constexpr EOptions Trim = EOptions::Trim;
+	const auto [View, ExpectedLines, Options] = GENERATE_COPY(table<FStringView, FStringViewArray, EParseLinesOptions>(
+	{
+		{TEXTVIEW(""), {}, SkipEmpty},
+		{TEXTVIEW(""), {TEXTVIEW("")}, KeepEmpty},
+		{TEXTVIEW("\n"), {}, SkipEmpty},
+		{TEXTVIEW("\n"), {TEXTVIEW("")}, KeepEmpty},
+		{TEXTVIEW("\r"), {}, SkipEmpty},
+		{TEXTVIEW("\r"), {TEXTVIEW("")}, KeepEmpty},
+		{TEXTVIEW("\r\n"), {}, SkipEmpty},
+		{TEXTVIEW("\r\n"), {TEXTVIEW("")}, KeepEmpty},
+		{TEXTVIEW("\n\n"), {}, SkipEmpty},
+		{TEXTVIEW("\n\n"), {TEXTVIEW(""), TEXTVIEW("")}, KeepEmpty},
+		{TEXTVIEW("\r\r"), {}, SkipEmpty},
+		{TEXTVIEW("\r\r"), {TEXTVIEW(""), TEXTVIEW("")}, KeepEmpty},
+		{TEXTVIEW("\r\n\r\n"), {}, SkipEmpty},
+		{TEXTVIEW("\r\n\r\n"), {TEXTVIEW(""), TEXTVIEW("")}, KeepEmpty},
+		{TEXTVIEW("\r\nABC").Left(2), {}, SkipEmpty},
+		{TEXTVIEW("\r\nABC").Left(2), {TEXTVIEW("")}, KeepEmpty},
+		{TEXTVIEW("\r\nABC\r\nDEF").Left(5), {TEXTVIEW("ABC")}, SkipEmpty},
+		{TEXTVIEW("\r\nABC\r\nDEF").Left(5), {TEXTVIEW(""), TEXTVIEW("ABC")}, KeepEmpty},
+		{TEXTVIEW("ABC DEF"), {TEXTVIEW("ABC DEF")}, SkipEmpty},
+		{TEXTVIEW("\nABC DEF\n"), {TEXTVIEW("ABC DEF")}, SkipEmpty},
+		{TEXTVIEW("\nABC DEF\n"), {TEXTVIEW(""), TEXTVIEW("ABC DEF")}, KeepEmpty},
+		{TEXTVIEW("\rABC DEF\r"), {TEXTVIEW("ABC DEF")}, SkipEmpty},
+		{TEXTVIEW("\rABC DEF\r"), {TEXTVIEW(""), TEXTVIEW("ABC DEF")}, KeepEmpty},
+		{TEXTVIEW("\r\nABC DEF\r\n"), {TEXTVIEW("ABC DEF")}, SkipEmpty},
+		{TEXTVIEW("\r\nABC DEF\r\n"), {TEXTVIEW(""), TEXTVIEW("ABC DEF")}, KeepEmpty},
+		{TEXTVIEW("\r\n\r\nABC DEF\r\n\r\n"), {TEXTVIEW("ABC DEF")}, SkipEmpty},
+		{TEXTVIEW("\r\n\r\nABC DEF\r\n\r\n"), {TEXTVIEW(""), TEXTVIEW(""), TEXTVIEW("ABC DEF"), TEXTVIEW("")}, KeepEmpty},
+		{TEXTVIEW("ABC\nDEF"), {TEXTVIEW("ABC"), TEXTVIEW("DEF")}, SkipEmpty},
+		{TEXTVIEW("ABC\rDEF"), {TEXTVIEW("ABC"), TEXTVIEW("DEF")}, SkipEmpty},
+		{TEXTVIEW("\r\nABC\r\nDEF\r\n"), {TEXTVIEW("ABC"), TEXTVIEW("DEF")}, SkipEmpty},
+		{TEXTVIEW("\r\nABC\r\nDEF\r\n"), {TEXTVIEW(""), TEXTVIEW("ABC"), TEXTVIEW("DEF")}, KeepEmpty},
+		{TEXTVIEW("\r\nABC\r\n\r\nDEF\r\n"), {TEXTVIEW("ABC"), TEXTVIEW("DEF")}, SkipEmpty},
+		{TEXTVIEW("\r\nABC\r\n\r\nDEF\r\n"), {TEXTVIEW(""), TEXTVIEW("ABC"), TEXTVIEW(""), TEXTVIEW("DEF")}, KeepEmpty},
+		{TEXTVIEW(" \t\r\n\t ABC \t\r\n\t \t\r\n\t DEF \t\r\n"), {TEXTVIEW(" \t"), TEXTVIEW("\t ABC \t"), TEXTVIEW("\t \t"), TEXTVIEW("\t DEF \t")}, SkipEmpty},
+		{TEXTVIEW(" \t\r\n\t ABC \t\r\n\t \t\r\n\t DEF \t\r\n"), {TEXTVIEW("ABC"), TEXTVIEW("DEF")}, SkipEmpty | Trim},
+		{TEXTVIEW(" \t\r\n\t ABC \t\r\n\t \t\r\n\t DEF \t\r\n"), {TEXTVIEW(" \t"), TEXTVIEW("\t ABC \t"), TEXTVIEW("\t \t"), TEXTVIEW("\t DEF \t")}, KeepEmpty},
+		{TEXTVIEW(" \t\r\n\t ABC \t\r\n\t \t\r\n\t DEF \t\r\n"), {TEXTVIEW(""), TEXTVIEW("ABC"), TEXTVIEW(""), TEXTVIEW("DEF")}, KeepEmpty | Trim},
+	}));
 
-	RunParseLinesTest(TEXTVIEW(""), {}, SkipEmpty);
-	RunParseLinesTest(TEXTVIEW(""), {TEXTVIEW("")}, KeepEmpty);
-	RunParseLinesTest(TEXTVIEW("\n"), {}, SkipEmpty);
-	RunParseLinesTest(TEXTVIEW("\n"), {TEXTVIEW("")}, KeepEmpty);
-	RunParseLinesTest(TEXTVIEW("\r"), {}, SkipEmpty);
-	RunParseLinesTest(TEXTVIEW("\r"), {TEXTVIEW("")}, KeepEmpty);
-	RunParseLinesTest(TEXTVIEW("\r\n"), {}, SkipEmpty);
-	RunParseLinesTest(TEXTVIEW("\r\n"), {TEXTVIEW("")}, KeepEmpty);
-	RunParseLinesTest(TEXTVIEW("\n\n"), {}, SkipEmpty);
-	RunParseLinesTest(TEXTVIEW("\n\n"), {TEXTVIEW(""), TEXTVIEW("")}, KeepEmpty);
-	RunParseLinesTest(TEXTVIEW("\r\r"), {}, SkipEmpty);
-	RunParseLinesTest(TEXTVIEW("\r\r"), {TEXTVIEW(""), TEXTVIEW("")}, KeepEmpty);
-	RunParseLinesTest(TEXTVIEW("\r\n\r\n"), {}, SkipEmpty);
-	RunParseLinesTest(TEXTVIEW("\r\n\r\n"), {TEXTVIEW(""), TEXTVIEW("")}, KeepEmpty);
-	RunParseLinesTest(TEXTVIEW("\r\nABC").Left(2), {}, SkipEmpty);
-	RunParseLinesTest(TEXTVIEW("\r\nABC").Left(2), {TEXTVIEW("")}, KeepEmpty);
-	RunParseLinesTest(TEXTVIEW("\r\nABC\r\nDEF").Left(5), {TEXTVIEW("ABC")}, SkipEmpty);
-	RunParseLinesTest(TEXTVIEW("\r\nABC\r\nDEF").Left(5), {TEXTVIEW(""), TEXTVIEW("ABC")}, KeepEmpty);
-	RunParseLinesTest(TEXTVIEW("ABC DEF"), {TEXTVIEW("ABC DEF")}, SkipEmpty);
-	RunParseLinesTest(TEXTVIEW("\nABC DEF\n"), {TEXTVIEW("ABC DEF")}, SkipEmpty);
-	RunParseLinesTest(TEXTVIEW("\nABC DEF\n"), {TEXTVIEW(""), TEXTVIEW("ABC DEF")}, KeepEmpty);
-	RunParseLinesTest(TEXTVIEW("\rABC DEF\r"), {TEXTVIEW("ABC DEF")}, SkipEmpty);
-	RunParseLinesTest(TEXTVIEW("\rABC DEF\r"), {TEXTVIEW(""), TEXTVIEW("ABC DEF")}, KeepEmpty);
-	RunParseLinesTest(TEXTVIEW("\r\nABC DEF\r\n"), {TEXTVIEW("ABC DEF")}, SkipEmpty);
-	RunParseLinesTest(TEXTVIEW("\r\nABC DEF\r\n"), {TEXTVIEW(""), TEXTVIEW("ABC DEF")}, KeepEmpty);
-	RunParseLinesTest(TEXTVIEW("\r\n\r\nABC DEF\r\n\r\n"), {TEXTVIEW("ABC DEF")}, SkipEmpty);
-	RunParseLinesTest(TEXTVIEW("\r\n\r\nABC DEF\r\n\r\n"), {TEXTVIEW(""), TEXTVIEW(""), TEXTVIEW("ABC DEF"), TEXTVIEW("")}, KeepEmpty);
-	RunParseLinesTest(TEXTVIEW("ABC\nDEF"), {TEXTVIEW("ABC"), TEXTVIEW("DEF")}, SkipEmpty);
-	RunParseLinesTest(TEXTVIEW("ABC\rDEF"), {TEXTVIEW("ABC"), TEXTVIEW("DEF")}, SkipEmpty);
-	RunParseLinesTest(TEXTVIEW("\r\nABC\r\nDEF\r\n"), {TEXTVIEW("ABC"), TEXTVIEW("DEF")}, SkipEmpty);
-	RunParseLinesTest(TEXTVIEW("\r\nABC\r\nDEF\r\n"), {TEXTVIEW(""), TEXTVIEW("ABC"), TEXTVIEW("DEF")}, KeepEmpty);
-	RunParseLinesTest(TEXTVIEW("\r\nABC\r\n\r\nDEF\r\n"), {TEXTVIEW("ABC"), TEXTVIEW("DEF")}, SkipEmpty);
-	RunParseLinesTest(TEXTVIEW("\r\nABC\r\n\r\nDEF\r\n"), {TEXTVIEW(""), TEXTVIEW("ABC"), TEXTVIEW(""), TEXTVIEW("DEF")}, KeepEmpty);
-	RunParseLinesTest(TEXTVIEW(" \t\r\n\t ABC \t\r\n\t \t\r\n\t DEF \t\r\n"), {TEXTVIEW(" \t"), TEXTVIEW("\t ABC \t"), TEXTVIEW("\t \t"), TEXTVIEW("\t DEF \t")}, SkipEmpty);
-	RunParseLinesTest(TEXTVIEW(" \t\r\n\t ABC \t\r\n\t \t\r\n\t DEF \t\r\n"), {TEXTVIEW("ABC"), TEXTVIEW("DEF")}, SkipEmpty | Trim);
-	RunParseLinesTest(TEXTVIEW(" \t\r\n\t ABC \t\r\n\t \t\r\n\t DEF \t\r\n"), {TEXTVIEW(" \t"), TEXTVIEW("\t ABC \t"), TEXTVIEW("\t \t"), TEXTVIEW("\t DEF \t")}, KeepEmpty);
-	RunParseLinesTest(TEXTVIEW(" \t\r\n\t ABC \t\r\n\t \t\r\n\t DEF \t\r\n"), {TEXTVIEW(""), TEXTVIEW("ABC"), TEXTVIEW(""), TEXTVIEW("DEF")}, KeepEmpty | Trim);
+	FStringViewArray ActualLines;
+	ParseLines(View, ActualLines, Options);
+	const bool bEqual = Algo::Compare(ActualLines, ExpectedLines);
+	if (!bEqual)
+	{
+		TStringBuilder<256> Input, Expected, Actual;
+		Input << QuoteEscape(View);
+		Expected << TEXTVIEW("[") << JoinBy(ExpectedLines, QuoteEscape, TEXTVIEW(", ")) << TEXTVIEW("]");
+		Actual << TEXTVIEW("[") << JoinBy(ActualLines, QuoteEscape, TEXTVIEW(", ")) << TEXTVIEW("]");
+		CAPTURE(Input, Expected, Actual);
+		CHECK(bEqual);
+	}
+	else
+	{
+		CHECK(bEqual);
+	}
 }
 
+} // UE::String
diff --git a/Engine/Source/Runtime/Core/Tests/String/ParseTokensTest.cpp b/Engine/Source/Runtime/Core/Tests/String/ParseTokensTest.cpp
--- a/Engine/Source/Runtime/Core/Tests/String/ParseTokensTest.cpp
+++ b/Engine/Source/Runtime/Core/Tests/String/ParseTokensTest.cpp
@@ -4,166 +4,161 @@
 
 #include "Algo/Compare.h"
 #include "Misc/StringBuilder.h"
-#include "Templates/EqualTo.h"
+#include "String/Escape.h"
+#include "String/Join.h"
 #include "TestHarness.h"
 
-#if WITH_DEV_AUTOMATION_TESTS
+namespace UE::String
+{
 
-TEST_CASE_METHOD(FAutomationTestFixture, "Core::String::Misc::Parse Tokens By String", "[Core][String][Smoke]")
+TEST_CASE("Core::String::ParseTokens", "[Core][String][Smoke]")
 {
-	using EOptions = UE::String::EParseTokensOptions;
+	using FCharArray = TArray<TCHAR, TInlineAllocator<8>>;
+	using FStringViewArray = TArray<FStringView, TInlineAllocator<8>>;
+
+	constexpr EParseTokensOptions KeepEmpty = EParseTokensOptions::None;
+	constexpr EParseTokensOptions SkipEmpty = EParseTokensOptions::SkipEmpty;
+	constexpr EParseTokensOptions IgnoreCase = EParseTokensOptions::IgnoreCase;
+	constexpr EParseTokensOptions Trim = EParseTokensOptions::Trim;
 
-	auto RunParseTokensTest = [this](
-		FStringView View,
-		std::initializer_list<FStringView> Delimiters,
-		std::initializer_list<FStringView> ExpectedTokens,
-		EOptions Options = EOptions::None)
+	SECTION("String")
 	{
-		TArray<FStringView, TInlineAllocator<8>> ResultTokens;
-		if (GetNum(Delimiters) == 1)
+		const auto [View, DelimiterArray, ExpectedTokens, Options] = GENERATE_COPY(table<FStringView, FStringViewArray, FStringViewArray, EParseTokensOptions>(
 		{
-			UE::String::ParseTokens(View, GetData(Delimiters)[0], ResultTokens, Options);
+			{TEXTVIEW(""),                {},                   {},                                                     SkipEmpty},
+			{TEXTVIEW(""),                {},                   {TEXTVIEW("")},                                         KeepEmpty},
+			{TEXTVIEW("ABC"),             {},                   {TEXTVIEW("ABC")},                                      KeepEmpty},
+
+			{TEXTVIEW(""),                {TEXTVIEW(",")},      {},                                                     SkipEmpty},
+			{TEXTVIEW(""),                {TEXTVIEW(",")},      {TEXTVIEW("")},                                         KeepEmpty},
+			{TEXTVIEW(","),               {TEXTVIEW(",")},      {},                                                     SkipEmpty},
+			{TEXTVIEW(","),               {TEXTVIEW(",")},      {TEXTVIEW(""), TEXTVIEW("")},                           KeepEmpty},
+			{TEXTVIEW(",,"),              {TEXTVIEW(",")},      {},                                                     SkipEmpty},
+			{TEXTVIEW(",,"),              {TEXTVIEW(",")},      {TEXTVIEW(""), TEXTVIEW(""), TEXTVIEW("")},             KeepEmpty},
+			{TEXTVIEW(", ,"),             {TEXTVIEW(",")},      {TEXTVIEW(" ")},                                        SkipEmpty},
+			{TEXTVIEW(", ,"),             {TEXTVIEW(",")},      {},                                                     SkipEmpty | Trim},
+			{TEXTVIEW(", ,"),             {TEXTVIEW(",")},      {TEXTVIEW(""), TEXTVIEW(" "), TEXTVIEW("")},            KeepEmpty},
+			{TEXTVIEW(", ,"),             {TEXTVIEW(",")},      {TEXTVIEW(""), TEXTVIEW(""), TEXTVIEW("")},             KeepEmpty | Trim},
+			{TEXTVIEW("ABC"),             {TEXTVIEW(",")},      {TEXTVIEW("ABC")},                                      KeepEmpty},
+			{TEXTVIEW("A,,C"),            {TEXTVIEW(",")},      {TEXTVIEW("A"), TEXTVIEW("C")},                         SkipEmpty},
+			{TEXTVIEW("A,,C"),            {TEXTVIEW(",")},      {TEXTVIEW("A"), TEXTVIEW(""), TEXTVIEW("C")},           KeepEmpty},
+			{TEXTVIEW("A,\tB\t,C"),       {TEXTVIEW(",")},      {TEXTVIEW("A"), TEXTVIEW("\tB\t"), TEXTVIEW("C")},      KeepEmpty},
+			{TEXTVIEW(",A, B ,C,"),       {TEXTVIEW(",")},      {TEXTVIEW("A"), TEXTVIEW(" B "), TEXTVIEW("C")},        SkipEmpty},
+			{TEXTVIEW(",A, B ,C,"),       {TEXTVIEW(",")},      {TEXTVIEW(""), TEXTVIEW("A"), TEXTVIEW(" B "), TEXTVIEW("C"), TEXTVIEW("")}, KeepEmpty},
+			{TEXTVIEW("A\u2022B\u2022C"), {TEXTVIEW("\u2022")}, {TEXTVIEW("A"), TEXTVIEW("B"), TEXTVIEW("C")},          KeepEmpty},
+
+			{TEXTVIEW("ABCDABCD"), {TEXTVIEW("AB")},   {TEXTVIEW("CD"), TEXTVIEW("CD")},               SkipEmpty},
+			{TEXTVIEW("ABCDABCD"), {TEXTVIEW("AB")},   {TEXTVIEW(""), TEXTVIEW("CD"), TEXTVIEW("CD")}, KeepEmpty},
+			{TEXTVIEW("ABCDABCD"), {TEXTVIEW("ABCD")}, {},                                             SkipEmpty},
+			{TEXTVIEW("ABCDABCD"), {TEXTVIEW("ABCD")}, {TEXTVIEW(""), TEXTVIEW(""), TEXTVIEW("")},     KeepEmpty},
+			{TEXTVIEW("ABCDABCD"), {TEXTVIEW("DA")},   {TEXTVIEW("ABC"), TEXTVIEW("BCD")},             KeepEmpty},
+
+			{TEXTVIEW("ABCDABCD"), {TEXTVIEW("B"),  TEXTVIEW("D")},  {TEXTVIEW("A"), TEXTVIEW("C"), TEXTVIEW("A"), TEXTVIEW("C")},               SkipEmpty},
+			{TEXTVIEW("ABCDABCD"), {TEXTVIEW("B"),  TEXTVIEW("D")},  {TEXTVIEW("A"), TEXTVIEW("C"), TEXTVIEW("A"), TEXTVIEW("C"), TEXTVIEW("")}, KeepEmpty},
+			{TEXTVIEW("ABCDABCD"), {TEXTVIEW("BC"), TEXTVIEW("DA")}, {TEXTVIEW("A"), TEXTVIEW("D")},                                             SkipEmpty},
+			{TEXTVIEW("ABCDABCD"), {TEXTVIEW("BC"), TEXTVIEW("DA")}, {TEXTVIEW("A"), TEXTVIEW(""), TEXTVIEW(""), TEXTVIEW("D")},                 KeepEmpty},
+
+			{TEXTVIEW("AbCdaBcDAbCd"), {TEXTVIEW("Bc"), TEXTVIEW("da")}, {TEXTVIEW("AbC"), TEXTVIEW("DAbCd")}, SkipEmpty},
+			{TEXTVIEW("AbCdaBcDAbCd"), {TEXTVIEW("Bc"), TEXTVIEW("da")}, {TEXTVIEW("A"), TEXTVIEW("d")},       SkipEmpty | IgnoreCase},
+
+			{TEXTVIEW("A\u2022\u2022B,,C"), {TEXTVIEW(",,"), TEXTVIEW("\u2022\u2022")},                     {TEXTVIEW("A"), TEXTVIEW("B"), TEXTVIEW("C")}, KeepEmpty},
+			{TEXTVIEW("A\u2022\u2022B\u0085\u0085C"), {TEXTVIEW("\u0085\u0085"), TEXTVIEW("\u2022\u2022")}, {TEXTVIEW("A"), TEXTVIEW("B"), TEXTVIEW("C")}, KeepEmpty},
+		}));
+
+		FStringViewArray ActualTokens;
+		if (GetNum(DelimiterArray) == 1)
+		{
+			ParseTokens(View, GetData(DelimiterArray)[0], ActualTokens, Options);
 		}
 		else
 		{
-			UE::String::ParseTokensMultiple(View, Delimiters, ResultTokens, Options);
+			ParseTokensMultiple(View, DelimiterArray, ActualTokens, Options);
 		}
 
-		TEST_TRUE(TEXT("UE::String::ParseTokens[Multiple] failed to parse"), Algo::Compare(ResultTokens, ExpectedTokens));
-
-		if (!Algo::Compare(ResultTokens, ExpectedTokens))
+		const bool bEqual = Algo::Compare(ActualTokens, ExpectedTokens);
+		if (!bEqual)
 		{
-			TStringBuilder<512> Error;
-			Error << TEXTVIEW("UE::String::ParseTokens[Multiple] failed to parse \"") << View << TEXTVIEW("\" with delimiters {");
-			Error.JoinQuoted(Delimiters, TEXTVIEW(", "), TEXTVIEW("\""));
-			Error << TEXTVIEW("} result {");
-			Error.JoinQuoted(ResultTokens, TEXTVIEW(", "), TEXTVIEW("\""));
-			Error << TEXTVIEW("} expected {");
-			Error.JoinQuoted(ExpectedTokens, TEXTVIEW(", "), TEXTVIEW("\""));
-			Error << TEXTVIEW("}");
-			FAIL_CHECK(Error.ToString());
+			TStringBuilder<256> Input, Delimiters, Expected, Actual;
+			Input << QuoteEscape(View);
+			Delimiters << TEXTVIEW("[") << JoinBy(DelimiterArray, QuoteEscape, TEXTVIEW(", ")) << TEXTVIEW("]");
+			Expected << TEXTVIEW("[") << JoinBy(ExpectedTokens, QuoteEscape, TEXTVIEW(", ")) << TEXTVIEW("]");
+			Actual << TEXTVIEW("[") << JoinBy(ActualTokens, QuoteEscape, TEXTVIEW(", ")) << TEXTVIEW("]");
+			CAPTURE(Input, Delimiters, Expected, Actual);
+			CHECK(bEqual);
 		}
-	};
-
-	constexpr EOptions KeepEmpty = EOptions::None;
-	constexpr EOptions SkipEmpty = EOptions::SkipEmpty;
-	constexpr EOptions IgnoreCase = EOptions::IgnoreCase;
-	constexpr EOptions Trim = EOptions::Trim;
-
-	RunParseTokensTest(TEXT(""),         {},                       {},                                 SkipEmpty);
-	RunParseTokensTest(TEXT(""),         {},                       {TEXT("")},                         KeepEmpty);
-	RunParseTokensTest(TEXT("ABC"),      {},                       {TEXT("ABC")});
-
-	RunParseTokensTest(TEXT(""),         {TEXT(",")},              {},                                 SkipEmpty);
-	RunParseTokensTest(TEXT(""),         {TEXT(",")},              {TEXT("")},                         KeepEmpty);
-	RunParseTokensTest(TEXT(","),        {TEXT(",")},              {},                                 SkipEmpty);
-	RunParseTokensTest(TEXT(","),        {TEXT(",")},              {TEXT(""), TEXT("")},               KeepEmpty);
-	RunParseTokensTest(TEXT(",,"),       {TEXT(",")},              {},                                 SkipEmpty);
-	RunParseTokensTest(TEXT(",,"),       {TEXT(",")},              {TEXT(""), TEXT(""), TEXT("")},     KeepEmpty);
-	RunParseTokensTest(TEXT(", ,"),      {TEXT(",")},              {TEXT(" ")},                        SkipEmpty);
-	RunParseTokensTest(TEXT(", ,"),      {TEXT(",")},              {},                                 SkipEmpty | Trim);
-	RunParseTokensTest(TEXT(", ,"),      {TEXT(",")},              {TEXT(""), TEXT(" "), TEXT("")},    KeepEmpty);
-	RunParseTokensTest(TEXT(", ,"),      {TEXT(",")},              {TEXT(""), TEXT(""), TEXT("")},    KeepEmpty | Trim);
-	RunParseTokensTest(TEXT("ABC"),      {TEXT(",")},              {TEXT("ABC")});
-	RunParseTokensTest(TEXT("A,,C"),     {TEXT(",")},              {TEXT("A"), TEXT("C")},             SkipEmpty);
-	RunParseTokensTest(TEXT("A,,C"),     {TEXT(",")},              {TEXT("A"), TEXT(""), TEXT("C")},   KeepEmpty);
-	RunParseTokensTest(TEXT("A,\tB\t,C"), {TEXT(",")},             {TEXT("A"), TEXT("\tB\t"), TEXT("C")});
-	RunParseTokensTest(TEXT(",A, B ,C,"), {TEXT(",")},             {TEXT("A"), TEXT(" B "), TEXT("C")},                     SkipEmpty);
-	RunParseTokensTest(TEXT(",A, B ,C,"), {TEXT(",")},             {TEXT(""), TEXT("A"), TEXT(" B "), TEXT("C"), TEXT("")}, KeepEmpty);
-	RunParseTokensTest(TEXT("A\u2022B\u2022C"), {TEXT("\u2022")},  {TEXT("A"), TEXT("B"), TEXT("C")});
-
-	RunParseTokensTest(TEXT("ABCDABCD"), {TEXT("AB")},             {TEXT("CD"), TEXT("CD")},           SkipEmpty);
-	RunParseTokensTest(TEXT("ABCDABCD"), {TEXT("AB")},             {TEXT(""), TEXT("CD"), TEXT("CD")}, KeepEmpty);
-	RunParseTokensTest(TEXT("ABCDABCD"), {TEXT("ABCD")},           {},                                 SkipEmpty);
-	RunParseTokensTest(TEXT("ABCDABCD"), {TEXT("ABCD")},           {TEXT(""), TEXT(""), TEXT("")},     KeepEmpty);
-	RunParseTokensTest(TEXT("ABCDABCD"), {TEXT("DA")},             {TEXT("ABC"), TEXT("BCD")});
-
-	RunParseTokensTest(TEXT("ABCDABCD"), {TEXT("B"),  TEXT("D")},  {TEXT("A"), TEXT("C"), TEXT("A"), TEXT("C")},           SkipEmpty);
-	RunParseTokensTest(TEXT("ABCDABCD"), {TEXT("B"),  TEXT("D")},  {TEXT("A"), TEXT("C"), TEXT("A"), TEXT("C"), TEXT("")}, KeepEmpty);
-	RunParseTokensTest(TEXT("ABCDABCD"), {TEXT("BC"), TEXT("DA")}, {TEXT("A"), TEXT("D")},                                 SkipEmpty);
-	RunParseTokensTest(TEXT("ABCDABCD"), {TEXT("BC"), TEXT("DA")}, {TEXT("A"), TEXT(""), TEXT(""), TEXT("D")},             KeepEmpty);
-
-	RunParseTokensTest(TEXT("AbCdaBcDAbCd"), {TEXT("Bc"), TEXT("da")}, {TEXT("AbC"), TEXT("DAbCd")}, SkipEmpty);
-	RunParseTokensTest(TEXT("AbCdaBcDAbCd"), {TEXT("Bc"), TEXT("da")}, {TEXT("A"), TEXT("d")}, SkipEmpty | IgnoreCase);
-
-	RunParseTokensTest(TEXT("A\u2022\u2022B,,C"), {TEXT(",,"), TEXT("\u2022\u2022")},                      {TEXT("A"), TEXT("B"), TEXT("C")});
-	RunParseTokensTest(TEXT("A\u2022\u2022B\u0085\u0085C"), {TEXT("\u0085\u0085"), TEXT("\u2022\u2022")},  {TEXT("A"), TEXT("B"), TEXT("C")});
-}
-
-TEST_CASE_METHOD(FAutomationTestFixture, "Core::String::Misc::Parse Tokens By Token", "[Core][String][Smoke]")
-{
-	using EOptions = UE::String::EParseTokensOptions;
+		else
+		{
+			CHECK(bEqual);
+		}
+	}
 
-	auto RunParseTokensTest = [this](
-		FStringView View,
-		std::initializer_list<TCHAR> Delimiters,
-		std::initializer_list<FStringView> ExpectedTokens,
-		EOptions Options = EOptions::None)
+	SECTION("Char")
 	{
-		TArray<FStringView, TInlineAllocator<8>> ResultTokens;
-		if (GetNum(Delimiters) == 1)
+		const auto [View, DelimiterArray, ExpectedTokens, Options] = GENERATE_COPY(table<FStringView, FCharArray, FStringViewArray, EParseTokensOptions>(
 		{
-			UE::String::ParseTokens(View, GetData(Delimiters)[0], ResultTokens, Options);
+			{TEXTVIEW(""),                {},               {},                                    SkipEmpty},
+			{TEXTVIEW(""),                {},               {TEXT("")},                            KeepEmpty},
+			{TEXTVIEW("ABC"),             {},               {TEXT("ABC")},                         KeepEmpty},
+
+			{TEXTVIEW(""),                {TEXT(',')},      {},                                    SkipEmpty},
+			{TEXTVIEW(""),                {TEXT(',')},      {TEXT("")},                            KeepEmpty},
+			{TEXTVIEW(","),               {TEXT(',')},      {},                                    SkipEmpty},
+			{TEXTVIEW(","),               {TEXT(',')},      {TEXT(""), TEXT("")},                  KeepEmpty},
+			{TEXTVIEW(",,"),              {TEXT(',')},      {},                                    SkipEmpty},
+			{TEXTVIEW(",,"),              {TEXT(',')},      {TEXT(""), TEXT(""), TEXT("")},        KeepEmpty},
+			{TEXTVIEW(", ,"),             {TEXT(',')},      {TEXT(" ")},                           SkipEmpty},
+			{TEXTVIEW(", ,"),             {TEXT(',')},      {},                                    SkipEmpty | Trim},
+			{TEXTVIEW(", ,"),             {TEXT(',')},      {TEXT(""), TEXT(" "), TEXT("")},       KeepEmpty},
+			{TEXTVIEW(", ,"),             {TEXT(',')},      {TEXT(""), TEXT(""), TEXT("")},        KeepEmpty | Trim},
+			{TEXTVIEW("ABC"),             {TEXT(',')},      {TEXT("ABC")},                         KeepEmpty},
+			{TEXTVIEW("A,,C"),            {TEXT(',')},      {TEXT("A"), TEXT("C")},                SkipEmpty},
+			{TEXTVIEW("A,,C"),            {TEXT(',')},      {TEXT("A"), TEXT(""), TEXT("C")},      KeepEmpty},
+			{TEXTVIEW("A,\tB\t,C"),       {TEXT(',')},      {TEXT("A"), TEXT("\tB\t"), TEXT("C")}, KeepEmpty},
+			{TEXTVIEW(",A, B ,C,"),       {TEXT(',')},      {TEXT("A"), TEXT(" B "), TEXT("C")},                     SkipEmpty},
+			{TEXTVIEW(",A, B ,C,"),       {TEXT(',')},      {TEXT(""), TEXT("A"), TEXT(" B "), TEXT("C"), TEXT("")}, KeepEmpty},
+			{TEXTVIEW("A\u2022B\u2022C"), {TEXT('\u2022')}, {TEXT("A"), TEXT("B"), TEXT("C")},     KeepEmpty},
+
+			{TEXTVIEW("ABCDABCD"),        {TEXT('B'), TEXT('D')},           {TEXT("A"), TEXT("C"), TEXT("A"), TEXT("C")},           SkipEmpty},
+			{TEXTVIEW("ABCDABCD"),        {TEXT('B'), TEXT('D')},           {TEXT("A"), TEXT("C"), TEXT("A"), TEXT("C"), TEXT("")}, KeepEmpty},
+			{TEXTVIEW("A\u2022B,C"),      {TEXT(','), TEXT('\u2022')},      {TEXT("A"), TEXT("B"), TEXT("C")},                      KeepEmpty},
+			{TEXTVIEW("A\u2022B\u0085C"), {TEXT('\u0085'), TEXT('\u2022')}, {TEXT("A"), TEXT("B"), TEXT("C")},                      KeepEmpty},
+
+			{TEXTVIEW("ABC"), {TEXT('b')}, {TEXT("ABC")}, SkipEmpty},
+			{TEXTVIEW("ABC"), {TEXT('b')}, {TEXT("A"), TEXT("C")}, SkipEmpty | IgnoreCase},
+
+			{TEXTVIEW("AbCdaBcD"), {TEXT('B'), TEXT('d')}, {TEXT("AbC"), TEXT("A"), TEXT("cD")},         SkipEmpty},
+			{TEXTVIEW("AbCdaBcD"), {TEXT('B'), TEXT('d')}, {TEXT("A"), TEXT("C"), TEXT("a"), TEXT("c")}, SkipEmpty | IgnoreCase},
+
+			{TEXTVIEW("A\u2022B\u2022C"), {TEXT('\u2022'), TEXT('b')}, {TEXT("A"), TEXT("B"), TEXT("C")}, SkipEmpty},
+			{TEXTVIEW("A\u2022B\u2022C"), {TEXT('\u2022'), TEXT('b')}, {TEXT("A"), TEXT("C")},            SkipEmpty | IgnoreCase},
+		}));
+
+		FStringViewArray ActualTokens;
+		if (GetNum(DelimiterArray) == 1)
+		{
+			ParseTokens(View, GetData(DelimiterArray)[0], ActualTokens, Options);
 		}
 		else
 		{
-			UE::String::ParseTokensMultiple(View, Delimiters, ResultTokens, Options);
+			ParseTokensMultiple(View, DelimiterArray, ActualTokens, Options);
 		}
 
-		TEST_TRUE(TEXT("UE::String::ParseTokens[Multiple] failed to parse"), Algo::Compare(ResultTokens, ExpectedTokens));
-
-		if (!Algo::Compare(ResultTokens, ExpectedTokens))
+		const bool bEqual = Algo::Compare(ActualTokens, ExpectedTokens);
+		if (!bEqual)
+		{
+			TStringBuilder<256> Input, Delimiters, Expected, Actual;
+			Input << QuoteEscape(View);
+			Delimiters << TEXTVIEW("[") << JoinQuoted(DelimiterArray, TEXTVIEW(", "), TEXT('\'')) << TEXTVIEW("]");
+			Expected << TEXTVIEW("[") << JoinBy(ExpectedTokens, QuoteEscape, TEXTVIEW(", ")) << TEXTVIEW("]");
+			Actual << TEXTVIEW("[") << JoinBy(ActualTokens, QuoteEscape, TEXTVIEW(", ")) << TEXTVIEW("]");
+			CAPTURE(Input, Delimiters, Expected, Actual);
+			CHECK(bEqual);
+		}
+		else
 		{
-			TStringBuilder<512> Error;
-			Error << TEXTVIEW("UE::String::ParseTokens[Multiple] failed to parse \"") << View << TEXTVIEW("\" with delimiters {");
-			Error.JoinQuoted(Delimiters, TEXTVIEW(", "), TEXTVIEW("'"));
-			Error << TEXTVIEW("} result {");
-			Error.JoinQuoted(ResultTokens, TEXTVIEW(", "), TEXTVIEW("\""));
-			Error << TEXTVIEW("} expected {");
-			Error.JoinQuoted(ExpectedTokens, TEXTVIEW(", "), TEXTVIEW("\""));
-			Error << TEXTVIEW("}");
-			FAIL_CHECK(Error.ToString());
+			CHECK(bEqual);
 		}
-	};
-
-	constexpr EOptions KeepEmpty = EOptions::None;
-	constexpr EOptions SkipEmpty = EOptions::SkipEmpty;
-	constexpr EOptions IgnoreCase = EOptions::IgnoreCase;
-	constexpr EOptions Trim = EOptions::Trim;
-
-	RunParseTokensTest(TEXT(""),         {},                       {},                                 SkipEmpty);
-	RunParseTokensTest(TEXT(""),         {},                       {TEXT("")},                         KeepEmpty);
-	RunParseTokensTest(TEXT("ABC"),      {},                       {TEXT("ABC")});
-
-	RunParseTokensTest(TEXT(""),         {TEXT(',')},              {},                                 SkipEmpty);
-	RunParseTokensTest(TEXT(""),         {TEXT(',')},              {TEXT("")},                         KeepEmpty);
-	RunParseTokensTest(TEXT(","),        {TEXT(',')},              {},                                 SkipEmpty);
-	RunParseTokensTest(TEXT(","),        {TEXT(',')},              {TEXT(""), TEXT("")},               KeepEmpty);
-	RunParseTokensTest(TEXT(",,"),       {TEXT(',')},              {},                                 SkipEmpty);
-	RunParseTokensTest(TEXT(",,"),       {TEXT(',')},              {TEXT(""), TEXT(""), TEXT("")},     KeepEmpty);
-	RunParseTokensTest(TEXT(", ,"),      {TEXT(',')},              {TEXT(" ")},                        SkipEmpty);
-	RunParseTokensTest(TEXT(", ,"),      {TEXT(',')},              {},                                 SkipEmpty | Trim);
-	RunParseTokensTest(TEXT(", ,"),      {TEXT(',')},              {TEXT(""), TEXT(" "), TEXT("")},    KeepEmpty);
-	RunParseTokensTest(TEXT(", ,"),      {TEXT(',')},              {TEXT(""), TEXT(""), TEXT("")},    KeepEmpty | Trim);
-	RunParseTokensTest(TEXT("ABC"),      {TEXT(',')},              {TEXT("ABC")});
-	RunParseTokensTest(TEXT("A,,C"),     {TEXT(',')},              {TEXT("A"), TEXT("C")},             SkipEmpty);
-	RunParseTokensTest(TEXT("A,,C"),     {TEXT(',')},              {TEXT("A"), TEXT(""), TEXT("C")},   KeepEmpty);
-	RunParseTokensTest(TEXT("A,\tB\t,C"), {TEXT(',')},             {TEXT("A"), TEXT("\tB\t"), TEXT("C")});
-	RunParseTokensTest(TEXT(",A, B ,C,"), {TEXT(',')},             {TEXT("A"), TEXT(" B "), TEXT("C")},                     SkipEmpty);
-	RunParseTokensTest(TEXT(",A, B ,C,"), {TEXT(',')},             {TEXT(""), TEXT("A"), TEXT(" B "), TEXT("C"), TEXT("")}, KeepEmpty);
-	RunParseTokensTest(TEXT("A\u2022B\u2022C"), {TEXT('\u2022')},  {TEXT("A"), TEXT("B"), TEXT("C")});
-
-	RunParseTokensTest(TEXT("ABCDABCD"), {TEXT('B'),  TEXT('D')},                  {TEXT("A"), TEXT("C"), TEXT("A"), TEXT("C")},           SkipEmpty);
-	RunParseTokensTest(TEXT("ABCDABCD"), {TEXT('B'),  TEXT('D')},                  {TEXT("A"), TEXT("C"), TEXT("A"), TEXT("C"), TEXT("")}, KeepEmpty);
-	RunParseTokensTest(TEXT("A\u2022B,C"), {TEXT(','), TEXT('\u2022')},            {TEXT("A"), TEXT("B"), TEXT("C")});
-	RunParseTokensTest(TEXT("A\u2022B\u0085C"), {TEXT('\u0085'), TEXT('\u2022')},  {TEXT("A"), TEXT("B"), TEXT("C")});
-
-	RunParseTokensTest(TEXT("ABC"), {TEXT('b')}, {TEXT("ABC")}, SkipEmpty);
-	RunParseTokensTest(TEXT("ABC"), {TEXT('b')}, {TEXT("A"), TEXT("C")}, SkipEmpty | IgnoreCase);
-	RunParseTokensTest(TEXT("AbCdaBcD"), {TEXT('B'),  TEXT('d')}, {TEXT("AbC"), TEXT("A"), TEXT("cD")}, SkipEmpty);
-	RunParseTokensTest(TEXT("AbCdaBcD"), {TEXT('B'),  TEXT('d')}, {TEXT("A"), TEXT("C"), TEXT("a"), TEXT("c")}, SkipEmpty | IgnoreCase);
-	RunParseTokensTest(TEXT("A\u2022B\u2022C"), {TEXT('\u2022'), TEXT('b')}, {TEXT("A"), TEXT("B"), TEXT("C")}, SkipEmpty);
-	RunParseTokensTest(TEXT("A\u2022B\u2022C"), {TEXT('\u2022'), TEXT('b')}, {TEXT("A"), TEXT("C")}, SkipEmpty | IgnoreCase);
+	}
 }
 
-#endif // WITH_DEV_AUTOMATION_TESTS
+} // UE::String
diff --git a/Engine/Source/Runtime/Core/Tests/Templates/ValueOrErrorTest.cpp b/Engine/Source/Runtime/Core/Tests/Templates/ValueOrErrorTest.cpp
--- a/Engine/Source/Runtime/Core/Tests/Templates/ValueOrErrorTest.cpp
+++ b/Engine/Source/Runtime/Core/Tests/Templates/ValueOrErrorTest.cpp
@@ -1,21 +1,24 @@
 // Copyright Epic Games, Inc. All Rights Reserved.
 
 #include "Templates/ValueOrError.h"
-#include "TestHarness.h"
 
+#include "TestHarness.h"
 
 #include <type_traits>
 
-static_assert(!std::is_constructible<TValueOrError<int, int>>::value, "Expected no default constructor.");
+TEST_CASE("Core::Templates::ValueOrError", "[Core][Templates][Smoke]")
+{
+	SECTION("Static")
+	{
+		STATIC_REQUIRE(!std::is_constructible<TValueOrError<int, int>>::value);
 
-static_assert(std::is_copy_constructible<TValueOrError<int, int>>::value, "Missing copy constructor.");
-static_assert(std::is_move_constructible<TValueOrError<int, int>>::value, "Missing move constructor.");
+		STATIC_REQUIRE(std::is_copy_constructible<TValueOrError<int, int>>::value);
+		STATIC_REQUIRE(std::is_move_constructible<TValueOrError<int, int>>::value);
 
-static_assert(std::is_copy_assignable<TValueOrError<int, int>>::value, "Missing copy assignment.");
-static_assert(std::is_move_assignable<TValueOrError<int, int>>::value, "Missing move assignment.");
+		STATIC_REQUIRE(std::is_copy_assignable<TValueOrError<int, int>>::value);
+		STATIC_REQUIRE(std::is_move_assignable<TValueOrError<int, int>>::value);
+	}
 
-TEST_CASE_METHOD(FAutomationTestFixture, "Core::Templates::TValueOrError::Smoke Test", "[Core][Templates][Smoke]")
-{
 	static int ValueCount = 0;
 	static int ErrorCount = 0;
 
@@ -53,139 +56,176 @@ TEST_CASE_METHOD(FAutomationTestFixture, "Core::Templates::TValueOrError::Smoke
 
 	using FTestType = TValueOrError<FTestValue, FTestError>;
 
-	// Test MakeValue Move
+	SECTION("MakeValue Move")
 	{
-		INFO("Test MakeValue Move");
 		FTestType ValueOrError = MakeValue(FTestValue());
-		CHECK_EQUAL(ValueCount, 1);
-		CHECK_EQUAL(ValueOrError.TryGetValue(), &ValueOrError.GetValue());
-		CHECK_EQUAL(ValueOrError.GetValue().Value, 1);
-		CHECK_FALSE(ValueOrError.HasError());
+		CHECK(ValueCount == 1);
+		CHECK(ValueOrError.TryGetValue() == &ValueOrError.GetValue());
+		CHECK(ValueOrError.GetValue().Value == 1);
 		CHECK(ValueOrError.HasValue());
 		CHECK(ValueOrError.TryGetError() == nullptr);
+		CHECK_FALSE(ValueOrError.HasError());
 	}
 
-	CHECK_EQUAL(ValueCount, 0);
-
-	// Test MakeValue Proxy
+	SECTION("MakeValue Proxy")
 	{
-		INFO("Test MakeValue Proxy");
 		FTestType ValueOrError = MakeValue(2, 6, 8);
-		CHECK_EQUAL(ValueCount, 1);
-		CHECK_EQUAL(ValueOrError.TryGetValue(), &ValueOrError.GetValue());
-		CHECK_EQUAL(ValueOrError.GetValue().Value, 16);
-		CHECK_FALSE(ValueOrError.HasError());
+		CHECK(ValueCount == 1);
+		CHECK(ValueOrError.TryGetValue() == &ValueOrError.GetValue());
+		CHECK(ValueOrError.GetValue().Value == 16);
 		CHECK(ValueOrError.HasValue());
 		CHECK(ValueOrError.TryGetError() == nullptr);
+		CHECK_FALSE(ValueOrError.HasError());
 	}
-	
-	CHECK_EQUAL(ValueCount, 0);
 
-	// Test StealValue
+	SECTION("StealValue")
 	{
-		INFO("Test StealValue");
 		FTestType ValueOrError = MakeValue(FTestValue());
 		FTestValue Value = ValueOrError.StealValue();
-		CHECK_EQUAL(ValueCount, 1);
-		CHECK_EQUAL(Value.Value, 1);
+		CHECK(ValueCount == 1);
+		CHECK(Value.Value == 1);
 		CHECK_FALSE(ValueOrError.HasError());
 		CHECK_FALSE(ValueOrError.HasValue());
 	}
-	
-	CHECK_EQUAL(ValueCount, 0);
 
-	// Test MakeError Move
+	SECTION("MakeError Move")
 	{
-		INFO("Test MakeError Move");
 		FTestType ValueOrError = MakeError(FTestError());
-		CHECK_EQUAL(ErrorCount, 1);
-		CHECK_EQUAL(ValueOrError.TryGetError(), &ValueOrError.GetError());
-		CHECK_EQUAL(ValueOrError.GetError().Error, 1);
-		CHECK_FALSE(ValueOrError.HasValue());
+		CHECK(ErrorCount == 1);
+		CHECK(ValueOrError.TryGetError() == &ValueOrError.GetError());
+		CHECK(ValueOrError.GetError().Error == 1);
 		CHECK(ValueOrError.HasError());
 		CHECK(ValueOrError.TryGetValue() == nullptr);
+		CHECK_FALSE(ValueOrError.HasValue());
 	}
-	
-	CHECK_EQUAL(ErrorCount, 0);
 
-	// Test MakeError Proxy
+	SECTION("MakeError Proxy")
 	{
-		INFO("Test MakeError Proxy");
 		FTestType ValueOrError = MakeError(4, 12);
-		CHECK_EQUAL(ErrorCount, 1);
-		CHECK_EQUAL(ValueOrError.TryGetError(), &ValueOrError.GetError());
-		CHECK_EQUAL(ValueOrError.GetError().Error, 16);
-		CHECK_FALSE(ValueOrError.HasValue());
+		CHECK(ErrorCount == 1);
+		CHECK(ValueOrError.TryGetError() == &ValueOrError.GetError());
+		CHECK(ValueOrError.GetError().Error == 16);
 		CHECK(ValueOrError.HasError());
 		CHECK(ValueOrError.TryGetValue() == nullptr);
+		CHECK_FALSE(ValueOrError.HasValue());
 	}
-	
-	CHECK_EQUAL(ErrorCount, 0);
 
-	// Test StealError
+	SECTION("StealError")
 	{
-		INFO("Test StealError");
 		FTestType ValueOrError = MakeError();
 		FTestError Error = ValueOrError.StealError();
-		CHECK_EQUAL(ErrorCount, 1);
-		CHECK_EQUAL(Error.Error, 1);
+		CHECK(ErrorCount == 1);
+		CHECK(Error.Error == 1);
 		CHECK_FALSE(ValueOrError.HasValue());
 		CHECK_FALSE(ValueOrError.HasError());
 	}
 
-	CHECK_EQUAL(ErrorCount, 0);
-
-	// Test Assignment
+	SECTION("Assignment")
 	{
-		INFO("Test Assignment");
 		FTestType ValueOrError = MakeValue();
 		ValueOrError = MakeValue();
-		CHECK_EQUAL(ValueCount, 1);
-		CHECK_EQUAL(ValueOrError.GetValue().Value, 2);
-		CHECK_EQUAL(ErrorCount, 0);
+		CHECK(ValueCount == 1);
+		CHECK(ValueOrError.GetValue().Value == 2);
+		CHECK(ErrorCount == 0);
 		ValueOrError = MakeError();
-		CHECK_EQUAL(ValueCount, 0);
-		CHECK_EQUAL(ErrorCount, 1);
+		CHECK(ValueCount == 0);
+		CHECK(ErrorCount == 1);
 		ValueOrError = MakeError();
-		CHECK_EQUAL(ValueCount, 0);
-		CHECK_EQUAL(ValueOrError.GetError().Error, 2);
-		CHECK_EQUAL(ErrorCount, 1);
+		CHECK(ValueCount == 0);
+		CHECK(ValueOrError.GetError().Error == 2);
+		CHECK(ErrorCount == 1);
 		ValueOrError = MakeValue();
-		CHECK_EQUAL(ValueCount, 1);
-		CHECK_EQUAL(ErrorCount, 0);
+		CHECK(ValueCount == 1);
+		CHECK(ErrorCount == 0);
 		FTestType UnsetValueOrError = MakeValue();
 		UnsetValueOrError.StealValue();
 		ValueOrError = MoveTemp(UnsetValueOrError);
-		CHECK_EQUAL(ValueCount, 0);
-		CHECK_EQUAL(ErrorCount, 0);
+		CHECK(ValueCount == 0);
+		CHECK(ErrorCount == 0);
 		CHECK_FALSE(ValueOrError.HasValue());
 		CHECK_FALSE(ValueOrError.HasError());
 	}
 
-	CHECK_EQUAL(ValueCount, 0);
-	CHECK_EQUAL(ErrorCount, 0);
-
-	// Test Move-Only Value/Error
+	SECTION("Move-Only Value/Error")
 	{
-		INFO("Move-Only Value/Error");
-		TValueOrError<FTestMoveOnly, FTestMoveOnly> Value = MakeValue(FTestMoveOnly());
-		TValueOrError<FTestMoveOnly, FTestMoveOnly> Error = MakeError(FTestMoveOnly());
+		TValueOrError<FTestMoveOnly, FTestMoveOnly> Value = MakeValue(FTestMoveOnly{1});
+		TValueOrError<FTestMoveOnly, FTestMoveOnly> Error = MakeError(FTestMoveOnly{1});
 		FTestMoveOnly MovedValue = MoveTemp(Value).GetValue();
 		FTestMoveOnly MovedError = MoveTemp(Error).GetError();
+		CHECK(MovedValue.Value == 1);
+		CHECK(MovedError.Value == 1);
 	}
 
-	// Test Integer Value/Error
+	SECTION("Integer Value/Error")
 	{
-		INFO("Test Integer Value/Error");
 		TValueOrError<int32, int32> ValueOrError = MakeValue();
-		CHECK_EQUAL(ValueOrError.GetValue(), 0);
+		CHECK(ValueOrError.GetValue() == 0);
 		ValueOrError = MakeValue(1);
-		CHECK_EQUAL(ValueOrError.GetValue(), 1);
+		CHECK(ValueOrError.GetValue() == 1);
 		ValueOrError = MakeError();
-		CHECK_EQUAL(ValueOrError.GetError(), 0);
+		CHECK(ValueOrError.GetError() == 0);
 		ValueOrError = MakeError(1);
-		CHECK_EQUAL(ValueOrError.GetError(), 1);
+		CHECK(ValueOrError.GetError() == 1);
 	}
 
-}
+	SECTION("Value/Void Value")
+	{
+		TValueOrError<FTestValue, void> ValueOrError = MakeValue(1, 2, 3);
+		CHECK(ValueCount == 1);
+		CHECK(ValueOrError.HasValue());
+		CHECK_FALSE(ValueOrError.HasError());
+		CHECK(ValueOrError.GetValue().Value == 6);
+		CHECK(ValueOrError.TryGetValue() == &ValueOrError.GetValue());
+		FTestValue Value = ValueOrError.StealValue();
+		CHECK_FALSE(ValueOrError.HasValue());
+		CHECK(Value.Value == 6);
+	}
+
+	SECTION("Value/Void Error")
+	{
+		TValueOrError<FTestValue, void> ValueOrError = MakeError();
+		CHECK(ValueCount == 0);
+		CHECK(ValueOrError.HasError());
+		CHECK_FALSE(ValueOrError.HasValue());
+		CHECK(ValueOrError.TryGetValue() == nullptr);
+	}
+
+	SECTION("Void/Error Value")
+	{
+		TValueOrError<void, FTestError> ValueOrError = MakeValue();
+		CHECK(ErrorCount == 0);
+		CHECK(ValueOrError.HasValue());
+		CHECK_FALSE(ValueOrError.HasError());
+		CHECK(ValueOrError.TryGetError() == nullptr);
+	}
 
+	SECTION("Void/Error Error")
+	{
+		TValueOrError<void, FTestError> ValueOrError = MakeError(1, 2);
+		CHECK(ErrorCount == 1);
+		CHECK(ValueOrError.HasError());
+		CHECK_FALSE(ValueOrError.HasValue());
+		CHECK(ValueOrError.GetError().Error == 3);
+		CHECK(ValueOrError.TryGetError() == &ValueOrError.GetError());
+		FTestError Error = ValueOrError.StealError();
+		CHECK_FALSE(ValueOrError.HasError());
+		CHECK(Error.Error == 3);
+	}
+
+	SECTION("Void/Void Value")
+	{
+		TValueOrError<void, void> ValueOrError = MakeValue();
+		CHECK(ValueOrError.HasValue());
+		CHECK_FALSE(ValueOrError.HasError());
+	}
+
+	SECTION("Void/Void Error")
+	{
+		TValueOrError<void, void> ValueOrError = MakeError();
+		CHECK(ValueOrError.HasError());
+		CHECK_FALSE(ValueOrError.HasValue());
+	}
+
+	CHECK(ValueCount == 0);
+	CHECK(ErrorCount == 0);
+}
