#ue5 - LWC: VectorMod() on Neon updated to match newer SSE implementation.

diff --git a/Engine/Source/Runtime/Core/Public/Math/UnrealMathNeon.h b/Engine/Source/Runtime/Core/Public/Math/UnrealMathNeon.h
--- a/Engine/Source/Runtime/Core/Public/Math/UnrealMathNeon.h
+++ b/Engine/Source/Runtime/Core/Public/Math/UnrealMathNeon.h
@@ -2715,35 +2715,31 @@ FORCEINLINE VectorRegister4Double VectorTruncate(const VectorRegister4Double& X)
 
 FORCEINLINE VectorRegister4Float VectorMod(const VectorRegister4Float& X, const VectorRegister4Float& Y)
 {
-	VectorRegister4Float Div = VectorDivide(X, Y);
-	// Floats where abs(f) >= 2^23 have no fractional portion, and larger values would overflow VectorTruncate.
-	VectorRegister4Float NoFractionMask = VectorCompareGE(VectorAbs(Div), GlobalVectorConstants::FloatNonFractional);
-	VectorRegister4Float Temp = VectorSelect(NoFractionMask, Div, VectorTruncate(Div));
-	VectorRegister4Float Result = VectorNegateMultiplyAdd(Y, Temp, X);
-	// Clamp to [-AbsY, AbsY] because of possible failures for very large numbers (>1e10) due to precision loss.
-	VectorRegister4Float AbsY = VectorAbs(Y);
-	return vmaxnmq_f32(VectorNegate(AbsY), vminnmq_f32(Result, AbsY));
+	// Check against invalid divisor
+	VectorRegister4Float InvalidDivisorMask = VectorCompareLE(VectorAbs(Y), GlobalVectorConstants::SmallNumber);
+	// Promote to doubles for better intermediate precision.
+	VectorRegister4Double DoubleX = MakeVectorRegisterDouble(X);
+	VectorRegister4Double DoubleY = MakeVectorRegisterDouble(Y);
+	// R = X - (Y * (Trunc(X / Y))
+	VectorRegister4Double Temp = VectorTruncate(VectorDivide(DoubleX, DoubleY));
+	VectorRegister4Double DoubleResult = VectorNegateMultiplyAdd(DoubleY, Temp, DoubleX);
+	// Convert back to floats. This is safe (in terms of not exceeding [-FLT_MAX,FLT_MAX]) because the answer is less than or equal to Abs(X) by definition.
+	VectorRegister4Float Result = MakeVectorRegisterFloatFromDouble(DoubleResult);
+	// Return 0 where divisor Y was too small	
+	Result = VectorSelect(InvalidDivisorMask, GlobalVectorConstants::FloatZero, Result);
+	return Result;
 }
 
 FORCEINLINE VectorRegister4Double VectorMod(const VectorRegister4Double& X, const VectorRegister4Double& Y)
 {
-	VectorRegister4Double Div = VectorDivide(X, Y);
-	// Floats where abs(f) >= 2^52 have no fractional portion, and larger values would overflow VectorTruncate.
-	VectorRegister4Double NoFractionMask = VectorCompareGE(VectorAbs(Div), GlobalVectorConstants::DoubleNonFractional);
-	VectorRegister4Double Temp = VectorSelect(NoFractionMask, Div, VectorTruncate(Div));
-	VectorRegister4Double Result = VectorNegateMultiplyAdd(Y, Temp, X);
-	// Clamp to [-AbsY, AbsY] because of possible failures for very large numbers (>1e10) due to precision loss.
-	VectorRegister4Double AbsY = VectorAbs(Y);
-	VectorRegister4Double NegAbsY = VectorNegate(AbsY);
-	
-	VectorRegister4Double Min;
-	Min.XY = vminnmq_f64(Result.XY, AbsY.XY);
-	Min.ZW = vminnmq_f64(Result.ZW, AbsY.ZW);
-	
-	VectorRegister4Double Max;
-	Max.XY = vmaxnmq_f64(NegAbsY.XY, Min.XY);
-	Max.ZW = vmaxnmq_f64(NegAbsY.ZW, Min.ZW);
-	return Max;
+	// Check against invalid divisor
+	VectorRegister4Double InvalidDivisorMask = VectorCompareLE(VectorAbs(Y), GlobalVectorConstants::DoubleSmallNumber);
+	// R = X - (Y * (Trunc(X / Y))
+	VectorRegister4Double Temp = VectorTruncate(VectorDivide(X, Y));
+	VectorRegister4Double DoubleResult = VectorNegateMultiplyAdd(Y, Temp, X);
+	// Return 0 where divisor Y was too small
+	VectorRegister4Double Result = VectorSelect(InvalidDivisorMask, GlobalVectorConstants::DoubleZero, DoubleResult);
+	return Result;
 }
 
 FORCEINLINE VectorRegister4Float VectorSign(const VectorRegister4Float& X)
