Added append to FString overloads for BytesToHex[Lower], FSHAHash, FGuid and optimized them a bit

diff --git a/Engine/Source/Runtime/Core/Private/Misc/Guid.cpp b/Engine/Source/Runtime/Core/Private/Misc/Guid.cpp
--- a/Engine/Source/Runtime/Core/Private/Misc/Guid.cpp
+++ b/Engine/Source/Runtime/Core/Private/Misc/Guid.cpp
@@ -3,6 +3,8 @@
 #include "Misc/Guid.h"
 #include "Algo/Replace.h"
 #include "Containers/ArrayView.h"
+#include "Containers/StaticArray.h"
+#include "Misc/ByteSwap.h"
 #include "Misc/Parse.h"
 #include "Math/Int128.h"
 #include "UObject/PropertyPortFlags.h"
@@ -44,41 +46,50 @@ bool FGuid::ImportTextItem(const TCHAR*& Buffer, int32 PortFlags, class UObject*
 }
 
 
-FString FGuid::ToString(EGuidFormats Format) const
+void FGuid::AppendString(FString& Out, EGuidFormats Format) const
 {
 	switch (Format)
 	{
 	case EGuidFormats::DigitsWithHyphens:
-		return FString::Printf(TEXT("%08X-%04X-%04X-%04X-%04X%08X"), A, B >> 16, B & 0xFFFF, C >> 16, C & 0xFFFF, D);
+		Out.Appendf(TEXT("%08X-%04X-%04X-%04X-%04X%08X"), A, B >> 16, B & 0xFFFF, C >> 16, C & 0xFFFF, D);
+		return;
 
 	case EGuidFormats::DigitsWithHyphensLower:
-		return FString::Printf(TEXT("%08x-%04x-%04x-%04x-%04x%08x"), A, B >> 16, B & 0xFFFF, C >> 16, C & 0xFFFF, D);
+		Out.Appendf(TEXT("%08x-%04x-%04x-%04x-%04x%08x"), A, B >> 16, B & 0xFFFF, C >> 16, C & 0xFFFF, D);
+		return;
 
 	case EGuidFormats::DigitsWithHyphensInBraces:
-		return FString::Printf(TEXT("{%08X-%04X-%04X-%04X-%04X%08X}"), A, B >> 16, B & 0xFFFF, C >> 16, C & 0xFFFF, D);
+		Out.Appendf(TEXT("{%08X-%04X-%04X-%04X-%04X%08X}"), A, B >> 16, B & 0xFFFF, C >> 16, C & 0xFFFF, D);
+		return;
 
 	case EGuidFormats::DigitsWithHyphensInParentheses:
-		return FString::Printf(TEXT("(%08X-%04X-%04X-%04X-%04X%08X)"), A, B >> 16, B & 0xFFFF, C >> 16, C & 0xFFFF, D);
+		Out.Appendf(TEXT("(%08X-%04X-%04X-%04X-%04X%08X)"), A, B >> 16, B & 0xFFFF, C >> 16, C & 0xFFFF, D);
+		return;
 
 	case EGuidFormats::HexValuesInBraces:
-		return FString::Printf(TEXT("{0x%08X,0x%04X,0x%04X,{0x%02X,0x%02X,0x%02X,0x%02X,0x%02X,0x%02X,0x%02X,0x%02X}}"), A, B >> 16, B & 0xFFFF, C >> 24, (C >> 16) & 0xFF, (C >> 8) & 0xFF, C & 0XFF, D >> 24, (D >> 16) & 0XFF, (D >> 8) & 0XFF, D & 0XFF);
+		Out.Appendf(TEXT("{0x%08X,0x%04X,0x%04X,{0x%02X,0x%02X,0x%02X,0x%02X,0x%02X,0x%02X,0x%02X,0x%02X}}"), A, B >> 16, B & 0xFFFF, C >> 24, (C >> 16) & 0xFF, (C >> 8) & 0xFF, C & 0XFF, D >> 24, (D >> 16) & 0XFF, (D >> 8) & 0XFF, D & 0XFF);
+		return;
 
 	case EGuidFormats::UniqueObjectGuid:
-		return FString::Printf(TEXT("%08X-%08X-%08X-%08X"), A, B, C, D);
+		Out.Appendf(TEXT("%08X-%08X-%08X-%08X"), A, B, C, D);
+		return;
 
 	case EGuidFormats::Short:
 	{
 		const uint32 Data[] = {A,B,C,D};
-		FString Result = FBase64::Encode(reinterpret_cast<const uint8*>(&Data), sizeof(Data));
-
-		Result.ReplaceCharInline(TEXT('+'), TEXT('-'), ESearchCase::CaseSensitive);
-		Result.ReplaceCharInline(TEXT('/'), TEXT('_'), ESearchCase::CaseSensitive);
+		TCHAR Buffer[25];
+		int32 Len = FBase64::Encode(reinterpret_cast<const uint8*>(Data), sizeof(Data), Buffer);
+		TArrayView<TCHAR> Result(Buffer, Len);
+		
+		Algo::Replace(Result, '+', '-');
+		Algo::Replace(Result, '/', '_');
 
 		// Remove trailing '=' base64 padding
-		check(Result.Len() == 24);
-		Result.RemoveAt(22, 2, false);
+		check(Len == 24);
+		Result.LeftChopInline(2);
 
-		return Result;
+		Out.Append(Result);
+		return;
 	}
 
 	case EGuidFormats::Base36Encoded:
@@ -92,33 +103,28 @@ FString FGuid::ToString(EGuidFormats Format) const
 
 		FUInt128 Value(A, B, C, D);
 
-		FString Result;
-		Result.Reserve(26);
-		TArray<TCHAR, FString::AllocatorType>& OutCharArray = Result.GetCharArray();
-
-		while(Value > 0)
+		TStaticArray<TCHAR, 25> Result;
+		for (int32 I = 24; I >= 0; --I)
 		{
 			uint32 Remainder;
 			Value = Value.Divide(36, Remainder);
-			OutCharArray.Add(Alphabet[Remainder]);
-		}
-
-		for(int32 i=OutCharArray.Num(); i<25; i++)
-		{
-			OutCharArray.Add(TEXT('0'));
+			Result[I] = Alphabet[Remainder];
 		}
 
-		OutCharArray.Add(TCHAR('\0'));
-
-		check(Result.Len() == 25);
-		return Result.Reverse();
+		Out.Append(Result);
+		return;
 	}
 
 	case EGuidFormats::DigitsLower:
-		return FString::Printf(TEXT("%08x%08x%08x%08x"), A, B, C, D);
+		Out.Appendf(TEXT("%08x%08x%08x%08x"), A, B, C, D);
+		return;
 
 	default:
-		return FString::Printf(TEXT("%08X%08X%08X%08X"), A, B, C, D);
+	{
+		uint32 Bytes[4] = { NETWORK_ORDER32(A), NETWORK_ORDER32(B), NETWORK_ORDER32(C), NETWORK_ORDER32(D) };
+		BytesToHex(reinterpret_cast<const uint8*>(Bytes), sizeof(Bytes), Out);
+		return;
+	}
 	}
 }
 
diff --git a/Engine/Source/Runtime/Core/Private/String/BytesToHex.cpp b/Engine/Source/Runtime/Core/Private/String/BytesToHex.cpp
--- a/Engine/Source/Runtime/Core/Private/String/BytesToHex.cpp
+++ b/Engine/Source/Runtime/Core/Private/String/BytesToHex.cpp
@@ -1,7 +1,7 @@
 // Copyright Epic Games, Inc. All Rights Reserved.
 
 #include "String/BytesToHex.h"
-
+#include "Containers/UnrealString.h"
 #include "Misc/StringBuilder.h"
 
 namespace UE::String
@@ -87,3 +87,27 @@ void BytesToHexLower(TConstArrayView<uint8> Bytes, FUtf8StringBuilderBase& Build
 }
 
 } // UE::String
+
+//////////////////// Functions declared in UnrealString.h ////////////////////
+
+void BytesToHex(const uint8* In, int32 Count, FString& Out)
+{
+	if (Count)
+	{
+		TArray<TCHAR>& OutArray = Out.GetCharArray();
+		OutArray.AddUninitialized(2 * Count + /* add null terminator */ (OutArray.Num() == 0));
+		UE::String::BytesToHex(MakeArrayView(In, Count), &OutArray.Last(2 * Count));
+		OutArray.Last() = '\0';
+	}
+}
+
+void BytesToHexLower(const uint8* In, int32 Count, FString& Out)
+{
+	if (Count)
+	{
+		TArray<TCHAR>& OutArray = Out.GetCharArray();
+		OutArray.AddUninitialized(2 * Count + /* add null terminator */ (OutArray.Num() == 0));
+		UE::String::BytesToHexLower(MakeArrayView(In, Count), &OutArray.Last(2 * Count));
+		OutArray.Last() = '\0';
+	}
+}
diff --git a/Engine/Source/Runtime/Core/Public/Containers/UnrealString.h b/Engine/Source/Runtime/Core/Public/Containers/UnrealString.h
--- a/Engine/Source/Runtime/Core/Public/Containers/UnrealString.h
+++ b/Engine/Source/Runtime/Core/Public/Containers/UnrealString.h
@@ -1933,23 +1933,25 @@ inline void ByteToHex(uint8 In, FString& Result)
 	Result += NibbleToTChar(In & 15);
 }
 
-/** 
- * Convert an array of bytes to hex
- * @param In byte array values to convert
- * @param Count number of bytes to convert
- * @return Hex value in string.
- */
-UE_NODISCARD inline FString BytesToHex(const uint8* In, int32 Count)
+/** Append bytes as uppercase hex string */
+CORE_API void BytesToHex(const uint8* Bytes, int32 NumBytes, FString& Out);
+/** Append bytes as lowercase hex string */
+CORE_API void BytesToHexLower(const uint8* Bytes, int32 NumBytes, FString& Out);
+
+/** Convert bytes to uppercase hex string */
+UE_NODISCARD inline FString BytesToHex(const uint8* Bytes, int32 NumBytes)
 {
-	FString Result;
-	Result.Empty(Count * 2);
+	FString Out;
+	BytesToHex(Bytes, NumBytes, Out);
+	return Out;
+}
 
-	while (Count)
-	{
-		ByteToHex(*In++, Result);
-		Count--;
-	}
-	return Result;
+/** Convert bytes to lowercase hex string */
+UE_NODISCARD inline FString BytesToHexLower(const uint8* Bytes, int32 NumBytes)
+{
+	FString Out;
+	BytesToHexLower(Bytes, NumBytes, Out);
+	return Out;
 }
 
 /**
diff --git a/Engine/Source/Runtime/Core/Public/Misc/Guid.h b/Engine/Source/Runtime/Core/Public/Misc/Guid.h
--- a/Engine/Source/Runtime/Core/Public/Misc/Guid.h
+++ b/Engine/Source/Runtime/Core/Public/Misc/Guid.h
@@ -306,20 +306,22 @@ public:
 	/**
 	 * Converts this GUID to its string representation.
 	 *
+	 * @param Format The string format to use.
 	 * @return The string representation.
 	 */
-	FString ToString() const
+	FString ToString(EGuidFormats Format = EGuidFormats::Digits) const
 	{
-		return ToString(EGuidFormats::Digits);
+		FString Out;
+		AppendString(Out, Format);
+		return Out;
 	}
 
 	/**
 	 * Converts this GUID to its string representation using the specified format.
 	 *
 	 * @param Format The string format to use.
-	 * @return The string representation.
 	 */
-	CORE_API FString ToString(EGuidFormats Format) const;
+	CORE_API void AppendString(FString& Out, EGuidFormats Format = EGuidFormats::Digits) const;
 
 	/**
 	 * Appends this GUID to the string builder using the specified format.
diff --git a/Engine/Source/Runtime/Core/Public/Misc/SecureHash.h b/Engine/Source/Runtime/Core/Public/Misc/SecureHash.h
--- a/Engine/Source/Runtime/Core/Public/Misc/SecureHash.h
+++ b/Engine/Source/Runtime/Core/Public/Misc/SecureHash.h
@@ -207,6 +207,11 @@ public:
 
 	static constexpr int32 GetStringLen() { return UE_ARRAY_COUNT(Hash) * 2; }
 
+	inline void AppendString(FString& Out) const
+	{
+		BytesToHex((const uint8*)Hash, sizeof(Hash), Out);
+	}
+
 	inline FString ToString() const
 	{
 		return BytesToHex((const uint8*)Hash, sizeof(Hash));
