Catch2 Conversion of Core Tests

diff --git a/Engine/Source/Developer/LowLevelTestsRunner/Private/TestCommon/CoreUtilities.cpp b/Engine/Source/Developer/LowLevelTestsRunner/Private/TestCommon/CoreUtilities.cpp
--- a/Engine/Source/Developer/LowLevelTestsRunner/Private/TestCommon/CoreUtilities.cpp
+++ b/Engine/Source/Developer/LowLevelTestsRunner/Private/TestCommon/CoreUtilities.cpp
@@ -93,6 +93,7 @@ void InitAllThreadPools(bool MultiThreaded)
 
 void InitTaskGraph(bool MultiThreaded, ENamedThreads::Type ThreadToAttach)
 {
+	FTaskTagScope Scope(ETaskTag::EGameThread);
 	FTaskGraphInterface::Startup(MultiThreaded ? FPlatformMisc::NumberOfWorkerThreadsToSpawn() : 1);
 	FTaskGraphInterface::Get().AttachToThread(ThreadToAttach);
 }
diff --git a/Engine/Source/Runtime/Core/Tests/Async/AsyncWorkTest.cpp b/Engine/Source/Runtime/Core/Tests/Async/AsyncWorkTest.cpp
--- a/Engine/Source/Runtime/Core/Tests/Async/AsyncWorkTest.cpp
+++ b/Engine/Source/Runtime/Core/Tests/Async/AsyncWorkTest.cpp
@@ -81,7 +81,7 @@ namespace AsyncWorkTestImpl
 	};
 }
 
-TEST_CASE_METHOD(FCoreTestFixture, "Core::Async::AsyncWork::Ensure Completion BusyWait Deadlock Test", "[Core][Async][Smoke]")
+TEST_CASE_METHOD(FCoreTestFixture, "Core::Async::Task::Ensure Completion BusyWait Deadlock Test", "[Core][Async][Smoke]")
 {
 	using namespace AsyncWorkTestImpl;
 	TRACE_CPUPROFILER_EVENT_SCOPE(FAsyncWorkEnsureCompletionBusyWaitDeadLockTest::RunTest);
diff --git a/Engine/Source/Runtime/Core/Tests/HAL/PlatformAtomicsTest.cpp b/Engine/Source/Runtime/Core/Tests/HAL/PlatformAtomicsTest.cpp
--- a/Engine/Source/Runtime/Core/Tests/HAL/PlatformAtomicsTest.cpp
+++ b/Engine/Source/Runtime/Core/Tests/HAL/PlatformAtomicsTest.cpp
@@ -21,8 +21,9 @@ bool TestInterlocked(volatile T* Dest, T ExpectedReturnValue, T ExpectedFinalVal
 	bool bSuccess = true;
 	const T ReturnValue = InterlockedFunc(Dest);
 	
-	TEST_TRUE(FString::Printf(TEXT("FPlatformAtomics::Interlocked%s on %s failed"), FunctionName, TypeName), ReturnValue == ExpectedReturnValue);
-	TEST_TRUE(FString::Printf(TEXT("FPlatformAtomics::Interlocked%s on %s failed"), FunctionName, TypeName), FPlatformAtomics::AtomicRead(Dest) == ExpectedFinalValue );
+	INFO(FString::Printf(TEXT("FPlatformAtomics::Interlocked%s on %s failed"), FunctionName, TypeName));
+	CHECK(ReturnValue == ExpectedReturnValue);
+	CHECK(FPlatformAtomics::AtomicRead(Dest) == ExpectedFinalValue );
 
 	return bSuccess;
 }
diff --git a/Engine/Source/Runtime/Core/Tests/HAL/PlatformStringTest.cpp b/Engine/Source/Runtime/Core/Tests/HAL/PlatformStringTest.cpp
--- a/Engine/Source/Runtime/Core/Tests/HAL/PlatformStringTest.cpp
+++ b/Engine/Source/Runtime/Core/Tests/HAL/PlatformStringTest.cpp
@@ -16,27 +16,27 @@ TEST_CASE_METHOD(FAutomationTestFixture, "Core::HAL::PlatformString::GetVarArgs"
 {
 	TCHAR Buffer[128];
 	InvokePlatformStringGetVarArgs(Buffer, TEXT("A%.*sZ"), 4, TEXT(" to B"));
-	TEST_TRUE(TEXT("GetVarArgs(%.*s)"), !FCString::Strcmp(Buffer, TEXT("A to Z") ));
+	CHECK(!FCString::Strcmp(Buffer, TEXT("A to Z")));
 }
 
 TEST_CASE_METHOD(FAutomationTestFixture, "Core::HAL::PlatformString::Strnlen", "[Core][HAL][Smoke]")
 {
-	TEST_EQUAL(TEXT("Strnlen(nullptr, 0)"), FPlatformString::Strnlen((const ANSICHAR*)nullptr, 0), 0);
-	TEST_EQUAL(TEXT("Strnlen(\"\", 0)"), FPlatformString::Strnlen("", 0), 0);
-	TEST_EQUAL(TEXT("Strnlen(\"1\", 0)"), FPlatformString::Strnlen("1", 0), 0);
-	TEST_EQUAL(TEXT("Strnlen(\"1\", 1)"), FPlatformString::Strnlen("1", 1), 1);
-	TEST_EQUAL(TEXT("Strnlen(\"1\", 2)"), FPlatformString::Strnlen("1", 2), 1);
-	TEST_EQUAL(TEXT("Strnlen(\"123\", 2)"), FPlatformString::Strnlen("123", 2), 2);
+	CHECK_EQUAL(FPlatformString::Strnlen((const ANSICHAR*)nullptr, 0), 0);
+	CHECK_EQUAL(FPlatformString::Strnlen("", 0), 0);
+	CHECK_EQUAL(FPlatformString::Strnlen("1", 0), 0);
+	CHECK_EQUAL(FPlatformString::Strnlen("1", 1), 1);
+	CHECK_EQUAL(FPlatformString::Strnlen("1", 2), 1);
+	CHECK_EQUAL(FPlatformString::Strnlen("123", 2), 2);
 	ANSICHAR AnsiBuffer[128] = "123456789";
-	TEST_EQUAL(TEXT("Strnlen(PaddedBuffer)"), FPlatformString::Strnlen(AnsiBuffer, UE_ARRAY_COUNT(AnsiBuffer)), 9);
+	CHECK_EQUAL(FPlatformString::Strnlen(AnsiBuffer, UE_ARRAY_COUNT(AnsiBuffer)), 9);
 
-	TEST_EQUAL(TEXT("Strnlen(nullptr, 0)"), FPlatformString::Strnlen((const TCHAR*)nullptr, 0), 0); //-V575
-	TEST_EQUAL(TEXT("Strnlen(\"\", 0)"), FPlatformString::Strnlen(TEXT(""), 0), 0); //-V575
-	TEST_EQUAL(TEXT("Strnlen(\"1\", 0)"), FPlatformString::Strnlen(TEXT("1"), 0), 0); //-V575
-	TEST_EQUAL(TEXT("Strnlen(\"1\", 1)"), FPlatformString::Strnlen(TEXT("1"), 1), 1);
-	TEST_EQUAL(TEXT("Strnlen(\"1\", 2)"), FPlatformString::Strnlen(TEXT("1"), 2), 1);
-	TEST_EQUAL(TEXT("Strnlen(\"123\", 2)"), FPlatformString::Strnlen(TEXT("123"), 2), 2);
+	CHECK_EQUAL(FPlatformString::Strnlen((const TCHAR*)nullptr, 0), 0); //-V575
+	CHECK_EQUAL(FPlatformString::Strnlen(TEXT(""), 0), 0); //-V575
+	CHECK_EQUAL(FPlatformString::Strnlen(TEXT("1"), 0), 0); //-V575
+	CHECK_EQUAL(FPlatformString::Strnlen(TEXT("1"), 1), 1);
+	CHECK_EQUAL(FPlatformString::Strnlen(TEXT("1"), 2), 1);
+	CHECK_EQUAL(FPlatformString::Strnlen(TEXT("123"), 2), 2);
 	TCHAR Buffer[128] = {};
 	FCString::Strcpy(Buffer, TEXT("123456789"));
-	TEST_EQUAL(TEXT("Strnlen(PaddedBuffer)"), FPlatformString::Strnlen(Buffer, UE_ARRAY_COUNT(Buffer)), 9);
+	CHECK_EQUAL(FPlatformString::Strnlen(Buffer, UE_ARRAY_COUNT(Buffer)), 9);
 }
diff --git a/Engine/Source/Runtime/Core/Tests/HAL/ThreadTest.cpp b/Engine/Source/Runtime/Core/Tests/HAL/ThreadTest.cpp
--- a/Engine/Source/Runtime/Core/Tests/HAL/ThreadTest.cpp
+++ b/Engine/Source/Runtime/Core/Tests/HAL/ThreadTest.cpp
@@ -16,7 +16,8 @@ namespace
 	void TestIsJoinableAfterCreation(FAutomationTestFixture& Test)
 	{
 		FThread Thread(TEXT("Test.Thread.TestIsJoinableAfterCreation"), []() { /*NOOP*/ });
-		TEST_TRUE(TEXT("FThread must be joinable after construction"), Thread.IsJoinable());
+		INFO("FThread must be joinable after construction");
+		CHECK(Thread.IsJoinable());
 		Thread.Join();
 		UE_LOG(LogTemp, Log, TEXT("%s completed"), StringCast<TCHAR>(__FUNCTION__).Get());
 	}
@@ -26,7 +27,8 @@ namespace
 		FThreadSafeBool bDone = false;
 		FThread Thread(TEXT("Test.Thread.TestIsJoinableAfterCompletion"), [&bDone]() { bDone = true; });
 		while (!bDone) {}; // wait for completion //-V529
-		TEST_TRUE(TEXT("FThread must still be joinable after completion"), Thread.IsJoinable());
+		INFO("FThread must still be joinable after completion");
+		CHECK(Thread.IsJoinable());
 		Thread.Join();
 		UE_LOG(LogTemp, Log, TEXT("%s completed"), StringCast<TCHAR>(__FUNCTION__).Get());
 	}
@@ -35,7 +37,8 @@ namespace
 	{
 		FThread Thread(TEXT("Test.Thread.TestIsNotJoinableAfterJoining"), []() { /*NOOP*/ });
 		Thread.Join();
-		TEST_FALSE(TEXT("FThread must not be joinable after joining"), Thread.IsJoinable());
+		INFO("FThread must not be joinable after joining");
+		CHECK_FALSE(Thread.IsJoinable());
 		UE_LOG(LogTemp, Log, TEXT("%s completed"), StringCast<TCHAR>(__FUNCTION__).Get());
 	}
 
@@ -51,7 +54,8 @@ namespace
 				});
 			Thread.Detach();
 			bReady = true; // make sure `Detach` is called before thread function exit
-			Test.TEST_FALSE(TEXT("FThread must not be joinable after detaching"), Thread.IsJoinable());
+			INFO("FThread must not be joinable after detaching");
+			CHECK_FALSE(Thread.IsJoinable());
 		}
 		// or thread function is completed fast and `FThreadImpl` releases the reference to itself before
 		// `Detach` call
@@ -61,7 +65,8 @@ namespace
 			FPlatformProcess::Sleep(0.1); // let the thread exit before detaching
 			Thread.Detach();
 			bReady = true; // make sure `Detach` is called before thread function exit
-			Test.TEST_FALSE(TEXT("FThread must not be joinable after detaching"), Thread.IsJoinable());
+			INFO("FThread must not be joinable after detaching");
+			CHECK_FALSE(Thread.IsJoinable());
 		}
 		UE_LOG(LogTemp, Log, TEXT("%s completed"), StringCast<TCHAR>(__FUNCTION__).Get());
 	}
@@ -78,12 +83,14 @@ namespace
 	{
 		{
 			FThread Thread;
-			TEST_FALSE(TEXT("Default-constructed FThread must be not joinable"), Thread.IsJoinable());
+			INFO("Default-constructed FThread must be not joinable");
+			CHECK_FALSE(Thread.IsJoinable());
 		}
 		{	// check that default constructed thread can be "upgraded" to joinable thread
 			FThread Thread;
 			Thread = FThread(TEXT("Test.Thread.TestDefaultConstruction"), []() { /* NOOP */ });
-			TEST_TRUE(TEXT("Move-constructed FThread from joinable thread must be joinable"), Thread.IsJoinable());
+			INFO("Move-constructed FThread from joinable thread must be joinable");
+			CHECK(Thread.IsJoinable());
 			Thread.Join();
 		}
 		UE_LOG(LogTemp, Log, TEXT("%s completed"), StringCast<TCHAR>(__FUNCTION__).Get());
@@ -111,8 +118,9 @@ namespace
 					FThreadSingletonTest::Get().SetTestField(2);
 				});
 			Thread.Join();
-			TEST_TRUE(TEXT("Thread singleton should be uninitialized by default"), bDefaultValuePass);
-			TEST_TRUE(TEXT("Thread singletons in different threads should be isolated"), FThreadSingletonTest::Get().GetTestField() == 1);
+
+			CHECK(bDefaultValuePass);
+			CHECK(FThreadSingletonTest::Get().GetTestField() == 1);
 		}
 		{	// check that ThreadSingleton entries don't point to invalid memory after cleanup
 			class FThreadSingletonFirst : public TThreadSingleton<FThreadSingletonFirst>
@@ -135,33 +143,36 @@ namespace
 					FThreadSingletonSecond::Get();
 				});
 			Thread.Join();
-			TEST_TRUE(TEXT("Thread singletons should be uninitialized by default"), FThreadSingletonFirst::TryGet() == nullptr);
-			TEST_TRUE(TEXT("Thread singletons should be uninitialized by default"), FThreadSingletonSecond::TryGet() == nullptr);
+			CHECK(FThreadSingletonFirst::TryGet() == nullptr);
+			CHECK(FThreadSingletonSecond::TryGet() == nullptr);
 		}
 		UE_LOG(LogTemp, Log, TEXT("%s completed"), StringCast<TCHAR>(__FUNCTION__).Get());
 	}
 
 	void TestMovability(FAutomationTestFixture& Test)
 	{
-		{	// move constructor with default-constructed thread
+		{	
+			INFO("move constructor with default - constructed thread");
 			FThread Src;
 			FThread Dst(MoveTemp(Src));
-			TEST_FALSE(TEXT("Default-constructed thread must stay not joinable after moving out"), Src.IsJoinable());
-			TEST_FALSE(TEXT("Move-constructed thread from not joinable thread must be not joinable"), Dst.IsJoinable());
+			CHECK_FALSE( Src.IsJoinable());
+			CHECK_FALSE(Dst.IsJoinable());
 		}
-		{	// move constructor with joinable thread
+		{	
+			INFO("move constructor with joinable thread");
 			FThread Src(TEXT("Test.Thread.TestMovability.1"), []() { /* NOOP */ });
 			FThread Dst(MoveTemp(Src));
-			TEST_FALSE(TEXT("Moved out thread must be not joinable"), Src.IsJoinable());
-			TEST_TRUE(TEXT("Move-constructed thread from joinable thread must be joinable"), Dst.IsJoinable());
+			CHECK_FALSE(Src.IsJoinable());
+			CHECK(Dst.IsJoinable());
 			Dst.Join();
 		}
-		{	// move assignment operator
+		{	
+			INFO("move assignment operator");
 			FThread Src(TEXT("Test.Thread.TestMovability.2"), []() { /* NOOP */ });
 			FThread Dst;
 			Dst = MoveTemp(Src);
-			TEST_FALSE(TEXT("Moved out thread must be not joinable"), Src.IsJoinable());
-			TEST_TRUE(TEXT("Move-assigned thread from joinable thread must be joinable"), Dst.IsJoinable());
+			CHECK_FALSE(Src.IsJoinable());
+			CHECK(Dst.IsJoinable());
 			Dst.Join();
 		}
 		{	// Failure test for move assignment operator of joinable thread
@@ -170,7 +181,8 @@ namespace
 			//Dst = MoveTemp(Src); // must assert that joinable thread wasn't joined before move-assignment, no way to test this
 			//Dst.Join();
 		}
-		{	// Move assignment operator of thread that has been joined
+		{	
+			INFO("Move assignment operator of thread that has been joined");
 			FThread Src(TEXT("Test.Thread.TestMovability.3"), []() { /* NOOP */ });
 			FThread Dst(TEXT("Test.Thread.TestMovability.4"), []() { /* NOOP */ });
 			Dst.Join();
diff --git a/Engine/Source/Runtime/Core/Tests/Hash/Blake3Test.cpp b/Engine/Source/Runtime/Core/Tests/Hash/Blake3Test.cpp
--- a/Engine/Source/Runtime/Core/Tests/Hash/Blake3Test.cpp
+++ b/Engine/Source/Runtime/Core/Tests/Hash/Blake3Test.cpp
@@ -5,7 +5,7 @@
 #include "Misc/StringBuilder.h"
 #include "TestHarness.h"
 
-TEST_CASE("Core::Hash::Blake3", "[Core][Hash][Smoke]")
+TEST_CASE("Core::Hash::FBlake3::Smoke Test", "[Core][Hash][Smoke]")
 {
 	struct FBlake3TestCase
 	{
diff --git a/Engine/Source/Runtime/Core/Tests/Misc/StringBuilderTest.cpp b/Engine/Source/Runtime/Core/Tests/Misc/StringBuilderTest.cpp
--- a/Engine/Source/Runtime/Core/Tests/Misc/StringBuilderTest.cpp
+++ b/Engine/Source/Runtime/Core/Tests/Misc/StringBuilderTest.cpp
@@ -6,7 +6,7 @@
 #include "String/Find.h"
 #include "TestHarness.h"
 
-TEST_CASE("Core::String::StringBuilder", "[Core][String][Smoke]")
+TEST_CASE("Core::String::Misc::StringBuilder", "[Core][String][Smoke]")
 {
 	SECTION("Static")
 	{
diff --git a/Engine/Source/Runtime/Core/Tests/Serialization/CompactBinaryTest.cpp b/Engine/Source/Runtime/Core/Tests/Serialization/CompactBinaryTest.cpp
--- a/Engine/Source/Runtime/Core/Tests/Serialization/CompactBinaryTest.cpp
+++ b/Engine/Source/Runtime/Core/Tests/Serialization/CompactBinaryTest.cpp
@@ -67,50 +67,59 @@ struct FCbAttachmentAccessors
 	static constexpr FIoHash (FCbFieldView::*AsType)(const FIoHash&) = &FCbFieldView::AsAttachment;
 };
 
-void TEST_EQUALCb(FAutomationTestFixture& Test, const FString& What, const FCbArrayView& Actual, const FCbArrayView& Expected)
+void TestEqualCb(FAutomationTestFixture& Test, const FString& What, const FCbArrayView& Actual, const FCbArrayView& Expected)
 {
-	TEST_TRUE(What, Actual.Equals(Expected));
+	INFO(What);
+	CHECK(Actual.Equals(Expected));
 }
 
-void TEST_EQUALCb(FAutomationTestFixture& Test, const FString& What, const FCbObjectView& Actual, const FCbObjectView& Expected)
+void TestEqualCb(FAutomationTestFixture& Test, const FString& What, const FCbObjectView& Actual, const FCbObjectView& Expected)
 {
-	TEST_TRUE(What, Actual.Equals(Expected));
+	INFO(What);
+	CHECK(Actual.Equals(Expected));
 }
 
-void TEST_EQUALCb(FAutomationTestFixture& Test, const FString& What, const FCbCustomById& Actual, const FCbCustomById& Expected)
+void TestEqualCb(FAutomationTestFixture& Test, const FString& What, const FCbCustomById& Actual, const FCbCustomById& Expected)
 {
-	TEST_TRUE(What, Actual.Id == Expected.Id && Actual.Data == Expected.Data);
+	INFO(What);
+	CHECK(Actual.Id == Expected.Id);
+	CHECK(Actual.Data == Expected.Data);
 }
 
-void TEST_EQUALCb(FAutomationTestFixture& Test, const FString& What, const FCbCustomByName& Actual, const FCbCustomByName& Expected)
+void TestEqualCb(FAutomationTestFixture& Test, const FString& What, const FCbCustomByName& Actual, const FCbCustomByName& Expected)
 {
-	TEST_TRUE(What, Actual.Name.Equals(Expected.Name, ESearchCase::CaseSensitive) && Actual.Data == Expected.Data);
+	INFO(What);
+	CHECK(Actual.Name.Equals(Expected.Name, ESearchCase::CaseSensitive));
+	CHECK(Actual.Data == Expected.Data);
 }
 
-void TEST_EQUALBytes(FAutomationTestFixture& Test, const TCHAR* What, FMemoryView Actual, TArrayView<const uint8> Expected)
+void TestEqualBytes(FAutomationTestFixture& Test, const TCHAR* What, FMemoryView Actual, TArrayView<const uint8> Expected)
 {
-	TEST_TRUE(What, Actual.EqualBytes(MakeMemoryView(Expected)));
+	INFO(What);
+	CHECK(Actual.EqualBytes(MakeMemoryView(Expected)));
 }
 
 template <typename A, typename B>
-void TEST_EQUALCb(FAutomationTestFixture& Test, const FString& What, const A& Actual, const B& Expected)
+void TestEqualCb(FAutomationTestFixture& Test, const FString& What, const A& Actual, const B& Expected)
 {
-	TEST_EQUAL(What, Actual, Expected);
+	INFO(What);
+	CHECK(Actual==Expected);
 }
 
 template <ECbFieldType FieldType, typename T = TCbFieldValueType<FieldType>, typename AccessorsType = TCbFieldTypeAccessors<FieldType>>
 void TestFieldNoClone(FAutomationTestFixture& Test, const TCHAR* What, FCbFieldView& Field, T ExpectedValue = T(), T DefaultValue = T(), ECbFieldError ExpectedError = ECbFieldError::None, const AccessorsType& Accessors = AccessorsType())
 {
-	TEST_EQUAL(FString::Printf(TEXT("FCbFieldView::Is[Type](%s)"), What), Invoke(Accessors.IsType, Field), ExpectedError != ECbFieldError::TypeError);
+	INFO(What);
+	CHECK_EQUAL(Invoke(Accessors.IsType, Field), (ExpectedError != ECbFieldError::TypeError));
 	if (ExpectedError == ECbFieldError::None && !Field.IsBool())
 	{
-		TEST_FALSE(FString::Printf(TEXT("FCbFieldView::AsBool(%s) == false"), What), Field.AsBool());
-		TEST_TRUE(FString::Printf(TEXT("FCbFieldView::AsBool(%s) -> HasError()"), What), Field.HasError());
-		TEST_EQUAL(FString::Printf(TEXT("FCbFieldView::AsBool(%s) -> GetError() == TypeError"), What), Field.GetError(), ECbFieldError::TypeError);
+		CHECK_FALSE(Field.AsBool());
+		CHECK(Field.HasError());
+		CHECK_EQUAL(Field.GetError(), ECbFieldError::TypeError);
 	}
-	TEST_EQUALCb(Test, FString::Printf(TEXT("FCbFieldView::As[Type](%s) -> Equal"), What), Invoke(Accessors.AsType, Field, DefaultValue), ExpectedValue);
-	TEST_EQUALCb(Test, FString::Printf(TEXT("FCbFieldView::As[Type](%s) -> HasError()"), What), Field.HasError(), ExpectedError != ECbFieldError::None);
-	TEST_EQUALCb(Test, FString::Printf(TEXT("FCbFieldView::As[Type](%s) -> GetError()"), What), Field.GetError(), ExpectedError);
+	TestEqualCb(Test, FString::Printf(TEXT("FCbFieldView::As[Type](%s) -> Equal"), What), Invoke(Accessors.AsType, Field, DefaultValue), ExpectedValue);
+	TestEqualCb(Test, FString::Printf(TEXT("FCbFieldView::As[Type](%s) -> HasError()"), What), Field.HasError(), ExpectedError != ECbFieldError::None);
+	TestEqualCb(Test, FString::Printf(TEXT("FCbFieldView::As[Type](%s) -> GetError()"), What), Field.GetError(), ExpectedError);
 }
 
 template <ECbFieldType FieldType, typename T = TCbFieldValueType<FieldType>, typename AccessorsType = TCbFieldTypeAccessors<FieldType>>
@@ -126,17 +135,17 @@ void TestField(FAutomationTestFixture& Test, const TCHAR* What, FCbFieldView& Fi
 	TestFieldNoClone<FieldType>(Test, What, Field, ExpectedValue, DefaultValue, ExpectedError, Accessors);
 	FCbField FieldClone = FCbField::Clone(Field);
 	TestFieldNoClone<FieldType>(Test, *FString::Printf(TEXT("%s, Clone"), What), FieldClone, ExpectedValue, DefaultValue, ExpectedError, Accessors);
-	TEST_TRUE(FString::Printf(TEXT("FCbFieldView::Equals(%s)"), What), Field.Equals(FieldClone));
+	CHECK(Field.Equals(FieldClone));
 }
 
 template <ECbFieldType FieldType, typename T = TCbFieldValueType<FieldType>, typename AccessorsType = TCbFieldTypeAccessors<FieldType>>
 void TestField(FAutomationTestFixture& Test, const TCHAR* What, TArrayView<const uint8> Value, T ExpectedValue = T(), T DefaultValue = T(), ECbFieldError ExpectedError = ECbFieldError::None, const AccessorsType& Accessors = AccessorsType())
 {
 	FCbFieldView Field(Value.GetData(), FieldType);
-	TEST_EQUAL(FString::Printf(TEXT("FCbFieldView::GetSize(%s)"), What), Field.GetSize(), uint64(Value.Num() + !FCbFieldType::HasFieldType(FieldType)));
-	TEST_TRUE(FString::Printf(TEXT("FCbFieldView::HasValue(%s)"), What), Field.HasValue());
-	TEST_FALSE(FString::Printf(TEXT("FCbFieldView::HasError(%s) == false"), What), Field.HasError());
-	TEST_EQUAL(FString::Printf(TEXT("FCbFieldView::GetError(%s) == None"), What), Field.GetError(), ECbFieldError::None);
+	CHECK_EQUAL(Field.GetSize(), uint64(Value.Num() + !FCbFieldType::HasFieldType(FieldType)));
+	CHECK(Field.HasValue());
+	CHECK_FALSE(Field.HasError());
+	CHECK_EQUAL(Field.GetError(), ECbFieldError::None);
 	TestField<FieldType>(Test, What, Field, ExpectedValue, DefaultValue, ExpectedError, Accessors);
 }
 
@@ -156,120 +165,122 @@ void TestFieldError(FAutomationTestFixture& Test, const TCHAR* What, TArrayView<
 
 TEST_CASE_METHOD(FAutomationTestFixture, "Core::Serialization::FCbField::Ctor Test", "[Core][Serialization][Smoke]")
 {
-	// Test FCbFieldView()
+	SECTION("Test FCbFieldView")
 	{
 		constexpr FCbFieldView DefaultField;
 		static_assert(!DefaultField.HasName(), "Error in HasName()");
 		static_assert(!DefaultField.HasValue(), "Error in HasValue()");
 		static_assert(!DefaultField.HasError(), "Error in HasError()");
 		static_assert(DefaultField.GetError() == ECbFieldError::None, "Error in GetError()");
-		TEST_EQUAL(TEXT("FCbFieldView()::GetSize() == 1"), DefaultField.GetSize(), uint64(1));
-		TEST_EQUAL(TEXT("FCbFieldView()::GetName().Len() == 0"), DefaultField.GetName().Len(), 0);
-		TEST_FALSE(TEXT("!FCbFieldView()::HasName()"), DefaultField.HasName());
-		TEST_FALSE(TEXT("!FCbFieldView()::HasValue()"), DefaultField.HasValue());
-		TEST_FALSE(TEXT("!FCbFieldView()::HasError()"), DefaultField.HasError());
-		TEST_EQUAL(TEXT("FCbFieldView()::GetError() == None"), DefaultField.GetError(), ECbFieldError::None);
-		TEST_EQUAL(TEXT("FCbFieldView()::GetHash()"), DefaultField.GetHash(), FIoHash::HashBuffer(MakeMemoryView<uint8>({uint8(ECbFieldType::None)})));
+		CHECK_EQUAL(DefaultField.GetSize(), uint64(1));
+		CHECK_EQUAL(DefaultField.GetName().Len(), 0);
+		CHECK_FALSE(DefaultField.HasName());
+		CHECK_FALSE(DefaultField.HasValue());
+		CHECK_FALSE(DefaultField.HasError());
+		CHECK_EQUAL(DefaultField.GetError(), ECbFieldError::None);
+		//CHECK(FIoHash::HashBuffer(MakeMemoryView<uint8>({uint8(ECbFieldType::None)})));
 		FMemoryView View;
-		TEST_FALSE(TEXT("FCbFieldView()::TryGetView()"), DefaultField.TryGetView(View));
+		CHECK_FALSE(DefaultField.TryGetView(View));
 	}
 
-	// Test FCbFieldView(None)
+	SECTION("Test FCbFieldView(None)")
 	{
 		FCbFieldView NoneField(nullptr, ECbFieldType::None);
-		TEST_EQUAL(TEXT("FCbFieldView(None)::GetSize() == 1"), NoneField.GetSize(), uint64(1));
-		TEST_EQUAL(TEXT("FCbFieldView(None)::GetName().Len() == 0"), NoneField.GetName().Len(), 0);
-		TEST_FALSE(TEXT("!FCbFieldView(None)::HasName()"), NoneField.HasName());
-		TEST_FALSE(TEXT("!FCbFieldView(None)::HasValue()"), NoneField.HasValue());
-		TEST_FALSE(TEXT("!FCbFieldView(None)::HasError()"), NoneField.HasError());
-		TEST_EQUAL(TEXT("FCbFieldView(None)::GetError() == None"), NoneField.GetError(), ECbFieldError::None);
-		TEST_EQUAL(TEXT("FCbFieldView(None)::GetHash()"), NoneField.GetHash(), FCbFieldView().GetHash());
+		CHECK_EQUAL(NoneField.GetSize(), uint64(1));
+		CHECK_EQUAL(NoneField.GetName().Len(), 0);
+		CHECK_FALSE(NoneField.HasName());
+		CHECK_FALSE(NoneField.HasValue());
+		CHECK_FALSE(NoneField.HasError());
+		CHECK_EQUAL(NoneField.GetError(), ECbFieldError::None);
+		CHECK_EQUAL(NoneField.GetHash(), FCbFieldView().GetHash());
 		FMemoryView View;
-		TEST_FALSE(TEXT("FCbFieldView(None)::TryGetView()"), NoneField.TryGetView(View));
+		CHECK_FALSE(NoneField.TryGetView(View));
 	}
 
-	// Test FCbFieldView(None|Type|Name)
+	SECTION("Test FCbFieldView(None|Type|Name)")
 	{
 		constexpr ECbFieldType FieldType = ECbFieldType::None | ECbFieldType::HasFieldName;
 		constexpr const ANSICHAR NoneBytes[] = { ANSICHAR(FieldType), 4, 'N', 'a', 'm', 'e' };
 		FCbFieldView NoneField(NoneBytes);
-		TEST_EQUAL(TEXT("FCbFieldView(None|Type|Name)::GetSize()"), NoneField.GetSize(), uint64(sizeof(NoneBytes)));
-		TEST_EQUAL(TEXT("FCbFieldView(None|Type|Name)::GetName()"), NoneField.GetName(), "Name"_U8SV);
-		TEST_TRUE(TEXT("FCbFieldView(None|Type|Name)::HasName()"), NoneField.HasName());
-		TEST_FALSE(TEXT("!FCbFieldView(None|Type|Name)::HasValue()"), NoneField.HasValue());
-		TEST_EQUAL(TEXT("FCbFieldView(None|Type|Name)::GetHash()"), NoneField.GetHash(), FIoHash::HashBuffer(MakeMemoryView(NoneBytes)));
+		CHECK_EQUAL(NoneField.GetSize(), uint64(sizeof(NoneBytes)));
+		CHECK_EQUAL(NoneField.GetName(), "Name"_U8SV);
+		CHECK(NoneField.HasName());
+		CHECK_FALSE(NoneField.HasValue());
+		CHECK_EQUAL(NoneField.GetHash(), FIoHash::HashBuffer(MakeMemoryView(NoneBytes)));
 		FMemoryView View;
-		TEST_TRUE(TEXT("FCbFieldView(None|Type|Name)::TryGetView()"), NoneField.TryGetView(View) && View == MakeMemoryView(NoneBytes));
+		CHECK(NoneField.TryGetView(View));
+		CHECK(View == MakeMemoryView(NoneBytes));
 
 		uint8 CopyBytes[sizeof(NoneBytes)];
 		NoneField.CopyTo(MakeMemoryView(CopyBytes));
-		TEST_TRUE(TEXT("FCbFieldView(None|Type|Name)::CopyTo()"), MakeMemoryView(NoneBytes).EqualBytes(MakeMemoryView(CopyBytes)));
+		CHECK(MakeMemoryView(NoneBytes).EqualBytes(MakeMemoryView(CopyBytes)));
 	}
 
-	// Test FCbFieldView(None|Type)
+	SECTION("Test FCbFieldView(None|Type)")
 	{
 		constexpr ECbFieldType FieldType = ECbFieldType::None;
 		constexpr const ANSICHAR NoneBytes[] = { ANSICHAR(FieldType) };
 		FCbFieldView NoneField(NoneBytes);
-		TEST_EQUAL(TEXT("FCbFieldView(None|Type)::GetSize()"), NoneField.GetSize(), uint64(sizeof(NoneBytes)));
-		TEST_EQUAL(TEXT("FCbFieldView(None|Type)::GetName()"), NoneField.GetName().Len(), 0);
-		TEST_FALSE(TEXT("FCbFieldView(None|Type)::HasName()"), NoneField.HasName());
-		TEST_FALSE(TEXT("!FCbFieldView(None|Type)::HasValue()"), NoneField.HasValue());
-		TEST_EQUAL(TEXT("FCbFieldView(None|Type)::GetHash()"), NoneField.GetHash(), FCbFieldView().GetHash());
+		CHECK_EQUAL(NoneField.GetSize(), uint64(sizeof(NoneBytes)));
+		CHECK_EQUAL(NoneField.GetName().Len(), 0);
+		CHECK_FALSE( NoneField.HasName());
+		CHECK_FALSE(NoneField.HasValue());
+		CHECK_EQUAL(NoneField.GetHash(), FCbFieldView().GetHash());
 		FMemoryView View;
-		TEST_TRUE(TEXT("FCbFieldView(None|Type)::TryGetView()"), NoneField.TryGetView(View) && View == MakeMemoryView(NoneBytes));
+		CHECK(NoneField.TryGetView(View));
+		CHECK(View == MakeMemoryView(NoneBytes));
 	}
 
-	// Test FCbFieldView(None|Name)
+	SECTION("Test FCbFieldView(None|Name)")
 	{
 		constexpr ECbFieldType FieldType = ECbFieldType::None | ECbFieldType::HasFieldName;
 		constexpr const ANSICHAR NoneBytes[] = { ANSICHAR(FieldType), 4, 'N', 'a', 'm', 'e' };
 		FCbFieldView NoneField(NoneBytes + 1, FieldType);
-		TEST_EQUAL(TEXT("FCbFieldView(None|Name)::GetSize()"), NoneField.GetSize(), uint64(sizeof(NoneBytes)));
-		TEST_EQUAL(TEXT("FCbFieldView(None|Name)::GetName()"), NoneField.GetName(), "Name"_U8SV);
-		TEST_TRUE(TEXT("FCbFieldView(None|Name)::HasName()"), NoneField.HasName());
-		TEST_FALSE(TEXT("!FCbFieldView(None|Name)::HasValue()"), NoneField.HasValue());
-		TEST_EQUAL(TEXT("FCbFieldView(None|Name)::GetHash()"), NoneField.GetHash(), FIoHash::HashBuffer(MakeMemoryView(NoneBytes)));
+		CHECK_EQUAL(NoneField.GetSize(), uint64(sizeof(NoneBytes)));
+		CHECK_EQUAL(NoneField.GetName(), "Name"_U8SV);
+		CHECK(NoneField.HasName());
+		CHECK_FALSE( NoneField.HasValue());
+		CHECK_EQUAL(NoneField.GetHash(), FIoHash::HashBuffer(MakeMemoryView(NoneBytes)));
 		FMemoryView View;
-		TEST_FALSE(TEXT("FCbFieldView(None|Name)::TryGetView()"), NoneField.TryGetView(View));
+		CHECK_FALSE(NoneField.TryGetView(View));
 
 		uint8 CopyBytes[sizeof(NoneBytes)];
 		NoneField.CopyTo(MakeMemoryView(CopyBytes));
-		TEST_TRUE(TEXT("FCbFieldView(None|Name)::CopyTo()"), MakeMemoryView(NoneBytes).EqualBytes(MakeMemoryView(CopyBytes)));
+		CHECK(MakeMemoryView(NoneBytes).EqualBytes(MakeMemoryView(CopyBytes)));
 	}
 
-	// Test FCbFieldView(None|EmptyName)
+	SECTION("Test FCbFieldView(None|EmptyName)")
 	{
 		constexpr ECbFieldType FieldType = ECbFieldType::None | ECbFieldType::HasFieldName;
 		constexpr const uint8 NoneBytes[] = { uint8(FieldType), 0 };
 		FCbFieldView NoneField(NoneBytes + 1, FieldType);
-		TEST_EQUAL(TEXT("FCbFieldView(None|EmptyName)::GetSize()"), NoneField.GetSize(), uint64(sizeof(NoneBytes)));
-		TEST_EQUAL(TEXT("FCbFieldView(None|EmptyName)::GetName()"), NoneField.GetName(), ""_U8SV);
-		TEST_TRUE(TEXT("FCbFieldView(None|EmptyName)::HasName()"), NoneField.HasName());
-		TEST_FALSE(TEXT("!FCbFieldView(None|EmptyName)::HasValue()"), NoneField.HasValue());
-		TEST_EQUAL(TEXT("FCbFieldView(None|EmptyName)::GetHash()"), NoneField.GetHash(), FIoHash::HashBuffer(MakeMemoryView(NoneBytes)));
+		CHECK_EQUAL(NoneField.GetSize(), uint64(sizeof(NoneBytes)));
+		CHECK_EQUAL(NoneField.GetName(), ""_U8SV);
+		CHECK(NoneField.HasName());
+		CHECK_FALSE(NoneField.HasValue());
+		CHECK_EQUAL(NoneField.GetHash(), FIoHash::HashBuffer(MakeMemoryView(NoneBytes)));
 		FMemoryView View;
-		TEST_FALSE(TEXT("FCbFieldView(None|EmptyName)::TryGetView()"), NoneField.TryGetView(View));
+		CHECK_FALSE(NoneField.TryGetView(View));
 	}
 }
 
 TEST_CASE_METHOD(FAutomationTestFixture, "Core::Serialization::FCbField::Null Test", "[Core][Serialization][Smoke]")
 {
-	// Test FCbFieldView(Null)
+	SECTION("Test FCbFieldView(Null)")
 	{
 		FCbFieldView NullField(nullptr, ECbFieldType::Null);
-		TEST_EQUAL(TEXT("FCbFieldView(Null)::GetSize() == 1"), NullField.GetSize(), uint64(1));
-		TEST_TRUE(TEXT("FCbFieldView(Null)::IsNull()"), NullField.IsNull());
-		TEST_TRUE(TEXT("FCbFieldView(Null)::HasValue()"), NullField.HasValue());
-		TEST_FALSE(TEXT("!FCbFieldView(Null)::HasError()"), NullField.HasError());
-		TEST_EQUAL(TEXT("FCbFieldView(Null)::GetError() == None"), NullField.GetError(), ECbFieldError::None);
-		TEST_EQUAL(TEXT("FCbFieldView(Null)::GetHash()"), NullField.GetHash(), FIoHash::HashBuffer(MakeMemoryView<uint8>({uint8(ECbFieldType::Null)})));
+		CHECK_EQUAL(NullField.GetSize(), uint64(1));
+		CHECK(NullField.IsNull());
+		CHECK(NullField.HasValue());
+		CHECK_FALSE(NullField.HasError());
+		CHECK_EQUAL(NullField.GetError(), ECbFieldError::None);
+		CHECK_EQUAL(NullField.GetHash(), FIoHash::HashBuffer(MakeMemoryView<uint8>({uint8(ECbFieldType::Null)})));
 	}
 
-	// Test FCbFieldView(None) as Null
+	SECTION("Test FCbFieldView(None) as Null")
 	{
 		FCbFieldView Field;
-		TEST_FALSE(TEXT("FCbFieldView(None)::IsNull()"), Field.IsNull());
+		CHECK_FALSE(Field.IsNull());
 	}
 }
 
@@ -282,63 +293,65 @@ TEST_CASE_METHOD(FAutomationTestFixture, "Core::Serialization::FCbField::Object
 
 	auto TestIntObject = [this](const FCbObjectView& Object, int32 ExpectedNum, uint64 ExpectedValueSize)
 	{
-		TEST_EQUAL(TEXT("FCbFieldView(Object)::AsObjectView().GetSize()"), Object.GetSize(), ExpectedValueSize + sizeof(ECbFieldType));
+		CHECK_EQUAL(Object.GetSize(), ExpectedValueSize + sizeof(ECbFieldType));
 
 		int32 ActualNum = 0;
 		for (FCbFieldViewIterator It = Object.CreateViewIterator(); It; ++It)
 		{
 			++ActualNum;
 			TestNotEqual(TEXT("FCbFieldView(Object) Iterator Name"), It->GetName().Len(), 0);
-			TEST_EQUAL(TEXT("FCbFieldView(Object) Iterator"), It->AsInt32(), ActualNum);
+			CHECK_EQUAL(It->AsInt32(), ActualNum);
 		}
-		TEST_EQUAL(TEXT("FCbFieldView(Object)::AsObjectView().CreateViewIterator() -> Count"), ActualNum, ExpectedNum);
+		CHECK_EQUAL( ActualNum, ExpectedNum);
 
 		ActualNum = 0;
 		for (FCbFieldView Field : Object)
 		{
 			++ActualNum;
 			TestNotEqual(TEXT("FCbFieldView(Object) Iterator Name"), Field.GetName().Len(), 0);
-			TEST_EQUAL(TEXT("FCbFieldView(Object) Range"), Field.AsInt32(), ActualNum);
+			CHECK_EQUAL(Field.AsInt32(), ActualNum);
 		}
-		TEST_EQUAL(TEXT("FCbFieldView(Object)::AsObjectView() Range -> Count"), ActualNum, ExpectedNum);
+		CHECK_EQUAL(ActualNum, ExpectedNum);
 
 		ActualNum = 0;
 		for (FCbFieldView Field : Object.AsFieldView())
 		{
 			++ActualNum;
 			TestNotEqual(TEXT("FCbFieldView(ObjectField) Iterator Name"), Field.GetName().Len(), 0);
-			TEST_EQUAL(TEXT("FCbFieldView(ObjectField) Range"), Field.AsInt32(), ActualNum);
+			CHECK_EQUAL(Field.AsInt32(), ActualNum);
 		}
-		TEST_EQUAL(TEXT("FCbFieldView(ObjectField)::AsObjectView() Range -> Count"), ActualNum, ExpectedNum);
+		CHECK_EQUAL( ActualNum, ExpectedNum);
 	};
 
-	// Test FCbFieldView(Object, Empty)
-	TestField<ECbFieldType::Object>(*this, TEXT("Object, Empty"), {0});
-
-	// Test FCbFieldView(Object, Empty)
+	SECTION("Test FCbFieldView(Object, Empty)")
+	{
+		TestField<ECbFieldType::Object>(*this, TEXT("Object, Empty"), { 0 });
+	}
+	
+	SECTION("Test FCbFieldView(Object, Empty)")
 	{
 		FCbObjectView Object;
 		TestIntObject(Object, 0, 1);
 
 		// Find fields that do not exist.
-		TEST_FALSE(TEXT("FCbObjectView()::Find(Missing)"), Object.FindView("Field"_U8SV).HasValue());
-		TEST_FALSE(TEXT("FCbObjectView()::FindViewIgnoreCase(Missing)"), Object.FindViewIgnoreCase("Field"_U8SV).HasValue());
-		TEST_FALSE(TEXT("FCbObjectView()::operator[](Missing)"), Object["Field"_U8SV].HasValue());
+		CHECK_FALSE(Object.FindView("Field"_U8SV).HasValue());
+		CHECK_FALSE(Object.FindViewIgnoreCase("Field"_U8SV).HasValue());
+		CHECK_FALSE(Object["Field"_U8SV].HasValue());
 
 		// Advance an iterator past the last field.
 		FCbFieldViewIterator It = Object.CreateViewIterator();
-		TEST_FALSE(TEXT("FCbObjectView()::CreateViewIterator() At End"), bool(It));
-		TEST_TRUE(TEXT("FCbObjectView()::CreateViewIterator() At End"), !It);
+		CHECK_FALSE(bool(It));
+		CHECK(!It);
 		for (int Count = 16; Count > 0; --Count)
 		{
 			++It;
 			It->AsInt32();
 		}
-		TEST_FALSE(TEXT("FCbObjectView()::CreateViewIterator() At End"), bool(It));
-		TEST_TRUE(TEXT("FCbObjectView()::CreateViewIterator() At End"), !It);
+		CHECK_FALSE(bool(It));
+		CHECK(!It);
 	}
 
-	// Test FCbFieldView(Object, NotEmpty)
+	SECTION("Test FCbFieldView(Object, NotEmpty)")
 	{
 		constexpr uint8 IntType = uint8(ECbFieldType::HasFieldName | ECbFieldType::IntegerPositive);
 		const uint8 Value[] = { 12, IntType, 1, 'A', 1, IntType, 1, 'B', 2, IntType, 1, 'C', 3 };
@@ -347,16 +360,16 @@ TEST_CASE_METHOD(FAutomationTestFixture, "Core::Serialization::FCbField::Object
 		FCbObject Object = FCbObject::Clone(Field.AsObjectView());
 		TestIntObject(Object, 3, sizeof(Value));
 		TestIntObject(Field.AsObjectView(), 3, sizeof(Value));
-		TEST_TRUE(TEXT("FCbObjectView::Equals()"), Object.Equals(Field.AsObjectView()));
-		TEST_EQUAL(TEXT("FCbObjectView::Find()"), Object.FindView("B"_ASV).AsInt32(), 2);
-		TEST_EQUAL(TEXT("FCbObjectView::Find()"), Object.FindView("b"_ASV).AsInt32(4), 4);
-		TEST_EQUAL(TEXT("FCbObjectView::FindViewIgnoreCase()"), Object.FindViewIgnoreCase("B"_ASV).AsInt32(), 2);
-		TEST_EQUAL(TEXT("FCbObjectView::FindViewIgnoreCase()"), Object.FindViewIgnoreCase("b"_ASV).AsInt32(), 2);
-		TEST_EQUAL(TEXT("FCbObjectView::operator[]"), Object["B"_ASV].AsInt32(), 2);
-		TEST_EQUAL(TEXT("FCbObjectView::operator[]"), Object["b"_ASV].AsInt32(4), 4);
+		CHECK(Object.Equals(Field.AsObjectView()));
+		CHECK_EQUAL(Object.FindView("B"_ASV).AsInt32(), 2);
+		CHECK_EQUAL(Object.FindView("b"_ASV).AsInt32(4), 4);
+		CHECK_EQUAL(Object.FindViewIgnoreCase("B"_ASV).AsInt32(), 2);
+		CHECK_EQUAL(Object.FindViewIgnoreCase("b"_ASV).AsInt32(), 2);
+		CHECK_EQUAL(Object["B"_ASV].AsInt32(), 2);
+		CHECK_EQUAL(Object["b"_ASV].AsInt32(4), 4);
 	}
 
-	// Test FCbFieldView(UniformObject, NotEmpty)
+	SECTION("Test FCbFieldView(UniformObject, NotEmpty)")
 	{
 		constexpr uint8 IntType = uint8(ECbFieldType::HasFieldName | ECbFieldType::IntegerPositive);
 		const uint8 Value[] = { 10, IntType, 1, 'A', 1, 1, 'B', 2, 1, 'C', 3 };
@@ -365,69 +378,79 @@ TEST_CASE_METHOD(FAutomationTestFixture, "Core::Serialization::FCbField::Object
 		FCbObject Object = FCbObject::Clone(Field.AsObjectView());
 		TestIntObject(Object, 3, sizeof(Value));
 		TestIntObject(Field.AsObjectView(), 3, sizeof(Value));
-		TEST_TRUE(TEXT("FCbObjectView(Uniform)::Equals()"), Object.Equals(Field.AsObjectView()));
-		TEST_EQUAL(TEXT("FCbObjectView(Uniform)::Find()"), Object.FindView("B"_ASV).AsInt32(), 2);
-		TEST_EQUAL(TEXT("FCbObjectView(Uniform)::Find()"), Object.Find("B"_ASV).AsInt32(), 2);
-		TEST_EQUAL(TEXT("FCbObjectView(Uniform)::Find()"), Object.FindView("b"_ASV).AsInt32(4), 4);
-		TEST_EQUAL(TEXT("FCbObjectView(Uniform)::Find()"), Object.Find("b"_ASV).AsInt32(4), 4);
-		TEST_EQUAL(TEXT("FCbObjectView(Uniform)::FindViewIgnoreCase()"), Object.FindViewIgnoreCase("B"_ASV).AsInt32(), 2);
-		TEST_EQUAL(TEXT("FCbObjectView(Uniform)::FindViewIgnoreCase()"), Object.FindIgnoreCase("B"_ASV).AsInt32(), 2);
-		TEST_EQUAL(TEXT("FCbObjectView(Uniform)::FindViewIgnoreCase()"), Object.FindViewIgnoreCase("b"_ASV).AsInt32(), 2);
-		TEST_EQUAL(TEXT("FCbObjectView(Uniform)::FindViewIgnoreCase()"), Object.FindIgnoreCase("b"_ASV).AsInt32(), 2);
-		TEST_EQUAL(TEXT("FCbObjectView(Uniform)::operator[]"), Object["B"_ASV].AsInt32(), 2);
-		TEST_EQUAL(TEXT("FCbObjectView(Uniform)::operator[]"), Object["b"_ASV].AsInt32(4), 4);
-
-		TEST_TRUE(TEXT("FCbObject::AsField()"), Object.GetOuterBuffer() == Object.AsField().AsObject().GetOuterBuffer());
-
-		// Equals
+		CHECK(Object.Equals(Field.AsObjectView()));
+		CHECK_EQUAL(Object.FindView("B"_ASV).AsInt32(), 2);
+		CHECK_EQUAL(Object.Find("B"_ASV).AsInt32(), 2);
+		CHECK_EQUAL(Object.FindView("b"_ASV).AsInt32(4), 4);
+		CHECK_EQUAL(Object.Find("b"_ASV).AsInt32(4), 4);
+		CHECK_EQUAL(Object.FindViewIgnoreCase("B"_ASV).AsInt32(), 2);
+		CHECK_EQUAL(Object.FindIgnoreCase("B"_ASV).AsInt32(), 2);
+		CHECK_EQUAL(Object.FindViewIgnoreCase("b"_ASV).AsInt32(), 2);
+		CHECK_EQUAL(Object.FindIgnoreCase("b"_ASV).AsInt32(), 2);
+		CHECK_EQUAL(Object["B"_ASV].AsInt32(), 2);
+		CHECK_EQUAL(Object["b"_ASV].AsInt32(4), 4);
+
+		CHECK(Object.GetOuterBuffer() == Object.AsField().AsObject().GetOuterBuffer());
+
 		const uint8 NamedValue[] = { 1, 'O', 10, IntType, 1, 'A', 1, 1, 'B', 2, 1, 'C', 3 };
 		FCbFieldView NamedField(NamedValue, ECbFieldType::UniformObject | ECbFieldType::HasFieldName);
-		TEST_TRUE(TEXT("FCbObjectView::Equals()"), Field.AsObjectView().Equals(NamedField.AsObjectView()));
-		TEST_TRUE(TEXT("FCbObjectView::AsFieldView().Equals()"), Field.Equals(Field.AsObjectView().AsFieldView()));
-		TEST_TRUE(TEXT("FCbObjectView::AsFieldView().Equals()"), NamedField.Equals(NamedField.AsObjectView().AsFieldView()));
-
-		// CopyTo
-		uint8 CopyBytes[sizeof(Value) + 1];
-		Field.AsObjectView().CopyTo(MakeMemoryView(CopyBytes));
-		TEST_TRUE(TEXT("FCbObjectView(NoType)::CopyTo()"), MakeMemoryView(Value).EqualBytes(MakeMemoryView(CopyBytes) + 1));
-		NamedField.AsObjectView().CopyTo(MakeMemoryView(CopyBytes));
-		TEST_TRUE(TEXT("FCbObjectView(NoType, Name)::CopyTo()"), MakeMemoryView(Value).EqualBytes(MakeMemoryView(CopyBytes) + 1));
-
-		// TryGetView
-		FMemoryView View;
-		TEST_TRUE(TEXT("FCbObjectView(Clone)::TryGetView()"), Object.TryGetView(View) && View == Object.GetOuterBuffer().GetView());
-		TEST_FALSE(TEXT("FCbObjectView(NoType)::TryGetView()"), Field.AsObjectView().TryGetView(View));
-		TEST_FALSE(TEXT("FCbObjectView(Name)::TryGetView()"), NamedField.AsObjectView().TryGetView(View));
 
-		// GetBuffer
-		TEST_TRUE(TEXT("FCbObjectView(Clone)::GetBuffer()"), Object.GetBuffer().ToShared().GetView() == Object.GetOuterBuffer().GetView());
-		TEST_TRUE(TEXT("FCbObjectView(NoType)::GetBuffer()"), FCbField::MakeView(Field).AsObject().GetBuffer().ToShared().GetView().EqualBytes(Object.GetOuterBuffer().GetView()));
-		TEST_TRUE(TEXT("FCbObjectView(Name)::GetBuffer()"), FCbField::MakeView(NamedField).AsObject().GetBuffer().ToShared().GetView().EqualBytes(Object.GetOuterBuffer().GetView()));
+		{
+			INFO("Equals");
+			CHECK(Field.AsObjectView().Equals(NamedField.AsObjectView()));
+			CHECK(Field.Equals(Field.AsObjectView().AsFieldView()));
+			CHECK(NamedField.Equals(NamedField.AsObjectView().AsFieldView()));
+		}
+		
+		{
+			INFO("CopyTo");
+			uint8 CopyBytes[sizeof(Value) + 1];
+			Field.AsObjectView().CopyTo(MakeMemoryView(CopyBytes));
+			CHECK(MakeMemoryView(Value).EqualBytes(MakeMemoryView(CopyBytes) + 1));
+			NamedField.AsObjectView().CopyTo(MakeMemoryView(CopyBytes));
+			CHECK(MakeMemoryView(Value).EqualBytes(MakeMemoryView(CopyBytes) + 1));
+		}
+	
+		{
+			INFO("TryGetView")
+			FMemoryView View;
+			CHECK(Object.TryGetView(View));
+			CHECK(View == Object.GetOuterBuffer().GetView());
+			CHECK_FALSE(Field.AsObjectView().TryGetView(View));
+			CHECK_FALSE(NamedField.AsObjectView().TryGetView(View));
+		}
+
+		{
+			INFO("GetBuffer")
+			CHECK(Object.GetBuffer().ToShared().GetView() == Object.GetOuterBuffer().GetView());
+			CHECK(FCbField::MakeView(Field).AsObject().GetBuffer().ToShared().GetView().EqualBytes(Object.GetOuterBuffer().GetView()));
+			CHECK(FCbField::MakeView(NamedField).AsObject().GetBuffer().ToShared().GetView().EqualBytes(Object.GetOuterBuffer().GetView()));
+		}
 	}
 
-	// Test FCbFieldView(None) as Object
+	SECTION("Test FCbFieldView(None) as Object")
 	{
 		FCbFieldView Field;
 		TestFieldError<ECbFieldType::Object>(*this, TEXT("Object, None"), Field, ECbFieldError::TypeError);
 		FCbField::MakeView(Field).AsObject();
 	}
 
-	// Test FCbObjectView(ObjectWithName) and CreateIterator
+	SECTION("Test FCbObjectView(ObjectWithName) and CreateIterator")
 	{
 		const uint8 ObjectType = uint8(ECbFieldType::Object | ECbFieldType::HasFieldName);
 		const uint8 Buffer[] = { ObjectType, 3, 'K', 'e', 'y', 4, uint8(ECbFieldType::HasFieldName | ECbFieldType::IntegerPositive), 1, 'F', 8 };
 		const FCbObjectView Object(Buffer);
-		TEST_EQUAL(TEXT("FCbObjectView(ObjectWithName)::GetSize()"), Object.GetSize(), uint64(6));
+		CHECK_EQUAL(Object.GetSize(), uint64(6));
 		const FCbObject ObjectClone = FCbObject::Clone(Object);
-		TEST_EQUAL(TEXT("FCbObject(ObjectWithName)::GetSize()"), ObjectClone.GetSize(), uint64(6));
-		TEST_TRUE(TEXT("FCbObjectView::Equals()"), Object.Equals(ObjectClone));
-		TEST_EQUAL(TEXT("FCbObjectView::GetHash()"), ObjectClone.GetHash(), Object.GetHash());
+		CHECK_EQUAL(ObjectClone.GetSize(), uint64(6));
+		CHECK(Object.Equals(ObjectClone));
+		CHECK_EQUAL(ObjectClone.GetHash(), Object.GetHash());
 		for (FCbFieldIterator It = ObjectClone.CreateIterator(); It; ++It)
 		{
 			FCbField Field = *It;
-			TEST_EQUAL(TEXT("FCbObject::CreateIterator().GetName()"), Field.GetName(), "F"_U8SV);
-			TEST_EQUAL(TEXT("FCbObject::CreateIterator().AsInt32()"), Field.AsInt32(), 8);
-			TEST_TRUE(TEXT("FCbObject::CreateIterator().IsOwned()"), Field.IsOwned());
+			CHECK_EQUAL(Field.GetName(), "F"_U8SV);
+			CHECK_EQUAL(Field.AsInt32(), 8);
+			CHECK(Field.IsOwned());
 		}
 		for (FCbFieldIterator It = ObjectClone.CreateIterator(), End; It != End; ++It)
 		{
@@ -436,49 +459,55 @@ TEST_CASE_METHOD(FAutomationTestFixture, "Core::Serialization::FCbField::Object
 		{
 		}
 
-		// CopyTo
-		uint8 CopyBytes[6];
-		Object.CopyTo(MakeMemoryView(CopyBytes));
-		TEST_TRUE(TEXT("FCbObjectView(Name)::CopyTo()"), ObjectClone.GetOuterBuffer().GetView().EqualBytes(MakeMemoryView(CopyBytes)));
-		ObjectClone.CopyTo(MakeMemoryView(CopyBytes));
-		TEST_TRUE(TEXT("FCbObjectView()::CopyTo()"), ObjectClone.GetOuterBuffer().GetView().EqualBytes(MakeMemoryView(CopyBytes)));
-
-		// GetBuffer
-		TEST_TRUE(TEXT("FCbField(ObjectWithName)::GetBuffer()"), FCbField(FSharedBuffer::MakeView(MakeMemoryView(Buffer))).GetBuffer().ToShared().GetView().EqualBytes(MakeMemoryView(Buffer)));
-		TEST_TRUE(TEXT("FCbField(ObjectWithNameNoType)::GetBuffer()"),
-			FCbField(FCbFieldView(Buffer + 1, ECbFieldType(ObjectType)), FSharedBuffer::MakeView(MakeMemoryView(Buffer))).GetBuffer().ToShared().GetView().EqualBytes(MakeMemoryView(Buffer)));
-
-		// Access Missing Field
-		TEST_FALSE(TEXT("FCbObject()[Missing]"), ObjectClone["M"_ASV].HasValue());
-		TEST_FALSE(TEXT("FCbField(Object)[Missing]"), ObjectClone.AsField()["M"_ASV].HasValue());
+		SECTION("CopyTo")
+		{
+			uint8 CopyBytes[6];
+			Object.CopyTo(MakeMemoryView(CopyBytes));
+			CHECK(ObjectClone.GetOuterBuffer().GetView().EqualBytes(MakeMemoryView(CopyBytes)));
+			ObjectClone.CopyTo(MakeMemoryView(CopyBytes));
+			CHECK(ObjectClone.GetOuterBuffer().GetView().EqualBytes(MakeMemoryView(CopyBytes)));
+		}
+		
+		SECTION("GetBuffer")
+		{
+			CHECK(FCbField(FSharedBuffer::MakeView(MakeMemoryView(Buffer))).GetBuffer().ToShared().GetView().EqualBytes(MakeMemoryView(Buffer)));
+			CHECK(FCbField(FCbFieldView(Buffer + 1, ECbFieldType(ObjectType)), FSharedBuffer::MakeView(MakeMemoryView(Buffer))).GetBuffer().ToShared().GetView().EqualBytes(MakeMemoryView(Buffer)));
+		}
+		
+		SECTION("Access Missing Field")
+		{
+			CHECK_FALSE(ObjectClone["M"_ASV].HasValue());
+			CHECK_FALSE(ObjectClone.AsField()["M"_ASV].HasValue());
+		}
+		
 	}
 
-	// Test FCbObjectView as FCbFieldViewIterator
+	SECTION("Test FCbObjectView as FCbFieldViewIterator")
 	{
 		uint32 Count = 0;
 		FCbObjectView Object;
 		for (FCbFieldView Field : FCbFieldViewIterator::MakeSingle(Object.AsFieldView()))
 		{
-			TEST_TRUE(TEXT("FCbObjectView::AsFieldView() as Iterator"), Field.IsObject());
+			CHECK(Field.IsObject());
 			++Count;
 		}
-		TEST_EQUAL(TEXT("FCbObjectView::AsFieldView() as Iterator Count"), Count, 1u);
+		CHECK_EQUAL(Count, 1u);
 	}
 
-	// Test FCbObject as FCbFieldIterator
+	SECTION("Test FCbObject as FCbFieldIterator")
 	{
 		uint32 Count = 0;
 		FCbObject Object;
 		Object.MakeOwned();
 		for (FCbField Field : FCbFieldIterator::MakeSingle(Object.AsField()))
 		{
-			TEST_TRUE(TEXT("FCbObject::AsFieldView() as Iterator"), Field.IsObject());
+			CHECK(Field.IsObject());
 			++Count;
 		}
-		TEST_EQUAL(TEXT("FCbObject::AsFieldView() as Iterator Count"), Count, 1u);
+		CHECK_EQUAL(Count, 1u);
 	}
 
-	// Test FCbObject(Empty) as FCbFieldIterator
+	SECTION("Test FCbObject(Empty) as FCbFieldIterator")
 	{
 		const uint8 Buffer[] = { uint8(ECbFieldType::Object), 0 };
 		const FCbObject Object = FCbObject::Clone(Buffer);
@@ -498,56 +527,58 @@ TEST_CASE_METHOD(FAutomationTestFixture, "Core::Serialization::FCbField::Array T
 
 	auto TestIntArray = [this](FCbArrayView Array, int32 ExpectedNum, uint64 ExpectedValueSize)
 	{
-		TEST_EQUAL(TEXT("FCbFieldView(Array)::AsArrayView().GetSize()"), Array.GetSize(), ExpectedValueSize + sizeof(ECbFieldType));
-		TEST_EQUAL(TEXT("FCbFieldView(Array)::AsArrayView().Num()"), Array.Num(), uint64(ExpectedNum));
+		CHECK_EQUAL(Array.GetSize(), ExpectedValueSize + sizeof(ECbFieldType));
+		CHECK_EQUAL(Array.Num(), uint64(ExpectedNum));
 
 		int32 ActualNum = 0;
 		for (FCbFieldViewIterator It = Array.CreateViewIterator(); It; ++It)
 		{
 			++ActualNum;
-			TEST_EQUAL(TEXT("FCbFieldView(Array) Iterator"), It->AsInt32(), ActualNum);
+			CHECK_EQUAL(It->AsInt32(), ActualNum);
 		}
-		TEST_EQUAL(TEXT("FCbFieldView(Array)::AsArrayView().CreateViewIterator() -> Count"), ActualNum, ExpectedNum);
+		CHECK_EQUAL(ActualNum, ExpectedNum);
 
 		ActualNum = 0;
 		for (FCbFieldView Field : Array)
 		{
 			++ActualNum;
-			TEST_EQUAL(TEXT("FCbFieldView(Array) Range"), Field.AsInt32(), ActualNum);
+			CHECK_EQUAL(Field.AsInt32(), ActualNum);
 		}
-		TEST_EQUAL(TEXT("FCbFieldView(Array)::AsArrayView() Range -> Count"), ActualNum, ExpectedNum);
+		CHECK_EQUAL(ActualNum, ExpectedNum);
 
 		ActualNum = 0;
 		for (FCbFieldView Field : Array.AsFieldView())
 		{
 			++ActualNum;
-			TEST_EQUAL(TEXT("FCbFieldView(ArrayField) Range"), Field.AsInt32(), ActualNum);
+			CHECK_EQUAL(Field.AsInt32(), ActualNum);
 		}
-		TEST_EQUAL(TEXT("FCbFieldView(ArrayField)::AsArrayView() Range -> Count"), ActualNum, ExpectedNum);
+		CHECK_EQUAL(ActualNum, ExpectedNum);
 	};
 
-	// Test FCbFieldView(Array, Empty)
-	TestField<ECbFieldType::Array>(*this, TEXT("Array, Empty"), { 1, 0 });
-
-	// Test FCbFieldView(Array, Empty)
+	SECTION("Test FCbFieldView(Array, Empty)")
+	{
+		TestField<ECbFieldType::Array>(*this, TEXT("Array, Empty"), { 1, 0 });
+	}
+	
+	SECTION("Test FCbFieldView(Array, Empty)")
 	{
 		FCbArrayView Array;
 		TestIntArray(Array, 0, 2);
 
 		// Advance an iterator past the last field.
 		FCbFieldViewIterator It = Array.CreateViewIterator();
-		TEST_FALSE(TEXT("FCbArrayView()::CreateViewIterator() At End"), bool(It));
-		TEST_TRUE(TEXT("FCbArrayView()::CreateViewIterator() At End"), !It);
+		CHECK_FALSE(bool(It));
+		CHECK(!It);
 		for (int Count = 16; Count > 0; --Count)
 		{
 			++It;
 			It->AsInt32();
 		}
-		TEST_FALSE(TEXT("FCbArrayView()::CreateViewIterator() At End"), bool(It));
-		TEST_TRUE(TEXT("FCbArrayView()::CreateViewIterator() At End"), !It);
+		CHECK_FALSE(bool(It));
+		CHECK(!It);
 	}
 
-	// Test FCbFieldView(Array, NotEmpty)
+	SECTION("Test FCbFieldView(Array, NotEmpty)")
 	{
 		constexpr uint8 IntType = uint8(ECbFieldType::IntegerPositive);
 		const uint8 Value[] = { 7, 3, IntType, 1, IntType, 2, IntType, 3 };
@@ -556,10 +587,10 @@ TEST_CASE_METHOD(FAutomationTestFixture, "Core::Serialization::FCbField::Array T
 		FCbArray Array = FCbArray::Clone(Field.AsArrayView());
 		TestIntArray(Array, 3, sizeof(Value));
 		TestIntArray(Field.AsArrayView(), 3, sizeof(Value));
-		TEST_TRUE(TEXT("FCbArrayView::Equals()"), Array.Equals(Field.AsArrayView()));
+		CHECK(Array.Equals(Field.AsArrayView()));
 	}
 
-	// Test FCbFieldView(UniformArray)
+	SECTION("Test FCbFieldView(UniformArray)")
 	{
 		constexpr uint8 IntType = uint8(ECbFieldType::IntegerPositive);
 		const uint8 Value[] = { 5, 3, IntType, 1, 2, 3 };
@@ -568,58 +599,68 @@ TEST_CASE_METHOD(FAutomationTestFixture, "Core::Serialization::FCbField::Array T
 		FCbArray Array = FCbArray::Clone(Field.AsArrayView());
 		TestIntArray(Array, 3, sizeof(Value));
 		TestIntArray(Field.AsArrayView(), 3, sizeof(Value));
-		TEST_TRUE(TEXT("FCbArrayView(Uniform)::Equals()"), Array.Equals(Field.AsArrayView()));
+		CHECK(Array.Equals(Field.AsArrayView()));
 
-		TEST_TRUE(TEXT("FCbArray::AsField()"), Array.GetOuterBuffer() == Array.AsField().AsArray().GetOuterBuffer());
+		CHECK(Array.GetOuterBuffer() == Array.AsField().AsArray().GetOuterBuffer());
 
-		// Equals
 		const uint8 NamedValue[] = { 1, 'A', 5, 3, IntType, 1, 2, 3 };
 		FCbFieldView NamedField(NamedValue, ECbFieldType::UniformArray | ECbFieldType::HasFieldName);
-		TEST_TRUE(TEXT("FCbArrayView::Equals()"), Field.AsArrayView().Equals(NamedField.AsArrayView()));
-		TEST_TRUE(TEXT("FCbArrayView::AsFieldView().Equals()"), Field.Equals(Field.AsArrayView().AsFieldView()));
-		TEST_TRUE(TEXT("FCbArrayView::AsFieldView().Equals()"), NamedField.Equals(NamedField.AsArrayView().AsFieldView()));
-
-		// CopyTo
-		uint8 CopyBytes[sizeof(Value) + 1];
-		Field.AsArrayView().CopyTo(MakeMemoryView(CopyBytes));
-		TEST_TRUE(TEXT("FCbArrayView(NoType)::CopyTo()"), MakeMemoryView(Value).EqualBytes(MakeMemoryView(CopyBytes) + 1));
-		NamedField.AsArrayView().CopyTo(MakeMemoryView(CopyBytes));
-		TEST_TRUE(TEXT("FCbArrayView(NoType, Name)::CopyTo()"), MakeMemoryView(Value).EqualBytes(MakeMemoryView(CopyBytes) + 1));
-
-		// TryGetView
-		FMemoryView View;
-		TEST_TRUE(TEXT("FCbArrayView(Clone)::TryGetView()"), Array.TryGetView(View) && View == Array.GetOuterBuffer().GetView());
-		TEST_FALSE(TEXT("FCbArrayView(NoType)::TryGetView()"), Field.AsArrayView().TryGetView(View));
-		TEST_FALSE(TEXT("FCbArrayView(Name)::TryGetView()"), NamedField.AsArrayView().TryGetView(View));
 
-		// GetBuffer
-		TEST_TRUE(TEXT("FCbArrayView(Clone)::GetBuffer()"), Array.GetBuffer().ToShared().GetView() == Array.GetOuterBuffer().GetView());
-		TEST_TRUE(TEXT("FCbArrayView(NoType)::GetBuffer()"), FCbField::MakeView(Field).AsArray().GetBuffer().ToShared().GetView().EqualBytes(Array.GetOuterBuffer().GetView()));
-		TEST_TRUE(TEXT("FCbArrayView(Name)::GetBuffer()"), FCbField::MakeView(NamedField).AsArray().GetBuffer().ToShared().GetView().EqualBytes(Array.GetOuterBuffer().GetView()));
+		{
+			INFO("Equals");
+			CHECK(Field.AsArrayView().Equals(NamedField.AsArrayView()));
+			CHECK(Field.Equals(Field.AsArrayView().AsFieldView()));
+			CHECK(NamedField.Equals(NamedField.AsArrayView().AsFieldView()));
+		}
+		
+		{
+			INFO("CopyTo");
+			uint8 CopyBytes[sizeof(Value) + 1];
+			Field.AsArrayView().CopyTo(MakeMemoryView(CopyBytes));
+			CHECK(MakeMemoryView(Value).EqualBytes(MakeMemoryView(CopyBytes) + 1));
+			NamedField.AsArrayView().CopyTo(MakeMemoryView(CopyBytes));
+			CHECK(MakeMemoryView(Value).EqualBytes(MakeMemoryView(CopyBytes) + 1));
+		}
+		
+		{
+			INFO("TryGetView");
+			FMemoryView View;
+			CHECK(Array.TryGetView(View));
+			CHECK(View == Array.GetOuterBuffer().GetView());
+			CHECK_FALSE(Field.AsArrayView().TryGetView(View));
+			CHECK_FALSE(NamedField.AsArrayView().TryGetView(View));
+		}
+		
+		{
+			INFO("GetBuffer");
+			CHECK(Array.GetBuffer().ToShared().GetView() == Array.GetOuterBuffer().GetView());
+			CHECK(FCbField::MakeView(Field).AsArray().GetBuffer().ToShared().GetView().EqualBytes(Array.GetOuterBuffer().GetView()));
+			CHECK(FCbField::MakeView(NamedField).AsArray().GetBuffer().ToShared().GetView().EqualBytes(Array.GetOuterBuffer().GetView()));
+		}
 	}
 
-	// Test FCbFieldView(None) as Array
+	SECTION("Test FCbFieldView(None) as Array")
 	{
 		FCbFieldView Field;
 		TestFieldError<ECbFieldType::Array>(*this, TEXT("Array, None"), Field, ECbFieldError::TypeError);
 		FCbField::MakeView(Field).AsArray();
 	}
 
-	// Test FCbArrayView(ArrayWithName) and CreateIterator
+	SECTION("Test FCbArrayView(ArrayWithName) and CreateIterator")
 	{
 		const uint8 ArrayType = uint8(ECbFieldType::Array | ECbFieldType::HasFieldName);
 		const uint8 Buffer[] = { ArrayType, 3, 'K', 'e', 'y', 3, 1, uint8(ECbFieldType::IntegerPositive), 8 };
 		const FCbArrayView Array(Buffer);
-		TEST_EQUAL(TEXT("FCbArrayView(ArrayWithName)::GetSize()"), Array.GetSize(), uint64(5));
+		CHECK_EQUAL(Array.GetSize(), uint64(5));
 		const FCbArray ArrayClone = FCbArray::Clone(Array);
-		TEST_EQUAL(TEXT("FCbArray(ArrayWithName)::GetSize()"), ArrayClone.GetSize(), uint64(5));
-		TEST_TRUE(TEXT("FCbArrayView::Equals()"), Array.Equals(ArrayClone));
-		TEST_EQUAL(TEXT("FCbArrayView::GetHash()"), ArrayClone.GetHash(), Array.GetHash());
+		CHECK_EQUAL(ArrayClone.GetSize(), uint64(5));
+		CHECK(Array.Equals(ArrayClone));
+		CHECK_EQUAL(ArrayClone.GetHash(), Array.GetHash());
 		for (FCbFieldIterator It = ArrayClone.CreateIterator(); It; ++It)
 		{
 			FCbField Field = *It;
-			TEST_EQUAL(TEXT("FCbArray::CreateIterator().AsInt32()"), Field.AsInt32(), 8);
-			TEST_TRUE(TEXT("FCbArray::CreateIterator().IsOwned()"), Field.IsOwned());
+			CHECK_EQUAL(Field.AsInt32(), 8);
+			CHECK(Field.IsOwned());
 		}
 		for (FCbFieldIterator It = ArrayClone.CreateIterator(), End; It != End; ++It)
 		{
@@ -628,45 +669,48 @@ TEST_CASE_METHOD(FAutomationTestFixture, "Core::Serialization::FCbField::Array T
 		{
 		}
 
-		// CopyTo
-		uint8 CopyBytes[5];
-		Array.CopyTo(MakeMemoryView(CopyBytes));
-		TEST_TRUE(TEXT("FCbArrayView(Name)::CopyTo()"), ArrayClone.GetOuterBuffer().GetView().EqualBytes(MakeMemoryView(CopyBytes)));
-		ArrayClone.CopyTo(MakeMemoryView(CopyBytes));
-		TEST_TRUE(TEXT("FCbArrayView()::CopyTo()"), ArrayClone.GetOuterBuffer().GetView().EqualBytes(MakeMemoryView(CopyBytes)));
-
-		// GetBuffer
-		TEST_TRUE(TEXT("FCbField(ArrayWithName)::GetBuffer()"), FCbField(FSharedBuffer::MakeView(MakeMemoryView(Buffer))).GetBuffer().ToShared().GetView().EqualBytes(MakeMemoryView(Buffer)));
-		TEST_TRUE(TEXT("FCbField(ArrayWithNameNoType)::GetBuffer()"),
-			FCbField(FCbFieldView(Buffer + 1, ECbFieldType(ArrayType)), FSharedBuffer::MakeView(MakeMemoryView(Buffer))).GetBuffer().ToShared().GetView().EqualBytes(MakeMemoryView(Buffer)));
+		{
+			INFO("CopyTo");
+			uint8 CopyBytes[5];
+			Array.CopyTo(MakeMemoryView(CopyBytes));
+			CHECK(ArrayClone.GetOuterBuffer().GetView().EqualBytes(MakeMemoryView(CopyBytes)));
+			ArrayClone.CopyTo(MakeMemoryView(CopyBytes));
+			CHECK(ArrayClone.GetOuterBuffer().GetView().EqualBytes(MakeMemoryView(CopyBytes)));
+		}
+		
+		{
+			INFO("GetBuffer")
+			CHECK(FCbField(FSharedBuffer::MakeView(MakeMemoryView(Buffer))).GetBuffer().ToShared().GetView().EqualBytes(MakeMemoryView(Buffer)));
+			CHECK(FCbField(FCbFieldView(Buffer + 1, ECbFieldType(ArrayType)), FSharedBuffer::MakeView(MakeMemoryView(Buffer))).GetBuffer().ToShared().GetView().EqualBytes(MakeMemoryView(Buffer)));
+		}
 	}
 
-	// Test FCbArrayView as FCbFieldViewIterator
+	SECTION("Test FCbArrayView as FCbFieldViewIterator")
 	{
 		uint32 Count = 0;
 		FCbArrayView Array;
 		for (FCbFieldView Field : FCbFieldViewIterator::MakeSingle(Array.AsFieldView()))
 		{
-			TEST_TRUE(TEXT("FCbArrayView::AsFieldView() as Iterator"), Field.IsArray());
+			CHECK(Field.IsArray());
 			++Count;
 		}
-		TEST_EQUAL(TEXT("FCbArrayView::AsFieldView() as Iterator Count"), Count, 1u);
+		CHECK_EQUAL(Count, 1u);
 	}
 
-	// Test FCbArray as FCbFieldIterator
+	SECTION("Test FCbArray as FCbFieldIterator")
 	{
 		uint32 Count = 0;
 		FCbArray Array;
 		Array.MakeOwned();
 		for (FCbField Field : FCbFieldIterator::MakeSingle(Array.AsField()))
 		{
-			TEST_TRUE(TEXT("FCbArray::AsFieldView() as Iterator"), Field.IsArray());
+			CHECK(Field.IsArray());
 			++Count;
 		}
-		TEST_EQUAL(TEXT("FCbArray::AsFieldView() as Iterator Count"), Count, 1u);
+		CHECK_EQUAL(Count, 1u);
 	}
 
-	// Test FCbArray(Empty) as FCbFieldIterator
+	SECTION("Test FCbArray(Empty) as FCbFieldIterator")
 	{
 		const uint8 Buffer[] = { uint8(ECbFieldType::Array), 1, 0 };
 		const FCbArray Array = FCbArray::Clone(Buffer);
@@ -685,10 +729,12 @@ TEST_CASE_METHOD(FAutomationTestFixture, "Core::Serialization::FCbField::Binary
 			[](FCbFieldView& Field, const FSharedBuffer& Default) -> FSharedBuffer { return static_cast<FCbField&>(Field).AsBinary(Default); };
 	};
 
-	// Test FCbFieldView(Binary, Empty)
-	TestField<ECbFieldType::Binary>(*this, TEXT("Binary, Empty"), {0});
+	SECTION("Test FCbFieldView(Binary, Empty)")
+	{
+		TestField<ECbFieldType::Binary>(*this, TEXT("Binary, Empty"), { 0 });
+	}
 
-	// Test FCbFieldView(Binary, Value)
+	SECTION("Test FCbFieldView(Binary, Value)")
 	{
 		const uint8 Value[] = { 3, 4, 5, 6 }; // Size: 3, Data: 4/5/6
 		FCbFieldView FieldView(Value, ECbFieldType::Binary);
@@ -696,12 +742,12 @@ TEST_CASE_METHOD(FAutomationTestFixture, "Core::Serialization::FCbField::Binary
 
 		FCbField Field = FCbField::Clone(FieldView);
 		Field.AsBinary();
-		TEST_FALSE(TEXT("Binary, Value, AsBinary -> GetOuterBuffer()"), Field.GetOuterBuffer().IsNull());
+		CHECK_FALSE(Field.GetOuterBuffer().IsNull());
 		MoveTemp(Field).AsBinary();
-		TEST_TRUE(TEXT("Binary, Value, AsBinary -> GetOuterBuffer()"), Field.GetOuterBuffer().IsNull());
+		CHECK(Field.GetOuterBuffer().IsNull());
 	}
 
-	// Test FCbFieldView(None) as Binary
+	SECTION("Test FCbFieldView(None) as Binary")
 	{
 		FCbFieldView FieldView;
 		const uint8 Default[] = { 1, 2, 3 };
@@ -714,23 +760,25 @@ TEST_CASE_METHOD(FAutomationTestFixture, "Core::Serialization::FCbField::Binary
 
 TEST_CASE_METHOD(FAutomationTestFixture, "Core::Serialization::FCbField::String Test", "[Core][Serialization][Smoke]")
 {
-	// Test FCbFieldView(String, Empty)
-	TestField<ECbFieldType::String>(*this, TEXT("String, Empty"), {0});
-
-	// Test FCbFieldView(String, Value)
+	SECTION("Test FCbFieldView(String, Empty)")
+	{
+		TestField<ECbFieldType::String>(*this, TEXT("String, Empty"), { 0 });
+	}
+	
+	SECTION("Test FCbFieldView(String, Value)")
 	{
 		const uint8 Value[] = { 3, 'A', 'B', 'C' }; // Size: 3, Data: ABC
 		TestField<ECbFieldType::String>(*this, TEXT("String, Value"), Value, FUtf8StringView(reinterpret_cast<const UTF8CHAR*>(Value) + 1, 3));
 	}
 
-	// Test FCbFieldView(String, OutOfRangeSize)
+	SECTION("Test FCbFieldView(String, OutOfRangeSize)")
 	{
 		uint8 Value[9];
 		WriteVarUInt(uint64(1) << 31, Value);
 		TestFieldError<ECbFieldType::String>(*this, TEXT("String, OutOfRangeSize"), Value, ECbFieldError::RangeError, "ABC"_U8SV);
 	}
 
-	// Test FCbFieldView(None) as String
+	SECTION("Test FCbFieldView(None) as String")
 	{
 		FCbFieldView Field;
 		TestFieldError<ECbFieldType::String>(*this, TEXT("String, None"), Field, ECbFieldError::TypeError, "ABC"_U8SV);
@@ -800,37 +848,41 @@ void TestIntegerField(FAutomationTestFixture& Test, ECbFieldType FieldType, EInt
 
 TEST_CASE_METHOD(FAutomationTestFixture, "Core::Serialization::FCbField::Integer Test", "[Core][Serialization][Smoke]")
 {
-	// Test FCbFieldView(IntegerPositive)
-	TestIntegerField(*this, ECbFieldType::IntegerPositive, EIntType::Pos7,  0x00);
-	TestIntegerField(*this, ECbFieldType::IntegerPositive, EIntType::Pos7,  0x7f);
-	TestIntegerField(*this, ECbFieldType::IntegerPositive, EIntType::Pos8,  0x80);
-	TestIntegerField(*this, ECbFieldType::IntegerPositive, EIntType::Pos8,  0xff);
-	TestIntegerField(*this, ECbFieldType::IntegerPositive, EIntType::Pos15, 0x0100);
-	TestIntegerField(*this, ECbFieldType::IntegerPositive, EIntType::Pos15, 0x7fff);
-	TestIntegerField(*this, ECbFieldType::IntegerPositive, EIntType::Pos16, 0x8000);
-	TestIntegerField(*this, ECbFieldType::IntegerPositive, EIntType::Pos16, 0xffff);
-	TestIntegerField(*this, ECbFieldType::IntegerPositive, EIntType::Pos31, 0x0001'0000);
-	TestIntegerField(*this, ECbFieldType::IntegerPositive, EIntType::Pos31, 0x7fff'ffff);
-	TestIntegerField(*this, ECbFieldType::IntegerPositive, EIntType::Pos32, 0x8000'0000);
-	TestIntegerField(*this, ECbFieldType::IntegerPositive, EIntType::Pos32, 0xffff'ffff);
-	TestIntegerField(*this, ECbFieldType::IntegerPositive, EIntType::Pos63, 0x0000'0001'0000'0000);
-	TestIntegerField(*this, ECbFieldType::IntegerPositive, EIntType::Pos63, 0x7fff'ffff'ffff'ffff);
-	TestIntegerField(*this, ECbFieldType::IntegerPositive, EIntType::Pos64, 0x8000'0000'0000'0000);
-	TestIntegerField(*this, ECbFieldType::IntegerPositive, EIntType::Pos64, 0xffff'ffff'ffff'ffff);
-
-	// Test FCbFieldView(IntegerNegative)
-	TestIntegerField(*this, ECbFieldType::IntegerNegative, EIntType::Neg7,  0x01);
-	TestIntegerField(*this, ECbFieldType::IntegerNegative, EIntType::Neg7,  0x80);
-	TestIntegerField(*this, ECbFieldType::IntegerNegative, EIntType::Neg15, 0x81);
-	TestIntegerField(*this, ECbFieldType::IntegerNegative, EIntType::Neg15, 0x8000);
-	TestIntegerField(*this, ECbFieldType::IntegerNegative, EIntType::Neg31, 0x8001);
-	TestIntegerField(*this, ECbFieldType::IntegerNegative, EIntType::Neg31, 0x8000'0000);
-	TestIntegerField(*this, ECbFieldType::IntegerNegative, EIntType::Neg63, 0x8000'0001);
-	TestIntegerField(*this, ECbFieldType::IntegerNegative, EIntType::Neg63, 0x8000'0000'0000'0000);
-	TestIntegerField(*this, ECbFieldType::IntegerNegative, EIntType::None,  0x8000'0000'0000'0001);
-	TestIntegerField(*this, ECbFieldType::IntegerNegative, EIntType::None,  0xffff'ffff'ffff'ffff);
-
-	// Test FCbFieldView(None) as Integer
+	SECTION("Test FCbFieldView(IntegerPositive)")
+	{
+		TestIntegerField(*this, ECbFieldType::IntegerPositive, EIntType::Pos7, 0x00);
+		TestIntegerField(*this, ECbFieldType::IntegerPositive, EIntType::Pos7, 0x7f);
+		TestIntegerField(*this, ECbFieldType::IntegerPositive, EIntType::Pos8, 0x80);
+		TestIntegerField(*this, ECbFieldType::IntegerPositive, EIntType::Pos8, 0xff);
+		TestIntegerField(*this, ECbFieldType::IntegerPositive, EIntType::Pos15, 0x0100);
+		TestIntegerField(*this, ECbFieldType::IntegerPositive, EIntType::Pos15, 0x7fff);
+		TestIntegerField(*this, ECbFieldType::IntegerPositive, EIntType::Pos16, 0x8000);
+		TestIntegerField(*this, ECbFieldType::IntegerPositive, EIntType::Pos16, 0xffff);
+		TestIntegerField(*this, ECbFieldType::IntegerPositive, EIntType::Pos31, 0x0001'0000);
+		TestIntegerField(*this, ECbFieldType::IntegerPositive, EIntType::Pos31, 0x7fff'ffff);
+		TestIntegerField(*this, ECbFieldType::IntegerPositive, EIntType::Pos32, 0x8000'0000);
+		TestIntegerField(*this, ECbFieldType::IntegerPositive, EIntType::Pos32, 0xffff'ffff);
+		TestIntegerField(*this, ECbFieldType::IntegerPositive, EIntType::Pos63, 0x0000'0001'0000'0000);
+		TestIntegerField(*this, ECbFieldType::IntegerPositive, EIntType::Pos63, 0x7fff'ffff'ffff'ffff);
+		TestIntegerField(*this, ECbFieldType::IntegerPositive, EIntType::Pos64, 0x8000'0000'0000'0000);
+		TestIntegerField(*this, ECbFieldType::IntegerPositive, EIntType::Pos64, 0xffff'ffff'ffff'ffff);
+	}
+
+	SECTION("Test FCbFieldView(IntegerNegative)")
+	{
+		TestIntegerField(*this, ECbFieldType::IntegerNegative, EIntType::Neg7, 0x01);
+		TestIntegerField(*this, ECbFieldType::IntegerNegative, EIntType::Neg7, 0x80);
+		TestIntegerField(*this, ECbFieldType::IntegerNegative, EIntType::Neg15, 0x81);
+		TestIntegerField(*this, ECbFieldType::IntegerNegative, EIntType::Neg15, 0x8000);
+		TestIntegerField(*this, ECbFieldType::IntegerNegative, EIntType::Neg31, 0x8001);
+		TestIntegerField(*this, ECbFieldType::IntegerNegative, EIntType::Neg31, 0x8000'0000);
+		TestIntegerField(*this, ECbFieldType::IntegerNegative, EIntType::Neg63, 0x8000'0001);
+		TestIntegerField(*this, ECbFieldType::IntegerNegative, EIntType::Neg63, 0x8000'0000'0000'0000);
+		TestIntegerField(*this, ECbFieldType::IntegerNegative, EIntType::None, 0x8000'0000'0000'0001);
+		TestIntegerField(*this, ECbFieldType::IntegerNegative, EIntType::None, 0xffff'ffff'ffff'ffff);
+	}
+	
+	SECTION("Test FCbFieldView(None) as Integer")
 	{
 		FCbFieldView Field;
 		TestFieldError<ECbFieldType::IntegerPositive>(*this, TEXT("Integer+, None"), Field, ECbFieldError::TypeError, uint64(8));
@@ -841,7 +893,7 @@ TEST_CASE_METHOD(FAutomationTestFixture, "Core::Serialization::FCbField::Integer
 
 TEST_CASE_METHOD(FAutomationTestFixture, "Core::Serialization::FCbField::Float Test", "[Core][Serialization][Smoke]")
 {
-	//// Test FCbFieldView(Float, 32-bit)
+	//SECTION("Test FCbFieldView(Float, 32-bit)")
 	//{
 	//	const uint8 Value[] = { 0xc0, 0x12, 0x34, 0x56 }; // -2.28444433f
 	//	TestField<ECbFieldType::Float32>(TEXT("Float32"), Value, -2.2844443321f);
@@ -850,7 +902,7 @@ TEST_CASE_METHOD(FAutomationTestFixture, "Core::Serialization::FCbField::Float T
 	//	TestField<ECbFieldType::Float64>(TEXT("Float32, AsDouble"), Field, -2.2844443321);
 	//}
 
-	// Test FCbFieldView(Float, 64-bit)
+	SECTION("Test FCbFieldView(Float, 64-bit)")
 	{
 		const uint8 Value[] = { 0xc1, 0x23, 0x45, 0x67, 0x89, 0xab, 0xcd, 0xef }; // -631475.76888888876
 		TestField<ECbFieldType::Float64>(*this, TEXT("Float64"), Value, -631475.76888888876);
@@ -859,7 +911,7 @@ TEST_CASE_METHOD(FAutomationTestFixture, "Core::Serialization::FCbField::Float T
 		TestFieldError<ECbFieldType::Float32>(*this, TEXT("Float64, AsFloat"), Field, ECbFieldError::RangeError, 8.0f);
 	}
 
-	// Test FCbFieldView(Integer+, MaxBinary32) as Float
+	SECTION("Test FCbFieldView(Integer+, MaxBinary32) as Float")
 	{
 		uint8 Value[9];
 		WriteVarUInt((uint64(1) << 24) - 1, Value); // 16,777,215
@@ -868,7 +920,7 @@ TEST_CASE_METHOD(FAutomationTestFixture, "Core::Serialization::FCbField::Float T
 		TestField<ECbFieldType::Float64>(*this, TEXT("Integer+, MaxBinary32, AsDouble"), Field, 16'777'215.0);
 	}
 
-	// Test FCbFieldView(Integer+, MaxBinary32+1) as Float
+	SECTION("Test FCbFieldView(Integer+, MaxBinary32+1) as Float")
 	{
 		uint8 Value[9];
 		WriteVarUInt(uint64(1) << 24, Value); // 16,777,216
@@ -877,7 +929,7 @@ TEST_CASE_METHOD(FAutomationTestFixture, "Core::Serialization::FCbField::Float T
 		TestField<ECbFieldType::Float64>(*this, TEXT("Integer+, MaxBinary32+1, AsDouble"), Field, 16'777'216.0);
 	}
 
-	// Test FCbFieldView(Integer+, MaxBinary64) as Float
+	SECTION("Test FCbFieldView(Integer+, MaxBinary64) as Float")
 	{
 		uint8 Value[9];
 		WriteVarUInt((uint64(1) << 53) - 1, Value); // 9,007,199,254,740,991
@@ -886,7 +938,7 @@ TEST_CASE_METHOD(FAutomationTestFixture, "Core::Serialization::FCbField::Float T
 		TestField<ECbFieldType::Float64>(*this, TEXT("Integer+, MaxBinary64, AsDouble"), Field, 9'007'199'254'740'991.0);
 	}
 
-	// Test FCbFieldView(Integer+, MaxBinary64+1) as Float
+	SECTION("Test FCbFieldView(Integer+, MaxBinary64+1) as Float")
 	{
 		uint8 Value[9];
 		WriteVarUInt(uint64(1) << 53, Value); // 9,007,199,254,740,992
@@ -895,7 +947,7 @@ TEST_CASE_METHOD(FAutomationTestFixture, "Core::Serialization::FCbField::Float T
 		TestFieldError<ECbFieldType::Float64>(*this, TEXT("Integer+, MaxBinary64+1, AsDouble"), Field, ECbFieldError::RangeError, 8.0);
 	}
 
-	// Test FCbFieldView(Integer+, MaxUInt64) as Float
+	SECTION("Test FCbFieldView(Integer+, MaxUInt64) as Float")
 	{
 		uint8 Value[9];
 		WriteVarUInt(uint64(-1), Value); // Max uint64
@@ -904,7 +956,7 @@ TEST_CASE_METHOD(FAutomationTestFixture, "Core::Serialization::FCbField::Float T
 		TestFieldError<ECbFieldType::Float64>(*this, TEXT("Integer+, MaxUInt64, AsDouble"), Field, ECbFieldError::RangeError, 8.0);
 	}
 
-	// Test FCbFieldView(Integer-, MaxBinary32) as Float
+	SECTION("Test FCbFieldView(Integer-, MaxBinary32) as Float")
 	{
 		uint8 Value[9];
 		WriteVarUInt((uint64(1) << 24) - 2, Value); // -16,777,215
@@ -913,7 +965,7 @@ TEST_CASE_METHOD(FAutomationTestFixture, "Core::Serialization::FCbField::Float T
 		TestField<ECbFieldType::Float64>(*this, TEXT("Integer-, MaxBinary32, AsDouble"), Field, -16'777'215.0);
 	}
 
-	// Test FCbFieldView(Integer-, MaxBinary32+1) as Float
+	SECTION("Test FCbFieldView(Integer-, MaxBinary32+1) as Float")
 	{
 		uint8 Value[9];
 		WriteVarUInt((uint64(1) << 24) - 1, Value); // -16,777,216
@@ -922,7 +974,7 @@ TEST_CASE_METHOD(FAutomationTestFixture, "Core::Serialization::FCbField::Float T
 		TestField<ECbFieldType::Float64>(*this, TEXT("Integer-, MaxBinary32+1, AsDouble"), Field, -16'777'216.0);
 	}
 
-	// Test FCbFieldView(Integer-, MaxBinary64) as Float
+	SECTION("Test FCbFieldView(Integer-, MaxBinary64) as Float")
 	{
 		uint8 Value[9];
 		WriteVarUInt((uint64(1) << 53) - 2, Value); // -9,007,199,254,740,991
@@ -931,7 +983,7 @@ TEST_CASE_METHOD(FAutomationTestFixture, "Core::Serialization::FCbField::Float T
 		TestField<ECbFieldType::Float64>(*this, TEXT("Integer-, MaxBinary64, AsDouble"), Field, -9'007'199'254'740'991.0);
 	}
 
-	// Test FCbFieldView(Integer-, MaxBinary64+1) as Float
+	SECTION("Test FCbFieldView(Integer-, MaxBinary64+1) as Float")
 	{
 		uint8 Value[9];
 		WriteVarUInt((uint64(1) << 53) - 1, Value); // -9,007,199,254,740,992
@@ -940,7 +992,7 @@ TEST_CASE_METHOD(FAutomationTestFixture, "Core::Serialization::FCbField::Float T
 		TestFieldError<ECbFieldType::Float64>(*this, TEXT("Integer-, MaxBinary64+1, AsDouble"), Field, ECbFieldError::RangeError, 8.0);
 	}
 
-	// Test FCbFieldView(None) as Float
+	SECTION("Test FCbFieldView(None) as Float")
 	{
 		FCbFieldView Field;
 		TestFieldError<ECbFieldType::Float32>(*this, TEXT("None, AsFloat"), Field, ECbFieldError::TypeError, 8.0f);
@@ -951,13 +1003,17 @@ TEST_CASE_METHOD(FAutomationTestFixture, "Core::Serialization::FCbField::Float T
 
 TEST_CASE_METHOD(FAutomationTestFixture, "Core::Serialization::FCbField::Bool Test", "[Core][Serialization][Smoke]")
 {
-	// Test FCbFieldView(Bool, False)
-	TestField<ECbFieldType::BoolFalse>(*this, TEXT("Bool, False"), {}, false, true);
-
-	// Test FCbFieldView(Bool, True)
-	TestField<ECbFieldType::BoolTrue>(*this, TEXT("Bool, True"), {}, true, false);
-
-	// Test FCbFieldView(None) as Bool
+	SECTION("Test FCbFieldView(Bool, False)")
+	{
+		TestField<ECbFieldType::BoolFalse>(*this, TEXT("Bool, False"), {}, false, true);
+	}
+	
+	SECTION("Test FCbFieldView(Bool, True)")
+	{
+		TestField<ECbFieldType::BoolTrue>(*this, TEXT("Bool, True"), {}, true, false);
+	}
+	
+	SECTION("Test FCbFieldView(None) as Bool")
 	{
 		FCbFieldView DefaultField;
 		TestFieldError<ECbFieldType::BoolFalse>(*this, TEXT("Bool, False, None"), DefaultField, ECbFieldError::TypeError, false);
@@ -971,19 +1027,23 @@ TEST_CASE_METHOD(FAutomationTestFixture, "Core::Serialization::FCbField::Object
 	const FIoHash::ByteArray ZeroBytes{};
 	const FIoHash::ByteArray SequentialBytes{1, 2, 3, 4, 5, 6, 7, 8, 9, 10, 11, 12, 13, 14, 15, 16, 17, 18, 19, 20};
 
-	// Test FCbFieldView(ObjectAttachment, Zero)
-	TestField<ECbFieldType::ObjectAttachment>(*this, TEXT("ObjectAttachment, Zero"), ZeroBytes);
-
-	// Test FCbFieldView(ObjectAttachment, NonZero)
-	TestField<ECbFieldType::ObjectAttachment>(*this, TEXT("ObjectAttachment, NonZero"), SequentialBytes, FIoHash(SequentialBytes));
-
-	// Test FCbFieldView(ObjectAttachment, NonZero) AsAttachment
+	SECTION("Test FCbFieldView(ObjectAttachment, Zero)")
+	{
+		TestField<ECbFieldType::ObjectAttachment>(*this, TEXT("ObjectAttachment, Zero"), ZeroBytes);
+	}
+	
+	SECTION("Test FCbFieldView(ObjectAttachment, NonZero)")
+	{
+		TestField<ECbFieldType::ObjectAttachment>(*this, TEXT("ObjectAttachment, NonZero"), SequentialBytes, FIoHash(SequentialBytes));
+	}
+	
+	SECTION("Test FCbFieldView(ObjectAttachment, NonZero) AsAttachment")
 	{
 		FCbFieldView Field(SequentialBytes, ECbFieldType::ObjectAttachment);
 		TestField<ECbFieldType::ObjectAttachment>(*this, TEXT("ObjectAttachment, NonZero, AsAttachment"), Field, FIoHash(SequentialBytes), FIoHash(), ECbFieldError::None, FCbAttachmentAccessors());
 	}
 
-	// Test FCbFieldView(None) as ObjectAttachment
+	SECTION("Test FCbFieldView(None) as ObjectAttachment")
 	{
 		FCbFieldView DefaultField;
 		TestFieldError<ECbFieldType::ObjectAttachment>(*this, TEXT("ObjectAttachment, None"), DefaultField, ECbFieldError::TypeError, FIoHash(SequentialBytes));
@@ -996,19 +1056,23 @@ TEST_CASE_METHOD(FAutomationTestFixture, "Core::Serialization::FCbField::Binary
 	const FIoHash::ByteArray ZeroBytes{};
 	const FIoHash::ByteArray SequentialBytes{1, 2, 3, 4, 5, 6, 7, 8, 9, 10, 11, 12, 13, 14, 15, 16, 17, 18, 19, 20};
 
-	// Test FCbFieldView(BinaryAttachment, Zero)
-	TestField<ECbFieldType::BinaryAttachment>(*this, TEXT("BinaryAttachment, Zero"), ZeroBytes);
-
-	// Test FCbFieldView(BinaryAttachment, NonZero)
-	TestField<ECbFieldType::BinaryAttachment>(*this, TEXT("BinaryAttachment, NonZero"), SequentialBytes, FIoHash(SequentialBytes));
-
-	// Test FCbFieldView(BinaryAttachment, NonZero) AsAttachment
+	SECTION("Test FCbFieldView(BinaryAttachment, Zero)")
+	{
+		TestField<ECbFieldType::BinaryAttachment>(*this, TEXT("BinaryAttachment, Zero"), ZeroBytes);
+	}
+	
+	SECTION("Test FCbFieldView(BinaryAttachment, NonZero)")
+	{
+		TestField<ECbFieldType::BinaryAttachment>(*this, TEXT("BinaryAttachment, NonZero"), SequentialBytes, FIoHash(SequentialBytes));
+	}
+	
+	SECTION("Test FCbFieldView(BinaryAttachment, NonZero) AsAttachment")
 	{
 		FCbFieldView Field(SequentialBytes, ECbFieldType::BinaryAttachment);
 		TestField<ECbFieldType::BinaryAttachment>(*this, TEXT("BinaryAttachment, NonZero, AsAttachment"), Field, FIoHash(SequentialBytes), FIoHash(), ECbFieldError::None, FCbAttachmentAccessors());
 	}
 
-	// Test FCbFieldView(None) as BinaryAttachment
+	SECTION("Test FCbFieldView(None) as BinaryAttachment")
 	{
 		FCbFieldView DefaultField;
 		TestFieldError<ECbFieldType::BinaryAttachment>(*this, TEXT("BinaryAttachment, None"), DefaultField, ECbFieldError::TypeError, FIoHash(SequentialBytes));
@@ -1021,25 +1085,29 @@ TEST_CASE_METHOD(FAutomationTestFixture, "Core::Serialization::FCbField::Hash Te
 	const FIoHash::ByteArray ZeroBytes{};
 	const FIoHash::ByteArray SequentialBytes{1, 2, 3, 4, 5, 6, 7, 8, 9, 10, 11, 12, 13, 14, 15, 16, 17, 18, 19, 20};
 
-	// Test FCbFieldView(Hash, Zero)
-	TestField<ECbFieldType::Hash>(*this, TEXT("Hash, Zero"), ZeroBytes);
-
-	// Test FCbFieldView(Hash, NonZero)
-	TestField<ECbFieldType::Hash>(*this, TEXT("Hash, NonZero"), SequentialBytes, FIoHash(SequentialBytes));
-
-	// Test FCbFieldView(None) as Hash
+	SECTION("Test FCbFieldView(Hash, Zero)")
+	{
+		TestField<ECbFieldType::Hash>(*this, TEXT("Hash, Zero"), ZeroBytes);
+	}
+	
+	SECTION("Test FCbFieldView(Hash, NonZero)")
+	{
+		TestField<ECbFieldType::Hash>(*this, TEXT("Hash, NonZero"), SequentialBytes, FIoHash(SequentialBytes));
+	}
+	
+	SECTION("Test FCbFieldView(None) as Hash")
 	{
 		FCbFieldView DefaultField;
 		TestFieldError<ECbFieldType::Hash>(*this, TEXT("Hash, None"), DefaultField, ECbFieldError::TypeError, FIoHash(SequentialBytes));
 	}
 
-	// Test FCbFieldView(ObjectAttachment) as Hash
+	SECTION("Test FCbFieldView(ObjectAttachment) as Hash")
 	{
 		FCbFieldView Field(SequentialBytes, ECbFieldType::ObjectAttachment);
 		TestField<ECbFieldType::Hash>(*this, TEXT("ObjectAttachment, NonZero, AsHash"), Field, FIoHash(SequentialBytes));
 	}
 
-	// Test FCbFieldView(BinaryAttachment) as Hash
+	SECTION("Test FCbFieldView(BinaryAttachment) as Hash")
 	{
 		FCbFieldView Field(SequentialBytes, ECbFieldType::BinaryAttachment);
 		TestField<ECbFieldType::Hash>(*this, TEXT("BinaryAttachment, NonZero, AsHash"), Field, FIoHash(SequentialBytes));
@@ -1052,13 +1120,17 @@ TEST_CASE_METHOD(FAutomationTestFixture, "Core::Serialization::FCbField::Uuid Te
 	const uint8 SequentialBytes[]{0, 1, 2, 3, 4, 5, 6, 7, 8, 9, 10, 11, 12, 13, 14, 15};
 	const FGuid SequentialGuid(TEXT("00010203-0405-0607-0809-0a0b0c0d0e0f"));
 
-	// Test FCbFieldView(Uuid, Zero)
-	TestField<ECbFieldType::Uuid>(*this, TEXT("Uuid, Zero"), ZeroBytes, FGuid(), SequentialGuid);
-
-	// Test FCbFieldView(Uuid, NonZero)
-	TestField<ECbFieldType::Uuid>(*this, TEXT("Uuid, NonZero"), SequentialBytes, SequentialGuid, FGuid());
-
-	// Test FCbFieldView(None) as Uuid
+	SECTION("Test FCbFieldView(Uuid, Zero)")
+	{
+		TestField<ECbFieldType::Uuid>(*this, TEXT("Uuid, Zero"), ZeroBytes, FGuid(), SequentialGuid);
+	}
+	
+	SECTION("Test FCbFieldView(Uuid, NonZero)")
+	{
+		TestField<ECbFieldType::Uuid>(*this, TEXT("Uuid, NonZero"), SequentialBytes, SequentialGuid, FGuid());
+	}
+	
+	SECTION("Test FCbFieldView(None) as Uuid")
 	{
 		FCbFieldView DefaultField;
 		TestFieldError<ECbFieldType::Uuid>(*this, TEXT("Uuid, None"), DefaultField, ECbFieldError::TypeError, FGuid::NewGuid());
@@ -1068,75 +1140,87 @@ TEST_CASE_METHOD(FAutomationTestFixture, "Core::Serialization::FCbField::Uuid Te
 
 TEST_CASE_METHOD(FAutomationTestFixture, "Core::Serialization::FCbField::Date Time", "[Core][Serialization][Smoke]")
 {
-	// Test FCbFieldView(DateTime, Zero)
-	TestField<ECbFieldType::DateTime>(*this, TEXT("DateTime, Zero"), {0, 0, 0, 0, 0, 0, 0, 0});
-
-	// Test FCbFieldView(DateTime, 0x1020'3040'5060'7080)
-	TestField<ECbFieldType::DateTime>(*this, TEXT("DateTime, NonZero"), {0x10, 0x20, 0x30, 0x40, 0x50, 0x60, 0x70, 0x80}, int64(0x1020'3040'5060'7080));
-
-	// Test FCbFieldView(DateTime, Zero) as FDateTime
+	SECTION("Test FCbFieldView(DateTime, Zero)")
+	{
+		TestField<ECbFieldType::DateTime>(*this, TEXT("DateTime, Zero"), { 0, 0, 0, 0, 0, 0, 0, 0 });
+	}
+	
+	SECTION("Test FCbFieldView(DateTime, 0x1020'3040'5060'7080)")
+	{
+		TestField<ECbFieldType::DateTime>(*this, TEXT("DateTime, NonZero"), { 0x10, 0x20, 0x30, 0x40, 0x50, 0x60, 0x70, 0x80 }, int64(0x1020'3040'5060'7080));
+	}
+	
+	SECTION("Test FCbFieldView(DateTime, Zero) as FDateTime")
 	{
 		const uint8 Value[] = {0, 0, 0, 0, 0, 0, 0, 0};
 		FCbFieldView Field(Value, ECbFieldType::DateTime);
-		TEST_EQUAL(TEXT("FCbFieldView()::AsDateTime()"), Field.AsDateTime(), FDateTime(0));
+		CHECK_EQUAL(Field.AsDateTime(), FDateTime(0));
 	}
 
-	// Test FCbFieldView(None) as DateTime
+	SECTION("Test FCbFieldView(None) as DateTime")
 	{
 		FCbFieldView DefaultField;
 		TestFieldError<ECbFieldType::DateTime>(*this, TEXT("DateTime, None"), DefaultField, ECbFieldError::TypeError);
 		const FDateTime DefaultValue(0x1020'3040'5060'7080);
-		TEST_EQUAL(TEXT("FCbFieldView()::AsDateTime()"), DefaultField.AsDateTime(DefaultValue), DefaultValue);
+		CHECK_EQUAL(DefaultField.AsDateTime(DefaultValue), DefaultValue);
 	}
 }
 
 TEST_CASE_METHOD(FAutomationTestFixture, "Core::Serialization::FCbField::Time Span", "[Core][Serialization][Smoke]")
 {
-	// Test FCbFieldView(TimeSpan, Zero)
-	TestField<ECbFieldType::TimeSpan>(*this, TEXT("TimeSpan, Zero"), {0, 0, 0, 0, 0, 0, 0, 0});
-
-	// Test FCbFieldView(TimeSpan, 0x1020'3040'5060'7080)
-	TestField<ECbFieldType::TimeSpan>(*this, TEXT("TimeSpan, NonZero"), {0x10, 0x20, 0x30, 0x40, 0x50, 0x60, 0x70, 0x80}, int64(0x1020'3040'5060'7080));
-
-	// Test FCbFieldView(TimeSpan, Zero) as FTimeSpan
+	SECTION("Test FCbFieldView(TimeSpan, Zero)")
+	{
+		TestField<ECbFieldType::TimeSpan>(*this, TEXT("TimeSpan, Zero"), { 0, 0, 0, 0, 0, 0, 0, 0 });
+	}
+	
+	SECTION("Test FCbFieldView(TimeSpan, 0x1020'3040'5060'7080)")
+	{
+		TestField<ECbFieldType::TimeSpan>(*this, TEXT("TimeSpan, NonZero"), { 0x10, 0x20, 0x30, 0x40, 0x50, 0x60, 0x70, 0x80 }, int64(0x1020'3040'5060'7080));
+	}
+	
+	SECTION("Test FCbFieldView(TimeSpan, Zero) as FTimeSpan")
 	{
 		const uint8 Value[] = {0, 0, 0, 0, 0, 0, 0, 0};
 		FCbFieldView Field(Value, ECbFieldType::TimeSpan);
-		TEST_EQUAL(TEXT("FCbFieldView()::AsTimeSpan()"), Field.AsTimeSpan(), FTimespan(0));
+		CHECK_EQUAL(Field.AsTimeSpan(), FTimespan(0));
 	}
 
-	// Test FCbFieldView(None) as TimeSpan
+	SECTION("Test FCbFieldView(None) as TimeSpan")
 	{
 		FCbFieldView DefaultField;
 		TestFieldError<ECbFieldType::TimeSpan>(*this, TEXT("TimeSpan, None"), DefaultField, ECbFieldError::TypeError);
 		const FTimespan DefaultValue(0x1020'3040'5060'7080);
-		TEST_EQUAL(TEXT("FCbFieldView()::AsTimeSpan()"), DefaultField.AsTimeSpan(DefaultValue), DefaultValue);
+		CHECK_EQUAL( DefaultField.AsTimeSpan(DefaultValue), DefaultValue);
 	}
 }
 
 
 TEST_CASE_METHOD(FAutomationTestFixture, "Core::Serialization::FCbField::ObjectId Span", "[Core][Serialization][Smoke]")
 {
-	// Test FCbFieldView(ObjectId, Zero)
-	TestField<ECbFieldType::ObjectId>(*this, TEXT("ObjectId, Zero"), {0, 0, 0, 0, 0, 0, 0, 0, 0, 0, 0, 0});
-
-	// Test FCbFieldView(ObjectId, 0x102030405060708090A0B0C0)
-	TestField<ECbFieldType::ObjectId>(*this, TEXT("ObjectId, NonZero"), {0x10, 0x20, 0x30, 0x40, 0x50, 0x60, 0x70, 0x80, 0x90, 0xA0, 0xB0, 0xC0},
-		FCbObjectId(MakeMemoryView<uint8>({0x10, 0x20, 0x30, 0x40, 0x50, 0x60, 0x70, 0x80, 0x90, 0xA0, 0xB0, 0xC0})));
-
-	// Test FCbFieldView(ObjectId, Zero) as FCbObjectId
+	SECTION("Test FCbFieldView(ObjectId, Zero)")
+	{
+		TestField<ECbFieldType::ObjectId>(*this, TEXT("ObjectId, Zero"), { 0, 0, 0, 0, 0, 0, 0, 0, 0, 0, 0, 0 });
+	}
+	
+	SECTION("Test FCbFieldView(ObjectId, 0x102030405060708090A0B0C0)")
+	{
+		TestField<ECbFieldType::ObjectId>(*this, TEXT("ObjectId, NonZero"), { 0x10, 0x20, 0x30, 0x40, 0x50, 0x60, 0x70, 0x80, 0x90, 0xA0, 0xB0, 0xC0 },
+			FCbObjectId(MakeMemoryView<uint8>({ 0x10, 0x20, 0x30, 0x40, 0x50, 0x60, 0x70, 0x80, 0x90, 0xA0, 0xB0, 0xC0 })));
+	}
+	
+	SECTION("Test FCbFieldView(ObjectId, Zero) as FCbObjectId")
 	{
 		const uint8 Value[] = {0, 0, 0, 0, 0, 0, 0, 0, 0, 0, 0, 0};
 		FCbFieldView Field(Value, ECbFieldType::ObjectId);
-		TEST_EQUAL(TEXT("FCbFieldView(ObjectId, Zero)::AsObjectId()"), Field.AsObjectId(), FCbObjectId());
+		CHECK_EQUAL(Field.AsObjectId(), FCbObjectId());
 	}
 
-	// Test FCbFieldView(None) as ObjectId
+	SECTION("Test FCbFieldView(None) as ObjectId")
 	{
 		FCbFieldView DefaultField;
 		TestFieldError<ECbFieldType::ObjectId>(*this, TEXT("ObjectId, None"), DefaultField, ECbFieldError::TypeError);
 		const FCbObjectId DefaultValue(MakeMemoryView<uint8>({0x10, 0x20, 0x30, 0x40, 0x50, 0x60, 0x70, 0x80, 0x90, 0xA0, 0xB0, 0xC0}));
-		TEST_EQUAL(TEXT("FCbFieldView(None)::AsObjectId()"), DefaultField.AsObjectId(DefaultValue), DefaultValue);
+		CHECK_EQUAL(DefaultField.AsObjectId(DefaultValue), DefaultValue);
 	}
 }
 
@@ -1154,7 +1238,7 @@ TEST_CASE_METHOD(FAutomationTestFixture, "Core::Serialization::FCbField::Custom
 		TUniqueFunction<FMemoryView (FCbFieldView& Field, FMemoryView Default)> AsType;
 	};
 
-	// Test FCbFieldView(CustomById, MinId, Empty)
+	SECTION("Test FCbFieldView(CustomById, MinId, Empty)")
 	{
 		const uint8 Value[] = {1, 0};
 		TestField<ECbFieldType::CustomById>(*this, TEXT("CustomById, MinId, Empty"), Value, FCbCustomById{0});
@@ -1162,7 +1246,7 @@ TEST_CASE_METHOD(FAutomationTestFixture, "Core::Serialization::FCbField::Custom
 		TestFieldError<ECbFieldType::CustomById>(*this, TEXT("CustomById, MinId, Empty, InvalidId"), Value, ECbFieldError::RangeError, MakeMemoryView<uint8>({1, 2, 3}), FCustomByIdAccessor(MAX_uint64));
 	}
 
-	// Test FCbFieldView(CustomById, MinId, Value)
+	SECTION("Test FCbFieldView(CustomById, MinId, Value)")
 	{
 		const uint8 Value[] = {5, 0, 1, 2, 3, 4};
 		TestFieldNoClone<ECbFieldType::CustomById>(*this, TEXT("CustomById, MinId, Value"), Value, FCbCustomById{0, MakeMemoryView(Value).Right(4)});
@@ -1170,7 +1254,7 @@ TEST_CASE_METHOD(FAutomationTestFixture, "Core::Serialization::FCbField::Custom
 		TestFieldError<ECbFieldType::CustomById>(*this, TEXT("CustomById, MinId, Value, InvalidId"), Value, ECbFieldError::RangeError, MakeMemoryView<uint8>({1, 2, 3}), FCustomByIdAccessor(MAX_uint64));
 	}
 
-	// Test FCbFieldView(CustomById, MaxId, Empty)
+	SECTION("Test FCbFieldView(CustomById, MaxId, Empty)")
 	{
 		const uint8 Value[] = {9, 0xff, 0xff, 0xff, 0xff, 0xff, 0xff, 0xff, 0xff, 0xff};
 		TestField<ECbFieldType::CustomById>(*this, TEXT("CustomById, MaxId, Empty"), Value, FCbCustomById{MAX_uint64});
@@ -1178,7 +1262,7 @@ TEST_CASE_METHOD(FAutomationTestFixture, "Core::Serialization::FCbField::Custom
 		TestFieldError<ECbFieldType::CustomById>(*this, TEXT("CustomById, MaxId, Empty, InvalidId"), Value, ECbFieldError::RangeError, MakeMemoryView<uint8>({1, 2, 3}), FCustomByIdAccessor(0));
 	}
 
-	// Test FCbFieldView(CustomById, MaxId, Value)
+	SECTION(" Test FCbFieldView(CustomById, MaxId, Value)")
 	{
 		const uint8 Value[] = {13, 0xff, 0xff, 0xff, 0xff, 0xff, 0xff, 0xff, 0xff, 0xff, 1, 2, 3, 4};
 		TestFieldNoClone<ECbFieldType::CustomById>(*this, TEXT("CustomById, MaxId, Value"), Value, FCbCustomById{MAX_uint64, MakeMemoryView(Value).Right(4)});
@@ -1186,13 +1270,13 @@ TEST_CASE_METHOD(FAutomationTestFixture, "Core::Serialization::FCbField::Custom
 		TestFieldError<ECbFieldType::CustomById>(*this, TEXT("CustomById, MaxId, Value, InvalidId"), Value, ECbFieldError::RangeError, MakeMemoryView<uint8>({1, 2, 3}), FCustomByIdAccessor(0));
 	}
 
-	// Test FCbFieldView(None) as CustomById
+	SECTION("Test FCbFieldView(None) as CustomById")
 	{
 		FCbFieldView DefaultField;
 		TestFieldError<ECbFieldType::CustomById>(*this, TEXT("CustomById, None"), DefaultField, ECbFieldError::TypeError, FCbCustomById{4, MakeMemoryView<uint8>({1, 2, 3})});
 		TestFieldError<ECbFieldType::CustomById>(*this, TEXT("CustomById, None, View"), DefaultField, ECbFieldError::TypeError, MakeMemoryView<uint8>({1, 2, 3}), FCustomByIdAccessor(0));
 		const uint8 DefaultValue[] = {1, 2, 3};
-		TEST_EQUAL(TEXT("FCbFieldView()::AsCustom(Id)"), DefaultField.AsCustom(0, MakeMemoryView(DefaultValue)), MakeMemoryView(DefaultValue));
+		CHECK_EQUAL(DefaultField.AsCustom(0, MakeMemoryView(DefaultValue)), MakeMemoryView(DefaultValue));
 	}
 }
 
@@ -1209,7 +1293,7 @@ TEST_CASE_METHOD(FAutomationTestFixture, "Core::Serialization::FCbField::Custom
 		TUniqueFunction<FMemoryView (FCbFieldView& Field, FMemoryView Default)> AsType;
 	};
 
-	// Test FCbFieldView(CustomByName, ABC, Empty)
+	SECTION("Test FCbFieldView(CustomByName, ABC, Empty)")
 	{
 		const uint8 Value[] = {4, 3, 'A', 'B', 'C'};
 		TestField<ECbFieldType::CustomByName>(*this, TEXT("CustomByName, MinId, Empty"), Value, FCbCustomByName{"ABC"_U8SV});
@@ -1217,7 +1301,7 @@ TEST_CASE_METHOD(FAutomationTestFixture, "Core::Serialization::FCbField::Custom
 		TestFieldError<ECbFieldType::CustomByName>(*this, TEXT("CustomByName, MinId, Empty, InvalidCase"), Value, ECbFieldError::RangeError, MakeMemoryView<uint8>({1, 2, 3}), FCustomByNameAccessor("abc"_U8SV));
 	}
 
-	// Test FCbFieldView(CustomByName, ABC, Value)
+	SECTION("Test FCbFieldView(CustomByName, ABC, Value)")
 	{
 		const uint8 Value[] = {8, 3, 'A', 'B', 'C', 1, 2, 3, 4};
 		TestFieldNoClone<ECbFieldType::CustomByName>(*this, TEXT("CustomByName, MinId, Value"), Value, FCbCustomByName{"ABC"_U8SV, MakeMemoryView(Value).Right(4)});
@@ -1225,13 +1309,13 @@ TEST_CASE_METHOD(FAutomationTestFixture, "Core::Serialization::FCbField::Custom
 		TestFieldError<ECbFieldType::CustomByName>(*this, TEXT("CustomByName, MinId, Value, InvalidCase"), Value, ECbFieldError::RangeError, MakeMemoryView<uint8>({1, 2, 3}), FCustomByNameAccessor("abc"_U8SV));
 	}
 
-	// Test FCbFieldView(None) as CustomByName
+	SECTION("Test FCbFieldView(None) as CustomByName")
 	{
 		FCbFieldView DefaultField;
 		TestFieldError<ECbFieldType::CustomByName>(*this, TEXT("CustomByName, None"), DefaultField, ECbFieldError::TypeError, FCbCustomByName{"ABC"_U8SV, MakeMemoryView<uint8>({1, 2, 3})});
 		TestFieldError<ECbFieldType::CustomByName>(*this, TEXT("CustomByName, None, View"), DefaultField, ECbFieldError::TypeError, MakeMemoryView<uint8>({1, 2, 3}), FCustomByNameAccessor("ABC"_U8SV));
 		const uint8 DefaultValue[] = {1, 2, 3};
-		TEST_EQUAL(TEXT("FCbFieldView()::AsCustom(Name)"), DefaultField.AsCustom("ABC"_U8SV, MakeMemoryView(DefaultValue)), MakeMemoryView(DefaultValue));
+		CHECK_EQUAL(DefaultField.AsCustom("ABC"_U8SV, MakeMemoryView(DefaultValue)), MakeMemoryView(DefaultValue));
 	}
 }
 
@@ -1395,17 +1479,17 @@ TEST_CASE_METHOD(FAutomationTestFixture, "Core::Serialization::FCbField::Iterate
 		Fields = Writer.Save();
 	}
 
-	TEST_EQUAL(TEXT("FCbFieldView::IterateAttachments Validate"), ValidateCompactBinaryRange(Fields.GetOuterBuffer(), ECbValidateMode::All), ECbValidateError::None);
+	CHECK_EQUAL(ValidateCompactBinaryRange(Fields.GetOuterBuffer(), ECbValidateMode::All), ECbValidateError::None);
 
 	uint32 AttachmentIndex = 0;
 
 	Fields.IterateRangeAttachments([this, &AttachmentIndex, &MakeTestHash](FCbFieldView Field)
 		{
-			TEST_TRUE(FString::Printf(TEXT("FCbFieldView::IterateAttachments(%u)::IsAttachment"), AttachmentIndex), Field.IsAttachment());
-			TEST_EQUAL(FString::Printf(TEXT("FCbFieldView::IterateAttachments(%u)"), AttachmentIndex), Field.AsAttachment(), MakeTestHash(AttachmentIndex));
+			CHECK(Field.IsAttachment());
+			CHECK_EQUAL(Field.AsAttachment(), MakeTestHash(AttachmentIndex));
 			++AttachmentIndex;
 		});
-	TEST_EQUAL(TEXT("FCbFieldView::IterateAttachments"), AttachmentIndex, 28);
+	CHECK_EQUAL(AttachmentIndex, 28);
 }
 
 
@@ -1430,16 +1514,16 @@ TEST_CASE_METHOD(FAutomationTestFixture, "Core::Serialization::FCbField::Field B
 	static_assert(std::is_constructible<FCbField, const FCbFieldView&, FCbArray&&>::value, "Missing constructor for FCbField");
 	static_assert(std::is_constructible<FCbField, const FCbFieldView&, FCbObject&&>::value, "Missing constructor for FCbField");
 
-	// Test FCbField()
+	SECTION("Test FCbField()")
 	{
 		FCbField DefaultField;
-		TEST_FALSE(TEXT("FCbField().HasValue()"), DefaultField.HasValue());
-		TEST_FALSE(TEXT("FCbField().IsOwned()"), DefaultField.IsOwned());
+		CHECK_FALSE( DefaultField.HasValue());
+		CHECK_FALSE( DefaultField.IsOwned());
 		DefaultField.MakeOwned();
-		TEST_TRUE(TEXT("FCbField().MakeOwned().IsOwned()"), DefaultField.IsOwned());
+		CHECK(DefaultField.IsOwned());
 	}
 
-	// Test Field w/ Type from Shared Buffer
+	SECTION("Test Field w/ Type from Shared Buffer")
 	{
 		const uint8 Value[] = { uint8(ECbFieldType::Binary), 3, 4, 5, 6 }; // Size: 3, Data: 4/5/6
 		FSharedBuffer ViewBuffer = FSharedBuffer::MakeView(MakeMemoryView(Value));
@@ -1458,33 +1542,33 @@ TEST_CASE_METHOD(FAutomationTestFixture, "Core::Serialization::FCbField::Field B
 		//FCbField InvalidOuterBuffer(ImplicitConv<FCbFieldView>(Owned), ViewBuffer);
 		//FCbField InvalidOuterBufferMove(ImplicitConv<FCbFieldView>(Owned), FSharedBuffer(ViewBuffer));
 
-		TEST_EQUAL(TEXT("FCbField(ViewBuffer)"), View.AsBinaryView(), ViewBuffer.GetView().Right(3));
-		TEST_EQUAL(TEXT("FCbField(ViewBuffer&&)"), ViewMove.AsBinaryView(), View.AsBinaryView());
-		TEST_EQUAL(TEXT("FCbField(ViewOuterField)"), ViewOuterField.AsBinaryView(), View.AsBinaryView());
-		TEST_EQUAL(TEXT("FCbField(ViewOuterBuffer)"), ViewOuterBuffer.AsBinaryView(), View.AsBinaryView());
-		TEST_EQUAL(TEXT("FCbField(OwnedBuffer)"), Owned.AsBinaryView(), OwnedBuffer.GetView().Right(3));
-		TEST_EQUAL(TEXT("FCbField(OwnedBuffer&&)"), OwnedMove.AsBinaryView(), Owned.AsBinaryView());
-		TEST_EQUAL(TEXT("FCbField(OwnedOuterField)"), OwnedOuterField.AsBinaryView(), Owned.AsBinaryView());
-		TEST_EQUAL(TEXT("FCbField(OwnedOuterBuffer)"), OwnedOuterBuffer.AsBinaryView(), Owned.AsBinaryView());
-
-		TEST_FALSE(TEXT("FCbField(ViewBuffer).IsOwned()"), View.IsOwned());
-		TEST_FALSE(TEXT("FCbField(ViewBuffer&&).IsOwned()"), ViewMove.IsOwned());
-		TEST_FALSE(TEXT("FCbField(ViewOuterField).IsOwned()"), ViewOuterField.IsOwned());
-		TEST_FALSE(TEXT("FCbField(ViewOuterBuffer).IsOwned()"), ViewOuterBuffer.IsOwned());
-		TEST_TRUE(TEXT("FCbField(OwnedBuffer).IsOwned()"), Owned.IsOwned());
-		TEST_TRUE(TEXT("FCbField(OwnedBuffer&&).IsOwned()"), OwnedMove.IsOwned());
-		TEST_TRUE(TEXT("FCbField(OwnedOuterField).IsOwned()"), OwnedOuterField.IsOwned());
-		TEST_TRUE(TEXT("FCbField(OwnedOuterBuffer).IsOwned()"), OwnedOuterBuffer.IsOwned());
+		CHECK_EQUAL(View.AsBinaryView(), ViewBuffer.GetView().Right(3));
+		CHECK_EQUAL(ViewMove.AsBinaryView(), View.AsBinaryView());
+		CHECK_EQUAL(ViewOuterField.AsBinaryView(), View.AsBinaryView());
+		CHECK_EQUAL(ViewOuterBuffer.AsBinaryView(), View.AsBinaryView());
+		CHECK_EQUAL(Owned.AsBinaryView(), OwnedBuffer.GetView().Right(3));
+		CHECK_EQUAL(OwnedMove.AsBinaryView(), Owned.AsBinaryView());
+		CHECK_EQUAL(OwnedOuterField.AsBinaryView(), Owned.AsBinaryView());
+		CHECK_EQUAL(OwnedOuterBuffer.AsBinaryView(), Owned.AsBinaryView());
+
+		CHECK_FALSE(View.IsOwned());
+		CHECK_FALSE(ViewMove.IsOwned());
+		CHECK_FALSE(ViewOuterField.IsOwned());
+		CHECK_FALSE(ViewOuterBuffer.IsOwned());
+		CHECK(Owned.IsOwned());
+		CHECK(OwnedMove.IsOwned());
+		CHECK(OwnedOuterField.IsOwned());
+		CHECK(OwnedOuterBuffer.IsOwned());
 
 		View.MakeOwned();
 		Owned.MakeOwned();
 		TestNotEqual(TEXT("FCbField(View).MakeOwned()"), View.AsBinaryView(), ViewBuffer.GetView().Right(3));
-		TEST_TRUE(TEXT("FCbField(View).MakeOwned().IsOwned()"), View.IsOwned());
-		TEST_EQUAL(TEXT("FCbField(Owned).MakeOwned()"), Owned.AsBinaryView(), OwnedBuffer.GetView().Right(3));
-		TEST_TRUE(TEXT("FCbField(Owned).MakeOwned().IsOwned()"), Owned.IsOwned());
+		CHECK(View.IsOwned());
+		CHECK_EQUAL(Owned.AsBinaryView(), OwnedBuffer.GetView().Right(3));
+		CHECK(Owned.IsOwned());
 	}
 
-	// Test Field w/ Type
+	SECTION("Test Field w/ Type")
 	{
 		const uint8 Value[] = { uint8(ECbFieldType::Binary), 3, 4, 5, 6 }; // Size: 3, Data: 4/5/6
 		const FCbFieldView Field(Value);
@@ -1495,23 +1579,23 @@ TEST_CASE_METHOD(FAutomationTestFixture, "Core::Serialization::FCbField::Field B
 		FCbField FieldClone = FCbField::Clone(Field);
 		FCbField FieldViewClone = FCbField::Clone(FieldView);
 
-		TEST_EQUAL(TEXT("FCbField::MakeView(Void)"), VoidView.AsBinaryView(), MakeMemoryView(Value).Right(3));
+		CHECK_EQUAL(VoidView.AsBinaryView(), MakeMemoryView(Value).Right(3));
 		TestNotEqual(TEXT("FCbField::Clone(Void)"), VoidClone.AsBinaryView(), MakeMemoryView(Value).Right(3));
-		TEST_TRUE(TEXT("FCbField::Clone(Void)->EqualBytes"), VoidClone.AsBinaryView().EqualBytes(VoidView.AsBinaryView()));
-		TEST_EQUAL(TEXT("FCbField::MakeView(Field)"), FieldView.AsBinaryView(), MakeMemoryView(Value).Right(3));
+		CHECK(VoidClone.AsBinaryView().EqualBytes(VoidView.AsBinaryView()));
+		CHECK_EQUAL(FieldView.AsBinaryView(), MakeMemoryView(Value).Right(3));
 		TestNotEqual(TEXT("FCbField::Clone(Field)"), FieldClone.AsBinaryView(), MakeMemoryView(Value).Right(3));
-		TEST_TRUE(TEXT("FCbField::Clone(Field)->EqualBytes"), FieldClone.AsBinaryView().EqualBytes(VoidView.AsBinaryView()));
+		CHECK(FieldClone.AsBinaryView().EqualBytes(VoidView.AsBinaryView()));
 		TestNotEqual(TEXT("FCbField::Clone(FieldView)"), FieldViewClone.AsBinaryView(), FieldView.AsBinaryView());
-		TEST_TRUE(TEXT("FCbField::Clone(FieldView)->EqualBytes"), FieldViewClone.AsBinaryView().EqualBytes(VoidView.AsBinaryView()));
+		CHECK(FieldViewClone.AsBinaryView().EqualBytes(VoidView.AsBinaryView()));
 
-		TEST_FALSE(TEXT("FCbField::MakeView(Void).IsOwned()"), VoidView.IsOwned());
-		TEST_TRUE(TEXT("FCbField::Clone(Void).IsOwned()"), VoidClone.IsOwned());
-		TEST_FALSE(TEXT("FCbField::MakeView(Field).IsOwned()"), FieldView.IsOwned());
-		TEST_TRUE(TEXT("FCbField::Clone(Field).IsOwned()"), FieldClone.IsOwned());
-		TEST_TRUE(TEXT("FCbField::Clone(FieldView).IsOwned()"), FieldViewClone.IsOwned());
+		CHECK_FALSE(VoidView.IsOwned());
+		CHECK(VoidClone.IsOwned());
+		CHECK_FALSE(FieldView.IsOwned());
+		CHECK(FieldClone.IsOwned());
+		CHECK(FieldViewClone.IsOwned());
 	}
 
-	// Test Field w/o Type
+	SECTION("Test Field w/o Type")
 	{
 		const uint8 Value[] = { 3, 4, 5, 6 }; // Size: 3, Data: 4/5/6
 		const FCbFieldView Field(Value, ECbFieldType::Binary);
@@ -1520,17 +1604,17 @@ TEST_CASE_METHOD(FAutomationTestFixture, "Core::Serialization::FCbField::Field B
 		FCbField FieldClone = FCbField::Clone(Field);
 		FCbField FieldViewClone = FCbField::Clone(FieldView);
 
-		TEST_EQUAL(TEXT("FCbField::MakeView(Field, NoType)"), FieldView.AsBinaryView(), MakeMemoryView(Value).Right(3));
-		TEST_TRUE(TEXT("FCbField::Clone(Field, NoType)"), FieldClone.AsBinaryView().EqualBytes(FieldView.AsBinaryView()));
-		TEST_TRUE(TEXT("FCbField::Clone(FieldView, NoType)"), FieldViewClone.AsBinaryView().EqualBytes(FieldView.AsBinaryView()));
+		CHECK_EQUAL(FieldView.AsBinaryView(), MakeMemoryView(Value).Right(3));
+		CHECK(FieldClone.AsBinaryView().EqualBytes(FieldView.AsBinaryView()));
+		CHECK(FieldViewClone.AsBinaryView().EqualBytes(FieldView.AsBinaryView()));
 
-		TEST_FALSE(TEXT("FCbField::MakeView(Field, NoType).IsOwned()"), FieldView.IsOwned());
-		TEST_TRUE(TEXT("FCbField::Clone(Field, NoType).IsOwned()"), FieldClone.IsOwned());
-		TEST_TRUE(TEXT("FCbField::Clone(FieldView, NoType).IsOwned()"), FieldViewClone.IsOwned());
+		CHECK_FALSE(FieldView.IsOwned());
+		CHECK(FieldClone.IsOwned());
+		CHECK(FieldViewClone.IsOwned());
 
 		FieldView.MakeOwned();
-		TEST_TRUE(TEXT("FCbField::MakeView(NoType).MakeOwned()"), FieldView.AsBinaryView().EqualBytes(MakeMemoryView(Value).Right(3)));
-		TEST_TRUE(TEXT("FCbField::MakeView(NoType).MakeOwned().IsOwned()"), FieldView.IsOwned());
+		CHECK(FieldView.AsBinaryView().EqualBytes(MakeMemoryView(Value).Right(3)));
+		CHECK(FieldView.IsOwned());
 	}
 }
 
@@ -1552,12 +1636,12 @@ TEST_CASE_METHOD(FAutomationTestFixture, "Core::Serialization::FCbField::Array B
 	static_assert(std::is_constructible<FCbArray, const FCbArrayView&, FCbArray&&>::value, "Missing constructor for FCbArray");
 	static_assert(std::is_constructible<FCbArray, const FCbArrayView&, FCbObject&&>::value, "Missing constructor for FCbArray");
 
-	// Test FCbArray()
+	SECTION("Test FCbArray()")
 	{
 		FCbArray DefaultArray;
-		TEST_FALSE(TEXT("FCbArray().IsOwned()"), DefaultArray.IsOwned());
+		CHECK_FALSE(DefaultArray.IsOwned());
 		DefaultArray.MakeOwned();
-		TEST_TRUE(TEXT("FCbArray().MakeOwned().IsOwned()"), DefaultArray.IsOwned());
+		CHECK(DefaultArray.IsOwned());
 	}
 }
 
@@ -1579,12 +1663,12 @@ TEST_CASE_METHOD(FAutomationTestFixture, "Core::Serialization::FCbField::Object
 	static_assert(std::is_constructible<FCbObject, const FCbObjectView&, FCbArray&&>::value, "Missing constructor for FCbObject");
 	static_assert(std::is_constructible<FCbObject, const FCbObjectView&, FCbObject&&>::value, "Missing constructor for FCbObject");
 
-	// Test FCbObject()
+	SECTION("Test FCbObject()")
 	{
 		FCbObject DefaultObject;
-		TEST_FALSE(TEXT("FCbObject().IsOwned()"), DefaultObject.IsOwned());
+		CHECK_FALSE(DefaultObject.IsOwned());
 		DefaultObject.MakeOwned();
-		TEST_TRUE(TEXT("FCbObject().MakeOwned().IsOwned()"), DefaultObject.IsOwned());
+		CHECK(DefaultObject.IsOwned());
 	}
 }
 
@@ -1606,13 +1690,13 @@ TEST_CASE_METHOD(FAutomationTestFixture, "Core::Serialization::FCbField::Buffer
 		return Count;
 	};
 
-	// Test FCbField[View]Iterator()
+	SECTION("Test FCbField[View]Iterator()")
 	{
-		TEST_EQUAL(TEXT("FCbFieldViewIterator()"), GetCount(FCbFieldViewIterator()), 0);
-		TEST_EQUAL(TEXT("FCbFieldIterator()"), GetCount(FCbFieldIterator()), 0);
+		CHECK_EQUAL(GetCount(FCbFieldViewIterator()), 0);
+		CHECK_EQUAL(GetCount(FCbFieldIterator()), 0);
 	}
 
-	// Test FCbField[View]Iterator(Range)
+	SECTION("Test FCbField[View]Iterator(Range)")
 	{
 		constexpr uint8 T = uint8(ECbFieldType::IntegerPositive);
 		const uint8 Value[] = { T, 0, T, 1, T, 2, T, 3 };
@@ -1626,92 +1710,103 @@ TEST_CASE_METHOD(FAutomationTestFixture, "Core::Serialization::FCbField::Buffer
 		const FCbFieldViewIterator FieldViewIt = FCbFieldViewIterator::MakeRange(View);
 		const FCbFieldIterator FieldIt = FCbFieldIterator::MakeRange(View);
 
-		TEST_EQUAL(TEXT("FCbFieldViewIterator::GetRangeHash()"), FieldViewIt.GetRangeHash(), FIoHash::HashBuffer(View));
-		TEST_EQUAL(TEXT("FCbFieldIterator::GetRangeHash()"), FieldIt.GetRangeHash(), FIoHash::HashBuffer(View));
+		CHECK_EQUAL(FieldViewIt.GetRangeHash(), FIoHash::HashBuffer(View));
+		CHECK_EQUAL(FieldIt.GetRangeHash(), FIoHash::HashBuffer(View));
 
 		FMemoryView RangeView;
-		TEST_TRUE(TEXT("FCbFieldViewIterator::TryGetRangeView()"), FieldViewIt.TryGetRangeView(RangeView) && RangeView == MakeMemoryView(Value));
-		TEST_TRUE(TEXT("FCbFieldIterator::TryGetRangeView()"), FieldIt.TryGetRangeView(RangeView) && RangeView == MakeMemoryView(Value));
+		CHECK(FieldViewIt.TryGetRangeView(RangeView));
+		CHECK(RangeView == MakeMemoryView(Value));
+		CHECK(FieldIt.TryGetRangeView(RangeView));
+		CHECK(RangeView == MakeMemoryView(Value));
 
-		TEST_EQUAL(TEXT("FCbFieldIterator::CloneRange(EmptyViewIt)"), GetCount(FCbFieldIterator::CloneRange(FCbFieldViewIterator())), 0);
-		TEST_EQUAL(TEXT("FCbFieldIterator::CloneRange(EmptyIt)"), GetCount(FCbFieldIterator::CloneRange(FCbFieldIterator())), 0);
+		CHECK_EQUAL(GetCount(FCbFieldIterator::CloneRange(FCbFieldViewIterator())), 0);
+		CHECK_EQUAL(GetCount(FCbFieldIterator::CloneRange(FCbFieldIterator())), 0);
 		const FCbFieldIterator FieldViewItClone = FCbFieldIterator::CloneRange(FieldViewIt);
 		const FCbFieldIterator FieldItClone = FCbFieldIterator::CloneRange(FieldIt);
-		TEST_EQUAL(TEXT("FCbFieldIterator::CloneRange(FieldViewIt)"), GetCount(FieldViewItClone), 4);
-		TEST_EQUAL(TEXT("FCbFieldIterator::CloneRange(FieldIt)"), GetCount(FieldItClone), 4);
+		CHECK_EQUAL(GetCount(FieldViewItClone), 4);
+		CHECK_EQUAL(GetCount(FieldItClone), 4);
 		TestNotEqual(TEXT("FCbFieldIterator::CloneRange(FieldViewIt).Equals()"), FieldViewItClone, FieldIt);
 		TestNotEqual(TEXT("FCbFieldIterator::CloneRange(FieldIt).Equals()"), FieldItClone, FieldIt);
 
-		TEST_EQUAL(TEXT("FCbFieldViewIterator::MakeRange(EmptyView)"), GetCount(FCbFieldViewIterator::MakeRange(EmptyView)), 0);
-		TEST_EQUAL(TEXT("FCbFieldIterator::MakeRange(BufferNullL)"), GetCount(FCbFieldIterator::MakeRange(NullBuffer)), 0);
-		TEST_EQUAL(TEXT("FCbFieldIterator::MakeRange(BufferNullR)"), GetCount(FCbFieldIterator::MakeRange(FSharedBuffer(NullBuffer))), 0);
+		CHECK_EQUAL(GetCount(FCbFieldViewIterator::MakeRange(EmptyView)), 0);
+		CHECK_EQUAL(GetCount(FCbFieldIterator::MakeRange(NullBuffer)), 0);
+		CHECK_EQUAL(GetCount(FCbFieldIterator::MakeRange(FSharedBuffer(NullBuffer))), 0);
 
-		TEST_EQUAL(TEXT("FCbFieldViewIterator::MakeRange(BufferView)"), GetCount(FCbFieldViewIterator::MakeRange(MakeMemoryView(Value))), 4);
-		TEST_EQUAL(TEXT("FCbFieldIterator::MakeRange(BufferCloneL)"), GetCount(FCbFieldIterator::MakeRange(Clone)), 4);
-		TEST_EQUAL(TEXT("FCbFieldIterator::MakeRange(BufferCloneR)"), GetCount(FCbFieldIterator::MakeRange(FSharedBuffer(Clone))), 4);
+		CHECK_EQUAL(GetCount(FCbFieldViewIterator::MakeRange(MakeMemoryView(Value))), 4);
+		CHECK_EQUAL(GetCount(FCbFieldIterator::MakeRange(Clone)), 4);
+		CHECK_EQUAL(GetCount(FCbFieldIterator::MakeRange(FSharedBuffer(Clone))), 4);
 
-		TEST_EQUAL(TEXT("FCbFieldIterator::MakeRangeView(FieldIt, BufferNullL)"), GetCount(FCbFieldIterator::MakeRangeView(FCbFieldViewIterator::MakeRange(View), NullBuffer)), 4);
-		TEST_EQUAL(TEXT("FCbFieldIterator::MakeRangeView(FieldIt, BufferNullR)"), GetCount(FCbFieldIterator::MakeRangeView(FCbFieldViewIterator::MakeRange(View), FSharedBuffer(NullBuffer))), 4);
-		TEST_EQUAL(TEXT("FCbFieldIterator::MakeRangeView(FieldIt, BufferViewL)"), GetCount(FCbFieldIterator::MakeRangeView(FCbFieldViewIterator::MakeRange(View), View)), 4);
-		TEST_EQUAL(TEXT("FCbFieldIterator::MakeRangeView(FieldIt, BufferViewR)"), GetCount(FCbFieldIterator::MakeRangeView(FCbFieldViewIterator::MakeRange(View), FSharedBuffer(View))), 4);
-		TEST_EQUAL(TEXT("FCbFieldIterator::MakeRangeView(FieldIt, BufferCloneL)"), GetCount(FCbFieldIterator::MakeRangeView(FCbFieldViewIterator::MakeRange(Clone), Clone)), 4);
-		TEST_EQUAL(TEXT("FCbFieldIterator::MakeRangeView(FieldIt, BufferCloneR)"), GetCount(FCbFieldIterator::MakeRangeView(FCbFieldViewIterator::MakeRange(Clone), FSharedBuffer(Clone))), 4);
+		CHECK_EQUAL(GetCount(FCbFieldIterator::MakeRangeView(FCbFieldViewIterator::MakeRange(View), NullBuffer)), 4);
+		CHECK_EQUAL(GetCount(FCbFieldIterator::MakeRangeView(FCbFieldViewIterator::MakeRange(View), FSharedBuffer(NullBuffer))), 4);
+		CHECK_EQUAL(GetCount(FCbFieldIterator::MakeRangeView(FCbFieldViewIterator::MakeRange(View), View)), 4);
+		CHECK_EQUAL(GetCount(FCbFieldIterator::MakeRangeView(FCbFieldViewIterator::MakeRange(View), FSharedBuffer(View))), 4);
+		CHECK_EQUAL(GetCount(FCbFieldIterator::MakeRangeView(FCbFieldViewIterator::MakeRange(Clone), Clone)), 4);
+		CHECK_EQUAL(GetCount(FCbFieldIterator::MakeRangeView(FCbFieldViewIterator::MakeRange(Clone), FSharedBuffer(Clone))), 4);
 
-		TEST_EQUAL(TEXT("FCbFieldViewIterator(FieldItL)"), GetCount(FCbFieldViewIterator(FieldIt)), 4);
-		TEST_EQUAL(TEXT("FCbFieldViewIterator(FieldItR)"), GetCount(FCbFieldViewIterator(FCbFieldIterator(FieldIt))), 4);
+		CHECK_EQUAL(GetCount(FCbFieldViewIterator(FieldIt)), 4);
+		CHECK_EQUAL(GetCount(FCbFieldViewIterator(FCbFieldIterator(FieldIt))), 4);
 
-		// Uniform
 		const uint8 UniformValue[] = { 0, 1, 2, 3 };
 		const FCbFieldViewIterator UniformFieldViewIt = FCbFieldViewIterator::MakeRange(MakeMemoryView(UniformValue), ECbFieldType::IntegerPositive);
-		TEST_EQUAL(TEXT("FCbFieldViewIterator::MakeRange(Uniform).GetRangeHash()"), UniformFieldViewIt.GetRangeHash(), FieldViewIt.GetRangeHash());
-		TEST_FALSE(TEXT("FCbFieldViewIterator::MakeRange(Uniform).TryGetRangeView()"), UniformFieldViewIt.TryGetRangeView(RangeView));
 		const FSharedBuffer UniformView = FSharedBuffer::MakeView(MakeMemoryView(UniformValue));
 		const FCbFieldIterator UniformFieldIt = FCbFieldIterator::MakeRange(UniformView, ECbFieldType::IntegerPositive);
-		TEST_EQUAL(TEXT("FCbFieldIterator::MakeRange(Uniform).GetRangeHash()"), UniformFieldIt.GetRangeHash(), FieldViewIt.GetRangeHash());
-		TEST_FALSE(TEXT("FCbFieldIterator::MakeRange(Uniform).TryGetRangeView()"), UniformFieldIt.TryGetRangeView(RangeView));
-
-		// Equals
-		TEST_TRUE(TEXT("FCbFieldViewIterator::Equals(Self)"), FieldViewIt.Equals(AsConst(FieldViewIt)));
-		TEST_TRUE(TEXT("FCbFieldViewIterator::Equals(OtherType)"), FieldViewIt.Equals(AsConst(FieldIt)));
-		TEST_TRUE(TEXT("FCbFieldIterator::Equals(Self)"), FieldIt.Equals(AsConst(FieldIt)));
-		TEST_TRUE(TEXT("FCbFieldIterator::Equals(OtherType)"), FieldIt.Equals(AsConst(FieldViewIt)));
-		TEST_FALSE(TEXT("FCbFieldViewIterator::Equals(OtherRange)"), FieldViewIt.Equals(FieldViewItClone));
-		TEST_FALSE(TEXT("FCbFieldIterator::Equals(OtherRange)"), FieldIt.Equals(FieldItClone));
-		TEST_TRUE(TEXT("FCbFieldViewIterator::Equals(Uniform, Self)"), UniformFieldViewIt.Equals(AsConst(UniformFieldViewIt)));
-		TEST_TRUE(TEXT("FCbFieldViewIterator::Equals(Uniform, OtherType)"), UniformFieldViewIt.Equals(UniformFieldIt));
-		TEST_TRUE(TEXT("FCbFieldIterator::Equals(Uniform, Self)"), UniformFieldIt.Equals(AsConst(UniformFieldIt)));
-		TEST_TRUE(TEXT("FCbFieldIterator::Equals(Uniform, OtherType)"), UniformFieldIt.Equals(UniformFieldViewIt));
-		TEST_FALSE(TEXT("FCbFieldViewIterator::Equals(SameValue, DifferentEnd)"),
-			FCbFieldViewIterator::MakeRange(MakeMemoryView(UniformValue), ECbFieldType::IntegerPositive)
+
+		{
+			INFO("Uniform");
+			CHECK_EQUAL(UniformFieldViewIt.GetRangeHash(), FieldViewIt.GetRangeHash());
+			CHECK_FALSE(UniformFieldViewIt.TryGetRangeView(RangeView));
+			CHECK_EQUAL(UniformFieldIt.GetRangeHash(), FieldViewIt.GetRangeHash());
+			CHECK_FALSE(UniformFieldIt.TryGetRangeView(RangeView));
+		}
+		
+		{
+			INFO("Equals");
+			CHECK(FieldViewIt.Equals(AsConst(FieldViewIt)));
+			CHECK(FieldViewIt.Equals(AsConst(FieldIt)));
+			CHECK(FieldIt.Equals(AsConst(FieldIt)));
+			CHECK(FieldIt.Equals(AsConst(FieldViewIt)));
+			CHECK_FALSE(FieldViewIt.Equals(FieldViewItClone));
+			CHECK_FALSE(FieldIt.Equals(FieldItClone));
+			CHECK(UniformFieldViewIt.Equals(AsConst(UniformFieldViewIt)));
+			CHECK(UniformFieldViewIt.Equals(UniformFieldIt));
+			CHECK(UniformFieldIt.Equals(AsConst(UniformFieldIt)));
+			CHECK(UniformFieldIt.Equals(UniformFieldViewIt));
+			CHECK_FALSE(FCbFieldViewIterator::MakeRange(MakeMemoryView(UniformValue), ECbFieldType::IntegerPositive)
 				.Equals(FCbFieldViewIterator::MakeRange(MakeMemoryView(UniformValue).LeftChop(1), ECbFieldType::IntegerPositive)));
-		TEST_FALSE(TEXT("FCbFieldViewIterator::Equals(DifferentValue, SameEnd)"),
-			FCbFieldViewIterator::MakeRange(MakeMemoryView(UniformValue), ECbFieldType::IntegerPositive)
+			CHECK_FALSE(FCbFieldViewIterator::MakeRange(MakeMemoryView(UniformValue), ECbFieldType::IntegerPositive)
 				.Equals(FCbFieldViewIterator::MakeRange(MakeMemoryView(UniformValue).RightChop(1), ECbFieldType::IntegerPositive)));
-
-		// CopyRangeTo
-		uint8 CopyBytes[sizeof(Value)];
-		FieldViewIt.CopyRangeTo(MakeMemoryView(CopyBytes));
-		TEST_TRUE(TEXT("FCbFieldViewIterator::MakeRange().CopyRangeTo()"), MakeMemoryView(CopyBytes).EqualBytes(MakeMemoryView(Value)));
-		FieldIt.CopyRangeTo(MakeMemoryView(CopyBytes));
-		TEST_TRUE(TEXT("FCbFieldIterator::MakeRange().CopyRangeTo()"), MakeMemoryView(CopyBytes).EqualBytes(MakeMemoryView(Value)));
-		UniformFieldViewIt.CopyRangeTo(MakeMemoryView(CopyBytes));
-		TEST_TRUE(TEXT("FCbFieldViewIterator::MakeRange(Uniform).CopyRangeTo()"), MakeMemoryView(CopyBytes).EqualBytes(MakeMemoryView(Value)));
-
-		// MakeRangeOwned
-		FCbFieldIterator OwnedFromView = UniformFieldIt;
-		OwnedFromView.MakeRangeOwned();
-		TEST_TRUE(TEXT("FCbFieldIterator::MakeRangeOwned(View)"), OwnedFromView.TryGetRangeView(RangeView) && RangeView.EqualBytes(MakeMemoryView(Value)));
-		FCbFieldIterator OwnedFromOwned = OwnedFromView;
-		OwnedFromOwned.MakeRangeOwned();
-		TEST_EQUAL(TEXT("FCbFieldIterator::MakeRangeOwned(Owned)"), OwnedFromOwned, OwnedFromView);
+		}
+		
+		{
+			INFO("CopyRangeTo");
+			uint8 CopyBytes[sizeof(Value)];
+			FieldViewIt.CopyRangeTo(MakeMemoryView(CopyBytes));
+			CHECK(MakeMemoryView(CopyBytes).EqualBytes(MakeMemoryView(Value)));
+			FieldIt.CopyRangeTo(MakeMemoryView(CopyBytes));
+			CHECK(MakeMemoryView(CopyBytes).EqualBytes(MakeMemoryView(Value)));
+			UniformFieldViewIt.CopyRangeTo(MakeMemoryView(CopyBytes));
+			CHECK(MakeMemoryView(CopyBytes).EqualBytes(MakeMemoryView(Value)));
+		}
+		
+		{
+			INFO("MakeRangeOwned");
+			FCbFieldIterator OwnedFromView = UniformFieldIt;
+			OwnedFromView.MakeRangeOwned();
+			CHECK(OwnedFromView.TryGetRangeView(RangeView));
+			CHECK(RangeView.EqualBytes(MakeMemoryView(Value)));
+			FCbFieldIterator OwnedFromOwned = OwnedFromView;
+			OwnedFromOwned.MakeRangeOwned();
+			CHECK_EQUAL(OwnedFromOwned, OwnedFromView);
+		}
+		
 
 		// These lines are expected to assert when uncommented.
 		//const FSharedBuffer ShortView = FSharedBuffer::MakeView(MakeMemoryView(Value).LeftChop(2));
-		//TEST_EQUAL(TEXT("FCbFieldIterator::MakeRangeView(FieldIt, InvalidBufferL)"), GetCount(FCbFieldIterator::MakeRangeView(FCbFieldViewIterator::MakeRange(*View), ShortView)), 4);
-		//TEST_EQUAL(TEXT("FCbFieldIterator::MakeRangeView(FieldIt, InvalidBufferR)"), GetCount(FCbFieldIterator::MakeRangeView(FCbFieldViewIterator::MakeRange(*View), FSharedBuffer(ShortView))), 4);
+		//CHECK_EQUAL(GetCount(FCbFieldIterator::MakeRangeView(FCbFieldViewIterator::MakeRange(*View), ShortView)), 4);
+		//CHECK_EQUAL(GetCount(FCbFieldIterator::MakeRangeView(FCbFieldViewIterator::MakeRange(*View), FSharedBuffer(ShortView))), 4);
 	}
 
-	// Test FCbField[View]Iterator(Scalar)
+	SECTION("Test FCbField[View]Iterator(Scalar)")
 	{
 		constexpr uint8 T = uint8(ECbFieldType::IntegerPositive);
 		const uint8 Value[] = { T, 0 };
@@ -1722,10 +1817,10 @@ TEST_CASE_METHOD(FAutomationTestFixture, "Core::Serialization::FCbField::Buffer
 		const FCbFieldView FieldView(Value);
 		const FCbField Field(View);
 
-		TEST_EQUAL(TEXT("FCbFieldViewIterator::MakeSingle(FieldViewL)"), GetCount(FCbFieldViewIterator::MakeSingle(FieldView)), 1);
-		TEST_EQUAL(TEXT("FCbFieldViewIterator::MakeSingle(FieldViewR)"), GetCount(FCbFieldViewIterator::MakeSingle(FCbFieldView(FieldView))), 1);
-		TEST_EQUAL(TEXT("FCbFieldIterator::MakeSingle(FieldL)"), GetCount(FCbFieldIterator::MakeSingle(Field)), 1);
-		TEST_EQUAL(TEXT("FCbFieldIterator::MakeSingle(FieldR)"), GetCount(FCbFieldIterator::MakeSingle(FCbField(Field))), 1);
+		CHECK_EQUAL(GetCount(FCbFieldViewIterator::MakeSingle(FieldView)), 1);
+		CHECK_EQUAL(GetCount(FCbFieldViewIterator::MakeSingle(FCbFieldView(FieldView))), 1);
+		CHECK_EQUAL(GetCount(FCbFieldIterator::MakeSingle(Field)), 1);
+		CHECK_EQUAL(GetCount(FCbFieldIterator::MakeSingle(FCbField(Field))), 1);
 	}
 }
 
@@ -1775,13 +1870,13 @@ TEST_CASE_METHOD(FAutomationTestFixture, "Core::Serialization::FCbField::Parse T
 	};
 
 	constexpr uint8 T = uint8(ECbFieldType::IntegerPositive | ECbFieldType::HasFieldName);
-	TEST_TRUE(TEXT("FCbObjectView Parse(None)"), TestParseObject({0}, 0, 0, 0, 0));
-	TEST_TRUE(TEXT("FCbObjectView Parse(ABCD)"), TestParseObject({16, T, 1, 'A', 1, T, 1, 'B', 2, T, 1, 'C', 3, T, 1, 'D', 4}, 1, 2, 3, 4));
-	TEST_TRUE(TEXT("FCbObjectView Parse(BCDA)"), TestParseObject({16, T, 1, 'B', 2, T, 1, 'C', 3, T, 1, 'D', 4, T, 1, 'A', 1}, 1, 2, 3, 4));
-	TEST_TRUE(TEXT("FCbObjectView Parse(BCD)"), TestParseObject({12, T, 1, 'B', 2, T, 1, 'C', 3, T, 1, 'D', 4}, 0, 2, 3, 4));
-	TEST_TRUE(TEXT("FCbObjectView Parse(BC)"), TestParseObject({8, T, 1, 'B', 2, T, 1, 'C', 3}, 0, 2, 3, 0));
-	TEST_TRUE(TEXT("FCbObjectView Parse(ABCDE)"), TestParseObject({20, T, 1, 'A', 1, T, 1, 'B', 2, T, 1, 'C', 3, T, 1, 'D', 4, T, 1, 'E', 5}, 1, 2, 3, 4));
-	TEST_TRUE(TEXT("FCbObjectView Parse(EABCD)"), TestParseObject({20, T, 1, 'E', 5, T, 1, 'A', 1, T, 1, 'B', 2, T, 1, 'C', 3, T, 1, 'D', 4}, 1, 2, 3, 4));
-	TEST_TRUE(TEXT("FCbObjectView Parse(DCBA)"), TestParseObject({16, T, 1, 'D', 4, T, 1, 'C', 3, T, 1, 'B', 2, T, 1, 'A', 1}, 1, 2, 3, 4));
+	CHECK(TestParseObject({0}, 0, 0, 0, 0));
+	CHECK(TestParseObject({16, T, 1, 'A', 1, T, 1, 'B', 2, T, 1, 'C', 3, T, 1, 'D', 4}, 1, 2, 3, 4));
+	CHECK(TestParseObject({16, T, 1, 'B', 2, T, 1, 'C', 3, T, 1, 'D', 4, T, 1, 'A', 1}, 1, 2, 3, 4));
+	CHECK(TestParseObject({12, T, 1, 'B', 2, T, 1, 'C', 3, T, 1, 'D', 4}, 0, 2, 3, 4));
+	CHECK(TestParseObject({8, T, 1, 'B', 2, T, 1, 'C', 3}, 0, 2, 3, 0));
+	CHECK(TestParseObject({20, T, 1, 'A', 1, T, 1, 'B', 2, T, 1, 'C', 3, T, 1, 'D', 4, T, 1, 'E', 5}, 1, 2, 3, 4));
+	CHECK(TestParseObject({20, T, 1, 'E', 5, T, 1, 'A', 1, T, 1, 'B', 2, T, 1, 'C', 3, T, 1, 'D', 4}, 1, 2, 3, 4));
+	CHECK(TestParseObject({16, T, 1, 'D', 4, T, 1, 'C', 3, T, 1, 'B', 2, T, 1, 'A', 1}, 1, 2, 3, 4));
 
 }
diff --git a/Engine/Source/Runtime/Core/Tests/String/EscapeTest.cpp b/Engine/Source/Runtime/Core/Tests/String/EscapeTest.cpp
--- a/Engine/Source/Runtime/Core/Tests/String/EscapeTest.cpp
+++ b/Engine/Source/Runtime/Core/Tests/String/EscapeTest.cpp
@@ -9,7 +9,7 @@
 namespace UE::String
 {
 
-TEST_CASE("Core::String::Escape", "[Core][String][Smoke]")
+TEST_CASE("Core::String::Misc::Escape", "[Core][String][Smoke]")
 {
 	SECTION("Escape")
 	{
diff --git a/Engine/Source/Runtime/Core/Tests/String/FindTest.cpp b/Engine/Source/Runtime/Core/Tests/String/FindTest.cpp
--- a/Engine/Source/Runtime/Core/Tests/String/FindTest.cpp
+++ b/Engine/Source/Runtime/Core/Tests/String/FindTest.cpp
@@ -8,7 +8,7 @@
 namespace UE::String
 {
 
-TEST_CASE("Core::String::Find", "[Core][String][Smoke]")
+TEST_CASE("Core::String::Misc::Find", "[Core][String][Smoke]")
 {
 	SECTION("FindFirst")
 	{
diff --git a/Engine/Source/Runtime/Core/Tests/String/JoinTest.cpp b/Engine/Source/Runtime/Core/Tests/String/JoinTest.cpp
--- a/Engine/Source/Runtime/Core/Tests/String/JoinTest.cpp
+++ b/Engine/Source/Runtime/Core/Tests/String/JoinTest.cpp
@@ -10,7 +10,7 @@
 namespace UE::String
 {
 
-TEST_CASE("Core::String::Join", "[Core][String][Smoke]")
+TEST_CASE("Core::String::Misc::Join", "[Core][String][Smoke]")
 {
 	SECTION("Join")
 	{
diff --git a/Engine/Source/Runtime/Core/Tests/String/ParseLinesTest.cpp b/Engine/Source/Runtime/Core/Tests/String/ParseLinesTest.cpp
--- a/Engine/Source/Runtime/Core/Tests/String/ParseLinesTest.cpp
+++ b/Engine/Source/Runtime/Core/Tests/String/ParseLinesTest.cpp
@@ -12,7 +12,7 @@
 namespace UE::String
 {
 
-TEST_CASE("Core::String::ParseLines", "[Core][String][Smoke]")
+TEST_CASE("Core::String::Misc::ParseLines", "[Core][String][Smoke]")
 {
 	using FStringViewArray = TArray<FStringView, TInlineAllocator<8>>;
 
diff --git a/Engine/Source/Runtime/Core/Tests/String/ParseTokensTest.cpp b/Engine/Source/Runtime/Core/Tests/String/ParseTokensTest.cpp
--- a/Engine/Source/Runtime/Core/Tests/String/ParseTokensTest.cpp
+++ b/Engine/Source/Runtime/Core/Tests/String/ParseTokensTest.cpp
@@ -11,7 +11,7 @@
 namespace UE::String
 {
 
-TEST_CASE("Core::String::ParseTokens", "[Core][String][Smoke]")
+TEST_CASE("Core::String::Misc::ParseTokens", "[Core][String][Smoke]")
 {
 	using FCharArray = TArray<TCHAR, TInlineAllocator<8>>;
 	using FStringViewArray = TArray<FStringView, TInlineAllocator<8>>;
diff --git a/Engine/Source/Runtime/Core/Tests/Tasks/TasksTest.cpp b/Engine/Source/Runtime/Core/Tests/Tasks/TasksTest.cpp
--- a/Engine/Source/Runtime/Core/Tests/Tasks/TasksTest.cpp
+++ b/Engine/Source/Runtime/Core/Tests/Tasks/TasksTest.cpp
@@ -594,7 +594,7 @@ namespace UE { namespace TasksTests
 		TlsValue = Dummy;
 	}
 
-	TEST_CASE_METHOD(FCoreTestFixture, "Core::Tasks::Tls::Tls", "[Core][Async][Smoke]")
+	TEST_CASE_METHOD(FCoreTestFixture, "Core::Async::Task::Tls", "[Core][Async][Smoke]")
 	{
 		UE_BENCHMARK(5, UeTlsStressTest<10000000>);
 		UE_BENCHMARK(5, ThreadLocalStressTest<10000000>);
