Added optional ANDROID_USE_NICE_VALUE_THREADPRIORITY define for FRunnableThreadAndroid::SetThreadPriority.
default off.
Enabling it sets the thread's priority via nice values.
Enabling ANDROID_USE_NICE_VALUE_THREADPRIORITY also sets Android pool threads to use all cores.

diff --git a/Engine/Source/Runtime/Core/Private/Android/AndroidPlatformProcess.cpp b/Engine/Source/Runtime/Core/Private/Android/AndroidPlatformProcess.cpp
--- a/Engine/Source/Runtime/Core/Private/Android/AndroidPlatformProcess.cpp
+++ b/Engine/Source/Runtime/Core/Private/Android/AndroidPlatformProcess.cpp
@@ -5,6 +5,7 @@
 =============================================================================*/
 
 #include "Android/AndroidPlatformProcess.h"
+#include "Android/AndroidPlatform.h"
 #include "Android/AndroidPlatformRunnableThread.h"
 #include "Android/AndroidPlatformAffinity.h"
 #include "Async/TaskGraphInterfaces.h"
@@ -12,6 +13,7 @@
 #include <sys/syscall.h>
 #include <pthread.h>
 #include <dlfcn.h>
+#include <sys/resource.h>
 
 #include "Android/AndroidJavaEnv.h"
 
@@ -269,3 +271,50 @@ uint64 FAndroidAffinity::GetLittleCoreMask()
 	}
 	return Mask;
 }
+
+static int AndroidThreadPriorityToNice(EThreadPriority NewPriority)
+{
+	static const int LOWEST = 19;
+	static const int BACKGROUND = 10;
+	static const int NORMAL = 0;
+	static const int FOREGROUND = -2;
+	static const int URGENT = -8;
+
+	static const int ThreadPriorityToNiceValue[7] = {
+		NORMAL	,	//	TPri_Normal
+		NORMAL - 2,	//	TPri_AboveNormal
+		NORMAL + 3,	//	TPri_BelowNormal
+		NORMAL - 4,	//	TPri_Highest
+		BACKGROUND,	//	TPri_Lowest
+		NORMAL+2,	//	TPri_SlightlyBelowNormal
+		URGENT,		//	TPri_TimeCritical
+	};
+	static_assert(UE_ARRAY_COUNT(ThreadPriorityToNiceValue) == EThreadPriority::TPri_Num, TEXT("This array is expected to be 1:1 mapping with EThreadPriority."));
+	return ThreadPriorityToNiceValue[(int)NewPriority];
+}
+
+#if ANDROID_USE_NICE_VALUE_THREADPRIORITY
+void FRunnableThreadAndroid::SetThreadPriority(pthread_t InThread, EThreadPriority NewPriority)
+{
+	check((int)NewPriority >= 0 && NewPriority < EThreadPriority::TPri_Num);
+	NewPriority = (EThreadPriority)FMath::Clamp((int)NewPriority, 0, (int)(EThreadPriority::TPri_Num - 1));
+
+	// Read the current policy
+	int32 InitialPolicy;
+	int32 NewPolicy = (NewPriority == EThreadPriority::TPri_Lowest) ? SCHED_BATCH : SCHED_NORMAL;
+	struct sched_param Sched = { };
+	pthread_getschedparam(InThread, &InitialPolicy, &Sched);
+
+	if (InitialPolicy != NewPolicy && sched_setscheduler(InThread, NewPolicy, &Sched)!=0)
+	{
+		UE_LOG(LogHAL, Error, TEXT("Failed to set %s thread scheduler, tid %d from %d to %d (errno %d)"), *GetThreadName(), ThreadID, InitialPolicy, NewPolicy, errno);
+	}
+
+	int InitialNice = getpriority(PRIO_PROCESS, ThreadID);
+	int NewNice = AndroidThreadPriorityToNice(NewPriority);
+	if (InitialNice != NewNice && setpriority(PRIO_PROCESS, ThreadID, NewNice) != 0)
+	{
+		UE_LOG(LogHAL, Error, TEXT("Failed to set %s thread priority, tid %d from %d to %d (errno %d)"), *GetThreadName(), ThreadID, InitialNice, NewNice, errno);
+	}
+}
+#endif
diff --git a/Engine/Source/Runtime/Core/Public/Android/AndroidPlatform.h b/Engine/Source/Runtime/Core/Public/Android/AndroidPlatform.h
--- a/Engine/Source/Runtime/Core/Public/Android/AndroidPlatform.h
+++ b/Engine/Source/Runtime/Core/Public/Android/AndroidPlatform.h
@@ -101,6 +101,11 @@ typedef FAndroidTypes FPlatformTypes;
 	#define USE_ANDROID_EVENTS						1
 #endif
 
+// Enable to set thread nice values when setting runnable thread priorities
+#ifndef ANDROID_USE_NICE_VALUE_THREADPRIORITY
+	#define ANDROID_USE_NICE_VALUE_THREADPRIORITY 0
+#endif
+
 // Function type macros.
 #define VARARGS													/* Functions with variable arguments */
 #define CDECL													/* Standard C function */
diff --git a/Engine/Source/Runtime/Core/Public/Android/AndroidPlatformAffinity.h b/Engine/Source/Runtime/Core/Public/Android/AndroidPlatformAffinity.h
--- a/Engine/Source/Runtime/Core/Public/Android/AndroidPlatformAffinity.h
+++ b/Engine/Source/Runtime/Core/Public/Android/AndroidPlatformAffinity.h
@@ -7,6 +7,7 @@ AndroidAffinity.h: Android affinity profile masks definitions.
 #pragma once
 
 #include "GenericPlatform/GenericPlatformAffinity.h"
+#include "AndroidPlatform.h"
 
 class FAndroidAffinity : public FGenericPlatformAffinity
 {
@@ -36,7 +37,11 @@ public:
 
 	static const CORE_API uint64 GetPoolThreadMask()
 	{
+#if ANDROID_USE_NICE_VALUE_THREADPRIORITY
+		return AllCores;
+#else
 		return GetLittleCoreMask();
+#endif
 	}
 
 	static const CORE_API uint64 GetTaskGraphThreadMask()
diff --git a/Engine/Source/Runtime/Core/Public/Android/AndroidPlatformRunnableThread.h b/Engine/Source/Runtime/Core/Public/Android/AndroidPlatformRunnableThread.h
--- a/Engine/Source/Runtime/Core/Public/Android/AndroidPlatformRunnableThread.h
+++ b/Engine/Source/Runtime/Core/Public/Android/AndroidPlatformRunnableThread.h
@@ -8,6 +8,7 @@
 
 #include "Runtime/Core/Private/HAL/PThreadRunnableThread.h"
 #include "Android/AndroidPlatformMisc.h"
+#include "AndroidPlatform.h"
 
 /**
 * Android implementation of the pthread functions
@@ -81,4 +82,8 @@ private:
 
 		return InStackSize;
 	}
+
+#if ANDROID_USE_NICE_VALUE_THREADPRIORITY
+	virtual void SetThreadPriority(pthread_t InThread, EThreadPriority NewPriority) override;
+#endif
 };
