Remove unbalanced optimization macro

diff --git a/Engine/Source/Runtime/Core/Tests/Internationalization/DateTimeFormattingRulesTest.cpp b/Engine/Source/Runtime/Core/Tests/Internationalization/DateTimeFormattingRulesTest.cpp
--- a/Engine/Source/Runtime/Core/Tests/Internationalization/DateTimeFormattingRulesTest.cpp
+++ b/Engine/Source/Runtime/Core/Tests/Internationalization/DateTimeFormattingRulesTest.cpp
@@ -146,6 +146,3 @@ TEST_CASE_METHOD(FAutomationTestFixture, "Core::Internationalization::FDateTime:
 		TestStr(TEXT("Testing Custom Date-Time Pattern (FText) - %l:%M %p"), FText::AsDateTime(TestDateTime, TEXT("%l:%M %p"), FText::GetInvariantTimeZone(), InvariantCulture).ToString(), TEXT("9:30 AM"));
 	}
 }
-
-PRAGMA_ENABLE_OPTIMIZATION
-
