Expose a new PlatFormMisc function to query if the game is running on cloud

diff --git a/Engine/Source/Runtime/Core/Public/GenericPlatform/GenericPlatformMisc.h b/Engine/Source/Runtime/Core/Public/GenericPlatform/GenericPlatformMisc.h
--- a/Engine/Source/Runtime/Core/Public/GenericPlatform/GenericPlatformMisc.h
+++ b/Engine/Source/Runtime/Core/Public/GenericPlatform/GenericPlatformMisc.h
@@ -1702,6 +1702,14 @@ public:
 	 */
 	static int GetMobilePropagateAlphaSetting();
 
+	/**
+	*	Return if the game is running on cloud server
+	*/
+	static bool IsRunningInCloud()
+	{
+		return false;
+	}
+
 #if !UE_BUILD_SHIPPING
 	/**
 	 * Returns any platform specific warning messages we want printed on screen
