LowLevelTests: Reverted the Core tests that were moved in 19498448 and 19571202

#rnx

diff --git a/Engine/Source/Runtime/Core/Private/Containers/ContainersTest.cpp b/Engine/Source/Runtime/Core/Private/Containers/ContainersTest.cpp
new file mode 100644
--- /dev/null
+++ b/Engine/Source/Runtime/Core/Private/Containers/ContainersTest.cpp
@@ -0,0 +1,847 @@
+// Copyright Epic Games, Inc. All Rights Reserved.
+
+#include "CoreTypes.h"
+#include "Misc/AssertionMacros.h"
+#include "Containers/Array.h"
+#include "Containers/Map.h"
+#include "Containers/Set.h"
+#include "Containers/SortedMap.h"
+#include "Containers/ArrayView.h"
+#include "Misc/AutomationTest.h"
+#include "Stats/StatsMisc.h"
+#include "Math/RandomStream.h"
+
+#if WITH_DEV_AUTOMATION_TESTS
+
+#define MAX_TEST_OBJECTS      65
+#define MAX_TEST_OBJECTS_STEP 1
+#define RANDOM_SEED 12345
+
+namespace
+{
+	namespace EContainerTestType
+	{
+		enum Type
+		{
+			MovedFrom,
+			Key,
+			Value
+		};
+	}
+
+	struct FContainerTestStats
+	{
+		int32  NextId;
+		int32  ConstructedIDs[MAX_TEST_OBJECTS];
+		int32* End;
+
+		FContainerTestStats()
+		{
+			Reset();
+		}
+
+		void Reset()
+		{
+			NextId = 1;
+			End    = ConstructedIDs;
+		}
+
+		int32 Num() const
+		{
+			return UE_PTRDIFF_TO_INT32(End - ConstructedIDs);
+		}
+
+		int32 Add()
+		{
+			// Ensure we're not constructing too many objects
+			check(Num() < MAX_TEST_OBJECTS);
+
+			// Store ID in array
+			return *End++ = NextId++;
+		}
+
+		void Remove(int32 ObjId)
+		{
+			for (int32* It = ConstructedIDs; It != End; ++It)
+			{
+				if (*It != ObjId)
+					continue;
+
+				// Remove this from the list
+				--End;
+				for (; It != End; ++It)
+				{
+					*It = *(It + 1);
+				}
+				return;
+			}
+
+			// We didn't find an entry for this - an invalid destructor call?
+			check(false);
+		}
+	} ContainerTestStats;
+
+	struct FContainerTestType
+	{
+		FContainerTestType(const TCHAR* InStr, EContainerTestType::Type InType)
+			: Str (InStr)
+			, Type(InType)
+			, Id  (ContainerTestStats.Add())
+		{
+		}
+
+		FContainerTestType(const FContainerTestType& Other)
+			: Str (Other.Str)
+			, Type(Other.Type)
+			, Id  (ContainerTestStats.Add())
+		{
+		}
+
+		FContainerTestType(FContainerTestType&& Other)
+			: Str (Other.Str)
+			, Type(Other.Type)
+			, Id  (ContainerTestStats.Add())
+		{
+			Other.Str  = NULL;
+			Other.Type = EContainerTestType::MovedFrom;
+		}
+
+		FContainerTestType& operator=(FContainerTestType&& Other)
+		{
+			Str  = Other.Str;
+			Type = Other.Type;
+
+			Other.Str  = NULL;
+			Other.Type = EContainerTestType::MovedFrom;
+			
+			return *( static_cast<FContainerTestType*>( this ) );
+		}
+
+		FContainerTestType& operator=(const FContainerTestType& Other)
+		{
+			Str  = Other.Str;
+			Type = Other.Type;
+			
+			return *( static_cast<FContainerTestType*>( this ) );
+		}
+
+		~FContainerTestType()
+		{
+			ContainerTestStats.Remove(Id);
+		}
+
+		friend bool operator==(const FContainerTestType& Lhs, const FContainerTestType& Rhs)
+		{
+			return Lhs.Type == Rhs.Type &&
+			       Lhs.Id   == Rhs.Id   &&
+			       !FCString::Strcmp(Lhs.Str, Rhs.Str);
+		}
+
+		friend bool operator!=(const FContainerTestType& Lhs, const FContainerTestType& Rhs)
+		{
+			return !(Lhs == Rhs);
+		}
+
+		const TCHAR*             Str;
+		EContainerTestType::Type Type;
+		int32                    Id;
+	};
+
+	struct FContainerTestKeyType : FContainerTestType
+	{
+		FContainerTestKeyType()
+			: FContainerTestType(TEXT("<default key>"), EContainerTestType::Key)
+		{
+		}
+
+		explicit FContainerTestKeyType(const TCHAR* InStr)
+			: FContainerTestType(InStr, EContainerTestType::Key)
+		{
+		}
+
+		FContainerTestKeyType(FContainerTestKeyType&& Other)
+			: FContainerTestType(MoveTemp(Other))
+		{
+		}
+
+		FContainerTestKeyType(const FContainerTestKeyType& Other)
+			: FContainerTestType(Other)
+		{
+		}
+
+		FContainerTestKeyType& operator=(FContainerTestKeyType&& Other)
+		{
+			(FContainerTestType&)*this = MoveTemp(Other);
+			
+			return *( static_cast<FContainerTestKeyType*>( this ) );
+		}
+
+		FContainerTestKeyType& operator=(const FContainerTestKeyType& Other)
+		{
+			(FContainerTestType&)*this = Other;
+			
+			return *( static_cast<FContainerTestKeyType*>( this ) );
+		}
+	};
+
+	struct FContainerTestValueType : FContainerTestType
+	{
+		FContainerTestValueType()
+			: FContainerTestType(TEXT("<default value>"), EContainerTestType::Value)
+		{
+		}
+
+		explicit FContainerTestValueType(const TCHAR* InStr)
+			: FContainerTestType(InStr, EContainerTestType::Value)
+		{
+		}
+
+		FContainerTestValueType(FContainerTestValueType&& Other)
+			: FContainerTestType(MoveTemp(Other))
+		{
+		}
+
+		FContainerTestValueType(const FContainerTestValueType& Other)
+			: FContainerTestType(Other)
+		{
+		}
+
+		FContainerTestValueType& operator=(FContainerTestValueType&& Other)
+		{
+			(FContainerTestType&)*this = MoveTemp(Other);
+			
+			return *( static_cast<FContainerTestValueType*>( this ) );
+		}
+
+		FContainerTestValueType& operator=(const FContainerTestValueType& Other)
+		{
+			(FContainerTestType&)*this = Other;
+			
+			return *( static_cast<FContainerTestValueType*>( this ) );
+		}
+	};
+
+	template <typename Container>
+	void CheckContainerElements(Container& Cont)
+	{
+		auto It  = Cont.CreateIterator();
+		auto CIt = Cont.CreateConstIterator();
+		for (auto& E : Cont)
+		{
+			check(*It  == E);
+			check(*CIt == E);
+
+			++It;
+			++CIt;
+		}
+	}
+
+	template <typename Container>
+	void CheckContainerNum(Container& Cont)
+	{
+		int32 Count = 0;
+		for (auto It = Cont.CreateIterator(); It; ++It)
+		{
+			++Count;
+		}
+
+		int32 CCount = 0;
+		for (auto It = Cont.CreateConstIterator(); It; ++It)
+		{
+			++CCount;
+		}
+
+		int32 RCount = 0;
+		for (auto& It : Cont)
+		{
+			++RCount;
+		}
+
+		check(Count  == Cont.Num());
+		check(CCount == Cont.Num());
+		check(RCount == Cont.Num());
+	}
+
+	template <typename Container>
+	void CheckContainerEnds(Container& Cont)
+	{
+		auto Iter  = Cont.CreateIterator();
+		auto CIter = Cont.CreateConstIterator();
+
+		for (int32 Num = Cont.Num(); Num; --Num)
+		{
+			++Iter;
+			++CIter;
+		}
+
+		check(!Iter);
+		check(!CIter);
+	}
+
+	template <typename KeyType>
+	KeyType GenerateTestKey(int32 Input)
+	{
+		return (KeyType)Input;
+	}
+
+	template <>
+	FName GenerateTestKey<FName>(int32 Input)
+	{
+		// Don't use _foo as we want to test the slower compare path
+		return FName(*FString::Printf(TEXT("TestName%d"), Input));
+	}
+
+	template <>
+	FString GenerateTestKey<FString>(int32 Input)
+	{
+		return FString::Printf(TEXT("TestString%d"), Input);
+	}
+
+	template <typename ContainerType, typename KeyType>
+	void RunContainerTests()
+	{
+		ContainerType Cont;
+
+		ContainerTestStats.Reset();
+		// Subtract one to account for temporaries that will be created during an Add
+		for (int32 Count = 0; Count < MAX_TEST_OBJECTS - 1; Count += MAX_TEST_OBJECTS_STEP)
+		{
+			for (int32 N = 0; N != Count; ++N)
+			{
+				Cont.Add(GenerateTestKey<KeyType>(N), FContainerTestValueType(TEXT("New Value")));
+				CheckContainerNum(Cont);
+				CheckContainerEnds(Cont);
+				CheckContainerElements(Cont);
+			}
+
+			for (int32 N = 0; N != Count; ++N)
+			{
+				Cont.Remove(GenerateTestKey<KeyType>(N));
+				CheckContainerNum(Cont);
+				CheckContainerEnds(Cont);
+				CheckContainerElements(Cont);
+			}
+
+			for (int32 N = 0; N != Count; ++N)
+			{
+				Cont.Add(GenerateTestKey<KeyType>((Count - 1) - N), FContainerTestValueType(TEXT("New Value")));
+				CheckContainerNum(Cont);
+				CheckContainerEnds(Cont);
+				CheckContainerElements(Cont);
+			}
+
+			for (int32 N = 0; N != Count; ++N)
+			{
+				Cont.Remove(GenerateTestKey<KeyType>(N));
+				CheckContainerNum(Cont);
+				CheckContainerEnds(Cont);
+				CheckContainerElements(Cont);
+			}
+		}
+	}
+
+	template <typename ContainerType, typename KeyType>
+	void RunPerformanceTest(const FString& Description, int32 NumObjects, int32 NumOperations)
+	{
+		ContainerTestStats.Reset();
+		
+		ContainerType Cont;
+		FRandomStream RandomStream(RANDOM_SEED);
+
+		// Prep keys, not part of performance test
+		TArray<KeyType> KeyArray;
+		KeyArray.Reserve(NumObjects);
+
+		for (int32 i = 0; i < NumObjects; i++)
+		{
+			KeyArray.Add(GenerateTestKey<KeyType>(i));
+		}
+
+		for (int32 i = 0; i < NumObjects; i++)
+		{
+			int32 SwapIndex = RandomStream.RandRange(0, NumObjects - 1);
+			if (i != SwapIndex)
+			{
+				KeyArray.Swap(i, SwapIndex);
+			}
+		}
+
+		FScopeLogTime LogTimePtr(*FString::Printf(TEXT("%s objects=%d count=%d"), *Description, NumObjects, NumOperations), nullptr, FScopeLogTime::ScopeLog_Milliseconds);
+
+		// Add elements in stably randomized order
+		for (int32 i = 0; i < NumObjects; i++)
+		{
+			Cont.Add(KeyArray[i], FString(TEXT("New Value")));
+		}
+		
+		// Now do searches
+		for (int32 i = 0; i < NumOperations; i++)
+		{
+			KeyType& Key = KeyArray[RandomStream.RandRange(0, NumObjects - 1)];
+
+			FString* FoundValue = Cont.Find(Key);
+			check(FoundValue);
+		}
+	}
+
+	template <typename ContainerType, typename KeyType>
+	void RunSetPerformanceTest(const FString& Description, int32 NumObjects, int32 NumOperations)
+	{
+		ContainerTestStats.Reset();
+
+		ContainerType Cont;
+		FRandomStream RandomStream(RANDOM_SEED);
+
+		// Prep keys, not part of performance test
+		TArray<KeyType> KeyArray;
+		KeyArray.Reserve(NumObjects);
+
+		for (int32 i = 0; i < NumObjects; i++)
+		{
+			KeyArray.Add(GenerateTestKey<KeyType>(i));
+		}
+
+		for (int32 i = 0; i < NumObjects; i++)
+		{
+			int32 SwapIndex = RandomStream.RandRange(0, NumObjects - 1);
+			if (i != SwapIndex)
+			{
+				KeyArray.Swap(i, SwapIndex);
+			}
+		}
+
+		FScopeLogTime LogTimePtr(*FString::Printf(TEXT("%s objects=%d count=%d"), *Description, NumObjects, NumOperations), nullptr, FScopeLogTime::ScopeLog_Milliseconds);
+
+		// Add elements in stably randomized order
+		for (int32 i = 0; i < NumObjects; i++)
+		{
+			Cont.Add(KeyArray[i]);
+		}
+
+		// Now do searches
+		for (int32 i = 0; i < NumOperations; i++)
+		{
+			KeyType& Key = KeyArray[RandomStream.RandRange(0, NumObjects - 1)];
+
+			bool FoundValue = Cont.Contains(Key);
+			check(FoundValue);
+		}
+	}
+
+	template <typename ContainerType, typename ElemType, typename GetRefFuncType>
+	void TestGetRef(ContainerType& Container, ElemType&& Elem, GetRefFuncType&& GetRefFunc)
+	{
+		auto& Ref = GetRefFunc(Container, Forward<ElemType>(Elem));
+	}
+
+	template<typename TValueType>
+	struct FCaseSensitiveLookupKeyFuncs : BaseKeyFuncs<TValueType, FString>
+	{
+		static FORCEINLINE const FString& GetSetKey(const TPair<FString, TValueType>& Element)
+		{
+			return Element.Key;
+		}
+		static FORCEINLINE bool Matches(const FString& A, const FString& B)
+		{
+			return A.Equals(B, ESearchCase::CaseSensitive);
+		}
+		static FORCEINLINE uint32 GetKeyHash(const FString& Key)
+		{
+			return FCrc::StrCrc32<TCHAR>(*Key);
+		}
+	};
+
+	void RunGetRefTests()
+	{
+		{
+			TArray<FString> Arr;
+			FString* Str1 = &Arr.AddDefaulted_GetRef();
+			check(Str1->IsEmpty());
+			check(Str1 == &Arr.Last());
+
+			FString* Str2 = &Arr.AddZeroed_GetRef();
+			check(Str2->IsEmpty());
+			check(Str2 == &Arr.Last());
+
+			FString* Str3 = &Arr.Add_GetRef(TEXT("Abc"));
+			check(*Str3 == TEXT("Abc"));
+			check(Str3 == &Arr.Last());
+
+			FString* Str4 = &Arr.Emplace_GetRef(TEXT("Def"));
+			check(*Str4 == TEXT("Def"));
+			check(Str4 == &Arr.Last());
+
+			FString* Str5 = &Arr.EmplaceAt_GetRef(3, TEXT("Ghi"));
+			check(*Str5 == TEXT("Ghi"));
+			check(Str5 == &Arr[3]);
+
+			FString* Str6 = &Arr.InsertDefaulted_GetRef(2);
+			check(Str6->IsEmpty());
+			check(Str6 == &Arr[2]);
+
+			FString* Str7 = &Arr.InsertZeroed_GetRef(4);
+			check(Str7->IsEmpty());
+			check(Str7 == &Arr[4]);
+		}
+	}
+}
+
+IMPLEMENT_SIMPLE_AUTOMATION_TEST(FContainersSmokeTest, "System.Core.Containers.Smoke", EAutomationTestFlags::ApplicationContextMask | EAutomationTestFlags::SmokeFilter)
+bool FContainersSmokeTest::RunTest( const FString& Parameters )
+{
+	RunContainerTests<TMap<int32, FContainerTestValueType>, int32>();
+
+	return true;
+}
+
+IMPLEMENT_SIMPLE_AUTOMATION_TEST(FContainersFullTest, "System.Core.Containers.Full", EAutomationTestFlags::ApplicationContextMask | EAutomationTestFlags::EngineFilter)
+bool FContainersFullTest::RunTest(const FString& Parameters)
+{
+	RunContainerTests<TMap<int32, FContainerTestValueType>, int32>();
+	RunContainerTests<TMap<FName, FContainerTestValueType>, FName>();
+	RunContainerTests<TMap<FString, FContainerTestValueType>, FString>();
+	RunContainerTests<TMap<int32, FContainerTestValueType, TInlineSetAllocator<64>>, int32>();
+	RunContainerTests<TMap<FString, FContainerTestValueType, FDefaultSetAllocator, FCaseSensitiveLookupKeyFuncs<FContainerTestValueType>>, FString>();
+
+	RunContainerTests<TSortedMap<int32, FContainerTestValueType>, int32>();
+	RunContainerTests<TSortedMap<FName, FContainerTestValueType, FDefaultAllocator, FNameLexicalLess>, FName>();
+	RunContainerTests<TSortedMap<FString, FContainerTestValueType>, FString>();
+	RunContainerTests<TSortedMap<FString, FContainerTestValueType, TInlineAllocator<64>>, FString>();
+
+	// Verify use of FName index sorter with SortedMap
+
+	TSortedMap<FName, int32, FDefaultAllocator, FNameFastLess> NameMap;
+	NameMap.Add(NAME_NameProperty);
+	NameMap.Add(NAME_FloatProperty);
+	NameMap.Add(NAME_None);
+	NameMap.Add(NAME_IntProperty);
+
+	auto It = NameMap.CreateConstIterator();
+
+	check(It->Key == NAME_None); ++It;
+	check(It->Key == NAME_IntProperty); ++It;
+	check(It->Key == NAME_FloatProperty); ++It;
+	check(It->Key == NAME_NameProperty); ++It;
+	check(!It);
+
+	RunGetRefTests();
+
+	return true;
+}
+
+IMPLEMENT_SIMPLE_AUTOMATION_TEST(FContainerPerformanceTest, "System.Core.Containers.Performance", EAutomationTestFlags::ApplicationContextMask | EAutomationTestFlags::EngineFilter)
+bool FContainerPerformanceTest::RunTest(const FString& Parameters)
+{
+	RunPerformanceTest<TMap<int32, FString>, int32>(TEXT("TMap int32"), 1, 1000000);
+	RunPerformanceTest<TMap<int32, FString>, int32>(TEXT("TMap int32"), 10, 1000000);
+	RunPerformanceTest<TMap<int32, FString>, int32>(TEXT("TMap int32"), 100, 1000000);
+	RunPerformanceTest<TMap<int32, FString>, int32>(TEXT("TMap int32"), 1000, 1000000);
+	RunPerformanceTest<TMap<int32, FString>, int32>(TEXT("TMap int32"), 10000, 1000000);
+
+	RunPerformanceTest<TMap<FName, FString>, FName>(TEXT("TMap FName"), 1, 1000000);
+	RunPerformanceTest<TMap<FName, FString>, FName>(TEXT("TMap FName"), 10, 1000000);
+	RunPerformanceTest<TMap<FName, FString>, FName>(TEXT("TMap FName"), 100, 1000000);
+	RunPerformanceTest<TMap<FName, FString>, FName>(TEXT("TMap FName"), 1000, 1000000);
+	RunPerformanceTest<TMap<FName, FString>, FName>(TEXT("TMap FName"), 10000, 1000000);
+
+	RunPerformanceTest<TMap<FString, FString>, FString>(TEXT("TMap FString"), 1, 1000000);
+	RunPerformanceTest<TMap<FString, FString>, FString>(TEXT("TMap FString"), 10, 1000000);
+	RunPerformanceTest<TMap<FString, FString>, FString>(TEXT("TMap FString"), 100, 1000000);
+	RunPerformanceTest<TMap<FString, FString>, FString>(TEXT("TMap FString"), 1000, 1000000);
+	RunPerformanceTest<TMap<FString, FString>, FString>(TEXT("TMap FString"), 10000, 1000000);
+
+	RunPerformanceTest<TSortedMap<int32, FString>, int32>(TEXT("TSortedMap int32"), 1, 1000000);
+	RunPerformanceTest<TSortedMap<int32, FString>, int32>(TEXT("TSortedMap int32"), 10, 1000000);
+	RunPerformanceTest<TSortedMap<int32, FString>, int32>(TEXT("TSortedMap int32"), 100, 1000000);
+	RunPerformanceTest<TSortedMap<int32, FString>, int32>(TEXT("TSortedMap int32"), 1000, 1000000);
+	RunPerformanceTest<TSortedMap<int32, FString>, int32>(TEXT("TSortedMap int32"), 10000, 1000000);
+
+	RunPerformanceTest<TSortedMap<FName, FString, FDefaultAllocator, FNameLexicalLess>, FName>(TEXT("TSortedMap FName"), 1, 1000000);
+	RunPerformanceTest<TSortedMap<FName, FString, FDefaultAllocator, FNameLexicalLess>, FName>(TEXT("TSortedMap FName"), 10, 1000000);
+	RunPerformanceTest<TSortedMap<FName, FString, FDefaultAllocator, FNameLexicalLess>, FName>(TEXT("TSortedMap FName"), 100, 1000000);
+	RunPerformanceTest<TSortedMap<FName, FString, FDefaultAllocator, FNameLexicalLess>, FName>(TEXT("TSortedMap FName"), 1000, 1000000);
+	RunPerformanceTest<TSortedMap<FName, FString, FDefaultAllocator, FNameLexicalLess>, FName>(TEXT("TSortedMap FName"), 10000, 1000000);
+
+	RunPerformanceTest<TSortedMap<FString, FString>, FString>(TEXT("TSortedMap FString"), 1, 1000000);
+	RunPerformanceTest<TSortedMap<FString, FString>, FString>(TEXT("TSortedMap FString"), 10, 1000000);
+	RunPerformanceTest<TSortedMap<FString, FString>, FString>(TEXT("TSortedMap FString"), 100, 1000000);
+	RunPerformanceTest<TSortedMap<FString, FString>, FString>(TEXT("TSortedMap FString"), 1000, 1000000);
+	RunPerformanceTest<TSortedMap<FString, FString>, FString>(TEXT("TSortedMap FString"), 10000, 1000000);
+
+	RunSetPerformanceTest<TSet<FName>, FName>(TEXT("TSet FName"), 1, 1000000);
+	RunSetPerformanceTest<TSet<FName>, FName>(TEXT("TSet FName"), 10, 1000000);
+	RunSetPerformanceTest<TSet<FName>, FName>(TEXT("TSet FName"), 100, 1000000);
+
+	RunSetPerformanceTest<TArray<FName>, FName>(TEXT("TArray FName"), 1, 1000000);
+	RunSetPerformanceTest<TArray<FName>, FName>(TEXT("TArray FName"), 10, 1000000);
+	RunSetPerformanceTest<TArray<FName>, FName>(TEXT("TArray FName"), 100, 1000000);
+
+	return true;
+}
+
+namespace
+{
+	struct FRecorder
+	{
+		FRecorder(uint32 InKey = 0, uint32 InPayload=0)
+			: Id(NextId++)
+			, Key(InKey)
+			, Payload(InPayload)
+		{
+		}
+		FRecorder(const FRecorder& Other)
+			: Id(NextId++)
+			, Key(Other.Key)
+			, Payload(Other.Payload)
+			, NumCopies(Other.NumCopies+1)
+			, NumMoves(Other.NumMoves)
+		{
+		}
+		FRecorder(FRecorder&& Other)
+			: Id(NextId++)
+			, Key(Other.Key)
+			, Payload(Other.Payload)
+			, NumCopies(Other.NumCopies)
+			, NumMoves(Other.NumMoves+1)
+		{
+		}
+
+		bool operator==(const FRecorder& Other) const
+		{
+			return Key == Other.Key;
+		}
+
+		uint32 Id;
+		uint32 Key;
+		uint32 Payload;
+		uint32 NumCopies = 0;
+		uint32 NumMoves = 0;
+
+		static uint32 NextId;
+	};
+	uint32 FRecorder::NextId = 0;
+
+	uint32 GetTypeHash(const FRecorder& Recorder)
+	{
+		return Recorder.Payload;
+	}
+
+}
+IMPLEMENT_SIMPLE_AUTOMATION_TEST(FContainersTSetTest, "System.Core.Containers.TSet", EAutomationTestFlags::ApplicationContextMask | EAutomationTestFlags::EngineFilter)
+bool FContainersTSetTest::RunTest(const FString& Parameters)
+{
+	enum class EArgType
+	{
+		Copy,
+		Move
+	};
+	enum class EHashType
+	{
+		Internal,
+		PassedIn
+	};
+
+	for (EHashType HashType : { EHashType::Internal, EHashType::PassedIn})
+	{
+		for (EArgType ArgType : {EArgType::Copy, EArgType::Move})
+		{
+			FString FuncName;
+			auto FullText = [HashType,ArgType,&FuncName](const TCHAR* Message)
+			{
+				const TCHAR* HashText = HashType == EHashType::Internal ? TEXT("") : TEXT("ByHash");
+				const TCHAR* ArgText = ArgType == EArgType::Copy ? TEXT("const&") : TEXT("&&");
+				return FString::Printf(TEXT("TSet::%s%s(%s) %s"), *FuncName, HashText, ArgText, Message);
+			};
+
+			// Test TSet::Add(const&), Add(&&), AddByHash(const&), AddByHash(&&)
+			FuncName = TEXT("Add");
+			{
+				TSet<FRecorder> Set;
+				FRecorder First(37, 43);
+				bool bAlreadyInSet = true;
+
+				if (HashType == EHashType::Internal)
+					if (ArgType == EArgType::Copy)
+						Set.Add(First, &bAlreadyInSet);
+					else
+						Set.Add(MoveTemp(First), &bAlreadyInSet);
+				else
+					if (ArgType == EArgType::Copy)
+						Set.AddByHash(GetTypeHash(First), First, &bAlreadyInSet);
+					else
+						Set.AddByHash(GetTypeHash(First), MoveTemp(First), &bAlreadyInSet);
+				TestFalse(FullText(TEXT("returns bAlreadyInSet==false for first add")), bAlreadyInSet);
+
+				FRecorder* Found = Set.Find(First);
+				if (ArgType == EArgType::Copy)
+					TestTrue(FullText(TEXT("constructs a copy")), Found && Found->Id > First.Id && Found->NumCopies > 0 && Found->Payload == First.Payload);
+				else
+					TestTrue(FullText(TEXT("constructs a move")), Found && Found->Id > First.Id && Found->NumCopies == 0 && Found->NumMoves >= 1 && Found->Payload == First.Payload);
+
+				uint32 FoundId = Found ? Found->Id : 0;
+				Found = Set.Find(First);
+				TestTrue(TEXT("Finding an element returns a reference, no copies"), Found && Found->Id == FoundId);
+
+				FRecorder Second(37, 56);
+				if (HashType == EHashType::Internal)
+					if (ArgType == EArgType::Copy)
+						Set.Add(Second, &bAlreadyInSet);
+					else
+						Set.Add(MoveTemp(Second), &bAlreadyInSet);
+				else
+					if (ArgType == EArgType::Copy)
+						Set.AddByHash(GetTypeHash(Second), Second, &bAlreadyInSet);
+					else
+						Set.AddByHash(GetTypeHash(Second), MoveTemp(Second), &bAlreadyInSet);
+				TestTrue(FullText(TEXT("returns bAlreadyInSet==true for second add")), bAlreadyInSet);
+				Found = Set.Find(First);
+				TestTrue(FullText(TEXT("with a duplicate key constructs a copy of the new key")), Found && Found->Id > Second.Id && Found->Payload == Second.Payload);
+			}
+
+			// Test TSet::FindOrAdd(const&), FindOrAdd(&&), FindOrAddByHash(const&), FindOrAddByHash(&&)
+			FuncName = TEXT("FindOrAdd");
+			{
+				TSet<FRecorder> Set;
+				FRecorder First(37, 43);
+				bool bAlreadyInSet = true;
+
+				FRecorder* FindOrAddResult;
+				if (HashType == EHashType::Internal)
+					if (ArgType == EArgType::Copy)
+						FindOrAddResult = &Set.FindOrAdd(First, &bAlreadyInSet);
+					else
+						FindOrAddResult = &Set.FindOrAdd(MoveTemp(First), &bAlreadyInSet);
+				else
+					if (ArgType == EArgType::Copy)
+						FindOrAddResult = &Set.FindOrAddByHash(GetTypeHash(First), First, &bAlreadyInSet);
+					else
+						FindOrAddResult = &Set.FindOrAddByHash(GetTypeHash(First), MoveTemp(First), &bAlreadyInSet);
+				TestFalse(FullText(TEXT("returns bAlreadyInSet==false for first add")), bAlreadyInSet);
+				if (ArgType == EArgType::Copy)
+					TestTrue(FullText(TEXT("on the first constructs a copy")), FindOrAddResult->Id > First.Id && FindOrAddResult->NumCopies > 0 && FindOrAddResult->Payload == First.Payload);
+				else
+					TestTrue(FullText(TEXT("on the first constructs a move")), FindOrAddResult->Id > First.Id && FindOrAddResult->NumCopies == 0 && FindOrAddResult->NumMoves >= 1 && FindOrAddResult->Payload == First.Payload);
+				uint32 FoundId = FindOrAddResult->Id;
+
+				FRecorder* Found = Set.Find(First);
+				TestTrue(FullText(TEXT("returns same value as future find")), Found&& Found->Id == FindOrAddResult->Id);
+				Found = Set.Find(First);
+				TestTrue(TEXT("Finding an element returns a reference, no copies"), Found && Found->Id == FoundId);
+
+				FRecorder Second(37, 56);
+				if (HashType == EHashType::Internal)
+					if (ArgType == EArgType::Copy)
+						FindOrAddResult = &Set.FindOrAdd(Second, &bAlreadyInSet);
+					else
+						FindOrAddResult = &Set.FindOrAdd(MoveTemp(Second), &bAlreadyInSet);
+				else
+					if (ArgType == EArgType::Copy)
+						FindOrAddResult = &Set.FindOrAddByHash(GetTypeHash(Second), Second, &bAlreadyInSet);
+					else
+						FindOrAddResult = &Set.FindOrAddByHash(GetTypeHash(Second), MoveTemp(Second), &bAlreadyInSet);
+				TestTrue(FullText(TEXT("returns bAlreadyInSet==true for second add")), bAlreadyInSet);
+				TestTrue(FullText(TEXT("with a duplicate key keeps the original key, returned from FindOrAdd")), FindOrAddResult->Id == FoundId && FindOrAddResult->Payload == First.Payload);
+				Found = Set.Find(First);
+				TestTrue(FullText(TEXT("with a duplicate key keeps the original key, returned from future Finds")), Found->Id == FoundId && Found->Payload == First.Payload);
+			}
+		}
+	}
+
+
+	return !HasAnyErrors();
+}
+
+
+namespace ArrayViewTests
+{
+	// commented out lines shouldn't compile
+
+	struct Base
+	{
+		int32 b;
+	};
+
+	struct Derived : public Base
+	{
+		int32 d;
+	};
+
+	template<typename T>
+	void TestFunction(TArrayView<T>)
+	{
+	}
+
+	template<typename T>
+	void TestFunction64(TArrayView64<T>)
+	{
+	}
+
+	bool RunTest()
+	{
+		// C array + derived-to-base conversions
+		Derived test1[13];
+		TestFunction<Derived>(test1);
+		//TestFunction<Base>(test1);
+		//TestFunction<const Base>(test1);
+		TestFunction64<Derived>(test1);
+
+		// C array of pointers + derived-to-base conversions
+		Derived* test2[13];
+		TestFunction<const Derived* const>(test2);
+		//TestFunction<const Derived*>(test2);
+		TestFunction<Derived* const>(test2);
+		//TestFunction<const Base* const>(test2);
+		TestFunction64<const Derived* const>(test2);
+		TestFunction64<Derived* const>(test2);
+
+		// TArray + derived-to-base conversions
+		TArray<Derived> test3;
+		TestFunction<Derived>(test3);
+		//TestFunction<Base>(test3);
+		//TestFunction<const Base>(test3);
+		TestFunction64<Derived>(test3);
+
+		// const TArray
+		const TArray<Base> test4;
+		TestFunction<const Base>(test4);
+		//TestFunction<Base>(test4);
+		TestFunction64<const Base>(test4);
+
+		// TArray of const
+		TArray<const Base> test5;
+		TestFunction<const Base>(test5);
+		//TestFunction<Base>(test5);
+		TestFunction64<const Base>(test5);
+
+		// temporary C array
+		struct Test6
+		{
+			Base test[13];
+		};
+		TestFunction<const Base>(Test6().test);
+		//TestFunction<Base>(Test6().test); // shouldn't compile but VS allows it :(
+		TestFunction64<const Base>(Test6().test);
+
+		// TArrayView64 from TArrayView
+		TArrayView<Derived> test7 = test1;
+		TestFunction64<Derived>(test7);
+		//TArrayView64<Derived> test7_64 = test1;
+		//TestFunction<Derived>(test7_64);
+
+		// TArray[64] from TArrayView[64]
+		TArrayView<Derived> test8;
+		TArrayView<Derived> test8_64;
+		TArray<Derived> test8_32from32(test8);
+		TArray<Derived> test8_32from64(test8_64);
+		TArray64<Derived> test8_64from32(test8);
+		TArray64<Derived> test8_64from64(test8_64);
+		test8_32from32 = test8;
+		test8_32from64 = test8_64;
+		test8_64from32 = test8;
+		test8_64from64 = test8_64;
+
+		return true;
+	}
+};
+
+#endif // WITH_DEV_AUTOMATION_TESTS
