LowLevelTests: Reverted the Core tests that were moved in 19498448 and 19571202

diff --git a/Engine/Source/Runtime/Core/Private/Tests/Containers/RingBufferTest.cpp b/Engine/Source/Runtime/Core/Private/Tests/Containers/RingBufferTest.cpp
new file mode 100644
--- /dev/null
+++ b/Engine/Source/Runtime/Core/Private/Tests/Containers/RingBufferTest.cpp
@@ -0,0 +1,1121 @@
+// Copyright Epic Games, Inc. All Rights Reserved.
+
+#include "Containers/RingBuffer.h"
+#include "Misc/AutomationTest.h"
+
+#if WITH_DEV_AUTOMATION_TESTS
+
+class FRingBufferTest : FAutomationTestBase
+{
+public:
+	static bool IsIntegerRange(const TRingBuffer<uint32>& Queue, uint32 Start, uint32 End, bool bForward = true)
+	{
+		if (Queue.Num() != End - Start)
+		{
+			return false;
+		}
+
+		// Peek elements in queue at given offset, peek from back to front
+		for (int32 It = 0; It < Queue.Num(); ++It)
+		{
+			uint32 QueueValue = bForward ? Queue[It] : Queue[Queue.Num() - 1 - It];
+			if (QueueValue != Start + It)
+			{
+				return false;
+			}
+		}
+
+		return true;
+	}
+
+	struct Counter
+	{
+		Counter(uint32 InValue = 0x12345)
+			:Value(InValue)
+		{
+			++NumVoid;
+		}
+		Counter(const Counter& Other)
+			:Value(Other.Value)
+		{
+			++NumCopy;
+		}
+		Counter(Counter&& Other)
+			:Value(Other.Value)
+		{
+			++NumMove;
+		}
+		~Counter()
+		{
+			++NumDestruct;
+		}
+
+		operator uint32() const
+		{
+			return Value;
+		}
+
+		bool operator==(const Counter& Other) const
+		{
+			return Value == Other.Value;
+		}
+		static int NumVoid;
+		static int NumCopy;
+		static int NumMove;
+		static int NumDestruct;
+		static void Clear()
+		{
+			NumVoid = NumCopy = NumMove = NumDestruct = 0;
+		}
+
+		uint32 Value;
+	};
+
+	FRingBufferTest(const FString& InName, const bool bInComplexTask)
+		: FAutomationTestBase(InName, bInComplexTask)
+	{
+	}
+
+	bool RunTest(const FString& Parameters)
+	{
+		// Test empty
+		{
+			TRingBuffer<uint32> Q(0);
+
+			TestTrue(TEXT("Test empty - IsEmpty"), Q.IsEmpty());
+			TestEqual(TEXT("Test empty - Size"), Q.Num(), 0);
+			TestEqual(TEXT("Test empty - Capacity"), Q.Max(), 0);
+			TestEqual(TEXT("Test empty - Iterator"), Q.begin(), Q.end());
+			TestEqual(TEXT("Test empty - ConvertPointerToIndex"), Q.ConvertPointerToIndex(nullptr), INDEX_NONE);
+			TestEqual(TEXT("Test empty - ConvertPointerToIndex"), Q.ConvertPointerToIndex(reinterpret_cast<uint32*>(this)), INDEX_NONE);
+			Q.Trim();
+			TestEqual(TEXT("Test Trim From empty - Size"), Q.Num(), 0);
+			TestEqual(TEXT("Test Trim From empty - Capacity"), Q.Max(), 0);
+			Q.Reset();
+			TestEqual(TEXT("Test Reset From empty - Size"), Q.Num(), 0);
+			TestEqual(TEXT("Test Reset From empty - Capacity"), Q.Max(), 0);
+			Q.Empty(0);
+			TestEqual(TEXT("Test Empty From empty - Size"), Q.Num(), 0);
+			TestEqual(TEXT("Test Empty From empty - Capacity"), Q.Max(), 0);
+			Q.PopFront(0);
+			Q.Pop(0);
+			TestEqual(TEXT("Test Pop on empty - Size"), Q.Num(), 0);
+			TestEqual(TEXT("Test Pop on empty - Capacity"), Q.Max(), 0);
+			TestEqual(TEXT("Test empty - IsValidIndex"), Q.IsValidIndex(0), false);
+
+
+			const TRingBuffer<uint32> ConstQ(0);
+			TestTrue(TEXT("Test const empty - IsEmpty"), ConstQ.IsEmpty());
+			TestEqual(TEXT("Test const empty - Size"), ConstQ.Num(), 0);
+			TestEqual(TEXT("Test const empty - Capacity"), ConstQ.Max(), 0);
+			TestEqual(TEXT("Test const empty - Iterator"), ConstQ.begin(), ConstQ.end());
+			TestEqual(TEXT("Test const empty - ConvertPointerToIndex"), ConstQ.ConvertPointerToIndex(reinterpret_cast<uint32*>(this)), INDEX_NONE);
+		}
+
+		// Test Adding a sequence of elements
+		{
+			const TRingBuffer<int32>::IndexType FirstSize = 8;
+
+			TRingBuffer<int32> Q(0);
+
+			TestEqual(TEXT("Test AddSequence - Capacity (Implementation Detail)"), Q.Max(), 0);
+			Q.Emplace(0);
+			TestEqual(TEXT("Test AddSequence - Size"), Q.Num(), 1);
+			TestEqual(TEXT("Test AddSequence - Capacity (Implementation Detail)"), Q.Max(), 1);
+			Q.Emplace(1);
+			TestEqual(TEXT("Test AddSequence - Size"), Q.Num(), 2);
+			TestEqual(TEXT("Implementation Detail - These tests expect that growing size will set capacity to successive powers of 2."), Q.Max(), 2);
+			for (int32 It = 2; It < FirstSize; ++It)
+			{
+				Q.Emplace(It);
+				TestEqual(TEXT("Test AddSequence - Size"), Q.Num(), It + 1);
+				TestEqual(TEXT("Test AddSequence - Capacity (Implementation Detail)"), static_cast<uint32>(Q.Max()), FMath::RoundUpToPowerOfTwo(It + 1));
+			}
+
+			for (int32 Index = 0; Index < FirstSize; ++Index)
+			{
+				TestEqual(TEXT("Test AddSequence - Expected values"), Q[Index], Index);
+				TestEqual(TEXT("Test AddSequence const- Expected values"), const_cast<TRingBuffer<int32>&>(Q)[Index], Index);
+			}
+
+			const TRingBuffer<int32>::IndexType SecondSize = 13;
+			for (int32 It = FirstSize; It < SecondSize; ++It)
+			{
+				Q.Emplace(It);
+				TestEqual(TEXT("Test AddSequence non powerof2 - Size"), Q.Num(), It + 1);
+				TestEqual(TEXT("Test AddSequence non powerof2 const - Capacity (Implementation Detail)"), static_cast<uint32>(Q.Max()), FMath::RoundUpToPowerOfTwo(It + 1));
+			}
+
+			for (int32 Index = 0; Index < FirstSize; ++Index)
+			{
+				TestEqual(TEXT("Test AddSequence non powerof2 - Expected values"), Q[Index], Index);
+				TestEqual(TEXT("Test AddSequence non powerof2 const - Expected values"), const_cast<TRingBuffer<int32>&>(Q)[Index], Index);
+			}
+		}
+
+		// Test Add under/over Capacity
+		{
+			const TRingBuffer<int32>::IndexType FirstElementsToAdd = 3;
+			const TRingBuffer<int32>::IndexType InitialCapacity = 8;
+			const TRingBuffer<int32>::IndexType SecondElementsToAdd = 9;
+
+			TRingBuffer<int32> Q(InitialCapacity);
+
+			for (int32 It = 0; It < FirstElementsToAdd; ++It)
+			{
+				Q.Emplace(It);
+			}
+
+			TestEqual(TEXT("Test Add under Capacity - Size"), Q.Num(), FirstElementsToAdd);
+			TestEqual(TEXT("Test Add under Capacity - Capacity"), Q.Max(), InitialCapacity);
+			for (int32 Index = 0; Index < FirstElementsToAdd; ++Index)
+			{
+				TestEqual(TEXT("Test Add under Capacity - Expected values"), Q[Index], Index);
+				TestEqual(TEXT("Test Add under Capacity const - Expected values"), const_cast<TRingBuffer<int32>&>(Q)[Index], Index);
+			}
+
+			for (int32 It = FirstElementsToAdd; It < SecondElementsToAdd; ++It)
+			{
+				Q.Emplace(It);
+			}
+
+			TestEqual(TEXT("Test Add over Capacity - Size"), Q.Num(), SecondElementsToAdd);
+			TestEqual(TEXT("Test Add over Capacity - Capacity (Implementation Detail)"), static_cast<uint32>(Q.Max()), FMath::RoundUpToPowerOfTwo(SecondElementsToAdd));
+			for (int32 Index = 0; Index < SecondElementsToAdd; ++Index)
+			{
+				TestEqual(TEXT("Test Add over Capacity - Expected values"), Q[Index], Index);
+				TestEqual(TEXT("Test Add over Capacity const - Expected values"), const_cast<TRingBuffer<int32>&>(Q)[Index], Index);
+			}
+		}
+
+		// Test Last/First
+		{
+			TRingBuffer<uint32> Q({ 0,1,2,3 });
+			TestEqual(TEXT("Test Last"), 3, Q.Last());
+			Q.Last() = 4;
+			TestEqual(TEXT("Test Last const"), 4, const_cast<TRingBuffer<uint32>&>(Q).Last());
+			TestEqual(TEXT("Test First"), 0, Q.First());
+			Q.First() = 5;
+			TestEqual(TEXT("Test First const"), 5, const_cast<TRingBuffer<uint32>&>(Q).First());
+		}
+
+		// Test PopFrontValue/PopValue
+		{
+			TRingBuffer<Counter> Q({ 31,32,33 });
+			Q.AddFront(30);
+
+			Counter::Clear();
+			Counter C(Q.PopFrontValue());
+			TestEqual(TEXT("PopFrontValue - PoppedValue"), C.Value, 30);
+			TestTrue(TEXT("PopFrontValue - ConstructorCounts"), Counter::NumMove > 0 && Counter::NumCopy == 0);
+			TestEqual(TEXT("PopFrontValue - Remaining Values"), Q, TRingBuffer<Counter>({ 31,32,33 }));
+			Counter::Clear();
+			TestEqual(TEXT("PopFrontValue Inline - PoppedValue"), Q.PopFrontValue().Value, 31);
+			TestTrue(TEXT("PopFrontValue Inline - ConstructorCounts"), Counter::NumCopy == 0);
+			TestEqual(TEXT("PopFrontValue Inline - Remaining Values"), Q, TRingBuffer<Counter>({ 32,33 }));
+
+			Counter::Clear();
+			Counter D(Q.PopValue());
+			TestEqual(TEXT("PopValue - PoppedValue"), D.Value, 33);
+			TestTrue(TEXT("PopValue - ConstructorCounts"), Counter::NumMove > 0 && Counter::NumCopy == 0);
+			TestEqual(TEXT("PopValue - Remaining Values"), Q, TRingBuffer<Counter>({ Counter(32) }));
+			Counter::Clear();
+			TestEqual(TEXT("PopValue Inline - PoppedValue"), Q.PopValue().Value, 32);
+			TestTrue(TEXT("PopValue Inline - ConstructorCounts"), Counter::NumCopy == 0);
+			TestTrue(TEXT("PopValue Inline - Remaining Values"), Q.IsEmpty());
+		}
+
+		// Test Initializer_List
+		{
+			const TRingBuffer<int32>::IndexType InitializerSize = 9;
+			TRingBuffer<int32> Q({ 0, 1, 2, 3, 4, 5, 6, 7, 8 });
+
+			TestEqual(TEXT("Test Initializer_List - Size"), Q.Num(), InitializerSize);
+			TestEqual(TEXT("Test Initializer_List - Capacity (Implementation Detail)"), static_cast<uint32>(Q.Max()), FMath::RoundUpToPowerOfTwo(InitializerSize));
+			for (int32 Index = 0; Index < InitializerSize; ++Index)
+			{
+				TestEqual(TEXT("Test Initializer_List - Expected values"), Q[Index], Index);
+			}
+		}
+
+		// Test RingBuffer's Copy Constructors et al
+		{
+			TRingBuffer<uint32> Original({ 0,1,2,3,4,5,6,7 });
+			TRingBuffer<uint32> Copy(Original);
+			TestEqual(TEXT("Copy Constructor"), Original, Copy);
+			TRingBuffer<uint32> Moved(MoveTemp(Copy));
+			TestEqual(TEXT("Move Constructor"), Original, Moved);
+			TestEqual(TEXT("Move Constructor did in fact move"), Copy.Max(), 0);
+			TRingBuffer<uint32> AssignCopy;
+			AssignCopy = Original;
+			TestEqual(TEXT("Copy Assignment"), Original, AssignCopy);
+			TRingBuffer<uint32> AssignMove;
+			AssignMove = MoveTemp(AssignCopy);
+			TestEqual(TEXT("Move Assignment"), Original, AssignMove);
+			TestEqual(TEXT("Move Assignment did in fact move"), AssignCopy.Max(), 0);
+		}
+
+		// Test Equality 
+		{
+			auto TestEquality = [this](const TCHAR* Message, bool ExpectedEqual, const TRingBuffer<int32>& A, const TRingBuffer<int32>& B)
+			{
+				TestEqual(*FString::Printf(TEXT("Test equality - %s - A == B"), Message), A == B, ExpectedEqual);
+				TestEqual(*FString::Printf(TEXT("Test equality - %s - B == A"), Message), B == A, ExpectedEqual);
+				TestEqual(*FString::Printf(TEXT("Test equality - %s - A != B"), Message), A != B, !ExpectedEqual);
+				TestEqual(*FString::Printf(TEXT("Test equality - %s - B != A"), Message), B != A, !ExpectedEqual);
+			};
+
+			TestEquality(TEXT("empty"), true, TRingBuffer<int32>(0), TRingBuffer<int32>(0));
+			TestEquality(TEXT("empty different capacities"), true, TRingBuffer<int32>(0), TRingBuffer<int32>(8));
+			TestEquality(TEXT("equal nonempty powerof2"), true, TRingBuffer<int32>({ 0, 1, 2, 3 }), TRingBuffer<int32>({ 0, 1, 2, 3 }));
+			TestEquality(TEXT("equal nonempty nonpowerof2"), true, TRingBuffer<int32>({ 0, 1, 2, 3, 4, 5 }), TRingBuffer<int32>({ 0, 1, 2, 3, 4, 5 }));
+			{
+				TRingBuffer<int32> QNum6Cap16(16);
+				for (int32 Index = 0; Index < 6; ++Index)
+				{
+					QNum6Cap16.Add(Index);
+				}
+				TestEquality(TEXT("equal nonempty different capacities"), true, QNum6Cap16, TRingBuffer<int32>({ 0, 1, 2, 3, 4, 5 }));
+			}
+
+			TestEquality(TEXT("empty to nonempty"), false, TRingBuffer<int32>(0), TRingBuffer<int32>({ 0, 1, 2, 3, 4, 5 }));
+			TestEquality(TEXT("smaller size to bigger size"), false, TRingBuffer<int32>({ 0, 1, 2 }), TRingBuffer<int32>({ 0, 1, 2, 3, 4, 5 }));
+			TestEquality(TEXT("same size different elements"), false, TRingBuffer<int32>({ 0, 1, 2 }), TRingBuffer<int32>({ 0, 1, 3 }));
+			TestEquality(TEXT("same elements different order"), false, TRingBuffer<int32>({ 0, 1, 2 }), TRingBuffer<int32>({ 0, 2, 1 }));
+
+			for (int HasPow2 = 0; HasPow2 < 2; ++HasPow2)
+			{
+				const int Count = HasPow2 ? 8 : 7;
+				TRingBuffer<int32> Q0Pop;
+				TRingBuffer<int32> Q1PopFront;
+				TRingBuffer<int32> Q2PopFront;
+				TRingBuffer<int32> Q1Pop;
+				TRingBuffer<int32> Q2Pop;
+				TRingBuffer<int32> Q2PopFront3Pop;
+				Q1PopFront.Add(47);
+				Q2PopFront.Add(576);
+				Q2PopFront.Add(-5);
+				Q2PopFront3Pop.Add(84);
+				Q2PopFront3Pop.Add(1000);
+				for (int Index = 0; Index < Count; ++Index)
+				{
+					Q0Pop.Add(Index);
+					Q1PopFront.Add(Index);
+					Q2PopFront.Add(Index);
+					Q1Pop.Add(Index);
+					Q2Pop.Add(Index);
+					Q2PopFront3Pop.Add(Index);
+				}
+				Q1PopFront.PopFront();
+				Q2PopFront.PopFront();
+				Q2PopFront.PopFront();
+				Q1Pop.Add(-18);
+				Q1Pop.Pop();
+				Q2Pop.Add(105);
+				Q2Pop.Add(219);
+				Q2Pop.Pop();
+				Q2Pop.Pop();
+				Q2PopFront3Pop.Add(456);
+				Q2PopFront3Pop.Add(654);
+				Q2PopFront3Pop.Add(8888888);
+				Q2PopFront3Pop.PopFront();
+				Q2PopFront3Pop.Pop();
+				Q2PopFront3Pop.PopFront();
+				Q2PopFront3Pop.Pop();
+				Q2PopFront3Pop.Pop();
+
+				const TCHAR* Names[] =
+				{
+					TEXT("Q0Pop"),
+					TEXT("Q1PopFront"),
+					TEXT("Q2PopFront"),
+					TEXT("Q1Pop"),
+					TEXT("Q2Pop"),
+					TEXT("Q2PopFront3Pop"),
+				};
+				TRingBuffer<int32>* Pops[] =
+				{
+					&Q0Pop,
+					&Q1PopFront,
+					&Q2PopFront,
+					&Q1Pop,
+					&Q2Pop,
+					&Q2PopFront3Pop
+				};
+
+
+				auto TestThesePops = [this, HasPow2, &TestEquality, &Names, &Pops](int TrialA, int TrialB)
+				{
+					TestEquality(*FString::Printf(TEXT("%s - %s - %s"), Names[TrialA], Names[TrialB], (HasPow2 ? TEXT("powerof2") : TEXT("nonpowerof2"))), true, *Pops[TrialA], *Pops[TrialB]);
+				};
+
+				for (int TrialA = 0; TrialA < UE_ARRAY_COUNT(Names); ++TrialA)
+				{
+					for (int TrialB = TrialA /* test each against itself as well */; TrialB < UE_ARRAY_COUNT(Names); ++TrialB)
+					{
+						TestThesePops(TrialA, TrialB);
+					}
+				}
+			}
+		}
+
+		// Test Add and pop all
+		for (int Direction = 0; Direction < 2; ++Direction)
+		{
+			bool bIsAddBack = Direction == 0;
+			auto GetMessage = [&bIsAddBack](const TCHAR* Message)
+			{
+				return FString::Printf(TEXT("Test %s (%s)"), Message, (bIsAddBack ? TEXT("AddBack") : TEXT("AddFront")));
+			};
+
+			// Test Mixed Adds and Pops
+			{
+				const TRingBuffer<uint32>::IndexType ElementsToAdd = 256;
+				const TRingBuffer<uint32>::IndexType ElementPopMod = 16;
+				const TRingBuffer<uint32>::IndexType ExpectedSize = 256 - ElementPopMod;
+				const TRingBuffer<uint32>::IndexType ExpectedCapacity = 256;
+
+				TRingBuffer<uint32> Q(4);
+
+				uint32 ExpectedPoppedValue = 0;
+				for (uint32 It = 0; It < 256; ++It)
+				{
+					if (bIsAddBack)
+					{
+						Q.Add(It);
+						TestEqual(*GetMessage(TEXT("Add and pop - Add")), It, Q[Q.Num() - 1]);
+					}
+					else
+					{
+						Q.AddFront(It);
+						TestEqual(*GetMessage(TEXT("Add and pop - Add")), It, Q[0]);
+					}
+
+					if (It % ElementPopMod == 0)
+					{
+						uint32 PoppedValue;
+						if (bIsAddBack)
+						{
+							PoppedValue = Q[0];
+							Q.PopFront();
+						}
+						else
+						{
+							PoppedValue = Q[Q.Num() - 1];
+							Q.Pop();
+						}
+						TestEqual(*GetMessage(TEXT("Add and pop - Pop")), ExpectedPoppedValue, PoppedValue);
+						++ExpectedPoppedValue;
+					}
+				}
+
+				TestEqual(*GetMessage(TEXT("Add and pop - Size")), Q.Num(), ExpectedSize);
+				TestEqual(*GetMessage(TEXT("Add and pop - Capacity")), Q.Max(), ExpectedCapacity);
+				TestTrue(*GetMessage(TEXT("Add and pop - IntegerRange")), IsIntegerRange(Q, ExpectedPoppedValue, ExpectedPoppedValue + ExpectedSize, bIsAddBack));
+			}
+
+
+			// Popping down to empty
+			{
+				const TRingBuffer<uint32>::IndexType ElementsToAdd = 256;
+
+				TRingBuffer<uint32> Q(ElementsToAdd);
+
+				TestTrue(*GetMessage(TEXT("Add and pop all - IsEmpty before")), Q.IsEmpty());
+				TestEqual(*GetMessage(TEXT("Add and pop all - Size before")), Q.Num(), 0);
+
+				for (TRingBuffer<int32>::IndexType It = 0; It < ElementsToAdd; ++It)
+				{
+					if (bIsAddBack)
+					{
+						Q.Add(It);
+					}
+					else
+					{
+						Q.AddFront(It);
+					}
+				}
+
+				TestEqual(*GetMessage(TEXT("Add and pop all - Size")), Q.Num(), ElementsToAdd);
+				TestEqual(*GetMessage(TEXT("Add and pop all - Capacity")), Q.Max(), ElementsToAdd);
+				TestTrue(*GetMessage(TEXT("Add and pop all - Expected")), IsIntegerRange(Q, 0, ElementsToAdd, bIsAddBack));
+
+				for (TRingBuffer<int32>::IndexType It = 0; It < ElementsToAdd; ++It)
+				{
+					if (bIsAddBack)
+					{
+						Q.PopFront();
+					}
+					else
+					{
+						Q.Pop();
+					}
+				}
+
+				TestTrue(*GetMessage(TEXT("Add and pop all - IsEmpty after")), Q.IsEmpty());
+				TestEqual(*GetMessage(TEXT("Add and pop all - Size after")), Q.Num(), 0);
+				TestEqual(*GetMessage(TEXT("Add and pop all - Capacity after")), Q.Max(), ElementsToAdd);
+			}
+
+			// Test index wrap
+			{
+				for (int32 Offset : {-12, -8, -5, -1, 0, 2, 7, 8, 15})
+				{
+					const TRingBuffer<uint32>::IndexType ElementsToAdd = 256;
+					const TRingBuffer<uint32>::IndexType ElementPopMod = 16;
+					const TRingBuffer<uint32>::IndexType ExpectedSize = 256 - ElementPopMod;
+					const TRingBuffer<uint32>::IndexType ExpectedCapacity = 256;
+
+					TRingBuffer<uint32> Q(8);
+
+					// Set front and afterback to an arbitrary offset
+					// Note that AfterBack is always exactly equal to Front + Num()
+					Q.Front = Offset;
+					Q.AfterBack = Q.Front;
+
+					TestTrue(*GetMessage(TEXT("index wrap - IsEmpty before")), Q.IsEmpty());
+					TestEqual(*GetMessage(TEXT("index wrap - Size before")), Q.Num(), 0);
+
+					for (TRingBuffer<uint32>::IndexType It = 0; It < ElementsToAdd; ++It)
+					{
+						if (bIsAddBack)
+						{
+							Q.Add(It);
+						}
+						else
+						{
+							Q.AddFront(It);
+						}
+					}
+
+					TestEqual(*GetMessage(TEXT("index wrap - Size")), Q.Num(), ElementsToAdd);
+					TestEqual(*GetMessage(TEXT("index wrap - Capacity")), Q.Max(), ElementsToAdd);
+					TestTrue(*GetMessage(TEXT("index wrap - Expected")), IsIntegerRange(Q, 0, ElementsToAdd, bIsAddBack));
+
+					for (TRingBuffer<int32>::IndexType It = 0; It < ElementsToAdd; ++It)
+					{
+						if (bIsAddBack)
+						{
+							Q.PopFront();
+						}
+						else
+						{
+							Q.Pop();
+						}
+					}
+
+					TestTrue(*GetMessage(TEXT("index wrap - IsEmpty after")), Q.IsEmpty());
+					TestEqual(*GetMessage(TEXT("index wrap - Size after")), Q.Num(), 0);
+					TestEqual(*GetMessage(TEXT("index wrap - Capacity after")), Q.Max(), ElementsToAdd);
+				}
+			}
+		}
+
+		// Test Trim
+		{
+			const TRingBuffer<int32>::IndexType ElementsToAdd = 9;
+			const TRingBuffer<int32>::IndexType ElementsToPop = 5;
+			const TRingBuffer<int32>::IndexType ExpectedCapacity = 16;
+			const TRingBuffer<int32>::IndexType ExpectedCapacityAfterTrim = 4;
+
+			TRingBuffer<uint32> Q(0);
+
+			for (TRingBuffer<int32>::IndexType It = 0; It < ElementsToAdd; ++It)
+			{
+				Q.Add(It);
+			}
+
+			TestEqual(TEXT("Test Trim - Size"), Q.Num(), ElementsToAdd);
+			TestEqual(TEXT("Test Trim - Capacity"), Q.Max(), ExpectedCapacity);
+			TestTrue(TEXT("Test Trim - Expected"), IsIntegerRange(Q, 0, ElementsToAdd));
+
+			for (TRingBuffer<int32>::IndexType It = 0; It < ElementsToPop; ++It)
+			{
+				Q.PopFront();
+			}
+
+			Q.Trim();
+
+			TestEqual(TEXT("Test Trim - Size"), Q.Num(), ElementsToAdd - ElementsToPop);
+			TestEqual(TEXT("Test Trim - Capacity"), Q.Max(), ExpectedCapacityAfterTrim);
+			TestTrue(TEXT("Test Trim - Expected"), IsIntegerRange(Q, ElementsToPop, ElementsToAdd));
+		}
+
+		// Test First and Last acting as two stacks
+		{
+			TRingBuffer<uint32> Q;
+
+			const uint32 ElementsToAdd = 64;
+			const uint32 ElementPopMod = 5;
+
+			for (uint32 It = 0; It < ElementsToAdd; ++It)
+			{
+				Q.Add(It);
+				TestEqual(TEXT("Test TwoStacks - AddBack"), Q.Last(), It);
+				Q.AddFront(It);
+				TestEqual(TEXT("Test TwoStacks - AddFront"), Q.First(), It);
+				if (It % ElementPopMod == 0)
+				{
+					uint32 AddValue = 0xfefefefe;
+					Q.Add(AddValue);
+					TestEqual(TEXT("Test TwoStacks - Sporadic Pop"), Q.Last(), AddValue);
+					Q.Pop();
+					Q.AddFront(AddValue);
+					TestEqual(TEXT("Test TwoStacks - Sporadic PopFront"), Q.First(), AddValue);
+					Q.PopFront();
+				}
+			}
+
+			TestEqual(TEXT("Test TwoStacks - MiddleSize"), Q.Num(), ElementsToAdd * 2);
+			for (uint32 It = 0; It < ElementsToAdd * 2; ++It)
+			{
+				TestEqual(*FString::Printf(TEXT("TwoStacks - Middle value %d"), It), Q[It], (It < ElementsToAdd ? ElementsToAdd - 1 - It : It - ElementsToAdd));
+			}
+
+			for (uint32 It = 0; It < ElementsToAdd; ++It)
+			{
+				TestEqual(TEXT("Test TwoStacks - Final Pop"), Q.Last(), ElementsToAdd - 1 - It);
+				Q.Pop();
+				TestEqual(TEXT("Test TwoStacks - Final PopFront"), Q.First(), ElementsToAdd - 1 - It);
+				Q.PopFront();
+			}
+
+			TestEqual(TEXT("Test TwoStacks - FinalSize"), Q.Num(), 0);
+		}
+
+		// Test adding into space that has been cleared from popping on the other side
+		{
+			for (int Direction = 0; Direction < 2; ++Direction)
+			{
+				bool bIsAddBack = Direction == 0;
+				auto GetMessage = [bIsAddBack](const TCHAR* Message)
+				{
+					return FString::Printf(TEXT("Test AddIntoPop - %s (%s)"), Message, (bIsAddBack ? TEXT("AddBack") : TEXT("AddFront")));
+				};
+				TRingBuffer<uint32> Q({ 0,1,2,3,4,5,6,7 });
+				TRingBuffer<int32>::IndexType InitialSize = 8;
+				TestEqual(*GetMessage(TEXT("InitialSize")), InitialSize, Q.Num());
+				TestEqual(*GetMessage(TEXT("InitialCapacity (Implementation Detail)")), InitialSize, Q.Max());
+
+				if (bIsAddBack)
+				{
+					Q.Pop();
+				}
+				else
+				{
+					Q.PopFront();
+				}
+				TestEqual(*GetMessage(TEXT("PoppedSize")), InitialSize - 1, Q.Num());
+				TestEqual(*GetMessage(TEXT("PoppedCapacity")), InitialSize, Q.Max());
+
+				if (bIsAddBack)
+				{
+					Q.AddFront(8);
+				}
+				else
+				{
+					Q.Add(8);
+				}
+				TestEqual(*GetMessage(TEXT("AddedSize")), InitialSize, Q.Num());
+				TestEqual(*GetMessage(TEXT("AddedCapacity")), InitialSize, Q.Max());
+				if (bIsAddBack)
+				{
+					TestEqual(*GetMessage(TEXT("AddedValues")), Q, TRingBuffer<uint32>({ 8,0,1,2,3,4,5,6 }));
+				}
+				else
+				{
+					TestEqual(*GetMessage(TEXT("AddedValues")), Q, TRingBuffer<uint32>({ 1,2,3,4,5,6,7,8 }));
+				}
+
+				if (bIsAddBack)
+				{
+					Q.AddFront(9);
+				}
+				else
+				{
+					Q.Add(9);
+				}
+				TestEqual(*GetMessage(TEXT("Second AddedSize")), InitialSize + 1, Q.Num());
+				TestEqual(*GetMessage(TEXT("Second AddedCapacity")), static_cast<uint32>(FMath::RoundUpToPowerOfTwo(InitialSize + 1)), Q.Max());
+				if (bIsAddBack)
+				{
+					TestEqual(*GetMessage(TEXT("Second AddedValues")), Q, TRingBuffer<uint32>({ 9,8,0,1,2,3,4,5,6 }));
+				}
+				else
+				{
+					TestEqual(*GetMessage(TEXT("Second AddedValues")), Q, TRingBuffer<uint32>({ 1,2,3,4,5,6,7,8,9 }));
+				}
+			}
+		}
+
+		// Test Empty to a capacity
+		{
+			TRingBuffer<uint32> Q(16);
+			TestEqual(TEXT("Test EmptyToCapacity - InitialCapacity"), 16, Q.Max());
+			Q.Empty(8);
+			TestEqual(TEXT("Test EmptyToCapacity - Lower"), 8, Q.Max());
+			Q.Empty(32);
+			TestEqual(TEXT("Test EmptyToCapacity - Higher"), 32, Q.Max());
+		}
+
+		// Test Different Add constructors
+		{
+			auto Clear = []()
+			{
+				Counter::Clear();
+			};
+			auto TestCounts = [this](const TCHAR* Message, int32 NumVoid, int32 NumCopy, int32 NumMove, int32 NumDestruct)
+			{
+				TestTrue(Message, NumVoid == Counter::NumVoid && NumCopy == Counter::NumCopy && NumMove == Counter::NumMove && NumDestruct == Counter::NumDestruct);
+			};
+
+			Clear();
+			{
+				TRingBuffer<Counter> QEmpty(4);
+				QEmpty.Reserve(8);
+				QEmpty.Empty();
+				TRingBuffer<Counter> QEmpty2(4);
+			}
+			TestCounts(TEXT("Test Add Constructors - Unallocated elements call no constructors/destructors"), 0, 0, 0, 0);
+			{
+				TRingBuffer<Counter> QEmpty(4);
+				QEmpty.Emplace();
+				QEmpty.Pop();
+				Clear();
+			}
+			TestCounts(TEXT("Test Add Constructors - Already removed element calls no destructors"), 0, 0, 0, 0);
+
+
+			uint32 MarkerValue = 0x54321;
+			Counter CounterA(MarkerValue);
+
+			TRingBuffer<Counter> Q(4);
+			Clear();
+			for (int Direction = 0; Direction < 2; ++Direction)
+			{
+				bool bAddBack = Direction == 0;
+				auto TestDirCounts = [this, bAddBack, &TestCounts, &Q, &Clear, MarkerValue](const TCHAR* Message, int32 NumVoid, int32 NumCopy, int32 NumMove, int32 NumDestruct, bool bWasInitialized = true)
+				{
+					const TCHAR* DirectionText = bAddBack ? TEXT("Back") : TEXT("Front");
+					bool bElementExists = Q.Num() == 1;
+					TestTrue(*FString::Printf(TEXT("Test Add Constructors - %s%s ElementExists"), Message, DirectionText), bElementExists);
+					if (bWasInitialized && bElementExists)
+					{
+						TestTrue(*FString::Printf(TEXT("Test Add Constructors - %s%s ValueEquals"), Message, DirectionText), Q.First().Value == MarkerValue);
+					}
+					Q.PopFront();
+					TestCounts(*FString::Printf(TEXT("Test Add Constructors - %s%s CountsEqual"), Message, DirectionText), NumVoid, NumCopy, NumMove, NumDestruct);
+					Clear();
+				};
+
+				if (bAddBack) Q.Add(CounterA); else Q.AddFront(CounterA);
+				TestDirCounts(TEXT("Copy Add"), 0, 1, 0, 1);
+				if (bAddBack) Q.Add_GetRef(CounterA); else Q.AddFront_GetRef(CounterA);
+				TestDirCounts(TEXT("Copy GetRef Add"), 0, 1, 0, 1);
+				if (bAddBack) Q.Add(MoveTemp(CounterA)); else Q.AddFront(MoveTemp(CounterA));
+				TestDirCounts(TEXT("Move Add"), 0, 0, 1, 1);
+				if (bAddBack) Q.Add_GetRef(MoveTemp(CounterA)); else Q.AddFront_GetRef(MoveTemp(CounterA));
+				TestDirCounts(TEXT("Move GetRef Add"), 0, 0, 1, 1);
+				if (bAddBack) Q.Emplace(MarkerValue); else Q.EmplaceFront(MarkerValue);
+				TestDirCounts(TEXT("Emplace"), 1, 0, 0, 1);
+				if (bAddBack) Q.Emplace_GetRef(MarkerValue); else Q.EmplaceFront_GetRef(MarkerValue);
+				TestDirCounts(TEXT("GetRef Emplace"), 1, 0, 0, 1);
+				if (bAddBack) Q.AddUninitialized(); else Q.AddFrontUninitialized();
+				TestDirCounts(TEXT("Uninitialized Add"), 0, 0, 0, 1, false);
+				if (bAddBack) Q.AddUninitialized_GetRef(); else Q.AddFrontUninitialized_GetRef();
+				TestDirCounts(TEXT("Uninitialized GetRef Add"), 0, 0, 0, 1, false);
+			}
+		}
+
+		TestShiftIndex<uint32>();
+		TestShiftIndex<Counter>();
+
+		// Test RemoveAt
+		{
+			{
+				TRingBuffer<uint32> Q{ 0,1,2,3,4,5,6,7 };
+				Q.RemoveAt(2);
+				TestEqual(TEXT("Test RemoveAt Front Closest"), TRingBuffer<uint32>({ 0,1,3,4,5,6,7 }), Q);
+			}
+			{
+				TRingBuffer<uint32> Q{ 0,1,2,3,4,5,6,7 };
+				Q.RemoveAt(5);
+				TestEqual(TEXT("Test RemoveAt Back Closest"), TRingBuffer<uint32>({ 0,1,2,3,4,6,7 }), Q);
+			}
+			{
+				TRingBuffer<uint32> Q{ 0,1,2,3,4,5,6,7 };
+				int32 Offset = 4;
+				Q.Front += Offset;
+				Q.AfterBack += Offset;
+				//Now equal to: TRingBuffer<uint32> Q{ 4,5,6,7,0,1,2,3 };
+				Q.RemoveAt(2);
+				TestEqual(TEXT("Test RemoveAt Front Closest With Offset"), TRingBuffer<uint32>({ 4,5,7,0,1,2,3 }), Q);
+			}
+			{
+				TRingBuffer<uint32> Q{ 0,1,2,3,4,5,6,7 };
+				int32 Offset = 4;
+				Q.Front += Offset;
+				Q.AfterBack += Offset;
+				//Now equal to: TRingBuffer<uint32> Q{ 4,5,6,7,0,1,2,3 };
+				Q.RemoveAt(5);
+				TestEqual(TEXT("Test RemoveAt Back Closest With Offset"), TRingBuffer<uint32>({ 4,5,6,7,0,2,3 }), Q);
+			}
+		}
+
+		// Test Iteration
+		{
+			{
+				TRingBuffer<uint32> Q{ 0,1,2,3,4,5,6,7 };
+				uint32 Counter = 0;
+				for (uint32 Value : Q)
+				{
+					TestEqual(TEXT("Test Iteration - Value"), Counter++, Value);
+				}
+				TestEqual(TEXT("Test Iteration - Num"), Counter, 8);
+			}
+			{
+				TRingBuffer<uint32> Q{ 4,5,6,7,0,1,2,3 };
+				int32 Offset = 4;
+				Q.Front += Offset;
+				Q.AfterBack += Offset;
+				// Now equal to 0,1,2,3,4,5,6,7
+				uint32 Counter = 0;
+				for (uint32 Value : Q)
+				{
+					TestEqual(TEXT("Test Iteration with Offset - Value"), Counter++, Value);
+				}
+				TestEqual(TEXT("Test Iteration with Offset  - Num"), Counter, 8);
+			}
+		}
+
+		// Test ConvertPointerToIndex
+		{
+			{
+				TRingBuffer<uint32> Q{ 4,5,6,7,0,1,2,3 };
+				int32 Offset = 4;
+				Q.Front += Offset;
+				Q.AfterBack += Offset;
+				// Now equal to 0,1,2,3,4,5,6,7
+				TestEqual(TEXT("Test ConvertPointerToIndex - before array"), Q.ConvertPointerToIndex(&Q[0] - 100), INDEX_NONE);
+				TestEqual(TEXT("Test ConvertPointerToIndex - after array"), Q.ConvertPointerToIndex(&Q[0] + 100), INDEX_NONE);
+				for (int32 It = 0; It < 8; ++It)
+				{
+					TestEqual(TEXT("Test ConvertPointerToIndex - Values"), Q.ConvertPointerToIndex(&Q[It]), It);
+				}
+			}
+
+			{
+				TRingBuffer<uint32> Q(16);
+				for (int32 It = 7; It >= 0; --It)
+				{
+					Q.AddFront(It);
+				}
+				Q.Pop();
+				// 8 Invalids, followed by 0,1,2,3,4,5,6, followed by Invalid
+				for (int32 It = 0; It < 7; ++It)
+				{
+					TestEqual(TEXT("Test ConvertPointerToIndex - Cap - Values"), Q.ConvertPointerToIndex(&Q[It]), It);
+				}
+				TestEqual(TEXT("Test ConvertPointerToIndex - Cap - After End"), Q.ConvertPointerToIndex(&Q[6] + 1), INDEX_NONE);
+				TestEqual(TEXT("Test ConvertPointerToIndex - Cap - Before Start"), Q.ConvertPointerToIndex(&Q[0] - 1), INDEX_NONE);
+			}
+		}
+
+		// Test that setting Front to its maximum value and then popping the maximum number of elements does not break the contract that Front < capacity in StorageModulo space
+		{
+			TRingBuffer<uint32> Q(8);
+			Q.AddFront(0);
+			for (uint32 It = 1; It < 8; ++It)
+			{
+				Q.Add(It);
+			}
+			TestTrue(TEXT("Test Front<Capacity - Setup"), (Q.Front & Q.IndexMask) == Q.IndexMask && Q.Num() == Q.Max());
+			Q.PopFront(8);
+			TestTrue(TEXT("Test Front<Capacity - Contract is true"), static_cast<uint32>(Q.Front) < static_cast<uint32>(Q.Max()));
+		}
+
+		// Test IsValidIndex
+		{
+			TRingBuffer<uint32> Q({ 0,1,2,3,4 });
+			for (int32 It = 0; It < Q.Num(); ++It)
+			{
+				TestEqual(TEXT("IsValidIndex - InRange"), Q.IsValidIndex(It), true);
+			}
+			TestEqual(TEXT("IsValidIndex - Negative"), Q.IsValidIndex(-1), false);
+			TestEqual(TEXT("IsValidIndex - Num()"), Q.IsValidIndex(Q.Num() + 1), false);
+			TestEqual(TEXT("IsValidIndex - Capacity"), Q.IsValidIndex(Q.Max()), false);
+			TestEqual(TEXT("IsValidIndex - Capacity + 1"), Q.IsValidIndex(Q.Max()+1), false);
+		}
+
+		// Test Compact
+		{
+			{
+				TRingBuffer<uint32> QEmpty;
+				TestEqual(TEXT("Compact - Empty zero capacity"), QEmpty.Compact().Num(), 0);
+				QEmpty.Add(1);
+				QEmpty.PopFront();
+				TestEqual(TEXT("Compact - Empty non-zero capacity"), QEmpty.Compact().Num(), 0);
+			}
+			{
+				TArrayView<uint32> View;
+				TRingBuffer<uint32> Q(8);
+				Q.AddFront(37);
+				View = Q.Compact();
+				TestTrue(TEXT("Compact - Front at end"), ArrayViewsEqual(View, TArrayView<const uint32>({ 37 })));
+			}
+			{
+				TArrayView<uint32> View;
+				TRingBuffer<uint32> Q(8);
+				for (TRingBuffer<uint32>::IndexType It = 0; It < 6; ++It)
+				{
+					Q.Add(It);
+				}
+				Q.PopFront();
+				TRingBuffer<uint32>::StorageModuloType SavedFront = Q.Front;
+				TestTrue(TEXT("Compact - Front in middle - setup"), SavedFront > 0);
+				View = Q.Compact();
+				TestTrue(TEXT("Compact - Front in middle - values"), ArrayViewsEqual(View, TArrayView<const uint32>({ 1,2,3,4,5 })));
+				TestTrue(TEXT("Compact - Front in middle - no reallocate"), Q.Front == SavedFront);
+			}
+			{
+				TArrayView<uint32> View;
+				TRingBuffer<uint32> Q(8);
+				for (TRingBuffer<uint32>::IndexType It = 1; It < 8; ++It)
+				{
+					Q.Add(It);
+				}
+				Q.AddFront(0);
+				TestTrue(TEXT("Compact - Full array front at end - setup"), (Q.Front & Q.IndexMask) == 7);
+				View = Q.Compact();
+				TestTrue(TEXT("Compact - Full array front at end - values"), ArrayViewsEqual(View, TArrayView<const uint32>({0,1,2,3,4,5,6,7 })));
+				TestTrue(TEXT("Compact - Full array front at end - reallocated"), Q.Front == 0);
+			}
+			{
+				TArrayView<uint32> View;
+				TRingBuffer<uint32> Q(8);
+				for (TRingBuffer<uint32>::IndexType It = 0; It < 8; ++It)
+				{
+					Q.Add(It);
+				}
+				uint32* SavedData = Q.AllocationData;
+				TestTrue(TEXT("Compact - Full array front at start - setup"), Q.Front == 0);
+				View = Q.Compact();
+				TestTrue(TEXT("Compact - Full array front at start - values"), ArrayViewsEqual(View, TArrayView<const uint32>({ 0,1,2,3,4,5,6,7 })));
+				TestTrue(TEXT("Compact - Full array front at start - no reallocate"), Q.AllocationData == SavedData);
+			}
+		}
+
+		// Test Remove
+		{
+			Counter Value;
+			{
+				TRingBuffer<Counter> Q;
+				Value.Value = 2;
+				Counter::Clear();
+				TestEqual(TEXT("Remove - empty"), Q.Remove(Value), 0);
+				TestEqual(TEXT("Remove - empty - destructor count"), Counter::NumDestruct, 0);
+			}
+			{
+				TRingBuffer<Counter> Q({ 0,1,2,3,4 });
+				Value.Value = 5;
+				Counter::Clear();
+				TestEqual(TEXT("Remove - no hits"), Q.Remove(Value), 0);
+				TestEqual(TEXT("Remove - no hits - destructor count"), Counter::NumDestruct, 0);
+				Q.Add(5);
+				TestTrue(TEXT("Remove - no hits - values"), Q == TRingBuffer<Counter>({ 0,1,2,3,4,5 }));
+			}
+			{
+				TRingBuffer<Counter> Q({ 1,2,3,4 });
+				Q.AddFront(0);
+				Value.Value = 0;
+				Counter::Clear();
+				TestEqual(TEXT("Remove - one element at front - num"), Q.Remove(Value), 1);
+				TestEqual(TEXT("Remove - one element at front - destructor count"), Counter::NumDestruct, 5);
+				Q.Add(5);
+				TestTrue(TEXT("Remove - one element at front - values"), Q == TRingBuffer<Counter>({ 1,2,3,4,5 }));
+			}
+			{
+				TRingBuffer<Counter> Q({ 0,1,2,3,4 });
+				Value.Value = 2;
+				Counter::Clear();
+				TestEqual(TEXT("Remove - one element in mid - num"), Q.Remove(Value), 1);
+				TestEqual(TEXT("Remove - one element in mid - destructor count"), Counter::NumDestruct, 3);
+				Q.Add(5);
+				TestTrue(TEXT("Remove - one element in mid - values"), Q == TRingBuffer<Counter>({ 0,1,3,4,5 }));
+			}
+			{
+				TRingBuffer<Counter> Q({ 1,2,3,4 });
+				Q.AddFront(0);
+				Value.Value = 2;
+				Counter::Clear();
+				TestEqual(TEXT("Remove - one element in mid - front at end"), Q.Remove(Value), 1);
+				TestEqual(TEXT("Remove - one element in mid - front at end - destructor count"), Counter::NumDestruct, 3);
+				Q.Add(5);
+				TestTrue(TEXT("Remove - one element in mid - front at end - values"), Q == TRingBuffer<Counter>({ 0,1,3,4,5 }));
+			}
+			{
+				TRingBuffer<Counter> Q({ 0,1,2,3,4 });
+				Value.Value = 4;
+				Counter::Clear();
+				TestEqual(TEXT("Remove - one element - element at end - num"), Q.Remove(Value), 1);
+				TestEqual(TEXT("Remove - one element - element at end - destructor count"), Counter::NumDestruct, 1);
+				Q.Add(5);
+				TestTrue(TEXT("Remove - one element - element at end - values"), Q == TRingBuffer<Counter>({ 0,1,2,3,5 }));
+			}
+			{
+				TRingBuffer<Counter> Q({ 1,2,3,4 });
+				Q.AddFront(4);
+				Value.Value = 4;
+				Counter::Clear();
+				TestEqual(TEXT("Remove - one element at front one at end - num"), Q.Remove(Value), 2);
+				TestEqual(TEXT("Remove - one element at front one at end - destructor count"), Counter::NumDestruct, 5);
+				Q.Add(5);
+				TestTrue(TEXT("Remove - one element at front one at end - values"), Q == TRingBuffer<Counter>({ 1,2,3,5 }));
+			}
+			{
+				TRingBuffer<Counter> Q({ 1,2,3,4 });
+				Q.AddFront(1);
+				Value.Value = 1;
+				Counter::Clear();
+				TestEqual(TEXT("Remove - two elements - front at end - num"), Q.Remove(Value), 2);
+				TestEqual(TEXT("Remove - two elements - front at end - destructor count"), Counter::NumDestruct, 5);
+				Q.Add(5);
+				TestTrue(TEXT("Remove - two elements - front at end - values"), Q == TRingBuffer<Counter>({ 2,3,4,5 }));
+			}
+		}
+
+		return true;
+	}
+
+	template <typename T>
+	void TestShiftIndex()
+	{
+		// Test shifts at specific points
+		{
+			{
+				TRingBuffer<T> Q{ 0,1,2,3,4,5,6,7 };
+				Q.ShiftIndexToFront(5);
+				TestEqual(TEXT("ShiftIndexToFront"), TRingBuffer<T>({ 5,0,1,2,3,4,6,7 }), Q);
+				Q.ShiftIndexToBack(3);
+				TestEqual(TEXT("ShiftIndexToBack"), TRingBuffer<T>({ 5,0,1,3,4,6,7,2 }), Q);
+			}
+
+			{
+				TRingBuffer<T> Q{ 0,1,2,3,4,5,6,7 };
+				int32 Offset = 4;
+				Q.Front += Offset;
+				Q.AfterBack += Offset;
+				//Now equal to: TRingBuffer<uint32> Q{ 4,5,6,7,0,1,2,3 };
+				Q.ShiftIndexToFront(5);
+				TestEqual(TEXT("ShiftIndexToFront With Offset"), TRingBuffer<T>({ 1,4,5,6,7,0,2,3 }), Q);
+				Q.ShiftIndexToBack(3);
+				TestEqual(TEXT("ShiftIndexToBack With Offset"), TRingBuffer<T>({ 1,4,5,7,0,2,3,6 }), Q);
+			}
+
+			{
+				TRingBuffer<T> Q{ 0,1,2,3,4,5,6,7,8 };
+				TestEqual(TEXT("ShiftIndexToFront Cap - Capacity"), Q.Max(), 16);
+				Q.ShiftIndexToFront(5);
+				TestEqual(TEXT("ShiftIndexToFront Cap"), TRingBuffer<T>({ 5,0,1,2,3,4,6,7,8 }), Q);
+				Q.ShiftIndexToBack(3);
+				TestEqual(TEXT("ShiftIndexToBack Cap"), TRingBuffer<T>({ 5,0,1,3,4,6,7,8,2 }), Q);
+			}
+
+			{
+				TRingBuffer<T> Q(16);
+				for (int32 It = 7; It >= 0; --It)
+				{
+					Q.AddFront(It);
+				}
+				Q.Pop();
+				// 8 Invalids, followed by 0,1,2,3,4,5,6, followed by Invalid
+				Q.ShiftIndexToFront(5);
+				TestEqual(TEXT("ShiftIndexToFront Cap With Offset"), TRingBuffer<T>({ 5,0,1,2,3,4,6 }), Q);
+				Q.ShiftIndexToBack(3);
+				TestEqual(TEXT("ShiftIndexToBack Cap With Offset"), TRingBuffer<T>({ 5,0,1,3,4,6,2 }), Q);
+			}
+
+			{
+				TRingBuffer<T> Q(16);
+				for (int32 It = 7; It >= 0; --It)
+				{
+					Q.AddFront(It);
+				}
+				Q.Add(8);
+				// 8, (AfterBack), followed by 7 Invalids, followed by (Start) 0,1,2,3,4,5,6,7
+				Q.ShiftIndexToFront(8);
+				TestEqual(TEXT("ShiftIndexToFront Cap With Wrapped"), TRingBuffer<T>({ 8,0,1,2,3,4,5,6,7 }), Q);
+				Q.ShiftIndexToBack(0);
+				TestEqual(TEXT("ShiftIndexToBack Cap With Wrapped"), TRingBuffer<T>({ 0,1,2,3,4,5,6,7,8 }), Q);
+			}
+		}
+
+		// Test ShiftIndex of each possible index
+		{
+			int32 Count = 8;
+			for (int32 It = 0; It < Count; ++It)
+			{
+				TRingBuffer<T> Q({ 0,1,2,3,4,5,6,7 });
+				Q.ShiftIndexToBack(It);
+				int32 CheckIndex = 0;
+				for (; CheckIndex < It; ++CheckIndex)
+				{
+					TestEqual(*FString::Printf(TEXT("ShiftIndexToBack Entire Array Values (%d,%d)"), It, CheckIndex), CheckIndex, Q[CheckIndex]);
+				}
+				for (; CheckIndex < Count - 1; ++CheckIndex)
+				{
+					TestEqual(*FString::Printf(TEXT("ShiftIndexToBack Entire Array Values (%d,%d)"), It, CheckIndex), CheckIndex + 1, Q[CheckIndex]);
+				}
+				TestEqual(*FString::Printf(TEXT("ShiftIndexToBack Entire Array Values (%d,%d)"), It, Count - 1), It, Q[Count - 1]);
+			}
+			for (int32 It = 0; It < Count; ++It)
+			{
+				TRingBuffer<T> Q({ 0,1,2,3,4,5,6,7 });
+				Q.ShiftIndexToFront(It);
+
+				TestEqual(*FString::Printf(TEXT("ShiftIndexToFront Entire Array Values (%d,%d)"), It, 0), It, Q[0]);
+				int32 CheckIndex = 1;
+				for (; CheckIndex <= It; ++CheckIndex)
+				{
+					TestEqual(*FString::Printf(TEXT("ShiftIndexToFront Entire Array Values (%d,%d)"), It, CheckIndex), CheckIndex - 1, Q[CheckIndex]);
+				}
+				for (; CheckIndex < Count; ++CheckIndex)
+				{
+					TestEqual(*FString::Printf(TEXT("ShiftIndexToFront Entire Array Values (%d,%d)"), It, CheckIndex), CheckIndex, Q[CheckIndex]);
+				}
+			}
+		}
+	}
+
+	template <typename T, typename U>
+	static bool ArrayViewsEqual(const TArrayView<T>& A, const TArrayView<U>& B)
+	{
+		int32 Num = A.Num();
+		if (Num != B.Num())
+		{
+			return false;
+		}
+		for (int It = 0; It < Num; ++It)
+		{
+			if (A[It] != B[It])
+			{
+				return false;
+			}
+		}
+		return true;
+	}
+};
+
+int FRingBufferTest::Counter::NumVoid = 0;
+int FRingBufferTest::Counter::NumCopy = 0;
+int FRingBufferTest::Counter::NumMove = 0;
+int FRingBufferTest::Counter::NumDestruct = 0;
+
+IMPLEMENT_CUSTOM_SIMPLE_AUTOMATION_TEST(FRingBufferTestSubClass, FRingBufferTest, "System.Core.Containers.RingBuffer", EAutomationTestFlags::ApplicationContextMask | EAutomationTestFlags::EngineFilter);
+bool FRingBufferTestSubClass::RunTest(const FString& Parameters)
+{
+	return FRingBufferTest::RunTest(Parameters);
+}
+
+
+#endif // #if WITH_DEV_AUTOMATION_TESTS
\ No newline at end of file
diff --git a/Engine/Source/Runtime/Core/Private/Tests/HAL/FileManagerGenericTest.cpp b/Engine/Source/Runtime/Core/Private/Tests/HAL/FileManagerGenericTest.cpp
new file mode 100644
--- /dev/null
+++ b/Engine/Source/Runtime/Core/Private/Tests/HAL/FileManagerGenericTest.cpp
@@ -0,0 +1,305 @@
+// Copyright Epic Games, Inc. All Rights Reserved.
+
+#include "Misc/AutomationTest.h"
+
+#if WITH_DEV_AUTOMATION_TESTS
+
+#include "HAL/FileManagerGeneric.h"
+
+#include "Containers/UnrealString.h"
+#include "HAL/FileManager.h"
+#include "HAL/PlatformFileManager.h"
+#include "Misc/Paths.h"
+#include "Misc/StringBuilder.h"
+
+class FArchiveFileReaderGenericTest : public FAutomationTestBase
+{
+public:
+	using FAutomationTestBase::FAutomationTestBase;
+	~FArchiveFileReaderGenericTest();
+	void TestInternalPrecache();
+
+private:
+	bool CreateTestFile();
+	void TestBytesValid(const TCHAR* What, uint8* Data, int64 NumBytes, int64 FileOffset);
+	int32 GetExpectedValue(int64 IntOffset);
+	void SetPosAndBuffer(TUniquePtr<FArchiveFileReaderGeneric>& Reader, int64 Pos, int64 BufferBase, int64 BufferSize);
+
+	FString TestFileName;
+	TUniquePtr<IFileHandle> ReadHandle;
+	int64 FileSize;
+};
+
+FArchiveFileReaderGenericTest::~FArchiveFileReaderGenericTest()
+{
+	ReadHandle.Reset();
+	if (!TestFileName.IsEmpty())
+	{
+		IFileManager& FileManager = IFileManager::Get();
+		FileManager.Delete(*TestFileName, false /* bRequireExists */, true /* bEvenReadOnly */, true /* Quiet */);
+	}
+	TestFileName.Empty();
+}
+
+void FArchiveFileReaderGenericTest::TestInternalPrecache()
+{
+	if (!CreateTestFile())
+	{
+		return;
+	}
+
+	uint32 BufferSize = 1024;
+	TArray<uint8> UnusedBytes;
+	UnusedBytes.SetNumUninitialized(BufferSize);
+	TUniquePtr<FArchiveFileReaderGeneric> Reader(new FArchiveFileReaderGeneric(ReadHandle.Release(), *TestFileName, FileSize, BufferSize));
+	TestEqual(TEXT("Initial Pos should be 0"), Reader->Pos, 0LL);
+	TestEqual(TEXT("Size should be what was passed in"), Reader->Size, FileSize);
+	TestEqual(TEXT("BufferSize should be what was passed in"), Reader->BufferSize, BufferSize);
+
+	// Vanilla InternalPrecache at start of file
+	bool Result = Reader->InternalPrecache(0, BufferSize);
+	TestEqual(TEXT("Vanilla0 - InternalPrecache should succeed"), Result, true);
+	TestEqual(TEXT("Vanilla0 - InternalPrecache should not alter Pos"), Reader->Pos, 0LL);
+	TestEqual(TEXT("Vanilla0 - BufferBase should have be set to PrecacheOffset aka Pos"), Reader->BufferBase, 0LL);
+	TestEqual(TEXT("Vanilla0 - Bytes precached should be set to BufferSize unless it runs out of room"), Reader->BufferArray.Num(), BufferSize);
+	TestBytesValid(TEXT("Vanilla0 - BufferArray should be expected bytes"), Reader->BufferArray.GetData(), Reader->BufferArray.Num(), Reader->BufferBase);
+
+	// InternalPrecache at PrecacheOffset != Pos is ignored, and returns false if the first byte at PrecacheOffset is not cached
+	Result = Reader->InternalPrecache(BufferSize  + BufferSize, BufferSize);
+	TestEqual(TEXT("PrecacheNotAtPos - Should fail"), Result, false);
+	TestEqual(TEXT("PrecacheNotAtPos - BufferBase should not be altered"), Reader->BufferBase, 0LL);
+	TestEqual(TEXT("PrecacheNotAtPos - Buffe should not be altered"), Reader->BufferArray.Num(), BufferSize);
+
+	// InternalPrecache at PrecacheOffset != Pos is ignored, but returns true if the first byte at PrecacheOffset is cached
+	Result = Reader->InternalPrecache(BufferSize - 1, BufferSize);
+	TestEqual(TEXT("PrecacheNotAtPos - Should succeed since the first byte at PrecacheOffset is buffered"), Result, true);
+	TestEqual(TEXT("PrecacheNotAtPos - BufferBase should not be altered"), Reader->BufferBase, 0LL);
+	TestEqual(TEXT("PrecacheNotAtPos - Buffe should not be altered"), Reader->BufferArray.Num(), BufferSize);
+
+	// InternalPrecache partway through the buffer should allocate the full buffer/not allocate anything, depending on bPrecacheAsSoonAsPossible
+	int64 PosStart = BufferSize / 2;
+	SetPosAndBuffer(Reader, PosStart, 0, BufferSize);
+	Result = Reader->InternalPrecache(PosStart, BufferSize);
+	TestEqual(TEXT("PrecachePartwayPosMoreThanBuffer - Should succeed"), Result, true);
+	if (FArchiveFileReaderGeneric::bPrecacheAsSoonAsPossible)
+	{
+		TestEqual(TEXT("PrecachePartwayPosMoreThanBuffer - BufferBase should be updated"), Reader->BufferBase, PosStart);
+	}
+	else
+	{
+		TestEqual(TEXT("PrecachePartwayPosMoreThanBuffer - BufferBase should not be updated"), Reader->BufferBase, 0LL);
+	}
+	TestEqual(TEXT("PrecachePartwayPosMoreThanBuffer - BufferCount should be set to BufferSize"), Reader->BufferArray.Num(), BufferSize);
+	TestBytesValid(TEXT("PrecachePartwayPosMoreThanBuffer - BufferBytes should match BufferBase"), Reader->BufferArray.GetData(), Reader->BufferArray.Num(), Reader->BufferBase);
+
+	// InternalPrecache partway through the buffer should allocate the full buffer/not allocate anything, and should handle the case of BufferArray's allocation being smaller than BufferSize
+	SetPosAndBuffer(Reader, 3*BufferSize/4, BufferSize/2, BufferSize/2);
+	Reader->BufferArray.Shrink();
+	Result = Reader->InternalPrecache(3*BufferSize/4, BufferSize);
+	TestEqual(TEXT("PrecachePartwayPosMoreThanBufferReallocation - Should succeed"), Result, true);
+	if (FArchiveFileReaderGeneric::bPrecacheAsSoonAsPossible)
+	{
+		TestEqual(TEXT("PrecachePartwayPosMoreThanBufferReallocation - BufferBase should be updated"), Reader->BufferBase, 3*BufferSize/4);
+		TestEqual(TEXT("PrecachePartwayPosMoreThanBufferReallocation - BufferCount should be set to BufferSize"), Reader->BufferArray.Num(), BufferSize);
+	}
+	else
+	{
+		TestEqual(TEXT("PrecachePartwayPosMoreThanBufferReallocation - BufferBase should not be updated"), Reader->BufferBase, BufferSize/2);
+		TestEqual(TEXT("PrecachePartwayPosMoreThanBufferReallocation - BufferCount should not be updated"), Reader->BufferArray.Num(), BufferSize/2);
+	}
+	TestBytesValid(TEXT("PrecachePartwayPosMoreThanBufferReallocation - BufferBytes should match BufferBase"), Reader->BufferArray.GetData(), Reader->BufferArray.Num(), Reader->BufferBase);
+
+	// InternalPrecache partway through the buffer, when the requested size is within the remaining buffer, should not change the buffer
+	PosStart = BufferSize / 2;
+	SetPosAndBuffer(Reader, PosStart, 0, BufferSize);
+	Result = Reader->InternalPrecache(PosStart, BufferSize / 4);
+	TestEqual(TEXT("PrecachePartwayPosLessThanBuffer - Should succeed"), Result, true);
+	TestEqual(TEXT("PrecachePartwayPosLessThanBuffer - BufferBase should not be updated"), Reader->BufferBase, 0LL);
+	TestEqual(TEXT("PrecachePartwayPosLessThanBuffer - BufferCount should be set to BufferSize"), Reader->BufferArray.Num(), BufferSize);
+	TestBytesValid(TEXT("PrecachePartwayPosLessThanBuffer - BufferBytes should match BufferBase"), Reader->BufferArray.GetData(), Reader->BufferArray.Num(), Reader->BufferBase);
+
+	// InternalPrecache at the end of the buffer should allocate the full buffer
+	PosStart = BufferSize;
+	SetPosAndBuffer(Reader, PosStart, 0, BufferSize);
+	Result = Reader->InternalPrecache(PosStart, BufferSize);
+	TestEqual(TEXT("PrecachePosEndOfBuffer - Should succeed"), Result, true);
+	TestEqual(TEXT("PrecachePosEndOfBuffer - BufferBase should be updated"), Reader->BufferBase, PosStart);
+	TestEqual(TEXT("PrecachePosEndOfBuffer - BufferCount should be set to BufferSize"), Reader->BufferArray.Num(), BufferSize);
+	TestBytesValid(TEXT("PrecachePosEndOfBuffer - BufferBytes should match BufferBase"), Reader->BufferArray.GetData(), Reader->BufferArray.Num(), Reader->BufferBase);
+
+	// InternalPrecache near the end of the file should clamp BufferSize to Size - Pos
+	PosStart = FileSize - 1;
+	SetPosAndBuffer(Reader, PosStart, 0, 0);
+	Result = Reader->InternalPrecache(PosStart, BufferSize);
+	TestEqual(TEXT("PrecacheNearEndOfFile - Should succeed"), Result, true);
+	TestEqual(TEXT("PrecacheNearEndOfFile - BufferBase should be updated"), Reader->BufferBase, PosStart);
+	TestEqual(TEXT("PrecacheNearEndOfFile - BufferCount should be set Min(BufferSize, Size - Pos)"), Reader->BufferArray.Num(), 1LL);
+	TestBytesValid(TEXT("PrecacheNearEndOfFile - BufferBytes should match BufferBase"), Reader->BufferArray.GetData(), Reader->BufferArray.Num(), Reader->BufferBase);
+
+	AddExpectedError(TEXT("ReadFile failed"));
+	// InternalPrecache at the end of the file should fail
+	SetPosAndBuffer(Reader, FileSize, 0, 16);
+	Result = Reader->InternalPrecache(FileSize, BufferSize);
+	TestEqual(TEXT("PrecacheEndOfFile - Should fail"), Result, false);
+	TestEqual(TEXT("PrecacheEndOfFile - BufferBase should not be updated"), Reader->BufferBase, 0LL);
+	TestEqual(TEXT("PrecacheEndOfFile - BufferCount should not be updated"), Reader->BufferArray.Num(), 16LL);
+
+	// If ReadLowLevel fails, InternalPrecache should return true if the first byte of PrecacheOffset is in the buffer
+	SetPosAndBuffer(Reader, 8, 0, 16);
+	Reader->SeekLowLevel(FileSize);
+	Result = Reader->InternalPrecache(8, 16);
+	TestEqual(TEXT("PrecacheReadFailsBytesRemain - Should succeed if bytes are left in the buffer"), Result, true);
+	if (FArchiveFileReaderGeneric::bPrecacheAsSoonAsPossible)
+	{
+		TestEqual(TEXT("PrecacheReadFailsBytesRemain - BufferBase should be updated"), Reader->BufferBase, 8LL);
+		TestEqual(TEXT("PrecacheReadFailsBytesRemain - BufferCount should be updated"), Reader->BufferArray.Num(), 8LL);
+	}
+	else
+	{
+		TestEqual(TEXT("PrecacheReadFailsBytesRemain - BufferBase should not be updated"), Reader->BufferBase, 0LL);
+		TestEqual(TEXT("PrecacheReadFailsBytesRemain - BufferCount should not be updated"), Reader->BufferArray.Num(), 16LL);
+	}
+
+	// If ReadLowLevel fails, InternalPrecache should return false if the first byte of PrecacheOffset is not in the buffer
+	SetPosAndBuffer(Reader, 16, 0, 16);
+	Reader->SeekLowLevel(FileSize);
+	Result = Reader->InternalPrecache(16, 16);
+	TestEqual(TEXT("PrecacheReadFails - Should fail if bytes are not left in the buffer"), Result, false);
+	TestEqual(TEXT("PrecacheReadFails - BufferCount should be updated to empty"), Reader->BufferArray.Num(), 0LL);
+
+	// TODO: Call ClearExpectedErrors to clear the AddExpectedError called above when it is implemented. Until then, make sure these failure-testing cases are the last
+	// cases tested in the test
+	//ClearExpectedErrors();
+}
+
+bool FArchiveFileReaderGenericTest::CreateTestFile()
+{
+	IFileManager& FileManager = IFileManager::Get();
+	IPlatformFile& PlatformFile = FPlatformFileManager::Get().GetPlatformFile();;
+	TestFileName = FPaths::CreateTempFilename(FPlatformProcess::UserTempDir(), TEXT("ArchiveFileReaderGenericTest"));
+	TUniquePtr<IFileHandle> WriteHandle(PlatformFile.OpenWrite(*TestFileName, false /* bAppend */, false /* bAllowRead */));
+	if (!WriteHandle)
+	{
+		AddError(FString::Printf(TEXT("Could not create test file %s."), *TestFileName));
+		return false;
+	}
+	TArray<int32> Builder;
+	int32 NumInts = 1024 * 128 + 17; // Make sure the file size is not a multiple of buffersize; we need to test that behavior
+	Builder.SetNumUninitialized(NumInts);
+	int32* Data = Builder.GetData();
+	for (int n = 0; n < Builder.Num(); ++n)
+	{
+		Data[n] = GetExpectedValue(n);
+	}
+	FileSize = NumInts * sizeof(int32);
+	WriteHandle->Write(reinterpret_cast<uint8*>(Data), FileSize);
+	WriteHandle.Reset();
+
+	ReadHandle.Reset(PlatformFile.OpenRead(*TestFileName, false /* bAllowWrite */));
+	if (!ReadHandle)
+	{
+		AddError(FString::Printf(TEXT("Could not open test file %s."), *TestFileName));
+		return false;
+	}
+	if (ReadHandle->Size() != FileSize)
+	{
+		AddError(FString::Printf(TEXT("Received incorrect file size from test file %s. Expected = %lld. Actual = %lld."), *TestFileName, FileSize, ReadHandle->Size()));
+		return false;
+	}
+
+	return true;
+}
+
+void FArchiveFileReaderGenericTest::TestBytesValid(const TCHAR* What, uint8* InData, int64 InNumBytes, int64 InFileOffset)
+{
+	uint8* Data = InData;
+	int64 NumBytes = InNumBytes;
+	int64 FileOffset = InFileOffset;
+
+	int64 StartByte = FileOffset % sizeof(int32);
+	if (StartByte != 0)
+	{
+		int64 FloorIndex = (FileOffset - StartByte)/sizeof(int32);
+		int64 ExpectedValue = GetExpectedValue(FloorIndex);
+		for (int n = 0; n < NumBytes && n < sizeof(int32) - StartByte; ++n)
+		{
+			if (Data[n] != reinterpret_cast<uint8*>(&ExpectedValue)[StartByte + n])
+			{
+				TestEqual(What, Data[n], reinterpret_cast<uint8*>(&ExpectedValue)[StartByte + n]);
+				return;
+			}
+		}
+		if (NumBytes < static_cast<int64>(sizeof(int32) - StartByte))
+		{
+			return;
+		}
+		int32 BytesRead = sizeof(int32) - StartByte;
+		Data += BytesRead;
+		FileOffset += BytesRead;
+		NumBytes -= BytesRead;
+	}
+	int32* IntData = reinterpret_cast<int32*>(Data);
+	int64 IntOffset = FileOffset / sizeof(int32);
+	int64 NumInts = NumBytes / sizeof(int32);
+	for (int n = 0; n < NumInts; ++n)
+	{
+		if (IntData[n] != GetExpectedValue(IntOffset + n))
+		{
+			TestEqual(What, IntData[n], GetExpectedValue(IntOffset + n));
+			return;
+		}
+	}
+
+	int64 NumEndBytes = NumBytes - NumInts * sizeof(int32);
+	if (NumEndBytes != 0)
+	{
+		uint8* DataStart = Data + NumInts * sizeof(int32);
+		int32 ExpectedValue = GetExpectedValue(IntOffset + NumInts);
+		for (int n = 0; n < NumEndBytes; ++n)
+		{
+			if (DataStart[n] != reinterpret_cast<uint8*>(&ExpectedValue)[n])
+			{
+				TestEqual(What, DataStart[n], reinterpret_cast<uint8*>(&ExpectedValue)[n]);
+				return;
+			}
+		}
+	}
+}
+
+void FArchiveFileReaderGenericTest::SetPosAndBuffer(TUniquePtr<FArchiveFileReaderGeneric>& Reader, int64 Pos, int64 BufferBase, int64 BufferSize)
+{
+	Reader->Pos = Pos;
+	Reader->BufferBase = BufferBase;
+	Reader->BufferArray.SetNumUninitialized(BufferSize, false /* bAllowShrink */);
+	if (BufferSize)
+	{
+		Reader->SeekLowLevel(BufferBase);
+		int64 BytesRead;
+		Reader->ReadLowLevel(Reader->BufferArray.GetData(), BufferSize, BytesRead);
+		TestEqual(TEXT("SetPosAndBuffer - ReadLowLevel read the requested bytes"), BytesRead, BufferSize);
+	}
+	// See the contract for Buffer window and LowLevel Pos in the variable comment on BufferArray
+	bool bPosWithinBuffer = BufferBase <= Pos && Pos < BufferBase + BufferSize;
+	if (!bPosWithinBuffer)
+	{
+		Reader->SeekLowLevel(Pos);
+	}
+	TestEqual(TEXT("SetPosAndBuffer - Pos set correctly"), Reader->Pos, Pos);
+	TestEqual(TEXT("SetPosAndBuffer - BufferBase set correctly"), Reader->BufferBase, BufferBase);
+	TestEqual(TEXT("SetPosAndBuffer - BufferSize set correctly"), Reader->BufferArray.Num(), BufferSize);
+	TestBytesValid(TEXT("SetPosAndBuffer - BufferBytes should match BufferBase"), Reader->BufferArray.GetData(), Reader->BufferArray.Num(), Reader->BufferBase);
+}
+
+int32 FArchiveFileReaderGenericTest::GetExpectedValue(int64 IntOffset)
+{
+	return 0xbe000000 + static_cast<int32>(IntOffset);
+}
+
+IMPLEMENT_CUSTOM_SIMPLE_AUTOMATION_TEST(FArchiveFileReaderGenericTestRunner, FArchiveFileReaderGenericTest, "System.Core.HAL.FArchiveFileReaderGeneric", EAutomationTestFlags::ApplicationContextMask | EAutomationTestFlags::EngineFilter)
+bool FArchiveFileReaderGenericTestRunner::RunTest(const FString& Parameters)
+{
+	TestInternalPrecache();
+	return !HasAnyErrors();
+}
+
+#endif
diff --git a/Engine/Source/Runtime/Core/Private/Tests/HAL/FileManagerTest.cpp b/Engine/Source/Runtime/Core/Private/Tests/HAL/FileManagerTest.cpp
new file mode 100644
--- /dev/null
+++ b/Engine/Source/Runtime/Core/Private/Tests/HAL/FileManagerTest.cpp
@@ -0,0 +1,16 @@
+// Copyright Epic Games, Inc. All Rights Reserved.
+
+#include "HAL/FileManager.h"
+
+#include "Misc/AutomationTest.h"
+
+#if WITH_DEV_AUTOMATION_TESTS
+
+IMPLEMENT_SIMPLE_AUTOMATION_TEST(FFileManagerReaderTest, "System.Core.HAL.FileManager.Reader", EAutomationTestFlags::ApplicationContextMask | EAutomationTestFlags::EngineFilter)
+
+bool FFileManagerReaderTest::RunTest(const FString& Parameters)
+{
+	return true;
+}
+
+#endif
\ No newline at end of file
diff --git a/Engine/Source/Runtime/Core/Private/Tests/Internationalization/ArchiveTest.cpp b/Engine/Source/Runtime/Core/Private/Tests/Internationalization/ArchiveTest.cpp
new file mode 100644
--- /dev/null
+++ b/Engine/Source/Runtime/Core/Private/Tests/Internationalization/ArchiveTest.cpp
@@ -0,0 +1,220 @@
+// Copyright Epic Games, Inc. All Rights Reserved.
+
+#include "CoreTypes.h"
+#include "Containers/UnrealString.h"
+#include "Misc/AutomationTest.h"
+#include "Internationalization/InternationalizationManifest.h"
+#include "Internationalization/InternationalizationArchive.h"
+#include "Internationalization/InternationalizationMetadata.h"
+
+#if WITH_DEV_AUTOMATION_TESTS
+
+IMPLEMENT_SIMPLE_AUTOMATION_TEST(FArchiveTest, "System.Core.Misc.Internationalization Archive", EAutomationTestFlags::EditorContext | EAutomationTestFlags::ClientContext | EAutomationTestFlags::SmokeFilter)
+
+
+bool FArchiveTest::RunTest( const FString& Parameters )
+{
+	// Key metadata
+	TSharedPtr< FLocMetadataObject > KeyMetadataA = MakeShareable( new FLocMetadataObject );
+	TSharedPtr< FLocMetadataObject > KeyMetadataB = MakeShareable( new FLocMetadataObject );
+
+	// Source metadata
+	TSharedPtr< FLocMetadataObject > SourceMetadataA = MakeShareable( new FLocMetadataObject );
+	TSharedPtr< FLocMetadataObject > SourceMetadataB = MakeShareable( new FLocMetadataObject );
+
+
+	// Setup KeyMetadataA
+	KeyMetadataA->SetStringField( TEXT("Gender"			),	TEXT("Masculine")	);
+	KeyMetadataA->SetStringField( TEXT("Plurality"		),	TEXT("Singular")	);
+	KeyMetadataA->SetStringField( TEXT("TargetGender"	),	TEXT("Masculine")	);
+	KeyMetadataA->SetStringField( TEXT("TargetPlurality"),	TEXT("Singular")	);
+
+	// Setup KeyMetadataB
+	KeyMetadataB->SetStringField( TEXT("Gender"			),	TEXT("Masculine")	);
+	KeyMetadataB->SetStringField( TEXT("Plurality"		),	TEXT("Singular")	);
+	KeyMetadataB->SetStringField( TEXT("TargetGender"	),	TEXT("Feminine")	);
+	KeyMetadataB->SetStringField( TEXT("TargetPlurality"),	TEXT("Singular")	);
+
+	// Setup source metadata
+	SourceMetadataA->SetBoolField( TEXT("*IsMature"), false );
+	SourceMetadataB->SetBoolField( TEXT("*IsMature"), true );
+
+	// Setup source item
+	FLocItem SourceA( TEXT("TextA") );
+	SourceA.MetadataObj = MakeShareable( new FLocMetadataObject(*SourceMetadataA) );
+
+	FLocItem Translation = SourceA;
+	Translation.Text = TEXT("TranslatedTextA");
+
+	FString TestNamespace = TEXT("TestNamespace");
+	FString SourceAKey = TEXT("TextA");
+
+	// Test entry add
+	{
+		bool TestOptionalTrue = true;
+		bool TestOptionalFalse = false;
+
+		// bIsOptional is not used as a key.  We ensure adding entries where the bIsOptional member is the only difference works as expected.
+		FInternationalizationArchive TestArchive;
+		TestArchive.AddEntry( TestNamespace, SourceAKey, SourceA, Translation, NULL, TestOptionalTrue );
+		// Add duplicate entry that differs in bIsOptional flag.  This add should report success because we already have an entry with matching 
+		//  namespace/source/keymetadata.  Differences in bIsOptional are not taken into consideration. 
+		bool bResult = TestArchive.AddEntry( TestNamespace, SourceAKey, SourceA, Translation, NULL, TestOptionalFalse );
+		TestTrue( TEXT("AddEntry result = true"), bResult);
+
+		// We should only have one entry in the archive
+		int32 EntryCount = 0;
+		for(auto Iter = TestArchive.GetEntriesBySourceTextIterator(); Iter; ++Iter, ++EntryCount);
+		TestEqual( TEXT("EntryCount == 1"), EntryCount, 1 );
+
+		// Make sure the original bIsOptional value is not overwritten by our second add.
+		TSharedPtr< FArchiveEntry > FoundEntry = TestArchive.FindEntryByKey( TestNamespace, SourceAKey, NULL );
+		if( !FoundEntry.IsValid() )
+		{
+			AddError(TEXT("FArchiveEntry could not find entry using FindEntryByKey."));
+		}
+		else
+		{
+			TestTrue( TEXT("FoundEntry->Namespace == Namespace"), FoundEntry->bIsOptional == TestOptionalTrue );
+		}
+	}
+
+	// Test lookup an entry
+	{
+		{
+			FInternationalizationArchive TestArchive;
+			TestArchive.AddEntry( TestNamespace, SourceAKey, SourceA, Translation, KeyMetadataA, false );
+
+			TSharedPtr< FArchiveEntry > FoundEntry = TestArchive.FindEntryByKey( TestNamespace, SourceAKey, KeyMetadataA );
+			if( !FoundEntry.IsValid() )
+			{
+				AddError(TEXT("FArchiveEntry could not find entry using FindEntryByKey."));
+			}
+			else
+			{
+				TestTrue( TEXT("FoundEntry->Namespace == Namespace"), FoundEntry->Namespace == TestNamespace );
+				TestTrue( TEXT("FoundEntry->Source == Source"), FoundEntry->Source == SourceA );
+				TestTrue( TEXT("FoundEntry->Translation == Translation"), FoundEntry->Translation == Translation );
+				if( FoundEntry->KeyMetadataObj == KeyMetadataA )
+				{
+					AddError(TEXT("FArchiveEntry KeyMetadataObj is not a unique object."));
+				}
+				TestTrue( TEXT("FoundEntry->KeyMetadataObj == KeyMetadataA"), *(FoundEntry->KeyMetadataObj) == *(KeyMetadataA) );
+			}
+
+			// Passing in a mismatched key metadata will fail to find the entry.  Any fallback logic is intended to happen at runtime
+			FoundEntry = TestArchive.FindEntryByKey( TestNamespace, SourceAKey, NULL );
+			TestInvalid( TEXT("!FoundEntry.IsValid()"), FoundEntry);
+
+			FoundEntry = TestArchive.FindEntryByKey( TestNamespace, SourceAKey, MakeShareable( new FLocMetadataObject() ) );
+			TestInvalid( TEXT("!FoundEntry.IsValid()"), FoundEntry);
+
+			FoundEntry = TestArchive.FindEntryByKey( TestNamespace, SourceAKey, KeyMetadataB );
+			TestInvalid( TEXT("!FoundEntry.IsValid()"), FoundEntry);
+		}
+
+		// Ensure we can properly lookup entries with non-null but empty key metadata
+		{
+			FInternationalizationArchive TestArchive;
+			TestArchive.AddEntry( TestNamespace, SourceAKey, SourceA, Translation, MakeShareable( new FLocMetadataObject() ), false );
+
+			TSharedPtr< FArchiveEntry > FoundEntry = TestArchive.FindEntryByKey( TestNamespace, SourceAKey, NULL );
+			if( !FoundEntry.IsValid() )
+			{
+				AddError(TEXT("FArchiveEntry could not find entry using FindEntryByKey."));
+			}
+			else
+			{
+				TestTrue( TEXT("FoundEntry->Namespace == Namespace"), FoundEntry->Namespace == TestNamespace );
+				TestTrue( TEXT("FoundEntry->Source == Source"), FoundEntry->Source == SourceA );
+			}
+		}
+
+		// Ensure we can properly lookup entries with null key metadata
+		{
+			FInternationalizationArchive TestArchive;
+			TestArchive.AddEntry( TestNamespace, SourceAKey, SourceA, Translation, NULL, false );
+
+			TSharedPtr< FArchiveEntry > FoundEntry = TestArchive.FindEntryByKey( TestNamespace, SourceAKey, NULL );
+			if( !FoundEntry.IsValid() )
+			{
+				AddError(TEXT("FArchiveEntry could not find entry using FindEntryByKey."));
+			}
+			else
+			{
+				TestTrue( TEXT("FoundEntry->Namespace == Namespace"), FoundEntry->Namespace == TestNamespace );
+				TestTrue( TEXT("FoundEntry->Source == Source"), FoundEntry->Source == SourceA );
+			}
+
+			// Test lookup with non-null but empty key metadata
+			FoundEntry = TestArchive.FindEntryByKey( TestNamespace, SourceAKey, MakeShareable( new FLocMetadataObject() ) );
+			if( !FoundEntry.IsValid() )
+			{
+				AddError(TEXT("FArchiveEntry could not find entry using FindEntryByKey."));
+			}
+			else
+			{
+				TestTrue( TEXT("FoundEntry->Namespace == Namespace"), FoundEntry->Namespace == TestNamespace );
+				TestTrue( TEXT("FoundEntry->Source == Source"), FoundEntry->Source == SourceA );
+			}
+		}
+
+		// Ensure lookup where source metadata has * prefixed entries work as expected.  Note: The source metadata object
+		//   supports a * prefix on metadata names which modifies the way we perform metadata comparison(ignores entry type and value, only name is checked)
+		{
+			FLocItem SourceCompare( TEXT("TextA") );
+			SourceCompare.MetadataObj = MakeShareable( new FLocMetadataObject() );
+			SourceCompare.MetadataObj->SetStringField( TEXT("*IsMature"), TEXT("") );
+
+			FInternationalizationArchive TestArchive;
+			// Added entry with String *IsMature entry
+			TestArchive.AddEntry( TestNamespace, SourceAKey, SourceCompare, Translation, KeyMetadataA, false );
+
+			// Finding entry using a source that has Boolean *IsMature
+			TSharedPtr< FArchiveEntry > FoundEntry = TestArchive.FindEntryByKey( TestNamespace, SourceAKey, KeyMetadataA );
+			if( !FoundEntry.IsValid() )
+			{
+				AddError(TEXT("FArchiveEntry could not find entry using FindEntryByKey."));
+			}
+			else
+			{
+				TestTrue( TEXT("FoundEntry->Namespace == Namespace"), FoundEntry->Namespace == TestNamespace );
+				TestTrue( TEXT("FoundEntry->Source == Source"), FoundEntry->Source == SourceA );
+			}
+
+			// Attempting to add an entry that only differs by a * prefexed source metadata entry value or type will result in success since
+			//  a 'matching' entry already exists in the archive.  We should, however, only have one entry in the archive
+			bool bResult = TestArchive.AddEntry( TestNamespace, SourceAKey, SourceA, Translation, KeyMetadataA, false );
+
+			TestTrue( TEXT("AddEntry result = true"), bResult);
+
+			// We should only have one entry in the archive
+			int32 EntryCount = 0;
+			for(auto Iter = TestArchive.GetEntriesBySourceTextIterator(); Iter; ++Iter, ++EntryCount);
+			TestEqual( TEXT("EntryCount == 1"), EntryCount, 1 );
+
+			// Check to see that the original type/value of the * prefixed entry did not get modified
+			FoundEntry = TestArchive.FindEntryByKey( TestNamespace, SourceAKey, KeyMetadataA );
+			if( !FoundEntry.IsValid() )
+			{
+				AddError(TEXT("FArchiveEntry could not find entry using FindEntryByKey."));
+			}
+			else
+			{
+				if( FoundEntry->Source.MetadataObj->HasTypedField< ELocMetadataType::String >( TEXT("*IsMature") ) )
+				{
+					TestTrue( TEXT("Metadata Type == String and Value == Empty string"), FoundEntry->Source.MetadataObj->GetStringField( TEXT("*IsMature") ) == TEXT("") );
+				}
+				else
+				{
+					AddError(TEXT("FArchiveEntry * prefixed metadata entry on source object was modified unexpectedly."));
+				}
+			}
+		}
+	}
+
+	return true;
+}
+
+
+#endif //WITH_DEV_AUTOMATION_TESTS
