LowLevelTests: Reverted the Core tests that were moved in 19498448 and 19571202

diff --git a/Engine/Source/Runtime/Core/Private/Tests/Internationalization/ManifestTest.cpp b/Engine/Source/Runtime/Core/Private/Tests/Internationalization/ManifestTest.cpp
new file mode 100644
--- /dev/null
+++ b/Engine/Source/Runtime/Core/Private/Tests/Internationalization/ManifestTest.cpp
@@ -0,0 +1,897 @@
+// Copyright Epic Games, Inc. All Rights Reserved.
+
+#include "CoreTypes.h"
+#include "Containers/UnrealString.h"
+#include "Misc/AutomationTest.h"
+#include "Internationalization/InternationalizationManifest.h"
+#include "Internationalization/InternationalizationMetadata.h"
+
+#if WITH_DEV_AUTOMATION_TESTS
+
+// Helper function to count the number of entries in a manifest
+int32 CountManifestEntries( const FInternationalizationManifest& Manifest )
+{
+	return Manifest.GetNumEntriesByKey();
+}
+
+IMPLEMENT_SIMPLE_AUTOMATION_TEST(FLocContextTest, "System.Core.Misc.Internationalization Context", EAutomationTestFlags::EditorContext | EAutomationTestFlags::ClientContext | EAutomationTestFlags::SmokeFilter)
+
+bool FLocContextTest::RunTest( const FString& Parameters )
+{
+	// Key metadata
+	TSharedPtr< FLocMetadataObject > KeyMetadataA = MakeShareable( new FLocMetadataObject );
+	TSharedPtr< FLocMetadataObject > KeyMetadataB = MakeShareable( new FLocMetadataObject );
+
+	// Info metadata
+	TSharedPtr< FLocMetadataObject > InfoMetadataA = MakeShareable( new FLocMetadataObject );
+	TSharedPtr< FLocMetadataObject > InfoMetadataB = MakeShareable( new FLocMetadataObject );
+
+	// Source metadata
+	TSharedPtr< FLocMetadataObject > SourceMetadataA = MakeShareable( new FLocMetadataObject );
+	TSharedPtr< FLocMetadataObject > SourceMetadataB = MakeShareable( new FLocMetadataObject );
+
+
+	// Setup KeyMetadataA
+	KeyMetadataA->SetStringField( TEXT("Gender"			),	TEXT("Masculine")	);
+	KeyMetadataA->SetStringField( TEXT("Plurality"		),	TEXT("Singular")	);
+	KeyMetadataA->SetStringField( TEXT("TargetGender"	),	TEXT("Masculine")	);
+	KeyMetadataA->SetStringField( TEXT("TargetPlurality"),	TEXT("Singular")	);
+
+	// Setup KeyMetadataB
+	KeyMetadataB->SetStringField( TEXT("Gender"			),	TEXT("Masculine")	);
+	KeyMetadataB->SetStringField( TEXT("Plurality"		),	TEXT("Singular")	);
+	KeyMetadataB->SetStringField( TEXT("TargetGender"	),	TEXT("Feminine")	);
+	KeyMetadataB->SetStringField( TEXT("TargetPlurality"),	TEXT("Singular")	);
+
+	// Setup source metadata
+	SourceMetadataA->SetBoolField( TEXT("*IsMature"), false );
+	SourceMetadataB->SetBoolField( TEXT("*IsMature"), true );
+
+	// Set InfoMetadataA
+	InfoMetadataA->SetStringField( TEXT("VoiceActorDirection"), TEXT("Go big or go home!") );
+
+	// Test FManifestContext
+	{
+		FManifestContext ContextA;
+		ContextA.Key			= TEXT("KeyA");
+		ContextA.SourceLocation = TEXT("SourceLocationA");
+		ContextA.InfoMetadataObj = MakeShareable( new FLocMetadataObject(*InfoMetadataA) );
+		ContextA.KeyMetadataObj = MakeShareable( new FLocMetadataObject(*KeyMetadataA) );
+
+		FManifestContext ContextB;
+		ContextB.Key			= TEXT("KeyB");
+		ContextB.SourceLocation = TEXT("SourceLocationB");
+		ContextB.InfoMetadataObj = MakeShareable( new FLocMetadataObject(*InfoMetadataB) );
+		ContextB.KeyMetadataObj = MakeShareable( new FLocMetadataObject(*KeyMetadataB) );
+
+
+		// Test copy ctor
+		{ //-V760
+			FManifestContext ContextAClone = ContextA;
+
+			if( ContextAClone.InfoMetadataObj == ContextA.InfoMetadataObj )
+			{
+				AddError(TEXT("FManifestContext InfoMetadataObj and its Clone are not unique objects."));
+			}
+
+			if( ContextAClone.KeyMetadataObj == ContextA.KeyMetadataObj )
+			{
+				AddError(TEXT("FManifestContext KeyMetadataObj and its Clone are not unique objects."));
+			}
+
+			TestEqual( TEXT("ContextAClone.Key == ContextA.Key"), ContextAClone.Key, ContextA.Key );
+			TestEqual( TEXT("ContextAClone.SourceLocation == ContextA.SourceLocation"), ContextAClone.SourceLocation, ContextA.SourceLocation );
+			TestEqual( TEXT("ContextAClone.bIsOptional == ContextA.bIsOptional"), ContextAClone.bIsOptional, ContextA.bIsOptional );
+			TestTrue( TEXT("ContextAClone.InfoMetadataObj == ContextA.InfoMetadataObj"), *(ContextAClone.InfoMetadataObj) == *(ContextA.InfoMetadataObj) );
+			TestTrue( TEXT("ContextAClone.KeyMetadataObj == ContextA.KeyMetadataObj"), *(ContextAClone.KeyMetadataObj) == *(ContextA.KeyMetadataObj) );
+
+			TestEqual( TEXT("ContextAClone == ContextA"), ContextAClone, ContextA );
+			TestFalse( TEXT("ContextAClone < ContextA"), ContextAClone < ContextA );
+
+		}
+
+		// Test assignment operator
+		{
+			FManifestContext ContextAClone = ContextA;
+
+			if( ContextAClone.InfoMetadataObj == ContextA.InfoMetadataObj )
+			{
+				AddError(TEXT("FManifestContext InfoMetadataObj and its Clone are not unique objects."));
+			}
+
+			if( ContextAClone.KeyMetadataObj == ContextA.KeyMetadataObj )
+			{
+				AddError(TEXT("FManifestContext KeyMetadataObj and its Clone are not unique objects."));
+			}
+
+			TestEqual( TEXT("ContextAClone.Key == ContextA.Key"), ContextAClone.Key, ContextA.Key );
+			TestEqual( TEXT("ContextAClone.SourceLocation == ContextA.SourceLocation"), ContextAClone.SourceLocation, ContextA.SourceLocation );
+			TestEqual( TEXT("ContextAClone.bIsOptional == ContextA.bIsOptional"), ContextAClone.bIsOptional, ContextA.bIsOptional );
+			TestTrue( TEXT("ContextAClone.InfoMetadataObj == ContextA.InfoMetadataObj"), *(ContextAClone.InfoMetadataObj) == *(ContextA.InfoMetadataObj) );
+			TestTrue( TEXT("ContextAClone.KeyMetadataObj == ContextA.KeyMetadataObj"), *(ContextAClone.KeyMetadataObj) == *(ContextA.KeyMetadataObj) );
+
+			TestEqual( TEXT("ContextAClone == ContextA"), ContextAClone, ContextA );
+			TestFalse( TEXT("ContextAClone < ContextA"), ContextAClone < ContextA );
+		}
+
+		// Test comparison operator
+		{
+			// Key and KeyMetadataObj members should be the only items that are taken into account when comparing
+			FManifestContext ContextAClone = ContextA;
+			TestEqual( TEXT("ContextAClone == ContextA"), ContextAClone, ContextA );
+
+			// Arbitrarily change all the non-important members
+			ContextAClone.SourceLocation = ContextA.SourceLocation + TEXT("New");
+			ContextAClone.bIsOptional = !ContextA.bIsOptional;
+			ContextAClone.InfoMetadataObj.Reset();
+			ContextAClone.InfoMetadataObj = MakeShareable( new FLocMetadataObject( *InfoMetadataB ) );
+			TestEqual( TEXT("ContextAClone == ContextA"), ContextAClone, ContextA );
+
+			// Changing the key in any way will cause comparison to fail
+			ContextAClone.Key = ContextAClone.Key.GetString() + TEXT("New");
+			TestNotEqual( TEXT("ContextAClone != ContextA"), ContextAClone, ContextA );
+
+			// Reset and test KeyMetadataObj change to one of the value entries
+			ContextAClone = ContextA;
+			ContextAClone.KeyMetadataObj->SetStringField( TEXT("TargetPlurality"), TEXT("Plural") );
+			TestNotEqual( TEXT("ContextAClone != ContextA"), ContextAClone, ContextA );
+
+			// Reset and test addition of entry to KeyMetadataObj
+			ContextAClone = ContextA;
+			ContextAClone.KeyMetadataObj->SetStringField( TEXT("NewField"), TEXT("NewFieldValue") );
+			TestNotEqual( TEXT("ContextAClone != ContextA"), ContextAClone, ContextA );
+
+			// Reset and test removal of entry from KeyMetadataObj
+			ContextAClone = ContextA;
+			ContextAClone.KeyMetadataObj->RemoveField( TEXT("TargetPlurality") );
+			TestNotEqual( TEXT("ContextAClone != ContextA"), ContextAClone, ContextA );
+
+			// Context with valid but empty KeyMetadataObject should be equivalent to Context with null KeyMetadataObject
+			FManifestContext ContextEmptyA;
+			FManifestContext ContextEmptyB;
+			ContextEmptyB.KeyMetadataObj = MakeShareable( new FLocMetadataObject );
+			TestEqual( TEXT("ContextEmptyA == ContextEmptyB"), ContextEmptyA, ContextEmptyB );
+		}
+
+		// Testing less than operator
+		{
+			TestTrue( TEXT("ContextA < ContextB"), ContextA < ContextB );
+
+			FManifestContext ContextAClone = ContextA;
+
+			// Differences in Key
+			TestFalse( TEXT("ContextA < ContextAClone"), ContextA < ContextAClone );
+			ContextAClone.Key = ContextAClone.Key.GetString() + TEXT("A");
+			//currently failing TestTrue( TEXT("ContextA < ContextAClone"), ContextA < ContextAClone );
+
+			// Adding new key metadata entry that will appear before other entries
+			ContextAClone = ContextA;
+			ContextAClone.KeyMetadataObj->SetStringField( TEXT("ANewKey"), TEXT("ANewValue") );
+			TestTrue( TEXT("ContextAClone < ContextA"), ContextAClone < ContextA );
+
+			// Adding new key metadata entry that will appear after other entries
+			ContextAClone = ContextA;
+			ContextAClone.KeyMetadataObj->SetStringField( TEXT("ZNewKey"), TEXT("ZNewValue") );
+			//currently failing TestTrue( TEXT("ContextA < ContextAClone"), ContextA < ContextAClone );
+
+			// Removing a key metadata entry
+			ContextAClone = ContextA;
+			ContextAClone.KeyMetadataObj->RemoveField( TEXT("TargetPlurality") ) ;
+			TestTrue( TEXT("ContextAClone < ContextA"), ContextAClone < ContextA );
+
+			// Changing a key metadata entry value
+			ContextAClone = ContextA;
+			ContextAClone.KeyMetadataObj->SetStringField( TEXT("TargetPlurality"), TEXT("A") ) ;
+			TestTrue( TEXT("ContextAClone < ContextA"), ContextAClone < ContextA );
+
+			FManifestContext ContextEmptyA;
+			FManifestContext ContextEmptyB;
+			TestFalse( TEXT("ContextEmptyA < ContextEmptyB"), ContextEmptyA < ContextEmptyB );
+			ContextEmptyB.KeyMetadataObj = MakeShareable( new FLocMetadataObject );
+			TestFalse( TEXT("ContextEmptyA < ContextEmptyB"), ContextEmptyA < ContextEmptyB );
+			TestFalse( TEXT("ContextEmptyB < ContextEmptyA"), ContextEmptyB < ContextEmptyA );
+			ContextEmptyB.KeyMetadataObj->SetStringField(TEXT("AMetadataKey"), TEXT("AMetadataValue"));
+			TestTrue( TEXT("ContextEmptyA < ContextEmptyB"), ContextEmptyA < ContextEmptyB );
+		}
+	}
+	return true;
+}
+
+
+IMPLEMENT_SIMPLE_AUTOMATION_TEST(FLocItemTest, "System.Core.Misc.Internationalization LocItem", EAutomationTestFlags::EditorContext | EAutomationTestFlags::ClientContext | EAutomationTestFlags::SmokeFilter)
+
+bool FLocItemTest::RunTest( const FString& Parameters )
+{
+	// Key metadata
+	TSharedPtr< FLocMetadataObject > KeyMetadataA = MakeShareable( new FLocMetadataObject );
+	TSharedPtr< FLocMetadataObject > KeyMetadataB = MakeShareable( new FLocMetadataObject );
+
+	// Info metadata
+	TSharedPtr< FLocMetadataObject > InfoMetadataA = MakeShareable( new FLocMetadataObject );
+	TSharedPtr< FLocMetadataObject > InfoMetadataB = MakeShareable( new FLocMetadataObject );
+
+	// Source metadata
+	TSharedPtr< FLocMetadataObject > SourceMetadataA = MakeShareable( new FLocMetadataObject );
+	TSharedPtr< FLocMetadataObject > SourceMetadataB = MakeShareable( new FLocMetadataObject );
+
+
+	// Setup KeyMetadataA
+	KeyMetadataA->SetStringField( TEXT("Gender"			),	TEXT("Masculine")	);
+	KeyMetadataA->SetStringField( TEXT("Plurality"		),	TEXT("Singular")	);
+	KeyMetadataA->SetStringField( TEXT("TargetGender"	),	TEXT("Masculine")	);
+	KeyMetadataA->SetStringField( TEXT("TargetPlurality"),	TEXT("Singular")	);
+
+	// Setup KeyMetadataB
+	KeyMetadataB->SetStringField( TEXT("Gender"			),	TEXT("Masculine")	);
+	KeyMetadataB->SetStringField( TEXT("Plurality"		),	TEXT("Singular")	);
+	KeyMetadataB->SetStringField( TEXT("TargetGender"	),	TEXT("Feminine")	);
+	KeyMetadataB->SetStringField( TEXT("TargetPlurality"),	TEXT("Singular")	);
+
+	// Setup source metadata
+	SourceMetadataA->SetBoolField( TEXT("*IsMature"), false );
+	SourceMetadataB->SetBoolField( TEXT("*IsMature"), true );
+
+	// Set InfoMetadataA
+	InfoMetadataA->SetStringField( TEXT("VoiceActorDirection"), TEXT("Go big or go home!") );
+
+
+	// Test FLocItem
+	{
+		FLocItem LocItemA( TEXT("TextA") );
+		LocItemA.MetadataObj = MakeShareable( new FLocMetadataObject(*SourceMetadataA) );
+
+		FLocItem LocItemB( TEXT("TextB") );
+		LocItemB.MetadataObj = MakeShareable( new FLocMetadataObject(*SourceMetadataB) );
+
+		// Test copy ctor
+		{
+			FLocItem LocItemAClone = LocItemA;
+
+			if( LocItemAClone.MetadataObj == LocItemA.MetadataObj )
+			{
+				AddError(TEXT("FLocItem MetadataObj and its Clone are not unique objects."));
+			}
+
+			TestEqual( TEXT("LocItemAClone.Text == LocItemA.Text"), LocItemAClone.Text, LocItemA.Text );
+			TestTrue( TEXT("LocItemAClone.MetadataObj == LocItemA.MetadataObj"), *(LocItemAClone.MetadataObj) == *(LocItemA.MetadataObj) );
+
+			TestEqual( TEXT("LocItemAClone == LocItemA"), LocItemAClone, LocItemA );
+			TestFalse( TEXT("LocItemAClone < LocItemA"), LocItemAClone < LocItemA );
+
+		}
+
+		// Test assignment operator
+		{
+			FLocItem LocItemAClone( TEXT("New") );
+
+			LocItemAClone = LocItemA;
+
+			if( LocItemAClone.MetadataObj == LocItemA.MetadataObj )
+			{
+				AddError(TEXT("FLocItem MetadataObj and its Clone are not unique objects."));
+			}
+
+			TestEqual( TEXT("LocItemAClone.Text == LocItemA.Text"), LocItemAClone.Text, LocItemA.Text );
+			TestTrue( TEXT("LocItemAClone.MetadataObj == LocItemA.MetadataObj"), *(LocItemAClone.MetadataObj) == *(LocItemA.MetadataObj) );
+
+			TestEqual( TEXT("LocItemAClone == LocItemA"), LocItemAClone, LocItemA );
+			TestFalse( TEXT("LocItemAClone < LocItemA"), LocItemAClone < LocItemA );
+			TestFalse( TEXT("LocItemA < LocItemAClone"), LocItemA < LocItemAClone );
+		}
+
+		// Test comparison operator
+		{
+			// Text and MetadataObj members should both be taken into account when comparing.  Note, Metadata supports a special * name prefix that causes the type 
+			//   and value of the metadata to be ignored when performing comparisons
+			FLocItem LocItemAClone = LocItemA;
+			TestEqual( TEXT("LocItemAClone == LocItemA"), LocItemAClone, LocItemA );
+
+			// Metadata with * prefix does not impact comparison but both FLocItems need a metadata where the name matches( type and value can be different )
+			FLocItem LocItemAClone2 = LocItemA;
+			LocItemAClone.MetadataObj->SetStringField( TEXT("*NewNonCompare"), TEXT("NewNonCompareValue") );
+			TestNotEqual( TEXT("LocItemAClone != LocItemAClone2"), LocItemAClone, LocItemAClone2 );
+
+			LocItemAClone2.MetadataObj->SetStringField( TEXT("*NewNonCompare"), TEXT("NewNonCompareValue2") );
+			TestEqual( TEXT("LocItemAClone == LocItemAClone2"), LocItemAClone, LocItemAClone2 );
+
+			// Changing the text in any way will cause comparison to fail
+			LocItemAClone = LocItemA;
+			LocItemAClone.Text = LocItemAClone.Text + TEXT("New");
+			TestNotEqual( TEXT("LocItemAClone != LocItemA"), LocItemAClone, LocItemA );
+
+			// LocItem with valid but empty KeyMetadataObject should be equivalent to LocItem with null KeyMetadataObject
+			FLocItem LocItemEmptyA(TEXT("TestText"));
+			FLocItem LocItemEmptyB(TEXT("TestText"));
+			LocItemEmptyB.MetadataObj = MakeShareable( new FLocMetadataObject );
+			TestEqual( TEXT("LocItemEmptyA == LocItemEmptyB"), LocItemEmptyA, LocItemEmptyB );
+		}
+
+		// Testing less than operator
+		{
+			TestTrue( TEXT("LocItemA < LocItemB"), LocItemA < LocItemB );
+			TestFalse( TEXT("LocItemB < LocItemA"), LocItemB < LocItemA );
+
+			FLocItem LocItemAClone = LocItemA;
+
+			// Differences in Text
+			TestFalse( TEXT("LocItemA < LocItemAClone"), LocItemA < LocItemAClone );
+			LocItemAClone.Text = LocItemAClone.Text + TEXT("A");
+			TestFalse( TEXT("LocItemAClone < LocItemA"), LocItemAClone < LocItemA );
+			//currently failing TestTrue( TEXT("LocItemA < LocItemAClone"), LocItemA < LocItemAClone );
+
+			// Adding new key metadata entry
+			LocItemAClone = LocItemA;
+			LocItemAClone.MetadataObj->SetStringField( TEXT("ANewKey"), TEXT("ANewValue") );
+			//currently failing TestTrue( TEXT("LocItemA < LocItemAClone"), LocItemA < LocItemAClone );
+			TestFalse( TEXT("LocItemAClone < LocItemA"), LocItemAClone < LocItemA );
+
+			// Removing a key metadata entry
+			LocItemAClone = LocItemA;
+			LocItemAClone.MetadataObj->RemoveField( TEXT("*IsMature") );
+			TestTrue( TEXT("LocItemAClone < LocItemA"), LocItemAClone < LocItemA );
+			TestFalse( TEXT("LocItemA < LocItemAClone"), LocItemA < LocItemAClone );
+
+			// Changing a key metadata entry value
+			LocItemAClone = LocItemA;
+			LocItemAClone.MetadataObj->SetBoolField( TEXT("*IsMature"),true );
+			//currently failing TestTrue( TEXT("LocItemA < LocItemAClone"), LocItemA < LocItemAClone );
+			TestFalse( TEXT("LocItemAClone < LocItemA"), LocItemAClone < LocItemA );
+
+			// Test null and non-null but empty Metadata
+			FLocItem LocItemEmptyA( TEXT("SameText") );
+			FLocItem LocItemEmptyB( TEXT("SameText") );
+			TestFalse( TEXT("LocItemEmptyA < LocItemEmptyB"), LocItemEmptyA < LocItemEmptyB );
+			LocItemEmptyB.MetadataObj = MakeShareable( new FLocMetadataObject );
+			TestFalse( TEXT("LocItemEmptyA < LocItemEmptyB"), LocItemEmptyA < LocItemEmptyB );
+			TestFalse( TEXT("LocItemEmptyB < LocItemEmptyA"), LocItemEmptyB < LocItemEmptyA );
+			LocItemEmptyB.MetadataObj->SetStringField(TEXT("AMetadataKey"), TEXT("AMetadataValue"));
+			TestTrue( TEXT("LocItemEmptyA < LocItemEmptyB"), LocItemEmptyA < LocItemEmptyB );
+		}
+	}
+	return true;
+}
+
+
+IMPLEMENT_SIMPLE_AUTOMATION_TEST(FManifestTest, "System.Core.Misc.Internationalization Manifest", EAutomationTestFlags::EditorContext | EAutomationTestFlags::ClientContext | EAutomationTestFlags::SmokeFilter)
+
+bool FManifestTest::RunTest( const FString& Parameters )
+{
+	// Key metadata
+	TSharedPtr< FLocMetadataObject > KeyMetadataA = MakeShareable( new FLocMetadataObject );
+	TSharedPtr< FLocMetadataObject > KeyMetadataB = MakeShareable( new FLocMetadataObject );
+
+	// Info metadata
+	TSharedPtr< FLocMetadataObject > InfoMetadataA = MakeShareable( new FLocMetadataObject );
+	TSharedPtr< FLocMetadataObject > InfoMetadataB = MakeShareable( new FLocMetadataObject );
+
+	// Source metadata
+	TSharedPtr< FLocMetadataObject > SourceMetadataA = MakeShareable( new FLocMetadataObject );
+	TSharedPtr< FLocMetadataObject > SourceMetadataB = MakeShareable( new FLocMetadataObject );
+
+
+	// Setup KeyMetadataA
+	KeyMetadataA->SetStringField( TEXT("Gender"			),	TEXT("Masculine")	);
+	KeyMetadataA->SetStringField( TEXT("Plurality"		),	TEXT("Singular")	);
+	KeyMetadataA->SetStringField( TEXT("TargetGender"	),	TEXT("Masculine")	);
+	KeyMetadataA->SetStringField( TEXT("TargetPlurality"),	TEXT("Singular")	);
+
+	// Setup KeyMetadataB
+	KeyMetadataB->SetStringField( TEXT("Gender"			),	TEXT("Masculine")	);
+	KeyMetadataB->SetStringField( TEXT("Plurality"		),	TEXT("Singular")	);
+	KeyMetadataB->SetStringField( TEXT("TargetGender"	),	TEXT("Feminine")	);
+	KeyMetadataB->SetStringField( TEXT("TargetPlurality"),	TEXT("Singular")	);
+
+	// Setup source metadata
+	SourceMetadataA->SetBoolField( TEXT("*IsMature"), false );
+	SourceMetadataB->SetBoolField( TEXT("*IsMature"), true );
+
+	// Set InfoMetadataA
+	InfoMetadataA->SetStringField( TEXT("VoiceActorDirection"), TEXT("Go big or go home!") );
+
+	// Test FInternationalizationManifest
+	{
+		FManifestContext ContextA;
+		ContextA.Key			= TEXT("KeyA");
+		ContextA.SourceLocation = TEXT("SourceLocationA");
+		ContextA.InfoMetadataObj = MakeShareable( new FLocMetadataObject(*InfoMetadataA) );
+		ContextA.KeyMetadataObj = MakeShareable( new FLocMetadataObject(*KeyMetadataA) );
+
+		FManifestContext ContextB;
+		ContextB.Key			= TEXT("KeyB");
+		ContextB.SourceLocation = TEXT("SourceLocationB");
+		ContextB.InfoMetadataObj = MakeShareable( new FLocMetadataObject(*InfoMetadataB) );
+		ContextB.KeyMetadataObj = MakeShareable( new FLocMetadataObject(*KeyMetadataB) );
+
+		FLocItem Source( TEXT("TestText") );
+		Source.MetadataObj = MakeShareable( new FLocMetadataObject(*SourceMetadataA) );
+
+		FString TestNamespace = TEXT("TestNamespace");
+
+
+		// Adding entries with exactly matching Source and matching Context
+		{
+			FInternationalizationManifest TestManifest;
+
+			TestManifest.AddSource( TestNamespace, Source, ContextA );
+			bool bResult = TestManifest.AddSource( TestNamespace, Source, ContextA );
+
+			// Adding a duplicate entry reports success but our entry count is not increased after the first entry is added
+			TestTrue( TEXT("AddSource result == true"), bResult);
+			TestEqual( TEXT("ManifestCount == 1"), CountManifestEntries(TestManifest), 1 );
+		}
+
+		// Adding entries with exactly matching Source but different Contexts
+		{
+			FInternationalizationManifest TestManifest;
+
+			TestManifest.AddSource( TestNamespace, Source, ContextA );
+			TestManifest.AddSource( TestNamespace, Source, ContextB );
+
+			TestEqual( TEXT("ManifestCount == 2"), CountManifestEntries(TestManifest), 2 );
+
+			// Test find by context
+			TSharedPtr< FManifestEntry > FoundEntry1 = TestManifest.FindEntryByContext( TestNamespace, ContextA );
+			if( !FoundEntry1.IsValid() )
+			{
+				AddError(TEXT("FManifestEntry could not find entry using FindEntryByContext."));
+			}
+			else
+			{
+				TestTrue( TEXT("FoundEntry->Source == Source"), FoundEntry1->Source == Source );
+				TestEqual( TEXT("FoundEntry->Context.Num() == 2"), FoundEntry1->Contexts.Num(), 2);
+			}
+
+			TSharedPtr< FManifestEntry > FoundEntry2 = TestManifest.FindEntryByContext( TestNamespace, ContextB );
+			if( !FoundEntry2.IsValid() )
+			{
+				AddError(TEXT("FManifestEntry could not find entry using FindEntryByContext."));
+			}
+			else
+			{
+				TestTrue( TEXT("FoundEntry->Source == Source"), FoundEntry2->Source == Source );
+				TestEqual( TEXT("FoundEntry->Context.Num() == 2"), FoundEntry2->Contexts.Num(), 2);
+			}
+
+			// Test find by source
+			TSharedPtr< FManifestEntry > FoundEntry3 = TestManifest.FindEntryBySource( TestNamespace, Source );
+			if( !FoundEntry3.IsValid() )
+			{
+				AddError(TEXT("FManifestEntry could not find entry using FindEntryBySource."));
+			}
+			else
+			{
+				TestTrue( TEXT("FoundEntry->Source == Source"), FoundEntry3->Source == Source );
+				TestEqual( TEXT("FoundEntry->Context.Num() == 2"), FoundEntry3->Contexts.Num(), 2);
+			}
+
+			TestTrue( TEXT("FoundEntry1 == FoundEntry2 == FoundEntry3"), (FoundEntry1 == FoundEntry2) && (FoundEntry1 == FoundEntry3) );
+
+		}
+
+
+
+		// Adding entries with Source that is NOT an exact match and matching context
+		{
+			// Source mismatched by Source Text.
+			{
+				FInternationalizationManifest TestManifest;
+
+				FLocItem ConflictingSourceA( TEXT("Conflicting TestTextA") );
+				FLocItem ConflictingSourceB( TEXT("Conflicting TestTextB") );
+
+				TestManifest.AddSource( TestNamespace, ConflictingSourceA, ContextA );
+				bool bResult = TestManifest.AddSource( TestNamespace, ConflictingSourceB, ContextA );
+
+				// Adding the second entry reports failure and entry count is not increased
+				TestFalse( TEXT("AddSource result == false"), bResult);
+				TestEqual( TEXT("ManifestCount == 1"), CountManifestEntries(TestManifest), 1 );	
+			}
+
+			// Source mismatched by standard metadata type (not * prefixed) 
+			{
+				FInternationalizationManifest ManifestTypeConflict;
+
+				FLocItem ConflictingSourceMetadataTypeA = Source;
+				FLocItem ConflictingSourceMetadataTypeB = Source;
+
+				// Set metadata with the same name but different type
+				ConflictingSourceMetadataTypeA.MetadataObj->SetBoolField( TEXT("ConflictingType"), true );
+				ConflictingSourceMetadataTypeB.MetadataObj->SetStringField( TEXT("ConflictingType"), TEXT("true") );
+
+				ManifestTypeConflict.AddSource( TestNamespace, ConflictingSourceMetadataTypeA, ContextA );
+				bool bResult = ManifestTypeConflict.AddSource( TestNamespace, ConflictingSourceMetadataTypeB, ContextA );
+
+				TestFalse( TEXT("AddSource result == false"), bResult);
+				TestEqual( TEXT("ManifestCount == 1"), CountManifestEntries(ManifestTypeConflict), 1 );
+
+				// Should not find a match when searching with the second Source we tried to add
+				TSharedPtr< FManifestEntry > FoundEntry = ManifestTypeConflict.FindEntryBySource( TestNamespace, ConflictingSourceMetadataTypeB );
+				TestInvalid( TEXT("FoundEntry is not valid"), FoundEntry);
+
+			}
+
+			// Source mismatched by standard metadata value (not * prefixed)
+			{
+				FInternationalizationManifest TestManifest;
+
+				FLocItem ConflictingSourceMetadataValueA = Source;
+				FLocItem ConflictingSourceMetadataValueB = Source;
+
+				// Set metadata with the same name and type but different value
+				ConflictingSourceMetadataValueA.MetadataObj->SetStringField( TEXT("ConflictingValue"), TEXT("A") );
+				ConflictingSourceMetadataValueB.MetadataObj->SetStringField( TEXT("ConflictingValue"), TEXT("B") );
+
+				TestManifest.AddSource( TestNamespace, ConflictingSourceMetadataValueA, ContextA );
+				bool bResult = TestManifest.AddSource( TestNamespace, ConflictingSourceMetadataValueB, ContextA );
+
+				TestFalse( TEXT("AddSource result == false"), bResult);
+				TestEqual( TEXT("ManifestCount == 1"), CountManifestEntries(TestManifest), 1 );
+
+				// Should not find a match when searching with the second Source we tried to add
+				TSharedPtr< FManifestEntry > FoundEntry = TestManifest.FindEntryBySource( TestNamespace, ConflictingSourceMetadataValueB );
+				TestInvalid( TEXT("FoundEntry is not valid"), FoundEntry);
+
+			}
+
+			// Source mismatched by * prefixed metadata type
+			{
+				FInternationalizationManifest TestManifest;
+
+				FLocItem ConflictingSourceMetadataTypeA = Source;
+				FLocItem ConflictingSourceMetadataTypeB = Source;
+
+				// Set metadata with the same name but different type
+				ConflictingSourceMetadataTypeA.MetadataObj->SetBoolField( TEXT("*ConflictingType"), true );
+				ConflictingSourceMetadataTypeB.MetadataObj->SetStringField( TEXT("*ConflictingType"), TEXT("true") );
+
+				TestManifest.AddSource( TestNamespace, ConflictingSourceMetadataTypeA, ContextA );
+
+				// Though the Source items are considered to be a match(they are equal) in this case, the manifest reports this as a conflict
+				//  and should not add an entry.  The reason is that AddSource does an 'exact' match check on the metadata object
+				bool bResult = TestManifest.AddSource( TestNamespace, ConflictingSourceMetadataTypeB, ContextA );
+
+				TestFalse( TEXT("AddSource result == false"), bResult);
+				TestEqual( TEXT("ManifestCount == 1"), CountManifestEntries(TestManifest), 1 );
+
+				// We should be able to find the entry using either Source FLocItem
+				TSharedPtr< FManifestEntry > FoundEntry1 = TestManifest.FindEntryBySource( TestNamespace, ConflictingSourceMetadataTypeA );
+				TestValid( TEXT("FoundEntry1 is valid"), FoundEntry1);
+
+				TSharedPtr< FManifestEntry > FoundEntry2 = TestManifest.FindEntryBySource( TestNamespace, ConflictingSourceMetadataTypeB );
+				TestValid( TEXT("FoundEntry2 is valid"), FoundEntry2);
+
+
+			}
+
+			// Source mismatched by * prefixed metadata value
+			{
+				FInternationalizationManifest TestManifest;
+
+				FLocItem ConflictingSourceMetadataValueA = Source;
+				FLocItem ConflictingSourceMetadataValueB = Source;
+
+				// Set metadata with the same name and type but different value
+				ConflictingSourceMetadataValueA.MetadataObj->SetStringField( TEXT("*ConflictingValue"), TEXT("A") );
+				ConflictingSourceMetadataValueB.MetadataObj->SetStringField( TEXT("*ConflictingValue"), TEXT("B") );
+
+				TestManifest.AddSource( TestNamespace, ConflictingSourceMetadataValueA, ContextA );
+
+				// Thought the Source items are considered to be a match(they are equal) in this case, the manifest reports this as a conflict
+				//  and should not add an entry.  The reason is that AddSource does an 'exact' match check on the metadata object
+				bool bResult = TestManifest.AddSource( TestNamespace, ConflictingSourceMetadataValueB, ContextA );
+
+				TestFalse( TEXT("AddSource result == false"), bResult);
+				TestEqual( TEXT("ManifestCount == 1"), CountManifestEntries(TestManifest), 1 );
+
+				// We should be able to find the entry using either Source FLocItem
+				TSharedPtr< FManifestEntry > FoundEntry1 = TestManifest.FindEntryBySource( TestNamespace, ConflictingSourceMetadataValueA );
+				TestValid( TEXT("FoundEntry1 is valid"), FoundEntry1);
+
+				TSharedPtr< FManifestEntry > FoundEntry2 = TestManifest.FindEntryBySource( TestNamespace, ConflictingSourceMetadataValueB );
+				TestValid( TEXT("FoundEntry2 is valid"), FoundEntry2);
+
+			}
+
+		}
+
+		// Adding entries with Source that is NOT an exact match and different context
+		{
+			// Source mismatched by Source Text.
+			{
+				FInternationalizationManifest TestManifest;
+
+				FLocItem ConflictingSourceA( TEXT("Conflicting TestTextA") );
+				FLocItem ConflictingSourceB( TEXT("Conflicting TestTextB") );
+
+				TestManifest.AddSource( TestNamespace, ConflictingSourceA, ContextA );
+				bool bResult = TestManifest.AddSource( TestNamespace, ConflictingSourceB, ContextB );
+
+				TestTrue( TEXT("AddSource result == true"), bResult);
+				TestEqual( TEXT("ManifestCount == 2"), CountManifestEntries(TestManifest), 2 );
+
+				// We should be able to find two unique entries by source
+				TSharedPtr< FManifestEntry > FoundEntry1 = TestManifest.FindEntryBySource( TestNamespace, ConflictingSourceA );
+				TestValid( TEXT("FoundEntry1 is valid"), FoundEntry1);
+
+				TSharedPtr< FManifestEntry > FoundEntry2 = TestManifest.FindEntryBySource( TestNamespace, ConflictingSourceB );
+				TestValid( TEXT("FoundEntry2 is valid"), FoundEntry2);
+
+				TestTrue( TEXT("FoundEntry1 != FoundEntry2"), FoundEntry1 != FoundEntry2 );
+
+				// We should be able to find two unique entries by context
+				TSharedPtr< FManifestEntry > FoundEntry3 = TestManifest.FindEntryByContext( TestNamespace, ContextA );
+				TestValid( TEXT("FoundEntry3 is valid"), FoundEntry3);
+
+				TSharedPtr< FManifestEntry > FoundEntry4 = TestManifest.FindEntryByContext( TestNamespace, ContextB );
+				TestValid( TEXT("FoundEntry3 is valid"), FoundEntry4);
+
+				TestTrue( TEXT("FoundEntry3 != FoundEntry4"), FoundEntry3 != FoundEntry4 );
+
+				// Found entry looked up by source A should match entry looked up by context A
+				TestTrue( TEXT("FoundEntry1 == FoundEntry3"), FoundEntry1 == FoundEntry3 );
+
+				// Found entry looked up by source B should match entry looked up by context B
+				TestTrue( TEXT("FoundEntry2 == FoundEntry4"), FoundEntry2 == FoundEntry4 );
+
+			}
+
+			// Source mismatched by standard metadata type (not * prefixed)
+			{
+				FInternationalizationManifest TestManifest;
+
+				FLocItem ConflictingSourceMetadataTypeA = Source;
+				FLocItem ConflictingSourceMetadataTypeB = Source;
+
+				// Set metadata with the same name but different type
+				ConflictingSourceMetadataTypeA.MetadataObj->SetBoolField( TEXT("ConflictingType"), true );
+				ConflictingSourceMetadataTypeB.MetadataObj->SetStringField( TEXT("ConflictingType"), TEXT("true") );
+
+				TestManifest.AddSource( TestNamespace, ConflictingSourceMetadataTypeA, ContextA );
+				bool bResult = TestManifest.AddSource( TestNamespace, ConflictingSourceMetadataTypeB, ContextB );
+
+				TestTrue( TEXT("AddSource result == true"), bResult);
+				TestEqual( TEXT("ManifestCount == 2"), CountManifestEntries(TestManifest), 2 );
+
+				// We should be able to find two unique entries by source
+				TSharedPtr< FManifestEntry > FoundEntry1 = TestManifest.FindEntryBySource( TestNamespace, ConflictingSourceMetadataTypeA );
+				TestValid( TEXT("FoundEntry1 is valid"), FoundEntry1);
+
+				TSharedPtr< FManifestEntry > FoundEntry2 = TestManifest.FindEntryBySource( TestNamespace, ConflictingSourceMetadataTypeB );
+				TestValid( TEXT("FoundEntry2 is valid"), FoundEntry2);
+
+				TestTrue( TEXT("FoundEntry1 != FoundEntry2"), FoundEntry1 != FoundEntry2 );
+
+				// We should be able to find two unique entries by context
+				TSharedPtr< FManifestEntry > FoundEntry3 = TestManifest.FindEntryByContext( TestNamespace, ContextA );
+				TestValid( TEXT("FoundEntry3 is valid"), FoundEntry3);
+
+				TSharedPtr< FManifestEntry > FoundEntry4 = TestManifest.FindEntryByContext( TestNamespace, ContextB );
+				TestValid( TEXT("FoundEntry3 is valid"), FoundEntry4);
+
+				TestTrue( TEXT("FoundEntry3 != FoundEntry4"), FoundEntry3 != FoundEntry4 );
+
+				// Found entry looked up by source A should match entry looked up by context A
+				TestTrue( TEXT("FoundEntry1 == FoundEntry3"), FoundEntry1 == FoundEntry3 );
+
+				// Found entry looked up by source B should match entry looked up by context B
+				TestTrue( TEXT("FoundEntry2 == FoundEntry4"), FoundEntry2 == FoundEntry4 );
+
+			}
+
+			// Source mismatched by standard metadata value (not * prefixed)
+			{
+				FInternationalizationManifest TestManifest;
+
+				FLocItem ConflictingSourceMetadataValueA = Source;
+				FLocItem ConflictingSourceMetadataValueB = Source;
+
+				// Set metadata with the same name and type but different value
+				ConflictingSourceMetadataValueA.MetadataObj->SetStringField( TEXT("ConflictingValue"), TEXT("A") );
+				ConflictingSourceMetadataValueB.MetadataObj->SetStringField( TEXT("ConflictingValue"), TEXT("B") );
+
+				TestManifest.AddSource( TestNamespace, ConflictingSourceMetadataValueA, ContextA );
+				bool bResult = TestManifest.AddSource( TestNamespace, ConflictingSourceMetadataValueB, ContextB );
+
+				TestTrue( TEXT("AddSource result == true"), bResult);
+				TestEqual( TEXT("ManifestCount == 2"), CountManifestEntries(TestManifest), 2 );
+
+				// We should be able to find two unique entries by source
+				TSharedPtr< FManifestEntry > FoundEntry1 = TestManifest.FindEntryBySource( TestNamespace, ConflictingSourceMetadataValueA );
+				TestValid( TEXT("FoundEntry1 is valid"), FoundEntry1);
+
+				TSharedPtr< FManifestEntry > FoundEntry2 = TestManifest.FindEntryBySource( TestNamespace, ConflictingSourceMetadataValueB );
+				TestValid( TEXT("FoundEntry2 is valid"), FoundEntry2);
+
+				TestTrue( TEXT("FoundEntry1 != FoundEntry2"), FoundEntry1 != FoundEntry2 );
+
+				// We should be able to find two unique entries by context
+				TSharedPtr< FManifestEntry > FoundEntry3 = TestManifest.FindEntryByContext( TestNamespace, ContextA );
+				TestValid( TEXT("FoundEntry3 is valid"), FoundEntry3);
+
+				TSharedPtr< FManifestEntry > FoundEntry4 = TestManifest.FindEntryByContext( TestNamespace, ContextB );
+				TestValid( TEXT("FoundEntry3 is valid"), FoundEntry4);
+
+				TestTrue( TEXT("FoundEntry3 != FoundEntry4"), FoundEntry3 != FoundEntry4 );
+
+				// Found entry looked up by source A should match entry looked up by context A
+				TestTrue( TEXT("FoundEntry1 == FoundEntry3"), FoundEntry1 == FoundEntry3 );
+
+				// Found entry looked up by source B should match entry looked up by context B
+				TestTrue( TEXT("FoundEntry2 == FoundEntry4"), FoundEntry2 == FoundEntry4 );
+
+			}
+
+			// Source mismatched by * prefixed metadata type
+			{
+				FInternationalizationManifest TestManifest;
+
+				FLocItem ConflictingSourceMetadataTypeA = Source;
+				FLocItem ConflictingSourceMetadataTypeB = Source;
+
+				// Set metadata with the same name but different type
+				ConflictingSourceMetadataTypeA.MetadataObj->SetBoolField( TEXT("*ConflictingType"), true );
+				ConflictingSourceMetadataTypeB.MetadataObj->SetStringField( TEXT("*ConflictingType"), TEXT("true") );
+
+				TestManifest.AddSource( TestNamespace, ConflictingSourceMetadataTypeA, ContextA );
+
+				// Though the Source items are considered to be a match(they are equal) in this case, the manifest reports this as a conflict
+				//  and should not add an entry.  The reason is that AddSource does an 'exact' match check on the metadata object
+				bool bResult = TestManifest.AddSource( TestNamespace, ConflictingSourceMetadataTypeB, ContextB );
+
+				TestFalse( TEXT("AddSource result == false"), bResult);
+				TestEqual( TEXT("ManifestCount == 1"), CountManifestEntries(TestManifest), 1 );
+
+				// We should be able to find the entry using either Source FLocItem
+				TSharedPtr< FManifestEntry > FoundEntry1 = TestManifest.FindEntryBySource( TestNamespace, ConflictingSourceMetadataTypeA );
+				TestValid( TEXT("FoundEntry1 is valid"), FoundEntry1);
+
+				TSharedPtr< FManifestEntry > FoundEntry2 = TestManifest.FindEntryBySource( TestNamespace, ConflictingSourceMetadataTypeB );
+				TestValid( TEXT("FoundEntry2 is valid"), FoundEntry2);
+			}
+
+			// Source mismatched by * prefixed metadata value
+			{
+				FInternationalizationManifest TestManifest;
+
+				FLocItem ConflictingSourceMetadataValueA = Source;
+				FLocItem ConflictingSourceMetadataValueB = Source;
+
+				// Set metadata with the same name and type but different value
+				ConflictingSourceMetadataValueA.MetadataObj->SetStringField( TEXT("*ConflictingValue"), TEXT("A") );
+				ConflictingSourceMetadataValueB.MetadataObj->SetStringField( TEXT("*ConflictingValue"), TEXT("B") );
+
+				TestManifest.AddSource( TestNamespace, ConflictingSourceMetadataValueA, ContextA );
+
+				// Thought the Source items are considered to be a match(they are equal) in this case, the manifest reports this as a conflict
+				//  and should not add an entry.  The reason is that AddSource does an 'exact' match check on the metadata object
+				bool bResult = TestManifest.AddSource( TestNamespace, ConflictingSourceMetadataValueB, ContextB );
+
+				TestFalse( TEXT("AddSource result == false"), bResult);
+				TestEqual( TEXT("ManifestCount == 1"), CountManifestEntries(TestManifest), 1 );
+
+				// We should be able to find the entry using either Source FLocItem
+				TSharedPtr< FManifestEntry > FoundEntry1 = TestManifest.FindEntryBySource( TestNamespace, ConflictingSourceMetadataValueA );
+				TestValid( TEXT("FoundEntry1 is valid"), FoundEntry1);
+
+				TSharedPtr< FManifestEntry > FoundEntry2 = TestManifest.FindEntryBySource( TestNamespace, ConflictingSourceMetadataValueB );
+				TestValid( TEXT("FoundEntry2 is valid"), FoundEntry2);
+
+			}
+
+		}
+
+		// Adding an entry that only differs in the optional flag 
+		{
+			FInternationalizationManifest TestManifest;
+
+			// Reports success and our entry count does not change.  bIsOptional is not a key and is not used during lookup.  When
+			//  we AddSource, we find an existing matching entry so AddSource returns true but no new entry is added.  The original
+			//  entry's bIsOptional value is not updated.
+			FManifestContext ContextConflictingOptionalFlag = ContextA;
+			ContextConflictingOptionalFlag.bIsOptional = !ContextA.bIsOptional;
+
+			TestManifest.AddSource( TestNamespace, Source, ContextA );
+			bool bResult = TestManifest.AddSource( TestNamespace, Source, ContextConflictingOptionalFlag);
+
+			TestTrue( TEXT("AddSource result == true"), bResult);
+			TestEqual( TEXT("ManifestCount == 1"), CountManifestEntries(TestManifest), 1 );
+
+			// We should be able to look up the existing entry using the ContextConflictingOptionalFlag context but the entries bIsOptional flag will match ContextA
+			TSharedPtr< FManifestEntry > FoundEntry = TestManifest.FindEntryByContext( TestNamespace, ContextConflictingOptionalFlag );
+			if( !FoundEntry.IsValid() )
+			{
+				AddError(TEXT("FManifestEntry could not find entry using FindEntryByContext."));
+			}
+			else
+			{
+				TestTrue( TEXT("FoundEntry->bIsOptional == ContextA->bIsOptional"), FoundEntry->Contexts[0].bIsOptional == ContextA.bIsOptional );
+			}
+		}
+
+		// Add an entry with null key metadata to see if we can retrieve it with non-null but empty key metadata
+		{
+			FInternationalizationManifest TestManifest;
+
+			FManifestContext ContextC;
+			ContextC.Key			= TEXT("KeyC");
+			ContextC.SourceLocation = TEXT("SourceLocationC");
+			ContextC.InfoMetadataObj = MakeShareable( new FLocMetadataObject(*InfoMetadataB) );
+			ContextC.KeyMetadataObj = NULL;
+
+			Source.MetadataObj = NULL;
+			TestManifest.AddSource( TestNamespace, Source, ContextC );
+
+			//  Now give our context and source valid but empty metadata
+			ContextC.KeyMetadataObj = MakeShareable( new FLocMetadataObject() );
+			Source.MetadataObj = MakeShareable( new FLocMetadataObject() );
+
+			// Ensure we find the entry we added by source
+			TSharedPtr< FManifestEntry > FoundEntry;
+			FoundEntry = TestManifest.FindEntryBySource( TestNamespace, Source );
+			if( !FoundEntry.IsValid() )
+			{
+				AddError(TEXT("FManifestEntry could not find entry using FindEntryBySource."));
+			}
+			else
+			{
+				TestTrue( TEXT("FoundEntry->Source == Source"), FoundEntry->Source == Source );
+			}
+
+			// Ensure we find the entry we added by context
+			FoundEntry = TestManifest.FindEntryByContext( TestNamespace, ContextC );
+			if( !FoundEntry.IsValid() )
+			{
+				AddError(TEXT("FManifestEntry could not find entry using FindEntryByContext."));
+			}
+			else
+			{
+				TestTrue( TEXT("FoundEntry->Source == Source"), FoundEntry->Source == Source );
+			}
+		}
+
+		// Add an entry with non-null but empty key metadata and see if we can retrieve null metadata
+		{
+			FInternationalizationManifest TestManifest;
+
+			FManifestContext ContextC;
+			ContextC.Key			= TEXT("KeyC");
+			ContextC.SourceLocation = TEXT("SourceLocationC");
+			ContextC.InfoMetadataObj = MakeShareable( new FLocMetadataObject(*InfoMetadataB) );
+			ContextC.KeyMetadataObj = MakeShareable( new FLocMetadataObject());
+
+			Source.MetadataObj = MakeShareable( new FLocMetadataObject());
+			TestManifest.AddSource( TestNamespace, Source, ContextC );
+
+			//  Now give our context and source Null metadata
+			ContextC.KeyMetadataObj = NULL;
+			Source.MetadataObj = NULL;
+
+			// Ensure we find the entry we added by source
+			TSharedPtr< FManifestEntry > FoundEntry;
+			FoundEntry = TestManifest.FindEntryBySource( TestNamespace, Source );
+			if( !FoundEntry.IsValid() )
+			{
+				AddError(TEXT("FManifestEntry could not find entry using FindEntryBySource."));
+			}
+			else
+			{
+				TestTrue( TEXT("FoundEntry->Source == Source"), FoundEntry->Source == Source );
+			}
+
+			// Ensure we find the entry we added by context
+			FoundEntry = TestManifest.FindEntryByContext( TestNamespace, ContextC );
+			if( !FoundEntry.IsValid() )
+			{
+				AddError(TEXT("FManifestEntry could not find entry using FindEntryByContext."));
+			}
+			else
+			{
+				TestTrue( TEXT("FoundEntry->Source == Source"), FoundEntry->Source == Source );
+			}
+		}
+	}
+
+	return true;
+}
+
+#endif //WITH_DEV_AUTOMATION_TESTS
diff --git a/Engine/Source/Runtime/Core/Private/Tests/Internationalization/MetadataTest.cpp b/Engine/Source/Runtime/Core/Private/Tests/Internationalization/MetadataTest.cpp
new file mode 100644
--- /dev/null
+++ b/Engine/Source/Runtime/Core/Private/Tests/Internationalization/MetadataTest.cpp
@@ -0,0 +1,376 @@
+// Copyright Epic Games, Inc. All Rights Reserved.
+
+#include "CoreTypes.h"
+#include "Containers/Array.h"
+#include "Containers/UnrealString.h"
+#include "Misc/AutomationTest.h"
+#include "Internationalization/InternationalizationMetadata.h"
+
+#if WITH_DEV_AUTOMATION_TESTS
+
+IMPLEMENT_SIMPLE_AUTOMATION_TEST(FMetadataTest, "System.Core.Misc.Internationalization Metadata", EAutomationTestFlags::EditorContext | EAutomationTestFlags::ClientContext | EAutomationTestFlags::SmokeFilter)
+
+BEGIN_FUNCTION_BUILD_OPTIMIZATION
+
+bool FMetadataTest::RunTest( const FString& Parameters )
+{
+	// BooleanValue metadata
+	bool BoolFalse = false;
+	bool BoolTrue = true;
+	TSharedPtr<FLocMetadataValue> MetadataValueBoolFalse = MakeShareable( new FLocMetadataValueBoolean( BoolFalse ) );
+	TSharedPtr<FLocMetadataValue> MetadataValueBoolTrue = MakeShareable( new FLocMetadataValueBoolean( BoolTrue ) );
+
+	// StringValue metadata
+	FString StringA = TEXT("A");
+	FString StringB = TEXT("B");
+	TSharedPtr<FLocMetadataValue> MetadataValueStringA = MakeShareable( new FLocMetadataValueString( StringA ) );
+	TSharedPtr<FLocMetadataValue> MetadataValueStringB = MakeShareable( new FLocMetadataValueString( StringB ) );
+
+	// ArrayValue metadata
+	TArray< TSharedPtr< FLocMetadataValue > > ArrayA;
+	TArray< TSharedPtr< FLocMetadataValue > > ArrayB;
+
+	ArrayA.Add( MetadataValueBoolFalse );
+	ArrayA.Add( MetadataValueStringA );
+
+	ArrayB.Add( MetadataValueBoolTrue );
+	ArrayB.Add( MetadataValueStringB );
+
+	TSharedPtr<FLocMetadataValue> MetadataValueArrayA = MakeShareable( new FLocMetadataValueArray( ArrayA ) );
+	TSharedPtr<FLocMetadataValue> MetadataValueArrayB = MakeShareable( new FLocMetadataValueArray( ArrayB ) );
+
+	// Object metadata
+	TSharedPtr< FLocMetadataObject > MetadataObjectA = MakeShareable( new FLocMetadataObject );
+	TSharedPtr< FLocMetadataObject > MetadataObjectB = MakeShareable( new FLocMetadataObject );
+
+	// Setup object A
+	MetadataObjectA->SetField( TEXT("MetadataBoolFalse"), MetadataValueBoolFalse );
+	MetadataObjectA->SetField( TEXT("MetadataStringA"), MetadataValueStringA );
+	MetadataObjectA->SetField( TEXT("MetadataArrayA"), MetadataValueArrayA );
+	MetadataObjectA->SetField( TEXT("*MetadataCompareModifier"), MetadataValueStringA);// Note: The * name prefix modifies the way entries in the object are compared.
+
+	// Setup object B
+	MetadataObjectB->SetField( TEXT("MetadataBoolFalse"), MetadataValueBoolTrue );
+	MetadataObjectB->SetField( TEXT("MetadataStringB"), MetadataValueStringB );
+	MetadataObjectB->SetField( TEXT("MetadataArrayB"), MetadataValueArrayB );
+	MetadataObjectA->SetBoolField( TEXT("*MetadataCompareModifier"), true);// Note: Different type/value. The * name prefix modifies the way entries in the object are compared.
+
+	// ObjectValue metadata
+	TSharedPtr< FLocMetadataValue > MetadataValueObjectA = MakeShareable( new FLocMetadataValueObject( MetadataObjectA ) );
+	TSharedPtr< FLocMetadataValue > MetadataValueObjectB = MakeShareable( new FLocMetadataValueObject( MetadataObjectB ) );
+
+	// Testing the bool meta data value type
+	{
+		TestFalse( TEXT("MetadataValueBoolFalse == MetadataValueBoolTrue"), *MetadataValueBoolFalse == *MetadataValueBoolTrue );
+		TestTrue( TEXT("MetadataValueBoolFalse < MetadataValueBoolTrue"), *MetadataValueBoolFalse < *MetadataValueBoolTrue );
+		TestFalse( TEXT("MetadataValueBoolTrue < MetadataValueBoolFalse"), *MetadataValueBoolTrue < *MetadataValueBoolFalse);
+
+		TestTrue( TEXT("MetadataValueBoolFalse < MetadataValueString"), *MetadataValueBoolFalse < *MetadataValueStringA );
+		TestTrue( TEXT("MetadataValueBoolTrue < MetadataValueString"), *MetadataValueBoolTrue < *MetadataValueStringA);
+
+		TestTrue( TEXT("MetadataValueBoolFalse < MetadataValueArray"), *MetadataValueBoolFalse < *MetadataValueArrayA );
+		TestTrue( TEXT("MetadataValueBoolTrue < MetadataValueArray"), *MetadataValueBoolTrue < *MetadataValueArrayA);
+
+		TestTrue( TEXT("MetadataValueBoolFalse < MetadataValueObject"), *MetadataValueBoolFalse < *MetadataValueObjectA );
+		TestTrue( TEXT("MetadataValueBoolTrue < MetadataValueObject"), *MetadataValueBoolTrue < *MetadataValueObjectA);
+
+		TSharedPtr<FLocMetadataValue> MetadataValueBoolFalseClone = MetadataValueBoolFalse->Clone();
+		TSharedPtr<FLocMetadataValue> MetadataValueBoolTrueClone = MetadataValueBoolTrue->Clone();
+
+		if( MetadataValueBoolFalse == MetadataValueBoolFalseClone )
+		{
+			AddError(TEXT("MetadataValueBool and its Clone are not unique objects."));
+		}
+		
+		TestTrue( TEXT("MetadataValueBoolFalseClone == MetadataValueBoolFalse"), *MetadataValueBoolFalseClone == *MetadataValueBoolFalse );
+		TestFalse( TEXT("MetadataValueBoolFalseClone < MetadataValueBoolFalse"), *MetadataValueBoolFalseClone < *MetadataValueBoolFalse );
+
+		TestTrue( TEXT("MetadataValueBoolTrueClone == MetadataValueBoolTrue"), *MetadataValueBoolTrueClone == *MetadataValueBoolTrue );
+		TestFalse( TEXT("MetadataValueBoolTrueClone < MetadataValueBoolTrue"), *MetadataValueBoolTrueClone < *MetadataValueBoolTrue );
+
+		// Test the bool metadata when it is part of an object
+		TSharedPtr< FLocMetadataObject > MetadataObjectFalse = MakeShareable( new FLocMetadataObject );
+		MetadataObjectFalse->SetField(TEXT("MetadataValueBool"), MetadataValueBoolFalse);
+
+		TSharedPtr< FLocMetadataObject > MetadataObjectTrue = MakeShareable( new FLocMetadataObject );
+		MetadataObjectTrue->SetField(TEXT("MetadataValueBool"), MetadataValueBoolTrue);
+
+		TestFalse( TEXT("GetBoolField(MetadataValueBool)"), MetadataObjectFalse->GetBoolField( TEXT("MetadataValueBool") ) );
+		TestTrue( TEXT("GetBoolField(MetadataValueBool)"), MetadataObjectTrue->GetBoolField( TEXT("MetadataValueBool") ) );
+
+		TestFalse( TEXT("MetadataObjectFalse == MetadataObjectTrue"), *MetadataObjectFalse == *MetadataObjectTrue );
+		TestTrue( TEXT("MetadataObjectFalse < MetadataObjectTrue"), *MetadataObjectFalse < *MetadataObjectTrue );
+	}
+
+	// Testing string metadata value type
+	{
+		TestFalse( TEXT("MetadataValueStringA == MetadataValueStringB"), *MetadataValueStringA == *MetadataValueStringB );
+		TestTrue( TEXT("MetadataValueStringA < MetadataValueStringB"), *MetadataValueStringA < *MetadataValueStringB );
+		TestFalse( TEXT("MetadataValueStringB < MetadataValueStringA"), *MetadataValueStringB < *MetadataValueStringA);
+
+		TestTrue( TEXT("MetadataValueString < MetadataValueArray"), *MetadataValueStringA < *MetadataValueArrayA );
+
+		TestTrue( TEXT("MetadataValueStringA < MetadataValueObject"), *MetadataValueStringA < *MetadataValueObjectA );
+
+
+		TSharedPtr<FLocMetadataValue> MetadataValueStringAClone = MetadataValueStringA->Clone();
+
+		if( MetadataValueStringA == MetadataValueStringAClone )
+		{
+			AddError(TEXT("MetadataValueString and its Clone are not unique objects."));
+		}
+
+		TestTrue( TEXT("MetadataValueStringAClone == MetadataValueStringA"), *MetadataValueStringAClone == *MetadataValueStringA );
+		TestFalse( TEXT("MetadataValueStringAClone < MetadataValueStringA"), *MetadataValueStringAClone < *MetadataValueStringA );
+		TestTrue( TEXT("MetadataValueStringAClone < MetadataValueStringB"), *MetadataValueStringAClone < *MetadataValueStringB );
+
+		// Test the string metadata when it is part of an object
+		TSharedPtr< FLocMetadataObject > TestMetadataObjectA = MakeShareable( new FLocMetadataObject );
+		TestMetadataObjectA->SetField(TEXT("MetadataValueString"), MetadataValueStringA);
+
+		TSharedPtr< FLocMetadataObject > TestMetadataObjectB = MakeShareable( new FLocMetadataObject );
+		TestMetadataObjectB->SetField(TEXT("MetadataValueString"), MetadataValueStringB);
+
+		TestTrue( TEXT("GetStringField(MetadataValueString) == StringA"), TestMetadataObjectA->GetStringField( TEXT("MetadataValueString") ) == StringA );
+	
+		TestFalse( TEXT("TestMetadataObjectA == TestMetadataObjectB"), *TestMetadataObjectA == *TestMetadataObjectB );
+		TestTrue( TEXT("TestMetadataObjectA < TestMetadataObjectB"), *TestMetadataObjectA < *TestMetadataObjectB );
+	}
+
+	// Testing array metadata value type
+	{
+		TestFalse( TEXT("MetadataValueArrayA == MetadataValueArrayB"), *MetadataValueArrayA == *MetadataValueArrayB );
+		TestTrue( TEXT("MetadataValueArrayA < MetadataValueArrayB"), *MetadataValueArrayA < *MetadataValueArrayB );
+		TestFalse( TEXT("MetadataValueArrayB < MetadataValueArrayA"), *MetadataValueArrayB < *MetadataValueArrayA);
+
+		TestTrue( TEXT("MetadataValueArrayA < MetadataValueObject"), *MetadataValueArrayA < *MetadataValueObjectA );
+
+		TSharedPtr<FLocMetadataValue> MetadataValueArrayAClone = MetadataValueArrayA->Clone();
+
+		if( MetadataValueArrayA == MetadataValueArrayAClone )
+		{
+			AddError(TEXT("MetadataValueString and its Clone are not unique objects."));
+		}
+
+		TestTrue( TEXT("MetadataValueArrayAClone == MetadataValueArrayA"), *MetadataValueArrayAClone == *MetadataValueArrayA );
+		TestFalse( TEXT("MetadataValueArrayAClone < MetadataValueArrayA"), *MetadataValueArrayAClone < *MetadataValueArrayA );
+		TestTrue( TEXT("MetadataValueArrayAClone < MetadataValueArrayB"), *MetadataValueArrayAClone < *MetadataValueArrayB );
+
+		// Test Less than and equality checks.  Metadata arrays are equivalent if they contain equivalent contents in any order.  
+		//  To calculate if a metadata array is less than another, we sort both arrays and check each entry index against its counterpart.
+		//  If we encounter an entry that is less than another we stop looking.
+		TArray< TSharedPtr< FLocMetadataValue > > ArrayC;
+		ArrayC.Add( MetadataValueBoolFalse );
+		ArrayC.Add( MetadataValueBoolFalse->Clone() );
+		TSharedPtr<FLocMetadataValue> MetadataValueArrayC = MakeShareable( new FLocMetadataValueArray( ArrayC ) );
+	
+		TestFalse( TEXT("MetadataValueArrayA == MetadataValueArrayC"), *MetadataValueArrayA == *MetadataValueArrayC );
+		TestTrue( TEXT("MetadataValueArrayC < MetadataValueArrayA"), *MetadataValueArrayC < *MetadataValueArrayA );
+		TestTrue( TEXT("MetadataValueArrayC < MetadataValueArrayB"), *MetadataValueArrayC < *MetadataValueArrayB );
+
+		TArray< TSharedPtr< FLocMetadataValue > > ArrayD;
+		ArrayD.Add( MetadataValueBoolFalse );
+		ArrayD.Add( MetadataValueBoolFalse->Clone() );
+		ArrayD.Add( MetadataValueBoolFalse->Clone() );
+		TSharedPtr<FLocMetadataValue> MetadataValueArrayD = MakeShareable( new FLocMetadataValueArray( ArrayD ) );
+
+		TestFalse( TEXT("MetadataValueArrayA == MetadataValueArrayD"), *MetadataValueArrayA == *MetadataValueArrayD );
+		TestFalse( TEXT("MetadataValueArrayC == MetadataValueArrayD"), *MetadataValueArrayC == *MetadataValueArrayD );
+		TestTrue( TEXT("MetadataValueArrayC < MetadataValueArrayD"), *MetadataValueArrayC < *MetadataValueArrayD );
+		TestTrue( TEXT("MetadataValueArrayD < MetadataValueArrayA"), *MetadataValueArrayD < *MetadataValueArrayA );
+
+		// Test the array metadata when it is part of an object
+		TSharedPtr< FLocMetadataObject > TestMetadataObjectA = MakeShareable( new FLocMetadataObject );
+		TestMetadataObjectA->SetField(TEXT("MetadataValueArray"), MetadataValueArrayA);
+
+		TSharedPtr< FLocMetadataObject > TestMetadataObjectB = MakeShareable( new FLocMetadataObject );
+		TestMetadataObjectB->SetField(TEXT("MetadataValueArray"), MetadataValueArrayB);
+
+		TestTrue( TEXT("GetArrayField(MetadataValueArray) == ArrayA"), TestMetadataObjectA->GetArrayField( TEXT("MetadataValueArray") ) == ArrayA );
+
+		TestFalse( TEXT("TestMetadataObjectA == TestMetadataObjectB"), *TestMetadataObjectA == *TestMetadataObjectB );
+		TestTrue( TEXT("TestMetadataObjectA < TestMetadataObjectB"), *TestMetadataObjectA < *TestMetadataObjectB );
+	}
+
+	// Testing object metadata value type
+	{
+		TestFalse( TEXT("MetadataValueObjectA == MetadataValueObjectB"), *MetadataValueObjectA == *MetadataValueObjectB );
+		TestTrue( TEXT("MetadataValueObjectA < MetadataValueObjectB"), *MetadataValueObjectA < *MetadataValueObjectB );
+		TestFalse( TEXT("MetadataValueObjectB < MetadataValueObjectA"), *MetadataValueObjectB < *MetadataValueObjectA);
+
+		TSharedPtr<FLocMetadataValue> MetadataValueObjectAClone = MetadataValueObjectA->Clone();
+
+		if( MetadataValueObjectA == MetadataValueObjectAClone )
+		{
+			AddError(TEXT("MetadataValueString and its Clone are not unique objects."));
+		}
+
+		TestTrue( TEXT("MetadataValueObjectAClone == MetadataValueObjectA"), *MetadataValueObjectAClone == *MetadataValueObjectA );
+		TestFalse( TEXT("MetadataValueObjectAClone < MetadataValueObjectA"), *MetadataValueObjectAClone < *MetadataValueObjectA );
+		TestTrue( TEXT("MetadataValueObjectAClone < MetadataValueObjectB"), *MetadataValueObjectAClone < *MetadataValueObjectB );
+
+
+		// Test the object metadata when it is part of another object
+		TSharedPtr< FLocMetadataObject > TestMetadataObjectA = MakeShareable( new FLocMetadataObject );
+		TestMetadataObjectA->SetField(TEXT("MetadataValueObject"), MetadataValueObjectA);
+
+		TSharedPtr< FLocMetadataObject > TestMetadataObjectB = MakeShareable( new FLocMetadataObject );
+		TestMetadataObjectB->SetField(TEXT("MetadataValueObject"), MetadataValueObjectB);
+
+		TestTrue( TEXT("GetObjectField(MetadataValueObject) == MetadataObjectA"), *TestMetadataObjectA->GetObjectField( TEXT("MetadataValueObject") ) == *MetadataObjectA );
+
+		TestFalse( TEXT("TestMetadataObjectA == TestMetadataObjectB"), *TestMetadataObjectA == *TestMetadataObjectB );
+		TestTrue( TEXT("TestMetadataObjectA < TestMetadataObjectB"), *TestMetadataObjectA < *TestMetadataObjectB );
+	}
+
+	// Testing Loc Metadata Object
+	{
+		TestFalse( TEXT("MetadataObjectA == MetadataObjectB"), *MetadataObjectA == *MetadataObjectB );
+		TestTrue( TEXT("MetadataObjectA < MetadataObjectB"), *MetadataObjectA < *MetadataObjectB );
+		TestFalse( TEXT("MetadataObjectB < MetadataObjectA"), *MetadataObjectB < *MetadataObjectA);
+
+		// Test copy ctor
+		{
+			FLocMetadataObject MetadataObjectAClone = *MetadataObjectA;
+			TestTrue( TEXT("MetadataObjectAClone == MetadataObjectA"), MetadataObjectAClone == *MetadataObjectA );
+		}
+		
+		// Test assignment
+		{
+			FLocMetadataObject MetadataObjectAClone = *MetadataObjectB;
+			// PVS-Studio complains about double initialization, but that's something we're testing
+			// so we disable the warning
+			MetadataObjectAClone = *MetadataObjectA; //-V519
+			TestTrue( TEXT("MetadataObjectAClone == MetadataObjectA"), MetadataObjectAClone == *MetadataObjectA );
+			TestFalse( TEXT("MetadataObjectAClone == MetadataObjectB"), MetadataObjectAClone == *MetadataObjectB );
+		}
+		
+		// Test comparison operator  
+		{
+			FLocMetadataObject MetadataObjectAClone = *MetadataObjectA;
+			
+			
+			// Adding standard entry
+			MetadataObjectAClone.SetStringField( TEXT("NewEntry"), TEXT("NewEntryValue") );
+			TestFalse( TEXT("MetadataObjectAClone == MetadataObjectA"), MetadataObjectAClone == *MetadataObjectA );
+			
+			// Adding non-standard entry.  Note metadata with * prefix in the name will ignore value and type when performing comparisons
+			MetadataObjectAClone = *MetadataObjectA;
+			MetadataObjectAClone.SetStringField( TEXT("*NewEntry"), TEXT("*NewEntryValue") );
+			TestFalse( TEXT("MetadataObjectAClone == MetadataObjectA"), MetadataObjectAClone == *MetadataObjectA );
+
+			// Value mismatch on entry with * prefix with same type
+			MetadataObjectAClone = *MetadataObjectA;
+			MetadataObjectAClone.SetStringField( TEXT("*NoCompare"), TEXT("NoCompare") );
+			FLocMetadataObject MetadataObjectAClone2 = *MetadataObjectA;
+			MetadataObjectAClone2.SetStringField( TEXT("*NoCompare"), TEXT("NoCompare2") );
+			TestTrue( TEXT("MetadataObjectAClone == MetadataObjectAClone2"), MetadataObjectAClone == MetadataObjectAClone2 );
+
+			// Value and type mismatch on entry with * prefix
+			MetadataObjectAClone = *MetadataObjectA;
+			MetadataObjectAClone.SetStringField( TEXT("*NoCompare"), TEXT("NoCompare") );
+			MetadataObjectAClone2 = *MetadataObjectA;
+			MetadataObjectAClone2.SetBoolField( TEXT("*NoCompare"), true );
+			TestTrue( TEXT("MetadataObjectAClone == MetadataObjectAClone2"), MetadataObjectAClone == MetadataObjectAClone2 );
+		
+			// Value mismatch on standard entry
+			MetadataObjectAClone = *MetadataObjectA;
+			MetadataObjectAClone.SetStringField( TEXT("DoCompare"), TEXT("DoCompare") );
+			MetadataObjectAClone2 = *MetadataObjectA;
+			MetadataObjectAClone2.SetStringField( TEXT("DoCompare"), TEXT("DoCompare2") );
+			TestFalse( TEXT("MetadataObjectAClone == MetadataObjectAClone2"), MetadataObjectAClone == MetadataObjectAClone2 );
+
+			// Value and type mismatch on standard entry
+			MetadataObjectAClone = *MetadataObjectA;
+			MetadataObjectAClone.SetStringField( TEXT("DoCompare"), TEXT("DoCompare") );
+			MetadataObjectAClone2 = *MetadataObjectA;
+			MetadataObjectAClone2.SetBoolField( TEXT("DoCompare"), true );
+			TestFalse( TEXT("MetadataObjectAClone == MetadataObjectAClone2"), MetadataObjectAClone == MetadataObjectAClone2 );
+		}		
+		
+		// Test IsExactMatch function.  Note: Differs from the == operator which takes into account the COMPARISON_MODIFIER_PREFIX
+		{
+			// This function behaves similar to the comparison operator except that it ensures all metadata entries match exactly.  In other
+			//  words, it will also perform exact match checks on * prefixed metadata items.
+			FLocMetadataObject MetadataObjectAClone = *MetadataObjectA;
+
+			TestTrue( TEXT("MetadataObjectAClone == MetadataObjectA"), MetadataObjectAClone.IsExactMatch( *MetadataObjectA) );
+
+			// Adding standard entry
+			MetadataObjectAClone.SetStringField( TEXT("NewEntry"), TEXT("NewEntryValue") );
+			TestFalse( TEXT("MetadataObjectAClone == MetadataObjectA"), MetadataObjectAClone.IsExactMatch( *MetadataObjectA) );
+
+			// Adding non-standard entry.  Note metadata with * prefix in the name will ignore value and type when performing comparisons
+			MetadataObjectAClone = *MetadataObjectA;
+			MetadataObjectAClone.SetStringField( TEXT("*NewEntry"), TEXT("*NewEntryValue") );
+			TestFalse( TEXT("MetadataObjectAClone == MetadataObjectA"),MetadataObjectAClone.IsExactMatch( *MetadataObjectA)  );
+
+			// Value mismatch on entry with * prefix with same type
+			MetadataObjectAClone = *MetadataObjectA;
+			MetadataObjectAClone.SetStringField( TEXT("*NoCompare"), TEXT("NoCompare") );
+			FLocMetadataObject MetadataObjectAClone2 = *MetadataObjectA;
+			MetadataObjectAClone2.SetStringField( TEXT("*NoCompare"), TEXT("NoCompare2") );
+			TestFalse( TEXT("MetadataObjectAClone == MetadataObjectAClone2"), MetadataObjectAClone.IsExactMatch( MetadataObjectAClone2) );
+
+			// Value and type mismatch on entry with * prefix
+			MetadataObjectAClone = *MetadataObjectA;
+			MetadataObjectAClone.SetStringField( TEXT("*NoCompare"), TEXT("NoCompare") );
+			MetadataObjectAClone2 = *MetadataObjectA;
+			MetadataObjectAClone2.SetBoolField( TEXT("*NoCompare"), true );
+			TestFalse( TEXT("MetadataObjectAClone == MetadataObjectAClone2"),  MetadataObjectAClone.IsExactMatch( MetadataObjectAClone2) );
+
+			// Value mismatch on standard entry
+			MetadataObjectAClone = *MetadataObjectA;
+			MetadataObjectAClone.SetStringField( TEXT("DoCompare"), TEXT("DoCompare") );
+			MetadataObjectAClone2 = *MetadataObjectA;
+			MetadataObjectAClone2.SetStringField( TEXT("DoCompare"), TEXT("DoCompare2") );
+			TestFalse( TEXT("MetadataObjectAClone == MetadataObjectAClone2"),  MetadataObjectAClone.IsExactMatch( MetadataObjectAClone2) );
+
+			// Value and type mismatch on standard entry
+			MetadataObjectAClone = *MetadataObjectA;
+			MetadataObjectAClone.SetStringField( TEXT("DoCompare"), TEXT("DoCompare") );
+			MetadataObjectAClone2 = *MetadataObjectA;
+			MetadataObjectAClone2.SetBoolField( TEXT("DoCompare"), true );
+			TestFalse( TEXT("MetadataObjectAClone == MetadataObjectAClone2"),  MetadataObjectAClone.IsExactMatch( MetadataObjectAClone2) );
+		}
+
+		// Test less than operator
+		{
+			
+			// Adding standard entry that would appear before other entries
+			FLocMetadataObject MetadataObjectAClone = *MetadataObjectA;
+			MetadataObjectAClone.SetStringField( TEXT("ANewEntry"), TEXT("NewEntryValue") );
+			TestTrue( TEXT("MetadataObjectAClone < MetadataObjectA"), MetadataObjectAClone < *MetadataObjectA );
+			
+			// Adding standard entry that would appear after other entries
+			MetadataObjectAClone = *MetadataObjectA;
+			MetadataObjectAClone.SetStringField( TEXT("ZNewEntry"), TEXT("NewEntryValue") );
+			TestTrue( TEXT("MetadataObjectA < MetadataObjectAClone"), *MetadataObjectA < MetadataObjectAClone );
+
+			// Adding non-standard entry that would appear before other entries.  Note metadata with * prefix has no special treatment in less than operator
+			MetadataObjectAClone = *MetadataObjectA;
+			MetadataObjectAClone.SetStringField( TEXT("*NewEntry"), TEXT("NewEntryValue") );
+			TestTrue( TEXT("MetadataObjectAClone < MetadataObjectA"), MetadataObjectAClone < *MetadataObjectA );
+
+			// Value mismatch on entry with * prefix with same type. Note metadata with * prefix has no special treatment in less than operator
+			MetadataObjectAClone = *MetadataObjectA;
+			MetadataObjectAClone.SetStringField( TEXT("*NoCompare"), TEXT("NoCompare") );
+			FLocMetadataObject MetadataObjectAClone2 = *MetadataObjectA;
+			MetadataObjectAClone2.SetStringField( TEXT("*NoCompare"), TEXT("NoCompare2") );
+			TestTrue( TEXT("MetadataObjectAClone < MetadataObjectAClone2"), MetadataObjectAClone < MetadataObjectAClone2 );
+
+			// Value and type mismatch on entry with * prefix.  Note metadata with * prefix has no special treatment in less than operator
+			MetadataObjectAClone = *MetadataObjectA;
+			MetadataObjectAClone.SetBoolField( TEXT("*NoCompare"), true );
+			MetadataObjectAClone2 = *MetadataObjectA;
+			MetadataObjectAClone2.SetStringField( TEXT("*NoCompare"), TEXT("NoCompare") );
+			TestTrue( TEXT("MetadataObjectAClone < MetadataObjectAClone2"), MetadataObjectAClone < MetadataObjectAClone2 );
+		}
+	}
+
+	return true;
+}
+
+END_FUNCTION_BUILD_OPTIMIZATION
+
+#endif //WITH_DEV_AUTOMATION_TESTS
