LowLevelTests: Reverted the Core tests that were moved in 19498448 and 19571202

diff --git a/Engine/Source/Runtime/Core/Private/Tests/Math/RangeBoundTest.cpp b/Engine/Source/Runtime/Core/Private/Tests/Math/RangeBoundTest.cpp
new file mode 100644
--- /dev/null
+++ b/Engine/Source/Runtime/Core/Private/Tests/Math/RangeBoundTest.cpp
@@ -0,0 +1,122 @@
+// Copyright Epic Games, Inc. All Rights Reserved.
+
+#include "CoreTypes.h"
+#include "Math/Range.h"
+#include "Misc/AutomationTest.h"
+
+#if WITH_DEV_AUTOMATION_TESTS
+
+IMPLEMENT_SIMPLE_AUTOMATION_TEST(FRangeBoundTest, "System.Core.Math.RangeBound", EAutomationTestFlags::EditorContext | EAutomationTestFlags::ClientContext | EAutomationTestFlags::SmokeFilter)
+
+bool FRangeBoundTest::RunTest( const FString& Parameters )
+{
+	// ensure template instantiation works
+	FDateRange DateRange;
+	FDoubleRange DoubleRange;
+	FFloatRange FloatRange;
+	FInt8Range Int8Range;
+	FInt16Range Int16Range;
+	FInt32Range Int32Range;
+	FInt64Range Int64Range;
+
+	// bound types must be correct after construction
+	FFloatRangeBound b1_1 = FFloatRangeBound::Exclusive(2.0f);
+	FFloatRangeBound b1_2 = FFloatRangeBound::Inclusive(2.0f);
+	FFloatRangeBound b1_3 = FFloatRangeBound::Open();
+	FFloatRangeBound b1_4 = 2;
+
+	TestTrue(TEXT("Exclusive bound constructor must create exclusive bound"), b1_1.IsExclusive());
+	TestTrue(TEXT("Exclusive bound constructor must create closed bound"), b1_1.IsClosed());
+	TestFalse(TEXT("Exclusive bound constructor must not create inclusive bound"), b1_1.IsInclusive());
+	TestFalse(TEXT("Exclusive bound constructor must not create open bound"), b1_1.IsOpen());
+	TestEqual(TEXT("Exclusive bound constructor must create the correct value"), b1_1.GetValue(), 2.0f);
+
+	TestTrue(TEXT("Inclusive bound constructor must create exclusive bound"), b1_2.IsInclusive());
+	TestTrue(TEXT("Inclusive bound constructor must create closed bound"), b1_2.IsClosed());
+	TestFalse(TEXT("Inclusive bound constructor must not create exclusive bound"), b1_2.IsExclusive());
+	TestFalse(TEXT("Inclusive bound constructor must not create open bound"), b1_2.IsOpen());
+	TestEqual(TEXT("Inclusive bound constructor must create the correct value"), b1_2.GetValue(), 2.0f);
+
+	TestTrue(TEXT("Open bound constructor must create open bound"), b1_3.IsOpen());
+	TestFalse(TEXT("Open bound constructor must not create closed bound"), b1_3.IsClosed());
+	TestFalse(TEXT("Open bound constructor must not create exclusive bound"), b1_3.IsExclusive());
+	TestFalse(TEXT("Open bound constructor must not create inclusive bound"), b1_3.IsInclusive());
+
+	TestTrue(TEXT("Implicit constructor must create an inclusive bound"), b1_4.IsInclusive());
+	TestEqual(TEXT("Implicit constructor must create the correct value"), b1_4, b1_2);
+
+	// comparisons must be correct
+	FFloatRangeBound b2_1 = FFloatRangeBound::Exclusive(2.0f);
+	FFloatRangeBound b2_2 = FFloatRangeBound::Exclusive(2.0f);
+	FFloatRangeBound b2_3 = FFloatRangeBound::Inclusive(2.0f);
+	FFloatRangeBound b2_4 = FFloatRangeBound::Inclusive(2.0f);
+	FFloatRangeBound b2_5 = FFloatRangeBound::Open();
+	FFloatRangeBound b2_6 = FFloatRangeBound::Open();
+
+	TestTrue(TEXT("Equal exclusive bounds must be equal"), b2_1 == b2_2);
+	TestTrue(TEXT("Equal inclusive bounds must be equal"), b2_3 == b2_4);
+	TestTrue(TEXT("Open bounds must be equal"), b2_5 == b2_6);
+
+	TestFalse(TEXT("Equal exclusive bounds must not be unequal"), b2_1 != b2_2);
+	TestFalse(TEXT("Equal inclusive bounds must not be unequal"), b2_3 != b2_4);
+	TestFalse(TEXT("Open bounds must not be unequal"), b2_5 != b2_6);
+
+	FFloatRangeBound b2_7 = FFloatRangeBound::Exclusive(3.0f);
+	FFloatRangeBound b2_8 = FFloatRangeBound::Inclusive(3.0f);
+
+	TestTrue(TEXT("Unequal exclusive bounds must be unequal"), b2_1 != b2_7);
+	TestTrue(TEXT("Unequal inclusive bounds must be unequal"), b2_2 != b2_8);
+
+	TestFalse(TEXT("Unequal exclusive bounds must not be equal"), b2_1 == b2_7);
+	TestFalse(TEXT("Unequal inclusive bounds must not be equal"), b2_2 == b2_8);
+
+	// min-max comparisons between bounds must be correct
+	FFloatRangeBound b3_1 = FFloatRangeBound::Exclusive(2.0f);
+	FFloatRangeBound b3_2 = FFloatRangeBound::Inclusive(2.0f);
+	FFloatRangeBound b3_3 = FFloatRangeBound::Open();
+
+	TestEqual(TEXT("'[2' must be less than '(2' <1>"), FFloatRangeBound::MinLower(b3_2, b3_1), b3_2);
+	TestEqual(TEXT("'[2' must be less than '(2' <2>"), FFloatRangeBound::MinLower(b3_1, b3_2), b3_2);
+	TestEqual(TEXT("Open lower bound must be less than '(2' <1>"), FFloatRangeBound::MinLower(b3_3, b3_1), b3_3);
+	TestEqual(TEXT("Open lower bound must be less than '(2' <2>"), FFloatRangeBound::MinLower(b3_1, b3_3), b3_3);
+	TestEqual(TEXT("Open lower bound must be less than '[2' <1>"), FFloatRangeBound::MinLower(b3_3, b3_2), b3_3);
+	TestEqual(TEXT("Open lower bound must be less than '[2' <2>"), FFloatRangeBound::MinLower(b3_2, b3_3), b3_3);
+
+	TestEqual(TEXT("'(2' must be greater than '[2' <1>"), FFloatRangeBound::MaxLower(b3_2, b3_1), b3_1);
+	TestEqual(TEXT("'(2' must be greater than '[2' <2>"), FFloatRangeBound::MaxLower(b3_1, b3_2), b3_1);
+	TestEqual(TEXT("'(2' must be greater than open lower bound <1>"), FFloatRangeBound::MaxLower(b3_3, b3_1), b3_1);
+	TestEqual(TEXT("'(2' must be greater than open lower bound <2>"), FFloatRangeBound::MaxLower(b3_1, b3_3), b3_1);
+	TestEqual(TEXT("'[2' must be greater than open lower bound <1>"), FFloatRangeBound::MaxLower(b3_3, b3_2), b3_2);
+	TestEqual(TEXT("'[2' must be greater than open lower bound <2>"), FFloatRangeBound::MaxLower(b3_2, b3_3), b3_2);
+
+	TestEqual(TEXT("'2)' must be less than '2]' <1>"), FFloatRangeBound::MinUpper(b3_2, b3_1), b3_1);
+	TestEqual(TEXT("'2)' must be less than '2]' <2>"), FFloatRangeBound::MinUpper(b3_1, b3_2), b3_1);
+	TestEqual(TEXT("'2)' must be less than open upper bound <1>"), FFloatRangeBound::MinUpper(b3_3, b3_1), b3_1);
+	TestEqual(TEXT("'2)' must be less than open upper bound <2>"), FFloatRangeBound::MinUpper(b3_1, b3_3), b3_1);
+	TestEqual(TEXT("'2]' must be less than open upper bound <1>"), FFloatRangeBound::MinUpper(b3_3, b3_2), b3_2);
+	TestEqual(TEXT("'2]' must be less than open upper bound"), FFloatRangeBound::MinUpper(b3_2, b3_3), b3_2);
+
+	TestEqual(TEXT("'2]' must be greater than '2)' <1>"), FFloatRangeBound::MaxUpper(b3_2, b3_1), b3_2);
+	TestEqual(TEXT("'2]' must be greater than '2)' <2>"), FFloatRangeBound::MaxUpper(b3_1, b3_2), b3_2);
+	TestEqual(TEXT("Open upper bound must be greater than '2)' <1>"), FFloatRangeBound::MaxUpper(b3_3, b3_1), b3_3);
+	TestEqual(TEXT("Open upper bound must be greater than '2)' <2>"), FFloatRangeBound::MaxUpper(b3_1, b3_3), b3_3);
+	TestEqual(TEXT("Open upper bound must be greater than '2]' <1>"), FFloatRangeBound::MaxUpper(b3_3, b3_2), b3_3);
+	TestEqual(TEXT("Open upper bound must be greater than '2]' <2>"), FFloatRangeBound::MaxUpper(b3_2, b3_3), b3_3);
+
+	FFloatRangeBound b4_1 = FFloatRangeBound::Exclusive(3);
+	FFloatRangeBound b4_2 = FFloatRangeBound::Inclusive(3);
+
+	TestEqual(TEXT("'(2' must be less than '[3' <1>"), FFloatRangeBound::MinLower(b3_1, b4_2), b3_1);
+	TestEqual(TEXT("'(2' must be less than '[3' <2>"), FFloatRangeBound::MinLower(b4_2, b3_1), b3_1);
+	TestEqual(TEXT("'[2' must be less than '[3' <1>"), FFloatRangeBound::MinLower(b3_2, b4_2), b3_2);
+	TestEqual(TEXT("'[2' must be less than '[3' <2>"), FFloatRangeBound::MinLower(b4_2, b3_2), b3_2);
+
+	TestEqual(TEXT("'[3' must be greater than '(2' <1>"), FFloatRangeBound::MaxLower(b3_1, b4_2), b4_2);
+	TestEqual(TEXT("'[3' must be greater than '(2' <2>"), FFloatRangeBound::MaxLower(b4_2, b3_1), b4_2);
+	TestEqual(TEXT("'[3' must be greater than '[2' <1>"), FFloatRangeBound::MaxLower(b3_2, b4_2), b4_2);
+	TestEqual(TEXT("'[3' must be greater than '[2' <2>"), FFloatRangeBound::MaxLower(b4_2, b3_2), b4_2);
+
+	return true;
+}
+
+#endif //WITH_DEV_AUTOMATION_TESTS
diff --git a/Engine/Source/Runtime/Core/Private/Tests/Math/RangeTest.cpp b/Engine/Source/Runtime/Core/Private/Tests/Math/RangeTest.cpp
new file mode 100644
--- /dev/null
+++ b/Engine/Source/Runtime/Core/Private/Tests/Math/RangeTest.cpp
@@ -0,0 +1,402 @@
+// Copyright Epic Games, Inc. All Rights Reserved.
+
+#include "CoreTypes.h"
+#include "Containers/Array.h"
+#include "Math/Range.h"
+#include "Misc/AutomationTest.h"
+
+#if WITH_DEV_AUTOMATION_TESTS
+
+// Disable optimization for RangeTest as it compiles very slowly in development builds
+PRAGMA_DISABLE_OPTIMIZATION
+
+
+IMPLEMENT_SIMPLE_AUTOMATION_TEST(FRangeTest, "System.Core.Math.Range", EAutomationTestFlags::EditorContext | EAutomationTestFlags::ClientContext | EAutomationTestFlags::SmokeFilter)
+
+bool FRangeTest::RunTest(const FString& Parameters)
+{
+	// single element constructor
+	FFloatRange r1_1 = FFloatRange(3.0f);
+
+	TestEqual(TEXT("Single element constructor must create the correct lower bound value"), r1_1.GetLowerBoundValue(), 3.0f);
+	TestEqual(TEXT("Single element constructor must create the correct upper bound value"), r1_1.GetUpperBoundValue(), 3.0f);
+	TestTrue(TEXT("Single element constructor must create an inclusive lower bound"), r1_1.GetLowerBound().IsInclusive());
+	TestTrue(TEXT("Single element constructor must create an inclusive upper bound"), r1_1.GetUpperBound().IsInclusive());
+	TestFalse(TEXT("Single element constructor must create non-empty range"), r1_1.IsEmpty());
+
+	// explicit element pair constructor
+	FFloatRange r2_1 = FFloatRange(1.0f, 4.0f);
+
+	TestEqual(TEXT("Explicit element pair constructor must create the correct lower bound value"), r2_1.GetLowerBoundValue(), 1.0f);
+	TestEqual(TEXT("Explicit element pair constructor must create the correct upper bound value"), r2_1.GetUpperBoundValue(), 4.0f);
+	TestTrue(TEXT("Explicit element pair constructor must create an inclusive lower bound"), r2_1.GetLowerBound().IsInclusive());
+	TestTrue(TEXT("Explicit element pair constructor must create an exclusive upper bound"), r2_1.GetUpperBound().IsExclusive());
+	TestFalse(TEXT("Explicit element pair constructor must create non-empty range"), r2_1.IsEmpty());
+
+	// bound pair constructors
+	FFloatRange r3_1 = FFloatRange(3.0f);
+	FFloatRange r3_2 = FFloatRange(1.0f, 4.0f);
+	FFloatRange r3_3 = FFloatRange(FFloatRangeBound(3.0f), FFloatRangeBound(3.0f));
+	FFloatRange r3_4 = FFloatRange(FFloatRangeBound::Inclusive(1.0f), FFloatRangeBound::Exclusive(4.0f));
+
+	TestEqual(TEXT("Default bound pair constructor must create the correct range"), r3_3, r3_1);
+	TestFalse(TEXT("Default bound pair constructor must create non-empty range"), r3_3.IsEmpty());
+	TestEqual(TEXT("Specific bound pair constructor must create the correct range"), r3_4, r3_2);
+	TestFalse(TEXT("Specific bound pair constructor must create non-empty range"), r3_4.IsEmpty());
+
+	// empty ranges
+	FFloatRange r4_1 = FFloatRange::Empty();
+	FFloatRange r4_2 = FFloatRange(4.0f, 1.0f);
+	FFloatRange r4_3 = FFloatRange(FFloatRangeBound::Inclusive(6.0f), FFloatRangeBound::Exclusive(2.0f));
+	FFloatRange r4_4 = FFloatRange(FFloatRangeBound::Exclusive(4.0f), FFloatRangeBound::Exclusive(4.0f));
+	FFloatRange r4_5 = FFloatRange(FFloatRangeBound::Exclusive(4.0f), FFloatRangeBound::Inclusive(4.0f));
+	FFloatRange r4_6 = FFloatRange(FFloatRangeBound::Inclusive(4.0f), FFloatRangeBound::Exclusive(4.0f));
+
+	TestTrue(TEXT("Empty range must be empty <1>"), r4_1.IsEmpty());
+	TestTrue(TEXT("Empty range must be empty <2>"), r4_2.IsEmpty());
+	TestTrue(TEXT("Empty range must be empty <3>"), r4_3.IsEmpty());
+	TestTrue(TEXT("Empty range must be empty <4>"), r4_4.IsEmpty());
+	TestTrue(TEXT("Empty range must be empty <5>"), r4_5.IsEmpty());
+	TestTrue(TEXT("Empty range must be empty <6>"), r4_6.IsEmpty());
+
+	TestEqual(TEXT("Empty ranges must be equal <1>"), r4_1, r4_2);
+	TestEqual(TEXT("Empty ranges must be equal <2>"), r4_1, r4_3);
+	TestEqual(TEXT("Empty ranges must be equal <3>"), r4_1, r4_4);
+	TestEqual(TEXT("Empty ranges must be equal <4>"), r4_1, r4_5);
+	TestEqual(TEXT("Empty ranges must be equal <5>"), r4_1, r4_6);
+
+	// inclusive-exclusive comparison
+	FFloatRange r5_1 = FFloatRange(FFloatRangeBound::Exclusive(1.0f), FFloatRangeBound::Exclusive(4.0f));
+	FFloatRange r5_2 = FFloatRange(FFloatRangeBound::Exclusive(1.0f), FFloatRangeBound::Inclusive(4.0f));
+	FFloatRange r5_3 = FFloatRange(FFloatRangeBound::Inclusive(1.0f), FFloatRangeBound::Exclusive(4.0f));
+	FFloatRange r5_4 = FFloatRange(FFloatRangeBound::Inclusive(1.0f), FFloatRangeBound::Inclusive(4.0f));
+
+	TestEqual(TEXT("Same inclusive-exclusive ranges must be equal <1>"), r5_1, r5_1);
+	TestEqual(TEXT("Same inclusive-exclusive ranges must be equal <2>"), r5_2, r5_2);
+	TestEqual(TEXT("Same inclusive-exclusive ranges must be equal <3>"), r5_3, r5_3);
+	TestEqual(TEXT("Same inclusive-exclusive ranges must be equal <4>"), r5_4, r5_4);
+
+	TestNotEqual(TEXT("Different inclusive-exclusive ranges must be different <1>"), r5_1, r5_2);
+	TestNotEqual(TEXT("Different inclusive-exclusive ranges must be different <2>"), r5_1, r5_3);
+	TestNotEqual(TEXT("Different inclusive-exclusive ranges must be different <3>"), r5_1, r5_4);
+	TestNotEqual(TEXT("Different inclusive-exclusive ranges must be different <4>"), r5_2, r5_3);
+	TestNotEqual(TEXT("Different inclusive-exclusive ranges must be different <5>"), r5_2, r5_4);
+	TestNotEqual(TEXT("Different inclusive-exclusive ranges must be different <6>"), r5_3, r5_4);
+
+	// adjoined ranges
+	FFloatRange r6_1 = FFloatRange(FFloatRangeBound::Exclusive(1.0f), FFloatRangeBound::Exclusive(4.0f));
+	FFloatRange r6_2 = FFloatRange(FFloatRangeBound::Exclusive(1.0f), FFloatRangeBound::Inclusive(4.0f));
+	FFloatRange r6_3 = FFloatRange(FFloatRangeBound::Exclusive(4.0f), FFloatRangeBound::Inclusive(8.0f));
+	FFloatRange r6_4 = FFloatRange(FFloatRangeBound::Inclusive(4.0f), FFloatRangeBound::Inclusive(8.0f));
+	FFloatRange r6_5 = FFloatRange(FFloatRangeBound::Inclusive(3.0f), FFloatRangeBound::Inclusive(8.0f));
+	FFloatRange r6_6 = FFloatRange(FFloatRangeBound::Inclusive(0.0f), FFloatRangeBound::Inclusive(2.0f));
+
+	TestTrue(TEXT("Adjoined ranges must be adjoined"), r6_1.Adjoins(r6_4));
+	TestTrue(TEXT("Adjoined ranges must be adjoined"), r6_2.Adjoins(r6_3));
+	TestFalse(TEXT("Separated ranges must not be adjoined"), r6_1.Adjoins(r6_3));
+	TestFalse(TEXT("Separated ranges must not be adjoined"), r6_2.Adjoins(r6_4));
+	TestFalse(TEXT("Overlapped ranges must not be adjoined <1>"), r6_1.Adjoins(r6_5));
+	TestFalse(TEXT("Overlapped ranges must not be adjoined <2>"), r6_1.Adjoins(r6_6));
+
+	FFloatRange r6_7 = FFloatRange(FFloatRangeBound::Exclusive(4.0f), FFloatRangeBound::Exclusive(2.0f));
+	FFloatRange r6_8 = FFloatRange(FFloatRangeBound::Inclusive(4.0f), FFloatRangeBound::Inclusive(2.0f));
+
+	TestFalse(TEXT("A non-empty and an empty range must not be adjoined <1>"), r6_1.Adjoins(r6_8));
+	TestFalse(TEXT("A non-empty and an empty range must not be adjoined <2>"), r6_2.Adjoins(r6_7));
+	
+	FFloatRange r6_9 = FFloatRange(FFloatRangeBound::Inclusive(2.0f), FFloatRangeBound::Exclusive(0.0f));
+	FFloatRange r6_10 = FFloatRange(FFloatRangeBound::Exclusive(2.0f), FFloatRangeBound::Exclusive(0.0f));
+
+	TestFalse(TEXT("Two empty ranges must not be adjoined <1>"), r6_7.Adjoins(r6_9));
+	TestFalse(TEXT("Two empty ranges must not be adjoined <2>"), r6_8.Adjoins(r6_10));
+
+	// @todo gmp: conjoined ranges
+
+	// @todo gmp: element containment
+
+	// @todo gmp: range containment
+
+	// @todo gmp: contiguous ranges
+
+	// bounds access
+	FFloatRange r11_1 = FFloatRange(FFloatRangeBound::Inclusive(1.0f), FFloatRangeBound::Exclusive(4.0f));
+	FFloatRange r11_2 = FFloatRange(FFloatRangeBound::Open(), FFloatRangeBound::Open());
+
+	TestTrue(TEXT("A closed range must have a lower bound"), r11_1.HasLowerBound());
+	TestTrue(TEXT("A closed range must have an upper bound"), r11_1.HasUpperBound());
+	TestFalse(TEXT("An open range must not have a lower bound"), r11_2.HasLowerBound());
+	TestFalse(TEXT("An open range must not have an upper bound"), r11_2.HasUpperBound());
+
+	TestEqual(TEXT("The lower bound value of [1, 4) must be 1"), r11_1.GetLowerBoundValue(), 1.0f);
+	TestEqual(TEXT("The upper bound value of [1, 4) must be 4"), r11_1.GetUpperBoundValue(), 4.0f);
+
+	// degenerate ranges
+	FFloatRange r12_1 = FFloatRange(3.0f);
+	FFloatRange r12_2 = FFloatRange(FFloatRangeBound::Inclusive(3.0f), FFloatRangeBound::Inclusive(3.0f));
+
+	TestTrue(TEXT("A range with a single element must be degenerate <1>"), r12_1.IsDegenerate());
+	TestTrue(TEXT("A range with a single element must be degenerate <2>"), r12_2.IsDegenerate());
+
+	// @todo gmp: overlapping ranges
+	FFloatRange r14_1 = FFloatRange(FFloatRangeBound::Exclusive(0.0f), FFloatRangeBound::Exclusive(2.0f));
+	FFloatRange r14_2 = FFloatRange(FFloatRangeBound::Inclusive(0.0f), FFloatRangeBound::Inclusive(2.0f));
+	FFloatRange r14_3 = FFloatRange(FFloatRangeBound::Exclusive(2.0f), FFloatRangeBound::Exclusive(4.0f));
+	FFloatRange r14_4 = FFloatRange(FFloatRangeBound::Inclusive(2.0f), FFloatRangeBound::Inclusive(4.0f));
+	FFloatRange r14_5 = FFloatRange(FFloatRangeBound::Inclusive(3.0f), FFloatRangeBound::Inclusive(5.0f));
+	FFloatRange r14_6 = FFloatRange(FFloatRangeBound::Inclusive(1.0f), FFloatRangeBound::Inclusive(3.0f));
+	FFloatRange r14_7 = FFloatRange(FFloatRangeBound::Exclusive(1.0f), FFloatRangeBound::Inclusive(3.0f));
+	FFloatRange r14_8 = FFloatRange(FFloatRangeBound::Inclusive(1.0f), FFloatRangeBound::Exclusive(3.0f));
+	FFloatRange r14_9 = FFloatRange(FFloatRangeBound::Exclusive(6.0f), FFloatRangeBound::Exclusive(10.0f));
+	FFloatRange r14_10 = FFloatRange(FFloatRangeBound::Inclusive(4.0f), FFloatRangeBound::Open());
+	FFloatRange r14_11 = FFloatRange(FFloatRangeBound::Open(), FFloatRangeBound::Exclusive(3.0f));
+	
+	TestTrue(TEXT("(2, 4) must overlap [3, 5]"), r14_3.Overlaps(r14_5));
+	TestTrue(TEXT("(2, 4) must overlap [1, 3]"), r14_3.Overlaps(r14_6));
+	TestTrue(TEXT("(2, 4) must overlap (1, 3]"), r14_3.Overlaps(r14_7));
+	TestTrue(TEXT("(2, 4) must overlap [1, 3)"), r14_3.Overlaps(r14_8));
+	TestTrue(TEXT("(2, 4) must overlap [2, 4]"), r14_3.Overlaps(r14_4));
+	TestTrue(TEXT("(2, 4) must overlap itself"), r14_3.Overlaps(r14_3));
+	TestTrue(TEXT("(2, 4) must not overlap (0, 2)"), !r14_3.Overlaps(r14_1));
+	TestTrue(TEXT("(2, 4) must not overlap [0, 2]"), !r14_3.Overlaps(r14_2));
+	TestTrue(TEXT("[2, 4] must not overlap (0, 2)"), !r14_4.Overlaps(r14_1));
+	TestTrue(TEXT("[2, 4] must overlap [0, 2]"), r14_4.Overlaps(r14_2));
+	TestTrue(TEXT("[2, 4] must not overlap (6, 10)"), !r14_4.Overlaps(r14_9));
+	TestTrue(TEXT("[2, 4] must overlap [4, inf)"), r14_4.Overlaps(r14_10));
+	TestTrue(TEXT("(2, 4) must not overlap [4, inf)"), !r14_3.Overlaps(r14_10));
+	TestTrue(TEXT("[1, 3] must not overlap [4, inf)"), !r14_6.Overlaps(r14_10));
+	TestTrue(TEXT("[1, 3] must overlap (inf, 3)"), r14_6.Overlaps(r14_11));
+	TestTrue(TEXT("(6, 10) must not overlap (inf, 3)"), !r14_9.Overlaps(r14_11));
+
+	// @todo gmp: range splitting
+
+	// difference
+	FFloatRange r15_1 = FFloatRange(FFloatRangeBound::Exclusive(7.0f), FFloatRangeBound::Exclusive(14.0f));		// X
+	FFloatRange r15_2 = FFloatRange(FFloatRangeBound::Inclusive(7.0f), FFloatRangeBound::Inclusive(14.0f));
+
+	TestEqual(TEXT("The difference between a range and itself must be an empty set <1>"), FFloatRange::Difference(r15_1, r15_1).Num(), 0);
+	TestEqual(TEXT("The difference between a range and itself must be an empty set <2>"), FFloatRange::Difference(r15_2, r15_2).Num(), 0);
+	
+	TArray<FFloatRange > d15_0 = FFloatRange::Difference(r15_2, r15_1);
+	{
+		TestEqual(TEXT("[7, 14] - (7, 14) must result in two ranges"), d15_0.Num(), 2);
+
+		bool b15_1 = (d15_0[0] == FFloatRange(FFloatRangeBound::Inclusive(7.0f), FFloatRangeBound::Inclusive(7.0f)));
+		bool b15_2 = (d15_0[1] == FFloatRange(FFloatRangeBound::Inclusive(14.0f), FFloatRangeBound::Inclusive(14.0f)));
+
+		TestTrue(TEXT("[7, 14] - (7, 14) must be {[7, 7], [14, 14]}"), b15_1 && b15_2);
+	}
+
+	FFloatRange r15_3 = FFloatRange(FFloatRangeBound::Inclusive(2.0f), FFloatRangeBound::Exclusive(9.0f));		// Y overlapping left
+	FFloatRange r15_4 = FFloatRange(FFloatRangeBound::Inclusive(2.0f), FFloatRangeBound::Inclusive(9.0f));
+	FFloatRange r15_5 = FFloatRange(FFloatRangeBound::Exclusive(8.0f), FFloatRangeBound::Exclusive(17.0f));		// Y overlapping right
+	FFloatRange r15_6 = FFloatRange(FFloatRangeBound::Inclusive(8.0f), FFloatRangeBound::Exclusive(17.0f));
+
+	TArray<FFloatRange > d15_1 = FFloatRange::Difference(r15_1, r15_3);
+	{
+		TestEqual(TEXT("(7, 14) - [2, 9) must result in one range"), d15_1.Num(), 1);
+		TestEqual(TEXT("(7, 14) - [2, 9) must be {[9, 14)}"), d15_1[0], FFloatRange(FFloatRangeBound::Inclusive(9.0f), FFloatRangeBound::Exclusive(14.0f)));
+	}
+
+	TArray<FFloatRange > d15_2 = FFloatRange::Difference(r15_1, r15_4);
+	{
+		TestEqual(TEXT("(7, 14) - [2, 9] must result in one range"), d15_2.Num(), 1);
+		TestEqual(TEXT("(7, 14) - [2, 9] must be {(9, 14)}"), d15_2[0], FFloatRange(FFloatRangeBound::Exclusive(9.0f), FFloatRangeBound::Exclusive(14.0f)));
+	}
+
+	TArray<FFloatRange > d15_3 = FFloatRange::Difference(r15_1, r15_5);
+	{
+		TestEqual(TEXT("(7, 14) - (8, 17) must result in one range"), d15_3.Num(), 1);
+		TestEqual(TEXT("(7, 14) - (8, 17) must be {(7, 8]}"), d15_3[0], FFloatRange(FFloatRangeBound::Exclusive(7.0f), FFloatRangeBound::Inclusive(8.0f)));
+	}
+
+	TArray<FFloatRange > d15_4 = FFloatRange::Difference(r15_1, r15_6);
+	{
+		TestEqual(TEXT("(7, 14) - [8, 17) must result in one range"), d15_4.Num(), 1);
+		TestEqual(TEXT("(7, 14) - [8, 17) must be {(7, 8)}"), d15_4[0], FFloatRange(FFloatRangeBound::Exclusive(7.0f), FFloatRangeBound::Exclusive(8.0f)));
+	}
+
+	TArray<FFloatRange > d15_5 = FFloatRange::Difference(r15_2, r15_3);
+	{
+		TestEqual(TEXT("[7, 14] - [2, 9) must result in one range"), d15_5.Num(), 1);
+		TestEqual(TEXT("[7, 14] - [2, 9) must be {[9, 14]}"), d15_5[0], FFloatRange(FFloatRangeBound::Inclusive(9.0f), FFloatRangeBound::Inclusive(14.0f)));
+	}
+
+	TArray<FFloatRange > d15_6 = FFloatRange::Difference(r15_2, r15_4);
+	{
+		TestEqual(TEXT("[7, 14] - [2, 9] must result in one range"), d15_6.Num(), 1);
+		TestEqual(TEXT("[7, 14] - [2, 9] must be {(9, 14]}"), d15_6[0], FFloatRange(FFloatRangeBound::Exclusive(9.0f), FFloatRangeBound::Inclusive(14.0f)));
+	}
+
+	TArray<FFloatRange > d15_7 = FFloatRange::Difference(r15_2, r15_5);
+	{
+		TestEqual(TEXT("[7, 14] - (8, 17) must result in one range"), d15_7.Num(), 1);
+		TestEqual(TEXT("[7, 14] - (8, 17) must be {[7, 8]}"), d15_7[0], FFloatRange(FFloatRangeBound::Inclusive(7.0f), FFloatRangeBound::Inclusive(8.0f)));
+	}
+
+	TArray<FFloatRange > d15_8 = FFloatRange::Difference(r15_2, r15_6);
+	{
+		TestEqual(TEXT("[7, 14] - [8, 17) must result in one range"), d15_8.Num(), 1);
+		TestEqual(TEXT("[7, 14] - [8, 17) must be {[7, 8)}"), d15_8[0], FFloatRange(FFloatRangeBound::Inclusive(7.0f), FFloatRangeBound::Exclusive(8.0f)));
+	}
+
+	FFloatRange r15_7 = FFloatRange(FFloatRangeBound::Inclusive(2.0f), FFloatRangeBound::Exclusive(7.0f));		// Y adjoint left
+	FFloatRange r15_8 = FFloatRange(FFloatRangeBound::Inclusive(2.0f), FFloatRangeBound::Inclusive(7.0f));
+
+	TArray<FFloatRange > d15_9 = FFloatRange::Difference(r15_1, r15_8);
+	{
+		TestEqual(TEXT("(7, 14) - [2, 7] must result in one set"), d15_9.Num(), 1);
+		TestEqual(TEXT("(7, 14) - [2, 7] must be {(7, 14)}"), d15_9[0], FFloatRange(FFloatRangeBound::Exclusive(7.0f), FFloatRangeBound::Exclusive(14.0f)));
+	}
+
+	TArray<FFloatRange > d15_10 = FFloatRange::Difference(r15_2, r15_7);
+	{
+		TestEqual(TEXT("[7, 14] - [2, 7) must result in one set"), d15_10.Num(), 1);
+		TestEqual(TEXT("[7, 14] - [2, 7) must be {[7, 14]}"), d15_10[0], FFloatRange(FFloatRangeBound::Inclusive(7.0f), FFloatRangeBound::Inclusive(14.0f)));
+	}
+
+	TArray<FFloatRange > d15_11 = FFloatRange::Difference(r15_2, r15_8);
+	{
+		TestEqual(TEXT("[7, 14] - [2, 7] must result in one range"), d15_11.Num(), 1);
+		TestEqual(TEXT("[7, 14] - [2, 7] must be {(7, 14]}"), d15_11[0], FFloatRange(FFloatRangeBound::Exclusive(7.0f), FFloatRangeBound::Inclusive(14.0f)));
+	}
+
+	FFloatRange r15_9 = FFloatRange(FFloatRangeBound::Exclusive(14.0f), FFloatRangeBound::Exclusive(16.0f));	// Y adjoint right
+	FFloatRange r15_10 = FFloatRange(FFloatRangeBound::Inclusive(14.0f), FFloatRangeBound::Inclusive(16.0f));
+
+	TArray<FFloatRange > d15_12 = FFloatRange::Difference(r15_1, r15_10);
+	{
+		TestEqual(TEXT("(7, 14) - [14, 18] must result in one range"), d15_12.Num(), 1);
+		TestEqual(TEXT("(7, 14) - [14, 18] must be {(7, 14)}"), d15_12[0], FFloatRange(FFloatRangeBound::Exclusive(7.0f), FFloatRangeBound::Exclusive(14.0f)));
+	}
+
+	TArray<FFloatRange > d15_13 = FFloatRange::Difference(r15_2, r15_9);
+	{
+		TestEqual(TEXT("[7, 14] - (14, 18) must result in one range"), d15_13.Num(), 1);
+		TestEqual(TEXT("[7, 14] - (14, 18) must be {[7, 14]}"), d15_13[0], FFloatRange(FFloatRangeBound::Inclusive(7.0f), FFloatRangeBound::Inclusive(14.0f)));
+	}
+
+	TArray<FFloatRange > d15_14 = FFloatRange::Difference(r15_2, r15_10);
+	{
+		TestEqual(TEXT("[7, 14] - [14, 18] must result in one range"), d15_14.Num(), 1);
+		TestEqual(TEXT("[7, 14] - [14, 18] must be {(7, 14)}"), d15_14[0], FFloatRange(FFloatRangeBound::Inclusive(7.0f), FFloatRangeBound::Exclusive(14.0f)));
+	}
+
+	FFloatRange r15_11 = FFloatRange(FFloatRangeBound::Inclusive(8.0f), FFloatRangeBound::Inclusive(13.0f));	// Y enclosed
+
+	TArray<FFloatRange > d15_15 = FFloatRange::Difference(r15_2, r15_11);
+	{
+		TestEqual(TEXT("[7, 14] - [8, 13] must result in two ranges"), d15_15.Num(), 2);
+
+		bool b15_1 = (d15_15[0] == FFloatRange(FFloatRangeBound::Inclusive(7.0f), FFloatRangeBound::Exclusive(8.0f)));
+		bool b15_2 = (d15_15[1] == FFloatRange(FFloatRangeBound::Exclusive(13.0f), FFloatRangeBound::Inclusive(14.0f)));
+
+		TestTrue(TEXT("[7, 14] - [8, 13] must be {[7, 8), (13, 14]}"), b15_1 && b15_2);
+	}
+	
+	FFloatRange r15_12 = FFloatRange(FFloatRangeBound::Inclusive(2.0f), FFloatRangeBound::Inclusive(4.0f));
+
+	TArray<FFloatRange > d15_16 = FFloatRange::Difference(r15_2, r15_12);
+	{
+		TestEqual(TEXT("[7, 14] - [2, 4) must result in one set"), d15_16.Num(), 1);
+		TestEqual(TEXT("[7, 14] - [2, 4) must be {[7, 14]}"), d15_16[0], FFloatRange(FFloatRangeBound::Inclusive(7.0f), FFloatRangeBound::Inclusive(14.0f)));
+	}
+
+	// range hulls
+	FFloatRange r16_1 = FFloatRange(FFloatRangeBound::Inclusive(7.0f), FFloatRangeBound::Exclusive(9.0f));
+	FFloatRange r16_2 = FFloatRange(FFloatRangeBound::Inclusive(11.0f), FFloatRangeBound::Exclusive(14.0f));
+	FFloatRange r16_3 = FFloatRange(4.0f, 1.0f);
+	FFloatRange r16_4 = FFloatRange(8.0f, 5.0f);
+
+	TestEqual(TEXT("The hull of [7, 9) and [11, 14) must be [7, 14)"), FFloatRange::Hull(r16_1, r16_2), FFloatRange(FFloatRangeBound::Inclusive(7.0f), FFloatRangeBound::Exclusive(14.0f)));
+	TestEqual(TEXT("The hull of [7, 9) and an empty range must be [7, 9)"), FFloatRange::Hull(r16_1, r16_3), FFloatRange(FFloatRangeBound::Inclusive(7.0f), FFloatRangeBound::Exclusive(9.0f)));
+	TestTrue(TEXT("The hull of two empty ranges must be empty"), FFloatRange::Hull(r16_3, r16_4).IsEmpty());
+
+	// intersections
+	FFloatRange r17_1 = FFloatRange(FFloatRangeBound::Inclusive(7.0f), FFloatRangeBound::Exclusive(14.0f));
+	FFloatRange r17_2 = FFloatRange(FFloatRangeBound::Inclusive(2.0f), FFloatRangeBound::Exclusive(8.0f));		// left overlap
+	FFloatRange r17_3 = FFloatRange(FFloatRangeBound::Inclusive(2.0f), FFloatRangeBound::Inclusive(8.0f));
+	FFloatRange r17_4 = FFloatRange(FFloatRangeBound::Exclusive(8.0f), FFloatRangeBound::Exclusive(16.0f));		// right overlap
+	FFloatRange r17_5 = FFloatRange(FFloatRangeBound::Inclusive(8.0f), FFloatRangeBound::Exclusive(16.0f));
+	FFloatRange r17_6 = FFloatRange(4.0f, 1.0f);
+	FFloatRange r17_7 = FFloatRange(8.0f, 5.0f);
+
+	TestEqual(TEXT("The intersection of [7, 14) and [2, 8) must be [7, 8)"), FFloatRange::Intersection(r17_1, r17_2), FFloatRange(FFloatRangeBound::Inclusive(7.0f), FFloatRangeBound::Exclusive(8.0f)));
+	TestEqual(TEXT("The intersection of [7, 14) and [2, 8] must be [7, 8]"), FFloatRange::Intersection(r17_1, r17_3), FFloatRange(FFloatRangeBound::Inclusive(7.0f), FFloatRangeBound::Inclusive(8.0f)));
+	TestEqual(TEXT("The intersection of [7, 14) and (8, 16) must be (8, 14)"), FFloatRange::Intersection(r17_1, r17_4), FFloatRange(FFloatRangeBound::Exclusive(8.0f), FFloatRangeBound::Exclusive(14.0f)));
+	TestEqual(TEXT("The intersection of [7, 14) and [8, 16) must be [8, 14)"), FFloatRange::Intersection(r17_1, r17_5), FFloatRange(FFloatRangeBound::Inclusive(8.0f), FFloatRangeBound::Exclusive(14.0f)));
+	TestTrue(TEXT("The intersection of a non-empty range and an empty range must be empty"), FFloatRange::Intersection(r17_1, r17_6).IsEmpty());
+	TestTrue(TEXT("The intersection of two empty ranges must be empty"), FFloatRange::Intersection(r17_6, r17_7).IsEmpty());
+
+	// union
+	FFloatRange r18_1 = FFloatRange(FFloatRangeBound::Inclusive(7.0f), FFloatRangeBound::Inclusive(14.0f));
+	FFloatRange r18_2 = FFloatRange(FFloatRangeBound::Exclusive(7.0f), FFloatRangeBound::Exclusive(14.0f));
+	FFloatRange r18_3 = FFloatRange(FFloatRangeBound::Inclusive(2.0f), FFloatRangeBound::Inclusive(8.0f));		// left overlap
+	FFloatRange r18_4 = FFloatRange(FFloatRangeBound::Exclusive(15.0f), FFloatRangeBound::Inclusive(16.0f));
+	FFloatRange r18_5 = FFloatRange(FFloatRangeBound::Inclusive(8.0f), FFloatRangeBound::Exclusive(16.0f));		// right overlap
+	FFloatRange r18_6 = FFloatRange(FFloatRangeBound::Inclusive(6.0f), FFloatRangeBound::Exclusive(7.0f));
+	
+	TArray<FFloatRange > u18_0 = FFloatRange::Union(r18_1, r18_1);
+	{
+		TestEqual(TEXT("[7, 14] unioned with itself must result in one range"), u18_0.Num(), 1);
+		TestEqual(TEXT("[7, 14] unioned with itself must result in itself"), u18_0[0], r18_1);
+	}
+
+	TArray<FFloatRange > u18_1 = FFloatRange::Union(r18_1, r18_2);
+	{
+		TestEqual(TEXT("[7, 14] unioned with (7, 14) must result in one range"), u18_1.Num(), 1);
+		TestEqual(TEXT("[7, 14] unioned with (7, 14) must result in {[7, 14]}"), u18_1[0], r18_1);
+	}
+
+	TArray<FFloatRange > u18_2 = FFloatRange::Union(r18_1, r18_3);
+	{
+		TestEqual(TEXT("[7, 14] unioned with [2, 8] must result in one range"), u18_2.Num(), 1);
+		TestEqual(TEXT("[7, 14] unioned with [2, 8] must result in {[2, 14]}"), u18_2[0], FFloatRange(FFloatRangeBound::Inclusive(2.f), FFloatRangeBound::Inclusive(14.f)));
+	}
+
+	TArray<FFloatRange > u18_3 = FFloatRange::Union(r18_1, r18_4);
+	{
+		TestEqual(TEXT("[7, 14] unioned with (15, 16] must result in two ranges"), u18_3.Num(), 2);
+		
+		bool b18_1 = (u18_3[0] == r18_1);
+		bool b18_2 = (u18_3[1] == r18_4);
+
+		TestTrue(TEXT("[7, 14] unioned with (15, 16] must result in {[7, 14], (15, 16]}"), b18_1 && b18_2);
+	}
+
+	TArray<FFloatRange > u18_4 = FFloatRange::Union(r18_1, FFloatRange::Empty());
+	{
+		TestEqual(TEXT("A non-empty range unioned with an empty range must result in one range"), u18_4.Num(), 1);
+		TestEqual(TEXT("A non-empty range unioned with an empty range must result in the non-empty range"), u18_4[0], r18_1);
+	}
+
+	TArray<FFloatRange > u18_5 = FFloatRange::Union(r18_1, r18_5);
+	{
+		TestEqual(TEXT("[7, 14] unioned with [8, 16) must result in one range"), u18_5.Num(), 1);
+		TestEqual(TEXT("[7, 14] unioned with [8, 16) must result in {[7, 16)}"), u18_5[0], FFloatRange(FFloatRangeBound::Inclusive(7.f), FFloatRangeBound::Exclusive(16.f)));
+	}
+
+	TArray<FFloatRange > u18_6 = FFloatRange::Union(r18_1, r18_6);
+	{
+		TestEqual(TEXT("[7, 14] unioned with [6, 7) must result in one range"), u18_6.Num(), 1);
+		TestEqual(TEXT("[7, 14] unioned with [6, 7) must result in {[6, 14]}"), u18_6[0], FFloatRange(FFloatRangeBound::Inclusive(6.f), FFloatRangeBound::Inclusive(14.f)));
+	}
+
+	TArray<FFloatRange > u18_7 = FFloatRange::Union(r18_2, r18_6);
+	{
+		TestEqual(TEXT("(7, 14) unioned with [6, 7) must result in two ranges"), u18_7.Num(), 2);
+		
+		bool b18_1 = (u18_7[0] == r18_2);
+		bool b18_2 = (u18_7[1] == r18_6);
+
+		TestTrue(TEXT("(7, 14) unioned with [6, 7) must result in {(7, 14), [6, 7)}"), b18_1 && b18_2);
+	}
+	
+	TestEqual(TEXT("An pair of empty ranges unioned must yield no ranges."), FFloatRange::Union(FFloatRange::Empty(), FFloatRange::Empty()).Num(), 0);
+
+	return true;
+}
+
+
+PRAGMA_ENABLE_OPTIMIZATION
+
+#endif //WITH_DEV_AUTOMATION_TESTS
diff --git a/Engine/Source/Runtime/Core/Private/Tests/Misc/ByteSwapTest.cpp b/Engine/Source/Runtime/Core/Private/Tests/Misc/ByteSwapTest.cpp
new file mode 100644
--- /dev/null
+++ b/Engine/Source/Runtime/Core/Private/Tests/Misc/ByteSwapTest.cpp
@@ -0,0 +1,111 @@
+// Copyright Epic Games, Inc. All Rights Reserved.
+
+#include "CoreTypes.h"
+#include "Misc/AutomationTest.h"
+#include "Misc/ByteSwap.h"
+
+#if WITH_DEV_AUTOMATION_TESTS
+
+/**
+ * Test byte swapping algorithms.
+ */
+IMPLEMENT_SIMPLE_AUTOMATION_TEST(FByteSwapTest, "System.Core.Misc.ByteSwap", EAutomationTestFlags::ApplicationContextMask | EAutomationTestFlags::SmokeFilter)
+
+bool FByteSwapTest::RunTest(const FString& Parameters)
+{
+	int16  ValS16 = static_cast<int16>(0x1122);
+	uint16 ValU16 = static_cast<uint16>(0x1122);
+	int32  ValS32 = 0xFFEE1122L;
+	uint32 ValU32 = 0xFFEE1122UL;
+	int64  ValS64 = 0xFFEEDDCC11223344ll;
+	uint64 ValU64 = 0xFFEEDDCC11223344ull;
+	float  ValF   = *reinterpret_cast<const float*>(&ValU32);
+	double ValD   = *reinterpret_cast<const double*>(&ValU64);
+	uint32 ExpectedF = ByteSwap(ValU32);
+	uint64 ExpectedD = ByteSwap(ValU64);
+	char16_t ValCh16 = static_cast<char16_t>(0x2233);
+
+	TestTrue("Swapping singned int16 value",  ByteSwap(ValS16) == 0x2211);
+	TestTrue("Swapping unsigned int16 value", ByteSwap(ValU16) == 0x2211);
+	TestTrue("Swapping singned int32 value",  ByteSwap(ValS32) == 0x2211EEFF);
+	TestTrue("Swapping unsigned int32 value", ByteSwap(ValU32) == 0x2211EEFF);
+	TestTrue("Swapping singned int64 value",  ByteSwap(ValS64) == 0x44332211CCDDEEFFll);
+	TestTrue("Swapping unsigned int64 value", ByteSwap(ValU64) == 0x44332211CCDDEEFFull);
+	TestTrue("Swapping float value",          ByteSwap(ValF)   == *reinterpret_cast<const float*>(&ExpectedF));
+	TestTrue("Swapping double value",         ByteSwap(ValD)   == *reinterpret_cast<const double*>(&ExpectedD));
+	TestTrue("Swapping char16_t value",       ByteSwap(ValCh16)== 0x3322);
+
+	return true;
+}
+
+// The byte swap benchmarking test are useful to compare the intrinsic implementation vs the generic implementation. Normally, the intrinsic is expected to be faster, but in some
+// cases, the intrinsic vs generic speed is within margin of error in optimized build. So the test result are not always consistent. The code remains there in case a new
+// implementation needs to be testes. Here some observations:
+//     - VC++ 2019 (16.4.3): The intrinsic versions are consistently (and significantly) faster than the generic version. The compiler poorly optimize the generic version.
+//     - Apple Clang (11.0.0): No performance change observed in release. Clang generates the same assembly for the intrinsic and the generic functions.
+//     - Linux Clang (8.0.1): No performance change observed in release. Clang generates the same assembly for the intrinsic and the generic functions.
+#define UE_BYTE_SWAP_BENCHMARK_ENABLED 0
+#if UE_BYTE_SWAP_BENCHMARK_ENABLED
+
+/**
+ * Compare the performance of swapping bytes using the intrinsic vs generic implementation. See comment above UE_BYTE_SWAP_BENCHMARK_ENABLED.
+ */
+IMPLEMENT_SIMPLE_AUTOMATION_TEST(FByteSwapPerformanceTest, "System.Core.Misc.ByteSwapPerf", EAutomationTestFlags::ApplicationContextMask | EAutomationTestFlags::PerfFilter)
+
+template<typename T, typename ByteSwapFn>
+TPair<FTimespan, T> BenchmarkByteswapping(T InitialValue, uint64 LoopCount, ByteSwapFn InByteSwapFn)
+{
+	static_assert(TIsIntegral<T>::Value, "The test only works with integers");
+	static_assert(!TIsSigned<T>::Value, "The test must support overflowing T when making the sum, value cannot be signed.");
+
+	T Sum = 0; // The sum value is not very relevant, but it is an observable result the compiler cannot ignore (prevent to optimize away the entire loop).
+	FDateTime StartTime = FDateTime::UtcNow();
+	for (uint64 I = 0; I < LoopCount; ++I)
+	{
+		Sum += InByteSwapFn(InitialValue + static_cast<T>(I)); // Expect InByteSwapFn to be inlined because the compiler knows the exact function to call.
+	}
+	FTimespan Duration = FDateTime::UtcNow() - StartTime;
+	return MakeTuple(Duration, Sum);
+}
+
+bool FByteSwapPerformanceTest::RunTest(const FString& Parameters)
+{
+#ifdef UE_BYTESWAP_INTRINSIC_PRIVATE_16
+	{
+		uint16 InitialValue = 0xF0F0;
+		TPair<FTimespan, uint16> Intrinsic = BenchmarkByteswapping(InitialValue, 1000000000, [](uint16 Val){ return ByteSwap(Val); });
+		TPair<FTimespan, uint16> Generic   = BenchmarkByteswapping(InitialValue, 1000000000, [](uint16 Val){ return Internal::ByteSwapGeneric16(Val); });
+		TestTrue("Swapping uint16 bytes is faster uing the compiler intrinsic than the generic implementation", Intrinsic.Get<0>() <= Generic.Get<0>());
+		TestTrue("Swapping uint16 bytes using intrinsic and generic algorithm produce the same values", Intrinsic.Get<1>() == Generic.Get<1>());
+		GLog->Logf(TEXT("Swapping 2 bytes using intrinsic is %s faster than generic version"), *LexToString(Generic.Get<0>().GetTotalMicroseconds() / Intrinsic.Get<0>().GetTotalMicroseconds()));
+	}
+#endif
+
+#ifdef UE_BYTESWAP_INTRINSIC_PRIVATE_32
+	{
+		uint32 InitialValue = 0xFF00FF00ul;
+		TPair<FTimespan, uint32> Intrinsic = BenchmarkByteswapping(InitialValue, 1000000000, [](uint32 Val){ return ByteSwap(Val); });
+		TPair<FTimespan, uint32> Generic   = BenchmarkByteswapping(InitialValue, 1000000000, [](uint32 Val){ return Internal::ByteSwapGeneric32(Val); });
+		TestTrue("Swapping uint32 bytes is faster uing the compiler intrinsic than the generic implementation", Intrinsic.Get<0>() <= Generic.Get<0>());
+		TestTrue("Swapping uint32 bytes using intrinsic and generic algorithm produce the same values", Intrinsic.Get<1>() == Generic.Get<1>());
+		GLog->Logf(TEXT("Swapping 4 bytes using intrinsic is %s faster than generic version"), *LexToString(Generic.Get<0>().GetTotalMicroseconds() / Intrinsic.Get<0>().GetTotalMicroseconds()));
+	}
+#endif
+
+#ifdef UE_BYTESWAP_INTRINSIC_PRIVATE_64
+	{
+		uint64 InitialValue = 0xFF00FF00FF00FF00ull;
+		TPair<FTimespan, uint64> Intrinsic = BenchmarkByteswapping(InitialValue, 1000000000, [](uint64 Val){ return ByteSwap(Val); });
+		TPair<FTimespan, uint64> Generic   = BenchmarkByteswapping(InitialValue, 1000000000, [](uint64 Val){ return Internal::ByteSwapGeneric64(Val); });
+		TestTrue("Swapping uint64 bytes is faster uing the compiler intrinsic than the generic implementation", Intrinsic.Get<0>() <= Generic.Get<0>());
+		TestTrue("Swapping uint64 bytes using intrinsic and generic algorithm produce the same values", Intrinsic.Get<1>() == Generic.Get<1>());
+		GLog->Logf(TEXT("Swapping 8 bytes using intrinsic is %s faster than generic version"), *LexToString(Generic.Get<0>().GetTotalMicroseconds() / Intrinsic.Get<0>().GetTotalMicroseconds()));
+	}
+#endif
+
+	return true;
+}
+
+#endif //UE_BYTE_SWAP_BENCHMARK_ENABLED
+
+#endif // WITH_DEV_AUTOMATION_TESTS
diff --git a/Engine/Source/Runtime/Core/Private/Tests/Misc/CircularBufferTest.cpp b/Engine/Source/Runtime/Core/Private/Tests/Misc/CircularBufferTest.cpp
new file mode 100644
--- /dev/null
+++ b/Engine/Source/Runtime/Core/Private/Tests/Misc/CircularBufferTest.cpp
@@ -0,0 +1,48 @@
+// Copyright Epic Games, Inc. All Rights Reserved.
+
+#include "CoreTypes.h"
+#include "Containers/UnrealString.h"
+#include "Containers/CircularBuffer.h"
+#include "Misc/AutomationTest.h"
+
+#if WITH_DEV_AUTOMATION_TESTS
+
+IMPLEMENT_SIMPLE_AUTOMATION_TEST(FCircularBufferTest, "System.Core.Misc.CircularBuffer", EAutomationTestFlags::ApplicationContextMask | EAutomationTestFlags::SmokeFilter)
+
+bool FCircularBufferTest::RunTest(const FString& Parameters)
+{
+	// buffer capacity
+	TCircularBuffer<int32> b1_1(127);
+	TCircularBuffer<int32> b1_2(128);
+	TCircularBuffer<int32> b1_3(129);
+
+	TestEqual(TEXT("Buffer capacity of 127 must be rounded up to 128"), b1_1.Capacity(), 128u);
+	TestEqual(TEXT("Buffer capacity of 128 must not change"), b1_2.Capacity(), 128u);
+	TestEqual(TEXT("Buffer capacity of 129 must be rounded up to 256"), b1_3.Capacity(), 256u);
+
+	// initial values
+	TCircularBuffer<int32> b2_1(64, 666);
+
+	for (uint32 Index = 0; Index < b2_1.Capacity(); ++Index)
+	{
+		TestEqual(FString::Printf(TEXT("Initial value must be correct (%i)"), Index), b2_1[Index], 666);
+	}
+
+	// indexing
+	TCircularBuffer<int32> b3_1(64, 0);
+
+	TestEqual(TEXT("Next index from 0 must be 1"), b3_1.GetNextIndex(0), 1u);
+	TestEqual(TEXT("Next index from 63 must be 0"), b3_1.GetNextIndex(63), 0u);
+	TestEqual(TEXT("Next index from 64 must be 1"), b3_1.GetNextIndex(64), 1u);
+
+	b3_1[0] = 42;
+	b3_1[65] = 42;
+
+	TestEqual(TEXT("Index 0 must be written and read correctly"), b3_1[0], 42);
+	TestEqual(TEXT("Index 1 must be written and read correctly"), b3_1[1], 42);
+	TestEqual(TEXT("Index 65 must be written and read correctly"), b3_1[65], 42);
+
+	return true;
+}
+
+#endif //WITH_DEV_AUTOMATION_TESTS
diff --git a/Engine/Source/Runtime/Core/Private/Tests/Misc/CircularQueueTest.cpp b/Engine/Source/Runtime/Core/Private/Tests/Misc/CircularQueueTest.cpp
new file mode 100644
--- /dev/null
+++ b/Engine/Source/Runtime/Core/Private/Tests/Misc/CircularQueueTest.cpp
@@ -0,0 +1,137 @@
+// Copyright Epic Games, Inc. All Rights Reserved.
+
+#include "CoreTypes.h"
+#include "Misc/AutomationTest.h"
+#include "Containers/CircularQueue.h"
+
+#if WITH_DEV_AUTOMATION_TESTS
+
+IMPLEMENT_SIMPLE_AUTOMATION_TEST(FCircularQueueTest, "System.Core.Misc.CircularQueue", EAutomationTestFlags::ApplicationContextMask | EAutomationTestFlags::SmokeFilter)
+
+bool FCircularQueueTest::RunTest(const FString& Parameters)
+{
+	const uint32 QueueSize = 8;
+
+	// empty queue
+	{
+		TCircularQueue<int32> Queue(QueueSize);
+
+		TestEqual(TEXT("Newly created queues must have zero elements"), Queue.Count(), 0u);
+		TestTrue(TEXT("Newly created queues must be empty"), Queue.IsEmpty());
+		TestFalse(TEXT("Newly created queues must not be full"), Queue.IsFull());
+
+		int32 Value;
+		TestFalse(TEXT("Peek must fail on an empty queue"), Queue.Peek(Value));
+		TestNull(TEXT("Peek must fail on an empty queue"), Queue.Peek());
+	}
+
+	// partially filled
+	{
+		TCircularQueue<int32> Queue(QueueSize);
+		int32 Value = 0;
+
+		TestTrue(TEXT("Adding to an empty queue must succeed"), Queue.Enqueue(666));
+		TestEqual(TEXT("After adding to an empty queue it must have one element"), Queue.Count(), 1u);
+		TestFalse(TEXT("Partially filled queues must not be empty"), Queue.IsEmpty());
+		TestFalse(TEXT("Partially filled queues must not be full"), Queue.IsFull());
+		TestTrue(TEXT("Peeking at a partially filled queue must succeed"), Queue.Peek(Value));
+		TestEqual(TEXT("The peeked at value must be correct"), Value, 666);
+
+		const int32* PeekValue = Queue.Peek();
+		TestNotNull(TEXT("Peeking at a partially filled queue must succeed"), PeekValue);
+		TestEqual(TEXT("The peeked at value must be correct"), *PeekValue, 666);
+	}
+
+	// full queue
+	for (uint32 PeekType = 0; PeekType < 2; PeekType++)
+	{
+		TCircularQueue<int32> Queue(QueueSize);
+
+		for (int32 Index = 0; Index < QueueSize - 1; ++Index)
+		{
+			TestTrue(TEXT("Adding to non-full queue must succeed"), Queue.Enqueue(Index));
+		}
+
+		TestFalse(TEXT("Full queues must not be empty"), Queue.IsEmpty());
+		TestTrue(TEXT("Full queues must be full"), Queue.IsFull());
+		TestFalse(TEXT("Adding to full queue must fail"), Queue.Enqueue(666));
+
+		int32 Value = 0;
+
+		for (int32 Index = 0; Index < QueueSize - 1; ++Index)
+		{
+			if (PeekType == 0)
+			{
+				TestTrue(TEXT("Peeking at a non-empty queue must succeed"), Queue.Peek(Value));
+				TestEqual(TEXT("The peeked at value must be correct"), Value, Index);
+
+				TestTrue(TEXT("Removing from a non-empty queue must succeed"), Queue.Dequeue(Value));
+				TestEqual(TEXT("The removed value must be correct"), Value, Index);
+			}
+			else
+			{
+				const int32* PeekValue = Queue.Peek();
+				TestNotNull(TEXT("Peeking at a non-empty queue must succeed"), PeekValue);
+				TestEqual(TEXT("The peeked at value must be correct"), *PeekValue, Index);
+
+				TestTrue(TEXT("Removing from a non-empty queue must succeed"), Queue.Dequeue());
+			}
+		}
+
+		TestTrue(TEXT("A queue that had all items removed must be empty"), Queue.IsEmpty());
+		TestFalse(TEXT("A queue that had all items removed must not be full"), Queue.IsFull());
+	}
+
+	// queue with index wrapping around
+	{
+		TCircularQueue<int32> Queue(QueueSize);
+
+		//Fill queue
+		for (int32 Index = 0; Index < QueueSize - 1; ++Index)
+		{
+			TestTrue(TEXT("Adding to non-full queue must succeed"), Queue.Enqueue(Index));
+		}
+
+		int32 Value = 0;
+		const int32 ExpectedSize = QueueSize - 1;
+		for (int32 Index = 0; Index < QueueSize; ++Index)
+		{
+			TestEqual(TEXT("Number of elements must be valid for all permutation of Tail and Head"), Queue.Count(), ExpectedSize);
+			TestTrue(TEXT("Removing from a non-empty queue must succeed"), Queue.Dequeue(Value));
+			TestTrue(TEXT("Adding to non-full queue must succeed"), Queue.Enqueue(Index));
+		}
+	}
+
+	// Non-zero initialization - ensures the backing store constructs and destructs non-POD objects properly.
+	{
+		static uint32 Const;
+		static uint32 Dest;
+		static uint32 CopyConst;
+		struct FNonPOD
+		{
+			FNonPOD() { Const++; }
+			~FNonPOD() { Dest++; }
+			FNonPOD(const FNonPOD&) { CopyConst++; }
+		};
+
+		// The current implementation of TCircularQueue doesn't call the held object constructors on but does call the
+		// destructors. Verify that (somewhat surprising, if typically ok for POD) behavior first.
+		Const = 0;
+		Dest = 0;
+		CopyConst = 0;
+		{
+			TCircularQueue<FNonPOD> Queue(QueueSize);
+
+			TestEqual(TEXT("The constructor should not run"), Const, 0);
+			TestEqual(TEXT("The destructor should not run"), Dest, 0);
+			TestEqual(TEXT("The copy constructor should not run"), CopyConst, 0);
+		}
+		TestEqual(TEXT("The constructor should not run"), Const, 0);
+		TestEqual(TEXT("The destructor should run"), Dest, QueueSize);
+		TestEqual(TEXT("The copy constructor should not run"), CopyConst, 0);
+	}
+
+	return true;
+}
+
+#endif //WITH_DEV_AUTOMATION_TESTS
diff --git a/Engine/Source/Runtime/Core/Private/Tests/Misc/GuidTest.cpp b/Engine/Source/Runtime/Core/Private/Tests/Misc/GuidTest.cpp
new file mode 100644
--- /dev/null
+++ b/Engine/Source/Runtime/Core/Private/Tests/Misc/GuidTest.cpp
@@ -0,0 +1,98 @@
+// Copyright Epic Games, Inc. All Rights Reserved.
+
+#include "CoreTypes.h"
+#include "Containers/UnrealString.h"
+#include "Misc/Guid.h"
+#include "Misc/AutomationTest.h"
+
+#if WITH_DEV_AUTOMATION_TESTS
+
+IMPLEMENT_SIMPLE_AUTOMATION_TEST(FGuidTest, "System.Core.Misc.Guid", EAutomationTestFlags::ApplicationContextMask | EAutomationTestFlags::SmokeFilter)
+
+bool FGuidTest::RunTest( const FString& Parameters )
+{
+	FGuid g = FGuid(305419896, 2271560481, 305419896, 2271560481);
+
+	// string conversions
+	TestEqual(TEXT("String conversion (Default) must return EGuidFormats::Digits string"), g.ToString(), g.ToString(EGuidFormats::Digits));
+	TestEqual<FString>(TEXT("String conversion (EGuidFormats::Digits)"), g.ToString(EGuidFormats::Digits), TEXT("12345678876543211234567887654321"));
+	TestEqual<FString>(TEXT("String conversion (EGuidFormats::DigitsWithHyphens)"), g.ToString(EGuidFormats::DigitsWithHyphens), TEXT("12345678-8765-4321-1234-567887654321"));
+	TestEqual<FString>(TEXT("String conversion (EGuidFormats::DigitsWithHyphensInBraces)"), g.ToString(EGuidFormats::DigitsWithHyphensInBraces), TEXT("{12345678-8765-4321-1234-567887654321}"));
+	TestEqual<FString>(TEXT("String conversion (EGuidFormats::DigitsWithHyphensInParentheses)"), g.ToString(EGuidFormats::DigitsWithHyphensInParentheses), TEXT("(12345678-8765-4321-1234-567887654321)"));
+	TestEqual<FString>(TEXT("String conversion (EGuidFormats::HexValuesInBraces)"), g.ToString(EGuidFormats::HexValuesInBraces), TEXT("{0x12345678,0x8765,0x4321,{0x12,0x34,0x56,0x78,0x87,0x65,0x43,0x21}}"));
+	TestEqual<FString>(TEXT("String conversion (EGuidFormats::UniqueObjectGuid)"), g.ToString(EGuidFormats::UniqueObjectGuid), TEXT("12345678-87654321-12345678-87654321"));
+	TestEqual<FString>(TEXT("String conversion (EGuidFormats::Short)"), g.ToString(EGuidFormats::Short), TEXT("eFY0EiFDZYd4VjQSIUNlhw"));
+	TestEqual<FString>(TEXT("String conversion (EGuidFormats::Base36Encoded)"), g.ToString(EGuidFormats::Base36Encoded), TEXT("12SRDE1NOQERISMPEF0H6KUU9"));
+
+	// parsing valid strings (exact)
+	FGuid g2_1;
+
+	TestTrue(TEXT("Parsing valid strings must succeed (EGuidFormats::Digits)"), FGuid::ParseExact(TEXT("12345678876543211234567887654321"), EGuidFormats::Digits, g2_1));
+	TestEqual(TEXT("Parsing valid strings must succeed (EGuidFormats::Digits)"), g2_1, g);
+
+	TestTrue(TEXT("Parsing valid strings must succeed (EGuidFormats::DigitsWithHyphens)"), FGuid::ParseExact(TEXT("12345678-8765-4321-1234-567887654321"), EGuidFormats::DigitsWithHyphens, g2_1));
+	TestEqual(TEXT("Parsing valid strings must succeed (EGuidFormats::DigitsWithHyphens)"), g2_1, g);
+
+	TestTrue(TEXT("Parsing valid strings must succeed (EGuidFormats::DigitsWithHyphensInBraces)"), FGuid::ParseExact(TEXT("{12345678-8765-4321-1234-567887654321}"), EGuidFormats::DigitsWithHyphensInBraces, g2_1));
+	TestEqual(TEXT("Parsing valid strings must succeed (EGuidFormats::DigitsWithHyphensInBraces)"), g2_1, g);
+
+	TestTrue(TEXT("Parsing valid strings must succeed (EGuidFormats::DigitsWithHyphensInParentheses)"), FGuid::ParseExact(TEXT("(12345678-8765-4321-1234-567887654321)"), EGuidFormats::DigitsWithHyphensInParentheses, g2_1));
+	TestEqual(TEXT("Parsing valid strings must succeed (EGuidFormats::DigitsWithHyphensInParentheses)"), g2_1, g);
+
+	TestTrue(TEXT("Parsing valid strings must succeed (EGuidFormats::HexValuesInBraces)"), FGuid::ParseExact(TEXT("{0x12345678,0x8765,0x4321,{0x12,0x34,0x56,0x78,0x87,0x65,0x43,0x21}}"), EGuidFormats::HexValuesInBraces, g2_1));
+	TestEqual(TEXT("Parsing valid strings must succeed (EGuidFormats::HexValuesInBraces)"), g2_1, g);
+
+	TestTrue(TEXT("Parsing valid strings must succeed (EGuidFormats::UniqueObjectGuid)"), FGuid::ParseExact(TEXT("12345678-87654321-12345678-87654321"), EGuidFormats::UniqueObjectGuid, g2_1));
+	TestEqual(TEXT("Parsing valid strings must succeed (EGuidFormats::UniqueObjectGuid)"), g2_1, g);
+
+	TestTrue(TEXT("Parsing valid strings must succeed (EGuidFormats::Short)"), FGuid::ParseExact(TEXT("eFY0EiFDZYd4VjQSIUNlhw"), EGuidFormats::Short, g2_1));
+	TestEqual(TEXT("Parsing valid strings must succeed (EGuidFormats::Short)"), g2_1, g);
+
+	TestTrue(TEXT("Parsing valid strings must succeed (EGuidFormats::Base36Encoded)"), FGuid::ParseExact(TEXT("12SRDE1NOQERISMPEF0H6KUU9"), EGuidFormats::Base36Encoded, g2_1));
+	TestEqual(TEXT("Parsing valid strings must succeed (EGuidFormats::Base36Encoded)"), g2_1, g);
+
+	// parsing invalid strings (exact)
+
+
+	// parsing valid strings (automatic)
+	FGuid g3_1;
+
+	TestTrue(TEXT("Parsing valid strings must succeed (12345678876543211234567887654321)"), FGuid::Parse(TEXT("12345678876543211234567887654321"), g3_1));
+	TestEqual(TEXT("Parsing valid strings must succeed (12345678876543211234567887654321)"), g3_1, g);
+
+	TestTrue(TEXT("Parsing valid strings must succeed (12345678-8765-4321-1234-567887654321)"), FGuid::Parse(TEXT("12345678-8765-4321-1234-567887654321"), g3_1));
+	TestEqual(TEXT("Parsing valid strings must succeed (12345678-8765-4321-1234-567887654321)"), g3_1, g);
+
+	TestTrue(TEXT("Parsing valid strings must succeed ({12345678-8765-4321-1234-567887654321})"), FGuid::Parse(TEXT("{12345678-8765-4321-1234-567887654321}"), g3_1));
+	TestEqual(TEXT("Parsing valid strings must succeed ({12345678-8765-4321-1234-567887654321})"), g3_1, g);
+
+	TestTrue(TEXT("Parsing valid strings must succeed ((12345678-8765-4321-1234-567887654321))"), FGuid::Parse(TEXT("(12345678-8765-4321-1234-567887654321)"), g3_1));
+	TestEqual(TEXT("Parsing valid strings must succeed ((12345678-8765-4321-1234-567887654321))"), g3_1, g);
+
+	TestTrue(TEXT("Parsing valid strings must succeed ({0x12345678,0x8765,0x4321,{0x12,0x34,0x56,0x78,0x87,0x65,0x43,0x21}})"), FGuid::Parse(TEXT("{0x12345678,0x8765,0x4321,{0x12,0x34,0x56,0x78,0x87,0x65,0x43,0x21}}"), g3_1));
+	TestEqual(TEXT("Parsing valid strings must succeed ({0x12345678,0x8765,0x4321,{0x12,0x34,0x56,0x78,0x87,0x65,0x43,0x21}})"), g3_1, g);
+
+	TestTrue(TEXT("Parsing valid strings must succeed (12345678-87654321-12345678-87654321)"), FGuid::Parse(TEXT("12345678-87654321-12345678-87654321"), g3_1));
+	TestEqual(TEXT("Parsing valid strings must succeed (12345678-87654321-12345678-87654321)"), g3_1, g);
+
+	TestTrue(TEXT("Parsing valid strings must succeed (eFY0EiFDZYd4VjQSIUNlhw)"), FGuid::Parse(TEXT("eFY0EiFDZYd4VjQSIUNlhw"), g3_1));
+	TestEqual(TEXT("Parsing valid strings must succeed (eFY0EiFDZYd4VjQSIUNlhw)"), g3_1, g);
+
+	TestTrue(TEXT("Parsing valid strings must succeed (12SRDE1NOQERISMPEF0H6KUU9)"), FGuid::Parse(TEXT("12SRDE1NOQERISMPEF0H6KUU9"), g3_1));
+	TestEqual(TEXT("Parsing valid strings must succeed (12SRDE1NOQERISMPEF0H6KUU9)"), g3_1, g);
+
+	//parsing invalid strings (automatic)
+
+	// GUID validation
+	FGuid g4_1 = FGuid::NewGuid();
+
+	TestTrue(TEXT("New GUIDs must be valid"), g4_1.IsValid());
+	
+	g4_1.Invalidate();
+
+	TestFalse(TEXT("Invalidated GUIDs must be invalid"), g4_1.IsValid());
+
+	return true;
+}
+
+#endif //WITH_DEV_AUTOMATION_TESTS
diff --git a/Engine/Source/Runtime/Core/Private/Tests/Misc/QueueTest.cpp b/Engine/Source/Runtime/Core/Private/Tests/Misc/QueueTest.cpp
new file mode 100644
--- /dev/null
+++ b/Engine/Source/Runtime/Core/Private/Tests/Misc/QueueTest.cpp
@@ -0,0 +1,70 @@
+// Copyright Epic Games, Inc. All Rights Reserved.
+
+#include "CoreTypes.h"
+#include "Containers/Queue.h"
+#include "Misc/AutomationTest.h"
+
+#if WITH_DEV_AUTOMATION_TESTS
+
+
+IMPLEMENT_SIMPLE_AUTOMATION_TEST(FQueueTest, "System.Core.Misc.Queue", EAutomationTestFlags::ApplicationContextMask | EAutomationTestFlags::SmokeFilter)
+
+
+bool FQueueTest::RunTest( const FString& Parameters )
+{
+	// empty queues
+	{
+		TQueue<int32> Queue;
+		int32 OutItem = 0;
+
+		TestTrue(TEXT("A new queue must be empty"), Queue.IsEmpty());
+
+		TestFalse(TEXT("A new queue must not dequeue anything"), Queue.Dequeue(OutItem));
+		TestFalse(TEXT("A new queue must not peek anything"), Queue.Peek(OutItem));
+	}
+
+	// insertion & removal
+	{
+		TQueue<int32> Queue;
+		int32 Item1 = 1;
+		int32 Item2 = 2;
+		int32 Item3 = 3;
+		int32 OutItem = 0;
+
+		TestTrue(TEXT("Inserting into a new queue must succeed"), Queue.Enqueue(Item1));
+		TestTrue(TEXT("Peek must succeed on a queue with one item"), Queue.Peek(OutItem));
+		TestEqual(TEXT("Peek must return the first value"), OutItem, Item1);
+
+		TestTrue(TEXT("Inserting into a non-empty queue must succeed"), Queue.Enqueue(Item2));
+		TestTrue(TEXT("Peek must succeed on a queue with two items"), Queue.Peek(OutItem));
+		TestEqual(TEXT("Peek must return the first item"), OutItem, Item1);
+
+		Queue.Enqueue(Item3);
+
+		TestTrue(TEXT("Dequeue must succeed on a queue with three items"), Queue.Dequeue(OutItem));
+		TestEqual(TEXT("Dequeue must return the first item"), OutItem, Item1);
+		TestTrue(TEXT("Dequeue must succeed on a queue with two items"), Queue.Dequeue(OutItem));
+		TestEqual(TEXT("Dequeue must return the second item"), OutItem, Item2);
+		TestTrue(TEXT("Dequeue must succeed on a queue with one item"), Queue.Dequeue(OutItem));
+		TestEqual(TEXT("Dequeue must return the third item"), OutItem, Item3);
+
+		TestTrue(TEXT("After removing all items, the queue must be empty"), Queue.IsEmpty());
+	}
+
+	// emptying
+	{
+		TQueue<int32> Queue;
+
+		Queue.Enqueue(1);
+		Queue.Enqueue(2);
+		Queue.Enqueue(3);
+		Queue.Empty();
+
+		TestTrue(TEXT("An emptied queue must be empty"), Queue.IsEmpty());
+	}
+
+	return true;
+}
+
+
+#endif //WITH_DEV_AUTOMATION_TESTS
diff --git a/Engine/Source/Runtime/Core/Private/Tests/Misc/TimespanTest.cpp b/Engine/Source/Runtime/Core/Private/Tests/Misc/TimespanTest.cpp
new file mode 100644
--- /dev/null
+++ b/Engine/Source/Runtime/Core/Private/Tests/Misc/TimespanTest.cpp
@@ -0,0 +1,139 @@
+// Copyright Epic Games, Inc. All Rights Reserved.
+
+#include "CoreTypes.h"
+#include "Containers/UnrealString.h"
+#include "Misc/Timespan.h"
+#include "Misc/AutomationTest.h"
+
+
+#if WITH_DEV_AUTOMATION_TESTS
+
+IMPLEMENT_SIMPLE_AUTOMATION_TEST(FTimespanTest, "System.Core.Misc.Timespan", EAutomationTestFlags::ApplicationContextMask | EAutomationTestFlags::SmokeFilter)
+
+
+bool FTimespanTest::RunTest(const FString& Parameters)
+{
+	// constructors must create equal objects
+	{
+		FTimespan ts1_1 = FTimespan(3, 2, 1);
+		FTimespan ts1_2 = FTimespan(0, 3, 2, 1);
+		FTimespan ts1_3 = FTimespan(0, 3, 2, 1, 0);
+
+		TestEqual(TEXT("Constructors must create equal objects (Hours/Minutes/Seconds vs. Days/Hours/Minutes/Seconds)"), ts1_1, ts1_2);
+		TestEqual(TEXT("Constructors must create equal objects (Hours/Minutes/Seconds vs. Days/Hours/Minutes/Seconds/FractionNano)"), ts1_1, ts1_3);
+	}
+
+	// component getters must return correct values
+	{
+		FTimespan ts2_1 = FTimespan(1, 2, 3, 4, 123456789);
+
+		TestEqual(TEXT("Component getters must return correct values (Days)"), ts2_1.GetDays(), 1);
+		TestEqual(TEXT("Component getters must return correct values (Hours)"), ts2_1.GetHours(), 2);
+		TestEqual(TEXT("Component getters must return correct values (Minutes)"), ts2_1.GetMinutes(), 3);
+		TestEqual(TEXT("Component getters must return correct values (Seconds)"), ts2_1.GetSeconds(), 4);
+		TestEqual(TEXT("Component getters must return correct values (FractionMilli)"), ts2_1.GetFractionMilli(), 123);
+		TestEqual(TEXT("Component getters must return correct values (FractionMicro)"), ts2_1.GetFractionMicro(), 123456);
+		TestEqual(TEXT("Component getters must return correct values (FractionNano)"), ts2_1.GetFractionNano(), 123456700);
+	}
+
+	// durations of positive and negative time spans must match
+	{
+		FTimespan ts3_1 = FTimespan(1, 2, 3, 4, 123456789);
+		FTimespan ts3_2 = FTimespan(-1, -2, -3, -4, -123456789);
+
+		TestEqual(TEXT("Durations of positive and negative time spans must match"), ts3_1.GetDuration(), ts3_2.GetDuration());
+	}
+
+	// static constructors must create correct values
+	{
+		TestEqual(TEXT("Static constructors must create correct values (FromDays)"), FTimespan::FromDays(123).GetTotalDays(), 123.0);
+		TestEqual(TEXT("Static constructors must create correct values (FromHours)"), FTimespan::FromHours(123).GetTotalHours(), 123.0);
+		TestEqual(TEXT("Static constructors must create correct values (FromMinutes)"), FTimespan::FromMinutes(123).GetTotalMinutes(), 123.0);
+		TestEqual(TEXT("Static constructors must create correct values (FromSeconds)"), FTimespan::FromSeconds(123).GetTotalSeconds(), 123.0);
+		TestEqual(TEXT("Static constructors must create correct values (FromMilliseconds)"), FTimespan::FromMilliseconds(123).GetTotalMilliseconds(), 123.0);
+		TestEqual(TEXT("Static constructors must create correct values (FromMicroseconds)"), FTimespan::FromMicroseconds(123).GetTotalMicroseconds(), 123.0);
+	}
+
+	// string conversions must return correct strings
+	{
+		FTimespan ts5_1 = FTimespan(1, 2, 3, 4, 123456789);
+
+		TestEqual<FString>(TEXT("String conversion (Default)"), ts5_1.ToString(), TEXT("+1.02:03:04.123"));
+		TestEqual<FString>(TEXT("String conversion (%d.%h:%m:%s.%f)"), ts5_1.ToString(TEXT("%d.%h:%m:%s.%f")), TEXT("+1.02:03:04.123"));
+		TestEqual<FString>(TEXT("String conversion (%d.%h:%m:%s.%u)"), ts5_1.ToString(TEXT("%d.%h:%m:%s.%u")), TEXT("+1.02:03:04.123456"));
+		TestEqual<FString>(TEXT("String conversion (%D.%h:%m:%s.%n)"), ts5_1.ToString(TEXT("%D.%h:%m:%s.%n")), TEXT("+00000001.02:03:04.123456700"));
+	}
+
+	// parsing valid strings must succeed
+	{
+		FTimespan ts6_t;
+
+		FTimespan ts6_1 = FTimespan(1, 2, 3, 4, 123000000);
+		FTimespan ts6_2 = FTimespan(1, 2, 3, 4, 123456000);
+		FTimespan ts6_3 = FTimespan(1, 2, 3, 4, 123456700);
+
+		TestTrue(TEXT("Parsing valid strings must succeed (+1.02:03:04.123)"), FTimespan::Parse(TEXT("+1.02:03:04.123"), ts6_t));
+		TestEqual(TEXT("Parsing valid strings must result in correct values (+1.02:03:04.123)"), ts6_t, ts6_1);
+
+		TestTrue(TEXT("Parsing valid strings must succeed (+1.02:03:04.123456)"), FTimespan::Parse(TEXT("+1.02:03:04.123456"), ts6_t));
+		TestEqual(TEXT("Parsing valid strings must result in correct values (+1.02:03:04.123456)"), ts6_t, ts6_2);
+
+		TestTrue(TEXT("Parsing valid strings must succeed (+1.02:03:04.123456789)"), FTimespan::Parse(TEXT("+1.02:03:04.123456789"), ts6_t));
+		TestEqual(TEXT("Parsing valid strings must result in correct values (+1.02:03:04.123456789)"), ts6_t, ts6_3);
+
+		FTimespan ts6_4 = FTimespan(-1, -2, -3, -4, -123000000);
+		FTimespan ts6_5 = FTimespan(-1, -2, -3, -4, -123456000);
+		FTimespan ts6_6 = FTimespan(-1, -2, -3, -4, -123456700);
+
+		TestTrue(TEXT("Parsing valid strings must succeed (-1.02:03:04.123)"), FTimespan::Parse(TEXT("-1.02:03:04.123"), ts6_t));
+		TestEqual(TEXT("Parsing valid strings must result in correct values (-1.02:03:04.123)"), ts6_t, ts6_4);
+
+		TestTrue(TEXT("Parsing valid strings must succeed (-1.02:03:04.123456)"), FTimespan::Parse(TEXT("-1.02:03:04.123456"), ts6_t));
+		TestEqual(TEXT("Parsing valid strings must result in correct values (-1.02:03:04.123456)"), ts6_t, ts6_5);
+
+		TestTrue(TEXT("Parsing valid strings must succeed (-1.02:03:04.123456789)"), FTimespan::Parse(TEXT("-1.02:03:04.123456789"), ts6_t));
+		TestEqual(TEXT("Parsing valid strings must result in correct values (-1.02:03:04.123456789)"), ts6_t, ts6_6);
+	}
+
+	// parsing invalid strings must fail
+	{
+		FTimespan ts7_1;
+
+		//TestFalse(TEXT("Parsing invalid strings must fail (1,02:03:04.005)"), FTimespan::Parse(TEXT("1,02:03:04.005"), ts7_1));
+		//TestFalse(TEXT("Parsing invalid strings must fail (1.1.02:03:04:005)"), FTimespan::Parse(TEXT("1.1.02:03:04:005"), ts7_1));
+		//TestFalse(TEXT("Parsing invalid strings must fail (04:005)"), FTimespan::Parse(TEXT("04:005"), ts7_1));
+	}
+
+	// `From*` converters must return correct value
+	// test normal and edge cases for polar conversions (FromMicroseconds() and FromDay()) and just normal case for all others
+	{
+		TestEqual(TEXT("FromMicroseconds(0) results in correct value"), FTimespan::FromMicroseconds(0), FTimespan(0));
+		TestEqual(TEXT("FromMicroseconds(1) results in correct value"), FTimespan::FromMicroseconds(1), FTimespan(1 * ETimespan::TicksPerMicrosecond));
+		TestEqual(TEXT("FromMicroseconds(1.1) results in correct value"), FTimespan::FromMicroseconds(1.1), FTimespan(1 * ETimespan::TicksPerMicrosecond + 1));
+		TestEqual(TEXT("FromMicroseconds(1.5) results in correct value"), FTimespan::FromMicroseconds(1.5), FTimespan(1 * ETimespan::TicksPerMicrosecond + 5));
+		TestEqual(TEXT("FromMicroseconds(1.499999999999997) results in 1.5 microsecs of ticks"), FTimespan::FromMicroseconds(1.499999999999997), FTimespan(1 * ETimespan::TicksPerMicrosecond + 5));
+		TestEqual(TEXT("FromMicroseconds(1.50000001) results in  in 1.5 microsecs of ticks"), FTimespan::FromMicroseconds(1.50000001), FTimespan(1 * ETimespan::TicksPerMicrosecond + 5));
+		TestEqual(TEXT("FromMicroseconds(-1) results in correct value"), FTimespan::FromMicroseconds(-1), FTimespan(-1 * ETimespan::TicksPerMicrosecond));
+		TestEqual(TEXT("FromMicroseconds(-1.1) results in correct value"), FTimespan::FromMicroseconds(-1.1), FTimespan(-1 * ETimespan::TicksPerMicrosecond - 1));
+		TestEqual(TEXT("FromMicroseconds(-1.5) results in correct value"), FTimespan::FromMicroseconds(-1.5), FTimespan(-1 * ETimespan::TicksPerMicrosecond - 5));
+
+		TestEqual(TEXT("FromMilliseconds(1) results in correct value"), FTimespan::FromMilliseconds(1), FTimespan(1 * ETimespan::TicksPerMillisecond));
+		TestEqual(TEXT("FromSeconds(1) results in correct value"), FTimespan::FromSeconds(1), FTimespan(1 * ETimespan::TicksPerSecond));
+		TestEqual(TEXT("FromMinutes(1) results in correct value"), FTimespan::FromMinutes(1), FTimespan(1 * ETimespan::TicksPerMinute));
+		TestEqual(TEXT("FromHours(1) results in correct value"), FTimespan::FromHours(1), FTimespan(1 * ETimespan::TicksPerHour));
+
+		TestEqual(TEXT("FromDays(0) results in correct value"), FTimespan::FromDays(0), FTimespan(0));
+		TestEqual(TEXT("FromDays(1) results in correct value"), FTimespan::FromDays(1), FTimespan(1 * ETimespan::TicksPerDay));
+		TestEqual(TEXT("FromDays(1.25) results in correct value (1 day and 6 hours of ticks)"), FTimespan::FromDays(1.25), FTimespan(1 * ETimespan::TicksPerDay + 6 * ETimespan::TicksPerHour));
+		TestEqual(TEXT("FromDays(1.5) results in correct value (1 day and 12 hours of ticks)"), FTimespan::FromDays(1.5), FTimespan(1 * ETimespan::TicksPerDay + 12 * ETimespan::TicksPerHour));
+		TestEqual(TEXT("FromDays(1.499999999999997) results in correct value (1 day and 12 hours of ticks)"), FTimespan::FromDays(1.499999999999997), FTimespan(1 * ETimespan::TicksPerDay + 12 * ETimespan::TicksPerHour));
+		TestEqual(TEXT("FromDays(-1) results in correct value"), FTimespan::FromDays(-1), FTimespan(-1 * ETimespan::TicksPerDay));
+		TestEqual(TEXT("FromDays(-1.25) results in correct value (minus 1 day and 6 hours of ticks)"), FTimespan::FromDays(-1.25), FTimespan(-1 * ETimespan::TicksPerDay - 6 * ETimespan::TicksPerHour));
+		TestEqual(TEXT("FromDays(-1.5) results in correct value (minus 1 day and 12 hours of ticks)"), FTimespan::FromDays(-1.5), FTimespan(-1 * ETimespan::TicksPerDay - 12 * ETimespan::TicksPerHour));
+	}
+
+	return true;
+}
+
+
+#endif //WITH_DEV_AUTOMATION_TESTS
diff --git a/Engine/Source/Runtime/Core/Private/Tests/Misc/TripleBufferTest.cpp b/Engine/Source/Runtime/Core/Private/Tests/Misc/TripleBufferTest.cpp
new file mode 100644
--- /dev/null
+++ b/Engine/Source/Runtime/Core/Private/Tests/Misc/TripleBufferTest.cpp
@@ -0,0 +1,106 @@
+// Copyright Epic Games, Inc. All Rights Reserved.
+
+#include "CoreTypes.h"
+#include "Containers/UnrealString.h"
+#include "Misc/AutomationTest.h"
+#include "Containers/TripleBuffer.h"
+#include "Math/RandomStream.h"
+
+
+IMPLEMENT_SIMPLE_AUTOMATION_TEST(FTripleBufferTest, "System.Core.Misc.TripleBuffer", EAutomationTestFlags::ApplicationContextMask | EAutomationTestFlags::SmokeFilter)
+
+
+bool FTripleBufferTest::RunTest(const FString& Parameters)
+{
+	// uninitialized buffer
+	{
+		TTripleBuffer<int32> Buffer(NoInit);
+
+		TestFalse(TEXT("Uninitialized triple buffer must not be dirty"), Buffer.IsDirty());
+	}
+
+	// initialized buffer
+	{
+		TTripleBuffer<int32> Buffer(1);
+
+		TestFalse(TEXT("Initialized triple buffer must not be dirty"), Buffer.IsDirty());
+		TestEqual(TEXT("Initialized triple buffer must have correct read buffer value"), Buffer.Read(), 1);
+
+		Buffer.SwapReadBuffers();
+
+		TestEqual(TEXT("Initialized triple buffer must have correct temp buffer value"), Buffer.Read(), 1);
+
+		Buffer.SwapWriteBuffers();
+
+		TestTrue(TEXT("Write buffer swap must set dirty flag"), Buffer.IsDirty());
+
+		Buffer.SwapReadBuffers();
+
+		TestFalse(TEXT("Read buffer swap must clear dirty flag"), Buffer.IsDirty());
+		TestEqual(TEXT("Initialized triple buffer must have correct temp buffer value"), Buffer.Read(), 1);
+	}
+
+	// pre-set buffer
+	{
+		int32 Array[3] = { 1, 2, 3 };
+		TTripleBuffer<int32> Buffer(Array);
+
+		int32 Read = Buffer.Read();
+		TestEqual(TEXT("Pre-set triple buffer must have correct Read buffer value"), Read, 3);
+
+		Buffer.SwapReadBuffers();
+
+		int32 Temp = Buffer.Read();
+		TestEqual(TEXT("Pre-set triple buffer must have correct Temp buffer value"), Temp, 1);
+
+		Buffer.SwapWriteBuffers();
+		Buffer.SwapReadBuffers();
+
+		int32 Write = Buffer.Read();
+		TestEqual(TEXT("Pre-set triple buffer must have correct Write buffer value"), Write, 2);
+	}
+
+	// operations
+	{
+		TTripleBuffer<int32> Buffer;
+
+		for (int32 Index = 0; Index < 6; ++Index)
+		{
+			int32& Write = Buffer.GetWriteBuffer(); Write = Index; Buffer.SwapWriteBuffers();
+			Buffer.SwapReadBuffers();
+			TestEqual(*FString::Printf(TEXT("Triple buffer must read correct value (%i)"), Index), Buffer.Read(), Index);
+		}
+
+		FRandomStream Rand;
+		int32 LastRead = -1;
+
+		for (int32 Index = 0; Index < 100; ++Index)
+		{
+			int32 Writes = Rand.GetUnsignedInt() % 4;
+
+			while (Writes > 0)
+			{
+				int32& Write = Buffer.GetWriteBuffer(); Write = Index; Buffer.SwapWriteBuffers();
+				--Writes;
+			}
+			
+			int32 Reads = Rand.GetUnsignedInt() % 4;
+
+			while (Reads > 0)
+			{
+				if (!Buffer.IsDirty())
+				{
+					break;
+				}
+
+				Buffer.SwapReadBuffers();
+				int32 Read = Buffer.Read();
+				TestTrue(TEXT("Triple buffer must read in increasing order"), Read > LastRead);
+				LastRead = Read;
+				--Reads;
+			}
+		}
+	}
+
+	return true;
+}
diff --git a/Engine/Source/Runtime/Core/Private/Tests/Serialization/MemoryReaderTest.cpp b/Engine/Source/Runtime/Core/Private/Tests/Serialization/MemoryReaderTest.cpp
new file mode 100644
--- /dev/null
+++ b/Engine/Source/Runtime/Core/Private/Tests/Serialization/MemoryReaderTest.cpp
@@ -0,0 +1,184 @@
+// Copyright Epic Games, Inc. All Rights Reserved.
+
+#include "CoreMinimal.h"
+#include "Misc/AutomationTest.h"
+#include "Serialization/MemoryReader.h"
+#include "Serialization/MemoryWriter.h"
+
+#if WITH_DEV_AUTOMATION_TESTS
+
+IMPLEMENT_SIMPLE_AUTOMATION_TEST(FMemoryReaderTest, "System.Core.Serialization.MemoryReader", EAutomationTestFlags::ApplicationContextMask | EAutomationTestFlags::SmokeFilter)
+
+bool FMemoryReaderTest::RunTest(const FString& Parameters)
+{
+	// Test reading uint64 byte swapped, ensuring that bytes really get swapped.
+	{
+		const uint64 WrittenValue  = 0x8877665544332211ull;
+		const uint64 ExpectedValue = 0x1122334455667788ull;
+
+		TArray<uint8> Bytes;
+		Bytes.AddUninitialized(sizeof(uint64));
+		*reinterpret_cast<uint64*>(Bytes.GetData()) = WrittenValue;
+
+		FMemoryReader Reader(Bytes);
+		Reader.SetByteSwapping(true);
+		uint64 ReadValue;
+		Reader << ReadValue;
+		TestTrue("Test reading uint64 in swapped byte order.", ReadValue == ExpectedValue);
+	}
+
+	// Write all supported types and read them back. (Assuming that MemoryWriterTest tests don't fail).
+	{
+		// Keeps the 'official' test value.
+		const uint8  TestValueU8   = 0x12;
+		const int8   TestValueS8   = 0x34;
+		const uint16 TestValueU16  = 0x1122;
+		const int16  TestValueS16  = 0x3344;
+		const uint32 TestValueU32  = 0x11223344;
+		const int32  TestValueS32  = 0x55667788;
+		const uint64 TestValueU64  = 0x1122334455667788;
+		const int64  TestValueS64  = 0x99AABBCCDDEEFF00;
+		const float  TestValueF    = 128.5f;
+		const double TestValueD    = 256.5;
+		const bool   TestValueB    = true;
+		const WIDECHAR TestValueCh = 0xF2;
+		const FString TestAnsiStr  = TEXT("Joe");
+		const FString TestUtf16Str = TEXT("\uC11C\uC6B8\uC0AC\uB78C"); // Must contains at least one non-ANSI codepoint to ensure UTF16 is used.
+
+		// The values to write.
+		TArray<uint8> Bytes;
+		FMemoryWriter Writer(Bytes);
+		uint8  ValueU8   = TestValueU8;
+		int8   ValueS8   = TestValueS8;
+		uint16 ValueU16  = TestValueU16;
+		int16  ValueS16  = TestValueS16;
+		uint32 ValueU32  = TestValueU32;
+		int32  ValueS32  = TestValueS32;
+		uint64 ValueU64  = TestValueU64;
+		int64  ValueS64  = TestValueS64;
+		float  ValueF    = TestValueF;
+		double ValueD    = TestValueD;
+		bool   ValueB    = TestValueB;
+		WIDECHAR ValueCh = TestValueCh;
+		FString AnsiStr  = TestAnsiStr;
+		FString Utf16Str = TestUtf16Str;
+
+		// Try to not align the values in the buffer, so that the test would fail (on platforms not supporting misaligned read) if the reader tries to cast and dereference a pointer type with misaligned address. Doing like: *(uint64*)Array.GetData();
+		Writer << ValueU8;
+		Writer << ValueU16;
+		Writer << ValueU32;
+		Writer << ValueU64;
+		Writer << ValueS8;
+		Writer << ValueS16;
+		Writer << ValueS32;
+		Writer << ValueS64;
+		Writer << ValueF;
+		Writer << ValueD;
+		Writer << ValueB;
+		Writer << ValueCh;
+		Writer << AnsiStr;
+		Writer << Utf16Str;
+
+		// Write again, but write the bytes swapped.
+		Writer.SetByteSwapping(true);
+
+		Writer << ValueU8;
+		Writer << ValueU16;
+		Writer << ValueU32;
+		Writer << ValueU64;
+		Writer << ValueS8;
+		Writer << ValueS16;
+		Writer << ValueS32;
+		Writer << ValueS64;
+		Writer << ValueF;
+		Writer << ValueD;
+		Writer << ValueB;
+		Writer << ValueCh;
+		Writer << AnsiStr;
+		Writer << Utf16Str;
+
+		// The value to reads without swapping the the value swapped back to this platform endianness.
+		uint8  ReadValueU8,  ReadValueSwapU8;
+		int8   ReadValueS8,  ReadValueSwapS8;
+		uint16 ReadValueU16, ReadValueSwapU16;
+		int16  ReadValueS16, ReadValueSwapS16;
+		uint32 ReadValueU32, ReadValueSwapU32;
+		int32  ReadValueS32, ReadValueSwapS32;
+		uint64 ReadValueU64, ReadValueSwapU64;
+		int64  ReadValueS64, ReadValueSwapS64;
+		float  ReadValueF,   ReadValueSwapF;
+		double ReadValueD,   ReadValueSwapD;
+		bool   ReadValueB,   ReadValueSwapB;
+		WIDECHAR ReadValueCh, ReadValueSwapCh;
+		FString ReadAnsiStr,  ReadSwapAnsiStr;
+		FString ReadUtf16Str, ReadSwapUtf16Str;
+
+		// Read the first set of values written in this platform endianness.
+		FMemoryReader Reader(Bytes);
+		Reader << ReadValueU8;
+		Reader << ReadValueU16;
+		Reader << ReadValueU32;
+		Reader << ReadValueU64;
+		Reader << ReadValueS8;
+		Reader << ReadValueS16;
+		Reader << ReadValueS32;
+		Reader << ReadValueS64;
+		Reader << ReadValueF;
+		Reader << ReadValueD;
+		Reader << ReadValueB;
+		Reader << ReadValueCh;
+		Reader << ReadAnsiStr;
+		Reader << ReadUtf16Str;
+
+		// Read the swapped values, swapping them back to their original value (Assuming that MemoryWriterTest tests don't fail)
+		Reader.SetByteSwapping(true);
+		Reader << ReadValueSwapU8;
+		Reader << ReadValueSwapU16;
+		Reader << ReadValueSwapU32;
+		Reader << ReadValueSwapU64;
+		Reader << ReadValueSwapS8;
+		Reader << ReadValueSwapS16;
+		Reader << ReadValueSwapS32;
+		Reader << ReadValueSwapS64;
+		Reader << ReadValueSwapF;
+		Reader << ReadValueSwapD;
+		Reader << ReadValueSwapB;
+		Reader << ReadValueSwapCh;
+		Reader << ReadSwapAnsiStr;
+		Reader << ReadSwapUtf16Str;
+
+		// Validate that the value read are the expected ones.
+		TestTrue("Test reading 'u8' from byte stream",             ReadValueU8      == TestValueU8);
+		TestTrue("Test reading 'u16' from byte stream",            ReadValueU16     == TestValueU16);
+		TestTrue("Test reading 'u32' from byte stream",            ReadValueU32     == TestValueU32);
+		TestTrue("Test reading 'u64' from byte stream",            ReadValueU64     == TestValueU64);
+		TestTrue("Test reading 's8' from byte stream",             ReadValueS8      == TestValueS8);
+		TestTrue("Test reading 's16' from byte stream",            ReadValueS16     == TestValueS16);
+		TestTrue("Test reading 's32' from byte stream",            ReadValueS32     == TestValueS32);
+		TestTrue("Test reading 's64' from byte stream",            ReadValueS64     == TestValueS64);
+		TestTrue("Test reading 'foat' from byte stream",           ReadValueF       == TestValueF);
+		TestTrue("Test reading 'double' from byte stream",         ReadValueD       == TestValueD);
+		TestTrue("Test reading 'bool' from byte stream",           ReadValueB       == TestValueB);
+		TestTrue("Test reading 'wchar' from byte stream",          ReadValueCh      == TestValueCh);
+		TestTrue("Test reading 'ansi str' from byte stream",       ReadAnsiStr      == TestAnsiStr);
+		TestTrue("Test reading 'utf16 str' from byte stream",      ReadUtf16Str     == TestUtf16Str);
+		TestTrue("Test reading 'u8-swapped' from byte stream",     ReadValueSwapU8  == TestValueU8);
+		TestTrue("Test reading 'u16-swapped' from byte stream",    ReadValueSwapU16 == TestValueU16);
+		TestTrue("Test reading 'u32-swapped' from byte stream",    ReadValueSwapU32 == TestValueU32);
+		TestTrue("Test reading 'u64-swapped' from byte stream",    ReadValueSwapU64 == TestValueU64);
+		TestTrue("Test reading 's8-swapped' from byte stream",     ReadValueSwapS8  == TestValueS8);
+		TestTrue("Test reading 's16'-swapped' from byte stream",   ReadValueSwapS16 == TestValueS16);
+		TestTrue("Test reading 's32-swapped' from byte stream",    ReadValueSwapS32 == TestValueS32);
+		TestTrue("Test reading 's64-swapped' from byte stream",    ReadValueSwapS64 == TestValueS64);
+		TestTrue("Test reading 'float-swapped' from byte stream",  ReadValueSwapF   == TestValueF);
+		TestTrue("Test reading 'double-swapped' from byte stream", ReadValueSwapD   == TestValueD);
+		TestTrue("Test reading 'bool-swapped' from byte stream",   ReadValueSwapB   == TestValueB);
+		TestTrue("Test reading 'wchar-swapped' from byte stream",  ReadValueSwapCh  == TestValueCh);
+		TestTrue("Test reading 'ansi-swapped' from byte stream",   ReadSwapAnsiStr  == TestAnsiStr);
+		TestTrue("Test reading 'utf16-swapped' from byte stream",  ReadSwapUtf16Str == TestUtf16Str);
+	}
+
+	return true;
+}
+
+#endif // WITH_DEV_AUTOMATION_TESTS
