Added missing PRAGMA_ENABLE_OPTIMIZATION

#rnx

diff --git a/Engine/Source/Runtime/Core/Private/Tests/Internationalization/DateTimeFormattingRulesTest.cpp b/Engine/Source/Runtime/Core/Private/Tests/Internationalization/DateTimeFormattingRulesTest.cpp
--- a/Engine/Source/Runtime/Core/Private/Tests/Internationalization/DateTimeFormattingRulesTest.cpp
+++ b/Engine/Source/Runtime/Core/Private/Tests/Internationalization/DateTimeFormattingRulesTest.cpp
@@ -160,4 +160,6 @@ bool FDateTimeFormattingRulesTest::RunTest (const FString& Parameters)
 	return true;
 }
 
+PRAGMA_ENABLE_OPTIMIZATION
+
 #endif //WITH_DEV_AUTOMATION_TESTS
