Added the test for the creation of an unique buffer from an array.

diff --git a/Engine/Source/Runtime/Core/Private/Tests/Memory/SharedBufferTest.cpp b/Engine/Source/Runtime/Core/Private/Tests/Memory/SharedBufferTest.cpp
--- a/Engine/Source/Runtime/Core/Private/Tests/Memory/SharedBufferTest.cpp
+++ b/Engine/Source/Runtime/Core/Private/Tests/Memory/SharedBufferTest.cpp
@@ -466,6 +466,20 @@ bool FBufferOwnerTest::RunTest(const FString& Parameters)
 	return true;
 }
 
+IMPLEMENT_SIMPLE_AUTOMATION_TEST(FUniqueBufferFromArrayTest, "System.Core.Memory.UniqueBufferFromArray", EAutomationTestFlags::ApplicationContextMask | EAutomationTestFlags::SmokeFilter)
+bool FUniqueBufferFromArrayTest::RunTest(const FString& Parameters)
+{
+	TArray<uint16> Array{ 1, 2, 3, 4, 5, 6, 7, 8 };
+	const void* const ExpectedData = Array.GetData();
+	const uint64 ExpectedSize = Array.Num() * sizeof(uint16);
+	const FUniqueBuffer Buffer = MakeUniqueBufferFromArray(MoveTemp(Array));
+	TestTrue(TEXT("MakeUniqueBufferFromArray -> Array.IsEmpty"), Array.IsEmpty());
+	TestEqual(TEXT("MakeUniqueBufferFromArray -> Buffer.GetData()"), Buffer.GetData(), ExpectedData);
+	TestEqual(TEXT("MakeUniqueBufferFromArray -> Buffer.GetSize()"), Buffer.GetSize(), ExpectedSize);
+
+	return true;
+}
+
 IMPLEMENT_SIMPLE_AUTOMATION_TEST(FSharedBufferFromArrayTest, "System.Core.Memory.SharedBufferFromArray", EAutomationTestFlags::ApplicationContextMask | EAutomationTestFlags::SmokeFilter)
 bool FSharedBufferFromArrayTest::RunTest(const FString& Parameters)
 {
