Fix mismatched #pragma push/pop when compiling with windows clang

#rnx

diff --git a/Engine/Source/Runtime/Core/Public/Misc/GeneratedTypeName.h b/Engine/Source/Runtime/Core/Public/Misc/GeneratedTypeName.h
--- a/Engine/Source/Runtime/Core/Public/Misc/GeneratedTypeName.h
+++ b/Engine/Source/Runtime/Core/Public/Misc/GeneratedTypeName.h
@@ -114,7 +114,7 @@ inline const TCHAR* GetGeneratedTypeName()
 	return Result.Array;
 }
 
-#ifdef _MSC_VER
+#if defined(_MSC_VER) && !defined(__clang__)
 	#pragma warning(pop)
 #endif
 
