Updating docs for TRotator based on further discussion and experimentation. Rotator interprets angles as intrinsic rotations and applies them in the order yaw, pitch, then roll. Positive angles are left-handed around Z, right-handed around +Y, and right-handed around +X.

diff --git a/Engine/Source/Runtime/Core/Public/Math/Rotator.h b/Engine/Source/Runtime/Core/Public/Math/Rotator.h
--- a/Engine/Source/Runtime/Core/Public/Math/Rotator.h
+++ b/Engine/Source/Runtime/Core/Public/Math/Rotator.h
@@ -22,9 +22,13 @@ namespace Math
  *
  * All rotation values are stored in degrees.
  *
- * The angles are interpreted as intrinsic rotations applied in the order Roll, then Pitch, then Yaw. I.e., an object would be rotated
- * first by the specified roll around its forward axis, then pitched around its (new) right axis, and then finally yawed around its 
- * (new) up axis.
+ * The angles are interpreted as intrinsic rotations applied in the order Yaw, then Pitch, then Roll. I.e., an object would be rotated
+ * first by the specified yaw around its up axis (with positive angles interpreted as clockwise when viewed from above, along -Z), 
+ * then pitched around its (new) right axis (with positive angles interpreted as 'nose up', i.e. clockwise when viewed along +Y), 
+ * and then finally rolled around its (final) forward axis (with positive angles interpreted as clockwise rotations when viewed along +X).
+ *
+ * Note that these conventions differ from quaternion axis/angle. UE Quat always considers a positive angle to be a left-handed rotation, 
+ * whereas Rotator treats yaw as left-handed but pitch and roll as right-handed.
  * 
  */
 template<typename T>
@@ -40,10 +44,10 @@ public:
 	/** Rotation around the right axis (around Y axis), Looking up and down (0=Straight Ahead, +Up, -Down) */
 	T Pitch;
 
-	/** Rotation around the up axis (around Z axis), Running in circles 0=East, +North, -South. */
+	/** Rotation around the up axis (around Z axis), Turning around (0=Forward, +Right, -Left)*/
 	T Yaw;
 
-	/** Rotation around the forward axis (around X axis), Tilting your head, 0=Straight, +Clockwise, -CCW. */
+	/** Rotation around the forward axis (around X axis), Tilting your head, (0=Straight, +Clockwise, -CCW) */
 	T Roll;
 
 public:
