Improved struct packing for TCounter (32 --> 20 bytes) and FScopedDetailTickStats (32 --> 24 bytes).

diff --git a/Engine/Source/Runtime/Core/Public/ProfilingDebugging/CountersTrace.h b/Engine/Source/Runtime/Core/Public/ProfilingDebugging/CountersTrace.h
--- a/Engine/Source/Runtime/Core/Public/ProfilingDebugging/CountersTrace.h
+++ b/Engine/Source/Runtime/Core/Public/ProfilingDebugging/CountersTrace.h
@@ -42,8 +42,8 @@ struct FCountersTrace
 	public:
 		TCounter(const TCHAR* InCounterName, ETraceCounterDisplayHint InCounterDisplayHint)
 			: Value(0)
-			, CounterId(0)
 			, CounterName(InCounterName)
+			, CounterId(0)
 			, CounterDisplayHint(InCounterDisplayHint)
 		{
 			CounterId = OutputInitCounter(InCounterName, CounterType, CounterDisplayHint);
@@ -92,18 +92,12 @@ struct FCountersTrace
 			--Value;
 			OutputSetValue(CounterId, Value);
 		}
-		
+
 	private:
 		ValueType Value;
-		uint16 CounterId;
 		const TCHAR* CounterName;
+		uint16 CounterId;
 		ETraceCounterDisplayHint CounterDisplayHint;
-
-		bool CheckCounterId()
-		{
-			CounterId = OutputInitCounter(CounterName, CounterType, CounterDisplayHint);
-			return !!CounterId;
-		}
 	};
 
 	using FCounterInt = TCounter<int64, TraceCounterType_Int>;
diff --git a/Engine/Source/Runtime/Engine/Private/LevelTick.cpp b/Engine/Source/Runtime/Engine/Private/LevelTick.cpp
--- a/Engine/Source/Runtime/Engine/Private/LevelTick.cpp
+++ b/Engine/Source/Runtime/Engine/Private/LevelTick.cpp
@@ -376,9 +376,9 @@ void FDetailedTickStats::DumpStats()
  * Constructor, keeping track of object's class and start time.
  */
 FScopedDetailTickStats::FScopedDetailTickStats( FDetailedTickStats& InDetailedTickStats, UObject* InObject )
-:	Object( InObject )
+:	DetailedTickStats( InDetailedTickStats )
+,	Object( InObject )
 ,	StartCycles( FPlatformTime::Cycles() )
-,	DetailedTickStats( InDetailedTickStats )
 {
 	bShouldTrackObjectClass = DetailedTickStats.BeginObject( Object->GetClass() );
 	bShouldTrackObject = DetailedTickStats.BeginObject( Object );
diff --git a/Engine/Source/Runtime/Engine/Public/UnrealEngine.h b/Engine/Source/Runtime/Engine/Public/UnrealEngine.h
--- a/Engine/Source/Runtime/Engine/Public/UnrealEngine.h
+++ b/Engine/Source/Runtime/Engine/Public/UnrealEngine.h
@@ -283,14 +283,15 @@ struct FScopedDetailTickStats
 
 
 private:
-	/** Object to track. 
-		Not GC safe, but we won't have anything in-flight during GC so that should be moot
-	*/
+	/** Detailed tick stats to update. */
+	FDetailedTickStats& DetailedTickStats;
+	/**
+	 * Object to track.
+	 * Not GC safe, but we won't have anything in-flight during GC so that should be moot.
+	 */
 	UObject* Object;
 	/** Tick start time. */
 	uint32 StartCycles;
-	/** Detailed tick stats to update. */
-	FDetailedTickStats& DetailedTickStats;
 	/** Whether object should be tracked. false e.g. when recursion is involved. */
 	bool bShouldTrackObject;
 	/** Whether object class should be tracked. false e.g. when recursion is involved. */
