VectorVM: allow dead random instruction to be removed.
Turned new VM back on and bumped the Niagara_LatestScriptCompileVersion.  #JIRA UE-147838 #preflight 624c47478e5ae00f0ab7e7f4 #rb rob.krajcarski

diff --git a/Engine/Plugins/FX/Niagara/Source/Niagara/Niagara.Build.cs b/Engine/Plugins/FX/Niagara/Source/Niagara/Niagara.Build.cs
--- a/Engine/Plugins/FX/Niagara/Source/Niagara/Niagara.Build.cs
+++ b/Engine/Plugins/FX/Niagara/Source/Niagara/Niagara.Build.cs
@@ -70,7 +70,7 @@ public class Niagara : ModuleRules
         PublicDefinitions.AddRange(
             new string[]
             {
-                "VECTORVM_SUPPORTS_EXPERIMENTAL=0",
+                "VECTORVM_SUPPORTS_EXPERIMENTAL=1",
                 "VECTORVM_SUPPORTS_LEGACY=1"
             });
     }
diff --git a/Engine/Source/Runtime/Core/Private/UObject/UE5MainStreamObjectVersion.cpp b/Engine/Source/Runtime/Core/Private/UObject/UE5MainStreamObjectVersion.cpp
--- a/Engine/Source/Runtime/Core/Private/UObject/UE5MainStreamObjectVersion.cpp
+++ b/Engine/Source/Runtime/Core/Private/UObject/UE5MainStreamObjectVersion.cpp
@@ -14,7 +14,7 @@ TMap<FGuid, FGuid> FUE5MainStreamObjectVersion::GetSystemGuids()
 	SystemGuids.Add(DevGuids.MATERIALSHADERMAP_DERIVEDDATA_VER, FGuid("65C01D817C9A4EEFAE9E988D41A1F3DD"));
 	SystemGuids.Add(DevGuids.NANITE_DERIVEDDATA_VER, FGuid("7E26CA5B0F5F03C40A3FD7173D88B8FA"));
 	SystemGuids.Add(DevGuids.NIAGARASHADERMAP_DERIVEDDATA_VER, FGuid("8A37C45D24F2423CBE5F8F371DE33575"));
-	SystemGuids.Add(DevGuids.Niagara_LatestScriptCompileVersion, FGuid("078A8ED467E24F94BC5DDB4BE2A4EB8E"));
+	SystemGuids.Add(DevGuids.Niagara_LatestScriptCompileVersion, FGuid("2B4EE121451E44299EF2EFED80BDD0CD"));
 	SystemGuids.Add(DevGuids.SkeletalMeshDerivedDataVersion, FGuid("FB7B2FB546DE49B2A2F88B7F69763CEE"));
 	SystemGuids.Add(DevGuids.STATICMESH_DERIVEDDATA_VER, FGuid("98F7D79A5811013825E75BBCC41ED3E9"));
 	SystemGuids.Add(DevGuids.POSESEARCHDB_DERIVEDDATA_VER, FGuid("4E595C2AC5E947D6BA9ABC874353E5BC"));
diff --git a/Engine/Source/Runtime/VectorVM/Private/VectorVMExperimental_Optimizer.inl b/Engine/Source/Runtime/VectorVM/Private/VectorVMExperimental_Optimizer.inl
--- a/Engine/Source/Runtime/VectorVM/Private/VectorVMExperimental_Optimizer.inl
+++ b/Engine/Source/Runtime/VectorVM/Private/VectorVMExperimental_Optimizer.inl
@@ -1226,9 +1226,8 @@ VECTORVM_API uint32 OptimizeVectorVMScript(const uint8 *InBytecode, int InByteco
 			for (uint32 i = 0; i < OptContext->Intermediate.NumInstructions; ++i)
 			{
 				FVectorVMOptimizeInstruction *Ins = OptContext->Intermediate.Instructions + i;
-				if (Ins->OpCat == EVectorVMOpCategory::Op && !(Ins->OpCode == EVectorVMOp::random || Ins->OpCode == EVectorVMOp::randomi))
+				if (Ins->OpCat == EVectorVMOpCategory::Op)
 				{
-					//can we remove random instructions? I dunno! so lets not for now
 					GetRegistersUsedForInstruction(OptContext, Ins, &RegUsage);
 					for (int OutputIdx = 0; OutputIdx < RegUsage.NumOutputRegisters; ++OutputIdx)
 					{
diff --git a/Engine/Source/Runtime/VectorVM/VectorVM.Build.cs b/Engine/Source/Runtime/VectorVM/VectorVM.Build.cs
--- a/Engine/Source/Runtime/VectorVM/VectorVM.Build.cs
+++ b/Engine/Source/Runtime/VectorVM/VectorVM.Build.cs
@@ -27,7 +27,7 @@ public class VectorVM : ModuleRules
         PublicDefinitions.AddRange(
             new string[]
             {
-                "VECTORVM_SUPPORTS_EXPERIMENTAL=0",
+                "VECTORVM_SUPPORTS_EXPERIMENTAL=1",
                 "VECTORVM_SUPPORTS_LEGACY=1"
             });
     }
