Renamed tests to match the directory structure

#rnx

diff --git a/Engine/Source/Runtime/Core/Tests/Hash/Blake3Test.cpp b/Engine/Source/Runtime/Core/Tests/Hash/Blake3Test.cpp
--- a/Engine/Source/Runtime/Core/Tests/Hash/Blake3Test.cpp
+++ b/Engine/Source/Runtime/Core/Tests/Hash/Blake3Test.cpp
@@ -5,7 +5,7 @@
 #include "Misc/StringBuilder.h"
 #include "TestHarness.h"
 
-TEST_CASE("Core::Hash::FBlake3::Smoke Test", "[Core][Hash][Smoke]")
+TEST_CASE("Core::Hash::Blake3", "[Core][Hash][Smoke]")
 {
 	struct FBlake3TestCase
 	{
diff --git a/Engine/Source/Runtime/Core/Tests/Misc/StringBuilderTest.cpp b/Engine/Source/Runtime/Core/Tests/Misc/StringBuilderTest.cpp
--- a/Engine/Source/Runtime/Core/Tests/Misc/StringBuilderTest.cpp
+++ b/Engine/Source/Runtime/Core/Tests/Misc/StringBuilderTest.cpp
@@ -6,7 +6,7 @@
 #include "String/Find.h"
 #include "TestHarness.h"
 
-TEST_CASE("Core::String::Misc::StringBuilder", "[Core][String][Smoke]")
+TEST_CASE("Core::String::StringBuilder", "[Core][String][Smoke]")
 {
 	SECTION("Static")
 	{
diff --git a/Engine/Source/Runtime/Core/Tests/String/EscapeTest.cpp b/Engine/Source/Runtime/Core/Tests/String/EscapeTest.cpp
--- a/Engine/Source/Runtime/Core/Tests/String/EscapeTest.cpp
+++ b/Engine/Source/Runtime/Core/Tests/String/EscapeTest.cpp
@@ -9,7 +9,7 @@
 namespace UE::String
 {
 
-TEST_CASE("Core::String::Misc::Escape", "[Core][String][Smoke]")
+TEST_CASE("Core::String::Escape", "[Core][String][Smoke]")
 {
 	SECTION("Escape")
 	{
diff --git a/Engine/Source/Runtime/Core/Tests/String/FindTest.cpp b/Engine/Source/Runtime/Core/Tests/String/FindTest.cpp
--- a/Engine/Source/Runtime/Core/Tests/String/FindTest.cpp
+++ b/Engine/Source/Runtime/Core/Tests/String/FindTest.cpp
@@ -8,7 +8,7 @@
 namespace UE::String
 {
 
-TEST_CASE("Core::String::Misc::Find", "[Core][String][Smoke]")
+TEST_CASE("Core::String::Find", "[Core][String][Smoke]")
 {
 	SECTION("FindFirst")
 	{
diff --git a/Engine/Source/Runtime/Core/Tests/String/JoinTest.cpp b/Engine/Source/Runtime/Core/Tests/String/JoinTest.cpp
--- a/Engine/Source/Runtime/Core/Tests/String/JoinTest.cpp
+++ b/Engine/Source/Runtime/Core/Tests/String/JoinTest.cpp
@@ -10,7 +10,7 @@
 namespace UE::String
 {
 
-TEST_CASE("Core::String::Misc::Join", "[Core][String][Smoke]")
+TEST_CASE("Core::String::Join", "[Core][String][Smoke]")
 {
 	SECTION("Join")
 	{
diff --git a/Engine/Source/Runtime/Core/Tests/String/ParseLinesTest.cpp b/Engine/Source/Runtime/Core/Tests/String/ParseLinesTest.cpp
--- a/Engine/Source/Runtime/Core/Tests/String/ParseLinesTest.cpp
+++ b/Engine/Source/Runtime/Core/Tests/String/ParseLinesTest.cpp
@@ -12,7 +12,7 @@
 namespace UE::String
 {
 
-TEST_CASE("Core::String::Misc::ParseLines", "[Core][String][Smoke]")
+TEST_CASE("Core::String::ParseLines", "[Core][String][Smoke]")
 {
 	using FStringViewArray = TArray<FStringView, TInlineAllocator<8>>;
 
diff --git a/Engine/Source/Runtime/Core/Tests/String/ParseTokensTest.cpp b/Engine/Source/Runtime/Core/Tests/String/ParseTokensTest.cpp
--- a/Engine/Source/Runtime/Core/Tests/String/ParseTokensTest.cpp
+++ b/Engine/Source/Runtime/Core/Tests/String/ParseTokensTest.cpp
@@ -11,7 +11,7 @@
 namespace UE::String
 {
 
-TEST_CASE("Core::String::Misc::ParseTokens", "[Core][String][Smoke]")
+TEST_CASE("Core::String::ParseTokens", "[Core][String][Smoke]")
 {
 	using FCharArray = TArray<TCHAR, TInlineAllocator<8>>;
 	using FStringViewArray = TArray<FStringView, TInlineAllocator<8>>;
