KiB/MiB/GiB/TiB are now accepted spellings for the units in the details view.

diff --git a/Engine/Source/Editor/UnrealEd/Classes/Editor/PropertyEditorTestObject.h b/Engine/Source/Editor/UnrealEd/Classes/Editor/PropertyEditorTestObject.h
--- a/Engine/Source/Editor/UnrealEd/Classes/Editor/PropertyEditorTestObject.h
+++ b/Engine/Source/Editor/UnrealEd/Classes/Editor/PropertyEditorTestObject.h
@@ -301,6 +301,9 @@ class UPropertyEditorTestObject : public UObject
 	UPROPERTY(EditAnywhere, Category=BasicProperties)
 	FTransform TransformProperty;
 
+	UPROPERTY(EditAnywhere, Category=Units, meta=(ForceUnits="GB"))
+	double GigabyteProperty;
+
 	UPROPERTY(EditAnywhere, Category=Classes)
 	TObjectPtr<UClass> ClassProperty;
 
diff --git a/Engine/Source/Runtime/Core/Private/Math/UnitConversion.cpp b/Engine/Source/Runtime/Core/Private/Math/UnitConversion.cpp
--- a/Engine/Source/Runtime/Core/Private/Math/UnitConversion.cpp
+++ b/Engine/Source/Runtime/Core/Private/Math/UnitConversion.cpp
@@ -63,10 +63,10 @@ FParseCandidate ParseCandidates[] = {
 	{ TEXT("RevolutionsPerMinute"),	EUnit::RevolutionsPerMinute },	{ TEXT("rpm"),		EUnit::RevolutionsPerMinute },
 
 	{ TEXT("Bytes"),				EUnit::Bytes },					{ TEXT("B"),		EUnit::Bytes },
-	{ TEXT("Kilobytes"),			EUnit::Kilobytes },				{ TEXT("KB"),		EUnit::Kilobytes },
-	{ TEXT("Megabytes"),			EUnit::Megabytes },				{ TEXT("MB"),		EUnit::Megabytes },
-	{ TEXT("Gigabytes"),			EUnit::Gigabytes },				{ TEXT("GB"),		EUnit::Gigabytes },
-	{ TEXT("Terabytes"),			EUnit::Terabytes },				{ TEXT("TB"),		EUnit::Terabytes },
+	{ TEXT("Kilobytes"),			EUnit::Kilobytes },				{ TEXT("KB"),		EUnit::Kilobytes },				{ TEXT("KiB"),		EUnit::Kilobytes },
+	{ TEXT("Megabytes"),			EUnit::Megabytes },				{ TEXT("MB"),		EUnit::Megabytes },				{ TEXT("MiB"),		EUnit::Megabytes },
+	{ TEXT("Gigabytes"),			EUnit::Gigabytes },				{ TEXT("GB"),		EUnit::Gigabytes },				{ TEXT("GiB"),		EUnit::Gigabytes },
+	{ TEXT("Terabytes"),			EUnit::Terabytes },				{ TEXT("TB"),		EUnit::Terabytes },				{ TEXT("TiB"),		EUnit::Terabytes },
 
 	{ TEXT("Lumens"),				EUnit::Lumens },				{ TEXT("lm"),		EUnit::Lumens },
 	{ TEXT("Candela"),				EUnit::Candela },				{ TEXT("cd"),		EUnit::Candela },
