Fix for static analysis warning

diff --git a/Engine/Source/Runtime/Core/Private/Tests/Internationalization/TextTest.cpp b/Engine/Source/Runtime/Core/Private/Tests/Internationalization/TextTest.cpp
--- a/Engine/Source/Runtime/Core/Private/Tests/Internationalization/TextTest.cpp
+++ b/Engine/Source/Runtime/Core/Private/Tests/Internationalization/TextTest.cpp
@@ -1385,10 +1385,10 @@ bool FTextFormatArgModifierTest::RunTest(const FString& Parameters)
 		}
 
 		{
-			const FTextFormat HppFormatText = INVTEXT("{Arg}|hpp(\uC774\uC5C8,​\uC600)");/* 이었/​였 */
+			const FTextFormat HppFormatText = INVTEXT("{Arg}|hpp(\uC774\uC5C8,​\uC601)");/* 이었/영*/
 			EnsureValidResult(FText::FormatNamed(HppFormatText, TEXT("Arg"), Consonant).ToString(), TEXT("\uC0AC\uB78C\uC774\uC5C8"), TEXT("HppResultConsonant"), HppFormatText.GetSourceText().ToString());
 			EnsureValidResult(FText::FormatNamed(HppFormatText, TEXT("Arg"), ConsonantRieul).ToString(), TEXT("\uC11C\uC6B8\uC774\uC5C8"), TEXT("HppResultConsonantRieul"), HppFormatText.GetSourceText().ToString());
-			EnsureValidResult(FText::FormatNamed(HppFormatText, TEXT("Arg"), Vowel).ToString(), TEXT("\uC0AC\uC790​\uC600"), TEXT("HppResultVowel"), HppFormatText.GetSourceText().ToString());
+			EnsureValidResult(FText::FormatNamed(HppFormatText, TEXT("Arg"), Vowel).ToString(), TEXT("\uC0AC\uC790​\uC601"), TEXT("HppResultVowel"), HppFormatText.GetSourceText().ToString());
 		}
 
 		{
@@ -1399,6 +1399,7 @@ bool FTextFormatArgModifierTest::RunTest(const FString& Parameters)
 		}
 	}
 
+
 	// Restore original culture
 	I18N.RestoreCultureState(OriginalCultureState);
 
