Split StructuredArchive code into individual files.

diff --git a/Engine/Source/Runtime/Core/Private/Serialization/StructuredArchive.cpp b/Engine/Source/Runtime/Core/Private/Serialization/StructuredArchive.cpp
--- a/Engine/Source/Runtime/Core/Private/Serialization/StructuredArchive.cpp
+++ b/Engine/Source/Runtime/Core/Private/Serialization/StructuredArchive.cpp
@@ -1,63 +1,12 @@
 // Copyright Epic Games, Inc. All Rights Reserved.
 
 #include "Serialization/StructuredArchive.h"
+#include "Serialization/StructuredArchiveContainer.h"
 #include "Containers/Set.h"
 #include "Containers/UnrealString.h"
 
 #if WITH_TEXT_ARCHIVE_SUPPORT
 
-//////////// FStructuredArchive::FContainer ////////////
-
-#if DO_STRUCTURED_ARCHIVE_CONTAINER_CHECKS
-
-struct FStructuredArchive::FContainer
-{
-	int  Index                   = 0;
-	int  Count                   = 0;
-	bool bAttributedValueWritten = false;
-
-#if DO_STRUCTURED_ARCHIVE_UNIQUE_FIELD_NAME_CHECKS
-	TSet<FString> KeyNames;
-#endif
-
-	explicit FContainer(int InCount)
-		: Count(InCount)
-	{
-	}
-};
-#endif
-
-FStructuredArchiveChildReader::FStructuredArchiveChildReader(FStructuredArchiveSlot InSlot)
-	: OwnedFormatter(nullptr)
-	, Archive(nullptr)
-{
-	FStructuredArchiveFormatter* Formatter = &InSlot.Ar.Formatter;
-	if (InSlot.GetUnderlyingArchive().IsTextFormat())
-	{
-		Formatter = OwnedFormatter = InSlot.Ar.Formatter.CreateSubtreeReader();
-	}
-
-	Archive = new FStructuredArchive(*Formatter);
-	Root.Emplace(Archive->Open());
-	InSlot.EnterRecord();
-}
-
-FStructuredArchiveChildReader::~FStructuredArchiveChildReader()
-{
-	Root.Reset();
-	Archive->Close();
-	delete Archive;
-	Archive = nullptr;
-
-	// If this is a text archive, we'll have created a subtree reader that our contained archive is using as 
-	// its formatter. We need to clean it up now.
-	if (OwnedFormatter)
-	{
-		delete OwnedFormatter;
-		OwnedFormatter = nullptr;
-	}
-}
-
 //////////// FStructuredArchive ////////////
 
 FStructuredArchive::FStructuredArchive(FArchiveFormatterType& InFormatter)
@@ -90,7 +39,7 @@ FStructuredArchiveSlot FStructuredArchive::Open()
 
 	CurrentSlotElementId = ElementIdGenerator.Generate();
 
-	return FStructuredArchiveSlot(*this, 0, CurrentSlotElementId);
+	return FStructuredArchiveSlot(FStructuredArchiveSlot::EPrivateToken{}, *this, 0, CurrentSlotElementId);
 }
 
 void FStructuredArchive::Close()
@@ -251,434 +200,4 @@ void FStructuredArchive::SetScope(UE::StructuredArchive::Private::FSlotPosition
 	}
 }
 
-//////////// FStructuredArchiveSlot ////////////
-
-FStructuredArchiveRecord FStructuredArchiveSlot::EnterRecord()
-{
-	int32 NewDepth = Ar.EnterSlotAsType(*this, UE::StructuredArchive::Private::EElementType::Record);
-
-#if DO_STRUCTURED_ARCHIVE_CONTAINER_CHECKS
-	Ar.CurrentContainer.Emplace(0);
-#endif
-
-	Ar.Formatter.EnterRecord();
-
-	return FStructuredArchiveRecord(Ar, NewDepth, ElementId);
-}
-
-FStructuredArchiveArray FStructuredArchiveSlot::EnterArray(int32& Num)
-{
-	int32 NewDepth = Ar.EnterSlotAsType(*this, UE::StructuredArchive::Private::EElementType::Array);
-
-	Ar.Formatter.EnterArray(Num);
-
-#if DO_STRUCTURED_ARCHIVE_CONTAINER_CHECKS
-	Ar.CurrentContainer.Emplace(Num);
-#endif
-
-	return FStructuredArchiveArray(Ar, NewDepth, ElementId);
-}
-
-FStructuredArchiveStream FStructuredArchiveSlot::EnterStream()
-{
-	int32 NewDepth = Ar.EnterSlotAsType(*this, UE::StructuredArchive::Private::EElementType::Stream);
-
-	Ar.Formatter.EnterStream();
-
-	return FStructuredArchiveStream(Ar, NewDepth, ElementId);
-}
-
-FStructuredArchiveMap FStructuredArchiveSlot::EnterMap(int32& Num)
-{
-	int32 NewDepth = Ar.EnterSlotAsType(*this, UE::StructuredArchive::Private::EElementType::Map);
-
-	Ar.Formatter.EnterMap(Num);
-
-#if DO_STRUCTURED_ARCHIVE_CONTAINER_CHECKS
-	Ar.CurrentContainer.Emplace(Num);
-#endif
-
-	return FStructuredArchiveMap(Ar, NewDepth, ElementId);
-}
-
-FStructuredArchiveSlot FStructuredArchiveSlot::EnterAttribute(FArchiveFieldName AttributeName)
-{
-	check(Ar.CurrentScope.Num() > 0);
-
-	int32 NewDepth = Depth + 1;
-	if (NewDepth >= Ar.CurrentScope.Num() || Ar.CurrentScope[NewDepth].Id != ElementId || Ar.CurrentScope[NewDepth].Type != UE::StructuredArchive::Private::EElementType::AttributedValue)
-	{
-		int32 NewDepthCheck = Ar.EnterSlotAsType(*this, UE::StructuredArchive::Private::EElementType::AttributedValue);
-		checkSlow(NewDepth == NewDepthCheck);
-
-		Ar.Formatter.EnterAttributedValue();
-
-#if DO_STRUCTURED_ARCHIVE_CONTAINER_CHECKS
-		Ar.CurrentContainer.Emplace(0);
-#endif
-	}
-
-	Ar.CurrentEnteringAttributeState = UE::StructuredArchive::Private::EEnteringAttributeState::NotEnteringAttribute;
-
-	UE::StructuredArchive::Private::FElementId AttributedValueId = Ar.CurrentScope[NewDepth].Id;
-
-	Ar.SetScope(UE::StructuredArchive::Private::FSlotPosition(NewDepth, AttributedValueId));
-
-	Ar.CurrentSlotElementId = Ar.ElementIdGenerator.Generate();
-
-#if DO_STRUCTURED_ARCHIVE_CONTAINER_CHECKS
-#if DO_STRUCTURED_ARCHIVE_UNIQUE_FIELD_NAME_CHECKS
-	if (!Ar.GetUnderlyingArchive().IsLoading())
-	{
-		FContainer& Container = *Ar.CurrentContainer.Top();
-		checkf(!Container.KeyNames.Contains(Name.Name), TEXT("Multiple attributes called '%s' serialized into attributed value"), AttributeName.Name);
-		Container.KeyNames.Add(Name.Name);
-	}
-#endif
-#endif
-
-	Ar.Formatter.EnterAttribute(AttributeName);
-
-	return FStructuredArchiveSlot(Ar, NewDepth, Ar.CurrentSlotElementId);
-}
-
-TOptional<FStructuredArchiveSlot> FStructuredArchiveSlot::TryEnterAttribute(FArchiveFieldName AttributeName, bool bEnterWhenWriting)
-{
-	check(Ar.CurrentScope.Num() > 0);
-
-	int32 NewDepth = Depth + 1;
-	if (NewDepth >= Ar.CurrentScope.Num() || Ar.CurrentScope[NewDepth].Id != ElementId || Ar.CurrentScope[NewDepth].Type != UE::StructuredArchive::Private::EElementType::AttributedValue)
-	{
-		int32 NewDepthCheck = Ar.EnterSlotAsType(*this, UE::StructuredArchive::Private::EElementType::AttributedValue);
-		checkSlow(NewDepth == NewDepthCheck);
-
-		Ar.Formatter.EnterAttributedValue();
-
-#if DO_STRUCTURED_ARCHIVE_CONTAINER_CHECKS
-		Ar.CurrentContainer.Emplace(0);
-#endif
-	}
-
-#if DO_STRUCTURED_ARCHIVE_CONTAINER_CHECKS
-#if DO_STRUCTURED_ARCHIVE_UNIQUE_FIELD_NAME_CHECKS
-	if (!Ar.GetUnderlyingArchive().IsLoading())
-	{
-		FContainer& Container = *Ar.CurrentContainer.Top();
-		checkf(!Container.KeyNames.Contains(Name.Name), TEXT("Multiple attributes called '%s' serialized into attributed value"), AttributeName.Name);
-		Container.KeyNames.Add(Name.Name);
-	}
-#endif
-#endif
-
-	UE::StructuredArchive::Private::FElementId AttributedValueId = Ar.CurrentScope[NewDepth].Id;
-
-	Ar.SetScope(UE::StructuredArchive::Private::FSlotPosition(NewDepth, AttributedValueId));
-
-	if (Ar.Formatter.TryEnterAttribute(AttributeName, bEnterWhenWriting))
-	{
-		Ar.CurrentEnteringAttributeState = UE::StructuredArchive::Private::EEnteringAttributeState::NotEnteringAttribute;
-
-		Ar.CurrentSlotElementId = Ar.ElementIdGenerator.Generate();
-
-		return FStructuredArchiveSlot(Ar, NewDepth, Ar.CurrentSlotElementId);
-	}
-	else
-	{
-		return {};
-	}
-}
-
-void FStructuredArchiveSlot::operator<< (uint8& Value)
-{
-	Ar.EnterSlot(*this);
-	Ar.Formatter.Serialize(Value);
-	Ar.LeaveSlot();
-}
-
-void FStructuredArchiveSlot::operator<< (uint16& Value)
-{
-	Ar.EnterSlot(*this);
-	Ar.Formatter.Serialize(Value);
-	Ar.LeaveSlot();
-}
-
-void FStructuredArchiveSlot::operator<< (uint32& Value)
-{
-	Ar.EnterSlot(*this);
-	Ar.Formatter.Serialize(Value);
-	Ar.LeaveSlot();
-}
-
-void FStructuredArchiveSlot::operator<< (uint64& Value)
-{
-	Ar.EnterSlot(*this);
-	Ar.Formatter.Serialize(Value);
-	Ar.LeaveSlot();
-}
-
-void FStructuredArchiveSlot::operator<< (int8& Value)
-{
-	Ar.EnterSlot(*this);
-	Ar.Formatter.Serialize(Value);
-	Ar.LeaveSlot();
-}
-
-void FStructuredArchiveSlot::operator<< (int16& Value)
-{
-	Ar.EnterSlot(*this);
-	Ar.Formatter.Serialize(Value);
-	Ar.LeaveSlot();
-}
-
-void FStructuredArchiveSlot::operator<< (int32& Value)
-{
-	Ar.EnterSlot(*this);
-	Ar.Formatter.Serialize(Value);
-	Ar.LeaveSlot();
-}
-
-void FStructuredArchiveSlot::operator<< (int64& Value)
-{
-	Ar.EnterSlot(*this);
-	Ar.Formatter.Serialize(Value);
-	Ar.LeaveSlot();
-}
-
-void FStructuredArchiveSlot::operator<< (float& Value)
-{
-	Ar.EnterSlot(*this);
-	Ar.Formatter.Serialize(Value);
-	Ar.LeaveSlot();
-}
-
-void FStructuredArchiveSlot::operator<< (double& Value)
-{
-	Ar.EnterSlot(*this);
-	Ar.Formatter.Serialize(Value);
-	Ar.LeaveSlot();
-}
-
-void FStructuredArchiveSlot::operator<< (bool& Value)
-{
-	Ar.EnterSlot(*this);
-	Ar.Formatter.Serialize(Value);
-	Ar.LeaveSlot();
-}
-
-void FStructuredArchiveSlot::operator<< (FString& Value)
-{
-	Ar.EnterSlot(*this);
-	Ar.Formatter.Serialize(Value);
-	Ar.LeaveSlot();
-}
-
-void FStructuredArchiveSlot::operator<< (FName& Value)
-{
-	Ar.EnterSlot(*this);
-	Ar.Formatter.Serialize(Value);
-	Ar.LeaveSlot();
-}
-
-void FStructuredArchiveSlot::operator<< (UObject*& Value)
-{
-	Ar.EnterSlot(*this);
-	Ar.Formatter.Serialize(Value);
-	Ar.LeaveSlot();
-}
-
-void FStructuredArchiveSlot::operator<< (FText& Value)
-{
-	Ar.EnterSlot(*this);
-	Ar.Formatter.Serialize(Value);
-	Ar.LeaveSlot();
-}
-
-void FStructuredArchiveSlot::operator<< (FWeakObjectPtr& Value)
-{
-	Ar.EnterSlot(*this);
-	Ar.Formatter.Serialize(Value);
-	Ar.LeaveSlot();
-}
-
-void FStructuredArchiveSlot::operator<< (FLazyObjectPtr& Value)
-{
-	Ar.EnterSlot(*this);
-	Ar.Formatter.Serialize(Value);
-	Ar.LeaveSlot();
-}
-
-void FStructuredArchiveSlot::operator<< (FObjectPtr& Value)
-{
-	Ar.EnterSlot(*this);
-	Ar.Formatter.Serialize(Value);
-	Ar.LeaveSlot();
-}
-
-void FStructuredArchiveSlot::operator<< (FSoftObjectPtr& Value)
-{
-	Ar.EnterSlot(*this);
-	Ar.Formatter.Serialize(Value);
-	Ar.LeaveSlot();
-}
-
-void FStructuredArchiveSlot::operator<< (FSoftObjectPath& Value)
-{
-	Ar.EnterSlot(*this);
-	Ar.Formatter.Serialize(Value);
-	Ar.LeaveSlot();
-}
-
-void FStructuredArchiveSlot::Serialize(TArray<uint8>& Data)
-{
-	Ar.EnterSlot(*this);
-	Ar.Formatter.Serialize(Data);
-	Ar.LeaveSlot();
-}
-
-void FStructuredArchiveSlot::Serialize(void* Data, uint64 DataSize)
-{
-	Ar.EnterSlot(*this);
-	Ar.Formatter.Serialize(Data, DataSize);
-	Ar.LeaveSlot();
-}
-
-//////////// FStructuredArchiveRecord ////////////
-
-FStructuredArchiveSlot FStructuredArchiveRecord::EnterField(FArchiveFieldName Name)
-{
-	Ar.SetScope(*this);
-
-	Ar.CurrentSlotElementId = Ar.ElementIdGenerator.Generate();
-
-#if DO_STRUCTURED_ARCHIVE_CONTAINER_CHECKS
-#if DO_STRUCTURED_ARCHIVE_UNIQUE_FIELD_NAME_CHECKS
-	if (!Ar.GetUnderlyingArchive().IsLoading())
-	{
-		FContainer& Container = *Ar.CurrentContainer.Top();
-		checkf(!Container.KeyNames.Contains(Name.Name), TEXT("Multiple keys called '%s' serialized into record"), Name.Name);
-		Container.KeyNames.Add(Name.Name);
-	}
-#endif
-#endif
-
-	Ar.Formatter.EnterField(Name);
-
-	return FStructuredArchiveSlot(Ar, Depth, Ar.CurrentSlotElementId);
-}
-
-FStructuredArchiveRecord FStructuredArchiveRecord::EnterRecord(FArchiveFieldName Name)
-{
-	return EnterField(Name).EnterRecord();
-}
-
-FStructuredArchiveArray FStructuredArchiveRecord::EnterArray(FArchiveFieldName Name, int32& Num)
-{
-	return EnterField(Name).EnterArray(Num);
-}
-
-FStructuredArchiveStream FStructuredArchiveRecord::EnterStream(FArchiveFieldName Name)
-{
-	return EnterField(Name).EnterStream();
-}
-
-FStructuredArchiveMap FStructuredArchiveRecord::EnterMap(FArchiveFieldName Name, int32& Num)
-{
-	return EnterField(Name).EnterMap(Num);
-}
-
-TOptional<FStructuredArchiveSlot> FStructuredArchiveRecord::TryEnterField(FArchiveFieldName Name, bool bEnterWhenWriting)
-{
-	Ar.SetScope(*this);
-
-#if DO_STRUCTURED_ARCHIVE_CONTAINER_CHECKS
-#if DO_STRUCTURED_ARCHIVE_UNIQUE_FIELD_NAME_CHECKS
-	if (!GetUnderlyingArchive().IsLoading())
-	{
-		FContainer& Container = *Ar.CurrentContainer.Top();
-		checkf(!Container.KeyNames.Contains(Name.Name), TEXT("Multiple keys called '%s' serialized into record"), Name.Name);
-		Container.KeyNames.Add(Name.Name);
-	}
-#endif
-#endif
-
-	if (Ar.Formatter.TryEnterField(Name, bEnterWhenWriting))
-	{
-		Ar.CurrentSlotElementId = Ar.ElementIdGenerator.Generate();
-		return FStructuredArchiveSlot(Ar, Depth, Ar.CurrentSlotElementId);
-	}
-	else
-	{
-		return {};
-	}
-}
-
-//////////// FStructuredArchiveArray ////////////
-
-FStructuredArchiveSlot FStructuredArchiveArray::EnterElement()
-{
-	Ar.SetScope(*this);
-
-#if DO_STRUCTURED_ARCHIVE_CONTAINER_CHECKS
-	checkf(Ar.CurrentContainer.Top()->Index < Ar.CurrentContainer.Top()->Count, TEXT("Serialized too many array elements"));
-#endif
-
-	Ar.CurrentSlotElementId = Ar.ElementIdGenerator.Generate();
-
-	Ar.Formatter.EnterArrayElement();
-
-	return FStructuredArchiveSlot(Ar, Depth, Ar.CurrentSlotElementId);
-}
-
-//////////// FStructuredArchiveStream ////////////
-
-FStructuredArchiveSlot FStructuredArchiveStream::EnterElement()
-{
-	Ar.SetScope(*this);
-
-	Ar.CurrentSlotElementId = Ar.ElementIdGenerator.Generate();
-
-	Ar.Formatter.EnterStreamElement();
-
-	return FStructuredArchiveSlot(Ar, Depth, Ar.CurrentSlotElementId);
-}
-
-//////////// FStructuredArchiveMap ////////////
-
-FStructuredArchiveSlot FStructuredArchiveMap::EnterElement(FString& Name)
-{
-	Ar.SetScope(*this);
-
-#if DO_STRUCTURED_ARCHIVE_CONTAINER_CHECKS
-	checkf(Ar.CurrentContainer.Top()->Index < Ar.CurrentContainer.Top()->Count, TEXT("Serialized too many map elements"));
-#endif
-
-	Ar.CurrentSlotElementId = Ar.ElementIdGenerator.Generate();
-
-#if DO_STRUCTURED_ARCHIVE_CONTAINER_CHECKS
-#if DO_STRUCTURED_ARCHIVE_UNIQUE_FIELD_NAME_CHECKS
-	if(Ar.GetUnderlyingArchive().IsSaving())
-	{
-		FContainer& Container = *Ar.CurrentContainer.Top();
-		checkf(!Container.KeyNames.Contains(Name), TEXT("Multiple keys called '%s' serialized into record"), *Name);
-		Container.KeyNames.Add(Name);
-	}
-#endif
-#endif
-
-	Ar.Formatter.EnterMapElement(Name);
-
-#if DO_STRUCTURED_ARCHIVE_CONTAINER_CHECKS
-#if DO_STRUCTURED_ARCHIVE_UNIQUE_FIELD_NAME_CHECKS
-	if(Ar.GetUnderlyingArchive().IsLoading())
-	{
-		FContainer& Container = *Ar.CurrentContainer.Top();
-		checkf(!Container.KeyNames.Contains(Name), TEXT("Multiple keys called '%s' serialized into record"), *Name);
-		Container.KeyNames.Add(Name);
-	}
-#endif
-#endif
-
-	return FStructuredArchiveSlot(Ar, Depth, Ar.CurrentSlotElementId);
-}
-
 #endif
diff --git a/Engine/Source/Runtime/Core/Private/Serialization/StructuredArchiveChildReader.cpp b/Engine/Source/Runtime/Core/Private/Serialization/StructuredArchiveChildReader.cpp
new file mode 100644
--- /dev/null
+++ b/Engine/Source/Runtime/Core/Private/Serialization/StructuredArchiveChildReader.cpp
@@ -0,0 +1,38 @@
+// Copyright Epic Games, Inc. All Rights Reserved.
+
+#include "Serialization/StructuredArchiveChildReader.h"
+
+#if WITH_TEXT_ARCHIVE_SUPPORT
+
+FStructuredArchiveChildReader::FStructuredArchiveChildReader(FStructuredArchiveSlot InSlot)
+	: OwnedFormatter(nullptr)
+	, Archive(nullptr)
+{
+	FStructuredArchiveFormatter* Formatter = &InSlot.Ar.Formatter;
+	if (InSlot.GetUnderlyingArchive().IsTextFormat())
+	{
+		Formatter = OwnedFormatter = InSlot.Ar.Formatter.CreateSubtreeReader();
+	}
+
+	Archive = new FStructuredArchive(*Formatter);
+	Root.Emplace(Archive->Open());
+	InSlot.EnterRecord();
+}
+
+FStructuredArchiveChildReader::~FStructuredArchiveChildReader()
+{
+	Root.Reset();
+	Archive->Close();
+	delete Archive;
+	Archive = nullptr;
+
+	// If this is a text archive, we'll have created a subtree reader that our contained archive is using as 
+	// its formatter. We need to clean it up now.
+	if (OwnedFormatter)
+	{
+		delete OwnedFormatter;
+		OwnedFormatter = nullptr;
+	}
+}
+
+#endif
diff --git a/Engine/Source/Runtime/Core/Private/Serialization/StructuredArchiveContainer.h b/Engine/Source/Runtime/Core/Private/Serialization/StructuredArchiveContainer.h
new file mode 100644
--- /dev/null
+++ b/Engine/Source/Runtime/Core/Private/Serialization/StructuredArchiveContainer.h
@@ -0,0 +1,33 @@
+// Copyright Epic Games, Inc. All Rights Reserved.
+
+#pragma once
+
+#include "Serialization/StructuredArchive.h"
+#include "Serialization/StructuredArchiveDefines.h"
+#include "Containers/Set.h"
+#include "Containers/UnrealString.h"
+
+#if WITH_TEXT_ARCHIVE_SUPPORT
+
+//////////// FStructuredArchive::FContainer ////////////
+
+#if DO_STRUCTURED_ARCHIVE_CONTAINER_CHECKS
+
+struct FStructuredArchive::FContainer
+{
+	int  Index = 0;
+	int  Count = 0;
+	bool bAttributedValueWritten = false;
+
+#if DO_STRUCTURED_ARCHIVE_UNIQUE_FIELD_NAME_CHECKS
+	TSet<FString> KeyNames;
+#endif
+
+	explicit FContainer(int InCount)
+		: Count(InCount)
+	{
+	}
+};
+#endif
+
+#endif
diff --git a/Engine/Source/Runtime/Core/Private/Serialization/StructuredArchiveSlots.cpp b/Engine/Source/Runtime/Core/Private/Serialization/StructuredArchiveSlots.cpp
new file mode 100644
--- /dev/null
+++ b/Engine/Source/Runtime/Core/Private/Serialization/StructuredArchiveSlots.cpp
@@ -0,0 +1,439 @@
+// Copyright Epic Games, Inc. All Rights Reserved.
+
+#include "Serialization/StructuredArchiveSlots.h"
+#include "Serialization/StructuredArchiveContainer.h"
+#include "Containers/UnrealString.h"
+
+#if WITH_TEXT_ARCHIVE_SUPPORT
+
+//////////// FStructuredArchiveSlot ////////////
+
+FStructuredArchiveRecord FStructuredArchiveSlot::EnterRecord()
+{
+	int32 NewDepth = Ar.EnterSlotAsType(*this, UE::StructuredArchive::Private::EElementType::Record);
+
+#if DO_STRUCTURED_ARCHIVE_CONTAINER_CHECKS
+	Ar.CurrentContainer.Emplace(0);
+#endif
+
+	Ar.Formatter.EnterRecord();
+
+	return FStructuredArchiveRecord(FStructuredArchiveRecord::EPrivateToken{}, Ar, NewDepth, ElementId);
+}
+
+FStructuredArchiveArray FStructuredArchiveSlot::EnterArray(int32& Num)
+{
+	int32 NewDepth = Ar.EnterSlotAsType(*this, UE::StructuredArchive::Private::EElementType::Array);
+
+	Ar.Formatter.EnterArray(Num);
+
+#if DO_STRUCTURED_ARCHIVE_CONTAINER_CHECKS
+	Ar.CurrentContainer.Emplace(Num);
+#endif
+
+	return FStructuredArchiveArray(FStructuredArchiveArray::EPrivateToken{}, Ar, NewDepth, ElementId);
+}
+
+FStructuredArchiveStream FStructuredArchiveSlot::EnterStream()
+{
+	int32 NewDepth = Ar.EnterSlotAsType(*this, UE::StructuredArchive::Private::EElementType::Stream);
+
+	Ar.Formatter.EnterStream();
+
+	return FStructuredArchiveStream(FStructuredArchiveStream::EPrivateToken{}, Ar, NewDepth, ElementId);
+}
+
+FStructuredArchiveMap FStructuredArchiveSlot::EnterMap(int32& Num)
+{
+	int32 NewDepth = Ar.EnterSlotAsType(*this, UE::StructuredArchive::Private::EElementType::Map);
+
+	Ar.Formatter.EnterMap(Num);
+
+#if DO_STRUCTURED_ARCHIVE_CONTAINER_CHECKS
+	Ar.CurrentContainer.Emplace(Num);
+#endif
+
+	return FStructuredArchiveMap(FStructuredArchiveMap::EPrivateToken{}, Ar, NewDepth, ElementId);
+}
+
+FStructuredArchiveSlot FStructuredArchiveSlot::EnterAttribute(FArchiveFieldName AttributeName)
+{
+	check(Ar.CurrentScope.Num() > 0);
+
+	int32 NewDepth = Depth + 1;
+	if (NewDepth >= Ar.CurrentScope.Num() || Ar.CurrentScope[NewDepth].Id != ElementId || Ar.CurrentScope[NewDepth].Type != UE::StructuredArchive::Private::EElementType::AttributedValue)
+	{
+		int32 NewDepthCheck = Ar.EnterSlotAsType(*this, UE::StructuredArchive::Private::EElementType::AttributedValue);
+		checkSlow(NewDepth == NewDepthCheck);
+
+		Ar.Formatter.EnterAttributedValue();
+
+#if DO_STRUCTURED_ARCHIVE_CONTAINER_CHECKS
+		Ar.CurrentContainer.Emplace(0);
+#endif
+	}
+
+	Ar.CurrentEnteringAttributeState = UE::StructuredArchive::Private::EEnteringAttributeState::NotEnteringAttribute;
+
+	UE::StructuredArchive::Private::FElementId AttributedValueId = Ar.CurrentScope[NewDepth].Id;
+
+	Ar.SetScope(UE::StructuredArchive::Private::FSlotPosition(NewDepth, AttributedValueId));
+
+	Ar.CurrentSlotElementId = Ar.ElementIdGenerator.Generate();
+
+#if DO_STRUCTURED_ARCHIVE_CONTAINER_CHECKS
+#if DO_STRUCTURED_ARCHIVE_UNIQUE_FIELD_NAME_CHECKS
+	if (!Ar.GetUnderlyingArchive().IsLoading())
+	{
+		FContainer& Container = *Ar.CurrentContainer.Top();
+		checkf(!Container.KeyNames.Contains(Name.Name), TEXT("Multiple attributes called '%s' serialized into attributed value"), AttributeName.Name);
+		Container.KeyNames.Add(Name.Name);
+	}
+#endif
+#endif
+
+	Ar.Formatter.EnterAttribute(AttributeName);
+
+	return FStructuredArchiveSlot(FStructuredArchiveSlot::EPrivateToken{}, Ar, NewDepth, Ar.CurrentSlotElementId);
+}
+
+TOptional<FStructuredArchiveSlot> FStructuredArchiveSlot::TryEnterAttribute(FArchiveFieldName AttributeName, bool bEnterWhenWriting)
+{
+	check(Ar.CurrentScope.Num() > 0);
+
+	int32 NewDepth = Depth + 1;
+	if (NewDepth >= Ar.CurrentScope.Num() || Ar.CurrentScope[NewDepth].Id != ElementId || Ar.CurrentScope[NewDepth].Type != UE::StructuredArchive::Private::EElementType::AttributedValue)
+	{
+		int32 NewDepthCheck = Ar.EnterSlotAsType(*this, UE::StructuredArchive::Private::EElementType::AttributedValue);
+		checkSlow(NewDepth == NewDepthCheck);
+
+		Ar.Formatter.EnterAttributedValue();
+
+#if DO_STRUCTURED_ARCHIVE_CONTAINER_CHECKS
+		Ar.CurrentContainer.Emplace(0);
+#endif
+	}
+
+#if DO_STRUCTURED_ARCHIVE_CONTAINER_CHECKS
+#if DO_STRUCTURED_ARCHIVE_UNIQUE_FIELD_NAME_CHECKS
+	if (!Ar.GetUnderlyingArchive().IsLoading())
+	{
+		FContainer& Container = *Ar.CurrentContainer.Top();
+		checkf(!Container.KeyNames.Contains(Name.Name), TEXT("Multiple attributes called '%s' serialized into attributed value"), AttributeName.Name);
+		Container.KeyNames.Add(Name.Name);
+	}
+#endif
+#endif
+
+	UE::StructuredArchive::Private::FElementId AttributedValueId = Ar.CurrentScope[NewDepth].Id;
+
+	Ar.SetScope(UE::StructuredArchive::Private::FSlotPosition(NewDepth, AttributedValueId));
+
+	if (Ar.Formatter.TryEnterAttribute(AttributeName, bEnterWhenWriting))
+	{
+		Ar.CurrentEnteringAttributeState = UE::StructuredArchive::Private::EEnteringAttributeState::NotEnteringAttribute;
+
+		Ar.CurrentSlotElementId = Ar.ElementIdGenerator.Generate();
+
+		return FStructuredArchiveSlot(FStructuredArchiveSlot::EPrivateToken{}, Ar, NewDepth, Ar.CurrentSlotElementId);
+	}
+	else
+	{
+		return {};
+	}
+}
+
+void FStructuredArchiveSlot::operator<< (uint8& Value)
+{
+	Ar.EnterSlot(*this);
+	Ar.Formatter.Serialize(Value);
+	Ar.LeaveSlot();
+}
+
+void FStructuredArchiveSlot::operator<< (uint16& Value)
+{
+	Ar.EnterSlot(*this);
+	Ar.Formatter.Serialize(Value);
+	Ar.LeaveSlot();
+}
+
+void FStructuredArchiveSlot::operator<< (uint32& Value)
+{
+	Ar.EnterSlot(*this);
+	Ar.Formatter.Serialize(Value);
+	Ar.LeaveSlot();
+}
+
+void FStructuredArchiveSlot::operator<< (uint64& Value)
+{
+	Ar.EnterSlot(*this);
+	Ar.Formatter.Serialize(Value);
+	Ar.LeaveSlot();
+}
+
+void FStructuredArchiveSlot::operator<< (int8& Value)
+{
+	Ar.EnterSlot(*this);
+	Ar.Formatter.Serialize(Value);
+	Ar.LeaveSlot();
+}
+
+void FStructuredArchiveSlot::operator<< (int16& Value)
+{
+	Ar.EnterSlot(*this);
+	Ar.Formatter.Serialize(Value);
+	Ar.LeaveSlot();
+}
+
+void FStructuredArchiveSlot::operator<< (int32& Value)
+{
+	Ar.EnterSlot(*this);
+	Ar.Formatter.Serialize(Value);
+	Ar.LeaveSlot();
+}
+
+void FStructuredArchiveSlot::operator<< (int64& Value)
+{
+	Ar.EnterSlot(*this);
+	Ar.Formatter.Serialize(Value);
+	Ar.LeaveSlot();
+}
+
+void FStructuredArchiveSlot::operator<< (float& Value)
+{
+	Ar.EnterSlot(*this);
+	Ar.Formatter.Serialize(Value);
+	Ar.LeaveSlot();
+}
+
+void FStructuredArchiveSlot::operator<< (double& Value)
+{
+	Ar.EnterSlot(*this);
+	Ar.Formatter.Serialize(Value);
+	Ar.LeaveSlot();
+}
+
+void FStructuredArchiveSlot::operator<< (bool& Value)
+{
+	Ar.EnterSlot(*this);
+	Ar.Formatter.Serialize(Value);
+	Ar.LeaveSlot();
+}
+
+void FStructuredArchiveSlot::operator<< (FString& Value)
+{
+	Ar.EnterSlot(*this);
+	Ar.Formatter.Serialize(Value);
+	Ar.LeaveSlot();
+}
+
+void FStructuredArchiveSlot::operator<< (FName& Value)
+{
+	Ar.EnterSlot(*this);
+	Ar.Formatter.Serialize(Value);
+	Ar.LeaveSlot();
+}
+
+void FStructuredArchiveSlot::operator<< (UObject*& Value)
+{
+	Ar.EnterSlot(*this);
+	Ar.Formatter.Serialize(Value);
+	Ar.LeaveSlot();
+}
+
+void FStructuredArchiveSlot::operator<< (FText& Value)
+{
+	Ar.EnterSlot(*this);
+	Ar.Formatter.Serialize(Value);
+	Ar.LeaveSlot();
+}
+
+void FStructuredArchiveSlot::operator<< (FWeakObjectPtr& Value)
+{
+	Ar.EnterSlot(*this);
+	Ar.Formatter.Serialize(Value);
+	Ar.LeaveSlot();
+}
+
+void FStructuredArchiveSlot::operator<< (FLazyObjectPtr& Value)
+{
+	Ar.EnterSlot(*this);
+	Ar.Formatter.Serialize(Value);
+	Ar.LeaveSlot();
+}
+
+void FStructuredArchiveSlot::operator<< (FObjectPtr& Value)
+{
+	Ar.EnterSlot(*this);
+	Ar.Formatter.Serialize(Value);
+	Ar.LeaveSlot();
+}
+
+void FStructuredArchiveSlot::operator<< (FSoftObjectPtr& Value)
+{
+	Ar.EnterSlot(*this);
+	Ar.Formatter.Serialize(Value);
+	Ar.LeaveSlot();
+}
+
+void FStructuredArchiveSlot::operator<< (FSoftObjectPath& Value)
+{
+	Ar.EnterSlot(*this);
+	Ar.Formatter.Serialize(Value);
+	Ar.LeaveSlot();
+}
+
+void FStructuredArchiveSlot::Serialize(TArray<uint8>& Data)
+{
+	Ar.EnterSlot(*this);
+	Ar.Formatter.Serialize(Data);
+	Ar.LeaveSlot();
+}
+
+void FStructuredArchiveSlot::Serialize(void* Data, uint64 DataSize)
+{
+	Ar.EnterSlot(*this);
+	Ar.Formatter.Serialize(Data, DataSize);
+	Ar.LeaveSlot();
+}
+
+//////////// FStructuredArchiveRecord ////////////
+
+FStructuredArchiveSlot FStructuredArchiveRecord::EnterField(FArchiveFieldName Name)
+{
+	Ar.SetScope(*this);
+
+	Ar.CurrentSlotElementId = Ar.ElementIdGenerator.Generate();
+
+#if DO_STRUCTURED_ARCHIVE_CONTAINER_CHECKS
+#if DO_STRUCTURED_ARCHIVE_UNIQUE_FIELD_NAME_CHECKS
+	if (!Ar.GetUnderlyingArchive().IsLoading())
+	{
+		FContainer& Container = *Ar.CurrentContainer.Top();
+		checkf(!Container.KeyNames.Contains(Name.Name), TEXT("Multiple keys called '%s' serialized into record"), Name.Name);
+		Container.KeyNames.Add(Name.Name);
+	}
+#endif
+#endif
+
+	Ar.Formatter.EnterField(Name);
+
+	return FStructuredArchiveSlot(FStructuredArchiveSlot::EPrivateToken{}, Ar, Depth, Ar.CurrentSlotElementId);
+}
+
+FStructuredArchiveRecord FStructuredArchiveRecord::EnterRecord(FArchiveFieldName Name)
+{
+	return EnterField(Name).EnterRecord();
+}
+
+FStructuredArchiveArray FStructuredArchiveRecord::EnterArray(FArchiveFieldName Name, int32& Num)
+{
+	return EnterField(Name).EnterArray(Num);
+}
+
+FStructuredArchiveStream FStructuredArchiveRecord::EnterStream(FArchiveFieldName Name)
+{
+	return EnterField(Name).EnterStream();
+}
+
+FStructuredArchiveMap FStructuredArchiveRecord::EnterMap(FArchiveFieldName Name, int32& Num)
+{
+	return EnterField(Name).EnterMap(Num);
+}
+
+TOptional<FStructuredArchiveSlot> FStructuredArchiveRecord::TryEnterField(FArchiveFieldName Name, bool bEnterWhenWriting)
+{
+	Ar.SetScope(*this);
+
+#if DO_STRUCTURED_ARCHIVE_CONTAINER_CHECKS
+#if DO_STRUCTURED_ARCHIVE_UNIQUE_FIELD_NAME_CHECKS
+	if (!GetUnderlyingArchive().IsLoading())
+	{
+		FContainer& Container = *Ar.CurrentContainer.Top();
+		checkf(!Container.KeyNames.Contains(Name.Name), TEXT("Multiple keys called '%s' serialized into record"), Name.Name);
+		Container.KeyNames.Add(Name.Name);
+	}
+#endif
+#endif
+
+	if (Ar.Formatter.TryEnterField(Name, bEnterWhenWriting))
+	{
+		Ar.CurrentSlotElementId = Ar.ElementIdGenerator.Generate();
+		return FStructuredArchiveSlot(FStructuredArchiveSlot::EPrivateToken{}, Ar, Depth, Ar.CurrentSlotElementId);
+	}
+	else
+	{
+		return {};
+	}
+}
+
+//////////// FStructuredArchiveArray ////////////
+
+FStructuredArchiveSlot FStructuredArchiveArray::EnterElement()
+{
+	Ar.SetScope(*this);
+
+#if DO_STRUCTURED_ARCHIVE_CONTAINER_CHECKS
+	checkf(Ar.CurrentContainer.Top()->Index < Ar.CurrentContainer.Top()->Count, TEXT("Serialized too many array elements"));
+#endif
+
+	Ar.CurrentSlotElementId = Ar.ElementIdGenerator.Generate();
+
+	Ar.Formatter.EnterArrayElement();
+
+	return FStructuredArchiveSlot(FStructuredArchiveSlot::EPrivateToken{}, Ar, Depth, Ar.CurrentSlotElementId);
+}
+
+//////////// FStructuredArchiveStream ////////////
+
+FStructuredArchiveSlot FStructuredArchiveStream::EnterElement()
+{
+	Ar.SetScope(*this);
+
+	Ar.CurrentSlotElementId = Ar.ElementIdGenerator.Generate();
+
+	Ar.Formatter.EnterStreamElement();
+
+	return FStructuredArchiveSlot(FStructuredArchiveSlot::EPrivateToken{}, Ar, Depth, Ar.CurrentSlotElementId);
+}
+
+//////////// FStructuredArchiveMap ////////////
+
+FStructuredArchiveSlot FStructuredArchiveMap::EnterElement(FString& Name)
+{
+	Ar.SetScope(*this);
+
+#if DO_STRUCTURED_ARCHIVE_CONTAINER_CHECKS
+	checkf(Ar.CurrentContainer.Top()->Index < Ar.CurrentContainer.Top()->Count, TEXT("Serialized too many map elements"));
+#endif
+
+	Ar.CurrentSlotElementId = Ar.ElementIdGenerator.Generate();
+
+#if DO_STRUCTURED_ARCHIVE_CONTAINER_CHECKS
+#if DO_STRUCTURED_ARCHIVE_UNIQUE_FIELD_NAME_CHECKS
+	if(Ar.GetUnderlyingArchive().IsSaving())
+	{
+		FContainer& Container = *Ar.CurrentContainer.Top();
+		checkf(!Container.KeyNames.Contains(Name), TEXT("Multiple keys called '%s' serialized into record"), *Name);
+		Container.KeyNames.Add(Name);
+	}
+#endif
+#endif
+
+	Ar.Formatter.EnterMapElement(Name);
+
+#if DO_STRUCTURED_ARCHIVE_CONTAINER_CHECKS
+#if DO_STRUCTURED_ARCHIVE_UNIQUE_FIELD_NAME_CHECKS
+	if(Ar.GetUnderlyingArchive().IsLoading())
+	{
+		FContainer& Container = *Ar.CurrentContainer.Top();
+		checkf(!Container.KeyNames.Contains(Name), TEXT("Multiple keys called '%s' serialized into record"), *Name);
+		Container.KeyNames.Add(Name);
+	}
+#endif
+#endif
+
+	return FStructuredArchiveSlot(FStructuredArchiveSlot::EPrivateToken{}, Ar, Depth, Ar.CurrentSlotElementId);
+}
+
+#endif
diff --git a/Engine/Source/Runtime/Core/Public/Serialization/StructuredArchive.h b/Engine/Source/Runtime/Core/Public/Serialization/StructuredArchive.h
--- a/Engine/Source/Runtime/Core/Public/Serialization/StructuredArchive.h
+++ b/Engine/Source/Runtime/Core/Public/Serialization/StructuredArchive.h
@@ -3,174 +3,18 @@
 #pragma once
 
 #include "CoreTypes.h"
+#include "Serialization/StructuredArchiveFwd.h"
 #include "Serialization/StructuredArchiveDefines.h"
+#include "Serialization/StructuredArchiveAdapters.h"
 #include "Serialization/StructuredArchiveFormatter.h"
+#include "Serialization/StructuredArchiveNameHelpers.h"
+#include "Serialization/StructuredArchiveSlots.h"
 #include "Formatters/BinaryArchiveFormatter.h"
-#include "Misc/Optional.h"
-#include "Concepts/Insertable.h"
-#include "Templates/Models.h"
 #include "Containers/Array.h"
-#include "Serialization/ArchiveProxy.h"
 #include "Templates/UniqueObj.h"
 
 namespace UE::StructuredArchive::Private
 {
-	/**
-	 * Class to contain a named value for serialization. Intended to be created as a temporary and passed to object serialization methods.
-	 */
-	template <typename T>
-	struct TNamedValue
-	{
-		FArchiveFieldName Name;
-		T& Value;
-	};
-
-	/**
-	 * Class to contain a named attribute for serialization. Intended to be created as a temporary and passed to object serialization methods.
-	 */
-	template <typename T>
-	struct TNamedAttribute
-	{
-		FArchiveFieldName Name;
-		T& Value;
-	};
-
-	/**
-	 * Class to contain a named attribute for serialization, with a default. Intended to be created as a temporary and passed to object
-	 * serialization methods, which can choose not to serialize the attribute if it matches the default.
-	 */
-	template <typename T>
-	struct TOptionalNamedAttribute
-	{
-		FArchiveFieldName Name;
-		T& Value;
-		const T& Default;
-	};
-
-	template <typename T>
-	FORCEINLINE TNamedValue<T> MakeNamedValue(FArchiveFieldName Name, T& Value)
-	{
-		return TNamedValue<T>{ Name, Value };
-	}
-
-	template <typename T>
-	FORCEINLINE TNamedAttribute<T> MakeNamedAttribute(FArchiveFieldName Name, T& Value)
-	{
-		return TNamedAttribute<T>{ Name, Value };
-	}
-
-	template <typename T>
-	FORCEINLINE TOptionalNamedAttribute<T> MakeOptionalNamedAttribute(FArchiveFieldName Name, T& Value, const typename TIdentity<T>::Type& Default)
-	{
-		return TOptionalNamedAttribute<T>{ Name, Value, Default };
-	}
-}
-
-/** Construct a TNamedValue given an ANSI string and value reference. */
-#define SA_VALUE(Name, Value) UE::StructuredArchive::Private::MakeNamedValue(FArchiveFieldName(Name), Value)
-
-/** Construct a TNamedAttribute given an ANSI string and value reference. */
-#define SA_ATTRIBUTE(Name, Value) UE::StructuredArchive::Private::MakeNamedAttribute(FArchiveFieldName(Name), Value)
-
-/** Construct a TOptionalNamedAttribute given an ANSI string and value reference. */
-#define SA_OPTIONAL_ATTRIBUTE(Name, Value, Default) UE::StructuredArchive::Private::MakeOptionalNamedAttribute(FArchiveFieldName(Name), Value, Default)
-
-/** Typedef for which formatter type to support */
-#if WITH_TEXT_ARCHIVE_SUPPORT
-	typedef FStructuredArchiveFormatter FArchiveFormatterType;
-#else
-	typedef FBinaryArchiveFormatter FArchiveFormatterType;
-#endif
-
-class FStructuredArchive;
-class FStructuredArchiveChildReader;
-class FStructuredArchiveSlot;
-class FStructuredArchiveRecord;
-class FStructuredArchiveArray;
-class FStructuredArchiveStream;
-class FStructuredArchiveMap;
-
-namespace UE::StructuredArchive::Private
-{
-	struct FElementId
-	{
-		FElementId() = default;
-
-		explicit FElementId(uint32 InId)
-			: Id(InId)
-		{
-		}
-
-		bool IsValid() const
-		{
-			return Id != 0;
-		}
-
-		void Reset()
-		{
-			Id = 0;
-		}
-
-		bool operator==(const FElementId& Rhs) const
-		{
-			return Id == Rhs.Id;
-		}
-
-		bool operator!=(const FElementId& Rhs) const
-		{
-			return Id != Rhs.Id;
-		}
-
-	private:
-		uint32 Id = 0;
-	};
-
-	// Represents a position of a slot within the hierarchy.
-	class CORE_API FSlotPosition
-	{
-		friend class FStructuredArchive;
-
-	public:
-		int32 Depth;
-		FElementId ElementId;
-
-		FORCEINLINE explicit FSlotPosition(int32 InDepth, FElementId InElementId)
-			: Depth(InDepth)
-			, ElementId(InElementId)
-		{
-		}
-	};
-
-	// The base class of all slot types
-	class CORE_API FSlotBase
-#if WITH_TEXT_ARCHIVE_SUPPORT
-		: protected FSlotPosition
-#endif
-	{
-		friend class FStructuredArchive;
-
-	public:
-		FArchive& GetUnderlyingArchive() const;
-
-		const FArchiveState& GetArchiveState() const;
-
-	protected:
-		FStructuredArchive& Ar;
-
-#if WITH_TEXT_ARCHIVE_SUPPORT
-		FORCEINLINE explicit FSlotBase(FStructuredArchive& InAr, int32 InDepth, FElementId InElementId)
-			: FSlotPosition(InDepth, InElementId)
-			, Ar(InAr)
-		{
-		}
-#else
-		FORCEINLINE FSlotBase(FStructuredArchive& InAr)
-			: Ar(InAr)
-		{
-		}
-#endif
-	};
-
 	enum class EElementType : unsigned char
 	{
 		Root,
@@ -186,185 +30,11 @@ namespace UE::StructuredArchive::Private
 		NotEnteringAttribute,
 		EnteringAttribute,
 	};
-}
-
-/**
- * Contains a value in the archive; either a field or array/map element. A slot does not know it's name or location,
- * and can merely have a value serialized into it. That value may be a literal (eg. int, float) or compound object
- * (eg. object, array, map).
- */
-class CORE_API FStructuredArchiveSlot final : public UE::StructuredArchive::Private::FSlotBase
-{
-public:
-	FStructuredArchiveRecord EnterRecord();
-	FStructuredArchiveArray EnterArray(int32& Num);
-	FStructuredArchiveStream EnterStream();
-	FStructuredArchiveMap EnterMap(int32& Num);
-	FStructuredArchiveSlot EnterAttribute(FArchiveFieldName AttributeName);
-	TOptional<FStructuredArchiveSlot> TryEnterAttribute(FArchiveFieldName AttributeName, bool bEnterWhenWriting);
-
-	// We don't support chaining writes to a single slot, so this returns void.
-	void operator << (uint8& Value);
-	void operator << (uint16& Value);
-	void operator << (uint32& Value);
-	void operator << (uint64& Value);
-	void operator << (int8& Value);
-	void operator << (int16& Value);
-	void operator << (int32& Value);
-	void operator << (int64& Value);
-	void operator << (float& Value);
-	void operator << (double& Value);
-	void operator << (bool& Value);
-	void operator << (FString& Value);
-	void operator << (FName& Value);
-	void operator << (UObject*& Value);
-	void operator << (FText& Value);
-	void operator << (FWeakObjectPtr& Value);
-	void operator << (FSoftObjectPtr& Value);
-	void operator << (FSoftObjectPath& Value);
-	void operator << (FLazyObjectPtr& Value);
-	void operator << (FObjectPtr& Value);
-
-	template <typename T>
-	FORCEINLINE void operator<<(TEnumAsByte<T>& Value)
-	{
-		uint8 Tmp = (uint8)Value.GetValue();
-		*this << Tmp;
-		Value = (T)Tmp;
-	}
 
-	template <
-		typename EnumType,
-		typename = typename TEnableIf<TIsEnumClass<EnumType>::Value>::Type
-	>
-	FORCEINLINE void operator<<(EnumType& Value)
-	{
-		*this << (__underlying_type(EnumType)&)Value;
-	}
-
-	template <typename T>
-	FORCEINLINE void operator<<(UE::StructuredArchive::Private::TNamedAttribute<T> Item)
-	{
-		EnterAttribute(Item.Name) << Item.Value;
-	}
-
-	template <typename T>
-	FORCEINLINE void operator<<(UE::StructuredArchive::Private::TOptionalNamedAttribute<T> Item)
-	{
-		if (TOptional<FStructuredArchiveSlot> Attribute = TryEnterAttribute(Item.Name, Item.Value != Item.Default))
-		{
-			Attribute.GetValue() << Item.Value;
-		}
-		else
-		{
-			Item.Value = Item.Default;
-		}
+	FElementId GetCurrentSlotElementIdImpl(FStructuredArchive& Ar);
+	FArchiveFormatterType& GetFormatterImpl(FStructuredArchive& Ar);
 	}
 
-	void Serialize(TArray<uint8>& Data);
-	void Serialize(void* Data, uint64 DataSize);
-
-	bool IsFilled() const;
-
-private:
-	friend FStructuredArchive;
-	friend FStructuredArchiveChildReader;
-	friend FStructuredArchiveSlot;
-	friend FStructuredArchiveRecord;
-	friend FStructuredArchiveArray;
-	friend FStructuredArchiveStream;
-	friend FStructuredArchiveMap;
-
-	using UE::StructuredArchive::Private::FSlotBase::FSlotBase;
-};
-
-/**
- * Represents a record in the structured archive. An object contains slots that are identified by FArchiveName,
- * which may be compiled out with binary-only archives.
- */
-class CORE_API FStructuredArchiveRecord final : public UE::StructuredArchive::Private::FSlotBase
-{
-public:
-	FStructuredArchiveSlot EnterField(FArchiveFieldName Name);
-	FStructuredArchiveRecord EnterRecord(FArchiveFieldName Name);
-	FStructuredArchiveArray EnterArray(FArchiveFieldName Name, int32& Num);
-	FStructuredArchiveStream EnterStream(FArchiveFieldName Name);
-	FStructuredArchiveMap EnterMap(FArchiveFieldName Name, int32& Num);
-
-	TOptional<FStructuredArchiveSlot> TryEnterField(FArchiveFieldName Name, bool bEnterForSaving);
-
-	template<typename T> FORCEINLINE FStructuredArchiveRecord& operator<<(UE::StructuredArchive::Private::TNamedValue<T> Item)
-	{
-		EnterField(Item.Name) << Item.Value;
-		return *this;
-	}
-
-private:
-	friend FStructuredArchive;
-	friend FStructuredArchiveSlot;
-
-	using UE::StructuredArchive::Private::FSlotBase::FSlotBase;
-};
-
-/**
- * Represents an array in the structured archive. An object contains slots that are identified by a FArchiveFieldName,
- * which may be compiled out with binary-only archives.
- */
-class CORE_API FStructuredArchiveArray final : public UE::StructuredArchive::Private::FSlotBase
-{
-public:
-	FStructuredArchiveSlot EnterElement();
-
-	template<typename T> FORCEINLINE FStructuredArchiveArray& operator<<(T& Item)
-	{
-		EnterElement() << Item;
-		return *this;
-	}
-
-private:
-	friend FStructuredArchive;
-	friend FStructuredArchiveSlot;
-
-	using UE::StructuredArchive::Private::FSlotBase::FSlotBase;
-};
-
-/**
- * Represents an unsized sequence of slots in the structured archive (similar to an array, but without a known size).
- */
-class CORE_API FStructuredArchiveStream final : public UE::StructuredArchive::Private::FSlotBase
-{
-public:
-	FStructuredArchiveSlot EnterElement();
-
-	template<typename T> FORCEINLINE FStructuredArchiveStream& operator<<(T& Item)
-	{
-		EnterElement() << Item;
-		return *this;
-	}
-
-private:
-	friend FStructuredArchive;
-	friend FStructuredArchiveSlot;
-
-	using UE::StructuredArchive::Private::FSlotBase::FSlotBase;
-};
-
-/**
- * Represents a map in the structured archive. A map is similar to a record, but keys can be read back out from an archive.
- * (This is an important distinction for binary archives).
- */
-class CORE_API FStructuredArchiveMap final : public UE::StructuredArchive::Private::FSlotBase
-{
-public:
-	FStructuredArchiveSlot EnterElement(FString& Name);
-
-private:
-	friend FStructuredArchive;
-	friend FStructuredArchiveSlot;
-
-	using UE::StructuredArchive::Private::FSlotBase::FSlotBase;
-};
-
 /**
  * Manages the state of an underlying FStructuredArchiveFormatter, and provides a consistent API for reading and writing to a structured archive.
  * 
@@ -391,7 +61,7 @@ public:
 	 *
 	 * @param InFormatter Formatter for the archive data
 	 */
-	FStructuredArchive(FArchiveFormatterType& InFormatter);
+	explicit FStructuredArchive(FArchiveFormatterType& InFormatter);
 	
 	/**
 	 * Default destructor. Closes the archive.
@@ -428,8 +98,9 @@ public:
 	FStructuredArchive& operator=(const FStructuredArchive&) = delete;
 
 private:
-
 	friend class FStructuredArchiveChildReader;
+	friend UE::StructuredArchive::Private::FElementId UE::StructuredArchive::Private::GetCurrentSlotElementIdImpl(FStructuredArchive& Ar);
+	friend FArchiveFormatterType& UE::StructuredArchive::Private::GetFormatterImpl(FStructuredArchive& Ar);
 
 	/**
 	* Reference to the formatter that actually writes out the data to the underlying archive
@@ -524,234 +195,7 @@ private:
 #endif
 };
 
-FORCEINLINE FArchive& UE::StructuredArchive::Private::FSlotBase::GetUnderlyingArchive() const
-{
-	return Ar.GetUnderlyingArchive();
-}
-
-FORCEINLINE const FArchiveState& UE::StructuredArchive::Private::FSlotBase::GetArchiveState() const
-{
-	return Ar.GetArchiveState();
-}
-
-FORCEINLINE bool FStructuredArchiveSlot::IsFilled() const
-{
-#if WITH_TEXT_ARCHIVE_SUPPORT
-	return Ar.CurrentSlotElementId != ElementId;
-#else
-	return true;
-#endif
-}
-
-template <typename T>
-FORCEINLINE_DEBUGGABLE void operator<<(FStructuredArchiveSlot Slot, TArray<T>& InArray)
-{
-	int32 NumElements = InArray.Num();
-	FStructuredArchiveArray Array = Slot.EnterArray(NumElements);
-
-	if (Slot.GetArchiveState().IsLoading())
-	{
-		InArray.SetNum(NumElements);
-	}
-
-	for (int32 ElementIndex = 0; ElementIndex < NumElements; ++ElementIndex)
-	{
-		FStructuredArchiveSlot ElementSlot = Array.EnterElement();
-		ElementSlot << InArray[ElementIndex];
-	}
-}
-
-template <>
-FORCEINLINE_DEBUGGABLE void operator<<(FStructuredArchiveSlot Slot, TArray<uint8>& InArray)
-{
-	Slot.Serialize(InArray);
-}
-
-/**
- * FStructuredArchiveChildReader
- *
- * Utility class for easily creating a structured archive that covers the data hierarchy underneath
- * the given slot
- *
- * Allows serialization code to get an archive instance for the current location, so that it can return to it
- * later on after the master archive has potentially moved on into a different location in the file.
- */
-class CORE_API FStructuredArchiveChildReader
-{
-public:
-
-	FStructuredArchiveChildReader(FStructuredArchiveSlot InSlot);
-	~FStructuredArchiveChildReader();
-
-	FORCEINLINE FStructuredArchiveSlot GetRoot() { return Root.GetValue(); }
-
-private:
-
-	FStructuredArchiveFormatter* OwnedFormatter;
-	FStructuredArchive* Archive;
-	TOptional<FStructuredArchiveSlot> Root;
-};
-
-class CORE_API FStructuredArchiveFromArchive
-{
-	UE_NONCOPYABLE(FStructuredArchiveFromArchive)
-
-	static constexpr uint32 ImplSize      = 400;
-	static constexpr uint32 ImplAlignment = 8;
-
-	struct FImpl;
-
-public:
-	explicit FStructuredArchiveFromArchive(FArchive& Ar);
-	~FStructuredArchiveFromArchive();
-
-	FStructuredArchiveSlot GetSlot();
-
-private:
-	// Implmented as a pimpl in order to reduce dependencies, but an inline one to avoid heap allocations
-	alignas(ImplAlignment) uint8 ImplStorage[ImplSize];
-};
-
-#if WITH_TEXT_ARCHIVE_SUPPORT
-
-class CORE_API FArchiveFromStructuredArchiveImpl : public FArchiveProxy
-{
-	UE_NONCOPYABLE(FArchiveFromStructuredArchiveImpl)
-
-	struct FImpl;
-
-public:
-	explicit FArchiveFromStructuredArchiveImpl(FStructuredArchiveSlot Slot);
-	virtual ~FArchiveFromStructuredArchiveImpl();
-
-	virtual void Flush() override;
-	virtual bool Close() override;
-
-	virtual int64 Tell() override;
-	virtual int64 TotalSize() override;
-	virtual void Seek(int64 InPos) override;
-	virtual bool AtEnd() override;
-
-	using FArchive::operator<<; // For visibility of the overloads we don't override
-
-	//~ Begin FArchive Interface
-	virtual FArchive& operator<<(class FName& Value) override;
-	virtual FArchive& operator<<(class UObject*& Value) override;
-	virtual FArchive& operator<<(class FText& Value) override;
-	//~ End FArchive Interface
-
-	virtual void Serialize(void* V, int64 Length) override;
-
-	virtual FArchive* GetCacheableArchive() override;
-
-	bool ContainsData() const;
-
-protected:
-	virtual bool Finalize(FStructuredArchiveRecord Record);
-	void OpenArchive();
-
-private:
-	void Commit();
-
-	// Implmented as a pimpl in order to reduce dependencies
-	TUniqueObj<FImpl> Pimpl;
-};
-
-class FArchiveFromStructuredArchive
-{
-public:
-	explicit FArchiveFromStructuredArchive(FStructuredArchiveSlot InSlot)
-		: Impl(InSlot)
-	{
-	}
-
-	      FArchive& GetArchive()       { return Impl; }
-	const FArchive& GetArchive() const { return Impl; }
-
-	void Close() { Impl.Close(); }
-
-private:
-	FArchiveFromStructuredArchiveImpl Impl;
-};
-
-#else
-
-class FArchiveFromStructuredArchive
-{
-public:
-	explicit FArchiveFromStructuredArchive(FStructuredArchiveSlot InSlot)
-		: Ar(InSlot.GetUnderlyingArchive())
-	{
-	}
-
-	      FArchive& GetArchive()       { return Ar; }
-	const FArchive& GetArchive() const { return Ar; }
-
-	void Close() {}
-
-private:
-	FArchive& Ar;
-};
-
-#endif
-
-/**
- * Adapter operator which allows a type to stream to an FArchive when it already supports streaming to an FStructuredArchiveSlot.
- *
- * @param  Ar   The archive to read from or write to.
- * @param  Obj  The object to read or write.
- *
- * @return  A reference to the same archive as Ar.
- */
-template <typename T>
-typename TEnableIf<
-	!TModels<CInsertable<FArchive&>, T>::Value && TModels<CInsertable<FStructuredArchiveSlot>, T>::Value,
-	FArchive&
->::Type operator<<(FArchive& Ar, T& Obj)
-{
-	FStructuredArchiveFromArchive ArAdapt(Ar);
-	ArAdapt.GetSlot() << Obj;
-	return Ar;
-}
-
-/**
- * Adapter operator which allows a type to stream to an FStructuredArchiveSlot when it already supports streaming to an FArchive.
- *
- * @param  Slot  The slot to read from or write to.
- * @param  Obj   The object to read or write.
- */
-template <typename T>
-typename TEnableIf<
-	TModels<CInsertable<FArchive&>, T>::Value &&
-	!TModels<CInsertable<FStructuredArchiveSlot>, T>::Value
->::Type operator<<(FStructuredArchiveSlot Slot, T& Obj)
-{
-#if WITH_TEXT_ARCHIVE_SUPPORT
-	FArchiveFromStructuredArchive Adapter(Slot);
-	FArchive& Ar = Adapter.GetArchive();
-#else
-	FArchive& Ar = Slot.GetUnderlyingArchive();
-#endif
-	Ar << Obj;
-#if WITH_TEXT_ARCHIVE_SUPPORT
-	Adapter.Close();
-#endif
-}
-
 #if !WITH_TEXT_ARCHIVE_SUPPORT
-	FORCEINLINE FStructuredArchiveChildReader::FStructuredArchiveChildReader(FStructuredArchiveSlot InSlot)
-		: OwnedFormatter(nullptr)
-		, Archive(nullptr)
-	{
-		Archive = new FStructuredArchive(InSlot.Ar.Formatter);
-		Root.Emplace(Archive->Open());
-	}
-
-	FORCEINLINE FStructuredArchiveChildReader::~FStructuredArchiveChildReader()
-	{
-	}
-
-	//////////// FStructuredArchive ////////////
 
 	FORCEINLINE FStructuredArchive::FStructuredArchive(FArchiveFormatterType& InFormatter)
 		: Formatter(InFormatter)
@@ -764,224 +208,31 @@ typename TEnableIf<
 
 	FORCEINLINE FStructuredArchiveSlot FStructuredArchive::Open()
 	{
-		return FStructuredArchiveSlot(*this);
+		return FStructuredArchiveSlot(FStructuredArchiveSlot::EPrivateToken{}, *this);
 	}
 
 	FORCEINLINE void FStructuredArchive::Close()
 	{
 	}
 
-	//////////// FStructuredArchiveSlot ////////////
-
-	FORCEINLINE FStructuredArchiveRecord FStructuredArchiveSlot::EnterRecord()
-	{
-		return FStructuredArchiveRecord(Ar);
-	}
-
-	FORCEINLINE FStructuredArchiveArray FStructuredArchiveSlot::EnterArray(int32& Num)
-	{
-		Ar.Formatter.EnterArray(Num);
-		return FStructuredArchiveArray(Ar);
-	}
-
-	FORCEINLINE FStructuredArchiveStream FStructuredArchiveSlot::EnterStream()
-	{
-		return FStructuredArchiveStream(Ar);
-	}
-
-	FORCEINLINE FStructuredArchiveMap FStructuredArchiveSlot::EnterMap(int32& Num)
-	{
-		Ar.Formatter.EnterMap(Num);
-		return FStructuredArchiveMap(Ar);
-	}
-
-	FORCEINLINE FStructuredArchiveSlot FStructuredArchiveSlot::EnterAttribute(FArchiveFieldName FieldName)
-	{
-		Ar.Formatter.EnterAttribute(FieldName);
-		return FStructuredArchiveSlot(Ar);
-	}
-
-	FORCEINLINE TOptional<FStructuredArchiveSlot> FStructuredArchiveSlot::TryEnterAttribute(FArchiveFieldName FieldName, bool bEnterWhenWriting)
-	{
-		if (Ar.Formatter.TryEnterAttribute(FieldName, bEnterWhenWriting))
-		{
-			return TOptional<FStructuredArchiveSlot>(Ar);
-		}
-		else
-		{
-			return TOptional<FStructuredArchiveSlot>();
-		}
-	}
-
-	FORCEINLINE void FStructuredArchiveSlot::operator<< (uint8& Value)
-	{
-		Ar.Formatter.Serialize(Value);
-	}
-
-	FORCEINLINE void FStructuredArchiveSlot::operator<< (uint16& Value)
-	{
-		Ar.Formatter.Serialize(Value);
-	}
-
-	FORCEINLINE void FStructuredArchiveSlot::operator<< (uint32& Value)
-	{
-		Ar.Formatter.Serialize(Value);
-	}
-
-	FORCEINLINE void FStructuredArchiveSlot::operator<< (uint64& Value)
-	{
-		Ar.Formatter.Serialize(Value);
-	}
-
-	FORCEINLINE void FStructuredArchiveSlot::operator<< (int8& Value)
-	{
-		Ar.Formatter.Serialize(Value);
-	}
-
-	FORCEINLINE void FStructuredArchiveSlot::operator<< (int16& Value)
-	{
-		Ar.Formatter.Serialize(Value);
-	}
-
-	FORCEINLINE void FStructuredArchiveSlot::operator<< (int32& Value)
-	{
-		Ar.Formatter.Serialize(Value);
-	}
-
-	FORCEINLINE void FStructuredArchiveSlot::operator<< (int64& Value)
-	{
-		Ar.Formatter.Serialize(Value);
-	}
-
-	FORCEINLINE void FStructuredArchiveSlot::operator<< (float& Value)
-	{
-		Ar.Formatter.Serialize(Value);
-	}
-
-	FORCEINLINE void FStructuredArchiveSlot::operator<< (double& Value)
-	{
-		Ar.Formatter.Serialize(Value);
-	}
-
-	FORCEINLINE void FStructuredArchiveSlot::operator<< (bool& Value)
-	{
-		Ar.Formatter.Serialize(Value);
-	}
-
-	FORCEINLINE void FStructuredArchiveSlot::operator<< (FString& Value)
-	{
-		Ar.Formatter.Serialize(Value);
-	}
-
-	FORCEINLINE void FStructuredArchiveSlot::operator<< (FName& Value)
-	{
-		Ar.Formatter.Serialize(Value);
-	}
-
-	FORCEINLINE void FStructuredArchiveSlot::operator<< (UObject*& Value)
-	{
-		Ar.Formatter.Serialize(Value);
-	}
-
-	FORCEINLINE void FStructuredArchiveSlot::operator<< (FText& Value)
-	{
-		Ar.Formatter.Serialize(Value);
-	}
-
-	FORCEINLINE void FStructuredArchiveSlot::operator<< (FWeakObjectPtr& Value)
-	{
-		Ar.Formatter.Serialize(Value);
-	}
-
-	FORCEINLINE void FStructuredArchiveSlot::operator<< (FSoftObjectPath& Value)
-	{
-		Ar.Formatter.Serialize(Value);
-	}
-
-	FORCEINLINE void FStructuredArchiveSlot::operator<< (FSoftObjectPtr& Value)
-	{
-		Ar.Formatter.Serialize(Value);
-	}
-
-	FORCEINLINE void FStructuredArchiveSlot::operator<< (FLazyObjectPtr& Value)
-	{
-		Ar.Formatter.Serialize(Value);
-	}
-
-	FORCEINLINE void FStructuredArchiveSlot::operator<< (FObjectPtr& Value)
-	{
-		Ar.Formatter.Serialize(Value);
-	}
-
-	FORCEINLINE void FStructuredArchiveSlot::Serialize(TArray<uint8>& Data)
-	{
-		Ar.Formatter.Serialize(Data);
-	}
-
-	FORCEINLINE void FStructuredArchiveSlot::Serialize(void* Data, uint64 DataSize)
-	{
-		Ar.Formatter.Serialize(Data, DataSize);
-	}
-
-	//////////// FStructuredArchiveRecord ////////////
-
-	FORCEINLINE FStructuredArchiveSlot FStructuredArchiveRecord::EnterField(FArchiveFieldName Name)
-	{
-		return FStructuredArchiveSlot(Ar);
-	}
-
-	FORCEINLINE FStructuredArchiveRecord FStructuredArchiveRecord::EnterRecord(FArchiveFieldName Name)
-	{
-		return EnterField(Name).EnterRecord();
-	}
+#endif
 
-	FORCEINLINE FStructuredArchiveArray FStructuredArchiveRecord::EnterArray(FArchiveFieldName Name, int32& Num)
+namespace UE::StructuredArchive::Private
 	{
-		return EnterField(Name).EnterArray(Num);
-	}
-
-	FORCEINLINE FStructuredArchiveStream FStructuredArchiveRecord::EnterStream(FArchiveFieldName Name)
+	FORCEINLINE FArchive& GetUnderlyingArchiveImpl(FStructuredArchive& Ar)
 	{
-		return EnterField(Name).EnterStream();
+		return Ar.GetUnderlyingArchive();
 	}
 
-	FORCEINLINE FStructuredArchiveMap FStructuredArchiveRecord::EnterMap(FArchiveFieldName Name, int32& Num)
+#if WITH_TEXT_ARCHIVE_SUPPORT
+	FORCEINLINE FElementId GetCurrentSlotElementIdImpl(FStructuredArchive& Ar)
 	{
-		return EnterField(Name).EnterMap(Num);
+		return Ar.CurrentSlotElementId;
 	}
+#endif
 
-	FORCEINLINE TOptional<FStructuredArchiveSlot> FStructuredArchiveRecord::TryEnterField(FArchiveFieldName Name, bool bEnterWhenWriting)
+	FORCEINLINE FArchiveFormatterType& GetFormatterImpl(FStructuredArchive& Ar)
 	{
-		if (Ar.Formatter.TryEnterField(Name, bEnterWhenWriting))
-		{
-			return TOptional<FStructuredArchiveSlot>(FStructuredArchiveSlot(Ar));
-		}
-		else
-		{
-			return TOptional<FStructuredArchiveSlot>();
+		return Ar.Formatter;
 		}
 	}
-
-	//////////// FStructuredArchiveArray ////////////
-
-	FORCEINLINE FStructuredArchiveSlot FStructuredArchiveArray::EnterElement()
-	{
-		return FStructuredArchiveSlot(Ar);
-	}
-
-	//////////// FStructuredArchiveStream ////////////
-
-	FORCEINLINE FStructuredArchiveSlot FStructuredArchiveStream::EnterElement()
-	{
-		return FStructuredArchiveSlot(Ar);
-	}
-
-	//////////// FStructuredArchiveMap ////////////
-
-	FORCEINLINE FStructuredArchiveSlot FStructuredArchiveMap::EnterElement(FString& Name)
-	{
-		Ar.Formatter.EnterMapElement(Name);
-		return FStructuredArchiveSlot(Ar);
-	}
-
-#endif
diff --git a/Engine/Source/Runtime/Core/Public/Serialization/StructuredArchiveAdapters.h b/Engine/Source/Runtime/Core/Public/Serialization/StructuredArchiveAdapters.h
new file mode 100644
--- /dev/null
+++ b/Engine/Source/Runtime/Core/Public/Serialization/StructuredArchiveAdapters.h
@@ -0,0 +1,157 @@
+// Copyright Epic Games, Inc. All Rights Reserved.
+
+#pragma once
+
+#include "CoreTypes.h"
+#include "Concepts/Insertable.h"
+#include "Templates/Models.h"
+#include "Serialization/ArchiveProxy.h"
+#include "Serialization/StructuredArchiveSlots.h"
+#include "Templates/UniqueObj.h"
+
+class CORE_API FStructuredArchiveFromArchive
+{
+	UE_NONCOPYABLE(FStructuredArchiveFromArchive)
+
+	static constexpr uint32 ImplSize = 400;
+	static constexpr uint32 ImplAlignment = 8;
+
+	struct FImpl;
+
+public:
+	explicit FStructuredArchiveFromArchive(FArchive& Ar);
+	~FStructuredArchiveFromArchive();
+
+	FStructuredArchiveSlot GetSlot();
+
+private:
+	// Implmented as a pimpl in order to reduce dependencies, but an inline one to avoid heap allocations
+	alignas(ImplAlignment) uint8 ImplStorage[ImplSize];
+};
+
+#if WITH_TEXT_ARCHIVE_SUPPORT
+
+class CORE_API FArchiveFromStructuredArchiveImpl : public FArchiveProxy
+{
+	UE_NONCOPYABLE(FArchiveFromStructuredArchiveImpl)
+
+		struct FImpl;
+
+public:
+	explicit FArchiveFromStructuredArchiveImpl(FStructuredArchiveSlot Slot);
+	virtual ~FArchiveFromStructuredArchiveImpl();
+
+	virtual void Flush() override;
+	virtual bool Close() override;
+
+	virtual int64 Tell() override;
+	virtual int64 TotalSize() override;
+	virtual void Seek(int64 InPos) override;
+	virtual bool AtEnd() override;
+
+	using FArchive::operator<<; // For visibility of the overloads we don't override
+
+	//~ Begin FArchive Interface
+	virtual FArchive& operator<<(class FName& Value) override;
+	virtual FArchive& operator<<(class UObject*& Value) override;
+	virtual FArchive& operator<<(class FText& Value) override;
+	//~ End FArchive Interface
+
+	virtual void Serialize(void* V, int64 Length) override;
+
+	virtual FArchive* GetCacheableArchive() override;
+
+	bool ContainsData() const;
+
+protected:
+	virtual bool Finalize(FStructuredArchiveRecord Record);
+	void OpenArchive();
+
+private:
+	void Commit();
+
+	// Implmented as a pimpl in order to reduce dependencies
+	TUniqueObj<FImpl> Pimpl;
+};
+
+class FArchiveFromStructuredArchive
+{
+public:
+	explicit FArchiveFromStructuredArchive(FStructuredArchiveSlot InSlot)
+		: Impl(InSlot)
+	{
+	}
+
+	FArchive& GetArchive() { return Impl; }
+	const FArchive& GetArchive() const { return Impl; }
+
+	void Close() { Impl.Close(); }
+
+private:
+	FArchiveFromStructuredArchiveImpl Impl;
+};
+
+#else
+
+class FArchiveFromStructuredArchive
+{
+public:
+	explicit FArchiveFromStructuredArchive(FStructuredArchiveSlot InSlot)
+		: Ar(InSlot.GetUnderlyingArchive())
+	{
+	}
+
+	FArchive& GetArchive() { return Ar; }
+	const FArchive& GetArchive() const { return Ar; }
+
+	void Close() {}
+
+private:
+	FArchive& Ar;
+};
+
+#endif
+
+/**
+ * Adapter operator which allows a type to stream to an FArchive when it already supports streaming to an FStructuredArchiveSlot.
+ *
+ * @param  Ar   The archive to read from or write to.
+ * @param  Obj  The object to read or write.
+ *
+ * @return  A reference to the same archive as Ar.
+ */
+template <typename T>
+typename TEnableIf<
+	!TModels<CInsertable<FArchive&>, T>::Value&& TModels<CInsertable<FStructuredArchiveSlot>, T>::Value,
+	FArchive&
+>::Type operator<<(FArchive& Ar, T& Obj)
+{
+	FStructuredArchiveFromArchive ArAdapt(Ar);
+	ArAdapt.GetSlot() << Obj;
+	return Ar;
+}
+
+/**
+ * Adapter operator which allows a type to stream to an FStructuredArchiveSlot when it already supports streaming to an FArchive.
+ *
+ * @param  Slot  The slot to read from or write to.
+ * @param  Obj   The object to read or write.
+ */
+template <typename T>
+typename TEnableIf<
+	TModels<CInsertable<FArchive&>, T>::Value &&
+	!TModels<CInsertable<FStructuredArchiveSlot>, T>::Value
+>::Type operator<<(FStructuredArchiveSlot Slot, T& Obj)
+{
+#if WITH_TEXT_ARCHIVE_SUPPORT
+	FArchiveFromStructuredArchive Adapter(Slot);
+	FArchive& Ar = Adapter.GetArchive();
+#else
+	FArchive& Ar = Slot.GetUnderlyingArchive();
+#endif
+	Ar << Obj;
+#if WITH_TEXT_ARCHIVE_SUPPORT
+	Adapter.Close();
+#endif
+}
+
diff --git a/Engine/Source/Runtime/Core/Public/Serialization/StructuredArchiveChildReader.h b/Engine/Source/Runtime/Core/Public/Serialization/StructuredArchiveChildReader.h
new file mode 100644
--- /dev/null
+++ b/Engine/Source/Runtime/Core/Public/Serialization/StructuredArchiveChildReader.h
@@ -0,0 +1,54 @@
+// Copyright Epic Games, Inc. All Rights Reserved.
+
+#pragma once
+
+#include "Serialization/StructuredArchiveDefines.h"
+
+/**
+ * FStructuredArchiveChildReader
+ *
+ * Utility class for easily creating a structured archive that covers the data hierarchy underneath
+ * the given slot
+ *
+ * Allows serialization code to get an archive instance for the current location, so that it can return to it
+ * later on after the master archive has potentially moved on into a different location in the file.
+ */
+class CORE_API FStructuredArchiveChildReader
+{
+public:
+	explicit FStructuredArchiveChildReader(FStructuredArchiveSlot InSlot);
+	~FStructuredArchiveChildReader();
+
+	// Non-copyable
+	FStructuredArchiveChildReader(FStructuredArchiveChildReader&&) = delete;
+	FStructuredArchiveChildReader(const FStructuredArchiveChildReader&) = delete;
+	FStructuredArchiveChildReader& operator=(FStructuredArchiveChildReader&&) = delete;
+	FStructuredArchiveChildReader& operator=(const FStructuredArchiveChildReader&) = delete;
+
+	FORCEINLINE FStructuredArchiveSlot GetRoot() const
+	{
+		return Root.GetValue();
+	}
+
+private:
+	FStructuredArchiveFormatter* OwnedFormatter;
+	FStructuredArchive* Archive;
+	TOptional<FStructuredArchiveSlot> Root;
+};
+
+#if !WITH_TEXT_ARCHIVE_SUPPORT
+
+FORCEINLINE FStructuredArchiveChildReader::FStructuredArchiveChildReader(FStructuredArchiveSlot InSlot)
+	: OwnedFormatter(nullptr)
+	, Archive(nullptr)
+{
+	Archive = new FStructuredArchive(InSlot.Ar.Formatter);
+	Root.Emplace(Archive->Open());
+}
+
+FORCEINLINE FStructuredArchiveChildReader::~FStructuredArchiveChildReader()
+{
+	delete Archive;
+}
+
+#endif
diff --git a/Engine/Source/Runtime/Core/Public/Serialization/StructuredArchiveFormatter.h b/Engine/Source/Runtime/Core/Public/Serialization/StructuredArchiveFormatter.h
--- a/Engine/Source/Runtime/Core/Public/Serialization/StructuredArchiveFormatter.h
+++ b/Engine/Source/Runtime/Core/Public/Serialization/StructuredArchiveFormatter.h
@@ -4,23 +4,7 @@
 
 #include "CoreTypes.h"
 #include "CoreFwd.h"
-
-// Define a structure to encapsulate a field name, which compiles to an empty object if WITH_TEXT_ARCHIVE_SUPPORT = 0
-struct FArchiveFieldName
-{
-#if WITH_TEXT_ARCHIVE_SUPPORT
-	const TCHAR* Name;
-#endif
-
-	explicit FArchiveFieldName(const TCHAR* InName) 
-#if WITH_TEXT_ARCHIVE_SUPPORT
-		: Name(InName)
-#endif
-	{
-	}
-};
-
-#define SA_FIELD_NAME(x) FArchiveFieldName(x)
+#include "Serialization/StructuredArchiveNameHelpers.h"
 
 /**
  * Specifies the type of a value in a slot. Used by FContextFreeArchiveFormatter for introspection.
diff --git a/Engine/Source/Runtime/Core/Public/Serialization/StructuredArchiveFwd.h b/Engine/Source/Runtime/Core/Public/Serialization/StructuredArchiveFwd.h
new file mode 100644
--- /dev/null
+++ b/Engine/Source/Runtime/Core/Public/Serialization/StructuredArchiveFwd.h
@@ -0,0 +1,25 @@
+// Copyright Epic Games, Inc. All Rights Reserved.
+
+#pragma once
+
+// Formatters
+class FBinaryArchiveFormatter;
+class FStructuredArchiveFormatter;
+
+// Archives
+class FStructuredArchive;
+class FStructuredArchiveChildReader;
+
+// Slots
+class FStructuredArchiveSlot;
+class FStructuredArchiveRecord;
+class FStructuredArchiveArray;
+class FStructuredArchiveStream;
+class FStructuredArchiveMap;
+
+/** Typedef for which formatter type to support */
+#if WITH_TEXT_ARCHIVE_SUPPORT
+	using FArchiveFormatterType = FStructuredArchiveFormatter;
+#else
+	using FArchiveFormatterType = FBinaryArchiveFormatter;
+#endif
diff --git a/Engine/Source/Runtime/Core/Public/Serialization/StructuredArchiveNameHelpers.h b/Engine/Source/Runtime/Core/Public/Serialization/StructuredArchiveNameHelpers.h
new file mode 100644
--- /dev/null
+++ b/Engine/Source/Runtime/Core/Public/Serialization/StructuredArchiveNameHelpers.h
@@ -0,0 +1,84 @@
+// Copyright Epic Games, Inc. All Rights Reserved.
+
+#pragma once
+
+#include "CoreTypes.h"
+
+// Define a structure to encapsulate a field name, which compiles to an empty object if WITH_TEXT_ARCHIVE_SUPPORT = 0
+struct FArchiveFieldName
+{
+#if WITH_TEXT_ARCHIVE_SUPPORT
+	const TCHAR* Name;
+#endif
+
+	explicit FArchiveFieldName(const TCHAR* InName)
+#if WITH_TEXT_ARCHIVE_SUPPORT
+		: Name(InName)
+#endif
+	{
+	}
+};
+
+namespace UE::StructuredArchive::Private
+{
+	/**
+	 * Class to contain a named value for serialization. Intended to be created as a temporary and passed to object serialization methods.
+	 */
+	template <typename T>
+	struct TNamedValue
+	{
+		FArchiveFieldName Name;
+		T& Value;
+	};
+
+	/**
+	 * Class to contain a named attribute for serialization. Intended to be created as a temporary and passed to object serialization methods.
+	 */
+	template <typename T>
+	struct TNamedAttribute
+	{
+		FArchiveFieldName Name;
+		T& Value;
+	};
+
+	/**
+	 * Class to contain a named attribute for serialization, with a default. Intended to be created as a temporary and passed to object
+	 * serialization methods, which can choose not to serialize the attribute if it matches the default.
+	 */
+	template <typename T>
+	struct TOptionalNamedAttribute
+	{
+		FArchiveFieldName Name;
+		T& Value;
+		const T& Default;
+	};
+
+	template <typename T>
+	FORCEINLINE TNamedValue<T> MakeNamedValue(FArchiveFieldName Name, T& Value)
+	{
+		return TNamedValue<T>{ Name, Value };
+	}
+
+	template <typename T>
+	FORCEINLINE TNamedAttribute<T> MakeNamedAttribute(FArchiveFieldName Name, T& Value)
+	{
+		return TNamedAttribute<T>{ Name, Value };
+	}
+
+	template <typename T>
+	FORCEINLINE TOptionalNamedAttribute<T> MakeOptionalNamedAttribute(FArchiveFieldName Name, T& Value, const typename TIdentity<T>::Type& Default)
+	{
+		return TOptionalNamedAttribute<T>{ Name, Value, Default };
+	}
+}
+
+#define SA_FIELD_NAME(x) FArchiveFieldName(x)
+
+/** Construct a TNamedValue given an ANSI string and value reference. */
+#define SA_VALUE(Name, Value) UE::StructuredArchive::Private::MakeNamedValue(FArchiveFieldName(Name), Value)
+
+/** Construct a TNamedAttribute given an ANSI string and value reference. */
+#define SA_ATTRIBUTE(Name, Value) UE::StructuredArchive::Private::MakeNamedAttribute(FArchiveFieldName(Name), Value)
+
+/** Construct a TOptionalNamedAttribute given an ANSI string and value reference. */
+#define SA_OPTIONAL_ATTRIBUTE(Name, Value, Default) UE::StructuredArchive::Private::MakeOptionalNamedAttribute(FArchiveFieldName(Name), Value, Default)
diff --git a/Engine/Source/Runtime/Core/Public/Serialization/StructuredArchiveSlotBase.h b/Engine/Source/Runtime/Core/Public/Serialization/StructuredArchiveSlotBase.h
new file mode 100644
--- /dev/null
+++ b/Engine/Source/Runtime/Core/Public/Serialization/StructuredArchiveSlotBase.h
@@ -0,0 +1,102 @@
+// Copyright Epic Games, Inc. All Rights Reserved.
+
+#pragma once
+
+#include "CoreTypes.h"
+
+class FStructuredArchive;
+
+namespace UE::StructuredArchive::Private
+{
+	FArchive& GetUnderlyingArchiveImpl(FStructuredArchive& Ar);
+
+	struct FElementId
+	{
+		FElementId() = default;
+
+		explicit FElementId(uint32 InId)
+			: Id(InId)
+		{
+		}
+
+		bool IsValid() const
+		{
+			return Id != 0;
+		}
+
+		void Reset()
+		{
+			Id = 0;
+		}
+
+		bool operator==(const FElementId& Rhs) const
+		{
+			return Id == Rhs.Id;
+		}
+
+		bool operator!=(const FElementId& Rhs) const
+		{
+			return Id != Rhs.Id;
+		}
+
+	private:
+		uint32 Id = 0;
+	};
+
+	// Represents a position of a slot within the hierarchy.
+	class CORE_API FSlotPosition
+	{
+		friend class FStructuredArchive;
+
+	public:
+		int32 Depth;
+		FElementId ElementId;
+
+		FORCEINLINE explicit FSlotPosition(int32 InDepth, FElementId InElementId)
+			: Depth(InDepth)
+			, ElementId(InElementId)
+		{
+		}
+	};
+
+	// The base class of all slot types
+	class CORE_API FSlotBase
+#if WITH_TEXT_ARCHIVE_SUPPORT
+		: protected FSlotPosition
+#endif
+	{
+		friend class FStructuredArchive;
+
+	protected:
+		// Token which needs to be passed to the constructor, so only friends are able to do so, but
+		// still allows emplacement without the container being a friend.
+		enum class EPrivateToken {};
+
+	public:
+#if WITH_TEXT_ARCHIVE_SUPPORT
+		FORCEINLINE explicit FSlotBase(EPrivateToken, FStructuredArchive& InAr, int32 InDepth, FElementId InElementId)
+			: FSlotPosition(InDepth, InElementId)
+			, Ar(InAr)
+		{
+		}
+#else
+		FORCEINLINE explicit FSlotBase(EPrivateToken, FStructuredArchive& InAr)
+			: Ar(InAr)
+		{
+		}
+#endif
+
+		FArchive& GetUnderlyingArchive() const
+		{
+			return GetUnderlyingArchiveImpl(Ar);
+		}
+
+		const FArchiveState& GetArchiveState() const
+		{
+			return GetUnderlyingArchiveImpl(Ar).GetArchiveState();
+		}
+
+	protected:
+		FStructuredArchive& Ar;
+	};
+}
diff --git a/Engine/Source/Runtime/Core/Public/Serialization/StructuredArchiveSlots.h b/Engine/Source/Runtime/Core/Public/Serialization/StructuredArchiveSlots.h
new file mode 100644
--- /dev/null
+++ b/Engine/Source/Runtime/Core/Public/Serialization/StructuredArchiveSlots.h
@@ -0,0 +1,447 @@
+// Copyright Epic Games, Inc. All Rights Reserved.
+
+#pragma once
+
+#include "CoreTypes.h"
+#include "Misc/Optional.h"
+#include "Serialization/StructuredArchiveFwd.h"
+#include "Serialization/StructuredArchiveNameHelpers.h"
+#include "Serialization/StructuredArchiveSlotBase.h"
+#include "Formatters/BinaryArchiveFormatter.h"
+
+class FStructuredArchive;
+class FStructuredArchiveChildReader;
+class FStructuredArchiveSlot;
+class FStructuredArchiveRecord;
+class FStructuredArchiveArray;
+class FStructuredArchiveStream;
+class FStructuredArchiveMap;
+
+namespace UE::StructuredArchive::Private
+{
+	FElementId GetCurrentSlotElementIdImpl(FStructuredArchive& Ar);
+	FArchiveFormatterType& GetFormatterImpl(FStructuredArchive& Ar);
+}
+
+/**
+ * Contains a value in the archive; either a field or array/map element. A slot does not know it's name or location,
+ * and can merely have a value serialized into it. That value may be a literal (eg. int, float) or compound object
+ * (eg. object, array, map).
+ */
+class CORE_API FStructuredArchiveSlot final : public UE::StructuredArchive::Private::FSlotBase
+{
+public:
+	FStructuredArchiveRecord EnterRecord();
+	FStructuredArchiveArray EnterArray(int32& Num);
+	FStructuredArchiveStream EnterStream();
+	FStructuredArchiveMap EnterMap(int32& Num);
+	FStructuredArchiveSlot EnterAttribute(FArchiveFieldName AttributeName);
+	TOptional<FStructuredArchiveSlot> TryEnterAttribute(FArchiveFieldName AttributeName, bool bEnterWhenWriting);
+
+	// We don't support chaining writes to a single slot, so this returns void.
+	void operator << (uint8& Value);
+	void operator << (uint16& Value);
+	void operator << (uint32& Value);
+	void operator << (uint64& Value);
+	void operator << (int8& Value);
+	void operator << (int16& Value);
+	void operator << (int32& Value);
+	void operator << (int64& Value);
+	void operator << (float& Value);
+	void operator << (double& Value);
+	void operator << (bool& Value);
+	void operator << (FString& Value);
+	void operator << (FName& Value);
+	void operator << (UObject*& Value);
+	void operator << (FText& Value);
+	void operator << (FWeakObjectPtr& Value);
+	void operator << (FSoftObjectPtr& Value);
+	void operator << (FSoftObjectPath& Value);
+	void operator << (FLazyObjectPtr& Value);
+	void operator << (FObjectPtr& Value);
+
+	template <typename T>
+	FORCEINLINE void operator<<(TEnumAsByte<T>& Value)
+	{
+		uint8 Tmp = (uint8)Value.GetValue();
+		*this << Tmp;
+		Value = (T)Tmp;
+	}
+
+	template <
+		typename EnumType,
+		typename = typename TEnableIf<TIsEnumClass<EnumType>::Value>::Type
+	>
+	FORCEINLINE void operator<<(EnumType& Value)
+	{
+		*this << (__underlying_type(EnumType)&)Value;
+	}
+
+	template <typename T>
+	FORCEINLINE void operator<<(UE::StructuredArchive::Private::TNamedAttribute<T> Item)
+	{
+		EnterAttribute(Item.Name) << Item.Value;
+	}
+
+	template <typename T>
+	FORCEINLINE void operator<<(UE::StructuredArchive::Private::TOptionalNamedAttribute<T> Item)
+	{
+		if (TOptional<FStructuredArchiveSlot> Attribute = TryEnterAttribute(Item.Name, Item.Value != Item.Default))
+		{
+			Attribute.GetValue() << Item.Value;
+		}
+		else
+		{
+			Item.Value = Item.Default;
+		}
+	}
+
+	void Serialize(TArray<uint8>& Data);
+	void Serialize(void* Data, uint64 DataSize);
+
+	FORCEINLINE bool IsFilled() const
+	{
+#if WITH_TEXT_ARCHIVE_SUPPORT
+		return UE::StructuredArchive::Private::GetCurrentSlotElementIdImpl(Ar) != ElementId;
+#else
+		return true;
+#endif
+	}
+
+private:
+	friend FStructuredArchive;
+	friend FStructuredArchiveChildReader;
+	friend FStructuredArchiveSlot;
+	friend FStructuredArchiveRecord;
+	friend FStructuredArchiveArray;
+	friend FStructuredArchiveStream;
+	friend FStructuredArchiveMap;
+
+	using UE::StructuredArchive::Private::FSlotBase::FSlotBase;
+};
+
+/**
+ * Represents a record in the structured archive. An object contains slots that are identified by FArchiveName,
+ * which may be compiled out with binary-only archives.
+ */
+class CORE_API FStructuredArchiveRecord final : public UE::StructuredArchive::Private::FSlotBase
+{
+public:
+	FStructuredArchiveSlot EnterField(FArchiveFieldName Name);
+	FStructuredArchiveRecord EnterRecord(FArchiveFieldName Name);
+	FStructuredArchiveArray EnterArray(FArchiveFieldName Name, int32& Num);
+	FStructuredArchiveStream EnterStream(FArchiveFieldName Name);
+	FStructuredArchiveMap EnterMap(FArchiveFieldName Name, int32& Num);
+
+	TOptional<FStructuredArchiveSlot> TryEnterField(FArchiveFieldName Name, bool bEnterForSaving);
+
+	template<typename T> FORCEINLINE FStructuredArchiveRecord& operator<<(UE::StructuredArchive::Private::TNamedValue<T> Item)
+	{
+		EnterField(Item.Name) << Item.Value;
+		return *this;
+	}
+
+private:
+	friend FStructuredArchive;
+	friend FStructuredArchiveSlot;
+
+	using UE::StructuredArchive::Private::FSlotBase::FSlotBase;
+};
+
+/**
+ * Represents an array in the structured archive. An object contains slots that are identified by a FArchiveFieldName,
+ * which may be compiled out with binary-only archives.
+ */
+class CORE_API FStructuredArchiveArray final : public UE::StructuredArchive::Private::FSlotBase
+{
+public:
+	FStructuredArchiveSlot EnterElement();
+
+	template<typename T> FORCEINLINE FStructuredArchiveArray& operator<<(T& Item)
+	{
+		EnterElement() << Item;
+		return *this;
+	}
+
+private:
+	friend FStructuredArchive;
+	friend FStructuredArchiveSlot;
+
+	using UE::StructuredArchive::Private::FSlotBase::FSlotBase;
+};
+
+/**
+ * Represents an unsized sequence of slots in the structured archive (similar to an array, but without a known size).
+ */
+class CORE_API FStructuredArchiveStream final : public UE::StructuredArchive::Private::FSlotBase
+{
+public:
+	FStructuredArchiveSlot EnterElement();
+
+	template<typename T> FORCEINLINE FStructuredArchiveStream& operator<<(T& Item)
+	{
+		EnterElement() << Item;
+		return *this;
+	}
+
+private:
+	friend FStructuredArchive;
+	friend FStructuredArchiveSlot;
+
+	using UE::StructuredArchive::Private::FSlotBase::FSlotBase;
+};
+
+/**
+ * Represents a map in the structured archive. A map is similar to a record, but keys can be read back out from an archive.
+ * (This is an important distinction for binary archives).
+ */
+class CORE_API FStructuredArchiveMap final : public UE::StructuredArchive::Private::FSlotBase
+{
+public:
+	FStructuredArchiveSlot EnterElement(FString& Name);
+
+private:
+	friend FStructuredArchive;
+	friend FStructuredArchiveSlot;
+
+	using UE::StructuredArchive::Private::FSlotBase::FSlotBase;
+};
+
+template <typename T>
+FORCEINLINE_DEBUGGABLE void operator<<(FStructuredArchiveSlot Slot, TArray<T>& InArray)
+{
+	int32 NumElements = InArray.Num();
+	FStructuredArchiveArray Array = Slot.EnterArray(NumElements);
+
+	if (Slot.GetArchiveState().IsLoading())
+	{
+		InArray.SetNum(NumElements);
+	}
+
+	for (int32 ElementIndex = 0; ElementIndex < NumElements; ++ElementIndex)
+	{
+		FStructuredArchiveSlot ElementSlot = Array.EnterElement();
+		ElementSlot << InArray[ElementIndex];
+	}
+}
+
+template <>
+FORCEINLINE_DEBUGGABLE void operator<<(FStructuredArchiveSlot Slot, TArray<uint8>& InArray)
+{
+	Slot.Serialize(InArray);
+}
+
+#if !WITH_TEXT_ARCHIVE_SUPPORT
+	//////////// FStructuredArchiveSlot ////////////
+	FORCEINLINE FStructuredArchiveRecord FStructuredArchiveSlot::EnterRecord()
+	{
+		return FStructuredArchiveRecord(FStructuredArchiveRecord::EPrivateToken{}, Ar);
+	}
+
+	FORCEINLINE FStructuredArchiveArray FStructuredArchiveSlot::EnterArray(int32& Num)
+	{
+		UE::StructuredArchive::Private::GetFormatterImpl(Ar).EnterArray(Num);
+		return FStructuredArchiveArray(FStructuredArchiveArray::EPrivateToken{}, Ar);
+	}
+
+	FORCEINLINE FStructuredArchiveStream FStructuredArchiveSlot::EnterStream()
+	{
+		return FStructuredArchiveStream(FStructuredArchiveStream::EPrivateToken{}, Ar);
+	}
+
+	FORCEINLINE FStructuredArchiveMap FStructuredArchiveSlot::EnterMap(int32& Num)
+	{
+		UE::StructuredArchive::Private::GetFormatterImpl(Ar).EnterMap(Num);
+		return FStructuredArchiveMap(FStructuredArchiveMap::EPrivateToken{}, Ar);
+	}
+
+	FORCEINLINE FStructuredArchiveSlot FStructuredArchiveSlot::EnterAttribute(FArchiveFieldName FieldName)
+	{
+		UE::StructuredArchive::Private::GetFormatterImpl(Ar).EnterAttribute(FieldName);
+		return FStructuredArchiveSlot(FStructuredArchiveSlot::EPrivateToken{}, Ar);
+	}
+
+	FORCEINLINE TOptional<FStructuredArchiveSlot> FStructuredArchiveSlot::TryEnterAttribute(FArchiveFieldName FieldName, bool bEnterWhenWriting)
+	{
+		if (UE::StructuredArchive::Private::GetFormatterImpl(Ar).TryEnterAttribute(FieldName, bEnterWhenWriting))
+		{
+			return TOptional<FStructuredArchiveSlot>(InPlace, FStructuredArchiveSlot::EPrivateToken{}, Ar);
+		}
+		else
+		{
+			return TOptional<FStructuredArchiveSlot>();
+		}
+	}
+
+	FORCEINLINE void FStructuredArchiveSlot::operator<< (uint8& Value)
+	{
+		UE::StructuredArchive::Private::GetFormatterImpl(Ar).Serialize(Value);
+	}
+
+	FORCEINLINE void FStructuredArchiveSlot::operator<< (uint16& Value)
+	{
+		UE::StructuredArchive::Private::GetFormatterImpl(Ar).Serialize(Value);
+	}
+
+	FORCEINLINE void FStructuredArchiveSlot::operator<< (uint32& Value)
+	{
+		UE::StructuredArchive::Private::GetFormatterImpl(Ar).Serialize(Value);
+	}
+
+	FORCEINLINE void FStructuredArchiveSlot::operator<< (uint64& Value)
+	{
+		UE::StructuredArchive::Private::GetFormatterImpl(Ar).Serialize(Value);
+	}
+
+	FORCEINLINE void FStructuredArchiveSlot::operator<< (int8& Value)
+	{
+		UE::StructuredArchive::Private::GetFormatterImpl(Ar).Serialize(Value);
+	}
+
+	FORCEINLINE void FStructuredArchiveSlot::operator<< (int16& Value)
+	{
+		UE::StructuredArchive::Private::GetFormatterImpl(Ar).Serialize(Value);
+	}
+
+	FORCEINLINE void FStructuredArchiveSlot::operator<< (int32& Value)
+	{
+		UE::StructuredArchive::Private::GetFormatterImpl(Ar).Serialize(Value);
+	}
+
+	FORCEINLINE void FStructuredArchiveSlot::operator<< (int64& Value)
+	{
+		UE::StructuredArchive::Private::GetFormatterImpl(Ar).Serialize(Value);
+	}
+
+	FORCEINLINE void FStructuredArchiveSlot::operator<< (float& Value)
+	{
+		UE::StructuredArchive::Private::GetFormatterImpl(Ar).Serialize(Value);
+	}
+
+	FORCEINLINE void FStructuredArchiveSlot::operator<< (double& Value)
+	{
+		UE::StructuredArchive::Private::GetFormatterImpl(Ar).Serialize(Value);
+	}
+
+	FORCEINLINE void FStructuredArchiveSlot::operator<< (bool& Value)
+	{
+		UE::StructuredArchive::Private::GetFormatterImpl(Ar).Serialize(Value);
+	}
+
+	FORCEINLINE void FStructuredArchiveSlot::operator<< (FString& Value)
+	{
+		UE::StructuredArchive::Private::GetFormatterImpl(Ar).Serialize(Value);
+	}
+
+	FORCEINLINE void FStructuredArchiveSlot::operator<< (FName& Value)
+	{
+		UE::StructuredArchive::Private::GetFormatterImpl(Ar).Serialize(Value);
+	}
+
+	FORCEINLINE void FStructuredArchiveSlot::operator<< (UObject*& Value)
+	{
+		UE::StructuredArchive::Private::GetFormatterImpl(Ar).Serialize(Value);
+	}
+
+	FORCEINLINE void FStructuredArchiveSlot::operator<< (FText& Value)
+	{
+		UE::StructuredArchive::Private::GetFormatterImpl(Ar).Serialize(Value);
+	}
+
+	FORCEINLINE void FStructuredArchiveSlot::operator<< (FWeakObjectPtr& Value)
+	{
+		UE::StructuredArchive::Private::GetFormatterImpl(Ar).Serialize(Value);
+	}
+
+	FORCEINLINE void FStructuredArchiveSlot::operator<< (FSoftObjectPath& Value)
+	{
+		UE::StructuredArchive::Private::GetFormatterImpl(Ar).Serialize(Value);
+	}
+
+	FORCEINLINE void FStructuredArchiveSlot::operator<< (FSoftObjectPtr& Value)
+	{
+		UE::StructuredArchive::Private::GetFormatterImpl(Ar).Serialize(Value);
+	}
+
+	FORCEINLINE void FStructuredArchiveSlot::operator<< (FLazyObjectPtr& Value)
+	{
+		UE::StructuredArchive::Private::GetFormatterImpl(Ar).Serialize(Value);
+	}
+
+	FORCEINLINE void FStructuredArchiveSlot::operator<< (FObjectPtr& Value)
+	{
+		UE::StructuredArchive::Private::GetFormatterImpl(Ar).Serialize(Value);
+	}
+
+	FORCEINLINE void FStructuredArchiveSlot::Serialize(TArray<uint8>& Data)
+	{
+		UE::StructuredArchive::Private::GetFormatterImpl(Ar).Serialize(Data);
+	}
+
+	FORCEINLINE void FStructuredArchiveSlot::Serialize(void* Data, uint64 DataSize)
+	{
+		UE::StructuredArchive::Private::GetFormatterImpl(Ar).Serialize(Data, DataSize);
+	}
+
+	//////////// FStructuredArchiveRecord ////////////
+
+	FORCEINLINE FStructuredArchiveSlot FStructuredArchiveRecord::EnterField(FArchiveFieldName Name)
+	{
+		return FStructuredArchiveSlot(FStructuredArchiveSlot::EPrivateToken{}, Ar);
+	}
+
+	FORCEINLINE FStructuredArchiveRecord FStructuredArchiveRecord::EnterRecord(FArchiveFieldName Name)
+	{
+		return EnterField(Name).EnterRecord();
+	}
+
+	FORCEINLINE FStructuredArchiveArray FStructuredArchiveRecord::EnterArray(FArchiveFieldName Name, int32& Num)
+	{
+		return EnterField(Name).EnterArray(Num);
+	}
+
+	FORCEINLINE FStructuredArchiveStream FStructuredArchiveRecord::EnterStream(FArchiveFieldName Name)
+	{
+		return EnterField(Name).EnterStream();
+	}
+
+	FORCEINLINE FStructuredArchiveMap FStructuredArchiveRecord::EnterMap(FArchiveFieldName Name, int32& Num)
+	{
+		return EnterField(Name).EnterMap(Num);
+	}
+
+	FORCEINLINE TOptional<FStructuredArchiveSlot> FStructuredArchiveRecord::TryEnterField(FArchiveFieldName Name, bool bEnterWhenWriting)
+	{
+		if (UE::StructuredArchive::Private::GetFormatterImpl(Ar).TryEnterField(Name, bEnterWhenWriting))
+		{
+			return TOptional<FStructuredArchiveSlot>(InPlace, FStructuredArchiveSlot(FStructuredArchiveSlot::EPrivateToken{}, Ar));
+		}
+		else
+		{
+			return TOptional<FStructuredArchiveSlot>();
+		}
+	}
+
+	//////////// FStructuredArchiveArray ////////////
+
+	FORCEINLINE FStructuredArchiveSlot FStructuredArchiveArray::EnterElement()
+	{
+		return FStructuredArchiveSlot(FStructuredArchiveSlot::EPrivateToken{}, Ar);
+	}
+
+	//////////// FStructuredArchiveStream ////////////
+
+	FORCEINLINE FStructuredArchiveSlot FStructuredArchiveStream::EnterElement()
+	{
+		return FStructuredArchiveSlot(FStructuredArchiveSlot::EPrivateToken{}, Ar);
+	}
+
+	//////////// FStructuredArchiveMap ////////////
+
+	FORCEINLINE FStructuredArchiveSlot FStructuredArchiveMap::EnterElement(FString& Name)
+	{
+		UE::StructuredArchive::Private::GetFormatterImpl(Ar).EnterMapElement(Name);
+		return FStructuredArchiveSlot(FStructuredArchiveSlot::EPrivateToken{}, Ar);
+	}
+
+#endif
diff --git a/Engine/Source/Runtime/CoreUObject/Private/UObject/LinkerLoad.cpp b/Engine/Source/Runtime/CoreUObject/Private/UObject/LinkerLoad.cpp
--- a/Engine/Source/Runtime/CoreUObject/Private/UObject/LinkerLoad.cpp
+++ b/Engine/Source/Runtime/CoreUObject/Private/UObject/LinkerLoad.cpp
@@ -41,6 +41,7 @@
 #include "Serialization/Formatters/BinaryArchiveFormatter.h"
 #include "Serialization/Formatters/JsonArchiveInputFormatter.h"
 #include "Serialization/ArchiveUObjectFromStructuredArchive.h"
+#include "Serialization/StructuredArchiveChildReader.h"
 #include "Serialization/UnversionedPropertySerialization.h"
 #include "Serialization/LoadTimeTracePrivate.h"
 #include "Serialization/EditorBulkData.h"
