Missing includes added.

diff --git a/Engine/Source/Runtime/Core/Public/Serialization/StructuredArchiveChildReader.h b/Engine/Source/Runtime/Core/Public/Serialization/StructuredArchiveChildReader.h
--- a/Engine/Source/Runtime/Core/Public/Serialization/StructuredArchiveChildReader.h
+++ b/Engine/Source/Runtime/Core/Public/Serialization/StructuredArchiveChildReader.h
@@ -2,7 +2,8 @@
 
 #pragma once
 
-#include "Serialization/StructuredArchiveDefines.h"
+#include "Misc/Optional.h"
+#include "Serialization/StructuredArchive.h"
 
 /**
  * FStructuredArchiveChildReader
diff --git a/Engine/Source/Runtime/Core/Public/Serialization/StructuredArchiveDefines.h b/Engine/Source/Runtime/Core/Public/Serialization/StructuredArchiveDefines.h
--- a/Engine/Source/Runtime/Core/Public/Serialization/StructuredArchiveDefines.h
+++ b/Engine/Source/Runtime/Core/Public/Serialization/StructuredArchiveDefines.h
@@ -2,6 +2,8 @@
 
 #pragma once
 
+#include "Misc/Build.h"
+
 /**
  * DO_STRUCTURED_ARCHIVE_CONTAINER_CHECKS - if set, checks that nested container types are serialized correctly.
  */
