Missing include added.

diff --git a/Engine/Source/Runtime/Core/Public/Serialization/StructuredArchiveNameHelpers.h b/Engine/Source/Runtime/Core/Public/Serialization/StructuredArchiveNameHelpers.h
--- a/Engine/Source/Runtime/Core/Public/Serialization/StructuredArchiveNameHelpers.h
+++ b/Engine/Source/Runtime/Core/Public/Serialization/StructuredArchiveNameHelpers.h
@@ -3,6 +3,7 @@
 #pragma once
 
 #include "CoreTypes.h"
+#include "Templates/Identity.h"
 
 // Define a structure to encapsulate a field name, which compiles to an empty object if WITH_TEXT_ARCHIVE_SUPPORT = 0
 struct FArchiveFieldName
