Incorporate type into ordering of the attributes to work around an issue where we can have multiple attributes with the same name, but compatible types (float vs position)

diff --git a/Engine/Plugins/FX/Niagara/Source/NiagaraEditor/Private/NiagaraHlslTranslator.cpp b/Engine/Plugins/FX/Niagara/Source/NiagaraEditor/Private/NiagaraHlslTranslator.cpp
--- a/Engine/Plugins/FX/Niagara/Source/NiagaraEditor/Private/NiagaraHlslTranslator.cpp
+++ b/Engine/Plugins/FX/Niagara/Source/NiagaraEditor/Private/NiagaraHlslTranslator.cpp
@@ -1817,7 +1817,15 @@ const FNiagaraTranslateResults &FHlslNiagaraTranslator::Translate(const FNiagara
 		TrimAttributes(CompileOptions, BasicAttributes);
 
 		// We sort the variables so that they end up in the same ordering between Spawn & Update...
-		Algo::SortBy(BasicAttributes, &FNiagaraVariable::GetName, FNameLexicalLess());
+		Algo::Sort(BasicAttributes, [](const FNiagaraVariable& Lhs, const FNiagaraVariable& Rhs)
+		{
+			const int32 NameDiff = Lhs.GetName().Compare(Rhs.GetName());
+			if (NameDiff)
+			{
+				return NameDiff < 0;
+			}
+			return Lhs.GetType().GetFName().Compare(Rhs.GetType().GetFName()) < 0;
+		});
 
 		ConvertFloatToHalf(CompileOptions, BasicAttributes);
 
diff --git a/Engine/Source/Runtime/Core/Private/UObject/FortniteMainBranchObjectVersion.cpp b/Engine/Source/Runtime/Core/Private/UObject/FortniteMainBranchObjectVersion.cpp
--- a/Engine/Source/Runtime/Core/Private/UObject/FortniteMainBranchObjectVersion.cpp
+++ b/Engine/Source/Runtime/Core/Private/UObject/FortniteMainBranchObjectVersion.cpp
@@ -11,7 +11,7 @@ TMap<FGuid, FGuid> FFortniteMainBranchObjectVersion::GetSystemGuids()
 	SystemGuids.Add(DevGuids.MATERIALSHADERMAP_DERIVEDDATA_VER, FGuid("1DA85DA5733C4A489A1BE117CDC5ACCC"));
 	SystemGuids.Add(DevGuids.NANITE_DERIVEDDATA_VER, FGuid("49FBF137DBFB4204B6C9E5B91DA15EBD"));
 	SystemGuids.Add(DevGuids.NIAGARASHADERMAP_DERIVEDDATA_VER, FGuid("2289D5116CF94BC0AFEAEC541468E645"));
-	SystemGuids.Add(DevGuids.Niagara_LatestScriptCompileVersion, FGuid("68A430A98C770B4BB1331E206470CC26"));
+	SystemGuids.Add(DevGuids.Niagara_LatestScriptCompileVersion, FGuid("A3D2BC12548B425DB10DCE6D0237381D"));
 	SystemGuids.Add(DevGuids.SkeletalMeshDerivedDataVersion, FGuid("ACF593EAAE354FCBB34CF44F5AA1BFD2"));
 	SystemGuids.Add(DevGuids.STATICMESH_DERIVEDDATA_VER, FGuid("3006d21f867648a68cf4edee185446b7"));
 
