Fix EGuidFormats::Short so that it uses Base64 encoding as the documentation states

Changes semantics of FGuid <-> Short string conversion, which might cause issues with persisted short GUIDs. The Short format is used sparingly and only to create unique strings via FGuid::NewGuid().ToString(EGuidFormats::Short), so I suspect we won't see any fallout from this.

FGuid::AppendString(FAnsiStringBuilderBase&) was totally broken and serialized stack pointers instead of GUID, but noone was using that path

diff --git a/Engine/Source/Runtime/Core/Private/Misc/Guid.cpp b/Engine/Source/Runtime/Core/Private/Misc/Guid.cpp
--- a/Engine/Source/Runtime/Core/Private/Misc/Guid.cpp
+++ b/Engine/Source/Runtime/Core/Private/Misc/Guid.cpp
@@ -76,9 +76,9 @@ void FGuid::AppendString(FString& Out, EGuidFormats Format) const
 
 	case EGuidFormats::Short:
 	{
-		const uint32 Data[] = {A,B,C,D};
+		uint32 Bytes[4] = { NETWORK_ORDER32(A), NETWORK_ORDER32(B), NETWORK_ORDER32(C), NETWORK_ORDER32(D) };
 		TCHAR Buffer[25];
-		int32 Len = FBase64::Encode(reinterpret_cast<const uint8*>(Data), sizeof(Data), Buffer);
+		int32 Len = FBase64::Encode(reinterpret_cast<const uint8*>(Bytes), sizeof(Bytes), Buffer);
 		TArrayView<TCHAR> Result(Buffer, Len);
 		
 		Algo::Replace(Result, '+', '-');
@@ -159,10 +159,10 @@ void FGuid::AppendString(FAnsiStringBuilderBase& Builder, EGuidFormats Format) c
 
 	case EGuidFormats::Short:
 	{
-		const uint32 Data[] = {A,B,C,D};
+		uint32 Bytes[4] = { NETWORK_ORDER32(A), NETWORK_ORDER32(B), NETWORK_ORDER32(C), NETWORK_ORDER32(D) };
 		int32 Index = Builder.AddUninitialized(24);
 		TArrayView<ANSICHAR> Result(Builder.GetData() + Index, 24);
-		int32 Len = FBase64::Encode(reinterpret_cast<const uint8*>(&Data), sizeof(Data), Result.GetData());
+		int32 Len = FBase64::Encode(reinterpret_cast<const uint8*>(Bytes), sizeof(Bytes), Result.GetData());
 
 		Algo::Replace(Result, '+', '-');
 		Algo::Replace(Result, '/', '_');
@@ -306,10 +306,10 @@ bool FGuid::ParseExact(const FString& GuidString, EGuidFormats Format, FGuid& Ou
 			return false;
 		}
 
-		OutGuid = FGuid(Data[0],
-						Data[1],
-						Data[2],
-						Data[3]);
+		OutGuid = FGuid(NETWORK_ORDER32(Data[0]),
+						NETWORK_ORDER32(Data[1]),
+						NETWORK_ORDER32(Data[2]),
+						NETWORK_ORDER32(Data[3]));
 		return true;
 	}
 
diff --git a/Engine/Source/Runtime/Core/Private/Tests/Misc/GuidTest.cpp b/Engine/Source/Runtime/Core/Private/Tests/Misc/GuidTest.cpp
--- a/Engine/Source/Runtime/Core/Private/Tests/Misc/GuidTest.cpp
+++ b/Engine/Source/Runtime/Core/Private/Tests/Misc/GuidTest.cpp
@@ -11,6 +11,9 @@ IMPLEMENT_SIMPLE_AUTOMATION_TEST(FGuidTest, "System.Core.Misc.Guid", EAutomation
 
 bool FGuidTest::RunTest( const FString& Parameters )
 {
+	// Base36		12SRDE1NOQERISMPEF0H6KUU9
+	// Base64		EjRWeIdlQyESNFZ4h2VDIQ
+	// Hex			0x12345678 0x87654321  0x12345678 0x87654321
 	FGuid g = FGuid(305419896, 2271560481, 305419896, 2271560481);
 
 	// string conversions
@@ -21,7 +24,7 @@ bool FGuidTest::RunTest( const FString& Parameters )
 	TestEqual<FString>(TEXT("String conversion (EGuidFormats::DigitsWithHyphensInParentheses)"), g.ToString(EGuidFormats::DigitsWithHyphensInParentheses), TEXT("(12345678-8765-4321-1234-567887654321)"));
 	TestEqual<FString>(TEXT("String conversion (EGuidFormats::HexValuesInBraces)"), g.ToString(EGuidFormats::HexValuesInBraces), TEXT("{0x12345678,0x8765,0x4321,{0x12,0x34,0x56,0x78,0x87,0x65,0x43,0x21}}"));
 	TestEqual<FString>(TEXT("String conversion (EGuidFormats::UniqueObjectGuid)"), g.ToString(EGuidFormats::UniqueObjectGuid), TEXT("12345678-87654321-12345678-87654321"));
-	TestEqual<FString>(TEXT("String conversion (EGuidFormats::Short)"), g.ToString(EGuidFormats::Short), TEXT("eFY0EiFDZYd4VjQSIUNlhw"));
+	TestEqual<FString>(TEXT("String conversion (EGuidFormats::Short)"), g.ToString(EGuidFormats::Short), TEXT("EjRWeIdlQyESNFZ4h2VDIQ"));
 	TestEqual<FString>(TEXT("String conversion (EGuidFormats::Base36Encoded)"), g.ToString(EGuidFormats::Base36Encoded), TEXT("12SRDE1NOQERISMPEF0H6KUU9"));
 
 	// parsing valid strings (exact)
@@ -45,7 +48,7 @@ bool FGuidTest::RunTest( const FString& Parameters )
 	TestTrue(TEXT("Parsing valid strings must succeed (EGuidFormats::UniqueObjectGuid)"), FGuid::ParseExact(TEXT("12345678-87654321-12345678-87654321"), EGuidFormats::UniqueObjectGuid, g2_1));
 	TestEqual(TEXT("Parsing valid strings must succeed (EGuidFormats::UniqueObjectGuid)"), g2_1, g);
 
-	TestTrue(TEXT("Parsing valid strings must succeed (EGuidFormats::Short)"), FGuid::ParseExact(TEXT("eFY0EiFDZYd4VjQSIUNlhw"), EGuidFormats::Short, g2_1));
+	TestTrue(TEXT("Parsing valid strings must succeed (EGuidFormats::Short)"), FGuid::ParseExact(TEXT("EjRWeIdlQyESNFZ4h2VDIQ"), EGuidFormats::Short, g2_1));
 	TestEqual(TEXT("Parsing valid strings must succeed (EGuidFormats::Short)"), g2_1, g);
 
 	TestTrue(TEXT("Parsing valid strings must succeed (EGuidFormats::Base36Encoded)"), FGuid::ParseExact(TEXT("12SRDE1NOQERISMPEF0H6KUU9"), EGuidFormats::Base36Encoded, g2_1));
@@ -75,8 +78,8 @@ bool FGuidTest::RunTest( const FString& Parameters )
 	TestTrue(TEXT("Parsing valid strings must succeed (12345678-87654321-12345678-87654321)"), FGuid::Parse(TEXT("12345678-87654321-12345678-87654321"), g3_1));
 	TestEqual(TEXT("Parsing valid strings must succeed (12345678-87654321-12345678-87654321)"), g3_1, g);
 
-	TestTrue(TEXT("Parsing valid strings must succeed (eFY0EiFDZYd4VjQSIUNlhw)"), FGuid::Parse(TEXT("eFY0EiFDZYd4VjQSIUNlhw"), g3_1));
-	TestEqual(TEXT("Parsing valid strings must succeed (eFY0EiFDZYd4VjQSIUNlhw)"), g3_1, g);
+	TestTrue(TEXT("Parsing valid strings must succeed (EjRWeIdlQyESNFZ4h2VDIQ)"), FGuid::Parse(TEXT("EjRWeIdlQyESNFZ4h2VDIQ"), g3_1));
+	TestEqual(TEXT("Parsing valid strings must succeed (EjRWeIdlQyESNFZ4h2VDIQ)"), g3_1, g);
 
 	TestTrue(TEXT("Parsing valid strings must succeed (12SRDE1NOQERISMPEF0H6KUU9)"), FGuid::Parse(TEXT("12SRDE1NOQERISMPEF0H6KUU9"), g3_1));
 	TestEqual(TEXT("Parsing valid strings must succeed (12SRDE1NOQERISMPEF0H6KUU9)"), g3_1, g);
