Fix mimalloc non-unity build for Mac platform

#rnx

diff --git a/Engine/Source/Runtime/Core/Private/HAL/MallocMimalloc.cpp b/Engine/Source/Runtime/Core/Private/HAL/MallocMimalloc.cpp
--- a/Engine/Source/Runtime/Core/Private/HAL/MallocMimalloc.cpp
+++ b/Engine/Source/Runtime/Core/Private/HAL/MallocMimalloc.cpp
@@ -11,6 +11,10 @@
 // Only use for supported platforms
 #if PLATFORM_SUPPORTS_MIMALLOC && MIMALLOC_ALLOCATOR_ALLOWED && PLATFORM_BUILDS_MIMALLOC
 
+#if PLATFORM_MAC
+#include "Templates/AlignmentTemplates.h"
+#endif
+
 /** Value we fill a memory block with after it is free, in UE_BUILD_DEBUG **/
 #define DEBUG_FILL_FREED (0xdd)
 
