Add a way to express fallback tag in Memory Trace

Manual memory tracking in LLM has a feature for defining a "fallback tag" for the allocation if no tag is currently set. This change adds a new scope type which allows us to express this behaviour in memory tracing as well.

diff --git a/Engine/Source/Runtime/Core/Private/ProfilingDebugging/TagTrace.cpp b/Engine/Source/Runtime/Core/Private/ProfilingDebugging/TagTrace.cpp
--- a/Engine/Source/Runtime/Core/Private/ProfilingDebugging/TagTrace.cpp
+++ b/Engine/Source/Runtime/Core/Private/ProfilingDebugging/TagTrace.cpp
@@ -35,34 +35,34 @@ UE_TRACE_EVENT_END()
 thread_local int32 GActiveTag;
 
 ////////////////////////////////////////////////////////////////////////////////
-FMemScope::FMemScope(int32 InTag)
+FMemScope::FMemScope(int32 InTag, bool bShouldActivate /*= true*/)
 {
-	if (UE_TRACE_CHANNELEXPR_IS_ENABLED(MemAllocChannel))
+	if (UE_TRACE_CHANNELEXPR_IS_ENABLED(MemAllocChannel) & bShouldActivate)
 	{
 		ActivateScope(InTag);
 	}
 }
 
 ////////////////////////////////////////////////////////////////////////////////
-FMemScope::FMemScope(ELLMTag InTag)
+FMemScope::FMemScope(ELLMTag InTag, bool bShouldActivate /*= true*/)
 {
-	if (UE_TRACE_CHANNELEXPR_IS_ENABLED(MemAllocChannel))
+	if (UE_TRACE_CHANNELEXPR_IS_ENABLED(MemAllocChannel) & bShouldActivate)
 	{
 		ActivateScope(static_cast<int32>(InTag));
 	}
 }
 
 ////////////////////////////////////////////////////////////////////////////////
-FMemScope::FMemScope(const FName& InName)
+FMemScope::FMemScope(const FName& InName, bool bShouldActivate /*= true*/)
 {
-	if (UE_TRACE_CHANNELEXPR_IS_ENABLED(MemAllocChannel))
+	if (UE_TRACE_CHANNELEXPR_IS_ENABLED(MemAllocChannel) & bShouldActivate)
 	{
 		ActivateScope(MemoryTrace_AnnounceFNameTag(InName));
 	}
 }
 
 ////////////////////////////////////////////////////////////////////////////////
-FMemScope::FMemScope(const UE::LLMPrivate::FTagData* TagData)
+FMemScope::FMemScope(const UE::LLMPrivate::FTagData* TagData, bool bShouldActivate /*= true*/)
 {
 	// TagData is opaque so we cant really use the input, additionally we 
 	// cannot count on LLM being active. Instead we have inserted an explicit 
diff --git a/Engine/Source/Runtime/Core/Public/ProfilingDebugging/TagTrace.h b/Engine/Source/Runtime/Core/Public/ProfilingDebugging/TagTrace.h
--- a/Engine/Source/Runtime/Core/Public/ProfilingDebugging/TagTrace.h
+++ b/Engine/Source/Runtime/Core/Public/ProfilingDebugging/TagTrace.h
@@ -41,10 +41,10 @@ CORE_API int32	MemoryTrace_GetActiveTag();
 class FMemScope
 {
 public:
-	CORE_API FMemScope(int32 InTag);
-	CORE_API FMemScope(ELLMTag InTag);
-	CORE_API FMemScope(const class FName& InName);
-	CORE_API FMemScope(const UE::LLMPrivate::FTagData* TagData);
+	CORE_API FMemScope(int32 InTag, bool bShouldActivate = true);
+	CORE_API FMemScope(ELLMTag InTag, bool bShouldActivate = true);
+	CORE_API FMemScope(const class FName& InName, bool bShouldActivate = true);
+	CORE_API FMemScope(const UE::LLMPrivate::FTagData* TagData, bool bShouldActivate = true);
 	CORE_API ~FMemScope();
 private:
 	void ActivateScope(int32 InTag);
@@ -52,6 +52,21 @@ private:
 	int32 PrevTag;
 };
 
+
+/**
+ * A scope that activates in case no existing scope is active.
+ */
+template<typename TagType>
+class FDefaultMemScope : public FMemScope
+{
+public:
+	CORE_API FDefaultMemScope(TagType InTag)
+		: FMemScope(InTag, MemoryTrace_GetActiveTag() == 0)
+	{
+		
+	}
+};
+
 /**
  * Used order to keep the tag for memory that is being reallocated.
  */
@@ -70,12 +85,14 @@ constexpr int32 TRACE_TAG = 257;
 ////////////////////////////////////////////////////////////////////////////////
 #define UE_MEMSCOPE(InTag)				FMemScope PREPROCESSOR_JOIN(MemScope,__LINE__)(InTag);
 #define UE_MEMSCOPE_PTR(InPtr)			FMemScopePtr PREPROCESSOR_JOIN(MemPtrScope,__LINE__)((uint64)InPtr);
+#define UE_MEMSCOPE_DEFAULT(InTag)		FDefaultMemScope PREPROCESSOR_JOIN(MemScope,__LINE__)(InTag);
 
 #else // UE_MEMORY_TAGS_TRACE_ENABLED
 
 ////////////////////////////////////////////////////////////////////////////////
 #define UE_MEMSCOPE(...)
 #define UE_MEMSCOPE_PTR(...)
+#define UE_MEMSCOPE_DEFAULT(...)
 
 #endif // UE_MEMORY_TAGS_TRACE_ENABLED
 
