FTraceAuxiliary: Added GetTraceDestination, IsConnected and GetActiveChannelsString to public interface.

diff --git a/Engine/Source/Runtime/Core/Private/ProfilingDebugging/TraceAuxiliary.cpp b/Engine/Source/Runtime/Core/Private/ProfilingDebugging/TraceAuxiliary.cpp
--- a/Engine/Source/Runtime/Core/Private/ProfilingDebugging/TraceAuxiliary.cpp
+++ b/Engine/Source/Runtime/Core/Private/ProfilingDebugging/TraceAuxiliary.cpp
@@ -1289,6 +1289,29 @@ void FTraceAuxiliary::EnableChannels()
 #endif
 }
 
+const TCHAR* FTraceAuxiliary::GetTraceDestination()
+{
+#if UE_TRACE_ENABLED
+	return GTraceAuxiliary.GetDest();
+#endif
+	return nullptr;
+}
+
+bool FTraceAuxiliary::IsConnected()
+{
+#if UE_TRACE_ENABLED
+	return GTraceAuxiliary.IsConnected();
+#endif
+	return false;
+}
+
+void FTraceAuxiliary::GetActiveChannelsString(FStringBuilderBase& String)
+{
+#if UE_TRACE_ENABLED
+	GTraceAuxiliary.GetActiveChannelsString(String);
+#endif
+}
+
 ////////////////////////////////////////////////////////////////////////////////
 void FTraceAuxiliary::TryAutoConnect()
 {
diff --git a/Engine/Source/Runtime/Core/Public/ProfilingDebugging/TraceAuxiliary.h b/Engine/Source/Runtime/Core/Public/ProfilingDebugging/TraceAuxiliary.h
--- a/Engine/Source/Runtime/Core/Public/ProfilingDebugging/TraceAuxiliary.h
+++ b/Engine/Source/Runtime/Core/Public/ProfilingDebugging/TraceAuxiliary.h
@@ -99,4 +99,23 @@ public:
 	 *  as channels can be announced on module loading.
 	 */
 	static void EnableChannels();
+
+	/**
+	 *  Returns the destination string that is currently being traced to.
+	 *  Contains either a file path or network address. Points to an empty string if tracing is disabled.
+	 */
+	static const TCHAR*	GetTraceDestination();
+	
+	/**
+	 *  Returns whether the trace system is currently connected to a trace sink (file or network)
+	 */
+	static bool	IsConnected();
+
+	/**
+	 *  Adds a comma separated list of currently active channels to the passed in StringBuilder
+	 */
+	static void	GetActiveChannelsString(FStringBuilderBase& String);
+
+
+	
 };
