Un-smoke tests that were not previously filtered as Smoke by the AutomationTest framework.

diff --git a/Engine/Source/Runtime/Core/Tests/Containers/DequeTest.cpp b/Engine/Source/Runtime/Core/Tests/Containers/DequeTest.cpp
--- a/Engine/Source/Runtime/Core/Tests/Containers/DequeTest.cpp
+++ b/Engine/Source/Runtime/Core/Tests/Containers/DequeTest.cpp
@@ -83,7 +83,7 @@ bool Test::EmplaceFirstPopLast(TDeque<int32Token>& Deque, int32 Count)
 // Unit tests
 //---------------------------------------------------------------------------------------------------------------------
 
-TEST_CASE("Core::Containers::TDeque::Reserve without data", "[Core][Containers][Smoke]")
+TEST_CASE("Core::Containers::TDeque::Reserve without data", "[Core][Containers][Deque]")
 {
 	int32Token::Reset();
 	TDeque<int32Token> Deque;
@@ -97,7 +97,7 @@ TEST_CASE("Core::Containers::TDeque::Reserve without data", "[Core][Containers][
 	CHECK(int32Token::EvenConstructionDestructionCalls(0));
 }
 
-TEST_CASE("Core::Containers::TDeque::Reserve EmplaceLast single element", "[Core][Containers][Smoke]")
+TEST_CASE("Core::Containers::TDeque::Reserve EmplaceLast single element", "[Core][Containers][Deque]")
 {
 	int32Token::Reset();
 	{
@@ -116,7 +116,7 @@ TEST_CASE("Core::Containers::TDeque::Reserve EmplaceLast single element", "[Core
 	CHECK(int32Token::EvenConstructionDestructionCalls(1));
 }
 
-TEST_CASE("Core::Containers::TDeque::Reset", "[Core][Containers][Smoke]")
+TEST_CASE("Core::Containers::TDeque::Reset", "[Core][Containers][Deque]")
 {
 	int32Token::Reset();
 	{
@@ -137,7 +137,7 @@ TEST_CASE("Core::Containers::TDeque::Reset", "[Core][Containers][Smoke]")
 	CHECK(int32Token::EvenConstructionDestructionCalls(1));
 }
 
-TEST_CASE("Core::Containers::TDeque::Empty", "[Core][Containers][Smoke]")
+TEST_CASE("Core::Containers::TDeque::Empty", "[Core][Containers][Deque]")
 {
 	int32Token::Reset();
 	{
@@ -150,7 +150,7 @@ TEST_CASE("Core::Containers::TDeque::Empty", "[Core][Containers][Smoke]")
 	CHECK(int32Token::EvenConstructionDestructionCalls(0));
 }
 
-TEST_CASE("Core::Containers::TDeque::Empty after single element EmplaceLast", "[Core][Containers][Smoke]")
+TEST_CASE("Core::Containers::TDeque::Empty after single element EmplaceLast", "[Core][Containers][Deque]")
 {
 	int32Token::Reset();
 	{
@@ -167,7 +167,7 @@ TEST_CASE("Core::Containers::TDeque::Empty after single element EmplaceLast", "[
 	CHECK(int32Token::EvenConstructionDestructionCalls(1));
 }
 
-TEST_CASE("Core::Containers::TDeque::EmplaceLast single element", "[Core][Containers][Smoke]")
+TEST_CASE("Core::Containers::TDeque::EmplaceLast single element", "[Core][Containers][Deque]")
 {
 	int32Token::Reset();
 	{
@@ -183,7 +183,7 @@ TEST_CASE("Core::Containers::TDeque::EmplaceLast single element", "[Core][Contai
 	CHECK(int32Token::EvenConstructionDestructionCalls(1));
 }
 
-TEST_CASE("Core::Containers::TDeque::EmplaceLast range to capacity", "[Core][Containers][Smoke]")
+TEST_CASE("Core::Containers::TDeque::EmplaceLast range to capacity", "[Core][Containers][Deque]")
 {
 	int32Token::Reset();
 	{
@@ -202,7 +202,7 @@ TEST_CASE("Core::Containers::TDeque::EmplaceLast range to capacity", "[Core][Con
 	CHECK(int32Token::EvenConstructionDestructionCalls(Test::DefaultCapacity * 10));
 }
 
-TEST_CASE("Core::Containers::TDeque::EmplaceLast range past capacity", "[Core][Containers][Smoke]")
+TEST_CASE("Core::Containers::TDeque::EmplaceLast range past capacity", "[Core][Containers][Deque]")
 {
 	int32Token::Reset();
 	{
@@ -224,7 +224,7 @@ TEST_CASE("Core::Containers::TDeque::EmplaceLast range past capacity", "[Core][C
 	CHECK(int32Token::EvenConstructionDestructionCalls(Test::DefaultCapacity + 1));
 }
 
-TEST_CASE("Core::Containers::TDeque::EmplaceFirst single element", "[Core][Containers][Smoke]")
+TEST_CASE("Core::Containers::TDeque::EmplaceFirst single element", "[Core][Containers][Deque]")
 {
 	int32Token::Reset();
 	{
@@ -240,7 +240,7 @@ TEST_CASE("Core::Containers::TDeque::EmplaceFirst single element", "[Core][Conta
 	CHECK(int32Token::EvenConstructionDestructionCalls(1));
 }
 
-TEST_CASE("Core::Containers::TDeque::EmplaceFirst range to capacity", "[Core][Containers][Smoke]")
+TEST_CASE("Core::Containers::TDeque::EmplaceFirst range to capacity", "[Core][Containers][Deque]")
 {
 	int32Token::Reset();
 	{
@@ -259,7 +259,7 @@ TEST_CASE("Core::Containers::TDeque::EmplaceFirst range to capacity", "[Core][Co
 	CHECK(int32Token::EvenConstructionDestructionCalls(Test::DefaultCapacity * 10));
 }
 
-TEST_CASE("Core::Containers::TDeque::PushLast single element (implicit move)", "[Core][Containers][Smoke]")
+TEST_CASE("Core::Containers::TDeque::PushLast single element (implicit move)", "[Core][Containers][Deque]")
 {
 	int32Token::Reset();
 	{
@@ -278,7 +278,7 @@ TEST_CASE("Core::Containers::TDeque::PushLast single element (implicit move)", "
 	CHECK(int32Token::NumMoveConstructorCalls() == 1);
 }
 
-TEST_CASE("Core::Containers::TDeque::PushLast single element from move", "[Core][Containers][Smoke]")
+TEST_CASE("Core::Containers::TDeque::PushLast single element from move", "[Core][Containers][Deque]")
 {
 	int32Token::Reset();
 	{
@@ -298,7 +298,7 @@ TEST_CASE("Core::Containers::TDeque::PushLast single element from move", "[Core]
 	CHECK(int32Token::NumMoveConstructorCalls() == 1);
 }
 
-TEST_CASE("Core::Containers::TDeque::PushLast single element from copy", "[Core][Containers][Smoke]")
+TEST_CASE("Core::Containers::TDeque::PushLast single element from copy", "[Core][Containers][Deque]")
 {
 	int32Token::Reset();
 	{
@@ -318,7 +318,7 @@ TEST_CASE("Core::Containers::TDeque::PushLast single element from copy", "[Core]
 	CHECK(int32Token::NumMoveConstructorCalls() == 0);
 }
 
-TEST_CASE("Core::Containers::TDeque::PushFirst single element (implicit move)", "[Core][Containers][Smoke]")
+TEST_CASE("Core::Containers::TDeque::PushFirst single element (implicit move)", "[Core][Containers][Deque]")
 {
 	int32Token::Reset();
 	{
@@ -337,7 +337,7 @@ TEST_CASE("Core::Containers::TDeque::PushFirst single element (implicit move)",
 	CHECK(int32Token::NumMoveConstructorCalls() == 1);
 }
 
-TEST_CASE("Core::Containers::TDeque::PushFirst single element from move", "[Core][Containers][Smoke]")
+TEST_CASE("Core::Containers::TDeque::PushFirst single element from move", "[Core][Containers][Deque]")
 {
 	int32Token::Reset();
 	{
@@ -357,7 +357,7 @@ TEST_CASE("Core::Containers::TDeque::PushFirst single element from move", "[Core
 	CHECK(int32Token::NumMoveConstructorCalls() == 1);
 }
 
-TEST_CASE("Core::Containers::TDeque::PushFirst single element from copy", "[Core][Containers][Smoke]")
+TEST_CASE("Core::Containers::TDeque::PushFirst single element from copy", "[Core][Containers][Deque]")
 {
 	int32Token::Reset();
 	{
@@ -382,7 +382,7 @@ void PopOne(TDeque<int32Token>& Deque)
 	Deque.PopFirst();
 }
 
-TEST_CASE("Core::Containers::TDeque::EmplaceLast/PopFirst single element", "[Core][Containers][Smoke]")
+TEST_CASE("Core::Containers::TDeque::EmplaceLast/PopFirst single element", "[Core][Containers][Deque]")
 {
 	int32Token::Reset();
 	{
@@ -402,7 +402,7 @@ TEST_CASE("Core::Containers::TDeque::EmplaceLast/PopFirst single element", "[Cor
 	CHECK(int32Token::EvenConstructionDestructionCalls(1));
 }
 
-TEST_CASE("Core::Containers::TDeque::EmplaceLast/PopFirst single element multiple times causing head/tail wrap around", "[Core][Containers][Smoke]")
+TEST_CASE("Core::Containers::TDeque::EmplaceLast/PopFirst single element multiple times causing head/tail wrap around", "[Core][Containers][Deque]")
 {
 	int32Token::Reset();
 	{
@@ -418,7 +418,7 @@ TEST_CASE("Core::Containers::TDeque::EmplaceLast/PopFirst single element multipl
 	CHECK(int32Token::EvenConstructionDestructionCalls(Test::DefaultCapacity * 2));
 }
 
-TEST_CASE("Core::Containers::TDeque::EmplaceLast/PopFirst range without head/tail wrap around", "[Core][Containers][Smoke]")
+TEST_CASE("Core::Containers::TDeque::EmplaceLast/PopFirst range without head/tail wrap around", "[Core][Containers][Deque]")
 {
 	int32Token::Reset();
 	{
@@ -430,7 +430,7 @@ TEST_CASE("Core::Containers::TDeque::EmplaceLast/PopFirst range without head/tai
 	CHECK(int32Token::EvenConstructionDestructionCalls(Test::DefaultCapacity));
 }
 
-TEST_CASE("Core::Containers::TDeque::EmplaceLast/PopFirst range with reallocation without head/tail wrap around", "[Core][Containers][Smoke]")
+TEST_CASE("Core::Containers::TDeque::EmplaceLast/PopFirst range with reallocation without head/tail wrap around", "[Core][Containers][Deque]")
 {
 	int32Token::Reset();
 	{
@@ -443,7 +443,7 @@ TEST_CASE("Core::Containers::TDeque::EmplaceLast/PopFirst range with reallocatio
 	CHECK(int32Token::EvenConstructionDestructionCalls(Test::DefaultCapacity + 1));
 }
 
-TEST_CASE("Core::Containers::TDeque::EmplaceLast/PopFirst range with head/tail wrap around", "[Core][Containers][Smoke]")
+TEST_CASE("Core::Containers::TDeque::EmplaceLast/PopFirst range with head/tail wrap around", "[Core][Containers][Deque]")
 {
 	int32Token::Reset();
 	{
@@ -458,7 +458,7 @@ TEST_CASE("Core::Containers::TDeque::EmplaceLast/PopFirst range with head/tail w
 	CHECK(int32Token::EvenConstructionDestructionCalls(Test::DefaultCapacity * (Test::DefaultCapacity - 1)));
 }
 
-TEST_CASE("Core::Containers::TDeque::EmplaceFirst/PopLast single element", "[Core][Containers][Smoke]")
+TEST_CASE("Core::Containers::TDeque::EmplaceFirst/PopLast single element", "[Core][Containers][Deque]")
 {
 	int32Token::Reset();
 	{
@@ -478,7 +478,7 @@ TEST_CASE("Core::Containers::TDeque::EmplaceFirst/PopLast single element", "[Cor
 	CHECK(int32Token::EvenConstructionDestructionCalls(1));
 }
 
-TEST_CASE("Core::Containers::TDeque::EmplaceFirst/PopLast range", "[Core][Containers][Smoke]")
+TEST_CASE("Core::Containers::TDeque::EmplaceFirst/PopLast range", "[Core][Containers][Deque]")
 {
 	int32Token::Reset();
 	{
@@ -490,7 +490,7 @@ TEST_CASE("Core::Containers::TDeque::EmplaceFirst/PopLast range", "[Core][Contai
 	CHECK(int32Token::EvenConstructionDestructionCalls(Test::DefaultCapacity));
 }
 
-TEST_CASE("Core::Containers::TDeque::EmplaceFirst/PopLast range with reallocation", "[Core][Containers][Smoke]")
+TEST_CASE("Core::Containers::TDeque::EmplaceFirst/PopLast range with reallocation", "[Core][Containers][Deque]")
 {
 	int32Token::Reset();
 	{
@@ -503,7 +503,7 @@ TEST_CASE("Core::Containers::TDeque::EmplaceFirst/PopLast range with reallocatio
 	CHECK(int32Token::EvenConstructionDestructionCalls(Test::DefaultCapacity + 1));
 }
 
-TEST_CASE("Core::Containers::TDeque::TryPopFirst", "[Core][Containers][Smoke]")
+TEST_CASE("Core::Containers::TDeque::TryPopFirst", "[Core][Containers][Deque]")
 {
 	int32Token::Reset();
 	{
@@ -526,7 +526,7 @@ TEST_CASE("Core::Containers::TDeque::TryPopFirst", "[Core][Containers][Smoke]")
 	CHECK(int32Token::EvenConstructionDestructionCalls(Test::DefaultCapacity + 1));  // + 1 for Result
 }
 
-TEST_CASE("Core::Containers::TDeque::TryPopFirst with reallocation", "[Core][Containers][Smoke]")
+TEST_CASE("Core::Containers::TDeque::TryPopFirst with reallocation", "[Core][Containers][Deque]")
 {
 	int32Token::Reset();
 	{
@@ -552,7 +552,7 @@ TEST_CASE("Core::Containers::TDeque::TryPopFirst with reallocation", "[Core][Con
 	CHECK(int32Token::EvenConstructionDestructionCalls(Test::DefaultCapacity + 2));
 }
 
-TEST_CASE("Core::Containers::TDeque::TryPopLast", "[Core][Containers][Smoke]")
+TEST_CASE("Core::Containers::TDeque::TryPopLast", "[Core][Containers][Deque]")
 {
 	int32Token::Reset();
 	{
@@ -575,7 +575,7 @@ TEST_CASE("Core::Containers::TDeque::TryPopLast", "[Core][Containers][Smoke]")
 	CHECK(int32Token::EvenConstructionDestructionCalls(Test::DefaultCapacity + 1));  // + 1 for Result
 }
 
-TEST_CASE("Core::Containers::TDeque::TryPopLast with reallocation", "[Core][Containers][Smoke]")
+TEST_CASE("Core::Containers::TDeque::TryPopLast with reallocation", "[Core][Containers][Deque]")
 {
 	int32Token::Reset();
 	{
@@ -601,7 +601,7 @@ TEST_CASE("Core::Containers::TDeque::TryPopLast with reallocation", "[Core][Cont
 	CHECK(int32Token::EvenConstructionDestructionCalls(Test::DefaultCapacity + 2));
 }
 
-TEST_CASE("Core::Containers::TDeque::Comparison simple", "[Core][Containers][Smoke]")
+TEST_CASE("Core::Containers::TDeque::Comparison simple", "[Core][Containers][Deque]")
 {
 	int32Token::Reset();
 	{
@@ -620,7 +620,7 @@ TEST_CASE("Core::Containers::TDeque::Comparison simple", "[Core][Containers][Smo
 	CHECK(int32Token::EvenConstructionDestructionCalls(Test::DefaultCapacity * 2));
 }
 
-TEST_CASE("Core::Containers::TDeque::Comparison with head/tail wrap around", "[Core][Containers][Smoke]")
+TEST_CASE("Core::Containers::TDeque::Comparison with head/tail wrap around", "[Core][Containers][Deque]")
 {
 	int32Token::Reset();
 	{
@@ -648,7 +648,7 @@ TEST_CASE("Core::Containers::TDeque::Comparison with head/tail wrap around", "[C
 	CHECK(int32Token::EvenConstructionDestructionCalls(Test::DefaultCapacity * 6));
 }
 
-TEST_CASE("Core::Containers::TDeque::Copy simple", "[Core][Containers][Smoke]")
+TEST_CASE("Core::Containers::TDeque::Copy simple", "[Core][Containers][Deque]")
 {
 	int32Token::Reset();
 	{
@@ -666,7 +666,7 @@ TEST_CASE("Core::Containers::TDeque::Copy simple", "[Core][Containers][Smoke]")
 	CHECK(int32Token::EvenConstructionDestructionCalls(Test::DefaultCapacity * 2));
 }
 
-TEST_CASE("Core::Containers::TDeque::Copy with head/tail wrap around", "[Core][Containers][Smoke]")
+TEST_CASE("Core::Containers::TDeque::Copy with head/tail wrap around", "[Core][Containers][Deque]")
 {
 	int32Token::Reset();
 	{
@@ -689,7 +689,7 @@ TEST_CASE("Core::Containers::TDeque::Copy with head/tail wrap around", "[Core][C
 	CHECK(int32Token::EvenConstructionDestructionCalls(Test::DefaultCapacity * 9));
 }
 
-TEST_CASE("Core::Containers::TDeque::Copy variable size with head/tail wrap around", "[Core][Containers][Smoke]")
+TEST_CASE("Core::Containers::TDeque::Copy variable size with head/tail wrap around", "[Core][Containers][Deque]")
 {
 	int32Token::Reset();
 	{
@@ -717,7 +717,7 @@ TEST_CASE("Core::Containers::TDeque::Copy variable size with head/tail wrap arou
 	CHECK(int32Token::EvenConstructionDestructionCalls());
 }
 
-TEST_CASE("Core::Containers::TDeque::Move simple", "[Core][Containers][Smoke]")
+TEST_CASE("Core::Containers::TDeque::Move simple", "[Core][Containers][Deque]")
 {
 	int32Token::Reset();
 	{
@@ -741,7 +741,7 @@ TEST_CASE("Core::Containers::TDeque::Move simple", "[Core][Containers][Smoke]")
 	CHECK(int32Token::EvenConstructionDestructionCalls(Test::DefaultCapacity + 1));
 }
 
-TEST_CASE("Core::Containers::TDeque::Move with head/tail wrap around", "[Core][Containers][Smoke]")
+TEST_CASE("Core::Containers::TDeque::Move with head/tail wrap around", "[Core][Containers][Deque]")
 {
 	int32Token::Reset();
 	{
@@ -769,7 +769,7 @@ TEST_CASE("Core::Containers::TDeque::Move with head/tail wrap around", "[Core][C
 	CHECK(int32Token::EvenConstructionDestructionCalls(Test::DefaultCapacity * 6));
 }
 
-TEST_CASE("Core::Containers::TDeque::Move variable size with head/tail wrap around", "[Core][Containers][Smoke]")
+TEST_CASE("Core::Containers::TDeque::Move variable size with head/tail wrap around", "[Core][Containers][Deque]")
 {
 	int32Token::Reset();
 	{
@@ -801,7 +801,7 @@ TEST_CASE("Core::Containers::TDeque::Move variable size with head/tail wrap arou
 	CHECK(int32Token::EvenConstructionDestructionCalls());
 }
 
-TEST_CASE("Core::Containers::TDeque::Iteration without head/tail wrap around", "[Core][Containers][Smoke]")
+TEST_CASE("Core::Containers::TDeque::Iteration without head/tail wrap around", "[Core][Containers][Deque]")
 {
 	int32Token::Reset();
 	{
@@ -826,7 +826,7 @@ TEST_CASE("Core::Containers::TDeque::Iteration without head/tail wrap around", "
 	CHECK(int32Token::EvenConstructionDestructionCalls(Test::DefaultCapacity));
 }
 
-TEST_CASE("Core::Containers::TDeque::Iteration with head/tail wrap around", "[Core][Containers][Smoke]")
+TEST_CASE("Core::Containers::TDeque::Iteration with head/tail wrap around", "[Core][Containers][Deque]")
 {
 	int32Token::Reset();
 	{
@@ -856,7 +856,7 @@ TEST_CASE("Core::Containers::TDeque::Iteration with head/tail wrap around", "[Co
 	CHECK(int32Token::EvenConstructionDestructionCalls(Test::DefaultCapacity * (Test::DefaultCapacity + 1)));
 }
 
-TEST_CASE("Core::Containers::TDeque::Iterator arithmetic", "[Core][Containers][Smoke]")
+TEST_CASE("Core::Containers::TDeque::Iterator arithmetic", "[Core][Containers][Deque]")
 {
 	TDeque<int32Token> Deque;
 	Deque.Reserve(Test::DefaultCapacity);
@@ -883,7 +883,7 @@ TEST_CASE("Core::Containers::TDeque::Iterator arithmetic", "[Core][Containers][S
 	}
 }
 
-TEST_CASE("Core::Containers::TDeque::Construct from std initializer_list", "[Core][Containers][Smoke]")
+TEST_CASE("Core::Containers::TDeque::Construct from std initializer_list", "[Core][Containers][Deque]")
 {
 	int32Token::Reset();
 	{
@@ -899,7 +899,7 @@ TEST_CASE("Core::Containers::TDeque::Construct from std initializer_list", "[Cor
 	CHECK(int32Token::EvenConstructionDestructionCalls(6 * 2 + 1));
 }
 
-TEST_CASE("Core::Containers::TDeque::Construct from empty std initializer_list", "[Core][Containers][Smoke]")
+TEST_CASE("Core::Containers::TDeque::Construct from empty std initializer_list", "[Core][Containers][Deque]")
 {
 	int32Token::Reset();
 	{
@@ -909,7 +909,7 @@ TEST_CASE("Core::Containers::TDeque::Construct from empty std initializer_list",
 	CHECK(int32Token::EvenConstructionDestructionCalls(0));
 }
 
-TEST_CASE("Core::Containers::TDeque::Assign from std initializer_list", "[Core][Containers][Smoke]")
+TEST_CASE("Core::Containers::TDeque::Assign from std initializer_list", "[Core][Containers][Deque]")
 {
 	int32Token::Reset();
 	{
diff --git a/Engine/Source/Runtime/Core/Tests/Containers/TokenTest.cpp b/Engine/Source/Runtime/Core/Tests/Containers/TokenTest.cpp
--- a/Engine/Source/Runtime/Core/Tests/Containers/TokenTest.cpp
+++ b/Engine/Source/Runtime/Core/Tests/Containers/TokenTest.cpp
@@ -3,7 +3,7 @@
 #include "TokenTest.h"
 #include "TestHarness.h"
 
-TEST_CASE("Core::Containers::TToken::Default", "[Core][Containers][Smoke]")
+TEST_CASE("Core::Containers::TToken::Default", "[Core][Containers][TToken]")
 {
 	int32Token::Reset();
 	{
@@ -18,7 +18,7 @@ TEST_CASE("Core::Containers::TToken::Default", "[Core][Containers][Smoke]")
 	CHECK(int32Token::NumDestructionCalls() == 1);
 }
 
-TEST_CASE("Core::Containers::TToken::Explicit constructor", "[Core][Containers][Smoke]")
+TEST_CASE("Core::Containers::TToken::Explicit constructor", "[Core][Containers][TToken]")
 {
 	int32Token::Reset();
 	{
@@ -35,7 +35,7 @@ TEST_CASE("Core::Containers::TToken::Explicit constructor", "[Core][Containers][
 	CHECK(int32Token::NumMoveAssignmentCalls() == 0);
 }
 
-TEST_CASE("Core::Containers::TToken::Copy constructor", "[Core][Containers][Smoke]")
+TEST_CASE("Core::Containers::TToken::Copy constructor", "[Core][Containers][TToken]")
 {
 	int32Token::Reset();
 	{
@@ -53,7 +53,7 @@ TEST_CASE("Core::Containers::TToken::Copy constructor", "[Core][Containers][Smok
 	CHECK(int32Token::NumMoveAssignmentCalls() == 0);
 }
 
-TEST_CASE("Core::Containers::TToken::Move constructor", "[Core][Containers][Smoke]")
+TEST_CASE("Core::Containers::TToken::Move constructor", "[Core][Containers][TToken]")
 {
 	int32Token::Reset();
 	{
@@ -71,7 +71,7 @@ TEST_CASE("Core::Containers::TToken::Move constructor", "[Core][Containers][Smok
 	CHECK(int32Token::NumMoveAssignmentCalls() == 0);
 }
 
-TEST_CASE("Core::Containers::TToken::Copy assignment", "[Core][Containers][Smoke]")
+TEST_CASE("Core::Containers::TToken::Copy assignment", "[Core][Containers][TToken]")
 {
 	int32Token::Reset();
 	{
@@ -90,7 +90,7 @@ TEST_CASE("Core::Containers::TToken::Copy assignment", "[Core][Containers][Smoke
 	CHECK(int32Token::NumMoveAssignmentCalls() == 0);
 }
 
-TEST_CASE("Core::Containers::TToken::Move assignment", "[Core][Containers][Smoke]")
+TEST_CASE("Core::Containers::TToken::Move assignment", "[Core][Containers][TToken]")
 {
 	int32Token::Reset();
 	{
