M - Symetric Task switching support for the Scheduler. This will allow Tasks to return an optional FTask* which is than used by the Scheduler to directly switch to and execute immediately regardless of it's priority.

diff --git a/Engine/Source/Runtime/Core/Private/Async/Fundamental/Scheduler.cpp b/Engine/Source/Runtime/Core/Private/Async/Fundamental/Scheduler.cpp
--- a/Engine/Source/Runtime/Core/Private/Async/Fundamental/Scheduler.cpp
+++ b/Engine/Source/Runtime/Core/Private/Async/Fundamental/Scheduler.cpp
@@ -151,15 +151,15 @@ namespace LowLevelTasks
 		}
 	}
 
-	void FScheduler::ExecuteTask(FTask& Task)
+	inline void FScheduler::ExecuteTask(FTask*& InOutTask)
 	{
 		FTask* ParentTask = FTask::ActiveTask;
-		FTask::ActiveTask = &Task;
+		FTask::ActiveTask = InOutTask;
 
-		if (!Task.IsBackgroundTask())
+		if (!InOutTask->IsBackgroundTask())
 		{
 			TRACE_CPUPROFILER_EVENT_SCOPE(ExecuteForegroundTask);
-			Task.ExecuteTask();
+			InOutTask = InOutTask->ExecuteTask();
 		}
 		else
 		{
@@ -181,7 +181,7 @@ namespace LowLevelTasks
 
 			{
 				TRACE_CPUPROFILER_EVENT_SCOPE(ExecuteBackgroundTask);
-				Task.ExecuteTask();
+				InOutTask = InOutTask->ExecuteTask();
 			}
 
 			if (!bSkipPriorityChange)
@@ -216,7 +216,14 @@ namespace LowLevelTasks
 			{
 				for (FTask* Task = QueueRegistry.Dequeue(); Task != nullptr; Task = QueueRegistry.Dequeue())
 				{
-					ExecuteTask(*Task);
+					while(Task)
+					{
+						ExecuteTask(Task);
+						if(Task)
+						{
+							verifySlow(Task->TryPrepareLaunch());
+						}
+					}
 				}
 			}
 		}
@@ -266,7 +273,15 @@ namespace LowLevelTasks
 		}
 		else
 		{
-			ExecuteTask(Task);
+			FTask* TaskPtr = &Task;
+			while(TaskPtr)
+			{
+				ExecuteTask(TaskPtr);
+				if(TaskPtr)
+				{
+					verifySlow(TaskPtr->TryPrepareLaunch());
+				}
+			}
 		}
 	}
 
@@ -299,33 +314,57 @@ namespace LowLevelTasks
 	template<FTask* (FSchedulerTls::FLocalQueueType::*DequeueFunction)(bool, bool), bool bIsBusyWaiting>
 	bool FScheduler::TryExecuteTaskFrom(FSchedulerTls::FLocalQueueType* Queue, FSchedulerTls::FQueueRegistry::FOutOfWork& OutOfWork, bool bPermitBackgroundWork, bool bDisableThrottleStealing)
 	{
-		for(int i = 0; i < 2; i++) // one retry if we pick up a task that cannot used during busy waiting.
+		bool AnyExecuted = false;
+
+		FTask* Task = (Queue->*DequeueFunction)(bPermitBackgroundWork, bDisableThrottleStealing);
+		while (Task)
 		{
-			FTask* Task = (Queue->*DequeueFunction)(bPermitBackgroundWork, bDisableThrottleStealing);
-			if (Task)
-			{	
+			if constexpr (bIsBusyWaiting)
+			{
 				FTask::FInitData InitData = Task->GetInitData();
 				bool bAllowBusyWaiting = EnumHasAnyFlags(InitData.Flags, ETaskFlags::AllowBusyWaiting);
-				if (bIsBusyWaiting && !bAllowBusyWaiting)
+				
+				if (!bAllowBusyWaiting || AnyExecuted) 
 				{
+					//either the task is not allowed during busy waiting or we have a  
+					//symetric switching task that we do not want to execute during BusyWaiting
+					//in either case we requeue and try again or we exit early with success.
 					QueueRegistry.Enqueue(Task, uint32(InitData.Priority));
-					continue;
-				}
-				if (OutOfWork.Stop())
-				{
-					//if we are ramping up work again we start waking up workers that might not have been woken while a lot of work might have been queued
-					if (!WakeUpWorker(bPermitBackgroundWork) && !FSchedulerTls::IsBackgroundWorker())
+
+					if(AnyExecuted)
 					{
-						WakeUpWorker(!bPermitBackgroundWork);
+						return AnyExecuted; 
+					}
+
+					Task = (Queue->*DequeueFunction)(bPermitBackgroundWork, bDisableThrottleStealing);
+					if(Task == nullptr)
+					{
+						return false;
 					}
 				}
+			}
+			else
+			{
+				checkSlow(FTask::ActiveTask == nullptr);
+			}
 
-				ExecuteTask(*Task);
-				return true;
+			if (OutOfWork.Stop())
+			{
+				//if we are ramping up work again we start waking up workers that might not have been woken while a lot of work might have been queued
+				if (!WakeUpWorker(bPermitBackgroundWork) && !FSchedulerTls::IsBackgroundWorker())
+				{
+					WakeUpWorker(!bPermitBackgroundWork);
+				}
+			}
+			
+			ExecuteTask(Task);
+			if(Task)
+			{
+				verifySlow(Task->TryPrepareLaunch());
 			}
-			return false;
+			AnyExecuted = true;
 		}
-		return false;
+		return AnyExecuted;
 	}
 
 	void FScheduler::WorkerMain(FSleepEvent* WorkerEvent, FSchedulerTls::FLocalQueueType* ExternalWorkerLocalQueue, uint32 WaitCycles, bool bPermitBackgroundWork)
diff --git a/Engine/Source/Runtime/Core/Private/Async/Fundamental/UnitTests.cpp b/Engine/Source/Runtime/Core/Private/Async/Fundamental/UnitTests.cpp
--- a/Engine/Source/Runtime/Core/Private/Async/Fundamental/UnitTests.cpp
+++ b/Engine/Source/Runtime/Core/Private/Async/Fundamental/UnitTests.cpp
@@ -80,7 +80,7 @@ namespace Tasks2Tests
 			}
 
 			FTask SignalTask;
-			SignalTask.Init(TEXT("Empty Signal Task"), [](){}, [&Tasks]()
+			SignalTask.Init(TEXT("Empty Signal Task"), [&Tasks]()
 			{
 				for (int i = 0; i < NumTasks; i++)
 				{
@@ -204,7 +204,7 @@ namespace Tasks2Tests
 			}
 
 			FTask SignalTask;
-			SignalTask.Init(TEXT("Signal Task"), [](){}, [&Tasks]()
+			SignalTask.Init(TEXT("Signal Task"), [&Tasks]()
 			{
 				for (int i = 0; i < NumTasks; i++)
 				{
@@ -380,7 +380,7 @@ namespace Tasks2Tests
 			}
 
 			FTask SignalTask;
-			SignalTask.Init(TEXT("Signal Task"), [](){}, [&Tasks]()
+			SignalTask.Init(TEXT("Signal Task"), [&Tasks]()
 			{
 				for (int i = 0; i < NumTasks; i++)
 				{
@@ -523,10 +523,10 @@ namespace Tasks2Tests
 				using FTaskHandle = TSharedPtr<FTask, ESPMode::ThreadSafe>;
 				FTaskHandle TaskHandle = MakeShared<FTask, ESPMode::ThreadSafe>();
 
-				TaskHandle->Init(TEXT("Refcounted Test"), [&TestValue]()
+				TaskHandle->Init(TEXT("Refcounted Test"), [&TestValue, TaskHandle]()
 				{
 					TestValue = 42;
-				}, [TaskHandle](){});
+				});
 				TryLaunch(*TaskHandle);
 
 				bool WasCanceled = TestCancel && TaskHandle->TryCancel();
@@ -583,17 +583,33 @@ namespace Tasks2Tests
 			using FTaskHandle = TUniquePtr<FTask>;
 			FTaskHandle TaskHandle = MakeUnique<FTask>();
 			FTask& Task = *TaskHandle;
-			Task.Init(TEXT("Uniqueptr Test"), [TestValue]()
+			Task.Init(TEXT("Uniqueptr Test"), [TestValue, TaskHandle = MoveTemp(TaskHandle)]()
 			{
 				*TestValue = 42;
-			}, [TestValue, TaskHandle = MoveTemp(TaskHandle)]()
-			{
-				verify(*TestValue == 42);
 				delete TestValue;
 			});
 			TryLaunch(Task);
 		}
 
+		//symetric switch
+		{
+			uint32 TestValue = 1337;
+
+			FTask TaskA;
+			FTask TaskB;
+			TaskB.Init(TEXT("TaskB"), [&TestValue]()
+			{
+				TestValue = 42;
+			});
+			TaskA.Init(TEXT("TaskA"), [&TaskB]() -> FTask*
+			{
+				return &TaskB;
+			});
+			verify(TryLaunch(TaskA));
+			BusyWaitForTask(TaskB);
+
+			verify(TestValue == 42);
+		}
 		return true;
 	}
 
@@ -676,4 +692,4 @@ namespace Tasks2Tests
 	}
 }
 
-#endif
\ No newline at end of file
+#endif
diff --git a/Engine/Source/Runtime/Core/Public/Async/Fundamental/Scheduler.h b/Engine/Source/Runtime/Core/Public/Async/Fundamental/Scheduler.h
--- a/Engine/Source/Runtime/Core/Public/Async/Fundamental/Scheduler.h
+++ b/Engine/Source/Runtime/Core/Public/Async/Fundamental/Scheduler.h
@@ -184,7 +184,7 @@ namespace LowLevelTasks
 		~FScheduler();
 
 	private: 
-		void ExecuteTask(FTask& Task);
+		void ExecuteTask(FTask*& InOutTask);
 		TUniquePtr<FThread> CreateWorker(bool bPermitBackgroundWork = false, FThread::EForkable IsForkable = FThread::NonForkable, FSleepEvent* ExternalWorkerEvent = nullptr, FSchedulerTls::FLocalQueueType* ExternalWorkerLocalQueue = nullptr, EThreadPriority Priority = EThreadPriority::TPri_Normal, uint64 InAffinity = 0);
 		void WorkerMain(struct FSleepEvent* WorkerEvent, FSchedulerTls::FLocalQueueType* ExternalWorkerLocalQueue, uint32 WaitCycles, bool bPermitBackgroundWork);
 		CORE_API void LaunchInternal(FTask& Task, EQueuePreference QueuePreference, bool bWakeUpWorker);
diff --git a/Engine/Source/Runtime/Core/Public/Async/Fundamental/Task.h b/Engine/Source/Runtime/Core/Public/Async/Fundamental/Task.h
--- a/Engine/Source/Runtime/Core/Public/Async/Fundamental/Task.h
+++ b/Engine/Source/Runtime/Core/Public/Async/Fundamental/Task.h
@@ -146,7 +146,7 @@ namespace LowLevelTasks
 
 	/*
 	* Generic implementation of a Deleter, it often comes up that one has to call a function to cleanup after a Task finished
-	* this can be done by capturing a TDeleter like this in the lambda of the Continuation: [Deleter(LowLevelTasks::TDeleter<Type, &Type::DeleteFunction>(value))](){}
+	* this can be done by capturing a TDeleter like so: [Deleter(LowLevelTasks::TDeleter<Type, &Type::DeleteFunction>(value))](){}
 	*/
 	template<typename Type, void (Type::*DeleteFunction)()>
 	class TDeleter
@@ -164,11 +164,17 @@ namespace LowLevelTasks
 			Other.Value = nullptr;
 		}
 
+		UE_DEPRECATED(5.1, "TDeleter::GetValue() has been deprecated. Please update to TDeleter::operator->().")
 		inline Type* GetValue() const
 		{
 			return Value;
 		}
 
+		inline Type* operator->() const
+		{
+			return Value;
+		}
+
 		inline ~TDeleter() 
 		{
 			if(Value)
@@ -293,7 +299,7 @@ namespace LowLevelTasks
 		};
 
 	private:
-		using FTaskDelegate = TTaskDelegate<void(bool), LOWLEVEL_TASK_SIZE - sizeof(FPackedData) - sizeof(void*)>;
+		using FTaskDelegate = TTaskDelegate<FTask*(bool), LOWLEVEL_TASK_SIZE - sizeof(FPackedData) - sizeof(void*)>;
 		FTaskDelegate Runnable;
 		mutable void* UserData = nullptr;
 		FPackedDataAtomic PackedData;
@@ -378,26 +384,32 @@ namespace LowLevelTasks
 		* try to expedite the task if succeded it will run immediately 
 		* but it will not set the completed state until the scheduler has executed it, because the scheduler still holds a reference.
 		* to check for completion in the context of expediting use WasExpedited. The TaskHandle canot be reused until IsCompleted returns true.
+		* @param Continuation: optional Continuation that needs to be executed or scheduled by the caller (can only be non null if the operation returned true)
 		*/
 		inline bool TryExpedite();
+		inline bool TryExpedite(FTask*& Continuation);
 
 		/*
 		* try to execute the task if it has not been launched yet the task will execute immediately.
+		* @param Continuation: optional Continuation that needs to be executed or scheduled by the caller (can only be non null if the operation returned true)
 		*/
 		inline bool TryExecute();
-
-		template<typename TRunnable, typename TContinuation>
-		inline void Init(const TCHAR* InDebugName, ETaskPriority InPriority, TRunnable&& InRunnable, TContinuation&& InContinuation, ETaskFlags Flags = ETaskFlags::DefaultFlags);
+		inline bool TryExecute(FTask*& Continuation);
 
 		template<typename TRunnable>
 		inline void Init(const TCHAR* InDebugName, ETaskPriority InPriority, TRunnable&& InRunnable, ETaskFlags Flags = ETaskFlags::DefaultFlags);
 
-		template<typename TRunnable, typename TContinuation>
-		inline void Init(const TCHAR* InDebugName, TRunnable&& InRunnable, TContinuation&& InContinuation, ETaskFlags Flags = ETaskFlags::DefaultFlags);
-
 		template<typename TRunnable>
 		inline void Init(const TCHAR* InDebugName, TRunnable&& InRunnable, ETaskFlags Flags = ETaskFlags::DefaultFlags);
 
+		template<typename TRunnable, typename TContinuation>
+		UE_DEPRECATED(5.1, "FTask::Init() has been deprecated. Please update to using a TDeleter.")
+		inline void Init(const TCHAR* InDebugName, ETaskPriority InPriority, TRunnable&& InRunnable, TContinuation&& InContinuation, ETaskFlags Flags = ETaskFlags::DefaultFlags);
+
+		template<typename TRunnable, typename TContinuation>
+		UE_DEPRECATED(5.1, "FTask::Init() has been deprecated. Please update to using a TDeleter.")
+		inline void Init(const TCHAR* InDebugName, TRunnable&& InRunnable, TContinuation&& InContinuation, ETaskFlags Flags = ETaskFlags::DefaultFlags);
+
 		inline const TCHAR* GetDebugName() const;
 		inline ETaskPriority GetPriority() const;
 		inline bool IsBackgroundTask() const;
@@ -429,7 +441,8 @@ namespace LowLevelTasks
 		//after calling this function the task can be considered dead
 		template<bool bIsExpeditingThread>
 		inline void TryFinish();
-		inline void ExecuteTask();
+
+		inline FTask* ExecuteTask();
 		inline void InheritParentData(ETaskPriority& Priority);
 	};
 
@@ -463,34 +476,53 @@ namespace LowLevelTasks
 		}
 	}
 
-	template<typename TRunnable, typename TContinuation>
-	inline void FTask::Init(const TCHAR* InDebugName, ETaskPriority InPriority, TRunnable&& InRunnable, TContinuation&& InContinuation, ETaskFlags Flags)
+	template<typename TRunnable>
+	inline void FTask::Init(const TCHAR* InDebugName, ETaskPriority InPriority, TRunnable&& InRunnable, ETaskFlags Flags)
 	{
 		checkf(IsCompleted(), TEXT("State: %d"), PackedData.load(std::memory_order_relaxed).GetState());
 		checkSlow(!Runnable.IsSet());
-		Runnable = [LocalRunnable = Forward<TRunnable>(InRunnable), LocalContinuation = Forward<TContinuation>(InContinuation)](const bool NotCanceled)
+		
+		//if the Runnable returns an FTask* than enable symetric switching
+		if constexpr (TIsSame<FTask*, decltype(UE::Core::Private::IsInvocable::DeclVal<TRunnable>()())>::Value)
 		{
-			if (NotCanceled)
+			Runnable = [LocalRunnable = Forward<TRunnable>(InRunnable)](const bool bNotCanceled) -> FTask*
 			{
-				LocalRunnable();
-			}
-			LocalContinuation();
-		};
+				if (bNotCanceled)
+				{
+					FTask* Task = LocalRunnable();
+					return Task;
+				}
+				return nullptr;
+			};
+		}
+		else
+		{
+			Runnable = [LocalRunnable = Forward<TRunnable>(InRunnable)](const bool bNotCanceled) -> FTask*
+			{
+				if (bNotCanceled)
+				{
+					LocalRunnable();
+				}
+				return nullptr;
+			};
+		}
 		InheritParentData(InPriority);
 		PackedData.store(FPackedData(InDebugName, InPriority, ETaskState::Ready, Flags), std::memory_order_release);
 	}
 
-	template<typename TRunnable>
-	inline void FTask::Init(const TCHAR* InDebugName, ETaskPriority InPriority, TRunnable&& InRunnable, ETaskFlags Flags)
+	template<typename TRunnable, typename TContinuation>
+	inline void FTask::Init(const TCHAR* InDebugName, ETaskPriority InPriority, TRunnable&& InRunnable, TContinuation&& InContinuation, ETaskFlags Flags)
 	{
 		checkf(IsCompleted(), TEXT("State: %d"), PackedData.load(std::memory_order_relaxed).GetState());
 		checkSlow(!Runnable.IsSet());
-		Runnable = [LocalRunnable = Forward<TRunnable>(InRunnable)](const bool NotCanceled)
+		Runnable = [LocalRunnable = Forward<TRunnable>(InRunnable), LocalContinuation = Forward<TContinuation>(InContinuation)](const bool NotCanceled)
 		{
 			if (NotCanceled)
 			{
 				LocalRunnable();
 			}
+			LocalContinuation();
+			return nullptr;
 		};
 		InheritParentData(InPriority);
 		PackedData.store(FPackedData(InDebugName, InPriority, ETaskState::Ready, Flags), std::memory_order_release);
@@ -533,7 +565,7 @@ namespace LowLevelTasks
 
 		if(bTryLaunchOnSuccess && WasCanceled && TryPrepareLaunch())
 		{
-			ExecuteTask();
+			verifySlow(ExecuteTask() == nullptr);
 			return true;
 		}
 		return WasCanceled;
@@ -554,16 +586,24 @@ namespace LowLevelTasks
 			|| PackedData.compare_exchange_strong(CanceledState, FPackedData(LocalPackedData, ETaskState::Scheduled), std::memory_order_release);
 	}
 
-	inline bool FTask::TryExecute()
+	inline bool FTask::TryExecute(FTask*& OutContinuation)
 	{
 		if(TryPrepareLaunch())
 		{
-			ExecuteTask();
+			OutContinuation  = ExecuteTask();
 			return true;
 		}
 		return false;
 	}
 
+	inline bool FTask::TryExecute()
+	{
+		FTask* Continuation = nullptr;
+		bool Result = TryExecute(Continuation);
+		checkSlow(Continuation == nullptr);
+		return Result;
+	}
+
 	template<bool bIsExpeditingThread>
 	inline void FTask::TryFinish()
 	{
@@ -583,28 +623,37 @@ namespace LowLevelTasks
 		}
 	}
 
-	inline bool FTask::TryExpedite()
+	inline bool FTask::TryExpedite(FTask*& OutContinuation)
 	{
 		FPackedData LocalPackedData = PackedData.load(std::memory_order_relaxed);
 		FPackedData ScheduledState(LocalPackedData, ETaskState::Scheduled);
 		if(PackedData.compare_exchange_strong(ScheduledState, FPackedData(LocalPackedData, ETaskState::Running), std::memory_order_acquire))
 		{
-			Runnable(true);
+			OutContinuation  = Runnable(true);
 			TryFinish<true>();
 			return true;
 		}
 		return false;
 	}
 
-	inline void FTask::ExecuteTask()
+	inline bool FTask::TryExpedite()
+	{
+		FTask* Continuation = nullptr;
+		bool Result = TryExpedite(Continuation);
+		checkSlow(Continuation == nullptr);
+		return Result;
+	}
+
+	inline FTask* FTask::ExecuteTask()
 	{
 		ETaskState PreviousState = PackedData.fetch_or(ETaskState::RunningFlag, std::memory_order_acquire);
 		checkSlow(EnumHasAnyFlags(PreviousState, ETaskState::ScheduledFlag));
 
+		FTask* Continuation = nullptr;
 		if(!EnumHasAnyFlags(PreviousState, ETaskState::RunningFlag)) //we are running or canceled
 		{
 			FTaskDelegate LocalRunnable;
-			Runnable.CallAndMove(LocalRunnable, !EnumHasAnyFlags(PreviousState, ETaskState::CanceledFlag));
+			Continuation = Runnable.CallAndMove(LocalRunnable, !EnumHasAnyFlags(PreviousState, ETaskState::CanceledFlag));
 			//do not access the task again after this call
 			//as by defitition the task can be considered dead
 			PreviousState = PackedData.fetch_or(ETaskState::CompletedFlag, std::memory_order_seq_cst);
@@ -615,6 +664,8 @@ namespace LowLevelTasks
 			checkSlow(PreviousState == ETaskState::Running || PreviousState == ETaskState::Expediting || PreviousState == ETaskState::Expedited);
 			TryFinish<false>();
 		}
+
+		return Continuation;
 	}
 
 	inline const TCHAR* FTask::GetDebugName() const
diff --git a/Engine/Source/Runtime/Core/Public/Misc/QueuedThreadPoolWrapper.h b/Engine/Source/Runtime/Core/Public/Misc/QueuedThreadPoolWrapper.h
--- a/Engine/Source/Runtime/Core/Public/Misc/QueuedThreadPoolWrapper.h
+++ b/Engine/Source/Runtime/Core/Public/Misc/QueuedThreadPoolWrapper.h
@@ -319,6 +319,13 @@ private:
 		}
 	}
 
+	void FinalizeExecution()
+	{
+		TaskCount.fetch_sub(1, std::memory_order_release);
+		bool bWakeUpWorker = false;
+		ScheduleTasks(bWakeUpWorker);
+	}
+
 	void AddQueuedWork(IQueuedWork* InQueuedWork, EQueuedWorkPriority InPriority = EQueuedWorkPriority::Normal) override
 	{
 		check(bIsExiting == false);
@@ -331,19 +338,11 @@ private:
 		LowLevelTasks::ETaskPriority TaskPriority = TaskPriorityMapper[int(Priority)];
 		LowLevelTasks::ETaskFlags Flags = (InQueuedWork->GetQueuedWorkFlags() & EQueuedWorkFlags::DoNotRunInsideBusyWait) == EQueuedWorkFlags::None ? LowLevelTasks::ETaskFlags::DefaultFlags : (LowLevelTasks::ETaskFlags::DefaultFlags & ~LowLevelTasks::ETaskFlags::AllowCancellation);
 
-		QueuedWorkInternalData->Task.Init(TEXT("FQueuedLowLevelThreadPoolTask"), TaskPriority, [InQueuedWork]
+		QueuedWorkInternalData->Task.Init(TEXT("FQueuedLowLevelThreadPoolTask"), TaskPriority, [InQueuedWork, InternalData = InQueuedWork->InternalData, Deleter = LowLevelTasks::TDeleter<FQueuedLowLevelThreadPool, &FQueuedLowLevelThreadPool::FinalizeExecution>{ this }]
 		{
 			FMemMark Mark(FMemStack::Get());
 			InQueuedWork->DoThreadedWork();
-		},
-		[this, InternalData = InQueuedWork->InternalData]()
-		{
-			TaskCount.fetch_sub(1, std::memory_order_release);
-			bool bWakeUpWorker = false;
-			ScheduleTasks(bWakeUpWorker);
-		},
-		Flags
-		);
+		}, Flags);
 
 		if (!bIsPaused)
 		{
