Math macro prefix fixes in non-vectorized transform.

diff --git a/Engine/Source/Runtime/Core/Public/Math/TransformNonVectorized.h b/Engine/Source/Runtime/Core/Public/Math/TransformNonVectorized.h
--- a/Engine/Source/Runtime/Core/Public/Math/TransformNonVectorized.h
+++ b/Engine/Source/Runtime/Core/Public/Math/TransformNonVectorized.h
@@ -644,17 +644,17 @@ public:
 
 private:
 
-	FORCEINLINE bool Private_RotationEquals(const TQuat<T>& InRotation, const FReal Tolerance = KINDA_SMALL_NUMBER) const
+	FORCEINLINE bool Private_RotationEquals(const TQuat<T>& InRotation, const FReal Tolerance = UE_KINDA_SMALL_NUMBER) const
 	{
 		return Rotation.Equals(InRotation, Tolerance);
 	}
 
-	FORCEINLINE bool Private_TranslationEquals(const TVector<T>& InTranslation, const FReal Tolerance = KINDA_SMALL_NUMBER) const
+	FORCEINLINE bool Private_TranslationEquals(const TVector<T>& InTranslation, const FReal Tolerance = UE_KINDA_SMALL_NUMBER) const
 	{
 		return Translation.Equals(InTranslation, Tolerance);
 	}
 
-	FORCEINLINE bool Private_Scale3DEquals(const TVector<T>& InScale3D, const FReal Tolerance = KINDA_SMALL_NUMBER) const
+	FORCEINLINE bool Private_Scale3DEquals(const TVector<T>& InScale3D, const FReal Tolerance = UE_KINDA_SMALL_NUMBER) const
 	{
 		return Scale3D.Equals(InScale3D, Tolerance);
 	}
@@ -662,19 +662,19 @@ private:
 public:
 
 	// Test if A's rotation equals B's rotation, within a tolerance. Preferred over "A.GetRotation().Equals(B.GetRotation())" because it is faster on some platforms.
-	FORCEINLINE static bool AreRotationsEqual(const TTransform<T>& A, const TTransform<T>& B, FReal Tolerance = KINDA_SMALL_NUMBER)
+	FORCEINLINE static bool AreRotationsEqual(const TTransform<T>& A, const TTransform<T>& B, FReal Tolerance = UE_KINDA_SMALL_NUMBER)
 	{
 		return A.Private_RotationEquals(B.Rotation, Tolerance);
 	}
 
 	// Test if A's translation equals B's translation, within a tolerance. Preferred over "A.GetTranslation().Equals(B.GetTranslation())" because it is faster on some platforms.
-	FORCEINLINE static bool AreTranslationsEqual(const TTransform<T>& A, const TTransform<T>& B, FReal Tolerance = KINDA_SMALL_NUMBER)
+	FORCEINLINE static bool AreTranslationsEqual(const TTransform<T>& A, const TTransform<T>& B, FReal Tolerance = UE_KINDA_SMALL_NUMBER)
 	{
 		return A.Private_TranslationEquals(B.Translation, Tolerance);
 	}
 
 	// Test if A's scale equals B's scale, within a tolerance. Preferred over "A.GetScale3D().Equals(B.GetScale3D())" because it is faster on some platforms.
-	FORCEINLINE static bool AreScale3DsEqual(const TTransform<T>& A, const TTransform<T>& B, FReal Tolerance = KINDA_SMALL_NUMBER)
+	FORCEINLINE static bool AreScale3DsEqual(const TTransform<T>& A, const TTransform<T>& B, FReal Tolerance = UE_KINDA_SMALL_NUMBER)
 	{
 		return A.Private_Scale3DEquals(B.Scale3D, Tolerance);
 	}
@@ -682,25 +682,25 @@ public:
 
 
 	// Test if this Transform's rotation equals another's rotation, within a tolerance. Preferred over "GetRotation().Equals(Other.GetRotation())" because it is faster on some platforms.
-	FORCEINLINE bool RotationEquals(const TTransform<T>& Other, FReal Tolerance = KINDA_SMALL_NUMBER) const
+	FORCEINLINE bool RotationEquals(const TTransform<T>& Other, FReal Tolerance = UE_KINDA_SMALL_NUMBER) const
 	{
 		return AreRotationsEqual(*this, Other, Tolerance);
 	}
 
 	// Test if this Transform's translation equals another's translation, within a tolerance. Preferred over "GetTranslation().Equals(Other.GetTranslation())" because it is faster on some platforms.
-	FORCEINLINE bool TranslationEquals(const TTransform<T>& Other, FReal Tolerance = KINDA_SMALL_NUMBER) const
+	FORCEINLINE bool TranslationEquals(const TTransform<T>& Other, FReal Tolerance = UE_KINDA_SMALL_NUMBER) const
 	{
 		return AreTranslationsEqual(*this, Other, Tolerance);
 	}
 
 	// Test if this Transform's scale equals another's scale, within a tolerance. Preferred over "GetScale3D().Equals(Other.GetScale3D())" because it is faster on some platforms.
-	FORCEINLINE bool Scale3DEquals(const TTransform<T>& Other, FReal Tolerance = KINDA_SMALL_NUMBER) const
+	FORCEINLINE bool Scale3DEquals(const TTransform<T>& Other, FReal Tolerance = UE_KINDA_SMALL_NUMBER) const
 	{
 		return AreScale3DsEqual(*this, Other, Tolerance);
 	}
 
 	// Test if all components of the transforms are equal, within a tolerance.
-	FORCEINLINE bool Equals(const TTransform<T>& Other, FReal Tolerance = KINDA_SMALL_NUMBER) const
+	FORCEINLINE bool Equals(const TTransform<T>& Other, FReal Tolerance = UE_KINDA_SMALL_NUMBER) const
 	{
 		return Private_TranslationEquals(Other.Translation, Tolerance) && Private_RotationEquals(Other.Rotation, Tolerance) && Private_Scale3DEquals(Other.Scale3D, Tolerance);
 	}
@@ -712,7 +712,7 @@ public:
 	}
 
 	// Test if rotation and translation components of the transforms are equal, within a tolerance.
-	FORCEINLINE bool EqualsNoScale(const TTransform<T>& Other, FReal Tolerance = KINDA_SMALL_NUMBER) const
+	FORCEINLINE bool EqualsNoScale(const TTransform<T>& Other, FReal Tolerance = UE_KINDA_SMALL_NUMBER) const
 	{
 		return Private_TranslationEquals(Other.Translation, Tolerance) && Private_RotationEquals(Other.Rotation, Tolerance);
 	}
@@ -899,7 +899,7 @@ public:
 	FORCEINLINE void Accumulate(const TTransform<T>& SourceAtom)
 	{
 		// Add ref pose relative animation to base animation, only if rotation is significant.
-		if (FMath::Square(SourceAtom.Rotation.W) < 1.f - DELTA * DELTA)
+		if (FMath::Square(SourceAtom.Rotation.W) < 1.f - UE_DELTA * UE_DELTA)
 		{
 			Rotation = SourceAtom.Rotation * Rotation;
 		}
@@ -929,7 +929,7 @@ public:
 		TTransform<T> SourceAtom(Atom * BlendWeight);
 
 		// Add ref pose relative animation to base animation, only if rotation is significant.
-		if (FMath::Square(SourceAtom.Rotation.W) < 1.f - DELTA * DELTA)
+		if (FMath::Square(SourceAtom.Rotation.W) < 1.f - UE_DELTA * UE_DELTA)
 		{
 			Rotation = SourceAtom.Rotation * Rotation;
 		}
@@ -1001,7 +1001,7 @@ public:
 		TTransform<T> SourceAtom(Atom * BlendWeight);
 
 		// Add ref pose relative animation to base animation, only if rotation is significant.
-		if (FMath::Square(SourceAtom.Rotation.W) < 1.f - DELTA * DELTA)
+		if (FMath::Square(SourceAtom.Rotation.W) < 1.f - UE_DELTA * UE_DELTA)
 		{
 			Rotation = SourceAtom.Rotation * Rotation;
 		}
@@ -1074,7 +1074,7 @@ public:
 		}
 
 		// Add ref pose relative animation to base animation, only if rotation is significant.
-		if (FMath::Square(DeltaAtom.Rotation.W) < 1.f - DELTA * DELTA)
+		if (FMath::Square(DeltaAtom.Rotation.W) < 1.f - UE_DELTA * UE_DELTA)
 		{
 			FinalAtom.Rotation = DeltaAtom.Rotation * FinalAtom.Rotation;
 		}
