Make FGenericPlatformMisc::CreateGuid less completely broken in case it is ever used
race on IncrementCounter meant it could easily create the same guid on different threads
should be rarely used because all platforms have their own implementation and only use this as a fallback

diff --git a/Engine/Source/Runtime/Core/Private/GenericPlatform/GenericPlatformMisc.cpp b/Engine/Source/Runtime/Core/Private/GenericPlatform/GenericPlatformMisc.cpp
--- a/Engine/Source/Runtime/Core/Private/GenericPlatform/GenericPlatformMisc.cpp
+++ b/Engine/Source/Runtime/Core/Private/GenericPlatform/GenericPlatformMisc.cpp
@@ -972,17 +972,31 @@ void FGenericPlatformMisc:: ClipboardPaste(class FString& Dest)
 void FGenericPlatformMisc::CreateGuid(FGuid& Guid)
 {
 	QUICK_SCOPE_CYCLE_COUNTER(STAT_FGenericPlatformMisc_CreateGuid);
+	
+	// this is a pretty terrible Guid maker that has way less than 128 bits of randomness
+	// luckily this is not used, as every platform has an override
+	//   do not use this!
+	// all the timing used for "randomness" can easily be the same on two threads
+	// you wind up with 15 bits of randomness from the stdlib rand() call
+	//   the atomic IncrementCounter is crucial to ensure simultaneous calls to CreateGuid
+	//	 on different threads don't produce the same Guid here
 
-	static uint16 IncrementCounter = 0; 
+	// note the first calls to this are in static initializers, people use it to init globals
 
+	static std::atomic<uint32> IncrementCounter = 0; 
 	static FDateTime InitialDateTime;
 	static uint64 InitialCycleCounter;
 
 	FDateTime EstimatedCurrentDateTime;
 
-	if (IncrementCounter == 0)
+	uint32 SequentialBits = IncrementCounter.fetch_add(1);
+
+	if (SequentialBits == 0) // not a thread-safe init, but our first call is not threaded, so okay
 	{
 		// Hack: First Guid can be created prior to FPlatformTime::InitTiming(), so do it here.
+		//	InitTiming is done in static init in CoreGlobals.cpp
+		//	but CreateGuid is also used in static initializers, so we may come first
+		//  this can result in InitTiming being called twice which is not entirely benign
 		FPlatformTime::InitTiming();
 
 		// uses FPlatformTime::SystemTime()
@@ -998,10 +1012,16 @@ void FGenericPlatformMisc::CreateGuid(FGuid& Guid)
 		EstimatedCurrentDateTime = InitialDateTime + ElapsedTime;
 	}
 
-	uint32 SequentialBits = static_cast<uint32>(IncrementCounter++); // Add sequential bits to ensure sequentially generated guids are unique even if Cycles is wrong
-	uint32 RandBits = FMath::Rand() & 0xFFFF; // Add randomness to improve uniqueness across machines
+	uint32 RandBits = FMath::Rand();
+	
+	// bit rotate 16 : 
+	SequentialBits = ( SequentialBits << 16 ) | ( SequentialBits >> 16 );
+
+	Guid = FGuid(RandBits ^ SequentialBits, EstimatedCurrentDateTime.GetTicks() >> 32, EstimatedCurrentDateTime.GetTicks() & 0xffffffff, FPlatformTime::Cycles());
 
-	Guid = FGuid(RandBits | (SequentialBits << 16), EstimatedCurrentDateTime.GetTicks() >> 32, EstimatedCurrentDateTime.GetTicks() & 0xffffffff, FPlatformTime::Cycles());
+	// note: all the platform Guid makers do this but we do not :
+	//	Result[1] = (Result[1] & 0xffff0fff) | 0x00004000; // version 4
+	//	Result[2] = (Result[2] & 0x3fffffff) | 0x80000000; // variant 1
 }
 
 const TCHAR* LexToString( EAppReturnType::Type Value )
