M - Coroutines can change the TLS Thread Singleton, so we need to make sure that the value only cleans out TLS on destruction if it is still in TLS first.

diff --git a/Engine/Source/Runtime/Core/Public/HAL/ThreadSingleton.h b/Engine/Source/Runtime/Core/Public/HAL/ThreadSingleton.h
--- a/Engine/Source/Runtime/Core/Public/HAL/ThreadSingleton.h
+++ b/Engine/Source/Runtime/Core/Public/HAL/ThreadSingleton.h
@@ -71,7 +71,10 @@ protected:
 	{
 		// Clean the dangling pointer from the TLS.
 		check(GetTlsSlot() != 0xFFFFFFFF);
-		FPlatformTLS::SetTlsValue(GetTlsSlot(), nullptr);
+		if(((FTlsAutoCleanup*)FPlatformTLS::GetTlsValue(GetTlsSlot())) == static_cast<FTlsAutoCleanup*>(this))
+		{
+			FPlatformTLS::SetTlsValue(GetTlsSlot(), nullptr);
+		}
 	}
 
 	/**
