Unification of FMallocTBB::TryMalloc() and FMallocTBB::TryRealloc() alignment behavior.

diff --git a/Engine/Source/Runtime/Core/Private/HAL/MallocTBB.cpp b/Engine/Source/Runtime/Core/Private/HAL/MallocTBB.cpp
--- a/Engine/Source/Runtime/Core/Private/HAL/MallocTBB.cpp
+++ b/Engine/Source/Runtime/Core/Private/HAL/MallocTBB.cpp
@@ -118,7 +118,9 @@ void* FMallocTBB::TryRealloc( void* Ptr, SIZE_T NewSize, uint32 Alignment )
 	}
 	else
 	{
-		NewPtr = scalable_realloc(Ptr, NewSize);
+		// Fulfill the promise of DEFAULT_ALIGNMENT, which aligns 16-byte or larger structures to 16 bytes,
+		// while TBB aligns to 8 by default.
+		NewPtr = scalable_aligned_realloc(Ptr, NewSize, NewSize >= 16 ? (uint32)16 : (uint32)8);
 	}
 #endif
 #if UE_BUILD_DEBUG
