Make sure we have FNameEntry symbol preserved so we can print FNames correctly when debugging

diff --git a/Engine/Source/Runtime/Core/Private/Misc/CoreGlobals.cpp b/Engine/Source/Runtime/Core/Private/Misc/CoreGlobals.cpp
--- a/Engine/Source/Runtime/Core/Private/Misc/CoreGlobals.cpp
+++ b/Engine/Source/Runtime/Core/Private/Misc/CoreGlobals.cpp
@@ -248,6 +248,9 @@ uint8** CORE_API GNameBlocksDebug = FNameDebugVisualizer::GetBlocks();
 FChunkedFixedUObjectArray*& CORE_API GObjectArrayForDebugVisualizers = GCoreObjectArrayForDebugVisualizers;
 UE::ObjectPath::Private::FStoredObjectPath*& GComplexObjectPathDebug = GCoreComplexObjectPathDebug;
 FObjectHandlePackageDebugData*& CORE_API GObjectHandlePackageDebug = GCoreObjectHandlePackageDebug;
+
+// Create a global to make sure FNameEntry symbol is around which is needed to dump FName info when debugging
+FNameEntry* CORE_API GFNameEntryPreserveSymbol = nullptr;
 #endif
 
 /** Game name, used for base game directory and ini among other things										*/
