- Add TEnumFlagsRange which allows iteration over each flag set in an enum flags variable.

diff --git a/Engine/Source/Runtime/Core/Public/Misc/EnumRange.h b/Engine/Source/Runtime/Core/Public/Misc/EnumRange.h
--- a/Engine/Source/Runtime/Core/Public/Misc/EnumRange.h
+++ b/Engine/Source/Runtime/Core/Public/Misc/EnumRange.h
@@ -171,6 +171,53 @@ namespace NEnumRangePrivate
 	};
 }
 
+namespace UE::EnumFlags::Private
+{
+	template <typename EnumType>
+	struct TIterator
+	{
+		typedef __underlying_type(EnumType) IntType;
+
+		FORCEINLINE explicit TIterator(EnumType InFlags)
+			: Flags(IntType(InFlags))
+		{
+		}
+
+		FORCEINLINE TIterator& operator++()
+		{
+			const IntType PoppedBit = Flags & (~Flags + 1);
+			Flags ^= PoppedBit;
+			return *this;
+		}
+
+		FORCEINLINE EnumType operator*() const
+		{
+			const IntType Result = Flags & (~Flags + 1); 
+			return (EnumType)Result;
+		}
+
+	private:
+		IntType Flags;
+
+		FORCEINLINE friend bool operator!=(const TIterator& Lhs, const TIterator& Rhs)
+		{
+			return Lhs.Flags != Rhs.Flags;
+		}
+	};
+
+	template <typename EnumType>
+	struct TRange
+	{
+		explicit TRange(EnumType InFlags) : Flags(InFlags) {}
+
+		Private::TIterator<EnumType> begin() const { return Private::TIterator<EnumType>(Flags); }
+		Private::TIterator<EnumType> end() const { return Private::TIterator<EnumType>(EnumType(0)); }
+
+	private:
+		EnumType Flags;
+	};
+} // namespace UE::EnumFlags::Private
+
 /**
  * Range type for iterating over enum values.  Enums should define themselves as iterable by specifying
  * one of the ENUM_RANGE_* macros.
@@ -186,3 +233,21 @@ template <typename EnumType>
 struct TEnumRange : NEnumRangePrivate::TEnumRange_Impl<EnumType, NEnumRangePrivate::TEnumRangeTraits<EnumType>::RangeType>
 {
 };
+
+/**
+ * Make a range for iterating over set flags in a flags enum.
+ *
+ * Example:
+ *
+ * EFlagThing Flags = EFlagThing::A | EFlagThing::B;
+ * for (EFlagThing Flag : MakeFlagsRange(Flags))
+ * {
+ *     // Loop is run twice, once with Flag = EFlagThing::A, once with Flag = EFlagThing::B
+ *     ...
+ * }
+ **/
+template <typename EnumType>
+UE::EnumFlags::Private::TRange<EnumType> MakeFlagsRange(EnumType Flags)
+{
+	return UE::EnumFlags::Private::TRange<EnumType>(Flags);
+}
\ No newline at end of file
