Move a UE_LOG that happens before we check -logcmds to Verbose so we dont by default spam out to the log when we dont want to

Currently this is spamming for every ShaderCompilerWorker, the more cores the more spam.

diff --git a/Engine/Source/Runtime/Core/Public/HAL/IPlatformFileCachedWrapper.h b/Engine/Source/Runtime/Core/Public/HAL/IPlatformFileCachedWrapper.h
--- a/Engine/Source/Runtime/Core/Public/HAL/IPlatformFileCachedWrapper.h
+++ b/Engine/Source/Runtime/Core/Public/HAL/IPlatformFileCachedWrapper.h
@@ -297,7 +297,7 @@ public:
 				bResult = true;
 			}
 
-			UE_LOG(LogPlatformFile, Log, TEXT("%s cached read wrapper"), bResult ? TEXT("Using") : TEXT("Not using"));
+			UE_LOG(LogPlatformFile, Verbose, TEXT("%s cached read wrapper"), bResult ? TEXT("Using") : TEXT("Not using"));
 		}
 #endif
 		return bResult;
