Instrument concurrent linear allocator

Add memory tracing instrumentation to concurrent linear allocator.

diff --git a/Engine/Source/Runtime/Core/Public/Experimental/ConcurrentLinearAllocator.h b/Engine/Source/Runtime/Core/Public/Experimental/ConcurrentLinearAllocator.h
--- a/Engine/Source/Runtime/Core/Public/Experimental/ConcurrentLinearAllocator.h
+++ b/Engine/Source/Runtime/Core/Public/Experimental/ConcurrentLinearAllocator.h
@@ -6,6 +6,7 @@
 #include "HAL/MallocAnsi.h"
 #include "HAL/UnrealMemory.h"
 #include "HAL/LowLevelMemTracker.h"
+#include "ProfilingDebugging/MemoryTrace.h"
 #include "Templates/UniquePtr.h"
 #include "Templates/UnrealTypeTraits.h"
 
@@ -41,11 +42,15 @@ struct FOsAllocator
 	FORCEINLINE static void* Malloc(SIZE_T Size, uint32 Alignment)
 	{
 		//LLM_SCOPED_PAUSE_TRACKING(ELLMAllocType::System);
-		return FPlatformMemory::BinnedAllocFromOS(Size);
+		void* Pointer = FPlatformMemory::BinnedAllocFromOS(Size);
+		MemoryTrace_Alloc(uint64(Pointer), Size, Alignment);
+		MemoryTrace_MarkAllocAsHeap(uint64(Pointer), EMemoryTraceRootHeap::SystemMemory);
+		return Pointer;
 	}
 
 	FORCEINLINE static void Free(void* Pointer, SIZE_T Size)
 	{
+		MemoryTrace_Free(uint64(Pointer));
 		return FPlatformMemory::BinnedFreeToOS(Pointer, Size);
 	}
 };
@@ -296,6 +301,7 @@ public:
 				Header->NextAllocationPtr = AlignedOffset + Size;
 				Header->Num++;
 
+				MemoryTrace_Alloc(uint64(AlignedOffset), Size, Alignment);
 				LLM_IF_ENABLED( FLowLevelMemTracker::Get().OnLowLevelAlloc(ELLMTracker::Default, reinterpret_cast<void*>(AlignedOffset), Size, ELLMTag::LinearAllocator) );
 				return reinterpret_cast<void*>(AlignedOffset);
 			}
@@ -316,6 +322,7 @@ public:
 
 					checkSlow(LargeAlignedOffset + Size <= LargeHeader->NextAllocationPtr);
 
+					MemoryTrace_Alloc(uint64(LargeAlignedOffset), Size, Alignment);
 					LLM_IF_ENABLED( FLowLevelMemTracker::Get().OnLowLevelAlloc(ELLMTracker::Default, reinterpret_cast<void*>(LargeAlignedOffset), Size, ELLMTag::LinearAllocator) );
 					return reinterpret_cast<void*>(LargeAlignedOffset);
 				}
@@ -335,6 +342,7 @@ public:
 				new (AllocationHeader) FAllocationHeader(Header, Size);
 				ASAN_POISON_MEMORY_REGION( AllocationHeader, sizeof(FAllocationHeader) );
 
+				MemoryTrace_Alloc(uint64(AllocationHeader), Size + sizeof(FAllocationHeader), Alignment);
 				LLM_IF_ENABLED( FLowLevelMemTracker::Get().OnLowLevelAlloc(ELLMTracker::Default, AllocationHeader, Size + sizeof(FAllocationHeader), ELLMTag::LinearAllocator) );
 				return reinterpret_cast<void*>(AlignedOffset);
 			}
@@ -357,6 +365,7 @@ public:
 					FAllocationHeader* AllocationHeader = new (reinterpret_cast<FAllocationHeader*>(LargeAlignedOffset) - 1) FAllocationHeader(LargeHeader, Size);
 					ASAN_POISON_MEMORY_REGION( AllocationHeader, sizeof(FAllocationHeader) );
 
+					MemoryTrace_Alloc(uint64(AllocationHeader), Size + sizeof(FAllocationHeader), Alignment);
 					LLM_IF_ENABLED( FLowLevelMemTracker::Get().OnLowLevelAlloc(ELLMTracker::Default, AllocationHeader, Size + sizeof(FAllocationHeader), ELLMTag::LinearAllocator) );
 					return reinterpret_cast<void*>(LargeAlignedOffset);
 				}
@@ -388,6 +397,7 @@ public:
 		{
 			if constexpr (SupportsFastPath)
 			{
+				MemoryTrace_Free(uint64(Pointer));
 				LLM_IF_ENABLED(FLowLevelMemTracker::Get().OnLowLevelFree(ELLMTracker::Default, Pointer));
 				FBlockHeader* Header = reinterpret_cast<FBlockHeader*>(AlignDown(Pointer, BlockAllocationTag::BlockSize));
 
@@ -402,6 +412,7 @@ public:
 			{
 				FAllocationHeader* AllocationHeader = GetAllocationHeader(Pointer);
 				ASAN_UNPOISON_MEMORY_REGION( AllocationHeader, sizeof(FAllocationHeader) );
+				MemoryTrace_Free(uint64(AllocationHeader));
 				LLM_IF_ENABLED(FLowLevelMemTracker::Get().OnLowLevelFree(ELLMTracker::Default, AllocationHeader));
 				FBlockHeader* Header = AllocationHeader->GetBlockHeader();
 				ASAN_POISON_MEMORY_REGION( AllocationHeader, sizeof(FAllocationHeader) + AllocationHeader->GetAllocationSize() );
