Logging: Fixed missing callstacks for Windows crashes

The crash reporter thread was calling GLog->Panic() which would make it the panic thread and prevent subsequent logging from any other thread. Panic from ReportCrash to ensure that the thread handling the exception is the panic thread, which allows it to log the callstack after the crash reporter finishes.

#rnx

diff --git a/Engine/Source/Runtime/Core/Private/Windows/WindowsPlatformCrashContext.cpp b/Engine/Source/Runtime/Core/Private/Windows/WindowsPlatformCrashContext.cpp
--- a/Engine/Source/Runtime/Core/Private/Windows/WindowsPlatformCrashContext.cpp
+++ b/Engine/Source/Runtime/Core/Private/Windows/WindowsPlatformCrashContext.cpp
@@ -1408,13 +1408,6 @@ private:
 		// Then try run time crash processing and broadcast information about a crash.
 		FCoreDelegates::OnHandleSystemError.Broadcast();
 
-		if (GLog)
-		{
-			// Panic flush the logs to make sure there are no entries queued.
-			// This is will skip output devices that are not safe to flush during a crash, for example, the editor log.
-			GLog->Panic();
-		}
-		
 		// Get the default settings for the crash context
 		ECrashContextType Type = ECrashContextType::Crash;
 		const TCHAR* ErrorMessage = TEXT("Unhandled exception");
@@ -1661,6 +1654,12 @@ int32 ReportCrash( LPEXCEPTION_POINTERS ExceptionInfo )
 	// (Can be called the first time from the RenderThread, then a second time from the MainThread.)
 	if (GCrashReportingThread)
 	{
+		if (GLog)
+		{
+			// Panic before reporting the crash to make sure that the logs are flushed as thoroughly as possible.
+			GLog->Panic();
+		}
+
 		if (FPlatformAtomics::InterlockedIncrement(&ReportCrashCallCount) == 1)
 		{
 			GCrashReportingThread->OnCrashed(ExceptionInfo);
