Fix 2 Insights callstack tracing crashes when capturing on some platforms.
First one because due to RTC in EOS not having frame pointers causing a crash when walking down the stack. Limiting callstack tracing with a scope.

The second one because of limited stack size (4k) in a Curl Helper thread, causing a crash when walking down the stack because the callstack walker was allocating 4k on the stack. Doing it inplace to allocate only 2k.

#rnx

diff --git a/Engine/Plugins/Online/EOSShared/Source/EOSShared/Private/EOSSDKManager.cpp b/Engine/Plugins/Online/EOSShared/Source/EOSShared/Private/EOSSDKManager.cpp
--- a/Engine/Plugins/Online/EOSShared/Source/EOSShared/Private/EOSSDKManager.cpp
+++ b/Engine/Plugins/Online/EOSShared/Source/EOSShared/Private/EOSSDKManager.cpp
@@ -13,6 +13,7 @@
 #include "Misc/Paths.h"
 #include "Misc/ConfigCacheIni.h"
 #include "ProfilingDebugging/CsvProfiler.h"
+#include "ProfilingDebugging/CallstackTrace.h"
 #include "Stats/Stats.h"
 
 #include "CoreGlobals.h"
@@ -24,23 +25,42 @@
 #include "eos_sdk.h"
 #include "eos_version.h"
 
+#ifndef EOS_TRACE_MALLOC
+#define EOS_TRACE_MALLOC 1
+#endif
+
 namespace
 {
 	static void* EOS_MEMORY_CALL EosMalloc(size_t Bytes, size_t Alignment)
 	{
 		LLM_SCOPE(ELLMTag::RealTimeCommunications);
+
+#if !EOS_TRACE_MALLOC
+		CALLSTACK_TRACE_LIMIT_CALLSTACKRESOLVE_SCOPE();
+#endif
+
 		return FMemory::Malloc(Bytes, Alignment);
 	}
 
 	static void* EOS_MEMORY_CALL EosRealloc(void* Ptr, size_t Bytes, size_t Alignment)
 	{
 		LLM_SCOPE(ELLMTag::RealTimeCommunications);
+
+#if !EOS_TRACE_MALLOC
+		CALLSTACK_TRACE_LIMIT_CALLSTACKRESOLVE_SCOPE();
+#endif
+
 		return FMemory::Realloc(Ptr, Bytes, Alignment);
 	}
 
 	static void EOS_MEMORY_CALL EosFree(void* Ptr)
 	{
 		LLM_SCOPE(ELLMTag::RealTimeCommunications);
+
+#if !EOS_TRACE_MALLOC
+		CALLSTACK_TRACE_LIMIT_CALLSTACKRESOLVE_SCOPE();
+#endif
+
 		FMemory::Free(Ptr);
 	}
 
diff --git a/Engine/Source/Runtime/Core/Private/ProfilingDebugging/CallstackTrace.cpp b/Engine/Source/Runtime/Core/Private/ProfilingDebugging/CallstackTrace.cpp
--- a/Engine/Source/Runtime/Core/Private/ProfilingDebugging/CallstackTrace.cpp
+++ b/Engine/Source/Runtime/Core/Private/ProfilingDebugging/CallstackTrace.cpp
@@ -14,6 +14,8 @@ void	CallstackTrace_InitializeInternal();
 UE_TRACE_CHANNEL_DEFINE(CallstackChannel)
 UE_TRACE_EVENT_DEFINE(Memory, CallstackSpec)
 
+uint32 GCallStackTracingTlsSlotIndex = MAX_uint32;
+
 ////////////////////////////////////////////////////////////////////////////////
 void CallstackTrace_Create(class FMalloc* InMalloc)
 {
@@ -27,6 +29,9 @@ void CallstackTrace_Create(class FMalloc* InMalloc)
 ////////////////////////////////////////////////////////////////////////////////
 void CallstackTrace_Initialize()
 {
+	GCallStackTracingTlsSlotIndex = FPlatformTLS::AllocTlsSlot();
+	//NOTE: we don't bother cleaning up TLS, this is only closed during real shutdown
+
 	static auto InitOnce = [&]
 	{
 		CallstackTrace_InitializeInternal();
diff --git a/Engine/Source/Runtime/Core/Private/ProfilingDebugging/Microsoft/WindowsCallstackTrace.cpp b/Engine/Source/Runtime/Core/Private/ProfilingDebugging/Microsoft/WindowsCallstackTrace.cpp
--- a/Engine/Source/Runtime/Core/Private/ProfilingDebugging/Microsoft/WindowsCallstackTrace.cpp
+++ b/Engine/Source/Runtime/Core/Private/ProfilingDebugging/Microsoft/WindowsCallstackTrace.cpp
@@ -770,7 +770,7 @@ void CallstackTrace_InitializeInternal()
 ////////////////////////////////////////////////////////////////////////////////
 uint32 CallstackTrace_GetCurrentId()
 {
-	void* AddressOfReturnAddress = PLATFORM_RETURN_ADDRESS_POINTER();
+	void* AddressOfReturnAddress = PLATFORM_RETURN_ADDRESS_FOR_CALLSTACKTRACING();
 	if (FBacktracer* Instance = FBacktracer::Get())
 	{
 		return Instance->GetBacktraceId(AddressOfReturnAddress);
diff --git a/Engine/Source/Runtime/Core/Public/ProfilingDebugging/CallstackTrace.h b/Engine/Source/Runtime/Core/Public/ProfilingDebugging/CallstackTrace.h
--- a/Engine/Source/Runtime/Core/Public/ProfilingDebugging/CallstackTrace.h
+++ b/Engine/Source/Runtime/Core/Public/ProfilingDebugging/CallstackTrace.h
@@ -36,10 +36,67 @@ void CallstackTrace_Initialize();
  */
 CORE_API uint32 CallstackTrace_GetCurrentId();
 
+ /**
+  * Callstack Trace Scoped Macro to avoid resolving the full callstack
+  * can be used when some external libraries are not compiled with frame pointers
+  * preventing us to resolve it without crashing. Instead the callstack will be 
+  * only the caller address.
+  */
+#define CALLSTACK_TRACE_LIMIT_CALLSTACKRESOLVE_SCOPE() 						FCallStackTraceLimitResolveScope PREPROCESSOR_JOIN(FCTLMScope,__LINE__)
+
+extern uint32 GCallStackTracingTlsSlotIndex;
+
+/**
+* @return the fallback callstack address
+*/
+inline void* CallstackTrace_GetFallbackPlatformReturnAddressData()
+{
+	if (GCallStackTracingTlsSlotIndex != MAX_uint32)
+		return FPlatformTLS::GetTlsValue(GCallStackTracingTlsSlotIndex);
+	else
+		return nullptr;
+}
+
+/**
+* @return Needs full callstack resolve
+*/
+inline bool  CallstackTrace_ResolveFullCallStack()
+{
+	return CallstackTrace_GetFallbackPlatformReturnAddressData() == nullptr;
+}
+
+/*
+ * Callstack Trace scope for override CallStack
+ */
+class CORE_API FCallStackTraceLimitResolveScope
+{
+public:
+	FORCENOINLINE FCallStackTraceLimitResolveScope()
+	{
+		if (GCallStackTracingTlsSlotIndex != MAX_uint32)
+		{
+			FPlatformTLS::SetTlsValue(GCallStackTracingTlsSlotIndex, PLATFORM_RETURN_ADDRESS_FOR_CALLSTACKTRACING());
+		}
+	}
+
+	FORCENOINLINE ~FCallStackTraceLimitResolveScope()
+	{
+		if (GCallStackTracingTlsSlotIndex != MAX_uint32)
+		{
+			FPlatformTLS::SetTlsValue(GCallStackTracingTlsSlotIndex, nullptr);
+		}
+	}
+};
+
 #else // UE_CALLSTACK_TRACE_ENABLED
 
 inline void CallstackTrace_Create(class FMalloc* Malloc) {}
 inline void CallstackTrace_Initialize() {}
 inline uint32 CallstackTrace_GetCurrentId() { return 0; }
+inline void* CallstackTrace_GetCurrentReturnAddressData() { return nullptr; }
+inline void* CallstackTrace_GetFallbackPlatformReturnAddressData() { return nullptr; }
+inline bool  CallstackTrace_ResolveFullCallStack() { return true; }
+
+#define CALLSTACK_TRACE_LIMIT_CALLSTACKRESOLVE_SCOPE()
 
 #endif // UE_CALLSTACK_TRACE_ENABLED
diff --git a/Engine/Source/Runtime/Core/Public/Windows/WindowsPlatform.h b/Engine/Source/Runtime/Core/Public/Windows/WindowsPlatform.h
--- a/Engine/Source/Runtime/Core/Public/Windows/WindowsPlatform.h
+++ b/Engine/Source/Runtime/Core/Public/Windows/WindowsPlatform.h
@@ -101,6 +101,8 @@ typedef FWindowsPlatformTypes FPlatformTypes;
 
 #define PLATFORM_SUPPORTS_BORDERLESS_WINDOW					1
 
+#define PLATFORM_RETURN_ADDRESS_FOR_CALLSTACKTRACING		PLATFORM_RETURN_ADDRESS_POINTER
+
 #define WINDOWS_USE_FEATURE_APPLICATIONMISC_CLASS			1
 #define WINDOWS_USE_FEATURE_PLATFORMPROCESS_CLASS			1
 #define WINDOWS_USE_FEATURE_PLATFORMMISC_CLASS				1
