Fix config warnings when loading templates.

diff --git a/Engine/Source/Editor/GameProjectGeneration/Private/GameProjectUtils.cpp b/Engine/Source/Editor/GameProjectGeneration/Private/GameProjectUtils.cpp
--- a/Engine/Source/Editor/GameProjectGeneration/Private/GameProjectUtils.cpp
+++ b/Engine/Source/Editor/GameProjectGeneration/Private/GameProjectUtils.cpp
@@ -1304,7 +1304,7 @@ UTemplateProjectDefs* GameProjectUtils::LoadTemplateDefs(const FString& ProjectD
 {
 	UTemplateProjectDefs* TemplateDefs = nullptr;
 
-	const FString TemplateDefsIniFilename = ProjectDirectory / TEXT("Config") / GetTemplateDefsFilename();
+	const FString TemplateDefsIniFilename = FConfigCacheIni::NormalizeConfigIniPath(ProjectDirectory / TEXT("Config") / GetTemplateDefsFilename());
 	if ( FPlatformFileManager::Get().GetPlatformFile().FileExists(*TemplateDefsIniFilename) )
 	{
 		UClass* ClassToConstruct = UDefaultTemplateProjectDefs::StaticClass();
diff --git a/Engine/Source/Editor/UnrealEd/Private/AssetGuideline.cpp b/Engine/Source/Editor/UnrealEd/Private/AssetGuideline.cpp
--- a/Engine/Source/Editor/UnrealEd/Private/AssetGuideline.cpp
+++ b/Engine/Source/Editor/UnrealEd/Private/AssetGuideline.cpp
@@ -85,7 +85,7 @@ void UAssetGuideline::PostLoad()
 		if (IConsoleManager::Get().FindConsoleVariable(*ProjectSetting.Key))
 		{
 			FString CurrentIniValue;
-			FString FilenamePath = FPaths::ProjectDir() + ProjectSetting.Filename;
+			FString FilenamePath = FConfigCacheIni::NormalizeConfigIniPath(FPaths::ProjectDir() + ProjectSetting.Filename);
 			if (GConfig->GetString(*ProjectSetting.Section, *ProjectSetting.Key, CurrentIniValue, FilenamePath))
 			{
 				if (CurrentIniValue != ProjectSetting.Value)
@@ -271,7 +271,7 @@ void UAssetGuideline::EnableMissingGuidelines(TArray<FString> IncorrectPlugins,
 		for (const FIniStringValue& IncorrectProjectSetting : IncorrectProjectSettings)
 		{
 			// Only fails if file DNE
-			FString FilenamePath = FPaths::ProjectDir() + IncorrectProjectSetting.Filename;
+			FString FilenamePath = FConfigCacheIni::NormalizeConfigIniPath(FPaths::ProjectDir() + IncorrectProjectSetting.Filename);
 			if (bSuccess && GConfig->Find(FilenamePath))
 			{
 				FGameProjectGenerationModule::Get().TryMakeProjectFileWriteable(FilenamePath);
diff --git a/Engine/Source/Runtime/Core/Private/Misc/ConfigCacheIni.cpp b/Engine/Source/Runtime/Core/Private/Misc/ConfigCacheIni.cpp
--- a/Engine/Source/Runtime/Core/Private/Misc/ConfigCacheIni.cpp
+++ b/Engine/Source/Runtime/Core/Private/Misc/ConfigCacheIni.cpp
@@ -2311,7 +2311,7 @@ FConfigFile* FConfigCacheIni::Find(const FString& Filename)
 	if (!Result && !bAreFileOperationsDisabled)
 	{
 		// Before attempting to add another file, double check that this doesn't exist at a normalized path.
-		const FString UnrealFileName = FPaths::CreateStandardFilename(FPaths::RemoveDuplicateSlashes(Filename));
+		const FString UnrealFileName = NormalizeConfigIniPath(Filename);
 		Result = FindConfigFile(UnrealFileName);
 		
 		if (!Result)
@@ -2326,7 +2326,9 @@ FConfigFile* FConfigCacheIni::Find(const FString& Filename)
 		}
 		else
 		{
-			UE_LOG(LogConfig, Warning, TEXT("GConfig::Find attempting to access config with non-normalized path %s"), *Filename);
+			// We could normalize always normalize paths, but we don't want to always incur the penalty of that
+			// when callers can cache the strings ahead of time.
+			UE_LOG(LogConfig, Warning, TEXT("GConfig::Find attempting to access config with non-normalized path %s. Please use FConfigCacheIni::NormalizeConfigIniPath before accessing INI files through ConfigCache."), *Filename);
 		}
 	}
 
@@ -4720,6 +4722,13 @@ void FConfigCacheIni::LoadConsoleVariablesFromINI()
 	IConsoleManager::Get().CallAllConsoleVariableSinks();
 }
 
+FString FConfigCacheIni::NormalizeConfigIniPath(const FString& NonNormalizedPath)
+{
+	// CreateStandardFilename may not actually do anything in certain cases (e.g., if we detect a network drive, non-root drive, etc.)
+	// At a minimum, we will remove double slashes to try and fix some errors.
+	return FPaths::CreateStandardFilename(FPaths::RemoveDuplicateSlashes(NonNormalizedPath));
+}
+
 FArchive& operator<<(FArchive& Ar, FConfigFile& ConfigFile)
 {
 	bool bHasSourceConfigFile = ConfigFile.SourceConfigFile != nullptr;
diff --git a/Engine/Source/Runtime/Core/Public/Misc/ConfigCacheIni.h b/Engine/Source/Runtime/Core/Public/Misc/ConfigCacheIni.h
--- a/Engine/Source/Runtime/Core/Public/Misc/ConfigCacheIni.h
+++ b/Engine/Source/Runtime/Core/Public/Misc/ConfigCacheIni.h
@@ -1161,6 +1161,36 @@ public:
 	 */
 	static void LoadConsoleVariablesFromINI();
 
+	/**
+	 * Normalizes file paths to INI files.
+	 *
+	 * If an INI file is accessed with multiple paths, then we can run into issues where we cache multiple versions
+	 * of the file. Specifically, any updates to the file may only be applied to one cached version, and could cause
+	 * changes to be lost.
+	 *
+	 * E.G.
+	 *
+	 *		// Standard path.
+	 *		C:\ProjectDir\Engine\Config\DefaultEngine.ini
+	 *
+	 *		// Lowercase drive, and an extra slash between ProjectDir and Engine.
+	 *		c:\ProjectDir\\Engine\Confg\DefaultEngine.ini
+	 *
+	 *		// Relative to a project binary.
+	 *		..\..\..\ConfigDefaultEngine.ini
+	 *
+	 *		The paths above could all be used to reference the same ini file (namely, DefaultEngine.ini).
+	 *		However, they would end up generating unique entries in the GConfigCache.
+	 *		That means any modifications to *one* of the entries would not propagate to the others, and if
+	 *		any / all of the ini files are saved, they will stomp changes to the other entries.
+	 *
+	 *		We can prevent these types of issues by enforcing normalized paths when accessing configs.
+	 *
+	 *	@param NonNormalizedPath	The path to the INI file we want to access.
+	 *	@return A normalized version of the path (may be the same as the input).
+	 */
+	static FString NormalizeConfigIniPath(const FString& NonNormalizedPath);
+	
 	/**
 	 * Save the current config cache state into a file for bootstrapping other processes.
 	 */
diff --git a/Engine/Source/Runtime/GameplayTags/Private/GameplayTagsManager.cpp b/Engine/Source/Runtime/GameplayTags/Private/GameplayTagsManager.cpp
--- a/Engine/Source/Runtime/GameplayTags/Private/GameplayTagsManager.cpp
+++ b/Engine/Source/Runtime/GameplayTags/Private/GameplayTagsManager.cpp
@@ -259,7 +259,7 @@ void UGameplayTagsManager::AddTagIniSearchPath(const FString& RootDir)
 			{
 				const FName TagSource = FName(*FPaths::GetCleanFilename(IniFilePath));
 				PathInfo->SourcesInPath.Add(TagSource);
-				PathInfo->TagIniList.Add(FPaths::CreateStandardFilename(IniFilePath));
+				PathInfo->TagIniList.Add(FConfigCacheIni::NormalizeConfigIniPath(IniFilePath));
 			}
 		}
 		
