Fix FFileHandleUnix::Write when writing files more than 2Gb in size.
#rnx

diff --git a/Engine/Source/Runtime/Core/Private/Unix/UnixPlatformFile.cpp b/Engine/Source/Runtime/Core/Private/Unix/UnixPlatformFile.cpp
--- a/Engine/Source/Runtime/Core/Private/Unix/UnixPlatformFile.cpp
+++ b/Engine/Source/Runtime/Core/Private/Unix/UnixPlatformFile.cpp
@@ -83,7 +83,10 @@ static FUnixFileRegistry GFileRegistry;
  */
 class CORE_API FFileHandleUnix : public FRegisteredFileHandle
 {
-	enum {READWRITE_SIZE = SSIZE_MAX};
+	// https://man7.org/linux/man-pages/man2/write.2.html
+	// On Linux, write() (and similar system calls) will transfer at most 0x7ffff000 (2,147,479,552) bytes,
+	// returning the number of bytes actually transferred.  (This is true on both 32-bit and 64-bit systems.)
+	enum {READWRITE_SIZE = 0x7ffff000};
 
 	FORCEINLINE bool IsValid()
 	{
@@ -210,12 +213,13 @@ public:
 			check(BytesToWrite >= 0);
 			int64 ThisSize = FMath::Min<int64>(READWRITE_SIZE, BytesToWrite);
 			check(Source);
-			if (write(FileHandle, Source, ThisSize) != ThisSize)
+			int64 WrittenSize = write(FileHandle, Source, ThisSize);
+			if (WrittenSize == -1)
 			{
 				return false;
 			}
-			Source += ThisSize;
-			BytesToWrite -= ThisSize;
+			Source += WrittenSize;
+			BytesToWrite -= WrittenSize;
 		}
 		return true;
 	}
