Implement OpenMapped for Linux inside of UnixPlatformFile

diff --git a/Engine/Source/Runtime/Core/Private/Unix/UnixPlatformFile.cpp b/Engine/Source/Runtime/Core/Private/Unix/UnixPlatformFile.cpp
--- a/Engine/Source/Runtime/Core/Private/Unix/UnixPlatformFile.cpp
+++ b/Engine/Source/Runtime/Core/Private/Unix/UnixPlatformFile.cpp
@@ -7,7 +7,9 @@
 #include "Containers/LruCache.h"
 #include "Logging/LogMacros.h"
 #include "Misc/Paths.h"
+#include "Async/MappedFileHandle.h"
 #include <sys/file.h>
+#include <sys/mman.h>
 
 #include "HAL/PlatformFileCommon.h"
 #include "HAL/PlatformFileManager.h"
@@ -684,6 +686,97 @@ public:
 
 FUnixFileMapper GCaseInsensMapper;
 
+class FUnixMappedFileRegion final : public IMappedFileRegion
+{
+public:
+	class FUnixMappedFileHandle* Parent;
+	const uint8* AlignedPtr;
+	uint64 AlignedSize;
+	FUnixMappedFileRegion(const uint8* InMappedPtr, const uint8* InAlignedPtr, size_t InMappedSize, uint64 InAlignedSize, const FString& InDebugFilename, size_t InDebugOffsetIntoFile, FUnixMappedFileHandle* InParent)
+		: IMappedFileRegion(InMappedPtr, InMappedSize, InDebugFilename, InDebugOffsetIntoFile)
+		, Parent(InParent)
+		, AlignedPtr(InAlignedPtr)
+		, AlignedSize(InAlignedSize)
+	{
+	}
+};
+
+class FUnixMappedFileHandle final : public IMappedFileHandle
+{	
+public:
+	FUnixMappedFileHandle(int InFileHandle, int64 FileSize, const FString& InFilename)
+		: IMappedFileHandle(FileSize)
+		, MappedPtr(nullptr)
+		, Filename(InFilename)
+		, NumOutstandingRegions(0)
+		, FileHandle(InFileHandle)
+	{
+		Alignment = sysconf(_SC_PAGE_SIZE);
+	}
+
+	virtual ~FUnixMappedFileHandle() override
+	{
+		check(!NumOutstandingRegions); // can't delete the file before you delete all outstanding regions
+		close(FileHandle);
+	}
+	
+	virtual IMappedFileRegion* MapRegion(int64 Offset = 0, int64 BytesToMap = MAX_int64, bool bPreloadHint = false) override
+	{
+		LLM_PLATFORM_SCOPE(ELLMTag::PlatformMMIO);
+		check(Offset < GetFileSize()); // don't map zero bytes and don't map off the end of the file
+		BytesToMap = FMath::Min<int64>(BytesToMap, GetFileSize() - Offset);
+		check(BytesToMap > 0); // don't map zero bytes
+
+		const int64 AlignedOffset = AlignDown(Offset, Alignment);
+		const int64 AlignedSize = Align(BytesToMap + Offset - AlignedOffset, Alignment);
+		
+		// if we are about to go off the end, let's not
+		if (AlignedOffset + AlignedSize > GetFileSize())
+		{
+			UE_LOG(LogLinux, Warning, TEXT("Mapping fell off the end, did we need to actually abort? [%lld + %lld > %lld]"), AlignedOffset, AlignedSize, GetFileSize());
+			return nullptr;
+		}
+
+		int Flags = MAP_PRIVATE;
+		if (bPreloadHint)
+		{
+			Flags |= MAP_POPULATE;
+		}
+		
+		const uint8* AlignedMapPtr = static_cast<const uint8*>(mmap(nullptr, AlignedSize, PROT_READ, Flags, FileHandle, AlignedOffset));
+		if (AlignedMapPtr == MAP_FAILED || AlignedMapPtr == nullptr)
+		{
+			UE_LOG(LogLinux, Warning, TEXT("Failed to map memory %s, error is %d"), *Filename, errno);
+			return nullptr;
+		}
+		LLM(FLowLevelMemTracker::Get().OnLowLevelAlloc(ELLMTracker::Platform, AlignedMapPtr, AlignedSize));
+
+		// create a mapping for this range
+		const uint8* MapPtr = AlignedMapPtr + Offset - AlignedOffset;
+		FUnixMappedFileRegion* Result = new FUnixMappedFileRegion(MapPtr, AlignedMapPtr, BytesToMap, AlignedSize, Filename, Offset, this);
+		NumOutstandingRegions++;
+		return Result;
+	}
+	
+	void UnMap(const FUnixMappedFileRegion* Region)
+	{
+		LLM_PLATFORM_SCOPE(ELLMTag::PlatformMMIO);
+		check(NumOutstandingRegions > 0);
+		NumOutstandingRegions--;
+		
+		LLM(FLowLevelMemTracker::Get().OnLowLevelFree(ELLMTracker::Platform, Region->AlignedPtr));
+		const int Res = munmap(const_cast<uint8*>(Region->AlignedPtr), Region->AlignedSize);
+		checkf(Res == 0, TEXT("Failed to unmap, error is %d, errno is %d [params: %x, %d]"), Res, errno, MappedPtr, GetFileSize());
+	}
+
+private:
+	const uint8* MappedPtr;
+	FString Filename;
+	int32 NumOutstandingRegions;
+	int32 Alignment;
+	int FileHandle;
+};
+
 /**
  * Unix File I/O implementation
 **/
@@ -996,11 +1089,39 @@ IFileHandle* FUnixPlatformFile::OpenWrite(const TCHAR* Filename, bool bAppend, b
 		return FileHandleUnix;
 	}
 
-	int ErrNo = errno;
+	const int ErrNo = errno;
 	UE_LOG_UNIX_FILE(Warning, TEXT( "open('%s', Flags=0x%08X) failed: errno=%d (%s)" ), *NormalizeFilename(Filename, true), Flags, ErrNo, UTF8_TO_TCHAR(strerror(ErrNo)));
 	return nullptr;
 }
 
+IMappedFileHandle* FUnixPlatformFile::OpenMapped(const TCHAR* Filename)
+{
+	const FString NormalizedFilename = NormalizeFilename(Filename, false);
+
+	constexpr int Flags = O_RDONLY;
+	const int32 Handle = open(TCHAR_TO_UTF8(*NormalizedFilename), Flags);
+	if (Handle == -1)
+	{
+		const int ErrNo = errno;
+		UE_LOG_UNIX_FILE(Warning, TEXT("open('%s', Flags=0x%08X) failed: errno=%d (%s)"), *NormalizedFilename, Flags, ErrNo, UTF8_TO_TCHAR(strerror(ErrNo)));
+
+		return nullptr;
+	}
+	
+	struct stat FileInfo;
+	FileInfo.st_size = -1;
+	const int StatResult = fstat(Handle, &FileInfo);
+	if (StatResult == -1)
+	{
+		const int ErrNo = errno;
+		UE_LOG_UNIX_FILE(Warning, TEXT("stat('%s', Flags=0x%08X) failed: errno=%d (%s)"), *NormalizedFilename, Flags, ErrNo, UTF8_TO_TCHAR(strerror(ErrNo)));
+
+		return nullptr;
+	}
+	
+	return new FUnixMappedFileHandle(Handle, FileInfo.st_size, NormalizedFilename); 
+}
+
 bool FUnixPlatformFile::DirectoryExists(const TCHAR* Directory)
 {
 	FString CaseSensitiveFilename;
diff --git a/Engine/Source/Runtime/Core/Public/Unix/UnixPlatformFile.h b/Engine/Source/Runtime/Core/Public/Unix/UnixPlatformFile.h
--- a/Engine/Source/Runtime/Core/Public/Unix/UnixPlatformFile.h
+++ b/Engine/Source/Runtime/Core/Public/Unix/UnixPlatformFile.h
@@ -43,6 +43,8 @@ public:
 
 	virtual IFileHandle* OpenRead(const TCHAR* Filename, bool bAllowWrite = false) override;
 	virtual IFileHandle* OpenWrite(const TCHAR* Filename, bool bAppend = false, bool bAllowRead = false) override;
+	virtual IMappedFileHandle* OpenMapped(const TCHAR* Filename) override;
+
 	virtual bool DirectoryExists(const TCHAR* Directory) override;
 	virtual bool CreateDirectory(const TCHAR* Directory) override;
 	virtual bool DeleteDirectory(const TCHAR* Directory) override;
