Remove unused read-ahead cache functionality

diff --git a/Engine/Source/Runtime/Core/Private/Unix/UnixPlatformFile.cpp b/Engine/Source/Runtime/Core/Private/Unix/UnixPlatformFile.cpp
--- a/Engine/Source/Runtime/Core/Private/Unix/UnixPlatformFile.cpp
+++ b/Engine/Source/Runtime/Core/Private/Unix/UnixPlatformFile.cpp
@@ -314,7 +314,6 @@ private:
 	// track if file is open for write
 	bool FileOpenAsWrite;
 
-	friend class FReadaheadCache;
 	friend class FUnixFileRegistry;
 };
 
