Fix compilation of nonunity build

#rnx

diff --git a/Engine/Source/Runtime/Core/Public/ProfilingDebugging/CallstackTrace.h b/Engine/Source/Runtime/Core/Public/ProfilingDebugging/CallstackTrace.h
--- a/Engine/Source/Runtime/Core/Public/ProfilingDebugging/CallstackTrace.h
+++ b/Engine/Source/Runtime/Core/Public/ProfilingDebugging/CallstackTrace.h
@@ -1,6 +1,8 @@
 // Copyright Epic Games, Inc. All Rights Reserved.
 #pragma once
 #include "TraceLog/Public/Trace/Config.h"
+#include "Math/NumericLimits.h"
+#include "HAL/PlatformTLS.h"
 
 ////////////////////////////////////////////////////////////////////////////////
 #if !defined(UE_CALLSTACK_TRACE_ENABLED)
