Allow functions/methods that only take args and an output device as parameters  to be bound to console commands. To to be used in cases where the UWorld* parameter is not required.

#rnx

- In theory the callback could just take all 3 parameter types then ignore the world pointer but in many places this means forward declaring UWorld* then documenting why the parameter is ignored.
- Adding this permuation of callback will help keep code cleaner and too the point.

diff --git a/Engine/Source/Runtime/Core/Private/HAL/ConsoleManager.cpp b/Engine/Source/Runtime/Core/Private/HAL/ConsoleManager.cpp
--- a/Engine/Source/Runtime/Core/Private/HAL/ConsoleManager.cpp
+++ b/Engine/Source/Runtime/Core/Private/HAL/ConsoleManager.cpp
@@ -1317,6 +1317,35 @@ private:
 	FConsoleCommandWithWorldAndArgsDelegate Delegate;
 };
 
+/* Console command that can be given args and an output device. */
+class FConsoleCommanWithArgsAndOutputDevice : public FConsoleCommandBase
+{
+
+public:
+	FConsoleCommanWithArgsAndOutputDevice(const FConsoleCommandWithArgsAndOutputDeviceDelegate& InitDelegate, const TCHAR* InitHelp, const EConsoleVariableFlags InitFlags)
+		: FConsoleCommandBase(InitHelp, InitFlags),
+		Delegate(InitDelegate)
+	{
+	}
+
+	// interface IConsoleCommand -----------------------------------
+
+	virtual void Release() override
+	{
+		delete this;
+	}
+
+	virtual bool Execute(const TArray< FString >& Args, UWorld* InWorld, FOutputDevice& OutputDevice) override
+	{
+		return Delegate.ExecuteIfBound(Args, OutputDevice);
+	}
+
+private:
+
+	/** User function to call when the console command is executed */
+	FConsoleCommandWithArgsAndOutputDeviceDelegate Delegate;
+};
+
 /* Console command that can be given a world parameter, args and an output device. */
 class FConsoleCommandWithWorldArgsAndOutputDevice : public FConsoleCommandBase
 {
@@ -1471,6 +1500,11 @@ IConsoleCommand* FConsoleManager::RegisterConsoleCommand(const TCHAR* Name, cons
 	return AddConsoleObject(Name, new FConsoleCommandWithWorldAndArgs(Command, Help, (EConsoleVariableFlags)Flags))->AsCommand();
 }
 
+IConsoleCommand* FConsoleManager::RegisterConsoleCommand(const TCHAR* Name, const TCHAR* Help, const FConsoleCommandWithArgsAndOutputDeviceDelegate& Command, uint32 Flags)
+{
+	return AddConsoleObject(Name, new FConsoleCommanWithArgsAndOutputDevice(Command, Help, (EConsoleVariableFlags)Flags))->AsCommand();
+}
+
 IConsoleCommand* FConsoleManager::RegisterConsoleCommand(const TCHAR* Name, const TCHAR* Help, const FConsoleCommandWithWorldArgsAndOutputDeviceDelegate& Command, uint32 Flags)
 {
 	return AddConsoleObject(Name, new FConsoleCommandWithWorldArgsAndOutputDevice(Command, Help, (EConsoleVariableFlags)Flags))->AsCommand();
diff --git a/Engine/Source/Runtime/Core/Public/HAL/ConsoleManager.h b/Engine/Source/Runtime/Core/Public/HAL/ConsoleManager.h
--- a/Engine/Source/Runtime/Core/Public/HAL/ConsoleManager.h
+++ b/Engine/Source/Runtime/Core/Public/HAL/ConsoleManager.h
@@ -68,6 +68,7 @@ public:
 	virtual IConsoleCommand* RegisterConsoleCommand(const TCHAR* Name, const TCHAR* Help, const FConsoleCommandWithArgsDelegate& Command, uint32 Flags) override;
 	virtual IConsoleCommand* RegisterConsoleCommand(const TCHAR* Name, const TCHAR* Help, const FConsoleCommandWithWorldDelegate& Command, uint32 Flags) override;
 	virtual IConsoleCommand* RegisterConsoleCommand(const TCHAR* Name, const TCHAR* Help, const FConsoleCommandWithWorldAndArgsDelegate& Command, uint32 Flags) override;
+	virtual IConsoleCommand* RegisterConsoleCommand(const TCHAR* Name, const TCHAR* Help, const FConsoleCommandWithArgsAndOutputDeviceDelegate& Command, uint32 Flags) override;
 	virtual IConsoleCommand* RegisterConsoleCommand(const TCHAR* Name, const TCHAR* Help, const FConsoleCommandWithWorldArgsAndOutputDeviceDelegate& Command, uint32 Flags) override;
 	virtual IConsoleCommand* RegisterConsoleCommand(const TCHAR* Name, const TCHAR* Help, const FConsoleCommandWithOutputDeviceDelegate& Command, uint32 Flags) override;
 	virtual IConsoleCommand* RegisterConsoleCommand(const TCHAR* Name, const TCHAR* Help, uint32 Flags) override;
diff --git a/Engine/Source/Runtime/Core/Public/HAL/IConsoleManager.h b/Engine/Source/Runtime/Core/Public/HAL/IConsoleManager.h
--- a/Engine/Source/Runtime/Core/Public/HAL/IConsoleManager.h
+++ b/Engine/Source/Runtime/Core/Public/HAL/IConsoleManager.h
@@ -156,6 +156,9 @@ DECLARE_DELEGATE_OneParam( FConsoleCommandWithWorldDelegate, UWorld* );
 /** Console command delegate type (with a world and arguments.)  This is a void callback function that always takes a list of arguments and a world. */
 DECLARE_DELEGATE_TwoParams(FConsoleCommandWithWorldAndArgsDelegate, const TArray< FString >&, UWorld*);
 
+/** Console command delegate type (with arguments and output device.)  This is a void callback function that always takes a list of arguments and output device. */
+DECLARE_DELEGATE_TwoParams(FConsoleCommandWithArgsAndOutputDeviceDelegate, const TArray< FString >&, FOutputDevice&);
+
 /** Console command delegate type (with a world arguments and output device.)  This is a void callback function that always takes a list of arguments, a world and output device. */
 DECLARE_DELEGATE_ThreeParams(FConsoleCommandWithWorldArgsAndOutputDeviceDelegate, const TArray< FString >&, UWorld*, FOutputDevice&);
 
@@ -271,13 +274,14 @@ struct FNullConsoleVariableDelegate
 	}
 };
 
-struct FConsoleVariableDelegate                            : FNullConsoleVariableDelegate<FConsoleVariableDelegate, IConsoleVariable*> {};
-struct FConsoleCommandDelegate                             : FNullConsoleVariableDelegate<FConsoleCommandDelegate> {};
-struct FConsoleCommandWithArgsDelegate                     : FNullConsoleVariableDelegate<FConsoleCommandWithArgsDelegate, const TArray<FString>&> {};
-struct FConsoleCommandWithWorldDelegate                    : FNullConsoleVariableDelegate<FConsoleCommandWithWorldDelegate, UWorld*> {};
-struct FConsoleCommandWithWorldAndArgsDelegate             : FNullConsoleVariableDelegate<FConsoleCommandWithWorldAndArgsDelegate, const TArray<FString>&, UWorld*> {};
-struct FConsoleCommandWithWorldArgsAndOutputDeviceDelegate : FNullConsoleVariableDelegate<FConsoleCommandWithWorldArgsAndOutputDeviceDelegate, const TArray<FString>&, UWorld*, FOutputDevice&> {};
-struct FConsoleCommandWithOutputDeviceDelegate             : FNullConsoleVariableDelegate<FConsoleCommandWithOutputDeviceDelegate, FOutputDevice&> {};
+struct FConsoleVariableDelegate								: FNullConsoleVariableDelegate<FConsoleVariableDelegate, IConsoleVariable*> {};
+struct FConsoleCommandDelegate								: FNullConsoleVariableDelegate<FConsoleCommandDelegate> {};
+struct FConsoleCommandWithArgsDelegate						: FNullConsoleVariableDelegate<FConsoleCommandWithArgsDelegate, const TArray<FString>&> {};
+struct FConsoleCommandWithWorldDelegate						: FNullConsoleVariableDelegate<FConsoleCommandWithWorldDelegate, UWorld*> {};
+struct FConsoleCommandWithWorldAndArgsDelegate				: FNullConsoleVariableDelegate<FConsoleCommandWithWorldAndArgsDelegate, const TArray<FString>&, UWorld*> {};
+struct FConsoleCommandWithArgsAndOutputDeviceDelegate		: FNullConsoleVariableDelegate<FConsoleCommandWithArgsAndOutputDeviceDelegate, const TArray<FString>&, FOutputDevice&> {};
+struct FConsoleCommandWithWorldArgsAndOutputDeviceDelegate	: FNullConsoleVariableDelegate<FConsoleCommandWithWorldArgsAndOutputDeviceDelegate, const TArray<FString>&, UWorld*, FOutputDevice&> {};
+struct FConsoleCommandWithOutputDeviceDelegate				: FNullConsoleVariableDelegate<FConsoleCommandWithOutputDeviceDelegate, FOutputDevice&> {};
 
 #endif
 
@@ -816,7 +820,17 @@ struct CORE_API IConsoleManager
 	* @param	Flags		Optional flags bitmask
 	*/
 	virtual IConsoleCommand* RegisterConsoleCommand(const TCHAR* Name, const TCHAR* Help, const FConsoleCommandWithWorldAndArgsDelegate& Command, uint32 Flags = ECVF_Default) = 0;
-
+	
+	/**
+	* Register a console command that takes arguments
+	*
+	* @param	Name		The name of this command (must not be nullptr)
+	* @param	Help		Help text for this command
+	* @param	Command		The user function to call when this command is executed
+	* @param	Flags		Optional flags bitmask
+	*/
+	virtual IConsoleCommand* RegisterConsoleCommand(const TCHAR* Name, const TCHAR* Help, const FConsoleCommandWithArgsAndOutputDeviceDelegate& Command, uint32 Flags = ECVF_Default) = 0;
+	
 	/**
 	* Register a console command that takes arguments
 	*
@@ -1755,6 +1769,19 @@ public:
 	}
 };
 
+/**
+ * Autoregistering console command with args and an output device
+ */
+class CORE_API FAutoConsoleCommandWithArgsAndOutputDevice : private FAutoConsoleObject
+{
+public:
+	
+	FAutoConsoleCommandWithArgsAndOutputDevice(const TCHAR* Name, const TCHAR* Help, const FConsoleCommandWithArgsAndOutputDeviceDelegate& Command, uint32 Flags = ECVF_Default)
+		: FAutoConsoleObject(IConsoleManager::Get().RegisterConsoleCommand(Name, Help, Command, Flags))
+	{
+	}
+};
+
 /**
  * Autoregistering console command with an output device
  */
