this line got lost in the previous big refactoring, adding it back. a cosmetic change for consistency, as it doesn't affect anything

diff --git a/Engine/Source/Runtime/Core/Private/Tests/Tasks/TasksTest.cpp b/Engine/Source/Runtime/Core/Private/Tests/Tasks/TasksTest.cpp
--- a/Engine/Source/Runtime/Core/Private/Tests/Tasks/TasksTest.cpp
+++ b/Engine/Source/Runtime/Core/Private/Tests/Tasks/TasksTest.cpp
@@ -256,6 +256,7 @@ namespace UE { namespace TasksTests
 			Signal2.Trigger();
 			verify(!Task.Wait(FTimespan::FromMilliseconds(100)));
 			Signal3.Trigger();
+			Task.Wait();
 		}
 
 		{	// nested in square
