Updated error message logging tags from LogLinux to LogUnixPlatformFile

diff --git a/Engine/Source/Runtime/Core/Private/Unix/UnixPlatformFile.cpp b/Engine/Source/Runtime/Core/Private/Unix/UnixPlatformFile.cpp
--- a/Engine/Source/Runtime/Core/Private/Unix/UnixPlatformFile.cpp
+++ b/Engine/Source/Runtime/Core/Private/Unix/UnixPlatformFile.cpp
@@ -732,7 +732,7 @@ public:
 		// if we are about to go off the end, let's not
 		if (AlignedOffset + AlignedSize > GetFileSize())
 		{
-			UE_LOG(LogLinux, Warning, TEXT("Mapping fell off the end, did we need to actually abort? [%lld + %lld > %lld]"), AlignedOffset, AlignedSize, GetFileSize());
+			UE_LOG(LogUnixPlatformFile, Warning, TEXT("Mapping fell off the end, did we need to actually abort? [%lld + %lld > %lld]"), AlignedOffset, AlignedSize, GetFileSize());
 			return nullptr;
 		}
 
@@ -745,7 +745,7 @@ public:
 		const uint8* AlignedMapPtr = static_cast<const uint8*>(mmap(nullptr, AlignedSize, PROT_READ, Flags, FileHandle, AlignedOffset));
 		if (AlignedMapPtr == MAP_FAILED || AlignedMapPtr == nullptr)
 		{
-			UE_LOG(LogLinux, Warning, TEXT("Failed to map memory %s, error is %d"), *Filename, errno);
+			UE_LOG(LogUnixPlatformFile, Warning, TEXT("Failed to map memory %s, error is %d"), *Filename, errno);
 			return nullptr;
 		}
 		LLM(FLowLevelMemTracker::Get().OnLowLevelAlloc(ELLMTracker::Platform, AlignedMapPtr, AlignedSize));
