FPaths and PakPlatformFile: suppress the Pak warning mounting to unregistered paths when the mount path is GameFeatureRootPrefix.
#rnx

diff --git a/Engine/Source/Editor/AddContentDialog/Private/ContentSourceProviders/FeaturePack/FeaturePackContentSource.cpp b/Engine/Source/Editor/AddContentDialog/Private/ContentSourceProviders/FeaturePack/FeaturePackContentSource.cpp
--- a/Engine/Source/Editor/AddContentDialog/Private/ContentSourceProviders/FeaturePack/FeaturePackContentSource.cpp
+++ b/Engine/Source/Editor/AddContentDialog/Private/ContentSourceProviders/FeaturePack/FeaturePackContentSource.cpp
@@ -158,7 +158,7 @@ FFeaturePackContentSource::FFeaturePackContentSource(FString InFeaturePackPath)
 	if( InFeaturePackPath.EndsWith(TEXT(".upack") ) == true )
 	{
 		bContentsInPakFile = true;
-		MountPoint = "root:/";
+		MountPoint = FPaths::GameFeatureRootPrefix();
 		// Create a pak platform file and mount the feature pack file.
 		FPakPlatformFile PakPlatformFile;
 		PakPlatformFile.Initialize(&FPlatformFileManager::Get().GetPlatformFile(), TEXT(""));
diff --git a/Engine/Source/Runtime/Core/Private/Misc/Paths.cpp b/Engine/Source/Runtime/Core/Private/Misc/Paths.cpp
--- a/Engine/Source/Runtime/Core/Private/Misc/Paths.cpp
+++ b/Engine/Source/Runtime/Core/Private/Misc/Paths.cpp
@@ -1045,12 +1045,15 @@ bool FPaths::IsDrive(const FString& InPath)
 	return false;
 }
 
-bool FPaths::IsRelative(const FString& InPath)
-{
 #if WITH_EDITOR
-	static const TCHAR RootPrefix[] = TEXT("root:/");
-#endif // WITH_EDITOR
+const TCHAR* FPaths::GameFeatureRootPrefix()
+{
+	return TEXT("root:/");
+}
+#endif
 
+bool FPaths::IsRelative(const FString& InPath)
+{
 	// The previous implementation of this function seemed to handle normalized and unnormalized paths, so this one does too for legacy reasons.
 	const uint32 PathLen = InPath.Len();
 	const bool IsRooted = PathLen &&
@@ -1059,7 +1062,7 @@ bool FPaths::IsRelative(const FString& InPath)
 			((InPath[0] == '\\') && (InPath[1] == '\\'))					// Root of the current directory on Windows. Also covers "\\" for UNC or "network" paths.
 			|| (InPath[1] == ':' && FChar::IsAlpha(InPath[0]))				// Starts with "<DriveLetter>:"
 #if WITH_EDITOR
-			|| (InPath.StartsWith(RootPrefix, ESearchCase::IgnoreCase))		// Feature packs use this
+			|| (InPath.StartsWith(GameFeatureRootPrefix(), ESearchCase::IgnoreCase))
 #endif // WITH_EDITOR
 			))
 		);
diff --git a/Engine/Source/Runtime/Core/Public/Misc/Paths.h b/Engine/Source/Runtime/Core/Public/Misc/Paths.h
--- a/Engine/Source/Runtime/Core/Public/Misc/Paths.h
+++ b/Engine/Source/Runtime/Core/Public/Misc/Paths.h
@@ -171,6 +171,11 @@ public:
 	 */
 	static FString RootDir();
 
+#if WITH_EDITOR
+	/** Returns the special path used when mounting FeaturePaks in editor */
+	static const TCHAR* GameFeatureRootPrefix();
+#endif
+
 	/**
 	 * Returns the base directory of the current project by looking at FApp::GetProjectName().
 	 * This is usually a subdirectory of the installation
diff --git a/Engine/Source/Runtime/PakFile/Private/IPlatformFilePak.cpp b/Engine/Source/Runtime/PakFile/Private/IPlatformFilePak.cpp
--- a/Engine/Source/Runtime/PakFile/Private/IPlatformFilePak.cpp
+++ b/Engine/Source/Runtime/PakFile/Private/IPlatformFilePak.cpp
@@ -7786,7 +7786,12 @@ bool FPakPlatformFile::Mount(const TCHAR* InPakFilename, uint32 PakOrder, const
 			UE_LOG(LogPakFile, Verbose, TEXT("OnPakFileMounted2Time == %lf"), OnPakFileMounted2Time);
 						
 			// skip this check for the default mountpoint, it will print false positives
-			if (FPaths::CreateStandardFilename(Pak->GetMountPoint()) != FPaths::CreateStandardFilename(FPaths::RootDir()))
+			FString NormalizedPakMountPoint = FPaths::CreateStandardFilename(Pak->GetMountPoint());
+			bool bIsMountingToRoot = NormalizedPakMountPoint == FPaths::CreateStandardFilename(FPaths::RootDir());
+#if WITH_EDITOR
+			bIsMountingToRoot |= NormalizedPakMountPoint == FPaths::CreateStandardFilename(FPaths::GameFeatureRootPrefix());
+#endif
+			if (!bIsMountingToRoot)
 			{
 				FString OutPackageName;
 				if (!FPackageName::TryConvertFilenameToLongPackageName(Pak->GetMountPoint(), OutPackageName))
