fixed a crash on local static stats task destruction that was caused by the allocator used for them was already destroyed

diff --git a/Engine/Source/Runtime/Core/Private/Stats/Stats2.cpp b/Engine/Source/Runtime/Core/Private/Stats/Stats2.cpp
--- a/Engine/Source/Runtime/Core/Private/Stats/Stats2.cpp
+++ b/Engine/Source/Runtime/Core/Private/Stats/Stats2.cpp
@@ -1574,6 +1574,10 @@ void FThreadStats::StartThread()
 
 static int32 CurrentEventIndex = 0;
 
+#if UE_STATS_THREAD_AS_PIPE
+static UE::Tasks::FTask LastFramesBarriers[MAX_STAT_LAG];
+#endif
+
 static FGraphEventRef LastFramesEvents[MAX_STAT_LAG];
 
 void FThreadStats::StopThread()
@@ -1601,6 +1605,10 @@ void FThreadStats::StopThread()
 	{
 		FPlatformProcess::Yield();
 	}
+	for (int32 Index = 0; Index < MAX_STAT_LAG; Index++)
+	{
+		LastFramesBarriers[Index] = {}; // release memory before the allocator is destroyed
+	}
 #else
 	for (int32 Index = 0; Index < MAX_STAT_LAG; Index++)
 	{
@@ -1630,13 +1638,9 @@ void FThreadStats::WaitForStats()
 #if UE_STATS_THREAD_AS_PIPE
 		if (FPlatformProcess::SupportsMultithreading())
 		{
-			static UE::Tasks::FTask LocalLastFramesEvents[MAX_STAT_LAG];
-			{
-				SCOPE_CYCLE_COUNTER(STAT_WaitForStats);
-				LocalLastFramesEvents[EventIndex].Wait();
-			}
-
-			LocalLastFramesEvents[EventIndex] = GStatsPipe.Launch(UE_SOURCE_LOCATION, [] {});
+			SCOPE_CYCLE_COUNTER(STAT_WaitForStats);
+			LastFramesBarriers[EventIndex].Wait();
+			LastFramesBarriers[EventIndex] = GStatsPipe.Launch(UE_SOURCE_LOCATION, [] {});
 		}
 		else
 #endif
