Update forked Session ProcessID post fork for child processes

Tested locally with forked servers and "kill -SIGSEGV" parent and child.
Tested on deployed servers by Jake + Brian.

#rnx

diff --git a/Engine/Source/Runtime/Core/Private/GenericPlatform/GenericPlatformCrashContext.cpp b/Engine/Source/Runtime/Core/Private/GenericPlatform/GenericPlatformCrashContext.cpp
--- a/Engine/Source/Runtime/Core/Private/GenericPlatform/GenericPlatformCrashContext.cpp
+++ b/Engine/Source/Runtime/Core/Private/GenericPlatform/GenericPlatformCrashContext.cpp
@@ -274,6 +274,17 @@ PRAGMA_ENABLE_DEPRECATION_WARNINGS
 
 	FCoreDelegates::ConfigReadyForUse.AddStatic(FGenericCrashContext::InitializeFromConfig);
 
+	FCoreDelegates::OnPostFork.AddLambda([](EForkProcessRole Role)
+	{
+		if (Role == EForkProcessRole::Child)
+		{
+			UE_LOG(LogCrashContext, VeryVerbose, TEXT("Updating forked child Session ProcessID: %u -> %u"), NCached::Session.ProcessId, FPlatformProcess::GetCurrentProcessId());
+
+			NCached::Session.ProcessId = FPlatformProcess::GetCurrentProcessId();
+			SerializeTempCrashContextToFile();
+		}
+	});
+
 	SerializeTempCrashContextToFile();
 
 	CleanupPlatformSpecificFiles();
