a workaround for MSVC compiler bug, to make TExecutableTask compile if instantiated inside a lambda. a simplified test: https://godbolt.org/z/ae96c7evn

diff --git a/Engine/Source/Runtime/Core/Public/Tasks/TaskPrivate.h b/Engine/Source/Runtime/Core/Public/Tasks/TaskPrivate.h
--- a/Engine/Source/Runtime/Core/Public/Tasks/TaskPrivate.h
+++ b/Engine/Source/Runtime/Core/Public/Tasks/TaskPrivate.h
@@ -701,13 +701,11 @@ namespace UE::Tasks
 		// In most cases it should be allocated on the heap and used with TRefCountPtr, e.g. @see FTaskHandle. With care, can be allocated on the stack, e.g. see 
 		// WaitingTask in FTaskBase::Wait().
 		// Implements memory allocation from a pooled fixed-size allocator tuned for the everage UE task size
-		template<typename TaskBodyType, typename Enable = void>
-		class TExecutableTask final : public TConcurrentLinearObject<TExecutableTask<TaskBodyType>, FTaskBlockAllocationTag>, public TTaskWithResult<TInvokeResult_T<TaskBodyType>>
+		template<typename TaskBodyType, typename ResultType = TInvokeResult_T<TaskBodyType>, typename Enable = void>
+		class TExecutableTask final : public TConcurrentLinearObject<TExecutableTask<TaskBodyType>, FTaskBlockAllocationTag>, public TTaskWithResult<ResultType>
 		{
 			UE_NONCOPYABLE(TExecutableTask);
 
-			using ResultType = TInvokeResult_T<TaskBodyType>;
-
 		public:
 			static TExecutableTask* Create(const TCHAR* DebugName, TaskBodyType TaskBody, ETaskPriority Priority, EExtendedTaskPriority InExtendedPriority = EExtendedTaskPriority::None)
 			{
