Fix some ini files not being read when reading config files from another project.

#rnx

- When trying to load a set of config files for another project, ini files under restricted  directories were not being loaded for the target project, but from the current project.
- PerformFinalExpansions needed to be changed to use the path of the path passed in as a parameter rather than FPaths::ProjectDir
- PerformFinalExpansions also caches values, which means if we later try to load a set of config files from another project we will end up using the previous projects path. We need to store the last project we cached values for and recalculate when the project path changes.

diff --git a/Engine/Source/Runtime/Core/Private/Misc/ConfigCacheIni.cpp b/Engine/Source/Runtime/Core/Private/Misc/ConfigCacheIni.cpp
--- a/Engine/Source/Runtime/Core/Private/Misc/ConfigCacheIni.cpp
+++ b/Engine/Source/Runtime/Core/Private/Misc/ConfigCacheIni.cpp
@@ -3917,6 +3917,8 @@ static FString PerformFinalExpansions(const FString InString, const FString& Pla
 	FScopeLock Lock(&UE::ConfigCacheIni::Private::ExpansionsCriticalSection);
 
 	static FString LastPlatform;
+	static FString LastProject;
+
 	static FString PlatformExtensionEngineDir;
 	static FString PlatformExtensionProjectDir;
 	static FString ProjectNotForLicenseesDir;
@@ -3925,25 +3927,27 @@ static FString PerformFinalExpansions(const FString InString, const FString& Pla
 	if (LastPlatform != PlatformName)
 	{
 		LastPlatform = PlatformName;
+
 		PlatformExtensionEngineDir = FPaths::Combine(*FPaths::EnginePlatformExtensionsDir(), *PlatformName).Replace(*FPaths::EngineDir(), *(FString(EngineDir) + "/"));
 		PlatformExtensionProjectDir = FPaths::Combine(*FPaths::ProjectPlatformExtensionsDir(), *PlatformName).Replace(*FPaths::ProjectDir(), *(FString(ProjectDir) + "/"));
 	}
 
-	// cache some slow operations
-	if (ProjectNotForLicenseesDir.Len() == 0)
+	if (LastProject != ProjectDir)
 	{
-		if (FPaths::IsUnderDirectory(FPaths::ProjectDir(), FPaths::EngineDir()))
+		LastProject = ProjectDir;
+
+		if (FPaths::IsUnderDirectory(ProjectDir, FPaths::EngineDir()))
 		{
-			FString RelativeDir = FPaths::ProjectDir();
+			FString RelativeDir = ProjectDir;
 			FPaths::MakePathRelativeTo(RelativeDir, *FPaths::EngineDir());
 			ProjectNotForLicenseesDir = FPaths::Combine(FPaths::EngineDir(), TEXT("Restricted/NotForLicensees"), RelativeDir);
 			ProjectNoRedistDir = FPaths::Combine(FPaths::EngineDir(), TEXT("Restricted/NoRedist"), RelativeDir);
 		}
 		else
 		{
-			ProjectNotForLicenseesDir = FPaths::Combine(FPaths::ProjectDir(), TEXT("Restricted/NotForLicensees"));
-			ProjectNoRedistDir = FPaths::Combine(FPaths::ProjectDir(), TEXT("Restricted/NoRedist"));
-		}
+			ProjectNotForLicenseesDir = FPaths::Combine(ProjectDir, TEXT("Restricted/NotForLicensees"));
+			ProjectNoRedistDir = FPaths::Combine(ProjectDir, TEXT("Restricted/NoRedist"));
+		}	
 	}
 
 	FString OutString = InString.Replace(TEXT("{ENGINE}"), EngineDir);
