We already call panic if we are crashing in the signal handler, avoid doing so here as well due to the same code path handling ensures
If we call panic we were no longer output to the log, after the ensure

diff --git a/Engine/Source/Runtime/Core/Private/Unix/UnixPlatformCrashContext.cpp b/Engine/Source/Runtime/Core/Private/Unix/UnixPlatformCrashContext.cpp
--- a/Engine/Source/Runtime/Core/Private/Unix/UnixPlatformCrashContext.cpp
+++ b/Engine/Source/Runtime/Core/Private/Unix/UnixPlatformCrashContext.cpp
@@ -595,7 +595,7 @@ void FUnixCrashContext::GenerateCrashInfoAndLaunchReporter(bool bReportingNonCra
 			FPaths::NormalizeDirectoryName(LogDstAbsolute);
 
 			// Flush out the log
-			GLog->Panic();
+			GLog->Flush();
 
 #if !NO_LOGGING
 			bool bMemoryOnly = FPlatformOutputDevices::GetLog()->IsMemoryOnly();
