Logging: Optimized formatting of logging and FDateTime

Formatting of FDateTime is about 7x faster for the logging use case. Formatting of log messages overall is about 3x faster with these changes.

diff --git a/Engine/Source/Runtime/Core/Private/Misc/DateTime.cpp b/Engine/Source/Runtime/Core/Private/Misc/DateTime.cpp
--- a/Engine/Source/Runtime/Core/Private/Misc/DateTime.cpp
+++ b/Engine/Source/Runtime/Core/Private/Misc/DateTime.cpp
@@ -233,43 +233,64 @@ FString FDateTime::ToString() const
 
 FString FDateTime::ToString(const TCHAR* Format) const
 {
-	TStringBuilder<32> Result;
+	TStringBuilder<48> Result;
 	ToString(Format, Result);
 	return Result.ToString();
 }
 
 void FDateTime::ToString(const TCHAR* Format, FStringBuilderBase& Result) const
 {
-	if (Format != nullptr)
+	if (UNLIKELY(!Format))
 	{
-		while (*Format != TCHAR('\0'))
+		return;
+	}
+
+	const auto AppendZeroPaddedInteger = [](FStringBuilderBase& Output, int32 Value, int32 DigitCount)
+	{
+		const int32 OutputIndex = Output.AddUninitialized(DigitCount);
+		TCHAR* const OutputData = Output.GetData();
+		for (; DigitCount--; Value /= 10)
 		{
-			if ((*Format == TCHAR('%')) && (*(++Format) != TCHAR('\0')))
-			{
-				switch (*Format)
-				{
-				case TCHAR('a'): Result.Append(IsMorning() ? TEXT("am") : TEXT("pm")); break;
-				case TCHAR('A'): Result.Append(IsMorning() ? TEXT("AM") : TEXT("PM")); break;
-				case TCHAR('d'): Result.Appendf(TEXT("%02i"), GetDay()); break;
-				case TCHAR('D'): Result.Appendf(TEXT("%03i"), GetDayOfYear()); break;
-				case TCHAR('m'): Result.Appendf(TEXT("%02i"), GetMonth()); break;
-				case TCHAR('y'): Result.Appendf(TEXT("%02i"), GetYear() % 100); break;
-				case TCHAR('Y'): Result.Appendf(TEXT("%04i"), GetYear()); break;
-				case TCHAR('h'): Result.Appendf(TEXT("%02i"), GetHour12()); break;
-				case TCHAR('H'): Result.Appendf(TEXT("%02i"), GetHour()); break;
-				case TCHAR('M'): Result.Appendf(TEXT("%02i"), GetMinute()); break;
-				case TCHAR('S'): Result.Appendf(TEXT("%02i"), GetSecond()); break;
-				case TCHAR('s'): Result.Appendf(TEXT("%03i"), GetMillisecond()); break;
-				default:		 Result.AppendChar(*Format);
-				}
-			}
-			else
-			{
-				Result.AppendChar(*Format);
-			}
+			OutputData[OutputIndex + DigitCount] = TEXT('0') + (Value % 10);
+		}
+	};
 
-			// move to the next one
-			Format++;
+	int32 Year = 0, Month = 0, Day = 0;
+	const auto GetLocalDateIfZero = [this, &Year, &Month, &Day](int32 Arg)
+	{
+		if (Arg == 0)
+		{
+			GetDate(Year, Month, Day);
+		}
+	};
+	const auto GetLocalYear = [this, &Year, &GetLocalDateIfZero] { GetLocalDateIfZero(Year); return Year; };
+	const auto GetLocalMonth = [this, &Month, &GetLocalDateIfZero] { GetLocalDateIfZero(Month); return Month; };
+	const auto GetLocalDay = [this, &Day, &GetLocalDateIfZero] { GetLocalDateIfZero(Day); return Day; };
+
+	for (; *Format; ++Format)
+	{
+		if (*Format == TCHAR('%') && *++Format)
+		{
+			switch (*Format)
+			{
+			case TCHAR('a'): Result.Append(IsMorning() ? TEXT("am") : TEXT("pm")); break;
+			case TCHAR('A'): Result.Append(IsMorning() ? TEXT("AM") : TEXT("PM")); break;
+			case TCHAR('D'): AppendZeroPaddedInteger(Result, GetDayOfYear(), 3); break;
+			case TCHAR('d'): AppendZeroPaddedInteger(Result, GetLocalDay(), 2); break;
+			case TCHAR('m'): AppendZeroPaddedInteger(Result, GetLocalMonth(), 2); break;
+			case TCHAR('y'): AppendZeroPaddedInteger(Result, GetLocalYear() % 100, 2); break;
+			case TCHAR('Y'): AppendZeroPaddedInteger(Result, GetLocalYear(), 4); break;
+			case TCHAR('h'): AppendZeroPaddedInteger(Result, GetHour12(), 2); break;
+			case TCHAR('H'): AppendZeroPaddedInteger(Result, GetHour(), 2); break;
+			case TCHAR('M'): AppendZeroPaddedInteger(Result, GetMinute(), 2); break;
+			case TCHAR('S'): AppendZeroPaddedInteger(Result, GetSecond(), 2); break;
+			case TCHAR('s'): AppendZeroPaddedInteger(Result, GetMillisecond(), 3); break;
+			default:         Result.AppendChar(*Format);
+			}
+		}
+		else
+		{
+			Result.AppendChar(*Format);
 		}
 	}
 }
diff --git a/Engine/Source/Runtime/Core/Private/Misc/OutputDeviceDebug.cpp b/Engine/Source/Runtime/Core/Private/Misc/OutputDeviceDebug.cpp
--- a/Engine/Source/Runtime/Core/Private/Misc/OutputDeviceDebug.cpp
+++ b/Engine/Source/Runtime/Core/Private/Misc/OutputDeviceDebug.cpp
@@ -3,34 +3,22 @@
 #include "Misc/OutputDeviceDebug.h"
 #include "CoreGlobals.h"
 #include "Misc/OutputDeviceHelper.h"
+#include "Misc/StringBuilder.h"
 
-/**
- * Serializes the passed in data unless the current event is suppressed.
- *
- * @param	Data	Text to log
- * @param	Event	Event name used for suppression purposes
- */
-void FOutputDeviceDebug::Serialize( const TCHAR* Data, ELogVerbosity::Type Verbosity, const class FName& Category, const double Time )
+void FOutputDeviceDebug::Serialize(const TCHAR* Data, ELogVerbosity::Type Verbosity, const FName& Category, const double Time)
 {
-	static bool Entry=false;
-	if( !GIsCriticalError || Entry )
+	if (Verbosity != ELogVerbosity::SetColor)
 	{
-		if (Verbosity != ELogVerbosity::SetColor)
-		{
-			FPlatformMisc::LowLevelOutputDebugStringf(TEXT("%s%s"),*FOutputDeviceHelper::FormatLogLine(Verbosity, Category, Data, GPrintLogTimes, Time),LINE_TERMINATOR);
-		}
-	}
-	else
-	{
-		Entry=true;
-		Serialize( Data, Verbosity, Category );
-		Entry=false;
+		TStringBuilder<512> Line;
+		FOutputDeviceHelper::AppendFormatLogLine(Line, Verbosity, Category, Data, GPrintLogTimes, Time);
+		Line.Append(LINE_TERMINATOR);
+		FPlatformMisc::LowLevelOutputDebugString(*Line);
 	}
 }
 
-void FOutputDeviceDebug::Serialize( const TCHAR* Data, ELogVerbosity::Type Verbosity, const class FName& Category )
+void FOutputDeviceDebug::Serialize(const TCHAR* Data, ELogVerbosity::Type Verbosity, const FName& Category)
 {
-	Serialize( Data, Verbosity, Category, -1.0 );
+	Serialize(Data, Verbosity, Category, -1.0);
 }
 
 bool FOutputDeviceDebug::CanBeUsedOnMultipleThreads() const
diff --git a/Engine/Source/Runtime/Core/Private/Misc/OutputDeviceHelper.cpp b/Engine/Source/Runtime/Core/Private/Misc/OutputDeviceHelper.cpp
--- a/Engine/Source/Runtime/Core/Private/Misc/OutputDeviceHelper.cpp
+++ b/Engine/Source/Runtime/Core/Private/Misc/OutputDeviceHelper.cpp
@@ -1,47 +1,51 @@
 // Copyright Epic Games, Inc. All Rights Reserved.
 
 #include "Misc/OutputDeviceHelper.h"
-#include "Misc/DateTime.h"
+#include "CoreGlobals.h"
 #include "HAL/PlatformTime.h"
-#include "Containers/StringConv.h"
 #include "Misc/App.h"
-#include "CoreGlobals.h"
-
-FString FOutputDeviceHelper::FormatLogLine( ELogVerbosity::Type Verbosity, const class FName& Category, const TCHAR* Message /*= nullptr*/, ELogTimes::Type LogTime /*= ELogTimes::None*/, const double Time /*= -1.0*/, int32* OutCategoryIndex)
+#include "Misc/DateTime.h"
+#include "Misc/StringBuilder.h"
+
+void FOutputDeviceHelper::AppendFormatLogLine(
+	FStringBuilderBase& Format,
+	const ELogVerbosity::Type Verbosity,
+	const FName& Category,
+	const TCHAR* const Message,
+	const ELogTimes::Type LogTime,
+	const double Time,
+	int32* const OutCategoryIndex)
 {
-	const bool bShowCategory = GPrintLogCategory && Category != NAME_None;
-	TStringBuilder<512> Format;
-
 	switch (LogTime)
 	{
 		case ELogTimes::SinceGStartTime:
-		{																	
+		{
 			const double RealTime = (Time < 0.0) ? (FPlatformTime::Seconds() - GStartTime) : Time;
-			Format.Append(*FString::Printf( TEXT( "[%07.2f][%3llu]" ), RealTime, GFrameCounter % 1000));
+			Format.Appendf(TEXT("[%07.2f][%3llu]"), RealTime, GFrameCounter % 1000);
 			break;
 		}
 
 		case ELogTimes::UTC:
-			Format.Append(TEXT("["));
-			FDateTime::UtcNow().ToString(TEXT("%Y.%m.%d-%H.%M.%S:%s"), Format);
-			Format.Append(TEXT("]["));
-			Format.Append(*FString::Printf(TEXT("%3llu"), GFrameCounter % 1000));
-			Format.Append(TEXT("]"));
+			FDateTime::UtcNow().ToString(TEXT("[%Y.%m.%d-%H.%M.%S:%s]"), Format);
+			Format.Appendf(TEXT("[%3llu]"), GFrameCounter % 1000);
 			break;
 
 		case ELogTimes::Local:
-			Format.Append(*FString::Printf(TEXT("[%s][%3llu]"), *FDateTime::Now().ToString(TEXT("%Y.%m.%d-%H.%M.%S:%s")), GFrameCounter % 1000));
+			FDateTime::Now().ToString(TEXT("[%Y.%m.%d-%H.%M.%S:%s]"), Format);
+			Format.Appendf(TEXT("[%3llu]"), GFrameCounter % 1000);
 			break;
 
 		case ELogTimes::Timecode:
-			Format.Append(*FString::Printf(TEXT("[%s][%3llu]"), *FApp::GetTimecode().ToString(), GFrameCounter % 1000));
+			Format.Appendf(TEXT("[%s][%3llu]"), *FApp::GetTimecode().ToString(), GFrameCounter % 1000);
 			break;
 
 		default:
 			break;
-	}	
+	}
 
-	if (OutCategoryIndex != nullptr)
+	const bool bShowCategory = GPrintLogCategory && !Category.IsNone();
+
+	if (OutCategoryIndex)
 	{
 		*OutCategoryIndex = bShowCategory ? Format.Len() : INDEX_NONE;
 	}
@@ -69,47 +73,65 @@ FString FOutputDeviceHelper::FormatLogLine( ELogVerbosity::Type Verbosity, const
 	{
 		Format.Append(Message);
 	}
+}
+
+void FOutputDeviceHelper::AppendFormatLogLine(
+	FUtf8StringBuilderBase& Format,
+	const ELogVerbosity::Type Verbosity,
+	const FName& Category,
+	const TCHAR* const Message,
+	const ELogTimes::Type LogTime,
+	const double Time,
+	int32* const OutCategoryIndex)
+{
+	TStringBuilder<128> Prefix;
+	AppendFormatLogLine(Prefix, Verbosity, Category, nullptr, LogTime, Time, OutCategoryIndex);
+	Format.Append(Prefix);
+
+	if (Message)
+	{
+		Format.Append(Message);
+	}
+}
+
+FString FOutputDeviceHelper::FormatLogLine(
+	const ELogVerbosity::Type Verbosity,
+	const FName& Category,
+	const TCHAR* const Message,
+	const ELogTimes::Type LogTime,
+	const double Time,
+	int32* const OutCategoryIndex)
+{
+	TStringBuilder<512> Format;
+	AppendFormatLogLine(Format, Verbosity, Category, Message, LogTime, Time, OutCategoryIndex);
 	return Format.ToString();
 }
 
-void FOutputDeviceHelper::FormatCastAndSerializeLine(FArchive& Output, const TCHAR* Data, ELogVerbosity::Type Verbosity, const class FName& Category, const double Time, bool bSuppressEventTag, bool bAutoEmitLineTerminator)
+void FOutputDeviceHelper::FormatCastAndSerializeLine(
+	FArchive& Output,
+	const TCHAR* const Message,
+	const ELogVerbosity::Type Verbosity,
+	const FName& Category,
+	const double Time,
+	const bool bSuppressEventTag,
+	const bool bAutoEmitLineTerminator)
 {
-	// First gather all strings to serialize (prefix, data, terminator) including their size before and after conversion
-	// then allocate the conversion buffer only once and serialize only once
-#if PLATFORM_UNIX
-	// on Linux, we still want to have logs with Windows line endings so they can be opened with Windows tools like infamous notepad.exe
-	static const TCHAR* Terminator = TEXT("\r\n");
-#else
-	static const TCHAR* Terminator = LINE_TERMINATOR;
-#endif // PLATFORM_UNIX
-	// Line Terminator won't change so we can keep everything as statics
-	static const int32 TerminatorLength = FCString::Strlen(Terminator);
-	static const int32 ConvertedTerminatorLength = FTCHARToUTF8_Convert::ConvertedLength(Terminator, TerminatorLength);
-
-	// Calculate data length
-	const int32 DataLength = FCString::Strlen(Data);
-	const int32 ConvertedDataLength = FTCHARToUTF8_Convert::ConvertedLength(Data, DataLength);
-
-	FString Prefix;
+	TUtf8StringBuilder<512> Format;
+
 	if (!bSuppressEventTag)
 	{
-		Prefix = FormatLogLine(Verbosity, Category, nullptr, GPrintLogTimes, Time);
+		AppendFormatLogLine(Format, Verbosity, Category, Message, GPrintLogTimes, Time);
+	}
+	else
+	{
+		Format.Append(Message);
 	}
 
-	// Calculate prefix length
-	const int32 PrefixLength = Prefix.Len();
-	const int32 ConvertedPrefixLength = FTCHARToUTF8_Convert::ConvertedLength(*Prefix, Prefix.Len());
-
-	// Allocate the conversion buffer. It's ok to always add some slack for the line terminator even if not required
-	TArray<UTF8CHAR, TInlineAllocator<2 * DEFAULT_STRING_CONVERSION_SIZE>> ConvertedText;
-	ConvertedText.AddUninitialized(ConvertedPrefixLength + ConvertedDataLength + (bAutoEmitLineTerminator ? ConvertedTerminatorLength : 0));
-	// Do the actual conversion to the pre-allocated buffer
-	FTCHARToUTF8_Convert::Convert(ConvertedText.GetData(), ConvertedPrefixLength, *Prefix, PrefixLength);
-	FTCHARToUTF8_Convert::Convert(ConvertedText.GetData() + ConvertedPrefixLength, ConvertedDataLength, Data, DataLength);
 	if (bAutoEmitLineTerminator)
 	{
-		FTCHARToUTF8_Convert::Convert(ConvertedText.GetData() + ConvertedPrefixLength + ConvertedDataLength, ConvertedTerminatorLength, Terminator, TerminatorLength);
+		// Use Windows line endings on Linux for compatibility with Windows tools like notepad.exe
+		Format.Append(PLATFORM_UNIX ? "\r\n" : LINE_TERMINATOR_ANSI);
 	}
-	// Serialize to the destination archive
-	Output.Serialize(ConvertedText.GetData(), ConvertedText.Num() * sizeof(ANSICHAR));
+
+	Output.Serialize(Format.GetData(), Format.Len() * sizeof(UTF8CHAR));
 }
diff --git a/Engine/Source/Runtime/Core/Public/Misc/OutputDeviceHelper.h b/Engine/Source/Runtime/Core/Public/Misc/OutputDeviceHelper.h
--- a/Engine/Source/Runtime/Core/Public/Misc/OutputDeviceHelper.h
+++ b/Engine/Source/Runtime/Core/Public/Misc/OutputDeviceHelper.h
@@ -3,22 +3,24 @@
 #pragma once
 
 #include "CoreTypes.h"
-#include "Containers/UnrealString.h"
+#include "Containers/StringFwd.h"
 #include "Logging/LogVerbosity.h"
+#include "Misc/OutputDevice.h"
+
+#if UE_ENABLE_INCLUDE_ORDER_DEPRECATED_IN_5_1
+#include "Containers/UnrealString.h"
+#endif
+
+class FArchive;
+class FName;
+class FString;
 
 /** Helper functions used by FOutputDevice derived classes **/
 struct CORE_API FOutputDeviceHelper
 {
-	/**
-	 * Converts verbosity to a string
-	 * @param Verbosity verbosity enum
-	 * @returns String representation of the verbosity enum
-	 */
-	UE_DEPRECATED(4.27, "VerbosityToString has moved to LogVerbosity.h and been renamed to ToString.")
-	inline static const TCHAR* VerbosityToString(ELogVerbosity::Type Verbosity)
-	{
-		return ToString(Verbosity);
-	}
+	/** Append a formatted log line to the string builder. */
+	static void AppendFormatLogLine(FStringBuilderBase& Output, ELogVerbosity::Type Verbosity, const FName& Category, const TCHAR* Message = nullptr, ELogTimes::Type LogTime = ELogTimes::None, double Time = -1.0, int32* OutCategoryIndex = nullptr);
+	static void AppendFormatLogLine(FUtf8StringBuilderBase& Output, ELogVerbosity::Type Verbosity, const FName& Category, const TCHAR* Message = nullptr, ELogTimes::Type LogTime = ELogTimes::None, double Time = -1.0, int32* OutCategoryIndex = nullptr);
 
 	/**
 	 * Formats a log line with date, time, category and verbosity prefix
@@ -30,7 +32,7 @@ struct CORE_API FOutputDeviceHelper
 	 * @param OutCategoryIndex (if non-null) The index of the category within the return string is written here, or INDEX_NONE if the category is suppressed
 	 * @returns Formatted log line
 	 */
-	static FString FormatLogLine(ELogVerbosity::Type Verbosity, const class FName& Category, const TCHAR* Message = nullptr, ELogTimes::Type LogTime = ELogTimes::None, const double Time = -1.0, int32* OutCategoryIndex = nullptr);
+	static FString FormatLogLine(ELogVerbosity::Type Verbosity, const FName& Category, const TCHAR* Message = nullptr, ELogTimes::Type LogTime = ELogTimes::None, double Time = -1.0, int32* OutCategoryIndex = nullptr);
 
 	/**
 	 * Formats, casts to ANSI char and serializes a message to archive. Optimized for small number of allocations and Serialize calls
@@ -42,5 +44,5 @@ struct CORE_API FOutputDeviceHelper
 	 * @param bSuppressEventTag True if the message date/time prefix should be suppressed
 	 * @param bAutoEmitLineTerminator True if the message should be automatically appended with a line terminator
 	 **/
-	static void FormatCastAndSerializeLine(class FArchive& Output, const TCHAR* Message, ELogVerbosity::Type Verbosity, const class FName& Category, const double Time, bool bSuppressEventTag, bool bAutoEmitLineTerminator);
+	static void FormatCastAndSerializeLine(FArchive& Output, const TCHAR* Message, ELogVerbosity::Type Verbosity, const FName& Category, double Time, bool bSuppressEventTag, bool bAutoEmitLineTerminator);
 };
