Make replay serialization errors (FBitReader::SetOverflowed) more obvious

diff --git a/Engine/Source/Runtime/Core/Private/Serialization/BitReader.cpp b/Engine/Source/Runtime/Core/Private/Serialization/BitReader.cpp
--- a/Engine/Source/Runtime/Core/Private/Serialization/BitReader.cpp
+++ b/Engine/Source/Runtime/Core/Private/Serialization/BitReader.cpp
@@ -11,6 +11,11 @@ PRAGMA_DISABLE_UNSAFE_TYPECAST_WARNINGS
 extern const uint8 GShift[8]={0x01,0x02,0x04,0x08,0x10,0x20,0x40,0x80};
 extern const uint8 GMask [8]={0x00,0x01,0x03,0x07,0x0f,0x1f,0x3f,0x7f};
 
+static TAutoConsoleVariable<bool> CVarLogFatalOnOverflow(
+	TEXT("BitReader.LogFatalOnOverflow"),
+	false,
+	TEXT("LogFatal if BitReader Overflows"));
+
 // Optimized arbitrary bit range memory copy routine.
 
 void appBitsCpy( uint8* Dest, int32 DestBit, uint8* Src, int32 SrcBit, int32 BitCount )
@@ -242,8 +247,16 @@ void FBitReader::CountMemory(FArchive& Ar) const
 
 void FBitReader::SetOverflowed(int64 LengthBits)
 {
-	UE_LOG(LogNetSerialization, Error, TEXT("FBitReader::SetOverflowed() called! (ReadLen: %i, Remaining: %i, Max: %i)"),
+	if (CVarLogFatalOnOverflow.GetValueOnAnyThread())
+	{
+		UE_LOG(LogNetSerialization, Fatal, TEXT("FBitReader::SetOverflowed() called! (ReadLen: %i, Remaining: %i, Max: %i)"),
 			LengthBits, (Num - Pos), Num);
+	}
+	else
+	{
+		ensureMsgf(false, TEXT("FBitReader::SetOverflowed() called! (ReadLen: %i, Remaining: %i, Max: %i)"),
+			LengthBits, (Num - Pos), Num);
+	}
 
 	SetError();
 }
diff --git a/Engine/Source/Runtime/Engine/Private/RepLayout.cpp b/Engine/Source/Runtime/Engine/Private/RepLayout.cpp
--- a/Engine/Source/Runtime/Engine/Private/RepLayout.cpp
+++ b/Engine/Source/Runtime/Engine/Private/RepLayout.cpp
@@ -4025,6 +4025,12 @@ bool FRepLayout::ReceiveProperties_BackwardsCompatible_r(
 				UE_LOG(LogRep, Warning, TEXT("ReceiveProperties_BackwardsCompatible_r: Property didn't read proper number of bits. Owner: %s, Name: %s, NetFieldExportHandle: %i, Checksum: %u, BitsLeft:%d"), *Owner->GetName(), *NetFieldExportGroup->NetFieldExports[NetFieldExportHandle].ExportName.ToString(), NetFieldExportHandle, Checksum, TempReader.GetBitsLeft());
 				return false;
 			}
+
+			if (TempReader.IsError())
+			{
+				UE_LOG(LogRep, Warning, TEXT("ReceiveProperties_BackwardsCompatible_r: Error Reading Property. Owner: %s, Name: %s, NetFieldExportHandle: %i, Checksum: %u"), *Owner->GetName(), *NetFieldExportGroup->NetFieldExports[NetFieldExportHandle].ExportName.ToString(), NetFieldExportHandle, Checksum);
+				return false;
+			}
 		}
 	}
 
