Fix issue with externally connected trace connections and status message.

diff --git a/Engine/Source/Runtime/Core/Private/ProfilingDebugging/TraceAuxiliary.cpp b/Engine/Source/Runtime/Core/Private/ProfilingDebugging/TraceAuxiliary.cpp
--- a/Engine/Source/Runtime/Core/Private/ProfilingDebugging/TraceAuxiliary.cpp
+++ b/Engine/Source/Runtime/Core/Private/ProfilingDebugging/TraceAuxiliary.cpp
@@ -682,15 +682,24 @@ static void TraceAuxiliaryStatus()
 	UE_LOG(LogConsoleResponse, Display, TEXT("Trace status ----------------------------------------------------------"));
 
 	// Status of data connection
-	const TCHAR* Dest = GTraceAuxiliary.GetDest();
 	TStringBuilder<256> ConnectionStr;
-	if (Dest && FCString::Strlen(Dest) > 0)
+	if (UE::Trace::IsTracing())
 	{
-		ConnectionStr.Appendf(TEXT("Tracing to '%s'"), Dest);
+		const TCHAR* Dest = GTraceAuxiliary.GetDest();
+		if (Dest && FCString::Strlen(Dest) > 0)
+		{
+			ConnectionStr.Appendf(TEXT("Tracing to '%s'"), Dest);
+		}
+		else 
+		{
+			// If GTraceAux doesn't know about the target but we are still tracing this is an externally initiated connection
+			// (e.g. connection command from Insights).
+			ConnectionStr = TEXT("Tracing to unknown target (externally set)");
+		}
 	}
 	else
 	{
-		ConnectionStr.Append(TEXT("Not tracing"));
+		ConnectionStr = TEXT("Not tracing");
 	}
 	UE_LOG(LogConsoleResponse, Display, TEXT("- Connection: %s"), ConnectionStr.ToString());
 
