Workaround for MSVC 19.31.31107 internal compiler error in some build configurations, related to calling member functions of OffsetLength multiple times

diff --git a/Engine/Source/Runtime/Core/Private/IO/IoStore.cpp b/Engine/Source/Runtime/Core/Private/IO/IoStore.cpp
--- a/Engine/Source/Runtime/Core/Private/IO/IoStore.cpp
+++ b/Engine/Source/Runtime/Core/Private/IO/IoStore.cpp
@@ -641,8 +641,8 @@ public:
 		}
 
 		const uint64 CompressionBlockSize = TocResource.Header.CompressionBlockSize;
-		int32 FirstBlockIndex = int32(OffsetLength.GetOffset() / CompressionBlockSize);
-		int32 LastBlockIndex = int32((Align(OffsetLength.GetOffset() + OffsetLength.GetLength(), CompressionBlockSize) - 1) / CompressionBlockSize);
+		int32 FirstBlockIndex = int32(ChunkInfo.Offset / CompressionBlockSize);
+		int32 LastBlockIndex = int32((Align(ChunkInfo.Offset + ChunkInfo.Size, CompressionBlockSize) - 1) / CompressionBlockSize);
 
 		ChunkInfo.CompressedSize = 0;
 		ChunkInfo.PartitionIndex = -1;
