Logging: Fixed a static analysis warning

#rnx

diff --git a/Engine/Source/Runtime/Core/Private/Misc/OutputDeviceHelper.cpp b/Engine/Source/Runtime/Core/Private/Misc/OutputDeviceHelper.cpp
--- a/Engine/Source/Runtime/Core/Private/Misc/OutputDeviceHelper.cpp
+++ b/Engine/Source/Runtime/Core/Private/Misc/OutputDeviceHelper.cpp
@@ -130,7 +130,14 @@ void FOutputDeviceHelper::FormatCastAndSerializeLine(
 	if (bAutoEmitLineTerminator)
 	{
 		// Use Windows line endings on Linux for compatibility with Windows tools like notepad.exe
-		Format.Append(PLATFORM_UNIX ? "\r\n" : LINE_TERMINATOR_ANSI);
+		if constexpr (PLATFORM_UNIX)
+		{
+			Format.Append("\r\n");
+		}
+		else
+		{
+			Format.Append(LINE_TERMINATOR_ANSI);
+		}
 	}
 
 	Output.Serialize(Format.GetData(), Format.Len() * sizeof(UTF8CHAR));
