#ue5 - LWC: Improve precision of some math functions (FMod / VectorMod) on SSE by using SVML where possible. FMath::Fmod() changed to use standard fmod() function. Fixed ambiguous naming conventions for many functions (inconsistency between InvSqrt, Reciprocal, etc and their estimate functions). Now the default for such functions is to be accurate, and the estimate functions are explicitly named with "Estimate".

Added missing vectorized implementation for VectorPow() for doubles. Fixed some missing VectorSqrt() overrides.

Did a small pass at fixing up some missing functions in UnrealMathFPU so it can pass start tests and HeadlessChaos tests when enabled, though there are a few more that need updates for Niagara (to come soon).

diff --git a/Engine/Source/Runtime/Core/Private/GenericPlatform/GenericPlatformMath.cpp b/Engine/Source/Runtime/Core/Private/GenericPlatform/GenericPlatformMath.cpp
--- a/Engine/Source/Runtime/Core/Private/GenericPlatform/GenericPlatformMath.cpp
+++ b/Engine/Source/Runtime/Core/Private/GenericPlatform/GenericPlatformMath.cpp
@@ -88,51 +88,25 @@ double FGenericPlatformMath::Atan2(double Y, double X)
 float FGenericPlatformMath::Fmod(float X, float Y)
 {
 	const float AbsY = FMath::Abs(Y);
-	if (AbsY <= 1.e-8f)
+	if (AbsY <= UE_SMALL_NUMBER) // Note: this constant should match that used by VectorMod() implementations
 	{
 		FmodReportError(X, Y);
 		return 0.0;
 	}
 
-	// Convert to double for better precision, since intermediate rounding can lose enough precision to skew the result.
-	const double DX = double(X);
-	const double DY = double(Y);
-
-	const double Div = (DX / DY);
-	const double IntPortion = DY * TruncToDouble(Div);
-	const double Result = DX - IntPortion;
-	// Convert back to float. This is safe because the result will by definition not exceed the X input.
-	return float(Result);
+	return fmodf(X, Y);
 }
 
 double FGenericPlatformMath::Fmod(double X, double Y)
 {
 	const double AbsY = FMath::Abs(Y);
-	if (AbsY <= 1.e-8)
+	if (AbsY <= UE_DOUBLE_SMALL_NUMBER) // Note: this constant should match that used by VectorMod() implementations
 	{
 		FmodReportError(X, Y);
 		return 0.0;
 	}
 
-#if 1
-	// Due to the lack of standard support across platforms for `long double`, we can't rely on this to always minimize the intermediate precision loss of division and multiplication.
-	// As a result, due to small precision loss the result here could be different than the std library version of fmod(), but our validation tests in UnrealMathTest should show
-	// that when different the results are within a small delta from either 0 or the value of Y. We could just use the library version here of course, but it is slower and also will not
-	// match the vectorized versions of Fmod (VectorRegisterMod) that use division and truncation.
-	const double Div = (X / Y);
-	const double IntPortion = Y * FMath::TruncToDouble(Div);
-	const double Result = X - IntPortion;
-	return Result;
-#elif 0
-	const long double LDX = (long double)(X);
-	const long double LDY = (long double)(Y);
-	const long double Div = (LDX / LDY);
-	const long double IntPortion = LDY * truncl(Div);
-	const long double Result = LDX - IntPortion;
-	return double(Result);
-#else
 	return fmod(X, Y);
-#endif
 }
 
 void FGenericPlatformMath::FmodReportError(float X, float Y)
diff --git a/Engine/Source/Runtime/Core/Private/Tests/Math/UnrealMathTest.cpp b/Engine/Source/Runtime/Core/Private/Tests/Math/UnrealMathTest.cpp
--- a/Engine/Source/Runtime/Core/Private/Tests/Math/UnrealMathTest.cpp
+++ b/Engine/Source/Runtime/Core/Private/Tests/Math/UnrealMathTest.cpp
@@ -32,7 +32,7 @@ alignas(alignof(VectorRegister4Float)) static float GScratch[16];
 alignas(alignof(VectorRegister4Double)) static double GScratchDouble[16];
 static float GSum;
 static double GSumDouble;
-static bool GPassing;
+static bool GPassing = true;
 
 #define MATHTEST_INLINE FORCENOINLINE // if you want to do performance testing change to FORCEINLINE or FORCEINLINE_DEBUGGABLE
 
@@ -1645,6 +1645,12 @@ bool RunDoubleVectorTest()
 	V0 = MakeVectorRegisterDouble(-3.0, 6.0, -3.0, 0.0);
 	LogTest<double>(TEXT("VectorCross"), TestVectorsEqual(V0, V1));
 
+	V0 = MakeVectorRegisterDouble(2.0, 4.0, 6.0, 8.0);
+	V1 = MakeVectorRegisterDouble(4.0, 3.0, 2.0, 1.0);
+	V1 = VectorPow(V0, V1);
+	V0 = MakeVectorRegister(16.0, 64.0, 36.0, 8.0);
+	LogTest<double>(TEXT("VectorPow"), TestVectorsEqual(V0, V1, 0.001));
+
 	// Vector component comparisons
 
 	// VectorCompareGT
@@ -1872,25 +1878,45 @@ bool RunDoubleVectorTest()
 	LogTest<double>(TEXT("VectorMergeXYZ_VecW-2"), TestVectorsEqual(V2, V3));
 
 	V0 = MakeVectorRegister(1.0, 1.0e6, 1.3e-8, 35.0);
-	V1 = VectorReciprocal(V0);
+	V1 = VectorReciprocalEstimate(V0);
 	V3 = VectorMultiply(V1, V0);
-	LogTest<double>(TEXT("VectorReciprocal"), TestVectorsEqual(VectorOneDouble(), V3, 0.008));
+	LogTest<double>(TEXT("VectorReciprocalEstimate"), TestVectorsEqual(VectorOneDouble(), V3, 0.008));
 
 	V0 = MakeVectorRegister(1.0, 1.0e6, 1.3e-8, 35.0);
-	V1 = VectorReciprocalAccurate(V0);
+	V1 = VectorReciprocal(V0);
 	V3 = VectorMultiply(V1, V0);
 	LogTest<double>(TEXT("VectorReciprocalAccurate"), TestVectorsEqual(VectorOneDouble(), V3, 1e-7));
 
 	V0 = MakeVectorRegister(1.0, 1.0e6, 1.3e-8, 35.0);
-	V1 = VectorReciprocalSqrt(V0);
+	V1 = VectorReciprocalSqrtEstimate(V0);
 	V3 = VectorMultiply(VectorMultiply(V1, V1), V0);
-	LogTest<double>(TEXT("VectorReciprocalSqrt"), TestVectorsEqual(VectorOneDouble(), V3, 0.007));
+	LogTest<double>(TEXT("VectorReciprocalSqrtEstimate"), TestVectorsEqual(VectorOneDouble(), V3, 0.007));
 
 	V0 = MakeVectorRegister(1.0, 1.0e6, 1.3e-8, 35.0);
-	V1 = VectorReciprocalSqrtAccurate(V0);
+	V1 = VectorReciprocalSqrt(V0);
 	V3 = VectorMultiply(VectorMultiply(V1, V1), V0);
 	LogTest<double>(TEXT("VectorReciprocalSqrtAccurate"), TestVectorsEqual(VectorOneDouble(), V3, 1e-6));
 
+	V0 = MakeVectorRegister(2.0f, -2.0f, 2.0f, -2.0f);
+	V1 = VectorReciprocalLenEstimate(V0);
+	V0 = MakeVectorRegister(0.25f, 0.25f, 0.25f, 0.25f);
+	LogTest<double>(TEXT("VectorReciprocalLenEstimate"), TestVectorsEqual(V0, V1, 0.004));
+
+	V0 = MakeVectorRegister(2.0f, -2.0f, 2.0f, -2.0f);
+	V1 = VectorReciprocalLen(V0);
+	V0 = MakeVectorRegister(0.25f, 0.25f, 0.25f, 0.25f);
+	LogTest<double>(TEXT("VectorReciprocalLenAccurate"), TestVectorsEqual(V0, V1, 0.0001));
+
+	V0 = MakeVectorRegister(2.0f, -2.0f, 2.0f, -2.0f);
+	V1 = VectorNormalizeEstimate(V0);
+	V0 = MakeVectorRegister(0.5f, -0.5f, 0.5f, -0.5f);
+	LogTest<double>(TEXT("VectorNormalizeEstimate"), TestVectorsEqual(V0, V1, 0.004));
+
+	V0 = MakeVectorRegister(2.0f, -2.0f, 2.0f, -2.0f);
+	V1 = VectorNormalize(V0);
+	V0 = MakeVectorRegister(0.5f, -0.5f, 0.5f, -0.5f);
+	LogTest<double>(TEXT("VectorNormalizeAccurate"), TestVectorsEqual(V0, V1, 1e-8));
+
 	// VectorMod
 	V0 = MakeVectorRegister(0.0, 3.2, 2.8, 1.5);
 	V1 = MakeVectorRegister(2.0, 1.2, 2.0, 3.0);
@@ -1904,6 +1930,12 @@ bool RunDoubleVectorTest()
 	V3 = VectorMod(V0, V1);
 	LogTest<double>(TEXT("VectorMod negative"), TestVectorsEqual(V2, V3));
 
+	V0 = MakeVectorRegister(89.9, 180.0, -256.0, -270.1);
+	V1 = MakeVectorRegister(360.0, 0.1, 360.0, 180.0);
+	V2 = TestReferenceMod(V0, V1);
+	V3 = VectorMod(V0, V1);
+	LogTest<double>(TEXT("VectorMod common"), TestVectorsEqual(V2, V3));
+
 	// VectorSign
 	V0 = MakeVectorRegister(2.0, -2.0, 0.0, -3.0);
 	V2 = MakeVectorRegister(1.0, -1.0, 1.0, -1.0);
@@ -2360,17 +2392,22 @@ bool FVectorRegisterAbstractionTest::RunTest(const FString& Parameters)
 	LogTest<float>( TEXT("VectorPow"), TestVectorsEqual( V0, V1, 0.001f ) );
 
 	V0 = MakeVectorRegister( 2.0f, -2.0f, 2.0f, -2.0f );
-	V1 = VectorReciprocalLen( V0 );
+	V1 = VectorReciprocalLenEstimate( V0 );
 	V0 = MakeVectorRegister( 0.25f, 0.25f, 0.25f, 0.25f );
-	LogTest<float>( TEXT("VectorReciprocalLen"), TestVectorsEqual( V0, V1, 0.004f ) );
+	LogTest<float>( TEXT("VectorReciprocalLenEstimate"), TestVectorsEqual( V0, V1, 0.004f ) );
+
+	V0 = MakeVectorRegister(2.0f, -2.0f, 2.0f, -2.0f);
+	V1 = VectorReciprocalLen(V0);
+	V0 = MakeVectorRegister(0.25f, 0.25f, 0.25f, 0.25f);
+	LogTest<float>(TEXT("VectorReciprocalLenAccurate"), TestVectorsEqual(V0, V1, 0.0001f));
 
 	V0 = MakeVectorRegister( 2.0f, -2.0f, 2.0f, -2.0f );
-	V1 = VectorNormalize( V0 );
+	V1 = VectorNormalizeEstimate( V0 );
 	V0 = MakeVectorRegister( 0.5f, -0.5f, 0.5f, -0.5f );
-	LogTest<float>( TEXT("VectorNormalize"), TestVectorsEqual( V0, V1, 0.004f ) );
+	LogTest<float>( TEXT("VectorNormalizeEstimate"), TestVectorsEqual( V0, V1, 0.004f ) );
 
 	V0 = MakeVectorRegister(2.0f, -2.0f, 2.0f, -2.0f);
-	V1 = VectorNormalizeAccurate(V0);
+	V1 = VectorNormalize(V0);
 	V0 = MakeVectorRegister(0.5f, -0.5f, 0.5f, -0.5f);
 	LogTest<float>(TEXT("VectorNormalizeAccurate"), TestVectorsEqual(V0, V1, 1e-8f));
 
@@ -2608,22 +2645,22 @@ bool FVectorRegisterAbstractionTest::RunTest(const FString& Parameters)
 	LogTest<float>( TEXT("VectorMergeXYZ_VecW-2"), TestVectorsEqual( V2, V3 ) );
 
 	V0 = MakeVectorRegister( 1.0f, 1.0e6f, 1.3e-8f, 35.0f );
-	V1 = VectorReciprocal( V0 );
+	V1 = VectorReciprocalEstimate( V0 );
 	V3 = VectorMultiply(V1, V0);
-	LogTest<float>( TEXT("VectorReciprocal"), TestVectorsEqual( VectorOne(), V3, 0.008f ) );
+	LogTest<float>( TEXT("VectorReciprocalEstimate"), TestVectorsEqual( VectorOne(), V3, 0.008f ) );
 
 	V0 = MakeVectorRegister( 1.0f, 1.0e6f, 1.3e-8f, 35.0f );
-	V1 = VectorReciprocalAccurate( V0 );
+	V1 = VectorReciprocal( V0 );
 	V3 = VectorMultiply(V1, V0);
 	LogTest<float>( TEXT("VectorReciprocalAccurate"), TestVectorsEqual( VectorOne(), V3, 1e-7f) );
 
 	V0 = MakeVectorRegister( 1.0f, 1.0e6f, 1.3e-8f, 35.0f );
-	V1 = VectorReciprocalSqrt( V0 );
+	V1 = VectorReciprocalSqrtEstimate( V0 );
 	V3 = VectorMultiply(VectorMultiply(V1, V1), V0);
-	LogTest<float>( TEXT("VectorReciprocalSqrt"), TestVectorsEqual( VectorOne(), V3, 0.007f ) );
+	LogTest<float>( TEXT("VectorReciprocalSqrtEstimate"), TestVectorsEqual( VectorOne(), V3, 0.007f ) );
 
 	V0 = MakeVectorRegister( 1.0f, 1.0e6f, 1.3e-8f, 35.0f );
-	V1 = VectorReciprocalSqrtAccurate( V0 );
+	V1 = VectorReciprocalSqrt( V0 );
 	V3 = VectorMultiply(VectorMultiply(V1, V1), V0);
 	LogTest<float>( TEXT("VectorReciprocalSqrtAccurate"), TestVectorsEqual( VectorOne(), V3, 1e-6f ) );
 
@@ -2652,6 +2689,12 @@ bool FVectorRegisterAbstractionTest::RunTest(const FString& Parameters)
 	V3 = VectorMod(V0, V1);
 	LogTest<float>( TEXT("VectorMod negative"), TestVectorsEqual(V2, V3));
 
+	V0 = MakeVectorRegister(89.9f, 180.0f, -256.0f, -270.1f);
+	V1 = MakeVectorRegister(360.0f, 0.1f, 360.0f, 180.0f);
+	V2 = TestReferenceMod(V0, V1);
+	V3 = VectorMod(V0, V1);
+	LogTest<float>(TEXT("VectorMod common"), TestVectorsEqual(V2, V3));
+
 	// VectorSign
 	V0 = MakeVectorRegister(2.0f, -2.0f, 0.0f, -3.0f);
 	V2 = MakeVectorRegister(1.0f, -1.0f, 1.0f, -1.0f);
diff --git a/Engine/Source/Runtime/Core/Public/Math/UnrealMathFPU.h b/Engine/Source/Runtime/Core/Public/Math/UnrealMathFPU.h
--- a/Engine/Source/Runtime/Core/Public/Math/UnrealMathFPU.h
+++ b/Engine/Source/Runtime/Core/Public/Math/UnrealMathFPU.h
@@ -36,7 +36,15 @@ struct VectorRegister2Double
 */
 struct alignas(16) VectorRegister4Double
 {
-	double V[4];
+	union
+	{
+		struct
+		{
+			VectorRegister2Double XY;
+			VectorRegister2Double ZW;
+		};
+		double V[4];
+	};
 
 	VectorRegister4Double() = default;
 
@@ -335,7 +343,8 @@ FORCEINLINE VectorRegister4Double VectorOneDouble(void)
  * @param ComponentIndex	Which component to get, X=0, Y=1, Z=2, W=3
  * @return					The component as a float
  */
-FORCEINLINE float VectorGetComponent(const VectorRegister4Float& Vec, uint32 ComponentIndex)
+template <uint32 ComponentIndex>
+FORCEINLINE float VectorGetComponentImpl(const VectorRegister4Float& Vec)
 {
 	return Vec.V[ComponentIndex];
 }
@@ -345,7 +354,8 @@ FORCEINLINE float VectorGetComponentDynamic(const VectorRegister4Float& Vec, uin
 	return Vec.V[ComponentIndex];
 }
 
-FORCEINLINE double VectorGetComponent(const VectorRegister4Double& Vec, uint32 ComponentIndex)
+template <uint32 ComponentIndex>
+FORCEINLINE double VectorGetComponentImpl(const VectorRegister4Double& Vec)
 {
 	return Vec.V[ComponentIndex];
 }
@@ -355,6 +365,8 @@ FORCEINLINE double VectorGetComponentDynamic(const VectorRegister4Double& Vec, u
 	return Vec.V[ComponentIndex];
 }
 
+#define VectorGetComponent(Vec, ComponentIndex) VectorGetComponentImpl<ComponentIndex>(Vec)
+
 
 /**
  * Loads 4 FLOATs from unaligned memory.
@@ -1203,10 +1215,20 @@ FORCEINLINE VectorRegister4Double VectorPow(const VectorRegister4Double& Base, c
 	return Vec;
 }
 
+FORCEINLINE VectorRegister4Float VectorSqrt(const VectorRegister4Float& Vec)
+{
+	return MakeVectorRegisterFloat(FMath::Sqrt(Vec.V[0]), FMath::Sqrt(Vec.V[1]), FMath::Sqrt(Vec.V[2]), FMath::Sqrt(Vec.V[3]));
+}
+
+FORCEINLINE VectorRegister4Double VectorSqrt(const VectorRegister4Double& Vec)
+{
+	return MakeVectorRegisterDouble(FMath::Sqrt(Vec.V[0]), FMath::Sqrt(Vec.V[1]), FMath::Sqrt(Vec.V[2]), FMath::Sqrt(Vec.V[3]));
+}
+
 /**
 * Returns an estimate of 1/sqrt(c) for each component of the vector
 *
-* @param Vector		Vector 
+* @param Vector		Vector
 * @return			VectorRegister(1/sqrt(t), 1/sqrt(t), 1/sqrt(t), 1/sqrt(t))
 */
 FORCEINLINE VectorRegister4Float VectorReciprocalSqrt(const VectorRegister4Float& Vec)
@@ -1219,21 +1241,52 @@ FORCEINLINE VectorRegister4Double VectorReciprocalSqrt(const VectorRegister4Doub
 	return MakeVectorRegisterDouble(1.0 / FMath::Sqrt(Vec.V[0]), 1.0 / FMath::Sqrt(Vec.V[1]), 1.0 / FMath::Sqrt(Vec.V[2]), 1.0 / FMath::Sqrt(Vec.V[3]));
 }
 
+/**
+* Returns an estimate of 1/sqrt(c) for each component of the vector
+*
+* @param Vector		Vector 
+* @return			VectorRegister(1/sqrt(t), 1/sqrt(t), 1/sqrt(t), 1/sqrt(t))
+*/
+FORCEINLINE VectorRegister4Float VectorReciprocalSqrtEstimate(const VectorRegister4Float& Vec)
+{
+	return MakeVectorRegisterFloat(FMath::InvSqrtEst(Vec.V[0]), FMath::InvSqrtEst(Vec.V[1]), FMath::InvSqrtEst(Vec.V[2]), FMath::InvSqrtEst(Vec.V[3]));
+}
+
+FORCEINLINE VectorRegister4Double VectorReciprocalSqrtEstimate(const VectorRegister4Double& Vec)
+{
+	return MakeVectorRegisterDouble(FMath::InvSqrtEst(Vec.V[0]), FMath::InvSqrtEst(Vec.V[1]), FMath::InvSqrtEst(Vec.V[2]), FMath::InvSqrtEst(Vec.V[3]));
+}
 
 /**
- * Computes an estimate of the reciprocal of a vector (component-wise) and returns the result.
+ * Computes the reciprocal of a vector (component-wise) and returns the result.
  *
  * @param Vec	1st vector
- * @return		VectorRegister( (Estimate) 1.0f / Vec.x, (Estimate) 1.0f / Vec.y, (Estimate) 1.0f / Vec.z, (Estimate) 1.0f / Vec.w )
+ * @return		VectorRegister4Float( 1.0f / Vec.x, 1.0f / Vec.y, 1.0f / Vec.z, 1.0f / Vec.w )
  */
 FORCEINLINE VectorRegister4Float VectorReciprocal(const VectorRegister4Float& Vec)
 {
-	return MakeVectorRegisterFloat(1.0f / Vec.V[0], 1.0f / Vec.V[1], 1.0f / Vec.V[2], 1.0f / Vec.V[3]);
+	return VectorDivide(GlobalVectorConstants::FloatOne, Vec);
 }
 
 FORCEINLINE VectorRegister4Double VectorReciprocal(const VectorRegister4Double& Vec)
 {
-	return MakeVectorRegisterDouble(1.0 / Vec.V[0], 1.0 / Vec.V[1], 1.0 / Vec.V[2], 1.0 / Vec.V[3]);
+	return VectorDivide(GlobalVectorConstants::DoubleOne, Vec);
+}
+
+/**
+ * Computes an estimate of the reciprocal of a vector (component-wise) and returns the result.
+ *
+ * @param Vec	1st vector
+ * @return		VectorRegister( (Estimate) 1.0f / Vec.x, (Estimate) 1.0f / Vec.y, (Estimate) 1.0f / Vec.z, (Estimate) 1.0f / Vec.w )
+ */
+FORCEINLINE VectorRegister4Float VectorReciprocalEstimate(const VectorRegister4Float& Vec)
+{
+	return VectorReciprocal(Vec);
+}
+
+FORCEINLINE VectorRegister4Double VectorReciprocalEstimate(const VectorRegister4Double& Vec)
+{
+	return VectorReciprocal(Vec);
 }
 
 /**
@@ -1268,40 +1321,22 @@ FORCEINLINE VectorRegister4Double VectorReciprocalLen(const VectorRegister4Doubl
 	return Result;
 }
 
-
 /**
-* Return the reciprocal of the square root of each component
-*
-* @param Vec		Vector 
-* @return			VectorRegister(1/sqrt(Vec.X), 1/sqrt(Vec.Y), 1/sqrt(Vec.Z), 1/sqrt(Vec.W))
-*/
-FORCEINLINE VectorRegister4Float VectorReciprocalSqrtAccurate(const VectorRegister4Float& Vec)
-{
-	return VectorReciprocalSqrt(Vec);
-}
-
-FORCEINLINE VectorRegister4Double VectorReciprocalSqrtAccurate(const VectorRegister4Double& Vec)
-{
-	return VectorReciprocalSqrt(Vec);
-}
-
-/**
- * Computes the reciprocal of a vector (component-wise) and returns the result.
+ * Return Reciprocal Length of the vector (estimate)
  *
- * @param Vec	1st vector
- * @return		VectorRegister( 1.0f / Vec.x, 1.0f / Vec.y, 1.0f / Vec.z, 1.0f / Vec.w )
+ * @param Vector		Vector
+ * @return			VectorRegister4Float(rlen, rlen, rlen, rlen) when rlen = 1/sqrt(dot4(V))
  */
-FORCEINLINE VectorRegister4Float VectorReciprocalAccurate(const VectorRegister4Float& Vec)
+FORCEINLINE VectorRegister4Float VectorReciprocalLenEstimate(const VectorRegister4Float& Vector)
 {
-	return VectorReciprocal(Vec);
+	return VectorReciprocalLen(Vector);
 }
 
-FORCEINLINE VectorRegister4Double VectorReciprocalAccurate(const VectorRegister4Double& Vec)
+FORCEINLINE VectorRegister4Double VectorReciprocalLenEstimate(const VectorRegister4Double& Vector)
 {
-	return VectorReciprocal(Vec);
+	return VectorReciprocalLen(Vector);
 }
 
-
 /**
 * Loads XYZ and sets W=0
 *
@@ -2227,13 +2262,23 @@ FORCEINLINE VectorRegister4Float VectorTruncate(const VectorRegister4Float& Vec)
 	return MakeVectorRegisterFloat(FMath::TruncToFloat(Vec.V[0]), FMath::TruncToFloat(Vec.V[1]), FMath::TruncToFloat(Vec.V[2]), FMath::TruncToFloat(Vec.V[3]));
 }
 
-
 FORCEINLINE VectorRegister4Double VectorTruncate(const VectorRegister4Double& Vec)
 {
 	return MakeVectorRegisterDouble(FMath::TruncToDouble(Vec.V[0]), FMath::TruncToDouble(Vec.V[1]), FMath::TruncToDouble(Vec.V[2]), FMath::TruncToDouble(Vec.V[3]));
 }
 
 
+FORCEINLINE VectorRegister4Float VectorRound(const VectorRegister4Float& Vec)
+{
+	return MakeVectorRegisterFloat(FMath::RoundToFloat(Vec.V[0]), FMath::RoundToFloat(Vec.V[1]), FMath::RoundToFloat(Vec.V[2]), FMath::RoundToFloat(Vec.V[3]));
+}
+
+FORCEINLINE VectorRegister4Double VectorRound(const VectorRegister4Double& Vec)
+{
+	return MakeVectorRegisterDouble(FMath::RoundToDouble(Vec.V[0]), FMath::RoundToDouble(Vec.V[1]), FMath::RoundToDouble(Vec.V[2]), FMath::RoundToDouble(Vec.V[3]));
+}
+
+
 FORCEINLINE VectorRegister4Float VectorMod(const VectorRegister4Float& X, const VectorRegister4Float& Y)
 {
 	return MakeVectorRegisterFloat(
@@ -2641,3 +2686,18 @@ FORCEINLINE VectorRegister4Int VectorIntLoad1(const void* Ptr)
 		IntSplat,
 		IntSplat);
 }
+
+//
+// TODO: Engine doesn't currently compile with PLATFORM_ENABLE_VECTORINTRINSICS=0, but when it does these need to be fixed up.
+//
+
+#define VectorSetZero()								MakeVectorRegisterFloat(0.f, 0.f, 0.f, 0.f)
+#define VectorSet1(F)								VectorSetFloat1(F)
+#define VectorIntSet1(F)							MakeVectorRegisterInt(F, F, F, F)
+#define VectorCastIntToFloat(Vec)                   VectorLoad((float*)(Vec.V))
+#define VectorCastFloatToInt(Vec)                   VectorIntLoad(Vec.V)
+#define VectorShuffleImmediate(Vec, I0, I1, I2, I3) VectorShuffle(Vec, Vec, I3, I2, I1, I0) // ShuffleImmediate shuffle is reversed from our logical shuffle.
+#define VectorShiftLeftImm(Vec, ImmAmt)             static_assert(false, "Unimplemented") // TODO: implement
+#define VectorShiftRightImmArithmetic(Vec, ImmAmt)  static_assert(false, "Unimplemented") // TODO: implement
+#define VectorShiftRightImmLogical(Vec, ImmAmt)		static_assert(false, "Unimplemented") // TODO: implement
+#define VectorIntExpandLow16To32(V0)				static_assert(false, "Unimplemented") // TODO: implement
diff --git a/Engine/Source/Runtime/Core/Public/Math/UnrealMathNeon.h b/Engine/Source/Runtime/Core/Public/Math/UnrealMathNeon.h
--- a/Engine/Source/Runtime/Core/Public/Math/UnrealMathNeon.h
+++ b/Engine/Source/Runtime/Core/Public/Math/UnrealMathNeon.h
@@ -907,12 +907,6 @@
 	return Result;
 }
 
-FORCEINLINE VectorRegister4Float VectorLerp(const VectorRegister4Float &Vec1, const VectorRegister4Float &Vec2, const VectorRegister4Float &Vec3)
-{
-	VectorRegister4Float SubVec = VectorSubtract(GlobalVectorConstants::FloatOne, Vec3);
-	return VectorMultiplyAdd(Vec2, Vec3, VectorMultiply(Vec1, SubVec));
-}
-
 /**
  * Multiplies two vectors (component-wise) and subtracts the result from the third vector.
  *
@@ -1588,36 +1582,17 @@
 }
 
 /**
-* Returns an estimate of 1/sqrt(c) for each component of the vector
-*
-* @param Vector		Vector 
-* @return			VectorRegister4Float(1/sqrt(t), 1/sqrt(t), 1/sqrt(t), 1/sqrt(t))
-*/
-FORCEINLINE VectorRegister4Float VectorReciprocalSqrt(const VectorRegister4Float& Vec)
-{
-	return vrsqrteq_f32(Vec);
-}
-
-FORCEINLINE VectorRegister4Double VectorReciprocalSqrt(const VectorRegister4Double& Vec)
-{
-	VectorRegister4Double Result;
-	Result.XY = vrsqrteq_f64(Vec.XY);
-	Result.ZW = vrsqrteq_f64(Vec.ZW);
-	return Result;
-}
-
-/**
  * Computes an estimate of the reciprocal of a vector (component-wise) and returns the result.
  *
  * @param Vec	1st vector
  * @return		VectorRegister4Float( (Estimate) 1.0f / Vec.x, (Estimate) 1.0f / Vec.y, (Estimate) 1.0f / Vec.z, (Estimate) 1.0f / Vec.w )
  */
-FORCEINLINE VectorRegister4Float VectorReciprocal(const VectorRegister4Float& Vec)
+FORCEINLINE VectorRegister4Float VectorReciprocalEstimate(const VectorRegister4Float& Vec)
 {
 	return vrecpeq_f32(Vec);
 }
 
-FORCEINLINE VectorRegister4Double VectorReciprocal(const VectorRegister4Double& Vec)
+FORCEINLINE VectorRegister4Double VectorReciprocalEstimate(const VectorRegister4Double& Vec)
 {
 	VectorRegister4Double Result;
 	Result.XY = vrecpeq_f64(Vec.XY);
@@ -1627,42 +1602,100 @@
 
 
 /**
-* Return Reciprocal Length of the vector
-*
-* @param Vector		Vector 
-* @return			VectorRegister4Float(rlen, rlen, rlen, rlen) when rlen = 1/sqrt(dot4(V))
-*/
-#define VectorReciprocalLen(Vector)		VectorReciprocalSqrt( VectorDot4( Vector, Vector ) )
+ * Computes the reciprocal of a vector (component-wise) and returns the result.
+ *
+ * @param Vec	1st vector
+ * @return		VectorRegister4Float( 1.0f / Vec.x, 1.0f / Vec.y, 1.0f / Vec.z, 1.0f / Vec.w )
+ */
+FORCEINLINE VectorRegister4Float VectorReciprocal(const VectorRegister4Float& Vec)
+{
+	// Perform two passes of Newton-Raphson iteration on the hardware estimate
+	// The built-in instruction (VRECPS) is not as accurate
+
+	// Initial estimate
+	VectorRegister4Float Reciprocal = VectorReciprocalEstimate(Vec);
+
+	// First iteration
+	VectorRegister4Float Squared = VectorMultiply(Reciprocal, Reciprocal);
+	VectorRegister4Float Double = VectorAdd(Reciprocal, Reciprocal);
+	Reciprocal = VectorNegateMultiplyAdd(Vec, Squared, Double);
+
+	// Second iteration
+	Squared = VectorMultiply(Reciprocal, Reciprocal);
+	Double = VectorAdd(Reciprocal, Reciprocal);
+	return VectorNegateMultiplyAdd(Vec, Squared, Double);
+}
+
+FORCEINLINE VectorRegister4Double VectorReciprocal(const VectorRegister4Double& Vec)
+{
+	return VectorDivide(GlobalVectorConstants::DoubleOne, Vec);
+}
 
 
 /**
-* Return the reciprocal of the square root of each component
-*
-* @param Vector		Vector
-* @return			VectorRegister4Float(1/sqrt(Vec.X), 1/sqrt(Vec.Y), 1/sqrt(Vec.Z), 1/sqrt(Vec.W))
-*/
-FORCEINLINE VectorRegister4Float VectorReciprocalSqrtAccurate(const VectorRegister4Float& Vec)
+ * Return the square root of each component
+ *
+ * @param Vector	Vector
+ * @return			VectorRegister4Float(sqrt(Vec.X), sqrt(Vec.Y), sqrt(Vec.Z), sqrt(Vec.W))
+ */
+FORCEINLINE VectorRegister4Float VectorSqrt(const VectorRegister4Float& Vec)
+{
+	return vsqrtq_f32(Vec);
+}
+
+FORCEINLINE VectorRegister4Double VectorSqrt(const VectorRegister4Double& Vec)
+{
+	VectorRegister4Double Result;
+	Result.XY = vsqrtq_f64(Vec.XY);
+	Result.ZW = vsqrtq_f64(Vec.ZW);
+	return Result;
+}
+
+/**
+ * Returns an estimate of 1/sqrt(c) for each component of the vector
+ *
+ * @param Vector	Vector 
+ * @return			VectorRegister4Float(1/sqrt(t), 1/sqrt(t), 1/sqrt(t), 1/sqrt(t))
+ */
+FORCEINLINE VectorRegister4Float VectorReciprocalSqrtEstimate(const VectorRegister4Float& Vec)
+{
+	return vrsqrteq_f32(Vec);
+}
+
+FORCEINLINE VectorRegister4Double VectorReciprocalSqrtEstimate(const VectorRegister4Double& Vec)
+{
+	VectorRegister4Double Result;
+	Result.XY = vrsqrteq_f64(Vec.XY);
+	Result.ZW = vrsqrteq_f64(Vec.ZW);
+	return Result;
+}
+
+/**
+ * Return the reciprocal of the square root of each component
+ *
+ * @param Vector	Vector
+ * @return			VectorRegister4Float(1/sqrt(Vec.X), 1/sqrt(Vec.Y), 1/sqrt(Vec.Z), 1/sqrt(Vec.W))
+ */
+FORCEINLINE VectorRegister4Float VectorReciprocalSqrt(const VectorRegister4Float& Vec)
 {
 	// Perform a single pass of Newton-Raphson iteration on the hardware estimate
 	// This is a builtin instruction (VRSQRTS)
 
 	// Initial estimate
-	VectorRegister4Float RecipSqrt = VectorReciprocalSqrt(Vec);
+	VectorRegister4Float RecipSqrt = VectorReciprocalSqrtEstimate(Vec);
 
 	// Two refinement
 	RecipSqrt = VectorMultiply(vrsqrtsq_f32(Vec, VectorMultiply(RecipSqrt, RecipSqrt)), RecipSqrt);
 	return VectorMultiply(vrsqrtsq_f32(Vec, VectorMultiply(RecipSqrt, RecipSqrt)), RecipSqrt);
 }
 
-#define VectorSqrt(Vec) vsqrtq_f32(Vec)
-
-FORCEINLINE VectorRegister4Double VectorReciprocalSqrtAccurate(const VectorRegister4Double& Vec)
+FORCEINLINE VectorRegister4Double VectorReciprocalSqrt(const VectorRegister4Double& Vec)
 {
 	// Perform a single pass of Newton-Raphson iteration on the hardware estimate
 	// This is a builtin instruction (VRSQRTS)
 
 	// Initial estimate
-	VectorRegister4Double RecipSqrt = VectorReciprocalSqrt(Vec);
+	VectorRegister4Double RecipSqrt = VectorReciprocalSqrtEstimate(Vec);
 
 	// Two refinement
 	VectorRegister4Double Tmp;
@@ -1675,37 +1708,39 @@
 	return VectorMultiply(Tmp, RecipSqrt);
 }
 
-
 /**
- * Computes the reciprocal of a vector (component-wise) and returns the result.
+ * Return Reciprocal Length of the vector
  *
- * @param Vec	1st vector
- * @return		VectorRegister4Float( 1.0f / Vec.x, 1.0f / Vec.y, 1.0f / Vec.z, 1.0f / Vec.w )
+ * @param Vector	Vector
+ * @return			VectorRegister4Float(rlen, rlen, rlen, rlen) when rlen = 1/sqrt(dot4(V))
  */
-FORCEINLINE VectorRegister4Float VectorReciprocalAccurate(const VectorRegister4Float& Vec)
+FORCEINLINE VectorRegister4Float VectorReciprocalLen(const VectorRegister4Float& Vector)
 {
-	// Perform two passes of Newton-Raphson iteration on the hardware estimate
-	// The built-in instruction (VRECPS) is not as accurate
-
-	// Initial estimate
-	VectorRegister4Float Reciprocal = VectorReciprocal(Vec);
-
-	// First iteration
-	VectorRegister4Float Squared = VectorMultiply(Reciprocal, Reciprocal);
-	VectorRegister4Float Double = VectorAdd(Reciprocal, Reciprocal);
-	Reciprocal = VectorNegateMultiplyAdd(Vec, Squared, Double);
-
-	// Second iteration
-	Squared = VectorMultiply(Reciprocal, Reciprocal);
-	Double = VectorAdd(Reciprocal, Reciprocal);
-	return VectorNegateMultiplyAdd(Vec, Squared, Double);
+	return VectorReciprocalSqrt(VectorDot4(Vector, Vector));
 }
 
-FORCEINLINE VectorRegister4Double VectorReciprocalAccurate(const VectorRegister4Double& Vec)
+FORCEINLINE VectorRegister4Double VectorReciprocalLen(const VectorRegister4Double& Vector)
 {
-	return VectorDivide(GlobalVectorConstants::DoubleOne, Vec);
+	return VectorReciprocalSqrt(VectorDot4(Vector, Vector));
 }
 
 /**
+ * Return Reciprocal Length of the vector (estimate)
+ *
+ * @param Vector	Vector
+ * @return			VectorRegister4Float(rlen, rlen, rlen, rlen) when rlen = 1/sqrt(dot4(V)) (estimate)
+ */
+FORCEINLINE VectorRegister4Float VectorReciprocalLenEstimate(const VectorRegister4Float& Vector)
+{
+	return VectorReciprocalSqrtEstimate(VectorDot4(Vector, Vector));
+}
+
+FORCEINLINE VectorRegister4Double VectorReciprocalLenEstimate(const VectorRegister4Double& Vector)
+{
+	return VectorReciprocalSqrtEstimate(VectorDot4(Vector, Vector));
+}
+
+
+/**
 * Loads XYZ and sets W=0
 *
diff --git a/Engine/Source/Runtime/Core/Public/Math/UnrealMathSSE.h b/Engine/Source/Runtime/Core/Public/Math/UnrealMathSSE.h
--- a/Engine/Source/Runtime/Core/Public/Math/UnrealMathSSE.h
+++ b/Engine/Source/Runtime/Core/Public/Math/UnrealMathSSE.h
@@ -1890,8 +1890,12 @@ FORCEINLINE VectorRegister4Double VectorCross(const VectorRegister4Double& Vec1,
  */
 FORCEINLINE VectorRegister4Float VectorPow( const VectorRegister4Float& Base, const VectorRegister4Float& Exponent )
 {
+#if UE_PLATFORM_MATH_USE_SVML
+	return _mm_pow_ps(Base, Exponent);
+#else
 	// using SseMath library
 	return SSE::exp_ps(_mm_mul_ps(SSE::log_ps(Base), Exponent));
+#endif
 /*
 	// old version, keeping for reference in case something breaks and we need to debug it.
 	union { VectorRegister4Float v; float f[4]; } B, E;
@@ -1903,6 +1907,11 @@ FORCEINLINE VectorRegister4Float VectorPow( const VectorRegister4Float& Base, co
 
 FORCEINLINE VectorRegister4Double VectorPow(const VectorRegister4Double& Base, const VectorRegister4Double& Exponent)
 {
+#if UE_PLATFORM_MATH_USE_SVML_AVX
+	return _mm256_pow_pd(Base, Exponent);
+#elif UE_PLATFORM_MATH_USE_SVML
+	return VectorRegister4Double(_mm_pow_pd(Base.XY, Exponent.XY), _mm_pow_pd(Base.ZW, Exponent.ZW));
+#else
 	AlignedDouble4 Values(Base);
 	AlignedDouble4 Exponents(Exponent);
 
@@ -1911,80 +1920,51 @@ FORCEINLINE VectorRegister4Double VectorPow(const VectorRegister4Double& Base, c
 	Values[2] = FMath::Pow(Values[2], Exponents[2]);
 	Values[3] = FMath::Pow(Values[3], Exponents[3]);
 	return Values.ToVectorRegister();
+#endif
 }
 
 /**
-* Returns an estimate of 1/sqrt(c) for each component of the vector
-*
-* @param Vector		Vector 
-* @return			VectorRegister4Float(1/sqrt(t), 1/sqrt(t), 1/sqrt(t), 1/sqrt(t))
-*/
-//#define VectorReciprocalSqrt(Vec)		_mm_rsqrt_ps( Vec )
-
-FORCEINLINE VectorRegister4Float VectorReciprocalSqrt(const VectorRegister4Float& Vec)
+ * Return the square root of each component
+ *
+ * @param Vector	Vector
+ * @return			VectorRegister4Float(sqrt(Vec.X), sqrt(Vec.Y), sqrt(Vec.Z), sqrt(Vec.W))
+ */
+FORCEINLINE VectorRegister4Float VectorSqrt(const VectorRegister4Float& Vec)
 {
-	return _mm_rsqrt_ps(Vec);
+	return _mm_sqrt_ps(Vec);
 }
 
-FORCEINLINE VectorRegister4Double VectorReciprocalSqrt(const VectorRegister4Double& Vec)
+FORCEINLINE VectorRegister4Double VectorSqrt(const VectorRegister4Double& Vec)
 {
-	// Not an estimate.
-	VectorRegister4Double Result;
-#if !UE_PLATFORM_MATH_USE_AVX
-	Result.XY = _mm_sqrt_pd(Vec.XY);
-	Result.ZW = _mm_sqrt_pd(Vec.ZW);
+#if UE_PLATFORM_MATH_USE_AVX
+	return _mm256_sqrt_pd(Vec);
 #else
-	Result = _mm256_sqrt_pd(Vec);
+	return VectorRegister4Double(_mm_sqrt_pd(Vec.XY), _mm_sqrt_pd(Vec.ZW));
 #endif
-	return VectorDivide(GlobalVectorConstants::DoubleOne, Result);
 }
 
-#define VectorSqrt(Vec) _mm_sqrt_ps(Vec)
-
 /**
- * Computes an estimate of the reciprocal of a vector (component-wise) and returns the result.
+ * Returns an estimate of 1/sqrt(c) for each component of the vector
  *
- * @param Vec	1st vector
- * @return		VectorRegister4Float( (Estimate) 1.0f / Vec.x, (Estimate) 1.0f / Vec.y, (Estimate) 1.0f / Vec.z, (Estimate) 1.0f / Vec.w )
+ * @param Vector		Vector
+ * @return			VectorRegister4Float(1/sqrt(t), 1/sqrt(t), 1/sqrt(t), 1/sqrt(t))
  */
-
-FORCEINLINE VectorRegister4Float VectorReciprocal(const VectorRegister4Float& Vec)
-{
-	return _mm_rcp_ps(Vec);
-}
-
-FORCEINLINE VectorRegister4Double VectorReciprocal(const VectorRegister4Double& Vec)
-{
-	// Not an estimate.
-	return VectorDivide(GlobalVectorConstants::DoubleOne, Vec);
-}
-
-/**
-* Return Reciprocal Length of the vector (estimate)
-*
-* @param Vector		Vector 
-* @return			VectorRegister4Float(rlen, rlen, rlen, rlen) when rlen = 1/sqrt(dot4(V))
-*/
-FORCEINLINE VectorRegister4Float VectorReciprocalLen( const VectorRegister4Float& Vector )
+FORCEINLINE VectorRegister4Float VectorReciprocalSqrtEstimate(const VectorRegister4Float& Vec)
 {
-	VectorRegister4Float RecipLen = VectorDot4( Vector, Vector );
-	return VectorReciprocalSqrt( RecipLen );
-}
-
-FORCEINLINE VectorRegister4Double VectorReciprocalLen(const VectorRegister4Double& Vector)
-{
-	VectorRegister4Double RecipLen = VectorDot4(Vector, Vector);
-	return VectorReciprocalSqrt(RecipLen);
+	return _mm_rsqrt_ps(Vec);
 }
 
 /**
-* Return the reciprocal of the square root of each component
-*
-* @param Vector		Vector 
-* @return			VectorRegister4Float(1/sqrt(Vec.X), 1/sqrt(Vec.Y), 1/sqrt(Vec.Z), 1/sqrt(Vec.W))
-*/
-FORCEINLINE VectorRegister4Float VectorReciprocalSqrtAccurate(const VectorRegister4Float& Vec)
+ * Return the reciprocal of the square root of each component
+ *
+ * @param Vector		Vector 
+ * @return			VectorRegister4Float(1/sqrt(Vec.X), 1/sqrt(Vec.Y), 1/sqrt(Vec.Z), 1/sqrt(Vec.W))
+ */
+FORCEINLINE VectorRegister4Float VectorReciprocalSqrt(const VectorRegister4Float& Vec)
 {
+#if UE_PLATFORM_MATH_USE_SVML
+	return _mm_invsqrt_ps(Vec);
+#else
 	// Perform two passes of Newton-Raphson iteration on the hardware estimate
 	//    v^-0.5 = x
 	// => x^2 = v^-1
@@ -2001,7 +1981,7 @@ FORCEINLINE VectorRegister4Float VectorReciprocalSqrtAccurate(const VectorRegist
 	const VectorRegister4Float VecDivBy2 = VectorMultiply(Vec, OneHalf);
 
 	// Initial estimate
-	const VectorRegister4Float x0 = VectorReciprocalSqrt(Vec);
+	const VectorRegister4Float x0 = VectorReciprocalSqrtEstimate(Vec);
 
 	// First iteration
 	VectorRegister4Float x1 = VectorMultiply(x0, x0);
@@ -2014,18 +1994,71 @@ FORCEINLINE VectorRegister4Float VectorReciprocalSqrtAccurate(const VectorRegist
 	x2 = VectorMultiplyAdd(x1, x2, x1);
 
 	return x2;
+#endif
 }
 
-FORCEINLINE VectorRegister4Double VectorReciprocalSqrtAccurate(const VectorRegister4Double& Vec)
+FORCEINLINE VectorRegister4Double VectorReciprocalSqrt(const VectorRegister4Double& Vec)
 {
-	VectorRegister4Double Result;
-#if !UE_PLATFORM_MATH_USE_AVX
-	Result.XY = _mm_sqrt_pd(Vec.XY);
-	Result.ZW = _mm_sqrt_pd(Vec.ZW);
+#if UE_PLATFORM_MATH_USE_SVML_AVX
+	return _mm256_invsqrt_pd(Vec);
+#elif UE_PLATFORM_MATH_USE_SVML
+	return VectorRegister4Double(_mm_invsqrt_pd(Vec.XY), _mm_invsqrt_pd(Vec.ZW));
+#elif UE_PLATFORM_MATH_USE_AVX
+	return VectorDivide(GlobalVectorConstants::DoubleOne, _mm256_sqrt_pd(Vec));
 #else
-	Result = _mm256_sqrt_pd(Vec);
+	return VectorDivide(GlobalVectorConstants::DoubleOne, VectorRegister4Double(_mm_sqrt_pd(Vec.XY), _mm_sqrt_pd(Vec.ZW)));
 #endif
-	return VectorDivide(GlobalVectorConstants::DoubleOne, Result);
+}
+
+FORCEINLINE VectorRegister4Double VectorReciprocalSqrtEstimate(const VectorRegister4Double& Vec)
+{
+	// TODO: Not an estimate
+	return VectorReciprocalSqrt(Vec);
+}
+
+
+/**
+ * Return Reciprocal Length of the vector
+ *
+ * @param Vector	Vector
+ * @return			VectorRegister4Float(rlen, rlen, rlen, rlen) when rlen = 1/sqrt(dot4(V))
+ */
+FORCEINLINE VectorRegister4Float VectorReciprocalLen(const VectorRegister4Float& Vector)
+{
+	return VectorReciprocalSqrt(VectorDot4(Vector, Vector));
+}
+
+FORCEINLINE VectorRegister4Double VectorReciprocalLen(const VectorRegister4Double& Vector)
+{
+	return VectorReciprocalSqrt(VectorDot4(Vector, Vector));
+}
+
+/**
+ * Return Reciprocal Length of the vector (estimate)
+ *
+ * @param Vector	Vector
+ * @return			VectorRegister4Float(rlen, rlen, rlen, rlen) when rlen = 1/sqrt(dot4(V))
+ */
+FORCEINLINE VectorRegister4Float VectorReciprocalLenEstimate(const VectorRegister4Float& Vector)
+{
+	return VectorReciprocalSqrtEstimate(VectorDot4(Vector, Vector));
+}
+
+FORCEINLINE VectorRegister4Double VectorReciprocalLenEstimate(const VectorRegister4Double& Vector)
+{
+	return VectorReciprocalSqrtEstimate(VectorDot4(Vector, Vector));
+}
+
+/**
+ * Computes an estimate of the reciprocal of a vector (component-wise) and returns the result.
+ *
+ * @param Vec	1st vector
+ * @return		VectorRegister4Float( (Estimate) 1.0f / Vec.x, (Estimate) 1.0f / Vec.y, (Estimate) 1.0f / Vec.z, (Estimate) 1.0f / Vec.w )
+ */
+
+FORCEINLINE VectorRegister4Float VectorReciprocalEstimate(const VectorRegister4Float& Vec)
+{
+	return _mm_rcp_ps(Vec);
 }
 
 /**
@@ -2034,7 +2067,7 @@ FORCEINLINE VectorRegister4Double VectorReciprocalSqrtAccurate(const VectorRegis
  * @param Vec	1st vector
  * @return		VectorRegister4Float( 1.0f / Vec.x, 1.0f / Vec.y, 1.0f / Vec.z, 1.0f / Vec.w )
  */
-FORCEINLINE VectorRegister4Float VectorReciprocalAccurate(const VectorRegister4Float& Vec)
+FORCEINLINE VectorRegister4Float VectorReciprocal(const VectorRegister4Float& Vec)
 {
 	// Perform two passes of Newton-Raphson iteration on the hardware estimate
 	//   x1 = x0 - f(x0) / f'(x0)
@@ -2049,7 +2082,7 @@ FORCEINLINE VectorRegister4Float VectorReciprocalAccurate(const VectorRegister4F
 	// => x1 = x0 - (x0 * Vec - 1) * x0 = 2 * x0 - Vec * x0^2 
 
 	// Initial estimate
-	const VectorRegister4Float x0 = VectorReciprocal(Vec);
+	const VectorRegister4Float x0 = VectorReciprocalEstimate(Vec);
 
 	// First iteration
 	const VectorRegister4Float x0Squared = VectorMultiply(x0, x0);
@@ -2064,15 +2097,15 @@ FORCEINLINE VectorRegister4Float VectorReciprocalAccurate(const VectorRegister4F
 	return x2;
 }
 
-FORCEINLINE VectorRegister4Double VectorReciprocalAccurate(const VectorRegister4Double& Vec)
+FORCEINLINE VectorRegister4Double VectorReciprocal(const VectorRegister4Double& Vec)
 {
 	return VectorDivide(GlobalVectorConstants::DoubleOne, Vec);
 }
 
-FORCEINLINE VectorRegister4Float VectorLerp(const VectorRegister4Float &Vec1, const VectorRegister4Float &Vec2, const VectorRegister4Float &Vec3)
+FORCEINLINE VectorRegister4Double VectorReciprocalEstimate(const VectorRegister4Double& Vec)
 {
-	VectorRegister4Float SubVec = VectorSubtract(GlobalVectorConstants::FloatOne, Vec3);
-	return VectorMultiplyAdd(Vec2, Vec3, VectorMultiply(Vec1, SubVec));
+	// Not an estimate.
+	return VectorReciprocal(Vec);
 }
 
 /**
@@ -2724,7 +2757,8 @@ FORCEINLINE VectorRegister4Double VectorTruncate(const VectorRegister4Double& V)
 	return Result;
 }
 
-FORCEINLINE VectorRegister4Float VectorRound(const VectorRegister4Float &Vec) {
+FORCEINLINE VectorRegister4Float VectorRound(const VectorRegister4Float &Vec)
+{
 #if UE_PLATFORM_MATH_USE_SSE4_1
 	return _mm_round_ps(Vec, _MM_FROUND_TO_NEAREST_INT |_MM_FROUND_NO_EXC);
 #else
@@ -2803,6 +2837,10 @@ FORCEINLINE VectorRegister4Float VectorMod(const VectorRegister4Float& X, const
 {
 	// Check against invalid divisor
 	VectorRegister4Float InvalidDivisorMask = VectorCompareLE(VectorAbs(Y), GlobalVectorConstants::SmallNumber);
+
+#if UE_PLATFORM_MATH_USE_SVML
+	VectorRegister4Float Result = _mm_fmod_ps(X, Y);
+#else
 	// Promote to doubles for better intermediate precision.
 	VectorRegister4Double DoubleX = MakeVectorRegisterDouble(X);
 	VectorRegister4Double DoubleY = MakeVectorRegisterDouble(Y);
@@ -2811,21 +2849,32 @@ FORCEINLINE VectorRegister4Float VectorMod(const VectorRegister4Float& X, const
 	VectorRegister4Double DoubleResult = VectorNegateMultiplyAdd(DoubleY, Temp, DoubleX);
 	// Convert back to floats. This is safe (in terms of not exceeding [-FLT_MAX,FLT_MAX]) because the answer is less than or equal to Abs(X) by definition.
 	VectorRegister4Float Result = MakeVectorRegisterFloatFromDouble(DoubleResult);
+#endif
+
 	// Return 0 where divisor Y was too small	
 	Result = VectorSelect(InvalidDivisorMask, GlobalVectorConstants::FloatZero, Result);
 	return Result;
+
 }
 
 FORCEINLINE VectorRegister4Double VectorMod(const VectorRegister4Double& X, const VectorRegister4Double& Y)
 {
 	// Check against invalid divisor
 	VectorRegister4Double InvalidDivisorMask = VectorCompareLE(VectorAbs(Y), GlobalVectorConstants::DoubleSmallNumber);
+
+#if UE_PLATFORM_MATH_USE_SVML_AVX
+	VectorRegister4Double DoubleResult = _mm256_fmod_pd(X, Y);
+#elif UE_PLATFORM_MATH_USE_SVML
+	VectorRegister4Double DoubleResult = VectorRegister4Double(_mm_fmod_pd(X.XY, Y.XY), _mm_fmod_pd(X.ZW, Y.ZW));
+#else
 	// R = X - (Y * (Trunc(X / Y))
 	VectorRegister4Double Temp = VectorTruncate(VectorDivide(X, Y));
 	VectorRegister4Double DoubleResult = VectorNegateMultiplyAdd(Y, Temp, X);
+#endif
+
 	// Return 0 where divisor Y was too small
-	VectorRegister4Double Result = VectorSelect(InvalidDivisorMask, GlobalVectorConstants::DoubleZero, DoubleResult);
-	return Result;
+	DoubleResult = VectorSelect(InvalidDivisorMask, GlobalVectorConstants::DoubleZero, DoubleResult);
+	return DoubleResult;
 }
 
 FORCEINLINE VectorRegister4Float VectorSign(const VectorRegister4Float& X)
@@ -2990,9 +3039,9 @@ FORCEINLINE VectorRegister4Float VectorSin(const VectorRegister4Float& V)
 FORCEINLINE VectorRegister4Double VectorSin(const VectorRegister4Double& V)
 {
 #if UE_PLATFORM_MATH_USE_SVML_AVX
-		return _mm256_sin_pd(V);
+	return _mm256_sin_pd(V);
 #elif UE_PLATFORM_MATH_USE_SVML
-		return VectorRegister4Double(_mm_sin_pd(V.XY), _mm_sin_pd(V.ZW));
+	return VectorRegister4Double(_mm_sin_pd(V.XY), _mm_sin_pd(V.ZW));
 #else
 	AlignedDouble4 Doubles(V);
 	Doubles[0] = FMath::Sin(Doubles[0]);
@@ -3008,7 +3057,6 @@ FORCEINLINE VectorRegister4Float VectorCos(const VectorRegister4Float& V)
 #if UE_PLATFORM_MATH_USE_SVML
 	return _mm_cos_ps(V);
 #else
-	//return VectorSin(VectorAdd(V, GlobalVectorConstants::PiByTwo));
 	return SSE::cos_ps(V);
 #endif
 }
@@ -3016,9 +3064,9 @@ FORCEINLINE VectorRegister4Float VectorCos(const VectorRegister4Float& V)
 FORCEINLINE VectorRegister4Double VectorCos(const VectorRegister4Double& V)
 {
 #if UE_PLATFORM_MATH_USE_SVML_AVX
-		return _mm256_cos_pd(V);
+	return _mm256_cos_pd(V);
 #elif UE_PLATFORM_MATH_USE_SVML
-		return VectorRegister4Double(_mm_cos_pd(V.XY), _mm_cos_pd(V.ZW));
+	return VectorRegister4Double(_mm_cos_pd(V.XY), _mm_cos_pd(V.ZW));
 #else
 	AlignedDouble4 Doubles(V);
 	Doubles[0] = FMath::Cos(Doubles[0]);
diff --git a/Engine/Source/Runtime/Core/Public/Math/UnrealMathVectorCommon.h b/Engine/Source/Runtime/Core/Public/Math/UnrealMathVectorCommon.h
--- a/Engine/Source/Runtime/Core/Public/Math/UnrealMathVectorCommon.h
+++ b/Engine/Source/Runtime/Core/Public/Math/UnrealMathVectorCommon.h
@@ -365,12 +365,43 @@ FORCEINLINE void VectorStoreFloat1(const VectorRegister4Double& Vec, int64* Dst)
 	VectorStoreFloat1(Vec, (double*)Dst);
 }
 
+///////////////////////////////////////////////////////////////////////////////////////////////////////////////////
+// Wrappers for clarifying "accurate" versus "estimate" for some functions,
+// important for old functions which were ambiguous in their precision.
+
+// Returns accurate reciprocal square root.
+FORCEINLINE VectorRegister4Float VectorReciprocalSqrtAccurate(const VectorRegister4Float& Vec)
+{
+	return VectorReciprocalSqrt(Vec);
+}
+
+FORCEINLINE VectorRegister4Double VectorReciprocalSqrtAccurate(const VectorRegister4Double& Vec)
+{
+	return VectorReciprocalSqrt(Vec);
+}
+
+/**
+ * Computes the reciprocal of a vector (component-wise) and returns the result.
+ *
+ * @param Vec	1st vector
+ * @return		VectorRegister( 1.0f / Vec.x, 1.0f / Vec.y, 1.0f / Vec.z, 1.0f / Vec.w )
+ */
+FORCEINLINE VectorRegister4Float VectorReciprocalAccurate(const VectorRegister4Float& Vec)
+{
+	return VectorReciprocal(Vec);
+}
+
+FORCEINLINE VectorRegister4Double VectorReciprocalAccurate(const VectorRegister4Double& Vec)
+{
+	return VectorReciprocal(Vec);
+}
+
 ///////////////////////////////////////////////////////////////////////////////////////////////////////////////////
 
 // Returns a normalized 4 vector = Vector / |Vector|.
 // There is no handling of zero length vectors, use VectorNormalizeSafe if this is a possible input.
 template<typename TVectorRegisterType>
-FORCEINLINE TVectorRegisterType VectorNormalizeAccurate( const TVectorRegisterType& Vector )
+FORCEINLINE TVectorRegisterType VectorNormalize( const TVectorRegisterType& Vector )
 {
 	const TVectorRegisterType SquareSum = VectorDot4(Vector, Vector);
 	const TVectorRegisterType InvLength = VectorReciprocalSqrtAccurate(SquareSum);
@@ -378,6 +409,18 @@ FORCEINLINE TVectorRegisterType VectorNormalizeAccurate( const TVectorRegisterTy
 	return NormalizedVector;
 }
 
+template<typename TVectorRegisterType>
+FORCEINLINE TVectorRegisterType VectorNormalizeAccurate(const TVectorRegisterType& Vector)
+{
+	return VectorNormalize(Vector);
+}
+
+template<typename TVectorRegisterType>
+FORCEINLINE TVectorRegisterType VectorNormalizeEstimate(const TVectorRegisterType& Vector)
+{
+	return VectorMultiply(Vector, VectorReciprocalLenEstimate(Vector));
+}
+
 // Returns ((Vector dot Vector) >= 1e-8) ? (Vector / |Vector|) : DefaultValue
 // Uses accurate 1/sqrt, not the estimate
 FORCEINLINE VectorRegister4Float VectorNormalizeSafe( const VectorRegister4Float& Vector, const VectorRegister4Float& DefaultValue )
@@ -475,6 +518,19 @@ template<> FORCEINLINE VectorRegister4Double FGenericPlatformMath::Max(const Vec
 	return VectorMax(A, B);
 }
 
+/** Lerp between two vectors */
+FORCEINLINE VectorRegister4Float VectorLerp(const VectorRegister4Float& A, const VectorRegister4Float& B, const VectorRegister4Float& Alpha)
+{
+	VectorRegister4Float SubVec = VectorSubtract(GlobalVectorConstants::FloatOne, Alpha);
+	return VectorMultiplyAdd(B, Alpha, VectorMultiply(A, SubVec));
+}
+
+FORCEINLINE VectorRegister4Double VectorLerp(const VectorRegister4Double& A, const VectorRegister4Double& B, const VectorRegister4Double& Alpha)
+{
+	VectorRegister4Double SubVec = VectorSubtract(GlobalVectorConstants::DoubleOne, Alpha);
+	return VectorMultiplyAdd(B, Alpha, VectorMultiply(A, SubVec));
+}
+
 // TCustomLerp for FMath::Lerp<VectorRegister4Float>()
 template<>
 struct TCustomLerp<VectorRegister4Float>
@@ -484,8 +540,7 @@ struct TCustomLerp<VectorRegister4Float>
 	// Specialization of Lerp function that works with vector registers
 	static FORCEINLINE VectorRegister4Float Lerp(const VectorRegister4Float& A, const VectorRegister4Float& B, const VectorRegister4Float& Alpha)
 	{
-		const VectorRegister4Float Delta = VectorSubtract(B, A);
-		return VectorMultiplyAdd(Alpha, Delta, A);
+		return VectorLerp(A, B, Alpha);
 	}
 };
 
@@ -498,8 +553,7 @@ struct TCustomLerp<VectorRegister4Double>
 	// Specialization of Lerp function that works with vector registers
 	static FORCEINLINE VectorRegister4Double Lerp(const VectorRegister4Double& A, const VectorRegister4Double& B, const VectorRegister4Double& Alpha)
 	{
-		const VectorRegister4Double Delta = VectorSubtract(B, A);
-		return VectorMultiplyAdd(Alpha, Delta, A);
+		return VectorLerp(A, B, Alpha);
 	}
 };
 
@@ -799,23 +853,6 @@ FORCEINLINE VectorRegister4Float VectorSet(uint32 X, uint32 Y, uint32 Z, uint32
 	return VectorSet(float(X), float(Y), float(Z), float(W));
 }
 
-/**
- * Normalize vector
- *
- * @param Vector		Vector to normalize
- * @return			Normalized VectorRegister4Float
- */
-FORCEINLINE VectorRegister4Float VectorNormalize(const VectorRegister4Float& Vec)
-{
-	return VectorMultiply(Vec, VectorReciprocalLen(Vec));
-}
-
-FORCEINLINE VectorRegister4Double VectorNormalize(const VectorRegister4Double& Vec)
-{
-	return VectorMultiply(Vec, VectorReciprocalLen(Vec));
-}
-
-
 FORCEINLINE VectorRegister4Float VectorFractional(const VectorRegister4Float& Vec)
 {
 	return VectorSubtract(Vec, VectorTruncate(Vec));
