Windows memory mapping improvements

Fix error handling, make PreloadHint respect in parameters and overload SupportsMemoryMappedFiles

diff --git a/Engine/Source/Runtime/Core/Private/Windows/WindowsPlatformFile.cpp b/Engine/Source/Runtime/Core/Private/Windows/WindowsPlatformFile.cpp
--- a/Engine/Source/Runtime/Core/Private/Windows/WindowsPlatformFile.cpp
+++ b/Engine/Source/Runtime/Core/Private/Windows/WindowsPlatformFile.cpp
@@ -821,19 +821,20 @@ public:
 
 	virtual void PreloadHint(int64 PreloadOffset = 0, int64 BytesToPreload = MAX_int64) override
 	{
-		// perhaps this could be done with a commit instead
-		int64 Size = GetMappedSize();
-		const uint8* Ptr = GetMappedPtr();
-		int32 FoolTheOptimizer = 0;
-		while (Size > 0)
-		{
-			FoolTheOptimizer += Ptr[0];
-			Size -= 4096;
-			Ptr += 4096;
-		}
-		if (FoolTheOptimizer == 0xbadf00d)
+		if (PreloadOffset < GetMappedSize())
 		{
-			FPlatformProcess::Sleep(0.0f); // this will more or less never happen, but we can't let the optimizer strip these reads
+			SIZE_T PageSize = FWindowsPlatformMemory::GetConstants().PageSize;
+			int64 SizeAfterOffset = GetMappedSize() - PreloadOffset;
+			int64 PreloadSize = FMath::Min(SizeAfterOffset, BytesToPreload);
+
+			for (const uint8* It = GetMappedPtr() + PreloadOffset, *End = It + PreloadSize; It < End; It += PageSize)
+			{
+				volatile uint8 Foo = *It;
+			}
+
+			// Once we drop Win7 support we can try this Win8 API instead
+			// WIN32_MEMORY_RANGE_ENTRY Range = { GetMappedPtr() + PreloadOffset, PreloadSize };
+			// PrefetchVirtualMemory(GetCurrentProcess(), 1, &Range, 0);
 		}
 	}
 };
@@ -882,8 +883,9 @@ public:
 
 		DWORD OpenMappingAccess = FILE_MAP_READ;
 
-		int64 AlignedOffset = AlignDown(Offset, 65536);
-		int64 AlignedSize = Align(BytesToMap + Offset - AlignedOffset, 65536);
+		SIZE_T VirtualMemoryGranularity = FWindowsPlatformMemory::GetConstants().OsAllocationGranularity;
+		int64 AlignedOffset = AlignDown(Offset, VirtualMemoryGranularity);
+		int64 AlignedSize = Align(BytesToMap + Offset - AlignedOffset, VirtualMemoryGranularity);
 
 		ULARGE_INTEGER LI;
 		LI.QuadPart = AlignedOffset;
@@ -1337,7 +1339,7 @@ public:
 			return nullptr;
 		}
 		HANDLE MappingHandle = CreateFileMapping(Handle, NULL, PAGE_READONLY, 0, 0, NULL);
-		if (MappingHandle == INVALID_HANDLE_VALUE)
+		if (MappingHandle == NULL)
 		{
 			TRACE_PLATFORMFILE_BEGIN_CLOSE(Handle);
 #if PLATFORMFILETRACE_ENABLED
diff --git a/Engine/Source/Runtime/Core/Public/Windows/WindowsPlatformProperties.h b/Engine/Source/Runtime/Core/Public/Windows/WindowsPlatformProperties.h
--- a/Engine/Source/Runtime/Core/Public/Windows/WindowsPlatformProperties.h
+++ b/Engine/Source/Runtime/Core/Public/Windows/WindowsPlatformProperties.h
@@ -78,6 +78,11 @@ struct FWindowsPlatformProperties
 		return IS_DEDICATED_SERVER;
 	}
 
+	static FORCEINLINE bool SupportsMemoryMappedFiles()
+	{
+		return true;
+	}
+
 	static FORCEINLINE bool SupportsAudioStreaming()
 	{
 		return !IsServerOnly();
