Optimize WITH_CASE_PRESERVING_NAME FName loading

LoadDisplayNames() 1.06s -> 0.24s when async loading the name batch in CachedAssetRegistry.bin during editor startup of an internal project

* Introduce FDisplayNameEntryId to avoid cache misses in FName::GetComparisonIdFromDisplayId() lookups
* Force name batch serialization to use display entry ids to help enforce cooking determinism
* Fix determinism problem when loading then saving via FPackageStoreNameMapBuilder
* Inline FNameEntryId::FromUnstableInt()
* Remove slow IsNumbered() check in CreateFromDisplayId when not using 4B FName
* Remove expensive heap sort by sorting up front instead

diff --git a/Engine/Source/Developer/IoStoreUtilities/Internal/PackageStoreOptimizer.h b/Engine/Source/Developer/IoStoreUtilities/Internal/PackageStoreOptimizer.h
--- a/Engine/Source/Developer/IoStoreUtilities/Internal/PackageStoreOptimizer.h
+++ b/Engine/Source/Developer/IoStoreUtilities/Internal/PackageStoreOptimizer.h
@@ -19,25 +19,27 @@ public:
 		NameMapType = InNameMapType;
 	}
 
-	void AddName(const FName& Name)
+	void AddName(FName Name)
 	{
-		const FNameEntryId ComparisonIndex = Name.GetComparisonIndex();
-		const FNameEntryId DisplayIndex = Name.GetDisplayIndex();
-		NameMap.Add(DisplayIndex);
+		AddName(FDisplayNameEntryId(Name));
+	}
+
+	void AddName(FDisplayNameEntryId DisplayId)
+	{
+		NameMap.Add(DisplayId);
 		int32 Index = NameMap.Num();
-		NameIndices.Add(ComparisonIndex, Index);
+		NameIndices.Add(DisplayId, Index);
 	}
 
 	void MarkNamesAsReferenced(const TArray<FName>& Names, TArray<int32>& OutNameIndices)
 	{
-		for (const FName& Name : Names)
+		for (FName Name : Names)
 		{
-			const FNameEntryId ComparisonIndex = Name.GetComparisonIndex();
-			const FNameEntryId DisplayIndex = Name.GetDisplayIndex();
-			int32& Index = NameIndices.FindOrAdd(ComparisonIndex);
+			FDisplayNameEntryId DisplayId(Name);
+			int32& Index = NameIndices.FindOrAdd(DisplayId);
 			if (Index == 0)
 			{
-				NameMap.Add(DisplayIndex);
+				NameMap.Add(DisplayId);
 				Index = NameMap.Num();
 			}
 
@@ -45,27 +47,24 @@ public:
 		}
 	}
 
-	void MarkNameAsReferenced(const FName& Name)
+	void MarkNameAsReferenced(FName Name)
 	{
-		const FNameEntryId ComparisonIndex = Name.GetComparisonIndex();
-		const FNameEntryId DisplayIndex = Name.GetDisplayIndex();
-		int32& Index = NameIndices.FindOrAdd(ComparisonIndex);
+		FDisplayNameEntryId DisplayId(Name);
+		int32& Index = NameIndices.FindOrAdd(DisplayId);
 		if (Index == 0)
 		{
-			NameMap.Add(DisplayIndex);
+			NameMap.Add(DisplayId);
 			Index = NameMap.Num();
 		}
 	}
 
-	FMappedName MapName(const FName& Name) const
+	FMappedName MapName(FName Name) const
 	{
-		const FNameEntryId Id = Name.GetComparisonIndex();
-		const int32* Index = NameIndices.Find(Id);
-		check(Index);
-		return FMappedName::Create(*Index - 1, Name.GetNumber(), NameMapType);
+		int32 Index = NameIndices.FindChecked(FDisplayNameEntryId(Name));
+		return FMappedName::Create(Index - 1, Name.GetNumber(), NameMapType);
 	}
 
-	const TArray<FNameEntryId>& GetNameMap() const
+	TConstArrayView<FDisplayNameEntryId> GetNameMap() const
 	{
 		return NameMap;
 	}
@@ -77,8 +76,8 @@ public:
 	}
 
 private:
-	TMap<FNameEntryId, int32> NameIndices;
-	TArray<FNameEntryId> NameMap;
+	TMap<FDisplayNameEntryId, int32> NameIndices;
+	TArray<FDisplayNameEntryId> NameMap;
 	FMappedName::EType NameMapType = FMappedName::EType::Package;
 };
 
@@ -374,7 +373,7 @@ private:
 		TOptional<FZenPackageVersioningInfo> VersioningInfo;
 		TArray<FPackageId> ImportedPackageIds;
 		TArray<uint64> ImportedPublicExportHashes;
-		TArray<FNameEntryId> NameMap;
+		TArray<FDisplayNameEntryId> NameMap;
 		TArray<FPackageObjectIndex> Imports; // FH: Imports might need to have more info to be able to resolve export hash with import as outer
 		TArray<FExportMapEntry> Exports;
 		TArray<FExportBundleHeader> ExportBundleHeaders;
diff --git a/Engine/Source/Developer/IoStoreUtilities/Private/IoStoreUtilities.cpp b/Engine/Source/Developer/IoStoreUtilities/Private/IoStoreUtilities.cpp
--- a/Engine/Source/Developer/IoStoreUtilities/Private/IoStoreUtilities.cpp
+++ b/Engine/Source/Developer/IoStoreUtilities/Private/IoStoreUtilities.cpp
@@ -3885,7 +3885,7 @@ int32 Describe(
 
 	TMap<FPackageObjectIndex, FScriptObjectDesc> ScriptObjectByGlobalIdMap;
 	FLargeMemoryReader ScriptObjectsArchive(ScriptObjectsBuffer.ValueOrDie().Data(), ScriptObjectsBuffer.ValueOrDie().DataSize());
-	TArray<FNameEntryId> GlobalNameMap = LoadNameBatch(ScriptObjectsArchive);
+	TArray<FDisplayNameEntryId> GlobalNameMap = LoadNameBatch(ScriptObjectsArchive);
 	int32 NumScriptObjects = 0;
 	ScriptObjectsArchive << NumScriptObjects;
 	const FScriptObjectEntry* ScriptObjectEntries = reinterpret_cast<const FScriptObjectEntry*>(ScriptObjectsBuffer.ValueOrDie().Data() + ScriptObjectsArchive.Tell());
@@ -3895,7 +3895,7 @@ int32 Describe(
 		FMappedName MappedName = ScriptObjectEntry.Mapped;
 		check(MappedName.IsGlobal());
 		FScriptObjectDesc& ScriptObjectDesc = ScriptObjectByGlobalIdMap.Add(ScriptObjectEntry.GlobalIndex);
-		ScriptObjectDesc.Name = FName::CreateFromDisplayId(GlobalNameMap[MappedName.GetIndex()], MappedName.GetNumber());
+		ScriptObjectDesc.Name = GlobalNameMap[MappedName.GetIndex()].ToName(MappedName.GetNumber());
 		ScriptObjectDesc.GlobalImportIndex = ScriptObjectEntry.GlobalIndex;
 		ScriptObjectDesc.OuterIndex = ScriptObjectEntry.OuterIndex;
 	}
@@ -4106,13 +4106,13 @@ int32 Describe(
 			HeaderDataReader << VersioningInfo;
 		}
 
-		TArray<FNameEntryId> PackageNameMap;
+		TArray<FDisplayNameEntryId> PackageNameMap;
 		{
 			TRACE_CPUPROFILER_EVENT_SCOPE(LoadNameBatch);
 			PackageNameMap = LoadNameBatch(HeaderDataReader);
 		}
 
-		Job.PackageDesc->PackageName = FName::CreateFromDisplayId(PackageNameMap[PackageSummary->Name.GetIndex()], PackageSummary->Name.GetNumber());
+		Job.PackageDesc->PackageName = PackageNameMap[PackageSummary->Name.GetIndex()].ToName(PackageSummary->Name.GetNumber());
 		Job.PackageDesc->PackageFlags = PackageSummary->PackageFlags;
 		Job.PackageDesc->NameCount = PackageNameMap.Num();
 		
@@ -4132,7 +4132,7 @@ int32 Describe(
 			const FExportMapEntry& ExportMapEntry = ExportMap[ExportIndex];
 			FExportDesc& ExportDesc = Job.PackageDesc->Exports[ExportIndex];
 			ExportDesc.Package = Job.PackageDesc;
-			ExportDesc.Name = FName::CreateFromDisplayId(PackageNameMap[ExportMapEntry.ObjectName.GetIndex()], ExportMapEntry.ObjectName.GetNumber());
+			ExportDesc.Name = PackageNameMap[ExportMapEntry.ObjectName.GetIndex()].ToName(ExportMapEntry.ObjectName.GetNumber());
 			ExportDesc.OuterIndex = ExportMapEntry.OuterIndex;
 			ExportDesc.ClassIndex = ExportMapEntry.ClassIndex;
 			ExportDesc.SuperIndex = ExportMapEntry.SuperIndex;
diff --git a/Engine/Source/Developer/IoStoreUtilities/Private/PackageStoreOptimizer.cpp b/Engine/Source/Developer/IoStoreUtilities/Private/PackageStoreOptimizer.cpp
--- a/Engine/Source/Developer/IoStoreUtilities/Private/PackageStoreOptimizer.cpp
+++ b/Engine/Source/Developer/IoStoreUtilities/Private/PackageStoreOptimizer.cpp
@@ -187,9 +187,9 @@ FPackageStorePackage* FPackageStoreOptimizer::CreatePackageFromPackageStoreHeade
 	}
 	Package->PackageFlags = PackageStoreHeaderData.Summary.PackageFlags;
 	Package->CookedHeaderSize = PackageStoreHeaderData.Summary.CookedHeaderSize;
-	for (const FNameEntryId& DisplayIndex : PackageStoreHeaderData.NameMap)
+	for (FDisplayNameEntryId DisplayId : PackageStoreHeaderData.NameMap)
 	{
-		Package->NameMapBuilder.AddName(FName::CreateFromDisplayId(DisplayIndex, 0));
+		Package->NameMapBuilder.AddName(DisplayId);
 	}
 	ProcessImports(PackageStoreHeaderData, Package);
 	ProcessExports(PackageStoreHeaderData, Package);
@@ -319,7 +319,7 @@ FPackageStoreOptimizer::FPackageStoreHeaderData FPackageStoreOptimizer::LoadPack
 		HeaderDataReader << VersioningInfo;
 	}
 
-	TArray<FNameEntryId>& NameMap = PackageStoreHeaderData.NameMap;
+	TArray<FDisplayNameEntryId>& NameMap = PackageStoreHeaderData.NameMap;
 	NameMap = LoadNameBatch(HeaderDataReader);
 
 	for (FPackageId PackageId : PackageStoreEntry.ImportedPackageIds)
@@ -654,7 +654,7 @@ void FPackageStoreOptimizer::ProcessExports(const FPackageStoreHeaderData& Packa
 	Package->Exports.SetNum(ExportCount);
 	Package->ExportGraphNodes.Reserve(ExportCount * 2);
 
-	const TArray<FNameEntryId>& NameMap = PackageStoreHeaderData.NameMap;
+	const TArray<FDisplayNameEntryId>& NameMap = PackageStoreHeaderData.NameMap;
 
 	Package->ImportedPublicExportHashes = PackageStoreHeaderData.ImportedPublicExportHashes;
 	for (int32 ExportIndex = 0; ExportIndex < ExportCount; ++ExportIndex)
@@ -663,7 +663,7 @@ void FPackageStoreOptimizer::ProcessExports(const FPackageStoreHeaderData& Packa
 		Package->ExportsSerialSize += ExportEntry.CookedSerialSize;
 
 		FPackageStorePackage::FExport& Export = Package->Exports[ExportIndex];
-		Export.ObjectName = FName::CreateFromDisplayId(NameMap[ExportEntry.ObjectName.GetIndex()], ExportEntry.ObjectName.GetNumber());
+		Export.ObjectName = ExportEntry.ObjectName.ResolveName(NameMap);
 		Export.PublicExportHash = ExportEntry.PublicExportHash;
 		Export.OuterIndex = ExportEntry.OuterIndex;
 		Export.ClassIndex = ExportEntry.ClassIndex;
@@ -1809,7 +1809,7 @@ FIoBuffer FPackageStoreOptimizer::CreateScriptObjectsBuffer() const
 void FPackageStoreOptimizer::LoadScriptObjectsBuffer(const FIoBuffer& ScriptObjectsBuffer)
 {
 	FLargeMemoryReader ScriptObjectsArchive(ScriptObjectsBuffer.Data(), ScriptObjectsBuffer.DataSize());
-	TArray<FNameEntryId> NameMap = LoadNameBatch(ScriptObjectsArchive);
+	TArray<FDisplayNameEntryId> NameMap = LoadNameBatch(ScriptObjectsArchive);
 	int32 NumScriptObjects;
 	ScriptObjectsArchive << NumScriptObjects;
 	for (int32 Index = 0; Index < NumScriptObjects; ++Index)
@@ -1818,7 +1818,7 @@ void FPackageStoreOptimizer::LoadScriptObjectsBuffer(const FIoBuffer& ScriptObje
 		ScriptObjectsArchive << Entry;
 		FScriptObjectData& ImportData = ScriptObjectsMap.Add(Entry.GlobalIndex);
 		FMappedName MappedName = Entry.Mapped;
-		ImportData.ObjectName = FName::CreateFromDisplayId(NameMap[MappedName.GetIndex()], MappedName.GetNumber());
+		ImportData.ObjectName = NameMap[MappedName.GetIndex()].ToName(MappedName.GetNumber());
 		ImportData.GlobalIndex = Entry.GlobalIndex;
 		ImportData.OuterIndex = Entry.OuterIndex;
 		ImportData.CDOClassIndex = Entry.CDOClassIndex;
diff --git a/Engine/Source/Runtime/AssetRegistry/Private/AssetRegistryArchive.cpp b/Engine/Source/Runtime/AssetRegistry/Private/AssetRegistryArchive.cpp
--- a/Engine/Source/Runtime/AssetRegistry/Private/AssetRegistryArchive.cpp
+++ b/Engine/Source/Runtime/AssetRegistry/Private/AssetRegistryArchive.cpp
@@ -101,11 +101,11 @@ FAssetRegistryWriter::FAssetRegistryWriter(const FAssetRegistryWriterOptions& Op
 	check(!IsLoading());
 }
 
-static TArray<FNameEntryId> FlattenIndex(const TMap<FNameEntryId, uint32>& Names)
+static TArray<FDisplayNameEntryId> FlattenIndex(const TMap<FDisplayNameEntryId, uint32>& Names)
 {
-	TArray<FNameEntryId> Out;
+	TArray<FDisplayNameEntryId> Out;
 	Out.SetNumZeroed(Names.Num());
-	for (TPair<FNameEntryId, uint32> Pair : Names)
+	for (TPair<FDisplayNameEntryId, uint32> Pair : Names)
 	{
 		Out[Pair.Value] = Pair.Key;
 	}
@@ -126,7 +126,7 @@ FAssetRegistryWriter::~FAssetRegistryWriter()
 
 FArchive& FAssetRegistryWriter::operator<<(FName& Value)
 {
-	FNameEntryId EntryId = Value.GetDisplayIndex();
+	FDisplayNameEntryId EntryId(Value);
 
 	uint32 Index = Names.FindOrAdd(EntryId, Names.Num());
 	check((Index & AssetRegistryNumberedNameBit) == 0);
@@ -162,7 +162,7 @@ FAssetRegistryReader::FAssetRegistryReader(FArchive& Inner, int32 NumWorkers)
 
 	if (NumWorkers > 0)
 	{
-		TFunction<TArray<FNameEntryId> ()> GetFutureNames = LoadNameBatchAsync(*this, NumWorkers);
+		TFunction<TArray<FDisplayNameEntryId> ()> GetFutureNames = LoadNameBatchAsync(*this, NumWorkers);
 
 		FixedTagPrivate::FAsyncStoreLoader StoreLoader;
 		Task = StoreLoader.ReadInitialDataAndKickLoad(*this, NumWorkers);
@@ -205,7 +205,7 @@ FArchive& FAssetRegistryReader::operator<<(FName& Out)
 		*this << Number;
 	}
 
-	Out = FName::CreateFromDisplayId(Names[Index], Number);
+	Out = Names[Index].ToName(Number);
 
 	return *this;
 }
diff --git a/Engine/Source/Runtime/AssetRegistry/Private/AssetRegistryArchive.h b/Engine/Source/Runtime/AssetRegistry/Private/AssetRegistryArchive.h
--- a/Engine/Source/Runtime/AssetRegistry/Private/AssetRegistryArchive.h
+++ b/Engine/Source/Runtime/AssetRegistry/Private/AssetRegistryArchive.h
@@ -25,7 +25,7 @@ public:
 	void WaitForTasks();
 
 private:
-	TArray<FNameEntryId> Names;
+	TArray<FDisplayNameEntryId> Names;
 	TRefCountPtr<const FixedTagPrivate::FStore> Tags;
 	TFuture<void> Task;
 
@@ -72,7 +72,7 @@ public:
 	void SerializeTagsAndBundles(const FAssetData& In);
 
 private:
-	TMap<FNameEntryId, uint32> Names;
+	TMap<FDisplayNameEntryId, uint32> Names;
 	FixedTagPrivate::FStoreBuilder Tags;
 	FArchive& TargetAr;
 
diff --git a/Engine/Source/Runtime/Core/Private/UObject/UnrealNames.cpp b/Engine/Source/Runtime/Core/Private/UObject/UnrealNames.cpp
--- a/Engine/Source/Runtime/Core/Private/UObject/UnrealNames.cpp
+++ b/Engine/Source/Runtime/Core/Private/UObject/UnrealNames.cpp
@@ -302,13 +302,6 @@ FArchive& operator<<(FArchive& Ar, FNameEntryId& Id)
 	return Ar;
 }
 
-FNameEntryId FNameEntryId::FromUnstableInt(uint32 Value)
-{
-	FNameEntryId Id;
-	Id.Value = Value;
-	return Id;
-}
-
 struct FNameSlot
 {
 	// Use the remaining few bits to store a hash that can determine inequality
@@ -892,14 +885,18 @@ template<ENameCase Sensitivity>
 struct FNameLoad
 {
 	FNameValue<Sensitivity> In;
-	FNameEntryId* Out = nullptr;
+	FDisplayNameEntryId* Out = nullptr;
 	bool bInReuseComparisonEntry = false;
 	bool bOutCreatedNewEntry = false;
+
+	void SetOut(FNameEntryId Id);
 };
 
 using FNameComparisonLoad = FNameLoad<ENameCase::IgnoreCase>;
+template<> void FNameComparisonLoad::SetOut(FNameEntryId Id) { Out->SetLoadedComparisonId(Id); }
 #if WITH_CASE_PRESERVING_NAME
 using FNameDisplayLoad = FNameLoad<ENameCase::CaseSensitive>;
+template<> void FNameDisplayLoad::SetOut(FNameEntryId Id) { Out->SetLoadedDifferentDisplayId(Id); }
 #endif
 
 class alignas(PLATFORM_CACHE_LINE_SIZE) FNamePoolShardBase : FNoncopyable
@@ -1192,7 +1189,7 @@ public:
 			if (Sensitivity == ENameCase::CaseSensitive && Request.bInReuseComparisonEntry)
 			{
 				// Reuse comparison id as display id
-				check(*Request.Out == GetExistingComparisonId(Request.In));
+				check(*Request.Out == GetExistingComparisonId(Request.In).GetValue());
 				Request.bOutCreatedNewEntry = false;
 				++NumNewSlots;
 			}
@@ -1206,7 +1203,7 @@ public:
 				// Probe and write Request.Out without branching.
 				// It'll be overwritten below if bOutCreatedNewEntry is false.
 				FNameSlot Slot = Probe(Request.In);
-				*Request.Out = Slot.GetId();
+				Request.SetOut(Slot.GetId());
 				Request.bOutCreatedNewEntry = !Slot.Used();
 				NumNewSlots += !Slot.Used();
 			}
@@ -1231,7 +1228,7 @@ public:
 					// Check Slot.Used() since batch may contain duplicates
 					FNameSlot& Slot = Probe(Request.In);
 					Request.bOutCreatedNewEntry = !Slot.Used();
-					*Request.Out = Slot.Used() ? Slot.GetId() : CreateAndInsertEntry<FNullScopeLock>(Slot, Request.In);
+					Request.SetOut(Slot.Used() ? Slot.GetId() : CreateAndInsertEntry<FNullScopeLock>(Slot, Request.In));
 				}
 			}
 			Entries->BatchUnlock();
@@ -1446,10 +1443,10 @@ public:
 
 	bool			IsValid(FNameEntryHandle Handle) const;
 
-	FNameEntryId	StoreValue(const FNameComparisonValue& Value);
+	FDisplayNameEntryId	StoreValue(const FNameComparisonValue& Value);
 	void			StoreBatch(uint32 ShardIdx, TArrayView<FNameComparisonLoad> Batch)	{ ComparisonShards[ShardIdx].InsertBatch(Batch); }
 #if WITH_CASE_PRESERVING_NAME
-	FNameEntryId	StoreValue(const FNameDisplayValue& Value, bool bReuseComparisonId);
+	FDisplayNameEntryId	StoreValue(const FNameDisplayValue& Value, bool bReuseComparisonId);
 	void			StoreBatch(uint32 ShardIdx, TArrayView<FNameDisplayLoad> Batch)		{ DisplayShards[ShardIdx].InsertBatch(Batch); }
 	bool			ReuseComparisonEntry(bool bAddedComparisonEntry, const FNameDisplayValue& DisplayValue);
 #endif
@@ -1583,7 +1580,7 @@ FNameEntryId FNamePool::Store(FNameStringView Name)
 
 #if WITH_CASE_PRESERVING_NAME
 	DisplayValue.ComparisonId = ComparisonId;
-	return StoreValue(DisplayValue, bAdded);
+	return StoreValue(DisplayValue, bAdded).ToDisplayId();
 #else
 	return ComparisonId;
 #endif
@@ -1648,7 +1645,7 @@ FNameEntryId FNamePool::StoreValueWithNumber(FNumberedNameDisplayValue Value, bo
 #endif // WITH_CASE_PRESERVING_NAME
 #endif // UE_FNAME_OUTLINE_NUMBER 
 
-FORCEINLINE FNameEntryId FNamePool::StoreValue(const FNameComparisonValue& ComparisonValue)
+FORCEINLINE FDisplayNameEntryId FNamePool::StoreValue(const FNameComparisonValue& ComparisonValue)
 {
 	bool bAdded = false;
 	FNameEntryId ComparisonId = ComparisonShards[ComparisonValue.Hash.ShardIndex].Insert(ComparisonValue, bAdded);
@@ -1658,7 +1655,7 @@ FORCEINLINE FNameEntryId FNamePool::StoreValue(const FNameComparisonValue& Compa
 	DisplayValue.ComparisonId = ComparisonId;
 	return StoreValue(DisplayValue, bAdded);
 #else
-	return ComparisonId;
+	return FDisplayNameEntryId::FromComparisonId(ComparisonId);
 #endif
 }
 
@@ -1668,19 +1665,21 @@ bool FNamePool::ReuseComparisonEntry(bool bAddedComparisonEntry, const FNameDisp
 	return bAddedComparisonEntry || EqualsSameDimensions<ENameCase::CaseSensitive>(Resolve(DisplayValue.ComparisonId), DisplayValue.Name);
 }
 
-FORCEINLINE FNameEntryId FNamePool::StoreValue(const FNameDisplayValue& DisplayValue, bool bAddedComparisonEntry)
+FORCEINLINE FDisplayNameEntryId FNamePool::StoreValue(const FNameDisplayValue& DisplayValue, bool bAddedComparisonEntry)
 {
-	FNamePoolShard<ENameCase::CaseSensitive>& DisplayShard = DisplayShards[DisplayValue.Hash.ShardIndex];
+	FDisplayNameEntryId Out = FDisplayNameEntryId::FromComparisonId(DisplayValue.ComparisonId);
 
+	FNamePoolShard<ENameCase::CaseSensitive>& DisplayShard = DisplayShards[DisplayValue.Hash.ShardIndex];
 	if (ReuseComparisonEntry(bAddedComparisonEntry, DisplayValue))
 	{
 		DisplayShard.InsertExistingEntry(DisplayValue.Hash, DisplayValue.ComparisonId);
-		return DisplayValue.ComparisonId;
 	}
 	else
 	{
-		return DisplayShard.Insert(DisplayValue, bAddedComparisonEntry);
+		Out.SetLoadedDifferentDisplayId(DisplayShard.Insert(DisplayValue, bAddedComparisonEntry));
 	}
+
+	return Out;
 }
 #endif
 
@@ -3853,6 +3852,47 @@ bool operator==(const FLazyName& A, const FLazyName& B)
 	FName batch serialization
 -----------------------------------------------------------------------------*/
 
+FORCEINLINE FNameEntryId FDisplayNameEntryId::ToDisplayId() const
+{
+	return GetDisplayId();
+}
+
+FORCEINLINE FDisplayNameEntryId FDisplayNameEntryId::FromComparisonId(FNameEntryId ComparisonId)
+{
+	return FDisplayNameEntryId(ComparisonId, ComparisonId);
+}
+
+#if WITH_CASE_PRESERVING_NAME
+
+FNameEntryId FDisplayNameEntryId::GetLoadedComparisonId() const
+{
+	checkf(SameIds(), TEXT("SetLoadedComparisonId() should've been called at this point. The display id should not have been created yet. "));
+	return GetComparisonId();
+}
+
+void FDisplayNameEntryId::SetLoadedDifferentDisplayId(FNameEntryId Id)
+{
+	static_assert((FNameEntryIdMask & DifferentIdsFlag) == 0u);
+	check((Id.ToUnstableInt() & DifferentIdsFlag) == 0u);
+	Value = Id.ToUnstableInt() | DifferentIdsFlag;
+}
+
+void FDisplayNameEntryId::SetLoadedComparisonId(FNameEntryId Id)
+{
+	check((Id.ToUnstableInt() & DifferentIdsFlag) == 0u);
+	Value = Id.ToUnstableInt();
+}
+
+#else
+
+FORCEINLINE void FDisplayNameEntryId::SetLoadedComparisonId(FNameEntryId InId)
+{
+	Id = InId;
+}
+
+#endif
+
+
 constexpr bool CanCastUtf16ToWideCharWithoutConversion()
 {
 #if PLATFORM_LITTLE_ENDIAN
@@ -4000,7 +4040,7 @@ static void SaveHeaderAndName(TArray<uint8>& Out, FNameStringView Name)
 	FMemory::Memcpy(NameData, Name.Data, NameBytes);
 }
 
-void SaveNameBatch(TArrayView<const FNameEntryId> Names, TArray<uint8>& OutNameData, TArray<uint8>& OutHashData)
+void SaveNameBatch(TConstArrayView<FDisplayNameEntryId> Names, TArray<uint8>& OutNameData, TArray<uint8>& OutHashData)
 {
 	OutNameData.Empty(/* average bytes per name guesstimate */ 40 * Names.Num());
 	OutHashData.Empty((/* hash version */ 1 + Names.Num()) * sizeof(uint64));
@@ -4010,9 +4050,9 @@ void SaveNameBatch(TArrayView<const FNameEntryId> Names, TArray<uint8>& OutNameD
 
 	// Save names and hashes
 	FNameBuffer CustomDecodeBuffer;
-	for (FNameEntryId EntryId : Names)
+	for (FDisplayNameEntryId EntryId : Names)
 	{
-		FNameStringView Name = GetNamePoolPostInit().Resolve(EntryId).MakeView(CustomDecodeBuffer);
+		FNameStringView Name = GetNamePoolPostInit().Resolve(EntryId.ToDisplayId()).MakeView(CustomDecodeBuffer);
 
 		SaveHeaderAndName(OutNameData, Name);
 
@@ -4020,7 +4060,7 @@ void SaveNameBatch(TArrayView<const FNameEntryId> Names, TArray<uint8>& OutNameD
 	}
 }
 
-void SaveNameBatch(TArrayView<const FNameEntryId> Names, FArchive& Ar)
+void SaveNameBatch(TConstArrayView<FDisplayNameEntryId> Names, FArchive& Ar)
 {
 	// Save Num and bail if zero
 	uint32 Num = Names.Num();
@@ -4040,9 +4080,9 @@ void SaveNameBatch(TArrayView<const FNameEntryId> Names, FArchive& Ar)
 	Strings.Reserve(Num * 40 /* guesstimate */);
 
 	FNameBuffer Buffer;
-	for (FNameEntryId EntryId : Names)
+	for (FDisplayNameEntryId EntryId : Names)
 	{
-		FNameStringView Name = GetNamePoolPostInit().Resolve(EntryId).MakeView(Buffer);
+		FNameStringView Name = GetNamePoolPostInit().Resolve(EntryId.ToDisplayId()).MakeView(Buffer);
 
 		Hashes.Add(GenerateLowerCaseHash(Name));
 		Headers.Add(FSerializedNameHeader(Name.Len, Name.bIsWide));
@@ -4071,7 +4111,7 @@ FORCENOINLINE void ReserveNameBatch(uint32 NameDataBytes, uint32 HashDataBytes)
 	GetNamePoolPostInit().Reserve(AddSlack(NameDataBytes), AddSlack(NumEntries));
 }
 
-static FNameEntryId BatchLoadNameWithoutHash(const UTF16CHAR* Str, uint32 Len)
+static FDisplayNameEntryId BatchLoadNameWithoutHash(const UTF16CHAR* Str, uint32 Len)
 {
 	WIDECHAR Temp[NAME_SIZE];
 	for (uint32 Idx = 0; Idx < Len; ++Idx)
@@ -4089,21 +4129,21 @@ static FNameEntryId BatchLoadNameWithoutHash(const UTF16CHAR* Str, uint32 Len)
 	return GetNamePoolPostInit().StoreValue(FNameComparisonValue(Name, Hash));
 }
 
-static FNameEntryId BatchLoadNameWithoutHash(const ANSICHAR* Str, uint32 Len)
+static FDisplayNameEntryId BatchLoadNameWithoutHash(const ANSICHAR* Str, uint32 Len)
 {
 	FNameStringView Name(Str, Len);
 	FNameHash Hash = HashName<ENameCase::IgnoreCase>(Name);
 	return GetNamePoolPostInit().StoreValue(FNameComparisonValue(Name, Hash));
 }
 
-static FNameEntryId BatchLoadNameWithoutHash(const FNameSerializedView& Name)
+static FDisplayNameEntryId BatchLoadNameWithoutHash(const FNameSerializedView& Name)
 {
 	return Name.bIsUtf16 ? BatchLoadNameWithoutHash(Name.Utf16, Name.Len)
 						 : BatchLoadNameWithoutHash(Name.Ansi, Name.Len);
 }
 
 template<typename CharType>
-FNameEntryId BatchLoadNameWithHash(const CharType* Str, uint32 Len, uint64 InHash)
+FDisplayNameEntryId BatchLoadNameWithHash(const CharType* Str, uint32 Len, uint64 InHash)
 {
 	FNameStringView Name(Str, Len);
 	FNameHash Hash(Str, Len, InHash);
@@ -4111,7 +4151,7 @@ FNameEntryId BatchLoadNameWithHash(const CharType* Str, uint32 Len, uint64 InHas
 	return GetNamePoolPostInit().StoreValue(FNameComparisonValue(Name, Hash));
 }
 
-static FNameEntryId BatchLoadNameWithHash(const FNameSerializedView& InName, uint64 InHash)
+static FDisplayNameEntryId BatchLoadNameWithHash(const FNameSerializedView& InName, uint64 InHash)
 {
 	if (InName.bIsUtf16)
 	{
@@ -4132,9 +4172,9 @@ static FNameEntryId BatchLoadNameWithHash(const FNameSerializedView& InName, uin
 	}
 }
 
-static void LoadInterleavedNameBatchInInputOrder(TArray<FNameEntryId>& Out, TArrayView<const uint8> Names)
+static void LoadInterleavedNameBatchInInputOrder(TArray<FDisplayNameEntryId>& Out, TArrayView<const uint8> Names)
 {
-	TArray<FNameEntryId>::TIterator OutIt(Out);
+	TArray<FDisplayNameEntryId>::TIterator OutIt(Out);
 	const uint8* NameIt = Names.GetData();
 	const uint8* NameEnd = Names.GetData() + Names.Num();
 
@@ -4148,11 +4188,11 @@ static void LoadInterleavedNameBatchInInputOrder(TArray<FNameEntryId>& Out, TArr
 	check(OutIt.GetIndex() == Out.Num());
 }
 
-static void LoadInterleavedNameBatchInInputOrder(TArray<FNameEntryId>& Out, TArrayView<const uint64> Hashes, TArrayView<const uint8> Names)
+static void LoadInterleavedNameBatchInInputOrder(TArray<FDisplayNameEntryId>& Out, TArrayView<const uint64> Hashes, TArrayView<const uint8> Names)
 {
 	check(Hashes.Num() == Out.Num());
 
-	TArray<FNameEntryId>::TIterator OutIt(Out);
+	TArray<FDisplayNameEntryId>::TIterator OutIt(Out);
 	const uint8* NameIt = Names.GetData();
 
 	for (uint64 Hash : Hashes)
@@ -4184,13 +4224,13 @@ static FNameSerializedView LoadHeaderAndName(FArchive& Ar, uint8* OutNameBuffer)
 	return LoadName(Ar, Header, OutNameBuffer);
 }
 
-static void LoadSeparatedNameBatchInInputOrder(	TArray<FNameEntryId>& Out,
+static void LoadSeparatedNameBatchInInputOrder(	TArray<FDisplayNameEntryId>& Out,
 												const uint64* Hashes,
 												TArrayView<const FSerializedNameHeader> Headers, 
 												TArrayView<const uint8> Strings)
 {
 	check(Out.Num() == Headers.Num());
-	TArray<FNameEntryId>::TIterator OutIt(Out);
+	TArray<FDisplayNameEntryId>::TIterator OutIt(Out);
 	const uint8* StringIt = Strings.GetData();
 	for (FSerializedNameHeader Header : Headers)
 	{
@@ -4201,13 +4241,13 @@ static void LoadSeparatedNameBatchInInputOrder(	TArray<FNameEntryId>& Out,
 	check(StringIt == Strings.end());
 }
 
-static void LoadSeparatedNameBatchInInputOrder(	TArray<FNameEntryId>& Out,
+static void LoadSeparatedNameBatchInInputOrder(	TArray<FDisplayNameEntryId>& Out,
 												TArrayView<const FSerializedNameHeader> Headers, 
 												TArrayView<const uint8> Strings)
 {
 	check(Out.Num() == Headers.Num());
 
-	TArray<FNameEntryId>::TIterator OutIt(Out);
+	TArray<FDisplayNameEntryId>::TIterator OutIt(Out);
 	const uint8* StringIt = Strings.GetData();
 	for (FSerializedNameHeader Header : Headers)
 	{
@@ -4218,29 +4258,6 @@ static void LoadSeparatedNameBatchInInputOrder(	TArray<FNameEntryId>& Out,
 	check(StringIt == Strings.end());
 }
 
-// @param Loads are sorted in shard order
-template<class LoadType>
-void LoadOneBatchPerShard(TArray<LoadType>& Loads)
-{
-	FNamePool& Pool = GetNamePoolPostInit();
-
-	LoadType* BatchIt = &Loads[0];
-	for (LoadType& Load : Loads)
-	{
-		if (Load.In.Hash.ShardIndex != BatchIt->In.Hash.ShardIndex)
-		{
-			Pool.StoreBatch(BatchIt->In.Hash.ShardIndex, TArrayView<LoadType>(BatchIt, &Load - BatchIt));
-			BatchIt = &Load;
-		}
-	}
-
-	TArrayView<LoadType> LastBatch(BatchIt, &Loads.Last() + 1 - BatchIt);
-	Pool.StoreBatch(BatchIt->In.Hash.ShardIndex, LastBatch);
-
-	check(BatchIt + LastBatch.Num() == Loads.GetData() + Loads.Num());
-}
-
-
 #if WITH_CASE_PRESERVING_NAME
 
 FORCENOINLINE // To help profiling
@@ -4248,22 +4265,30 @@ static void LoadDisplayNames(const TArray<FNameComparisonLoad>& ComparisonLoads)
 {
 	FNamePool& Pool = GetNamePoolPostInit();
 
+	// Reserve load requests
+	TArray<FNameDisplayLoad> LoadShards[FNamePoolShards];
+	const uint32 ReservePerShard = (8 * ComparisonLoads.Num()) / (7 * FNamePoolShards);
+	for (TArray<FNameDisplayLoad>& LoadShard : LoadShards)
+	{
+		LoadShard.Reserve(ReservePerShard);
+	}
+
 	// Hash display names and prepare load requests
-	TArray<FNameDisplayLoad> DisplayLoads;
-	DisplayLoads.Reserve(ComparisonLoads.Num());
 	for (const FNameComparisonLoad& ComparisonLoad : ComparisonLoads)
 	{
 		FNameDisplayValue DisplayValue(ComparisonLoad.In.Name);
-		DisplayValue.ComparisonId = *ComparisonLoad.Out;
+		DisplayValue.ComparisonId = ComparisonLoad.Out->GetLoadedComparisonId();
 
 		bool bReuseEntry = Pool.ReuseComparisonEntry(ComparisonLoad.bOutCreatedNewEntry, DisplayValue);
-		DisplayLoads.Emplace(FNameDisplayLoad {DisplayValue, ComparisonLoad.Out, bReuseEntry});
+		LoadShards[DisplayValue.Hash.ShardIndex].Emplace(FNameDisplayLoad {DisplayValue, ComparisonLoad.Out, bReuseEntry});
 	}
 
-	// Sort loads in shard order
-	Algo::HeapSort(DisplayLoads, [](const FNameDisplayLoad& A, const FNameDisplayLoad& B) { return A.In.Hash.ShardIndex < B.In.Hash.ShardIndex; });
-
-	LoadOneBatchPerShard(DisplayLoads);
+	// Issue load requests
+	uint32 ShardIndex = 0;
+	for (TArrayView<FNameDisplayLoad> LoadShard : LoadShards)
+	{
+		Pool.StoreBatch(ShardIndex++, LoadShard);
+	}
 }
 
 #endif  // WITH_CASE_PRESERVING_NAME
@@ -4298,7 +4323,7 @@ static void InitializeTargets(FShardTargetArray& Targets, const TArrayView<const
 }
 
 
-static void LoadInterleavedNameBatchInShardOrder(	TArray<FNameEntryId>& Out,
+static void LoadInterleavedNameBatchInShardOrder(	TArray<FDisplayNameEntryId>& Out,
 													TArrayView<const uint64> Hashes,
 													TArrayView<const uint8> HeadersAndStrings)
 {
@@ -4349,7 +4374,7 @@ static void LoadInterleavedNameBatchInShardOrder(	TArray<FNameEntryId>& Out,
 #endif
 }
 
-static void LoadSeparatedNameBatchInShardOrder(	TArray<FNameEntryId>& Out,
+static void LoadSeparatedNameBatchInShardOrder(	TArray<FDisplayNameEntryId>& Out,
 												const uint64* Hashes,
 												TArrayView<const FSerializedNameHeader> Headers, 
 												TArrayView<const uint8> Strings)											
@@ -4440,7 +4465,7 @@ static void LoadSeparatedNameBatchInShardOrder(	TArray<FNameEntryId>& Out,
 #endif
 }
 
-void LoadNameBatch(TArray<FNameEntryId>& OutNames, TArrayView<const uint8> NameData, TArrayView<const uint8> HashData)
+void LoadNameBatch(TArray<FDisplayNameEntryId>& OutNames, TArrayView<const uint8> NameData, TArrayView<const uint8> HashData)
 {
 	check(IsAligned(NameData.GetData(), sizeof(uint64)));
 	check(IsAligned(HashData.GetData(), sizeof(uint64)));
@@ -4507,11 +4532,11 @@ struct FNameBatchLoader
 		return !Ar.IsError();
 	}
 
-	TArray<FNameEntryId> Load()
+	TArray<FDisplayNameEntryId> Load()
 	{
 		check(Headers.Num());
 
-		TArray<FNameEntryId> Out;
+		TArray<FDisplayNameEntryId> Out;
 		Out.SetNumUninitialized(Headers.Num());
 
 		if (Hashes.Num() == 0)
@@ -4531,7 +4556,7 @@ struct FNameBatchLoader
 	}
 };
 
-TArray<FNameEntryId> LoadNameBatch(FArchive& Ar)
+TArray<FDisplayNameEntryId> LoadNameBatch(FArchive& Ar)
 {
 	FNameBatchLoader Loader;
 
@@ -4540,7 +4565,7 @@ TArray<FNameEntryId> LoadNameBatch(FArchive& Ar)
 		return Loader.Load();
 	}
 
-	return TArray<FNameEntryId>();
+	return TArray<FDisplayNameEntryId>();
 }
 
 struct FNameBatchAsyncLoader : FNameBatchLoader
@@ -4567,7 +4592,7 @@ struct FNameBatchAsyncLoader : FNameBatchLoader
 		DoneEvent->Trigger();
 	}
 
-	TArray<FNameEntryId> GetResult()
+	TArray<FDisplayNameEntryId> GetResult()
 	{
 		check(DoneEvent);
 
@@ -4579,10 +4604,10 @@ struct FNameBatchAsyncLoader : FNameBatchLoader
 	}
 
 	FEvent* DoneEvent = nullptr;
-	TArray<FNameEntryId> Out;
+	TArray<FDisplayNameEntryId> Out;
 };
 
-TFunction<TArray<FNameEntryId>()> LoadNameBatchAsync(FArchive& Ar, uint32 MaxWorkers)
+TFunction<TArray<FDisplayNameEntryId>()> LoadNameBatchAsync(FArchive& Ar, uint32 MaxWorkers)
 {
 	check(MaxWorkers > 0);
 
@@ -4604,13 +4629,13 @@ TFunction<TArray<FNameEntryId>()> LoadNameBatchAsync(FArchive& Ar, uint32 MaxWor
 		}
 	}
 
-	return []() { return TArray<FNameEntryId>(); };
+	return []() { return TArray<FDisplayNameEntryId>(); };
 }
 
 
 #if 0 && ALLOW_NAME_BATCH_SAVING  
 
-FORCENOINLINE void PerfTestLoadNameBatch(TArray<FNameEntryId>& OutNames, TArrayView<const uint8> NameData, TArrayView<const uint8> HashData)
+FORCENOINLINE void PerfTestLoadNameBatch(TArray<FLoadedNameEntryId>& OutNames, TArrayView<const uint8> NameData, TArrayView<const uint8> HashData)
 {
 	LoadNameBatch(OutNames, NameData, HashData);
 }
@@ -4679,7 +4704,7 @@ int LoadNameBatchTestFiles(bool bArchivePath)
 	{
 		TArray<uint8> ArchData = ReadBlobFile(TEXT("TestNameBatch.Archive"));
 		FMemoryReader ArchReader(ArchData);
-		TArray<FNameEntryId> NameEntries = LoadNameBatch(ArchiveReader);
+		TArray<FLoadedNameEntryId> NameEntries = LoadNameBatch(ArchiveReader);
 		return NameEntries.Num();
 	}
 	else
@@ -4687,7 +4712,7 @@ int LoadNameBatchTestFiles(bool bArchivePath)
 		TArray<uint8> NameData = ReadBlobFile(TEXT("TestNameBatch.Names"));
 		TArray<uint8> HashData = ReadBlobFile(TEXT("TestNameBatch.Hashes"));
 
-		TArray<FNameEntryId> NameEntries;
+		TArray<FLoadedNameEntryId> NameEntries;
 		if (HashData.Num())
 		{
 			ReserveNameBatch(NameData.Num(), HashData.Num());
@@ -4704,7 +4729,7 @@ static void TestNameBatch()
 {
 #if ALLOW_NAME_BATCH_SAVING
 
-	TArray<FNameEntryId> Names;
+	TArray<FDisplayNameEntryId> Names;
 	TArray<uint8> NameData;
 	TArray<uint8> HashData;
 
@@ -4715,8 +4740,8 @@ static void TestNameBatch()
 	check(Names.Num() == 0);
 
 	// Test empty / "None" name and another EName
-	Names.Add(FName().GetComparisonIndex());
-	Names.Add(FName(NAME_Box).GetComparisonIndex());
+	Names.Add(FDisplayNameEntryId(FName()));
+	Names.Add(FDisplayNameEntryId(FName(NAME_Box)));
 
 	// Test long strings
 	FString MaxLengthAnsi;
@@ -4732,26 +4757,32 @@ static void TestNameBatch()
 
 	for (const FString& MaxLength : {MaxLengthAnsi, MaxLengthWide})
 	{
-		Names.Add(FName(*MaxLength).GetComparisonIndex());
-		Names.Add(FName(*MaxLength + NAME_SIZE - 255).GetComparisonIndex());
-		Names.Add(FName(*MaxLength + NAME_SIZE - 256).GetComparisonIndex());
-		Names.Add(FName(*MaxLength + NAME_SIZE - 257).GetComparisonIndex());
+		Names.Add(FDisplayNameEntryId(FName(*MaxLength)));
+		Names.Add(FDisplayNameEntryId(FName(*MaxLength + NAME_SIZE - 255)));
+		Names.Add(FDisplayNameEntryId(FName(*MaxLength + NAME_SIZE - 256)));
+		Names.Add(FDisplayNameEntryId(FName(*MaxLength + NAME_SIZE - 257)));
 	}
 
 	// Test UTF-16 alignment
 	FString Wide("Wide ");
 	Wide[4] = 60000;
 
-	Names.Add(FName(*Wide).GetComparisonIndex());
-	Names.Add(FName("odd").GetComparisonIndex());
-	Names.Add(FName(*Wide).GetComparisonIndex());
-	Names.Add(FName("even").GetComparisonIndex());
-	Names.Add(FName(*Wide).GetComparisonIndex());
+	Names.Add(FDisplayNameEntryId(FName(*Wide)));
+	Names.Add(FDisplayNameEntryId(FName("odd")));
+	Names.Add(FDisplayNameEntryId(FName(*Wide)));
+	Names.Add(FDisplayNameEntryId(FName("even")));
+	Names.Add(FDisplayNameEntryId(FName(*Wide)));
+
+	// Test preserve different casing
+	FName("case");
+	Names.Add(FDisplayNameEntryId(FName("CASE")));
+	Names.Add(FDisplayNameEntryId(FName("case")));
+	Names.Add(FDisplayNameEntryId(FName("casE")));
 
 	// Roundtrip names
 	SaveNameBatch(MakeArrayView(Names), NameData, HashData);
 	check(NameData.Num() > 0);
-	TArray<FNameEntryId> LoadedNames;
+	TArray<FDisplayNameEntryId> LoadedNames;
 	LoadNameBatch(LoadedNames, MakeArrayView(NameData), MakeArrayView(HashData));
 	check(LoadedNames == Names);
 
@@ -4785,21 +4816,21 @@ static void TestNameBatch()
 
 	// Test FArchive-based Load/SaveNameBatch()
 
-	auto TestArchiveRoundtrip = [](const TArray<FNameEntryId>& Names)
+	auto TestArchiveRoundtrip = [](const TArray<FDisplayNameEntryId>& Names)
 	{
 		TArray<uint8> ArchiveData;
 		FMemoryWriter ArchiveWriter(ArchiveData);
 		SaveNameBatch(MakeArrayView(Names), ArchiveWriter);
 
 		FMemoryReader ArchiveReader(ArchiveData);
-		TArray<FNameEntryId> ArchiveRoundtrippedNames = LoadNameBatch(ArchiveReader);
+		TArray<FDisplayNameEntryId> ArchiveRoundtrippedNames = LoadNameBatch(ArchiveReader);
 		
 		check(ArchiveRoundtrippedNames == Names);
 	};
 
 	TestArchiveRoundtrip(Names);
 	
-	TArray<FNameEntryId> LargeBatch = Names;
+	TArray<FDisplayNameEntryId> LargeBatch = Names;
 	for (char C1 = 'A'; C1 <= 'z' ; ++C1)
 	{
 		for (char C2 = 'A'; C2 < 'z'; ++C2)
@@ -4807,7 +4838,7 @@ static void TestNameBatch()
 			char Str[]  = {C1, C2, '\0'};
 
 			// Use display index to get some duplicate comparison ids like "aa" and "aA"
-			LargeBatch.Add(FName(Str).GetDisplayIndex());
+			LargeBatch.Add(FDisplayNameEntryId(FName(Str)));
 		}
 	}
 	check(LargeBatch.Num() > FNamePoolShards);
@@ -5215,4 +5246,4 @@ namespace UE::Name::Private
 		GetNamePoolPostInit().LogHashStats(OutDevice);
 		Out.Logf(TEXT("Wrote FName hash stats to %s"), *Filename);
 	}
-}
\ No newline at end of file
+}
diff --git a/Engine/Source/Runtime/Core/Public/IO/IoContainerHeader.h b/Engine/Source/Runtime/Core/Public/IO/IoContainerHeader.h
--- a/Engine/Source/Runtime/Core/Public/IO/IoContainerHeader.h
+++ b/Engine/Source/Runtime/Core/Public/IO/IoContainerHeader.h
@@ -81,7 +81,7 @@ struct FIoContainerHeader
 	uint32 PackageCount = 0;
 	TArray<FPackageId> PackageIds;
 	TArray<uint8> StoreEntries; //FPackageStoreEntry[PackageCount]
-	TArray<FNameEntryId> RedirectsNameMap;
+	TArray<FDisplayNameEntryId> RedirectsNameMap;
 	TArray<FIoContainerHeaderLocalizedPackage> LocalizedPackages;
 	TArray<FIoContainerHeaderPackageRedirect> PackageRedirects;
 
diff --git a/Engine/Source/Runtime/Core/Public/Serialization/MappedName.h b/Engine/Source/Runtime/Core/Public/Serialization/MappedName.h
--- a/Engine/Source/Runtime/Core/Public/Serialization/MappedName.h
+++ b/Engine/Source/Runtime/Core/Public/Serialization/MappedName.h
@@ -3,6 +3,7 @@
 #pragma once
 
 #include "CoreMinimal.h"
+#include "UObject/NameBatchSerialization.h"
 
 /**
  * Index and name number into a name map.
@@ -83,6 +84,11 @@ public:
 		return Index != Other.Index || Number != Other.Number;
 	}
 
+	FName ResolveName(TConstArrayView<FDisplayNameEntryId> Names) const
+	{
+		return Names[GetIndex()].ToName(GetNumber());
+	}
+
 	CORE_API friend FArchive& operator<<(FArchive& Ar, FMappedName& MappedName);
 
 private:
@@ -112,8 +118,7 @@ public:
 		check(MappedName.GetType() == NameMapType);
 		check(MappedName.GetIndex() < uint32(NameEntries.Num()));
 
-		FNameEntryId NameEntry = NameEntries[MappedName.GetIndex()];
-		return FName::CreateFromDisplayId(NameEntry, MappedName.GetNumber());
+		return NameEntries[MappedName.GetIndex()].ToName(MappedName.GetNumber());
 	}
 
 	bool TryGetName(const FMappedName& MappedName, FName& OutName) const
@@ -123,8 +128,7 @@ public:
 		uint32 Index = MappedName.GetIndex();
 		if (Index < uint32(NameEntries.Num()))
 		{
-			FNameEntryId NameEntry = NameEntries[MappedName.GetIndex()];
-			OutName = FName::CreateFromDisplayId(NameEntry, MappedName.GetNumber());
+			OutName = NameEntries[MappedName.GetIndex()].ToName(MappedName.GetNumber());
 
 			return true;
 		}
@@ -134,15 +138,11 @@ public:
 
 	FMinimalName GetMinimalName(const FMappedName& MappedName) const
 	{
-		check(MappedName.GetType() == NameMapType);
-		check(MappedName.GetIndex() < uint32(NameEntries.Num()));
-
-		FNameEntryId DisplayId = NameEntries[MappedName.GetIndex()];
-		FNameEntryId ComparisonId = FName::GetComparisonIdFromDisplayId(DisplayId);
-		return NameToMinimalName(FName(ComparisonId, DisplayId, MappedName.GetNumber()));
+		// Wasteful, looks up display name then discards it
+		return NameToMinimalName(GetName(MappedName));
 	}
 
 private:
-	TArray<FNameEntryId> NameEntries;
+	TArray<FDisplayNameEntryId> NameEntries;
 	FMappedName::EType NameMapType = FMappedName::EType::Global;
 };
diff --git a/Engine/Source/Runtime/Core/Public/UObject/NameBatchSerialization.h b/Engine/Source/Runtime/Core/Public/UObject/NameBatchSerialization.h
--- a/Engine/Source/Runtime/Core/Public/UObject/NameBatchSerialization.h
+++ b/Engine/Source/Runtime/Core/Public/UObject/NameBatchSerialization.h
@@ -8,14 +8,67 @@
 
 #define ALLOW_NAME_BATCH_SAVING PLATFORM_LITTLE_ENDIAN && !PLATFORM_TCHAR_IS_4_BYTES
 
+//////////////////////////////////////////////////////////////////////////
+
+/** 
+ * Serialization util that optimizes WITH_CASE_PRESERVING_NAME-loading by reducing comparison id lookups
+ * 
+ * Stores 32-bit display entry id with an unused bit to indicate if FName::GetComparisonIdFromDisplayId lookup is needed.
+ * 
+ * Note that only display entries should be saved to make output deterministic.
+ */
+class FDisplayNameEntryId
+{
+public:
+	explicit FDisplayNameEntryId(FName Name) : FDisplayNameEntryId(Name.GetDisplayIndex(), Name.GetComparisonIndex()) {}
+	FORCEINLINE FName ToName(uint32 Number) const { return FName(GetComparisonId(), GetDisplayId(), Number); }
+
+private:
+#if WITH_CASE_PRESERVING_NAME
+	static constexpr uint32 DifferentIdsFlag = 1u << 31;
+	static constexpr uint32 DisplayIdMask = ~DifferentIdsFlag;
+
+	uint32 Value = 0;
+
+	FDisplayNameEntryId(FNameEntryId Id, FNameEntryId CmpId) : Value(Id.ToUnstableInt() | (Id != CmpId) * DifferentIdsFlag) {}
+	FORCEINLINE bool SameIds() const { return (Value & DifferentIdsFlag) == 0; }
+	FORCEINLINE FNameEntryId GetDisplayId() const	 { return FNameEntryId::FromUnstableInt(Value & DisplayIdMask); }
+	FORCEINLINE FNameEntryId GetComparisonId() const { return SameIds() ? GetDisplayId() : FName::GetComparisonIdFromDisplayId(GetDisplayId()); }
+	friend bool operator==(FDisplayNameEntryId A, FDisplayNameEntryId B) { return A.Value == B.Value; }
+#else
+	FNameEntryId Id;
+
+	FDisplayNameEntryId(FNameEntryId InId, FNameEntryId) : Id(InId) {}
+	FORCEINLINE FNameEntryId GetDisplayId() const	 { return Id; }
+	FORCEINLINE FNameEntryId GetComparisonId() const { return Id; }
+	friend bool operator==(FDisplayNameEntryId A, FDisplayNameEntryId B) { return A.Id == B.Id; }
+#endif
+	friend bool operator==(FNameEntryId A, FDisplayNameEntryId B) { return A == B.GetDisplayId(); }
+	friend bool operator==(FDisplayNameEntryId A, FNameEntryId B) { return A.GetDisplayId() == B; }
+	friend uint32 GetTypeHash(FDisplayNameEntryId InId) { return GetTypeHash(InId.GetDisplayId()); }
+
+public: // Internal functions for batch serialization code - intentionally lacking CORE_API
+	static FDisplayNameEntryId FromComparisonId(FNameEntryId ComparisonId);
+	FNameEntryId ToDisplayId() const;
+	void SetLoadedComparisonId(FNameEntryId ComparisonId); // Called first
+#if WITH_CASE_PRESERVING_NAME
+	void SetLoadedDifferentDisplayId(FNameEntryId DisplayId); // Called second if display id differs
+	FNameEntryId GetLoadedComparisonId() const; // Get the already loaded comparison id
+#endif
+};
+
+//////////////////////////////////////////////////////////////////////////
+
 #if ALLOW_NAME_BATCH_SAVING
-// Save comparison entries in given order to a name blob and a versioned hash blob.
-CORE_API void SaveNameBatch(TArrayView<const FNameEntryId> Names, TArray<uint8>& OutNameData, TArray<uint8>& OutHashData);
+// Save display entries in given order to a name blob and a versioned hash blob.
+CORE_API void SaveNameBatch(TArrayView<const FDisplayNameEntryId> Names, TArray<uint8>& OutNameData, TArray<uint8>& OutHashData);
 
-// Save comparison entries in given order to an archive
-CORE_API void SaveNameBatch(TArrayView<const FNameEntryId> Names, FArchive& Out);
+// Save display entries in given order to an archive
+CORE_API void SaveNameBatch(TArrayView<const FDisplayNameEntryId> Names, FArchive& Out);
 #endif
 
+//////////////////////////////////////////////////////////////////////////
+
 // Reserve memory in preparation for batch loading
 //
 // @param Bytes for existing and new names.
@@ -26,12 +79,12 @@ CORE_API void ReserveNameBatch(uint32 NameDataBytes, uint32 HashDataBytes);
 // Names are rehashed if hash algorithm version doesn't match.
 //
 // @param NameData, HashData must be 8-byte aligned.
-CORE_API void LoadNameBatch(TArray<FNameEntryId>& OutNames, TArrayView<const uint8> NameData, TArrayView<const uint8> HashData);
+CORE_API void LoadNameBatch(TArray<FDisplayNameEntryId>& OutNames, TArrayView<const uint8> NameData, TArrayView<const uint8> HashData);
 
 // Load names and precalculated hashes from an archive
 //
 // Names are rehashed if hash algorithm version doesn't match.
-CORE_API TArray<FNameEntryId> LoadNameBatch(FArchive& Ar);
+CORE_API TArray<FDisplayNameEntryId> LoadNameBatch(FArchive& Ar);
 
 // Load names and precalculated hashes from an archive using multiple workers
 //
@@ -42,4 +95,6 @@ CORE_API TArray<FNameEntryId> LoadNameBatch(FArchive& Ar);
 // @param Ar is drained synchronously
 // @param MaxWorkers > 0
 // @return function that waits before returning result, like a simple future.
-CORE_API TFunction<TArray<FNameEntryId>()> LoadNameBatchAsync(FArchive& Ar, uint32 MaxWorkers);
\ No newline at end of file
+CORE_API TFunction<TArray<FDisplayNameEntryId>()> LoadNameBatchAsync(FArchive& Ar, uint32 MaxWorkers);
+
+//////////////////////////////////////////////////////////////////////////
\ No newline at end of file
diff --git a/Engine/Source/Runtime/Core/Public/UObject/NameTypes.h b/Engine/Source/Runtime/Core/Public/UObject/NameTypes.h
--- a/Engine/Source/Runtime/Core/Public/UObject/NameTypes.h
+++ b/Engine/Source/Runtime/Core/Public/UObject/NameTypes.h
@@ -85,7 +85,12 @@ struct FNameEntryId
 	uint32 ToUnstableInt() const { return Value; }
 
 	/** Create from unstable int produced by this process */
-	CORE_API static FNameEntryId FromUnstableInt(uint32 UnstableInt);
+	static FNameEntryId FromUnstableInt(uint32 UnstableInt)
+	{
+		FNameEntryId Id;
+		Id.Value = UnstableInt;
+		return Id;
+	}
 
 	FORCEINLINE static FNameEntryId FromEName(EName Ename)
 	{
@@ -836,7 +841,9 @@ public:
 	 */
 	static FName CreateFromDisplayId(FNameEntryId DisplayId, int32 Number)
 	{
-		check(ResolveEntry(DisplayId)->IsNumbered() == false); // This id should be unnumbered i.e. returned from GetDisplayIndex on an FName.
+#if UE_FNAME_OUTLINE_NUMBER
+		checkSlow(ResolveEntry(DisplayId)->IsNumbered() == false); // This id should be unnumbered i.e. returned from GetDisplayIndex on an FName.
+#endif
 		return FName(GetComparisonIdFromDisplayId(DisplayId), DisplayId, Number);
 	}
 
diff --git a/Engine/Source/Runtime/PakFile/Private/FilePackageStore.cpp b/Engine/Source/Runtime/PakFile/Private/FilePackageStore.cpp
--- a/Engine/Source/Runtime/PakFile/Private/FilePackageStore.cpp
+++ b/Engine/Source/Runtime/PakFile/Private/FilePackageStore.cpp
@@ -204,15 +204,15 @@ void FFilePackageStore::Update()
 			FName& LocalizedPackageSourceName = LocalizedPackages.FindOrAdd(LocalizedPackage.SourcePackageId);
 			if (LocalizedPackageSourceName.IsNone())
 			{
-				FNameEntryId NameEntry = ContainerHeader->RedirectsNameMap[LocalizedPackage.SourcePackageName.GetIndex()];
-				LocalizedPackageSourceName = FName::CreateFromDisplayId(NameEntry, LocalizedPackage.SourcePackageName.GetNumber());
+				FDisplayNameEntryId NameEntry = ContainerHeader->RedirectsNameMap[LocalizedPackage.SourcePackageName.GetIndex()];
+				LocalizedPackageSourceName = NameEntry.ToName(LocalizedPackage.SourcePackageName.GetNumber());
 			}
 		}
 
 		for (const FIoContainerHeaderPackageRedirect& Redirect : ContainerHeader->PackageRedirects)
 		{
-			FNameEntryId NameEntry = ContainerHeader->RedirectsNameMap[Redirect.SourcePackageName.GetIndex()];
-			FName SourcePackageName = FName::CreateFromDisplayId(NameEntry, Redirect.SourcePackageName.GetNumber());
+			FDisplayNameEntryId NameEntry = ContainerHeader->RedirectsNameMap[Redirect.SourcePackageName.GetIndex()];
+			FName SourcePackageName = NameEntry.ToName(Redirect.SourcePackageName.GetNumber());
 			RedirectsPackageMap.Emplace(Redirect.SourcePackageId, MakeTuple(SourcePackageName, Redirect.TargetPackageId));
 		}
 	}
