Non-Unity fix

diff --git a/Engine/Source/Runtime/Core/Private/Misc/OutputDeviceArchiveWrapper.cpp b/Engine/Source/Runtime/Core/Private/Misc/OutputDeviceArchiveWrapper.cpp
--- a/Engine/Source/Runtime/Core/Private/Misc/OutputDeviceArchiveWrapper.cpp
+++ b/Engine/Source/Runtime/Core/Private/Misc/OutputDeviceArchiveWrapper.cpp
@@ -2,6 +2,7 @@
 
 #include "Misc/OutputDeviceArchiveWrapper.h"
 #include "Misc/OutputDeviceHelper.h"
+#include "Serialization/Archive.h"
 
 void FOutputDeviceArchiveWrapper::Flush()
 {
