Kill FAndroidPlatformString::Strncmp, which reads one character out of bounds for unterminated strings, and fallback to my faster FGenericPlatform::Strncmp algorithm

diff --git a/Engine/Source/Runtime/Core/Public/Android/AndroidPlatformString.h b/Engine/Source/Runtime/Core/Public/Android/AndroidPlatformString.h
--- a/Engine/Source/Runtime/Core/Public/Android/AndroidPlatformString.h
+++ b/Engine/Source/Runtime/Core/Public/Android/AndroidPlatformString.h
@@ -124,18 +124,6 @@ struct FAndroidPlatformString : public FGenericPlatformString
 		return 0;
 	}
 
-	static FORCEINLINE int32 Strncmp( const WIDECHAR* String1, const WIDECHAR* String2, SIZE_T Count )
-	{
-		for (; (*String1 || *String2) && Count > 0; String1++, String2++, Count--)
-		{
-			if (*String1 != *String2)
-			{
-				return *String1 - *String2;
-			}
-		}
-		return 0;
-	}
-
 	static FORCEINLINE int32 Strlen( const WIDECHAR* String )
 	{
 		if (!String)
