Replaced uses of _U8SV with UTF8TEXTVIEW

#rnx

diff --git a/Engine/Source/Developer/DerivedDataCache/Private/DerivedDataBuildOutput.cpp b/Engine/Source/Developer/DerivedDataCache/Private/DerivedDataBuildOutput.cpp
--- a/Engine/Source/Developer/DerivedDataCache/Private/DerivedDataBuildOutput.cpp
+++ b/Engine/Source/Developer/DerivedDataCache/Private/DerivedDataBuildOutput.cpp
@@ -31,26 +31,26 @@ static FUtf8StringView LexToString(EBuildOutputMessageLevel Level)
 {
 	switch (Level)
 	{
-	case EBuildOutputMessageLevel::Error: return "Error"_U8SV;
-	case EBuildOutputMessageLevel::Warning: return "Warning"_U8SV;
-	case EBuildOutputMessageLevel::Display: return "Display"_U8SV;
-	default: return "Unknown"_U8SV;
+	case EBuildOutputMessageLevel::Error: return UTF8TEXTVIEW("Error");
+	case EBuildOutputMessageLevel::Warning: return UTF8TEXTVIEW("Warning");
+	case EBuildOutputMessageLevel::Display: return UTF8TEXTVIEW("Display");
+	default: return UTF8TEXTVIEW("Unknown");
 	}
 }
 
 static bool TryLexFromString(EBuildOutputMessageLevel& OutLevel, FUtf8StringView String)
 {
-	if (String.Equals("Error"_U8SV, ESearchCase::CaseSensitive))
+	if (String.Equals(UTF8TEXTVIEW("Error"), ESearchCase::CaseSensitive))
 	{
 		OutLevel = EBuildOutputMessageLevel::Error;
 		return true;
 	}
-	if (String.Equals("Warning"_U8SV, ESearchCase::CaseSensitive))
+	if (String.Equals(UTF8TEXTVIEW("Warning"), ESearchCase::CaseSensitive))
 	{
 		OutLevel = EBuildOutputMessageLevel::Warning;
 		return true;
 	}
-	if (String.Equals("Display"_U8SV, ESearchCase::CaseSensitive))
+	if (String.Equals(UTF8TEXTVIEW("Display"), ESearchCase::CaseSensitive))
 	{
 		OutLevel = EBuildOutputMessageLevel::Display;
 		return true;
@@ -62,20 +62,20 @@ static FUtf8StringView LexToString(EBuildOutputLogLevel Level)
 {
 	switch (Level)
 	{
-	case EBuildOutputLogLevel::Error: return "Error"_U8SV;
-	case EBuildOutputLogLevel::Warning: return "Warning"_U8SV;
-	default: return "Unknown"_U8SV;
+	case EBuildOutputLogLevel::Error: return UTF8TEXTVIEW("Error");
+	case EBuildOutputLogLevel::Warning: return UTF8TEXTVIEW("Warning");
+	default: return UTF8TEXTVIEW("Unknown");
 	}
 }
 
 static bool TryLexFromString(EBuildOutputLogLevel& OutLevel, FUtf8StringView String)
 {
-	if (String.Equals("Error"_U8SV, ESearchCase::CaseSensitive))
+	if (String.Equals(UTF8TEXTVIEW("Error"), ESearchCase::CaseSensitive))
 	{
 		OutLevel = EBuildOutputLogLevel::Error;
 		return true;
 	}
-	if (String.Equals("Warning"_U8SV, ESearchCase::CaseSensitive))
+	if (String.Equals(UTF8TEXTVIEW("Warning"), ESearchCase::CaseSensitive))
 	{
 		OutLevel = EBuildOutputLogLevel::Warning;
 		return true;
@@ -212,7 +212,7 @@ FBuildOutputInternal::FBuildOutputInternal(
 	checkf(!Name.IsEmpty(), TEXT("A build output requires a non-empty name."));
 	AssertValidBuildFunctionName(Function, Name);
 	Output.MakeOwned();
-	FCbField MetaField = InOutput["Meta"_ASV];
+	FCbField MetaField = InOutput[ANSITEXTVIEW("Meta")];
 	bOutIsValid = TryLoad() && (!MetaField || MetaField.IsObject());
 	Meta = MoveTemp(MetaField).AsObject();
 }
@@ -252,11 +252,11 @@ bool FBuildOutputInternal::TryLoad()
 
 	if (Values.IsEmpty())
 	{
-		for (FCbFieldView Value : OutputView["Values"_ASV])
+		for (FCbFieldView Value : OutputView[ANSITEXTVIEW("Values")])
 		{
-			const FValueId Id = Value["Id"_ASV].AsObjectId();
-			const FIoHash& RawHash = Value["RawHash"_ASV].AsAttachment();
-			const uint64 RawSize = Value["RawSize"_ASV].AsUInt64(MAX_uint64);
+			const FValueId Id = Value[ANSITEXTVIEW("Id")].AsObjectId();
+			const FIoHash& RawHash = Value[ANSITEXTVIEW("RawHash")].AsAttachment();
+			const uint64 RawSize = Value[ANSITEXTVIEW("RawSize")].AsUInt64(MAX_uint64);
 			if (Id.IsNull() || RawHash.IsZero() || RawSize == MAX_uint64)
 			{
 				return false;
@@ -265,7 +265,7 @@ bool FBuildOutputInternal::TryLoad()
 		}
 	}
 
-	if (FCbFieldView MessagesField = OutputView["Messages"_ASV])
+	if (FCbFieldView MessagesField = OutputView[ANSITEXTVIEW("Messages")])
 	{
 		if (!MessagesField.IsArray())
 		{
@@ -274,8 +274,8 @@ bool FBuildOutputInternal::TryLoad()
 		Messages.Reserve(MessagesField.AsArrayView().Num());
 		for (FCbFieldView MessageField : MessagesField)
 		{
-			const FUtf8StringView LevelName = MessageField["Level"_ASV].AsString();
-			const FUtf8StringView Message = MessageField["Message"_ASV].AsString();
+			const FUtf8StringView LevelName = MessageField[ANSITEXTVIEW("Level")].AsString();
+			const FUtf8StringView Message = MessageField[ANSITEXTVIEW("Message")].AsString();
 			EBuildOutputMessageLevel Level;
 			if (LevelName.IsEmpty() || Message.IsEmpty() || !TryLexFromString(Level, LevelName))
 			{
@@ -285,7 +285,7 @@ bool FBuildOutputInternal::TryLoad()
 		}
 	}
 
-	if (FCbFieldView LogsField = OutputView["Logs"_ASV])
+	if (FCbFieldView LogsField = OutputView[ANSITEXTVIEW("Logs")])
 	{
 		if (!LogsField.IsArray())
 		{
@@ -294,9 +294,9 @@ bool FBuildOutputInternal::TryLoad()
 		Logs.Reserve(LogsField.AsArrayView().Num());
 		for (FCbFieldView LogField : LogsField)
 		{
-			const FUtf8StringView LevelName = LogField["Level"_ASV].AsString();
-			const FUtf8StringView Category = LogField["Category"_ASV].AsString();
-			const FUtf8StringView Message = LogField["Message"_ASV].AsString();
+			const FUtf8StringView LevelName = LogField[ANSITEXTVIEW("Level")].AsString();
+			const FUtf8StringView Category = LogField[ANSITEXTVIEW("Category")].AsString();
+			const FUtf8StringView Message = LogField[ANSITEXTVIEW("Message")].AsString();
 			EBuildOutputLogLevel Level;
 			if (LevelName.IsEmpty() || Category.IsEmpty() || Message.IsEmpty() || !TryLexFromString(Level, LevelName))
 			{
@@ -314,28 +314,28 @@ void FBuildOutputInternal::Save(FCbWriter& Writer) const
 	Writer.BeginObject();
 	if (!Values.IsEmpty())
 	{
-		Writer.BeginArray("Values"_ASV);
+		Writer.BeginArray(ANSITEXTVIEW("Values"));
 		for (const FValueWithId& Value : Values)
 		{
 			Writer.BeginObject();
-			Writer.AddObjectId("Id"_ASV, Value.GetId());
-			Writer.AddBinaryAttachment("RawHash"_ASV, Value.GetRawHash());
-			Writer.AddInteger("RawSize"_ASV, Value.GetRawSize());
+			Writer.AddObjectId(ANSITEXTVIEW("Id"), Value.GetId());
+			Writer.AddBinaryAttachment(ANSITEXTVIEW("RawHash"), Value.GetRawHash());
+			Writer.AddInteger(ANSITEXTVIEW("RawSize"), Value.GetRawSize());
 			Writer.EndObject();
 		}
 		Writer.EndArray();
 	}
-	if (FCbField MessagesField = Output["Messages"_ASV])
+	if (FCbField MessagesField = Output[ANSITEXTVIEW("Messages")])
 	{
-		Writer.AddField("Messages"_ASV, MessagesField);
+		Writer.AddField(ANSITEXTVIEW("Messages"), MessagesField);
 	}
-	if (FCbField LogsField = Output["Logs"_ASV])
+	if (FCbField LogsField = Output[ANSITEXTVIEW("Logs")])
 	{
-		Writer.AddField("Logs"_ASV, LogsField);
+		Writer.AddField(ANSITEXTVIEW("Logs"), LogsField);
 	}
 	if (Meta)
 	{
-		Writer.AddObject("Meta"_ASV, Meta);
+		Writer.AddObject(ANSITEXTVIEW("Meta"), Meta);
 	}
 	Writer.EndObject();
 }
@@ -346,13 +346,13 @@ void FBuildOutputInternal::Save(FCacheRecordBuilder& RecordBuilder) const
 	if (!Messages.IsEmpty() || !Logs.IsEmpty())
 	{
 		TCbWriter<1024> Writer;
-		if (FCbField MessagesField = Output["Messages"_ASV])
+		if (FCbField MessagesField = Output[ANSITEXTVIEW("Messages")])
 		{
-			Writer.AddField("Messages"_ASV, MessagesField);
+			Writer.AddField(ANSITEXTVIEW("Messages"), MessagesField);
 		}
-		if (FCbField LogsField = Output["Logs"_ASV])
+		if (FCbField LogsField = Output[ANSITEXTVIEW("Logs")])
 		{
-			Writer.AddField("Logs"_ASV, LogsField);
+			Writer.AddField(ANSITEXTVIEW("Logs"), LogsField);
 		}
 		RecordBuilder.AddValue(GBuildOutputValueId, Writer.Save().GetBuffer());
 	}
@@ -383,8 +383,8 @@ void FBuildOutputBuilderInternal::AddMessage(const FBuildOutputMessage& Message)
 	bHasError |= Message.Level == EBuildOutputMessageLevel::Error;
 	bHasMessages = true;
 	MessageWriter.BeginObject();
-	MessageWriter.AddString("Level"_ASV, LexToString(Message.Level));
-	MessageWriter.AddString("Message"_ASV, Message.Message);
+	MessageWriter.AddString(ANSITEXTVIEW("Level"), LexToString(Message.Level));
+	MessageWriter.AddString(ANSITEXTVIEW("Message"), Message.Message);
 	MessageWriter.EndObject();
 }
 
@@ -393,9 +393,9 @@ void FBuildOutputBuilderInternal::AddLog(const FBuildOutputLog& Log)
 	bHasError |= Log.Level == EBuildOutputLogLevel::Error;
 	bHasLogs = true;
 	LogWriter.BeginObject();
-	LogWriter.AddString("Level"_ASV, LexToString(Log.Level));
-	LogWriter.AddString("Category"_ASV, Log.Category);
-	LogWriter.AddString("Message"_ASV, Log.Message);
+	LogWriter.AddString(ANSITEXTVIEW("Level"), LexToString(Log.Level));
+	LogWriter.AddString(ANSITEXTVIEW("Category"), Log.Category);
+	LogWriter.AddString(ANSITEXTVIEW("Message"), Log.Message);
 	LogWriter.EndObject();
 }
 
@@ -414,12 +414,12 @@ FBuildOutput FBuildOutputBuilderInternal::Build()
 		Writer.BeginObject();
 		if (bHasMessages)
 		{
-			Writer.AddArray("Messages"_ASV, MessageWriter.Save().AsArray());
+			Writer.AddArray(ANSITEXTVIEW("Messages"), MessageWriter.Save().AsArray());
 			MessageWriter.Reset();
 		}
 		if (bHasLogs)
 		{
-			Writer.AddArray("Logs"_ASV, LogWriter.Save().AsArray());
+			Writer.AddArray(ANSITEXTVIEW("Logs"), LogWriter.Save().AsArray());
 			LogWriter.Reset();
 		}
 		Writer.EndObject();
diff --git a/Engine/Source/Runtime/Core/Private/Tests/Serialization/CompactBinaryTest.cpp b/Engine/Source/Runtime/Core/Private/Tests/Serialization/CompactBinaryTest.cpp
--- a/Engine/Source/Runtime/Core/Private/Tests/Serialization/CompactBinaryTest.cpp
+++ b/Engine/Source/Runtime/Core/Private/Tests/Serialization/CompactBinaryTest.cpp
@@ -208,7 +208,7 @@ bool FCbFieldNoneTest::RunTest(const FString& Parameters)
 		constexpr const ANSICHAR NoneBytes[] = { ANSICHAR(FieldType), 4, 'N', 'a', 'm', 'e' };
 		FCbFieldView NoneField(NoneBytes);
 		TestEqual(TEXT("FCbFieldView(None|Type|Name)::GetSize()"), NoneField.GetSize(), uint64(sizeof(NoneBytes)));
-		TestEqual(TEXT("FCbFieldView(None|Type|Name)::GetName()"), NoneField.GetName(), "Name"_U8SV);
+		TestEqual(TEXT("FCbFieldView(None|Type|Name)::GetName()"), NoneField.GetName(), UTF8TEXTVIEW("Name"));
 		TestTrue(TEXT("FCbFieldView(None|Type|Name)::HasName()"), NoneField.HasName());
 		TestFalse(TEXT("!FCbFieldView(None|Type|Name)::HasValue()"), NoneField.HasValue());
 		TestEqual(TEXT("FCbFieldView(None|Type|Name)::GetHash()"), NoneField.GetHash(), FIoHash::HashBuffer(MakeMemoryView(NoneBytes)));
@@ -240,7 +240,7 @@ bool FCbFieldNoneTest::RunTest(const FString& Parameters)
 		constexpr const ANSICHAR NoneBytes[] = { ANSICHAR(FieldType), 4, 'N', 'a', 'm', 'e' };
 		FCbFieldView NoneField(NoneBytes + 1, FieldType);
 		TestEqual(TEXT("FCbFieldView(None|Name)::GetSize()"), NoneField.GetSize(), uint64(sizeof(NoneBytes)));
-		TestEqual(TEXT("FCbFieldView(None|Name)::GetName()"), NoneField.GetName(), "Name"_U8SV);
+		TestEqual(TEXT("FCbFieldView(None|Name)::GetName()"), NoneField.GetName(), UTF8TEXTVIEW("Name"));
 		TestTrue(TEXT("FCbFieldView(None|Name)::HasName()"), NoneField.HasName());
 		TestFalse(TEXT("!FCbFieldView(None|Name)::HasValue()"), NoneField.HasValue());
 		TestEqual(TEXT("FCbFieldView(None|Name)::GetHash()"), NoneField.GetHash(), FIoHash::HashBuffer(MakeMemoryView(NoneBytes)));
@@ -258,7 +258,7 @@ bool FCbFieldNoneTest::RunTest(const FString& Parameters)
 		constexpr const uint8 NoneBytes[] = { uint8(FieldType), 0 };
 		FCbFieldView NoneField(NoneBytes + 1, FieldType);
 		TestEqual(TEXT("FCbFieldView(None|EmptyName)::GetSize()"), NoneField.GetSize(), uint64(sizeof(NoneBytes)));
-		TestEqual(TEXT("FCbFieldView(None|EmptyName)::GetName()"), NoneField.GetName(), ""_U8SV);
+		TestEqual(TEXT("FCbFieldView(None|EmptyName)::GetName()"), NoneField.GetName(), UTF8TEXTVIEW(""));
 		TestTrue(TEXT("FCbFieldView(None|EmptyName)::HasName()"), NoneField.HasName());
 		TestFalse(TEXT("!FCbFieldView(None|EmptyName)::HasValue()"), NoneField.HasValue());
 		TestEqual(TEXT("FCbFieldView(None|EmptyName)::GetHash()"), NoneField.GetHash(), FIoHash::HashBuffer(MakeMemoryView(NoneBytes)));
@@ -341,9 +341,9 @@ bool FCbFieldObjectTest::RunTest(const FString& Parameters)
 		TestIntObject(Object, 0, 1);
 
 		// Find fields that do not exist.
-		TestFalse(TEXT("FCbObjectView()::Find(Missing)"), Object.FindView("Field"_U8SV).HasValue());
-		TestFalse(TEXT("FCbObjectView()::FindViewIgnoreCase(Missing)"), Object.FindViewIgnoreCase("Field"_U8SV).HasValue());
-		TestFalse(TEXT("FCbObjectView()::operator[](Missing)"), Object["Field"_U8SV].HasValue());
+		TestFalse(TEXT("FCbObjectView()::Find(Missing)"), Object.FindView(UTF8TEXTVIEW("Field")).HasValue());
+		TestFalse(TEXT("FCbObjectView()::FindViewIgnoreCase(Missing)"), Object.FindViewIgnoreCase(UTF8TEXTVIEW("Field")).HasValue());
+		TestFalse(TEXT("FCbObjectView()::operator[](Missing)"), Object[UTF8TEXTVIEW("Field")].HasValue());
 
 		// Advance an iterator past the last field.
 		FCbFieldViewIterator It = Object.CreateViewIterator();
@@ -368,12 +368,12 @@ bool FCbFieldObjectTest::RunTest(const FString& Parameters)
 		TestIntObject(Object, 3, sizeof(Value));
 		TestIntObject(Field.AsObjectView(), 3, sizeof(Value));
 		TestTrue(TEXT("FCbObjectView::Equals()"), Object.Equals(Field.AsObjectView()));
-		TestEqual(TEXT("FCbObjectView::Find()"), Object.FindView("B"_ASV).AsInt32(), 2);
-		TestEqual(TEXT("FCbObjectView::Find()"), Object.FindView("b"_ASV).AsInt32(4), 4);
-		TestEqual(TEXT("FCbObjectView::FindViewIgnoreCase()"), Object.FindViewIgnoreCase("B"_ASV).AsInt32(), 2);
-		TestEqual(TEXT("FCbObjectView::FindViewIgnoreCase()"), Object.FindViewIgnoreCase("b"_ASV).AsInt32(), 2);
-		TestEqual(TEXT("FCbObjectView::operator[]"), Object["B"_ASV].AsInt32(), 2);
-		TestEqual(TEXT("FCbObjectView::operator[]"), Object["b"_ASV].AsInt32(4), 4);
+		TestEqual(TEXT("FCbObjectView::Find()"), Object.FindView(ANSITEXTVIEW("B")).AsInt32(), 2);
+		TestEqual(TEXT("FCbObjectView::Find()"), Object.FindView(ANSITEXTVIEW("b")).AsInt32(4), 4);
+		TestEqual(TEXT("FCbObjectView::FindViewIgnoreCase()"), Object.FindViewIgnoreCase(ANSITEXTVIEW("B")).AsInt32(), 2);
+		TestEqual(TEXT("FCbObjectView::FindViewIgnoreCase()"), Object.FindViewIgnoreCase(ANSITEXTVIEW("b")).AsInt32(), 2);
+		TestEqual(TEXT("FCbObjectView::operator[]"), Object[ANSITEXTVIEW("B")].AsInt32(), 2);
+		TestEqual(TEXT("FCbObjectView::operator[]"), Object[ANSITEXTVIEW("b")].AsInt32(4), 4);
 	}
 
 	// Test FCbFieldView(UniformObject, NotEmpty)
@@ -386,16 +386,16 @@ bool FCbFieldObjectTest::RunTest(const FString& Parameters)
 		TestIntObject(Object, 3, sizeof(Value));
 		TestIntObject(Field.AsObjectView(), 3, sizeof(Value));
 		TestTrue(TEXT("FCbObjectView(Uniform)::Equals()"), Object.Equals(Field.AsObjectView()));
-		TestEqual(TEXT("FCbObjectView(Uniform)::Find()"), Object.FindView("B"_ASV).AsInt32(), 2);
-		TestEqual(TEXT("FCbObjectView(Uniform)::Find()"), Object.Find("B"_ASV).AsInt32(), 2);
-		TestEqual(TEXT("FCbObjectView(Uniform)::Find()"), Object.FindView("b"_ASV).AsInt32(4), 4);
-		TestEqual(TEXT("FCbObjectView(Uniform)::Find()"), Object.Find("b"_ASV).AsInt32(4), 4);
-		TestEqual(TEXT("FCbObjectView(Uniform)::FindViewIgnoreCase()"), Object.FindViewIgnoreCase("B"_ASV).AsInt32(), 2);
-		TestEqual(TEXT("FCbObjectView(Uniform)::FindViewIgnoreCase()"), Object.FindIgnoreCase("B"_ASV).AsInt32(), 2);
-		TestEqual(TEXT("FCbObjectView(Uniform)::FindViewIgnoreCase()"), Object.FindViewIgnoreCase("b"_ASV).AsInt32(), 2);
-		TestEqual(TEXT("FCbObjectView(Uniform)::FindViewIgnoreCase()"), Object.FindIgnoreCase("b"_ASV).AsInt32(), 2);
-		TestEqual(TEXT("FCbObjectView(Uniform)::operator[]"), Object["B"_ASV].AsInt32(), 2);
-		TestEqual(TEXT("FCbObjectView(Uniform)::operator[]"), Object["b"_ASV].AsInt32(4), 4);
+		TestEqual(TEXT("FCbObjectView(Uniform)::Find()"), Object.FindView(ANSITEXTVIEW("B")).AsInt32(), 2);
+		TestEqual(TEXT("FCbObjectView(Uniform)::Find()"), Object.Find(ANSITEXTVIEW("B")).AsInt32(), 2);
+		TestEqual(TEXT("FCbObjectView(Uniform)::Find()"), Object.FindView(ANSITEXTVIEW("b")).AsInt32(4), 4);
+		TestEqual(TEXT("FCbObjectView(Uniform)::Find()"), Object.Find(ANSITEXTVIEW("b")).AsInt32(4), 4);
+		TestEqual(TEXT("FCbObjectView(Uniform)::FindViewIgnoreCase()"), Object.FindViewIgnoreCase(ANSITEXTVIEW("B")).AsInt32(), 2);
+		TestEqual(TEXT("FCbObjectView(Uniform)::FindViewIgnoreCase()"), Object.FindIgnoreCase(ANSITEXTVIEW("B")).AsInt32(), 2);
+		TestEqual(TEXT("FCbObjectView(Uniform)::FindViewIgnoreCase()"), Object.FindViewIgnoreCase(ANSITEXTVIEW("b")).AsInt32(), 2);
+		TestEqual(TEXT("FCbObjectView(Uniform)::FindViewIgnoreCase()"), Object.FindIgnoreCase(ANSITEXTVIEW("b")).AsInt32(), 2);
+		TestEqual(TEXT("FCbObjectView(Uniform)::operator[]"), Object[ANSITEXTVIEW("B")].AsInt32(), 2);
+		TestEqual(TEXT("FCbObjectView(Uniform)::operator[]"), Object[ANSITEXTVIEW("b")].AsInt32(4), 4);
 
 		TestTrue(TEXT("FCbObject::AsField()"), Object.GetOuterBuffer() == Object.AsField().AsObject().GetOuterBuffer());
 
@@ -445,7 +445,7 @@ bool FCbFieldObjectTest::RunTest(const FString& Parameters)
 		for (FCbFieldIterator It = ObjectClone.CreateIterator(); It; ++It)
 		{
 			FCbField Field = *It;
-			TestEqual(TEXT("FCbObject::CreateIterator().GetName()"), Field.GetName(), "F"_U8SV);
+			TestEqual(TEXT("FCbObject::CreateIterator().GetName()"), Field.GetName(), UTF8TEXTVIEW("F"));
 			TestEqual(TEXT("FCbObject::CreateIterator().AsInt32()"), Field.AsInt32(), 8);
 			TestTrue(TEXT("FCbObject::CreateIterator().IsOwned()"), Field.IsOwned());
 		}
@@ -469,8 +469,8 @@ bool FCbFieldObjectTest::RunTest(const FString& Parameters)
 			FCbField(FCbFieldView(Buffer + 1, ECbFieldType(ObjectType)), FSharedBuffer::MakeView(MakeMemoryView(Buffer))).GetBuffer().ToShared().GetView().EqualBytes(MakeMemoryView(Buffer)));
 
 		// Access Missing Field
-		TestFalse(TEXT("FCbObject()[Missing]"), ObjectClone["M"_ASV].HasValue());
-		TestFalse(TEXT("FCbField(Object)[Missing]"), ObjectClone.AsField()["M"_ASV].HasValue());
+		TestFalse(TEXT("FCbObject()[Missing]"), ObjectClone[ANSITEXTVIEW("M")].HasValue());
+		TestFalse(TEXT("FCbField(Object)[Missing]"), ObjectClone.AsField()[ANSITEXTVIEW("M")].HasValue());
 	}
 
 	// Test FCbObjectView as FCbFieldViewIterator
@@ -755,13 +755,13 @@ bool FCbFieldStringTest::RunTest(const FString& Parameters)
 	{
 		uint8 Value[9];
 		WriteVarUInt(uint64(1) << 31, Value);
-		TestFieldError<ECbFieldType::String>(TEXT("String, OutOfRangeSize"), Value, ECbFieldError::RangeError, "ABC"_U8SV);
+		TestFieldError<ECbFieldType::String>(TEXT("String, OutOfRangeSize"), Value, ECbFieldError::RangeError, UTF8TEXTVIEW("ABC"));
 	}
 
 	// Test FCbFieldView(None) as String
 	{
 		FCbFieldView Field;
-		TestFieldError<ECbFieldType::String>(TEXT("String, None"), Field, ECbFieldError::TypeError, "ABC"_U8SV);
+		TestFieldError<ECbFieldType::String>(TEXT("String, None"), Field, ECbFieldError::TypeError, UTF8TEXTVIEW("ABC"));
 	}
 
 	return true;
@@ -1273,26 +1273,26 @@ bool FCbFieldCustomByNameTest::RunTest(const FString& Parameters)
 	// Test FCbFieldView(CustomByName, ABC, Empty)
 	{
 		const uint8 Value[] = {4, 3, 'A', 'B', 'C'};
-		TestField<ECbFieldType::CustomByName>(TEXT("CustomByName, MinId, Empty"), Value, FCbCustomByName{"ABC"_U8SV});
-		TestField<ECbFieldType::CustomByName>(TEXT("CustomByName, MinId, Empty, View"), Value, FMemoryView(), MakeMemoryView<uint8>({1, 2, 3}), ECbFieldError::None, FCustomByNameAccessor("ABC"_U8SV));
-		TestFieldError<ECbFieldType::CustomByName>(TEXT("CustomByName, MinId, Empty, InvalidCase"), Value, ECbFieldError::RangeError, MakeMemoryView<uint8>({1, 2, 3}), FCustomByNameAccessor("abc"_U8SV));
+		TestField<ECbFieldType::CustomByName>(TEXT("CustomByName, MinId, Empty"), Value, FCbCustomByName{UTF8TEXTVIEW("ABC")});
+		TestField<ECbFieldType::CustomByName>(TEXT("CustomByName, MinId, Empty, View"), Value, FMemoryView(), MakeMemoryView<uint8>({1, 2, 3}), ECbFieldError::None, FCustomByNameAccessor(UTF8TEXTVIEW("ABC")));
+		TestFieldError<ECbFieldType::CustomByName>(TEXT("CustomByName, MinId, Empty, InvalidCase"), Value, ECbFieldError::RangeError, MakeMemoryView<uint8>({1, 2, 3}), FCustomByNameAccessor(UTF8TEXTVIEW("abc")));
 	}
 
 	// Test FCbFieldView(CustomByName, ABC, Value)
 	{
 		const uint8 Value[] = {8, 3, 'A', 'B', 'C', 1, 2, 3, 4};
-		TestFieldNoClone<ECbFieldType::CustomByName>(TEXT("CustomByName, MinId, Value"), Value, FCbCustomByName{"ABC"_U8SV, MakeMemoryView(Value).Right(4)});
-		TestFieldNoClone<ECbFieldType::CustomByName>(TEXT("CustomByName, MinId, Value, View"), Value, MakeMemoryView(Value).Right(4), FMemoryView(), ECbFieldError::None, FCustomByNameAccessor("ABC"_U8SV));
-		TestFieldError<ECbFieldType::CustomByName>(TEXT("CustomByName, MinId, Value, InvalidCase"), Value, ECbFieldError::RangeError, MakeMemoryView<uint8>({1, 2, 3}), FCustomByNameAccessor("abc"_U8SV));
+		TestFieldNoClone<ECbFieldType::CustomByName>(TEXT("CustomByName, MinId, Value"), Value, FCbCustomByName{UTF8TEXTVIEW("ABC"), MakeMemoryView(Value).Right(4)});
+		TestFieldNoClone<ECbFieldType::CustomByName>(TEXT("CustomByName, MinId, Value, View"), Value, MakeMemoryView(Value).Right(4), FMemoryView(), ECbFieldError::None, FCustomByNameAccessor(UTF8TEXTVIEW("ABC")));
+		TestFieldError<ECbFieldType::CustomByName>(TEXT("CustomByName, MinId, Value, InvalidCase"), Value, ECbFieldError::RangeError, MakeMemoryView<uint8>({1, 2, 3}), FCustomByNameAccessor(UTF8TEXTVIEW("abc")));
 	}
 
 	// Test FCbFieldView(None) as CustomByName
 	{
 		FCbFieldView DefaultField;
-		TestFieldError<ECbFieldType::CustomByName>(TEXT("CustomByName, None"), DefaultField, ECbFieldError::TypeError, FCbCustomByName{"ABC"_U8SV, MakeMemoryView<uint8>({1, 2, 3})});
-		TestFieldError<ECbFieldType::CustomByName>(TEXT("CustomByName, None, View"), DefaultField, ECbFieldError::TypeError, MakeMemoryView<uint8>({1, 2, 3}), FCustomByNameAccessor("ABC"_U8SV));
+		TestFieldError<ECbFieldType::CustomByName>(TEXT("CustomByName, None"), DefaultField, ECbFieldError::TypeError, FCbCustomByName{UTF8TEXTVIEW("ABC"), MakeMemoryView<uint8>({1, 2, 3})});
+		TestFieldError<ECbFieldType::CustomByName>(TEXT("CustomByName, None, View"), DefaultField, ECbFieldError::TypeError, MakeMemoryView<uint8>({1, 2, 3}), FCustomByNameAccessor(UTF8TEXTVIEW("ABC")));
 		const uint8 DefaultValue[] = {1, 2, 3};
-		TestEqual(TEXT("FCbFieldView()::AsCustom(Name)"), DefaultField.AsCustom("ABC"_U8SV, MakeMemoryView(DefaultValue)), MakeMemoryView(DefaultValue));
+		TestEqual(TEXT("FCbFieldView()::AsCustom(Name)"), DefaultField.AsCustom(UTF8TEXTVIEW("ABC"), MakeMemoryView(DefaultValue)), MakeMemoryView(DefaultValue));
 	}
 
 	return true;
@@ -1826,22 +1826,22 @@ bool FCbFieldParseTest::RunTest(const FString& Parameters)
 		for (FCbFieldViewIterator It = Object.CreateViewIterator(); It;)
 		{
 			const FCbFieldViewIterator Last = It;
-			if (It.GetName().Equals("A"_U8SV))
+			if (It.GetName().Equals(UTF8TEXTVIEW("A")))
 			{
 				A = It.AsUInt32();
 				++It;
 			}
-			if (It.GetName().Equals("B"_U8SV))
+			if (It.GetName().Equals(UTF8TEXTVIEW("B")))
 			{
 				B = It.AsUInt32();
 				++It;
 			}
-			if (It.GetName().Equals("C"_U8SV))
+			if (It.GetName().Equals(UTF8TEXTVIEW("C")))
 			{
 				C = It.AsUInt32();
 				++It;
 			}
-			if (It.GetName().Equals("D"_U8SV))
+			if (It.GetName().Equals(UTF8TEXTVIEW("D")))
 			{
 				D = It.AsUInt32();
 				++It;
diff --git a/Engine/Source/Runtime/Core/Private/Tests/Serialization/CompactBinaryWriterTest.cpp b/Engine/Source/Runtime/Core/Private/Tests/Serialization/CompactBinaryWriterTest.cpp
--- a/Engine/Source/Runtime/Core/Private/Tests/Serialization/CompactBinaryWriterTest.cpp
+++ b/Engine/Source/Runtime/Core/Private/Tests/Serialization/CompactBinaryWriterTest.cpp
@@ -39,7 +39,7 @@ bool FCbWriterObjectTest::RunTest(const FString& Parameters)
 	// Test Named Empty Object
 	{
 		Writer.Reset();
-		Writer.SetName("Object"_ASV);
+		Writer.SetName(ANSITEXTVIEW("Object"));
 		Writer.BeginObject();
 		Writer.EndObject();
 		FCbField Field = Writer.Save();
@@ -54,16 +54,16 @@ bool FCbWriterObjectTest::RunTest(const FString& Parameters)
 	{
 		Writer.Reset();
 		Writer.BeginObject();
-		Writer.SetName("Integer"_ASV).AddInteger(0);
-		Writer.SetName("Float"_ASV).AddFloat(0.0f);
+		Writer.SetName(ANSITEXTVIEW("Integer")).AddInteger(0);
+		Writer.SetName(ANSITEXTVIEW("Float")).AddFloat(0.0f);
 		Writer.EndObject();
 		FCbField Field = Writer.Save();
 		if (TestEqual(TEXT("FCbWriter(Object, Basic) Validate"), ValidateCompactBinary(Field.GetOuterBuffer(), ECbValidateMode::All), ECbValidateError::None))
 		{
 			TestTrue(TEXT("FCbWriter(Object, Basic).IsObject()"), Field.IsObject());
 			FCbObjectView Object = Field.AsObjectView();
-			TestTrue(TEXT("FCbWriter(Object, Basic).AsObjectView()[Integer]"), Object["Integer"_ASV].IsInteger());
-			TestTrue(TEXT("FCbWriter(Object, Basic).AsObjectView()[Float]"), Object["Float"_ASV].IsFloat());
+			TestTrue(TEXT("FCbWriter(Object, Basic).AsObjectView()[Integer]"), Object[ANSITEXTVIEW("Integer")].IsInteger());
+			TestTrue(TEXT("FCbWriter(Object, Basic).AsObjectView()[Float]"), Object[ANSITEXTVIEW("Float")].IsFloat());
 		}
 	}
 
@@ -71,16 +71,16 @@ bool FCbWriterObjectTest::RunTest(const FString& Parameters)
 	{
 		Writer.Reset();
 		Writer.BeginObject();
-		Writer.SetName("Field1"_ASV).AddInteger(0);
-		Writer.SetName("Field2"_ASV).AddInteger(1);
+		Writer.SetName(ANSITEXTVIEW("Field1")).AddInteger(0);
+		Writer.SetName(ANSITEXTVIEW("Field2")).AddInteger(1);
 		Writer.EndObject();
 		FCbField Field = Writer.Save();
 		if (TestEqual(TEXT("FCbWriter(Object, Uniform) Validate"), ValidateCompactBinary(Field.GetOuterBuffer(), ECbValidateMode::All), ECbValidateError::None))
 		{
 			TestTrue(TEXT("FCbWriter(Object, Uniform).IsObject()"), Field.IsObject());
 			FCbObjectView Object = Field.AsObjectView();
-			TestTrue(TEXT("FCbWriter(Object, Uniform).AsObjectView()[Field1]"), Object["Field1"_ASV].IsInteger());
-			TestTrue(TEXT("FCbWriter(Object, Uniform).AsObjectView()[Field2]"), Object["Field2"_ASV].IsInteger());
+			TestTrue(TEXT("FCbWriter(Object, Uniform).AsObjectView()[Field1]"), Object[ANSITEXTVIEW("Field1")].IsInteger());
+			TestTrue(TEXT("FCbWriter(Object, Uniform).AsObjectView()[Field2]"), Object[ANSITEXTVIEW("Field2")].IsInteger());
 		}
 	}
 
@@ -110,7 +110,7 @@ bool FCbWriterArrayTest::RunTest(const FString& Parameters)
 	// Test Named Empty Array
 	{
 		Writer.Reset();
-		Writer.SetName("Array"_ASV);
+		Writer.SetName(ANSITEXTVIEW("Array"));
 		Writer.BeginArray();
 		Writer.EndArray();
 		FCbField Field = Writer.Save();
@@ -186,12 +186,12 @@ bool FCbWriterNullTest::RunTest(const FString& Parameters)
 	// Test Null with Name
 	{
 		Writer.Reset();
-		Writer.SetName("Null"_ASV);
+		Writer.SetName(ANSITEXTVIEW("Null"));
 		Writer.AddNull();
 		FCbField Field = Writer.Save();
 		if (TestEqual(TEXT("FCbWriter(Null, Name) Validate"), ValidateCompactBinary(Field.GetOuterBuffer(), ECbValidateMode::All), ECbValidateError::None))
 		{
-			TestEqual(TEXT("FCbWriter(Null, Name).GetName()"), Field.GetName(), "Null"_U8SV);
+			TestEqual(TEXT("FCbWriter(Null, Name).GetName()"), Field.GetName(), UTF8TEXTVIEW("Null"));
 			TestTrue(TEXT("FCbWriter(Null, Name).HasName()"), Field.HasName());
 			TestTrue(TEXT("FCbWriter(Null, Name).IsNull()"), Field.IsNull());
 		}
@@ -208,9 +208,9 @@ bool FCbWriterNullTest::RunTest(const FString& Parameters)
 		Writer.EndArray();
 
 		Writer.BeginObject();
-		Writer.SetName("N1"_ASV).AddNull();
-		Writer.SetName("N2"_ASV).AddNull();
-		Writer.SetName("N3"_ASV).AddNull();
+		Writer.SetName(ANSITEXTVIEW("N1")).AddNull();
+		Writer.SetName(ANSITEXTVIEW("N2")).AddNull();
+		Writer.SetName(ANSITEXTVIEW("N3")).AddNull();
 		Writer.EndObject();
 
 		FCbFieldIterator Fields = Writer.Save();
@@ -278,12 +278,12 @@ bool FCbWriterBinaryTest::RunTest(const FString& Parameters)
 	{
 		Writer.Reset();
 		const uint8 BinaryValue[] = { 1, 2, 3, 4, 5, 6 };
-		Writer.SetName("Binary"_ASV);
+		Writer.SetName(ANSITEXTVIEW("Binary"));
 		Writer.AddBinary(BinaryValue, sizeof(BinaryValue));
 		FCbField Field = Writer.Save();
 		if (TestEqual(TEXT("FCbWriter(Binary, Array) Validate"), ValidateCompactBinary(Field.GetOuterBuffer(), ECbValidateMode::All), ECbValidateError::None))
 		{
-			TestEqual(TEXT("FCbWriter(Binary, Array).GetName()"), Field.GetName(), "Binary"_U8SV);
+			TestEqual(TEXT("FCbWriter(Binary, Array).GetName()"), Field.GetName(), UTF8TEXTVIEW("Binary"));
 			TestTrue(TEXT("FCbWriter(Binary, Array).HasName()"), Field.HasName());
 			TestTrue(TEXT("FCbWriter(Binary, Array).IsBinary()"), Field.IsBinary());
 			TestTrue(TEXT("FCbWriter(Binary, Array).AsBinaryView()"), Field.AsBinaryView().EqualBytes(MakeMemoryView(BinaryValue)));
@@ -320,17 +320,17 @@ bool FCbWriterStringTest::RunTest(const FString& Parameters)
 	// Test Basic Strings
 	{
 		Writer.Reset();
-		Writer.SetName("String"_ASV).AddString("Value"_ASV);
-		Writer.SetName("String"_ASV).AddString(TEXTVIEW("Value"));
+		Writer.SetName(ANSITEXTVIEW("String")).AddString(ANSITEXTVIEW("Value"));
+		Writer.SetName(ANSITEXTVIEW("String")).AddString(TEXTVIEW("Value"));
 		FCbFieldIterator Fields = Writer.Save();
 		if (TestEqual(TEXT("FCbWriter(String, Basic) Validate"), ValidateCompactBinaryRange(Fields.GetOuterBuffer(), ECbValidateMode::All), ECbValidateError::None))
 		{
 			for (FCbFieldView Field : Fields)
 			{
-				TestEqual(TEXT("FCbWriter(String, Basic).GetName()"), Field.GetName(), "String"_U8SV);
+				TestEqual(TEXT("FCbWriter(String, Basic).GetName()"), Field.GetName(), UTF8TEXTVIEW("String"));
 				TestTrue(TEXT("FCbWriter(String, Basic).HasName()"), Field.HasName());
 				TestTrue(TEXT("FCbWriter(String, Basic).IsString()"), Field.IsString());
-				TestEqual(TEXT("FCbWriter(String, Basic).AsString()"), Field.AsString(), "Value"_U8SV);
+				TestEqual(TEXT("FCbWriter(String, Basic).AsString()"), Field.AsString(), UTF8TEXTVIEW("Value"));
 			}
 		}
 	}
@@ -360,14 +360,14 @@ bool FCbWriterStringTest::RunTest(const FString& Parameters)
 	{
 		Writer.Reset();
 		WIDECHAR Value[2] = { 0xd83d, 0xde00 };
-		Writer.AddString("\xf0\x9f\x98\x80"_ASV);
+		Writer.AddString(ANSITEXTVIEW("\xf0\x9f\x98\x80"));
 		Writer.AddString(FWideStringView(Value, UE_ARRAY_COUNT(Value)));
 		FCbFieldIterator Fields = Writer.Save();
 		if (TestEqual(TEXT("FCbWriter(String, Unicode) Validate"), ValidateCompactBinaryRange(Fields.GetOuterBuffer(), ECbValidateMode::All), ECbValidateError::None))
 		{
 			for (FCbFieldView Field : Fields)
 			{
-				TestEqual(TEXT("FCbWriter(String, Unicode).AsString()"), Field.AsString(), "\xf0\x9f\x98\x80"_U8SV);
+				TestEqual(TEXT("FCbWriter(String, Unicode).AsString()"), Field.AsString(), UTF8TEXTVIEW("\xf0\x9f\x98\x80"));
 			}
 		}
 	}
@@ -550,9 +550,9 @@ bool FCbWriterBoolTest::RunTest(const FString& Parameters)
 		Writer.EndArray();
 
 		Writer.BeginObject();
-		Writer.SetName("B1"_ASV).AddBool(false);
-		Writer.SetName("B2"_ASV).AddBool(false);
-		Writer.SetName("B3"_ASV).AddBool(false);
+		Writer.SetName(ANSITEXTVIEW("B1")).AddBool(false);
+		Writer.SetName(ANSITEXTVIEW("B2")).AddBool(false);
+		Writer.SetName(ANSITEXTVIEW("B3")).AddBool(false);
 		Writer.EndObject();
 
 		FCbFieldIterator Fields = Writer.Save();
@@ -812,8 +812,8 @@ bool FCbWriterCustomByNameTest::RunTest(const FString& Parameters)
 	};
 	const FCustomValue Values[] =
 	{
-		{ "Type1"_U8SV, {1, 2, 3} },
-		{ "Type2"_U8SV, {4, 5, 6} },
+		{ UTF8TEXTVIEW("Type1"), {1, 2, 3} },
+		{ UTF8TEXTVIEW("Type2"), {4, 5, 6} },
 	};
 
 	for (const FCustomValue& Value : Values)
@@ -848,44 +848,44 @@ bool FCbWriterComplexTest::RunTest(const FString& Parameters)
 		Writer.BeginObject();
 
 		const uint8 LocalField[] = { uint8(ECbFieldType::IntegerPositive | ECbFieldType::HasFieldName), 1, 'I', 42 };
-		Writer.AddField("FieldViewCopy"_ASV, FCbFieldView(LocalField));
-		Writer.AddField("FieldCopy"_ASV, FCbField(FSharedBuffer::Clone(MakeMemoryView(LocalField))));
+		Writer.AddField(ANSITEXTVIEW("FieldViewCopy"), FCbFieldView(LocalField));
+		Writer.AddField(ANSITEXTVIEW("FieldCopy"), FCbField(FSharedBuffer::Clone(MakeMemoryView(LocalField))));
 
 		const uint8 LocalObject[] = { uint8(ECbFieldType::Object | ECbFieldType::HasFieldName), 1, 'O', 7,
 			uint8(ECbFieldType::IntegerPositive | ECbFieldType::HasFieldName), 1, 'I', 42,
 			uint8(ECbFieldType::Null | ECbFieldType::HasFieldName), 1, 'N' };
-		Writer.AddObject("ObjectViewCopy"_ASV, FCbObjectView(LocalObject));
-		Writer.AddObject("ObjectCopy"_ASV, FCbObject(FSharedBuffer::Clone(MakeMemoryView(LocalObject))));
+		Writer.AddObject(ANSITEXTVIEW("ObjectViewCopy"), FCbObjectView(LocalObject));
+		Writer.AddObject(ANSITEXTVIEW("ObjectCopy"), FCbObject(FSharedBuffer::Clone(MakeMemoryView(LocalObject))));
 
 		const uint8 LocalArray[] = { uint8(ECbFieldType::UniformArray | ECbFieldType::HasFieldName), 1, 'A', 4, 2,
 			uint8(ECbFieldType::IntegerPositive), 42, 21 };
-		Writer.AddArray("ArrayViewCopy"_ASV, FCbArrayView(LocalArray));
-		Writer.AddArray("ArrayCopy"_ASV, FCbArray(FSharedBuffer::Clone(MakeMemoryView(LocalArray))));
+		Writer.AddArray(ANSITEXTVIEW("ArrayViewCopy"), FCbArrayView(LocalArray));
+		Writer.AddArray(ANSITEXTVIEW("ArrayCopy"), FCbArray(FSharedBuffer::Clone(MakeMemoryView(LocalArray))));
 
-		Writer.AddNull("Null"_ASV);
+		Writer.AddNull(ANSITEXTVIEW("Null"));
 
-		Writer.BeginObject("Binary"_ASV);
+		Writer.BeginObject(ANSITEXTVIEW("Binary"));
 		{
-			Writer.AddBinary("Empty"_ASV, FMemoryView());
-			Writer.AddBinary("Value"_ASV, MakeMemoryView("BinaryValue"));
-			Writer.AddBinary("LargeViewValue"_ASV, MakeMemoryView(FString::ChrN(256, TEXT('.'))));
-			Writer.AddBinary("LargeValue"_ASV, FSharedBuffer::Clone(MakeMemoryView(FString::ChrN(256, TEXT('!')))));
+			Writer.AddBinary(ANSITEXTVIEW("Empty"), FMemoryView());
+			Writer.AddBinary(ANSITEXTVIEW("Value"), MakeMemoryView("BinaryValue"));
+			Writer.AddBinary(ANSITEXTVIEW("LargeViewValue"), MakeMemoryView(FString::ChrN(256, TEXT('.'))));
+			Writer.AddBinary(ANSITEXTVIEW("LargeValue"), FSharedBuffer::Clone(MakeMemoryView(FString::ChrN(256, TEXT('!')))));
 		}
 		Writer.EndObject();
 
-		Writer.BeginObject("Strings"_ASV);
+		Writer.BeginObject(ANSITEXTVIEW("Strings"));
 		{
-			Writer.AddString("AnsiString"_ASV, "AnsiValue"_ASV);
-			Writer.AddString("WideString"_ASV, FString::ChrN(256, TEXT('.')));
-			Writer.AddString("EmptyAnsiString"_ASV, FAnsiStringView());
-			Writer.AddString("EmptyWideString"_ASV, FWideStringView());
-			Writer.AddString("EmptyUtf8String"_U8SV, FUtf8StringView());
+			Writer.AddString(ANSITEXTVIEW("AnsiString"), ANSITEXTVIEW("AnsiValue"));
+			Writer.AddString(ANSITEXTVIEW("WideString"), FString::ChrN(256, TEXT('.')));
+			Writer.AddString(ANSITEXTVIEW("EmptyAnsiString"), FAnsiStringView());
+			Writer.AddString(ANSITEXTVIEW("EmptyWideString"), FWideStringView());
+			Writer.AddString(UTF8TEXTVIEW("EmptyUtf8String"), FUtf8StringView());
 			Writer.AddString("AnsiStringLiteral", "AnsiValue");
 			Writer.AddString("WideStringLiteral", TEXT("AnsiValue"));
 		}
 		Writer.EndObject();
 
-		Writer.BeginArray("Integers"_ASV);
+		Writer.BeginArray(ANSITEXTVIEW("Integers"));
 		{
 			Writer.AddInteger(int32(-1));
 			Writer.AddInteger(int64(-1));
@@ -900,7 +900,7 @@ bool FCbWriterComplexTest::RunTest(const FString& Parameters)
 		}
 		Writer.EndArray();
 
-		Writer.BeginArray("UniformIntegers"_ASV);
+		Writer.BeginArray(ANSITEXTVIEW("UniformIntegers"));
 		{
 			Writer.AddInteger(0);
 			Writer.AddInteger(MAX_int32);
@@ -910,41 +910,41 @@ bool FCbWriterComplexTest::RunTest(const FString& Parameters)
 		}
 		Writer.EndArray();
 
-		Writer.AddFloat("Float32"_ASV, 1.0f);
-		Writer.AddFloat("Float64as32"_ASV, 2.0);
-		Writer.AddFloat("Float64"_ASV, 3.0e100);
+		Writer.AddFloat(ANSITEXTVIEW("Float32"), 1.0f);
+		Writer.AddFloat(ANSITEXTVIEW("Float64as32"), 2.0);
+		Writer.AddFloat(ANSITEXTVIEW("Float64"), 3.0e100);
 
-		Writer.AddBool("False"_ASV, false);
-		Writer.AddBool("True"_ASV, true);
+		Writer.AddBool(ANSITEXTVIEW("False"), false);
+		Writer.AddBool(ANSITEXTVIEW("True"), true);
 
-		Writer.AddObjectAttachment("ObjectAttachment"_ASV, FIoHash());
-		Writer.AddBinaryAttachment("BinaryAttachment"_ASV, FIoHash());
-		Writer.AddAttachment("Attachment"_ASV, FCbAttachment());
+		Writer.AddObjectAttachment(ANSITEXTVIEW("ObjectAttachment"), FIoHash());
+		Writer.AddBinaryAttachment(ANSITEXTVIEW("BinaryAttachment"), FIoHash());
+		Writer.AddAttachment(ANSITEXTVIEW("Attachment"), FCbAttachment());
 
-		Writer.AddHash("Hash"_ASV, FIoHash());
-		Writer.AddUuid("Uuid"_ASV, FGuid());
+		Writer.AddHash(ANSITEXTVIEW("Hash"), FIoHash());
+		Writer.AddUuid(ANSITEXTVIEW("Uuid"), FGuid());
 
-		Writer.AddDateTimeTicks("DateTimeZero"_ASV, 0);
-		Writer.AddDateTime("DateTime2020"_ASV, FDateTime(2020, 5, 13, 15, 10));
+		Writer.AddDateTimeTicks(ANSITEXTVIEW("DateTimeZero"), 0);
+		Writer.AddDateTime(ANSITEXTVIEW("DateTime2020"), FDateTime(2020, 5, 13, 15, 10));
 
-		Writer.AddTimeSpanTicks("TimeSpanZero"_ASV, 0);
-		Writer.AddTimeSpan("TimeSpan"_ASV, FTimespan(1, 2, 4, 8));
+		Writer.AddTimeSpanTicks(ANSITEXTVIEW("TimeSpanZero"), 0);
+		Writer.AddTimeSpan(ANSITEXTVIEW("TimeSpan"), FTimespan(1, 2, 4, 8));
 
-		Writer.AddObjectId("ObjectIdZero"_ASV, FCbObjectId());
-		Writer.AddObjectId("ObjectId"_ASV, FCbObjectId(MakeMemoryView<uint8>({1, 2, 3, 4, 5, 6, 7, 8, 9, 10, 11, 12})));
+		Writer.AddObjectId(ANSITEXTVIEW("ObjectIdZero"), FCbObjectId());
+		Writer.AddObjectId(ANSITEXTVIEW("ObjectId"), FCbObjectId(MakeMemoryView<uint8>({1, 2, 3, 4, 5, 6, 7, 8, 9, 10, 11, 12})));
 
-		Writer.BeginObject("NestedObjects"_ASV);
+		Writer.BeginObject(ANSITEXTVIEW("NestedObjects"));
 		{
-			Writer.BeginObject("Empty"_ASV);
+			Writer.BeginObject(ANSITEXTVIEW("Empty"));
 			Writer.EndObject();
 
-			Writer.BeginObject("Null"_ASV);
-			Writer.AddNull("Null"_ASV);
+			Writer.BeginObject(ANSITEXTVIEW("Null"));
+			Writer.AddNull(ANSITEXTVIEW("Null"));
 			Writer.EndObject();
 		}
 		Writer.EndObject();
 
-		Writer.BeginArray("NestedArrays"_ASV);
+		Writer.BeginArray(ANSITEXTVIEW("NestedArrays"));
 		{
 			Writer.BeginArray();
 			Writer.EndArray();
@@ -969,32 +969,32 @@ bool FCbWriterComplexTest::RunTest(const FString& Parameters)
 		}
 		Writer.EndArray();
 
-		Writer.BeginArray("ArrayOfObjects"_ASV);
+		Writer.BeginArray(ANSITEXTVIEW("ArrayOfObjects"));
 		{
 			Writer.BeginObject();
 			Writer.EndObject();
 
 			Writer.BeginObject();
-			Writer.AddNull("Null"_ASV);
+			Writer.AddNull(ANSITEXTVIEW("Null"));
 			Writer.EndObject();
 		}
 		Writer.EndArray();
 
-		Writer.BeginArray("LargeArray"_ASV);
+		Writer.BeginArray(ANSITEXTVIEW("LargeArray"));
 		for (int Index = 0; Index < 256; ++Index)
 		{
 			Writer.AddInteger(Index - 128);
 		}
 		Writer.EndArray();
 
-		Writer.BeginArray("LargeUniformArray"_ASV);
+		Writer.BeginArray(ANSITEXTVIEW("LargeUniformArray"));
 		for (int Index = 0; Index < 256; ++Index)
 		{
 			Writer.AddInteger(Index);
 		}
 		Writer.EndArray();
 
-		Writer.BeginArray("NestedUniformArray"_ASV);
+		Writer.BeginArray(ANSITEXTVIEW("NestedUniformArray"));
 		for (int Index = 0; Index < 16; ++Index)
 		{
 			Writer.BeginArray();
@@ -1047,48 +1047,48 @@ bool FCbWriterStreamTest::RunTest(const FString& Parameters)
 		Writer.BeginObject();
 
 		const uint8 LocalField[] = { uint8(ECbFieldType::IntegerPositive | ECbFieldType::HasFieldName), 1, 'I', 42 };
-		Writer << "FieldCopy"_ASV << FCbFieldView(LocalField);
+		Writer << ANSITEXTVIEW("FieldCopy") << FCbFieldView(LocalField);
 
 		const uint8 LocalObject[] = { uint8(ECbFieldType::Object | ECbFieldType::HasFieldName), 1, 'O', 7,
 			uint8(ECbFieldType::IntegerPositive | ECbFieldType::HasFieldName), 1, 'I', 42,
 			uint8(ECbFieldType::Null | ECbFieldType::HasFieldName), 1, 'N' };
-		Writer << "ObjectCopy"_ASV << FCbObjectView(LocalObject);
+		Writer << ANSITEXTVIEW("ObjectCopy") << FCbObjectView(LocalObject);
 
 		const uint8 LocalArray[] = { uint8(ECbFieldType::UniformArray | ECbFieldType::HasFieldName), 1, 'A', 4, 2,
 			uint8(ECbFieldType::IntegerPositive), 42, 21 };
-		Writer << "ArrayCopy"_ASV << FCbArrayView(LocalArray);
+		Writer << ANSITEXTVIEW("ArrayCopy") << FCbArrayView(LocalArray);
 
-		Writer << "Null"_ASV << nullptr;
+		Writer << ANSITEXTVIEW("Null") << nullptr;
 
-		Writer << "Strings"_ASV;
+		Writer << ANSITEXTVIEW("Strings");
 		Writer.BeginObject();
 		Writer
-			<< "AnsiString"_ASV << "AnsiValue"_ASV
-			<< "AnsiStringLiteral"_ASV << "AnsiValue"
-			<< "WideString"_ASV << TEXTVIEW("WideValue")
-			<< "WideStringLiteral"_ASV << TEXT("WideValue");
+			<< ANSITEXTVIEW("AnsiString") << ANSITEXTVIEW("AnsiValue")
+			<< ANSITEXTVIEW("AnsiStringLiteral") << "AnsiValue"
+			<< ANSITEXTVIEW("WideString") << TEXTVIEW("WideValue")
+			<< ANSITEXTVIEW("WideStringLiteral") << TEXT("WideValue");
 		Writer.EndObject();
 
-		Writer << "Integers"_ASV;
+		Writer << ANSITEXTVIEW("Integers");
 		Writer.BeginArray();
 		Writer << int32(-1) << int64(-1) << uint32(1) << uint64(1);
 		Writer.EndArray();
 
-		Writer << "Float32"_ASV << 1.0f;
-		Writer << "Float64"_ASV << 2.0;
+		Writer << ANSITEXTVIEW("Float32") << 1.0f;
+		Writer << ANSITEXTVIEW("Float64") << 2.0;
 
-		Writer << "False"_ASV << false << "True"_ASV << true;
+		Writer << ANSITEXTVIEW("False") << false << ANSITEXTVIEW("True") << true;
 
-		Writer << "Attachment"_ASV << FCbAttachment();
+		Writer << ANSITEXTVIEW("Attachment") << FCbAttachment();
 
-		Writer << "Hash"_ASV << FIoHash();
-		Writer << "Uuid"_ASV << FGuid();
+		Writer << ANSITEXTVIEW("Hash") << FIoHash();
+		Writer << ANSITEXTVIEW("Uuid") << FGuid();
 
-		Writer << "DateTime"_ASV << FDateTime(2020, 5, 13, 15, 10);
-		Writer << "TimeSpan"_ASV << FTimespan(1, 2, 4, 8);
+		Writer << ANSITEXTVIEW("DateTime") << FDateTime(2020, 5, 13, 15, 10);
+		Writer << ANSITEXTVIEW("TimeSpan") << FTimespan(1, 2, 4, 8);
 
-		Writer << "ObjectIdZero"_ASV << FCbObjectId();
-		Writer << "ObjectId"_ASV << FCbObjectId(MakeMemoryView<uint8>({1, 2, 3, 4, 5, 6, 7, 8, 9, 10, 11, 12}));
+		Writer << ANSITEXTVIEW("ObjectIdZero") << FCbObjectId();
+		Writer << ANSITEXTVIEW("ObjectId") << FCbObjectId(MakeMemoryView<uint8>({1, 2, 3, 4, 5, 6, 7, 8, 9, 10, 11, 12}));
 
 		Writer << "LiteralName" << nullptr;
 
@@ -1126,10 +1126,10 @@ bool FCbWriterStateTest::RunTest(const FString& Parameters)
 	//Writer.Save(MakeMemoryView(ZeroFieldLarge));
 
 	// Assert on empty name.
-	//Writer.SetName(""_ASV);
+	//Writer.SetName(ANSITEXTVIEW(""));
 
 	// Assert on name after name.
-	//Writer.SetName("Field"_ASV).SetName("Field"_ASV);
+	//Writer.SetName(ANSITEXTVIEW("Field")).SetName(ANSITEXTVIEW("Field"));
 
 	// Assert on missing name.
 	//Writer.BeginObject();
@@ -1138,7 +1138,7 @@ bool FCbWriterStateTest::RunTest(const FString& Parameters)
 
 	// Assert on name in array.
 	//Writer.BeginArray();
-	//Writer.SetName("Field"_ASV);
+	//Writer.SetName(ANSITEXTVIEW("Field"));
 	//Writer.EndArray();
 
 	// Assert on save in object.
@@ -1163,7 +1163,7 @@ bool FCbWriterStateTest::RunTest(const FString& Parameters)
 
 	// Assert on object end after name with no value.
 	//Writer.BeginObject();
-	//Writer.SetName("Field"_ASV);
+	//Writer.SetName(ANSITEXTVIEW("Field"));
 	//Writer.EndObject();
 
 	// Assert on writing a field with no value.
diff --git a/Engine/Source/Runtime/Core/Public/Containers/UnrealString.h b/Engine/Source/Runtime/Core/Public/Containers/UnrealString.h
--- a/Engine/Source/Runtime/Core/Public/Containers/UnrealString.h
+++ b/Engine/Source/Runtime/Core/Public/Containers/UnrealString.h
@@ -44,7 +44,7 @@ int32        GetNum(const FString& String);
  *
  * - Do not use the u8"..." prefix (gives the wrong array type until C++20).
  * - Use UTF8TEXT("...") for array literals (type is const UTF8CHAR[n]).
- * - Use "..."_U8SV for string view literals (type is FUtf8StringView).
+ * - Use UTF8TEXTVIEW("...") for string view literals (type is FUtf8StringView).
  * - Use \uxxxx or \Uxxxxxxxx escape sequences rather than \x to specify Unicode code points.
  */
 class CORE_API FString
