Logging: Fixed a static analysis warning

#rnx

diff --git a/Engine/Source/Runtime/Core/Private/Misc/OutputDeviceHelper.cpp b/Engine/Source/Runtime/Core/Private/Misc/OutputDeviceHelper.cpp
--- a/Engine/Source/Runtime/Core/Private/Misc/OutputDeviceHelper.cpp
+++ b/Engine/Source/Runtime/Core/Private/Misc/OutputDeviceHelper.cpp
@@ -129,15 +129,12 @@ void FOutputDeviceHelper::FormatCastAndSerializeLine(
 
 	if (bAutoEmitLineTerminator)
 	{
+	#if PLATFORM_UNIX
 		// Use Windows line endings on Linux for compatibility with Windows tools like notepad.exe
-		if constexpr (PLATFORM_UNIX)
-		{
-			Format.Append("\r\n");
-		}
-		else
-		{
-			Format.Append(LINE_TERMINATOR_ANSI);
-		}
+		Format.Append("\r\n");
+	#else
+		Format.Append(LINE_TERMINATOR_ANSI);
+	#endif
 	}
 
 	Output.Serialize(Format.GetData(), Format.Len() * sizeof(UTF8CHAR));
