StringView: Deprecated the public user-defined literals for string views

diff --git a/Engine/Source/Runtime/Core/Public/Containers/StringView.h b/Engine/Source/Runtime/Core/Public/Containers/StringView.h
--- a/Engine/Source/Runtime/Core/Public/Containers/StringView.h
+++ b/Engine/Source/Runtime/Core/Public/Containers/StringView.h
@@ -42,7 +42,7 @@ namespace UE::Core::Private
  * ProcessPath(Object->GetPathName());
  *
  * A string view is implicitly constructible from null-terminated strings, from contiguous ranges
- * of characters such as FString and TStringBuilder, and from literals such as TEXT("..."_SV).
+ * of characters such as FString and TStringBuilder, and from literals such as TEXTVIEW("...").
  *
  * A string view is cheap to copy and is meant to be passed by value. Avoid passing by reference.
  *
@@ -69,7 +69,7 @@ namespace UE::Core::Private
  *	UseString(CString);
  *	UseString(StringBuilder);
  *	UseString(TEXT("ABC"));
- *	UseString(TEXT("ABC"_SV));
+ *	UseString(TEXTVIEW("ABC"));
  * @endcode
  */
 template <typename CharType>
@@ -351,27 +351,32 @@ constexpr inline auto GetNum(TStringView<CharType> String)
 //////////////////////////////////////////////////////////////////////////
 
 #if PLATFORM_TCHAR_IS_UTF8CHAR
+UE_DEPRECATED(5.1, "The _SV literal is deprecated. Use TEXTVIEW.")
 inline FStringView operator "" _SV(const ANSICHAR* String, size_t Size)
 {
 	return FStringView((const UTF8CHAR*)String, Size);
 }
 #else
+UE_DEPRECATED(5.1, "The _SV literal is deprecated. Use TEXTVIEW.")
 constexpr inline FStringView operator "" _SV(const TCHAR* String, size_t Size)
 {
 	return FStringView(String, Size);
 }
 #endif
 
+UE_DEPRECATED(5.1, "The _ASV literal is deprecated. Use ANSITEXTVIEW.")
 constexpr inline FAnsiStringView operator "" _ASV(const ANSICHAR* String, size_t Size)
 {
 	return FAnsiStringView(String, Size);
 }
 
+UE_DEPRECATED(5.1, "The _WSV literal is deprecated. Use WIDETEXTVIEW.")
 constexpr inline FWideStringView operator "" _WSV(const WIDECHAR* String, size_t Size)
 {
 	return FWideStringView(String, Size);
 }
 
+UE_DEPRECATED(5.1, "The _U8SV literal is deprecated. Use UTF8TEXTVIEW.")
 /*constexpr*/ inline FUtf8StringView operator "" _U8SV(const ANSICHAR* String, size_t Size)
 {
 	// Would like this operator to be constexpr, but cannot be until after this operator can take a UTF8CHAR*
@@ -380,13 +385,42 @@ constexpr inline FWideStringView operator "" _WSV(const WIDECHAR* String, size_t
 }
 
 #if PLATFORM_TCHAR_IS_UTF8CHAR
-	#define TEXTVIEW(str) (str##_SV)
+inline FStringView operator "" _PrivateSV(const ANSICHAR* String, size_t Size)
+{
+	return FStringView((const UTF8CHAR*)String, Size);
+}
+#else
+constexpr inline FStringView operator "" _PrivateSV(const TCHAR* String, size_t Size)
+{
+	return FStringView(String, Size);
+}
+#endif
+
+constexpr inline FAnsiStringView operator "" _PrivateASV(const ANSICHAR* String, size_t Size)
+{
+	return FAnsiStringView(String, Size);
+}
+
+constexpr inline FWideStringView operator "" _PrivateWSV(const WIDECHAR* String, size_t Size)
+{
+	return FWideStringView(String, Size);
+}
+
+/*constexpr*/ inline FUtf8StringView operator "" _PrivateU8SV(const ANSICHAR* String, size_t Size)
+{
+	// Would like this operator to be constexpr, but cannot be until after this operator can take a UTF8CHAR*
+	// rather than an ANSICHAR*, which won't be until we have C++20 char8_t string literals.
+	return FUtf8StringView(reinterpret_cast<const UTF8CHAR*>(String), Size);
+}
+
+#if PLATFORM_TCHAR_IS_UTF8CHAR
+	#define TEXTVIEW(str) (str##_PrivateSV)
 #else
-	#define TEXTVIEW(str) TEXT(str##_SV)
+	#define TEXTVIEW(str) TEXT(str##_PrivateSV)
 #endif
-#define ANSITEXTVIEW(str) (str##_ASV)
-#define WIDETEXTVIEW(str) PREPROCESSOR_JOIN(WIDETEXT(str), _WSV)
-#define UTF8TEXTVIEW(str) (str##_U8SV)
+#define ANSITEXTVIEW(str) (str##_PrivateASV)
+#define WIDETEXTVIEW(str) PREPROCESSOR_JOIN(WIDETEXT(str), _PrivateWSV)
+#define UTF8TEXTVIEW(str) (str##_PrivateU8SV)
 
 //////////////////////////////////////////////////////////////////////////
 
