StallDetector support for Linux

### Features

This change enables the StallDetector watchdog in Editor to submit reports to crashreporter about threads violating instrumented deadlines in the source code. This feature was available prior on Windows, and this change adds Linux support.

### Notes

New APIs:
ReportStall()
CaptureThreadPortableCallStack()

Many APIs are updated from purely "Ensure" naming to more general naming. Stalls are more like Ensures than crashes, and so the appropriate renames to make the code readable and clear have been made. In some places Ensure is replaced with the clearer: Continuable Event nomenclature.

### Testing

I synthesized an ensure on Linux, and did the same for a stall. I then compared crash report XML file to make sure they contain accurate data in the callstack, portable callstack, and other fields in the report. I also noted that the stall information was showing as expected in the crash reporter.

diff --git a/Engine/Source/Programs/CrashReportClient/Private/CrashReportClientApp.cpp b/Engine/Source/Programs/CrashReportClient/Private/CrashReportClientApp.cpp
--- a/Engine/Source/Programs/CrashReportClient/Private/CrashReportClientApp.cpp
+++ b/Engine/Source/Programs/CrashReportClient/Private/CrashReportClientApp.cpp
@@ -621,7 +621,7 @@ FPlatformErrorReport CollectErrorReport(FRecoveryService* RecoveryService, uint3
 	if (RecoveryService && 
 		DirectoryExists && 
 		SharedCrashContext.UserSettings.bSendUsageData && 
-		!FGenericCrashContext::IsTypeContinuable(SharedCrashContext.CrashType)
+		!FPlatformCrashContext::IsTypeContinuable(SharedCrashContext.CrashType)
 	{
 		RecoveryService->CollectFiles(ReportDirectoryAbsolutePath);
 	}
diff --git a/Engine/Source/Runtime/Core/Private/GenericPlatform/GenericPlatformCrashContext.cpp b/Engine/Source/Runtime/Core/Private/GenericPlatform/GenericPlatformCrashContext.cpp
--- a/Engine/Source/Runtime/Core/Private/GenericPlatform/GenericPlatformCrashContext.cpp
+++ b/Engine/Source/Runtime/Core/Private/GenericPlatform/GenericPlatformCrashContext.cpp
@@ -1219,6 +1219,18 @@ void FGenericCrashContext::CapturePortableCallStack(void* ErrorProgramCounter, v
 	SetPortableCallStack(StackTraceCursor, StackTraceDepth);
 }
 
+void FGenericCrashContext::CaptureThreadPortableCallStack(const uint64 ThreadId, void* Context)
+{
+	// Capture the stack trace
+	static const int StackTraceMaxDepth = 100;
+	uint64 StackTrace[StackTraceMaxDepth];
+	FMemory::Memzero(StackTrace);
+	int32 StackTraceDepth = FPlatformStackWalk::CaptureThreadStackBackTrace(ThreadId, StackTrace, StackTraceMaxDepth, Context);
+
+	// Generate the portable callstack from it
+	SetPortableCallStack(StackTrace, StackTraceDepth);
+}
+
 void FGenericCrashContext::CapturePortableCallStack(int32 NumStackFramesToIgnore, void* Context)
 {
 	CapturePortableCallStack(nullptr, Context);
diff --git a/Engine/Source/Runtime/Core/Private/ProfilingDebugging/StallDetector.cpp b/Engine/Source/Runtime/Core/Private/ProfilingDebugging/StallDetector.cpp
--- a/Engine/Source/Runtime/Core/Private/ProfilingDebugging/StallDetector.cpp
+++ b/Engine/Source/Runtime/Core/Private/ProfilingDebugging/StallDetector.cpp
@@ -21,10 +21,6 @@ TRACE_DECLARE_FLOAT_COUNTER(StallTimeSeconds, TEXT("StallDetector/TimeSeconds"))
 // use the heart beat clock to account for process suspend
 #define STALL_DETECTOR_HEART_BEAT_CLOCK 1
 
-#if STALL_DETECTOR_DEBUG && _MSC_VER
-#pragma optimize( "", off )
-#endif // STALL_DETECTOR_DEBUG && _MSC_VER
-
 #if STALL_DETECTOR_HEART_BEAT_CLOCK
  #include "HAL/ThreadHeartBeat.h"
 #endif // STALL_DETECTOR_HEART_BEAT_CLOCK
diff --git a/Engine/Source/Runtime/Core/Private/Unix/UnixPlatformCrashContext.cpp b/Engine/Source/Runtime/Core/Private/Unix/UnixPlatformCrashContext.cpp
--- a/Engine/Source/Runtime/Core/Private/Unix/UnixPlatformCrashContext.cpp
+++ b/Engine/Source/Runtime/Core/Private/Unix/UnixPlatformCrashContext.cpp
@@ -93,7 +93,7 @@ void FGenericCrashContext::CleanupPlatformSpecificFiles()
 {
 }
 
-__thread siginfo_t FUnixCrashContext::FakeSiginfoForEnsures;
+__thread siginfo_t FUnixCrashContext::FakeSiginfoForDiagnostics;
 
 FUnixCrashContext::~FUnixCrashContext()
 {
@@ -114,19 +114,19 @@ void FUnixCrashContext::InitFromSignal(int32 InSignal, siginfo_t* InInfo, void*
 	FCString::Strcat(SignalDescription, UE_ARRAY_COUNT( SignalDescription ) - 1, *DescribeSignal(Signal, Info, Context));
 }
 
-void FUnixCrashContext::InitFromEnsureHandler(const TCHAR* EnsureMessage, const void* CrashAddress)
+void FUnixCrashContext::InitFromDiagnostics(const void* InAddress)
 {
 	Signal = SIGTRAP;
 
-	FakeSiginfoForEnsures.si_signo = SIGTRAP;
-	FakeSiginfoForEnsures.si_code = TRAP_TRACE;
-	FakeSiginfoForEnsures.si_addr = const_cast<void *>(CrashAddress);
-	Info = &FakeSiginfoForEnsures;
+	FakeSiginfoForDiagnostics.si_signo = SIGTRAP;
+	FakeSiginfoForDiagnostics.si_code = TRAP_TRACE;
+	FakeSiginfoForDiagnostics.si_addr = const_cast<void *>(InAddress);
+	Info = &FakeSiginfoForDiagnostics;
 
 	Context = nullptr;
 
 	// set signal description to a more human-readable one for ensures
-	FCString::Strcpy(SignalDescription, UE_ARRAY_COUNT(SignalDescription) - 1, EnsureMessage);
+	FCString::Strcpy(SignalDescription, UE_ARRAY_COUNT(SignalDescription) - 1, ErrorMessage);
 
 	// only need the first string
 	for (int Idx = 0; Idx < UE_ARRAY_COUNT(SignalDescription); ++Idx)
@@ -284,6 +284,8 @@ void FUnixCrashContext::CaptureStackTrace(void* ErrorProgramCounter)
 	// Only do work the first time this function is called - this is mainly a carry over from Windows where it can be called multiple times, left intact for extra safety.
 	if (!bCapturedBacktrace)
 	{
+		bCapturedBacktrace = true;
+
 		static const SIZE_T StackTraceSize = 65535;
 		static ANSICHAR StackTrace[StackTraceSize];
 		StackTrace[0] = 0;
@@ -299,8 +301,33 @@ void FUnixCrashContext::CaptureStackTrace(void* ErrorProgramCounter)
 
 		FCString::Strncat( GErrorHist, UTF8_TO_TCHAR(StackTrace), UE_ARRAY_COUNT(GErrorHist) - 1 );
 		CreateExceptionInfoString(Signal, Info, Context);
+	}
+}
 
+void FUnixCrashContext::CaptureThreadStackTrace(uint32_t ThreadId)
+{
+	// Only do work the first time this function is called - this is mainly a carry over from Windows where it can be called multiple times, left intact for extra safety.
+	if (!bCapturedBacktrace)
+	{
 		bCapturedBacktrace = true;
+
+		CaptureThreadPortableCallStack(ThreadId, this);
+
+		// The crash report XML has an element <CallStack>. However, CrashReportClient will run after the UE process generates the report. It will overwrite this
+		//  part of the report with whatever is in the CALLSTACK section of Diagnostics.txt. Because of this we have to call APIs here to ensure that MiniDumpCallstackInfo
+		//  gets filled out. So, we iterate through the portable callstack for the thread and call ProgramCounterToHumanReadableString passing the CrashContext (this).
+		static const SIZE_T StackTraceSize = 65535;
+		static ANSICHAR StackTrace[StackTraceSize];
+		StackTrace[0] = 0;
+		for ( int i=0; i<CallStack.Num(); i++ )
+		{
+			FPlatformStackWalk::ProgramCounterToHumanReadableString( i, CallStack[i].BaseAddress + CallStack[i].Offset, StackTrace, StackTraceSize, this );
+			FCStringAnsi::Strncat(StackTrace, LINE_TERMINATOR_ANSI, (int32)StackTraceSize);
+		}
+
+		// We do not set the ExceptionInfo string here as it just gets details for the exact callsite,
+		//  and this function by definition is just capturing the state of some other specific thread
+		//  (and not some instrumentation point or instruction pointer that generated a signal)
 	}
 }
 
@@ -488,8 +515,10 @@ void FUnixCrashContext::AddPlatformSpecificProperties() const
 	}
 }
 
-void FUnixCrashContext::GenerateCrashInfoAndLaunchReporter(bool bReportingNonCrash) const
+void FUnixCrashContext::GenerateCrashInfoAndLaunchReporter() const
 {
+	const bool bReportingNonCrash = IsTypeContinuable(Type);
+
 	// do not report crashes for tools (particularly for crash reporter itself)
 #if !IS_PROGRAM
 
@@ -572,7 +601,19 @@ void FUnixCrashContext::GenerateCrashInfoAndLaunchReporter(bool bReportingNonCra
 
 	if (!bSkipCRC)
 	{
-		FString CrashInfoFolder = FPaths::Combine(*FPaths::ProjectSavedDir(), TEXT("Crashes"), *FString::Printf(TEXT("%sinfo-%s-pid-%d-%s"), bReportingNonCrash ? TEXT("ensure") : TEXT("crash"), FApp::GetProjectName(), getpid(), *CrashGuid));
+		const TCHAR* TypeString = TEXT("crash");
+		switch(Type)
+		{
+		case ECrashContextType::Ensure:
+			TypeString = TEXT("ensure");
+			break;
+
+		case ECrashContextType::Stall:
+			TypeString = TEXT("stall");
+			break;
+		}
+
+		FString CrashInfoFolder = FPaths::Combine(*FPaths::ProjectSavedDir(), TEXT("Crashes"), *FString::Printf(TEXT("%sinfo-%s-pid-%d-%s"), TypeString, FApp::GetProjectName(), getpid(), *CrashGuid));
 		FString CrashInfoAbsolute = FPaths::ConvertRelativePathToFull(CrashInfoFolder);
 		if (IFileManager::Get().MakeDirectory(*CrashInfoAbsolute, true))
 		{
diff --git a/Engine/Source/Runtime/Core/Private/Unix/UnixPlatformStackWalk.cpp b/Engine/Source/Runtime/Core/Private/Unix/UnixPlatformStackWalk.cpp
--- a/Engine/Source/Runtime/Core/Private/Unix/UnixPlatformStackWalk.cpp
+++ b/Engine/Source/Runtime/Core/Private/Unix/UnixPlatformStackWalk.cpp
@@ -531,8 +531,12 @@ bool FUnixPlatformStackWalk::ProgramCounterToHumanReadableString( int32 CurrentC
 			// Get filename, source file and line number
 			FUnixCrashContext* UnixContext = static_cast< FUnixCrashContext* >( Context );
 
-			// for ensure, use the fast path - do not even attempt to get detailed info as it will result in long hitch
-			bool bAddDetailedInfo = UnixContext ? UnixContext->GetType() != ECrashContextType::Ensure : false;
+			// do not even attempt to get detailed info for continuable events (like ensure) as it will result in long hitch, use the fast path
+			bool bAddDetailedInfo = false;
+			if (UnixContext)
+			{
+				bAddDetailedInfo = !FPlatformCrashContext::IsTypeContinuable(UnixContext->GetType());
+			}
 
 			// Program counters in the backtrace point to the location from where the execution will be resumed (in all frames except the one where we crashed),
 			// which results in callstack pointing to the next lines in code. In order to determine the source line where the actual call happened, we need to go
@@ -944,44 +948,67 @@ void ReportGPUCrash(const TCHAR* ErrorMessage, void* ProgramCounter)
 	FPlatformMisc::RaiseException(1);
 }
 
-static FCriticalSection EnsureLock;
+static FCriticalSection ReportLock;
 static bool bReentranceGuard = false;
 
 void ReportEnsure(const TCHAR* ErrorMessage, void* ProgramCounter)
 {
 	// Simple re-entrance guard.
-	EnsureLock.Lock();
+	ReportLock.Lock();
 
 	if (bReentranceGuard)
 	{
-		EnsureLock.Unlock();
+		ReportLock.Unlock();
 		return;
 	}
 
 	bReentranceGuard = true;
 
 	FUnixCrashContext EnsureContext(ECrashContextType::Ensure, ErrorMessage);
-	EnsureContext.InitFromEnsureHandler(ErrorMessage, ProgramCounter);
+	EnsureContext.InitFromDiagnostics(ProgramCounter);
 
 	EnsureContext.CaptureStackTrace(ProgramCounter);
-	EnsureContext.GenerateCrashInfoAndLaunchReporter(true);
+	EnsureContext.GenerateCrashInfoAndLaunchReporter();
+
+	bReentranceGuard = false;
+	ReportLock.Unlock();
+}
+
+void ReportStall(const TCHAR* Message, uint32 ThreadId)
+{
+	// Simple re-entrance guard.
+	ReportLock.Lock();
+
+	if (bReentranceGuard)
+	{
+		ReportLock.Unlock();
+		return;
+	}
+
+	bReentranceGuard = true;
+
+	FUnixCrashContext StallContext(ECrashContextType::Stall, Message);
+	StallContext.InitFromDiagnostics();
+
+	StallContext.CaptureThreadStackTrace(ThreadId);
+	StallContext.GenerateCrashInfoAndLaunchReporter();
 
 	bReentranceGuard = false;
-	EnsureLock.Unlock();
+	ReportLock.Unlock();
 }
 
 void ReportHang(const TCHAR* ErrorMessage, const uint64* StackFrames, int32 NumStackFrames, uint32 HungThreadId)
 {
-	EnsureLock.Lock();
+	ReportLock.Lock();
 	if (!bReentranceGuard)
 	{
 		bReentranceGuard = true;
 
-		FUnixCrashContext EnsureContext(ECrashContextType::Hang, ErrorMessage);
-		EnsureContext.SetPortableCallStack(StackFrames, NumStackFrames);
-		EnsureContext.GenerateCrashInfoAndLaunchReporter(true);
+		FUnixCrashContext HangContext(ECrashContextType::Hang, ErrorMessage);
+		HangContext.SetPortableCallStack(StackFrames, NumStackFrames);
+		HangContext.GenerateCrashInfoAndLaunchReporter();
 
 		bReentranceGuard = false;
 	}
-	EnsureLock.Unlock();
+	ReportLock.Unlock();
 }
diff --git a/Engine/Source/Runtime/Core/Public/GenericPlatform/GenericPlatformCrashContext.h b/Engine/Source/Runtime/Core/Public/GenericPlatform/GenericPlatformCrashContext.h
--- a/Engine/Source/Runtime/Core/Public/GenericPlatform/GenericPlatformCrashContext.h
+++ b/Engine/Source/Runtime/Core/Public/GenericPlatform/GenericPlatformCrashContext.h
@@ -473,12 +473,19 @@ public:
 	void SetNumMinidumpFramesToIgnore(int32 InNumMinidumpFramesToIgnore);
 
 	/**
-	 * Generate raw call stack for crash report (image base + offset)
+	 * Generate raw call stack for crash report (image base + offset) for the calling thread
 	 * @param ErrorProgramCounter The program counter of where the occur occurred in the callstack being captured
 	 * @param Context Optional thread context information
 	 */
 	void CapturePortableCallStack(void* ErrorProgramCounter, void* Context);
 
+	/**
+	 * Generate raw call stack for crash report (image base + offset) for a different thread
+	 * @param InThreadId The thread id of the thread to capture the callstack for
+	 * @param Context Optional thread context information
+	 */
+	void CaptureThreadPortableCallStack(const uint64 ThreadId, void* Context);
+
 	UE_DEPRECATED(5.0, "")
 	void CapturePortableCallStack(int32 NumStackFramesToIgnore, void* Context);
 	
diff --git a/Engine/Source/Runtime/Core/Public/HAL/ExceptionHandling.h b/Engine/Source/Runtime/Core/Public/HAL/ExceptionHandling.h
--- a/Engine/Source/Runtime/Core/Public/HAL/ExceptionHandling.h
+++ b/Engine/Source/Runtime/Core/Public/HAL/ExceptionHandling.h
@@ -38,6 +38,7 @@ extern CORE_API void ReportHang(const TCHAR*, const uint64* StackFrames, int32 N
 extern CORE_API void ReportAssert(const TCHAR* ErrorMessage, void* ProgramCounter);
 extern CORE_API void ReportGPUCrash(const TCHAR* ErrorMessage, void* ProgramCounter);
 extern CORE_API void ReportEnsure(const TCHAR* ErrorMessage, void* ProgramCounter);
+extern CORE_API void ReportStall(const TCHAR* ErrorMessage, uint32 ThreadId);
 extern CORE_API void ReportHang(const TCHAR*, const uint64* StackFrames, int32 NumStackFrames, uint32 HungThreadId);
 #elif PLATFORM_HOLOLENS
 #include "HoloLens/HoloLensSystemIncludes.h"
diff --git a/Engine/Source/Runtime/Core/Public/ProfilingDebugging/StallDetector.h b/Engine/Source/Runtime/Core/Public/ProfilingDebugging/StallDetector.h
--- a/Engine/Source/Runtime/Core/Public/ProfilingDebugging/StallDetector.h
+++ b/Engine/Source/Runtime/Core/Public/ProfilingDebugging/StallDetector.h
@@ -15,11 +15,11 @@
 #include <atomic>
 
 #ifndef STALL_DETECTOR
- #if WITH_EDITOR && PLATFORM_WINDOWS && !UE_BUILD_SHIPPING && COUNTERSTRACE_ENABLED
+ #if WITH_EDITOR && ( PLATFORM_WINDOWS || PLATFORM_UNIX ) && !UE_BUILD_SHIPPING && COUNTERSTRACE_ENABLED
   #define STALL_DETECTOR 1
- #else // WITH_EDITOR && PLATFORM_WINDOWS && !UE_BUILD_SHIPPING && COUNTERSTRACE_ENABLED
+ #else
   #define STALL_DETECTOR 0
- #endif // WITH_EDITOR && PLATFORM_WINDOWS && !UE_BUILD_SHIPPING && COUNTERSTRACE_ENABLED
+ #endif
 #endif // STALL_DETECTOR
 
 #if STALL_DETECTOR
diff --git a/Engine/Source/Runtime/Core/Public/Unix/UnixPlatformCrashContext.h b/Engine/Source/Runtime/Core/Public/Unix/UnixPlatformCrashContext.h
--- a/Engine/Source/Runtime/Core/Public/Unix/UnixPlatformCrashContext.h
+++ b/Engine/Source/Runtime/Core/Public/Unix/UnixPlatformCrashContext.h
@@ -45,8 +45,8 @@ struct CORE_API FUnixCrashContext : public FGenericCrashContext
 	/** Memory reserved for minidump-style callstack info */
 	char MinidumpCallstackInfo[16384];
 
-	/** Fake siginfo used when handling ensure()s */
-	static __thread siginfo_t	FakeSiginfoForEnsures;
+	/** Fake siginfo used when handling ensure(), etc */
+	static __thread siginfo_t	FakeSiginfoForDiagnostics;
 
 	/** The PC of the first function used when handling a crash. Used to figure out the number of frames to ignore */
 	uint64* FirstCrashHandlerFrame = nullptr;
@@ -79,27 +79,31 @@ struct CORE_API FUnixCrashContext : public FGenericCrashContext
 	void InitFromSignal(int32 InSignal, siginfo_t* InInfo, void* InContext);
 
 	/**
-	 * Inits the crash context from ensure handler
+	 * Inits the crash context from some diagnostic handler (ensure, stall, etc...)
 	 *
-	 * @param EnsureMessage Message about the ensure that failed
-	 * @param CrashAddress address where "crash" happened
+	 * @param InAddress address where the event happened, for optional callstack trimming
 	 */
-	void InitFromEnsureHandler(const TCHAR* EnsureMessage, const void* CrashAddress);
+	void InitFromDiagnostics(const void* InAddress = nullptr);
 
 	/**
-	 * Populates crash context stack trace and a few related fields
+	 * Populates crash context stack trace and a few related fields for the calling thread
 	 *
 	 */
 	void CaptureStackTrace(void* ErrorProgramCounter);
 
 	/**
-	 * Generates a new crash report containing information needed for the crash reporter and launches it; may not return.
+	 * Populates crash context stack trace and a few related fields for another thread
 	 *
-	 * @param bReportingNonCrash if true, we are not reporting a crash (but e.g. ensure()), so don't re-raise the signal.
+	 * @param ThreadId The thread to capture the stack trace from
+	 */
+	void CaptureThreadStackTrace(uint32_t ThreadId);
+
+	/**
+	 * Generates a new crash report containing information needed for the crash reporter and launches it; may not return.
 	 *
-	 * @return If bReportingNonCrash is false, the function will not return
+	 * @return If the crash type is not continuable, the function will not return
 	 */
-	void GenerateCrashInfoAndLaunchReporter(bool bReportingNonCrash = false) const;
+	void GenerateCrashInfoAndLaunchReporter() const;
 
 	/**
 	 * Sets whether this crash represents a non-crash event like an ensure
