Fix bugs in crash malloc.
When a small pool was  exhausted it became impossible to free any allocations in that pool, because Free used the same search function as Malloc which only considered pools with space available.
Refactor to store bookkeeping and allocations separately, and match allocations to pools with pointer arithmetic.

diff --git a/Engine/Source/Runtime/Core/Private/GenericPlatform/GenericPlatformMallocCrash.cpp b/Engine/Source/Runtime/Core/Private/GenericPlatform/GenericPlatformMallocCrash.cpp
--- a/Engine/Source/Runtime/Core/Private/GenericPlatform/GenericPlatformMallocCrash.cpp
+++ b/Engine/Source/Runtime/Core/Private/GenericPlatform/GenericPlatformMallocCrash.cpp
@@ -6,14 +6,17 @@
 #include "Templates/AlignmentTemplates.h"
 #include "HAL/LowLevelMemTracker.h"
 #include "ProfilingDebugging/MemoryTrace.h"
+#include "Algo/BinarySearch.h"
+
 
 /** Describes a pool. */
 struct FPoolDesc
 {
-	FPoolDesc( const uint32 InSize, const uint32 InNumAllocs ):
-		Size(InSize),
-		NumAllocs(InNumAllocs)
-	{}
+	constexpr FPoolDesc(uint32 InSize, const uint32 InNumAllocs)
+		: Size(InSize)
+		, NumAllocs(InNumAllocs)
+	{
+	}
 
 	/** Size of the pool. */
 	const uint32 Size;
@@ -22,42 +25,52 @@ struct FPoolDesc
 	const uint32 NumAllocs;
 };
 
-/** Generated by the FGenericPlatformMallocCrash::PrintPoolsUsage. */
-const FPoolDesc& FGenericPlatformMallocCrash::GetPoolDesc( uint32 Index ) const
-{
 #ifdef NUM_ALLOCS_MODIFIER
 	#error NUM_ALLOCS_MODIFIER already defined
 #endif // NUM_ALLOCS_MODIFIER
 
 #if PLATFORM_UNIX
-	// Linux crash handling exhausts some of pools sometimes
+// Linux crash handling exhausts some of pools sometimes
 	#define NUM_ALLOCS_MODIFIER		2
 #else
 	#define NUM_ALLOCS_MODIFIER		1
 #endif // PLATFORM_UNIX
 
-	static const FPoolDesc AllPoolDesc[NUM_POOLS]=
-	{
-		FPoolDesc(   64, 224 * NUM_ALLOCS_MODIFIER),
-		FPoolDesc(   96, 144 * NUM_ALLOCS_MODIFIER),
-		FPoolDesc(  128,  80 * NUM_ALLOCS_MODIFIER),
-		FPoolDesc(  192, 560 * NUM_ALLOCS_MODIFIER),
-		FPoolDesc(  256, 384 * NUM_ALLOCS_MODIFIER),
-		FPoolDesc(  384, 208 * NUM_ALLOCS_MODIFIER),
-		FPoolDesc(  512,  48 * NUM_ALLOCS_MODIFIER),
-		FPoolDesc(  768,  32 * NUM_ALLOCS_MODIFIER),
-		FPoolDesc( 1024,  32 * NUM_ALLOCS_MODIFIER),
-		FPoolDesc( 2048,  32 * NUM_ALLOCS_MODIFIER),
-		FPoolDesc( 4096,  32 * NUM_ALLOCS_MODIFIER),
-		FPoolDesc( 8192,  32 * NUM_ALLOCS_MODIFIER),
-		FPoolDesc(16384,  16 * NUM_ALLOCS_MODIFIER),
-		FPoolDesc(32768,  16 * NUM_ALLOCS_MODIFIER),
-	};
-	return AllPoolDesc[Index];
+template<uint32 AllocSize>
+constexpr FPoolDesc MakePoolDesc(uint32 NumAllocs)
+{
+	static_assert(IsAligned(AllocSize, FGenericPlatformMallocCrash::REQUIRED_ALIGNMENT),
+		"All pool sizes in FGenericPlatformMallocCrash must be aligned to required alignment.");
+	return FPoolDesc(AllocSize, NumAllocs * NUM_ALLOCS_MODIFIER);
 }
 
 #undef NUM_ALLOCS_MODIFIER
 
+
+static constexpr FPoolDesc AllPoolDesc[FGenericPlatformMallocCrash::NUM_POOLS] =
+{
+	MakePoolDesc<64>(224),
+	MakePoolDesc<96>(144),
+	MakePoolDesc<128>(80),
+	MakePoolDesc<192>(560),
+	MakePoolDesc<256>(384),
+	MakePoolDesc<384>(208),
+	MakePoolDesc<512>(48),
+	MakePoolDesc<768>(32),
+	MakePoolDesc<1024>(32),
+	MakePoolDesc<2048>(32),
+	MakePoolDesc<4096>(32),
+	MakePoolDesc<8192>(32),
+	MakePoolDesc<16384>(16),
+	MakePoolDesc<32768>(16),
+};
+
+/** Generated by the FGenericPlatformMallocCrash::PrintPoolsUsage. */
+const FPoolDesc& FGenericPlatformMallocCrash::GetPoolDesc( uint32 Index ) const
+{
+	return AllPoolDesc[Index];
+}
+
 struct FMallocCrashPool
 {
 	uint32 NumUsed;
@@ -69,7 +82,10 @@ struct FMallocCrashPool
 	const uint32 AllocationSize;
 
 	/** Fixed list of allocations for the specified size for this pool. */
-	FPtrInfo* Allocations[FGenericPlatformMallocCrash::MAX_NUM_ALLOCS_IN_POOL];
+	TArrayView<FPtrInfo> Allocations;
+
+	// All allocations are used within this region of memory, use to match allocations to pools when freeing
+	uint8* AllocBase;
 
 	/** Maximum number of allocations that can be made for this pool. */
 	const uint32 MaxNumAllocations;
@@ -82,22 +98,21 @@ struct FMallocCrashPool
 		MaxUsedIndex(0),
 		MaxNumUsed(0),
 		TotalNumUsed(0),
-		AllocationSize( PoolDesc.Size+FGenericPlatformMallocCrash::PER_ALLOC_OVERHEAD ),
+		AllocationSize( PoolDesc.Size ),
 		MaxNumAllocations( PoolDesc.NumAllocs )
 	{
-		for( uint32 Index = 0; Index < MaxNumAllocations; ++Index )
-		{
-			uint8* NewPtr = (uint8*)Outer.AllocateFromSmallPool( AllocationSize );
-			Allocations[Index] = new(NewPtr) FPtrInfo( NewPtr+FGenericPlatformMallocCrash::PER_ALLOC_OVERHEAD );
-		}
+		Allocations = TArrayView<FPtrInfo>(
+			reinterpret_cast<FPtrInfo*>(Outer.AllocateFromBookkeeping(MaxNumAllocations * sizeof(FPtrInfo))), 
+			MaxNumAllocations
+			);
 
-		// Zero the rest.
-		for( uint32 Index = MaxNumAllocations; Index < FGenericPlatformMallocCrash::MAX_NUM_ALLOCS_IN_POOL; ++Index )
+		AllocBase = (uint8*)Outer.AllocateFromSmallPool(AllocationSize * MaxNumAllocations);
+		for( uint32 Index = 0; Index < MaxNumAllocations; ++Index )
 		{
-			Allocations[Index] = nullptr;
+			new (&Allocations[Index]) FPtrInfo(AllocBase + Index * AllocationSize);
 		}
-
-		AllocatedMemory = MaxNumAllocations*AllocationSize + sizeof(Allocations);
+		
+		AllocatedMemory = MaxNumAllocations*AllocationSize + MaxNumAllocations*sizeof(FPtrInfo);
 
 #ifdef	_DEBUG
 		FPlatformMisc::LowLevelOutputDebugStringf( TEXT( "FMallocCrashPool Size=%5u Num=%4i Mem=%8u\n" ), PoolDesc.Size, MaxNumAllocations, AllocatedMemory );
@@ -110,12 +125,11 @@ struct FMallocCrashPool
 
 		// Find an unused allocation.
 		FPtrInfo* PtrInfo = nullptr;
-		for( uint32 Index = 0; Index < MaxNumAllocations; ++Index )
+		for (uint32 Index = 0; Index < MaxNumAllocations; ++Index)
 		{
-			FPtrInfo* PtrIt = Allocations[Index];
-			if( PtrIt->Size == 0 )
+			if(FPtrInfo& PtrIt = Allocations[Index]; PtrIt.Size == 0)
 			{
-				PtrInfo = PtrIt;
+				PtrInfo = &PtrIt;
 				MaxUsedIndex = FMath::Max(MaxUsedIndex,Index);
 				break;
 			}
@@ -131,7 +145,7 @@ struct FMallocCrashPool
 			FMemory::Memset( (void*)PtrInfo->Ptr, FGenericPlatformMallocCrash::MEM_TAG, PtrInfo->Size );
 			DebugVerify();
 
-			//FPlatformMisc::LowLevelOutputDebugStringf( TEXT( "Malloc Size=%u, PooledPtr=0x%016llx \n" ), AllocationSize, (uint64)PtrInfo->Ptr );
+			//FPlatformMisc::LowLevelOutputDebugStringf( TEXT( "Malloc Requested=%u Size=%u, PooledPtr=0x%016llx \n" ), InAllocationSize, AllocationSize, (uint64)PtrInfo->Ptr );
 
 			return PtrInfo->Ptr;
 		}
@@ -151,28 +165,36 @@ struct FMallocCrashPool
 		//const uint32 PtrSize = FGenericPlatformMallocCrash::GetAllocationSize(Ptr);
 		//FPlatformMisc::LowLevelOutputDebugStringf( TEXT( "Free SizeWithOverhead=%u, PooledPtr=0x%016llx \n" ), PtrSize, (uint64)Ptr );
 
-		bool bRemoved = false;
-		for( uint32 Index = 0; Index < MaxNumAllocations; ++Index )
+		int32 Index = int32((Ptr - AllocBase) / AllocationSize);
+		if (Index < 0 || Index >=(int32)MaxNumAllocations)
 		{
-			FPtrInfo* PtrIt = Allocations[Index];
-			if( PtrIt->Ptr == Ptr )
-			{
-				FMemory::Memset( (void*)PtrIt->Ptr, FGenericPlatformMallocCrash::MEM_WIPETAG, PtrIt->Size );
+			FPlatformMisc::LowLevelOutputDebugStringf(TEXT("Failed to free PoolSize=%u, PooledPtr=0x%016llx \n"),
+				AllocationSize, (uint64)Ptr);
 
-				PtrIt->Size = 0;
-				bRemoved = true;
-				NumUsed--;
+			UE_DEBUG_BREAK();
+		}
+		else
+		{
+			FPtrInfo& PtrIt = Allocations[Index];
+			FMemory::Memset((void*)PtrIt.Ptr, FGenericPlatformMallocCrash::MEM_WIPETAG, PtrIt.Size);
 
-				break;
-			}
+			PtrIt.Size = 0;
+			NumUsed--;
 		}
 
-		if( !bRemoved )
+		
+		DebugVerify();
+	}
+
+	uint64 GetAllocationSize(uint8* Ptr) const
+	{
+		SIZE_T Offset = Ptr - AllocBase;
+		int32 Index = int32(Offset / AllocationSize);
+		if (Index < 0 || (uint32)Index >= MaxNumAllocations)
 		{
 			UE_DEBUG_BREAK();
 		}
-
-		DebugVerify();
+		return Allocations[Index].Size;
 	}
 
 	FORCEINLINE static FMallocCrashPool& GetPool(int32 Index)
@@ -181,6 +203,11 @@ struct FMallocCrashPool
 		return ((FMallocCrashPool*)&BytePools)[Index];
 	}
 
+	FORCEINLINE static TConstArrayView<FMallocCrashPool> GetPools()
+	{
+		return TConstArrayView<FMallocCrashPool>((const FMallocCrashPool*)&BytePools, FGenericPlatformMallocCrash::NUM_POOLS);
+	}
+
 private:
 	void DebugVerify()
 	{
@@ -201,32 +228,37 @@ private:
 
 TTypeCompatibleBytes<FMallocCrashPool[FGenericPlatformMallocCrash::NUM_POOLS]> FMallocCrashPool::BytePools;
 
-FGenericPlatformMallocCrash::FGenericPlatformMallocCrash( FMalloc* MainMalloc ) :
-	CrashedThreadId( 0 ),
-	LargeMemoryPoolOffset( 0 ),
-	SmallMemoryPoolOffset( 0 ),
-	PreviousMalloc( MainMalloc )
+
+FGenericPlatformMallocCrash::FGenericPlatformMallocCrash( FMalloc* MainMalloc ) 
+	: SmallMemoryPoolSize(CalculateSmallPoolTotalSize())
+	, BookkeepingPoolSize(CalculateBookkeepingPoolTotalSize())
+	, PreviousMalloc(MainMalloc)
 {
 	LLM_SCOPE(ELLMTag::GenericPlatformMallocCrash);
 	LLM_PLATFORM_SCOPE(ELLMTag::GenericPlatformMallocCrashPlatform);
 
 	const uint32 LargeMemoryPoolSize = Align((int32)LARGE_MEMORYPOOL_SIZE,SafePageSize());
-	LargeMemoryPool = (uint8*)FPlatformMemory::BinnedAllocFromOS( LargeMemoryPoolSize );
-	SmallMemoryPool = (uint8*)FPlatformMemory::BinnedAllocFromOS( (SIZE_T)GetSmallPoolTotalSize() );
+	LargeMemoryPool = (uint8*)FPlatformMemory::BinnedAllocFromOS(LargeMemoryPoolSize);
+	SmallMemoryPool = (uint8*)FPlatformMemory::BinnedAllocFromOS((SIZE_T)SmallMemoryPoolSize);
+	BookkeepingPool = (uint8*)FPlatformMemory::BinnedAllocFromOS((SIZE_T)BookkeepingPoolSize);
 
 	LLM(FLowLevelMemTracker::Get().OnLowLevelAlloc(ELLMTracker::Default, LargeMemoryPool, LargeMemoryPoolSize));
-	LLM(FLowLevelMemTracker::Get().OnLowLevelAlloc(ELLMTracker::Default, SmallMemoryPool, GetSmallPoolTotalSize()));
+	LLM(FLowLevelMemTracker::Get().OnLowLevelAlloc(ELLMTracker::Default, SmallMemoryPool, SmallMemoryPoolSize));
+	LLM(FLowLevelMemTracker::Get().OnLowLevelAlloc(ELLMTracker::Default, BookkeepingPool, BookkeepingPoolSize));
 
 	MemoryTrace_Alloc((uint64)LargeMemoryPool, LargeMemoryPoolSize, alignof(uint8), EMemoryTraceRootHeap::SystemMemory);
-	MemoryTrace_Alloc((uint64)SmallMemoryPool, GetSmallPoolTotalSize(), alignof(uint8), EMemoryTraceRootHeap::SystemMemory);
+	MemoryTrace_Alloc((uint64)SmallMemoryPool, SmallMemoryPoolSize, alignof(uint8), EMemoryTraceRootHeap::SystemMemory);
+	MemoryTrace_Alloc((uint64)BookkeepingPool, BookkeepingPoolSize, alignof(uint8), EMemoryTraceRootHeap::SystemMemory);
 
-	if( !SmallMemoryPool || !LargeMemoryPool )
+	if( !SmallMemoryPool || !LargeMemoryPool || !BookkeepingPool )
 	{
 		FPlatformMisc::LowLevelOutputDebugString( TEXT( "Memory pools allocations failed, exiting...\n" ) );
 		FPlatformMisc::RequestExit(true);
 	}
 
-	if((UPTRINT(LargeMemoryPool) & (REQUIRED_ALIGNMENT - 1)) != 0 || (UPTRINT(SmallMemoryPool) & (REQUIRED_ALIGNMENT - 1)) != 0)
+	if(!IsAligned(LargeMemoryPool, REQUIRED_ALIGNMENT) 
+	|| !IsAligned(SmallMemoryPool, REQUIRED_ALIGNMENT) 
+	|| !IsAligned(BookkeepingPool, REQUIRED_ALIGNMENT))
 	{
 		FPlatformMisc::LowLevelOutputDebugString( TEXT( "OS allocations must be aligned to a value multiple of 16, exiting...\n" ) );
 		FPlatformMisc::RequestExit(true);
@@ -275,7 +307,7 @@ void* FGenericPlatformMallocCrash::Malloc( SIZE_T Size, uint32 Alignment )
 
 	if( IsOnCrashedThread() )
 	{
-		FMallocCrashPool* Pool = FindPoolFromSize( Size32 );
+		FMallocCrashPool* Pool = ChoosePoolForSize( Size32 );
 		if( Pool )
 		{
 			const uint8* PooledPtr = Pool->AllocateFromPool( Size32 );	
@@ -283,7 +315,7 @@ void* FGenericPlatformMallocCrash::Malloc( SIZE_T Size, uint32 Alignment )
 		}
 		else
 		{
-			const uint32 SizeWithOverhead = Size32 + PER_ALLOC_OVERHEAD;
+			const uint32 SizeWithOverhead = Size32 + sizeof(FPtrInfo);
 			LargeMemoryPoolOffset = Align( LargeMemoryPoolOffset, REQUIRED_ALIGNMENT );
 
 			if( LargeMemoryPoolOffset + SizeWithOverhead <= LARGE_MEMORYPOOL_SIZE )
@@ -293,16 +325,16 @@ void* FGenericPlatformMallocCrash::Malloc( SIZE_T Size, uint32 Alignment )
 
 				FPtrInfo* PtrInfo = (FPtrInfo*)(LargeMemoryPool+ReturnMemoryPoolOffset);
 				PtrInfo->Size = Size32;
-				PtrInfo->Ptr = LargeMemoryPool+ReturnMemoryPoolOffset+PER_ALLOC_OVERHEAD;
+				PtrInfo->Ptr = LargeMemoryPool+ReturnMemoryPoolOffset+sizeof(FPtrInfo);
 
 				FPlatformMisc::LowLevelOutputDebugStringf( TEXT( "Malloc Size=%d LargeMemoryPoolOffset=%d \n" ), Size32, LargeMemoryPoolOffset );
 				return (void*)PtrInfo->Ptr;
 			}
 			else
 			{
-				UE_DEBUG_BREAK();
 				FPlatformMisc::LowLevelOutputDebugStringf( TEXT( "MallocCrash run out of memory allocating %u bytes, free %u bytes\n" ), Size32, LARGE_MEMORYPOOL_SIZE-LargeMemoryPoolOffset );
-				FPlatformMisc::LowLevelOutputDebugString( TEXT( "Please increase LARGE_MEMORYPOOL_SIZE, exiting...\n" ) );
+				FPlatformMisc::LowLevelOutputDebugString(TEXT("Please increase LARGE_MEMORYPOOL_SIZE, exiting...\n"));
+				UE_DEBUG_BREAK();
 				FPlatformMisc::RequestExit( true );			
 			}
 		}
@@ -376,7 +408,7 @@ void FGenericPlatformMallocCrash::Free( void* Ptr )
 	{
 		if( IsPtrInSmallPool(Ptr) )
 		{
-			FMallocCrashPool* Pool = FindPoolFromSize( GetAllocationSize(Ptr) );
+			FMallocCrashPool* Pool = FindPoolForAlloc(Ptr);
 			if( Pool )
 			{
 				Pool->TryFreeFromPool( (uint8*)Ptr );
@@ -440,7 +472,7 @@ bool FGenericPlatformMallocCrash::IsPtrInLargePool( void* Ptr ) const
 
 bool FGenericPlatformMallocCrash::IsPtrInSmallPool( void* Ptr ) const
 {
-	const bool bResult = (Ptr >= &SmallMemoryPool[0] && Ptr < &SmallMemoryPool[GetSmallPoolTotalSize()]);
+	const bool bResult = (Ptr >= &SmallMemoryPool[0] && Ptr < &SmallMemoryPool[SmallMemoryPoolSize]);
 	return bResult;
 }
 
@@ -450,7 +482,7 @@ bool FGenericPlatformMallocCrash::GetAllocationSize( void *Original, SIZE_T &Siz
 	return true;
 }
 
-uint32 FGenericPlatformMallocCrash::GetSmallPoolTotalSize() const
+uint32 FGenericPlatformMallocCrash::CalculateSmallPoolTotalSize() const
 {
 	static uint32 TotalSize = 0;
 	if( TotalSize == 0 )
@@ -460,7 +492,7 @@ uint32 FGenericPlatformMallocCrash::GetSmallPoolTotalSize() const
 			const FPoolDesc& PoolDesc = GetPoolDesc(Index);
 			check(PoolDesc.NumAllocs%16==0);
 			check(PoolDesc.Size%16==0);
-			TotalSize += PoolDesc.NumAllocs*(PoolDesc.Size+PER_ALLOC_OVERHEAD);
+			TotalSize += PoolDesc.NumAllocs*PoolDesc.Size;
 		}
 
 		TotalSize = Align(TotalSize,SafePageSize());
@@ -468,6 +500,25 @@ uint32 FGenericPlatformMallocCrash::GetSmallPoolTotalSize() const
 	return TotalSize;
 }
 
+
+uint32 FGenericPlatformMallocCrash::CalculateBookkeepingPoolTotalSize() const
+{
+	static uint32 TotalSize = 0;
+	if (TotalSize == 0)
+	{
+		for (uint32 Index = 0; Index < NUM_POOLS; ++Index)
+		{
+			const FPoolDesc& PoolDesc = GetPoolDesc(Index);
+			check(PoolDesc.NumAllocs % 16 == 0);
+			check(PoolDesc.Size % 16 == 0);
+			TotalSize += PoolDesc.NumAllocs * sizeof(FPtrInfo);
+		}
+
+		TotalSize = Align(TotalSize, SafePageSize());
+	}
+	return TotalSize;
+}
+
 void FGenericPlatformMallocCrash::InitializeSmallPools()
 {
 	uint32 MallocCrashOverhead = 0;
@@ -478,22 +529,22 @@ void FGenericPlatformMallocCrash::InitializeSmallPools()
 		MallocCrashOverhead += NewPool->AllocatedMemory;
 	}
 
-	check(SmallMemoryPoolOffset<=GetSmallPoolTotalSize());
+	check(SmallMemoryPoolOffset<=CalculateSmallPoolTotalSize());
 }
 
-FMallocCrashPool* FGenericPlatformMallocCrash::FindPoolFromSize( uint32 AllocationSize ) const
+FMallocCrashPool* FGenericPlatformMallocCrash::ChoosePoolForSize( uint32 AllocationSize ) const
 {
 	for( uint32 Index = 0; Index < NUM_POOLS; ++Index )
 	{
 		FMallocCrashPool* Pool = &FMallocCrashPool::GetPool(Index);
 
 		// If we have used up this pools allocation limit lets try the next one
-		if ( Pool->NumUsed >= Pool->MaxNumAllocations )
+		if (Pool->NumUsed >= Pool->MaxNumAllocations)
 		{
 			continue;
 		}
 
-		if( AllocationSize <= Pool->AllocationSize-PER_ALLOC_OVERHEAD )
+		if (AllocationSize <= Pool->AllocationSize)
 		{
 			return Pool;
 		}
@@ -503,25 +554,53 @@ FMallocCrashPool* FGenericPlatformMallocCrash::FindPoolFromSize( uint32 Allocati
 	return nullptr;
 }
 
+FMallocCrashPool* FGenericPlatformMallocCrash::FindPoolForAlloc(void* Ptr) const
+{
+	TArrayView<const FMallocCrashPool> Pools = FMallocCrashPool::GetPools();
+	int32 IndexBeyond = Algo::UpperBoundBy(Pools, (uint8*)Ptr, [](const FMallocCrashPool& Pool) { return Pool.AllocBase; });
+	return &FMallocCrashPool::GetPool(IndexBeyond - 1);
+}
+
 uint8* FGenericPlatformMallocCrash::AllocateFromSmallPool( uint32 AllocationSize )
 {
-	if( SmallMemoryPoolOffset + AllocationSize <= GetSmallPoolTotalSize() )
+	if( SmallMemoryPoolOffset + AllocationSize <= SmallMemoryPoolSize )
 	{
 		const uint32 ReturnMemoryPoolOffset = SmallMemoryPoolOffset;
 		SmallMemoryPoolOffset += AllocationSize;
 		return (uint8*)SmallMemoryPool+ReturnMemoryPoolOffset;
 	}
-	else
+
+	check(0);
+	return nullptr;
+}
+
+uint8* FGenericPlatformMallocCrash::AllocateFromBookkeeping(uint32 AllocationSize)
+{
+	if (BookkeepingPoolOffset + AllocationSize <= BookkeepingPoolSize)
 	{
-		check(0);
+		const uint32 ReturnMemoryPoolOffset = BookkeepingPoolOffset;
+		BookkeepingPoolOffset += AllocationSize;
+		return (uint8*)BookkeepingPool + ReturnMemoryPoolOffset;
 	}
+
+	check(0);
 	return nullptr;
 }
 
-uint32 FGenericPlatformMallocCrash::GetAllocationSize( void *Original )
+uint64 FGenericPlatformMallocCrash::GetAllocationSize(void *Original)
 {
-	FPtrInfo* PtrInfo = (FPtrInfo*)((uint8*)Original-PER_ALLOC_OVERHEAD);
-	return (uint32)(PtrInfo->Size);
+	if (IsPtrInSmallPool(Original))
+	{
+		FMallocCrashPool* Pool = FindPoolForAlloc(Original);
+		return Pool->GetAllocationSize((uint8*)Original);
+	}
+	else if (IsPtrInLargePool(Original))
+	{
+		FPtrInfo* Info = (FPtrInfo*)((uint8*)Original - sizeof(FPtrInfo));
+		return Info->Size;
+	}
+	
+	return 0;
 }
 
 uint32 FGenericPlatformMallocCrash::SafePageSize()
diff --git a/Engine/Source/Runtime/Core/Public/GenericPlatform/GenericPlatformMallocCrash.h b/Engine/Source/Runtime/Core/Public/GenericPlatform/GenericPlatformMallocCrash.h
--- a/Engine/Source/Runtime/Core/Public/GenericPlatform/GenericPlatformMallocCrash.h
+++ b/Engine/Source/Runtime/Core/Public/GenericPlatform/GenericPlatformMallocCrash.h
@@ -45,12 +45,10 @@ struct FGenericPlatformMallocCrash final : public FMalloc
 	friend struct FPoolDesc;
 	friend struct FMallocCrashPool;
 
-private:
 	enum
 	{
 		LARGE_MEMORYPOOL_SIZE = 2 * 1024 * 1024,
 		REQUIRED_ALIGNMENT = 16,
-		PER_ALLOC_OVERHEAD = sizeof(FPtrInfo),
 
 		NUM_POOLS = 14,
 		MAX_NUM_ALLOCS_IN_POOL = 2048,
@@ -114,13 +112,19 @@ protected:
 	bool IsPtrInSmallPool(void* Ptr) const;
 
 	const FPoolDesc& GetPoolDesc( uint32 Index ) const;
-	uint32 GetSmallPoolTotalSize() const;
+	uint32 CalculateSmallPoolTotalSize() const;
+	uint32 CalculateBookkeepingPoolTotalSize() const;
 
 	void InitializeSmallPools();
-	FMallocCrashPool* FindPoolFromSize( uint32 AllocationSize ) const;
+
+	// Choose a pool to make a new allocation from of the given size - may use a larger pool than necessary if small pools are exhausted.
+	FMallocCrashPool* ChoosePoolForSize( uint32 AllocationSize ) const;
+	// Find the pool that an allocation was made from if any.
+	FMallocCrashPool* FindPoolForAlloc(void* Ptr) const;
+	uint8* AllocateFromBookkeeping(uint32 AllocationSize);
 	uint8* AllocateFromSmallPool( uint32 AllocationSize );
 
-	static uint32 GetAllocationSize(void *Original);
+	uint64 GetAllocationSize(void *Original);
 
 	/**
 	 * @return page size, if page size is not initialized, returns 64k.
@@ -132,22 +136,30 @@ protected:
 	FCriticalSection InternalLock;
 
 	/** The id of the thread that crashed. */
-	int32 CrashedThreadId;
+	int32 CrashedThreadId = 0;
 
 	/** Preallocated memory pool for large allocations. */
-	uint8* LargeMemoryPool;
+	uint8* LargeMemoryPool = nullptr;
 
 	/** Current position in the large memory pool. */
-	uint32 LargeMemoryPoolOffset;
+	uint32 LargeMemoryPoolOffset = 0;
 
 	/** Preallocated memory pool for small allocations. */
-	uint8* SmallMemoryPool;
+	uint8* SmallMemoryPool = nullptr;
 
 	/** Current position in the small memory pool. */
-	uint32 SmallMemoryPoolOffset;
+	uint32 SmallMemoryPoolOffset = 0;
+	uint32 SmallMemoryPoolSize = 0;
+
+	/** Preallocated memory pool for bookkeeping. */
+	uint8* BookkeepingPool = nullptr;
+
+	/** Current position in the bookkeeping pool. */
+	uint32 BookkeepingPoolOffset = 0;
+	uint32 BookkeepingPoolSize = 0;
 
 	/** Previously used malloc. */
-	FMalloc* PreviousMalloc;
+	FMalloc* PreviousMalloc = nullptr;
 };
 
 
