use non-inlined GetActiveTasks() in modular builds because ActiveTask can't be dllexport if it is also thread_local.

#rnx

diff --git a/Engine/Source/Runtime/Core/Private/Async/Fundamental/Scheduler.cpp b/Engine/Source/Runtime/Core/Private/Async/Fundamental/Scheduler.cpp
--- a/Engine/Source/Runtime/Core/Private/Async/Fundamental/Scheduler.cpp
+++ b/Engine/Source/Runtime/Core/Private/Async/Fundamental/Scheduler.cpp
@@ -285,7 +285,7 @@ namespace LowLevelTasks
 		}
 	}
 
-#if PLATFORM_DESKTOP
+#if PLATFORM_DESKTOP || !IS_MONOLITHIC
 	const FTask* FTask::GetActiveTask()
 	{
 		return ActiveTask;
diff --git a/Engine/Source/Runtime/Core/Public/Async/Fundamental/Task.h b/Engine/Source/Runtime/Core/Public/Async/Fundamental/Task.h
--- a/Engine/Source/Runtime/Core/Public/Async/Fundamental/Task.h
+++ b/Engine/Source/Runtime/Core/Public/Async/Fundamental/Task.h
@@ -357,7 +357,7 @@ namespace LowLevelTasks
 			return !EnumHasAnyFlags(State, ~ETaskState::CanceledFlag); 
 		}
 
-#if PLATFORM_DESKTOP
+#if PLATFORM_DESKTOP || !IS_MONOLITHIC
 		/*
 		* get the currently active task if any
 		*/
