added `TTask::IsAwaitable()` to detect if the task can be waited for from the current thread. The only case that is detected right now is a waiting from inside a task that leads to a deadlock. Required by parallel animation evaluation to detect that a code is executed inside a task during task retraction.

diff --git a/Engine/Source/Runtime/Core/Private/Tests/Tasks/TasksTest.cpp b/Engine/Source/Runtime/Core/Private/Tests/Tasks/TasksTest.cpp
--- a/Engine/Source/Runtime/Core/Private/Tests/Tasks/TasksTest.cpp
+++ b/Engine/Source/Runtime/Core/Private/Tests/Tasks/TasksTest.cpp
@@ -322,6 +322,37 @@ namespace UE { namespace TasksTests
 			Task.Wait();
 		}
 
+		{	// test basic functionality of TTask::IsAwaitable()
+			FTask Task;
+			Task.Launch(UE_SOURCE_LOCATION,
+				[&Task] 
+				{
+					check(!Task.IsAwaitable()); // Task.Wait() would deadlock if called here inside its execution
+				}
+			);
+			check(Task.IsAwaitable());
+			Task.Wait();
+		}
+
+
+		{	// TTask::IsAwaitable() returns false when called from its inner task execution
+			FTask Outer;
+			Outer.Launch(UE_SOURCE_LOCATION,
+				[&Outer]
+				{
+					auto InnerTaskBody = [&Outer] { check(!Outer.IsAwaitable()); /* still inside `Outer` execution */ };
+
+					// using the private API for `EExtendedTaskPriority::Inline` to force inner task execution inside the outer task
+					using namespace UE::Tasks::Private;
+					using FInnerTask = TExecutableTask<decltype(InnerTaskBody)>;
+
+					FInnerTask InnerTask{ UE_SOURCE_LOCATION, MoveTemp(InnerTaskBody), ETaskPriority::Default, EExtendedTaskPriority::Inline };
+					verify(InnerTask.TryLaunch()); // blocking execution of InnerTask due to its "inline" priority
+				}
+			);
+			Outer.Wait();
+		}
+
 		//{	// Cancelled task is completed only after nested
 		//	FTaskEvent ParentBlocker{ UE_SOURCE_LOCATION };
 		//	FTaskEvent NestedBlocker{ UE_SOURCE_LOCATION };
diff --git a/Engine/Source/Runtime/Core/Public/Tasks/Task.h b/Engine/Source/Runtime/Core/Public/Tasks/Task.h
--- a/Engine/Source/Runtime/Core/Public/Tasks/Task.h
+++ b/Engine/Source/Runtime/Core/Public/Tasks/Task.h
@@ -135,7 +135,7 @@ namespace UE::Tasks
 			void Launch(
 				const TCHAR* DebugName,
 				TaskBodyType&& TaskBody,
-				LowLevelTasks::ETaskPriority Priority = LowLevelTasks::ETaskPriority::Normal
+				ETaskPriority Priority = ETaskPriority::Normal
 			)
 			{
 				check(!IsValid());
@@ -159,7 +159,7 @@ namespace UE::Tasks
 				const TCHAR* DebugName,
 				TaskBodyType&& TaskBody,
 				PrerequisitesCollectionType&& Prerequisites,
-				LowLevelTasks::ETaskPriority Priority = LowLevelTasks::ETaskPriority::Normal
+				ETaskPriority Priority = ETaskPriority::Normal
 			)
 			{
 				check(!IsValid());
@@ -172,6 +172,11 @@ namespace UE::Tasks
 				Task->TryLaunch();
 			}
 
+			bool IsAwaitable() const
+			{
+				return IsValid() && Pimpl->IsAwaitable();
+			}
+
 		protected:
 			TRefCountPtr<FTaskBase> Pimpl;
 		};
@@ -261,7 +266,7 @@ namespace UE::Tasks
 	TTask<TInvokeResult_T<TaskBodyType>> Launch(
 		const TCHAR* DebugName,
 		TaskBodyType&& TaskBody,
-		LowLevelTasks::ETaskPriority Priority = LowLevelTasks::ETaskPriority::Normal
+		ETaskPriority Priority = ETaskPriority::Normal
 	)
 	{
 		using FResult = TInvokeResult_T<TaskBodyType>;
@@ -282,7 +287,7 @@ namespace UE::Tasks
 		const TCHAR* DebugName,
 		TaskBodyType&& TaskBody,
 		PrerequisitesCollectionType&& Prerequisites,
-		LowLevelTasks::ETaskPriority Priority = LowLevelTasks::ETaskPriority::Normal
+		ETaskPriority Priority = ETaskPriority::Normal
 	)
 	{
 		using FResult = TInvokeResult_T<TaskBodyType>;
diff --git a/Engine/Source/Runtime/Core/Public/Tasks/TaskPrivate.h b/Engine/Source/Runtime/Core/Public/Tasks/TaskPrivate.h
--- a/Engine/Source/Runtime/Core/Public/Tasks/TaskPrivate.h
+++ b/Engine/Source/Runtime/Core/Public/Tasks/TaskPrivate.h
@@ -118,6 +118,13 @@ namespace UE::Tasks
 			virtual bool TryExecuteTask() = 0;
 
 		public:
+			// returns true if it's valid to wait for the task completion.
+			// it's not valid to wait for a task e.g. from inside task's execution, as this would deadlock
+			bool IsAwaitable() const
+			{
+				return FPlatformTLS::GetCurrentThreadId() != ExecutingThreadId.load(std::memory_order_relaxed);
+			}
+
 			EExtendedTaskPriority GetExtendedPriority() const
 			{
 				return ExtendedPriority;
@@ -457,6 +464,7 @@ namespace UE::Tasks
 				ReleasePrerequisites();
 
 				FTaskBase* PrevTask = ExchangeCurrentTask(this);
+				ExecutingThreadId.store(FPlatformTLS::GetCurrentThreadId(), std::memory_order_relaxed);
 
 				if (GetPipe() != nullptr)
 				{
@@ -470,6 +478,8 @@ namespace UE::Tasks
 					FinishPipeExecution();
 				}
 
+				ExecutingThreadId.store(FThread::InvalidThreadId, std::memory_order_relaxed); // no need to sync with loads as they matter only if
+				// executed by the same thread
 				ExchangeCurrentTask(PrevTask);
 
 				// close the task if there are no pending nested tasks
@@ -657,6 +667,8 @@ namespace UE::Tasks
 			TLockFreePointerListUnordered<FTaskBase, 0> Prerequisites;
 
 			FPipe* Pipe{ nullptr };
+
+			std::atomic<uint32> ExecutingThreadId = FThread::InvalidThreadId;
 		};
 
 		// an extension of FTaskBase for tasks that return a result.
