tasks integration: Tasks System integrated into the old TaskGraph API as an intermediary step into transition to the new API. Disabled by default. The only remaining part of the old TaskGraph is named threads implementation that is not directly supported by the new Tasks System.

also includes the fix for parallel animation evaluation to work correctly with task retraction that is globally enabled by these change. animation fix was reviewed by [REDACTED_PERSON]

diff --git a/Engine/Source/Runtime/Core/Private/Async/TaskGraph.cpp b/Engine/Source/Runtime/Core/Private/Async/TaskGraph.cpp
--- a/Engine/Source/Runtime/Core/Private/Async/TaskGraph.cpp
+++ b/Engine/Source/Runtime/Core/Private/Async/TaskGraph.cpp
@@ -775,7 +775,7 @@ PRAGMA_ENABLE_DEPRECATION_WARNINGS
 	virtual void EnqueueFromThisThread(int32 QueueIndex, FBaseGraphTask* Task) override
 	{
 		checkThreadGraph(Task && Queue(QueueIndex).StallRestartEvent); // make sure we are started up
-		uint32 PriIndex = ENamedThreads::GetTaskPriority(Task->ThreadToExecuteOn) ? 0 : 1;
+		uint32 PriIndex = ENamedThreads::GetTaskPriority(Task->GetThreadToExecuteOn()) ? 0 : 1;
 		int32 ThreadToStart = Queue(QueueIndex).StallQueue.Push(Task, PriIndex);
 		check(ThreadToStart < 0); // if I am stalled, then how can I be queueing a task?
 	}
@@ -809,7 +809,7 @@ PRAGMA_ENABLE_DEPRECATION_WARNINGS
 		TestRandomizedThreads();
 		checkThreadGraph(Task && Queue(QueueIndex).StallRestartEvent); // make sure we are started up
 
-		uint32 PriIndex = ENamedThreads::GetTaskPriority(Task->ThreadToExecuteOn) ? 0 : 1;
+		uint32 PriIndex = ENamedThreads::GetTaskPriority(Task->GetThreadToExecuteOn()) ? 0 : 1;
 		int32 ThreadToStart = Queue(QueueIndex).StallQueue.Push(Task, PriIndex);
 
 		if (ThreadToStart >= 0)
@@ -1420,10 +1420,10 @@ public:
 	{
 #if STATS
 PRAGMA_DISABLE_DEPRECATION_WARNINGS
-		if (ENamedThreads::GetThreadIndex(Task->ThreadToExecuteOn) == ENamedThreads::StatsThread)
+		if (ENamedThreads::GetThreadIndex(Task->GetThreadToExecuteOn()) == ENamedThreads::StatsThread)
 PRAGMA_ENABLE_DEPRECATION_WARNINGS
 		{
-			checkf(ENamedThreads::GetQueueIndex(Task->ThreadToExecuteOn) == ENamedThreads::MainQueue, TEXT("`StatsThread_Local` is not supported and unlikely to be used intentionally. Check that `Task->ThreadToExecuteOn` value is not rubbish: %d"), Task->ThreadToExecuteOn);
+			checkf(ENamedThreads::GetQueueIndex(Task->GetThreadToExecuteOn()) == ENamedThreads::MainQueue, TEXT("`StatsThread_Local` is not supported and unlikely to be used intentionally. Check that `Task->ThreadToExecuteOn` value is not rubbish: %d"), Task->GetThreadToExecuteOn());
 
 			extern CORE_API UE::Tasks::FPipe GStatsPipe;
 			GStatsPipe.Launch(UE_SOURCE_LOCATION,
@@ -1446,7 +1446,7 @@ PRAGMA_ENABLE_DEPRECATION_WARNINGS
 	static void RedirectAudioTasksToPipe(FBaseGraphTask* Task)
 	{
 PRAGMA_DISABLE_DEPRECATION_WARNINGS
-		if (ENamedThreads::GetThreadIndex(Task->ThreadToExecuteOn) == ENamedThreads::AudioThread)
+		if (ENamedThreads::GetThreadIndex(Task->GetThreadToExecuteOn()) == ENamedThreads::AudioThread)
 PRAGMA_ENABLE_DEPRECATION_WARNINGS
 		{
 			extern CORE_API UE::Tasks::FPipe GAudioPipe;
@@ -1488,8 +1488,8 @@ PRAGMA_ENABLE_DEPRECATION_WARNINGS
 			TASKGRAPH_SCOPE_CYCLE_COUNTER(3, STAT_TaskGraph_QueueTask_AnyThread);
 			if (FTaskGraphInterface::IsMultithread())
 			{
-				uint32 TaskPriority = ENamedThreads::GetTaskPriority(Task->ThreadToExecuteOn);
-				int32 Priority = ENamedThreads::GetThreadPriorityIndex(Task->ThreadToExecuteOn);
+				uint32 TaskPriority = ENamedThreads::GetTaskPriority(Task->GetThreadToExecuteOn());
+				int32 Priority = ENamedThreads::GetThreadPriorityIndex(Task->GetThreadToExecuteOn());
 				if (Priority == (ENamedThreads::BackgroundThreadPriority >> ENamedThreads::ThreadPriorityShift) && (!bCreatedBackgroundPriorityThreads || !ENamedThreads::bHasBackgroundThreads))
 				{
 					Priority = ENamedThreads::NormalThreadPriority >> ENamedThreads::ThreadPriorityShift; // we don't have background threads, promote to normal
@@ -2021,6 +2021,9 @@ private:
 
 		if (ENamedThreads::GetThreadIndex(InThreadToExecuteOn) == ENamedThreads::AnyThread)
 		{
+#if TASKGRAPH_NEW_FRONTEND
+			checkNoEntry();
+#else
 			uint32 ThreadPriority = GetThreadPriorityIndex(InThreadToExecuteOn);
 			check(ThreadPriority < uint32(ENamedThreads::NumThreadPriorities));
 			LowLevelTasks::ETaskPriority Conversion[int(ENamedThreads::NumThreadPriorities)] = { LowLevelTasks::ETaskPriority::Normal, LowLevelTasks::ETaskPriority::High, LowLevelTasks::ETaskPriority::BackgroundNormal };
@@ -2031,14 +2034,16 @@ private:
 				Priority = LowLevelTasks::ETaskPriority::BackgroundHigh;
 			}
 
-			Task->TaskHandle.Init(TEXT("TaskGraphTask"), Priority, [Task, InThreadToExecuteOn, Deleter(LowLevelTasks::TDeleter<FBaseGraphTask, &FBaseGraphTask::DeleteTask>(Task))]()
+			Task->GetTaskHandle().Init(TEXT("TaskGraphTask"), Priority, [Task, InThreadToExecuteOn, Deleter(LowLevelTasks::TDeleter<FBaseGraphTask, &FBaseGraphTask::DeleteTask>(Task))]()
 			{
 				Task->Execute(NewTasks, InThreadToExecuteOn, false);
 			});
 	
 			bWakeUpWorker |= LowLevelTasks::FSchedulerTls::IsBusyWaiting();
-			verifySlow(LowLevelTasks::TryLaunch(Task->TaskHandle, bWakeUpWorker ? LowLevelTasks::EQueuePreference::GlobalQueuePreference : LowLevelTasks::EQueuePreference::LocalQueuePreference, bWakeUpWorker));
+
+			verifySlow(LowLevelTasks::TryLaunch(Task->GetTaskHandle(), bWakeUpWorker ? LowLevelTasks::EQueuePreference::GlobalQueuePreference : LowLevelTasks::EQueuePreference::LocalQueuePreference, bWakeUpWorker));
 			return;
+#endif
 		}
 
 		ENamedThreads::Type CurrentThreadIfKnown;
@@ -2134,6 +2139,26 @@ private:
 	{
 		TaskTrace::FWaitingScope WaitingScope(GetTraceIds(Tasks));
 		TRACE_CPUPROFILER_EVENT_SCOPE(WaitUntilTasksComplete);
+
+#if TASKGRAPH_NEW_FRONTEND
+		bool HasNamedThreadTasks = false;
+		for (const FGraphEventRef& Task : Tasks)
+		{
+			if (!Task->IsNamedThreadTask())
+			{
+				Task->Wait();
+			}
+			else
+			{
+				HasNamedThreadTasks = true;
+			}
+		}
+
+		if (!HasNamedThreadTasks)
+		{
+			return;
+		}
+#endif
 	
 		ENamedThreads::Type CurrentThread = CurrentThreadIfKnown;
 		if (ENamedThreads::GetThreadIndex(CurrentThreadIfKnown) == ENamedThreads::AnyThread)
@@ -2428,6 +2453,11 @@ bool FTaskGraphInterface::IsMultithread()
 	return FPlatformProcess::SupportsMultithreading() || (FForkProcessHelper::IsForkedMultithreadInstance() && GAllowTaskGraphForkMultithreading);
 }
 
+
+// Statics and some implementations from FBaseGraphTask and FGraphEvent
+
+#if !TASKGRAPH_NEW_FRONTEND
+
 #if !(UE_BUILD_SHIPPING || UE_BUILD_TEST)
 void FBaseGraphTask::LogPossiblyInvalidSubsequentsTask(const TCHAR* TaskName)
 {
@@ -2526,6 +2556,8 @@ FGraphEvent::~FGraphEvent()
 	CheckDontCompleteUntilIsEmpty(); // We should not have any wait untils outstanding
 }
 
+#endif
+
 DECLARE_CYCLE_STAT(TEXT("FBroadcastTask"), STAT_FBroadcastTask, STATGROUP_TaskGraphTasks);
 
 static int32 GPrintBroadcastWarnings = true;
diff --git a/Engine/Source/Runtime/Core/Private/Tasks/TaskPrivate.cpp b/Engine/Source/Runtime/Core/Private/Tasks/TaskPrivate.cpp
--- a/Engine/Source/Runtime/Core/Private/Tasks/TaskPrivate.cpp
+++ b/Engine/Source/Runtime/Core/Private/Tasks/TaskPrivate.cpp
@@ -3,8 +3,43 @@
 #include "Tasks/TaskPrivate.h"
 #include "Tasks/Pipe.h"
 
+#if TASKGRAPH_NEW_FRONTEND
+#include "Async/TaskGraphInterfaces.h"
+#endif
+
 namespace UE::Tasks::Private
 {
+	void FTaskBase::Schedule()
+	{
+#if TASKGRAPH_NEW_FRONTEND
+		if (IsNamedThreadTask())
+		{
+			ENamedThreads::Type ConversionMap[] =
+			{
+				ENamedThreads::GameThread,
+				(ENamedThreads::Type)(ENamedThreads::GameThread | ENamedThreads::HighTaskPriority),
+				(ENamedThreads::Type)(ENamedThreads::GameThread | ENamedThreads::LocalQueue),
+				(ENamedThreads::Type)(ENamedThreads::GameThread | ENamedThreads::HighTaskPriority | ENamedThreads::LocalQueue),
+
+				ENamedThreads::GetRenderThread(),
+				(ENamedThreads::Type)(ENamedThreads::GetRenderThread() | ENamedThreads::HighTaskPriority),
+				(ENamedThreads::Type)(ENamedThreads::GetRenderThread() | ENamedThreads::LocalQueue),
+				(ENamedThreads::Type)(ENamedThreads::GetRenderThread() | ENamedThreads::HighTaskPriority | ENamedThreads::LocalQueue),
+
+				ENamedThreads::RHIThread,
+				(ENamedThreads::Type)(ENamedThreads::RHIThread | ENamedThreads::HighTaskPriority),
+				(ENamedThreads::Type)(ENamedThreads::RHIThread | ENamedThreads::LocalQueue),
+				(ENamedThreads::Type)(ENamedThreads::RHIThread | ENamedThreads::HighTaskPriority | ENamedThreads::LocalQueue)
+			};
+
+			FTaskGraphInterface::Get().QueueTask(static_cast<FBaseGraphTask*>(this), true, ConversionMap[(int32)ExtendedPriority - (int32)EExtendedTaskPriority::GameThreadNormalPri]);
+			return;
+		}
+#endif
+
+		LowLevelTasks::FScheduler::Get().TryLaunch(LowLevelTask, LowLevelTasks::EQueuePreference::GlobalQueuePreference, /*bWakeUpWorker=*/ true);
+	}
+
 	FTaskBase* FTaskBase::TryPushIntoPipe()
 	{
 		return GetPipe()->PushIntoPipe(*this);
diff --git a/Engine/Source/Runtime/Core/Public/Async/TaskGraphInterfaces.h b/Engine/Source/Runtime/Core/Public/Async/TaskGraphInterfaces.h
--- a/Engine/Source/Runtime/Core/Public/Async/TaskGraphInterfaces.h
+++ b/Engine/Source/Runtime/Core/Public/Async/TaskGraphInterfaces.h
@@ -28,19 +28,30 @@
 #include "Async/Fundamental/Task.h"
 
 #include "Async/TaskTrace.h"
+#include "Tasks/TaskPrivate.h"
 
 #if !defined(STATS)
 #error "STATS must be defined as either zero or one."
 #endif
 
+// what level of checking to perform...normally checkSlow but could be ensure or check
+#define checkThreadGraph checkSlow
 
+#if TASKGRAPH_NEW_FRONTEND
 
+class FBaseGraphTask;
 
+using FGraphEvent = FBaseGraphTask;
+using FGraphEventRef = TRefCountPtr<FBaseGraphTask>;
+
+#else
+
+/** Convenience typedef for a reference counted pointer to a graph event **/
+typedef TRefCountPtr<class FGraphEvent> FGraphEventRef;
+
+#endif
 
-// what level of checking to perform...normally checkSlow but could be ensure or check
-#define checkThreadGraph checkSlow
 
-class FGraphEvent;
 
 //#define checkThreadGraph(x) ((x)||((*(char*)3) = 0))
 
@@ -294,9 +305,6 @@ namespace ESubsequentsMode
 	};
 }
 
-/** Convenience typedef for a reference counted pointer to a graph event **/
-typedef TRefCountPtr<class FGraphEvent> FGraphEventRef;
-
 /** Convenience typedef for a an array a graph events **/
 typedef TArray<FGraphEventRef, TInlineAllocator<4> > FGraphEventArray;
 
@@ -307,7 +315,12 @@ CORE_API TArray<TaskTrace::FId> GetTraceIds(const FGraphEventArray& Tasks);
 class FTaskGraphInterface
 {
 	friend class FBaseGraphTask;
-	/** 
+
+#if TASKGRAPH_NEW_FRONTEND
+	friend UE::Tasks::Private::FTaskBase;
+#endif
+
+	/**
 	 *	Internal function to queue a task
 	 *	@param	Task; the task to queue
 	 *	@param	ThreadToExecuteOn; Either a named thread for a threadlocked task or ENamedThreads::AnyThread for a task that is to run on a worker thread
@@ -461,6 +474,283 @@ struct FTaskGraphBlockAllocationTag : FDefaultBlockAllocationTag
 	using Allocator = TBlockAllocationCache<BlockSize, FAlignedAllocator>;
 };
 
+#if TASKGRAPH_NEW_FRONTEND
+
+/** 
+ *	Base class for all tasks. A replacement for `FBaseGraphTask` and `FGraphEvent` from the old API, based on `Tasks::Private::FTaskBase` functionality
+ **/
+
+class FBaseGraphTask : public UE::Tasks::Private::FTaskBase
+{
+public:
+	explicit FBaseGraphTask(const FGraphEventArray* InPrerequisites)
+		: FTaskBase(/*InitRefCount=*/ 1)
+	{		
+		if (InPrerequisites != nullptr)
+		{
+			for (const FGraphEventRef& Prereq : *InPrerequisites)
+			{
+				AddPrerequisites(*Prereq);
+			}
+		}
+	}
+
+	void Init(const TCHAR* InDebugName, UE::Tasks::ETaskPriority InPriority, UE::Tasks::Private::EExtendedTaskPriority InExtendedPriority)
+	{
+		FTaskBase::Init(InDebugName, InPriority, InExtendedPriority);
+	}
+
+	void Unlock(ENamedThreads::Type CurrentThreadIfKnown = ENamedThreads::AnyThread)
+	{
+		TryLaunch();
+	}
+
+	void Execute(TArray<FBaseGraphTask*>& NewTasks, ENamedThreads::Type CurrentThread, bool bDeleteOnCompletion)
+	{	// only called for named thread tasks, normal tasks are executed using `FTaskBase` API directly (see `TGraphTask`)
+		checkSlow(NewTasks.Num() == 0);
+		checkSlow(bDeleteOnCompletion);
+		checkSlow(IsNamedThreadTask());
+		verify(TryExecuteTask());
+		ReleaseInternalReference(); // named tasks are executed by named threads, outside of the scheduler
+	}
+
+	FGraphEventRef GetCompletionEvent()
+	{
+		return this;
+	}
+
+	void DontCompleteUntil(FGraphEventRef NestedTask)
+	{
+		checkSlow(UE::Tasks::Private::GetCurrentTask() == this); // a nested task can be added only from inside of parent's execution
+		AddNested(*NestedTask);
+	}
+
+	bool IsComplete() const
+	{
+		return IsCompleted(); // the new API uses a slightly different name
+	}
+
+	static FGraphEventRef CreateGraphEvent();
+
+	void DispatchSubsequents(ENamedThreads::Type CurrentThreadIfKnown = ENamedThreads::AnyThread)
+	{
+		AddRef(); // scheduler's reference
+		TryLaunch();
+	}
+
+	void DispatchSubsequents(TArray<FBaseGraphTask*>& NewTasks, ENamedThreads::Type CurrentThreadIfKnown = ENamedThreads::AnyThread)
+	{
+		check(NewTasks.Num() == 0); // the feature is not used
+		DispatchSubsequents();
+	}
+
+	void SetDebugName(const TCHAR* InDebugName)
+	{	// incompatible with the new API that requires debug name during task construction and doesn't allow to set it later. "debug name" feature was added
+		// to the old API recently, is used only in a couple of places and will be fixed manually by switching to the new API
+	}
+
+	void Wait(ENamedThreads::Type CurrentThreadIfKnown = ENamedThreads::AnyThread)
+	{
+		if (!IsNamedThreadTask())
+		{
+			FTaskBase::Wait();
+		}
+		else
+		{
+			FTaskGraphInterface::Get().WaitUntilTaskCompletes(this, CurrentThreadIfKnown);
+		}
+	}
+
+	ENamedThreads::Type GetThreadToExecuteOn() const
+	{
+		return TranslatePriority(GetExtendedPriority());
+	}
+
+	// task priority translation from the old API to the new API
+	static void TranslatePriority(ENamedThreads::Type ThreadType, UE::Tasks::ETaskPriority& OutPriority, UE::Tasks::Private::EExtendedTaskPriority& OutExtendedPriority)
+	{
+		using namespace UE::Tasks;
+
+		ENamedThreads::Type ThreadIndex = ENamedThreads::GetThreadIndex(ThreadType);
+		if (ThreadIndex != ENamedThreads::AnyThread)
+		{
+			check(ThreadIndex == ENamedThreads::GameThread || ThreadIndex == ENamedThreads::GetRenderThread() || ThreadIndex == ENamedThreads::RHIThread);
+			Private::EExtendedTaskPriority ConversionMap[] =
+			{
+				Private::EExtendedTaskPriority::RHIThreadNormalPri,
+				Private::EExtendedTaskPriority::None, // invalid, maps from ENamedThreads::AudioThread
+				Private::EExtendedTaskPriority::GameThreadNormalPri,
+				Private::EExtendedTaskPriority::RenderThreadNormalPri
+			};
+			OutExtendedPriority = ConversionMap[ThreadIndex - ENamedThreads::RHIThread];
+			OutExtendedPriority = (Private::EExtendedTaskPriority)((int32)OutExtendedPriority + (ENamedThreads::GetTaskPriority(ThreadType) != ENamedThreads::NormalTaskPriority ? 1 : 0));
+			OutExtendedPriority = (Private::EExtendedTaskPriority)((int32)OutExtendedPriority + (ENamedThreads::GetQueueIndex(ThreadType) != ENamedThreads::MainQueue ? 2 : 0));
+			OutPriority = ETaskPriority::Count;
+		}
+		else
+		{
+			OutExtendedPriority = Private::EExtendedTaskPriority::None;
+			uint32 ThreadPriority = GetThreadPriorityIndex(ThreadType);
+			check(ThreadPriority < uint32(ENamedThreads::NumThreadPriorities));
+			ETaskPriority ConversionMap[int(ENamedThreads::NumThreadPriorities)] = { ETaskPriority::Normal, ETaskPriority::High, ETaskPriority::BackgroundNormal };
+			OutPriority = ConversionMap[ThreadPriority];
+		}
+
+		if (OutPriority == ETaskPriority::BackgroundNormal && GetTaskPriority(ThreadType))
+		{
+			OutPriority = ETaskPriority::BackgroundHigh;
+		}
+	}
+
+	// task priority translation from the new API to the old API
+	static ENamedThreads::Type TranslatePriority(UE::Tasks::Private::EExtendedTaskPriority Priority)
+	{
+		using namespace UE::Tasks::Private;
+
+		checkf(Priority >= EExtendedTaskPriority::GameThreadNormalPri && Priority <= EExtendedTaskPriority::RHIThreadHiPriLocalQueue, TEXT("only named threads can call this method: %d"), Priority);
+
+		int32 ConversionMap[] =
+		{
+				ENamedThreads::GameThread, // GameThreadNormalPri
+				ENamedThreads::GameThread | ENamedThreads::HighTaskPriority, // GameThreadHiPri
+				ENamedThreads::GameThread | ENamedThreads::LocalQueue, // GameThreadNormalPriLocalQueue
+				ENamedThreads::GameThread | ENamedThreads::HighTaskPriority | ENamedThreads::LocalQueue, // GameThreadHiPriLocalQueue
+
+				ENamedThreads::GetRenderThread(), // RenderThreadNormalPri
+				ENamedThreads::GetRenderThread() | ENamedThreads::HighTaskPriority, // RenderThreadHiPri
+				ENamedThreads::GetRenderThread() | ENamedThreads::LocalQueue, // RenderThreadNormalPriLocalQueue
+				ENamedThreads::GetRenderThread() | ENamedThreads::HighTaskPriority | ENamedThreads::LocalQueue, // RenderThreadHiPriLocalQueue
+
+				ENamedThreads::RHIThread, // RHIThreadNormalPri
+				ENamedThreads::RHIThread | ENamedThreads::HighTaskPriority, // RHIThreadHiPri
+				ENamedThreads::RHIThread | ENamedThreads::LocalQueue, // RHIThreadNormalPriLocalQueue
+				ENamedThreads::RHIThread | ENamedThreads::HighTaskPriority | ENamedThreads::LocalQueue // RHIThreadHiPriLocalQueue
+		};
+
+		return (ENamedThreads::Type)ConversionMap[(int32)Priority - (int32)EExtendedTaskPriority::GameThreadNormalPri];
+	}
+};
+
+// the new task implementation integrated into the old task API
+template<typename TTask>
+class TGraphTask : public TConcurrentLinearObject<TGraphTask<TTask>, FTaskGraphBlockAllocationTag>, public FBaseGraphTask
+{
+public:
+	/**
+	 *	This is a helper class returned from the factory. It constructs the embeded task with a set of arguments and sets the task up and makes it ready to execute.
+	 *	The task may complete before these routines even return.
+	 **/
+	class FConstructor
+	{
+	public:
+		UE_NONCOPYABLE(FConstructor);
+
+		/** Passthrough internal task constructor and dispatch. Note! Generally speaking references will not pass through; use pointers */
+		template<typename...T>
+		FORCEINLINE_DEBUGGABLE FGraphEventRef ConstructAndDispatchWhenReady(T&&... Args)
+		{
+			FGraphEventRef Ref{ ConstructAndHold(Forward<T>(Args)...) };
+			Ref->TryLaunch();
+			return Ref;
+		}
+
+		/** Passthrough internal task constructor and hold. */
+		template<typename...T>
+		FORCEINLINE_DEBUGGABLE TGraphTask* ConstructAndHold(T&&... Args)
+		{
+			TGraphTask* Task = new TGraphTask(Prerequisites);
+			TTask* TaskObject = new(&Task->TaskStorage) TTask(Forward<T>(Args)...);
+
+			UE::Tasks::ETaskPriority Pri;
+			UE::Tasks::Private::EExtendedTaskPriority ExtPri;
+			TranslatePriority(TaskObject->GetDesiredThread(), Pri, ExtPri);
+
+			Task->Init(Pri, ExtPri);
+
+			return Task;
+		}
+
+		FConstructor(const FGraphEventArray* InPrerequisites)
+			: Prerequisites(InPrerequisites)
+		{
+		}
+
+	private:
+		const FGraphEventArray* Prerequisites;
+	};
+
+	/**
+	 *	Factory to create a task and return the helper object to construct the embedded task and set it up for execution.
+	 *	@param Prerequisites; the list of FGraphEvents that must be completed prior to this task executing.
+	 *	@param CurrentThreadIfKnown; provides the index of the thread we are running on. Can be ENamedThreads::AnyThread if the current thread is unknown.
+	 *	@return a temporary helper class which can be used to complete the process.
+	**/
+	static FConstructor CreateTask(const FGraphEventArray* Prerequisites = nullptr, ENamedThreads::Type CurrentThreadIfKnown = ENamedThreads::AnyThread)
+	{
+		return FConstructor(Prerequisites);
+	}
+
+	virtual ~TGraphTask() override
+	{
+		DestructItem(TaskStorage.GetTypedPtr());
+	}
+
+private:
+	explicit TGraphTask(const FGraphEventArray* InPrerequisites)
+		: FBaseGraphTask(InPrerequisites)
+	{
+	}
+
+	void Init(UE::Tasks::ETaskPriority InPriority, UE::Tasks::Private::EExtendedTaskPriority InExtendedPriority)
+	{
+		FBaseGraphTask::Init(TEXT("GraphTask"), InPriority, InExtendedPriority);
+	}
+
+	virtual bool TryExecuteTask() override
+	{
+		return TryExecute(
+			[](UE::Tasks::Private::FTaskBase& Task)
+			{
+				TGraphTask& This = static_cast<TGraphTask&>(Task);
+				FGraphEventRef GraphEventRef{ &This };
+				TTask* TaskObject = This.TaskStorage.GetTypedPtr();
+				ENamedThreads::Type ThreadIndex = ENamedThreads::GetThreadIndex(TaskObject->GetDesiredThread());
+
+				TaskObject->DoTask(ThreadIndex, GraphEventRef);
+			}
+		);
+	}
+
+private:
+	TTypeCompatibleBytes<TTask> TaskStorage;
+};
+
+// an adaptation of FBaseGraphTask to be used as a standalone FGraphEvent
+class FGraphEventImpl : public TConcurrentLinearObject<FGraphEventImpl, FTaskGraphBlockAllocationTag>, public FBaseGraphTask
+{
+public:
+	FGraphEventImpl()
+		: FBaseGraphTask(nullptr)
+	{
+		Init(TEXT("GraphEvent"), UE::Tasks::ETaskPriority::Normal, UE::Tasks::Private::EExtendedTaskPriority::TaskEvent);
+	}
+
+private:
+	virtual bool TryExecuteTask() override
+	{
+		checkNoEntry(); // graph events are never executed
+		return true;
+	}
+};
+
+inline FGraphEventRef FBaseGraphTask::CreateGraphEvent()
+{
+	FGraphEventImpl* GraphEvent = new FGraphEventImpl;
+	return FGraphEventRef{ GraphEvent, /*bAddRef = */ false };
+}
+
+#else // TASKGRAPH_NEW_FRONTEND
+
 /** 
  *	Base class for all tasks. 
  *	Tasks go through a very specific life stage progression, and this is verified.
@@ -548,6 +838,16 @@ protected:
 #endif
 	}
 
+	LowLevelTasks::FTask& GetTaskHandle()
+	{
+		return TaskHandle;
+	}
+
+	ENamedThreads::Type GetThreadToExecuteOn() const
+	{
+		return ThreadToExecuteOn;
+	}
+
 private:
 	friend class FNamedTaskThread;
 	friend class FTaskThreadBase;
@@ -555,6 +855,7 @@ private:
 	friend class FGraphEvent;
 	friend class FTaskGraphImplementation;
 	friend class FTaskGraphCompatibilityImplementation;
+
 	LowLevelTasks::FTask TaskHandle;
 
 	// Subclass API
@@ -1110,6 +1411,7 @@ private:
 	FGraphEventRef				Subsequents;
 };
 
+#endif // TASKGRAPH_NEW_FRONTEND
 
 /** 
  *	FReturnGraphTask is a task used to return flow control from a named thread back to the original caller of ProcessThreadUntilRequestReturn
diff --git a/Engine/Source/Runtime/Core/Public/Tasks/TaskPrivate.h b/Engine/Source/Runtime/Core/Public/Tasks/TaskPrivate.h
--- a/Engine/Source/Runtime/Core/Public/Tasks/TaskPrivate.h
+++ b/Engine/Source/Runtime/Core/Public/Tasks/TaskPrivate.h
@@ -20,6 +20,10 @@
 #include <atomic>
 #include <type_traits>
 
+#if !defined(TASKGRAPH_NEW_FRONTEND)
+#define TASKGRAPH_NEW_FRONTEND 0
+#endif
+
 namespace UE::Tasks
 {
 	using LowLevelTasks::ETaskPriority;
@@ -41,6 +45,24 @@ namespace UE::Tasks
 			None,
 			Inline, // a task priority for "inline" task execution - a task is executed "inline" by the thread that unlocked it, w/o scheduling
 			TaskEvent, // a task priority used by task events, allows to shortcut task execution
+
+#if TASKGRAPH_NEW_FRONTEND
+			// for integration with named threads
+			GameThreadNormalPri,
+			GameThreadHiPri,
+			GameThreadNormalPriLocalQueue,
+			GameThreadHiPriLocalQueue,
+
+			RenderThreadNormalPri,
+			RenderThreadHiPri,
+			RenderThreadNormalPriLocalQueue,
+			RenderThreadHiPriLocalQueue,
+
+			RHIThreadNormalPri,
+			RHIThreadHiPri,
+			RHIThreadNormalPriLocalQueue,
+			RHIThreadHiPriLocalQueue
+#endif
 		};
 
 		// An abstract base class for task implementation. 
@@ -125,6 +147,13 @@ namespace UE::Tasks
 				return FPlatformTLS::GetCurrentThreadId() != ExecutingThreadId.load(std::memory_order_relaxed);
 			}
 
+#if TASKGRAPH_NEW_FRONTEND
+			bool IsNamedThreadTask() const
+			{
+				return ExtendedPriority >= EExtendedTaskPriority::GameThreadNormalPri;
+			}
+#endif
+
 			EExtendedTaskPriority GetExtendedPriority() const
 			{
 				return ExtendedPriority;
@@ -606,10 +635,7 @@ namespace UE::Tasks
 				return true;
 			}
 
-			void Schedule()
-			{
-				LowLevelTasks::FScheduler::Get().TryLaunch(LowLevelTask, LowLevelTasks::EQueuePreference::GlobalQueuePreference, /*bWakeUpWorker=*/ true);
-			}
+			CORE_API void Schedule();
 
 			// is called when the task has no pending prerequisites. Returns the previous piped task if any
 			CORE_API FTaskBase* TryPushIntoPipe();
diff --git a/Engine/Source/Runtime/Engine/Private/Components/SkeletalMeshComponent.cpp b/Engine/Source/Runtime/Engine/Private/Components/SkeletalMeshComponent.cpp
--- a/Engine/Source/Runtime/Engine/Private/Components/SkeletalMeshComponent.cpp
+++ b/Engine/Source/Runtime/Engine/Private/Components/SkeletalMeshComponent.cpp
@@ -4016,7 +4016,12 @@ void USkeletalMeshComponent::CompleteParallelAnimationEvaluation(bool bDoPostAni
 
 bool USkeletalMeshComponent::HandleExistingParallelEvaluationTask(bool bBlockOnTask, bool bPerformPostAnimEvaluation)
 {
-	if (IsRunningParallelEvaluation()) // We are already processing eval on another thread
+	// We are already processing eval and we are not inside it
+	if (IsRunningParallelEvaluation()
+#if TASKGRAPH_NEW_FRONTEND 
+		&& ParallelAnimationEvaluationTask->IsAwaitable()
+#endif
+	)
 	{
 		if (bBlockOnTask)
 		{
diff --git a/Engine/Source/Runtime/Engine/Private/Streaming/StreamingManagerTexture.cpp b/Engine/Source/Runtime/Engine/Private/Streaming/StreamingManagerTexture.cpp
--- a/Engine/Source/Runtime/Engine/Private/Streaming/StreamingManagerTexture.cpp
+++ b/Engine/Source/Runtime/Engine/Private/Streaming/StreamingManagerTexture.cpp
@@ -194,7 +194,11 @@ FRenderAssetStreamingManager::~FRenderAssetStreamingManager()
 
 TArray<FStreamingRenderAsset>& FRenderAssetStreamingManager::GetStreamingRenderAssetsAsyncSafe()
 {
-	if (StreamingRenderAssetsSyncEvent.IsValid())
+	if (StreamingRenderAssetsSyncEvent.IsValid()
+#if TASKGRAPH_NEW_FRONTEND
+		&& StreamingRenderAssetsSyncEvent->IsAwaitable()
+#endif
+	)
 	{
 		StreamingRenderAssetsSyncEvent->Wait(ENamedThreads::GameThread);
 	}
