Workaround for IncludeTool error.

diff --git a/Engine/Source/Runtime/Core/Public/Serialization/StructuredArchiveFwd.h b/Engine/Source/Runtime/Core/Public/Serialization/StructuredArchiveFwd.h
--- a/Engine/Source/Runtime/Core/Public/Serialization/StructuredArchiveFwd.h
+++ b/Engine/Source/Runtime/Core/Public/Serialization/StructuredArchiveFwd.h
@@ -2,6 +2,8 @@
 
 #pragma once
 
+#include <type_traits>
+
 // Formatters
 class FBinaryArchiveFormatter;
 class FStructuredArchiveFormatter;
@@ -18,8 +20,5 @@ class FStructuredArchiveStream;
 class FStructuredArchiveMap;
 
 /** Typedef for which formatter type to support */
-#if WITH_TEXT_ARCHIVE_SUPPORT
-	using FArchiveFormatterType = FStructuredArchiveFormatter;
-#else
-	using FArchiveFormatterType = FBinaryArchiveFormatter;
-#endif
+/** Written using std::conditional_t to work around IncludeTool complaints about #if blocks in Fwd.h files */
+using FArchiveFormatterType = std::conditional_t<WITH_TEXT_ARCHIVE_SUPPORT, FStructuredArchiveFormatter, FBinaryArchiveFormatter>;
