Add crash metadata for hung thread or present hangs.

diff --git a/Engine/Source/Runtime/Core/Private/HAL/ThreadHeartBeat.cpp b/Engine/Source/Runtime/Core/Private/HAL/ThreadHeartBeat.cpp
--- a/Engine/Source/Runtime/Core/Private/HAL/ThreadHeartBeat.cpp
+++ b/Engine/Source/Runtime/Core/Private/HAL/ThreadHeartBeat.cpp
@@ -13,6 +13,7 @@
 #include "Misc/CommandLine.h"
 #include "Misc/CoreDelegates.h"
 #include "HAL/ExceptionHandling.h"
+#include "GenericPlatform/GenericPlatformCrashContext.h"
 #include "Stats/Stats.h"
 #include "Async/TaskGraphInterfaces.h"
 #include "ProfilingDebugging/CsvProfiler.h"
@@ -174,6 +175,7 @@ void FORCENOINLINE FThreadHeartBeat::OnPresentHang(double HangDuration)
 #if MINIMAL_FATAL_HANG_DETECTION
 
 	LastHungThreadId = FThreadHeartBeat::PresentThreadId;
+	FGenericCrashContext::SetEngineData(TEXT("HungThread"), TEXT("Present"));
 #if PLATFORM_SWITCH
 	FPlatformCrashContext::UpdateDynamicData();
 #endif
@@ -214,6 +216,7 @@ void FORCENOINLINE FThreadHeartBeat::OnHang(double HangDuration, uint32 ThreadTh
 #if MINIMAL_FATAL_HANG_DETECTION
 
 	LastHungThreadId = ThreadThatHung;
+	FGenericCrashContext::SetEngineData(TEXT("HungThread"), LexToString(ThreadThatHung));
 #if PLATFORM_SWITCH
 	FPlatformCrashContext::UpdateDynamicData();
 #endif
