Fix non-unity error.

diff --git a/Engine/Source/Runtime/Core/Private/GenericPlatform/GenericPlatformMallocCrash.cpp b/Engine/Source/Runtime/Core/Private/GenericPlatform/GenericPlatformMallocCrash.cpp
--- a/Engine/Source/Runtime/Core/Private/GenericPlatform/GenericPlatformMallocCrash.cpp
+++ b/Engine/Source/Runtime/Core/Private/GenericPlatform/GenericPlatformMallocCrash.cpp
@@ -1,6 +1,7 @@
 // Copyright Epic Games, Inc. All Rights Reserved.
 
 #include "GenericPlatform/GenericPlatformMallocCrash.h"
+#include "Containers/ArrayView.h"
 #include "HAL/PlatformProcess.h"
 #include "HAL/PlatformTLS.h"
 #include "Templates/AlignmentTemplates.h"
