Add MicrosoftPlatformCrashContext and MicrosoftPlatformStackWalk to share code between Microsoft platforms.

diff --git a/Engine/Source/Runtime/Core/Private/Microsoft/MicrosoftPlatformCrashContext.cpp b/Engine/Source/Runtime/Core/Private/Microsoft/MicrosoftPlatformCrashContext.cpp
new file mode 100644
--- /dev/null
+++ b/Engine/Source/Runtime/Core/Private/Microsoft/MicrosoftPlatformCrashContext.cpp
@@ -0,0 +1,106 @@
+// Copyright Epic Games, Inc. All Rights Reserved.
+
+#include "Microsoft/MicrosoftPlatformCrashContext.h"
+#include "HAL/ThreadManager.h"
+
+#include "Microsoft/AllowMicrosoftPlatformTypes.h"
+	#include <psapi.h>
+#include "Microsoft/HideMicrosoftPlatformTypes.h"
+
+void FMicrosoftPlatformCrashContext::CaptureAllThreadContexts()
+{
+	TArray<typename FThreadManager::FThreadStackBackTrace> StackTraces;
+	FThreadManager::Get().GetAllThreadStackBackTraces(StackTraces);
+
+	for (const FThreadManager::FThreadStackBackTrace& Thread : StackTraces)
+	{
+		AddPortableThreadCallStack(Thread.ThreadId, *Thread.ThreadName, Thread.ProgramCounters.GetData(), Thread.ProgramCounters.Num());
+	}
+}
+
+void FMicrosoftPlatformCrashContext::AddPortableThreadCallStack(uint32 ThreadId, const TCHAR* ThreadName, const uint64* StackFrames, int32 NumStackFrames)
+{
+	FModuleHandleArray ProcModuleHandles;
+	GetProcModuleHandles(ProcessHandle, ProcModuleHandles);
+
+	FThreadStackFrames Thread;
+	Thread.ThreadId = ThreadId;
+	Thread.ThreadName = FString(ThreadName);
+	ConvertProgramCountersToStackFrames(ProcessHandle, ProcModuleHandles, StackFrames, NumStackFrames, Thread.StackFrames);
+	ThreadCallStacks.Push(Thread);
+}
+
+void FMicrosoftPlatformCrashContext::SetPortableCallStack(const uint64* StackTrace, int32 StackTraceDepth)
+{
+	FModuleHandleArray ProcessModuleHandles;
+	GetProcModuleHandles(ProcessHandle, ProcessModuleHandles);
+	ConvertProgramCountersToStackFrames(ProcessHandle, ProcessModuleHandles, StackTrace, StackTraceDepth, CallStack);
+}
+
+void FMicrosoftPlatformCrashContext::GetProcModuleHandles(const FProcHandle& ProcessHandle, FModuleHandleArray& OutHandles)
+{
+	// Get all the module handles for the current process. Each module handle is its base address.
+	for (;;)
+	{
+		DWORD BufferSize = OutHandles.Num() * sizeof(HMODULE);
+		DWORD RequiredBufferSize = 0;
+		if (!EnumProcessModulesEx(ProcessHandle.IsValid() ? ProcessHandle.Get() : GetCurrentProcess(), (HMODULE*)OutHandles.GetData(), BufferSize, &RequiredBufferSize, LIST_MODULES_ALL))
+		{
+			// We do not want partial set of modules in case this fails.
+			OutHandles.Empty();
+			return;
+		}
+		if (RequiredBufferSize <= BufferSize)
+		{
+			break;
+		}
+		OutHandles.SetNum(RequiredBufferSize / sizeof(HMODULE));
+	}
+	// Sort the handles by address. This allows us to do a binary search for the module containing an address.
+	Algo::Sort(OutHandles);
+}
+
+void FMicrosoftPlatformCrashContext::ConvertProgramCountersToStackFrames(
+	const FProcHandle& ProcessHandle,
+	const FModuleHandleArray& SortedModuleHandles,
+	const uint64* ProgramCounters,
+	int32 NumPCs,
+	TArray<FCrashStackFrame>& OutStackFrames)
+{
+	// Prepare the callstack buffer
+	OutStackFrames.Reset(NumPCs);
+
+	TCHAR Buffer[PLATFORM_MAX_FILEPATH_LENGTH_DEPRECATED];
+	FString Unknown = TEXT("Unknown");
+
+	HANDLE Process = ProcessHandle.IsValid() ? ProcessHandle.Get() : GetCurrentProcess();
+	// Create the crash context
+	for (int32 Idx = 0; Idx < NumPCs; ++Idx)
+	{
+		int32 ModuleIdx = Algo::UpperBound(SortedModuleHandles, (void*)ProgramCounters[Idx]) - 1;
+		if (ModuleIdx < 0 || ModuleIdx >= SortedModuleHandles.Num())
+		{
+			OutStackFrames.Add(FCrashStackFrame(Unknown, 0, ProgramCounters[Idx]));
+		}
+		else
+		{
+			FStringView ModuleName = Unknown;
+			if (GetModuleFileNameExW(Process, (HMODULE)SortedModuleHandles[ModuleIdx], Buffer, UE_ARRAY_COUNT(Buffer)) != 0)
+			{
+				ModuleName = Buffer;
+				if (int32 CharIdx = 0; ModuleName.FindLastChar('\\', CharIdx))
+				{
+					ModuleName.RightChopInline(CharIdx +1);
+				}
+				if (int32 CharIdx = 0; ModuleName.FindLastChar('.', CharIdx))
+				{
+					ModuleName.LeftInline(CharIdx);
+				}
+			}
+
+			uint64 BaseAddress = (uint64)SortedModuleHandles[ModuleIdx];
+			uint64 Offset = ProgramCounters[Idx] - BaseAddress;
+			OutStackFrames.Add(FCrashStackFrame(FString(ModuleName), BaseAddress, Offset));
+		}
+	}
+}
diff --git a/Engine/Source/Runtime/Core/Private/Microsoft/MicrosoftPlatformStackWalk.cpp b/Engine/Source/Runtime/Core/Private/Microsoft/MicrosoftPlatformStackWalk.cpp
new file mode 100644
--- /dev/null
+++ b/Engine/Source/Runtime/Core/Private/Microsoft/MicrosoftPlatformStackWalk.cpp
@@ -0,0 +1,77 @@
+// Copyright Epic Games, Inc. All Rights Reserved.
+
+#include "Microsoft/MicrosoftPlatformStackWalk.h"
+
+#include "Misc/Paths.h"
+
+#include "Microsoft/AllowMicrosoftPlatformTypes.h"
+	#include <psapi.h>
+#include "Microsoft/HideMicrosoftPlatformTypes.h"
+
+bool FMicrosoftPlatformStackWalk::ExtractInfoFromModule(void* ProcessHandle, void* Module, FStackWalkModuleInfo& OutInfo)
+{
+	MODULEINFO ModuleInfo = { 0 };
+	WCHAR ModuleName[MAX_PATH] = { 0 };
+	WCHAR ImageName[MAX_PATH] = { 0 };
+#if PLATFORM_64BITS
+	static_assert(sizeof(MODULEINFO) == 24, "Broken alignment for 64bit Windows include.");
+#else
+	static_assert(sizeof(MODULEINFO) == 12, "Broken alignment for 32bit Windows include.");
+#endif
+	GetModuleInformation(ProcessHandle, (HMODULE)Module, &ModuleInfo, sizeof(ModuleInfo));
+	GetModuleFileNameEx(ProcessHandle, (HMODULE)Module, ImageName, MAX_PATH);
+	GetModuleBaseName(ProcessHandle, (HMODULE)Module, ModuleName, MAX_PATH);
+
+	const uint8* ModuleBase = reinterpret_cast<const uint8*>(ModuleInfo.lpBaseOfDll);
+	const auto* DosHeader = reinterpret_cast<const IMAGE_DOS_HEADER*>(ModuleBase);
+	const auto* NtHeaders = reinterpret_cast<const IMAGE_NT_HEADERS*>(ModuleBase + DosHeader->e_lfanew);
+	bool bValid = DosHeader->e_magic == IMAGE_DOS_SIGNATURE && NtHeaders->Signature == IMAGE_NT_SIGNATURE;
+	if (!bValid)
+	{
+		return false;
+	}
+
+	const IMAGE_OPTIONAL_HEADER& OptionalHeader = NtHeaders->OptionalHeader;
+	OutInfo.BaseOfImage = (uint64)ModuleInfo.lpBaseOfDll;
+	OutInfo.ImageSize = OptionalHeader.SizeOfImage;
+	OutInfo.TimeDateStamp = NtHeaders->FileHeader.TimeDateStamp;
+	FString BaseModuleName = FPaths::GetBaseFilename(ModuleName);
+	FCString::Strncpy(OutInfo.ModuleName, *BaseModuleName, 32);
+	FCString::Strcpy(OutInfo.LoadedImageName, ImageName);
+	FCString::Strcpy(OutInfo.ImageName, ImageName);
+
+	// Find the guid and age of the binary, used to match debug files
+	const IMAGE_DATA_DIRECTORY& DebugInfoEntry = OptionalHeader.DataDirectory[IMAGE_DIRECTORY_ENTRY_DEBUG];
+	const auto* DebugEntries = (IMAGE_DEBUG_DIRECTORY*)(ModuleBase + DebugInfoEntry.VirtualAddress);
+	for (uint32 i = 0, n = DebugInfoEntry.Size / sizeof(DebugEntries[0]); i < n; ++i)
+	{
+		const IMAGE_DEBUG_DIRECTORY& Entry = DebugEntries[i];
+		if (Entry.Type == IMAGE_DEBUG_TYPE_CODEVIEW)
+		{
+			struct FCodeView7
+			{
+				uint32  Signature;
+				uint32  Guid[4];
+				uint32  Age;
+			};
+
+			if (Entry.SizeOfData < sizeof(FCodeView7))
+			{
+				continue;
+			}
+
+			const auto* CodeView7 = (FCodeView7*)(ModuleBase + Entry.AddressOfRawData);
+			if (CodeView7->Signature != 'SDSR')
+			{
+				continue;
+			}
+
+			static_assert(sizeof(OutInfo.PdbSig70) == sizeof(uint32) * 4);
+			FMemory::Memcpy(&OutInfo.PdbSig70, CodeView7->Guid, sizeof(OutInfo.PdbSig70));
+			OutInfo.PdbAge = CodeView7->Age;
+			break;
+		}
+	}
+
+	return true;
+}
\ No newline at end of file
diff --git a/Engine/Source/Runtime/Core/Private/Windows/WindowsPlatformCrashContext.cpp b/Engine/Source/Runtime/Core/Private/Windows/WindowsPlatformCrashContext.cpp
--- a/Engine/Source/Runtime/Core/Private/Windows/WindowsPlatformCrashContext.cpp
+++ b/Engine/Source/Runtime/Core/Private/Windows/WindowsPlatformCrashContext.cpp
@@ -23,21 +23,23 @@
 #include "Misc/OutputDeviceFile.h"
 #include "Serialization/Archive.h"
 #include "Windows/WindowsPlatformStackWalk.h"
-#include "Windows/WindowsHWrapper.h"
-#include "Windows/AllowWindowsPlatformTypes.h"
 #include "Templates/UniquePtr.h"
 #include "Templates/UnrealTemplate.h"
 #include "Misc/OutputDeviceArchiveWrapper.h"
 #include "HAL/ThreadManager.h"
 #include "BuildSettings.h"
 #include "CoreGlobals.h"
+#include <atomic>
+
+#include "Windows/WindowsHWrapper.h"
+#include "Windows/AllowWindowsPlatformTypes.h"
 #include <strsafe.h>
 #include <dbghelp.h>
 #include <Shlwapi.h>
 #include <psapi.h>
 #include <tlhelp32.h>
 #include <shellapi.h>
-#include <atomic>
+// Windows platform types intentionall not hidden til later
 
 #ifndef UE_LOG_CRASH_CALLSTACK
 	#define UE_LOG_CRASH_CALLSTACK 1
@@ -168,85 +170,6 @@ void FGenericCrashContext::CleanupPlatformSpecificFiles()
 	IFileManager::Get().Delete(*GPUMiniDumpPath);
 }
 
-
-void FWindowsPlatformCrashContext::GetProcModuleHandles(const FProcHandle& ProcessHandle, FModuleHandleArray& OutHandles)
-{
-	// Get all the module handles for the current process. Each module handle is its base address.
-	for (;;)
-	{
-		DWORD BufferSize = OutHandles.Num() * sizeof(HMODULE);
-		DWORD RequiredBufferSize = 0;
-		if (!EnumProcessModulesEx(ProcessHandle.IsValid() ? ProcessHandle.Get() : GetCurrentProcess(), (HMODULE*)OutHandles.GetData(), BufferSize, &RequiredBufferSize, LIST_MODULES_ALL))
-		{
-			// We do not want partial set of modules in case this fails.
-			OutHandles.Empty();
-			return;
-		}
-		if (RequiredBufferSize <= BufferSize)
-		{
-			break;
-		}
-		OutHandles.SetNum(RequiredBufferSize / sizeof(HMODULE));
-	}
-	// Sort the handles by address. This allows us to do a binary search for the module containing an address.
-	Algo::Sort(OutHandles);
-}
-
-void FWindowsPlatformCrashContext::ConvertProgramCountersToStackFrames(
-	const FProcHandle& ProcessHandle,
-	const FModuleHandleArray& SortedModuleHandles,
-	const uint64* ProgramCounters,
-	int32 NumPCs,
-	TArray<FCrashStackFrame>& OutStackFrames)
-{
-	// Prepare the callstack buffer
-	OutStackFrames.Reset(NumPCs);
-
-	// Create the crash context
-	for (int32 Idx = 0; Idx < NumPCs; ++Idx)
-	{
-		int32 ModuleIdx = Algo::UpperBound(SortedModuleHandles, (void*)ProgramCounters[Idx]) - 1;
-		if (ModuleIdx < 0 || ModuleIdx >= SortedModuleHandles.Num())
-		{
-			OutStackFrames.Add(FCrashStackFrame(TEXT("Unknown"), 0, ProgramCounters[Idx]));
-		}
-		else
-		{
-			TCHAR ModuleName[MAX_PATH];
-			if (GetModuleFileNameExW(ProcessHandle.IsValid() ? ProcessHandle.Get() : GetCurrentProcess(), (HMODULE)SortedModuleHandles[ModuleIdx], ModuleName, MAX_PATH) != 0)
-			{
-				TCHAR* ModuleNameEnd = FCString::Strrchr(ModuleName, '\\');
-				if (ModuleNameEnd != nullptr)
-				{
-					FMemory::Memmove(ModuleName, ModuleNameEnd + 1, (FCString::Strlen(ModuleNameEnd + 1) + 1) * sizeof(TCHAR));
-				}
-
-				TCHAR* ModuleNameExt = FCString::Strrchr(ModuleName, '.');
-				if (ModuleNameExt != nullptr)
-				{
-					*ModuleNameExt = TCHAR('\0');
-				}
-			}
-			else
-			{
-				const DWORD Err = GetLastError();
-				FCString::Strcpy(ModuleName, TEXT("Unknown"));
-			}
-
-			uint64 BaseAddress = (uint64)SortedModuleHandles[ModuleIdx];
-			uint64 Offset = ProgramCounters[Idx] - BaseAddress;
-			OutStackFrames.Add(FCrashStackFrame(ModuleName, BaseAddress, Offset));
-		}
-	}
-}
-
-void FWindowsPlatformCrashContext::SetPortableCallStack(const uint64* StackTrace, int32 StackTraceDepth)
-{
-	FModuleHandleArray ProcessModuleHandles;
-	GetProcModuleHandles(ProcessHandle, ProcessModuleHandles);
-	ConvertProgramCountersToStackFrames(ProcessHandle, ProcessModuleHandles, StackTrace, StackTraceDepth, CallStack);
-}
-
 void FWindowsPlatformCrashContext::AddPlatformSpecificProperties() const
 {
 	AddCrashProperty(TEXT("PlatformIsRunningWindows"), 1);
@@ -308,18 +231,6 @@ void FWindowsPlatformCrashContext::AddThreadContextString(
 	OutStr += LINE_TERMINATOR;
 }
 
-void FWindowsPlatformCrashContext::AddPortableThreadCallStack(uint32 ThreadId, const TCHAR* ThreadName, const uint64* StackFrames, int32 NumStackFrames)
-{
-	FModuleHandleArray ProcModuleHandles;
-	GetProcModuleHandles(ProcessHandle, ProcModuleHandles);
-
-	FThreadStackFrames Thread;
-	Thread.ThreadId = ThreadId;
-	Thread.ThreadName = FString(ThreadName);
-	ConvertProgramCountersToStackFrames(ProcessHandle, ProcModuleHandles, StackFrames, NumStackFrames, Thread.StackFrames);
-	ThreadCallStacks.Push(Thread);
-}
-
 void FWindowsPlatformCrashContext::CopyPlatformSpecificFiles(const TCHAR* OutputDirectory, void* Context)
 {
 	FGenericCrashContext::CopyPlatformSpecificFiles(OutputDirectory, Context);
@@ -351,17 +262,6 @@ void FWindowsPlatformCrashContext::CopyPlatformSpecificFiles(const TCHAR* Output
 	}
 }
 
-void FWindowsPlatformCrashContext::CaptureAllThreadContexts()
-{
-	TArray<typename FThreadManager::FThreadStackBackTrace> StackTraces;
-	FThreadManager::Get().GetAllThreadStackBackTraces(StackTraces);
-
-	for (const FThreadManager::FThreadStackBackTrace& Thread : StackTraces)
-	{
-		AddPortableThreadCallStack(Thread.ThreadId, *Thread.ThreadName, Thread.ProgramCounters.GetData(), Thread.ProgramCounters.Num());
-	}	
-}
-
 
 namespace
 {
@@ -1901,4 +1801,6 @@ namespace {
 
 
 
-
+#ifdef WINDOWS_PLATFORM_TYPES_GUARD
+	static_assert(false, "Missing HideWindowsPlatformTypes.h in " __FILE__ );
+#endif
diff --git a/Engine/Source/Runtime/Core/Public/Microsoft/MicrosoftPlatformCrashContext.h b/Engine/Source/Runtime/Core/Public/Microsoft/MicrosoftPlatformCrashContext.h
new file mode 100644
--- /dev/null
+++ b/Engine/Source/Runtime/Core/Public/Microsoft/MicrosoftPlatformCrashContext.h
@@ -0,0 +1,32 @@
+// Copyright Epic Games, Inc. All Rights Reserved.
+
+#pragma once
+
+#include "CoreTypes.h"
+#include "GenericPlatform/GenericPlatformCrashContext.h"
+
+struct CORE_API FMicrosoftPlatformCrashContext : public FGenericCrashContext
+{
+	FMicrosoftPlatformCrashContext(ECrashContextType InType, const TCHAR* InErrorMessage)
+		: FGenericCrashContext(InType, InErrorMessage)
+	{
+	}
+
+	void CaptureAllThreadContexts();
+
+	virtual void SetPortableCallStack(const uint64* StackTrace, int32 StackTraceDepth) override;
+	virtual void AddPortableThreadCallStack(uint32 ThreadId, const TCHAR* ThreadName, const uint64* StackFrames, int32 NumStackFrames) override;
+
+protected:
+	// Helpers
+	typedef TArray<void*, TInlineAllocator<128>> FModuleHandleArray;
+
+	static void GetProcModuleHandles(const FProcHandle& Process, FModuleHandleArray& OutHandles);
+
+	static void ConvertProgramCountersToStackFrames(
+		const FProcHandle& Process,
+		const FModuleHandleArray& SortedModuleHandles,
+		const uint64* ProgramCounters,
+		int32 NumPCs,
+		TArray<FCrashStackFrame>& OutStackFrames);
+};
\ No newline at end of file
diff --git a/Engine/Source/Runtime/Core/Public/Microsoft/MicrosoftPlatformStackWalk.h b/Engine/Source/Runtime/Core/Public/Microsoft/MicrosoftPlatformStackWalk.h
new file mode 100644
--- /dev/null
+++ b/Engine/Source/Runtime/Core/Public/Microsoft/MicrosoftPlatformStackWalk.h
@@ -0,0 +1,19 @@
+// Copyright Epic Games, Inc. All Rights Reserved.
+
+#pragma once
+
+#include "CoreTypes.h"
+#include "GenericPlatform/GenericPlatformStackWalk.h"
+
+/**
+ * Microsoft-shared implementation of the stack walking.
+ **/
+struct CORE_API FMicrosoftPlatformStackWalk
+	: public FGenericPlatformStackWalk
+{
+
+protected:
+	// Extract debug info for a module from the module header in memory. 
+	// Can directly read the information even when the current target can't load the symbols itself or use certain DbgHelp APIs.
+	static bool ExtractInfoFromModule(void* ProcessHandle, void* ModuleHandle, FStackWalkModuleInfo& OutInfo);
+};
\ No newline at end of file
diff --git a/Engine/Source/Runtime/Core/Public/Windows/WindowsPlatformCrashContext.h b/Engine/Source/Runtime/Core/Public/Windows/WindowsPlatformCrashContext.h
--- a/Engine/Source/Runtime/Core/Public/Windows/WindowsPlatformCrashContext.h
+++ b/Engine/Source/Runtime/Core/Public/Windows/WindowsPlatformCrashContext.h
@@ -3,32 +3,26 @@
 #pragma once
 
 #include "CoreTypes.h"
-#include "GenericPlatform/GenericPlatformCrashContext.h"
+#include "Microsoft/MicrosoftPlatformCrashContext.h"
 
 #if !defined(WINDOWS_CRASHCONTEXT_WITH_CUSTOM_HANDLERS)
 	#define WINDOWS_CRASHCONTEXT_WITH_CUSTOM_HANDLERS 0
 #endif
 
 
-struct CORE_API FWindowsPlatformCrashContext : public FGenericCrashContext
+struct CORE_API FWindowsPlatformCrashContext : public FMicrosoftPlatformCrashContext
 {
 	static const TCHAR* const UEGPUAftermathMinidumpName;
 	
 	FWindowsPlatformCrashContext(ECrashContextType InType, const TCHAR* InErrorMessage)
-		: FGenericCrashContext(InType, InErrorMessage)
+		: FMicrosoftPlatformCrashContext(InType, InErrorMessage)
 	{
 	}
 
-	virtual void SetPortableCallStack(const uint64* StackFrames, int32 NumStackFrames);
-
+	
 	virtual void AddPlatformSpecificProperties() const override;
-
-	virtual void AddPortableThreadCallStack(uint32 ThreadId, const TCHAR* ThreadName, const uint64* StackFrames, int32 NumStackFrames) override;
-
 	virtual void CopyPlatformSpecificFiles(const TCHAR* OutputDirectory, void* Context) override;
 
-	void CaptureAllThreadContexts();
-
 protected:
 	virtual bool GetPlatformAllThreadContextsString(FString& OutStr) const override;
 
@@ -36,15 +30,6 @@ private:
 	// Helpers
 	typedef TArray<void*, TInlineAllocator<128>> FModuleHandleArray;
 	
-	static void GetProcModuleHandles(const FProcHandle& Process, FModuleHandleArray& OutHandles);
-
-	static void ConvertProgramCountersToStackFrames(
-		const FProcHandle& Process,
-		const FModuleHandleArray& SortedModuleHandles,
-		const uint64* ProgramCounters,
-		int32 NumPCs,
-		TArray<FCrashStackFrame>& OutStackFrames);
-
 	static void AddThreadContextString(
 		uint32 CrashedThreadId,
 		uint32 ThreadId,
@@ -55,3 +40,4 @@ private:
 };
 
 typedef FWindowsPlatformCrashContext FPlatformCrashContext;
+
diff --git a/Engine/Source/Runtime/Core/Public/Windows/WindowsPlatformStackWalk.h b/Engine/Source/Runtime/Core/Public/Windows/WindowsPlatformStackWalk.h
--- a/Engine/Source/Runtime/Core/Public/Windows/WindowsPlatformStackWalk.h
+++ b/Engine/Source/Runtime/Core/Public/Windows/WindowsPlatformStackWalk.h
@@ -3,13 +3,13 @@
 #pragma once
 
 #include "CoreTypes.h"
-#include "GenericPlatform/GenericPlatformStackWalk.h"
+#include "Microsoft/MicrosoftPlatformStackWalk.h"
 
 /**
  * Windows implementation of the stack walking.
  **/
 struct CORE_API FWindowsPlatformStackWalk
-	: public FGenericPlatformStackWalk
+	: public FMicrosoftPlatformStackWalk
 {
 	static bool InitStackWalking();
 	static bool InitStackWalkingForProcess(const FProcHandle& Process);
