Fix use of FThreadManager::FThreadStackBackTrace when platform has all-thread-backtraces compiled out.

diff --git a/Engine/Source/Runtime/Core/Private/Microsoft/MicrosoftPlatformCrashContext.cpp b/Engine/Source/Runtime/Core/Private/Microsoft/MicrosoftPlatformCrashContext.cpp
--- a/Engine/Source/Runtime/Core/Private/Microsoft/MicrosoftPlatformCrashContext.cpp
+++ b/Engine/Source/Runtime/Core/Private/Microsoft/MicrosoftPlatformCrashContext.cpp
@@ -9,6 +9,7 @@
 
 void FMicrosoftPlatformCrashContext::CaptureAllThreadContexts()
 {
+#if PLATFORM_SUPPORTS_ALL_THREAD_BACKTRACES
 	TArray<typename FThreadManager::FThreadStackBackTrace> StackTraces;
 	FThreadManager::Get().GetAllThreadStackBackTraces(StackTraces);
 
@@ -16,6 +17,7 @@ void FMicrosoftPlatformCrashContext::CaptureAllThreadContexts()
 	{
 		AddPortableThreadCallStack(Thread.ThreadId, *Thread.ThreadName, Thread.ProgramCounters.GetData(), Thread.ProgramCounters.Num());
 	}
+#endif
 }
 
 void FMicrosoftPlatformCrashContext::AddPortableThreadCallStack(uint32 ThreadId, const TCHAR* ThreadName, const uint64* StackFrames, int32 NumStackFrames)
