Allow platforms to opt-in to all-thread backtraces, defaulting to on for windows and mac still.

diff --git a/Engine/Source/Runtime/Core/Private/HAL/ThreadingBase.cpp b/Engine/Source/Runtime/Core/Private/HAL/ThreadingBase.cpp
--- a/Engine/Source/Runtime/Core/Private/HAL/ThreadingBase.cpp
+++ b/Engine/Source/Runtime/Core/Private/HAL/ThreadingBase.cpp
@@ -555,7 +555,7 @@ const FString& FThreadManager::GetThreadNameInternal(uint32 ThreadId)
 	return NoThreadName;
 }
 
-#if PLATFORM_WINDOWS || PLATFORM_MAC
+#if PLATFORM_SUPPORTS_ALL_THREAD_BACKTRACES
 static void GetAllThreadStackBackTraces_ProcessSingle(
 	uint32 CurThreadId,
 	uint32 ThreadId,
diff --git a/Engine/Source/Runtime/Core/Public/HAL/ThreadManager.h b/Engine/Source/Runtime/Core/Public/HAL/ThreadManager.h
--- a/Engine/Source/Runtime/Core/Public/HAL/ThreadManager.h
+++ b/Engine/Source/Runtime/Core/Public/HAL/ThreadManager.h
@@ -7,6 +7,10 @@
 #include "Containers/Map.h"
 #include "HAL/CriticalSection.h"
 
+#ifndef PLATFORM_SUPPORTS_ALL_THREAD_BACKTRACES
+	#define PLATFORM_SUPPORTS_ALL_THREAD_BACKTRACES (PLATFORM_WINDOWS || PLATFORM_MAC)
+#endif
+
 /**
  * Manages runnables and runnable threads.
  */
@@ -57,7 +61,7 @@ public:
 		return Get().GetThreadNameInternal(ThreadId);
 	}
 
-#if PLATFORM_WINDOWS || PLATFORM_MAC
+#if PLATFORM_SUPPORTS_ALL_THREAD_BACKTRACES
 	struct FThreadStackBackTrace
 	{
 		uint32 ThreadId;
