#ue5 - VectorMod() uses either SVML fmod or std::fmod(). Previous alternate implementations could result in different results with certain inputs, especially when FMA3 was not available.

diff --git a/Engine/Source/Runtime/Core/Private/Tests/Math/UnrealMathTest.cpp b/Engine/Source/Runtime/Core/Private/Tests/Math/UnrealMathTest.cpp
--- a/Engine/Source/Runtime/Core/Private/Tests/Math/UnrealMathTest.cpp
+++ b/Engine/Source/Runtime/Core/Private/Tests/Math/UnrealMathTest.cpp
@@ -877,7 +877,6 @@ MATHTEST_INLINE VectorRegister4Double TestReferenceMod(const VectorRegister4Doub
 		(double)fmod(AF[3], MF[3]));
 }
 
-
 // SinCos
 template<typename FloatType, typename VectorRegisterType>
 MATHTEST_INLINE void TestReferenceSinCos(VectorRegisterType& S, VectorRegisterType& C, const VectorRegisterType& VAngles)
@@ -1930,13 +1929,11 @@ bool RunDoubleVectorTest()
 	V3 = VectorMod(V0, V1);
 	LogTest<double>(TEXT("VectorMod negative"), TestVectorsEqual(V2, V3));
 
-	/* Disable due to being broken for non-windows platforms UE-150201
 	V0 = MakeVectorRegister(89.9, 180.0, -256.0, -270.1);
 	V1 = MakeVectorRegister(360.0, 0.1, 360.0, 180.0);
 	V2 = TestReferenceMod(V0, V1);
 	V3 = VectorMod(V0, V1);
 	LogTest<double>(TEXT("VectorMod common"), TestVectorsEqual(V2, V3));
-	*/
 
 	// VectorSign
 	V0 = MakeVectorRegister(2.0, -2.0, 0.0, -3.0);
diff --git a/Engine/Source/Runtime/Core/Public/Math/UnrealMathNeon.h b/Engine/Source/Runtime/Core/Public/Math/UnrealMathNeon.h
--- a/Engine/Source/Runtime/Core/Public/Math/UnrealMathNeon.h
+++ b/Engine/Source/Runtime/Core/Public/Math/UnrealMathNeon.h
@@ -2752,14 +2752,14 @@ FORCEINLINE VectorRegister4Float VectorMod(const VectorRegister4Float& X, const
 {
 	// Check against invalid divisor
 	VectorRegister4Float InvalidDivisorMask = VectorCompareLE(VectorAbs(Y), GlobalVectorConstants::SmallNumber);
-	// Promote to doubles for better intermediate precision.
-	VectorRegister4Double DoubleX = MakeVectorRegisterDouble(X);
-	VectorRegister4Double DoubleY = MakeVectorRegisterDouble(Y);
-	// R = X - (Y * (Trunc(X / Y))
-	VectorRegister4Double Temp = VectorTruncate(VectorDivide(DoubleX, DoubleY));
-	VectorRegister4Double DoubleResult = VectorNegateMultiplyAdd(DoubleY, Temp, DoubleX);
-	// Convert back to floats. This is safe (in terms of not exceeding [-FLT_MAX,FLT_MAX]) because the answer is less than or equal to Abs(X) by definition.
-	VectorRegister4Float Result = MakeVectorRegisterFloatFromDouble(DoubleResult);
+	
+	AlignedFloat4 XFloats(X), YFloats(Y);
+	XFloats[0] = fmodf(XFloats[0], YFloats[0]);
+	XFloats[1] = fmodf(XFloats[1], YFloats[1]);
+	XFloats[2] = fmodf(XFloats[2], YFloats[2]);
+	XFloats[3] = fmodf(XFloats[3], YFloats[3]);
+	VectorRegister4Float Result = XFloats.ToVectorRegister();
+
 	// Return 0 where divisor Y was too small	
 	Result = VectorSelect(InvalidDivisorMask, GlobalVectorConstants::FloatZero, Result);
 	return Result;
@@ -2769,12 +2769,17 @@ FORCEINLINE VectorRegister4Double VectorMod(const VectorRegister4Double& X, cons
 {
 	// Check against invalid divisor
 	VectorRegister4Double InvalidDivisorMask = VectorCompareLE(VectorAbs(Y), GlobalVectorConstants::DoubleSmallNumber);
-	// R = X - (Y * (Trunc(X / Y))
-	VectorRegister4Double Temp = VectorTruncate(VectorDivide(X, Y));
-	VectorRegister4Double DoubleResult = VectorNegateMultiplyAdd(Y, Temp, X);
+	
+	AlignedDouble4 XDoubles(X), YDoubles(Y);
+	XDoubles[0] = fmod(XDoubles[0], YDoubles[0]);
+	XDoubles[1] = fmod(XDoubles[1], YDoubles[1]);
+	XDoubles[2] = fmod(XDoubles[2], YDoubles[2]);
+	XDoubles[3] = fmod(XDoubles[3], YDoubles[3]);
+	VectorRegister4Double DoubleResult = XDoubles.ToVectorRegister();
+
 	// Return 0 where divisor Y was too small
-	VectorRegister4Double Result = VectorSelect(InvalidDivisorMask, GlobalVectorConstants::DoubleZero, DoubleResult);
-	return Result;
+	DoubleResult = VectorSelect(InvalidDivisorMask, GlobalVectorConstants::DoubleZero, DoubleResult);
+	return DoubleResult;
 }
 
 FORCEINLINE VectorRegister4Float VectorSign(const VectorRegister4Float& X)
diff --git a/Engine/Source/Runtime/Core/Public/Math/UnrealMathSSE.h b/Engine/Source/Runtime/Core/Public/Math/UnrealMathSSE.h
--- a/Engine/Source/Runtime/Core/Public/Math/UnrealMathSSE.h
+++ b/Engine/Source/Runtime/Core/Public/Math/UnrealMathSSE.h
@@ -2860,20 +2860,17 @@ FORCEINLINE VectorRegister4Float VectorMod(const VectorRegister4Float& X, const
 #if UE_PLATFORM_MATH_USE_SVML
 	VectorRegister4Float Result = _mm_fmod_ps(X, Y);
 #else
-	// Promote to doubles for better intermediate precision.
-	VectorRegister4Double DoubleX = MakeVectorRegisterDouble(X);
-	VectorRegister4Double DoubleY = MakeVectorRegisterDouble(Y);
-	// R = X - (Y * (Trunc(X / Y))
-	VectorRegister4Double Temp = VectorTruncate(VectorDivide(DoubleX, DoubleY));
-	VectorRegister4Double DoubleResult = VectorNegateMultiplyAdd(DoubleY, Temp, DoubleX);
-	// Convert back to floats. This is safe (in terms of not exceeding [-FLT_MAX,FLT_MAX]) because the answer is less than or equal to Abs(X) by definition.
-	VectorRegister4Float Result = MakeVectorRegisterFloatFromDouble(DoubleResult);
+	AlignedFloat4 XFloats(X), YFloats(Y);
+	XFloats[0] = fmodf(XFloats[0], YFloats[0]);
+	XFloats[1] = fmodf(XFloats[1], YFloats[1]);
+	XFloats[2] = fmodf(XFloats[2], YFloats[2]);
+	XFloats[3] = fmodf(XFloats[3], YFloats[3]);
+	VectorRegister4Float Result = XFloats.ToVectorRegister();
 #endif
 
 	// Return 0 where divisor Y was too small	
 	Result = VectorSelect(InvalidDivisorMask, GlobalVectorConstants::FloatZero, Result);
 	return Result;
-
 }
 
 FORCEINLINE VectorRegister4Double VectorMod(const VectorRegister4Double& X, const VectorRegister4Double& Y)
@@ -2886,9 +2883,12 @@ FORCEINLINE VectorRegister4Double VectorMod(const VectorRegister4Double& X, cons
 #elif UE_PLATFORM_MATH_USE_SVML
 	VectorRegister4Double DoubleResult = VectorRegister4Double(_mm_fmod_pd(X.XY, Y.XY), _mm_fmod_pd(X.ZW, Y.ZW));
 #else
-	// R = X - (Y * (Trunc(X / Y))
-	VectorRegister4Double Temp = VectorTruncate(VectorDivide(X, Y));
-	VectorRegister4Double DoubleResult = VectorNegateMultiplyAdd(Y, Temp, X);
+	AlignedDouble4 XDoubles(X), YDoubles(Y);
+	XDoubles[0] = fmod(XDoubles[0], YDoubles[0]);
+	XDoubles[1] = fmod(XDoubles[1], YDoubles[1]);
+	XDoubles[2] = fmod(XDoubles[2], YDoubles[2]);
+	XDoubles[3] = fmod(XDoubles[3], YDoubles[3]);
+	VectorRegister4Double DoubleResult = XDoubles.ToVectorRegister();
 #endif
 
 	// Return 0 where divisor Y was too small
