Move support for adding all thread callstacks to crash context from Windows to Microsoft.

diff --git a/Engine/Source/Runtime/Core/Private/Microsoft/MicrosoftPlatformCrashContext.cpp b/Engine/Source/Runtime/Core/Private/Microsoft/MicrosoftPlatformCrashContext.cpp
--- a/Engine/Source/Runtime/Core/Private/Microsoft/MicrosoftPlatformCrashContext.cpp
+++ b/Engine/Source/Runtime/Core/Private/Microsoft/MicrosoftPlatformCrashContext.cpp
@@ -106,3 +106,58 @@ void FMicrosoftPlatformCrashContext::ConvertProgramCountersToStackFrames(
 		}
 	}
 }
+
+bool FMicrosoftPlatformCrashContext::GetPlatformAllThreadContextsString(FString& OutStr) const
+{
+	for (const FThreadStackFrames& Thread : ThreadCallStacks)
+	{
+		AddThreadContextString(
+			CrashedThreadId,
+			Thread.ThreadId,
+			Thread.ThreadName,
+			Thread.StackFrames,
+			OutStr
+		);
+	}
+	return !OutStr.IsEmpty();
+}
+
+void FMicrosoftPlatformCrashContext::AddThreadContextString(
+	uint32 CrashedThreadId,
+	uint32 ThreadId,
+	const FString& ThreadName,
+	const TArray<FCrashStackFrame>& StackFrames,
+	FString& OutStr)
+{
+	OutStr += TEXT("<Thread>");
+	{
+		OutStr += TEXT("<CallStack>");
+
+		int32 MaxModuleNameLen = 0;
+		for (const FCrashStackFrame& StFrame : StackFrames)
+		{
+			MaxModuleNameLen = FMath::Max(MaxModuleNameLen, StFrame.ModuleName.Len());
+		}
+
+		FString CallstackStr;
+		for (const FCrashStackFrame& StFrame : StackFrames)
+		{
+			CallstackStr += FString::Printf(TEXT("%-*s 0x%016llx + %-16llx"), MaxModuleNameLen + 1, *StFrame.ModuleName, StFrame.BaseAddress, StFrame.Offset);
+			CallstackStr += LINE_TERMINATOR;
+		}
+		AppendEscapedXMLString(OutStr, *CallstackStr);
+		OutStr += TEXT("</CallStack>");
+		OutStr += LINE_TERMINATOR;
+	}
+	OutStr += FString::Printf(TEXT("<IsCrashed>%s</IsCrashed>"), ThreadId == CrashedThreadId ? TEXT("true") : TEXT("false"));
+	OutStr += LINE_TERMINATOR;
+	// TODO: do we need thread register states?
+	OutStr += TEXT("<Registers></Registers>");
+	OutStr += LINE_TERMINATOR;
+	OutStr += FString::Printf(TEXT("<ThreadID>%d</ThreadID>"), ThreadId);
+	OutStr += LINE_TERMINATOR;
+	OutStr += FString::Printf(TEXT("<ThreadName>%s</ThreadName>"), *ThreadName);
+	OutStr += LINE_TERMINATOR;
+	OutStr += TEXT("</Thread>");
+	OutStr += LINE_TERMINATOR;
+}
diff --git a/Engine/Source/Runtime/Core/Private/Windows/WindowsPlatformCrashContext.cpp b/Engine/Source/Runtime/Core/Private/Windows/WindowsPlatformCrashContext.cpp
--- a/Engine/Source/Runtime/Core/Private/Windows/WindowsPlatformCrashContext.cpp
+++ b/Engine/Source/Runtime/Core/Private/Windows/WindowsPlatformCrashContext.cpp
@@ -176,60 +176,6 @@ void FWindowsPlatformCrashContext::AddPlatformSpecificProperties() const
 	AddCrashProperty(TEXT("IsRunningOnBattery"), FPlatformMisc::IsRunningOnBattery());
 }
 
-bool FWindowsPlatformCrashContext::GetPlatformAllThreadContextsString(FString& OutStr) const
-{
-	for (const FThreadStackFrames& Thread : ThreadCallStacks)
-	{
-		AddThreadContextString(
-			CrashedThreadId, 
-			Thread.ThreadId, 
-			Thread.ThreadName, 
-			Thread.StackFrames,
-			OutStr
-		);
-	}
-	return !OutStr.IsEmpty();
-}
-
-void FWindowsPlatformCrashContext::AddThreadContextString(
-	uint32 CrashedThreadId,
-	uint32 ThreadId,
-	const FString& ThreadName,
-	const TArray<FCrashStackFrame>& StackFrames,
-	FString& OutStr)
-{
-	OutStr += TEXT("<Thread>");
-	{
-		OutStr += TEXT("<CallStack>");
-
-		int32 MaxModuleNameLen = 0;
-		for (const FCrashStackFrame& StFrame : StackFrames)
-		{
-			MaxModuleNameLen = FMath::Max(MaxModuleNameLen, StFrame.ModuleName.Len());
-		}
-
-		FString CallstackStr;
-		for (const FCrashStackFrame& StFrame : StackFrames)
-		{
-			CallstackStr += FString::Printf(TEXT("%-*s 0x%016llx + %-16llx"), MaxModuleNameLen + 1, *StFrame.ModuleName, StFrame.BaseAddress, StFrame.Offset);
-			CallstackStr += LINE_TERMINATOR;
-		}
-		AppendEscapedXMLString(OutStr, *CallstackStr);
-		OutStr += TEXT("</CallStack>");
-		OutStr += LINE_TERMINATOR;
-	}
-	OutStr += FString::Printf(TEXT("<IsCrashed>%s</IsCrashed>"), ThreadId == CrashedThreadId ? TEXT("true") : TEXT("false"));
-	OutStr += LINE_TERMINATOR;
-	// TODO: do we need thread register states?
-	OutStr += TEXT("<Registers></Registers>");
-	OutStr += LINE_TERMINATOR;
-	OutStr += FString::Printf(TEXT("<ThreadID>%d</ThreadID>"), ThreadId);
-	OutStr += LINE_TERMINATOR;
-	OutStr += FString::Printf(TEXT("<ThreadName>%s</ThreadName>"), *ThreadName);
-	OutStr += LINE_TERMINATOR;
-	OutStr += TEXT("</Thread>");
-	OutStr += LINE_TERMINATOR;
-}
 
 void FWindowsPlatformCrashContext::CopyPlatformSpecificFiles(const TCHAR* OutputDirectory, void* Context)
 {
diff --git a/Engine/Source/Runtime/Core/Public/Microsoft/MicrosoftPlatformCrashContext.h b/Engine/Source/Runtime/Core/Public/Microsoft/MicrosoftPlatformCrashContext.h
--- a/Engine/Source/Runtime/Core/Public/Microsoft/MicrosoftPlatformCrashContext.h
+++ b/Engine/Source/Runtime/Core/Public/Microsoft/MicrosoftPlatformCrashContext.h
@@ -29,4 +29,14 @@ protected:
 		const uint64* ProgramCounters,
 		int32 NumPCs,
 		TArray<FCrashStackFrame>& OutStackFrames);
+
+	virtual bool GetPlatformAllThreadContextsString(FString& OutStr) const override;
+
+	static void AddThreadContextString(
+		uint32 CrashedThreadId,
+		uint32 ThreadId,
+		const FString& ThreadName,
+		const TArray<FCrashStackFrame>& StackFrames,
+		FString& OutStr);
+
 };
\ No newline at end of file
diff --git a/Engine/Source/Runtime/Core/Public/Windows/WindowsPlatformCrashContext.h b/Engine/Source/Runtime/Core/Public/Windows/WindowsPlatformCrashContext.h
--- a/Engine/Source/Runtime/Core/Public/Windows/WindowsPlatformCrashContext.h
+++ b/Engine/Source/Runtime/Core/Public/Windows/WindowsPlatformCrashContext.h
@@ -22,21 +22,6 @@ struct CORE_API FWindowsPlatformCrashContext : public FMicrosoftPlatformCrashCon
 	
 	virtual void AddPlatformSpecificProperties() const override;
 	virtual void CopyPlatformSpecificFiles(const TCHAR* OutputDirectory, void* Context) override;
-
-protected:
-	virtual bool GetPlatformAllThreadContextsString(FString& OutStr) const override;
-
-private:
-	// Helpers
-	typedef TArray<void*, TInlineAllocator<128>> FModuleHandleArray;
-	
-	static void AddThreadContextString(
-		uint32 CrashedThreadId,
-		uint32 ThreadId,
-		const FString& ThreadName,
-		const TArray<FCrashStackFrame>& StackFrames,
-		FString& OutStr);	
-
 };
 
 typedef FWindowsPlatformCrashContext FPlatformCrashContext;
