Add support for capturing loaded modules in crash context.
Duplicated "PlatformCallbackResult" field as more clearly named "CrashTrigger".

diff --git a/Engine/Source/Runtime/Core/Private/GenericPlatform/GenericPlatformCrashContext.cpp b/Engine/Source/Runtime/Core/Private/GenericPlatform/GenericPlatformCrashContext.cpp
--- a/Engine/Source/Runtime/Core/Private/GenericPlatform/GenericPlatformCrashContext.cpp
+++ b/Engine/Source/Runtime/Core/Private/GenericPlatform/GenericPlatformCrashContext.cpp
@@ -43,6 +43,37 @@ DEFINE_LOG_CATEGORY_STATIC(LogCrashContext, Display, All);
 
 extern CORE_API bool GIsGPUCrashed;
 
+/**
+ * A function-like type that creates a TStringBuilder to xml-escape a string
+ */
+template<int BufferSize=512>
+class FXmlEscapedString : public TStringBuilderWithBuffer<TCHAR, BufferSize>
+{
+public:
+	explicit FXmlEscapedString(FStringView Str)
+	{
+		FGenericCrashContext::AppendEscapedXMLString(*this, Str);
+	}
+};
+
+static bool NeedsEscape(FStringView Str)
+{
+	for (TCHAR C : Str)
+	{
+		switch (C)
+		{
+		case TCHAR('&'):
+		case TCHAR('"'):
+		case TCHAR('\''):
+		case TCHAR('<'):
+		case TCHAR('>'):
+		case TCHAR('\r'):
+			return true;
+		}
+	}
+	return false;
+}
+
 /*-----------------------------------------------------------------------------
 	FGenericCrashContext
 -----------------------------------------------------------------------------*/
@@ -725,6 +756,7 @@ void FGenericCrashContext::SerializeContentToBuffer() const
 	// Add new portable callstack element with crash stack
 	AddPortableCallStack();
 	AddPortableCallStackHash();
+	AddModules();
 
 	{
 		FString AllThreadStacks;
@@ -743,7 +775,9 @@ void FGenericCrashContext::SerializeContentToBuffer() const
 	BeginSection( CommonBuffer, PlatformPropertiesTag );
 	AddPlatformSpecificProperties();
 	// The name here is a bit cryptic, but we keep it to avoid breaking backend stuff.
-	AddCrashProperty(TEXT("PlatformCallbackResult"), NCached::Session.CrashType);
+	AddCrashProperty(TEXT("PlatformCallbackResult"), NCached::Session.CrashTrigger);
+	// New name we can phase in for the crash trigger to distinguish real crashes from debug
+	AddCrashProperty(TEXT("CrashTrigger"), NCached::Session.CrashTrigger);
 	EndSection( CommonBuffer, PlatformPropertiesTag );
 
 	// Add the engine data
@@ -803,7 +837,7 @@ void FGenericCrashContext::SetDeploymentName(const FString& EpicApp)
 
 void FGenericCrashContext::SetCrashTrigger(ECrashTrigger Type)
 {
-	NCached::Session.CrashType = (int32)Type;
+	NCached::Session.CrashTrigger = (int32)Type;
 }
 
 void FGenericCrashContext::GetUniqueCrashName(TCHAR* GUIDBuffer, int32 BufferSize) const
@@ -831,18 +865,13 @@ void FGenericCrashContext::SerializeAsXML( const TCHAR* Filename ) const
 	FFileHelper::SaveStringToFile( CommonBuffer, Filename, FFileHelper::EEncodingOptions::AutoDetect );
 }
 
-void FGenericCrashContext::AddCrashPropertyInternal(FString& Buffer, const TCHAR* PropertyName, const TCHAR* PropertyValue)
+void FGenericCrashContext::AddCrashPropertyInternal(FString& Buffer, FStringView PropertyName, FStringView PropertyValue)
 {
-	Buffer += TEXT( "<" );
-	Buffer += PropertyName;
-	Buffer += TEXT( ">" );
-
-	AppendEscapedXMLString(Buffer, PropertyValue);
-
-	Buffer += TEXT( "</" );
-	Buffer += PropertyName;
-	Buffer += TEXT( ">" );
-	Buffer += LINE_TERMINATOR;
+	Buffer.Appendf(TEXT("<%.*s>%s</%.*s>") LINE_TERMINATOR, 
+		PropertyName.Len(), PropertyName.GetData(),
+		*FXmlEscapedString(PropertyValue), 
+		PropertyName.Len(), PropertyName.GetData()
+		);
 }
 
 void FGenericCrashContext::AddPlatformSpecificProperties() const
@@ -889,6 +918,23 @@ void FGenericCrashContext::AddPortableCallStackHash() const
 	AddCrashProperty(TEXT("PCallStackHash"), *EscapedPortableHash);
 }
 
+void FGenericCrashContext::AppendPortableCallstack(FString& OutBuffer, TConstArrayView<FCrashStackFrame> StackFrames)
+{
+	OutBuffer += LINE_TERMINATOR;
+
+	// Get the max module name length for padding
+	int32 MaxModuleLength = 0;
+	for (const FCrashStackFrame& Frame : StackFrames)
+	{
+		MaxModuleLength = FMath::Max(MaxModuleLength, FXmlEscapedString(Frame.ModuleName).Len());
+	}
+
+	for (const FCrashStackFrame& Frame : StackFrames)
+	{
+		OutBuffer.Appendf(TEXT("%-*s 0x%016llx + %-16llx") LINE_TERMINATOR, MaxModuleLength + 1, *FXmlEscapedString(Frame.ModuleName), Frame.BaseAddress, Frame.Offset);
+	}
+}
+
 void FGenericCrashContext::AddPortableCallStack() const
 {	
 	if (CallStack.Num() == 0)
@@ -897,26 +943,28 @@ void FGenericCrashContext::AddPortableCallStack() const
 		return;
 	}
 
-	FString CrashStackBuffer = LINE_TERMINATOR;
-
-	// Get the max module name length for padding
-	int32 MaxModuleLength = 0;
-	for (TArray<FCrashStackFrame>::TConstIterator It(CallStack); It; ++It)
-	{
-		MaxModuleLength = FMath::Max(MaxModuleLength, It->ModuleName.Len());
-	}
+	BeginSection(CommonBuffer, TEXT("PCallStack"));
+	AppendPortableCallstack(CommonBuffer, CallStack);
+	EndSection(CommonBuffer, TEXT("PCallStack"));
+}
 
-	for (TArray<FCrashStackFrame>::TConstIterator It(CallStack); It; ++It)
-	{
-		CrashStackBuffer += FString::Printf(TEXT("%-*s 0x%016llx + %-16llx"),MaxModuleLength + 1, *It->ModuleName, It->BaseAddress, It->Offset);
-		CrashStackBuffer += LINE_TERMINATOR;
-	}
 
-	FString EscapedStackBuffer;
 
-	AppendEscapedXMLString(EscapedStackBuffer, *CrashStackBuffer);
+void FGenericCrashContext::AddModules() const
+{
+	BeginSection(CommonBuffer, TEXT("Modules"));
+	for (const FStackWalkModuleInfo& Module : ModulesInfo)
+	{
+		CommonBuffer.Appendf(TEXT("<Module Signature=\"%08x%08x\" DebugSignature=\"%08x-%04x-%04x-%04x%04x\">"),
+			Module.TimeDateStamp, Module.ImageSize, 
+			Module.PdbSig70.Data1, Module.PdbSig70.Data2, Module.PdbSig70.Data3, *(uint32*)&Module.PdbSig70.Data4[0], *(uint32*)&Module.PdbSig70.Data4[4]
+			);
+		
+		AppendEscapedXMLString(CommonBuffer, Module.ModuleName);
 
-	AddCrashProperty(TEXT("PCallStack"), *EscapedStackBuffer);
+		CommonBuffer += TEXT("</Module>") LINE_TERMINATOR;
+	}
+	EndSection(CommonBuffer, TEXT("Modules"));
 }
 
 void FGenericCrashContext::AddHeader(FString& Buffer)
@@ -932,30 +980,25 @@ void FGenericCrashContext::AddFooter(FString& Buffer)
 
 void FGenericCrashContext::BeginSection(FString& Buffer, const TCHAR* SectionName)
 {
-	Buffer += TEXT( "<" );
-	Buffer += SectionName;
-	Buffer += TEXT(">");
-	Buffer += LINE_TERMINATOR;
+	Buffer.Appendf(TEXT("<%s>") LINE_TERMINATOR, SectionName);
 }
 
 void FGenericCrashContext::EndSection(FString& Buffer, const TCHAR* SectionName)
 {
-	Buffer += TEXT( "</" );
-	Buffer += SectionName;
-	Buffer += TEXT( ">" );
-	Buffer += LINE_TERMINATOR;
+	Buffer.Appendf(TEXT("</%s>") LINE_TERMINATOR, SectionName);
 }
 
-void FGenericCrashContext::AppendEscapedXMLString(FString& OutBuffer, const TCHAR* Text)
+template<typename DEST>
+static void AppendEscapedXMLString(DEST& OutBuffer, FStringView Text)
 {
-	if (!Text)
+	if (Text.IsEmpty())
 	{
 		return;
 	}
 
-	while (*Text)
+	for (TCHAR C : Text)
 	{
-		switch (*Text)
+		switch (C)
 		{
 		case TCHAR('&'):
 			OutBuffer += TEXT("&amp;");
@@ -975,13 +1018,21 @@ void FGenericCrashContext::AppendEscapedXMLString(FString& OutBuffer, const TCHA
 		case TCHAR('\r'):
 			break;
 		default:
-			OutBuffer += *Text;
+			OutBuffer += C;
 		};
-
-		Text++;
 	}
 }
 
+void FGenericCrashContext::AppendEscapedXMLString(FString& OutBuffer, FStringView Text)
+{
+	::AppendEscapedXMLString(OutBuffer, Text);
+}
+
+void FGenericCrashContext::AppendEscapedXMLString(FStringBuilderBase& OutBuffer, FStringView Text)
+{
+	::AppendEscapedXMLString(OutBuffer, Text);
+}
+
 FString FGenericCrashContext::UnescapeXMLString( const FString& Text )
 {
 	return Text
@@ -1286,6 +1337,23 @@ void FGenericCrashContext::AddPortableThreadCallStack(uint32 ThreadId, const TCH
 	// Not implemented for generic class
 }
 
+void FGenericCrashContext::CaptureModules()
+{
+	ModulesInfo.Reset();
+	GetModules(ModulesInfo);
+}
+
+void FGenericCrashContext::GetModules(TArray<FStackWalkModuleInfo>& OutModules) const
+{
+	int32 Count = FPlatformStackWalk::GetProcessModuleCount();
+	if (Count > 0)
+	{
+		OutModules.Reset();
+		OutModules.AddZeroed(Count);
+		Count = FPlatformStackWalk::GetProcessModuleSignatures(OutModules.GetData(), Count);
+		OutModules.SetNum(Count);
+	}
+}
 
 
 void FGenericCrashContext::CopyPlatformSpecificFiles(const TCHAR* OutputDirectory, void* Context)
diff --git a/Engine/Source/Runtime/Core/Public/GenericPlatform/GenericPlatformCrashContext.h b/Engine/Source/Runtime/Core/Public/GenericPlatform/GenericPlatformCrashContext.h
--- a/Engine/Source/Runtime/Core/Public/GenericPlatform/GenericPlatformCrashContext.h
+++ b/Engine/Source/Runtime/Core/Public/GenericPlatform/GenericPlatformCrashContext.h
@@ -142,7 +142,7 @@ struct FSessionContext
 	int32 					NumberOfCoresIncludingHyperthreads;
 	int32 					SecondsSinceStart;
 	int32 					CrashDumpMode;
-	int32					CrashType;
+	int32					CrashTrigger;
 	int32					OOMAllocationAlignment;
 	uint64					OOMAllocationSize;
 	TCHAR 					GameName[CR_MAX_GENERIC_FIELD_CHARS];
@@ -394,7 +394,10 @@ public:
 	}
 
 	/** Escapes and appends specified text to XML string */
-	static void AppendEscapedXMLString( FString& OutBuffer, const TCHAR* Text );
+	static void AppendEscapedXMLString(FString& OutBuffer, FStringView Text );
+	static void AppendEscapedXMLString(FStringBuilderBase& OutBuffer, FStringView Text);
+
+	static void AppendPortableCallstack(FString& OutBuffer, TConstArrayView<FCrashStackFrame> StackFrames);
 
 	/** Unescapes a specified XML string, naive implementation. */
 	static FString UnescapeXMLString( const FString& Text );
@@ -495,6 +498,12 @@ public:
 	/** Gets the portable callstack to a specified stack and puts it into OutCallStack */
 	virtual void GetPortableCallStack(const uint64* StackFrames, int32 NumStackFrames, TArray<FCrashStackFrame>& OutCallStack) const;
 
+	/** Store info about loaded modules */
+	virtual void CaptureModules();
+
+	/** Gets info about loaded modules and stores it in the given array */
+	virtual void GetModules(TArray<FStackWalkModuleInfo>& OutModules) const;
+	
 	/** Adds a portable callstack for a thread */
 	virtual void AddPortableThreadCallStack(uint32 ThreadId, const TCHAR* ThreadName, const uint64* StackFrames, int32 NumStackFrames);
 
@@ -549,6 +558,7 @@ protected:
 	int NumMinidumpFramesToIgnore;
 	TArray<FCrashStackFrame> CallStack;
 	TArray<FThreadStackFrames> ThreadCallStacks;
+	TArray<FStackWalkModuleInfo> ModulesInfo;
 
 	/** Allow platform implementations to provide a callstack property. Primarily used when non-native code triggers a crash. */
 	virtual const TCHAR* GetCallstackProperty() const;
@@ -568,13 +578,13 @@ private:
 	static void SerializeUserSettings(FString& Buffer);
 
 	/** Writes a common property to the buffer. */
-	static void AddCrashPropertyInternal(FString& Buffer, const TCHAR* PropertyName, const TCHAR* PropertyValue);
+	static void AddCrashPropertyInternal(FString& Buffer, FStringView PropertyName, FStringView PropertyValue);
 
 	/** Writes a common property to the buffer. */
 	template <typename Type>
-	static void AddCrashPropertyInternal(FString& Buffer, const TCHAR* PropertyName, const Type& Value)
+	static void AddCrashPropertyInternal(FString& Buffer, FStringView PropertyName, const Type& Value)
 	{
-		AddCrashPropertyInternal(Buffer, PropertyName, *TTypeToString<Type>::ToString(Value));
+		AddCrashPropertyInternal(Buffer, PropertyName, FStringView(TTypeToString<Type>::ToString(Value)));
 	}
 
 	/** Serializes platform specific properties to the buffer. */
@@ -586,6 +596,9 @@ private:
 	/** Produces a hash based on the offsets of the portable callstack and adds it to the xml */
 	void AddPortableCallStackHash() const;
 
+	/** Add module/pdb information to the crash report xml */
+	void AddModules() const;
+
 	/** Writes header information to the buffer. */
 	static void AddHeader(FString& Buffer);
 
