Fixed memory corruption in MallocBinned when allocating memory with alignment requirements higher than DEFAULT_BINNED_ALLOCATOR_ALIGNMENT (16 bytes)

diff --git a/Engine/Source/Runtime/Core/Private/HAL/MallocBinned.cpp b/Engine/Source/Runtime/Core/Private/HAL/MallocBinned.cpp
--- a/Engine/Source/Runtime/Core/Private/HAL/MallocBinned.cpp
+++ b/Engine/Source/Runtime/Core/Private/HAL/MallocBinned.cpp
@@ -169,7 +169,7 @@ struct FMallocBinned::Private
 #if USE_OS_SMALL_BLOCK_ALLOC
 #if USE_OS_SMALL_BLOCK_GRAB_MEMORY_FROM_OS
 	enum { SMALL_BLOCK_MAX_TOTAL_POOL_SIZE = 0x20000000 };
-	enum { SMALL_BLOCK_GRAB_ALLOC_ALIGN = 16 };
+	enum { SMALL_BLOCK_GRAB_ALLOC_ALIGN = DEFAULT_BINNED_ALLOCATOR_ALIGNMENT };
 	enum { SMALL_BLOCK_GRAB_MAX_ALLOC_SIZE = 256 };
 	enum { SMALL_BLOCK_GRAB_MIN_ALLOC_SIZE = SMALL_BLOCK_GRAB_ALLOC_ALIGN };
 	
@@ -1288,8 +1288,6 @@ void* FMallocBinned::Malloc(SIZE_T Size, uint32 Alignment)
 
 	Private::FlushPendingFrees(*this);
 
-	Alignment = FMath::Max<uint32>(Alignment, Private::DEFAULT_BINNED_ALLOCATOR_ALIGNMENT);
-	Size = Align(Size, Alignment);
 	MEM_TIME(MemTime -= FPlatformTime::Seconds());
 	
 	BINNED_INCREMENT_STATCOUNTER(CurrentAllocs);
@@ -1298,7 +1296,7 @@ void* FMallocBinned::Malloc(SIZE_T Size, uint32 Alignment)
 	FFreeMem* Free = nullptr;
 	bool bUsePools = true;
 #if USE_OS_SMALL_BLOCK_ALLOC && !USE_OS_SMALL_BLOCK_GRAB_MEMORY_FROM_OS
-	if (FPlatformMemory::IsNanoMallocAvailable() && Size <= Private::SMALL_BLOCK_POOL_SIZE)
+	if (FPlatformMemory::IsNanoMallocAvailable() && (Size <= Private::SMALL_BLOCK_POOL_SIZE) && (Alignment <= Private::DEFAULT_BINNED_ALLOCATOR_ALIGNMENT))
 	{
 		//Make sure we have initialized our hash buckets even if we are using the NANO_MALLOC grabber, as otherwise we can end
 		//up making bad assumptions and trying to grab invalid data during a Realloc of this data.
@@ -1308,16 +1306,15 @@ void* FMallocBinned::Malloc(SIZE_T Size, uint32 Alignment)
 		}
 
 		bUsePools = false;
-		UPTRINT AlignedSize = Align(Size, Alignment);
 		SIZE_T ActualPoolSize; //TODO: use this to reduce waste?
-		Free = (FFreeMem*)Private::SmallOSAlloc(*this, AlignedSize, ActualPoolSize);
+		Free = (FFreeMem*)Private::SmallOSAlloc(*this, Size, ActualPoolSize);
 		check(FPlatformMemory::PtrIsOSMalloc(Free));
 		
 		if(!FPlatformMemory::PtrIsFromNanoMalloc(Free))
 		{
 			// This means we've overflowed the nano zone's internal buckets, which are fixed
 			// So we need to fall back to UE's allocator
-			Private::SmallOSFree(*this, Free, AlignedSize);
+			Private::SmallOSFree(*this, Free, Size);
 			bUsePools = true;
 			Free = nullptr;
 		}
@@ -1325,7 +1322,10 @@ void* FMallocBinned::Malloc(SIZE_T Size, uint32 Alignment)
 #endif
 	if (bUsePools)
 	{
-		if( Size < BinnedSizeLimit)
+		Alignment = FMath::Max<uint32>(Alignment, Private::DEFAULT_BINNED_ALLOCATOR_ALIGNMENT);
+		const uint32 PAGE_ALIGNMENT = 4096;
+
+		if( (Size < BinnedSizeLimit) && (Alignment == Private::DEFAULT_BINNED_ALLOCATOR_ALIGNMENT) )
 		{
 			// Allocate from pool.
 			FPoolTable* Table = MemSizeToPoolTable[Size];
@@ -1345,7 +1345,7 @@ void* FMallocBinned::Malloc(SIZE_T Size, uint32 Alignment)
 			Free = Private::AllocateBlockFromPool(*this, Table, Pool, Alignment);
 		}
 		else if ( ((Size >= BinnedSizeLimit && Size <= PagePoolTable[0].BlockSize) ||
-				   (Size > PageSize && Size <= PagePoolTable[1].BlockSize)))
+				   (Size > PageSize && Size <= PagePoolTable[1].BlockSize)) && (Alignment <= PAGE_ALIGNMENT))
 		{
 			// Bucket in a pool of 3*PageSize or 6*PageSize
 			uint32 BinType = Size < PageSize ? 0 : 1;
@@ -1419,10 +1419,6 @@ void* FMallocBinned::Realloc( void* Ptr, SIZE_T NewSize, uint32 Alignment )
 {
 	Alignment = FMath::Max<uint32>(Alignment, Private::DEFAULT_BINNED_ALLOCATOR_ALIGNMENT);
 	const uint32 NewSizeUnmodified = NewSize;
-	if (NewSize)
-	{
-		NewSize = Align(NewSize, Alignment);
-	}
 	MEM_TIME(MemTime -= FPlatformTime::Seconds());
 	UPTRINT BasePtr;
 	void* NewPtr = Ptr;
@@ -1431,6 +1427,7 @@ void* FMallocBinned::Realloc( void* Ptr, SIZE_T NewSize, uint32 Alignment )
 #if USE_OS_SMALL_BLOCK_ALLOC && !USE_OS_SMALL_BLOCK_GRAB_MEMORY_FROM_OS
 		if(FPlatformMemory::PtrIsOSMalloc(Ptr))
 		{
+			NewSize = Align(NewSize, Alignment);
 			NewPtr = ::realloc(Ptr, NewSize);
 			
 			if(!FPlatformMemory::PtrIsFromNanoMalloc(NewPtr))
@@ -1452,17 +1449,13 @@ void* FMallocBinned::Realloc( void* Ptr, SIZE_T NewSize, uint32 Alignment )
 			{
 				// Allocated from pool, so grow or shrink if necessary.
 				check(Pool->TableIndex > 0); // it isn't possible to allocate a size of 0, Malloc will increase the size to DEFAULT_BINNED_ALLOCATOR_ALIGNMENT
-				if (NewSizeUnmodified > MemSizeToPoolTable[Pool->TableIndex]->BlockSize || NewSizeUnmodified <= MemSizeToPoolTable[Pool->TableIndex - 1]->BlockSize)
+				if (NewSizeUnmodified > MemSizeToPoolTable[Pool->TableIndex]->BlockSize || NewSizeUnmodified <= MemSizeToPoolTable[Pool->TableIndex - 1]->BlockSize
+					|| !IsAligned(Ptr, Alignment))		// If reallocation requires a pointer to have a different alignment than the original we need to allocate new mem and copy old over
 				{
 					NewPtr = Malloc(NewSizeUnmodified, Alignment);
 					FMemory::Memcpy(NewPtr, Ptr, FMath::Min<SIZE_T>(NewSizeUnmodified, MemSizeToPoolTable[Pool->TableIndex]->BlockSize));
 					Free( Ptr );
 				}
-				else if (((UPTRINT)Ptr & (UPTRINT)(Alignment - 1)) != 0)
-				{
-					NewPtr = Align(Ptr, Alignment);
-					FMemory::Memmove(NewPtr, Ptr, NewSize);
-				}
 			}
 			else
 			{
@@ -1560,7 +1553,6 @@ bool FMallocBinned::GetAllocationSize(void *Original, SIZE_T &SizeOut)
 SIZE_T FMallocBinned::QuantizeSize(SIZE_T Size, uint32 Alignment)
 {
 	Alignment = FMath::Max<uint32>(Alignment, Private::DEFAULT_BINNED_ALLOCATOR_ALIGNMENT);
-	Size = Align(Size, Alignment);
 
 	SIZE_T Result;
 #if USE_OS_SMALL_BLOCK_ALLOC && !USE_OS_SMALL_BLOCK_GRAB_MEMORY_FROM_OS
