the fix 17455132 was lost during recent tasks refactoring. reintroducing it

diff --git a/Engine/Source/Runtime/Core/Public/Tasks/TaskPrivate.h b/Engine/Source/Runtime/Core/Public/Tasks/TaskPrivate.h
--- a/Engine/Source/Runtime/Core/Public/Tasks/TaskPrivate.h
+++ b/Engine/Source/Runtime/Core/Public/Tasks/TaskPrivate.h
@@ -760,11 +760,6 @@ namespace UE::Tasks
 				new(&TaskBodyStorage) TaskBodyType(MoveTemp(TaskBody));
 			}
 
-			virtual ~TExecutableTask() override
-			{
-				DestructItem(TaskBodyStorage.GetTypedPtr());
-			}
-
 			virtual bool TryExecuteTask() override
 			{
 				return FTaskBase::TryExecute(
@@ -772,6 +767,9 @@ namespace UE::Tasks
 					{ 
 						TExecutableTask& This = static_cast<TExecutableTask&>(Task);
 						new(&This.ResultStorage) ResultType{ Invoke(*This.TaskBodyStorage.GetTypedPtr()) }; 
+
+						// destroy the task body as soon as we are done with it, as it can have captured data sensitive to destruction order
+						DestructItem(This.TaskBodyStorage.GetTypedPtr());
 					}
 				);
 			}
@@ -802,11 +800,6 @@ namespace UE::Tasks
 				new(&TaskBodyStorage) TaskBodyType(MoveTemp(TaskBody));
 			}
 
-			virtual ~TExecutableTask() override
-			{
-				DestructItem(TaskBodyStorage.GetTypedPtr());
-			}
-
 			virtual bool TryExecuteTask() override
 			{
 				return TryExecute(
@@ -814,6 +807,9 @@ namespace UE::Tasks
 					{
 						TExecutableTask& This = static_cast<TExecutableTask&>(Task);
 						Invoke(*This.TaskBodyStorage.GetTypedPtr()); 
+
+						// destroy the task body as soon as we are done with it, as it can have captured data sensitive to destruction order
+						DestructItem(This.TaskBodyStorage.GetTypedPtr());
 					}
 				);
 			}
