Report the Windows build number and revision number in the log at startup.

Also adds support for detecting Windows Server 2022.

#rnx

diff --git a/Engine/Source/Runtime/Core/Private/Windows/WindowsPlatformMisc.cpp b/Engine/Source/Runtime/Core/Private/Windows/WindowsPlatformMisc.cpp
--- a/Engine/Source/Runtime/Core/Private/Windows/WindowsPlatformMisc.cpp
+++ b/Engine/Source/Runtime/Core/Private/Windows/WindowsPlatformMisc.cpp
@@ -186,7 +186,16 @@ int32 GetOSVersionsHelper( TCHAR* OutOSVersionLabel, int32 OSVersionLabelLength,
 				}
 				else
 				{
-					OSVersionLabel = TEXT("Windows Server 2019");
+					// Same thing here, Windows Server 2019 and 2022 both claim to be version 10.0, so use the
+					// build number to decide.
+					if (OsVersionInfo.dwBuildNumber >= 20348)
+					{
+						OSVersionLabel = TEXT("Windows Server 2022");
+					}
+					else
+					{
+						OSVersionLabel = TEXT("Windows Server 2019");
+					}
 				}
 
 				// For Windows 10, get the release number and append that to the string too (eg. 1709 = Fall Creators Update). 
@@ -207,6 +216,15 @@ int32 GetOSVersionsHelper( TCHAR* OutOSVersionLabel, int32 OSVersionLabelLength,
 							OSVersionLabel += FString::Printf(TEXT(" (Release %s)"), *ReleaseId);
 						}
 					}
+
+					FString UpdateBuildRevision;
+					if (!FWindowsPlatformMisc::QueryRegKey(HKEY_LOCAL_MACHINE, TEXT("SOFTWARE\\Microsoft\\Windows NT\\CurrentVersion"), TEXT("UBR"), UpdateBuildRevision))
+					{
+						UpdateBuildRevision = TEXT("UNKNOWN");
+					}
+
+					// Add the build number as displayed by the winver utility.
+					OSVersionLabel += FString::Printf(TEXT(" [%u.%u.%u.%s]"), OsVersionInfo.dwMajorVersion, OsVersionInfo.dwMinorVersion, OsVersionInfo.dwBuildNumber, *UpdateBuildRevision);
 				}
 				break;
 			default:
@@ -2815,19 +2833,50 @@ bool FWindowsPlatformMisc::QueryRegKey( const Windows::HKEY InKey, const TCHAR*
 		const uint32 RegFlags = (RegistryIndex == 0) ? KEY_WOW64_32KEY : KEY_WOW64_64KEY;
 		if (RegOpenKeyEx( InKey, InSubKey, 0, KEY_READ | RegFlags, &Key ) == ERROR_SUCCESS)
 		{
-			::DWORD Size = 0;
+			::DWORD Size = 0, ValueType = 0;
 			// First, we'll call RegQueryValueEx to find out how large of a buffer we need
-			if ((RegQueryValueEx( Key, InValueName, NULL, NULL, NULL, &Size ) == ERROR_SUCCESS) && Size)
+			if ((RegQueryValueEx( Key, InValueName, NULL, &ValueType, NULL, &Size ) == ERROR_SUCCESS) && Size)
 			{
-				// Allocate a buffer to hold the value and call the function again to get the data
-				char *Buffer = new char[Size];
-				if (RegQueryValueEx( Key, InValueName, NULL, NULL, (LPBYTE)Buffer, &Size ) == ERROR_SUCCESS)
+				switch (ValueType)
 				{
-					const uint32 Length = (Size / sizeof(TCHAR)) - 1;
-					OutData = FString( Length, (TCHAR*)Buffer );
-					bSuccess = true;
+					case REG_DWORD:
+					{
+						::DWORD Value;
+						if (RegQueryValueEx(Key, InValueName, NULL, NULL, (LPBYTE)&Value, &Size) == ERROR_SUCCESS)
+						{
+							OutData = FString::Printf(TEXT("%d"), Value);
+							bSuccess = true;
+						}
+						break;
+					}
+
+					case REG_QWORD:
+					{
+						int64 Value;
+						if (RegQueryValueEx(Key, InValueName, NULL, NULL, (LPBYTE)&Value, &Size) == ERROR_SUCCESS)
+						{
+							OutData = FString::Printf(TEXT("%lld"), Value);
+							bSuccess = true;
+						}
+						break;
+					}
+
+					case REG_SZ:
+					case REG_EXPAND_SZ:
+					case REG_MULTI_SZ:
+					{
+						// Allocate a buffer to hold the value and call the function again to get the data
+						char* Buffer = new char[Size];
+						if (RegQueryValueEx(Key, InValueName, NULL, NULL, (LPBYTE)Buffer, &Size) == ERROR_SUCCESS)
+						{
+							const uint32 Length = (Size / sizeof(TCHAR)) - 1;
+							OutData = FString(Length, (TCHAR*)Buffer);
+							bSuccess = true;
+						}
+						delete[] Buffer;
+						break;
+					}
 				}
-				delete [] Buffer;
 			}
 			RegCloseKey( Key );
 		}
