Add cast to the specified octree type in the generic octree callback.

diff --git a/Engine/Source/Runtime/Core/Public/Math/GenericOctree.h b/Engine/Source/Runtime/Core/Public/Math/GenericOctree.h
--- a/Engine/Source/Runtime/Core/Public/Math/GenericOctree.h
+++ b/Engine/Source/Runtime/Core/Public/Math/GenericOctree.h
@@ -1034,7 +1034,7 @@ private:
 	template <typename Semantics>
 	typename TEnableIf<TModels<COctreeSemanticsV2, Semantics>::Value>::Type SetOctreeSemanticsElementId(const ElementType& Element, FOctreeElementId2 Id)
 	{
-		Semantics::SetElementId(*this, Element, Id);
+		Semantics::SetElementId(static_cast<typename Semantics::FOctree&>(*this), Element, Id);
 	}
 
 protected:
