#ue5 - Make UnrealMathFPU usable again when vector intrinsics are disabled or when explicitly swapped in within VectorRegister.h. Fixed some related deprecated UE constants usage. Fixed NaN handling for comparisons in UnrealMathFPU to pass the UnrealMathTest startup tests.

diff --git a/Engine/Source/Runtime/Core/Private/Math/Color.cpp b/Engine/Source/Runtime/Core/Private/Math/Color.cpp
--- a/Engine/Source/Runtime/Core/Private/Math/Color.cpp
+++ b/Engine/Source/Runtime/Core/Private/Math/Color.cpp
@@ -172,7 +172,7 @@ static uint8 stbir__linear_to_srgb_uchar_fast(float in)
     return (uint8) ((bias + scale*t) >> 16);
 }
 
-#if PLATFORM_CPU_X86_FAMILY
+#if PLATFORM_CPU_X86_FAMILY && PLATFORM_ENABLE_VECTORINTRINSICS
 
 static FColor ConvertLinearToSRGBSSE2(const FLinearColor& InColor)
 {
@@ -250,7 +250,7 @@ FColor FLinearColor::ToFColorSRGB() const
 	// convert to 0, as do negative values, and out-of-range
 	// positive values convert to 255.
 
-#if PLATFORM_CPU_X86_FAMILY
+#if PLATFORM_CPU_X86_FAMILY && PLATFORM_ENABLE_VECTORINTRINSICS
 	return ConvertLinearToSRGBSSE2(*this);
 #else
 	return FColor(
diff --git a/Engine/Source/Runtime/Core/Private/Math/Transform.cpp b/Engine/Source/Runtime/Core/Private/Math/Transform.cpp
--- a/Engine/Source/Runtime/Core/Private/Math/Transform.cpp
+++ b/Engine/Source/Runtime/Core/Private/Math/Transform.cpp
@@ -140,7 +140,7 @@ void TTransform<T>::SetToRelativeTransform(const TTransform<T>& ParentTransform)
  	TMatrix<T> BM = ParentTransform.ToMatrixWithScale();
 #endif
 
-	const TVector<T> SafeRecipScale3D = GetSafeScaleReciprocal(ParentTransform.Scale3D, SMALL_NUMBER);
+	const TVector<T> SafeRecipScale3D = GetSafeScaleReciprocal(ParentTransform.Scale3D, UE_SMALL_NUMBER);
 	const TQuat<T> InverseRot = ParentTransform.Rotation.Inverse();
 
 	Scale3D *= SafeRecipScale3D;	
@@ -160,7 +160,7 @@ void TTransform<T>::GetRelativeTransformUsingMatrixWithScale(TTransform<T>* OutT
 	TMatrix<T> AM = Base->ToMatrixWithScale();
 	TMatrix<T> BM = Relative->ToMatrixWithScale();
 	// get combined scale
-	TVector<T> SafeRecipScale3D = GetSafeScaleReciprocal(Relative->Scale3D, SMALL_NUMBER);
+	TVector<T> SafeRecipScale3D = GetSafeScaleReciprocal(Relative->Scale3D, UE_SMALL_NUMBER);
 	TVector<T> DesiredScale3D = Base->Scale3D*SafeRecipScale3D;
 	ConstructTransformFromMatrixWithDesiredScale(AM, BM.Inverse(), DesiredScale3D, *OutTransform);
 }
@@ -183,7 +183,7 @@ TTransform<T> TTransform<T>::GetRelativeTransform(const TTransform<T>& Other) co
 	}
 	else
 	{
-		TVector<T> SafeRecipScale3D = GetSafeScaleReciprocal(Other.Scale3D, SMALL_NUMBER);
+		TVector<T> SafeRecipScale3D = GetSafeScaleReciprocal(Other.Scale3D, UE_SMALL_NUMBER);
 		Result.Scale3D = Scale3D*SafeRecipScale3D;
 
 		if (Other.Rotation.IsNormalized() == false)
diff --git a/Engine/Source/Runtime/Core/Private/Math/UnrealMath.cpp b/Engine/Source/Runtime/Core/Private/Math/UnrealMath.cpp
--- a/Engine/Source/Runtime/Core/Private/Math/UnrealMath.cpp
+++ b/Engine/Source/Runtime/Core/Private/Math/UnrealMath.cpp
@@ -479,7 +479,7 @@ FQuat4d FRotator3d::Quaternion() const
 	const VectorRegister4Double Result = VectorAdd(LeftTerm, RightTerm);
 	FQuat4d RotationQuat = FQuat4d::MakeFromVectorRegister(Result);
 #else
-	const double DEG_TO_RAD = DOUBLE_PI / (180.0);
+	const double DEG_TO_RAD = UE_DOUBLE_PI / (180.0);
 	const double RADS_DIVIDED_BY_2 = DEG_TO_RAD / 2.0;
 	double SP, SY, SR;
 	double CP, CY, CR;
diff --git a/Engine/Source/Runtime/Core/Private/Tests/Math/UnrealMathTest.cpp b/Engine/Source/Runtime/Core/Private/Tests/Math/UnrealMathTest.cpp
--- a/Engine/Source/Runtime/Core/Private/Tests/Math/UnrealMathTest.cpp
+++ b/Engine/Source/Runtime/Core/Private/Tests/Math/UnrealMathTest.cpp
@@ -1662,6 +1662,8 @@ bool RunDoubleVectorTest()
 	V2 = VectorCompareGT(V0, VectorNaN);
 	V3 = VectorZeroMaskDouble;
 	LogTest<double>(TEXT("VectorCompareGT (NaN)"), TestVectorsEqualBitwise(V2, V3));
+	V2 = VectorCompareGT(VectorNaN, V0);
+	LogTest<double>(TEXT("VectorCompareGT (NaN)"), TestVectorsEqualBitwise(V2, V3));
 
 	// VectorCompareGE
 	V0 = MakeVectorRegisterDouble(1.0f, 3.0f, 2.0f, 8.0f);
@@ -1682,6 +1684,8 @@ bool RunDoubleVectorTest()
 	V2 = VectorCompareLT(V0, VectorNaN);
 	V3 = VectorZeroMaskDouble;
 	LogTest<double>(TEXT("VectorCompareLT (NaN)"), TestVectorsEqualBitwise(V2, V3));
+	V2 = VectorCompareLT(VectorNaN, V0);
+	LogTest<double>(TEXT("VectorCompareLT (NaN)"), TestVectorsEqualBitwise(V2, V3));
 
 	// VectorCompareLE
 	V0 = MakeVectorRegisterDouble(2.0f, 4.0f, 2.0f, 1.0f);
@@ -1702,6 +1706,8 @@ bool RunDoubleVectorTest()
 	V2 = VectorCompareEQ(V0, VectorNaN);
 	V3 = VectorZeroMaskDouble;
 	LogTest<double>(TEXT("VectorCompareEQ (NaN)"), TestVectorsEqualBitwise(V2, V3));
+	V2 = VectorCompareEQ(VectorNaN, V0);
+	LogTest<double>(TEXT("VectorCompareEQ (NaN)"), TestVectorsEqualBitwise(V2, V3));
 	// NaN:NaN comparisons are undefined according to the Intel instruction manual, and will vary in optimized versus debug builds.
 	/*
 	V2 = VectorCompareEQ(VectorNaN, VectorNaN);
@@ -1719,6 +1725,8 @@ bool RunDoubleVectorTest()
 	V2 = VectorCompareNE(V0, VectorNaN);
 	V3 = GlobalVectorConstants::DoubleAllMask();
 	LogTest<double>(TEXT("VectorCompareNE (NaN)"), TestVectorsEqualBitwise(V2, V3));
+	V2 = VectorCompareNE(VectorNaN, V0);
+	LogTest<double>(TEXT("VectorCompareNE (NaN)"), TestVectorsEqualBitwise(V2, V3));
 	// NaN:NaN comparisons are undefined according to the Intel instruction manual, and will vary in optimized versus debug builds.
 	/*
 	V2 = VectorCompareNE(VectorNaN, VectorNaN);
@@ -2528,6 +2536,8 @@ bool FVectorRegisterAbstractionTest::RunTest(const FString& Parameters)
 	V2 = VectorCompareGT(V0, VectorNaN);
 	V3 = VectorZeroMaskFloat;
 	LogTest<float>(TEXT("VectorCompareGT (NaN)"), TestVectorsEqualBitwise(V2, V3));
+	V2 = VectorCompareGT(VectorNaN, V0);
+	LogTest<float>(TEXT("VectorCompareGT (NaN)"), TestVectorsEqualBitwise(V2, V3));
 
 	// VectorCompareGE
 	V0 = MakeVectorRegister(1.0f, 3.0f, 2.0f, 8.0f);
@@ -2548,6 +2558,8 @@ bool FVectorRegisterAbstractionTest::RunTest(const FString& Parameters)
 	V2 = VectorCompareLT(V0, VectorNaN);
 	V3 = VectorZeroMaskFloat;
 	LogTest<float>(TEXT("VectorCompareLT (NaN)"), TestVectorsEqualBitwise(V2, V3));
+	V2 = VectorCompareLT(VectorNaN, V0);
+	LogTest<float>(TEXT("VectorCompareLT (NaN)"), TestVectorsEqualBitwise(V2, V3));
 
 	// VectorCompareLE
 	V0 = MakeVectorRegister(2.0f, 4.0f, 2.0f, 1.0f);
@@ -2568,6 +2580,8 @@ bool FVectorRegisterAbstractionTest::RunTest(const FString& Parameters)
 	V2 = VectorCompareEQ(V0, VectorNaN);
 	V3 = VectorZeroMaskFloat;
 	LogTest<float>(TEXT("VectorCompareEQ (NaN)"), TestVectorsEqualBitwise(V2, V3));
+	V2 = VectorCompareEQ(VectorNaN, V0);
+	LogTest<float>(TEXT("VectorCompareEQ (NaN)"), TestVectorsEqualBitwise(V2, V3));
 	// NaN:NaN comparisons are undefined according to the Intel instruction manual, and will vary in optimized versus debug builds.
 	/*
 	V2 = VectorCompareEQ(VectorNaN, VectorNaN);
@@ -2584,6 +2598,8 @@ bool FVectorRegisterAbstractionTest::RunTest(const FString& Parameters)
 	V2 = VectorCompareNE(V0, VectorNaN);
 	V3 = GlobalVectorConstants::AllMask();
 	LogTest<float>(TEXT("VectorCompareNE (NaN)"), TestVectorsEqualBitwise(V2, V3));
+	V2 = VectorCompareNE(VectorNaN, V0);
+	LogTest<float>(TEXT("VectorCompareNE (NaN)"), TestVectorsEqualBitwise(V2, V3));
 	// NaN:NaN comparisons are undefined according to the Intel instruction manual, and will vary in optimized versus debug builds.
 	/*
 	V2 = VectorCompareNE(VectorNaN, VectorNaN);
diff --git a/Engine/Source/Runtime/Core/Public/Math/Quat.h b/Engine/Source/Runtime/Core/Public/Math/Quat.h
--- a/Engine/Source/Runtime/Core/Public/Math/Quat.h
+++ b/Engine/Source/Runtime/Core/Public/Math/Quat.h
@@ -1095,7 +1095,7 @@ FORCEINLINE bool TQuat<T>::IsNormalized() const
 	QuatVectorRegister TestValue = VectorAbs(VectorSubtract(VectorOne(), VectorDot4(A, A)));
 	return !VectorAnyGreaterThan(TestValue, GlobalVectorConstants::ThreshQuatNormalized);
 #else
-	return (FMath::Abs(1.f - SizeSquared()) < THRESH_QUAT_NORMALIZED);
+	return (FMath::Abs(1.f - SizeSquared()) < UE_THRESH_QUAT_NORMALIZED);
 #endif // PLATFORM_ENABLE_VECTORINTRINSICS
 }
 
@@ -1144,7 +1144,7 @@ FORCEINLINE TVector<T> TQuat<T>::GetRotationAxis() const
 	return V;
 #else
 	const T SquareSum = X * X + Y * Y + Z * Z;
-	if (SquareSum < SMALL_NUMBER)
+	if (SquareSum < UE_SMALL_NUMBER)
 	{
 		return TVector<T>::XAxisVector;
 	}
diff --git a/Engine/Source/Runtime/Core/Public/Math/TransformNonVectorized.h b/Engine/Source/Runtime/Core/Public/Math/TransformNonVectorized.h
--- a/Engine/Source/Runtime/Core/Public/Math/TransformNonVectorized.h
+++ b/Engine/Source/Runtime/Core/Public/Math/TransformNonVectorized.h
@@ -42,6 +42,7 @@ struct alignas(alignof(TQuat<T>)) TTransform
 	friend Z_Construct_UScriptStruct_FTransform_Statics;
 
 	using FReal = T;
+	using TransformVectorRegister = TVectorRegisterType<T>;
 
 protected:
 	/** Rotation of this transformation, as a quaternion. */
@@ -508,7 +509,7 @@ public:
 	FORCEINLINE static bool AnyHasNegativeScale(const TVector<T>& InScale3D, const TVector<T>& InOtherScale3D);
 	FORCEINLINE void ScaleTranslation(const TVector<T>& InScale3D);
 	FORCEINLINE void ScaleTranslation(const FReal& Scale);
-	FORCEINLINE void RemoveScaling(FReal Tolerance = SMALL_NUMBER);
+	FORCEINLINE void RemoveScaling(FReal Tolerance = UE_SMALL_NUMBER);
 	FORCEINLINE T GetMaximumAxisScale() const;
 	FORCEINLINE T GetMinimumAxisScale() const;
 
@@ -529,8 +530,8 @@ public:
 	*/
 	CORE_API void SetToRelativeTransform(const TTransform<T>& ParentTransform);
 
-	FORCEINLINE FVector4 TransformFVector4(const FVector4& V) const;
-	FORCEINLINE FVector4 TransformFVector4NoScale(const FVector4& V) const;
+	FORCEINLINE TVector4<T> TransformFVector4(const TVector4<T>& V) const;
+	FORCEINLINE TVector4<T> TransformFVector4NoScale(const TVector4<T>& V) const;
 	FORCEINLINE TVector<T> TransformPosition(const TVector<T>& V) const;
 	FORCEINLINE TVector<T> TransformPositionNoScale(const TVector<T>& V) const;
 
@@ -564,7 +565,7 @@ public:
 	FORCEINLINE TVector<T> GetScaledAxis(EAxis::Type InAxis) const;
 	FORCEINLINE TVector<T> GetUnitAxis(EAxis::Type InAxis) const;
 	FORCEINLINE void Mirror(EAxis::Type MirrorAxis, EAxis::Type FlipAxis);
-	FORCEINLINE static TVector<T> GetSafeScaleReciprocal(const TVector<T>& InScale, FReal Tolerance = SMALL_NUMBER);
+	FORCEINLINE static TVector<T> GetSafeScaleReciprocal(const TVector<T>& InScale, FReal Tolerance = UE_SMALL_NUMBER);
 
 	// temp function for easy conversion
 	FORCEINLINE TVector<T> GetLocation() const
@@ -879,6 +880,12 @@ public:
 		DiagnosticCheckNaN_Scale3D();
 	}
 
+	// For low-level VectorRegister programming
+	TransformVectorRegister GetTranslationRegister() const { return VectorLoadFloat3_W0(&Translation); }
+	TransformVectorRegister GetRotationRegister() const { return VectorLoad(&Rotation); }
+	void SetTranslationRegister(TransformVectorRegister InTranslation) { VectorStoreFloat3(InTranslation, &Translation); }
+	void SetRotationRegister(TransformVectorRegister InRotation) { VectorStore(InRotation, &Rotation); }
+
 	/** @note: Added template type function for Accumulate
 	* The template type isn't much useful yet, but it is with the plan to move forward
 	* to unify blending features with just type of additive or full pose
@@ -1357,9 +1364,9 @@ FORCEINLINE TTransform<T> TTransform<T>::GetScaled(TVector<T> InScale) const
 }
 
 
-/** Transform homogenous FVector4, ignoring the scaling part of this transform **/
+/** Transform homogenous TVector4<T>, ignoring the scaling part of this transform **/
 template<typename T>
-FORCEINLINE FVector4 TTransform<T>::TransformFVector4NoScale(const FVector4& V) const
+FORCEINLINE TVector4<T> TTransform<T>::TransformFVector4NoScale(const TVector4<T>& V) const
 {
 	DiagnosticCheckNaN_All();
 
@@ -1368,19 +1375,19 @@ FORCEINLINE FVector4 TTransform<T>::TransformFVector4NoScale(const FVector4& V)
 
 	//Transform using QST is following
 	//QST(P) = Q*S*P*-Q + T where Q = quaternion, S = scale, T = translation
-	FVector4 Transform = FVector4(Rotation.RotateVector(TVector<T>(V)), 0.f);
+	TVector4<T> Transform = TVector4<T>(Rotation.RotateVector(TVector<T>(V)), 0.f);
 	if (V.W == 1.f)
 	{
-		Transform += FVector4(Translation, 1.f);
+		Transform += TVector4<T>(Translation, 1.f);
 	}
 
 	return Transform;
 }
 
 
-/** Transform FVector4 **/
+/** Transform TVector4<T> **/
 template<typename T>
-FORCEINLINE FVector4 TTransform<T>::TransformFVector4(const FVector4& V) const
+FORCEINLINE TVector4<T> TTransform<T>::TransformFVector4(const TVector4<T>& V) const
 {
 	DiagnosticCheckNaN_All();
 
@@ -1390,10 +1397,10 @@ FORCEINLINE FVector4 TTransform<T>::TransformFVector4(const FVector4& V) const
 	//Transform using QST is following
 	//QST(P) = Q*S*P*-Q + T where Q = quaternion, S = scale, T = translation
 
-	FVector4 Transform = FVector4(Rotation.RotateVector(Scale3D*TVector<T>(V)), 0.f);
+	TVector4<T> Transform = TVector4<T>(Rotation.RotateVector(Scale3D*TVector<T>(V)), 0.f);
 	if (V.W == 1.f)
 	{
-		Transform += FVector4(Translation, 1.f);
+		Transform += TVector4<T>(Translation, 1.f);
 	}
 
 	return Transform;
diff --git a/Engine/Source/Runtime/Core/Public/Math/TransformVectorized.h b/Engine/Source/Runtime/Core/Public/Math/TransformVectorized.h
--- a/Engine/Source/Runtime/Core/Public/Math/TransformVectorized.h
+++ b/Engine/Source/Runtime/Core/Public/Math/TransformVectorized.h
@@ -944,6 +944,12 @@ public:
 		DiagnosticCheckNaN_Scale3D();
 	}
 	
+	// For low-level VectorRegister programming
+	const TPersistentVectorRegisterType<T>& GetTranslationRegister() const { return Translation; }
+	const TPersistentVectorRegisterType<T>& GetRotationRegister() const { return Rotation; }
+	void SetTranslationRegister(TransformVectorRegister InTranslation) { Translation = InTranslation; }
+	void SetRotationRegister(TransformVectorRegister InRotation) { Rotation = InRotation; }
+
 	/** @note : Added template type function for Accumulate
 	  * The template type isn't much useful yet, but it is with the plan to move forward
 	  * to unify blending features with just type of additive or full pose
diff --git a/Engine/Source/Runtime/Core/Public/Math/UnrealMathFPU.h b/Engine/Source/Runtime/Core/Public/Math/UnrealMathFPU.h
--- a/Engine/Source/Runtime/Core/Public/Math/UnrealMathFPU.h
+++ b/Engine/Source/Runtime/Core/Public/Math/UnrealMathFPU.h
@@ -879,19 +879,19 @@ FORCEINLINE VectorRegister4Double VectorDot4(const VectorRegister4Double& Vec1,
 FORCEINLINE VectorRegister4Float VectorCompareEQ(const VectorRegister4Float& Vec1, const VectorRegister4Float& Vec2)
 {
 	return MakeVectorRegisterFloat(
-		uint32(Vec1.V[0] == Vec2.V[0] ? 0xFFFFFFFF : 0),
-		uint32(Vec1.V[1] == Vec2.V[1] ? 0xFFFFFFFF : 0),
-		uint32(Vec1.V[2] == Vec2.V[2] ? 0xFFFFFFFF : 0),
-		uint32(Vec1.V[3] == Vec2.V[3] ? 0xFFFFFFFF : 0));
+		uint32(Vec1.V[0] == Vec2.V[0] && !std::isunordered(Vec1.V[0], Vec2.V[0]) ? 0xFFFFFFFF : 0),
+		uint32(Vec1.V[1] == Vec2.V[1] && !std::isunordered(Vec1.V[1], Vec2.V[1]) ? 0xFFFFFFFF : 0),
+		uint32(Vec1.V[2] == Vec2.V[2] && !std::isunordered(Vec1.V[2], Vec2.V[2]) ? 0xFFFFFFFF : 0),
+		uint32(Vec1.V[3] == Vec2.V[3] && !std::isunordered(Vec1.V[3], Vec2.V[3]) ? 0xFFFFFFFF : 0));
 }
 
 FORCEINLINE VectorRegister4Double VectorCompareEQ(const VectorRegister4Double& Vec1, const VectorRegister4Double& Vec2)
 {
 	return MakeVectorRegisterDouble(
-		uint64(Vec1.V[0] == Vec2.V[0] ? 0xFFFFFFFFFFFFFFFF : 0),
-		uint64(Vec1.V[1] == Vec2.V[1] ? 0xFFFFFFFFFFFFFFFF : 0),
-		uint64(Vec1.V[2] == Vec2.V[2] ? 0xFFFFFFFFFFFFFFFF : 0),
-		uint64(Vec1.V[3] == Vec2.V[3] ? 0xFFFFFFFFFFFFFFFF : 0));
+		uint64(Vec1.V[0] == Vec2.V[0] && !std::isunordered(Vec1.V[0], Vec2.V[0]) ? 0xFFFFFFFFFFFFFFFF : 0),
+		uint64(Vec1.V[1] == Vec2.V[1] && !std::isunordered(Vec1.V[1], Vec2.V[1]) ? 0xFFFFFFFFFFFFFFFF : 0),
+		uint64(Vec1.V[2] == Vec2.V[2] && !std::isunordered(Vec1.V[2], Vec2.V[2]) ? 0xFFFFFFFFFFFFFFFF : 0),
+		uint64(Vec1.V[3] == Vec2.V[3] && !std::isunordered(Vec1.V[3], Vec2.V[3]) ? 0xFFFFFFFFFFFFFFFF : 0));
 }
 
 /**
@@ -904,19 +904,19 @@ FORCEINLINE VectorRegister4Double VectorCompareEQ(const VectorRegister4Double& V
 FORCEINLINE VectorRegister4Float VectorCompareNE(const VectorRegister4Float& Vec1, const VectorRegister4Float& Vec2)
 {
 	return MakeVectorRegisterFloat(
-		uint32(Vec1.V[0] != Vec2.V[0] ? 0xFFFFFFFF : 0),
-		uint32(Vec1.V[1] != Vec2.V[1] ? 0xFFFFFFFF : 0),
-		uint32(Vec1.V[2] != Vec2.V[2] ? 0xFFFFFFFF : 0),
-		uint32(Vec1.V[3] != Vec2.V[3] ? 0xFFFFFFFF : 0));
+		uint32(Vec1.V[0] != Vec2.V[0] || std::isunordered(Vec1.V[0], Vec2.V[0]) ? 0xFFFFFFFF : 0),
+		uint32(Vec1.V[1] != Vec2.V[1] || std::isunordered(Vec1.V[1], Vec2.V[1]) ? 0xFFFFFFFF : 0),
+		uint32(Vec1.V[2] != Vec2.V[2] || std::isunordered(Vec1.V[2], Vec2.V[2]) ? 0xFFFFFFFF : 0),
+		uint32(Vec1.V[3] != Vec2.V[3] || std::isunordered(Vec1.V[3], Vec2.V[3]) ? 0xFFFFFFFF : 0));
 }
 
 FORCEINLINE VectorRegister4Double VectorCompareNE(const VectorRegister4Double& Vec1, const VectorRegister4Double& Vec2)
 {
 	return MakeVectorRegisterDouble(
-		uint64(Vec1.V[0] != Vec2.V[0] ? 0xFFFFFFFFFFFFFFFF : 0),
-		uint64(Vec1.V[1] != Vec2.V[1] ? 0xFFFFFFFFFFFFFFFF : 0),
-		uint64(Vec1.V[2] != Vec2.V[2] ? 0xFFFFFFFFFFFFFFFF : 0),
-		uint64(Vec1.V[3] != Vec2.V[3] ? 0xFFFFFFFFFFFFFFFF : 0));
+		uint64(Vec1.V[0] != Vec2.V[0] || std::isunordered(Vec1.V[0], Vec2.V[0]) ? 0xFFFFFFFFFFFFFFFF : 0),
+		uint64(Vec1.V[1] != Vec2.V[1] || std::isunordered(Vec1.V[1], Vec2.V[1]) ? 0xFFFFFFFFFFFFFFFF : 0),
+		uint64(Vec1.V[2] != Vec2.V[2] || std::isunordered(Vec1.V[2], Vec2.V[2]) ? 0xFFFFFFFFFFFFFFFF : 0),
+		uint64(Vec1.V[3] != Vec2.V[3] || std::isunordered(Vec1.V[3], Vec2.V[3]) ? 0xFFFFFFFFFFFFFFFF : 0));
 }
 
 /**
@@ -930,19 +930,19 @@ FORCEINLINE VectorRegister4Double VectorCompareNE(const VectorRegister4Double& V
 FORCEINLINE VectorRegister4Float VectorCompareGT(const VectorRegister4Float& Vec1, const VectorRegister4Float& Vec2)
 {
 	return MakeVectorRegisterFloat(
-		uint32(Vec1.V[0] > Vec2.V[0] ? 0xFFFFFFFF : 0),
-		uint32(Vec1.V[1] > Vec2.V[1] ? 0xFFFFFFFF : 0),
-		uint32(Vec1.V[2] > Vec2.V[2] ? 0xFFFFFFFF : 0),
-		uint32(Vec1.V[3] > Vec2.V[3] ? 0xFFFFFFFF : 0));
+		uint32(std::isgreater(Vec1.V[0], Vec2.V[0]) ? 0xFFFFFFFF : 0),
+		uint32(std::isgreater(Vec1.V[1], Vec2.V[1]) ? 0xFFFFFFFF : 0),
+		uint32(std::isgreater(Vec1.V[2], Vec2.V[2]) ? 0xFFFFFFFF : 0),
+		uint32(std::isgreater(Vec1.V[3], Vec2.V[3]) ? 0xFFFFFFFF : 0));
 }
 
 FORCEINLINE VectorRegister4Double VectorCompareGT(const VectorRegister4Double& Vec1, const VectorRegister4Double& Vec2)
 {
 	return MakeVectorRegisterDouble(
-		uint64(Vec1.V[0] > Vec2.V[0] ? 0xFFFFFFFFFFFFFFFF : 0),
-		uint64(Vec1.V[1] > Vec2.V[1] ? 0xFFFFFFFFFFFFFFFF : 0),
-		uint64(Vec1.V[2] > Vec2.V[2] ? 0xFFFFFFFFFFFFFFFF : 0),
-		uint64(Vec1.V[3] > Vec2.V[3] ? 0xFFFFFFFFFFFFFFFF : 0));
+		uint64(std::isgreater(Vec1.V[0], Vec2.V[0]) ? 0xFFFFFFFFFFFFFFFF : 0),
+		uint64(std::isgreater(Vec1.V[1], Vec2.V[1]) ? 0xFFFFFFFFFFFFFFFF : 0),
+		uint64(std::isgreater(Vec1.V[2], Vec2.V[2]) ? 0xFFFFFFFFFFFFFFFF : 0),
+		uint64(std::isgreater(Vec1.V[3], Vec2.V[3]) ? 0xFFFFFFFFFFFFFFFF : 0));
 }
 
 /**
@@ -956,19 +956,19 @@ FORCEINLINE VectorRegister4Double VectorCompareGT(const VectorRegister4Double& V
 FORCEINLINE VectorRegister4Float VectorCompareGE(const VectorRegister4Float& Vec1, const VectorRegister4Float& Vec2)
 {
 	return MakeVectorRegisterFloat(
-		uint32(Vec1.V[0] >= Vec2.V[0] ? 0xFFFFFFFF : 0),
-		uint32(Vec1.V[1] >= Vec2.V[1] ? 0xFFFFFFFF : 0),
-		uint32(Vec1.V[2] >= Vec2.V[2] ? 0xFFFFFFFF : 0),
-		uint32(Vec1.V[3] >= Vec2.V[3] ? 0xFFFFFFFF : 0));
+		uint32(std::isgreaterequal(Vec1.V[0], Vec2.V[0]) ? 0xFFFFFFFF : 0),
+		uint32(std::isgreaterequal(Vec1.V[1], Vec2.V[1]) ? 0xFFFFFFFF : 0),
+		uint32(std::isgreaterequal(Vec1.V[2], Vec2.V[2]) ? 0xFFFFFFFF : 0),
+		uint32(std::isgreaterequal(Vec1.V[3], Vec2.V[3]) ? 0xFFFFFFFF : 0));
 }
 
 FORCEINLINE VectorRegister4Double VectorCompareGE(const VectorRegister4Double& Vec1, const VectorRegister4Double& Vec2)
 {
 	return MakeVectorRegisterDouble(
-		uint64(Vec1.V[0] >= Vec2.V[0] ? 0xFFFFFFFFFFFFFFFF : 0),
-		uint64(Vec1.V[1] >= Vec2.V[1] ? 0xFFFFFFFFFFFFFFFF : 0),
-		uint64(Vec1.V[2] >= Vec2.V[2] ? 0xFFFFFFFFFFFFFFFF : 0),
-		uint64(Vec1.V[3] >= Vec2.V[3] ? 0xFFFFFFFFFFFFFFFF : 0));
+		uint64(std::isgreaterequal(Vec1.V[0], Vec2.V[0]) ? 0xFFFFFFFFFFFFFFFF : 0),
+		uint64(std::isgreaterequal(Vec1.V[1], Vec2.V[1]) ? 0xFFFFFFFFFFFFFFFF : 0),
+		uint64(std::isgreaterequal(Vec1.V[2], Vec2.V[2]) ? 0xFFFFFFFFFFFFFFFF : 0),
+		uint64(std::isgreaterequal(Vec1.V[3], Vec2.V[3]) ? 0xFFFFFFFFFFFFFFFF : 0));
 }
 
 /**
@@ -981,19 +981,19 @@ FORCEINLINE VectorRegister4Double VectorCompareGE(const VectorRegister4Double& V
 FORCEINLINE VectorRegister4Float VectorCompareLT(const VectorRegister4Float& Vec1, const VectorRegister4Float& Vec2)
 {
 	return MakeVectorRegisterFloat(
-		uint32(Vec1.V[0] < Vec2.V[0] ? 0xFFFFFFFF : 0),
-		uint32(Vec1.V[1] < Vec2.V[1] ? 0xFFFFFFFF : 0),
-		uint32(Vec1.V[2] < Vec2.V[2] ? 0xFFFFFFFF : 0),
-		uint32(Vec1.V[3] < Vec2.V[3] ? 0xFFFFFFFF : 0));
+		uint32(std::isless(Vec1.V[0], Vec2.V[0]) ? 0xFFFFFFFF : 0),
+		uint32(std::isless(Vec1.V[1], Vec2.V[1]) ? 0xFFFFFFFF : 0),
+		uint32(std::isless(Vec1.V[2], Vec2.V[2]) ? 0xFFFFFFFF : 0),
+		uint32(std::isless(Vec1.V[3], Vec2.V[3]) ? 0xFFFFFFFF : 0));
 }
 
 FORCEINLINE VectorRegister4Double VectorCompareLT(const VectorRegister4Double& Vec1, const VectorRegister4Double& Vec2)
 {
 	return MakeVectorRegisterDouble(
-		uint64(Vec1.V[0] < Vec2.V[0] ? 0xFFFFFFFFFFFFFFFF : 0),
-		uint64(Vec1.V[1] < Vec2.V[1] ? 0xFFFFFFFFFFFFFFFF : 0),
-		uint64(Vec1.V[2] < Vec2.V[2] ? 0xFFFFFFFFFFFFFFFF : 0),
-		uint64(Vec1.V[3] < Vec2.V[3] ? 0xFFFFFFFFFFFFFFFF : 0));
+		uint64(std::isless(Vec1.V[0], Vec2.V[0]) ? 0xFFFFFFFFFFFFFFFF : 0),
+		uint64(std::isless(Vec1.V[1], Vec2.V[1]) ? 0xFFFFFFFFFFFFFFFF : 0),
+		uint64(std::isless(Vec1.V[2], Vec2.V[2]) ? 0xFFFFFFFFFFFFFFFF : 0),
+		uint64(std::isless(Vec1.V[3], Vec2.V[3]) ? 0xFFFFFFFFFFFFFFFF : 0));
 }
 
 /**
@@ -1006,21 +1006,22 @@ FORCEINLINE VectorRegister4Double VectorCompareLT(const VectorRegister4Double& V
 FORCEINLINE VectorRegister4Float VectorCompareLE(const VectorRegister4Float& Vec1, const VectorRegister4Float& Vec2)
 {
 	return MakeVectorRegisterFloat(
-		uint32(Vec1.V[0] <= Vec2.V[0] ? 0xFFFFFFFF : 0),
-		uint32(Vec1.V[1] <= Vec2.V[1] ? 0xFFFFFFFF : 0),
-		uint32(Vec1.V[2] <= Vec2.V[2] ? 0xFFFFFFFF : 0),
-		uint32(Vec1.V[3] <= Vec2.V[3] ? 0xFFFFFFFF : 0));
+		uint32(std::islessequal(Vec1.V[0], Vec2.V[0]) ? 0xFFFFFFFF : 0),
+		uint32(std::islessequal(Vec1.V[1], Vec2.V[1]) ? 0xFFFFFFFF : 0),
+		uint32(std::islessequal(Vec1.V[2], Vec2.V[2]) ? 0xFFFFFFFF : 0),
+		uint32(std::islessequal(Vec1.V[3], Vec2.V[3]) ? 0xFFFFFFFF : 0));
 }
 
 FORCEINLINE VectorRegister4Double VectorCompareLE(const VectorRegister4Double& Vec1, const VectorRegister4Double& Vec2)
 {
 	return MakeVectorRegisterDouble(
-		uint64(Vec1.V[0] <= Vec2.V[0] ? 0xFFFFFFFFFFFFFFFF : 0),
-		uint64(Vec1.V[1] <= Vec2.V[1] ? 0xFFFFFFFFFFFFFFFF : 0),
-		uint64(Vec1.V[2] <= Vec2.V[2] ? 0xFFFFFFFFFFFFFFFF : 0),
-		uint64(Vec1.V[3] <= Vec2.V[3] ? 0xFFFFFFFFFFFFFFFF : 0));
+		uint64(std::islessequal(Vec1.V[0], Vec2.V[0]) ? 0xFFFFFFFFFFFFFFFF : 0),
+		uint64(std::islessequal(Vec1.V[1], Vec2.V[1]) ? 0xFFFFFFFFFFFFFFFF : 0),
+		uint64(std::islessequal(Vec1.V[2], Vec2.V[2]) ? 0xFFFFFFFFFFFFFFFF : 0),
+		uint64(std::islessequal(Vec1.V[3], Vec2.V[3]) ? 0xFFFFFFFFFFFFFFFF : 0));
 }
 
+
 /**
  * Returns an integer bit-mask (0x00 - 0x0f) based on the sign-bit for each component in a vector.
  *
@@ -2687,10 +2688,6 @@ FORCEINLINE VectorRegister4Int VectorIntLoad1(const void* Ptr)
 		IntSplat);
 }
 
-//
-// TODO: Engine doesn't currently compile with PLATFORM_ENABLE_VECTORINTRINSICS=0, but when it does these need to be fixed up.
-//
-
 #define VectorSetZero()								MakeVectorRegisterFloat(0.f, 0.f, 0.f, 0.f)
 #define VectorSet1(F)								VectorSetFloat1(F)
 #define VectorIntSet1(F)							MakeVectorRegisterInt(F, F, F, F)
diff --git a/Engine/Source/Runtime/Experimental/Chaos/Public/Chaos/Transform.h b/Engine/Source/Runtime/Experimental/Chaos/Public/Chaos/Transform.h
--- a/Engine/Source/Runtime/Experimental/Chaos/Public/Chaos/Transform.h
+++ b/Engine/Source/Runtime/Experimental/Chaos/Public/Chaos/Transform.h
@@ -180,12 +180,6 @@ namespace Chaos
 			return BaseTransform::ToMatrixNoScale();
 		}
 
-		// For low-level VectorRegister programming
-		const TPersistentVectorRegisterType<FRealSingle>& GetTranslationRegister() const { return Translation; }
-		const TPersistentVectorRegisterType<FRealSingle>& GetRotationRegister() const { return Rotation; }
-		void SetTranslationRegister(TransformVectorRegister InTranslation) { Translation = InTranslation; }
-		void SetRotationRegister(TransformVectorRegister InRotation) { Rotation = InRotation; }
-
 		CHAOS_API PMatrix<FRealSingle, 4, 4> operator*(const PMatrix<FRealSingle, 4, 4>& Matrix) const;
 		
 		inline TRigidTransform<FRealSingle, 3> operator*(const TRigidTransform<FRealSingle, 3>& Other) const
@@ -247,7 +241,7 @@ namespace Chaos
 #else
 			Result.Rotation = B.Rotation * A.Rotation;
 			Result.Translation = B.Rotation * A.Translation + B.Translation;
-			Result.Scale3D = FVector::OneVector;
+			Result.Scale3D = FVector3f::OneVector;
 #endif
 
 			return Result;
@@ -290,12 +284,6 @@ namespace Chaos
 			return BaseTransform::ToMatrixNoScale();
 		}
 
-		// For low-level VectorRegister programming
-		const TPersistentVectorRegisterType<FRealDouble>& GetTranslationRegister() const { return Translation; }
-		const TPersistentVectorRegisterType<FRealDouble>& GetRotationRegister() const { return Rotation; }
-		void SetTranslationRegister(TransformVectorRegister InTranslation) { Translation = InTranslation; }
-		void SetRotationRegister(TransformVectorRegister InRotation) { Rotation = InRotation; }
-
 		CHAOS_API PMatrix<FRealDouble, 4, 4> operator*(const Chaos::PMatrix<FRealDouble, 4, 4>& Matrix) const;
 		
 		inline TRigidTransform<FRealDouble, 3> operator*(const TRigidTransform<FRealDouble, 3>& Other) const
@@ -357,7 +345,7 @@ namespace Chaos
 #else
 			Result.Rotation = B.Rotation * A.Rotation;
 			Result.Translation = B.Rotation * A.Translation + B.Translation;
-			Result.Scale3D = FVector::OneVector;
+			Result.Scale3D = FVector3d::OneVector;
 #endif
 
 			return Result;
diff --git a/Engine/Source/Runtime/Experimental/Chaos/Public/Chaos/VectorUtility.h b/Engine/Source/Runtime/Experimental/Chaos/Public/Chaos/VectorUtility.h
--- a/Engine/Source/Runtime/Experimental/Chaos/Public/Chaos/VectorUtility.h
+++ b/Engine/Source/Runtime/Experimental/Chaos/Public/Chaos/VectorUtility.h
@@ -8,18 +8,7 @@
  * @param V	vector
  * @return		VectorRegister4Float( B.x, A.y, A.z, A.w)
  */
-#if !defined(_MSC_VER)|| PLATFORM_ENABLE_VECTORINTRINSICS_NEON
-FORCEINLINE VectorRegister4Float VectorCast4IntTo4Float(const VectorRegister4Int& V)
-{
-	return VectorRegister4Float(V);
-}
-#else
-FORCEINLINE VectorRegister4Float VectorCast4IntTo4Float(const VectorRegister4Int& V)
-{
-	return _mm_castsi128_ps(V);
-}
-
-#endif
+#define VectorCast4IntTo4Float(Vec)		VectorCastIntToFloat(Vec)
 
 /**
  * Cast VectorRegister4Float in VectorRegister4Int
@@ -27,20 +16,7 @@ FORCEINLINE VectorRegister4Float VectorCast4IntTo4Float(const VectorRegister4Int
  * @param V	vector
  * @return		VectorCast4FloatTo4Int( B.x, A.y, A.z, A.w)
  */
-#if !defined(_MSC_VER) || PLATFORM_ENABLE_VECTORINTRINSICS_NEON
-FORCEINLINE VectorRegister4Int VectorCast4FloatTo4Int(const VectorRegister4Float& V)
-{
-	return VectorRegister4Int(V);
-}
-#else
-FORCEINLINE VectorRegister4Int VectorCast4FloatTo4Int(const VectorRegister4Float& V)
-{
-	return _mm_castps_si128(V);
-}
-#endif
-
-
-#if PLATFORM_ENABLE_VECTORINTRINSICS_NEON
+#define VectorCast4FloatTo4Int(Vec)		VectorCastFloatToInt(Vec)
 
 /**
  * Selects and interleaves the lower two SP FP values from A and B.
@@ -51,57 +27,13 @@ FORCEINLINE VectorRegister4Int VectorCast4FloatTo4Int(const VectorRegister4Float
  */
 FORCEINLINE VectorRegister4Float VectorUnpackLo(const VectorRegister4Float& A, const VectorRegister4Float& B)
 {
+#if PLATFORM_ENABLE_VECTORINTRINSICS_NEON
 	return vzip1q_f32(A, B);
-}
-
-/**
- * Moves the lower 2 SP FP values of b to the upper 2 SP FP values of the result. The lower 2 SP FP values of a are passed through to the result.
- *
- * @param A	1st vector
- * @param B	2nd vector
- * @return		VectorRegister4Float( A.x, A.y, B.x, B.y)
-  */
-FORCEINLINE VectorRegister4Float VectorMoveLh(const VectorRegister4Float& A, const VectorRegister4Float& B)
-{
-	return vzip1q_f64(A, B);
-}
-
-/**
- * Combines two vectors using bitwise NOT AND (treating each vector as a 128 bit field)
- *
- * @param Vec1	1st vector
- * @param Vec2	2nd vector
- * @return		VectorRegister4Float( for each bit i: !Vec1[i] & Vec2[i] )
- */
-FORCEINLINE VectorRegister4Float VectorBitwiseNotAnd(const VectorRegister4Float& A, const VectorRegister4Float& B)
-{
-	return (VectorRegister4Float)vandq_u32(vmvnq_u32((VectorRegister4Int)A), (VectorRegister4Int)B);
-}
-
-/**
- * Return square root.
- *
- * @param A	1st vector
-  * @return		VectorRegister4Float( sqrt(A.x), sqrt(A.y), sqrt(A.z), sqrt(A.w))
-  */
-//FORCEINLINE VectorRegister4Float VectorSqrt(const VectorRegister4Float& A)
-//{
-//	return vsqrtq_f32(A);
-//}
-
-
-#else // PLATFORM_ENABLE_VECTORINTRINSICS_NEON
-
-/**
- * Selects and interleaves the lower two SP FP values from A and B.
- *
- * @param A	1st vector
- * @param B	2nd vector
- * @return		VectorRegister4Float( A.x, B.x, A.y, B.y)
- */
-FORCEINLINE VectorRegister4Float VectorUnpackLo(const VectorRegister4Float& A, const VectorRegister4Float& B)
-{
+#elif PLATFORM_ENABLE_VECTORINTRINSICS
 	return _mm_unpacklo_ps(A, B);
+#else
+	return MakeVectorRegisterFloat(A.V[0], B.V[0], A.V[1], B.V[1]);
+#endif
 }
 
 /**
@@ -113,7 +45,11 @@ FORCEINLINE VectorRegister4Float VectorUnpackLo(const VectorRegister4Float& A, c
   */
 FORCEINLINE VectorRegister4Float VectorMoveLh(const VectorRegister4Float& A, const VectorRegister4Float& B)
 {
-	return _mm_movelh_ps(A, B);
+#if PLATFORM_ENABLE_VECTORINTRINSICS_NEON
+	return vzip1q_f64(A, B);
+#else
+	return VectorCombineLow(A, B);
+#endif
 }
 
 
@@ -126,20 +62,15 @@ FORCEINLINE VectorRegister4Float VectorMoveLh(const VectorRegister4Float& A, con
  */
 FORCEINLINE VectorRegister4Float VectorBitwiseNotAnd(const VectorRegister4Float& A, const VectorRegister4Float& B)
 {
+#if PLATFORM_ENABLE_VECTORINTRINSICS_NEON
+	return (VectorRegister4Float)vandq_u32(vmvnq_u32((VectorRegister4Int)A), (VectorRegister4Int)B);
+#elif PLATFORM_ENABLE_VECTORINTRINSICS
 	return _mm_andnot_ps(A, B);
+#else
+	return MakeVectorRegisterFloat(
+		uint32(~((uint32*)(A.V))[0] & ((uint32*)(B.V))[0]),
+		uint32(~((uint32*)(A.V))[1] & ((uint32*)(B.V))[1]),
+		uint32(~((uint32*)(A.V))[2] & ((uint32*)(B.V))[2]),
+		uint32(~((uint32*)(A.V))[3] & ((uint32*)(B.V))[3]));
+#endif
 }
-
-/**
- * Return square root.
- *
- * @param A	1st vector
-  * @return		VectorRegister4Float( sqrt(A.x), sqrt(A.y), sqrt(A.z), sqrt(A.w))
-  */
-//FORCEINLINE VectorRegister4Float VectorSqrt(const VectorRegister4Float& A)
-//{
-//	return _mm_sqrt_ps(A);
-//}
-
-#endif // PLATFORM_ENABLE_VECTORINTRINSICS_NEON
-
-
