fix GeometryProcessingUnitTest compile on "Debug Editor" -- update _DEBUG-only GenericPlatformMallocCrash functions

diff --git a/Engine/Source/Runtime/Core/Private/GenericPlatform/GenericPlatformMallocCrash.cpp b/Engine/Source/Runtime/Core/Private/GenericPlatform/GenericPlatformMallocCrash.cpp
--- a/Engine/Source/Runtime/Core/Private/GenericPlatform/GenericPlatformMallocCrash.cpp
+++ b/Engine/Source/Runtime/Core/Private/GenericPlatform/GenericPlatformMallocCrash.cpp
@@ -215,8 +215,7 @@ private:
 #ifdef	_DEBUG
 		for( uint32 Index = 0; Index < MaxNumAllocations; ++Index )
 		{
-			FPtrInfo* PtrIt = Allocations[Index];
-			if( PtrIt->Size > 32768 )
+			if (FPtrInfo& PtrIt = Allocations[Index]; PtrIt.Size > 32768)
 			{
 				UE_DEBUG_BREAK();
 			}
@@ -267,7 +266,7 @@ FGenericPlatformMallocCrash::FGenericPlatformMallocCrash( FMalloc* MainMalloc )
 
 	InitializeSmallPools();
 #ifdef	_DEBUG
-	FPlatformMisc::LowLevelOutputDebugStringf( TEXT( "FGenericPlatformMallocCrash overhead is %u bytes\n" ), LargeMemoryPoolSize+GetSmallPoolTotalSize() );
+	FPlatformMisc::LowLevelOutputDebugStringf( TEXT( "FGenericPlatformMallocCrash overhead is %u bytes\n" ), LargeMemoryPoolSize+CalculateSmallPoolTotalSize() );
 #endif // _DEBUG
 }
 
@@ -437,14 +436,14 @@ void FGenericPlatformMallocCrash::PrintPoolsUsage()
 	for( uint32 Index = 0; Index < FGenericPlatformMallocCrash::NUM_POOLS; ++Index )
 	{
 		const FMallocCrashPool& CrashPool = FMallocCrashPool::GetPool(Index);
-		FPlatformMisc::LowLevelOutputDebugStringf( TEXT( "FPoolDesc(%5u,%4u),\n" ), CrashPool.AllocationSize-FGenericPlatformMallocCrash::PER_ALLOC_OVERHEAD, CrashPool.MaxUsedIndex );
+		FPlatformMisc::LowLevelOutputDebugStringf( TEXT( "FPoolDesc(%5u,%4u),\n" ), CrashPool.AllocationSize, CrashPool.MaxUsedIndex );
 	}
 
 	FPlatformMisc::LowLevelOutputDebugString( TEXT( "FPoolDesc tweaked\n" ) );
 	for( uint32 Index = 0; Index < FGenericPlatformMallocCrash::NUM_POOLS; ++Index )
 	{
 		const FMallocCrashPool& CrashPool = FMallocCrashPool::GetPool(Index);
-		FPlatformMisc::LowLevelOutputDebugStringf( TEXT( "FPoolDesc(%5u,%4u),\n" ), CrashPool.AllocationSize-FGenericPlatformMallocCrash::PER_ALLOC_OVERHEAD, Align(CrashPool.MaxUsedIndex*2+16,16) );
+		FPlatformMisc::LowLevelOutputDebugStringf( TEXT( "FPoolDesc(%5u,%4u),\n" ), CrashPool.AllocationSize, Align(CrashPool.MaxUsedIndex*2+16,16) );
 	}
 	FPlatformMisc::LowLevelOutputDebugStringf( TEXT( "LargeMemoryPoolOffset=%u\n" ), LargeMemoryPoolOffset );
 #endif // _DEBUG
