Remove references to EnumProcessModules when the current winapi partition doesn't allow access to it.

diff --git a/Engine/Source/Runtime/Core/Private/Microsoft/MicrosoftPlatformCrashContext.cpp b/Engine/Source/Runtime/Core/Private/Microsoft/MicrosoftPlatformCrashContext.cpp
--- a/Engine/Source/Runtime/Core/Private/Microsoft/MicrosoftPlatformCrashContext.cpp
+++ b/Engine/Source/Runtime/Core/Private/Microsoft/MicrosoftPlatformCrashContext.cpp
@@ -41,6 +41,8 @@ void FMicrosoftPlatformCrashContext::SetPortableCallStack(const uint64* StackTra
 
 void FMicrosoftPlatformCrashContext::GetProcModuleHandles(const FProcHandle& ProcessHandle, FModuleHandleArray& OutHandles)
 {
+	// Most but not all Microsoft platforms have acess to EnumProcessModules*
+#if WINAPI_FAMILY_PARTITION(WINAPI_PARTITION_DESKTOP | WINAPI_PARTITION_SYSTEM | WINAPI_PARTITION_GAMES)
 	// Get all the module handles for the current process. Each module handle is its base address.
 	for (;;)
 	{
@@ -60,6 +62,7 @@ void FMicrosoftPlatformCrashContext::GetProcModuleHandles(const FProcHandle& Pro
 	}
 	// Sort the handles by address. This allows us to do a binary search for the module containing an address.
 	Algo::Sort(OutHandles);
+#endif
 }
 
 void FMicrosoftPlatformCrashContext::ConvertProgramCountersToStackFrames(
