Extend preprocessor condition to use __builtin_bit_cast() intrinsic for Clang.

#rnx

diff --git a/Engine/Source/Runtime/Core/Public/Templates/TypeCompatibleBytes.h b/Engine/Source/Runtime/Core/Public/Templates/TypeCompatibleBytes.h
--- a/Engine/Source/Runtime/Core/Public/Templates/TypeCompatibleBytes.h
+++ b/Engine/Source/Runtime/Core/Public/Templates/TypeCompatibleBytes.h
@@ -79,7 +79,7 @@ inline ToType BitCast(const FromType& From)
 // Ensure we can use this builtin - seems to be present on Clang 9, GCC 11 and MSVC 19.26,
 // but gives spurious "non-void function 'BitCast' should return a value" errors on some
 // Mac and Android toolchains when building PCHs, so avoid those.
-#if (PLATFORM_APPLE && __clang_major__ >= 13) || (PLATFORM_ANDROID && __clang_major__ >= 13) || (!PLATFORM_APPLE && !PLATFORM_ANDROID && defined(__clang__) && __clang_major__ >= 9) || (defined(__GNUC__) && __GNUC__ >= 11) || (defined(_MSC_VER) && _MSC_VER >= 1926)
+#if (PLATFORM_APPLE && __clang_major__ >= 13) || (PLATFORM_ANDROID && __clang_major__ >= 13) || (PLATFORM_SWITCH && __clang_major__ >= 13) || (!PLATFORM_APPLE && !PLATFORM_ANDROID && !PLATFORM_SWITCH && defined(__clang__) && __clang_major__ >= 9) || (defined(__GNUC__) && __GNUC__ >= 11) || (defined(_MSC_VER) && _MSC_VER >= 1926)
 	return __builtin_bit_cast(ToType, From);
 #else
 	TTypeCompatibleBytes<ToType> Result;
