Resolve race condition that triggers false hang detections

#rn Fix Core Previously an errant heartbeat present hang could be triggered if PLATFORM_PRESENT_HANG_DETECTION_ON_BY_DEFAULT was disabled but PresentHangDuration was enabled in the project's configuration.  This change resolves the initialization of present hang detection settings preventing a race condition that triggered the errant hang.

diff --git a/Engine/Source/Runtime/Core/Private/HAL/ThreadHeartBeat.cpp b/Engine/Source/Runtime/Core/Private/HAL/ThreadHeartBeat.cpp
--- a/Engine/Source/Runtime/Core/Private/HAL/ThreadHeartBeat.cpp
+++ b/Engine/Source/Runtime/Core/Private/HAL/ThreadHeartBeat.cpp
@@ -431,6 +431,23 @@ void FThreadHeartBeat::InitSettings()
 	CurrentPresentDuration = ConfigPresentDuration * HangDurationMultiplier;
 
 	bHangsAreFatal = bNewHangsAreFatal;
+
+	// Update the existing thread and present hang durations.
+	// Only increase existing thread's heartbeats.
+	// We don't want to decrease here, threads and present logic will pick up a smaller hang duration 
+	// the next time they call HeartBeat() or PresentFrame().
+	for (TPair<uint32, FHeartBeatInfo>& Pair : ThreadHeartBeat)
+	{
+		if (Pair.Value.HangDuration < CurrentHangDuration)
+		{
+			Pair.Value.HangDuration = CurrentHangDuration;
+		}
+	}
+	
+	if (PresentHeartBeat.HangDuration < CurrentPresentDuration)
+	{
+		PresentHeartBeat.HangDuration = CurrentPresentDuration;
+	}
 }
 
 void FThreadHeartBeat::HeartBeat(bool bReadConfig)
@@ -869,23 +886,6 @@ void FThreadHeartBeat::SetDurationMultiplier(double NewMultiplier)
 	InitSettings();
 
 	UE_LOG(LogCore, Display, TEXT("Setting hang detector multiplier to %.4fs. New hang duration: %.4fs. New present duration: %.4fs."), NewMultiplier, CurrentHangDuration, CurrentPresentDuration);
-
-	// Update the existing thread's hang durations.
-	for (TPair<uint32, FHeartBeatInfo>& Pair : ThreadHeartBeat)
-	{
-		// Only increase existing thread's heartbeats.
-		// We don't want to decrease here, otherwise reducing the multiplier could cause a false detection.
-		// Threads will pick up a smaller hang duration the next time they call HeartBeat().
-		if (Pair.Value.HangDuration < CurrentHangDuration)
-		{
-			Pair.Value.HangDuration = CurrentHangDuration;
-		}
-	}
-
-	if (PresentHeartBeat.HangDuration < CurrentPresentDuration)
-	{
-		PresentHeartBeat.HangDuration = CurrentPresentDuration;
-	}
 #endif
 }
 
