Split the default state of FGameplayEventDescriptor to a separate non-property allocation

This avoids instances having a copy of the default data (which could be manipulated in the editor), and instead the default state is found via the archetype chain when needed at runtime. This does require that the EventContext type used by BroadcastEvent owns the FGameplayEventDescriptor property being used, and extra validation rules have been added to enforce this.

Note: This split was previously done for FGameplayEventFunction and follows that pattern.

#rnx

diff --git a/Engine/Source/Runtime/Core/Public/UObject/FortniteMainBranchObjectVersion.h b/Engine/Source/Runtime/Core/Public/UObject/FortniteMainBranchObjectVersion.h
--- a/Engine/Source/Runtime/Core/Public/UObject/FortniteMainBranchObjectVersion.h
+++ b/Engine/Source/Runtime/Core/Public/UObject/FortniteMainBranchObjectVersion.h
@@ -205,6 +205,9 @@ struct CORE_API FFortniteMainBranchObjectVersion
 		// Storing inertia tensor as vec3 instead of matrix.
 		ChaosInertiaConvertedToVec3,
 
+		// Migrated event definitions to the CDO/archetype data
+		MigratedEventDefinitionToDefaults,
+
 		// -----<new versions can be added above this line>-------------------------------------------------
 		VersionPlusOne,
 		LatestVersion = VersionPlusOne - 1
