Report monitor resolution and scene viewport size in the log.

#rnx

diff --git a/Engine/Source/Runtime/Core/Private/Windows/WindowsPlatformMisc.cpp b/Engine/Source/Runtime/Core/Private/Windows/WindowsPlatformMisc.cpp
--- a/Engine/Source/Runtime/Core/Private/Windows/WindowsPlatformMisc.cpp
+++ b/Engine/Source/Runtime/Core/Private/Windows/WindowsPlatformMisc.cpp
@@ -2505,8 +2505,33 @@ static void GetVideoDriverDetails(const FString& Key, FGPUDriverInfo& Out)
 	FWindowsPlatformMisc::QueryRegKey(HKEY_LOCAL_MACHINE, *Key, TEXT("DriverDate"), Out.DriverDate);
 }
 
+static BOOL CALLBACK MonitorEnumProc(HMONITOR Monitor, HDC MonitorDC, LPRECT Rect, LPARAM UserData)
+{
+	int* NumMonitors = (int*)UserData;
+	*NumMonitors += 1;
+
+	MONITORINFOEX MonitorInfoEx;
+	MonitorInfoEx.cbSize = sizeof(MonitorInfoEx);
+	GetMonitorInfo(Monitor, &MonitorInfoEx);
+
+	UE_LOG(LogWindows, Log, TEXT("    resolution: %dx%d, work area: (%d, %d) -> (%d, %d), device: '%s'%s"),
+		MonitorInfoEx.rcMonitor.right - MonitorInfoEx.rcMonitor.left, MonitorInfoEx.rcMonitor.bottom - MonitorInfoEx.rcMonitor.top,
+		MonitorInfoEx.rcWork.left, MonitorInfoEx.rcWork.top, MonitorInfoEx.rcWork.right, MonitorInfoEx.rcWork.bottom,
+		MonitorInfoEx.szDevice,
+		MonitorInfoEx.dwFlags & MONITORINFOF_PRIMARY ? TEXT(" [PRIMARY]") : TEXT("")
+	);
+
+	return TRUE;
+}
+
 FGPUDriverInfo FWindowsPlatformMisc::GetGPUDriverInfo(const FString& DeviceDescription)
 {
+	// Also report monitor information here, for lack of a better place.
+	UE_LOG(LogWindows, Log, TEXT("Attached monitors:"));
+	int NumMonitors = 0;
+	EnumDisplayMonitors(nullptr, nullptr, MonitorEnumProc, (LPARAM)&NumMonitors);
+	UE_LOG(LogWindows, Log, TEXT("Found %d attached monitors."), NumMonitors);
+
 	// to distinguish failed GetGPUDriverInfo() from call to GetGPUDriverInfo()
 	FGPUDriverInfo Ret;
 
@@ -2539,11 +2564,12 @@ FGPUDriverInfo FWindowsPlatformMisc::GetGPUDriverInfo(const FString& DeviceDescr
 				break;
 			}
 
-			UE_LOG(LogWindows, Log, TEXT("   %d. '%s' (P:%d D:%d)"),
+			UE_LOG(LogWindows, Log, TEXT("   %d. '%s' (P:%d D:%d), name: '%s'"),
 				i,
 				Device.DeviceString,
 				(Device.StateFlags & DISPLAY_DEVICE_PRIMARY_DEVICE) != 0,
-				(Device.StateFlags & DISPLAY_DEVICE_ATTACHED_TO_DESKTOP) != 0
+				(Device.StateFlags & DISPLAY_DEVICE_ATTACHED_TO_DESKTOP) != 0,
+				Device.DeviceName
 				);
 
 			if(Method == 3)
diff --git a/Engine/Source/Runtime/Engine/Private/Slate/SceneViewport.cpp b/Engine/Source/Runtime/Engine/Private/Slate/SceneViewport.cpp
--- a/Engine/Source/Runtime/Engine/Private/Slate/SceneViewport.cpp
+++ b/Engine/Source/Runtime/Engine/Private/Slate/SceneViewport.cpp
@@ -1522,6 +1522,17 @@ void FSceneViewport::ResizeViewport(uint32 NewSizeX, uint32 NewSizeY, EWindowMod
 			}
 		}
 
+		const TCHAR* WindowModeName;
+		switch (NewWindowMode)
+		{
+		case EWindowMode::Fullscreen: WindowModeName = TEXT("Fullscreen"); break;
+		case EWindowMode::WindowedFullscreen: WindowModeName = TEXT("WindowedFullscreen"); break;
+		case EWindowMode::Windowed: WindowModeName = TEXT("Windowed"); break;
+		default: WindowModeName = TEXT("INVALID"); break;
+		}
+
+		UE_LOG(LogViewport, Log, TEXT("Scene viewport resized to %dx%d, mode %s."), NewSizeX, NewSizeY, WindowModeName);
+
 		bIsResizing = true;
 
 		UpdateViewportRHI(false, NewSizeX, NewSizeY, NewWindowMode, PF_Unknown);
