[Android] Enable Infinite loading screen detection. Enables minimal hang detector. Sets hang detector to 30s, which seems t be required to enable loading screen hang detection. Set loading screen hang detection to 3 minutes to start. Have seen normal loading screens take ~90 seconds.

diff --git a/Engine/Source/Runtime/Core/Public/Android/AndroidPlatform.h b/Engine/Source/Runtime/Core/Public/Android/AndroidPlatform.h
--- a/Engine/Source/Runtime/Core/Public/Android/AndroidPlatform.h
+++ b/Engine/Source/Runtime/Core/Public/Android/AndroidPlatform.h
@@ -62,7 +62,7 @@ typedef FAndroidTypes FPlatformTypes;
 #define PLATFORM_SUPPORTS_VIRTUAL_TEXTURE_STREAMING	1
 #define PLATFORM_USE_ANSI_POSIX_MALLOC				1
 #define PLATFORM_SUPPORTS_LANDSCAPE_VISUAL_MESH_LOD_STREAMING 1
-
+#define PLATFORM_USE_MINIMAL_HANG_DETECTION			1
 #define PLATFORM_CODE_SECTION(Name)					__attribute__((section(Name)))
 
 #define PLATFORM_GLOBAL_LOG_CATEGORY				LogAndroid
diff --git a/Engine/Source/Runtime/Core/Public/Android/AndroidPlatformMisc.h b/Engine/Source/Runtime/Core/Public/Android/AndroidPlatformMisc.h
--- a/Engine/Source/Runtime/Core/Public/Android/AndroidPlatformMisc.h
+++ b/Engine/Source/Runtime/Core/Public/Android/AndroidPlatformMisc.h
@@ -59,7 +59,7 @@ public:
 
 	static bool AllowThreadHeartBeat()
 	{
-		return false;
+		return true;
 	}
 
 	struct FCPUStatTime{
