#ue5 - Added VectorMod360(), for specific use by Rotator normalization where speed is more important than the stability of stock fmod() under extreme ranges. 30% - 75% faster depending on the intrinsics available to the alternate VectorMod().

Given that the second argument is always 360, we can remove some tests for bad ranges of input and use simple truncation/multiplication to compute the result, while retaining good accuracy within reasonable FRotator input value ranges.

diff --git a/Engine/Source/Runtime/Core/Private/Math/UnrealMath.cpp b/Engine/Source/Runtime/Core/Private/Math/UnrealMath.cpp
--- a/Engine/Source/Runtime/Core/Private/Math/UnrealMath.cpp
+++ b/Engine/Source/Runtime/Core/Private/Math/UnrealMath.cpp
@@ -381,7 +381,7 @@ FQuat4f FRotator3f::Quaternion() const
 
 #if PLATFORM_ENABLE_VECTORINTRINSICS
 	const VectorRegister4Float Angles = MakeVectorRegisterFloat(Pitch, Yaw, Roll, 0.0f);
-	const VectorRegister4Float AnglesNoWinding = VectorMod(Angles, GlobalVectorConstants::Float360);
+	const VectorRegister4Float AnglesNoWinding = VectorMod360(Angles);
 	const VectorRegister4Float HalfAngles = VectorMultiply(AnglesNoWinding, GlobalVectorConstants::DEG_TO_RAD_HALF);
 
 	VectorRegister4Float SinAngles, CosAngles;
@@ -450,7 +450,7 @@ FQuat4d FRotator3d::Quaternion() const
 
 #if PLATFORM_ENABLE_VECTORINTRINSICS
 	const VectorRegister4Double Angles = MakeVectorRegisterDouble(Pitch, Yaw, Roll, 0.0);
-	const VectorRegister4Double AnglesNoWinding = VectorMod(Angles, GlobalVectorConstants::Double360);
+	const VectorRegister4Double AnglesNoWinding = VectorMod360(Angles);
 	const VectorRegister4Double HalfAngles = VectorMultiply(AnglesNoWinding, GlobalVectorConstants::DOUBLE_DEG_TO_RAD_HALF);
 
 	VectorRegister4Double SinAngles, CosAngles;
diff --git a/Engine/Source/Runtime/Core/Private/Tests/Math/UnrealMathTest.cpp b/Engine/Source/Runtime/Core/Private/Tests/Math/UnrealMathTest.cpp
--- a/Engine/Source/Runtime/Core/Private/Tests/Math/UnrealMathTest.cpp
+++ b/Engine/Source/Runtime/Core/Private/Tests/Math/UnrealMathTest.cpp
@@ -1943,6 +1943,31 @@ bool RunDoubleVectorTest()
 	V3 = VectorMod(V0, V1);
 	LogTest<double>(TEXT("VectorMod common"), TestVectorsEqual(V2, V3));
 
+	// VectorMod360
+	V0 = MakeVectorRegister(89.9, 180.0, -256.0, -270.1);
+	V1 = GlobalVectorConstants::Double360;
+	V2 = TestReferenceMod(V0, V1);
+	V3 = VectorMod360(V0);
+	LogTest<double>(TEXT("VectorMod360"), TestVectorsEqual(V2, V3));
+
+	V0 = MakeVectorRegister(0.0, 1079.9, -1720.1, 12345.12345);
+	V1 = GlobalVectorConstants::Double360;
+	V2 = TestReferenceMod(V0, V1);
+	V3 = VectorMod360(V0);
+	LogTest<double>(TEXT("VectorMod360"), TestVectorsEqual(V2, V3));
+
+#if UE_BUILD_DEBUG
+	V1 = GlobalVectorConstants::Double360;
+	double RotStep = 1.01;
+	for (double F = -720.0; F <= 720.0; F += RotStep)
+	{
+		V0 = MakeVectorRegister(F, F + 0.001, F + 0.025, F + 0.6125);
+		V2 = TestReferenceMod(V0, V1);
+		V3 = VectorMod360(V0);
+		LogTest<double>(TEXT("VectorMod360"), TestVectorsEqual(V2, V3));
+	}
+#endif // UE_BUILD_DEBUG
+
 	// VectorSign
 	V0 = MakeVectorRegister(2.0, -2.0, 0.0, -3.0);
 	V2 = MakeVectorRegister(1.0, -1.0, 1.0, -1.0);
@@ -2710,6 +2735,31 @@ bool FVectorRegisterAbstractionTest::RunTest(const FString& Parameters)
 	V3 = VectorMod(V0, V1);
 	LogTest<float>(TEXT("VectorMod common"), TestVectorsEqual(V2, V3));
 
+	// VectorMod360
+	V0 = MakeVectorRegister(89.9f, 180.0f, -256.0f, -270.1f);
+	V1 = GlobalVectorConstants::Float360;
+	V2 = TestReferenceMod(V0, V1);
+	V3 = VectorMod360(V0);
+	LogTest<float>(TEXT("VectorMod360"), TestVectorsEqual(V2, V3));
+
+	V0 = MakeVectorRegister(0.0f, -1079.9f, -1720.1f, 12345.12345f);
+	V1 = GlobalVectorConstants::Float360;
+	V2 = TestReferenceMod(V0, V1);
+	V3 = VectorMod360(V0);
+	LogTest<float>(TEXT("VectorMod360"), TestVectorsEqual(V2, V3));
+
+#if UE_BUILD_DEBUG
+	V1 = GlobalVectorConstants::Float360;
+	float RotStep = 1.01f;
+	for (float F = -720.f; F <= 720.f; F += RotStep)
+	{
+		V0 = MakeVectorRegister(F, F + 0.001f, F + 0.025f, F + 0.6125f);
+		V2 = TestReferenceMod(V0, V1);
+		V3 = VectorMod360(V0);
+		LogTest<float>(TEXT("VectorMod360"), TestVectorsEqual(V2, V3));
+	}
+#endif // UE_BUILD_DEBUG
+
 	// VectorSign
 	V0 = MakeVectorRegister(2.0f, -2.0f, 0.0f, -3.0f);
 	V2 = MakeVectorRegister(1.0f, -1.0f, 1.0f, -1.0f);
@@ -3518,7 +3568,7 @@ bool FVectorRegisterAbstractionTest::RunTest(const FString& Parameters)
 
 	if (!GPassing)
 	{
-		UE_LOG(LogUnrealMathTest, Fatal, TEXT("VectorIntrinsics Failed."));
+		UE_LOG(LogUnrealMathTest, Fatal, TEXT("VectorIntrinsics <float> Failed."));
 	}
 
 	if (!RunDoubleVectorTest())
diff --git a/Engine/Source/Runtime/Core/Public/Math/UnrealMathVectorCommon.h b/Engine/Source/Runtime/Core/Public/Math/UnrealMathVectorCommon.h
--- a/Engine/Source/Runtime/Core/Public/Math/UnrealMathVectorCommon.h
+++ b/Engine/Source/Runtime/Core/Public/Math/UnrealMathVectorCommon.h
@@ -594,11 +594,28 @@ FORCEINLINE VectorRegister4Double VectorNormalizeQuaternion(const VectorRegister
 	return VectorNormalizeSafe(UnnormalizedQuat, GlobalVectorConstants::Double0001);
 }
 
+// VectorMod360: Essentially VectorMod(X, 360) but using faster computation that is still accurate given the known input constraints.
+FORCEINLINE VectorRegister4Float VectorMod360(const VectorRegister4Float& X)
+{
+	// R = X - (Y * (Trunc(X / Y))
+	VectorRegister4Float Temp = VectorTruncate(VectorDivide(X, GlobalVectorConstants::Float360));
+	VectorRegister4Float FloatResult = VectorNegateMultiplyAdd(GlobalVectorConstants::Float360, Temp, X);
+	return FloatResult;
+}
+
+FORCEINLINE VectorRegister4Double VectorMod360(const VectorRegister4Double& X)
+{
+	// R = X - (Y * (Trunc(X / Y))
+	VectorRegister4Double Temp = VectorTruncate(VectorDivide(X, GlobalVectorConstants::Double360));
+	VectorRegister4Double DoubleResult = VectorNegateMultiplyAdd(GlobalVectorConstants::Double360, Temp, X);
+	return DoubleResult;
+}
+
 // Normalize Rotator
 FORCEINLINE VectorRegister4Float VectorNormalizeRotator(const VectorRegister4Float& UnnormalizedRotator)
 {
 	// shift in the range [-360,360]
-	VectorRegister4Float V0	= VectorMod( UnnormalizedRotator, GlobalVectorConstants::Float360);
+	VectorRegister4Float V0	= VectorMod360(UnnormalizedRotator);
 	VectorRegister4Float V1	= VectorAdd( V0, GlobalVectorConstants::Float360 );
 	VectorRegister4Float V2	= VectorSelect(VectorCompareGE(V0, VectorZeroFloat()), V0, V1);
 
@@ -612,7 +629,7 @@ FORCEINLINE VectorRegister4Float VectorNormalizeRotator(const VectorRegister4Flo
 FORCEINLINE VectorRegister4Double VectorNormalizeRotator(const VectorRegister4Double& UnnormalizedRotator)
 {
 	// shift in the range [-360,360]
-	VectorRegister4Double V0 = VectorMod(UnnormalizedRotator, GlobalVectorConstants::Double360);
+	VectorRegister4Double V0 = VectorMod360(UnnormalizedRotator);
 	VectorRegister4Double V1 = VectorAdd(V0, GlobalVectorConstants::Double360);
 	VectorRegister4Double V2 = VectorSelect(VectorCompareGE(V0, VectorZeroDouble()), V0, V1);
 
