Faster AES and SHA1 implementations using architecture speciifc instructions.

On x86 uses AESNI and SHA1 instructions detected with cpuid at runtime.
On ARMv8 uses NEON crypto instructions when enabled in compiler.

diff --git a/Engine/Source/Runtime/Core/Private/Misc/AES.cpp b/Engine/Source/Runtime/Core/Private/Misc/AES.cpp
--- a/Engine/Source/Runtime/Core/Private/Misc/AES.cpp
+++ b/Engine/Source/Runtime/Core/Private/Misc/AES.cpp
@@ -4,1133 +4,966 @@
 #include "Misc/AssertionMacros.h"
 #include "Misc/CString.h"
 #include "Misc/CoreDelegates.h"
+#include "HAL/PlatformMemory.h"
 
-// This is using the reference implementation of rijndael encryption algorithm
-// http://www.efgh.com/software/rijndael.htm
+#if defined(_M_AMD64) || defined(__x86_64__)
 
-// This approach has the shortcoming that it encrypts and decrypts each 128-bit block separately. If the plaintext contains identical 128-byte blocks, as many text files do, 
-// the blocks will be encrypted identically. This makes some of the plaintext structure visible in the ciphertext, even to someone who does not have the key. 
-// The usual practice is to combine each block after the first with the previous blocks (usually by some kind of XOR operation) before encrypting it. This hides repeated 
-// blocks very effectively, but it can wreak havoc if even one block of ciphertext is corrupted. The corrupted block AND ALL SUBSEQUENT BLOCKS will become unreadable
+#include <wmmintrin.h>
+#if PLATFORM_COMPILER_CLANG
+#include <cpuid.h>
+#endif
 
-// The approach used in CryptoPP is better, but much more complicated.
+#define UE_PLATFORM_AES_X86 1
 
-// It may be useful to integrate this into a class derived from FArchive to allow streaming encryption/decryption - but that's a mission for another day.
+// if x86 AESNI is enabled with compiler, use AESNI without fallback
+#if defined(__AES__) || PLATFORM_ALWAYS_HAS_AESNI
+#define UE_PLATFORM_AES_FALLBACK 0
+#endif
 
-#define TEST_ENCRYPTION			0
-#define AES_KEYBITS				256
+#elif defined(__aarch64__)
 
-#define KEYLENGTH( keybits )		( ( keybits ) / 8 )
-#define RKLENGTH( keybits )		( ( keybits ) / 8 + 28 )
-#define NROUNDS( keybits )			( ( keybits ) / 32 + 6 )
+// on ARMv8 enable AES instructions unconditionally only when they are enabled with compiler
+#ifdef __ARM_FEATURE_CRYPTO
+#define UE_PLATFORM_AES_ARMV8 1
+#define UE_PLATFORM_AES_FALLBACK 0
+#endif
 
-static const uint32 Te0[256] =
-{
-	0xc66363a5U, 0xf87c7c84U, 0xee777799U, 0xf67b7b8dU,
-	0xfff2f20dU, 0xd66b6bbdU, 0xde6f6fb1U, 0x91c5c554U,
-	0x60303050U, 0x02010103U, 0xce6767a9U, 0x562b2b7dU,
-	0xe7fefe19U, 0xb5d7d762U, 0x4dababe6U, 0xec76769aU,
-	0x8fcaca45U, 0x1f82829dU, 0x89c9c940U, 0xfa7d7d87U,
-	0xeffafa15U, 0xb25959ebU, 0x8e4747c9U, 0xfbf0f00bU,
-	0x41adadecU, 0xb3d4d467U, 0x5fa2a2fdU, 0x45afafeaU,
-	0x239c9cbfU, 0x53a4a4f7U, 0xe4727296U, 0x9bc0c05bU,
-	0x75b7b7c2U, 0xe1fdfd1cU, 0x3d9393aeU, 0x4c26266aU,
-	0x6c36365aU, 0x7e3f3f41U, 0xf5f7f702U, 0x83cccc4fU,
-	0x6834345cU, 0x51a5a5f4U, 0xd1e5e534U, 0xf9f1f108U,
-	0xe2717193U, 0xabd8d873U, 0x62313153U, 0x2a15153fU,
-	0x0804040cU, 0x95c7c752U, 0x46232365U, 0x9dc3c35eU,
-	0x30181828U, 0x379696a1U, 0x0a05050fU, 0x2f9a9ab5U,
-	0x0e070709U, 0x24121236U, 0x1b80809bU, 0xdfe2e23dU,
-	0xcdebeb26U, 0x4e272769U, 0x7fb2b2cdU, 0xea75759fU,
-	0x1209091bU, 0x1d83839eU, 0x582c2c74U, 0x341a1a2eU,
-	0x361b1b2dU, 0xdc6e6eb2U, 0xb45a5aeeU, 0x5ba0a0fbU,
-	0xa45252f6U, 0x763b3b4dU, 0xb7d6d661U, 0x7db3b3ceU,
-	0x5229297bU, 0xdde3e33eU, 0x5e2f2f71U, 0x13848497U,
-	0xa65353f5U, 0xb9d1d168U, 0x00000000U, 0xc1eded2cU,
-	0x40202060U, 0xe3fcfc1fU, 0x79b1b1c8U, 0xb65b5bedU,
-	0xd46a6abeU, 0x8dcbcb46U, 0x67bebed9U, 0x7239394bU,
-	0x944a4adeU, 0x984c4cd4U, 0xb05858e8U, 0x85cfcf4aU,
-	0xbbd0d06bU, 0xc5efef2aU, 0x4faaaae5U, 0xedfbfb16U,
-	0x864343c5U, 0x9a4d4dd7U, 0x66333355U, 0x11858594U,
-	0x8a4545cfU, 0xe9f9f910U, 0x04020206U, 0xfe7f7f81U,
-	0xa05050f0U, 0x783c3c44U, 0x259f9fbaU, 0x4ba8a8e3U,
-	0xa25151f3U, 0x5da3a3feU, 0x804040c0U, 0x058f8f8aU,
-	0x3f9292adU, 0x219d9dbcU, 0x70383848U, 0xf1f5f504U,
-	0x63bcbcdfU, 0x77b6b6c1U, 0xafdada75U, 0x42212163U,
-	0x20101030U, 0xe5ffff1aU, 0xfdf3f30eU, 0xbfd2d26dU,
-	0x81cdcd4cU, 0x180c0c14U, 0x26131335U, 0xc3ecec2fU,
-	0xbe5f5fe1U, 0x359797a2U, 0x884444ccU, 0x2e171739U,
-	0x93c4c457U, 0x55a7a7f2U, 0xfc7e7e82U, 0x7a3d3d47U,
-	0xc86464acU, 0xba5d5de7U, 0x3219192bU, 0xe6737395U,
-	0xc06060a0U, 0x19818198U, 0x9e4f4fd1U, 0xa3dcdc7fU,
-	0x44222266U, 0x542a2a7eU, 0x3b9090abU, 0x0b888883U,
-	0x8c4646caU, 0xc7eeee29U, 0x6bb8b8d3U, 0x2814143cU,
-	0xa7dede79U, 0xbc5e5ee2U, 0x160b0b1dU, 0xaddbdb76U,
-	0xdbe0e03bU, 0x64323256U, 0x743a3a4eU, 0x140a0a1eU,
-	0x924949dbU, 0x0c06060aU, 0x4824246cU, 0xb85c5ce4U,
-	0x9fc2c25dU, 0xbdd3d36eU, 0x43acacefU, 0xc46262a6U,
-	0x399191a8U, 0x319595a4U, 0xd3e4e437U, 0xf279798bU,
-	0xd5e7e732U, 0x8bc8c843U, 0x6e373759U, 0xda6d6db7U,
-	0x018d8d8cU, 0xb1d5d564U, 0x9c4e4ed2U, 0x49a9a9e0U,
-	0xd86c6cb4U, 0xac5656faU, 0xf3f4f407U, 0xcfeaea25U,
-	0xca6565afU, 0xf47a7a8eU, 0x47aeaee9U, 0x10080818U,
-	0x6fbabad5U, 0xf0787888U, 0x4a25256fU, 0x5c2e2e72U,
-	0x381c1c24U, 0x57a6a6f1U, 0x73b4b4c7U, 0x97c6c651U,
-	0xcbe8e823U, 0xa1dddd7cU, 0xe874749cU, 0x3e1f1f21U,
-	0x964b4bddU, 0x61bdbddcU, 0x0d8b8b86U, 0x0f8a8a85U,
-	0xe0707090U, 0x7c3e3e42U, 0x71b5b5c4U, 0xcc6666aaU,
-	0x904848d8U, 0x06030305U, 0xf7f6f601U, 0x1c0e0e12U,
-	0xc26161a3U, 0x6a35355fU, 0xae5757f9U, 0x69b9b9d0U,
-	0x17868691U, 0x99c1c158U, 0x3a1d1d27U, 0x279e9eb9U,
-	0xd9e1e138U, 0xebf8f813U, 0x2b9898b3U, 0x22111133U,
-	0xd26969bbU, 0xa9d9d970U, 0x078e8e89U, 0x339494a7U,
-	0x2d9b9bb6U, 0x3c1e1e22U, 0x15878792U, 0xc9e9e920U,
-	0x87cece49U, 0xaa5555ffU, 0x50282878U, 0xa5dfdf7aU,
-	0x038c8c8fU, 0x59a1a1f8U, 0x09898980U, 0x1a0d0d17U,
-	0x65bfbfdaU, 0xd7e6e631U, 0x844242c6U, 0xd06868b8U,
-	0x824141c3U, 0x299999b0U, 0x5a2d2d77U, 0x1e0f0f11U,
-	0x7bb0b0cbU, 0xa85454fcU, 0x6dbbbbd6U, 0x2c16163aU,
-};
+#endif
 
-static const uint32 Te1[256] =
-{
-	0xa5c66363U, 0x84f87c7cU, 0x99ee7777U, 0x8df67b7bU,
-	0x0dfff2f2U, 0xbdd66b6bU, 0xb1de6f6fU, 0x5491c5c5U,
-	0x50603030U, 0x03020101U, 0xa9ce6767U, 0x7d562b2bU,
-	0x19e7fefeU, 0x62b5d7d7U, 0xe64dababU, 0x9aec7676U,
-	0x458fcacaU, 0x9d1f8282U, 0x4089c9c9U, 0x87fa7d7dU,
-	0x15effafaU, 0xebb25959U, 0xc98e4747U, 0x0bfbf0f0U,
-	0xec41adadU, 0x67b3d4d4U, 0xfd5fa2a2U, 0xea45afafU,
-	0xbf239c9cU, 0xf753a4a4U, 0x96e47272U, 0x5b9bc0c0U,
-	0xc275b7b7U, 0x1ce1fdfdU, 0xae3d9393U, 0x6a4c2626U,
-	0x5a6c3636U, 0x417e3f3fU, 0x02f5f7f7U, 0x4f83ccccU,
-	0x5c683434U, 0xf451a5a5U, 0x34d1e5e5U, 0x08f9f1f1U,
-	0x93e27171U, 0x73abd8d8U, 0x53623131U, 0x3f2a1515U,
-	0x0c080404U, 0x5295c7c7U, 0x65462323U, 0x5e9dc3c3U,
-	0x28301818U, 0xa1379696U, 0x0f0a0505U, 0xb52f9a9aU,
-	0x090e0707U, 0x36241212U, 0x9b1b8080U, 0x3ddfe2e2U,
-	0x26cdebebU, 0x694e2727U, 0xcd7fb2b2U, 0x9fea7575U,
-	0x1b120909U, 0x9e1d8383U, 0x74582c2cU, 0x2e341a1aU,
-	0x2d361b1bU, 0xb2dc6e6eU, 0xeeb45a5aU, 0xfb5ba0a0U,
-	0xf6a45252U, 0x4d763b3bU, 0x61b7d6d6U, 0xce7db3b3U,
-	0x7b522929U, 0x3edde3e3U, 0x715e2f2fU, 0x97138484U,
-	0xf5a65353U, 0x68b9d1d1U, 0x00000000U, 0x2cc1ededU,
-	0x60402020U, 0x1fe3fcfcU, 0xc879b1b1U, 0xedb65b5bU,
-	0xbed46a6aU, 0x468dcbcbU, 0xd967bebeU, 0x4b723939U,
-	0xde944a4aU, 0xd4984c4cU, 0xe8b05858U, 0x4a85cfcfU,
-	0x6bbbd0d0U, 0x2ac5efefU, 0xe54faaaaU, 0x16edfbfbU,
-	0xc5864343U, 0xd79a4d4dU, 0x55663333U, 0x94118585U,
-	0xcf8a4545U, 0x10e9f9f9U, 0x06040202U, 0x81fe7f7fU,
-	0xf0a05050U, 0x44783c3cU, 0xba259f9fU, 0xe34ba8a8U,
-	0xf3a25151U, 0xfe5da3a3U, 0xc0804040U, 0x8a058f8fU,
-	0xad3f9292U, 0xbc219d9dU, 0x48703838U, 0x04f1f5f5U,
-	0xdf63bcbcU, 0xc177b6b6U, 0x75afdadaU, 0x63422121U,
-	0x30201010U, 0x1ae5ffffU, 0x0efdf3f3U, 0x6dbfd2d2U,
-	0x4c81cdcdU, 0x14180c0cU, 0x35261313U, 0x2fc3ececU,
-	0xe1be5f5fU, 0xa2359797U, 0xcc884444U, 0x392e1717U,
-	0x5793c4c4U, 0xf255a7a7U, 0x82fc7e7eU, 0x477a3d3dU,
-	0xacc86464U, 0xe7ba5d5dU, 0x2b321919U, 0x95e67373U,
-	0xa0c06060U, 0x98198181U, 0xd19e4f4fU, 0x7fa3dcdcU,
-	0x66442222U, 0x7e542a2aU, 0xab3b9090U, 0x830b8888U,
-	0xca8c4646U, 0x29c7eeeeU, 0xd36bb8b8U, 0x3c281414U,
-	0x79a7dedeU, 0xe2bc5e5eU, 0x1d160b0bU, 0x76addbdbU,
-	0x3bdbe0e0U, 0x56643232U, 0x4e743a3aU, 0x1e140a0aU,
-	0xdb924949U, 0x0a0c0606U, 0x6c482424U, 0xe4b85c5cU,
-	0x5d9fc2c2U, 0x6ebdd3d3U, 0xef43acacU, 0xa6c46262U,
-	0xa8399191U, 0xa4319595U, 0x37d3e4e4U, 0x8bf27979U,
-	0x32d5e7e7U, 0x438bc8c8U, 0x596e3737U, 0xb7da6d6dU,
-	0x8c018d8dU, 0x64b1d5d5U, 0xd29c4e4eU, 0xe049a9a9U,
-	0xb4d86c6cU, 0xfaac5656U, 0x07f3f4f4U, 0x25cfeaeaU,
-	0xafca6565U, 0x8ef47a7aU, 0xe947aeaeU, 0x18100808U,
-	0xd56fbabaU, 0x88f07878U, 0x6f4a2525U, 0x725c2e2eU,
-	0x24381c1cU, 0xf157a6a6U, 0xc773b4b4U, 0x5197c6c6U,
-	0x23cbe8e8U, 0x7ca1ddddU, 0x9ce87474U, 0x213e1f1fU,
-	0xdd964b4bU, 0xdc61bdbdU, 0x860d8b8bU, 0x850f8a8aU,
-	0x90e07070U, 0x427c3e3eU, 0xc471b5b5U, 0xaacc6666U,
-	0xd8904848U, 0x05060303U, 0x01f7f6f6U, 0x121c0e0eU,
-	0xa3c26161U, 0x5f6a3535U, 0xf9ae5757U, 0xd069b9b9U,
-	0x91178686U, 0x5899c1c1U, 0x273a1d1dU, 0xb9279e9eU,
-	0x38d9e1e1U, 0x13ebf8f8U, 0xb32b9898U, 0x33221111U,
-	0xbbd26969U, 0x70a9d9d9U, 0x89078e8eU, 0xa7339494U,
-	0xb62d9b9bU, 0x223c1e1eU, 0x92158787U, 0x20c9e9e9U,
-	0x4987ceceU, 0xffaa5555U, 0x78502828U, 0x7aa5dfdfU,
-	0x8f038c8cU, 0xf859a1a1U, 0x80098989U, 0x171a0d0dU,
-	0xda65bfbfU, 0x31d7e6e6U, 0xc6844242U, 0xb8d06868U,
-	0xc3824141U, 0xb0299999U, 0x775a2d2dU, 0x111e0f0fU,
-	0xcb7bb0b0U, 0xfca85454U, 0xd66dbbbbU, 0x3a2c1616U,
-};
 
-static const uint32 Te2[256] =
-{
-	0x63a5c663U, 0x7c84f87cU, 0x7799ee77U, 0x7b8df67bU,
-	0xf20dfff2U, 0x6bbdd66bU, 0x6fb1de6fU, 0xc55491c5U,
-	0x30506030U, 0x01030201U, 0x67a9ce67U, 0x2b7d562bU,
-	0xfe19e7feU, 0xd762b5d7U, 0xabe64dabU, 0x769aec76U,
-	0xca458fcaU, 0x829d1f82U, 0xc94089c9U, 0x7d87fa7dU,
-	0xfa15effaU, 0x59ebb259U, 0x47c98e47U, 0xf00bfbf0U,
-	0xadec41adU, 0xd467b3d4U, 0xa2fd5fa2U, 0xafea45afU,
-	0x9cbf239cU, 0xa4f753a4U, 0x7296e472U, 0xc05b9bc0U,
-	0xb7c275b7U, 0xfd1ce1fdU, 0x93ae3d93U, 0x266a4c26U,
-	0x365a6c36U, 0x3f417e3fU, 0xf702f5f7U, 0xcc4f83ccU,
-	0x345c6834U, 0xa5f451a5U, 0xe534d1e5U, 0xf108f9f1U,
-	0x7193e271U, 0xd873abd8U, 0x31536231U, 0x153f2a15U,
-	0x040c0804U, 0xc75295c7U, 0x23654623U, 0xc35e9dc3U,
-	0x18283018U, 0x96a13796U, 0x050f0a05U, 0x9ab52f9aU,
-	0x07090e07U, 0x12362412U, 0x809b1b80U, 0xe23ddfe2U,
-	0xeb26cdebU, 0x27694e27U, 0xb2cd7fb2U, 0x759fea75U,
-	0x091b1209U, 0x839e1d83U, 0x2c74582cU, 0x1a2e341aU,
-	0x1b2d361bU, 0x6eb2dc6eU, 0x5aeeb45aU, 0xa0fb5ba0U,
-	0x52f6a452U, 0x3b4d763bU, 0xd661b7d6U, 0xb3ce7db3U,
-	0x297b5229U, 0xe33edde3U, 0x2f715e2fU, 0x84971384U,
-	0x53f5a653U, 0xd168b9d1U, 0x00000000U, 0xed2cc1edU,
-	0x20604020U, 0xfc1fe3fcU, 0xb1c879b1U, 0x5bedb65bU,
-	0x6abed46aU, 0xcb468dcbU, 0xbed967beU, 0x394b7239U,
-	0x4ade944aU, 0x4cd4984cU, 0x58e8b058U, 0xcf4a85cfU,
-	0xd06bbbd0U, 0xef2ac5efU, 0xaae54faaU, 0xfb16edfbU,
-	0x43c58643U, 0x4dd79a4dU, 0x33556633U, 0x85941185U,
-	0x45cf8a45U, 0xf910e9f9U, 0x02060402U, 0x7f81fe7fU,
-	0x50f0a050U, 0x3c44783cU, 0x9fba259fU, 0xa8e34ba8U,
-	0x51f3a251U, 0xa3fe5da3U, 0x40c08040U, 0x8f8a058fU,
-	0x92ad3f92U, 0x9dbc219dU, 0x38487038U, 0xf504f1f5U,
-	0xbcdf63bcU, 0xb6c177b6U, 0xda75afdaU, 0x21634221U,
-	0x10302010U, 0xff1ae5ffU, 0xf30efdf3U, 0xd26dbfd2U,
-	0xcd4c81cdU, 0x0c14180cU, 0x13352613U, 0xec2fc3ecU,
-	0x5fe1be5fU, 0x97a23597U, 0x44cc8844U, 0x17392e17U,
-	0xc45793c4U, 0xa7f255a7U, 0x7e82fc7eU, 0x3d477a3dU,
-	0x64acc864U, 0x5de7ba5dU, 0x192b3219U, 0x7395e673U,
-	0x60a0c060U, 0x81981981U, 0x4fd19e4fU, 0xdc7fa3dcU,
-	0x22664422U, 0x2a7e542aU, 0x90ab3b90U, 0x88830b88U,
-	0x46ca8c46U, 0xee29c7eeU, 0xb8d36bb8U, 0x143c2814U,
-	0xde79a7deU, 0x5ee2bc5eU, 0x0b1d160bU, 0xdb76addbU,
-	0xe03bdbe0U, 0x32566432U, 0x3a4e743aU, 0x0a1e140aU,
-	0x49db9249U, 0x060a0c06U, 0x246c4824U, 0x5ce4b85cU,
-	0xc25d9fc2U, 0xd36ebdd3U, 0xacef43acU, 0x62a6c462U,
-	0x91a83991U, 0x95a43195U, 0xe437d3e4U, 0x798bf279U,
-	0xe732d5e7U, 0xc8438bc8U, 0x37596e37U, 0x6db7da6dU,
-	0x8d8c018dU, 0xd564b1d5U, 0x4ed29c4eU, 0xa9e049a9U,
-	0x6cb4d86cU, 0x56faac56U, 0xf407f3f4U, 0xea25cfeaU,
-	0x65afca65U, 0x7a8ef47aU, 0xaee947aeU, 0x08181008U,
-	0xbad56fbaU, 0x7888f078U, 0x256f4a25U, 0x2e725c2eU,
-	0x1c24381cU, 0xa6f157a6U, 0xb4c773b4U, 0xc65197c6U,
-	0xe823cbe8U, 0xdd7ca1ddU, 0x749ce874U, 0x1f213e1fU,
-	0x4bdd964bU, 0xbddc61bdU, 0x8b860d8bU, 0x8a850f8aU,
-	0x7090e070U, 0x3e427c3eU, 0xb5c471b5U, 0x66aacc66U,
-	0x48d89048U, 0x03050603U, 0xf601f7f6U, 0x0e121c0eU,
-	0x61a3c261U, 0x355f6a35U, 0x57f9ae57U, 0xb9d069b9U,
-	0x86911786U, 0xc15899c1U, 0x1d273a1dU, 0x9eb9279eU,
-	0xe138d9e1U, 0xf813ebf8U, 0x98b32b98U, 0x11332211U,
-	0x69bbd269U, 0xd970a9d9U, 0x8e89078eU, 0x94a73394U,
-	0x9bb62d9bU, 0x1e223c1eU, 0x87921587U, 0xe920c9e9U,
-	0xce4987ceU, 0x55ffaa55U, 0x28785028U, 0xdf7aa5dfU,
-	0x8c8f038cU, 0xa1f859a1U, 0x89800989U, 0x0d171a0dU,
-	0xbfda65bfU, 0xe631d7e6U, 0x42c68442U, 0x68b8d068U,
-	0x41c38241U, 0x99b02999U, 0x2d775a2dU, 0x0f111e0fU,
-	0xb0cb7bb0U, 0x54fca854U, 0xbbd66dbbU, 0x163a2c16U,
-};
+#ifndef UE_PLATFORM_AES_FALLBACK
+#define UE_PLATFORM_AES_FALLBACK 1
+#endif
+#ifndef UE_PLATFORM_AES_X86
+#define UE_PLATFORM_AES_X86 0
+#endif
+#ifndef UE_PLATFORM_AES_ARMV8
+#define UE_PLATFORM_AES_ARMV8 0
+#endif
+
+
+#define AES256_ROUND_COUNT 14
+
+
+#if UE_PLATFORM_AES_FALLBACK
 
-static const uint32 Te3[256] =
+struct FAesExpandedKey
 {
-	0x6363a5c6U, 0x7c7c84f8U, 0x777799eeU, 0x7b7b8df6U,
-	0xf2f20dffU, 0x6b6bbdd6U, 0x6f6fb1deU, 0xc5c55491U,
-	0x30305060U, 0x01010302U, 0x6767a9ceU, 0x2b2b7d56U,
-	0xfefe19e7U, 0xd7d762b5U, 0xababe64dU, 0x76769aecU,
-	0xcaca458fU, 0x82829d1fU, 0xc9c94089U, 0x7d7d87faU,
-	0xfafa15efU, 0x5959ebb2U, 0x4747c98eU, 0xf0f00bfbU,
-	0xadadec41U, 0xd4d467b3U, 0xa2a2fd5fU, 0xafafea45U,
-	0x9c9cbf23U, 0xa4a4f753U, 0x727296e4U, 0xc0c05b9bU,
-	0xb7b7c275U, 0xfdfd1ce1U, 0x9393ae3dU, 0x26266a4cU,
-	0x36365a6cU, 0x3f3f417eU, 0xf7f702f5U, 0xcccc4f83U,
-	0x34345c68U, 0xa5a5f451U, 0xe5e534d1U, 0xf1f108f9U,
-	0x717193e2U, 0xd8d873abU, 0x31315362U, 0x15153f2aU,
-	0x04040c08U, 0xc7c75295U, 0x23236546U, 0xc3c35e9dU,
-	0x18182830U, 0x9696a137U, 0x05050f0aU, 0x9a9ab52fU,
-	0x0707090eU, 0x12123624U, 0x80809b1bU, 0xe2e23ddfU,
-	0xebeb26cdU, 0x2727694eU, 0xb2b2cd7fU, 0x75759feaU,
-	0x09091b12U, 0x83839e1dU, 0x2c2c7458U, 0x1a1a2e34U,
-	0x1b1b2d36U, 0x6e6eb2dcU, 0x5a5aeeb4U, 0xa0a0fb5bU,
-	0x5252f6a4U, 0x3b3b4d76U, 0xd6d661b7U, 0xb3b3ce7dU,
-	0x29297b52U, 0xe3e33eddU, 0x2f2f715eU, 0x84849713U,
-	0x5353f5a6U, 0xd1d168b9U, 0x00000000U, 0xeded2cc1U,
-	0x20206040U, 0xfcfc1fe3U, 0xb1b1c879U, 0x5b5bedb6U,
-	0x6a6abed4U, 0xcbcb468dU, 0xbebed967U, 0x39394b72U,
-	0x4a4ade94U, 0x4c4cd498U, 0x5858e8b0U, 0xcfcf4a85U,
-	0xd0d06bbbU, 0xefef2ac5U, 0xaaaae54fU, 0xfbfb16edU,
-	0x4343c586U, 0x4d4dd79aU, 0x33335566U, 0x85859411U,
-	0x4545cf8aU, 0xf9f910e9U, 0x02020604U, 0x7f7f81feU,
-	0x5050f0a0U, 0x3c3c4478U, 0x9f9fba25U, 0xa8a8e34bU,
-	0x5151f3a2U, 0xa3a3fe5dU, 0x4040c080U, 0x8f8f8a05U,
-	0x9292ad3fU, 0x9d9dbc21U, 0x38384870U, 0xf5f504f1U,
-	0xbcbcdf63U, 0xb6b6c177U, 0xdada75afU, 0x21216342U,
-	0x10103020U, 0xffff1ae5U, 0xf3f30efdU, 0xd2d26dbfU,
-	0xcdcd4c81U, 0x0c0c1418U, 0x13133526U, 0xecec2fc3U,
-	0x5f5fe1beU, 0x9797a235U, 0x4444cc88U, 0x1717392eU,
-	0xc4c45793U, 0xa7a7f255U, 0x7e7e82fcU, 0x3d3d477aU,
-	0x6464acc8U, 0x5d5de7baU, 0x19192b32U, 0x737395e6U,
-	0x6060a0c0U, 0x81819819U, 0x4f4fd19eU, 0xdcdc7fa3U,
-	0x22226644U, 0x2a2a7e54U, 0x9090ab3bU, 0x8888830bU,
-	0x4646ca8cU, 0xeeee29c7U, 0xb8b8d36bU, 0x14143c28U,
-	0xdede79a7U, 0x5e5ee2bcU, 0x0b0b1d16U, 0xdbdb76adU,
-	0xe0e03bdbU, 0x32325664U, 0x3a3a4e74U, 0x0a0a1e14U,
-	0x4949db92U, 0x06060a0cU, 0x24246c48U, 0x5c5ce4b8U,
-	0xc2c25d9fU, 0xd3d36ebdU, 0xacacef43U, 0x6262a6c4U,
-	0x9191a839U, 0x9595a431U, 0xe4e437d3U, 0x79798bf2U,
-	0xe7e732d5U, 0xc8c8438bU, 0x3737596eU, 0x6d6db7daU,
-	0x8d8d8c01U, 0xd5d564b1U, 0x4e4ed29cU, 0xa9a9e049U,
-	0x6c6cb4d8U, 0x5656faacU, 0xf4f407f3U, 0xeaea25cfU,
-	0x6565afcaU, 0x7a7a8ef4U, 0xaeaee947U, 0x08081810U,
-	0xbabad56fU, 0x787888f0U, 0x25256f4aU, 0x2e2e725cU,
-	0x1c1c2438U, 0xa6a6f157U, 0xb4b4c773U, 0xc6c65197U,
-	0xe8e823cbU, 0xdddd7ca1U, 0x74749ce8U, 0x1f1f213eU,
-	0x4b4bdd96U, 0xbdbddc61U, 0x8b8b860dU, 0x8a8a850fU,
-	0x707090e0U, 0x3e3e427cU, 0xb5b5c471U, 0x6666aaccU,
-	0x4848d890U, 0x03030506U, 0xf6f601f7U, 0x0e0e121cU,
-	0x6161a3c2U, 0x35355f6aU, 0x5757f9aeU, 0xb9b9d069U,
-	0x86869117U, 0xc1c15899U, 0x1d1d273aU, 0x9e9eb927U,
-	0xe1e138d9U, 0xf8f813ebU, 0x9898b32bU, 0x11113322U,
-	0x6969bbd2U, 0xd9d970a9U, 0x8e8e8907U, 0x9494a733U,
-	0x9b9bb62dU, 0x1e1e223cU, 0x87879215U, 0xe9e920c9U,
-	0xcece4987U, 0x5555ffaaU, 0x28287850U, 0xdfdf7aa5U,
-	0x8c8c8f03U, 0xa1a1f859U, 0x89898009U, 0x0d0d171aU,
-	0xbfbfda65U, 0xe6e631d7U, 0x4242c684U, 0x6868b8d0U,
-	0x4141c382U, 0x9999b029U, 0x2d2d775aU, 0x0f0f111eU,
-	0xb0b0cb7bU, 0x5454fca8U, 0xbbbbd66dU, 0x16163a2cU,
+	uint32 Key[4 * (AES256_ROUND_COUNT + 1)];
 };
 
-static const uint32 Te4[256] =
+static const uint8 AesSbox[256] =
 {
-	0x63636363U, 0x7c7c7c7cU, 0x77777777U, 0x7b7b7b7bU,
-	0xf2f2f2f2U, 0x6b6b6b6bU, 0x6f6f6f6fU, 0xc5c5c5c5U,
-	0x30303030U, 0x01010101U, 0x67676767U, 0x2b2b2b2bU,
-	0xfefefefeU, 0xd7d7d7d7U, 0xababababU, 0x76767676U,
-	0xcacacacaU, 0x82828282U, 0xc9c9c9c9U, 0x7d7d7d7dU,
-	0xfafafafaU, 0x59595959U, 0x47474747U, 0xf0f0f0f0U,
-	0xadadadadU, 0xd4d4d4d4U, 0xa2a2a2a2U, 0xafafafafU,
-	0x9c9c9c9cU, 0xa4a4a4a4U, 0x72727272U, 0xc0c0c0c0U,
-	0xb7b7b7b7U, 0xfdfdfdfdU, 0x93939393U, 0x26262626U,
-	0x36363636U, 0x3f3f3f3fU, 0xf7f7f7f7U, 0xccccccccU,
-	0x34343434U, 0xa5a5a5a5U, 0xe5e5e5e5U, 0xf1f1f1f1U,
-	0x71717171U, 0xd8d8d8d8U, 0x31313131U, 0x15151515U,
-	0x04040404U, 0xc7c7c7c7U, 0x23232323U, 0xc3c3c3c3U,
-	0x18181818U, 0x96969696U, 0x05050505U, 0x9a9a9a9aU,
-	0x07070707U, 0x12121212U, 0x80808080U, 0xe2e2e2e2U,
-	0xebebebebU, 0x27272727U, 0xb2b2b2b2U, 0x75757575U,
-	0x09090909U, 0x83838383U, 0x2c2c2c2cU, 0x1a1a1a1aU,
-	0x1b1b1b1bU, 0x6e6e6e6eU, 0x5a5a5a5aU, 0xa0a0a0a0U,
-	0x52525252U, 0x3b3b3b3bU, 0xd6d6d6d6U, 0xb3b3b3b3U,
-	0x29292929U, 0xe3e3e3e3U, 0x2f2f2f2fU, 0x84848484U,
-	0x53535353U, 0xd1d1d1d1U, 0x00000000U, 0xededededU,
-	0x20202020U, 0xfcfcfcfcU, 0xb1b1b1b1U, 0x5b5b5b5bU,
-	0x6a6a6a6aU, 0xcbcbcbcbU, 0xbebebebeU, 0x39393939U,
-	0x4a4a4a4aU, 0x4c4c4c4cU, 0x58585858U, 0xcfcfcfcfU,
-	0xd0d0d0d0U, 0xefefefefU, 0xaaaaaaaaU, 0xfbfbfbfbU,
-	0x43434343U, 0x4d4d4d4dU, 0x33333333U, 0x85858585U,
-	0x45454545U, 0xf9f9f9f9U, 0x02020202U, 0x7f7f7f7fU,
-	0x50505050U, 0x3c3c3c3cU, 0x9f9f9f9fU, 0xa8a8a8a8U,
-	0x51515151U, 0xa3a3a3a3U, 0x40404040U, 0x8f8f8f8fU,
-	0x92929292U, 0x9d9d9d9dU, 0x38383838U, 0xf5f5f5f5U,
-	0xbcbcbcbcU, 0xb6b6b6b6U, 0xdadadadaU, 0x21212121U,
-	0x10101010U, 0xffffffffU, 0xf3f3f3f3U, 0xd2d2d2d2U,
-	0xcdcdcdcdU, 0x0c0c0c0cU, 0x13131313U, 0xececececU,
-	0x5f5f5f5fU, 0x97979797U, 0x44444444U, 0x17171717U,
-	0xc4c4c4c4U, 0xa7a7a7a7U, 0x7e7e7e7eU, 0x3d3d3d3dU,
-	0x64646464U, 0x5d5d5d5dU, 0x19191919U, 0x73737373U,
-	0x60606060U, 0x81818181U, 0x4f4f4f4fU, 0xdcdcdcdcU,
-	0x22222222U, 0x2a2a2a2aU, 0x90909090U, 0x88888888U,
-	0x46464646U, 0xeeeeeeeeU, 0xb8b8b8b8U, 0x14141414U,
-	0xdedededeU, 0x5e5e5e5eU, 0x0b0b0b0bU, 0xdbdbdbdbU,
-	0xe0e0e0e0U, 0x32323232U, 0x3a3a3a3aU, 0x0a0a0a0aU,
-	0x49494949U, 0x06060606U, 0x24242424U, 0x5c5c5c5cU,
-	0xc2c2c2c2U, 0xd3d3d3d3U, 0xacacacacU, 0x62626262U,
-	0x91919191U, 0x95959595U, 0xe4e4e4e4U, 0x79797979U,
-	0xe7e7e7e7U, 0xc8c8c8c8U, 0x37373737U, 0x6d6d6d6dU,
-	0x8d8d8d8dU, 0xd5d5d5d5U, 0x4e4e4e4eU, 0xa9a9a9a9U,
-	0x6c6c6c6cU, 0x56565656U, 0xf4f4f4f4U, 0xeaeaeaeaU,
-	0x65656565U, 0x7a7a7a7aU, 0xaeaeaeaeU, 0x08080808U,
-	0xbabababaU, 0x78787878U, 0x25252525U, 0x2e2e2e2eU,
-	0x1c1c1c1cU, 0xa6a6a6a6U, 0xb4b4b4b4U, 0xc6c6c6c6U,
-	0xe8e8e8e8U, 0xddddddddU, 0x74747474U, 0x1f1f1f1fU,
-	0x4b4b4b4bU, 0xbdbdbdbdU, 0x8b8b8b8bU, 0x8a8a8a8aU,
-	0x70707070U, 0x3e3e3e3eU, 0xb5b5b5b5U, 0x66666666U,
-	0x48484848U, 0x03030303U, 0xf6f6f6f6U, 0x0e0e0e0eU,
-	0x61616161U, 0x35353535U, 0x57575757U, 0xb9b9b9b9U,
-	0x86868686U, 0xc1c1c1c1U, 0x1d1d1d1dU, 0x9e9e9e9eU,
-	0xe1e1e1e1U, 0xf8f8f8f8U, 0x98989898U, 0x11111111U,
-	0x69696969U, 0xd9d9d9d9U, 0x8e8e8e8eU, 0x94949494U,
-	0x9b9b9b9bU, 0x1e1e1e1eU, 0x87878787U, 0xe9e9e9e9U,
-	0xcecececeU, 0x55555555U, 0x28282828U, 0xdfdfdfdfU,
-	0x8c8c8c8cU, 0xa1a1a1a1U, 0x89898989U, 0x0d0d0d0dU,
-	0xbfbfbfbfU, 0xe6e6e6e6U, 0x42424242U, 0x68686868U,
-	0x41414141U, 0x99999999U, 0x2d2d2d2dU, 0x0f0f0f0fU,
-	0xb0b0b0b0U, 0x54545454U, 0xbbbbbbbbU, 0x16161616U,
+	0x63, 0x7c, 0x77, 0x7b, 0xf2, 0x6b, 0x6f, 0xc5, 0x30, 0x01, 0x67, 0x2b, 0xfe, 0xd7, 0xab, 0x76,
+	0xca, 0x82, 0xc9, 0x7d, 0xfa, 0x59, 0x47, 0xf0, 0xad, 0xd4, 0xa2, 0xaf, 0x9c, 0xa4, 0x72, 0xc0,
+	0xb7, 0xfd, 0x93, 0x26, 0x36, 0x3f, 0xf7, 0xcc, 0x34, 0xa5, 0xe5, 0xf1, 0x71, 0xd8, 0x31, 0x15,
+	0x04, 0xc7, 0x23, 0xc3, 0x18, 0x96, 0x05, 0x9a, 0x07, 0x12, 0x80, 0xe2, 0xeb, 0x27, 0xb2, 0x75,
+	0x09, 0x83, 0x2c, 0x1a, 0x1b, 0x6e, 0x5a, 0xa0, 0x52, 0x3b, 0xd6, 0xb3, 0x29, 0xe3, 0x2f, 0x84,
+	0x53, 0xd1, 0x00, 0xed, 0x20, 0xfc, 0xb1, 0x5b, 0x6a, 0xcb, 0xbe, 0x39, 0x4a, 0x4c, 0x58, 0xcf,
+	0xd0, 0xef, 0xaa, 0xfb, 0x43, 0x4d, 0x33, 0x85, 0x45, 0xf9, 0x02, 0x7f, 0x50, 0x3c, 0x9f, 0xa8,
+	0x51, 0xa3, 0x40, 0x8f, 0x92, 0x9d, 0x38, 0xf5, 0xbc, 0xb6, 0xda, 0x21, 0x10, 0xff, 0xf3, 0xd2,
+	0xcd, 0x0c, 0x13, 0xec, 0x5f, 0x97, 0x44, 0x17, 0xc4, 0xa7, 0x7e, 0x3d, 0x64, 0x5d, 0x19, 0x73,
+	0x60, 0x81, 0x4f, 0xdc, 0x22, 0x2a, 0x90, 0x88, 0x46, 0xee, 0xb8, 0x14, 0xde, 0x5e, 0x0b, 0xdb,
+	0xe0, 0x32, 0x3a, 0x0a, 0x49, 0x06, 0x24, 0x5c, 0xc2, 0xd3, 0xac, 0x62, 0x91, 0x95, 0xe4, 0x79,
+	0xe7, 0xc8, 0x37, 0x6d, 0x8d, 0xd5, 0x4e, 0xa9, 0x6c, 0x56, 0xf4, 0xea, 0x65, 0x7a, 0xae, 0x08,
+	0xba, 0x78, 0x25, 0x2e, 0x1c, 0xa6, 0xb4, 0xc6, 0xe8, 0xdd, 0x74, 0x1f, 0x4b, 0xbd, 0x8b, 0x8a,
+	0x70, 0x3e, 0xb5, 0x66, 0x48, 0x03, 0xf6, 0x0e, 0x61, 0x35, 0x57, 0xb9, 0x86, 0xc1, 0x1d, 0x9e,
+	0xe1, 0xf8, 0x98, 0x11, 0x69, 0xd9, 0x8e, 0x94, 0x9b, 0x1e, 0x87, 0xe9, 0xce, 0x55, 0x28, 0xdf,
+	0x8c, 0xa1, 0x89, 0x0d, 0xbf, 0xe6, 0x42, 0x68, 0x41, 0x99, 0x2d, 0x0f, 0xb0, 0x54, 0xbb, 0x16,
 };
 
-static const uint32 Td0[256] =
+static const uint8 AesISbox[256] =
 {
-	0x51f4a750U, 0x7e416553U, 0x1a17a4c3U, 0x3a275e96U,
-	0x3bab6bcbU, 0x1f9d45f1U, 0xacfa58abU, 0x4be30393U,
-	0x2030fa55U, 0xad766df6U, 0x88cc7691U, 0xf5024c25U,
-	0x4fe5d7fcU, 0xc52acbd7U, 0x26354480U, 0xb562a38fU,
-	0xdeb15a49U, 0x25ba1b67U, 0x45ea0e98U, 0x5dfec0e1U,
-	0xc32f7502U, 0x814cf012U, 0x8d4697a3U, 0x6bd3f9c6U,
-	0x038f5fe7U, 0x15929c95U, 0xbf6d7aebU, 0x955259daU,
-	0xd4be832dU, 0x587421d3U, 0x49e06929U, 0x8ec9c844U,
-	0x75c2896aU, 0xf48e7978U, 0x99583e6bU, 0x27b971ddU,
-	0xbee14fb6U, 0xf088ad17U, 0xc920ac66U, 0x7dce3ab4U,
-	0x63df4a18U, 0xe51a3182U, 0x97513360U, 0x62537f45U,
-	0xb16477e0U, 0xbb6bae84U, 0xfe81a01cU, 0xf9082b94U,
-	0x70486858U, 0x8f45fd19U, 0x94de6c87U, 0x527bf8b7U,
-	0xab73d323U, 0x724b02e2U, 0xe31f8f57U, 0x6655ab2aU,
-	0xb2eb2807U, 0x2fb5c203U, 0x86c57b9aU, 0xd33708a5U,
-	0x302887f2U, 0x23bfa5b2U, 0x02036abaU, 0xed16825cU,
-	0x8acf1c2bU, 0xa779b492U, 0xf307f2f0U, 0x4e69e2a1U,
-	0x65daf4cdU, 0x0605bed5U, 0xd134621fU, 0xc4a6fe8aU,
-	0x342e539dU, 0xa2f355a0U, 0x058ae132U, 0xa4f6eb75U,
-	0x0b83ec39U, 0x4060efaaU, 0x5e719f06U, 0xbd6e1051U,
-	0x3e218af9U, 0x96dd063dU, 0xdd3e05aeU, 0x4de6bd46U,
-	0x91548db5U, 0x71c45d05U, 0x0406d46fU, 0x605015ffU,
-	0x1998fb24U, 0xd6bde997U, 0x894043ccU, 0x67d99e77U,
-	0xb0e842bdU, 0x07898b88U, 0xe7195b38U, 0x79c8eedbU,
-	0xa17c0a47U, 0x7c420fe9U, 0xf8841ec9U, 0x00000000U,
-	0x09808683U, 0x322bed48U, 0x1e1170acU, 0x6c5a724eU,
-	0xfd0efffbU, 0x0f853856U, 0x3daed51eU, 0x362d3927U,
-	0x0a0fd964U, 0x685ca621U, 0x9b5b54d1U, 0x24362e3aU,
-	0x0c0a67b1U, 0x9357e70fU, 0xb4ee96d2U, 0x1b9b919eU,
-	0x80c0c54fU, 0x61dc20a2U, 0x5a774b69U, 0x1c121a16U,
-	0xe293ba0aU, 0xc0a02ae5U, 0x3c22e043U, 0x121b171dU,
-	0x0e090d0bU, 0xf28bc7adU, 0x2db6a8b9U, 0x141ea9c8U,
-	0x57f11985U, 0xaf75074cU, 0xee99ddbbU, 0xa37f60fdU,
-	0xf701269fU, 0x5c72f5bcU, 0x44663bc5U, 0x5bfb7e34U,
-	0x8b432976U, 0xcb23c6dcU, 0xb6edfc68U, 0xb8e4f163U,
-	0xd731dccaU, 0x42638510U, 0x13972240U, 0x84c61120U,
-	0x854a247dU, 0xd2bb3df8U, 0xaef93211U, 0xc729a16dU,
-	0x1d9e2f4bU, 0xdcb230f3U, 0x0d8652ecU, 0x77c1e3d0U,
-	0x2bb3166cU, 0xa970b999U, 0x119448faU, 0x47e96422U,
-	0xa8fc8cc4U, 0xa0f03f1aU, 0x567d2cd8U, 0x223390efU,
-	0x87494ec7U, 0xd938d1c1U, 0x8ccaa2feU, 0x98d40b36U,
-	0xa6f581cfU, 0xa57ade28U, 0xdab78e26U, 0x3fadbfa4U,
-	0x2c3a9de4U, 0x5078920dU, 0x6a5fcc9bU, 0x547e4662U,
-	0xf68d13c2U, 0x90d8b8e8U, 0x2e39f75eU, 0x82c3aff5U,
-	0x9f5d80beU, 0x69d0937cU, 0x6fd52da9U, 0xcf2512b3U,
-	0xc8ac993bU, 0x10187da7U, 0xe89c636eU, 0xdb3bbb7bU,
-	0xcd267809U, 0x6e5918f4U, 0xec9ab701U, 0x834f9aa8U,
-	0xe6956e65U, 0xaaffe67eU, 0x21bccf08U, 0xef15e8e6U,
-	0xbae79bd9U, 0x4a6f36ceU, 0xea9f09d4U, 0x29b07cd6U,
-	0x31a4b2afU, 0x2a3f2331U, 0xc6a59430U, 0x35a266c0U,
-	0x744ebc37U, 0xfc82caa6U, 0xe090d0b0U, 0x33a7d815U,
-	0xf104984aU, 0x41ecdaf7U, 0x7fcd500eU, 0x1791f62fU,
-	0x764dd68dU, 0x43efb04dU, 0xccaa4d54U, 0xe49604dfU,
-	0x9ed1b5e3U, 0x4c6a881bU, 0xc12c1fb8U, 0x4665517fU,
-	0x9d5eea04U, 0x018c355dU, 0xfa877473U, 0xfb0b412eU,
-	0xb3671d5aU, 0x92dbd252U, 0xe9105633U, 0x6dd64713U,
-	0x9ad7618cU, 0x37a10c7aU, 0x59f8148eU, 0xeb133c89U,
-	0xcea927eeU, 0xb761c935U, 0xe11ce5edU, 0x7a47b13cU,
-	0x9cd2df59U, 0x55f2733fU, 0x1814ce79U, 0x73c737bfU,
-	0x53f7cdeaU, 0x5ffdaa5bU, 0xdf3d6f14U, 0x7844db86U,
-	0xcaaff381U, 0xb968c43eU, 0x3824342cU, 0xc2a3405fU,
-	0x161dc372U, 0xbce2250cU, 0x283c498bU, 0xff0d9541U,
-	0x39a80171U, 0x080cb3deU, 0xd8b4e49cU, 0x6456c190U,
-	0x7bcb8461U, 0xd532b670U, 0x486c5c74U, 0xd0b85742U,
+	0x52, 0x09, 0x6a, 0xd5, 0x30, 0x36, 0xa5, 0x38, 0xbf, 0x40, 0xa3, 0x9e, 0x81, 0xf3, 0xd7, 0xfb,
+	0x7c, 0xe3, 0x39, 0x82, 0x9b, 0x2f, 0xff, 0x87, 0x34, 0x8e, 0x43, 0x44, 0xc4, 0xde, 0xe9, 0xcb,
+	0x54, 0x7b, 0x94, 0x32, 0xa6, 0xc2, 0x23, 0x3d, 0xee, 0x4c, 0x95, 0x0b, 0x42, 0xfa, 0xc3, 0x4e,
+	0x08, 0x2e, 0xa1, 0x66, 0x28, 0xd9, 0x24, 0xb2, 0x76, 0x5b, 0xa2, 0x49, 0x6d, 0x8b, 0xd1, 0x25,
+	0x72, 0xf8, 0xf6, 0x64, 0x86, 0x68, 0x98, 0x16, 0xd4, 0xa4, 0x5c, 0xcc, 0x5d, 0x65, 0xb6, 0x92,
+	0x6c, 0x70, 0x48, 0x50, 0xfd, 0xed, 0xb9, 0xda, 0x5e, 0x15, 0x46, 0x57, 0xa7, 0x8d, 0x9d, 0x84,
+	0x90, 0xd8, 0xab, 0x00, 0x8c, 0xbc, 0xd3, 0x0a, 0xf7, 0xe4, 0x58, 0x05, 0xb8, 0xb3, 0x45, 0x06,
+	0xd0, 0x2c, 0x1e, 0x8f, 0xca, 0x3f, 0x0f, 0x02, 0xc1, 0xaf, 0xbd, 0x03, 0x01, 0x13, 0x8a, 0x6b,
+	0x3a, 0x91, 0x11, 0x41, 0x4f, 0x67, 0xdc, 0xea, 0x97, 0xf2, 0xcf, 0xce, 0xf0, 0xb4, 0xe6, 0x73,
+	0x96, 0xac, 0x74, 0x22, 0xe7, 0xad, 0x35, 0x85, 0xe2, 0xf9, 0x37, 0xe8, 0x1c, 0x75, 0xdf, 0x6e,
+	0x47, 0xf1, 0x1a, 0x71, 0x1d, 0x29, 0xc5, 0x89, 0x6f, 0xb7, 0x62, 0x0e, 0xaa, 0x18, 0xbe, 0x1b,
+	0xfc, 0x56, 0x3e, 0x4b, 0xc6, 0xd2, 0x79, 0x20, 0x9a, 0xdb, 0xc0, 0xfe, 0x78, 0xcd, 0x5a, 0xf4,
+	0x1f, 0xdd, 0xa8, 0x33, 0x88, 0x07, 0xc7, 0x31, 0xb1, 0x12, 0x10, 0x59, 0x27, 0x80, 0xec, 0x5f,
+	0x60, 0x51, 0x7f, 0xa9, 0x19, 0xb5, 0x4a, 0x0d, 0x2d, 0xe5, 0x7a, 0x9f, 0x93, 0xc9, 0x9c, 0xef,
+	0xa0, 0xe0, 0x3b, 0x4d, 0xae, 0x2a, 0xf5, 0xb0, 0xc8, 0xeb, 0xbb, 0x3c, 0x83, 0x53, 0x99, 0x61,
+	0x17, 0x2b, 0x04, 0x7e, 0xba, 0x77, 0xd6, 0x26, 0xe1, 0x69, 0x14, 0x63, 0x55, 0x21, 0x0c, 0x7d,
 };
 
-static const uint32 Td1[256] =
+static const uint32 AesTbox[4][256] =
 {
-	0x5051f4a7U, 0x537e4165U, 0xc31a17a4U, 0x963a275eU,
-	0xcb3bab6bU, 0xf11f9d45U, 0xabacfa58U, 0x934be303U,
-	0x552030faU, 0xf6ad766dU, 0x9188cc76U, 0x25f5024cU,
-	0xfc4fe5d7U, 0xd7c52acbU, 0x80263544U, 0x8fb562a3U,
-	0x49deb15aU, 0x6725ba1bU, 0x9845ea0eU, 0xe15dfec0U,
-	0x02c32f75U, 0x12814cf0U, 0xa38d4697U, 0xc66bd3f9U,
-	0xe7038f5fU, 0x9515929cU, 0xebbf6d7aU, 0xda955259U,
-	0x2dd4be83U, 0xd3587421U, 0x2949e069U, 0x448ec9c8U,
-	0x6a75c289U, 0x78f48e79U, 0x6b99583eU, 0xdd27b971U,
-	0xb6bee14fU, 0x17f088adU, 0x66c920acU, 0xb47dce3aU,
-	0x1863df4aU, 0x82e51a31U, 0x60975133U, 0x4562537fU,
-	0xe0b16477U, 0x84bb6baeU, 0x1cfe81a0U, 0x94f9082bU,
-	0x58704868U, 0x198f45fdU, 0x8794de6cU, 0xb7527bf8U,
-	0x23ab73d3U, 0xe2724b02U, 0x57e31f8fU, 0x2a6655abU,
-	0x07b2eb28U, 0x032fb5c2U, 0x9a86c57bU, 0xa5d33708U,
-	0xf2302887U, 0xb223bfa5U, 0xba02036aU, 0x5ced1682U,
-	0x2b8acf1cU, 0x92a779b4U, 0xf0f307f2U, 0xa14e69e2U,
-	0xcd65daf4U, 0xd50605beU, 0x1fd13462U, 0x8ac4a6feU,
-	0x9d342e53U, 0xa0a2f355U, 0x32058ae1U, 0x75a4f6ebU,
-	0x390b83ecU, 0xaa4060efU, 0x065e719fU, 0x51bd6e10U,
-	0xf93e218aU, 0x3d96dd06U, 0xaedd3e05U, 0x464de6bdU,
-	0xb591548dU, 0x0571c45dU, 0x6f0406d4U, 0xff605015U,
-	0x241998fbU, 0x97d6bde9U, 0xcc894043U, 0x7767d99eU,
-	0xbdb0e842U, 0x8807898bU, 0x38e7195bU, 0xdb79c8eeU,
-	0x47a17c0aU, 0xe97c420fU, 0xc9f8841eU, 0x00000000U,
-	0x83098086U, 0x48322bedU, 0xac1e1170U, 0x4e6c5a72U,
-	0xfbfd0effU, 0x560f8538U, 0x1e3daed5U, 0x27362d39U,
-	0x640a0fd9U, 0x21685ca6U, 0xd19b5b54U, 0x3a24362eU,
-	0xb10c0a67U, 0x0f9357e7U, 0xd2b4ee96U, 0x9e1b9b91U,
-	0x4f80c0c5U, 0xa261dc20U, 0x695a774bU, 0x161c121aU,
-	0x0ae293baU, 0xe5c0a02aU, 0x433c22e0U, 0x1d121b17U,
-	0x0b0e090dU, 0xadf28bc7U, 0xb92db6a8U, 0xc8141ea9U,
-	0x8557f119U, 0x4caf7507U, 0xbbee99ddU, 0xfda37f60U,
-	0x9ff70126U, 0xbc5c72f5U, 0xc544663bU, 0x345bfb7eU,
-	0x768b4329U, 0xdccb23c6U, 0x68b6edfcU, 0x63b8e4f1U,
-	0xcad731dcU, 0x10426385U, 0x40139722U, 0x2084c611U,
-	0x7d854a24U, 0xf8d2bb3dU, 0x11aef932U, 0x6dc729a1U,
-	0x4b1d9e2fU, 0xf3dcb230U, 0xec0d8652U, 0xd077c1e3U,
-	0x6c2bb316U, 0x99a970b9U, 0xfa119448U, 0x2247e964U,
-	0xc4a8fc8cU, 0x1aa0f03fU, 0xd8567d2cU, 0xef223390U,
-	0xc787494eU, 0xc1d938d1U, 0xfe8ccaa2U, 0x3698d40bU,
-	0xcfa6f581U, 0x28a57adeU, 0x26dab78eU, 0xa43fadbfU,
-	0xe42c3a9dU, 0x0d507892U, 0x9b6a5fccU, 0x62547e46U,
-	0xc2f68d13U, 0xe890d8b8U, 0x5e2e39f7U, 0xf582c3afU,
-	0xbe9f5d80U, 0x7c69d093U, 0xa96fd52dU, 0xb3cf2512U,
-	0x3bc8ac99U, 0xa710187dU, 0x6ee89c63U, 0x7bdb3bbbU,
-	0x09cd2678U, 0xf46e5918U, 0x01ec9ab7U, 0xa8834f9aU,
-	0x65e6956eU, 0x7eaaffe6U, 0x0821bccfU, 0xe6ef15e8U,
-	0xd9bae79bU, 0xce4a6f36U, 0xd4ea9f09U, 0xd629b07cU,
-	0xaf31a4b2U, 0x312a3f23U, 0x30c6a594U, 0xc035a266U,
-	0x37744ebcU, 0xa6fc82caU, 0xb0e090d0U, 0x1533a7d8U,
-	0x4af10498U, 0xf741ecdaU, 0x0e7fcd50U, 0x2f1791f6U,
-	0x8d764dd6U, 0x4d43efb0U, 0x54ccaa4dU, 0xdfe49604U,
-	0xe39ed1b5U, 0x1b4c6a88U, 0xb8c12c1fU, 0x7f466551U,
-	0x049d5eeaU, 0x5d018c35U, 0x73fa8774U, 0x2efb0b41U,
-	0x5ab3671dU, 0x5292dbd2U, 0x33e91056U, 0x136dd647U,
-	0x8c9ad761U, 0x7a37a10cU, 0x8e59f814U, 0x89eb133cU,
-	0xeecea927U, 0x35b761c9U, 0xede11ce5U, 0x3c7a47b1U,
-	0x599cd2dfU, 0x3f55f273U, 0x791814ceU, 0xbf73c737U,
-	0xea53f7cdU, 0x5b5ffdaaU, 0x14df3d6fU, 0x867844dbU,
-	0x81caaff3U, 0x3eb968c4U, 0x2c382434U, 0x5fc2a340U,
-	0x72161dc3U, 0x0cbce225U, 0x8b283c49U, 0x41ff0d95U,
-	0x7139a801U, 0xde080cb3U, 0x9cd8b4e4U, 0x906456c1U,
-	0x617bcb84U, 0x70d532b6U, 0x74486c5cU, 0x42d0b857U,
+	{
+		0xa56363c6, 0x847c7cf8, 0x997777ee, 0x8d7b7bf6, 0x0df2f2ff, 0xbd6b6bd6, 0xb16f6fde, 0x54c5c591,
+		0x50303060, 0x03010102, 0xa96767ce, 0x7d2b2b56, 0x19fefee7, 0x62d7d7b5, 0xe6abab4d, 0x9a7676ec,
+		0x45caca8f, 0x9d82821f, 0x40c9c989, 0x877d7dfa, 0x15fafaef, 0xeb5959b2, 0xc947478e, 0x0bf0f0fb,
+		0xecadad41, 0x67d4d4b3, 0xfda2a25f, 0xeaafaf45, 0xbf9c9c23, 0xf7a4a453, 0x967272e4, 0x5bc0c09b,
+		0xc2b7b775, 0x1cfdfde1, 0xae93933d, 0x6a26264c, 0x5a36366c, 0x413f3f7e, 0x02f7f7f5, 0x4fcccc83,
+		0x5c343468, 0xf4a5a551, 0x34e5e5d1, 0x08f1f1f9, 0x937171e2, 0x73d8d8ab, 0x53313162, 0x3f15152a,
+		0x0c040408, 0x52c7c795, 0x65232346, 0x5ec3c39d, 0x28181830, 0xa1969637, 0x0f05050a, 0xb59a9a2f,
+		0x0907070e, 0x36121224, 0x9b80801b, 0x3de2e2df, 0x26ebebcd, 0x6927274e, 0xcdb2b27f, 0x9f7575ea,
+		0x1b090912, 0x9e83831d, 0x742c2c58, 0x2e1a1a34, 0x2d1b1b36, 0xb26e6edc, 0xee5a5ab4, 0xfba0a05b,
+		0xf65252a4, 0x4d3b3b76, 0x61d6d6b7, 0xceb3b37d, 0x7b292952, 0x3ee3e3dd, 0x712f2f5e, 0x97848413,
+		0xf55353a6, 0x68d1d1b9, 0x00000000, 0x2cededc1, 0x60202040, 0x1ffcfce3, 0xc8b1b179, 0xed5b5bb6,
+		0xbe6a6ad4, 0x46cbcb8d, 0xd9bebe67, 0x4b393972, 0xde4a4a94, 0xd44c4c98, 0xe85858b0, 0x4acfcf85,
+		0x6bd0d0bb, 0x2aefefc5, 0xe5aaaa4f, 0x16fbfbed, 0xc5434386, 0xd74d4d9a, 0x55333366, 0x94858511,
+		0xcf45458a, 0x10f9f9e9, 0x06020204, 0x817f7ffe, 0xf05050a0, 0x443c3c78, 0xba9f9f25, 0xe3a8a84b,
+		0xf35151a2, 0xfea3a35d, 0xc0404080, 0x8a8f8f05, 0xad92923f, 0xbc9d9d21, 0x48383870, 0x04f5f5f1,
+		0xdfbcbc63, 0xc1b6b677, 0x75dadaaf, 0x63212142, 0x30101020, 0x1affffe5, 0x0ef3f3fd, 0x6dd2d2bf,
+		0x4ccdcd81, 0x140c0c18, 0x35131326, 0x2fececc3, 0xe15f5fbe, 0xa2979735, 0xcc444488, 0x3917172e,
+		0x57c4c493, 0xf2a7a755, 0x827e7efc, 0x473d3d7a, 0xac6464c8, 0xe75d5dba, 0x2b191932, 0x957373e6,
+		0xa06060c0, 0x98818119, 0xd14f4f9e, 0x7fdcdca3, 0x66222244, 0x7e2a2a54, 0xab90903b, 0x8388880b,
+		0xca46468c, 0x29eeeec7, 0xd3b8b86b, 0x3c141428, 0x79dedea7, 0xe25e5ebc, 0x1d0b0b16, 0x76dbdbad,
+		0x3be0e0db, 0x56323264, 0x4e3a3a74, 0x1e0a0a14, 0xdb494992, 0x0a06060c, 0x6c242448, 0xe45c5cb8,
+		0x5dc2c29f, 0x6ed3d3bd, 0xefacac43, 0xa66262c4, 0xa8919139, 0xa4959531, 0x37e4e4d3, 0x8b7979f2,
+		0x32e7e7d5, 0x43c8c88b, 0x5937376e, 0xb76d6dda, 0x8c8d8d01, 0x64d5d5b1, 0xd24e4e9c, 0xe0a9a949,
+		0xb46c6cd8, 0xfa5656ac, 0x07f4f4f3, 0x25eaeacf, 0xaf6565ca, 0x8e7a7af4, 0xe9aeae47, 0x18080810,
+		0xd5baba6f, 0x887878f0, 0x6f25254a, 0x722e2e5c, 0x241c1c38, 0xf1a6a657, 0xc7b4b473, 0x51c6c697,
+		0x23e8e8cb, 0x7cdddda1, 0x9c7474e8, 0x211f1f3e, 0xdd4b4b96, 0xdcbdbd61, 0x868b8b0d, 0x858a8a0f,
+		0x907070e0, 0x423e3e7c, 0xc4b5b571, 0xaa6666cc, 0xd8484890, 0x05030306, 0x01f6f6f7, 0x120e0e1c,
+		0xa36161c2, 0x5f35356a, 0xf95757ae, 0xd0b9b969, 0x91868617, 0x58c1c199, 0x271d1d3a, 0xb99e9e27,
+		0x38e1e1d9, 0x13f8f8eb, 0xb398982b, 0x33111122, 0xbb6969d2, 0x70d9d9a9, 0x898e8e07, 0xa7949433,
+		0xb69b9b2d, 0x221e1e3c, 0x92878715, 0x20e9e9c9, 0x49cece87, 0xff5555aa, 0x78282850, 0x7adfdfa5,
+		0x8f8c8c03, 0xf8a1a159, 0x80898909, 0x170d0d1a, 0xdabfbf65, 0x31e6e6d7, 0xc6424284, 0xb86868d0,
+		0xc3414182, 0xb0999929, 0x772d2d5a, 0x110f0f1e, 0xcbb0b07b, 0xfc5454a8, 0xd6bbbb6d, 0x3a16162c,
+	},
+	{
+		0x6363c6a5, 0x7c7cf884, 0x7777ee99, 0x7b7bf68d, 0xf2f2ff0d, 0x6b6bd6bd, 0x6f6fdeb1, 0xc5c59154,
+		0x30306050, 0x01010203, 0x6767cea9, 0x2b2b567d, 0xfefee719, 0xd7d7b562, 0xabab4de6, 0x7676ec9a,
+		0xcaca8f45, 0x82821f9d, 0xc9c98940, 0x7d7dfa87, 0xfafaef15, 0x5959b2eb, 0x47478ec9, 0xf0f0fb0b,
+		0xadad41ec, 0xd4d4b367, 0xa2a25ffd, 0xafaf45ea, 0x9c9c23bf, 0xa4a453f7, 0x7272e496, 0xc0c09b5b,
+		0xb7b775c2, 0xfdfde11c, 0x93933dae, 0x26264c6a, 0x36366c5a, 0x3f3f7e41, 0xf7f7f502, 0xcccc834f,
+		0x3434685c, 0xa5a551f4, 0xe5e5d134, 0xf1f1f908, 0x7171e293, 0xd8d8ab73, 0x31316253, 0x15152a3f,
+		0x0404080c, 0xc7c79552, 0x23234665, 0xc3c39d5e, 0x18183028, 0x969637a1, 0x05050a0f, 0x9a9a2fb5,
+		0x07070e09, 0x12122436, 0x80801b9b, 0xe2e2df3d, 0xebebcd26, 0x27274e69, 0xb2b27fcd, 0x7575ea9f,
+		0x0909121b, 0x83831d9e, 0x2c2c5874, 0x1a1a342e, 0x1b1b362d, 0x6e6edcb2, 0x5a5ab4ee, 0xa0a05bfb,
+		0x5252a4f6, 0x3b3b764d, 0xd6d6b761, 0xb3b37dce, 0x2929527b, 0xe3e3dd3e, 0x2f2f5e71, 0x84841397,
+		0x5353a6f5, 0xd1d1b968, 0x00000000, 0xededc12c, 0x20204060, 0xfcfce31f, 0xb1b179c8, 0x5b5bb6ed,
+		0x6a6ad4be, 0xcbcb8d46, 0xbebe67d9, 0x3939724b, 0x4a4a94de, 0x4c4c98d4, 0x5858b0e8, 0xcfcf854a,
+		0xd0d0bb6b, 0xefefc52a, 0xaaaa4fe5, 0xfbfbed16, 0x434386c5, 0x4d4d9ad7, 0x33336655, 0x85851194,
+		0x45458acf, 0xf9f9e910, 0x02020406, 0x7f7ffe81, 0x5050a0f0, 0x3c3c7844, 0x9f9f25ba, 0xa8a84be3,
+		0x5151a2f3, 0xa3a35dfe, 0x404080c0, 0x8f8f058a, 0x92923fad, 0x9d9d21bc, 0x38387048, 0xf5f5f104,
+		0xbcbc63df, 0xb6b677c1, 0xdadaaf75, 0x21214263, 0x10102030, 0xffffe51a, 0xf3f3fd0e, 0xd2d2bf6d,
+		0xcdcd814c, 0x0c0c1814, 0x13132635, 0xececc32f, 0x5f5fbee1, 0x979735a2, 0x444488cc, 0x17172e39,
+		0xc4c49357, 0xa7a755f2, 0x7e7efc82, 0x3d3d7a47, 0x6464c8ac, 0x5d5dbae7, 0x1919322b, 0x7373e695,
+		0x6060c0a0, 0x81811998, 0x4f4f9ed1, 0xdcdca37f, 0x22224466, 0x2a2a547e, 0x90903bab, 0x88880b83,
+		0x46468cca, 0xeeeec729, 0xb8b86bd3, 0x1414283c, 0xdedea779, 0x5e5ebce2, 0x0b0b161d, 0xdbdbad76,
+		0xe0e0db3b, 0x32326456, 0x3a3a744e, 0x0a0a141e, 0x494992db, 0x06060c0a, 0x2424486c, 0x5c5cb8e4,
+		0xc2c29f5d, 0xd3d3bd6e, 0xacac43ef, 0x6262c4a6, 0x919139a8, 0x959531a4, 0xe4e4d337, 0x7979f28b,
+		0xe7e7d532, 0xc8c88b43, 0x37376e59, 0x6d6ddab7, 0x8d8d018c, 0xd5d5b164, 0x4e4e9cd2, 0xa9a949e0,
+		0x6c6cd8b4, 0x5656acfa, 0xf4f4f307, 0xeaeacf25, 0x6565caaf, 0x7a7af48e, 0xaeae47e9, 0x08081018,
+		0xbaba6fd5, 0x7878f088, 0x25254a6f, 0x2e2e5c72, 0x1c1c3824, 0xa6a657f1, 0xb4b473c7, 0xc6c69751,
+		0xe8e8cb23, 0xdddda17c, 0x7474e89c, 0x1f1f3e21, 0x4b4b96dd, 0xbdbd61dc, 0x8b8b0d86, 0x8a8a0f85,
+		0x7070e090, 0x3e3e7c42, 0xb5b571c4, 0x6666ccaa, 0x484890d8, 0x03030605, 0xf6f6f701, 0x0e0e1c12,
+		0x6161c2a3, 0x35356a5f, 0x5757aef9, 0xb9b969d0, 0x86861791, 0xc1c19958, 0x1d1d3a27, 0x9e9e27b9,
+		0xe1e1d938, 0xf8f8eb13, 0x98982bb3, 0x11112233, 0x6969d2bb, 0xd9d9a970, 0x8e8e0789, 0x949433a7,
+		0x9b9b2db6, 0x1e1e3c22, 0x87871592, 0xe9e9c920, 0xcece8749, 0x5555aaff, 0x28285078, 0xdfdfa57a,
+		0x8c8c038f, 0xa1a159f8, 0x89890980, 0x0d0d1a17, 0xbfbf65da, 0xe6e6d731, 0x424284c6, 0x6868d0b8,
+		0x414182c3, 0x999929b0, 0x2d2d5a77, 0x0f0f1e11, 0xb0b07bcb, 0x5454a8fc, 0xbbbb6dd6, 0x16162c3a,
+	},
+	{
+		0x63c6a563, 0x7cf8847c, 0x77ee9977, 0x7bf68d7b, 0xf2ff0df2, 0x6bd6bd6b, 0x6fdeb16f, 0xc59154c5,
+		0x30605030, 0x01020301, 0x67cea967, 0x2b567d2b, 0xfee719fe, 0xd7b562d7, 0xab4de6ab, 0x76ec9a76,
+		0xca8f45ca, 0x821f9d82, 0xc98940c9, 0x7dfa877d, 0xfaef15fa, 0x59b2eb59, 0x478ec947, 0xf0fb0bf0,
+		0xad41ecad, 0xd4b367d4, 0xa25ffda2, 0xaf45eaaf, 0x9c23bf9c, 0xa453f7a4, 0x72e49672, 0xc09b5bc0,
+		0xb775c2b7, 0xfde11cfd, 0x933dae93, 0x264c6a26, 0x366c5a36, 0x3f7e413f, 0xf7f502f7, 0xcc834fcc,
+		0x34685c34, 0xa551f4a5, 0xe5d134e5, 0xf1f908f1, 0x71e29371, 0xd8ab73d8, 0x31625331, 0x152a3f15,
+		0x04080c04, 0xc79552c7, 0x23466523, 0xc39d5ec3, 0x18302818, 0x9637a196, 0x050a0f05, 0x9a2fb59a,
+		0x070e0907, 0x12243612, 0x801b9b80, 0xe2df3de2, 0xebcd26eb, 0x274e6927, 0xb27fcdb2, 0x75ea9f75,
+		0x09121b09, 0x831d9e83, 0x2c58742c, 0x1a342e1a, 0x1b362d1b, 0x6edcb26e, 0x5ab4ee5a, 0xa05bfba0,
+		0x52a4f652, 0x3b764d3b, 0xd6b761d6, 0xb37dceb3, 0x29527b29, 0xe3dd3ee3, 0x2f5e712f, 0x84139784,
+		0x53a6f553, 0xd1b968d1, 0x00000000, 0xedc12ced, 0x20406020, 0xfce31ffc, 0xb179c8b1, 0x5bb6ed5b,
+		0x6ad4be6a, 0xcb8d46cb, 0xbe67d9be, 0x39724b39, 0x4a94de4a, 0x4c98d44c, 0x58b0e858, 0xcf854acf,
+		0xd0bb6bd0, 0xefc52aef, 0xaa4fe5aa, 0xfbed16fb, 0x4386c543, 0x4d9ad74d, 0x33665533, 0x85119485,
+		0x458acf45, 0xf9e910f9, 0x02040602, 0x7ffe817f, 0x50a0f050, 0x3c78443c, 0x9f25ba9f, 0xa84be3a8,
+		0x51a2f351, 0xa35dfea3, 0x4080c040, 0x8f058a8f, 0x923fad92, 0x9d21bc9d, 0x38704838, 0xf5f104f5,
+		0xbc63dfbc, 0xb677c1b6, 0xdaaf75da, 0x21426321, 0x10203010, 0xffe51aff, 0xf3fd0ef3, 0xd2bf6dd2,
+		0xcd814ccd, 0x0c18140c, 0x13263513, 0xecc32fec, 0x5fbee15f, 0x9735a297, 0x4488cc44, 0x172e3917,
+		0xc49357c4, 0xa755f2a7, 0x7efc827e, 0x3d7a473d, 0x64c8ac64, 0x5dbae75d, 0x19322b19, 0x73e69573,
+		0x60c0a060, 0x81199881, 0x4f9ed14f, 0xdca37fdc, 0x22446622, 0x2a547e2a, 0x903bab90, 0x880b8388,
+		0x468cca46, 0xeec729ee, 0xb86bd3b8, 0x14283c14, 0xdea779de, 0x5ebce25e, 0x0b161d0b, 0xdbad76db,
+		0xe0db3be0, 0x32645632, 0x3a744e3a, 0x0a141e0a, 0x4992db49, 0x060c0a06, 0x24486c24, 0x5cb8e45c,
+		0xc29f5dc2, 0xd3bd6ed3, 0xac43efac, 0x62c4a662, 0x9139a891, 0x9531a495, 0xe4d337e4, 0x79f28b79,
+		0xe7d532e7, 0xc88b43c8, 0x376e5937, 0x6ddab76d, 0x8d018c8d, 0xd5b164d5, 0x4e9cd24e, 0xa949e0a9,
+		0x6cd8b46c, 0x56acfa56, 0xf4f307f4, 0xeacf25ea, 0x65caaf65, 0x7af48e7a, 0xae47e9ae, 0x08101808,
+		0xba6fd5ba, 0x78f08878, 0x254a6f25, 0x2e5c722e, 0x1c38241c, 0xa657f1a6, 0xb473c7b4, 0xc69751c6,
+		0xe8cb23e8, 0xdda17cdd, 0x74e89c74, 0x1f3e211f, 0x4b96dd4b, 0xbd61dcbd, 0x8b0d868b, 0x8a0f858a,
+		0x70e09070, 0x3e7c423e, 0xb571c4b5, 0x66ccaa66, 0x4890d848, 0x03060503, 0xf6f701f6, 0x0e1c120e,
+		0x61c2a361, 0x356a5f35, 0x57aef957, 0xb969d0b9, 0x86179186, 0xc19958c1, 0x1d3a271d, 0x9e27b99e,
+		0xe1d938e1, 0xf8eb13f8, 0x982bb398, 0x11223311, 0x69d2bb69, 0xd9a970d9, 0x8e07898e, 0x9433a794,
+		0x9b2db69b, 0x1e3c221e, 0x87159287, 0xe9c920e9, 0xce8749ce, 0x55aaff55, 0x28507828, 0xdfa57adf,
+		0x8c038f8c, 0xa159f8a1, 0x89098089, 0x0d1a170d, 0xbf65dabf, 0xe6d731e6, 0x4284c642, 0x68d0b868,
+		0x4182c341, 0x9929b099, 0x2d5a772d, 0x0f1e110f, 0xb07bcbb0, 0x54a8fc54, 0xbb6dd6bb, 0x162c3a16,
+	},
+	{
+		0xc6a56363, 0xf8847c7c, 0xee997777, 0xf68d7b7b, 0xff0df2f2, 0xd6bd6b6b, 0xdeb16f6f, 0x9154c5c5,
+		0x60503030, 0x02030101, 0xcea96767, 0x567d2b2b, 0xe719fefe, 0xb562d7d7, 0x4de6abab, 0xec9a7676,
+		0x8f45caca, 0x1f9d8282, 0x8940c9c9, 0xfa877d7d, 0xef15fafa, 0xb2eb5959, 0x8ec94747, 0xfb0bf0f0,
+		0x41ecadad, 0xb367d4d4, 0x5ffda2a2, 0x45eaafaf, 0x23bf9c9c, 0x53f7a4a4, 0xe4967272, 0x9b5bc0c0,
+		0x75c2b7b7, 0xe11cfdfd, 0x3dae9393, 0x4c6a2626, 0x6c5a3636, 0x7e413f3f, 0xf502f7f7, 0x834fcccc,
+		0x685c3434, 0x51f4a5a5, 0xd134e5e5, 0xf908f1f1, 0xe2937171, 0xab73d8d8, 0x62533131, 0x2a3f1515,
+		0x080c0404, 0x9552c7c7, 0x46652323, 0x9d5ec3c3, 0x30281818, 0x37a19696, 0x0a0f0505, 0x2fb59a9a,
+		0x0e090707, 0x24361212, 0x1b9b8080, 0xdf3de2e2, 0xcd26ebeb, 0x4e692727, 0x7fcdb2b2, 0xea9f7575,
+		0x121b0909, 0x1d9e8383, 0x58742c2c, 0x342e1a1a, 0x362d1b1b, 0xdcb26e6e, 0xb4ee5a5a, 0x5bfba0a0,
+		0xa4f65252, 0x764d3b3b, 0xb761d6d6, 0x7dceb3b3, 0x527b2929, 0xdd3ee3e3, 0x5e712f2f, 0x13978484,
+		0xa6f55353, 0xb968d1d1, 0x00000000, 0xc12ceded, 0x40602020, 0xe31ffcfc, 0x79c8b1b1, 0xb6ed5b5b,
+		0xd4be6a6a, 0x8d46cbcb, 0x67d9bebe, 0x724b3939, 0x94de4a4a, 0x98d44c4c, 0xb0e85858, 0x854acfcf,
+		0xbb6bd0d0, 0xc52aefef, 0x4fe5aaaa, 0xed16fbfb, 0x86c54343, 0x9ad74d4d, 0x66553333, 0x11948585,
+		0x8acf4545, 0xe910f9f9, 0x04060202, 0xfe817f7f, 0xa0f05050, 0x78443c3c, 0x25ba9f9f, 0x4be3a8a8,
+		0xa2f35151, 0x5dfea3a3, 0x80c04040, 0x058a8f8f, 0x3fad9292, 0x21bc9d9d, 0x70483838, 0xf104f5f5,
+		0x63dfbcbc, 0x77c1b6b6, 0xaf75dada, 0x42632121, 0x20301010, 0xe51affff, 0xfd0ef3f3, 0xbf6dd2d2,
+		0x814ccdcd, 0x18140c0c, 0x26351313, 0xc32fecec, 0xbee15f5f, 0x35a29797, 0x88cc4444, 0x2e391717,
+		0x9357c4c4, 0x55f2a7a7, 0xfc827e7e, 0x7a473d3d, 0xc8ac6464, 0xbae75d5d, 0x322b1919, 0xe6957373,
+		0xc0a06060, 0x19988181, 0x9ed14f4f, 0xa37fdcdc, 0x44662222, 0x547e2a2a, 0x3bab9090, 0x0b838888,
+		0x8cca4646, 0xc729eeee, 0x6bd3b8b8, 0x283c1414, 0xa779dede, 0xbce25e5e, 0x161d0b0b, 0xad76dbdb,
+		0xdb3be0e0, 0x64563232, 0x744e3a3a, 0x141e0a0a, 0x92db4949, 0x0c0a0606, 0x486c2424, 0xb8e45c5c,
+		0x9f5dc2c2, 0xbd6ed3d3, 0x43efacac, 0xc4a66262, 0x39a89191, 0x31a49595, 0xd337e4e4, 0xf28b7979,
+		0xd532e7e7, 0x8b43c8c8, 0x6e593737, 0xdab76d6d, 0x018c8d8d, 0xb164d5d5, 0x9cd24e4e, 0x49e0a9a9,
+		0xd8b46c6c, 0xacfa5656, 0xf307f4f4, 0xcf25eaea, 0xcaaf6565, 0xf48e7a7a, 0x47e9aeae, 0x10180808,
+		0x6fd5baba, 0xf0887878, 0x4a6f2525, 0x5c722e2e, 0x38241c1c, 0x57f1a6a6, 0x73c7b4b4, 0x9751c6c6,
+		0xcb23e8e8, 0xa17cdddd, 0xe89c7474, 0x3e211f1f, 0x96dd4b4b, 0x61dcbdbd, 0x0d868b8b, 0x0f858a8a,
+		0xe0907070, 0x7c423e3e, 0x71c4b5b5, 0xccaa6666, 0x90d84848, 0x06050303, 0xf701f6f6, 0x1c120e0e,
+		0xc2a36161, 0x6a5f3535, 0xaef95757, 0x69d0b9b9, 0x17918686, 0x9958c1c1, 0x3a271d1d, 0x27b99e9e,
+		0xd938e1e1, 0xeb13f8f8, 0x2bb39898, 0x22331111, 0xd2bb6969, 0xa970d9d9, 0x07898e8e, 0x33a79494,
+		0x2db69b9b, 0x3c221e1e, 0x15928787, 0xc920e9e9, 0x8749cece, 0xaaff5555, 0x50782828, 0xa57adfdf,
+		0x038f8c8c, 0x59f8a1a1, 0x09808989, 0x1a170d0d, 0x65dabfbf, 0xd731e6e6, 0x84c64242, 0xd0b86868,
+		0x82c34141, 0x29b09999, 0x5a772d2d, 0x1e110f0f, 0x7bcbb0b0, 0xa8fc5454, 0x6dd6bbbb, 0x2c3a1616,
+	},
 };
 
-static const uint32 Td2[256] =
+static const uint32 AesITbox[4][256] =
 {
-	0xa75051f4U, 0x65537e41U, 0xa4c31a17U, 0x5e963a27U,
-	0x6bcb3babU, 0x45f11f9dU, 0x58abacfaU, 0x03934be3U,
-	0xfa552030U, 0x6df6ad76U, 0x769188ccU, 0x4c25f502U,
-	0xd7fc4fe5U, 0xcbd7c52aU, 0x44802635U, 0xa38fb562U,
-	0x5a49deb1U, 0x1b6725baU, 0x0e9845eaU, 0xc0e15dfeU,
-	0x7502c32fU, 0xf012814cU, 0x97a38d46U, 0xf9c66bd3U,
-	0x5fe7038fU, 0x9c951592U, 0x7aebbf6dU, 0x59da9552U,
-	0x832dd4beU, 0x21d35874U, 0x692949e0U, 0xc8448ec9U,
-	0x896a75c2U, 0x7978f48eU, 0x3e6b9958U, 0x71dd27b9U,
-	0x4fb6bee1U, 0xad17f088U, 0xac66c920U, 0x3ab47dceU,
-	0x4a1863dfU, 0x3182e51aU, 0x33609751U, 0x7f456253U,
-	0x77e0b164U, 0xae84bb6bU, 0xa01cfe81U, 0x2b94f908U,
-	0x68587048U, 0xfd198f45U, 0x6c8794deU, 0xf8b7527bU,
-	0xd323ab73U, 0x02e2724bU, 0x8f57e31fU, 0xab2a6655U,
-	0x2807b2ebU, 0xc2032fb5U, 0x7b9a86c5U, 0x08a5d337U,
-	0x87f23028U, 0xa5b223bfU, 0x6aba0203U, 0x825ced16U,
-	0x1c2b8acfU, 0xb492a779U, 0xf2f0f307U, 0xe2a14e69U,
-	0xf4cd65daU, 0xbed50605U, 0x621fd134U, 0xfe8ac4a6U,
-	0x539d342eU, 0x55a0a2f3U, 0xe132058aU, 0xeb75a4f6U,
-	0xec390b83U, 0xefaa4060U, 0x9f065e71U, 0x1051bd6eU,
-	0x8af93e21U, 0x063d96ddU, 0x05aedd3eU, 0xbd464de6U,
-	0x8db59154U, 0x5d0571c4U, 0xd46f0406U, 0x15ff6050U,
-	0xfb241998U, 0xe997d6bdU, 0x43cc8940U, 0x9e7767d9U,
-	0x42bdb0e8U, 0x8b880789U, 0x5b38e719U, 0xeedb79c8U,
-	0x0a47a17cU, 0x0fe97c42U, 0x1ec9f884U, 0x00000000U,
-	0x86830980U, 0xed48322bU, 0x70ac1e11U, 0x724e6c5aU,
-	0xfffbfd0eU, 0x38560f85U, 0xd51e3daeU, 0x3927362dU,
-	0xd9640a0fU, 0xa621685cU, 0x54d19b5bU, 0x2e3a2436U,
-	0x67b10c0aU, 0xe70f9357U, 0x96d2b4eeU, 0x919e1b9bU,
-	0xc54f80c0U, 0x20a261dcU, 0x4b695a77U, 0x1a161c12U,
-	0xba0ae293U, 0x2ae5c0a0U, 0xe0433c22U, 0x171d121bU,
-	0x0d0b0e09U, 0xc7adf28bU, 0xa8b92db6U, 0xa9c8141eU,
-	0x198557f1U, 0x074caf75U, 0xddbbee99U, 0x60fda37fU,
-	0x269ff701U, 0xf5bc5c72U, 0x3bc54466U, 0x7e345bfbU,
-	0x29768b43U, 0xc6dccb23U, 0xfc68b6edU, 0xf163b8e4U,
-	0xdccad731U, 0x85104263U, 0x22401397U, 0x112084c6U,
-	0x247d854aU, 0x3df8d2bbU, 0x3211aef9U, 0xa16dc729U,
-	0x2f4b1d9eU, 0x30f3dcb2U, 0x52ec0d86U, 0xe3d077c1U,
-	0x166c2bb3U, 0xb999a970U, 0x48fa1194U, 0x642247e9U,
-	0x8cc4a8fcU, 0x3f1aa0f0U, 0x2cd8567dU, 0x90ef2233U,
-	0x4ec78749U, 0xd1c1d938U, 0xa2fe8ccaU, 0x0b3698d4U,
-	0x81cfa6f5U, 0xde28a57aU, 0x8e26dab7U, 0xbfa43fadU,
-	0x9de42c3aU, 0x920d5078U, 0xcc9b6a5fU, 0x4662547eU,
-	0x13c2f68dU, 0xb8e890d8U, 0xf75e2e39U, 0xaff582c3U,
-	0x80be9f5dU, 0x937c69d0U, 0x2da96fd5U, 0x12b3cf25U,
-	0x993bc8acU, 0x7da71018U, 0x636ee89cU, 0xbb7bdb3bU,
-	0x7809cd26U, 0x18f46e59U, 0xb701ec9aU, 0x9aa8834fU,
-	0x6e65e695U, 0xe67eaaffU, 0xcf0821bcU, 0xe8e6ef15U,
-	0x9bd9bae7U, 0x36ce4a6fU, 0x09d4ea9fU, 0x7cd629b0U,
-	0xb2af31a4U, 0x23312a3fU, 0x9430c6a5U, 0x66c035a2U,
-	0xbc37744eU, 0xcaa6fc82U, 0xd0b0e090U, 0xd81533a7U,
-	0x984af104U, 0xdaf741ecU, 0x500e7fcdU, 0xf62f1791U,
-	0xd68d764dU, 0xb04d43efU, 0x4d54ccaaU, 0x04dfe496U,
-	0xb5e39ed1U, 0x881b4c6aU, 0x1fb8c12cU, 0x517f4665U,
-	0xea049d5eU, 0x355d018cU, 0x7473fa87U, 0x412efb0bU,
-	0x1d5ab367U, 0xd25292dbU, 0x5633e910U, 0x47136dd6U,
-	0x618c9ad7U, 0x0c7a37a1U, 0x148e59f8U, 0x3c89eb13U,
-	0x27eecea9U, 0xc935b761U, 0xe5ede11cU, 0xb13c7a47U,
-	0xdf599cd2U, 0x733f55f2U, 0xce791814U, 0x37bf73c7U,
-	0xcdea53f7U, 0xaa5b5ffdU, 0x6f14df3dU, 0xdb867844U,
-	0xf381caafU, 0xc43eb968U, 0x342c3824U, 0x405fc2a3U,
-	0xc372161dU, 0x250cbce2U, 0x498b283cU, 0x9541ff0dU,
-	0x017139a8U, 0xb3de080cU, 0xe49cd8b4U, 0xc1906456U,
-	0x84617bcbU, 0xb670d532U, 0x5c74486cU, 0x5742d0b8U,
+	{
+		0x50a7f451, 0x5365417e, 0xc3a4171a, 0x965e273a, 0xcb6bab3b, 0xf1459d1f, 0xab58faac, 0x9303e34b,
+		0x55fa3020, 0xf66d76ad, 0x9176cc88, 0x254c02f5, 0xfcd7e54f, 0xd7cb2ac5, 0x80443526, 0x8fa362b5,
+		0x495ab1de, 0x671bba25, 0x980eea45, 0xe1c0fe5d, 0x02752fc3, 0x12f04c81, 0xa397468d, 0xc6f9d36b,
+		0xe75f8f03, 0x959c9215, 0xeb7a6dbf, 0xda595295, 0x2d83bed4, 0xd3217458, 0x2969e049, 0x44c8c98e,
+		0x6a89c275, 0x78798ef4, 0x6b3e5899, 0xdd71b927, 0xb64fe1be, 0x17ad88f0, 0x66ac20c9, 0xb43ace7d,
+		0x184adf63, 0x82311ae5, 0x60335197, 0x457f5362, 0xe07764b1, 0x84ae6bbb, 0x1ca081fe, 0x942b08f9,
+		0x58684870, 0x19fd458f, 0x876cde94, 0xb7f87b52, 0x23d373ab, 0xe2024b72, 0x578f1fe3, 0x2aab5566,
+		0x0728ebb2, 0x03c2b52f, 0x9a7bc586, 0xa50837d3, 0xf2872830, 0xb2a5bf23, 0xba6a0302, 0x5c8216ed,
+		0x2b1ccf8a, 0x92b479a7, 0xf0f207f3, 0xa1e2694e, 0xcdf4da65, 0xd5be0506, 0x1f6234d1, 0x8afea6c4,
+		0x9d532e34, 0xa055f3a2, 0x32e18a05, 0x75ebf6a4, 0x39ec830b, 0xaaef6040, 0x069f715e, 0x51106ebd,
+		0xf98a213e, 0x3d06dd96, 0xae053edd, 0x46bde64d, 0xb58d5491, 0x055dc471, 0x6fd40604, 0xff155060,
+		0x24fb9819, 0x97e9bdd6, 0xcc434089, 0x779ed967, 0xbd42e8b0, 0x888b8907, 0x385b19e7, 0xdbeec879,
+		0x470a7ca1, 0xe90f427c, 0xc91e84f8, 0x00000000, 0x83868009, 0x48ed2b32, 0xac70111e, 0x4e725a6c,
+		0xfbff0efd, 0x5638850f, 0x1ed5ae3d, 0x27392d36, 0x64d90f0a, 0x21a65c68, 0xd1545b9b, 0x3a2e3624,
+		0xb1670a0c, 0x0fe75793, 0xd296eeb4, 0x9e919b1b, 0x4fc5c080, 0xa220dc61, 0x694b775a, 0x161a121c,
+		0x0aba93e2, 0xe52aa0c0, 0x43e0223c, 0x1d171b12, 0x0b0d090e, 0xadc78bf2, 0xb9a8b62d, 0xc8a91e14,
+		0x8519f157, 0x4c0775af, 0xbbdd99ee, 0xfd607fa3, 0x9f2601f7, 0xbcf5725c, 0xc53b6644, 0x347efb5b,
+		0x7629438b, 0xdcc623cb, 0x68fcedb6, 0x63f1e4b8, 0xcadc31d7, 0x10856342, 0x40229713, 0x2011c684,
+		0x7d244a85, 0xf83dbbd2, 0x1132f9ae, 0x6da129c7, 0x4b2f9e1d, 0xf330b2dc, 0xec52860d, 0xd0e3c177,
+		0x6c16b32b, 0x99b970a9, 0xfa489411, 0x2264e947, 0xc48cfca8, 0x1a3ff0a0, 0xd82c7d56, 0xef903322,
+		0xc74e4987, 0xc1d138d9, 0xfea2ca8c, 0x360bd498, 0xcf81f5a6, 0x28de7aa5, 0x268eb7da, 0xa4bfad3f,
+		0xe49d3a2c, 0x0d927850, 0x9bcc5f6a, 0x62467e54, 0xc2138df6, 0xe8b8d890, 0x5ef7392e, 0xf5afc382,
+		0xbe805d9f, 0x7c93d069, 0xa92dd56f, 0xb31225cf, 0x3b99acc8, 0xa77d1810, 0x6e639ce8, 0x7bbb3bdb,
+		0x097826cd, 0xf418596e, 0x01b79aec, 0xa89a4f83, 0x656e95e6, 0x7ee6ffaa, 0x08cfbc21, 0xe6e815ef,
+		0xd99be7ba, 0xce366f4a, 0xd4099fea, 0xd67cb029, 0xafb2a431, 0x31233f2a, 0x3094a5c6, 0xc066a235,
+		0x37bc4e74, 0xa6ca82fc, 0xb0d090e0, 0x15d8a733, 0x4a9804f1, 0xf7daec41, 0x0e50cd7f, 0x2ff69117,
+		0x8dd64d76, 0x4db0ef43, 0x544daacc, 0xdf0496e4, 0xe3b5d19e, 0x1b886a4c, 0xb81f2cc1, 0x7f516546,
+		0x04ea5e9d, 0x5d358c01, 0x737487fa, 0x2e410bfb, 0x5a1d67b3, 0x52d2db92, 0x335610e9, 0x1347d66d,
+		0x8c61d79a, 0x7a0ca137, 0x8e14f859, 0x893c13eb, 0xee27a9ce, 0x35c961b7, 0xede51ce1, 0x3cb1477a,
+		0x59dfd29c, 0x3f73f255, 0x79ce1418, 0xbf37c773, 0xeacdf753, 0x5baafd5f, 0x146f3ddf, 0x86db4478,
+		0x81f3afca, 0x3ec468b9, 0x2c342438, 0x5f40a3c2, 0x72c31d16, 0x0c25e2bc, 0x8b493c28, 0x41950dff,
+		0x7101a839, 0xdeb30c08, 0x9ce4b4d8, 0x90c15664, 0x6184cb7b, 0x70b632d5, 0x745c6c48, 0x4257b8d0,
+	},
+	{
+		0xa7f45150, 0x65417e53, 0xa4171ac3, 0x5e273a96, 0x6bab3bcb, 0x459d1ff1, 0x58faacab, 0x03e34b93,
+		0xfa302055, 0x6d76adf6, 0x76cc8891, 0x4c02f525, 0xd7e54ffc, 0xcb2ac5d7, 0x44352680, 0xa362b58f,
+		0x5ab1de49, 0x1bba2567, 0x0eea4598, 0xc0fe5de1, 0x752fc302, 0xf04c8112, 0x97468da3, 0xf9d36bc6,
+		0x5f8f03e7, 0x9c921595, 0x7a6dbfeb, 0x595295da, 0x83bed42d, 0x217458d3, 0x69e04929, 0xc8c98e44,
+		0x89c2756a, 0x798ef478, 0x3e58996b, 0x71b927dd, 0x4fe1beb6, 0xad88f017, 0xac20c966, 0x3ace7db4,
+		0x4adf6318, 0x311ae582, 0x33519760, 0x7f536245, 0x7764b1e0, 0xae6bbb84, 0xa081fe1c, 0x2b08f994,
+		0x68487058, 0xfd458f19, 0x6cde9487, 0xf87b52b7, 0xd373ab23, 0x024b72e2, 0x8f1fe357, 0xab55662a,
+		0x28ebb207, 0xc2b52f03, 0x7bc5869a, 0x0837d3a5, 0x872830f2, 0xa5bf23b2, 0x6a0302ba, 0x8216ed5c,
+		0x1ccf8a2b, 0xb479a792, 0xf207f3f0, 0xe2694ea1, 0xf4da65cd, 0xbe0506d5, 0x6234d11f, 0xfea6c48a,
+		0x532e349d, 0x55f3a2a0, 0xe18a0532, 0xebf6a475, 0xec830b39, 0xef6040aa, 0x9f715e06, 0x106ebd51,
+		0x8a213ef9, 0x06dd963d, 0x053eddae, 0xbde64d46, 0x8d5491b5, 0x5dc47105, 0xd406046f, 0x155060ff,
+		0xfb981924, 0xe9bdd697, 0x434089cc, 0x9ed96777, 0x42e8b0bd, 0x8b890788, 0x5b19e738, 0xeec879db,
+		0x0a7ca147, 0x0f427ce9, 0x1e84f8c9, 0x00000000, 0x86800983, 0xed2b3248, 0x70111eac, 0x725a6c4e,
+		0xff0efdfb, 0x38850f56, 0xd5ae3d1e, 0x392d3627, 0xd90f0a64, 0xa65c6821, 0x545b9bd1, 0x2e36243a,
+		0x670a0cb1, 0xe757930f, 0x96eeb4d2, 0x919b1b9e, 0xc5c0804f, 0x20dc61a2, 0x4b775a69, 0x1a121c16,
+		0xba93e20a, 0x2aa0c0e5, 0xe0223c43, 0x171b121d, 0x0d090e0b, 0xc78bf2ad, 0xa8b62db9, 0xa91e14c8,
+		0x19f15785, 0x0775af4c, 0xdd99eebb, 0x607fa3fd, 0x2601f79f, 0xf5725cbc, 0x3b6644c5, 0x7efb5b34,
+		0x29438b76, 0xc623cbdc, 0xfcedb668, 0xf1e4b863, 0xdc31d7ca, 0x85634210, 0x22971340, 0x11c68420,
+		0x244a857d, 0x3dbbd2f8, 0x32f9ae11, 0xa129c76d, 0x2f9e1d4b, 0x30b2dcf3, 0x52860dec, 0xe3c177d0,
+		0x16b32b6c, 0xb970a999, 0x489411fa, 0x64e94722, 0x8cfca8c4, 0x3ff0a01a, 0x2c7d56d8, 0x903322ef,
+		0x4e4987c7, 0xd138d9c1, 0xa2ca8cfe, 0x0bd49836, 0x81f5a6cf, 0xde7aa528, 0x8eb7da26, 0xbfad3fa4,
+		0x9d3a2ce4, 0x9278500d, 0xcc5f6a9b, 0x467e5462, 0x138df6c2, 0xb8d890e8, 0xf7392e5e, 0xafc382f5,
+		0x805d9fbe, 0x93d0697c, 0x2dd56fa9, 0x1225cfb3, 0x99acc83b, 0x7d1810a7, 0x639ce86e, 0xbb3bdb7b,
+		0x7826cd09, 0x18596ef4, 0xb79aec01, 0x9a4f83a8, 0x6e95e665, 0xe6ffaa7e, 0xcfbc2108, 0xe815efe6,
+		0x9be7bad9, 0x366f4ace, 0x099fead4, 0x7cb029d6, 0xb2a431af, 0x233f2a31, 0x94a5c630, 0x66a235c0,
+		0xbc4e7437, 0xca82fca6, 0xd090e0b0, 0xd8a73315, 0x9804f14a, 0xdaec41f7, 0x50cd7f0e, 0xf691172f,
+		0xd64d768d, 0xb0ef434d, 0x4daacc54, 0x0496e4df, 0xb5d19ee3, 0x886a4c1b, 0x1f2cc1b8, 0x5165467f,
+		0xea5e9d04, 0x358c015d, 0x7487fa73, 0x410bfb2e, 0x1d67b35a, 0xd2db9252, 0x5610e933, 0x47d66d13,
+		0x61d79a8c, 0x0ca1377a, 0x14f8598e, 0x3c13eb89, 0x27a9ceee, 0xc961b735, 0xe51ce1ed, 0xb1477a3c,
+		0xdfd29c59, 0x73f2553f, 0xce141879, 0x37c773bf, 0xcdf753ea, 0xaafd5f5b, 0x6f3ddf14, 0xdb447886,
+		0xf3afca81, 0xc468b93e, 0x3424382c, 0x40a3c25f, 0xc31d1672, 0x25e2bc0c, 0x493c288b, 0x950dff41,
+		0x01a83971, 0xb30c08de, 0xe4b4d89c, 0xc1566490, 0x84cb7b61, 0xb632d570, 0x5c6c4874, 0x57b8d042,
+	},
+	{
+		0xf45150a7, 0x417e5365, 0x171ac3a4, 0x273a965e, 0xab3bcb6b, 0x9d1ff145, 0xfaacab58, 0xe34b9303,
+		0x302055fa, 0x76adf66d, 0xcc889176, 0x02f5254c, 0xe54ffcd7, 0x2ac5d7cb, 0x35268044, 0x62b58fa3,
+		0xb1de495a, 0xba25671b, 0xea45980e, 0xfe5de1c0, 0x2fc30275, 0x4c8112f0, 0x468da397, 0xd36bc6f9,
+		0x8f03e75f, 0x9215959c, 0x6dbfeb7a, 0x5295da59, 0xbed42d83, 0x7458d321, 0xe0492969, 0xc98e44c8,
+		0xc2756a89, 0x8ef47879, 0x58996b3e, 0xb927dd71, 0xe1beb64f, 0x88f017ad, 0x20c966ac, 0xce7db43a,
+		0xdf63184a, 0x1ae58231, 0x51976033, 0x5362457f, 0x64b1e077, 0x6bbb84ae, 0x81fe1ca0, 0x08f9942b,
+		0x48705868, 0x458f19fd, 0xde94876c, 0x7b52b7f8, 0x73ab23d3, 0x4b72e202, 0x1fe3578f, 0x55662aab,
+		0xebb20728, 0xb52f03c2, 0xc5869a7b, 0x37d3a508, 0x2830f287, 0xbf23b2a5, 0x0302ba6a, 0x16ed5c82,
+		0xcf8a2b1c, 0x79a792b4, 0x07f3f0f2, 0x694ea1e2, 0xda65cdf4, 0x0506d5be, 0x34d11f62, 0xa6c48afe,
+		0x2e349d53, 0xf3a2a055, 0x8a0532e1, 0xf6a475eb, 0x830b39ec, 0x6040aaef, 0x715e069f, 0x6ebd5110,
+		0x213ef98a, 0xdd963d06, 0x3eddae05, 0xe64d46bd, 0x5491b58d, 0xc471055d, 0x06046fd4, 0x5060ff15,
+		0x981924fb, 0xbdd697e9, 0x4089cc43, 0xd967779e, 0xe8b0bd42, 0x8907888b, 0x19e7385b, 0xc879dbee,
+		0x7ca1470a, 0x427ce90f, 0x84f8c91e, 0x00000000, 0x80098386, 0x2b3248ed, 0x111eac70, 0x5a6c4e72,
+		0x0efdfbff, 0x850f5638, 0xae3d1ed5, 0x2d362739, 0x0f0a64d9, 0x5c6821a6, 0x5b9bd154, 0x36243a2e,
+		0x0a0cb167, 0x57930fe7, 0xeeb4d296, 0x9b1b9e91, 0xc0804fc5, 0xdc61a220, 0x775a694b, 0x121c161a,
+		0x93e20aba, 0xa0c0e52a, 0x223c43e0, 0x1b121d17, 0x090e0b0d, 0x8bf2adc7, 0xb62db9a8, 0x1e14c8a9,
+		0xf1578519, 0x75af4c07, 0x99eebbdd, 0x7fa3fd60, 0x01f79f26, 0x725cbcf5, 0x6644c53b, 0xfb5b347e,
+		0x438b7629, 0x23cbdcc6, 0xedb668fc, 0xe4b863f1, 0x31d7cadc, 0x63421085, 0x97134022, 0xc6842011,
+		0x4a857d24, 0xbbd2f83d, 0xf9ae1132, 0x29c76da1, 0x9e1d4b2f, 0xb2dcf330, 0x860dec52, 0xc177d0e3,
+		0xb32b6c16, 0x70a999b9, 0x9411fa48, 0xe9472264, 0xfca8c48c, 0xf0a01a3f, 0x7d56d82c, 0x3322ef90,
+		0x4987c74e, 0x38d9c1d1, 0xca8cfea2, 0xd498360b, 0xf5a6cf81, 0x7aa528de, 0xb7da268e, 0xad3fa4bf,
+		0x3a2ce49d, 0x78500d92, 0x5f6a9bcc, 0x7e546246, 0x8df6c213, 0xd890e8b8, 0x392e5ef7, 0xc382f5af,
+		0x5d9fbe80, 0xd0697c93, 0xd56fa92d, 0x25cfb312, 0xacc83b99, 0x1810a77d, 0x9ce86e63, 0x3bdb7bbb,
+		0x26cd0978, 0x596ef418, 0x9aec01b7, 0x4f83a89a, 0x95e6656e, 0xffaa7ee6, 0xbc2108cf, 0x15efe6e8,
+		0xe7bad99b, 0x6f4ace36, 0x9fead409, 0xb029d67c, 0xa431afb2, 0x3f2a3123, 0xa5c63094, 0xa235c066,
+		0x4e7437bc, 0x82fca6ca, 0x90e0b0d0, 0xa73315d8, 0x04f14a98, 0xec41f7da, 0xcd7f0e50, 0x91172ff6,
+		0x4d768dd6, 0xef434db0, 0xaacc544d, 0x96e4df04, 0xd19ee3b5, 0x6a4c1b88, 0x2cc1b81f, 0x65467f51,
+		0x5e9d04ea, 0x8c015d35, 0x87fa7374, 0x0bfb2e41, 0x67b35a1d, 0xdb9252d2, 0x10e93356, 0xd66d1347,
+		0xd79a8c61, 0xa1377a0c, 0xf8598e14, 0x13eb893c, 0xa9ceee27, 0x61b735c9, 0x1ce1ede5, 0x477a3cb1,
+		0xd29c59df, 0xf2553f73, 0x141879ce, 0xc773bf37, 0xf753eacd, 0xfd5f5baa, 0x3ddf146f, 0x447886db,
+		0xafca81f3, 0x68b93ec4, 0x24382c34, 0xa3c25f40, 0x1d1672c3, 0xe2bc0c25, 0x3c288b49, 0x0dff4195,
+		0xa8397101, 0x0c08deb3, 0xb4d89ce4, 0x566490c1, 0xcb7b6184, 0x32d570b6, 0x6c48745c, 0xb8d04257,
+	},
+	{
+		0x5150a7f4, 0x7e536541, 0x1ac3a417, 0x3a965e27, 0x3bcb6bab, 0x1ff1459d, 0xacab58fa, 0x4b9303e3,
+		0x2055fa30, 0xadf66d76, 0x889176cc, 0xf5254c02, 0x4ffcd7e5, 0xc5d7cb2a, 0x26804435, 0xb58fa362,
+		0xde495ab1, 0x25671bba, 0x45980eea, 0x5de1c0fe, 0xc302752f, 0x8112f04c, 0x8da39746, 0x6bc6f9d3,
+		0x03e75f8f, 0x15959c92, 0xbfeb7a6d, 0x95da5952, 0xd42d83be, 0x58d32174, 0x492969e0, 0x8e44c8c9,
+		0x756a89c2, 0xf478798e, 0x996b3e58, 0x27dd71b9, 0xbeb64fe1, 0xf017ad88, 0xc966ac20, 0x7db43ace,
+		0x63184adf, 0xe582311a, 0x97603351, 0x62457f53, 0xb1e07764, 0xbb84ae6b, 0xfe1ca081, 0xf9942b08,
+		0x70586848, 0x8f19fd45, 0x94876cde, 0x52b7f87b, 0xab23d373, 0x72e2024b, 0xe3578f1f, 0x662aab55,
+		0xb20728eb, 0x2f03c2b5, 0x869a7bc5, 0xd3a50837, 0x30f28728, 0x23b2a5bf, 0x02ba6a03, 0xed5c8216,
+		0x8a2b1ccf, 0xa792b479, 0xf3f0f207, 0x4ea1e269, 0x65cdf4da, 0x06d5be05, 0xd11f6234, 0xc48afea6,
+		0x349d532e, 0xa2a055f3, 0x0532e18a, 0xa475ebf6, 0x0b39ec83, 0x40aaef60, 0x5e069f71, 0xbd51106e,
+		0x3ef98a21, 0x963d06dd, 0xddae053e, 0x4d46bde6, 0x91b58d54, 0x71055dc4, 0x046fd406, 0x60ff1550,
+		0x1924fb98, 0xd697e9bd, 0x89cc4340, 0x67779ed9, 0xb0bd42e8, 0x07888b89, 0xe7385b19, 0x79dbeec8,
+		0xa1470a7c, 0x7ce90f42, 0xf8c91e84, 0x00000000, 0x09838680, 0x3248ed2b, 0x1eac7011, 0x6c4e725a,
+		0xfdfbff0e, 0x0f563885, 0x3d1ed5ae, 0x3627392d, 0x0a64d90f, 0x6821a65c, 0x9bd1545b, 0x243a2e36,
+		0x0cb1670a, 0x930fe757, 0xb4d296ee, 0x1b9e919b, 0x804fc5c0, 0x61a220dc, 0x5a694b77, 0x1c161a12,
+		0xe20aba93, 0xc0e52aa0, 0x3c43e022, 0x121d171b, 0x0e0b0d09, 0xf2adc78b, 0x2db9a8b6, 0x14c8a91e,
+		0x578519f1, 0xaf4c0775, 0xeebbdd99, 0xa3fd607f, 0xf79f2601, 0x5cbcf572, 0x44c53b66, 0x5b347efb,
+		0x8b762943, 0xcbdcc623, 0xb668fced, 0xb863f1e4, 0xd7cadc31, 0x42108563, 0x13402297, 0x842011c6,
+		0x857d244a, 0xd2f83dbb, 0xae1132f9, 0xc76da129, 0x1d4b2f9e, 0xdcf330b2, 0x0dec5286, 0x77d0e3c1,
+		0x2b6c16b3, 0xa999b970, 0x11fa4894, 0x472264e9, 0xa8c48cfc, 0xa01a3ff0, 0x56d82c7d, 0x22ef9033,
+		0x87c74e49, 0xd9c1d138, 0x8cfea2ca, 0x98360bd4, 0xa6cf81f5, 0xa528de7a, 0xda268eb7, 0x3fa4bfad,
+		0x2ce49d3a, 0x500d9278, 0x6a9bcc5f, 0x5462467e, 0xf6c2138d, 0x90e8b8d8, 0x2e5ef739, 0x82f5afc3,
+		0x9fbe805d, 0x697c93d0, 0x6fa92dd5, 0xcfb31225, 0xc83b99ac, 0x10a77d18, 0xe86e639c, 0xdb7bbb3b,
+		0xcd097826, 0x6ef41859, 0xec01b79a, 0x83a89a4f, 0xe6656e95, 0xaa7ee6ff, 0x2108cfbc, 0xefe6e815,
+		0xbad99be7, 0x4ace366f, 0xead4099f, 0x29d67cb0, 0x31afb2a4, 0x2a31233f, 0xc63094a5, 0x35c066a2,
+		0x7437bc4e, 0xfca6ca82, 0xe0b0d090, 0x3315d8a7, 0xf14a9804, 0x41f7daec, 0x7f0e50cd, 0x172ff691,
+		0x768dd64d, 0x434db0ef, 0xcc544daa, 0xe4df0496, 0x9ee3b5d1, 0x4c1b886a, 0xc1b81f2c, 0x467f5165,
+		0x9d04ea5e, 0x015d358c, 0xfa737487, 0xfb2e410b, 0xb35a1d67, 0x9252d2db, 0xe9335610, 0x6d1347d6,
+		0x9a8c61d7, 0x377a0ca1, 0x598e14f8, 0xeb893c13, 0xceee27a9, 0xb735c961, 0xe1ede51c, 0x7a3cb147,
+		0x9c59dfd2, 0x553f73f2, 0x1879ce14, 0x73bf37c7, 0x53eacdf7, 0x5f5baafd, 0xdf146f3d, 0x7886db44,
+		0xca81f3af, 0xb93ec468, 0x382c3424, 0xc25f40a3, 0x1672c31d, 0xbc0c25e2, 0x288b493c, 0xff41950d,
+		0x397101a8, 0x08deb30c, 0xd89ce4b4, 0x6490c156, 0x7b6184cb, 0xd570b632, 0x48745c6c, 0xd04257b8,
+	},
 };
 
-static const uint32 Td3[256] =
+static inline uint32 AesEncryptMix(uint32 x)
 {
-	0xf4a75051U, 0x4165537eU, 0x17a4c31aU, 0x275e963aU,
-	0xab6bcb3bU, 0x9d45f11fU, 0xfa58abacU, 0xe303934bU,
-	0x30fa5520U, 0x766df6adU, 0xcc769188U, 0x024c25f5U,
-	0xe5d7fc4fU, 0x2acbd7c5U, 0x35448026U, 0x62a38fb5U,
-	0xb15a49deU, 0xba1b6725U, 0xea0e9845U, 0xfec0e15dU,
-	0x2f7502c3U, 0x4cf01281U, 0x4697a38dU, 0xd3f9c66bU,
-	0x8f5fe703U, 0x929c9515U, 0x6d7aebbfU, 0x5259da95U,
-	0xbe832dd4U, 0x7421d358U, 0xe0692949U, 0xc9c8448eU,
-	0xc2896a75U, 0x8e7978f4U, 0x583e6b99U, 0xb971dd27U,
-	0xe14fb6beU, 0x88ad17f0U, 0x20ac66c9U, 0xce3ab47dU,
-	0xdf4a1863U, 0x1a3182e5U, 0x51336097U, 0x537f4562U,
-	0x6477e0b1U, 0x6bae84bbU, 0x81a01cfeU, 0x082b94f9U,
-	0x48685870U, 0x45fd198fU, 0xde6c8794U, 0x7bf8b752U,
-	0x73d323abU, 0x4b02e272U, 0x1f8f57e3U, 0x55ab2a66U,
-	0xeb2807b2U, 0xb5c2032fU, 0xc57b9a86U, 0x3708a5d3U,
-	0x2887f230U, 0xbfa5b223U, 0x036aba02U, 0x16825cedU,
-	0xcf1c2b8aU, 0x79b492a7U, 0x07f2f0f3U, 0x69e2a14eU,
-	0xdaf4cd65U, 0x05bed506U, 0x34621fd1U, 0xa6fe8ac4U,
-	0x2e539d34U, 0xf355a0a2U, 0x8ae13205U, 0xf6eb75a4U,
-	0x83ec390bU, 0x60efaa40U, 0x719f065eU, 0x6e1051bdU,
-	0x218af93eU, 0xdd063d96U, 0x3e05aeddU, 0xe6bd464dU,
-	0x548db591U, 0xc45d0571U, 0x06d46f04U, 0x5015ff60U,
-	0x98fb2419U, 0xbde997d6U, 0x4043cc89U, 0xd99e7767U,
-	0xe842bdb0U, 0x898b8807U, 0x195b38e7U, 0xc8eedb79U,
-	0x7c0a47a1U, 0x420fe97cU, 0x841ec9f8U, 0x00000000U,
-	0x80868309U, 0x2bed4832U, 0x1170ac1eU, 0x5a724e6cU,
-	0x0efffbfdU, 0x8538560fU, 0xaed51e3dU, 0x2d392736U,
-	0x0fd9640aU, 0x5ca62168U, 0x5b54d19bU, 0x362e3a24U,
-	0x0a67b10cU, 0x57e70f93U, 0xee96d2b4U, 0x9b919e1bU,
-	0xc0c54f80U, 0xdc20a261U, 0x774b695aU, 0x121a161cU,
-	0x93ba0ae2U, 0xa02ae5c0U, 0x22e0433cU, 0x1b171d12U,
-	0x090d0b0eU, 0x8bc7adf2U, 0xb6a8b92dU, 0x1ea9c814U,
-	0xf1198557U, 0x75074cafU, 0x99ddbbeeU, 0x7f60fda3U,
-	0x01269ff7U, 0x72f5bc5cU, 0x663bc544U, 0xfb7e345bU,
-	0x4329768bU, 0x23c6dccbU, 0xedfc68b6U, 0xe4f163b8U,
-	0x31dccad7U, 0x63851042U, 0x97224013U, 0xc6112084U,
-	0x4a247d85U, 0xbb3df8d2U, 0xf93211aeU, 0x29a16dc7U,
-	0x9e2f4b1dU, 0xb230f3dcU, 0x8652ec0dU, 0xc1e3d077U,
-	0xb3166c2bU, 0x70b999a9U, 0x9448fa11U, 0xe9642247U,
-	0xfc8cc4a8U, 0xf03f1aa0U, 0x7d2cd856U, 0x3390ef22U,
-	0x494ec787U, 0x38d1c1d9U, 0xcaa2fe8cU, 0xd40b3698U,
-	0xf581cfa6U, 0x7ade28a5U, 0xb78e26daU, 0xadbfa43fU,
-	0x3a9de42cU, 0x78920d50U, 0x5fcc9b6aU, 0x7e466254U,
-	0x8d13c2f6U, 0xd8b8e890U, 0x39f75e2eU, 0xc3aff582U,
-	0x5d80be9fU, 0xd0937c69U, 0xd52da96fU, 0x2512b3cfU,
-	0xac993bc8U, 0x187da710U, 0x9c636ee8U, 0x3bbb7bdbU,
-	0x267809cdU, 0x5918f46eU, 0x9ab701ecU, 0x4f9aa883U,
-	0x956e65e6U, 0xffe67eaaU, 0xbccf0821U, 0x15e8e6efU,
-	0xe79bd9baU, 0x6f36ce4aU, 0x9f09d4eaU, 0xb07cd629U,
-	0xa4b2af31U, 0x3f23312aU, 0xa59430c6U, 0xa266c035U,
-	0x4ebc3774U, 0x82caa6fcU, 0x90d0b0e0U, 0xa7d81533U,
-	0x04984af1U, 0xecdaf741U, 0xcd500e7fU, 0x91f62f17U,
-	0x4dd68d76U, 0xefb04d43U, 0xaa4d54ccU, 0x9604dfe4U,
-	0xd1b5e39eU, 0x6a881b4cU, 0x2c1fb8c1U, 0x65517f46U,
-	0x5eea049dU, 0x8c355d01U, 0x877473faU, 0x0b412efbU,
-	0x671d5ab3U, 0xdbd25292U, 0x105633e9U, 0xd647136dU,
-	0xd7618c9aU, 0xa10c7a37U, 0xf8148e59U, 0x133c89ebU,
-	0xa927eeceU, 0x61c935b7U, 0x1ce5ede1U, 0x47b13c7aU,
-	0xd2df599cU, 0xf2733f55U, 0x14ce7918U, 0xc737bf73U,
-	0xf7cdea53U, 0xfdaa5b5fU, 0x3d6f14dfU, 0x44db8678U,
-	0xaff381caU, 0x68c43eb9U, 0x24342c38U, 0xa3405fc2U,
-	0x1dc37216U, 0xe2250cbcU, 0x3c498b28U, 0x0d9541ffU,
-	0xa8017139U, 0x0cb3de08U, 0xb4e49cd8U, 0x56c19064U,
-	0xcb84617bU, 0x32b670d5U, 0x6c5c7448U, 0xb85742d0U,
-};
+	return ((uint32)AesSbox[(uint8)(x >>  0)] <<  0)
+		 ^ ((uint32)AesSbox[(uint8)(x >>  8)] <<  8)
+		 ^ ((uint32)AesSbox[(uint8)(x >> 16)] << 16)
+		 ^ ((uint32)AesSbox[(uint8)(x >> 24)] << 24);
+}
 
-static const uint32 Td4[256] =
+static inline uint32 AesDecryptMix(uint32 x)
 {
-	0x52525252U, 0x09090909U, 0x6a6a6a6aU, 0xd5d5d5d5U,
-	0x30303030U, 0x36363636U, 0xa5a5a5a5U, 0x38383838U,
-	0xbfbfbfbfU, 0x40404040U, 0xa3a3a3a3U, 0x9e9e9e9eU,
-	0x81818181U, 0xf3f3f3f3U, 0xd7d7d7d7U, 0xfbfbfbfbU,
-	0x7c7c7c7cU, 0xe3e3e3e3U, 0x39393939U, 0x82828282U,
-	0x9b9b9b9bU, 0x2f2f2f2fU, 0xffffffffU, 0x87878787U,
-	0x34343434U, 0x8e8e8e8eU, 0x43434343U, 0x44444444U,
-	0xc4c4c4c4U, 0xdedededeU, 0xe9e9e9e9U, 0xcbcbcbcbU,
-	0x54545454U, 0x7b7b7b7bU, 0x94949494U, 0x32323232U,
-	0xa6a6a6a6U, 0xc2c2c2c2U, 0x23232323U, 0x3d3d3d3dU,
-	0xeeeeeeeeU, 0x4c4c4c4cU, 0x95959595U, 0x0b0b0b0bU,
-	0x42424242U, 0xfafafafaU, 0xc3c3c3c3U, 0x4e4e4e4eU,
-	0x08080808U, 0x2e2e2e2eU, 0xa1a1a1a1U, 0x66666666U,
-	0x28282828U, 0xd9d9d9d9U, 0x24242424U, 0xb2b2b2b2U,
-	0x76767676U, 0x5b5b5b5bU, 0xa2a2a2a2U, 0x49494949U,
-	0x6d6d6d6dU, 0x8b8b8b8bU, 0xd1d1d1d1U, 0x25252525U,
-	0x72727272U, 0xf8f8f8f8U, 0xf6f6f6f6U, 0x64646464U,
-	0x86868686U, 0x68686868U, 0x98989898U, 0x16161616U,
-	0xd4d4d4d4U, 0xa4a4a4a4U, 0x5c5c5c5cU, 0xccccccccU,
-	0x5d5d5d5dU, 0x65656565U, 0xb6b6b6b6U, 0x92929292U,
-	0x6c6c6c6cU, 0x70707070U, 0x48484848U, 0x50505050U,
-	0xfdfdfdfdU, 0xededededU, 0xb9b9b9b9U, 0xdadadadaU,
-	0x5e5e5e5eU, 0x15151515U, 0x46464646U, 0x57575757U,
-	0xa7a7a7a7U, 0x8d8d8d8dU, 0x9d9d9d9dU, 0x84848484U,
-	0x90909090U, 0xd8d8d8d8U, 0xababababU, 0x00000000U,
-	0x8c8c8c8cU, 0xbcbcbcbcU, 0xd3d3d3d3U, 0x0a0a0a0aU,
-	0xf7f7f7f7U, 0xe4e4e4e4U, 0x58585858U, 0x05050505U,
-	0xb8b8b8b8U, 0xb3b3b3b3U, 0x45454545U, 0x06060606U,
-	0xd0d0d0d0U, 0x2c2c2c2cU, 0x1e1e1e1eU, 0x8f8f8f8fU,
-	0xcacacacaU, 0x3f3f3f3fU, 0x0f0f0f0fU, 0x02020202U,
-	0xc1c1c1c1U, 0xafafafafU, 0xbdbdbdbdU, 0x03030303U,
-	0x01010101U, 0x13131313U, 0x8a8a8a8aU, 0x6b6b6b6bU,
-	0x3a3a3a3aU, 0x91919191U, 0x11111111U, 0x41414141U,
-	0x4f4f4f4fU, 0x67676767U, 0xdcdcdcdcU, 0xeaeaeaeaU,
-	0x97979797U, 0xf2f2f2f2U, 0xcfcfcfcfU, 0xcecececeU,
-	0xf0f0f0f0U, 0xb4b4b4b4U, 0xe6e6e6e6U, 0x73737373U,
-	0x96969696U, 0xacacacacU, 0x74747474U, 0x22222222U,
-	0xe7e7e7e7U, 0xadadadadU, 0x35353535U, 0x85858585U,
-	0xe2e2e2e2U, 0xf9f9f9f9U, 0x37373737U, 0xe8e8e8e8U,
-	0x1c1c1c1cU, 0x75757575U, 0xdfdfdfdfU, 0x6e6e6e6eU,
-	0x47474747U, 0xf1f1f1f1U, 0x1a1a1a1aU, 0x71717171U,
-	0x1d1d1d1dU, 0x29292929U, 0xc5c5c5c5U, 0x89898989U,
-	0x6f6f6f6fU, 0xb7b7b7b7U, 0x62626262U, 0x0e0e0e0eU,
-	0xaaaaaaaaU, 0x18181818U, 0xbebebebeU, 0x1b1b1b1bU,
-	0xfcfcfcfcU, 0x56565656U, 0x3e3e3e3eU, 0x4b4b4b4bU,
-	0xc6c6c6c6U, 0xd2d2d2d2U, 0x79797979U, 0x20202020U,
-	0x9a9a9a9aU, 0xdbdbdbdbU, 0xc0c0c0c0U, 0xfefefefeU,
-	0x78787878U, 0xcdcdcdcdU, 0x5a5a5a5aU, 0xf4f4f4f4U,
-	0x1f1f1f1fU, 0xddddddddU, 0xa8a8a8a8U, 0x33333333U,
-	0x88888888U, 0x07070707U, 0xc7c7c7c7U, 0x31313131U,
-	0xb1b1b1b1U, 0x12121212U, 0x10101010U, 0x59595959U,
-	0x27272727U, 0x80808080U, 0xececececU, 0x5f5f5f5fU,
-	0x60606060U, 0x51515151U, 0x7f7f7f7fU, 0xa9a9a9a9U,
-	0x19191919U, 0xb5b5b5b5U, 0x4a4a4a4aU, 0x0d0d0d0dU,
-	0x2d2d2d2dU, 0xe5e5e5e5U, 0x7a7a7a7aU, 0x9f9f9f9fU,
-	0x93939393U, 0xc9c9c9c9U, 0x9c9c9c9cU, 0xefefefefU,
-	0xa0a0a0a0U, 0xe0e0e0e0U, 0x3b3b3b3bU, 0x4d4d4d4dU,
-	0xaeaeaeaeU, 0x2a2a2a2aU, 0xf5f5f5f5U, 0xb0b0b0b0U,
-	0xc8c8c8c8U, 0xebebebebU, 0xbbbbbbbbU, 0x3c3c3c3cU,
-	0x83838383U, 0x53535353U, 0x99999999U, 0x61616161U,
-	0x17171717U, 0x2b2b2b2bU, 0x04040404U, 0x7e7e7e7eU,
-	0xbabababaU, 0x77777777U, 0xd6d6d6d6U, 0x26262626U,
-	0xe1e1e1e1U, 0x69696969U, 0x14141414U, 0x63636363U,
-	0x55555555U, 0x21212121U, 0x0c0c0c0cU, 0x7d7d7d7dU,
-	};
-
-static const uint32 rcon[] =
+	return (AesITbox[0][AesSbox[(uint8)(x >>  0)]])
+		 ^ (AesITbox[1][AesSbox[(uint8)(x >>  8)]])
+		 ^ (AesITbox[2][AesSbox[(uint8)(x >> 16)]])
+		 ^ (AesITbox[3][AesSbox[(uint8)(x >> 24)]]);
+}
+
+static inline uint32 RotateRight(uint32 X, int N)
 {
-	0x01000000, 0x02000000, 0x04000000, 0x08000000,
-	0x10000000, 0x20000000, 0x40000000, 0x80000000,
-	0x1B000000, 0x36000000,
-	/* for 128-bit blocks, Rijndael never uses more than 10 rcon values */
-};
+	return (X >> N) | (X << (32 - N));
+}
 
-#define GETU32(plaintext) (((uint32)(plaintext)[0] << 24) ^ \
-						((uint32)(plaintext)[1] << 16) ^ \
-						((uint32)(plaintext)[2] <<  8) ^ \
-						((uint32)(plaintext)[3] <<  0))
-
-#define PUTU32(ciphertext, st) { (ciphertext)[0] = (uint8)((st) >> 24); \
-							  (ciphertext)[1] = (uint8)((st) >> 16); \
-							  (ciphertext)[2] = (uint8)((st) >>  8); \
-							  (ciphertext)[3] = (uint8)((st) >>  0); }
-
-/**
- * Expand the cipher key into the encryption key schedule.
- *
- * @return the number of rounds for the given cipher key size.
- */
-static int32 rijndaelSetupEncrypt( uint32* rk, const uint8* key, int32 keybits )
+static inline void AesEncryptExpand(FAesExpandedKey* EncryptKey, const uint8* Key)
 {
-	int32 i = 0;
-	uint32 temp;
-
-	rk[0] = GETU32( key + 0 );
-	rk[1] = GETU32( key + 4 );
-	rk[2] = GETU32( key + 8 );
-	rk[3] = GETU32( key + 12 );
-	if( keybits == 128 )
+	uint32* EKey = EncryptKey->Key;
+
+	EKey[0] = FPlatformMemory::ReadUnaligned<uint32>(Key +  0);
+	EKey[1] = FPlatformMemory::ReadUnaligned<uint32>(Key +  4);
+	EKey[2] = FPlatformMemory::ReadUnaligned<uint32>(Key +  8);
+	EKey[3] = FPlatformMemory::ReadUnaligned<uint32>(Key + 12);
+	EKey[4] = FPlatformMemory::ReadUnaligned<uint32>(Key + 16);
+	EKey[5] = FPlatformMemory::ReadUnaligned<uint32>(Key + 20);
+	EKey[6] = FPlatformMemory::ReadUnaligned<uint32>(Key + 24);
+	EKey[7] = FPlatformMemory::ReadUnaligned<uint32>(Key + 28);
+
+	for (int Index = 0; Index < AES256_ROUND_COUNT/2; Index++, EKey += 8)
 	{
-		for( ; ; )
+		// rcon[] = { 0x01, 0x02, 0x04, 0x08, 0x10, 0x20, 0x40, 0x80 }
+		uint32 Rcon = 1U << Index;
+
+		EKey[ 8] = EKey[0] ^ RotateRight(AesEncryptMix(EKey[7]), 8) ^ Rcon;
+		EKey[ 9] = EKey[1] ^ EKey[ 8];
+		EKey[10] = EKey[2] ^ EKey[ 9];
+		EKey[11] = EKey[3] ^ EKey[10];
+
+		if (Index < AES256_ROUND_COUNT/2 - 1)
 		{
-			temp = rk[3];
-			rk[4] = rk[0] ^
-				   ( Te4[( temp >> 16 ) & 0xff] & 0xff000000 ) ^
-				   ( Te4[( temp >>  8 ) & 0xff] & 0x00ff0000 ) ^
-				   ( Te4[( temp >>  0 ) & 0xff] & 0x0000ff00 ) ^
-				   ( Te4[( temp >> 24 ) & 0xff] & 0x000000ff ) ^
-				   rcon[i];
-			rk[5] = rk[1] ^ rk[4];
-			rk[6] = rk[2] ^ rk[5];
-			rk[7] = rk[3] ^ rk[6];
-			if( ++i == 10 )
-			{
-				return 10;
-			}
-
-			rk += 4;
+			EKey[12] = EKey[4] ^ AesEncryptMix(EKey[11]);
+			EKey[13] = EKey[5] ^ EKey[12];
+			EKey[14] = EKey[6] ^ EKey[13];
+			EKey[15] = EKey[7] ^ EKey[14];
 		}
 	}
-	rk[4] = GETU32( key + 16 );
-	rk[5] = GETU32( key + 20 );
-	if( keybits == 192 )
+}
+
+static inline void AesDecryptExpand(FAesExpandedKey* DecryptKey, const uint8* Key)
+{
+	FAesExpandedKey EncryptKey;
+	AesEncryptExpand(&EncryptKey, Key);
+
+	const uint32* EKey = EncryptKey.Key + (4 * AES256_ROUND_COUNT);
+	uint32* DKey = DecryptKey->Key;
+
+	DKey[0] = EKey[0];
+	DKey[1] = EKey[1];
+	DKey[2] = EKey[2];
+	DKey[3] = EKey[3];
+	DKey += 4;
+	EKey -= 4;
+
+	for (int Round = 1; Round < AES256_ROUND_COUNT; Round++)
+	{
+		DKey[0] = AesDecryptMix(EKey[0]);
+		DKey[1] = AesDecryptMix(EKey[1]);
+		DKey[2] = AesDecryptMix(EKey[2]);
+		DKey[3] = AesDecryptMix(EKey[3]);
+		DKey += 4;
+		EKey -= 4;
+	}
+
+	DKey[0] = EKey[0];
+	DKey[1] = EKey[1];
+	DKey[2] = EKey[2];
+	DKey[3] = EKey[3];
+}
+
+static inline void AesEncrypt(const uint8* Key, uint8* Contents, uint32 NumBytes)
+{
+#define AES_ENC(a, b, c, d) (      \
+	AesTbox[0][(uint8)(a >>  0)] ^ \
+	AesTbox[1][(uint8)(b >>  8)] ^ \
+	AesTbox[2][(uint8)(c >> 16)] ^ \
+	AesTbox[3][(uint8)(d >> 24)]   \
+)
+
+#define AES_ENC_LAST(a, b, c, d) (      \
+	(AesSbox[(uint8)(a >>  0)] <<  0) ^ \
+	(AesSbox[(uint8)(b >>  8)] <<  8) ^ \
+	(AesSbox[(uint8)(c >> 16)] << 16) ^ \
+	(AesSbox[(uint8)(d >> 24)] << 24)   \
+)
+
+	FAesExpandedKey EncryptKey;
+	AesEncryptExpand(&EncryptKey, Key);
+
+	const uint32* EKey = EncryptKey.Key;
+
+	for (uint32 Offset = 0; Offset < NumBytes; Offset += 16)
 	{
-		for( ; ; )
+		uint8* Block = Contents + Offset;
+
+		uint32 s0 = FPlatformMemory::ReadUnaligned<uint32>(Block +  0) ^ EKey[0];
+		uint32 s1 = FPlatformMemory::ReadUnaligned<uint32>(Block +  4) ^ EKey[1];
+		uint32 s2 = FPlatformMemory::ReadUnaligned<uint32>(Block +  8) ^ EKey[2];
+		uint32 s3 = FPlatformMemory::ReadUnaligned<uint32>(Block + 12) ^ EKey[3];
+
+		for (int Round = 1; Round < AES256_ROUND_COUNT; Round++)
 		{
-			temp = rk[5];
-			rk[6] = rk[0] ^
-				   ( Te4[( temp >> 16 ) & 0xff] & 0xff000000 ) ^
-				   ( Te4[( temp >>  8 ) & 0xff] & 0x00ff0000 ) ^
-				   ( Te4[( temp >>  0 ) & 0xff] & 0x0000ff00 ) ^
-				   ( Te4[( temp >> 24 ) & 0xff] & 0x000000ff ) ^
-				   rcon[i];
-			rk[7] = rk[1] ^ rk[6];
-			rk[8] = rk[2] ^ rk[7];
-			rk[9] = rk[3] ^ rk[8];
-			if( ++i == 8 )
-			{
-				return 12;
-			}
-
-			rk[10] = rk[4] ^ rk[9];
-			rk[11] = rk[5] ^ rk[10];
-			rk += 6;
+			uint32 t0 = AES_ENC(s0, s1, s2, s3);
+			uint32 t1 = AES_ENC(s1, s2, s3, s0);
+			uint32 t2 = AES_ENC(s2, s3, s0, s1);
+			uint32 t3 = AES_ENC(s3, s0, s1, s2);
+
+			s0 = t0 ^ EKey[4 * Round + 0];
+			s1 = t1 ^ EKey[4 * Round + 1];
+			s2 = t2 ^ EKey[4 * Round + 2];
+			s3 = t3 ^ EKey[4 * Round + 3];
 		}
+
+		uint32 t0 = AES_ENC_LAST(s0, s1, s2, s3);
+		uint32 t1 = AES_ENC_LAST(s1, s2, s3, s0);
+		uint32 t2 = AES_ENC_LAST(s2, s3, s0, s1);
+		uint32 t3 = AES_ENC_LAST(s3, s0, s1, s2);
+
+		s0 = t0 ^ EKey[4 * AES256_ROUND_COUNT + 0];
+		s1 = t1 ^ EKey[4 * AES256_ROUND_COUNT + 1];
+		s2 = t2 ^ EKey[4 * AES256_ROUND_COUNT + 2];
+		s3 = t3 ^ EKey[4 * AES256_ROUND_COUNT + 3];
+
+		FPlatformMemory::WriteUnaligned<uint32>(Block +  0, s0);
+		FPlatformMemory::WriteUnaligned<uint32>(Block +  4, s1);
+		FPlatformMemory::WriteUnaligned<uint32>(Block +  8, s2);
+		FPlatformMemory::WriteUnaligned<uint32>(Block + 12, s3);
 	}
-	rk[6] = GETU32( key + 24 );
-	rk[7] = GETU32( key + 28 );
-	if( keybits == 256 )
+
+#undef AES_ENC
+#undef AES_ENC_LAST
+}
+
+static inline void AesDecrypt(const uint8* Key, uint8* Contents, uint32 NumBytes)
+{
+#define AES_DEC(a, b, c, d) (       \
+	AesITbox[0][(uint8)(a >>  0)] ^ \
+	AesITbox[1][(uint8)(b >>  8)] ^ \
+	AesITbox[2][(uint8)(c >> 16)] ^ \
+	AesITbox[3][(uint8)(d >> 24)]   \
+)
+
+#define AES_DEC_LAST(a, b, c, d) (       \
+	(AesISbox[(uint8)(a >>  0)] <<  0) ^ \
+	(AesISbox[(uint8)(b >>  8)] <<  8) ^ \
+	(AesISbox[(uint8)(c >> 16)] << 16) ^ \
+	(AesISbox[(uint8)(d >> 24)] << 24)   \
+)
+
+	FAesExpandedKey DecryptKey;
+	AesDecryptExpand(&DecryptKey, Key);
+
+	const uint32* DKey = DecryptKey.Key;
+
+	for (uint32 Offset = 0; Offset < NumBytes; Offset += 16)
 	{
-		for( ; ; )
+		uint8* Block = Contents + Offset;
+
+		uint32 s0 = FPlatformMemory::ReadUnaligned<uint32>(Block +  0) ^ DKey[0];
+		uint32 s1 = FPlatformMemory::ReadUnaligned<uint32>(Block +  4) ^ DKey[1];
+		uint32 s2 = FPlatformMemory::ReadUnaligned<uint32>(Block +  8) ^ DKey[2];
+		uint32 s3 = FPlatformMemory::ReadUnaligned<uint32>(Block + 12) ^ DKey[3];
+
+		for (int Round = 1; Round < AES256_ROUND_COUNT; Round++)
 		{
-			temp = rk[7];
-			rk[8] = rk[0] ^
-				   ( Te4[( temp >> 16 ) & 0xff] & 0xff000000 ) ^
-				   ( Te4[( temp >>  8 ) & 0xff] & 0x00ff0000 ) ^
-				   ( Te4[( temp >>  0 ) & 0xff] & 0x0000ff00 ) ^
-				   ( Te4[( temp >> 24 ) & 0xff] & 0x000000ff ) ^
-				   rcon[i];
-			rk[9] = rk[1] ^ rk[8];
-			rk[10] = rk[2] ^ rk[9];
-			rk[11] = rk[3] ^ rk[10];
-			if( ++i == 7 )
-			{
-				return 14;
-			}
-
-			temp = rk[11];
-			rk[12] = rk[4] ^
-					( Te4[( temp >> 24 ) & 0xff] & 0xff000000 ) ^
-					( Te4[( temp >> 16 ) & 0xff] & 0x00ff0000 ) ^
-					( Te4[( temp >>  8 ) & 0xff] & 0x0000ff00 ) ^
-					( Te4[( temp >>  0 ) & 0xff] & 0x000000ff );
-			rk[13] = rk[5] ^ rk[12];
-			rk[14] = rk[6] ^ rk[13];
-			rk[15] = rk[7] ^ rk[14];
-			rk += 8;
+			uint32 t0 = AES_DEC(s0, s3, s2, s1);
+			uint32 t1 = AES_DEC(s1, s0, s3, s2);
+			uint32 t2 = AES_DEC(s2, s1, s0, s3);
+			uint32 t3 = AES_DEC(s3, s2, s1, s0);
+
+			s0 = t0 ^ DKey[4 * Round + 0];
+			s1 = t1 ^ DKey[4 * Round + 1];
+			s2 = t2 ^ DKey[4 * Round + 2];
+			s3 = t3 ^ DKey[4 * Round + 3];
 		}
+
+		uint32 t0 = AES_DEC_LAST(s0, s3, s2, s1);
+		uint32 t1 = AES_DEC_LAST(s1, s0, s3, s2);
+		uint32 t2 = AES_DEC_LAST(s2, s1, s0, s3);
+		uint32 t3 = AES_DEC_LAST(s3, s2, s1, s0);
+
+		s0 = t0 ^ DKey[4 * AES256_ROUND_COUNT + 0];
+		s1 = t1 ^ DKey[4 * AES256_ROUND_COUNT + 1];
+		s2 = t2 ^ DKey[4 * AES256_ROUND_COUNT + 2];
+		s3 = t3 ^ DKey[4 * AES256_ROUND_COUNT + 3];
+
+		FPlatformMemory::WriteUnaligned<uint32>(Block +  0, s0);
+		FPlatformMemory::WriteUnaligned<uint32>(Block +  4, s1);
+		FPlatformMemory::WriteUnaligned<uint32>(Block +  8, s2);
+		FPlatformMemory::WriteUnaligned<uint32>(Block + 12, s3);
 	}
 
-	return 0;
+#undef AES_DEC
+#undef AES_DEC_LAST
+}
+
+#endif
+
+#if UE_PLATFORM_AES_X86
+
+static inline bool DetectAesInstructions()
+{
+	int info[4];
+#if PLATFORM_COMPILER_CLANG
+	__cpuid(1, info[0], info[1], info[2], info[3]);
+#else
+	__cpuid(info, 1);
+#endif
+	return (info[2] & 0x02000000) != 0;
 }
 
-/**
- * Expand the cipher key into the decryption key schedule.
- *
- * @return the number of rounds for the given cipher key size.
- */
-static int32 rijndaelSetupDecrypt( uint32* rk, const uint8* key, int32 keybits )
+static bool CanUseAesInstructions()
 {
-	int32 nrounds, i, j;
-	uint32 temp;
+#if UE_PLATFORM_AES_FALLBACK
+	// run cpuid only once
+	static bool bCanUse = DetectAesInstructions();
+	return bCanUse;
+#else
+	// no AES fallback code means AESNI will be used unconditionally
+	return true;
+#endif
+}
 
-	/* expand the cipher key: */
-	nrounds = rijndaelSetupEncrypt( rk, key, keybits );
+// https://www.intel.com/content/dam/doc/white-paper/advanced-encryption-standard-new-instructions-set-paper.pdf
+
+#define AES_INIT_ENC_X86(ctx, x, y, rcon)         \
+{                                                 \
+    __m128i z, t;                                 \
+    ctx = y;                                      \
+    z = _mm_aeskeygenassist_si128(y, rcon);       \
+    z = _mm_shuffle_epi32(z, rcon ? 0xff : 0xaa); \
+    t = _mm_slli_si128(x, 4);                     \
+    x = _mm_xor_si128(x, t);                      \
+    t = _mm_slli_si128(t, 4);                     \
+    x = _mm_xor_si128(x, t);                      \
+    t = _mm_slli_si128(t, 4);                     \
+    x = _mm_xor_si128(x, t);                      \
+    x = _mm_xor_si128(x, z);                      \
+}
 
-	/* invert the order of the round keys: */
-	for( i = 0, j = 4 * nrounds; i < j; i += 4, j -= 4 )
-	{
-		temp = rk[i + 0]; 
-		rk[i + 0] = rk[j + 0];
-		rk[j + 0] = temp;
+#define AES_INIT_DEC_X86(ctx, x, y, rcon)         \
+{                                                 \
+    __m128i z, t;                                 \
+    ctx = _mm_aesimc_si128(y);                    \
+    z = _mm_aeskeygenassist_si128(y, rcon);       \
+    z = _mm_shuffle_epi32(z, rcon ? 0xff : 0xaa); \
+    t = _mm_slli_si128(x, 4);                     \
+    x = _mm_xor_si128(x, t);                      \
+    t = _mm_slli_si128(t, 4);                     \
+    x = _mm_xor_si128(x, t);                      \
+    t = _mm_slli_si128(t, 4);                     \
+    x = _mm_xor_si128(x, t);                      \
+    x = _mm_xor_si128(x, z);                      \
+}
 
-		temp = rk[i + 1];
-		rk[i + 1] = rk[j + 1];
-		rk[j + 1] = temp;
+#if PLATFORM_COMPILER_CLANG
+#pragma clang attribute push (__attribute__((target("aes"))), apply_to=function)
+#endif
 
-		temp = rk[i + 2];
-		rk[i + 2] = rk[j + 2];
-		rk[j + 2] = temp;
+static inline void AesEncryptX86(const void* Key, void* Contents, uint32 NumBytes)
+{
+	const __m128i* Key128 = reinterpret_cast<const __m128i*>(Key);
+	__m128i x = _mm_loadu_si128(Key128 + 0);
+	__m128i y = _mm_loadu_si128(Key128 + 1);
+
+	__m128i EKey[AES256_ROUND_COUNT + 1];
+	EKey[0] = x;
+	AES_INIT_ENC_X86(EKey[ 1], x, y, 0x01);
+	AES_INIT_ENC_X86(EKey[ 2], y, x, 0x00);
+	AES_INIT_ENC_X86(EKey[ 3], x, y, 0x02);
+	AES_INIT_ENC_X86(EKey[ 4], y, x, 0x00);
+	AES_INIT_ENC_X86(EKey[ 5], x, y, 0x04);
+	AES_INIT_ENC_X86(EKey[ 6], y, x, 0x00);
+	AES_INIT_ENC_X86(EKey[ 7], x, y, 0x08);
+	AES_INIT_ENC_X86(EKey[ 8], y, x, 0x00);
+	AES_INIT_ENC_X86(EKey[ 9], x, y, 0x10);
+	AES_INIT_ENC_X86(EKey[10], y, x, 0x00);
+	AES_INIT_ENC_X86(EKey[11], x, y, 0x20);
+	AES_INIT_ENC_X86(EKey[12], y, x, 0x00);
+	AES_INIT_ENC_X86(EKey[13], x, y, 0x40);
+	EKey[14] = x;
+
+	// unrolling loop to do multiple AES blocks is faster on newer CPUs
+	// for example, ~10% improvement on Zen2 threadripper (4989 -> 5501 MB/s)
+	__m128i* Data = reinterpret_cast<__m128i*>(Contents);
+	while (NumBytes >= 32)
+	{
+		__m128i Block0 = _mm_loadu_si128(Data + 0);
+		__m128i Block1 = _mm_loadu_si128(Data + 1);
+
+		Block0 = _mm_xor_si128(Block0, EKey[0]);
+		Block1 = _mm_xor_si128(Block1, EKey[0]);
+		Block0 = _mm_aesenc_si128(Block0, EKey[1]);
+		Block1 = _mm_aesenc_si128(Block1, EKey[1]);
+		Block0 = _mm_aesenc_si128(Block0, EKey[2]);
+		Block1 = _mm_aesenc_si128(Block1, EKey[2]);
+		Block0 = _mm_aesenc_si128(Block0, EKey[3]);
+		Block1 = _mm_aesenc_si128(Block1, EKey[3]);
+		Block0 = _mm_aesenc_si128(Block0, EKey[4]);
+		Block1 = _mm_aesenc_si128(Block1, EKey[4]);
+		Block0 = _mm_aesenc_si128(Block0, EKey[5]);
+		Block1 = _mm_aesenc_si128(Block1, EKey[5]);
+		Block0 = _mm_aesenc_si128(Block0, EKey[6]);
+		Block1 = _mm_aesenc_si128(Block1, EKey[6]);
+		Block0 = _mm_aesenc_si128(Block0, EKey[7]);
+		Block1 = _mm_aesenc_si128(Block1, EKey[7]);
+		Block0 = _mm_aesenc_si128(Block0, EKey[8]);
+		Block1 = _mm_aesenc_si128(Block1, EKey[8]);
+		Block0 = _mm_aesenc_si128(Block0, EKey[9]);
+		Block1 = _mm_aesenc_si128(Block1, EKey[9]);
+		Block0 = _mm_aesenc_si128(Block0, EKey[10]);
+		Block1 = _mm_aesenc_si128(Block1, EKey[10]);
+		Block0 = _mm_aesenc_si128(Block0, EKey[11]);
+		Block1 = _mm_aesenc_si128(Block1, EKey[11]);
+		Block0 = _mm_aesenc_si128(Block0, EKey[12]);
+		Block1 = _mm_aesenc_si128(Block1, EKey[12]);
+		Block0 = _mm_aesenc_si128(Block0, EKey[13]);
+		Block1 = _mm_aesenc_si128(Block1, EKey[13]);
+		Block0 = _mm_aesenclast_si128(Block0, EKey[14]);
+		Block1 = _mm_aesenclast_si128(Block1, EKey[14]);
+
+		_mm_storeu_si128(Data + 0, Block0);
+		_mm_storeu_si128(Data + 1, Block1);
+		Data += 2;
+		NumBytes -= 32;
+	}
 
-		temp = rk[i + 3];
-		rk[i + 3] = rk[j + 3];
-		rk[j + 3] = temp;
+	if (NumBytes != 0)
+	{
+		__m128i Block = _mm_loadu_si128(Data);
+
+		Block = _mm_xor_si128(Block, EKey[0]);
+		Block = _mm_aesenc_si128(Block, EKey[1]);
+		Block = _mm_aesenc_si128(Block, EKey[2]);
+		Block = _mm_aesenc_si128(Block, EKey[3]);
+		Block = _mm_aesenc_si128(Block, EKey[4]);
+		Block = _mm_aesenc_si128(Block, EKey[5]);
+		Block = _mm_aesenc_si128(Block, EKey[6]);
+		Block = _mm_aesenc_si128(Block, EKey[7]);
+		Block = _mm_aesenc_si128(Block, EKey[8]);
+		Block = _mm_aesenc_si128(Block, EKey[9]);
+		Block = _mm_aesenc_si128(Block, EKey[10]);
+		Block = _mm_aesenc_si128(Block, EKey[11]);
+		Block = _mm_aesenc_si128(Block, EKey[12]);
+		Block = _mm_aesenc_si128(Block, EKey[13]);
+		Block = _mm_aesenclast_si128(Block, EKey[14]);
+
+		_mm_storeu_si128(Data, Block);
 	}
+}
 
-	/* apply the inverse MixColumn transform to all round keys but the first and the last: */
-	for( i = 1; i < nrounds; i++ )
+static inline void AesDecryptX86(const void* Key, void* Contents, uint32 NumBytes)
+{
+	const __m128i* Key128 = reinterpret_cast<const __m128i*>(Key);
+	__m128i x = _mm_loadu_si128(Key128 + 0);
+	__m128i y = _mm_loadu_si128(Key128 + 1);
+
+	__m128i DKey[AES256_ROUND_COUNT + 1];
+	DKey[14] = x;
+	AES_INIT_DEC_X86(DKey[13], x, y, 0x01);
+	AES_INIT_DEC_X86(DKey[12], y, x, 0x00);
+	AES_INIT_DEC_X86(DKey[11], x, y, 0x02);
+	AES_INIT_DEC_X86(DKey[10], y, x, 0x00);
+	AES_INIT_DEC_X86(DKey[ 9], x, y, 0x04);
+	AES_INIT_DEC_X86(DKey[ 8], y, x, 0x00);
+	AES_INIT_DEC_X86(DKey[ 7], x, y, 0x08);
+	AES_INIT_DEC_X86(DKey[ 6], y, x, 0x00);
+	AES_INIT_DEC_X86(DKey[ 5], x, y, 0x10);
+	AES_INIT_DEC_X86(DKey[ 4], y, x, 0x00);
+	AES_INIT_DEC_X86(DKey[ 3], x, y, 0x20);
+	AES_INIT_DEC_X86(DKey[ 2], y, x, 0x00);
+	AES_INIT_DEC_X86(DKey[ 1], x, y, 0x40);
+	DKey[0] = x;
+
+	__m128i* Data = reinterpret_cast<__m128i*>(Contents);
+	while (NumBytes >= 32)
 	{
-		rk += 4;
-		rk[0] =
-			Td0[Te4[( rk[0] >> 24 ) & 0xff] & 0xff] ^
-			Td1[Te4[( rk[0] >> 16 ) & 0xff] & 0xff] ^
-			Td2[Te4[( rk[0] >>  8 ) & 0xff] & 0xff] ^
-			Td3[Te4[( rk[0] >>  0 ) & 0xff] & 0xff];
-		rk[1] =
-			Td0[Te4[( rk[1] >> 24 ) & 0xff] & 0xff] ^
-			Td1[Te4[( rk[1] >> 16 ) & 0xff] & 0xff] ^
-			Td2[Te4[( rk[1] >>  8 ) & 0xff] & 0xff] ^
-			Td3[Te4[( rk[1] >>  0 ) & 0xff] & 0xff];
-		rk[2] =
-			Td0[Te4[( rk[2] >> 24 ) & 0xff] & 0xff] ^
-			Td1[Te4[( rk[2] >> 16 ) & 0xff] & 0xff] ^
-			Td2[Te4[( rk[2] >>  8 ) & 0xff] & 0xff] ^
-			Td3[Te4[( rk[2] >>  0 ) & 0xff] & 0xff];
-		rk[3] =
-			Td0[Te4[( rk[3] >> 24 ) & 0xff] & 0xff] ^
-			Td1[Te4[( rk[3] >> 16 ) & 0xff] & 0xff] ^
-			Td2[Te4[( rk[3] >>  8 ) & 0xff] & 0xff] ^
-			Td3[Te4[( rk[3] >>  0 ) & 0xff] & 0xff];
+		__m128i Block0 = _mm_loadu_si128(Data + 0);
+		__m128i Block1 = _mm_loadu_si128(Data + 1);
+
+		Block0 = _mm_xor_si128(Block0, DKey[0]);
+		Block1 = _mm_xor_si128(Block1, DKey[0]);
+		Block0 = _mm_aesdec_si128(Block0, DKey[1]);
+		Block1 = _mm_aesdec_si128(Block1, DKey[1]);
+		Block0 = _mm_aesdec_si128(Block0, DKey[2]);
+		Block1 = _mm_aesdec_si128(Block1, DKey[2]);
+		Block0 = _mm_aesdec_si128(Block0, DKey[3]);
+		Block1 = _mm_aesdec_si128(Block1, DKey[3]);
+		Block0 = _mm_aesdec_si128(Block0, DKey[4]);
+		Block1 = _mm_aesdec_si128(Block1, DKey[4]);
+		Block0 = _mm_aesdec_si128(Block0, DKey[5]);
+		Block1 = _mm_aesdec_si128(Block1, DKey[5]);
+		Block0 = _mm_aesdec_si128(Block0, DKey[6]);
+		Block1 = _mm_aesdec_si128(Block1, DKey[6]);
+		Block0 = _mm_aesdec_si128(Block0, DKey[7]);
+		Block1 = _mm_aesdec_si128(Block1, DKey[7]);
+		Block0 = _mm_aesdec_si128(Block0, DKey[8]);
+		Block1 = _mm_aesdec_si128(Block1, DKey[8]);
+		Block0 = _mm_aesdec_si128(Block0, DKey[9]);
+		Block1 = _mm_aesdec_si128(Block1, DKey[9]);
+		Block0 = _mm_aesdec_si128(Block0, DKey[10]);
+		Block1 = _mm_aesdec_si128(Block1, DKey[10]);
+		Block0 = _mm_aesdec_si128(Block0, DKey[11]);
+		Block1 = _mm_aesdec_si128(Block1, DKey[11]);
+		Block0 = _mm_aesdec_si128(Block0, DKey[12]);
+		Block1 = _mm_aesdec_si128(Block1, DKey[12]);
+		Block0 = _mm_aesdec_si128(Block0, DKey[13]);
+		Block1 = _mm_aesdec_si128(Block1, DKey[13]);
+		Block0 = _mm_aesdeclast_si128(Block0, DKey[14]);
+		Block1 = _mm_aesdeclast_si128(Block1, DKey[14]);
+
+		_mm_storeu_si128(Data + 0, Block0);
+		_mm_storeu_si128(Data + 1, Block1);
+		Data += 2;
+		NumBytes -= 32;
 	}
 
-	return nrounds;
+	if (NumBytes != 0)
+	{
+		__m128i Block = _mm_loadu_si128(Data);
+
+		Block = _mm_xor_si128(Block, DKey[0]);
+		Block = _mm_aesdec_si128(Block, DKey[1]);
+		Block = _mm_aesdec_si128(Block, DKey[2]);
+		Block = _mm_aesdec_si128(Block, DKey[3]);
+		Block = _mm_aesdec_si128(Block, DKey[4]);
+		Block = _mm_aesdec_si128(Block, DKey[5]);
+		Block = _mm_aesdec_si128(Block, DKey[6]);
+		Block = _mm_aesdec_si128(Block, DKey[7]);
+		Block = _mm_aesdec_si128(Block, DKey[8]);
+		Block = _mm_aesdec_si128(Block, DKey[9]);
+		Block = _mm_aesdec_si128(Block, DKey[10]);
+		Block = _mm_aesdec_si128(Block, DKey[11]);
+		Block = _mm_aesdec_si128(Block, DKey[12]);
+		Block = _mm_aesdec_si128(Block, DKey[13]);
+		Block = _mm_aesdeclast_si128(Block, DKey[14]);
+
+		_mm_storeu_si128(Data, Block);
+	}
 }
 
-static void rijndaelEncrypt( const uint32 *rk, int32 nrounds, const uint8 plaintext[16], uint8 ciphertext[16] )
+#if PLATFORM_COMPILER_CLANG
+#pragma clang attribute pop
+#endif
+
+#endif
+
+
+#if UE_PLATFORM_AES_ARMV8
+
+struct FAesExpandedKeyARMV8
+{
+	uint8x16_t Key[AES256_ROUND_COUNT + 1];
+};
+
+static inline uint32 AesEncryptMixARMV8(uint32 x)
+{
+	uint8x16_t Block = vreinterpretq_u8_u32(vmovq_n_u32(x));
+	Block = vaeseq_u8(Block, vdupq_n_u8(0));
+	return vgetq_lane_u32(vreinterpretq_u32_u8(Block), 0);
+}
+
+static inline uint32 RotateRight(uint32 X, int N)
+{
+	return (X >> N) | (X << (32 - N));
+}
+
+static inline void AesEncryptExpandARMV8(FAesExpandedKeyARMV8* EncryptKey, const uint8* Key)
 {
-	uint32 s0, s1, s2, s3, t0, t1, t2, t3;
-
-	/*
-	 * map byte array block to cipher state
-	 * and add initial round key:
-	 */
-	s0 = GETU32( plaintext +  0 ) ^ rk[0];
-	s1 = GETU32( plaintext +  4 ) ^ rk[1];
-	s2 = GETU32( plaintext +  8 ) ^ rk[2];
-	s3 = GETU32( plaintext + 12 ) ^ rk[3];
-
-	/* round 1: */
-	t0 = Te0[s0 >> 24] ^ Te1[(s1 >> 16) & 0xff] ^ Te2[(s2 >>  8) & 0xff] ^ Te3[s3 & 0xff] ^ rk[ 4];
-	t1 = Te0[s1 >> 24] ^ Te1[(s2 >> 16) & 0xff] ^ Te2[(s3 >>  8) & 0xff] ^ Te3[s0 & 0xff] ^ rk[ 5];
-	t2 = Te0[s2 >> 24] ^ Te1[(s3 >> 16) & 0xff] ^ Te2[(s0 >>  8) & 0xff] ^ Te3[s1 & 0xff] ^ rk[ 6];
-	t3 = Te0[s3 >> 24] ^ Te1[(s0 >> 16) & 0xff] ^ Te2[(s1 >>  8) & 0xff] ^ Te3[s2 & 0xff] ^ rk[ 7];
-	/* round 2: */
-	s0 = Te0[t0 >> 24] ^ Te1[(t1 >> 16) & 0xff] ^ Te2[(t2 >>  8) & 0xff] ^ Te3[t3 & 0xff] ^ rk[ 8];
-	s1 = Te0[t1 >> 24] ^ Te1[(t2 >> 16) & 0xff] ^ Te2[(t3 >>  8) & 0xff] ^ Te3[t0 & 0xff] ^ rk[ 9];
-	s2 = Te0[t2 >> 24] ^ Te1[(t3 >> 16) & 0xff] ^ Te2[(t0 >>  8) & 0xff] ^ Te3[t1 & 0xff] ^ rk[10];
-	s3 = Te0[t3 >> 24] ^ Te1[(t0 >> 16) & 0xff] ^ Te2[(t1 >>  8) & 0xff] ^ Te3[t2 & 0xff] ^ rk[11];
-	/* round 3: */
-	t0 = Te0[s0 >> 24] ^ Te1[(s1 >> 16) & 0xff] ^ Te2[(s2 >>  8) & 0xff] ^ Te3[s3 & 0xff] ^ rk[12];
-	t1 = Te0[s1 >> 24] ^ Te1[(s2 >> 16) & 0xff] ^ Te2[(s3 >>  8) & 0xff] ^ Te3[s0 & 0xff] ^ rk[13];
-	t2 = Te0[s2 >> 24] ^ Te1[(s3 >> 16) & 0xff] ^ Te2[(s0 >>  8) & 0xff] ^ Te3[s1 & 0xff] ^ rk[14];
-	t3 = Te0[s3 >> 24] ^ Te1[(s0 >> 16) & 0xff] ^ Te2[(s1 >>  8) & 0xff] ^ Te3[s2 & 0xff] ^ rk[15];
-	/* round 4: */
-	s0 = Te0[t0 >> 24] ^ Te1[(t1 >> 16) & 0xff] ^ Te2[(t2 >>  8) & 0xff] ^ Te3[t3 & 0xff] ^ rk[16];
-	s1 = Te0[t1 >> 24] ^ Te1[(t2 >> 16) & 0xff] ^ Te2[(t3 >>  8) & 0xff] ^ Te3[t0 & 0xff] ^ rk[17];
-	s2 = Te0[t2 >> 24] ^ Te1[(t3 >> 16) & 0xff] ^ Te2[(t0 >>  8) & 0xff] ^ Te3[t1 & 0xff] ^ rk[18];
-	s3 = Te0[t3 >> 24] ^ Te1[(t0 >> 16) & 0xff] ^ Te2[(t1 >>  8) & 0xff] ^ Te3[t2 & 0xff] ^ rk[19];
-	/* round 5: */
-	t0 = Te0[s0 >> 24] ^ Te1[(s1 >> 16) & 0xff] ^ Te2[(s2 >>  8) & 0xff] ^ Te3[s3 & 0xff] ^ rk[20];
-	t1 = Te0[s1 >> 24] ^ Te1[(s2 >> 16) & 0xff] ^ Te2[(s3 >>  8) & 0xff] ^ Te3[s0 & 0xff] ^ rk[21];
-	t2 = Te0[s2 >> 24] ^ Te1[(s3 >> 16) & 0xff] ^ Te2[(s0 >>  8) & 0xff] ^ Te3[s1 & 0xff] ^ rk[22];
-	t3 = Te0[s3 >> 24] ^ Te1[(s0 >> 16) & 0xff] ^ Te2[(s1 >>  8) & 0xff] ^ Te3[s2 & 0xff] ^ rk[23];
-	/* round 6: */
-	s0 = Te0[t0 >> 24] ^ Te1[(t1 >> 16) & 0xff] ^ Te2[(t2 >>  8) & 0xff] ^ Te3[t3 & 0xff] ^ rk[24];
-	s1 = Te0[t1 >> 24] ^ Te1[(t2 >> 16) & 0xff] ^ Te2[(t3 >>  8) & 0xff] ^ Te3[t0 & 0xff] ^ rk[25];
-	s2 = Te0[t2 >> 24] ^ Te1[(t3 >> 16) & 0xff] ^ Te2[(t0 >>  8) & 0xff] ^ Te3[t1 & 0xff] ^ rk[26];
-	s3 = Te0[t3 >> 24] ^ Te1[(t0 >> 16) & 0xff] ^ Te2[(t1 >>  8) & 0xff] ^ Te3[t2 & 0xff] ^ rk[27];
-	/* round 7: */
-	t0 = Te0[s0 >> 24] ^ Te1[(s1 >> 16) & 0xff] ^ Te2[(s2 >>  8) & 0xff] ^ Te3[s3 & 0xff] ^ rk[28];
-	t1 = Te0[s1 >> 24] ^ Te1[(s2 >> 16) & 0xff] ^ Te2[(s3 >>  8) & 0xff] ^ Te3[s0 & 0xff] ^ rk[29];
-	t2 = Te0[s2 >> 24] ^ Te1[(s3 >> 16) & 0xff] ^ Te2[(s0 >>  8) & 0xff] ^ Te3[s1 & 0xff] ^ rk[30];
-	t3 = Te0[s3 >> 24] ^ Te1[(s0 >> 16) & 0xff] ^ Te2[(s1 >>  8) & 0xff] ^ Te3[s2 & 0xff] ^ rk[31];
-	/* round 8: */
-	s0 = Te0[t0 >> 24] ^ Te1[(t1 >> 16) & 0xff] ^ Te2[(t2 >>  8) & 0xff] ^ Te3[t3 & 0xff] ^ rk[32];
-	s1 = Te0[t1 >> 24] ^ Te1[(t2 >> 16) & 0xff] ^ Te2[(t3 >>  8) & 0xff] ^ Te3[t0 & 0xff] ^ rk[33];
-	s2 = Te0[t2 >> 24] ^ Te1[(t3 >> 16) & 0xff] ^ Te2[(t0 >>  8) & 0xff] ^ Te3[t1 & 0xff] ^ rk[34];
-	s3 = Te0[t3 >> 24] ^ Te1[(t0 >> 16) & 0xff] ^ Te2[(t1 >>  8) & 0xff] ^ Te3[t2 & 0xff] ^ rk[35];
-	/* round 9: */
-	t0 = Te0[s0 >> 24] ^ Te1[(s1 >> 16) & 0xff] ^ Te2[(s2 >>  8) & 0xff] ^ Te3[s3 & 0xff] ^ rk[36];
-	t1 = Te0[s1 >> 24] ^ Te1[(s2 >> 16) & 0xff] ^ Te2[(s3 >>  8) & 0xff] ^ Te3[s0 & 0xff] ^ rk[37];
-	t2 = Te0[s2 >> 24] ^ Te1[(s3 >> 16) & 0xff] ^ Te2[(s0 >>  8) & 0xff] ^ Te3[s1 & 0xff] ^ rk[38];
-	t3 = Te0[s3 >> 24] ^ Te1[(s0 >> 16) & 0xff] ^ Te2[(s1 >>  8) & 0xff] ^ Te3[s2 & 0xff] ^ rk[39];
-
-	if( nrounds > 10 )
+	uint32 K[4 * (AES256_ROUND_COUNT + 1)];
+
+	uint32* EKey = K;
+	FMemory::Memcpy(EKey, Key, FAES::FAESKey::KeySize);
+
+	for (int Index = 0; Index < AES256_ROUND_COUNT/2; Index++, EKey += 8)
 	{
-		/* round 10: */
-		s0 = Te0[t0 >> 24] ^ Te1[(t1 >> 16) & 0xff] ^ Te2[(t2 >>  8) & 0xff] ^ Te3[t3 & 0xff] ^ rk[40];
-		s1 = Te0[t1 >> 24] ^ Te1[(t2 >> 16) & 0xff] ^ Te2[(t3 >>  8) & 0xff] ^ Te3[t0 & 0xff] ^ rk[41];
-		s2 = Te0[t2 >> 24] ^ Te1[(t3 >> 16) & 0xff] ^ Te2[(t0 >>  8) & 0xff] ^ Te3[t1 & 0xff] ^ rk[42];
-		s3 = Te0[t3 >> 24] ^ Te1[(t0 >> 16) & 0xff] ^ Te2[(t1 >>  8) & 0xff] ^ Te3[t2 & 0xff] ^ rk[43];
-		/* round 11: */
-		t0 = Te0[s0 >> 24] ^ Te1[(s1 >> 16) & 0xff] ^ Te2[(s2 >>  8) & 0xff] ^ Te3[s3 & 0xff] ^ rk[44];
-		t1 = Te0[s1 >> 24] ^ Te1[(s2 >> 16) & 0xff] ^ Te2[(s3 >>  8) & 0xff] ^ Te3[s0 & 0xff] ^ rk[45];
-		t2 = Te0[s2 >> 24] ^ Te1[(s3 >> 16) & 0xff] ^ Te2[(s0 >>  8) & 0xff] ^ Te3[s1 & 0xff] ^ rk[46];
-		t3 = Te0[s3 >> 24] ^ Te1[(s0 >> 16) & 0xff] ^ Te2[(s1 >>  8) & 0xff] ^ Te3[s2 & 0xff] ^ rk[47];
-
-		if( nrounds > 12 )
+		// rcon[] = { 0x01, 0x02, 0x04, 0x08, 0x10, 0x20, 0x40, 0x80 }
+		uint32 Rcon = 1U << Index;
+
+		EKey[ 8] = EKey[0] ^ RotateRight(AesEncryptMixARMV8(EKey[7]), 8) ^ Rcon;
+		EKey[ 9] = EKey[1] ^ EKey[ 8];
+		EKey[10] = EKey[2] ^ EKey[ 9];
+		EKey[11] = EKey[3] ^ EKey[10];
+
+		if (Index < AES256_ROUND_COUNT/2 - 1)
 		{
-			/* round 12: */
-			s0 = Te0[t0 >> 24] ^ Te1[(t1 >> 16) & 0xff] ^ Te2[(t2 >>  8) & 0xff] ^ Te3[t3 & 0xff] ^ rk[48];
-			s1 = Te0[t1 >> 24] ^ Te1[(t2 >> 16) & 0xff] ^ Te2[(t3 >>  8) & 0xff] ^ Te3[t0 & 0xff] ^ rk[49];
-			s2 = Te0[t2 >> 24] ^ Te1[(t3 >> 16) & 0xff] ^ Te2[(t0 >>  8) & 0xff] ^ Te3[t1 & 0xff] ^ rk[50];
-			s3 = Te0[t3 >> 24] ^ Te1[(t0 >> 16) & 0xff] ^ Te2[(t1 >>  8) & 0xff] ^ Te3[t2 & 0xff] ^ rk[51];
-			/* round 13: */
-			t0 = Te0[s0 >> 24] ^ Te1[(s1 >> 16) & 0xff] ^ Te2[(s2 >>  8) & 0xff] ^ Te3[s3 & 0xff] ^ rk[52];
-			t1 = Te0[s1 >> 24] ^ Te1[(s2 >> 16) & 0xff] ^ Te2[(s3 >>  8) & 0xff] ^ Te3[s0 & 0xff] ^ rk[53];
-			t2 = Te0[s2 >> 24] ^ Te1[(s3 >> 16) & 0xff] ^ Te2[(s0 >>  8) & 0xff] ^ Te3[s1 & 0xff] ^ rk[54];
-			t3 = Te0[s3 >> 24] ^ Te1[(s0 >> 16) & 0xff] ^ Te2[(s1 >>  8) & 0xff] ^ Te3[s2 & 0xff] ^ rk[55];
+			EKey[12] = EKey[4] ^ AesEncryptMixARMV8(EKey[11]);
+			EKey[13] = EKey[5] ^ EKey[12];
+			EKey[14] = EKey[6] ^ EKey[13];
+			EKey[15] = EKey[7] ^ EKey[14];
 		}
 	}
 
-	rk += ((UPTRINT)nrounds) << 2;
-
-	/*
-	 * apply last round and
-	 * map cipher state to byte array block:
-	 */
-	s0 = 
-		( Te4[( t0 >> 24 ) & 0xff] & 0xff000000 ) ^ 
-		( Te4[( t1 >> 16 ) & 0xff] & 0x00ff0000 ) ^ 
-		( Te4[( t2 >>  8 ) & 0xff] & 0x0000ff00 ) ^ 
-		( Te4[( t3 >>  0 ) & 0xff] & 0x000000ff ) ^	
-		rk[0];
-	PUTU32( ciphertext, s0 );
-
-	s1 =	 
-		( Te4[( t1 >> 24 ) & 0xff] & 0xff000000 ) ^
-		( Te4[( t2 >> 16 ) & 0xff] & 0x00ff0000 ) ^
-		( Te4[( t3 >>  8 ) & 0xff] & 0x0000ff00 ) ^
-		( Te4[( t0 >>  0 ) & 0xff] & 0x000000ff ) ^
-		rk[1];
-	PUTU32( ciphertext + 4, s1 );
-
-	s2 = 
-		( Te4[( t2 >> 24 ) & 0xff] & 0xff000000 ) ^
-		( Te4[( t3 >> 16 ) & 0xff] & 0x00ff0000 ) ^
-		( Te4[( t0 >>  8 ) & 0xff] & 0x0000ff00 ) ^
-		( Te4[( t1 >>  0 ) & 0xff] & 0x000000ff ) ^
-		rk[2];
-	PUTU32( ciphertext + 8, s2 );
-
-	s3 = 
-		( Te4[( t3 >> 24 ) & 0xff] & 0xff000000 ) ^
-		( Te4[( t0 >> 16 ) & 0xff] & 0x00ff0000 ) ^
-		( Te4[( t1 >>  8 ) & 0xff] & 0x0000ff00 ) ^
-		( Te4[( t2 >>  0 ) & 0xff] & 0x000000ff ) ^
-		rk[3];
-	PUTU32( ciphertext + 12, s3 );
+	for (int Round=0; Round<=AES256_ROUND_COUNT; Round++)
+	{
+		EncryptKey->Key[Round] = vreinterpretq_u8_u32(vld1q_u32(K + Round*4));
+	}
 }
 
-static void rijndaelDecrypt( const uint32 *rk, int32 nrounds, const uint8 ciphertext[16], uint8 plaintext[16] )
+static inline void AesDecryptExpandARMV8(FAesExpandedKeyARMV8* DecryptKey, const uint8* Key)
 {
-	uint32 s0, s1, s2, s3, t0, t1, t2, t3;
-
-	/*
-	 * map byte array block to cipher state
-	 * and add initial round key:
-	 */
-	s0 = GETU32( ciphertext ) ^ rk[0];
-	s1 = GETU32( ciphertext + 4 ) ^ rk[1];
-	s2 = GETU32( ciphertext + 8 ) ^ rk[2];
-	s3 = GETU32( ciphertext + 12 ) ^ rk[3];
-
-	/* round 1: */
-	t0 = Td0[s0 >> 24] ^ Td1[(s3 >> 16) & 0xff] ^ Td2[(s2 >>  8) & 0xff] ^ Td3[s1 & 0xff] ^ rk[ 4];
-	t1 = Td0[s1 >> 24] ^ Td1[(s0 >> 16) & 0xff] ^ Td2[(s3 >>  8) & 0xff] ^ Td3[s2 & 0xff] ^ rk[ 5];
-	t2 = Td0[s2 >> 24] ^ Td1[(s1 >> 16) & 0xff] ^ Td2[(s0 >>  8) & 0xff] ^ Td3[s3 & 0xff] ^ rk[ 6];
-	t3 = Td0[s3 >> 24] ^ Td1[(s2 >> 16) & 0xff] ^ Td2[(s1 >>  8) & 0xff] ^ Td3[s0 & 0xff] ^ rk[ 7];
-	/* round 2: */
-	s0 = Td0[t0 >> 24] ^ Td1[(t3 >> 16) & 0xff] ^ Td2[(t2 >>  8) & 0xff] ^ Td3[t1 & 0xff] ^ rk[ 8];
-	s1 = Td0[t1 >> 24] ^ Td1[(t0 >> 16) & 0xff] ^ Td2[(t3 >>  8) & 0xff] ^ Td3[t2 & 0xff] ^ rk[ 9];
-	s2 = Td0[t2 >> 24] ^ Td1[(t1 >> 16) & 0xff] ^ Td2[(t0 >>  8) & 0xff] ^ Td3[t3 & 0xff] ^ rk[10];
-	s3 = Td0[t3 >> 24] ^ Td1[(t2 >> 16) & 0xff] ^ Td2[(t1 >>  8) & 0xff] ^ Td3[t0 & 0xff] ^ rk[11];
-	/* round 3: */
-	t0 = Td0[s0 >> 24] ^ Td1[(s3 >> 16) & 0xff] ^ Td2[(s2 >>  8) & 0xff] ^ Td3[s1 & 0xff] ^ rk[12];
-	t1 = Td0[s1 >> 24] ^ Td1[(s0 >> 16) & 0xff] ^ Td2[(s3 >>  8) & 0xff] ^ Td3[s2 & 0xff] ^ rk[13];
-	t2 = Td0[s2 >> 24] ^ Td1[(s1 >> 16) & 0xff] ^ Td2[(s0 >>  8) & 0xff] ^ Td3[s3 & 0xff] ^ rk[14];
-	t3 = Td0[s3 >> 24] ^ Td1[(s2 >> 16) & 0xff] ^ Td2[(s1 >>  8) & 0xff] ^ Td3[s0 & 0xff] ^ rk[15];
-	/* round 4: */
-	s0 = Td0[t0 >> 24] ^ Td1[(t3 >> 16) & 0xff] ^ Td2[(t2 >>  8) & 0xff] ^ Td3[t1 & 0xff] ^ rk[16];
-	s1 = Td0[t1 >> 24] ^ Td1[(t0 >> 16) & 0xff] ^ Td2[(t3 >>  8) & 0xff] ^ Td3[t2 & 0xff] ^ rk[17];
-	s2 = Td0[t2 >> 24] ^ Td1[(t1 >> 16) & 0xff] ^ Td2[(t0 >>  8) & 0xff] ^ Td3[t3 & 0xff] ^ rk[18];
-	s3 = Td0[t3 >> 24] ^ Td1[(t2 >> 16) & 0xff] ^ Td2[(t1 >>  8) & 0xff] ^ Td3[t0 & 0xff] ^ rk[19];
-	/* round 5: */
-	t0 = Td0[s0 >> 24] ^ Td1[(s3 >> 16) & 0xff] ^ Td2[(s2 >>  8) & 0xff] ^ Td3[s1 & 0xff] ^ rk[20];
-	t1 = Td0[s1 >> 24] ^ Td1[(s0 >> 16) & 0xff] ^ Td2[(s3 >>  8) & 0xff] ^ Td3[s2 & 0xff] ^ rk[21];
-	t2 = Td0[s2 >> 24] ^ Td1[(s1 >> 16) & 0xff] ^ Td2[(s0 >>  8) & 0xff] ^ Td3[s3 & 0xff] ^ rk[22];
-	t3 = Td0[s3 >> 24] ^ Td1[(s2 >> 16) & 0xff] ^ Td2[(s1 >>  8) & 0xff] ^ Td3[s0 & 0xff] ^ rk[23];
-	/* round 6: */
-	s0 = Td0[t0 >> 24] ^ Td1[(t3 >> 16) & 0xff] ^ Td2[(t2 >>  8) & 0xff] ^ Td3[t1 & 0xff] ^ rk[24];
-	s1 = Td0[t1 >> 24] ^ Td1[(t0 >> 16) & 0xff] ^ Td2[(t3 >>  8) & 0xff] ^ Td3[t2 & 0xff] ^ rk[25];
-	s2 = Td0[t2 >> 24] ^ Td1[(t1 >> 16) & 0xff] ^ Td2[(t0 >>  8) & 0xff] ^ Td3[t3 & 0xff] ^ rk[26];
-	s3 = Td0[t3 >> 24] ^ Td1[(t2 >> 16) & 0xff] ^ Td2[(t1 >>  8) & 0xff] ^ Td3[t0 & 0xff] ^ rk[27];
-	/* round 7: */
-	t0 = Td0[s0 >> 24] ^ Td1[(s3 >> 16) & 0xff] ^ Td2[(s2 >>  8) & 0xff] ^ Td3[s1 & 0xff] ^ rk[28];
-	t1 = Td0[s1 >> 24] ^ Td1[(s0 >> 16) & 0xff] ^ Td2[(s3 >>  8) & 0xff] ^ Td3[s2 & 0xff] ^ rk[29];
-	t2 = Td0[s2 >> 24] ^ Td1[(s1 >> 16) & 0xff] ^ Td2[(s0 >>  8) & 0xff] ^ Td3[s3 & 0xff] ^ rk[30];
-	t3 = Td0[s3 >> 24] ^ Td1[(s2 >> 16) & 0xff] ^ Td2[(s1 >>  8) & 0xff] ^ Td3[s0 & 0xff] ^ rk[31];
-	/* round 8: */
-	s0 = Td0[t0 >> 24] ^ Td1[(t3 >> 16) & 0xff] ^ Td2[(t2 >>  8) & 0xff] ^ Td3[t1 & 0xff] ^ rk[32];
-	s1 = Td0[t1 >> 24] ^ Td1[(t0 >> 16) & 0xff] ^ Td2[(t3 >>  8) & 0xff] ^ Td3[t2 & 0xff] ^ rk[33];
-	s2 = Td0[t2 >> 24] ^ Td1[(t1 >> 16) & 0xff] ^ Td2[(t0 >>  8) & 0xff] ^ Td3[t3 & 0xff] ^ rk[34];
-	s3 = Td0[t3 >> 24] ^ Td1[(t2 >> 16) & 0xff] ^ Td2[(t1 >>  8) & 0xff] ^ Td3[t0 & 0xff] ^ rk[35];
-	/* round 9: */
-	t0 = Td0[s0 >> 24] ^ Td1[(s3 >> 16) & 0xff] ^ Td2[(s2 >>  8) & 0xff] ^ Td3[s1 & 0xff] ^ rk[36];
-	t1 = Td0[s1 >> 24] ^ Td1[(s0 >> 16) & 0xff] ^ Td2[(s3 >>  8) & 0xff] ^ Td3[s2 & 0xff] ^ rk[37];
-	t2 = Td0[s2 >> 24] ^ Td1[(s1 >> 16) & 0xff] ^ Td2[(s0 >>  8) & 0xff] ^ Td3[s3 & 0xff] ^ rk[38];
-	t3 = Td0[s3 >> 24] ^ Td1[(s2 >> 16) & 0xff] ^ Td2[(s1 >>  8) & 0xff] ^ Td3[s0 & 0xff] ^ rk[39];
-	if( nrounds > 10 )
+	FAesExpandedKeyARMV8 EncryptKey;
+	AesEncryptExpandARMV8(&EncryptKey, Key);
+
+	DecryptKey->Key[0] = EncryptKey.Key[AES256_ROUND_COUNT];
+	for (int Round=1; Round<AES256_ROUND_COUNT; Round++)
 	{
-		/* round 10: */
-		s0 = Td0[t0 >> 24] ^ Td1[(t3 >> 16) & 0xff] ^ Td2[(t2 >>  8) & 0xff] ^ Td3[t1 & 0xff] ^ rk[40];
-		s1 = Td0[t1 >> 24] ^ Td1[(t0 >> 16) & 0xff] ^ Td2[(t3 >>  8) & 0xff] ^ Td3[t2 & 0xff] ^ rk[41];
-		s2 = Td0[t2 >> 24] ^ Td1[(t1 >> 16) & 0xff] ^ Td2[(t0 >>  8) & 0xff] ^ Td3[t3 & 0xff] ^ rk[42];
-		s3 = Td0[t3 >> 24] ^ Td1[(t2 >> 16) & 0xff] ^ Td2[(t1 >>  8) & 0xff] ^ Td3[t0 & 0xff] ^ rk[43];
-		/* round 11: */
-		t0 = Td0[s0 >> 24] ^ Td1[(s3 >> 16) & 0xff] ^ Td2[(s2 >>  8) & 0xff] ^ Td3[s1 & 0xff] ^ rk[44];
-		t1 = Td0[s1 >> 24] ^ Td1[(s0 >> 16) & 0xff] ^ Td2[(s3 >>  8) & 0xff] ^ Td3[s2 & 0xff] ^ rk[45];
-		t2 = Td0[s2 >> 24] ^ Td1[(s1 >> 16) & 0xff] ^ Td2[(s0 >>  8) & 0xff] ^ Td3[s3 & 0xff] ^ rk[46];
-		t3 = Td0[s3 >> 24] ^ Td1[(s2 >> 16) & 0xff] ^ Td2[(s1 >>  8) & 0xff] ^ Td3[s0 & 0xff] ^ rk[47];
-		if( nrounds > 12 )
-		{
-			/* round 12: */
-			s0 = Td0[t0 >> 24] ^ Td1[(t3 >> 16) & 0xff] ^ Td2[(t2 >>  8) & 0xff] ^ Td3[t1 & 0xff] ^ rk[48];
-			s1 = Td0[t1 >> 24] ^ Td1[(t0 >> 16) & 0xff] ^ Td2[(t3 >>  8) & 0xff] ^ Td3[t2 & 0xff] ^ rk[49];
-			s2 = Td0[t2 >> 24] ^ Td1[(t1 >> 16) & 0xff] ^ Td2[(t0 >>  8) & 0xff] ^ Td3[t3 & 0xff] ^ rk[50];
-			s3 = Td0[t3 >> 24] ^ Td1[(t2 >> 16) & 0xff] ^ Td2[(t1 >>  8) & 0xff] ^ Td3[t0 & 0xff] ^ rk[51];
-			/* round 13: */
-			t0 = Td0[s0 >> 24] ^ Td1[(s3 >> 16) & 0xff] ^ Td2[(s2 >>  8) & 0xff] ^ Td3[s1 & 0xff] ^ rk[52];
-			t1 = Td0[s1 >> 24] ^ Td1[(s0 >> 16) & 0xff] ^ Td2[(s3 >>  8) & 0xff] ^ Td3[s2 & 0xff] ^ rk[53];
-			t2 = Td0[s2 >> 24] ^ Td1[(s1 >> 16) & 0xff] ^ Td2[(s0 >>  8) & 0xff] ^ Td3[s3 & 0xff] ^ rk[54];
-			t3 = Td0[s3 >> 24] ^ Td1[(s2 >> 16) & 0xff] ^ Td2[(s1 >>  8) & 0xff] ^ Td3[s0 & 0xff] ^ rk[55];
-		}
+		DecryptKey->Key[Round] = vaesimcq_u8(EncryptKey.Key[AES256_ROUND_COUNT-Round]);
+	}
+	DecryptKey->Key[AES256_ROUND_COUNT] = EncryptKey.Key[0];
+}
+
+static inline void AesEncryptARMV8(const uint8* Key, uint8* Contents, uint32 NumBytes)
+{
+	FAesExpandedKeyARMV8 EncryptKey;
+	AesEncryptExpandARMV8(&EncryptKey, Key);
+
+	const uint8x16_t* EKey = EncryptKey.Key;
+
+	for (uint32 Offset = 0; Offset < NumBytes; Offset += 16)
+	{
+		uint8* Data = Contents + Offset;
+
+		uint8x16_t Block = vld1q_u8(Data);
+		Block = vaesmcq_u8(vaeseq_u8(Block, EKey[0]));
+		Block = vaesmcq_u8(vaeseq_u8(Block, EKey[1]));
+		Block = vaesmcq_u8(vaeseq_u8(Block, EKey[2]));
+		Block = vaesmcq_u8(vaeseq_u8(Block, EKey[3]));
+		Block = vaesmcq_u8(vaeseq_u8(Block, EKey[4]));
+		Block = vaesmcq_u8(vaeseq_u8(Block, EKey[5]));
+		Block = vaesmcq_u8(vaeseq_u8(Block, EKey[6]));
+		Block = vaesmcq_u8(vaeseq_u8(Block, EKey[7]));
+		Block = vaesmcq_u8(vaeseq_u8(Block, EKey[8]));
+		Block = vaesmcq_u8(vaeseq_u8(Block, EKey[9]));
+		Block = vaesmcq_u8(vaeseq_u8(Block, EKey[10]));
+		Block = vaesmcq_u8(vaeseq_u8(Block, EKey[11]));
+		Block = vaesmcq_u8(vaeseq_u8(Block, EKey[12]));
+		Block = vaeseq_u8(Block, EKey[13]);
+		Block = veorq_u8(Block, EKey[14]);
+		vst1q_u8(Data, Block);
+	}
+}
+
+static inline void AesDecryptARMV8(const uint8* Key, uint8* Contents, uint32 NumBytes)
+{
+	FAesExpandedKeyARMV8 DecryptKey;
+	AesDecryptExpandARMV8(&DecryptKey, Key);
+
+	const uint8x16_t* DKey = DecryptKey.Key;
+
+	for (uint32 Offset = 0; Offset < NumBytes; Offset += 16)
+	{
+		uint8* Data = Contents + Offset;
+
+		uint8x16_t Block = vld1q_u8(Data);
+		Block = vaesimcq_u8(vaesdq_u8(Block, DKey[0]));
+		Block = vaesimcq_u8(vaesdq_u8(Block, DKey[1]));
+		Block = vaesimcq_u8(vaesdq_u8(Block, DKey[2]));
+		Block = vaesimcq_u8(vaesdq_u8(Block, DKey[3]));
+		Block = vaesimcq_u8(vaesdq_u8(Block, DKey[4]));
+		Block = vaesimcq_u8(vaesdq_u8(Block, DKey[5]));
+		Block = vaesimcq_u8(vaesdq_u8(Block, DKey[6]));
+		Block = vaesimcq_u8(vaesdq_u8(Block, DKey[7]));
+		Block = vaesimcq_u8(vaesdq_u8(Block, DKey[8]));
+		Block = vaesimcq_u8(vaesdq_u8(Block, DKey[9]));
+		Block = vaesimcq_u8(vaesdq_u8(Block, DKey[10]));
+		Block = vaesimcq_u8(vaesdq_u8(Block, DKey[11]));
+		Block = vaesimcq_u8(vaesdq_u8(Block, DKey[12]));
+		Block = vaesdq_u8(Block, DKey[13]);
+		Block = veorq_u8(Block, DKey[14]);
+		vst1q_u8(Data, Block);
 	}
-	rk += (UPTRINT)nrounds << 2;
-	/*
-	 * apply last round and
-	 * map cipher state to byte array block:
-	 */
-	s0 = 
-		( Td4[( t0 >> 24 ) & 0xff] & 0xff000000 ) ^
-		( Td4[( t3 >> 16 ) & 0xff] & 0x00ff0000 ) ^
-		( Td4[( t2 >>  8 ) & 0xff] & 0x0000ff00 ) ^
-		( Td4[( t1 >>  0 ) & 0xff] & 0x000000ff ) ^
-		rk[0];
-	PUTU32( plaintext, s0 );
-
-	s1 = 
-		( Td4[( t1 >> 24 ) & 0xff] & 0xff000000 ) ^
-		( Td4[( t0 >> 16 ) & 0xff] & 0x00ff0000 ) ^
-		( Td4[( t3 >>  8 ) & 0xff] & 0x0000ff00 ) ^
-		( Td4[( t2 >>  0 ) & 0xff] & 0x000000ff ) ^
-		rk[1];
-	PUTU32( plaintext + 4, s1 );
-
-	s2 = 
-		( Td4[( t2 >> 24 ) & 0xff] & 0xff000000 ) ^
-		( Td4[( t1 >> 16 ) & 0xff] & 0x00ff0000 ) ^
-		( Td4[( t0 >>  8 ) & 0xff] & 0x0000ff00 ) ^
-		( Td4[( t3 >>  0 ) & 0xff] & 0x000000ff ) ^
-		rk[2];
-	PUTU32( plaintext + 8, s2 );
-
-	s3 = 
-		( Td4[( t3 >> 24 ) & 0xff] & 0xff000000 ) ^
-		( Td4[( t2 >> 16 ) & 0xff] & 0x00ff0000 ) ^
-		( Td4[( t1 >>  8 ) & 0xff] & 0x0000ff00 ) ^
-		( Td4[( t0 >>  0 ) & 0xff] & 0x000000ff ) ^
-		rk[3];
-	PUTU32( plaintext + 12, s3 );
 }
 
+#endif
+
+
 void FAES::EncryptData(uint8 *Contents, uint32 NumBytes, const FAESKey& Key)
 {
 	checkf(Key.IsValid(), TEXT("No valid encryption key specified"));
@@ -1145,34 +978,24 @@ void FAES::EncryptData(uint8 *Contents, uint32 NumBytes, const ANSICHAR* Key)
 
 void FAES::EncryptData(uint8* Contents, uint32 NumBytes, const uint8* KeyBytes, uint32 NumKeyBytes)
 {
-	uint32 rk[RKLENGTH(AES_KEYBITS)] = { 0 };
-	int32 nrounds;
-
-	checkf((NumBytes & (AESBlockSize - 1)) == 0, TEXT("NumBytes needs to be a multiple of 16 bytes"));
-	checkf(NumKeyBytes >= KEYLENGTH(AES_KEYBITS), TEXT("AES key needs to be at least %d characters"), KEYLENGTH(AES_KEYBITS));
+	checkf(NumBytes % AESBlockSize == 0, TEXT("NumBytes needs to be a multiple of 16 bytes"));
+	checkf(NumKeyBytes == FAESKey::KeySize, TEXT("AES-256 key needs to be 32 bytes long"));
 
-#if TEST_ENCRYPTION
-	TArray<uint8> OriginalBlob;
-	OriginalBlob.AddUninitialized( NumBytes );
-	FMemory::Memcpy( OriginalBlob.GetData(), Contents, NumBytes );
+#if UE_PLATFORM_AES_ARMV8
+	AesEncryptARMV8(KeyBytes, Contents, NumBytes);
+	return;
 #endif
 
-	// Set up the rk buffer
-	nrounds = rijndaelSetupEncrypt(rk, KeyBytes, AES_KEYBITS);
-
-	// Encrypt the data a block at a time
-	for( uint32 Offset = 0; Offset < NumBytes; Offset += AESBlockSize )
+#if UE_PLATFORM_AES_X86
+	if (CanUseAesInstructions())
 	{
-		rijndaelEncrypt( rk, nrounds, Contents + Offset, Contents + Offset );
+		AesEncryptX86(KeyBytes, Contents, NumBytes);
+		return;
 	}
+#endif
 
-#if TEST_ENCRYPTION
-	TArray<uint8> DecryptedBlob;
-	DecryptedBlob.AddUninitialized( NumBytes );
-	FMemory::Memcpy( DecryptedBlob.GetData(), Contents, NumBytes );
-
-	DecryptData( DecryptedBlob.GetData(), DecryptedBlob.Num() );
-	checkf( DecryptedBlob == OriginalBlob, TEXT( "Decrypted data different from original data - this is bad" ) );
+#if UE_PLATFORM_AES_FALLBACK
+	AesEncrypt(KeyBytes, Contents, NumBytes);
 #endif
 }
 
@@ -1190,18 +1013,23 @@ void FAES::DecryptData(uint8 *Contents, uint32 NumBytes, const ANSICHAR* Key)
 
 void FAES::DecryptData(uint8* Contents, uint32 NumBytes, const uint8* KeyBytes, uint32 NumKeyBytes)
 {
-	uint32 rk[RKLENGTH(AES_KEYBITS)] = { 0 };
-	int32 nrounds;
-
-	checkf((NumBytes & (AESBlockSize - 1)) == 0, TEXT("NumBytes needs to tbe a multiple of 16 bytes"));
-	checkf(NumKeyBytes >= KEYLENGTH(AES_KEYBITS), TEXT("AES key needs to be at least %d characters"), KEYLENGTH(AES_KEYBITS));
+	checkf(NumBytes % AESBlockSize == 0, TEXT("NumBytes needs to be a multiple of 16 bytes"));
+	checkf(NumKeyBytes == FAESKey::KeySize, TEXT("AES-256 key needs to be 32 bytes long"));
 
-	// Set up the rk buffer
-	nrounds = rijndaelSetupDecrypt(rk, KeyBytes, AES_KEYBITS);
+#if UE_PLATFORM_AES_ARMV8
+	AesDecryptARMV8(KeyBytes, Contents, NumBytes);
+	return;
+#endif
 
-	// Decrypt the data a block at a time
-	for( uint32 Offset = 0; Offset < NumBytes; Offset += AESBlockSize )
+#if UE_PLATFORM_AES_X86
+	if (CanUseAesInstructions())
 	{
-		rijndaelDecrypt( rk, nrounds, Contents + Offset, Contents + Offset );
+		AesDecryptX86(KeyBytes, Contents, NumBytes);
+		return;
 	}
+#endif
+
+#if UE_PLATFORM_AES_FALLBACK
+	AesDecrypt(KeyBytes, Contents, NumBytes);
+#endif
 }
diff --git a/Engine/Source/Runtime/Core/Private/Misc/SecureHash.cpp b/Engine/Source/Runtime/Core/Private/Misc/SecureHash.cpp
--- a/Engine/Source/Runtime/Core/Private/Misc/SecureHash.cpp
+++ b/Engine/Source/Runtime/Core/Private/Misc/SecureHash.cpp
@@ -5,6 +5,42 @@
 #include "HAL/FileManager.h"
 #include "Misc/Paths.h"
 
+#if defined(_M_AMD64) || defined(__x86_64__)
+
+#include <immintrin.h>
+#if PLATFORM_COMPILER_CLANG
+#include <shaintrin.h>
+#include <cpuid.h>
+#endif
+
+#define UE_PLATFORM_SHA_X86 1
+
+// if x86 SHA is enabled with compiler, use SHA instructions without fallback
+#if defined(__SHA__) || PLATFORM_ALWAYS_HAS_SHA
+#define UE_PLATFORM_SHA_FALLBACK 0
+#endif
+
+#elif defined(__aarch64__)
+
+// on ARMv8 enable SHA instructions unconditionally only when they are enabled with compiler
+#ifdef __ARM_FEATURE_CRYPTO
+#define UE_PLATFORM_SHA_ARMV8 1
+#define UE_PLATFORM_SHA_FALLBACK 0
+#endif
+
+#endif
+
+
+#ifndef UE_PLATFORM_SHA_FALLBACK
+#define UE_PLATFORM_SHA_FALLBACK 1
+#endif
+#ifndef UE_PLATFORM_SHA_X86
+#define UE_PLATFORM_SHA_X86 0
+#endif
+#ifndef UE_PLATFORM_SHA_ARMV8
+#define UE_PLATFORM_SHA_ARMV8 0
+#endif
+
 
 DEFINE_LOG_CATEGORY_STATIC(LogSecureHash, Log, All);
 
@@ -34,12 +70,7 @@ enum {_S42=10};
 enum {_S43=15};
 enum {_S44=21};
 
-static uint8 PADDING[64] = {
-	0x80, 0, 0, 0, 0, 0, 0, 0, 0, 0, 0, 0, 0, 0, 0, 0, 0, 0, 0,
-	0, 0, 0, 0, 0, 0, 0, 0, 0, 0, 0, 0, 0, 0, 0, 0, 0, 0, 0, 0,
-	0, 0, 0, 0, 0, 0, 0, 0, 0, 0, 0, 0, 0, 0, 0, 0, 0, 0, 0, 0,
-	0, 0, 0, 0, 0
-};
+static const uint8 PADDING[64] = { 0x80 };
 
 //
 // Basic MD5 transformations.
@@ -585,13 +616,12 @@ TMap<FString, uint8*> FSHA1::ScriptSHAHashMap;
 #endif
 
 #if PLATFORM_LITTLE_ENDIAN
-	#define SHABLK0(i) (m_block->l[i] = (ROL32(m_block->l[i],24) & 0xFF00FF00) | (ROL32(m_block->l[i],8) & 0x00FF00FF))
+#define SHABLK0(i) (block[i] = (ROL32(block[i],24) & 0xFF00FF00) | (ROL32(block[i],8) & 0x00FF00FF))
 #else
-	#define SHABLK0(i) (m_block->l[i])
+#define SHABLK0(i) (block[i])
 #endif
 
-#define SHABLK(i) (m_block->l[i&15] = ROL32(m_block->l[(i+13)&15] ^ m_block->l[(i+8)&15] \
-	^ m_block->l[(i+2)&15] ^ m_block->l[i&15],1))
+#define SHABLK(i) (block[i&15] = ROL32(block[(i+13)&15] ^ block[(i+8)&15] ^ block[(i+2)&15] ^ block[i&15],1))
 
 // SHA-1 rounds
 #define _R0(v,w,x,y,z,i) { z+=((w&(x^y))^y)+SHABLK0(i)+0x5A827999+ROL32(v,5); w=ROL32(w,30); }
@@ -600,6 +630,8 @@ TMap<FString, uint8*> FSHA1::ScriptSHAHashMap;
 #define _R3(v,w,x,y,z,i) { z+=(((w|x)&y)|(w&x))+SHABLK(i)+0x8F1BBCDC+ROL32(v,5); w=ROL32(w,30); }
 #define _R4(v,w,x,y,z,i) { z+=(w^x^y)+SHABLK(i)+0xCA62C1D6+ROL32(v,5); w=ROL32(w,30); }
 
+#define SHA1_BLOCK_SIZE 64
+
 FArchive& operator<<( FArchive& Ar, FSHAHash& G )
 {
 	Ar.Serialize(&G.Hash, sizeof(G.Hash));
@@ -623,8 +655,6 @@ void Freeze::IntrinsicToString(const FSHAHash& Object, const FTypeLayoutDesc& Ty
 
 FSHA1::FSHA1()
 {
-	m_block = (SHA1_WORKSPACE_BLOCK *)m_workspace;
-
 	Reset();
 }
 
@@ -642,45 +672,495 @@ void FSHA1::Reset()
 	m_state[3] = 0x10325476;
 	m_state[4] = 0xC3D2E1F0;
 
-	m_count[0] = 0;
-	m_count[1] = 0;
-}
-
-void FSHA1::Transform(uint32 *state, const uint8 *buffer)
-{
-	// Copy state[] to working vars
-	uint32 a = state[0], b = state[1], c = state[2], d = state[3], e = state[4];
-
-	FMemory::Memcpy(m_block, buffer, 64);
-
-	// 4 rounds of 20 operations each. Loop unrolled.
-	_R0(a,b,c,d,e, 0); _R0(e,a,b,c,d, 1); _R0(d,e,a,b,c, 2); _R0(c,d,e,a,b, 3);
-	_R0(b,c,d,e,a, 4); _R0(a,b,c,d,e, 5); _R0(e,a,b,c,d, 6); _R0(d,e,a,b,c, 7);
-	_R0(c,d,e,a,b, 8); _R0(b,c,d,e,a, 9); _R0(a,b,c,d,e,10); _R0(e,a,b,c,d,11);
-	_R0(d,e,a,b,c,12); _R0(c,d,e,a,b,13); _R0(b,c,d,e,a,14); _R0(a,b,c,d,e,15);
-	_R1(e,a,b,c,d,16); _R1(d,e,a,b,c,17); _R1(c,d,e,a,b,18); _R1(b,c,d,e,a,19);
-	_R2(a,b,c,d,e,20); _R2(e,a,b,c,d,21); _R2(d,e,a,b,c,22); _R2(c,d,e,a,b,23);
-	_R2(b,c,d,e,a,24); _R2(a,b,c,d,e,25); _R2(e,a,b,c,d,26); _R2(d,e,a,b,c,27);
-	_R2(c,d,e,a,b,28); _R2(b,c,d,e,a,29); _R2(a,b,c,d,e,30); _R2(e,a,b,c,d,31);
-	_R2(d,e,a,b,c,32); _R2(c,d,e,a,b,33); _R2(b,c,d,e,a,34); _R2(a,b,c,d,e,35);
-	_R2(e,a,b,c,d,36); _R2(d,e,a,b,c,37); _R2(c,d,e,a,b,38); _R2(b,c,d,e,a,39);
-	_R3(a,b,c,d,e,40); _R3(e,a,b,c,d,41); _R3(d,e,a,b,c,42); _R3(c,d,e,a,b,43);
-	_R3(b,c,d,e,a,44); _R3(a,b,c,d,e,45); _R3(e,a,b,c,d,46); _R3(d,e,a,b,c,47);
-	_R3(c,d,e,a,b,48); _R3(b,c,d,e,a,49); _R3(a,b,c,d,e,50); _R3(e,a,b,c,d,51);
-	_R3(d,e,a,b,c,52); _R3(c,d,e,a,b,53); _R3(b,c,d,e,a,54); _R3(a,b,c,d,e,55);
-	_R3(e,a,b,c,d,56); _R3(d,e,a,b,c,57); _R3(c,d,e,a,b,58); _R3(b,c,d,e,a,59);
-	_R4(a,b,c,d,e,60); _R4(e,a,b,c,d,61); _R4(d,e,a,b,c,62); _R4(c,d,e,a,b,63);
-	_R4(b,c,d,e,a,64); _R4(a,b,c,d,e,65); _R4(e,a,b,c,d,66); _R4(d,e,a,b,c,67);
-	_R4(c,d,e,a,b,68); _R4(b,c,d,e,a,69); _R4(a,b,c,d,e,70); _R4(e,a,b,c,d,71);
-	_R4(d,e,a,b,c,72); _R4(c,d,e,a,b,73); _R4(b,c,d,e,a,74); _R4(a,b,c,d,e,75);
-	_R4(e,a,b,c,d,76); _R4(d,e,a,b,c,77); _R4(c,d,e,a,b,78); _R4(b,c,d,e,a,79);
-
-	// Add the working vars back into state
-	state[0] += a;
-	state[1] += b;
-	state[2] += c;
-	state[3] += d;
-	state[4] += e;
+	// input byte count
+	m_count = 0;
+}
+
+#if UE_PLATFORM_SHA_X86
+
+static inline bool DetectShaInstructions()
+{
+	int info1[4];
+	int info7[4];
+#if PLATFORM_COMPILER_CLANG
+	__cpuid(1, info1[0], info1[1], info1[2], info1[3]);
+	__cpuid_count(7, 0, info7[0], info7[1], info7[2], info7[3]);
+#else
+	__cpuid(info1, 1);
+	__cpuidex(info7, 7, 0);
+#endif
+	bool bHasSSSE3 = (info1[2] & 0x00000200) != 0;
+	bool bHasSHA = (info7[1] & 0x20000000) != 0;
+	return bHasSSSE3 && bHasSHA;
+}
+
+static bool CanUseShaInstructions()
+{
+#if UE_PLATFORM_SHA_FALLBACK
+	// run cpuid only once
+	static bool bCanUse = DetectShaInstructions();
+	return bCanUse;
+#else
+	// no SHA fallback code means SHA instructions will be used unconditionally
+	return true;
+#endif
+}
+
+#if PLATFORM_COMPILER_CLANG
+#pragma clang attribute push (__attribute__((target("sha"))), apply_to=function)
+#endif
+
+static void Sha1TransformX86(uint32* State, const void* Buffer, uint64 Size)
+{
+	// https://www.intel.com/content/www/us/en/developer/articles/technical/intel-sha-extensions.html
+	// https://www.intel.com/content/www/us/en/docs/intrinsics-guide/index.html#othertechs=SHA
+
+	// to reverse byte order
+	__m128i REVERSE_BYTES = _mm_setr_epi8(
+		0x0f, 0x0e, 0x0d, 0x0c, 0x0b, 0x0a, 0x09, 0x08,
+		0x07, 0x06, 0x05, 0x04, 0x03, 0x02, 0x01, 0x00);
+
+	// load current state
+	__m128i ABCD = _mm_loadu_si128((const __m128i*)State);
+	__m128i E0 = _mm_setr_epi32(0, 0, 0, State[4]);
+	ABCD = _mm_shuffle_epi32(ABCD, _MM_SHUFFLE(0, 1, 2, 3));
+
+	const __m128i* Data = (const __m128i*)Buffer;
+	while (Size >= SHA1_BLOCK_SIZE)
+	{
+		__m128i E1;
+
+		// remember current state
+		__m128i ABCD_SAVE = ABCD;
+		__m128i E_SAVE = E0;
+
+		// load message from buffer as 32-bit big-endian integers
+		__m128i M0 = _mm_loadu_si128(Data + 0);
+		__m128i M1 = _mm_loadu_si128(Data + 1);
+		__m128i M2 = _mm_loadu_si128(Data + 2);
+		__m128i M3 = _mm_loadu_si128(Data + 3);
+		M0 = _mm_shuffle_epi8(M0, REVERSE_BYTES);
+		M1 = _mm_shuffle_epi8(M1, REVERSE_BYTES);
+		M2 = _mm_shuffle_epi8(M2, REVERSE_BYTES);
+		M3 = _mm_shuffle_epi8(M3, REVERSE_BYTES);
+
+		// 0..3
+		E0 = _mm_add_epi32(E0, M0);
+		E1 = ABCD;
+		ABCD = _mm_sha1rnds4_epu32(ABCD, E0, 0);
+
+		// 4..7
+		E1 = _mm_sha1nexte_epu32(E1, M1);
+		E0 = ABCD;
+		ABCD = _mm_sha1rnds4_epu32(ABCD, E1, 0);
+		M0 = _mm_sha1msg1_epu32(M0, M1);
+
+		// 8..11
+		E0 = _mm_sha1nexte_epu32(E0, M2);
+		E1 = ABCD;
+		ABCD = _mm_sha1rnds4_epu32(ABCD, E0, 0);
+		M1 = _mm_sha1msg1_epu32(M1, M2);
+		M0 = _mm_xor_si128(M0, M2);
+
+		// 12..15
+		E1 = _mm_sha1nexte_epu32(E1, M3);
+		E0 = ABCD;
+		M0 = _mm_sha1msg2_epu32(M0, M3);
+		ABCD = _mm_sha1rnds4_epu32(ABCD, E1, 0);
+		M2 = _mm_sha1msg1_epu32(M2, M3);
+		M1 = _mm_xor_si128(M1, M3);
+
+		// 16..19
+		E0 = _mm_sha1nexte_epu32(E0, M0);
+		E1 = ABCD;
+		M1 = _mm_sha1msg2_epu32(M1, M0);
+		ABCD = _mm_sha1rnds4_epu32(ABCD, E0, 0);
+		M3 = _mm_sha1msg1_epu32(M3, M0);
+		M2 = _mm_xor_si128(M2, M0);
+
+		// 20..23
+		E1 = _mm_sha1nexte_epu32(E1, M1);
+		E0 = ABCD;
+		M2 = _mm_sha1msg2_epu32(M2, M1);
+		ABCD = _mm_sha1rnds4_epu32(ABCD, E1, 1);
+		M0 = _mm_sha1msg1_epu32(M0, M1);
+		M3 = _mm_xor_si128(M3, M1);
+
+		// 24..27
+		E0 = _mm_sha1nexte_epu32(E0, M2);
+		E1 = ABCD;
+		M3 = _mm_sha1msg2_epu32(M3, M2);
+		ABCD = _mm_sha1rnds4_epu32(ABCD, E0, 1);
+		M1 = _mm_sha1msg1_epu32(M1, M2);
+		M0 = _mm_xor_si128(M0, M2);
+
+		// 28..31
+		E1 = _mm_sha1nexte_epu32(E1, M3);
+		E0 = ABCD;
+		M0 = _mm_sha1msg2_epu32(M0, M3);
+		ABCD = _mm_sha1rnds4_epu32(ABCD, E1, 1);
+		M2 = _mm_sha1msg1_epu32(M2, M3);
+		M1 = _mm_xor_si128(M1, M3);
+
+		// 32..35
+		E0 = _mm_sha1nexte_epu32(E0, M0);
+		E1 = ABCD;
+		M1 = _mm_sha1msg2_epu32(M1, M0);
+		ABCD = _mm_sha1rnds4_epu32(ABCD, E0, 1);
+		M3 = _mm_sha1msg1_epu32(M3, M0);
+		M2 = _mm_xor_si128(M2, M0);
+
+		// 36..39
+		E1 = _mm_sha1nexte_epu32(E1, M1);
+		E0 = ABCD;
+		M2 = _mm_sha1msg2_epu32(M2, M1);
+		ABCD = _mm_sha1rnds4_epu32(ABCD, E1, 1);
+		M0 = _mm_sha1msg1_epu32(M0, M1);
+		M3 = _mm_xor_si128(M3, M1);
+
+		// 40..43
+		E0 = _mm_sha1nexte_epu32(E0, M2);
+		E1 = ABCD;
+		M3 = _mm_sha1msg2_epu32(M3, M2);
+		ABCD = _mm_sha1rnds4_epu32(ABCD, E0, 2);
+		M1 = _mm_sha1msg1_epu32(M1, M2);
+		M0 = _mm_xor_si128(M0, M2);
+
+		// 44..47
+		E1 = _mm_sha1nexte_epu32(E1, M3);
+		E0 = ABCD;
+		M0 = _mm_sha1msg2_epu32(M0, M3);
+		ABCD = _mm_sha1rnds4_epu32(ABCD, E1, 2);
+		M2 = _mm_sha1msg1_epu32(M2, M3);
+		M1 = _mm_xor_si128(M1, M3);
+
+		// 48..51
+		E0 = _mm_sha1nexte_epu32(E0, M0);
+		E1 = ABCD;
+		M1 = _mm_sha1msg2_epu32(M1, M0);
+		ABCD = _mm_sha1rnds4_epu32(ABCD, E0, 2);
+		M3 = _mm_sha1msg1_epu32(M3, M0);
+		M2 = _mm_xor_si128(M2, M0);
+
+		// 52..55
+		E1 = _mm_sha1nexte_epu32(E1, M1);
+		E0 = ABCD;
+		M2 = _mm_sha1msg2_epu32(M2, M1);
+		ABCD = _mm_sha1rnds4_epu32(ABCD, E1, 2);
+		M0 = _mm_sha1msg1_epu32(M0, M1);
+		M3 = _mm_xor_si128(M3, M1);
+
+		// 56..59
+		E0 = _mm_sha1nexte_epu32(E0, M2);
+		E1 = ABCD;
+		M3 = _mm_sha1msg2_epu32(M3, M2);
+		ABCD = _mm_sha1rnds4_epu32(ABCD, E0, 2);
+		M1 = _mm_sha1msg1_epu32(M1, M2);
+		M0 = _mm_xor_si128(M0, M2);
+
+		// 60..63
+		E1 = _mm_sha1nexte_epu32(E1, M3);
+		E0 = ABCD;
+		M0 = _mm_sha1msg2_epu32(M0, M3);
+		ABCD = _mm_sha1rnds4_epu32(ABCD, E1, 3);
+		M2 = _mm_sha1msg1_epu32(M2, M3);
+		M1 = _mm_xor_si128(M1, M3);
+
+		// 64..67
+		E0 = _mm_sha1nexte_epu32(E0, M0);
+		E1 = ABCD;
+		M1 = _mm_sha1msg2_epu32(M1, M0);
+		ABCD = _mm_sha1rnds4_epu32(ABCD, E0, 3);
+		M3 = _mm_sha1msg1_epu32(M3, M0);
+		M2 = _mm_xor_si128(M2, M0);
+
+		// 68..71
+		E1 = _mm_sha1nexte_epu32(E1, M1);
+		E0 = ABCD;
+		M2 = _mm_sha1msg2_epu32(M2, M1);
+		ABCD = _mm_sha1rnds4_epu32(ABCD, E1, 3);
+		M3 = _mm_xor_si128(M3, M1);
+
+		// 72..75
+		E0 = _mm_sha1nexte_epu32(E0, M2);
+		E1 = ABCD;
+		M3 = _mm_sha1msg2_epu32(M3, M2);
+		ABCD = _mm_sha1rnds4_epu32(ABCD, E0, 3);
+
+		// 76..79
+		E1 = _mm_sha1nexte_epu32(E1, M3);
+		E0 = ABCD;
+		ABCD = _mm_sha1rnds4_epu32(ABCD, E1, 3);
+
+		// add the working vars back into state
+		E0 = _mm_sha1nexte_epu32(E0, E_SAVE);
+		ABCD = _mm_add_epi32(ABCD, ABCD_SAVE);
+
+		Data += 4;
+		Size -= SHA1_BLOCK_SIZE;
+	}
+
+	// save state
+	ABCD = _mm_shuffle_epi32(ABCD, _MM_SHUFFLE(0, 1, 2, 3));
+	_mm_storeu_si128((__m128i*)State, ABCD);
+	State[4] = _mm_extract_epi32(E0, 3);
+}
+
+#if PLATFORM_COMPILER_CLANG
+#pragma clang attribute pop
+#endif
+
+#endif
+
+#if UE_PLATFORM_SHA_ARMV8
+
+static void Sha1TransformARMV8(uint32* State, const uint8_t* Buffer, uint64 Size)
+{
+	// https://developer.arm.com/architectures/instruction-sets/intrinsics/#f:@navigationhierarchiesinstructiongroup=[Cryptography,SHA1]
+
+	// load current state
+	uint32x4_t ABCD = vld1q_u32(State);
+	uint32 E0 = State[4];
+
+	while (Size >= SHA1_BLOCK_SIZE)
+	{
+		// remember current state
+		uint32x4_t ABCD_SAVE = ABCD;
+		uint32_t E0_SAVE = E0;
+
+		// load message from buffer as 32-bit big-endian integers
+		uint32x4_t M0 = vreinterpretq_u32_u8(vrev32q_u8(vld1q_u8(Buffer + 0 * 16)));
+		uint32x4_t M1 = vreinterpretq_u32_u8(vrev32q_u8(vld1q_u8(Buffer + 1 * 16)));
+		uint32x4_t M2 = vreinterpretq_u32_u8(vrev32q_u8(vld1q_u8(Buffer + 2 * 16)));
+		uint32x4_t M3 = vreinterpretq_u32_u8(vrev32q_u8(vld1q_u8(Buffer + 3 * 16)));
+
+		uint32x4_t T0 = vaddq_u32(M0, vdupq_n_u32(0x5A827999));
+		uint32x4_t T1 = vaddq_u32(M1, vdupq_n_u32(0x5A827999));
+		uint32 E1;
+
+		// 0..3
+		E1 = vsha1h_u32(vgetq_lane_u32(ABCD, 0));
+		ABCD = vsha1cq_u32(ABCD, E0, T0);
+		T0 = vaddq_u32(M2, vdupq_n_u32(0x5A827999));
+		M0 = vsha1su0q_u32(M0, M1, M2);
+
+		// 4..7
+		E0 = vsha1h_u32(vgetq_lane_u32(ABCD, 0));
+		ABCD = vsha1cq_u32(ABCD, E1, T1);
+		T1 = vaddq_u32(M3, vdupq_n_u32(0x5A827999));
+		M0 = vsha1su1q_u32(M0, M3);
+		M1 = vsha1su0q_u32(M1, M2, M3);
+
+		// 8..11
+		E1 = vsha1h_u32(vgetq_lane_u32(ABCD, 0));
+		ABCD = vsha1cq_u32(ABCD, E0, T0);
+		T0 = vaddq_u32(M0, vdupq_n_u32(0x5A827999));
+		M1 = vsha1su1q_u32(M1, M0);
+		M2 = vsha1su0q_u32(M2, M3, M0);
+
+		// 12..15
+		E0 = vsha1h_u32(vgetq_lane_u32(ABCD, 0));
+		ABCD = vsha1cq_u32(ABCD, E1, T1);
+		T1 = vaddq_u32(M1, vdupq_n_u32(0x6ED9EBA1));
+		M2 = vsha1su1q_u32(M2, M1);
+		M3 = vsha1su0q_u32(M3, M0, M1);
+
+		// 16..19
+		E1 = vsha1h_u32(vgetq_lane_u32(ABCD, 0));
+		ABCD = vsha1cq_u32(ABCD, E0, T0);
+		T0 = vaddq_u32(M2, vdupq_n_u32(0x6ED9EBA1));
+		M3 = vsha1su1q_u32(M3, M2);
+		M0 = vsha1su0q_u32(M0, M1, M2);
+
+		// 20..23
+		E0 = vsha1h_u32(vgetq_lane_u32(ABCD, 0));
+		ABCD = vsha1pq_u32(ABCD, E1, T1);
+		T1 = vaddq_u32(M3, vdupq_n_u32(0x6ED9EBA1));
+		M0 = vsha1su1q_u32(M0, M3);
+		M1 = vsha1su0q_u32(M1, M2, M3);
+
+		// 24..27
+		E1 = vsha1h_u32(vgetq_lane_u32(ABCD, 0));
+		ABCD = vsha1pq_u32(ABCD, E0, T0);
+		T0 = vaddq_u32(M0, vdupq_n_u32(0x6ED9EBA1));
+		M1 = vsha1su1q_u32(M1, M0);
+		M2 = vsha1su0q_u32(M2, M3, M0);
+
+		// 28..31
+		E0 = vsha1h_u32(vgetq_lane_u32(ABCD, 0));
+		ABCD = vsha1pq_u32(ABCD, E1, T1);
+		T1 = vaddq_u32(M1, vdupq_n_u32(0x6ED9EBA1));
+		M2 = vsha1su1q_u32(M2, M1);
+		M3 = vsha1su0q_u32(M3, M0, M1);
+
+		// 32..35
+		E1 = vsha1h_u32(vgetq_lane_u32(ABCD, 0));
+		ABCD = vsha1pq_u32(ABCD, E0, T0);
+		T0 = vaddq_u32(M2, vdupq_n_u32(0x8F1BBCDC));
+		M3 = vsha1su1q_u32(M3, M2);
+		M0 = vsha1su0q_u32(M0, M1, M2);
+
+		// 36..39
+		E0 = vsha1h_u32(vgetq_lane_u32(ABCD, 0));
+		ABCD = vsha1pq_u32(ABCD, E1, T1);
+		T1 = vaddq_u32(M3, vdupq_n_u32(0x8F1BBCDC));
+		M0 = vsha1su1q_u32(M0, M3);
+		M1 = vsha1su0q_u32(M1, M2, M3);
+
+		// 40..43
+		E1 = vsha1h_u32(vgetq_lane_u32(ABCD, 0));
+		ABCD = vsha1mq_u32(ABCD, E0, T0);
+		T0 = vaddq_u32(M0, vdupq_n_u32(0x8F1BBCDC));
+		M1 = vsha1su1q_u32(M1, M0);
+		M2 = vsha1su0q_u32(M2, M3, M0);
+
+		// 44..47
+		E0 = vsha1h_u32(vgetq_lane_u32(ABCD, 0));
+		ABCD = vsha1mq_u32(ABCD, E1, T1);
+		T1 = vaddq_u32(M1, vdupq_n_u32(0x8F1BBCDC));
+		M2 = vsha1su1q_u32(M2, M1);
+		M3 = vsha1su0q_u32(M3, M0, M1);
+
+		// 48..51
+		E1 = vsha1h_u32(vgetq_lane_u32(ABCD, 0));
+		ABCD = vsha1mq_u32(ABCD, E0, T0);
+		T0 = vaddq_u32(M2, vdupq_n_u32(0x8F1BBCDC));
+		M3 = vsha1su1q_u32(M3, M2);
+		M0 = vsha1su0q_u32(M0, M1, M2);
+
+		// 52..55
+		E0 = vsha1h_u32(vgetq_lane_u32(ABCD, 0));
+		ABCD = vsha1mq_u32(ABCD, E1, T1);
+		T1 = vaddq_u32(M3, vdupq_n_u32(0xCA62C1D6));
+		M0 = vsha1su1q_u32(M0, M3);
+		M1 = vsha1su0q_u32(M1, M2, M3);
+
+		// 56..59
+		E1 = vsha1h_u32(vgetq_lane_u32(ABCD, 0));
+		ABCD = vsha1mq_u32(ABCD, E0, T0);
+		T0 = vaddq_u32(M0, vdupq_n_u32(0xCA62C1D6));
+		M1 = vsha1su1q_u32(M1, M0);
+		M2 = vsha1su0q_u32(M2, M3, M0);
+
+		// 60..63
+		E0 = vsha1h_u32(vgetq_lane_u32(ABCD, 0));
+		ABCD = vsha1pq_u32(ABCD, E1, T1);
+		T1 = vaddq_u32(M1, vdupq_n_u32(0xCA62C1D6));
+		M2 = vsha1su1q_u32(M2, M1);
+		M3 = vsha1su0q_u32(M3, M0, M1);
+
+		// 64..67
+		E1 = vsha1h_u32(vgetq_lane_u32(ABCD, 0));
+		ABCD = vsha1pq_u32(ABCD, E0, T0);
+		T0 = vaddq_u32(M2, vdupq_n_u32(0xCA62C1D6));
+		M3 = vsha1su1q_u32(M3, M2);
+		M0 = vsha1su0q_u32(M0, M1, M2);
+
+		// 68..71
+		E0 = vsha1h_u32(vgetq_lane_u32(ABCD, 0));
+		ABCD = vsha1pq_u32(ABCD, E1, T1);
+		T1 = vaddq_u32(M3, vdupq_n_u32(0xCA62C1D6));
+		M0 = vsha1su1q_u32(M0, M3);
+
+		// 72..75
+		E1 = vsha1h_u32(vgetq_lane_u32(ABCD, 0));
+		ABCD = vsha1pq_u32(ABCD, E0, T0);
+
+		// 76..79
+		E0 = vsha1h_u32(vgetq_lane_u32(ABCD, 0));
+		ABCD = vsha1pq_u32(ABCD, E1, T1);
+
+		// add the working vars back into state
+		E0 += E0_SAVE;
+		ABCD = vaddq_u32(ABCD_SAVE, ABCD);
+
+		Buffer += SHA1_BLOCK_SIZE;
+		Size -= SHA1_BLOCK_SIZE;
+	}
+
+	// save state
+	vst1q_u32(State, ABCD);
+	State[4] = E0;
+}
+
+#endif
+
+void FSHA1::Transform(const uint8* buffer, uint64 len)
+{
+#if UE_PLATFORM_SHA_X86
+	if (CanUseShaInstructions())
+	{
+		Sha1TransformX86(m_state, buffer, len);
+		return;
+	}
+#endif
+
+#if UE_PLATFORM_SHA_ARMV8
+	Sha1TransformARMV8(m_state, buffer, len);
+	return;
+#endif
+
+#if UE_PLATFORM_SHA_FALLBACK
+
+	// load current state
+	uint32 a = m_state[0];
+	uint32 b = m_state[1];
+	uint32 c = m_state[2];
+	uint32 d = m_state[3];
+	uint32 e = m_state[4];
+
+	while (len >= SHA1_BLOCK_SIZE)
+	{
+		uint32 block[16];
+		FMemory::Memcpy(block, buffer, SHA1_BLOCK_SIZE);
+
+		// remember current state
+		uint32 a0 = a;
+		uint32 b0 = b;
+		uint32 c0 = c;
+		uint32 d0 = d;
+		uint32 e0 = e;
+
+		// 4 rounds of 20 operations each. Loop unrolled.
+		_R0(a,b,c,d,e, 0); _R0(e,a,b,c,d, 1); _R0(d,e,a,b,c, 2); _R0(c,d,e,a,b, 3);
+		_R0(b,c,d,e,a, 4); _R0(a,b,c,d,e, 5); _R0(e,a,b,c,d, 6); _R0(d,e,a,b,c, 7);
+		_R0(c,d,e,a,b, 8); _R0(b,c,d,e,a, 9); _R0(a,b,c,d,e,10); _R0(e,a,b,c,d,11);
+		_R0(d,e,a,b,c,12); _R0(c,d,e,a,b,13); _R0(b,c,d,e,a,14); _R0(a,b,c,d,e,15);
+		_R1(e,a,b,c,d,16); _R1(d,e,a,b,c,17); _R1(c,d,e,a,b,18); _R1(b,c,d,e,a,19);
+		_R2(a,b,c,d,e,20); _R2(e,a,b,c,d,21); _R2(d,e,a,b,c,22); _R2(c,d,e,a,b,23);
+		_R2(b,c,d,e,a,24); _R2(a,b,c,d,e,25); _R2(e,a,b,c,d,26); _R2(d,e,a,b,c,27);
+		_R2(c,d,e,a,b,28); _R2(b,c,d,e,a,29); _R2(a,b,c,d,e,30); _R2(e,a,b,c,d,31);
+		_R2(d,e,a,b,c,32); _R2(c,d,e,a,b,33); _R2(b,c,d,e,a,34); _R2(a,b,c,d,e,35);
+		_R2(e,a,b,c,d,36); _R2(d,e,a,b,c,37); _R2(c,d,e,a,b,38); _R2(b,c,d,e,a,39);
+		_R3(a,b,c,d,e,40); _R3(e,a,b,c,d,41); _R3(d,e,a,b,c,42); _R3(c,d,e,a,b,43);
+		_R3(b,c,d,e,a,44); _R3(a,b,c,d,e,45); _R3(e,a,b,c,d,46); _R3(d,e,a,b,c,47);
+		_R3(c,d,e,a,b,48); _R3(b,c,d,e,a,49); _R3(a,b,c,d,e,50); _R3(e,a,b,c,d,51);
+		_R3(d,e,a,b,c,52); _R3(c,d,e,a,b,53); _R3(b,c,d,e,a,54); _R3(a,b,c,d,e,55);
+		_R3(e,a,b,c,d,56); _R3(d,e,a,b,c,57); _R3(c,d,e,a,b,58); _R3(b,c,d,e,a,59);
+		_R4(a,b,c,d,e,60); _R4(e,a,b,c,d,61); _R4(d,e,a,b,c,62); _R4(c,d,e,a,b,63);
+		_R4(b,c,d,e,a,64); _R4(a,b,c,d,e,65); _R4(e,a,b,c,d,66); _R4(d,e,a,b,c,67);
+		_R4(c,d,e,a,b,68); _R4(b,c,d,e,a,69); _R4(a,b,c,d,e,70); _R4(e,a,b,c,d,71);
+		_R4(d,e,a,b,c,72); _R4(c,d,e,a,b,73); _R4(b,c,d,e,a,74); _R4(a,b,c,d,e,75);
+		_R4(e,a,b,c,d,76); _R4(d,e,a,b,c,77); _R4(c,d,e,a,b,78); _R4(b,c,d,e,a,79);
+
+		// Add the working vars back into state
+		a += a0;
+		b += b0;
+		c += c0;
+		d += d0;
+		e += e0;
+
+		buffer += SHA1_BLOCK_SIZE;
+		len -= SHA1_BLOCK_SIZE;
+	}
+
+	// save state
+	m_state[0] = a;
+	m_state[1] = b;
+	m_state[2] = c;
+	m_state[3] = d;
+	m_state[4] = e;
+#endif
 }
 
 // do not remove #undefs, or you can get name collision with libc++'s <functional>
@@ -689,32 +1169,43 @@ void FSHA1::Transform(uint32 *state, const uint8 *buffer)
 #undef _R2
 #undef _R3
 #undef _R4
+#undef SHABLK
+#undef SHABLK0
+#undef ROL32
 
 // Use this function to hash in binary data
-void FSHA1::Update(const uint8 *data, uint64 len)
+void FSHA1::Update(const uint8* data, uint64 len)
 {
-	uint64 i, j;
+	if (len == 0)
+	{
+		return;
+	}
 
-	j = (m_count[0] >> 3) & 63;
+	size_t left = m_count % SHA1_BLOCK_SIZE;
+	size_t fill = SHA1_BLOCK_SIZE - left;
 
-	uint64 inputBitCount = len << 3;
-	uint64 newBitCount = (uint64(m_count[0]) | (uint64(m_count[1]) << 32)) + inputBitCount;
-	m_count[0] = uint32(newBitCount);
-	m_count[1] = uint32(newBitCount >> 32);
+	m_count += len;
 
-	if((j + len) > 63)
+	if (left && len >= fill)
 	{
-		i = 64 - j;
-		FMemory::Memcpy(&m_buffer[j], data, i);
-		Transform(m_state, m_buffer);
-
-		for( ; i + 63 < len; i += 64) Transform(m_state, &data[i]);
+		// process internal buffer if it can be filled fully
+		FMemory::Memcpy(&m_buffer[left], data, fill);
+		Transform(m_buffer, SHA1_BLOCK_SIZE);
+		data += fill;
+		len -= fill;
+		left = 0;
+	}
 
-		j = 0;
+	if (len >= SHA1_BLOCK_SIZE)
+	{
+		// process whole blocks in one call as much as possible
+		Transform(data, len);
+		data += len & ~(SHA1_BLOCK_SIZE - 1);
+		len %= SHA1_BLOCK_SIZE;
 	}
-	else i = 0;
 
-	FMemory::Memcpy(&m_buffer[j], &data[i], len - i);
+	// remember any leftover bytes into internal buffer
+	FMemory::Memcpy(&m_buffer[left], data, len);
 }
 
 // Use this function to hash in strings
@@ -725,26 +1216,24 @@ void FSHA1::UpdateWithString(const TCHAR *String, uint32 Length)
 
 void FSHA1::Final()
 {
-	uint32 i;
+	// processed input bit count, in big endian
 	uint8 finalcount[8];
-
-	for(i = 0; i < 8; i++)
+	for (int i = 0; i < 8; i++)
 	{
-		finalcount[i] = (uint8)((m_count[((i >= 4) ? 0 : 1)] >> ((3 - (i & 3)) * 8) ) & 255); // Endian independent
+		finalcount[i] = (uint8)((m_count << 3) >> ((7 - i) * 8));
 	}
 
-	Update((uint8*)"\200", 1);
-
-	while ((m_count[0] & 504) != 448)
-	{
-		Update((uint8*)"\0", 1);
-	}
+	// pad to 56 mod 64
+	size_t last = m_count % SHA1_BLOCK_SIZE;
+	size_t padn = (last < SHA1_BLOCK_SIZE - 8) ? (SHA1_BLOCK_SIZE - 8 - last) : (2 * SHA1_BLOCK_SIZE - 8 - last);
+	Update(PADDING, padn);
 
-	Update(finalcount, 8); // Cause a SHA1Transform()
+	// last 8 bytes is size
+	Update(finalcount, sizeof(finalcount));
 
-	for(i = 0; i < 20; i++)
+	for (int i = 0; i < 20; i++)
 	{
-		m_digest[i] = (uint8)((m_state[i >> 2] >> ((3 - (i & 3)) * 8) ) & 255);
+		m_digest[i] = (uint8)(m_state[i >> 2] >> ((3 - (i & 3)) * 8));
 	}
 }
 
diff --git a/Engine/Source/Runtime/Core/Public/HAL/Platform.h b/Engine/Source/Runtime/Core/Public/HAL/Platform.h
--- a/Engine/Source/Runtime/Core/Public/HAL/Platform.h
+++ b/Engine/Source/Runtime/Core/Public/HAL/Platform.h
@@ -182,6 +182,12 @@
 #ifndef PLATFORM_ALWAYS_HAS_FMA3
 	#define PLATFORM_ALWAYS_HAS_FMA3			0
 #endif
+#ifndef PLATFORM_ALWAYS_HAS_AESNI
+	#define PLATFORM_ALWAYS_HAS_AESNI			0
+#endif
+#ifndef PLATFORM_ALWAYS_HAS_SHA
+	#define PLATFORM_ALWAYS_HAS_SHA				0
+#endif
 
 
 #ifndef PLATFORM_HAS_CPUID
diff --git a/Engine/Source/Runtime/Core/Public/Misc/AES.h b/Engine/Source/Runtime/Core/Public/Misc/AES.h
--- a/Engine/Source/Runtime/Core/Public/Misc/AES.h
+++ b/Engine/Source/Runtime/Core/Public/Misc/AES.h
@@ -5,8 +5,16 @@
 #include "CoreTypes.h"
 #include "HAL/UnrealMemory.h"
 
-#define AES_BLOCK_SIZE 16
+// AES-256 implementation - using ECB mode for multiple blocks
 
+// The currently implemented approach has the shortcoming that it encrypts and decrypts each
+// 128-bit block separately. If the plaintext contains identical 128-byte blocks, the blocks
+// will be encrypted identically. This makes some of the plaintext structure visible in the
+// ciphertext, even to someone who does not have the key.
+
+// DO NOT USE this functionality for any new place where you might need encryption. Because
+// current code is meant for keeping backwards compatiblity with existing data. Better way
+// to use AES would be integrated with authentication, for example, in AES-GCM mode.
 
 struct CORE_API FAES
 {
@@ -64,7 +72,7 @@ struct CORE_API FAES
 	 *
 	 * @param Contents the buffer to encrypt
 	 * @param NumBytes the size of the buffer
-	 * @param Key a null terminated string that is a 32 byte multiple length
+	 * @param Key a null terminated string that is a 32 bytes long
 	 */
 	static void EncryptData(uint8* Contents, uint32 NumBytes, const ANSICHAR* Key);
 
@@ -73,7 +81,8 @@ struct CORE_API FAES
 	*
 	* @param Contents the buffer to encrypt
 	* @param NumBytes the size of the buffer
-	* @param Key a byte array that is a 32 byte multiple length
+	* @param Key a byte array that is a 32 byte length
+	* @param NumKeyBytes length of Key byte array, must be 32
 	*/
 	static void EncryptData(uint8* Contents, uint32 NumBytes, const uint8* KeyBytes, uint32 NumKeyBytes);
 
@@ -82,7 +91,7 @@ struct CORE_API FAES
 	*
 	* @param Contents the buffer to encrypt
 	* @param NumBytes the size of the buffer
-	* @param Key a null terminated string that is a 32 byte multiple length
+	* @param Key An FAESKey object containing the decryption key
 	*/
 	static void DecryptData(uint8* Contents, uint32 NumBytes, const FAESKey& Key);
 
@@ -91,7 +100,7 @@ struct CORE_API FAES
 	 *
 	 * @param Contents the buffer to encrypt
 	 * @param NumBytes the size of the buffer
-	 * @param Key a null terminated string that is a 32 byte multiple length
+	 * @param Key a null terminated string that is a 32 bytes long
 	 */
 	static void DecryptData(uint8* Contents, uint32 NumBytes, const ANSICHAR* Key);
 
@@ -100,7 +109,8 @@ struct CORE_API FAES
 	*
 	* @param Contents the buffer to encrypt
 	* @param NumBytes the size of the buffer
-	* @param Key a byte array that is a 32 byte multiple length
+	* @param Key a byte array that is a 32 byte length
+	* @param NumKeyBytes length of Key byte array, must be 32
 	*/
 	static void DecryptData(uint8* Contents, uint32 NumBytes, const uint8* KeyBytes, uint32 NumKeyBytes);
 };
diff --git a/Engine/Source/Runtime/Core/Public/Misc/SecureHash.h b/Engine/Source/Runtime/Core/Public/Misc/SecureHash.h
--- a/Engine/Source/Runtime/Core/Public/Misc/SecureHash.h
+++ b/Engine/Source/Runtime/Core/Public/Misc/SecureHash.h
@@ -184,13 +184,6 @@ private:
  *	Web: http://www.dominik-reichl.de/
  */
 
-
-typedef union
-{
-	uint8  c[64];
-	uint32 l[16];
-} SHA1_WORKSPACE_BLOCK;
-
 /** This divider string is beween full file hashes and script hashes */
 #define HASHES_SHA_DIVIDER "+++"
 
@@ -285,7 +278,7 @@ public:
 	~FSHA1();
 
 	uint32 m_state[5];
-	uint32 m_count[2];
+	uint64 m_count;
 	uint32 __reserved1[1];
 	uint8  m_buffer[64];
 	uint8  m_digest[20];
@@ -373,11 +366,7 @@ public:
 
 private:
 	// Private SHA-1 transformation
-	void Transform(uint32 *state, const uint8 *buffer);
-
-	// Member variables
-	uint8 m_workspace[64];
-	SHA1_WORKSPACE_BLOCK *m_block; // SHA1 pointer to the byte array above
+	void Transform(const uint8* buffer, uint64 len);
 
 	/** Global map of filename to hash value, filled out in InitializeFileHashesFromBuffer */
 	static TMap<FString, uint8*> FullFileSHAHashMap;
diff --git a/Engine/Source/Runtime/Core/Tests/Misc/AESTest.cpp b/Engine/Source/Runtime/Core/Tests/Misc/AESTest.cpp
new file mode 100644
--- /dev/null
+++ b/Engine/Source/Runtime/Core/Tests/Misc/AESTest.cpp
@@ -0,0 +1,858 @@
+// Copyright Epic Games, Inc. All Rights Reserved.
+
+#if WITH_LOW_LEVEL_TESTS
+
+#include "Misc/AES.h"
+
+#include "HAL/UnrealMemory.h"
+#include "String/HexToBytes.h"
+#include "String/BytesToHex.h"
+#include "Containers/StringView.h"
+#include "Containers/Array.h"
+#include "Misc/StringBuilder.h"
+
+#include "TestHarness.h"
+
+TEST_CASE("Core::Misc::AES", "[Core][AES][Smoke]")
+{
+	auto RunTest = [](FAnsiStringView Key, FAnsiStringView PlainText, FAnsiStringView CipherText, int IterCount = 1)
+	{
+		checkf((Key.Len()/2) == FAES::FAESKey::KeySize, TEXT("key size must be 32 bytes = 256 bits"));
+		checkf((PlainText.Len()/2) % FAES::AESBlockSize == 0, TEXT("plain-text must be multiple of block size"));
+		checkf(CipherText.Len() == PlainText.Len(), TEXT("expected cipher-text should be same length as plain-text"));
+
+		TArray<uint8> KeyBytes;
+		KeyBytes.AddUninitialized(FAES::FAESKey::KeySize);
+		UE::String::HexToBytes(Key, KeyBytes.GetData());
+
+		TArray<uint8> PlainTextBytes;
+		PlainTextBytes.AddUninitialized(PlainText.Len() / 2);
+		UE::String::HexToBytes(PlainText, PlainTextBytes.GetData());
+
+		TArray<uint8> CipherTextBytes;
+		CipherTextBytes.AddUninitialized(CipherText.Len() / 2);
+		UE::String::HexToBytes(CipherText, CipherTextBytes.GetData());
+
+		// encryption
+		{
+			TArray<uint8> Encrypted = PlainTextBytes;
+			for (int Iter = 0; Iter < IterCount; Iter++)
+			{
+				FAES::EncryptData(Encrypted.GetData(), Encrypted.Num(), KeyBytes.GetData(), KeyBytes.Num());
+			}
+
+			FAnsiStringBuilderBase EncryptedHex;
+			UE::String::BytesToHex(Encrypted, EncryptedHex);
+
+			REQUIRE(EncryptedHex.ToView() == CipherText);
+		}
+
+		// decryption
+		{
+			TArray<uint8> Decrypted = CipherTextBytes;
+			for (int Iter = 0; Iter < IterCount; Iter++)
+			{
+				FAES::DecryptData(Decrypted.GetData(), Decrypted.Num(), KeyBytes.GetData(), KeyBytes.Num());
+			}
+
+			FAnsiStringBuilderBase DecryptedHex;
+			UE::String::BytesToHex(Decrypted, DecryptedHex);
+
+			REQUIRE(DecryptedHex.ToView() == PlainText);
+		}
+	};
+
+	struct AesTestCase
+	{
+		FAnsiStringView Key;
+		FAnsiStringView PlainText;
+		FAnsiStringView CipherText;
+	};
+
+	// Appendix C, https://csrc.nist.gov/csrc/media/publications/fips/197/final/documents/fips-197.pdf
+	SECTION("Example Vector")
+	{
+		RunTest("000102030405060708090a0b0c0d0e0f101112131415161718191a1b1c1d1e1f", "00112233445566778899aabbccddeeff", "8ea2b7ca516745bfeafc49904b496089");
+	}
+
+	// https://csrc.nist.gov/Projects/cryptographic-algorithm-validation-program/Block-Ciphers#AES
+
+	// Appendix B
+	SECTION("GFSbox Known Answer")
+	{
+		static const AesTestCase GFSboxCases[] =
+		{
+			{ "0000000000000000000000000000000000000000000000000000000000000000", "014730f80ac625fe84f026c60bfd547d", "5c9d844ed46f9885085e5d6a4f94c7d7" },
+			{ "0000000000000000000000000000000000000000000000000000000000000000", "0b24af36193ce4665f2825d7b4749c98", "a9ff75bd7cf6613d3731c77c3b6d0c04" },
+			{ "0000000000000000000000000000000000000000000000000000000000000000", "761c1fe41a18acf20d241650611d90f1", "623a52fcea5d443e48d9181ab32c7421" },
+			{ "0000000000000000000000000000000000000000000000000000000000000000", "8a560769d605868ad80d819bdba03771", "38f2c7ae10612415d27ca190d27da8b4" },
+			{ "0000000000000000000000000000000000000000000000000000000000000000", "91fbef2d15a97816060bee1feaa49afe", "1bc704f1bce135ceb810341b216d7abe" },
+		};
+		for (auto& Test : GFSboxCases)
+		{
+			RunTest(Test.Key, Test.PlainText, Test.CipherText);
+		}
+	}
+
+	// Appendix C
+	SECTION("KeySbox Known Answer")
+	{
+		static const AesTestCase KeySboxCases[] =
+		{
+			{ "c47b0294dbbbee0fec4757f22ffeee3587ca4730c3d33b691df38bab076bc558", "00000000000000000000000000000000", "46f2fb342d6f0ab477476fc501242c5f" },
+			{ "28d46cffa158533194214a91e712fc2b45b518076675affd910edeca5f41ac64", "00000000000000000000000000000000", "4bf3b0a69aeb6657794f2901b1440ad4" },
+			{ "c1cc358b449909a19436cfbb3f852ef8bcb5ed12ac7058325f56e6099aab1a1c", "00000000000000000000000000000000", "352065272169abf9856843927d0674fd" },
+			{ "984ca75f4ee8d706f46c2d98c0bf4a45f5b00d791c2dfeb191b5ed8e420fd627", "00000000000000000000000000000000", "4307456a9e67813b452e15fa8fffe398" },
+			{ "b43d08a447ac8609baadae4ff12918b9f68fc1653f1269222f123981ded7a92f", "00000000000000000000000000000000", "4663446607354989477a5c6f0f007ef4" },
+			{ "1d85a181b54cde51f0e098095b2962fdc93b51fe9b88602b3f54130bf76a5bd9", "00000000000000000000000000000000", "531c2c38344578b84d50b3c917bbb6e1" },
+			{ "dc0eba1f2232a7879ded34ed8428eeb8769b056bbaf8ad77cb65c3541430b4cf", "00000000000000000000000000000000", "fc6aec906323480005c58e7e1ab004ad" },
+			{ "f8be9ba615c5a952cabbca24f68f8593039624d524c816acda2c9183bd917cb9", "00000000000000000000000000000000", "a3944b95ca0b52043584ef02151926a8" },
+			{ "797f8b3d176dac5b7e34a2d539c4ef367a16f8635f6264737591c5c07bf57a3e", "00000000000000000000000000000000", "a74289fe73a4c123ca189ea1e1b49ad5" },
+			{ "6838d40caf927749c13f0329d331f448e202c73ef52c5f73a37ca635d4c47707", "00000000000000000000000000000000", "b91d4ea4488644b56cf0812fa7fcf5fc" },
+			{ "ccd1bc3c659cd3c59bc437484e3c5c724441da8d6e90ce556cd57d0752663bbc", "00000000000000000000000000000000", "304f81ab61a80c2e743b94d5002a126b" },
+			{ "13428b5e4c005e0636dd338405d173ab135dec2a25c22c5df0722d69dcc43887", "00000000000000000000000000000000", "649a71545378c783e368c9ade7114f6c" },
+			{ "07eb03a08d291d1b07408bf3512ab40c91097ac77461aad4bb859647f74f00ee", "00000000000000000000000000000000", "47cb030da2ab051dfc6c4bf6910d12bb" },
+			{ "90143ae20cd78c5d8ebdd6cb9dc1762427a96c78c639bccc41a61424564eafe1", "00000000000000000000000000000000", "798c7c005dee432b2c8ea5dfa381ecc3" },
+			{ "b7a5794d52737475d53d5a377200849be0260a67a2b22ced8bbef12882270d07", "00000000000000000000000000000000", "637c31dc2591a07636f646b72daabbe7" },
+			{ "fca02f3d5011cfc5c1e23165d413a049d4526a991827424d896fe3435e0bf68e", "00000000000000000000000000000000", "179a49c712154bbffbe6e7a84a18e220" },
+		};
+		for (auto& Test : KeySboxCases)
+		{
+			RunTest(Test.Key, Test.PlainText, Test.CipherText);
+		}
+	}
+
+	// Appendix D
+	SECTION("VarTxt Known Answer")
+	{
+		static const AesTestCase VarTxtCases[] =
+		{
+			{ "0000000000000000000000000000000000000000000000000000000000000000", "80000000000000000000000000000000", "ddc6bf790c15760d8d9aeb6f9a75fd4e" },
+			{ "0000000000000000000000000000000000000000000000000000000000000000", "c0000000000000000000000000000000", "0a6bdc6d4c1e6280301fd8e97ddbe601" },
+			{ "0000000000000000000000000000000000000000000000000000000000000000", "e0000000000000000000000000000000", "9b80eefb7ebe2d2b16247aa0efc72f5d" },
+			{ "0000000000000000000000000000000000000000000000000000000000000000", "f0000000000000000000000000000000", "7f2c5ece07a98d8bee13c51177395ff7" },
+			{ "0000000000000000000000000000000000000000000000000000000000000000", "f8000000000000000000000000000000", "7818d800dcf6f4be1e0e94f403d1e4c2" },
+			{ "0000000000000000000000000000000000000000000000000000000000000000", "fc000000000000000000000000000000", "e74cd1c92f0919c35a0324123d6177d3" },
+			{ "0000000000000000000000000000000000000000000000000000000000000000", "fe000000000000000000000000000000", "8092a4dcf2da7e77e93bdd371dfed82e" },
+			{ "0000000000000000000000000000000000000000000000000000000000000000", "ff000000000000000000000000000000", "49af6b372135acef10132e548f217b17" },
+			{ "0000000000000000000000000000000000000000000000000000000000000000", "ff800000000000000000000000000000", "8bcd40f94ebb63b9f7909676e667f1e7" },
+			{ "0000000000000000000000000000000000000000000000000000000000000000", "ffc00000000000000000000000000000", "fe1cffb83f45dcfb38b29be438dbd3ab" },
+			{ "0000000000000000000000000000000000000000000000000000000000000000", "ffe00000000000000000000000000000", "0dc58a8d886623705aec15cb1e70dc0e" },
+			{ "0000000000000000000000000000000000000000000000000000000000000000", "fff00000000000000000000000000000", "c218faa16056bd0774c3e8d79c35a5e4" },
+			{ "0000000000000000000000000000000000000000000000000000000000000000", "fff80000000000000000000000000000", "047bba83f7aa841731504e012208fc9e" },
+			{ "0000000000000000000000000000000000000000000000000000000000000000", "fffc0000000000000000000000000000", "dc8f0e4915fd81ba70a331310882f6da" },
+			{ "0000000000000000000000000000000000000000000000000000000000000000", "fffe0000000000000000000000000000", "1569859ea6b7206c30bf4fd0cbfac33c" },
+			{ "0000000000000000000000000000000000000000000000000000000000000000", "ffff0000000000000000000000000000", "300ade92f88f48fa2df730ec16ef44cd" },
+			{ "0000000000000000000000000000000000000000000000000000000000000000", "ffff8000000000000000000000000000", "1fe6cc3c05965dc08eb0590c95ac71d0" },
+			{ "0000000000000000000000000000000000000000000000000000000000000000", "ffffc000000000000000000000000000", "59e858eaaa97fec38111275b6cf5abc0" },
+			{ "0000000000000000000000000000000000000000000000000000000000000000", "ffffe000000000000000000000000000", "2239455e7afe3b0616100288cc5a723b" },
+			{ "0000000000000000000000000000000000000000000000000000000000000000", "fffff000000000000000000000000000", "3ee500c5c8d63479717163e55c5c4522" },
+			{ "0000000000000000000000000000000000000000000000000000000000000000", "fffff800000000000000000000000000", "d5e38bf15f16d90e3e214041d774daa8" },
+			{ "0000000000000000000000000000000000000000000000000000000000000000", "fffffc00000000000000000000000000", "b1f4066e6f4f187dfe5f2ad1b17819d0" },
+			{ "0000000000000000000000000000000000000000000000000000000000000000", "fffffe00000000000000000000000000", "6ef4cc4de49b11065d7af2909854794a" },
+			{ "0000000000000000000000000000000000000000000000000000000000000000", "ffffff00000000000000000000000000", "ac86bc606b6640c309e782f232bf367f" },
+			{ "0000000000000000000000000000000000000000000000000000000000000000", "ffffff80000000000000000000000000", "36aff0ef7bf3280772cf4cac80a0d2b2" },
+			{ "0000000000000000000000000000000000000000000000000000000000000000", "ffffffc0000000000000000000000000", "1f8eedea0f62a1406d58cfc3ecea72cf" },
+			{ "0000000000000000000000000000000000000000000000000000000000000000", "ffffffe0000000000000000000000000", "abf4154a3375a1d3e6b1d454438f95a6" },
+			{ "0000000000000000000000000000000000000000000000000000000000000000", "fffffff0000000000000000000000000", "96f96e9d607f6615fc192061ee648b07" },
+			{ "0000000000000000000000000000000000000000000000000000000000000000", "fffffff8000000000000000000000000", "cf37cdaaa0d2d536c71857634c792064" },
+			{ "0000000000000000000000000000000000000000000000000000000000000000", "fffffffc000000000000000000000000", "fbd6640c80245c2b805373f130703127" },
+			{ "0000000000000000000000000000000000000000000000000000000000000000", "fffffffe000000000000000000000000", "8d6a8afe55a6e481badae0d146f436db" },
+			{ "0000000000000000000000000000000000000000000000000000000000000000", "ffffffff000000000000000000000000", "6a4981f2915e3e68af6c22385dd06756" },
+			{ "0000000000000000000000000000000000000000000000000000000000000000", "ffffffff800000000000000000000000", "42a1136e5f8d8d21d3101998642d573b" },
+			{ "0000000000000000000000000000000000000000000000000000000000000000", "ffffffffc00000000000000000000000", "9b471596dc69ae1586cee6158b0b0181" },
+			{ "0000000000000000000000000000000000000000000000000000000000000000", "ffffffffe00000000000000000000000", "753665c4af1eff33aa8b628bf8741cfd" },
+			{ "0000000000000000000000000000000000000000000000000000000000000000", "fffffffff00000000000000000000000", "9a682acf40be01f5b2a4193c9a82404d" },
+			{ "0000000000000000000000000000000000000000000000000000000000000000", "fffffffff80000000000000000000000", "54fafe26e4287f17d1935f87eb9ade01" },
+			{ "0000000000000000000000000000000000000000000000000000000000000000", "fffffffffc0000000000000000000000", "49d541b2e74cfe73e6a8e8225f7bd449" },
+			{ "0000000000000000000000000000000000000000000000000000000000000000", "fffffffffe0000000000000000000000", "11a45530f624ff6f76a1b3826626ff7b" },
+			{ "0000000000000000000000000000000000000000000000000000000000000000", "ffffffffff0000000000000000000000", "f96b0c4a8bc6c86130289f60b43b8fba" },
+			{ "0000000000000000000000000000000000000000000000000000000000000000", "ffffffffff8000000000000000000000", "48c7d0e80834ebdc35b6735f76b46c8b" },
+			{ "0000000000000000000000000000000000000000000000000000000000000000", "ffffffffffc000000000000000000000", "2463531ab54d66955e73edc4cb8eaa45" },
+			{ "0000000000000000000000000000000000000000000000000000000000000000", "ffffffffffe000000000000000000000", "ac9bd8e2530469134b9d5b065d4f565b" },
+			{ "0000000000000000000000000000000000000000000000000000000000000000", "fffffffffff000000000000000000000", "3f5f9106d0e52f973d4890e6f37e8a00" },
+			{ "0000000000000000000000000000000000000000000000000000000000000000", "fffffffffff800000000000000000000", "20ebc86f1304d272e2e207e59db639f0" },
+			{ "0000000000000000000000000000000000000000000000000000000000000000", "fffffffffffc00000000000000000000", "e67ae6426bf9526c972cff072b52252c" },
+			{ "0000000000000000000000000000000000000000000000000000000000000000", "fffffffffffe00000000000000000000", "1a518dddaf9efa0d002cc58d107edfc8" },
+			{ "0000000000000000000000000000000000000000000000000000000000000000", "ffffffffffff00000000000000000000", "ead731af4d3a2fe3b34bed047942a49f" },
+			{ "0000000000000000000000000000000000000000000000000000000000000000", "ffffffffffff80000000000000000000", "b1d4efe40242f83e93b6c8d7efb5eae9" },
+			{ "0000000000000000000000000000000000000000000000000000000000000000", "ffffffffffffc0000000000000000000", "cd2b1fec11fd906c5c7630099443610a" },
+			{ "0000000000000000000000000000000000000000000000000000000000000000", "ffffffffffffe0000000000000000000", "a1853fe47fe29289d153161d06387d21" },
+			{ "0000000000000000000000000000000000000000000000000000000000000000", "fffffffffffff0000000000000000000", "4632154179a555c17ea604d0889fab14" },
+			{ "0000000000000000000000000000000000000000000000000000000000000000", "fffffffffffff8000000000000000000", "dd27cac6401a022e8f38f9f93e774417" },
+			{ "0000000000000000000000000000000000000000000000000000000000000000", "fffffffffffffc000000000000000000", "c090313eb98674f35f3123385fb95d4d" },
+			{ "0000000000000000000000000000000000000000000000000000000000000000", "fffffffffffffe000000000000000000", "cc3526262b92f02edce548f716b9f45c" },
+			{ "0000000000000000000000000000000000000000000000000000000000000000", "ffffffffffffff000000000000000000", "c0838d1a2b16a7c7f0dfcc433c399c33" },
+			{ "0000000000000000000000000000000000000000000000000000000000000000", "ffffffffffffff800000000000000000", "0d9ac756eb297695eed4d382eb126d26" },
+			{ "0000000000000000000000000000000000000000000000000000000000000000", "ffffffffffffffc00000000000000000", "56ede9dda3f6f141bff1757fa689c3e1" },
+			{ "0000000000000000000000000000000000000000000000000000000000000000", "ffffffffffffffe00000000000000000", "768f520efe0f23e61d3ec8ad9ce91774" },
+			{ "0000000000000000000000000000000000000000000000000000000000000000", "fffffffffffffff00000000000000000", "b1144ddfa75755213390e7c596660490" },
+			{ "0000000000000000000000000000000000000000000000000000000000000000", "fffffffffffffff80000000000000000", "1d7c0c4040b355b9d107a99325e3b050" },
+			{ "0000000000000000000000000000000000000000000000000000000000000000", "fffffffffffffffc0000000000000000", "d8e2bb1ae8ee3dcf5bf7d6c38da82a1a" },
+			{ "0000000000000000000000000000000000000000000000000000000000000000", "fffffffffffffffe0000000000000000", "faf82d178af25a9886a47e7f789b98d7" },
+			{ "0000000000000000000000000000000000000000000000000000000000000000", "ffffffffffffffff0000000000000000", "9b58dbfd77fe5aca9cfc190cd1b82d19" },
+			{ "0000000000000000000000000000000000000000000000000000000000000000", "ffffffffffffffff8000000000000000", "77f392089042e478ac16c0c86a0b5db5" },
+			{ "0000000000000000000000000000000000000000000000000000000000000000", "ffffffffffffffffc000000000000000", "19f08e3420ee69b477ca1420281c4782" },
+			{ "0000000000000000000000000000000000000000000000000000000000000000", "ffffffffffffffffe000000000000000", "a1b19beee4e117139f74b3c53fdcb875" },
+			{ "0000000000000000000000000000000000000000000000000000000000000000", "fffffffffffffffff000000000000000", "a37a5869b218a9f3a0868d19aea0ad6a" },
+			{ "0000000000000000000000000000000000000000000000000000000000000000", "fffffffffffffffff800000000000000", "bc3594e865bcd0261b13202731f33580" },
+			{ "0000000000000000000000000000000000000000000000000000000000000000", "fffffffffffffffffc00000000000000", "811441ce1d309eee7185e8c752c07557" },
+			{ "0000000000000000000000000000000000000000000000000000000000000000", "fffffffffffffffffe00000000000000", "959971ce4134190563518e700b9874d1" },
+			{ "0000000000000000000000000000000000000000000000000000000000000000", "ffffffffffffffffff00000000000000", "76b5614a042707c98e2132e2e805fe63" },
+			{ "0000000000000000000000000000000000000000000000000000000000000000", "ffffffffffffffffff80000000000000", "7d9fa6a57530d0f036fec31c230b0cc6" },
+			{ "0000000000000000000000000000000000000000000000000000000000000000", "ffffffffffffffffffc0000000000000", "964153a83bf6989a4ba80daa91c3e081" },
+			{ "0000000000000000000000000000000000000000000000000000000000000000", "ffffffffffffffffffe0000000000000", "a013014d4ce8054cf2591d06f6f2f176" },
+			{ "0000000000000000000000000000000000000000000000000000000000000000", "fffffffffffffffffff0000000000000", "d1c5f6399bf382502e385eee1474a869" },
+			{ "0000000000000000000000000000000000000000000000000000000000000000", "fffffffffffffffffff8000000000000", "0007e20b8298ec354f0f5fe7470f36bd" },
+			{ "0000000000000000000000000000000000000000000000000000000000000000", "fffffffffffffffffffc000000000000", "b95ba05b332da61ef63a2b31fcad9879" },
+			{ "0000000000000000000000000000000000000000000000000000000000000000", "fffffffffffffffffffe000000000000", "4620a49bd967491561669ab25dce45f4" },
+			{ "0000000000000000000000000000000000000000000000000000000000000000", "ffffffffffffffffffff000000000000", "12e71214ae8e04f0bb63d7425c6f14d5" },
+			{ "0000000000000000000000000000000000000000000000000000000000000000", "ffffffffffffffffffff800000000000", "4cc42fc1407b008fe350907c092e80ac" },
+			{ "0000000000000000000000000000000000000000000000000000000000000000", "ffffffffffffffffffffc00000000000", "08b244ce7cbc8ee97fbba808cb146fda" },
+			{ "0000000000000000000000000000000000000000000000000000000000000000", "ffffffffffffffffffffe00000000000", "39b333e8694f21546ad1edd9d87ed95b" },
+			{ "0000000000000000000000000000000000000000000000000000000000000000", "fffffffffffffffffffff00000000000", "3b271f8ab2e6e4a20ba8090f43ba78f3" },
+			{ "0000000000000000000000000000000000000000000000000000000000000000", "fffffffffffffffffffff80000000000", "9ad983f3bf651cd0393f0a73cccdea50" },
+			{ "0000000000000000000000000000000000000000000000000000000000000000", "fffffffffffffffffffffc0000000000", "8f476cbff75c1f725ce18e4bbcd19b32" },
+			{ "0000000000000000000000000000000000000000000000000000000000000000", "fffffffffffffffffffffe0000000000", "905b6267f1d6ab5320835a133f096f2a" },
+			{ "0000000000000000000000000000000000000000000000000000000000000000", "ffffffffffffffffffffff0000000000", "145b60d6d0193c23f4221848a892d61a" },
+			{ "0000000000000000000000000000000000000000000000000000000000000000", "ffffffffffffffffffffff8000000000", "55cfb3fb6d75cad0445bbc8dafa25b0f" },
+			{ "0000000000000000000000000000000000000000000000000000000000000000", "ffffffffffffffffffffffc000000000", "7b8e7098e357ef71237d46d8b075b0f5" },
+			{ "0000000000000000000000000000000000000000000000000000000000000000", "ffffffffffffffffffffffe000000000", "2bf27229901eb40f2df9d8398d1505ae" },
+			{ "0000000000000000000000000000000000000000000000000000000000000000", "fffffffffffffffffffffff000000000", "83a63402a77f9ad5c1e931a931ecd706" },
+			{ "0000000000000000000000000000000000000000000000000000000000000000", "fffffffffffffffffffffff800000000", "6f8ba6521152d31f2bada1843e26b973" },
+			{ "0000000000000000000000000000000000000000000000000000000000000000", "fffffffffffffffffffffffc00000000", "e5c3b8e30fd2d8e6239b17b44bd23bbd" },
+			{ "0000000000000000000000000000000000000000000000000000000000000000", "fffffffffffffffffffffffe00000000", "1ac1f7102c59933e8b2ddc3f14e94baa" },
+			{ "0000000000000000000000000000000000000000000000000000000000000000", "ffffffffffffffffffffffff00000000", "21d9ba49f276b45f11af8fc71a088e3d" },
+			{ "0000000000000000000000000000000000000000000000000000000000000000", "ffffffffffffffffffffffff80000000", "649f1cddc3792b4638635a392bc9bade" },
+			{ "0000000000000000000000000000000000000000000000000000000000000000", "ffffffffffffffffffffffffc0000000", "e2775e4b59c1bc2e31a2078c11b5a08c" },
+			{ "0000000000000000000000000000000000000000000000000000000000000000", "ffffffffffffffffffffffffe0000000", "2be1fae5048a25582a679ca10905eb80" },
+			{ "0000000000000000000000000000000000000000000000000000000000000000", "fffffffffffffffffffffffff0000000", "da86f292c6f41ea34fb2068df75ecc29" },
+			{ "0000000000000000000000000000000000000000000000000000000000000000", "fffffffffffffffffffffffff8000000", "220df19f85d69b1b562fa69a3c5beca5" },
+			{ "0000000000000000000000000000000000000000000000000000000000000000", "fffffffffffffffffffffffffc000000", "1f11d5d0355e0b556ccdb6c7f5083b4d" },
+			{ "0000000000000000000000000000000000000000000000000000000000000000", "fffffffffffffffffffffffffe000000", "62526b78be79cb384633c91f83b4151b" },
+			{ "0000000000000000000000000000000000000000000000000000000000000000", "ffffffffffffffffffffffffff000000", "90ddbcb950843592dd47bbef00fdc876" },
+			{ "0000000000000000000000000000000000000000000000000000000000000000", "ffffffffffffffffffffffffff800000", "2fd0e41c5b8402277354a7391d2618e2" },
+			{ "0000000000000000000000000000000000000000000000000000000000000000", "ffffffffffffffffffffffffffc00000", "3cdf13e72dee4c581bafec70b85f9660" },
+			{ "0000000000000000000000000000000000000000000000000000000000000000", "ffffffffffffffffffffffffffe00000", "afa2ffc137577092e2b654fa199d2c43" },
+			{ "0000000000000000000000000000000000000000000000000000000000000000", "fffffffffffffffffffffffffff00000", "8d683ee63e60d208e343ce48dbc44cac" },
+			{ "0000000000000000000000000000000000000000000000000000000000000000", "fffffffffffffffffffffffffff80000", "705a4ef8ba2133729c20185c3d3a4763" },
+			{ "0000000000000000000000000000000000000000000000000000000000000000", "fffffffffffffffffffffffffffc0000", "0861a861c3db4e94194211b77ed761b9" },
+			{ "0000000000000000000000000000000000000000000000000000000000000000", "fffffffffffffffffffffffffffe0000", "4b00c27e8b26da7eab9d3a88dec8b031" },
+			{ "0000000000000000000000000000000000000000000000000000000000000000", "ffffffffffffffffffffffffffff0000", "5f397bf03084820cc8810d52e5b666e9" },
+			{ "0000000000000000000000000000000000000000000000000000000000000000", "ffffffffffffffffffffffffffff8000", "63fafabb72c07bfbd3ddc9b1203104b8" },
+			{ "0000000000000000000000000000000000000000000000000000000000000000", "ffffffffffffffffffffffffffffc000", "683e2140585b18452dd4ffbb93c95df9" },
+			{ "0000000000000000000000000000000000000000000000000000000000000000", "ffffffffffffffffffffffffffffe000", "286894e48e537f8763b56707d7d155c8" },
+			{ "0000000000000000000000000000000000000000000000000000000000000000", "fffffffffffffffffffffffffffff000", "a423deabc173dcf7e2c4c53e77d37cd1" },
+			{ "0000000000000000000000000000000000000000000000000000000000000000", "fffffffffffffffffffffffffffff800", "eb8168313e1cfdfdb5e986d5429cf172" },
+			{ "0000000000000000000000000000000000000000000000000000000000000000", "fffffffffffffffffffffffffffffc00", "27127daafc9accd2fb334ec3eba52323" },
+			{ "0000000000000000000000000000000000000000000000000000000000000000", "fffffffffffffffffffffffffffffe00", "ee0715b96f72e3f7a22a5064fc592f4c" },
+			{ "0000000000000000000000000000000000000000000000000000000000000000", "ffffffffffffffffffffffffffffff00", "29ee526770f2a11dcfa989d1ce88830f" },
+			{ "0000000000000000000000000000000000000000000000000000000000000000", "ffffffffffffffffffffffffffffff80", "0493370e054b09871130fe49af730a5a" },
+			{ "0000000000000000000000000000000000000000000000000000000000000000", "ffffffffffffffffffffffffffffffc0", "9b7b940f6c509f9e44a4ee140448ee46" },
+			{ "0000000000000000000000000000000000000000000000000000000000000000", "ffffffffffffffffffffffffffffffe0", "2915be4a1ecfdcbe3e023811a12bb6c7" },
+			{ "0000000000000000000000000000000000000000000000000000000000000000", "fffffffffffffffffffffffffffffff0", "7240e524bc51d8c4d440b1be55d1062c" },
+			{ "0000000000000000000000000000000000000000000000000000000000000000", "fffffffffffffffffffffffffffffff8", "da63039d38cb4612b2dc36ba26684b93" },
+			{ "0000000000000000000000000000000000000000000000000000000000000000", "fffffffffffffffffffffffffffffffc", "0f59cb5a4b522e2ac56c1a64f558ad9a" },
+			{ "0000000000000000000000000000000000000000000000000000000000000000", "fffffffffffffffffffffffffffffffe", "7bfe9d876c6d63c1d035da8fe21c409d" },
+			{ "0000000000000000000000000000000000000000000000000000000000000000", "ffffffffffffffffffffffffffffffff", "acdace8078a32b1a182bfa4987ca1347" },
+		};
+		for (auto& Test : VarTxtCases)
+		{
+			RunTest(Test.Key, Test.PlainText, Test.CipherText);
+		}
+	}
+
+	// Appendix E
+	SECTION("VarKey Known Answer")
+	{
+		static const AesTestCase VarKeyCases[] =
+		{
+			{ "8000000000000000000000000000000000000000000000000000000000000000", "00000000000000000000000000000000", "e35a6dcb19b201a01ebcfa8aa22b5759" },
+			{ "c000000000000000000000000000000000000000000000000000000000000000", "00000000000000000000000000000000", "b29169cdcf2d83e838125a12ee6aa400" },
+			{ "e000000000000000000000000000000000000000000000000000000000000000", "00000000000000000000000000000000", "d8f3a72fc3cdf74dfaf6c3e6b97b2fa6" },
+			{ "f000000000000000000000000000000000000000000000000000000000000000", "00000000000000000000000000000000", "1c777679d50037c79491a94da76a9a35" },
+			{ "f800000000000000000000000000000000000000000000000000000000000000", "00000000000000000000000000000000", "9cf4893ecafa0a0247a898e040691559" },
+			{ "fc00000000000000000000000000000000000000000000000000000000000000", "00000000000000000000000000000000", "8fbb413703735326310a269bd3aa94b2" },
+			{ "fe00000000000000000000000000000000000000000000000000000000000000", "00000000000000000000000000000000", "60e32246bed2b0e859e55c1cc6b26502" },
+			{ "ff00000000000000000000000000000000000000000000000000000000000000", "00000000000000000000000000000000", "ec52a212f80a09df6317021bc2a9819e" },
+			{ "ff80000000000000000000000000000000000000000000000000000000000000", "00000000000000000000000000000000", "f23e5b600eb70dbccf6c0b1d9a68182c" },
+			{ "ffc0000000000000000000000000000000000000000000000000000000000000", "00000000000000000000000000000000", "a3f599d63a82a968c33fe26590745970" },
+			{ "ffe0000000000000000000000000000000000000000000000000000000000000", "00000000000000000000000000000000", "d1ccb9b1337002cbac42c520b5d67722" },
+			{ "fff0000000000000000000000000000000000000000000000000000000000000", "00000000000000000000000000000000", "cc111f6c37cf40a1159d00fb59fb0488" },
+			{ "fff8000000000000000000000000000000000000000000000000000000000000", "00000000000000000000000000000000", "dc43b51ab609052372989a26e9cdd714" },
+			{ "fffc000000000000000000000000000000000000000000000000000000000000", "00000000000000000000000000000000", "4dcede8da9e2578f39703d4433dc6459" },
+			{ "fffe000000000000000000000000000000000000000000000000000000000000", "00000000000000000000000000000000", "1a4c1c263bbccfafc11782894685e3a8" },
+			{ "ffff000000000000000000000000000000000000000000000000000000000000", "00000000000000000000000000000000", "937ad84880db50613423d6d527a2823d" },
+			{ "ffff800000000000000000000000000000000000000000000000000000000000", "00000000000000000000000000000000", "610b71dfc688e150d8152c5b35ebc14d" },
+			{ "ffffc00000000000000000000000000000000000000000000000000000000000", "00000000000000000000000000000000", "27ef2495dabf323885aab39c80f18d8b" },
+			{ "ffffe00000000000000000000000000000000000000000000000000000000000", "00000000000000000000000000000000", "633cafea395bc03adae3a1e2068e4b4e" },
+			{ "fffff00000000000000000000000000000000000000000000000000000000000", "00000000000000000000000000000000", "6e1b482b53761cf631819b749a6f3724" },
+			{ "fffff80000000000000000000000000000000000000000000000000000000000", "00000000000000000000000000000000", "976e6f851ab52c771998dbb2d71c75a9" },
+			{ "fffffc0000000000000000000000000000000000000000000000000000000000", "00000000000000000000000000000000", "85f2ba84f8c307cf525e124c3e22e6cc" },
+			{ "fffffe0000000000000000000000000000000000000000000000000000000000", "00000000000000000000000000000000", "6bcca98bf6a835fa64955f72de4115fe" },
+			{ "ffffff0000000000000000000000000000000000000000000000000000000000", "00000000000000000000000000000000", "2c75e2d36eebd65411f14fd0eb1d2a06" },
+			{ "ffffff8000000000000000000000000000000000000000000000000000000000", "00000000000000000000000000000000", "bd49295006250ffca5100b6007a0eade" },
+			{ "ffffffc000000000000000000000000000000000000000000000000000000000", "00000000000000000000000000000000", "a190527d0ef7c70f459cd3940df316ec" },
+			{ "ffffffe000000000000000000000000000000000000000000000000000000000", "00000000000000000000000000000000", "bbd1097a62433f79449fa97d4ee80dbf" },
+			{ "fffffff000000000000000000000000000000000000000000000000000000000", "00000000000000000000000000000000", "07058e408f5b99b0e0f061a1761b5b3b" },
+			{ "fffffff800000000000000000000000000000000000000000000000000000000", "00000000000000000000000000000000", "5fd1f13fa0f31e37fabde328f894eac2" },
+			{ "fffffffc00000000000000000000000000000000000000000000000000000000", "00000000000000000000000000000000", "fc4af7c948df26e2ef3e01c1ee5b8f6f" },
+			{ "fffffffe00000000000000000000000000000000000000000000000000000000", "00000000000000000000000000000000", "829fd7208fb92d44a074a677ee9861ac" },
+			{ "ffffffff00000000000000000000000000000000000000000000000000000000", "00000000000000000000000000000000", "ad9fc613a703251b54c64a0e76431711" },
+			{ "ffffffff80000000000000000000000000000000000000000000000000000000", "00000000000000000000000000000000", "33ac9eccc4cc75e2711618f80b1548e8" },
+			{ "ffffffffc0000000000000000000000000000000000000000000000000000000", "00000000000000000000000000000000", "2025c74b8ad8f4cda17ee2049c4c902d" },
+			{ "ffffffffe0000000000000000000000000000000000000000000000000000000", "00000000000000000000000000000000", "f85ca05fe528f1ce9b790166e8d551e7" },
+			{ "fffffffff0000000000000000000000000000000000000000000000000000000", "00000000000000000000000000000000", "6f6238d8966048d4967154e0dad5a6c9" },
+			{ "fffffffff8000000000000000000000000000000000000000000000000000000", "00000000000000000000000000000000", "f2b21b4e7640a9b3346de8b82fb41e49" },
+			{ "fffffffffc000000000000000000000000000000000000000000000000000000", "00000000000000000000000000000000", "f836f251ad1d11d49dc344628b1884e1" },
+			{ "fffffffffe000000000000000000000000000000000000000000000000000000", "00000000000000000000000000000000", "077e9470ae7abea5a9769d49182628c3" },
+			{ "ffffffffff000000000000000000000000000000000000000000000000000000", "00000000000000000000000000000000", "e0dcc2d27fc9865633f85223cf0d611f" },
+			{ "ffffffffff800000000000000000000000000000000000000000000000000000", "00000000000000000000000000000000", "be66cfea2fecd6bf0ec7b4352c99bcaa" },
+			{ "ffffffffffc00000000000000000000000000000000000000000000000000000", "00000000000000000000000000000000", "df31144f87a2ef523facdcf21a427804" },
+			{ "ffffffffffe00000000000000000000000000000000000000000000000000000", "00000000000000000000000000000000", "b5bb0f5629fb6aae5e1839a3c3625d63" },
+			{ "fffffffffff00000000000000000000000000000000000000000000000000000", "00000000000000000000000000000000", "3c9db3335306fe1ec612bdbfae6b6028" },
+			{ "fffffffffff80000000000000000000000000000000000000000000000000000", "00000000000000000000000000000000", "3dd5c34634a79d3cfcc8339760e6f5f4" },
+			{ "fffffffffffc0000000000000000000000000000000000000000000000000000", "00000000000000000000000000000000", "82bda118a3ed7af314fa2ccc5c07b761" },
+			{ "fffffffffffe0000000000000000000000000000000000000000000000000000", "00000000000000000000000000000000", "2937a64f7d4f46fe6fea3b349ec78e38" },
+			{ "ffffffffffff0000000000000000000000000000000000000000000000000000", "00000000000000000000000000000000", "225f068c28476605735ad671bb8f39f3" },
+			{ "ffffffffffff8000000000000000000000000000000000000000000000000000", "00000000000000000000000000000000", "ae682c5ecd71898e08942ac9aa89875c" },
+			{ "ffffffffffffc000000000000000000000000000000000000000000000000000", "00000000000000000000000000000000", "5e031cb9d676c3022d7f26227e85c38f" },
+			{ "ffffffffffffe000000000000000000000000000000000000000000000000000", "00000000000000000000000000000000", "a78463fb064db5d52bb64bfef64f2dda" },
+			{ "fffffffffffff000000000000000000000000000000000000000000000000000", "00000000000000000000000000000000", "8aa9b75e784593876c53a00eae5af52b" },
+			{ "fffffffffffff800000000000000000000000000000000000000000000000000", "00000000000000000000000000000000", "3f84566df23da48af692722fe980573a" },
+			{ "fffffffffffffc00000000000000000000000000000000000000000000000000", "00000000000000000000000000000000", "31690b5ed41c7eb42a1e83270a7ff0e6" },
+			{ "fffffffffffffe00000000000000000000000000000000000000000000000000", "00000000000000000000000000000000", "77dd7702646d55f08365e477d3590eda" },
+			{ "ffffffffffffff00000000000000000000000000000000000000000000000000", "00000000000000000000000000000000", "4c022ac62b3cb78d739cc67b3e20bb7e" },
+			{ "ffffffffffffff80000000000000000000000000000000000000000000000000", "00000000000000000000000000000000", "092fa137ce18b5dfe7906f550bb13370" },
+			{ "ffffffffffffffc0000000000000000000000000000000000000000000000000", "00000000000000000000000000000000", "3e0cdadf2e68353c0027672c97144dd3" },
+			{ "ffffffffffffffe0000000000000000000000000000000000000000000000000", "00000000000000000000000000000000", "d8c4b200b383fc1f2b2ea677618a1d27" },
+			{ "fffffffffffffff0000000000000000000000000000000000000000000000000", "00000000000000000000000000000000", "11825f99b0e9bb3477c1c0713b015aac" },
+			{ "fffffffffffffff8000000000000000000000000000000000000000000000000", "00000000000000000000000000000000", "f8b9fffb5c187f7ddc7ab10f4fb77576" },
+			{ "fffffffffffffffc000000000000000000000000000000000000000000000000", "00000000000000000000000000000000", "ffb4e87a32b37d6f2c8328d3b5377802" },
+			{ "fffffffffffffffe000000000000000000000000000000000000000000000000", "00000000000000000000000000000000", "d276c13a5d220f4da9224e74896391ce" },
+			{ "ffffffffffffffff000000000000000000000000000000000000000000000000", "00000000000000000000000000000000", "94efe7a0e2e031e2536da01df799c927" },
+			{ "ffffffffffffffff800000000000000000000000000000000000000000000000", "00000000000000000000000000000000", "8f8fd822680a85974e53a5a8eb9d38de" },
+			{ "ffffffffffffffffc00000000000000000000000000000000000000000000000", "00000000000000000000000000000000", "e0f0a91b2e45f8cc37b7805a3042588d" },
+			{ "ffffffffffffffffe00000000000000000000000000000000000000000000000", "00000000000000000000000000000000", "597a6252255e46d6364dbeeda31e279c" },
+			{ "fffffffffffffffff00000000000000000000000000000000000000000000000", "00000000000000000000000000000000", "f51a0f694442b8f05571797fec7ee8bf" },
+			{ "fffffffffffffffff80000000000000000000000000000000000000000000000", "00000000000000000000000000000000", "9ff071b165b5198a93dddeebc54d09b5" },
+			{ "fffffffffffffffffc0000000000000000000000000000000000000000000000", "00000000000000000000000000000000", "c20a19fd5758b0c4bc1a5df89cf73877" },
+			{ "fffffffffffffffffe0000000000000000000000000000000000000000000000", "00000000000000000000000000000000", "97120166307119ca2280e9315668e96f" },
+			{ "ffffffffffffffffff0000000000000000000000000000000000000000000000", "00000000000000000000000000000000", "4b3b9f1e099c2a09dc091e90e4f18f0a" },
+			{ "ffffffffffffffffff8000000000000000000000000000000000000000000000", "00000000000000000000000000000000", "eb040b891d4b37f6851f7ec219cd3f6d" },
+			{ "ffffffffffffffffffc000000000000000000000000000000000000000000000", "00000000000000000000000000000000", "9f0fdec08b7fd79aa39535bea42db92a" },
+			{ "ffffffffffffffffffe000000000000000000000000000000000000000000000", "00000000000000000000000000000000", "2e70f168fc74bf911df240bcd2cef236" },
+			{ "fffffffffffffffffff000000000000000000000000000000000000000000000", "00000000000000000000000000000000", "462ccd7f5fd1108dbc152f3cacad328b" },
+			{ "fffffffffffffffffff800000000000000000000000000000000000000000000", "00000000000000000000000000000000", "a4af534a7d0b643a01868785d86dfb95" },
+			{ "fffffffffffffffffffc00000000000000000000000000000000000000000000", "00000000000000000000000000000000", "ab980296197e1a5022326c31da4bf6f3" },
+			{ "fffffffffffffffffffe00000000000000000000000000000000000000000000", "00000000000000000000000000000000", "f97d57b3333b6281b07d486db2d4e20c" },
+			{ "ffffffffffffffffffff00000000000000000000000000000000000000000000", "00000000000000000000000000000000", "f33fa36720231afe4c759ade6bd62eb6" },
+			{ "ffffffffffffffffffff80000000000000000000000000000000000000000000", "00000000000000000000000000000000", "fdcfac0c02ca538343c68117e0a15938" },
+			{ "ffffffffffffffffffffc0000000000000000000000000000000000000000000", "00000000000000000000000000000000", "ad4916f5ee5772be764fc027b8a6e539" },
+			{ "ffffffffffffffffffffe0000000000000000000000000000000000000000000", "00000000000000000000000000000000", "2e16873e1678610d7e14c02d002ea845" },
+			{ "fffffffffffffffffffff0000000000000000000000000000000000000000000", "00000000000000000000000000000000", "4e6e627c1acc51340053a8236d579576" },
+			{ "fffffffffffffffffffff8000000000000000000000000000000000000000000", "00000000000000000000000000000000", "ab0c8410aeeead92feec1eb430d652cb" },
+			{ "fffffffffffffffffffffc000000000000000000000000000000000000000000", "00000000000000000000000000000000", "e86f7e23e835e114977f60e1a592202e" },
+			{ "fffffffffffffffffffffe000000000000000000000000000000000000000000", "00000000000000000000000000000000", "e68ad5055a367041fade09d9a70a794b" },
+			{ "ffffffffffffffffffffff000000000000000000000000000000000000000000", "00000000000000000000000000000000", "0791823a3c666bb6162825e78606a7fe" },
+			{ "ffffffffffffffffffffff800000000000000000000000000000000000000000", "00000000000000000000000000000000", "dcca366a9bf47b7b868b77e25c18a364" },
+			{ "ffffffffffffffffffffffc00000000000000000000000000000000000000000", "00000000000000000000000000000000", "684c9efc237e4a442965f84bce20247a" },
+			{ "ffffffffffffffffffffffe00000000000000000000000000000000000000000", "00000000000000000000000000000000", "a858411ffbe63fdb9c8aa1bfaed67b52" },
+			{ "fffffffffffffffffffffff00000000000000000000000000000000000000000", "00000000000000000000000000000000", "04bc3da2179c3015498b0e03910db5b8" },
+			{ "fffffffffffffffffffffff80000000000000000000000000000000000000000", "00000000000000000000000000000000", "40071eeab3f935dbc25d00841460260f" },
+			{ "fffffffffffffffffffffffc0000000000000000000000000000000000000000", "00000000000000000000000000000000", "0ebd7c30ed2016e08ba806ddb008bcc8" },
+			{ "fffffffffffffffffffffffe0000000000000000000000000000000000000000", "00000000000000000000000000000000", "15c6becf0f4cec7129cbd22d1a79b1b8" },
+			{ "ffffffffffffffffffffffff0000000000000000000000000000000000000000", "00000000000000000000000000000000", "0aeede5b91f721700e9e62edbf60b781" },
+			{ "ffffffffffffffffffffffff8000000000000000000000000000000000000000", "00000000000000000000000000000000", "266581af0dcfbed1585e0a242c64b8df" },
+			{ "ffffffffffffffffffffffffc000000000000000000000000000000000000000", "00000000000000000000000000000000", "6693dc911662ae473216ba22189a511a" },
+			{ "ffffffffffffffffffffffffe000000000000000000000000000000000000000", "00000000000000000000000000000000", "7606fa36d86473e6fb3a1bb0e2c0adf5" },
+			{ "fffffffffffffffffffffffff000000000000000000000000000000000000000", "00000000000000000000000000000000", "112078e9e11fbb78e26ffb8899e96b9a" },
+			{ "fffffffffffffffffffffffff800000000000000000000000000000000000000", "00000000000000000000000000000000", "40b264e921e9e4a82694589ef3798262" },
+			{ "fffffffffffffffffffffffffc00000000000000000000000000000000000000", "00000000000000000000000000000000", "8d4595cb4fa7026715f55bd68e2882f9" },
+			{ "fffffffffffffffffffffffffe00000000000000000000000000000000000000", "00000000000000000000000000000000", "b588a302bdbc09197df1edae68926ed9" },
+			{ "ffffffffffffffffffffffffff00000000000000000000000000000000000000", "00000000000000000000000000000000", "33f7502390b8a4a221cfecd0666624ba" },
+			{ "ffffffffffffffffffffffffff80000000000000000000000000000000000000", "00000000000000000000000000000000", "3d20253adbce3be2373767c4d822c566" },
+			{ "ffffffffffffffffffffffffffc0000000000000000000000000000000000000", "00000000000000000000000000000000", "a42734a3929bf84cf0116c9856a3c18c" },
+			{ "ffffffffffffffffffffffffffe0000000000000000000000000000000000000", "00000000000000000000000000000000", "e3abc4939457422bb957da3c56938c6d" },
+			{ "fffffffffffffffffffffffffff0000000000000000000000000000000000000", "00000000000000000000000000000000", "972bdd2e7c525130fadc8f76fc6f4b3f" },
+			{ "fffffffffffffffffffffffffff8000000000000000000000000000000000000", "00000000000000000000000000000000", "84a83d7b94c699cbcb8a7d9b61f64093" },
+			{ "fffffffffffffffffffffffffffc000000000000000000000000000000000000", "00000000000000000000000000000000", "ce61d63514aded03d43e6ebfc3a9001f" },
+			{ "fffffffffffffffffffffffffffe000000000000000000000000000000000000", "00000000000000000000000000000000", "6c839dd58eeae6b8a36af48ed63d2dc9" },
+			{ "ffffffffffffffffffffffffffff000000000000000000000000000000000000", "00000000000000000000000000000000", "cd5ece55b8da3bf622c4100df5de46f9" },
+			{ "ffffffffffffffffffffffffffff800000000000000000000000000000000000", "00000000000000000000000000000000", "3b6f46f40e0ac5fc0a9c1105f800f48d" },
+			{ "ffffffffffffffffffffffffffffc00000000000000000000000000000000000", "00000000000000000000000000000000", "ba26d47da3aeb028de4fb5b3a854a24b" },
+			{ "ffffffffffffffffffffffffffffe00000000000000000000000000000000000", "00000000000000000000000000000000", "87f53bf620d3677268445212904389d5" },
+			{ "fffffffffffffffffffffffffffff00000000000000000000000000000000000", "00000000000000000000000000000000", "10617d28b5e0f4605492b182a5d7f9f6" },
+			{ "fffffffffffffffffffffffffffff80000000000000000000000000000000000", "00000000000000000000000000000000", "9aaec4fabbf6fae2a71feff02e372b39" },
+			{ "fffffffffffffffffffffffffffffc0000000000000000000000000000000000", "00000000000000000000000000000000", "3a90c62d88b5c42809abf782488ed130" },
+			{ "fffffffffffffffffffffffffffffe0000000000000000000000000000000000", "00000000000000000000000000000000", "f1f1c5a40899e15772857ccb65c7a09a" },
+			{ "ffffffffffffffffffffffffffffff0000000000000000000000000000000000", "00000000000000000000000000000000", "190843d29b25a3897c692ce1dd81ee52" },
+			{ "ffffffffffffffffffffffffffffff8000000000000000000000000000000000", "00000000000000000000000000000000", "a866bc65b6941d86e8420a7ffb0964db" },
+			{ "ffffffffffffffffffffffffffffffc000000000000000000000000000000000", "00000000000000000000000000000000", "8193c6ff85225ced4255e92f6e078a14" },
+			{ "ffffffffffffffffffffffffffffffe000000000000000000000000000000000", "00000000000000000000000000000000", "9661cb2424d7d4a380d547f9e7ec1cb9" },
+			{ "fffffffffffffffffffffffffffffff000000000000000000000000000000000", "00000000000000000000000000000000", "86f93d9ec08453a071e2e2877877a9c8" },
+			{ "fffffffffffffffffffffffffffffff800000000000000000000000000000000", "00000000000000000000000000000000", "27eefa80ce6a4a9d598e3fec365434d2" },
+			{ "fffffffffffffffffffffffffffffffc00000000000000000000000000000000", "00000000000000000000000000000000", "d62068444578e3ab39ce7ec95dd045dc" },
+			{ "fffffffffffffffffffffffffffffffe00000000000000000000000000000000", "00000000000000000000000000000000", "b5f71d4dd9a71fe5d8bc8ba7e6ea3048" },
+			{ "ffffffffffffffffffffffffffffffff00000000000000000000000000000000", "00000000000000000000000000000000", "6825a347ac479d4f9d95c5cb8d3fd7e9" },
+			{ "ffffffffffffffffffffffffffffffff80000000000000000000000000000000", "00000000000000000000000000000000", "e3714e94a5778955cc0346358e94783a" },
+			{ "ffffffffffffffffffffffffffffffffc0000000000000000000000000000000", "00000000000000000000000000000000", "d836b44bb29e0c7d89fa4b2d4b677d2a" },
+			{ "ffffffffffffffffffffffffffffffffe0000000000000000000000000000000", "00000000000000000000000000000000", "5d454b75021d76d4b84f873a8f877b92" },
+			{ "fffffffffffffffffffffffffffffffff0000000000000000000000000000000", "00000000000000000000000000000000", "c3498f7eced2095314fc28115885b33f" },
+			{ "fffffffffffffffffffffffffffffffff8000000000000000000000000000000", "00000000000000000000000000000000", "6e668856539ad8e405bd123fe6c88530" },
+			{ "fffffffffffffffffffffffffffffffffc000000000000000000000000000000", "00000000000000000000000000000000", "8680db7f3a87b8605543cfdbe6754076" },
+			{ "fffffffffffffffffffffffffffffffffe000000000000000000000000000000", "00000000000000000000000000000000", "6c5d03b13069c3658b3179be91b0800c" },
+			{ "ffffffffffffffffffffffffffffffffff000000000000000000000000000000", "00000000000000000000000000000000", "ef1b384ac4d93eda00c92add0995ea5f" },
+			{ "ffffffffffffffffffffffffffffffffff800000000000000000000000000000", "00000000000000000000000000000000", "bf8115805471741bd5ad20a03944790f" },
+			{ "ffffffffffffffffffffffffffffffffffc00000000000000000000000000000", "00000000000000000000000000000000", "c64c24b6894b038b3c0d09b1df068b0b" },
+			{ "ffffffffffffffffffffffffffffffffffe00000000000000000000000000000", "00000000000000000000000000000000", "3967a10cffe27d0178545fbf6a40544b" },
+			{ "fffffffffffffffffffffffffffffffffff00000000000000000000000000000", "00000000000000000000000000000000", "7c85e9c95de1a9ec5a5363a8a053472d" },
+			{ "fffffffffffffffffffffffffffffffffff80000000000000000000000000000", "00000000000000000000000000000000", "a9eec03c8abec7ba68315c2c8c2316e0" },
+			{ "fffffffffffffffffffffffffffffffffffc0000000000000000000000000000", "00000000000000000000000000000000", "cac8e414c2f388227ae14986fc983524" },
+			{ "fffffffffffffffffffffffffffffffffffe0000000000000000000000000000", "00000000000000000000000000000000", "5d942b7f4622ce056c3ce3ce5f1dd9d6" },
+			{ "ffffffffffffffffffffffffffffffffffff0000000000000000000000000000", "00000000000000000000000000000000", "d240d648ce21a3020282c3f1b528a0b6" },
+			{ "ffffffffffffffffffffffffffffffffffff8000000000000000000000000000", "00000000000000000000000000000000", "45d089c36d5c5a4efc689e3b0de10dd5" },
+			{ "ffffffffffffffffffffffffffffffffffffc000000000000000000000000000", "00000000000000000000000000000000", "b4da5df4becb5462e03a0ed00d295629" },
+			{ "ffffffffffffffffffffffffffffffffffffe000000000000000000000000000", "00000000000000000000000000000000", "dcf4e129136c1a4b7a0f38935cc34b2b" },
+			{ "fffffffffffffffffffffffffffffffffffff000000000000000000000000000", "00000000000000000000000000000000", "d9a4c7618b0ce48a3d5aee1a1c0114c4" },
+			{ "fffffffffffffffffffffffffffffffffffff800000000000000000000000000", "00000000000000000000000000000000", "ca352df025c65c7b0bf306fbee0f36ba" },
+			{ "fffffffffffffffffffffffffffffffffffffc00000000000000000000000000", "00000000000000000000000000000000", "238aca23fd3409f38af63378ed2f5473" },
+			{ "fffffffffffffffffffffffffffffffffffffe00000000000000000000000000", "00000000000000000000000000000000", "59836a0e06a79691b36667d5380d8188" },
+			{ "ffffffffffffffffffffffffffffffffffffff00000000000000000000000000", "00000000000000000000000000000000", "33905080f7acf1cdae0a91fc3e85aee4" },
+			{ "ffffffffffffffffffffffffffffffffffffff80000000000000000000000000", "00000000000000000000000000000000", "72c9e4646dbc3d6320fc6689d93e8833" },
+			{ "ffffffffffffffffffffffffffffffffffffffc0000000000000000000000000", "00000000000000000000000000000000", "ba77413dea5925b7f5417ea47ff19f59" },
+			{ "ffffffffffffffffffffffffffffffffffffffe0000000000000000000000000", "00000000000000000000000000000000", "6cae8129f843d86dc786a0fb1a184970" },
+			{ "fffffffffffffffffffffffffffffffffffffff0000000000000000000000000", "00000000000000000000000000000000", "fcfefb534100796eebbd990206754e19" },
+			{ "fffffffffffffffffffffffffffffffffffffff8000000000000000000000000", "00000000000000000000000000000000", "8c791d5fdddf470da04f3e6dc4a5b5b5" },
+			{ "fffffffffffffffffffffffffffffffffffffffc000000000000000000000000", "00000000000000000000000000000000", "c93bbdc07a4611ae4bb266ea5034a387" },
+			{ "fffffffffffffffffffffffffffffffffffffffe000000000000000000000000", "00000000000000000000000000000000", "c102e38e489aa74762f3efc5bb23205a" },
+			{ "ffffffffffffffffffffffffffffffffffffffff000000000000000000000000", "00000000000000000000000000000000", "93201481665cbafc1fcc220bc545fb3d" },
+			{ "ffffffffffffffffffffffffffffffffffffffff800000000000000000000000", "00000000000000000000000000000000", "4960757ec6ce68cf195e454cfd0f32ca" },
+			{ "ffffffffffffffffffffffffffffffffffffffffc00000000000000000000000", "00000000000000000000000000000000", "feec7ce6a6cbd07c043416737f1bbb33" },
+			{ "ffffffffffffffffffffffffffffffffffffffffe00000000000000000000000", "00000000000000000000000000000000", "11c5413904487a805d70a8edd9c35527" },
+			{ "fffffffffffffffffffffffffffffffffffffffff00000000000000000000000", "00000000000000000000000000000000", "347846b2b2e36f1f0324c86f7f1b98e2" },
+			{ "fffffffffffffffffffffffffffffffffffffffff80000000000000000000000", "00000000000000000000000000000000", "332eee1a0cbd19ca2d69b426894044f0" },
+			{ "fffffffffffffffffffffffffffffffffffffffffc0000000000000000000000", "00000000000000000000000000000000", "866b5b3977ba6efa5128efbda9ff03cd" },
+			{ "fffffffffffffffffffffffffffffffffffffffffe0000000000000000000000", "00000000000000000000000000000000", "cc1445ee94c0f08cdee5c344ecd1e233" },
+			{ "ffffffffffffffffffffffffffffffffffffffffff0000000000000000000000", "00000000000000000000000000000000", "be288319029363c2622feba4b05dfdfe" },
+			{ "ffffffffffffffffffffffffffffffffffffffffff8000000000000000000000", "00000000000000000000000000000000", "cfd1875523f3cd21c395651e6ee15e56" },
+			{ "ffffffffffffffffffffffffffffffffffffffffffc000000000000000000000", "00000000000000000000000000000000", "cb5a408657837c53bf16f9d8465dce19" },
+			{ "ffffffffffffffffffffffffffffffffffffffffffe000000000000000000000", "00000000000000000000000000000000", "ca0bf42cb107f55ccff2fc09ee08ca15" },
+			{ "fffffffffffffffffffffffffffffffffffffffffff000000000000000000000", "00000000000000000000000000000000", "fdd9bbb4a7dc2e4a23536a5880a2db67" },
+			{ "fffffffffffffffffffffffffffffffffffffffffff800000000000000000000", "00000000000000000000000000000000", "ede447b362c484993dec9442a3b46aef" },
+			{ "fffffffffffffffffffffffffffffffffffffffffffc00000000000000000000", "00000000000000000000000000000000", "10dffb05904bff7c4781df780ad26837" },
+			{ "fffffffffffffffffffffffffffffffffffffffffffe00000000000000000000", "00000000000000000000000000000000", "c33bc13e8de88ac25232aa7496398783" },
+			{ "ffffffffffffffffffffffffffffffffffffffffffff00000000000000000000", "00000000000000000000000000000000", "ca359c70803a3b2a3d542e8781dea975" },
+			{ "ffffffffffffffffffffffffffffffffffffffffffff80000000000000000000", "00000000000000000000000000000000", "bcc65b526f88d05b89ce8a52021fdb06" },
+			{ "ffffffffffffffffffffffffffffffffffffffffffffc0000000000000000000", "00000000000000000000000000000000", "db91a38855c8c4643851fbfb358b0109" },
+			{ "ffffffffffffffffffffffffffffffffffffffffffffe0000000000000000000", "00000000000000000000000000000000", "ca6e8893a114ae8e27d5ab03a5499610" },
+			{ "fffffffffffffffffffffffffffffffffffffffffffff0000000000000000000", "00000000000000000000000000000000", "6629d2b8df97da728cdd8b1e7f945077" },
+			{ "fffffffffffffffffffffffffffffffffffffffffffff8000000000000000000", "00000000000000000000000000000000", "4570a5a18cfc0dd582f1d88d5c9a1720" },
+			{ "fffffffffffffffffffffffffffffffffffffffffffffc000000000000000000", "00000000000000000000000000000000", "72bc65aa8e89562e3f274d45af1cd10b" },
+			{ "fffffffffffffffffffffffffffffffffffffffffffffe000000000000000000", "00000000000000000000000000000000", "98551da1a6503276ae1c77625f9ea615" },
+			{ "ffffffffffffffffffffffffffffffffffffffffffffff000000000000000000", "00000000000000000000000000000000", "0ddfe51ced7e3f4ae927daa3fe452cee" },
+			{ "ffffffffffffffffffffffffffffffffffffffffffffff800000000000000000", "00000000000000000000000000000000", "db826251e4ce384b80218b0e1da1dd4c" },
+			{ "ffffffffffffffffffffffffffffffffffffffffffffffc00000000000000000", "00000000000000000000000000000000", "2cacf728b88abbad7011ed0e64a1680c" },
+			{ "ffffffffffffffffffffffffffffffffffffffffffffffe00000000000000000", "00000000000000000000000000000000", "330d8ee7c5677e099ac74c9994ee4cfb" },
+			{ "fffffffffffffffffffffffffffffffffffffffffffffff00000000000000000", "00000000000000000000000000000000", "edf61ae362e882ddc0167474a7a77f3a" },
+			{ "fffffffffffffffffffffffffffffffffffffffffffffff80000000000000000", "00000000000000000000000000000000", "6168b00ba7859e0970ecfd757efecf7c" },
+			{ "fffffffffffffffffffffffffffffffffffffffffffffffc0000000000000000", "00000000000000000000000000000000", "d1415447866230d28bb1ea18a4cdfd02" },
+			{ "fffffffffffffffffffffffffffffffffffffffffffffffe0000000000000000", "00000000000000000000000000000000", "516183392f7a8763afec68a060264141" },
+			{ "ffffffffffffffffffffffffffffffffffffffffffffffff0000000000000000", "00000000000000000000000000000000", "77565c8d73cfd4130b4aa14d8911710f" },
+			{ "ffffffffffffffffffffffffffffffffffffffffffffffff8000000000000000", "00000000000000000000000000000000", "37232a4ed21ccc27c19c9610078cabac" },
+			{ "ffffffffffffffffffffffffffffffffffffffffffffffffc000000000000000", "00000000000000000000000000000000", "804f32ea71828c7d329077e712231666" },
+			{ "ffffffffffffffffffffffffffffffffffffffffffffffffe000000000000000", "00000000000000000000000000000000", "d64424f23cb97215e9c2c6f28d29eab7" },
+			{ "fffffffffffffffffffffffffffffffffffffffffffffffff000000000000000", "00000000000000000000000000000000", "023e82b533f68c75c238cebdb2ee89a2" },
+			{ "fffffffffffffffffffffffffffffffffffffffffffffffff800000000000000", "00000000000000000000000000000000", "193a3d24157a51f1ee0893f6777417e7" },
+			{ "fffffffffffffffffffffffffffffffffffffffffffffffffc00000000000000", "00000000000000000000000000000000", "84ecacfcd400084d078612b1945f2ef5" },
+			{ "fffffffffffffffffffffffffffffffffffffffffffffffffe00000000000000", "00000000000000000000000000000000", "1dcd8bb173259eb33a5242b0de31a455" },
+			{ "ffffffffffffffffffffffffffffffffffffffffffffffffff00000000000000", "00000000000000000000000000000000", "35e9eddbc375e792c19992c19165012b" },
+			{ "ffffffffffffffffffffffffffffffffffffffffffffffffff80000000000000", "00000000000000000000000000000000", "8a772231c01dfdd7c98e4cfddcc0807a" },
+			{ "ffffffffffffffffffffffffffffffffffffffffffffffffffc0000000000000", "00000000000000000000000000000000", "6eda7ff6b8319180ff0d6e65629d01c3" },
+			{ "ffffffffffffffffffffffffffffffffffffffffffffffffffe0000000000000", "00000000000000000000000000000000", "c267ef0e2d01a993944dd397101413cb" },
+			{ "fffffffffffffffffffffffffffffffffffffffffffffffffff0000000000000", "00000000000000000000000000000000", "e9f80e9d845bcc0f62926af72eabca39" },
+			{ "fffffffffffffffffffffffffffffffffffffffffffffffffff8000000000000", "00000000000000000000000000000000", "6702990727aa0878637b45dcd3a3b074" },
+			{ "fffffffffffffffffffffffffffffffffffffffffffffffffffc000000000000", "00000000000000000000000000000000", "2e2e647d5360e09230a5d738ca33471e" },
+			{ "fffffffffffffffffffffffffffffffffffffffffffffffffffe000000000000", "00000000000000000000000000000000", "1f56413c7add6f43d1d56e4f02190330" },
+			{ "ffffffffffffffffffffffffffffffffffffffffffffffffffff000000000000", "00000000000000000000000000000000", "69cd0606e15af729d6bca143016d9842" },
+			{ "ffffffffffffffffffffffffffffffffffffffffffffffffffff800000000000", "00000000000000000000000000000000", "a085d7c1a500873a20099c4caa3c3f5b" },
+			{ "ffffffffffffffffffffffffffffffffffffffffffffffffffffc00000000000", "00000000000000000000000000000000", "4fc0d230f8891415b87b83f95f2e09d1" },
+			{ "ffffffffffffffffffffffffffffffffffffffffffffffffffffe00000000000", "00000000000000000000000000000000", "4327d08c523d8eba697a4336507d1f42" },
+			{ "fffffffffffffffffffffffffffffffffffffffffffffffffffff00000000000", "00000000000000000000000000000000", "7a15aab82701efa5ae36ab1d6b76290f" },
+			{ "fffffffffffffffffffffffffffffffffffffffffffffffffffff80000000000", "00000000000000000000000000000000", "5bf0051893a18bb30e139a58fed0fa54" },
+			{ "fffffffffffffffffffffffffffffffffffffffffffffffffffffc0000000000", "00000000000000000000000000000000", "97e8adf65638fd9cdf3bc22c17fe4dbd" },
+			{ "fffffffffffffffffffffffffffffffffffffffffffffffffffffe0000000000", "00000000000000000000000000000000", "1ee6ee326583a0586491c96418d1a35d" },
+			{ "ffffffffffffffffffffffffffffffffffffffffffffffffffffff0000000000", "00000000000000000000000000000000", "26b549c2ec756f82ecc48008e529956b" },
+			{ "ffffffffffffffffffffffffffffffffffffffffffffffffffffff8000000000", "00000000000000000000000000000000", "70377b6da669b072129e057cc28e9ca5" },
+			{ "ffffffffffffffffffffffffffffffffffffffffffffffffffffffc000000000", "00000000000000000000000000000000", "9c94b8b0cb8bcc919072262b3fa05ad9" },
+			{ "ffffffffffffffffffffffffffffffffffffffffffffffffffffffe000000000", "00000000000000000000000000000000", "2fbb83dfd0d7abcb05cd28cad2dfb523" },
+			{ "fffffffffffffffffffffffffffffffffffffffffffffffffffffff000000000", "00000000000000000000000000000000", "96877803de77744bb970d0a91f4debae" },
+			{ "fffffffffffffffffffffffffffffffffffffffffffffffffffffff800000000", "00000000000000000000000000000000", "7379f3370cf6e5ce12ae5969c8eea312" },
+			{ "fffffffffffffffffffffffffffffffffffffffffffffffffffffffc00000000", "00000000000000000000000000000000", "02dc99fa3d4f98ce80985e7233889313" },
+			{ "fffffffffffffffffffffffffffffffffffffffffffffffffffffffe00000000", "00000000000000000000000000000000", "1e38e759075ba5cab6457da51844295a" },
+			{ "ffffffffffffffffffffffffffffffffffffffffffffffffffffffff00000000", "00000000000000000000000000000000", "70bed8dbf615868a1f9d9b05d3e7a267" },
+			{ "ffffffffffffffffffffffffffffffffffffffffffffffffffffffff80000000", "00000000000000000000000000000000", "234b148b8cb1d8c32b287e896903d150" },
+			{ "ffffffffffffffffffffffffffffffffffffffffffffffffffffffffc0000000", "00000000000000000000000000000000", "294b033df4da853f4be3e243f7e513f4" },
+			{ "ffffffffffffffffffffffffffffffffffffffffffffffffffffffffe0000000", "00000000000000000000000000000000", "3f58c950f0367160adec45f2441e7411" },
+			{ "fffffffffffffffffffffffffffffffffffffffffffffffffffffffff0000000", "00000000000000000000000000000000", "37f655536a704e5ace182d742a820cf4" },
+			{ "fffffffffffffffffffffffffffffffffffffffffffffffffffffffff8000000", "00000000000000000000000000000000", "ea7bd6bb63418731aeac790fe42d61e8" },
+			{ "fffffffffffffffffffffffffffffffffffffffffffffffffffffffffc000000", "00000000000000000000000000000000", "e74a4c999b4c064e48bb1e413f51e5ea" },
+			{ "fffffffffffffffffffffffffffffffffffffffffffffffffffffffffe000000", "00000000000000000000000000000000", "ba9ebefdb4ccf30f296cecb3bc1943e8" },
+			{ "ffffffffffffffffffffffffffffffffffffffffffffffffffffffffff000000", "00000000000000000000000000000000", "3194367a4898c502c13bb7478640a72d" },
+			{ "ffffffffffffffffffffffffffffffffffffffffffffffffffffffffff800000", "00000000000000000000000000000000", "da797713263d6f33a5478a65ef60d412" },
+			{ "ffffffffffffffffffffffffffffffffffffffffffffffffffffffffffc00000", "00000000000000000000000000000000", "d1ac39bb1ef86b9c1344f214679aa376" },
+			{ "ffffffffffffffffffffffffffffffffffffffffffffffffffffffffffe00000", "00000000000000000000000000000000", "2fdea9e650532be5bc0e7325337fd363" },
+			{ "fffffffffffffffffffffffffffffffffffffffffffffffffffffffffff00000", "00000000000000000000000000000000", "d3a204dbd9c2af158b6ca67a5156ce4a" },
+			{ "fffffffffffffffffffffffffffffffffffffffffffffffffffffffffff80000", "00000000000000000000000000000000", "3a0a0e75a8da36735aee6684d965a778" },
+			{ "fffffffffffffffffffffffffffffffffffffffffffffffffffffffffffc0000", "00000000000000000000000000000000", "52fc3e620492ea99641ea168da5b6d52" },
+			{ "fffffffffffffffffffffffffffffffffffffffffffffffffffffffffffe0000", "00000000000000000000000000000000", "d2e0c7f15b4772467d2cfc873000b2ca" },
+			{ "ffffffffffffffffffffffffffffffffffffffffffffffffffffffffffff0000", "00000000000000000000000000000000", "563531135e0c4d70a38f8bdb190ba04e" },
+			{ "ffffffffffffffffffffffffffffffffffffffffffffffffffffffffffff8000", "00000000000000000000000000000000", "a8a39a0f5663f4c0fe5f2d3cafff421a" },
+			{ "ffffffffffffffffffffffffffffffffffffffffffffffffffffffffffffc000", "00000000000000000000000000000000", "d94b5e90db354c1e42f61fabe167b2c0" },
+			{ "ffffffffffffffffffffffffffffffffffffffffffffffffffffffffffffe000", "00000000000000000000000000000000", "50e6d3c9b6698a7cd276f96b1473f35a" },
+			{ "fffffffffffffffffffffffffffffffffffffffffffffffffffffffffffff000", "00000000000000000000000000000000", "9338f08e0ebee96905d8f2e825208f43" },
+			{ "fffffffffffffffffffffffffffffffffffffffffffffffffffffffffffff800", "00000000000000000000000000000000", "8b378c86672aa54a3a266ba19d2580ca" },
+			{ "fffffffffffffffffffffffffffffffffffffffffffffffffffffffffffffc00", "00000000000000000000000000000000", "cca7c3086f5f9511b31233da7cab9160" },
+			{ "fffffffffffffffffffffffffffffffffffffffffffffffffffffffffffffe00", "00000000000000000000000000000000", "5b40ff4ec9be536ba23035fa4f06064c" },
+			{ "ffffffffffffffffffffffffffffffffffffffffffffffffffffffffffffff00", "00000000000000000000000000000000", "60eb5af8416b257149372194e8b88749" },
+			{ "ffffffffffffffffffffffffffffffffffffffffffffffffffffffffffffff80", "00000000000000000000000000000000", "2f005a8aed8a361c92e440c15520cbd1" },
+			{ "ffffffffffffffffffffffffffffffffffffffffffffffffffffffffffffffc0", "00000000000000000000000000000000", "7b03627611678a997717578807a800e2" },
+			{ "ffffffffffffffffffffffffffffffffffffffffffffffffffffffffffffffe0", "00000000000000000000000000000000", "cf78618f74f6f3696e0a4779b90b5a77" },
+			{ "fffffffffffffffffffffffffffffffffffffffffffffffffffffffffffffff0", "00000000000000000000000000000000", "03720371a04962eaea0a852e69972858" },
+			{ "fffffffffffffffffffffffffffffffffffffffffffffffffffffffffffffff8", "00000000000000000000000000000000", "1f8a8133aa8ccf70e2bd3285831ca6b7" },
+			{ "fffffffffffffffffffffffffffffffffffffffffffffffffffffffffffffffc", "00000000000000000000000000000000", "27936bd27fb1468fc8b48bc483321725" },
+			{ "fffffffffffffffffffffffffffffffffffffffffffffffffffffffffffffffe", "00000000000000000000000000000000", "b07d4f3e2cd2ef2eb545980754dfea0f" },
+			{ "ffffffffffffffffffffffffffffffffffffffffffffffffffffffffffffffff", "00000000000000000000000000000000", "4bf85f1b5d54adbc307b0a048389adcb" },
+		};
+		for (auto& Test : VarKeyCases)
+		{
+			RunTest(Test.Key, Test.PlainText, Test.CipherText);
+		}
+	}
+
+	// 6.4 Monte Carlo Test, from ECBMCT256.rsp
+	SECTION("Monte Carlo")
+	{
+		static const AesTestCase MonteCarloCases[] =
+		{
+			{ "f9e8389f5b80712e3886cc1fa2d28a3b8c9cd88a2d4a54c6aa86ce0fef944be0", "b379777f9050e2a818f2940cbbd9aba4", "6893ebaf0a1fccc704326529fdfb60db" },
+			{ "db9ea5a2284fa17fb63e13bf891c8e42e40f332527559801aeb4ab26126f2b3b", "6893ebaf0a1fccc704326529fdfb60db", "f3c78a5e85e5439bf26d5818718157d6" },
+			{ "7099ed88e82744228a5303ae2ef6c0d017c8b97ba2b0db9a5cd9f33e63ee7ced", "f3c78a5e85e5439bf26d5818718157d6", "2326b958b00b3050697eedb08cc20504" },
+			{ "5e9e65ea96e78dd4fb78ea1184f6ebde34ee002312bbebca35a71e8eef2c79e9", "2326b958b00b3050697eedb08cc20504", "ec4332d5e3cebd3e0f5fc51452f4560d" },
+			{ "33acf1cafc822646dc869e905bd26f9ad8ad32f6f17556f43af8db9abdd82fe4", "ec4332d5e3cebd3e0f5fc51452f4560d", "5da58b5ef2076340d555f861c3449a77" },
+			{ "eb0ae85c1b44d5db4729d268f49be2a08508b9a8037235b4efad23fb7e9cb593", "5da58b5ef2076340d555f861c3449a77", "307d50c18a0b6a08402ff131d72cb7ec" },
+			{ "fac93b561a9b6a0e809d71ecdb980afab575e96989795fbcaf82d2caa9b0027f", "307d50c18a0b6a08402ff131d72cb7ec", "92c34165a2963e77e05e2d6fc2d931d5" },
+			{ "a0559e41d58af36174a67246df87541b27b6a80c2bef61cb4fdcffa56b6933aa", "92c34165a2963e77e05e2d6fc2d931d5", "cb33d519a1fdb1d5fbb185c47870c1ed" },
+			{ "e48824d6c2251d3a27f38fb543c31fc1ec857d158a12d01eb46d7a611319f247", "cb33d519a1fdb1d5fbb185c47870c1ed", "78fb452f384c8f870e572890588f3728" },
+			{ "7a33440ad7c69d583355c745e5c88c47947e383ab25e5f99ba3a52f14b96c56f", "78fb452f384c8f870e572890588f3728", "12375e02a8bbc84b00feaab54a66db43" },
+			{ "0e6877c7fdc234efb9afcd96b4ebdb83864966381ae597d2bac4f84401f01e2c", "12375e02a8bbc84b00feaab54a66db43", "eab1606610b55c857f2b4bf1cf3feba0" },
+			{ "ec91d3550c79ab4914a26987725ab1396cf8065e0a50cb57c5efb3b5cecff58c", "eab1606610b55c857f2b4bf1cf3feba0", "6c73381147de97961cc26ad26602a45a" },
+			{ "83caed5a49579b3a55a71e5ece5966e5008b3e4f4d8e5cc1d92dd967a8cd51d6", "6c73381147de97961cc26ad26602a45a", "e76c08fd29bf015352003c636fee5ff9" },
+			{ "91b9b50908968361dcd8f4ba236fa199e7e736b264315d928b2de504c7230e2f", "e76c08fd29bf015352003c636fee5ff9", "6f26b8191a2b059dcdeb3dbabc437c29" },
+			{ "94fc46213c870f7965b88773afe93b1388c18eab7e1a580f46c6d8be7b607206", "6f26b8191a2b059dcdeb3dbabc437c29", "e91e2fcef14dd4251caec97c45223fef" },
+			{ "b0deff009aff61f65763b0b9fdd39a9061dfa1658f578c2a5a6811c23e424de9", "e91e2fcef14dd4251caec97c45223fef", "98fb2122912360f07916e4802c0ea1e5" },
+			{ "968f3c88d27f1be8decb00c4d464d369f92480471e74ecda237ef542124cec0c", "98fb2122912360f07916e4802c0ea1e5", "b4a15c59976b39da50da8ed393f27a62" },
+			{ "35fe06a60309581565a97232140668464d85dc1e891fd50073a47b9181be966e", "b4a15c59976b39da50da8ed393f27a62", "d8b57d7a72ef92409c51d40bb8c4cbc5" },
+			{ "4739043d7750bcf4a6f269a3d54083ca9530a164fbf04740eff5af9a397a5dab", "d8b57d7a72ef92409c51d40bb8c4cbc5", "548cc893e80caf5a601c2381517f8c5b" },
+			{ "4b065b5195f2ddf6f5d0aed72ff7a1e3c1bc69f713fce81a8fe98c1b6805d1f0", "548cc893e80caf5a601c2381517f8c5b", "1c2238c560d678d40b48cc8034add0c4" },
+			{ "1a949129e14c5963d997c86a6352ea53dd9e5132732a90ce84a1409b5ca80134", "1c2238c560d678d40b48cc8034add0c4", "2771ff806e061df8ad4aa877717bf309" },
+			{ "91cd3e48f4b42c432bed9848583e6dc7faefaeb21d2c8d3629ebe8ec2dd3f23d", "2771ff806e061df8ad4aa877717bf309", "065593fa1fcdb481bb27f334505543f5" },
+			{ "42d6f7b585a0a0d356c59be3d07d4d41fcba3d4802e139b792cc1bd87d86b1c8", "065593fa1fcdb481bb27f334505543f5", "d3f309c0039b15d14eb8b739a94b94fa" },
+			{ "72900f0ecbf90fa058805deb430815072f493488017a2c66dc74ace1d4cd2532", "d3f309c0039b15d14eb8b739a94b94fa", "4e441a3fb277d6fbe0ed7c6e080d9a9f" },
+			{ "12ef8789b91a8e35fd0ac79457a906f0610d2eb7b30dfa9d3c99d08fdcc0bfad", "4e441a3fb277d6fbe0ed7c6e080d9a9f", "cb9241bc964cbc9823531f68e3a03b7c" },
+			{ "2cbe5980ad54c5dbf2ee3db1e9875733aa9f6f0b254146051fcacfe73f6084d1", "cb9241bc964cbc9823531f68e3a03b7c", "70b17c3e869aca6076617a2cf75e9f4a" },
+			{ "dc4da3812736a2a603bcc6390763c5c5da2e1335a3db8c6569abb5cbc83e1b9b", "70b17c3e869aca6076617a2cf75e9f4a", "f4b8a6ed6d6d72aff59484314f210bb9" },
+			{ "2e4e21a6eb77bb39a1cb5cd20cc5fd3b2e96b5d8ceb6feca9c3f31fa871f1022", "f4b8a6ed6d6d72aff59484314f210bb9", "f56a4597beaeafc0c14ee73988bcbee8" },
+			{ "54d5c037e61fd09e5cf57fd34c5e0192dbfcf04f7018510a5d71d6c30fa3aeca", "f56a4597beaeafc0c14ee73988bcbee8", "04333079d5352236e1c79213f3f38dbe" },
+			{ "3a5da203e03fa399caeb1fac63679b56dfcfc036a52d733cbcb644d0fc502374", "04333079d5352236e1c79213f3f38dbe", "22c97ecdf4af830f94b11951f41e4d29" },
+			{ "5c84147d6ab2051b56b0993a7cbfa306fd06befb5182f03328075d81084e6e5d", "22c97ecdf4af830f94b11951f41e4d29", "def46a3b39c8048431d2491d97daa6ea" },
+			{ "7af77aa155a33f658283ebc3e9eb708923f2d4c0684af4b719d5149c9f94c8b7", "def46a3b39c8048431d2491d97daa6ea", "4c6367ad8a2190366c3d730fe5eeb6ee" },
+			{ "edabba9dc9d87357bb91da6931c743e16f91b36de26b648175e867937a7a7e59", "4c6367ad8a2190366c3d730fe5eeb6ee", "90bd09b4eb0f7d8397c0026cefea8fb3" },
+			{ "ab7ef1b4a1e43771de88e158ad26a419ff2cbad909641902e22865ff9590f1ea", "90bd09b4eb0f7d8397c0026cefea8fb3", "25aefcfa6ec98fae81b93afad7761711" },
+			{ "6b0145c6cbeaae320be86c2909c8d643da82462367ad96ac63915f0542e6e6fb", "25aefcfa6ec98fae81b93afad7761711", "607b26f4eb3585e3e886e75c0f3a01cf" },
+			{ "b46d28c20a614b6d986e92795258631cbaf960d78c98134f8b17b8594ddce734", "607b26f4eb3585e3e886e75c0f3a01cf", "fe35fe16a3290042c28c021ee9ede73f" },
+			{ "93c5c8403410764717e547e3d07b0ec344cc9ec12fb1130d499bba47a431000b", "fe35fe16a3290042c28c021ee9ede73f", "0141871c357a00ad37ae65597830cac8" },
+			{ "cf85ea60e0c94e611fedba48e19e9693458d19dd1acb13a07e35df1edc01cac3", "0141871c357a00ad37ae65597830cac8", "df8bc7782ae10eb2c1bdfdc8887e4907" },
+			{ "d40735932ea269aeb51fa3caf0c176d89a06dea5302a1d12bf8822d6547f83c4", "df8bc7782ae10eb2c1bdfdc8887e4907", "1ac65c65b6cd8ab751f1d908ebd962ab" },
+			{ "9a7120341819fe54fc8a750d17ecf20f80c082c086e797a5ee79fbdebfa6e16f", "1ac65c65b6cd8ab751f1d908ebd962ab", "9f10b53fb6adf7189e277d04e351aa9d" },
+			{ "8c19642c172cdb804059f751b3f25cff1fd037ff304a60bd705e86da5cf74bf2", "9f10b53fb6adf7189e277d04e351aa9d", "c94ecc943ddc8d1eeafadd173cf73fdc" },
+			{ "e1cc9575e51d6a8240e6f3cd8958e1a6d69efb6b0d96eda39aa45bcd6000742e", "c94ecc943ddc8d1eeafadd173cf73fdc", "5dd513d74b72f9d46b06b4b96a94c3bd" },
+			{ "ce41ff2d70b2169fad6a5f9ead4471e88b4be8bc46e41477f1a2ef740a94b793", "5dd513d74b72f9d46b06b4b96a94c3bd", "1594a31cc22cdbaf2011a9a317538608" },
+			{ "7b467b87180772677903a043bc63c3ef9edf4ba084c8cfd8d1b346d71dc7319b", "1594a31cc22cdbaf2011a9a317538608", "4d5e8d6fa3e2ea343b5afb77124529aa" },
+			{ "ca068b3afc717d1f6e8dc7e4eca0f56fd381c6cf272a25eceae9bda00f821831", "4d5e8d6fa3e2ea343b5afb77124529aa", "1bbf651cf61c295c96e73a210483d7a1" },
+			{ "d04d80998a971ec52d5390fa7eaf78eac83ea3d3d1360cb07c0e87810b01cf90", "1bbf651cf61c295c96e73a210483d7a1", "84b42f75c3b62bbf21707ba66cca399f" },
+			{ "2cfb1c71e8cac872e6a6dc90b0195acf4c8a8ca61280270f5d7efc2767cbf60f", "84b42f75c3b62bbf21707ba66cca399f", "a0c514e1e6a9c659f605eff4cf4951a0" },
+			{ "405bc450a19e6dae76847b8eac858c8bec4f9847f429e156ab7b13d3a882a7af", "a0c514e1e6a9c659f605eff4cf4951a0", "1dcc38307c6b3c31e25f868b279b3711" },
+			{ "35755fe97d9aeb1e67c1f4ca5a40ce1ff183a0778842dd67492495588f1990be", "1dcc38307c6b3c31e25f868b279b3711", "cdd744574664be7b221d7a2921b4f0b2" },
+			{ "7e9d8b4bbd2aad70d229247d5880a0cb3c54e420ce26631c6b39ef71aead600c", "cdd744574664be7b221d7a2921b4f0b2", "e602ef85184ad4a86ed339e9403f541c" },
+			{ "73be0bb5b74a8f6324f8a5f0acbf5faada560ba5d66cb7b405ead698ee923410", "e602ef85184ad4a86ed339e9403f541c", "7c227f03e605ff14c42fb9f8e8786e84" },
+			{ "c82e22faeca51a38045f2a3a04a0e6e2a67474a6306948a0c1c56f6006ea5a94", "7c227f03e605ff14c42fb9f8e8786e84", "94c0466b27aa4d361b19c250329c14a0" },
+			{ "7c1dd42bdff414e3733b846cbb00c43432b432cd17c30596dadcad3034764e34", "94c0466b27aa4d361b19c250329c14a0", "fd6b281c93c1fd8e9b83b69ecf722f3b" },
+			{ "aa7ec003efbca86fe5ad028d0571cf3acfdf1ad18402f818415f1baefb04610f", "fd6b281c93c1fd8e9b83b69ecf722f3b", "a75b28be92dc53de9d44714918e3c541" },
+			{ "14ba26bafb6c8410104d70ec23876db56884326f16deabc6dc1b6ae7e3e7a44e", "a75b28be92dc53de9d44714918e3c541", "48693025ba68f3fa3515e7112b6e32aa" },
+			{ "33b40738b39ebe39118938c6461ebd7820ed024aacb6583ce90e8df6c88996e4", "48693025ba68f3fa3515e7112b6e32aa", "6bfa37cc5ed6e5d4ddf40b1301957f73" },
+			{ "0d16eab1b56e2a468bf1ba1035129d574b173586f260bde834fa86e5c91ce997", "6bfa37cc5ed6e5d4ddf40b1301957f73", "ca18b89c3225877e02e21042af9836a8" },
+			{ "44b55514a5b4a5a4294f3cf8e75ff73c810f8d1ac0453a96361896a76684df3f", "ca18b89c3225877e02e21042af9836a8", "c3ebc08368be58f8f514cfbd5f98ed52" },
+			{ "265f1a7a53151bc51526fe476e8c63ad42e44d99a8fb626ec30c591a391c326d", "c3ebc08368be58f8f514cfbd5f98ed52", "788552397ddf2d5841dbc0e73dd11150" },
+			{ "23fa4630937015585ac065e6a7bd93023a611fa0d5244f3682d799fd04cd233d", "788552397ddf2d5841dbc0e73dd11150", "e9d728bc19b50809bc1245552091d0b2" },
+			{ "9dd3a2408e5694fced28964c7f442e6ed3b6371ccc91473f3ec5dca8245cf38f", "e9d728bc19b50809bc1245552091d0b2", "ce08d78b58179b7957eeab6f74ddb10f" },
+			{ "055d53ef574a63f77c0b3a88e7e7cdb01dbee0979486dc46692b77c750814280", "ce08d78b58179b7957eeab6f74ddb10f", "33ce33e80fcc262e9e862f5ac50c14df" },
+			{ "c1a90cdef5d5daaf7b182d6b409940e62e70d37f9b4afa68f7ad589d958d565f", "33ce33e80fcc262e9e862f5ac50c14df", "5c4d21e6045f11f9f91a8cef130304e0" },
+			{ "6a257a4f65bdd7d882ff45ffbb9541e5723df2999f15eb910eb7d472868e52bf", "5c4d21e6045f11f9f91a8cef130304e0", "4f50679c2dbfa63b5e55d031afacd44b" },
+			{ "5843bf4d6afc64273d328e1baf7821b03d6d9505b2aa4daa50e20443292286f4", "4f50679c2dbfa63b5e55d031afacd44b", "29f64e79de48f3d3d5dff92aa0b1783b" },
+			{ "8915060a33758a300c053d365c304d4e149bdb7c6ce2be79853dfd698993fecf", "29f64e79de48f3d3d5dff92aa0b1783b", "8c5bf92515a92449846593a3bdbe521b" },
+			{ "be2b2a41b3491511c181508b66ec96f798c02259794b9a3001586eca342dacd4", "8c5bf92515a92449846593a3bdbe521b", "5511fef7a2238277aaa9e25f1d097c19" },
+			{ "4d6f38935317ffa0cb29bdfa8a43cd19cdd1dcaedb681847abf18c952924d0cd", "5511fef7a2238277aaa9e25f1d097c19", "2ecd75c6a5eace1c541e12db150a2143" },
+			{ "1549ea3d167152c85223b0a603d1fa61e31ca9687e82d65bffef9e4e3c2ef18e", "2ecd75c6a5eace1c541e12db150a2143", "6a4cfcb9e8dae9610df99509daac7be0" },
+			{ "ba593f6731f70edab83f5cdf02527436895055d196583f3af2160b47e6828a6e", "6a4cfcb9e8dae9610df99509daac7be0", "185772deff51807147ac5350249b3e1a" },
+			{ "7126a05712f987f1c9249bf9a7c17b4f9107270f6909bf4bb5ba5817c219b474", "185772deff51807147ac5350249b3e1a", "bfb1ce7df706bc972e1b8306d44aa135" },
+			{ "3d0fa3983ebc8a3f64c4135d7cd3195e2eb6e9729e0f03dc9ba1db1116531541", "bfb1ce7df706bc972e1b8306d44aa135", "f21df1e99a781dba4a68ff3491848f99" },
+			{ "88c88ec4d63eb481cd8d6e8e8d2e2715dcab189b04771e66d1c9242587d79ad8", "f21df1e99a781dba4a68ff3491848f99", "e4e8c9e4963a44f5cf27767e4e42fa61" },
+			{ "989b9545625a4025f7725d63847213883843d17f924d5a931eee525bc99560b9", "e4e8c9e4963a44f5cf27767e4e42fa61", "27475121d0e3367d0e0c2d9fb39cfe95" },
+			{ "f49b61f190f982aa866d5c8f2e5749781f04805e42ae6cee10e27fc47a099e2c", "27475121d0e3367d0e0c2d9fb39cfe95", "c2990626cb34c5d7fafe0430650ed907" },
+			{ "101537d5f633ad8dc7e8cc058ce7fe20dd9d8678899aa939ea1c7bf41f07472b", "c2990626cb34c5d7fafe0430650ed907", "1c5d0cca2845b66c371c1760f81e024a" },
+			{ "1f6acaa66674733d28dfec92c43c7e76c1c08ab2a1df1f55dd006c94e7194561", "1c5d0cca2845b66c371c1760f81e024a", "6f5a53addd93c9a136401a804a710419" },
+			{ "b55bdefe370646699012062df748b6aaae9ad91f7c4cd6f4eb407614ad684178", "6f5a53addd93c9a136401a804a710419", "8e0c070be7109ead0e8cc0aa0bf95e61" },
+			{ "de09c34a85d56a2748559c309f04eb722096de149b5c4859e5ccb6bea6911f19", "8e0c070be7109ead0e8cc0aa0bf95e61", "db53b5306561db899a635a56c56e7239" },
+			{ "2650ecd133a5df3825072a4df2d95d15fbc56b24fe3d93d07fafece863ff6d20", "db53b5306561db899a635a56c56e7239", "72e5ac05990d3f5508295f95f8973313" },
+			{ "a8599f0edd6446b4bb9371e380bf33488920c7216730ac857786b37d9b685e33", "72e5ac05990d3f5508295f95f8973313", "86a8332f16f997bc02af5271c64e7e0f" },
+			{ "672e7b0d497458e260084662c566394e0f88f40e71c93b397529e10c5d26203c", "86a8332f16f997bc02af5271c64e7e0f", "8e53858ce7ad0d25410f886eeeca1e0e" },
+			{ "73ea41a50cdd5a98644e6f8d974af21381db71829664361c34266962b3ec3e32", "8e53858ce7ad0d25410f886eeeca1e0e", "c6a1d493d607d2cef6aef29ee878c434" },
+			{ "14969b3ff78ccdd66c53b8027a79563e477aa5114063e4d2c2889bfc5b94fa06", "c6a1d493d607d2cef6aef29ee878c434", "a4973e0cfaf8d6ac6758615fc81e06df" },
+			{ "d09e2741b9480c2166c9c1832654763de3ed9b1dba9b327ea5d0faa3938afcd9", "a4973e0cfaf8d6ac6758615fc81e06df", "67e3243d8bc81887517708a29a5d9ae5" },
+			{ "22c65ead303a0395cc9c065ada930ddd840ebf2031532af9f4a7f20109d7663c", "67e3243d8bc81887517708a29a5d9ae5", "de2f35df51644cf5d79984be6c17c14c" },
+			{ "8679ef0357516245ec489386419646bc5a218aff6037660c233e76bf65c0a770", "de2f35df51644cf5d79984be6c17c14c", "9eaedabb3bd046ec27a2bdafabdbc861" },
+			{ "6043902c115f107639e7c92ba4952301c48f50445be720e0049ccb10ce1b6f11", "9eaedabb3bd046ec27a2bdafabdbc861", "0cde7e6c1aecbbe996865bf30b178de7" },
+			{ "aca443a9e31033ea73b1eeda6e7d6ca1c8512e28410b9b09921a90e3c50ce2f6", "0cde7e6c1aecbbe996865bf30b178de7", "ba25bf5be1e4c099b9e45eedd7369cb5" },
+			{ "60362e2b0bf998c875c087c370c4ca4772749173a0ef5b902bfece0e123a7e43", "ba25bf5be1e4c099b9e45eedd7369cb5", "805be62789549ce6af74966467f41135" },
+			{ "6d7f0f7584162a1fa4dd6764548f355af22f775429bbc776848a586a75ce6f76", "805be62789549ce6af74966467f41135", "ab6001c6c4c56e8ca393c5fd173505ba" },
+			{ "d8ecc39ac1d00c53216f6e64e826a7a9594f7692ed7ea9fa27199d9762fb6acc", "ab6001c6c4c56e8ca393c5fd173505ba", "3ba3673f4f495dd1541d47c22b7921c5" },
+			{ "493108f91caedf714652149a2b2030fe62ec11ada237f42b7304da5549824b09", "3ba3673f4f495dd1541d47c22b7921c5", "b24fe17cdc5c8cfa4260c38691b57bfa" },
+			{ "446af6dd5f58755aeaa0a1226d8c584fd0a3f0d17e6b78d1316419d3d83730f3", "b24fe17cdc5c8cfa4260c38691b57bfa", "86d999a63b96f6c9d9aaf3be6202977b" },
+			{ "65f92d4e1723d5e58aeb350c79df28de567a697745fd8e18e8ceea6dba35a788", "86d999a63b96f6c9d9aaf3be6202977b", "c4712aa733f9737f91e4ed61609e02f1" },
+			{ "915cee6af4ea95623f7122acda5e9040920b43d07604fd67792a070cdaaba579", "c4712aa733f9737f91e4ed61609e02f1", "0e8c1a77b280f4c753682768fd6f3b23" },
+			{ "8e9bb2887fe60d42db4d827f00ba68ff9c8759a7c48409a02a42206427c49e5a", "0e8c1a77b280f4c753682768fd6f3b23", "dccb684d47c480cc1317dcaa451234c0" },
+			{ "cbf85a6645469e5df882fe840776b6aa404c31ea8340896c3955fcce62d6aa9a", "dccb684d47c480cc1317dcaa451234c0", "9a2c4f07489c14265e33ac031d02b3d8" },
+			{ "3ea3c33d7439ab3c478c01907f13cda7da607eedcbdc9d4a676650cd7fd41942", "9a2c4f07489c14265e33ac031d02b3d8", "5c8e622ddbd32ee79c17572e8b3ee61c" },
+			{ "312c5b43263c1af8d1e35c0f24d1004386ee1cc0100fb3adfb7107e3f4eaff5e", "5c8e622ddbd32ee79c17572e8b3ee61c", "c5d2cb3d5b7ff0e23e308967ee074825" },
+			{ "2b09ba39b834062b9e93f48373b8dd018dedf1e5ba1b8af831ebbacbc92a2643", "1f9b9b213f1884fa98b62dd6639fd33b", "89649bd0115f30bd878567610223a59d" },
+			{ "58ac71619fdc3ac73a17f285319e1cd492766ac485030e02a95d971daab5f578", "aecd334ef8fb0c51b6896ae065d8be28", "1f9b9b213f1884fa98b62dd6639fd33b" },
+			{ "f6e3cca2cd628c10625c62cf08b385743cbb598a7df802531fd4fdfdcf6d4b50", "f1938dd245c055e9c380336ff8450d9d", "aecd334ef8fb0c51b6896ae065d8be28" },
+			{ "a9ae3bd7d454f19d69289875ff009d16cd28d458383857badc54ce92372846cd", "42d5c4a13b748800ebfe0f67781dcff1", "f1938dd245c055e9c380336ff8450d9d" },
+			{ "e3bc6232865f0476149be7162ef8fc4b8ffd10f9034cdfba37aac1f54f35893c", "469968a00226f0aae7acfda02b2ce0ae", "42d5c4a13b748800ebfe0f67781dcff1" },
+			{ "451fbfba7f09b625540822fac8c8e30bc9647859016a2f10d0063c5564196992", "eab144d6f80ccfd2fae95d16784718ac", "469968a00226f0aae7acfda02b2ce0ae" },
+			{ "f53e21e46a0ec97c980d49d6f4b81ec423d53c8ff966e0c22aef61431c5e713e", "e04b91c3f084d733d3d0c1c7c152695a", "eab144d6f80ccfd2fae95d16784718ac" },
+			{ "5fb3a2cbdbe6971fcd345961bdcba5f6c39ead4c09e237f1f93fa084dd0c1864", "e47db9b8c7fcc9459fa0a7fc84047b5b", "e04b91c3f084d733d3d0c1c7c152695a" },
+			{ "6b145436c5ae0dcb2477f94f5bdb037927e314f4ce1efeb4669f07785908633f", "f617d26b55da999d65b2d236358b2e60", "e47db9b8c7fcc9459fa0a7fc84047b5b" },
+			{ "126752da3d7f9a0ffc1a46ccaa8a3925d1f4c69f9bc46729032dd54e6c834d5f", "066ffea799ad5f09d03cb868deb1591e", "f617d26b55da999d65b2d236358b2e60" },
+			{ "aa3daafbcd647f0a9eaa027146ab7aded79b383802693820d3116d26b2321441", "6d28dbb55fb10c1c3226609997a4fb38", "066ffea799ad5f09d03cb868deb1591e" },
+			{ "fed7b5c1f84d3dc185222da3fffac7a1bab3e38d5dd8343ce1370dbf2596ef79", "9dd78987562c9c099d67ccc927fcec13", "6d28dbb55fb10c1c3226609997a4fb38" },
+			{ "c36db4af0bee11dad6ea59679cfb988c27646a0a0bf4a8357c50c176026a036a", "b5cc780ec75b8d8c5128c4d26fa520ea", "9dd78987562c9c099d67ccc927fcec13" },
+			{ "7b5aebe66f96867504ee9efd2bd56bda92a81204ccaf25b92d7805a46dcf2380", "c715bd2f23b98c66bf56e5680c6037f7", "b5cc780ec75b8d8c5128c4d26fa520ea" },
+			{ "81a73bfcacc994f47aeec0d655e0485055bdaf2bef16a9df922ee0cc61af1477", "e00984f9d93551e95ae85cfe4a829b82", "c715bd2f23b98c66bf56e5680c6037f7" },
+			{ "2678e4ccda2eb6ad7f0e2843020da2d8b5b42bd23623f836c8c6bc322b2d8ff5", "bfa487fcb2331afe3ad82e59064b4aae", "e00984f9d93551e95ae85cfe4a829b82" },
+			{ "03565af5c267a62e0d2500ead06bc7d90a10ac2e8410e2c8f21e926b2d66c55b", "4c278ca418fd1c5718883da0ba4a066f", "bfa487fcb2331afe3ad82e59064b4aae" },
+			{ "ec719a4ed851c36f806fbc78183026cc4637208a9cedfe9fea96afcb972cc334", "ceed1c8b96d84951cadf809d3ea1b397", "4c278ca418fd1c5718883da0ba4a066f" },
+			{ "f138ad912201beecd15ebe2049f980e588da3c010a35b7ce20492f56a98d70a3", "f88aed84bca42a06647d6504a5c48379", "ceed1c8b96d84951cadf809d3ea1b397" },
+			{ "8da22644e1267212742200dbf6b0ba947050d185b6919dc844344a520c49f3da", "c451b75375d0fec1cc1825eefbcc1372", "f88aed84bca42a06647d6504a5c48379" },
+			{ "5104213aabaccafe44129632f67cf2dab40166d6c3416309882c6fbcf785e0a8", "9a54dfd56bff41ea17bcf08d3761e95b", "c451b75375d0fec1cc1825eefbcc1372" },
+			{ "4c7c62be3f5a9d09934940284565e9c92e55b903a8be22e39f909f31c0e409f3", "76e074c587f112f3ffd17b390fa8cffc", "9a54dfd56bff41ea17bcf08d3761e95b" },
+			{ "9871765871b1dae363fca3961224414a58b5cdc62f4f30106041e408cf4cc60f", "5b14fd4c40a127733c20cd84648ec287", "76e074c587f112f3ffd17b390fa8cffc" },
+			{ "58e044758e5e4f300eb08f8b6ccd967203a1308a6fee17635c61298cabc20488", "01c234e2392f42831b1172352f322286", "5b14fd4c40a127733c20cd84648ec287" },
+			{ "ffbb7104e110fdfe689f68ffa48b79590263046856c155e047705bb984f0260e", "ab142165cb093648938660fd800d8896", "01c234e2392f42831b1172352f322286" },
+			{ "91aefec30a79b2524b19a611c0c4952da977250d9dc863a8d4f63b4404fdae98", "5c28ad9b5721edadc30e09b4541fd724", "ab142165cb093648938660fd800d8896" },
+			{ "00d0f69c36671a9efeb7b8637e22af38f55f8896cae98e0517f832f050e279bc", "5a9ad24fe456f4f766e839a8101373c3", "5c28ad9b5721edadc30e09b4541fd724" },
+			{ "b22571c62ee4e0b4ebd7d776703c1345afc55ad92ebf7af271100b5840f10a7f", "a30e9a7730628b560c71967bdb35efce", "5a9ad24fe456f4f766e839a8101373c3" },
+			{ "1431d01006510179dd13e2a83e1b8b7b0ccbc0ae1eddf1a47d619d239bc4e5b1", "7118cc8a057df7ea9b7fb096decf05d9", "a30e9a7730628b560c71967bdb35efce" },
+			{ "06104531e2d000fff0cc8336b62e9bb87dd30c241ba0064ee61e2db5450be068", "c550ff93a36fbb8e9a505682e07fb981", "7118cc8a057df7ea9b7fb096decf05d9" },
+			{ "8e5ebd5c9bf858f3379bfb88c7c75319b883f3b7b8cfbdc07c4e7b37a57459e9", "1375e7b4164adc861bbb77c50aba445c", "c550ff93a36fbb8e9a505682e07fb981" },
+			{ "ccc9445c3aab1c03ec8027593d370664abf61403ae85614667f50cf2afce1db5", "5bc8b411c0fae92bad54b4e5c3a8826e", "1375e7b4164adc861bbb77c50aba445c" },
+			{ "b804994228d618a9b7a3e0c3d8d2d373f03ea0126e7f886dcaa1b8176c669fdb", "36c9de7c01298ff06c5e3ecf12507cd6", "5bc8b411c0fae92bad54b4e5c3a8826e" },
+			{ "8a784f32cdd7e2ed19e85faf04de2015c6f77e6e6f56079da6ff86d87e36e30d", "35d814c0f6c07ae72bbe7cd1751d7527", "36c9de7c01298ff06c5e3ecf12507cd6" },
+			{ "71e3cc0f60a2eafd83c3a455eefc4c9cf32f6aae99967d7a8d41fa090b2b962a", "73c39ce5ea1b5eab18c9c15221e258d1", "35d814c0f6c07ae72bbe7cd1751d7527" },
+			{ "71bef3d2beffa17d37c6c3b28ef2928280ecf64b738d23d195883b5b2ac9cefb", "ac67508a8bcab69f96ead29121718cd7", "73c39ce5ea1b5eab18c9c15221e258d1" },
+			{ "9aa1ffa9e175f03a6bcc71fc2aea3a522c8ba6c1f847954e0362e9ca0bb8422c", "49d2cac2ca39a3b66302e5f3650e71d7", "ac67508a8bcab69f96ead29121718cd7" },
+			{ "4e480b73ab5c659538354b2474c9367b65596c03327e36f860600c396eb633fb", "e54638674ea0b3e837772a2866723753", "49d2cac2ca39a3b66302e5f3650e71d7" },
+			{ "aeecd55c38bef76367a21469d08721d2801f54647cde85105717261108c404a8", "074d933018e9e0558b416e39548994ff", "e54638674ea0b3e837772a2866723753" },
+			{ "e36c6a30f734a679ab7af2f9ff767cfe8752c75464376545dc5648285c4d9057", "04825a9ba20646acadb2cbdbd23530c0", "074d933018e9e0558b416e39548994ff" },
+			{ "124806708fc5fc73a1e50c61e5a76b8783d09dcfc63123e971e483f38e78a097", "25b98d2837580f9f3da8f2593f7b1d72", "04825a9ba20646acadb2cbdbd23530c0" },
+			{ "560afb5529a8b433738348a296408cada66910e7f1692c764c4c71aab103bde5", "048b220c70d875cad6901e40f45c5b03", "25b98d2837580f9f3da8f2593f7b1d72" },
+			{ "8b7c61dd007d9d0d1937cfc69b63c114a2e232eb81b159bc9adc6fea455fe6e6", "34e2a33412c0e682b29b26d5409d015f", "048b220c70d875cad6901e40f45c5b03" },
+			{ "30560ca0a31d8145228d34fbf8415683960091df9371bf3e2847493f05c2e7b9", "942c057914f3b88fa0a40bb41d02b051", "34e2a33412c0e682b29b26d5409d015f" },
+			{ "20c7e6e45c08ba0756c379b6cd4d6886022c94a6878207b188e3428b18c057e8", "6f4e5c40e4135abfdf294b4d8c9304b3", "942c057914f3b88fa0a40bb41d02b051" },
+			{ "2a6c928bf8c2b32ce82ed0be1bc954b96d62c8e663915d0e57ca09c69453535b", "cac032ddc2b27443d70d0f2dfb1a1d87", "6f4e5c40e4135abfdf294b4d8c9304b3" },
+			{ "74994d46526b2e9d5470a0a687f726f2a7a2fa3ba123294d80c706eb6f494edc", "6aec33b2c1ae39237b13e86555722d1a", "cac032ddc2b27443d70d0f2dfb1a1d87" },
+			{ "5aaec28d079eef5f1c42df377c7e80accd4ec989608d106efbd4ee8e3a3b63c6", "2df34173bbff331488b23687c755f617", "6aec33b2c1ae39237b13e86555722d1a" },
+			{ "9bfa3c08b8f8b8758f0a1f4b0ef86be4e0bd88fadb72237a7366d809fd6e95d1", "07649dd5c1c7b38810cbecbad9c037ed", "2df34173bbff331488b23687c755f617" },
+			{ "188e3a027b75fede978c76e70a853f95e7d9152f1ab590f263ad34b324aea23c", "ef3d9092321649d9bbb88f6448e3847b", "07649dd5c1c7b38810cbecbad9c037ed" },
+			{ "aae2c6363be3dbb71c21f3a2d9f123f108e485bd28a3d92bd815bbd76c4d2647", "1e5058b3e986d4aad80a11f909af1cca", "ef3d9092321649d9bbb88f6448e3847b" },
+			{ "85a51c6d62ba97564ac955a5409c362916b4dd0ec1250d81001faa2e65e23a8d", "9bf81052596b786b50b3f3506c009169", "1e5058b3e986d4aad80a11f909af1cca" },
+			{ "0d09600f108a86cfd58f3c211a40c38d8d4ccd5c984e75ea50ac597e09e2abe4", "81bdaf86cf4255e5826294a51d2f9992", "9bf81052596b786b50b3f3506c009169" },
+			{ "14ab478866e1d160311aa8820df354fb0cf162da570c200fd2cecddb14cd3276", "1fdc5fee9eb640ff547a98d9e1601b52", "81bdaf86cf4255e5826294a51d2f9992" },
+			{ "39f18405f0b1da647f99479e87d4d153132d3d34c9ba60f086b45502f5ad2924", "936609876cd6eebf31db5a6af7839844", "1fdc5fee9eb640ff547a98d9e1601b52" },
+			{ "85a900bc1563b16e01481e45d1125734804b34b3a56c8e4fb76f0f68022eb160", "9ea8b6942287f470e9866a48f4dc661a", "936609876cd6eebf31db5a6af7839844" },
+			{ "c43d677fc0a1becbb0214221b866f0be1ee3822787eb7a3f5ee96520f6f2d77a", "1dfc2cf82bfaa2732eb525916ef9ef13", "9ea8b6942287f470e9866a48f4dc661a" },
+			{ "584f10659c2d3c41f6a1926de4595eb2031faedfac11d84c705c40b1980b3869", "a5d8d1d60281a152b9b76142e3e5a512", "1dfc2cf82bfaa2732eb525916ef9ef13" },
+			{ "8eba336927e18b13eb3858a89709768ba6c77f09ae90791ec9eb21f37bee9d7b", "8ac0278f2ab100a5964173d76934d732", "a5d8d1d60281a152b9b76142e3e5a512" },
+			{ "13f1558dcae09062bfdccdd0037e5c432c075886842179bb5faa522412da4a49", "30233b70c4d6fd31b23cc77bc73559ea", "8ac0278f2ab100a5964173d76934d732" },
+			{ "d736669842aa88f43f1de421d38f391f1c2463f640f7848aed96955fd5ef13a3", "766630dfd41e4146a827f55cba745b6a", "30233b70c4d6fd31b23cc77bc73559ea" },
+			{ "0ca8353df90e73606588a11675bd37a76a42532994e9c5cc45b160036f9b48c9", "177f22d98f009c214a44b08cf40b790d", "766630dfd41e4146a827f55cba745b6a" },
+			{ "136ec3a1a65ee7bb0fd09af7364b66607d3d71f01be959ed0ff5d08f9b9031c4", "e9c916a21dec9662bdd91a0f865f8b47", "177f22d98f009c214a44b08cf40b790d" },
+			{ "ba5c2bd413daca7866de11c11db2926694f467520605cf8fb22cca801dcfba83", "92b3ab3e28f6c606508706e0c3d76d28", "e9c916a21dec9662bdd91a0f865f8b47" },
+			{ "0fab4768fa7f104a3579f534508182920647cc6c2ef30989e2abcc60de18d7ab", "3652053dcba591f101944fb1a5eb2b99", "92b3ab3e28f6c606508706e0c3d76d28" },
+			{ "e766aa23c075deaa6fae68384f5b205e3015c951e5569878e33f83d17bf3fc32", "242260d2387c6f1a415834b34e300a7e", "3652053dcba591f101944fb1a5eb2b99" },
+			{ "9f1941442336bfffa1d6b27dce7845971437a983dd2af762a267b76235c3f64c", "a2a58ed050231b0d15f8cd99f5e171fb", "242260d2387c6f1a415834b34e300a7e" },
+			{ "429a3cdc98124fafe33156a66b4e8c2ab69227538d09ec6fb79f7afbc02287b7", "ec6631b17e5c97193fd9651b64feffe7", "a2a58ed050231b0d15f8cd99f5e171fb" },
+			{ "92764a89e0de63823b9d9a5a001868035af416e2f3557b7688461fe0a4dc7850", "1eef8afa419a9341085a2d611185e42c", "ec6631b17e5c97193fd9651b64feffe7" },
+			{ "0c0bf13d79d1297e137e6f5995cb0a2c441b9c18b2cfe837801c3281b5599c7c", "d4eb4e591cf856b816e2207264e426e4", "1eef8afa419a9341085a2d611185e42c" },
+			{ "de31f995867bb45813d369349cb7434290f0d241ae37be8f96fe12f3d1bdba98", "401e14b126348e1aedd13662a19c23fb", "d4eb4e591cf856b816e2207264e426e4" },
+			{ "3b1a02b5600444db09118e99765fa61dd0eec6f0880330957b2f249170219963", "38bfea7e1eece741eb6a9601605129dd", "401e14b126348e1aedd13662a19c23fb" },
+			{ "1fbcf4ffb36f76b88323f745b3152947e8512c8e96efd7d49045b2901070b0be", "71cd66b3b3553b6aa632c338a307e11a", "38bfea7e1eece741eb6a9601605129dd" },
+			{ "b963bfe1c513f0dba7e24b5cfa1535a3999c4a3d25baecbe367771a8b37751a4", "34c2972d43e59eb0c44b7b69a540e7de", "71cd66b3b3553b6aa632c338a307e11a" },
+			{ "23458e2a4332ce00e7967a438cd86448ad5edd10665f720ef23c0ac11637b67a", "8988b878d0c6a84b6262351e033fb1b7", "34c2972d43e59eb0c44b7b69a540e7de" },
+			{ "e552b8a0f474a8da502e749e3f0639b124d66568b699da45905e3fdf150807cd", "010e24d01c37edca127484c8d8137772", "8988b878d0c6a84b6262351e033fb1b7" },
+			{ "65a0645dd971bdb3b0d3a079ef96f3a925d841b8aaae378f822abb17cd1b70bf", "6c7b2e83630a8c8b6219839297d7c647", "010e24d01c37edca127484c8d8137772" },
+			{ "9867d55ae9f56a03c435f4b25f2eee2c49a36f3bc9a4bb04e03338855accb6f8", "ccba81435ecf162f70d30b3163cb3858", "6c7b2e83630a8c8b6219839297d7c647" },
+			{ "e4a3addb683bbe85ed84b3ce044bcc538519ee78976bad2b90e033b439078ea0", "a0575bc9b8931f34cd6072e2015e7331", "ccba81435ecf162f70d30b3163cb3858" },
+			{ "6f543088bf24577d55a685e3637d85e8254eb5b12ff8b21f5d8041563859fd91", "a44046a87574e6d3e6c6446fb3c324bc", "a0575bc9b8931f34cd6072e2015e7331" },
+			{ "5fff671ca24c3b241ddb92761f02dfe3810ef3195a8c54ccbb4605398b9ad92d", "ef6abb8ae007b730be45aa85e39f362d", "a44046a87574e6d3e6c6446fb3c324bc" },
+			{ "db897d48a317f66412a8c52111c3f3b06e644893ba8be3fc0503afbc6805ef00", "e5bd2edfb7e563f6cd86da757c4c370f", "ef6abb8ae007b730be45aa85e39f362d" },
+			{ "78aec1da8e8e79007af12c05d800a18d8bd9664c0d6e800ac88575c91449d80f", "33caa267a64fb3cac210271bb8f25369", "e5bd2edfb7e563f6cd86da757c4c370f" },
+			{ "2c3e5387a51addb0b7789ef0289109b3b813c42bab2133c00a9552d2acbb8b66", "27bb26990ca5e12d387bd659cb7e7860", "33caa267a64fb3cac210271bb8f25369" },
+			{ "85685fe6996d514e5928d7196f41b5d29fa8e2b2a784d2ed32ee848b67c5f306", "d866ace71916e00c1559bba4149380a0", "27bb26990ca5e12d387bd659cb7e7860" },
+			{ "82abcc9a00bae55087cf5d87413f095547ce4e55be9232e127b73f2f735673a6", "8a4d6a0af8488a82a0c31332db3e36de", "d866ace71916e00c1559bba4149380a0" },
+			{ "81d983ce59413e84d7af4558d2e3941ccd83245f46dab86387742c1da8684578", "05f6ff378b837cd375cddfae79ec2a9d", "8a4d6a0af8488a82a0c31332db3e36de" },
+			{ "bbe6d0d13a6b146bb23db6885ba66487c875db68cd59c4b0f2b9f3b3d1846fe5", "b706bdb4019d604cc35faa07197c318d", "05f6ff378b837cd375cddfae79ec2a9d" },
+			{ "8092df6673b2b77783f2ba73c22e8d3c7f7366dcccc4a4fc31e659b4c8f85e68", "b26354bc77d3ce16909ca505d2381e4c", "b706bdb4019d604cc35faa07197c318d" },
+			{ "3118fc57b93677af8e14c1b5961b76f3cd103260bb176aeaa17afcb11ac04024", "1515ef6a25cf3943eadadc36a029194a", "b26354bc77d3ce16909ca505d2381e4c" },
+			{ "c4bde4a0c5f80329be6dd515f5bf6dabd805dd0a9ed853a94ba02087bae9596e", "7b4ad946dfb59c80b1dc5cedb2fa87a0", "1515ef6a25cf3943eadadc36a029194a" },
+			{ "80b47f9126b0b1f5ab408bb9ea5f349aa34f044c416dcf29fa7c7c6a0813dece", "fa10a53e14cde5ef1e9a8a5692847aa5", "7b4ad946dfb59c80b1dc5cedb2fa87a0" },
+			{ "86a43939c8cd243f7ff2a9658524373b595fa17255a02ac6e4e6f63c9a97a46b", "4c9bed07308abcd0b09265b5efbe301c", "fa10a53e14cde5ef1e9a8a5692847aa5" },
+			{ "74a8f02a15285153e0ae160fbddcf6bf15c44c75652a96165474938975299477", "75b7496c42809a3a592acbd8e069269e", "4c9bed07308abcd0b09265b5efbe301c" },
+			{ "293b40a2f70cc2c4262554663fb6f0986073051927aa0c2c0d5e58519540b2e9", "2dbc01185af7e084d90578468b6b10ef", "75b7496c42809a3a592acbd8e069269e" },
+			{ "b04d4f4eb708eef86c5b6b8e08a665084dcf04017d5deca8d45b20171e2ba206", "c95b3a9f689c9dce4995c24f72dd5162", "2dbc01185af7e084d90578468b6b10ef" },
+			{ "7cb9f7211815e0cb48b71286f84f80a184943e9e15c171669dcee2586cf6f364", "616e55c3bf113e2c18cae3c61b7eb7d1", "c95b3a9f689c9dce4995c24f72dd5162" },
+			{ "b3cdb46cf92aa0b96e87212bc650d5e5e5fa6b5daad04f4a8504019e778844b5", "87e8b80767ebbdbad75cb94f4cb54f3b", "616e55c3bf113e2c18cae3c61b7eb7d1" },
+			{ "512c2a3821eb53af613141c71e1076656212d35acd3bf2f05258b8d13b3d0b8e", "c83e20e18f2b1457788954b49fd84307", "87e8b80767ebbdbad75cb94f4cb54f3b" },
+			{ "9977c985745bc33954a2ce898bc8febdaa2cf3bb4210e6a72ad1ec65a4e54889", "e3d3868f578caf34e36445bf14cefc68", "c83e20e18f2b1457788954b49fd84307" },
+		};
+		for (auto& Test : MonteCarloCases)
+		{
+			RunTest(Test.Key, Test.PlainText, Test.CipherText, 1000);
+		}
+	}
+
+	// 6.3 Multi-block Message Test, from ECBMMT256.rsp
+	SECTION("Multi-block Message")
+	{
+		static const AesTestCase MultiBlockCases[] =
+		{
+			{
+				"cc22da787f375711c76302bef0979d8eddf842829c2b99ef3dd04e23e54cc24b",
+				"ccc62c6b0a09a671d64456818db29a4d",
+				"df8634ca02b13a125b786e1dce90658b",
+			},
+			{
+				"7a52e4d342aa07255a7e7c34266cf7302abe2d4dd7ec4468a46187ee61825ffa",
+				"7e771c6ee4b26db89050e982ba7e9803c8da34606434dd85d2910e538076d001",
+				"a91d8b2ddf37520bc469470ad0dd6394923143ce55386beb1f9c4bd51584658e",
+			},
+			{
+				"605c4139c961b496ca5148f1bdb1bb1901f2101943a0ec10fcdc403d3b0c285a",
+				"68c9885ba2be03181f65f1e04e83d6ba6880467550bcf099be26dc9d9c0af15ab02abac07c116ac862a41da90cfa604f",
+				"a7603d29bbba4c77208bf2f3df9f5ec85204adce012299f2cce7b326ce78f5cf8040343dd291e8cf9f3645726368dc20",
+			},
+			{
+				"f984b0f534fc0ae2c0a8593e16ab8365f25fcc9c5947f9a2db45b588160d35c3",
+				"351fee099122e371c4830f409c6c4411186d22176f7138b054f16b3c79679c2f520685651ba8e4b61c08dccb2c31982f743631a97524d2ca4d351ac23546c178",
+				"8b9c9e692c16e7059818e285e85d8fa5433dee2aff9fec61d6a0a781e24b24f64902fbd18cef7461ad7760cfb2442fb74ffd9be108a386545f2a216430ef16fb",
+			},
+			{
+				"ba42b760bb5a5de21acb9aba214c9783cd71ea841ada018580abc4e1be3b76dd",
+				"4b4b12d6ee6fc0bf987eaafe2634aad464781ff4c83d3f8a61a6af7c0a6d51f0e3855d0e02feb307652a6f562bfebe4604baf1b4e7cdd01603f231bcf7a0c95645a141b704008cd8d62979201a4c84e2",
+				"fa18d25e37ea0ce94f0949efc0edecc6a40fada8f007fd8e760afed0a83ebb350c82b03baaa6ee19f791bb9bd1b44d27a76fc6eb0e1c0017d68776ed69a541851a732e46ef328def064baf6a0a755588",
+			},
+			{
+				"1d29079cc34ab5a3bc713f416a129f9d26ada15fca458cc2731404ea857d2f79",
+				"b36809198132518d815aa37f32f40dd7a952ec8bd63355700837b450b396b33c72125e23482a84a42b916021c3de78156f85c6a78906167feb64afd8b1d935d641c8ce1a89f3849588eeeb9910d40336ca385fc37a5e87bb84ab9ccbb05b3a28",
+				"7ae803b14914d156439f580c02c592bf9a41b7b80c20168129a33fae2290403f01a6fabe4fea7ac770ffbc6c421f8e013b9e831674ef17eb27d446dded3baf50686b809c18b6bbd588cd3c7423b64987e9164b7e1e66987dc9319790ae27b3e8",
+			},
+			{
+				"2e39c585ce4900d323ce29713bebe73a1be08a0cb22e9f1310fcc14ad4b9b23e",
+				"e199899e1e12cffcb28909aec51b36c2f96fab49ef32b9650cc38aa37d2f4c8b785f9176c590f6a07e04037e13f7535290d5f5fc23aa1113d9dacf34a812749ab27ecfefc83dd3622d1285fa9d5c192a8e4857a5b6454473cdb8ff4594f1e906165e08b22effae6b491a55ca6d30ce73",
+				"2c4998428e72f6d96e982a316f73bf2a7da81730909b65403489ab92ada6de11882d08742f90f0f109d3420b00b8abe6873f4fdd1474923da2c5bdea2e4523ffca2132015ecf7c9cac9de2f956b112c4ba8e4c8e4b354a3df6e4652d6a77ae982a24d15cff71b979538d49994fd38761",
+			},
+			{
+				"85405c4f0ebbe8f29228f02f1ff184e2f5e7857e8933c2a1d08f61ecb9b68111",
+				"0f5321db6fd9d816d88e28183a739d90974e76095caa9f12f11fe49c8fd35fa3524132118f397cdf6728857d9c9a3c74a4fde44afca80aa5bf1cbab47789f2cb3394574dda5727cfbea96f7a74a07eb1e45599f49ce7f0056ac3d14929221c70dbd3f759f83ac22f06994ed96a8e49917eddfdd2e3703b78199c91234ca6c3dc",
+				"2ac6de212da0434bea9cdd7332637307131d31e8c4b0c1fd02298e249bfa9f64f34aae45faadf79d971ae82b033d903f6b18adec171ec83ccd147b44d05dec5cedb574534e8901385534c3f1dae4aebee0eb2164975e8b4e85a7b5a766ffd8247885c2b6429eb659b92c8d953af92b54517933566654d804665112bec17ff3a4",
+			},
+			{
+				"f157285db00e64c2791668a54493966e3039a19426605056b95b7eac5106667d",
+				"3637f71f60a430322980349ad414fcfdc14f87e9915d210e8b7be5aa3e09814468e0399d17e72fe40ee1e1296a89f31486e12fd71bc7ca61acc9e8d4213a633ab285c87406c2a3729c87fdaa6b0122c2c543a89081ddac4559b15fe59102c2fbdacead8a755e16469b1b90041da31270481cfafe0bc951235cda51c4d78924efa862efc9ace20fc7d344a321c9984a84",
+				"ca0a683e759c1312928fe01198f625bba044bb9003e82b9279b6812feefe54e30ccc0ca51fb858eba9ed4667e85c146a424071e7bf603f1f538ead57a3e29c583549c08e279bb078cc51e42eeb3d2443da965192cd0478e8d7ea1343cf90196f520933e3aab6db2dcd9a76534e05483eea2d373deae95d6213b91db2b96aa6adce727cdf7e43bd01978d07e0241cf1c1",
+			},
+			{
+				"44a2b5a7453e49f38261904f21ac797641d1bcd8ddedd293f319449fe63b2948",
+				"c91b8a7b9c511784b6a37f73b290516bb9ef1e8df68d89bf49169eac4039650c4307b6260e9c4e93650223440252f5c7d31c26c56209cbd095bf035b9705880a1628832daf9da587a6e77353dbbce189f963235df160c008a753e8ccea1e0732aa469a97659c42e6e31c16a723153e39958abe5b8ad88ff2e89af40622ca0b0d6729a26c1ae04d3b8367b548c4a6335f0e5a9ec914bb6113c05cd0112552bc21",
+				"05d51af0e2b61e2c06cb1e843fee3172825e63b5d1ce8183b7e1db6268db5aa726521f46e948028aa443af9ebd8b7c6baf958067ab0d4a8ac530ecbb68cdfc3eb93034a428eb7e8f6a3813cea6189068dfecfa268b7ecd5987f8cb2732c6882bbec8f716bac254d72269230aec5dc7f5a6b866fd305242552d400f5b0404f19cbfe7291fab690ecfe6018c4309fc639d1b65fcb65e643edb0ad1f09cfe9cee4a",
+			},
+			{
+				"a81fd6ca56683d0f5445659dde4d995dc65f4bce208963053e28d7f2df517ce4",
+				"8b2b1b22f733ac09d1196d6be6a87a72",
+				"4154c0be71072945d8156f5f046d198d",
+			},
+			{
+				"bc1415aa119c29a9a27c0ea9d19ed50ace86c488b8e39d6a056424fb23cddb3e",
+				"a8fab53790af3519cf21978e3cf03ba8e52b902be23311bf17f1ad2c5ff37c16",
+				"0a5f3278d7d96632e0508391e813f06b35d8d754ddf5867240d3168dd69f4a66",
+			},
+			{
+				"0d0ec1b61ebc5177c4513ef1d7d5bb97d06abaa2d33710a8eda6d3709acf0705",
+				"f17733de8f763110ef4b305594936ca2bb75119ad65261be32ba919a2c3ef8b8f1c42f62b8474362e53ee7cc6c82a647",
+				"c1c83fa3cd3d52524876e715bc28efe7c7c4256a139e9d2c874ea029bf56b792ba0606cdd39ddbbdf3b1874304d16d05",
+			},
+			{
+				"ae67999eb240a9e5b6e3f0bd6b504530000be51343b8bc3b0ae8e0eed1335f98",
+				"959b7d696e2267f3d465f6f77ff2453296cfa54843363981484853e67cc21d340b803d6d653213d1037b81d849c5acc361771a5a072c9c29d6bc5096c38c9c86",
+				"254c015626baa3ed2d7f05f5d0a98c8cc2f29dd7a4410ea41d74db4e2c5d1be1bd0a327f7b4a4703e66cf9e3a54a4d5b1d87ebc842148b52adf9aacdf05281e1",
+			},
+			{
+				"dde079371133ebd68df061b56f0efd3a14c137ced35a30e0eb68422cb924dc3d",
+				"5889ad2c09a6307611e6115a78c13566de05b5892f7850fb917f83898e07486cc9ce746e891db102a0f570d7ad3c2804f40cdfe23bcc8b2c8a3bfba48632892d3df3bb7bbd1029b915cab2faf281c4e6",
+				"ea1fd2f064548906ad10ce1240758868ed9fb32921dabe18681232a8308b955ad0d28e45c9f34af64bec1d7bfb626dbbf393cc090a8c64f8569b9870f008e801f7001578d8d286b0cc5eebaba2c920e1",
+			},
+			{
+				"d68a345159178b9bd2e3bd7a13c9512ee9b397944eff81a8df28b44890a2df3b",
+				"a56e000be19eced20479cbe8964fdb366c437ca3b5fd9d04c339c1a51bdb0aad46df1b78ef05b02e492464136155d76dd32c3b4aa4198c3926cb75d2578434ab4fff0894508b6f60ff39686a0fd151d1bffba7a786b1bc02acd2d23b56e45749",
+				"59722586a56ed6c8207d6a0a9f7278588f5203782fb64e6ffd71f1486d732d10e1ed7d25d6b66db15beeaf71c8e16cbbac2dd8bf0728090dba7c09573b142a788d3c99b670033644d5e4a30b94a69bab618186bcfd824a59a93163dfe07036f2",
+			},
+			{
+				"5a20ff7da47c7e853becca0ca5f32bc80e17de973337146f7e1f3c93725a850d",
+				"34467d07c2e49c44e9003da836786169e7c0fb54f1e2f178387c2f759d50809ce84d6705fe6350907d5e945e262af378753163f5569be9c09c874d254dde2545898d4cbd2eaa1ade9c8a02f8d76d418565d02075e082999bbf6f2ba985f65f17c73cc0cc29ac04824db98b3cc2b77126",
+				"bed9f2db31ebbfac57f3745e0dacf71087cb35c94326b9ba3b1b7dc35f30003942f798e714350a52d0425842b6b3da7dfc3790d6c1bfb76642e29a1a507fab6e02a4d277a7ae0527feb219cdaa1cd8dbe09653ec632c7a05b2ade91b7f5405d664a74895d2a0d9b24a0b6007cf2f18b1",
+			},
+			{
+				"082d33edd0a1ad3de596762d711bae6f3188a12c7b6ced987fc7e8c9cd7a3cc9",
+				"51c527c0985369234a599f476731e8b51ac7a44bfbe37ec71a641ad71ad464f9e45467d82a0d101f67043b87e2da34de183459295000cea9e0f0ccdd82270fc70b880fa80104e0b78b7a5b1620dee83da84702223b277a0966b10ddf44ef06b98e48788092c763895d95f13fd389fff5706e70876d5af8c197cdebb38a4d2ba6",
+				"496553f24d47264ef74a584b893ccd4ea1cfc0b1044aac15997e6dfab9fca28da68557c058c2e1a0803508e4c7706cc1568521069cbd9d644b8dd40da4c5a182f50b68a56088fa16dde1b8bc8269fca20395b40575b050ab57cfdf76d208743840dfcfc9ec0139571835b0fe458f4498aee6c1d716898e91c32ed55acb2be458",
+			},
+			{
+				"3e1934554952b87a07d6371721f4d78a0e9adadc42be7347a2fcd8f53c81990b",
+				"5f967e10da68887567b514293149697212888efc9fd90604459f8a9c4fa1f4117f0214fa675b680571bd980cf941bfccd4826e1fe1420a0ba59554341b507e760d2d85c7219b3f5e261eef2d201bc134d2ec32d8b97157e3de91f1c312b26ff193e637bf78012850f23d0536715e51a68be730c55375774642345e0daaa4c40185985cab38df09dac9a688e6b597d327",
+				"2b3c036be68fba7a625dc72a0e68b6677326f8751714e0e1a14df2b73d5b5df5424dc640130947d643e4b8c4fc702d59a195bfdab2e49077f027097ec23d66cfd82575900e589a21935c17524a968ad47060b7ee805b8871530652ee8c905a2633eb98c2a022bb459bec6d3ec181d7c22aac68194822d2b81212566f62fd42e49cc8b8409110b028a3bd40ef29da1c0a",
+			},
+			{
+				"c4a71e055a7254dda360693fe1be49f10faa6731c36dbaa6590b05974e185c5b",
+				"31fd5a307e279b2f34581e2c432379df8eccbaf79532938916711cd377540b9045373e47f2214b8f876040af733f6c9d8f03a7c58f8714d2fbb4c14af59c75b483adc718946ee907a18286cc4efd206789064b6f1b195f0d0d234468e4f00e6f1cad5cd3b9c0a643b3c0dd09280ff2e2a5929183409384dd72dc94e39687ea2b623d5d776700bd8b36e6130ffde966f134c4b1f35f29c5cc4a03297e1ccc9539",
+				"2c487fa96f4090c56aa1b5be81918a934c9492878fb0cd686dcf8d17d86485454c51237bbd09205dcef1552f430dd098b9d827a694730c133a0222c77f540f9d5fc2d36af359583c9e3b49df884228a64de79b67f66207c8281360b99b214042ce61367ff97960e944453cd63679bb44708897d29bc5e70f9fc8f1f715143fbb00f7f5c1b7b161ec26d8d41d36fab0fa8a85c3ee6ce4d37007eb7a89d6753590",
+			},
+		};
+		for (auto& Test : MultiBlockCases)
+		{
+			RunTest(Test.Key, Test.PlainText, Test.CipherText);
+		}
+	}
+}
+
+#endif
diff --git a/Engine/Source/Runtime/Core/Tests/Misc/SecureHashTest.cpp b/Engine/Source/Runtime/Core/Tests/Misc/SecureHashTest.cpp
new file mode 100644
--- /dev/null
+++ b/Engine/Source/Runtime/Core/Tests/Misc/SecureHashTest.cpp
@@ -0,0 +1,211 @@
+// Copyright Epic Games, Inc. All Rights Reserved.
+
+#if WITH_LOW_LEVEL_TESTS
+
+#include "Misc/SecureHash.h"
+
+#include "HAL/UnrealMemory.h"
+#include "String/HexToBytes.h"
+#include "Containers/StringView.h"
+#include "Containers/StringConv.h"
+#include "Containers/Array.h"
+#include "Misc/StringBuilder.h"
+
+#include "TestHarness.h"
+
+TEST_CASE("Core::Hash::SecureHash", "[Core][SHA1][Smoke]")
+{
+	struct FSHA1TestCase
+	{
+		FAnsiStringView Hash;
+		FAnsiStringView Input;
+	};
+
+	// https://csrc.nist.gov/CSRC/media/Projects/Cryptographic-Standards-and-Guidelines/documents/examples/SHA1.pdf
+	SECTION("Examples")
+	{
+		static const FSHA1TestCase TestCases[] =
+		{
+			{ "a9993e364706816aba3e25717850c26c9cd0d89d", "abc" },
+			{ "84983e441c3bd26ebaae4aa1f95129e5e54670f1", "abcdbcdecdefdefgefghfghighijhijkijkljklmklmnlmnomnopnopq" },
+		};
+
+		for (const auto& TestCase : TestCases)
+		{
+			FSHA1 Hash;
+			Hash.Update((const uint8*)TestCase.Input.GetData(), TestCase.Input.Len());
+
+			TStringBuilder<40> Hex;
+			Hex << Hash.Finalize();
+
+			REQUIRE(Hex.ToView() == StringCast<TCHAR>(TestCase.Hash.GetData()));
+		}
+	}
+
+	// https://csrc.nist.gov/projects/cryptographic-algorithm-validation-program/secure-hashing#shavs
+	SECTION("Validation")
+	{
+		static const FSHA1TestCase TestCases[] =
+		{
+			// from SHA1ShortMsg.rsp
+			{ "da39a3ee5e6b4b0d3255bfef95601890afd80709", "" },
+			{ "c1dfd96eea8cc2b62785275bca38ac261256e278", "36" },
+			{ "0a1c2d555bbe431ad6288af5a54f93e0449c9232", "195a" },
+			{ "bf36ed5d74727dfd5d7854ec6b1d49468d8ee8aa", "df4bd2" },
+			{ "b78bae6d14338ffccfd5d5b5674a275f6ef9c717", "549e959e" },
+			{ "60b7d5bb560a1acf6fa45721bd0abb419a841a89", "f7fb1be205" },
+			{ "a6d338459780c08363090fd8fc7d28dc80e8e01f", "c0e5abeaea63" },
+			{ "860328d80509500c1783169ebf0ba0c4b94da5e5", "63bfc1ed7f78ab" },
+			{ "24a2c34b976305277ce58c2f42d5092031572520", "7e3d7b3eada98866" },
+			{ "411ccee1f6e3677df12698411eb09d3ff580af97", "9e61e55d9ed37b1c20" },
+			{ "05c915b5ed4e4c4afffc202961f3174371e90b5c", "9777cf90dd7c7e863506" },
+			{ "af320b42d7785ca6c8dd220463be23a2d2cb5afc", "4eb08c9e683c94bea00dfa" },
+			{ "9f4e66b6ceea40dcf4b9166c28f1c88474141da9", "0938f2e2ebb64f8af8bbfc91" },
+			{ "e6c4363c0852951991057f40de27ec0890466f01", "74c9996d14e87d3e6cbea7029d" },
+			{ "046a7b396c01379a684a894558779b07d8c7da20", "51dca5c0f8e5d49596f32d3eb874" },
+			{ "d58a262ee7b6577c07228e71ae9b3e04c8abcda9", "3a36ea49684820a2adc7fc4175ba78" },
+			{ "a150de927454202d94e656de4c7c0ca691de955d", "3552694cdf663fd94b224747ac406aaf" },
+			{ "35a4b39fef560e7ea61246676e1b7e13d587be30", "f216a1cbde2446b1edf41e93481d33e2ed" },
+			{ "7ce69b1acdce52ea7dbd382531fa1a83df13cae7", "a3cf714bf112647e727e8cfd46499acd35a6" },
+			{ "b47be2c64124fa9a124a887af9551a74354ca411", "148de640f3c11591a6f8c5c48632c5fb79d3b7" },
+			{ "8bb8c0d815a9c68a1d2910f39d942603d807fbcc", "63a3cc83fd1ec1b6680e9974a0514e1a9ecebb6a" },
+			{ "b486f87fb833ebf0328393128646a6f6e660fcb1", "875a90909a8afc92fb7070047e9d081ec92f3d08b8" },
+			{ "76159368f99dece30aadcfb9b7b41dab33688858", "444b25f9c9259dc217772cc4478c44b6feff62353673" },
+			{ "dbc1cb575ce6aeb9dc4ebf0f843ba8aeb1451e89", "487351c8a5f440e4d03386483d5fe7bb669d41adcbfdb7" },
+			{ "d7a98289679005eb930ab75efd8f650f991ee952", "46b061ef132b87f6d3b0ee2462f67d910977da20aed13705" },
+			{ "fda26fa9b4874ab701ed0bb64d134f89b9c4cc50", "3842b6137bb9d27f3ca5bafe5bbb62858344fe4ba5c41589a5" },
+			{ "c2ff7ccde143c8f0601f6974b1903eb8d5741b6e", "44d91d3d465a4111462ba0c7ec223da6735f4f5200453cf132c3" },
+			{ "643c9dc20a929608f6caa9709d843ca6fa7a76f4", "cce73f2eabcb52f785d5a6df63c0a105f34a91ca237fe534ee399d" },
+			{ "509ef787343d5b5a269229b961b96241864a3d74", "664e6e7946839203037a65a12174b244de8cbc6ec3f578967a84f9ce" },
+			{ "b61ce538f1a1e6c90432b233d7af5b6524ebfbe3", "9597f714b2e45e3399a7f02aec44921bd78be0fefee0c5e9b499488f6e" },
+			{ "5b7b94076b2fc20d6adb82479e6b28d07c902b75", "75c5ad1f3cbd22e8a95fc3b089526788fb4ebceed3e7d4443da6e081a35e" },
+			{ "6066db99fc358952cf7fb0ec4d89cb0158ed91d7", "dd245bffe6a638806667768360a95d0574e1a0bd0d18329fdb915ca484ac0d" },
+			{ "b89962c94d60f6a332fd60f6f07d4f032a586b76", "0321794b739418c24e7c2e565274791c4be749752ad234ed56cb0a6347430c6b" },
+			{ "17bda899c13d35413d2546212bcd8a93ceb0657b", "4c3dcf95c2f0b5258c651fcd1d51bd10425d6203067d0748d37d1340d9ddda7db3" },
+			{ "badcdd53fdc144b8bf2cc1e64d10f676eebe66ed", "b8d12582d25b45290a6e1bb95da429befcfdbf5b4dd41cdf3311d6988fa17cec0723" },
+			{ "01b4646180f1f6d2e06bbe22c20e50030322673a", "6fda97527a662552be15efaeba32a3aea4ed449abb5c1ed8d9bfff544708a425d69b72" },
+			{ "10016dc3a2719f9034ffcc689426d28292c42fc9", "09fa2792acbb2417e8ed269041cc03c77006466e6e7ae002cf3f1af551e8ce0bb506d705" },
+			{ "9f42fa2bce6ef021d93c6b2d902273797e426535", "5efa2987da0baf0a54d8d728792bcfa707a15798dc66743754406914d1cfe3709b1374eaeb" },
+			{ "cdf48bacbff6f6152515323f9b43a286e0cb8113", "2836de99c0f641cd55e89f5af76638947b8227377ef88bfba662e5682babc1ec96c6992bc9a0" },
+			{ "b88fb75274b9b0fd57c0045988cfcef6c3ce6554", "42143a2b9e1d0b354df3264d08f7b602f54aad922a3d63006d097f683dc11b90178423bff2f7fe" },
+			{ "c06d3a6a12d9e8db62e8cff40ca23820d61d8aa7", "eb60c28ad8aeda807d69ebc87552024ad8aca68204f1bcd29dc5a81dd228b591e2efb7c4df75ef03" },
+			{ "6e40f9e83a4be93874bc97cdebb8da6889ae2c7a", "7de4ba85ec54747cdc42b1f23546b7e490e31280f066e52fac117fd3b0792e4de62d5843ee98c72015" },
+			{ "3efc940c312ef0dfd4e1143812248db89542f6a5", "e70653637bc5e388ccd8dc44e5eace36f7398f2bac993042b9bc2f4fb3b0ee7e23a96439dc01134b8c7d" },
+			{ "a0cf03f7badd0c3c3c4ea3717f5a4fb7e67b2e56", "dd37bc9f0b3a4788f9b54966f252174c8ce487cbe59c53c22b81bf77621a7ce7616dcb5b1e2ee63c2c309b" },
+			{ "a544e06f1a07ceb175a51d6d9c0111b3e15e9859", "5f485c637ae30b1e30497f0fb7ec364e13c906e2813daa34161b7ac4a4fd7a1bddd79601bbd22cef1f57cbc7" },
+			{ "199d986ed991b99a071f450c6b1121a727e8c735", "f6c237fb3cfe95ec8414cc16d203b4874e644cc9a543465cad2dc563488a659e8a2e7c981e2a9f22e5e868ffe1" },
+			{ "33bac6104b0ad6128d091b5d5e2999099c9f05de", "da7ab3291553c659873c95913768953c6e526d3a26590898c0ade89ff56fbd110f1436af590b17fed49f8c4b2b1e" },
+			{ "76d7db6e18c1f4ae225ce8ccc93c8f9a0dfeb969", "8cfa5fd56ee239ca47737591cba103e41a18acf8e8d257b0dbe8851134a81ff6b2e97104b39b76e19da256a17ce52d" },
+			{ "f652f3b1549f16710c7402895911e2b86a9b2aee", "57e89659d878f360af6de45a9a5e372ef40c384988e82640a3d5e4b76d2ef181780b9a099ac06ef0f8a7f3f764209720" },
+			{ "63faebb807f32be708cf00fc35519991dc4e7f68", "b91e64235dbd234eea2ae14a92a173ebe835347239cff8b02074416f55c6b60dc6ced06ae9f8d705505f0d617e4b29aef9" },
+			{ "0e6730bc4a0e9322ea205f4edfff1fffda26af0a", "e42a67362a581e8cf3d847502215755d7ad425ca030c4360b0f7ef513e6980265f61c9fa18dd9ce668f38dbc2a1ef8f83cd6" },
+			{ "b61a3a6f42e8e6604b93196c43c9e84d5359e6fe", "634db92c22010e1cbf1e1623923180406c515272209a8acc42de05cc2e96a1e94c1f9f6b93234b7f4c55de8b1961a3bf352259" },
+			{ "32d979ca1b3ed0ed8c890d99ec6dd85e6c16abf4", "cc6ca3a8cb391cd8a5aff1faa7b3ffbdd21a5a3ce66cfaddbfe8b179e4c860be5ec66bd2c6de6a39a25622f9f2fcb3fc05af12b5" },
+			{ "6f18190bd2d02fc93bce64756575cea36d08b1c3", "7c0e6a0d35f8ac854c7245ebc73693731bbbc3e6fab644466de27bb522fcb99307126ae718fe8f00742e6e5cb7a687c88447cbc961" },
+			{ "68f525feea1d8dbe0117e417ca46708d18d7629a", "c5581d40b331e24003901bd6bf244aca9e9601b9d81252bb38048642731f1146b8a4c69f88e148b2c8f8c14f15e1d6da57b2daa9991e" },
+			{ "a7272e2308622ff7a339460adc61efd0ea8dabdc", "ec6b4a88713df27c0f2d02e738b69db43abda3921317259c864c1c386e9a5a3f533dc05f3beeb2bec2aac8e06db4c6cb3cddcf697e03d5" },
+			{ "aef843b86916c16f66c84d83a6005d23fd005c9e", "0321736beba578e90abc1a90aa56157d871618f6de0d764cc8c91e06c68ecd3b9de3824064503384db67beb7fe012232dacaef93a000fba7" },
+			{ "be2cd6f380969be59cde2dff5e848a44e7880bd6", "d0a249a97b5f1486721a50d4c4ab3f5d674a0e29925d5bf2678ef6d8d521e456bd84aa755328c83fc890837726a8e7877b570dba39579aabdd" },
+			{ "e5eb4543deee8f6a5287845af8b593a95a9749a1", "c32138531118f08c7dcc292428ad20b45ab27d9517a18445f38b8f0c2795bcdfe3ffe384e65ecbf74d2c9d0da88398575326074904c1709ba072" },
+			{ "534c850448dd486787b62bdec2d4a0b140a1b170", "b0f4cfb939ea785eabb7e7ca7c476cdd9b227f015d905368ba00ae96b9aaf720297491b3921267576b72c8f58d577617e844f9f0759b399c6b064c" },
+			{ "6fbfa6e4edce4cc85a845bf0d228dc39acefc2fa", "bd02e51b0cf2c2b8d204a026b41a66fbfc2ac37ee9411fc449c8d1194a0792a28ee731407dfc89b6dfc2b10faa27723a184afef8fd83def858a32d3f" },
+			{ "018872691d9b04e8220e09187df5bc5fa6257cd9", "e33146b83e4bb671392218da9a77f8d9f5974147182fb95ba662cb66011989c16d9af104735d6f79841aa4d1df276615b50108df8a29dbc9de31f4260d" },
+			{ "d98d512a35572f8bd20de62e9510cc21145c5bf4", "411c13c75073c1e2d4b1ecf13139ba9656cd35c14201f1c7c6f0eeb58d2dbfe35bfdeccc92c3961cfabb590bc1eb77eac15732fb0275798680e0c7292e50" },
+			{ "9f3ea255f6af95c5454e55d7354cabb45352ea0b", "f2c76ef617fa2bfc8a4d6bcbb15fe88436fdc2165d3074629579079d4d5b86f5081ab177b4c3f530376c9c924cbd421a8daf8830d0940c4fb7589865830699" },
+			{ "a70cfbfe7563dd0e665c7c6715a96a8d756950c0", "45927e32ddf801caf35e18e7b5078b7f5435278212ec6bb99df884f49b327c6486feae46ba187dc1cc9145121e1492e6b06e9007394dc33b7748f86ac3207cfe" },
+
+			// first 5 from SHA1LongMsg.rsp
+			{
+				"d8fd6a91ef3b6ced05b98358a99107c1fac8c807",
+				"7c9c67323a1df1adbfe5ceb415eaef0155ece2820f4d50c1ec22cba4928ac656"
+				"c83fe585db6a78ce40bc42757aba7e5a3f582428d6ca68d0c3978336a6efb729"
+				"613e8d9979016204bfd921322fdd5222183554447de5e6e9bbe6edf76d7b71e1"
+				"8dc2e8d6dc89b7398364f652fafc734329aafa3dcd45d4f31e388e4fafd7fc64"
+				"95f37ca5cbab7f54d586463da4bfeaa3bae09f7b8e9239d832b4f0a733aa609c"
+				"c1f8d4"
+			},
+			{
+				"4a75a406f4de5f9e1132069d66717fc424376388",
+				"6cb70d19c096200f9249d2dbc04299b0085eb068257560be3a307dbd741a3378"
+				"ebfa03fcca610883b07f7fea563a866571822472dade8a0bec4b98202d47a344"
+				"312976a7bcb3964427eacb5b0525db22066599b81be41e5adaf157d925fac04b"
+				"06eb6e01deb753babf33be16162b214e8db017212fafa512cdc8c0d0a15c10f6"
+				"32e8f4f47792c64d3f026004d173df50cf0aa7976066a79a8d78deeeec951dab"
+				"7cc90f68d16f786671feba0b7d269d92941c4f02f432aa5ce2aab6194dcc6fd3"
+				"ae36c8433274ef6b1bd0d314636be47ba38d1948343a38bf9406523a0b2a8cd7"
+				"8ed6266ee3c9b5c60620b308cc6b3a73c6060d5268a7d82b6a33b93a6fd6fe1d"
+				"e55231d12c97"
+			},
+			{
+				"a135e32581bb06289b8c83f040e9421ec79bbe01",
+				"6487972d88d0dd390d8d09d134860f263f88df7a3412457adf510dcf164e6cf0"
+				"41679b3a19fcc542af6a236ab03d66b2e8a155d1061ab7859f75732775fff682"
+				"f8f4d5e50d3ab3770f4f66cb138155b4715d245b8069948ea016a45b7ef0fdde"
+				"93188c57eef4717f3425181de5b9a5d4e0a2963f2a67a340eb1ae994b98a48ab"
+				"19b90ab74391c50426d28287ac4f1eb93f5af1a68c7dae40876b8afaaf35a192"
+				"93c1952e957978abee40ec32f2aa880c956c7eb72f117b397cefcfb4e75ace3b"
+				"081776e46b13521e93559d453e32ab74ebc0859b9a8dd4d1d39000ebe95f984d"
+				"80a3f5004dc91a051dfbdfe9194f4f9a483e4e7955577fb0933464c63eaec771"
+				"044d59abc3029a079519f8460a693b25b4ce207ae9d9447fc4c5446e6dad234e"
+				"9afdec0c562798cd0297318399e838be385845c6dd79ede66e2ae80afec6738d"
+				"4d9bf44c8d9eddff6c5cd2c94e340e0ddac40384b9a1408c9a4b98c37a6081d5"
+				"220fba92f1d03144db"
+			},
+			{
+				"b22b87ea30f4050913f8f0241fc2ae2c319f52e7",
+				"bd74e7f607cd7d905e90175d67650a6dc2f8a4e2d4ab1249ca88812bda7984de"
+				"ccbbb6a1ba90a0e91434ddf5e6137ba85e39a598890a7f635d335242fce0e9e0"
+				"37303b6c51e54aec06614ad5ccce06d9599c80016530d7fbb1da6eb548084b2b"
+				"05babd7d553642443efda726a1fd71a8bc087c44f285e2bccf661ead475a7267"
+				"3e4386fc4eea5197c4f13c0feb0a85bc8e67e28ab872684bbebdaa527f3c253d"
+				"ebb2dc12c2693f8e9e2651b9345c0abed7a0fafa3e5d305386c95acb7a172e54"
+				"13ef08e73b1bd4d0d6832e4c035bc8559f9b0cbd0caf037a30707641c0545356"
+				"bee151a24068d70674ef1befe16f872aef4060faaad1a968c39c45dbd7595de8"
+				"f472016b5ab812d77e545fca55000ee5ce773edaa129eac6473410c2499013b4"
+				"be895f6c0f734becfe994306e776262d4528ed8577218e3cc5201f1d9e5f3f62"
+				"230eb2caea014becfba60fcb1f3997aa5b3bb622b7205c714348ba155c30a79a"
+				"2cea43b070cada807e630b4086b129051898e1d9e68d1d0ecc9429d20d6a1403"
+				"e0035a442b37bf508eb87e8ea347a3e68427b6d48ed299ba65ecb37b38754f45"
+				"47423eaea2aec403338db2dcfe61cff4a8d17c383656981e1838a23866b91d09"
+				"698f39175d98af4175caed53"
+			},
+			{
+				"d742931bc51d4d44ff938783be73dc7beccc980b",
+				"a52638f0efb19bff5ec95fcde4ac9aabd95e14d2e5f84c551f43bc5376855e71"
+				"519b6f877248739a20cd790b85baa00d5503da5cb056f02d4aacc760c91fe1fd"
+				"6efb26def817e5a9c56616023bc9e2fe662765dae2c0b2edfcbe17db140da30c"
+				"466de65c49c6f81496bbbd1acd81666455f23bb243dd987d7ea1362a20faac84"
+				"1f1a36692cfcb4c3dbf5f6bb058c36296b8be64e9b56adc5187cacb7b58c054f"
+				"422a9e6d6a61229fdc3b494da98f5a33ed1bee14b2d2f6ad1177ffe99a6bb553"
+				"f7c4a6d0cb9e498ee0b63f388235d86c26c9d96e50fa7d1eb3bcb9279940c47a"
+				"8510d7fb175b3279318d5fe45823baba5dbe31c33c7649fe447061db78b33baa"
+				"3637b854163fe34915e931b9f3040807d9217d7b3fed62370dbe806c006b21cd"
+				"5061d24490f366e4d5f23e201a7ec83ae31b46fe2108d1af56cc9d42f9117eca"
+				"1cb5ab344c1fc334b9cf0d7f9739043bc3d413b3aa6e9d5067c240c52b4c5b89"
+				"e25ccd8a136a002008a9273f30dec3f2c1736c04a1c7ce0087c9f25d5ec5bff2"
+				"ea7ec0b0ad7c278f0ca712c9ae150e472521d958d0bd6da9ff0939725924b2ed"
+				"7b410a0ce2fe3f6b0bf25884d885ec223605e318fdf6803218a9a06ce5103c62"
+				"ded035087a98519b4eb180d778d7656b3d4811aaf11a128317d1acb3ca316639"
+				"5c51c90a3cf164071d0d132c54b3810a8211ec7774d2288447abe7afd030375a"
+				"3bed4c7cf1b28097c02e98ea36bf49e74d89fbe74ec6cc1def5cd8c8beb5b8ad"
+				"c3cb48c56182ad337e3b9778e4a6c4"
+			},
+		};
+
+		for (const auto& TestCase : TestCases)
+		{
+			TArray<uint8> Input;
+			Input.AddUninitialized(TestCase.Input.Len() / 2);
+			UE::String::HexToBytes(TestCase.Input, Input.GetData());
+
+			FSHAHash Hash = FSHA1::HashBuffer(Input.GetData(), Input.Num());
+
+			TStringBuilder<40> Hex;
+			Hex << Hash;
+
+			REQUIRE(Hex.ToView() == StringCast<TCHAR>(TestCase.Hash.GetData()));
+		}
+	}
+}
+
+#endif
