Fix for broken assert strings in TArray64.

diff --git a/Engine/Source/Runtime/Core/Public/Containers/Array.h b/Engine/Source/Runtime/Core/Public/Containers/Array.h
--- a/Engine/Source/Runtime/Core/Public/Containers/Array.h
+++ b/Engine/Source/Runtime/Core/Public/Containers/Array.h
@@ -688,7 +688,7 @@ public:
 		// Template property, branch will be optimized out
 		if (AllocatorType::RequireRangeCheck)
 		{
-			checkf((Index >= 0) & (Index < ArrayNum),TEXT("Array index out of bounds: %i from an array of size %i"),Index,ArrayNum); // & for one branch
+			checkf((Index >= 0) & (Index < ArrayNum),TEXT("Array index out of bounds: %lld from an array of size %lld"),(long long)Index, (long long)ArrayNum); // & for one branch
 		}
 	}
 
@@ -1341,7 +1341,7 @@ private:
 		checkSlow((Count >= 0) & (Index >= 0) & (Index <= ArrayNum));
 
 		SizeType NewNum = Count;
-		checkf((OtherSizeType)NewNum == Count, TEXT("Invalid number of elements to add to this array type: %llu"), (unsigned long long)NewNum);
+		checkf((OtherSizeType)NewNum == Count, TEXT("Invalid number of elements to add to this array type: %lld"), (long long)NewNum);
 
 		const SizeType OldNum = ArrayNum;
 		if ((ArrayNum += Count) > ArrayMax)
@@ -1524,7 +1524,7 @@ public:
 	 */
 	FORCEINLINE void CheckAddress(const ElementType* Addr) const
 	{
-		checkf(Addr < GetData() || Addr >= (GetData() + ArrayMax), TEXT("Attempting to use a container element (%p) which already comes from the container being modified (%p, ArrayMax: %d, ArrayNum: %d, SizeofElement: %d)!"), Addr, GetData(), ArrayMax, ArrayNum, sizeof(ElementType));
+		checkf(Addr < GetData() || Addr >= (GetData() + ArrayMax), TEXT("Attempting to use a container element (%p) which already comes from the container being modified (%p, ArrayMax: %lld, ArrayNum: %lld, SizeofElement: %d)!"), Addr, GetData(), (long long)ArrayMax, (long long)ArrayNum, sizeof(ElementType));
 	}
 
 	/**
@@ -2835,7 +2835,7 @@ private:
 	void CopyToEmpty(const OtherElementType* OtherData, OtherSizeType OtherNum, SizeType PrevMax, SizeType ExtraSlack)
 	{
 		SizeType NewNum = (SizeType)OtherNum;
-		checkf((OtherSizeType)NewNum == OtherNum, TEXT("Invalid number of elements to add to this array type: %llu"), (unsigned long long)NewNum);
+		checkf((OtherSizeType)NewNum == OtherNum, TEXT("Invalid number of elements to add to this array type: %lld"), (long long)NewNum);
 
 		checkSlow(ExtraSlack >= 0);
 		ArrayNum = NewNum;
