Logging: Added a workaround to log callstacks for failed assertions

diff --git a/Engine/Source/Runtime/Core/Private/Mac/MacErrorOutputDevice.cpp b/Engine/Source/Runtime/Core/Private/Mac/MacErrorOutputDevice.cpp
--- a/Engine/Source/Runtime/Core/Private/Mac/MacErrorOutputDevice.cpp
+++ b/Engine/Source/Runtime/Core/Private/Mac/MacErrorOutputDevice.cpp
@@ -34,8 +34,14 @@ void FMacErrorOutputDevice::Serialize( const TCHAR* Msg, ELogVerbosity::Type Ver
 		{
 			UE_LOG(LogMac, Error, TEXT("appError called: %s"), Msg );
 		}
-		FCString::Strncpy( GErrorHist, Msg, UE_ARRAY_COUNT(GErrorHist) - 5 );
-		FCString::Strncat( GErrorHist, TEXT("\r\n\r\n"), UE_ARRAY_COUNT(GErrorHist) - 1  );
+
+		// CheckVerifyFailedImpl writes GErrorHist including a callstack and then calls GError->Logf with only the
+		// assertion expression and description. Keep GErrorHist intact if it begins with Msg.
+		if (FCString::Strncmp(GErrorHist, Msg, FMath::Min<int32>(UE_ARRAY_COUNT(GErrorHist), FCString::Strlen(Msg))))
+		{
+			FCString::Strncpy(GErrorHist, Msg, UE_ARRAY_COUNT(GErrorHist) - 5);
+			FCString::Strncat(GErrorHist, TEXT("\r\n\r\n"), UE_ARRAY_COUNT(GErrorHist) - 1);
+		}
 		ErrorPos = FCString::Strlen(GErrorHist);
 	}
 	else
diff --git a/Engine/Source/Runtime/Core/Private/Unix/UnixErrorOutputDevice.cpp b/Engine/Source/Runtime/Core/Private/Unix/UnixErrorOutputDevice.cpp
--- a/Engine/Source/Runtime/Core/Private/Unix/UnixErrorOutputDevice.cpp
+++ b/Engine/Source/Runtime/Core/Private/Unix/UnixErrorOutputDevice.cpp
@@ -41,8 +41,14 @@ void FUnixErrorOutputDevice::Serialize(const TCHAR* Msg, ELogVerbosity::Type Ver
 		{
 			UE_LOG(LogCore, Error, TEXT("appError called: %s"), Msg );
 		}
-		FCString::Strncpy( GErrorHist, Msg, UE_ARRAY_COUNT(GErrorHist) - 5 );
-		FCString::Strncat( GErrorHist, TEXT("\r\n\r\n"), UE_ARRAY_COUNT(GErrorHist) - 1 );
+
+		// CheckVerifyFailedImpl writes GErrorHist including a callstack and then calls GError->Logf with only the
+		// assertion expression and description. Keep GErrorHist intact if it begins with Msg.
+		if (FCString::Strncmp(GErrorHist, Msg, FMath::Min<int32>(UE_ARRAY_COUNT(GErrorHist), FCString::Strlen(Msg))))
+		{
+			FCString::Strncpy(GErrorHist, Msg, UE_ARRAY_COUNT(GErrorHist) - 5);
+			FCString::Strncat(GErrorHist, TEXT("\r\n\r\n"), UE_ARRAY_COUNT(GErrorHist) - 1);
+		}
 		ErrorPos = FCString::Strlen(GErrorHist);
 	}
 	else
diff --git a/Engine/Source/Runtime/Core/Private/Windows/WindowsErrorOutputDevice.cpp b/Engine/Source/Runtime/Core/Private/Windows/WindowsErrorOutputDevice.cpp
--- a/Engine/Source/Runtime/Core/Private/Windows/WindowsErrorOutputDevice.cpp
+++ b/Engine/Source/Runtime/Core/Private/Windows/WindowsErrorOutputDevice.cpp
@@ -50,8 +50,13 @@ void FWindowsErrorOutputDevice::Serialize( const TCHAR* Msg, ELogVerbosity::Type
 			UE_LOG(LogWindows, Error, TEXT("Windows GetLastError: %s (%i)"), FPlatformMisc::GetSystemErrorMessage(ErrorBuffer, 1024, LastError), LastError);
 		}
 
-		FCString::Strncpy(GErrorHist, Msg, UE_ARRAY_COUNT(GErrorHist) - 5);
-		FCString::Strncat(GErrorHist, TEXT("\r\n\r\n"), UE_ARRAY_COUNT(GErrorHist) - 1);
+		// CheckVerifyFailedImpl writes GErrorHist including a callstack and then calls GError->Logf with only the
+		// assertion expression and description. Keep GErrorHist intact if it begins with Msg.
+		if (FCString::Strncmp(GErrorHist, Msg, FMath::Min<int32>(UE_ARRAY_COUNT(GErrorHist), FCString::Strlen(Msg))))
+		{
+			FCString::Strncpy(GErrorHist, Msg, UE_ARRAY_COUNT(GErrorHist) - 5);
+			FCString::Strncat(GErrorHist, TEXT("\r\n\r\n"), UE_ARRAY_COUNT(GErrorHist) - 1);
+		}
 	}
 	else
 	{
