Remove PLATFORM_PRESENT_HANG_DETECTION_ON_BY_DEFAULT.  The last remaining consumer of this was removed with a previous CL.

#rnx

diff --git a/Engine/Source/Runtime/Core/Private/HAL/ThreadHeartBeat.cpp b/Engine/Source/Runtime/Core/Private/HAL/ThreadHeartBeat.cpp
--- a/Engine/Source/Runtime/Core/Private/HAL/ThreadHeartBeat.cpp
+++ b/Engine/Source/Runtime/Core/Private/HAL/ThreadHeartBeat.cpp
@@ -383,12 +383,7 @@ void FThreadHeartBeat::InitSettings()
 
 	// Default to 25 seconds if not overridden in config.
 	double NewHangDuration = 25.0;
-
-#if	PLATFORM_PRESENT_HANG_DETECTION_ON_BY_DEFAULT
-	double NewPresentDuration = 25.0;
-#else
 	double NewPresentDuration = 0.0;
-#endif //PLATFORM_PRESENT_HANG_DETECTION_ON_BY_DEFAULT
 
 	bool bNewHangsAreFatal = !!(UE_ASSERT_ON_HANG);
 
diff --git a/Engine/Source/Runtime/Core/Public/HAL/Platform.h b/Engine/Source/Runtime/Core/Public/HAL/Platform.h
--- a/Engine/Source/Runtime/Core/Public/HAL/Platform.h
+++ b/Engine/Source/Runtime/Core/Public/HAL/Platform.h
@@ -517,10 +517,6 @@
 	#define PLATFORM_USE_MINIMAL_HANG_DETECTION					0
 #endif
 
-#ifndef PLATFORM_PRESENT_HANG_DETECTION_ON_BY_DEFAULT
-	#define PLATFORM_PRESENT_HANG_DETECTION_ON_BY_DEFAULT		0
-#endif
-
 #ifndef PLATFORM_USE_GENERIC_STRING_IMPLEMENTATION
 	#define PLATFORM_USE_GENERIC_STRING_IMPLEMENTATION			1
 #endif
