Fix path permission list not supporting trailing slashes

diff --git a/Engine/Source/Runtime/Core/Private/Misc/PathPermissionList.cpp b/Engine/Source/Runtime/Core/Private/Misc/PathPermissionList.cpp
--- a/Engine/Source/Runtime/Core/Private/Misc/PathPermissionList.cpp
+++ b/Engine/Source/Runtime/Core/Private/Misc/PathPermissionList.cpp
@@ -2,6 +2,7 @@
 
 #include "Misc/NamePermissionList.h"
 #include "Misc/StringBuilder.h"
+#include "Misc/PathViews.h"
 
 
 bool FPathPermissionList::PassesFilter(const FStringView Item) const
@@ -46,7 +47,7 @@ bool FPathPermissionList::PassesStartsWithFilter(const FStringView Item, const b
 		bool bPassedAllowList = false;
 		for (const auto& Other : AllowList)
 		{
-			if (Item.StartsWith(Other.Key) && (Item.Len() <= Other.Key.Len() || Item[Other.Key.Len()] == TEXT('/')))
+			if (FPathViews::IsParentPathOf(Other.Key, Item))
 			{
 				bPassedAllowList = true;
 				break;
@@ -55,7 +56,7 @@ bool FPathPermissionList::PassesStartsWithFilter(const FStringView Item, const b
 			if (bAllowParentPaths)
 			{
 				// If allowing parent paths (eg, when filtering folders), then we must also check if the item has a AllowList child path
-				if (FStringView(Other.Key).StartsWith(Item) && (Other.Key.Len() <= Item.Len() || Other.Key[Item.Len()] == TEXT('/')))
+				if (FPathViews::IsParentPathOf(Item, Other.Key))
 				{
 					bPassedAllowList = true;
 					break;
@@ -73,7 +74,7 @@ bool FPathPermissionList::PassesStartsWithFilter(const FStringView Item, const b
 	{
 		for (const auto& Other : DenyList)
 		{
-			if (Item.StartsWith(Other.Key) && (Item.Len() <= Other.Key.Len() || Item[Other.Key.Len()] == TEXT('/')))
+			if (FPathViews::IsParentPathOf(Other.Key, Item))
 			{
 				return false;
 			}
