Remove code that breaks *.sym file paths for Unix platform

From CL 19498448

diff --git a/Engine/Source/Runtime/Core/Private/Unix/UnixPlatformStackWalk.cpp b/Engine/Source/Runtime/Core/Private/Unix/UnixPlatformStackWalk.cpp
--- a/Engine/Source/Runtime/Core/Private/Unix/UnixPlatformStackWalk.cpp
+++ b/Engine/Source/Runtime/Core/Private/Unix/UnixPlatformStackWalk.cpp
@@ -328,23 +328,6 @@ namespace
 			if (FPaths::IsRelative(info.dli_fname))
 			{
 				FCStringAnsi::Strcpy(ModuleSymbolPath, TCHAR_TO_UTF8(FPlatformProcess::BaseDir()));
-#ifdef WITH_LOW_LEVEL_TESTS
-				// Low level tests live one level above the base directory in a folder <ModuleName>Tests
-				// Sometimes this folder can also be just <ModuleName> if the target was compiled with the tests
-				if (FPaths::DirectoryExists(FPaths::Combine(ModuleSymbolPath, TCHAR_TO_UTF8(*FPaths::GetBaseFilename(out_SymbolInfo.ModuleName)))))
-				{
-					FCStringAnsi::Strcat(ModuleSymbolPath, "/");
-					FCStringAnsi::Strcat(ModuleSymbolPath, TCHAR_TO_UTF8(*FPaths::GetBaseFilename(out_SymbolInfo.ModuleName)));
-					FCStringAnsi::Strcat(ModuleSymbolPath, "/");
-				}
-				else
-				{
-					FCStringAnsi::Strcat(ModuleSymbolPath, "/");
-					FCStringAnsi::Strcat(ModuleSymbolPath, TCHAR_TO_UTF8(*FPaths::GetBaseFilename(out_SymbolInfo.ModuleName)));
-					FCStringAnsi::Strcat(ModuleSymbolPath, "Tests");
-					FCStringAnsi::Strcat(ModuleSymbolPath, "/");
-				}
-#endif
 				FCStringAnsi::Strcat(ModuleSymbolPath, TCHAR_TO_UTF8(*FPaths::GetBaseFilename(out_SymbolInfo.ModuleName)));
 				FCStringAnsi::Strcat(ModuleSymbolPath, ".sym");
 			}
