ModuleManager behavior nop, clean up some checks and comment

diff --git a/Engine/Source/Runtime/Core/Private/Modules/ModuleManager.cpp b/Engine/Source/Runtime/Core/Private/Modules/ModuleManager.cpp
--- a/Engine/Source/Runtime/Core/Private/Modules/ModuleManager.cpp
+++ b/Engine/Source/Runtime/Core/Private/Modules/ModuleManager.cpp
@@ -128,7 +128,6 @@ IModuleInterface* FModuleManager::GetModulePtr_Internal(FName ModuleName)
 		return nullptr;
 	}
 
-	// Access the Module C pointer directly without creating any non-thread safe shared pointers which would unsafely modify the shared pointer's refcount
 	return ModuleInfo->Module.Get();
 }
 
@@ -206,6 +205,10 @@ bool FModuleManager::IsModuleLoaded( const FName InModuleName ) const
 		if( ModuleInfo.Module.IsValid()  )
 		{
 			// Module is loaded and ready
+
+			// note: not checking (bIsReady || GameThread) , that might be wrong
+			//   see difference with GetModule()
+			// in fact this function could just be replaced with GetModule() != null
 			return true;
 		}
 	}
@@ -393,12 +396,15 @@ IModuleInterface& FModuleManager::LoadModuleChecked( const FName InModuleName )
 
 IModuleInterface* FModuleManager::LoadModuleWithFailureReason(const FName InModuleName, EModuleLoadResult& OutFailureReason, ELoadModuleFlags InLoadModuleFlags)
 {
-#if 0
-	ensureMsgf(IsInGameThread(), TEXT("ModuleManager: Attempting to load '%s' outside the main thread.  Please call LoadModule on the main/game thread only.  You can use GetModule or GetModuleChecked instead, those are safe to call outside the game thread."), *InModuleName.ToString());
-#endif
-
 	IModuleInterface* LoadedModule = nullptr;
 	OutFailureReason = EModuleLoadResult::Success;
+	
+	// note that this behaves differently than ::LoadModule(), when called from not-game-thread
+	//	 LoadModule just redirects to ::GetModule on non-game-thread
+	
+#if !(UE_BUILD_SHIPPING || UE_BUILD_TEST)
+	WarnIfItWasntSafeToLoadHere(InModuleName);
+#endif
 
 	// Do fast check for existing module, this is the most common case
 	ModuleInfoPtr FoundModulePtr = FindModule(InModuleName);
@@ -409,13 +415,18 @@ IModuleInterface* FModuleManager::LoadModuleWithFailureReason(const FName InModu
 
 		if (LoadedModule)
 		{
-#if !(UE_BUILD_SHIPPING || UE_BUILD_TEST)
-			WarnIfItWasntSafeToLoadHere(InModuleName);
-#endif
+			// note: this function does not check (bIsReady || IsInGameThread()) the way GetModule() does
+			//   that looks like a bug if called from off-game-thread
+
 			return LoadedModule;
 		}
 	}
 
+	// doing LoadModule off GameThread should not be done
+	// already warned above
+	// enable this ensure when we know it's okay
+//	ensureMsgf(IsInGameThread(), TEXT("ModuleManager: Attempting to load '%s' outside the main thread.  Please call LoadModule on the main/game thread only.  You can use GetModule or GetModuleChecked instead, those are safe to call outside the game thread."), *InModuleName.ToString());
+
 	SCOPED_BOOT_TIMING("LoadModule");
 	TRACE_CPUPROFILER_EVENT_SCOPE_TEXT(*FString::Printf(TEXT("LoadModule_%s"), *InModuleName.ToString()));
 	DECLARE_SCOPE_CYCLE_COUNTER(TEXT("Module Load"), STAT_ModuleLoad, STATGROUP_LoadTime);
@@ -435,6 +446,12 @@ IModuleInterface* FModuleManager::LoadModuleWithFailureReason(const FName InModu
 
 		// Ptr will always be valid at this point
 		FoundModulePtr = FindModule(InModuleName);
+
+		// NOTE: Module is now findable , calls to Find or Load Module will find this module pointer
+		//  but it's not initialized yet (bIsReady is false so Get from other threads will fail)
+		// this AddModule must be done before the module is initialized
+		// because StartupModule may call functions that Find/Load on this module
+		// and they should get back this pointer, even though it is not finished initializing yet
 	}
 	
 	// Grab the module info.  This has the file name of the module, as well as other info.
@@ -574,6 +591,7 @@ IModuleInterface* FModuleManager::LoadModuleWithFailureReason(const FName InModu
 
 							// Startup the module
 							ModuleInfo->Module->StartupModule();
+
 							// The module might try to load other dependent modules in StartupModule. In this case, we want those modules shut down AFTER this one because we may still depend on the module at shutdown.
 							ModuleInfo->LoadOrder = FModuleInfo::CurrentLoadOrder++;
 
@@ -790,7 +808,20 @@ IModuleInterface* FModuleManager::GetModule( const FName InModuleName )
 	}
 
 	// For loading purpose, the GameThread is allowed to query modules that are not yet ready
-	return (ModuleInfo->bIsReady || IsInGameThread()) ? ModuleInfo->Module.Get() : nullptr;
+	// bIsReady load acquire (bIsReady is TAtomic) :
+	if ( ModuleInfo->bIsReady || IsInGameThread() )
+	{
+		return ModuleInfo->Module.Get();
+	}
+	
+#if !UE_BUILD_SHIPPING
+	// if you hit this it's almost always a bug
+	// it means your call to GetModule() is running at the same time that module is loading
+	// so you will see null or not depending on timing
+	UE_LOG(LogModuleManager, Warning, TEXT("GetModule racing against IsReady: %s"), *InModuleName.ToString());
+#endif
+
+	return nullptr;
 }
 
 bool FModuleManager::Exec( UWorld* Inworld, const TCHAR* Cmd, FOutputDevice& Ar )
