Optimized StaticDuplicateObject by pooling its temporary memory block allocation.
Added FPooledLargeMemoryData that allows lockless access to temporary memory blocks.

#tests Squads games on WSL, using FramePro to validate performance gains

diff --git a/Engine/Source/Runtime/Core/Private/Serialization/LargeMemoryData.cpp b/Engine/Source/Runtime/Core/Private/Serialization/LargeMemoryData.cpp
--- a/Engine/Source/Runtime/Core/Private/Serialization/LargeMemoryData.cpp
+++ b/Engine/Source/Runtime/Core/Private/Serialization/LargeMemoryData.cpp
@@ -4,6 +4,17 @@
 #include "Logging/LogMacros.h"
 #include "CoreGlobals.h"
 
+
+int32 GPooledLargeMemoryData_MaxPoolLength = 2;
+static FAutoConsoleVariableRef CPooledLargeMemoryData_MaxPoolLength(
+	TEXT("s.LargeMemoryDataMaxPoolLength"),
+	GPooledLargeMemoryData_MaxPoolLength,
+	TEXT("Limit LargeMemoryData pool size to the given number of elements."),
+	ECVF_Default
+
+);
+static constexpr int64 GPooledLargeMemoryData_MaxElementSize = 128 * 1024;
+
 /*----------------------------------------------------------------------------
 	FLargeMemoryData
 ----------------------------------------------------------------------------*/
@@ -103,3 +114,47 @@ void FLargeMemoryData::GrowBuffer()
 		Data = (uint8*)FMemory::Malloc(MaxBytes);
 	}
 }
+
+/*----------------------------------------------------------------------------
+	FPooledLargeMemoryData
+----------------------------------------------------------------------------*/
+
+TLockFreePointerListUnordered<FLargeMemoryData, 0> FPooledLargeMemoryData::FreeList;
+std::atomic<int32> FPooledLargeMemoryData::FreeListLength;
+
+FPooledLargeMemoryData::FPooledLargeMemoryData()
+{
+	Data = FreeList.Pop();
+	if (Data != nullptr)
+	{
+		--FreeListLength;
+	}
+	else
+	{
+		Data = new FLargeMemoryData();
+	}
+}
+
+FPooledLargeMemoryData::~FPooledLargeMemoryData()
+{
+	if (Data->GetSize() > GPooledLargeMemoryData_MaxElementSize)
+	{
+		delete Data;
+	}
+	else
+	{
+		if (FreeListLength.fetch_add(1) >= GPooledLargeMemoryData_MaxPoolLength)
+		{
+			// There's a chance of a race condition here if removal and addition are done at the same time.
+			// The outcome of that would be unnecessary destruction of Data that could be pooled,
+			// which is considered OK for the sake of simplicity.
+			--FreeListLength;
+			delete Data;
+		}
+		else
+		{
+			FreeList.Push(Data);
+		}
+	}
+}
+
diff --git a/Engine/Source/Runtime/Core/Public/Serialization/LargeMemoryData.h b/Engine/Source/Runtime/Core/Public/Serialization/LargeMemoryData.h
--- a/Engine/Source/Runtime/Core/Public/Serialization/LargeMemoryData.h
+++ b/Engine/Source/Runtime/Core/Public/Serialization/LargeMemoryData.h
@@ -74,3 +74,18 @@ private:
 	/** Resizes the data buffer to at least NumBytes with some slack */
 	void GrowBuffer();
 };
+
+/**
+* Pooled storage of FLargeMemoryData instances, allowing allocation-free and lock-free access.
+*/
+class CORE_API FPooledLargeMemoryData
+{
+public:
+	FPooledLargeMemoryData();
+	~FPooledLargeMemoryData();
+	FLargeMemoryData& Get() { return *Data; }
+private:
+	FLargeMemoryData* Data;
+	static TLockFreePointerListUnordered<FLargeMemoryData, 0> FreeList;
+	static std::atomic<int32> FreeListLength;
+};
\ No newline at end of file
diff --git a/Engine/Source/Runtime/CoreUObject/Private/UObject/UObjectGlobals.cpp b/Engine/Source/Runtime/CoreUObject/Private/UObject/UObjectGlobals.cpp
--- a/Engine/Source/Runtime/CoreUObject/Private/UObject/UObjectGlobals.cpp
+++ b/Engine/Source/Runtime/CoreUObject/Private/UObject/UObjectGlobals.cpp
@@ -2634,7 +2634,7 @@ UObject* StaticDuplicateObjectEx( FObjectDuplicationParameters& Parameters )
 		DupRootObject = StaticConstructObject_Internal(Params);
 	}
 
-	FLargeMemoryData ObjectData;
+	FPooledLargeMemoryData ObjectData;
 
 	FUObjectAnnotationSparse<FDuplicatedObject,false>  DuplicatedObjectAnnotation;
 
@@ -2656,7 +2656,7 @@ UObject* StaticDuplicateObjectEx( FObjectDuplicationParameters& Parameters )
 	// Read from the source object(s)
 	FDuplicateDataWriter Writer(
 		DuplicatedObjectAnnotation,				// Ref: Object annotation which stores the duplicated object for each source object
-		ObjectData,								// Out: Serialized object data
+		ObjectData.Get(),						// Out: Serialized object data
 		Parameters.SourceObject,				// Source object to copy
 		DupRootObject,							// Destination object to copy into
 		Parameters.FlagMask,					// Flags to be copied for duplicated objects
@@ -2687,7 +2687,7 @@ UObject* StaticDuplicateObjectEx( FObjectDuplicationParameters& Parameters )
 	};
 
 	TRefCountPtr<FUObjectSerializeContext> LoadContext(FUObjectThreadContext::Get().GetSerializeContext());
-	FDuplicateDataReader Reader(DuplicatedObjectAnnotation, ObjectData, Parameters.PortFlags, Parameters.DestOuter);
+	FDuplicateDataReader Reader(DuplicatedObjectAnnotation, ObjectData.Get(), Parameters.PortFlags, Parameters.DestOuter);
 	Reader.SetSerializeContext(LoadContext);
 	for(int32 ObjectIndex = 0;ObjectIndex < SerializedObjects.Num();ObjectIndex++)
 	{
