FLargeMemoryData buildfix: missing include

diff --git a/Engine/Source/Runtime/Core/Public/Serialization/LargeMemoryData.h b/Engine/Source/Runtime/Core/Public/Serialization/LargeMemoryData.h
--- a/Engine/Source/Runtime/Core/Public/Serialization/LargeMemoryData.h
+++ b/Engine/Source/Runtime/Core/Public/Serialization/LargeMemoryData.h
@@ -3,6 +3,7 @@
 #pragma once
 
 #include "CoreTypes.h"
+#include "Containers/LockFreeList.h"
 
 /**
 * Data storage for the large memory reader and writer.
