FLargeMemoryData NonUnity buildfix

diff --git a/Engine/Source/Runtime/Core/Private/Serialization/LargeMemoryData.cpp b/Engine/Source/Runtime/Core/Private/Serialization/LargeMemoryData.cpp
--- a/Engine/Source/Runtime/Core/Private/Serialization/LargeMemoryData.cpp
+++ b/Engine/Source/Runtime/Core/Private/Serialization/LargeMemoryData.cpp
@@ -3,6 +3,7 @@
 #include "Serialization/LargeMemoryData.h"
 #include "Logging/LogMacros.h"
 #include "CoreGlobals.h"
+#include "HAL/IConsoleManager.h"
 
 
 int32 GPooledLargeMemoryData_MaxPoolLength = 2;
