Logging: Flush logs when an assertion fails

The previous call to FlushThreadedLogs serialized buffered logs to the output devices, but did not flush the output devices. Recent changes to logging mean that Flush is now thread-safe and can be called here to flush the output devices too. This means that the log file, console, and other output devices will contain every log prior to the failed assertion, and will be visible when we break in the debugger.

diff --git a/Engine/Source/Runtime/Core/Private/Misc/AssertionMacros.cpp b/Engine/Source/Runtime/Core/Private/Misc/AssertionMacros.cpp
--- a/Engine/Source/Runtime/Core/Private/Misc/AssertionMacros.cpp
+++ b/Engine/Source/Runtime/Core/Private/Misc/AssertionMacros.cpp
@@ -517,7 +517,7 @@ void FORCENOINLINE FDebug::CheckVerifyFailedImpl(
 	{
 		// Flushing the logs here increases the likelihood that recent messages will be written to the log file, stdout and the debugger console.
 		// Without this, some of the recent messages may not be reported when debugger stops due to an assertion failure.
-		GLog->FlushThreadedLogs();
+		GLog->Flush();
 	}
 
 	if (!FPlatformMisc::IsDebuggerPresent())
