#ue5 - LWC: add missing VectorClamp() for doubles, it was only defined for floats, and also not in UnrealMathFPU.

diff --git a/Engine/Source/Runtime/Core/Private/Tests/Math/UnrealMathTest.cpp b/Engine/Source/Runtime/Core/Private/Tests/Math/UnrealMathTest.cpp
--- a/Engine/Source/Runtime/Core/Private/Tests/Math/UnrealMathTest.cpp
+++ b/Engine/Source/Runtime/Core/Private/Tests/Math/UnrealMathTest.cpp
@@ -1796,6 +1796,13 @@ bool RunDoubleVectorTest()
 	V0 = MakeVectorRegisterDouble(4.0f, 4.0f, 6.0f, 8.0f);
 	LogTest<double>(TEXT("VectorMax"), TestVectorsEqual(V0, V1));
 
+	V0 = MakeVectorRegisterDouble(-2.0f, 3.0f, -4.0f, 5.0f);
+	V1 = MakeVectorRegisterDouble( 4.0f, 4.0f,  6.0f, 8.0f);
+	V2 = MakeVectorRegisterDouble(-8.0f, 5.0f, -1.0f, 1.0f);
+	V2 = VectorClamp(V2, V0, V1);
+	V3 = MakeVectorRegisterDouble(-2.0f, 4.0f, -1.0f, 5.0f);
+	LogTest<double>(TEXT("VectorClamp"), TestVectorsEqual(V2, V3));
+
 	V0 = MakeVectorRegisterDouble(1.0, 3.0, 5.0, 7.0);
 	V1 = MakeVectorRegisterDouble(2.0, 4.0, 6.0, 8.0);
 	V1 = VectorCombineHigh(V0, V1);
@@ -2480,6 +2487,13 @@ bool FVectorRegisterAbstractionTest::RunTest(const FString& Parameters)
 	V0 = MakeVectorRegister( 4.0f, 4.0f, 6.0f, 8.0f );
 	LogTest<float>( TEXT("VectorMax"), TestVectorsEqual( V0, V1 ) );
 
+	V0 = MakeVectorRegister(-2.0f, 3.0f, -4.0f, 5.0f);
+	V1 = MakeVectorRegister( 4.0f, 4.0f,  6.0f, 8.0f);
+	V2 = MakeVectorRegister(-8.0f, 5.0f, -1.0f, 1.0f);
+	V2 = VectorClamp(V2, V0, V1);
+	V3 = MakeVectorRegister(-2.0f, 4.0f, -1.0f, 5.0f);
+	LogTest<float>(TEXT("VectorClamp"), TestVectorsEqual(V2, V3));
+
 	V0 = MakeVectorRegister(1.0f, 3.0f, 5.0f, 7.0f);
 	V1 = MakeVectorRegister(2.0f, 4.0f, 6.0f, 8.0f);
 	V1 = VectorCombineHigh(V0, V1);
diff --git a/Engine/Source/Runtime/Core/Public/Math/UnrealMathNeon.h b/Engine/Source/Runtime/Core/Public/Math/UnrealMathNeon.h
--- a/Engine/Source/Runtime/Core/Public/Math/UnrealMathNeon.h
+++ b/Engine/Source/Runtime/Core/Public/Math/UnrealMathNeon.h
@@ -2195,12 +2195,6 @@ FORCEINLINE VectorRegister4Double VectorMax(VectorRegister4Double Vec1, VectorRe
 	return Result;
 }
 
-FORCEINLINE VectorRegister4Float VectorClamp(const VectorRegister4Float& Vec1, const VectorRegister4Float& Vec2, const VectorRegister4Float& Vec3) 
-{
-	return VectorMin(VectorMax(Vec1, Vec2), Vec3);
-}
-
-
 /**
  * Merges the XYZ components of one vector with the W component of another vector and returns the result.
  *
diff --git a/Engine/Source/Runtime/Core/Public/Math/UnrealMathSSE.h b/Engine/Source/Runtime/Core/Public/Math/UnrealMathSSE.h
--- a/Engine/Source/Runtime/Core/Public/Math/UnrealMathSSE.h
+++ b/Engine/Source/Runtime/Core/Public/Math/UnrealMathSSE.h
@@ -2349,11 +2349,6 @@ FORCEINLINE VectorRegister4Double VectorMax(const VectorRegister4Double& Vec1, c
 	return Result;
 }
 
-FORCEINLINE VectorRegister4Float VectorClamp(const VectorRegister4Float& Vec1, const VectorRegister4Float& Vec2, const VectorRegister4Float& Vec3) 
-{
-	return _mm_min_ps(_mm_max_ps(Vec1, Vec2), Vec3);
-}
-
 /**
 * Creates a vector by combining two high components from each vector
 *
diff --git a/Engine/Source/Runtime/Core/Public/Math/UnrealMathVectorCommon.h b/Engine/Source/Runtime/Core/Public/Math/UnrealMathVectorCommon.h
--- a/Engine/Source/Runtime/Core/Public/Math/UnrealMathVectorCommon.h
+++ b/Engine/Source/Runtime/Core/Public/Math/UnrealMathVectorCommon.h
@@ -489,6 +489,17 @@ FORCEINLINE uint32 VectorAllLesserThan(VectorRegister4Double Vec1, VectorRegiste
 	return !VectorAnyGreaterThan(Vec1, Vec2);
 }
 
+/** Clamps X to be between VecMin and VecMax, inclusive. */
+FORCEINLINE VectorRegister4Float VectorClamp(const VectorRegister4Float& X, const VectorRegister4Float& VecMin, const VectorRegister4Float& VecMax)
+{
+	return VectorMin(VectorMax(X, VecMin), VecMax);
+}
+
+FORCEINLINE VectorRegister4Double VectorClamp(const VectorRegister4Double& X, const VectorRegister4Double& VecMin, const VectorRegister4Double& VecMax)
+{
+	return VectorMin(VectorMax(X, VecMin), VecMax);
+}
+
 /*----------------------------------------------------------------------------
 	VectorRegister specialization of templates.
 ----------------------------------------------------------------------------*/
