UnrealMathTest Color conversion tests: don't test NaN/inf when compiled with fast/finite math only

These compiler flags promise that NaN/Inf won't occur and allows the compiler to make transformations that break in the presence of these values. As long as we compile with this on, there's no expectation that these values behave in any particular way, and no sane way to test it.

diff --git a/Engine/Source/Runtime/Core/Private/Tests/Math/UnrealMathTest.cpp b/Engine/Source/Runtime/Core/Private/Tests/Math/UnrealMathTest.cpp
--- a/Engine/Source/Runtime/Core/Private/Tests/Math/UnrealMathTest.cpp
+++ b/Engine/Source/Runtime/Core/Private/Tests/Math/UnrealMathTest.cpp
@@ -4269,6 +4269,22 @@ bool FInitVectorTest::RunTest(const FString& Parameters)
 	return true;
 }
 
+// On GCC and Clang, setting -ffast-math or enabling some other unsafe floating point
+// optimizations set these #defines. As long as these compiler flags are active, the code
+// effectively promises that there are no infinites or NaNs, and although the original code
+// is written to handle these correctly, the compiler will happily apply transforms that
+// break it.
+// 
+// With these compiler flags on (and thus allowing the compiler to make transforms that
+// assume Inf/NaN don't occur), there is no expectation that any code can handle Inf/NaN
+// correctly (since the compiler is free to break whatever you write), so don't even test
+// it.
+#if defined(__FAST_MATH__) || defined(__FINITE_MATH_ONLY__)
+static constexpr bool GColorConversionsTestInfNaNs = false;
+#else
+static constexpr bool GColorConversionsTestInfNaNs = true;
+#endif
+
 IMPLEMENT_SIMPLE_AUTOMATION_TEST(FColorConversionTest, "System.Core.Math.ColorConversion", EAutomationTestFlags::ApplicationContextMask | EAutomationTestFlags::EngineFilter)
 bool FColorConversionTest::RunTest(const FString& Parameters)
 {
@@ -4350,15 +4366,19 @@ bool FColorConversionTest::RunTest(const FString& Parameters)
 		TestEqual(InTestName, ConvertedColor, ExpectedColor);
 	};
 
-	TestExtremalValue(TEXT("Extremal NaN"), NaN, 0);
-	TestExtremalValue(TEXT("Extremal -Inf"), -PosInf, 0);
+	if (GColorConversionsTestInfNaNs)
+	{
+		TestExtremalValue(TEXT("Extremal NaN"), NaN, 0);
+		TestExtremalValue(TEXT("Extremal -Inf"), -PosInf, 0);
+		TestExtremalValue(TEXT("Extremal +Inf"), PosInf, 255);
+	}
+
 	TestExtremalValue(TEXT("Extremal -MediumLarge"), -MediumLarge, 0);
 	TestExtremalValue(TEXT("Extremal 0"), 0.0f, 0);
 	TestExtremalValue(TEXT("Extremal +subnorm"), PosSubnormal, 0);
 	TestExtremalValue(TEXT("Extremal 1"), 1.0f, 255);
 	TestExtremalValue(TEXT("Extremal Mediumlarge"), MediumLarge, 255);
 	TestExtremalValue(TEXT("Extremal VeryLarge"), VeryLarge, 255);
-	TestExtremalValue(TEXT("Extremal +Inf"), PosInf, 255);
 
 	return true;
 }
@@ -4459,11 +4479,19 @@ bool FColorConversionHeavyTest::RunTest(const FString& Parameters)
 			{
 				uint32 CurrentBits = BucketIndex * ItemsInBucket + WithinBucketIndex;
 
+				// When asked to not test Infs/NaNs, skip that exponent entirely.
+				if (!GColorConversionsTestInfNaNs && (CurrentBits & 0x7f800000u) == 0x7f800000u)
+				{
+					continue;
+				}
+
 				// Don't put the same value in every color channel; we want to make sure we catch
 				// accidental channel swaps too! Three different inputs only; A is special anyway.
+				//
+				// XOR constants here are chosen to keep exponent bits same.
 				float R = FPlatformMath::AsFloat(CurrentBits);
 				float G = FPlatformMath::AsFloat(CurrentBits ^ 1234567u);
-				float B = FPlatformMath::AsFloat(CurrentBits ^ 89376762u);
+				float B = FPlatformMath::AsFloat(CurrentBits ^ 5490682u);
 				FLinearColor LinearInput(R, G, B, R);
 				FColor Converted = LinearInput.ToFColorSRGB();
 
@@ -4495,6 +4523,9 @@ bool FColorConversionHeavyTest::RunTest(const FString& Parameters)
 						Shared.FailedConverted = Converted;
 						Shared.FailedExpected = Expected;
 					}
+
+					// Break out of loop on first error
+					break;
 				}
 			}
 
