UE-149055 HoloLens FScheduler::ExecuteTask checkSlow(RunnableThread && RunnableThread->GetThreadPriority() == WorkerPriority);
-ThreadPriority and several other members were being shadowed by FHoloLensRunnableThread on top of FRunnableThread, confusion followed.  Several fixes to FWindowsRunnableThread were missing from FHoloLensRunnableThread
-This is a relatively minimal fix.  A better fix would be to refactor and put shared code into a FRunnableThreadMicrosoft from which this and windows could derive.  Unfortunately that runs into problem due to how another platform is using FRunnableThreadWindows.h along with delicacies about virtual functions, crash reporting, etc and it would all have to be refactored together.

diff --git a/Engine/Source/Runtime/Core/Private/HoloLens/HoloLensRunnableThread.cpp b/Engine/Source/Runtime/Core/Private/HoloLens/HoloLensRunnableThread.cpp
--- a/Engine/Source/Runtime/Core/Private/HoloLens/HoloLensRunnableThread.cpp
+++ b/Engine/Source/Runtime/Core/Private/HoloLens/HoloLensRunnableThread.cpp
@@ -13,7 +13,7 @@ uint32 FRunnableThreadHoloLens::GuardedRun()
 {
 	uint32 ExitCode = 1;
 
-	FPlatformProcess::SetThreadAffinityMask(ThreadAffintyMask);
+	FPlatformProcess::SetThreadAffinityMask(ThreadAffinityMask);
 
 #if !PLATFORM_SEH_EXCEPTIONS_DISABLED
 	if (!FPlatformMisc::IsDebuggerPresent() || GAlwaysReportCrash)
@@ -76,3 +76,25 @@ uint32 FRunnableThreadHoloLens::Run()
 
 	return ExitCode;
 }
+
+int FRunnableThreadHoloLens::TranslateThreadPriority(EThreadPriority Priority)
+{
+	// If this triggers, 
+	static_assert(TPri_Num == 7, "Need to add a case for new TPri_xxx enum value");
+
+	switch (Priority)
+	{
+	case TPri_AboveNormal:			return THREAD_PRIORITY_ABOVE_NORMAL;
+	case TPri_Normal:				return THREAD_PRIORITY_NORMAL;
+	case TPri_BelowNormal:			return THREAD_PRIORITY_BELOW_NORMAL;
+	case TPri_Highest:				return THREAD_PRIORITY_HIGHEST;
+	case TPri_TimeCritical:			return THREAD_PRIORITY_HIGHEST;
+	case TPri_Lowest:				return THREAD_PRIORITY_LOWEST;
+
+		// There is no such things as slightly below normal on Windows.
+		// This can't be below normal since we don't want latency sensitive task to go to efficient cores on Alder Lake.
+	case TPri_SlightlyBelowNormal:	return THREAD_PRIORITY_NORMAL;
+
+	default: UE_LOG(LogHAL, Fatal, TEXT("Unknown Priority passed to TranslateThreadPriority()")); return THREAD_PRIORITY_NORMAL;
+	}
+}
diff --git a/Engine/Source/Runtime/Core/Private/HoloLens/HoloLensRunnableThread.h b/Engine/Source/Runtime/Core/Private/HoloLens/HoloLensRunnableThread.h
--- a/Engine/Source/Runtime/Core/Private/HoloLens/HoloLensRunnableThread.h
+++ b/Engine/Source/Runtime/Core/Private/HoloLens/HoloLensRunnableThread.h
@@ -13,23 +13,16 @@
 #include "HoloLens/AllowWindowsPlatformTypes.h"
 
 /**
-* This is the base interface for all runnable thread classes. It specifies the
-* methods used in managing its life cycle.
+* This class has a lot of code that is very similar to FRunnableThreadWin.  It would make a lot of sense to refactor into an FRunnableThreadMicrosoft 
+* from which both this and FRunnableThreadWin could derive.  Unfortunately another platform is using the FRunnableThreadWin .h and providing its own
+* function definitions in a different .cpp and there are dire perforance, debugging, crash reporting, and exception handling concerns about the use of virtual
+* functions in this thread class.  So this refactor would have to be done carefully on all platforms.
 */
 class FRunnableThreadHoloLens
 	: public FRunnableThread
 {
 	/** The thread handle for the thread. */
-	HANDLE Thread;
-
-	/** Sync event to make sure that Init() has been completed before allowing the main thread to continue. */
-	FEventRef ThreadInitSyncEvent;
-
-	/** The priority to run the thread at. */
-	EThreadPriority ThreadPriority;
-
-	/** The Affinity to run the thread with. */
-	uint64 ThreadAffintyMask;
+	HANDLE Thread = 0;
 
 	/**
 	* Helper function to set thread names, visible by the debugger.
@@ -96,15 +89,6 @@ class FRunnableThreadHoloLens
 	uint32 Run();
 
 public:
-
-	FRunnableThreadHoloLens()
-		: Thread(NULL)
-		, ThreadInitSyncEvent()
-		, ThreadPriority(TPri_Normal)
-	{
-
-	}
-
 	~FRunnableThreadHoloLens()
 	{
 		// Clean up our thread if it is still active
@@ -121,10 +105,7 @@ public:
 		{
 			ThreadPriority = NewPriority;
 			// Change the priority on the thread
-			::SetThreadPriority(Thread,
-				ThreadPriority == TPri_AboveNormal ? THREAD_PRIORITY_ABOVE_NORMAL :
-				ThreadPriority == TPri_BelowNormal ? THREAD_PRIORITY_BELOW_NORMAL :
-				THREAD_PRIORITY_NORMAL);
+			::SetThreadPriority(Thread, TranslateThreadPriority(ThreadPriority));
 		}
 	}
 
@@ -174,6 +155,7 @@ public:
 		WaitForSingleObjectEx(Thread, INFINITE, FALSE);
 	}
 
+	static int TranslateThreadPriority(EThreadPriority Priority);
 protected:
 
 	virtual bool CreateInternal(FRunnable* InRunnable, const TCHAR* InThreadName,
@@ -183,12 +165,16 @@ protected:
 	{
 		check(InRunnable);
 		Runnable = InRunnable;
-		ThreadAffintyMask = InThreadAffinityMask;
-		ThreadName = InThreadName ? InThreadName : TEXT("Unnamed UE4");
+		ThreadAffinityMask = InThreadAffinityMask;
+
+		// Create a sync event to guarantee the Init() function is called first
+		ThreadInitSyncEvent = FPlatformProcess::GetSynchEventFromPool(true);
+
+		ThreadName = InThreadName ? InThreadName : TEXT("Unnamed UE");
 		ThreadPriority = InThreadPri;
 
-		// Create the new thread
-		Thread = CreateThread(NULL, InStackSize, _ThreadProc, this, STACK_SIZE_PARAM_IS_A_RESERVATION, (DWORD *)&ThreadID);
+		// Create the thread as suspended, so we can ensure ThreadId is initialized and the thread manager knows about the thread before it runs.
+		Thread = CreateThread(NULL, InStackSize, _ThreadProc, this, STACK_SIZE_PARAM_IS_A_RESERVATION | CREATE_SUSPENDED, (DWORD *)&ThreadID);
 
 		// If it fails, clear all the vars
 		if (Thread == NULL)
@@ -197,6 +183,8 @@ protected:
 		}
 		else
 		{
+			ResumeThread(Thread);
+
 			// Let the thread start up, then set the name for debug purposes.
 			ThreadInitSyncEvent->Wait(INFINITE);
 			SetThreadName(ThreadID, TCHAR_TO_ANSI(*ThreadName));
@@ -205,6 +193,9 @@ protected:
 			SetThreadPriority(InThreadPri);
 		}
 
+		// Cleanup the sync event
+		FPlatformProcess::ReturnSynchEventToPool(ThreadInitSyncEvent);
+		ThreadInitSyncEvent = nullptr;
 		return Thread != NULL;
 	}
 };
