Fix for niagara component renderer bound attributes getting trimmed

diff --git a/Engine/Plugins/FX/Niagara/Source/Niagara/Private/NiagaraComponentRendererProperties.cpp b/Engine/Plugins/FX/Niagara/Source/Niagara/Private/NiagaraComponentRendererProperties.cpp
--- a/Engine/Plugins/FX/Niagara/Source/Niagara/Private/NiagaraComponentRendererProperties.cpp
+++ b/Engine/Plugins/FX/Niagara/Source/Niagara/Private/NiagaraComponentRendererProperties.cpp
@@ -639,8 +639,8 @@ FText UNiagaraComponentRendererProperties::GetWidgetDisplayName() const
 
 TArray<FNiagaraVariable> UNiagaraComponentRendererProperties::GetBoundAttributes() const
 {
-	TArray<FNiagaraVariable> BoundAttributes;
-	BoundAttributes.Reserve(PropertyBindings.Num() + (bAssignComponentsOnParticleID ? 2 : 1));
+	TArray<FNiagaraVariable> BoundAttributes = Super::GetBoundAttributes();
+	BoundAttributes.Reserve(BoundAttributes.Num() + PropertyBindings.Num() + (bAssignComponentsOnParticleID ? 2 : 1));
 
 	BoundAttributes.Add(SYS_PARAM_PARTICLES_COMPONENTS_ENABLED);
 	if (bAssignComponentsOnParticleID)
diff --git a/Engine/Source/Runtime/Core/Private/UObject/FortniteMainBranchObjectVersion.cpp b/Engine/Source/Runtime/Core/Private/UObject/FortniteMainBranchObjectVersion.cpp
--- a/Engine/Source/Runtime/Core/Private/UObject/FortniteMainBranchObjectVersion.cpp
+++ b/Engine/Source/Runtime/Core/Private/UObject/FortniteMainBranchObjectVersion.cpp
@@ -11,7 +11,7 @@ TMap<FGuid, FGuid> FFortniteMainBranchObjectVersion::GetSystemGuids()
 	SystemGuids.Add(DevGuids.MATERIALSHADERMAP_DERIVEDDATA_VER, FGuid("1DA85DA5733C4A489A1BE117CDC5ACCC"));
 	SystemGuids.Add(DevGuids.NANITE_DERIVEDDATA_VER, FGuid("49FBF137DBFB4204B6C9E5B91DA15EBD"));
 	SystemGuids.Add(DevGuids.NIAGARASHADERMAP_DERIVEDDATA_VER, FGuid("2289D5116CF94BC0AFEAEC541468E645"));
-	SystemGuids.Add(DevGuids.Niagara_LatestScriptCompileVersion, FGuid("E3B6BEB506864EB697362DA7ADC799F1"));
+	SystemGuids.Add(DevGuids.Niagara_LatestScriptCompileVersion, FGuid("E93049C042704A758D6B65386328DC9F"));
 	SystemGuids.Add(DevGuids.SkeletalMeshDerivedDataVersion, FGuid("ACF593EAAE354FCBB34CF44F5AA1BFD2"));
 	SystemGuids.Add(DevGuids.STATICMESH_DERIVEDDATA_VER, FGuid("3006d21f867648a68cf4edee185446b7"));
 
