FileManager: Fixed CreateFileReader to respect the flags passed to it

diff --git a/Engine/Source/Runtime/Core/Private/HAL/FileManagerGeneric.cpp b/Engine/Source/Runtime/Core/Private/HAL/FileManagerGeneric.cpp
--- a/Engine/Source/Runtime/Core/Private/HAL/FileManagerGeneric.cpp
+++ b/Engine/Source/Runtime/Core/Private/HAL/FileManagerGeneric.cpp
@@ -55,7 +55,7 @@ FArchive* FFileManagerGeneric::CreateFileReaderInternal( const TCHAR* InFilename
 		}
 		return NULL;
 	}
-	return new FArchiveFileReaderGeneric( Handle, InFilename, Handle->Size(), BufferSize );
+	return new FArchiveFileReaderGeneric( Handle, InFilename, Handle->Size(), BufferSize, Flags );
 }
 
 /**
