moved `EExtendedTaskPriority` from the private to the public Tasks System API. This allows:
1) to use "inline" priority publicly, which is often helpful for perf
2) to use Tasks System with named threads (only if "tasks integration" is enabled - `TASKGRAPH_NEW_FRONTEND`)

diff --git a/Engine/Source/Runtime/Core/Private/Tests/Tasks/TasksTest.cpp b/Engine/Source/Runtime/Core/Private/Tests/Tasks/TasksTest.cpp
--- a/Engine/Source/Runtime/Core/Private/Tests/Tasks/TasksTest.cpp
+++ b/Engine/Source/Runtime/Core/Private/Tests/Tasks/TasksTest.cpp
@@ -340,19 +340,33 @@ namespace UE { namespace TasksTests
 			Outer.Launch(UE_SOURCE_LOCATION,
 				[&Outer]
 				{
-					auto InnerTaskBody = [&Outer] { check(!Outer.IsAwaitable()); /* still inside `Outer` execution */ };
-
-					// using the private API for `EExtendedTaskPriority::Inline` to force inner task execution inside the outer task
-					using namespace UE::Tasks::Private;
-					using FInnerTask = TExecutableTask<decltype(InnerTaskBody)>;
-
-					FInnerTask InnerTask{ UE_SOURCE_LOCATION, MoveTemp(InnerTaskBody), ETaskPriority::Default, EExtendedTaskPriority::Inline };
-					verify(InnerTask.TryLaunch()); // blocking execution of InnerTask due to its "inline" priority
+					// the inner task is executed "inline" (synchronously)
+					FTask Inner = Launch
+					(
+						UE_SOURCE_LOCATION, 
+						[&Outer] { check(!Outer.IsAwaitable()); /* still inside `Outer` execution */ }, 
+						ETaskPriority::Default, 
+						EExtendedTaskPriority::Inline
+					);
+					check(Inner.IsCompleted());
 				}
 			);
 			Outer.Wait();
 		}
 
+#if TASKGRAPH_NEW_FRONTEND
+		{	// a basic test for a named thread task
+			FTask GTTask = Launch
+			(
+				UE_SOURCE_LOCATION, 
+				[] { check(IsInGameThread()); }, 
+				ETaskPriority::Default, 
+				EExtendedTaskPriority::GameThreadNormalPri
+			);
+			GTTask.Wait();
+		}
+#endif
+
 		//{	// Cancelled task is completed only after nested
 		//	FTaskEvent ParentBlocker{ UE_SOURCE_LOCATION };
 		//	FTaskEvent NestedBlocker{ UE_SOURCE_LOCATION };
diff --git a/Engine/Source/Runtime/Core/Public/Async/TaskGraphInterfaces.h b/Engine/Source/Runtime/Core/Public/Async/TaskGraphInterfaces.h
--- a/Engine/Source/Runtime/Core/Public/Async/TaskGraphInterfaces.h
+++ b/Engine/Source/Runtime/Core/Public/Async/TaskGraphInterfaces.h
@@ -495,7 +495,7 @@ public:
 		}
 	}
 
-	void Init(const TCHAR* InDebugName, UE::Tasks::ETaskPriority InPriority, UE::Tasks::Private::EExtendedTaskPriority InExtendedPriority)
+	void Init(const TCHAR* InDebugName, UE::Tasks::ETaskPriority InPriority, UE::Tasks::EExtendedTaskPriority InExtendedPriority)
 	{
 		FTaskBase::Init(InDebugName, InPriority, InExtendedPriority);
 	}
@@ -567,7 +567,7 @@ public:
 	}
 
 	// task priority translation from the old API to the new API
-	static void TranslatePriority(ENamedThreads::Type ThreadType, UE::Tasks::ETaskPriority& OutPriority, UE::Tasks::Private::EExtendedTaskPriority& OutExtendedPriority)
+	static void TranslatePriority(ENamedThreads::Type ThreadType, UE::Tasks::ETaskPriority& OutPriority, UE::Tasks::EExtendedTaskPriority& OutExtendedPriority)
 	{
 		using namespace UE::Tasks;
 
@@ -575,21 +575,21 @@ public:
 		if (ThreadIndex != ENamedThreads::AnyThread)
 		{
 			check(ThreadIndex == ENamedThreads::GameThread || ThreadIndex == ENamedThreads::GetRenderThread() || ThreadIndex == ENamedThreads::RHIThread);
-			Private::EExtendedTaskPriority ConversionMap[] =
+			EExtendedTaskPriority ConversionMap[] =
 			{
-				Private::EExtendedTaskPriority::RHIThreadNormalPri,
-				Private::EExtendedTaskPriority::None, // invalid, maps from ENamedThreads::AudioThread
-				Private::EExtendedTaskPriority::GameThreadNormalPri,
-				Private::EExtendedTaskPriority::RenderThreadNormalPri
+				EExtendedTaskPriority::RHIThreadNormalPri,
+				EExtendedTaskPriority::None, // invalid, maps from ENamedThreads::AudioThread
+				EExtendedTaskPriority::GameThreadNormalPri,
+				EExtendedTaskPriority::RenderThreadNormalPri
 			};
 			OutExtendedPriority = ConversionMap[ThreadIndex - ENamedThreads::RHIThread];
-			OutExtendedPriority = (Private::EExtendedTaskPriority)((int32)OutExtendedPriority + (ENamedThreads::GetTaskPriority(ThreadType) != ENamedThreads::NormalTaskPriority ? 1 : 0));
-			OutExtendedPriority = (Private::EExtendedTaskPriority)((int32)OutExtendedPriority + (ENamedThreads::GetQueueIndex(ThreadType) != ENamedThreads::MainQueue ? 2 : 0));
+			OutExtendedPriority = (EExtendedTaskPriority)((int32)OutExtendedPriority + (ENamedThreads::GetTaskPriority(ThreadType) != ENamedThreads::NormalTaskPriority ? 1 : 0));
+			OutExtendedPriority = (EExtendedTaskPriority)((int32)OutExtendedPriority + (ENamedThreads::GetQueueIndex(ThreadType) != ENamedThreads::MainQueue ? 2 : 0));
 			OutPriority = ETaskPriority::Count;
 		}
 		else
 		{
-			OutExtendedPriority = Private::EExtendedTaskPriority::None;
+			OutExtendedPriority = EExtendedTaskPriority::None;
 			uint32 ThreadPriority = GetThreadPriorityIndex(ThreadType);
 			check(ThreadPriority < uint32(ENamedThreads::NumThreadPriorities));
 			ETaskPriority ConversionMap[int(ENamedThreads::NumThreadPriorities)] = { ETaskPriority::Normal, ETaskPriority::High, ETaskPriority::BackgroundNormal };
@@ -603,9 +603,9 @@ public:
 	}
 
 	// task priority translation from the new API to the old API
-	static ENamedThreads::Type TranslatePriority(UE::Tasks::Private::EExtendedTaskPriority Priority)
+	static ENamedThreads::Type TranslatePriority(UE::Tasks::EExtendedTaskPriority Priority)
 	{
-		using namespace UE::Tasks::Private;
+		using namespace UE::Tasks;
 
 		checkf(Priority >= EExtendedTaskPriority::GameThreadNormalPri && Priority <= EExtendedTaskPriority::RHIThreadHiPriLocalQueue, TEXT("only named threads can call this method: %d"), Priority);
 
@@ -662,7 +662,7 @@ public:
 			TTask* TaskObject = new(&Task->TaskStorage) TTask(Forward<T>(Args)...);
 
 			UE::Tasks::ETaskPriority Pri;
-			UE::Tasks::Private::EExtendedTaskPriority ExtPri;
+			UE::Tasks::EExtendedTaskPriority ExtPri;
 			TranslatePriority(TaskObject->GetDesiredThread(), Pri, ExtPri);
 
 			Task->Init(Pri, ExtPri);
@@ -701,7 +701,7 @@ private:
 	{
 	}
 
-	void Init(UE::Tasks::ETaskPriority InPriority, UE::Tasks::Private::EExtendedTaskPriority InExtendedPriority)
+	void Init(UE::Tasks::ETaskPriority InPriority, UE::Tasks::EExtendedTaskPriority InExtendedPriority)
 	{
 		FBaseGraphTask::Init(TEXT("GraphTask"), InPriority, InExtendedPriority);
 	}
@@ -732,7 +732,7 @@ public:
 	FGraphEventImpl()
 		: FBaseGraphTask(nullptr)
 	{
-		Init(TEXT("GraphEvent"), UE::Tasks::ETaskPriority::Normal, UE::Tasks::Private::EExtendedTaskPriority::TaskEvent);
+		Init(TEXT("GraphEvent"), UE::Tasks::ETaskPriority::Normal, UE::Tasks::EExtendedTaskPriority::TaskEvent);
 	}
 
 private:
diff --git a/Engine/Source/Runtime/Core/Public/Tasks/Pipe.h b/Engine/Source/Runtime/Core/Public/Tasks/Pipe.h
--- a/Engine/Source/Runtime/Core/Public/Tasks/Pipe.h
+++ b/Engine/Source/Runtime/Core/Public/Tasks/Pipe.h
@@ -52,12 +52,18 @@ namespace UE::Tasks
 		// @Priority - task priority, can affect task scheduling once it's passed the pipe
 		// @return Task instance that can be used to wait for task completion or to obtain the result of task execution
 		template<typename TaskBodyType>
-		TTask<TInvokeResult_T<TaskBodyType>> Launch(const TCHAR* InDebugName, TaskBodyType&& TaskBody, LowLevelTasks::ETaskPriority Priority = LowLevelTasks::ETaskPriority::Default)
+		TTask<TInvokeResult_T<TaskBodyType>> Launch
+		(
+			const TCHAR* InDebugName, 
+			TaskBodyType&& TaskBody, 
+			ETaskPriority Priority = ETaskPriority::Default,
+			EExtendedTaskPriority ExtendedPriority = EExtendedTaskPriority::None
+		)
 		{
 			using FResult = TInvokeResult_T<TaskBodyType>;
 			using FExecutableTask = Private::TExecutableTask<std::decay_t<TaskBodyType>>;
 
-			FExecutableTask* Task = FExecutableTask::Create(InDebugName, Forward<TaskBodyType>(TaskBody), Priority, Private::EExtendedTaskPriority::None);
+			FExecutableTask* Task = FExecutableTask::Create(InDebugName, Forward<TaskBodyType>(TaskBody), Priority, ExtendedPriority);
 			Task->SetPipe(*this);
 			Task->TryLaunch();
 			return TTask<FResult>{ Task };
@@ -70,12 +76,19 @@ namespace UE::Tasks
 		// @Priority - task priority, can affect task scheduling once it's passed the pipe
 		// @return Task instance that can be used to wait for task completion or to obtain the result of task execution
 		template<typename TaskBodyType, typename PrerequisitesCollectionType>
-		TTask<TInvokeResult_T<TaskBodyType>> Launch(const TCHAR* InDebugName, TaskBodyType&& TaskBody, PrerequisitesCollectionType&& Prerequisites, LowLevelTasks::ETaskPriority Priority = LowLevelTasks::ETaskPriority::Default)
+		TTask<TInvokeResult_T<TaskBodyType>> Launch
+		(
+			const TCHAR* InDebugName, 
+			TaskBodyType&& TaskBody, 
+			PrerequisitesCollectionType&& Prerequisites, 
+			ETaskPriority Priority = ETaskPriority::Default,
+			EExtendedTaskPriority ExtendedPriority = EExtendedTaskPriority::None
+		)
 		{
 			using FResult = TInvokeResult_T<TaskBodyType>;
 			using FExecutableTask = Private::TExecutableTask<std::decay_t<TaskBodyType>>;
 
-			FExecutableTask* Task = FExecutableTask::Create(InDebugName, Forward<TaskBodyType>(TaskBody), Priority, Private::EExtendedTaskPriority::None);
+			FExecutableTask* Task = FExecutableTask::Create(InDebugName, Forward<TaskBodyType>(TaskBody), Priority, ExtendedPriority);
 			Task->SetPipe(*this);
 			Task->AddPrerequisites(Forward<PrerequisitesCollectionType>(Prerequisites));
 			Task->TryLaunch();
diff --git a/Engine/Source/Runtime/Core/Public/Tasks/Task.h b/Engine/Source/Runtime/Core/Public/Tasks/Task.h
--- a/Engine/Source/Runtime/Core/Public/Tasks/Task.h
+++ b/Engine/Source/Runtime/Core/Public/Tasks/Task.h
@@ -135,13 +135,14 @@ namespace UE::Tasks
 			void Launch(
 				const TCHAR* DebugName,
 				TaskBodyType&& TaskBody,
-				ETaskPriority Priority = ETaskPriority::Normal
+				ETaskPriority Priority = ETaskPriority::Normal,
+				EExtendedTaskPriority ExtendedPriority = EExtendedTaskPriority::None
 			)
 			{
 				check(!IsValid());
 
 				using FExecutableTask = Private::TExecutableTask<std::decay_t<TaskBodyType>>;
-				FExecutableTask* Task = FExecutableTask::Create(DebugName, Forward<TaskBodyType>(TaskBody), Priority);
+				FExecutableTask* Task = FExecutableTask::Create(DebugName, Forward<TaskBodyType>(TaskBody), Priority, ExtendedPriority);
 				// this must happen before launching, to support an ability to access the task itself from inside it
 				*Pimpl.GetInitReference() = Task;
 				Task->TryLaunch();
@@ -159,13 +160,14 @@ namespace UE::Tasks
 				const TCHAR* DebugName,
 				TaskBodyType&& TaskBody,
 				PrerequisitesCollectionType&& Prerequisites,
-				ETaskPriority Priority = ETaskPriority::Normal
+				ETaskPriority Priority = ETaskPriority::Normal,
+				EExtendedTaskPriority ExtendedPriority = EExtendedTaskPriority::None
 			)
 			{
 				check(!IsValid());
 
 				using FExecutableTask = Private::TExecutableTask<std::decay_t<TaskBodyType>>;
-				FExecutableTask* Task = FExecutableTask::Create(DebugName, Forward<TaskBodyType>(TaskBody), Priority);
+				FExecutableTask* Task = FExecutableTask::Create(DebugName, Forward<TaskBodyType>(TaskBody), Priority, ExtendedPriority);
 				Task->AddPrerequisites(Forward<PrerequisitesCollectionType>(Prerequisites));
 				// this must happen before launching, to support an ability to access the task itself from inside it
 				*Pimpl.GetInitReference() = Task;
@@ -266,12 +268,13 @@ namespace UE::Tasks
 	TTask<TInvokeResult_T<TaskBodyType>> Launch(
 		const TCHAR* DebugName,
 		TaskBodyType&& TaskBody,
-		ETaskPriority Priority = ETaskPriority::Normal
+		ETaskPriority Priority = ETaskPriority::Normal,
+		EExtendedTaskPriority ExtendedPriority = EExtendedTaskPriority::None
 	)
 	{
 		using FResult = TInvokeResult_T<TaskBodyType>;
 		TTask<FResult> Task;
-		Task.Launch(DebugName, Forward<TaskBodyType>(TaskBody), Priority);
+		Task.Launch(DebugName, Forward<TaskBodyType>(TaskBody), Priority, ExtendedPriority);
 		return Task;
 	}
 
@@ -287,12 +290,13 @@ namespace UE::Tasks
 		const TCHAR* DebugName,
 		TaskBodyType&& TaskBody,
 		PrerequisitesCollectionType&& Prerequisites,
-		ETaskPriority Priority = ETaskPriority::Normal
+		ETaskPriority Priority = ETaskPriority::Normal,
+		EExtendedTaskPriority ExtendedPriority = EExtendedTaskPriority::None
 	)
 	{
 		using FResult = TInvokeResult_T<TaskBodyType>;
 		TTask<FResult> Task;
-		Task.Launch(DebugName, Forward<TaskBodyType>(TaskBody), Forward<PrerequisitesCollectionType>(Prerequisites), Priority);
+		Task.Launch(DebugName, Forward<TaskBodyType>(TaskBody), Forward<PrerequisitesCollectionType>(Prerequisites), Priority, ExtendedPriority);
 		return Task;
 	}
 
@@ -347,7 +351,7 @@ namespace UE::Tasks
 		auto WaitingTaskBody = [&CompletionEvent] { CompletionEvent->Trigger(); };
 		using FWaitingTask = Private::TExecutableTask<decltype(WaitingTaskBody)>;
 
-		FWaitingTask WaitingTask{ TEXT("Waiting Task"), MoveTemp(WaitingTaskBody), ETaskPriority::Default /* doesn't matter */, Private::EExtendedTaskPriority::Inline };
+		FWaitingTask WaitingTask{ TEXT("Waiting Task"), MoveTemp(WaitingTaskBody), ETaskPriority::Default /* doesn't matter */,  EExtendedTaskPriority::Inline };
 		WaitingTask.AddPrerequisites(Tasks);
 
 		if (WaitingTask.TryLaunch())
@@ -387,7 +391,7 @@ namespace UE::Tasks
 		auto WaitingTaskBody = [CompletionEvent] { CompletionEvent->Trigger(); };
 		using FWaitingTask = Private::TExecutableTask<decltype(WaitingTaskBody)>;
 
-		TRefCountPtr<FWaitingTask> WaitingTask{ FWaitingTask::Create(TEXT("Waiting Task"), MoveTemp(WaitingTaskBody), ETaskPriority::Default /* doesn't matter */, Private::EExtendedTaskPriority::Inline), /*bAddRef=*/ false};
+		TRefCountPtr<FWaitingTask> WaitingTask{ FWaitingTask::Create(TEXT("Waiting Task"), MoveTemp(WaitingTaskBody), ETaskPriority::Default /* doesn't matter */, EExtendedTaskPriority::Inline), /*bAddRef=*/ false};
 		WaitingTask->AddPrerequisites(Tasks);
 
 		if (WaitingTask->TryLaunch())
diff --git a/Engine/Source/Runtime/Core/Public/Tasks/TaskPrivate.h b/Engine/Source/Runtime/Core/Public/Tasks/TaskPrivate.h
--- a/Engine/Source/Runtime/Core/Public/Tasks/TaskPrivate.h
+++ b/Engine/Source/Runtime/Core/Public/Tasks/TaskPrivate.h
@@ -28,6 +28,32 @@ namespace UE::Tasks
 {
 	using LowLevelTasks::ETaskPriority;
 
+	// special task priorities for tasks that are never sent to the scheduler
+	enum class EExtendedTaskPriority
+	{
+		None,
+		Inline, // a task priority for "inline" task execution - a task is executed "inline" by the thread that unlocked it, w/o scheduling
+		TaskEvent, // a task priority used by task events, allows to shortcut task execution
+
+#if TASKGRAPH_NEW_FRONTEND
+		// for integration with named threads
+		GameThreadNormalPri,
+		GameThreadHiPri,
+		GameThreadNormalPriLocalQueue,
+		GameThreadHiPriLocalQueue,
+
+		RenderThreadNormalPri,
+		RenderThreadHiPri,
+		RenderThreadNormalPriLocalQueue,
+		RenderThreadHiPriLocalQueue,
+
+		RHIThreadNormalPri,
+		RHIThreadHiPri,
+		RHIThreadNormalPriLocalQueue,
+		RHIThreadHiPriLocalQueue
+#endif
+	};
+
 	class FPipe;
 
 	namespace Private
@@ -39,32 +65,6 @@ namespace UE::Tasks
 		// sets the current task and returns the previous current task
 		CORE_API FTaskBase* ExchangeCurrentTask(FTaskBase* Task);
 
-		// special task priorities for tasks that are never scheduled
-		enum class EExtendedTaskPriority
-		{
-			None,
-			Inline, // a task priority for "inline" task execution - a task is executed "inline" by the thread that unlocked it, w/o scheduling
-			TaskEvent, // a task priority used by task events, allows to shortcut task execution
-
-#if TASKGRAPH_NEW_FRONTEND
-			// for integration with named threads
-			GameThreadNormalPri,
-			GameThreadHiPri,
-			GameThreadNormalPriLocalQueue,
-			GameThreadHiPriLocalQueue,
-
-			RenderThreadNormalPri,
-			RenderThreadHiPri,
-			RenderThreadNormalPriLocalQueue,
-			RenderThreadHiPriLocalQueue,
-
-			RHIThreadNormalPri,
-			RHIThreadHiPri,
-			RHIThreadNormalPriLocalQueue,
-			RHIThreadHiPriLocalQueue
-#endif
-		};
-
 		// An abstract base class for task implementation. 
 		// Implements internal logic of task prerequisites, nested tasks and deep task retraction.
 		// Implements intrusive ref-counting and so can be used with TRefCountPtr.
