- FIxed compile error when DISABLE_GENERATED_INI_WHEN_COOKED is set

diff --git a/Engine/Source/Runtime/Core/Private/Misc/ConfigContext.cpp b/Engine/Source/Runtime/Core/Private/Misc/ConfigContext.cpp
--- a/Engine/Source/Runtime/Core/Private/Misc/ConfigContext.cpp
+++ b/Engine/Source/Runtime/Core/Private/Misc/ConfigContext.cpp
@@ -269,13 +269,13 @@ bool FConfigContext::PerformLoad()
 	else
 	{
 #if DISABLE_GENERATED_INI_WHEN_COOKED
-		if (Context.BaseIniName == TEXT("GameUserSettings"))
+		if (BaseIniName == TEXT("GameUserSettings"))
 		{
 			// If we asked to disable ini when cooked, disable all ini files except GameUserSettings, which stores user preferences
-			Context.bAllowGeneratedIniWhenCooked = false;
+			bAllowGeneratedIniWhenCooked = false;
 			if (FPlatformProperties::RequiresCookedData())
 			{
-				Context.ConfigFile->NoSave = true;
+				ConfigFile->NoSave = true;
 			}
 		}
 #endif
