- FIxed compile error

diff --git a/Engine/Source/Runtime/Core/Private/Misc/ConfigCacheIni.cpp b/Engine/Source/Runtime/Core/Private/Misc/ConfigCacheIni.cpp
--- a/Engine/Source/Runtime/Core/Private/Misc/ConfigCacheIni.cpp
+++ b/Engine/Source/Runtime/Core/Private/Misc/ConfigCacheIni.cpp
@@ -4388,7 +4388,7 @@ void FConfigCacheIni::InitializeConfigSystem()
 #else
 
 	// load in the default ini files
-	bool bEngineConfigCreated = GConfig->InitializeKnownConfigFiles(nullptr, bDefaultEngineIniRequired);
+	bool bEngineConfigCreated = GConfig->InitializeKnownConfigFiles(nullptr, false);
 #endif
 
 	// verify if needed
