ForwardAsTuple added.

diff --git a/Engine/Source/Runtime/Core/Public/Templates/Tuple.h b/Engine/Source/Runtime/Core/Public/Templates/Tuple.h
--- a/Engine/Source/Runtime/Core/Public/Templates/Tuple.h
+++ b/Engine/Source/Runtime/Core/Public/Templates/Tuple.h
@@ -49,6 +49,9 @@ class FArchive;
 template <typename... Types>
 struct TTuple;
 
+template <typename... Types>
+TTuple<std::decay_t<Types>...> MakeTuple(Types&&... Args);
+
 namespace UE::Core::Private::Tuple
 {
 	enum EForwardingConstructor { ForwardingConstructor };
@@ -521,10 +524,10 @@ namespace UE::Core::Private::Tuple
 		(void)Temp;
 	}
 
-	template <typename... Types>
-	FORCEINLINE TTuple<typename TDecay<Types>::Type...> MakeTupleImpl(Types&&... Args)
+	template <typename... ElementTypes, typename... Types>
+	FORCEINLINE TTuple<ElementTypes...> MakeTupleImpl(Types&&... Args)
 	{
-		return TTuple<typename TDecay<Types>::Type...>(Forward<Types>(Args)...);
+		return TTuple<ElementTypes...>(Forward<Types>(Args)...);
 	}
 
 	template <typename IntegerSequence>
@@ -536,7 +539,7 @@ namespace UE::Core::Private::Tuple
 		template <typename TupleType, typename FuncType>
 		static decltype(auto) Do(TupleType&& Tuple, FuncType Func)
 		{
-			return MakeTupleImpl(Func(Forward<TupleType>(Tuple).template Get<Indices>())...);
+			return MakeTuple(Func(Forward<TupleType>(Tuple).template Get<Indices>())...);
 		}
 	};
 
@@ -856,9 +859,34 @@ using TTupleElement = UE::Core::Private::Tuple::TCVTupleElement<Index, const vol
  * }
  */
 template <typename... Types>
-FORCEINLINE TTuple<typename TDecay<Types>::Type...> MakeTuple(Types&&... Args)
+FORCEINLINE TTuple<std::decay_t<Types>...> MakeTuple(Types&&... Args)
+{
+	return UE::Core::Private::Tuple::MakeTupleImpl<std::decay_t<Types>...>(Forward<Types>(Args)...);
+}
+
+
+/**
+ * Makes a TTuple from some arguments.  Unlike MakeTuple, the TTuple element types are references and retain the
+ * same value category of the arguments, like the Forward function.
+ *
+ * @param  Args  The arguments used to construct the tuple.
+ * @return A tuple containing forwarded references to the arguments.
+ *
+ * Example:
+ *
+ * template <typename... Ts>
+ * void Foo(const TTuple<Ts...>&);
+ *
+ * void Func(const int32 A, FString&& B)
+ * {
+ *     // Calls Foo<const int32&, const TCHAR(&)[6], FString&&>(...);
+ *     Foo(ForwardAsTuple(A, TEXT("Hello"), MoveTemp(B)));
+ * }
+ */
+template <typename... Types>
+FORCEINLINE TTuple<Types&&...> ForwardAsTuple(Types&&... Args)
 {
-	return UE::Core::Private::Tuple::MakeTupleImpl(Forward<Types>(Args)...);
+	return UE::Core::Private::Tuple::MakeTupleImpl<Types&&...>(Forward<Types>(Args)...);
 }
 
 
