removed "BlockingWorkers" TaskGraph unit test as it's not relevant anymore

diff --git a/Engine/Source/Runtime/Core/Private/Tests/Async/TaskGraphTest.cpp b/Engine/Source/Runtime/Core/Private/Tests/Async/TaskGraphTest.cpp
--- a/Engine/Source/Runtime/Core/Private/Tests/Async/TaskGraphTest.cpp
+++ b/Engine/Source/Runtime/Core/Private/Tests/Async/TaskGraphTest.cpp
@@ -1427,64 +1427,6 @@ PRAGMA_ENABLE_DEPRECATION_WARNINGS
 
 		return true;
 	}
-
-	IMPLEMENT_SIMPLE_AUTOMATION_TEST(FTaskGraphBlockingWorkersTest, "System.Core.TaskGraph.BlockingWorkers", EAutomationTestFlags::ApplicationContextMask | EAutomationTestFlags::EngineFilter | EAutomationTestFlags::Disabled);
-
-	template<uint32 Num, bool bBackgroundWorkersToo>
-	void BlockingWorkersByFEvent();
-
-	bool FTaskGraphBlockingWorkersTest::RunTest(const FString& Parameters)
-	{
-		FPlatformProcess::Sleep(1.0f);
-		UE_BENCHMARK(10, BlockingWorkersByFEvent<100, false>);
-		UE_BENCHMARK(10, BlockingWorkersByFEvent<100, true>);
-		return true;
-	}
-
-	template<uint32 Num, bool bBackgroundWorkersToo>
-	void BlockingWorkersByFEvent()
-	{
-		for (int N = 0; N != Num; ++N)
-		{
-			int32 NumWorkers = bBackgroundWorkersToo ? LowLevelTasks::FScheduler::Get().GetNumWorkers() : GNumForegroundWorkers;
-
-			TArray<LowLevelTasks::FTask> WorkerBlockers; // tasks that block worker threads
-			WorkerBlockers.AddDefaulted(NumWorkers);
-
-			FEventRef BlockersBlocker{ EEventMode::ManualReset }; // blocks `WorkerBlockers`, manual reset because multiple threads are waiting for it
-
-			std::atomic<int32> NumWorkersNotBlocked{ NumWorkers };
-			FEventRef AllWorkersBlockedSignal;
-
-			for (int i = 0; i != NumWorkers; ++i)
-			{
-				WorkerBlockers[i].Init(TEXT("BlockingWorkers"), LowLevelTasks::ETaskPriority::High,
-					[i, &BlockersBlocker, &AllWorkersBlockedSignal, &NumWorkersNotBlocked]
-					{
-						verify(LowLevelTasks::FSchedulerTls::GetAffinityIndex() == i);
-						if (--NumWorkersNotBlocked == 0)
-						{
-							AllWorkersBlockedSignal->Trigger();
-						}
-						else
-						{
-							BlockersBlocker->Wait();
-						}
-					});
-
-				verify(LowLevelTasks::TryLaunchAffinity(WorkerBlockers[i], i));
-			}
-
-			verify(AllWorkersBlockedSignal->Wait(FTimespan::FromSeconds(3)));
-
-			BlockersBlocker->Trigger();
-			for (int i = 0; i != NumWorkers; ++i)
-			{
-				while(!WorkerBlockers[i].IsCompleted())
-				{}
-			}
-		}
-	}
 }
 
 #endif //WITH_DEV_AUTOMATION_TESTS
