New Version: LevelInstanceActorGuidSerialize

diff --git a/Engine/Source/Runtime/Core/Public/UObject/FortniteMainBranchObjectVersion.h b/Engine/Source/Runtime/Core/Public/UObject/FortniteMainBranchObjectVersion.h
--- a/Engine/Source/Runtime/Core/Public/UObject/FortniteMainBranchObjectVersion.h
+++ b/Engine/Source/Runtime/Core/Public/UObject/FortniteMainBranchObjectVersion.h
@@ -208,6 +208,9 @@ struct CORE_API FFortniteMainBranchObjectVersion
 		// Migrated event definitions to the CDO/archetype data
 		MigratedEventDefinitionToDefaults,
 
+		// Serialize LevelInstanceActorGuid on new ILevelInstanceInterface implementation
+		LevelInstanceActorGuidSerialize,
+
 		// -----<new versions can be added above this line>-------------------------------------------------
 		VersionPlusOne,
 		LatestVersion = VersionPlusOne - 1
