Avoid return value is ignored warnings when DO_ENSURE is false or WITH_CODE_ANALYSIS is true

diff --git a/Engine/Source/Runtime/Core/Public/Misc/AssertionMacros.h b/Engine/Source/Runtime/Core/Public/Misc/AssertionMacros.h
--- a/Engine/Source/Runtime/Core/Public/Misc/AssertionMacros.h
+++ b/Engine/Source/Runtime/Core/Public/Misc/AssertionMacros.h
@@ -359,10 +359,10 @@ RetType FORCENOINLINE UE_DEBUG_SECTION DispatchCheckVerify(InnerType&& Inner, Ar
 
 #else	// DO_ENSURE
 
-	#define ensure(           InExpression                ) (!!(InExpression))
-	#define ensureMsgf(       InExpression, InFormat, ... ) (!!(InExpression))
-	#define ensureAlways(     InExpression                ) (!!(InExpression))
-	#define ensureAlwaysMsgf( InExpression, InFormat, ... ) (!!(InExpression))
+	#define ensure(           InExpression                ) [&](){return (!!(InExpression));}()
+	#define ensureMsgf(       InExpression, InFormat, ... ) [&](){return (!!(InExpression));}()
+	#define ensureAlways(     InExpression                ) [&](){return (!!(InExpression));}()
+	#define ensureAlwaysMsgf( InExpression, InFormat, ... ) [&](){return (!!(InExpression));}()
 
 #endif	// DO_CHECK
 
