Logging: Added a timing event scope to FlushThreadedLogs

#rnx

diff --git a/Engine/Source/Runtime/Core/Private/Misc/OutputDeviceRedirector.cpp b/Engine/Source/Runtime/Core/Private/Misc/OutputDeviceRedirector.cpp
--- a/Engine/Source/Runtime/Core/Private/Misc/OutputDeviceRedirector.cpp
+++ b/Engine/Source/Runtime/Core/Private/Misc/OutputDeviceRedirector.cpp
@@ -441,6 +441,7 @@ void FOutputDeviceRedirector::FlushThreadedLogs(EOutputDeviceRedirectorFlushOpti
 	{
 		if (!EnumHasAnyFlags(Options, EOutputDeviceRedirectorFlushOptions::Async))
 		{
+			TRACE_CPUPROFILER_EVENT_SCOPE(FOutputDeviceRedirector::FlushThreadedLogs);
 			FEventRef IdleEvent(EEventMode::ManualReset);
 			if (State->ThreadIdleEvents.EnqueueAndReturnWasEmpty(IdleEvent.Get()))
 			{
