Add TTiedTupleAdder which allows algorithms as Algo::Transform to output to multiple containers

diff --git a/Engine/Source/Runtime/Core/Private/Templates/TuplesTest.cpp b/Engine/Source/Runtime/Core/Private/Templates/TuplesTest.cpp
--- a/Engine/Source/Runtime/Core/Private/Templates/TuplesTest.cpp
+++ b/Engine/Source/Runtime/Core/Private/Templates/TuplesTest.cpp
@@ -240,6 +240,13 @@ static_assert(TTupleIndex<double, const volatile TTuple<double, float&, char&&>>
 static_assert(TTupleIndex<float&, const volatile TTuple<double, float&, char&&>>::Value == 1, "");
 static_assert(TTupleIndex<char&&, const volatile TTuple<double, float&, char&&>>::Value == 2, "");
 
+// Check that TIsTuple works for values, lvalue references and rvalue references, even with qualifiers on the tuple
+static_assert(!TIsTuple<double>::Value										, "");
+static_assert(TIsTuple<               TTuple<double, float&, char&&>>::Value, "");
+static_assert(TIsTuple<const          TTuple<double, float&, char&&>>::Value, "");
+static_assert(TIsTuple<      volatile TTuple<double, float&, char&&>>::Value, "");
+static_assert(TIsTuple<const volatile TTuple<double, float&, char&&>>::Value, "");
+
 // These shouldn't compile - ideally giving a meaningful error message
 #if 0
 	// TTupleElement passed a non-tuple
@@ -256,4 +263,7 @@ static_assert(TTupleIndex<char&&, const volatile TTuple<double, float&, char&&>>
 
 	// Type appears multiple times in tuple
 	static_assert(TTupleIndex<int, TTuple<int, float&, int>>::Value == 0, "");
+
+	// Multiple arguments
+	static_assert(TIsTuple<double, double>::Value, "");
 #endif
diff --git a/Engine/Source/Runtime/Core/Public/Algo/TiedTupleOutput.h b/Engine/Source/Runtime/Core/Public/Algo/TiedTupleOutput.h
new file mode 100644
--- /dev/null
+++ b/Engine/Source/Runtime/Core/Public/Algo/TiedTupleOutput.h
@@ -0,0 +1,56 @@
+﻿// Copyright Epic Games, Inc. All Rights Reserved.
+
+#pragma once
+
+#include "CoreTypes.h"
+#include "Templates/Tuple.h"
+
+namespace Algo
+{
+	namespace Private
+	{
+		template <typename... OutputTypes>
+		struct TTiedTupleAdder
+		{
+			template <typename TupleType>
+			void Add(TupleType&& Values)
+			{
+				static_assert(TIsTuple<TupleType>::Value);
+				static_assert(sizeof...(OutputTypes) == TTupleArity<typename TDecay<TupleType>::Type>::Value, "Wrong length of tuple for output");
+
+				VisitTupleElements(
+					[](auto& Output, auto&& Value)
+					{
+						Output.Add(Forward<decltype(Value)>(Value));
+					},
+					Outputs,
+					Forward<TupleType>(Values)
+				);
+			}
+
+			TTuple<OutputTypes&...> Outputs;
+		};
+	}
+	
+
+	/**
+	 * Ties n objects with an Add function, usually containers, into one object with an Add function accepting an n-tuple
+	 * that forwards the Add calls.
+	 *
+	 * This is useful for algorithms such as Algo::Transform and Algo::Copy.
+	 *
+	 * Example:
+	 *	TArray<int32> Out1;
+	 *	TArray<FString> Out2;
+	 *	TieTupleAdd(Out1, Out2).Add(MakeTuple(42, TEXT("42")));
+	 *	// Out1 = { 42 }, Out2 = { "42" }
+	 *
+	 *	@param Outputs Objects with an Add() function
+	 *	@return An object with an Add function that will call Add on the tied elements
+	 */
+	template <typename... OutputTypes>
+	Private::TTiedTupleAdder<OutputTypes...> TieTupleAdd(OutputTypes&... Outputs)
+	{
+		return { { Outputs... } };
+	}
+}
\ No newline at end of file
diff --git a/Engine/Source/Runtime/Core/Public/Algo/Transform.h b/Engine/Source/Runtime/Core/Public/Algo/Transform.h
--- a/Engine/Source/Runtime/Core/Public/Algo/Transform.h
+++ b/Engine/Source/Runtime/Core/Public/Algo/Transform.h
@@ -5,6 +5,34 @@
 #include "CoreTypes.h"
 #include "Templates/Invoke.h"
 
+/*
+ * Inspired by std::transform.
+ *
+ * Simple example:
+ *	TArray<FString> Out;
+ *	TArray<int32> Inputs { 1, 2, 3, 4, 5, 6 };
+ *	Algo::Transform(
+ *		Inputs,
+ *		Out,
+ *		[](int32 Input) { return LexToString(Input); }
+ *	);
+ *	// Out1 == [ "1", "2", "3", "4", "5", "6" ]
+ *
+ * You can also use Transform to output to multiple targets that have an Add function, by using :
+ * 
+ *	TArray<FString> Out1;
+ *	TArray<int32> Out2;
+ *
+ *	TArray<int32> Inputs = { 1, 2, 3, 4, 5, 6 };
+ *	Algo::Transform(
+ *		Inputs,
+ *		TiedTupleAdd(Out1, Out2),
+ *		[](int32 Input) { return ForwardAsTuple(LexToString(Input), Input * Input); }
+ *	);
+ *
+ *	// Out1 == [ "1", "2", "3", "4", "5", "6" ]
+ *	// Out2 == [ 1, 4, 9, 16, 25, 36 ]
+ */
 namespace Algo
 {
 	/**
@@ -16,7 +44,7 @@ namespace Algo
 	 * @param  Trans      Transformation operation
 	 */
 	template <typename InT, typename OutT, typename PredicateT, typename TransformT>
-	FORCEINLINE void TransformIf(const InT& Input, OutT& Output, PredicateT Predicate, TransformT Trans)
+	FORCEINLINE void TransformIf(const InT& Input, OutT&& Output, PredicateT Predicate, TransformT Trans)
 	{
 		for (const auto& Value : Input)
 		{
@@ -35,7 +63,7 @@ namespace Algo
 	 * @param  Trans   Transformation operation
 	 */
 	template <typename InT, typename OutT, typename TransformT>
-	FORCEINLINE void Transform(const InT& Input, OutT& Output, TransformT Trans)
+	FORCEINLINE void Transform(const InT& Input, OutT&& Output, TransformT Trans)
 	{
 		for (const auto& Value : Input)
 		{
diff --git a/Engine/Source/Runtime/Core/Public/Templates/Tuple.h b/Engine/Source/Runtime/Core/Public/Templates/Tuple.h
--- a/Engine/Source/Runtime/Core/Public/Templates/Tuple.h
+++ b/Engine/Source/Runtime/Core/Public/Templates/Tuple.h
@@ -989,3 +989,16 @@ public:
 	using type = typename TTupleElement<N, TTuple<ArgTypes...>>::Type;
 };
 #endif
+
+template <typename T> constexpr bool TIsTuple_V = false;
+
+template <typename... Types> constexpr bool TIsTuple_V<               TTuple<Types...>> = true;
+template <typename... Types> constexpr bool TIsTuple_V<const          TTuple<Types...>> = true;
+template <typename... Types> constexpr bool TIsTuple_V<      volatile TTuple<Types...>> = true;
+template <typename... Types> constexpr bool TIsTuple_V<const volatile TTuple<Types...>> = true;
+
+template <typename T>
+struct TIsTuple
+{
+	enum { Value = TIsTuple_V<T> };
+};
\ No newline at end of file
