Fix memory mapping a full file generally going past the file size when trying to grab some aligned memory

diff --git a/Engine/Source/Runtime/Core/Private/Unix/UnixPlatformFile.cpp b/Engine/Source/Runtime/Core/Private/Unix/UnixPlatformFile.cpp
--- a/Engine/Source/Runtime/Core/Private/Unix/UnixPlatformFile.cpp
+++ b/Engine/Source/Runtime/Core/Private/Unix/UnixPlatformFile.cpp
@@ -698,10 +698,12 @@ public:
 		, AlignedSize(InAlignedSize)
 	{
 	}
+
+	virtual ~FUnixMappedFileRegion();
 };
 
 class FUnixMappedFileHandle final : public IMappedFileHandle
-{	
+{
 public:
 	FUnixMappedFileHandle(int InFileHandle, int64 FileSize, const FString& InFilename)
 		: IMappedFileHandle(FileSize)
@@ -718,7 +720,7 @@ public:
 		check(!NumOutstandingRegions); // can't delete the file before you delete all outstanding regions
 		close(FileHandle);
 	}
-	
+
 	virtual IMappedFileRegion* MapRegion(int64 Offset = 0, int64 BytesToMap = MAX_int64, bool bPreloadHint = false) override
 	{
 		LLM_PLATFORM_SCOPE(ELLMTag::PlatformMMIO);
@@ -727,13 +729,12 @@ public:
 		check(BytesToMap > 0); // don't map zero bytes
 
 		const int64 AlignedOffset = AlignDown(Offset, Alignment);
-		const int64 AlignedSize = Align(BytesToMap + Offset - AlignedOffset, Alignment);
-		
-		// if we are about to go off the end, let's not
+		int64 AlignedSize = Align(BytesToMap + Offset - AlignedOffset, Alignment);
+
+		// if we are about to go off the end, let's just go to the end of the file
 		if (AlignedOffset + AlignedSize > GetFileSize())
 		{
-			UE_LOG(LogUnixPlatformFile, Warning, TEXT("Mapping fell off the end, did we need to actually abort? [%lld + %lld > %lld]"), AlignedOffset, AlignedSize, GetFileSize());
-			return nullptr;
+			AlignedSize = GetFileSize() - AlignedOffset;
 		}
 
 		int Flags = MAP_PRIVATE;
@@ -741,7 +742,7 @@ public:
 		{
 			Flags |= MAP_POPULATE;
 		}
-		
+
 		const uint8* AlignedMapPtr = static_cast<const uint8*>(mmap(nullptr, AlignedSize, PROT_READ, Flags, FileHandle, AlignedOffset));
 		if (AlignedMapPtr == MAP_FAILED || AlignedMapPtr == nullptr)
 		{
@@ -756,13 +757,13 @@ public:
 		NumOutstandingRegions++;
 		return Result;
 	}
-	
+
 	void UnMap(const FUnixMappedFileRegion* Region)
 	{
 		LLM_PLATFORM_SCOPE(ELLMTag::PlatformMMIO);
 		check(NumOutstandingRegions > 0);
 		NumOutstandingRegions--;
-		
+
 		LLM(FLowLevelMemTracker::Get().OnLowLevelFree(ELLMTracker::Platform, Region->AlignedPtr));
 		const int Res = munmap(const_cast<uint8*>(Region->AlignedPtr), Region->AlignedSize);
 		checkf(Res == 0, TEXT("Failed to unmap, error is %d, errno is %d [params: %x, %d]"), Res, errno, MappedPtr, GetFileSize());
@@ -776,6 +777,11 @@ private:
 	int FileHandle;
 };
 
+FUnixMappedFileRegion::~FUnixMappedFileRegion()
+{
+	Parent->UnMap(this);
+}
+
 /**
  * Unix File I/O implementation
 **/
