Fixing deprecation warning in HoloLens.

diff --git a/Engine/Source/Runtime/Core/Private/HoloLens/HoloLensProcess.cpp b/Engine/Source/Runtime/Core/Private/HoloLens/HoloLensProcess.cpp
--- a/Engine/Source/Runtime/Core/Private/HoloLens/HoloLensProcess.cpp
+++ b/Engine/Source/Runtime/Core/Private/HoloLens/HoloLensProcess.cpp
@@ -113,7 +113,9 @@ FEvent* FHoloLensProcess::CreateSynchEvent(bool bIsManualReset /*= false*/)
 		Event = new FSingleThreadEvent();
 	}
 	// If the internal create fails, delete the instance and return NULL
+	PRAGMA_DISABLE_DEPRECATION_WARNINGS
 	if (!Event->Create(bIsManualReset))
+	PRAGMA_ENABLE_DEPRECATION_WARNINGS
 	{
 		delete Event;
 		Event = NULL;
