Fix FName debug visualizers on certain clang platforms where the 19725380 fix didn't work

diff --git a/Engine/Source/Runtime/Core/Private/Misc/CoreGlobals.cpp b/Engine/Source/Runtime/Core/Private/Misc/CoreGlobals.cpp
--- a/Engine/Source/Runtime/Core/Private/Misc/CoreGlobals.cpp
+++ b/Engine/Source/Runtime/Core/Private/Misc/CoreGlobals.cpp
@@ -248,9 +248,6 @@ uint8** CORE_API GNameBlocksDebug = FNameDebugVisualizer::GetBlocks();
 FChunkedFixedUObjectArray*& CORE_API GObjectArrayForDebugVisualizers = GCoreObjectArrayForDebugVisualizers;
 UE::ObjectPath::Private::FStoredObjectPath*& GComplexObjectPathDebug = GCoreComplexObjectPathDebug;
 FObjectHandlePackageDebugData*& CORE_API GObjectHandlePackageDebug = GCoreObjectHandlePackageDebug;
-
-// Create a global to make sure FNameEntry symbol is around which is needed to dump FName info when debugging
-FNameEntry* CORE_API GFNameEntryPreserveSymbol = nullptr;
 #endif
 
 /** Game name, used for base game directory and ini among other things										*/
diff --git a/Engine/Source/Runtime/Core/Private/UObject/UnrealNames.cpp b/Engine/Source/Runtime/Core/Private/UObject/UnrealNames.cpp
--- a/Engine/Source/Runtime/Core/Private/UObject/UnrealNames.cpp
+++ b/Engine/Source/Runtime/Core/Private/UObject/UnrealNames.cpp
@@ -1962,6 +1962,10 @@ static uint16 GetRawNonCasePreservingHash(const TCharType* Source)
 	FNameEntry
 -----------------------------------------------------------------------------*/
 
+// Make Clang keep FNameEntry type debug info needed for debug visualizers
+struct FClangKeepDebugInfo {};
+FNameEntry::FNameEntry(FClangKeepDebugInfo) {}
+
 void FNameEntry::StoreName(const ANSICHAR* InName, uint32 Len)
 {
 	FPlatformMemory::Memcpy(AnsiName, InName, sizeof(ANSICHAR) * Len);
diff --git a/Engine/Source/Runtime/Core/Public/UObject/NameTypes.h b/Engine/Source/Runtime/Core/Public/UObject/NameTypes.h
--- a/Engine/Source/Runtime/Core/Public/UObject/NameTypes.h
+++ b/Engine/Source/Runtime/Core/Public/UObject/NameTypes.h
@@ -231,6 +231,7 @@ private:
 #endif // UE_FNAME_OUTLINE_NUMBER
 	};
 
+	FNameEntry(struct FClangKeepDebugInfo);
 	FNameEntry(const FNameEntry&) = delete;
 	FNameEntry(FNameEntry&&) = delete;
 	FNameEntry& operator=(const FNameEntry&) = delete;
