Tasks System: minor improvement to deadlock detection. now detects waiting for task completion from inside that task execution, even if inside its inner task.

diff --git a/Engine/Source/Runtime/Core/Public/Tasks/TaskPrivate.h b/Engine/Source/Runtime/Core/Public/Tasks/TaskPrivate.h
--- a/Engine/Source/Runtime/Core/Public/Tasks/TaskPrivate.h
+++ b/Engine/Source/Runtime/Core/Public/Tasks/TaskPrivate.h
@@ -337,7 +337,7 @@ namespace UE::Tasks
 
 				if (!TryExecuteTask())
 				{
-					return false; // still locked by prerequisites or another thread managed to set execution flag first
+					return false; // still locked by prerequisites, or another thread managed to set execution flag first, or we're inside this task execution
 					// we could try to help with nested tasks execution (the task execution could already spawned a couple of nested tasks sitting in the queue). 
 					// it's unclear how important this is, but this would definitely lead to more complicated impl. we can revisit this once we see such instances in profiler captures
 				}
@@ -852,9 +852,9 @@ namespace UE::Tasks
 				return;
 			}
 
-			if (GetCurrentTask() == this)
+			if (!IsAwaitable())
 			{
-				UE_LOG(LogTemp, Fatal, TEXT("A task waiting for itself detected"));
+				UE_LOG(LogTemp, Fatal, TEXT("Deadlock detected! A task can't be waited here, e.g. because it's being executed by the currect thread"));
 				return;
 			}
 
