added unit-tests for Tasks System to demonstrate how pipe retraction works

diff --git a/Engine/Source/Runtime/Core/Private/Tests/Tasks/TasksTest.cpp b/Engine/Source/Runtime/Core/Private/Tests/Tasks/TasksTest.cpp
--- a/Engine/Source/Runtime/Core/Private/Tests/Tasks/TasksTest.cpp
+++ b/Engine/Source/Runtime/Core/Private/Tests/Tasks/TasksTest.cpp
@@ -1007,6 +1007,32 @@ namespace UE { namespace TasksTests
 
 		TwoLevelsDeepRetractionTest();
 
+		{	// retraction of a piped task
+			FPipe Pipe{ UE_SOURCE_LOCATION };
+			FTask PipedTask = Pipe.Launch(UE_SOURCE_LOCATION, [] {});
+			PipedTask.Wait(); // the pipe is not blocked, so the task is retracted successfully
+		}
+
+		{	// deep retraction of a piped task: waiting for a piped task when the pipe has other incomplete tasks before it
+			FPipe Pipe{ UE_SOURCE_LOCATION };
+			FTask PipedTask1 = Pipe.Launch(UE_SOURCE_LOCATION, [] {});
+			FTask PipedTask2 = Pipe.Launch(UE_SOURCE_LOCATION, [] {});
+			PipedTask2.Wait(); // the pipe is not blocked, deep retraction first executes `PipedTask1`, and then `PipedTask2`
+		}
+
+		// an example of a deadlock caused by waiting for a piped task from inside execution of another task of the same pipe
+		//{	// retraction of a piped task from inside that pipe
+		//	FPipe Pipe{ UE_SOURCE_LOCATION };
+		//	FTask PipedOuterTask = Pipe.Launch(UE_SOURCE_LOCATION, 
+		//		[&Pipe] () mutable
+		//		{
+		//			FTask PipedInnerTask = Pipe.Launch(UE_SOURCE_LOCATION, [] {});
+		//			PipedInnerTask.Wait(); // deadlock as the pipe is blocked (we are inside it)
+		//		}
+		//	);
+		//	PipedOuterTask.Wait();
+		//}
+
 		ResumeEvent.Trigger();
 		LowLevelTasks::BusyWaitForTasks<LowLevelTasks::FTask>(WorkerBlockers);
 
