Fixed compilation warnings reported after adding missing support for CompileEnvironment.UnsafeTypeCastWarningLevel (additional per-module type conversion warnings) on one of the platforms.

#rnx

diff --git a/Engine/Source/Runtime/Core/Private/HAL/MallocBinnedArena.cpp b/Engine/Source/Runtime/Core/Private/HAL/MallocBinnedArena.cpp
--- a/Engine/Source/Runtime/Core/Private/HAL/MallocBinnedArena.cpp
+++ b/Engine/Source/Runtime/Core/Private/HAL/MallocBinnedArena.cpp
@@ -13,6 +13,8 @@ PRAGMA_DISABLE_UNSAFE_TYPECAST_WARNINGS
 #include "HAL/MemoryMisc.h"
 #include "HAL/PlatformMisc.h"
 
+#include <limits>
+
 struct FMallocBinnedArena::FPoolInfoSmall
 {
 	enum ECanary
@@ -570,12 +572,13 @@ void FMallocBinnedArena::InitMallocBinned()
 		if (Size > ArenaParams.BasePageSize)
 		{
 			check(Size % 4096 == 0); // calculations are done assume 4k is the smallest page size we will ever see
-			SizeTable.Emplace(Size, ArenaParams.AllocationGranularity, Size / 4096, ArenaParams.BasePageSize, ArenaParams.MinimumAlignment);
+			check(Size / 4096 <= std::numeric_limits<uint8>::max())		// Make sure we don't try to allocate more pages than fits in our counter.
+			SizeTable.Emplace(Size, ArenaParams.AllocationGranularity, (uint8)(Size / 4096), ArenaParams.BasePageSize, ArenaParams.MinimumAlignment);
 		}
 		else
 		{
 			// it is difficult to test what would actually make a good bucket size here, wouldn't want a prime number, 33 for example because that would take 33 pages a slab
-			SizeTable.Emplace(Size, ArenaParams.AllocationGranularity, 1, ArenaParams.BasePageSize, ArenaParams.MinimumAlignment);
+			SizeTable.Emplace(Size, ArenaParams.AllocationGranularity, (uint8)1, ArenaParams.BasePageSize, ArenaParams.MinimumAlignment);
 		}
 		ArenaParams.PoolCount++;
 	}
diff --git a/Engine/Source/Runtime/Core/Private/HAL/MallocBinnedGPU.cpp b/Engine/Source/Runtime/Core/Private/HAL/MallocBinnedGPU.cpp
--- a/Engine/Source/Runtime/Core/Private/HAL/MallocBinnedGPU.cpp
+++ b/Engine/Source/Runtime/Core/Private/HAL/MallocBinnedGPU.cpp
@@ -13,6 +13,8 @@ PRAGMA_DISABLE_UNSAFE_TYPECAST_WARNINGS
 #include "HAL/MemoryMisc.h"
 #include "HAL/PlatformMisc.h"
 
+#include <limits>
+
 struct FMallocBinnedGPU::FPoolInfoSmall
 {
 	enum ECanary
@@ -581,12 +583,13 @@ void FMallocBinnedGPU::InitMallocBinned()
 		if (Size > ArenaParams.BasePageSize)
 		{
 			check(Size % 4096 == 0); // calculations are done assume 4k is the smallest page size we will ever see
-			SizeTable.Emplace(Size, ArenaParams.AllocationGranularity, Size / 4096, ArenaParams.BasePageSize, ArenaParams.MinimumAlignment);
+			check(Size / 4096 <= std::numeric_limits<uint8>::max())		// Make sure we don't try to allocate more pages than fits in our counter.
+			SizeTable.Emplace(Size, ArenaParams.AllocationGranularity, (uint8)(Size / 4096), ArenaParams.BasePageSize, ArenaParams.MinimumAlignment);
 		}
 		else
 		{
 			// it is difficult to test what would actually make a good bucket size here, wouldn't want a prime number, 33 for example because that would take 33 pages a slab
-			SizeTable.Emplace(Size, ArenaParams.AllocationGranularity, 1, ArenaParams.BasePageSize, ArenaParams.MinimumAlignment);
+			SizeTable.Emplace(Size, ArenaParams.AllocationGranularity, (uint8)1, ArenaParams.BasePageSize, ArenaParams.MinimumAlignment);
 		}
 		ArenaParams.PoolCount++;
 	}
