Make updated ensure macros add parentheses consistent with how prior macro worked
Fix case where an if ensure was not using parentheses consistent with typical coding style
#rnx

diff --git a/Engine/Plugins/FX/Niagara/Source/Niagara/Private/NiagaraSystem.cpp b/Engine/Plugins/FX/Niagara/Source/Niagara/Private/NiagaraSystem.cpp
--- a/Engine/Plugins/FX/Niagara/Source/Niagara/Private/NiagaraSystem.cpp
+++ b/Engine/Plugins/FX/Niagara/Source/Niagara/Private/NiagaraSystem.cpp
@@ -3048,7 +3048,7 @@ void UNiagaraSystem::InitEmitterCompiledData()
 			const UNiagaraEmitter* Emitter = EmitterHandle.GetInstance();
 			FNiagaraDataSetCompiledData& EmitterDataSetCompiledData = NewEmitterCompiledData[EmitterIdx]->DataSetCompiledData;
 			FNiagaraDataSetCompiledData& GPUCaptureCompiledData = NewEmitterCompiledData[EmitterIdx]->GPUCaptureDataSetCompiledData;
-			if ensureMsgf(Emitter != nullptr, TEXT("Failed to get Emitter Instance from Emitter Handle in post compile, please investigate."))
+			if (ensureMsgf(Emitter != nullptr, TEXT("Failed to get Emitter Instance from Emitter Handle in post compile, please investigate.")))
 			{
 				static FName GPUCaptureDataSetName = TEXT("GPU Capture Dataset");
 				InitEmitterVariableAliasNames(NewEmitterCompiledData[EmitterIdx].Get(), Emitter);
diff --git a/Engine/Source/Runtime/Core/Public/Misc/AssertionMacros.h b/Engine/Source/Runtime/Core/Public/Misc/AssertionMacros.h
--- a/Engine/Source/Runtime/Core/Public/Misc/AssertionMacros.h
+++ b/Engine/Source/Runtime/Core/Public/Misc/AssertionMacros.h
@@ -359,10 +359,10 @@ RetType FORCENOINLINE UE_DEBUG_SECTION DispatchCheckVerify(InnerType&& Inner, Ar
 
 #else	// DO_ENSURE
 
-	#define ensure(           InExpression                ) [&](){return (!!(InExpression));}()
-	#define ensureMsgf(       InExpression, InFormat, ... ) [&](){return (!!(InExpression));}()
-	#define ensureAlways(     InExpression                ) [&](){return (!!(InExpression));}()
-	#define ensureAlwaysMsgf( InExpression, InFormat, ... ) [&](){return (!!(InExpression));}()
+	#define ensure(           InExpression                ) ([&](){return (!!(InExpression));}())
+	#define ensureMsgf(       InExpression, InFormat, ... ) ([&](){return (!!(InExpression));}())
+	#define ensureAlways(     InExpression                ) ([&](){return (!!(InExpression));}())
+	#define ensureAlwaysMsgf( InExpression, InFormat, ... ) ([&](){return (!!(InExpression));}())
 
 #endif	// DO_CHECK
 
