Logging: Fixed a deadlock when an assertion failed within a buffered output device

#rnx

diff --git a/Engine/Source/Runtime/Core/Private/Misc/OutputDeviceRedirector.cpp b/Engine/Source/Runtime/Core/Private/Misc/OutputDeviceRedirector.cpp
--- a/Engine/Source/Runtime/Core/Private/Misc/OutputDeviceRedirector.cpp
+++ b/Engine/Source/Runtime/Core/Private/Misc/OutputDeviceRedirector.cpp
@@ -139,6 +139,9 @@ struct FOutputDeviceRedirectorState
 	/** A queue of events to trigger when the dedicated master thread is idle. */
 	TDepletableMpscQueue<FEvent*, FOutputDeviceLinearAllocator> ThreadIdleEvents;
 
+	/** The ID of the thread holding the master lock. */
+	std::atomic<uint32> LockedThreadId = MAX_uint32;
+
 	/** The ID of the master thread. Logging from other threads will be buffered for processing by the master thread. */
 	std::atomic<uint32> MasterThreadId = FPlatformTLS::GetCurrentThreadId();
 
@@ -160,6 +163,10 @@ struct FOutputDeviceRedirectorState
 
 	bool CanLockFromThread(const uint32 ThreadId) const
 	{
+		if (UNLIKELY(ThreadId == LockedThreadId.load(std::memory_order_relaxed)))
+		{
+			return false;
+		}
 		const uint32 LocalPanicThreadId = PanicThreadId.load(std::memory_order_relaxed);
 		return LocalPanicThreadId == MAX_uint32 || LocalPanicThreadId == ThreadId;
 	}
@@ -243,6 +250,7 @@ public:
 		TRACE_CPUPROFILER_EVENT_SCOPE(FOutputDevicesWriteScopeLock);
 		// Take the lock before modifying the state, to avoid contention on the LSB.
 		State.OutputDevicesLock.WriteLock();
+		State.LockedThreadId.store(FPlatformTLS::GetCurrentThreadId(), std::memory_order_relaxed);
 		// Set the LSB to flag to read locks that a write lock is waiting.
 		uint32 LockState = State.OutputDevicesLockState.fetch_or(uint32(1), std::memory_order_relaxed);
 		check((LockState & 1) == 0);
@@ -263,6 +271,7 @@ public:
 		// Clear the LSB to allow read locks after the unlock below.
 		uint32 LockState = State.OutputDevicesLockState.fetch_and(~uint32(1), std::memory_order_release);
 		check((LockState & 1) == 1);
+		State.LockedThreadId.store(MAX_uint32, std::memory_order_relaxed);
 		State.OutputDevicesLock.WriteUnlock();
 	}
 
@@ -307,11 +316,19 @@ public:
 					bLocked = true;
 				}
 			}
+			if (bLocked)
+			{
+				State.LockedThreadId.store(ThreadId, std::memory_order_relaxed);
+			}
 		}
 	}
 
 	FORCEINLINE ~FOutputDevicesMasterScope()
 	{
+		if (bLocked)
+		{
+			State.LockedThreadId.store(MAX_uint32, std::memory_order_relaxed);
+		}
 		if (bNeedsUnlock)
 		{
 			State.OutputDevicesLock.WriteUnlock();
