fixed a discrepency between declaration of FBaseGraphTask::LifeStage and its usage

diff --git a/Engine/Source/Runtime/Core/Public/Async/TaskGraphInterfaces.h b/Engine/Source/Runtime/Core/Public/Async/TaskGraphInterfaces.h
--- a/Engine/Source/Runtime/Core/Public/Async/TaskGraphInterfaces.h
+++ b/Engine/Source/Runtime/Core/Public/Async/TaskGraphInterfaces.h
@@ -907,7 +907,7 @@ private:
 	FThreadSafeCounter			NumberOfPrerequistitesOutstanding; 
 
 
-#if !UE_BUILD_SHIPPING
+#if DO_GUARD_SLOW
 	// Life stage verification
 	// Tasks go through 8 steps, in order. In non-final builds, we track them with a thread safe counter and verify that the progression is correct.
 	enum ELifeStage
