- Fixed a config issue with deeply nested projects

diff --git a/Engine/Source/Runtime/Core/Private/Misc/ConfigCacheIni.cpp b/Engine/Source/Runtime/Core/Private/Misc/ConfigCacheIni.cpp
--- a/Engine/Source/Runtime/Core/Private/Misc/ConfigCacheIni.cpp
+++ b/Engine/Source/Runtime/Core/Private/Misc/ConfigCacheIni.cpp
@@ -578,7 +578,7 @@ struct FConfigLayer
 	**************************************************/
 
 	// Engine/Base.ini
-	{ TEXT("AbsoluteBase"),				TEXT("../../../Engine/Config/Base.ini"), EConfigLayerFlags::NoExpand},
+	{ TEXT("AbsoluteBase"),				TEXT("{TRUEENGINE}/Config/Base.ini"), EConfigLayerFlags::NoExpand},
 	
 	// Engine/Base*.ini
 	{ TEXT("Base"),						TEXT("{ENGINE}/Config/Base{TYPE}.ini") },
@@ -5935,6 +5935,7 @@ bool GenerateDestIniFile(FConfigContext& Context)
 FString PerformFinalExpansions(FConfigContext& Context, const FString& InString, const FString& Platform)
 {
 	FString OutString = InString.Replace(TEXT("{ENGINE}"), *Context.EngineRootDir);
+	OutString = OutString.Replace(TEXT("{TRUEENGINE}"), *FPaths::EngineDir());
 	OutString = OutString.Replace(TEXT("{PROJECT}"), *Context.ProjectRootDir);
 	OutString = OutString.Replace(TEXT("{RESTRICTEDPROJECT_NFL}"), *Context.ProjectNotForLicenseesDir);
 	OutString = OutString.Replace(TEXT("{RESTRICTEDPROJECT_NR}"), *Context.ProjectNoRedistDir);
