Packaging Reference Chunk Database - "Lossless Patch Preventer"

Can now provide reference iostore containers to reuse compressed chunks from. If a match is found on *the decompressed data*, instead of recompressing the blocks, they are read off of disk. This allows tweaks of the compressor algorithm without introducing changes as the runtime still sees the exact same data. Additionally this allows for fairly dramatic staging speedups as nvme speeds are significantly faster than high effort compressions. This is distinct from the DDC compression because:
1) DDC compression ties in the compressor version/method
2) We are explicitly interested in chunks that are deployed to end users, not merely cached for speed.

To facilitate this, several changes were made to IoStore:

FIoStoreReader now directly reads from IFileHandles* instead of routing through the GenericPlatformFile async read system, as that system is sensitive to build #defines and can result in constant file opens under load (indeed, for anything not a .pak file, every read is an open/close).

Cold file cache read speed improvements from ~140MB/s to ~1 GB/s. Hot is more.

Additionally:
    FIoStoreReader switched to UE::Tasks from taskgraph for tasks in order to facilitate task retraction during waits as the previous ReadAsync call was trivial to deadlock when called from worker threads due to its use of TFuture<>.
    FIoStoreReader::ReadCompressed now returns the compressed blocks as they were on disk - padded to AES encryption block size.

diff --git a/Engine/Source/Developer/IoStoreUtilities/Private/IoStoreUtilities.cpp b/Engine/Source/Developer/IoStoreUtilities/Private/IoStoreUtilities.cpp
--- a/Engine/Source/Developer/IoStoreUtilities/Private/IoStoreUtilities.cpp
+++ b/Engine/Source/Developer/IoStoreUtilities/Private/IoStoreUtilities.cpp
@@ -82,6 +82,134 @@ static const uint64 DefaultCompressionBlockSize = 64 << 10;
 static const uint64 DefaultCompressionBlockAlignment = 64 << 10;
 static const uint64 DefaultMemoryMappingAlignment = 16 << 10;
 
+static TUniquePtr<FIoStoreReader> CreateIoStoreReader(const TCHAR* Path, const FKeyChain& KeyChain);
+
+class FIoStoreChunkDatabase : public IIoStoreWriterReferenceChunkDatabase
+{
+public:
+
+	TArray<TUniquePtr<FIoStoreReader>> Readers;
+	struct FReaderChunks
+	{
+		int32 ReaderIndex;
+		TMap<FIoChunkHash, FIoChunkId> Chunks;
+	};
+
+	TMap<FIoContainerId, FReaderChunks> ChunkDatabase;
+	int32 RequestCount = 0;
+	int32 FulfillCount = 0;
+	int32 ContainerNotFound = 0;
+	int64 FulfillBytes = 0;
+	
+
+	bool Init(const FString& InGlobalContainerFileName, const FKeyChain& InDecryptionKeychain)
+	{
+		double StartTime = FPlatformTime::Seconds();
+
+		TUniquePtr<FIoStoreReader> GlobalReader = CreateIoStoreReader(*InGlobalContainerFileName, InDecryptionKeychain);
+		if (GlobalReader.IsValid() == false)
+		{
+			UE_LOG(LogIoStore, Error, TEXT("Failed to open reference chunk container %s"), *InGlobalContainerFileName);
+			return false;
+		}
+
+		FString Directory = FPaths::GetPath(InGlobalContainerFileName);
+		FPaths::NormalizeDirectoryName(Directory);
+
+		TArray<FString> FoundContainerFiles;
+		IFileManager::Get().FindFiles(FoundContainerFiles, *(Directory / TEXT("*.utoc")), true, false);
+		TArray<FString> ContainerFilePaths;
+		for (const FString& Filename : FoundContainerFiles)
+		{
+			ContainerFilePaths.Emplace(Directory / Filename);
+		}
+
+		int64 IoChunkCount = 0;		
+		for (const FString& ContainerFilePath : ContainerFilePaths)
+		{
+			TUniquePtr<FIoStoreReader> Reader = CreateIoStoreReader(*ContainerFilePath, InDecryptionKeychain);
+			if (Reader.IsValid() == false)
+			{
+				UE_LOG(LogIoStore, Error, TEXT("Failed to open reference chunk container %s"), *InGlobalContainerFileName);
+				return false;
+			}
+			FReaderChunks& ReaderChunks = ChunkDatabase.FindOrAdd(Reader->GetContainerId());
+
+			Reader->EnumerateChunks([&ReaderChunks](const FIoStoreTocChunkInfo& ChunkInfo)
+			{
+				ReaderChunks.Chunks.Add(TPair<FIoChunkHash, FIoChunkId>(ChunkInfo.Hash, ChunkInfo.Id));
+				return true;
+			});
+
+			IoChunkCount += ReaderChunks.Chunks.Num();
+			ReaderChunks.ReaderIndex = Readers.Num();
+			Readers.Add(MoveTemp(Reader));
+		}
+
+		UE_LOG(LogIoStore, Display, TEXT("Block reference loaded %d containers and %s chunks, in %.1f seconds"), Readers.Num(), *FText::AsNumber(IoChunkCount).ToString(), FPlatformTime::Seconds() - StartTime);
+		return true;
+	}
+
+	// Not thread safe, called from the BeginCompress dispatch thread.
+	virtual bool RetrieveChunk(const TPair<FIoContainerId, FIoChunkHash>& InChunkKey, const FName& InCompressionMethod, uint64 InUncompressedSize, uint64 InNumChunkBlocks, TUniqueFunction<void(TIoStatusOr<FIoStoreCompressedReadResult>)> InCompleteCallback)
+	{
+		RequestCount++;
+
+		FReaderChunks* ReaderChunks = ChunkDatabase.Find(InChunkKey.Key);
+		if (ReaderChunks == nullptr)
+		{
+			// Container doesn't exist - likely provided the path to a different project. Mark this
+			// error as happening once so we can log at the end.
+			ContainerNotFound++;
+			return false;
+		}
+
+		FIoChunkId* ChunkId = ReaderChunks->Chunks.Find(InChunkKey.Value);
+		if (ChunkId == nullptr)
+		{
+			// No exact chunk data match - this is a normal exit condition for a changed block.
+			return false;
+		}
+
+		uint64 TotalCompressedSize = 0;
+		uint64 TotalUncompressedSize = 0;
+		uint32 CompressedBlockCount = 0;
+		Readers[ReaderChunks->ReaderIndex]->EnumerateCompressedBlocksForChunk(*ChunkId, [&TotalUncompressedSize, &CompressedBlockCount, &TotalCompressedSize](const FIoStoreTocCompressedBlockInfo& BlockInfo)
+		{			
+			TotalCompressedSize += BlockInfo.CompressedSize;
+			TotalUncompressedSize += BlockInfo.UncompressedSize;
+			CompressedBlockCount ++;
+			return true;
+		});
+
+		if (TotalUncompressedSize != InUncompressedSize)
+		{
+			// Shocked if this happens - hash match with different data!
+			return false;
+		}
+
+		if (CompressedBlockCount != InNumChunkBlocks)
+		{
+			// Different CompressionBlockSize between the builds
+			return false;
+		}
+
+		FulfillBytes += TotalCompressedSize;
+		FulfillCount++;
+
+		//
+		// At this point we know we can use the block so we can go async.
+		//
+		FFunctionGraphTask::CreateAndDispatchWhenReady([this, ChunkId, ReaderIndex = ReaderChunks->ReaderIndex, CompleteCallback = MoveTemp(InCompleteCallback)]()
+		{
+			TIoStatusOr<FIoStoreCompressedReadResult> Result = Readers[ReaderIndex]->ReadCompressed(*ChunkId, FIoReadOptions());
+			CompleteCallback(Result);
+		}, TStatId(), nullptr, ENamedThreads::AnyHiPriThreadNormalTask);
+
+		return true;
+	}
+};
+
 struct FReleasedPackages
 {
 	TSet<FName> PackageNames;
@@ -771,6 +899,8 @@ struct FIoStoreArguments
 	FKeyChain PatchKeyChain;
 	FString DLCPluginPath;
 	FString DLCName;
+	FString ReferenceChunkGlobalContainerFileName;
+	FKeyChain ReferenceChunkKeys;
 	FReleasedPackages ReleasedPackages;
 	TUniquePtr<FCookedPackageStore> PackageStore;
 	TUniquePtr<FIoBuffer> ScriptObjects;
@@ -3169,6 +3299,17 @@ int32 CreateTarget(const FIoStoreArguments& Arguments, const FIoStoreWriterSetti
 	ChunkIdCsv.CreateOutputFile(CookedDir);
 #endif
 
+	TSharedPtr<IIoStoreWriterReferenceChunkDatabase> ChunkDatabase;
+	if (Arguments.ReferenceChunkGlobalContainerFileName.Len())
+	{
+		ChunkDatabase = MakeShared<FIoStoreChunkDatabase>();
+		if (((FIoStoreChunkDatabase&)*ChunkDatabase).Init(Arguments.ReferenceChunkGlobalContainerFileName, Arguments.ReferenceChunkKeys) == false)
+		{
+			UE_LOG(LogIoStore, Error, TEXT("Failed to initialize reference chunk store."));
+			return 1;
+		}
+	}
+
 	TArray<FLegacyCookedPackage*> Packages;
 	FPackageNameMap PackageNameMap;
 	FPackageIdMap PackageIdMap;
@@ -3228,6 +3369,7 @@ int32 CreateTarget(const FIoStoreArguments& Arguments, const FIoStoreWriterSetti
 				ContainerSettings.bGenerateDiffPatch = ContainerTarget->bGenerateDiffPatch;
 				ContainerTarget->IoStoreWriter = IoStoreWriterContext->CreateContainer(*ContainerTarget->OutputPath, ContainerSettings);
 				ContainerTarget->IoStoreWriter->EnableDiskLayoutOrdering(ContainerTarget->PatchSourceReaders);
+				ContainerTarget->IoStoreWriter->SetReferenceChunkDatabase(ChunkDatabase);
 				IoStoreWriters.Add(ContainerTarget->IoStoreWriter);
 				if (!ContainerTarget->OptionalSegmentOutputPath.IsEmpty())
 				{
@@ -3553,6 +3695,25 @@ int32 CreateTarget(const FIoStoreArguments& Arguments, const FIoStoreWriterSetti
 	UE_LOG(LogIoStore, Display, TEXT("Output: %8d Public runtime script objects"), PackageStoreOptimizer.GetTotalScriptObjectCount());
 	UE_LOG(LogIoStore, Display, TEXT("Output: %8.2lf MB InitialLoadData"), (double)InitialLoadSize / 1024.0 / 1024.0);
 
+	if (ChunkDatabase.IsValid())
+	{
+		uint64 TotalCompressedBytes = 0;
+		for (const FIoStoreWriterResult& Result : IoStoreWriterResults)
+		{
+			TotalCompressedBytes += Result.CompressedContainerSize;
+		}
+
+		FIoStoreChunkDatabase& ChunkDatabaseRef = (FIoStoreChunkDatabase&)*ChunkDatabase;
+		UE_LOG(LogIoStore, Display, TEXT("Reference Chunk: %s reused bytes out of %s possible: %.1f%%"), *FText::AsNumber(ChunkDatabaseRef.FulfillBytes).ToString(), *FText::AsNumber(TotalCompressedBytes).ToString(), 100.0 * ChunkDatabaseRef.FulfillBytes / TotalCompressedBytes);
+		UE_LOG(LogIoStore, Display, TEXT("Reference Chunk: %s chunks found / %s requests"), *FText::AsNumber(ChunkDatabaseRef.FulfillCount).ToString(), *FText::AsNumber(ChunkDatabaseRef.RequestCount).ToString());
+		if (ChunkDatabaseRef.ContainerNotFound)
+		{
+			UE_LOG(LogIoStore, Warning, TEXT("Reference Chunk had %s requests for a container that wasn't loaded. This means the "), *FText::AsNumber(ChunkDatabaseRef.ContainerNotFound).ToString());
+			UE_LOG(LogIoStore, Warning, TEXT("new output has a container that wasn't deployed before. If that doesn't sound right"));
+			UE_LOG(LogIoStore, Warning, TEXT("verify that you used reference containers from the same project."));
+		}
+	}
+
 	return 0;
 }
 
@@ -3969,6 +4130,223 @@ bool LegacyListIoStoreContainer(
 	return true;
 }
 
+int32 ProfileReadSpeed(const TCHAR* InCommandLine, const FKeyChain& InKeyChain)
+{
+	FString ContainerPath;
+	if (FParse::Value(InCommandLine, TEXT("Container="), ContainerPath) == false)
+	{
+		UE_LOG(LogIoStore, Display, TEXT(""));
+		UE_LOG(LogIoStore, Display, TEXT("ProfileReadSpeed"));
+		UE_LOG(LogIoStore, Display, TEXT(""));
+		UE_LOG(LogIoStore, Display, TEXT("Reads the given utoc file using given a read method. This uses FIoStoreReader, which is not"));
+		UE_LOG(LogIoStore, Display, TEXT("the system the runtime uses to load and stream iostore containers! It's for utility/debug use only."));
+		UE_LOG(LogIoStore, Display, TEXT(""));
+		UE_LOG(LogIoStore, Display, TEXT("Arguments:"));
+		UE_LOG(LogIoStore, Display, TEXT(""));
+		UE_LOG(LogIoStore, Display, TEXT("    -Container=path/to/utoc                      [required] The .utoc file to read."));
+		UE_LOG(LogIoStore, Display, TEXT("    -ReadType={Read, ReadAsync, ReadCompressed}  What read function to use on FIoStoreReader. Default: Read"));
+		UE_LOG(LogIoStore, Display, TEXT("    -cryptokeys=path/to/crypto.json              [required if encrypted] The keys to decrypt the container."));
+		UE_LOG(LogIoStore, Display, TEXT("    -MaxJobCount=#                               The number of outstanding read tasks to maintain. Default: 512."));
+		UE_LOG(LogIoStore, Display, TEXT("    -Validate                                    Whether to hash the reads and verify they match. Invalid for ReadCompressed. Default: disabled"));
+		return 1;
+	}
+
+	TUniquePtr<FIoStoreReader> Reader = CreateIoStoreReader(*ContainerPath, InKeyChain);
+	if (Reader.IsValid() == false)
+	{
+		return 1; // Already logged
+	}
+
+	TArray<FIoChunkId> Chunks;
+	Reader->EnumerateChunks([&Chunks](const FIoStoreTocChunkInfo& ChunkInfo)
+	{
+		Chunks.Add(ChunkInfo.Id);
+		return true;
+	});
+
+	enum class EReadType
+	{
+		Read,
+		ReadAsync,
+		ReadCompressed
+	};
+
+	auto ReadTypeToString = [](EReadType InReadType)
+	{
+		switch (InReadType)
+		{
+		case EReadType::Read: return TEXT("Read");
+		case EReadType::ReadAsync: return TEXT("ReadAsync");
+		case EReadType::ReadCompressed: return TEXT("ReadCompressed");
+		default: return TEXT("INVALID");
+		}
+	};
+
+	EReadType ReadType = EReadType::Read;
+	int32 MaxOutstandingJobs = 512;
+	bool bValidate = false;
+
+	bValidate = FParse::Param(InCommandLine, TEXT("Validate"));
+	FParse::Value(InCommandLine, TEXT("MaxJobCount="), MaxOutstandingJobs);
+
+	FString ReadTypeRaw;
+	if (FParse::Value(InCommandLine, TEXT("ReadType="), ReadTypeRaw))
+	{
+		if (ReadTypeRaw.Compare(TEXT("Read"), ESearchCase::IgnoreCase) == 0)
+		{
+			ReadType = EReadType::Read;
+		}
+		else if (ReadTypeRaw.Compare(TEXT("ReadAsync"), ESearchCase::IgnoreCase) == 0)
+		{
+			ReadType = EReadType::ReadAsync;
+		}
+		else if (ReadTypeRaw.Compare(TEXT("ReadCompressed"), ESearchCase::IgnoreCase) == 0)
+		{
+			ReadType = EReadType::ReadCompressed;
+		}
+		else
+		{
+			UE_LOG(LogIoStore, Error, TEXT("Invalid -ReadType provided: %s. Valid are {Read, ReadAsync, ReadCompressed}"), *ReadTypeRaw);
+			return 1;
+		}
+	}
+
+	if (MaxOutstandingJobs <= 0)
+	{
+		UE_LOG(LogIoStore, Error, TEXT("Invalid -MaxJobCount provided: %d. Specify a positive integer"), MaxOutstandingJobs);
+		return 1;
+	}
+
+	if (ReadType == EReadType::ReadCompressed && 
+		bValidate)
+	{
+		UE_LOG(LogIoStore, Error, TEXT("Can't validate ReadCompressed as the data is not decompressed and thus can't be hashed"));
+		return 1;
+	}
+
+	UE_LOG(LogIoStore, Display, TEXT("MaxJobCount:            %s"), *FText::AsNumber(MaxOutstandingJobs).ToString());
+	UE_LOG(LogIoStore, Display, TEXT("ReadType:               %s"), ReadTypeToString(ReadType));
+	UE_LOG(LogIoStore, Display, TEXT("Validation:             %s"), bValidate ? TEXT("Enabled") : TEXT("Disabled"));
+	UE_LOG(LogIoStore, Display, TEXT("Container Encrypted:    %s"), EnumHasAnyFlags(Reader->GetContainerFlags(), EIoContainerFlags::Encrypted) ? TEXT("Yes") : TEXT("No"));
+	
+	// We need a resettable event, so we can't use any of the task system events.
+	FEvent* JustGotSpaceEvent = FPlatformProcess::GetSynchEventFromPool();
+	UE::Tasks::FTaskEvent CompletedEvent(TEXT("ProfileReadDone"));
+
+	std::atomic_int32_t OutstandingJobs = 0;
+	std::atomic_int32_t TotalJobsRemaining = Chunks.Num();
+	std::atomic_int64_t BytesRead = 0;
+
+	double StartTime = FPlatformTime::Seconds();
+	UE_LOG(LogIoStore, Display, TEXT("Dispatching %s chunk reads (%s max at one time)"), *FText::AsNumber(Chunks.Num()).ToString(), *FText::AsNumber(MaxOutstandingJobs).ToString());
+
+	for (FIoChunkId& Id : Chunks)
+	{
+		for (;;)
+		{
+			int32 CurrentOutstanding = OutstandingJobs.load();
+			if (CurrentOutstanding == MaxOutstandingJobs)
+			{
+				// Wait for one to complete.
+				JustGotSpaceEvent->Wait();
+				continue;
+			}
+			else if (CurrentOutstanding > MaxOutstandingJobs)
+			{
+				UE_LOG(LogIoStore, Warning, TEXT("Synch error -- too many jobs oustanding %d"), CurrentOutstanding);
+			}
+			break;
+		}
+
+		OutstandingJobs++;
+		UE::Tasks::Launch(TEXT("IoStoreUtil::ReadJob"), [Id = Id, &OutstandingJobs, &JustGotSpaceEvent, &TotalJobsRemaining, &BytesRead, &Reader, &CompletedEvent, MaxOutstandingJobs, ReadType, bValidate]()
+		{
+			FIoChunkHash ReadHash;
+			bool bHashValid = false;
+
+			switch (ReadType)
+			{
+			case EReadType::ReadCompressed:
+				{
+					FIoStoreCompressedReadResult Result = Reader->ReadCompressed(Id, FIoReadOptions()).ValueOrDie();
+					BytesRead += Result.IoBuffer.GetSize();
+					break;
+				}
+			case EReadType::Read:
+				{
+					FIoBuffer Result = Reader->Read(Id, FIoReadOptions()).ValueOrDie();
+					BytesRead += Result.GetSize();
+
+					if (bValidate)
+					{
+						ReadHash = FIoChunkHash::HashBuffer(Result.GetData(), Result.GetSize());
+						bHashValid = true;
+					}
+
+					break;
+				}
+			case EReadType::ReadAsync:
+				{
+					UE::Tasks::TTask<TIoStatusOr<FIoBuffer>> Tasks = Reader->ReadAsync(Id, FIoReadOptions());
+					Tasks.BusyWait();
+					FIoBuffer Result = Tasks.GetResult().ValueOrDie();
+
+					BytesRead += Result.GetSize();
+
+					if (bValidate)
+					{
+						ReadHash = FIoChunkHash::HashBuffer(Result.GetData(), Result.GetSize());
+						bHashValid = true;
+					}
+					break;
+				}
+			}
+
+			if (bHashValid)
+			{
+				FIoChunkHash CheckAgainstHash = Reader->GetChunkInfo(Id).ValueOrDie().Hash;
+				if (ReadHash != CheckAgainstHash)
+				{
+					UE_LOG(LogIoStore, Warning, TEXT("Read hash mismatch: Chunk %s"), *LexToString(Id));
+				}
+			}
+
+			if (OutstandingJobs.fetch_add(-1) == MaxOutstandingJobs)
+			{
+				// We are the first to make space in our limit, so release the dispatch thread to add more.
+				JustGotSpaceEvent->Trigger();
+			}
+
+			int32 JobsRemaining = TotalJobsRemaining.fetch_add(-1);
+			if ((JobsRemaining % 1000) == 1)
+			{
+				UE_LOG(LogIoStore, Display, TEXT("Jobs Remaining: %d"), JobsRemaining - 1);
+			}
+
+			// Were we the last job issued?
+			if (JobsRemaining == 1)
+			{
+				CompletedEvent.Trigger();
+			}
+		});
+	}
+
+	{
+		double WaitStartTime = FPlatformTime::Seconds();
+		CompletedEvent.BusyWait();
+		UE_LOG(LogIoStore, Display, TEXT("Waited %.1f seconds"), FPlatformTime::Seconds() - WaitStartTime);
+	}
+
+	FPlatformProcess::ReturnSynchEventToPool(JustGotSpaceEvent);
+
+	double TotalTime = FPlatformTime::Seconds() - StartTime;
+
+	int64 BytesPerSecond = int64(BytesRead.load() / TotalTime);
+
+	UE_LOG(LogIoStore, Display, TEXT("%s bytes in %.1f seconds; %s bytes per second"), *FText::AsNumber((int64)BytesRead.load()).ToString(), TotalTime, *FText::AsNumber(BytesPerSecond).ToString());
+	return 0;
+}
+
 int32 Describe(
 	const FString& GlobalContainerPath,
 	const FKeyChain& KeyChain,
@@ -5486,15 +5864,20 @@ bool ExtractFilesFromIoStoreContainer(
 		}
 	};
 
-	TArray<TFuture<bool>> ExtractTasks;
-	ExtractTasks.SetNum(Entries.Num());
+	TArray<UE::Tasks::TTask<bool>> ExtractTasks;
+	ExtractTasks.Reserve(Entries.Num());
 	EAsyncExecution ThreadPool = EAsyncExecution::ThreadPool;
 	for (int32 EntryIndex = 0; EntryIndex < Entries.Num(); ++EntryIndex)
 	{
 		const FEntry& Entry = Entries[EntryIndex];
-		ExtractTasks[EntryIndex] = Reader->ReadAsync(Entry.ChunkId, FIoReadOptions())
-			.Next([&Entry, &WriteFile](TIoStatusOr<FIoBuffer> ReadChunkResult)
+
+		UE::Tasks::TTask<TIoStatusOr<FIoBuffer>> ReadTask = Reader->ReadAsync(Entry.ChunkId, FIoReadOptions());
+
+		// Once the read is done, write out the file.
+		ExtractTasks.Emplace(UE::Tasks::Launch(TEXT("IoStore_Extract"), 
+			[&Entry, &WriteFile, ReadTask]() mutable
 			{
+				TIoStatusOr<FIoBuffer> ReadChunkResult = ReadTask.GetResult();
 				if (!ReadChunkResult.IsOk())
 				{
 					UE_LOG(LogIoStore, Error, TEXT("Failed reading chunk for file \"%s\" (%s)."), *Entry.SourceFileName, *ReadChunkResult.Status().ToString());
@@ -5524,13 +5907,14 @@ bool ExtractFilesFromIoStoreContainer(
 					return false;
 				}
 				return true;
-			});
+			},
+			Prerequisites(ReadTask)));
 	}
 
 	int32 ErrorCount = 0;
 	for (int32 EntryIndex = 0; EntryIndex < Entries.Num(); ++EntryIndex)
 	{
-		if (ExtractTasks[EntryIndex].Get())
+		if (ExtractTasks[EntryIndex].GetResult())
 		{
 			const FEntry& Entry = Entries[EntryIndex];
 			if (OutOrderMap != nullptr)
@@ -5832,6 +6216,26 @@ bool ParseContainerGenerationArguments(FIoStoreArguments& Arguments, FIoStoreWri
 		WriterSettings.CompressionBlockAlignment = BlockAlignment;
 	}
 
+	//
+	// If a filename to a global.utoc container is provided, all containers in that directory will have their compressed blocks be
+	// made available for the new containers to reuse. This provides two benefits:
+	//	1.	Saves compression time for the new blocks, as ssd/nvme io times are significantly faster.
+	//	2.	Prevents trivial bit changes in the compressor from causing patch changes down the line, 
+	//		allowing worry-free compressor upgrading.
+	//
+	// This should be a path to your last released containers. If those containers are encrypted, be sure to
+	// provide keys via -ReferenceContainerCryptoKeys.
+	//
+	FParse::Value(FCommandLine::Get(), TEXT("-ReferenceContainerGlobalFileName="), Arguments.ReferenceChunkGlobalContainerFileName);
+
+	FString CryptoKeysCacheFilename;
+	if (FParse::Value(FCommandLine::Get(), TEXT("-ReferenceContainerCryptoKeys="), CryptoKeysCacheFilename))
+	{
+		UE_LOG(LogIoStore, Display, TEXT("Parsing reference container crypto keys from a crypto key cache file '%s'"), *CryptoKeysCacheFilename);
+		KeyChainUtilities::LoadKeyChainFromFile(CryptoKeysCacheFilename, Arguments.ReferenceChunkKeys);
+	}
+
+
 	uint64 PatchPaddingAlignment = 0;
 	if (ParseSizeArgument(FCommandLine::Get(), TEXT("-patchpaddingalign="), PatchPaddingAlignment))
 	{
@@ -6037,6 +6441,11 @@ int32 CreateIoStoreContainerFiles(const TCHAR* CmdLine)
 		bool bIncludeExportHashes = FParse::Param(FCommandLine::Get(), TEXT("IncludeExportHashes"));
 		return Describe(ContainerPathOrWildcard, Arguments.KeyChain, PackageFilter, OutPath, bIncludeExportHashes);
 	}
+	else if (FParse::Param(FCommandLine::Get(), TEXT("ProfileReadSpeed")))
+	{
+		// Load the .UTOC file provided and read it in its entirety, cmdline parsed in function
+		return ProfileReadSpeed(FCommandLine::Get(), Arguments.KeyChain);
+	}
 	else if (FParse::Param(FCommandLine::Get(), TEXT("Diff")))
 	{
 		FString SourcePath, TargetPath, OutPath;
diff --git a/Engine/Source/Developer/LauncherServices/Private/Launcher/LauncherWorker.cpp b/Engine/Source/Developer/LauncherServices/Private/Launcher/LauncherWorker.cpp
--- a/Engine/Source/Developer/LauncherServices/Private/Launcher/LauncherWorker.cpp
+++ b/Engine/Source/Developer/LauncherServices/Private/Launcher/LauncherWorker.cpp
@@ -827,6 +827,16 @@ FString FLauncherWorker::CreateUATCommand( const ILauncherProfileRef& InProfile,
 		UATCommand += TEXT(" -SkipCookingEditorContent");
 	}
 
+	FString StageAdditionalCommandLine;
+	if (InProfile->IsUsingIoStore() &&
+		InProfile->GetReferenceContainerGlobalFileName().Len())
+	{
+		StageAdditionalCommandLine += TEXT(" -ReferenceContainerGlobalFileName=\"") + InProfile->GetReferenceContainerGlobalFileName() + TEXT("\"");
+		if (InProfile->GetReferenceContainerCryptoKeysFileName().Len())
+		{
+			StageAdditionalCommandLine += TEXT(" -ReferenceContainerCryptoKeys=\"") + InProfile->GetReferenceContainerCryptoKeysFileName() + TEXT("\"") ;
+		}
+	}
 	// stage/package/deploy
 	if (InProfile->GetDeploymentMode() != ELauncherProfileDeploymentModes::DoNotDeploy)
 	{
@@ -867,6 +877,7 @@ FString FLauncherWorker::CreateUATCommand( const ILauncherProfileRef& InProfile,
 				UATCommand += StageDirectory;
 				UATCommand += DeviceCommand;
 				UATCommand += AdditionalCommandLine;
+				UATCommand += StageAdditionalCommandLine;
 
 				FCommandDesc Desc;
 				FText Command = FText::Format(LOCTEXT("LauncherDeployDesc", "Deploying content for {0}"), FText::FromString(Platforms.RightChop(1)));
@@ -912,6 +923,7 @@ FString FLauncherWorker::CreateUATCommand( const ILauncherProfileRef& InProfile,
 			UATCommand += StageDirectory;
 			UATCommand += CommandLine;
 			UATCommand += AdditionalCommandLine;
+			UATCommand += StageAdditionalCommandLine;
 
 			FCommandDesc Desc;
 			FText Command = FText::Format(LOCTEXT("LauncherPackageDesc", "Packaging content for {0}"), FText::FromString(Platforms.RightChop(1)));
diff --git a/Engine/Source/Developer/LauncherServices/Private/Profiles/LauncherProfile.h b/Engine/Source/Developer/LauncherServices/Private/Profiles/LauncherProfile.h
--- a/Engine/Source/Developer/LauncherServices/Private/Profiles/LauncherProfile.h
+++ b/Engine/Source/Developer/LauncherServices/Private/Profiles/LauncherProfile.h
@@ -49,7 +49,7 @@ enum ELauncherVersion
 	LAUNCHERSERVICES_ADDEDBUILDMODE = 28,
 	LAUNCHERSERVICES_ADDEDUSEIOSTORE = 29,
 	LAUNCHERSERVICES_ADDEDMAKEBINARYCONFIG = 30,
-
+	LAUNCHERSERVICES_ADDEDREFERENCECONTAINERS = 31,
 	//ADD NEW STUFF HERE
 
 
@@ -536,7 +536,22 @@ public:
 		BasedOnReleaseVersionName = InBasedOnReleaseVersionName;
 	}
 
-
+	virtual FString GetReferenceContainerGlobalFileName() const override
+	{
+		return ReferenceContainerGlobalFileName;
+	}
+	virtual void SetReferenceContainerGlobalFileName(const FString& InReferenceContainerGlobalFileName) override
+	{
+		ReferenceContainerGlobalFileName = InReferenceContainerGlobalFileName;
+	}
+	virtual FString GetReferenceContainerCryptoKeysFileName() const override
+	{
+		return ReferenceContainerCryptoKeysFileName;
+	}
+	virtual void SetReferenceContainerCryptoKeysFileName(const FString& InReferenceContainerCryptoKeysFileName) override
+	{
+		ReferenceContainerCryptoKeysFileName = InReferenceContainerCryptoKeysFileName;
+	}
 
 	virtual ELauncherProfileDeploymentModes::Type GetDeploymentMode( ) const override
 	{
@@ -1031,6 +1046,12 @@ public:
 			Archive << bMakeBinaryConfig;
 		}
 
+		if (Version >= LAUNCHERSERVICES_ADDEDREFERENCECONTAINERS)
+		{
+			Archive << ReferenceContainerGlobalFileName;
+			Archive << ReferenceContainerCryptoKeysFileName;
+		}
+
 		else if(Archive.IsLoading())
 		{
 			BuildMode = BuildGame ? ELauncherProfileBuildModes::Build : ELauncherProfileBuildModes::DoNotBuild;
@@ -1161,6 +1182,8 @@ public:
 		Writer.WriteValue("CreateReleaseVersion", CreateReleaseVersion);
 		Writer.WriteValue("CreateReleaseVersionName", CreateReleaseVersionName);
 		Writer.WriteValue("BasedOnReleaseVersionName", BasedOnReleaseVersionName);
+		Writer.WriteValue("ReferenceContainerGlobalFileName", ReferenceContainerGlobalFileName);
+		Writer.WriteValue("ReferenceContainerCryptoKeysFileName", ReferenceContainerCryptoKeysFileName);
 		Writer.WriteValue("CreateDLC", CreateDLC);
 		Writer.WriteValue("DLCName", DLCName);
 		Writer.WriteValue("GenerateChunks", bGenerateChunks);
@@ -1446,6 +1469,7 @@ public:
 		Writer.WriteValue("ForDistribution", IsForDistribution());
 
 		// stage/package/deploy
+		bool bIsStaging = false;
 		if (GetDeploymentMode() != ELauncherProfileDeploymentModes::DoNotDeploy)
 		{
 			switch (GetDeploymentMode())
@@ -1464,6 +1488,7 @@ public:
 			case ELauncherProfileDeploymentModes::FileServer:
 				{
 					Writer.WriteValue("stage", true);
+					bIsStaging = true;
 					Writer.WriteValue("deploy", true);
 				}
 				break;
@@ -1480,10 +1505,22 @@ public:
 			if (GetPackagingMode() == ELauncherProfilePackagingModes::Locally)
 			{
 				Writer.WriteValue("stage", true);
+				bIsStaging = true;
 				Writer.WriteValue("package", true);
 			}
 		}
 
+		if (bIsStaging &&
+			GetReferenceContainerGlobalFileName().Len())
+		{
+			// (only iostore uses this) - pass reference block database
+			Writer.WriteValue("ReferenceContainerGlobalFileName", GetReferenceContainerGlobalFileName());
+			if (GetReferenceContainerCryptoKeysFileName().Len())
+			{
+				Writer.WriteValue("ReferenceContainerCryptoKeys", GetReferenceContainerCryptoKeysFileName());
+			}
+		}
+
 		/*
 		"script", ""
 		"project", ""
@@ -1795,6 +1832,8 @@ public:
 		CreateReleaseVersion = Object.GetBoolField("CreateReleaseVersion");
 		CreateReleaseVersionName = Object.GetStringField("CreateReleaseVersionName");
 		BasedOnReleaseVersionName = Object.GetStringField("BasedOnReleaseVersionName");
+		ReferenceContainerCryptoKeysFileName = Object.GetStringField("ReferenceContainerCryptoKeysFileName");
+		ReferenceContainerGlobalFileName = Object.GetStringField("ReferenceContainerGlobalFileName");
 		CreateDLC = Object.GetBoolField("CreateDLC");
 		DLCName = Object.GetStringField("DLCName");
 		bGenerateChunks = Object.GetBoolField("GenerateChunks");
@@ -2857,6 +2896,15 @@ private:
 	// Version name of the HTTPChunkInstall data
 	FString HttpChunkDataReleaseName;
 
+	// if present, iostore container creation will try to use existing compressed blocks
+	// instead of compressing new ones, to avoid patches from compressor version changes,
+	// and to speed up iostore container creation time. See IoStoreUtilities.cpp ReferenceContainerGlobalFileName.
+	FString ReferenceContainerGlobalFileName;
+
+	// If ReferenceContainerGlobalFileName refers to encrypted containers, this is the filename of
+	// the json file containing the keys.
+	FString ReferenceContainerCryptoKeysFileName;
+
 	// create a release version of the content (this can be used to base dlc / patches from)
 	bool CreateReleaseVersion;
 
diff --git a/Engine/Source/Developer/LauncherServices/Public/ILauncherProfile.h b/Engine/Source/Developer/LauncherServices/Public/ILauncherProfile.h
--- a/Engine/Source/Developer/LauncherServices/Public/ILauncherProfile.h
+++ b/Engine/Source/Developer/LauncherServices/Public/ILauncherProfile.h
@@ -1118,6 +1118,15 @@ public:
 	virtual FString GetBasedOnReleaseVersionName() const = 0;
 
 	virtual void SetBasedOnReleaseVersionName(const FString& InBasedOnReleaseVersion) = 0;
+	 
+	/**
+	* Provides a database of compressed iostore chunks to reuse during the
+	* staging process. See IoStoreUtilities.cpp ReferenceContainerGlobalFileName.
+	*/
+	virtual FString GetReferenceContainerGlobalFileName() const = 0;
+	virtual void SetReferenceContainerGlobalFileName(const FString& InReferenceContainerGlobalFileName) = 0;
+	virtual FString GetReferenceContainerCryptoKeysFileName() const = 0;
+	virtual void SetReferenceContainerCryptoKeysFileName(const FString& InReferenceContainerCryptoKeysFileName) = 0;
 
 
 	/**
diff --git a/Engine/Source/Developer/ProjectLauncher/Private/Widgets/Package/SProjectLauncherPackagePage.cpp b/Engine/Source/Developer/ProjectLauncher/Private/Widgets/Package/SProjectLauncherPackagePage.cpp
--- a/Engine/Source/Developer/ProjectLauncher/Private/Widgets/Package/SProjectLauncherPackagePage.cpp
+++ b/Engine/Source/Developer/ProjectLauncher/Private/Widgets/Package/SProjectLauncherPackagePage.cpp
@@ -138,7 +138,7 @@ void SProjectLauncherPackagePage::HandlePackagingModeMenuEntryClicked(ELauncherP
 		SelectedProfile->SetPackagingMode(PackagingMode);
 
 		check(ProjectLauncherPackagingSettings.IsValid());
-		ProjectLauncherPackagingSettings->UpdateDirectoryPathText();
+		ProjectLauncherPackagingSettings->UpdatePackagingModeWidgets();
 	}
 }
 
diff --git a/Engine/Source/Developer/ProjectLauncher/Private/Widgets/Package/SProjectLauncherPackagingSettings.cpp b/Engine/Source/Developer/ProjectLauncher/Private/Widgets/Package/SProjectLauncherPackagingSettings.cpp
--- a/Engine/Source/Developer/ProjectLauncher/Private/Widgets/Package/SProjectLauncherPackagingSettings.cpp
+++ b/Engine/Source/Developer/ProjectLauncher/Private/Widgets/Package/SProjectLauncherPackagingSettings.cpp
@@ -161,6 +161,82 @@ void SProjectLauncherPackagingSettings::Construct(const FArguments& InArgs, cons
 									.Text(LOCTEXT("MakeBinaryConfigCheckBoxText", "Make a binary config file for faster runtime startup times"))
 								]
 							]
+
+						+ SVerticalBox::Slot()
+							.AutoHeight()
+							.Padding(0.0, 4.0, 0.0, 0.0)
+							[
+								SNew(STextBlock)
+								.Text(LOCTEXT("RefBlockDbFileNameText", "Optional I/O Store Reference Block Database:"))
+							]
+
+						+ SVerticalBox::Slot()
+							.AutoHeight()
+							.Padding(0.0, 4.0, 0.0, 0.0)
+							[
+								SNew(SHorizontalBox)
+
+								+ SHorizontalBox::Slot()
+									.FillWidth(1.0)
+									.Padding(0.0, 0.0, 0.0, 3.0)
+									[
+										// path textbox for previously compressed containers.
+										SAssignNew(ReferenceContainerGlobalFileName, SEditableTextBox)
+										.IsEnabled(this, &SProjectLauncherPackagingSettings::IsReferenceBlockDbEditable)
+										.OnTextCommitted(this, &SProjectLauncherPackagingSettings::OnRefBlockDbFileNameTextCommitted)
+										.OnTextChanged(this, &SProjectLauncherPackagingSettings::OnRefBlockDbFileNameTextChanged)
+										.ToolTipText(LOCTEXT("RefBlockDbFileNameToolTip", "The path to a global.utoc in a directory of previously released iostore containers."))
+										.HintText(LOCTEXT("RefBlockDbFileNameHint", "<path/to/global.utoc>"))
+									]
+
+								+ SHorizontalBox::Slot()
+									.AutoWidth()
+									.HAlign(HAlign_Right)
+									.Padding(4.0, 0.0, 0.0, 0.0)
+									[
+										// browse button for the global.utoc for a previously staged game.
+										SNew(SButton)
+										.ContentPadding(FMargin(6.0, 2.0))
+										.IsEnabled(this, &SProjectLauncherPackagingSettings::IsReferenceBlockDbEditable)
+										.Text(LOCTEXT("BrowseButtonText", "Browse..."))
+										.ToolTipText(LOCTEXT("BrowseForRefDbButtonToolTip", "Browse for the global.utoc file from a previously staged/released game."))
+										.OnClicked(this, &SProjectLauncherPackagingSettings::HandleRefBlockDbBrowseButtonClicked)
+									]
+							] // end reference container line
+
+						+ SVerticalBox::Slot()
+							.AutoHeight()
+							.Padding(0.0, 4.0, 0.0, 0.0)
+							[
+								SNew(SHorizontalBox)
+
+								+ SHorizontalBox::Slot()
+								.FillWidth(1.0)
+								.Padding(0.0, 0.0, 0.0, 3.0)
+								[
+									// path textbox for previously compressed containers' crypto keys
+									SAssignNew(ReferenceContainerCryptoKeysFileName, SEditableTextBox)
+									.IsEnabled(this, &SProjectLauncherPackagingSettings::IsReferenceBlockDbEditable)
+									.OnTextCommitted(this, &SProjectLauncherPackagingSettings::OnRefBlockCryptoFileNameTextCommitted)
+									.OnTextChanged(this, &SProjectLauncherPackagingSettings::OnRefBlockCryptoFileNameTextChanged)
+									.ToolTipText(LOCTEXT("RefBlockDbFileNameToolTip", "The path to a crypto.json to decrypt the reference block containers, if needed."))
+									.HintText(LOCTEXT("RefBlockDbFileNameHint", "<path/to/crypto.json, optional>"))
+								]
+
+								+ SHorizontalBox::Slot()
+									.AutoWidth()
+									.HAlign(HAlign_Right)
+									.Padding(4.0, 0.0, 0.0, 0.0)
+									[
+										// browse button for the crypto.json to decrypt the ref block containers, if needed
+										SNew(SButton)
+										.ContentPadding(FMargin(6.0, 2.0))
+										.IsEnabled(this, &SProjectLauncherPackagingSettings::IsReferenceBlockDbEditable)
+										.Text(LOCTEXT("BrowseButtonText", "Browse..."))
+										.ToolTipText(LOCTEXT("BrowseForRefDbButtonToolTip", "Browse for the crypto.json file to decrypt the reference containers, if needed."))
+										.OnClicked(this, &SProjectLauncherPackagingSettings::HandleRefBlockCryptoBrowseButtonClicked)
+									]
+							] // end reference container crypto keys line
 					]
 			]
 	];
@@ -170,9 +246,16 @@ void SProjectLauncherPackagingSettings::Construct(const FArguments& InArgs, cons
 END_SLATE_FUNCTION_BUILD_OPTIMIZATION
 
 
-void SProjectLauncherPackagingSettings::UpdateDirectoryPathText()
+void SProjectLauncherPackagingSettings::UpdatePackagingModeWidgets()
 {
 	DirectoryPathTextBox->SetText(HandleDirectoryPathText());
+
+	ILauncherProfilePtr SelectedProfile = Model->GetSelectedProfile();
+	if (SelectedProfile.IsValid())
+	{
+		ReferenceContainerCryptoKeysFileName->SetText(FText::FromString(SelectedProfile->GetReferenceContainerCryptoKeysFileName()));
+		ReferenceContainerGlobalFileName->SetText(FText::FromString(SelectedProfile->GetReferenceContainerGlobalFileName()));
+	}
 }
 
 
@@ -321,7 +404,7 @@ FText SProjectLauncherPackagingSettings::HandleHintPathText() const
 
 void SProjectLauncherPackagingSettings::HandleProfileManagerProfileSelected(const ILauncherProfilePtr& SelectedProfile, const ILauncherProfilePtr& PreviousProfile)
 {
-	UpdateDirectoryPathText();
+	UpdatePackagingModeWidgets();
 }
 
 
@@ -399,4 +482,133 @@ void SProjectLauncherPackagingSettings::OnTextCommitted(const FText& InText, ETe
 }
 
 
+
+bool SProjectLauncherPackagingSettings::IsReferenceBlockDbEditable() const
+{
+	if (IsEditable() == false)
+	{
+		return false;
+	}
+
+	ILauncherProfilePtr SelectedProfile = Model->GetSelectedProfile();
+	if (SelectedProfile.IsValid())
+	{
+		return SelectedProfile->IsUsingIoStore();
+	}
+	return false;
+}
+
+void SProjectLauncherPackagingSettings::OnRefBlockDbFileNameTextChanged(const FText& InText)
+{
+	ILauncherProfilePtr SelectedProfile = Model->GetSelectedProfile();
+	if (SelectedProfile.IsValid())
+	{
+		SelectedProfile->SetReferenceContainerGlobalFileName(InText.ToString());
+	}
+}
+
+
+void SProjectLauncherPackagingSettings::OnRefBlockDbFileNameTextCommitted(const FText& InText, ETextCommit::Type CommitInfo)
+{
+	if (CommitInfo == ETextCommit::OnEnter)
+	{
+		ILauncherProfilePtr SelectedProfile = Model->GetSelectedProfile();
+		if (SelectedProfile.IsValid())
+		{
+			SelectedProfile->SetReferenceContainerGlobalFileName(InText.ToString());
+		}
+	}
+}
+
+FReply SProjectLauncherPackagingSettings::HandleRefBlockDbBrowseButtonClicked()
+{
+	IDesktopPlatform* DesktopPlatform = FDesktopPlatformModule::Get();
+	if (DesktopPlatform)
+	{
+		TSharedPtr<SWindow> ParentWindow = FSlateApplication::Get().FindWidgetWindow(AsShared());
+		void* ParentWindowHandle = (ParentWindow.IsValid() && ParentWindow->GetNativeWindow().IsValid()) ? ParentWindow->GetNativeWindow()->GetOSWindowHandle() : nullptr;
+
+		TArray<FString> SelectedFiles;
+		bool bFileSelected = DesktopPlatform->OpenFileDialog(
+			ParentWindowHandle, 
+			LOCTEXT("GlobalUtocTitle", "Select a global.utoc file").ToString(),
+			ReferenceContainerGlobalFileName->GetText().ToString(),
+			TEXT(""), 
+			TEXT("global.utoc files|global.utoc"),
+			EFileDialogFlags::None /* only single file selection */,
+			SelectedFiles
+			);
+
+		if (bFileSelected && SelectedFiles.Num())
+		{
+			ReferenceContainerGlobalFileName->SetText(FText::FromString(SelectedFiles[0]));
+
+			ILauncherProfilePtr SelectedProfile = Model->GetSelectedProfile();
+			if (SelectedProfile.IsValid())
+			{
+				SelectedProfile->SetReferenceContainerGlobalFileName(SelectedFiles[0]);
+			}
+		}
+	}
+
+	return FReply::Handled();
+}
+
+
+void SProjectLauncherPackagingSettings::OnRefBlockCryptoFileNameTextChanged(const FText& InText)
+{
+	ILauncherProfilePtr SelectedProfile = Model->GetSelectedProfile();
+	if (SelectedProfile.IsValid())
+	{
+		SelectedProfile->SetReferenceContainerCryptoKeysFileName(InText.ToString());
+	}
+}
+
+
+void SProjectLauncherPackagingSettings::OnRefBlockCryptoFileNameTextCommitted(const FText& InText, ETextCommit::Type CommitInfo)
+{
+	if (CommitInfo == ETextCommit::OnEnter)
+	{
+		ILauncherProfilePtr SelectedProfile = Model->GetSelectedProfile();
+		if (SelectedProfile.IsValid())
+		{
+			SelectedProfile->SetReferenceContainerCryptoKeysFileName(InText.ToString());
+		}
+	}
+}
+
+FReply SProjectLauncherPackagingSettings::HandleRefBlockCryptoBrowseButtonClicked()
+{
+	IDesktopPlatform* DesktopPlatform = FDesktopPlatformModule::Get();
+	if (DesktopPlatform)
+	{
+		TSharedPtr<SWindow> ParentWindow = FSlateApplication::Get().FindWidgetWindow(AsShared());
+		void* ParentWindowHandle = (ParentWindow.IsValid() && ParentWindow->GetNativeWindow().IsValid()) ? ParentWindow->GetNativeWindow()->GetOSWindowHandle() : nullptr;
+
+		TArray<FString> SelectedFiles;
+		bool bFileSelected = DesktopPlatform->OpenFileDialog(
+			ParentWindowHandle,
+			LOCTEXT("GlobalUtocKeysTitle", "Select a crypto.json file").ToString(),
+			ReferenceContainerCryptoKeysFileName->GetText().ToString(),
+			TEXT(""),
+			TEXT("crypto.json files|crypto.json"),
+			EFileDialogFlags::None /* only single file selection */,
+			SelectedFiles
+		);
+
+		if (bFileSelected && SelectedFiles.Num())
+		{
+			ReferenceContainerCryptoKeysFileName->SetText(FText::FromString(SelectedFiles[0]));
+
+			ILauncherProfilePtr SelectedProfile = Model->GetSelectedProfile();
+			if (SelectedProfile.IsValid())
+			{
+				SelectedProfile->SetReferenceContainerCryptoKeysFileName(SelectedFiles[0]);
+			}
+		}
+	}
+
+	return FReply::Handled();
+}
+
 #undef LOCTEXT_NAMESPACE
diff --git a/Engine/Source/Developer/ProjectLauncher/Private/Widgets/Package/SProjectLauncherPackagingSettings.h b/Engine/Source/Developer/ProjectLauncher/Private/Widgets/Package/SProjectLauncherPackagingSettings.h
--- a/Engine/Source/Developer/ProjectLauncher/Private/Widgets/Package/SProjectLauncherPackagingSettings.h
+++ b/Engine/Source/Developer/ProjectLauncher/Private/Widgets/Package/SProjectLauncherPackagingSettings.h
@@ -39,8 +39,8 @@ public:
 	 */
 	void Construct( const FArguments& InArgs, const TSharedRef<FProjectLauncherModel>& InModel );
 
-	// Changes the directory path to the correct packaging mode version
-	void UpdateDirectoryPathText();
+	// Update any widgets to reflect the change in packaging mode.
+	void UpdatePackagingModeWidgets();
 
 private:
 
@@ -56,6 +56,16 @@ private:
 	void HandleMakeBinaryConfigCheckStateChanged(ECheckBoxState NewState);
 	ECheckBoxState HandleMakeBinaryConfigCheckBoxIsChecked() const;
 
+	// Reference block database handlers.
+	bool IsReferenceBlockDbEditable() const;
+	void OnRefBlockDbFileNameTextCommitted(const FText& InText, ETextCommit::Type CommitInfo);
+	void OnRefBlockDbFileNameTextChanged(const FText& InText);
+	FReply HandleRefBlockDbBrowseButtonClicked();
+	void OnRefBlockCryptoFileNameTextCommitted(const FText& InText, ETextCommit::Type CommitInfo);
+	void OnRefBlockCryptoFileNameTextChanged(const FText& InText);
+	FReply HandleRefBlockCryptoBrowseButtonClicked();
+
+
 	// Callback for getting the content text of the 'Directory' label.
 	FText HandleDirectoryTitleText() const;
 	FText HandleDirectoryPathText() const;
@@ -82,4 +92,7 @@ private:
 
 	// Holds the repository path text box.
 	TSharedPtr<SEditableTextBox> DirectoryPathTextBox;
+
+	TSharedPtr<SEditableTextBox> ReferenceContainerGlobalFileName;
+	TSharedPtr<SEditableTextBox> ReferenceContainerCryptoKeysFileName;
 };
diff --git a/Engine/Source/Programs/AutomationTool/AutomationUtils/ProjectParams.cs b/Engine/Source/Programs/AutomationTool/AutomationUtils/ProjectParams.cs
--- a/Engine/Source/Programs/AutomationTool/AutomationUtils/ProjectParams.cs
+++ b/Engine/Source/Programs/AutomationTool/AutomationUtils/ProjectParams.cs
@@ -483,7 +483,9 @@ namespace AutomationTool
             string CreateReleaseVersion = null,
 			string CreateReleaseVersionBasePath = null,
 			string BasedOnReleaseVersionBasePath = null,
-            bool? GeneratePatch = null,
+			string ReferenceContainerGlobalFileName = null,
+			string ReferenceContainerCryptoKeys = null,
+			bool? GeneratePatch = null,
 			bool? AddPatchLevel = null,
 			bool? StageBaseReleasePaks = null,
             string DiscVersion = null,
@@ -649,6 +651,8 @@ namespace AutomationTool
 			this.Cook = GetParamValueIfNotSpecified(Command, Cook, this.Cook, "cook");
 			this.CreateReleaseVersionBasePath = ParseParamValueIfNotSpecified(Command, CreateReleaseVersionBasePath, "createreleaseversionroot", String.Empty);
 			this.BasedOnReleaseVersionBasePath = ParseParamValueIfNotSpecified(Command, BasedOnReleaseVersionBasePath, "basedonreleaseversionroot", String.Empty);
+			this.ReferenceContainerGlobalFileName = ParseParamValueIfNotSpecified(Command, ReferenceContainerGlobalFileName, "ReferenceContainerGlobalFileName", String.Empty);
+			this.ReferenceContainerCryptoKeys = ParseParamValueIfNotSpecified(Command, ReferenceContainerCryptoKeys, "ReferenceContainerCryptoKeys", String.Empty); 
 			this.OriginalReleaseVersion = ParseParamValueIfNotSpecified(Command, OriginalReleaseVersion, "originalreleaseversion", String.Empty);
 			this.CreateReleaseVersion = ParseParamValueIfNotSpecified(Command, CreateReleaseVersion, "createreleaseversion", String.Empty);
             this.BasedOnReleaseVersion = ParseParamValueIfNotSpecified(Command, BasedOnReleaseVersion, "basedonreleaseversion", String.Empty);
@@ -1704,6 +1708,18 @@ namespace AutomationTool
 		/// </summary>
 		public string CreateReleaseVersionBasePath;
 
+		/// <summary>
+		/// Stage: Path to the global.utoc file for a directory of iostore containers to use as a source of compressed
+		/// chunks when writing new containers. See -ReferenceContainerGlobalFileName in IoStoreUtilities.cpp.
+		/// </summary>
+		public string ReferenceContainerGlobalFileName;
+
+		/// <summary>
+		/// Stage: Path to the crypto.json file to use for decrypting ReferenceContainerFlobalFileName, if needed.
+		/// </summary>
+		public string ReferenceContainerCryptoKeys;
+
+
 		/// <summary>
 		/// Are we generating a patch, generate a patch from a previously released version of the game (use CreateReleaseVersion to create a release). 
 		/// this requires BasedOnReleaseVersion
diff --git a/Engine/Source/Programs/AutomationTool/Scripts/BuildCookRun.Automation.cs b/Engine/Source/Programs/AutomationTool/Scripts/BuildCookRun.Automation.cs
--- a/Engine/Source/Programs/AutomationTool/Scripts/BuildCookRun.Automation.cs
+++ b/Engine/Source/Programs/AutomationTool/Scripts/BuildCookRun.Automation.cs
@@ -80,6 +80,18 @@ public class BuildCookRun : BuildCommand
             Params.CulturesToCook = new ParamList<string>(CulturesToCook.Split('+'));
         }
 
+		var ReferenceContainerGlobalFileName = ParseParamValue("ReferenceContainerGlobalFileName");
+		if (!String.IsNullOrEmpty(ReferenceContainerGlobalFileName))
+		{
+			Params.ReferenceContainerGlobalFileName = ReferenceContainerGlobalFileName;
+		}
+
+		var ReferenceContainerCryptoKeys = ParseParamValue("ReferenceContainerCryptoKeys");
+		if (!String.IsNullOrEmpty(ReferenceContainerCryptoKeys))
+		{
+			Params.ReferenceContainerCryptoKeys = ReferenceContainerCryptoKeys;
+		}
+
 		if (Params.DedicatedServer)
 		{
 			foreach (var ServerPlatformInstance in Params.ServerTargetPlatformInstances)
diff --git a/Engine/Source/Programs/AutomationTool/Scripts/CopyBuildToStagingDirectory.Automation.cs b/Engine/Source/Programs/AutomationTool/Scripts/CopyBuildToStagingDirectory.Automation.cs
--- a/Engine/Source/Programs/AutomationTool/Scripts/CopyBuildToStagingDirectory.Automation.cs
+++ b/Engine/Source/Programs/AutomationTool/Scripts/CopyBuildToStagingDirectory.Automation.cs
@@ -3133,6 +3133,16 @@ namespace AutomationScripts
 				}
 
 
+				if (Params.ReferenceContainerGlobalFileName.Length != 0)
+				{
+					AdditionalArgs += " -ReferenceContainerGlobalFileName=\"" + Params.ReferenceContainerGlobalFileName + "\"";
+					if (Params.ReferenceContainerCryptoKeys.Length != 0)
+					{
+						AdditionalArgs += " -ReferenceContainerCryptoKeys=\"" + Params.ReferenceContainerCryptoKeys + "\"";
+					}
+				}
+
+
 				AdditionalArgs += " " + Params.AdditionalIoStoreOptions;
 
 				RunIoStore(Params, SC, IoStoreCommandsFileName, FinalOrderFiles, AdditionalArgs);
@@ -4135,6 +4145,15 @@ namespace AutomationScripts
 				CommandLineBuilder.Append("-iterative ");
 			}
 
+			if (!string.IsNullOrEmpty(Params.ReferenceContainerGlobalFileName))
+			{
+				CommandLineBuilder.AppendFormat("-ReferenceContainerGlobalFileName={0}", Params.ReferenceContainerGlobalFileName);
+				if (!string.IsNullOrEmpty(Params.ReferenceContainerCryptoKeys))
+				{
+					CommandLineBuilder.AppendFormat("-ReferenceContainerCryptoKeys={0}", Params.ReferenceContainerCryptoKeys);
+				}
+			}
+
 			String CompleteCommandLine = CommandLineBuilder.ToString().Trim();
 
 			if (string.IsNullOrEmpty(CompleteCommandLine))
diff --git a/Engine/Source/Runtime/Core/Private/IO/IoStore.cpp b/Engine/Source/Runtime/Core/Private/IO/IoStore.cpp
--- a/Engine/Source/Runtime/Core/Private/IO/IoStore.cpp
+++ b/Engine/Source/Runtime/Core/Private/IO/IoStore.cpp
@@ -24,6 +24,7 @@
 #include "Serialization/LargeMemoryReader.h"
 #include "Serialization/MemoryWriter.h"
 #include "ProfilingDebugging/CountersTrace.h"
+#include "Tasks/Task.h"
 
 DEFINE_LOG_CATEGORY(LogIoStore);
 
@@ -203,6 +204,9 @@ struct FChunkBlock
 {
 	const uint8* UncompressedData = nullptr;
 	FIoBuffer* IoBuffer = nullptr;
+
+	// This is the size of the actual block after encryption alignment, and is
+	// set in EncryptAndSign. This happens whether or not the container is encrypted.
 	uint64 Size = 0;
 	uint64 CompressedSize = 0;
 	uint64 UncompressedSize = 0;
@@ -405,6 +409,10 @@ public:
 		QueueEntry->Request->PrepareSourceBufferAsync(QueueEntry->BeginCompressionBarrier);
 	}
 
+	//
+	// This must be called prior to the consumer being dispatched in order to prevent resource
+	// contention deadlock (even if using retracting waits)
+	//
 	FIoBuffer* AllocCompressionBuffer(int32 TotalEntryChunkBlocksCount)
 	{
 		if (TotalEntryChunkBlocksCount > TotalCompressionBufferCount)
@@ -683,6 +691,11 @@ public:
 	{
 	}
 
+	void SetReferenceChunkDatabase(TSharedPtr<IIoStoreWriterReferenceChunkDatabase> InReferenceChunkDatabase)
+	{
+		ReferenceChunkDatabase = InReferenceChunkDatabase;
+	}
+
 	void EnumerateChunks(TFunction<bool(const FIoStoreTocChunkInfo&)>&& Callback) const
 	{
 		const FIoStoreTocResource& TocResource = Toc.GetTocResource();
@@ -1436,18 +1449,70 @@ private:
 		}
 		
 		Entry->ChunkBlocks.SetNum(int32(NumChunkBlocks));
+		{
+			// We must allocate resources for our tasks up front to prevent resource deadlock.
+			// Note that for Reference Chunk loaded blocks this will reserve more than actually
+			// needed.
+			uint64 BytesToProcess = Entry->UncompressedSize;
+			const uint8* UncompressedData = SourceBuffer->Data();
+			for (int32 BlockIndex = 0; BlockIndex < NumChunkBlocks; ++BlockIndex)
+			{
+				FChunkBlock& Block = Entry->ChunkBlocks[BlockIndex];
+				Block.IoBuffer = WriterContext->AllocCompressionBuffer(int32(NumChunkBlocks));
+				Block.CompressionMethod = CompressionMethod;
+				Block.UncompressedSize = FMath::Min(BytesToProcess, WriterSettings.CompressionBlockSize);
+				Block.UncompressedData = UncompressedData;
+				BytesToProcess -= Block.UncompressedSize;
+				UncompressedData += Block.UncompressedSize;
+			}
+		}
 
-		uint64 BytesToProcess = Entry->UncompressedSize;
-		const uint8* UncompressedData = SourceBuffer->Data();
-		for (int32 BlockIndex = 0; BlockIndex < NumChunkBlocks; ++BlockIndex)
+		//
+		// Check if this chunk exists in the reference cache.
+		//
+		if (ReferenceChunkDatabase.IsValid())
 		{
-			FChunkBlock& Block = Entry->ChunkBlocks[BlockIndex];
-			Block.IoBuffer = WriterContext->AllocCompressionBuffer(int32(NumChunkBlocks));
-			Block.CompressionMethod = CompressionMethod;
-			Block.UncompressedSize = FMath::Min(BytesToProcess, WriterSettings.CompressionBlockSize);
-			Block.UncompressedData = UncompressedData;
-			BytesToProcess -= Block.UncompressedSize;
-			UncompressedData += Block.UncompressedSize;
+			TPair<FIoContainerId, FIoChunkHash> ChunkKey(ContainerSettings.ContainerId, Entry->ChunkHash);
+
+			// Valid chunks must create the same decompressed bits, but can have different compressed bits.
+			// Since we are on a lightweight dispatch thread, the actual read is async, as is the processing
+			// of the results.
+			bool GotChunk = ReferenceChunkDatabase->RetrieveChunk(ChunkKey, CompressionMethod, Entry->UncompressedSize, NumChunkBlocks, [this, Entry](TIoStatusOr<FIoStoreCompressedReadResult> InReadResult)
+			{
+				// If we fail here, in order to recover we effectively need to re-kick this chunk's
+				// BeginCompress()... however, this is just a direct read and should only fail
+				// in catastrophic scenarios (loss of connection on a network drive?).
+				FIoStoreCompressedReadResult ReadResult = InReadResult.ValueOrDie();
+
+				uint8* ReferenceData = ReadResult.IoBuffer.GetData();
+				for (int32 BlockIndex = 0; BlockIndex < ReadResult.Blocks.Num(); ++BlockIndex)
+				{
+					FIoStoreCompressedBlockInfo& ReferenceBlock = ReadResult.Blocks[BlockIndex];
+					FChunkBlock& Block = Entry->ChunkBlocks[BlockIndex];
+					Block.CompressionMethod = ReferenceBlock.CompressionMethod;
+					Block.CompressedSize = ReferenceBlock.CompressedSize;
+
+					// Future optimization: ReadCompressed returns the memory ready to encrypt in one
+					// large contiguous buffer (i.e. padded). We could use the FIoBuffer functionality of referencing a 
+					// sub block from a parent buffer, however this would mean that we need to add support
+					// for tracking the memory usage in order to remain within our prescribed limits. To do this
+					// requires releasing the entire chunk's memory at once after WriteEntry.
+					// As it stands, we temporarily use untracked memory in the ReadCompressed call (in RetrieveChunk),
+					// then immediately copy it to tracked memory. There's some waste as tracked memory is mod CompressionBlockSize
+					// and we are post compression, so with the average 50% compression rate, we're using double the memory
+					// we "could".
+					FMemory::Memcpy(Block.IoBuffer->GetData(), ReferenceData, Block.CompressedSize);
+					ReferenceData += ReferenceBlock.AlignedSize;
+				}
+
+				Entry->FinishCompressionBarrier->DispatchSubsequents();
+			});
+
+			if (GotChunk)
+			{
+				// Lambda handles dispatch subsequents
+				return;
+			}
 		}
 
 		if (CompressionMethod == NAME_None)
@@ -1740,6 +1805,8 @@ private:
 	bool						bHasMemoryMappedEntry = false;
 	bool						bHasFlushed = false;
 	bool						bHasResult = false;
+	TSharedPtr<IIoStoreWriterReferenceChunkDatabase> ReferenceChunkDatabase;
+
 
 	friend class FIoStoreWriterContextImpl;
 };
@@ -1905,8 +1972,78 @@ public:
 
 	}
 
+	//
+	// GenericPlatformFile isn't designed around a lot of jobs throwing accesses at it, so instead we 
+	// use IFileHandle directly and round robin between a number of file handles in order to saturate
+	// year 2022 ssd drives. For a file hot in the windows file cache, you can get 4+ GB/s with as few as 
+	// 4 file handles, however a cold file you need upwards of 32 in order to reach ~1.5 GB/s. This is
+	// low because IoStore reads are comparatively small - at most you're reading compression block sized
+	// chunks with uncompressed, however with Oodle those get cut by ~half, so with a default block size
+	// of 64kb, reads are generally less than 32kb, which is tough to use and get full ssd bandwidth out of.
+	//
+	static constexpr uint32 NumHandlesPerFile = 12;
+	struct FContainerFileAccess
+	{
+		FCriticalSection HandleLock[NumHandlesPerFile];
+		IFileHandle* Handle[NumHandlesPerFile];
+		std::atomic_uint32_t NextHandleIndex;
+		bool bValid = false;
+
+		FContainerFileAccess(IPlatformFile& Ipf, const TCHAR* ContainerFileName)
+		{
+			bValid = true;
+			for (uint32 i=0; i < NumHandlesPerFile; i++)
+			{
+				Handle[i] = Ipf.OpenRead(ContainerFileName);
+				if (Handle[i] == nullptr)
+				{
+					bValid = false;
+				}
+			}
+		}
+
+		bool IsValid() const { return bValid; }
+	};
+	
+
+	// Kick off an async read from the iostore container, rotating between the file handles for the partition.
+	UE::Tasks::FTask StartAsyncRead(int32 InPartitionIndex, int64 InPartitionOffset, int64 InReadAmount, uint8* OutBuffer, std::atomic_bool* OutSuccess) const
+	{
+		return UE::Tasks::Launch(TEXT("FIoStoreReader_AsyncRead"), [this, InPartitionIndex, InPartitionOffset, OutBuffer, InReadAmount, OutSuccess]() mutable
+		{
+			FContainerFileAccess* ContainerFileAccess = this->ContainerFileAccessors[InPartitionIndex].Get();
+
+			// Round robin between the file handles. Since we are always reading blocks, everything is ~roughly~ the same
+			// size so we don't have to worry about a single huge read backing up one handle.
+			uint32 OurIndex = ContainerFileAccess->NextHandleIndex.fetch_add(1);
+			OurIndex %= NumHandlesPerFile;
+
+			// Each file handle can only be touched by one task at a time. We use an OS lock so that the OS scheduler
+			// knows we're in a wait state and who we're waiting on.
+			//
+			// CAUTION if any overload of IFileHandle launches tasks (... unlikely ...) this could deadlock if NumHandlesPerFile is more
+			// than the number of worker threads, as the OS lock will not do task retraction.
+			{
+				TRACE_CPUPROFILER_EVENT_SCOPE(FIoStoreReader_StartAsyncRead_Lock);
+				ContainerFileAccess->HandleLock[OurIndex].Lock();
+			}
+
+			bool bReadSucceeded;
+			{
+				TRACE_CPUPROFILER_EVENT_SCOPE(FIoStoreReader_StartAsyncRead_SeekAndRead);
+				ContainerFileAccess->Handle[OurIndex]->Seek(InPartitionOffset);
+				bReadSucceeded = ContainerFileAccess->Handle[OurIndex]->Read(OutBuffer, InReadAmount);
+			}
+
+			OutSuccess->store(bReadSucceeded);
+			ContainerFileAccess->HandleLock[OurIndex].Unlock();
+		});
+	}
+
 	UE_NODISCARD FIoStatus Initialize(const TCHAR* InContainerPath, const TMap<FGuid, FAES::FAESKey>& InDecryptionKeys)
 	{
+		ContainerPath = InContainerPath;
+
 		TStringBuilder<256> TocFilePath;
 		TocFilePath.Append(InContainerPath);
 		TocFilePath.Append(TEXT(".utoc"));
@@ -1921,7 +2058,7 @@ public:
 		Toc.Initialize();
 
 		IPlatformFile& Ipf = FPlatformFileManager::Get().GetPlatformFile();
-		ContainerFileHandles.Reserve(TocResource.Header.PartitionCount);
+		ContainerFileAccessors.Reserve(TocResource.Header.PartitionCount);
 		for (uint32 PartitionIndex = 0; PartitionIndex < TocResource.Header.PartitionCount; ++PartitionIndex)
 		{
 			TStringBuilder<256> ContainerFilePath;
@@ -1931,8 +2068,9 @@ public:
 				ContainerFilePath.Append(FString::Printf(TEXT("_s%d"), PartitionIndex));
 			}
 			ContainerFilePath.Append(TEXT(".ucas"));
-			ContainerFileHandles.Emplace(Ipf.OpenAsyncRead(*ContainerFilePath));
-			if (!ContainerFileHandles.Last())
+
+			ContainerFileAccessors.Emplace(new FContainerFileAccess(Ipf, *ContainerFilePath));
+			if (ContainerFileAccessors.Last().IsValid() == false)
 			{
 				return FIoStatusBuilder(EIoErrorCode::FileOpenFailed) << TEXT("Failed to open IoStore container file '") << *TocFilePath << TEXT("'");
 			}
@@ -2030,28 +2168,29 @@ public:
 		}
 	}
 
-	TFuture<TIoStatusOr<FIoBuffer>> ReadAsync(const FIoChunkId& ChunkId, const FIoReadOptions& Options) const
+	UE::Tasks::TTask<TIoStatusOr<FIoBuffer>> ReadAsync(const FIoChunkId& ChunkId, const FIoReadOptions& Options) const
 	{
 		TRACE_CPUPROFILER_EVENT_SCOPE(ReadChunkAsync);
 
 		struct FState
 		{
-			TPromise<TIoStatusOr<FIoBuffer>> Promise;
-			IAsyncReadRequest* ReadRequest = nullptr;
-			uint8* CompressedBuffer = nullptr;
+			TArray64<uint8> CompressedBuffer;
 			uint64 CompressedSize = 0;
 			uint64 UncompressedSize = 0;
 			TOptional<FIoBuffer> UncompressedBuffer;
-			TAtomic<bool> bReadFailed{ false };
-			TAtomic<bool> bUncompressFailed{ false };
+			std::atomic_bool bReadSucceeded {false};
+			std::atomic_bool bUncompressFailed { false };
 		};
 
 		const FIoOffsetAndLength* OffsetAndLength = Toc.GetOffsetAndLength(ChunkId);
 		if (!OffsetAndLength )
 		{
-			TPromise<TIoStatusOr<FIoBuffer>> Promise;
-			Promise.EmplaceValue(FIoStatus(EIoErrorCode::NotFound, TEXT("Unknown chunk ID")));
-			return Promise.GetFuture();
+			// Currently there's no way to make a task with a valid result that just emplaces
+			// without running.
+			return UE::Tasks::Launch(TEXT("FIoStoreRead_Error"), 
+				[] { return TIoStatusOr<FIoBuffer>(FIoStatus(EIoErrorCode::NotFound, TEXT("Unknown chunk ID"))); },
+				UE::Tasks::ETaskPriority::Normal,
+				UE::Tasks::EExtendedTaskPriority::Inline); // force execution on this thread
 		}
 
 		const uint64 RequestedOffset = Options.GetOffset();
@@ -2064,9 +2203,12 @@ public:
 		const int32 BlockCount = LastBlockIndex - FirstBlockIndex + 1;
 		if (!BlockCount)
 		{
-			TPromise<TIoStatusOr<FIoBuffer>> Promise;
-			Promise.EmplaceValue(FIoBuffer());
-			return Promise.GetFuture();
+			// Currently there's no way to make a task with a valid result that just emplaces
+			// without running.
+			return UE::Tasks::Launch(TEXT("FIoStoreRead_Empty"),
+				[] { return TIoStatusOr<FIoBuffer>(); },
+				UE::Tasks::ETaskPriority::Normal,
+				UE::Tasks::EExtendedTaskPriority::Inline); // force execution on this thread
 		}
 		const FIoStoreTocCompressedBlockEntry& FirstBlock = TocResource.CompressionBlocks[FirstBlockIndex];
 		const FIoStoreTocCompressedBlockEntry& LastBlock = TocResource.CompressionBlocks[LastBlockIndex];
@@ -2077,119 +2219,105 @@ public:
 		FState* State = new FState();
 		State->CompressedSize = ReadEndOffset - ReadStartOffset;
 		State->UncompressedSize = ResolvedSize;
+		State->CompressedBuffer.AddUninitialized(State->CompressedSize);
+		State->UncompressedBuffer.Emplace(State->UncompressedSize);
 
-		FGraphEventArray PrepareDecompressionPrereqs;
-		FGraphEventArray DecompressionPrereqs;
-		FGraphEventArray FinalizePreReqs;
-
-		FGraphEventRef ReadFinishedBarrier = FGraphEvent::CreateGraphEvent();
-		PrepareDecompressionPrereqs.Add(ReadFinishedBarrier);
+		UE::Tasks::FTask ReadJob = StartAsyncRead(PartitionIndex, ReadStartOffset, (int32)State->CompressedSize, State->CompressedBuffer.GetData(), &State->bReadSucceeded);
 
-		FGraphEventRef PrepareDecompressionTask = FFunctionGraphTask::CreateAndDispatchWhenReady([State]()
-			{
-				if (State->bReadFailed.Load())
-				{
-					return;
-				}
+		UE::Tasks::TTask<TIoStatusOr<FIoBuffer>> ReturnTask = UE::Tasks::Launch(TEXT("FIoStoreReader::AsyncRead"), [this, State, PartitionIndex, CompressionBlockSize, ResolvedOffset, FirstBlockIndex, LastBlockIndex, ResolvedSize, ReadStartOffset, &TocResource]()
+		{			
+			UE::Tasks::FTaskEvent DecompressionDoneEvent(TEXT("FIoStoreReader::DecompressionDone"));
 
-				State->CompressedBuffer = State->ReadRequest->GetReadResults();
-				State->UncompressedBuffer.Emplace(State->UncompressedSize);
-			}, TStatId(), & PrepareDecompressionPrereqs, ENamedThreads::AnyHiPriThreadNormalTask);
-		DecompressionPrereqs.Add(PrepareDecompressionTask);
+			std::atomic_int32_t DecompressionJobsRemaining = LastBlockIndex - FirstBlockIndex + 1;
 
-		uint64 CompressedSourceOffset = 0;
-		uint64 UncompressedDestinationOffset = 0;
-		uint64 OffsetInBlock = ResolvedOffset % CompressionBlockSize;
-		uint64 RemainingSize = ResolvedSize;
-		for (int32 BlockIndex = FirstBlockIndex; BlockIndex <= LastBlockIndex; ++BlockIndex)
-		{
-			FGraphEventRef DecompressBlockTask = FFunctionGraphTask::CreateAndDispatchWhenReady([this, State, BlockIndex, CompressedSourceOffset, UncompressedDestinationOffset, OffsetInBlock, RemainingSize]()
+			uint64 CompressedSourceOffset = 0;
+			uint64 UncompressedDestinationOffset = 0;
+			uint64 OffsetInBlock = ResolvedOffset % CompressionBlockSize;
+			uint64 RemainingSize = ResolvedSize;
+			for (int32 BlockIndex = FirstBlockIndex; BlockIndex <= LastBlockIndex; ++BlockIndex)
+			{
+				UE::Tasks::FTask DecompressBlockTask = UE::Tasks::Launch(TEXT("FIoStoreReader::Decompress"), [this, State, BlockIndex, CompressedSourceOffset, UncompressedDestinationOffset, OffsetInBlock, RemainingSize, &DecompressionDoneEvent, &DecompressionJobsRemaining]()
 				{
-					if (State->bReadFailed.Load())
+					if (State->bReadSucceeded)
 					{
-						return;
-					}
-					uint8* CompressedSource = State->CompressedBuffer + CompressedSourceOffset;
-					uint8* UncompressedDestination = State->UncompressedBuffer->Data() + UncompressedDestinationOffset;
-					const FIoStoreTocResource& TocResource = Toc.GetTocResource();
-					const FIoStoreTocCompressedBlockEntry& CompressionBlock = TocResource.CompressionBlocks[BlockIndex];
-					const uint32 RawSize = Align(CompressionBlock.GetCompressedSize(), FAES::AESBlockSize);
-					const uint32 UncompressedSize = CompressionBlock.GetUncompressedSize();
-					FName CompressionMethod = TocResource.CompressionMethods[CompressionBlock.GetCompressionMethodIndex()];
-					if (EnumHasAnyFlags(TocResource.Header.ContainerFlags, EIoContainerFlags::Encrypted))
-					{
-						TRACE_CPUPROFILER_EVENT_SCOPE(Decrypt);
-						check(CompressedSource + RawSize <= State->CompressedBuffer + State->CompressedSize);
-						FAES::DecryptData(CompressedSource, RawSize, DecryptionKey);
-					}
-					if (CompressionMethod.IsNone())
-					{
-						check(UncompressedDestination + UncompressedSize - OffsetInBlock <= State->UncompressedBuffer->Data() + State->UncompressedBuffer->DataSize());
-						FMemory::Memcpy(UncompressedDestination, CompressedSource + OffsetInBlock, UncompressedSize - OffsetInBlock);
-					}
-					else
-					{
-						bool bUncompressed;
-						if (OffsetInBlock || RemainingSize < UncompressedSize)
+						uint8* CompressedSource = State->CompressedBuffer.GetData() + CompressedSourceOffset;
+						uint8* UncompressedDestination = State->UncompressedBuffer->Data() + UncompressedDestinationOffset;
+						const FIoStoreTocResource& TocResource = Toc.GetTocResource();
+						const FIoStoreTocCompressedBlockEntry& CompressionBlock = TocResource.CompressionBlocks[BlockIndex];
+						const uint32 RawSize = Align(CompressionBlock.GetCompressedSize(), FAES::AESBlockSize);
+						const uint32 UncompressedSize = CompressionBlock.GetUncompressedSize();
+						FName CompressionMethod = TocResource.CompressionMethods[CompressionBlock.GetCompressionMethodIndex()];
+						if (EnumHasAnyFlags(TocResource.Header.ContainerFlags, EIoContainerFlags::Encrypted))
 						{
-							TArray<uint8> TempBuffer;
-							TempBuffer.SetNumUninitialized(UncompressedSize);
-							bUncompressed = FCompression::UncompressMemory(CompressionMethod, TempBuffer.GetData(), UncompressedSize, CompressedSource, CompressionBlock.GetCompressedSize());
-							uint64 CopySize = FMath::Min<uint64>(UncompressedSize - OffsetInBlock, RemainingSize);
-							FMemory::Memcpy(UncompressedDestination, TempBuffer.GetData() + OffsetInBlock, CopySize);
+							TRACE_CPUPROFILER_EVENT_SCOPE(Decrypt);
+							check(CompressedSource + RawSize <= State->CompressedBuffer.GetData() + State->CompressedSize);
+							FAES::DecryptData(CompressedSource, RawSize, DecryptionKey);
 						}
-						else
+						if (CompressionMethod.IsNone())
 						{
-							check(UncompressedDestination + UncompressedSize <= State->UncompressedBuffer->Data() + State->UncompressedBuffer->DataSize());
-							bUncompressed = FCompression::UncompressMemory(CompressionMethod, UncompressedDestination, UncompressedSize, CompressedSource, CompressionBlock.GetCompressedSize());
+							check(UncompressedDestination + UncompressedSize - OffsetInBlock <= State->UncompressedBuffer->Data() + State->UncompressedBuffer->DataSize());
+							FMemory::Memcpy(UncompressedDestination, CompressedSource + OffsetInBlock, UncompressedSize - OffsetInBlock);
 						}
-						if (!bUncompressed)
+						else
 						{
-							State->bUncompressFailed.Store(true);
+							bool bUncompressed;
+							if (OffsetInBlock || RemainingSize < UncompressedSize)
+							{
+								TArray<uint8> TempBuffer;
+								TempBuffer.SetNumUninitialized(UncompressedSize);
+								bUncompressed = FCompression::UncompressMemory(CompressionMethod, TempBuffer.GetData(), UncompressedSize, CompressedSource, CompressionBlock.GetCompressedSize());
+								uint64 CopySize = FMath::Min<uint64>(UncompressedSize - OffsetInBlock, RemainingSize);
+								FMemory::Memcpy(UncompressedDestination, TempBuffer.GetData() + OffsetInBlock, CopySize);
+							}
+							else
+							{
+								check(UncompressedDestination + UncompressedSize <= State->UncompressedBuffer->Data() + State->UncompressedBuffer->DataSize());
+								bUncompressed = FCompression::UncompressMemory(CompressionMethod, UncompressedDestination, UncompressedSize, CompressedSource, CompressionBlock.GetCompressedSize());
+							}
+							if (!bUncompressed)
+							{
+								State->bUncompressFailed = true;
+							}
 						}
+					} // end if read succeeded
+
+					if (DecompressionJobsRemaining.fetch_add(-1) == 1)
+					{
+						DecompressionDoneEvent.Trigger();
 					}
-				}, TStatId(), &DecompressionPrereqs, ENamedThreads::AnyHiPriThreadNormalTask);
-			FinalizePreReqs.Add(DecompressBlockTask);
-			const FIoStoreTocCompressedBlockEntry& CompressionBlock = TocResource.CompressionBlocks[BlockIndex];
-			const uint32 RawSize = Align(CompressionBlock.GetCompressedSize(), FAES::AESBlockSize);
-			CompressedSourceOffset += RawSize;
-			UncompressedDestinationOffset += CompressionBlock.GetUncompressedSize();
-			RemainingSize -= CompressionBlock.GetUncompressedSize();
-			OffsetInBlock = 0;
-		}
+				}); // end decompression lambda
 
-		FAsyncFileCallBack ReadRequestCallback = [State, ReadFinishedBarrier](bool bWasCancelled, IAsyncReadRequest*)
-		{
-			if (bWasCancelled)
+				const FIoStoreTocCompressedBlockEntry& CompressionBlock = TocResource.CompressionBlocks[BlockIndex];
+				const uint32 RawSize = Align(CompressionBlock.GetCompressedSize(), FAES::AESBlockSize);
+				CompressedSourceOffset += RawSize;
+				UncompressedDestinationOffset += CompressionBlock.GetUncompressedSize();
+				RemainingSize -= CompressionBlock.GetUncompressedSize();
+				OffsetInBlock = 0;
+			} // end for each block
+
+			// Wait for everything
+			DecompressionDoneEvent.BusyWait();
+
+			TIoStatusOr<FIoBuffer> Result;
+			if (State->bReadSucceeded == false)
 			{
-				State->bReadFailed.Store(true);
+				Result = FIoStatus(EIoErrorCode::ReadError, TEXT("Failed reading chunk from container file"));
 			}
-			ReadFinishedBarrier->DispatchSubsequents();
-		};
-
-		State->ReadRequest = ContainerFileHandles[PartitionIndex]->ReadRequest(ReadStartOffset, State->CompressedSize, AIOP_Normal, &ReadRequestCallback);
-		
-		FGraphEventRef FinalizeTask = FFunctionGraphTask::CreateAndDispatchWhenReady([State]()
+			else if (State->bUncompressFailed)
 			{
-				if (State->bReadFailed.Load())
-				{
-					State->Promise.EmplaceValue(FIoStatus(EIoErrorCode::ReadError, TEXT("Failed reading chunk from container file")));
-				}
-				else if (State->bUncompressFailed.Load())
-				{
-					State->Promise.EmplaceValue(FIoStatus(EIoErrorCode::ReadError, TEXT("Failed uncompressing chunk")));
-				}
-				else
-				{
-					State->Promise.SetValue(State->UncompressedBuffer.GetValue());
-				}
-				delete State->CompressedBuffer;
-				delete State->ReadRequest;
-				delete State;
-			}, TStatId(), &FinalizePreReqs, ENamedThreads::AnyHiPriThreadNormalTask);
+				Result = FIoStatus(EIoErrorCode::ReadError, TEXT("Failed uncompressing chunk"));
+			}
+			else
+			{
+				Result = State->UncompressedBuffer.GetValue();
+			}
+			delete State;
 
-		return State->Promise.GetFuture();
-	}
+			return Result;
+		}, UE::Tasks::Prerequisites(ReadJob)); // end read and compress lambda launch
+
+		return ReturnTask;
+	} // end ReadAsync
 
 	TIoStatusOr<FIoBuffer> Read(const FIoChunkId& ChunkId, const FIoReadOptions& Options) const
 	{
@@ -2222,12 +2350,13 @@ public:
 		// We try to overlap the IO for the next block with the decrypt/decompress for the current
 		// block, which requires two IO buffers:
 		TArray<uint8> CompressedBuffers[2];
+		std::atomic_bool AsyncReadSucceeded[2];
 
 		int32 FirstBlockIndex = int32(ResolvedOffset / CompressionBlockSize);
 		int32 LastBlockIndex = int32((Align(ResolvedOffset + ResolvedSize, CompressionBlockSize) - 1) / CompressionBlockSize);
 
 		// Lambda to kick off a read with a sufficient output buffer.
-		auto LaunchBlockRead = [&TocResource, this](int32 BlockIndex, TArray<uint8>& DestinationBuffer)
+		auto LaunchBlockRead = [&TocResource, this](int32 BlockIndex, TArray<uint8>& DestinationBuffer, std::atomic_bool* OutReadSucceeded)
 		{
 			const uint64 CompressionBlockSize = TocResource.Header.CompressionBlockSize;
 			const FIoStoreTocCompressedBlockEntry& CompressionBlock = TocResource.CompressionBlocks[BlockIndex];
@@ -2245,14 +2374,14 @@ public:
 
 			int32 PartitionIndex = int32(CompressionBlock.GetOffset() / TocResource.Header.PartitionSize);
 			int64 PartitionOffset = int64(CompressionBlock.GetOffset() % TocResource.Header.PartitionSize);
-			return ContainerFileHandles[PartitionIndex]->ReadRequest(PartitionOffset, SizeForDecrypt, AIOP_Normal, nullptr, DestinationBuffer.GetData());
+			return StartAsyncRead(PartitionIndex, PartitionOffset, SizeForDecrypt, DestinationBuffer.GetData(), OutReadSucceeded);
 		};
 
 
 		// Kick off the first async read
-		TUniquePtr<IAsyncReadRequest> NextReadRequest;
+		UE::Tasks::FTask NextReadRequest;
 		uint8 NextReadBufferIndex = 0;
-		NextReadRequest.Reset(LaunchBlockRead(FirstBlockIndex, CompressedBuffers[NextReadBufferIndex]));
+		NextReadRequest = LaunchBlockRead(FirstBlockIndex, CompressedBuffers[NextReadBufferIndex], &AsyncReadSucceeded[NextReadBufferIndex]);
 
 		uint64 UncompressedDestinationOffset = 0;
 		uint64 OffsetInBlock = ResolvedOffset % CompressionBlockSize;
@@ -2261,23 +2390,29 @@ public:
 		for (int32 BlockIndex = FirstBlockIndex; BlockIndex <= LastBlockIndex; ++BlockIndex)
 		{
 			// Kick off the next block's IO if there is one
-			TUniquePtr<IAsyncReadRequest> ReadRequest(MoveTemp(NextReadRequest));
-			uint8 OurBufferIndex = NextReadBufferIndex;			
+			UE::Tasks::FTask ReadRequest(MoveTemp(NextReadRequest));
+			uint8 OurBufferIndex = NextReadBufferIndex;
 			if (BlockIndex + 1 <= LastBlockIndex)
 			{
 				NextReadBufferIndex = NextReadBufferIndex ^ 1;
-				NextReadRequest.Reset(LaunchBlockRead(BlockIndex + 1, CompressedBuffers[NextReadBufferIndex]));
+				NextReadRequest = LaunchBlockRead(BlockIndex + 1, CompressedBuffers[NextReadBufferIndex], &AsyncReadSucceeded[NextReadBufferIndex]);
 			}
 
 			// Now, wait for _our_ block's IO
 			{
 				TRACE_CPUPROFILER_EVENT_SCOPE(WaitForIo);
-				ReadRequest->WaitCompletion();
+				ReadRequest.BusyWait();
+			}
+
+			if (AsyncReadSucceeded[OurBufferIndex] == false)
+			{
+				return FIoStatus(EIoErrorCode::ReadError, TEXT("Failed async read in FIoStoreReader::ReadCompressed"));
 			}
 
 			const FIoStoreTocCompressedBlockEntry& CompressionBlock = TocResource.CompressionBlocks[BlockIndex];
-			uint32 RawSize = Align(CompressionBlock.GetCompressedSize(), FAES::AESBlockSize);
 
+			// This also happened in the LaunchBlockRead call, so we know the buffer has the necessary size.
+			uint32 RawSize = Align(CompressionBlock.GetCompressedSize(), FAES::AESBlockSize);
 			if (EnumHasAnyFlags(TocResource.Header.ContainerFlags, EIoContainerFlags::Encrypted))
 			{
 				TRACE_CPUPROFILER_EVENT_SCOPE(Decrypt);
@@ -2298,7 +2433,7 @@ public:
 				if (OffsetInBlock || RemainingSize < UncompressedSize)
 				{
 					// If this block is larger than the amount of data actually requested, decompress to a temp
-					// buffer and then copy out. Should never happen when reading the entire file.
+					// buffer and then copy out. Should never happen when reading the entire chunk.
 					TempBuffer.SetNumUninitialized(UncompressedSize);
 					bUncompressed = FCompression::UncompressMemory(CompressionMethod, TempBuffer.GetData(), UncompressedSize, CompressedBuffers[OurBufferIndex].GetData(), CompressionBlock.GetCompressedSize());
 					uint64 CopySize = FMath::Min<uint64>(UncompressedSize - OffsetInBlock, RemainingSize);
@@ -2326,12 +2461,14 @@ public:
 	{
 		TRACE_CPUPROFILER_EVENT_SCOPE(ReadChunkCompressed);
 
+		// Find where in the virtual file the chunk exists.
 		const FIoOffsetAndLength* OffsetAndLength = Toc.GetOffsetAndLength(ChunkId);
 		if (!OffsetAndLength)
 		{
 			return FIoStatus(EIoErrorCode::NotFound, TEXT("Unknown chunk ID"));
 		}
 
+		// Combine with offset/size requested by the reader.
 		uint64 RequestedOffset = Options.GetOffset();
 		uint64 ResolvedOffset = OffsetAndLength->GetOffset() + RequestedOffset;
 		uint64 ResolvedSize = 0;
@@ -2340,52 +2477,76 @@ public:
 			ResolvedSize = FMath::Min(Options.GetSize(), OffsetAndLength->GetLength() - RequestedOffset);
 		}
 
+		// Find what compressed blocks this read straddles.
 		const FIoStoreTocResource& TocResource = Toc.GetTocResource();
 		const uint64 CompressionBlockSize = TocResource.Header.CompressionBlockSize;
-		TArray<uint8> CompressedBuffer;
-		CompressedBuffer.SetNumUninitialized(static_cast<int32>(CompressionBlockSize));
 		int32 FirstBlockIndex = int32(ResolvedOffset / CompressionBlockSize);
 		int32 LastBlockIndex = int32((Align(ResolvedOffset + ResolvedSize, CompressionBlockSize) - 1) / CompressionBlockSize);
+
+		// Determine size of the result and set up output buffers
 		uint64 TotalCompressedSize = 0;
+		uint64 TotalAlignedSize = 0;
 		for (int32 BlockIndex = FirstBlockIndex; BlockIndex <= LastBlockIndex; ++BlockIndex)
 		{
 			const FIoStoreTocCompressedBlockEntry& CompressionBlock = TocResource.CompressionBlocks[BlockIndex];
 			TotalCompressedSize += CompressionBlock.GetCompressedSize();
+			TotalAlignedSize += Align(CompressionBlock.GetCompressedSize(), FAES::AESBlockSize);
 		}
+
 		FIoStoreCompressedReadResult Result;
-		Result.IoBuffer = FIoBuffer(TotalCompressedSize);
+		Result.IoBuffer = FIoBuffer(TotalAlignedSize);
 		Result.Blocks.Reserve(LastBlockIndex + 1 - FirstBlockIndex);
 		Result.UncompressedOffset = ResolvedOffset % CompressionBlockSize;
 		Result.UncompressedSize = ResolvedSize;
-		uint8* Dst = Result.IoBuffer.Data();
+		Result.TotalCompressedSize = TotalCompressedSize;
+
+		// Set up the result blocks.
+		uint64 CurrentOffset = 0;
 		for (int32 BlockIndex = FirstBlockIndex; BlockIndex <= LastBlockIndex; ++BlockIndex)
 		{
 			const FIoStoreTocCompressedBlockEntry& CompressionBlock = TocResource.CompressionBlocks[BlockIndex];
-			uint32 RawSize = Align(CompressionBlock.GetCompressedSize(), FAES::AESBlockSize);
-			if (uint32(CompressedBuffer.Num()) < RawSize)
-			{
-				CompressedBuffer.SetNumUninitialized(RawSize);
-			}
-			
+			FIoStoreCompressedBlockInfo& BlockInfo = Result.Blocks.AddDefaulted_GetRef();
+			BlockInfo.CompressionMethod = TocResource.CompressionMethods[CompressionBlock.GetCompressionMethodIndex()];
+			BlockInfo.CompressedSize = CompressionBlock.GetCompressedSize();
+			BlockInfo.UncompressedSize = CompressionBlock.GetUncompressedSize();
+			BlockInfo.OffsetInBuffer = CurrentOffset;
+			BlockInfo.AlignedSize = Align(CompressionBlock.GetCompressedSize(), FAES::AESBlockSize);
+			CurrentOffset += BlockInfo.AlignedSize;
+		}
+
+		uint8* OutputBuffer = Result.IoBuffer.Data();
+
+		// We can read the entire thing at once since we obligate the caller to skip the alignment padding.
+		{
+			const FIoStoreTocCompressedBlockEntry& CompressionBlock = TocResource.CompressionBlocks[FirstBlockIndex];
 			int32 PartitionIndex = int32(CompressionBlock.GetOffset() / TocResource.Header.PartitionSize);
 			int64 PartitionOffset = int64(CompressionBlock.GetOffset() % TocResource.Header.PartitionSize);
-			TUniquePtr<IAsyncReadRequest> ReadRequest(ContainerFileHandles[PartitionIndex]->ReadRequest(PartitionOffset, RawSize, AIOP_Normal, nullptr, CompressedBuffer.GetData()));
+
+			std::atomic_bool bReadSucceeded;
+			UE::Tasks::FTask ReadTask = StartAsyncRead(PartitionIndex, PartitionOffset, TotalAlignedSize, OutputBuffer, &bReadSucceeded);
+
 			{
 				TRACE_CPUPROFILER_EVENT_SCOPE(WaitForIo);
-				ReadRequest->WaitCompletion();
+				ReadTask.BusyWait();
 			}
 
-			if (EnumHasAnyFlags(TocResource.Header.ContainerFlags, EIoContainerFlags::Encrypted))
+			if (bReadSucceeded == false)
 			{
-				TRACE_CPUPROFILER_EVENT_SCOPE(Decrypt);
-				FAES::DecryptData(CompressedBuffer.GetData(), RawSize, DecryptionKey);
+				UE_LOG(LogIoStore, Error, TEXT("Read from container %s failed (partition %d, offset %lld, size %d)"), *ContainerPath, PartitionIndex, PartitionOffset, TotalAlignedSize);
+				return FIoStoreCompressedReadResult();
+			}
+		}
+
+		if (EnumHasAnyFlags(TocResource.Header.ContainerFlags, EIoContainerFlags::Encrypted))
+		{
+			TRACE_CPUPROFILER_EVENT_SCOPE(Decrypt);
+
+			for (int32 BlockIndex = FirstBlockIndex; BlockIndex <= LastBlockIndex; ++BlockIndex)
+			{
+				FIoStoreCompressedBlockInfo& OutputBlock = Result.Blocks[BlockIndex - FirstBlockIndex];
+				uint8* Buffer = OutputBuffer + OutputBlock.OffsetInBuffer;
+				FAES::DecryptData(Buffer, OutputBlock.AlignedSize, DecryptionKey);
 			}
-			FIoStoreCompressedBlockInfo& BlockInfo = Result.Blocks.AddDefaulted_GetRef();
-			BlockInfo.CompressionMethod = TocResource.CompressionMethods[CompressionBlock.GetCompressionMethodIndex()];
-			BlockInfo.CompressedSize = CompressionBlock.GetCompressedSize();
-			BlockInfo.UncompressedSize = CompressionBlock.GetUncompressedSize();
-			FMemory::Memcpy(Dst, CompressedBuffer.GetData(), BlockInfo.CompressedSize);
-			Dst += BlockInfo.CompressedSize;
 		}
 		return Result;
 	}
@@ -2417,6 +2578,38 @@ public:
 		return Toc.GetTocResource().CompressionMethods;
 	}
 
+	bool EnumerateCompressedBlocksForChunk(const FIoChunkId& ChunkId, TFunction<bool(const FIoStoreTocCompressedBlockInfo&)>&& Callback)
+	{
+		const FIoOffsetAndLength* OffsetAndLength = Toc.GetOffsetAndLength(ChunkId);
+		if (!OffsetAndLength)
+		{
+			return false;
+		}
+
+		// Find what compressed blocks this chunk straddles.
+		const FIoStoreTocResource& TocResource = Toc.GetTocResource();
+		const uint64 CompressionBlockSize = TocResource.Header.CompressionBlockSize;
+		int32 FirstBlockIndex = int32(OffsetAndLength->GetOffset() / CompressionBlockSize);
+		int32 LastBlockIndex = int32((Align(OffsetAndLength->GetOffset() + OffsetAndLength->GetLength(), CompressionBlockSize) - 1) / CompressionBlockSize);
+
+		for (int32 BlockIndex = FirstBlockIndex; BlockIndex <= LastBlockIndex; ++BlockIndex)
+		{
+			const FIoStoreTocCompressedBlockEntry& Entry = TocResource.CompressionBlocks[BlockIndex];
+			FIoStoreTocCompressedBlockInfo Info
+			{
+				Entry.GetOffset(),
+				Entry.GetCompressedSize(),
+				Entry.GetUncompressedSize(),
+				Entry.GetCompressionMethodIndex()
+			};
+			if (!Callback(Info))
+			{
+				break;
+			}
+		}
+		return true;
+	}
+
 	void EnumerateCompressedBlocks(TFunction<bool(const FIoStoreTocCompressedBlockInfo&)>&& Callback)
 	{
 		const FIoStoreTocResource& TocResource = Toc.GetTocResource();
@@ -2443,7 +2636,8 @@ private:
 
 	FIoStoreToc Toc;
 	FAES::FAESKey DecryptionKey;
-	TArray<TUniquePtr<IAsyncReadFileHandle>> ContainerFileHandles;
+	TArray<TUniquePtr<FContainerFileAccess>> ContainerFileAccessors;
+	FString ContainerPath;
 	FIoDirectoryIndexReader DirectoryIndexReader;
 	TMap<int32, FString> ChunkFileNamesMap;
 };
@@ -2508,7 +2702,7 @@ TIoStatusOr<FIoStoreCompressedReadResult> FIoStoreReader::ReadCompressed(const F
 	return Impl->ReadCompressed(Chunk, Options);
 }
 
-TFuture<TIoStatusOr<FIoBuffer>> FIoStoreReader::ReadAsync(const FIoChunkId& Chunk, const FIoReadOptions& Options) const
+UE::Tasks::TTask<TIoStatusOr<FIoBuffer>> FIoStoreReader::ReadAsync(const FIoChunkId& Chunk, const FIoReadOptions& Options) const
 {
 	return Impl->ReadAsync(Chunk, Options);
 }
@@ -2533,6 +2727,11 @@ void FIoStoreReader::EnumerateCompressedBlocks(TFunction<bool(const FIoStoreTocC
 	Impl->EnumerateCompressedBlocks(MoveTemp(Callback));
 }
 
+void FIoStoreReader::EnumerateCompressedBlocksForChunk(const FIoChunkId& Chunk, TFunction<bool(const FIoStoreTocCompressedBlockInfo&)>&& Callback) const
+{
+	Impl->EnumerateCompressedBlocksForChunk(Chunk, MoveTemp(Callback));
+}
+
 FIoStatus FIoStoreTocResource::Read(const TCHAR* TocFilePath, EIoStoreTocReadOptions ReadOptions, FIoStoreTocResource& OutTocResource)
 {
 	check(TocFilePath != nullptr);
diff --git a/Engine/Source/Runtime/Core/Public/IO/IoDispatcher.h b/Engine/Source/Runtime/Core/Public/IO/IoDispatcher.h
--- a/Engine/Source/Runtime/Core/Public/IO/IoDispatcher.h
+++ b/Engine/Source/Runtime/Core/Public/IO/IoDispatcher.h
@@ -20,6 +20,7 @@
 #include "Async/TaskGraphInterfaces.h"
 #include "GenericPlatform/GenericPlatformFile.h"
 #include "Async/Future.h"
+#include "Tasks/Task.h"
 
 class FIoRequest;
 class FIoDispatcher;
@@ -1209,16 +1210,54 @@ struct FIoStoreTocCompressedBlockInfo
 struct FIoStoreCompressedBlockInfo
 {
 	FName CompressionMethod;
+	// The size of relevant data in the block (i.e. what you pass to decompress)
 	uint32 CompressedSize;
+	// The size of the _block_ after decompression. This is not adjusted for any FIoReadOptions used.
 	uint32 UncompressedSize;
+	// The size of the data this block takes in IoBuffer (i.e. after padding for decryption)
+	uint32 AlignedSize;
+	// Where in IoBuffer this block starts.
+	uint64 OffsetInBuffer;
 };
 
 struct FIoStoreCompressedReadResult
 {
 	FIoBuffer IoBuffer;
 	TArray<FIoStoreCompressedBlockInfo> Blocks;
-	uint64 UncompressedOffset;
-	uint64 UncompressedSize;
+	// There is where the data starts in IoBuffer (for when you pass in a data range via FIoReadOptions)
+	uint64 UncompressedOffset = 0;
+	// This is the total size requested via FIoReadOptions. Notably, if you requested a narrow range, you could
+	// add up all the block uncompressed sizes and it would be larger than this.
+	uint64 UncompressedSize = 0;
+	// This is the total size of compressed data, which is less than IoBuffer size due to padding for decryption.
+	uint64 TotalCompressedSize = 0;
+};
+
+
+class IIoStoreWriterReferenceChunkDatabase
+{
+public:
+	CORE_API virtual ~IIoStoreWriterReferenceChunkDatabase() = default;
+
+	/*
+	* Used by IIoStoreWriter to check and see if there's a reference chunk that matches the data that
+	* IoStoreWriter wants to compress and write. Validity checks must be synchronous - if a chunk can't be
+	* used for some reason (no matching chunk exists or otherwise), this function must return false and not 
+	* call InCompletionCallback. The input parameters are all for match finding and validating.
+	* 
+	* Once a matching chunk is found, it is read from the source iostore container asynchronously, and upon
+	* completion InCompletionCallback is called with the raw output from FIoStoreReader::ReadCompressed (i.e.
+	* FIoStoreCompressedReadResult). Failures once the async read process has started are currently fatal due to needing
+	* to rekick the BeginCompress for the chunk.
+	* 
+	* For the moment, changes in compression method are allowed.
+	* 
+	* RetrieveChunk is not currently thread safe and must be called from a single thread.
+	* 
+	* Chunks provided *MUST* decompress to bits that hash to the exact value provided in InChunkKey (i.e. be exactly the same bits),
+	* and also be the same number of blocks (i.e. same CompressionBlockSize)
+	*/
+	CORE_API virtual bool RetrieveChunk(const TPair<FIoContainerId, FIoChunkHash>& InChunkKey, const FName& InCompressionMethod, uint64 InUncompressedSize, uint64 InNumChunkBlocks, TUniqueFunction<void(TIoStatusOr<FIoStoreCompressedReadResult>)> InCompletionCallback) = 0;
 };
 
 
@@ -1227,6 +1266,11 @@ class IIoStoreWriter
 public:
 	CORE_API virtual ~IIoStoreWriter() = default;
 
+	/**
+	*	If a reference database is provided, the IoStoreWriter implementation may elect to reuse compressed blocks
+	*	from previous containers instead of recompressing input data. This must be set before any writes are appended.
+	*/
+	CORE_API virtual void SetReferenceChunkDatabase(TSharedPtr<IIoStoreWriterReferenceChunkDatabase> ReferenceChunkDatabase) = 0;
 	CORE_API virtual void EnableDiskLayoutOrdering(const TArray<TUniquePtr<FIoStoreReader>>& PatchSourceReaders = TArray<TUniquePtr<FIoStoreReader>>()) = 0;
 	CORE_API virtual void Append(const FIoChunkId& ChunkId, FIoBuffer Chunk, const FIoWriteOptions& WriteOptions, uint64 OrderHint = MAX_uint64) = 0;
 	CORE_API virtual void Append(const FIoChunkId& ChunkId, IIoStoreWriteRequest* Request, const FIoWriteOptions& WriteOptions) = 0;
@@ -1248,9 +1292,18 @@ public:
 	CORE_API void EnumerateChunks(TFunction<bool(const FIoStoreTocChunkInfo&)>&& Callback) const;
 	CORE_API TIoStatusOr<FIoStoreTocChunkInfo> GetChunkInfo(const FIoChunkId& Chunk) const;
 	CORE_API TIoStatusOr<FIoStoreTocChunkInfo> GetChunkInfo(const uint32 TocEntryIndex) const;
+
+	// Reads the chunk off the disk, decrypting/decompressing as necessary.
 	CORE_API TIoStatusOr<FIoBuffer> Read(const FIoChunkId& Chunk, const FIoReadOptions& Options) const;
+	
+	// As Read(), except returns a task that will contain the result after a .Wait/.BusyWait.
+	CORE_API UE::Tasks::TTask<TIoStatusOr<FIoBuffer>> ReadAsync(const FIoChunkId& Chunk, const FIoReadOptions& Options) const;
+
+	// Reads and decrypts if necessary the compressed blocks, but does _not_ decompress them. The totality of the data is stored
+	// in FIoStoreCompressedReadResult::FIoBuffer as a contiguous buffer, however each block is padded during encryption, so
+	// either use FIoStoreCompressedBlockInfo::AlignedSize to advance through the buffer, or use FIoStoreCompressedBlockInfo::OffsetInBuffer
+	// directly.
 	CORE_API TIoStatusOr<FIoStoreCompressedReadResult> ReadCompressed(const FIoChunkId& Chunk, const FIoReadOptions& Options) const;
-	CORE_API TFuture<TIoStatusOr<FIoBuffer>> ReadAsync(const FIoChunkId& Chunk, const FIoReadOptions& Options) const;
 
 	CORE_API const FIoDirectoryIndexReader& GetDirectoryIndexReader() const;
 
@@ -1260,6 +1313,7 @@ public:
 	CORE_API uint64 GetCompressionBlockSize() const;
 	CORE_API const TArray<FName>& GetCompressionMethods() const;
 	CORE_API void EnumerateCompressedBlocks(TFunction<bool(const FIoStoreTocCompressedBlockInfo&)>&& Callback) const;
+	CORE_API void EnumerateCompressedBlocksForChunk(const FIoChunkId& Chunk, TFunction<bool(const FIoStoreTocCompressedBlockInfo&)>&& Callback) const;
 
 private:
 	FIoStoreReaderImpl* Impl;
