fix platform file trace for modular builds.

#rnx

diff --git a/Engine/Source/Runtime/Core/Public/ProfilingDebugging/PlatformFileTrace.h b/Engine/Source/Runtime/Core/Public/ProfilingDebugging/PlatformFileTrace.h
--- a/Engine/Source/Runtime/Core/Public/ProfilingDebugging/PlatformFileTrace.h
+++ b/Engine/Source/Runtime/Core/Public/ProfilingDebugging/PlatformFileTrace.h
@@ -13,7 +13,7 @@
 
 #if PLATFORMFILETRACE_ENABLED
 
-struct FPlatformFileTrace
+struct CORE_API FPlatformFileTrace
 {
 	static void BeginOpen(const TCHAR* Path);
 	static void EndOpen(uint64 FileHandle);
@@ -28,7 +28,7 @@ struct FPlatformFileTrace
 	static void BeginWrite(uint64 WriteHandle, uint64 FileHandle, uint64 Offset, uint64 Size);
 	static void EndWrite(uint64 WriteHandle, uint64 SizeWritten);
 
-	CORE_API static uint32 GetOpenFileHandleCount();
+	static uint32 GetOpenFileHandleCount();
 };
 
 #define TRACE_PLATFORMFILE_BEGIN_OPEN(Path) \
