- Changed Plugins to use an FConfigContext, uising the FConfigContext::ReadPluginFile() factory, this dramatically reduces the number of files to scan, solves the BaseFoo.ini vs DefaultFoo.ini issue for Engine vs Project plugins
- Moved the IniCacheSet into the FConfigContext, and out of FConfigCacheIni. We could possibly get rid of it now entirely with the reduced number of paths we look for (and the fact that plugin scanning will be multithreaded soon)
- Removed the TRUEENGINE workaround now that Plugins are using a context to load the file, this makes the workaround unnecessary since plugins don't fake the EngineConfigDir

diff --git a/Engine/Source/Runtime/Core/Private/Misc/ConfigCacheIni.cpp b/Engine/Source/Runtime/Core/Private/Misc/ConfigCacheIni.cpp
--- a/Engine/Source/Runtime/Core/Private/Misc/ConfigCacheIni.cpp
+++ b/Engine/Source/Runtime/Core/Private/Misc/ConfigCacheIni.cpp
@@ -41,7 +41,6 @@ namespace Private
 }
 }
 }
-const TSet<FString>* FConfigCacheIni::IniCacheSet = nullptr;
 
 namespace
 {
@@ -434,7 +433,7 @@ static void FixupArrayOfStructKeysForSection(FConfigSection* Section, const FStr
 /**
  * Check if an ini file exists, allowing a delegate to determine if it will handle loading it
  */
-static bool DoesConfigFileExistWrapper(const TCHAR* IniFile)
+static bool DoesConfigFileExistWrapper(const TCHAR* IniFile, const TSet<FString>* IniCacheSet=nullptr)
 {
 	// will any delegates return contents via PreLoadConfigFileDelegate()?
 	int32 ResponderCount = 0;
@@ -451,7 +450,6 @@ static bool DoesConfigFileExistWrapper(const TCHAR* IniFile)
 	// which would pass by silently. Realistically, in most cases this would never
 	// have caused an issue, but using FPlatformProperties::RequiresCookedData
 	// to prevent using the cache in that case ensures full consistency.
-	const TSet<FString>* const IniCacheSet = FConfigCacheIni::GetIniCacheSet();
 	if (IniCacheSet && FPlatformProperties::RequiresCookedData())
 	{ 
 		const FString IniFileString(IniFile);
@@ -578,7 +576,7 @@ struct FConfigLayer
 	**************************************************/
 
 	// Engine/Base.ini
-	{ TEXT("AbsoluteBase"),				TEXT("{TRUEENGINE}/Config/Base.ini"), EConfigLayerFlags::NoExpand},
+	{ TEXT("AbsoluteBase"),				TEXT("{ENGINE}/Config/Base.ini"), EConfigLayerFlags::NoExpand},
 	
 	// Engine/Base*.ini
 	{ TEXT("Base"),						TEXT("{ENGINE}/Config/Base{TYPE}.ini") },
@@ -1367,7 +1365,7 @@ void FConfigFile::AddDynamicLayerToHierarchy(const FString& Filename)
  * @param ConfigFile - This is the FConfigFile which will have the contents of the .ini loaded into and Combined()
  *
  **/
-/*static */bool LoadIniFileHierarchy(const FConfigFileHierarchy& HierarchyToLoad, FConfigFile& ConfigFile, bool bUseCache)
+static bool LoadIniFileHierarchy(const FConfigFileHierarchy& HierarchyToLoad, FConfigFile& ConfigFile, bool bUseCache, const TSet<FString>* IniCacheSet)
 {
 	// Traverse ini list back to front, merging along the way.
 	for (const TPair<int32, FString>& HierarchyIt : HierarchyToLoad)
@@ -1377,7 +1375,7 @@ void FConfigFile::AddDynamicLayerToHierarchy(const FString& Filename)
 		const FString& IniFileName = HierarchyIt.Value;
 
 		// Spit out friendly error if there was a problem locating .inis (e.g. bad command line parameter or missing folder, ...).
-		if (IsUsingLocalIniFile(*IniFileName, nullptr) && !DoesConfigFileExistWrapper(*IniFileName))
+		if (IsUsingLocalIniFile(*IniFileName, nullptr) && !DoesConfigFileExistWrapper(*IniFileName, IniCacheSet))
 		{
 			if (bIsRequired)
 			{
@@ -3517,7 +3515,7 @@ bool FConfigCacheIni::ForEachEntry(const FKeyValueSink& Visitor, const TCHAR* Se
  */
 static bool GenerateDestIniFile(FConfigFile& DestConfigFile, const FString& DestIniFilename, const FConfigFileHierarchy& SourceIniHierarchy, bool bAllowGeneratedINIs, const bool bUseHierarchyCache)
 {
-	bool bResult = LoadIniFileHierarchy(SourceIniHierarchy, *DestConfigFile.SourceConfigFile, bUseHierarchyCache);
+	bool bResult = LoadIniFileHierarchy(SourceIniHierarchy, *DestConfigFile.SourceConfigFile, bUseHierarchyCache, nullptr);
 	if ( bResult == false )
 	{
 		return false;
@@ -3544,7 +3542,7 @@ static bool GenerateDestIniFile(FConfigFile& DestConfigFile, const FString& Dest
 	auto RegenerateFileLambda = [](const FConfigFileHierarchy& InSourceIniHierarchy, FConfigFile& InDestConfigFile, const bool bInUseCache)
 	{
 		// Regenerate the file.
-		bool bReturnValue = LoadIniFileHierarchy(InSourceIniHierarchy, InDestConfigFile, bInUseCache);
+		bool bReturnValue = LoadIniFileHierarchy(InSourceIniHierarchy, InDestConfigFile, bInUseCache, nullptr);
 		if (InDestConfigFile.SourceConfigFile)
 		{
 			delete InDestConfigFile.SourceConfigFile;
@@ -4939,7 +4937,7 @@ void FConfigFile::UpdateSections(const TCHAR* DiskFilename, const TCHAR* IniRoot
 		SourceConfigFile = new FConfigFile();
 
 		// now when Write it called below, it will diff against SourceIniHierarchy
-		LoadIniFileHierarchy(SourceIniHierarchy, *SourceConfigFile, true);
+		LoadIniFileHierarchy(SourceIniHierarchy, *SourceConfigFile, true, nullptr);
 	}
 
 	WriteInternal(DiskFilename, true, SectionTexts, SectionOrder);
@@ -5730,7 +5728,7 @@ void FConfigCacheIni::ClearOtherPlatformConfigs()
  */
 bool GenerateDestIniFile(FConfigContext& Context)
 {
-	bool bResult = LoadIniFileHierarchy(Context.ConfigFile->SourceIniHierarchy, *Context.ConfigFile->SourceConfigFile, Context.bUseHierarchyCache);
+	bool bResult = LoadIniFileHierarchy(Context.ConfigFile->SourceIniHierarchy, *Context.ConfigFile->SourceConfigFile, Context.bUseHierarchyCache, Context.IniCacheSet);
 	if (bResult == false)
 	{
 		return false;
@@ -5758,7 +5756,7 @@ bool GenerateDestIniFile(FConfigContext& Context)
 	auto RegenerateFileLambda = [](const FConfigFileHierarchy& InSourceIniHierarchy, FConfigFile& InDestConfigFile, const bool bInUseCache)
 	{
 		// Regenerate the file.
-		bool bReturnValue = LoadIniFileHierarchy(InSourceIniHierarchy, InDestConfigFile, bInUseCache);
+		bool bReturnValue = LoadIniFileHierarchy(InSourceIniHierarchy, InDestConfigFile, bInUseCache, nullptr);
 		if (InDestConfigFile.SourceConfigFile)
 		{
 			delete InDestConfigFile.SourceConfigFile;
@@ -5935,7 +5933,6 @@ bool GenerateDestIniFile(FConfigContext& Context)
 FString PerformFinalExpansions(FConfigContext& Context, const FString& InString, const FString& Platform)
 {
 	FString OutString = InString.Replace(TEXT("{ENGINE}"), *Context.EngineRootDir);
-	OutString = OutString.Replace(TEXT("{TRUEENGINE}"), *FPaths::EngineDir());
 	OutString = OutString.Replace(TEXT("{PROJECT}"), *Context.ProjectRootDir);
 	OutString = OutString.Replace(TEXT("{RESTRICTEDPROJECT_NFL}"), *Context.ProjectNotForLicenseesDir);
 	OutString = OutString.Replace(TEXT("{RESTRICTEDPROJECT_NR}"), *Context.ProjectNoRedistDir);
diff --git a/Engine/Source/Runtime/Core/Public/Misc/ConfigCacheIni.h b/Engine/Source/Runtime/Core/Public/Misc/ConfigCacheIni.h
--- a/Engine/Source/Runtime/Core/Public/Misc/ConfigCacheIni.h
+++ b/Engine/Source/Runtime/Core/Public/Misc/ConfigCacheIni.h
@@ -624,22 +624,6 @@ enum class EConfigCacheType : uint8
 // Set of all cached config files.
 class CORE_API FConfigCacheIni : private TMap<FString,FConfigFile>
 {
-public:
-
-	static void SetIniCacheSet(const TSet<FString>* InIniCacheSet)
-	{
-		IniCacheSet = InIniCacheSet;
-	}
-
-	static const TSet<FString>* GetIniCacheSet()
-	{
-		return IniCacheSet;
-	}
-
-private:
-
-	static const TSet<FString>* IniCacheSet;
-	
 public:
 	
 	// Basic functions.
diff --git a/Engine/Source/Runtime/Core/Public/Misc/ConfigContext.h b/Engine/Source/Runtime/Core/Public/Misc/ConfigContext.h
--- a/Engine/Source/Runtime/Core/Public/Misc/ConfigContext.h
+++ b/Engine/Source/Runtime/Core/Public/Misc/ConfigContext.h
@@ -152,6 +152,9 @@ public:
 	bool bDefaultEngineRequired = false;
 	bool bIsForPlugin = false;
 
+	// if this is non-null, it contains a set of pre-scanned ini files to use to find files, instead of looking on disk
+	const TSet<FString>* IniCacheSet = nullptr;
+
 protected:
 
 	bool bDoNotResetConfigFile = false;
@@ -172,4 +175,3 @@ public:
 
 bool GenerateDestIniFile(FConfigContext& Context);
 void LoadAnIniFile(const FString& FilenameToLoad, FConfigFile& ConfigFile);
-//bool LoadIniFileHierarchy(const class FConfigFileHierarchy& HierarchyToLoad, FConfigFile& ConfigFile, bool bUseCache);
\ No newline at end of file
diff --git a/Engine/Source/Runtime/Projects/Private/PluginManager.cpp b/Engine/Source/Runtime/Projects/Private/PluginManager.cpp
--- a/Engine/Source/Runtime/Projects/Private/PluginManager.cpp
+++ b/Engine/Source/Runtime/Projects/Private/PluginManager.cpp
@@ -9,6 +9,7 @@
 #include "Misc/CommandLine.h"
 #include "Misc/Paths.h"
 #include "Misc/ConfigCacheIni.h"
+#include "Misc/ConfigContext.h"
 #include "Misc/ScopedSlowTask.h"
 #include "Misc/CoreDelegates.h"
 #include "Misc/App.h"
@@ -1121,20 +1122,6 @@ bool FPluginManager::ConfigureEnabledPlugins()
 				UE_LOG(LogPluginManager, Log, TEXT("Mounting %s plugin %s"), *EnumToString(Plugin.Type), *Plugin.GetName());
 				UE_LOG(LogPluginManager, Verbose, TEXT("Plugin path: %s"), *Plugin.FileName);
 
-				// Load <PluginName>.ini config file if it exists
-				FString PluginConfigDir = FPaths::GetPath(Plugin.FileName) / TEXT("Config/");
-				FString EngineConfigDir = FPaths::EngineConfigDir();
-				FString SourceConfigDir = FPaths::SourceConfigDir();
-
-				// Load Engine plugins out of BasePluginName.ini and the engine directory, game plugins out of DefaultPluginName.ini
-				if (Plugin.GetLoadedFrom() == EPluginLoadedFrom::Engine)
-				{
-					EngineConfigDir = PluginConfigDir;
-				}
-				else
-				{
-					SourceConfigDir = PluginConfigDir;
-				}
 
 				// Build the config system key for PluginName.ini
 				FString PluginConfigFilename = GConfig->GetConfigFilename(*Plugin.Name);
@@ -1144,27 +1131,23 @@ bool FPluginManager::ConfigureEnabledPlugins()
 
 					FConfigFile& PluginConfig = GConfig->Add(PluginConfigFilename, FConfigFile());
 
-					// We probably *don't* need to set and unset this every time,
-					// but this feels safer in case INI files are added for a plugin and then we try to load it later.
-					// Dynamically added / generated INI files should be irrelevant for the LoadExternalIniFile call.
+					FConfigContext Context = FConfigContext::ReadIntoPluginFile(PluginConfig, FPaths::GetPath(Plugin.FileName));
+
 					if (PluginSystemDefs::IsCachingIniFilesForProcessing())
 					{
-						FConfigCacheIni::SetIniCacheSet(&AllIniFiles);
-					}	
+						Context.IniCacheSet = &AllIniFiles;
+					}
 
-					// This will write out an ini to PluginConfigFilename
-					if (!FConfigCacheIni::LoadExternalIniFile(PluginConfig, *Plugin.Name, *EngineConfigDir, *SourceConfigDir, true, nullptr, false, true))
+					if (!Context.Load(*Plugin.Name))
 					{
 						// Nothing to add, remove from map
 						GConfig->Remove(PluginConfigFilename);
 					}
-
-					if (PluginSystemDefs::IsCachingIniFilesForProcessing())
-					{
-						FConfigCacheIni::SetIniCacheSet(nullptr);
-					}	
 				}
 
+				// Load <PluginName>.ini config file if it exists
+				FString PluginConfigDir = FPaths::GetPath(Plugin.FileName) / TEXT("Config/");
+
 				// override config cache entries with plugin configs (Engine.ini, Game.ini, etc in <PluginDir>\Config\)
 				TArray<FString> PluginConfigs;
 				IFileManager::Get().FindFiles(PluginConfigs, *PluginConfigDir, TEXT("ini"));
