Add NullOpt for TOptional

diff --git a/Engine/Source/Runtime/Core/Public/Misc/Optional.h b/Engine/Source/Runtime/Core/Public/Misc/Optional.h
--- a/Engine/Source/Runtime/Core/Public/Misc/Optional.h
+++ b/Engine/Source/Runtime/Core/Public/Misc/Optional.h
@@ -8,6 +8,12 @@
 #include "Templates/UnrealTemplate.h"
 #include "Serialization/Archive.h"
 
+struct FNullOpt
+{
+	explicit constexpr FNullOpt(int) {}
+};
+constexpr FNullOpt NullOpt{0};
+
 /**
  * When we have an optional value IsSet() returns true, and GetValue() is meaningful.
  * Otherwise GetValue() is not meaningful.
@@ -18,7 +24,7 @@ struct TOptional
 public:
 	using ElementType = OptionalType;
 	
-	/** Construct an OptionaType with a valid value. */
+	/** Construct an OptionalType with a valid value. */
 	TOptional(const OptionalType& InValue)
 	{
 		new(&Value) OptionalType(InValue);
@@ -35,6 +41,12 @@ public:
 		new(&Value) OptionalType(Forward<ArgTypes>(Args)...);
 		bIsSet = true;
 	}
+	
+	/** Construct an OptionalType with an invalid value. */
+	TOptional(FNullOpt)
+		: bIsSet(false)
+	{
+	}
 
 	/** Construct an OptionalType with no value; i.e. unset */
 	TOptional()
