- Fix for FFileHelper::SaveStringToFile not closing a file handle

#rnx

diff --git a/Engine/Source/Runtime/Core/Private/Misc/FileHelper.cpp b/Engine/Source/Runtime/Core/Private/Misc/FileHelper.cpp
--- a/Engine/Source/Runtime/Core/Private/Misc/FileHelper.cpp
+++ b/Engine/Source/Runtime/Core/Private/Misc/FileHelper.cpp
@@ -616,7 +616,10 @@ bool FFileHelper::SaveStringToFile( FStringView String, const TCHAR* Filename,
 		return false;
 
 	if( String.IsEmpty() )
+	{
+		Ar->Close();
 		return true;
+	}
 
 	bool SaveAsUnicode = EncodingOptions == EEncodingOptions::ForceUnicode || ( EncodingOptions == EEncodingOptions::AutoDetect && !FCString::IsPureAnsi(String.GetData(), String.Len()) );
 	if( EncodingOptions == EEncodingOptions::ForceUTF8 )
