Insert dummy values into memory images for names like vtable pointers.

diff --git a/Engine/Source/Runtime/Core/Private/Serialization/MemoryImage.cpp b/Engine/Source/Runtime/Core/Private/Serialization/MemoryImage.cpp
--- a/Engine/Source/Runtime/Core/Private/Serialization/MemoryImage.cpp
+++ b/Engine/Source/Runtime/Core/Private/Serialization/MemoryImage.cpp
@@ -1406,6 +1406,7 @@ void FMemoryImageResult::SaveToArchive(FArchive& Ar) const
 	Ar << const_cast<FMemoryImageResult*>(this)->TargetLayoutParameters;
 
 	Ar << FrozenSize;
+	int64 FrozenObjectOffset = Ar.Tell();
 	Ar.Serialize(FrozenObject, FrozenSize);
 
 	PointerTable->SaveToArchive(Ar, TargetLayoutParameters, FrozenObject);
@@ -1470,6 +1471,48 @@ void FMemoryImageResult::SaveToArchive(FArchive& Ar) const
 	
 	SerializeNames(ScriptNames, ScriptNameCounts, Ar);
 	SerializeNames(MemoryImageNames, MemoryImageNameCounts, Ar);
+
+	int64 EndOffset = Ar.Tell();
+
+	// Replace patchable values with zeroes for determinism. 
+	// We can't do this when we flatten an object because we may still want to read that object on the current platform. 
+	for (const FMemoryImageVTablePointer& Patch : VTables)
+	{
+		Ar.Seek(FrozenObjectOffset + Patch.Offset);
+		if (TargetLayoutParameters.Is32Bit())
+		{
+			int32 Zero = 0;
+			Ar.Serialize(&Zero, sizeof(Zero));
+		}
+		else
+		{
+			int64 Zero = 0;
+			Ar.Serialize(&Zero, sizeof(Zero));
+		}
+	}
+
+	for (const FMemoryImageNamePointer& Patch : ScriptNames)
+	{
+		Ar.Seek(FrozenObjectOffset + Patch.Offset);
+		FScriptName NullName;
+		Ar.Serialize(&NullName, sizeof(FScriptName));
+	}
+
+	for (const FMemoryImageNamePointer& Patch : MemoryImageNames)
+	{
+		Ar.Seek(FrozenObjectOffset + Patch.Offset);
+		int32 Zeroes[3] = {0, 0, 0};
+		if (TargetLayoutParameters.WithCasePreservingFName())
+		{
+			Ar.Serialize(Zeroes, sizeof(int32) * 3);
+		}
+		else
+		{
+			Ar.Serialize(Zeroes, sizeof(int32) * 2);
+		}
+	}
+
+	Ar.Seek(EndOffset);
 }
 
 static inline void ApplyVTablePatch(void* FrozenObject, const FTypeLayoutDesc& DerivedType, uint32 VTableOffset, uint32 Offset)
@@ -1715,9 +1758,9 @@ uint32 FMemoryImageSection::WriteVTable(const FTypeLayoutDesc& TypeDesc, const F
 	return VTablePointer.Offset;
 }
 
-uint32 FMemoryImageSection::WriteFMemoryImageName(TConstArrayView<uint8> InBytes, const FName& Name)
+uint32 FMemoryImageSection::WriteFMemoryImageName(int32 NumBytes, const FName& Name)
 {
-	const uint32 Offset = WriteBytes(InBytes.GetData(), InBytes.Num());
+	const uint32 Offset = WriteZeroBytes(NumBytes);
 	if (!Name.IsNone())
 	{
 		FMemoryImageNamePointer& NamePointer = MemoryImageNames.AddDefaulted_GetRef();
@@ -1729,7 +1772,7 @@ uint32 FMemoryImageSection::WriteFMemoryImageName(TConstArrayView<uint8> InBytes
 
 uint32 FMemoryImageSection::WriteFScriptName(const FScriptName& Name)
 {
-	const uint32 Offset = WriteBytes(Name);
+	const uint32 Offset = WriteZeroBytes(sizeof(Name));
 	if (!Name.IsNone())
 	{
 		FMemoryImageNamePointer& NamePointer = ScriptNames.AddDefaulted_GetRef();
@@ -1963,9 +2006,9 @@ uint32 FMemoryImageWriter::WriteVTable(const FTypeLayoutDesc& TypeDesc, const FT
 	return Section->WriteVTable(TypeDesc, DerivedTypeDesc);
 }
 
-uint32 FMemoryImageWriter::WriteFMemoryImageName(TConstArrayView<uint8> InBytes, const FName& Name)
+uint32 FMemoryImageWriter::WriteFMemoryImageName(int32 NumBytes, const FName& Name)
 {
-	return Section->WriteFMemoryImageName(InBytes, Name);
+	return Section->WriteFMemoryImageName(NumBytes, Name);
 }
 
 uint32 FMemoryImageWriter::WriteFScriptName(const FScriptName& Name)
diff --git a/Engine/Source/Runtime/Core/Private/UObject/UnrealNames.cpp b/Engine/Source/Runtime/Core/Private/UObject/UnrealNames.cpp
--- a/Engine/Source/Runtime/Core/Private/UObject/UnrealNames.cpp
+++ b/Engine/Source/Runtime/Core/Private/UObject/UnrealNames.cpp
@@ -5023,39 +5023,18 @@ void Freeze::IntrinsicWriteMemoryImage(FMemoryImageWriter& Writer, const FMemory
 	const FPlatformTypeLayoutParameters& TargetLayoutParameters = Writer.GetTargetLayoutParams();
 	if (TargetLayoutParameters.IsCurrentPlatform())
 	{
-		TConstArrayView<uint8> Bytes{ reinterpret_cast<const uint8*>(&Object), sizeof(Object) };
-		Writer.WriteFMemoryImageName(Bytes, FName(Object));
+		Writer.WriteFMemoryImageName(sizeof(FMemoryImageName), FName(Object));
 	}
 	else
 	{
 		ensureMsgf(WITH_CASE_PRESERVING_NAME, TEXT("Cooking memory images containing names when WITH_CASE_PRESERVING_NAME is false can lead to cook determinism issues with string casing."));
-
 		if (!TargetLayoutParameters.WithCasePreservingFName())
 		{
-			TMemoryImageNameLayout<0> ToWrite;
-			ToWrite.ComparisonIndex = Object.ComparisonIndex;
-#if !UE_FNAME_OUTLINE_NUMBER
-			ToWrite.NumberOrDummy = Object.Number;
-#endif
-
-			TConstArrayView<uint8> Bytes{ reinterpret_cast<const uint8*>(&ToWrite), sizeof(ToWrite) };
-			Writer.WriteFMemoryImageName(Bytes, FName(Object));
+			Writer.WriteFMemoryImageName(sizeof(TMemoryImageNameLayout<0>), FName(Object));
 		}
 		else
 		{
-			TMemoryImageNameLayout<1> ToWrite;
-			ToWrite.ComparisonIndex = Object.ComparisonIndex;
-#if !UE_FNAME_OUTLINE_NUMBER
-			ToWrite.NumberOrDummy = Object.Number;
-#endif
-#if WITH_CASE_PRESERVING_NAME
-			ToWrite.DisplayIndex = Object.DisplayIndex;
-#else
-			ToWrite.DisplayIndex = Object.ComparisonIndex;
-#endif
-
-			TConstArrayView<uint8> Bytes{ reinterpret_cast<const uint8*>(&ToWrite), sizeof(ToWrite) };
-			Writer.WriteFMemoryImageName(Bytes, FName(Object));
+			Writer.WriteFMemoryImageName(sizeof(TMemoryImageNameLayout<1>), FName(Object));
 		}
 	}
 }
diff --git a/Engine/Source/Runtime/Core/Public/Serialization/MemoryImage.h b/Engine/Source/Runtime/Core/Public/Serialization/MemoryImage.h
--- a/Engine/Source/Runtime/Core/Public/Serialization/MemoryImage.h
+++ b/Engine/Source/Runtime/Core/Public/Serialization/MemoryImage.h
@@ -231,13 +231,20 @@ public:
 		return Offset;
 	}
 
+	uint32 WriteZeroBytes(int32 Num)
+	{
+		const uint32 Offset = GetOffset();
+		Bytes.SetNumZeroed(Offset + Num);
+		return Offset;
+	}
+
 	template<typename T>
 	uint32 WriteBytes(const T& Data) { return WriteBytes(&Data, sizeof(T)); }
 
 	FMemoryImageSection* WritePointer(const FTypeLayoutDesc& StaticTypeDesc, const FTypeLayoutDesc& DerivedTypeDesc, uint32* OutOffsetToBase = nullptr);
 	uint32 WriteRawPointerSizedBytes(uint64 PointerValue);
 	uint32 WriteVTable(const FTypeLayoutDesc& TypeDesc, const FTypeLayoutDesc& DerivedTypeDesc);
-	uint32 WriteFMemoryImageName(TConstArrayView<uint8> InBytes, const FName& Name);
+	uint32 WriteFMemoryImageName(int32 NumBytes, const FName& Name);
 	uint32 WriteFScriptName(const FScriptName& Name);
 	uint32 Flatten(FMemoryImageResult& OutResult) const;
 
diff --git a/Engine/Source/Runtime/Core/Public/Serialization/MemoryImageWriter.h b/Engine/Source/Runtime/Core/Public/Serialization/MemoryImageWriter.h
--- a/Engine/Source/Runtime/Core/Public/Serialization/MemoryImageWriter.h
+++ b/Engine/Source/Runtime/Core/Public/Serialization/MemoryImageWriter.h
@@ -41,7 +41,7 @@ public:
 	uint32 WriteNullPointer();
 	uint32 WriteRawPointerSizedBytes(uint64 PointerValue);
 	uint32 WriteVTable(const FTypeLayoutDesc& TypeDesc, const FTypeLayoutDesc& DerivedTypeDesc);
-	uint32 WriteFMemoryImageName(TConstArrayView<uint8> InBytes, const FName& Name);
+	uint32 WriteFMemoryImageName(int32 NumBytes, const FName& Name);
 	uint32 WriteFScriptName(const FScriptName& Name);
 
 	template<typename T>
