- Deleting old unused code, and a related useless checkSlow that was firing in debug

diff --git a/Engine/Source/Runtime/Core/Private/Misc/ConfigCacheIni.cpp b/Engine/Source/Runtime/Core/Private/Misc/ConfigCacheIni.cpp
--- a/Engine/Source/Runtime/Core/Private/Misc/ConfigCacheIni.cpp
+++ b/Engine/Source/Runtime/Core/Private/Misc/ConfigCacheIni.cpp
@@ -548,7 +548,6 @@ enum class EConfigLayerFlags : int32
 	None						= 0,
 	AllowCommandLineOverride	= (1 << 1),
 	DedicatedServerOnly			= (1 << 2), // replaces Default, Base, and (NOT {PLATFORM} yet) with an empty string
-	GenerateCacheKey			= (1 << 3),
 	NoExpand					= (1 << 4),
 	RequiresCustomConfig		= (1 << 5), // disabled if no custom config specified
 };
@@ -583,9 +582,9 @@ struct FConfigLayer
 	// Engine/Platform/BasePlatform*.ini
 	{ TEXT("BasePlatform"),				TEXT("{ENGINE}/Config/{PLATFORM}/Base{PLATFORM}{TYPE}.ini")  },
 	// Project/Default*.ini
-	{ TEXT("ProjectDefault"),			TEXT("{PROJECT}/Config/Default{TYPE}.ini"), EConfigLayerFlags::AllowCommandLineOverride | EConfigLayerFlags::GenerateCacheKey },
+	{ TEXT("ProjectDefault"),			TEXT("{PROJECT}/Config/Default{TYPE}.ini"), EConfigLayerFlags::AllowCommandLineOverride },
 	// Project/Generated*.ini Reserved for files generated by build process and should never be checked in 
-	{ TEXT("ProjectGenerated"),			TEXT("{PROJECT}/Config/Generated{TYPE}.ini"), EConfigLayerFlags::GenerateCacheKey },
+	{ TEXT("ProjectGenerated"),			TEXT("{PROJECT}/Config/Generated{TYPE}.ini") },
 	// Project/Custom/CustomConfig/Default*.ini only if CustomConfig is defined
 	{ TEXT("CustomConfig"),				TEXT("{PROJECT}/Config/Custom/{CUSTOMCONFIG}/Default{TYPE}.ini"), EConfigLayerFlags::RequiresCustomConfig },
 	// Engine/Platform/Platform*.ini
@@ -601,7 +600,7 @@ struct FConfigLayer
 	// UserDir/.../User*.ini
 	{ TEXT("UserDir"),					TEXT("{USER}Unreal Engine/Engine/Config/User{TYPE}.ini"), EConfigLayerFlags::NoExpand },
 	// Project/User*.ini
-	{ TEXT("GameDirUser"),				TEXT("{PROJECT}/Config/User{TYPE}.ini"), EConfigLayerFlags::GenerateCacheKey | EConfigLayerFlags::NoExpand },
+	{ TEXT("GameDirUser"),				TEXT("{PROJECT}/Config/User{TYPE}.ini"), EConfigLayerFlags::NoExpand },
 };
 
 
@@ -622,7 +621,7 @@ FConfigLayer GPluginLayers[] =
 	// Project/Default.ini
 	{ TEXT("ProjectDefault"),			TEXT("{PROJECT}/Config/Default{TYPE}.ini") },
 	// Project/Platform/.ini
-	{ TEXT("ProjectDefault"),			TEXT("{PROJECT}/Config/{PLATFORM}/{PLATFORM}{TYPE}.ini") },
+	{ TEXT("ProjectDefault"),			TEXT("{PROJECT}/Config/{PLATFORM}/{PLATFORM}{TYPE}.ini"), EConfigLayerFlags::NoExpand },
 };
 
 
@@ -3883,8 +3882,6 @@ void FConfigFile::AddStaticLayersToHierarchy(const TCHAR* InBaseIniName, const T
 		{
 			// we assume none of the more special tags in expanded ones
 			checkfSlow(FCString::Strstr(Layer.Path, TEXT("{USERSETTINGS}")) == nullptr && FCString::Strstr(Layer.Path, TEXT("{USER}")) == nullptr, TEXT("Expanded config %s shouldn't have a {USER*} tags in it"), *Layer.Path);
-			// last layer is expected to not be expanded
-			checkfSlow(LayerIndex < UE_ARRAY_COUNT(GConfigLayers) - 1, TEXT("Final layer %s shouldn't be an expansion layer, as it needs to generate the  hierarchy cache key"), *Layer.Path);
 
 			// loop over all the possible expansions
 			for (int32 ExpansionIndex = 0; ExpansionIndex < UE_ARRAY_COUNT(GConfigExpansions); ExpansionIndex++)
@@ -3905,10 +3902,6 @@ void FConfigFile::AddStaticLayersToHierarchy(const TCHAR* InBaseIniName, const T
 					ConditionalOverrideIniFilename(ExpandedPath, InBaseIniName);
 				}
 
-				// check if we should be generating the cache key - only at the end of all expansions
-				bool bGenerateCacheKey = EnumHasAnyFlags(Layer.Flag, EConfigLayerFlags::GenerateCacheKey) && ExpansionIndex == UE_ARRAY_COUNT(GConfigExpansions) - 1;
-				checkfSlow(!(bGenerateCacheKey && bHasPlatformTag), TEXT("EConfigLayerFlags::GenerateCacheKey shouldn't have a platform tag"));
-
 				const FDataDrivenPlatformInfo& Info = FDataDrivenPlatformInfoRegistry::GetPlatformInfo(LocalPlatformName);
 
 				// go over parents, and then this platform, unless there's no platform tag, then we simply want to run through the loop one time to add it to the
@@ -6001,8 +5994,6 @@ void FConfigFile::AddStaticLayersToHierarchy(FConfigContext& Context)
 		{
 			// we assume none of the more special tags in expanded ones
 			checkfSlow(FCString::Strstr(Layer.Path, TEXT("{USERSETTINGS}")) == nullptr && FCString::Strstr(Layer.Path, TEXT("{USER}")) == nullptr, TEXT("Expanded config %s shouldn't have a {USER*} tags in it"), *Layer.Path);
-			// last layer is expected to not be expanded
-			checkfSlow(LayerIndex < NumLayers - 1, TEXT("Final layer %s shouldn't be an expansion layer, as it needs to generate the  hierarchy cache key"), *Layer.Path);
 
 			// loop over all the possible expansions
 			for (int32 ExpansionIndex = 0; ExpansionIndex < UE_ARRAY_COUNT(GConfigExpansions); ExpansionIndex++)
@@ -6029,10 +6020,6 @@ void FConfigFile::AddStaticLayersToHierarchy(FConfigContext& Context)
 					ConditionalOverrideIniFilename(ExpandedPath, *Context.BaseIniName);
 				}
 
-				// check if we should be generating the cache key - only at the end of all expansions
-				bool bGenerateCacheKey = EnumHasAnyFlags(Layer.Flag, EConfigLayerFlags::GenerateCacheKey) && ExpansionIndex == UE_ARRAY_COUNT(GConfigExpansions) - 1;
-				checkfSlow(!(bGenerateCacheKey && bHasPlatformTag), TEXT("EConfigLayerFlags::GenerateCacheKey shouldn't have a platform tag"));
-
 				const FDataDrivenPlatformInfo& Info = FDataDrivenPlatformInfoRegistry::GetPlatformInfo(Context.Platform);
 
 				// go over parents, and then this platform, unless there's no platform tag, then we simply want to run through the loop one time to add it to the
