Moved bookkeeping to beginning of large chunks so large page allocator could fallback to CachedPageOSAllocator without crashing at canary check
Improved very large page allocator to keep alive FreeList upon a certain amount of frames. this to prevent physical memory fragmentation

#rnx

diff --git a/Engine/Source/Runtime/Core/Private/HAL/Allocators/CachedOSVeryLargePageAllocator.cpp b/Engine/Source/Runtime/Core/Private/HAL/Allocators/CachedOSVeryLargePageAllocator.cpp
--- a/Engine/Source/Runtime/Core/Private/HAL/Allocators/CachedOSVeryLargePageAllocator.cpp
+++ b/Engine/Source/Runtime/Core/Private/HAL/Allocators/CachedOSVeryLargePageAllocator.cpp
@@ -5,10 +5,25 @@
 #include "Logging/LogMacros.h"
 #include "CoreGlobals.h"
 #include "HAL/LowLevelMemTracker.h"
+#include "ProfilingDebugging/CsvProfiler.h"
 
 #if UE_USE_VERYLARGEPAGEALLOCATOR
+static int32 GVeryLargePageAllocatorMaxEmptyBackstore = 10;
+static FAutoConsoleVariableRef CVarVeryLargePageAllocatorMaxEmptyBackstore(
+	TEXT("VeryLargePageAllocator.MaxEmptyBackstore"),
+	GVeryLargePageAllocatorMaxEmptyBackstore,
+	TEXT(""));
 
 CORE_API bool GEnableVeryLargePageAllocator = true;
+
+#if CSV_PROFILER
+CSV_DECLARE_CATEGORY_MODULE_EXTERN(CORE_API, FMemory);
+
+static volatile int32 GLargePageAllocatorCommitCount = 0;
+static volatile int32 GLargePageAllocatorDecommitCount = 0;
+static volatile int32 GLargePageAllocatorLargePageCount = 0;
+#endif
+
 void FCachedOSVeryLargePageAllocator::Init()
 {
 	Block = FPlatformMemory::FPlatformVirtualMemoryBlock::AllocateVirtual(AddressSpaceToReserve);
@@ -26,6 +41,8 @@ void FCachedOSVeryLargePageAllocator::Init()
 		FreeLargePagesHead[i] = nullptr;
 		UsedLargePagesWithSpaceHead[i] = nullptr;
 		UsedLargePagesHead[i] = nullptr;
+		EmptyButAvailableLargePagesHead[i] = nullptr;
+		EmptyBackStoreCount[i] = 0;
 	}
 #if UE_VERYLARGEPAGEALLOCATOR_TAKEONALL64KBALLOCATIONS
 	for (int i = 0; i < NumberOfLargePages / 2; i++)
@@ -63,27 +80,75 @@ void* FCachedOSVeryLargePageAllocator::Allocate(SIZE_T Size, uint32 AllocationHi
 		if (AllocationHint == FMemory::AllocationHints::SmallPool)
 #endif
 		{
+			bool bLinkToUsedLargePagesWithSpaceHead = false;
 			FLargePage* LargePage = UsedLargePagesWithSpaceHead[AllocationHint];
+
 			if (LargePage == nullptr)
 			{
-				LargePage = FreeLargePagesHead[AllocationHint];
-				if (LargePage)
+				bLinkToUsedLargePagesWithSpaceHead = true;
+				LargePage = EmptyButAvailableLargePagesHead[AllocationHint];
+
+				if (LargePage != nullptr)
 				{
 					LargePage->AllocationHint = AllocationHint;
 					LargePage->Unlink();
+					EmptyBackStoreCount[AllocationHint] -= 1;
+				}
+				else
+				{
+					LargePage = FreeLargePagesHead[AllocationHint];
+
+					if (LargePage != nullptr)
 					{
+						LargePage->AllocationHint = AllocationHint;
+						LargePage->Unlink();
+						{
 #if UE_ALLOW_OSMEMORYLOCKFREE
-						FScopeUnlock FScopeUnlock(Mutex);
+							FScopeUnlock FScopeUnlock(Mutex);
 #endif
-						LLM_PLATFORM_SCOPE(ELLMTag::FMalloc);
-						Block.Commit(LargePage->BaseAddress - AddressSpaceReserved, SizeOfLargePage);
+							LLM_PLATFORM_SCOPE(ELLMTag::FMalloc);
+#if UE_USE_VERYLARGEPAGEALLOCATOR_FALLBACKPATH
+							if (!Block.Commit(LargePage->BaseAddress - AddressSpaceReserved, SizeOfLargePage, false))
+							{
+#if UE_ALLOW_OSMEMORYLOCKFREE
+								if (Mutex != nullptr)
+								{
+									FScopeLock FScopeLock(Mutex);
+									LargePage->LinkHead(FreeLargePagesHead[AllocationHint]);
+								}
+								else
+#endif
+								{
+									LargePage->LinkHead(FreeLargePagesHead[AllocationHint]);
+								}
+
+								// Fallback to regular allocator
+								LargePage = nullptr;
+							}
+#if CSV_PROFILER
+							else
+							{
+								FPlatformAtomics::InterlockedIncrement(&GLargePageAllocatorCommitCount);
+								FPlatformAtomics::InterlockedIncrement(&GLargePageAllocatorLargePageCount);
+							}
+#endif
+						}
 					}
-					LargePage->LinkHead(UsedLargePagesWithSpaceHead[AllocationHint]);
-					CachedFree += SizeOfLargePage;
+#else 
+							Block.Commit(LargePage->BaseAddress - AddressSpaceReserved, SizeOfLargePage);
+						}
+					}
+#endif // UE_USE_VERYLARGEPAGEALLOCATOR_FALLBACKPATH
 				}
 			}
 			if (LargePage != nullptr)
 			{
+				if (bLinkToUsedLargePagesWithSpaceHead)
+				{
+					LargePage->LinkHead(UsedLargePagesWithSpaceHead[AllocationHint]);
+					CachedFree += SizeOfLargePage;
+				}
+
 				ret = LargePage->Allocate();
 				if (ret)
 				{
@@ -127,16 +192,33 @@ void FCachedOSVeryLargePageAllocator::Free(void* Ptr, SIZE_T Size, FCriticalSect
 
 		if (LargePage->NumberOfFreeSubPages == NumberOfSubPagesPerLargePage)
 		{
-			// totally free, need to move which list we are in and remove the backing store
+			// totally free
 			LargePage->Unlink();
+
+			// move it to EmptyButAvailableLargePagesHead if that pool of backstore is not full yet
+			if (EmptyBackStoreCount[LargePage->AllocationHint] < GVeryLargePageAllocatorMaxEmptyBackstore)
+			{
+				LargePage->LinkHead(EmptyButAvailableLargePagesHead[LargePage->AllocationHint]);
+				EmptyBackStoreCount[LargePage->AllocationHint] += 1;
+			}
+			else
 			{
+				// need to move which list we are in and remove the backing store
+				{
 #if UE_ALLOW_OSMEMORYLOCKFREE
-				FScopeUnlock FScopeUnlock(Mutex);
+					FScopeUnlock FScopeUnlock(Mutex);
 #endif
-				Block.Decommit(LargePage->BaseAddress - AddressSpaceReserved, SizeOfLargePage);
+					Block.Decommit(LargePage->BaseAddress - AddressSpaceReserved, SizeOfLargePage);
+				}
+
+#if CSV_PROFILER
+				FPlatformAtomics::InterlockedIncrement(&GLargePageAllocatorDecommitCount);
+				FPlatformAtomics::InterlockedDecrement(&GLargePageAllocatorLargePageCount);
+#endif
+
+				LargePage->LinkHead(FreeLargePagesHead[LargePage->AllocationHint]);
+				CachedFree -= SizeOfLargePage;
 			}
-			LargePage->LinkHead(FreeLargePagesHead[LargePage->AllocationHint]);
-			CachedFree -= SizeOfLargePage;
 		}
 		else if (LargePage->NumberOfFreeSubPages == 1)
 		{
@@ -203,4 +285,23 @@ void FCachedOSVeryLargePageAllocator::FreeAll(FCriticalSection* Mutex)
 {
 	CachedOSPageAllocator.FreeAll(Mutex);
 }
-#endif
\ No newline at end of file
+
+void FCachedOSVeryLargePageAllocator::UpdateStats()
+{
+#if CSV_PROFILER
+	int32 BackingStoreCount = 0;
+	for (int i = 0; i < FMemory::AllocationHints::Max; i++)
+	{
+		BackingStoreCount += EmptyBackStoreCount[i];
+	}
+
+	CSV_CUSTOM_STAT(FMemory, LargeAllocatorCommitCount, GLargePageAllocatorCommitCount, ECsvCustomStatOp::Set);
+	CSV_CUSTOM_STAT(FMemory, LargeAllocatorDecommitCount, GLargePageAllocatorDecommitCount, ECsvCustomStatOp::Set);
+	CSV_CUSTOM_STAT(FMemory, LargeAllocatorBackingStoreCount, BackingStoreCount, ECsvCustomStatOp::Set);
+	CSV_CUSTOM_STAT(FMemory, LargeAllocatorPageCount, GLargePageAllocatorLargePageCount, ECsvCustomStatOp::Set);
+
+	GLargePageAllocatorCommitCount = 0;
+	GLargePageAllocatorDecommitCount = 0;
+#endif
+}
+#endif
diff --git a/Engine/Source/Runtime/Core/Private/HAL/MallocBinned2.cpp b/Engine/Source/Runtime/Core/Private/HAL/MallocBinned2.cpp
--- a/Engine/Source/Runtime/Core/Private/HAL/MallocBinned2.cpp
+++ b/Engine/Source/Runtime/Core/Private/HAL/MallocBinned2.cpp
@@ -1506,6 +1506,7 @@ void FMallocBinned2::UpdateStats()
 	CSV_CUSTOM_STAT(FMemory, AllocatorCachedSlackMB, (int32)(CachedOSPageAllocator.GetCachedFreeTotal()/(1024*1024)), ECsvCustomStatOp::Set);
 #endif
 
+	CachedOSPageAllocator.UpdateStats();
 	FScopedVirtualMallocTimer::UpdateStats();
 }
 
diff --git a/Engine/Source/Runtime/Core/Public/HAL/Allocators/CachedOSPageAllocator.h b/Engine/Source/Runtime/Core/Public/HAL/Allocators/CachedOSPageAllocator.h
--- a/Engine/Source/Runtime/Core/Public/HAL/Allocators/CachedOSPageAllocator.h
+++ b/Engine/Source/Runtime/Core/Public/HAL/Allocators/CachedOSPageAllocator.h
@@ -47,7 +47,10 @@ struct TCachedOSPageAllocator : private FCachedOSPageAllocator
 	{
 		return FreeAllImpl(FreedPageBlocks, FreedPageBlocksNum, CachedTotal, Mutex);
 	}
+	void UpdateStats()
+	{
 
+	}
 	uint64 GetCachedFreeTotal()
 	{
 		return CachedTotal;
diff --git a/Engine/Source/Runtime/Core/Public/HAL/Allocators/CachedOSVeryLargePageAllocator.h b/Engine/Source/Runtime/Core/Public/HAL/Allocators/CachedOSVeryLargePageAllocator.h
--- a/Engine/Source/Runtime/Core/Public/HAL/Allocators/CachedOSVeryLargePageAllocator.h
+++ b/Engine/Source/Runtime/Core/Public/HAL/Allocators/CachedOSVeryLargePageAllocator.h
@@ -9,6 +9,10 @@
 #include "HAL/PlatformMemory.h"
 
 
+#ifndef UE_USE_VERYLARGEPAGEALLOCATOR_FALLBACKPATH
+#define UE_USE_VERYLARGEPAGEALLOCATOR_FALLBACKPATH UE_USE_VERYLARGEPAGEALLOCATOR
+#endif // UE_USE_VERYLARGEPAGEALLOCATOR_FALLBACKPATH
+
 #if UE_USE_VERYLARGEPAGEALLOCATOR
 
 #if PLATFORM_64BITS
@@ -66,6 +70,8 @@ public:
 
 	void FreeAll(FCriticalSection* Mutex = nullptr);
 
+	void UpdateStats();
+
 	uint64 GetCachedFreeTotal()
 	{
 		return CachedFree + CachedOSPageAllocator.GetCachedFreeTotal();
@@ -89,6 +95,7 @@ private:
 	uintptr_t	AddressSpaceReservedEndSmallPool;
 	uintptr_t	AddressSpaceReservedEnd;
 	uint64		CachedFree;
+	int32		EmptyBackStoreCount[FMemory::AllocationHints::Max];
 
 	FPlatformMemory::FPlatformVirtualMemoryBlock Block;
 
@@ -134,6 +141,8 @@ private:
 
 	FLargePage* UsedLargePagesWithSpaceHead[FMemory::AllocationHints::Max];	// has backing store and still has room
 
+	FLargePage*	EmptyButAvailableLargePagesHead[FMemory::AllocationHints::Max];	// has backing store and is empty
+
 	FLargePage	LargePagesArray[NumberOfLargePages];
 
 	TCachedOSPageAllocator<CACHEDOSVERYLARGEPAGEALLOCATOR_MAX_CACHED_OS_FREES, CACHEDOSVERYLARGEPAGEALLOCATOR_BYTE_LIMIT> CachedOSPageAllocator;
diff --git a/Engine/Source/Runtime/Core/Public/HAL/MallocBinned2.h b/Engine/Source/Runtime/Core/Public/HAL/MallocBinned2.h
--- a/Engine/Source/Runtime/Core/Public/HAL/MallocBinned2.h
+++ b/Engine/Source/Runtime/Core/Public/HAL/MallocBinned2.h
@@ -48,8 +48,11 @@
 	#define DEFAULT_GMallocBinned2BundleSize BINNED2_LARGE_ALLOC
 #endif
 
+// When book keeping is at the end of FFreeBlock, MallocBinned2 cannot tell if the allocation comes from a large allocation (higher than 64KB, also named as "OSAllocation") 
+// or from VeryLargePageAllocator that fell back to FCachedOSPageAllocator. In both cases the allocation (large or small) might be aligned to 64KB.
+// bookKeeping at the end needs to be disabled if we want VeryLargePageAllocator to properly fallback to regular TCachedOSPageAllocator if needed
 #ifndef BINNED2_BOOKKEEPING_AT_THE_END_OF_LARGEBLOCK
-#define BINNED2_BOOKKEEPING_AT_THE_END_OF_LARGEBLOCK 1
+#define BINNED2_BOOKKEEPING_AT_THE_END_OF_LARGEBLOCK 0
 #endif
 
 // If we are emulating forking on a windows server or are a linux server, enable support for avoiding dirtying pages owned by the parent. 
@@ -296,7 +299,7 @@ class CORE_API FMallocBinned2 : public FMalloc
 	FORCEINLINE bool IsOSAllocation(const void* Ptr)
 	{
 #if UE_USE_VERYLARGEPAGEALLOCATOR && !PLATFORM_UNIX && !PLATFORM_ANDROID
-		return !CachedOSPageAllocator.IsPartOf(Ptr);
+		return !CachedOSPageAllocator.IsPartOf(Ptr) && IsAligned(Ptr, BINNED2_LARGE_ALLOC);
 #else
 		return IsAligned(Ptr, BINNED2_LARGE_ALLOC);
 #endif
