Added UpdateStats stub to FPooledVirtualMemoryAllocator

diff --git a/Engine/Source/Runtime/Core/Public/HAL/Allocators/PooledVirtualMemoryAllocator.h b/Engine/Source/Runtime/Core/Public/HAL/Allocators/PooledVirtualMemoryAllocator.h
--- a/Engine/Source/Runtime/Core/Public/HAL/Allocators/PooledVirtualMemoryAllocator.h
+++ b/Engine/Source/Runtime/Core/Public/HAL/Allocators/PooledVirtualMemoryAllocator.h
@@ -78,6 +78,9 @@ struct FPooledVirtualMemoryAllocator
 	/** Returns free memory in the pools */
 	uint64 GetCachedFreeTotal();
 
+	/** Update memory stats (does nothing in that implementation) */
+	void UpdateStats() {}
+
 private:
 
 	enum Limits
