- Fixed another checkSlow

diff --git a/Engine/Source/Runtime/Core/Private/Misc/ConfigCacheIni.cpp b/Engine/Source/Runtime/Core/Private/Misc/ConfigCacheIni.cpp
--- a/Engine/Source/Runtime/Core/Private/Misc/ConfigCacheIni.cpp
+++ b/Engine/Source/Runtime/Core/Private/Misc/ConfigCacheIni.cpp
@@ -621,7 +621,7 @@ FConfigLayer GPluginLayers[] =
 	// Project/Default.ini
 	{ TEXT("ProjectDefault"),			TEXT("{PROJECT}/Config/Default{TYPE}.ini") },
 	// Project/Platform/.ini
-	{ TEXT("ProjectDefault"),			TEXT("{PROJECT}/Config/{PLATFORM}/{PLATFORM}{TYPE}.ini"), EConfigLayerFlags::NoExpand },
+	{ TEXT("ProjectDefault"),			TEXT("{PROJECT}/Config/{PLATFORM}/{PLATFORM}{TYPE}.ini")  },
 };
 
 
