- Fixed the Save dir for some editor configs, also fixes plugins loading into Editor.ini

diff --git a/Engine/Source/Runtime/Core/Private/Misc/ConfigContext.cpp b/Engine/Source/Runtime/Core/Private/Misc/ConfigContext.cpp
--- a/Engine/Source/Runtime/Core/Private/Misc/ConfigContext.cpp
+++ b/Engine/Source/Runtime/Core/Private/Misc/ConfigContext.cpp
@@ -27,6 +27,11 @@ FConfigContext::FConfigContext(FConfigCacheIni* InConfigSystem, bool InIsHierarc
 		// but save Generated ini files to, say, WindowsEditor
 		SavePlatform = FPlatformProperties::PlatformName();
 	}
+	else if (Platform == FPlatformProperties::IniPlatformName())
+	{
+		// but save Generated ini files to, say, WindowsEditor
+		SavePlatform = FPlatformProperties::PlatformName();
+	}
 	else
 	{
 		SavePlatform = Platform;
diff --git a/Engine/Source/Runtime/Core/Public/Misc/ConfigContext.h b/Engine/Source/Runtime/Core/Public/Misc/ConfigContext.h
--- a/Engine/Source/Runtime/Core/Public/Misc/ConfigContext.h
+++ b/Engine/Source/Runtime/Core/Public/Misc/ConfigContext.h
@@ -50,7 +50,7 @@ public:
 	 */
 	static FConfigContext ReadIntoGConfig()
 	{
-		return FConfigContext(GConfig, true, FPlatformProperties::IniPlatformName());
+		return FConfigContext(GConfig, true, FString());
 	}
 
 	/**
@@ -58,7 +58,7 @@ public:
 	 */
 	static FConfigContext ReadSingleIntoGConfig()
 	{
-		return FConfigContext(GConfig, false, FPlatformProperties::IniPlatformName());
+		return FConfigContext(GConfig, false, FString());
 	}
 
 	/**
@@ -82,7 +82,7 @@ public:
 	 */
 	static FConfigContext ReadIntoPluginFile(FConfigFile& DestConfigFile, const FString& PluginRootDir)
 	{
-		FConfigContext Context(nullptr, true, FPlatformProperties::IniPlatformName(), &DestConfigFile);
+		FConfigContext Context(nullptr, true, FString(), &DestConfigFile);
 		Context.bIsForPlugin = true;
 		Context.PluginRootDir = PluginRootDir;
 
