Fixed VectorLoadURGBA16N and VectorLoadSRGBA16N on NEON and FPU

diff --git a/Engine/Source/Runtime/Core/Public/Math/UnrealMathFPU.h b/Engine/Source/Runtime/Core/Public/Math/UnrealMathFPU.h
--- a/Engine/Source/Runtime/Core/Public/Math/UnrealMathFPU.h
+++ b/Engine/Source/Runtime/Core/Public/Math/UnrealMathFPU.h
@@ -2348,10 +2348,10 @@ FORCEINLINE VectorRegister4Float VectorLoadURGBA16N(void* Ptr)
 	float V[4];
 	uint16* E = (uint16*)Ptr;
 
-	V[0] = float(E[0]) / 65535.0f;
-	V[1] = float(E[1]) / 65535.0f;
-	V[2] = float(E[2]) / 65535.0f;
-	V[3] = float(E[3]) / 65535.0f;
+	V[0] = float(E[0]);
+	V[1] = float(E[1]);
+	V[2] = float(E[2]);
+	V[3] = float(E[3]);
 
 	return MakeVectorRegisterFloat(V[0], V[1], V[2], V[3]);
 }
@@ -2368,10 +2368,10 @@ FORCEINLINE VectorRegister4Float VectorLoadSRGBA16N(void* Ptr)
 	float V[4];
 	int16* E = (int16*)Ptr;
 
-	V[0] = float(E[0]) / 32767.0f;
-	V[1] = float(E[1]) / 32767.0f;
-	V[2] = float(E[2]) / 32767.0f;
-	V[3] = float(E[3]) / 32767.0f;
+	V[0] = float(E[0]);
+	V[1] = float(E[1]);
+	V[2] = float(E[2]);
+	V[3] = float(E[3]);
 
 	return MakeVectorRegisterFloat(V[0], V[1], V[2], V[3]);
 }
diff --git a/Engine/Source/Runtime/Core/Public/Math/UnrealMathNeon.h b/Engine/Source/Runtime/Core/Public/Math/UnrealMathNeon.h
--- a/Engine/Source/Runtime/Core/Public/Math/UnrealMathNeon.h
+++ b/Engine/Source/Runtime/Core/Public/Math/UnrealMathNeon.h
@@ -2875,9 +2875,7 @@ FORCEINLINE VectorRegister4Float VectorLoadURGBA16N(const uint16* E)
 	V[2] = float(E[2]);
 	V[3] = float(E[3]);
 
-	VectorRegister4Float Vec = VectorLoad(V);
-	VectorRegister4Float Div = vdupq_n_f32(1.0f / 65535.0f);
-	return VectorMultiply(Vec, Div);
+	return VectorLoad(V);
 }
 
 /**
@@ -2897,9 +2895,7 @@ FORCEINLINE VectorRegister4Float VectorLoadSRGBA16N(const void* Ptr)
 	V[2] = float(E[2]);
 	V[3] = float(E[3]);
 
-	VectorRegister4Float Vec = VectorLoad(V);
-	VectorRegister4Float Div = vdupq_n_f32(1.0f / 32767.0f);
-	return VectorMultiply(Vec, Div);
+	return VectorLoad(V);
 }
 
 /**
