Remove Modules section from crash context as it's not compatible with old WER-based schema and it shouldn't be necessary for symbolization anyway.

diff --git a/Engine/Source/Runtime/Core/Private/GenericPlatform/GenericPlatformCrashContext.cpp b/Engine/Source/Runtime/Core/Private/GenericPlatform/GenericPlatformCrashContext.cpp
--- a/Engine/Source/Runtime/Core/Private/GenericPlatform/GenericPlatformCrashContext.cpp
+++ b/Engine/Source/Runtime/Core/Private/GenericPlatform/GenericPlatformCrashContext.cpp
@@ -756,7 +756,6 @@ void FGenericCrashContext::SerializeContentToBuffer() const
 	// Add new portable callstack element with crash stack
 	AddPortableCallStack();
 	AddPortableCallStackHash();
-	AddModules();
 
 	{
 		FString AllThreadStacks;
@@ -950,23 +949,6 @@ void FGenericCrashContext::AddPortableCallStack() const
 
 
 
-void FGenericCrashContext::AddModules() const
-{
-	BeginSection(CommonBuffer, TEXT("Modules"));
-	for (const FStackWalkModuleInfo& Module : ModulesInfo)
-	{
-		CommonBuffer.Appendf(TEXT("<Module Signature=\"%08x%08x\" DebugSignature=\"%08x-%04x-%04x-%04x%04x\">"),
-			Module.TimeDateStamp, Module.ImageSize, 
-			Module.PdbSig70.Data1, Module.PdbSig70.Data2, Module.PdbSig70.Data3, *(uint32*)&Module.PdbSig70.Data4[0], *(uint32*)&Module.PdbSig70.Data4[4]
-			);
-		
-		AppendEscapedXMLString(CommonBuffer, Module.ModuleName);
-
-		CommonBuffer += TEXT("</Module>") LINE_TERMINATOR;
-	}
-	EndSection(CommonBuffer, TEXT("Modules"));
-}
-
 void FGenericCrashContext::AddHeader(FString& Buffer)
 {
 	Buffer += TEXT("<?xml version=\"1.0\" encoding=\"UTF-8\"?>") LINE_TERMINATOR;
