fixed static analysis compilation

diff --git a/Engine/Source/Runtime/Core/Public/Async/TaskGraphInterfaces.h b/Engine/Source/Runtime/Core/Public/Async/TaskGraphInterfaces.h
--- a/Engine/Source/Runtime/Core/Public/Async/TaskGraphInterfaces.h
+++ b/Engine/Source/Runtime/Core/Public/Async/TaskGraphInterfaces.h
@@ -907,7 +907,7 @@ private:
 	FThreadSafeCounter			NumberOfPrerequistitesOutstanding; 
 
 
-#if DO_GUARD_SLOW
+#if DO_GUARD_SLOW || USING_CODE_ANALYSIS
 	// Life stage verification
 	// Tasks go through 8 steps, in order. In non-final builds, we track them with a thread safe counter and verify that the progression is correct.
 	enum ELifeStage
