Fixed FramePro re-initializing FThreadIdleStats singleton on shutdown.
- don't trace FramePro waits when FramePro is not recording (FramePro uses wait events only when recording).

diff --git a/Engine/Source/Runtime/Core/Private/FramePro/FrameProUE4.cpp b/Engine/Source/Runtime/Core/Private/FramePro/FrameProUE4.cpp
--- a/Engine/Source/Runtime/Core/Private/FramePro/FrameProUE4.cpp
+++ b/Engine/Source/Runtime/Core/Private/FramePro/FrameProUE4.cpp
@@ -563,7 +563,7 @@
 		int Platform::WaitEvent(void* p_os_event_mem, int timeout)
 		{
 			FEvent* p_event = GetOSEvent(p_os_event_mem);
-			return (timeout==-1) ? p_event->Wait() : p_event->Wait(timeout);
+			return (timeout==-1) ? p_event->Wait(MAX_uint32, !FramePro::IsConnected()) : p_event->Wait(timeout, !FramePro::IsConnected());
 		}
 
 		//------------------------------------------------------------------------
