TVector::BoundToBox: Change inconsistent pass-by-value to pass-by-reference.
AIGraph::UpdateClassData: Fix incorrect recursive child call: call needs to be made on the child not on the current node.
From PR #8466 submitted by LindyHopperGT.
#rnx

diff --git a/Engine/Source/Editor/AIGraph/Private/AIGraph.cpp b/Engine/Source/Editor/AIGraph/Private/AIGraph.cpp
--- a/Engine/Source/Editor/AIGraph/Private/AIGraph.cpp
+++ b/Engine/Source/Editor/AIGraph/Private/AIGraph.cpp
@@ -146,9 +146,9 @@ void UAIGraph::UpdateClassData()
 
 			for (int32 SubIdx = 0; SubIdx < Node->SubNodes.Num(); SubIdx++)
 			{
-				if (Node->SubNodes[SubIdx])
+				if (UAIGraphNode* SubNode = Node->SubNodes[SubIdx])
 				{
-					Node->UpdateNodeClassData();
+					SubNode->UpdateNodeClassData();
 				}
 			}
 		}
diff --git a/Engine/Source/Runtime/Core/Public/Math/Vector.h b/Engine/Source/Runtime/Core/Public/Math/Vector.h
--- a/Engine/Source/Runtime/Core/Public/Math/Vector.h
+++ b/Engine/Source/Runtime/Core/Public/Math/Vector.h
@@ -694,7 +694,7 @@ public:
     TVector<T> BoundToCube(T Radius) const;
 
     /** Get a copy of this vector, clamped inside of a cube. */
-    TVector<T> BoundToBox(const TVector<T>& Min, const TVector<T> Max) const;
+    TVector<T> BoundToBox(const TVector<T>& Min, const TVector<T>& Max) const;
 
     /** Create a copy of this vector, with its magnitude clamped between Min and Max. */
     TVector<T> GetClampedToSize(T Min, T Max) const;
@@ -1763,7 +1763,7 @@ FORCEINLINE TVector<T> TVector<T>::BoundToCube(T Radius) const
 }
 
 template<typename T>
-FORCEINLINE TVector<T> TVector<T>::BoundToBox(const TVector<T>& Min, const TVector<T> Max) const
+FORCEINLINE TVector<T> TVector<T>::BoundToBox(const TVector<T>& Min, const TVector<T>& Max) const
 {
     return TVector<T>
     (
