Avoid calling `memcpy`with `nullptr` argument when `TArray::operator==` is invoked on two empty arrays.

diff --git a/Engine/Source/Runtime/Core/Private/Containers/ContainersTest.cpp b/Engine/Source/Runtime/Core/Private/Containers/ContainersTest.cpp
--- a/Engine/Source/Runtime/Core/Private/Containers/ContainersTest.cpp
+++ b/Engine/Source/Runtime/Core/Private/Containers/ContainersTest.cpp
@@ -2,6 +2,7 @@
 
 #include "CoreTypes.h"
 #include "Misc/AssertionMacros.h"
+#include "Concepts/EqualityComparable.h"
 #include "Containers/Array.h"
 #include "Containers/Map.h"
 #include "Containers/Set.h"
@@ -236,6 +237,15 @@ namespace
 		}
 	}
 
+	template <typename Container>
+	void CheckContainerSelfEquality(Container& Cont)
+	{
+		if constexpr (TModels<CEqualityComparable, Container>::Value)
+		{
+			check(Cont == Cont);
+		}
+	}
+
 	template <typename Container>
 	void CheckContainerNum(Container& Cont)
 	{
@@ -312,6 +322,7 @@ namespace
 				CheckContainerNum(Cont);
 				CheckContainerEnds(Cont);
 				CheckContainerElements(Cont);
+				CheckContainerSelfEquality(Cont);
 			}
 
 			for (int32 N = 0; N != Count; ++N)
@@ -320,6 +331,7 @@ namespace
 				CheckContainerNum(Cont);
 				CheckContainerEnds(Cont);
 				CheckContainerElements(Cont);
+				CheckContainerSelfEquality(Cont);
 			}
 
 			for (int32 N = 0; N != Count; ++N)
@@ -328,6 +340,7 @@ namespace
 				CheckContainerNum(Cont);
 				CheckContainerEnds(Cont);
 				CheckContainerElements(Cont);
+				CheckContainerSelfEquality(Cont);
 			}
 
 			for (int32 N = 0; N != Count; ++N)
@@ -336,10 +349,18 @@ namespace
 				CheckContainerNum(Cont);
 				CheckContainerEnds(Cont);
 				CheckContainerElements(Cont);
+				CheckContainerSelfEquality(Cont);
 			}
 		}
 	}
 
+	template <typename Container>
+	void RunEmptyContainerSelfEqualityTest()
+	{
+		Container Cont;
+		CheckContainerSelfEquality(Cont);
+	}
+
 	template <typename ContainerType, typename KeyType>
 	void RunPerformanceTest(const FString& Description, int32 NumObjects, int32 NumOperations)
 	{
@@ -508,6 +529,8 @@ bool FContainersFullTest::RunTest(const FString& Parameters)
 	RunContainerTests<TSortedMap<FString, FContainerTestValueType>, FString>();
 	RunContainerTests<TSortedMap<FString, FContainerTestValueType, TInlineAllocator<64>>, FString>();
 
+	RunEmptyContainerSelfEqualityTest<TArray<int32>>();
+
 	// Verify use of FName index sorter with SortedMap
 
 	TSortedMap<FName, int32, FDefaultAllocator, FNameFastLess> NameMap;
diff --git a/Engine/Source/Runtime/Core/Public/Templates/MemoryOps.h b/Engine/Source/Runtime/Core/Public/Templates/MemoryOps.h
--- a/Engine/Source/Runtime/Core/Public/Templates/MemoryOps.h
+++ b/Engine/Source/Runtime/Core/Public/Templates/MemoryOps.h
@@ -248,7 +248,7 @@ FORCEINLINE bool CompareItems(const ElementType* A, const ElementType* B, SizeTy
 {
 	if constexpr (TTypeTraits<ElementType>::IsBytewiseComparable)
 	{
-		return !FMemory::Memcmp(A, B, sizeof(ElementType) * Count);
+		return !Count || !FMemory::Memcmp(A, B, sizeof(ElementType) * Count);
 	}
 	else
 	{
