Add missing include to EqualityComparable.h

diff --git a/Engine/Source/Runtime/Core/Public/Concepts/EqualityComparable.h b/Engine/Source/Runtime/Core/Public/Concepts/EqualityComparable.h
--- a/Engine/Source/Runtime/Core/Public/Concepts/EqualityComparable.h
+++ b/Engine/Source/Runtime/Core/Public/Concepts/EqualityComparable.h
@@ -2,6 +2,8 @@
 
 #pragma once
 
+#include "Templates/Models.h"
+
 /**
  * Describes a type comparable with another type.
  */
