fixed task retraction while waiting on named threads. check the new unit test for exact conditions

diff --git a/Engine/Source/Runtime/Core/Private/Async/TaskGraph.cpp b/Engine/Source/Runtime/Core/Private/Async/TaskGraph.cpp
--- a/Engine/Source/Runtime/Core/Private/Async/TaskGraph.cpp
+++ b/Engine/Source/Runtime/Core/Private/Async/TaskGraph.cpp
@@ -2140,26 +2140,6 @@ private:
 		TaskTrace::FWaitingScope WaitingScope(GetTraceIds(Tasks));
 		TRACE_CPUPROFILER_EVENT_SCOPE(WaitUntilTasksComplete);
 
-#if TASKGRAPH_NEW_FRONTEND
-		bool HasNamedThreadTasks = false;
-		for (const FGraphEventRef& Task : Tasks)
-		{
-			if (!Task->IsNamedThreadTask())
-			{
-				Task->Wait();
-			}
-			else
-			{
-				HasNamedThreadTasks = true;
-			}
-		}
-
-		if (!HasNamedThreadTasks)
-		{
-			return;
-		}
-#endif
-	
 		ENamedThreads::Type CurrentThread = CurrentThreadIfKnown;
 		if (ENamedThreads::GetThreadIndex(CurrentThreadIfKnown) == ENamedThreads::AnyThread)
 		{
@@ -2199,8 +2179,31 @@ private:
 			// named thread process tasks while we wait
 			TGraphTask<FReturnGraphTask>::CreateTask(&Tasks, CurrentThread).ConstructAndDispatchWhenReady(CurrentThread);
 			ProcessThreadUntilRequestReturn(CurrentThread);
+			return;
 		}
-		else if (LowLevelTasks::FScheduler::Get().IsWorkerThread() || LowLevelTasks::FReserveScheduler::Get().IsWorkerThread())
+
+#if TASKGRAPH_NEW_FRONTEND
+		// use the waiting logic of the new frontend, except for named thread tasks
+		bool HasNamedThreadTasks = false;
+		for (const FGraphEventRef& Task : Tasks)
+		{
+			if (!Task->IsNamedThreadTask())
+			{
+				Task->Wait();
+			}
+			else
+			{
+				HasNamedThreadTasks = true;
+			}
+		}
+
+		if (!HasNamedThreadTasks)
+		{
+			return;
+		}
+#endif
+
+		if (LowLevelTasks::FScheduler::Get().IsWorkerThread() || LowLevelTasks::FReserveScheduler::Get().IsWorkerThread())
 		{
 			// a worker thread gets blocked, involve a reserve worker to utilise an idle core
 			bool bSuccess = false; // has a reserve worker helped?
diff --git a/Engine/Source/Runtime/Core/Private/Tasks/TaskPrivate.cpp b/Engine/Source/Runtime/Core/Private/Tasks/TaskPrivate.cpp
--- a/Engine/Source/Runtime/Core/Private/Tasks/TaskPrivate.cpp
+++ b/Engine/Source/Runtime/Core/Private/Tasks/TaskPrivate.cpp
@@ -73,4 +73,31 @@ namespace UE::Tasks::Private
 		CurrentTask = Task;
 		return PrevTask;
 	}
+
+	bool TryWaitOnNamedThread(FTaskBase& Task)
+	{
+#if TASKGRAPH_NEW_FRONTEND
+		// handle waiting only on a named thread and if not called from inside a task
+		FTaskGraphInterface& TaskGraph = FTaskGraphInterface::Get();
+		ENamedThreads::Type CurrentThread = TaskGraph.GetCurrentThreadIfKnown();
+		if (CurrentThread < ENamedThreads::ActualRenderingThread /* is a named thread? */ && !TaskGraph.IsThreadProcessingTasks(CurrentThread))
+		{
+			// execute other tasks of this named thread while waiting
+			ETaskPriority Dummy;
+			EExtendedTaskPriority ExtendedPriority;
+			FBaseGraphTask::TranslatePriority(CurrentThread, Dummy, ExtendedPriority);
+
+			auto TaskBody = [CurrentThread, &TaskGraph] { TaskGraph.RequestReturn(CurrentThread); };
+			using FReturnFromNamedThreadTask = TExecutableTask<decltype(TaskBody)>;
+			FReturnFromNamedThreadTask ReturnTask{ TEXT("ReturnFromNamedThreadTask"), MoveTemp(TaskBody), ETaskPriority::High, ExtendedPriority };
+			ReturnTask.AddPrerequisites(Task);
+			ReturnTask.TryLaunch(); // the result doesn't matter
+
+			TaskGraph.ProcessThreadUntilRequestReturn(CurrentThread);
+			return true;
+		}
+#endif
+
+		return false;
+	}
 }
diff --git a/Engine/Source/Runtime/Core/Private/Tests/Async/TaskGraphTest.cpp b/Engine/Source/Runtime/Core/Private/Tests/Async/TaskGraphTest.cpp
--- a/Engine/Source/Runtime/Core/Private/Tests/Async/TaskGraphTest.cpp
+++ b/Engine/Source/Runtime/Core/Private/Tests/Async/TaskGraphTest.cpp
@@ -1045,6 +1045,12 @@ PRAGMA_ENABLE_DEPRECATION_WARNINGS
 			check(LocalQueueTask.GetRefCount() == 1);
 		}
 
+		{	// launch a GT task, then an any-thread task that depends on it. wait for the any-thread task. this was a deadlock on the new frontend
+			FGraphEventRef GTTask = FFunctionGraphTask::CreateAndDispatchWhenReady([] { IsInGameThread(); }, TStatId{}, nullptr, ENamedThreads::GameThread);
+			FGraphEventRef AnyThreadTask = FFunctionGraphTask::CreateAndDispatchWhenReady([] {}, TStatId{}, GTTask);
+			AnyThreadTask->Wait();
+		}
+
 		return true;
 	}
 
diff --git a/Engine/Source/Runtime/Core/Public/Async/TaskGraphInterfaces.h b/Engine/Source/Runtime/Core/Public/Async/TaskGraphInterfaces.h
--- a/Engine/Source/Runtime/Core/Public/Async/TaskGraphInterfaces.h
+++ b/Engine/Source/Runtime/Core/Public/Async/TaskGraphInterfaces.h
@@ -551,14 +551,13 @@ public:
 
 	void Wait(ENamedThreads::Type CurrentThreadIfKnown = ENamedThreads::AnyThread)
 	{
-		if (!IsNamedThreadTask())
+		// local queue have to be handled by the original TaskGraph implementation. The new frontend doesn't support local queues
+		if (ENamedThreads::GetQueueIndex(CurrentThreadIfKnown) != ENamedThreads::MainQueue)
 		{
-			FTaskBase::Wait();
-		}
-		else
-		{
-			FTaskGraphInterface::Get().WaitUntilTaskCompletes(this, CurrentThreadIfKnown);
+			return FTaskGraphInterface::Get().WaitUntilTaskCompletes(this, CurrentThreadIfKnown);
 		}
+
+		return FTaskBase::Wait();
 	}
 
 	ENamedThreads::Type GetThreadToExecuteOn() const
diff --git a/Engine/Source/Runtime/Core/Public/Tasks/TaskPrivate.h b/Engine/Source/Runtime/Core/Public/Tasks/TaskPrivate.h
--- a/Engine/Source/Runtime/Core/Public/Tasks/TaskPrivate.h
+++ b/Engine/Source/Runtime/Core/Public/Tasks/TaskPrivate.h
@@ -292,6 +292,14 @@ namespace UE::Tasks
 					return true;
 				}
 
+#if TASKGRAPH_NEW_FRONTEND
+				// task retraction is not supported for named thread tasks
+				if (IsNamedThreadTask())
+				{
+					return false;
+				}
+#endif
+
 				// avoid stack overflow. is not expected in a real-life cases but happens in stress tests
 				if (RecursionDepth == 200)
 				{
@@ -818,6 +826,10 @@ namespace UE::Tasks
 			TTypeCompatibleBytes<TaskBodyType> TaskBodyStorage;
 		};
 
+		// waiting on named threads that replicates TaskGraph logic
+		// returns true if called on a named thread
+		CORE_API bool TryWaitOnNamedThread(FTaskBase& Task);
+
 		// a special kind of task that is used for signalling or dependency management. It can have prerequisites or be used as a prerequisite for other tasks. 
 		// It's optimized for the fact that it doesn't have a task body and so doesn't need to be scheduled and executed
 		class FTaskEventBase : public TConcurrentLinearObject<FTaskEventBase, FTaskBlockAllocationTag>, public FTaskBase
@@ -844,17 +856,28 @@ namespace UE::Tasks
 
 		inline void FTaskBase::Wait()
 		{
+			if (IsCompleted())
+			{
+				return;
+			}
+
 			TaskTrace::FWaitingScope WaitingScope(GetTraceId());
 			TRACE_CPUPROFILER_EVENT_SCOPE(Tasks::Wait);
 
+			if (!IsAwaitable())
+			{
+				UE_LOG(LogTemp, Fatal, TEXT("Deadlock detected! A task can't be waited here, e.g. because it's being executed by the currect thread"));
+				return;
+			}
+
 			if (TryRetractAndExecute())
 			{
 				return;
 			}
 
-			if (!IsAwaitable())
+			// if we are on a named thread, handle waiting in TaskGraph-specific style
+			if (TryWaitOnNamedThread(*this))
 			{
-				UE_LOG(LogTemp, Fatal, TEXT("Deadlock detected! A task can't be waited here, e.g. because it's being executed by the currect thread"));
 				return;
 			}
 
