ThreadHeartBeat::OnHang broadcast OnHang delegate before triggering a memory exception. Delegate implementation will be called from the hang detector thread and not from the hung thread. Disabled in shipping build

#rnx

diff --git a/Engine/Source/Runtime/Core/Private/HAL/ThreadHeartBeat.cpp b/Engine/Source/Runtime/Core/Private/HAL/ThreadHeartBeat.cpp
--- a/Engine/Source/Runtime/Core/Private/HAL/ThreadHeartBeat.cpp
+++ b/Engine/Source/Runtime/Core/Private/HAL/ThreadHeartBeat.cpp
@@ -220,6 +220,12 @@ void FORCENOINLINE FThreadHeartBeat::OnHang(double HangDuration, uint32 ThreadTh
 #if PLATFORM_SWITCH
 	FPlatformCrashContext::UpdateDynamicData();
 #endif
+
+#if !UE_BUILD_SHIPPING
+	// Delegate implementation will be called from the hang detector thread and not from the hung thread
+	OnHangDelegate.ExecuteIfBound(ThreadThatHung);
+#endif
+
 	// We want to avoid all memory allocations if a hang is detected.
 	// Force a crash in a way that will generate a crash report.
 
diff --git a/Engine/Source/Runtime/Core/Public/HAL/ThreadHeartBeat.h b/Engine/Source/Runtime/Core/Public/HAL/ThreadHeartBeat.h
--- a/Engine/Source/Runtime/Core/Public/HAL/ThreadHeartBeat.h
+++ b/Engine/Source/Runtime/Core/Public/HAL/ThreadHeartBeat.h
@@ -17,6 +17,10 @@
 DECLARE_DELEGATE_OneParam(FOnThreadStuck, uint32);
 DECLARE_DELEGATE_OneParam(FOnThreadUnstuck, uint32);
 
+#if !UE_BUILD_SHIPPING
+DECLARE_DELEGATE_OneParam(FOnHangDelegate, uint32);
+#endif
+
 /**
  * Our own local clock.
  * Platforms that support suspend/resume have problems where a suspended title acts like
@@ -140,6 +144,9 @@ class CORE_API FThreadHeartBeat : public FRunnable
 
 	FOnThreadStuck OnStuck;
 	FOnThreadUnstuck OnUnstuck;
+#if !UE_BUILD_SHIPPING
+	FOnHangDelegate OnHangDelegate;
+#endif
 
 	FThreadHeartBeat();
 	virtual ~FThreadHeartBeat();
@@ -234,6 +241,15 @@ public:
 	FOnThreadStuck& GetOnThreadStuck() { return OnStuck; }
 	FOnThreadUnstuck& GetOnThreadUnstuck() { return OnUnstuck; }
 
+	/*
+	* Get delegate for callback on hang.
+	* Delegate implementation will be called from the hang detector thread and not from the hung thread
+	* Disabled in shipping build
+	*/
+#if !UE_BUILD_SHIPPING
+	FOnHangDelegate& GetOnHangDelegate() { return OnHangDelegate; }
+#endif
+
 	/*
 	*  Get hang duration threshold.
 	*/
