[Niagara] - nudge the DDC version for niagara scripts until Stu can dig into the core issue

diff --git a/Engine/Source/Runtime/Core/Private/UObject/FortniteNCBranchObjectVersion.cpp b/Engine/Source/Runtime/Core/Private/UObject/FortniteNCBranchObjectVersion.cpp
--- a/Engine/Source/Runtime/Core/Private/UObject/FortniteNCBranchObjectVersion.cpp
+++ b/Engine/Source/Runtime/Core/Private/UObject/FortniteNCBranchObjectVersion.cpp
@@ -7,6 +7,7 @@ TMap<FGuid, FGuid> FFortniteNCBranchObjectVersion::GetSystemGuids()
 	const FDevSystemGuids& DevGuids = FDevSystemGuids::Get();
 
 	SystemGuids.Add(DevGuids.STATICMESH_DERIVEDDATA_VER, FGuid("F16D2728168B4CDEBFB382B38FFE9E06"));
+	SystemGuids.Add(DevGuids.Niagara_LatestScriptCompileVersion, FGuid("08A58144744FCA4A9BF6DD45DFC4EC13"));
 
 	return SystemGuids;
 }
