Array Iterator: Add RemoveCurrentSwap.
Give RingBuffer its own iterator class since RemoveCurrentSwap does not fit the usual usage of a ringbuffer, which is supposed to maintain order.
#rnx

diff --git a/Engine/Source/Runtime/Core/Public/Containers/Array.h b/Engine/Source/Runtime/Core/Public/Containers/Array.h
--- a/Engine/Source/Runtime/Core/Public/Containers/Array.h
+++ b/Engine/Source/Runtime/Core/Public/Containers/Array.h
@@ -154,6 +154,17 @@ public:
 		Index--;
 	}
 
+	/**
+	 * Removes current element in array by swapping it with the end element and popping it from the end.
+	 * This invalidates the current iterator value and it must be incremented.
+	 * Note this modifies the order of the remaining elements in the array.
+	 */
+	void RemoveCurrentSwap()
+	{
+		Container.RemoveAtSwap(Index);
+		Index--;
+	}
+
 	FORCEINLINE friend bool operator==(const TIndexedContainerIterator& Lhs, const TIndexedContainerIterator& Rhs) { return &Lhs.Container == &Rhs.Container && Lhs.Index == Rhs.Index; }
 	FORCEINLINE friend bool operator!=(const TIndexedContainerIterator& Lhs, const TIndexedContainerIterator& Rhs) { return &Lhs.Container != &Rhs.Container || Lhs.Index != Rhs.Index; }
 
diff --git a/Engine/Source/Runtime/Core/Public/Containers/RingBuffer.h b/Engine/Source/Runtime/Core/Public/Containers/RingBuffer.h
--- a/Engine/Source/Runtime/Core/Public/Containers/RingBuffer.h
+++ b/Engine/Source/Runtime/Core/Public/Containers/RingBuffer.h
@@ -4,7 +4,6 @@
 
 #include "CoreTypes.h"
 #include "Misc/AssertionMacros.h"
-#include "Containers/Array.h" // TIndexedContainerIterator
 #include "Containers/ArrayView.h"
 #include "Templates/IsPODType.h"
 #include "Templates/IsTriviallyDestructible.h"
@@ -45,6 +44,116 @@ template <> struct TMakeUnsigned<uint32> { using Type = uint32; };
 template <> struct TMakeUnsigned<int64>  { using Type = uint64; };
 template <> struct TMakeUnsigned<uint64> { using Type = uint64; };
 
+template< typename ContainerType, typename ElementType, typename SizeType>
+class TRingBufferIterator
+{
+public:
+	TRingBufferIterator(ContainerType& InContainer, SizeType StartIndex = 0)
+		: Container(InContainer)
+		, Index(StartIndex)
+	{
+	}
+
+	/** Advances iterator to the next element in the container. */
+	TRingBufferIterator& operator++()
+	{
+		++Index;
+		return *this;
+	}
+	TRingBufferIterator operator++(int)
+	{
+		TRingBufferIterator Tmp(*this);
+		++Index;
+		return Tmp;
+	}
+
+	/** Moves iterator to the previous element in the container. */
+	TRingBufferIterator& operator--()
+	{
+		--Index;
+		return *this;
+	}
+	TRingBufferIterator operator--(int)
+	{
+		TRingBufferIterator Tmp(*this);
+		--Index;
+		return Tmp;
+	}
+
+	/** iterator arithmetic support */
+	TRingBufferIterator& operator+=(SizeType Offset)
+	{
+		Index += Offset;
+		return *this;
+	}
+
+	TRingBufferIterator operator+(SizeType Offset) const
+	{
+		TRingBufferIterator Tmp(*this);
+		return Tmp += Offset;
+	}
+
+	TRingBufferIterator& operator-=(SizeType Offset)
+	{
+		return *this += -Offset;
+	}
+
+	TRingBufferIterator operator-(SizeType Offset) const
+	{
+		TRingBufferIterator Tmp(*this);
+		return Tmp -= Offset;
+	}
+
+	FORCEINLINE ElementType& operator* () const
+	{
+		return Container[Index];
+	}
+
+	FORCEINLINE ElementType* operator->() const
+	{
+		return &Container[Index];
+	}
+
+	/** conversion to "bool" returning true if the iterator has not reached the last element. */
+	FORCEINLINE explicit operator bool() const
+	{
+		return Container.IsValidIndex(Index);
+	}
+
+	/** Returns an index to the current element. */
+	SizeType GetIndex() const
+	{
+		return Index;
+	}
+
+	/** Resets the iterator to the first element. */
+	void Reset()
+	{
+		Index = 0;
+	}
+
+	/** Sets the iterator to one past the last element. */
+	void SetToEnd()
+	{
+		Index = Container.Num();
+	}
+
+	/** Removes current element in array. This invalidates the current iterator value and it must be incremented */
+	void RemoveCurrent()
+	{
+		Container.RemoveAt(Index);
+		Index--;
+	}
+
+	FORCEINLINE friend bool operator==(const TRingBufferIterator& Lhs, const TRingBufferIterator& Rhs) { return &Lhs.Container == &Rhs.Container && Lhs.Index == Rhs.Index; }
+	FORCEINLINE friend bool operator!=(const TRingBufferIterator& Lhs, const TRingBufferIterator& Rhs) { return &Lhs.Container != &Rhs.Container || Lhs.Index != Rhs.Index; }
+
+private:
+
+	ContainerType& Container;
+	SizeType      Index;
+};
+
 /**
  * RingBuffer - an array with a Front and Back pointer and with implicit wraparound to the beginning of the array when reaching the end of the array when iterating from Front to Back
  * Useful for providing O(1) push/pop at the end of the array (for Queue or Stack) while still having high cache coherency during iteration.
@@ -67,8 +176,8 @@ public:
 	/** Type used to communicate size and capacity and counts */
 	typedef typename TMakeUnsigned<typename Allocator::SizeType>::Type SizeType;
 	/** Iterator type used for ranged-for traversal. */
-	typedef TIndexedContainerIterator<TRingBuffer, ElementType, typename Allocator::SizeType> TIterator;
-	typedef TIndexedContainerIterator<const TRingBuffer, const ElementType, typename Allocator::SizeType> TConstIterator;
+	typedef TRingBufferIterator<TRingBuffer, ElementType, typename Allocator::SizeType> TIterator;
+	typedef TRingBufferIterator<const TRingBuffer, const ElementType, typename Allocator::SizeType> TConstIterator;
 private:
 	/**
 	 * Type used for variables that are indexes into the underlying storage.
diff --git a/Engine/Source/Runtime/CoreUObject/Private/UObject/SavePackage2.cpp b/Engine/Source/Runtime/CoreUObject/Private/UObject/SavePackage2.cpp
--- a/Engine/Source/Runtime/CoreUObject/Private/UObject/SavePackage2.cpp
+++ b/Engine/Source/Runtime/CoreUObject/Private/UObject/SavePackage2.cpp
@@ -2083,21 +2083,20 @@ ESavePackageResult BeginCachePlatformCookedData(FSaveContext& SaveContext)
 		const ITargetPlatform* TargetPlatform = SaveContext.GetTargetPlatform();
 		TArray<UObject*> ObjectsInPackage;
 		GetObjectsWithPackage(SaveContext.GetPackage(), ObjectsInPackage);
-		for (int32 Index = 0; Index < ObjectsInPackage.Num(); /** incremented in loop */)
+		for (TArray<UObject*>::TIterator Iter(ObjectsInPackage); Iter; ++Iter)
 		{
-			UObject* Object = ObjectsInPackage[Index];
+			UObject* Object = *Iter;
 			if (SaveContext.IsUnsaveable(Object))
 			{
-				ObjectsInPackage.RemoveAtSwap(Index);
+				Iter.RemoveCurrentSwap();
 				continue;
 			}
 			Object->BeginCacheForCookedPlatformData(TargetPlatform);
 			if (Object->IsCachedCookedPlatformDataLoaded(TargetPlatform))
 			{
-				ObjectsInPackage.RemoveAtSwap(Index);
+				Iter.RemoveCurrentSwap();
 				continue;
 			}
-			++Index;
 		}
 		if (ObjectsInPackage.Num())
 		{
@@ -2115,14 +2114,13 @@ ESavePackageResult BeginCachePlatformCookedData(FSaveContext& SaveContext)
 						SaveContext.GetFilename(), *ObjectsInPackage[0]->GetFullName());
 					return ESavePackageResult::Error;
 				}
-				for (int32 Index = 0; Index < ObjectsInPackage.Num(); /** incremented in loop */)
+				for (TArray<UObject*>::TIterator Iter(ObjectsInPackage); Iter; ++Iter)
 				{
-					if (ObjectsInPackage[Index]->IsCachedCookedPlatformDataLoaded(TargetPlatform))
+					if ((*Iter)->IsCachedCookedPlatformDataLoaded(TargetPlatform))
 					{
-						ObjectsInPackage.RemoveAtSwap(Index);
+						Iter.RemoveCurrentSwap();
 						continue;
 					}
-					++Index;
 				}
 			}
 		}
