Fixed compiler error (forgot FramePro isn't built on PS4).

diff --git a/Engine/Source/Runtime/Core/Private/FramePro/FrameProProfiler.cpp b/Engine/Source/Runtime/Core/Private/FramePro/FrameProProfiler.cpp
--- a/Engine/Source/Runtime/Core/Private/FramePro/FrameProProfiler.cpp
+++ b/Engine/Source/Runtime/Core/Private/FramePro/FrameProProfiler.cpp
@@ -426,7 +426,7 @@ FString FFrameProProfiler::StartFrameProRecording(const FString& FilenameRoot, i
 	GFrameProEnabled = true;
 
 	// Enable named events as well
-	++GCycleStatsShouldEmitNamedEvents:
+	++GCycleStatsShouldEmitNamedEvents;
 
 	// Set recording flag
 	GFrameProIsRecording = true;
