PR #8596: Add dll linkage for FNumberFormattingOptions::GetTypeHash & FNumberParsingOptions::GetTypeHash (Contributed by [REDACTED_PERSON])

diff --git a/Engine/Source/Runtime/Core/Public/Internationalization/Text.h b/Engine/Source/Runtime/Core/Public/Internationalization/Text.h
--- a/Engine/Source/Runtime/Core/Public/Internationalization/Text.h
+++ b/Engine/Source/Runtime/Core/Public/Internationalization/Text.h
@@ -212,10 +212,10 @@ struct CORE_API FNumberFormattingOptions
 	int32 MaximumFractionalDigits;
 	FNumberFormattingOptions& SetMaximumFractionalDigits( int32 InValue ){ MaximumFractionalDigits = InValue; return *this; }
 
-	friend void operator<<(FStructuredArchive::FSlot Slot, FNumberFormattingOptions& Value);
+	friend CORE_API void operator<<(FStructuredArchive::FSlot Slot, FNumberFormattingOptions& Value);
 
 	/** Get the hash code to use for the given formatting options */
-	friend uint32 GetTypeHash( const FNumberFormattingOptions& Key );
+	friend CORE_API uint32 GetTypeHash( const FNumberFormattingOptions& Key );
 
 	/** Check to see if our formatting options match the other formatting options */
 	bool IsIdentical( const FNumberFormattingOptions& Other ) const;
@@ -242,10 +242,10 @@ struct CORE_API FNumberParsingOptions
 	bool UseClamping;
 	FNumberParsingOptions& SetUseClamping(bool InValue) { UseClamping = InValue; return *this; }
 
-	friend void operator<<(FStructuredArchive::FSlot Slot, FNumberParsingOptions& Value);
+	friend CORE_API void operator<<(FStructuredArchive::FSlot Slot, FNumberParsingOptions& Value);
 
 	/** Get the hash code to use for the given parsing options */
-	friend uint32 GetTypeHash( const FNumberParsingOptions& Key );
+	friend CORE_API uint32 GetTypeHash( const FNumberParsingOptions& Key );
 
 	/** Check to see if our parsing options match the other parsing options */
 	bool IsIdentical( const FNumberParsingOptions& Other ) const;
