[Insights] Added "TraceAnalysis" and "Trace/CpuProfile" LLM_SCOPEs (in an effort to reduce the Untagged allocations):
  - "TraceAnalysis" -- captures all mem allocated by the FAnalysisProcessor.
  - "Trace/CpuProfile" -- captures allocation of per-thread buffers used by the cpu profiler trace.

diff --git a/Engine/Source/Developer/TraceAnalysis/Private/Analysis/Processor.cpp b/Engine/Source/Developer/TraceAnalysis/Private/Analysis/Processor.cpp
--- a/Engine/Source/Developer/TraceAnalysis/Private/Analysis/Processor.cpp
+++ b/Engine/Source/Developer/TraceAnalysis/Private/Analysis/Processor.cpp
@@ -33,6 +33,8 @@ FAnalysisProcessor::FImpl::~FImpl()
 ////////////////////////////////////////////////////////////////////////////////
 uint32 FAnalysisProcessor::FImpl::Run()
 {
+	LLM_SCOPE_BYNAME(TEXT("TraceAnalysis"));
+
 	AnalysisEngine.Begin();
 
 	FStreamBuffer Buffer(4 << 20);
diff --git a/Engine/Source/Runtime/Core/Private/ProfilingDebugging/CpuProfilerTrace.cpp b/Engine/Source/Runtime/Core/Private/ProfilingDebugging/CpuProfilerTrace.cpp
--- a/Engine/Source/Runtime/Core/Private/ProfilingDebugging/CpuProfilerTrace.cpp
+++ b/Engine/Source/Runtime/Core/Private/ProfilingDebugging/CpuProfilerTrace.cpp
@@ -115,6 +115,7 @@ thread_local uint32 FCpuProfilerTraceInternal::ThreadDepth = 0;
 
 FCpuProfilerTraceInternal::FThreadBuffer* FCpuProfilerTraceInternal::CreateThreadBuffer()
 {
+	LLM_SCOPE_BYNAME(TEXT("Trace/CpuProfiler"));
 	ThreadBuffer = new FThreadBuffer();
 	ThreadBuffer->Register();
 	return ThreadBuffer;
