Fixed issue with non axis aligned HLOD volumes
* Brushes or rotated volumes were including more actors than they should have

diff --git a/Engine/Source/Editor/UnrealEd/Private/HierarchicalLOD.cpp b/Engine/Source/Editor/UnrealEd/Private/HierarchicalLOD.cpp
--- a/Engine/Source/Editor/UnrealEd/Private/HierarchicalLOD.cpp
+++ b/Engine/Source/Editor/UnrealEd/Private/HierarchicalLOD.cpp
@@ -226,12 +226,10 @@ void FHierarchicalLODBuilder::BuildClusters(ULevel* InLevel)
 							if (VolumePtr && (*VolumePtr)->AppliesToHLODLevel(LODId))
 							{
 								AHierarchicalLODVolume* Volume = *VolumePtr;
-								FBox HLODVolumeBox = Volume->GetComponentsBoundingBox(true);
 
-								auto IsInVolume = [Volume, &HLODVolumeBox](const AActor* Actor)
+								auto IsInVolume = [Volume](const AActor* Actor)
 								{
-									FBox ActorBox = Actor->GetComponentsBoundingBox(true);
-									return HLODVolumeBox.IsInside(ActorBox) || (Volume->bIncludeOverlappingActors && HLODVolumeBox.Intersect(ActorBox));
+									return Volume->IsActorIncluded(Actor);
 								};
 
 								EvaluateValidActors(IsInVolume);
@@ -516,8 +514,7 @@ void FHierarchicalLODBuilder::InitializeClusters(ULevel* InLevel, const int32 LO
 		FBox ActorBox = InActor->GetComponentsBoundingBox(true);
 		for (TPair<AHierarchicalLODVolume*, FLODCluster>& Cluster : HLODVolumeClusters)
 		{
-			FBox HLODVolumeBox = Cluster.Key->GetComponentsBoundingBox(true);
-			if (HLODVolumeBox.IsInside(ActorBox) || (Cluster.Key->bIncludeOverlappingActors && HLODVolumeBox.Intersect(ActorBox)))
+			if (Cluster.Key->IsActorIncluded(InActor))
 			{
 				Cluster.Value += InActor;
 				return true;
diff --git a/Engine/Source/Editor/UnrealEd/Private/HierarchicalLODVolume.cpp b/Engine/Source/Editor/UnrealEd/Private/HierarchicalLODVolume.cpp
--- a/Engine/Source/Editor/UnrealEd/Private/HierarchicalLODVolume.cpp
+++ b/Engine/Source/Editor/UnrealEd/Private/HierarchicalLODVolume.cpp
@@ -3,6 +3,7 @@
 #include "HierarchicalLODVolume.h"
 #include "Engine/CollisionProfile.h"
 #include "Components/BrushComponent.h"
+#include "Algo/AllOf.h"
 
 AHierarchicalLODVolume::AHierarchicalLODVolume(const FObjectInitializer& ObjectInitializer)
 	: Super(ObjectInitializer)
@@ -29,4 +30,20 @@ bool AHierarchicalLODVolume::AppliesToHLODLevel(int32 LODIdx) const
 {
 	return ApplyOnlyToSpecificHLODLevels.Num() == 0 ||
 		   ApplyOnlyToSpecificHLODLevels.Contains(LODIdx);
+}
+
+bool AHierarchicalLODVolume::IsActorIncluded(const AActor* InActor) const
+{
+	FBox ActorBoundingBox = InActor->GetComponentsBoundingBox(true);
+
+	if (bIncludeOverlappingActors)
+	{
+		return EncompassesPoint(ActorBoundingBox.GetCenter(), ActorBoundingBox.GetExtent().Size());
+	}
+	else
+	{
+		FVector Vertices[8];
+		ActorBoundingBox.GetVertices(Vertices);
+		return Algo::AllOf(Vertices, [this](const FVector& Vertex) { return EncompassesPoint(Vertex); });
+	}		
 }
\ No newline at end of file
diff --git a/Engine/Source/Editor/UnrealEd/Public/HierarchicalLODVolume.h b/Engine/Source/Editor/UnrealEd/Public/HierarchicalLODVolume.h
--- a/Engine/Source/Editor/UnrealEd/Public/HierarchicalLODVolume.h
+++ b/Engine/Source/Editor/UnrealEd/Public/HierarchicalLODVolume.h
@@ -17,6 +17,9 @@ class AHierarchicalLODVolume : public AVolume
 	virtual bool NeedsLoadForClient() const override { return false; }
 	virtual bool NeedsLoadForServer() const override { return false; }
 	virtual bool IsLevelBoundsRelevant() const override { return false; }
+
+	bool IsActorIncluded(const AActor* InActor) const;
+
 public:
 	bool AppliesToHLODLevel(int32 LODIdx) const;
 
diff --git a/Engine/Source/Runtime/Core/Public/Math/Box.h b/Engine/Source/Runtime/Core/Public/Math/Box.h
--- a/Engine/Source/Runtime/Core/Public/Math/Box.h
+++ b/Engine/Source/Runtime/Core/Public/Math/Box.h
@@ -492,6 +492,13 @@ public:
 	 */
 	FString ToString() const;
 
+	/**
+	 * Get the vertices that make up this box.
+	 * 
+	 * 
+	 */
+	void GetVertices( TVector<T> (&Vertices)[8] ) const;
+
 public:
 
 	/** 
@@ -731,20 +738,24 @@ TBox<T> TBox<T>::TransformBy(const TTransform<T>& M) const
 	return TransformBy(M.ToMatrixWithScale());
 }
 
+template<typename T>
+void TBox<T>::GetVertices(TVector<T>(&Vertices)[8]) const
+{
+	Vertices[0] = TVector<T>(Min);
+	Vertices[1] = TVector<T>(Min.X, Min.Y, Max.Z);
+	Vertices[2] = TVector<T>(Min.X, Max.Y, Min.Z);
+	Vertices[3] = TVector<T>(Max.X, Min.Y, Min.Z);
+	Vertices[4] = TVector<T>(Max.X, Max.Y, Min.Z);
+	Vertices[5] = TVector<T>(Max.X, Min.Y, Max.Z);
+	Vertices[6] = TVector<T>(Min.X, Max.Y, Max.Z);
+	Vertices[7] = TVector<T>(Max);
+}
+
 template<typename T>
 TBox<T> TBox<T>::InverseTransformBy(const TTransform<T>& M) const
 {
-	TVector<T> Vertices[8] =
-	{
-		TVector<T>(Min),
-		TVector<T>(Min.X, Min.Y, Max.Z),
-		TVector<T>(Min.X, Max.Y, Min.Z),
-		TVector<T>(Max.X, Min.Y, Min.Z),
-		TVector<T>(Max.X, Max.Y, Min.Z),
-		TVector<T>(Max.X, Min.Y, Max.Z),
-		TVector<T>(Min.X, Max.Y, Max.Z),
-		TVector<T>(Max)
-	};
+	TVector<T> Vertices[8];
+	GetVertices(Vertices);
 
 	TBox<T> NewBox(ForceInit);
 
@@ -760,17 +771,8 @@ TBox<T> TBox<T>::InverseTransformBy(const TTransform<T>& M) const
 template<typename T>
 TBox<T> TBox<T>::TransformProjectBy(const TMatrix<T>& ProjM) const
 {
-	TVector<T> Vertices[8] =
-	{
-		TVector<T>(Min),
-		TVector<T>(Min.X, Min.Y, Max.Z),
-		TVector<T>(Min.X, Max.Y, Min.Z),
-		TVector<T>(Max.X, Min.Y, Min.Z),
-		TVector<T>(Max.X, Max.Y, Min.Z),
-		TVector<T>(Max.X, Min.Y, Max.Z),
-		TVector<T>(Min.X, Max.Y, Max.Z),
-		TVector<T>(Max)
-	};
+	TVector<T> Vertices[8];
+	GetVertices(Vertices);
 
 	TBox<T> NewBox(ForceInit);
 
