PR #8891: TokenizedMessage: Move exported symbols to .cpp (Contributed by [REDACTED_PERSON])

Also used the opportunity to convert from MakeShareable to MakeShared.

diff --git a/Engine/Source/Runtime/Core/Private/Logging/TokenizedMessage.cpp b/Engine/Source/Runtime/Core/Private/Logging/TokenizedMessage.cpp
--- a/Engine/Source/Runtime/Core/Private/Logging/TokenizedMessage.cpp
+++ b/Engine/Source/Runtime/Core/Private/Logging/TokenizedMessage.cpp
@@ -8,7 +8,7 @@
 
 TSharedRef<FTokenizedMessage> FTokenizedMessage::Create(EMessageSeverity::Type InSeverity, const FText& InMessageText)
 {
-	TSharedRef<FTokenizedMessage> Message = MakeShareable(new FTokenizedMessage());
+	TSharedRef<FTokenizedMessage> Message = MakeShared<FTokenizedMessage>(FPrivateToken());
 	Message->SetSeverity( InSeverity );
 	Message->AddToken( FSeverityToken::Create(InSeverity) );
 	if(!InMessageText.IsEmpty())
@@ -144,12 +144,37 @@ TSharedPtr<IMessageToken> FTokenizedMessage::GetMessageLink() const
 	return MessageLink;
 }
 
+TSharedRef<FTextToken> FTextToken::Create(const FText& InMessage)
+{
+	return MakeShared<FTextToken>(FPrivateToken(), InMessage);
+}
+
+TSharedRef<FDynamicTextToken> FDynamicTextToken::Create(const TAttribute<FText>& InMessage)
+{
+	return MakeShared<FDynamicTextToken>(FPrivateToken(), InMessage);
+}
+
+TSharedRef<FImageToken> FImageToken::Create(const FName& InImageName)
+{
+	return MakeShared<FImageToken>(FPrivateToken(), InImageName);
+}
+
+TSharedRef<FSeverityToken> FSeverityToken::Create(EMessageSeverity::Type InSeverity)
+{
+	return MakeShared<FSeverityToken>(FPrivateToken(), InSeverity);
+}
+
+TSharedRef<FURLToken> FURLToken::Create(const FString& InURL, const FText& InMessage)
+{
+	return MakeShared<FURLToken>(FPrivateToken(), InURL, InMessage);
+}
+
 void FURLToken::VisitURL(const TSharedRef<IMessageToken>& Token, FString InURL)
 {	
 	FPlatformProcess::LaunchURL(*InURL, NULL, NULL);
 }
 
-FURLToken::FURLToken( const FString& InURL, const FText& InMessage )
+FURLToken::FURLToken(const FString& InURL, const FText& InMessage)
 {
 	URL = InURL;
 
@@ -177,7 +202,7 @@ void FAssetNameToken::FindAsset(const TSharedRef<IMessageToken>& Token, FString
 
 TSharedRef<FAssetNameToken> FAssetNameToken::Create(const FString& InAssetName, const FText& InMessage)
 {
-	return MakeShareable(new FAssetNameToken(InAssetName, InMessage));
+	return MakeShared<FAssetNameToken>(FPrivateToken(), InAssetName, InMessage);
 }
 
 FAssetNameToken::FAssetNameToken(const FString& InAssetName, const FText& InMessage)
@@ -195,10 +220,10 @@ FAssetNameToken::FAssetNameToken(const FString& InAssetName, const FText& InMess
 	MessageTokenActivated = FOnMessageTokenActivated::CreateStatic(&FAssetNameToken::FindAsset, AssetName);
 }
 
-FDocumentationToken::FDocumentationToken( FString InDocumentationLink, FString InPreviewExcerptLink, FString InPreviewExcerptName )
-	: DocumentationLink(MoveTemp(InDocumentationLink))
-	, PreviewExcerptLink(MoveTemp(InPreviewExcerptLink))
-	, PreviewExcerptName(MoveTemp(InPreviewExcerptName))
+FDocumentationToken::FDocumentationToken(const FString& InDocumentationLink, const FString& InPreviewExcerptLink, const FString& InPreviewExcerptName)
+	: DocumentationLink(InDocumentationLink)
+	, PreviewExcerptLink(InPreviewExcerptLink)
+	, PreviewExcerptName(InPreviewExcerptName)
 {
 	if (!PreviewExcerptName.IsEmpty())
 	{
@@ -206,16 +231,26 @@ FDocumentationToken::FDocumentationToken( FString InDocumentationLink, FString I
 	}
 }
 
+TSharedRef<FActionToken> FActionToken::Create(const FText& ActionName, const FText& ActionDescription, const FOnActionTokenExecuted& Action, bool bInSingleUse)
+{
+	return MakeShared<FActionToken>(FPrivateToken(), ActionName, ActionDescription, Action, bInSingleUse);
+}
+
+TSharedRef<FTutorialToken> FTutorialToken::Create(const FString& TutorialAssetName)
+{
+	return MakeShared<FTutorialToken>(FPrivateToken(), TutorialAssetName);
+}
+
 TSharedRef<FDocumentationToken> FDocumentationToken::Create(const FString& InDocumentationLink, const FString& InPreviewExcerptLink, const FString& InPreviewExcerptName)
 {
-	return MakeShareable(new FDocumentationToken(InDocumentationLink, InPreviewExcerptLink, InPreviewExcerptName));
+	return MakeShared<FDocumentationToken>(FPrivateToken(), InDocumentationLink, InPreviewExcerptLink, InPreviewExcerptName);
 }
 
 FOnMessageTokenActivated FActorToken::DefaultMessageTokenActivated;
 
 TSharedRef<FActorToken> FActorToken::Create(const FString& InActorPath, const FGuid& InActorGuid, const FText& InMessage)
 {
-	return MakeShareable(new FActorToken(InActorPath, InActorGuid, InMessage));
+	return MakeShared<FActorToken>(FPrivateToken(), InActorPath, InActorGuid, InMessage);
 }
 
 FActorToken::FActorToken(const FString& InActorPath, const FGuid& InActorGuid, const FText& InMessage)
diff --git a/Engine/Source/Runtime/Core/Public/Logging/TokenizedMessage.h b/Engine/Source/Runtime/Core/Public/Logging/TokenizedMessage.h
--- a/Engine/Source/Runtime/Core/Public/Logging/TokenizedMessage.h
+++ b/Engine/Source/Runtime/Core/Public/Logging/TokenizedMessage.h
@@ -107,6 +107,9 @@ protected:
 /** This class represents a rich tokenized message, such as would be used for compiler output with 'hyperlinks' to source file locations */
 class FTokenizedMessage : public TSharedFromThis<FTokenizedMessage>
 {
+	// The private token allows only members or friends to call MakeShared.
+	struct FPrivateToken { explicit FPrivateToken() = default; };
+
 public:
 
 	/** 
@@ -193,12 +196,12 @@ public:
 	/** Assigns Identifier for the message to the provided name */
 	CORE_API void SetIdentifier(FName InIdentifier);
 
-private:
 	/** Private constructor - we want to only create these structures as shared references via Create() */
-	FTokenizedMessage()
+	explicit FTokenizedMessage(FPrivateToken)
 		: Severity( EMessageSeverity::Info )
 		, Identifier(NAME_None)
-	{ }
+	{
+	}
 
 protected:
 
@@ -220,12 +223,12 @@ private:
 /** Basic message token with a localized text payload */
 class FTextToken : public IMessageToken
 {
+	// The private token allows only members or friends to call MakeShared.
+	struct FPrivateToken { explicit FPrivateToken() = default; };
+
 public:
 	/** Factory method, tokens can only be constructed as shared refs */
-	CORE_API static TSharedRef<FTextToken> Create( const FText& InMessage )
-	{
-		return MakeShareable(new FTextToken(InMessage));
-	}
+	CORE_API static TSharedRef<FTextToken> Create(const FText& InMessage);
 
 	/** Begin IMessageToken interface */
 	virtual EMessageToken::Type GetType() const override
@@ -234,9 +237,8 @@ public:
 	}
 	/** End IMessageToken interface */
 
-private:
 	/** Private constructor */
-	FTextToken( const FText& InMessage )
+	FTextToken(FPrivateToken, const FText& InMessage)
 	{
 		CachedText = InMessage;
 	}
@@ -245,12 +247,12 @@ private:
 /** Message token with a localized attribute text payload */
 class FDynamicTextToken : public IMessageToken
 {
+	// The private token allows only members or friends to call MakeShared.
+	struct FPrivateToken { explicit FPrivateToken() = default; };
+
 public:
 	/** Factory method, tokens can only be constructed as shared refs */
-	CORE_API static TSharedRef<FDynamicTextToken> Create(const TAttribute<FText>& InMessage)
-	{
-		return MakeShareable(new FDynamicTextToken(InMessage));
-	}
+	CORE_API static TSharedRef<FDynamicTextToken> Create(const TAttribute<FText>& InMessage);
 
 	/** Begin IMessageToken interface */
 	virtual EMessageToken::Type GetType() const override
@@ -269,14 +271,14 @@ public:
 		return Message;
 	}
 
-private:
 	/** Private constructor */
-	FDynamicTextToken(const TAttribute<FText>& InMessage)
+	FDynamicTextToken(FPrivateToken, const TAttribute<FText>& InMessage)
 		: Message(InMessage)
 	{
 		CachedText = InMessage.Get();
 	}
 
+private:
 	/** The attribute text of this message */
 	TAttribute<FText> Message;
 };
@@ -284,12 +286,12 @@ private:
 /** Basic message token with an icon/image payload */
 class FImageToken : public IMessageToken
 {
+	// The private token allows only members or friends to call MakeShared.
+	struct FPrivateToken { explicit FPrivateToken() = default; };
+
 public:
 	/** Factory method, tokens can only be constructed as shared refs */
-	CORE_API static TSharedRef<FImageToken> Create( const FName& InImageName )
-	{
-		return MakeShareable(new FImageToken(InImageName));
-	}
+	CORE_API static TSharedRef<FImageToken> Create(const FName& InImageName);
 
 	/** Begin IMessageToken interface */
 	virtual EMessageToken::Type GetType() const override
@@ -304,14 +306,14 @@ public:
 		return ImageName;
 	}
 
-private:
 	/** Private constructor */
-	FImageToken( const FName& InImageName )
+	FImageToken(FPrivateToken, const FName& InImageName)
 		: ImageName(InImageName)
 	{
 		CachedText = FText::FromName( InImageName );
 	}
 
+private:
 	/** A name to be used as a brush in this message */
 	FName ImageName;
 };
@@ -319,12 +321,12 @@ private:
 /** Basic message token with a severity payload */
 class FSeverityToken : public IMessageToken
 {
+	// The private token allows only members or friends to call MakeShared.
+	struct FPrivateToken { explicit FPrivateToken() = default; };
+
 public:
 	/** Factory method, tokens can only be constructed as shared refs */
-	CORE_API static TSharedRef<FSeverityToken> Create( EMessageSeverity::Type InSeverity )
-	{
-		return MakeShareable(new FSeverityToken(InSeverity));
-	}
+	CORE_API static TSharedRef<FSeverityToken> Create(EMessageSeverity::Type InSeverity);
 
 	/** Begin IMessageToken interface */
 	virtual EMessageToken::Type GetType() const override
@@ -339,14 +341,14 @@ public:
 		return Severity;
 	}
 
-private:
 	/** Private constructor */
-	FSeverityToken( EMessageSeverity::Type InSeverity )
+	FSeverityToken(FPrivateToken, EMessageSeverity::Type InSeverity)
 		: Severity(InSeverity)
 	{
 		CachedText = FTokenizedMessage::GetSeverityText( InSeverity );
 	}
 
+private:
 	/** A severity for this token */
 	EMessageSeverity::Type Severity;
 };
@@ -354,12 +356,12 @@ private:
 /** Basic message token that defaults is activated method to traverse a URL */
 class FURLToken : public IMessageToken
 {
+	// The private token allows only members or friends to call MakeShared.
+	struct FPrivateToken { explicit FPrivateToken() = default; };
+
 public:
 	/** Factory method, tokens can only be constructed as shared refs */
-	CORE_API static TSharedRef<FURLToken> Create( const FString& InURL, const FText& InMessage = FText() )
-	{
-		return MakeShareable(new FURLToken(InURL, InMessage));
-	}
+	CORE_API static TSharedRef<FURLToken> Create(const FString& InURL, const FText& InMessage = FText());
 
 	/** Begin IMessageToken interface */
 	virtual EMessageToken::Type GetType() const override
@@ -374,9 +376,15 @@ public:
 		return URL;
 	}
 
+	/** Private constructor */
+	inline FURLToken(FPrivateToken, const FString& InURL, const FText& InMessage)
+		: FURLToken(InURL, InMessage)
+	{
+	}
+
 private:
 	/** Private constructor */
-	FURLToken( const FString& InURL, const FText& InMessage );
+	FURLToken(const FString& InURL, const FText& InMessage);
 
 	/**
 	 * Delegate used to visit a URL
@@ -395,6 +403,9 @@ private:
  */
 class FAssetNameToken : public IMessageToken
 {
+	// The private token allows only members or friends to call MakeShared.
+	struct FPrivateToken { explicit FPrivateToken() = default; };
+
 public:
 	/** Factory method, tokens can only be constructed as shared refs */
 	CORE_API static TSharedRef<FAssetNameToken> Create(const FString& InAssetName, const FText& InMessage = FText());
@@ -413,14 +424,20 @@ public:
 	}
 
 	DECLARE_DELEGATE_OneParam(FOnGotoAsset, const FString&);
-	CORE_API static FOnGotoAsset& OnGotoAsset()
+	static FOnGotoAsset& OnGotoAsset()
 	{
 		return GotoAsset;
 	}
 
+	/** Private constructor */
+	inline FAssetNameToken(FPrivateToken, const FString& InAssetName, const FText& InMessage)
+		: FAssetNameToken(InAssetName, InMessage)
+	{
+	}
+
 private:
 	/** Private constructor */
-	FAssetNameToken( const FString& InAssetName, const FText& InMessage );
+	FAssetNameToken(const FString& InAssetName, const FText& InMessage);
 
 	/**
 	 * Delegate used to find a file
@@ -441,6 +458,9 @@ private:
  */
 class FDocumentationToken : public IMessageToken
 {
+	// The private token allows only members or friends to call MakeShared.
+	struct FPrivateToken { explicit FPrivateToken() = default; };
+
 public:
 	/** Factory method, tokens can only be constructed as shared refs */
 	CORE_API static TSharedRef<FDocumentationToken> Create(const FString& InDocumentationLink, const FString& InPreviewExcerptLink = FString(), const FString& InPreviewExcerptName = FString());
@@ -470,9 +490,14 @@ public:
 		return PreviewExcerptName;
 	}
 
+	inline FDocumentationToken(FPrivateToken, const FString& InDocumentationLink, const FString& InPreviewExcerptLink, const FString& InPreviewExcerptName)
+		: FDocumentationToken(InDocumentationLink, InPreviewExcerptLink, InPreviewExcerptName)
+	{
+	}
+
 protected:
 	/** Protected constructor */
-	FDocumentationToken( FString InDocumentationLink, FString InPreviewExcerptLink, FString InPreviewExcerptName );
+	FDocumentationToken(const FString& InDocumentationLink, const FString& InPreviewExcerptLink, const FString& InPreviewExcerptName);
 
 private:
 	/** The documentation path we link to when clicked */
@@ -494,13 +519,13 @@ DECLARE_DELEGATE(FOnActionTokenExecuted);
 class FActionToken
 	: public IMessageToken
 {
+	// The private token allows only members or friends to call MakeShared.
+	struct FPrivateToken { explicit FPrivateToken() = default; };
+
 public:
 
 	/** Factory method, tokens can only be constructed as shared refs */
-	CORE_API static TSharedRef<FActionToken> Create( const FText& ActionName, const FText& ActionDescription, const FOnActionTokenExecuted& Action, bool bInSingleUse=false )
-	{
-		return MakeShareable(new FActionToken(ActionName, ActionDescription, Action, bInSingleUse));
-	}
+	CORE_API static TSharedRef<FActionToken> Create(const FText& ActionName, const FText& ActionDescription, const FOnActionTokenExecuted& Action, bool bInSingleUse = false);
 
 	/** Executes the assigned action delegate. */
 	void ExecuteAction()
@@ -521,14 +546,17 @@ public:
 		return Action.IsBound() && (!bSingleUse || !bActionExecuted);
 	}
 
-public:
-
 	// IMessageToken interface
 	virtual EMessageToken::Type GetType() const override
 	{
 		return EMessageToken::Action;
 	}
 
+	inline FActionToken(FPrivateToken, const FText& InActionName, const FText& InActionDescription, const FOnActionTokenExecuted& InAction, bool bInSingleUse)
+		: FActionToken(InActionName, InActionDescription, InAction, bInSingleUse)
+	{
+	}
+
 protected:
 
 	/** Hidden constructor. */
@@ -560,13 +588,13 @@ private:
 class FTutorialToken
 	: public IMessageToken
 {
+	// The private token allows only members or friends to call MakeShared.
+	struct FPrivateToken { explicit FPrivateToken() = default; };
+
 public:
 
 	/** Factory method, tokens can only be constructed as shared refs */
-	CORE_API static TSharedRef<FTutorialToken> Create( const FString& TutorialAssetName )
-	{
-		return MakeShareable(new FTutorialToken(TutorialAssetName));
-	}
+	CORE_API static TSharedRef<FTutorialToken> Create(const FString& TutorialAssetName);
 
 public:
 
@@ -583,6 +611,11 @@ public:
 		return TutorialAssetName;
 	}
 
+	FTutorialToken(FPrivateToken, const FString& InTutorialAssetName)
+		: FTutorialToken(InTutorialAssetName)
+	{
+	}
+
 protected:
 	/** Protected constructor */
 	FTutorialToken( const FString& InTutorialAssetName )
@@ -600,6 +633,9 @@ private:
  */
 class FActorToken : public IMessageToken
 {
+	// The private token allows only members or friends to call MakeShared.
+	struct FPrivateToken { explicit FPrivateToken() = default; };
+
 public:
 	/** Factory method, tokens can only be constructed as shared refs */
 	CORE_API static TSharedRef<FActorToken> Create(const FString& InActorPath, const FGuid& InActorGuid, const FText& InMessage = FText());
@@ -626,11 +662,17 @@ public:
 	}
 
 	/** Get the delegate for default token activation */
-	CORE_API static FOnMessageTokenActivated& DefaultOnMessageTokenActivated()
+	static FOnMessageTokenActivated& DefaultOnMessageTokenActivated()
 	{
 		return DefaultMessageTokenActivated;
 	}
 
+	/** Private constructor */
+	inline FActorToken(FPrivateToken, const FString& InActorPath, const FGuid& InActorGuid, const FText& InMessage)
+		: FActorToken(InActorPath, InActorGuid, InMessage)
+	{
+	}
+
 private:
 	/** Private constructor */
 	FActorToken(const FString& InActorPath, const FGuid& InActorGuid, const FText& InMessage);
