Avoid calling GLog->TearDown() from RequestExit as it may deadlock

#rnx

diff --git a/Engine/Source/Runtime/Core/Private/Mac/MacPlatformMisc.cpp b/Engine/Source/Runtime/Core/Private/Mac/MacPlatformMisc.cpp
--- a/Engine/Source/Runtime/Core/Private/Mac/MacPlatformMisc.cpp
+++ b/Engine/Source/Runtime/Core/Private/Mac/MacPlatformMisc.cpp
@@ -708,9 +708,7 @@ void FMacPlatformMisc::RequestExit( bool Force )
 		// Make sure the log is flushed.
 		if (GLog)
 		{
-			// This may be called from other thread, so set this thread as the master.
-			GLog->SetCurrentThreadAsMasterThread();
-			GLog->TearDown();
+			GLog->Flush();
 		}
 
 
diff --git a/Engine/Source/Runtime/Core/Private/Unix/UnixPlatformMisc.cpp b/Engine/Source/Runtime/Core/Private/Unix/UnixPlatformMisc.cpp
--- a/Engine/Source/Runtime/Core/Private/Unix/UnixPlatformMisc.cpp
+++ b/Engine/Source/Runtime/Core/Private/Unix/UnixPlatformMisc.cpp
@@ -379,9 +379,7 @@ void FUnixPlatformMisc::RequestExit(bool Force)
 		// Make sure the log is flushed.
 		if (!GEnteredSignalHandler && GLog)
 		{
-			// This may be called from other thread, so set this thread as the master.
-			GLog->SetCurrentThreadAsMasterThread();
-			GLog->TearDown();
+			GLog->Flush();
 		}
 
 		// Force immediate exit. Cannot call abort() here, because abort() raises SIGABRT which we treat as crash
diff --git a/Engine/Source/Runtime/Core/Private/Windows/WindowsPlatformMisc.cpp b/Engine/Source/Runtime/Core/Private/Windows/WindowsPlatformMisc.cpp
--- a/Engine/Source/Runtime/Core/Private/Windows/WindowsPlatformMisc.cpp
+++ b/Engine/Source/Runtime/Core/Private/Windows/WindowsPlatformMisc.cpp
@@ -951,9 +951,7 @@ void FWindowsPlatformMisc::RequestExitWithStatus(bool Force, uint8 ReturnCode)
 		// Make sure the log is flushed.
 		if (GLog)
 		{
-			// This may be called from other thread, so set this thread as the master.
-			GLog->SetCurrentThreadAsMasterThread();
-			GLog->TearDown();
+			GLog->Flush();
 		}
 
 		TerminateProcess(GetCurrentProcess(), ReturnCode);
