Make TCoroFrame movable only

#rnx

diff --git a/Engine/Source/Runtime/Core/Public/Experimental/Coroutine/Coroutine.h b/Engine/Source/Runtime/Core/Public/Experimental/Coroutine/Coroutine.h
--- a/Engine/Source/Runtime/Core/Public/Experimental/Coroutine/Coroutine.h
+++ b/Engine/Source/Runtime/Core/Public/Experimental/Coroutine/Coroutine.h
@@ -872,14 +872,34 @@ public:
 	using promise_type = CoroTask_Detail::TFramePromise<ReturnType>;
 	coroutine_handle<promise_type> CoroutineHandle;
 
-	COROFORCEINLINE TCoroFrame(coroutine_handle<promise_type> InCoroutineHandle) : CoroutineHandle(InCoroutineHandle)
+	// Non-copyable
+	TCoroFrame(const TCoroFrame&) = delete;
+	TCoroFrame& operator=(const TCoroFrame&) = delete;
+
+	COROFORCEINLINE TCoroFrame(coroutine_handle<promise_type> InCoroutineHandle)
+		: CoroutineHandle(InCoroutineHandle)
+	{
+	}
+
+	/**
+	 * Move constructor
+	 */
+	COROFORCEINLINE TCoroFrame(TCoroFrame&& Other)
+		: CoroutineHandle(MoveTemp(Other.CoroutineHandle))
 	{
+		// Make sure we explicitly set the moved coroutine to null as we might
+		// rely on different implementations which might not garantee this behavior.
+		Other.CoroutineHandle = nullptr;
 	}
 
 	COROFORCEINLINE ~TCoroFrame()
 	{
-		coroCheck(CoroutineHandle.done());
-		CoroutineHandle.destroy();
+		// Do not try to destroy a coroframe that has been moved.
+		if (CoroutineHandle)
+		{
+			coroCheck(CoroutineHandle.done());
+			CoroutineHandle.destroy();
+		}
 	}
 
 	COROFORCEINLINE auto GetResult()
