Some TArray comment fixes from this PR:

https://github.com/EpicGames/UnrealEngine/pull/8386

diff --git a/Engine/Source/Runtime/Core/Public/Containers/Array.h b/Engine/Source/Runtime/Core/Public/Containers/Array.h
--- a/Engine/Source/Runtime/Core/Public/Containers/Array.h
+++ b/Engine/Source/Runtime/Core/Public/Containers/Array.h
@@ -53,7 +53,7 @@ template <typename T, typename AllocatorType> inline void* operator new(size_t S
  * - A method T& operator\[\](SizeType index) which returns a reference to a contained object by index.
  * - A method void RemoveAt(SizeType index) which removes the element at index
  */
-template< typename ContainerType, typename ElementType, typename SizeType>
+template <typename ContainerType, typename ElementType, typename SizeType>
 class TIndexedContainerIterator
 {
 public:
@@ -290,7 +290,7 @@ namespace UE4Array_Private
 	};
 
 	// Assume elements are compatible with themselves - avoids problems with generated copy
-	// constuctors of arrays of forwarded types, e.g.:
+	// constructors of arrays of forwarded types, e.g.:
 	//
 	// struct FThing;
 	//
@@ -1266,7 +1266,7 @@ public:
 	 *   - it is safe to call BulkSerialize on TTransArrays
 	 *
 	 * IMPORTANT:
-	 *   - This is Overridden in XeD3dResourceArray.h Please make certain changes are propogated accordingly
+	 *   - This is Overridden in XeD3dResourceArray.h Please make certain changes are propagated accordingly
 	 *
 	 * @param Ar	FArchive to bulk serialize this TArray to/from
 	 */
@@ -1451,7 +1451,7 @@ public:
 	/**
 	 * Inserts given elements into the array at given location.
 	 *
-	 * @param Items Array of elements to insert.
+	 * @param InitList Array of elements to insert.
 	 * @param InIndex Tells where to insert the new elements.
 	 * @returns Location at which the item was inserted.
 	 */
@@ -1633,11 +1633,11 @@ private:
 			if (NumToMove)
 			{
 				FMemory::Memmove
-					(
+				(
 					(uint8*)AllocatorInstance.GetAllocation() + (Index)* sizeof(ElementType),
 					(uint8*)AllocatorInstance.GetAllocation() + (Index + Count) * sizeof(ElementType),
 					NumToMove * sizeof(ElementType)
-					);
+				);
 			}
 			ArrayNum -= Count;
 
@@ -1650,12 +1650,10 @@ private:
 
 public:
 	/**
-	 * Removes an element (or elements) at given location optionally shrinking
+	 * Removes an element (or elements) at given location, then shrinks
 	 * the array.
 	 *
 	 * @param Index Location in array of the element to remove.
-	 * @param Count (Optional) Number of elements to remove. Default is 1.
-	 * @param bAllowShrinking (Optional) Tells if this call can shrink array if suitable after remove. Default is true.
 	 */
 	FORCEINLINE void RemoveAt(SizeType Index)
 	{
@@ -1663,7 +1661,7 @@ public:
 	}
 
 	/**
-	 * Removes an element (or elements) at given location optionally shrinking
+	 * Removes an element (or elements) at given location, optionally shrinking
 	 * the array.
 	 *
 	 * @param Index Location in array of the element to remove.
@@ -1694,7 +1692,7 @@ private:
 			if (NumElementsToMoveIntoHole)
 			{
 				FMemory::Memcpy(
-					(uint8*)AllocatorInstance.GetAllocation() + (Index)* sizeof(ElementType),
+					(uint8*)AllocatorInstance.GetAllocation() + (Index) * sizeof(ElementType),
 					(uint8*)AllocatorInstance.GetAllocation() + (ArrayNum - NumElementsToMoveIntoHole) * sizeof(ElementType),
 					NumElementsToMoveIntoHole * sizeof(ElementType)
 					);
@@ -1710,16 +1708,13 @@ private:
 
 public:
 	/**
-	 * Removes an element (or elements) at given location optionally shrinking
+	 * Removes an element (or elements) at given location, then shrinks
 	 * the array.
 	 *
 	 * This version is much more efficient than RemoveAt (O(Count) instead of
 	 * O(ArrayNum)), but does not preserve the order.
 	 *
 	 * @param Index Location in array of the element to remove.
-	 * @param Count (Optional) Number of elements to remove. Default is 1.
-	 * @param bAllowShrinking (Optional) Tells if this call can shrink array if
-	 *                        suitable after remove. Default is true.
 	 */
 	FORCEINLINE void RemoveAtSwap(SizeType Index)
 	{
@@ -1727,7 +1722,7 @@ public:
 	}
 
 	/**
-	 * Removes an element (or elements) at given location optionally shrinking
+	 * Removes an element (or elements) at given location, optionally shrinking
 	 * the array.
 	 *
 	 * This version is much more efficient than RemoveAt (O(Count) instead of
@@ -1804,9 +1799,11 @@ public:
 	}
 
 	/**
-	 * Resizes array to given number of elements. New elements will be zeroed.
+	 * Resizes array to given number of elements, optionally shrinking it.
+	 * New elements will be zeroed.
 	 *
 	 * @param NewNum New size of the array.
+	 * @param bAllowShrinking Tell if this function can shrink the memory in-use if suitable.
 	 */
 	void SetNumZeroed(SizeType NewNum, bool bAllowShrinking = true)
 	{
@@ -2312,7 +2309,7 @@ public:
 	 *
 	 * Move semantics version.
 	 *
-	 * @param Args Item to add.
+	 * @param Item Item to add.
 	 * @returns Index of the element in the array.
 	 * @see Add, AddDefaulted, AddZeroed, Append, Insert
 	 */
@@ -2321,7 +2318,7 @@ public:
 	/**
 	 * Adds unique element to array if it doesn't exist.
 	 *
-	 * @param Args Item to add.
+	 * @param Item Item to add.
 	 * @returns Index of the element in the array.
 	 * @see Add, AddDefaulted, AddZeroed, Append, Insert
 	 */
@@ -2555,7 +2552,7 @@ public:
 			(uint8*)AllocatorInstance.GetAllocation() + (sizeof(ElementType)*FirstIndexToSwap),
 			(uint8*)AllocatorInstance.GetAllocation() + (sizeof(ElementType)*SecondIndexToSwap),
 			sizeof(ElementType)
-			);
+		);
 	}
 
 	/**
