Fix compile errors when Find() is used on a const TRobinHoodHashMap/Set, add basic automated test

diff --git a/Engine/Source/Runtime/Core/Private/Tests/Containers/RobinHoodHashTableTest.cpp b/Engine/Source/Runtime/Core/Private/Tests/Containers/RobinHoodHashTableTest.cpp
new file mode 100644
--- /dev/null
+++ b/Engine/Source/Runtime/Core/Private/Tests/Containers/RobinHoodHashTableTest.cpp
@@ -0,0 +1,46 @@
+// Copyright Epic Games, Inc. All Rights Reserved.
+
+#include "Misc/AutomationTest.h"
+
+#if WITH_DEV_AUTOMATION_TESTS
+
+#include "Experimental/Containers/RobinHoodHashTable.h"
+
+IMPLEMENT_SIMPLE_AUTOMATION_TEST(FRobinHoodHashTableTest, "System.Core.Containers.RobinHoodShashTable.Basic", EAutomationTestFlags::ApplicationContextMask | EAutomationTestFlags::SmokeFilter)
+bool FRobinHoodHashTableTest::RunTest(const FString& Parameters)
+{
+	Experimental::TRobinHoodHashMap<int32, int32> Table;
+	static constexpr int32 N = 1999; // Some reasonable prime number of elements
+
+	// Insert some unique values into the table
+	for (int32 Index = 0; Index < N; ++Index)
+	{
+		bool bAlreadyExists = false;
+		Table.FindOrAdd(Index, Index, bAlreadyExists);
+		TestEqual(TEXT("bAlreadyExists"), bAlreadyExists, false);
+	}
+
+	// Confirm that everything was added
+	for (int32 Index = 0; Index < N; ++Index)
+	{
+		bool bAlreadyExists = false;
+		Table.FindOrAdd(Index, Index, bAlreadyExists);
+		TestEqual(TEXT("bAlreadyExists"), bAlreadyExists, true);
+	}
+
+	// Check that everything can be accessed via const Find()
+	{
+		const Experimental::TRobinHoodHashMap<int32, int32>& TableConst = Table;
+		for (int32 Index = 0; Index < N; ++Index)
+		{
+			bool bAlreadyExists = false;
+			const int32* FoundPtr = TableConst.Find(Index);
+			TestNotNull(TEXT("FoundPtr"), FoundPtr);
+			TestEqual(TEXT("*FoundPtr"), *FoundPtr, Index);
+		}
+	}
+
+	return true;
+}
+
+#endif // WITH_DEV_AUTOMATION_TESTS
diff --git a/Engine/Source/Runtime/Core/Public/Experimental/Containers/RobinHoodHashTable.h b/Engine/Source/Runtime/Core/Public/Experimental/Containers/RobinHoodHashTable.h
--- a/Engine/Source/Runtime/Core/Public/Experimental/Containers/RobinHoodHashTable.h
+++ b/Engine/Source/Runtime/Core/Public/Experimental/Containers/RobinHoodHashTable.h
@@ -230,6 +230,7 @@ namespace RobinHoodHashTable_Private
 	class TKeyValue
 	{
 		using FindValueType = ValueType*;
+		using FindValueTypeConst = const ValueType *;
 		using ElementType = TPair<const KeyType, ValueType>;
 
 		template<typename, typename, typename, typename>
@@ -249,6 +250,11 @@ namespace RobinHoodHashTable_Private
 			return &Pair.Value;
 		}
 
+		inline FindValueTypeConst FindImpl() const
+		{
+			return &Pair.Value;
+		}
+
 		template<typename DeducedValueType>
 		inline void Update(DeducedValueType&& InValue)
 		{
@@ -275,6 +281,7 @@ namespace RobinHoodHashTable_Private
 	class TKeyValue<KeyType, FUnitType>
 	{
 		using FindValueType = const KeyType*;
+		using FindValueTypeConst = FindValueType;
 		using ElementType = const KeyType;
 
 		template<typename, typename, typename, typename>
@@ -311,6 +318,7 @@ namespace RobinHoodHashTable_Private
 		using InlineOneAllocatorType = TInlineAllocator<1, HashMapAllocator>;
 		using KeyValueType = RobinHoodHashTable_Private::TKeyValue<KeyType, ValueType>;
 		using FindValueType = typename KeyValueType::FindValueType;
+		using FindValueTypeConst = typename KeyValueType::FindValueTypeConst;
 		using ElementType = typename KeyValueType::ElementType;
 		using IndexType = uint32;
 		using SizeType = SIZE_T;
@@ -929,7 +937,7 @@ namespace RobinHoodHashTable_Private
 			return nullptr;
 		}
 
-		const FindValueType Find(const KeyType& Key) const
+		FindValueTypeConst Find(const KeyType& Key) const
 		{
 			FHashElementId Id = FindId(Key);
 			if (Id.IsValid())
@@ -1128,6 +1136,7 @@ class TRobinHoodHashMap : public RobinHoodHashTable_Private::TRobinHoodHashTable
 	using Base = typename RobinHoodHashTable_Private::TRobinHoodHashTable<KeyType, ValueType, Hasher, HashMapAllocator>;
 	using IndexType = typename Base::IndexType;
 	using FindValueType = typename Base::FindValueType;
+	using FindValueTypeConst = typename Base::FindValueTypeConst;
 
 public:
 	TRobinHoodHashMap() : Base()
@@ -1412,6 +1421,7 @@ class TRobinHoodHashSet : public RobinHoodHashTable_Private::TRobinHoodHashTable
 	using Base = typename RobinHoodHashTable_Private::TRobinHoodHashTable<KeyType, Unit, Hasher, HashMapAllocator>;
 	using IndexType = typename Base::IndexType;
 	using FindValueType = typename Base::FindValueType;
+	using FindValueTypeConst = typename Base::FindValueTypeConst;
 
 public:
 	TRobinHoodHashSet() : Base()
