Reduce use of UE_DECLARE_LWC_TYPE macro to improve Intellisense.
Moved Math forward declarations from CoreFwd.h to Math/MathFwd.h

diff --git a/Engine/Plugins/Animation/LiveLink/Source/LiveLinkComponents/Public/Controllers/LiveLinkTransformController.h b/Engine/Plugins/Animation/LiveLink/Source/LiveLinkComponents/Public/Controllers/LiveLinkTransformController.h
--- a/Engine/Plugins/Animation/LiveLink/Source/LiveLinkComponents/Public/Controllers/LiveLinkTransformController.h
+++ b/Engine/Plugins/Animation/LiveLink/Source/LiveLinkComponents/Public/Controllers/LiveLinkTransformController.h
@@ -2,12 +2,12 @@
 
 #pragma once
 
+#include "Math/MathFwd.h"
 #include "LiveLinkControllerBase.h"
 #include "Engine/EngineTypes.h"
 #include "LiveLinkTransformController.generated.h"
 
 struct FLiveLinkTransformStaticData;
-UE_DECLARE_LWC_TYPE(Transform, 3);
 class USceneComponent;
 
 USTRUCT(BlueprintType)
diff --git a/Engine/Plugins/Runtime/MeshModelingToolset/Source/ModelingComponents/Public/Mechanics/CubeGrid.h b/Engine/Plugins/Runtime/MeshModelingToolset/Source/ModelingComponents/Public/Mechanics/CubeGrid.h
--- a/Engine/Plugins/Runtime/MeshModelingToolset/Source/ModelingComponents/Public/Mechanics/CubeGrid.h
+++ b/Engine/Plugins/Runtime/MeshModelingToolset/Source/ModelingComponents/Public/Mechanics/CubeGrid.h
@@ -8,8 +8,6 @@
 #include "FrameTypes.h"
 #include "InteractionMechanic.h"
 
-UE_DECLARE_LWC_TYPE(Ray, 3);
-
 namespace UE {
 namespace Geometry {
 
diff --git a/Engine/Source/Editor/UnrealEd/Public/Tools/LegacyEdModeWidgetHelpers.h b/Engine/Source/Editor/UnrealEd/Public/Tools/LegacyEdModeWidgetHelpers.h
--- a/Engine/Source/Editor/UnrealEd/Public/Tools/LegacyEdModeWidgetHelpers.h
+++ b/Engine/Source/Editor/UnrealEd/Public/Tools/LegacyEdModeWidgetHelpers.h
@@ -2,8 +2,9 @@
 
 #pragma once
 
-#include "Tools/LegacyEdModeInterfaces.h"
+#include "Math/MathFwd.h"
 #include "Math/Vector.h"
+#include "Tools/LegacyEdModeInterfaces.h"
 #include "Tools/UEdMode.h"
 
 #include "LegacyEdModeWidgetHelpers.generated.h"
@@ -16,7 +17,6 @@ class FPrimitiveDrawInterface;
 class FSceneView;
 class FViewport;
 struct FConvexVolume;
-UE_DECLARE_LWC_TYPE(Matrix, 44);
 
 enum EModeTools : int8;
 class HHitProxy;
diff --git a/Engine/Source/Programs/IncludeTool/IncludeTool/Rules.cs b/Engine/Source/Programs/IncludeTool/IncludeTool/Rules.cs
--- a/Engine/Source/Programs/IncludeTool/IncludeTool/Rules.cs
+++ b/Engine/Source/Programs/IncludeTool/IncludeTool/Rules.cs
@@ -419,8 +419,7 @@ namespace IncludeTool
 			"/Engine/Source/Runtime/SlateCore/Public/Fonts/ShapedTextFwd.h", // Typedef isn't a forward declaration
 			"/Engine/Source/Runtime/Slate/Public/Framework/Text/ShapedTextCacheFwd.h", // Typedef isn't a forward declaration
 			"/Engine/Source/Runtime/MovieScene/Public/MovieSceneFwd.h",
-			"/Engine/Source/Runtime/Core/Public/CoreFwd.h",	// invalid forward declaration - UE_DECLARE_LWC_TYPE macro (emitting typedefs and namespaces)
-			"/Engine/Source/Runtime/Core/Public/Math/MathFwd.h", // invalid forward declaration - UE_DECLARE_LWC_TYPE macro (emitting typedefs and namespaces)
+			"/Engine/Source/Runtime/Core/Public/Math/MathFwd.h", // invalid forward declaration - 'namespace UE::Math'
 			"/Engine/Source/Runtime/Core/Public/Containers/ContainersFwd.h", // invalid forward declaration - 'template<> struct TIsContiguousContainer<Type> { static constexpr bool Value = true; };'
 			"/Engine/Source/Runtime/Core/Public/Containers/StringFwd.h", // invalid forward declaration - 'template<> struct TIsContiguousContainer<Type> { static constexpr bool Value = true; };'
 			"/Engine/Source/Runtime/Core/Public/Internationalization/StringTableCoreFwd.h", // Typedef isn't a forward declaration
diff --git a/Engine/Source/Runtime/ApplicationCore/Public/GenericPlatform/GenericPlatformApplicationMisc.h b/Engine/Source/Runtime/ApplicationCore/Public/GenericPlatform/GenericPlatformApplicationMisc.h
--- a/Engine/Source/Runtime/ApplicationCore/Public/GenericPlatform/GenericPlatformApplicationMisc.h
+++ b/Engine/Source/Runtime/ApplicationCore/Public/GenericPlatform/GenericPlatformApplicationMisc.h
@@ -6,7 +6,6 @@
 #include "GenericPlatform/GenericApplication.h"
 #include "HAL/IConsoleManager.h"
 
-UE_DECLARE_LWC_TYPE(Vector2,,FVector2D);
 class IPlatformInputDeviceMapper;
 
 /**
diff --git a/Engine/Source/Runtime/AugmentedReality/Private/ARDebugDrawHelpers.h b/Engine/Source/Runtime/AugmentedReality/Private/ARDebugDrawHelpers.h
--- a/Engine/Source/Runtime/AugmentedReality/Private/ARDebugDrawHelpers.h
+++ b/Engine/Source/Runtime/AugmentedReality/Private/ARDebugDrawHelpers.h
@@ -2,10 +2,9 @@
 
 #pragma once
 
-#include "Misc/LargeWorldCoordinates.h"
+#include "Math/MathFwd.h"
 
 class UWorld;
-UE_DECLARE_LWC_TYPE(Vector, 3);
 class FString;
 struct FColor;
 
diff --git a/Engine/Source/Runtime/Core/Public/CoreFwd.h b/Engine/Source/Runtime/Core/Public/CoreFwd.h
--- a/Engine/Source/Runtime/Core/Public/CoreFwd.h
+++ b/Engine/Source/Runtime/Core/Public/CoreFwd.h
@@ -4,6 +4,7 @@
 
 #include "CoreTypes.h"
 #include "Containers/ContainersFwd.h"
+#include "Math/MathFwd.h"
 #include "UObject/UObjectHierarchyFwd.h"
 
 // Basic types
@@ -15,69 +16,7 @@ class FFeedbackContext;
 struct FDateTime;
 struct FGuid;
 
-// Math
-UE_DECLARE_LWC_TYPE(Sphere, 3);
-UE_DECLARE_LWC_TYPE(Box2,, FBox2D);
-UE_DECLARE_LWC_TYPE(Box, 3);
-struct FColor;
-struct FLinearColor;
-UE_DECLARE_LWC_TYPE(Matrix, 44);
-UE_DECLARE_LWC_TYPE(Plane, 4);
-UE_DECLARE_LWC_TYPE(Quat, 4);
-UE_DECLARE_LWC_TYPE(Rotator, 3);
-UE_DECLARE_LWC_TYPE(Transform, 3);
-UE_DECLARE_LWC_TYPE(Vector2,, FVector2D);
-UE_DECLARE_LWC_TYPE(Vector, 3);
-UE_DECLARE_LWC_TYPE(Vector4);
-
-namespace UE
-{
-namespace Math
-{
-template <typename IntType>
-struct TIntPoint;
-
-template <typename IntType>
-struct TIntRect;
-
-template <typename IntType>
-struct TIntVector2;
-
-template <typename IntType>
-struct TIntVector3;
-
-template <typename IntType>
-struct TIntVector4;
-
-template<typename T, typename TExtent = T>
-struct TBoxSphereBounds;
-}
-}
-
-using FIntPoint = UE::Math::TIntPoint<int32>;
-using FUintPoint = UE::Math::TIntPoint<uint32>;
-
-using FIntRect = UE::Math::TIntRect<int32>;
-using FUintRect = UE::Math::TIntRect<uint32>;
-
-using FIntVector3 = UE::Math::TIntVector3<int32>;
-using FUintVector3 = UE::Math::TIntVector3<uint32>;
-
-using FIntVector = FIntVector3;
-
-using FIntVector2 = UE::Math::TIntVector2<int32>;
-using FUintVector2 = UE::Math::TIntVector2<uint32>;
-
-using FIntVector4 = UE::Math::TIntVector4<int32>;
-using FUintVector4 = UE::Math::TIntVector4<uint32>;
-
-using FBoxSphereBounds3f = UE::Math::TBoxSphereBounds<float, float>;
-using FBoxSphereBounds3d = UE::Math::TBoxSphereBounds<double, double>;
-// FCompactBoxSphereBounds always stores float extents
-using FCompactBoxSphereBounds3d = UE::Math::TBoxSphereBounds<double, float>;
-
-using FBoxSphereBounds = FBoxSphereBounds3d;
-using FCompactBoxSphereBounds = FCompactBoxSphereBounds3d;
+// Math - See Math/MathFwd.h
 
 // Misc
 struct FResourceSizeEx;
diff --git a/Engine/Source/Runtime/Core/Public/Math/MathFwd.h b/Engine/Source/Runtime/Core/Public/Math/MathFwd.h
--- a/Engine/Source/Runtime/Core/Public/Math/MathFwd.h
+++ b/Engine/Source/Runtime/Core/Public/Math/MathFwd.h
@@ -2,15 +2,96 @@
 
 #pragma once
 
-UE_DECLARE_LWC_TYPE(Sphere, 3);
-UE_DECLARE_LWC_TYPE(Box2,, FBox2D);
-UE_DECLARE_LWC_TYPE(Box, 3);
+#include "HAL/Platform.h"
+
+struct FColor;
+struct FLinearColor;
+
+// Non-macro declaration of the most commonly used types to assist Intellisense.
+namespace UE::Math
+{
+	// Forward declaration of templates
+	template<typename T> struct TVector;
+	template<typename T> struct TVector2;
+	template<typename T> struct TVector4;
+	template<typename T> struct TQuat;
+	template<typename T> struct TMatrix;
+	template<typename T> struct TPlane;
+	template<typename T> struct TTransform;
+	template<typename T> struct TSphere;
+	template<typename T> struct TBox;
+	template<typename T> struct TBox2;
+	template<typename T> struct TRotator;
+	template<typename T, typename TExtent = T> struct TBoxSphereBounds;
+	template <typename IntType> struct TIntPoint;
+	template <typename IntType> struct TIntRect;
+	template <typename IntType> struct TIntVector2;
+	template <typename IntType> struct TIntVector3;
+	template <typename IntType> struct TIntVector4;
+}
+
+// Forward declaration of concrete types				// Macro version - declares all three variants.
+using FVector 		= UE::Math::TVector<double>;		// UE_DECLARE_LWC_TYPE(Vector, 3);
+using FVector2D 	= UE::Math::TVector2<double>;		// UE_DECLARE_LWC_TYPE(Vector2,, FVector2D);
+using FVector4 		= UE::Math::TVector4<double>;		// UE_DECLARE_LWC_TYPE(Vector4);
+using FQuat 		= UE::Math::TQuat<double>;			// UE_DECLARE_LWC_TYPE(Quat, 4);
+using FMatrix 		= UE::Math::TMatrix<double>;		// UE_DECLARE_LWC_TYPE(Matrix, 44);
+using FPlane 		= UE::Math::TPlane<double>;			// UE_DECLARE_LWC_TYPE(Plane, 4);
+using FTransform 	= UE::Math::TTransform<double>;		// UE_DECLARE_LWC_TYPE(Transform, 3);
+using FSphere 		= UE::Math::TSphere<double>;		// UE_DECLARE_LWC_TYPE(Sphere, 3);
+using FBox 			= UE::Math::TBox<double>;			// UE_DECLARE_LWC_TYPE(Box, 3);
+using FBox2D 		= UE::Math::TBox2<double>;			// UE_DECLARE_LWC_TYPE(Box2,, FBox2D);
+using FRotator 		= UE::Math::TRotator<double>;		// UE_DECLARE_LWC_TYPE(Rotator, 3);
+
+using FVector3d 	= UE::Math::TVector<double>;
+using FVector2d 	= UE::Math::TVector2<double>;
+using FVector4d 	= UE::Math::TVector4<double>;
+using FQuat4d 		= UE::Math::TQuat<double>;
+using FMatrix44d 	= UE::Math::TMatrix<double>;
+using FPlane4d 		= UE::Math::TPlane<double>;
+using FTransform3d 	= UE::Math::TTransform<double>;
+using FSphere3d 	= UE::Math::TSphere<double>;
+using FBox3d 		= UE::Math::TBox<double>;
+using FBox2d 		= UE::Math::TBox2<double>;
+using FRotator3d 	= UE::Math::TRotator<double>;
+
+using FVector3f		= UE::Math::TVector<float>;
+using FVector2f		= UE::Math::TVector2<float>;
+using FVector4f		= UE::Math::TVector4<float>;
+using FQuat4f		= UE::Math::TQuat<float>;
+using FMatrix44f	= UE::Math::TMatrix<float>;
+using FPlane4f		= UE::Math::TPlane<float>;
+using FTransform3f	= UE::Math::TTransform<float>;
+using FSphere3f		= UE::Math::TSphere<float>;
+using FBox3f		= UE::Math::TBox<float>;
+using FBox2f		= UE::Math::TBox2<float>;
+using FRotator3f	= UE::Math::TRotator<float>;
+
+
 struct FColor;
 struct FLinearColor;
-UE_DECLARE_LWC_TYPE(Matrix, 44);
-UE_DECLARE_LWC_TYPE(Plane, 4);
-UE_DECLARE_LWC_TYPE(Quat, 4);
-UE_DECLARE_LWC_TYPE(Rotator, 3);
-UE_DECLARE_LWC_TYPE(Transform, 3);
-UE_DECLARE_LWC_TYPE(Vector, 3);
-UE_DECLARE_LWC_TYPE(Vector4);
+
+using FIntPoint = UE::Math::TIntPoint<int32>;
+using FUintPoint = UE::Math::TIntPoint<uint32>;
+
+using FIntRect = UE::Math::TIntRect<int32>;
+using FUintRect = UE::Math::TIntRect<uint32>;
+
+using FIntVector3 = UE::Math::TIntVector3<int32>;
+using FUintVector3 = UE::Math::TIntVector3<uint32>;
+
+using FIntVector = FIntVector3;
+
+using FIntVector2 = UE::Math::TIntVector2<int32>;
+using FUintVector2 = UE::Math::TIntVector2<uint32>;
+
+using FIntVector4 = UE::Math::TIntVector4<int32>;
+using FUintVector4 = UE::Math::TIntVector4<uint32>;
+
+using FBoxSphereBounds3f = UE::Math::TBoxSphereBounds<float, float>;
+using FBoxSphereBounds3d = UE::Math::TBoxSphereBounds<double, double>;
+// FCompactBoxSphereBounds always stores float extents
+using FCompactBoxSphereBounds3d = UE::Math::TBoxSphereBounds<double, float>;
+
+using FBoxSphereBounds = FBoxSphereBounds3d;
+using FCompactBoxSphereBounds = FCompactBoxSphereBounds3d;
\ No newline at end of file
diff --git a/Engine/Source/Runtime/Core/Public/Math/UnrealMathUtility.h b/Engine/Source/Runtime/Core/Public/Math/UnrealMathUtility.h
--- a/Engine/Source/Runtime/Core/Public/Math/UnrealMathUtility.h
+++ b/Engine/Source/Runtime/Core/Public/Math/UnrealMathUtility.h
@@ -5,6 +5,7 @@
 #include "CoreTypes.h"
 #include "Misc/AssertionMacros.h"
 #include "HAL/PlatformMath.h"
+#include "Math/MathFwd.h"
 #include "Templates/Decay.h"
 #include "Templates/IsFloatingPoint.h"
 #include "Templates/IsIntegral.h"
@@ -23,18 +24,7 @@
 -----------------------------------------------------------------------------*/
 
 // Forward declarations.
-UE_DECLARE_LWC_TYPE(Vector2,, FVector2D);
-UE_DECLARE_LWC_TYPE(Vector, 3);
-UE_DECLARE_LWC_TYPE(Vector4);
-UE_DECLARE_LWC_TYPE(Plane, 4);
-UE_DECLARE_LWC_TYPE(Box, 3);
-UE_DECLARE_LWC_TYPE(Rotator, 3);
-UE_DECLARE_LWC_TYPE(Matrix, 44);
-UE_DECLARE_LWC_TYPE(Quat, 4);
 struct  FTwoVectors;
-UE_DECLARE_LWC_TYPE(Transform, 3);
-UE_DECLARE_LWC_TYPE(Sphere, 3);
-
 struct FLinearColor;
 template<typename ElementType>
 class TRange;
diff --git a/Engine/Source/Runtime/Core/Public/Math/Vector.h b/Engine/Source/Runtime/Core/Public/Math/Vector.h
--- a/Engine/Source/Runtime/Core/Public/Math/Vector.h
+++ b/Engine/Source/Runtime/Core/Public/Math/Vector.h
@@ -4,6 +4,7 @@
 
 #include "CoreTypes.h"
 #include "Misc/AssertionMacros.h"
+#include "Math/MathFwd.h"
 #include "Math/NumericLimits.h"
 #include "Misc/Crc.h"
 #include "Math/UnrealMathUtility.h"
@@ -40,8 +41,6 @@ namespace UE
 namespace Math
 {
 
-template<typename T> struct TPlane;
-
 /**
  * A vector in 3-D space composed of components (X, Y, Z) with floating point precision.
  */
@@ -2470,8 +2469,6 @@ FORCEINLINE uint32 GetTypeHash(const TVector<T>& Vector)
 } // namespace UE::Math
 } // namespace UE
 
-UE_DECLARE_LWC_TYPE(Vector, 3);
-
 template<> struct TCanBulkSerialize<FVector3f> { enum { Value = true }; };
 template<> struct TIsPODType<FVector3f> { enum { Value = true }; };
 template<> struct TIsUECoreVariant<FVector3f> { enum { Value = true }; };
diff --git a/Engine/Source/Runtime/CoreUObject/Public/UObject/NoExportTypes.h b/Engine/Source/Runtime/CoreUObject/Public/UObject/NoExportTypes.h
--- a/Engine/Source/Runtime/CoreUObject/Public/UObject/NoExportTypes.h
+++ b/Engine/Source/Runtime/CoreUObject/Public/UObject/NoExportTypes.h
@@ -6,6 +6,11 @@
 
 #pragma once
 
+// Help intellisense to avoid interpreting this file's declaration of FVector etc as it assumes !CPP by default
+#ifndef CPP
+#define CPP 1
+#endif
+
 #if CPP
 
 // Include the real definitions of the noexport classes below to allow the generated cpp file to compile.
diff --git a/Engine/Source/Runtime/Engine/Public/Animation/AnimTrace.h b/Engine/Source/Runtime/Engine/Public/Animation/AnimTrace.h
--- a/Engine/Source/Runtime/Engine/Public/Animation/AnimTrace.h
+++ b/Engine/Source/Runtime/Engine/Public/Animation/AnimTrace.h
@@ -29,7 +29,6 @@ struct FAnimationCacheBonesContext;
 struct FPoseContext;
 struct FComponentSpacePoseContext;
 class FName;
-UE_DECLARE_LWC_TYPE(Rotator, 3);
 struct FAnimNode_SequencePlayerBase;
 struct FAnimNotifyEvent;
 struct FPassedMarker;
diff --git a/Engine/Source/Runtime/Engine/Public/AudioDevice.h b/Engine/Source/Runtime/Engine/Public/AudioDevice.h
--- a/Engine/Source/Runtime/Engine/Public/AudioDevice.h
+++ b/Engine/Source/Runtime/Engine/Public/AudioDevice.h
@@ -60,7 +60,6 @@ struct FActiveSound;
 struct FAttenuationFocusData;
 struct FAudioComponentParam;
 struct FAudioQualitySettings;
-UE_DECLARE_LWC_TYPE(Rotator, 3);
 struct FWaveInstance;
 
 LLM_DECLARE_TAG_API(AudioSpatializationPlugins, ENGINE_API);
diff --git a/Engine/Source/Runtime/Experimental/Chaos/Public/Chaos/Rotation.h b/Engine/Source/Runtime/Experimental/Chaos/Public/Chaos/Rotation.h
--- a/Engine/Source/Runtime/Experimental/Chaos/Public/Chaos/Rotation.h
+++ b/Engine/Source/Runtime/Experimental/Chaos/Public/Chaos/Rotation.h
@@ -12,7 +12,7 @@
 #else
 #include <array>
 
-struct FQuat
+struct _FQuat
 {
 public:
 	const Chaos::FReal operator[](const int32 i) const
@@ -26,11 +26,12 @@ public:
 	std::array<Chaos::FReal, 3> angles;
 	static MakeFromEuler(const Vector<Chaos::FReal, 3>& InAngles)
 	{
-		FQuat Quat;
+		_FQuat Quat;
 		Quat.angles = InAngles;
 		return Quat;
 	}
 };
+using FQuat = _FQuat;	// Work around include tool not understanding that this won't be compiled alongside MathFwd.h
 #endif
 
 namespace Chaos
diff --git a/Engine/Source/Runtime/Experimental/Chaos/Public/Chaos/Transform.h b/Engine/Source/Runtime/Experimental/Chaos/Public/Chaos/Transform.h
--- a/Engine/Source/Runtime/Experimental/Chaos/Public/Chaos/Transform.h
+++ b/Engine/Source/Runtime/Experimental/Chaos/Public/Chaos/Transform.h
@@ -10,16 +10,15 @@
 #include "Math/Transform.h"
 #else
 //TODO(mlentine): If we use this class in engine we need to make it more efficient.
-//TODO(mlentine): This should really be a class as there is a lot of functionality but static anlysis current forbids this.
-struct FTransform
+struct _FTransform
 {
 public:
-	FTransform() {}
-	FTransform(const Chaos::TRotation<Chaos::FReal, 3>& Rotation, const Chaos::TVector<Chaos::FReal, 3>& Translation)
+	_FTransform() {}
+	_FTransform(const Chaos::TRotation<Chaos::FReal, 3>& Rotation, const Chaos::TVector<Chaos::FReal, 3>& Translation)
 		: MRotation(Rotation), MTranslation(Translation)
 	{
 	}
-	FTransform(const FMatrix& Matrix)
+	_FTransform(const FMatrix& Matrix)
 	{
 		MTranslation[0] = Matrix.M[0][3];
 		MTranslation[1] = Matrix.M[1][3];
@@ -39,7 +38,7 @@ public:
 			MRotation[2] = 0;
 		}
 	}
-	FTransform(const FTransform& Transform)
+	_FTransform(const _FTransform& Transform)
 		: MRotation(Transform.MRotation), MTranslation(Transform.MTranslation)
 	{
 	}
@@ -106,6 +105,7 @@ private:
 	Chaos::TRotation<Chaos::FReal, 3> MRotation;
 	Chaos::TVector<Chaos::FReal, 3> MTranslation;
 };
+using FTransform = _FTransform;	// Work around include tool not understanding that this can't be compiled alongside MathFwd.h
 #endif
 
 namespace Chaos
diff --git a/Engine/Source/Runtime/Experimental/ChaosCore/Public/Chaos/Real.h b/Engine/Source/Runtime/Experimental/ChaosCore/Public/Chaos/Real.h
--- a/Engine/Source/Runtime/Experimental/ChaosCore/Public/Chaos/Real.h
+++ b/Engine/Source/Runtime/Experimental/ChaosCore/Public/Chaos/Real.h
@@ -3,7 +3,6 @@
 #pragma once
 
 #include "Templates/AreTypesEqual.h"
-#include "Misc/LargeWorldCoordinates.h"
 
 namespace Chaos
 {
