FTaskPriorityCVar - a console var to configure Tasks System priorities, by analogy with FAutoConsoleTaskPriority for TaskGraph

diff --git a/Engine/Source/Runtime/Core/Private/Tasks/TaskPrivate.cpp b/Engine/Source/Runtime/Core/Private/Tasks/TaskPrivate.cpp
--- a/Engine/Source/Runtime/Core/Private/Tasks/TaskPrivate.cpp
+++ b/Engine/Source/Runtime/Core/Private/Tasks/TaskPrivate.cpp
@@ -7,97 +7,216 @@
 #include "Async/TaskGraphInterfaces.h"
 #endif
 
-namespace UE::Tasks::Private
+namespace UE::Tasks
 {
-	void FTaskBase::Schedule()
+	namespace Private
 	{
-#if TASKGRAPH_NEW_FRONTEND
-		if (IsNamedThreadTask())
+		void FTaskBase::Schedule()
 		{
-			ENamedThreads::Type ConversionMap[] =
+#if TASKGRAPH_NEW_FRONTEND
+			if (IsNamedThreadTask())
 			{
-				ENamedThreads::GameThread,
-				(ENamedThreads::Type)(ENamedThreads::GameThread | ENamedThreads::HighTaskPriority),
-				(ENamedThreads::Type)(ENamedThreads::GameThread | ENamedThreads::LocalQueue),
-				(ENamedThreads::Type)(ENamedThreads::GameThread | ENamedThreads::HighTaskPriority | ENamedThreads::LocalQueue),
-
-				ENamedThreads::GetRenderThread(),
-				(ENamedThreads::Type)(ENamedThreads::GetRenderThread() | ENamedThreads::HighTaskPriority),
-				(ENamedThreads::Type)(ENamedThreads::GetRenderThread() | ENamedThreads::LocalQueue),
-				(ENamedThreads::Type)(ENamedThreads::GetRenderThread() | ENamedThreads::HighTaskPriority | ENamedThreads::LocalQueue),
-
-				ENamedThreads::RHIThread,
-				(ENamedThreads::Type)(ENamedThreads::RHIThread | ENamedThreads::HighTaskPriority),
-				(ENamedThreads::Type)(ENamedThreads::RHIThread | ENamedThreads::LocalQueue),
-				(ENamedThreads::Type)(ENamedThreads::RHIThread | ENamedThreads::HighTaskPriority | ENamedThreads::LocalQueue)
-			};
-
-			FTaskGraphInterface::Get().QueueTask(static_cast<FBaseGraphTask*>(this), true, ConversionMap[(int32)ExtendedPriority - (int32)EExtendedTaskPriority::GameThreadNormalPri]);
-			return;
+				ENamedThreads::Type ConversionMap[] =
+				{
+					ENamedThreads::GameThread,
+					(ENamedThreads::Type)(ENamedThreads::GameThread | ENamedThreads::HighTaskPriority),
+					(ENamedThreads::Type)(ENamedThreads::GameThread | ENamedThreads::LocalQueue),
+					(ENamedThreads::Type)(ENamedThreads::GameThread | ENamedThreads::HighTaskPriority | ENamedThreads::LocalQueue),
+
+					ENamedThreads::GetRenderThread(),
+					(ENamedThreads::Type)(ENamedThreads::GetRenderThread() | ENamedThreads::HighTaskPriority),
+					(ENamedThreads::Type)(ENamedThreads::GetRenderThread() | ENamedThreads::LocalQueue),
+					(ENamedThreads::Type)(ENamedThreads::GetRenderThread() | ENamedThreads::HighTaskPriority | ENamedThreads::LocalQueue),
+
+					ENamedThreads::RHIThread,
+					(ENamedThreads::Type)(ENamedThreads::RHIThread | ENamedThreads::HighTaskPriority),
+					(ENamedThreads::Type)(ENamedThreads::RHIThread | ENamedThreads::LocalQueue),
+					(ENamedThreads::Type)(ENamedThreads::RHIThread | ENamedThreads::HighTaskPriority | ENamedThreads::LocalQueue)
+				};
+
+				FTaskGraphInterface::Get().QueueTask(static_cast<FBaseGraphTask*>(this), true, ConversionMap[(int32)ExtendedPriority - (int32)EExtendedTaskPriority::GameThreadNormalPri]);
+				return;
+			}
+#endif
+
+			LowLevelTasks::FScheduler::Get().TryLaunch(LowLevelTask, LowLevelTasks::EQueuePreference::GlobalQueuePreference, /*bWakeUpWorker=*/ true);
 		}
+
+		FTaskBase* FTaskBase::TryPushIntoPipe()
+		{
+			return GetPipe()->PushIntoPipe(*this);
+		}
+
+		void FTaskBase::StartPipeExecution()
+		{
+			GetPipe()->ExecutionStarted();
+		}
+
+		void FTaskBase::FinishPipeExecution()
+		{
+			GetPipe()->ExecutionFinished();
+		}
+
+		void FTaskBase::ClearPipe()
+		{
+			GetPipe()->ClearTask(*this);
+		}
+
+		static thread_local FTaskBase* CurrentTask = nullptr;
+
+		FTaskBase* GetCurrentTask()
+		{
+			return CurrentTask;
+		}
+
+		FTaskBase* ExchangeCurrentTask(FTaskBase* Task)
+		{
+			FTaskBase* PrevTask = CurrentTask;
+			CurrentTask = Task;
+			return PrevTask;
+		}
+
+		bool TryWaitOnNamedThread(FTaskBase& Task)
+		{
+#if TASKGRAPH_NEW_FRONTEND
+			// handle waiting only on a named thread and if not called from inside a task
+			FTaskGraphInterface& TaskGraph = FTaskGraphInterface::Get();
+			ENamedThreads::Type CurrentThread = TaskGraph.GetCurrentThreadIfKnown();
+			if (CurrentThread < ENamedThreads::ActualRenderingThread /* is a named thread? */ && !TaskGraph.IsThreadProcessingTasks(CurrentThread))
+			{
+				// execute other tasks of this named thread while waiting
+				ETaskPriority Dummy;
+				EExtendedTaskPriority ExtendedPriority;
+				FBaseGraphTask::TranslatePriority(CurrentThread, Dummy, ExtendedPriority);
+
+				auto TaskBody = [CurrentThread, &TaskGraph] { TaskGraph.RequestReturn(CurrentThread); };
+				using FReturnFromNamedThreadTask = TExecutableTask<decltype(TaskBody)>;
+				FReturnFromNamedThreadTask ReturnTask{ TEXT("ReturnFromNamedThreadTask"), MoveTemp(TaskBody), ETaskPriority::High, ExtendedPriority };
+				ReturnTask.AddPrerequisites(Task);
+				ReturnTask.TryLaunch(); // the result doesn't matter
+
+				TaskGraph.ProcessThreadUntilRequestReturn(CurrentThread);
+				return true;
+			}
 #endif
 
-		LowLevelTasks::FScheduler::Get().TryLaunch(LowLevelTask, LowLevelTasks::EQueuePreference::GlobalQueuePreference, /*bWakeUpWorker=*/ true);
+			return false;
+		}
 	}
 
-	FTaskBase* FTaskBase::TryPushIntoPipe()
+	const TCHAR* ToString(EExtendedTaskPriority ExtendedPriority)
 	{
-		return GetPipe()->PushIntoPipe(*this);
-	}
+		if (ExtendedPriority < EExtendedTaskPriority::None || ExtendedPriority >= EExtendedTaskPriority::Count)
+		{
+			return nullptr;
+		}
 
-	void FTaskBase::StartPipeExecution()
-	{
-		GetPipe()->ExecutionStarted();
-	}
+		const TCHAR* ExtendedTaskPriorityToStr[] =
+		{
+			TEXT("None"),
+			TEXT("Inline"),
+			TEXT("TaskEvent"),
 
-	void FTaskBase::FinishPipeExecution()
-	{
-		GetPipe()->ExecutionFinished();
+#if TASKGRAPH_NEW_FRONTEND
+			TEXT("GameThreadNormalPri"),
+			TEXT("GameThreadHiPri"),
+			TEXT("GameThreadNormalPriLocalQueue"),
+			TEXT("GameThreadHiPriLocalQueue"),
+
+			TEXT("RenderThreadNormalPri"),
+			TEXT("RenderThreadHiPri"),
+			TEXT("RenderThreadNormalPriLocalQueue"),
+			TEXT("RenderThreadHiPriLocalQueue"),
+
+			TEXT("RHIThreadNormalPri"),
+			TEXT("RHIThreadHiPri"),
+			TEXT("RHIThreadNormalPriLocalQueue"),
+			TEXT("RHIThreadHiPriLocalQueue")
+#endif
+		};
+		return ExtendedTaskPriorityToStr[(int32)ExtendedPriority];
 	}
 
-	void FTaskBase::ClearPipe()
+	bool ToExtendedTaskPriority(const TCHAR* ExtendedPriorityStr, EExtendedTaskPriority& OutExtendedPriority)
 	{
-		GetPipe()->ClearTask(*this);
-	}
+#define CONVERT_EXTENDED_TASK_PRIORITY(ExtendedTaskPriority)\
+		if (FCString::Stricmp(ExtendedPriorityStr, ToString(EExtendedTaskPriority::ExtendedTaskPriority)) == 0)\
+		{\
+			OutExtendedPriority = EExtendedTaskPriority::ExtendedTaskPriority;\
+			return true;\
+		}
+
+		CONVERT_EXTENDED_TASK_PRIORITY(None);
+		CONVERT_EXTENDED_TASK_PRIORITY(Inline);
+		CONVERT_EXTENDED_TASK_PRIORITY(TaskEvent);
+
+#if TASKGRAPH_NEW_FRONTEND
+		CONVERT_EXTENDED_TASK_PRIORITY(GameThreadNormalPri);
+		CONVERT_EXTENDED_TASK_PRIORITY(GameThreadHiPri);
+		CONVERT_EXTENDED_TASK_PRIORITY(GameThreadNormalPriLocalQueue);
+		CONVERT_EXTENDED_TASK_PRIORITY(GameThreadHiPriLocalQueue);
+
+		CONVERT_EXTENDED_TASK_PRIORITY(RenderThreadNormalPri);
+		CONVERT_EXTENDED_TASK_PRIORITY(RenderThreadHiPri);
+		CONVERT_EXTENDED_TASK_PRIORITY(RenderThreadNormalPriLocalQueue);
+		CONVERT_EXTENDED_TASK_PRIORITY(RenderThreadHiPriLocalQueue);
+
+		CONVERT_EXTENDED_TASK_PRIORITY(RHIThreadNormalPri);
+		CONVERT_EXTENDED_TASK_PRIORITY(RHIThreadHiPri);
+		CONVERT_EXTENDED_TASK_PRIORITY(RHIThreadNormalPriLocalQueue);
+		CONVERT_EXTENDED_TASK_PRIORITY(RHIThreadHiPriLocalQueue);
+#endif
 
-	static thread_local FTaskBase* CurrentTask = nullptr;
+#undef CONVERT_EXTENDED_TASK_PRIORITY
 
-	FTaskBase* GetCurrentTask()
+		return false;
+	}
+
+	FString FTaskPriorityCVar::CreateFullHelpText(const TCHAR* Name, const TCHAR* OriginalHelp)
 	{
-		return CurrentTask;
+		TStringBuilder<1024> TaskPriorities;
+		for (int i = 0; i != (int)ETaskPriority::Count; ++i)
+		{
+			TaskPriorities.Append(ToString((ETaskPriority)i));
+			TaskPriorities.Append(TEXT(", "));
+		}
+		TaskPriorities.RemoveSuffix(2); // remove the last ", "
+
+		TStringBuilder<1024> ExtendedTaskPriorities;
+		for (int i = 0; i != (int)EExtendedTaskPriority::Count; ++i)
+		{
+			ExtendedTaskPriorities.Append(ToString((EExtendedTaskPriority)i));
+			ExtendedTaskPriorities.Append(TEXT(", "));
+		}
+		ExtendedTaskPriorities.RemoveSuffix(2); // remove the last ", "
+
+		return FString::Printf(
+			TEXT("%s\n")
+			TEXT("Arguments are task priority and extended task priority (optional) separated by a space: [TaskPriority] [ExtendedTaskPriority]\n")
+			TEXT("where TaskPriority is in [%s]\n")
+			TEXT("and ExtendedTaskPriority is in [%s].\n")
+			TEXT("Example: \"%s %s %s\" or \"%s\"")
+			, OriginalHelp, *TaskPriorities, *ExtendedTaskPriorities, Name, ToString((ETaskPriority)0), ToString((EExtendedTaskPriority)0), ToString((ETaskPriority)0));
 	}
 
-	FTaskBase* ExchangeCurrentTask(FTaskBase* Task)
+	FString FTaskPriorityCVar::ConfigStringFromPriorities(ETaskPriority InPriority, EExtendedTaskPriority InExtendedPriority)
 	{
-		FTaskBase* PrevTask = CurrentTask;
-		CurrentTask = Task;
-		return PrevTask;
+		return FString{ ToString(InPriority) } + TEXT(" ") + ToString(InExtendedPriority);
 	}
 
-	bool TryWaitOnNamedThread(FTaskBase& Task)
+	void FTaskPriorityCVar::OnSettingChanged(IConsoleVariable* InVariable)
 	{
-#if TASKGRAPH_NEW_FRONTEND
-		// handle waiting only on a named thread and if not called from inside a task
-		FTaskGraphInterface& TaskGraph = FTaskGraphInterface::Get();
-		ENamedThreads::Type CurrentThread = TaskGraph.GetCurrentThreadIfKnown();
-		if (CurrentThread < ENamedThreads::ActualRenderingThread /* is a named thread? */ && !TaskGraph.IsThreadProcessingTasks(CurrentThread))
+		FString PriorityStr, ExtendedPriorityStr;
+		static FString Delimiter{ " " };
+		if (RawSetting.Split(Delimiter, &PriorityStr, &ExtendedPriorityStr))
 		{
-			// execute other tasks of this named thread while waiting
-			ETaskPriority Dummy;
-			EExtendedTaskPriority ExtendedPriority;
-			FBaseGraphTask::TranslatePriority(CurrentThread, Dummy, ExtendedPriority);
-
-			auto TaskBody = [CurrentThread, &TaskGraph] { TaskGraph.RequestReturn(CurrentThread); };
-			using FReturnFromNamedThreadTask = TExecutableTask<decltype(TaskBody)>;
-			FReturnFromNamedThreadTask ReturnTask{ TEXT("ReturnFromNamedThreadTask"), MoveTemp(TaskBody), ETaskPriority::High, ExtendedPriority };
-			ReturnTask.AddPrerequisites(Task);
-			ReturnTask.TryLaunch(); // the result doesn't matter
-
-			TaskGraph.ProcessThreadUntilRequestReturn(CurrentThread);
-			return true;
+			verify(ToTaskPriority(*PriorityStr, Priority));
+			verify(ToExtendedTaskPriority(*ExtendedPriorityStr, ExtendedPriority));
+		}
+		else
+		{
+			verify(ToTaskPriority(*RawSetting, Priority));
+			ExtendedPriority = EExtendedTaskPriority::None;
 		}
-#endif
-
-		return false;
 	}
 }
diff --git a/Engine/Source/Runtime/Core/Private/Tests/Tasks/TasksTest.cpp b/Engine/Source/Runtime/Core/Private/Tests/Tasks/TasksTest.cpp
--- a/Engine/Source/Runtime/Core/Private/Tests/Tasks/TasksTest.cpp
+++ b/Engine/Source/Runtime/Core/Private/Tests/Tasks/TasksTest.cpp
@@ -1281,6 +1281,45 @@ namespace UE { namespace TasksTests
 
 		return true;
 	}
+
+	IMPLEMENT_SIMPLE_AUTOMATION_TEST(FTasksPriorityCVarTest, "System.Core.Tasks.PriorityCVar", EAutomationTestFlags::ApplicationContextMask | EAutomationTestFlags::EngineFilter);
+
+	bool FTasksPriorityCVarTest::RunTest(const FString& Parameters)
+	{
+		// set every combination of task priority and extended task priority
+		const TCHAR* CVarName = TEXT("TasksPriorityTestCVar");
+		UE::Tasks::FTaskPriorityCVar CVar{ CVarName, TEXT("CVarHelp"), ETaskPriority::Normal, EExtendedTaskPriority::None };
+		IConsoleVariable* CVarPtr = IConsoleManager::Get().FindConsoleVariable(CVarName);
+		for (int Priority = 0; Priority != (int)ETaskPriority::Count; ++Priority)
+		{
+			for (int ExtendedPriority = 0; ExtendedPriority != (int)EExtendedTaskPriority::Count; ++ExtendedPriority)
+			{
+				TStringBuilder<1024> TaskPriorities;
+				TaskPriorities.Append(ToString((ETaskPriority)Priority));
+				TaskPriorities.Append(TEXT(" "));
+				TaskPriorities.Append(ToString((EExtendedTaskPriority)ExtendedPriority));
+
+				CVarPtr->Set(*TaskPriorities);
+
+				check(CVar.GetTaskPriority() == (ETaskPriority)Priority);
+				check(CVar.GetExtendedTaskPriority() == (EExtendedTaskPriority)ExtendedPriority);
+				check(CVarPtr->GetString() == *TaskPriorities);
+			}
+		}
+
+		// test setting only task priority
+		CVarPtr->Set(TEXT("High"));
+		check(CVar.GetTaskPriority() == ETaskPriority::High);
+		check(CVar.GetExtendedTaskPriority() == EExtendedTaskPriority::None);
+		CVarPtr->Set(TEXT("Normal"));
+		check(CVar.GetTaskPriority() == ETaskPriority::Normal);
+
+		// usage example
+		CVarPtr->Set(TEXT("Normal"));
+		Launch(UE_SOURCE_LOCATION, [] {}, CVar.GetTaskPriority(), CVar.GetExtendedTaskPriority()).Wait();
+
+		return true;
+	}
 }}
 
 #endif // WITH_DEV_AUTOMATION_TESTS
diff --git a/Engine/Source/Runtime/Core/Public/Async/TaskGraphInterfaces.h b/Engine/Source/Runtime/Core/Public/Async/TaskGraphInterfaces.h
--- a/Engine/Source/Runtime/Core/Public/Async/TaskGraphInterfaces.h
+++ b/Engine/Source/Runtime/Core/Public/Async/TaskGraphInterfaces.h
@@ -606,7 +606,7 @@ public:
 	{
 		using namespace UE::Tasks;
 
-		checkf(Priority >= EExtendedTaskPriority::GameThreadNormalPri && Priority <= EExtendedTaskPriority::RHIThreadHiPriLocalQueue, TEXT("only named threads can call this method: %d"), Priority);
+		checkf(Priority >= EExtendedTaskPriority::GameThreadNormalPri && Priority < EExtendedTaskPriority::Count, TEXT("only named threads can call this method: %d"), Priority);
 
 		int32 ConversionMap[] =
 		{
diff --git a/Engine/Source/Runtime/Core/Public/Tasks/Task.h b/Engine/Source/Runtime/Core/Public/Tasks/Task.h
--- a/Engine/Source/Runtime/Core/Public/Tasks/Task.h
+++ b/Engine/Source/Runtime/Core/Public/Tasks/Task.h
@@ -472,6 +472,45 @@ namespace UE::Tasks
 		check(Parent != nullptr);
 		Parent->AddNested(*Nested.Pimpl);
 	}
+
+	// Console variable for configuring task priorities
+	// Example:
+	//		FTaskPriorityCVar CVar{ TEXT("CVarName"), TEXT("CVarHelp"), ETaskPriority::Normal, EExtendedTaskPriority::None };
+	//		Launch(UE_SOURCE_LOCATION, [] {}, CVar.GetTaskPriority(), CVar.GetExtendedTaskPriority()).Wait();
+	class FTaskPriorityCVar
+	{
+	public:
+		FTaskPriorityCVar(const TCHAR* Name, const TCHAR* Help, ETaskPriority DefaultPriority, EExtendedTaskPriority DefaultExtendedPriority)
+			: RawSetting(ConfigStringFromPriorities(DefaultPriority, DefaultExtendedPriority))
+			, FullHelpText(CreateFullHelpText(Name, Help))
+			, Variable(Name, RawSetting, *FullHelpText, FConsoleVariableDelegate::CreateRaw(this, &FTaskPriorityCVar::OnSettingChanged), ECVF_Default)
+			, Priority(DefaultPriority)
+			, ExtendedPriority(DefaultExtendedPriority)
+		{
+		}
+
+		ETaskPriority GetTaskPriority() const
+		{
+			return Priority;
+		}
+
+		EExtendedTaskPriority GetExtendedTaskPriority() const
+		{
+			return ExtendedPriority;
+		}
+
+	private:
+		static FString CreateFullHelpText(const TCHAR* Name, const TCHAR* OriginalHelp);
+		static FString ConfigStringFromPriorities(ETaskPriority InPriority, EExtendedTaskPriority InExtendedPriority);
+		void OnSettingChanged(IConsoleVariable* Variable);
+
+	private:
+		FString RawSetting;
+		FString FullHelpText;
+		FAutoConsoleVariableRef Variable;
+		ETaskPriority Priority;
+		EExtendedTaskPriority ExtendedPriority;
+	};
 }
 
 template <typename... TaskTypes>
diff --git a/Engine/Source/Runtime/Core/Public/Tasks/TaskPrivate.h b/Engine/Source/Runtime/Core/Public/Tasks/TaskPrivate.h
--- a/Engine/Source/Runtime/Core/Public/Tasks/TaskPrivate.h
+++ b/Engine/Source/Runtime/Core/Public/Tasks/TaskPrivate.h
@@ -27,6 +27,8 @@
 namespace UE::Tasks
 {
 	using LowLevelTasks::ETaskPriority;
+	using LowLevelTasks::ToString;
+	using LowLevelTasks::ToTaskPriority;
 
 	// special task priorities for tasks that are never sent to the scheduler
 	enum class EExtendedTaskPriority
@@ -50,10 +52,15 @@ namespace UE::Tasks
 		RHIThreadNormalPri,
 		RHIThreadHiPri,
 		RHIThreadNormalPriLocalQueue,
-		RHIThreadHiPriLocalQueue
+		RHIThreadHiPriLocalQueue,
 #endif
+
+		Count
 	};
 
+	const TCHAR* ToString(EExtendedTaskPriority ExtendedPriority);
+	bool ToExtendedTaskPriority(const TCHAR* ExtendedPriorityStr, EExtendedTaskPriority& OutExtendedPriority);
+
 	class FPipe;
 
 	namespace Private
