Add missing ISPC forward declarations

diff --git a/Engine/Source/Runtime/Core/Public/Math/MathFwd.h b/Engine/Source/Runtime/Core/Public/Math/MathFwd.h
--- a/Engine/Source/Runtime/Core/Public/Math/MathFwd.h
+++ b/Engine/Source/Runtime/Core/Public/Math/MathFwd.h
@@ -30,6 +30,18 @@ namespace UE::Math
 	template <typename IntType> struct TIntVector4;
 }
 
+// ISPC forward declarations
+namespace ispc
+{
+	struct FVector;
+	struct FVector2D;
+	struct FVector4;
+	struct FQuat;
+	struct FMatrix;
+	struct FPlane;
+	struct FTransform;
+}
+
 // Forward declaration of concrete types				// Macro version - declares all three variants.
 using FVector 		= UE::Math::TVector<double>;		// UE_DECLARE_LWC_TYPE(Vector, 3);
 using FVector2D 	= UE::Math::TVector2<double>;		// UE_DECLARE_LWC_TYPE(Vector2,, FVector2D);
