Fixed CIS error (-NoPCH -DisableUnity).

diff --git a/Engine/Source/Developer/TraceAnalysis/Private/Analysis/Processor.cpp b/Engine/Source/Developer/TraceAnalysis/Private/Analysis/Processor.cpp
--- a/Engine/Source/Developer/TraceAnalysis/Private/Analysis/Processor.cpp
+++ b/Engine/Source/Developer/TraceAnalysis/Private/Analysis/Processor.cpp
@@ -2,6 +2,7 @@
 
 #include "Processor.h"
 #include "HAL/Event.h"
+#include "HAL/LowLevelMemTracker.h"
 #include "HAL/PlatformProcess.h"
 #include "HAL/RunnableThread.h"
 #include "StreamReader.h"
diff --git a/Engine/Source/Runtime/Core/Private/ProfilingDebugging/CpuProfilerTrace.cpp b/Engine/Source/Runtime/Core/Private/ProfilingDebugging/CpuProfilerTrace.cpp
--- a/Engine/Source/Runtime/Core/Private/ProfilingDebugging/CpuProfilerTrace.cpp
+++ b/Engine/Source/Runtime/Core/Private/ProfilingDebugging/CpuProfilerTrace.cpp
@@ -1,6 +1,7 @@
 // Copyright Epic Games, Inc. All Rights Reserved.
 #include "ProfilingDebugging/CpuProfilerTrace.h"
 #include "Trace/Trace.inl"
+#include "HAL/LowLevelMemTracker.h"
 #include "HAL/PlatformTime.h"
 #include "HAL/PlatformTLS.h"
 #include "HAL/TlsAutoCleanup.h"
