Defer String Table find/load during module CDO registration, until StartupModule can be called. This resolves warnings that arise during CDO construction before game modules are loaded during cooking.
#rnx

diff --git a/Engine/Source/Runtime/Core/Private/Internationalization/StringTableCore.cpp b/Engine/Source/Runtime/Core/Private/Internationalization/StringTableCore.cpp
--- a/Engine/Source/Runtime/Core/Private/Internationalization/StringTableCore.cpp
+++ b/Engine/Source/Runtime/Core/Private/Internationalization/StringTableCore.cpp
@@ -10,6 +10,7 @@ DEFINE_LOG_CATEGORY(LogStringTable);
 
 
 IStringTableEngineBridge* IStringTableEngineBridge::InstancePtr = nullptr;
+std::atomic<int8> IStringTableEngineBridge::DeferFindOrLoad(0);
 
 
 namespace StringTableRedirects
diff --git a/Engine/Source/Runtime/Core/Private/Modules/ModuleManager.cpp b/Engine/Source/Runtime/Core/Private/Modules/ModuleManager.cpp
--- a/Engine/Source/Runtime/Core/Private/Modules/ModuleManager.cpp
+++ b/Engine/Source/Runtime/Core/Private/Modules/ModuleManager.cpp
@@ -12,6 +12,7 @@
 #include "Misc/ScopeLock.h"
 #include "Misc/DataDrivenPlatformInfoRegistry.h"
 #include "Serialization/LoadTimeTrace.h"
+#include "Internationalization/StringTableCore.h"
 #include "Misc/FileHelper.h"
 #include "Serialization/MemoryReader.h"
 #include "Serialization/MemoryWriter.h"
@@ -481,6 +482,10 @@ IModuleInterface* FModuleManager::LoadModuleWithFailureReason(const FName InModu
 			TRACE_LOADTIME_REQUEST_GROUP_SCOPE(TEXT("LoadModule - %s"), *InModuleName.ToString());
 #if USE_PER_MODULE_UOBJECT_BOOTSTRAP || WITH_VERSE
 			{
+				// Defer String Table find/load during CDO registration, as it may happen 
+				// before StartupModule has had a chance to load the String Table
+				IStringTableEngineBridge::FScopedDeferFindOrLoad DeferStringTableFindOrLoad;
+
 				ProcessLoadedObjectsCallback.Broadcast(InModuleName, bCanProcessNewlyLoadedObjects);
 			}
 #endif
@@ -525,6 +530,10 @@ IModuleInterface* FModuleManager::LoadModuleWithFailureReason(const FName InModu
 		// in the module being loaded.
 		if (bCanProcessNewlyLoadedObjects)
 		{
+			// Defer String Table find/load during CDO registration, as it may happen 
+			// before StartupModule has had a chance to load the String Table
+			IStringTableEngineBridge::FScopedDeferFindOrLoad DeferStringTableFindOrLoad;
+
 			ProcessLoadedObjectsCallback.Broadcast(NAME_None, bCanProcessNewlyLoadedObjects);
 		}
 
@@ -561,14 +570,20 @@ IModuleInterface* FModuleManager::LoadModuleWithFailureReason(const FName InModu
 			ModuleInfo->Handle = FPlatformProcess::GetDllHandle(*ModuleFileToLoad);
 			if (ModuleInfo->Handle != nullptr)
 			{
-				// First things first.  If the loaded DLL has UObjects in it, then their generated code's
-				// static initialization will have run during the DLL loading phase, and we'll need to
-				// go in and make sure those new UObject classes are properly registered.
-						// Sometimes modules are loaded before even the UObject systems are ready.  We need to assume
-						// these modules aren't using UObjects.
-							// OK, we've verified that loading the module caused new UObject classes to be
-							// registered, so we'll treat this module as a module with UObjects in it.
-				ProcessLoadedObjectsCallback.Broadcast(InModuleName, bCanProcessNewlyLoadedObjects);
+				{
+					// Defer String Table find/load during CDO registration, as it may happen 
+					// before StartupModule has had a chance to load the String Table
+					IStringTableEngineBridge::FScopedDeferFindOrLoad DeferStringTableFindOrLoad;
+
+					// First things first.  If the loaded DLL has UObjects in it, then their generated code's
+					// static initialization will have run during the DLL loading phase, and we'll need to
+					// go in and make sure those new UObject classes are properly registered.
+					// Sometimes modules are loaded before even the UObject systems are ready.  We need to assume
+					// these modules aren't using UObjects.
+					// OK, we've verified that loading the module caused new UObject classes to be
+					// registered, so we'll treat this module as a module with UObjects in it.
+					ProcessLoadedObjectsCallback.Broadcast(InModuleName, bCanProcessNewlyLoadedObjects);
+				}
 
 				// Find our "InitializeModule" global function, which must exist for all module DLLs
 				FInitializeModuleFunctionPtr InitializeModuleFunctionPtr =
diff --git a/Engine/Source/Runtime/Core/Public/Internationalization/StringTableCore.h b/Engine/Source/Runtime/Core/Public/Internationalization/StringTableCore.h
--- a/Engine/Source/Runtime/Core/Public/Internationalization/StringTableCore.h
+++ b/Engine/Source/Runtime/Core/Public/Internationalization/StringTableCore.h
@@ -1,6 +1,7 @@
 // Copyright Epic Games, Inc. All Rights Reserved.
 #pragma once
 
+#include <atomic>
 #include "Internationalization/StringTableCoreFwd.h"
 #include "Containers/Map.h"
 #include "UObject/NameTypes.h"
@@ -171,6 +172,24 @@ private:
 class CORE_API IStringTableEngineBridge
 {
 public:
+	/**
+	 * Scope object used to temporarily defer String Table find/load (eg, during module load).
+	 */
+	struct FScopedDeferFindOrLoad
+	{
+		FScopedDeferFindOrLoad()
+		{
+			++DeferFindOrLoad;
+		}
+
+		~FScopedDeferFindOrLoad()
+		{
+			--DeferFindOrLoad;
+		}
+
+		UE_NONCOPYABLE(FScopedDeferFindOrLoad);
+	};
+
 	/**
 	 * Callback used when loading string table assets.
 	 * @param The name of the table we were asked to load.
@@ -185,6 +204,7 @@ public:
 	static bool CanFindOrLoadStringTableAsset()
 	{
 		return FInternationalization::IsAvailable()
+			&& DeferFindOrLoad.load(std::memory_order::memory_order_relaxed) <= 0
 			&& (!InstancePtr || InstancePtr->CanFindOrLoadStringTableAssetImpl());
 	}
 
@@ -269,6 +289,9 @@ protected:
 
 	/** Singleton instance, populated by the derived type */
 	static IStringTableEngineBridge* InstancePtr;
+
+	/** Whether String Table find/load is currently deferred (eg, during module load) */
+	static std::atomic<int8> DeferFindOrLoad;
 };
 
 /** String table redirect utils */
