Windows: Always print the error string when it's provided to ErrorOutputDevice::Serialize

### Notes
Many other platforms already do this, this just brings parity to Windows.
I think by design the error message is meant to be reported via a call to HandleError() or LogFormattedMessageWithCallstack(), but sometimes neither happens. The reasons are unclear due to lack of information.

This print may be redundant (the message be printed twice in some situations, once with a callstack and once without), but that's better than not having it at all in edge cases.

#rnx

diff --git a/Engine/Source/Runtime/Core/Private/Windows/WindowsErrorOutputDevice.cpp b/Engine/Source/Runtime/Core/Private/Windows/WindowsErrorOutputDevice.cpp
--- a/Engine/Source/Runtime/Core/Private/Windows/WindowsErrorOutputDevice.cpp
+++ b/Engine/Source/Runtime/Core/Private/Windows/WindowsErrorOutputDevice.cpp
@@ -33,14 +33,17 @@ void FWindowsErrorOutputDevice::Serialize( const TCHAR* Msg, ELogVerbosity::Type
    
 	if( !GIsCriticalError )
 	{   
-		const int32 LastError = ::GetLastError();
-
 		// First appError.
 		GIsCriticalError = 1;
-		TCHAR ErrorBuffer[1024];
-		ErrorBuffer[0] = TCHAR('\0');
+
+		// There is a bit of ambiguity around which SEH __except is to call HandleError or LogFormattedMessageWithCallstack
+		// We print it here so that the message, at very least, is always included with the logs
+		UE_LOG(LogWindows, Error, TEXT("appError called: %s"), Msg);
 
 		// Windows error.
+		const int32 LastError = ::GetLastError();
+		TCHAR ErrorBuffer[1024];
+		ErrorBuffer[0] = TCHAR('\0');
 		if (LastError == 0)
 		{
 			UE_LOG(LogWindows, Log, TEXT("Windows GetLastError: %s (%i)"), FPlatformMisc::GetSystemErrorMessage(ErrorBuffer, 1024, LastError), LastError);
