Adding support for configuring an array property to be empty.
If an array property has elements defines in a Base config and an editor wants to configure the property to be empty in a Default config, nothing was being written.
The first problem was that only object properties with values are written to config files. In this case, we can use the !Property=ClearValue.
The second problem was handling this !Property case inside the config writing code. The code that detects array properties had to be adjusted to detect a !Property entry and check if the property is also empty in the config hierarchy.

diff --git a/Engine/Source/Runtime/Core/Private/Misc/ConfigCacheIni.cpp b/Engine/Source/Runtime/Core/Private/Misc/ConfigCacheIni.cpp
--- a/Engine/Source/Runtime/Core/Private/Misc/ConfigCacheIni.cpp
+++ b/Engine/Source/Runtime/Core/Private/Misc/ConfigCacheIni.cpp
@@ -2072,55 +2072,79 @@ void FConfigFile::ProcessSourceAndCheckAgainstBackup()
 	}
 }
 
+static TArray<FString> GetSourceProperties(const FConfigFileHierarchy& SourceIniHierarchy, int IniCombineThreshold, const FString& SectionName, const FString& PropertyName)
+{
+	// Build a config file out of this default configs hierarchy.
+	FConfigCacheIni Hierarchy(EConfigCacheType::Temporary);
+
+	int32 HighestFileIndex = 0;
+	TArray<int32> ExistingEntries;
+	SourceIniHierarchy.GetKeys(ExistingEntries);
+	for (const int32& NextEntry : ExistingEntries)
+	{
+		HighestFileIndex = NextEntry > HighestFileIndex ? NextEntry : HighestFileIndex;
+	}
+
+	const FString& LastFileInHierarchy = SourceIniHierarchy.FindChecked(HighestFileIndex);
+	FConfigFile& DefaultConfigFile = Hierarchy.Add(LastFileInHierarchy, FConfigFile());
+
+	for (const auto& HierarchyFileIt : SourceIniHierarchy)
+	{
+		// Combine everything up to the level we're writing, but not including it.
+		// Inclusion would result in a bad feedback loop where on subsequent writes 
+		// we would be diffing against the same config we've just written to.
+		if (HierarchyFileIt.Key < IniCombineThreshold)
+		{
+			DefaultConfigFile.Combine(HierarchyFileIt.Value);
+		}
+	}
+
+	// Remove any array elements from the default configs hierearchy, we will add these in below
+	// Note.	This compensates for an issue where strings in the hierarchy have a slightly different format
+	//			to how the config system wishes to serialize them.
+	TArray<FString> SourceArrayProperties;
+	Hierarchy.GetArray(*SectionName, *PropertyName, SourceArrayProperties, *LastFileInHierarchy);
+
+	return SourceArrayProperties;
+}
+
 void FConfigFile::ProcessPropertyAndWriteForDefaults(int IniCombineThreshold, const TArray<const FConfigValue*>& InCompletePropertyToProcess, FString& OutText, const FString& SectionName, const FString& PropertyName)
 {
 	// Only process against a hierarchy if this config file has one.
 	if (SourceIniHierarchy.Num() > 0)
 	{
-		// Handle array elements from the configs hierarchy.
-		if (PropertyName.StartsWith(TEXT("+")) || InCompletePropertyToProcess.Num() > 1)
-		{
-			// Build a config file out of this default configs hierarchy.
-			FConfigCacheIni Hierarchy(EConfigCacheType::Temporary);
+		FString CleanedPropertyName = PropertyName;
+		const bool bHadPlus = CleanedPropertyName.RemoveFromStart(TEXT("+"));
+		const bool bHadBang = CleanedPropertyName.RemoveFromStart(TEXT("!"));
 
-			int32 HighestFileIndex = 0;
-			TArray<int32> ExistingEntries;
-			SourceIniHierarchy.GetKeys(ExistingEntries);
-			for (const int32& NextEntry : ExistingEntries)
-			{
-				HighestFileIndex = NextEntry > HighestFileIndex ? NextEntry : HighestFileIndex;
-			}
+		const FString PropertyNameWithRemoveOp = TEXT("-") + CleanedPropertyName;
 
-			const FString& LastFileInHierarchy = SourceIniHierarchy.FindChecked(HighestFileIndex);
-			FConfigFile& DefaultConfigFile = Hierarchy.Add(LastFileInHierarchy, FConfigFile());
-
-			for (const auto& HierarchyFileIt : SourceIniHierarchy)
+		// look for pointless !Clear entries that the config system wrote out when it noticed the user didn't have any entries
+		if (bHadBang && InCompletePropertyToProcess.Num() == 1 && InCompletePropertyToProcess[0]->GetSavedValue() == TEXT("__ClearArray__"))
+		{
+			const TArray<FString> SourceArrayProperties = GetSourceProperties(SourceIniHierarchy, IniCombineThreshold, SectionName, CleanedPropertyName);
+			for (const FString& NextElement : SourceArrayProperties)
 			{
-				// Combine everything up to the level we're writing, but not including it.
-				// Inclusion would result in a bad feedback loop where on subsequent writes 
-				// we would be diffing against the same config we've just written to.
-				if (HierarchyFileIt.Key < IniCombineThreshold)
-				{
-					DefaultConfigFile.Combine(HierarchyFileIt.Value);
-				}
+				OutText.Append(GenerateExportedPropertyLine(PropertyNameWithRemoveOp, NextElement));
 			}
 
-			// Remove any array elements from the default configs hierearchy, we will add these in below
-			// Note.	This compensates for an issue where strings in the hierarchy have a slightly different format
-			//			to how the config system wishes to serialize them.
-			TArray<FString> ArrayProperties;
-			Hierarchy.GetArray(*SectionName, *PropertyName.Replace(TEXT("+"), TEXT("")), ArrayProperties, *LastFileInHierarchy);
+			// We don't need to write out the ! entry so just leave now.
+			return;
+		}
 
-			for (const FString& NextElement : ArrayProperties)
+		// Handle array elements from the configs hierarchy.
+		if (bHadPlus || InCompletePropertyToProcess.Num() > 1)
+		{
+			const TArray<FString> SourceArrayProperties = GetSourceProperties(SourceIniHierarchy, IniCombineThreshold, SectionName, CleanedPropertyName);
+			for (const FString& NextElement : SourceArrayProperties)
 			{
-				FString PropertyNameWithRemoveOp = PropertyName.Replace(TEXT("+"), TEXT("-"));
 				OutText.Append(GenerateExportedPropertyLine(PropertyNameWithRemoveOp, NextElement));
 			}
 		}
 	}
 
 	// Write the properties out to a file.
-	for ( const FConfigValue* PropertyIt : InCompletePropertyToProcess)
+	for (const FConfigValue* PropertyIt : InCompletePropertyToProcess)
 	{
 		OutText.Append(GenerateExportedPropertyLine(PropertyName, PropertyIt->GetSavedValue()));
 	}
diff --git a/Engine/Source/Runtime/CoreUObject/Private/UObject/Obj.cpp b/Engine/Source/Runtime/CoreUObject/Private/UObject/Obj.cpp
--- a/Engine/Source/Runtime/CoreUObject/Private/UObject/Obj.cpp
+++ b/Engine/Source/Runtime/CoreUObject/Private/UObject/Obj.cpp
@@ -2730,6 +2730,11 @@ void UObject::SaveConfig( uint64 Flags, const TCHAR* InFilename, FConfigCacheIni
 						Array->Inner->ExportTextItem_Direct( Buffer, ArrayHelper.GetRawPtr(i), ArrayHelper.GetRawPtr(i), this, PortFlags );
 						Sec->Add(*CompleteKey, *Buffer);
 					}
+					if (ArrayHelper.Num() == 0 && bIsADefaultIniWrite)
+					{
+						const FString EmptyKey = FString::Printf(TEXT("!%s"), *Key);
+						Sec->Add(*EmptyKey, TEXT("__ClearArray__"));
+					}
 				}
 			}
 			else
