ConfigCache GetXOrDefault to return value gotten by value rather than fill a reference
does not indicate if return value came from config or the default value
useful for one-time init of static variables in code locations where config may be queried too often, like :
 static int32 bMyConfigValue = GConfig->GetIntOrDefault(Section,Key,DefaultValue,ConfigFilename);

diff --git a/Engine/Source/Runtime/Core/Public/Misc/ConfigCacheIni.h b/Engine/Source/Runtime/Core/Public/Misc/ConfigCacheIni.h
--- a/Engine/Source/Runtime/Core/Public/Misc/ConfigCacheIni.h
+++ b/Engine/Source/Runtime/Core/Public/Misc/ConfigCacheIni.h
@@ -922,6 +922,53 @@ public:
 	{
 		return GetArray(Section, Key, Value, Filename);
 	}
+		
+	// Return a config value if found, if not found return default value
+	// does not indicate if return value came from config or the default value
+	// useful for one-time init of static variables in code locations where config may be queried too often, like :
+	//  static int32 bMyConfigValue = GConfig->GetIntOrDefault(Section,Key,DefaultValue,ConfigFilename);
+	int32 GetIntOrDefault(const TCHAR* Section, const TCHAR* Key, const int32 DefaultValue, const FString& Filename)
+	{
+		int32 Value = DefaultValue;
+		GetInt(Section,Key,Value,Filename);
+		return Value;
+	}
+	float GetFloatOrDefault(const TCHAR* Section, const TCHAR* Key, const float DefaultValue, const FString& Filename)
+	{
+		float Value = DefaultValue;
+		GetFloat(Section,Key,Value,Filename);
+		return Value;
+	}
+	bool GetBoolOrDefault(const TCHAR* Section, const TCHAR* Key, const bool DefaultValue, const FString& Filename)
+	{
+		bool Value = DefaultValue;
+		GetBool(Section,Key,Value,Filename);
+		return Value;
+	}
+	FString GetStringOrDefault(const TCHAR* Section, const TCHAR* Key, const FString & DefaultValue, const FString& Filename)
+	{
+		FString Value;
+		if ( GetString(Section,Key,Value,Filename) )
+		{
+			return Value;
+		}
+		else
+		{
+			return DefaultValue;
+		}
+	}
+	FText GetTextOrDefault(const TCHAR* Section, const TCHAR* Key, const FText & DefaultValue, const FString& Filename)
+	{
+		FText Value;
+		if ( GetText(Section,Key,Value,Filename) )
+		{
+			return Value;
+		}
+		else
+		{
+			return DefaultValue;
+		}
+	}
 
 	void SetInt
 	(
diff --git a/Engine/Source/Runtime/Engine/Private/Texture.cpp b/Engine/Source/Runtime/Engine/Private/Texture.cpp
--- a/Engine/Source/Runtime/Engine/Private/Texture.cpp
+++ b/Engine/Source/Runtime/Engine/Private/Texture.cpp
@@ -715,16 +715,8 @@ void UTexture::PostEditChangeProperty(FPropertyChangedEvent& PropertyChangedEven
 #if WITH_EDITORONLY_DATA
 static bool IsEnableLegacyAlphaCoverageThresholdScaling()
 {
-	static struct ReadConfigOnce
-	{
-		bool bBool;
-		ReadConfigOnce()
-		{
-			bBool = false;
-			GConfig->GetBool(TEXT("Texture"), TEXT("EnableLegacyAlphaCoverageThresholdScaling"), bBool, GEditorIni);
-		}
-	} ConfigValue;
-	return ConfigValue.bBool;
+	static bool bConfigBool = GConfig->GetBoolOrDefault(TEXT("Texture"), TEXT("EnableLegacyAlphaCoverageThresholdScaling"), false, GEditorIni);
+	return bConfigBool;
 }
 #endif
 
