[Insights] Fixed StatsTrace incorrectly tracing disabled stats (i.e. when the stat name is set to "None").

diff --git a/Engine/Source/Runtime/Core/Private/Stats/StatsTrace.cpp b/Engine/Source/Runtime/Core/Private/Stats/StatsTrace.cpp
--- a/Engine/Source/Runtime/Core/Private/Stats/StatsTrace.cpp
+++ b/Engine/Source/Runtime/Core/Private/Stats/StatsTrace.cpp
@@ -106,74 +106,95 @@ void FStatsTraceInternal::EndEncodeOp(FThreadState* ThreadState, uint8* BufferPt
 
 void FStatsTrace::DeclareStat(const FName& Stat, const ANSICHAR* Name, const TCHAR* Description, const ANSICHAR* Group, bool IsFloatingPoint, bool IsMemory, bool ShouldClearEveryFrame)
 {
-	uint32 NameLen = FCStringAnsi::Strlen(Name);
-	uint32 DescriptionLen = FCString::Strlen(Description);
-	uint32 GroupLen = FCStringAnsi::Strlen(Group);
-		
-	UE_TRACE_LOG(Stats, Spec, StatsChannel, NameLen + DescriptionLen + GroupLen)
-		<< Spec.Id(Stat.GetComparisonIndex().ToUnstableInt())
-		<< Spec.IsFloatingPoint(IsFloatingPoint)
-		<< Spec.IsMemory(IsMemory)
-		<< Spec.ShouldClearEveryFrame(ShouldClearEveryFrame)
-		<< Spec.Name(Name, NameLen)
-		<< Spec.Description(Description, DescriptionLen)
-		<< Spec.Group(Group, GroupLen);
+	if (!Stat.IsNone())
+	{
+		uint32 NameLen = FCStringAnsi::Strlen(Name);
+		uint32 DescriptionLen = FCString::Strlen(Description);
+		uint32 GroupLen = FCStringAnsi::Strlen(Group);
+
+		UE_TRACE_LOG(Stats, Spec, StatsChannel, NameLen + DescriptionLen + GroupLen)
+			<< Spec.Id(Stat.GetComparisonIndex().ToUnstableInt())
+			<< Spec.IsFloatingPoint(IsFloatingPoint)
+			<< Spec.IsMemory(IsMemory)
+			<< Spec.ShouldClearEveryFrame(ShouldClearEveryFrame)
+			<< Spec.Name(Name, NameLen)
+			<< Spec.Description(Description, DescriptionLen)
+			<< Spec.Group(Group, GroupLen);
+	}
 }
 
 void FStatsTrace::Increment(const FName& Stat)
 {
-	FStatsTraceInternal::FThreadState* ThreadState;
-	uint8* BufferPtr;
-	FStatsTraceInternal::BeginEncodeOp(Stat, FStatsTraceInternal::Increment, ThreadState, BufferPtr);
-	ThreadState->BufferSize = uint16(BufferPtr - ThreadState->Buffer);
-	FStatsTraceInternal::EndEncodeOp(ThreadState, BufferPtr);
+	if (!Stat.IsNone())
+	{
+		FStatsTraceInternal::FThreadState* ThreadState;
+		uint8* BufferPtr;
+		FStatsTraceInternal::BeginEncodeOp(Stat, FStatsTraceInternal::Increment, ThreadState, BufferPtr);
+		ThreadState->BufferSize = uint16(BufferPtr - ThreadState->Buffer);
+		FStatsTraceInternal::EndEncodeOp(ThreadState, BufferPtr);
+	}
 }
 
 void FStatsTrace::Decrement(const FName& Stat)
 {
-	FStatsTraceInternal::FThreadState* ThreadState;
-	uint8* BufferPtr;
-	FStatsTraceInternal::BeginEncodeOp(Stat, FStatsTraceInternal::Decrement, ThreadState, BufferPtr);
-	ThreadState->BufferSize = uint16(BufferPtr - ThreadState->Buffer);
-	FStatsTraceInternal::EndEncodeOp(ThreadState, BufferPtr);
+	if (!Stat.IsNone())
+	{
+		FStatsTraceInternal::FThreadState* ThreadState;
+		uint8* BufferPtr;
+		FStatsTraceInternal::BeginEncodeOp(Stat, FStatsTraceInternal::Decrement, ThreadState, BufferPtr);
+		ThreadState->BufferSize = uint16(BufferPtr - ThreadState->Buffer);
+		FStatsTraceInternal::EndEncodeOp(ThreadState, BufferPtr);
+	}
 }
 
 void FStatsTrace::Add(const FName& Stat, int64 Amount)
 {
-	FStatsTraceInternal::FThreadState* ThreadState;
-	uint8* BufferPtr;
-	FStatsTraceInternal::BeginEncodeOp(Stat, FStatsTraceInternal::AddInteger, ThreadState, BufferPtr);
-	FTraceUtils::EncodeZigZag(Amount, BufferPtr);
-	FStatsTraceInternal::EndEncodeOp(ThreadState, BufferPtr);
+	if (!Stat.IsNone())
+	{
+		FStatsTraceInternal::FThreadState* ThreadState;
+		uint8* BufferPtr;
+		FStatsTraceInternal::BeginEncodeOp(Stat, FStatsTraceInternal::AddInteger, ThreadState, BufferPtr);
+		FTraceUtils::EncodeZigZag(Amount, BufferPtr);
+		FStatsTraceInternal::EndEncodeOp(ThreadState, BufferPtr);
+	}
 }
 
 void FStatsTrace::Add(const FName& Stat, double Amount)
 {
-	FStatsTraceInternal::FThreadState* ThreadState;
-	uint8* BufferPtr;
-	FStatsTraceInternal::BeginEncodeOp(Stat, FStatsTraceInternal::AddFloat, ThreadState, BufferPtr);
-	memcpy(BufferPtr, &Amount, sizeof(double));
-	BufferPtr += sizeof(double);
-	FStatsTraceInternal::EndEncodeOp(ThreadState, BufferPtr);
+	if (!Stat.IsNone())
+	{
+		FStatsTraceInternal::FThreadState* ThreadState;
+		uint8* BufferPtr;
+		FStatsTraceInternal::BeginEncodeOp(Stat, FStatsTraceInternal::AddFloat, ThreadState, BufferPtr);
+		memcpy(BufferPtr, &Amount, sizeof(double));
+		BufferPtr += sizeof(double);
+		FStatsTraceInternal::EndEncodeOp(ThreadState, BufferPtr);
+	}
 }
 
 void FStatsTrace::Set(const FName& Stat, int64 Value)
 {
-	FStatsTraceInternal::FThreadState* ThreadState;
-	uint8* BufferPtr;
-	FStatsTraceInternal::BeginEncodeOp(Stat, FStatsTraceInternal::SetInteger, ThreadState, BufferPtr);
-	FTraceUtils::EncodeZigZag(Value, BufferPtr);
-	FStatsTraceInternal::EndEncodeOp(ThreadState, BufferPtr);
+	if (!Stat.IsNone())
+	{
+		FStatsTraceInternal::FThreadState* ThreadState;
+		uint8* BufferPtr;
+		FStatsTraceInternal::BeginEncodeOp(Stat, FStatsTraceInternal::SetInteger, ThreadState, BufferPtr);
+		FTraceUtils::EncodeZigZag(Value, BufferPtr);
+		FStatsTraceInternal::EndEncodeOp(ThreadState, BufferPtr);
+	}
 }
 
 void FStatsTrace::Set(const FName& Stat, double Value)
 {
-	FStatsTraceInternal::FThreadState* ThreadState;
-	uint8* BufferPtr;
-	FStatsTraceInternal::BeginEncodeOp(Stat, FStatsTraceInternal::SetFloat, ThreadState, BufferPtr);
-	memcpy(BufferPtr, &Value, sizeof(double));
-	BufferPtr += sizeof(double);
-	FStatsTraceInternal::EndEncodeOp(ThreadState, BufferPtr);
+	if (!Stat.IsNone())
+	{
+		FStatsTraceInternal::FThreadState* ThreadState;
+		uint8* BufferPtr;
+		FStatsTraceInternal::BeginEncodeOp(Stat, FStatsTraceInternal::SetFloat, ThreadState, BufferPtr);
+		memcpy(BufferPtr, &Value, sizeof(double));
+		BufferPtr += sizeof(double);
+		FStatsTraceInternal::EndEncodeOp(ThreadState, BufferPtr);
+	}
 }
 
-#endif
\ No newline at end of file
+#endif
diff --git a/Engine/Source/Runtime/Core/Public/Stats/Stats2.h b/Engine/Source/Runtime/Core/Public/Stats/Stats2.h
--- a/Engine/Source/Runtime/Core/Public/Stats/Stats2.h
+++ b/Engine/Source/Runtime/Core/Public/Stats/Stats2.h
@@ -2062,8 +2062,9 @@ struct FStat_##StatName\
 {\
 	if (FThreadStats::IsCollectingData() || !GET_STATISEVERYFRAME(Stat)) \
 	{ \
-		FThreadStats::AddMessage(GET_STATFNAME(Stat), EStatOperation::Add, int64(1));\
-		TRACE_STAT_INCREMENT(GET_STATFNAME(Stat)); \
+		const FName StatName = GET_STATFNAME(Stat); \
+		FThreadStats::AddMessage(StatName, EStatOperation::Add, int64(1));\
+		TRACE_STAT_INCREMENT(StatName); \
 	} \
 }
 #define INC_FLOAT_STAT_BY(Stat, Amount) \
@@ -2073,8 +2074,9 @@ struct FStat_##StatName\
 	{ \
 		if (FThreadStats::IsCollectingData() || !GET_STATISEVERYFRAME(Stat)) \
 		{ \
-			FThreadStats::AddMessage(GET_STATFNAME(Stat), EStatOperation::Add, AddAmount); \
-			TRACE_STAT_ADD(GET_STATFNAME(Stat), AddAmount); \
+			const FName StatName = GET_STATFNAME(Stat); \
+			FThreadStats::AddMessage(StatName, EStatOperation::Add, AddAmount); \
+			TRACE_STAT_ADD(StatName, AddAmount); \
 		} \
 	} \
 }
@@ -2085,8 +2087,9 @@ struct FStat_##StatName\
 	{ \
 		if (FThreadStats::IsCollectingData() || !GET_STATISEVERYFRAME(Stat)) \
 		{ \
-			FThreadStats::AddMessage(GET_STATFNAME(Stat), EStatOperation::Add, AddAmount); \
-			TRACE_STAT_ADD(GET_STATFNAME(Stat), AddAmount); \
+			const FName StatName = GET_STATFNAME(Stat); \
+			FThreadStats::AddMessage(StatName, EStatOperation::Add, AddAmount); \
+			TRACE_STAT_ADD(StatName, AddAmount); \
 		} \
 	} \
 }
@@ -2097,8 +2100,9 @@ struct FStat_##StatName\
 	{ \
 		if (FThreadStats::IsCollectingData() || !GET_STATISEVERYFRAME(Stat)) \
 		{ \
-			FThreadStats::AddMessage(GET_STATFNAME(Stat), EStatOperation::Add, AddAmount); \
-			TRACE_STAT_ADD(GET_STATFNAME(Stat), AddAmount); \
+			const FName StatName = GET_STATFNAME(Stat); \
+			FThreadStats::AddMessage(StatName, EStatOperation::Add, AddAmount); \
+			TRACE_STAT_ADD(StatName, AddAmount); \
 		} \
 	} \
 }
@@ -2106,8 +2110,9 @@ struct FStat_##StatName\
 {\
 	if (FThreadStats::IsCollectingData() || !GET_STATISEVERYFRAME(Stat)) \
 	{ \
-		FThreadStats::AddMessage(GET_STATFNAME(Stat), EStatOperation::Subtract, int64(1));\
-		TRACE_STAT_DECREMENT(GET_STATFNAME(Stat)); \
+		const FName StatName = GET_STATFNAME(Stat); \
+		FThreadStats::AddMessage(StatName, EStatOperation::Subtract, int64(1));\
+		TRACE_STAT_DECREMENT(StatName); \
 	} \
 }
 #define DEC_FLOAT_STAT_BY(Stat,Amount) \
@@ -2117,8 +2122,9 @@ struct FStat_##StatName\
 	{ \
 		if (FThreadStats::IsCollectingData() || !GET_STATISEVERYFRAME(Stat)) \
 		{ \
-			FThreadStats::AddMessage(GET_STATFNAME(Stat), EStatOperation::Subtract, SubtractAmount); \
-			TRACE_STAT_ADD(GET_STATFNAME(Stat), -SubtractAmount); \
+			const FName StatName = GET_STATFNAME(Stat); \
+			FThreadStats::AddMessage(StatName, EStatOperation::Subtract, SubtractAmount); \
+			TRACE_STAT_ADD(StatName, -SubtractAmount); \
 		} \
 	} \
 }
@@ -2129,8 +2135,9 @@ struct FStat_##StatName\
 	{ \
 		if (FThreadStats::IsCollectingData() || !GET_STATISEVERYFRAME(Stat)) \
 		{ \
-			FThreadStats::AddMessage(GET_STATFNAME(Stat), EStatOperation::Subtract, SubtractAmount); \
-			TRACE_STAT_ADD(GET_STATFNAME(Stat), -SubtractAmount); \
+			const FName StatName = GET_STATFNAME(Stat); \
+			FThreadStats::AddMessage(StatName, EStatOperation::Subtract, SubtractAmount); \
+			TRACE_STAT_ADD(StatName, -SubtractAmount); \
 		} \
 	} \
 }
@@ -2141,8 +2148,9 @@ struct FStat_##StatName\
 	{ \
 		if (FThreadStats::IsCollectingData() || !GET_STATISEVERYFRAME(Stat)) \
 		{ \
-			FThreadStats::AddMessage(GET_STATFNAME(Stat), EStatOperation::Subtract, SubtractAmount); \
-			TRACE_STAT_ADD(GET_STATFNAME(Stat), -SubtractAmount); \
+			const FName StatName = GET_STATFNAME(Stat); \
+			FThreadStats::AddMessage(StatName, EStatOperation::Subtract, SubtractAmount); \
+			TRACE_STAT_ADD(StatName, -SubtractAmount); \
 		} \
 	} \
 }
@@ -2150,27 +2158,30 @@ struct FStat_##StatName\
 {\
 	if (FThreadStats::IsCollectingData() || !GET_STATISEVERYFRAME(Stat)) \
 	{ \
+		const FName StatName = GET_STATFNAME(Stat); \
 		int64 SetValue = int64(Value); \
-		FThreadStats::AddMessage(GET_STATFNAME(Stat), EStatOperation::Set, SetValue); \
-		TRACE_STAT_SET(GET_STATFNAME(Stat), SetValue); \
+		FThreadStats::AddMessage(StatName, EStatOperation::Set, SetValue); \
+		TRACE_STAT_SET(StatName, SetValue); \
 	} \
 }
 #define SET_DWORD_STAT(Stat,Value) \
 {\
 	if (FThreadStats::IsCollectingData() || !GET_STATISEVERYFRAME(Stat)) \
 	{ \
+		const FName StatName = GET_STATFNAME(Stat); \
 		int64 SetValue = int64(Value); \
-		FThreadStats::AddMessage(GET_STATFNAME(Stat), EStatOperation::Set, SetValue); \
-		TRACE_STAT_SET(GET_STATFNAME(Stat), SetValue); \
+		FThreadStats::AddMessage(StatName, EStatOperation::Set, SetValue); \
+		TRACE_STAT_SET(StatName, SetValue); \
 	} \
 }
 #define SET_FLOAT_STAT(Stat,Value) \
 {\
 	if (FThreadStats::IsCollectingData() || !GET_STATISEVERYFRAME(Stat)) \
 	{ \
+		const FName StatName = GET_STATFNAME(Stat); \
 		double SetValue = double(Value); \
-		FThreadStats::AddMessage(GET_STATFNAME(Stat), EStatOperation::Set, SetValue); \
-		TRACE_STAT_SET(GET_STATFNAME(Stat), SetValue); \
+		FThreadStats::AddMessage(StatName, EStatOperation::Set, SetValue); \
+		TRACE_STAT_SET(StatName, SetValue); \
 	} \
 }
 
