Ensure HoloLens only includes the HoloLens versions of the Windows API wrapper files.

#rnx

diff --git a/Engine/Platforms/Hololens/Source/Runtime/Core/Public/HoloLensSystemIncludes.h b/Engine/Platforms/Hololens/Source/Runtime/Core/Public/HoloLensSystemIncludes.h
--- a/Engine/Platforms/Hololens/Source/Runtime/Core/Public/HoloLensSystemIncludes.h
+++ b/Engine/Platforms/Hololens/Source/Runtime/Core/Public/HoloLensSystemIncludes.h
@@ -9,7 +9,7 @@
 #endif
 
 #include "HoloLensPlatformCompilerSetup.h"
-#include "Windows/MinimalWindowsApi.h"
+#include "Microsoft/MinimalWindowsApi.h"
 
 // Macro for releasing COM objects
 #define SAFE_RELEASE(p) { if(p) { (p)->Release(); (p)=NULL; } }
diff --git a/Engine/Platforms/Hololens/Source/Runtime/Core/Public/Microsoft/AllowMicrosoftPlatformAtomicsPrivate.h b/Engine/Platforms/Hololens/Source/Runtime/Core/Public/Microsoft/AllowMicrosoftPlatformAtomicsPrivate.h
--- a/Engine/Platforms/Hololens/Source/Runtime/Core/Public/Microsoft/AllowMicrosoftPlatformAtomicsPrivate.h
+++ b/Engine/Platforms/Hololens/Source/Runtime/Core/Public/Microsoft/AllowMicrosoftPlatformAtomicsPrivate.h
@@ -3,7 +3,7 @@
 #ifndef WINDOWS_PLATFORM_ATOMICS_GUARD
 	#define WINDOWS_PLATFORM_ATOMICS_GUARD
 #else
-	#error Nesting AllowWindowsPlatformAtomics.h is not allowed!
+	#error Nesting AllowMicrosoftPlatformAtomics.h is not allowed!
 #endif
 
 #define InterlockedIncrement _InterlockedIncrement
@@ -12,6 +12,9 @@
 #define InterlockedExchange _InterlockedExchange
 #define InterlockedExchangeAdd _InterlockedExchangeAdd
 #define InterlockedCompareExchange _InterlockedCompareExchange
+#define InterlockedAnd _InterlockedAnd
+#define InterlockedOr _InterlockedOr
+#define InterlockedXor _InterlockedXor
 
 #if PLATFORM_64BITS
 	#define InterlockedCompareExchangePointer _InterlockedCompareExchangePointer
diff --git a/Engine/Platforms/Hololens/Source/Runtime/Core/Public/Microsoft/HideMicrosoftPlatformAtomicsPrivate.h b/Engine/Platforms/Hololens/Source/Runtime/Core/Public/Microsoft/HideMicrosoftPlatformAtomicsPrivate.h
--- a/Engine/Platforms/Hololens/Source/Runtime/Core/Public/Microsoft/HideMicrosoftPlatformAtomicsPrivate.h
+++ b/Engine/Platforms/Hololens/Source/Runtime/Core/Public/Microsoft/HideMicrosoftPlatformAtomicsPrivate.h
@@ -18,4 +18,7 @@
 #undef InterlockedCompareExchange64
 #undef InterlockedIncrement64
 #undef InterlockedDecrement64
+#undef InterlockedAnd
+#undef InterlockedOr
+#undef InterlockedXor
 
diff --git a/Engine/Platforms/Hololens/Source/Runtime/Core/Public/Microsoft/PostWindowsApiPrivate.h b/Engine/Platforms/Hololens/Source/Runtime/Core/Public/Microsoft/PostWindowsApiPrivate.h
--- a/Engine/Platforms/Hololens/Source/Runtime/Core/Public/Microsoft/PostWindowsApiPrivate.h
+++ b/Engine/Platforms/Hololens/Source/Runtime/Core/Public/Microsoft/PostWindowsApiPrivate.h
@@ -44,7 +44,7 @@ THIRD_PARTY_INCLUDES_END
 #undef GetTempFileName
 #undef GetFreeSpace
 
-// Undefine all the atomics. AllowWindowsPlatformAtomics/HideWindowsPlatformAtomics temporarily defining these macros.
+// Undefine all the atomics. AllowMicrosoftPlatformAtomics/HideMicrosoftPlatformAtomics temporarily defining these macros.
 #undef InterlockedIncrement
 #undef InterlockedDecrement
 #undef InterlockedAdd
@@ -57,6 +57,9 @@ THIRD_PARTY_INCLUDES_END
 #undef InterlockedCompareExchange64
 #undef InterlockedIncrement64
 #undef InterlockedDecrement64
+#undef InterlockedAnd
+#undef InterlockedOr
+#undef InterlockedXor
 
 // Restore any previously defined macros
 #pragma pop_macro("MAX_uint8")
@@ -72,4 +75,3 @@ THIRD_PARTY_INCLUDES_END
 #if !defined(_XTL_) && (_WIN32_WINNT < 0x0403)
 #error SetCriticalSectionSpinCount requires _WIN32_WINNT >= 0x0403
 #endif
-#undef InterlockedDecrement64
diff --git a/Engine/Platforms/Hololens/Source/Runtime/TraceLog/Private/HoloLensTrace.cpp b/Engine/Platforms/Hololens/Source/Runtime/TraceLog/Private/HoloLensTrace.cpp
--- a/Engine/Platforms/Hololens/Source/Runtime/TraceLog/Private/HoloLensTrace.cpp
+++ b/Engine/Platforms/Hololens/Source/Runtime/TraceLog/Private/HoloLensTrace.cpp
@@ -4,13 +4,13 @@
 
 #if UE_TRACE_ENABLED && PLATFORM_HOLOLENS
 
-#include "Windows/AllowWindowsPlatformTypes.h"
-#	include "Windows/WindowsHWrapper.h"
+#include "Microsoft/AllowMicrosoftPlatformTypes.h"
+#	include "Microsoft/WindowsHWrapper.h"
 #	define _WINSOCK_DEPRECATED_NO_WARNINGS
 #	include <winsock2.h>
 #	include <ws2tcpip.h>
 #	pragma comment(lib, "ws2_32.lib")
-#include "Windows/HideWindowsPlatformTypes.h"
+#include "Microsoft/HideMicrosoftPlatformTypes.h"
 
 #pragma warning(push)
 #pragma warning(disable : 6031) // WSAStartup() return ignore  - we're error tolerant
diff --git a/Engine/Plugins/Runtime/WindowsMixedReality/Source/HoloLensAR/Private/HoloLensCameraImageTexture.h b/Engine/Plugins/Runtime/WindowsMixedReality/Source/HoloLensAR/Private/HoloLensCameraImageTexture.h
--- a/Engine/Plugins/Runtime/WindowsMixedReality/Source/HoloLensAR/Private/HoloLensCameraImageTexture.h
+++ b/Engine/Plugins/Runtime/WindowsMixedReality/Source/HoloLensAR/Private/HoloLensCameraImageTexture.h
@@ -6,7 +6,7 @@
 #include "ARTextures.h"
 
 #if SUPPORTS_WINDOWS_MIXED_REALITY_AR
-	#include "Windows/AllowWindowsPlatformTypes.h"
+	#include "Microsoft/AllowMicrosoftPlatformTypes.h"
 
 	THIRD_PARTY_INCLUDES_START
 		#include <windows.h>
@@ -16,7 +16,7 @@
 	THIRD_PARTY_INCLUDES_END
 
 	#include "Microsoft/COMPointer.h"
-	#include "Windows/HideWindowsPlatformTypes.h"
+	#include "Microsoft/HideMicrosoftPlatformTypes.h"
 #endif
 
 #include "HoloLensCameraImageTexture.generated.h"
diff --git a/Engine/Source/Runtime/Core/Public/Microsoft/MinWindows.h b/Engine/Source/Runtime/Core/Public/Microsoft/MinWindows.h
--- a/Engine/Source/Runtime/Core/Public/Microsoft/MinWindows.h
+++ b/Engine/Source/Runtime/Core/Public/Microsoft/MinWindows.h
@@ -3,8 +3,6 @@
 
 #if PLATFORM_WINDOWS
 	#include "Windows/MinWindows.h"
-#elif PLATFORM_HOLOLENS
-	#include "HoloLens/MinWindows.h"
 #else
 	#include "Microsoft/MinWindowsPrivate.h"
 #endif
diff --git a/Engine/Source/Runtime/Core/Public/Microsoft/PostWindowsApi.h b/Engine/Source/Runtime/Core/Public/Microsoft/PostWindowsApi.h
--- a/Engine/Source/Runtime/Core/Public/Microsoft/PostWindowsApi.h
+++ b/Engine/Source/Runtime/Core/Public/Microsoft/PostWindowsApi.h
@@ -2,8 +2,6 @@
 
 #if PLATFORM_WINDOWS
 	#include "Windows/PostWindowsApi.h"
-#elif PLATFORM_HOLOLENS
-	#include "HoloLens/PostWindowsApi.h"
 #else
 	#include "Microsoft/PostWindowsApiPrivate.h"
 #endif
diff --git a/Engine/Source/Runtime/Core/Public/Microsoft/PreWindowsApi.h b/Engine/Source/Runtime/Core/Public/Microsoft/PreWindowsApi.h
--- a/Engine/Source/Runtime/Core/Public/Microsoft/PreWindowsApi.h
+++ b/Engine/Source/Runtime/Core/Public/Microsoft/PreWindowsApi.h
@@ -4,8 +4,6 @@
 
 #if PLATFORM_WINDOWS
 	#include "Windows/PreWindowsApi.h"
-#elif PLATFORM_HOLOLENS
-	#include "HoloLens/PreWindowsApi.h"
 #else
 	#include "Microsoft/PreWindowsApiPrivate.h"
 #endif
diff --git a/Engine/Source/Runtime/Core/Public/Windows/AllowWindowsPlatformAtomics.h b/Engine/Source/Runtime/Core/Public/Windows/AllowWindowsPlatformAtomics.h
--- a/Engine/Source/Runtime/Core/Public/Windows/AllowWindowsPlatformAtomics.h
+++ b/Engine/Source/Runtime/Core/Public/Windows/AllowWindowsPlatformAtomics.h
@@ -2,7 +2,7 @@
 
 // #TODO: redirect to platform-agnostic version for the time being. Eventually this will become an error
 #include "HAL/Platform.h"
-#if !PLATFORM_WINDOWS && !PLATFORM_HOLOLENS
+#if !PLATFORM_WINDOWS
 	#include "Microsoft/AllowMicrosoftPlatformAtomics.h"
 #else
 
diff --git a/Engine/Source/Runtime/Core/Public/Windows/AllowWindowsPlatformTypes.h b/Engine/Source/Runtime/Core/Public/Windows/AllowWindowsPlatformTypes.h
--- a/Engine/Source/Runtime/Core/Public/Windows/AllowWindowsPlatformTypes.h
+++ b/Engine/Source/Runtime/Core/Public/Windows/AllowWindowsPlatformTypes.h
@@ -2,7 +2,7 @@
 
 // #TODO: redirect to platform-agnostic version for the time being. Eventually this will become an error
 #include "HAL/Platform.h"
-#if !PLATFORM_WINDOWS && !PLATFORM_HOLOLENS
+#if !PLATFORM_WINDOWS
 	#include "Microsoft/AllowMicrosoftPlatformTypes.h"
 #else
 
diff --git a/Engine/Source/Runtime/Core/Public/Windows/HideWindowsPlatformAtomics.h b/Engine/Source/Runtime/Core/Public/Windows/HideWindowsPlatformAtomics.h
--- a/Engine/Source/Runtime/Core/Public/Windows/HideWindowsPlatformAtomics.h
+++ b/Engine/Source/Runtime/Core/Public/Windows/HideWindowsPlatformAtomics.h
@@ -2,7 +2,7 @@
 
 // #TODO: redirect to platform-agnostic version for the time being. Eventually this will become an error
 #include "HAL/Platform.h"
-#if !PLATFORM_WINDOWS && !PLATFORM_HOLOLENS
+#if !PLATFORM_WINDOWS
 	#include "Microsoft/HideMicrosoftPlatformAtomics.h"
 #else
 
diff --git a/Engine/Source/Runtime/Core/Public/Windows/HideWindowsPlatformTypes.h b/Engine/Source/Runtime/Core/Public/Windows/HideWindowsPlatformTypes.h
--- a/Engine/Source/Runtime/Core/Public/Windows/HideWindowsPlatformTypes.h
+++ b/Engine/Source/Runtime/Core/Public/Windows/HideWindowsPlatformTypes.h
@@ -2,7 +2,7 @@
 
 // #TODO: redirect to platform-agnostic version for the time being. Eventually this will become an error
 #include "HAL/Platform.h"
-#if !PLATFORM_WINDOWS && !PLATFORM_HOLOLENS
+#if !PLATFORM_WINDOWS
 	#include "Microsoft/HideMicrosoftPlatformTypes.h"
 #else
 
diff --git a/Engine/Source/Runtime/Core/Public/Windows/MinWindows.h b/Engine/Source/Runtime/Core/Public/Windows/MinWindows.h
--- a/Engine/Source/Runtime/Core/Public/Windows/MinWindows.h
+++ b/Engine/Source/Runtime/Core/Public/Windows/MinWindows.h
@@ -4,7 +4,7 @@
 
 // #TODO: redirect to platform-agnostic version for the time being. Eventually this will become an error
 #include "HAL/Platform.h"
-#if !PLATFORM_WINDOWS && !PLATFORM_HOLOLENS
+#if !PLATFORM_WINDOWS
 	#include "Microsoft/MinWindows.h"
 #else
 
diff --git a/Engine/Source/Runtime/Core/Public/Windows/MinimalWindowsApi.h b/Engine/Source/Runtime/Core/Public/Windows/MinimalWindowsApi.h
--- a/Engine/Source/Runtime/Core/Public/Windows/MinimalWindowsApi.h
+++ b/Engine/Source/Runtime/Core/Public/Windows/MinimalWindowsApi.h
@@ -2,8 +2,8 @@
 
 #pragma once
 
-#if !PLATFORM_WINDOWS && !PLATFORM_HOLOLENS
-	#error this file is for PLATFORM_WINDOWS or PLATFORM_HOLOLENS only
+#if !PLATFORM_WINDOWS
+	#error this file is for PLATFORM_WINDOWS only
 #endif
 
 // redirect to new file
diff --git a/Engine/Source/Runtime/Core/Public/Windows/PostWindowsApi.h b/Engine/Source/Runtime/Core/Public/Windows/PostWindowsApi.h
--- a/Engine/Source/Runtime/Core/Public/Windows/PostWindowsApi.h
+++ b/Engine/Source/Runtime/Core/Public/Windows/PostWindowsApi.h
@@ -2,7 +2,7 @@
 
 // #TODO: redirect to platform-agnostic version for the time being. Eventually this will become an error
 #include "HAL/Platform.h"
-#if !PLATFORM_WINDOWS && !PLATFORM_HOLOLENS
+#if !PLATFORM_WINDOWS
 	#include "Microsoft/PostWindowsApi.h"
 #else
 
@@ -56,23 +56,21 @@ THIRD_PARTY_INCLUDES_END
 #undef GetFreeSpace
 
 // Undefine all the atomics. AllowWindowsPlatformAtomics/HideWindowsPlatformAtomics temporarily defining these macros.
-#if PLATFORM_WINDOWS || PLATFORM_HOLOLENS
-	#undef InterlockedIncrement
-	#undef InterlockedDecrement
-	#undef InterlockedAdd
-	#undef InterlockedExchange
-	#undef InterlockedExchangeAdd
-	#undef InterlockedCompareExchange
-	#undef InterlockedCompareExchangePointer
-	#undef InterlockedExchange64
-	#undef InterlockedExchangeAdd64
-	#undef InterlockedCompareExchange64
-	#undef InterlockedIncrement64
-	#undef InterlockedDecrement64
-	#undef InterlockedAnd
-	#undef InterlockedOr
-	#undef InterlockedXor
-#endif
+#undef InterlockedIncrement
+#undef InterlockedDecrement
+#undef InterlockedAdd
+#undef InterlockedExchange
+#undef InterlockedExchangeAdd
+#undef InterlockedCompareExchange
+#undef InterlockedCompareExchangePointer
+#undef InterlockedExchange64
+#undef InterlockedExchangeAdd64
+#undef InterlockedCompareExchange64
+#undef InterlockedIncrement64
+#undef InterlockedDecrement64
+#undef InterlockedAnd
+#undef InterlockedOr
+#undef InterlockedXor
 
 // Restore any previously defined macros
 #pragma pop_macro("MAX_uint8")
diff --git a/Engine/Source/Runtime/Core/Public/Windows/PreWindowsApi.h b/Engine/Source/Runtime/Core/Public/Windows/PreWindowsApi.h
--- a/Engine/Source/Runtime/Core/Public/Windows/PreWindowsApi.h
+++ b/Engine/Source/Runtime/Core/Public/Windows/PreWindowsApi.h
@@ -2,7 +2,7 @@
 
 // #TODO: redirect to platform-agnostic version for the time being. Eventually this will become an error
 #include "HAL/Platform.h"
-#if !PLATFORM_WINDOWS && !PLATFORM_HOLOLENS
+#if !PLATFORM_WINDOWS
 	#include "Microsoft/PreWindowsApi.h"
 #else
 
diff --git a/Engine/Source/Runtime/Core/Public/Windows/WindowsHWrapper.h b/Engine/Source/Runtime/Core/Public/Windows/WindowsHWrapper.h
--- a/Engine/Source/Runtime/Core/Public/Windows/WindowsHWrapper.h
+++ b/Engine/Source/Runtime/Core/Public/Windows/WindowsHWrapper.h
@@ -3,7 +3,7 @@
 
 // #TODO: redirect to platform-agnostic version for the time being. Eventually this will become an error
 #include "HAL/Platform.h"
-#if !PLATFORM_WINDOWS && !PLATFORM_HOLOLENS
+#if !PLATFORM_WINDOWS
 	#include "Microsoft/WindowsHWrapper.h"
 #else
 
