Add a way for the entire asset virtualization system to be disabled in code

#rnx

- A new define has been added 'UE_DISABLE_VIRTUALIZATION_SYSTEM' which when set to 1 will prevent the virtualization system from finding a valid system to mount, meaning it will default to the null implementation.
- The expected use it to add 'GlobalDefinitions.Add("UE_DISABLE_VIRTUALIZATION_SYSTEM=1");' to a target.cs to disable virtualization for a target
- This wouldn't be needed if the virtualization module was a plugin as it could be disabled by a .uproject file. We should consider moving the module.

diff --git a/Engine/Source/Runtime/Core/Private/Virtualization/VirtualizationSystem.cpp b/Engine/Source/Runtime/Core/Private/Virtualization/VirtualizationSystem.cpp
--- a/Engine/Source/Runtime/Core/Private/Virtualization/VirtualizationSystem.cpp
+++ b/Engine/Source/Runtime/Core/Private/Virtualization/VirtualizationSystem.cpp
@@ -6,6 +6,13 @@
 #include "Misc/App.h"
 #include "Misc/ConfigCacheIni.h"
 
+// Can be defined as 1 by programs target.cs files to disable the virtualization system
+// entirely. This could be removed in the future if we move the virtualization module
+// to be a plugin
+#ifndef UE_DISABLE_VIRTUALIZATION_SYSTEM
+	#define UE_DISABLE_VIRTUALIZATION_SYSTEM 0
+#endif
+
 namespace UE::Virtualization
 {
 
@@ -134,12 +141,16 @@ void Initialize(const FInitParams& InitParams)
 {
 	FName SystemName;
 
+#if UE_DISABLE_VIRTUALIZATION_SYSTEM == 0
 	FString RawSystemName;
 	if (InitParams.ConfigFile.GetString(TEXT("Core.ContentVirtualization"), TEXT("SystemName"), RawSystemName))
 	{
 		SystemName = FName(RawSystemName);
 		UE_LOG(LogVirtualization, Display, TEXT("VirtualizationSystem name found in ini file: %s"), *RawSystemName);
 	}
+#else
+	UE_LOG(LogVirtualization, Display, TEXT("The virtualization system has been disabled by code"));
+#endif //UE_DISABLE_VIRTUALIZATION_SYSTEM
 
 	if (!SystemName.IsNone())
 	{
