Fixed a few issues in MallocBinned

diff --git a/Engine/Source/Runtime/Core/Private/HAL/MallocBinned.cpp b/Engine/Source/Runtime/Core/Private/HAL/MallocBinned.cpp
--- a/Engine/Source/Runtime/Core/Private/HAL/MallocBinned.cpp
+++ b/Engine/Source/Runtime/Core/Private/HAL/MallocBinned.cpp
@@ -708,6 +708,7 @@ struct FMallocBinned::Private
 		if(FPlatformMemory::PtrIsOSMalloc(Ptr))
 		{
 			SmallOSFree(Allocator, Ptr, Private::SMALL_BLOCK_POOL_SIZE);
+			Allocator.bNanoMallocAvailable = true;
 			return;
 		}
 #endif
@@ -1156,6 +1157,7 @@ FMallocBinned::FMallocBinned(uint32 InPageSize, uint64 AddressLimit)
 #if USE_OS_SMALL_BLOCK_ALLOC
 	
 	FPlatformMemory::NanoMallocInit();
+	bNanoMallocAvailable = FPlatformMemory::IsNanoMallocAvailable();
 	
 #if USE_OS_SMALL_BLOCK_GRAB_MEMORY_FROM_OS
 
@@ -1293,10 +1295,11 @@ void* FMallocBinned::Malloc(SIZE_T Size, uint32 Alignment)
 	BINNED_INCREMENT_STATCOUNTER(CurrentAllocs);
 	BINNED_INCREMENT_STATCOUNTER(TotalAllocs);
 	
+	Size = FMath::Max(Size, (SIZE_T)1);
 	FFreeMem* Free = nullptr;
 	bool bUsePools = true;
 #if USE_OS_SMALL_BLOCK_ALLOC && !USE_OS_SMALL_BLOCK_GRAB_MEMORY_FROM_OS
-	if (FPlatformMemory::IsNanoMallocAvailable() && (Size <= Private::SMALL_BLOCK_POOL_SIZE) && (Alignment <= Private::DEFAULT_BINNED_ALLOCATOR_ALIGNMENT))
+	if (bNanoMallocAvailable && (Size <= Private::SMALL_BLOCK_POOL_SIZE) && (Alignment <= Private::DEFAULT_BINNED_ALLOCATOR_ALIGNMENT))
 	{
 		//Make sure we have initialized our hash buckets even if we are using the NANO_MALLOC grabber, as otherwise we can end
 		//up making bad assumptions and trying to grab invalid data during a Realloc of this data.
@@ -1317,6 +1320,7 @@ void* FMallocBinned::Malloc(SIZE_T Size, uint32 Alignment)
 			Private::SmallOSFree(*this, Free, Size);
 			bUsePools = true;
 			Free = nullptr;
+			bNanoMallocAvailable = false;
 		}
 	}
 #endif
@@ -1436,8 +1440,10 @@ void* FMallocBinned::Realloc( void* Ptr, SIZE_T NewSize, uint32 Alignment )
 				// Fall back to UE's allocator
 				Ptr = NewPtr;
 				NewPtr = Malloc(NewSizeUnmodified, Alignment);
-				FMemory::Memcpy(NewPtr, Ptr, NewSize);
+				const SIZE_T OldSize = malloc_size(Ptr);
+				FMemory::Memcpy(NewPtr, Ptr, OldSize);
 				Private::SmallOSFree(*this, Ptr, NewSize);
+				bNanoMallocAvailable = true;
 			}
 		}
 		else
@@ -1521,7 +1527,7 @@ bool FMallocBinned::GetAllocationSize(void *Original, SIZE_T &SizeOut)
 #if USE_OS_SMALL_BLOCK_ALLOC && !USE_OS_SMALL_BLOCK_GRAB_MEMORY_FROM_OS
 	if(FPlatformMemory::PtrIsOSMalloc(Original))
 	{
-		SizeOut = Private::SMALL_BLOCK_POOL_SIZE;
+		SizeOut = malloc_size(Original);
 		return true;
 	}
 #endif
diff --git a/Engine/Source/Runtime/Core/Public/HAL/MallocBinned.h b/Engine/Source/Runtime/Core/Public/HAL/MallocBinned.h
--- a/Engine/Source/Runtime/Core/Public/HAL/MallocBinned.h
+++ b/Engine/Source/Runtime/Core/Public/HAL/MallocBinned.h
@@ -268,6 +268,10 @@ private:
 	
 #endif //STATS
 
+#if USE_OS_SMALL_BLOCK_ALLOC
+	bool			bNanoMallocAvailable;
+#endif
+
 public:
 	// FMalloc interface.
 	// InPageSize - First parameter is page size, all allocs from BinnedAllocFromOS() MUST be aligned to this size
