Fixed MallocBinned compilation on Mac

diff --git a/Engine/Source/Runtime/Core/Private/HAL/MallocBinned.cpp b/Engine/Source/Runtime/Core/Private/HAL/MallocBinned.cpp
--- a/Engine/Source/Runtime/Core/Private/HAL/MallocBinned.cpp
+++ b/Engine/Source/Runtime/Core/Private/HAL/MallocBinned.cpp
@@ -708,7 +708,6 @@ struct FMallocBinned::Private
 		if(FPlatformMemory::PtrIsOSMalloc(Ptr))
 		{
 			SmallOSFree(Allocator, Ptr, Private::SMALL_BLOCK_POOL_SIZE);
-			Allocator.bNanoMallocAvailable = true;
 			return;
 		}
 #endif
@@ -886,6 +885,7 @@ struct FMallocBinned::Private
 	{
 #if PLATFORM_IOS
 		::free(Ptr);
+		Allocator.bNanoMallocAvailable = true;
 #else
 		FPlatformMemory::BinnedFreeToOS(Ptr, Size);
 #endif
@@ -1443,7 +1443,6 @@ void* FMallocBinned::Realloc( void* Ptr, SIZE_T NewSize, uint32 Alignment )
 				const SIZE_T OldSize = malloc_size(Ptr);
 				FMemory::Memcpy(NewPtr, Ptr, OldSize);
 				Private::SmallOSFree(*this, Ptr, NewSize);
-				bNanoMallocAvailable = true;
 			}
 		}
 		else
