Logging: Changed CanBeUsedOnPanicThread() to return false by default

The previous default of CanBeUsedOnAnyThread() proved unsafe since some output devices can safely be used on any thread but cannot safely be used during a crash. CanBeUsedOnAnyThread() was used pre-5.1 to control serialization to an output device during a crash, but optimizations in 5.1 have uncovered that some of these were never safe.

#rnx

diff --git a/Engine/Source/Programs/CrashReportClient/Private/CrashReportAnalyticsSessionSummary.cpp b/Engine/Source/Programs/CrashReportClient/Private/CrashReportAnalyticsSessionSummary.cpp
--- a/Engine/Source/Programs/CrashReportClient/Private/CrashReportAnalyticsSessionSummary.cpp
+++ b/Engine/Source/Programs/CrashReportClient/Private/CrashReportAnalyticsSessionSummary.cpp
@@ -451,6 +451,11 @@ bool FCrashReportAnalyticsSessionSummary::CanBeUsedOnMultipleThreads() const
 	return true;
 }
 
+bool FCrashReportAnalyticsSessionSummary::CanBeUsedOnPanicThread() const
+{
+	return true;
+}
+
 void FCrashReportAnalyticsSessionSummary::LogEvent(const FString& Event)
 {
 	LogEvent(*Event);
diff --git a/Engine/Source/Programs/CrashReportClient/Private/CrashReportAnalyticsSessionSummary.h b/Engine/Source/Programs/CrashReportClient/Private/CrashReportAnalyticsSessionSummary.h
--- a/Engine/Source/Programs/CrashReportClient/Private/CrashReportAnalyticsSessionSummary.h
+++ b/Engine/Source/Programs/CrashReportClient/Private/CrashReportAnalyticsSessionSummary.h
@@ -39,6 +39,7 @@ public:
 	virtual void Serialize(const TCHAR* V, ELogVerbosity::Type Verbosity, const FName& Category, const double Time) override;
 	virtual bool CanBeUsedOnAnyThread() const override;
 	virtual bool CanBeUsedOnMultipleThreads() const override;
+	virtual bool CanBeUsedOnPanicThread() const override;
 
 	/** Logs an events in the analytics session summary diagnostic log. The event text is expected to be short and concise. */
 	void LogEvent(const TCHAR* Event, bool bForwardToUELog = true);
diff --git a/Engine/Source/Runtime/ApplicationCore/Private/Linux/LinuxConsoleOutputDevice.cpp b/Engine/Source/Runtime/ApplicationCore/Private/Linux/LinuxConsoleOutputDevice.cpp
--- a/Engine/Source/Runtime/ApplicationCore/Private/Linux/LinuxConsoleOutputDevice.cpp
+++ b/Engine/Source/Runtime/ApplicationCore/Private/Linux/LinuxConsoleOutputDevice.cpp
@@ -106,7 +106,12 @@ void FLinuxConsoleOutputDevice::Serialize(const TCHAR* Data, ELogVerbosity::Type
 	}
 }
 
-bool FLinuxConsoleOutputDevice::CanBeUsedOnAnyThread() const 
+bool FLinuxConsoleOutputDevice::CanBeUsedOnAnyThread() const
+{
+	return true;
+}
+
+bool FLinuxConsoleOutputDevice::CanBeUsedOnPanicThread() const
 {
 	return true;
 }
diff --git a/Engine/Source/Runtime/ApplicationCore/Private/Windows/WindowsConsoleOutputDevice.cpp b/Engine/Source/Runtime/ApplicationCore/Private/Windows/WindowsConsoleOutputDevice.cpp
--- a/Engine/Source/Runtime/ApplicationCore/Private/Windows/WindowsConsoleOutputDevice.cpp
+++ b/Engine/Source/Runtime/ApplicationCore/Private/Windows/WindowsConsoleOutputDevice.cpp
@@ -445,7 +445,7 @@ bool FWindowsConsoleOutputDevice::IsAttached()
 	}
 	else
 	{
-		FreeConsole();		
+		FreeConsole();
 	}
 	return false;
 }
@@ -455,3 +455,7 @@ bool FWindowsConsoleOutputDevice::CanBeUsedOnAnyThread() const
 	return true;
 }
 
+bool FWindowsConsoleOutputDevice::CanBeUsedOnPanicThread() const 
+{
+	return true;
+}
diff --git a/Engine/Source/Runtime/ApplicationCore/Public/Linux/LinuxConsoleOutputDevice.h b/Engine/Source/Runtime/ApplicationCore/Public/Linux/LinuxConsoleOutputDevice.h
--- a/Engine/Source/Runtime/ApplicationCore/Public/Linux/LinuxConsoleOutputDevice.h
+++ b/Engine/Source/Runtime/ApplicationCore/Public/Linux/LinuxConsoleOutputDevice.h
@@ -50,6 +50,7 @@ public:
 	virtual bool IsAttached() {return false;}
 
 	virtual bool CanBeUsedOnAnyThread() const override;
+	virtual bool CanBeUsedOnPanicThread() const override;
 
 	/**
 	 * Displays text on the console and scrolls if necessary.
diff --git a/Engine/Source/Runtime/ApplicationCore/Public/Windows/WindowsConsoleOutputDevice.h b/Engine/Source/Runtime/ApplicationCore/Public/Windows/WindowsConsoleOutputDevice.h
--- a/Engine/Source/Runtime/ApplicationCore/Public/Windows/WindowsConsoleOutputDevice.h
+++ b/Engine/Source/Runtime/ApplicationCore/Public/Windows/WindowsConsoleOutputDevice.h
@@ -84,6 +84,7 @@ public:
 	virtual bool IsAttached();
 
 	virtual bool CanBeUsedOnAnyThread() const override;
+	virtual bool CanBeUsedOnPanicThread() const override;
 
 	/**
 	 * Displays text on the console and scrolls if necessary.
diff --git a/Engine/Source/Runtime/Core/Private/Misc/OutputDeviceRedirector.cpp b/Engine/Source/Runtime/Core/Private/Misc/OutputDeviceRedirector.cpp
--- a/Engine/Source/Runtime/Core/Private/Misc/OutputDeviceRedirector.cpp
+++ b/Engine/Source/Runtime/Core/Private/Misc/OutputDeviceRedirector.cpp
@@ -2,6 +2,7 @@
 
 #include "Misc/OutputDeviceRedirector.h"
 
+#include "Containers/BitArray.h"
 #include "Containers/DepletableMpscQueue.h"
 #include "Experimental/ConcurrentLinearAllocator.h"
 #include "HAL/Event.h"
@@ -133,12 +134,12 @@ struct FOutputDeviceRedirectorState
 	/** A lock to synchronize access to the thread. */
 	FRWLock ThreadLock;
 
-	/** An event to wake the dedicated master thread to process buffered lines. */
-	std::atomic<FEvent*> ThreadWakeEvent = nullptr;
-
 	/** A queue of events to trigger when the dedicated master thread is idle. */
 	TDepletableMpscQueue<FEvent*, FOutputDeviceLinearAllocator> ThreadIdleEvents;
 
+	/** An event to wake the dedicated master thread to process buffered lines. */
+	std::atomic<FEvent*> ThreadWakeEvent = nullptr;
+
 	/** The ID of the thread holding the master lock. */
 	std::atomic<uint32> LockedThreadId = MAX_uint32;
 
@@ -151,6 +152,10 @@ struct FOutputDeviceRedirectorState
 	/** Whether the backlog is enabled. */
 	bool bEnableBacklog = false;
 
+	/** Whether the output device at the corresponding index can be used on the panic thread. */
+	TBitArray<TInlineAllocator<1>> BufferedOutputDevicesCanBeUsedOnPanicThread;
+	TBitArray<TInlineAllocator<1>> UnbufferedOutputDevicesCanBeUsedOnPanicThread;
+
 	bool IsMasterThread(const uint32 ThreadId) const
 	{
 		return ThreadId == MasterThreadId.load(std::memory_order_relaxed);
@@ -171,6 +176,9 @@ struct FOutputDeviceRedirectorState
 		return LocalPanicThreadId == MAX_uint32 || LocalPanicThreadId == ThreadId;
 	}
 
+	void AddOutputDevice(FOutputDevice* OutputDevice);
+	void RemoveOutputDevice(FOutputDevice* OutputDevice);
+
 	bool TryStartThread();
 	bool TryStopThread();
 
@@ -179,12 +187,18 @@ struct FOutputDeviceRedirectorState
 	void FlushBufferedLines();
 
 	template <typename OutputDevicesType, typename FunctionType, typename... ArgTypes>
-	FORCEINLINE void BroadcastTo(const uint32 ThreadId, const OutputDevicesType& OutputDevices, FunctionType&& Function, ArgTypes&&... Args)
-	{
+	FORCEINLINE void BroadcastTo(
+		const uint32 ThreadId,
+		const OutputDevicesType& OutputDevices,
+		const TBitArray<TInlineAllocator<1>>& CanBeUsedOnPanicThread,
+		FunctionType&& Function,
+		ArgTypes&&... Args)
+	{
+		int32 Index = 0;
 		const bool bIsPanicThread = IsPanicThread(ThreadId);
 		for (FOutputDevice* OutputDevice : OutputDevices)
 		{
-			if (!bIsPanicThread || OutputDevice->CanBeUsedOnPanicThread())
+			if (!bIsPanicThread || CanBeUsedOnPanicThread[Index++])
 			{
 				Invoke(Function, OutputDevice, Forward<ArgTypes>(Args)...);
 			}
@@ -343,6 +357,42 @@ private:
 	bool bLocked = false;
 };
 
+void FOutputDeviceRedirectorState::AddOutputDevice(FOutputDevice* OutputDevice)
+{
+	const auto AddTo = [OutputDevice](TArray<FOutputDevice*>& OutputDevices, TBitArray<TInlineAllocator<1>>& Flags)
+	{
+		const int32 Count = OutputDevices.Num();
+		if (OutputDevices.AddUnique(OutputDevice) == Count)
+		{
+			Flags.Add(OutputDevice->CanBeUsedOnPanicThread());
+		}
+	};
+	UE::Private::FOutputDevicesWriteScopeLock ScopeLock(*this);
+	if (OutputDevice->CanBeUsedOnMultipleThreads())
+	{
+		AddTo(UnbufferedOutputDevices, UnbufferedOutputDevicesCanBeUsedOnPanicThread);
+	}
+	else
+	{
+		AddTo(BufferedOutputDevices, BufferedOutputDevicesCanBeUsedOnPanicThread);
+	}
+}
+
+void FOutputDeviceRedirectorState::RemoveOutputDevice(FOutputDevice* OutputDevice)
+{
+	const auto RemoveFrom = [OutputDevice](TArray<FOutputDevice*>& OutputDevices, TBitArray<TInlineAllocator<1>>& Flags)
+	{
+		if (const int32 Index = OutputDevices.FindLast(OutputDevice); Index != INDEX_NONE)
+		{
+			OutputDevices.RemoveAt(Index);
+			Flags.RemoveAt(Index);
+		}
+	};
+	UE::Private::FOutputDevicesWriteScopeLock ScopeLock(*this);
+	RemoveFrom(BufferedOutputDevices, BufferedOutputDevicesCanBeUsedOnPanicThread);
+	RemoveFrom(UnbufferedOutputDevices, UnbufferedOutputDevicesCanBeUsedOnPanicThread);
+}
+
 bool FOutputDeviceRedirectorState::TryStartThread()
 {
 	if (FWriteScopeLock ThreadScopeLock(ThreadLock); !ThreadWakeEvent.load(std::memory_order_relaxed))
@@ -402,7 +452,8 @@ void FOutputDeviceRedirectorState::FlushBufferedLines()
 	const uint32 ThreadId = FPlatformTLS::GetCurrentThreadId();
 	BufferedLines.Deplete([this, ThreadId](UE::Private::FOutputDeviceLine&& Line)
 	{
-		BroadcastTo(ThreadId, BufferedOutputDevices, UE_PROJECTION_MEMBER(FOutputDevice, Serialize),
+		BroadcastTo(ThreadId, BufferedOutputDevices, BufferedOutputDevicesCanBeUsedOnPanicThread,
+			UE_PROJECTION_MEMBER(FOutputDevice, Serialize),
 			Line.Data, Line.Verbosity, Line.Category, Line.Time);
 	});
 }
@@ -424,15 +475,7 @@ void FOutputDeviceRedirector::AddOutputDevice(FOutputDevice* OutputDevice)
 {
 	if (OutputDevice)
 	{
-		UE::Private::FOutputDevicesWriteScopeLock ScopeLock(*State);
-		if (OutputDevice->CanBeUsedOnMultipleThreads())
-		{
-			State->UnbufferedOutputDevices.AddUnique(OutputDevice);
-		}
-		else
-		{
-			State->BufferedOutputDevices.AddUnique(OutputDevice);
-		}
+		State->AddOutputDevice(OutputDevice);
 	}
 }
 
@@ -440,9 +483,7 @@ void FOutputDeviceRedirector::RemoveOutputDevice(FOutputDevice* OutputDevice)
 {
 	if (OutputDevice)
 	{
-		UE::Private::FOutputDevicesWriteScopeLock ScopeLock(*State);
-		State->BufferedOutputDevices.Remove(OutputDevice);
-		State->UnbufferedOutputDevices.Remove(OutputDevice);
+		State->RemoveOutputDevice(OutputDevice);
 	}
 }
 
@@ -537,7 +578,8 @@ void FOutputDeviceRedirector::Serialize(const TCHAR* const Data, const ELogVerbo
 	const uint32 ThreadId = FPlatformTLS::GetCurrentThreadId();
 
 	// Serialize directly to any output devices which don't require buffering
-	State->BroadcastTo(ThreadId, State->UnbufferedOutputDevices, UE_PROJECTION_MEMBER(FOutputDevice, Serialize),
+	State->BroadcastTo(ThreadId, State->UnbufferedOutputDevices, State->UnbufferedOutputDevicesCanBeUsedOnPanicThread,
+		UE_PROJECTION_MEMBER(FOutputDevice, Serialize),
 		Data, Verbosity, Category, RealTime);
 
 	// Serialize to the backlog when not in panic mode. This will deadlock in panic mode when the
@@ -557,7 +599,8 @@ void FOutputDeviceRedirector::Serialize(const TCHAR* const Data, const ELogVerbo
 		if (UE::Private::FOutputDevicesMasterScope MasterLock(*State); MasterLock.IsLocked() && State->IsMasterThread(ThreadId))
 		{
 			State->FlushBufferedLines();
-			State->BroadcastTo(ThreadId, State->BufferedOutputDevices, UE_PROJECTION_MEMBER(FOutputDevice, Serialize),
+			State->BroadcastTo(ThreadId, State->BufferedOutputDevices, State->BufferedOutputDevicesCanBeUsedOnPanicThread,
+				UE_PROJECTION_MEMBER(FOutputDevice, Serialize),
 				Data, Verbosity, Category, RealTime);
 			if (UNLIKELY(State->IsPanicThread(ThreadId)))
 			{
@@ -598,8 +641,8 @@ void FOutputDeviceRedirector::Flush()
 	{
 		State->FlushBufferedLines();
 		const uint32 ThreadId = FPlatformTLS::GetCurrentThreadId();
-		State->BroadcastTo(ThreadId, State->BufferedOutputDevices, &FOutputDevice::Flush);
-		State->BroadcastTo(ThreadId, State->UnbufferedOutputDevices, &FOutputDevice::Flush);
+		State->BroadcastTo(ThreadId, State->BufferedOutputDevices, State->BufferedOutputDevicesCanBeUsedOnPanicThread, &FOutputDevice::Flush);
+		State->BroadcastTo(ThreadId, State->UnbufferedOutputDevices, State->UnbufferedOutputDevicesCanBeUsedOnPanicThread, &FOutputDevice::Flush);
 	}
 }
 
diff --git a/Engine/Source/Runtime/Core/Public/Misc/OutputDevice.h b/Engine/Source/Runtime/Core/Public/Misc/OutputDevice.h
--- a/Engine/Source/Runtime/Core/Public/Misc/OutputDevice.h
+++ b/Engine/Source/Runtime/Core/Public/Misc/OutputDevice.h
@@ -201,10 +201,11 @@ public:
 
 	/**
 	 * @return whether this output device can be used after a panic (crash or fatal error) has been flagged.
+	 * @note The return value is cached by AddOutputDevice because calling this during a panic may fail.
 	 */
 	virtual bool CanBeUsedOnPanicThread() const
 	{
-		return CanBeUsedOnAnyThread();
+		return false;
 	}
 
 	// Simple text printing.
diff --git a/Engine/Source/Runtime/Core/Public/Misc/OutputDeviceDebug.h b/Engine/Source/Runtime/Core/Public/Misc/OutputDeviceDebug.h
--- a/Engine/Source/Runtime/Core/Public/Misc/OutputDeviceDebug.h
+++ b/Engine/Source/Runtime/Core/Public/Misc/OutputDeviceDebug.h
@@ -24,5 +24,10 @@ public:
 	}
 
 	virtual bool CanBeUsedOnMultipleThreads() const override;
+
+	virtual bool CanBeUsedOnPanicThread() const override
+	{
+		return true;
+	}
 };
 
diff --git a/Engine/Source/Runtime/Core/Public/Misc/OutputDeviceError.h b/Engine/Source/Runtime/Core/Public/Misc/OutputDeviceError.h
--- a/Engine/Source/Runtime/Core/Public/Misc/OutputDeviceError.h
+++ b/Engine/Source/Runtime/Core/Public/Misc/OutputDeviceError.h
@@ -9,6 +9,11 @@
 class CORE_API FOutputDeviceError : public FOutputDevice
 {
 public:
+	virtual bool CanBeUsedOnPanicThread() const override
+	{
+		return true;
+	}
+
 	virtual void HandleError()=0;
 
 	/** Sets the location of the instruction that raise the next error */
diff --git a/Engine/Source/Runtime/Core/Public/Misc/OutputDeviceFile.h b/Engine/Source/Runtime/Core/Public/Misc/OutputDeviceFile.h
--- a/Engine/Source/Runtime/Core/Public/Misc/OutputDeviceFile.h
+++ b/Engine/Source/Runtime/Core/Public/Misc/OutputDeviceFile.h
@@ -147,6 +147,10 @@ public:
 	{
 		return true;
 	}
+	virtual bool CanBeUsedOnPanicThread() const override
+	{
+		return true;
+	}
 	//~ End FOutputDevice Interface.
 
 	/** Creates a backup copy of a log file if it already exists */
diff --git a/Engine/Source/Runtime/Core/Public/Misc/OutputDeviceMemory.h b/Engine/Source/Runtime/Core/Public/Misc/OutputDeviceMemory.h
--- a/Engine/Source/Runtime/Core/Public/Misc/OutputDeviceMemory.h
+++ b/Engine/Source/Runtime/Core/Public/Misc/OutputDeviceMemory.h
@@ -58,6 +58,10 @@ public:
 	{
 		return true;
 	}
+	virtual bool CanBeUsedOnPanicThread() const override
+	{
+		return true;
+	}
 	virtual bool IsMemoryOnly() const override
 	{
 		return true;
diff --git a/Engine/Source/Runtime/Launch/Private/LaunchEngineLoop.cpp b/Engine/Source/Runtime/Launch/Private/LaunchEngineLoop.cpp
--- a/Engine/Source/Runtime/Launch/Private/LaunchEngineLoop.cpp
+++ b/Engine/Source/Runtime/Launch/Private/LaunchEngineLoop.cpp
@@ -377,6 +377,11 @@ public:
 		return true;
 	}
 
+	virtual bool CanBeUsedOnPanicThread() const override
+	{
+		return true;
+	}
+
 	virtual void Serialize( const TCHAR* V, ELogVerbosity::Type Verbosity, const class FName& Category ) override
 	{
 		if (Verbosity <= AllowedLogVerbosity)
