[Insights] Added more LLM_SCOPEs ("Trace/CpuProfiler" and "Trace/Stats").

diff --git a/Engine/Source/Runtime/Core/Private/ProfilingDebugging/CpuProfilerTrace.cpp b/Engine/Source/Runtime/Core/Private/ProfilingDebugging/CpuProfilerTrace.cpp
--- a/Engine/Source/Runtime/Core/Private/ProfilingDebugging/CpuProfilerTrace.cpp
+++ b/Engine/Source/Runtime/Core/Private/ProfilingDebugging/CpuProfilerTrace.cpp
@@ -171,6 +171,7 @@ void FCpuProfilerTrace::OutputBeginDynamicEvent(const ANSICHAR* Name, const ANSI
 	uint32 SpecId = ThreadBuffer->DynamicAnsiScopeNamesMap.FindRef(Name);
 	if (!SpecId)
 	{
+		LLM_SCOPE_BYNAME(TEXT("Trace/CpuProfiler"));
 		int32 NameSize = strlen(Name) + 1;
 		ANSICHAR* NameCopy = reinterpret_cast<ANSICHAR*>(ThreadBuffer->DynamicScopeNamesMemory.Alloc(NameSize, alignof(ANSICHAR)));
 		FMemory::Memmove(NameCopy, Name, NameSize);
@@ -186,6 +187,7 @@ void FCpuProfilerTrace::OutputBeginDynamicEvent(const TCHAR* Name, const ANSICHA
 	uint32 SpecId = ThreadBuffer->DynamicTCharScopeNamesMap.FindRef(Name);
 	if (!SpecId)
 	{
+		LLM_SCOPE_BYNAME(TEXT("Trace/CpuProfiler"));
 		int32 NameSize = (FCString::Strlen(Name) + 1) * sizeof(TCHAR);
 		TCHAR* NameCopy = reinterpret_cast<TCHAR*>(ThreadBuffer->DynamicScopeNamesMemory.Alloc(NameSize, alignof(TCHAR)));
 		FMemory::Memmove(NameCopy, Name, NameSize);
@@ -201,6 +203,7 @@ void FCpuProfilerTrace::OutputBeginDynamicEvent(const FName& Name, const ANSICHA
 	uint32 SpecId = ThreadBuffer->DynamicFNameScopeNamesMap.FindRef(Name.GetComparisonIndex());
 	if (!SpecId)
 	{
+		LLM_SCOPE_BYNAME(TEXT("Trace/CpuProfiler"));
 		const FNameEntry* NameEntry = Name.GetDisplayNameEntry();
 		if (NameEntry->IsWide())
 		{
diff --git a/Engine/Source/Runtime/Core/Private/Stats/StatsTrace.cpp b/Engine/Source/Runtime/Core/Private/Stats/StatsTrace.cpp
--- a/Engine/Source/Runtime/Core/Private/Stats/StatsTrace.cpp
+++ b/Engine/Source/Runtime/Core/Private/Stats/StatsTrace.cpp
@@ -3,6 +3,7 @@
 #include "ProfilingDebugging/MiscTrace.h"
 #include "Templates/Atomic.h"
 #include "Misc/CString.h"
+#include "HAL/LowLevelMemTracker.h"
 #include "HAL/PlatformTime.h"
 #include "HAL/PlatformTLS.h"
 #include "Trace/Trace.inl"
@@ -67,6 +68,7 @@ thread_local FStatsTraceInternal::FThreadState* FStatsTraceInternal::ThreadLocal
 
 FStatsTraceInternal::FThreadState* FStatsTraceInternal::InitThreadState()
 {
+	LLM_SCOPE_BYNAME(TEXT("Trace/Stats"));
 	ThreadLocalThreadState = new FThreadState();
 	ThreadLocalThreadState->BufferSize = 0;
 	ThreadLocalThreadState->LastCycle = 0;
@@ -197,4 +199,4 @@ void FStatsTrace::Set(const FName& Stat, double Value)
 	}
 }
 
-#endif
+#endif // STATSTRACE_ENABLED
