PThread: Added error checking to thread joining in Kill and WaitForCompletion.

Internally, we call pthread_join in these function, which may actually fail. Some of the reasons are:
   - A thread trying to join itself.
   - A deadlock e.g. two threads trying to join each other.
   - Another thread is already joining the thread.

If one of this cases happens, the following snippet would likely crash, because we used to unconditionally return 'true':

if (Thread->Kill(true))
    delete Thread;

In many places we don't even perform this check but just kill and delete threads. For this reason, I've also added a log message to provide some details if we ever crash because of this again. An assert could be a better way of dealing with it, but because Kill (public API) returns true/false to indicate if the thread has been really stopped, we can't assume that no one has code relying on it.

diff --git a/Engine/Source/Runtime/Core/Private/HAL/PThreadRunnableThread.h b/Engine/Source/Runtime/Core/Private/HAL/PThreadRunnableThread.h
--- a/Engine/Source/Runtime/Core/Private/HAL/PThreadRunnableThread.h
+++ b/Engine/Source/Runtime/Core/Private/HAL/PThreadRunnableThread.h
@@ -274,7 +274,11 @@ public:
 		// brute force kill that thread. Very bad as that might leak.
 		if (bShouldWait && bThreadStartedAndNotCleanedUp)
 		{
-			pthread_join(Thread, nullptr);
+			int JoinResult = pthread_join(Thread, nullptr);
+
+			UE_CLOG(JoinResult != 0, LogHAL, Warning, TEXT("Failed to join thread %u (%s). (err=%d, %s)"), ThreadID, *ThreadName, JoinResult, UTF8_TO_TCHAR(strerror(JoinResult)));
+
+			bDidExitOK = (JoinResult == 0);
 			bThreadStartedAndNotCleanedUp = false;
 		}
 
@@ -287,7 +291,10 @@ public:
 		// Block until this thread exits
 		if (bThreadStartedAndNotCleanedUp)
 		{
-			pthread_join(Thread, nullptr);
+			int JoinResult = pthread_join(Thread, nullptr);
+
+			UE_CLOG(JoinResult != 0, LogHAL, Warning, TEXT("Failed to join thread %u (%s). (err=%d, %s)"), ThreadID, *ThreadName, JoinResult, UTF8_TO_TCHAR(strerror(JoinResult)));
+
 			bThreadStartedAndNotCleanedUp = false;
 		}
 	}
