- Changed FConfigCacheIni to be a map of FString->FConfigFile* instead of FConfigFile, so if a file is added, it won't reallocate memory of other files
- Moved the TMap of non-known config files to a member instead of superclass for FConfigCacheIni. Exposed some minor bugs, which are now fixed

diff --git a/Engine/Source/Runtime/Core/Private/Misc/ConfigCacheIni.cpp b/Engine/Source/Runtime/Core/Private/Misc/ConfigCacheIni.cpp
--- a/Engine/Source/Runtime/Core/Private/Misc/ConfigCacheIni.cpp
+++ b/Engine/Source/Runtime/Core/Private/Misc/ConfigCacheIni.cpp
@@ -1960,7 +1960,7 @@ FConfigFile* FConfigCacheIni::FindConfigFile( const FString& Filename )
 
 	if (Result == nullptr)
 	{
-		Result = Super::Find(Filename);
+		Result = OtherFiles.FindRef(Filename);
 	}
 	return Result;
 }
@@ -2017,11 +2017,11 @@ FConfigFile* FConfigCacheIni::FindConfigFileWithBaseName(FName BaseName)
 		return Result;
 	}
 
-	for (TPair<FString,FConfigFile>& CurrentFilePair : *this)
+	for (TPair<FString,FConfigFile*>& CurrentFilePair : OtherFiles)
 	{
-		if (CurrentFilePair.Value.Name == BaseName)
+		if (CurrentFilePair.Value->Name == BaseName)
 		{
-			return &CurrentFilePair.Value;
+			return CurrentFilePair.Value;
 		}
 	}
 	return nullptr;
@@ -2034,9 +2034,9 @@ bool FConfigCacheIni::ContainsConfigFile(const FConfigFile* ConfigFile) const
 	// since the point at which the caller received the ConfigFile pointer
 	// they are testing. It is the caller's responsibility to not try to hold
 	// on to the ConfigFile pointer during writes to this ConfigCacheIni
-	for (const TPair<FString, FConfigFile>& CurrentFilePair : *this)
+	for (const TPair<FString, FConfigFile*>& CurrentFilePair : OtherFiles)
 	{
-		if (ConfigFile == &CurrentFilePair.Value)
+		if (ConfigFile == CurrentFilePair.Value)
 		{
 			return true;
 		}
@@ -2057,7 +2057,7 @@ bool FConfigCacheIni::ContainsConfigFile(const FConfigFile* ConfigFile) const
 TArray<FString> FConfigCacheIni::GetFilenames()
 {
 	TArray<FString> Result;
-	Super::GetKeys(Result);
+	OtherFiles.GetKeys(Result);
 
 	for (const FConfigCacheIni::FKnownConfigFiles::FKnownConfigFile& File : KnownFiles.Files)
 	{
@@ -2068,36 +2068,34 @@ TArray<FString> FConfigCacheIni::GetFilenames()
 }
 
 
-void FConfigCacheIni::Flush( bool Read, const FString& Filename )
+void FConfigCacheIni::Flush(bool bRemoveFromCache, const FString& Filename )
 {
 	// never Flush temporary cache objects
-	if (Type == EConfigCacheType::Temporary)
+	if (Type != EConfigCacheType::Temporary)
 	{
-		return;
-	}
-
-	// write out the files if we can
-	if (!bAreFileOperationsDisabled)
-	{
-		for (TIterator It(*this); It; ++It)
+		// write out the files if we can
+		if (!bAreFileOperationsDisabled)
 		{
-			if (Filename.Len() == 0 || It.Key()==Filename)
+			for (TPair<FString, FConfigFile*>& Pair : OtherFiles)
 			{
-				It.Value().Write(*It.Key());
+				if (Filename.Len() == 0 || Pair.Key == Filename)
+				{
+					Pair.Value->Write(*Pair.Key);
+				}
 			}
-		}
 
-		// now flush the known files (all or a single file)
-		for (FConfigCacheIni::FKnownConfigFiles::FKnownConfigFile& File : KnownFiles.Files)
-		{
-			if (Filename.Len() == 0 || Filename == File.IniName.ToString())
+			// now flush the known files (all or a single file)
+			for (FConfigCacheIni::FKnownConfigFiles::FKnownConfigFile& File : KnownFiles.Files)
 			{
-				File.IniFile.Write(*File.IniPath);
+				if (Filename.Len() == 0 || Filename == File.IniName.ToString())
+				{
+					File.IniFile.Write(*File.IniPath);
+				}
 			}
 		}
 	}
 
-	if( Read )
+	if (bRemoveFromCache)
 	{
 		// we can't read it back in if file operations are disabled
 		if (bAreFileOperationsDisabled)
@@ -2112,7 +2110,11 @@ void FConfigCacheIni::Flush( bool Read, const FString& Filename )
 		}
 		else
 		{
-			Empty();
+			for (TPair<FString, FConfigFile*>& It : OtherFiles)
+			{
+				delete It.Value;
+			}
+			OtherFiles.Empty();
 		}
 	}
 }
@@ -2289,9 +2291,6 @@ void FConfigCacheIni::LoadFile( const FString& Filename, const FConfigFile* Fall
 	{
 		UE_LOG(LogConfig, Warning, TEXT( "FConfigCacheIni::LoadFile failed loading file as it was 0 size.  Filename was:  %s" ), *Filename);
 	}
-
-	// Avoid memory wasted in array slack.
-	Shrink();
 }
 
 
@@ -2599,7 +2598,7 @@ bool FConfigCacheIni::GetSectionNames( const FString& Filename, TArray<FString>&
 	FConfigFile* File = Find(Filename);
 	if ( File != nullptr )
 	{
-		out_SectionNames.Empty(Num());
+		out_SectionNames.Empty(File->Num());
 		for ( FConfigFile::TIterator It(*File); It; ++It )
 		{
 			out_SectionNames.Add(It.Key());
@@ -2690,7 +2689,7 @@ void FConfigCacheIni::Dump(FOutputDevice& Ar, const TCHAR* BaseIniName)
 	if (BaseIniName == nullptr)
 	{
 		Ar.Log( TEXT("Files map:") );
-		TMap<FString,FConfigFile>::Dump( Ar );
+		OtherFiles.Dump( Ar );
 	}
 
 	for (const FConfigCacheIni::FKnownConfigFiles::FKnownConfigFile& File : KnownFiles.Files)
@@ -2698,11 +2697,11 @@ void FConfigCacheIni::Dump(FOutputDevice& Ar, const TCHAR* BaseIniName)
 		DumpFile(Ar, File.IniName.ToString(), File.IniFile);
 	}
 
-	for ( TIterator It(*this); It; ++It )
+	for (TPair<FString, FConfigFile*> Pair : OtherFiles)
 	{
-		if (BaseIniName == nullptr || FPaths::GetBaseFilename(It.Key()) == BaseIniName)
+		if (BaseIniName == nullptr || FPaths::GetBaseFilename(Pair.Key) == BaseIniName)
 		{
-			DumpFile(Ar, It.Key(), It.Value());
+			DumpFile(Ar, Pair.Key, *Pair.Value);
 		}
 	}
 }
@@ -3161,10 +3160,10 @@ void FConfigCacheIni::ShowMemoryUsage( FOutputDevice& Ar )
 {
 	FConfigMemoryData ConfigCacheMemoryData;
 
-	for ( TIterator FileIt(*this); FileIt; ++FileIt )
+	for (TPair<FString, FConfigFile*>& Pair : OtherFiles)
 	{
-		FString Filename = FileIt.Key();
-		FConfigFile& ConfigFile = FileIt.Value();
+		FString Filename = Pair.Key;
+		FConfigFile* ConfigFile = Pair.Value;
 
 		FArchiveCountConfigMem MemAr;
 
@@ -3172,10 +3171,15 @@ void FConfigCacheIni::ShowMemoryUsage( FOutputDevice& Ar )
 		MemAr << Filename;
 
 		// count the bytes used for storing the array of SectionName->Section pairs
-		MemAr << ConfigFile;
+		MemAr << *ConfigFile;
 		
 		ConfigCacheMemoryData.AddConfigFile(Filename, MemAr);
 	}
+	{
+		FArchiveCountConfigMem MemAr;
+		MemAr << KnownFiles;
+		ConfigCacheMemoryData.AddConfigFile(TEXT("KnownFiles"), MemAr);
+	}
 
 	// add a little extra spacing between the columns
 	ConfigCacheMemoryData.SizeIndent += 10;
@@ -3183,7 +3187,7 @@ void FConfigCacheIni::ShowMemoryUsage( FOutputDevice& Ar )
 
 	// record the memory used by the FConfigCacheIni's TMap
 	FArchiveCountConfigMem MemAr;
-	CountBytes(MemAr);
+	OtherFiles.CountBytes(MemAr);
 
 	SIZE_T TotalMemoryUsage=MemAr.GetNum();
 	SIZE_T MaxMemoryUsage=MemAr.GetMax();
@@ -3217,7 +3221,7 @@ SIZE_T FConfigCacheIni::GetMaxMemoryUsage()
 {
 	// record the memory used by the FConfigCacheIni's TMap
 	FArchiveCountConfigMem MemAr;
-	CountBytes(MemAr);
+	OtherFiles.CountBytes(MemAr);
 
 	SIZE_T TotalMemoryUsage=MemAr.GetNum();
 	SIZE_T MaxMemoryUsage=MemAr.GetMax();
@@ -3225,10 +3229,10 @@ SIZE_T FConfigCacheIni::GetMaxMemoryUsage()
 
 	FConfigMemoryData ConfigCacheMemoryData;
 
-	for ( TIterator FileIt(*this); FileIt; ++FileIt )
+	for (TPair<FString, FConfigFile*>& Pair : OtherFiles)
 	{
-		FString Filename = FileIt.Key();
-		FConfigFile& ConfigFile = FileIt.Value();
+		FString Filename = Pair.Key;
+		FConfigFile* ConfigFile = Pair.Value;
 
 		FArchiveCountConfigMem FileMemAr;
 
@@ -3236,10 +3240,15 @@ SIZE_T FConfigCacheIni::GetMaxMemoryUsage()
 		FileMemAr << Filename;
 
 		// count the bytes used for storing the array of SectionName->Section pairs
-		FileMemAr << ConfigFile;
+		FileMemAr << *ConfigFile;
 
 		ConfigCacheMemoryData.AddConfigFile(Filename, FileMemAr);
 	}
+	{
+		FArchiveCountConfigMem FileMemAr;
+		FileMemAr << KnownFiles;
+		ConfigCacheMemoryData.AddConfigFile(TEXT("KnownFiles"), FileMemAr);
+	}
 
 	for ( int32 Index = 0; Index < ConfigCacheMemoryData.MemoryData.Num(); Index++ )
 	{
@@ -3315,8 +3324,29 @@ void FConfigCacheIni::SaveCurrentStateForBootstrap(const TCHAR* Filename)
 
 void FConfigCacheIni::Serialize(FArchive& Ar)
 {
-	FConfigCacheIni::Super* Parent = static_cast<FConfigCacheIni::Super*>(this);
-	Ar << *Parent;
+	if (Ar.IsLoading())
+	{
+		int Num;
+		Ar << Num;
+		for (int Index = 0; Index < Num; Index++)
+		{
+			FString Filename;
+			FConfigFile* File = new FConfigFile;
+			Ar << Filename;
+			Ar << *File;
+			OtherFiles.Add(Filename, File);
+		}
+	}
+	else
+	{
+		int Num = OtherFiles.Num();
+		Ar << Num;
+		for (TPair<FString, FConfigFile*>& It : OtherFiles)
+		{
+			Ar << It.Key;
+			Ar << *It.Value;
+		}
+	}
 	Ar << KnownFiles;
 	Ar << bAreFileOperationsDisabled;
 	Ar << bIsReadyForUse;
diff --git a/Engine/Source/Runtime/Core/Public/Misc/ConfigCacheIni.h b/Engine/Source/Runtime/Core/Public/Misc/ConfigCacheIni.h
--- a/Engine/Source/Runtime/Core/Public/Misc/ConfigCacheIni.h
+++ b/Engine/Source/Runtime/Core/Public/Misc/ConfigCacheIni.h
@@ -599,10 +599,9 @@ enum class EConfigCacheType : uint8
 };
 
 // Set of all cached config files.
-class CORE_API FConfigCacheIni : private TMap<FString,FConfigFile>
+class CORE_API FConfigCacheIni
 {
 public:
-	
 	// Basic functions.
 	FConfigCacheIni(EConfigCacheType Type);
 
@@ -705,19 +704,19 @@ public:
 	/** Finds Config file that matches the base name such as "Engine" */
 	FConfigFile* FindConfigFileWithBaseName(FName BaseName);
 
-	using Super = TMap<FString, FConfigFile>;
 	FConfigFile& Add(const FString& Filename, const FConfigFile& File)
 	{
-		return Super::Add(Filename, File);
+		return *OtherFiles.Add(Filename, new FConfigFile(File));
 	}
 	int32 Remove(const FString& Filename)
 	{
-		return Super::Remove(Filename);
+		delete OtherFiles.FindRef(Filename);
+		return OtherFiles.Remove(Filename);
 	}
 	TArray<FString> GetFilenames();
 
 
-	void Flush( bool Read, const FString& Filename=TEXT("") );
+	void Flush(bool bRemoveFromCache, const FString& Filename=TEXT(""));
 
 	void LoadFile( const FString& InFilename, const FConfigFile* Fallback = NULL, const TCHAR* PlatformString = NULL );
 	void SetFile( const FString& InFilename, const FConfigFile* NewConfigFile );
@@ -1262,6 +1261,8 @@ private:
 	/** The filenames for the known files in this config */
 	FKnownConfigFiles KnownFiles;
 
+	TMap<FString, FConfigFile*> OtherFiles;
+
 	friend FConfigContext;
 };
 
@@ -1330,4 +1331,4 @@ CORE_API void DeleteRecordedConfigReadsFromIni();
  * Helper function to deal with "True","False","Yes","No","On","Off"
  */
 UE_DEPRECATED(5.1, "Use UE::ConfigUtilities::ConvertValueFromHumanFriendlyValue")
-CORE_API const TCHAR* ConvertValueFromHumanFriendlyValue(const TCHAR* Value);
\ No newline at end of file
+CORE_API const TCHAR* ConvertValueFromHumanFriendlyValue(const TCHAR* Value);
