Promoted Zen project id generation to FApp

diff --git a/Engine/Source/Developer/IoStoreUtilities/Internal/ZenStoreHttpClient.h b/Engine/Source/Developer/IoStoreUtilities/Internal/ZenStoreHttpClient.h
--- a/Engine/Source/Developer/IoStoreUtilities/Internal/ZenStoreHttpClient.h
+++ b/Engine/Source/Developer/IoStoreUtilities/Internal/ZenStoreHttpClient.h
@@ -64,8 +64,6 @@ public:
 	static const UTF8CHAR* FindOrAddAttachmentId(FUtf8StringView AttachmentText);
 	static const UTF8CHAR* FindAttachmentId(FUtf8StringView AttachmentText);
 
-	static FString GetProjectId();
-	static FString GenerateDefaultProjectId();
 private:
 	TIoStatusOr<FIoBuffer> ReadOpLogUri(FStringBuilderBase& ChunkUri, uint64 Offset = 0, uint64 Size = ~0ull);
 
diff --git a/Engine/Source/Developer/IoStoreUtilities/Private/ZenStoreHttpClient.cpp b/Engine/Source/Developer/IoStoreUtilities/Private/ZenStoreHttpClient.cpp
--- a/Engine/Source/Developer/IoStoreUtilities/Private/ZenStoreHttpClient.cpp
+++ b/Engine/Source/Developer/IoStoreUtilities/Private/ZenStoreHttpClient.cpp
@@ -770,32 +770,4 @@ const UTF8CHAR* FZenStoreHttpClient::FindAttachmentId(FUtf8StringView Attachment
 	return Existing;
 }
 
-// Duplicated in StorageServerPlatformFile.cpp to avoid having a public API in a shared module
-static FString GetProjectPathId()
-{
-	FString ProjectFilePath = FPaths::GetProjectFilePath();
-	FPaths::NormalizeFilename(ProjectFilePath);
-	FString AbsProjectFilePath = IFileManager::Get().ConvertToAbsolutePathForExternalAppForRead(*ProjectFilePath);
-	FTCHARToUTF8 AbsProjectFilePathUTF8(*AbsProjectFilePath);
-
-	FString HashString = FMD5::HashBytes((unsigned char*)AbsProjectFilePathUTF8.Get(), AbsProjectFilePathUTF8.Length()).Left(8);
-	return FString::Printf(TEXT("%s.%.8s"), FApp::GetProjectName(), *HashString);
-}
-
-FString FZenStoreHttpClient::GetProjectId()
-{
-	FString ProjectId;
-	if (FParse::Value(FCommandLine::Get(), TEXT("-ZenStoreProject="), ProjectId) == false)
-	{
-		ProjectId = GetProjectPathId();
-	}
-
-	return ProjectId;
-}
-
-FString FZenStoreHttpClient::GenerateDefaultProjectId()
-{
-	return GetProjectPathId();
-}
-
 }
diff --git a/Engine/Source/Developer/IoStoreUtilities/Private/ZenStoreWriter.cpp b/Engine/Source/Developer/IoStoreUtilities/Private/ZenStoreWriter.cpp
--- a/Engine/Source/Developer/IoStoreUtilities/Private/ZenStoreWriter.cpp
+++ b/Engine/Source/Developer/IoStoreUtilities/Private/ZenStoreWriter.cpp
@@ -175,7 +175,7 @@ FZenStoreWriter::FZenStoreWriter(
 {
 	StaticInit();
 
-	ProjectId = FZenStoreHttpClient::GetProjectId();
+	ProjectId = FApp::GetZenStoreProjectId();
 	
 	if (FParse::Value(FCommandLine::Get(), TEXT("-ZenStorePlatform="), OplogId) == false)
 	{
diff --git a/Engine/Source/Editor/UnrealEd/Private/TargetDomain/TargetDomainUtils.cpp b/Engine/Source/Editor/UnrealEd/Private/TargetDomain/TargetDomainUtils.cpp
--- a/Engine/Source/Editor/UnrealEd/Private/TargetDomain/TargetDomainUtils.cpp
+++ b/Engine/Source/Editor/UnrealEd/Private/TargetDomain/TargetDomainUtils.cpp
@@ -417,7 +417,7 @@ FEditorDomainOplog::FEditorDomainOplog()
 {
 	StaticInit();
 
-	FString ProjectId = FZenStoreHttpClient::GetProjectId();
+	FString ProjectId = FApp::GetZenStoreProjectId();
 	FString OplogId = TEXT("EditorDomain");
 
 	FString RootDir = FPaths::RootDir();
diff --git a/Engine/Source/Runtime/Core/Private/Misc/App.cpp b/Engine/Source/Runtime/Core/Private/Misc/App.cpp
--- a/Engine/Source/Runtime/Core/Private/Misc/App.cpp
+++ b/Engine/Source/Runtime/Core/Private/Misc/App.cpp
@@ -10,6 +10,7 @@
 #include "UObject/DevObjectVersion.h"
 #include "Misc/EngineVersion.h"
 #include "Misc/NetworkVersion.h"
+#include "Misc/SecureHash.h"
 
 DEFINE_LOG_CATEGORY_STATIC(LogApp, Log, All);
 
@@ -372,4 +373,30 @@ void FApp::PrintStartupLogMessages()
 	UE_LOG(LogInit, Log, TEXT("Installed Engine Build: %d"), FApp::IsEngineInstalled() ? 1 : 0);
 
 	FDevVersionRegistration::DumpVersionsToLog();
+}
+
+FString FApp::GetZenStoreProjectId()
+{
+	FString ProjectId;
+	if (FParse::Value(FCommandLine::Get(), TEXT("-ZenStoreProject="), ProjectId))
+	{
+		return ProjectId;
+	}
+
+#if PLATFORM_DESKTOP
+	if (FPaths::IsProjectFilePathSet())
+	{
+		FString ProjectFilePath = FPaths::GetProjectFilePath();
+		FPaths::NormalizeFilename(ProjectFilePath);
+		FString AbsProjectFilePath = IFileManager::Get().ConvertToAbsolutePathForExternalAppForRead(*ProjectFilePath);
+		FTCHARToUTF8 AbsProjectFilePathUTF8(*AbsProjectFilePath);
+
+		FString HashString = FMD5::HashBytes((unsigned char*)AbsProjectFilePathUTF8.Get(), AbsProjectFilePathUTF8.Length()).Left(8);
+		return FString::Printf(TEXT("%s.%.8s"), FApp::GetProjectName(), *HashString);
+	}
+	UE_LOG(LogInit, Fatal, TEXT("GetZenStoreProjectId() called before having a valid project file path"));
+#else
+	UE_LOG(LogInit, Fatal, TEXT("-ZenStoreProject command line argument is required to run from Zen"));
+#endif
+	return FString();
 }
\ No newline at end of file
diff --git a/Engine/Source/Runtime/Core/Public/Misc/App.h b/Engine/Source/Runtime/Core/Public/Misc/App.h
--- a/Engine/Source/Runtime/Core/Public/Misc/App.h
+++ b/Engine/Source/Runtime/Core/Public/Misc/App.h
@@ -211,6 +211,13 @@ public:
 		SessionUsers.Remove(UserName);
 	}
 
+	/**
+	 * Gets the Zen store project id for the current application instance.
+	 *
+	 * @return Zen store project id.
+	 */
+	static FString GetZenStoreProjectId();
+
 	/**
 	 * Gets the globally unique identifier of this application instance.
 	 *
diff --git a/Engine/Source/Runtime/StorageServerClient/Private/StorageServerConnection.cpp b/Engine/Source/Runtime/StorageServerClient/Private/StorageServerConnection.cpp
--- a/Engine/Source/Runtime/StorageServerClient/Private/StorageServerConnection.cpp
+++ b/Engine/Source/Runtime/StorageServerClient/Private/StorageServerConnection.cpp
@@ -66,18 +66,6 @@ static uint64 GetCompressedOffset(const FCompressedBuffer& Buffer, uint64 RawOff
 	return 0;
 }
 
-// Duplicated in ZenStoreHttpClient.cpp to avoid having a public API in a shared module
-static FString GetProjectPathId()
-{
-	FString ProjectFilePath = FPaths::GetProjectFilePath();
-	FPaths::NormalizeFilename(ProjectFilePath);
-	FString AbsProjectFilePath = IFileManager::Get().ConvertToAbsolutePathForExternalAppForRead(*ProjectFilePath);
-	FTCHARToUTF8 AbsProjectFilePathUTF8(*AbsProjectFilePath);
-
-	FString HashString = FMD5::HashBytes((unsigned char*)AbsProjectFilePathUTF8.Get(), AbsProjectFilePathUTF8.Length()).Left(8);
-	return FString::Printf(TEXT("%s.%.8s"), FApp::GetProjectName(), *HashString);
-}
-
 FStorageServerRequest::FStorageServerRequest(FAnsiStringView Verb, FAnsiStringView Resource, FAnsiStringView Hostname, EStorageServerContentType Accept)
 : AcceptType(Accept)
 {
@@ -451,7 +439,7 @@ bool FStorageServerConnection::Initialize(TArrayView<const FString> InHostAddres
 	}
 	else
 	{
-		OplogPath.Append(TCHAR_TO_ANSI(*GetProjectPathId()));
+		OplogPath.Append(TCHAR_TO_ANSI(*FApp::GetZenStoreProjectId()));
 	}
 	OplogPath.Append("/oplog/");
 	if (InPlatformNameOverride)
