Add support for RotationVector conversions to Quat

diff --git a/Engine/Source/Runtime/Core/Public/Math/Quat.h b/Engine/Source/Runtime/Core/Public/Math/Quat.h
--- a/Engine/Source/Runtime/Core/Public/Math/Quat.h
+++ b/Engine/Source/Runtime/Core/Public/Math/Quat.h
@@ -400,19 +400,33 @@ public:
 	FORCEINLINE T SizeSquared() const;
 
 
-	/** Get the angle of this quaternion */
+	/** Get the angle in radians of this quaternion */
 	FORCEINLINE T GetAngle() const;
 
 	/** 
 	 * get the axis and angle of rotation of this quaternion
 	 *
 	 * @param Axis{out] vector of axis of the quaternion
-	 * @param Angle{out] angle of the quaternion
+	 * @param Angle{out] angle of the quaternion in radians
 	 * @warning : assumes normalized quaternions.
 	 */
 	void ToAxisAndAngle(TVector<T>& Axis, float& Angle) const;
 	void ToAxisAndAngle(TVector<T>& Axis, double& Angle) const;
 
+	/**
+	 * Get the rotation vector corresponding to this quaternion. 
+	 * The direction of the vector represents the rotation axis, 
+	 * and the magnitude the angle in radians.
+	 */
+	TVector<T> ToRotationVector() const;
+
+	/**
+	 * Constructs a quaternion corresponding to the rotation vector. 
+	 * The direction of the vector represents the rotation axis, 
+	 * and the magnitude the angle in radians.
+	 */
+	static TQuat<T> MakeFromRotationVector(const TVector<T>& RotationVector);
+
 	/** 
 	 * Get the swing and twist decomposition for a specified axis
 	 *
@@ -1133,6 +1147,20 @@ FORCEINLINE void TQuat<T>::ToAxisAndAngle(TVector<T>& Axis, double& Angle) const
 	Axis = GetRotationAxis();
 }
 
+template<typename T>
+FORCEINLINE TVector<T> TQuat<T>::ToRotationVector() const
+{
+	TQuat<T> RotQ = Log();
+	return TVector<T>(RotQ.X * 2.0f, RotQ.Y * 2.0f, RotQ.Z * 2.0f);
+}
+
+template<typename T>
+FORCEINLINE TQuat<T> TQuat<T>::MakeFromRotationVector(const TVector<T>& RotationVector)
+{
+	TQuat<T> RotQ(RotationVector.X * 0.5f, RotationVector.Y * 0.5f, RotationVector.Z * 0.5f, 0.0f);
+	return RotQ.Exp();
+}
+
 template<typename T>
 FORCEINLINE TVector<T> TQuat<T>::GetRotationAxis() const
 {
