Remove NDA platform define references

diff --git a/Engine/Plugins/Experimental/NNI/Source/Private/NeuralNetworkInferenceQA/Private/ModelUnitTester.cpp b/Engine/Plugins/Experimental/NNI/Source/Private/NeuralNetworkInferenceQA/Private/ModelUnitTester.cpp
--- a/Engine/Plugins/Experimental/NNI/Source/Private/NeuralNetworkInferenceQA/Private/ModelUnitTester.cpp
+++ b/Engine/Plugins/Experimental/NNI/Source/Private/NeuralNetworkInferenceQA/Private/ModelUnitTester.cpp
@@ -5,7 +5,9 @@
 #include "NeuralNetworkInferenceQAUtils.h"
 #include "Misc/Paths.h"
 
-
+#ifndef SKIP_ModelLoadAccuracyAndSpeedTests
+#define SKIP_ModelLoadAccuracyAndSpeedTests 0
+#endif
 
 /* FModelUnitTester axuiliary functions
  *****************************************************************************/
@@ -124,12 +126,12 @@ bool FModelUnitTester::GlobalTest(const FString& InProjectContentDir, const FStr
 	// Run tests
 	UE_LOG(LogNeuralNetworkInferenceQA, Display,
 		TEXT("ENeuralBackEnd::UEOnly test for FModelUnitTester::GlobalTest disabled due to not being compatible with ORT (linking error issues)."));
-#if !PLATFORM_PS4
+#if !SKIP_ModelLoadAccuracyAndSpeedTests
 	return ModelLoadAccuracyAndSpeedTests(InProjectContentDir, InModelZooRelativeDirectory, ModelNames, InputArrayValues, CPUGroundTruths,
 		GPUGroundTruths, CPURepetitionsForUEAndORTBackEnd, GPURepetitionsForUEAndORTBackEnd, CPURepetitionsForUEOnlyBackEnd,
 		GPURepetitionsForUEOnlyBackEnd);
 #else
-	UE_LOG(LogNeuralNetworkInferenceQA, Display, TEXT("ENeuralBackEnd::FModelUnitTester::GlobalTest disabled due to not being compatible with cooked PS4 build"));
+	UE_LOG(LogNeuralNetworkInferenceQA, Display, TEXT("ENeuralBackEnd::FModelUnitTester::GlobalTest disabled due to not being compatible with cooked platform build"));
 	return true; 
 #endif
 }
diff --git a/Engine/Source/Runtime/Core/Private/FramePro/FramePro.cpp b/Engine/Source/Runtime/Core/Private/FramePro/FramePro.cpp
--- a/Engine/Source/Runtime/Core/Private/FramePro/FramePro.cpp
+++ b/Engine/Source/Runtime/Core/Private/FramePro/FramePro.cpp
@@ -8502,7 +8502,7 @@ namespace FramePro
 //
 
 //------------------------------------------------------------------------
-#if FRAMEPRO_LINUX_BASED_PLATFORM && !FRAMEPRO_PLATFORM_PS4 && !(FRAMEPRO_PLATFORM_UE4 && PLATFORM_PS4)
+#if FRAMEPRO_LINUX_BASED_PLATFORM && !FRAMEPRO_PLATFORM_PS4
 
 //------------------------------------------------------------------------
 #if FRAMEPRO_ENABLE_CALLSTACKS
diff --git a/Engine/Source/Runtime/Core/Private/FramePro/FrameProUE4.cpp b/Engine/Source/Runtime/Core/Private/FramePro/FrameProUE4.cpp
--- a/Engine/Source/Runtime/Core/Private/FramePro/FrameProUE4.cpp
+++ b/Engine/Source/Runtime/Core/Private/FramePro/FrameProUE4.cpp
@@ -149,7 +149,7 @@
 		//------------------------------------------------------------------------
 		Platform::Enum Platform::GetPlatformEnum()
 		{
-			//@EPIC: allow external definition + XBOXONE removed
+			//@EPIC: allow external definition + NDA platforms removed
 			#if defined(FRAMEPRO_UNREAL_PLATFORM)
 				return FRAMEPRO_UNREAL_PLATFORM;
 			//@EPIC end
@@ -159,16 +159,12 @@
 				return Platform::Linux;
 			#elif PLATFORM_UNIX
 				return Platform::Linux;
-			#elif PLATFORM_PS4
-				return Platform::PS4;
 			#elif PLATFORM_ANDROID
 				return Platform::Android;
 			#elif PLATFORM_MAC
 				return Platform::Mac;
 			#elif PLATFORM_IOS
 				return Platform::iOS;
-			#elif PLATFORM_SWITCH
-				return Platform::Switch;
 			#else
 				//@EPIC: begin - useful error
 				#error unknown platform or FRAMEPRO_UNREAL_PLATFORM not defined
@@ -176,34 +172,16 @@
 			#endif
 		}
 
-		//------------------------------------------------------------------------
-		namespace PS4ContextSwitches
-		{
-			void* CreateContextSwitchRecorder(Allocator* p_allocator);
-			void DestroyContextSwitchRecorder(void* p_obj, Allocator* p_allocator);
-			bool StartRecordingContextSitches(void* p_obj, Platform::ContextSwitchCallbackFunction callback_function, void* p_context, DynamicString& error, Allocator* p_allocator);
-			void StopRecordingContextSitches(void* p_obj);
-			void FlushContextSwitches(void*);
-		}
-
 		//------------------------------------------------------------------------
 		void* Platform::CreateContextSwitchRecorder(Allocator* p_allocator)
 		{
-			#if PLATFORM_PS4 && FRAMEPRO_PS4_CONTEXT_SWITCH_RECORDING
-				return PS4ContextSwitches::CreateContextSwitchRecorder(p_allocator);
-			#else
 				return GenericPlatform::CreateContextSwitchRecorder(p_allocator);
-			#endif
 		}
 
 		//------------------------------------------------------------------------
 		void Platform::DestroyContextSwitchRecorder(void* p_context_switch_recorder, Allocator* p_allocator)
 		{
-			#if PLATFORM_PS4 && FRAMEPRO_PS4_CONTEXT_SWITCH_RECORDING
-				PS4ContextSwitches::DestroyContextSwitchRecorder(p_context_switch_recorder, p_allocator);
-			#else
 				GenericPlatform::DestroyContextSwitchRecorder(p_context_switch_recorder, p_allocator);
-			#endif
 		}
 
 		//------------------------------------------------------------------------
@@ -214,14 +192,6 @@
 			DynamicString& error,
 			Allocator* p_allocator)
 		{
-			#if PLATFORM_PS4 && FRAMEPRO_PS4_CONTEXT_SWITCH_RECORDING
-				return PS4ContextSwitches::StartRecordingContextSitches(
-					p_context_switch_recorder,
-					p_callback,
-					p_context,
-					error,
-					p_allocator);
-			#else
 				FRAMEPRO_UNREFERENCED(p_allocator);
 
 				return GenericPlatform::StartRecordingContextSitches(
@@ -229,32 +199,23 @@
 					p_callback,
 					p_context,
 					error);
-			#endif
 		}
 
 		//------------------------------------------------------------------------
 		void Platform::StopRecordingContextSitches(void* p_context_switch_recorder)
 		{
-			#if PLATFORM_PS4 && FRAMEPRO_PS4_CONTEXT_SWITCH_RECORDING
-				PS4ContextSwitches::StopRecordingContextSitches(p_context_switch_recorder);
-			#else
 				GenericPlatform::StopRecordingContextSitches(p_context_switch_recorder);
-			#endif
 		}
 
 		//------------------------------------------------------------------------
 		void Platform::FlushContextSwitches(void* p_context_switch_recorder)
 		{
-			#if PLATFORM_PS4 && FRAMEPRO_PS4_CONTEXT_SWITCH_RECORDING
-				PS4ContextSwitches::FlushContextSwitches(p_context_switch_recorder);
-			#else
 				GenericPlatform::FlushContextSwitches(p_context_switch_recorder);
-			#endif
 		}
 
 		//------------------------------------------------------------------------
-		#if PLATFORM_PS4 && FRAMEPRO_ENUMERATE_ALL_MODULES
-			namespace EnumModulesPS4 { void EnumerateModules(Array<ModulePacket*>& module_packets, Allocator* p_allocator); }
+		#if defined(FRAMEPRO_ENUM_MODULES_PLATFORM) && FRAMEPRO_ENUMERATE_ALL_MODULES
+			namespace FRAMEPRO_ENUM_MODULES_PLATFORM { void EnumerateModules(Array<ModulePacket*>& module_packets, Allocator* p_allocator); }
 		#endif
 
 		//------------------------------------------------------------------------
@@ -262,11 +223,6 @@
 			namespace EnumModulesLinux { void EnumerateModules(Array<ModulePacket*>& module_packets, Allocator* p_allocator); }
 		#endif
 
-		//------------------------------------------------------------------------
-		#if PLATFORM_SWITCH
-			namespace EnumModulesSwitch { void EnumerateModules(Array<ModulePacket*>& module_packets, Allocator* p_allocator); }
-		#endif
-
 		//------------------------------------------------------------------------
 		void Platform::EnumerateModules(Array<ModulePacket*>& module_packets, Allocator* p_allocator)
 		{
@@ -275,10 +231,8 @@
 					EnumModulesWindows::EnumerateModules(module_packets, p_allocator);
 				#elif PLATFORM_LINUX
 					EnumModulesLinux::EnumerateModules(module_packets, p_allocator);
-				#elif PLATFORM_PS4
-					EnumModulesPS4::EnumerateModules(module_packets, p_allocator);
-				#elif PLATFORM_SWITCH
-					EnumModulesSwitch::EnumerateModules(module_packets, p_allocator);
+				#elif defined(FRAMEPRO_ENUM_MODULES_PLATFORM)
+					FRAMEPRO_ENUM_MODULES_PLATFORM::EnumerateModules(module_packets, p_allocator);
 				#else
 					ModulePacket* p_module_packet = (ModulePacket*)p_allocator->Alloc(sizeof(ModulePacket));
 					if (!p_module_packet)
diff --git a/Engine/Source/Runtime/Core/Private/Memory/MemoryArena.cpp b/Engine/Source/Runtime/Core/Private/Memory/MemoryArena.cpp
--- a/Engine/Source/Runtime/Core/Private/Memory/MemoryArena.cpp
+++ b/Engine/Source/Runtime/Core/Private/Memory/MemoryArena.cpp
@@ -19,10 +19,6 @@ extern "C" {
 #define USE_SPIN_LOCKS 0
 //#define USE_RECURSIVE_LOCKS 1
 
-#if PLATFORM_PS4
-#	define malloc_getpagesize 4096
-#endif
-
 #include "ThirdParty/dlmalloc/malloc-2.8.6.h"
 
 #if PLATFORM_WINDOWS
diff --git a/Engine/Source/Runtime/Core/Private/Misc/Paths.cpp b/Engine/Source/Runtime/Core/Private/Misc/Paths.cpp
--- a/Engine/Source/Runtime/Core/Private/Misc/Paths.cpp
+++ b/Engine/Source/Runtime/Core/Private/Misc/Paths.cpp
@@ -434,14 +434,8 @@ FString FPaths::AudioCaptureDir()
 
 FString FPaths::ProjectLogDir()
 {
-#if PLATFORM_PS4
-	const FString* OverrideDir = FPS4PlatformFile::GetOverrideLogDirectory();
-	if (OverrideDir != nullptr)
-	{
-		return *OverrideDir;
-	}
-#elif PLATFORM_SWITCH
-	const FString* OverrideDir = FSwitchPlatformFile::GetOverrideLogDirectory();
+#if defined(OVERRIDE_LOG_DIRECTORY_PLATFORM)
+	const FString* OverrideDir = OVERRIDE_LOG_DIRECTORY_PLATFORM::GetOverrideLogDirectory();
 	if (OverrideDir != nullptr)
 	{
 		return *OverrideDir;
diff --git a/Engine/Source/Runtime/Core/Public/FramePro/FrameProUE4.h b/Engine/Source/Runtime/Core/Public/FramePro/FrameProUE4.h
--- a/Engine/Source/Runtime/Core/Public/FramePro/FrameProUE4.h
+++ b/Engine/Source/Runtime/Core/Public/FramePro/FrameProUE4.h
@@ -44,10 +44,14 @@
 		#define FRAMEPRO_API CORE_API
 
 		// Windows or Linux based platform
-#if !defined(FRAMEPRO_WIN_BASED_PLATFORM) //@EPIC begin - allow external definition + XBOXONE removed
+//@EPIC begin - allow external definition + NDA platforms removed
+#if !defined(FRAMEPRO_WIN_BASED_PLATFORM)
 		#define FRAMEPRO_WIN_BASED_PLATFORM (PLATFORM_WINDOWS)
-#endif //@EPIC: end
-		#define FRAMEPRO_LINUX_BASED_PLATFORM (!FRAMEPRO_WIN_BASED_PLATFORM && !PLATFORM_SWITCH)
+#endif
+#if !defined(FRAMEPRO_LINUX_BASED_PLATFORM)
+		#define FRAMEPRO_LINUX_BASED_PLATFORM (!FRAMEPRO_WIN_BASED_PLATFORM)
+#endif
+//@EPIC: end
 
 		#define FRAMEPRO_USE_TLS_SLOTS 1
 
diff --git a/Engine/Source/Runtime/Core/Public/HAL/IPlatformFileCachedWrapper.h b/Engine/Source/Runtime/Core/Public/HAL/IPlatformFileCachedWrapper.h
--- a/Engine/Source/Runtime/Core/Public/HAL/IPlatformFileCachedWrapper.h
+++ b/Engine/Source/Runtime/Core/Public/HAL/IPlatformFileCachedWrapper.h
@@ -281,12 +281,15 @@ public:
 	}
 	virtual bool ShouldBeUsed(IPlatformFile* Inner, const TCHAR* CmdLine) const override
 	{
+#ifndef PLATFORM_PROVIDES_FILE_CACHE
+#define PLATFORM_PROVIDES_FILE_CACHE 0
+#endif
 		// Default to false on platforms that already do platform file level caching
-		bool bResult = !PLATFORM_PS4 && !PLATFORM_WINDOWS && FPlatformProperties::RequiresCookedData();
+		bool bResult = !PLATFORM_PROVIDES_FILE_CACHE && !PLATFORM_WINDOWS && FPlatformProperties::RequiresCookedData();
 
 		// Allow a choice between shorter load times or less memory on desktop platforms.
 		// Note: this cannot be in config since they aren't read at that point.
-#if (PLATFORM_DESKTOP || PLATFORM_PS4)
+#if (PLATFORM_DESKTOP || PLATFORM_PROVIDES_FILE_CACHE)
 		{
 			if (FParse::Param(CmdLine, TEXT("NoCachedReadFile")))
 			{
diff --git a/Engine/Source/Runtime/Core/Public/HAL/UnrealMemory.h b/Engine/Source/Runtime/Core/Public/HAL/UnrealMemory.h
--- a/Engine/Source/Runtime/Core/Public/HAL/UnrealMemory.h
+++ b/Engine/Source/Runtime/Core/Public/HAL/UnrealMemory.h
@@ -32,7 +32,7 @@ FORCEINLINE void DoGamethreadHook(int32 Index)
 }
 #endif
 
-#define TIME_MALLOC (0 && PLATFORM_PS4)
+#define TIME_MALLOC (0)
 
 #if TIME_MALLOC
 
diff --git a/Engine/Source/Runtime/Core/Public/MemPro/MemPro.h b/Engine/Source/Runtime/Core/Public/MemPro/MemPro.h
--- a/Engine/Source/Runtime/Core/Public/MemPro/MemPro.h
+++ b/Engine/Source/Runtime/Core/Public/MemPro/MemPro.h
@@ -67,8 +67,13 @@ OR IN CONNECTION WITH THE USE OR PERFORMANCE OF THIS SOFTWARE.
 // **** The Target Platform ****
 
 // define ONE of these
-#if (defined(_WIN32) || defined(_WIN64) || defined(WIN32) || defined(WIN64) || defined(__WIN32__) || defined(__WINDOWS__)) && !defined(MEMPRO_PLATFORM_XBOXONE) //@EPIC: allow external definition
-	#if defined(_XBOX_ONE) //@EPIC:  allow external definition
+//@EPIC BEGIN external definition of platforms
+#if defined(MEMPRO_PLATFORM_XBOXONE)
+#elif defined(MEMPRO_PLATFORM_PS4)
+#elif defined(MEMPRO_PLATFORM_SWITCH)
+#elif defined(_WIN32) || defined(_WIN64) || defined(WIN32) || defined(WIN64) || defined(__WIN32__)
+//@EPIC END
+	#if defined(_XBOX_ONE)
 		#define MEMPRO_PLATFORM_XBOXONE
 	#elif defined(_XBOX)
 		#define MEMPRO_PLATFORM_XBOX360
@@ -77,12 +82,6 @@ OR IN CONNECTION WITH THE USE OR PERFORMANCE OF THIS SOFTWARE.
 	#endif
 #elif defined(__APPLE__)
 	#define MEMPRO_PLATFORM_APPLE
-#elif defined(PS4) || (defined(__UNREAL__) && PLATFORM_PS4)
-	#define MEMPRO_PLATFORM_PS4
-//@EPIC BEGIN: other platforms
-#elif defined(__UNREAL__) && PLATFORM_SWITCH
-	#define MEMPRO_PLATFORM_SWITCH
-//@EPIC END
 #else
 	#define MEMPRO_PLATFORM_UNIX
 #endif
diff --git a/Engine/Source/Runtime/Engine/Private/EngineSessionManager.cpp b/Engine/Source/Runtime/Engine/Private/EngineSessionManager.cpp
--- a/Engine/Source/Runtime/Engine/Private/EngineSessionManager.cpp
+++ b/Engine/Source/Runtime/Engine/Private/EngineSessionManager.cpp
@@ -27,9 +27,15 @@
 
 DEFINE_LOG_CATEGORY(LogEngineSessionManager);
 
+#if !defined(IGNORE_SESSION_SHUTDOWN_IN_DEACTIVATED_STATE)
+	#define IGNORE_SESSION_SHUTDOWN_IN_DEACTIVATED_STATE 0
+#endif
 #if !defined(IGNORE_SESSION_SHUTDOWN_IN_BACKGROUND_STATE)
 	#define IGNORE_SESSION_SHUTDOWN_IN_BACKGROUND_STATE 0
 #endif
+#if !defined(UNABLE_TO_SET_TERMINATING_FLAG)
+	#define UNABLE_TO_SET_TERMINATING_FLAG 0
+#endif
 
 namespace SessionManagerDefs
 {
@@ -45,7 +51,6 @@ namespace SessionManagerDefs
 	static const FString TerminatedSessionToken(TEXT("Terminated"));
 	static const FString DebuggerSessionToken(TEXT("Debugger"));
 	static const FString AbnormalSessionToken(TEXT("AbnormalShutdown"));
-	static const FString PS4SessionToken(TEXT("AbnormalShutdownPS4"));
 	static const FString SessionRecordListSection(TEXT("List"));
 	static const FString EditorSessionRecordSectionPrefix(TEXT("Unreal Engine/Editor Sessions/"));
 	static const FString GameSessionRecordSectionPrefix(TEXT("Unreal Engine/Game Sessions/"));
@@ -445,16 +450,16 @@ void FEngineSessionManager::SendAbnormalShutdownReport(const FSessionRecord& Rec
 
 #if PLATFORM_WINDOWS | PLATFORM_MAC | PLATFORM_UNIX
 	// do nothing
-#elif PLATFORM_PS4
+#elif IGNORE_SESSION_SHUTDOWN_IN_DEACTIVATED_STATE
 	if (Record.bIsDeactivated && !Record.bCrashed)
 	{
-		// Shutting down in deactivated state on PS4 is normal - don't report it
+		// Shutting down in deactivated state on this platform is normal - don't report it
 		return;
 	}
 #elif IGNORE_SESSION_SHUTDOWN_IN_BACKGROUND_STATE
 	if (Record.bIsInBackground && !Record.bCrashed)
 	{
-		// Shutting down in background state on XB1 is normal - don't report it
+		// Shutting down in background state on this platform is normal - don't report it
 		return;
 	}
 #else
@@ -469,13 +474,13 @@ void FEngineSessionManager::SendAbnormalShutdownReport(const FSessionRecord& Rec
 		SessionIdString = SessionId.ToString(EGuidFormats::DigitsWithHyphensInBraces);
 	}
 
-#if !PLATFORM_PS4
+#if !UNABLE_TO_SET_TERMINATING_FLAG
 	FString ShutdownTypeString = Record.bCrashed ? SessionManagerDefs::CrashSessionToken :
 		(Record.bWasEverDebugger ? SessionManagerDefs::DebuggerSessionToken :
 		(Record.bIsTerminating ? SessionManagerDefs::TerminatedSessionToken : SessionManagerDefs::AbnormalSessionToken));
 #else
-	// PS4 cannot set the crash flag so report abnormal shutdowns with a specific token meaning "crash or abnormal shutdown".
-	FString ShutdownTypeString = Record.bWasEverDebugger ? SessionManagerDefs::DebuggerSessionToken : SessionManagerDefs::PS4SessionToken;
+	// Platform cannot set the crash flag so report abnormal shutdowns with a specific token meaning "crash or abnormal shutdown".
+	FString ShutdownTypeString = Record.bWasEverDebugger ? SessionManagerDefs::DebuggerSessionToken : SessionManagerDefs::AbnormalSessionToken + FPlatformProperties::PlatformName();
 #endif
 
 	const FString& RunTypeString = Record.Mode == EEngineSessionManagerMode::Editor ? SessionManagerDefs::EditorValueString : SessionManagerDefs::GameValueString;
diff --git a/Engine/Source/Runtime/WebBrowser/Private/WebBrowserSingleton.cpp b/Engine/Source/Runtime/WebBrowser/Private/WebBrowserSingleton.cpp
--- a/Engine/Source/Runtime/WebBrowser/Private/WebBrowserSingleton.cpp
+++ b/Engine/Source/Runtime/WebBrowser/Private/WebBrowserSingleton.cpp
@@ -451,7 +451,7 @@ FWebBrowserSingleton::~FWebBrowserSingleton()
 		// Shut down CEF.
 		CefShutdown();
 	}
-#elif PLATFORM_IOS || PLATFORM_PS4 || (PLATFORM_ANDROID && USE_ANDROID_JNI)
+#elif PLATFORM_IOS || PLATFORM_SPECIFIC_WEB_BROWSER || (PLATFORM_ANDROID && USE_ANDROID_JNI)
 	{
 		FScopeLock Lock(&WindowInterfacesCS);
 		// Clear this before CefShutdown() below
@@ -634,7 +634,7 @@ TSharedPtr<IWebBrowserWindow> FWebBrowserSingleton::CreateBrowserWindow(const FC
 		WindowInterfaces.Add(NewBrowserWindow);
 	}
 	return NewBrowserWindow;
-#elif PLATFORM_PS4
+#elif PLATFORM_SPECIFIC_WEB_BROWSER
 	// Create new window
 	TSharedPtr<FWebBrowserWindow> NewBrowserWindow = MakeShareable(new FWebBrowserWindow(
 		WindowSettings.InitialURL,
@@ -749,7 +749,7 @@ bool FWebBrowserSingleton::Tick(float DeltaTime)
 		}
 	}
 
-#elif PLATFORM_IOS || PLATFORM_PS4 || (PLATFORM_ANDROID && USE_ANDROID_JNI)
+#elif PLATFORM_IOS || PLATFORM_SPECIFIC_WEB_BROWSER || (PLATFORM_ANDROID && USE_ANDROID_JNI)
 	FScopeLock Lock(&WindowInterfacesCS);
 	bool bIsSlateAwake = FSlateApplication::IsInitialized() && !FSlateApplication::Get().IsSlateAsleep();
 	// Remove any windows that have been deleted and check whether it's currently visible
diff --git a/Engine/Source/Runtime/WebBrowser/Private/WebBrowserSingleton.h b/Engine/Source/Runtime/WebBrowser/Private/WebBrowserSingleton.h
--- a/Engine/Source/Runtime/WebBrowser/Private/WebBrowserSingleton.h
+++ b/Engine/Source/Runtime/WebBrowser/Private/WebBrowserSingleton.h
@@ -168,7 +168,7 @@ private:
 	/** List of currently existing browser windows */
 #if WITH_CEF3
 	TArray<TWeakPtr<FCEFWebBrowserWindow>>	WindowInterfaces;
-#elif PLATFORM_IOS || PLATFORM_PS4 || (PLATFORM_ANDROID && USE_ANDROID_JNI)
+#elif PLATFORM_IOS || PLATFORM_SPECIFIC_WEB_BROWSER || (PLATFORM_ANDROID && USE_ANDROID_JNI)
 	TArray<TWeakPtr<IWebBrowserWindow>>	WindowInterfaces;
 #endif
 
