Remove redundant calls to AlignArbitrary().  This is a small part of a larger pull request which has already been addressed:

GitHub# 9078 : Fix occasional crash in TArray<FVector4>::Add

diff --git a/Engine/Source/Runtime/Core/Private/HAL/MallocTBB.cpp b/Engine/Source/Runtime/Core/Private/HAL/MallocTBB.cpp
--- a/Engine/Source/Runtime/Core/Private/HAL/MallocTBB.cpp
+++ b/Engine/Source/Runtime/Core/Private/HAL/MallocTBB.cpp
@@ -46,7 +46,7 @@ void* FMallocTBB::TryMalloc( SIZE_T Size, uint32 Alignment )
 	// Contrary to scalable_malloc, scalable_aligned_malloc returns nullptr when trying to allocate 0 bytes, which is inconsistent with system malloc, so
 	// for 0 bytes requests we actually allocate sizeof(size_t), which is exactly what scalable_malloc does internally in such case.
 	// scalable_aligned_realloc and scalable_realloc behave the same in this regard, so this is only needed here.
-	Alignment = AlignArbitrary(FMath::Max((uint32)16, Alignment), (uint32)16);
+	Alignment = FMath::Max((uint32)16, Alignment);
 	NewPtr = scalable_aligned_malloc(Size ? Size : sizeof(size_t), Alignment);
 #else
 	if( Alignment != DEFAULT_ALIGNMENT )
@@ -108,7 +108,7 @@ void* FMallocTBB::TryRealloc( void* Ptr, SIZE_T NewSize, uint32 Alignment )
 	void* NewPtr = nullptr;
 #if PLATFORM_MAC
 	// macOS expects all allocations to be aligned to 16 bytes, but TBBs default alignment is 8, so on Mac we always have to use scalable_aligned_realloc
-	Alignment = AlignArbitrary(FMath::Max((uint32)16, Alignment), (uint32)16);
+	Alignment = FMath::Max((uint32)16, Alignment);
 	NewPtr = scalable_aligned_realloc(Ptr, NewSize, Alignment);
 #else
 	if (Alignment != DEFAULT_ALIGNMENT)
@@ -120,7 +120,7 @@ void* FMallocTBB::TryRealloc( void* Ptr, SIZE_T NewSize, uint32 Alignment )
 	{
 		// Fulfill the promise of DEFAULT_ALIGNMENT, which aligns 16-byte or larger structures to 16 bytes,
 		// while TBB aligns to 8 by default.
-		NewPtr = scalable_aligned_realloc(Ptr, NewSize, NewSize >= 16 ? (uint32)16 : (uint32)8);
+		NewPtr = scalable_aligned_realloc( Ptr, NewSize, NewSize >= 16 ? (uint32)16 : (uint32)8);
 	}
 #endif
 #if UE_BUILD_DEBUG
