Build fix for unintended assert introduced by a dummy platform implementation of virtual memory allocator code.

diff --git a/Engine/Source/Runtime/Core/Private/Memory/VirtualStackAllocator.cpp b/Engine/Source/Runtime/Core/Private/Memory/VirtualStackAllocator.cpp
--- a/Engine/Source/Runtime/Core/Private/Memory/VirtualStackAllocator.cpp
+++ b/Engine/Source/Runtime/Core/Private/Memory/VirtualStackAllocator.cpp
@@ -43,41 +43,48 @@ FVirtualStackAllocator::FVirtualStackAllocator(size_t RequestedStackSize, EVirtu
 FVirtualStackAllocator::~FVirtualStackAllocator()
 {
     check(GetAllocatedBytes() == 0);
-    VirtualMemory.FreeVirtual();
+	if (NextUncommittedPage != nullptr)
+	{
+		VirtualMemory.FreeVirtual();
+	}
 }
 
 void* FVirtualStackAllocator::Allocate(size_t Size, size_t Alignment)
 {
-    void* const AllocationStart = Align(NextAllocationStart, Alignment);
-    void* const AllocationEnd = OffsetPointer(AllocationStart, Size);
+	void* const AllocationStart = Align(NextAllocationStart, Alignment);
+	if (Size > 0)
+	{
+		void* const AllocationEnd = OffsetPointer(AllocationStart, Size);
 
-    if (Size > (TotalReservationSize - PageSize))
-    {
-        FPlatformMemory::OnOutOfMemory(Size, Alignment);
-    }
+		if (Size > (TotalReservationSize - PageSize))
+		{
+			FPlatformMemory::OnOutOfMemory(Size, Alignment);
+		}
 
-    void* const UsableMemoryEnd = OffsetPointer(VirtualMemory.GetVirtualPointer(), TotalReservationSize - PageSize);
+		void* const UsableMemoryEnd = OffsetPointer(VirtualMemory.GetVirtualPointer(), TotalReservationSize - PageSize);
 
-	// After the high water mark is established, needing to commit pages should be rare
-    if (UNLIKELY(AllocationEnd > NextUncommittedPage))
-    {
-        // We need to commit some more pages. Let's see how many
-        uintptr_t RequiredAdditionalCommit = PointerDifference(AllocationEnd, NextUncommittedPage);
-        // CommitByPtr doesn't round up the size for you
-        size_t SizeToCommit = Align(RequiredAdditionalCommit, PageSize);
-        VirtualMemory.CommitByPtr(NextUncommittedPage, SizeToCommit);
+		// After the high water mark is established, needing to commit pages should be rare
+		if (UNLIKELY(AllocationEnd > NextUncommittedPage))
+		{
+			// We need to commit some more pages. Let's see how many
+			uintptr_t RequiredAdditionalCommit = PointerDifference(AllocationEnd, NextUncommittedPage);
+			// CommitByPtr doesn't round up the size for you
+			size_t SizeToCommit = Align(RequiredAdditionalCommit, PageSize);
+			VirtualMemory.CommitByPtr(NextUncommittedPage, SizeToCommit);
 
-		LLM_IF_ENABLED(FLowLevelMemTracker::Get().OnLowLevelAlloc(ELLMTracker::Default, NextUncommittedPage, SizeToCommit, LLM_TAG_NAME(VirtualStackAllocator)));
+			LLM_IF_ENABLED(FLowLevelMemTracker::Get().OnLowLevelAlloc(ELLMTracker::Default, NextUncommittedPage, SizeToCommit, LLM_TAG_NAME(VirtualStackAllocator)));
 
-        NextUncommittedPage = Align(AllocationEnd, PageSize);
-    }
+			NextUncommittedPage = Align(AllocationEnd, PageSize);
+		}
 
-	if ((char*)AllocationEnd > (char*)RecentHighWaterMark)
-	{
-		RecentHighWaterMark = Align(AllocationEnd, PageSize);
-	}
+		if ((char*)AllocationEnd > (char*)RecentHighWaterMark)
+		{
+			RecentHighWaterMark = Align(AllocationEnd, PageSize);
+		}
 
-    NextAllocationStart = AllocationEnd;
+		NextAllocationStart = AllocationEnd;
+	}
+ 
 
     return AllocationStart;
 }
