When saving an ini file, the SourceIniHierarchy wasn't being configured correctly, causing things like array clearing to not be handled correctly.

diff --git a/Engine/Source/Runtime/Core/Private/Misc/ConfigCacheIni.cpp b/Engine/Source/Runtime/Core/Private/Misc/ConfigCacheIni.cpp
--- a/Engine/Source/Runtime/Core/Private/Misc/ConfigCacheIni.cpp
+++ b/Engine/Source/Runtime/Core/Private/Misc/ConfigCacheIni.cpp
@@ -3883,6 +3883,8 @@ void FConfigFile::UpdateSections(const TCHAR* DiskFilename, const TCHAR* IniRoot
 		// now when Write it called below, it will diff against this SourceConfigFile
 		FConfigContext BaseContext = FConfigContext::ReadUpToBeforeFile(*SourceConfigFile, OverridePlatform, DiskFilename);
 		BaseContext.Load(IniRootName);
+
+		SourceIniHierarchy = SourceConfigFile->SourceIniHierarchy;
 	}
 
 	WriteInternal(DiskFilename, true, SectionTexts, SectionOrder);
