Andriod: Enabling Logging in PGO builds

diff --git a/Engine/Source/Runtime/Core/Private/Android/AndroidPlatformMisc.cpp b/Engine/Source/Runtime/Core/Private/Android/AndroidPlatformMisc.cpp
--- a/Engine/Source/Runtime/Core/Private/Android/AndroidPlatformMisc.cpp
+++ b/Engine/Source/Runtime/Core/Private/Android/AndroidPlatformMisc.cpp
@@ -266,7 +266,7 @@ void FAndroidMisc::LocalPrint(const TCHAR *Message)
 {
 	// Builds for distribution should not have logging in them:
 	// http://developer.android.com/tools/publishing/preparing.html#publishing-configure
-#if !UE_BUILD_SHIPPING
+#if !UE_BUILD_SHIPPING || ENABLE_PGO_PROFILE
 	const int MAX_LOG_LENGTH = 4096;
 	// not static since may be called by different threads
 	wchar_t MessageBuffer[MAX_LOG_LENGTH];
