LowLevelMemTracker: Fix array iterator failed ensure when tag construction adds an implicitly-defined parent tag that is discovered during Tick's UpdateTags function.
#rnx

diff --git a/Engine/Source/Runtime/Core/Private/HAL/LowLevelMemTracker.cpp b/Engine/Source/Runtime/Core/Private/HAL/LowLevelMemTracker.cpp
--- a/Engine/Source/Runtime/Core/Private/HAL/LowLevelMemTracker.cpp
+++ b/Engine/Source/Runtime/Core/Private/HAL/LowLevelMemTracker.cpp
@@ -930,8 +930,11 @@ void FLowLevelMemTracker::UpdateTags()
 	bool bNeedsResort = false;
 	{
 		FReadScopeLock ScopeLock(TagDataLock);
-		for (FTagData* TagData : *TagDatas)
+		// Cannot use a ranged-for because FinishConstruct can drop our lock and add elements to TagDatas
+		// Check the valid-index condition on every loop iteration
+		for (int32 TagIndex = 0; TagIndex < TagDatas->Num(); ++TagIndex)
 		{
+			FTagData* TagData = (*TagDatas)[TagIndex];
 			FinishConstruct(TagData, UE::LLMPrivate::ETagReferenceSource::FunctionAPI);
 			const FTagData* Parent = TagData->GetParent();
 			if (Parent && Parent->GetIndex() > TagData->GetIndex())
