GitHub #8600 : Rename Ar to StructuredArchive

diff --git a/Engine/Source/Runtime/Core/Private/Serialization/StructuredArchiveChildReader.cpp b/Engine/Source/Runtime/Core/Private/Serialization/StructuredArchiveChildReader.cpp
--- a/Engine/Source/Runtime/Core/Private/Serialization/StructuredArchiveChildReader.cpp
+++ b/Engine/Source/Runtime/Core/Private/Serialization/StructuredArchiveChildReader.cpp
@@ -1,38 +1,38 @@
-// Copyright Epic Games, Inc. All Rights Reserved.
-
-#include "Serialization/StructuredArchiveChildReader.h"
-
-#if WITH_TEXT_ARCHIVE_SUPPORT
-
-FStructuredArchiveChildReader::FStructuredArchiveChildReader(FStructuredArchiveSlot InSlot)
-	: OwnedFormatter(nullptr)
-	, Archive(nullptr)
-{
-	FStructuredArchiveFormatter* Formatter = &InSlot.Ar.Formatter;
-	if (InSlot.GetUnderlyingArchive().IsTextFormat())
-	{
-		Formatter = OwnedFormatter = InSlot.Ar.Formatter.CreateSubtreeReader();
-	}
-
-	Archive = new FStructuredArchive(*Formatter);
-	Root.Emplace(Archive->Open());
-	InSlot.EnterRecord();
-}
-
-FStructuredArchiveChildReader::~FStructuredArchiveChildReader()
-{
-	Root.Reset();
-	Archive->Close();
-	delete Archive;
-	Archive = nullptr;
-
-	// If this is a text archive, we'll have created a subtree reader that our contained archive is using as 
-	// its formatter. We need to clean it up now.
-	if (OwnedFormatter)
-	{
-		delete OwnedFormatter;
-		OwnedFormatter = nullptr;
-	}
-}
-
-#endif
+// Copyright Epic Games, Inc. All Rights Reserved.
+
+#include "Serialization/StructuredArchiveChildReader.h"
+
+#if WITH_TEXT_ARCHIVE_SUPPORT
+
+FStructuredArchiveChildReader::FStructuredArchiveChildReader(FStructuredArchiveSlot InSlot)
+	: OwnedFormatter(nullptr)
+	, StructuredArchive(nullptr)
+{
+	FStructuredArchiveFormatter* Formatter = &InSlot.StructuredArchive.Formatter;
+	if (InSlot.GetUnderlyingArchive().IsTextFormat())
+	{
+		Formatter = OwnedFormatter = InSlot.StructuredArchive.Formatter.CreateSubtreeReader();
+	}
+
+	StructuredArchive = new FStructuredArchive(*Formatter);
+	Root.Emplace(StructuredArchive->Open());
+	InSlot.EnterRecord();
+}
+
+FStructuredArchiveChildReader::~FStructuredArchiveChildReader()
+{
+	Root.Reset();
+	StructuredArchive->Close();
+	delete StructuredArchive;
+	StructuredArchive = nullptr;
+
+	// If this is a text archive, we'll have created a subtree reader that our contained archive is using as 
+	// its formatter. We need to clean it up now.
+	if (OwnedFormatter)
+	{
+		delete OwnedFormatter;
+		OwnedFormatter = nullptr;
+	}
+}
+
+#endif
diff --git a/Engine/Source/Runtime/Core/Private/Serialization/StructuredArchiveSlots.cpp b/Engine/Source/Runtime/Core/Private/Serialization/StructuredArchiveSlots.cpp
--- a/Engine/Source/Runtime/Core/Private/Serialization/StructuredArchiveSlots.cpp
+++ b/Engine/Source/Runtime/Core/Private/Serialization/StructuredArchiveSlots.cpp
@@ -1,439 +1,439 @@
-// Copyright Epic Games, Inc. All Rights Reserved.
-
-#include "Serialization/StructuredArchiveSlots.h"
-#include "Serialization/StructuredArchiveContainer.h"
-#include "Containers/UnrealString.h"
-
-#if WITH_TEXT_ARCHIVE_SUPPORT
-
-//////////// FStructuredArchiveSlot ////////////
-
-FStructuredArchiveRecord FStructuredArchiveSlot::EnterRecord()
-{
-	int32 NewDepth = Ar.EnterSlotAsType(*this, UE::StructuredArchive::Private::EElementType::Record);
-
-#if DO_STRUCTURED_ARCHIVE_CONTAINER_CHECKS
-	Ar.CurrentContainer.Emplace(0);
-#endif
-
-	Ar.Formatter.EnterRecord();
-
-	return FStructuredArchiveRecord(FStructuredArchiveRecord::EPrivateToken{}, Ar, NewDepth, ElementId);
-}
-
-FStructuredArchiveArray FStructuredArchiveSlot::EnterArray(int32& Num)
-{
-	int32 NewDepth = Ar.EnterSlotAsType(*this, UE::StructuredArchive::Private::EElementType::Array);
-
-	Ar.Formatter.EnterArray(Num);
-
-#if DO_STRUCTURED_ARCHIVE_CONTAINER_CHECKS
-	Ar.CurrentContainer.Emplace(Num);
-#endif
-
-	return FStructuredArchiveArray(FStructuredArchiveArray::EPrivateToken{}, Ar, NewDepth, ElementId);
-}
-
-FStructuredArchiveStream FStructuredArchiveSlot::EnterStream()
-{
-	int32 NewDepth = Ar.EnterSlotAsType(*this, UE::StructuredArchive::Private::EElementType::Stream);
-
-	Ar.Formatter.EnterStream();
-
-	return FStructuredArchiveStream(FStructuredArchiveStream::EPrivateToken{}, Ar, NewDepth, ElementId);
-}
-
-FStructuredArchiveMap FStructuredArchiveSlot::EnterMap(int32& Num)
-{
-	int32 NewDepth = Ar.EnterSlotAsType(*this, UE::StructuredArchive::Private::EElementType::Map);
-
-	Ar.Formatter.EnterMap(Num);
-
-#if DO_STRUCTURED_ARCHIVE_CONTAINER_CHECKS
-	Ar.CurrentContainer.Emplace(Num);
-#endif
-
-	return FStructuredArchiveMap(FStructuredArchiveMap::EPrivateToken{}, Ar, NewDepth, ElementId);
-}
-
-FStructuredArchiveSlot FStructuredArchiveSlot::EnterAttribute(FArchiveFieldName AttributeName)
-{
-	check(Ar.CurrentScope.Num() > 0);
-
-	int32 NewDepth = Depth + 1;
-	if (NewDepth >= Ar.CurrentScope.Num() || Ar.CurrentScope[NewDepth].Id != ElementId || Ar.CurrentScope[NewDepth].Type != UE::StructuredArchive::Private::EElementType::AttributedValue)
-	{
-		int32 NewDepthCheck = Ar.EnterSlotAsType(*this, UE::StructuredArchive::Private::EElementType::AttributedValue);
-		checkSlow(NewDepth == NewDepthCheck);
-
-		Ar.Formatter.EnterAttributedValue();
-
-#if DO_STRUCTURED_ARCHIVE_CONTAINER_CHECKS
-		Ar.CurrentContainer.Emplace(0);
-#endif
-	}
-
-	Ar.CurrentEnteringAttributeState = UE::StructuredArchive::Private::EEnteringAttributeState::NotEnteringAttribute;
-
-	UE::StructuredArchive::Private::FElementId AttributedValueId = Ar.CurrentScope[NewDepth].Id;
-
-	Ar.SetScope(UE::StructuredArchive::Private::FSlotPosition(NewDepth, AttributedValueId));
-
-	Ar.CurrentSlotElementId = Ar.ElementIdGenerator.Generate();
-
-#if DO_STRUCTURED_ARCHIVE_CONTAINER_CHECKS
-#if DO_STRUCTURED_ARCHIVE_UNIQUE_FIELD_NAME_CHECKS
-	if (!Ar.GetUnderlyingArchive().IsLoading())
-	{
-		FContainer& Container = *Ar.CurrentContainer.Top();
-		checkf(!Container.KeyNames.Contains(Name.Name), TEXT("Multiple attributes called '%s' serialized into attributed value"), AttributeName.Name);
-		Container.KeyNames.Add(Name.Name);
-	}
-#endif
-#endif
-
-	Ar.Formatter.EnterAttribute(AttributeName);
-
-	return FStructuredArchiveSlot(FStructuredArchiveSlot::EPrivateToken{}, Ar, NewDepth, Ar.CurrentSlotElementId);
-}
-
-TOptional<FStructuredArchiveSlot> FStructuredArchiveSlot::TryEnterAttribute(FArchiveFieldName AttributeName, bool bEnterWhenWriting)
-{
-	check(Ar.CurrentScope.Num() > 0);
-
-	int32 NewDepth = Depth + 1;
-	if (NewDepth >= Ar.CurrentScope.Num() || Ar.CurrentScope[NewDepth].Id != ElementId || Ar.CurrentScope[NewDepth].Type != UE::StructuredArchive::Private::EElementType::AttributedValue)
-	{
-		int32 NewDepthCheck = Ar.EnterSlotAsType(*this, UE::StructuredArchive::Private::EElementType::AttributedValue);
-		checkSlow(NewDepth == NewDepthCheck);
-
-		Ar.Formatter.EnterAttributedValue();
-
-#if DO_STRUCTURED_ARCHIVE_CONTAINER_CHECKS
-		Ar.CurrentContainer.Emplace(0);
-#endif
-	}
-
-#if DO_STRUCTURED_ARCHIVE_CONTAINER_CHECKS
-#if DO_STRUCTURED_ARCHIVE_UNIQUE_FIELD_NAME_CHECKS
-	if (!Ar.GetUnderlyingArchive().IsLoading())
-	{
-		FContainer& Container = *Ar.CurrentContainer.Top();
-		checkf(!Container.KeyNames.Contains(Name.Name), TEXT("Multiple attributes called '%s' serialized into attributed value"), AttributeName.Name);
-		Container.KeyNames.Add(Name.Name);
-	}
-#endif
-#endif
-
-	UE::StructuredArchive::Private::FElementId AttributedValueId = Ar.CurrentScope[NewDepth].Id;
-
-	Ar.SetScope(UE::StructuredArchive::Private::FSlotPosition(NewDepth, AttributedValueId));
-
-	if (Ar.Formatter.TryEnterAttribute(AttributeName, bEnterWhenWriting))
-	{
-		Ar.CurrentEnteringAttributeState = UE::StructuredArchive::Private::EEnteringAttributeState::NotEnteringAttribute;
-
-		Ar.CurrentSlotElementId = Ar.ElementIdGenerator.Generate();
-
-		return FStructuredArchiveSlot(FStructuredArchiveSlot::EPrivateToken{}, Ar, NewDepth, Ar.CurrentSlotElementId);
-	}
-	else
-	{
-		return {};
-	}
-}
-
-void FStructuredArchiveSlot::operator<< (uint8& Value)
-{
-	Ar.EnterSlot(*this);
-	Ar.Formatter.Serialize(Value);
-	Ar.LeaveSlot();
-}
-
-void FStructuredArchiveSlot::operator<< (uint16& Value)
-{
-	Ar.EnterSlot(*this);
-	Ar.Formatter.Serialize(Value);
-	Ar.LeaveSlot();
-}
-
-void FStructuredArchiveSlot::operator<< (uint32& Value)
-{
-	Ar.EnterSlot(*this);
-	Ar.Formatter.Serialize(Value);
-	Ar.LeaveSlot();
-}
-
-void FStructuredArchiveSlot::operator<< (uint64& Value)
-{
-	Ar.EnterSlot(*this);
-	Ar.Formatter.Serialize(Value);
-	Ar.LeaveSlot();
-}
-
-void FStructuredArchiveSlot::operator<< (int8& Value)
-{
-	Ar.EnterSlot(*this);
-	Ar.Formatter.Serialize(Value);
-	Ar.LeaveSlot();
-}
-
-void FStructuredArchiveSlot::operator<< (int16& Value)
-{
-	Ar.EnterSlot(*this);
-	Ar.Formatter.Serialize(Value);
-	Ar.LeaveSlot();
-}
-
-void FStructuredArchiveSlot::operator<< (int32& Value)
-{
-	Ar.EnterSlot(*this);
-	Ar.Formatter.Serialize(Value);
-	Ar.LeaveSlot();
-}
-
-void FStructuredArchiveSlot::operator<< (int64& Value)
-{
-	Ar.EnterSlot(*this);
-	Ar.Formatter.Serialize(Value);
-	Ar.LeaveSlot();
-}
-
-void FStructuredArchiveSlot::operator<< (float& Value)
-{
-	Ar.EnterSlot(*this);
-	Ar.Formatter.Serialize(Value);
-	Ar.LeaveSlot();
-}
-
-void FStructuredArchiveSlot::operator<< (double& Value)
-{
-	Ar.EnterSlot(*this);
-	Ar.Formatter.Serialize(Value);
-	Ar.LeaveSlot();
-}
-
-void FStructuredArchiveSlot::operator<< (bool& Value)
-{
-	Ar.EnterSlot(*this);
-	Ar.Formatter.Serialize(Value);
-	Ar.LeaveSlot();
-}
-
-void FStructuredArchiveSlot::operator<< (FString& Value)
-{
-	Ar.EnterSlot(*this);
-	Ar.Formatter.Serialize(Value);
-	Ar.LeaveSlot();
-}
-
-void FStructuredArchiveSlot::operator<< (FName& Value)
-{
-	Ar.EnterSlot(*this);
-	Ar.Formatter.Serialize(Value);
-	Ar.LeaveSlot();
-}
-
-void FStructuredArchiveSlot::operator<< (UObject*& Value)
-{
-	Ar.EnterSlot(*this);
-	Ar.Formatter.Serialize(Value);
-	Ar.LeaveSlot();
-}
-
-void FStructuredArchiveSlot::operator<< (FText& Value)
-{
-	Ar.EnterSlot(*this);
-	Ar.Formatter.Serialize(Value);
-	Ar.LeaveSlot();
-}
-
-void FStructuredArchiveSlot::operator<< (FWeakObjectPtr& Value)
-{
-	Ar.EnterSlot(*this);
-	Ar.Formatter.Serialize(Value);
-	Ar.LeaveSlot();
-}
-
-void FStructuredArchiveSlot::operator<< (FLazyObjectPtr& Value)
-{
-	Ar.EnterSlot(*this);
-	Ar.Formatter.Serialize(Value);
-	Ar.LeaveSlot();
-}
-
-void FStructuredArchiveSlot::operator<< (FObjectPtr& Value)
-{
-	Ar.EnterSlot(*this);
-	Ar.Formatter.Serialize(Value);
-	Ar.LeaveSlot();
-}
-
-void FStructuredArchiveSlot::operator<< (FSoftObjectPtr& Value)
-{
-	Ar.EnterSlot(*this);
-	Ar.Formatter.Serialize(Value);
-	Ar.LeaveSlot();
-}
-
-void FStructuredArchiveSlot::operator<< (FSoftObjectPath& Value)
-{
-	Ar.EnterSlot(*this);
-	Ar.Formatter.Serialize(Value);
-	Ar.LeaveSlot();
-}
-
-void FStructuredArchiveSlot::Serialize(TArray<uint8>& Data)
-{
-	Ar.EnterSlot(*this);
-	Ar.Formatter.Serialize(Data);
-	Ar.LeaveSlot();
-}
-
-void FStructuredArchiveSlot::Serialize(void* Data, uint64 DataSize)
-{
-	Ar.EnterSlot(*this);
-	Ar.Formatter.Serialize(Data, DataSize);
-	Ar.LeaveSlot();
-}
-
-//////////// FStructuredArchiveRecord ////////////
-
-FStructuredArchiveSlot FStructuredArchiveRecord::EnterField(FArchiveFieldName Name)
-{
-	Ar.SetScope(*this);
-
-	Ar.CurrentSlotElementId = Ar.ElementIdGenerator.Generate();
-
-#if DO_STRUCTURED_ARCHIVE_CONTAINER_CHECKS
-#if DO_STRUCTURED_ARCHIVE_UNIQUE_FIELD_NAME_CHECKS
-	if (!Ar.GetUnderlyingArchive().IsLoading())
-	{
-		FContainer& Container = *Ar.CurrentContainer.Top();
-		checkf(!Container.KeyNames.Contains(Name.Name), TEXT("Multiple keys called '%s' serialized into record"), Name.Name);
-		Container.KeyNames.Add(Name.Name);
-	}
-#endif
-#endif
-
-	Ar.Formatter.EnterField(Name);
-
-	return FStructuredArchiveSlot(FStructuredArchiveSlot::EPrivateToken{}, Ar, Depth, Ar.CurrentSlotElementId);
-}
-
-FStructuredArchiveRecord FStructuredArchiveRecord::EnterRecord(FArchiveFieldName Name)
-{
-	return EnterField(Name).EnterRecord();
-}
-
-FStructuredArchiveArray FStructuredArchiveRecord::EnterArray(FArchiveFieldName Name, int32& Num)
-{
-	return EnterField(Name).EnterArray(Num);
-}
-
-FStructuredArchiveStream FStructuredArchiveRecord::EnterStream(FArchiveFieldName Name)
-{
-	return EnterField(Name).EnterStream();
-}
-
-FStructuredArchiveMap FStructuredArchiveRecord::EnterMap(FArchiveFieldName Name, int32& Num)
-{
-	return EnterField(Name).EnterMap(Num);
-}
-
-TOptional<FStructuredArchiveSlot> FStructuredArchiveRecord::TryEnterField(FArchiveFieldName Name, bool bEnterWhenWriting)
-{
-	Ar.SetScope(*this);
-
-#if DO_STRUCTURED_ARCHIVE_CONTAINER_CHECKS
-#if DO_STRUCTURED_ARCHIVE_UNIQUE_FIELD_NAME_CHECKS
-	if (!GetUnderlyingArchive().IsLoading())
-	{
-		FContainer& Container = *Ar.CurrentContainer.Top();
-		checkf(!Container.KeyNames.Contains(Name.Name), TEXT("Multiple keys called '%s' serialized into record"), Name.Name);
-		Container.KeyNames.Add(Name.Name);
-	}
-#endif
-#endif
-
-	if (Ar.Formatter.TryEnterField(Name, bEnterWhenWriting))
-	{
-		Ar.CurrentSlotElementId = Ar.ElementIdGenerator.Generate();
-		return FStructuredArchiveSlot(FStructuredArchiveSlot::EPrivateToken{}, Ar, Depth, Ar.CurrentSlotElementId);
-	}
-	else
-	{
-		return {};
-	}
-}
-
-//////////// FStructuredArchiveArray ////////////
-
-FStructuredArchiveSlot FStructuredArchiveArray::EnterElement()
-{
-	Ar.SetScope(*this);
-
-#if DO_STRUCTURED_ARCHIVE_CONTAINER_CHECKS
-	checkf(Ar.CurrentContainer.Top()->Index < Ar.CurrentContainer.Top()->Count, TEXT("Serialized too many array elements"));
-#endif
-
-	Ar.CurrentSlotElementId = Ar.ElementIdGenerator.Generate();
-
-	Ar.Formatter.EnterArrayElement();
-
-	return FStructuredArchiveSlot(FStructuredArchiveSlot::EPrivateToken{}, Ar, Depth, Ar.CurrentSlotElementId);
-}
-
-//////////// FStructuredArchiveStream ////////////
-
-FStructuredArchiveSlot FStructuredArchiveStream::EnterElement()
-{
-	Ar.SetScope(*this);
-
-	Ar.CurrentSlotElementId = Ar.ElementIdGenerator.Generate();
-
-	Ar.Formatter.EnterStreamElement();
-
-	return FStructuredArchiveSlot(FStructuredArchiveSlot::EPrivateToken{}, Ar, Depth, Ar.CurrentSlotElementId);
-}
-
-//////////// FStructuredArchiveMap ////////////
-
-FStructuredArchiveSlot FStructuredArchiveMap::EnterElement(FString& Name)
-{
-	Ar.SetScope(*this);
-
-#if DO_STRUCTURED_ARCHIVE_CONTAINER_CHECKS
-	checkf(Ar.CurrentContainer.Top()->Index < Ar.CurrentContainer.Top()->Count, TEXT("Serialized too many map elements"));
-#endif
-
-	Ar.CurrentSlotElementId = Ar.ElementIdGenerator.Generate();
-
-#if DO_STRUCTURED_ARCHIVE_CONTAINER_CHECKS
-#if DO_STRUCTURED_ARCHIVE_UNIQUE_FIELD_NAME_CHECKS
-	if(Ar.GetUnderlyingArchive().IsSaving())
-	{
-		FContainer& Container = *Ar.CurrentContainer.Top();
-		checkf(!Container.KeyNames.Contains(Name), TEXT("Multiple keys called '%s' serialized into record"), *Name);
-		Container.KeyNames.Add(Name);
-	}
-#endif
-#endif
-
-	Ar.Formatter.EnterMapElement(Name);
-
-#if DO_STRUCTURED_ARCHIVE_CONTAINER_CHECKS
-#if DO_STRUCTURED_ARCHIVE_UNIQUE_FIELD_NAME_CHECKS
-	if(Ar.GetUnderlyingArchive().IsLoading())
-	{
-		FContainer& Container = *Ar.CurrentContainer.Top();
-		checkf(!Container.KeyNames.Contains(Name), TEXT("Multiple keys called '%s' serialized into record"), *Name);
-		Container.KeyNames.Add(Name);
-	}
-#endif
-#endif
-
-	return FStructuredArchiveSlot(FStructuredArchiveSlot::EPrivateToken{}, Ar, Depth, Ar.CurrentSlotElementId);
-}
-
-#endif
+// Copyright Epic Games, Inc. All Rights Reserved.
+
+#include "Serialization/StructuredArchiveSlots.h"
+#include "Serialization/StructuredArchiveContainer.h"
+#include "Containers/UnrealString.h"
+
+#if WITH_TEXT_ARCHIVE_SUPPORT
+
+//////////// FStructuredArchiveSlot ////////////
+
+FStructuredArchiveRecord FStructuredArchiveSlot::EnterRecord()
+{
+	int32 NewDepth = StructuredArchive.EnterSlotAsType(*this, UE::StructuredArchive::Private::EElementType::Record);
+
+#if DO_STRUCTURED_ARCHIVE_CONTAINER_CHECKS
+	StructuredArchive.CurrentContainer.Emplace(0);
+#endif
+
+	StructuredArchive.Formatter.EnterRecord();
+
+	return FStructuredArchiveRecord(FStructuredArchiveRecord::EPrivateToken{}, StructuredArchive, NewDepth, ElementId);
+}
+
+FStructuredArchiveArray FStructuredArchiveSlot::EnterArray(int32& Num)
+{
+	int32 NewDepth = StructuredArchive.EnterSlotAsType(*this, UE::StructuredArchive::Private::EElementType::Array);
+
+	StructuredArchive.Formatter.EnterArray(Num);
+
+#if DO_STRUCTURED_ARCHIVE_CONTAINER_CHECKS
+	StructuredArchive.CurrentContainer.Emplace(Num);
+#endif
+
+	return FStructuredArchiveArray(FStructuredArchiveArray::EPrivateToken{}, StructuredArchive, NewDepth, ElementId);
+}
+
+FStructuredArchiveStream FStructuredArchiveSlot::EnterStream()
+{
+	int32 NewDepth = StructuredArchive.EnterSlotAsType(*this, UE::StructuredArchive::Private::EElementType::Stream);
+
+	StructuredArchive.Formatter.EnterStream();
+
+	return FStructuredArchiveStream(FStructuredArchiveStream::EPrivateToken{}, StructuredArchive, NewDepth, ElementId);
+}
+
+FStructuredArchiveMap FStructuredArchiveSlot::EnterMap(int32& Num)
+{
+	int32 NewDepth = StructuredArchive.EnterSlotAsType(*this, UE::StructuredArchive::Private::EElementType::Map);
+
+	StructuredArchive.Formatter.EnterMap(Num);
+
+#if DO_STRUCTURED_ARCHIVE_CONTAINER_CHECKS
+	StructuredArchive.CurrentContainer.Emplace(Num);
+#endif
+
+	return FStructuredArchiveMap(FStructuredArchiveMap::EPrivateToken{}, StructuredArchive, NewDepth, ElementId);
+}
+
+FStructuredArchiveSlot FStructuredArchiveSlot::EnterAttribute(FArchiveFieldName AttributeName)
+{
+	check(StructuredArchive.CurrentScope.Num() > 0);
+
+	int32 NewDepth = Depth + 1;
+	if (NewDepth >= StructuredArchive.CurrentScope.Num() || StructuredArchive.CurrentScope[NewDepth].Id != ElementId || StructuredArchive.CurrentScope[NewDepth].Type != UE::StructuredArchive::Private::EElementType::AttributedValue)
+	{
+		int32 NewDepthCheck = StructuredArchive.EnterSlotAsType(*this, UE::StructuredArchive::Private::EElementType::AttributedValue);
+		checkSlow(NewDepth == NewDepthCheck);
+
+		StructuredArchive.Formatter.EnterAttributedValue();
+
+#if DO_STRUCTURED_ARCHIVE_CONTAINER_CHECKS
+		StructuredArchive.CurrentContainer.Emplace(0);
+#endif
+	}
+
+	StructuredArchive.CurrentEnteringAttributeState = UE::StructuredArchive::Private::EEnteringAttributeState::NotEnteringAttribute;
+
+	UE::StructuredArchive::Private::FElementId AttributedValueId = StructuredArchive.CurrentScope[NewDepth].Id;
+
+	StructuredArchive.SetScope(UE::StructuredArchive::Private::FSlotPosition(NewDepth, AttributedValueId));
+
+	StructuredArchive.CurrentSlotElementId = StructuredArchive.ElementIdGenerator.Generate();
+
+#if DO_STRUCTURED_ARCHIVE_CONTAINER_CHECKS
+#if DO_STRUCTURED_ARCHIVE_UNIQUE_FIELD_NAME_CHECKS
+	if (!StructuredArchive.GetUnderlyingArchive().IsLoading())
+	{
+		FContainer& Container = *StructuredArchive.CurrentContainer.Top();
+		checkf(!Container.KeyNames.Contains(Name.Name), TEXT("Multiple attributes called '%s' serialized into attributed value"), AttributeName.Name);
+		Container.KeyNames.Add(Name.Name);
+	}
+#endif
+#endif
+
+	StructuredArchive.Formatter.EnterAttribute(AttributeName);
+
+	return FStructuredArchiveSlot(FStructuredArchiveSlot::EPrivateToken{}, StructuredArchive, NewDepth, StructuredArchive.CurrentSlotElementId);
+}
+
+TOptional<FStructuredArchiveSlot> FStructuredArchiveSlot::TryEnterAttribute(FArchiveFieldName AttributeName, bool bEnterWhenWriting)
+{
+	check(StructuredArchive.CurrentScope.Num() > 0);
+
+	int32 NewDepth = Depth + 1;
+	if (NewDepth >= StructuredArchive.CurrentScope.Num() || StructuredArchive.CurrentScope[NewDepth].Id != ElementId || StructuredArchive.CurrentScope[NewDepth].Type != UE::StructuredArchive::Private::EElementType::AttributedValue)
+	{
+		int32 NewDepthCheck = StructuredArchive.EnterSlotAsType(*this, UE::StructuredArchive::Private::EElementType::AttributedValue);
+		checkSlow(NewDepth == NewDepthCheck);
+
+		StructuredArchive.Formatter.EnterAttributedValue();
+
+#if DO_STRUCTURED_ARCHIVE_CONTAINER_CHECKS
+		StructuredArchive.CurrentContainer.Emplace(0);
+#endif
+	}
+
+#if DO_STRUCTURED_ARCHIVE_CONTAINER_CHECKS
+#if DO_STRUCTURED_ARCHIVE_UNIQUE_FIELD_NAME_CHECKS
+	if (!StructuredArchive.GetUnderlyingArchive().IsLoading())
+	{
+		FContainer& Container = *StructuredArchive.CurrentContainer.Top();
+		checkf(!Container.KeyNames.Contains(Name.Name), TEXT("Multiple attributes called '%s' serialized into attributed value"), AttributeName.Name);
+		Container.KeyNames.Add(Name.Name);
+	}
+#endif
+#endif
+
+	UE::StructuredArchive::Private::FElementId AttributedValueId = StructuredArchive.CurrentScope[NewDepth].Id;
+
+	StructuredArchive.SetScope(UE::StructuredArchive::Private::FSlotPosition(NewDepth, AttributedValueId));
+
+	if (StructuredArchive.Formatter.TryEnterAttribute(AttributeName, bEnterWhenWriting))
+	{
+		StructuredArchive.CurrentEnteringAttributeState = UE::StructuredArchive::Private::EEnteringAttributeState::NotEnteringAttribute;
+
+		StructuredArchive.CurrentSlotElementId = StructuredArchive.ElementIdGenerator.Generate();
+
+		return FStructuredArchiveSlot(FStructuredArchiveSlot::EPrivateToken{}, StructuredArchive, NewDepth, StructuredArchive.CurrentSlotElementId);
+	}
+	else
+	{
+		return {};
+	}
+}
+
+void FStructuredArchiveSlot::operator<< (uint8& Value)
+{
+	StructuredArchive.EnterSlot(*this);
+	StructuredArchive.Formatter.Serialize(Value);
+	StructuredArchive.LeaveSlot();
+}
+
+void FStructuredArchiveSlot::operator<< (uint16& Value)
+{
+	StructuredArchive.EnterSlot(*this);
+	StructuredArchive.Formatter.Serialize(Value);
+	StructuredArchive.LeaveSlot();
+}
+
+void FStructuredArchiveSlot::operator<< (uint32& Value)
+{
+	StructuredArchive.EnterSlot(*this);
+	StructuredArchive.Formatter.Serialize(Value);
+	StructuredArchive.LeaveSlot();
+}
+
+void FStructuredArchiveSlot::operator<< (uint64& Value)
+{
+	StructuredArchive.EnterSlot(*this);
+	StructuredArchive.Formatter.Serialize(Value);
+	StructuredArchive.LeaveSlot();
+}
+
+void FStructuredArchiveSlot::operator<< (int8& Value)
+{
+	StructuredArchive.EnterSlot(*this);
+	StructuredArchive.Formatter.Serialize(Value);
+	StructuredArchive.LeaveSlot();
+}
+
+void FStructuredArchiveSlot::operator<< (int16& Value)
+{
+	StructuredArchive.EnterSlot(*this);
+	StructuredArchive.Formatter.Serialize(Value);
+	StructuredArchive.LeaveSlot();
+}
+
+void FStructuredArchiveSlot::operator<< (int32& Value)
+{
+	StructuredArchive.EnterSlot(*this);
+	StructuredArchive.Formatter.Serialize(Value);
+	StructuredArchive.LeaveSlot();
+}
+
+void FStructuredArchiveSlot::operator<< (int64& Value)
+{
+	StructuredArchive.EnterSlot(*this);
+	StructuredArchive.Formatter.Serialize(Value);
+	StructuredArchive.LeaveSlot();
+}
+
+void FStructuredArchiveSlot::operator<< (float& Value)
+{
+	StructuredArchive.EnterSlot(*this);
+	StructuredArchive.Formatter.Serialize(Value);
+	StructuredArchive.LeaveSlot();
+}
+
+void FStructuredArchiveSlot::operator<< (double& Value)
+{
+	StructuredArchive.EnterSlot(*this);
+	StructuredArchive.Formatter.Serialize(Value);
+	StructuredArchive.LeaveSlot();
+}
+
+void FStructuredArchiveSlot::operator<< (bool& Value)
+{
+	StructuredArchive.EnterSlot(*this);
+	StructuredArchive.Formatter.Serialize(Value);
+	StructuredArchive.LeaveSlot();
+}
+
+void FStructuredArchiveSlot::operator<< (FString& Value)
+{
+	StructuredArchive.EnterSlot(*this);
+	StructuredArchive.Formatter.Serialize(Value);
+	StructuredArchive.LeaveSlot();
+}
+
+void FStructuredArchiveSlot::operator<< (FName& Value)
+{
+	StructuredArchive.EnterSlot(*this);
+	StructuredArchive.Formatter.Serialize(Value);
+	StructuredArchive.LeaveSlot();
+}
+
+void FStructuredArchiveSlot::operator<< (UObject*& Value)
+{
+	StructuredArchive.EnterSlot(*this);
+	StructuredArchive.Formatter.Serialize(Value);
+	StructuredArchive.LeaveSlot();
+}
+
+void FStructuredArchiveSlot::operator<< (FText& Value)
+{
+	StructuredArchive.EnterSlot(*this);
+	StructuredArchive.Formatter.Serialize(Value);
+	StructuredArchive.LeaveSlot();
+}
+
+void FStructuredArchiveSlot::operator<< (FWeakObjectPtr& Value)
+{
+	StructuredArchive.EnterSlot(*this);
+	StructuredArchive.Formatter.Serialize(Value);
+	StructuredArchive.LeaveSlot();
+}
+
+void FStructuredArchiveSlot::operator<< (FLazyObjectPtr& Value)
+{
+	StructuredArchive.EnterSlot(*this);
+	StructuredArchive.Formatter.Serialize(Value);
+	StructuredArchive.LeaveSlot();
+}
+
+void FStructuredArchiveSlot::operator<< (FObjectPtr& Value)
+{
+	StructuredArchive.EnterSlot(*this);
+	StructuredArchive.Formatter.Serialize(Value);
+	StructuredArchive.LeaveSlot();
+}
+
+void FStructuredArchiveSlot::operator<< (FSoftObjectPtr& Value)
+{
+	StructuredArchive.EnterSlot(*this);
+	StructuredArchive.Formatter.Serialize(Value);
+	StructuredArchive.LeaveSlot();
+}
+
+void FStructuredArchiveSlot::operator<< (FSoftObjectPath& Value)
+{
+	StructuredArchive.EnterSlot(*this);
+	StructuredArchive.Formatter.Serialize(Value);
+	StructuredArchive.LeaveSlot();
+}
+
+void FStructuredArchiveSlot::Serialize(TArray<uint8>& Data)
+{
+	StructuredArchive.EnterSlot(*this);
+	StructuredArchive.Formatter.Serialize(Data);
+	StructuredArchive.LeaveSlot();
+}
+
+void FStructuredArchiveSlot::Serialize(void* Data, uint64 DataSize)
+{
+	StructuredArchive.EnterSlot(*this);
+	StructuredArchive.Formatter.Serialize(Data, DataSize);
+	StructuredArchive.LeaveSlot();
+}
+
+//////////// FStructuredArchiveRecord ////////////
+
+FStructuredArchiveSlot FStructuredArchiveRecord::EnterField(FArchiveFieldName Name)
+{
+	StructuredArchive.SetScope(*this);
+
+	StructuredArchive.CurrentSlotElementId = StructuredArchive.ElementIdGenerator.Generate();
+
+#if DO_STRUCTURED_ARCHIVE_CONTAINER_CHECKS
+#if DO_STRUCTURED_ARCHIVE_UNIQUE_FIELD_NAME_CHECKS
+	if (!StructuredArchive.GetUnderlyingArchive().IsLoading())
+	{
+		FContainer& Container = *Ar.CurrentContainer.Top();
+		checkf(!Container.KeyNames.Contains(Name.Name), TEXT("Multiple keys called '%s' serialized into record"), Name.Name);
+		Container.KeyNames.Add(Name.Name);
+	}
+#endif
+#endif
+
+	StructuredArchive.Formatter.EnterField(Name);
+
+	return FStructuredArchiveSlot(FStructuredArchiveSlot::EPrivateToken{}, StructuredArchive, Depth, StructuredArchive.CurrentSlotElementId);
+}
+
+FStructuredArchiveRecord FStructuredArchiveRecord::EnterRecord(FArchiveFieldName Name)
+{
+	return EnterField(Name).EnterRecord();
+}
+
+FStructuredArchiveArray FStructuredArchiveRecord::EnterArray(FArchiveFieldName Name, int32& Num)
+{
+	return EnterField(Name).EnterArray(Num);
+}
+
+FStructuredArchiveStream FStructuredArchiveRecord::EnterStream(FArchiveFieldName Name)
+{
+	return EnterField(Name).EnterStream();
+}
+
+FStructuredArchiveMap FStructuredArchiveRecord::EnterMap(FArchiveFieldName Name, int32& Num)
+{
+	return EnterField(Name).EnterMap(Num);
+}
+
+TOptional<FStructuredArchiveSlot> FStructuredArchiveRecord::TryEnterField(FArchiveFieldName Name, bool bEnterWhenWriting)
+{
+	StructuredArchive.SetScope(*this);
+
+#if DO_STRUCTURED_ARCHIVE_CONTAINER_CHECKS
+#if DO_STRUCTURED_ARCHIVE_UNIQUE_FIELD_NAME_CHECKS
+	if (!GetUnderlyingArchive().IsLoading())
+	{
+		FContainer& Container = *Ar.CurrentContainer.Top();
+		checkf(!Container.KeyNames.Contains(Name.Name), TEXT("Multiple keys called '%s' serialized into record"), Name.Name);
+		Container.KeyNames.Add(Name.Name);
+	}
+#endif
+#endif
+
+	if (StructuredArchive.Formatter.TryEnterField(Name, bEnterWhenWriting))
+	{
+		StructuredArchive.CurrentSlotElementId = StructuredArchive.ElementIdGenerator.Generate();
+		return FStructuredArchiveSlot(FStructuredArchiveSlot::EPrivateToken{}, StructuredArchive, Depth, StructuredArchive.CurrentSlotElementId);
+	}
+	else
+	{
+		return {};
+	}
+}
+
+//////////// FStructuredArchiveArray ////////////
+
+FStructuredArchiveSlot FStructuredArchiveArray::EnterElement()
+{
+	StructuredArchive.SetScope(*this);
+
+#if DO_STRUCTURED_ARCHIVE_CONTAINER_CHECKS
+	checkf(StructuredArchive.CurrentContainer.Top()->Index < StructuredArchive.CurrentContainer.Top()->Count, TEXT("Serialized too many array elements"));
+#endif
+
+	StructuredArchive.CurrentSlotElementId = StructuredArchive.ElementIdGenerator.Generate();
+
+	StructuredArchive.Formatter.EnterArrayElement();
+
+	return FStructuredArchiveSlot(FStructuredArchiveSlot::EPrivateToken{}, StructuredArchive, Depth, StructuredArchive.CurrentSlotElementId);
+}
+
+//////////// FStructuredArchiveStream ////////////
+
+FStructuredArchiveSlot FStructuredArchiveStream::EnterElement()
+{
+	StructuredArchive.SetScope(*this);
+
+	StructuredArchive.CurrentSlotElementId = StructuredArchive.ElementIdGenerator.Generate();
+
+	StructuredArchive.Formatter.EnterStreamElement();
+
+	return FStructuredArchiveSlot(FStructuredArchiveSlot::EPrivateToken{}, StructuredArchive, Depth, StructuredArchive.CurrentSlotElementId);
+}
+
+//////////// FStructuredArchiveMap ////////////
+
+FStructuredArchiveSlot FStructuredArchiveMap::EnterElement(FString& Name)
+{
+	StructuredArchive.SetScope(*this);
+
+#if DO_STRUCTURED_ARCHIVE_CONTAINER_CHECKS
+	checkf(StructuredArchive.CurrentContainer.Top()->Index < StructuredArchive.CurrentContainer.Top()->Count, TEXT("Serialized too many map elements"));
+#endif
+
+	StructuredArchive.CurrentSlotElementId = StructuredArchive.ElementIdGenerator.Generate();
+
+#if DO_STRUCTURED_ARCHIVE_CONTAINER_CHECKS
+#if DO_STRUCTURED_ARCHIVE_UNIQUE_FIELD_NAME_CHECKS
+	if(StructuredArchive.GetUnderlyingArchive().IsSaving())
+	{
+		FContainer& Container = *StructuredArchive.CurrentContainer.Top();
+		checkf(!Container.KeyNames.Contains(Name), TEXT("Multiple keys called '%s' serialized into record"), *Name);
+		Container.KeyNames.Add(Name);
+	}
+#endif
+#endif
+
+	StructuredArchive.Formatter.EnterMapElement(Name);
+
+#if DO_STRUCTURED_ARCHIVE_CONTAINER_CHECKS
+#if DO_STRUCTURED_ARCHIVE_UNIQUE_FIELD_NAME_CHECKS
+	if(StructuredArchive.GetUnderlyingArchive().IsLoading())
+	{
+		FContainer& Container = *StructuredArchive.CurrentContainer.Top();
+		checkf(!Container.KeyNames.Contains(Name), TEXT("Multiple keys called '%s' serialized into record"), *Name);
+		Container.KeyNames.Add(Name);
+	}
+#endif
+#endif
+
+	return FStructuredArchiveSlot(FStructuredArchiveSlot::EPrivateToken{}, StructuredArchive, Depth, StructuredArchive.CurrentSlotElementId);
+}
+
+#endif
diff --git a/Engine/Source/Runtime/Core/Public/Serialization/StructuredArchive.h b/Engine/Source/Runtime/Core/Public/Serialization/StructuredArchive.h
--- a/Engine/Source/Runtime/Core/Public/Serialization/StructuredArchive.h
+++ b/Engine/Source/Runtime/Core/Public/Serialization/StructuredArchive.h
@@ -218,21 +218,21 @@
 #endif
 
 namespace UE::StructuredArchive::Private
+{
+	FORCEINLINE FArchive& GetUnderlyingArchiveImpl(FStructuredArchive& StructuredArchive)
 	{
-	FORCEINLINE FArchive& GetUnderlyingArchiveImpl(FStructuredArchive& Ar)
-	{
-		return Ar.GetUnderlyingArchive();
+		return StructuredArchive.GetUnderlyingArchive();
 	}
 
 #if WITH_TEXT_ARCHIVE_SUPPORT
-	FORCEINLINE FElementId GetCurrentSlotElementIdImpl(FStructuredArchive& Ar)
+	FORCEINLINE FElementId GetCurrentSlotElementIdImpl(FStructuredArchive& StructuredArchive)
 	{
-		return Ar.CurrentSlotElementId;
+		return StructuredArchive.CurrentSlotElementId;
 	}
 #endif
 
-	FORCEINLINE FArchiveFormatterType& GetFormatterImpl(FStructuredArchive& Ar)
+	FORCEINLINE FArchiveFormatterType& GetFormatterImpl(FStructuredArchive& StructuredArchive)
 	{
-		return Ar.Formatter;
-		}
+		return StructuredArchive.Formatter;
 	}
+}
diff --git a/Engine/Source/Runtime/Core/Public/Serialization/StructuredArchiveChildReader.h b/Engine/Source/Runtime/Core/Public/Serialization/StructuredArchiveChildReader.h
--- a/Engine/Source/Runtime/Core/Public/Serialization/StructuredArchiveChildReader.h
+++ b/Engine/Source/Runtime/Core/Public/Serialization/StructuredArchiveChildReader.h
@@ -33,7 +33,7 @@
 
 private:
 	FStructuredArchiveFormatter* OwnedFormatter;
-	FStructuredArchive* Archive;
+	FStructuredArchive* StructuredArchive;
 	TOptional<FStructuredArchiveSlot> Root;
 };
 
@@ -41,15 +41,15 @@
 
 FORCEINLINE FStructuredArchiveChildReader::FStructuredArchiveChildReader(FStructuredArchiveSlot InSlot)
 	: OwnedFormatter(nullptr)
-	, Archive(nullptr)
+	, StructuredArchive(nullptr)
 {
-	Archive = new FStructuredArchive(InSlot.Ar.Formatter);
-	Root.Emplace(Archive->Open());
+	StructuredArchive = new FStructuredArchive(InSlot.StructuredArchive.Formatter);
+	Root.Emplace(StructuredArchive->Open());
 }
 
 FORCEINLINE FStructuredArchiveChildReader::~FStructuredArchiveChildReader()
 {
-	delete Archive;
+	delete StructuredArchive;
 }
 
 #endif
diff --git a/Engine/Source/Runtime/Core/Public/Serialization/StructuredArchiveSlotBase.h b/Engine/Source/Runtime/Core/Public/Serialization/StructuredArchiveSlotBase.h
--- a/Engine/Source/Runtime/Core/Public/Serialization/StructuredArchiveSlotBase.h
+++ b/Engine/Source/Runtime/Core/Public/Serialization/StructuredArchiveSlotBase.h
@@ -74,29 +74,29 @@
 
 	public:
 #if WITH_TEXT_ARCHIVE_SUPPORT
-		FORCEINLINE explicit FSlotBase(EPrivateToken, FStructuredArchive& InAr, int32 InDepth, FElementId InElementId)
+		FORCEINLINE explicit FSlotBase(EPrivateToken, FStructuredArchive& InStructuredArchive, int32 InDepth, FElementId InElementId)
 			: FSlotPosition(InDepth, InElementId)
-			, Ar(InAr)
+			, StructuredArchive(InStructuredArchive)
 		{
 		}
 #else
-		FORCEINLINE explicit FSlotBase(EPrivateToken, FStructuredArchive& InAr)
-			: Ar(InAr)
+		FORCEINLINE explicit FSlotBase(EPrivateToken, FStructuredArchive& InStructuredArchive)
+			: StructuredArchive(InStructuredArchive)
 		{
 		}
 #endif
 
 		FArchive& GetUnderlyingArchive() const
 		{
-			return GetUnderlyingArchiveImpl(Ar);
+			return GetUnderlyingArchiveImpl(StructuredArchive);
 		}
 
 		const FArchiveState& GetArchiveState() const
 		{
-			return GetUnderlyingArchiveImpl(Ar).GetArchiveState();
+			return GetUnderlyingArchiveImpl(StructuredArchive).GetArchiveState();
 		}
 
 	protected:
-		FStructuredArchive& Ar;
+		FStructuredArchive& StructuredArchive;
 	};
 }
diff --git a/Engine/Source/Runtime/Core/Public/Serialization/StructuredArchiveSlots.h b/Engine/Source/Runtime/Core/Public/Serialization/StructuredArchiveSlots.h
--- a/Engine/Source/Runtime/Core/Public/Serialization/StructuredArchiveSlots.h
+++ b/Engine/Source/Runtime/Core/Public/Serialization/StructuredArchiveSlots.h
@@ -19,8 +19,8 @@
 
 namespace UE::StructuredArchive::Private
 {
-	FElementId GetCurrentSlotElementIdImpl(FStructuredArchive& Ar);
-	FArchiveFormatterType& GetFormatterImpl(FStructuredArchive& Ar);
+	FElementId GetCurrentSlotElementIdImpl(FStructuredArchive& StructuredArchive);
+	FArchiveFormatterType& GetFormatterImpl(FStructuredArchive& StructuredArchive);
 }
 
 /**
@@ -102,7 +102,7 @@
 	FORCEINLINE bool IsFilled() const
 	{
 #if WITH_TEXT_ARCHIVE_SUPPORT
-		return UE::StructuredArchive::Private::GetCurrentSlotElementIdImpl(Ar) != ElementId;
+		return UE::StructuredArchive::Private::GetCurrentSlotElementIdImpl(StructuredArchive) != ElementId;
 #else
 		return true;
 #endif
@@ -235,37 +235,37 @@
 	//////////// FStructuredArchiveSlot ////////////
 	FORCEINLINE FStructuredArchiveRecord FStructuredArchiveSlot::EnterRecord()
 	{
-		return FStructuredArchiveRecord(FStructuredArchiveRecord::EPrivateToken{}, Ar);
+		return FStructuredArchiveRecord(FStructuredArchiveRecord::EPrivateToken{}, StructuredArchive);
 	}
 
 	FORCEINLINE FStructuredArchiveArray FStructuredArchiveSlot::EnterArray(int32& Num)
 	{
-		UE::StructuredArchive::Private::GetFormatterImpl(Ar).EnterArray(Num);
-		return FStructuredArchiveArray(FStructuredArchiveArray::EPrivateToken{}, Ar);
+		UE::StructuredArchive::Private::GetFormatterImpl(StructuredArchive).EnterArray(Num);
+		return FStructuredArchiveArray(FStructuredArchiveArray::EPrivateToken{}, StructuredArchive);
 	}
 
 	FORCEINLINE FStructuredArchiveStream FStructuredArchiveSlot::EnterStream()
 	{
-		return FStructuredArchiveStream(FStructuredArchiveStream::EPrivateToken{}, Ar);
+		return FStructuredArchiveStream(FStructuredArchiveStream::EPrivateToken{}, StructuredArchive);
 	}
 
 	FORCEINLINE FStructuredArchiveMap FStructuredArchiveSlot::EnterMap(int32& Num)
 	{
-		UE::StructuredArchive::Private::GetFormatterImpl(Ar).EnterMap(Num);
-		return FStructuredArchiveMap(FStructuredArchiveMap::EPrivateToken{}, Ar);
+		UE::StructuredArchive::Private::GetFormatterImpl(StructuredArchive).EnterMap(Num);
+		return FStructuredArchiveMap(FStructuredArchiveMap::EPrivateToken{}, StructuredArchive);
 	}
 
 	FORCEINLINE FStructuredArchiveSlot FStructuredArchiveSlot::EnterAttribute(FArchiveFieldName FieldName)
 	{
-		UE::StructuredArchive::Private::GetFormatterImpl(Ar).EnterAttribute(FieldName);
-		return FStructuredArchiveSlot(FStructuredArchiveSlot::EPrivateToken{}, Ar);
+		UE::StructuredArchive::Private::GetFormatterImpl(StructuredArchive).EnterAttribute(FieldName);
+		return FStructuredArchiveSlot(FStructuredArchiveSlot::EPrivateToken{}, StructuredArchive);
 	}
 
 	FORCEINLINE TOptional<FStructuredArchiveSlot> FStructuredArchiveSlot::TryEnterAttribute(FArchiveFieldName FieldName, bool bEnterWhenWriting)
 	{
-		if (UE::StructuredArchive::Private::GetFormatterImpl(Ar).TryEnterAttribute(FieldName, bEnterWhenWriting))
+		if (UE::StructuredArchive::Private::GetFormatterImpl(StructuredArchive).TryEnterAttribute(FieldName, bEnterWhenWriting))
 		{
-			return TOptional<FStructuredArchiveSlot>(InPlace, FStructuredArchiveSlot::EPrivateToken{}, Ar);
+			return TOptional<FStructuredArchiveSlot>(InPlace, FStructuredArchiveSlot::EPrivateToken{}, StructuredArchive);
 		}
 		else
 		{
@@ -275,119 +275,119 @@
 
 	FORCEINLINE void FStructuredArchiveSlot::operator<< (uint8& Value)
 	{
-		UE::StructuredArchive::Private::GetFormatterImpl(Ar).Serialize(Value);
+		UE::StructuredArchive::Private::GetFormatterImpl(StructuredArchive).Serialize(Value);
 	}
 
 	FORCEINLINE void FStructuredArchiveSlot::operator<< (uint16& Value)
 	{
-		UE::StructuredArchive::Private::GetFormatterImpl(Ar).Serialize(Value);
+		UE::StructuredArchive::Private::GetFormatterImpl(StructuredArchive).Serialize(Value);
 	}
 
 	FORCEINLINE void FStructuredArchiveSlot::operator<< (uint32& Value)
 	{
-		UE::StructuredArchive::Private::GetFormatterImpl(Ar).Serialize(Value);
+		UE::StructuredArchive::Private::GetFormatterImpl(StructuredArchive).Serialize(Value);
 	}
 
 	FORCEINLINE void FStructuredArchiveSlot::operator<< (uint64& Value)
 	{
-		UE::StructuredArchive::Private::GetFormatterImpl(Ar).Serialize(Value);
+		UE::StructuredArchive::Private::GetFormatterImpl(StructuredArchive).Serialize(Value);
 	}
 
 	FORCEINLINE void FStructuredArchiveSlot::operator<< (int8& Value)
 	{
-		UE::StructuredArchive::Private::GetFormatterImpl(Ar).Serialize(Value);
+		UE::StructuredArchive::Private::GetFormatterImpl(StructuredArchive).Serialize(Value);
 	}
 
 	FORCEINLINE void FStructuredArchiveSlot::operator<< (int16& Value)
 	{
-		UE::StructuredArchive::Private::GetFormatterImpl(Ar).Serialize(Value);
+		UE::StructuredArchive::Private::GetFormatterImpl(StructuredArchive).Serialize(Value);
 	}
 
 	FORCEINLINE void FStructuredArchiveSlot::operator<< (int32& Value)
 	{
-		UE::StructuredArchive::Private::GetFormatterImpl(Ar).Serialize(Value);
+		UE::StructuredArchive::Private::GetFormatterImpl(StructuredArchive).Serialize(Value);
 	}
 
 	FORCEINLINE void FStructuredArchiveSlot::operator<< (int64& Value)
 	{
-		UE::StructuredArchive::Private::GetFormatterImpl(Ar).Serialize(Value);
+		UE::StructuredArchive::Private::GetFormatterImpl(StructuredArchive).Serialize(Value);
 	}
 
 	FORCEINLINE void FStructuredArchiveSlot::operator<< (float& Value)
 	{
-		UE::StructuredArchive::Private::GetFormatterImpl(Ar).Serialize(Value);
+		UE::StructuredArchive::Private::GetFormatterImpl(StructuredArchive).Serialize(Value);
 	}
 
 	FORCEINLINE void FStructuredArchiveSlot::operator<< (double& Value)
 	{
-		UE::StructuredArchive::Private::GetFormatterImpl(Ar).Serialize(Value);
+		UE::StructuredArchive::Private::GetFormatterImpl(StructuredArchive).Serialize(Value);
 	}
 
 	FORCEINLINE void FStructuredArchiveSlot::operator<< (bool& Value)
 	{
-		UE::StructuredArchive::Private::GetFormatterImpl(Ar).Serialize(Value);
+		UE::StructuredArchive::Private::GetFormatterImpl(StructuredArchive).Serialize(Value);
 	}
 
 	FORCEINLINE void FStructuredArchiveSlot::operator<< (FString& Value)
 	{
-		UE::StructuredArchive::Private::GetFormatterImpl(Ar).Serialize(Value);
+		UE::StructuredArchive::Private::GetFormatterImpl(StructuredArchive).Serialize(Value);
 	}
 
 	FORCEINLINE void FStructuredArchiveSlot::operator<< (FName& Value)
 	{
-		UE::StructuredArchive::Private::GetFormatterImpl(Ar).Serialize(Value);
+		UE::StructuredArchive::Private::GetFormatterImpl(StructuredArchive).Serialize(Value);
 	}
 
 	FORCEINLINE void FStructuredArchiveSlot::operator<< (UObject*& Value)
 	{
-		UE::StructuredArchive::Private::GetFormatterImpl(Ar).Serialize(Value);
+		UE::StructuredArchive::Private::GetFormatterImpl(StructuredArchive).Serialize(Value);
 	}
 
 	FORCEINLINE void FStructuredArchiveSlot::operator<< (FText& Value)
 	{
-		UE::StructuredArchive::Private::GetFormatterImpl(Ar).Serialize(Value);
+		UE::StructuredArchive::Private::GetFormatterImpl(StructuredArchive).Serialize(Value);
 	}
 
 	FORCEINLINE void FStructuredArchiveSlot::operator<< (FWeakObjectPtr& Value)
 	{
-		UE::StructuredArchive::Private::GetFormatterImpl(Ar).Serialize(Value);
+		UE::StructuredArchive::Private::GetFormatterImpl(StructuredArchive).Serialize(Value);
 	}
 
 	FORCEINLINE void FStructuredArchiveSlot::operator<< (FSoftObjectPath& Value)
 	{
-		UE::StructuredArchive::Private::GetFormatterImpl(Ar).Serialize(Value);
+		UE::StructuredArchive::Private::GetFormatterImpl(StructuredArchive).Serialize(Value);
 	}
 
 	FORCEINLINE void FStructuredArchiveSlot::operator<< (FSoftObjectPtr& Value)
 	{
-		UE::StructuredArchive::Private::GetFormatterImpl(Ar).Serialize(Value);
+		UE::StructuredArchive::Private::GetFormatterImpl(StructuredArchive).Serialize(Value);
 	}
 
 	FORCEINLINE void FStructuredArchiveSlot::operator<< (FLazyObjectPtr& Value)
 	{
-		UE::StructuredArchive::Private::GetFormatterImpl(Ar).Serialize(Value);
+		UE::StructuredArchive::Private::GetFormatterImpl(StructuredArchive).Serialize(Value);
 	}
 
 	FORCEINLINE void FStructuredArchiveSlot::operator<< (FObjectPtr& Value)
 	{
-		UE::StructuredArchive::Private::GetFormatterImpl(Ar).Serialize(Value);
+		UE::StructuredArchive::Private::GetFormatterImpl(StructuredArchive).Serialize(Value);
 	}
 
 	FORCEINLINE void FStructuredArchiveSlot::Serialize(TArray<uint8>& Data)
 	{
-		UE::StructuredArchive::Private::GetFormatterImpl(Ar).Serialize(Data);
+		UE::StructuredArchive::Private::GetFormatterImpl(StructuredArchive).Serialize(Data);
 	}
 
 	FORCEINLINE void FStructuredArchiveSlot::Serialize(void* Data, uint64 DataSize)
 	{
-		UE::StructuredArchive::Private::GetFormatterImpl(Ar).Serialize(Data, DataSize);
+		UE::StructuredArchive::Private::GetFormatterImpl(StructuredArchive).Serialize(Data, DataSize);
 	}
 
 	//////////// FStructuredArchiveRecord ////////////
 
 	FORCEINLINE FStructuredArchiveSlot FStructuredArchiveRecord::EnterField(FArchiveFieldName Name)
 	{
-		return FStructuredArchiveSlot(FStructuredArchiveSlot::EPrivateToken{}, Ar);
+		return FStructuredArchiveSlot(FStructuredArchiveSlot::EPrivateToken{}, StructuredArchive);
 	}
 
 	FORCEINLINE FStructuredArchiveRecord FStructuredArchiveRecord::EnterRecord(FArchiveFieldName Name)
@@ -412,9 +412,9 @@
 
 	FORCEINLINE TOptional<FStructuredArchiveSlot> FStructuredArchiveRecord::TryEnterField(FArchiveFieldName Name, bool bEnterWhenWriting)
 	{
-		if (UE::StructuredArchive::Private::GetFormatterImpl(Ar).TryEnterField(Name, bEnterWhenWriting))
+		if (UE::StructuredArchive::Private::GetFormatterImpl(StructuredArchive).TryEnterField(Name, bEnterWhenWriting))
 		{
-			return TOptional<FStructuredArchiveSlot>(InPlace, FStructuredArchiveSlot(FStructuredArchiveSlot::EPrivateToken{}, Ar));
+			return TOptional<FStructuredArchiveSlot>(InPlace, FStructuredArchiveSlot(FStructuredArchiveSlot::EPrivateToken{}, StructuredArchive));
 		}
 		else
 		{
@@ -426,22 +426,22 @@
 
 	FORCEINLINE FStructuredArchiveSlot FStructuredArchiveArray::EnterElement()
 	{
-		return FStructuredArchiveSlot(FStructuredArchiveSlot::EPrivateToken{}, Ar);
+		return FStructuredArchiveSlot(FStructuredArchiveSlot::EPrivateToken{}, StructuredArchive);
 	}
 
 	//////////// FStructuredArchiveStream ////////////
 
 	FORCEINLINE FStructuredArchiveSlot FStructuredArchiveStream::EnterElement()
 	{
-		return FStructuredArchiveSlot(FStructuredArchiveSlot::EPrivateToken{}, Ar);
+		return FStructuredArchiveSlot(FStructuredArchiveSlot::EPrivateToken{}, StructuredArchive);
 	}
 
 	//////////// FStructuredArchiveMap ////////////
 
 	FORCEINLINE FStructuredArchiveSlot FStructuredArchiveMap::EnterElement(FString& Name)
 	{
-		UE::StructuredArchive::Private::GetFormatterImpl(Ar).EnterMapElement(Name);
-		return FStructuredArchiveSlot(FStructuredArchiveSlot::EPrivateToken{}, Ar);
+		UE::StructuredArchive::Private::GetFormatterImpl(StructuredArchive).EnterMapElement(Name);
+		return FStructuredArchiveSlot(FStructuredArchiveSlot::EPrivateToken{}, StructuredArchive);
 	}
 
 #endif
