Implement FPlatformMisc::GameTemporaryDownloadDir()

diff --git a/Engine/Source/Runtime/Core/Private/GenericPlatform/GenericPlatformMisc.cpp b/Engine/Source/Runtime/Core/Private/GenericPlatform/GenericPlatformMisc.cpp
--- a/Engine/Source/Runtime/Core/Private/GenericPlatform/GenericPlatformMisc.cpp
+++ b/Engine/Source/Runtime/Core/Private/GenericPlatform/GenericPlatformMisc.cpp
@@ -1377,6 +1377,11 @@
 	return *GamePersistentDownloadDir;
 }
 
+const TCHAR* FGenericPlatformMisc::GameTemporaryDownloadDir()
+{
+	return nullptr;
+}
+
 const TCHAR* FGenericPlatformMisc::GeneratedConfigDir()
 {
 	static FString Dir = FPaths::ProjectSavedDir() / TEXT("Config/");
diff --git a/Engine/Source/Runtime/Core/Public/GenericPlatform/GenericPlatformMisc.h b/Engine/Source/Runtime/Core/Public/GenericPlatform/GenericPlatformMisc.h
--- a/Engine/Source/Runtime/Core/Public/GenericPlatform/GenericPlatformMisc.h
+++ b/Engine/Source/Runtime/Core/Public/GenericPlatform/GenericPlatformMisc.h
@@ -1157,6 +1157,11 @@
 	 */
 	static const TCHAR* GamePersistentDownloadDir();
 
+	/**
+	 *	Return the temporary directory
+	 */
+	static const TCHAR* GameTemporaryDownloadDir();
+
 	static const TCHAR* GeneratedConfigDir();
 
 	static const TCHAR* GetUBTPlatform();
