GitHub #8819 : Resolved TODO for moving TMakeSigned and TMakeUnsigned out of RingBuffer into to its own file

diff --git a/Engine/Source/Runtime/Core/Public/Containers/RingBuffer.h b/Engine/Source/Runtime/Core/Public/Containers/RingBuffer.h
--- a/Engine/Source/Runtime/Core/Public/Containers/RingBuffer.h
+++ b/Engine/Source/Runtime/Core/Public/Containers/RingBuffer.h
@@ -7,43 +7,11 @@
 #include "Containers/ArrayView.h"
 #include "Templates/IsPODType.h"
 #include "Templates/IsTriviallyDestructible.h"
+#include "Templates/MakeSigned.h"
+#include "Templates/MakeUnsigned.h"
 #include "Templates/MemoryOps.h"
 #include "Templates/UnrealTemplate.h"
 
-// PUBLIC_RINGBUFFER_TODO: Move TMakeSigned and TMakeUnsigned into MakeSigned.h
-template <typename T>
-struct TMakeSigned
-{
-	static_assert(sizeof(T) == 0, "Unsupported type in TMakeSigned<T>.");
-};
-template <typename T> struct TMakeSigned<const          T> { using Type = const          typename TMakeSigned<T>::Type; };
-template <typename T> struct TMakeSigned<      volatile T> { using Type =       volatile typename TMakeSigned<T>::Type; };
-template <typename T> struct TMakeSigned<const volatile T> { using Type = const volatile typename TMakeSigned<T>::Type; };
-template <> struct TMakeSigned<int8>   { using Type = int8; };
-template <> struct TMakeSigned<uint8>  { using Type = int8; };
-template <> struct TMakeSigned<int16>  { using Type = int16; };
-template <> struct TMakeSigned<uint16> { using Type = int16; };
-template <> struct TMakeSigned<int32>  { using Type = int32; };
-template <> struct TMakeSigned<uint32> { using Type = int32; };
-template <> struct TMakeSigned<int64>  { using Type = int64; };
-template <> struct TMakeSigned<uint64> { using Type = int64; };
-template <typename T>
-struct TMakeUnsigned
-{
-	static_assert(sizeof(T) == 0, "Unsupported type in TMakeUnsigned<T>.");
-};
-template <typename T> struct TMakeUnsigned<const          T> { using Type = const          typename TMakeUnsigned<T>::Type; };
-template <typename T> struct TMakeUnsigned<      volatile T> { using Type =       volatile typename TMakeUnsigned<T>::Type; };
-template <typename T> struct TMakeUnsigned<const volatile T> { using Type = const volatile typename TMakeUnsigned<T>::Type; };
-template <> struct TMakeUnsigned<int8>   { using Type = uint8; };
-template <> struct TMakeUnsigned<uint8>  { using Type = uint8; };
-template <> struct TMakeUnsigned<int16>  { using Type = uint16; };
-template <> struct TMakeUnsigned<uint16> { using Type = uint16; };
-template <> struct TMakeUnsigned<int32>  { using Type = uint32; };
-template <> struct TMakeUnsigned<uint32> { using Type = uint32; };
-template <> struct TMakeUnsigned<int64>  { using Type = uint64; };
-template <> struct TMakeUnsigned<uint64> { using Type = uint64; };
-
 template< typename ContainerType, typename ElementType, typename SizeType>
 class TRingBufferIterator
 {
diff --git a/Engine/Source/Runtime/Core/Public/Templates/MakeSigned.h b/Engine/Source/Runtime/Core/Public/Templates/MakeSigned.h
new file mode 100644
--- /dev/null
+++ b/Engine/Source/Runtime/Core/Public/Templates/MakeSigned.h
@@ -0,0 +1,27 @@
+// Copyright Epic Games, Inc. All Rights Reserved.
+
+#pragma once
+
+#include "CoreTypes.h"
+
+/**
+ * Traits class which gets the signed version of an integer type.
+ */
+template <typename T>
+struct TMakeSigned
+{
+	static_assert(sizeof(T) == 0, "Unsupported type in TMakeSigned<T>.");
+};
+
+template <typename T> struct TMakeSigned<const          T> { using Type = const          typename TMakeSigned<T>::Type; };
+template <typename T> struct TMakeSigned<      volatile T> { using Type =       volatile typename TMakeSigned<T>::Type; };
+template <typename T> struct TMakeSigned<const volatile T> { using Type = const volatile typename TMakeSigned<T>::Type; };
+
+template <> struct TMakeSigned<int8  > { using Type = int8;  };
+template <> struct TMakeSigned<uint8 > { using Type = int8;  };
+template <> struct TMakeSigned<int16 > { using Type = int16; };
+template <> struct TMakeSigned<uint16> { using Type = int16; };
+template <> struct TMakeSigned<int32 > { using Type = int32; };
+template <> struct TMakeSigned<uint32> { using Type = int32; };
+template <> struct TMakeSigned<int64 > { using Type = int64; };
+template <> struct TMakeSigned<uint64> { using Type = int64; };
diff --git a/Engine/Source/Runtime/Core/Public/Templates/MakeUnsigned.h b/Engine/Source/Runtime/Core/Public/Templates/MakeUnsigned.h
new file mode 100644
--- /dev/null
+++ b/Engine/Source/Runtime/Core/Public/Templates/MakeUnsigned.h
@@ -0,0 +1,27 @@
+// Copyright Epic Games, Inc. All Rights Reserved.
+
+#pragma once
+
+#include "CoreTypes.h"
+
+/**
+ * Traits class which gets the unsigned version of an integer type.
+ */
+template <typename T>
+struct TMakeUnsigned
+{
+	static_assert(sizeof(T) == 0, "Unsupported type in TMakeUnsigned<T>.");
+};
+
+template <typename T> struct TMakeUnsigned<const          T> { using Type = const          typename TMakeUnsigned<T>::Type; };
+template <typename T> struct TMakeUnsigned<      volatile T> { using Type =       volatile typename TMakeUnsigned<T>::Type; };
+template <typename T> struct TMakeUnsigned<const volatile T> { using Type = const volatile typename TMakeUnsigned<T>::Type; };
+
+template <> struct TMakeUnsigned<int8  > { using Type = uint8;  };
+template <> struct TMakeUnsigned<uint8 > { using Type = uint8;  };
+template <> struct TMakeUnsigned<int16 > { using Type = uint16; };
+template <> struct TMakeUnsigned<uint16> { using Type = uint16; };
+template <> struct TMakeUnsigned<int32 > { using Type = uint32; };
+template <> struct TMakeUnsigned<uint32> { using Type = uint32; };
+template <> struct TMakeUnsigned<int64 > { using Type = uint64; };
+template <> struct TMakeUnsigned<uint64> { using Type = uint64; };
