Add support for TVariant visitors that use rvalue-reference function parameters

Minor cleanup: remove undef for no-longer-used TVARIANT_STORAGE_USE_RECURSIVE_TEMPLATE

diff --git a/Engine/Source/Runtime/Core/Public/Misc/TVariantMeta.h b/Engine/Source/Runtime/Core/Public/Misc/TVariantMeta.h
--- a/Engine/Source/Runtime/Core/Public/Misc/TVariantMeta.h
+++ b/Engine/Source/Runtime/Core/Public/Misc/TVariantMeta.h
@@ -95,7 +95,7 @@
 
 		/** Interpret the underlying data as the type in the variant parameter pack at the compile-time index. This function is used to implement Visit and should not be used directly */
 		template <SIZE_T N>
-		auto& GetValueAsIndexedType()
+		auto& GetValueAsIndexedType() &
 		{
 			using ReturnType = typename TNthTypeFromParameterPack<N, Ts...>::Type;
 			return *reinterpret_cast<ReturnType*>(&Storage);
@@ -103,8 +103,16 @@
 
 		/** Interpret the underlying data as the type in the variant parameter pack at the compile-time index. This function is used to implement Visit and should not be used directly */
 		template <SIZE_T N>
-		const auto& GetValueAsIndexedType() const
+		auto&& GetValueAsIndexedType() &&
 		{
+			using ReturnType = typename TNthTypeFromParameterPack<N, Ts...>::Type;
+			return (ReturnType&&)GetValueAsIndexedType<N>();
+		}
+
+		/** Interpret the underlying data as the type in the variant parameter pack at the compile-time index. This function is used to implement Visit and should not be used directly */
+		template <SIZE_T N>
+		const auto& GetValueAsIndexedType() const&
+		{
 			// Temporarily remove the const qualifier so we can implement GetValueAsIndexedType in one location.
 			return const_cast<TVariantStorage*>(this)->template GetValueAsIndexedType<N>();
 		}
@@ -276,6 +284,12 @@
 	}
 
 	template <typename... Ts>
+	FORCEINLINE TVariantStorage<Ts...>&& CastToStorage(TVariant<Ts...>&& Variant)
+	{
+		return (TVariantStorage<Ts...>&&)(*(TVariantStorage<Ts...>*)(&Variant));
+	}
+
+	template <typename... Ts>
 	FORCEINLINE const TVariantStorage<Ts...>& CastToStorage(const TVariant<Ts...>& Variant)
 	{
 		return *(const TVariantStorage<Ts...>*)(&Variant);
@@ -286,7 +300,7 @@
 	inline decltype(auto) VisitApplyEncoded(Func&& Callable, Variants&&... Args)
 	{
 		constexpr SIZE_T VariantSizes[] = { TVariantSize<Variants>::Value... };
-		return Callable(CastToStorage(Args).template GetValueAsIndexedType<DecodeIndex(EncodedIndex, VariantIndices, VariantSizes)>()...);
+		return Callable(CastToStorage(Forward<Variants>(Args)).template GetValueAsIndexedType<DecodeIndex(EncodedIndex, VariantIndices, VariantSizes)>()...);
 	}
 
 	/**
@@ -317,5 +331,3 @@
 } // namespace Private
 } // namespace Core
 } // namespace UE
-
-#undef TVARIANT_STORAGE_USE_RECURSIVE_TEMPLATE
