Added CacheAllocated statistics for trace system (all memory allocated in cache buffers).
Fixed "Important Events cache" stats displayed by [REDACTED_PERSON] console command.

diff --git a/Engine/Source/Runtime/Core/Private/ProfilingDebugging/TraceAuxiliary.cpp b/Engine/Source/Runtime/Core/Private/ProfilingDebugging/TraceAuxiliary.cpp
--- a/Engine/Source/Runtime/Core/Private/ProfilingDebugging/TraceAuxiliary.cpp
+++ b/Engine/Source/Runtime/Core/Private/ProfilingDebugging/TraceAuxiliary.cpp
@@ -58,6 +58,7 @@
 ////////////////////////////////////////////////////////////////////////////////
 DECLARE_STATS_GROUP(TEXT("TraceLog"), STATGROUP_Trace, STATCAT_Advanced);
 DECLARE_MEMORY_STAT(TEXT("Memory used"), STAT_TraceMemoryUsed, STATGROUP_Trace);
+DECLARE_MEMORY_STAT(TEXT("Important event cache"), STAT_TraceCacheAllocated, STATGROUP_Trace);
 DECLARE_MEMORY_STAT(TEXT("Important event cache used"), STAT_TraceCacheUsed, STATGROUP_Trace);
 DECLARE_MEMORY_STAT(TEXT("Important event cache waste"), STAT_TraceCacheWaste, STATGROUP_Trace);
 DECLARE_MEMORY_STAT(TEXT("Sent"), STAT_TraceSent, STATGROUP_Trace);
@@ -587,6 +588,7 @@
 			UE::Trace::FStatistics Stats;
 			UE::Trace::GetStatistics(Stats);
 			SET_MEMORY_STAT(STAT_TraceMemoryUsed, Stats.MemoryUsed);
+			SET_MEMORY_STAT(STAT_TraceCacheAllocated, Stats.CacheAllocated);
 			SET_MEMORY_STAT(STAT_TraceCacheUsed, Stats.CacheUsed);
 			SET_MEMORY_STAT(STAT_TraceCacheWaste, Stats.CacheWaste);
 			SET_MEMORY_STAT(STAT_TraceSent, Stats.BytesSent);
@@ -735,11 +737,11 @@
 	constexpr double MiB = 1.0 / (1024.0 * 1024.0);
 	UE_LOG(LogConsoleResponse, Display, TEXT("- Memory used: %.02f MiB"),
 		double(Stats.MemoryUsed) * MiB);
-	// Stats.CacheWaste is actually the total size of allocated buffers.
-	UE_LOG(LogConsoleResponse, Display, TEXT("- Important Events cache: %.02f MiB (%.02f MiB used + %0.02f MiB waste)"),
-		double(Stats.CacheWaste) * MiB,
+	UE_LOG(LogConsoleResponse, Display, TEXT("- Important Events cache: %.02f MiB (%.02f MiB used + %0.02f MiB unused | %0.02f MiB waste)"),
+		double(Stats.CacheAllocated) * MiB,
 		double(Stats.CacheUsed) * MiB,
-		double(Stats.CacheWaste - Stats.CacheUsed) * MiB);
+		double(Stats.CacheAllocated - Stats.CacheUsed) * MiB,
+		double(Stats.CacheWaste) * MiB);
 	UE_LOG(LogConsoleResponse, Display, TEXT("- Sent: %.02f MiB"),
 		double(Stats.BytesSent) * MiB);
 
diff --git a/Engine/Source/Runtime/TraceLog/Private/Trace/Important/Cache.cpp b/Engine/Source/Runtime/TraceLog/Private/Trace/Important/Cache.cpp
--- a/Engine/Source/Runtime/TraceLog/Private/Trace/Important/Cache.cpp
+++ b/Engine/Source/Runtime/TraceLog/Private/Trace/Important/Cache.cpp
@@ -50,6 +50,10 @@
 ////////////////////////////////////////////////////////////////////////////////
 static FCacheBuffer* Writer_CacheCreateBuffer(uint32 Size)
 {
+#if TRACE_PRIVATE_STATISTICS
+	GTraceStatistics.CacheAllocated += Size;
+#endif
+
 	void* Block = Writer_MemoryAllocate(sizeof(FCacheBuffer) + Size, alignof(FCacheBuffer));
 	auto* Buffer = (FCacheBuffer*)Block;
 	Buffer->Size = Size;
diff --git a/Engine/Source/Runtime/TraceLog/Public/Trace/Trace.h b/Engine/Source/Runtime/TraceLog/Public/Trace/Trace.h
--- a/Engine/Source/Runtime/TraceLog/Public/Trace/Trace.h
+++ b/Engine/Source/Runtime/TraceLog/Public/Trace/Trace.h
@@ -93,8 +93,9 @@
 	uint64		BytesSent;
 	uint64		BytesTraced;
 	uint64		MemoryUsed;
-	uint32		CacheUsed;		// Important-marked events are
-	uint32		CacheWaste;		// stored in the cache.
+	uint32		CacheAllocated;	// Total memory allocated in cache buffers
+	uint32		CacheUsed;		// Used cache memory; Important-marked events are stored in the cache.
+	uint32		CacheWaste;		// Unused memory from retired cache buffers
 };
 
 UE_TRACE_API void	SetMemoryHooks(AllocFunc Alloc, FreeFunc Free) UE_TRACE_IMPL();
