Improved comments around the use of MakeShared.
New ConstCastWeakPtr and StaticCastWeakPtr functions to match the SharedPtr/Ref equivalents.
Lots of sharedptr-related functions made [[nodiscard]].

diff --git a/Engine/Source/Runtime/Core/Public/Templates/SharedPointer.h b/Engine/Source/Runtime/Core/Public/Templates/SharedPointer.h
--- a/Engine/Source/Runtime/Core/Public/Templates/SharedPointer.h
+++ b/Engine/Source/Runtime/Core/Public/Templates/SharedPointer.h
@@ -43,11 +43,14 @@
  *	Additionally, the following helper classes and functions are defined:
  *
  *		MakeShareable() - Used to initialize shared pointers from C++ pointers (enables implicit conversion)
+ *		MakeShared<T>(...) - Used to construct a T alongside its controller, saving an allocation.
  *		TSharedFromThis - You can derive your own class from this to acquire a TSharedRef from "this"
  *		StaticCastSharedRef() - Static cast utility function, typically used to downcast to a derived type. 
  *		ConstCastSharedRef() - Converts a 'const' reference to 'mutable' smart reference
  *		StaticCastSharedPtr() - Dynamic cast utility function, typically used to downcast to a derived type. 
  *		ConstCastSharedPtr() - Converts a 'const' smart pointer to 'mutable' smart pointer
+ *		StaticCastWeakPtr() - Dynamic cast utility function, typically used to downcast to a derived type.
+ *		ConstCastWeakPtr() - Converts a 'const' smart pointer to 'mutable' smart pointer
  *
  *
  *	Examples:
@@ -58,6 +61,7 @@
  *		- Use TSharedRef instead of TSharedPtr whenever possible -- it can never be nullptr!
  *		- You can call TSharedPtr::Reset() to release a reference to your object (and potentially deallocate) 
  *		- Use the MakeShareable() helper function to implicitly convert to TSharedRefs or TSharedPtrs
+ *		- Prefer MakeShared<T>(...) to MakeShareable(new T(...))
  *		- You can never reset a TSharedRef or assign it to nullptr, but you can assign it a new object
  *		- Shared pointers assume ownership of objects -- no need to call delete yourself!
  *		- Usually you should "operator new" when passing a C++ pointer to a new shared pointer
@@ -68,16 +72,16 @@
  *		- You can create a typedef to TSharedRef< MyClass > to make it easier to type
  *		- For best performance, minimize calls to TWeakPtr::Pin (or conversions to TSharedRef/TSharedPtr)
  *		- Your class can return itself as a shared reference if you derive from TSharedFromThis
- *		- To downcast a pointer to a derived object class, to the StaticCastSharedPtr function
+ *		- To downcast a pointer to a derived object class, use the StaticCastSharedRef, StaticCastSharedRef or StaticCastWeakPtr functions
  *		- 'const' objects are fully supported with shared pointers!
- *		- You can make a 'const' shared pointer mutable using the ConstCastSharedPtr function
+ *		- You can make a 'const' pointer mutable using the ConstCastSharedRef, ConstCastSharedPtr or ConstCastWeakPtr functions
  *		
  *
  *	Limitations:
  *
  *		- Shared pointers are not compatible with Unreal objects (UObject classes)!
  *		- Currently only types with that have regular destructors (no custom deleters)
- *		- Dynamically-allocated arrays are not supported yet (e.g. MakeSharable( new int32[20] ))
+ *		- Dynamically-allocated arrays are not supported yet (e.g. MakeShareable( new int32[20] ))
  *		- Implicit conversion of TSharedPtr/TSharedRef to bool is not supported yet
  *
  *
@@ -119,7 +123,7 @@
  * @param  InSharedRef  The shared reference to cast
  */
 template< class CastToType, class CastFromType, ESPMode Mode >
-FORCEINLINE TSharedRef< CastToType, Mode > StaticCastSharedRef( TSharedRef< CastFromType, Mode > const& InSharedRef )
+[[nodiscard]] FORCEINLINE TSharedRef< CastToType, Mode > StaticCastSharedRef( TSharedRef< CastFromType, Mode > const& InSharedRef )
 {
 	return TSharedRef< CastToType, Mode >( InSharedRef, SharedPointerInternals::FStaticCastTag() );
 }
@@ -465,7 +469,7 @@ public:
 	 *
 	 * @return  The object owned by this shared reference
 	 */
-	FORCEINLINE ObjectType& Get() const
+	[[nodiscard]] FORCEINLINE ObjectType& Get() const
 	{
 		// Should never be nullptr as TSharedRef is never nullable
 		checkSlow( IsValid() );
@@ -477,7 +481,7 @@ public:
 	 *
 	 * @return  Reference to the object
 	 */
-	FORCEINLINE ObjectType& operator*() const
+	[[nodiscard]] FORCEINLINE ObjectType& operator*() const
 	{
 		// Should never be nullptr as TSharedRef is never nullable
 		checkSlow( IsValid() );
@@ -489,7 +493,7 @@ public:
 	 *
 	 * @return  Returns a pointer to the object referenced by this shared reference
 	 */
-	FORCEINLINE ObjectType* operator->() const
+	[[nodiscard]] FORCEINLINE ObjectType* operator->() const
 	{
 		// Should never be nullptr as TSharedRef is never nullable
 		checkSlow( IsValid() );
@@ -502,7 +506,7 @@ public:
 	 *
 	 * @return  Number of shared references to the object (including this reference.)
 	 */
-	FORCEINLINE const int32 GetSharedReferenceCount() const
+	[[nodiscard]] FORCEINLINE int32 GetSharedReferenceCount() const
 	{
 		return SharedReferenceCount.GetSharedReferenceCount();
 	}
@@ -514,7 +518,7 @@ public:
 	 *
 	 * @return  True if there is only one shared reference to the object, and this is it!
 	 */
-	FORCEINLINE const bool IsUnique() const
+	[[nodiscard]] FORCEINLINE bool IsUnique() const
 	{
 		return SharedReferenceCount.IsUnique();
 	}
@@ -572,7 +576,7 @@ private:
 	 *
 	 * @return  True if the shared reference is valid and can be dereferenced
 	 */
-	FORCEINLINE const bool IsValid() const
+	[[nodiscard]] FORCEINLINE bool IsValid() const
 	{
 		return Object != nullptr;
 	}
@@ -584,7 +588,7 @@ private:
 	 *
 	 * @return  Hash code value
 	 */
-	friend uint32 GetTypeHash( const TSharedRef< ObjectType, Mode >& InSharedRef )
+	[[nodiscard]] friend uint32 GetTypeHash( const TSharedRef< ObjectType, Mode >& InSharedRef )
 	{
 		return ::PointerHash( InSharedRef.Object );
 	}
@@ -1001,7 +1005,7 @@ public:
 	 * @return  Reference to the object
 	 */
 	// NOTE: The following is an Unreal extension to standard shared_ptr behavior
-	FORCEINLINE TSharedRef< ObjectType, Mode > ToSharedRef() const&
+	[[nodiscard]] FORCEINLINE TSharedRef< ObjectType, Mode > ToSharedRef() const&
 	{
 		// If this assert goes off, it means a shared reference was created from a shared pointer that was nullptr.
 		// Shared references are never allowed to be null.  Consider using TSharedPtr instead.
@@ -1015,7 +1019,7 @@ public:
 	 * @return  Reference to the object
 	 */
 	 // NOTE: The following is an Unreal extension to standard shared_ptr behavior
-	FORCEINLINE TSharedRef< ObjectType, Mode > ToSharedRef() &&
+	[[nodiscard]] FORCEINLINE TSharedRef< ObjectType, Mode > ToSharedRef() &&
 	{
 		// If this assert goes off, it means a shared reference was created from a shared pointer that was nullptr.
 		// Shared references are never allowed to be null.  Consider using TSharedPtr instead.
@@ -1028,7 +1032,7 @@ public:
 	 *
 	 * @return  The object owned by this shared pointer, or nullptr
 	 */
-	FORCEINLINE ObjectType* Get() const
+	[[nodiscard]] FORCEINLINE ObjectType* Get() const
 	{
 		return Object;
 	}
@@ -1038,7 +1042,7 @@ public:
 	 *
 	 * @return  True if the shared pointer is valid and can be dereferenced
 	 */
-	FORCEINLINE explicit operator bool() const
+	[[nodiscard]] FORCEINLINE explicit operator bool() const
 	{
 		return Object != nullptr;
 	}
@@ -1048,7 +1052,7 @@ public:
 	 *
 	 * @return  True if the shared pointer is valid and can be dereferenced
 	 */
-	FORCEINLINE const bool IsValid() const
+	[[nodiscard]] FORCEINLINE const bool IsValid() const
 	{
 		return Object != nullptr;
 	}
@@ -1058,7 +1062,7 @@ public:
 	 *
 	 * @return  Reference to the object
 	 */
-	FORCEINLINE typename FMakeReferenceTo<ObjectType>::Type operator*() const
+	[[nodiscard]] FORCEINLINE typename FMakeReferenceTo<ObjectType>::Type operator*() const
 	{
 		check( IsValid() );
 		return *Object;
@@ -1069,7 +1073,7 @@ public:
 	 *
 	 * @return  Returns a pointer to the object referenced by this shared pointer
 	 */
-	FORCEINLINE ObjectType* operator->() const
+	[[nodiscard]] FORCEINLINE ObjectType* operator->() const
 	{
 		check( IsValid() );
 		return Object;
@@ -1090,7 +1094,7 @@ public:
 	 *
 	 * @return  Number of shared references to the object (including this reference.)
 	 */
-	FORCEINLINE const int32 GetSharedReferenceCount() const
+	[[nodiscard]] FORCEINLINE int32 GetSharedReferenceCount() const
 	{
 		return SharedReferenceCount.GetSharedReferenceCount();
 	}
@@ -1102,7 +1106,7 @@ public:
 	 *
 	 * @return  True if there is only one shared reference to the object, and this is it!
 	 */
-	FORCEINLINE const bool IsUnique() const
+	[[nodiscard]] FORCEINLINE bool IsUnique() const
 	{
 		return SharedReferenceCount.IsUnique();
 	}
@@ -1167,7 +1171,7 @@ private:
 	 *
 	 * @return  Hash code value
 	 */
-	friend uint32 GetTypeHash( const TSharedPtr< ObjectType, Mode >& InSharedPtr )
+	[[nodiscard]] friend uint32 GetTypeHash( const TSharedPtr< ObjectType, Mode >& InSharedPtr )
 	{
 		return ::PointerHash( InSharedPtr.Object );
 	}
@@ -1190,9 +1194,6 @@ private:
 	SharedPointerInternals::FSharedReferencer< Mode > SharedReferenceCount;
 };
 
-//template< class ObjectType, ESPMode Mode >
-//class TSharedPtr
-
 namespace Freeze
 {
 	template<class ObjectType, ESPMode Mode>
@@ -1258,6 +1259,36 @@ public:
 	{
 	}
 
+	/**
+	 * Special constructor used internally to statically cast one weak pointer type to another.  You
+	 * should never call this constructor directly.  Instead, use the StaticCastWeakPtr() function.
+	 * This constructor creates a weak pointer as a weak reference to an existing weak pointer after
+	 * statically casting that pointer's object.  This constructor is needed for static casts.
+	 *
+	 * @param  InWeakPtr  The weak pointer whose object we should create an additional reference to
+	 */
+	template <typename OtherType>
+	FORCEINLINE TWeakPtr(TWeakPtr< OtherType, Mode > const& InWeakPtr, SharedPointerInternals::FStaticCastTag)
+		: Object(static_cast<ObjectType*>(InWeakPtr.Object))
+		, WeakReferenceCount(InWeakPtr.WeakReferenceCount)
+	{
+	}
+
+	/**
+	 * Special constructor used internally to cast a 'const' weak pointer a 'mutable' pointer.  You
+	 * should never call this constructor directly.  Instead, use the ConstCastWeakPtr() function.
+	 * This constructor creates a weak pointer as a weak reference to an existing weak pointer after
+	 * const casting that pointer's object.  This constructor is needed for const casts.
+	 *
+	 * @param  InWeakPtr  The weak pointer whose object we should create an additional reference to
+	 */
+	template <typename OtherType>
+	FORCEINLINE TWeakPtr(TWeakPtr< OtherType, Mode > const& InWeakPtr, SharedPointerInternals::FConstCastTag)
+		: Object(const_cast<ObjectType*>(InWeakPtr.Object))
+		, WeakReferenceCount(InWeakPtr.WeakReferenceCount)
+	{
+	}
+
 	/**
 	 * Constructs a weak pointer from a weak pointer of another type.
 	 * This constructor is intended to allow derived-to-base conversions.
@@ -1401,7 +1432,7 @@ public:
 	 *
 	 * @return  Shared pointer for this object (will only be valid if still referenced!)
 	 */
-	FORCEINLINE TSharedPtr< ObjectType, Mode > Pin() const&
+	[[nodiscard]] FORCEINLINE TSharedPtr< ObjectType, Mode > Pin() const&
 	{
 		return TSharedPtr< ObjectType, Mode >( *this );
 	}
@@ -1414,7 +1445,7 @@ public:
 	 *
 	 * @return  Shared pointer for this object (will only be valid if still referenced!)
 	 */
-	FORCEINLINE TSharedPtr< ObjectType, Mode > Pin() &&
+	[[nodiscard]] FORCEINLINE TSharedPtr< ObjectType, Mode > Pin() &&
 	{
 		return TSharedPtr< ObjectType, Mode >( MoveTemp( *this ) );
 	}
@@ -1424,7 +1455,7 @@ public:
 	 *
 	 * @return  True if the weak pointer is valid and a pin operator would have succeeded
 	 */
-	FORCEINLINE const bool IsValid() const
+	[[nodiscard]] FORCEINLINE bool IsValid() const
 	{
 		return Object != nullptr && WeakReferenceCount.IsValid();
 	}
@@ -1441,7 +1472,7 @@ public:
 	/**
 	 * Returns true if the object this weak pointer points to is the same as the specified object pointer.
 	 */
-	FORCEINLINE bool HasSameObject( const void* InOtherPtr ) const
+	[[nodiscard]] FORCEINLINE bool HasSameObject( const void* InOtherPtr ) const
 	{
 		return Pin().Get() == InOtherPtr;
 	}
@@ -1455,7 +1486,7 @@ private:
 	 *
 	 * @return  Hash code value
 	 */
-	friend uint32 GetTypeHash( const TWeakPtr< ObjectType, Mode >& InWeakPtr )
+	[[nodiscard]] friend uint32 GetTypeHash( const TWeakPtr< ObjectType, Mode >& InWeakPtr )
 	{
 		return ::PointerHash( InWeakPtr.Object );
 	}
@@ -1498,7 +1529,7 @@ public:
 	 *
 	 * @return	Returns this object as a shared pointer
 	 */
-	TSharedRef< ObjectType, Mode > AsShared()
+	[[nodiscard]] TSharedRef< ObjectType, Mode > AsShared()
 	{
 		TSharedPtr< ObjectType, Mode > SharedThis( WeakThis.Pin() );
 
@@ -1525,7 +1556,7 @@ public:
 	 *
 	 * @return	Returns this object as a shared pointer (const)
 	 */
-	TSharedRef< ObjectType const, Mode > AsShared() const
+	[[nodiscard]] TSharedRef< ObjectType const, Mode > AsShared() const
 	{
 		TSharedPtr< ObjectType const, Mode > SharedThis( WeakThis );
 
@@ -1552,7 +1583,7 @@ public:
 	 *
 	 * @return	Returns this object as a shared pointer
 	 */
-	TWeakPtr< ObjectType, Mode > AsWeak()
+	[[nodiscard]] TWeakPtr< ObjectType, Mode > AsWeak()
 	{
 		TWeakPtr< ObjectType, Mode > Result = WeakThis;
 
@@ -1570,7 +1601,7 @@ public:
 		// Now that we've verified the pointer is valid, we'll return it!
 		return Result;
 	}
-	TWeakPtr< ObjectType const, Mode > AsWeak() const
+	[[nodiscard]] TWeakPtr< ObjectType const, Mode > AsWeak() const
 	{
 		TWeakPtr< ObjectType const, Mode > Result = WeakThis;
 
@@ -1600,7 +1631,7 @@ protected:
 	 * @return	Returns this object as a shared pointer
 	 */
 	template< class OtherType >
-	FORCEINLINE static TSharedRef< OtherType, Mode > SharedThis( OtherType* ThisPtr )
+	[[nodiscard]] FORCEINLINE static TSharedRef< OtherType, Mode > SharedThis( OtherType* ThisPtr )
 	{
 		return StaticCastSharedRef< OtherType >( ThisPtr->AsShared() );
 	}
@@ -1614,7 +1645,7 @@ protected:
 	 * @return	Returns this object as a shared pointer (const)
 	 */
 	template< class OtherType >
-	FORCEINLINE static TSharedRef< OtherType const, Mode > SharedThis( const OtherType* ThisPtr )
+	[[nodiscard]] FORCEINLINE static TSharedRef< OtherType const, Mode > SharedThis( const OtherType* ThisPtr )
 	{
 		return StaticCastSharedRef< OtherType const >( ThisPtr->AsShared() );
 	}
@@ -1656,7 +1687,7 @@ public:		// @todo: Ideally this would be private, but template sharing problems
 	 *  - Is currently within its constructor (so the shared instance isn't yet available).
 	 *  - Is currently within its destructor (so the shared instance is no longer available).
 	 */
-	FORCEINLINE bool DoesSharedInstanceExist() const
+	[[nodiscard]] FORCEINLINE bool DoesSharedInstanceExist() const
 	{
 		return WeakThis.IsValid();
 	}
@@ -1692,7 +1723,7 @@ private:
  * @return  True if the two shared references are equal
  */
 template< class ObjectTypeA, class ObjectTypeB, ESPMode Mode >
-FORCEINLINE bool operator==( TSharedRef< ObjectTypeA, Mode > const& InSharedRefA, TSharedRef< ObjectTypeB, Mode > const& InSharedRefB )
+[[nodiscard]] FORCEINLINE bool operator==( TSharedRef< ObjectTypeA, Mode > const& InSharedRefA, TSharedRef< ObjectTypeB, Mode > const& InSharedRefB )
 {
 	return &( InSharedRefA.Get() ) == &( InSharedRefB.Get() );
 }
@@ -1704,7 +1735,7 @@ FORCEINLINE bool operator==( TSharedRef< ObjectTypeA, Mode > const& InSharedRefA
  * @return  True if the two shared references are not equal
  */
 template< class ObjectTypeA, class ObjectTypeB, ESPMode Mode >
-FORCEINLINE bool operator!=( TSharedRef< ObjectTypeA, Mode > const& InSharedRefA, TSharedRef< ObjectTypeB, Mode > const& InSharedRefB )
+[[nodiscard]] FORCEINLINE bool operator!=( TSharedRef< ObjectTypeA, Mode > const& InSharedRefA, TSharedRef< ObjectTypeB, Mode > const& InSharedRefB )
 {
 	return &( InSharedRefA.Get() ) != &( InSharedRefB.Get() );
 }
@@ -1716,7 +1747,7 @@ FORCEINLINE bool operator!=( TSharedRef< ObjectTypeA, Mode > const& InSharedRefA
  * @return  True if the two shared pointers are equal
  */
 template< class ObjectTypeA, class ObjectTypeB, ESPMode Mode >
-FORCEINLINE bool operator==( TSharedPtr< ObjectTypeA, Mode > const& InSharedPtrA, TSharedPtr< ObjectTypeB, Mode > const& InSharedPtrB )
+[[nodiscard]] FORCEINLINE bool operator==( TSharedPtr< ObjectTypeA, Mode > const& InSharedPtrA, TSharedPtr< ObjectTypeB, Mode > const& InSharedPtrB )
 {
 	return InSharedPtrA.Get() == InSharedPtrB.Get();
 }
@@ -1728,7 +1759,7 @@ FORCEINLINE bool operator==( TSharedPtr< ObjectTypeA, Mode > const& InSharedPtrA
  * @return  True if the shared pointer is null
  */
 template< class ObjectTypeA, ESPMode Mode >
-FORCEINLINE bool operator==( TSharedPtr< ObjectTypeA, Mode > const& InSharedPtrA, TYPE_OF_NULLPTR )
+[[nodiscard]] FORCEINLINE bool operator==( TSharedPtr< ObjectTypeA, Mode > const& InSharedPtrA, TYPE_OF_NULLPTR )
 {
 	return !InSharedPtrA.IsValid();
 }
@@ -1740,7 +1771,7 @@ FORCEINLINE bool operator==( TSharedPtr< ObjectTypeA, Mode > const& InSharedPtrA
  * @return  True if the shared pointer is null
  */
 template< class ObjectTypeB, ESPMode Mode >
-FORCEINLINE bool operator==( TYPE_OF_NULLPTR, TSharedPtr< ObjectTypeB, Mode > const& InSharedPtrB )
+[[nodiscard]] FORCEINLINE bool operator==( TYPE_OF_NULLPTR, TSharedPtr< ObjectTypeB, Mode > const& InSharedPtrB )
 {
 	return !InSharedPtrB.IsValid();
 }
@@ -1752,7 +1783,7 @@ FORCEINLINE bool operator==( TYPE_OF_NULLPTR, TSharedPtr< ObjectTypeB, Mode > co
  * @return  True if the two shared pointers are not equal
  */
 template< class ObjectTypeA, class ObjectTypeB, ESPMode Mode >
-FORCEINLINE bool operator!=( TSharedPtr< ObjectTypeA, Mode > const& InSharedPtrA, TSharedPtr< ObjectTypeB, Mode > const& InSharedPtrB )
+[[nodiscard]] FORCEINLINE bool operator!=( TSharedPtr< ObjectTypeA, Mode > const& InSharedPtrA, TSharedPtr< ObjectTypeB, Mode > const& InSharedPtrB )
 {
 	return InSharedPtrA.Get() != InSharedPtrB.Get();
 }
@@ -1764,7 +1795,7 @@ FORCEINLINE bool operator!=( TSharedPtr< ObjectTypeA, Mode > const& InSharedPtrA
  * @return  True if the shared pointer is not null
  */
 template< class ObjectTypeA, ESPMode Mode >
-FORCEINLINE bool operator!=( TSharedPtr< ObjectTypeA, Mode > const& InSharedPtrA, TYPE_OF_NULLPTR )
+[[nodiscard]] FORCEINLINE bool operator!=( TSharedPtr< ObjectTypeA, Mode > const& InSharedPtrA, TYPE_OF_NULLPTR )
 {
 	return InSharedPtrA.IsValid();
 }
@@ -1776,7 +1807,7 @@ FORCEINLINE bool operator!=( TSharedPtr< ObjectTypeA, Mode > const& InSharedPtrA
  * @return  True if the shared pointer is not null
  */
 template< class ObjectTypeB, ESPMode Mode >
-FORCEINLINE bool operator!=( TYPE_OF_NULLPTR, TSharedPtr< ObjectTypeB, Mode > const& InSharedPtrB )
+[[nodiscard]] FORCEINLINE bool operator!=( TYPE_OF_NULLPTR, TSharedPtr< ObjectTypeB, Mode > const& InSharedPtrB )
 {
 	return InSharedPtrB.IsValid();
 }
@@ -1788,7 +1819,7 @@ FORCEINLINE bool operator!=( TYPE_OF_NULLPTR, TSharedPtr< ObjectTypeB, Mode > co
  * @return  True if the shared reference and shared pointer are "equal"
  */
 template< class ObjectTypeA, class ObjectTypeB, ESPMode Mode >
-FORCEINLINE bool operator==( TSharedRef< ObjectTypeA, Mode > const& InSharedRef, TSharedPtr< ObjectTypeB, Mode > const& InSharedPtr )
+[[nodiscard]] FORCEINLINE bool operator==( TSharedRef< ObjectTypeA, Mode > const& InSharedRef, TSharedPtr< ObjectTypeB, Mode > const& InSharedPtr )
 {
 	return InSharedPtr.IsValid() && InSharedPtr.Get() == &( InSharedRef.Get() );
 }
@@ -1800,7 +1831,7 @@ FORCEINLINE bool operator==( TSharedRef< ObjectTypeA, Mode > const& InSharedRef,
  * @return  True if the shared reference and shared pointer are not "equal"
  */
 template< class ObjectTypeA, class ObjectTypeB, ESPMode Mode >
-FORCEINLINE bool operator!=( TSharedRef< ObjectTypeA, Mode > const& InSharedRef, TSharedPtr< ObjectTypeB, Mode > const& InSharedPtr )
+[[nodiscard]] FORCEINLINE bool operator!=( TSharedRef< ObjectTypeA, Mode > const& InSharedRef, TSharedPtr< ObjectTypeB, Mode > const& InSharedPtr )
 {
 	return !InSharedPtr.IsValid() || ( InSharedPtr.Get() != &( InSharedRef.Get() ) );
 }
@@ -1812,7 +1843,7 @@ FORCEINLINE bool operator!=( TSharedRef< ObjectTypeA, Mode > const& InSharedRef,
  * @return  True if the shared reference and shared pointer are "equal"
  */
 template< class ObjectTypeA, class ObjectTypeB, ESPMode Mode >
-FORCEINLINE bool operator==( TSharedPtr< ObjectTypeB, Mode > const& InSharedPtr, TSharedRef< ObjectTypeA, Mode > const& InSharedRef )
+[[nodiscard]] FORCEINLINE bool operator==( TSharedPtr< ObjectTypeB, Mode > const& InSharedPtr, TSharedRef< ObjectTypeA, Mode > const& InSharedRef )
 {
 	return InSharedRef == InSharedPtr;
 }
@@ -1824,7 +1855,7 @@ FORCEINLINE bool operator==( TSharedPtr< ObjectTypeB, Mode > const& InSharedPtr,
  * @return  True if the shared reference and shared pointer are not "equal"
  */
 template< class ObjectTypeA, class ObjectTypeB, ESPMode Mode >
-FORCEINLINE bool operator!=( TSharedPtr< ObjectTypeB, Mode > const& InSharedPtr, TSharedRef< ObjectTypeA, Mode > const& InSharedRef )
+[[nodiscard]] FORCEINLINE bool operator!=( TSharedPtr< ObjectTypeB, Mode > const& InSharedPtr, TSharedRef< ObjectTypeA, Mode > const& InSharedRef )
 {
 	return InSharedRef != InSharedPtr;
 }
@@ -1836,7 +1867,7 @@ FORCEINLINE bool operator!=( TSharedPtr< ObjectTypeB, Mode > const& InSharedPtr,
  * @return  True if the two weak pointers are equal
  */
 template< class ObjectTypeA, class ObjectTypeB, ESPMode Mode >
-FORCEINLINE bool operator==( TWeakPtr< ObjectTypeA, Mode > const& InWeakPtrA, TWeakPtr< ObjectTypeB, Mode > const& InWeakPtrB )
+[[nodiscard]] FORCEINLINE bool operator==( TWeakPtr< ObjectTypeA, Mode > const& InWeakPtrA, TWeakPtr< ObjectTypeB, Mode > const& InWeakPtrB )
 {
 	return InWeakPtrA.Pin().Get() == InWeakPtrB.Pin().Get();
 }
@@ -1848,7 +1879,7 @@ FORCEINLINE bool operator==( TWeakPtr< ObjectTypeA, Mode > const& InWeakPtrA, TW
  * @return  True if the weak pointer and the shared ref are equal
  */
 template< class ObjectTypeA, class ObjectTypeB, ESPMode Mode >
-FORCEINLINE bool operator==( TWeakPtr< ObjectTypeA, Mode > const& InWeakPtrA, TSharedRef< ObjectTypeB, Mode > const& InSharedRefB )
+[[nodiscard]] FORCEINLINE bool operator==( TWeakPtr< ObjectTypeA, Mode > const& InWeakPtrA, TSharedRef< ObjectTypeB, Mode > const& InSharedRefB )
 {
 	return InWeakPtrA.Pin().Get() == &InSharedRefB.Get();
 }
@@ -1860,7 +1891,7 @@ FORCEINLINE bool operator==( TWeakPtr< ObjectTypeA, Mode > const& InWeakPtrA, TS
  * @return  True if the weak pointer and the shared ptr are equal
  */
 template< class ObjectTypeA, class ObjectTypeB, ESPMode Mode >
-FORCEINLINE bool operator==( TWeakPtr< ObjectTypeA, Mode > const& InWeakPtrA, TSharedPtr< ObjectTypeB, Mode > const& InSharedPtrB )
+[[nodiscard]] FORCEINLINE bool operator==( TWeakPtr< ObjectTypeA, Mode > const& InWeakPtrA, TSharedPtr< ObjectTypeB, Mode > const& InSharedPtrB )
 {
 	return InWeakPtrA.Pin().Get() == InSharedPtrB.Get();
 }
@@ -1872,7 +1903,7 @@ FORCEINLINE bool operator==( TWeakPtr< ObjectTypeA, Mode > const& InWeakPtrA, TS
  * @return  True if the weak pointer and the shared ref are equal
  */
 template< class ObjectTypeA, class ObjectTypeB, ESPMode Mode >
-FORCEINLINE bool operator==( TSharedRef< ObjectTypeA, Mode > const& InSharedRefA, TWeakPtr< ObjectTypeB, Mode > const& InWeakPtrB )
+[[nodiscard]] FORCEINLINE bool operator==( TSharedRef< ObjectTypeA, Mode > const& InSharedRefA, TWeakPtr< ObjectTypeB, Mode > const& InWeakPtrB )
 {
 	return &InSharedRefA.Get() == InWeakPtrB.Pin().Get();
 }
@@ -1884,7 +1915,7 @@ FORCEINLINE bool operator==( TSharedRef< ObjectTypeA, Mode > const& InSharedRefA
  * @return  True if the weak pointer and the shared ptr are equal
  */
 template< class ObjectTypeA, class ObjectTypeB, ESPMode Mode >
-FORCEINLINE bool operator==( TSharedPtr< ObjectTypeA, Mode > const& InSharedPtrA, TWeakPtr< ObjectTypeB, Mode > const& InWeakPtrB )
+[[nodiscard]] FORCEINLINE bool operator==( TSharedPtr< ObjectTypeA, Mode > const& InSharedPtrA, TWeakPtr< ObjectTypeB, Mode > const& InWeakPtrB )
 {
 	return InSharedPtrA.Get() == InWeakPtrB.Pin().Get();
 }
@@ -1896,7 +1927,7 @@ FORCEINLINE bool operator==( TSharedPtr< ObjectTypeA, Mode > const& InSharedPtrA
  * @return  True if the weak pointer is null
  */
 template< class ObjectTypeA, ESPMode Mode >
-FORCEINLINE bool operator==( TWeakPtr< ObjectTypeA, Mode > const& InWeakPtrA, TYPE_OF_NULLPTR )
+[[nodiscard]] FORCEINLINE bool operator==( TWeakPtr< ObjectTypeA, Mode > const& InWeakPtrA, TYPE_OF_NULLPTR )
 {
 	return !InWeakPtrA.IsValid();
 }
@@ -1908,7 +1939,7 @@ FORCEINLINE bool operator==( TWeakPtr< ObjectTypeA, Mode > const& InWeakPtrA, TY
  * @return  True if the weak pointer is null
  */
 template< class ObjectTypeB, ESPMode Mode >
-FORCEINLINE bool operator==( TYPE_OF_NULLPTR, TWeakPtr< ObjectTypeB, Mode > const& InWeakPtrB )
+[[nodiscard]] FORCEINLINE bool operator==( TYPE_OF_NULLPTR, TWeakPtr< ObjectTypeB, Mode > const& InWeakPtrB )
 {
 	return !InWeakPtrB.IsValid();
 }
@@ -1920,7 +1951,7 @@ FORCEINLINE bool operator==( TYPE_OF_NULLPTR, TWeakPtr< ObjectTypeB, Mode > cons
  * @return  True if the two weak pointers are not equal
  */
 template< class ObjectTypeA, class ObjectTypeB, ESPMode Mode >
-FORCEINLINE bool operator!=( TWeakPtr< ObjectTypeA, Mode > const& InWeakPtrA, TWeakPtr< ObjectTypeB, Mode > const& InWeakPtrB )
+[[nodiscard]] FORCEINLINE bool operator!=( TWeakPtr< ObjectTypeA, Mode > const& InWeakPtrA, TWeakPtr< ObjectTypeB, Mode > const& InWeakPtrB )
 {
 	return InWeakPtrA.Pin().Get() != InWeakPtrB.Pin().Get();
 }
@@ -1932,7 +1963,7 @@ FORCEINLINE bool operator!=( TWeakPtr< ObjectTypeA, Mode > const& InWeakPtrA, TW
  * @return  True if the weak pointer and the shared ref are not equal
  */
 template< class ObjectTypeA, class ObjectTypeB, ESPMode Mode >
-FORCEINLINE bool operator!=( TWeakPtr< ObjectTypeA, Mode > const& InWeakPtrA, TSharedRef< ObjectTypeB, Mode > const& InSharedRefB )
+[[nodiscard]] FORCEINLINE bool operator!=( TWeakPtr< ObjectTypeA, Mode > const& InWeakPtrA, TSharedRef< ObjectTypeB, Mode > const& InSharedRefB )
 {
 	return InWeakPtrA.Pin().Get() != &InSharedRefB.Get();
 }
@@ -1944,7 +1975,7 @@ FORCEINLINE bool operator!=( TWeakPtr< ObjectTypeA, Mode > const& InWeakPtrA, TS
  * @return  True if the weak pointer and the shared ptr are not equal
  */
 template< class ObjectTypeA, class ObjectTypeB, ESPMode Mode >
-FORCEINLINE bool operator!=( TWeakPtr< ObjectTypeA, Mode > const& InWeakPtrA, TSharedPtr< ObjectTypeB, Mode > const& InSharedPtrB )
+[[nodiscard]] FORCEINLINE bool operator!=( TWeakPtr< ObjectTypeA, Mode > const& InWeakPtrA, TSharedPtr< ObjectTypeB, Mode > const& InSharedPtrB )
 {
 	return InWeakPtrA.Pin().Get() != InSharedPtrB.Get();
 }
@@ -1956,7 +1987,7 @@ FORCEINLINE bool operator!=( TWeakPtr< ObjectTypeA, Mode > const& InWeakPtrA, TS
  * @return  True if the weak pointer and the shared ref are not equal
  */
 template< class ObjectTypeA, class ObjectTypeB, ESPMode Mode >
-FORCEINLINE bool operator!=( TSharedRef< ObjectTypeA, Mode > const& InSharedRefA, TWeakPtr< ObjectTypeB, Mode > const& InWeakPtrB )
+[[nodiscard]] FORCEINLINE bool operator!=( TSharedRef< ObjectTypeA, Mode > const& InSharedRefA, TWeakPtr< ObjectTypeB, Mode > const& InWeakPtrB )
 {
 	return &InSharedRefA.Get() != InWeakPtrB.Pin().Get();
 }
@@ -1968,7 +1999,7 @@ FORCEINLINE bool operator!=( TSharedRef< ObjectTypeA, Mode > const& InSharedRefA
  * @return  True if the weak pointer and the shared ptr are not equal
  */
 template< class ObjectTypeA, class ObjectTypeB, ESPMode Mode >
-FORCEINLINE bool operator!=( TSharedPtr< ObjectTypeA, Mode > const& InSharedPtrA, TWeakPtr< ObjectTypeB, Mode > const& InWeakPtrB )
+[[nodiscard]] FORCEINLINE bool operator!=( TSharedPtr< ObjectTypeA, Mode > const& InSharedPtrA, TWeakPtr< ObjectTypeB, Mode > const& InWeakPtrB )
 {
 	return InSharedPtrA.Get() != InWeakPtrB.Pin().Get();
 }
@@ -1980,7 +2011,7 @@ FORCEINLINE bool operator!=( TSharedPtr< ObjectTypeA, Mode > const& InSharedPtrA
  * @return  True if the weak pointer is not null
  */
 template< class ObjectTypeA, ESPMode Mode >
-FORCEINLINE bool operator!=( TWeakPtr< ObjectTypeA, Mode > const& InWeakPtrA, TYPE_OF_NULLPTR )
+[[nodiscard]] FORCEINLINE bool operator!=( TWeakPtr< ObjectTypeA, Mode > const& InWeakPtrA, TYPE_OF_NULLPTR )
 {
 	return InWeakPtrA.IsValid();
 }
@@ -1992,7 +2023,7 @@ FORCEINLINE bool operator!=( TWeakPtr< ObjectTypeA, Mode > const& InWeakPtrA, TY
  * @return  True if the weak pointer is not null
  */
 template< class ObjectTypeB, ESPMode Mode >
-FORCEINLINE bool operator!=( TYPE_OF_NULLPTR, TWeakPtr< ObjectTypeB, Mode > const& InWeakPtrB )
+[[nodiscard]] FORCEINLINE bool operator!=( TYPE_OF_NULLPTR, TWeakPtr< ObjectTypeB, Mode > const& InWeakPtrB )
 {
 	return InWeakPtrB.IsValid();
 }
@@ -2004,9 +2035,21 @@ FORCEINLINE bool operator!=( TYPE_OF_NULLPTR, TWeakPtr< ObjectTypeB, Mode > cons
  * @param  InSharedPtr  The shared pointer to cast
  */
 template< class CastToType, class CastFromType, ESPMode Mode >
-FORCEINLINE TSharedPtr< CastToType, Mode > StaticCastSharedPtr( TSharedPtr< CastFromType, Mode > const& InSharedPtr )
+[[nodiscard]] FORCEINLINE TSharedPtr< CastToType, Mode > StaticCastSharedPtr(TSharedPtr< CastFromType, Mode > const& InSharedPtr)
 {
-	return TSharedPtr< CastToType, Mode >( InSharedPtr, SharedPointerInternals::FStaticCastTag() );
+	return TSharedPtr< CastToType, Mode >(InSharedPtr, SharedPointerInternals::FStaticCastTag());
+}
+
+
+/**
+ * Casts a weak pointer of one type to another type. (static_cast)  Useful for down-casting.
+ *
+ * @param  InWeakPtr  The weak pointer to cast
+ */
+template< class CastToType, class CastFromType, ESPMode Mode >
+[[nodiscard]] FORCEINLINE TWeakPtr< CastToType, Mode > StaticCastWeakPtr(TWeakPtr< CastFromType, Mode > const& InWeakPtr)
+{
+	return TWeakPtr< CastToType, Mode >(InWeakPtr, SharedPointerInternals::FStaticCastTag());
 }
 
 
@@ -2016,7 +2059,7 @@ FORCEINLINE TSharedPtr< CastToType, Mode > StaticCastSharedPtr( TSharedPtr< Cast
  * @param  InSharedRef  The shared reference to cast
  */
 template< class CastToType, class CastFromType, ESPMode Mode >
-FORCEINLINE TSharedRef< CastToType, Mode > ConstCastSharedRef( TSharedRef< CastFromType, Mode > const& InSharedRef )
+[[nodiscard]] FORCEINLINE TSharedRef< CastToType, Mode > ConstCastSharedRef( TSharedRef< CastFromType, Mode > const& InSharedRef )
 {
 	return TSharedRef< CastToType, Mode >( InSharedRef, SharedPointerInternals::FConstCastTag() );
 }
@@ -2028,12 +2071,24 @@ FORCEINLINE TSharedRef< CastToType, Mode > ConstCastSharedRef( TSharedRef< CastF
  * @param  InSharedPtr  The shared pointer to cast
  */
 template< class CastToType, class CastFromType, ESPMode Mode >
-FORCEINLINE TSharedPtr< CastToType, Mode > ConstCastSharedPtr( TSharedPtr< CastFromType, Mode > const& InSharedPtr )
+[[nodiscard]] FORCEINLINE TSharedPtr< CastToType, Mode > ConstCastSharedPtr( TSharedPtr< CastFromType, Mode > const& InSharedPtr )
 {
 	return TSharedPtr< CastToType, Mode >( InSharedPtr, SharedPointerInternals::FConstCastTag() );
 }
 
 
+/**
+ * Casts a 'const' weak pointer to 'mutable' weak pointer. (const_cast)
+ *
+ * @param  InWeakPtr  The weak pointer to cast
+ */
+template< class CastToType, class CastFromType, ESPMode Mode >
+[[nodiscard]] FORCEINLINE TWeakPtr< CastToType, Mode > ConstCastWeakPtr(TWeakPtr< CastFromType, Mode > const& InWeakPtr)
+{
+	return TWeakPtr< CastToType, Mode >(InWeakPtr, SharedPointerInternals::FConstCastTag());
+}
+
+
 /**
  * MakeShareable utility function.  Wrap object pointers with MakeShareable to allow them to be implicitly
  * converted to shared pointers!  This is useful in assignment operations, or when returning a shared
@@ -2041,7 +2096,7 @@ FORCEINLINE TSharedPtr< CastToType, Mode > ConstCastSharedPtr( TSharedPtr< CastF
  */
 // NOTE: The following is an Unreal extension to standard shared_ptr behavior
 template< class ObjectType >
-FORCEINLINE SharedPointerInternals::TRawPtrProxy< ObjectType > MakeShareable( ObjectType* InObject )
+[[nodiscard]] FORCEINLINE SharedPointerInternals::TRawPtrProxy< ObjectType > MakeShareable( ObjectType* InObject )
 {
 	return SharedPointerInternals::TRawPtrProxy< ObjectType >( InObject );
 }
@@ -2054,7 +2109,7 @@ FORCEINLINE SharedPointerInternals::TRawPtrProxy< ObjectType > MakeShareable( Ob
  */
 // NOTE: The following is an Unreal extension to standard shared_ptr behavior
 template< class ObjectType, class DeleterType >
-FORCEINLINE SharedPointerInternals::TRawPtrProxyWithDeleter< ObjectType, DeleterType > MakeShareable( ObjectType* InObject, DeleterType&& InDeleter )
+[[nodiscard]] FORCEINLINE SharedPointerInternals::TRawPtrProxyWithDeleter< ObjectType, DeleterType > MakeShareable( ObjectType* InObject, DeleterType&& InDeleter )
 {
 	return SharedPointerInternals::TRawPtrProxyWithDeleter< ObjectType, DeleterType >( InObject, Forward< DeleterType >( InDeleter ) );
 }
@@ -2068,7 +2123,7 @@ FORCEINLINE SharedPointerInternals::TRawPtrProxyWithDeleter< ObjectType, Deleter
  *	  friend class SharedPointerInternals::TIntrusiveReferenceController;
  */
 template <typename InObjectType, ESPMode InMode = ESPMode::ThreadSafe, typename... InArgTypes>
-FORCEINLINE TSharedRef<InObjectType, InMode> MakeShared(InArgTypes&&... Args)
+[[nodiscard]] FORCEINLINE TSharedRef<InObjectType, InMode> MakeShared(InArgTypes&&... Args)
 {
 	SharedPointerInternals::TIntrusiveReferenceController<InObjectType, InMode>* Controller = SharedPointerInternals::NewIntrusiveReferenceController<InMode, InObjectType>(Forward<InArgTypes>(Args)...);
 	return UE::Core::Private::MakeSharedRef<InObjectType, InMode>(Controller->GetObjectPtr(), (SharedPointerInternals::TReferenceControllerBase<InMode>*)Controller);
diff --git a/Engine/Source/Runtime/Core/Public/Templates/SharedPointerInternals.h b/Engine/Source/Runtime/Core/Public/Templates/SharedPointerInternals.h
--- a/Engine/Source/Runtime/Core/Public/Templates/SharedPointerInternals.h
+++ b/Engine/Source/Runtime/Core/Public/Templates/SharedPointerInternals.h
@@ -359,6 +359,26 @@ namespace SharedPointerInternals
 		template <typename... ArgTypes>
 		explicit TIntrusiveReferenceController(ArgTypes&&... Args)
 		{
+			// If this fails to compile when trying to called MakeShared with a non-public constructor,
+			// do not make SharedPointerInternals::TIntrusiveReferenceController a friend.
+			//
+			// Instead, prefer this pattern:
+			//
+			//     class FMyType
+			//     {
+			//     private:
+			//         struct FPrivateToken { explicit FPrivateToken() = default; };
+			//
+			//     public:
+			//         // This has an equivalent access level to a private constructor,
+			//         // as only friends of FMyType will have access to FPrivateToken,
+			//         // but MakeShared can legally call it since it's public.
+			//         explicit FMyType(FPrivateToken, int32 Int, float Real, const TCHAR* String);
+			//     };
+			//
+			//     // Won't compile if the caller doesn't have access to FMyType::FPrivateToken
+			//     TSharedPtr<FMyType> Val = MakeShared<FMyType>(FMyType::FPrivateToken{}, 5, 3.14f, TEXT("Banana"));
+			//
 			new ((void*)&ObjectStorage) ObjectType(Forward<ArgTypes>(Args)...);
 		}
 
