fix a bug where FScopedSlowTask can get stuck auto-cancelling if the dialog-create-delay is longer than it takes to call ShouldCancel()

diff --git a/Engine/Source/Runtime/Core/Private/Misc/SlowTask.cpp b/Engine/Source/Runtime/Core/Private/Misc/SlowTask.cpp
--- a/Engine/Source/Runtime/Core/Private/Misc/SlowTask.cpp
+++ b/Engine/Source/Runtime/Core/Private/Misc/SlowTask.cpp
@@ -202,7 +202,7 @@ void FSlowTask::MakeDialog(bool bShowCancelButton, bool bAllowInPIE)
 
 bool FSlowTask::ShouldCancel() const
 {
-	if (bEnabled)
+	if (bEnabled && bCreatedDialog)
 	{
 		check(IsInGameThread()); // FSlowTask is only meant to be used on the main thread currently
 
