Comment about failed compilation of non-public constructors copied from TIntrusiveReferenceController to TArray::Emplace.

diff --git a/Engine/Source/Runtime/Core/Public/Containers/Array.h b/Engine/Source/Runtime/Core/Public/Containers/Array.h
--- a/Engine/Source/Runtime/Core/Public/Containers/Array.h
+++ b/Engine/Source/Runtime/Core/Public/Containers/Array.h
@@ -1975,6 +1975,29 @@ public:
 	FORCEINLINE SizeType Emplace(ArgsType&&... Args)
 	{
 		const SizeType Index = AddUninitialized(1);
+
+		// If this fails to compile when trying to call Emplace with a non-public constructor,
+		// do not make TArray a friend.
+		//
+		// Instead, prefer this pattern:
+		//
+		//     class FMyType
+		//     {
+		//     private:
+		//         struct FPrivateToken { explicit FPrivateToken() = default; };
+		//
+		//     public:
+		//         // This has an equivalent access level to a private constructor,
+		//         // as only friends of FMyType will have access to FPrivateToken,
+		//         // but Emplace can legally call it since it's public.
+		//         explicit FMyType(FPrivateToken, int32 Int, float Real, const TCHAR* String);
+		//     };
+		//
+		//     TArray<FMyType> Arr:
+		//
+		//     // Won't compile if the caller doesn't have access to FMyType::FPrivateToken
+		//     Arr.Emplace(FMyType::FPrivateToken{}, 5, 3.14f, TEXT("Banana"));
+		//
 		new(GetData() + Index) ElementType(Forward<ArgsType>(Args)...);
 		return Index;
 	}
diff --git a/Engine/Source/Runtime/Core/Public/Templates/SharedPointerInternals.h b/Engine/Source/Runtime/Core/Public/Templates/SharedPointerInternals.h
--- a/Engine/Source/Runtime/Core/Public/Templates/SharedPointerInternals.h
+++ b/Engine/Source/Runtime/Core/Public/Templates/SharedPointerInternals.h
@@ -359,7 +359,7 @@ namespace SharedPointerInternals
 		template <typename... ArgTypes>
 		explicit TIntrusiveReferenceController(ArgTypes&&... Args)
 		{
-			// If this fails to compile when trying to called MakeShared with a non-public constructor,
+			// If this fails to compile when trying to call MakeShared with a non-public constructor,
 			// do not make SharedPointerInternals::TIntrusiveReferenceController a friend.
 			//
 			// Instead, prefer this pattern:
