A temporary iOS only workaround for a memory allocator crash on UE5/Main until UE-154181 is resolved.

diff --git a/Engine/Source/Runtime/Core/Private/Apple/ApplePlatformMemory.cpp b/Engine/Source/Runtime/Core/Private/Apple/ApplePlatformMemory.cpp
--- a/Engine/Source/Runtime/Core/Private/Apple/ApplePlatformMemory.cpp
+++ b/Engine/Source/Runtime/Core/Private/Apple/ApplePlatformMemory.cpp
@@ -298,7 +298,13 @@ FMalloc* FApplePlatformMemory::BaseAllocator()
 	}
 	else
 	{
+		// A temporary workaround for an iOS specific memory allocator issue causing crashes on Main.
+		// See jira UE-154181 for details. Remove once jira is fixed.
+#if PLATFORM_IOS
+		AllocatorToUse = EMemoryAllocatorToUse::Ansi;
+#else
 		AllocatorToUse = EMemoryAllocatorToUse::Binned;
+#endif
 	}
 	
 	// Force ansi malloc in some cases
