ParallelFor: Cancel unused tasks and time the wait for tasks to complete

When worker threads are heavily utilized by slow tasks, ParallelFor may execute most of the work on the calling thread before most (or any) of its tasks have executed. Allowing these tasks to execute later is wasteful because each task executed on a background worker pays about 0.3us to lower and then raise thread priority. Canceling the tasks can avoid that cost when paired with an upcoming change to the scheduler, which saves about 40us for each call to ParallelFor in a heavily-utilized pool of 126 worker threads.

diff --git a/Engine/Source/Runtime/Core/Public/Async/ParallelFor.h b/Engine/Source/Runtime/Core/Public/Async/ParallelFor.h
--- a/Engine/Source/Runtime/Core/Public/Async/ParallelFor.h
+++ b/Engine/Source/Runtime/Core/Public/Async/ParallelFor.h
@@ -344,6 +344,15 @@ namespace ParallelForImpl
 		FParallelExecutor LocalExecutor(MoveTemp(Data), NumWorkers, Priority);
 		const bool bFinishedLast = LocalExecutor(true, DebugName);
 
+		// Cancel tasks that have not yet launched since they will otherwise waste time on worker threads
+		{
+			TRACE_CPUPROFILER_EVENT_SCOPE(ParallelFor.Cancel);
+			for (LowLevelTasks::FTask& Task : LocalExecutor.GetData()->Tasks)
+			{
+				Task.TryCancel(LowLevelTasks::ECancellationFlags::PrelaunchCancellation);
+			}
+		}
+
 		if(!bFinishedLast)
 		{
 			const bool bPumpRenderingThread  = (Flags & EParallelForFlags::PumpRenderingThread) != EParallelForFlags::None;
@@ -360,6 +369,7 @@ namespace ParallelForImpl
 			{
 				// FinishedSignal waits here if some other thread finishes the last item
 				// Data must live on until all of the tasks are cleared which might be long after this function exits
+				TRACE_CPUPROFILER_EVENT_SCOPE(ParallelFor.Wait);
 				FinishedSignal->Wait();
 			}
 		}
