Remove hardcoded NDA platform define.  This can be defined for a game by specifying bAllowConfidentialPlatformDefines as a target rule.

diff --git a/Engine/Source/Runtime/Core/Public/HAL/Platform.h b/Engine/Source/Runtime/Core/Public/HAL/Platform.h
--- a/Engine/Source/Runtime/Core/Public/HAL/Platform.h
+++ b/Engine/Source/Runtime/Core/Public/HAL/Platform.h
@@ -20,9 +20,6 @@
 	#define PLATFORM_MAC_X86 0
 	#define PLATFORM_MAC_ARM64 0
 #endif
-#if !defined(PLATFORM_PS4)
-	#define PLATFORM_PS4 0
-#endif
 #if !defined(PLATFORM_IOS)
 	#define PLATFORM_IOS 0
 #endif
