remove Hololens from ExceptionHandling.h

 - only one function used in one place, so easier to extern it directly.

#rnx

diff --git a/Engine/Platforms/Hololens/Source/Runtime/Core/Private/HoloLensRunnableThread.cpp b/Engine/Platforms/Hololens/Source/Runtime/Core/Private/HoloLensRunnableThread.cpp
--- a/Engine/Platforms/Hololens/Source/Runtime/Core/Private/HoloLensRunnableThread.cpp
+++ b/Engine/Platforms/Hololens/Source/Runtime/Core/Private/HoloLensRunnableThread.cpp
@@ -5,9 +5,14 @@
 #include "HAL/ExceptionHandling.h"
 #include "Misc/FeedbackContext.h"
 #include "CoreGlobals.h"
+#include "HoloLensSystemIncludes.h"
+#include <excpt.h>
 
 DEFINE_LOG_CATEGORY_STATIC(LogThreadingWindows, Log, All);
 
+extern CORE_API int32 ReportCrash(Windows::LPEXCEPTION_POINTERS ExceptionInfo);
+
+
 
 uint32 FRunnableThreadHoloLens::GuardedRun()
 {
diff --git a/Engine/Source/Runtime/Core/Public/HAL/ExceptionHandling.h b/Engine/Source/Runtime/Core/Public/HAL/ExceptionHandling.h
--- a/Engine/Source/Runtime/Core/Public/HAL/ExceptionHandling.h
+++ b/Engine/Source/Runtime/Core/Public/HAL/ExceptionHandling.h
@@ -40,13 +40,6 @@ extern CORE_API void ReportGPUCrash(const TCHAR* ErrorMessage, void* ProgramCoun
 extern CORE_API void ReportEnsure(const TCHAR* ErrorMessage, void* ProgramCounter);
 extern CORE_API void ReportStall(const TCHAR* ErrorMessage, uint32 ThreadId);
 extern CORE_API void ReportHang(const TCHAR*, const uint64* StackFrames, int32 NumStackFrames, uint32 HungThreadId);
-#elif PLATFORM_HOLOLENS
-#include "HoloLensSystemIncludes.h"
-#include <excpt.h>
-extern CORE_API int32 ReportCrash(Windows::LPEXCEPTION_POINTERS ExceptionInfo);
-extern CORE_API void ReportEnsure(const TCHAR* ErrorMessage, void* ProgramCounter);
-extern CORE_API void ReportHang(const TCHAR*, const uint64* StackFrames, int32 NumStackFrames, uint32 HungThreadId);
-extern CORE_API void NewReportEnsure(const TCHAR* ErrorMessage);
 #endif
 
 
