GitHub 8898 : FIntVector: Add component-wise multiplication operator

diff --git a/Engine/Source/Runtime/Core/Public/Math/IntVector.h b/Engine/Source/Runtime/Core/Public/Math/IntVector.h
--- a/Engine/Source/Runtime/Core/Public/Math/IntVector.h
+++ b/Engine/Source/Runtime/Core/Public/Math/IntVector.h
@@ -180,6 +180,21 @@ PRAGMA_ENABLE_DEPRECATION_WARNINGS
 		return X != Other.X || Y != Other.Y || Z != Other.Z;
 	}
 
+	/**
+	 * Multiplies this vector with another vector, using component-wise multiplication.
+	 *
+	 * @param Other The point to multiply with.
+	 * @return Reference to this point after multiplication.
+	 */
+	TIntVector3& operator*=(const TIntVector3& Other)
+	{
+		X *= Other.X;
+		Y *= Other.Y;
+		Z *= Other.Z;
+
+		return *this;
+	}
+
 	/**
 	 * Scales this point.
 	 *
@@ -240,6 +255,17 @@ PRAGMA_ENABLE_DEPRECATION_WARNINGS
 		return *this;
 	}
 
+	/**
+	 * Gets the result of component-wise multiplication of this point by another.
+	 *
+	 * @param Other The point to multiply with.
+	 * @return The result of multiplication.
+	 */
+	TIntVector3 operator*(const TIntVector3& Other) const
+	{
+		return TIntVector3(*this) *= Other;
+	}
+
 	/**
 	 * Gets the result of scaling on this point.
 	 *
@@ -712,6 +738,22 @@ PRAGMA_ENABLE_DEPRECATION_WARNINGS
 		return X != Other.X || Y != Other.Y || Z != Other.Z || W != Other.W;
 	}
 
+	/**
+	 * Multiplies this vector with another vector, using component-wise multiplication.
+	 *
+	 * @param Other The point to multiply with.
+	 * @return Reference to this point after multiplication.
+	 */
+	TIntVector4& operator*=(const TIntVector4& Other)
+	{
+		X *= Other.X;
+		Y *= Other.Y;
+		Z *= Other.Z;
+		W *= Other.W;
+
+		return *this;
+	}
+
 	/**
 	 * Scales this point.
 	 *
@@ -792,6 +834,17 @@ PRAGMA_ENABLE_DEPRECATION_WARNINGS
 		return *this;
 	}
 
+	/**
+	 * Gets the result of component-wise multiplication of this point by another.
+	 *
+	 * @param Other The point to multiply with.
+	 * @return The result of multiplication.
+	 */
+	TIntVector4 operator*(const TIntVector4& Other) const
+	{
+		return TIntVector4(*this) *= Other;
+	}
+
 	/**
 	 * Gets the result of scaling on this point.
 	 *
