Memory tracing refactorings

As the memory tracing implementation grew and callstack tracing was extracted it became clear that alot of the code which was implemented as platform specific should really be platform independent. This changelist moves all MemoryTrace_* methods except the MemoryTrace_Initialize to a platform independent layer. This also allows for internalizing the "do-not-allocate-while-tracing-an-allocation" flag. Also took the oportunity to rename MemoryTrace.cpp in Microsoft folder to better reflect platform specific file.

diff --git a/Engine/Source/Runtime/Core/Private/ProfilingDebugging/MemoryAllocationTrace.cpp b/Engine/Source/Runtime/Core/Private/ProfilingDebugging/MemoryAllocationTrace.cpp
deleted file mode 100644
diff --git a/Engine/Source/Runtime/Core/Private/ProfilingDebugging/MemoryAllocationTrace.h b/Engine/Source/Runtime/Core/Private/ProfilingDebugging/MemoryAllocationTrace.h
deleted file mode 100644
diff --git a/Engine/Source/Runtime/Core/Private/ProfilingDebugging/MemoryTrace.cpp b/Engine/Source/Runtime/Core/Private/ProfilingDebugging/MemoryTrace.cpp
new file mode 100644
--- /dev/null
+++ b/Engine/Source/Runtime/Core/Private/ProfilingDebugging/MemoryTrace.cpp
@@ -0,0 +1,683 @@
+// Copyright Epic Games, Inc. All Rights Reserved.
+
+#include "ProfilingDebugging/MemoryTrace.h"
+#include "HAL/PlatformTime.h"
+#include "ProfilingDebugging/CallstackTrace.h"
+#include "ProfilingDebugging/TagTrace.h"
+#include "ProfilingDebugging/TraceMalloc.h"
+#include "Trace/Trace.h"
+
+#if UE_TRACE_ENABLED
+UE_TRACE_CHANNEL_DEFINE(MemAllocChannel, "Memory allocations", true)
+#endif
+
+#if UE_MEMORY_TRACE_ENABLED
+
+////////////////////////////////////////////////////////////////////////////////
+void 	MemoryTrace_InitTags(FMalloc*);
+void	MemoryTrace_EnableTracePump();
+
+////////////////////////////////////////////////////////////////////////////////
+namespace
+{
+	// Controls how often time markers are emitted (default every 4095 allocation)
+	constexpr uint32 MarkerSamplePeriod	= (4 << 10) - 1;
+	
+	// Number of bits shifted bits to SizeLower
+	constexpr uint32 SizeShift = 3;
+	
+	// Counter to track when time marker is emitted
+	std::atomic<uint32>	GMarkerCounter(0);
+	
+	// If enabled also pumps the Trace system itself. Used on process shutdown
+	// when worker thread has been killed, but memory events still occurs.
+	bool				GDoPumpTrace;
+	
+	// Temporarily disables any internal operation that causes allocations. Used to
+	// avoid recursive behaviour when memory tracing needs to allocate memory through
+	// TraceMalloc
+	thread_local bool	GDoNotAllocateInTrace;
+	
+	// Set on initialization, on some platforms we hook allocator functions very early
+	// before Trace has the ability to allocate memory.
+	bool				GTraceAllowed; 
+}
+
+////////////////////////////////////////////////////////////////////////////////
+namespace UE {
+namespace Trace {
+	TRACELOG_API void Update();
+} // namespace Trace
+} // namespace UE
+
+////////////////////////////////////////////////////////////////////////////////
+
+UE_TRACE_EVENT_BEGIN(Memory, Init, NoSync|Important)
+	UE_TRACE_EVENT_FIELD(uint32, MarkerPeriod)
+	UE_TRACE_EVENT_FIELD(uint8, Version)
+	UE_TRACE_EVENT_FIELD(uint8, MinAlignment)
+	UE_TRACE_EVENT_FIELD(uint8, SizeShift)
+	UE_TRACE_EVENT_FIELD(uint8, Mode)
+UE_TRACE_EVENT_END()
+
+UE_TRACE_EVENT_BEGIN(Memory, Marker)
+	UE_TRACE_EVENT_FIELD(uint64, Cycle)
+UE_TRACE_EVENT_END()
+
+UE_TRACE_EVENT_BEGIN(Memory, Alloc)
+	UE_TRACE_EVENT_FIELD(uint64, Address)
+	UE_TRACE_EVENT_FIELD(uint32, CallstackId)
+	UE_TRACE_EVENT_FIELD(uint32, Size)
+	UE_TRACE_EVENT_FIELD(uint8, AlignmentPow2_SizeLower)
+	UE_TRACE_EVENT_FIELD(uint8, RootHeap)
+UE_TRACE_EVENT_END()
+
+UE_TRACE_EVENT_BEGIN(Memory, AllocSystem)
+	UE_TRACE_EVENT_FIELD(uint64, Address)
+	UE_TRACE_EVENT_FIELD(uint32, CallstackId)
+	UE_TRACE_EVENT_FIELD(uint32, Size)
+	UE_TRACE_EVENT_FIELD(uint8, AlignmentPow2_SizeLower)
+UE_TRACE_EVENT_END()
+
+UE_TRACE_EVENT_BEGIN(Memory, AllocVideo)
+	UE_TRACE_EVENT_FIELD(uint64, Address)
+	UE_TRACE_EVENT_FIELD(uint32, CallstackId)
+	UE_TRACE_EVENT_FIELD(uint32, Size)
+	UE_TRACE_EVENT_FIELD(uint8, AlignmentPow2_SizeLower)
+UE_TRACE_EVENT_END()
+
+UE_TRACE_EVENT_BEGIN(Memory, Free)
+	UE_TRACE_EVENT_FIELD(uint64, Address)
+	UE_TRACE_EVENT_FIELD(uint8, RootHeap)
+UE_TRACE_EVENT_END()
+
+UE_TRACE_EVENT_BEGIN(Memory, FreeSystem)
+	UE_TRACE_EVENT_FIELD(uint64, Address)
+UE_TRACE_EVENT_END()
+
+UE_TRACE_EVENT_BEGIN(Memory, FreeVideo)
+	UE_TRACE_EVENT_FIELD(uint64, Address)
+UE_TRACE_EVENT_END()
+
+UE_TRACE_EVENT_BEGIN(Memory, ReallocAlloc)
+	UE_TRACE_EVENT_FIELD(uint64, Address)
+	UE_TRACE_EVENT_FIELD(uint32, CallstackId)
+	UE_TRACE_EVENT_FIELD(uint32, Size)
+	UE_TRACE_EVENT_FIELD(uint8, AlignmentPow2_SizeLower)
+	UE_TRACE_EVENT_FIELD(uint8, RootHeap)
+UE_TRACE_EVENT_END()
+
+UE_TRACE_EVENT_BEGIN(Memory, ReallocAllocSystem)
+	UE_TRACE_EVENT_FIELD(uint64, Address)
+	UE_TRACE_EVENT_FIELD(uint32, CallstackId)
+	UE_TRACE_EVENT_FIELD(uint32, Size)
+	UE_TRACE_EVENT_FIELD(uint8, AlignmentPow2_SizeLower)
+UE_TRACE_EVENT_END()
+
+UE_TRACE_EVENT_BEGIN(Memory, ReallocFree)
+	UE_TRACE_EVENT_FIELD(uint64, Address)
+	UE_TRACE_EVENT_FIELD(uint8, RootHeap)
+UE_TRACE_EVENT_END()
+
+UE_TRACE_EVENT_BEGIN(Memory, ReallocFreeSystem)
+	UE_TRACE_EVENT_FIELD(uint64, Address)
+UE_TRACE_EVENT_END()
+
+UE_TRACE_EVENT_BEGIN(Memory, HeapSpec, NoSync|Important)
+	UE_TRACE_EVENT_FIELD(HeapId, Id)
+	UE_TRACE_EVENT_FIELD(HeapId, ParentId)
+	UE_TRACE_EVENT_FIELD(uint16, Flags)
+	UE_TRACE_EVENT_FIELD(UE::Trace::WideString, Name)
+UE_TRACE_EVENT_END()
+
+UE_TRACE_EVENT_BEGIN(Memory, HeapMarkAlloc)
+	UE_TRACE_EVENT_FIELD(uint64, Address)
+	UE_TRACE_EVENT_FIELD(uint16, Flags)
+	UE_TRACE_EVENT_FIELD(HeapId, Heap)
+UE_TRACE_EVENT_END()
+
+UE_TRACE_EVENT_BEGIN(Memory, HeapUnmarkAlloc)
+	UE_TRACE_EVENT_FIELD(uint64, Address)
+	UE_TRACE_EVENT_FIELD(HeapId, Heap)
+UE_TRACE_EVENT_END()
+
+// If layout of the above events are changed, bump this version number
+constexpr uint8 MemoryTraceVersion = 1;
+
+////////////////////////////////////////////////////////////////////////////////
+class FMallocWrapper
+	: public FMalloc
+{
+public:
+							FMallocWrapper(FMalloc* InMalloc);
+
+private:
+	struct FCookie
+	{
+		uint64				Tag  : 16;
+		uint64				Bias : 8;
+		uint64				Size : 40;
+	};
+
+	static uint32			GetActualAlignment(SIZE_T Size, uint32 Alignment);
+	virtual void*			Malloc(SIZE_T Size, uint32 Alignment) override;
+	virtual void*			Realloc(void* PrevAddress, SIZE_T NewSize, uint32 Alignment) override;
+	virtual void			Free(void* Address) override;
+	virtual bool			IsInternallyThreadSafe() const override						{ return InnerMalloc->IsInternallyThreadSafe(); }
+	virtual void			UpdateStats() override										{ InnerMalloc->UpdateStats(); }
+	virtual void			GetAllocatorStats(FGenericMemoryStats& Out) override		{ InnerMalloc->GetAllocatorStats(Out); }
+	virtual void			DumpAllocatorStats(FOutputDevice& Ar) override				{ InnerMalloc->DumpAllocatorStats(Ar); }
+	virtual bool			ValidateHeap() override										{ return InnerMalloc->ValidateHeap(); }
+	virtual bool			GetAllocationSize(void* Address, SIZE_T &SizeOut) override	{ return InnerMalloc->GetAllocationSize(Address, SizeOut); }
+	virtual void			SetupTLSCachesOnCurrentThread() override					{ return InnerMalloc->SetupTLSCachesOnCurrentThread(); }
+	virtual void			OnPreFork() override										{ InnerMalloc->OnPreFork();}
+	virtual void			OnPostFork() override										{ InnerMalloc->OnPostFork(); }
+
+	FMalloc*				InnerMalloc;
+};
+
+////////////////////////////////////////////////////////////////////////////////
+FMallocWrapper::FMallocWrapper(FMalloc* InMalloc)
+: InnerMalloc(InMalloc)
+{
+}
+
+////////////////////////////////////////////////////////////////////////////////
+uint32 FMallocWrapper::GetActualAlignment(SIZE_T Size, uint32 Alignment)
+{
+	// Defaults; if size is < 16 then alignment is 8 else 16.
+	uint32 DefaultAlignment = 8 << uint32(Size >= 16);
+	return (Alignment < DefaultAlignment) ? DefaultAlignment : Alignment;
+}
+
+////////////////////////////////////////////////////////////////////////////////
+void* FMallocWrapper::Malloc(SIZE_T Size, uint32 Alignment)
+{
+	if (Size == 0)
+	{
+		return nullptr;
+	}
+
+	uint32 ActualAlignment = GetActualAlignment(Size, Alignment);
+	void* Address = InnerMalloc->Malloc(Size, Alignment);
+
+	MemoryTrace_Alloc((uint64)Address, Size, ActualAlignment);
+
+	return Address;
+}
+
+////////////////////////////////////////////////////////////////////////////////
+void* FMallocWrapper::Realloc(void* PrevAddress, SIZE_T NewSize, uint32 Alignment)
+{
+	// This simplifies things and means reallocs trace events are true reallocs
+	if (PrevAddress == nullptr)
+	{
+		return Malloc(NewSize, Alignment);
+	}
+
+	if (NewSize == 0)
+	{
+		Free(PrevAddress);
+		return nullptr;
+	}
+
+	MemoryTrace_ReallocFree((uint64)PrevAddress);
+
+	void* RetAddress = InnerMalloc->Realloc(PrevAddress, NewSize, Alignment);
+
+	Alignment = GetActualAlignment(NewSize, Alignment);
+	MemoryTrace_ReallocAlloc((uint64)RetAddress, NewSize, Alignment);
+
+	return RetAddress;
+}
+
+////////////////////////////////////////////////////////////////////////////////
+void FMallocWrapper::Free(void* Address)
+{
+	if (Address == nullptr)
+	{
+		return;
+	}
+
+	MemoryTrace_Free((uint64)Address);
+
+	void* InnerAddress = Address;
+
+	return InnerMalloc->Free(InnerAddress);
+}
+
+////////////////////////////////////////////////////////////////////////////////
+template <class T>
+class alignas(alignof(T)) FUndestructed
+{
+public:
+	template <typename... ArgTypes>
+	void Construct(ArgTypes... Args)
+	{
+		::new (Buffer) T(Args...);
+		bIsConstructed = true;
+	}
+
+	bool IsConstructed() const
+	{
+		return bIsConstructed;
+	}
+
+	T* operator & ()	{ return (T*)Buffer; }
+	T* operator -> ()	{ return (T*)Buffer; }
+
+protected:
+	uint8 Buffer[sizeof(T)];
+	bool bIsConstructed;
+};
+
+////////////////////////////////////////////////////////////////////////////////
+static FUndestructed<FTraceMalloc> GTraceMalloc;
+
+////////////////////////////////////////////////////////////////////////////////
+FMalloc* MemoryTrace_CreateInternal(FMalloc* InMalloc)
+{
+	// Some OSes (i.e. Windows) will terminate all threads except the main
+	// one as part of static deinit. However we may receive more memory
+	// trace events that would get lost as Trace's worker thread has been
+	// terminated. So flush the last remaining memory events trace needs
+	// to be updated which we will do that in response to to memory events.
+	// We'll use an atexit can to know when Trace is probably no longer
+	// getting ticked.
+	atexit([]() { MemoryTrace_EnableTracePump(); });
+
+	GTraceMalloc.Construct(InMalloc);
+
+	// Both tag and callstack tracing need to use the wrapped trace malloc
+	// so we can break out tracing memory overhead (and not cause recursive behaviour).
+	MemoryTrace_InitTags(&GTraceMalloc);
+	CallstackTrace_Create(&GTraceMalloc);
+
+
+	static FUndestructed<FMallocWrapper> SMallocWrapper;
+	SMallocWrapper.Construct(InMalloc);
+
+	return &SMallocWrapper;
+}
+
+////////////////////////////////////////////////////////////////////////////////
+void MemoryTrace_Initialize()
+{
+	// At this point we initialized the system to allow tracing.
+	GTraceAllowed = true;
+	
+	UE_TRACE_LOG(Memory, Init, MemAllocChannel)
+		<< Init.MarkerPeriod(MarkerSamplePeriod + 1)
+		<< Init.Version(MemoryTraceVersion)
+		<< Init.MinAlignment(uint8(MIN_ALIGNMENT))
+		<< Init.SizeShift(uint8(SizeShift));
+
+	const HeapId SystemRootHeap = MemoryTrace_RootHeapSpec(TEXT("System memory"));
+	check(SystemRootHeap == EMemoryTraceRootHeap::SystemMemory);
+	const HeapId VideoRootHeap = MemoryTrace_RootHeapSpec(TEXT("Video memory"));
+	check(VideoRootHeap == EMemoryTraceRootHeap::VideoMemory);
+
+	static_assert((1 << SizeShift) - 1 <= MIN_ALIGNMENT, "Not enough bits to pack size fields");
+
+#if !UE_MEMORY_TRACE_LATE_INIT
+	// On some platforms callstack initialization cannot happen this early in the process. It is initialized
+	// in other locations when UE_MEMORY_TRACE_LATE_INIT is defined. Until that point allocations cannot have
+	// callstacks.
+	CallstackTrace_Initialize();
+#endif
+}
+
+////////////////////////////////////////////////////////////////////////////////
+void MemoryTrace_EnableTracePump()
+{
+	GDoPumpTrace = true;
+}
+
+////////////////////////////////////////////////////////////////////////////////
+void MemoryTrace_UpdateInternal()
+{
+	const uint32 TheCount = GMarkerCounter.fetch_add(1, std::memory_order_relaxed);
+	if ((TheCount & MarkerSamplePeriod) == 0)
+	{
+		UE_TRACE_LOG(Memory, Marker, MemAllocChannel)
+			<< Marker.Cycle(FPlatformTime::Cycles64());
+	}
+
+	if (GDoPumpTrace)
+	{
+		UE::Trace::Update();
+	}
+}
+
+////////////////////////////////////////////////////////////////////////////////
+void MemoryTrace_Alloc(uint64 Address, uint64 Size, uint32 Alignment, HeapId RootHeap)
+{
+	check(RootHeap < 16);
+	if (!GTraceAllowed)
+	{
+		return;
+	}
+		
+	const uint32 AlignmentPow2 = uint32(FPlatformMath::CountTrailingZeros(Alignment));
+	const uint32 Alignment_SizeLower = (AlignmentPow2 << SizeShift) | uint32(Size & ((1 << SizeShift) - 1));
+	const uint32 CallstackId = GDoNotAllocateInTrace ? 0 : CallstackTrace_GetCurrentId();
+	
+	switch (RootHeap)
+	{
+		case EMemoryTraceRootHeap::SystemMemory:
+		{
+			UE_TRACE_LOG(Memory, AllocSystem, MemAllocChannel)
+				<< AllocSystem.CallstackId(CallstackId)
+				<< AllocSystem.Address(uint64(Address))
+				<< AllocSystem.Size(uint32(Size >> SizeShift))
+				<< AllocSystem.AlignmentPow2_SizeLower(uint8(Alignment_SizeLower));
+			break;
+		}
+
+		case EMemoryTraceRootHeap::VideoMemory:
+		{
+			UE_TRACE_LOG(Memory, AllocVideo, MemAllocChannel)
+				<< AllocVideo.CallstackId(CallstackId)
+				<< AllocVideo.Address(uint64(Address))
+				<< AllocVideo.Size(uint32(Size >> SizeShift))
+				<< AllocVideo.AlignmentPow2_SizeLower(uint8(Alignment_SizeLower));
+			break;
+		}
+
+		default:
+		{
+			UE_TRACE_LOG(Memory, Alloc, MemAllocChannel)
+				<< Alloc.CallstackId(CallstackId)
+				<< Alloc.Address(uint64(Address))
+				<< Alloc.RootHeap(uint8(RootHeap))
+				<< Alloc.Size(uint32(Size >> SizeShift))
+				<< Alloc.AlignmentPow2_SizeLower(uint8(Alignment_SizeLower));
+			break;
+		}
+	}
+
+	MemoryTrace_UpdateInternal();
+}
+
+////////////////////////////////////////////////////////////////////////////////
+void MemoryTrace_Free(uint64 Address, HeapId RootHeap)
+{
+	check(RootHeap < 16);
+	if (!GTraceAllowed)
+	{
+		return;
+	}
+	
+	switch (RootHeap)
+	{
+		case EMemoryTraceRootHeap::SystemMemory:
+			{
+				UE_TRACE_LOG(Memory, FreeSystem, MemAllocChannel)
+					<< FreeSystem.Address(uint64(Address));
+				break;
+			}
+		case EMemoryTraceRootHeap::VideoMemory:
+			{
+				UE_TRACE_LOG(Memory, FreeVideo, MemAllocChannel)
+					<< FreeVideo.Address(uint64(Address));
+				break;
+			}
+		default:
+			{
+				UE_TRACE_LOG(Memory, Free, MemAllocChannel)
+					<< Free.Address(uint64(Address))
+					<< Free.RootHeap(uint8(RootHeap));
+				break;
+			}
+	}
+
+	MemoryTrace_UpdateInternal();
+}
+
+////////////////////////////////////////////////////////////////////////////////
+void MemoryTrace_ReallocAlloc(uint64 Address, uint64 Size, uint32 Alignment, HeapId RootHeap)
+{
+	check(RootHeap < 16);
+	if (!GTraceAllowed)
+	{
+		return;
+	}
+	
+	const uint32 AlignmentPow2 = uint32(FPlatformMath::CountTrailingZeros(Alignment));
+	const uint32 Alignment_SizeLower = (AlignmentPow2 << SizeShift) | uint32(Size & ((1 << SizeShift) - 1));
+	const uint32 CallstackId = GDoNotAllocateInTrace ? 0 : CallstackTrace_GetCurrentId();
+
+	switch (RootHeap)
+	{
+		case EMemoryTraceRootHeap::SystemMemory:
+		{
+			UE_TRACE_LOG(Memory, ReallocAllocSystem, MemAllocChannel)
+				<< ReallocAllocSystem.CallstackId(CallstackId)
+				<< ReallocAllocSystem.Address(uint64(Address))
+				<< ReallocAllocSystem.Size(uint32(Size >> SizeShift))
+				<< ReallocAllocSystem.AlignmentPow2_SizeLower(uint8(Alignment_SizeLower));
+			break;
+		}
+
+		default:
+		{
+			UE_TRACE_LOG(Memory, ReallocAlloc, MemAllocChannel)
+				<< ReallocAlloc.CallstackId(CallstackId)
+				<< ReallocAlloc.Address(uint64(Address))
+				<< ReallocAlloc.RootHeap(uint8(RootHeap))
+				<< ReallocAlloc.Size(uint32(Size >> SizeShift))
+				<< ReallocAlloc.AlignmentPow2_SizeLower(uint8(Alignment_SizeLower));
+			break;
+		}
+	}
+
+	MemoryTrace_UpdateInternal();
+}
+
+////////////////////////////////////////////////////////////////////////////////
+void MemoryTrace_ReallocFree(uint64 Address, HeapId RootHeap)
+{
+	check(RootHeap < 16);
+	if (!GTraceAllowed)
+	{
+		return;
+	}
+
+	switch (RootHeap)
+	{
+		case EMemoryTraceRootHeap::SystemMemory:
+		{
+			UE_TRACE_LOG(Memory, ReallocFreeSystem, MemAllocChannel)
+				<< ReallocFreeSystem.Address(uint64(Address));
+			break;
+		}
+
+		default:
+		{
+			UE_TRACE_LOG(Memory, ReallocFree, MemAllocChannel)
+				<< ReallocFree.Address(uint64(Address))
+				<< ReallocFree.RootHeap(uint8(RootHeap));
+			break;
+		}
+	}
+
+	MemoryTrace_UpdateInternal();
+}
+
+////////////////////////////////////////////////////////////////////////////////
+HeapId MemoryTrace_HeapSpec(HeapId ParentId, const TCHAR* Name, EMemoryTraceHeapFlags Flags) 
+{
+	if (!GTraceAllowed)
+	{
+		return 0;
+	}
+	
+	static std::atomic<HeapId> HeapIdCount(EMemoryTraceRootHeap::EndReserved + 1); //Reserve indexes for root heaps
+	const HeapId Id = HeapIdCount.fetch_add(1);
+	const uint32 NameLen = FCString::Strlen(Name);
+	const uint32 DataSize = NameLen * sizeof(TCHAR);
+	check(ParentId < Id);
+
+	UE_TRACE_LOG(Memory, HeapSpec, MemAllocChannel, DataSize)
+		<< HeapSpec.Id(Id)
+		<< HeapSpec.ParentId(ParentId)
+		<< HeapSpec.Name(Name, NameLen)
+		<< HeapSpec.Flags(uint16(Flags));
+
+	return Id;
+}
+
+////////////////////////////////////////////////////////////////////////////////
+HeapId MemoryTrace_RootHeapSpec(const TCHAR* Name, EMemoryTraceHeapFlags Flags)
+{
+	if (!GTraceAllowed)
+	{
+		return 0;
+	}
+	
+	static std::atomic<HeapId> RootHeapCount(0);
+	const HeapId Id = RootHeapCount.fetch_add(1);
+	check(Id <= EMemoryTraceRootHeap::EndReserved);
+
+	const uint32 NameLen = FCString::Strlen(Name);
+	const uint32 DataSize = NameLen * sizeof(TCHAR);
+
+	UE_TRACE_LOG(Memory, HeapSpec, MemAllocChannel, DataSize)
+		<< HeapSpec.Id(Id)
+		<< HeapSpec.ParentId(HeapId(~0))
+		<< HeapSpec.Name(Name, NameLen)
+		<< HeapSpec.Flags(uint16(EMemoryTraceHeapFlags::Root|Flags));
+
+	return Id;
+}
+
+////////////////////////////////////////////////////////////////////////////////
+void MemoryTrace_MarkAllocAsHeap(uint64 Address, HeapId Heap, EMemoryTraceHeapAllocationFlags Flags)
+{
+	if (!GTraceAllowed)
+	{
+		return;
+	}
+	
+	UE_TRACE_LOG(Memory, HeapMarkAlloc, MemAllocChannel)
+		<< HeapMarkAlloc.Address(uint64(Address))
+		<< HeapMarkAlloc.Heap(Heap)
+		<< HeapMarkAlloc.Flags(uint16(EMemoryTraceHeapAllocationFlags::Heap | Flags));
+}
+
+////////////////////////////////////////////////////////////////////////////////
+void MemoryTrace_UnmarkAllocAsHeap(uint64 Address, HeapId Heap)
+{
+	if (!GTraceAllowed)
+	{
+		return;
+	}
+	
+	// Sets all flags to zero
+	UE_TRACE_LOG(Memory, HeapUnmarkAlloc, MemAllocChannel)
+		<< HeapUnmarkAlloc.Address(uint64(Address))
+		<< HeapUnmarkAlloc.Heap(Heap);
+}
+
+#endif // UE_MEMORY_TRACE_ENABLED
+
+
+/////////////////////////////////////////////////////////////////////////////
+FTraceMalloc::FTraceMalloc(FMalloc* InMalloc)
+{
+	WrappedMalloc = InMalloc;
+}
+
+/////////////////////////////////////////////////////////////////////////////
+FTraceMalloc::~FTraceMalloc()
+{
+}
+
+/////////////////////////////////////////////////////////////////////////////
+void* FTraceMalloc::Malloc(SIZE_T Count, uint32 Alignment)
+{
+#if UE_MEMORY_TRACE_ENABLED
+	void* NewPtr;
+	{
+		TGuardValue<bool> _(GDoNotAllocateInTrace, true);
+		NewPtr = WrappedMalloc->Malloc(Count, Alignment);
+	}
+
+	const uint64 Size = Count;
+	const uint32 AlignmentPow2 = uint32(FPlatformMath::CountTrailingZeros(Alignment));
+	const uint32 Alignment_SizeLower = (AlignmentPow2 << SizeShift) | uint32(Size & ((1 << SizeShift) - 1));
+
+	UE_MEMSCOPE(TRACE_TAG);
+	
+	UE_TRACE_LOG(Memory, Alloc, MemAllocChannel)
+		<< Alloc.CallstackId(0)
+		<< Alloc.Address(uint64(NewPtr))
+		<< Alloc.RootHeap(uint8(EMemoryTraceRootHeap::SystemMemory))
+		<< Alloc.Size(uint32(Size >> SizeShift))
+		<< Alloc.AlignmentPow2_SizeLower(uint8(Alignment_SizeLower));
+	
+	return NewPtr;
+#else
+	return WrappedMalloc->Malloc(Count, Alignment);
+#endif //UE_MEMORY_TRACE_ENABLED
+}
+
+/////////////////////////////////////////////////////////////////////////////
+void* FTraceMalloc::Realloc(void* Original, SIZE_T Count, uint32 Alignment)
+{
+#if UE_MEMORY_TRACE_ENABLED
+	void* NewPtr = nullptr;
+	const uint64 Size = Count;
+	const uint32 AlignmentPow2 = uint32(FPlatformMath::CountTrailingZeros(Alignment));
+	const uint32 Alignment_SizeLower = (AlignmentPow2 << SizeShift) | uint32(Size & ((1 << SizeShift) - 1));
+
+	UE_MEMSCOPE(TRACE_TAG);
+	
+	UE_TRACE_LOG(Memory, ReallocFree, MemAllocChannel)
+		<< ReallocFree.Address(uint64(Original))
+		<< ReallocFree.RootHeap(uint8(EMemoryTraceRootHeap::SystemMemory));
+	
+	{
+		TGuardValue<bool> _(GDoNotAllocateInTrace, true);
+		NewPtr = WrappedMalloc->Realloc(Original, Count, Alignment);
+	}
+	
+	UE_TRACE_LOG(Memory, ReallocAlloc, MemAllocChannel)
+		<< ReallocAlloc.CallstackId(0)
+		<< ReallocAlloc.Address(uint64(NewPtr))
+		<< ReallocAlloc.RootHeap(uint8(EMemoryTraceRootHeap::SystemMemory))
+		<< ReallocAlloc.Size(uint32(Size >> SizeShift))
+		<< ReallocAlloc.AlignmentPow2_SizeLower(uint8(Alignment_SizeLower));
+	
+	return NewPtr;
+#else
+	return WrappedMalloc->Realloc(Original, Count, Alignment);
+#endif //UE_MEMORY_TRACE_ENABLED
+}
+
+/////////////////////////////////////////////////////////////////////////////
+void FTraceMalloc::Free(void* Original)
+{
+#if UE_MEMORY_TRACE_ENABLED
+	UE_TRACE_LOG(Memory, Free, MemAllocChannel)
+		<< Free.Address(uint64(Original))
+		<< Free.RootHeap(uint8(EMemoryTraceRootHeap::SystemMemory));
+	
+	{
+		TGuardValue<bool> _(GDoNotAllocateInTrace, true);
+		WrappedMalloc->Free(Original);
+	}
+#else
+	WrappedMalloc->Free(Original);
+#endif
+}
+
+/////////////////////////////////////////////////////////////////////////////
+bool FTraceMalloc::ShouldTrace()
+{
+#if UE_MEMORY_TRACE_ENABLED
+	return !GDoNotAllocateInTrace;
+#else
+	return true;
+#endif
+}
diff --git a/Engine/Source/Runtime/Core/Private/ProfilingDebugging/Microsoft/MemoryTrace.cpp b/Engine/Source/Runtime/Core/Private/ProfilingDebugging/Microsoft/MemoryTrace.cpp
deleted file mode 100644
diff --git a/Engine/Source/Runtime/Core/Private/ProfilingDebugging/Microsoft/WindowsMemoryTrace.cpp b/Engine/Source/Runtime/Core/Private/ProfilingDebugging/Microsoft/WindowsMemoryTrace.cpp
new file mode 100644
--- /dev/null
+++ b/Engine/Source/Runtime/Core/Private/ProfilingDebugging/Microsoft/WindowsMemoryTrace.cpp
@@ -0,0 +1,344 @@
+// Copyright Epic Games, Inc. All Rights Reserved.
+
+#include "ProfilingDebugging/MemoryTrace.h"
+
+#include "ProfilingDebugging/CallstackTrace.h"
+#include "ProfilingDebugging/TraceMalloc.h"
+
+#if UE_MEMORY_TRACE_ENABLED
+
+#include "Containers/StringView.h"
+#include "CoreTypes.h"
+#include "HAL/MemoryBase.h"
+#include "HAL/Platform.h"
+#include "HAL/PlatformTime.h"
+#include "Misc/CoreDelegates.h"
+#include "Misc/CString.h"
+#include "Trace/Trace.inl"
+
+#include "Windows/AllowWindowsPlatformTypes.h"
+#include <winternl.h>
+#include <intrin.h>
+
+////////////////////////////////////////////////////////////////////////////////
+FMalloc* MemoryTrace_CreateInternal(FMalloc*);
+
+////////////////////////////////////////////////////////////////////////////////
+struct FAddrPack
+{
+			FAddrPack() = default;
+			FAddrPack(UPTRINT Addr, uint16 Value) { Set(Addr, Value); }
+	void	Set(UPTRINT Addr, uint16 Value) { Inner = uint64(Addr) | (uint64(Value) << 48ull); }
+	uint64	Inner;
+};
+static_assert(sizeof(FAddrPack) == sizeof(uint64), "");
+
+////////////////////////////////////////////////////////////////////////////////
+#if defined(PLATFORM_SUPPORTS_TRACE_WIN32_VIRTUAL_MEMORY_HOOKS)
+class FTextSectionEditor
+{
+public:
+							FTextSectionEditor(void* InBase);
+							~FTextSectionEditor();
+	template <typename T>
+	T*						Hook(T* Target, T* HookFunction);
+
+private:
+	static void*			GetActualAddress(void* Function);
+	uint8*					AllocateTrampoline(unsigned int Size);
+	void*					HookImpl(void* Target, void* HookFunction);
+	uint8*					TrampolineTail;
+	void*					Base;
+	size_t					Size;
+	DWORD					Protection;
+};
+
+////////////////////////////////////////////////////////////////////////////////
+FTextSectionEditor::FTextSectionEditor(void* InBase)
+{
+	InBase = GetActualAddress(InBase);
+
+	MEMORY_BASIC_INFORMATION MemInfo;
+	VirtualQuery(InBase, &MemInfo, sizeof(MemInfo));
+	Base = MemInfo.BaseAddress;
+	Size = MemInfo.RegionSize;
+
+	VirtualProtect(Base, Size, PAGE_EXECUTE_READWRITE, &Protection);
+
+	TrampolineTail = (uint8*)Base + Size;
+}
+
+////////////////////////////////////////////////////////////////////////////////
+FTextSectionEditor::~FTextSectionEditor()
+{
+	VirtualProtect(Base, Size, Protection, &Protection);
+	FlushInstructionCache(GetCurrentProcess(), Base, Size);
+}
+
+////////////////////////////////////////////////////////////////////////////////
+void* FTextSectionEditor::GetActualAddress(void* Function)
+{
+	uint8* Addr = (uint8*)Function;
+	int Offset = unsigned(Addr[0] & 0xf0) == 0x40; // REX prefix
+	if (Addr[Offset + 0] == 0xff && Addr[Offset + 1] == 0x25)
+	{
+		Addr += Offset;
+		Addr = *(uint8**)(Addr + 6 + *(uint32*)(Addr + 2));
+	}
+	return Addr;
+}
+
+////////////////////////////////////////////////////////////////////////////////
+uint8* FTextSectionEditor::AllocateTrampoline(unsigned int PatchSize)
+{
+	static const int TrampolineSize = 24;
+	uint8* NextTail = TrampolineTail - TrampolineSize;
+	do 
+	{
+		check(*NextTail == 0);
+		++NextTail;
+	}
+	while (NextTail != TrampolineTail);
+
+	TrampolineTail -= TrampolineSize;
+	return TrampolineTail;
+}
+
+////////////////////////////////////////////////////////////////////////////////
+template <typename T>
+T* FTextSectionEditor::Hook(T* Target, T* HookFunction)
+{
+	return (T*)HookImpl((void*)Target, (void*)HookFunction);
+}
+
+////////////////////////////////////////////////////////////////////////////////
+void* FTextSectionEditor::HookImpl(void* Target, void* HookFunction)
+{
+	Target = GetActualAddress(Target);
+
+	const uint8* Start = (uint8*)Target;
+	const uint8* Read = Start;
+	do
+	{
+		Read += (Read[0] & 0xf0) == 0x40; // REX prefix
+		uint8 Inst = *Read++;
+		if (unsigned(Inst - 0x80) < 0x0cu)
+		{
+			uint8 ModRm = *Read++;
+			Read += ((ModRm & 0300) < 0300) & ((ModRm & 0007) == 0004); // SIB
+			switch (ModRm & 0300) // Disp[8|32]
+			{
+			case 0100: Read += 1; break;
+			case 0200: Read += 5; break;
+			}
+			Read += (Inst == 0x83);
+		}
+		else if (unsigned(Inst - 0x50) >= 0x10u)
+			check(!"Unknown instruction");
+	}
+	while (Read - Start < 6);
+
+	int PatchSize = int(Read - Start);
+	uint8* TrampolinePtr = AllocateTrampoline(PatchSize);
+
+	*(void**)TrampolinePtr = HookFunction;
+
+	uint8* PatchJmp = TrampolinePtr + sizeof(void*);
+	memcpy(PatchJmp, Start, PatchSize);
+
+	PatchJmp += PatchSize;
+	*PatchJmp = 0xe9;
+	*(int32*)(PatchJmp + 1) = int32(intptr_t(Start + PatchSize) - intptr_t(PatchJmp)) - 5;
+
+	uint16* HookJmp = (uint16*)Target;
+	HookJmp[0] = 0x25ff;
+	*(int32*)(HookJmp + 1) = int32(intptr_t(TrampolinePtr) - intptr_t(HookJmp + 3));
+
+	return PatchJmp - PatchSize;
+}
+
+////////////////////////////////////////////////////////////////////////////////
+class FVirtualWinApiHooks
+{
+public:
+	static void				Initialize(bool bInLight);
+
+private:
+							FVirtualWinApiHooks();
+	static bool				bLight;
+	static LPVOID WINAPI	VmAlloc(LPVOID Address, SIZE_T Size, DWORD Type, DWORD Protect);
+	static LPVOID WINAPI	VmAllocEx(HANDLE Process, LPVOID Address, SIZE_T Size, DWORD Type, DWORD Protect);
+	static LPVOID WINAPI	VmAlloc2(HANDLE Process, LPVOID BaseAddress, SIZE_T Size, ULONG AllocationType, ULONG PageProtection, /*MEM_EXTENDED_PARAMETER* */ void* ExtendedParameters, ULONG ParameterCount);
+	static BOOL WINAPI		VmFree(LPVOID Address, SIZE_T Size, DWORD Type);
+	static BOOL WINAPI		VmFreeEx(HANDLE Process, LPVOID Address, SIZE_T Size, DWORD Type);
+	static LPVOID			(WINAPI *VmAllocOrig)(LPVOID, SIZE_T, DWORD, DWORD);
+	static LPVOID			(WINAPI *VmAllocExOrig)(HANDLE, LPVOID, SIZE_T, DWORD, DWORD);
+	static LPVOID			(WINAPI *VmAlloc2Orig)(HANDLE, LPVOID, SIZE_T, ULONG, ULONG, /*MEM_EXTENDED_PARAMETER* */ void*, ULONG);
+	static BOOL				(WINAPI *VmFreeOrig)(LPVOID, SIZE_T, DWORD);
+	static BOOL				(WINAPI *VmFreeExOrig)(HANDLE, LPVOID, SIZE_T, DWORD);
+
+	typedef LPVOID(__stdcall* FnVirtualAlloc2)(HANDLE, LPVOID, SIZE_T, ULONG, ULONG, /* MEM_EXTENDED_PARAMETER* */ void*, ULONG);
+};
+
+////////////////////////////////////////////////////////////////////////////////
+bool	FVirtualWinApiHooks::bLight;
+LPVOID	(WINAPI *FVirtualWinApiHooks::VmAllocOrig)(LPVOID, SIZE_T, DWORD, DWORD);
+LPVOID	(WINAPI *FVirtualWinApiHooks::VmAllocExOrig)(HANDLE, LPVOID, SIZE_T, DWORD, DWORD);
+LPVOID	(WINAPI *FVirtualWinApiHooks::VmAlloc2Orig)(HANDLE, LPVOID, SIZE_T, ULONG, ULONG, /*MEM_EXTENDED_PARAMETER* */ void*, ULONG);
+BOOL	(WINAPI *FVirtualWinApiHooks::VmFreeOrig)(LPVOID, SIZE_T, DWORD);
+BOOL	(WINAPI *FVirtualWinApiHooks::VmFreeExOrig)(HANDLE, LPVOID, SIZE_T, DWORD);
+
+////////////////////////////////////////////////////////////////////////////////
+void FVirtualWinApiHooks::Initialize(bool bInLight)
+{
+	bLight = bInLight;
+
+	{
+		FTextSectionEditor Editor((void*)VirtualAlloc);
+		VmAllocOrig = Editor.Hook(VirtualAlloc, &FVirtualWinApiHooks::VmAlloc);
+		VmFreeOrig = Editor.Hook(VirtualFree, &FVirtualWinApiHooks::VmFree); 
+		
+		{
+			FTextSectionEditor EditorEx((void*)VirtualAllocEx);
+			VmAllocExOrig = Editor.Hook(VirtualAllocEx, &FVirtualWinApiHooks::VmAllocEx);
+			VmFreeExOrig = Editor.Hook(VirtualFreeEx, &FVirtualWinApiHooks::VmFreeEx);
+		}
+
+#if PLATFORM_WINDOWS
+#if (NTDDI_VERSION >= NTDDI_WIN10_RS4)
+		{
+			FTextSectionEditor EditorAlloc2((void*)VirtualAlloc2);
+			VmAlloc2Orig = Editor.Hook(VirtualAlloc2, &FVirtualWinApiHooks::VmAlloc2);
+		}
+#else // NTDDI_VERSION
+		{
+			VmAlloc2Orig = nullptr;
+			HINSTANCE DllInstance;
+			DllInstance = LoadLibrary(TEXT("kernelbase.dll"));
+			if (DllInstance != NULL)
+			{
+#pragma warning(push)
+#pragma warning(disable: 4191) // 'type cast': unsafe conversion from 'FARPROC' to 'FVirtualWinApiHooks::FnVirtualAlloc2'
+				VmAlloc2Orig = (FnVirtualAlloc2)GetProcAddress(DllInstance, "VirtualAlloc2");
+#pragma warning(pop)
+				FreeLibrary(DllInstance);
+			}
+			if (VmAlloc2Orig)
+			{
+				FTextSectionEditor EditorAlloc2((void*)VmAlloc2Orig);
+				VmAlloc2Orig = Editor.Hook(VmAlloc2Orig, &FVirtualWinApiHooks::VmAlloc2);
+			}
+		}
+#endif // NTDDI_VERSION
+#endif // PLATFORM_WINDOWS
+	}
+}
+
+////////////////////////////////////////////////////////////////////////////////
+LPVOID WINAPI FVirtualWinApiHooks::VmAlloc(LPVOID Address, SIZE_T Size, DWORD Type, DWORD Protect)
+{
+	LPVOID Ret = VmAllocOrig(Address, Size, Type, Protect);
+	// Track any reserve for now. Going forward we need events to differentiate reserves/commits and
+	// corresponding information on frees.
+	if (Ret != nullptr && (Type & MEM_RESERVE))
+	{
+		MemoryTrace_Alloc((uint64)Ret, Size, 0, EMemoryTraceRootHeap::SystemMemory);
+		MemoryTrace_MarkAllocAsHeap((uint64)Ret, EMemoryTraceRootHeap::SystemMemory);
+	}
+
+	return Ret;
+}
+
+////////////////////////////////////////////////////////////////////////////////
+BOOL WINAPI FVirtualWinApiHooks::VmFree(LPVOID Address, SIZE_T Size, DWORD Type)
+{
+	if (Type & MEM_RELEASE)
+	{
+		MemoryTrace_Free((uint64)Address, EMemoryTraceRootHeap::SystemMemory);
+	}
+	return VmFreeOrig(Address, Size, Type);
+}
+
+////////////////////////////////////////////////////////////////////////////////
+LPVOID WINAPI FVirtualWinApiHooks::VmAllocEx(HANDLE Process, LPVOID Address, SIZE_T Size, DWORD Type, DWORD Protect)
+{
+	LPVOID Ret = VmAllocExOrig(Process, Address, Size, Type, Protect);
+	if (Process == GetCurrentProcess() && Ret != nullptr && (Type & MEM_RESERVE))
+	{
+		MemoryTrace_Alloc((uint64)Ret, Size, 0);
+		MemoryTrace_MarkAllocAsHeap((uint64)Ret, EMemoryTraceRootHeap::SystemMemory);
+	}
+
+	return Ret;
+}
+
+////////////////////////////////////////////////////////////////////////////////
+BOOL WINAPI FVirtualWinApiHooks::VmFreeEx(HANDLE Process, LPVOID Address, SIZE_T Size, DWORD Type)
+{
+	if (Process == GetCurrentProcess() && (Type & MEM_RELEASE))
+	{
+		MemoryTrace_Free((uint64)Address, EMemoryTraceRootHeap::SystemMemory);
+	}
+
+	return VmFreeExOrig(Process, Address, Size, Type);
+}
+
+////////////////////////////////////////////////////////////////////////////////
+LPVOID WINAPI FVirtualWinApiHooks::VmAlloc2(HANDLE Process, LPVOID BaseAddress, SIZE_T Size, ULONG Type, ULONG PageProtection, /*MEM_EXTENDED_PARAMETER* */ void* ExtendedParameters, ULONG ParameterCount)
+{
+	LPVOID Ret = VmAlloc2Orig(Process, BaseAddress, Size, Type, PageProtection, ExtendedParameters, ParameterCount);
+	if (Process == GetCurrentProcess() && Ret != nullptr && (Type & MEM_RESERVE))
+	{
+		MemoryTrace_Alloc((uint64)Ret, Size, 0);
+		MemoryTrace_MarkAllocAsHeap((uint64)Ret, EMemoryTraceRootHeap::SystemMemory);
+	}
+
+	return Ret;
+}
+#endif // defined(PLATFORM_SUPPORTS_TRACE_WIN32_VIRTUAL_MEMORY_HOOKS)
+
+
+////////////////////////////////////////////////////////////////////////////////
+FMalloc* MemoryTrace_Create(FMalloc* InMalloc)
+{
+	bool bEnabled = false;
+	const TCHAR* CmdLine = ::GetCommandLineW();
+	if (const TCHAR* TraceArg = FCString::Strstr(CmdLine, TEXT("-trace=")))
+	{
+		const TCHAR* Start = TraceArg + 7;
+		const TCHAR* End = Start;
+		for (; *End != ' ' && *End != '\0'; ++End);
+
+		for (const TCHAR* c = Start; c < End + 1; ++c)
+		{
+			if (*c == ' ' || *c == '\0' || *c == ',')
+			{
+				FStringView View(Start, uint32(c - Start));
+				if (View.Equals(TEXT("memalloc"), ESearchCase::IgnoreCase) || View.Equals(TEXT("memory"), ESearchCase::IgnoreCase))
+				{
+					bEnabled = true;
+					break;
+				}
+
+				Start = c + 1;
+			}
+		}
+	}
+
+	if (bEnabled)
+	{
+		FMalloc* OutMalloc = MemoryTrace_CreateInternal(InMalloc);
+		
+	#if defined(PLATFORM_SUPPORTS_TRACE_WIN32_VIRTUAL_MEMORY_HOOKS)
+		FVirtualWinApiHooks::Initialize(false);
+	#endif // defined(PLATFORM_SUPPORTS_TRACE_WIN32_VIRTUAL_MEMORY_HOOKS)
+		
+		return OutMalloc;
+	}
+
+	return InMalloc;
+}
+
+#include "Windows/HideWindowsPlatformTypes.h"
+
+#endif // UE_MEMORY_TRACE_ENABLED
diff --git a/Engine/Source/Runtime/Core/Private/ProfilingDebugging/TraceAuxiliary.cpp b/Engine/Source/Runtime/Core/Private/ProfilingDebugging/TraceAuxiliary.cpp
--- a/Engine/Source/Runtime/Core/Private/ProfilingDebugging/TraceAuxiliary.cpp
+++ b/Engine/Source/Runtime/Core/Private/ProfilingDebugging/TraceAuxiliary.cpp
@@ -1251,10 +1251,6 @@
 	}
 #endif
 
-#if UE_MEMORY_TRACE_ENABLED
-	MemoryTrace_InitializeLate();
-#endif
-
 	// Trace out information about this session. This is done before initialisation,
 	// so that it is always sent (all channels are enabled prior to initialisation).
 	const TCHAR* BranchName = BuildSettings::GetBranchName();
@@ -1300,6 +1296,10 @@
 	}
 	UE::Trace::Initialize(Desc);
 
+	// Workaround for the fact that even if StartFromCommandlineArguments will enable channels
+	// specified by the commandline, UE::Trace::Initialize will reset all channels.
+	GTraceAuxiliary.EnableCommandlineChannelsPostInitialize();
+	
 	// Always register end frame updates. This path is short circuited if a worker thread exists.
 	GTraceAuxiliary.StartEndFramePump();
 	if (!FForkProcessHelper::IsForkRequested())
@@ -1329,10 +1329,6 @@
 	});
 
 	UE::Trace::ThreadRegister(TEXT("GameThread"), FPlatformTLS::GetCurrentThreadId(), -1);
-	
-	// Workaround for the fact that even if StartFromCommandlineArguments will enable channels
-	// specified by the commandline, UE::Trace::Initialize will reset all channels.
-	GTraceAuxiliary.EnableCommandlineChannelsPostInitialize();
 #endif
 }
 
diff --git a/Engine/Source/Runtime/Core/Public/ProfilingDebugging/MemoryTrace.h b/Engine/Source/Runtime/Core/Public/ProfilingDebugging/MemoryTrace.h
--- a/Engine/Source/Runtime/Core/Public/ProfilingDebugging/MemoryTrace.h
+++ b/Engine/Source/Runtime/Core/Public/ProfilingDebugging/MemoryTrace.h
@@ -66,12 +66,6 @@ CORE_API UE_TRACE_CHANNEL_EXTERN(MemAllocChannel);
 class FMalloc* MemoryTrace_Create(class FMalloc* InMalloc);
 void MemoryTrace_Initialize();
 
-#if UE_MEMORY_TRACE_LATE_INIT
-void MemoryTrace_InitializeLate();
-#else
-static void MemoryTrace_InitializeLate() {}
-#endif // UE_MEMORY_TRACE_LATE_INIT
-
 /**
  * Register a new heap specification (name). Use the returned value when marking heaps.
  * @param ParentId Heap id of parent heap.
diff --git a/Engine/Source/Runtime/Core/Public/ProfilingDebugging/TraceMalloc.h b/Engine/Source/Runtime/Core/Public/ProfilingDebugging/TraceMalloc.h
--- a/Engine/Source/Runtime/Core/Public/ProfilingDebugging/TraceMalloc.h
+++ b/Engine/Source/Runtime/Core/Public/ProfilingDebugging/TraceMalloc.h
@@ -13,6 +13,7 @@ public:
 	virtual void* Realloc(void* Original, SIZE_T Count, uint32 Alignment) override;
 	virtual void Free(void* Original) override;
 
+	UE_DEPRECATED(5.1, "ShouldTrace is deprecated. Details should be internal now.")
 	static bool ShouldTrace();
 
 	virtual void OnPreFork() override
