Trace: Enable Default Trace Channels automatically for Editor

### Features
Prior to this change there are no trace channels active for any Editor process, and default channels were accessed by passing -trace or -trace=default.

This change enables default channels automatically for Editor.

diff --git a/Engine/Source/Runtime/Core/Private/ProfilingDebugging/TraceAuxiliary.cpp b/Engine/Source/Runtime/Core/Private/ProfilingDebugging/TraceAuxiliary.cpp
--- a/Engine/Source/Runtime/Core/Private/ProfilingDebugging/TraceAuxiliary.cpp
+++ b/Engine/Source/Runtime/Core/Private/ProfilingDebugging/TraceAuxiliary.cpp
@@ -1093,6 +1093,12 @@ static bool StartFromCommandlineArguments(const TCHAR* CommandLine)
 	{
 		Channels = GDefaultChannels;
 	}
+#if WITH_EDITOR
+	else
+	{
+		Channels = GDefaultChannels;
+	}
+#endif
 
 	// By default, if any channels are enabled we trace to memory.
 	FTraceAuxiliary::EConnectionType Type = FTraceAuxiliary::EConnectionType::None;
