Remove macOS specific logic in FPaths::GeneratedConfigDir().   Using ~/library/preferences/UnrealEngine folder based on project name can lead to clashes with different projects using the same name.  Use the same logic as other platforms and store generated config ini files in the project/saved/ folder.  This also will fix the minor issue of leaving lots of per project config files around in ~/library/preferences/UnrealEngine when you delete a project.

Tested multiple blueprint and C++ projects with the same name in different directories; running the editor + PIE + cooking & packing.

diff --git a/Engine/Source/Runtime/Core/Private/Misc/Paths.cpp b/Engine/Source/Runtime/Core/Private/Misc/Paths.cpp
--- a/Engine/Source/Runtime/Core/Private/Misc/Paths.cpp
+++ b/Engine/Source/Runtime/Core/Private/Misc/Paths.cpp
@@ -395,11 +395,7 @@ FString FPaths::SourceConfigDir()
 
 FString FPaths::GeneratedConfigDir()
 {
-#if PLATFORM_MAC // @todo, move this to Mac FPlatformMisc.
-	return FPlatformProcess::UserPreferencesDir();
-#else
 	return FPlatformMisc::GeneratedConfigDir();
-#endif
 }
 
 FString FPaths::SandboxesDir()
