releasing FIoBatchImpl::Callback right after it's called to avoid a long living allocation that keeps a larger memory block from being returned to the system

diff --git a/Engine/Source/Runtime/Core/Private/IO/IoDispatcher.cpp b/Engine/Source/Runtime/Core/Private/IO/IoDispatcher.cpp
--- a/Engine/Source/Runtime/Core/Private/IO/IoDispatcher.cpp
+++ b/Engine/Source/Runtime/Core/Private/IO/IoDispatcher.cpp
@@ -676,6 +676,7 @@ private:
 				Result = Status;
 			}
 			Request->Callback(Result);
+			Request->Callback = {};
 		}
 		if (Batch)
 		{
