Avoid map reallocation for AutomationTest filters

#rnx

diff --git a/Engine/Source/Runtime/Core/Private/Misc/AutomationTest.cpp b/Engine/Source/Runtime/Core/Private/Misc/AutomationTest.cpp
--- a/Engine/Source/Runtime/Core/Private/Misc/AutomationTest.cpp
+++ b/Engine/Source/Runtime/Core/Private/Misc/AutomationTest.cpp
@@ -1209,9 +1209,10 @@ uint32 FAutomationTestBase::ExtractAutomationTestFlags(FString InTagNotation)
 		.ParseIntoArray(OutputParts, TEXT(";"), true);
 	for (auto it = OutputParts.begin(); it != OutputParts.end(); ++it)
 	{
-		if (FStringToEAutomationTestFlagMap.Contains(*it))
+		auto Value = EAutomationTestFlags::FromString(*it);
+		if (Value != EAutomationTestFlags::None)
 		{
-			Result |= FStringToEAutomationTestFlagMap[*it];
+			Result |= Value;
 		}
 	}
 	return Result;
diff --git a/Engine/Source/Runtime/Core/Public/Misc/AutomationTest.h b/Engine/Source/Runtime/Core/Public/Misc/AutomationTest.h
--- a/Engine/Source/Runtime/Core/Public/Misc/AutomationTest.h
+++ b/Engine/Source/Runtime/Core/Public/Misc/AutomationTest.h
@@ -49,12 +49,13 @@
 
 /**
 * Flags for specifying automation test requirements/behavior
-* Update FStringToEAutomationTestFlagMap when updating this enum.
+* Update GetTestFlagsMap when updating this enum.
 */
-namespace EAutomationTestFlags
+struct EAutomationTestFlags
 {
 	enum Type
 	{
+		None = 0x00000000,
 		//~ Application context required for the test
 		// Test is suitable for running within the editor
 		EditorContext = 0x00000001,
@@ -107,6 +108,49 @@ namespace EAutomationTestFlags
 		NegativeFilter				= 0x20000000,
 		FilterMask = SmokeFilter | EngineFilter | ProductFilter | PerfFilter | StressFilter | NegativeFilter
 	};
+
+	static const TMap<FString, Type>& GetTestFlagsMap()
+	{
+		LLM_SCOPE_BYNAME(TEXT("AutomationTest"));
+		/** String to EAutomationTestFlags map */
+		static const TMap<FString, Type> FlagsMap = {
+			{ TEXT("EditorContext"), Type::EditorContext},
+			{ TEXT("ClientContext"), Type::ClientContext},
+			{ TEXT("ServerContext"), Type::ServerContext},
+			{ TEXT("CommandletContext"), Type::CommandletContext},
+			{ TEXT("ApplicationContextMask"), Type::ApplicationContextMask},
+			{ TEXT("NonNullRHI"), Type::NonNullRHI},
+			{ TEXT("RequiresUser"), Type::RequiresUser},
+			{ TEXT("FeatureMask"), Type::FeatureMask},
+			{ TEXT("Disabled"), Type::Disabled},
+			{ TEXT("CriticalPriority"), Type::CriticalPriority},
+			{ TEXT("HighPriority"), Type::HighPriority},
+			{ TEXT("HighPriorityAndAbove"), Type::HighPriorityAndAbove},
+			{ TEXT("MediumPriority"), Type::MediumPriority},
+			{ TEXT("MediumPriorityAndAbove"), Type::MediumPriorityAndAbove},
+			{ TEXT("LowPriority"), Type::LowPriority},
+			{ TEXT("PriorityMask"), Type::PriorityMask},
+			{ TEXT("SmokeFilter"), Type::SmokeFilter},
+			{ TEXT("EngineFilter"), Type::EngineFilter},
+			{ TEXT("ProductFilter"), Type::ProductFilter},
+			{ TEXT("PerfFilter"), Type::PerfFilter},
+			{ TEXT("SmokeFilter"), Type::SmokeFilter},
+			{ TEXT("StressFilter"), Type::StressFilter},
+			{ TEXT("NegativeFilter"), Type::NegativeFilter},
+			{ TEXT("FilterMask"), Type::FilterMask}
+		};
+		return FlagsMap;
+	}
+
+	static const Type FromString(FString Name)
+	{
+		static auto FlagMap = GetTestFlagsMap();
+		if (FlagMap.Contains(Name))
+		{
+			return FlagMap[Name];
+		}
+		return Type::None;
+	}
 };
 
 /** Flags for indicating the matching type to use for an expected error */
@@ -1793,34 +1837,6 @@ protected:
 	/** Extracts a combined EAutomationTestFlags value from a string representation using tag notation "[Filter_1]...[Filter_n][Tag_1]...[Tag_m]" */
 	uint32 ExtractAutomationTestFlags(FString InTagNotation);
 
-	/** String to EAutomationTestFlags map */
-	const TMap<FString, uint32> FStringToEAutomationTestFlagMap = {
-		{ TEXT("EditorContext"), EAutomationTestFlags::EditorContext},
-		{ TEXT("ClientContext"), EAutomationTestFlags::ClientContext},
-		{ TEXT("ServerContext"), EAutomationTestFlags::ServerContext},
-		{ TEXT("CommandletContext"), EAutomationTestFlags::CommandletContext},
-		{ TEXT("ApplicationContextMask"), EAutomationTestFlags::ApplicationContextMask},
-		{ TEXT("NonNullRHI"), EAutomationTestFlags::NonNullRHI},
-		{ TEXT("RequiresUser"), EAutomationTestFlags::RequiresUser},
-		{ TEXT("FeatureMask"), EAutomationTestFlags::FeatureMask},
-		{ TEXT("Disabled"), EAutomationTestFlags::Disabled},
-		{ TEXT("CriticalPriority"), EAutomationTestFlags::CriticalPriority},
-		{ TEXT("HighPriority"), EAutomationTestFlags::HighPriority},
-		{ TEXT("HighPriorityAndAbove"), EAutomationTestFlags::HighPriorityAndAbove},
-		{ TEXT("MediumPriority"), EAutomationTestFlags::MediumPriority},
-		{ TEXT("MediumPriorityAndAbove"), EAutomationTestFlags::MediumPriorityAndAbove},
-		{ TEXT("LowPriority"), EAutomationTestFlags::LowPriority},
-		{ TEXT("PriorityMask"), EAutomationTestFlags::PriorityMask},
-		{ TEXT("SmokeFilter"), EAutomationTestFlags::SmokeFilter},
-		{ TEXT("EngineFilter"), EAutomationTestFlags::EngineFilter},
-		{ TEXT("ProductFilter"), EAutomationTestFlags::ProductFilter},
-		{ TEXT("PerfFilter"), EAutomationTestFlags::PerfFilter},
-		{ TEXT("SmokeFilter"), EAutomationTestFlags::SmokeFilter},
-		{ TEXT("StressFilter"), EAutomationTestFlags::StressFilter},
-		{ TEXT("NegativeFilter"), EAutomationTestFlags::NegativeFilter},
-		{ TEXT("FilterMask"), EAutomationTestFlags::FilterMask}
-	};
-
 protected:
 
 	//Flag to indicate if this is a complex task
