CompactBinary: Replaced local ReadUnaligned function with FPlatformMemory::ReadUnaligned

#rnx

diff --git a/Engine/Source/Runtime/Core/Private/Serialization/CompactBinaryValidation.cpp b/Engine/Source/Runtime/Core/Private/Serialization/CompactBinaryValidation.cpp
--- a/Engine/Source/Runtime/Core/Private/Serialization/CompactBinaryValidation.cpp
+++ b/Engine/Source/Runtime/Core/Private/Serialization/CompactBinaryValidation.cpp
@@ -11,25 +11,6 @@
 
 ///////////////////////////////////////////////////////////////////////////////////////////////////////////////////////
 
-namespace CompactBinaryValidationPrivate
-{
-
-template <typename T>
-static constexpr FORCEINLINE T ReadUnaligned(const void* const Memory)
-{
-#if PLATFORM_SUPPORTS_UNALIGNED_LOADS
-	return *static_cast<const T*>(Memory);
-#else
-	T Value;
-	FMemory::Memcpy(&Value, Memory, sizeof(Value));
-	return Value;
-#endif
-}
-
-}
-
-///////////////////////////////////////////////////////////////////////////////////////////////////////////////////////
-
 /**
  * Adds the given error(s) to the error mask.
  *
@@ -112,7 +93,7 @@ static void ValidateCbFloat64(FMemoryView& View, ECbValidateMode Mode, ECbValida
 	{
 		if (EnumHasAnyFlags(Mode, ECbValidateMode::Format))
 		{
-			const uint64 RawValue = NETWORK_ORDER64(CompactBinaryValidationPrivate::ReadUnaligned<uint64>(View.GetData()));
+			const uint64 RawValue = NETWORK_ORDER64(FPlatformMemory::ReadUnaligned<uint64>(View.GetData()));
 			const double Value = reinterpret_cast<const double&>(RawValue);
 			if (Value == double(float(Value)))
 			{
diff --git a/Engine/Source/Runtime/Core/Public/Serialization/CompactBinaryValue.h b/Engine/Source/Runtime/Core/Public/Serialization/CompactBinaryValue.h
--- a/Engine/Source/Runtime/Core/Public/Serialization/CompactBinaryValue.h
+++ b/Engine/Source/Runtime/Core/Public/Serialization/CompactBinaryValue.h
@@ -11,23 +11,6 @@
 #include "Serialization/CompactBinary.h"
 #include "Serialization/VarInt.h"
 
-namespace UE::CompactBinary::Private
-{
-
-template <typename T>
-static constexpr FORCEINLINE T ReadUnaligned(const void* const Memory)
-{
-#if PLATFORM_SUPPORTS_UNALIGNED_LOADS
-	return *static_cast<const T*>(Memory);
-#else
-	T Value;
-	FMemory::Memcpy(&Value, Memory, sizeof(Value));
-	return Value;
-#endif
-}
-
-} // UE::CompactBinary::Private
-
 /**
  * A type that provides unchecked access to compact binary values.
  *
@@ -176,13 +159,13 @@ FORCEINLINE int64 FCbValue::AsIntegerNegative() const
 
 FORCEINLINE float FCbValue::AsFloat32() const
 {
-	const uint32 Value = NETWORK_ORDER32(UE::CompactBinary::Private::ReadUnaligned<uint32>(Data));
+	const uint32 Value = NETWORK_ORDER32(FPlatformMemory::ReadUnaligned<uint32>(Data));
 	return reinterpret_cast<const float&>(Value);
 }
 
 FORCEINLINE double FCbValue::AsFloat64() const
 {
-	const uint64 Value = NETWORK_ORDER64(UE::CompactBinary::Private::ReadUnaligned<uint64>(Data));
+	const uint64 Value = NETWORK_ORDER64(FPlatformMemory::ReadUnaligned<uint64>(Data));
 	return reinterpret_cast<const double&>(Value);
 }
 
@@ -209,12 +192,12 @@ FORCEINLINE FGuid FCbValue::AsUuid() const
 
 FORCEINLINE int64 FCbValue::AsDateTimeTicks() const
 {
-	return NETWORK_ORDER64(UE::CompactBinary::Private::ReadUnaligned<int64>(Data));
+	return NETWORK_ORDER64(FPlatformMemory::ReadUnaligned<int64>(Data));
 }
 
 FORCEINLINE int64 FCbValue::AsTimeSpanTicks() const
 {
-	return NETWORK_ORDER64(UE::CompactBinary::Private::ReadUnaligned<int64>(Data));
+	return NETWORK_ORDER64(FPlatformMemory::ReadUnaligned<int64>(Data));
 }
 
 FORCEINLINE FCbObjectId FCbValue::AsObjectId() const
