GitHub #9202 : Semantic fix

Replaced nullptr_t with TYPE_OF_NULLPTR to prevent error reporting in IDE's utilizing clangd like QT Creator

diff --git a/Engine/Source/Runtime/Core/Public/Async/ParallelFor.h b/Engine/Source/Runtime/Core/Public/Async/ParallelFor.h
--- a/Engine/Source/Runtime/Core/Public/Async/ParallelFor.h
+++ b/Engine/Source/Runtime/Core/Public/Async/ParallelFor.h
@@ -58,7 +58,7 @@ namespace ParallelForImpl
 
 	// Helper specialization for "no context", which changes the assumed body call signature
 	template <typename FunctionType>
-	inline void CallBody(const FunctionType& Body, const TArrayView<nullptr_t>&, int32, int32 Index)
+	inline void CallBody(const FunctionType& Body, const TArrayView<TYPE_OF_NULLPTR>&, int32, int32 Index)
 	{
 		Body(Index);
 	}
@@ -388,7 +388,7 @@ inline void ParallelFor(int32 Num, TFunctionRef<void(int32)> Body, bool bForceSi
 {
 	ParallelForImpl::ParallelForInternal(TEXT("ParallelFor"), Num, 1, Body, [](){},
 		(bForceSingleThread ? EParallelForFlags::ForceSingleThread : EParallelForFlags::None) | 
-		(bPumpRenderingThread ? EParallelForFlags::PumpRenderingThread : EParallelForFlags::None), TArrayView<nullptr_t>());
+		(bPumpRenderingThread ? EParallelForFlags::PumpRenderingThread : EParallelForFlags::None), TArrayView<TYPE_OF_NULLPTR>());
 }
 
 /**
@@ -401,7 +401,7 @@ inline void ParallelFor(int32 Num, TFunctionRef<void(int32)> Body, bool bForceSi
 template<typename FunctionType>
 inline void ParallelForTemplate(int32 Num, const FunctionType& Body, EParallelForFlags Flags = EParallelForFlags::None)
 {
-	ParallelForImpl::ParallelForInternal(TEXT("ParallelFor"), Num, 1, Body, [](){}, Flags, TArrayView<nullptr_t>());
+	ParallelForImpl::ParallelForInternal(TEXT("ParallelFor"), Num, 1, Body, [](){}, Flags, TArrayView<TYPE_OF_NULLPTR>());
 }
 
 /**
@@ -416,7 +416,7 @@ inline void ParallelForTemplate(int32 Num, const FunctionType& Body, EParallelFo
 template<typename FunctionType>
 inline void ParallelForTemplate(const TCHAR* DebugName, int32 Num, int32 MinBatchSize, const FunctionType& Body, EParallelForFlags Flags = EParallelForFlags::None)
 {
-	ParallelForImpl::ParallelForInternal(DebugName, Num, MinBatchSize, Body, [](){}, Flags, TArrayView<nullptr_t>());
+	ParallelForImpl::ParallelForInternal(DebugName, Num, MinBatchSize, Body, [](){}, Flags, TArrayView<TYPE_OF_NULLPTR>());
 }
 
 /** 
@@ -431,7 +431,7 @@ inline void ParallelForTemplate(const TCHAR* DebugName, int32 Num, int32 MinBatc
 **/
 inline void ParallelFor(int32 Num, TFunctionRef<void(int32)> Body, EParallelForFlags Flags = EParallelForFlags::None)
 {
-	ParallelForImpl::ParallelForInternal(TEXT("ParallelFor"), Num, 1, Body, [](){}, Flags, TArrayView<nullptr_t>());
+	ParallelForImpl::ParallelForInternal(TEXT("ParallelFor"), Num, 1, Body, [](){}, Flags, TArrayView<TYPE_OF_NULLPTR>());
 }
 
 /** 
@@ -448,7 +448,7 @@ inline void ParallelFor(int32 Num, TFunctionRef<void(int32)> Body, EParallelForF
 **/
 inline void ParallelFor(const TCHAR* DebugName, int32 Num, int32 MinBatchSize, TFunctionRef<void(int32)> Body, EParallelForFlags Flags = EParallelForFlags::None)
 {
-	ParallelForImpl::ParallelForInternal(DebugName, Num, MinBatchSize, Body, [](){}, Flags, TArrayView<nullptr_t>());
+	ParallelForImpl::ParallelForInternal(DebugName, Num, MinBatchSize, Body, [](){}, Flags, TArrayView<TYPE_OF_NULLPTR>());
 }
 
 /** 
@@ -463,7 +463,7 @@ inline void ParallelForWithPreWork(int32 Num, TFunctionRef<void(int32)> Body, TF
 {
 	ParallelForImpl::ParallelForInternal(TEXT("ParallelFor"), Num, 1, Body, CurrentThreadWorkToDoBeforeHelping,
 		(bForceSingleThread ? EParallelForFlags::ForceSingleThread : EParallelForFlags::None) |
-		(bPumpRenderingThread ? EParallelForFlags::PumpRenderingThread : EParallelForFlags::None), TArrayView<nullptr_t>());
+		(bPumpRenderingThread ? EParallelForFlags::PumpRenderingThread : EParallelForFlags::None), TArrayView<TYPE_OF_NULLPTR>());
 }
 
 /** 
@@ -476,7 +476,7 @@ inline void ParallelForWithPreWork(int32 Num, TFunctionRef<void(int32)> Body, TF
 **/
 inline void ParallelForWithPreWork(int32 Num, TFunctionRef<void(int32)> Body, TFunctionRef<void()> CurrentThreadWorkToDoBeforeHelping, EParallelForFlags Flags = EParallelForFlags::None)
 {
-	ParallelForImpl::ParallelForInternal(TEXT("ParallelFor"), Num, 1, Body, CurrentThreadWorkToDoBeforeHelping, Flags, TArrayView<nullptr_t>());
+	ParallelForImpl::ParallelForInternal(TEXT("ParallelFor"), Num, 1, Body, CurrentThreadWorkToDoBeforeHelping, Flags, TArrayView<TYPE_OF_NULLPTR>());
 }
 
 /** 
@@ -491,7 +491,7 @@ inline void ParallelForWithPreWork(int32 Num, TFunctionRef<void(int32)> Body, TF
 **/
 inline void ParallelForWithPreWork(const TCHAR* DebugName, int32 Num, int32 MinBatchSize, TFunctionRef<void(int32)> Body, TFunctionRef<void()> CurrentThreadWorkToDoBeforeHelping, EParallelForFlags Flags = EParallelForFlags::None)
 {
-	ParallelForImpl::ParallelForInternal(DebugName, Num, MinBatchSize, Body, CurrentThreadWorkToDoBeforeHelping, Flags, TArrayView<nullptr_t>());
+	ParallelForImpl::ParallelForInternal(DebugName, Num, MinBatchSize, Body, CurrentThreadWorkToDoBeforeHelping, Flags, TArrayView<TYPE_OF_NULLPTR>());
 }
 
 /** 
