re-fix slow task getting stuck on auto-cancel w/ delayed dialogs, this time accounting for nested slow tasks

diff --git a/Engine/Source/Runtime/Core/Private/Misc/SlowTask.cpp b/Engine/Source/Runtime/Core/Private/Misc/SlowTask.cpp
--- a/Engine/Source/Runtime/Core/Private/Misc/SlowTask.cpp
+++ b/Engine/Source/Runtime/Core/Private/Misc/SlowTask.cpp
@@ -202,7 +202,7 @@ void FSlowTask::MakeDialog(bool bShowCancelButton, bool bAllowInPIE)
 
 bool FSlowTask::ShouldCancel() const
 {
-	if (bEnabled && bCreatedDialog)
+	if (bEnabled && GIsSlowTask)
 	{
 		check(IsInGameThread()); // FSlowTask is only meant to be used on the main thread currently
 
