Don't skip loading ini files if the path isn't fully qualified by ending in the .ini filetype.
This can prevent the CommandLineAlias from not being loaded if the ini file is skipped.

diff --git a/Engine/Source/Runtime/Core/Private/Misc/ConfigContext.cpp b/Engine/Source/Runtime/Core/Private/Misc/ConfigContext.cpp
--- a/Engine/Source/Runtime/Core/Private/Misc/ConfigContext.cpp
+++ b/Engine/Source/Runtime/Core/Private/Misc/ConfigContext.cpp
@@ -123,7 +123,8 @@ const FConfigContext::FPerPlatformDirs& FConfigContext::GetPerPlatformDirs(const
 bool FConfigContext::Load(const TCHAR* InBaseIniName, FString& OutFinalFilename)
 {
 	// for single file loads, just return early of the file doesn't exist
-	if (!bIsHierarchicalConfig && !DoesConfigFileExistWrapper(InBaseIniName, IniCacheSet))
+	const bool bBaseIniNameIsFullInIFilePath = FString(InBaseIniName).EndsWith(TEXT(".ini"));
+	if (!bIsHierarchicalConfig && bBaseIniNameIsFullInIFilePath && !DoesConfigFileExistWrapper(InBaseIniName, IniCacheSet))
 	{
 		return false;
 	}
