Allow users to opt-out of default crash handling.

This change allows users to opt-out of the engines default crash handling. Currently only implemented on Windows where the SEH code forwards the exception to externally registered handlers.

Thanks to [REDACTED_PERSON] for PR.

diff --git a/Engine/Source/Runtime/Core/Private/Windows/WindowsPlatformMisc.cpp b/Engine/Source/Runtime/Core/Private/Windows/WindowsPlatformMisc.cpp
--- a/Engine/Source/Runtime/Core/Private/Windows/WindowsPlatformMisc.cpp
+++ b/Engine/Source/Runtime/Core/Private/Windows/WindowsPlatformMisc.cpp
@@ -681,6 +681,19 @@ void FWindowsPlatformMisc::SetGracefulTerminationHandler()
 #endif // !UE_BUILD_SHIPPING && PLATFORM_CPU_X86_FAMILY
 }
 
+static ECrashHandlingType GCrashHandlingType; /* = ECrashHandlingType::Default */
+
+ECrashHandlingType FWindowsPlatformMisc::GetCrashHandlingType()
+{
+	return GCrashHandlingType;
+}
+
+ECrashHandlingType FWindowsPlatformMisc::SetCrashHandlingType(ECrashHandlingType Type)
+{
+	GCrashHandlingType = Type;
+	return GCrashHandlingType;
+}
+
 int32 FWindowsPlatformMisc::GetMaxPathLength()
 {
 	struct FLongPathsEnabled
diff --git a/Engine/Source/Runtime/Core/Private/Windows/WindowsRunnableThread.cpp b/Engine/Source/Runtime/Core/Private/Windows/WindowsRunnableThread.cpp
--- a/Engine/Source/Runtime/Core/Private/Windows/WindowsRunnableThread.cpp
+++ b/Engine/Source/Runtime/Core/Private/Windows/WindowsRunnableThread.cpp
@@ -79,7 +79,7 @@ uint32 FRunnableThreadWin::GuardedRun()
 			ExitCode = Run();
 		}
 #if !PLATFORM_SEH_EXCEPTIONS_DISABLED
-		__except (ReportCrash( GetExceptionInformation() ))
+		__except (FPlatformMisc::GetCrashHandlingType() == ECrashHandlingType::Default ? ReportCrash(GetExceptionInformation()) : EXCEPTION_CONTINUE_SEARCH)
 		{
 			__try
 			{
diff --git a/Engine/Source/Runtime/Core/Public/GenericPlatform/GenericPlatformMisc.h b/Engine/Source/Runtime/Core/Public/GenericPlatform/GenericPlatformMisc.h
--- a/Engine/Source/Runtime/Core/Public/GenericPlatform/GenericPlatformMisc.h
+++ b/Engine/Source/Runtime/Core/Public/GenericPlatform/GenericPlatformMisc.h
@@ -55,6 +55,15 @@
 };
 
 /**
+ * Controls behaviour for built in crash handling
+ */
+enum class ECrashHandlingType : uint8
+{
+	Default,	/* UE default handling */
+	Disabled,	/* UE doesn't handle crashes on main thread, general workers, rendering thread. */
+};
+
+/**
  * Returns the string representation of the specified EBuildConfiguration value.
  *
  * @param Configuration The string to get the EBuildConfiguration for.
@@ -523,6 +532,17 @@
 	static void SetCrashHandler(void (* CrashHandler)(const FGenericCrashContext& Context)) { }
 
 	/**
+	 * Gets the current crash handling type.
+	 */
+	static ECrashHandlingType GetCrashHandlingType() { return ECrashHandlingType::Default; }
+	
+	/**
+	 * Sets the type of crash handling done by the engine. Returns the new crash handling type. If disabling crash
+	 * handling make sure an alternative is active since threads that crash will just disappear.
+	 */
+	static ECrashHandlingType SetCrashHandlingType(ECrashHandlingType Type) { return ECrashHandlingType::Default; }
+	
+	/**
 	 * Retrieve a environment variable from the system
 	 *
 	 * @param VariableName The name of the variable (ie "Path")
diff --git a/Engine/Source/Runtime/Core/Public/Windows/WindowsPlatformMisc.h b/Engine/Source/Runtime/Core/Public/Windows/WindowsPlatformMisc.h
--- a/Engine/Source/Runtime/Core/Public/Windows/WindowsPlatformMisc.h
+++ b/Engine/Source/Runtime/Core/Public/Windows/WindowsPlatformMisc.h
@@ -47,6 +47,8 @@ struct CORE_API FWindowsPlatformMisc
 	static void PlatformPreInit();
 	static void PlatformInit();
 	static void SetGracefulTerminationHandler();
+	static ECrashHandlingType GetCrashHandlingType();
+	static ECrashHandlingType SetCrashHandlingType(ECrashHandlingType);
 	static int32 GetMaxPathLength();
 
 	UE_DEPRECATED(4.21, "void FPlatformMisc::GetEnvironmentVariable(Name, Result, Length) is deprecated. Use FString FPlatformMisc::GetEnvironmentVariable(Name) instead.")
diff --git a/Engine/Source/Runtime/Launch/Private/Windows/LaunchWindows.cpp b/Engine/Source/Runtime/Launch/Private/Windows/LaunchWindows.cpp
--- a/Engine/Source/Runtime/Launch/Private/Windows/LaunchWindows.cpp
+++ b/Engine/Source/Runtime/Launch/Private/Windows/LaunchWindows.cpp
@@ -147,7 +147,7 @@ LAUNCH_API int32 GuardedMainWrapper( const TCHAR* CmdLine )
 			ErrorLevel = GuardedMain( CmdLine );
 		}
 #if !PLATFORM_SEH_EXCEPTIONS_DISABLED
-		__except( ReportCrash( GetExceptionInformation() ), EXCEPTION_CONTINUE_SEARCH )
+		__except( FPlatformMisc::GetCrashHandlingType() == ECrashHandlingType::Default ? (ReportCrash( GetExceptionInformation()), EXCEPTION_CONTINUE_SEARCH) : EXCEPTION_CONTINUE_SEARCH )
 		{
 			// Deliberately do nothing but avoid warning C6322: Empty _except block.
 			(void)0;
diff --git a/Engine/Source/Runtime/RenderCore/Private/RenderingThread.cpp b/Engine/Source/Runtime/RenderCore/Private/RenderingThread.cpp
--- a/Engine/Source/Runtime/RenderCore/Private/RenderingThread.cpp
+++ b/Engine/Source/Runtime/RenderCore/Private/RenderingThread.cpp
@@ -526,7 +526,9 @@ PRAGMA_ENABLE_DEPRECATION_WARNINGS
 				RenderingThreadMain( TaskGraphBoundSyncEvent );
 			}
 #if !PLATFORM_SEH_EXCEPTIONS_DISABLED
-			__except(FlushRHILogsAndReportCrash(GetExceptionInformation()))
+			__except (FPlatformMisc::GetCrashHandlingType() == ECrashHandlingType::Default ?
+							FlushRHILogsAndReportCrash(GetExceptionInformation()) : 
+							EXCEPTION_CONTINUE_SEARCH)
 			{
 #if !NO_LOGGING
 				// Dump the error and flush the log. This is the same logging behavior as FWindowsErrorOutputDevice::HandleError which is called in GuardedMain's caller's __except
