Add comments to Quat.h following swarm review 20258682

diff --git a/Engine/Source/Runtime/Core/Public/Math/Quat.h b/Engine/Source/Runtime/Core/Public/Math/Quat.h
--- a/Engine/Source/Runtime/Core/Public/Math/Quat.h
+++ b/Engine/Source/Runtime/Core/Public/Math/Quat.h
@@ -406,9 +406,9 @@ public:
 	/** 
 	 * get the axis and angle of rotation of this quaternion
 	 *
-	 * @param Axis{out] vector of axis of the quaternion
-	 * @param Angle{out] angle of the quaternion in radians
-	 * @warning : assumes normalized quaternions.
+	 * @param Axis{out] Normalized rotation axis of the quaternion
+	 * @param Angle{out] Angle of the quaternion in radians
+	 * @warning : Requires this quaternion to be normalized.
 	 */
 	void ToAxisAndAngle(TVector<T>& Axis, float& Angle) const;
 	void ToAxisAndAngle(TVector<T>& Axis, double& Angle) const;
@@ -417,6 +417,7 @@ public:
 	 * Get the rotation vector corresponding to this quaternion. 
 	 * The direction of the vector represents the rotation axis, 
 	 * and the magnitude the angle in radians.
+	 * @warning : Requires this quaternion to be normalized.
 	 */
 	TVector<T> ToRotationVector() const;
 
@@ -476,6 +477,7 @@ public:
 
 	/**
 	 * @return inverse of this quaternion
+	 * @warning : Requires this quaternion to be normalized.
 	 */
 	FORCEINLINE TQuat<T> Inverse() const;
 
@@ -1150,6 +1152,7 @@ FORCEINLINE void TQuat<T>::ToAxisAndAngle(TVector<T>& Axis, double& Angle) const
 template<typename T>
 FORCEINLINE TVector<T> TQuat<T>::ToRotationVector() const
 {
+	checkSlow(IsNormalized());
 	TQuat<T> RotQ = Log();
 	return TVector<T>(RotQ.X * 2.0f, RotQ.Y * 2.0f, RotQ.Z * 2.0f);
 }
