Add extra info to fname stats command for hash table memory usage. Fix units for fixed overhead.

diff --git a/Engine/Source/Runtime/Core/Private/UObject/UnrealNames.cpp b/Engine/Source/Runtime/Core/Private/UObject/UnrealNames.cpp
--- a/Engine/Source/Runtime/Core/Private/UObject/UnrealNames.cpp
+++ b/Engine/Source/Runtime/Core/Private/UObject/UnrealNames.cpp
@@ -1762,7 +1762,12 @@ uint32 FNamePool::NumSlots() const
 
 void FNamePool::LogStats(FOutputDevice& Ar) const
 {
-	Ar.Logf(TEXT("%i FNames using in %ikB + %ikB"), NumEntries(), sizeof(FNamePool), Entries.NumBlocks() * FNameEntryAllocator::BlockSizeBytes / 1024);
+	Ar.Logf(TEXT("%i FNames using in %" SIZE_T_FMT " bytes fixed + %.2f kB entries + %.2f kB slots"),
+		NumEntries(), 
+		sizeof(FNamePool), 
+		Entries.NumBlocks() * FNameEntryAllocator::BlockSizeBytes / 1024.0, 
+		NumSlots() * sizeof(FNameSlot) / 1024.0
+	);
 	Ar.Logf(TEXT("%d ansi FNames"), NumAnsiEntries());
 	Ar.Logf(TEXT("%d wide FNames"), NumWideEntries());
 #if UE_FNAME_OUTLINE_NUMBER
