Auto enable "stat namedevents" when a razor or trace capture is in progress.

Functionality is controlled by stats.AutoEnableNamedEventsWhenProfiling, default disabled in engine. Enabled in fort.

diff --git a/Engine/Source/Runtime/Core/Private/Misc/CoreMisc.cpp b/Engine/Source/Runtime/Core/Private/Misc/CoreMisc.cpp
--- a/Engine/Source/Runtime/Core/Private/Misc/CoreMisc.cpp
+++ b/Engine/Source/Runtime/Core/Private/Misc/CoreMisc.cpp
@@ -493,3 +493,25 @@ namespace UEAsserts_Private
 	}
 }
 #endif
+
+static bool GAutoEnableNamedEventsWhenProfiling = false;
+static FAutoConsoleVariableRef GCVarAutoEnableNamedEventsWhenProfiling(
+	TEXT("stats.AutoEnableNamedEventsWhenProfiling"),
+	GAutoEnableNamedEventsWhenProfiling,
+	TEXT("If 1, toggles named events on when a profiler is detected and capturing. Toggles named events off if 0 or when profiling stops."),
+	ECVF_Default
+);
+
+void FAutoNamedEventsToggler::Update(bool bIsProfiling)
+{
+	if (bIsProfiling && !bSetNamedEventsEnabled && GAutoEnableNamedEventsWhenProfiling)
+	{
+		++GCycleStatsShouldEmitNamedEvents;
+		bSetNamedEventsEnabled = true;
+	}
+	else if (!bIsProfiling && bSetNamedEventsEnabled)
+	{
+		GCycleStatsShouldEmitNamedEvents = FMath::Max(0, GCycleStatsShouldEmitNamedEvents - 1);
+		bSetNamedEventsEnabled = false;
+	}
+}
diff --git a/Engine/Source/Runtime/Core/Public/Misc/CoreMisc.h b/Engine/Source/Runtime/Core/Public/Misc/CoreMisc.h
--- a/Engine/Source/Runtime/Core/Public/Misc/CoreMisc.h
+++ b/Engine/Source/Runtime/Core/Public/Misc/CoreMisc.h
@@ -263,6 +263,20 @@ struct CORE_API FScopedScriptExceptionHandler
 	~FScopedScriptExceptionHandler();
 };
 
+/**
+ * Enables named events when profiling.
+ * Increments/decrements GCycleStatsShouldEmitNamedEvents based on bIsProfiling.
+ * Functionality is controlled by stats.AutoEnableNamedEventsWhenProfiling.
+ */
+class FAutoNamedEventsToggler
+{
+public:
+	CORE_API void Update(bool bIsProfiling);
+
+private:
+	bool bSetNamedEventsEnabled = false;
+};
+
 /** 
  * This define enables the blueprint runaway and exception stack trace checks
  * If this is true, it will create a FBlueprintContextTracker (previously FBlueprintExceptionTracker) which is defined in Script.h
diff --git a/Engine/Source/Runtime/Launch/Private/LaunchEngineLoop.cpp b/Engine/Source/Runtime/Launch/Private/LaunchEngineLoop.cpp
--- a/Engine/Source/Runtime/Launch/Private/LaunchEngineLoop.cpp
+++ b/Engine/Source/Runtime/Launch/Private/LaunchEngineLoop.cpp
@@ -22,6 +22,7 @@
 #include "Misc/Paths.h"
 #include "Misc/PathViews.h"
 #include "Misc/ConfigCacheIni.h"
+#include "Misc/CoreMisc.h"
 #include "Misc/ConfigUtilities.h"
 #include "Misc/OutputDeviceHelper.h"
 #include "Misc/OutputDeviceRedirector.h"
@@ -5067,7 +5068,10 @@ void FEngineLoop::Tick()
 	TCHAR IndexedFrameString[32] = { 0 };
 	const TCHAR* FrameString = nullptr;
 #if CPUPROFILERTRACE_ENABLED
-	if (UE_TRACE_CHANNELEXPR_IS_ENABLED(CpuChannel))
+	const bool bTraceCpuChannelEnabled = UE_TRACE_CHANNELEXPR_IS_ENABLED(CpuChannel);
+	static FAutoNamedEventsToggler TraceNamedEventsToggler;
+	TraceNamedEventsToggler.Update(bTraceCpuChannelEnabled && UE::Trace::IsTracing());
+	if (bTraceCpuChannelEnabled)
 	{
 		FrameString = TEXT("FEngineLoop");
 	}
