Restore logging for setting CVars from config but only when the value changes.

diff --git a/Engine/Source/Runtime/Core/Private/Misc/ConfigUtilities.cpp b/Engine/Source/Runtime/Core/Private/Misc/ConfigUtilities.cpp
--- a/Engine/Source/Runtime/Core/Private/Misc/ConfigUtilities.cpp
+++ b/Engine/Source/Runtime/Core/Private/Misc/ConfigUtilities.cpp
@@ -31,7 +31,7 @@ const TCHAR* ConvertValueFromHumanFriendlyValue( const TCHAR* Value )
 }
 
 
-void OnSetCVarFromIniEntry(const TCHAR *IniFile, const TCHAR *Key, const TCHAR* Value, uint32 SetBy, bool bAllowCheating)
+void OnSetCVarFromIniEntry(const TCHAR *IniFile, const TCHAR *Key, const TCHAR* Value, uint32 SetBy, bool bAllowCheating, bool bNoLogging)
 {
 	check(IniFile && Key && Value);
 	check((SetBy & ECVF_FlagMask) == 0);
@@ -57,7 +57,12 @@ void OnSetCVarFromIniEntry(const TCHAR *IniFile, const TCHAR *Key, const TCHAR*
 
 		if(bAllowChange)
 		{
-			UE_LOG(LogConfig, Verbose, TEXT("Setting CVar [[%s:%s]]"), Key, Value);
+#if !NO_LOGGING
+			bool BoolValue = CVar->IsVariableBool() ? CVar->GetBool() : false;
+			int32 IntValue = CVar->IsVariableInt() ? CVar->GetInt() : 0;
+			float FloatValue = CVar->IsVariableFloat() ? CVar->GetFloat() : 0.0f;
+			FString StringValue = CVar->IsVariableString() ? CVar->GetString() : FString();
+#endif
 			if (SetBy == ECVF_SetByMask)
 			{
 				CVar->SetWithCurrentPriority(Value);
@@ -66,6 +71,26 @@ void OnSetCVarFromIniEntry(const TCHAR *IniFile, const TCHAR *Key, const TCHAR*
 			{
 				CVar->Set(Value, (EConsoleVariableFlags)SetBy);
 			}
+#if !NO_LOGGING
+			bool bChanged = false;
+			if (CVar->IsVariableBool())
+			{
+				bChanged = BoolValue != CVar->GetBool();
+			}
+			else if (CVar->IsVariableInt())
+			{
+				bChanged = IntValue != CVar->GetInt();
+			}
+			else if (CVar->IsVariableFloat())
+			{
+				bChanged = FloatValue != CVar->GetFloat();
+			}
+			else if (CVar->IsVariableString())
+			{
+				bChanged = StringValue != CVar->GetString();
+			}
+			UE_CLOG(!bNoLogging && bChanged, LogConfig, Log, TEXT("Set CVar [[%s:%s]]"), Key, Value);
+#endif
 		}
 		else
 		{
@@ -108,7 +133,7 @@ void ApplyCVarSettingsFromIni(const TCHAR* InSectionName, const TCHAR* InIniFile
 			const FString& KeyString = It.Key().GetPlainNameString(); 
 			const FString& ValueString = It.Value().GetValue();
 
-			OnSetCVarFromIniEntry(InIniFilename, *KeyString, *ValueString, SetBy, bAllowCheating);
+			OnSetCVarFromIniEntry(InIniFilename, *KeyString, *ValueString, SetBy, bAllowCheating, false);
 		}
 	}
 }
@@ -229,7 +254,7 @@ public:
 					UE_LOG(LogConfig, Display, TEXT("Applied changed CVAR value %s %s -> %s Config %s %s"), 
 						*KeyString, *CurrentValue, HumanFriendlyValue, *IniFilename, *SectionName);
 
-					OnSetCVarFromIniEntry(*IniFilename, *KeyString, *ValueString, SetBy, IniHistory.bAllowCheating);
+					OnSetCVarFromIniEntry(*IniFilename, *KeyString, *ValueString, SetBy, IniHistory.bAllowCheating, false);
 				}
 			}
 		}
diff --git a/Engine/Source/Runtime/Core/Public/Misc/ConfigUtilities.h b/Engine/Source/Runtime/Core/Public/Misc/ConfigUtilities.h
--- a/Engine/Source/Runtime/Core/Public/Misc/ConfigUtilities.h
+++ b/Engine/Source/Runtime/Core/Public/Misc/ConfigUtilities.h
@@ -17,7 +17,7 @@ namespace UE::ConfigUtilities
 	/**
 	 * Single function to set a cvar from ini (handing friendly names, cheats for shipping and message about cheats in non shipping)
 	 */
-	CORE_API void OnSetCVarFromIniEntry(const TCHAR* IniFile, const TCHAR* Key, const TCHAR* Value, uint32 SetBy, bool bAllowCheating=false);
+	CORE_API void OnSetCVarFromIniEntry(const TCHAR* IniFile, const TCHAR* Key, const TCHAR* Value, uint32 SetBy, bool bAllowCheating=false, bool bNoLogging=false);
 	
 	/**
 	 * Helper function to read the contents of an ini file and a specified group of cvar parameters, where sections in the ini file are marked [InName]
diff --git a/Engine/Source/Runtime/Engine/Private/DeviceProfiles/DeviceProfileManager.cpp b/Engine/Source/Runtime/Engine/Private/DeviceProfiles/DeviceProfileManager.cpp
--- a/Engine/Source/Runtime/Engine/Private/DeviceProfiles/DeviceProfileManager.cpp
+++ b/Engine/Source/Runtime/Engine/Private/DeviceProfiles/DeviceProfileManager.cpp
@@ -387,7 +387,7 @@ void UDeviceProfileManager::SetDeviceProfileCVars(const FString& DeviceProfileNa
 
 		// Set by scalability or DP, depending
 		uint32 CVarPriority = bIsScalabilityBucket ? ECVF_SetByScalability : ECVF_SetByDeviceProfile;
-		UE::ConfigUtilities::OnSetCVarFromIniEntry(*GDeviceProfilesIni, *CVarKey, *CVarValue, CVarPriority);
+		UE::ConfigUtilities::OnSetCVarFromIniEntry(*GDeviceProfilesIni, *CVarKey, *CVarValue, CVarPriority, false, true);
 	}
 
 
