Avoid running Smoke Tests when running commandlet.
* + Update flags for UObject.Class AttemptToFindUninitializedScriptStructMembers to avoid running on commandlet anyway

#rnx

diff --git a/Engine/Source/Runtime/Core/Private/Misc/AutomationTest.cpp b/Engine/Source/Runtime/Core/Private/Misc/AutomationTest.cpp
--- a/Engine/Source/Runtime/Core/Private/Misc/AutomationTest.cpp
+++ b/Engine/Source/Runtime/Core/Private/Misc/AutomationTest.cpp
@@ -271,7 +271,7 @@ bool FAutomationTestFramework::RunSmokeTests()
 	//@todo - better determination of whether to run than requires cooked data
 	// Ensure there isn't another slow task in progress when trying to run unit tests
 	const bool bRequiresCookedData = FPlatformProperties::RequiresCookedData();
-	if ((!bRequiresCookedData && !GIsSlowTask && !GIsPlayInEditorWorld && !FPlatformProperties::IsProgram()) || bForceSmokeTests)
+	if ((!bRequiresCookedData && !GIsSlowTask && !GIsPlayInEditorWorld && !FPlatformProperties::IsProgram() && !IsRunningCommandlet()) || bForceSmokeTests)
 	{
 		TArray<FAutomationTestInfo> TestInfo;
 
@@ -322,9 +322,9 @@ bool FAutomationTestFramework::RunSmokeTests()
 			FAutomationTestFramework::DumpAutomationTestExecutionInfo( OutExecutionInfoMap );
 		}
 	}
-	else if( bRequiresCookedData )
+	else if( bRequiresCookedData || IsRunningCommandlet())
 	{
-		UE_LOG( LogAutomationTest, Log, TEXT( "Skipping unit tests for the cooked build." ) );
+		UE_LOG( LogAutomationTest, Log, TEXT( "Skipping unit tests for the cooked build and commandlet." ) );
 	}
 	else if (!FPlatformProperties::IsProgram())
 	{
diff --git a/Engine/Source/Runtime/CoreUObject/Private/UObject/Class.cpp b/Engine/Source/Runtime/CoreUObject/Private/UObject/Class.cpp
--- a/Engine/Source/Runtime/CoreUObject/Private/UObject/Class.cpp
+++ b/Engine/Source/Runtime/CoreUObject/Private/UObject/Class.cpp
@@ -3656,7 +3656,7 @@ FAutoConsoleCommandWithWorldAndArgs GCmdListBadScriptStructs(
 			FStructUtils::AttemptToFindUninitializedScriptStructMembers();
 		}));
 
-IMPLEMENT_SIMPLE_AUTOMATION_TEST(FAutomationTestAttemptToFindUninitializedScriptStructMembers, "UObject.Class AttemptToFindUninitializedScriptStructMembers", EAutomationTestFlags::EditorContext | EAutomationTestFlags::ApplicationContextMask | EAutomationTestFlags::ServerContext | EAutomationTestFlags::SmokeFilter)
+IMPLEMENT_SIMPLE_AUTOMATION_TEST(FAutomationTestAttemptToFindUninitializedScriptStructMembers, "UObject.Class AttemptToFindUninitializedScriptStructMembers", EAutomationTestFlags::EditorContext | EAutomationTestFlags::ClientContext | EAutomationTestFlags::ServerContext | EAutomationTestFlags::SmokeFilter)
 bool FAutomationTestAttemptToFindUninitializedScriptStructMembers::RunTest(const FString& Parameters)
 {
 	// This test fails when running tests under UHT because there is no TestUninitializedScriptStructMembersTest, so just skip it in that config.
