Non-unity build fix.

diff --git a/Engine/Source/Runtime/Core/Public/ProfilingDebugging/TraceAuxiliary.h b/Engine/Source/Runtime/Core/Public/ProfilingDebugging/TraceAuxiliary.h
--- a/Engine/Source/Runtime/Core/Public/ProfilingDebugging/TraceAuxiliary.h
+++ b/Engine/Source/Runtime/Core/Public/ProfilingDebugging/TraceAuxiliary.h
@@ -5,6 +5,7 @@
 #include "CoreGlobals.h"
 #include "HAL/Platform.h"
 #include "Logging/LogMacros.h"
+#include "Delegates/Delegate.h"
 
 ////////////////////////////////////////////////////////////////////////////////
 class CORE_API FTraceAuxiliary
