#ue5 - Add explicit overloads of FMath::Clamp() for all float and all double arguments because PVS static analysis resolves those incorrectly otherwise and generates false float to int warnings.

diff --git a/Engine/Source/Runtime/Core/Public/Math/UnrealMathUtility.h b/Engine/Source/Runtime/Core/Public/Math/UnrealMathUtility.h
--- a/Engine/Source/Runtime/Core/Public/Math/UnrealMathUtility.h
+++ b/Engine/Source/Runtime/Core/Public/Math/UnrealMathUtility.h
@@ -591,11 +591,12 @@ public:
 	/** Allow mixing float/double arguments, promoting to highest precision type. */
 	MIX_FLOATS_3_ARGS(Clamp);
 	
-	/** Clamps X to be between Min and Max, inclusive. Overload to support int64/int32 mixed types. */
-	UE_NODISCARD static constexpr FORCEINLINE int64 Clamp(const int64 X, const int32 Min, const int32 Max)
-	{
-		return (X < Min) ? Min : (X < Max) ? X : Max;
-	}
+	/** Clamps X to be between Min and Max, inclusive. Explicitly defined here for floats/doubles because static analysis gets confused between template and int versions. */
+	UE_NODISCARD static constexpr FORCEINLINE float Clamp(const float X, const float Min, const float Max) { return Clamp<float>(X, Min, Max); }
+	UE_NODISCARD static constexpr FORCEINLINE double Clamp(const double X, const double Min, const double Max) { return Clamp<double>(X, Min, Max); }
+
+	/** Clamps X to be between Min and Max, inclusive. Overload to support mixed int64/int32 types. */
+	UE_NODISCARD static constexpr FORCEINLINE int64 Clamp(const int64 X, const int32 Min, const int32 Max) { return Clamp<int64>(X, Min, Max); }
 
 	/** Wraps X to be between Min and Max, inclusive. */
 	/** When X can wrap to both Min and Max, it will wrap to Min if it lies below the range and wrap to Max if it is above the range. */
