Add operator<< to FIntVector2.

diff --git a/Engine/Source/Runtime/Core/Public/Math/IntVector.h b/Engine/Source/Runtime/Core/Public/Math/IntVector.h
--- a/Engine/Source/Runtime/Core/Public/Math/IntVector.h
+++ b/Engine/Source/Runtime/Core/Public/Math/IntVector.h
@@ -597,6 +597,18 @@ PRAGMA_ENABLE_DEPRECATION_WARNINGS
 	{
 		return X!=Other.X || Y!=Other.Y;
 	}
+
+	/**
+	 * Serializes the Vector2.
+	 *
+	 * @param Ar The archive to serialize into.
+	 * @param Vector The vector to serialize.
+	 * @return Reference to the Archive after serialization.
+	 */
+	friend FArchive& operator<<(FArchive& Ar, TIntVector2& Vector)
+	{
+		return Ar << Vector.X << Vector.Y;
+	}
 };
 
 ///////////////////////////////////////////////////////////////////////////////////////////////////
