Fixed Insights scopes no longer reporting file/line information.

diff --git a/Engine/Source/Runtime/Core/Public/ProfilingDebugging/CpuProfilerTrace.h b/Engine/Source/Runtime/Core/Public/ProfilingDebugging/CpuProfilerTrace.h
--- a/Engine/Source/Runtime/Core/Public/ProfilingDebugging/CpuProfilerTrace.h
+++ b/Engine/Source/Runtime/Core/Public/ProfilingDebugging/CpuProfilerTrace.h
@@ -113,28 +113,28 @@ struct FCpuProfilerTrace
 			BeginEventCommon(InSpecId);
 		}
 
-		FORCEINLINE FEventScope(uint32& InOutSpecId, const ANSICHAR* InEventString, bool bInCondition)
+		FORCEINLINE FEventScope(uint32& InOutSpecId, const ANSICHAR* InEventString, bool bInCondition, const ANSICHAR* File, uint32 Line)
 			: bEnabled(bInCondition && CpuChannel)
 		{
-			BeginEventCommon(InOutSpecId, InEventString);
+			BeginEventCommon(InOutSpecId, InEventString, File, Line);
 		}
 
-		FORCEINLINE FEventScope(uint32& InOutSpecId, const ANSICHAR* InEventString, const UE::Trace::FChannel& InChannel, bool bInCondition)
+		FORCEINLINE FEventScope(uint32& InOutSpecId, const ANSICHAR* InEventString, const UE::Trace::FChannel& InChannel, bool bInCondition, const ANSICHAR* File, uint32 Line)
 			: bEnabled(bInCondition && (CpuChannel | InChannel))
 		{
-			BeginEventCommon(InOutSpecId, InEventString);
+			BeginEventCommon(InOutSpecId, InEventString, File, Line);
 		}
 
-		FORCEINLINE FEventScope(uint32& InOutSpecId, const TCHAR* InEventString, bool bInCondition)
+		FORCEINLINE FEventScope(uint32& InOutSpecId, const TCHAR* InEventString, bool bInCondition, const ANSICHAR* File, uint32 Line)
 			: bEnabled(bInCondition && CpuChannel)
 		{
-			BeginEventCommon(InOutSpecId, InEventString);
+			BeginEventCommon(InOutSpecId, InEventString, File, Line);
 		}
 
-		FORCEINLINE FEventScope(uint32& InOutSpecId, const TCHAR* InEventString, const UE::Trace::FChannel& InChannel, bool bInCondition)
+		FORCEINLINE FEventScope(uint32& InOutSpecId, const TCHAR* InEventString, const UE::Trace::FChannel& InChannel, bool bInCondition, const ANSICHAR* File, uint32 Line)
 			: bEnabled(bInCondition && (CpuChannel | InChannel))
 		{
-			BeginEventCommon(InOutSpecId, InEventString);
+			BeginEventCommon(InOutSpecId, InEventString, File, Line);
 		}
 
 		FORCEINLINE ~FEventScope()
@@ -154,25 +154,25 @@ struct FCpuProfilerTrace
 			}
 		}
 
-		FORCEINLINE void BeginEventCommon(uint32& InOutSpecId, const ANSICHAR* InEventString)
+		FORCEINLINE void BeginEventCommon(uint32& InOutSpecId, const ANSICHAR* InEventString, const ANSICHAR* File, uint32 Line)
 		{
 			if (bEnabled)
 			{
 				if (InOutSpecId == 0)
 				{
-					InOutSpecId = FCpuProfilerTrace::OutputEventType(InEventString);
+					InOutSpecId = FCpuProfilerTrace::OutputEventType(InEventString, File, Line);
 				}
 				OutputBeginEvent(InOutSpecId);
 			}
 		}
 
-		FORCEINLINE void BeginEventCommon(uint32& InOutSpecId, const TCHAR* InEventString)
+		FORCEINLINE void BeginEventCommon(uint32& InOutSpecId, const TCHAR* InEventString, const ANSICHAR* File, uint32 Line)
 		{
 			if (bEnabled)
 			{
 				if (InOutSpecId == 0)
 				{
-					InOutSpecId = FCpuProfilerTrace::OutputEventType(InEventString);
+					InOutSpecId = FCpuProfilerTrace::OutputEventType(InEventString, File, Line);
 				}
 				OutputBeginEvent(InOutSpecId);
 			}
@@ -239,12 +239,12 @@ struct FCpuProfilerTrace
 // Advanced macro for integrating e.g. stats/named events with cpu trace
 // Traces a scoped event previously declared with TRACE_CPUPROFILER_EVENT_DECLARE, conditionally
 #define TRACE_CPUPROFILER_EVENT_SCOPE_USE(DeclName, NameStr, ScopeName, Condition) \
-	FCpuProfilerTrace::FEventScope ScopeName(DeclName, NameStr, Condition);
+	FCpuProfilerTrace::FEventScope ScopeName(DeclName, NameStr, Condition, __FILE__, __LINE__);
 
 // Advanced macro for integrating e.g. stats/named events with cpu trace
 // Traces a scoped event previously declared with TRACE_CPUPROFILER_EVENT_DECLARE, conditionally
 #define TRACE_CPUPROFILER_EVENT_SCOPE_USE_ON_CHANNEL(DeclName, NameStr, ScopeName, Channel, Condition) \
-	FCpuProfilerTrace::FEventScope ScopeName(DeclName, NameStr, Channel, Condition);
+	FCpuProfilerTrace::FEventScope ScopeName(DeclName, NameStr, Channel, Condition, __FILE__, __LINE__);
 
 // Trace a scoped cpu timing event providing a static string (const ANSICHAR* or const TCHAR*)
 // as the scope name. It will use the Cpu trace channel.
