Fixing up places that use FPathPermissionList with class paths. Added an optional EPathPermissionListType constructor param to catch places that pass short path names to FPathPermissionList when it's being used to store class path names

diff --git a/Engine/Plugins/Editor/ContentBrowser/ContentBrowserClassDataSource/Source/ContentBrowserClassDataSource/Private/ContentBrowserClassDataSource.cpp b/Engine/Plugins/Editor/ContentBrowser/ContentBrowserClassDataSource/Source/ContentBrowserClassDataSource/Private/ContentBrowserClassDataSource.cpp
--- a/Engine/Plugins/Editor/ContentBrowser/ContentBrowserClassDataSource/Source/ContentBrowserClassDataSource/Private/ContentBrowserClassDataSource.cpp
+++ b/Engine/Plugins/Editor/ContentBrowser/ContentBrowserClassDataSource/Source/ContentBrowserClassDataSource/Private/ContentBrowserClassDataSource.cpp
@@ -301,7 +301,7 @@ void UContentBrowserClassDataSource::CompileFilter(const FName InPath, const FCo
 			for (UClass* ChildClassObject : ChildClassObjects)
 			{
 				const bool bPassesInclusiveFilter = ClassPathsToInclude.Num() == 0 || ClassPathsToInclude.Contains(*ChildClassObject->GetPathName());
-				const bool bPassesPermissionCheck = !ClassPermissionList || ClassPermissionList->PassesFilter(ChildClassObject->GetFName());
+				const bool bPassesPermissionCheck = !ClassPermissionList || ClassPermissionList->PassesFilter(ChildClassObject->GetClassPathName().ToString());
 
 				if (bPassesInclusiveFilter && bPassesPermissionCheck)
 				{
diff --git a/Engine/Source/Developer/AssetTools/Private/AssetTools.cpp b/Engine/Source/Developer/AssetTools/Private/AssetTools.cpp
--- a/Engine/Source/Developer/AssetTools/Private/AssetTools.cpp
+++ b/Engine/Source/Developer/AssetTools/Private/AssetTools.cpp
@@ -192,10 +192,22 @@ UAssetToolsImpl::UAssetToolsImpl(const FObjectInitializer& ObjectInitializer)
 	GConfig->GetArray(TEXT("AssetTools"), TEXT("SupportedAssetTypes"), SupportedTypesArray, GEditorIni);
 	for (int32 i = 0; i < (int32)EAssetClassAction::AllAssetActions; ++i)
 	{
-		AssetClassPermissionList.Add(MakeShared<FPathPermissionList>());
-		for (const FString& Type : SupportedTypesArray)
+		AssetClassPermissionList.Add(MakeShared<FPathPermissionList>(EPathPermissionListType::ClassPaths));
+		for (FString& Type : SupportedTypesArray)
 		{
-			AssetClassPermissionList[i]->AddAllowListItem("AssetToolsConfigFile", *Type);
+			if (FPackageName::IsShortPackageName(Type))
+			{
+				FTopLevelAssetPath TypePath = UClass::TryConvertShortTypeNameToPathName<UStruct>(Type, ELogVerbosity::Warning, TEXT("AssetToolsImpl"));
+				if (TypePath.IsNull())
+				{
+					UE_LOG(LogAssetTools, Warning, TEXT("Failed to convert short type name \"%s\" to path name. Please update SupportedAssetTypes entries in [AssetTools] ini section"), *Type);
+				}
+				else
+				{
+					Type = TypePath.ToString();
+				}
+			}
+			AssetClassPermissionList[i]->AddAllowListItem("AssetToolsConfigFile", Type);
 		}
 		AssetClassPermissionList[i]->OnFilterChanged().AddUObject(this, &UAssetToolsImpl::AssetClassPermissionListChanged, (EAssetClassAction)i);
 	}
@@ -337,7 +349,7 @@ void UAssetToolsImpl::RegisterAssetTypeActions(const TSharedRef<IAssetTypeAction
 	bool bSupported = false;
 	if (const UClass* SupportedClass = NewActions->GetSupportedClass())
 	{
-		bSupported = GetAssetClassPathPermissionList(EAssetClassAction::CreateAsset)->PassesFilter(SupportedClass->GetFName());
+		bSupported = GetAssetClassPathPermissionList(EAssetClassAction::CreateAsset)->PassesFilter(SupportedClass->GetClassPathName().ToString());
 	}
 	else
 	{
diff --git a/Engine/Source/Editor/UnrealEd/Classes/Settings/ClassViewerSettings.h b/Engine/Source/Editor/UnrealEd/Classes/Settings/ClassViewerSettings.h
--- a/Engine/Source/Editor/UnrealEd/Classes/Settings/ClassViewerSettings.h
+++ b/Engine/Source/Editor/UnrealEd/Classes/Settings/ClassViewerSettings.h
@@ -61,7 +61,10 @@ protected:
 
 	// UObject overrides
 	virtual void PostEditChangeProperty(struct FPropertyChangedEvent& PropertyChangedEvent) override;
+	virtual void PostInitProperties() override;
+	virtual void PostLoad() override;
 
+	void FixupShortNames();
 
 	// Holds an event delegate that is executed when a setting has changed.
 	static FSettingChangedEvent SettingChangedEvent;
diff --git a/Engine/Source/Editor/UnrealEd/Private/Settings/SettingsClasses.cpp b/Engine/Source/Editor/UnrealEd/Private/Settings/SettingsClasses.cpp
--- a/Engine/Source/Editor/UnrealEd/Private/Settings/SettingsClasses.cpp
+++ b/Engine/Source/Editor/UnrealEd/Private/Settings/SettingsClasses.cpp
@@ -44,6 +44,8 @@
 
 #define LOCTEXT_NAMESPACE "SettingsClasses"
 
+DEFINE_LOG_CATEGORY_STATIC(LogSettingsClasses, Log, All);
+
 /* UContentBrowserSettings interface
  *****************************************************************************/
 
@@ -90,6 +92,11 @@ void UClassViewerSettings::PostEditChangeProperty(struct FPropertyChangedEvent&
 		? PropertyChangedEvent.Property->GetFName()
 		: NAME_None;
 
+	if (PropertyChangedEvent.MemberProperty && PropertyChangedEvent.MemberProperty->GetFName() == GET_MEMBER_NAME_CHECKED(UClassViewerSettings, AllowedClasses))
+	{
+		FixupShortNames();
+	}
+
 	if (!FUnrealEdMisc::Get().IsDeletePreferences())
 	{
 		SaveConfig();
@@ -98,6 +105,38 @@ void UClassViewerSettings::PostEditChangeProperty(struct FPropertyChangedEvent&
 	SettingChangedEvent.Broadcast(Name);
 }
 
+void UClassViewerSettings::PostInitProperties()
+{
+	Super::PostInitProperties();
+	FixupShortNames();
+}
+
+void UClassViewerSettings::PostLoad()
+{
+	Super::PostLoad();
+	FixupShortNames();
+}
+
+void UClassViewerSettings::FixupShortNames()
+{
+	for (FString& ClassName : AllowedClasses)
+	{
+		// Empty string may represent a new entry that's just been added in the editor
+		if (ClassName.Len() && FPackageName::IsShortPackageName(ClassName))
+		{
+			FTopLevelAssetPath ClassPathName = UClass::TryConvertShortTypeNameToPathName<UStruct>(ClassName, ELogVerbosity::Warning, TEXT("ClassViewerSettings"));
+			if (!ClassPathName.IsNull())
+			{
+				ClassName = ClassPathName.ToString();
+			}
+			else
+			{
+				UE_LOG(LogSettingsClasses, Warning, TEXT("Unable to convert short class name \"%s\" to path names for %s.AllowedClasses. Please update it manually"), *ClassName, *GetPathName());
+			}
+		}
+	}
+}
+
 /* UStructViewerSettings interface
 *****************************************************************************/
 
diff --git a/Engine/Source/Runtime/Core/Private/Misc/PathPermissionList.cpp b/Engine/Source/Runtime/Core/Private/Misc/PathPermissionList.cpp
--- a/Engine/Source/Runtime/Core/Private/Misc/PathPermissionList.cpp
+++ b/Engine/Source/Runtime/Core/Private/Misc/PathPermissionList.cpp
@@ -12,6 +12,8 @@ bool FPathPermissionList::PassesFilter(const FStringView Item) const
 		return false;
 	}
 
+	VerifyItemMatchesListType(Item);
+
 	if (AllowList.Num() > 0 || DenyList.Num() > 0)
 	{
 		const uint32 ItemHash = GetTypeHash(Item);
@@ -42,6 +44,8 @@ bool FPathPermissionList::PassesFilter(const TCHAR* Item) const
 
 bool FPathPermissionList::PassesStartsWithFilter(const FStringView Item, const bool bAllowParentPaths) const
 {
+	VerifyItemMatchesListType(Item);
+
 	if (AllowList.Num() > 0)
 	{
 		bool bPassedAllowList = false;
@@ -101,6 +105,8 @@ bool FPathPermissionList::PassesStartsWithFilter(const TCHAR* Item, const bool b
 
 bool FPathPermissionList::AddDenyListItem(const FName OwnerName, const FStringView Item)
 {
+	VerifyItemMatchesListType(Item);
+
 	const uint32 ItemHash = GetTypeHash(Item);
 
 	FPermissionListOwners* Owners = DenyList.FindByHash(ItemHash, Item);
@@ -164,6 +170,8 @@ bool FPathPermissionList::RemoveDenyListItem(const FName OwnerName, const TCHAR*
 
 bool FPathPermissionList::AddAllowListItem(const FName OwnerName, const FStringView Item)
 {
+	VerifyItemMatchesListType(Item);
+
 	const uint32 ItemHash = GetTypeHash(Item);
 
 	FPermissionListOwners* Owners = AllowList.FindByHash(ItemHash, Item);
@@ -328,6 +336,8 @@ bool FPathPermissionList::UnregisterOwners(const TArray<FName>& OwnerNames)
 
 bool FPathPermissionList::Append(const FPathPermissionList& Other)
 {
+	ensureAlwaysMsgf(ListType == Other.ListType, TEXT("Trying to combine PathPermissionLists of different types"));
+
 	bool bFilterChanged = false;
 	{
 		TGuardValue<bool> Guard(bSuppressOnFilterChanged, true);
@@ -439,4 +449,19 @@ bool FPathPermissionList::UnregisterOwnersAndAppend(const TArray<FName>& OwnerNa
 	}
 
 	return bFilterChanged;
-}
\ No newline at end of file
+}
+
+// Extracted the ensure condition into a separate function so that logs are easier to read
+FORCEINLINE static bool IsClassPathNameOrNone(const FStringView Item)
+{
+	return !Item.Len() || Item[0] == '/' || Item == TEXTVIEW("None");
+}
+
+void FPathPermissionList::VerifyItemMatchesListType(const FStringView Item) const
+{
+	if (ListType == EPathPermissionListType::ClassPaths)
+	{
+		// Long names always have / as first character
+		ensureAlwaysMsgf(IsClassPathNameOrNone(Item), TEXT("Short class name \"%.*s\" provided for PathPermissionList representing class paths"), Item.Len(), Item.GetData());
+	}
+}
diff --git a/Engine/Source/Runtime/Core/Public/Misc/NamePermissionList.h b/Engine/Source/Runtime/Core/Public/Misc/NamePermissionList.h
--- a/Engine/Source/Runtime/Core/Public/Misc/NamePermissionList.h
+++ b/Engine/Source/Runtime/Core/Public/Misc/NamePermissionList.h
@@ -109,10 +109,19 @@ protected:
 	bool bSuppressOnFilterChanged = false;
 };
 
+enum class EPathPermissionListType
+{
+	Default,	// Default path permission list
+	ClassPaths	// Class permission list
+};
+
 class CORE_API FPathPermissionList : public TSharedFromThis<FPathPermissionList>
 {
 public:
-	FPathPermissionList() {}
+	FPathPermissionList(EPathPermissionListType InType = EPathPermissionListType::Default) 
+		: ListType(InType)
+	{
+	}
 	virtual ~FPathPermissionList() {}
 	
 	/** Returns true if passes filter restrictions using exact match */
@@ -263,6 +272,12 @@ public:
 
 protected:
 
+	/**
+	 * Checks if an item is of a valid format for this list
+	 * @return True if the item passes list type test.
+	 */
+	void VerifyItemMatchesListType(const FStringView Item) const;
+
 	/** List if items to filter out */
 	TMap<FString, FPermissionListOwners> DenyList;
 
@@ -277,4 +292,7 @@ protected:
 
 	/** Temporarily prevent delegate from being triggered */
 	bool bSuppressOnFilterChanged = false;
+
+	/** Type of paths this list represent */
+	EPathPermissionListType ListType;
 };
