#UE Temporarily add ensure to track down the cause for a long-standing warning

diff --git a/Engine/Source/Runtime/Core/Private/Misc/AutomationTest.cpp b/Engine/Source/Runtime/Core/Private/Misc/AutomationTest.cpp
--- a/Engine/Source/Runtime/Core/Private/Misc/AutomationTest.cpp
+++ b/Engine/Source/Runtime/Core/Private/Misc/AutomationTest.cpp
@@ -316,6 +316,7 @@ bool FAutomationTestFramework::RunSmokeTests()
 				UE_LOG(LogAutomationTest, Warning, TEXT("Smoke tests took >2s to run (%.2fs). '%s' took %dms. "
 					"SmokeFilter tier tests should take less than 1ms. Please optimize or move '%s' to a slower tier than SmokeFilter."), 
 					TotalDuration, SlowestTestName, static_cast<int32>(1000*SlowestTestDuration), SlowestTestName);
+				ensure(false);
 			}
 #endif
 
