Get rid of modular features bool that's ensuring in situations where there's only two threads

diff --git a/Engine/Source/Runtime/Core/Private/Features/ModularFeatures.cpp b/Engine/Source/Runtime/Core/Private/Features/ModularFeatures.cpp
--- a/Engine/Source/Runtime/Core/Private/Features/ModularFeatures.cpp
+++ b/Engine/Source/Runtime/Core/Private/Features/ModularFeatures.cpp
@@ -13,25 +13,23 @@ IModularFeatures& IModularFeatures::Get()
 void FModularFeatures::LockModularFeatureList()
 {
 	ModularFeaturesMapCriticalSection.Lock();
-	bModularFeatureListLocked = true;
 }
 
 void FModularFeatures::UnlockModularFeatureList()
 {
-	bModularFeatureListLocked = false;
 	ModularFeaturesMapCriticalSection.Unlock();
 }
 
 int32 FModularFeatures::GetModularFeatureImplementationCount( const FName Type )
 {
-	ensureMsgf(IsInGameThread() || bModularFeatureListLocked, TEXT("IModularFeature counting is not thread-safe unless wrapped with LockModularFeatureList/UnlockModularFeatureList"));
+	// IModularFeature counting is not thread-safe unless wrapped with LockModularFeatureList/UnlockModularFeatureList if you are crashing here
 
 	return ModularFeaturesMap.Num( Type );
 }
 
 IModularFeature* FModularFeatures::GetModularFeatureImplementation( const FName Type, const int32 Index )
 {
-	ensureMsgf(IsInGameThread() || bModularFeatureListLocked, TEXT("IModularFeature fetching is not thread-safe unless wrapped with LockModularFeatureList/UnlockModularFeatureList"));
+	// IModularFeature fetching is not thread-safe unless wrapped with LockModularFeatureList/UnlockModularFeatureList if you are crashing here
 
 	IModularFeature* ModularFeature = nullptr;
 
diff --git a/Engine/Source/Runtime/Core/Private/Features/ModularFeatures.h b/Engine/Source/Runtime/Core/Private/Features/ModularFeatures.h
--- a/Engine/Source/Runtime/Core/Private/Features/ModularFeatures.h
+++ b/Engine/Source/Runtime/Core/Private/Features/ModularFeatures.h
@@ -37,10 +37,7 @@ private:
 	/** Maps each feature type to a list of known providers of that feature */
 	TMultiMap< FName, class IModularFeature* > ModularFeaturesMap;
 
-	/** Used to fire ensure when ModularFeaturesMap is inspected cross-thread without being locked properly */
-	std::atomic<bool> bModularFeatureListLocked = false;
-
-	/** Lock modular features map so it can be used across threads */
+	/** Lock modular features map so it can be used across threads, this could be a RW lock instead of a crit sec */
 	FCriticalSection ModularFeaturesMapCriticalSection;
 
 	/** Event used to inform clients that a modular feature has been registered */
