Fixed callstack tracing not capturing correct backtrace for majority of allocations in Editor Win64 Debug.
The cause of the bug is the fact that few allocations are made before loaded modules are registered. For those initial allocs, the stack cannot be properly backtraced and the code caches that info for its function on top of the stack (which is "MemoryTrace_Alloc" in debug config and "CallstackTrace_InitializeInternal" for development config). Further callstacks that encounters those functions stops backtracing, so the result is a "one frame callstack". Debug build is seriously affected by this bug because "MemoryTrace_Alloc" is present in the callstack of majority of allocations (so all of them will be reduced to one frame).

diff --git a/Engine/Source/Runtime/Core/Private/ProfilingDebugging/Microsoft/WindowsCallstackTrace.cpp b/Engine/Source/Runtime/Core/Private/ProfilingDebugging/Microsoft/WindowsCallstackTrace.cpp
--- a/Engine/Source/Runtime/Core/Private/ProfilingDebugging/Microsoft/WindowsCallstackTrace.cpp
+++ b/Engine/Source/Runtime/Core/Private/ProfilingDebugging/Microsoft/WindowsCallstackTrace.cpp
@@ -135,7 +135,7 @@ private:
 	struct FModule
 	{
 		uint32				Id;
-		uint32				IdSize = 0;
+		uint32				IdSize;
 		uint32				NumFunctions;
 #if BACKTRACE_DBGLVL >= 1
 		uint16				NumFpTypes;
@@ -203,6 +203,7 @@ private:
 ////////////////////////////////////////////////////////////////////////////////
 FBacktracer* FBacktracer::Instance = nullptr;
 static bool GFullBacktraces = UE_CALLSTACK_TRACE_FULL_CALLSTACKS;
+static bool GModulesAreInitialized = false;
 
 ////////////////////////////////////////////////////////////////////////////////
 FBacktracer::FBacktracer(FMalloc* InMalloc)
@@ -231,7 +232,6 @@ FBacktracer::~FBacktracer()
 	{
 		Malloc->Free(Module.Functions);
 	}
-
 }
 
 ////////////////////////////////////////////////////////////////////////////////
@@ -566,7 +566,7 @@ const FBacktracer::FFunction* FBacktracer::LookupFunction(UPTRINT Address, FLook
 ////////////////////////////////////////////////////////////////////////////////
 uint32 FBacktracer::GetBacktraceId(void* AddressOfReturnAddress) const
 {
-	FLookupState LookupState;
+	FLookupState LookupState = {};
 	FCallstackTracer::FBacktraceEntry BacktraceEntry;
 
 	UPTRINT* StackPointer = (UPTRINT*)AddressOfReturnAddress;
@@ -640,8 +640,13 @@ uint32 FBacktracer::GetBacktraceId(void* AddressOfReturnAddress) const
 		if (Function == nullptr)
 		{
 #if BACKTRACE_LOCK_FREE
-			TGuardValue<bool> ReentranceGuard(bReentranceCheck, true);
-			FunctionLookups.Emplace(RetAddr, -1);
+			// LookupFunction fails when modules are not yet registered. In this case, we do not want the address
+			// to be added to the lookup map, but to retry the lookup later when modules are properly registered.
+			if (GModulesAreInitialized)
+			{
+				TGuardValue<bool> ReentranceGuard(bReentranceCheck, true);
+				FunctionLookups.Emplace(RetAddr, -1);
+			}
 #endif
 			break;
 		}
@@ -712,7 +717,7 @@ int32 FBacktracer::GetFrameSize(void* FunctionAddress) const
 {
 	FScopeLock _(&Lock);
 
-	FLookupState LookupState;
+	FLookupState LookupState = {};
 	const FFunction* Function = LookupFunction(UPTRINT(FunctionAddress), LookupState);
 	if (FunctionAddress == nullptr)
 	{
@@ -765,6 +770,7 @@ void CallstackTrace_CreateInternal(FMalloc* Malloc)
 void CallstackTrace_InitializeInternal()
 {
 	Modules_Initialize();
+	GModulesAreInitialized = true;
 }
 
 ////////////////////////////////////////////////////////////////////////////////
diff --git a/Engine/Source/Runtime/Core/Private/ProfilingDebugging/Microsoft/WindowsModuleDiagnostics.cpp b/Engine/Source/Runtime/Core/Private/ProfilingDebugging/Microsoft/WindowsModuleDiagnostics.cpp
--- a/Engine/Source/Runtime/Core/Private/ProfilingDebugging/Microsoft/WindowsModuleDiagnostics.cpp
+++ b/Engine/Source/Runtime/Core/Private/ProfilingDebugging/Microsoft/WindowsModuleDiagnostics.cpp
@@ -51,8 +51,8 @@ private:
 	void						OnDllNotification(unsigned int Reason, const void* DataPtr);
 	static FModuleTrace*		Instance;
 	SubscriberArray				Subscribers;
-	void*						CallbackCookie;
-	HeapId						ProgramHeapId;
+	void*						CallbackCookie = nullptr;
+	HeapId						ProgramHeapId = 0;
 };
 
 ////////////////////////////////////////////////////////////////////////////////
@@ -67,8 +67,11 @@ FModuleTrace::FModuleTrace(FMalloc* InMalloc)
 ////////////////////////////////////////////////////////////////////////////////
 FModuleTrace::~FModuleTrace()
 {
-	FNtDllFunction UnregisterFunc("LdrUnregisterDllNotification");
-	UnregisterFunc(CallbackCookie);
+	if (CallbackCookie)
+	{
+		FNtDllFunction UnregisterFunc("LdrUnregisterDllNotification");
+		UnregisterFunc(CallbackCookie);
+	}
 }
 
 ////////////////////////////////////////////////////////////////////////////////
