Low level tests adapter that can translate Catch2-defined test to a Functional Test

diff --git a/Engine/Source/Programs/UnrealBuildTool/Configuration/TestTargetRules.cs b/Engine/Source/Programs/UnrealBuildTool/Configuration/TestTargetRules.cs
--- a/Engine/Source/Programs/UnrealBuildTool/Configuration/TestTargetRules.cs
+++ b/Engine/Source/Programs/UnrealBuildTool/Configuration/TestTargetRules.cs
@@ -1,7 +1,10 @@
 // Copyright Epic Games, Inc. All Rights Reserved.
 
 using EpicGames.Core;
+using System;
 using System.IO;
+using System.Reflection;
+using System.Runtime.Serialization;
 using UnrealBuildBase;
 
 namespace UnrealBuildTool
@@ -100,6 +103,35 @@ namespace UnrealBuildTool
 			}
 		}
 
+		/// <summary>
+		/// Constructs a valid TestTargetRules instance from an existent TargetRules instance.
+		/// </summary>
+		public static TestTargetRules Create(TargetRules Rules, TargetInfo TargetInfo)
+		{
+			Type TestTargetRulesType = typeof(TestTargetRules);
+			TestTargetRules TestRules = (TestTargetRules)FormatterServices.GetUninitializedObject(TestTargetRulesType);
+
+			// Initialize the logger before calling the constructor
+			TestRules.Logger = Rules.Logger;
+
+			ConstructorInfo? Constructor = TestTargetRulesType.GetConstructor(new Type[] { typeof(TargetRules), typeof(TargetInfo) });
+			if (Constructor == null)
+			{
+				throw new BuildException("No constructor found on {0} which takes first argument of type TargetRules and second of type TargetInfo.", TestTargetRulesType.Name);
+			}
+
+			try
+			{
+				Constructor.Invoke(TestRules, new object[] { Rules, TargetInfo });
+			}
+			catch (Exception Ex)
+			{
+				throw new BuildException(Ex, "Unable to instantiate instance of '{0}' object type from compiled assembly '{1}'.  Unreal Build Tool creates an instance of your module's 'Rules' object in order to find out about your module's requirements.  The CLR exception details may provide more information:  {2}", TestTargetRulesType.Name, Path.GetFileNameWithoutExtension(TestTargetRulesType.Assembly?.Location), Ex.ToString());
+			}
+
+			return TestRules;
+		}
+
 		/// <summary>
 		/// Constructor for TestTargetRules based on existing target.
 		/// TestTargetRules is setup as a program and is linked monolithically.
diff --git a/Engine/Source/Programs/UnrealBuildTool/Configuration/UEBuildTarget.cs b/Engine/Source/Programs/UnrealBuildTool/Configuration/UEBuildTarget.cs
--- a/Engine/Source/Programs/UnrealBuildTool/Configuration/UEBuildTarget.cs
+++ b/Engine/Source/Programs/UnrealBuildTool/Configuration/UEBuildTarget.cs
@@ -3888,6 +3888,7 @@ namespace UnrealBuildTool
 			GlobalCompileEnvironment.Definitions.Add("WITH_PERF_AUTOMATION_TESTS=" + (bCompilePerfTests ? "1" : "0"));
 
 			GlobalCompileEnvironment.Definitions.Add(String.Format("WITH_LOW_LEVEL_TESTS={0}", Rules.WithLowLevelTests ? "1" : "0"));
+			GlobalCompileEnvironment.Definitions.Add(String.Format("WITH_TESTS={0}", Rules.WithLowLevelTests || bCompileDevTests || bCompilePerfTests ? "1" : "0"));
 
 			GlobalCompileEnvironment.Definitions.Add("UNICODE");
 			GlobalCompileEnvironment.Definitions.Add("_UNICODE");
diff --git a/Engine/Source/Programs/UnrealBuildTool/System/RulesAssembly.cs b/Engine/Source/Programs/UnrealBuildTool/System/RulesAssembly.cs
--- a/Engine/Source/Programs/UnrealBuildTool/System/RulesAssembly.cs
+++ b/Engine/Source/Programs/UnrealBuildTool/System/RulesAssembly.cs
@@ -673,7 +673,7 @@
 
 			if (IsTestTarget)
 			{
-				Rules = new TestTargetRules(Rules, TargetInfo);
+				Rules = TestTargetRules.Create(Rules, TargetInfo);
 			}
 
 			return Rules;
diff --git a/Engine/Source/Runtime/Core/Private/Misc/AutomationTest.cpp b/Engine/Source/Runtime/Core/Private/Misc/AutomationTest.cpp
--- a/Engine/Source/Runtime/Core/Private/Misc/AutomationTest.cpp
+++ b/Engine/Source/Runtime/Core/Private/Misc/AutomationTest.cpp
@@ -18,6 +18,38 @@
 
 DEFINE_LOG_CATEGORY_STATIC(LogAutomationTest, Warning, All);
 
+CORE_API const TMap<FString, EAutomationTestFlags::Type>& EAutomationTestFlags::GetTestFlagsMap()
+{
+	LLM_SCOPE_BYNAME(TEXT("AutomationTest"));
+	/** String to EAutomationTestFlags map */
+	static const TMap<FString, Type> FlagsMap = {
+		{ TEXT("EditorContext"), Type::EditorContext},
+		{ TEXT("ClientContext"), Type::ClientContext},
+		{ TEXT("ServerContext"), Type::ServerContext},
+		{ TEXT("CommandletContext"), Type::CommandletContext},
+		{ TEXT("ApplicationContextMask"), Type::ApplicationContextMask},
+		{ TEXT("NonNullRHI"), Type::NonNullRHI},
+		{ TEXT("RequiresUser"), Type::RequiresUser},
+		{ TEXT("FeatureMask"), Type::FeatureMask},
+		{ TEXT("Disabled"), Type::Disabled},
+		{ TEXT("CriticalPriority"), Type::CriticalPriority},
+		{ TEXT("HighPriority"), Type::HighPriority},
+		{ TEXT("HighPriorityAndAbove"), Type::HighPriorityAndAbove},
+		{ TEXT("MediumPriority"), Type::MediumPriority},
+		{ TEXT("MediumPriorityAndAbove"), Type::MediumPriorityAndAbove},
+		{ TEXT("LowPriority"), Type::LowPriority},
+		{ TEXT("PriorityMask"), Type::PriorityMask},
+		{ TEXT("SmokeFilter"), Type::SmokeFilter},
+		{ TEXT("EngineFilter"), Type::EngineFilter},
+		{ TEXT("ProductFilter"), Type::ProductFilter},
+		{ TEXT("PerfFilter"), Type::PerfFilter},
+		{ TEXT("StressFilter"), Type::StressFilter},
+		{ TEXT("NegativeFilter"), Type::NegativeFilter},
+		{ TEXT("FilterMask"), Type::FilterMask}
+	};
+	return FlagsMap;
+};
+
 /*
 	Determine the level that a log item should be written to the automation log based on the properties of the current test. 
 	only Display/Warning/Error are supported in the automation log so anything with NoLogging/Log will not be shown
diff --git a/Engine/Source/Runtime/Core/Public/Misc/AutomationTest.h b/Engine/Source/Runtime/Core/Public/Misc/AutomationTest.h
--- a/Engine/Source/Runtime/Core/Public/Misc/AutomationTest.h
+++ b/Engine/Source/Runtime/Core/Public/Misc/AutomationTest.h
@@ -109,38 +109,7 @@ struct EAutomationTestFlags
 		FilterMask = SmokeFilter | EngineFilter | ProductFilter | PerfFilter | StressFilter | NegativeFilter
 	};
 
-	static const TMap<FString, Type>& GetTestFlagsMap()
-	{
-		LLM_SCOPE_BYNAME(TEXT("AutomationTest"));
-		/** String to EAutomationTestFlags map */
-		static const TMap<FString, Type> FlagsMap = {
-			{ TEXT("EditorContext"), Type::EditorContext},
-			{ TEXT("ClientContext"), Type::ClientContext},
-			{ TEXT("ServerContext"), Type::ServerContext},
-			{ TEXT("CommandletContext"), Type::CommandletContext},
-			{ TEXT("ApplicationContextMask"), Type::ApplicationContextMask},
-			{ TEXT("NonNullRHI"), Type::NonNullRHI},
-			{ TEXT("RequiresUser"), Type::RequiresUser},
-			{ TEXT("FeatureMask"), Type::FeatureMask},
-			{ TEXT("Disabled"), Type::Disabled},
-			{ TEXT("CriticalPriority"), Type::CriticalPriority},
-			{ TEXT("HighPriority"), Type::HighPriority},
-			{ TEXT("HighPriorityAndAbove"), Type::HighPriorityAndAbove},
-			{ TEXT("MediumPriority"), Type::MediumPriority},
-			{ TEXT("MediumPriorityAndAbove"), Type::MediumPriorityAndAbove},
-			{ TEXT("LowPriority"), Type::LowPriority},
-			{ TEXT("PriorityMask"), Type::PriorityMask},
-			{ TEXT("SmokeFilter"), Type::SmokeFilter},
-			{ TEXT("EngineFilter"), Type::EngineFilter},
-			{ TEXT("ProductFilter"), Type::ProductFilter},
-			{ TEXT("PerfFilter"), Type::PerfFilter},
-			{ TEXT("SmokeFilter"), Type::SmokeFilter},
-			{ TEXT("StressFilter"), Type::StressFilter},
-			{ TEXT("NegativeFilter"), Type::NegativeFilter},
-			{ TEXT("FilterMask"), Type::FilterMask}
-		};
-		return FlagsMap;
-	}
+	static CORE_API const TMap<FString, Type>& GetTestFlagsMap();
 
 	static const Type FromString(FString Name)
 	{
@@ -3136,48 +3105,6 @@ public: \
 		virtual FString GetBeautifiedTestName() const override { return PrettyName; } \
 	};
 
-/**
-* This macro works the same way IMPLEMENT_SIMPLE_AUTOMATION_TEST_PRIVATE does with the following specifications:
-*	- TFalgs is expected to be declared as a string and it can include extra test tags along with EAutomationTestFlags flags
-		e.g. "[Flag_1]...[Flag_n][Tag_1][Tag_2]...[Tag_m]". The flags will be converted to an appropriate combined EAutomationTestFlags value
-*	- PrettyName is declared using double column notation "Namespace_1::Namespace_2:: ... ::Namespace_n::TestName"
-*		and it is converted to its dot notation "Namespace_1.Namespace_2. ... .Namespace_n.TestName"
-*	- It is mainly used when declaring tests that can run both as Low Level Tests with Catch2 or Functional Tests with FAutomationTestFramework
-*/
-#define IMPLEMENT_SIMPLE_AUTOMATION_TEST_PRIVATE_LLT( TClass, TBaseClass, PrettyName, TFlags, FileName, LineNumber ) \
-	class TClass : public TBaseClass \
-	{ \
-	public: \
-		TClass( const FString& InName ) \
-		:TBaseClass( InName, false ) {\
-			TestFlags = ExtractAutomationTestFlags(TFlags);\
-			PrettyNameDotNotation = PrettyName.Replace(TEXT("::"), TEXT("."));\
-			assert(TestFlags & EAutomationTestFlags::ApplicationContextMask, "AutomationTest has no application flag.  It shouldn't run.  See AutomationTest.h."); \
-			assert(	((TestFlags & EAutomationTestFlags::FilterMask) == EAutomationTestFlags::SmokeFilter) || \
-					((TestFlags & EAutomationTestFlags::FilterMask) == EAutomationTestFlags::EngineFilter) || \
-					((TestFlags & EAutomationTestFlags::FilterMask) == EAutomationTestFlags::ProductFilter) || \
-					((TestFlags & EAutomationTestFlags::FilterMask) == EAutomationTestFlags::PerfFilter) || \
-					((TestFlags & EAutomationTestFlags::FilterMask) == EAutomationTestFlags::StressFilter) || \
-					((TestFlags & EAutomationTestFlags::FilterMask) == EAutomationTestFlags::NegativeFilter), \
-					"All AutomationTests must have exactly 1 filter type specified.  See AutomationTest.h."); \
-		} \
-		virtual uint32 GetTestFlags() const override { return TestFlags; } \
-		virtual bool IsStressTest() const { return false; } \
-		virtual uint32 GetRequiredDeviceNum() const override { return 1; } \
-		virtual FString GetTestSourceFileName() const override { return FileName; } \
-		virtual int32 GetTestSourceFileLine() const override { return LineNumber; } \
-	protected: \
-		virtual void GetTests(TArray<FString>& OutBeautifiedNames, TArray <FString>& OutTestCommands) const override \
-		{ \
-			OutBeautifiedNames.Add(PrettyNameDotNotation)); \
-			OutTestCommands.Add(FString()); \
-		} \
-		virtual bool RunTest(const FString& Parameters) override; \
-		virtual FString GetBeautifiedTestName() const override { return PrettyNameDotNotation; } \
-	private: \
-		uint32 TestFlags;\
-		FString PrettyNameDotNotation;\
-	};
 
 
 #define IMPLEMENT_COMPLEX_AUTOMATION_TEST_PRIVATE( TClass, TBaseClass, PrettyName, TFlags, FileName, LineNumber ) \
diff --git a/Engine/Source/Runtime/Core/Public/Misc/LowLevelTestAdapter.h b/Engine/Source/Runtime/Core/Public/Misc/LowLevelTestAdapter.h
new file mode 100644
--- /dev/null
+++ b/Engine/Source/Runtime/Core/Public/Misc/LowLevelTestAdapter.h
@@ -0,0 +1,75 @@
+// Copyright Epic Games, Inc. All Rights Reserved.
+
+#pragma once
+
+#if !WITH_LOW_LEVEL_TESTS
+
+#include "CoreTypes.h"
+#include "Misc/AssertionMacros.h"
+#include "CoreMiscDefines.h"
+
+#define IMPLEMENT_SIMPLE_AUTOMATION_TEST_PRIVATE_LLT( TClass, PrettyName, TFlags, FileName, LineNumber ) \
+		class TClass : public FAutomationTestBase \
+		{ \
+		public:\
+			TClass( const FString& InName) \
+			: FAutomationTestBase( InName, false ) \
+			{ \
+				TestFlags = ExtractAutomationTestFlags(TFlags); \
+				PrettyNameDotNotation = FString(PrettyName).Replace(TEXT("::"), TEXT(".")); \
+				if (!(TestFlags & EAutomationTestFlags::ApplicationContextMask)) \
+				{ \
+					TestFlags |= EAutomationTestFlags::ApplicationContextMask; \
+				} \
+				if (!(TestFlags & EAutomationTestFlags::FilterMask)) \
+				{ \
+					TestFlags |= EAutomationTestFlags::EngineFilter; \
+				} \
+			} \
+			virtual uint32 GetTestFlags() const override { return TestFlags; } \
+			virtual bool IsStressTest() const { return false; } \
+			virtual uint32 GetRequiredDeviceNum() const override { return 1; } \
+			virtual FString GetTestSourceFileName() const override { return FileName; } \
+			virtual int32 GetTestSourceFileLine() const override { return LineNumber; } \
+		protected: \
+			virtual void GetTests(TArray<FString>& OutBeautifiedNames, TArray <FString>& OutTestCommands) const override \
+			{ \
+				OutBeautifiedNames.Add(PrettyNameDotNotation); \
+				OutTestCommands.Add(FString());\
+			} \
+			void TestBody(const FString& Parameters); \
+			virtual bool RunTest(const FString& Parameters) { \
+				TestBody(Parameters); \
+				return true; \
+			} \
+			virtual FString GetBeautifiedTestName() const override { return PrettyNameDotNotation; } \
+		private:\
+			uint32 TestFlags; \
+			FString PrettyNameDotNotation; \
+		};
+
+#define LLT_JOIN(Prefix, Counter) LLT_JOIN_INNER(Prefix, Counter)
+#define LLT_JOIN_INNER(Prefix, Counter) Prefix##Counter
+
+#define TEST_CASE_NAMED(TClass, StrName, PrettyName, TFlags) \
+		IMPLEMENT_SIMPLE_AUTOMATION_TEST_PRIVATE_LLT(TClass, PrettyName, TFlags, __FILE__, __LINE__) \
+		namespace \
+		{ \
+			TClass LLT_JOIN(TClass, Instance)(TEXT(StrName)); \
+		} \
+		void TClass::TestBody(const FString& Parameters)
+
+
+#define TEST_CASE_GENERATED_NAME_UNIQUE LLT_JOIN(FLLTAdaptedTest, __COUNTER__)
+#define LLT_STR(Macro) #Macro
+#define LLT_STR_EXPAND(Macro) LLT_STR(Macro)
+#define TEST_CASE_GENERATED_NAME_UNIQUE_STR LLT_STR_EXPAND(TEST_CASE_GENERATED_NAME_UNIQUE)
+#define TEST_CASE(PrettyName, TFlags) TEST_CASE_NAMED(TEST_CASE_GENERATED_NAME_UNIQUE, TEST_CASE_GENERATED_NAME_UNIQUE_STR, PrettyName, TFlags)
+
+#define CHECK(Expr) if (!(Expr)) { FAutomationTestFramework::Get().GetCurrentTest()->AddError(TEXT("Condition failed")); }
+#define CHECK_FALSE(Expr) if (Expr) { FAutomationTestFramework::Get().GetCurrentTest()->AddError(TEXT("Condition expected to return false but returned true")); }
+#define REQUIRE(Expr) if (!(Expr)) { FAutomationTestFramework::Get().GetCurrentTest()->AddError(TEXT("Required condition failed, interrupting test")); return; }
+
+#define SECTION(Text) AddInfo(TEXT(Text));
+
+#endif // !WITH_LOW_LEVEL_TESTS
\ No newline at end of file
diff --git a/Engine/Source/Runtime/Core/Public/Tests/TestHarness.h b/Engine/Source/Runtime/Core/Public/Tests/TestHarness.h
new file mode 100644
--- /dev/null
+++ b/Engine/Source/Runtime/Core/Public/Tests/TestHarness.h
@@ -0,0 +1,17 @@
+// Copyright Epic Games, Inc. All Rights Reserved.
+
+#pragma once
+
+/**
+ * Common test harness include file. Use this with tests that can be run as both low level tests or functional tests.
+ * The two types of frameworks are as follows:
+ * - Low level tests using Catch2, generates a monolithic test application, built separate
+ * - UE legacy test automation framework, tests are built into the target executable
+ */
+
+#if WITH_LOW_LEVEL_TESTS
+#include "LowLevelTestsRunner/Public/TestHarness.h"
+#elif WITH_AUTOMATION_TESTS
+#include "Misc/AutomationTest.h"
+#include "Misc/LowLevelTestAdapter.h"
+#endif
diff --git a/Engine/Source/Runtime/Core/Tests/Containers/DequeTest.cpp b/Engine/Source/Runtime/Core/Tests/Containers/DequeTest.cpp
--- a/Engine/Source/Runtime/Core/Tests/Containers/DequeTest.cpp
+++ b/Engine/Source/Runtime/Core/Tests/Containers/DequeTest.cpp
@@ -1,12 +1,13 @@
 // Copyright Epic Games, Inc. All Rights Reserved.
 
-#if WITH_LOW_LEVEL_TESTS
+#if WITH_TESTS
 
 #include "Containers/Deque.h"
 #include "CoreMinimal.h"
 #include "TokenTest.h"
 #include "TestUtils.h"
-#include "TestHarness.h"
+
+#include "Tests/TestHarness.h"
 
 namespace Deque
 { 
@@ -85,7 +86,7 @@ bool Test::EmplaceFirstPopLast(TDeque<int32Token>& Deque, int32 Count)
 // Unit tests
 //---------------------------------------------------------------------------------------------------------------------
 
-TEST_CASE("Core::Containers::TDeque::Reserve without data", "[Core][Containers][Deque]")
+TEST_CASE("System::Core::Containers::TDeque::Reserve without data", "[SmokeFilter][Core][Containers][Deque]")
 {
 	int32Token::Reset();
 	TDeque<int32Token> Deque;
@@ -99,7 +100,7 @@ TEST_CASE("Core::Containers::TDeque::Reserve without data", "[Core][Containers][
 	CHECK(int32Token::EvenConstructionDestructionCalls(0));
 }
 
-TEST_CASE("Core::Containers::TDeque::Reserve EmplaceLast single element", "[Core][Containers][Deque]")
+TEST_CASE("System::Core::Containers::TDeque::Reserve EmplaceLast single element", "[SmokeFilter][Core][Containers][Deque]")
 {
 	int32Token::Reset();
 	{
@@ -118,7 +119,7 @@ TEST_CASE("Core::Containers::TDeque::Reserve EmplaceLast single element", "[Core
 	CHECK(int32Token::EvenConstructionDestructionCalls(1));
 }
 
-TEST_CASE("Core::Containers::TDeque::Reset", "[Core][Containers][Deque]")
+TEST_CASE("System::Core::Containers::TDeque::Reset", "[SmokeFilter][Core][Containers][Deque]")
 {
 	int32Token::Reset();
 	{
@@ -139,7 +140,7 @@ TEST_CASE("Core::Containers::TDeque::Reset", "[Core][Containers][Deque]")
 	CHECK(int32Token::EvenConstructionDestructionCalls(1));
 }
 
-TEST_CASE("Core::Containers::TDeque::Empty", "[Core][Containers][Deque]")
+TEST_CASE("System::Core::Containers::TDeque::Empty", "[SmokeFilter][Core][Containers][Deque]")
 {
 	int32Token::Reset();
 	{
@@ -152,7 +153,7 @@ TEST_CASE("Core::Containers::TDeque::Empty", "[Core][Containers][Deque]")
 	CHECK(int32Token::EvenConstructionDestructionCalls(0));
 }
 
-TEST_CASE("Core::Containers::TDeque::Empty after single element EmplaceLast", "[Core][Containers][Deque]")
+TEST_CASE("System::Core::Containers::TDeque::Empty after single element EmplaceLast", "[SmokeFilter][Core][Containers][Deque]")
 {
 	int32Token::Reset();
 	{
@@ -169,7 +170,7 @@ TEST_CASE("Core::Containers::TDeque::Empty after single element EmplaceLast", "[
 	CHECK(int32Token::EvenConstructionDestructionCalls(1));
 }
 
-TEST_CASE("Core::Containers::TDeque::EmplaceLast single element", "[Core][Containers][Deque]")
+TEST_CASE("System::Core::Containers::TDeque::EmplaceLast single element", "[SmokeFilter][Core][Containers][Deque]")
 {
 	int32Token::Reset();
 	{
@@ -185,7 +186,7 @@ TEST_CASE("Core::Containers::TDeque::EmplaceLast single element", "[Core][Contai
 	CHECK(int32Token::EvenConstructionDestructionCalls(1));
 }
 
-TEST_CASE("Core::Containers::TDeque::EmplaceLast range to capacity", "[Core][Containers][Deque]")
+TEST_CASE("System::Core::Containers::TDeque::EmplaceLast range to capacity", "[SmokeFilter][Core][Containers][Deque]")
 {
 	int32Token::Reset();
 	{
@@ -204,7 +205,7 @@ TEST_CASE("Core::Containers::TDeque::EmplaceLast range to capacity", "[Core][Con
 	CHECK(int32Token::EvenConstructionDestructionCalls(Test::DefaultCapacity * 10));
 }
 
-TEST_CASE("Core::Containers::TDeque::EmplaceLast range past capacity", "[Core][Containers][Deque]")
+TEST_CASE("System::Core::Containers::TDeque::EmplaceLast range past capacity", "[SmokeFilter][Core][Containers][Deque]")
 {
 	int32Token::Reset();
 	{
@@ -226,7 +227,7 @@ TEST_CASE("Core::Containers::TDeque::EmplaceLast range past capacity", "[Core][C
 	CHECK(int32Token::EvenConstructionDestructionCalls(Test::DefaultCapacity + 1));
 }
 
-TEST_CASE("Core::Containers::TDeque::EmplaceFirst single element", "[Core][Containers][Deque]")
+TEST_CASE("System::Core::Containers::TDeque::EmplaceFirst single element", "[SmokeFilter][Core][Containers][Deque]")
 {
 	int32Token::Reset();
 	{
@@ -242,7 +243,7 @@ TEST_CASE("Core::Containers::TDeque::EmplaceFirst single element", "[Core][Conta
 	CHECK(int32Token::EvenConstructionDestructionCalls(1));
 }
 
-TEST_CASE("Core::Containers::TDeque::EmplaceFirst range to capacity", "[Core][Containers][Deque]")
+TEST_CASE("System::Core::Containers::TDeque::EmplaceFirst range to capacity", "[SmokeFilter][Core][Containers][Deque]")
 {
 	int32Token::Reset();
 	{
@@ -261,7 +262,7 @@ TEST_CASE("Core::Containers::TDeque::EmplaceFirst range to capacity", "[Core][Co
 	CHECK(int32Token::EvenConstructionDestructionCalls(Test::DefaultCapacity * 10));
 }
 
-TEST_CASE("Core::Containers::TDeque::PushLast single element (implicit move)", "[Core][Containers][Deque]")
+TEST_CASE("System::Core::Containers::TDeque::PushLast single element (implicit move)", "[SmokeFilter][Core][Containers][Deque]")
 {
 	int32Token::Reset();
 	{
@@ -280,7 +281,7 @@ TEST_CASE("Core::Containers::TDeque::PushLast single element (implicit move)", "
 	CHECK(int32Token::NumMoveConstructorCalls() == 1);
 }
 
-TEST_CASE("Core::Containers::TDeque::PushLast single element from move", "[Core][Containers][Deque]")
+TEST_CASE("System::Core::Containers::TDeque::PushLast single element from move", "[SmokeFilter][Core][Containers][Deque]")
 {
 	int32Token::Reset();
 	{
@@ -300,7 +301,7 @@ TEST_CASE("Core::Containers::TDeque::PushLast single element from move", "[Core]
 	CHECK(int32Token::NumMoveConstructorCalls() == 1);
 }
 
-TEST_CASE("Core::Containers::TDeque::PushLast single element from copy", "[Core][Containers][Deque]")
+TEST_CASE("System::Core::Containers::TDeque::PushLast single element from copy", "[SmokeFilter][Core][Containers][Deque]")
 {
 	int32Token::Reset();
 	{
@@ -320,7 +321,7 @@ TEST_CASE("Core::Containers::TDeque::PushLast single element from copy", "[Core]
 	CHECK(int32Token::NumMoveConstructorCalls() == 0);
 }
 
-TEST_CASE("Core::Containers::TDeque::PushFirst single element (implicit move)", "[Core][Containers][Deque]")
+TEST_CASE("System::Core::Containers::TDeque::PushFirst single element (implicit move)", "[SmokeFilter][Core][Containers][Deque]")
 {
 	int32Token::Reset();
 	{
@@ -339,7 +340,7 @@ TEST_CASE("Core::Containers::TDeque::PushFirst single element (implicit move)",
 	CHECK(int32Token::NumMoveConstructorCalls() == 1);
 }
 
-TEST_CASE("Core::Containers::TDeque::PushFirst single element from move", "[Core][Containers][Deque]")
+TEST_CASE("System::Core::Containers::TDeque::PushFirst single element from move", "[SmokeFilter][Core][Containers][Deque]")
 {
 	int32Token::Reset();
 	{
@@ -359,7 +360,7 @@ TEST_CASE("Core::Containers::TDeque::PushFirst single element from move", "[Core
 	CHECK(int32Token::NumMoveConstructorCalls() == 1);
 }
 
-TEST_CASE("Core::Containers::TDeque::PushFirst single element from copy", "[Core][Containers][Deque]")
+TEST_CASE("System::Core::Containers::TDeque::PushFirst single element from copy", "[SmokeFilter][Core][Containers][Deque]")
 {
 	int32Token::Reset();
 	{
@@ -384,7 +385,7 @@ void PopOne(TDeque<int32Token>& Deque)
 	Deque.PopFirst();
 }
 
-TEST_CASE("Core::Containers::TDeque::EmplaceLast/PopFirst single element", "[Core][Containers][Deque]")
+TEST_CASE("System::Core::Containers::TDeque::EmplaceLast/PopFirst single element", "[SmokeFilter][Core][Containers][Deque]")
 {
 	int32Token::Reset();
 	{
@@ -404,7 +405,7 @@ TEST_CASE("Core::Containers::TDeque::EmplaceLast/PopFirst single element", "[Cor
 	CHECK(int32Token::EvenConstructionDestructionCalls(1));
 }
 
-TEST_CASE("Core::Containers::TDeque::EmplaceLast/PopFirst single element multiple times causing head/tail wrap around", "[Core][Containers][Deque]")
+TEST_CASE("System::Core::Containers::TDeque::EmplaceLast/PopFirst single element multiple times causing head/tail wrap around", "[SmokeFilter][Core][Containers][Deque]")
 {
 	int32Token::Reset();
 	{
@@ -420,7 +421,7 @@ TEST_CASE("Core::Containers::TDeque::EmplaceLast/PopFirst single element multipl
 	CHECK(int32Token::EvenConstructionDestructionCalls(Test::DefaultCapacity * 2));
 }
 
-TEST_CASE("Core::Containers::TDeque::EmplaceLast/PopFirst range without head/tail wrap around", "[Core][Containers][Deque]")
+TEST_CASE("System::Core::Containers::TDeque::EmplaceLast/PopFirst range without head/tail wrap around", "[SmokeFilter][Core][Containers][Deque]")
 {
 	int32Token::Reset();
 	{
@@ -432,7 +433,7 @@ TEST_CASE("Core::Containers::TDeque::EmplaceLast/PopFirst range without head/tai
 	CHECK(int32Token::EvenConstructionDestructionCalls(Test::DefaultCapacity));
 }
 
-TEST_CASE("Core::Containers::TDeque::EmplaceLast/PopFirst range with reallocation without head/tail wrap around", "[Core][Containers][Deque]")
+TEST_CASE("System::Core::Containers::TDeque::EmplaceLast/PopFirst range with reallocation without head/tail wrap around", "[SmokeFilter][Core][Containers][Deque]")
 {
 	int32Token::Reset();
 	{
@@ -445,7 +446,7 @@ TEST_CASE("Core::Containers::TDeque::EmplaceLast/PopFirst range with reallocatio
 	CHECK(int32Token::EvenConstructionDestructionCalls(Test::DefaultCapacity + 1));
 }
 
-TEST_CASE("Core::Containers::TDeque::EmplaceLast/PopFirst range with head/tail wrap around", "[Core][Containers][Deque]")
+TEST_CASE("System::Core::Containers::TDeque::EmplaceLast/PopFirst range with head/tail wrap around", "[SmokeFilter][Core][Containers][Deque]")
 {
 	int32Token::Reset();
 	{
@@ -460,7 +461,7 @@ TEST_CASE("Core::Containers::TDeque::EmplaceLast/PopFirst range with head/tail w
 	CHECK(int32Token::EvenConstructionDestructionCalls(Test::DefaultCapacity * (Test::DefaultCapacity - 1)));
 }
 
-TEST_CASE("Core::Containers::TDeque::EmplaceFirst/PopLast single element", "[Core][Containers][Deque]")
+TEST_CASE("System::Core::Containers::TDeque::EmplaceFirst/PopLast single element", "[SmokeFilter][Core][Containers][Deque]")
 {
 	int32Token::Reset();
 	{
@@ -480,7 +481,7 @@ TEST_CASE("Core::Containers::TDeque::EmplaceFirst/PopLast single element", "[Cor
 	CHECK(int32Token::EvenConstructionDestructionCalls(1));
 }
 
-TEST_CASE("Core::Containers::TDeque::EmplaceFirst/PopLast range", "[Core][Containers][Deque]")
+TEST_CASE("System::Core::Containers::TDeque::EmplaceFirst/PopLast range", "[SmokeFilter][Core][Containers][Deque]")
 {
 	int32Token::Reset();
 	{
@@ -492,7 +493,7 @@ TEST_CASE("Core::Containers::TDeque::EmplaceFirst/PopLast range", "[Core][Contai
 	CHECK(int32Token::EvenConstructionDestructionCalls(Test::DefaultCapacity));
 }
 
-TEST_CASE("Core::Containers::TDeque::EmplaceFirst/PopLast range with reallocation", "[Core][Containers][Deque]")
+TEST_CASE("System::Core::Containers::TDeque::EmplaceFirst/PopLast range with reallocation", "[SmokeFilter][Core][Containers][Deque]")
 {
 	int32Token::Reset();
 	{
@@ -505,7 +506,7 @@ TEST_CASE("Core::Containers::TDeque::EmplaceFirst/PopLast range with reallocatio
 	CHECK(int32Token::EvenConstructionDestructionCalls(Test::DefaultCapacity + 1));
 }
 
-TEST_CASE("Core::Containers::TDeque::TryPopFirst", "[Core][Containers][Deque]")
+TEST_CASE("System::Core::Containers::TDeque::TryPopFirst", "[SmokeFilter][Core][Containers][Deque]")
 {
 	int32Token::Reset();
 	{
@@ -528,7 +529,7 @@ TEST_CASE("Core::Containers::TDeque::TryPopFirst", "[Core][Containers][Deque]")
 	CHECK(int32Token::EvenConstructionDestructionCalls(Test::DefaultCapacity + 1));  // + 1 for Result
 }
 
-TEST_CASE("Core::Containers::TDeque::TryPopFirst with reallocation", "[Core][Containers][Deque]")
+TEST_CASE("System::Core::Containers::TDeque::TryPopFirst with reallocation", "[SmokeFilter][Core][Containers][Deque]")
 {
 	int32Token::Reset();
 	{
@@ -554,7 +555,7 @@ TEST_CASE("Core::Containers::TDeque::TryPopFirst with reallocation", "[Core][Con
 	CHECK(int32Token::EvenConstructionDestructionCalls(Test::DefaultCapacity + 2));
 }
 
-TEST_CASE("Core::Containers::TDeque::TryPopLast", "[Core][Containers][Deque]")
+TEST_CASE("System::Core::Containers::TDeque::TryPopLast", "[SmokeFilter][Core][Containers][Deque]")
 {
 	int32Token::Reset();
 	{
@@ -577,7 +578,7 @@ TEST_CASE("Core::Containers::TDeque::TryPopLast", "[Core][Containers][Deque]")
 	CHECK(int32Token::EvenConstructionDestructionCalls(Test::DefaultCapacity + 1));  // + 1 for Result
 }
 
-TEST_CASE("Core::Containers::TDeque::TryPopLast with reallocation", "[Core][Containers][Deque]")
+TEST_CASE("System::Core::Containers::TDeque::TryPopLast with reallocation", "[SmokeFilter][Core][Containers][Deque]")
 {
 	int32Token::Reset();
 	{
@@ -603,7 +604,7 @@ TEST_CASE("Core::Containers::TDeque::TryPopLast with reallocation", "[Core][Cont
 	CHECK(int32Token::EvenConstructionDestructionCalls(Test::DefaultCapacity + 2));
 }
 
-TEST_CASE("Core::Containers::TDeque::Comparison simple", "[Core][Containers][Deque]")
+TEST_CASE("System::Core::Containers::TDeque::Comparison simple", "[SmokeFilter][Core][Containers][Deque]")
 {
 	int32Token::Reset();
 	{
@@ -622,7 +623,7 @@ TEST_CASE("Core::Containers::TDeque::Comparison simple", "[Core][Containers][Deq
 	CHECK(int32Token::EvenConstructionDestructionCalls(Test::DefaultCapacity * 2));
 }
 
-TEST_CASE("Core::Containers::TDeque::Comparison with head/tail wrap around", "[Core][Containers][Deque]")
+TEST_CASE("System::Core::Containers::TDeque::Comparison with head/tail wrap around", "[SmokeFilter][Core][Containers][Deque]")
 {
 	int32Token::Reset();
 	{
@@ -650,7 +651,7 @@ TEST_CASE("Core::Containers::TDeque::Comparison with head/tail wrap around", "[C
 	CHECK(int32Token::EvenConstructionDestructionCalls(Test::DefaultCapacity * 6));
 }
 
-TEST_CASE("Core::Containers::TDeque::Copy simple", "[Core][Containers][Deque]")
+TEST_CASE("System::Core::Containers::TDeque::Copy simple", "[SmokeFilter][Core][Containers][Deque]")
 {
 	int32Token::Reset();
 	{
@@ -668,7 +669,7 @@ TEST_CASE("Core::Containers::TDeque::Copy simple", "[Core][Containers][Deque]")
 	CHECK(int32Token::EvenConstructionDestructionCalls(Test::DefaultCapacity * 2));
 }
 
-TEST_CASE("Core::Containers::TDeque::Copy with head/tail wrap around", "[Core][Containers][Deque]")
+TEST_CASE("System::Core::Containers::TDeque::Copy with head/tail wrap around", "[SmokeFilter][Core][Containers][Deque]")
 {
 	int32Token::Reset();
 	{
@@ -691,7 +692,7 @@ TEST_CASE("Core::Containers::TDeque::Copy with head/tail wrap around", "[Core][C
 	CHECK(int32Token::EvenConstructionDestructionCalls(Test::DefaultCapacity * 9));
 }
 
-TEST_CASE("Core::Containers::TDeque::Copy variable size with head/tail wrap around", "[Core][Containers][Deque]")
+TEST_CASE("System::Core::Containers::TDeque::Copy variable size with head/tail wrap around", "[SmokeFilter][Core][Containers][Deque]")
 {
 	int32Token::Reset();
 	{
@@ -719,7 +720,7 @@ TEST_CASE("Core::Containers::TDeque::Copy variable size with head/tail wrap arou
 	CHECK(int32Token::EvenConstructionDestructionCalls());
 }
 
-TEST_CASE("Core::Containers::TDeque::Move simple", "[Core][Containers][Deque]")
+TEST_CASE("System::Core::Containers::TDeque::Move simple", "[SmokeFilter][Core][Containers][Deque]")
 {
 	int32Token::Reset();
 	{
@@ -743,7 +744,7 @@ TEST_CASE("Core::Containers::TDeque::Move simple", "[Core][Containers][Deque]")
 	CHECK(int32Token::EvenConstructionDestructionCalls(Test::DefaultCapacity + 1));
 }
 
-TEST_CASE("Core::Containers::TDeque::Move with head/tail wrap around", "[Core][Containers][Deque]")
+TEST_CASE("System::Core::Containers::TDeque::Move with head/tail wrap around", "[SmokeFilter][Core][Containers][Deque]")
 {
 	int32Token::Reset();
 	{
@@ -771,7 +772,7 @@ TEST_CASE("Core::Containers::TDeque::Move with head/tail wrap around", "[Core][C
 	CHECK(int32Token::EvenConstructionDestructionCalls(Test::DefaultCapacity * 6));
 }
 
-TEST_CASE("Core::Containers::TDeque::Move variable size with head/tail wrap around", "[Core][Containers][Deque]")
+TEST_CASE("System::Core::Containers::TDeque::Move variable size with head/tail wrap around", "[SmokeFilter][Core][Containers][Deque]")
 {
 	int32Token::Reset();
 	{
@@ -803,7 +804,7 @@ TEST_CASE("Core::Containers::TDeque::Move variable size with head/tail wrap arou
 	CHECK(int32Token::EvenConstructionDestructionCalls());
 }
 
-TEST_CASE("Core::Containers::TDeque::Iteration without head/tail wrap around", "[Core][Containers][Deque]")
+TEST_CASE("System::Core::Containers::TDeque::Iteration without head/tail wrap around", "[SmokeFilter][Core][Containers][Deque]")
 {
 	int32Token::Reset();
 	{
@@ -828,7 +829,7 @@ TEST_CASE("Core::Containers::TDeque::Iteration without head/tail wrap around", "
 	CHECK(int32Token::EvenConstructionDestructionCalls(Test::DefaultCapacity));
 }
 
-TEST_CASE("Core::Containers::TDeque::Iteration with head/tail wrap around", "[Core][Containers][Deque]")
+TEST_CASE("System::Core::Containers::TDeque::Iteration with head/tail wrap around", "[SmokeFilter][Core][Containers][Deque]")
 {
 	int32Token::Reset();
 	{
@@ -858,7 +859,7 @@ TEST_CASE("Core::Containers::TDeque::Iteration with head/tail wrap around", "[Co
 	CHECK(int32Token::EvenConstructionDestructionCalls(Test::DefaultCapacity * (Test::DefaultCapacity + 1)));
 }
 
-TEST_CASE("Core::Containers::TDeque::Iterator arithmetic", "[Core][Containers][Deque]")
+TEST_CASE("System::Core::Containers::TDeque::Iterator arithmetic", "[SmokeFilter][Core][Containers][Deque]")
 {
 	TDeque<int32Token> Deque;
 	Deque.Reserve(Test::DefaultCapacity);
@@ -885,7 +886,7 @@ TEST_CASE("Core::Containers::TDeque::Iterator arithmetic", "[Core][Containers][D
 	}
 }
 
-TEST_CASE("Core::Containers::TDeque::Construct from std initializer_list", "[Core][Containers][Deque]")
+TEST_CASE("System::Core::Containers::TDeque::Construct from std initializer_list", "[SmokeFilter][Core][Containers][Deque]")
 {
 	int32Token::Reset();
 	{
@@ -901,7 +902,7 @@ TEST_CASE("Core::Containers::TDeque::Construct from std initializer_list", "[Cor
 	CHECK(int32Token::EvenConstructionDestructionCalls(6 * 2 + 1));
 }
 
-TEST_CASE("Core::Containers::TDeque::Construct from empty std initializer_list", "[Core][Containers][Deque]")
+TEST_CASE("System::Core::Containers::TDeque::Construct from empty std initializer_list", "[SmokeFilter][Core][Containers][Deque]")
 {
 	int32Token::Reset();
 	{
@@ -911,7 +912,7 @@ TEST_CASE("Core::Containers::TDeque::Construct from empty std initializer_list",
 	CHECK(int32Token::EvenConstructionDestructionCalls(0));
 }
 
-TEST_CASE("Core::Containers::TDeque::Assign from std initializer_list", "[Core][Containers][Deque]")
+TEST_CASE("System::Core::Containers::TDeque::Assign from std initializer_list", "[SmokeFilter][Core][Containers][Deque]")
 {
 	int32Token::Reset();
 	{
