Revise macOS file locking solution, via flock(), to suppport multiple instances for shared and exclusive access to the same log file.   Fixes overwrite of existing instance log file by second instance.  Read will get a shared lock - try to get exclusive lock on a file when opening for write to prove nothing else is currently using the file - downgrade lock to shared if bAllowRead is set.

#rnx

diff --git a/Engine/Source/Runtime/Core/Private/Apple/ApplePlatformFile.cpp b/Engine/Source/Runtime/Core/Private/Apple/ApplePlatformFile.cpp
--- a/Engine/Source/Runtime/Core/Private/Apple/ApplePlatformFile.cpp
+++ b/Engine/Source/Runtime/Core/Private/Apple/ApplePlatformFile.cpp
@@ -568,9 +568,9 @@ IFileHandle* FApplePlatformFile::OpenRead(const TCHAR* Filename, bool bAllowWrit
 	if (Handle != -1)
 	{
 		TRACE_PLATFORMFILE_END_OPEN(Handle);
-#if PLATFORM_MAC && !UE_BUILD_SHIPPING
+#if PLATFORM_MAC && UE_EDITOR && !UE_BUILD_SHIPPING
 		// No blocking attempt shared lock, failure means we should not have opened the file for reading, protect against multiple instances and client/server versions
-		if(!bAllowWrite && flock(Handle, LOCK_NB | LOCK_SH) == -1)
+		if(flock(Handle, LOCK_NB | LOCK_SH) != 0)
 		{
 			TRACE_PLATFORMFILE_BEGIN_CLOSE(Handle);
 			int CloseResult = close(Handle);
@@ -623,8 +623,8 @@ IFileHandle* FApplePlatformFile::OpenWrite(const TCHAR* Filename, bool bAppend,
 	{
 		TRACE_PLATFORMFILE_END_OPEN(Handle);
 #if PLATFORM_MAC && UE_EDITOR && !UE_BUILD_SHIPPING
-		// No blocking attempt exclusive lock, failure means we should not have opened the file for writing, protect against multiple instances and client/server versions
-		if(!bAllowRead && flock(Handle, LOCK_NB | LOCK_EX) == -1)
+		// No blocking attempt EXclusive lock, failure means we should not have opened the file for writing, protect against multiple instances and client/server versions
+		if(flock(Handle, LOCK_NB | LOCK_EX) != 0)
 		{
 			TRACE_PLATFORMFILE_BEGIN_CLOSE(Handle);
 			int CloseResult = close(Handle);
@@ -642,6 +642,12 @@ IFileHandle* FApplePlatformFile::OpenWrite(const TCHAR* Filename, bool bAppend,
 #endif
 			return nullptr;
 		}
+		
+		// We have created the writer, if reading is required downgrade the lock to SHared
+		if(bAllowRead)
+		{
+			flock(Handle, LOCK_NB | LOCK_SH);
+		}
 #endif
 		
 		// Truncate after locking as lock may fail - don't use O_TRUNC in open flags
