Default initialize FAutomationTestBase::bSuppressLogs to false and honor its value if a derived class change it.

PR #9263: Fix AutomationTest not using bSuppressLogs properly (Contributed by [REDACTED_PERSON])

diff --git a/Engine/Source/Runtime/Core/Public/Misc/AutomationTest.h b/Engine/Source/Runtime/Core/Public/Misc/AutomationTest.h
--- a/Engine/Source/Runtime/Core/Public/Misc/AutomationTest.h
+++ b/Engine/Source/Runtime/Core/Public/Misc/AutomationTest.h
@@ -1411,7 +1411,7 @@ public:
 	 */
 	virtual bool SuppressLogs()
 	{
-		return false;
+		return bSuppressLogs;
 	}
 
 	/**
@@ -1815,7 +1815,7 @@ protected:
 	bool bRunOnSeparateThread;
 
 	/** Flag to suppress logs */
-	bool bSuppressLogs;
+	bool bSuppressLogs = false;
 
 	/** Name of the test */
 	FString TestName;
