Fix test harness macro that's undefined for most tools.
#rnx

diff --git a/Engine/Source/Runtime/Core/Public/Tests/TestHarness.h b/Engine/Source/Runtime/Core/Public/Tests/TestHarness.h
--- a/Engine/Source/Runtime/Core/Public/Tests/TestHarness.h
+++ b/Engine/Source/Runtime/Core/Public/Tests/TestHarness.h
@@ -11,7 +11,7 @@
 
 #if WITH_LOW_LEVEL_TESTS
 #include "LowLevelTestsRunner/Public/TestHarness.h"
-#elif WITH_AUTOMATION_TESTS
+#elif defined(WITH_AUTOMATION_TESTS) || (WITH_DEV_AUTOMATION_TESTS || WITH_PERF_AUTOMATION_TESTS)
 #include "Misc/AutomationTest.h"
 #include "Misc/LowLevelTestAdapter.h"
 #endif
