Refactored all threads callstack capturing on Android with an aim to reduce timeouts

diff --git a/Engine/Source/Runtime/Core/Private/Android/AndroidPlatformCrashContext.cpp b/Engine/Source/Runtime/Core/Private/Android/AndroidPlatformCrashContext.cpp
--- a/Engine/Source/Runtime/Core/Private/Android/AndroidPlatformCrashContext.cpp
+++ b/Engine/Source/Runtime/Core/Private/Android/AndroidPlatformCrashContext.cpp
@@ -226,14 +226,18 @@ void FAndroidCrashContext::StoreCrashInfo(bool bWriteLog) const
 
 //Create a separate file containing thread context info (callstacks etc) in xml form
 // This is added to the crash report xml at during pre-processing time.
-void FAndroidCrashContext::DumpAllThreadCallstacks() const
+void FAndroidCrashContext::DumpAllThreadCallstacks(FAsyncThreadBackTrace* BackTrace, int NumThreads) const
 {
+	if (NumThreads == 0)
+	{
+		return;
+	}
+
 	char FilePath[FAndroidCrashContext::CrashReportMaxPathSize] = { 0 };
 	FCStringAnsi::Strcpy(FilePath, ReportDirectory);
 	FCStringAnsi::Strcat(FilePath, FAndroidCrashContext::CrashReportMaxPathSize, "/AllThreads.txt");
 	TArray<FCrashStackFrame> CrashStackFrames;
 	CrashStackFrames.Empty(32);
-	uint32 CallstacksRecorded = 0;
 	int DestHandle = open(FilePath, O_CREAT | O_WRONLY, S_IRUSR | S_IWUSR);
 	if (DestHandle >= 0)
 	{
@@ -248,29 +252,28 @@ void FAndroidCrashContext::DumpAllThreadCallstacks() const
 			write(FileHandle, "\n", 1);
 		};
 
-		auto WriteThreadEntry = [Writeln, Write, this, DestHandle, &CrashStackFrames, &CallstacksRecorded, CurrentThreadID](uint32 ThreadID, const char* ThreadName)
+		auto CaptureCallstack = [CurrentThreadID](FAsyncThreadBackTrace* BackTrace)
 		{
-			// Capture the stack trace
-			static const int StackTraceMaxDepth = 100;
-			uint64 StackTrace[StackTraceMaxDepth];
-			FMemory::Memzero(StackTrace);
-			uint32 Depth = 0;
-			if (CurrentThreadID == ThreadID)
+			if (CurrentThreadID == BackTrace->ThreadID)
 			{
-				Depth = FPlatformStackWalk::CaptureStackBackTrace(StackTrace, StackTraceMaxDepth);
+				BackTrace->Depth = FPlatformStackWalk::CaptureStackBackTrace(BackTrace->BackTrace, BackTrace->StackTraceMaxDepth);
+				BackTrace->Flag.store(1, std::memory_order_release);
 			}
 			else
 			{
-				Depth = FPlatformStackWalk::CaptureThreadStackBackTrace(ThreadID, StackTrace, StackTraceMaxDepth);
+				FPlatformStackWalk::CaptureThreadStackBackTraceAsync(BackTrace);
 			}
+		};
 
-			if (Depth)
+		auto WriteThreadEntry = [Writeln, Write, this, DestHandle, &CrashStackFrames](FAsyncThreadBackTrace* BackTrace)
+		{
+			if (BackTrace->Depth)
 			{
 				ANSICHAR Line[256];
 				Writeln(DestHandle, "<Thread>");
 				Write(DestHandle, "<CallStack>");
 				// Write stack
-				GetPortableCallStack(StackTrace, Depth, CrashStackFrames);
+				GetPortableCallStack(BackTrace->BackTrace, BackTrace->Depth, CrashStackFrames);
 				for (const FCrashStackFrame& CrashStackFrame : CrashStackFrames)
 				{
 					FCStringAnsi::Strncpy(Line, TCHAR_TO_UTF8(*CrashStackFrame.ModuleName), UE_ARRAY_COUNT(Line));
@@ -284,34 +287,78 @@ void FAndroidCrashContext::DumpAllThreadCallstacks() const
 				Writeln(DestHandle, "<IsCrashed>false</IsCrashed>");
 				Writeln(DestHandle, "<Registers/>");
 				// write Thread Id
-				FCStringAnsi::Strncpy(Line, ItoANSI((uint64)ThreadID, (uint64)10), UE_ARRAY_COUNT(Line));
+				FCStringAnsi::Strncpy(Line, ItoANSI((uint64)BackTrace->ThreadID, (uint64)10), UE_ARRAY_COUNT(Line));
 				Write(DestHandle, "<ThreadID>");
 				Write(DestHandle, Line);
 				Writeln(DestHandle, "</ThreadID>");
 
 				// write Thread Name
 				Write(DestHandle, "<ThreadName>");
-				FCStringAnsi::Strncpy(Line, ThreadName, UE_ARRAY_COUNT(Line));
+				FMemory::Memcpy(Line, BackTrace->ThreadName, FAsyncThreadBackTrace::MaxThreadName);
+				Line[FAsyncThreadBackTrace::MaxThreadName] = 0;
 				Write(DestHandle, Line);
 				Writeln(DestHandle, "</ThreadName>");
 
 				Writeln(DestHandle, "</Thread>");
-				CallstacksRecorded++;
 			}
 		};
 
 		Writeln(DestHandle, "<Threads>");
 
+		uint32 CaptureCallstacks = 0;
+		uint32 CallstacksRecorded = 0;
+		for (int i = 0; i < NumThreads; ++i)
+		{
+			BackTrace[i].Flag.store(0, std::memory_order_relaxed);
+			BackTrace[i].Depth = 0;
+		}
+
 		// On android the Game thread is that which calls android_main entry point. 
 		// Explicitly call it here as the thread manager is not aware of it.
-		WriteThreadEntry(GGameThreadId, "GameThread");
+		if (CrashingThreadId != GGameThreadId)
+		{
+			BackTrace[0].ThreadID = GGameThreadId;
+			FCStringAnsi::Strcpy(BackTrace[0].ThreadName, "GameThread");
+			CaptureCallstack(&BackTrace[0]);
+			++CaptureCallstacks;
+		}
 
-		// For each thread append it's info to the file.
-		FThreadManager::Get().ForEachThread([WriteThreadEntry, this, DestHandle, &CrashStackFrames, &CallstacksRecorded](uint32 ThreadID, FRunnableThread* Runnable)
+		FThreadManager::Get().ForEachThread([CaptureCallstack, BackTrace, &CaptureCallstacks, &NumThreads, &CrashingThreadId = CrashingThreadId](uint32 ThreadID, FRunnableThread* Runnable)
 		{
-			WriteThreadEntry(ThreadID, TCHAR_TO_UTF8(*Runnable->GetThreadName()));
+			if ((CaptureCallstacks < NumThreads) && (CrashingThreadId != ThreadID))
+			{
+				FAsyncThreadBackTrace* Trace = &BackTrace[CaptureCallstacks];
+				Trace->ThreadID = ThreadID;
+				const FString& ThreadName = Runnable->GetThreadName();
+				FMemory::Memcpy(Trace->ThreadName, TCHAR_TO_UTF8(*ThreadName), FMath::Min(ThreadName.Len() + 1, FAsyncThreadBackTrace::MaxThreadName));
+
+				CaptureCallstack(Trace);
+				++CaptureCallstacks;
+			}
 		});
 
+		const float PollTime = 0.001f;
+		extern float GThreadCallStackMaxWait;
+
+		for (float CurrentTime = 0; CurrentTime <= GThreadCallStackMaxWait; CurrentTime += PollTime)
+		{
+			CallstacksRecorded = 0;
+			for (int i = 0; i < CaptureCallstacks; ++i)
+			{
+				CallstacksRecorded += BackTrace[i].Flag.load(std::memory_order_acquire);
+			}
+			if (CallstacksRecorded == CaptureCallstacks)
+			{
+				break;
+			}
+			FPlatformProcess::SleepNoStats(PollTime);
+		}
+
+		for (int i = 0; i < CaptureCallstacks; ++i)
+		{
+			WriteThreadEntry(&BackTrace[i]);
+		}
+
 		Writeln(DestHandle, "</Threads>");
 		close(DestHandle);
 
diff --git a/Engine/Source/Runtime/Core/Private/Android/AndroidPlatformMisc.cpp b/Engine/Source/Runtime/Core/Private/Android/AndroidPlatformMisc.cpp
--- a/Engine/Source/Runtime/Core/Private/Android/AndroidPlatformMisc.cpp
+++ b/Engine/Source/Runtime/Core/Private/Android/AndroidPlatformMisc.cpp
@@ -1136,7 +1136,7 @@ private:
 		FPlatformAtomics::AtomicStore(&handling_signal, 0);
 	}
 
-	static void HandleTargetSignal(int Signal, siginfo* Info, void* Context)
+	static void HandleTargetSignal(int Signal, siginfo* Info, void* Context, uint32 CrashingThreadId)
 	{
 		FPlatformStackWalk::HandleBackTraceSignal(Info, Context);
 	}
@@ -1255,8 +1255,8 @@ FString FAndroidMisc::GetFatalSignalMessage(int Signal, siginfo* Info)
 }
 
 // Making the signal handler available to track down issues with failing crash handler.
-static void (*GFatalSignalHandlerOverrideFunc)(int Signal, struct siginfo* Info, void* Context) = nullptr;
-void FAndroidMisc::OverrideFatalSignalHandler(void (*FatalSignalHandlerOverrideFunc)(int Signal, struct siginfo* Info, void* Context))
+static void (*GFatalSignalHandlerOverrideFunc)(int Signal, struct siginfo* Info, void* Context, uint32 CrashingThreadId) = nullptr;
+void FAndroidMisc::OverrideFatalSignalHandler(void (*FatalSignalHandlerOverrideFunc)(int Signal, struct siginfo* Info, void* Context, uint32 CrashingThreadId))
 {
 	GFatalSignalHandlerOverrideFunc = FatalSignalHandlerOverrideFunc;
 }
@@ -1309,11 +1309,11 @@ protected:
 		raise(Signal);
 	}
 
-	static void HandleTargetSignal(int Signal, siginfo* Info, void* Context)
+	static void HandleTargetSignal(int Signal, siginfo* Info, void* Context, uint32 CrashingThreadId)
 	{
 		if (GFatalSignalHandlerOverrideFunc)
 		{
-			GFatalSignalHandlerOverrideFunc(Signal, Info, Context);
+			GFatalSignalHandlerOverrideFunc(Signal, Info, Context, CrashingThreadId);
 		}
 		else
 		{
@@ -1323,7 +1323,7 @@ protected:
 			FString Message = FAndroidMisc::GetFatalSignalMessage(Signal, Info);
 			FAndroidCrashContext CrashContext(ECrashContextType::Crash, *Message);
 
-			CrashContext.InitFromSignal(Signal, Info, Context);
+			CrashContext.InitFromSignal(Signal, Info, Context, CrashingThreadId);
 			CrashContext.CaptureCrashInfo();
 			if (GCrashHandlerPointer)
 			{
diff --git a/Engine/Source/Runtime/Core/Private/Android/AndroidPlatformStackWalk.cpp b/Engine/Source/Runtime/Core/Private/Android/AndroidPlatformStackWalk.cpp
--- a/Engine/Source/Runtime/Core/Private/Android/AndroidPlatformStackWalk.cpp
+++ b/Engine/Source/Runtime/Core/Private/Android/AndroidPlatformStackWalk.cpp
@@ -340,34 +340,33 @@ static FAutoConsoleVariableRef CVarAndroidPlatformThreadCallStackRequestMaxWait(
 	GThreadCallStackRequestMaxWait,
 	TEXT("The number of seconds to spin before an individual back trace has timed out."));
 
-static float GThreadCallStackMaxWait = 5.0f;
+float GThreadCallStackMaxWait = 5.0f;
 static TAutoConsoleVariable<float> CVarAndroidPlatformThreadCallStackMaxWait(
 	TEXT("AndroidPlatformThreadStackWalk.MaxWait"),
 	GThreadCallStackMaxWait,
 	TEXT("The number of seconds allowed to spin before killing the process, with the assumption the back trace handler has hung."));
 
 #if ANDROID_HAS_RTSIGNALS
-/** Passed in through sigqueue for gathering of a callstack from a signal */
-struct ThreadStackUserData
-{
-	uint64* BackTrace;
-	int32 BackTraceCount;
-	SIZE_T CallStackSize;
-};
-
-int32 ThreadStackBackTraceStatus = 0;
-static const int32 ThreadStackBackTraceCurrentStatus_RUNNING = -2;
-static const int32 ThreadStackBackTraceCurrentStatus_DONE = -3;
-
-static ThreadStackUserData SignalThreadStackUserData;
+static FAsyncThreadBackTrace SignalThreadStackUserData;
+/*Async stack backtracing capturing works needs to work in two modes:
+ * 1. Serial - when we fire a request and wait for it (or time out) before firing next request.
+ *    In case of a time out, serial request uses static SignalThreadStackUserData, so it's possible we can send a second request
+ *    after the one that timed out and first request will be processed at this time, providing a wrong data to the second request.
+ *    To prevent this from happening, we check if there's no serial request in flight before firing one
+ * 2. Broadcast - when we fire a bunch of requests one ofter the other and then wait on all of them.
+ *    This mode is not immune to an issue described before, but this is called only in a crash handler, so we are not going to fire a new request afterwards.
+ */ 
+static std::atomic<FAsyncThreadBackTrace*> InFlightSerialRequest(nullptr);
 
 // the callback when THREAD_CALLSTACK_GENERATOR is being processed.
 void FAndroidPlatformStackWalk::HandleBackTraceSignal(siginfo* Info, void* Context)
 {
-	if (FPlatformAtomics::InterlockedCompareExchange(&ThreadStackBackTraceStatus, ThreadStackBackTraceCurrentStatus_RUNNING, Info->si_value.sival_int) == Info->si_value.sival_int)
+	FAsyncThreadBackTrace* BackTrace = (FAsyncThreadBackTrace*)Info->si_value.sival_ptr;
+	BackTrace->Depth = FPlatformStackWalk::CaptureStackBackTrace(BackTrace->BackTrace, BackTrace->StackTraceMaxDepth, Context);
+	BackTrace->Flag.store(1, std::memory_order_release);
+	if (InFlightSerialRequest.load(std::memory_order_relaxed) == BackTrace)
 	{
-		SignalThreadStackUserData.BackTraceCount = FPlatformStackWalk::CaptureStackBackTrace(SignalThreadStackUserData.BackTrace, SignalThreadStackUserData.CallStackSize, Context);
-		FPlatformAtomics::AtomicStore(&ThreadStackBackTraceStatus, ThreadStackBackTraceCurrentStatus_DONE);
+		InFlightSerialRequest.store(nullptr, std::memory_order_release);
 	}
 }
 
@@ -387,68 +386,65 @@ uint32 FAndroidPlatformStackWalk::CaptureThreadStackBackTrace(uint64 ThreadId, u
 		bHasReentered = false;
 	};
 
-	auto GatherCallstackFromThread = [](uint64 TargetThreadId)
+	if (InFlightSerialRequest.load(std::memory_order_acquire) != nullptr)
+	{
+		return 0;
+	}
+
+	SignalThreadStackUserData.Depth = 0;
+	SignalThreadStackUserData.ThreadID = ThreadId;
+	SignalThreadStackUserData.Flag.store(0, std::memory_order_release);
+
+	auto WaitForSignalHandlerToFinishOrCrash = [BackTrace]()
 	{
-		static int32 ThreadStackBackTraceNextRequest = 0;
-		int32 CurrentThreadStackBackTrace = ThreadStackBackTraceNextRequest++;
+		const float PollTime = 0.001f;
 
-		auto WaitForSignalHandlerToFinishOrCrash = [CurrentThreadStackBackTrace]()
+		for (float CurrentTime = 0; CurrentTime <= GThreadCallStackRequestMaxWait; CurrentTime += PollTime)
 		{
-			const float PollTime = 0.001f;
-
-			for (float CurrentTime = 0; CurrentTime <= GThreadCallStackMaxWait; CurrentTime += PollTime)
+			if (SignalThreadStackUserData.Flag.load(std::memory_order_acquire))
 			{
-				if (FPlatformAtomics::InterlockedCompareExchange(&ThreadStackBackTraceStatus, ThreadStackBackTraceNextRequest, ThreadStackBackTraceCurrentStatus_DONE) == ThreadStackBackTraceCurrentStatus_DONE)
-				{
-					// success 
-					return SignalThreadStackUserData.BackTraceCount;
-				}
-
-				// signal timed out
-				if (CurrentTime > GThreadCallStackRequestMaxWait && FPlatformAtomics::InterlockedCompareExchange(&ThreadStackBackTraceStatus, ThreadStackBackTraceNextRequest, CurrentThreadStackBackTrace) == CurrentThreadStackBackTrace)
-				{
-					// request not yet started, skip it.
-					return 0;
-				}
-				FPlatformProcess::SleepNoStats(PollTime);
+				FMemory::Memcpy(BackTrace, SignalThreadStackUserData.BackTrace, SignalThreadStackUserData.Depth * sizeof(*BackTrace));
+				return SignalThreadStackUserData.Depth;
 			}
 
-			// We have waited for as long as we should for the signal handler to finish. Assume it has hang and we need to kill our selfs
-			*(int*)0x10 = 0x0;
-			return 0;
-		};
-
-		sigval UserData;
-		UserData.sival_int = CurrentThreadStackBackTrace;
-
-		siginfo_t info;
-		memset(&info, 0, sizeof(siginfo_t));
-		info.si_signo = THREAD_CALLSTACK_GENERATOR;
-		info.si_code = SI_QUEUE;
-		info.si_pid = syscall(SYS_getpid);
-		info.si_uid = syscall(SYS_getuid);
-		info.si_value = UserData;
-
-		// Avoid using sigqueue here as if the ThreadId is already blocked and in a signal handler
-		// sigqueue will try a different thread signal handler and report the wrong callstack
-		if (syscall(SYS_rt_tgsigqueueinfo, info.si_pid, TargetThreadId, THREAD_CALLSTACK_GENERATOR, &info) == 0)
-		{
-			return WaitForSignalHandlerToFinishOrCrash();
-		}
-		else
-		{
-			// we failed to send the signal, update the current status.
-			FPlatformAtomics::AtomicStore(&ThreadStackBackTraceStatus, ThreadStackBackTraceNextRequest);
+			FPlatformProcess::SleepNoStats(PollTime);
 		}
 
+		// Time out
 		return 0;
 	};
 
-	SignalThreadStackUserData.CallStackSize = MaxDepth;
-	SignalThreadStackUserData.BackTrace = BackTrace;
-	SignalThreadStackUserData.BackTraceCount = 0;
+	InFlightSerialRequest.store(&SignalThreadStackUserData, std::memory_order_relaxed);
+	if (CaptureThreadStackBackTraceAsync(&SignalThreadStackUserData))
+	{
+		return WaitForSignalHandlerToFinishOrCrash();
+	}
+	
+	InFlightSerialRequest.store(nullptr, std::memory_order_relaxed);
+	return 0;
+}
+
+int FAndroidPlatformStackWalk::CaptureThreadStackBackTraceAsync(FAsyncThreadBackTrace* BackTrace)
+{
+	sigval UserData;
+	UserData.sival_ptr = BackTrace;
+
+	siginfo_t info;
+	memset(&info, 0, sizeof(siginfo_t));
+	info.si_signo = THREAD_CALLSTACK_GENERATOR;
+	info.si_code = SI_QUEUE;
+	info.si_pid = syscall(SYS_getpid);
+	info.si_uid = syscall(SYS_getuid);
+	info.si_value = UserData;
+
+	// Avoid using sigqueue here as if the ThreadId is already blocked and in a signal handler
+	// sigqueue will try a different thread signal handler and report the wrong callstack
+	if (syscall(SYS_rt_tgsigqueueinfo, info.si_pid, BackTrace->ThreadID, THREAD_CALLSTACK_GENERATOR, &info) == 0)
+	{
+		return 1;
+	}
 
-	return GatherCallstackFromThread(ThreadId);
+	return 0;
 }
 #else
 uint32 FAndroidPlatformStackWalk::CaptureThreadStackBackTrace(uint64 ThreadId, uint64* BackTrace, uint32 MaxDepth, void* Context)
diff --git a/Engine/Source/Runtime/Core/Private/Android/AndroidSignals.h b/Engine/Source/Runtime/Core/Private/Android/AndroidSignals.h
--- a/Engine/Source/Runtime/Core/Private/Android/AndroidSignals.h
+++ b/Engine/Source/Runtime/Core/Private/Android/AndroidSignals.h
@@ -93,6 +93,7 @@ protected:
 			FPlatformProcess::SleepNoStats(0.0f);
 		}
 		SignalParams.Signal = InSignal;
+		SignalParams.ThreadID = FPlatformTLS::GetCurrentThreadId();
 		SignalParams.Info = InInfo;
 		SignalParams.Context = InContext;
 
@@ -166,13 +167,14 @@ private:
 			pthread_exit(0);
 			return;
 		}
-		Derived::HandleTargetSignal(SignalParams.Signal, SignalParams.Info, SignalParams.Context);
+		Derived::HandleTargetSignal(SignalParams.Signal, SignalParams.Info, SignalParams.Context, SignalParams.ThreadID);
 		FPlatformAtomics::AtomicStore(&SignalThreadStatus, (int32)ESignalThreadStatus::Complete);
 	}
 
 	struct FSignalParams
 	{
 		int Signal;
+		uint32 ThreadID;
 		siginfo* Info;
 		void* Context;
 	};
diff --git a/Engine/Source/Runtime/Core/Public/Android/AndroidPlatformCrashContext.h b/Engine/Source/Runtime/Core/Public/Android/AndroidPlatformCrashContext.h
--- a/Engine/Source/Runtime/Core/Public/Android/AndroidPlatformCrashContext.h
+++ b/Engine/Source/Runtime/Core/Public/Android/AndroidPlatformCrashContext.h
@@ -9,6 +9,9 @@ struct CORE_API FAndroidCrashContext : public FGenericCrashContext
 {
 	/** Signal number */
 	int32 Signal;
+
+	/** Id of a thread that crashed */
+	uint32 CrashingThreadId;
 	
 	/** Additional signal info */
 	siginfo* Info;
@@ -29,9 +32,10 @@ struct CORE_API FAndroidCrashContext : public FGenericCrashContext
 	 * @param InInfo additional info (e.g. address we tried to read, etc)
 	 * @param InContext thread context
 	 */
-	void InitFromSignal(int32 InSignal, siginfo* InInfo, void* InContext)
+	void InitFromSignal(int32 InSignal, siginfo* InInfo, void* InContext, uint32 InCrashingThreadId)
 	{
 		Signal = InSignal;
+		CrashingThreadId = InCrashingThreadId;
 		Info = InInfo;
 		Context = InContext;
 	}
@@ -60,7 +64,8 @@ struct CORE_API FAndroidCrashContext : public FGenericCrashContext
 	// generate an absolute path to a crash report folder.
 	static void GenerateReportDirectoryName(char(&DirectoryNameOUT)[CrashReportMaxPathSize]);
 
-	void DumpAllThreadCallstacks() const;
+	// expects externally allocated memory for all threads found in FThreadManager + main thread
+	void DumpAllThreadCallstacks(FAsyncThreadBackTrace* BackTrace, int NumThreads) const;
 
 	/** Async-safe ItoA */
 	static const ANSICHAR* ItoANSI(uint64 Val, uint64 Base, uint32 Len = 0);
diff --git a/Engine/Source/Runtime/Core/Public/Android/AndroidPlatformMisc.h b/Engine/Source/Runtime/Core/Public/Android/AndroidPlatformMisc.h
--- a/Engine/Source/Runtime/Core/Public/Android/AndroidPlatformMisc.h
+++ b/Engine/Source/Runtime/Core/Public/Android/AndroidPlatformMisc.h
@@ -219,7 +219,7 @@ public:
 
 	// To help track down issues with failing crash handler.
 	static FString GetFatalSignalMessage(int Signal, siginfo* Info);
-	static void OverrideFatalSignalHandler(void (*FatalSignalHandlerOverrideFunc)(int Signal, struct siginfo* Info, void* Context));
+	static void OverrideFatalSignalHandler(void (*FatalSignalHandlerOverrideFunc)(int Signal, struct siginfo* Info, void* Context, uint32 CrashingThreadId));
 	// To help track down issues with failing crash handler.
 
 	static bool IsInSignalHandler();
diff --git a/Engine/Source/Runtime/Core/Public/Android/AndroidPlatformStackWalk.h b/Engine/Source/Runtime/Core/Public/Android/AndroidPlatformStackWalk.h
--- a/Engine/Source/Runtime/Core/Public/Android/AndroidPlatformStackWalk.h
+++ b/Engine/Source/Runtime/Core/Public/Android/AndroidPlatformStackWalk.h
@@ -8,6 +8,17 @@
 #pragma once
 #include "GenericPlatform/GenericPlatformStackWalk.h"
 
+struct FAsyncThreadBackTrace
+{
+	std::atomic<int32> Flag;		//0 or 1 to indicate if stack capture has been finished
+	int32 Depth;
+	static constexpr int StackTraceMaxDepth = 100;
+	uint64 BackTrace[StackTraceMaxDepth];
+	uint32 ThreadID;
+	static constexpr int MaxThreadName = 20;
+	char ThreadName[MaxThreadName];
+};
+
 /**
 * Android platform stack walking
 */
@@ -20,6 +31,7 @@ struct CORE_API FAndroidPlatformStackWalk : public FGenericPlatformStackWalk
 	static bool SymbolInfoToHumanReadableString(const FProgramCounterSymbolInfo& SymbolInfo, ANSICHAR* HumanReadableString, SIZE_T HumanReadableStringSize);
 
 	static uint32 CaptureThreadStackBackTrace(uint64 ThreadId, uint64* BackTrace, uint32 MaxDepth, void* Context = nullptr);
+	static int CaptureThreadStackBackTraceAsync(FAsyncThreadBackTrace* BackTrace);
 
 	static void HandleBackTraceSignal(siginfo* Info, void* Context);
 
diff --git a/Engine/Source/Runtime/Core/Public/HAL/ThreadManager.h b/Engine/Source/Runtime/Core/Public/HAL/ThreadManager.h
--- a/Engine/Source/Runtime/Core/Public/HAL/ThreadManager.h
+++ b/Engine/Source/Runtime/Core/Public/HAL/ThreadManager.h
@@ -40,6 +40,9 @@ public:
 	*/
 	void RemoveThread(class FRunnableThread* Thread);
 
+	/** Get the number of registered threads */
+	int32 NumThreads() const { return Threads.Num(); }
+
 	/** Ticks all fake threads and their runnable objects. */
 	void Tick();
 
