Changed FOutputDeviceBlockAllocationTag::Allocator to disregard Alignment param and use DEFAULT_ALIGNMENT instead

diff --git a/Engine/Source/Runtime/Core/Private/Misc/OutputDeviceRedirector.cpp b/Engine/Source/Runtime/Core/Private/Misc/OutputDeviceRedirector.cpp
--- a/Engine/Source/Runtime/Core/Private/Misc/OutputDeviceRedirector.cpp
+++ b/Engine/Source/Runtime/Core/Private/Misc/OutputDeviceRedirector.cpp
@@ -47,7 +47,7 @@ struct FOutputDeviceBlockAllocationTag : FDefaultBlockAllocationTag
 
 		FORCEINLINE static void* Malloc(SIZE_T Size, uint32 Alignment)
 		{
-			return FMemory::Malloc(Size, Alignment);
+			return FMemory::Malloc(Size, DEFAULT_ALIGNMENT);
 		}
 
 		FORCEINLINE static void Free(void* Pointer, SIZE_T Size)
