Fixed bugs introduced by latest changes to MallocBinned

diff --git a/Engine/Source/Runtime/Core/Private/Apple/ApplePlatformMemory.cpp b/Engine/Source/Runtime/Core/Private/Apple/ApplePlatformMemory.cpp
--- a/Engine/Source/Runtime/Core/Private/Apple/ApplePlatformMemory.cpp
+++ b/Engine/Source/Runtime/Core/Private/Apple/ApplePlatformMemory.cpp
@@ -298,13 +298,7 @@ FMalloc* FApplePlatformMemory::BaseAllocator()
 	}
 	else
 	{
-		// A temporary workaround for an iOS specific memory allocator issue causing crashes on Main.
-		// See jira UE-154181 for details. Remove once jira is fixed.
-#if PLATFORM_IOS
-		AllocatorToUse = EMemoryAllocatorToUse::Ansi;
-#else
 		AllocatorToUse = EMemoryAllocatorToUse::Binned;
-#endif
 	}
 	
 	// Force ansi malloc in some cases
diff --git a/Engine/Source/Runtime/Core/Private/HAL/MallocBinned.cpp b/Engine/Source/Runtime/Core/Private/HAL/MallocBinned.cpp
--- a/Engine/Source/Runtime/Core/Private/HAL/MallocBinned.cpp
+++ b/Engine/Source/Runtime/Core/Private/HAL/MallocBinned.cpp
@@ -613,7 +613,7 @@ struct FMallocBinned::Private
 		UPTRINT OsBytes = Align(Bytes, PageSize);
 
 		checkSlow(Blocks >= 1);
-		checkSlow(Blocks * Table->BlockSize <= Bytes && PoolSize >= Bytes);
+		checkSlow(PoolSize >= Bytes);
 		
 		FFreeMem* Free = nullptr;
 		SIZE_T ActualPoolSize; //TODO: use this to reduce waste?
@@ -628,7 +628,7 @@ struct FMallocBinned::Private
 			Free = (FFreeMem*)OSAlloc(Allocator, OsBytes, ActualPoolSize);
 		}
 		
-		checkSlow(!((UPTRINT)Free & (PageSize - 1)));
+        checkSlow(IsAligned(Free, PageSize));
 		if( !Free )
 		{
 			OutOfMemory(OsBytes);
@@ -766,7 +766,7 @@ struct FMallocBinned::Private
 		else
 		{
 			// Free an OS allocation.
-			checkSlow(!((UPTRINT)Ptr & (Allocator.PageSize - 1)));
+			checkSlow(IsAligned(Ptr, Allocator.PageSize));
 			SIZE_T OsBytes = Pool->GetOsBytes(Allocator.PageSize, Allocator.BinnedOSTableIndex);
 
 			BINNED_ADD_STATCOUNTER(Allocator.UsedCurrent,  -(int64)Pool->GetBytes());
@@ -1287,6 +1287,8 @@ void* FMallocBinned::Malloc(SIZE_T Size, uint32 Alignment)
 #ifdef USE_COARSE_GRAIN_LOCKS
 	FScopeLock ScopedLock(&AccessGuard);
 #endif
+    
+    ensure(Alignment <= PageSize);
 
 	Private::FlushPendingFrees(*this);
 
@@ -1336,8 +1338,6 @@ void* FMallocBinned::Malloc(SIZE_T Size, uint32 Alignment)
 #ifdef USE_FINE_GRAIN_LOCKS
 			FScopeLock TableLock(&Table->CriticalSection);
 #endif
-			checkSlow(Size <= Table->BlockSize);
-
 			Private::TrackStats(Table, (uint32)Size);
 
 			FPoolInfo* Pool = Table->FirstPool;
@@ -1358,8 +1358,6 @@ void* FMallocBinned::Malloc(SIZE_T Size, uint32 Alignment)
 #ifdef USE_FINE_GRAIN_LOCKS
 			FScopeLock TableLock(&Table->CriticalSection);
 #endif
-			checkSlow(Size <= Table->BlockSize);
-
 			Private::TrackStats(Table, (uint32)Size);
 
 			FPoolInfo* Pool = Table->FirstPool;
@@ -1422,7 +1420,7 @@ void* FMallocBinned::Malloc(SIZE_T Size, uint32 Alignment)
 void* FMallocBinned::Realloc( void* Ptr, SIZE_T NewSize, uint32 Alignment )
 {
 	Alignment = FMath::Max<uint32>(Alignment, Private::DEFAULT_BINNED_ALLOCATOR_ALIGNMENT);
-	const uint32 NewSizeUnmodified = NewSize;
+    ensure(Alignment <= PageSize);
 	MEM_TIME(MemTime -= FPlatformTime::Seconds());
 	UPTRINT BasePtr;
 	void* NewPtr = Ptr;
@@ -1431,19 +1429,23 @@ void* FMallocBinned::Realloc( void* Ptr, SIZE_T NewSize, uint32 Alignment )
 #if USE_OS_SMALL_BLOCK_ALLOC && !USE_OS_SMALL_BLOCK_GRAB_MEMORY_FROM_OS
 		if(FPlatformMemory::PtrIsOSMalloc(Ptr))
 		{
-			NewSize = Align(NewSize, Alignment);
-			NewPtr = ::realloc(Ptr, NewSize);
-			
-			if(!FPlatformMemory::PtrIsFromNanoMalloc(NewPtr))
-			{
-				// We've overflowed the nano region
-				// Fall back to UE's allocator
-				Ptr = NewPtr;
-				NewPtr = Malloc(NewSizeUnmodified, Alignment);
-				const SIZE_T OldSize = malloc_size(Ptr);
-				FMemory::Memcpy(NewPtr, Ptr, OldSize);
-				Private::SmallOSFree(*this, Ptr, NewSize);
-			}
+            bool bFallback = true;
+            const SIZE_T OldSize = malloc_size(Ptr);
+            if ((NewSize <= Private::SMALL_BLOCK_POOL_SIZE) && (Alignment <= Private::DEFAULT_BINNED_ALLOCATOR_ALIGNMENT))
+            {
+                NewPtr = ::realloc(Ptr, NewSize);
+                bFallback = !FPlatformMemory::PtrIsFromNanoMalloc(NewPtr);
+            }
+
+            if (bFallback)
+            {
+                // We've overflowed the nano region
+                // Fall back to UE's allocator
+                Ptr = NewPtr;
+                NewPtr = Malloc(NewSize, Alignment);
+                FMemory::Memcpy(NewPtr, Ptr, OldSize);
+                Private::SmallOSFree(*this, Ptr, OldSize);
+            }
 		}
 		else
 #endif
@@ -1454,11 +1456,11 @@ void* FMallocBinned::Realloc( void* Ptr, SIZE_T NewSize, uint32 Alignment )
 			{
 				// Allocated from pool, so grow or shrink if necessary.
 				check(Pool->TableIndex > 0); // it isn't possible to allocate a size of 0, Malloc will increase the size to DEFAULT_BINNED_ALLOCATOR_ALIGNMENT
-				if (NewSizeUnmodified > MemSizeToPoolTable[Pool->TableIndex]->BlockSize || NewSizeUnmodified <= MemSizeToPoolTable[Pool->TableIndex - 1]->BlockSize
+				if (NewSize > MemSizeToPoolTable[Pool->TableIndex]->BlockSize || NewSize <= MemSizeToPoolTable[Pool->TableIndex - 1]->BlockSize
 					|| !IsAligned(Ptr, Alignment))		// If reallocation requires a pointer to have a different alignment than the original we need to allocate new mem and copy old over
 				{
-					NewPtr = Malloc(NewSizeUnmodified, Alignment);
-					FMemory::Memcpy(NewPtr, Ptr, FMath::Min<SIZE_T>(NewSizeUnmodified, MemSizeToPoolTable[Pool->TableIndex]->BlockSize));
+					NewPtr = Malloc(NewSize, Alignment);
+					FMemory::Memcpy(NewPtr, Ptr, FMath::Min<SIZE_T>(NewSize, MemSizeToPoolTable[Pool->TableIndex]->BlockSize));
 					Free( Ptr );
 				}
 			}
@@ -1468,8 +1470,8 @@ void* FMallocBinned::Realloc( void* Ptr, SIZE_T NewSize, uint32 Alignment )
 				if( NewSize > Pool->GetOsBytes(PageSize, BinnedOSTableIndex) || NewSize * 3 < Pool->GetOsBytes(PageSize, BinnedOSTableIndex) * 2 )
 				{
 					// Grow or shrink.
-					NewPtr = Malloc(NewSizeUnmodified, Alignment);
-					FMemory::Memcpy(NewPtr, Ptr, FMath::Min<SIZE_T>(NewSizeUnmodified, Pool->GetBytes()));
+					NewPtr = Malloc(NewSize, Alignment);
+					FMemory::Memcpy(NewPtr, Ptr, FMath::Min<SIZE_T>(NewSize, Pool->GetBytes()));
 					Free( Ptr );
 				}
 				else
@@ -1484,14 +1486,14 @@ void* FMallocBinned::Realloc( void* Ptr, SIZE_T NewSize, uint32 Alignment )
 					BINNED_ADD_STATCOUNTER(UsedCurrent, UsedChange);
 					BINNED_PEAK_STATCOUNTER(UsedPeak, UsedCurrent);
 					BINNED_ADD_STATCOUNTER(WasteCurrent, (Pool->GetBytes() - NewSize));
-					Pool->SetAllocationSizes(NewSizeUnmodified, Pool->GetOsBytes(PageSize, BinnedOSTableIndex), BinnedOSTableIndex, BinnedOSTableIndex);
+					Pool->SetAllocationSizes(NewSize, Pool->GetOsBytes(PageSize, BinnedOSTableIndex), BinnedOSTableIndex, BinnedOSTableIndex);
 				}
 			}
 		}
 	}
 	else if( Ptr == nullptr )
 	{
-		NewPtr = Malloc(NewSizeUnmodified, Alignment);
+		NewPtr = Malloc(NewSize, Alignment);
 	}
 	else
 	{
@@ -1558,6 +1560,7 @@ bool FMallocBinned::GetAllocationSize(void *Original, SIZE_T &SizeOut)
 SIZE_T FMallocBinned::QuantizeSize(SIZE_T Size, uint32 Alignment)
 {
 	Alignment = FMath::Max<uint32>(Alignment, Private::DEFAULT_BINNED_ALLOCATOR_ALIGNMENT);
+    ensure(Alignment <= PageSize);
 
 	SIZE_T Result;
 #if USE_OS_SMALL_BLOCK_ALLOC && !USE_OS_SMALL_BLOCK_GRAB_MEMORY_FROM_OS
