Trivial cast fixes for IPackageWriter

#rnx

- FPackageInfo and FBulkDataInfo have a member 'MultiOutputIndex' which is set to be a uint32, but wil be used as the chunk index for ::CreateIoChunkId which can only support uint16, so there is no point in 'MultiOutputIndex' being a larger data type as it would silently fail later.

diff --git a/Engine/Source/Runtime/Core/Public/Serialization/PackageWriter.h b/Engine/Source/Runtime/Core/Public/Serialization/PackageWriter.h
--- a/Engine/Source/Runtime/Core/Public/Serialization/PackageWriter.h
+++ b/Engine/Source/Runtime/Core/Public/Serialization/PackageWriter.h
@@ -99,7 +99,7 @@ public:
 		FString		LooseFilePath;
 		uint64		HeaderSize = 0;
 		FIoChunkId	ChunkId = FIoChunkId::InvalidChunkId;
-		uint32		MultiOutputIndex = 0;
+		uint16		MultiOutputIndex = 0;
 	};
 
 	/** Write package data (exports and serialized header)
@@ -125,7 +125,7 @@ public:
 		EType		BulkDataType = BulkSegment;
 		FString		LooseFilePath;
 		FIoChunkId	ChunkId = FIoChunkId::InvalidChunkId;
-		uint32		MultiOutputIndex = 0;
+		uint16		MultiOutputIndex = 0;
 	};
 
 	/** Write bulk data for the current package
@@ -138,7 +138,7 @@ public:
 		FName		PackageName;
 		FString		Filename;
 		FIoChunkId	ChunkId = FIoChunkId::InvalidChunkId;
-		uint32		MultiOutputIndex = 0;
+		uint16		MultiOutputIndex = 0;
 	};
 
 	/** Write separate files written by UObjects during cooking via UObject::CookAdditionalFiles. */
@@ -148,7 +148,7 @@ public:
 	{
 		/** Associated Package Name Entry */
 		FName	PackageName;
-		uint32	MultiOutputIndex = 0;
+		uint16	MultiOutputIndex = 0;
 	};
 	/** Write separate data written by UObjects via FLinkerSave::AdditionalDataToAppend. */
 	virtual void WriteLinkerAdditionalData(const FLinkerAdditionalDataInfo& Info, const FIoBuffer& Data, const TArray<FFileRegion>& FileRegions) = 0;
