Modular builds: Consoles: Added initialization of PLATFORM_CONSOLE_DYNAMIC_LINK in the master builds script based on whether the build is monolythis or moduler to avoid having to set this variable from each project's script separately.

#rnx

diff --git a/Engine/Source/Runtime/Core/Public/HAL/Platform.h b/Engine/Source/Runtime/Core/Public/HAL/Platform.h
--- a/Engine/Source/Runtime/Core/Public/HAL/Platform.h
+++ b/Engine/Source/Runtime/Core/Public/HAL/Platform.h
@@ -649,6 +649,10 @@
 	#define PLATFORM_REQUIRES_TYPELESS_RESOURCE_DISCARD_WORKAROUND 0
 #endif
 
+#ifndef PLATFORM_CONSOLE_DYNAMIC_LINK
+	#define PLATFORM_CONSOLE_DYNAMIC_LINK 0
+#endif
+
 // deprecated, do not use
 #define PLATFORM_HAS_THREADSAFE_RHIGetRenderQueryResult	#
 #define PLATFORM_SUPPORTS_RHI_THREAD #
