Added CountBy, a version of Count taking a value and a projection

diff --git a/Engine/Source/Runtime/Core/Public/Algo/Count.h b/Engine/Source/Runtime/Core/Public/Algo/Count.h
--- a/Engine/Source/Runtime/Core/Public/Algo/Count.h
+++ b/Engine/Source/Runtime/Core/Public/Algo/Count.h
@@ -45,4 +45,25 @@ namespace Algo
 		}
 		return Result;
 	}
+
+	/**
+	 * Counts elements of a range whose projection equals the supplied value
+	 *
+	 * @param  Input       Any iterable type
+	 * @param  InValue     Value to compare against
+	 * @param  Projection  The projection to apply to the element
+	 */
+	template <typename InT, typename ValueT, typename ProjectionT>
+	FORCEINLINE SIZE_T CountBy(const InT& Input, const ValueT& InValue, ProjectionT Proj)
+	{
+		SIZE_T Result = 0;
+		for (const auto& Value : Input)
+		{
+			if (Invoke(Proj, Value) == InValue)
+			{
+				++Result;
+			}
+		}
+		return Result;
+	}
 }
