Restore the ability to load NFL config files in cooked builds
This is needed for samples that store private ids in NFL folders for internal builds

diff --git a/Engine/Source/Runtime/Core/Public/Misc/ConfigHierarchy.h b/Engine/Source/Runtime/Core/Public/Misc/ConfigHierarchy.h
--- a/Engine/Source/Runtime/Core/Public/Misc/ConfigHierarchy.h
+++ b/Engine/Source/Runtime/Core/Public/Misc/ConfigHierarchy.h
@@ -79,7 +79,7 @@ static FConfigLayerExpansion GConfigExpansions[] =
 	{ 
 		TEXT("{ENGINE}/"),						TEXT("{ENGINE}/Restricted/NotForLicensees/"),	
 		TEXT("{PROJECT}/Config/"),				TEXT("{RESTRICTEDPROJECT_NFL}/Config/"), 
-		EConfigExpansionFlags::ForUncooked 
+		EConfigExpansionFlags::ForUncooked | EConfigExpansionFlags::ForCooked
 	},
 	{ 
 		TEXT("{ENGINE}/"),						TEXT("{ENGINE}/Restricted/NoRedist/"),			
@@ -98,7 +98,7 @@ static FConfigLayerExpansion GConfigExpansions[] =
 	{
 		TEXT("{ENGINE}/Config/{PLATFORM}/"),	TEXT("{ENGINE}/Restricted/NotForLicensees/Platforms/{PLATFORM}/Config/"),	
 		TEXT("{PROJECT}/Config/{PLATFORM}/"),	TEXT("{RESTRICTEDPROJECT_NFL}/Platforms/{PLATFORM}/Config/"), 
-		EConfigExpansionFlags::ForUncooked 
+		EConfigExpansionFlags::ForUncooked | EConfigExpansionFlags::ForCooked
 	},
 	{
 		TEXT("{ENGINE}/Config/{PLATFORM}/"),	TEXT("{ENGINE}/Restricted/NoRedist/Platforms/{PLATFORM}/Config/"),			
