String tracing

Adds a tracing subsystem to trace deduplicated strings and FNames.

diff --git a/Engine/Source/Developer/TraceServices/Private/Analyzers/StringsAnalyzer.cpp b/Engine/Source/Developer/TraceServices/Private/Analyzers/StringsAnalyzer.cpp
new file mode 100644
--- /dev/null
+++ b/Engine/Source/Developer/TraceServices/Private/Analyzers/StringsAnalyzer.cpp
@@ -0,0 +1,83 @@
+﻿// Copyright Epic Games, Inc. All Rights Reserved.
+#include "StringsAnalyzer.h"
+#include "../Common/Utils.h"
+#include "TraceServices/Model/Definitions.h"
+#include "TraceServices/Model/Strings.h"
+
+namespace TraceServices
+{
+	
+/////////////////////////////////////////////////////////////////////////////////////////
+void FStringsAnalyzer::OnAnalysisBegin(const FOnAnalysisContext& Context)
+{
+	Context.InterfaceBuilder.RouteEvent(RouteId_StaticString, "Strings", "StaticString");
+	Context.InterfaceBuilder.RouteEvent(RouteId_FName, "Strings", "FName");
+	Context.InterfaceBuilder.RouteEvent(RouteId_StaticStringNoSync, "Strings", "StaticStringNoSync");
+	Context.InterfaceBuilder.RouteEvent(RouteId_FNameNoSync, "Strings", "FNameNoSync");
+}
+
+/////////////////////////////////////////////////////////////////////////////////////////
+bool FStringsAnalyzer::OnEvent(uint16 RouteId, EStyle Style, const FOnEventContext& Context)
+{
+	switch (RouteId)
+	{
+		// The layout of these two events is identical so they can be handled using the same code.
+	case RouteId_StaticString:
+	case RouteId_FName:
+	case RouteId_StaticStringNoSync:
+	case RouteId_FNameNoSync:
+		{
+			const FEventData& EventData = Context.EventData;
+			IDefinitionProvider* DefinitionProvider = GetDefinitionProvider(Session);
+
+			FWideStringView DisplayWide;
+			FAnsiStringView DisplayAnsi;
+			const TCHAR* Display = nullptr;
+			if (DisplayAnsi.Compare("LyraTestAssetName1") == 0)
+			{
+				UE_DEBUG_BREAK();
+			}
+			if (EventData.GetString("DisplayWide", DisplayWide) && DisplayWide.Len() > 0)
+			{
+				Display = Session.StoreString(DisplayWide);
+			}
+			else if (EventData.GetString("DisplayAnsi", DisplayAnsi) && DisplayAnsi.Len() > 0)
+			{
+				const FString Str(DisplayAnsi);
+				Display = Session.StoreString(Str);
+			}
+			else
+			{
+				UE_LOG(LogTraceServices, Warning, TEXT("Empty string definition detected."));
+				return true;
+			}
+
+			FStringDefinition* Instance = DefinitionProvider->Create<FStringDefinition>();
+			Instance->Display = Display;
+
+			if (RouteId == RouteId_FName)
+			{
+				auto Id = EventData.GetDefinitionId<uint32>();
+				if (Id.Id == 1651656)
+				{
+					UE_DEBUG_BREAK();
+				}
+				// Overwrite type id of Strings.FNameNoSync to Strings.FNames type id.
+				//Id.RefTypeId = ???;
+				DefinitionProvider->Register<FStringDefinition>(Instance, Id);
+			}
+			else
+			{
+				auto Id = EventData.GetDefinitionId<uint64>();
+				// Overwrite type id of Strings.StaticStringNoSync to Strings.StaticString type id.
+				//Id.RefTypeId = ???;
+				DefinitionProvider->Register<FStringDefinition>(Instance, Id);
+			}
+		}
+		break;
+	default: ;
+	}
+	return true;
+}
+
+}
diff --git a/Engine/Source/Developer/TraceServices/Private/Analyzers/StringsAnalyzer.h b/Engine/Source/Developer/TraceServices/Private/Analyzers/StringsAnalyzer.h
new file mode 100644
--- /dev/null
+++ b/Engine/Source/Developer/TraceServices/Private/Analyzers/StringsAnalyzer.h
@@ -0,0 +1,31 @@
+﻿// Copyright Epic Games, Inc. All Rights Reserved.
+#pragma once
+#include "Trace/Analyzer.h"
+#include "TraceServices/Model/AnalysisSession.h"
+
+namespace TraceServices
+{
+
+class FStringsAnalyzer : public UE::Trace::IAnalyzer
+{
+public:
+	FStringsAnalyzer(IAnalysisSession& InSession)
+		: Session(InSession)
+	{
+	}
+
+private:
+	enum ERouteId : uint16
+	{
+		RouteId_StaticString,
+		RouteId_FName,
+		RouteId_StaticStringNoSync,
+		RouteId_FNameNoSync,
+	};
+	virtual void OnAnalysisBegin(const FOnAnalysisContext& Context) override;
+	virtual bool OnEvent(uint16 RouteId, EStyle Style, const FOnEventContext& Context) override;
+
+	IAnalysisSession& Session;
+};
+
+} // namespace TraceServices
diff --git a/Engine/Source/Developer/TraceServices/Public/TraceServices/Model/Strings.h b/Engine/Source/Developer/TraceServices/Public/TraceServices/Model/Strings.h
new file mode 100644
--- /dev/null
+++ b/Engine/Source/Developer/TraceServices/Public/TraceServices/Model/Strings.h
@@ -0,0 +1,16 @@
+﻿// Copyright Epic Games, Inc. All Rights Reserved.
+#pragma  once
+
+#include "TraceServices/Model/AnalysisSession.h"
+#include "Templates/Function.h"
+
+namespace TraceServices
+{
+	struct FStringDefinition
+	{
+		/**
+		 * Display string. The pointer is valid during the analysis session.
+		 */
+		const TCHAR* Display;
+	};
+}
\ No newline at end of file
diff --git a/Engine/Source/Runtime/Core/Private/ProfilingDebugging/StringTrace.cpp b/Engine/Source/Runtime/Core/Private/ProfilingDebugging/StringTrace.cpp
new file mode 100644
--- /dev/null
+++ b/Engine/Source/Runtime/Core/Private/ProfilingDebugging/StringTrace.cpp
@@ -0,0 +1,92 @@
+﻿// Copyright Epic Games, Inc. All Rights Reserved.
+#include "ProfilingDebugging/StringsTrace.h"
+#include "CoreSharedPCH.h"
+#include "HAL/CriticalSection.h"
+
+//////////////////////////////////////////////////////////////////////////////////////////
+UE_TRACE_EVENT_DEFINE(Strings, StaticString);
+UE_TRACE_EVENT_DEFINE(Strings, FName);
+UE_TRACE_EVENT_DEFINE(Strings, StaticStringNoSync);
+UE_TRACE_EVENT_DEFINE(Strings, FNameNoSync);
+
+#if UE_TRACE_ENABLED
+
+//////////////////////////////////////////////////////////////////////////////////////////
+static struct FStaticStringTracer
+{
+	UE::Trace::FEventRef64 TraceString(const ANSICHAR* InString)
+	{
+		FScopeLock _(&LookupLock);
+		uint64 Id = uint64(InString);
+		bool bAlreadyInSet = false;
+		TracedAnsiStrings.FindOrAdd(Id, &bAlreadyInSet);
+		if (!bAlreadyInSet)
+		{
+			UE_TRACE_LOG_DEFINITION(Strings, StaticString, Id, true)
+				 << StaticString.DisplayAnsi(InString, FCStringAnsi::Strlen(InString));
+		}
+		return UE::Trace::FEventRef64(Id, UE_TRACE_GET_DEFINITION_TYPE_ID(Strings, StaticString));
+	}
+
+	UE::Trace::FEventRef64 TraceString(const TCHAR* InString)
+	{
+		FScopeLock _(&LookupLock);
+		uint64 Id = uint64(InString);
+		bool bAlreadyInSet = false;
+		TracedWideStrings.FindOrAdd(Id, &bAlreadyInSet);
+		if (!bAlreadyInSet)
+		{
+			return UE_TRACE_LOG_DEFINITION(Strings, StaticString, Id, true)
+				 << StaticString.DisplayWide(InString, FCString::Strlen(InString));
+		}
+		return UE::Trace::FEventRef64(Id, UE_TRACE_GET_DEFINITION_TYPE_ID(Strings, StaticString));
+	}
+	
+	void TraceStringsOnConnection()
+	{
+		FScopeLock _(&LookupLock);
+		for(const uint64 Ptr : TracedWideStrings)
+		{
+			const TCHAR* String = (const TCHAR*) Ptr;
+			UE_TRACE_LOG_DEFINITION(Strings, StaticStringNoSync, Ptr, true)
+				<< StaticStringNoSync.DisplayWide(String, FCString::Strlen(String));
+		}
+		for(const uint64 Ptr : TracedAnsiStrings)
+		{
+			const ANSICHAR* String = (const ANSICHAR*) Ptr;
+			UE_TRACE_LOG_DEFINITION(Strings, StaticStringNoSync, Ptr, true)
+				<< StaticStringNoSync.DisplayAnsi(String, FCStringAnsi::Strlen(String));
+		}
+	}
+
+	FCriticalSection LookupLock;
+	TSet<uint64> TracedWideStrings;
+	TSet<uint64> TracedAnsiStrings;
+} GStaticStringTracer;
+
+//////////////////////////////////////////////////////////////////////////////////////////
+UE::Trace::FEventRef32 FStringTrace::GetNameRef(const FName& Name)
+{
+	return FName::TraceName(Name);
+}
+
+//////////////////////////////////////////////////////////////////////////////////////////
+UE::Trace::FEventRef64 FStringTrace::GetStaticStringRef(const TCHAR* String)
+{
+	return GStaticStringTracer.TraceString(String);
+}
+
+//////////////////////////////////////////////////////////////////////////////////////////
+UE::Trace::FEventRef64 FStringTrace::GetStaticStringRef(const ANSICHAR* String)
+{
+	return GStaticStringTracer.TraceString(String);
+}
+
+//////////////////////////////////////////////////////////////////////////////////////////
+void FStringTrace::OnConnection()
+{
+	GStaticStringTracer.TraceStringsOnConnection();
+	FName::TraceNamesOnConnection();
+}
+
+#endif //UE_TRACE_ENABLED
\ No newline at end of file
diff --git a/Engine/Source/Runtime/Core/Private/ProfilingDebugging/TraceAuxiliary.cpp b/Engine/Source/Runtime/Core/Private/ProfilingDebugging/TraceAuxiliary.cpp
--- a/Engine/Source/Runtime/Core/Private/ProfilingDebugging/TraceAuxiliary.cpp
+++ b/Engine/Source/Runtime/Core/Private/ProfilingDebugging/TraceAuxiliary.cpp
@@ -1315,6 +1315,9 @@ void FTraceAuxiliary::Initialize(const TCHAR* CommandLine)
 	// specified by the commandline, UE::Trace::Initialize will reset all channels.
 	GTraceAuxiliary.EnableCommandlineChannelsPostInitialize();
 	
+	// Setup known on connection callbacks
+	OnConnection.AddStatic(FStringTrace::OnConnection);
+
 	// Always register end frame updates. This path is short circuited if a worker thread exists.
 	GTraceAuxiliary.StartEndFramePump();
 	if (!FForkProcessHelper::IsForkRequested())
diff --git a/Engine/Source/Runtime/Core/Private/UObject/UnrealNames.cpp b/Engine/Source/Runtime/Core/Private/UObject/UnrealNames.cpp
--- a/Engine/Source/Runtime/Core/Private/UObject/UnrealNames.cpp
+++ b/Engine/Source/Runtime/Core/Private/UObject/UnrealNames.cpp
@@ -21,6 +21,7 @@
 #include "UObject/ObjectVersion.h"
 #include "HAL/ThreadSafeCounter.h"
 #include "Misc/ScopeRWLock.h"
+#include "ProfilingDebugging/StringsTrace.h"
 #include "Containers/Set.h"
 #include "Internationalization/Text.h"
 #include "Internationalization/Internationalization.h"
@@ -302,6 +303,107 @@ FArchive& operator<<(FArchive& Ar, FNameEntryId& Id)
 	return Ar;
 }
 
+
+template<uint32 NumBits>
+class TNameAtomicBitSet
+{
+	using WordType = uint64;
+	static constexpr uint32 BitsPerWord = 8 * sizeof(WordType);
+	static constexpr uint32 NumWords = NumBits / BitsPerWord;
+	static constexpr uint32 BitIndexMask = BitsPerWord - 1;
+	static_assert((NumBits % BitsPerWord) == 0);
+	
+public:
+	/**
+	 * Atomically set a bit in the bit set.
+	 * @param Index Index of bit to set
+	 * @return Previous value of bit
+	 */
+	bool SetBitAtomic(uint32 Index)
+	{
+		const uint32 WordIndex = Index / BitsPerWord;
+		const uint32 BitIndex = Index & BitIndexMask;
+		const WordType ValueMask = WordType(1) << BitIndex;
+		volatile WordType* DestDword = GetBuffer() + WordIndex;
+		checkSlow(DestDword >= GetBuffer() && DestDword < (GetBuffer() + NumWords));
+		const WordType PrevDwordValue = std::atomic_fetch_or((std::atomic<WordType>*)DestDword, ValueMask);
+		return !!(PrevDwordValue & ValueMask);
+	}
+	
+	/**
+	 * Atomically clear a bit in the bit set.
+	 * @param Index Index of bit to clear
+	 * @return Previous value of bit
+	 */
+	bool ClearBitAtomic(uint32 Index)
+	{
+		const uint32 WordIndex = Index / BitsPerWord;
+		const uint32 BitIndex = Index & BitIndexMask;
+		const WordType ValueMask = WordType(1) << BitIndex;
+		volatile WordType* DestDword = GetBuffer() + WordIndex;
+		checkSlow(DestDword >= GetBuffer() && DestDword < (GetBuffer() + NumWords));
+		const WordType PrevDwordValue = std::atomic_fetch_and((std::atomic<WordType>*)DestDword, ~ValueMask);
+		return !!(PrevDwordValue & ValueMask);
+	}
+
+	/**
+	 * Iterates over the set and issues the lambda for each index that is set.
+	 * @param Func Functor to call
+	 */
+	template<typename Lambda>
+	void ForEachSetBit(Lambda Func) const
+	{
+		const WordType* Words = Data.load(std::memory_order_relaxed);
+		if (!Words)
+		{
+			return;
+		}
+
+		uint32 WordIndex(0);
+		for (std::atomic<WordType> AtomicWord : MakeArrayView(Words, NumWords))
+		{
+			const WordType Word = AtomicWord.load();
+			for (uint32 BitIdx = 0; BitIdx < BitsPerWord; ++BitIdx)
+			{
+				if (Word & (WordType(1) << BitIdx))
+				{
+					Func(BitIdx + (WordIndex * BitsPerWord));
+				}
+			}
+			++WordIndex;
+		}
+	}
+	
+private:
+
+
+	/**
+	 * Returns or create the buffer.
+	 * @return Pointer to buffer
+	 */
+	WordType* GetBuffer()
+	{
+		WordType* Buffer = Data.load(std::memory_order_relaxed);
+		if (!Buffer)
+		{
+			LLM_SCOPE(ELLMTag::FName);
+			constexpr uint32 BufferSizeBytes = NumWords * sizeof(WordType);
+			Buffer = (WordType*) FMemory::MallocZeroed(BufferSizeBytes, alignof(WordType));
+			WordType* Expected = nullptr;
+			if (UNLIKELY(!Data.compare_exchange_strong(Expected, Buffer)))
+			{
+				// Another thread assigned the value before us. Free our buffer and use
+				// the buffer from the other thread.
+				FMemory::Free(Buffer);
+				Buffer = Expected;
+			}
+		}
+		return Buffer;
+	}
+	
+	std::atomic<WordType*> Data;
+};
+
 struct FNameSlot
 {
 	// Use the remaining few bits to store a hash that can determine inequality
@@ -1466,6 +1568,11 @@ public:
 	void IterateNumberedNames(TFunctionRef<void(FName)> Func) const;
 #endif
 
+#if UE_TRACE_ENABLED
+	UE::Trace::FEventRef32 Trace(const FNameEntryId& EntryId);
+	void RetraceAll() const;
+#endif
+
 private:
 	enum { MaxENames = 512 };
 
@@ -1476,6 +1583,11 @@ private:
 #endif
 	FNamePoolShard<ENameCase::IgnoreCase> ComparisonShards[FNamePoolShards];
 
+#if UE_TRACE_ENABLED
+	using FTracedBitSet = TNameAtomicBitSet<(1 << FNameBlockOffsetBits)>;
+	FTracedBitSet TracedNames[FNameMaxBlocks];
+#endif
+	
 	// Put constant lookup on separate cache line to avoid it being constantly invalidated by insertion
 	alignas(PLATFORM_CACHE_LINE_SIZE) FNameEntryId ENameToEntry[(uint32)EName::MaxHardcodedNameIndex] = {};
 	uint32 LargestEnameUnstableId;
@@ -1760,14 +1872,75 @@ uint32 FNamePool::NumSlots() const
 	return SlotCapacity;
 }
 
+#if UE_TRACE_ENABLED
+UE::Trace::FEventRef32 FNamePool::Trace(const FNameEntryId& EntryId)
+{
+	const FNameEntryHandle Handle(EntryId);
+	const uint32 DefinitionId = EntryId.ToUnstableInt();
+	if (!TracedNames[Handle.Block].SetBitAtomic(Handle.Offset))
+	{
+		const FNameEntry& DisplayEntry = Resolve(Handle);
+
+		if (DisplayEntry.IsWide())
+		{
+			TCHAR Chars[NAME_SIZE];
+			DisplayEntry.GetUnterminatedName(Chars, NAME_SIZE);
+			const uint32 Size = DisplayEntry.GetNameLength();
+			const auto Ref = UE_TRACE_LOG_DEFINITION(Strings, FName, DefinitionId, true)
+				 << FName.DisplayWide(Chars, Size);
+			return Ref;
+		}
+		else
+		{
+			ANSICHAR Chars[NAME_SIZE];
+			DisplayEntry.GetAnsiName(Chars);
+			const uint32 Size = DisplayEntry.GetNameLength();
+			const auto Ref = UE_TRACE_LOG_DEFINITION(Strings, FName, DefinitionId, true)
+				 << FName.DisplayAnsi(Chars, Size);
+			return Ref;
+		}
+	}
+	return UE::Trace::MakeEventRef(DefinitionId, UE_TRACE_GET_DEFINITION_TYPE_ID(Strings, FName));
+}
+
+void FNamePool::RetraceAll() const
+{
+	uint32 Block = 0;
+	for (const FTracedBitSet& Set : TracedNames)
+	{
+		++Block;
+		Set.ForEachSetBit([&](uint32 Index)
+		{
+			const FNameEntryHandle Handle(Block, Index);
+			const FNameEntry& DisplayEntry = Resolve(Handle);
+			const FNameEntryId EntryId = Handle;
+			const uint32 DefinitionId = EntryId.ToUnstableInt(); 
+				
+			if (DisplayEntry.IsWide())
+			{
+				 TCHAR Chars[NAME_SIZE];
+				 DisplayEntry.GetUnterminatedName(Chars, NAME_SIZE);
+				 const uint32 Size = DisplayEntry.GetNameLength();
+				 const auto Ref = UE_TRACE_LOG_DEFINITION(Strings, FNameNoSync, DefinitionId, true)
+					  << FNameNoSync.DisplayWide(Chars, Size);
+			}
+			else
+			{
+				 ANSICHAR Chars[NAME_SIZE];
+				 DisplayEntry.GetAnsiName(Chars);
+				 const uint32 Size = DisplayEntry.GetNameLength();
+				 const auto Ref = UE_TRACE_LOG_DEFINITION(Strings, FNameNoSync, DefinitionId, true)
+					  << FNameNoSync.DisplayAnsi(Chars, Size);
+			}
+		});
+
+	}
+}
+#endif //UE_TRACE_ENABLED
+
 void FNamePool::LogStats(FOutputDevice& Ar) const
 {
-	Ar.Logf(TEXT("%i FNames using in %" SIZE_T_FMT " bytes fixed + %.2f kB entries + %.2f kB slots"),
-		NumEntries(), 
-		sizeof(FNamePool), 
-		Entries.NumBlocks() * FNameEntryAllocator::BlockSizeBytes / 1024.0, 
-		NumSlots() * sizeof(FNameSlot) / 1024.0
-	);
+	Ar.Logf(TEXT("%i FNames using in %ikB + %ikB"), NumEntries(), sizeof(FNamePool), Entries.NumBlocks() * FNameEntryAllocator::BlockSizeBytes / 1024);
 	Ar.Logf(TEXT("%d ansi FNames"), NumAnsiEntries());
 	Ar.Logf(TEXT("%d wide FNames"), NumWideEntries());
 #if UE_FNAME_OUTLINE_NUMBER
@@ -2260,6 +2433,20 @@ FNameEntry const* FName::GetEntry(FNameEntryId Id)
 	return ResolveEntryRecursive(Id);
 }
 
+#if UE_TRACE_ENABLED
+UE::Trace::FEventRef32 FName::TraceName(const FName& Name)
+{
+	const FNameEntryId DisplayId = Name.GetDisplayIndexFast();
+	return GetNamePool().Trace(DisplayId);
+}
+
+void FName::TraceNamesOnConnection()
+{
+	const FNamePool& Pool = GetNamePool();
+	Pool.RetraceAll();
+}
+#endif // UE_TRACE_ENABLED
+
 FString FName::NameToDisplayString( const FString& InDisplayName, const bool bIsBool )
 {
 	// Copy the characters out so that we can modify the string in place
@@ -3657,6 +3844,34 @@ PRAGMA_ENABLE_DEPRECATION_WARNINGS
 		check(Id == OldId);
 	}
 #endif
+
+	{
+		TNameAtomicBitSet<1024> TestSet;
+		uint32 SetBits[] = { 0, 3, 34, 35, 53, 65, 523, 534, 654, 1001, 1023 };
+		for (uint32 Index : SetBits)
+		{
+			const bool bWasSet = TestSet.SetBitAtomic(Index);
+			check(bWasSet == false);
+		}
+
+		const bool b32WasSet = TestSet.SetBitAtomic(34);
+		check(b32WasSet);
+
+		uint32 Expected = 0;
+		TestSet.ForEachSetBit([&](uint32 Index)
+		{
+			check(Index == SetBits[Expected]);
+			++Expected;
+		});
+		check(Expected == UE_ARRAY_COUNT(SetBits));
+
+		const bool b1023WasSet = TestSet.ClearBitAtomic(1023);
+		check(b1023WasSet == true);
+		TestSet.ForEachSetBit([&](uint32 Index)
+		{
+			check(Index != 1023);
+		});
+	}
 #endif // DO_CHECK
 }
 
diff --git a/Engine/Source/Runtime/Core/Public/ProfilingDebugging/StringsTrace.h b/Engine/Source/Runtime/Core/Public/ProfilingDebugging/StringsTrace.h
new file mode 100644
--- /dev/null
+++ b/Engine/Source/Runtime/Core/Public/ProfilingDebugging/StringsTrace.h
@@ -0,0 +1,88 @@
+﻿// Copyright Epic Games, Inc. All Rights Reserved.
+
+#pragma once
+
+#include "CoreFwd.h"
+#include "Trace/Trace.h"
+#include "Trace/Trace.inl"
+
+CORE_API UE_TRACE_EVENT_BEGIN_EXTERN(Strings, StaticString, Definition64bit)
+	UE_TRACE_EVENT_FIELD(UE::Trace::WideString, DisplayWide)
+	UE_TRACE_EVENT_FIELD(UE::Trace::AnsiString, DisplayAnsi)
+UE_TRACE_EVENT_END()
+
+CORE_API UE_TRACE_EVENT_BEGIN_EXTERN(Strings, FName, Definition32bit)
+	UE_TRACE_EVENT_FIELD(UE::Trace::WideString, DisplayWide)
+	UE_TRACE_EVENT_FIELD(UE::Trace::AnsiString, DisplayAnsi)
+UE_TRACE_EVENT_END()
+
+UE_TRACE_EVENT_BEGIN_EXTERN(Strings, StaticStringNoSync, Definition64bit | NoSync)
+	UE_TRACE_EVENT_FIELD(UE::Trace::WideString, DisplayWide)
+	UE_TRACE_EVENT_FIELD(UE::Trace::AnsiString, DisplayAnsi)
+UE_TRACE_EVENT_END()
+
+UE_TRACE_EVENT_BEGIN_EXTERN(Strings, FNameNoSync, Definition32bit | NoSync)
+	UE_TRACE_EVENT_FIELD(UE::Trace::WideString, DisplayWide)
+	UE_TRACE_EVENT_FIELD(UE::Trace::AnsiString, DisplayAnsi)
+UE_TRACE_EVENT_END()
+
+/**
+ * Utility class to trace deduplicated strings, FNames and static strings. Each function returning a
+ * reference that can be used when tracing events.
+ *
+ * A event defined as:
+ * \code
+ * UE_TRACE_EVENT_BEGIN(Asset)
+ *		UE_TRACE_EVENT_REFERENCE_FIELD(Strings, FName, Name)
+ *		UE_TRACE_EVENT_REFERENCE_FIELD(Strings, StaticString, File)
+ * UE_TRACE_EVENT_END()
+ *
+ * Can be emitted with:
+ * auto NameRef = FStringTrace::GetNameRef(Object->GetFName());
+ * auto FileRef = FStringTrace::GetStaticStringRef(__FILE__);
+ * UE_TRACE_LOG(Asset)
+ *		<< Asset.Name(NameRef)
+ *		<< Asset.File(FileRef);
+ * \endcode
+ */
+class FStringTrace
+{
+	
+public:
+	/**
+	 * Gets the trace id of a FName.
+	 * @param Name Name to trace.
+	 * @return Id that can be used to reference the name.
+	 */
+	CORE_API static UE::Trace::FEventRef32 GetNameRef(const FName& Name);
+
+	/**
+	 * Gets the trace id of a static string. Will use the address of the string
+	 * for deduplication, so it is important to only pass static strings to this
+	 * function.
+	 * @param String Static string
+	 * @return Id that can be used to reference the string.
+	 */
+	CORE_API static UE::Trace::FEventRef64 GetStaticStringRef(const TCHAR* String);
+	
+	/**
+	 * Gets the trace id of a static string. Will use the address of the string
+	 * for deduplication, so it is important to only pass static strings to this
+	 * function.
+	 * @param String Static string
+	 * @return Id that can be used to reference the string.
+	 */
+	CORE_API static UE::Trace::FEventRef64 GetStaticStringRef(const ANSICHAR* String);
+
+	/**
+	 * On a new connection callback.
+	 */
+	static void OnConnection();
+};
+
+#if !UE_TRACE_ENABLED
+FORCEINLINE UE::Trace::FEventRef32 FStringTrace::GetNameRef(const FName& Name) { return UE::Trace::FEventRef32(0,0); }
+FORCEINLINE UE::Trace::FEventRef64 FStringTrace::GetStaticStringRef(const TCHAR* String) { return UE::Trace::FEventRef64(0,0); }
+FORCEINLINE UE::Trace::FEventRef64 FStringTrace::GetStaticStringRef(const ANSICHAR* String) { return UE::Trace::FEventRef64(0,0); }
+FORCEINLINE void FStringTrace::OnConnection() {};
+#endif
\ No newline at end of file
diff --git a/Engine/Source/Runtime/Core/Public/UObject/NameTypes.h b/Engine/Source/Runtime/Core/Public/UObject/NameTypes.h
--- a/Engine/Source/Runtime/Core/Public/UObject/NameTypes.h
+++ b/Engine/Source/Runtime/Core/Public/UObject/NameTypes.h
@@ -15,6 +15,7 @@
 #include "Templates/Atomic.h"
 #include "Serialization/MemoryLayout.h"
 #include "Misc/StringBuilder.h"
+#include "Trace/Trace.h"
 
 /*----------------------------------------------------------------------------
 	Definitions.
@@ -1030,6 +1031,11 @@ PRAGMA_ENABLE_DEPRECATION_WARNINGS
 	static FNameEntry const* GetEntry(EName Ename);
 	static FNameEntry const* GetEntry(FNameEntryId Id);
 
+#if UE_TRACE_ENABLED
+	static UE::Trace::FEventRef32 TraceName(const FName& Name);
+	static void TraceNamesOnConnection();
+#endif
+
 	//@}
 
 	/** Run autotest on FNames. */
