Fix for minor typos in Array and ArrayView documentation.

diff --git a/Engine/Source/Runtime/Core/Public/Containers/Array.h b/Engine/Source/Runtime/Core/Public/Containers/Array.h
--- a/Engine/Source/Runtime/Core/Public/Containers/Array.h
+++ b/Engine/Source/Runtime/Core/Public/Containers/Array.h
@@ -679,8 +679,8 @@ public:
 	}
 
 	/**
-	 * Checks array invariants: if array size is greater than zero and less
-	 * than maximum.
+	 * Checks array invariants: if array size is greater than or equal to zero and less
+	 * than or equal to the maximum.
 	 */
 	FORCEINLINE void CheckInvariants() const
 	{
@@ -748,7 +748,7 @@ public:
 	}
 
 	/**
-	 * Array bracket operator. Returns reference to element at give index.
+	 * Array bracket operator. Returns reference to element at given index.
 	 *
 	 * @returns Reference to indexed element.
 	 */
@@ -759,7 +759,7 @@ public:
 	}
 
 	/**
-	 * Array bracket operator. Returns reference to element at give index.
+	 * Array bracket operator. Returns reference to element at given index.
 	 *
 	 * Const version of the above.
 	 *
@@ -1126,7 +1126,7 @@ public:
 	}
 
 	/**
-	 * Checks if this array contains element for which the predicate is true.
+	 * Checks if this array contains an element for which the predicate is true.
 	 *
 	 * @param Predicate to use
 	 * @returns	True if found. False otherwise.
diff --git a/Engine/Source/Runtime/Core/Public/Containers/ArrayView.h b/Engine/Source/Runtime/Core/Public/Containers/ArrayView.h
--- a/Engine/Source/Runtime/Core/Public/Containers/ArrayView.h
+++ b/Engine/Source/Runtime/Core/Public/Containers/ArrayView.h
@@ -113,7 +113,7 @@ namespace ArrayViewPrivate
  *
  * Caution:
  *   Treat a view like a *reference* to the elements in the array. DO NOT free or reallocate the array while the view exists!
- *   For this reason, be mindful of lifetime when constructing TArrayViews from rvalue initializer lists:
+ *   For this reason, be mindful of lifetimes when constructing TArrayViews from rvalue initializer lists:
  *
  *   TArrayView<int> View = { 1, 2, 3 }; // construction of array view from rvalue initializer list
  *   int n = View[0]; // undefined behavior, as the initializer list was destroyed at the end of the previous line
@@ -208,7 +208,7 @@ public:
 	/**
 	 * Helper function for returning a typed pointer to the first array entry.
 	 *
-	 * @returns Pointer to first array entry or nullptr if ArrayMax == 0.
+	 * @returns Pointer to first array entry.
 	 */
 	FORCEINLINE ElementType* GetData() const
 	{
@@ -234,8 +234,7 @@ public:
 	}
 
 	/**
-	 * Checks array invariants: if array size is greater than zero and less
-	 * than maximum.
+	 * Checks array invariants: if array size is greater than or equal to zero.
 	 */
 	FORCEINLINE void CheckInvariants() const
 	{
@@ -269,7 +268,7 @@ public:
 	}
 
 	/**
-	 * Tests if index is valid, i.e. than or equal to zero, and less than the number of elements in the array.
+	 * Tests if index is valid, i.e. greater than or equal to zero, and less than the number of elements in the array.
 	 *
 	 * @param Index Index to test.
 	 *
@@ -302,7 +301,7 @@ public:
 	}
 
 	/**
-	 * Array bracket operator. Returns reference to element at give index.
+	 * Array bracket operator. Returns reference to element at given index.
 	 *
 	 * @returns Reference to indexed element.
 	 */
@@ -328,7 +327,7 @@ public:
 
 	/**
 	 * Returns a sliced view
-	   The is similar to Mid(), but with a narrow contract, i.e. slicing outside of the range of the view is illegal.
+	 * This is similar to Mid(), but with a narrow contract, i.e. slicing outside of the range of the view is illegal.
 	 *
 	 * @param Index starting index of the new view
 	 * @param InNum number of elements in the new view
@@ -633,7 +632,7 @@ public:
 	}
 
 	/**
-	 * Checks if this array contains element for which the predicate is true.
+	 * Checks if this array contains an element for which the predicate is true.
 	 *
 	 * @param Predicate to use
 	 *
