Ensure Task timeouts don't pass a negative wait time to FEvents, as that will cause a check failure.

#rnx

diff --git a/Engine/Source/Runtime/Core/Public/Tasks/Task.h b/Engine/Source/Runtime/Core/Public/Tasks/Task.h
--- a/Engine/Source/Runtime/Core/Public/Tasks/Task.h
+++ b/Engine/Source/Runtime/Core/Public/Tasks/Task.h
@@ -401,7 +401,7 @@ namespace UE::Tasks
 			return true;
 		}
 
-		return CompletionEvent->Wait(Timeout.GetRemainingTime());
+		return CompletionEvent->Wait((uint32)FMath::Clamp<int64>(Timeout.GetRemainingTime().GetTicks() / ETimespan::TicksPerMillisecond, 0, MAX_uint32));
 	}
 
 	// wait for multiple tasks while executing other tasks
diff --git a/Engine/Source/Runtime/Core/Public/Tasks/TaskPrivate.h b/Engine/Source/Runtime/Core/Public/Tasks/TaskPrivate.h
--- a/Engine/Source/Runtime/Core/Public/Tasks/TaskPrivate.h
+++ b/Engine/Source/Runtime/Core/Public/Tasks/TaskPrivate.h
@@ -945,7 +945,7 @@ namespace UE::Tasks
 				return true;
 			}
 
-			return CompletionEvent->Wait(Timeout.GetRemainingTime());
+			return CompletionEvent->Wait((uint32)FMath::Clamp<int64>(Timeout.GetRemainingTime().GetTicks() / ETimespan::TicksPerMillisecond, 0, MAX_uint32));
 		}
 
 		// task retraction of multiple tasks. 
