LogTrace: Workaround for crash in TRACE_LOG_MESSAGE when encounters string arguments not null-terminated (Added and enabled LOGTRACE_RUNTIME_FORMATTING_ENABLED which traces the log messages already preformatted).

diff --git a/Engine/Source/Runtime/Core/Private/Logging/LogTrace.cpp b/Engine/Source/Runtime/Core/Private/Logging/LogTrace.cpp
--- a/Engine/Source/Runtime/Core/Private/Logging/LogTrace.cpp
+++ b/Engine/Source/Runtime/Core/Private/Logging/LogTrace.cpp
@@ -3,10 +3,12 @@
 
 #if LOGTRACE_ENABLED
 
-#include "Trace/Trace.inl"
-#include "Templates/Function.h"
 #include "HAL/PlatformTime.h"
 #include "HAL/PlatformTLS.h"
+#include "Misc/ScopeLock.h"
+#include "Misc/VarArgs.h"
+#include "Templates/Function.h"
+#include "Trace/Trace.inl"
 
 UE_TRACE_CHANNEL_DEFINE(LogChannel)
 
@@ -62,4 +64,21 @@ void FLogTrace::OutputLogMessageInternal(const void* LogPoint, uint16 EncodedFor
 		<< LogMessage.FormatArgs(EncodedFormatArgs, EncodedFormatArgsSize);
 }
 
-#endif
+#if LOGTRACE_RUNTIME_FORMATTING_ENABLED
+static FCriticalSection* GetLogTraceStaticBufferGuard()
+{
+	static FCriticalSection CS;
+	return &CS;
+}
+
+static TCHAR LogTraceStaticBuffer[8192];
+
+void FLogTrace::OutputLogMessageSimple(const void* LogPoint, const TCHAR* Fmt, ...)
+{
+	FScopeLock MsgLock(GetLogTraceStaticBufferGuard());
+	GET_VARARGS(LogTraceStaticBuffer, UE_ARRAY_COUNT(LogTraceStaticBuffer), UE_ARRAY_COUNT(LogTraceStaticBuffer) - 1, Fmt, Fmt);
+	FLogTrace::OutputLogMessage(LogPoint, LogTraceStaticBuffer);
+}
+#endif // LOGTRACE_RUNTIME_FORMATTING_ENABLED
+
+#endif // LOGTRACE_ENABLED
diff --git a/Engine/Source/Runtime/Core/Public/Logging/LogTrace.h b/Engine/Source/Runtime/Core/Public/Logging/LogTrace.h
--- a/Engine/Source/Runtime/Core/Public/Logging/LogTrace.h
+++ b/Engine/Source/Runtime/Core/Public/Logging/LogTrace.h
@@ -16,6 +16,10 @@
 #if LOGTRACE_ENABLED
 #include "ProfilingDebugging/FormatArgsTrace.h"
 
+#if !defined(LOGTRACE_RUNTIME_FORMATTING_ENABLED)
+#define LOGTRACE_RUNTIME_FORMATTING_ENABLED 1
+#endif
+
 UE_TRACE_CHANNEL_EXTERN(LogChannel, CORE_API)
 
 struct FLogCategoryBase;
@@ -24,6 +28,7 @@ struct FLogTrace
 {
 	CORE_API static void OutputLogCategory(const FLogCategoryBase* Category, const TCHAR* Name, ELogVerbosity::Type DefaultVerbosity);
 	CORE_API static void OutputLogMessageSpec(const void* LogPoint, const FLogCategoryBase* Category, ELogVerbosity::Type Verbosity, const ANSICHAR* File, int32 Line, const TCHAR* Format);
+	CORE_API static void OutputLogMessageSimple(const void* LogPoint, const TCHAR* Fmt, ...);
 
 	template <typename... Types>
 	FORCENOINLINE static void OutputLogMessage(const void* LogPoint, Types... FormatArgs)
@@ -43,6 +48,19 @@ private:
 #define TRACE_LOG_CATEGORY(Category, Name, DefaultVerbosity) \
 	FLogTrace::OutputLogCategory(Category, Name, DefaultVerbosity);
 
+#if LOGTRACE_RUNTIME_FORMATTING_ENABLED
+#define TRACE_LOG_MESSAGE(Category, Verbosity, Format, ...) \
+	if (UE_TRACE_CHANNELEXPR_IS_ENABLED(LogChannel)) \
+	{ \
+		static bool PREPROCESSOR_JOIN(__LogPoint, __LINE__); \
+		if (!PREPROCESSOR_JOIN(__LogPoint, __LINE__)) \
+		{ \
+			FLogTrace::OutputLogMessageSpec(&PREPROCESSOR_JOIN(__LogPoint, __LINE__), &Category, ELogVerbosity::Verbosity, __FILE__, __LINE__, TEXT("%s")); \
+			PREPROCESSOR_JOIN(__LogPoint, __LINE__) = true; \
+		} \
+		FLogTrace::OutputLogMessageSimple(&PREPROCESSOR_JOIN(__LogPoint, __LINE__), Format, ##__VA_ARGS__); \
+	}
+#else
 #define TRACE_LOG_MESSAGE(Category, Verbosity, Format, ...) \
 	if (UE_TRACE_CHANNELEXPR_IS_ENABLED(LogChannel)) \
 	{ \
@@ -54,6 +72,7 @@ private:
 		} \
 		FLogTrace::OutputLogMessage(&PREPROCESSOR_JOIN(__LogPoint, __LINE__), ##__VA_ARGS__); \
 	}
+#endif
 
 #else
 #define TRACE_LOG_CATEGORY(Category, Name, DefaultVerbosity)
