PlatformFile: Removed warnings from SetTimeStamp

These warnings are not logged on any other platforms, there is no way to silence them, and they can occur despite the best efforts of the caller because another process could delete the file at the same time.

diff --git a/Engine/Source/Runtime/Core/Private/Windows/WindowsPlatformFile.cpp b/Engine/Source/Runtime/Core/Private/Windows/WindowsPlatformFile.cpp
--- a/Engine/Source/Runtime/Core/Private/Windows/WindowsPlatformFile.cpp
+++ b/Engine/Source/Runtime/Core/Private/Windows/WindowsPlatformFile.cpp
@@ -1067,10 +1067,7 @@ public:
 		{
 			TRACE_PLATFORMFILE_END_OPEN(Handle);
 			const FILETIME ModificationFileTime = UEDateTimeToWindowsFileTime(DateTime);
-			if (!SetFileTime(Handle, nullptr, nullptr, &ModificationFileTime))
-			{
-				UE_LOG(LogTemp, Warning, TEXT("SetTimeStamp: Failed to SetFileTime on %s"), Filename);
-			}
+			SetFileTime(Handle, nullptr, nullptr, &ModificationFileTime);
 			TRACE_PLATFORMFILE_BEGIN_CLOSE(Handle);
 #if PLATFORMFILETRACE_ENABLED
 			// MSVC static analysis has a rule that reports the argument to CloseHandle as uninitialized memory after the call to CloseHandle, so we have to save it ahead of time
@@ -1093,7 +1090,6 @@ public:
 		else
 		{
 			TRACE_PLATFORMFILE_FAIL_OPEN(Filename);
-			UE_LOG(LogTemp, Warning, TEXT("SetTimeStamp: Failed to open file %s"), Filename);
 		}
 	}
 
