Default init std::atomic as by default it will be uninited pointer, which may not allow it to get allocated when doing GetBuffer

diff --git a/Engine/Source/Runtime/Core/Private/UObject/UnrealNames.cpp b/Engine/Source/Runtime/Core/Private/UObject/UnrealNames.cpp
--- a/Engine/Source/Runtime/Core/Private/UObject/UnrealNames.cpp
+++ b/Engine/Source/Runtime/Core/Private/UObject/UnrealNames.cpp
@@ -401,7 +401,7 @@ private:
 		return Buffer;
 	}
 	
-	std::atomic<WordType*> Data;
+	std::atomic<WordType*> Data{};
 };
 
 struct FNameSlot
@@ -3845,7 +3845,6 @@ PRAGMA_ENABLE_DEPRECATION_WARNINGS
 	}
 #endif
 
-#if 0
 	{
 		TNameAtomicBitSet<1024> TestSet;
 		uint32 SetBits[] = { 0, 3, 34, 35, 53, 65, 523, 534, 654, 1001, 1023 };
@@ -3873,8 +3872,6 @@ PRAGMA_ENABLE_DEPRECATION_WARNINGS
 			check(Index != 1023);
 		});
 	}
-#endif
-
 #endif // DO_CHECK
 }
 
