Tasks: fixed traces of the new frontend

diff --git a/Engine/Source/Runtime/Core/Private/Tasks/TaskPrivate.cpp b/Engine/Source/Runtime/Core/Private/Tasks/TaskPrivate.cpp
--- a/Engine/Source/Runtime/Core/Private/Tasks/TaskPrivate.cpp
+++ b/Engine/Source/Runtime/Core/Private/Tasks/TaskPrivate.cpp
@@ -13,6 +13,8 @@ namespace UE::Tasks
 	{
 		void FTaskBase::Schedule()
 		{
+			TaskTrace::Scheduled(GetTraceId());
+
 #if TASKGRAPH_NEW_FRONTEND
 			if (IsNamedThreadTask())
 			{
diff --git a/Engine/Source/Runtime/Core/Public/Async/TaskGraphInterfaces.h b/Engine/Source/Runtime/Core/Public/Async/TaskGraphInterfaces.h
--- a/Engine/Source/Runtime/Core/Public/Async/TaskGraphInterfaces.h
+++ b/Engine/Source/Runtime/Core/Public/Async/TaskGraphInterfaces.h
@@ -648,7 +648,7 @@ public:
 		template<typename...T>
 		FORCEINLINE_DEBUGGABLE FGraphEventRef ConstructAndDispatchWhenReady(T&&... Args)
 		{
-			FGraphEventRef Ref{ ConstructAndHold(Forward<T>(Args)...) };
+			FGraphEventRef Ref{ ConstructAndHoldImpl(Forward<T>(Args)...) };
 			Ref->TryLaunch();
 			return Ref;
 		}
@@ -656,6 +656,15 @@ public:
 		/** Passthrough internal task constructor and hold. */
 		template<typename...T>
 		FORCEINLINE_DEBUGGABLE TGraphTask* ConstructAndHold(T&&... Args)
+		{
+			TGraphTask* Task = ConstructAndHoldImpl(Forward<T>(Args)...);
+			TaskTrace::Created(Task->GetTraceId());
+			return Task;
+		}
+
+	private:
+		template<typename...T>
+		FORCEINLINE_DEBUGGABLE TGraphTask* ConstructAndHoldImpl(T&&... Args)
 		{
 			TGraphTask* Task = new TGraphTask(Prerequisites);
 			TTask* TaskObject = new(&Task->TaskStorage) TTask(Forward<T>(Args)...);
diff --git a/Engine/Source/Runtime/Core/Public/Tasks/TaskPrivate.h b/Engine/Source/Runtime/Core/Public/Tasks/TaskPrivate.h
--- a/Engine/Source/Runtime/Core/Public/Tasks/TaskPrivate.h
+++ b/Engine/Source/Runtime/Core/Public/Tasks/TaskPrivate.h
@@ -140,6 +140,7 @@ namespace UE::Tasks
 			virtual ~FTaskBase()
 			{
 				check(IsCompleted());
+				TaskTrace::Destroyed(GetTraceId());
 			}
 
 			// will be called to execute the task, must be implemented by a derived class that should call `FTaskBase::TryExecute` and pass the task body
@@ -255,7 +256,12 @@ namespace UE::Tasks
 			// returns false if the task is already completed and the subsequent wasn't added
 			bool AddSubsequent(FTaskBase& Subsequent)
 			{
-				return Subsequents.PushIfNotClosed(&Subsequent);
+				bool bSucceeded = Subsequents.PushIfNotClosed(&Subsequent);
+				if (bSucceeded)
+				{
+					TaskTrace::SubsequentAdded(GetTraceId(), Subsequent.GetTraceId());
+				}
+				return bSucceeded;
 			}
 
 			// A piped task is executed after the previous task from this pipe is completed. Tasks from the same pipe are not executed
@@ -414,6 +420,7 @@ namespace UE::Tasks
 
 				if (Nested.AddSubsequent(*this)) // "release" memory order
 				{
+					TaskTrace::NestedAdded(GetTraceId(), Nested.GetTraceId());
 					Nested.AddRef(); // keep it alive as we store it in `Prerequisites` and we can need it to try to retract it. it's released on closing the task
 					Prerequisites.Push(&Nested);
 				}
@@ -503,8 +510,6 @@ namespace UE::Tasks
 				AddRef(); // `LowLevelTask` will automatically release the internal reference after execution, but there can be pending nested tasks, so keep it alive
 				// it's released either later here if the task is closed, or when the last nested task is completed and unlocks its parent (in `TryUnlock`)
 
-				TaskTrace::FTaskTimingEventScope TaskEventScope(GetTraceId());
-
 				ReleasePrerequisites();
 
 				FTaskBase* PrevTask = ExchangeCurrentTask(this);
@@ -515,7 +520,10 @@ namespace UE::Tasks
 					StartPipeExecution();
 				}
 
-				TaskBody(*this);
+				{
+					TaskTrace::FTaskTimingEventScope TaskEventScope(GetTraceId());
+					TaskBody(*this);
+				}
 
 				if (GetPipe() != nullptr)
 				{
