Fix warning V501: There are identical sub-expressions to the left and to the right of the '==' operator: Deque.Max() == Deque.Max()
#rnx

diff --git a/Engine/Source/Runtime/Core/Tests/Containers/DequeTest.cpp b/Engine/Source/Runtime/Core/Tests/Containers/DequeTest.cpp
--- a/Engine/Source/Runtime/Core/Tests/Containers/DequeTest.cpp
+++ b/Engine/Source/Runtime/Core/Tests/Containers/DequeTest.cpp
@@ -415,7 +415,7 @@ TEST_CASE("System::Core::Containers::TDeque::EmplaceLast/PopFirst single element
 		for (int32 i = 0; i < Test::DefaultCapacity * 2; ++i)
 		{
 			CHECK(Test::EmplaceLastPopFirst(Deque, 1));  // Rotates head and tail
-			CHECK(Deque.Max() == Deque.Max());
+			CHECK(Deque.Max() == Test::DefaultCapacity);
 		}
 	}
 	CHECK(int32Token::EvenConstructionDestructionCalls(Test::DefaultCapacity * 2));
