Fix non-unity
#rnx

diff --git a/Engine/Source/Runtime/Core/Private/Experimental/Misc/ExecutionResource.cpp b/Engine/Source/Runtime/Core/Private/Experimental/Misc/ExecutionResource.cpp
--- a/Engine/Source/Runtime/Core/Private/Experimental/Misc/ExecutionResource.cpp
+++ b/Engine/Source/Runtime/Core/Private/Experimental/Misc/ExecutionResource.cpp
@@ -1,5 +1,6 @@
 // Copyright Epic Games, Inc. All Rights Reserved.
 #include "Experimental/Misc/ExecutionResource.h"
+#include "Containers/Array.h"
 
 namespace ExecutionResourceImpl
 {
diff --git a/Engine/Source/Runtime/Core/Public/Experimental/Misc/ExecutionResource.h b/Engine/Source/Runtime/Core/Public/Experimental/Misc/ExecutionResource.h
--- a/Engine/Source/Runtime/Core/Public/Experimental/Misc/ExecutionResource.h
+++ b/Engine/Source/Runtime/Core/Public/Experimental/Misc/ExecutionResource.h
@@ -1,6 +1,8 @@
 // Copyright Epic Games, Inc. All Rights Reserved.
 #pragma once
 
+#include "Templates/RefCounting.h"
+
 /**
  * Simple ref-counted interface that represents a resource that is held or required for the execution of some task item.
  */
