CompactBinary: Added a compact mode to the conversion to JSON

diff --git a/Engine/Source/Runtime/Core/Private/Serialization/CompactBinarySerialization.cpp b/Engine/Source/Runtime/Core/Private/Serialization/CompactBinarySerialization.cpp
--- a/Engine/Source/Runtime/Core/Private/Serialization/CompactBinarySerialization.cpp
+++ b/Engine/Source/Runtime/Core/Private/Serialization/CompactBinarySerialization.cpp
@@ -268,10 +268,16 @@ FArchive& operator<<(FArchive& Ar, FCbObject& Object)
 class FCbJsonWriter
 {
 public:
-	explicit FCbJsonWriter(FUtf8StringBuilderBase& InBuilder)
+	explicit FCbJsonWriter(
+		FUtf8StringBuilderBase& InBuilder,
+		FUtf8StringView InLineTerminator = LINE_TERMINATOR_ANSI,
+		FUtf8StringView InIndent = "\t",
+		FUtf8StringView InSpace = " ")
 		: Builder(InBuilder)
+		, Indent(InIndent)
+		, Space(InSpace)
 	{
-		NewLineAndIndent << LINE_TERMINATOR_ANSI;
+		NewLineAndIndent << InLineTerminator;
 	}
 
 	void WriteField(FCbFieldView Field)
@@ -282,7 +288,7 @@ public:
 		if (FUtf8StringView Name = Field.GetName(); !Name.IsEmpty())
 		{
 			AppendQuotedString(Name);
-			Builder << ANSITEXTVIEW(": ");
+			Builder << ANSITEXTVIEW(":") << Space;
 		}
 
 		switch (FCbValue Accessor = Field.GetValue(); Accessor.GetType())
@@ -293,13 +299,13 @@ public:
 		case ECbFieldType::Object:
 		case ECbFieldType::UniformObject:
 			Builder << '{';
-			NewLineAndIndent << '\t';
+			NewLineAndIndent << Indent;
 			bNeedsNewLine = true;
 			for (FCbFieldView It : Field)
 			{
 				WriteField(It);
 			}
-			NewLineAndIndent.RemoveSuffix(1);
+			NewLineAndIndent.RemoveSuffix(Indent.Len());
 			if (bNeedsComma)
 			{
 				WriteOptionalNewLine();
@@ -309,13 +315,13 @@ public:
 		case ECbFieldType::Array:
 		case ECbFieldType::UniformArray:
 			Builder << '[';
-			NewLineAndIndent << '\t';
+			NewLineAndIndent << Indent;
 			bNeedsNewLine = true;
 			for (FCbFieldView It : Field)
 			{
 				WriteField(It);
 			}
-			NewLineAndIndent.RemoveSuffix(1);
+			NewLineAndIndent.RemoveSuffix(Indent.Len());
 			if (bNeedsComma)
 			{
 				WriteOptionalNewLine();
@@ -378,21 +384,21 @@ public:
 		case ECbFieldType::CustomById:
 		{
 			FCbCustomById Custom = Accessor.AsCustomById();
-			Builder << "{ \"Id\": ";
+			Builder << "{" << Space << "\"Id\":" << Space;
 			Builder << Custom.Id;
-			Builder << ", \"Data\": ";
+			Builder << "," << Space << "\"Data\":" << Space;
 			AppendBase64String(Custom.Data);
-			Builder << " }";
+			Builder << Space << "}";
 			break;
 		}
 		case ECbFieldType::CustomByName:
 		{
 			FCbCustomByName Custom = Accessor.AsCustomByName();
-			Builder << "{ \"Name\": ";
+			Builder << "{" << Space << "\"Name\":" << Space;
 			AppendQuotedString(Custom.Name);
-			Builder << ", \"Data\": ";
+			Builder << "," << Space << "\"Data\":" << Space;
 			AppendBase64String(Custom.Data);
-			Builder << " }";
+			Builder << Space << "}";
 			break;
 		}
 		default:
@@ -459,8 +465,8 @@ private:
 	void AppendBase64String(FMemoryView Value)
 	{
 		Builder << '"';
-		checkf(Value.GetSize() <= 512 * 1024 * 1024, TEXT("Encoding 512 MiB or larger is not supported. ")
-			TEXT("Size: " UINT64_FMT), Value.GetSize());
+		checkf(Value.GetSize() <= 512 * 1024 * 1024,
+			TEXT("Encoding 512 MiB or larger is not supported. Size: " UINT64_FMT), Value.GetSize());
 		const uint32 EncodedSize = FBase64::GetEncodedDataSize(uint32(Value.GetSize()));
 		const int32 EncodedIndex = Builder.AddUninitialized(int32(EncodedSize));
 		FBase64::Encode(static_cast<const uint8*>(Value.GetData()), uint32(Value.GetSize()),
@@ -471,6 +477,8 @@ private:
 private:
 	FUtf8StringBuilderBase& Builder;
 	TUtf8StringBuilder<32> NewLineAndIndent;
+	FUtf8StringView Indent;
+	FUtf8StringView Space;
 	bool bNeedsComma{false};
 	bool bNeedsNewLine{false};
 };
@@ -481,4 +489,10 @@ void CompactBinaryToJson(const FCbObjectView& Object, FUtf8StringBuilderBase& Bu
 	Writer.WriteField(Object.AsFieldView());
 }
 
+void CompactBinaryToCompactJson(const FCbObjectView& Object, FUtf8StringBuilderBase& Builder)
+{
+	FCbJsonWriter Writer(Builder, "", "", "");
+	Writer.WriteField(Object.AsFieldView());
+}
+
 ///////////////////////////////////////////////////////////////////////////////////////////////////////////////////////
diff --git a/Engine/Source/Runtime/Core/Public/Serialization/CompactBinarySerialization.h b/Engine/Source/Runtime/Core/Public/Serialization/CompactBinarySerialization.h
--- a/Engine/Source/Runtime/Core/Public/Serialization/CompactBinarySerialization.h
+++ b/Engine/Source/Runtime/Core/Public/Serialization/CompactBinarySerialization.h
@@ -77,6 +77,10 @@ CORE_API FArchive& operator<<(FArchive& Ar, FCbObject& Object);
 
 ///////////////////////////////////////////////////////////////////////////////////////////////////////////////////////
 
+/** Convert the object to JSON in a multi-line format with indentation. */
 CORE_API void CompactBinaryToJson(const FCbObjectView& Object, FUtf8StringBuilderBase& Builder);
 
+/** Convert the object to JSON in a compact format with no added whitespace. */
+CORE_API void CompactBinaryToCompactJson(const FCbObjectView& Object, FUtf8StringBuilderBase& Builder);
+
 ///////////////////////////////////////////////////////////////////////////////////////////////////////////////////////
