Make TVector2 return bool to match TVector

diff --git a/Engine/Source/Runtime/Core/Public/Math/Vector2D.h b/Engine/Source/Runtime/Core/Public/Math/Vector2D.h
--- a/Engine/Source/Runtime/Core/Public/Math/Vector2D.h
+++ b/Engine/Source/Runtime/Core/Public/Math/Vector2D.h
@@ -502,11 +502,13 @@ public:
 
 	/**
 	* Normalize this vector in-place if it is large enough, set it to (0,0) otherwise.
+	* (Note this is different from TVector<>::Normalize, which leaves the vector unchanged if it is too small to normalize.)
 	*
 	* @param Tolerance Minimum squared length of vector for normalization.
 	* @see GetSafeNormal()
+	* @return true if the vector was normalized correctly, false if it was too small and set to zero.
 	*/
-	void Normalize(T Tolerance=UE_SMALL_NUMBER);
+	bool Normalize(T Tolerance=UE_SMALL_NUMBER);
 
 	/**
 	* Checks whether vector is near to zero within a specified tolerance.
@@ -1051,7 +1053,7 @@ FORCEINLINE TVector2<T> TVector2<T>::GetSafeNormal(T Tolerance) const
 }
 
 template<typename T>
-FORCEINLINE void TVector2<T>::Normalize(T Tolerance)
+FORCEINLINE bool TVector2<T>::Normalize(T Tolerance)
 {
 	const T SquareSum = X*X + Y*Y;
 	if(SquareSum > Tolerance)
@@ -1059,10 +1061,11 @@ FORCEINLINE void TVector2<T>::Normalize(T Tolerance)
 		const T Scale = FMath::InvSqrt(SquareSum);
 		X *= Scale;
 		Y *= Scale;
-		return;
+		return true;
 	}
 	X = 0.0f;
 	Y = 0.0f;
+	return false;
 }
 
 template<typename T>
