Removal of TEXT literal concatenations, which doesn't work in UTF-8 mode.

diff --git a/Engine/Source/Developer/DerivedDataCache/Private/DerivedDataBackends.cpp b/Engine/Source/Developer/DerivedDataCache/Private/DerivedDataBackends.cpp
--- a/Engine/Source/Developer/DerivedDataCache/Private/DerivedDataBackends.cpp
+++ b/Engine/Source/Developer/DerivedDataCache/Private/DerivedDataBackends.cpp
@@ -571,8 +571,8 @@ public:
 
 		if (Hierarchy)
 		{
-			UE_LOG(LogDerivedDataCache, Warning, TEXT("Node %s is disabled because there may be only one hierarchy node. ")
-				TEXT("Confirm there is only one hierarchy in the cache graph and that it is inside of any async node."), NodeName);
+			UE_LOG(LogDerivedDataCache, Warning, TEXT("Node %s is disabled because there may be only one hierarchy node. "
+				"Confirm there is only one hierarchy in the cache graph and that it is inside of any async node."), NodeName);
 			return MakeTuple(nullptr, ECacheStoreFlags::None);
 		}
 
diff --git a/Engine/Source/Developer/DerivedDataCache/Private/DerivedDataBuildJobContext.cpp b/Engine/Source/Developer/DerivedDataCache/Private/DerivedDataBuildJobContext.cpp
--- a/Engine/Source/Developer/DerivedDataCache/Private/DerivedDataBuildJobContext.cpp
+++ b/Engine/Source/Developer/DerivedDataCache/Private/DerivedDataBuildJobContext.cpp
@@ -163,8 +163,8 @@ void FBuildJobContext::EndAsyncBuild()
 
 void FBuildJobContext::SetCacheBucket(FCacheBucket Bucket)
 {
-	checkf(!Bucket.IsNull(), TEXT("Null cache bucket not allowed for build of '%s' by %s. ")
-		TEXT("The cache can be disabled by calling SetCachePolicyMask(~ECachePolicy::Default)."),
+	checkf(!Bucket.IsNull(), TEXT("Null cache bucket not allowed for build of '%s' by %s. "
+		"The cache can be disabled by calling SetCachePolicyMask(~ECachePolicy::Default)."),
 		*Job.GetName(), *WriteToString<32>(Job.GetFunction()));
 	CacheKey.Bucket = Bucket;
 }
@@ -172,12 +172,12 @@ void FBuildJobContext::SetCacheBucket(FCacheBucket Bucket)
 void FBuildJobContext::SetCachePolicyMask(ECachePolicy Policy)
 {
 	checkf(EnumHasAllFlags(Policy, ECachePolicy::SkipData),
-		TEXT("SkipData flags may not be masked out on the cache policy for build of '%s' by %s. ")
-		TEXT("Flags for skipping data may be set indirectly through FBuildPolicy."),
+		TEXT("SkipData flags may not be masked out on the cache policy for build of '%s' by %s. "
+		     "Flags for skipping data may be set indirectly through FBuildPolicy."),
 		*Job.GetName(), *WriteToString<32>(Job.GetFunction()));
 	checkf(EnumHasAllFlags(Policy, ECachePolicy::KeepAlive),
-		TEXT("KeepAlive flag may not be masked out on the cache policy for build of '%s' by %s. ")
-		TEXT("Flags for cache record lifetime may be set indirectly through FBuildPolicy."),
+		TEXT("KeepAlive flag may not be masked out on the cache policy for build of '%s' by %s. "
+		     "Flags for cache record lifetime may be set indirectly through FBuildPolicy."),
 		*Job.GetName(), *WriteToString<32>(Job.GetFunction()));
 	checkf(EnumHasAllFlags(Policy, ECachePolicy::PartialRecord),
 		TEXT("PartialRecord flag may not be masked out on the cache policy for build of '%s' by %s."),
@@ -188,16 +188,16 @@ void FBuildJobContext::SetCachePolicyMask(ECachePolicy Policy)
 void FBuildJobContext::SetBuildPolicyMask(EBuildPolicy Policy)
 {
 	checkf(EnumHasAllFlags(Policy, EBuildPolicy::Cache),
-		TEXT("Cache flags may not be masked out on the build policy for build of '%s' by %s. ")
-		TEXT("Flags for modifying cache operations may be set through SetCachePolicyMask."),
+		TEXT("Cache flags may not be masked out on the build policy for build of '%s' by %s. "
+		     "Flags for modifying cache operations may be set through SetCachePolicyMask."),
 		*Job.GetName(), *WriteToString<32>(Job.GetFunction()));
 	checkf(EnumHasAllFlags(Policy, EBuildPolicy::CacheKeepAlive),
-		TEXT("CacheKeepAlive flag may not be masked out on the build policy for build of '%s' by %s. ")
-		TEXT("Flags for cache record lifetime may only be set through the build session."),
+		TEXT("CacheKeepAlive flag may not be masked out on the build policy for build of '%s' by %s. "
+		     "Flags for cache record lifetime may only be set through the build session."),
 		*Job.GetName(), *WriteToString<32>(Job.GetFunction()));
 	checkf(EnumHasAllFlags(Policy, EBuildPolicy::SkipData),
-		TEXT("SkipData flag may not be masked out on the build policy for build of '%s' by %s. ")
-		TEXT("Flags for skipping the data may only be set through the build session."),
+		TEXT("SkipData flag may not be masked out on the build policy for build of '%s' by %s. "
+		     "Flags for skipping the data may only be set through the build session."),
 		*Job.GetName(), *WriteToString<32>(Job.GetFunction()));
 	BuildPolicyMask = Policy;
 }
diff --git a/Engine/Source/Developer/DerivedDataCache/Private/DerivedDataCachePolicy.cpp b/Engine/Source/Developer/DerivedDataCache/Private/DerivedDataCachePolicy.cpp
--- a/Engine/Source/Developer/DerivedDataCache/Private/DerivedDataCachePolicy.cpp
+++ b/Engine/Source/Developer/DerivedDataCache/Private/DerivedDataCachePolicy.cpp
@@ -139,8 +139,8 @@ public:
 		checkf(Value.Id.IsValid(), TEXT("Failed to add value policy because the ID is null."));
 		const int32 Index = Algo::LowerBoundBy(Values, Value.Id, &FCacheValuePolicy::Id);
 		checkf(!(Values.IsValidIndex(Index) && Values[Index].Id == Value.Id),
-			TEXT("Failed to add value policy with ID %s because it has an existing value policy with that ID. ")
-			TEXT("New: %s. Existing: %s."),
+			TEXT("Failed to add value policy with ID %s because it has an existing value policy with that ID. "
+			     "New: %s. Existing: %s."),
 			*WriteToString<32>(Value.Id), *WriteToString<128>(Value.Policy), *WriteToString<128>(Values[Index].Policy));
 		Values.Insert(Value, Index);
 	}
diff --git a/Engine/Source/Developer/DerivedDataCache/Private/DerivedDataRequestOwner.cpp b/Engine/Source/Developer/DerivedDataCache/Private/DerivedDataRequestOwner.cpp
--- a/Engine/Source/Developer/DerivedDataCache/Private/DerivedDataRequestOwner.cpp
+++ b/Engine/Source/Developer/DerivedDataCache/Private/DerivedDataRequestOwner.cpp
@@ -74,8 +74,8 @@ void FRequestOwnerShared::Begin(IRequest* Request)
 	{
 		FWriteScopeLock WriteLock(Lock);
 		checkf(BarrierCount > 0 || !bBeginExecuted,
-			TEXT("At least one FRequestBarrier must be in scope when beginning a request after the first request. ")
-			TEXT("The overload of End that invokes a callback handles this automatically for most use cases."));
+			TEXT("At least one FRequestBarrier must be in scope when beginning a request after the first request. "
+			     "The overload of End that invokes a callback handles this automatically for most use cases."));
 		check(Request);
 		Requests.Add(Request);
 		bBeginExecuted = true;
diff --git a/Engine/Source/Developer/DerivedDataCache/Private/FileSystemCacheStore.cpp b/Engine/Source/Developer/DerivedDataCache/Private/FileSystemCacheStore.cpp
--- a/Engine/Source/Developer/DerivedDataCache/Private/FileSystemCacheStore.cpp
+++ b/Engine/Source/Developer/DerivedDataCache/Private/FileSystemCacheStore.cpp
@@ -941,21 +941,21 @@ FFileSystemCacheStore::FFileSystemCacheStore(
 		}
 
 		UE_LOG(LogDerivedDataCache, Display,
-			TEXT("Performance to %s: Latency=%.02fms. RandomReadSpeed=%.02fMBs, RandomWriteSpeed=%.02fMBs. ")
-			TEXT("Assigned SpeedClass '%s'"),
+			TEXT("Performance to %s: Latency=%.02fms. RandomReadSpeed=%.02fMBs, RandomWriteSpeed=%.02fMBs. "
+			     "Assigned SpeedClass '%s'"),
 			*CachePath, SpeedStats.LatencyMS, SpeedStats.ReadSpeedMBs, SpeedStats.WriteSpeedMBs, LexToString(SpeedClass));
 
 		if (SpeedClass <= EBackendSpeedClass::Slow && !bReadOnly)
 		{
 			if (GIsBuildMachine)
 			{
-				UE_LOG(LogDerivedDataCache, Display, TEXT("Access to %s appears to be slow. ")
-					TEXT("'Touch' will be disabled and queries/writes will be limited."), *CachePath);
+				UE_LOG(LogDerivedDataCache, Display, TEXT("Access to %s appears to be slow. "
+					"'Touch' will be disabled and queries/writes will be limited."), *CachePath);
 			}
 			else
 			{
-				UE_LOG(LogDerivedDataCache, Warning, TEXT("Access to %s appears to be slow. ")
-					TEXT("'Touch' will be disabled and queries/writes will be limited."), *CachePath);
+				UE_LOG(LogDerivedDataCache, Warning, TEXT("Access to %s appears to be slow. "
+					"'Touch' will be disabled and queries/writes will be limited."), *CachePath);
 			}
 			bTouch = false;
 			//bReadOnly = true;
@@ -1144,8 +1144,8 @@ bool FFileSystemCacheStore::RunSpeedTest(
 					TCHAR ErrorBuffer[1024];
 					FPlatformMisc::GetSystemErrorMessage(ErrorBuffer, 1024, ErrorCode);
 					UE_LOG(LogDerivedDataCache, Warning,
-						TEXT("Fail to create %s, derived data cache to this directory will be read only. ")
-						TEXT("WriteError: %u (%s)"), *File, ErrorCode, ErrorBuffer);
+						TEXT("Fail to create %s, derived data cache to this directory will be read only. "
+						     "WriteError: %u (%s)"), *File, ErrorCode, ErrorBuffer);
 					bTestDataExists = false;
 					bWriteTestPassed = false;
 					break;
@@ -1527,8 +1527,8 @@ bool FFileSystemCacheStore::PutCacheRecord(
 					if (!Value.HasData() && !GetCacheContentExists(Key, Value.GetRawHash()))
 					{
 						UE_LOG(LogDerivedDataCache, Log,
-							TEXT("%s: Cache put of non-deterministic record will overwrite existing record due to ")
-							TEXT("missing value %s with hash %s for %s from '%.*s'"),
+							TEXT("%s: Cache put of non-deterministic record will overwrite existing record due to "
+							     "missing value %s with hash %s for %s from '%.*s'"),
 							*CachePath, *WriteToString<16>(Value.GetId()), *WriteToString<48>(Value.GetRawHash()),
 							*WriteToString<96>(Key), Name.Len(), Name.GetData());
 						return false;
diff --git a/Engine/Source/Developer/DerivedDataCache/Private/MemoryCacheStore.cpp b/Engine/Source/Developer/DerivedDataCache/Private/MemoryCacheStore.cpp
--- a/Engine/Source/Developer/DerivedDataCache/Private/MemoryCacheStore.cpp
+++ b/Engine/Source/Developer/DerivedDataCache/Private/MemoryCacheStore.cpp
@@ -261,8 +261,8 @@ void FMemoryCacheStore::Put(
 		if (MaxCacheSize > 0 && (CurrentCacheSize + RequiredSize) > MaxCacheSize)
 		{
 			UE_CLOG(!bMaxSizeExceeded, LogDerivedDataCache, Display,
-				TEXT("Failed to cache data. Maximum cache size reached. ")
-				TEXT("CurrentSize %" UINT64_FMT " KiB / MaxSize: %" UINT64_FMT " KiB"),
+				TEXT("Failed to cache data. Maximum cache size reached. "
+				     "CurrentSize %" UINT64_FMT " KiB / MaxSize: %" UINT64_FMT " KiB"),
 				CurrentCacheSize / 1024, MaxCacheSize / 1024);
 			bMaxSizeExceeded = true;
 			Status = EStatus::Ok;
@@ -455,8 +455,8 @@ void FMemoryCacheStore::PutValue(
 		if (MaxCacheSize > 0 && (CurrentCacheSize + RequiredSize) > MaxCacheSize)
 		{
 			UE_CLOG(!bMaxSizeExceeded, LogDerivedDataCache, Display,
-				TEXT("Failed to cache data. Maximum cache size reached. ")
-				TEXT("CurrentSize %" UINT64_FMT " KiB / MaxSize: %" UINT64_FMT " KiB"),
+				TEXT("Failed to cache data. Maximum cache size reached. "
+				     "CurrentSize %" UINT64_FMT " KiB / MaxSize: %" UINT64_FMT " KiB"),
 				CurrentCacheSize / 1024, MaxCacheSize / 1024);
 			bMaxSizeExceeded = true;
 			continue;
diff --git a/Engine/Source/Developer/DerivedDataCache/Private/PakFileCacheStore.cpp b/Engine/Source/Developer/DerivedDataCache/Private/PakFileCacheStore.cpp
--- a/Engine/Source/Developer/DerivedDataCache/Private/PakFileCacheStore.cpp
+++ b/Engine/Source/Developer/DerivedDataCache/Private/PakFileCacheStore.cpp
@@ -364,8 +364,8 @@ void FPakFileCacheStore::PutCachedData(const TCHAR* CacheKey, TArrayView<const u
 				}
 
 				UE_LOG(LogDerivedDataCache, Warning,
-					TEXT("%s: Repeated put of %s with different sized contents. Multiple contents will be in the file, ")
-					TEXT("but only the last will be in the index. This has wasted %" INT64_FMT " bytes in the file."),
+					TEXT("%s: Repeated put of %s with different sized contents. Multiple contents will be in the file, "
+					     "but only the last will be in the index. This has wasted %" INT64_FMT " bytes in the file."),
 					*CachePath, CacheKey, Item->Size);
 				CacheItems.Remove(Key);
 			}
@@ -595,10 +595,10 @@ bool IPakFileCacheStore::SortAndCopy(const FString &InputFilename, const FString
 
 	// Write out a TOC listing for debugging
 	FStringOutputDevice Output;
-	Output.Logf(TEXT("Asset,Size") LINE_TERMINATOR);
+	Output.Logf(TEXT("Asset,Size" LINE_TERMINATOR_ANSI));
 	for(int KeyIndex = 0; KeyIndex < KeyNames.Num(); KeyIndex++)
 	{
-		Output.Logf(TEXT("%s,%d") LINE_TERMINATOR, *KeyNames[KeyIndex], KeySizes[KeyIndex]);
+		Output.Logf(TEXT("%s,%d" LINE_TERMINATOR_ANSI), *KeyNames[KeyIndex], KeySizes[KeyIndex]);
 	}
 	FFileHelper::SaveStringToFile(Output, *FPaths::Combine(*FPaths::GetPath(OutputFilename), *(FPaths::GetBaseFilename(OutputFilename) + TEXT(".csv"))));
 	return true;
diff --git a/Engine/Source/Developer/DerivedDataCache/Private/Tests/HttpCacheStoreTest.cpp b/Engine/Source/Developer/DerivedDataCache/Private/Tests/HttpCacheStoreTest.cpp
--- a/Engine/Source/Developer/DerivedDataCache/Private/Tests/HttpCacheStoreTest.cpp
+++ b/Engine/Source/Developer/DerivedDataCache/Private/Tests/HttpCacheStoreTest.cpp
@@ -22,15 +22,15 @@
 
 DEFINE_LOG_CATEGORY_STATIC(LogHttpDerivedDataBackendTests, Log, All);
 
-#define TEST_NAME_ROOT TEXT("System.DerivedDataCache.HttpDerivedDataBackend")
+#define TEST_NAME_ROOT "System.DerivedDataCache.HttpDerivedDataBackend"
 
 #define IMPLEMENT_HTTPDERIVEDDATA_AUTOMATION_TEST( TClass, PrettyName, TFlags ) \
-	IMPLEMENT_CUSTOM_COMPLEX_AUTOMATION_TEST(TClass, FHttpCacheStoreTestBase, TEST_NAME_ROOT PrettyName, TFlags) \
+	IMPLEMENT_CUSTOM_COMPLEX_AUTOMATION_TEST(TClass, FHttpCacheStoreTestBase, TEXT(TEST_NAME_ROOT PrettyName), TFlags) \
 	void TClass::GetTests(TArray<FString>& OutBeautifiedNames, TArray <FString>& OutTestCommands) const \
 	{ \
 		if (CheckPrequisites()) \
 		{ \
-			OutBeautifiedNames.Add(TEST_NAME_ROOT PrettyName); \
+			OutBeautifiedNames.Add(TEXT(TEST_NAME_ROOT PrettyName)); \
 			OutTestCommands.Add(FString()); \
 		} \
 	}
@@ -598,7 +598,7 @@ TArray<FValue> CreateTestCacheValues(ICacheStore* InTestBackend, uint32 InNumVal
 }
 
 // Tests basic functionality for structured cache operations
-IMPLEMENT_HTTPDERIVEDDATA_AUTOMATION_TEST(FHttpCacheStoreTest, TEXT(".CacheStore"), EAutomationTestFlags::EditorContext | EAutomationTestFlags::ProductFilter)
+IMPLEMENT_HTTPDERIVEDDATA_AUTOMATION_TEST(FHttpCacheStoreTest, ".CacheStore", EAutomationTestFlags::EditorContext | EAutomationTestFlags::ProductFilter)
 bool FHttpCacheStoreTest::RunTest(const FString& Parameters)
 {
 	using namespace UE::DerivedData;
diff --git a/Engine/Source/Runtime/Core/Private/Async/TaskGraph.cpp b/Engine/Source/Runtime/Core/Private/Async/TaskGraph.cpp
--- a/Engine/Source/Runtime/Core/Private/Async/TaskGraph.cpp
+++ b/Engine/Source/Runtime/Core/Private/Async/TaskGraph.cpp
@@ -422,10 +422,10 @@ static ENamedThreads::Type ThreadPriorityFromChar(TCHAR InChar)
 FString FAutoConsoleTaskPriority::CreateFullHelpText(const TCHAR* Name, const TCHAR* OriginalHelp)
 {
 	return FString::Printf(
-		TEXT("%s\n")
-		TEXT("Arguments are three characters: [ThreadPriority][TaskPriority][TaskPriorityIfForcedToNormalThreadPriority] ")
-		TEXT("where ThreadPriority is 'h' or 'n' or 'b' (high/normal/background) and TaskPriority is 'h' or 'n' (high/normal). ")
-		TEXT("Example: %s bnh")
+		TEXT("%s\n"
+		     "Arguments are three characters: [ThreadPriority][TaskPriority][TaskPriorityIfForcedToNormalThreadPriority] "
+		     "where ThreadPriority is 'h' or 'n' or 'b' (high/normal/background) and TaskPriority is 'h' or 'n' (high/normal). "
+		     "Example: %s bnh")
 		, OriginalHelp, Name);
 }
 
diff --git a/Engine/Source/Runtime/Core/Private/Compression/CompressedBuffer.cpp b/Engine/Source/Runtime/Core/Private/Compression/CompressedBuffer.cpp
--- a/Engine/Source/Runtime/Core/Private/Compression/CompressedBuffer.cpp
+++ b/Engine/Source/Runtime/Core/Private/Compression/CompressedBuffer.cpp
@@ -277,8 +277,8 @@ FCompositeBuffer FBlockEncoder::Compress(const FCompositeBuffer& RawData, const
 
 	const int64 BlockCount = FMath::DivideAndRoundUp(RawSize, BlockSize);
 	check(BlockCount > 0);
-	checkf(BlockCount <= MAX_int32, TEXT("Raw data of size %" UINT64_FMT " with block size %" UINT64_FMT " requires ")
-		TEXT("%" UINT64_FMT " blocks, but the limit is %d."), RawSize, BlockSize, BlockCount, MAX_int32);
+	checkf(BlockCount <= MAX_int32, TEXT("Raw data of size %" UINT64_FMT " with block size %" UINT64_FMT " requires "
+		"%" UINT64_FMT " blocks, but the limit is %d."), RawSize, BlockSize, BlockCount, MAX_int32);
 
 	FBlake3 RawHash;
 	const uint64 MetaSize = sizeof(uint32) * BlockCount;
diff --git a/Engine/Source/Runtime/Core/Private/Containers/StackTracker.cpp b/Engine/Source/Runtime/Core/Private/Containers/StackTracker.cpp
--- a/Engine/Source/Runtime/Core/Private/Containers/StackTracker.cpp
+++ b/Engine/Source/Runtime/Core/Private/Containers/StackTracker.cpp
@@ -219,7 +219,7 @@ void FStackTracker::DumpStackTraces(int32 StackThreshold, FOutputDevice& Ar, flo
 				ANSICHAR AddressInformation[512];
 				AddressInformation[0] = 0;
 				FPlatformStackWalk::ProgramCounterToHumanReadableString( AddressIndex, CallStack.Addresses[AddressIndex], AddressInformation, UE_ARRAY_COUNT(AddressInformation)-1 );
-				CallStackString = CallStackString + LINE_TERMINATOR TEXT(",,,") + FString(AddressInformation);
+				CallStackString = CallStackString + TEXT(LINE_TERMINATOR_ANSI ",,,") + FString(AddressInformation);
 			}
 
 			// Finally log with ',' prefix so "Log:" can easily be discarded as row in Excel.
diff --git a/Engine/Source/Runtime/Core/Private/FileCache/FileCache.cpp b/Engine/Source/Runtime/Core/Private/FileCache/FileCache.cpp
--- a/Engine/Source/Runtime/Core/Private/FileCache/FileCache.cpp
+++ b/Engine/Source/Runtime/Core/Private/FileCache/FileCache.cpp
@@ -28,8 +28,8 @@ static int32 GFileCacheBlockSizeKB = 64;
 static FAutoConsoleVariableRef CVarFileCacheBlockSize(
 	TEXT("fc.BlockSize"),
 	GFileCacheBlockSizeKB,
-	TEXT("Size of each block in KB in the global file cache object\n")
-	TEXT("Should match packaging compression block size for optimal reading from packege"),
+	TEXT("Size of each block in KB in the global file cache object\n"
+	     "Should match packaging compression block size for optimal reading from packege"),
 	ECVF_ReadOnly
 );
 
diff --git a/Engine/Source/Runtime/Core/Private/GenericPlatform/GenericPlatformCrashContext.cpp b/Engine/Source/Runtime/Core/Private/GenericPlatform/GenericPlatformCrashContext.cpp
--- a/Engine/Source/Runtime/Core/Private/GenericPlatform/GenericPlatformCrashContext.cpp
+++ b/Engine/Source/Runtime/Core/Private/GenericPlatform/GenericPlatformCrashContext.cpp
@@ -866,7 +866,7 @@ void FGenericCrashContext::SerializeAsXML( const TCHAR* Filename ) const
 
 void FGenericCrashContext::AddCrashPropertyInternal(FString& Buffer, FStringView PropertyName, FStringView PropertyValue)
 {
-	Buffer.Appendf(TEXT("<%.*s>%s</%.*s>") LINE_TERMINATOR, 
+	Buffer.Appendf(TEXT("<%.*s>%s</%.*s>" LINE_TERMINATOR_ANSI), 
 		PropertyName.Len(), PropertyName.GetData(),
 		*FXmlEscapedString(PropertyValue), 
 		PropertyName.Len(), PropertyName.GetData()
@@ -930,7 +930,7 @@ void FGenericCrashContext::AppendPortableCallstack(FString& OutBuffer, TConstArr
 
 	for (const FCrashStackFrame& Frame : StackFrames)
 	{
-		OutBuffer.Appendf(TEXT("%-*s 0x%016llx + %-16llx") LINE_TERMINATOR, MaxModuleLength + 1, *FXmlEscapedString(Frame.ModuleName), Frame.BaseAddress, Frame.Offset);
+		OutBuffer.Appendf(TEXT("%-*s 0x%016llx + %-16llx" LINE_TERMINATOR_ANSI), MaxModuleLength + 1, *FXmlEscapedString(Frame.ModuleName), Frame.BaseAddress, Frame.Offset);
 	}
 }
 
@@ -951,7 +951,7 @@ void FGenericCrashContext::AddPortableCallStack() const
 
 void FGenericCrashContext::AddHeader(FString& Buffer)
 {
-	Buffer += TEXT("<?xml version=\"1.0\" encoding=\"UTF-8\"?>") LINE_TERMINATOR;
+	Buffer += TEXT("<?xml version=\"1.0\" encoding=\"UTF-8\"?>" LINE_TERMINATOR_ANSI);
 	BeginSection(Buffer, TEXT("FGenericCrashContext") );
 }
 
@@ -962,12 +962,12 @@ void FGenericCrashContext::AddFooter(FString& Buffer)
 
 void FGenericCrashContext::BeginSection(FString& Buffer, const TCHAR* SectionName)
 {
-	Buffer.Appendf(TEXT("<%s>") LINE_TERMINATOR, SectionName);
+	Buffer.Appendf(TEXT("<%s>" LINE_TERMINATOR_ANSI), SectionName);
 }
 
 void FGenericCrashContext::EndSection(FString& Buffer, const TCHAR* SectionName)
 {
-	Buffer.Appendf(TEXT("</%s>") LINE_TERMINATOR, SectionName);
+	Buffer.Appendf(TEXT("</%s>" LINE_TERMINATOR_ANSI), SectionName);
 }
 
 template<typename DEST>
diff --git a/Engine/Source/Runtime/Core/Private/GenericPlatform/GenericPlatformFile.cpp b/Engine/Source/Runtime/Core/Private/GenericPlatform/GenericPlatformFile.cpp
--- a/Engine/Source/Runtime/Core/Private/GenericPlatform/GenericPlatformFile.cpp
+++ b/Engine/Source/Runtime/Core/Private/GenericPlatform/GenericPlatformFile.cpp
@@ -202,9 +202,9 @@ static int32 GCacheHandleForPakFilesOnly = 1;
 static FAutoConsoleVariableRef CVarCacheHandleForPakFilesOnly(
 	TEXT("AsyncReadFile.CacheHandleForPakFilesOnly"),
 	GCacheHandleForPakFilesOnly,
-	TEXT("Control how Async read handle caches the underlying platform handle for files.\n")
-	TEXT("0: Cache the underlying platform handles for all files.\n")
-	TEXT("1: Cache the underlying platform handle for .pak files only (default).\n"),
+	TEXT("Control how Async read handle caches the underlying platform handle for files.\n"
+	     "0: Cache the underlying platform handles for all files.\n"
+	     "1: Cache the underlying platform handle for .pak files only (default).\n"),
 	ECVF_Default
 );
 
diff --git a/Engine/Source/Runtime/Core/Private/GenericPlatform/GenericPlatformMemory.cpp b/Engine/Source/Runtime/Core/Private/GenericPlatform/GenericPlatformMemory.cpp
--- a/Engine/Source/Runtime/Core/Private/GenericPlatform/GenericPlatformMemory.cpp
+++ b/Engine/Source/Runtime/Core/Private/GenericPlatform/GenericPlatformMemory.cpp
@@ -56,9 +56,9 @@ namespace GenericPlatformMemory
 	static FAutoConsoleVariableRef CVarMemoryPressureCriticalThresholdMB(
 		TEXT("memory.MemoryPressureCriticalThresholdMB"),
 		GMemoryPressureCriticalThresholdMB,
-		TEXT("When the available physical memory drops below this threshold memory stats will consider this to be at critical pressure.\n")
-		TEXT("Where a platform can specifically state it's memory pressure this test maybe ignored.\n")
-		TEXT("0 (default) critical pressure will not use the threshold."),
+		TEXT("When the available physical memory drops below this threshold memory stats will consider this to be at critical pressure.\n"
+		     "Where a platform can specifically state it's memory pressure this test maybe ignored.\n"
+		     "0 (default) critical pressure will not use the threshold."),
 		ECVF_Default);
 }
 
@@ -216,13 +216,13 @@ void FGenericPlatformMemory::OnOutOfMemory(uint64 Size, uint32 Alignment)
 			MemoryTrace_Free((uint64)BackupOOMMemoryPool, EMemoryTraceRootHeap::SystemMemory);
 		}
 
-		UE_LOG(LogMemory, Warning, TEXT("MemoryStats:")\
-			TEXT("\n\tAvailablePhysical %llu")\
-			TEXT("\n\t AvailableVirtual %llu")\
-			TEXT("\n\t     UsedPhysical %llu")\
-			TEXT("\n\t PeakUsedPhysical %llu")\
-			TEXT("\n\t      UsedVirtual %llu")\
-			TEXT("\n\t  PeakUsedVirtual %llu"),
+		UE_LOG(LogMemory, Warning, TEXT("MemoryStats:"
+			"\n\tAvailablePhysical %llu"
+			"\n\t AvailableVirtual %llu"
+			"\n\t     UsedPhysical %llu"
+			"\n\t PeakUsedPhysical %llu"
+			"\n\t      UsedVirtual %llu"
+			"\n\t  PeakUsedVirtual %llu"),
 			(uint64)PlatformMemoryStats.AvailablePhysical,
 			(uint64)PlatformMemoryStats.AvailableVirtual,
 			(uint64)PlatformMemoryStats.UsedPhysical,
diff --git a/Engine/Source/Runtime/Core/Private/GenericPlatform/GenericPlatformTime.cpp b/Engine/Source/Runtime/Core/Private/GenericPlatform/GenericPlatformTime.cpp
--- a/Engine/Source/Runtime/Core/Private/GenericPlatform/GenericPlatformTime.cpp
+++ b/Engine/Source/Runtime/Core/Private/GenericPlatform/GenericPlatformTime.cpp
@@ -257,8 +257,8 @@ protected:
 static FAutoConsoleCommand CPUTimeDumpCommand
 (
 	TEXT("CPUTime.Dump"),
-	TEXT("Usage -Delay=[NumSeconds=30]\n")
-	TEXT("If Delay==0, disables printing the CPU usage to the log\n")
-	TEXT("If Delay>0, starts printing the average CPU usage from the last n frames, clamps between 10 and 300"),
+	TEXT("Usage -Delay=[NumSeconds=30]\n"
+	     "If Delay==0, disables printing the CPU usage to the log\n"
+	     "If Delay>0, starts printing the average CPU usage from the last n frames, clamps between 10 and 300"),
 	FConsoleCommandWithArgsDelegate::CreateRaw( &FCPUTimeDump::Get(), &FCPUTimeDump::ExecuteCommand )
 );
diff --git a/Engine/Source/Runtime/Core/Private/Memory/CompositeBuffer.cpp b/Engine/Source/Runtime/Core/Private/Memory/CompositeBuffer.cpp
--- a/Engine/Source/Runtime/Core/Private/Memory/CompositeBuffer.cpp
+++ b/Engine/Source/Runtime/Core/Private/Memory/CompositeBuffer.cpp
@@ -113,8 +113,8 @@ void FCompositeBuffer::IterateRange(uint64 Offset, uint64 Size, TFunctionRef<voi
 void FCompositeBuffer::IterateRange(uint64 Offset, uint64 Size,
 	TFunctionRef<void (FMemoryView View, const FSharedBuffer& ViewOuter)> Visitor) const
 {
-	checkf(Offset + Size <= GetSize(), TEXT("Failed to access %" UINT64_FMT " bytes at offset %" UINT64_FMT)
-		TEXT(" of a composite buffer containing %" UINT64_FMT " bytes."), Size, Offset, GetSize());
+	checkf(Offset + Size <= GetSize(), TEXT("Failed to access %" UINT64_FMT " bytes at offset %" UINT64_FMT
+		" of a composite buffer containing %" UINT64_FMT " bytes."), Size, Offset, GetSize());
 	for (const FSharedBuffer& Segment : Segments)
 	{
 		if (const uint64 SegmentSize = Segment.GetSize(); Offset <= SegmentSize)
diff --git a/Engine/Source/Runtime/Core/Private/Misc/AssertionMacros.cpp b/Engine/Source/Runtime/Core/Private/Misc/AssertionMacros.cpp
--- a/Engine/Source/Runtime/Core/Private/Misc/AssertionMacros.cpp
+++ b/Engine/Source/Runtime/Core/Private/Misc/AssertionMacros.cpp
@@ -67,7 +67,7 @@ namespace
 	};
 }
 
-#define FILE_LINE_DESC TEXT(" [File:%s] [Line: %i] ")
+#define FILE_LINE_DESC_ANSI " [File:%s] [Line: %i] "
 
 /*
 	Ensure behavior
@@ -146,7 +146,7 @@ static void AssertFailedImplV(const ANSICHAR* Expr, const ANSICHAR* File, int32
 	FCString::Sprintf(ErrorString, TEXT("%s"), ANSI_TO_TCHAR(Expr));
 	if (GError)
 	{
-		GError->Logf(TEXT("Assertion failed: %s") FILE_LINE_DESC TEXT("\n%s\n"), ErrorString, ANSI_TO_TCHAR(File), Line, DescriptionString);
+		GError->Logf(TEXT("Assertion failed: %s" FILE_LINE_DESC_ANSI "\n%s\n"), ErrorString, ANSI_TO_TCHAR(File), Line, DescriptionString);
 	}
 }
 
@@ -177,14 +177,14 @@ FORCENOINLINE void StaticFailDebug( const TCHAR* Error, const ANSICHAR* File, in
 	}
 
 	FScopeLock Lock( &GetFailDebugCriticalSection());
-	FPlatformMisc::LowLevelOutputDebugStringf(TEXT("%s") FILE_LINE_DESC TEXT("\n%s\n"), Error, ANSI_TO_TCHAR(File), Line, Description);
+	FPlatformMisc::LowLevelOutputDebugStringf(TEXT("%s" FILE_LINE_DESC_ANSI "\n%s\n"), Error, ANSI_TO_TCHAR(File), Line, Description);
 
 	// Copy the detailed error into the error message.
 	TCHAR ErrorMessage[4096];
-	if (FCString::Snprintf( ErrorMessage, UE_ARRAY_COUNT( ErrorMessage ), TEXT( "%s" ) FILE_LINE_DESC TEXT( "\n%s\n" ), Error, ANSI_TO_TCHAR( File ), Line, DescriptionAndTrace) < 0)
+	if (FCString::Snprintf( ErrorMessage, UE_ARRAY_COUNT( ErrorMessage ), TEXT( "%s" FILE_LINE_DESC_ANSI "\n%s\n" ), Error, ANSI_TO_TCHAR( File ), Line, DescriptionAndTrace) < 0)
 	{
 		// Description and callstack was too long to fit in GErrorMessage. Use only description
-		FCString::Snprintf( ErrorMessage, UE_ARRAY_COUNT( ErrorMessage ), TEXT( "%s" ) FILE_LINE_DESC TEXT( "\n%s\n<< callstack too long >>" ), Error, ANSI_TO_TCHAR( File ), Line, Description);
+		FCString::Snprintf( ErrorMessage, UE_ARRAY_COUNT( ErrorMessage ), TEXT( "%s" FILE_LINE_DESC_ANSI "\n%s\n<< callstack too long >>" ), Error, ANSI_TO_TCHAR( File ), Line, Description);
 	}
 
 	// Copy the error message to the error history.
@@ -362,11 +362,11 @@ FORCENOINLINE void FDebug::EnsureFailed(const ANSICHAR* Expr, const ANSICHAR* Fi
 #if !NO_LOGGING
 		if (GEnsuresAreErrors)
 		{
-			UE_LOG(LogOutputDevice, Error, TEXT("%s") FILE_LINE_DESC TEXT("\n%s\n"), ErrorString, ANSI_TO_TCHAR(File), Line, Msg);
+			UE_LOG(LogOutputDevice, Error, TEXT("%s" FILE_LINE_DESC_ANSI "\n%s\n"), ErrorString, ANSI_TO_TCHAR(File), Line, Msg);
 		}
 		else
 		{
-			UE_LOG(LogOutputDevice, Warning, TEXT("%s") FILE_LINE_DESC TEXT("\n%s\n"), ErrorString, ANSI_TO_TCHAR(File), Line, Msg);
+			UE_LOG(LogOutputDevice, Warning, TEXT("%s" FILE_LINE_DESC_ANSI "\n%s\n"), ErrorString, ANSI_TO_TCHAR(File), Line, Msg);
 		}
 #endif
 	}
@@ -377,7 +377,7 @@ FORCENOINLINE void FDebug::EnsureFailed(const ANSICHAR* Expr, const ANSICHAR* Fi
 
 		// Create a final string that we'll output to the log (and error history buffer)
 		TCHAR ErrorMsg[16384];
-		FCString::Snprintf(ErrorMsg, UE_ARRAY_COUNT(ErrorMsg), TEXT("Ensure condition failed: %s [File:%s] [Line: %i]") LINE_TERMINATOR TEXT("%s") LINE_TERMINATOR TEXT("Stack: ") LINE_TERMINATOR, ANSI_TO_TCHAR(Expr), ANSI_TO_TCHAR(File), Line, Msg);
+		FCString::Snprintf(ErrorMsg, UE_ARRAY_COUNT(ErrorMsg), TEXT("Ensure condition failed: %s [File:%s] [Line: %i]" LINE_TERMINATOR_ANSI "%s" LINE_TERMINATOR_ANSI "Stack: " LINE_TERMINATOR_ANSI), ANSI_TO_TCHAR(Expr), ANSI_TO_TCHAR(File), Line, Msg);
 
 		// No debugger attached, so generate a call stack and submit a crash report
 		// Walk the stack and dump it to the allocated memory.
@@ -607,4 +607,4 @@ FORCENOINLINE void FDebug::DumpStackTraceToLog(const TCHAR* Heading, const ELogV
 #endif
 }
 
-#undef FILE_LINE_DESC
+#undef FILE_LINE_DESC_ANSI
diff --git a/Engine/Source/Runtime/Core/Private/Misc/ConfigCacheIni.cpp b/Engine/Source/Runtime/Core/Private/Misc/ConfigCacheIni.cpp
--- a/Engine/Source/Runtime/Core/Private/Misc/ConfigCacheIni.cpp
+++ b/Engine/Source/Runtime/Core/Private/Misc/ConfigCacheIni.cpp
@@ -1414,12 +1414,12 @@ void FConfigFile::WriteToStringInternal(FString& InOutText, bool bIsADefaultIniW
 					// If this is the first property we are writing of this section, then print the section name
 					if( InOutText.Len() == InitialInOutTextSize )
 					{
-						InOutText.Appendf(TEXT("[%s]") LINE_TERMINATOR, *SectionName);
+						InOutText.Appendf(TEXT("[%s]" LINE_TERMINATOR_ANSI), *SectionName);
 
 						// and if the section has any array of struct uniqueness keys, add them here
 						for (auto It = Section.ArrayOfStructKeys.CreateConstIterator(); It; ++It)
 						{
-							InOutText.Appendf(TEXT("@%s=%s") LINE_TERMINATOR, *It.Key().ToString(), *It.Value());
+							InOutText.Appendf(TEXT("@%s=%s" LINE_TERMINATOR_ANSI), *It.Key().ToString(), *It.Value());
 						}
 					}
 
@@ -1474,7 +1474,7 @@ void FConfigFile::WriteToStringInternal(FString& InOutText, bool bIsADefaultIniW
 		SectionNamesLeftToWrite.Add(kvpair.Key);
 	}
 
-	static const FString BlankLine(LINE_TERMINATOR LINE_TERMINATOR);
+	static const FString BlankLine(TEXT(LINE_TERMINATOR_ANSI LINE_TERMINATOR_ANSI));
 	auto AddSectionToText = [&InOutText, &InOutSectionTexts, &SectionNamesLeftToWrite](const FString& SectionName)
 	{
 		FString* SectionText = InOutSectionTexts.Find(SectionName);
@@ -1798,7 +1798,7 @@ void FConfigFile::SaveSourceToBackupFile()
 		const FString& SectionName = SectionIterator.Key();
 		const FConfigSection& Section = SectionIterator.Value();
 
-		Text += FString::Printf( TEXT("[%s]") LINE_TERMINATOR, *SectionName);
+		Text += FString::Printf( TEXT("[%s]" LINE_TERMINATOR_ANSI), *SectionName);
 
 		for( FConfigSection::TConstIterator PropertyIterator(Section); PropertyIterator; ++PropertyIterator )
 		{
@@ -3925,7 +3925,7 @@ public:
 		// Rebuild the file with the updated section.
 
 		FString NewFile = IniFileMakeup.BeforeSection + IniFileMakeup.Section + IniFileMakeup.AfterSection;
-		if (!NewFile.EndsWith(LINE_TERMINATOR LINE_TERMINATOR))
+		if (!NewFile.EndsWith(TEXT(LINE_TERMINATOR_ANSI LINE_TERMINATOR_ANSI)))
 		{
 			NewFile.AppendChars(LINE_TERMINATOR, TCString<TCHAR>::Strlen(LINE_TERMINATOR));
 		}
diff --git a/Engine/Source/Runtime/Core/Private/Misc/OutputDeviceAnsiError.cpp b/Engine/Source/Runtime/Core/Private/Misc/OutputDeviceAnsiError.cpp
--- a/Engine/Source/Runtime/Core/Private/Misc/OutputDeviceAnsiError.cpp
+++ b/Engine/Source/Runtime/Core/Private/Misc/OutputDeviceAnsiError.cpp
@@ -81,7 +81,7 @@ void FOutputDeviceAnsiError::HandleError()
 	if (GLog)
 	{
 		// print to log and flush it
-		UE_LOG( LogHAL, Log, TEXT( "=== Critical error: ===" ) LINE_TERMINATOR LINE_TERMINATOR TEXT( "%s" ) LINE_TERMINATOR, GErrorExceptionDescription );
+		UE_LOG( LogHAL, Log, TEXT( "=== Critical error: ===" LINE_TERMINATOR_ANSI LINE_TERMINATOR_ANSI "%s" LINE_TERMINATOR_ANSI), GErrorExceptionDescription );
 		UE_LOG(LogHAL, Log, TEXT("%s"), GErrorHist);
 
 		GLog->Panic();
diff --git a/Engine/Source/Runtime/Core/Private/Serialization/MemoryImage.cpp b/Engine/Source/Runtime/Core/Private/Serialization/MemoryImage.cpp
--- a/Engine/Source/Runtime/Core/Private/Serialization/MemoryImage.cpp
+++ b/Engine/Source/Runtime/Core/Private/Serialization/MemoryImage.cpp
@@ -21,7 +21,7 @@ IMPLEMENT_TYPE_LAYOUT(FPlatformTypeLayoutParameters);
 IMPLEMENT_TYPE_LAYOUT(FHashedName);
 
 // Guid for versioning format changes/bug fixes to memory images
-#define MEMORYIMAGE_DERIVEDDATA_VER		TEXT("B6A0A753ED62499D8F30931E61FCAB05")
+#define MEMORYIMAGE_DERIVEDDATA_VER_ANSI		"B6A0A753ED62499D8F30931E61FCAB05"
 
 static const uint32 NumTypeLayoutDescHashBuckets = 4357u;
 static const FTypeLayoutDesc* GTypeLayoutHashBuckets[NumTypeLayoutDescHashBuckets] = { nullptr };
@@ -111,7 +111,7 @@ FArchive& FPlatformTypeLayoutParameters::Serialize(FArchive& Ar)
 
 void FPlatformTypeLayoutParameters::AppendKeyString(FString& KeyString) const
 {
-	KeyString += FString::Printf(TEXT("FL_%08x_MFA_%08x_V_") MEMORYIMAGE_DERIVEDDATA_VER, Flags, MaxFieldAlignment);
+	KeyString += FString::Printf(TEXT("FL_%08x_MFA_%08x_V_" MEMORYIMAGE_DERIVEDDATA_VER_ANSI), Flags, MaxFieldAlignment);
 }
 
 // evaluated during static-initialization, so logging from regular check() macros won't work correctly
@@ -1717,9 +1717,9 @@ FMemoryImageSection* FMemoryImageSection::WritePointer(const FTypeLayoutDesc& St
 
 	checkf(OutOffsetToBase || bStaticTypeMatchesDerived, TEXT("Must consider OffsetToBase if static/derived types are different, %s/%s"), DerivedTypeDesc.Name, StaticTypeDesc.Name);
 	checkf(TypeDependencyIndex != INDEX_NONE || bStaticTypeMatchesDerived,
-		TEXT("Unable to store pointer to derived type %s, different from static type %s\n")
-		TEXT("Ensure UE_MEMORYIMAGE_TRACK_TYPE_DEPENDENCIES is set\n")
-		TEXT("Make sure derived type is not declared using DECLARE_INLINE_TYPE_LAYOUT()"),
+		TEXT("Unable to store pointer to derived type %s, different from static type %s\n"
+		     "Ensure UE_MEMORYIMAGE_TRACK_TYPE_DEPENDENCIES is set\n"
+		     "Make sure derived type is not declared using DECLARE_INLINE_TYPE_LAYOUT()"),
 		DerivedTypeDesc.Name, StaticTypeDesc.Name);
 
 	FFrozenMemoryImagePtr FrozenPtr;
diff --git a/Engine/Source/Runtime/Core/Private/Windows/WindowsPlatformMisc.cpp b/Engine/Source/Runtime/Core/Private/Windows/WindowsPlatformMisc.cpp
--- a/Engine/Source/Runtime/Core/Private/Windows/WindowsPlatformMisc.cpp
+++ b/Engine/Source/Runtime/Core/Private/Windows/WindowsPlatformMisc.cpp
@@ -83,12 +83,12 @@ THIRD_PARTY_INCLUDES_END
 static TAutoConsoleVariable<int32> CVarDriverDetectionMethod(
 	TEXT("r.DriverDetectionMethod"),
 	4,
-	TEXT("Defines which implementation is used to detect the GPU driver (to check for old drivers, logs and statistics)\n")
-	TEXT("  0: Iterate available drivers in registry and choose the one with the same name, if in question use next method (happens)\n")
-	TEXT("  1: Get the driver of the primary adapter (might not be correct when dealing with multiple adapters)\n")
-	TEXT("  2: Use DirectX LUID (would be the best, not yet implemented)\n")
-	TEXT("  3: Use Windows functions, use the primary device (might be wrong when API is using another adapter)\n")
-	TEXT("  4: Use Windows functions, use names such as DirectX Device (newest, most promising)"),
+	TEXT("Defines which implementation is used to detect the GPU driver (to check for old drivers, logs and statistics)\n"
+	     "  0: Iterate available drivers in registry and choose the one with the same name, if in question use next method (happens)\n"
+	     "  1: Get the driver of the primary adapter (might not be correct when dealing with multiple adapters)\n"
+	     "  2: Use DirectX LUID (would be the best, not yet implemented)\n"
+	     "  3: Use Windows functions, use the primary device (might be wrong when API is using another adapter)\n"
+	     "  4: Use Windows functions, use names such as DirectX Device (newest, most promising)"),
 	ECVF_RenderThreadSafe);
 
 int32 GetOSVersionsHelper( TCHAR* OutOSVersionLabel, int32 OSVersionLabelLength, TCHAR* OutOSSubVersionLabel, int32 OSSubVersionLabelLength )
@@ -2104,13 +2104,13 @@ void FWindowsPlatformMisc::PromptForRemoteDebugging(bool bIsEnsure)
 		FPlatformStackWalk::UploadLocalSymbols();
 
 		FCString::Sprintf(GErrorRemoteDebugPromptMessage, 
-			TEXT("Have a programmer remote debug this crash?\n")
-			TEXT("Hit NO to exit and submit error report as normal.\n")
-			TEXT("Otherwise, contact a programmer for remote debugging,\n")
-			TEXT("giving them the changelist number below.\n")
-			TEXT("Once they have confirmed they are connected to the machine,\n")
-			TEXT("hit YES to allow them to debug the crash.\n")
-			TEXT("[Changelist = %d]"),
+			TEXT("Have a programmer remote debug this crash?\n"
+			     "Hit NO to exit and submit error report as normal.\n"
+			     "Otherwise, contact a programmer for remote debugging,\n"
+			     "giving them the changelist number below.\n"
+			     "Once they have confirmed they are connected to the machine,\n"
+			     "hit YES to allow them to debug the crash.\n"
+			     "[Changelist = %d]"),
 			FEngineVersion::Current().GetChangelist());
 		FSlowHeartBeatScope SuspendHeartBeat;
 		if (MessageBox(0, GErrorRemoteDebugPromptMessage, TEXT("CRASHED"), MB_YESNO|MB_SYSTEMMODAL) == IDYES)
diff --git a/Engine/Source/Runtime/Core/Private/Windows/WindowsRunnableThread.cpp b/Engine/Source/Runtime/Core/Private/Windows/WindowsRunnableThread.cpp
--- a/Engine/Source/Runtime/Core/Private/Windows/WindowsRunnableThread.cpp
+++ b/Engine/Source/Runtime/Core/Private/Windows/WindowsRunnableThread.cpp
@@ -88,7 +88,7 @@ uint32 FRunnableThreadWin::GuardedRun()
 				GWarn->Flush();
 
 				// Append the thread name at the end of the error report.
-				FCString::Strncat( GErrorHist, LINE_TERMINATOR TEXT( "Crash in runnable thread " ), UE_ARRAY_COUNT( GErrorHist ) );
+				FCString::Strncat( GErrorHist, TEXT(LINE_TERMINATOR_ANSI "Crash in runnable thread " ), UE_ARRAY_COUNT( GErrorHist ) );
 				FCString::Strncat( GErrorHist, *ThreadName, UE_ARRAY_COUNT( GErrorHist ) );
 
 				// Crashed.
diff --git a/Engine/Source/Runtime/Core/Public/UObject/NameTypes.h b/Engine/Source/Runtime/Core/Public/UObject/NameTypes.h
--- a/Engine/Source/Runtime/Core/Public/UObject/NameTypes.h
+++ b/Engine/Source/Runtime/Core/Public/UObject/NameTypes.h
@@ -142,7 +142,8 @@ typedef FNameEntryId NAME_INDEX;
 
 
 /** this is the character used to separate a subobject root from its subobjects in a path name. */
-#define SUBOBJECT_DELIMITER				TEXT(":")
+#define SUBOBJECT_DELIMITER_ANSI		":"
+#define SUBOBJECT_DELIMITER				TEXT(SUBOBJECT_DELIMITER_ANSI)
 
 /** this is the character used to separate a subobject root from its subobjects in a path name, as a char */
 #define SUBOBJECT_DELIMITER_CHAR		TEXT(':')
diff --git a/Engine/Source/Runtime/CoreUObject/Private/Misc/PackageName.cpp b/Engine/Source/Runtime/CoreUObject/Private/Misc/PackageName.cpp
--- a/Engine/Source/Runtime/CoreUObject/Private/Misc/PackageName.cpp
+++ b/Engine/Source/Runtime/CoreUObject/Private/Misc/PackageName.cpp
@@ -627,8 +627,8 @@ void FPackageName::InternalFilenameToLongPackageName(FStringView InFilename, FSt
 
 	if (bIsValidLongPackageName && Result.Len() != Filename.Len())
 	{
-		UE_LOG(LogPackageName, Warning, TEXT("TryConvertFilenameToLongPackageName was passed an ObjectPath (%.*s) rather than a PackageName or FilePath; it will be converted to the PackageName. ")
-			TEXT("Accepting ObjectPaths is deprecated behavior and will be removed in a future release; TryConvertFilenameToLongPackageName will fail on ObjectPaths."), InFilename.Len(), InFilename.GetData());
+		UE_LOG(LogPackageName, Warning, TEXT("TryConvertFilenameToLongPackageName was passed an ObjectPath (%.*s) rather than a PackageName or FilePath; it will be converted to the PackageName. "
+			"Accepting ObjectPaths is deprecated behavior and will be removed in a future release; TryConvertFilenameToLongPackageName will fail on ObjectPaths."), InFilename.Len(), InFilename.GetData());
 	}
 
 	{
@@ -2488,8 +2488,8 @@ bool FPackageNameTests::RunTest(const FString& Parameters)
 				ActualSubObjectName != ExpectedSubObjectName || ActualExtension != ExpectedExtension ||
 				ActualFlexNameType != ExpectedFlexNameType || ActualFailureReason != ExpectedFailureReason)
 			{
-				AddError(FString::Printf(TEXT("Path '%.*s' failed FPackageName::TryConvertToMountedPath\n")
-					TEXT("got      %s,'%s','%s','%s','%s','%s',%d,%d,\nexpected %s,'%.*s','%.*s','%.*s','%.*s','%.*s',%d,%d."),
+				AddError(FString::Printf(TEXT("Path '%.*s' failed FPackageName::TryConvertToMountedPath\n"
+					"got      %s,'%s','%s','%s','%s','%s',%d,%d,\nexpected %s,'%.*s','%.*s','%.*s','%.*s','%.*s',%d,%d."),
 					InPath.Len(), InPath.GetData(),
 					bActualResult ? TEXT("true") : TEXT("false"),
 					*ActualLocalPath, *ActualPackageName, *ActualObjectName, *ActualSubObjectName, *ActualExtension,
diff --git a/Engine/Source/Runtime/CoreUObject/Private/UObject/Class.cpp b/Engine/Source/Runtime/CoreUObject/Private/UObject/Class.cpp
--- a/Engine/Source/Runtime/CoreUObject/Private/UObject/Class.cpp
+++ b/Engine/Source/Runtime/CoreUObject/Private/UObject/Class.cpp
@@ -4703,8 +4703,8 @@ void UClass::ValidateRuntimeReplicationData()
 						if (EnumHasAnyFlags(Struct->StructFlags, STRUCT_NetDeltaSerializeNative))
 						{
 							UE_LOG(LogClass, Warning,
-									TEXT("Property %s::%s (SourceClass: %s) with custom net delta serializer was not registered in ")
-									TEXT("GetLifetimeReplicatedProps. This property will replicate but you should still register it."),
+									TEXT("Property %s::%s (SourceClass: %s) with custom net delta serializer was not registered in "
+									     "GetLifetimeReplicatedProps. This property will replicate but you should still register it."),
 									ToCStr(GetName()), ToCStr(RepProp->GetName()), ToCStr(RepProp->GetOwnerClass()->GetName()));
 
 							continue;
@@ -4712,8 +4712,8 @@ void UClass::ValidateRuntimeReplicationData()
 					}
 
 					UE_LOG(LogClass, Warning,
-							TEXT("Property %s::%s (SourceClass: %s) was not registered in GetLifetimeReplicatedProps. ")
-							TEXT("This property will not be replicated. Use DISABLE_REPLICATED_PROPERTY if not replicating was intentional."),
+							TEXT("Property %s::%s (SourceClass: %s) was not registered in GetLifetimeReplicatedProps. "
+							     "This property will not be replicated. Use DISABLE_REPLICATED_PROPERTY if not replicating was intentional."),
 							ToCStr(GetName()), ToCStr(RepProp->GetName()), ToCStr(RepProp->GetOwnerClass()->GetName()));
 				}
 			}
@@ -6143,7 +6143,7 @@ bool UClass::IsShortTypeName(FStringView InClassPath)
 	{
 		return false;
 	}
-	return UE::String::FindFirstOfAnyChar(InClassPath, TEXT("./\\" SUBOBJECT_DELIMITER)) == INDEX_NONE;
+	return UE::String::FindFirstOfAnyChar(InClassPath, TEXT("./\\" SUBOBJECT_DELIMITER_ANSI)) == INDEX_NONE;
 }
 
 
diff --git a/Engine/Source/Runtime/CoreUObject/Private/UObject/Field.cpp b/Engine/Source/Runtime/CoreUObject/Private/UObject/Field.cpp
--- a/Engine/Source/Runtime/CoreUObject/Private/UObject/Field.cpp
+++ b/Engine/Source/Runtime/CoreUObject/Private/UObject/Field.cpp
@@ -770,7 +770,7 @@ FText FField::GetToolTipText(bool bShortTooltip) const
 	}
 	else
 	{
-		LocalizedToolTip = FText::Join(FText::FromString(TEXT(":") LINE_TERMINATOR), DisplayName, LocalizedToolTip);
+		LocalizedToolTip = FText::Join(FText::FromString(TEXT(":" LINE_TERMINATOR_ANSI)), DisplayName, LocalizedToolTip);
 	}
 
 	return LocalizedToolTip;
diff --git a/Engine/Source/Runtime/CoreUObject/Private/UObject/Obj.cpp b/Engine/Source/Runtime/CoreUObject/Private/UObject/Obj.cpp
--- a/Engine/Source/Runtime/CoreUObject/Private/UObject/Obj.cpp
+++ b/Engine/Source/Runtime/CoreUObject/Private/UObject/Obj.cpp
@@ -3498,7 +3498,7 @@ void ParseFunctionFlags(uint32 Flags, TArray<const TCHAR*>& Results)
 		TEXT("BlueprintEvent"),			// FUNC_BlueprintEvent
 		TEXT("BlueprintPure"),			// FUNC_BlueprintPure
 		TEXT("0x20000000"),
-		TEXT("Const")					// FUNC_Const
+		TEXT("Const"),					// FUNC_Const
 		TEXT("0x80000000"),
 	};
 
diff --git a/Engine/Source/Runtime/CoreUObject/Private/UObject/SavePackage/SavePackageUtilities.cpp b/Engine/Source/Runtime/CoreUObject/Private/UObject/SavePackage/SavePackageUtilities.cpp
--- a/Engine/Source/Runtime/CoreUObject/Private/UObject/SavePackage/SavePackageUtilities.cpp
+++ b/Engine/Source/Runtime/CoreUObject/Private/UObject/SavePackage/SavePackageUtilities.cpp
@@ -1892,8 +1892,8 @@ ESavePackageResult SaveBulkData(FLinkerSave* Linker, int64& InOutStartOffset, co
 		}
 		if (bBulkItemIsMapped && bRequestSaveByReference)
 		{
-			UE_LOG(LogSavePackage, Warning, TEXT("%s has bulk data that is mapped, but the save method is SAVE_BulkDataByReference.")
-				TEXT("This is not currently supported. Will not be mapped."), Filename);
+			UE_LOG(LogSavePackage, Warning, TEXT("%s has bulk data that is mapped, but the save method is SAVE_BulkDataByReference."
+				"This is not currently supported. Will not be mapped."), Filename);
 			BulkDataFlags &= ~BULKDATA_MemoryMappedPayload;
 			bBulkItemIsMapped = false;
 		}
diff --git a/Engine/Source/Runtime/CoreUObject/Private/UObject/TopLevelAssetPath.cpp b/Engine/Source/Runtime/CoreUObject/Private/UObject/TopLevelAssetPath.cpp
--- a/Engine/Source/Runtime/CoreUObject/Private/UObject/TopLevelAssetPath.cpp
+++ b/Engine/Source/Runtime/CoreUObject/Private/UObject/TopLevelAssetPath.cpp
@@ -127,7 +127,7 @@ bool FTopLevelAssetPath::TrySetPath(FWideStringView Path)
 			return true;
 		}
 
-		FAsciiSet Delim2("." SUBOBJECT_DELIMITER);
+		FAsciiSet Delim2(TEXT("." SUBOBJECT_DELIMITER_ANSI));
 		if (FAsciiSet::HasAny(AssetNameView, Delim2))
 		{
 			// Subobject path or is malformed and contains multiple '.' delimiters.
@@ -142,6 +142,13 @@ bool FTopLevelAssetPath::TrySetPath(FWideStringView Path)
 	return false;
 }
 
+bool FTopLevelAssetPath::TrySetPath(FUtf8StringView Path)
+{
+	TStringBuilder<FName::StringBufferSize> Wide;
+	Wide << Path;
+	return TrySetPath(Wide);
+}
+
 bool FTopLevelAssetPath::TrySetPath(FAnsiStringView Path)
 {
 	TStringBuilder<FName::StringBufferSize> Wide;
diff --git a/Engine/Source/Runtime/CoreUObject/Private/UObject/UObjectGlobals.cpp b/Engine/Source/Runtime/CoreUObject/Private/UObject/UObjectGlobals.cpp
--- a/Engine/Source/Runtime/CoreUObject/Private/UObject/UObjectGlobals.cpp
+++ b/Engine/Source/Runtime/CoreUObject/Private/UObject/UObjectGlobals.cpp
@@ -3107,15 +3107,15 @@ UObject* StaticAllocateObject
 
 			if (InClass->HasAnyClassFlags(CLASS_PerObjectConfig) && InOuter != nullptr && InOuter->GetOutermost() == GetTransientPackage())
 			{
-				ErrorPrefix = TEXT("PerObjectConfig object using the transient package, has triggered a name conflict and will now crash.\n")
-						TEXT("To avoid this, don't use the transient package for PerObjectConfig objects.\n")
-						TEXT("This has the side effect, of using the full path name for config ini sections. Use 'OverridePerObjectConfigSection' to keep the short name.\n\n");
+				ErrorPrefix = TEXT("PerObjectConfig object using the transient package, has triggered a name conflict and will now crash.\n"
+						"To avoid this, don't use the transient package for PerObjectConfig objects.\n"
+						"This has the side effect, of using the full path name for config ini sections. Use 'OverridePerObjectConfigSection' to keep the short name.\n\n");
 			}
 
 			UE_LOG(LogUObjectGlobals, Fatal,
-				TEXT("%sObjects have the same fully qualified name but different paths.\n")
-				TEXT("\tNew Object: %s %s.%s\n")
-				TEXT("\tExisting Object: %s"),
+				TEXT("%sObjects have the same fully qualified name but different paths.\n"
+				     "\tNew Object: %s %s.%s\n"
+				     "\tExisting Object: %s"),
 				ErrorPrefix, *InClass->GetName(), InOuter ? *InOuter->GetPathName() : TEXT(""), *InName.ToString(),
 				*Obj->GetFullName());
 		}
@@ -4948,7 +4948,7 @@ FString FAssetMsg::FormatPathForAssetLog(const UObject* Object)
 
 FString FAssetMsg::GetAssetLogString(const TCHAR* Path, const FString& Message)
 {
-	return FString::Printf(ASSET_LOG_FORMAT_STRING TEXT("%s"), *FAssetMsg::FormatPathForAssetLog(Path), *Message);
+	return FString::Printf(TEXT(ASSET_LOG_FORMAT_STRING_ANSI "%s"), *FAssetMsg::FormatPathForAssetLog(Path), *Message);
 }
 
 FString FAssetMsg::GetAssetLogString(const UObject* Object, const FString& Message)
diff --git a/Engine/Source/Runtime/CoreUObject/Public/UObject/TopLevelAssetPath.h b/Engine/Source/Runtime/CoreUObject/Public/UObject/TopLevelAssetPath.h
--- a/Engine/Source/Runtime/CoreUObject/Public/UObject/TopLevelAssetPath.h
+++ b/Engine/Source/Runtime/CoreUObject/Public/UObject/TopLevelAssetPath.h
@@ -60,6 +60,7 @@ struct COREUOBJECT_API FTopLevelAssetPath
 	bool TrySetPath(FName InPackageName, FName InAssetName);
 	/** Sets asset path of this reference based on a string path. Resets this object and returns false if the string is empty or does not represent a top level asset path. */
 	bool TrySetPath(FWideStringView Path);
+	bool TrySetPath(FUtf8StringView Path);
 	bool TrySetPath(FAnsiStringView Path);
 	template<typename CharType>
 	bool TrySetPath(const CharType* Path) { return TrySetPath(TStringView<CharType>(Path)); }
diff --git a/Engine/Source/Runtime/CoreUObject/Public/UObject/UObjectGlobals.h b/Engine/Source/Runtime/CoreUObject/Public/UObject/UObjectGlobals.h
--- a/Engine/Source/Runtime/CoreUObject/Public/UObject/UObjectGlobals.h
+++ b/Engine/Source/Runtime/CoreUObject/Public/UObject/UObjectGlobals.h
@@ -2848,7 +2848,8 @@ struct FAssetMsg
 	static COREUOBJECT_API FString GetAssetLogString(const UObject* Object, const FString& Message);
 };
 
-#define ASSET_LOG_FORMAT_STRING TEXT("[AssetLog] %s: ")
+#define ASSET_LOG_FORMAT_STRING_ANSI "[AssetLog] %s: "
+#define ASSET_LOG_FORMAT_STRING TEXT(ASSET_LOG_FORMAT_STRING_ANSI)
 
 #if NO_LOGGING
 	#define UE_ASSET_LOG(...)
diff --git a/Engine/Source/Runtime/CoreUObject/Tests/ObjectPtrTest.cpp b/Engine/Source/Runtime/CoreUObject/Tests/ObjectPtrTest.cpp
--- a/Engine/Source/Runtime/CoreUObject/Tests/ObjectPtrTest.cpp
+++ b/Engine/Source/Runtime/CoreUObject/Tests/ObjectPtrTest.cpp
@@ -321,7 +321,7 @@ TEST_CASE_METHOD(FObjectPtrTestBase, "CoreUObject::TObjectPtr::Long Path", "[Cor
 
 // @TODO: OBJPTR: We should have a test that ensures that lazy loading of an object with an external package is handled correctly.
 //				  This should also include external packages in the outer chain of the target object.
-// IMPLEMENT_CUSTOM_SIMPLE_AUTOMATION_TEST(FObjectPtrTestExternalPackages, FObjectPtrTestBase, TEST_NAME_ROOT TEXT(".ExternalPackages"), ObjectPtrTestFlags)
+// IMPLEMENT_CUSTOM_SIMPLE_AUTOMATION_TEST(FObjectPtrTestExternalPackages, FObjectPtrTestBase, TEXT(TEST_NAME_ROOT ".ExternalPackages"), ObjectPtrTestFlags)
 // bool FObjectPtrTestExternalPackages::RunTest(const FString& Parameters)
 // {
 // 	const FName TestExternalPackage1Name(TEXT("/Engine/Test/ObjectPtrExternalPackages1/Transient"));
@@ -348,7 +348,7 @@ TEST_CASE_METHOD(FObjectPtrTestBase, "CoreUObject::TObjectPtr::Long Path", "[Cor
 // @TODO: OBJPTR: We should have a test that ensures that we can (de)serialize an FObjectPtr to FLinkerSave/FLinkerLoad and that upon load the object
 //			pointer is not resolved if we are in a configuration that supports lazy load.  This is proving difficult due to the restrictions around how
 //			FLinkerSave/FLinkerLoad is used.
-// IMPLEMENT_CUSTOM_SIMPLE_AUTOMATION_TEST(FObjectPtrTestLinkerSerializeBehavior, FObjectPtrTestBase, TEST_NAME_ROOT TEXT(".LinkerSerialize"), ObjectPtrTestFlags)
+// IMPLEMENT_CUSTOM_SIMPLE_AUTOMATION_TEST(FObjectPtrTestLinkerSerializeBehavior, FObjectPtrTestBase, TEXT(TEST_NAME_ROOT ".LinkerSerialize"), ObjectPtrTestFlags)
 // bool FObjectPtrTestLinkerSerializeBehavior::RunTest(const FString& Parameters)
 // {
 // 	FSnapshotObjectRefMetrics ObjectRefMetrics(*this);
diff --git a/Engine/Source/Runtime/Launch/Private/LaunchEngineLoop.cpp b/Engine/Source/Runtime/Launch/Private/LaunchEngineLoop.cpp
--- a/Engine/Source/Runtime/Launch/Private/LaunchEngineLoop.cpp
+++ b/Engine/Source/Runtime/Launch/Private/LaunchEngineLoop.cpp
@@ -5616,21 +5616,21 @@ void FEngineLoop::ClearPendingCleanupObjects()
 static TAutoConsoleVariable<int32> CVarLogTimestamp(
 	TEXT("log.Timestamp"),
 	1,
-	TEXT("Defines if time is included in each line in the log file and in what form. Layout: [time][frame mod 1000]\n")
-	TEXT("  0 = Do not display log timestamps\n")
-	TEXT("  1 = Log time stamps in UTC and frame time (default) e.g. [2015.11.25-21.28.50:803][376]\n")
-	TEXT("  2 = Log timestamps in seconds elapsed since GStartTime e.g. [0130.29][420]")
-	TEXT("  3 = Log timestamps in local time and frame time e.g. [2017.08.04-17.59.50:803][420]")
-	TEXT("  4 = Log timestamps with the engine's timecode and frame time e.g. [17:59:50:18][420]"),
+	TEXT("Defines if time is included in each line in the log file and in what form. Layout: [time][frame mod 1000]\n"
+	     "  0 = Do not display log timestamps\n"
+	     "  1 = Log time stamps in UTC and frame time (default) e.g. [2015.11.25-21.28.50:803][376]\n"
+	     "  2 = Log timestamps in seconds elapsed since GStartTime e.g. [0130.29][420]"
+	     "  3 = Log timestamps in local time and frame time e.g. [2017.08.04-17.59.50:803][420]"
+	     "  4 = Log timestamps with the engine's timecode and frame time e.g. [17:59:50:18][420]"),
 	ECVF_Default);
 
 
 static TAutoConsoleVariable<int32> CVarLogCategory(
 	TEXT("log.Category"),
 	1,
-	TEXT("Defines if the categoy is included in each line in the log file and in what form.\n")
-	TEXT("  0 = Do not log category\n")
-	TEXT("  2 = Log the category (default)"),
+	TEXT("Defines if the categoy is included in each line in the log file and in what form.\n"
+	     "  0 = Do not log category\n"
+	     "  2 = Log the category (default)"),
 	ECVF_Default);
 
 
@@ -5793,7 +5793,7 @@ bool FEngineLoop::AppInit( )
 
 
 	// Error history.
-	FCString::Strcpy(GErrorHist, TEXT("Fatal error!" LINE_TERMINATOR LINE_TERMINATOR));
+	FCString::Strcpy(GErrorHist, TEXT("Fatal error!" LINE_TERMINATOR_ANSI LINE_TERMINATOR_ANSI));
 
 	// Platform specific pre-init.
 	{
