Remove another hardcoded NDA platform define.This can be defined for a game by specifying bAllowConfidentialPlatformDefines as a target rule.

diff --git a/Engine/Source/Runtime/Core/Public/HAL/Platform.h b/Engine/Source/Runtime/Core/Public/HAL/Platform.h
--- a/Engine/Source/Runtime/Core/Public/HAL/Platform.h
+++ b/Engine/Source/Runtime/Core/Public/HAL/Platform.h
@@ -10,9 +10,6 @@
 #if !defined(PLATFORM_WINDOWS)
 	#define PLATFORM_WINDOWS 0
 #endif
-#if !defined(PLATFORM_XBOXONE)
-	#define PLATFORM_XBOXONE 0
-#endif
 #if !defined(PLATFORM_MAC)
 	#define PLATFORM_MAC 0
 	// If PLATFORM_MAC is defined these will be set appropriately in
