Fix on-connection strings having different type ids.

Since there were two different events for declaring strings and FNames, one regular and one no-sync version used in on-connection callback, references would not be correctly resolved. We remove the no sync version and use only one version. In order to solve the synchronizing we defer resolving the strings in analysis.

diff --git a/Engine/Source/Runtime/Core/Private/ProfilingDebugging/StringTrace.cpp b/Engine/Source/Runtime/Core/Private/ProfilingDebugging/StringTrace.cpp
--- a/Engine/Source/Runtime/Core/Private/ProfilingDebugging/StringTrace.cpp
+++ b/Engine/Source/Runtime/Core/Private/ProfilingDebugging/StringTrace.cpp
@@ -45,17 +45,19 @@ static struct FStaticStringTracer
 	void TraceStringsOnConnection()
 	{
 		FScopeLock _(&LookupLock);
+		// Retrace all strings that has been referenced previously. Note that we create the reference id
+		// manually so that we get the correct type id for sync version.
 		for(const uint64 Ptr : TracedWideStrings)
 		{
 			const TCHAR* String = (const TCHAR*) Ptr;
-			UE_TRACE_LOG_DEFINITION(Strings, StaticStringNoSync, Ptr, true)
-				<< StaticStringNoSync.DisplayWide(String, FCString::Strlen(String));
+			UE_TRACE_LOG_DEFINITION(Strings, StaticString, Ptr, true)
+				<< StaticString.DisplayWide(String, FCString::Strlen(String));
 		}
 		for(const uint64 Ptr : TracedAnsiStrings)
 		{
 			const ANSICHAR* String = (const ANSICHAR*) Ptr;
-			UE_TRACE_LOG_DEFINITION(Strings, StaticStringNoSync, Ptr, true)
-				<< StaticStringNoSync.DisplayAnsi(String, FCStringAnsi::Strlen(String));
+			UE_TRACE_LOG_DEFINITION(Strings, StaticString, Ptr, true)
+				<< StaticString.DisplayAnsi(String, FCStringAnsi::Strlen(String));
 		}
 	}
 
diff --git a/Engine/Source/Runtime/Core/Private/UObject/UnrealNames.cpp b/Engine/Source/Runtime/Core/Private/UObject/UnrealNames.cpp
--- a/Engine/Source/Runtime/Core/Private/UObject/UnrealNames.cpp
+++ b/Engine/Source/Runtime/Core/Private/UObject/UnrealNames.cpp
@@ -1918,19 +1918,19 @@ void FNamePool::RetraceAll() const
 				
 			if (DisplayEntry.IsWide())
 			{
-				 TCHAR Chars[NAME_SIZE];
-				 DisplayEntry.GetUnterminatedName(Chars, NAME_SIZE);
-				 const uint32 Size = DisplayEntry.GetNameLength();
-				 const auto Ref = UE_TRACE_LOG_DEFINITION(Strings, FNameNoSync, DefinitionId, true)
-					  << FNameNoSync.DisplayWide(Chars, Size);
+				TCHAR Chars[NAME_SIZE];
+				DisplayEntry.GetUnterminatedName(Chars, NAME_SIZE);
+				const uint32 Size = DisplayEntry.GetNameLength();
+				UE_TRACE_LOG_DEFINITION(Strings, FName, DefinitionId, true)
+					<< FName.DisplayWide(Chars, Size);
 			}
 			else
 			{
-				 ANSICHAR Chars[NAME_SIZE];
-				 DisplayEntry.GetAnsiName(Chars);
-				 const uint32 Size = DisplayEntry.GetNameLength();
-				 const auto Ref = UE_TRACE_LOG_DEFINITION(Strings, FNameNoSync, DefinitionId, true)
-					  << FNameNoSync.DisplayAnsi(Chars, Size);
+				ANSICHAR Chars[NAME_SIZE];
+				DisplayEntry.GetAnsiName(Chars);
+				const uint32 Size = DisplayEntry.GetNameLength();
+				UE_TRACE_LOG_DEFINITION(Strings, FName, DefinitionId, true)
+					<< FName.DisplayAnsi(Chars, Size);
 			}
 		});
 
diff --git a/Engine/Source/Runtime/Core/Public/ProfilingDebugging/StringsTrace.h b/Engine/Source/Runtime/Core/Public/ProfilingDebugging/StringsTrace.h
--- a/Engine/Source/Runtime/Core/Public/ProfilingDebugging/StringsTrace.h
+++ b/Engine/Source/Runtime/Core/Public/ProfilingDebugging/StringsTrace.h
@@ -6,24 +6,12 @@
 #include "Trace/Trace.h"
 #include "Trace/Trace.inl"
 
-CORE_API UE_TRACE_EVENT_BEGIN_EXTERN(Strings, StaticString, Definition64bit)
+CORE_API UE_TRACE_EVENT_BEGIN_EXTERN(Strings, StaticString, NoSync|Definition64bit)
 	UE_TRACE_EVENT_FIELD(UE::Trace::WideString, DisplayWide)
 	UE_TRACE_EVENT_FIELD(UE::Trace::AnsiString, DisplayAnsi)
 UE_TRACE_EVENT_END()
 
-CORE_API UE_TRACE_EVENT_BEGIN_EXTERN(Strings, FName, Definition32bit)
-	UE_TRACE_EVENT_FIELD(UE::Trace::WideString, DisplayWide)
-	UE_TRACE_EVENT_FIELD(UE::Trace::AnsiString, DisplayAnsi)
-UE_TRACE_EVENT_END()
-
-// The following events are used when on connection callback is triggered. Do not use or reference
-
-UE_TRACE_EVENT_BEGIN_EXTERN(Strings, StaticStringNoSync, Definition64bit | NoSync)
-	UE_TRACE_EVENT_FIELD(UE::Trace::WideString, DisplayWide)
-	UE_TRACE_EVENT_FIELD(UE::Trace::AnsiString, DisplayAnsi)
-UE_TRACE_EVENT_END()
-
-UE_TRACE_EVENT_BEGIN_EXTERN(Strings, FNameNoSync, Definition32bit | NoSync)
+CORE_API UE_TRACE_EVENT_BEGIN_EXTERN(Strings, FName, NoSync|Definition32bit)
 	UE_TRACE_EVENT_FIELD(UE::Trace::WideString, DisplayWide)
 	UE_TRACE_EVENT_FIELD(UE::Trace::AnsiString, DisplayAnsi)
 UE_TRACE_EVENT_END()
