TDepletableMpscQueue: fixed memory order bug reported by TSan

diff --git a/Engine/Source/Runtime/Core/Public/Containers/DepletableMpscQueue.h b/Engine/Source/Runtime/Core/Public/Containers/DepletableMpscQueue.h
--- a/Engine/Source/Runtime/Core/Public/Containers/DepletableMpscQueue.h
+++ b/Engine/Source/Runtime/Core/Public/Containers/DepletableMpscQueue.h
@@ -65,7 +65,7 @@ namespace UE
 
 			FNode* Prev = Tail.exchange(New, std::memory_order_release); // `release` to make sure the new node is fully constructed before
 			// it becomes visible to the consumer
-			check(Prev->Next.load(std::memory_order_relaxed) == nullptr);
+			check(Prev->Next.load(std::memory_order_relaxed) == nullptr); // `Tail` is assigned before its Next
 
 			Prev->Next.store(New, std::memory_order_relaxed);
 
@@ -92,10 +92,10 @@ namespace UE
 			// reset the queue to the empty state. this redirects producers to start from `Sentinel` again.
 			// take note of the tail on resetting it because the list can be still not fully linked and so `Node.Next == nullptr` can't be 
 			// used to detect the end of the list
-			FNode* Last = Tail.exchange(&Sentinel, std::memory_order_release); // `release` to make `Sentinel.Next = nullptr` happen before
-			// modifying `Tail`
-			check(Last->Next.load(std::memory_order_relaxed) == nullptr);
-			// the previously queued items are detached from the instance (as a linked list, though potential not fully linked yet)
+			FNode* Last = Tail.exchange(&Sentinel, std::memory_order_acquire); // `acquire` to sync with producers' tail modifications,
+			// `Sentinel.Next = nullptr` must happen before modifying `Tail` but it's important only for the single consumer so doesn't need extra sync
+			check(Last->Next.load(std::memory_order_relaxed) == nullptr); // `Tail` is assigned before its Next
+			// the previously queued items are detached from the instance (as a linked list, though potentially not fully linked yet)
 
 			check(Last != &Sentinel); // can't be empty because of `First != nullptr` above
 
