add pool support for larger than default alignment (16) in MallocBinned3

diff --git a/Engine/Source/Runtime/Core/Private/HAL/MallocBinned3.cpp b/Engine/Source/Runtime/Core/Private/HAL/MallocBinned3.cpp
--- a/Engine/Source/Runtime/Core/Private/HAL/MallocBinned3.cpp
+++ b/Engine/Source/Runtime/Core/Private/HAL/MallocBinned3.cpp
@@ -956,9 +956,44 @@ void* FMallocBinned3::MallocExternal(SIZE_T Size, uint32 Alignment)
 {
 	static_assert(DEFAULT_ALIGNMENT <= BINNED3_MINIMUM_ALIGNMENT, "DEFAULT_ALIGNMENT is assumed to be zero"); // used below
 
-	// Only allocate from the small pools if the size is small enough and the alignment isn't crazy large.
-	// With large alignments, we'll waste a lot of memory allocating an entire page, but such alignments are highly unlikely in practice.
-	if ((Size <= BINNED3_MAX_SMALL_POOL_SIZE) & (Alignment <= BINNED3_MINIMUM_ALIGNMENT)) // one branch, not two
+	// Fast path: Allocate from the small pools if the size is small enough and the alignment <= binned3 min alignment.
+	//            Larger alignments can waste a lot of memory allocating an entire page, so some smaller alignments are 
+	//			  handled in the fallback path if less than a predefined max small pool alignment. 
+
+	bool UsePools = (Size <= BINNED3_MAX_SMALL_POOL_SIZE) && (Alignment <= BINNED3_MINIMUM_ALIGNMENT);
+	
+	if (!UsePools)
+	{
+		// fallback: test for non-default/mininum alignments and handle a subset of them to reduce memory waste of 64KB page requirements
+		// 			 e.g. code that wants to use aligned avx loads will use 32 byte alignments 
+
+		Alignment			= FMath::Max<uint32>(Alignment, BINNED3_MINIMUM_ALIGNMENT);
+		size_t AlignedSize	= Align(Size, Alignment);
+
+		if ((AlignedSize <= BINNED3_MAX_SMALL_POOL_SIZE) && (Alignment <= BINNED3_MAX_SMALL_POOL_ALIGNMENT))
+		{
+			// Start at the naturally aligned size as that is guaranteed to be aligned for our given alignment in the allocator
+			// Note: We must handle the case where the naturally aligned size is not a valid pool entry and the next higher size
+			//       may not match the same alignment. In that case we do a search upwards as they are likely nearby aligned sizes 
+			//		 for our small range of alignemnts allowed
+			uint32 PoolIndex = BoundSizeToPoolIndex(AlignedSize);
+			do
+			{
+				uint32 BlockSize = PoolIndexToBlockSize(PoolIndex);
+				if (IsAligned(BlockSize, Alignment))
+				{
+					// we found a matching pool for our alignment and size requirements, so modify the size request to match
+					Size = SIZE_T(BlockSize);
+					UsePools = true;
+					break;
+				}
+
+				PoolIndex++;
+			} while (PoolIndex < BINNED3_SMALL_POOL_COUNT);
+		}
+	}
+
+	if (UsePools) 
 	{
 		uint32 PoolIndex = BoundSizeToPoolIndex(Size);
 		FPerThreadFreeBlockLists* Lists = GMallocBinned3PerThreadCaches ? FPerThreadFreeBlockLists::Get() : nullptr;
diff --git a/Engine/Source/Runtime/Core/Public/HAL/MallocBinned3.h b/Engine/Source/Runtime/Core/Public/HAL/MallocBinned3.h
--- a/Engine/Source/Runtime/Core/Public/HAL/MallocBinned3.h
+++ b/Engine/Source/Runtime/Core/Public/HAL/MallocBinned3.h
@@ -27,6 +27,12 @@
 #define BINNED3_BASE_PAGE_SIZE				4096			// Minimum "page size" for binned3
 #define BINNED3_MINIMUM_ALIGNMENT_SHIFT		4				// Alignment of blocks, expressed as a shift
 #define BINNED3_MINIMUM_ALIGNMENT			16				// Alignment of blocks
+
+#ifndef BINNED3_MAX_SMALL_POOL_ALIGNMENT
+#define BINNED3_MAX_SMALL_POOL_ALIGNMENT	128
+#endif
+
+
 #ifndef BINNED3_MAX_SMALL_POOL_SIZE
 #if USE_CACHED_PAGE_ALLOCATOR_FOR_LARGE_ALLOCS
 #define BINNED3_MAX_SMALL_POOL_SIZE			(BINNEDCOMMON_MAX_LISTED_SMALL_POOL_SIZE)	// Maximum medium block size
