Tasks: switched FTaskEvent and pure FGraphEvent to lock-free TLS cache allocator to reduce mem footprint as currently used LinearConcurrentAllocator is vulnerable to long-living allocs

diff --git a/Engine/Source/Runtime/Core/Private/Async/TaskGraph.cpp b/Engine/Source/Runtime/Core/Private/Async/TaskGraph.cpp
--- a/Engine/Source/Runtime/Core/Private/Async/TaskGraph.cpp
+++ b/Engine/Source/Runtime/Core/Private/Async/TaskGraph.cpp
@@ -24,7 +24,6 @@
 #include "HAL/IConsoleManager.h"
 #include "Misc/App.h"
 #include "Misc/Fork.h"
-#include "Containers/LockFreeFixedSizeAllocator.h"
 #include "Async/TaskGraphInterfaces.h"
 #include "HAL/LowLevelMemTracker.h"
 #include "HAL/ThreadHeartBeat.h"
@@ -2456,10 +2455,13 @@ bool FTaskGraphInterface::IsMultithread()
 	return FPlatformProcess::SupportsMultithreading() || (FForkProcessHelper::IsForkedMultithreadInstance() && GAllowTaskGraphForkMultithreading);
 }
 
+#if TASKGRAPH_NEW_FRONTEND
 
-// Statics and some implementations from FBaseGraphTask and FGraphEvent
+FGraphEventImplAllocator GraphEventImplAllocator;
 
-#if !TASKGRAPH_NEW_FRONTEND
+#else
+
+// Statics and some implementations from FBaseGraphTask and FGraphEvent
 
 #if !(UE_BUILD_SHIPPING || UE_BUILD_TEST)
 void FBaseGraphTask::LogPossiblyInvalidSubsequentsTask(const TCHAR* TaskName)
diff --git a/Engine/Source/Runtime/Core/Private/Tasks/TaskPrivate.cpp b/Engine/Source/Runtime/Core/Private/Tasks/TaskPrivate.cpp
--- a/Engine/Source/Runtime/Core/Private/Tasks/TaskPrivate.cpp
+++ b/Engine/Source/Runtime/Core/Private/Tasks/TaskPrivate.cpp
@@ -11,6 +11,8 @@ namespace UE::Tasks
 {
 	namespace Private
 	{
+		FTaskEventBaseAllocator TaskEventBaseAllocator;
+
 		void FTaskBase::Schedule()
 		{
 			TaskTrace::Scheduled(GetTraceId());
diff --git a/Engine/Source/Runtime/Core/Public/Async/TaskGraphInterfaces.h b/Engine/Source/Runtime/Core/Public/Async/TaskGraphInterfaces.h
--- a/Engine/Source/Runtime/Core/Public/Async/TaskGraphInterfaces.h
+++ b/Engine/Source/Runtime/Core/Public/Async/TaskGraphInterfaces.h
@@ -734,7 +734,7 @@ private:
 };
 
 // an adaptation of FBaseGraphTask to be used as a standalone FGraphEvent
-class FGraphEventImpl : public TConcurrentLinearObject<FGraphEventImpl, FTaskGraphBlockAllocationTag>, public FBaseGraphTask
+class FGraphEventImpl : public FBaseGraphTask
 {
 public:
 	FGraphEventImpl()
@@ -743,6 +743,9 @@ public:
 		Init(TEXT("GraphEvent"), UE::Tasks::ETaskPriority::Normal, UE::Tasks::EExtendedTaskPriority::TaskEvent);
 	}
 
+	static void* operator new(size_t Size);
+	static void operator delete(void* Ptr);
+
 private:
 	virtual bool TryExecuteTask() override
 	{
@@ -751,6 +754,19 @@ private:
 	}
 };
 
+using FGraphEventImplAllocator = TLockFreeFixedSizeAllocator_TLSCache<sizeof(FGraphEventImpl), PLATFORM_CACHE_LINE_SIZE>;
+CORE_API extern FGraphEventImplAllocator GraphEventImplAllocator;
+
+inline void* FGraphEventImpl::operator new(size_t Size)
+{
+	return GraphEventImplAllocator.Allocate();
+}
+
+inline void FGraphEventImpl::operator delete(void* Ptr)
+{
+	GraphEventImplAllocator.Free(Ptr);
+}
+
 inline FGraphEventRef FBaseGraphTask::CreateGraphEvent()
 {
 	FGraphEventImpl* GraphEvent = new FGraphEventImpl;
diff --git a/Engine/Source/Runtime/Core/Public/Tasks/TaskPrivate.h b/Engine/Source/Runtime/Core/Public/Tasks/TaskPrivate.h
--- a/Engine/Source/Runtime/Core/Public/Tasks/TaskPrivate.h
+++ b/Engine/Source/Runtime/Core/Public/Tasks/TaskPrivate.h
@@ -847,7 +847,7 @@ namespace UE::Tasks
 
 		// a special kind of task that is used for signalling or dependency management. It can have prerequisites or be used as a prerequisite for other tasks. 
 		// It's optimized for the fact that it doesn't have a task body and so doesn't need to be scheduled and executed
-		class FTaskEventBase : public TConcurrentLinearObject<FTaskEventBase, FTaskBlockAllocationTag>, public FTaskBase
+		class FTaskEventBase : public FTaskBase
 		{
 		public:
 			static FTaskEventBase* Create(const TCHAR* DebugName)
@@ -855,6 +855,9 @@ namespace UE::Tasks
 				return new FTaskEventBase(DebugName);
 			}
 
+			static void* operator new(size_t Size);
+			static void operator delete(void* Ptr);
+
 		private:
 			FTaskEventBase(const TCHAR* InDebugName)
 				: FTaskBase(/*InitRefCount=*/ 1) // for the initial reference (we don't increment it on passing to `TRefCountPtr`)
@@ -920,6 +923,20 @@ namespace UE::Tasks
 			}
 		}
 
+		using FTaskEventBaseAllocator = TLockFreeFixedSizeAllocator_TLSCache<sizeof(FTaskEventBase), PLATFORM_CACHE_LINE_SIZE>;
+		CORE_API extern FTaskEventBaseAllocator TaskEventBaseAllocator;
+
+		inline void* FTaskEventBase::operator new(size_t Size)
+		{
+			return TaskEventBaseAllocator.Allocate();
+		}
+
+		inline void FTaskEventBase::operator delete(void* Ptr)
+		{
+			TaskEventBaseAllocator.Free(Ptr);
+		}
+
+
 		inline bool FTaskBase::Wait(FTimespan InTimeout)
 		{
 			TaskTrace::FWaitingScope WaitingScope(GetTraceId());
