Modular builds: Signing: Added missing signing and encryption initialization to the main module definition macro.

It resolves encrypted archives not being handled correctly when using modular builds. It was a problem e.g. for ShooterGame. Without the keys embedded, the game would fail to mount encrypted/signed archives.

#rnx

diff --git a/Engine/Source/Runtime/Core/Private/Misc/CoreDelegates.cpp b/Engine/Source/Runtime/Core/Private/Misc/CoreDelegates.cpp
--- a/Engine/Source/Runtime/Core/Private/Misc/CoreDelegates.cpp
+++ b/Engine/Source/Runtime/Core/Private/Misc/CoreDelegates.cpp
@@ -235,7 +235,7 @@ FCoreDelegates::FGetOnScreenMessagesDelegate FCoreDelegates::OnGetOnScreenMessag
 typedef void(*TSigningKeyFunc)(TArray<uint8>&, TArray<uint8>&);
 typedef void(*TEncryptionKeyFunc)(unsigned char[32]);
 
-void RegisterSigningKeyCallback(TSigningKeyFunc InCallback)
+CORE_API void RegisterSigningKeyCallback(TSigningKeyFunc InCallback)
 {
 	FCoreDelegates::GetPakSigningKeysDelegate().BindLambda([InCallback](TArray<uint8>& OutExponent, TArray<uint8>& OutModulus)
 	{
@@ -243,7 +243,7 @@ void RegisterSigningKeyCallback(TSigningKeyFunc InCallback)
 	});
 }
 
-void RegisterEncryptionKeyCallback(TEncryptionKeyFunc InCallback)
+CORE_API void RegisterEncryptionKeyCallback(TEncryptionKeyFunc InCallback)
 {
 	FCoreDelegates::GetPakEncryptionKeyDelegate().BindLambda([InCallback](uint8 OutKey[32])
 	{
diff --git a/Engine/Source/Runtime/Core/Public/Modules/ModuleManager.h b/Engine/Source/Runtime/Core/Public/Modules/ModuleManager.h
--- a/Engine/Source/Runtime/Core/Public/Modules/ModuleManager.h
+++ b/Engine/Source/Runtime/Core/Public/Modules/ModuleManager.h
@@ -828,7 +828,7 @@ class FDefaultGameModuleImpl
 	{ \
 		FSigningKeyRegistration() \
 		{ \
-			extern void RegisterSigningKeyCallback(void (*)(TArray<uint8>&, TArray<uint8>&)); \
+			extern CORE_API void RegisterSigningKeyCallback(void (*)(TArray<uint8>&, TArray<uint8>&)); \
 			RegisterSigningKeyCallback(&Callback); \
 		} \
 		static void Callback(TArray<uint8>& OutExponent, TArray<uint8>& OutModulus) \
@@ -856,7 +856,7 @@ class FDefaultGameModuleImpl
 	{ \
 		FEncryptionKeyRegistration() \
 		{ \
-			extern void RegisterEncryptionKeyCallback(void (*)(unsigned char OutKey[32])); \
+			extern CORE_API void RegisterEncryptionKeyCallback(void (*)(unsigned char OutKey[32])); \
 			RegisterEncryptionKeyCallback(&Callback); \
 		} \
 		static void Callback(unsigned char OutKey[32]) \
@@ -955,6 +955,8 @@ class FDefaultGameModuleImpl
 
 	#define IMPLEMENT_PRIMARY_GAME_MODULE( ModuleImplClass, ModuleName, GameName ) \
 		/* Nothing special to do for modular builds.  The game name will be set via the command-line */ \
+		IMPLEMENT_SIGNING_KEY_REGISTRATION() \
+		IMPLEMENT_ENCRYPTION_KEY_REGISTRATION() \
 		IMPLEMENT_TARGET_NAME_REGISTRATION() \
 		IMPLEMENT_GAME_MODULE( ModuleImplClass, ModuleName )
 #endif	//IS_MONOLITHIC
