Avoid resolving the display string in FTextStringHelper::WriteToBuffer when writing complex text

#rnx

diff --git a/Engine/Source/Runtime/Core/Private/Internationalization/Text.cpp b/Engine/Source/Runtime/Core/Private/Internationalization/Text.cpp
--- a/Engine/Source/Runtime/Core/Private/Internationalization/Text.cpp
+++ b/Engine/Source/Runtime/Core/Private/Internationalization/Text.cpp
@@ -1760,15 +1760,13 @@ bool FTextStringHelper::ReadFromString(const TCHAR* Buffer, FText& OutValue, con
 
 void FTextStringHelper::WriteToBuffer(FString& Buffer, const FText& Value, const bool bRequiresQuotes, const bool bStripPackageNamespace)
 {
-	const FString& StringValue = FTextInspector::GetDisplayString(Value);
-
 	// Culture invariant text?
 	if (Value.IsCultureInvariant())
 	{
 #define LOC_DEFINE_REGION
 		// Produces INVTEXT("...")
 		Buffer += TEXT("INVTEXT(\"");
-		Buffer += StringValue.ReplaceCharWithEscapedChar();
+		Buffer += Value.ToString().ReplaceCharWithEscapedChar();
 		Buffer += TEXT("\")");
 #undef LOC_DEFINE_REGION
 	}
@@ -1780,12 +1778,12 @@ void FTextStringHelper::WriteToBuffer(FString& Buffer, const FText& Value, const
 	else if (bRequiresQuotes)
 	{
 		Buffer += TEXT("\"");
-		Buffer += StringValue.ReplaceCharWithEscapedChar();
+		Buffer += Value.ToString().ReplaceCharWithEscapedChar();
 		Buffer += TEXT("\"");
 	}
 	else
 	{
-		Buffer += StringValue;
+		Buffer += Value.ToString();
 	}
 }
 
