PR #9123: Support structured bindings for UE_LOG, checkf, ensureMsgf. (Contributed by [REDACTED_PERSON])

diff --git a/Engine/Source/Runtime/Core/Public/Misc/AssertionMacros.h b/Engine/Source/Runtime/Core/Public/Misc/AssertionMacros.h
--- a/Engine/Source/Runtime/Core/Public/Misc/AssertionMacros.h
+++ b/Engine/Source/Runtime/Core/Public/Misc/AssertionMacros.h
@@ -234,10 +234,10 @@ RetType FORCENOINLINE UE_DEBUG_SECTION DispatchCheckVerify(InnerType&& Inner, Ar
 		{ \
 			if(UNLIKELY(!(expr))) \
 			{ \
-				DispatchCheckVerify([&] () UE_DEBUG_SECTION \
+				DispatchCheckVerify([] (const auto& LFormat, const auto&... UE_LOG_Args) UE_DEBUG_SECTION \
 				{ \
-					FDebug::CheckVerifyFailedImpl(#expr, __FILE__, __LINE__, PLATFORM_RETURN_ADDRESS(), format, ##__VA_ARGS__); \
-				}); \
+					FDebug::CheckVerifyFailedImpl(#expr, __FILE__, __LINE__, PLATFORM_RETURN_ADDRESS(), LFormat, UE_LOG_Args...); \
+				}, format, ##__VA_ARGS__); \
 				PLATFORM_BREAK_IF_DESIRED(); \
 				CA_ASSUME(false); \
 			} \
@@ -334,14 +334,14 @@ RetType FORCENOINLINE UE_DEBUG_SECTION DispatchCheckVerify(InnerType&& Inner, Ar
 
 #if DO_ENSURE && !USING_CODE_ANALYSIS // The Visual Studio 2013 analyzer doesn't understand these complex conditionals
 
-	#define UE_ENSURE_IMPL(Capture, Always, InExpression, ...) \
-		(LIKELY(!!(InExpression)) || (DispatchCheckVerify<bool>([Capture] () FORCENOINLINE UE_DEBUG_SECTION \
+	#define UE_ENSURE_IMPL(Always, InExpression, ...) \
+		(LIKELY(!!(InExpression)) || (DispatchCheckVerify<bool>([] (const auto&... UE_LOG_Args) FORCENOINLINE UE_DEBUG_SECTION \
 		{ \
 			static bool bExecuted = false; \
 			if ((!bExecuted || Always) && FPlatformMisc::IsEnsureAllowed()) \
 			{ \
 				bExecuted = true; \
-				FDebug::OptionallyLogFormattedEnsureMessageReturningFalse(true, #InExpression, __FILE__, __LINE__, PLATFORM_RETURN_ADDRESS(), ##__VA_ARGS__); \
+				FDebug::OptionallyLogFormattedEnsureMessageReturningFalse(true, #InExpression, __FILE__, __LINE__, PLATFORM_RETURN_ADDRESS(), UE_LOG_Args...); \
 				if (!FPlatformMisc::IsDebuggerPresent()) \
 				{ \
 					FPlatformMisc::PromptForRemoteDebugging(true); \
@@ -350,12 +350,12 @@ RetType FORCENOINLINE UE_DEBUG_SECTION DispatchCheckVerify(InnerType&& Inner, Ar
 				return true; \
 			} \
 			return false; \
-		}) && ([] () { PLATFORM_BREAK_IF_DESIRED(); } (), false)))
+		}, ##__VA_ARGS__) && ([] () { PLATFORM_BREAK_IF_DESIRED(); } (), false)))
 
-	#define ensure(           InExpression                ) UE_ENSURE_IMPL( , false, InExpression, TEXT(""))
-	#define ensureMsgf(       InExpression, InFormat, ... ) UE_ENSURE_IMPL(&, false, InExpression, InFormat, ##__VA_ARGS__)
-	#define ensureAlways(     InExpression                ) UE_ENSURE_IMPL( , true,  InExpression, TEXT(""))
-	#define ensureAlwaysMsgf( InExpression, InFormat, ... ) UE_ENSURE_IMPL(&, true,  InExpression, InFormat, ##__VA_ARGS__)
+	#define ensure(           InExpression                ) UE_ENSURE_IMPL(false, InExpression, TEXT(""))
+	#define ensureMsgf(       InExpression, InFormat, ... ) UE_ENSURE_IMPL(false, InExpression, InFormat, ##__VA_ARGS__)
+	#define ensureAlways(     InExpression                ) UE_ENSURE_IMPL(true,  InExpression, TEXT(""))
+	#define ensureAlwaysMsgf( InExpression, InFormat, ... ) UE_ENSURE_IMPL(true,  InExpression, InFormat, ##__VA_ARGS__)
 
 #else	// DO_ENSURE
 
@@ -431,12 +431,12 @@ CORE_API void VARARGS LowLevelFatalErrorHandler(const ANSICHAR* File, int32 Line
 #define LowLevelFatalError(Format, ...) \
 	{ \
 		static_assert(TIsArrayOrRefOfType<decltype(Format), TCHAR>::Value, "Formatting string must be a TCHAR array."); \
-		DispatchCheckVerify([&] () FORCENOINLINE UE_DEBUG_SECTION \
+		DispatchCheckVerify([] (const auto& LFormat, const auto&... UE_LOG_Args) FORCENOINLINE UE_DEBUG_SECTION \
 		{ \
 			void* ProgramCounter = PLATFORM_RETURN_ADDRESS(); \
-			LowLevelFatalErrorHandler(__FILE__, __LINE__, ProgramCounter, Format, ##__VA_ARGS__); \
+			LowLevelFatalErrorHandler(__FILE__, __LINE__, ProgramCounter, LFormat, UE_LOG_Args...); \
 			UE_DEBUG_BREAK_AND_PROMPT_FOR_REMOTE(); \
 			FDebug::ProcessFatalError(ProgramCounter); \
-		}); \
+		}, Format, ##__VA_ARGS__); \
 	}
 
