Fix for event names in UTF-8 mode.

diff --git a/Engine/Source/Runtime/Core/Public/HAL/PlatformMisc.h b/Engine/Source/Runtime/Core/Public/HAL/PlatformMisc.h
--- a/Engine/Source/Runtime/Core/Public/HAL/PlatformMisc.h
+++ b/Engine/Source/Runtime/Core/Public/HAL/PlatformMisc.h
@@ -113,7 +113,7 @@ public:
 #if PLATFORM_USES_ANSI_STRING_FOR_EXTERNAL_PROFILING
 #define NAMED_EVENT_STR(x) x
 #else
-#define NAMED_EVENT_STR(x) L##x
+#define NAMED_EVENT_STR(x) TEXT(x)
 #endif
 
 #define SCOPED_NAMED_EVENT(Name, Color)\
