Logging: Fixed the use of FORCENOINLINE and UE_DEBUG_SECTION for logging and assertions

Since DispatchCheckVerify is already FORCENOINLINE, the use of it on the lambdas causes DispatchCheckVerify to be act as an additional and unnecessary wrapper.

diff --git a/Engine/Source/Runtime/Core/Public/Logging/LogMacros.h b/Engine/Source/Runtime/Core/Public/Logging/LogMacros.h
--- a/Engine/Source/Runtime/Core/Public/Logging/LogMacros.h
+++ b/Engine/Source/Runtime/Core/Public/Logging/LogMacros.h
@@ -193,7 +193,7 @@ private:
 		UE_LOG_EXPAND_IS_FATAL(Verbosity, PREPROCESSOR_NOTHING, if (!CategoryName.IsSuppressed(ELogVerbosity::Verbosity))) \
 			{ \
 				UE_VALIDATE_FORMAT_STRING(Format, ##__VA_ARGS__); \
-				DispatchCheckVerify([] (const auto& LCategoryName, const auto& LFormat, const auto&... UE_LOG_Args) FORCENOINLINE \
+				DispatchCheckVerify([] (const auto& LCategoryName, const auto& LFormat, const auto&... UE_LOG_Args) UE_DEBUG_SECTION \
 				{ \
 					TRACE_LOG_MESSAGE(LCategoryName, Verbosity, LFormat, UE_LOG_Args...) \
 					UE_LOG_EXPAND_IS_FATAL(Verbosity, \
@@ -297,7 +297,7 @@ private:
 			{ \
 				if (Condition) \
 				{ \
-					DispatchCheckVerify([] (const auto& LCategoryName, const auto& LFormat, const auto&... UE_LOG_Args) FORCENOINLINE \
+					DispatchCheckVerify([] (const auto& LCategoryName, const auto& LFormat, const auto&... UE_LOG_Args) UE_DEBUG_SECTION \
 					{ \
 						TRACE_LOG_MESSAGE(LCategoryName, Verbosity, LFormat, UE_LOG_Args...) \
 						UE_LOG_EXPAND_IS_FATAL(Verbosity, \
diff --git a/Engine/Source/Runtime/Core/Public/Misc/AssertionMacros.h b/Engine/Source/Runtime/Core/Public/Misc/AssertionMacros.h
--- a/Engine/Source/Runtime/Core/Public/Misc/AssertionMacros.h
+++ b/Engine/Source/Runtime/Core/Public/Misc/AssertionMacros.h
@@ -335,7 +335,7 @@ RetType FORCENOINLINE UE_DEBUG_SECTION DispatchCheckVerify(InnerType&& Inner, Ar
 #if DO_ENSURE && !USING_CODE_ANALYSIS // The Visual Studio 2013 analyzer doesn't understand these complex conditionals
 
 	#define UE_ENSURE_IMPL(Always, InExpression, ...) \
-		(LIKELY(!!(InExpression)) || (DispatchCheckVerify<bool>([] (const auto&... UE_LOG_Args) FORCENOINLINE UE_DEBUG_SECTION \
+		(LIKELY(!!(InExpression)) || (DispatchCheckVerify<bool>([] (const auto&... UE_LOG_Args) UE_DEBUG_SECTION \
 		{ \
 			static bool bExecuted = false; \
 			if ((!bExecuted || Always) && FPlatformMisc::IsEnsureAllowed()) \
@@ -431,7 +431,7 @@ CORE_API void VARARGS LowLevelFatalErrorHandler(const ANSICHAR* File, int32 Line
 #define LowLevelFatalError(Format, ...) \
 	{ \
 		static_assert(TIsArrayOrRefOfType<decltype(Format), TCHAR>::Value, "Formatting string must be a TCHAR array."); \
-		DispatchCheckVerify([] (const auto& LFormat, const auto&... UE_LOG_Args) FORCENOINLINE UE_DEBUG_SECTION \
+		DispatchCheckVerify([] (const auto& LFormat, const auto&... UE_LOG_Args) UE_DEBUG_SECTION \
 		{ \
 			void* ProgramCounter = PLATFORM_RETURN_ADDRESS(); \
 			LowLevelFatalErrorHandler(__FILE__, __LINE__, ProgramCounter, LFormat, UE_LOG_Args...); \
