Add trace.bookmark console command which simply emits a TRACE_BOOKMARK.

- Argument is a single string which becomes the bookmark name.

diff --git a/Engine/Source/Runtime/Core/Private/ProfilingDebugging/TraceAuxiliary.cpp b/Engine/Source/Runtime/Core/Private/ProfilingDebugging/TraceAuxiliary.cpp
--- a/Engine/Source/Runtime/Core/Private/ProfilingDebugging/TraceAuxiliary.cpp
+++ b/Engine/Source/Runtime/Core/Private/ProfilingDebugging/TraceAuxiliary.cpp
@@ -848,6 +848,12 @@ static void TraceAuxiliarySnapshotFile(const TArray<FString>& Args)
 	}
 }
 
+////////////////////////////////////////////////////////////////////////////////
+static void TraceBookmark(const TArray<FString>& Args)
+{
+	TRACE_BOOKMARK(TEXT("%s"), Args.Num() ? *Args[0] : TEXT(""));
+}
+
 ////////////////////////////////////////////////////////////////////////////////
 static FAutoConsoleCommand TraceAuxiliarySendCmd(
 	TEXT("Trace.Send"),
@@ -932,6 +938,14 @@ static FAutoConsoleCommand TraceAuxiliarySnapshotFileCmd(
 	FConsoleCommandWithArgsDelegate::CreateStatic(TraceAuxiliarySnapshotFile)
 );
 
+////////////////////////////////////////////////////////////////////////////////
+static FAutoConsoleCommand TraceBookmarkCmd(
+	TEXT("Trace.Bookmark"),
+	TEXT("[Name] - Emits a TRACE_BOOKMARK() event with the given string name."
+	),
+	FConsoleCommandWithArgsDelegate::CreateStatic(TraceBookmark)
+);
+
 #endif // UE_TRACE_ENABLED
 
 
