Logging: Made GPrintLogX globals atomic to fix warnings from TSan

#rnx

diff --git a/Engine/Source/Runtime/Core/Private/Misc/CoreGlobals.cpp b/Engine/Source/Runtime/Core/Private/Misc/CoreGlobals.cpp
--- a/Engine/Source/Runtime/Core/Private/Misc/CoreGlobals.cpp
+++ b/Engine/Source/Runtime/Core/Private/Misc/CoreGlobals.cpp
@@ -387,11 +387,11 @@ bool					GShouldSuspendRenderingThread	= false;
 /** Determines what kind of trace should occur, NAME_None for none.											*/
 FLazyName				GCurrentTraceName;
 /** How to print the time in log output																		*/
-ELogTimes::Type			GPrintLogTimes					= ELogTimes::None;
-/** How to print the category in log output. */
-bool					GPrintLogCategory = true;
-/** How to print the verbosity in log output. */
-bool					GPrintLogVerbosity = true;
+TSAN_ATOMIC(ELogTimes::Type)	GPrintLogTimes			= ELogTimes::None;
+/** How to print the category in log output.																*/
+TSAN_ATOMIC(bool)		GPrintLogCategory				= true;
+/** How to print the verbosity in log output.																*/
+TSAN_ATOMIC(bool)		GPrintLogVerbosity				= true;
 
 #if USE_HITCH_DETECTION
 TSAN_ATOMIC(bool)				GHitchDetected(false);
diff --git a/Engine/Source/Runtime/Core/Public/CoreGlobals.h b/Engine/Source/Runtime/Core/Public/CoreGlobals.h
--- a/Engine/Source/Runtime/Core/Public/CoreGlobals.h
+++ b/Engine/Source/Runtime/Core/Public/CoreGlobals.h
@@ -496,13 +496,13 @@ extern CORE_API bool GShouldSuspendRenderingThread;
 extern CORE_API FLazyName GCurrentTraceName;
 
 /** How to print the time in log output. */
-extern CORE_API ELogTimes::Type GPrintLogTimes;
+extern CORE_API TSAN_ATOMIC(ELogTimes::Type) GPrintLogTimes;
 
 /** How to print the category in log output. */
-extern CORE_API bool GPrintLogCategory;
+extern CORE_API TSAN_ATOMIC(bool) GPrintLogCategory;
 
 /** How to print the verbosity in log output. */
-extern CORE_API bool GPrintLogVerbosity;
+extern CORE_API TSAN_ATOMIC(bool) GPrintLogVerbosity;
 
 #if USE_HITCH_DETECTION
 /** Used by the lightweight stats and FGameThreadHitchHeartBeat to print a stat stack for hitches in shipping builds. */
