Reverted support for structured bindings in ensureMsgf

Compiling for Windows with Clang fails with "error: definition of dllimport data" with the change.

#rnx

diff --git a/Engine/Source/Runtime/Core/Public/Misc/AssertionMacros.h b/Engine/Source/Runtime/Core/Public/Misc/AssertionMacros.h
--- a/Engine/Source/Runtime/Core/Public/Misc/AssertionMacros.h
+++ b/Engine/Source/Runtime/Core/Public/Misc/AssertionMacros.h
@@ -334,14 +334,14 @@ RetType FORCENOINLINE UE_DEBUG_SECTION DispatchCheckVerify(InnerType&& Inner, Ar
 
 #if DO_ENSURE && !USING_CODE_ANALYSIS // The Visual Studio 2013 analyzer doesn't understand these complex conditionals
 
-	#define UE_ENSURE_IMPL(Always, InExpression, ...) \
-		(LIKELY(!!(InExpression)) || (DispatchCheckVerify<bool>([] (const auto&... UE_LOG_Args) UE_DEBUG_SECTION \
+	#define UE_ENSURE_IMPL(Capture, Always, InExpression, ...) \
+		(LIKELY(!!(InExpression)) || (DispatchCheckVerify<bool>([Capture] () UE_DEBUG_SECTION \
 		{ \
 			static bool bExecuted = false; \
 			if ((!bExecuted || Always) && FPlatformMisc::IsEnsureAllowed()) \
 			{ \
 				bExecuted = true; \
-				FDebug::OptionallyLogFormattedEnsureMessageReturningFalse(true, #InExpression, __FILE__, __LINE__, PLATFORM_RETURN_ADDRESS(), UE_LOG_Args...); \
+				FDebug::OptionallyLogFormattedEnsureMessageReturningFalse(true, #InExpression, __FILE__, __LINE__, PLATFORM_RETURN_ADDRESS(), ##__VA_ARGS__); \
 				if (!FPlatformMisc::IsDebuggerPresent()) \
 				{ \
 					FPlatformMisc::PromptForRemoteDebugging(true); \
@@ -350,12 +350,12 @@ RetType FORCENOINLINE UE_DEBUG_SECTION DispatchCheckVerify(InnerType&& Inner, Ar
 				return true; \
 			} \
 			return false; \
-		}, ##__VA_ARGS__) && ([] () { PLATFORM_BREAK_IF_DESIRED(); } (), false)))
+		}) && ([] () { PLATFORM_BREAK_IF_DESIRED(); } (), false)))
 
-	#define ensure(           InExpression                ) UE_ENSURE_IMPL(false, InExpression, TEXT(""))
-	#define ensureMsgf(       InExpression, InFormat, ... ) UE_ENSURE_IMPL(false, InExpression, InFormat, ##__VA_ARGS__)
-	#define ensureAlways(     InExpression                ) UE_ENSURE_IMPL(true,  InExpression, TEXT(""))
-	#define ensureAlwaysMsgf( InExpression, InFormat, ... ) UE_ENSURE_IMPL(true,  InExpression, InFormat, ##__VA_ARGS__)
+	#define ensure(           InExpression                ) UE_ENSURE_IMPL( , false, InExpression, TEXT(""))
+	#define ensureMsgf(       InExpression, InFormat, ... ) UE_ENSURE_IMPL(&, false, InExpression, InFormat, ##__VA_ARGS__)
+	#define ensureAlways(     InExpression                ) UE_ENSURE_IMPL( , true,  InExpression, TEXT(""))
+	#define ensureAlwaysMsgf( InExpression, InFormat, ... ) UE_ENSURE_IMPL(&, true,  InExpression, InFormat, ##__VA_ARGS__)
 
 #else	// DO_ENSURE
 
