Modular builds: TThreadSingleton: Added 'noinline' attribute to GetTls to work around problems on some platforms with local statics not being handled correctly in inlined members of exported classes.

In this case, we sometimes export classes derived from this template. Without this workaround, we ended up with each module having its own static.

#rnx

diff --git a/Engine/Source/Runtime/Core/Public/HAL/ThreadSingleton.h b/Engine/Source/Runtime/Core/Public/HAL/ThreadSingleton.h
--- a/Engine/Source/Runtime/Core/Public/HAL/ThreadSingleton.h
+++ b/Engine/Source/Runtime/Core/Public/HAL/ThreadSingleton.h
@@ -53,6 +53,9 @@ class TThreadSingleton : public FTlsAutoCleanup
 	/**
 	 * @return TLS slot that holds a TThreadSingleton.
 	 */
+#if PLATFORM_CONSOLE_DYNAMIC_LINK
+	FORCENOINLINE
+#endif
 	static uint32& GetTlsSlot()
 	{
 		static uint32 TlsSlot = 0xFFFFFFFF;
