fixed TSAN false positive data race warnings about TReferenceControllerBase. TSAN doesn't handle memory fences properly (https://stackoverflow.com/questions/70542993/why-does-the-thread-sanitizer-complain-about-acquire-release-thread-fences). As using release decrement and acquire memory fence is not proven to be superior for supported platforms, it was decided to switch to acq_rel decrement as it's what is used by most of other popular implementations for ref counting

diff --git a/Engine/Source/Runtime/Core/Public/Tasks/TaskPrivate.h b/Engine/Source/Runtime/Core/Public/Tasks/TaskPrivate.h
--- a/Engine/Source/Runtime/Core/Public/Tasks/TaskPrivate.h
+++ b/Engine/Source/Runtime/Core/Public/Tasks/TaskPrivate.h
@@ -95,10 +95,9 @@ namespace UE::Tasks
 
 			void Release()
 			{
-				uint32 LocalRefCount = RefCount.fetch_sub(1, std::memory_order_release) - 1;
+				uint32 LocalRefCount = RefCount.fetch_sub(1, std::memory_order_acq_rel) - 1;
 				if (LocalRefCount == 0)
 				{
-					std::atomic_thread_fence(std::memory_order_acquire);
 					delete this;
 				}
 			}
diff --git a/Engine/Source/Runtime/Core/Public/Templates/SharedPointerInternals.h b/Engine/Source/Runtime/Core/Public/Templates/SharedPointerInternals.h
--- a/Engine/Source/Runtime/Core/Public/Templates/SharedPointerInternals.h
+++ b/Engine/Source/Runtime/Core/Public/Templates/SharedPointerInternals.h
@@ -193,18 +193,14 @@ namespace SharedPointerInternals
 		{
 			if constexpr (Mode == ESPMode::ThreadSafe)
 			{
-				// std::memory_order_release is used here so that, if we do end up executing the destructor, it's not possible
+				// std::memory_order_acq_rel is used here so that, if we do end up executing the destructor, it's not possible
 				// for side effects from executing the destructor end up being visible before we've determined that the shared
 				// reference count is actually zero.
 
-				int32 OldSharedCount = SharedReferenceCount.fetch_sub(1, std::memory_order_release);
+				int32 OldSharedCount = SharedReferenceCount.fetch_sub(1, std::memory_order_acq_rel);
 				checkSlow(OldSharedCount > 0);
 				if (OldSharedCount == 1)
 				{
-					// Ensure that all other threads' accesses to the object are visible to this thread before we call the
-					// destructor.
-					std::atomic_thread_fence(std::memory_order_acquire);
-
 					// Last shared reference was released!  Destroy the referenced object.
 					DestroyObject();
 
@@ -255,14 +251,12 @@ namespace SharedPointerInternals
 		{
 			if constexpr (Mode == ESPMode::ThreadSafe)
 			{
-				// See ReleaseSharedReference for the same reasons that std::memory_order_release and std::memory_order_acquire are used in this function.
+				// See ReleaseSharedReference for the same reasons that std::memory_order_acq_rel is used in this function.
 
-				int32 OldWeakCount = WeakReferenceCount.fetch_sub(1, std::memory_order_release);
+				int32 OldWeakCount = WeakReferenceCount.fetch_sub(1, std::memory_order_acq_rel);
 				checkSlow(OldWeakCount > 0);
 				if (OldWeakCount == 1)
 				{
-					std::atomic_thread_fence(std::memory_order_acquire);
-
 					// Disable this if running clang's static analyzer. Passing shared pointers
 					// and references to functions it cannot reason about, produces false
 					// positives about use-after-free in the TSharedPtr/TSharedRef destructors.
