DiffOnly Callstacks: Fix truncation of callstack information due to StackIgnoreCount being incorrectly incremented for inline functions.
#rnx

diff --git a/Engine/Source/Runtime/Core/Private/Windows/WindowsPlatformStackWalk.cpp b/Engine/Source/Runtime/Core/Private/Windows/WindowsPlatformStackWalk.cpp
--- a/Engine/Source/Runtime/Core/Private/Windows/WindowsPlatformStackWalk.cpp
+++ b/Engine/Source/Runtime/Core/Private/Windows/WindowsPlatformStackWalk.cpp
@@ -242,11 +242,18 @@ void FWindowsPlatformStackWalk::StackWalkAndDump( ANSICHAR* HumanReadableString,
 	InitStackWalking();
 
 	// If the callstack is for the executing thread, ignore this function
-	if(Context == nullptr)
+	if (Context == nullptr)
 	{
 		IgnoreCount++;
 	}
 	FGenericPlatformStackWalk::StackWalkAndDump(HumanReadableString, HumanReadableStringSize, IgnoreCount, Context);
+	// If we incremented IgnoreCount, make sure we have instructions after StackWalkAndDump so the compiler
+	// can not remove this function from the callstack using Tail Call Elimination
+	if (Context == nullptr)
+	{
+		static volatile int32 ForceCompilerToReturnHere = 0;
+		ForceCompilerToReturnHere += static_cast<int32>(HumanReadableStringSize);
+	}
 }
 
 void FWindowsPlatformStackWalk::StackWalkAndDump( ANSICHAR* HumanReadableString, SIZE_T HumanReadableStringSize, void* ProgramCounter, void* Context )
diff --git a/Engine/Source/Runtime/Core/Public/GenericPlatform/GenericPlatformStackWalk.h b/Engine/Source/Runtime/Core/Public/GenericPlatform/GenericPlatformStackWalk.h
--- a/Engine/Source/Runtime/Core/Public/GenericPlatform/GenericPlatformStackWalk.h
+++ b/Engine/Source/Runtime/Core/Public/GenericPlatform/GenericPlatformStackWalk.h
@@ -233,7 +233,7 @@ struct CORE_API FGenericPlatformStackWalk
 	 * @param	IgnoreCount			Number of stack entries to ignore (some are guaranteed to be in the stack walking code)
 	 * @param	Context				Optional thread context information
 	 */ 
-	static void StackWalkAndDump( ANSICHAR* HumanReadableString, SIZE_T HumanReadableStringSize, int32 IgnoreCount, void* Context = nullptr );
+	FORCENOINLINE static void StackWalkAndDump( ANSICHAR* HumanReadableString, SIZE_T HumanReadableStringSize, int32 IgnoreCount, void* Context = nullptr );  // FORCENOINLINE so it can be counted during StackTrace
 
 	/**
 	 * Walks the stack and appends the human readable string to the passed in one.
diff --git a/Engine/Source/Runtime/Core/Public/Windows/WindowsPlatformStackWalk.h b/Engine/Source/Runtime/Core/Public/Windows/WindowsPlatformStackWalk.h
--- a/Engine/Source/Runtime/Core/Public/Windows/WindowsPlatformStackWalk.h
+++ b/Engine/Source/Runtime/Core/Public/Windows/WindowsPlatformStackWalk.h
@@ -18,9 +18,9 @@ struct CORE_API FWindowsPlatformStackWalk
 
 	static void ProgramCounterToSymbolInfo( uint64 ProgramCounter, FProgramCounterSymbolInfo& out_SymbolInfo );
 	static void ProgramCounterToSymbolInfoEx( uint64 ProgramCounter, FProgramCounterSymbolInfoEx& out_SymbolInfo );
-	static uint32 CaptureStackBackTrace( uint64* BackTrace, uint32 MaxDepth, void* Context = nullptr );
+	FORCENOINLINE static uint32 CaptureStackBackTrace( uint64* BackTrace, uint32 MaxDepth, void* Context = nullptr );  // FORCENOINLINE so it can be counted during StackTrace
 	static uint32 CaptureThreadStackBackTrace( uint64 ThreadId, uint64* BackTrace, uint32 MaxDepth, void* Context = nullptr);
-	static void StackWalkAndDump( ANSICHAR* HumanReadableString, SIZE_T HumanReadableStringSize, int32 IgnoreCount, void* Context = nullptr );
+	FORCENOINLINE static void StackWalkAndDump( ANSICHAR* HumanReadableString, SIZE_T HumanReadableStringSize, int32 IgnoreCount, void* Context = nullptr );  // FORCENOINLINE so it can be counted during StackTrace
 	static void StackWalkAndDump( ANSICHAR* HumanReadableString, SIZE_T HumanReadableStringSize, void* ProgramCounter, void* Context = nullptr );
 	static void ThreadStackWalkAndDump(ANSICHAR* HumanReadableString, SIZE_T HumanReadableStringSize, int32 IgnoreCount, uint32 ThreadId);
 
diff --git a/Engine/Source/Runtime/CoreUObject/Private/Serialization/ArchiveStackTrace.cpp b/Engine/Source/Runtime/CoreUObject/Private/Serialization/ArchiveStackTrace.cpp
--- a/Engine/Source/Runtime/CoreUObject/Private/Serialization/ArchiveStackTrace.cpp
+++ b/Engine/Source/Runtime/CoreUObject/Private/Serialization/ArchiveStackTrace.cpp
@@ -229,11 +229,10 @@ void FArchiveStackTrace::Serialize(void* InData, int64 Num)
 {
 	if (Num)
 	{
-#if UE_BUILD_DEBUG
-		const int32 StackIgnoreCount = 5;
-#else
-		const int32 StackIgnoreCount = 4;
-#endif
+		// Remove this function from the reported stack. Note that all functions that we want to
+		// remove from the stack using this method have to be marked with FORCENOINLINE so they're
+		// there for us to remove and we don't remove meaningful functions above them instead.
+		const int32 StackIgnoreCount = 1;
 
 		static struct FBreakAtOffsetSettings
 		{
diff --git a/Engine/Source/Runtime/CoreUObject/Public/Serialization/ArchiveStackTrace.h b/Engine/Source/Runtime/CoreUObject/Public/Serialization/ArchiveStackTrace.h
--- a/Engine/Source/Runtime/CoreUObject/Public/Serialization/ArchiveStackTrace.h
+++ b/Engine/Source/Runtime/CoreUObject/Public/Serialization/ArchiveStackTrace.h
@@ -229,7 +229,7 @@ public:
 	virtual ~FArchiveStackTrace();
 
 	//~ Begin FArchive Interface
-	virtual void Serialize(void* Data, int64 Num) override;
+	FORCENOINLINE virtual void Serialize(void* Data, int64 Num) override; // FORCENOINLINE so it can be counted during StackTrace
 	virtual void SetSerializeContext(FUObjectSerializeContext* InLoadContext) override;
 	virtual FUObjectSerializeContext* GetSerializeContext() override;
 #if WITH_EDITOR
