Tasks integration: implemented a helper `CreateCompletionHandle()`.
Tasks are allocated by ConcurrentLinearAllocator that have a problem with long-living allocs as a single alloc can hold an entire memory page thus driving total mem usage up. `CreateCompletionHandle()` creates a task event that uses an allocator that doesn't have this issue. This task event depends on the parent task and can be used to wait for this task completion

diff --git a/Engine/Source/Runtime/Core/Private/Tests/Async/TaskGraphTest.cpp b/Engine/Source/Runtime/Core/Private/Tests/Async/TaskGraphTest.cpp
--- a/Engine/Source/Runtime/Core/Private/Tests/Async/TaskGraphTest.cpp
+++ b/Engine/Source/Runtime/Core/Private/Tests/Async/TaskGraphTest.cpp
@@ -1433,6 +1433,24 @@ PRAGMA_ENABLE_DEPRECATION_WARNINGS
 
 		return true;
 	}
+
+	IMPLEMENT_SIMPLE_AUTOMATION_TEST(FTaskGraphCreateCompletionHandleTest, "System.Core.Async.TaskGraph.CreateCompletionHandle", EAutomationTestFlags::ApplicationContextMask | EAutomationTestFlags::EngineFilter);
+
+	bool FTaskGraphCreateCompletionHandleTest::RunTest(const FString& Parameters)
+	{
+		FGraphEventRef UnblockTask = FGraphEvent::CreateGraphEvent(); // to block initially the following task
+		FGraphEventRef Task = FFunctionGraphTask::CreateAndDispatchWhenReady([] {}, TStatId{}, UnblockTask); // supposedly long-living task
+		FGraphEventRef CompletionHandle = Task->CreateCompletionHandle();
+		// check that the completion handle is not signalling as the task is blocked
+		FPlatformProcess::Sleep(0.1f);
+		check(!CompletionHandle->IsComplete());
+		// unblock the task
+		UnblockTask->DispatchSubsequents();
+		// wating for the completion handle instead of waiting for the task, should succeed
+		CompletionHandle->Wait();
+
+		return true;
+	}
 }
 
 #endif //WITH_DEV_AUTOMATION_TESTS
diff --git a/Engine/Source/Runtime/Core/Private/Tests/Tasks/TasksTest.cpp b/Engine/Source/Runtime/Core/Private/Tests/Tasks/TasksTest.cpp
--- a/Engine/Source/Runtime/Core/Private/Tests/Tasks/TasksTest.cpp
+++ b/Engine/Source/Runtime/Core/Private/Tests/Tasks/TasksTest.cpp
@@ -1320,6 +1320,24 @@ namespace UE { namespace TasksTests
 
 		return true;
 	}
+
+	IMPLEMENT_SIMPLE_AUTOMATION_TEST(FTasksCreateCompletionHandleTest, "System.Core.Async.Tasks.CreateCompletionHandle", EAutomationTestFlags::ApplicationContextMask | EAutomationTestFlags::EngineFilter);
+
+	bool FTasksCreateCompletionHandleTest::RunTest(const FString& Parameters)
+	{
+		FTaskEvent UnblockTask{ UE_SOURCE_LOCATION }; // to block initially the following task
+		FTask Task = Launch(UE_SOURCE_LOCATION, [] {}, Prerequisites(UnblockTask)); // supposedly long-living task
+		FTask CompletionHandle = Task.CreateCompletionHandle();
+		// check that the completion handle is not signalling as the task is blocked
+		FPlatformProcess::Sleep(0.1f);
+		check(!CompletionHandle.IsCompleted());
+		// unblock the task
+		UnblockTask.Trigger();
+		// wating for the completion handle instead of waiting for the task, should succeed
+		check(CompletionHandle.Wait(FTimespan::FromMilliseconds(100)));
+
+		return true;
+	}
 }}
 
 #endif // WITH_DEV_AUTOMATION_TESTS
diff --git a/Engine/Source/Runtime/Core/Public/Async/TaskGraphInterfaces.h b/Engine/Source/Runtime/Core/Public/Async/TaskGraphInterfaces.h
--- a/Engine/Source/Runtime/Core/Public/Async/TaskGraphInterfaces.h
+++ b/Engine/Source/Runtime/Core/Public/Async/TaskGraphInterfaces.h
@@ -628,6 +628,27 @@ public:
 
 		return (ENamedThreads::Type)ConversionMap[(int32)Priority - (int32)EExtendedTaskPriority::GameThreadNormalPri];
 	}
+
+	// see Tasks::FTaskHandle::CreateCompletionHandle description
+	FGraphEventRef CreateCompletionHandle()
+	{
+		if (IsCompleted())
+		{
+			return {};
+		}
+
+		// CreateGraphEvent() uses an allocator that doesn't have an issue with long-living allocs
+		FGraphEventRef CompletionHandle = FGraphEvent::CreateGraphEvent();
+		if (!CompletionHandle->AddPrerequisites(*this))
+		{
+			return {}; // too late, the task is already completed
+		}
+
+		// trigger the completion handle so the only thing that holds it from signalling is the task itself
+		CompletionHandle->DispatchSubsequents();
+
+		return CompletionHandle;
+	}
 };
 
 // the new task implementation integrated into the old task API
@@ -1068,6 +1089,13 @@ public:
 #endif
 	}
 
+	// does nothing and is needed only for compatibility with the new frontend
+	FGraphEventRef CreateCompletionHandle()
+	{
+		// nothing to do here as this instance already serves as a completion handle
+		return this;
+	}
+
 private:
 	friend class TRefCountPtr<FGraphEvent>;
 	friend class TLockFreeClassAllocator_TLSCache<FGraphEvent, PLATFORM_CACHE_LINE_SIZE>;
diff --git a/Engine/Source/Runtime/Core/Public/Tasks/Task.h b/Engine/Source/Runtime/Core/Public/Tasks/Task.h
--- a/Engine/Source/Runtime/Core/Public/Tasks/Task.h
+++ b/Engine/Source/Runtime/Core/Public/Tasks/Task.h
@@ -180,6 +180,35 @@ namespace UE::Tasks
 				return IsValid() && Pimpl->IsAwaitable();
 			}
 
+			// Creates and returns a "task event" that can be used to wait for the task completion.
+			// Regular tasks are allocated by a fast allocator that doesn't handle well long-living tasks, as such tasks can keep an entire memory page alive, 
+			// thus pushing up total memory consumption. In most cases such tasks are stored as class members, or as global vars.
+			// hovewer, task events use a different allocator and doesn't cause this issue, and so can be used for long-living tasks.
+			// Make sure to profile and identify that you indeed have "long-living task" problem before using this function, 
+			// as otherwise it would be a needless overhead.
+			// Task events don't support execution result. if you do need execution result, most probably you can reset your task right after getting its execution
+			// result, so it's not long-living anymore. 
+			FTaskHandle CreateCompletionHandle()
+			{
+				if (!IsValid() || IsCompleted())
+				{
+					return {};
+				}
+
+				// `FTaskEventBase` uses an allocator that doesn't have an issue with long-living allocs
+				Private::FTaskEventBase* CompletionHandle{ Private::FTaskEventBase::Create(TEXT("CompletionHandle")) };
+				if (!CompletionHandle->AddPrerequisites(*Pimpl))
+				{
+					delete CompletionHandle;
+					return {}; // too late, the task is already completed
+				}
+
+				CompletionHandle->AddRef(); // internal reference that is released when the handle is completed
+				// trigger the completion handle so the only thing that holds it from signalling is the task itself
+				CompletionHandle->TryLaunch();
+				return FTaskHandle{ CompletionHandle };
+			}
+
 		protected:
 			TRefCountPtr<FTaskBase> Pimpl;
 		};
diff --git a/Engine/Source/Runtime/Core/Public/Tasks/TaskPrivate.h b/Engine/Source/Runtime/Core/Public/Tasks/TaskPrivate.h
--- a/Engine/Source/Runtime/Core/Public/Tasks/TaskPrivate.h
+++ b/Engine/Source/Runtime/Core/Public/Tasks/TaskPrivate.h
@@ -169,7 +169,7 @@ namespace UE::Tasks
 			// The task will be executed only when all prerequisites are completed. The task type must be a task handle that holds a pointer to
 			// FTaskBase as its `Pimpl` member (see Tasks::TTaskBase).
 			// Must not be called concurrently
-			void AddPrerequisites(FTaskBase& Prerequisite)
+			bool AddPrerequisites(FTaskBase& Prerequisite)
 			{
 				checkf(NumLocks.load(std::memory_order_relaxed) >= NumInitialLocks && NumLocks.load(std::memory_order_relaxed) < ExecutionFlag, TEXT("Prerequisites can be added only before the task is launched"));
 
@@ -180,25 +180,25 @@ namespace UE::Tasks
 				// `AddSubsequent` provides required sync
 				checkf(PrevNumLocks + 1 < ExecutionFlag, TEXT("Max number of nested tasks reached: %d"), ExecutionFlag);
 
-				if (Prerequisite.AddSubsequent(*this)) // linearisation point, acq_rel semantic
-				{
-					Prerequisite.AddRef(); // keep it alive until this task's execution
-					Prerequisites.Push(&Prerequisite); // release memory order
-				}
-				else
+				if (!Prerequisite.AddSubsequent(*this)) // linearisation point, acq_rel semantic
 				{
 					// failed to add the prerequisite (too late), correct the number
 					NumLocks.fetch_sub(1, std::memory_order_relaxed); // relaxed because the previous `AddSubsequent` call provides required sync
+					return false;
 				}
+
+				Prerequisite.AddRef(); // keep it alive until this task's execution
+				Prerequisites.Push(&Prerequisite); // release memory order
+				return true;
 			}
 
 			// The task will be executed only when all prerequisites are completed. The task type must be a task handle that holds a pointer to
 			// FTaskBase as its `Pimpl` member (see Tasks::TTaskBase).
 			// Must not be called concurrently
 			template<typename HigherLevelTaskType, decltype(std::declval<HigherLevelTaskType>().Pimpl)* = nullptr>
-			void AddPrerequisites(const HigherLevelTaskType& Prerequisite)
+			bool AddPrerequisites(const HigherLevelTaskType& Prerequisite)
 			{
-				AddPrerequisites(*Prerequisite.Pimpl);
+				return AddPrerequisites(*Prerequisite.Pimpl);
 			}
 
 			// The task will be executed only when all prerequisites are completed.
