CompactBinary: Converted to FPlatformMemory::ReadUnaligned

#rnx

diff --git a/Engine/Source/Runtime/Core/Public/Serialization/CompactBinaryValue.h b/Engine/Source/Runtime/Core/Public/Serialization/CompactBinaryValue.h
--- a/Engine/Source/Runtime/Core/Public/Serialization/CompactBinaryValue.h
+++ b/Engine/Source/Runtime/Core/Public/Serialization/CompactBinaryValue.h
@@ -3,7 +3,7 @@
 #pragma once
 
 #include "Containers/StringView.h"
-#include "HAL/UnrealMemory.h"
+#include "HAL/PlatformMemory.h"
 #include "IO/IoHash.h"
 #include "Memory/MemoryView.h"
 #include "Misc/ByteSwap.h"
@@ -181,8 +181,7 @@ FORCEINLINE FIoHash FCbValue::AsHash() const
 
 FORCEINLINE FGuid FCbValue::AsUuid() const
 {
-	FGuid Value;
-	FMemory::Memcpy(&Value, Data, sizeof(FGuid));
+	FGuid Value = FPlatformMemory::ReadUnaligned<FGuid>(Data);
 	Value.A = NETWORK_ORDER32(Value.A);
 	Value.B = NETWORK_ORDER32(Value.B);
 	Value.C = NETWORK_ORDER32(Value.C);
