CompactBinary: Added LoadFromCompactBinary as a standardized way to load from compact binary

The intent is that function take a field and an output parameter, with an optional default value:

bool LoadFromCompactBinary(FCbFieldView Field, FTypeName& OutValue, FTypeName Default = {});

Implementations must return true if the field was loaded successfully, and false if any required part of the field was missing or its value had the wrong type or range. The output value must be assigned before the function returns.

diff --git a/Engine/Source/Runtime/Core/Private/Serialization/CompactBinarySerialization.cpp b/Engine/Source/Runtime/Core/Private/Serialization/CompactBinarySerialization.cpp
--- a/Engine/Source/Runtime/Core/Private/Serialization/CompactBinarySerialization.cpp
+++ b/Engine/Source/Runtime/Core/Private/Serialization/CompactBinarySerialization.cpp
@@ -265,6 +265,32 @@ FArchive& operator<<(FArchive& Ar, FCbObject& Object)
 
 ///////////////////////////////////////////////////////////////////////////////////////////////////////////////////////
 
+bool LoadFromCompactBinary(FCbFieldView Field, FUtf8StringBuilderBase& OutValue)
+{
+	OutValue << Field.AsString();
+	return !Field.HasError();
+}
+
+bool LoadFromCompactBinary(FCbFieldView Field, FWideStringBuilderBase& OutValue)
+{
+	OutValue << Field.AsString();
+	return !Field.HasError();
+}
+
+bool LoadFromCompactBinary(FCbFieldView Field, FGuid& OutValue)
+{
+	OutValue = Field.AsUuid();
+	return !Field.HasError();
+}
+
+bool LoadFromCompactBinary(FCbFieldView Field, FGuid& OutValue, const FGuid& Default)
+{
+	OutValue = Field.AsUuid(Default);
+	return !Field.HasError();
+}
+
+///////////////////////////////////////////////////////////////////////////////////////////////////////////////////////
+
 class FCbJsonWriter
 {
 public:
diff --git a/Engine/Source/Runtime/Core/Public/Serialization/CompactBinarySerialization.h b/Engine/Source/Runtime/Core/Public/Serialization/CompactBinarySerialization.h
--- a/Engine/Source/Runtime/Core/Public/Serialization/CompactBinarySerialization.h
+++ b/Engine/Source/Runtime/Core/Public/Serialization/CompactBinarySerialization.h
@@ -77,6 +77,101 @@ CORE_API FArchive& operator<<(FArchive& Ar, FCbObject& Object);
 
 ///////////////////////////////////////////////////////////////////////////////////////////////////////////////////////
 
+/**
+ * LoadFromCompactBinary attempts to load the output value from compact binary.
+ *
+ * Implementations of LoadCompactBinary are expected to assign the output even on failure.
+ * Implementations may accept an optional default value to assign in case of load failure.
+ *
+ * @return true if required fields had the values with the correct type and range, otherwise false.
+ */
+
+CORE_API bool LoadFromCompactBinary(FCbFieldView Field, FUtf8StringBuilderBase& OutValue);
+CORE_API bool LoadFromCompactBinary(FCbFieldView Field, FWideStringBuilderBase& OutValue);
+
+inline bool LoadFromCompactBinary(FCbFieldView Field, int8& OutValue, const int8 Default = 0)
+{
+	OutValue = Field.AsInt8(Default);
+	return !Field.HasError();
+}
+
+inline bool LoadFromCompactBinary(FCbFieldView Field, int16& OutValue, const int16 Default = 0)
+{
+	OutValue = Field.AsInt16(Default);
+	return !Field.HasError();
+}
+
+inline bool LoadFromCompactBinary(FCbFieldView Field, int32& OutValue, const int32 Default = 0)
+{
+	OutValue = Field.AsInt32(Default);
+	return !Field.HasError();
+}
+
+inline bool LoadFromCompactBinary(FCbFieldView Field, int64& OutValue, const int64 Default = 0)
+{
+	OutValue = Field.AsInt64(Default);
+	return !Field.HasError();
+}
+
+inline bool LoadFromCompactBinary(FCbFieldView Field, uint8& OutValue, const uint8 Default = 0)
+{
+	OutValue = Field.AsUInt8(Default);
+	return !Field.HasError();
+}
+
+inline bool LoadFromCompactBinary(FCbFieldView Field, uint16& OutValue, const uint16 Default = 0)
+{
+	OutValue = Field.AsUInt16(Default);
+	return !Field.HasError();
+}
+
+inline bool LoadFromCompactBinary(FCbFieldView Field, uint32& OutValue, const uint32 Default = 0)
+{
+	OutValue = Field.AsUInt32(Default);
+	return !Field.HasError();
+}
+
+inline bool LoadFromCompactBinary(FCbFieldView Field, uint64& OutValue, const uint64 Default = 0)
+{
+	OutValue = Field.AsUInt64(Default);
+	return !Field.HasError();
+}
+
+inline bool LoadFromCompactBinary(FCbFieldView Field, float& OutValue, const float Default = 0.0f)
+{
+	OutValue = Field.AsFloat(Default);
+	return !Field.HasError();
+}
+
+inline bool LoadFromCompactBinary(FCbFieldView Field, double& OutValue, const double Default = 0.0)
+{
+	OutValue = Field.AsDouble(Default);
+	return !Field.HasError();
+}
+
+inline bool LoadFromCompactBinary(FCbFieldView Field, bool& OutValue, const bool Default = false)
+{
+	OutValue = Field.AsBool(Default);
+	return !Field.HasError();
+}
+
+inline bool LoadFromCompactBinary(FCbFieldView Field, FIoHash& OutValue, const FIoHash& Default = FIoHash())
+{
+	OutValue = Field.AsHash(Default);
+	return !Field.HasError();
+}
+
+CORE_API bool LoadFromCompactBinary(FCbFieldView Field, FGuid& OutValue);
+CORE_API bool LoadFromCompactBinary(FCbFieldView Field, FGuid& OutValue, const FGuid& Default);
+
+inline bool LoadFromCompactBinary(FCbFieldView Field, FCbObjectId& OutValue, const FCbObjectId& Default = FCbObjectId())
+{
+	OutValue = Field.AsObjectId(Default);
+	return !Field.HasError();
+}
+
+///////////////////////////////////////////////////////////////////////////////////////////////////////////////////////
+
 /** Convert the object to JSON in a multi-line format with indentation. */
 CORE_API void CompactBinaryToJson(const FCbObjectView& Object, FUtf8StringBuilderBase& Builder);
 
