More FCString<UTF8CHAR> support.

diff --git a/Engine/Source/Runtime/Core/Public/Misc/CString.h b/Engine/Source/Runtime/Core/Public/Misc/CString.h
--- a/Engine/Source/Runtime/Core/Public/Misc/CString.h
+++ b/Engine/Source/Runtime/Core/Public/Misc/CString.h
@@ -947,13 +947,58 @@ inline int32 TCString<ANSICHAR>::SnprintfImpl(CharType* Dest, int32 DestSize, co
 	return Result;
 }
 
-template <> 
+template <>
 FORCEINLINE bool TCString<ANSICHAR>::ToBool(const ANSICHAR* Str)
 {
 	return FToBoolHelper::FromCStringAnsi(Str);
 }
 
-template <> 
+/*-----------------------------------------------------------------------------
+	TCString<UTF8CHAR> specializations
+-----------------------------------------------------------------------------*/
+template <> FORCEINLINE
+bool TCString<UTF8CHAR>::IsPureAnsi(const CharType* Str)
+{
+	for (; *Str; Str++)
+	{
+		if ((uint8)*Str > 0x7f)
+		{
+			return false;
+		}
+	}
+	return true;
+}
+
+template <> FORCEINLINE
+bool TCString<UTF8CHAR>::IsPureAnsi(const CharType* Str, const SIZE_T StrLen)
+{
+	for (SIZE_T Idx = 0; Idx < StrLen; Idx++, Str++)
+	{
+		if ((uint8)*Str > 0x7f)
+		{
+			return false;
+		}
+	}
+	return true;
+}
+
+template <>
+inline int32 TCString<UTF8CHAR>::SprintfImpl(CharType* Dest, const CharType* Fmt, ...)
+{
+	int32	Result = -1;
+	GET_VARARGS_RESULT_UTF8(Dest, MAX_SPRINTF, MAX_SPRINTF - 1, Fmt, Fmt, Result);
+	return Result;
+}
+
+template <>
+inline int32 TCString<UTF8CHAR>::SnprintfImpl(CharType* Dest, int32 DestSize, const CharType* Fmt, ...)
+{
+	int32	Result = -1;
+	GET_VARARGS_RESULT_UTF8(Dest, DestSize, DestSize - 1, Fmt, Fmt, Result);
+	return Result;
+}
+
+template <>
 FORCEINLINE bool TCString<UTF8CHAR>::ToBool(const UTF8CHAR* Str)
 {
 	return FToBoolHelper::FromCStringUtf8(Str);
diff --git a/Engine/Source/Runtime/Core/Public/Misc/VarArgs.h b/Engine/Source/Runtime/Core/Public/Misc/VarArgs.h
--- a/Engine/Source/Runtime/Core/Public/Misc/VarArgs.h
+++ b/Engine/Source/Runtime/Core/Public/Misc/VarArgs.h
@@ -69,6 +69,17 @@ VarArgs helper macros.
 		} \
 		va_end(ap); \
 	}
+#define GET_VARARGS_RESULT_UTF8(msg, msgsize, len, lastarg, fmt, result) \
+	{ \
+		va_list ap; \
+		va_start(ap, lastarg); \
+		result = FCStringUtf8::GetVarArgs(msg, msgsize, fmt, ap); \
+		if (result >= msgsize) \
+		{ \
+			result = -1; \
+		} \
+		va_end(ap); \
+	}
 
 /*-----------------------------------------------------------------------------
 Ugly VarArgs type checking (debug builds only).
