Fix for csv profiler pause event not firing at start of capture

+ Added a new delegate to the CSV profiler that fires once on frame 0 of the capture so we can send the pause event as soon as valid.
+ The PerformanceTrackingSubsystem now handles sending the pause/resume events as it made more sense.

diff --git a/Engine/Source/Runtime/Core/Private/ProfilingDebugging/CsvProfiler.cpp b/Engine/Source/Runtime/Core/Private/ProfilingDebugging/CsvProfiler.cpp
--- a/Engine/Source/Runtime/Core/Private/ProfilingDebugging/CsvProfiler.cpp
+++ b/Engine/Source/Runtime/Core/Private/ProfilingDebugging/CsvProfiler.cpp
@@ -2835,6 +2835,11 @@ void FCsvProfiler::BeginFrame()
 		if (GCsvProfilerIsCapturing)
 		{
 			GFrameBoundaries.AddBeginFrameTimestamp(ECsvTimeline::Gamethread);
+
+			if (CaptureFrameNumber == 0)
+			{
+				OnCSVProfileFirstFrameDelegate.Broadcast();
+			}
 		}
 	}
 
diff --git a/Engine/Source/Runtime/Core/Public/ProfilingDebugging/CsvProfiler.h b/Engine/Source/Runtime/Core/Public/ProfilingDebugging/CsvProfiler.h
--- a/Engine/Source/Runtime/Core/Public/ProfilingDebugging/CsvProfiler.h
+++ b/Engine/Source/Runtime/Core/Public/ProfilingDebugging/CsvProfiler.h
@@ -287,9 +287,14 @@ public:
 
 	CORE_API void GetFrameExecCommands(TArray<FString>& OutFrameCommands) const;
 
+	/** Called right before we start capturing. */
 	DECLARE_MULTICAST_DELEGATE(FOnCSVProfileStart);
 	FOnCSVProfileStart& OnCSVProfileStart() { return OnCSVProfileStartDelegate; }
 
+	/** Called when csv frame 0 begins its capture. This is when CsvEvents and CustomStats will start being collected. */
+	DECLARE_MULTICAST_DELEGATE(FOnCSVProfileFirstFrame);
+	FOnCSVProfileFirstFrame& OnCSVProfileFirstFrame() { return OnCSVProfileFirstFrameDelegate; }
+
 	DECLARE_MULTICAST_DELEGATE(FOnCSVProfileEnd);
 	FOnCSVProfileEnd& OnCSVProfileEnd() { return OnCSVProfileEndDelegate; }
 	
@@ -340,6 +345,7 @@ private:
 	ECsvProfilerFlags CurrentFlags;
 
 	FOnCSVProfileStart OnCSVProfileStartDelegate;
+	FOnCSVProfileFirstFrame OnCSVProfileFirstFrameDelegate;
 	FOnCSVProfileEnd OnCSVProfileEndDelegate;
 	
 	FOnCSVProfileFinished OnCSVProfileFinishedDelegate;
