Added an overload for `TMpscQueue::Dequeue()` to align with `TSpscQueue` API

diff --git a/Engine/Source/Runtime/Core/Public/Containers/MpscQueue.h b/Engine/Source/Runtime/Core/Public/Containers/MpscQueue.h
--- a/Engine/Source/Runtime/Core/Public/Containers/MpscQueue.h
+++ b/Engine/Source/Runtime/Core/Public/Containers/MpscQueue.h
@@ -73,6 +73,18 @@ public:
 		return Res;
 	}
 
+	bool Dequeue(ElementType& OutElem)
+	{
+		TOptional<ElementType> LocalElement = Dequeue();
+		if (LocalElement.IsSet())
+		{
+			OutElem = LocalElement.GetValue();
+			return true;
+		}
+
+		return false;
+	}
+
 private:
 	struct FNode
 	{
