added `IsEmpty()` to TSpscQueue and TMpscQueue

diff --git a/Engine/Source/Runtime/Core/Public/Containers/MpscQueue.h b/Engine/Source/Runtime/Core/Public/Containers/MpscQueue.h
--- a/Engine/Source/Runtime/Core/Public/Containers/MpscQueue.h
+++ b/Engine/Source/Runtime/Core/Public/Containers/MpscQueue.h
@@ -85,6 +85,11 @@ public:
 		return false;
 	}
 
+	bool IsEmpty() const
+	{
+		return Tail->Next.load(std::memory_order_acquire) != nullptr;
+	}
+
 private:
 	struct FNode
 	{
diff --git a/Engine/Source/Runtime/Core/Public/Containers/SpscQueue.h b/Engine/Source/Runtime/Core/Public/Containers/SpscQueue.h
--- a/Engine/Source/Runtime/Core/Public/Containers/SpscQueue.h
+++ b/Engine/Source/Runtime/Core/Public/Containers/SpscQueue.h
@@ -92,6 +92,13 @@ public:
 		return false;
 	}
 
+	bool IsEmpty() const
+	{
+		FNode* LocalTail = Tail.load(std::memory_order_relaxed);
+		FNode* LocalTailNext = LocalTail->Next.load(std::memory_order_acquire);
+		return LocalTailNext != nullptr;
+	}
+
 private:
 	struct FNode
 	{
