TArray::Append(<any contiguous range>) added.
Fix for TArrayElementsAreCompatible, which was always evaluating to true.

diff --git a/Engine/Source/Runtime/Core/Public/Containers/Array.h b/Engine/Source/Runtime/Core/Public/Containers/Array.h
--- a/Engine/Source/Runtime/Core/Public/Containers/Array.h
+++ b/Engine/Source/Runtime/Core/Public/Containers/Array.h
@@ -26,6 +26,7 @@
 #include "Templates/Sorting.h"
 #include "Templates/AlignmentTemplates.h"
 #include "Templates/IsConstructible.h"
+#include "Traits/ElementType.h"
 
 #include <type_traits>
 
@@ -269,6 +270,14 @@ private:
 
 namespace UE4Array_Private
 {
+	// Simply forwards to an unqualified GetData(), but can be called from within a container or view
+	// where GetData() is already a member and so hides any others.
+	template <typename T>
+	FORCEINLINE decltype(auto) GetDataHelper(T&& Arg)
+	{
+		return GetData(Forward<T>(Arg));
+	}
+
 	template <typename FromArrayType, typename ToArrayType>
 	struct TCanMoveTArrayPointersBetweenArrayTypes
 	{
@@ -302,7 +311,24 @@ namespace UE4Array_Private
 	// This should be changed to use std::disjunction and std::is_constructible, and the usage
 	// changed to use ::value instead of ::Value, when std::disjunction (C++17) is available everywhere.
 	template <typename DestType, typename SourceType>
-	using TArrayElementsAreCompatible = TOrValue<std::is_same<DestType, std::decay_t<DestType>>::value, TIsConstructible<DestType, SourceType>>;
+	using TArrayElementsAreCompatible = TOrValue<std::is_same<DestType, std::decay_t<SourceType>>::value, TIsConstructible<DestType, SourceType>>;
+
+	template <typename T>
+	struct TIsTArrayOrDerivedFromTArray
+	{
+		template <typename ElementType, typename AllocatorType>
+		static char (&Resolve(const volatile TArray<ElementType, AllocatorType>*))[2];
+
+		static char(&Resolve(...))[1];
+
+		enum { Value = sizeof(Resolve((T*)nullptr)) == 2 };
+	};
+
+	template <typename ElementType, typename RangeType>
+	struct TTypeIsCompatibleWithRangeElementType
+	{
+		enum { Value = TArrayElementsAreCompatible<ElementType, TElementType_T<RangeType>>::Value };
+	};
 }
 
 
@@ -1899,6 +1925,44 @@ public:
 		ArrayNum += SourceCount;
 	}
 
+	/**
+	 * Appends the elements from a contiguous range to this array.
+	 *
+	 * @param Source The range of elements to append.
+	 * @see Add, Insert
+	 */
+	template <
+		typename RangeType,
+		typename RangeValueType = std::remove_reference_t<RangeType>,
+		typename RangeElementType = TElementType_T<RangeValueType>,
+		std::enable_if_t<
+			TAnd<
+				TIsContiguousContainer<RangeValueType>,
+				TNot<UE4Array_Private::TIsTArrayOrDerivedFromTArray<RangeValueType>>,
+				UE4Array_Private::TTypeIsCompatibleWithRangeElementType<ElementType, RangeType>
+			>::Value
+		>* = nullptr
+	>
+	void Append(RangeType&& Source)
+	{
+		auto InCount = GetNum(Source);
+		checkf((InCount >= 0) && ((sizeof(InCount) < sizeof(SizeType)) || (InCount <= static_cast<decltype(InCount)>(TNumericLimits<SizeType>::Max()))), TEXT("Invalid range size: %lld"), (long long)InCount);
+
+		// Do nothing if the source is empty.
+		if (!InCount)
+		{
+			return;
+		}
+
+		SizeType SourceCount = (SizeType)InCount;
+
+		// Allocate memory for the new elements.
+		Reserve(ArrayNum + SourceCount);
+		ConstructItems<ElementType>(GetData() + ArrayNum, UE4Array_Private::GetDataHelper(Source), SourceCount);
+
+		ArrayNum += SourceCount;
+	}
+
 	/**
 	 * Adds a raw array of elements to the end of the TArray.
 	 *
