Setting the NetCL for 21.20

diff --git a/Engine/Source/Runtime/Core/Public/Misc/NetworkVersion.h b/Engine/Source/Runtime/Core/Public/Misc/NetworkVersion.h
--- a/Engine/Source/Runtime/Core/Public/Misc/NetworkVersion.h
+++ b/Engine/Source/Runtime/Core/Public/Misc/NetworkVersion.h
@@ -8,7 +8,7 @@
 #include "Delegates/Delegate.h"
 
 // The version number used for determining network compatibility. If zero, uses the engine compatible version.
-#define ENGINE_NET_VERSION  0
+#define ENGINE_NET_VERSION  20788829
 
 // The version number used for determining replay compatibility
 #define ENGINE_REPLAY_VERSION  ENGINE_NET_VERSION
