Fix TSan build failure

Need to handle std::atomic in TGuardValue in a nicer way...

diff --git a/Engine/Source/Runtime/Core/Private/Misc/CoreGlobals.cpp b/Engine/Source/Runtime/Core/Private/Misc/CoreGlobals.cpp
--- a/Engine/Source/Runtime/Core/Private/Misc/CoreGlobals.cpp
+++ b/Engine/Source/Runtime/Core/Private/Misc/CoreGlobals.cpp
@@ -387,7 +387,7 @@ bool					GShouldSuspendRenderingThread	= false;
 /** Determines what kind of trace should occur, NAME_None for none.											*/
 FLazyName				GCurrentTraceName;
 /** How to print the time in log output																		*/
-TSAN_ATOMIC(ELogTimes::Type)	GPrintLogTimes			= ELogTimes::None;
+ELogTimes::Type	GPrintLogTimes			= ELogTimes::None;
 /** How to print the category in log output.																*/
 TSAN_ATOMIC(bool)		GPrintLogCategory				= true;
 /** How to print the verbosity in log output.																*/
diff --git a/Engine/Source/Runtime/Core/Public/CoreGlobals.h b/Engine/Source/Runtime/Core/Public/CoreGlobals.h
--- a/Engine/Source/Runtime/Core/Public/CoreGlobals.h
+++ b/Engine/Source/Runtime/Core/Public/CoreGlobals.h
@@ -496,7 +496,7 @@ extern CORE_API bool GShouldSuspendRenderingThread;
 extern CORE_API FLazyName GCurrentTraceName;
 
 /** How to print the time in log output. */
-extern CORE_API TSAN_ATOMIC(ELogTimes::Type) GPrintLogTimes;
+extern CORE_API ELogTimes::Type GPrintLogTimes;
 
 /** How to print the category in log output. */
 extern CORE_API TSAN_ATOMIC(bool) GPrintLogCategory;
