Fix bad scoped locks in various places.

Added nodiscard on the scope lock helpers to prevent such errors in the future.

#rnx

diff --git a/Engine/Plugins/FastBuildController/Source/Private/FastBuildControllerModule.cpp b/Engine/Plugins/FastBuildController/Source/Private/FastBuildControllerModule.cpp
--- a/Engine/Plugins/FastBuildController/Source/Private/FastBuildControllerModule.cpp
+++ b/Engine/Plugins/FastBuildController/Source/Private/FastBuildControllerModule.cpp
@@ -288,14 +288,14 @@ void FFastBuildControllerModule::RegisterDispatchedTask(FTask* DispatchedTask)
 {
 	check(DispatchedTask)
 	{
-		FScopeLock(TasksCS.Get());
+		FScopeLock Lock(TasksCS.Get());
 		DispatchedTasks.Add(DispatchedTask->ID , DispatchedTask);
 	}
 }
 
 void FFastBuildControllerModule::ReEnqueueDispatchedTasks()
 {
-	FScopeLock(TasksCS.Get());
+	FScopeLock Lock(TasksCS.Get());
 	// Reclaim dispatched (incomplete) tasks 
 	for (const TPair<uint32, FTask*>& DispatchedTaskEntry : DispatchedTasks)
 	{
@@ -309,7 +309,7 @@ void FFastBuildControllerModule::DeRegisterDispatchedTasks(const TArray<uint32>&
 {
 	for (const uint32& TaskID : InTasksID)
 	{
-		FScopeLock(TasksCS.Get());
+		FScopeLock Lock(TasksCS.Get());
 		DispatchedTasks.Remove(TaskID);
 	}
 }
diff --git a/Engine/Plugins/FastBuildController/Source/Private/FastBuildJobProcessor.cpp b/Engine/Plugins/FastBuildController/Source/Private/FastBuildJobProcessor.cpp
--- a/Engine/Plugins/FastBuildController/Source/Private/FastBuildJobProcessor.cpp
+++ b/Engine/Plugins/FastBuildController/Source/Private/FastBuildJobProcessor.cpp
@@ -245,7 +245,7 @@ void FFastBuildJobProcessor::SubmitPendingJobs()
 	TArray<FTask*> TasksToSubmit;
 	FTask* PendingTask = nullptr;
 	{
-		FScopeLock(ControllerModule.GetTasksCS());
+		FScopeLock Lock(ControllerModule.GetTasksCS());
 		while (ControllerModule.AreTasksPending())
 		{
 			PendingTask = ControllerModule.DequeueTask();
diff --git a/Engine/Source/Developer/DerivedDataCache/Private/DerivedDataRequestOwner.cpp b/Engine/Source/Developer/DerivedDataCache/Private/DerivedDataRequestOwner.cpp
--- a/Engine/Source/Developer/DerivedDataCache/Private/DerivedDataRequestOwner.cpp
+++ b/Engine/Source/Developer/DerivedDataCache/Private/DerivedDataRequestOwner.cpp
@@ -89,7 +89,12 @@ void FRequestOwnerShared::Begin(IRequest* Request)
 	for (EPriority CheckPriority; ; NewPriority = CheckPriority)
 	{
 		Request->SetPriority(NewPriority);
-		CheckPriority = (FReadScopeLock(Lock), Priority);
+		
+		{
+			FReadScopeLock ScopeLock(Lock);
+			CheckPriority = Priority;
+		}
+
 		if (CheckPriority == NewPriority)
 		{
 			break;
@@ -255,7 +260,11 @@ void FRequestOwnerShared::KeepAlive()
 
 void FRequestOwnerShared::Destroy()
 {
-	const bool bLocalKeepAlive = (FWriteScopeLock(Lock), bKeepAlive);
+	bool bLocalKeepAlive;
+	{
+		FWriteScopeLock ScopeLock(Lock);
+		bLocalKeepAlive = bKeepAlive;
+	}
 	if (!bLocalKeepAlive)
 	{
 		Cancel();
diff --git a/Engine/Source/Runtime/Core/Private/HAL/Allocators/CachedOSPageAllocator.cpp b/Engine/Source/Runtime/Core/Private/HAL/Allocators/CachedOSPageAllocator.cpp
--- a/Engine/Source/Runtime/Core/Private/HAL/Allocators/CachedOSPageAllocator.cpp
+++ b/Engine/Source/Runtime/Core/Private/HAL/Allocators/CachedOSPageAllocator.cpp
@@ -60,7 +60,7 @@ void* FCachedOSPageAllocator::AllocateImpl(SIZE_T Size, uint32 CachedByteLimit,
 			void* Ptr = nullptr;
 			{
 #if UE_ALLOW_OSMEMORYLOCKFREE
-				FScopeUnlock FScopeUnlock(Mutex);
+				FScopeUnlock ScopeUnlock(Mutex);
 #endif
 				Ptr = FPlatformMemory::BinnedAllocFromOS(Size);
 			}
@@ -86,7 +86,7 @@ void* FCachedOSPageAllocator::AllocateImpl(SIZE_T Size, uint32 CachedByteLimit,
 	void* Ret = nullptr;
 	{
 #if UE_ALLOW_OSMEMORYLOCKFREE
-		FScopeUnlock FScopeUnlock(Mutex);
+		FScopeUnlock ScopeUnlock(Mutex);
 #endif
 		Ret = FPlatformMemory::BinnedAllocFromOS(Size);
 	}
@@ -98,7 +98,7 @@ void FCachedOSPageAllocator::FreeImpl(void* Ptr, SIZE_T Size, uint32 NumCacheBlo
 	if (FPlatformMemory::BinnedPlatformHasMemoryPoolForThisSize(Size) || Size > CachedByteLimit / 4)
 	{
 #if UE_ALLOW_OSMEMORYLOCKFREE
-		FScopeUnlock FScopeUnlock(Mutex);
+		FScopeUnlock ScopeUnlock(Mutex);
 #endif
 		FPlatformMemory::BinnedFreeToOS(Ptr, Size);
 		return;
@@ -117,7 +117,7 @@ void FCachedOSPageAllocator::FreeImpl(void* Ptr, SIZE_T Size, uint32 NumCacheBlo
 		}
 		{
 #if UE_ALLOW_OSMEMORYLOCKFREE
-			FScopeUnlock FScopeUnlock(Mutex);
+			FScopeUnlock ScopeUnlock(Mutex);
 #endif
 			FPlatformMemory::BinnedFreeToOS(FreePtr, FreeSize);
 		}
diff --git a/Engine/Source/Runtime/Core/Private/HAL/Allocators/CachedOSVeryLargePageAllocator.cpp b/Engine/Source/Runtime/Core/Private/HAL/Allocators/CachedOSVeryLargePageAllocator.cpp
--- a/Engine/Source/Runtime/Core/Private/HAL/Allocators/CachedOSVeryLargePageAllocator.cpp
+++ b/Engine/Source/Runtime/Core/Private/HAL/Allocators/CachedOSVeryLargePageAllocator.cpp
@@ -104,7 +104,7 @@ void* FCachedOSVeryLargePageAllocator::Allocate(SIZE_T Size, uint32 AllocationHi
 						LargePage->Unlink();
 						{
 #if UE_ALLOW_OSMEMORYLOCKFREE
-							FScopeUnlock FScopeUnlock(Mutex);
+							FScopeUnlock ScopeUnlock(Mutex);
 #endif
 							LLM_PLATFORM_SCOPE(ELLMTag::FMalloc);
 #if UE_USE_VERYLARGEPAGEALLOCATOR_FALLBACKPATH
@@ -113,7 +113,7 @@ void* FCachedOSVeryLargePageAllocator::Allocate(SIZE_T Size, uint32 AllocationHi
 #if UE_ALLOW_OSMEMORYLOCKFREE
 								if (Mutex != nullptr)
 								{
-									FScopeLock FScopeLock(Mutex);
+									FScopeLock Lock(Mutex);
 									LargePage->LinkHead(FreeLargePagesHead[AllocationHint]);
 								}
 								else
@@ -206,7 +206,7 @@ void FCachedOSVeryLargePageAllocator::Free(void* Ptr, SIZE_T Size, FCriticalSect
 				// need to move which list we are in and remove the backing store
 				{
 #if UE_ALLOW_OSMEMORYLOCKFREE
-					FScopeUnlock FScopeUnlock(Mutex);
+					FScopeUnlock ScopeUnlock(Mutex);
 #endif
 					Block.Decommit(LargePage->BaseAddress - AddressSpaceReserved, SizeOfLargePage);
 				}
diff --git a/Engine/Source/Runtime/Core/Public/Misc/ScopeLock.h b/Engine/Source/Runtime/Core/Public/Misc/ScopeLock.h
--- a/Engine/Source/Runtime/Core/Public/Misc/ScopeLock.h
+++ b/Engine/Source/Runtime/Core/Public/Misc/ScopeLock.h
@@ -32,7 +32,7 @@ public:
 	 *
 	 * @param InSynchObject The synchronization object to manage
 	 */
-	FScopeLock( FCriticalSection* InSynchObject )
+	UE_NODISCARD FScopeLock( FCriticalSection* InSynchObject )
 		: SynchObject(InSynchObject)
 	{
 		check(SynchObject);
@@ -99,7 +99,7 @@ public:
 	 *
 	 * @param InSynchObject The synchronization object to manage, can be null.
 	 */
-	FScopeUnlock(FCriticalSection* InSynchObject)
+	UE_NODISCARD FScopeUnlock(FCriticalSection* InSynchObject)
 		: SynchObject(InSynchObject)
 	{
 		if (InSynchObject)
@@ -151,7 +151,7 @@ namespace UE
 	public:
 		UE_NONCOPYABLE(TScopeLock);
 
-		TScopeLock(MutexType& InMutex)
+		UE_NODISCARD TScopeLock(MutexType& InMutex)
 			: Mutex(&InMutex)
 		{
 			check(Mutex);
@@ -193,7 +193,7 @@ namespace UE
 	public:
 		UE_NONCOPYABLE(TScopeUnlock);
 
-		TScopeUnlock(MutexType& InMutex)
+		UE_NODISCARD TScopeUnlock(MutexType& InMutex)
 			: Mutex(&InMutex)
 		{
 			check(Mutex);
diff --git a/Engine/Source/Runtime/Core/Public/Misc/ScopeRWLock.h b/Engine/Source/Runtime/Core/Public/Misc/ScopeRWLock.h
--- a/Engine/Source/Runtime/Core/Public/Misc/ScopeRWLock.h
+++ b/Engine/Source/Runtime/Core/Public/Misc/ScopeRWLock.h
@@ -11,7 +11,7 @@
 class FReadScopeLock
 {
 public:
-	explicit FReadScopeLock(FRWLock& InLock)
+	UE_NODISCARD explicit FReadScopeLock(FRWLock& InLock)
 		: Lock(InLock)
 	{
 		Lock.ReadLock();
@@ -32,7 +32,7 @@ private:
 class FWriteScopeLock
 {
 public:
-	explicit FWriteScopeLock(FRWLock& InLock)
+	UE_NODISCARD explicit FWriteScopeLock(FRWLock& InLock)
 		: Lock(InLock)
 	{
 		Lock.WriteLock();
@@ -67,7 +67,7 @@ enum FRWScopeLockType
 class FRWScopeLock
 {
 public:
-	explicit FRWScopeLock(FRWLock& InLockObject,FRWScopeLockType InLockType)
+	UE_NODISCARD explicit FRWScopeLock(FRWLock& InLockObject,FRWScopeLockType InLockType)
 	: LockObject(InLockObject)
 	, LockType(InLockType)
 	{
diff --git a/Engine/Source/Runtime/Engine/Private/Streaming/AsyncTextureStreaming.cpp b/Engine/Source/Runtime/Engine/Private/Streaming/AsyncTextureStreaming.cpp
--- a/Engine/Source/Runtime/Engine/Private/Streaming/AsyncTextureStreaming.cpp
+++ b/Engine/Source/Runtime/Engine/Private/Streaming/AsyncTextureStreaming.cpp
@@ -324,7 +324,7 @@ void FRenderAssetStreamingMipCalcTask::ApplyPakStateChanges_Async()
 
 	// Acquire the pending file state changes from the streaming manager.
 	{
-		FScopeLock(&StreamingManager.MountedStateDirtyFilesCS);
+		FScopeLock Lock(&StreamingManager.MountedStateDirtyFilesCS);
 		FMemory::Memswap(&MountedStateDirtyFiles, &StreamingManager.MountedStateDirtyFiles, sizeof(FIoFilenameHashSet));
 		FMemory::Memswap(&bRecacheAllFiles, &StreamingManager.bRecacheAllFiles, sizeof(bool));
 	}
