Removing nodiscard from scope lock constructors because sadly some compilers we use don't like that.

#rnx

diff --git a/Engine/Source/Runtime/Core/Public/Misc/ScopeLock.h b/Engine/Source/Runtime/Core/Public/Misc/ScopeLock.h
--- a/Engine/Source/Runtime/Core/Public/Misc/ScopeLock.h
+++ b/Engine/Source/Runtime/Core/Public/Misc/ScopeLock.h
@@ -32,7 +32,7 @@ public:
 	 *
 	 * @param InSynchObject The synchronization object to manage
 	 */
-	UE_NODISCARD FScopeLock( FCriticalSection* InSynchObject )
+	FScopeLock( FCriticalSection* InSynchObject )
 		: SynchObject(InSynchObject)
 	{
 		check(SynchObject);
@@ -99,7 +99,7 @@ public:
 	 *
 	 * @param InSynchObject The synchronization object to manage, can be null.
 	 */
-	UE_NODISCARD FScopeUnlock(FCriticalSection* InSynchObject)
+	FScopeUnlock(FCriticalSection* InSynchObject)
 		: SynchObject(InSynchObject)
 	{
 		if (InSynchObject)
@@ -151,7 +151,7 @@ namespace UE
 	public:
 		UE_NONCOPYABLE(TScopeLock);
 
-		UE_NODISCARD TScopeLock(MutexType& InMutex)
+		TScopeLock(MutexType& InMutex)
 			: Mutex(&InMutex)
 		{
 			check(Mutex);
@@ -193,7 +193,7 @@ namespace UE
 	public:
 		UE_NONCOPYABLE(TScopeUnlock);
 
-		UE_NODISCARD TScopeUnlock(MutexType& InMutex)
+		TScopeUnlock(MutexType& InMutex)
 			: Mutex(&InMutex)
 		{
 			check(Mutex);
diff --git a/Engine/Source/Runtime/Core/Public/Misc/ScopeRWLock.h b/Engine/Source/Runtime/Core/Public/Misc/ScopeRWLock.h
--- a/Engine/Source/Runtime/Core/Public/Misc/ScopeRWLock.h
+++ b/Engine/Source/Runtime/Core/Public/Misc/ScopeRWLock.h
@@ -11,7 +11,7 @@
 class FReadScopeLock
 {
 public:
-	UE_NODISCARD explicit FReadScopeLock(FRWLock& InLock)
+	explicit FReadScopeLock(FRWLock& InLock)
 		: Lock(InLock)
 	{
 		Lock.ReadLock();
@@ -32,7 +32,7 @@ private:
 class FWriteScopeLock
 {
 public:
-	UE_NODISCARD explicit FWriteScopeLock(FRWLock& InLock)
+	explicit FWriteScopeLock(FRWLock& InLock)
 		: Lock(InLock)
 	{
 		Lock.WriteLock();
@@ -67,7 +67,7 @@ enum FRWScopeLockType
 class FRWScopeLock
 {
 public:
-	UE_NODISCARD explicit FRWScopeLock(FRWLock& InLockObject,FRWScopeLockType InLockType)
+	explicit FRWScopeLock(FRWLock& InLockObject,FRWScopeLockType InLockType)
 	: LockObject(InLockObject)
 	, LockType(InLockType)
 	{
