Fixed compilation for Android

diff --git a/Engine/Source/Runtime/Core/Private/Android/AndroidPlatformStackWalk.cpp b/Engine/Source/Runtime/Core/Private/Android/AndroidPlatformStackWalk.cpp
--- a/Engine/Source/Runtime/Core/Private/Android/AndroidPlatformStackWalk.cpp
+++ b/Engine/Source/Runtime/Core/Private/Android/AndroidPlatformStackWalk.cpp
@@ -82,7 +82,7 @@ bool FAndroidPlatformStackWalk::InitStackWalking()
 	}
 	else
 	{
-		unw_disable_signal_frame_test(1);
+		//unw_disable_signal_frame_test(1);
 		__android_log_print(ANDROID_LOG_DEBUG, "UE", "XOM has been detected");
 	}
 
