Adding FStringView constructor to FStringFormatArg.

diff --git a/Engine/Source/Runtime/Core/Private/Misc/StringFormatter.cpp b/Engine/Source/Runtime/Core/Private/Misc/StringFormatter.cpp
--- a/Engine/Source/Runtime/Core/Private/Misc/StringFormatter.cpp
+++ b/Engine/Source/Runtime/Core/Private/Misc/StringFormatter.cpp
@@ -13,6 +13,7 @@ FStringFormatArg::FStringFormatArg( const uint64 Value ) : Type(UInt), UIntValue
 FStringFormatArg::FStringFormatArg( const float Value ) : Type(Double), DoubleValue(Value) {}
 FStringFormatArg::FStringFormatArg( const double Value ) : Type(Double), DoubleValue(Value) {}
 FStringFormatArg::FStringFormatArg( FString Value ) : Type(String), StringValue(MoveTemp(Value)) {}
+FStringFormatArg::FStringFormatArg( FStringView Value ) : Type(String), StringValue(Value) {}
 FStringFormatArg::FStringFormatArg( const ANSICHAR* Value ) : Type(StringLiteralANSI), StringLiteralANSIValue(Value) {}
 FStringFormatArg::FStringFormatArg( const WIDECHAR* Value ) : Type(StringLiteralWIDE), StringLiteralWIDEValue(Value) {}
 FStringFormatArg::FStringFormatArg( const UCS2CHAR* Value ) : Type(StringLiteralUCS2), StringLiteralUCS2Value(Value) {}
diff --git a/Engine/Source/Runtime/Core/Public/Misc/StringFormatArg.h b/Engine/Source/Runtime/Core/Public/Misc/StringFormatArg.h
--- a/Engine/Source/Runtime/Core/Public/Misc/StringFormatArg.h
+++ b/Engine/Source/Runtime/Core/Public/Misc/StringFormatArg.h
@@ -41,6 +41,7 @@ struct CORE_API FStringFormatArg
 	FStringFormatArg( const float Value );
 	FStringFormatArg( const double Value );
 	FStringFormatArg( FString Value );
+	FStringFormatArg( FStringView Value );
 	FStringFormatArg( const ANSICHAR* Value );
 	FStringFormatArg( const WIDECHAR* Value );
 	FStringFormatArg( const UCS2CHAR* Value );
