Box2D: IsInsideOrOn

diff --git a/Engine/Source/Runtime/Core/Public/Math/Box2D.h b/Engine/Source/Runtime/Core/Public/Math/Box2D.h
--- a/Engine/Source/Runtime/Core/Public/Math/Box2D.h
+++ b/Engine/Source/Runtime/Core/Public/Math/Box2D.h
@@ -313,6 +313,18 @@ public:
 		return ((TestPoint.X > Min.X) && (TestPoint.X < Max.X) && (TestPoint.Y > Min.Y) && (TestPoint.Y < Max.Y));
 	}
 
+	/**
+	 * Checks whether the given point is inside or on this box.
+	 *
+	 * @param Point The point to test.
+	 * @return true if point is inside or on this box, otherwise false.
+	 * @see IsInside
+	 */
+	bool IsInsideOrOn( const TVector2<T>& TestPoint ) const
+	{
+		return ((TestPoint.X >= Min.X) && (TestPoint.X <= Max.X) && (TestPoint.Y >= Min.Y) && (TestPoint.Y <= Max.Y));
+	}
+
 	/** 
 	 * Checks whether the given box is fully encapsulated by this box.
 	 * 
