Add command line argument for underrun mode

diff --git a/Engine/Source/Runtime/Core/Private/HAL/MallocStomp2.cpp b/Engine/Source/Runtime/Core/Private/HAL/MallocStomp2.cpp
--- a/Engine/Source/Runtime/Core/Private/HAL/MallocStomp2.cpp
+++ b/Engine/Source/Runtime/Core/Private/HAL/MallocStomp2.cpp
@@ -355,8 +355,9 @@ FMalloc* FMallocStomp2::OverrideIfEnabled(FMalloc* InUsedAlloc)
 		const TCHAR* CommandLine = ::GetCommandLineW();
 		FParse::Value(CommandLine, TEXT("MallocStomp2MinSize="), MinSize);
 		FParse::Value(CommandLine, TEXT("MallocStomp2MaxSize="), MaxSize);
+		bool bUseUnderrunMode = FParse::Param(CommandLine, TEXT("MallocStomp2UnderrunMode"));
 #endif
-		GMallocStomp2 = new FMallocStomp2(InUsedAlloc);
+		GMallocStomp2 = new FMallocStomp2(InUsedAlloc, bUseUnderrunMode);
 		GMallocStomp2->Init();
 		return GMallocStomp2;
 	}
