Fix undefined local in Malloc Stomp 2.

diff --git a/Engine/Source/Runtime/Core/Private/HAL/MallocStomp2.cpp b/Engine/Source/Runtime/Core/Private/HAL/MallocStomp2.cpp
--- a/Engine/Source/Runtime/Core/Private/HAL/MallocStomp2.cpp
+++ b/Engine/Source/Runtime/Core/Private/HAL/MallocStomp2.cpp
@@ -356,6 +356,8 @@ FMalloc* FMallocStomp2::OverrideIfEnabled(FMalloc* InUsedAlloc)
 		FParse::Value(CommandLine, TEXT("MallocStomp2MinSize="), MinSize);
 		FParse::Value(CommandLine, TEXT("MallocStomp2MaxSize="), MaxSize);
 		bool bUseUnderrunMode = FParse::Param(CommandLine, TEXT("MallocStomp2UnderrunMode"));
+#else
+		bool bUseUnderrunMode = false;
 #endif
 		GMallocStomp2 = new FMallocStomp2(InUsedAlloc, bUseUnderrunMode);
 		GMallocStomp2->Init();
diff --git a/Engine/Source/Runtime/Core/Public/HAL/MallocStomp2.h b/Engine/Source/Runtime/Core/Public/HAL/MallocStomp2.h
--- a/Engine/Source/Runtime/Core/Public/HAL/MallocStomp2.h
+++ b/Engine/Source/Runtime/Core/Public/HAL/MallocStomp2.h
@@ -3,6 +3,7 @@
 #pragma once
 #include "CoreTypes.h"
 #include "HAL/MemoryBase.h"
+#include "HAL/PlatformTLS.h"
 
 #ifndef WITH_MALLOC_STOMP2
 #define WITH_MALLOC_STOMP2 0
