fixed TSAN warning complaining about concurrent access to FTextLocalizationManager::InitializedFlags. Make it atomic as it doesn't look this can't affect overall performance

diff --git a/Engine/Source/Runtime/Core/Private/Internationalization/TextLocalizationManager.cpp b/Engine/Source/Runtime/Core/Private/Internationalization/TextLocalizationManager.cpp
--- a/Engine/Source/Runtime/Core/Private/Internationalization/TextLocalizationManager.cpp
+++ b/Engine/Source/Runtime/Core/Private/Internationalization/TextLocalizationManager.cpp
@@ -364,8 +364,9 @@ void InitEngineTextLocalization()
 	// Setting InitializedFlags to None ensures we don't pick up the culture change
 	// notification if ApplyDefaultCultureSettings changes the default culture
 	{
-		TGuardValue<ETextLocalizationManagerInitializedFlags> InitializedFlagsGuard(FTextLocalizationManager::Get().InitializedFlags, ETextLocalizationManagerInitializedFlags::None);
+		ETextLocalizationManagerInitializedFlags OldFlags = FTextLocalizationManager::Get().InitializedFlags.exchange(ETextLocalizationManagerInitializedFlags::None);
 		ApplyDefaultCultureSettings(ApplyLocLoadFlags);
+		FTextLocalizationManager::Get().InitializedFlags = OldFlags;
 	}
 
 #if WITH_EDITOR
@@ -381,7 +382,7 @@ void InitEngineTextLocalization()
 #endif
 
 	FTextLocalizationManager::Get().LoadLocalizationResourcesForCulture(FInternationalization::Get().GetCurrentLanguage()->GetName(), LocLoadFlags);
-	FTextLocalizationManager::Get().InitializedFlags |= ETextLocalizationManagerInitializedFlags::Engine;
+	FTextLocalizationManager::Get().InitializedFlags = FTextLocalizationManager::Get().InitializedFlags.load() | ETextLocalizationManagerInitializedFlags::Engine;
 }
 
 static FGraphEventRef InitGameTextLocalizationTask;
@@ -405,8 +406,9 @@ void BeginInitGameTextLocalization()
 	// notification if ApplyDefaultCultureSettings changes the default culture
 	const FString PreviousLanguage = FInternationalization::Get().GetCurrentLanguage()->GetName();
 	{
-		TGuardValue<ETextLocalizationManagerInitializedFlags> InitializedFlagsGuard(FTextLocalizationManager::Get().InitializedFlags, ETextLocalizationManagerInitializedFlags::None);
+		ETextLocalizationManagerInitializedFlags OldFlags = FTextLocalizationManager::Get().InitializedFlags.exchange(ETextLocalizationManagerInitializedFlags::None);
 		ApplyDefaultCultureSettings(ELocalizationLoadFlags::Game);
+		FTextLocalizationManager::Get().InitializedFlags = OldFlags;
 	}
 	const FString CurrentLanguage = FInternationalization::Get().GetCurrentLanguage()->GetName();
 
@@ -422,8 +424,8 @@ void BeginInitGameTextLocalization()
 		LocLoadFlags |= ELocalizationLoadFlags::Additional;
 	}
 
-	FTextLocalizationManager::Get().InitializedFlags |= ETextLocalizationManagerInitializedFlags::Initializing;
-	auto TaskLambda = [LocLoadFlags, InitializedFlags = FTextLocalizationManager::Get().InitializedFlags]()
+	FTextLocalizationManager::Get().InitializedFlags = FTextLocalizationManager::Get().InitializedFlags.load() | ETextLocalizationManagerInitializedFlags::Initializing;
+	auto TaskLambda = [LocLoadFlags, InitializedFlags = FTextLocalizationManager::Get().InitializedFlags.load()]()
 	{
 		SCOPED_BOOT_TIMING("InitGameTextLocalization");
 
@@ -995,7 +997,7 @@ void FTextLocalizationManager::OnPakFileMounted(const IPakFile& PakFile)
 	// Track this so that full resource refreshes (eg, changing culture) work as expected
 	LocResTextSource->RegisterChunkId(ChunkId);
 
-	if (!EnumHasAnyFlags(InitializedFlags, ETextLocalizationManagerInitializedFlags::Game))
+	if (!EnumHasAnyFlags(InitializedFlags.load(), ETextLocalizationManagerInitializedFlags::Game))
 	{
 		// If we've not yet initialized game localization then don't bother patching, as the full initialization path will load the data for this chunk
 		return;
diff --git a/Engine/Source/Runtime/Core/Public/Internationalization/TextLocalizationManager.h b/Engine/Source/Runtime/Core/Public/Internationalization/TextLocalizationManager.h
--- a/Engine/Source/Runtime/Core/Public/Internationalization/TextLocalizationManager.h
+++ b/Engine/Source/Runtime/Core/Public/Internationalization/TextLocalizationManager.h
@@ -17,6 +17,8 @@
 #include "Internationalization/LocKeyFuncs.h"
 #include "Internationalization/LocalizedTextSourceTypes.h"
 
+#include <atomic>
+
 struct FPolyglotTextData;
 class ILocalizedTextSource;
 class IPakFile;
@@ -77,7 +79,7 @@ private:
 	typedef TMap<FTextId, FDisplayStringEntry> FDisplayStringLookupTable;
 
 private:
-	ETextLocalizationManagerInitializedFlags InitializedFlags = ETextLocalizationManagerInitializedFlags::None;
+	std::atomic<ETextLocalizationManagerInitializedFlags> InitializedFlags{ ETextLocalizationManagerInitializedFlags::None };
 	
 	bool IsInitialized() const
 	{
@@ -86,7 +88,7 @@ private:
 
 	bool IsInitializing() const
 	{
-		return EnumHasAnyFlags(InitializedFlags, ETextLocalizationManagerInitializedFlags::Initializing);
+		return EnumHasAnyFlags(InitializedFlags.load(), ETextLocalizationManagerInitializedFlags::Initializing);
 	}
 
 	mutable FCriticalSection DisplayStringLookupTableCS;
