Add nodiscard to the scoped lock constructors to prevent bugs where the variable is not named and it's destroyed immediately, instead of at the end of the scope.

We need a separate UE_NODISCARD_CTOR macro for this, because older compilers don't like nodiscard on constructors.

#rnx

diff --git a/Engine/Source/Runtime/Core/Public/HAL/Platform.h b/Engine/Source/Runtime/Core/Public/HAL/Platform.h
--- a/Engine/Source/Runtime/Core/Public/HAL/Platform.h
+++ b/Engine/Source/Runtime/Core/Public/HAL/Platform.h
@@ -701,6 +701,20 @@
 	#define UE_NODISCARD
 #endif
 
+// Use before a constructor declaration to warn when an unnamed temporary object is ignored, e.g. FScopeLock(CS); instead of FScopeLock Lock(CS);
+// We can't just use UE_NODISCARD in this case because older compilers don't like [[nodiscard]] on constructors.
+#if !defined(UE_NODISCARD_CTOR) && defined(__has_cpp_attribute)
+	#if __has_cpp_attribute(nodiscard)
+		#if (defined(_MSC_VER) && _MSC_VER >= 1924) || (defined(__clang__) && __clang_major__ >= 10)
+			#define UE_NODISCARD_CTOR [[nodiscard]]
+		#endif
+	#endif
+#endif
+
+#ifndef UE_NODISCARD_CTOR
+	#define UE_NODISCARD_CTOR
+#endif
+
 /* Use before a function declaration to indicate that the function never returns */
 #if !defined(UE_NORETURN) && defined(__has_cpp_attribute)
 	#if __has_cpp_attribute(noreturn)
diff --git a/Engine/Source/Runtime/Core/Public/Misc/ScopeLock.h b/Engine/Source/Runtime/Core/Public/Misc/ScopeLock.h
--- a/Engine/Source/Runtime/Core/Public/Misc/ScopeLock.h
+++ b/Engine/Source/Runtime/Core/Public/Misc/ScopeLock.h
@@ -32,7 +32,7 @@ public:
 	 *
 	 * @param InSynchObject The synchronization object to manage
 	 */
-	FScopeLock( FCriticalSection* InSynchObject )
+	UE_NODISCARD_CTOR FScopeLock( FCriticalSection* InSynchObject )
 		: SynchObject(InSynchObject)
 	{
 		check(SynchObject);
@@ -99,7 +99,7 @@ public:
 	 *
 	 * @param InSynchObject The synchronization object to manage, can be null.
 	 */
-	FScopeUnlock(FCriticalSection* InSynchObject)
+	UE_NODISCARD_CTOR FScopeUnlock(FCriticalSection* InSynchObject)
 		: SynchObject(InSynchObject)
 	{
 		if (InSynchObject)
@@ -151,7 +151,7 @@ namespace UE
 	public:
 		UE_NONCOPYABLE(TScopeLock);
 
-		TScopeLock(MutexType& InMutex)
+		UE_NODISCARD_CTOR TScopeLock(MutexType& InMutex)
 			: Mutex(&InMutex)
 		{
 			check(Mutex);
@@ -193,7 +193,7 @@ namespace UE
 	public:
 		UE_NONCOPYABLE(TScopeUnlock);
 
-		TScopeUnlock(MutexType& InMutex)
+		UE_NODISCARD_CTOR TScopeUnlock(MutexType& InMutex)
 			: Mutex(&InMutex)
 		{
 			check(Mutex);
diff --git a/Engine/Source/Runtime/Core/Public/Misc/ScopeRWLock.h b/Engine/Source/Runtime/Core/Public/Misc/ScopeRWLock.h
--- a/Engine/Source/Runtime/Core/Public/Misc/ScopeRWLock.h
+++ b/Engine/Source/Runtime/Core/Public/Misc/ScopeRWLock.h
@@ -11,7 +11,7 @@
 class FReadScopeLock
 {
 public:
-	explicit FReadScopeLock(FRWLock& InLock)
+	UE_NODISCARD_CTOR explicit FReadScopeLock(FRWLock& InLock)
 		: Lock(InLock)
 	{
 		Lock.ReadLock();
@@ -32,7 +32,7 @@ private:
 class FWriteScopeLock
 {
 public:
-	explicit FWriteScopeLock(FRWLock& InLock)
+	UE_NODISCARD_CTOR explicit FWriteScopeLock(FRWLock& InLock)
 		: Lock(InLock)
 	{
 		Lock.WriteLock();
@@ -67,7 +67,7 @@ enum FRWScopeLockType
 class FRWScopeLock
 {
 public:
-	explicit FRWScopeLock(FRWLock& InLockObject,FRWScopeLockType InLockType)
+	UE_NODISCARD_CTOR explicit FRWScopeLock(FRWLock& InLockObject,FRWScopeLockType InLockType)
 	: LockObject(InLockObject)
 	, LockType(InLockType)
 	{
