Tasks: added "destroyed" trace to the old frontend

diff --git a/Engine/Source/Runtime/Core/Private/Async/TaskGraph.cpp b/Engine/Source/Runtime/Core/Private/Async/TaskGraph.cpp
--- a/Engine/Source/Runtime/Core/Private/Async/TaskGraph.cpp
+++ b/Engine/Source/Runtime/Core/Private/Async/TaskGraph.cpp
@@ -2559,6 +2559,8 @@ FGraphEvent::~FGraphEvent()
 	}
 #endif
 	CheckDontCompleteUntilIsEmpty(); // We should not have any wait untils outstanding
+
+	TaskTrace::Destroyed(GetTraceId());
 }
 
 #endif
