Fix some thread heartbeat suspend count issues:

1. We will never do resume logic if bAllThreads is false since the global count won't be incremented.
2. If a checkpoint is added while a heartbeat suspend object is active, it will attempt to decrement the checkpoint count below zero.

diff --git a/Engine/Source/Runtime/Core/Private/HAL/ThreadHeartBeat.cpp b/Engine/Source/Runtime/Core/Private/HAL/ThreadHeartBeat.cpp
--- a/Engine/Source/Runtime/Core/Private/HAL/ThreadHeartBeat.cpp
+++ b/Engine/Source/Runtime/Core/Private/HAL/ThreadHeartBeat.cpp
@@ -94,6 +94,8 @@ FThreadHeartBeat::FThreadHeartBeat()
 	, LastHungThreadId(InvalidThreadId)
 	, LastStuckThreadId(InvalidThreadId)
 	, bHangsAreFatal(false)
+	, GlobalSuspendCount(0)
+	, CheckpointSuspendCount(0)
 	, Clock(HangDetectorClock_MaxTimeStep_MS / 1000)
 {
 	// Start with the frame-present based hang detection disabled. This will be automatically enabled on
@@ -692,7 +694,7 @@ void FThreadHeartBeat::MonitorCheckpointStart(FName EndCheckpoint, double TimeTo
 		HeartBeatInfo.LastHeartBeatTime = Clock.Seconds();
 		HeartBeatInfo.HangDuration = TimeToReachCheckpoint;
 		HeartBeatInfo.HeartBeatName = EndCheckpoint;
-		HeartBeatInfo.SuspendedCount = 0;
+		HeartBeatInfo.SuspendedCount = CheckpointSuspendCount;
 	}
 #endif
 }
@@ -787,9 +789,10 @@ void FThreadHeartBeat::SuspendHeartBeat(bool bAllThreads)
 
 	// Suspend the checkpoint heartbeats
 	{
-		FScopeLock HeartBeatLock(&CheckpointHeartBeatCritical);
 		if (!bAllThreads)
 		{
+			FScopeLock HeartBeatLock(&CheckpointHeartBeatCritical);
+			++CheckpointSuspendCount;
 			for (TPair<FName, FHeartBeatInfo>& HeartBeatEntry : CheckpointHeartBeat)
 			{
 				HeartBeatEntry.Value.Suspend();
@@ -806,15 +809,16 @@ void FThreadHeartBeat::ResumeHeartBeat(bool bAllThreads)
 	bool bLastThreadResumed = false;
 	{
 		FScopeLock HeartBeatLock(&HeartBeatCritical);
-		if (GlobalSuspendCount.GetValue() == 0)
-		{
-			// Resume without matching Suspend, ignore it
-			return;
-		}
-		
+
 		const double CurrentTime = Clock.Seconds();
 		if (bAllThreads)
 		{
+			if (GlobalSuspendCount.GetValue() == 0)
+			{
+				// Resume without matching Suspend, ignore it
+				return;
+			}
+		
 			if (GlobalSuspendCount.Decrement() == 0)
 			{
 				bLastThreadResumed = true;
@@ -840,8 +844,8 @@ void FThreadHeartBeat::ResumeHeartBeat(bool bAllThreads)
 
 	// Resume the checkpoint heartbeats
 	{
-		FScopeLock HeartBeatLock(&CheckpointHeartBeatCritical);
 		const double CurrentTime = Clock.Seconds();
+		FScopeLock HeartBeatLock(&CheckpointHeartBeatCritical);
 		if (bAllThreads)
 		{
 			if (bLastThreadResumed)
@@ -854,6 +858,8 @@ void FThreadHeartBeat::ResumeHeartBeat(bool bAllThreads)
 		}
 		else
 		{
+			check(CheckpointSuspendCount > 0);
+			CheckpointSuspendCount--;
 			for (TPair<FName, FHeartBeatInfo>& HeartBeatEntry : CheckpointHeartBeat)
 			{
 				HeartBeatEntry.Value.Resume(CurrentTime);
diff --git a/Engine/Source/Runtime/Core/Public/HAL/ThreadHeartBeat.h b/Engine/Source/Runtime/Core/Public/HAL/ThreadHeartBeat.h
--- a/Engine/Source/Runtime/Core/Public/HAL/ThreadHeartBeat.h
+++ b/Engine/Source/Runtime/Core/Public/HAL/ThreadHeartBeat.h
@@ -137,8 +137,8 @@ class CORE_API FThreadHeartBeat : public FRunnable
 
 	bool bHangsAreFatal;
 
-	/** Global suspended count */
 	FThreadSafeCounter GlobalSuspendCount;
+	int32 CheckpointSuspendCount;
 
 	FThreadHeartBeatClock Clock;
 
