Resolve Windows typecast warning in clang UnrealGame

#rnx

diff --git a/Engine/Source/Runtime/Core/Private/Windows/WindowsPlatformMemory.cpp b/Engine/Source/Runtime/Core/Private/Windows/WindowsPlatformMemory.cpp
--- a/Engine/Source/Runtime/Core/Private/Windows/WindowsPlatformMemory.cpp
+++ b/Engine/Source/Runtime/Core/Private/Windows/WindowsPlatformMemory.cpp
@@ -73,13 +73,13 @@ void FWindowsPlatformMemory::Init()
 	const FPlatformMemoryConstants& MemoryConstants = FPlatformMemory::GetConstants();
 #if PLATFORM_32BITS	
 	UE_LOG(LogMemory, Log, TEXT("Memory total: Physical=%.1fGB (%dGB approx) Virtual=%.1fGB"), 
-		float(MemoryConstants.TotalPhysical/1024.0/1024.0/1024.0),
+		(double)MemoryConstants.TotalPhysical/1024.0/1024.0/1024.0,
 		MemoryConstants.TotalPhysicalGB, 
-		float(MemoryConstants.TotalVirtual/1024.0/1024.0/1024.0) );
+		(double)MemoryConstants.TotalVirtual/1024.0/1024.0/1024.0 );
 #else
 	// Logging virtual memory size for 64bits is pointless.
 	UE_LOG(LogMemory, Log, TEXT("Memory total: Physical=%.1fGB (%dGB approx)"), 
-		float(MemoryConstants.TotalPhysical/1024.0/1024.0/1024.0),
+		(double)MemoryConstants.TotalPhysical/1024.0/1024.0/1024.0,
 		MemoryConstants.TotalPhysicalGB );
 #endif //PLATFORM_32BITS
 
diff --git a/Engine/Source/Runtime/Core/Private/Windows/WindowsPlatformTime.cpp b/Engine/Source/Runtime/Core/Private/Windows/WindowsPlatformTime.cpp
--- a/Engine/Source/Runtime/Core/Private/Windows/WindowsPlatformTime.cpp
+++ b/Engine/Source/Runtime/Core/Private/Windows/WindowsPlatformTime.cpp
@@ -28,8 +28,8 @@ double FWindowsPlatformTime::InitTiming(void)
 {
 	LARGE_INTEGER Frequency;
 	verify( QueryPerformanceFrequency(&Frequency) );
-	SecondsPerCycle = 1.0 / Frequency.QuadPart;
-	SecondsPerCycle64 = 1.0 / Frequency.QuadPart;
+	SecondsPerCycle = 1.0 / (double)Frequency.QuadPart;
+	SecondsPerCycle64 = 1.0 / (double)Frequency.QuadPart;
 
 	// Due to some limitation of the OS, we limit the polling frequency to 4 times per second, 
 	// but it should be enough for longterm CPU usage monitoring.
diff --git a/Engine/Source/Runtime/Core/Public/Windows/WindowsPlatformTime.h b/Engine/Source/Runtime/Core/Public/Windows/WindowsPlatformTime.h
--- a/Engine/Source/Runtime/Core/Public/Windows/WindowsPlatformTime.h
+++ b/Engine/Source/Runtime/Core/Public/Windows/WindowsPlatformTime.h
@@ -26,7 +26,7 @@ struct CORE_API FWindowsPlatformTime
 		Windows::QueryPerformanceCounter(&Cycles);
 
 		// add big number to make bugs apparent where return value is being passed to float
-		return Cycles.QuadPart * GetSecondsPerCycle() + 16777216.0;
+		return (double)Cycles.QuadPart * GetSecondsPerCycle() + 16777216.0;
 	}
 
 	static FORCEINLINE uint32 Cycles()
