Resolve typecast warning in Math/Rotator.h

#rnx

diff --git a/Engine/Source/Runtime/Core/Public/Math/Rotator.h b/Engine/Source/Runtime/Core/Public/Math/Rotator.h
--- a/Engine/Source/Runtime/Core/Public/Math/Rotator.h
+++ b/Engine/Source/Runtime/Core/Public/Math/Rotator.h
@@ -684,7 +684,7 @@ template<typename T>
 FORCEINLINE uint8 TRotator<T>::CompressAxisToByte( T Angle )
 {
 	// map [0->360) to [0->256) and mask off any winding
-	return FMath::RoundToInt(Angle * 256.f / 360.f) & 0xFF;
+	return FMath::RoundToInt(Angle * (T)256.f / (T)360.f) & 0xFF;
 }
 
 template<typename T>
