Resolve typecast warning in ProfilingDebugging

#rnx

diff --git a/Engine/Source/Runtime/Core/Private/ProfilingDebugging/StallDetector.cpp b/Engine/Source/Runtime/Core/Private/ProfilingDebugging/StallDetector.cpp
--- a/Engine/Source/Runtime/Core/Private/ProfilingDebugging/StallDetector.cpp
+++ b/Engine/Source/Runtime/Core/Private/ProfilingDebugging/StallDetector.cpp
@@ -118,7 +118,7 @@ uint32 UE::FStallDetectorRunnable::Run()
 		}
 
 		// Sleep an interval, the resolution at which we want to detect an overage
-		FPlatformProcess::SleepNoStats(0.005);
+		FPlatformProcess::SleepNoStats(0.005f);
 	}
 
 	return 0;
@@ -184,7 +184,7 @@ void UE::FStallDetectorStats::TabulateStats(TArray<TabulatedResult>& TabulatedRe
 			FScopeLock Lock(&InStats->StatsSection);
 			if (InStats->TriggerCount.Get() && InStats->BudgetSeconds > 0.0)
 			{
-				OverageRatio = (InStats->OverageSeconds.Get() / InStats->TriggerCount.Get()) / InStats->BudgetSeconds;
+				OverageRatio = (InStats->OverageSeconds.Get() / (double)InStats->TriggerCount.Get()) / InStats->BudgetSeconds;
 			}
 		}
 
@@ -586,7 +586,7 @@ void UE::FStallDetector::Startup()
 	{
 		UE_LOG(LogStall, Log, TEXT("Startup..."));
 
-		check(FPlatformTime::GetSecondsPerCycle());
+		check(FPlatformTime::GetSecondsPerCycle() > 0.0);
 
 		// Cannot be a global due to clock member
 		Runnable = new FStallDetectorRunnable();
diff --git a/Engine/Source/Runtime/Core/Public/ProfilingDebugging/CookStats.h b/Engine/Source/Runtime/Core/Public/ProfilingDebugging/CookStats.h
--- a/Engine/Source/Runtime/Core/Public/ProfilingDebugging/CookStats.h
+++ b/Engine/Source/Runtime/Core/Public/ProfilingDebugging/CookStats.h
@@ -248,12 +248,12 @@ public:
 						TEXT("Call"), CallName,
 						TEXT("HitOrMiss"), HitOrMiss == EHitOrMiss::Hit ? TEXT("Hit") : TEXT("Miss"),
 						TEXT("Count"), GetAccumulatedValue(HitOrMiss, EStatType::Counter, bGameThread),
-						TEXT("TimeSec"), GetAccumulatedValue(HitOrMiss, EStatType::Cycles, bGameThread) * FPlatformTime::GetSecondsPerCycle(),
-						TEXT("MB"), GetAccumulatedValue(HitOrMiss, EStatType::Bytes, bGameThread) / (1024.0 * 1024.0),
+						TEXT("TimeSec"), (double)GetAccumulatedValue(HitOrMiss, EStatType::Cycles, bGameThread) * FPlatformTime::GetSecondsPerCycle(),
+						TEXT("MB"), (double)GetAccumulatedValue(HitOrMiss, EStatType::Bytes, bGameThread) / (1024.0 * 1024.0),
 						TEXT("MB/s"), GetAccumulatedValue(HitOrMiss, EStatType::Cycles, bGameThread) == 0
 							? 0.0
-							: GetAccumulatedValue(HitOrMiss, EStatType::Bytes, bGameThread) / (1024.0 * 1024.0) /
-							 (GetAccumulatedValue(HitOrMiss, EStatType::Cycles, bGameThread) * FPlatformTime::GetSecondsPerCycle()),
+							: (double)GetAccumulatedValue(HitOrMiss, EStatType::Bytes, bGameThread) / (1024.0 * 1024.0) /
+							 ((double)GetAccumulatedValue(HitOrMiss, EStatType::Cycles, bGameThread) * FPlatformTime::GetSecondsPerCycle()),
 						TEXT("Node"), NodeName
 						));
 				}
