[Insights] Fixed initialisation of Platform Events (context switches and sampling profiling) by adding a PostInit call to enable functionality based on enabled channels after the system is initialized.

diff --git a/Engine/Source/Runtime/Core/Private/ProfilingDebugging/PlatformEvents.cpp b/Engine/Source/Runtime/Core/Private/ProfilingDebugging/PlatformEvents.cpp
--- a/Engine/Source/Runtime/Core/Private/ProfilingDebugging/PlatformEvents.cpp
+++ b/Engine/Source/Runtime/Core/Private/ProfilingDebugging/PlatformEvents.cpp
@@ -29,3 +29,23 @@ EPlatformEvent PlatformEvents_GetEvent(const FString& Name)
 		return EPlatformEvent::None;
 	}
 }
+
+/////////////////////////////////////////////////////////////////////
+
+#if PLATFORM_SUPPORTS_PLATFORM_EVENTS
+
+void PlatformEvents_PostInit()
+{
+	if (TRACE_PRIVATE_CHANNELEXPR_IS_ENABLED(ContextSwitchChannel))
+	{
+		PlatformEvents_Enable(EPlatformEvent::ContextSwitch);
+	}
+	else if (TRACE_PRIVATE_CHANNELEXPR_IS_ENABLED(StackSamplingChannel))
+	{
+		PlatformEvents_Enable(EPlatformEvent::StackSampling);
+	}
+}
+
+#endif // PLATFORM_SUPPORTS_PLATFORM_EVENTS
+
+/////////////////////////////////////////////////////////////////////
diff --git a/Engine/Source/Runtime/Core/Private/ProfilingDebugging/PlatformEvents.h b/Engine/Source/Runtime/Core/Private/ProfilingDebugging/PlatformEvents.h
--- a/Engine/Source/Runtime/Core/Private/ProfilingDebugging/PlatformEvents.h
+++ b/Engine/Source/Runtime/Core/Private/ProfilingDebugging/PlatformEvents.h
@@ -62,6 +62,7 @@ EPlatformEvent PlatformEvents_GetEvent(const FString& Name);
 #if PLATFORM_SUPPORTS_PLATFORM_EVENTS
 
 void PlatformEvents_Init(uint32 SamplingIntervalUsec);
+void PlatformEvents_PostInit();
 void PlatformEvents_Enable(EPlatformEvent Event);
 void PlatformEvents_Disable(EPlatformEvent Event);
 void PlatformEvents_Stop();
@@ -69,6 +70,7 @@ void PlatformEvents_Stop();
 #else // PLATFORM_SUPPORTS_PLATFORM_EVENTS
 
 static void PlatformEvents_Init(uint32 SamplingIntervalUsec) {}
+static void PlatformEvents_PostInit() {}
 static void PlatformEvents_Enable(EPlatformEvent Event) {}
 static void PlatformEvents_Disable(EPlatformEvent Event) {}
 static void PlatformEvents_Stop() {}
diff --git a/Engine/Source/Runtime/Core/Private/ProfilingDebugging/TraceAuxiliary.cpp b/Engine/Source/Runtime/Core/Private/ProfilingDebugging/TraceAuxiliary.cpp
--- a/Engine/Source/Runtime/Core/Private/ProfilingDebugging/TraceAuxiliary.cpp
+++ b/Engine/Source/Runtime/Core/Private/ProfilingDebugging/TraceAuxiliary.cpp
@@ -1348,6 +1348,7 @@ void FTraceAuxiliary::Initialize(const TCHAR* CommandLine)
 	uint32 Microseconds = 1000;
 	FParse::Value(CommandLine, TEXT("-samplinginterval="), Microseconds);
 	PlatformEvents_Init(Microseconds);
+	PlatformEvents_PostInit();
 
 #if CSV_PROFILER
 	FCoreDelegates::OnEndFrame.AddRaw(&GTraceAuxiliary, &FTraceAuxiliaryImpl::UpdateCsvStats);
