UE_SECURITY_LOG deprecated and fixed to work in UTF-8 mode.
CLOSE_CONNECTION_DUE_TO_SECURITY_VIOLATION deprecated.

diff --git a/Engine/Plugins/Online/OnlineSubsystemUtils/Source/OnlineSubsystemUtils/Private/IpNetDriver.cpp b/Engine/Plugins/Online/OnlineSubsystemUtils/Source/OnlineSubsystemUtils/Private/IpNetDriver.cpp
--- a/Engine/Plugins/Online/OnlineSubsystemUtils/Source/OnlineSubsystemUtils/Private/IpNetDriver.cpp
+++ b/Engine/Plugins/Online/OnlineSubsystemUtils/Source/OnlineSubsystemUtils/Private/IpNetDriver.cpp
@@ -1270,7 +1270,7 @@ void UIpNetDriver::TickDispatch(float DeltaTime)
 
 					if (Connection != nullptr)
 					{
-						UE_SECURITY_LOG(Connection, ESecurityEvent::Malformed_Packet, TEXT("Received Packet with bytes > max MTU"));
+						UE_LOG(LogSecurity, Warning, TEXT("%s: Malformed_Packet: Received Packet with bytes > max MTU"), *Connection->RemoteAddressToString());
 					}
 				}
 				else
diff --git a/Engine/Source/Runtime/Core/Public/Logging/LogMacros.h b/Engine/Source/Runtime/Core/Public/Logging/LogMacros.h
--- a/Engine/Source/Runtime/Core/Public/Logging/LogMacros.h
+++ b/Engine/Source/Runtime/Core/Public/Logging/LogMacros.h
@@ -136,7 +136,7 @@ private:
 	#define DEFINE_LOG_CATEGORY_STATIC(...)
 	#define DECLARE_LOG_CATEGORY_CLASS(...)
 	#define DEFINE_LOG_CATEGORY_CLASS(...)
-	#define UE_SECURITY_LOG(...)
+	#define UE_SECURITY_LOG(...) DEPRECATED_MACRO(5.1, "UE_SECURITY_LOG has been deprecated in favor of UE_LOG")
 
 #else
 
@@ -275,13 +275,14 @@ private:
 	***/
 	#define UE_SECURITY_LOG(NetConnection, SecurityEventType, Format, ...) \
 	{ \
+		DEPRECATED_MACRO(5.1, "UE_SECURITY_LOG has been deprecated in favor of UE_LOG") \
 		static_assert(TIsArrayOrRefOfType<decltype(Format), TCHAR>::Value, "Formatting string must be a TCHAR array."); \
 		check(NetConnection != nullptr); \
 		CA_CONSTANT_IF((ELogVerbosity::Warning & ELogVerbosity::VerbosityMask) <= ELogVerbosity::COMPILED_IN_MINIMUM_VERBOSITY && (ELogVerbosity::Warning & ELogVerbosity::VerbosityMask) <= FLogCategoryLogSecurity::CompileTimeVerbosity) \
 		{ \
 			if (!LogSecurity.IsSuppressed(ELogVerbosity::Warning)) \
 			{ \
-				FMsg::Logf_Internal(UE_LOG_SOURCE_FILE(__FILE__), __LINE__, LogSecurity.GetCategoryName(), ELogVerbosity::Warning, TEXT("%s: %s: ") Format, *(NetConnection->RemoteAddressToString()), ToString(SecurityEventType), ##__VA_ARGS__); \
+				FMsg::Logf_Internal(UE_LOG_SOURCE_FILE(__FILE__), __LINE__, LogSecurity.GetCategoryName(), ELogVerbosity::Warning, TEXT("%s: %s: %s"), *(NetConnection->RemoteAddressToString()), ToString(SecurityEventType), *FString::Printf(Format, ##__VA_ARGS__)); \
 			} \
 		} \
 	}
@@ -403,11 +404,12 @@ private:
 ***/
 #define CLOSE_CONNECTION_DUE_TO_SECURITY_VIOLATION_INNER(NetConnection, SecurityEventType, Format, ...) \
 { \
+	DEPRECATED_MACRO(5.1, "CLOSE_CONNECTION_DUE_TO_SECURITY_VIOLATION has been deprecated") \
 	static_assert(TIsArrayOrRefOfType<decltype(Format), TCHAR>::Value, "Formatting string must be a TCHAR array."); \
 	check(NetConnection != nullptr); \
 	FString SecurityPrint = FString::Printf(Format, ##__VA_ARGS__); \
-	UE_SECURITY_LOG(NetConnection, SecurityEventType, Format, ##__VA_ARGS__); \
-	UE_SECURITY_LOG(NetConnection, ESecurityEvent::Closed, TEXT("Connection closed")); \
+	UE_LOG(LogSecurity, Warning, TEXT("%s: %s: %s"), *(NetConnection)->RemoteAddressToString(), ToString(SecurityEventType), SecurityEventType, *SecurityPrint); \
+	UE_LOG(LogSecurity, Warning, TEXT("%s: Closed: Connection closed"), *(NetConnection)->RemoteAddressToString()); \
 	NetConnection->Close({FromSecurityEvent(SecurityEventType), SecurityPrint}); \
 }
 #if USE_SERVER_PERF_COUNTERS
diff --git a/Engine/Source/Runtime/Engine/Private/DataChannel.cpp b/Engine/Source/Runtime/Engine/Private/DataChannel.cpp
--- a/Engine/Source/Runtime/Engine/Private/DataChannel.cpp
+++ b/Engine/Source/Runtime/Engine/Private/DataChannel.cpp
@@ -1753,7 +1753,7 @@ void UControlChannel::ReceivedBunch( FInBunch& Bunch )
 			FString DebugMessage;
 			if (FNetControlMessage<NMT_SecurityViolation>::Receive(Bunch, DebugMessage))
 			{
-				UE_SECURITY_LOG(Connection, ESecurityEvent::Closed, TEXT("%s"), *DebugMessage);
+				UE_LOG(LogSecurity, Warning, TEXT("%s: Closed: %s"), *Connection->RemoteAddressToString(), *DebugMessage);
 				break;
 			}
 		}
