added `Peek` to TSpscQueue and TMpscQueue

diff --git a/Engine/Source/Runtime/Core/Private/Tests/Containers/ConcurrentQueuesTest.cpp b/Engine/Source/Runtime/Core/Private/Tests/Containers/ConcurrentQueuesTest.cpp
--- a/Engine/Source/Runtime/Core/Private/Tests/Containers/ConcurrentQueuesTest.cpp
+++ b/Engine/Source/Runtime/Core/Private/Tests/Containers/ConcurrentQueuesTest.cpp
@@ -404,6 +404,43 @@ namespace UE { namespace ConcurrentQueuesTests
 			}
 		}
 
+		{	// test `Peek()`
+			struct FMovableOnly
+			{
+				explicit FMovableOnly(int32 InId)
+					: Id(InId)
+				{
+				}
+
+				FMovableOnly(FMovableOnly&&) = default;
+				FMovableOnly& operator=(FMovableOnly&&) = default;
+
+				int32 Id;
+			};
+
+			{
+				TSpscQueue<FMovableOnly> Q;
+				Q.Enqueue(42);
+				FMovableOnly* Peeked = Q.Peek();
+				check(Peeked);
+				check(Peeked->Id == 42);
+				TOptional<FMovableOnly> Dequeued = Q.Dequeue();
+				check(Dequeued);
+				check(Dequeued.GetValue().Id == Peeked->Id);
+			}
+
+			{
+				TMpscQueue<FMovableOnly> Q;
+				Q.Enqueue(42);
+				FMovableOnly* Peeked = Q.Peek();
+				check(Peeked);
+				check(Peeked->Id == 42);
+				TOptional<FMovableOnly> Dequeued = Q.Dequeue();
+				check(Dequeued);
+				check(Dequeued.GetValue().Id == Peeked->Id);
+			}
+		}
+
 		UE_BENCHMARK(5, TestTCircularQueueSingleThread<5'000'000>);
 		UE_BENCHMARK(5, TestQueueSingleThread<5'000'000, TQueueAdapter<TQueue<uint32, EQueueMode::Spsc>>>);
 		UE_BENCHMARK(5, TestQueueSingleThread<5'000'000, TQueueAdapter<TQueue<uint32, EQueueMode::Mpsc>>>);
diff --git a/Engine/Source/Runtime/Core/Public/Containers/MpscQueue.h b/Engine/Source/Runtime/Core/Public/Containers/MpscQueue.h
--- a/Engine/Source/Runtime/Core/Public/Containers/MpscQueue.h
+++ b/Engine/Source/Runtime/Core/Public/Containers/MpscQueue.h
@@ -85,6 +85,21 @@ public:
 		return false;
 	}
 
+	// as there can be only one consumer, a consumer can safely "peek" the tail of the queue.
+	// returns a pointer to the tail if the queue is not empty, nullptr otherwise
+	// there's no overload with TOptional as it doesn't support references
+	ElementType* Peek() const
+	{
+		FNode* Next = Tail->Next.load(std::memory_order_acquire);
+
+		if (Next == nullptr)
+		{
+			return nullptr;
+		}
+
+		return (ElementType*)&Next->Value;
+	}
+
 	bool IsEmpty() const
 	{
 		return Tail->Next.load(std::memory_order_acquire) != nullptr;
diff --git a/Engine/Source/Runtime/Core/Public/Containers/SpscQueue.h b/Engine/Source/Runtime/Core/Public/Containers/SpscQueue.h
--- a/Engine/Source/Runtime/Core/Public/Containers/SpscQueue.h
+++ b/Engine/Source/Runtime/Core/Public/Containers/SpscQueue.h
@@ -99,6 +99,22 @@ public:
 		return LocalTailNext != nullptr;
 	}
 
+	// as there can be only one consumer, a consumer can safely "peek" the tail of the queue.
+	// returns a pointer to the tail if the queue is not empty, nullptr otherwise
+	// there's no overload with TOptional as it doesn't support references
+	ElementType* Peek() const
+	{
+		FNode* LocalTail = Tail.load(std::memory_order_relaxed);
+		FNode* LocalTailNext = LocalTail->Next.load(std::memory_order_acquire);
+
+		if (LocalTailNext == nullptr)
+		{
+			return nullptr;
+		}
+
+		return (ElementType*)&LocalTailNext->Value;
+	}
+
 private:
 	struct FNode
 	{
