fixed IsEmpty() in T[Spsc/Mpsc]Queue

diff --git a/Engine/Source/Runtime/Core/Private/Tests/Containers/ConcurrentQueuesTest.cpp b/Engine/Source/Runtime/Core/Private/Tests/Containers/ConcurrentQueuesTest.cpp
--- a/Engine/Source/Runtime/Core/Private/Tests/Containers/ConcurrentQueuesTest.cpp
+++ b/Engine/Source/Runtime/Core/Private/Tests/Containers/ConcurrentQueuesTest.cpp
@@ -333,8 +333,8 @@ namespace UE { namespace ConcurrentQueuesTests
 			// SPSC
 
 			{	// destroy queue while it's holding one unconsumed item
-				TSpscQueue<FNonTrivial> Q;
-				Q.Enqueue(MakeUnique<int>(1));
+			TSpscQueue<FNonTrivial> Q;
+			Q.Enqueue(MakeUnique<int>(1));
 			}
 
 			{	// destroy queue while it's holding one cached consumed time
@@ -404,6 +404,24 @@ namespace UE { namespace ConcurrentQueuesTests
 			}
 		}
 
+		{ // test `TSpscQueue::IsEmpty()`
+			TSpscQueue<int> Q;
+			check(Q.IsEmpty());
+			Q.Enqueue(42);
+			check(!Q.IsEmpty());
+			Q.Dequeue();
+			check(Q.IsEmpty());
+		}
+
+		{ // test `TMpscQueue::IsEmpty()`
+			TMpscQueue<int> Q;
+			check(Q.IsEmpty());
+			Q.Enqueue(42);
+			check(!Q.IsEmpty());
+			Q.Dequeue();
+			check(Q.IsEmpty());
+		}
+
 		{	// test `Peek()`
 			struct FMovableOnly
 			{
diff --git a/Engine/Source/Runtime/Core/Public/Containers/MpscQueue.h b/Engine/Source/Runtime/Core/Public/Containers/MpscQueue.h
--- a/Engine/Source/Runtime/Core/Public/Containers/MpscQueue.h
+++ b/Engine/Source/Runtime/Core/Public/Containers/MpscQueue.h
@@ -102,7 +102,7 @@ public:
 
 	bool IsEmpty() const
 	{
-		return Tail->Next.load(std::memory_order_acquire) != nullptr;
+		return Tail->Next.load(std::memory_order_acquire) == nullptr;
 	}
 
 private:
diff --git a/Engine/Source/Runtime/Core/Public/Containers/SpscQueue.h b/Engine/Source/Runtime/Core/Public/Containers/SpscQueue.h
--- a/Engine/Source/Runtime/Core/Public/Containers/SpscQueue.h
+++ b/Engine/Source/Runtime/Core/Public/Containers/SpscQueue.h
@@ -96,7 +96,7 @@ public:
 	{
 		FNode* LocalTail = Tail.load(std::memory_order_relaxed);
 		FNode* LocalTailNext = LocalTail->Next.load(std::memory_order_acquire);
-		return LocalTailNext != nullptr;
+		return LocalTailNext == nullptr;
 	}
 
 	// as there can be only one consumer, a consumer can safely "peek" the tail of the queue.
