Compile warning fixes for precision loss when using UnsafeTypeCast warnings

diff --git a/Engine/Source/Runtime/Core/Public/Misc/GenericGrowableAllocator.h b/Engine/Source/Runtime/Core/Public/Misc/GenericGrowableAllocator.h
--- a/Engine/Source/Runtime/Core/Public/Misc/GenericGrowableAllocator.h
+++ b/Engine/Source/Runtime/Core/Public/Misc/GenericGrowableAllocator.h
@@ -53,7 +53,7 @@ public:
 	 */
 	virtual uint64 GetWasteApproximation()
 	{
-		double Waste = ((double)TotalWaste / (double)TotalAllocs) * CurrentAllocs;
+		double Waste = (static_cast<double>(TotalWaste) / static_cast<double>(TotalAllocs)) * static_cast<double>(CurrentAllocs);
 		return (uint64)Waste;
 	}
 
@@ -855,8 +855,8 @@ private:
 	 void OutOfMemory(uint32 Size)
 	 {
 #if !UE_BUILD_SHIPPING
-		 FPlatformMisc::LowLevelOutputDebugStringf(TEXT("FGrowableAllocator: OOM allocating %dbytes %fMB"), Size, Size / 1024.0f / 1024.0f);
-		 UE_LOG(LogCore, Fatal, TEXT("FGrowableAllocator: OOM allocating %dbytes %fMB"), Size, Size / 1024.0f / 1024.0f);
+		 FPlatformMisc::LowLevelOutputDebugStringf(TEXT("FGrowableAllocator: OOM allocating %dbytes %fMB"), Size, static_cast<float>(Size) / 1024.0f / 1024.0f);
+		 UE_LOG(LogCore, Fatal, TEXT("FGrowableAllocator: OOM allocating %dbytes %fMB"), Size, static_cast<float>(Size) / 1024.0f / 1024.0f);
 #endif
 	 }
 
