Add ParallelForWithExistingTaskContext, which is similar to ParallelForWithTaskContext except takes a pre-allocated array of task contexts and launches at most one task per context

diff --git a/Engine/Source/Runtime/Core/Public/Async/ParallelFor.h b/Engine/Source/Runtime/Core/Public/Async/ParallelFor.h
--- a/Engine/Source/Runtime/Core/Public/Async/ParallelFor.h
+++ b/Engine/Source/Runtime/Core/Public/Async/ParallelFor.h
@@ -101,6 +101,13 @@ namespace ParallelForImpl
 		check(Num >= 0);
 
 		int32 NumWorkers = GetNumberOfThreadTasks(Num, MinBatchSize, Flags);
+
+		if (!Contexts.IsEmpty())
+		{
+			// Use at most as many workers as there are contexts when task contexts are used.
+			NumWorkers = FMath::Min(NumWorkers, Contexts.Num());
+		}
+
 		//single threaded mode
 		if (NumWorkers <= 1)
 		{
@@ -661,4 +668,45 @@ inline void ParallelForWithTaskContext(const TCHAR* DebugName, TArray<ContextTyp
 		OutContexts.AddDefaulted(NumContexts);
 		ParallelForImpl::ParallelForInternal(DebugName, Num, MinBatchSize, Body, [](){}, Flags, TArrayView<ContextType>(OutContexts));
 	}
-}
\ No newline at end of file
+}
+
+/**
+*	General purpose parallel for that uses the taskgraph. This variant takes an array of user-defined context
+*	objects for each task that may get spawned to do work (one task per context at most), and passes them to
+*	the loop body to give it a task-local "workspace" that can be mutated without need for synchronization primitives.
+*	@param Contexts; User-privided array of user-defined task-level context objects
+*	@param Num; number of calls of Body; Body(0), Body(1)....Body(Num - 1)
+*	@param MinBatchSize; Minimum Size of a Batch (will only launch DivUp(Num, MinBatchSize) Workers 
+*	@param Body; Function to call from multiple threads
+*	@param Flags; Used to customize the behavior of the ParallelFor if needed.
+*	Notes: Please add stats around to calls to parallel for and within your lambda as appropriate. Do not clog the task graph with long running tasks or tasks that block.
+**/
+template <typename ContextType, typename FunctionType>
+inline void ParallelForWithExistingTaskContext(TArrayView<ContextType> Contexts, int32 Num, int32 MinBatchSize, const FunctionType& Body, EParallelForFlags Flags = EParallelForFlags::None)
+{
+	if (Num > 0)
+	{
+		ParallelForImpl::ParallelForInternal(TEXT("ParallelFor"), Num, MinBatchSize, Body, [](){}, Flags, Contexts);
+	}
+}
+
+/**
+*	General purpose parallel for that uses the taskgraph. This variant takes an array of user-defined context
+*	objects for each task that may get spawned to do work (one task per context at most), and passes them to
+*	the loop body to give it a task-local "workspace" that can be mutated without need for synchronization primitives.
+*	@param DebugName; ProfilingScope and Debugname
+*	@param Contexts; User-privided array of user-defined task-level context objects
+*	@param Num; number of calls of Body; Body(0), Body(1)....Body(Num - 1)
+*	@param MinBatchSize; Minimum Size of a Batch (will only launch DivUp(Num, MinBatchSize) Workers 
+*	@param Body; Function to call from multiple threads
+*	@param Flags; Used to customize the behavior of the ParallelFor if needed.
+*	Notes: Please add stats around to calls to parallel for and within your lambda as appropriate. Do not clog the task graph with long running tasks or tasks that block.
+**/
+template <typename ContextType, typename FunctionType>
+inline void ParallelForWithExistingTaskContext(const TCHAR* DebugName, TArrayView<ContextType> Contexts, int32 Num, int32 MinBatchSize, const FunctionType& Body, EParallelForFlags Flags = EParallelForFlags::None)
+{
+	if (Num > 0)
+	{
+		ParallelForImpl::ParallelForInternal(DebugName, Num, MinBatchSize, Body, [](){}, Flags, Contexts);
+	}
+}
