LLM: Allow parents to be defined after their children.
Fix parent pointer in child tags that are defined before their parents; it was previously accidentally being set to the base parent instead of the parent tag when the parent was found.
Fix deadlocks in error messages due to allocating memory while inside the LLM CriticalSection.
#rnx

diff --git a/Engine/Source/Runtime/Core/Private/HAL/LowLevelMemTracker.cpp b/Engine/Source/Runtime/Core/Private/HAL/LowLevelMemTracker.cpp
--- a/Engine/Source/Runtime/Core/Private/HAL/LowLevelMemTracker.cpp
+++ b/Engine/Source/Runtime/Core/Private/HAL/LowLevelMemTracker.cpp
@@ -1643,11 +1643,19 @@ UE::LLMPrivate::FTagData& FLowLevelMemTracker::RegisterTagData(FName Name, FName
 		// Note that it is not valid for an LLM_SCOPE to be called before a formal registration
 		// (e.g. LLM_DECLARE_TAG). If a formal registration exists for a tag, it must precede its use
 		// in any LLM_SCOPE calls.
-		if (ReferenceSource != UE::LLMPrivate::ETagReferenceSource::Scope &&
-			ReferenceSource != UE::LLMPrivate::ETagReferenceSource::ImplicitParent &&
-			ReferenceSource != UE::LLMPrivate::ETagReferenceSource::FunctionAPI)
-		{
-			ReportDuplicateTagName(TagDataForName, ReferenceSource);
+		if ((ParentName.IsNone() && TagDataForName->GetParent() != nullptr) ||
+			(!ParentName.IsNone() && (TagDataForName->GetParent() == nullptr || TagDataForName->GetParent()->GetName() != ParentName)) ||
+			StatName != TagDataForName->GetStatName() ||
+			SummaryStatName != TagDataForName->GetSummaryStatName() ||
+			bHasEnumTag != TagDataForName->HasEnumTag() ||
+			(bHasEnumTag && EnumTag != TagDataForName->GetEnumTag()))
+		{
+			if (ReferenceSource != UE::LLMPrivate::ETagReferenceSource::Scope &&
+				ReferenceSource != UE::LLMPrivate::ETagReferenceSource::ImplicitParent &&
+				ReferenceSource != UE::LLMPrivate::ETagReferenceSource::FunctionAPI)
+			{
+				ReportDuplicateTagName(TagDataForName, ReferenceSource);
+			}
 		}
 
 		// Abandon the string work we've done and return the version that was added by the other source.
@@ -1698,23 +1706,27 @@ void FLowLevelMemTracker::ReportDuplicateTagName(UE::LLMPrivate::FTagData* TagDa
 {
 	using namespace UE::LLMPrivate;
 
+	// We're inside the LLM lock, so do not allow these GetName calls to allocate memory from FMalloc.
 	if (ReferenceSource == ETagReferenceSource::FunctionAPI || ReferenceSource == ETagReferenceSource::Scope || ReferenceSource == ETagReferenceSource::ImplicitParent)
 	{
-		LLMCheckf(false, TEXT("LLM Error: Unexpected call to RegisterTagData(%s) from LLM_SCOPE or function call when the tag already exists."), *TagData->GetName().ToString());
+		LLMCheckf(false, TEXT("LLM Error: Unexpected call to RegisterTagData(%s) from LLM_SCOPE or function call when the tag already exists."),
+			*WriteToString<FName::StringBufferSize>(TagData->GetName()));
 	}
 	else if (TagData->GetReferenceSource() == ETagReferenceSource::FunctionAPI || TagData->GetReferenceSource() == ETagReferenceSource::Scope)
 	{
-		LLMCheckf(false, TEXT("LLM Error: Tag %s parsed from %s after it was already used in an LLM_SCOPE or LLM api call."), *TagData->GetName().ToString(), ToString(ReferenceSource));
+		LLMCheckf(false, TEXT("LLM Error: Tag %s parsed from %s after it was already used in an LLM_SCOPE or LLM api call."),
+			*WriteToString<FName::StringBufferSize>(TagData->GetName()), ToString(ReferenceSource));
 	}
 	else if (TagData->GetReferenceSource() == ETagReferenceSource::ImplicitParent)
 	{
 		LLMCheckf(false, TEXT("LLM Error: Tag %s parsed from %s after it was already used as an implicit parent in another tag. Add LLM_DEFINE_TAG(%s) in cpp, or move it to the same module as the child tag using it, so that it will be defined before the child tag tries to use it."),
-			*TagData->GetName().ToString(), ToString(ReferenceSource), *TagData->GetName().ToString());
+			*WriteToString<FName::StringBufferSize>(TagData->GetName()), ToString(ReferenceSource),
+			*WriteToString<FName::StringBufferSize>(TagData->GetName()));
 	}
 	else
 	{
 		LLMCheckf(false, TEXT("LLM Error: Multiple occurrences of a unique tag name %s in ELLMTag or LLM_DEFINE_TAG. First occurrence: %s. Second occurrence: %s."),
-			*TagData->GetName().ToString(), ToString(TagData->GetReferenceSource()), ToString(ReferenceSource));
+			*WriteToString<FName::StringBufferSize>(TagData->GetName()), ToString(TagData->GetReferenceSource()), ToString(ReferenceSource));
 	}
 }
 
@@ -1755,7 +1767,7 @@ void FLowLevelMemTracker::FinishConstruct(UE::LLMPrivate::FTagData* TagData, UE:
 					// Another thread got in and finished construction while we were outside of the lock
 					return;
 				}
-				ParentDataPtr = TagDataNameMap->Find(GetTagName_CustomName());
+				ParentDataPtr = TagDataNameMap->Find(ParentName);
 				LLMCheck(ParentDataPtr);
 			}
 			FTagData* ParentData = *ParentDataPtr;
@@ -2912,13 +2924,13 @@ namespace LLMPrivate
 
 	const FTagData* FTagData::GetParent() const
 	{
-		LLMCheckf(!bParentIsName, TEXT("GetParent called on TagData %s before SetParent was called"), *Name.ToString());
+		LLMCheckf(!bParentIsName, TEXT("GetParent called on TagData %s before SetParent was called"), *WriteToString<FName::StringBufferSize>(Name));
 		return Parent;
 	}
 
 	FName FTagData::GetParentName() const
 	{
-		LLMCheckf(bParentIsName, TEXT("GetParent called on TagData %s after SetParent was called"), *Name.ToString());
+		LLMCheckf(bParentIsName, TEXT("GetParentName called on TagData %s after SetParent was called"), *WriteToString<FName::StringBufferSize>(Name));
 		return ParentName;
 	}
 
