Add a core delegate which is broadcasts once GWarn and GError are first set up

#rnx

- This allows programs to override GWarn and GError before any messages are passed through them by FEngineLoop::PreInit.

diff --git a/Engine/Source/Runtime/Core/Private/Misc/CoreDelegates.cpp b/Engine/Source/Runtime/Core/Private/Misc/CoreDelegates.cpp
--- a/Engine/Source/Runtime/Core/Private/Misc/CoreDelegates.cpp
+++ b/Engine/Source/Runtime/Core/Private/Misc/CoreDelegates.cpp
@@ -95,6 +95,7 @@ FCoreDelegates::FPakSigningKeysDelegate& FCoreDelegates::GetPakSigningKeysDelega
 #endif
 FSimpleMulticastDelegate FCoreDelegates::OnShutdownAfterError;
 FSimpleMulticastDelegate FCoreDelegates::OnInit;
+FSimpleMulticastDelegate FCoreDelegates::OnOutputDevicesInit;
 FSimpleMulticastDelegate FCoreDelegates::OnPostEngineInit;
 FSimpleMulticastDelegate FCoreDelegates::OnAllModuleLoadingPhasesComplete;
 FSimpleMulticastDelegate FCoreDelegates::OnFEngineLoopInitComplete;
diff --git a/Engine/Source/Runtime/Core/Public/Misc/CoreDelegates.h b/Engine/Source/Runtime/Core/Public/Misc/CoreDelegates.h
--- a/Engine/Source/Runtime/Core/Public/Misc/CoreDelegates.h
+++ b/Engine/Source/Runtime/Core/Public/Misc/CoreDelegates.h
@@ -255,6 +255,9 @@ public:
 	// Called when appInit is called, very early in startup
 	static FSimpleMulticastDelegate OnInit;
 
+	// Called during FEngineLoop::PreInit after GWarn & GError have been first set so that they can be overridden before anything in PreInit uses them
+	static FSimpleMulticastDelegate OnOutputDevicesInit;
+
 	// Called at the end of UEngine::Init, right before loading PostEngineInit modules for both normal execution and commandlets
 	static FSimpleMulticastDelegate OnPostEngineInit;
 
diff --git a/Engine/Source/Runtime/Launch/Private/LaunchEngineLoop.cpp b/Engine/Source/Runtime/Launch/Private/LaunchEngineLoop.cpp
--- a/Engine/Source/Runtime/Launch/Private/LaunchEngineLoop.cpp
+++ b/Engine/Source/Runtime/Launch/Private/LaunchEngineLoop.cpp
@@ -1789,6 +1789,8 @@ int32 FEngineLoop::PreInitPreStartupScreen(const TCHAR* CmdLine)
 		GError = FPlatformOutputDevices::GetError();
 		GWarn = FPlatformOutputDevices::GetFeedbackContext();
 #endif
+
+		FCoreDelegates::OnOutputDevicesInit.Broadcast();
 	}
 
 	// Avoiding potential exploits by not exposing command line overrides in the shipping games.
