Logging: Fixed a data race in FAsyncWriter in the file output device

diff --git a/Engine/Source/Runtime/Core/Private/Misc/OutputDeviceFile.cpp b/Engine/Source/Runtime/Core/Private/Misc/OutputDeviceFile.cpp
--- a/Engine/Source/Runtime/Core/Private/Misc/OutputDeviceFile.cpp
+++ b/Engine/Source/Runtime/Core/Private/Misc/OutputDeviceFile.cpp
@@ -95,7 +95,7 @@ void FAsyncWriter::SerializeBufferToArchive()
 		// We won't be modifying it anyway and will later serialize new data in the next iteration.
 		// Here we only serialize what we know exists at the beginning of this function.
 		int32 ThisThreadStartPos = BufferStartPos.Load(EMemoryOrder::Relaxed);
-		int32 ThisThreadEndPos   = BufferEndPos  .Load(EMemoryOrder::Relaxed);
+		int32 ThisThreadEndPos   = BufferEndPos.Load(EMemoryOrder::SequentiallyConsistent);
 
 		if (ThisThreadEndPos >= ThisThreadStartPos)
 		{
@@ -109,7 +109,7 @@ void FAsyncWriter::SerializeBufferToArchive()
 		}
 
 		// Modify the start pos. Only the worker thread modifies this value so it's ok to not guard it with a critical section.
-		BufferStartPos = ThisThreadEndPos;
+		BufferStartPos.Store(ThisThreadEndPos, EMemoryOrder::Relaxed);
 
 		// Decrement the request counter, we now know we serialized at least one request.
 		// We might have serialized more requests but it's irrelevant, the counter will go down to 0 eventually
@@ -254,7 +254,7 @@ void FAsyncWriter::Serialize(void* InData, int64 Length)
 	}
 
 	// Update the end position and let the async thread know we need to write to disk
-	BufferEndPos = (ThisThreadEndPos + Length) % Buffer.Num();
+	BufferEndPos.Store((ThisThreadEndPos + Length) % Buffer.Num(), EMemoryOrder::SequentiallyConsistent);
 	SerializeRequestCounter.Increment();
 
 	// No async thread? Serialize now.
