Added micro and nano-seconds to unit conversion utilities

diff --git a/Engine/Source/Runtime/Core/Private/Math/UnitConversion.cpp b/Engine/Source/Runtime/Core/Private/Math/UnitConversion.cpp
--- a/Engine/Source/Runtime/Core/Private/Math/UnitConversion.cpp
+++ b/Engine/Source/Runtime/Core/Private/Math/UnitConversion.cpp
@@ -73,6 +73,8 @@ FParseCandidate ParseCandidates[] = {
 	{ TEXT("Lux"),					EUnit::Lux },					{ TEXT("lx"),		EUnit::Lux },
 	{ TEXT("CandelaPerMeterSquared"), EUnit::CandelaPerMeter2 },	{ TEXT("cd/m2"),	EUnit::CandelaPerMeter2 },		{ TEXT("CandelaPerMeter2"),		EUnit::CandelaPerMeter2 },
 
+	{ TEXT("Nanoseconds"),			EUnit::Nanoseconds },			{ TEXT("ns"),		EUnit::Nanoseconds },
+	{ TEXT("Microseconds"),			EUnit::Microseconds },			{ TEXT("us"),		EUnit::Microseconds },			{ TEXT("Microseconds"),			EUnit::Microseconds },			{ TEXT("\u00B5s"),	EUnit::Microseconds },
 	{ TEXT("Milliseconds"),			EUnit::Milliseconds },			{ TEXT("ms"),		EUnit::Milliseconds },
 	{ TEXT("Seconds"),				EUnit::Seconds },				{ TEXT("s"),		EUnit::Seconds },
 	{ TEXT("Minutes"),				EUnit::Minutes },				{ TEXT("min"),		EUnit::Minutes },
@@ -111,7 +113,7 @@ const TCHAR* const DisplayStrings[] = {
 
 	TEXT("lm"), TEXT("cd"), TEXT("lux"), TEXT("cd/m2"),
 
-	TEXT("ms"), TEXT("s"), TEXT("min"), TEXT("hr"), TEXT("dy"), TEXT("mth"), TEXT("yr"),
+	TEXT("ns"), TEXT("\u00B5s"), TEXT("ms"), TEXT("s"), TEXT("min"), TEXT("hr"), TEXT("dy"), TEXT("mth"), TEXT("yr"),
 
 	TEXT("ppi"),
 
@@ -142,7 +144,7 @@ const EUnitType UnitTypes[] = {
 
 	EUnitType::LuminousFlux, EUnitType::LuminousIntensity, EUnitType::Illuminance, EUnitType::Luminance,
 
-	EUnitType::Time,		EUnitType::Time,		EUnitType::Time,		EUnitType::Time,		EUnitType::Time,		EUnitType::Time,		EUnitType::Time,
+	EUnitType::Time,		EUnitType::Time,		EUnitType::Time,		EUnitType::Time,		EUnitType::Time,		EUnitType::Time,		EUnitType::Time,		EUnitType::Time,		EUnitType::Time,
 
 	EUnitType::PixelDensity,
 
@@ -646,6 +648,8 @@ namespace UnitConversion
 			case EUnit::Days:				Factor *= 24;		// fallthrough
 											return Factor;
 
+			case EUnit::Nanoseconds:		Factor /= 1000;		// fallthrough
+			case EUnit::Microseconds:		Factor /= 1000;		// fallthrough
 			case EUnit::Milliseconds:		Factor /= 1000;		// fallthrough
 			case EUnit::Seconds:			Factor /= 60;		// fallthrough
 			case EUnit::Minutes:			Factor /= 60;		// fallthrough
diff --git a/Engine/Source/Runtime/Core/Public/Math/UnitConversion.h b/Engine/Source/Runtime/Core/Public/Math/UnitConversion.h
--- a/Engine/Source/Runtime/Core/Public/Math/UnitConversion.h
+++ b/Engine/Source/Runtime/Core/Public/Math/UnitConversion.h
@@ -47,7 +47,7 @@ enum class EUnit : uint8
 	Lumens, Candela, Lux, CandelaPerMeter2,
 
 	/** Time units */
-	Milliseconds, Seconds, Minutes, Hours, Days, Months, Years,
+	Nanoseconds, Microseconds, Milliseconds, Seconds, Minutes, Hours, Days, Months, Years,
 
 	/** Pixel density units */
 	PixelsPerInch,
