TIsCharEncodingCompatibleWith, TIsFixedWidthCharEncoding and TIsCharEncodingSimplyConvertibleTo traits, moved out of from FPlatformString because Platform is not available early enough in the compilation process.

diff --git a/Engine/Source/Runtime/Core/Private/GenericPlatform/GenericPlatformString.cpp b/Engine/Source/Runtime/Core/Private/GenericPlatform/GenericPlatformString.cpp
--- a/Engine/Source/Runtime/Core/Private/GenericPlatform/GenericPlatformString.cpp
+++ b/Engine/Source/Runtime/Core/Private/GenericPlatform/GenericPlatformString.cpp
@@ -676,7 +676,7 @@ namespace UE::Core::Private
 template <typename DestEncoding, typename SourceEncoding>
 void FGenericPlatformString::LogBogusChars(const SourceEncoding* Src, int32 SrcSize)
 {
-	static_assert(IsFixedWidthEncoding<SourceEncoding>(), "Currently unimplemented for non-fixed-width source conversions");
+	static_assert(TIsFixedWidthCharEncoding_V<SourceEncoding>, "Currently unimplemented for non-fixed-width source conversions");
 
 	FString SrcStr;
 	bool    bFoundBogusChars = false;
diff --git a/Engine/Source/Runtime/Core/Public/Containers/StringConv.h b/Engine/Source/Runtime/Core/Public/Containers/StringConv.h
--- a/Engine/Source/Runtime/Core/Public/Containers/StringConv.h
+++ b/Engine/Source/Runtime/Core/Public/Containers/StringConv.h
@@ -18,6 +18,7 @@
 #include "Templates/UnrealTemplate.h"
 #include "Templates/UnrealTypeTraits.h"
 #include "Traits/ElementType.h"
+#include "Traits/IsCharEncodingCompatibleWith.h"
 #include "Traits/IsContiguousContainer.h"
 #include <type_traits>
 
@@ -32,7 +33,7 @@ public:
 
 	template <
 		typename CharType,
-		std::enable_if_t<FPlatformString::IsCharEncodingCompatibleWith<CharType, FromType>()>* = nullptr
+		std::enable_if_t<TIsCharEncodingCompatibleWith_V<CharType, FromType>>* = nullptr
 	>
 	FORCEINLINE static void Convert(To* Dest, int32 DestLen, const CharType* Source, int32 SourceLen)
 	{
@@ -42,7 +43,7 @@ public:
 
 	template <
 		typename CharType,
-		std::enable_if_t<FPlatformString::IsCharEncodingCompatibleWith<CharType, FromType>()>* = nullptr
+		std::enable_if_t<TIsCharEncodingCompatibleWith_V<CharType, FromType>>* = nullptr
 	>
 	static int32 ConvertedLength(const CharType* Source, int32 SourceLen)
 	{
@@ -293,7 +294,7 @@ public:
 	 */
 	template <
 		typename SrcBufferType,
-		std::enable_if_t<FPlatformString::IsCharEncodingCompatibleWith<SrcBufferType, FromType>()>* = nullptr
+		std::enable_if_t<TIsCharEncodingCompatibleWith_V<SrcBufferType, FromType>>* = nullptr
 	>
 	static FORCEINLINE int32 Convert(IntendedToType* Dest, int32 DestLen, const SrcBufferType* Source, int32 SourceLen)
 	{
@@ -307,7 +308,7 @@ public:
 	}
 	template <
 		typename SrcBufferType,
-		std::enable_if_t<FPlatformString::IsCharEncodingCompatibleWith<SrcBufferType, FromType>()>* = nullptr
+		std::enable_if_t<TIsCharEncodingCompatibleWith_V<SrcBufferType, FromType>>* = nullptr
 	>
 	static FORCEINLINE int32 Convert(ToType* Dest, int32 DestLen, const SrcBufferType* Source, int32 SourceLen)
 	{
@@ -322,7 +323,7 @@ public:
 	 */
 	template <
 		typename SrcBufferType,
-		std::enable_if_t<FPlatformString::IsCharEncodingCompatibleWith<SrcBufferType, FromType>()>* = nullptr
+		std::enable_if_t<TIsCharEncodingCompatibleWith_V<SrcBufferType, FromType>>* = nullptr
 	>
 	static FORCEINLINE int32 ConvertedLength(const SrcBufferType* Source, int32 SourceLen)
 	{
@@ -347,7 +348,7 @@ public:
 	 */
 	template <
 		typename SrcBufferType,
-		std::enable_if_t<FPlatformString::IsCharEncodingCompatibleWith<SrcBufferType, FromType>()>* = nullptr
+		std::enable_if_t<TIsCharEncodingCompatibleWith_V<SrcBufferType, FromType>>* = nullptr
 	>
 	static FORCEINLINE void Convert(ToType* Dest, const int32 DestLen, const SrcBufferType* Source, const int32 SourceLen)
 	{
@@ -367,7 +368,7 @@ public:
 	 */
 	template <
 		typename SrcBufferType,
-		std::enable_if_t<FPlatformString::IsCharEncodingCompatibleWith<SrcBufferType, FromType>()>* = nullptr
+		std::enable_if_t<TIsCharEncodingCompatibleWith_V<SrcBufferType, FromType>>* = nullptr
 	>
 	static int32 ConvertedLength(const SrcBufferType* Source, const int32 SourceLen)
 	{
@@ -455,7 +456,7 @@ public:
 	 */
 	template <
 		typename SrcBufferType,
-		std::enable_if_t<FPlatformString::IsCharEncodingCompatibleWith<SrcBufferType, FromType>()>* = nullptr
+		std::enable_if_t<TIsCharEncodingCompatibleWith_V<SrcBufferType, FromType>>* = nullptr
 	>
 	static FORCEINLINE void Convert(ToType* Dest, int32 DestLen, const SrcBufferType* Source, int32 SourceLen)
 	{
@@ -469,7 +470,7 @@ public:
 	 */
 	template <
 		typename SrcBufferType,
-		std::enable_if_t<FPlatformString::IsCharEncodingCompatibleWith<SrcBufferType, FromType>()>* = nullptr
+		std::enable_if_t<TIsCharEncodingCompatibleWith_V<SrcBufferType, FromType>>* = nullptr
 	>
 	static FORCEINLINE int32 ConvertedLength(const SrcBufferType* Source, int32 SourceLen)
 	{
@@ -531,7 +532,7 @@ public:
 	 */
 	template <
 		typename SrcBufferType,
-		std::enable_if_t<FPlatformString::IsCharEncodingCompatibleWith<SrcBufferType, FromType>()>* = nullptr
+		std::enable_if_t<TIsCharEncodingCompatibleWith_V<SrcBufferType, FromType>>* = nullptr
 	>
 	static FORCEINLINE void Convert(ToType* Dest, const int32 DestLen, const SrcBufferType* Source, const int32 SourceLen)
 	{
@@ -547,7 +548,7 @@ public:
 	 */
 	template <
 		typename SrcBufferType,
-		std::enable_if_t<FPlatformString::IsCharEncodingCompatibleWith<SrcBufferType, FromType>()>* = nullptr
+		std::enable_if_t<TIsCharEncodingCompatibleWith_V<SrcBufferType, FromType>>* = nullptr
 	>
 	static int32 ConvertedLength(const SrcBufferType* Source, const int32 SourceLen)
 	{
@@ -672,7 +673,7 @@ class TStringConversion : private Converter, private TInlineAllocator<DefaultCon
 public:
 	template <
 		typename SrcBufferType,
-		std::enable_if_t<FPlatformString::IsCharEncodingCompatibleWith<SrcBufferType, FromType>()>* = nullptr
+		std::enable_if_t<TIsCharEncodingCompatibleWith_V<SrcBufferType, FromType>>* = nullptr
 	>
 	explicit TStringConversion(const SrcBufferType* Source)
 	{
@@ -693,7 +694,7 @@ public:
 
 	template <
 		typename SrcBufferType,
-		std::enable_if_t<FPlatformString::IsCharEncodingCompatibleWith<SrcBufferType, FromType>()>* = nullptr
+		std::enable_if_t<TIsCharEncodingCompatibleWith_V<SrcBufferType, FromType>>* = nullptr
 	>
 	TStringConversion(const SrcBufferType* Source, int32 SourceLen)
 	{
@@ -729,9 +730,9 @@ public:
 		std::enable_if_t<
 			TAnd<
 				TIsContiguousContainer<FromRangeType>,
-				TNot<TIsArray<typename TRemoveReference<FromRangeType>::Type>>
-			>::Value &&
-			FPlatformString::IsCharEncodingCompatibleWith<FromRangeCharType, FromType>()
+				TNot<TIsArray<typename TRemoveReference<FromRangeType>::Type>>,
+				TIsCharEncodingCompatibleWith<FromRangeCharType, FromType>
+			>::Value
 		>* = nullptr
 	>
 	TStringConversion(FromRangeType&& Source)
@@ -824,7 +825,7 @@ public:
 public:
 	template <
 		typename SrcBufferType,
-		std::enable_if_t<FPlatformString::IsCharEncodingCompatibleWith<SrcBufferType, FromType>()>* = nullptr
+		std::enable_if_t<TIsCharEncodingCompatibleWith_V<SrcBufferType, FromType>>* = nullptr
 	>
 	explicit TStringPointer(const SrcBufferType* Source)
 	{
@@ -842,7 +843,7 @@ public:
 
 	template <
 		typename SrcBufferType,
-		std::enable_if_t<FPlatformString::IsCharEncodingCompatibleWith<SrcBufferType, FromType>()>* = nullptr
+		std::enable_if_t<TIsCharEncodingCompatibleWith_V<SrcBufferType, FromType>>* = nullptr
 	>
 	TStringPointer(const SrcBufferType* Source, int32 SourceLen)
 	{
@@ -873,9 +874,9 @@ public:
 		std::enable_if_t<
 			TAnd<
 				TIsContiguousContainer<FromRangeType>,
-				TNot<TIsArray<typename TRemoveReference<FromRangeType>::Type>>
-			>::Value &&
-			FPlatformString::IsCharEncodingCompatibleWith<FromRangeCharType, FromType>()
+				TNot<TIsArray<typename TRemoveReference<FromRangeType>::Type>>,
+				TIsCharEncodingCompatibleWith<FromRangeCharType, FromType>
+			>::Value
 		>* = nullptr
 	>
 	TStringPointer(FromRangeType&& Source)
@@ -1061,7 +1062,7 @@ typedef TStringPointer<wchar_t, TCHAR> FWCharToTCHAR;
 template <typename To, int32 DefaultConversionSize = DEFAULT_STRING_CONVERSION_SIZE, typename From>
 FORCEINLINE auto StringCast(const From* Str)
 {
-	if constexpr (FPlatformString::IsCharEncodingCompatibleWith<From, To>())
+	if constexpr (TIsCharEncodingCompatibleWith_V<From, To>)
 	{
 		return TStringPointer<To>((const To*)Str);
 	}
@@ -1080,7 +1081,7 @@ FORCEINLINE auto StringCast(const From* Str)
 template <typename To, int32 DefaultConversionSize = DEFAULT_STRING_CONVERSION_SIZE, typename From>
 FORCEINLINE auto StringCast(const From* Str, int32 Len)
 {
-	if constexpr (FPlatformString::IsCharEncodingCompatibleWith<From, To>())
+	if constexpr (TIsCharEncodingCompatibleWith_V<From, To>)
 	{
 		return TStringPointer<To>((const To*)Str, Len);
 	}
@@ -1208,7 +1209,7 @@ private:
 template <typename From, typename To, int32 DefaultConversionSize = DEFAULT_STRING_CONVERSION_SIZE>
 FORCEINLINE auto StringMemoryPassthru(To* Buffer, int32 BufferSize, int32 SourceLength)
 {
-	if constexpr (FPlatformString::IsCharEncodingCompatibleWith<From, To>())
+	if constexpr (TIsCharEncodingCompatibleWith_V<From, To>)
 	{
 		check(SourceLength <= BufferSize);
 		return TPassthruPointer<From>((From*)Buffer);
diff --git a/Engine/Source/Runtime/Core/Public/Containers/StringView.h b/Engine/Source/Runtime/Core/Public/Containers/StringView.h
--- a/Engine/Source/Runtime/Core/Public/Containers/StringView.h
+++ b/Engine/Source/Runtime/Core/Public/Containers/StringView.h
@@ -12,6 +12,7 @@
 #include "String/Find.h"
 #include "Templates/UnrealTemplate.h"
 #include "Traits/ElementType.h"
+#include "Traits/IsCharEncodingCompatibleWith.h"
 #include "Traits/IsCharType.h"
 #include "Traits/IsContiguousContainer.h"
 #include <type_traits>
@@ -99,8 +100,15 @@
 	}
 
 	/** Construct a view of the null-terminated string pointed to by InData. */
-	template <typename OtherCharType,
-		std::enable_if_t<FPlatformString::IsCharEncodingCompatibleWith<OtherCharType, CharType>()>* = nullptr>
+	template <
+		typename OtherCharType,
+		std::enable_if_t<
+			TAnd<
+				TIsCharType<OtherCharType>,
+				TIsCharEncodingCompatibleWith<OtherCharType, CharType>
+			>::Value
+		>* = nullptr
+	>
 	constexpr inline TStringView(const OtherCharType* InData)
 		: DataPtr((const CharType*)InData)
 		, Size(InData ? TCString<CharType>::Strlen((const CharType*)InData) : 0)
@@ -108,8 +116,15 @@
 	}
 
 	/** Construct a view of InSize characters beginning at InData. */
-	template <typename OtherCharType,
-		std::enable_if_t<FPlatformString::IsCharEncodingCompatibleWith<OtherCharType, CharType>()>* = nullptr>
+	template <
+		typename OtherCharType,
+		std::enable_if_t<
+			TAnd<
+				TIsCharType<OtherCharType>,
+				TIsCharEncodingCompatibleWith<OtherCharType, CharType>
+			>::Value
+		>* = nullptr
+	>
 	constexpr inline TStringView(const OtherCharType* InData, int32 InSize)
 		: DataPtr((const CharType*)InData)
 		, Size(InSize)
@@ -119,9 +134,11 @@
 	/** Construct a view from a contiguous range of characters, such as FString or TStringBuilder. */
 	template <typename CharRangeType,
 		std::enable_if_t<
-			TIsContiguousContainer<CharRangeType>::Value &&
-			TIsCharType<TElementType_T<CharRangeType>>::Value &&
-			FPlatformString::IsCharEncodingCompatibleWith<TElementType_T<CharRangeType>, CharType>() &&
+			TAnd<
+				TIsContiguousContainer<CharRangeType>,
+				TIsCharType<TElementType_T<CharRangeType>>,
+				TIsCharEncodingCompatibleWith<TElementType_T<CharRangeType>, CharType>
+			>::Value &&
 			!std::is_array_v<std::remove_reference_t<CharRangeType>> &&
 			!std::is_same_v<CharRangeType, ViewType>
 		>* = nullptr>
@@ -152,7 +169,7 @@
 	/** Modifies the view to remove the given number of characters from the end. */
 	inline void RemoveSuffix(int32 CharCount) { Size -= CharCount; }
 	/** Resets to an empty view */
-	inline void Reset() { DataPtr = nullptr; Size = 0; }
+	inline void		Reset()								{ DataPtr = nullptr; Size = 0; }
 
 	// Operations
 
@@ -210,7 +227,7 @@
 	 * Check whether this view is equivalent to a character range.
 	 *
 	 * @param Other        A character range that is comparable with the character type of this view.
-	 * @param SearchCase   Whether the comparison should ignore case.
+	 * @param SearchCase Whether the comparison should ignore case.
 	 */
 	template <typename OtherRangeType, decltype(MakeStringView(DeclVal<OtherRangeType>()))* = nullptr>
 	[[nodiscard]] inline bool Equals(OtherRangeType&& Other, ESearchCase::Type SearchCase = ESearchCase::CaseSensitive) const
@@ -223,7 +240,7 @@
 	 * Check whether this view is equivalent to a string view.
 	 *
 	 * @param Other        A string that is comparable with the character type of this view.
-	 * @param SearchCase   Whether the comparison should ignore case.
+	 * @param SearchCase Whether the comparison should ignore case.
 	 */
 	template <typename OtherCharType,
 		std::enable_if_t<TIsCharType<OtherCharType>::Value>* = nullptr>
@@ -233,7 +250,7 @@
 	 * Compare this view with a character range.
 	 *
 	 * @param Other        A character range that is comparable with the character type of this view.
-	 * @param SearchCase   Whether the comparison should ignore case.
+	 * @param SearchCase Whether the comparison should ignore case.
 	 * @return 0 is equal, negative if this view is less, positive if this view is greater.
 	 */
 	template <typename OtherRangeType, decltype(MakeStringView(DeclVal<OtherRangeType>()))* = nullptr>
@@ -257,7 +274,7 @@
 	 * Compare this view with a null-terminated string.
 	 *
 	 * @param Other        A null-terminated string that is comparable with the character type of this view.
-	 * @param SearchCase   Whether the comparison should ignore case.
+	 * @param SearchCase Whether the comparison should ignore case.
 	 * @return 0 is equal, negative if this view is less, positive if this view is greater.
 	 */
 	template <typename OtherCharType>
@@ -299,7 +316,7 @@
 	 * Search the view for the first occurrence of a character.
 	 *
 	 * @param Search           The character to search for. Comparison is lexicographic.
-	 * @param OutIndex [out]   The position at which the character was found, or INDEX_NONE if not found.
+	 * @param OutIndex [out] The position at which the character was found, or INDEX_NONE if not found.
 	 * @return True if the character was found in the view, otherwise false.
 	 */
 	inline bool FindChar(CharType Search, int32& OutIndex) const;
@@ -308,7 +325,7 @@
 	 * Search the view for the last occurrence of a character.
 	 *
 	 * @param Search           The character to search for. Comparison is lexicographic.
-	 * @param OutIndex [out]   The position at which the character was found, or INDEX_NONE if not found.
+	 * @param OutIndex [out] The position at which the character was found, or INDEX_NONE if not found.
 	 * @return True if the character was found in the view, otherwise false.
 	 */
 	inline bool FindLastChar(CharType Search, int32& OutIndex) const;
@@ -620,14 +637,14 @@
 
 template <typename CharType>
 inline int32 TStringView<CharType>::Find(const ViewType Search, const int32 StartPosition) const
-{
+			{
 	const int32 Index = UE::String::FindFirst(RightChop(StartPosition), Search);
 	return Index == INDEX_NONE ? INDEX_NONE : Index + StartPosition;
 }
 
 template <typename CharType>
 inline bool TStringView<CharType>::FindChar(const CharType Search, int32& OutIndex) const
-{
+	{
 	OutIndex = UE::String::FindFirstChar(*this, Search);
 	return OutIndex != INDEX_NONE;
 }
diff --git a/Engine/Source/Runtime/Core/Public/Containers/UnrealString.h b/Engine/Source/Runtime/Core/Public/Containers/UnrealString.h
--- a/Engine/Source/Runtime/Core/Public/Containers/UnrealString.h
+++ b/Engine/Source/Runtime/Core/Public/Containers/UnrealString.h
@@ -26,6 +26,7 @@
 #include "Templates/TypeHash.h"
 #include "Templates/IsFloatingPoint.h"
 #include "Traits/IsCharType.h"
+#include "Traits/IsCharEncodingSimplyConvertibleTo.h"
 
 struct FStringFormatArg;
 template<typename InKeyType,typename InValueType,typename SetAllocator ,typename KeyFuncs > class TMap;
@@ -410,7 +411,7 @@ public:
 	>
 	FORCEINLINE FString& operator+=(AppendedCharType Char)
 	{
-		if constexpr (FPlatformString::IsCharEncodingSimplyConvertibleTo<AppendedCharType, TCHAR>())
+		if constexpr (TIsCharEncodingSimplyConvertibleTo_V<AppendedCharType, TCHAR>)
 		{
 			return AppendChar((TCHAR)Char);
 		}
diff --git a/Engine/Source/Runtime/Core/Public/GenericPlatform/GenericPlatformString.h b/Engine/Source/Runtime/Core/Public/GenericPlatform/GenericPlatformString.h
--- a/Engine/Source/Runtime/Core/Public/GenericPlatform/GenericPlatformString.h
+++ b/Engine/Source/Runtime/Core/Public/GenericPlatform/GenericPlatformString.h
@@ -4,26 +4,29 @@
 
 #include "CoreTypes.h"
 #include "Templates/EnableIf.h"
+#include "Traits/IsCharEncodingCompatibleWith.h"
+#include "Traits/IsCharEncodingSimplyConvertibleTo.h"
 #include "Traits/IsCharType.h"
+#include "Traits/IsFixedWidthCharEncoding.h"
 #include "GenericPlatform/GenericPlatformStricmp.h"
 #include <type_traits>
 
 namespace UE::Core::Private
 {
 	// The Dest parameter is just used for overload resolution
-	CORE_API int32 GetConvertedLength(const UTF8CHAR* Dest, const ANSICHAR* Src, int32 SrcLen);
-	CORE_API int32 GetConvertedLength(const UTF8CHAR* Dest, const WIDECHAR* Src, int32 SrcLen);
-	CORE_API int32 GetConvertedLength(const UTF8CHAR* Dest, const UCS2CHAR* Src, int32 SrcLen);
-	CORE_API int32 GetConvertedLength(const ANSICHAR* Dest, const UTF8CHAR* Src, int32 SrcLen);
-	CORE_API int32 GetConvertedLength(const WIDECHAR* Dest, const UTF8CHAR* Src, int32 SrcLen);
-	CORE_API int32 GetConvertedLength(const UCS2CHAR* Dest, const UTF8CHAR* Src, int32 SrcLen);
-
-	CORE_API UTF8CHAR* Convert(UTF8CHAR* Dest, int32 DestLen, const ANSICHAR* Src, int32 SrcLen);
-	CORE_API UTF8CHAR* Convert(UTF8CHAR* Dest, int32 DestLen, const WIDECHAR* Src, int32 SrcLen);
-	CORE_API UTF8CHAR* Convert(UTF8CHAR* Dest, int32 DestLen, const UCS2CHAR* Src, int32 SrcLen);
-	CORE_API ANSICHAR* Convert(ANSICHAR* Dest, int32 DestLen, const UTF8CHAR* Src, int32 SrcLen);
-	CORE_API WIDECHAR* Convert(WIDECHAR* Dest, int32 DestLen, const UTF8CHAR* Src, int32 SrcLen);
-	CORE_API UCS2CHAR* Convert(UCS2CHAR* Dest, int32 DestLen, const UTF8CHAR* Src, int32 SrcLen);
+	CORE_API int32 GetConvertedLength(const UTF8CHAR* Dest, const ANSICHAR*  Src, int32 SrcLen);
+	CORE_API int32 GetConvertedLength(const UTF8CHAR* Dest, const WIDECHAR*  Src, int32 SrcLen);
+	CORE_API int32 GetConvertedLength(const UTF8CHAR* Dest, const UCS2CHAR*  Src, int32 SrcLen);
+	CORE_API int32 GetConvertedLength(const ANSICHAR* Dest, const UTF8CHAR*  Src, int32 SrcLen);
+	CORE_API int32 GetConvertedLength(const WIDECHAR* Dest, const UTF8CHAR*  Src, int32 SrcLen);
+	CORE_API int32 GetConvertedLength(const UCS2CHAR* Dest, const UTF8CHAR*  Src, int32 SrcLen);
+
+	CORE_API UTF8CHAR* Convert(UTF8CHAR* Dest, int32 DestLen, const ANSICHAR*  Src, int32 SrcLen);
+	CORE_API UTF8CHAR* Convert(UTF8CHAR* Dest, int32 DestLen, const WIDECHAR*  Src, int32 SrcLen);
+	CORE_API UTF8CHAR* Convert(UTF8CHAR* Dest, int32 DestLen, const UCS2CHAR*  Src, int32 SrcLen);
+	CORE_API ANSICHAR* Convert(ANSICHAR* Dest, int32 DestLen, const UTF8CHAR*  Src, int32 SrcLen);
+	CORE_API WIDECHAR* Convert(WIDECHAR* Dest, int32 DestLen, const UTF8CHAR*  Src, int32 SrcLen);
+	CORE_API UCS2CHAR* Convert(UCS2CHAR* Dest, int32 DestLen, const UTF8CHAR*  Src, int32 SrcLen);
 }
 
 // These will be moved inside GenericPlatformString.cpp when the platform layer handles UTF-16
@@ -47,18 +50,10 @@ struct FGenericPlatformString : public FGenericPlatformStricmp
 	 * Tests whether an encoding has fixed-width characters
 	 */
 	template <typename Encoding>
+	UE_DEPRECATED(5.1, "FPlatformString::IsFixedWidthEncoding<T>() has been deprecated in favor of TIsFixedWidthCharEncoding_V<T>")
 	static constexpr bool IsFixedWidthEncoding()
 	{
-		static_assert(TIsCharType<Encoding>::Value, "Encoding is not a char type");
-
-		return
-			std::is_same_v<Encoding, ANSICHAR> ||
-			std::is_same_v<Encoding, UCS2CHAR> || // this may not be true when PLATFORM_UCS2CHAR_IS_UTF16CHAR == 1, but this is the legacy behavior
-			std::is_same_v<Encoding, WIDECHAR> || // the UCS2CHAR comment also applies to WIDECHAR
-#if PLATFORM_TCHAR_IS_CHAR16
-			std::is_same_v<Encoding, wchar_t>  || // the UCS2CHAR comment also applies to wchar_t
-#endif
-			std::is_same_v<Encoding, UTF32CHAR>;
+		return TIsFixedWidthCharEncoding_V<Encoding>;
 	}
 
 	/**
@@ -68,65 +63,20 @@ struct FGenericPlatformString : public FGenericPlatformStricmp
 	 * UTF-8, but UTF-8 is not compatible with ANSI.
 	 */
 	template <typename SrcEncoding, typename DestEncoding>
+	UE_DEPRECATED(5.1, "FPlatformString::IsCharEncodingCompatibleWith<A, B>() has been deprecated in favor of TIsCharEncodingCompatibleWith_V<A, B>")
 	static constexpr bool IsCharEncodingCompatibleWith()
 	{
-		if constexpr (std::is_same_v<SrcEncoding, DestEncoding>)
-		{
-			return true;
-		}
-		else if constexpr (std::is_same_v<SrcEncoding, ANSICHAR> && std::is_same_v<DestEncoding, UTF8CHAR>)
-		{
-			return true;
-		}
-		else if constexpr (std::is_same_v<SrcEncoding, UCS2CHAR> && std::is_same_v<DestEncoding, UTF16CHAR>)
-		{
-			return true;
-		}
-		else if constexpr (std::is_same_v<SrcEncoding, WIDECHAR> && std::is_same_v<DestEncoding, UCS2CHAR>)
-		{
-			return true;
-		}
-		else if constexpr (std::is_same_v<SrcEncoding, UCS2CHAR> && std::is_same_v<DestEncoding, WIDECHAR>)
-		{
-			return true;
-		}
-#if PLATFORM_TCHAR_IS_CHAR16
-		else if constexpr (std::is_same_v<SrcEncoding, WIDECHAR> && std::is_same_v<DestEncoding, wchar_t>)
-		{
-			return true;
-		}
-		else if constexpr (std::is_same_v<SrcEncoding, wchar_t> && std::is_same_v<DestEncoding, WIDECHAR>)
-		{
-			return true;
-		}
-#endif
-		else
-		{
-			return false;
-		}
-	};
+		return TIsCharEncodingCompatibleWith_V<SrcEncoding, DestEncoding>;
+	}
 
 	/**
 	 * Tests whether you can simply (i.e. by assignment) encode code units from the source encoding as the destination encoding.
 	 */
 	template <typename SourceEncoding, typename DestEncoding>
+	UE_DEPRECATED(5.1, "FPlatformString::IsCharEncodingSimplyConvertibleTo<A, B>() has been deprecated in favor of TIsCharEncodingSimplyConvertibleTo_V<A, B>")
 	static constexpr bool IsCharEncodingSimplyConvertibleTo()
 	{
-		if constexpr (IsCharEncodingCompatibleWith<SourceEncoding, DestEncoding>())
-		{
-			// Binary-compatible conversions are always simple
-			return true;
-		}
-		else if constexpr (IsFixedWidthEncoding<SourceEncoding>() && sizeof(DestEncoding) >= sizeof(SourceEncoding))
-		{
-			// Converting from a fixed-width encoding to a wider or same encoding should always be possible,
-			// as should ANSICHAR->UTF8CHAR and UCS2CHAR->UTF16CHAR,
-			return true;
-		}
-		else
-		{
-			return false;
-		}
+		return TIsCharEncodingSimplyConvertibleTo_V<SourceEncoding, DestEncoding>;
 	}
 
 	/**
@@ -144,14 +94,14 @@ struct FGenericPlatformString : public FGenericPlatformStricmp
 		static_assert(TIsCharType<DestEncoding  >::Value, "Destination encoding is not a char type");
 
 		// This is only defined for fixed-width encodings, because codepoints cannot be represented in a single variable-width code unit.
-		static_assert(IsFixedWidthEncoding<SourceEncoding>(), "Source encoding is not fixed-width");
+		static_assert(TIsFixedWidthCharEncoding_V<SourceEncoding>, "Source encoding is not fixed-width");
 
-		if constexpr (IsCharEncodingSimplyConvertibleTo<SourceEncoding, DestEncoding>())
+		if constexpr (TIsCharEncodingSimplyConvertibleTo_V<SourceEncoding, DestEncoding>)
 		{
 			// Simple conversions mean conversion is always possible
 			return true;
 		}
-		else if constexpr (!IsFixedWidthEncoding<DestEncoding>())
+		else if constexpr (!TIsFixedWidthCharEncoding_V<DestEncoding>)
 		{
 			// Converting all codepoints to a variable-width encoding should always be possible
 			return true;
@@ -199,9 +149,9 @@ struct FGenericPlatformString : public FGenericPlatformStricmp
 	 * Metafunction which tests whether a given character type represents a fixed-width encoding.
 	 */
 	template <typename T>
-	struct UE_DEPRECATED(5.0, "TIsFixedWidthEncoding is deprecated, use IsFixedWidthEncoding<T>() instead.") TIsFixedWidthEncoding
+	struct UE_DEPRECATED(5.0, "TIsFixedWidthEncoding is deprecated, use TIsFixedWidthCharEncoding_V<T> instead.") TIsFixedWidthEncoding
 	{
-		enum { Value = IsFixedWidthEncoding<T>() };
+		enum { Value = TIsFixedWidthCharEncoding_V<T> };
 	};
 
 
@@ -214,9 +164,9 @@ struct FGenericPlatformString : public FGenericPlatformStricmp
 	 * Same size is a minimum requirement.
 	 */
 	template <typename EncodingA, typename EncodingB>
-	struct UE_DEPRECATED(5.0, "TAreEncodingsCompatible is deprecated, use IsCharEncodingCompatibleWith<SrcEncoding, DestEncoding>() instead.") TAreEncodingsCompatible
+	struct UE_DEPRECATED(5.0, "TAreEncodingsCompatible is deprecated, use TIsCharEncodingCompatibleWith_V<SrcEncoding, DestEncoding> instead.") TAreEncodingsCompatible
 	{
-		enum { Value = IsFixedWidthEncoding<EncodingA>() && IsFixedWidthEncoding<EncodingB>() && sizeof(EncodingA) == sizeof(EncodingB) };
+		enum { Value = TIsFixedWidthCharEncoding_V<EncodingA> && TIsFixedWidthCharEncoding_V<EncodingB> && sizeof(EncodingA) == sizeof(EncodingB) };
 	};
 
 	/**
@@ -233,7 +183,7 @@ struct FGenericPlatformString : public FGenericPlatformStricmp
 	template <typename SourceEncoding, typename DestEncoding>
 	static FORCEINLINE DestEncoding* Convert(DestEncoding* Dest, int32 DestSize, const SourceEncoding* Src, int32 SrcSize)
 	{
-		if constexpr (IsCharEncodingCompatibleWith<SourceEncoding, DestEncoding>())
+		if constexpr (TIsCharEncodingCompatibleWith_V<SourceEncoding, DestEncoding>)
 		{
 			if (DestSize < SrcSize)
 			{
@@ -242,7 +192,7 @@ struct FGenericPlatformString : public FGenericPlatformStricmp
 
 			return (DestEncoding*)Memcpy(Dest, Src, SrcSize * sizeof(SourceEncoding)) + SrcSize;
 		}
-		else if constexpr (IsCharEncodingSimplyConvertibleTo<SourceEncoding, DestEncoding>())
+		else if constexpr (TIsCharEncodingSimplyConvertibleTo_V<SourceEncoding, DestEncoding>)
 		{
 			const int32 Size = DestSize <= SrcSize ? DestSize : SrcSize;
 			for (int I = 0; I < Size; ++I)
@@ -253,7 +203,7 @@ struct FGenericPlatformString : public FGenericPlatformStricmp
 
 			return DestSize < SrcSize ? nullptr : Dest + Size;
 		}
-		else if constexpr (IsFixedWidthEncoding<SourceEncoding>() && IsFixedWidthEncoding<DestEncoding>())
+		else if constexpr (TIsFixedWidthCharEncoding_V<SourceEncoding> && TIsFixedWidthCharEncoding_V<DestEncoding>)
 		{
 			const int32 Size = DestSize <= SrcSize ? DestSize : SrcSize;
 			bool bInvalidChars = false;
@@ -297,7 +247,7 @@ struct FGenericPlatformString : public FGenericPlatformStricmp
 	template <typename DestEncoding, typename SourceEncoding>
 	static int32 ConvertedLength(const SourceEncoding* Src, int32 SrcSize)
 	{
-		if constexpr (IsCharEncodingSimplyConvertibleTo<SourceEncoding, DestEncoding>() || (IsFixedWidthEncoding<SourceEncoding>() && IsFixedWidthEncoding<DestEncoding>()))
+		if constexpr (TIsCharEncodingSimplyConvertibleTo_V<SourceEncoding, DestEncoding> || (TIsFixedWidthCharEncoding_V<SourceEncoding> && TIsFixedWidthCharEncoding_V<DestEncoding>))
 		{
 			return SrcSize;
 		}
diff --git a/Engine/Source/Runtime/Core/Public/Misc/StringBuilder.h b/Engine/Source/Runtime/Core/Public/Misc/StringBuilder.h
--- a/Engine/Source/Runtime/Core/Public/Misc/StringBuilder.h
+++ b/Engine/Source/Runtime/Core/Public/Misc/StringBuilder.h
@@ -12,6 +12,7 @@
 #include "Templates/IsValidVariadicFunctionArg.h"
 #include "Templates/UnrealTemplate.h"
 #include "Traits/IsCharType.h"
+#include "Traits/IsCharEncodingSimplyConvertibleTo.h"
 #include "Traits/IsContiguousContainer.h"
 #include <type_traits>
 
@@ -165,7 +166,7 @@ public:
 	>
 	inline BuilderType& AppendChar(AppendedCharType Char)
 	{
-		if constexpr (FPlatformString::IsCharEncodingSimplyConvertibleTo<AppendedCharType, CharType>())
+		if constexpr (TIsCharEncodingSimplyConvertibleTo_V<AppendedCharType, CharType>)
 	{
 		EnsureAdditionalCapacity(1);
 			*CurPos++ = (CharType)Char;
diff --git a/Engine/Source/Runtime/Core/Public/Traits/IsCharEncodingCompatibleWith.h b/Engine/Source/Runtime/Core/Public/Traits/IsCharEncodingCompatibleWith.h
new file mode 100644
--- /dev/null
+++ b/Engine/Source/Runtime/Core/Public/Traits/IsCharEncodingCompatibleWith.h
@@ -0,0 +1,69 @@
+// Copyright Epic Games, Inc. All Rights Reserved.
+
+#pragma once
+
+#include <type_traits>
+#include "Traits/IsCharType.h"
+
+#ifndef PLATFORM_TCHAR_IS_CHAR16
+	#error "Traits/IsCharEncodingCompatibleWith.h should be included after platform headers"
+#endif
+
+namespace UE::Core::Private
+{
+	template <typename SrcEncoding, typename DestEncoding>
+	static constexpr bool IsCharEncodingCompatibleWithImpl()
+	{
+		if constexpr (std::is_same_v<SrcEncoding, DestEncoding>)
+		{
+			return true;
+		}
+		else if constexpr (std::is_same_v<SrcEncoding, ANSICHAR> && std::is_same_v<DestEncoding, UTF8CHAR>)
+		{
+			return true;
+		}
+		else if constexpr (std::is_same_v<SrcEncoding, UCS2CHAR> && std::is_same_v<DestEncoding, UTF16CHAR>)
+		{
+			return true;
+		}
+		else if constexpr (std::is_same_v<SrcEncoding, WIDECHAR> && std::is_same_v<DestEncoding, UCS2CHAR>)
+		{
+			return true;
+		}
+		else if constexpr (std::is_same_v<SrcEncoding, UCS2CHAR> && std::is_same_v<DestEncoding, WIDECHAR>)
+		{
+			return true;
+		}
+#if PLATFORM_TCHAR_IS_CHAR16
+		else if constexpr (std::is_same_v<SrcEncoding, WIDECHAR> && std::is_same_v<DestEncoding, wchar_t>)
+		{
+			return true;
+		}
+		else if constexpr (std::is_same_v<SrcEncoding, wchar_t> && std::is_same_v<DestEncoding, WIDECHAR>)
+		{
+			return true;
+		}
+#endif
+		else
+		{
+			return false;
+		}
+	};
+}
+
+/**
+ * Trait which tests if a source char type is binary compatible with a destination char type.
+
+ * This is not commutative.  For example, ANSI is compatible with UTF-8, but UTF-8 is not compatible with ANSI.
+ */
+template <typename SrcEncoding, typename DestEncoding>
+struct TIsCharEncodingCompatibleWith
+{
+	static_assert(TIsCharType<SrcEncoding >::Value, "SrcEncoding is not a character encoding type");
+	static_assert(TIsCharType<DestEncoding>::Value, "DestEncoding is not a character encoding type");
+
+	static constexpr bool Value = UE::Core::Private::IsCharEncodingCompatibleWithImpl<std::remove_cv_t<SrcEncoding>, std::remove_cv_t<DestEncoding>>();
+};
+
+template <typename SrcEncoding, typename DestEncoding>
+constexpr inline bool TIsCharEncodingCompatibleWith_V = TIsCharEncodingCompatibleWith<SrcEncoding, DestEncoding>::Value;
diff --git a/Engine/Source/Runtime/Core/Public/Traits/IsCharEncodingSimplyConvertibleTo.h b/Engine/Source/Runtime/Core/Public/Traits/IsCharEncodingSimplyConvertibleTo.h
new file mode 100644
--- /dev/null
+++ b/Engine/Source/Runtime/Core/Public/Traits/IsCharEncodingSimplyConvertibleTo.h
@@ -0,0 +1,46 @@
+// Copyright Epic Games, Inc. All Rights Reserved.
+
+#pragma once
+
+#include <type_traits>
+#include "Traits/IsCharEncodingCompatibleWith.h"
+#include "Traits/IsFixedWidthCharEncoding.h"
+#include "Traits/IsCharType.h"
+
+namespace UE::Core::Private
+{
+	template <typename SrcEncoding, typename DestEncoding>
+	static constexpr bool IsCharEncodingSimplyConvertibleToImpl()
+	{
+		if constexpr (TIsCharEncodingCompatibleWith_V<SrcEncoding, DestEncoding>)
+		{
+			// Binary-compatible conversions are always simple
+			return true;
+		}
+		else if constexpr (TIsFixedWidthCharEncoding_V<SrcEncoding> && sizeof(DestEncoding) >= sizeof(SrcEncoding))
+		{
+			// Converting from a fixed-width encoding to a wider or same encoding should always be possible,
+			// as should ANSICHAR->UTF8CHAR and UCS2CHAR->UTF16CHAR
+			return true;
+		}
+		else
+		{
+			return false;
+		}
+	}
+}
+
+/**
+ * Trait which tests if code units of the source encoding are simply convertible (i.e. by assignment) to the destination encoding.
+ */
+template <typename SrcEncoding, typename DestEncoding>
+struct TIsCharEncodingSimplyConvertibleTo
+{
+	static_assert(TIsCharType<SrcEncoding >::Value, "SrcEncoding is not a character encoding type");
+	static_assert(TIsCharType<DestEncoding>::Value, "DestEncoding is not a character encoding type");
+
+	static constexpr bool Value = UE::Core::Private::IsCharEncodingSimplyConvertibleToImpl<SrcEncoding, DestEncoding>();
+};
+
+template <typename SrcEncoding, typename DestEncoding>
+constexpr inline bool TIsCharEncodingSimplyConvertibleTo_V = TIsCharEncodingSimplyConvertibleTo<SrcEncoding, DestEncoding>::Value;
diff --git a/Engine/Source/Runtime/Core/Public/Traits/IsCharType.h b/Engine/Source/Runtime/Core/Public/Traits/IsCharType.h
--- a/Engine/Source/Runtime/Core/Public/Traits/IsCharType.h
+++ b/Engine/Source/Runtime/Core/Public/Traits/IsCharType.h
@@ -5,7 +5,7 @@
 #include "CoreTypes.h"
 
 /**
- * Type trait which tests if a type is a char type.
+ * Type trait which tests if a type is a character encoding type.
  */
 template<typename T> struct TIsCharType            { enum { Value = false }; };
 template<>           struct TIsCharType<ANSICHAR>  { enum { Value = true  }; };
@@ -19,3 +19,6 @@ template<>           struct TIsCharType<UTF32CHAR> { enum { Value = true  }; };
 #if PLATFORM_TCHAR_IS_CHAR16
 template<>           struct TIsCharType<wchar_t>   { enum { Value = true  }; };
 #endif
+
+template <typename T>
+constexpr inline bool TIsCharType_V = TIsCharType<T>::Value;
diff --git a/Engine/Source/Runtime/Core/Public/Traits/IsFixedWidthCharEncoding.h b/Engine/Source/Runtime/Core/Public/Traits/IsFixedWidthCharEncoding.h
new file mode 100644
--- /dev/null
+++ b/Engine/Source/Runtime/Core/Public/Traits/IsFixedWidthCharEncoding.h
@@ -0,0 +1,40 @@
+// Copyright Epic Games, Inc. All Rights Reserved.
+
+#pragma once
+
+#include <type_traits>
+#include "Traits/IsCharType.h"
+
+#ifndef PLATFORM_TCHAR_IS_CHAR16
+	#error "Traits/IsFixedWidthCharEncoding.h should be included after platform headers"
+#endif
+
+namespace UE::Core::Private
+{
+	/**
+	 * Tests whether an encoding has fixed-width characters
+	 */
+	template <typename Encoding>
+	static constexpr bool IsFixedWidthEncodingImpl()
+	{
+		return
+			std::is_same_v<Encoding, ANSICHAR> ||
+			std::is_same_v<Encoding, UCS2CHAR> || // this may not be true when PLATFORM_UCS2CHAR_IS_UTF16CHAR == 1, but this is the legacy behavior
+			std::is_same_v<Encoding, WIDECHAR> || // the UCS2CHAR comment also applies to WIDECHAR
+#if PLATFORM_TCHAR_IS_CHAR16
+			std::is_same_v<Encoding, wchar_t> || // the UCS2CHAR comment also applies to wchar_t
+#endif
+			std::is_same_v<Encoding, UTF32CHAR>;
+	}
+}
+
+template <typename Encoding>
+struct TIsFixedWidthCharEncoding
+{
+	static_assert(TIsCharType_V<Encoding>, "Encoding is not a character encoding type");
+
+	static constexpr bool Value = UE::Core::Private::IsFixedWidthEncodingImpl<std::remove_cv_t<Encoding>>();
+};
+
+template <typename Encoding>
+constexpr inline bool TIsFixedWidthCharEncoding_V = TIsFixedWidthCharEncoding<Encoding>::Value;
diff --git a/Engine/Source/Runtime/Engine/Public/ReplayTypes.h b/Engine/Source/Runtime/Engine/Public/ReplayTypes.h
--- a/Engine/Source/Runtime/Engine/Public/ReplayTypes.h
+++ b/Engine/Source/Runtime/Engine/Public/ReplayTypes.h
@@ -11,6 +11,7 @@
 #include "Net/Common/Packets/PacketTraits.h"
 #include "IPAddress.h"
 #include "Serialization/BitReader.h"
+#include "Traits/IsCharEncodingCompatibleWith.h"
 #include "ReplayTypes.generated.h"
 
 ENGINE_API DECLARE_LOG_CATEGORY_EXTERN(LogDemo, Log, All);
@@ -459,7 +460,7 @@ public:
 private:
 	void EnableFastStringSerialization()
 	{
-		if (FPlatformString::IsCharEncodingCompatibleWith<WIDECHAR, TCHAR>())
+		if constexpr (TIsCharEncodingCompatibleWith<WIDECHAR, TCHAR>::Value)
 		{
 			Archive.SetForceUnicode(true);
 		}
@@ -467,7 +468,7 @@ private:
 
 	void RestoreStringSerialization()
 	{
-		if (FPlatformString::IsCharEncodingCompatibleWith<WIDECHAR, TCHAR>())
+		if constexpr (TIsCharEncodingCompatibleWith<WIDECHAR, TCHAR>::Value)
 		{
 			Archive.SetForceUnicode(bWasUnicode);
 		}
