Missing UTF32->UTF8 conversions added.

diff --git a/Engine/Source/Runtime/Core/Private/GenericPlatform/GenericPlatformString.cpp b/Engine/Source/Runtime/Core/Private/GenericPlatform/GenericPlatformString.cpp
--- a/Engine/Source/Runtime/Core/Private/GenericPlatform/GenericPlatformString.cpp
+++ b/Engine/Source/Runtime/Core/Private/GenericPlatform/GenericPlatformString.cpp
@@ -604,6 +604,12 @@ namespace UE::Core::Private
 		int32 Result = ConvertToUTF8(Dest, INT32_MAX, Source, SourceLen);
 		return Result;
 	}
+	int32 GetConvertedLength(const UTF8CHAR*, const UTF32CHAR* Source, int32 SourceLen)
+	{
+		TCountingOutputIterator<UTF8CHAR> Dest;
+		int32 Result = ConvertToUTF8(Dest, INT32_MAX, Source, SourceLen);
+		return Result;
+	}
 	int32 GetConvertedLength(const ANSICHAR*, const UTF8CHAR* Source, int32 SourceLen)
 	{
 		TCountingOutputIterator<ANSICHAR> Dest;
@@ -647,6 +653,14 @@ namespace UE::Core::Private
 		}
 		return Dest;
 	}
+	UTF8CHAR* Convert(UTF8CHAR* Dest, int32 DestLen, const UTF32CHAR* Src, int32 SrcLen)
+	{
+		if (ConvertToUTF8(Dest, DestLen, Src, SrcLen) == -1)
+		{
+			return nullptr;
+		}
+		return Dest;
+	}
 	ANSICHAR* Convert(ANSICHAR* Dest, int32 DestLen, const UTF8CHAR* Src, int32 SrcLen)
 	{
 		if (ConvertFromUTF8<ANSICHAR>(Dest, DestLen, Src, SrcLen) == -1)
diff --git a/Engine/Source/Runtime/Core/Public/GenericPlatform/GenericPlatformString.h b/Engine/Source/Runtime/Core/Public/GenericPlatform/GenericPlatformString.h
--- a/Engine/Source/Runtime/Core/Public/GenericPlatform/GenericPlatformString.h
+++ b/Engine/Source/Runtime/Core/Public/GenericPlatform/GenericPlatformString.h
@@ -17,6 +17,7 @@ namespace UE::Core::Private
 	CORE_API int32 GetConvertedLength(const UTF8CHAR* Dest, const ANSICHAR*  Src, int32 SrcLen);
 	CORE_API int32 GetConvertedLength(const UTF8CHAR* Dest, const WIDECHAR*  Src, int32 SrcLen);
 	CORE_API int32 GetConvertedLength(const UTF8CHAR* Dest, const UCS2CHAR*  Src, int32 SrcLen);
+	CORE_API int32 GetConvertedLength(const UTF8CHAR* Dest, const UTF32CHAR* Src, int32 SrcLen);
 	CORE_API int32 GetConvertedLength(const ANSICHAR* Dest, const UTF8CHAR*  Src, int32 SrcLen);
 	CORE_API int32 GetConvertedLength(const WIDECHAR* Dest, const UTF8CHAR*  Src, int32 SrcLen);
 	CORE_API int32 GetConvertedLength(const UCS2CHAR* Dest, const UTF8CHAR*  Src, int32 SrcLen);
@@ -24,6 +25,7 @@ namespace UE::Core::Private
 	CORE_API UTF8CHAR* Convert(UTF8CHAR* Dest, int32 DestLen, const ANSICHAR*  Src, int32 SrcLen);
 	CORE_API UTF8CHAR* Convert(UTF8CHAR* Dest, int32 DestLen, const WIDECHAR*  Src, int32 SrcLen);
 	CORE_API UTF8CHAR* Convert(UTF8CHAR* Dest, int32 DestLen, const UCS2CHAR*  Src, int32 SrcLen);
+	CORE_API UTF8CHAR* Convert(UTF8CHAR* Dest, int32 DestLen, const UTF32CHAR* Src, int32 SrcLen);
 	CORE_API ANSICHAR* Convert(ANSICHAR* Dest, int32 DestLen, const UTF8CHAR*  Src, int32 SrcLen);
 	CORE_API WIDECHAR* Convert(WIDECHAR* Dest, int32 DestLen, const UTF8CHAR*  Src, int32 SrcLen);
 	CORE_API UCS2CHAR* Convert(UCS2CHAR* Dest, int32 DestLen, const UTF8CHAR*  Src, int32 SrcLen);
