Missing header added.

diff --git a/Engine/Source/Runtime/Core/Public/Templates/IsArrayOrRefOfTypeByPredicate.h b/Engine/Source/Runtime/Core/Public/Templates/IsArrayOrRefOfTypeByPredicate.h
new file mode 100644
--- /dev/null
+++ b/Engine/Source/Runtime/Core/Public/Templates/IsArrayOrRefOfTypeByPredicate.h
@@ -0,0 +1,30 @@
+// Copyright Epic Games, Inc. All Rights Reserved.
+
+#pragma once
+
+#include "CoreTypes.h"
+
+
+/**
+ * Type trait which returns true if the type T is an array or a reference to an array of a type which matches the predicate.
+ */
+template <typename T, template <typename> class Predicate>
+struct TIsArrayOrRefOfTypeByPredicate
+{
+	enum { Value = false };
+};
+
+template <typename ArrType, template <typename> class Predicate> struct TIsArrayOrRefOfTypeByPredicate<               ArrType[], Predicate> { enum { Value = Predicate<ArrType>::Value }; };
+template <typename ArrType, template <typename> class Predicate> struct TIsArrayOrRefOfTypeByPredicate<const          ArrType[], Predicate> { enum { Value = Predicate<ArrType>::Value }; };
+template <typename ArrType, template <typename> class Predicate> struct TIsArrayOrRefOfTypeByPredicate<      volatile ArrType[], Predicate> { enum { Value = Predicate<ArrType>::Value }; };
+template <typename ArrType, template <typename> class Predicate> struct TIsArrayOrRefOfTypeByPredicate<const volatile ArrType[], Predicate> { enum { Value = Predicate<ArrType>::Value }; };
+
+template <typename ArrType, unsigned int N, template <typename> class Predicate> struct TIsArrayOrRefOfTypeByPredicate<               ArrType[N], Predicate> { enum { Value = Predicate<ArrType>::Value }; };
+template <typename ArrType, unsigned int N, template <typename> class Predicate> struct TIsArrayOrRefOfTypeByPredicate<const          ArrType[N], Predicate> { enum { Value = Predicate<ArrType>::Value }; };
+template <typename ArrType, unsigned int N, template <typename> class Predicate> struct TIsArrayOrRefOfTypeByPredicate<      volatile ArrType[N], Predicate> { enum { Value = Predicate<ArrType>::Value }; };
+template <typename ArrType, unsigned int N, template <typename> class Predicate> struct TIsArrayOrRefOfTypeByPredicate<const volatile ArrType[N], Predicate> { enum { Value = Predicate<ArrType>::Value }; };
+
+template <typename ArrType, unsigned int N, template <typename> class Predicate> struct TIsArrayOrRefOfTypeByPredicate<               ArrType(&)[N], Predicate> { enum { Value = Predicate<ArrType>::Value }; };
+template <typename ArrType, unsigned int N, template <typename> class Predicate> struct TIsArrayOrRefOfTypeByPredicate<const          ArrType(&)[N], Predicate> { enum { Value = Predicate<ArrType>::Value }; };
+template <typename ArrType, unsigned int N, template <typename> class Predicate> struct TIsArrayOrRefOfTypeByPredicate<      volatile ArrType(&)[N], Predicate> { enum { Value = Predicate<ArrType>::Value }; };
+template <typename ArrType, unsigned int N, template <typename> class Predicate> struct TIsArrayOrRefOfTypeByPredicate<const volatile ArrType(&)[N], Predicate> { enum { Value = Predicate<ArrType>::Value }; };
