IOS LowLevelMemTracker: Add SmallOSAlloc ::malloc calls to LLM.
#rnx

diff --git a/Engine/Source/Runtime/Core/Private/HAL/MallocBinned.cpp b/Engine/Source/Runtime/Core/Private/HAL/MallocBinned.cpp
--- a/Engine/Source/Runtime/Core/Private/HAL/MallocBinned.cpp
+++ b/Engine/Source/Runtime/Core/Private/HAL/MallocBinned.cpp
@@ -884,6 +884,7 @@ struct FMallocBinned::Private
 	static FORCEINLINE void SmallOSFree(FMallocBinned& Allocator, void* Ptr, SIZE_T Size)
 	{
 #if PLATFORM_IOS
+		LLM(FLowLevelMemTracker::Get().OnLowLevelFree(ELLMTracker::Platform, Ptr));
 		::free(Ptr);
 		Allocator.bNanoMallocAvailable = true;
 #else
@@ -964,6 +965,7 @@ struct FMallocBinned::Private
 		{
 			UE_LOG(LogTemp, Warning, TEXT("malloc failure allocating %d, error code: %d"), NewSize, errno);
 		}
+		LLM(FLowLevelMemTracker::Get().OnLowLevelAlloc(ELLMTracker::Platform, Ptr, NewSize));
 		return Ptr;
 #else
 		(void)OutActualSize;
@@ -1433,7 +1435,10 @@ void* FMallocBinned::Realloc( void* Ptr, SIZE_T NewSize, uint32 Alignment )
             const SIZE_T OldSize = malloc_size(Ptr);
             if ((NewSize <= Private::SMALL_BLOCK_POOL_SIZE) && (Alignment <= Private::DEFAULT_BINNED_ALLOCATOR_ALIGNMENT))
             {
+                LLM_PLATFORM_SCOPE(ELLMTag::FMalloc);
+                LLM(FLowLevelMemTracker::Get().OnLowLevelFree(ELLMTracker::Platform, Ptr));
                 NewPtr = ::realloc(Ptr, NewSize);
+                LLM(FLowLevelMemTracker::Get().OnLowLevelAlloc(ELLMTracker::Platform, NewPtr, NewSize));
                 bFallback = !FPlatformMemory::PtrIsFromNanoMalloc(NewPtr);
             }
 
