IoStore:
Default to using the current build crypto keys when reading previous builds for layout matching
Fixed incorrect order of previous build chunks when layout matching against a previous build

diff --git a/Engine/Source/Developer/IoStoreUtilities/Private/IoStoreUtilities.cpp b/Engine/Source/Developer/IoStoreUtilities/Private/IoStoreUtilities.cpp
--- a/Engine/Source/Developer/IoStoreUtilities/Private/IoStoreUtilities.cpp
+++ b/Engine/Source/Developer/IoStoreUtilities/Private/IoStoreUtilities.cpp
@@ -6370,11 +6370,14 @@ bool ParseContainerGenerationArguments(FIoStoreArguments& Arguments, FIoStoreWri
 	WriterSettings.bEnableFileRegions = Arguments.bFileRegions;
 
 	FString PatchReferenceCryptoKeysFilename;
-	FKeyChain PatchKeyChain;
 	if (FParse::Value(FCommandLine::Get(), TEXT("PatchCryptoKeys="), PatchReferenceCryptoKeysFilename))
 	{
 		KeyChainUtilities::LoadKeyChainFromFile(PatchReferenceCryptoKeysFilename, Arguments.PatchKeyChain);
 	}
+	else
+	{
+		Arguments.PatchKeyChain = Arguments.KeyChain;
+	}
 
 	return true;
 }
diff --git a/Engine/Source/Runtime/Core/Private/IO/IoStore.cpp b/Engine/Source/Runtime/Core/Private/IO/IoStore.cpp
--- a/Engine/Source/Runtime/Core/Private/IO/IoStore.cpp
+++ b/Engine/Source/Runtime/Core/Private/IO/IoStore.cpp
@@ -756,21 +756,32 @@ public:
 
 		for (const TUniquePtr<FIoStoreReader>& PatchSourceReader : PatchSourceReaders)
 		{
-			PatchSourceReader->EnumerateChunks([this, &PrevEntryLink](const FIoStoreTocChunkInfo& ChunkInfo)
+			TArray<TPair<uint64, FLayoutEntry*>> LayoutEntriesWithOffsets;
+			PatchSourceReader->EnumerateChunks([this, &PrevEntryLink, &LayoutEntriesWithOffsets](const FIoStoreTocChunkInfo& ChunkInfo)
 				{
 					FLayoutEntry* PreviousBuildEntry = new FLayoutEntry();
-					LayoutEntries.Add(PreviousBuildEntry);
-
 					PreviousBuildEntry->Hash = ChunkInfo.Hash;
 					PreviousBuildEntry->PartitionIndex = ChunkInfo.PartitionIndex;
 					PreviousBuildEntry->CompressedSize = ChunkInfo.CompressedSize;
-					PrevEntryLink->Next = PreviousBuildEntry;
-					PreviousBuildEntry->Prev = PrevEntryLink;
-					PrevEntryLink = PreviousBuildEntry;
+					LayoutEntriesWithOffsets.Emplace(ChunkInfo.Offset, PreviousBuildEntry);
 					PreviousBuildLayoutEntryByChunkId.Add(ChunkInfo.Id, PreviousBuildEntry);
 					return true;
 				});
 
+			// Sort entries by offset
+			Algo::Sort(LayoutEntriesWithOffsets, [](const TPair<uint64, FLayoutEntry*>& A, const TPair<uint64, FLayoutEntry*>& B)
+				{
+					return A.Get<0>() < B.Get<0>();
+				});
+
+			for (const TPair<uint64, FLayoutEntry*>& EntryWithOffset : LayoutEntriesWithOffsets)
+			{
+				FLayoutEntry* PreviousBuildEntry = EntryWithOffset.Get<1>();
+				LayoutEntries.Add(PreviousBuildEntry);
+				PrevEntryLink->Next = PreviousBuildEntry;
+				PreviousBuildEntry->Prev = PrevEntryLink;
+				PrevEntryLink = PreviousBuildEntry;
+			}
 			if (!ContainerSettings.bGenerateDiffPatch)
 			{
 				break;
