Fix shipping compile error.
#rnx

diff --git a/Engine/Source/Runtime/Core/Private/ProfilingDebugging/CookStats.cpp b/Engine/Source/Runtime/Core/Private/ProfilingDebugging/CookStats.cpp
--- a/Engine/Source/Runtime/Core/Private/ProfilingDebugging/CookStats.cpp
+++ b/Engine/Source/Runtime/Core/Private/ProfilingDebugging/CookStats.cpp
@@ -57,7 +57,7 @@ void TracePackageAssetClass(uint64 InId, const FStringView InName)
 
 bool ShouldTracePackageInfo()
 {
-	return CookChannel.IsEnabled();
+	return UE_TRACE_CHANNELEXPR_IS_ENABLED(CookChannel);
 }
 
 #endif
