Add some additional sanity checks/logging to FPersistentStorageManager inital scan.

diff --git a/Engine/Source/Runtime/Core/Public/HAL/IPlatformFileManagedStorageWrapper.h b/Engine/Source/Runtime/Core/Public/HAL/IPlatformFileManagedStorageWrapper.h
--- a/Engine/Source/Runtime/Core/Public/HAL/IPlatformFileManagedStorageWrapper.h
+++ b/Engine/Source/Runtime/Core/Public/HAL/IPlatformFileManagedStorageWrapper.h
@@ -376,10 +376,19 @@ public:
 
 							// This must be done under the lock because another thread may be modifying or deleting the file while we scan
 							FFileStatData StatData = IPlatformFile::GetPlatformPhysical().GetStatData(FilenameOrDirectory);
-							if (ensure(StatData.bIsValid))
+							if (ensureAlways(StatData.bIsValid))
 							{
+								if (!ensureAlways(StatData.FileSize >= 0))
+								{
+									UE_LOG(LogPlatformFileManagedStorage, Error, TEXT("Invalid File Size for %s!"), FilenameOrDirectory);
+								}
+								
 								Man.AddOrUpdateFile(File, StatData.FileSize);
 							}
+							else
+							{
+								UE_LOG(LogPlatformFileManagedStorage, Error, TEXT("Invalid Stat Data for %s!"), FilenameOrDirectory);
+							}
 						}
 					}
 
