Resolve mac\ios typecast warnings

#rnx

diff --git a/Engine/Source/Runtime/Core/Private/Apple/ApplePlatformDebugEvents.cpp b/Engine/Source/Runtime/Core/Private/Apple/ApplePlatformDebugEvents.cpp
--- a/Engine/Source/Runtime/Core/Private/Apple/ApplePlatformDebugEvents.cpp
+++ b/Engine/Source/Runtime/Core/Private/Apple/ApplePlatformDebugEvents.cpp
@@ -209,7 +209,7 @@ uint16 FApplePlatformDebugEvents::GetEventCode(FString String)
 			}
 			else
 			{
-				Code = Names.Num() + 1;
+				Code = (uint16)Names.Num() + 1;
 				check(Code < 16384);
 				Names.Add(Hash, Code);
 				UE_LOG(LogInstruments, Display, TEXT("New Event Code: %u : %s"), (uint32)Code, *String);
diff --git a/Engine/Source/Runtime/Core/Private/Apple/ApplePlatformFile.cpp b/Engine/Source/Runtime/Core/Private/Apple/ApplePlatformFile.cpp
--- a/Engine/Source/Runtime/Core/Private/Apple/ApplePlatformFile.cpp
+++ b/Engine/Source/Runtime/Core/Private/Apple/ApplePlatformFile.cpp
@@ -538,7 +538,7 @@ void FApplePlatformFile::SetTimeStamp(const TCHAR* Filename, const FDateTime Dat
 	// change the modification time only
 	struct utimbuf Times;
 	Times.actime = FileInfo.st_atime;
-	Times.modtime = (DateTime - MacEpoch).GetTotalSeconds();
+	Times.modtime = (time_t)(DateTime - MacEpoch).GetTotalSeconds();
 	utime(TCHAR_TO_UTF8(*NormalizeFilename(Filename)), &Times);
 }
 
diff --git a/Engine/Source/Runtime/Core/Private/Apple/ApplePlatformMisc.cpp b/Engine/Source/Runtime/Core/Private/Apple/ApplePlatformMisc.cpp
--- a/Engine/Source/Runtime/Core/Private/Apple/ApplePlatformMisc.cpp
+++ b/Engine/Source/Runtime/Core/Private/Apple/ApplePlatformMisc.cpp
@@ -253,7 +253,7 @@ void GetBytesForFont(const NSString* InFontName, OUT TArray<uint8>& OutBytes)
 		}
 		SearchRange <<= 4;
 
-		uint16 RangeShift = (TableCount << 4) - SearchRange;
+		uint16 RangeShift = (uint16)((TableCount << 4) - SearchRange);
 
 		// Write font header (also called sfnt header, offset subtable)
 		FFontHeader* OffsetTable = (FFontHeader*)DataPtr;
@@ -302,13 +302,13 @@ TArray<uint8> FApplePlatformMisc::GetSystemFontBytes()
 {
 #if PLATFORM_MAC
 	// Gather some details about the system font
-	uint32 SystemFontSize = [NSFont systemFontSize];
+	uint32 SystemFontSize = (uint32)[NSFont systemFontSize];
 	NSString* SystemFontName = [NSFont systemFontOfSize:SystemFontSize].fontName;
 #elif PLATFORM_TVOS
 	NSString* SystemFontName = [UIFont preferredFontForTextStyle:UIFontTextStyleBody].fontName;
 #else
 	// Gather some details about the system font
-	uint32 SystemFontSize = [UIFont systemFontSize];
+	uint32 SystemFontSize = (uint32)[UIFont systemFontSize];
 	NSString* SystemFontName = [UIFont systemFontOfSize:SystemFontSize].fontName;
 #endif
 
diff --git a/Engine/Source/Runtime/Core/Private/IOS/IOSPlatformFile.cpp b/Engine/Source/Runtime/Core/Private/IOS/IOSPlatformFile.cpp
--- a/Engine/Source/Runtime/Core/Private/IOS/IOSPlatformFile.cpp
+++ b/Engine/Source/Runtime/Core/Private/IOS/IOSPlatformFile.cpp
@@ -741,7 +741,7 @@ void FIOSPlatformFile::SetTimeStamp(const TCHAR* Filename, const FDateTime DateT
 	// change the modification time only
 	struct utimbuf Times;
 	Times.actime = FileInfo.st_atime;
-	Times.modtime = (DateTime - IOSEpoch).GetTotalSeconds();
+	Times.modtime = (time_t)(DateTime - IOSEpoch).GetTotalSeconds();
 	utime(TCHAR_TO_UTF8(*IOSFilename), &Times);
 }
 
diff --git a/Engine/Source/Runtime/Core/Private/IOS/IOSPlatformMisc.cpp b/Engine/Source/Runtime/Core/Private/IOS/IOSPlatformMisc.cpp
--- a/Engine/Source/Runtime/Core/Private/IOS/IOSPlatformMisc.cpp
+++ b/Engine/Source/Runtime/Core/Private/IOS/IOSPlatformMisc.cpp
@@ -236,7 +236,7 @@ int FIOSPlatformMisc::GetBatteryLevel()
 float FIOSPlatformMisc::GetBrightness()
 {
 #if !PLATFORM_TVOS
-	return [UIScreen mainScreen].brightness;
+	return (float)[UIScreen mainScreen].brightness;
 #else
 	return 1.0f;
 #endif // !PLATFORM_TVOS
diff --git a/Engine/Source/Runtime/Core/Private/Mac/MacPlatformMisc.cpp b/Engine/Source/Runtime/Core/Private/Mac/MacPlatformMisc.cpp
--- a/Engine/Source/Runtime/Core/Private/Mac/MacPlatformMisc.cpp
+++ b/Engine/Source/Runtime/Core/Private/Mac/MacPlatformMisc.cpp
@@ -2376,7 +2376,7 @@ void FMacCrashContext::GenerateCrashInfoAndLaunchReporter() const
 		posix_spawnattr_init(&SpawnAttr);
 
 		{
-			uint32 SpawnFlags = POSIX_SPAWN_SETPGROUP;
+			uint16 SpawnFlags = POSIX_SPAWN_SETPGROUP;
 			posix_spawnattr_setflags(&SpawnAttr, SpawnFlags);
 		}
 
@@ -2801,7 +2801,7 @@ T GetMacGPUStat(TMap<FString, float> const& Stats, FString StatName)
 	T Result = (T)0;
 	if(Stats.Contains(StatName))
 	{
-		Result = Stats.FindRef(StatName);
+		Result = (T)Stats.FindRef(StatName);
 	}
 	return Result;
 }
