Direct FTCHARToUTF8_Convert usage deprecated and replaced with StringCast or FPlatformString.

diff --git a/Engine/Plugins/Runtime/OpenXR/Source/OpenXRHMD/Private/OpenXRHMDModule.cpp b/Engine/Plugins/Runtime/OpenXR/Source/OpenXRHMD/Private/OpenXRHMDModule.cpp
--- a/Engine/Plugins/Runtime/OpenXR/Source/OpenXRHMD/Private/OpenXRHMDModule.cpp
+++ b/Engine/Plugins/Runtime/OpenXR/Source/OpenXRHMD/Private/OpenXRHMDModule.cpp
@@ -637,9 +637,9 @@ bool FOpenXRHMDModule::InitInstance()
 	Info.type = XR_TYPE_INSTANCE_CREATE_INFO;
 	Info.next = nullptr;
 	Info.createFlags = 0;
-	FTCHARToUTF8_Convert::Convert(Info.applicationInfo.applicationName, XR_MAX_APPLICATION_NAME_SIZE, *AppName, AppName.Len() + 1);
+	FPlatformString::Convert((UTF8CHAR*)Info.applicationInfo.applicationName, XR_MAX_APPLICATION_NAME_SIZE, *AppName, AppName.Len() + 1);
 	Info.applicationInfo.applicationVersion = static_cast<uint32>(BuildSettings::GetCurrentChangelist()) | (BuildSettings::IsLicenseeVersion() ? 0x80000000 : 0);
-	FTCHARToUTF8_Convert::Convert(Info.applicationInfo.engineName, XR_MAX_ENGINE_NAME_SIZE, *EngineName, EngineName.Len() + 1);
+	FPlatformString::Convert((UTF8CHAR*)Info.applicationInfo.engineName, XR_MAX_ENGINE_NAME_SIZE, *EngineName, EngineName.Len() + 1);
 	Info.applicationInfo.engineVersion = (uint32)(FEngineVersion::Current().GetMinor() << 16 | FEngineVersion::Current().GetPatch());
 	Info.applicationInfo.apiVersion = XR_CURRENT_API_VERSION;
 
diff --git a/Engine/Plugins/VirtualProduction/RemoteControl/Source/WebRemoteControl/Private/WebRemoteControlUtils.cpp b/Engine/Plugins/VirtualProduction/RemoteControl/Source/WebRemoteControl/Private/WebRemoteControlUtils.cpp
--- a/Engine/Plugins/VirtualProduction/RemoteControl/Source/WebRemoteControl/Private/WebRemoteControlUtils.cpp
+++ b/Engine/Plugins/VirtualProduction/RemoteControl/Source/WebRemoteControl/Private/WebRemoteControlUtils.cpp
@@ -14,15 +14,15 @@ void WebRemoteControlUtils::ConvertToTCHAR(TConstArrayView<uint8> InUTF8Payload,
 void WebRemoteControlUtils::ConvertToUTF8(TConstArrayView<uint8> InTCHARPayload, TArray<uint8>& OutUTF8Payload)
 {
 	int32 StartIndex = OutUTF8Payload.Num();
-	OutUTF8Payload.AddUninitialized(FTCHARToUTF8_Convert::ConvertedLength((TCHAR*)InTCHARPayload.GetData(), InTCHARPayload.Num() / sizeof(TCHAR)) * sizeof(ANSICHAR));
-	FTCHARToUTF8_Convert::Convert((ANSICHAR*)(OutUTF8Payload.GetData() + StartIndex), (OutUTF8Payload.Num() - StartIndex) / sizeof(ANSICHAR), (TCHAR*)InTCHARPayload.GetData(), InTCHARPayload.Num() / sizeof(TCHAR));
+	OutUTF8Payload.AddUninitialized(FPlatformString::ConvertedLength<UTF8CHAR>((TCHAR*)InTCHARPayload.GetData(), InTCHARPayload.Num() / sizeof(TCHAR)) * sizeof(ANSICHAR));
+	FPlatformString::Convert((UTF8CHAR*)(OutUTF8Payload.GetData() + StartIndex), (OutUTF8Payload.Num() - StartIndex) / sizeof(ANSICHAR), (TCHAR*)InTCHARPayload.GetData(), InTCHARPayload.Num() / sizeof(TCHAR));
 }
 
 void WebRemoteControlUtils::ConvertToUTF8(const FString& InString, TArray<uint8>& OutUTF8Payload)
 {
 	int32 StartIndex = OutUTF8Payload.Num();
-	OutUTF8Payload.AddUninitialized(FTCHARToUTF8_Convert::ConvertedLength(*InString, InString.Len()) * sizeof(ANSICHAR));
-	FTCHARToUTF8_Convert::Convert((ANSICHAR*)(OutUTF8Payload.GetData() + StartIndex), (OutUTF8Payload.Num() - StartIndex) / sizeof(ANSICHAR), *InString, InString.Len());
+	OutUTF8Payload.AddUninitialized(FPlatformString::ConvertedLength<UTF8CHAR>(*InString, InString.Len()) * sizeof(ANSICHAR));
+	FPlatformString::Convert((UTF8CHAR*)(OutUTF8Payload.GetData() + StartIndex), (OutUTF8Payload.Num() - StartIndex) / sizeof(ANSICHAR), *InString, InString.Len());
 }
 
 TSharedRef<IStructSerializerBackend> WebRemoteControlUtils::CreateJsonSerializerBackend(FMemoryWriter& Writer)
diff --git a/Engine/Source/Runtime/Analytics/AnalyticsET/Private/AnalyticsProviderETEventCache.cpp b/Engine/Source/Runtime/Analytics/AnalyticsET/Private/AnalyticsProviderETEventCache.cpp
--- a/Engine/Source/Runtime/Analytics/AnalyticsET/Private/AnalyticsProviderETEventCache.cpp
+++ b/Engine/Source/Runtime/Analytics/AnalyticsET/Private/AnalyticsProviderETEventCache.cpp
@@ -68,7 +68,6 @@ namespace EventCacheStatic
 	// This is the buffer we will convert strings into UTF8 into, since it's difficult to convert them directly into a TArray<>, since it doesn't know how to resize itself.
 	// We also don't want to walk the string once to count the chars if we don't have to. so we pay the price to copy directly into a stack-allocated buffer most of the time,
 	// but let it spill over to a dynamic allocation for long strings.
-	//typedef TStringConversion<FTCHARToUTF8_Convert, ConversionBufferSize> FPayloadUTF8Converter;
 	typedef TStringBuilder<ConversionBufferSize> FJsonStringBuilder;
 
 	const ANSICHAR* PayloadTemplate = "{\"Events\":[]}";
@@ -98,9 +97,9 @@ namespace EventCacheStatic
 		// *** ORIGINAL, simpler code. But slower. ***
 		// convert directly into new array, precompute length
 		// get the string length and expand our buffer to fit it.
-		//const int32 StrLen = FTCHARToUTF8_Convert::ConvertedLength(Str, Len);
+		//const int32 StrLen = FPlatformString::ConvertedLength<UTF8CHAR>(Str, Len);
 		//UTF8Stream.SetNumUninitialized(OldLen + StrLen, false);
-		//FTCHARToUTF8_Convert::Convert(&UTF8Stream[OldLen], StrLen, Str, Len);
+		//FPlatformString::Convert((UTF8CHAR*)&UTF8Stream[OldLen], StrLen, Str, Len);
 
 		// optimistically allocate a bit of extra space and see if we fill up the buffer.
 		// If we do, lengthen the buffer a bit and try again.
@@ -114,7 +113,7 @@ namespace EventCacheStatic
 			// make space for the string
 			UTF8Stream.SetNumUninitialized(OldLen + StrLen, false);
 			// convert it to UTF8
-			int32 CharsWritten = FTCHARToUTF8_Convert::Convert((UTF8CHAR*)&UTF8Stream[OldLen], StrLen, Str, Len);
+			int32 CharsWritten = FPlatformString::Convert((UTF8CHAR*)&UTF8Stream[OldLen], StrLen, Str, Len) - (UTF8CHAR*)&UTF8Stream[OldLen];
 			// figure out how many characters were actually written 
 			if (CharsWritten >= 0)
 			{
@@ -129,10 +128,10 @@ namespace EventCacheStatic
 				// if we grow too much, give up and compute the true chars needed.
 				if (SizeMultiplier >= 2.0)
 				{
-					const int32 ActualCharsNeeded = FTCHARToUTF8_Convert::ConvertedLength(Str, Len);
+					const int32 ActualCharsNeeded = FPlatformString::ConvertedLength<UTF8CHAR>(Str, Len);
 					UTF8Stream.SetNumUninitialized(OldLen + ActualCharsNeeded, false);
 					// convert it to UTF8 using the known number of charts
-					FTCHARToUTF8_Convert::Convert((UTF8CHAR*)&UTF8Stream[OldLen], ActualCharsNeeded, Str, Len);
+					FPlatformString::Convert((UTF8CHAR*)&UTF8Stream[OldLen], ActualCharsNeeded, Str, Len);
 					bWroteFullString = true;
 				}
 			}
diff --git a/Engine/Source/Runtime/Core/Private/Serialization/CompactBinaryWriter.cpp b/Engine/Source/Runtime/Core/Private/Serialization/CompactBinaryWriter.cpp
--- a/Engine/Source/Runtime/Core/Private/Serialization/CompactBinaryWriter.cpp
+++ b/Engine/Source/Runtime/Core/Private/Serialization/CompactBinaryWriter.cpp
@@ -398,7 +398,7 @@ void FCbWriter::AddString(const FUtf8StringView Value)
 void FCbWriter::AddString(const FWideStringView Value)
 {
 	BeginField();
-	const uint32 Size = uint32(FTCHARToUTF8_Convert::ConvertedLength(Value.GetData(), Value.Len()));
+	const uint32 Size = uint32(FPlatformString::ConvertedLength<UTF8CHAR>(Value.GetData(), Value.Len()));
 	const uint32 SizeByteCount = MeasureVarUInt(Size);
 	const int64 Offset = Data.AddUninitialized(SizeByteCount + Size);
 	uint8* StringData = Data.GetData() + Offset;
@@ -406,7 +406,7 @@ void FCbWriter::AddString(const FWideStringView Value)
 	StringData += SizeByteCount;
 	if (Size > 0)
 	{
-		FTCHARToUTF8_Convert::Convert(reinterpret_cast<UTF8CHAR*>(StringData), Size, Value.GetData(), Value.Len());
+		FPlatformString::Convert(reinterpret_cast<UTF8CHAR*>(StringData), Size, Value.GetData(), Value.Len());
 	}
 	EndField(ECbFieldType::String);
 }
diff --git a/Engine/Source/Runtime/Core/Public/Containers/StringConv.h b/Engine/Source/Runtime/Core/Public/Containers/StringConv.h
--- a/Engine/Source/Runtime/Core/Public/Containers/StringConv.h
+++ b/Engine/Source/Runtime/Core/Public/Containers/StringConv.h
@@ -173,7 +173,7 @@ namespace StringConv
 	FUnused GetLegacyFromType(...);
 }
 
-namespace UE4StringConv_Private
+namespace UEStringConv_Private
 {
 	/**
 	 * This is a basic object which counts how many times it has been incremented
@@ -206,130 +206,132 @@ namespace UE4StringConv_Private
 	private:
 		int32 Counter;
 	};
-}
-
-// This should be replaced with Platform stuff when FPlatformString starts to know about UTF-8.
-class FTCHARToUTF8_Convert
-{
-	// Leave this until we convert all the call sites and can change FromType to this.
-	using IntendedToType = UTF8CHAR;
 
-public:
-	typedef TCHAR    FromType;
-	typedef ANSICHAR ToType;
-
-	/**
-	 * Convert Codepoint into UTF-8 characters.
-	 *
-	 * @param Codepoint Codepoint to expand into UTF-8 bytes
-	 * @param OutputIterator Output iterator to write UTF-8 bytes into
-	 * @param OutputIteratorByteSizeRemaining Maximum number of ANSI characters that can be written to OutputIterator
-	 * @return Number of characters written for Codepoint
-	 */
-	template <typename BufferType>
-	static int32 Utf8FromCodepoint(uint32 Codepoint, BufferType OutputIterator, uint32 OutputIteratorByteSizeRemaining)
+	// This should be replaced with Platform stuff when FPlatformString starts to know about UTF-8.
+	class FTCHARToUTF8_Convert
 	{
-		// Ensure we have at least one character in size to write
-		if (OutputIteratorByteSizeRemaining < sizeof(ToType))
-		{
-			return 0;
-		}
-
-		const BufferType OutputIteratorStartPosition = OutputIterator;
-
-		if (!StringConv::IsValidCodepoint(Codepoint))
-		{
-			Codepoint = UNICODE_BOGUS_CHAR_CODEPOINT;
-		}
-		else if (StringConv::IsHighSurrogate(Codepoint) || StringConv::IsLowSurrogate(Codepoint)) // UTF-8 Characters are not allowed to encode codepoints in the surrogate pair range
-		{
-			Codepoint = UNICODE_BOGUS_CHAR_CODEPOINT;
-		}
-
-		// Do the encoding...
-		if (Codepoint < 0x80)
-		{
-			*(OutputIterator++) = (ToType)Codepoint;
-		}
-		else if (Codepoint < 0x800)
+		// Leave this until we convert all the call sites and can change FromType to this.
+		using IntendedToType = UTF8CHAR;
+	
+	public:
+		typedef TCHAR    FromType;
+		typedef ANSICHAR ToType;
+	
+		/**
+		 * Convert Codepoint into UTF-8 characters.
+		 *
+		 * @param Codepoint Codepoint to expand into UTF-8 bytes
+		 * @param OutputIterator Output iterator to write UTF-8 bytes into
+		 * @param OutputIteratorByteSizeRemaining Maximum number of ANSI characters that can be written to OutputIterator
+		 * @return Number of characters written for Codepoint
+		 */
+		template <typename BufferType>
+		static int32 Utf8FromCodepoint(uint32 Codepoint, BufferType OutputIterator, uint32 OutputIteratorByteSizeRemaining)
 		{
-			if (OutputIteratorByteSizeRemaining >= 2)
+			// Ensure we have at least one character in size to write
+			if (OutputIteratorByteSizeRemaining < sizeof(ToType))
 			{
-				*(OutputIterator++) = (ToType)((Codepoint >> 6)         | 128 | 64);
-				*(OutputIterator++) = (ToType)((Codepoint       & 0x3F) | 128);
+				return 0;
 			}
-		}
-		else if (Codepoint < 0x10000)
-		{
-			if (OutputIteratorByteSizeRemaining >= 3)
+	
+			const BufferType OutputIteratorStartPosition = OutputIterator;
+	
+			if (!StringConv::IsValidCodepoint(Codepoint))
+			{
+				Codepoint = UNICODE_BOGUS_CHAR_CODEPOINT;
+			}
+			else if (StringConv::IsHighSurrogate(Codepoint) || StringConv::IsLowSurrogate(Codepoint)) // UTF-8 Characters are not allowed to encode codepoints in the surrogate pair range
+			{
+				Codepoint = UNICODE_BOGUS_CHAR_CODEPOINT;
+			}
+	
+			// Do the encoding...
+			if (Codepoint < 0x80)
+			{
+				*(OutputIterator++) = (ToType)Codepoint;
+			}
+			else if (Codepoint < 0x800)
 			{
-				*(OutputIterator++) = (ToType) ((Codepoint >> 12)        | 128 | 64 | 32);
-				*(OutputIterator++) = (ToType)(((Codepoint >> 6) & 0x3F) | 128);
-				*(OutputIterator++) = (ToType) ((Codepoint       & 0x3F) | 128);
+				if (OutputIteratorByteSizeRemaining >= 2)
+				{
+					*(OutputIterator++) = (ToType)((Codepoint >> 6)         | 128 | 64);
+					*(OutputIterator++) = (ToType)((Codepoint       & 0x3F) | 128);
+				}
 			}
+			else if (Codepoint < 0x10000)
+			{
+				if (OutputIteratorByteSizeRemaining >= 3)
+				{
+					*(OutputIterator++) = (ToType) ((Codepoint >> 12)        | 128 | 64 | 32);
+					*(OutputIterator++) = (ToType)(((Codepoint >> 6) & 0x3F) | 128);
+					*(OutputIterator++) = (ToType) ((Codepoint       & 0x3F) | 128);
+				}
+			}
+			else
+			{
+				if (OutputIteratorByteSizeRemaining >= 4)
+				{
+					*(OutputIterator++) = (ToType) ((Codepoint >> 18)         | 128 | 64 | 32 | 16);
+					*(OutputIterator++) = (ToType)(((Codepoint >> 12) & 0x3F) | 128);
+					*(OutputIterator++) = (ToType)(((Codepoint >> 6 ) & 0x3F) | 128);
+					*(OutputIterator++) = (ToType) ((Codepoint        & 0x3F) | 128);
+				}
+			}
+	
+			return UE_PTRDIFF_TO_INT32(OutputIterator - OutputIteratorStartPosition);
 		}
-		else
+	
+		/**
+		 * Converts a Source string into UTF8 and stores it in Dest.
+		 *
+		 * @param Dest      The destination output iterator. Usually ANSICHAR*, but you can supply your own output iterator.
+		 *                  One can determine the number of characters written by checking the offset of Dest when the function returns.
+		 * @param DestLen   The length of the destination buffer.
+		 * @param Source    The source string to convert.
+		 * @param SourceLen The length of the source string.
+		 * @return          The number of bytes written to Dest, up to DestLen, or -1 if the entire Source string could did not fit in DestLen bytes.
+		 */
+		template <
+			typename SrcBufferType,
+			std::enable_if_t<TIsCharEncodingCompatibleWith_V<SrcBufferType, FromType>>* = nullptr
+		>
+		static FORCEINLINE int32 Convert(IntendedToType* Dest, int32 DestLen, const SrcBufferType* Source, int32 SourceLen)
 		{
-			if (OutputIteratorByteSizeRemaining >= 4)
+			IntendedToType* Result = FPlatformString::Convert(Dest, DestLen, (const FromType*)Source, SourceLen);
+			if (!Result)
 			{
-				*(OutputIterator++) = (ToType) ((Codepoint >> 18)         | 128 | 64 | 32 | 16);
-				*(OutputIterator++) = (ToType)(((Codepoint >> 12) & 0x3F) | 128);
-				*(OutputIterator++) = (ToType)(((Codepoint >> 6 ) & 0x3F) | 128);
-				*(OutputIterator++) = (ToType) ((Codepoint        & 0x3F) | 128);
+				return -1;
 			}
+	
+			return (int32)(Result - Dest);
 		}
-
-		return UE_PTRDIFF_TO_INT32(OutputIterator - OutputIteratorStartPosition);
-	}
-
-	/**
-	 * Converts a Source string into UTF8 and stores it in Dest.
-	 *
-	 * @param Dest      The destination output iterator. Usually ANSICHAR*, but you can supply your own output iterator.
-	 *                  One can determine the number of characters written by checking the offset of Dest when the function returns.
-	 * @param DestLen   The length of the destination buffer.
-	 * @param Source    The source string to convert.
-	 * @param SourceLen The length of the source string.
-	 * @return          The number of bytes written to Dest, up to DestLen, or -1 if the entire Source string could did not fit in DestLen bytes.
-	 */
-	template <
-		typename SrcBufferType,
-		std::enable_if_t<TIsCharEncodingCompatibleWith_V<SrcBufferType, FromType>>* = nullptr
-	>
-	static FORCEINLINE int32 Convert(IntendedToType* Dest, int32 DestLen, const SrcBufferType* Source, int32 SourceLen)
-	{
-		IntendedToType* Result = FPlatformString::Convert(Dest, DestLen, (const FromType*)Source, SourceLen);
-		if (!Result)
+		template <
+			typename SrcBufferType,
+			std::enable_if_t<TIsCharEncodingCompatibleWith_V<SrcBufferType, FromType>>* = nullptr
+		>
+		static FORCEINLINE int32 Convert(ToType* Dest, int32 DestLen, const SrcBufferType* Source, int32 SourceLen)
 		{
-			return -1;
+			// This helper function is only for backwards compatibility and won't be necessary once IntendedToType becomes the ToType
+			return FTCHARToUTF8_Convert::Convert((IntendedToType*)Dest, DestLen, Source, SourceLen);
 		}
+	
+		/**
+		 * Determines the length of the converted string.
+		 *
+		 * @return The length of the string in UTF-8 code units.
+		 */
+		template <
+			typename SrcBufferType,
+			std::enable_if_t<TIsCharEncodingCompatibleWith_V<SrcBufferType, FromType>>* = nullptr
+		>
+		static FORCEINLINE int32 ConvertedLength(const SrcBufferType* Source, int32 SourceLen)
+		{
+			return FPlatformString::ConvertedLength<IntendedToType>((const FromType*)Source, SourceLen);
+		}
+	};
+}
 
-		return (int32)(Result - Dest);
-	}
-	template <
-		typename SrcBufferType,
-		std::enable_if_t<TIsCharEncodingCompatibleWith_V<SrcBufferType, FromType>>* = nullptr
-	>
-	static FORCEINLINE int32 Convert(ToType* Dest, int32 DestLen, const SrcBufferType* Source, int32 SourceLen)
-	{
-		// This helper function is only for backwards compatibility and won't be necessary once IntendedToType becomes the ToType
-		return FTCHARToUTF8_Convert::Convert((IntendedToType*)Dest, DestLen, Source, SourceLen);
-	}
-
-	/**
-	 * Determines the length of the converted string.
-	 *
-	 * @return The length of the string in UTF-8 code units.
-	 */
-	template <
-		typename SrcBufferType,
-		std::enable_if_t<TIsCharEncodingCompatibleWith_V<SrcBufferType, FromType>>* = nullptr
-	>
-	static FORCEINLINE int32 ConvertedLength(const SrcBufferType* Source, int32 SourceLen)
-	{
-		return FPlatformString::ConvertedLength<IntendedToType>((const FromType*)Source, SourceLen);
-	}
-};
+using FTCHARToUTF8_Convert UE_DEPRECATED(5.1, "FTCHARToUTF8_Convert has been deprecated in favor of FPlatformString::Convert and StringCast") = UEStringConv_Private::FTCHARToUTF8_Convert;
 
 class FUTF8ToTCHAR_Convert
 {
@@ -474,7 +476,7 @@ public:
 	>
 	static FORCEINLINE int32 ConvertedLength(const SrcBufferType* Source, int32 SourceLen)
 	{
-		UE4StringConv_Private::FCountingOutputIterator Dest;
+		UEStringConv_Private::FCountingOutputIterator Dest;
 		const int32 DestLen = SourceLen * 2;
 		Convert_Impl(Dest, DestLen, (const FromType*)Source, SourceLen);
 
@@ -552,7 +554,7 @@ public:
 	>
 	static int32 ConvertedLength(const SrcBufferType* Source, const int32 SourceLen)
 	{
-		UE4StringConv_Private::FCountingOutputIterator Dest;
+		UEStringConv_Private::FCountingOutputIterator Dest;
 		Convert_Impl(Dest, MAX_int32, (const FromType*)Source, SourceLen);
 
 		return Dest.GetCount();
@@ -962,7 +964,7 @@ struct TElementType<TStringPointer<FromType, ToType>>
  */
 
 // These should be replaced with StringCasts when FPlatformString starts to know about UTF-8.
-typedef TStringConversion<FTCHARToUTF8_Convert> FTCHARToUTF8;
+typedef TStringConversion<UEStringConv_Private::FTCHARToUTF8_Convert> FTCHARToUTF8;
 typedef TStringConversion<FUTF8ToTCHAR_Convert> FUTF8ToTCHAR;
 
 // Usage of these should be replaced with StringCasts.
diff --git a/Engine/Source/Runtime/Online/HTTP/Private/Curl/CurlHttp.cpp b/Engine/Source/Runtime/Online/HTTP/Private/Curl/CurlHttp.cpp
--- a/Engine/Source/Runtime/Online/HTTP/Private/Curl/CurlHttp.cpp
+++ b/Engine/Source/Runtime/Online/HTTP/Private/Curl/CurlHttp.cpp
@@ -331,10 +331,10 @@ void FCurlHttpRequest::SetContentAsString(const FString& ContentString)
 		return;
 	}
 
-	int32 Utf8Length = FTCHARToUTF8_Convert::ConvertedLength(*ContentString, ContentString.Len());
+	int32 Utf8Length = FPlatformString::ConvertedLength<UTF8CHAR>(*ContentString, ContentString.Len());
 	TArray<uint8> Buffer;
 	Buffer.SetNumUninitialized(Utf8Length);
-	FTCHARToUTF8_Convert::Convert((UTF8CHAR*)Buffer.GetData(), Buffer.Num(), *ContentString, ContentString.Len());
+	FPlatformString::Convert((UTF8CHAR*)Buffer.GetData(), Buffer.Num(), *ContentString, ContentString.Len());
 	RequestPayload = MakeUnique<FRequestPayloadInMemory>(MoveTemp(Buffer));
 	bIsRequestPayloadSeekable = true;
 }
diff --git a/Engine/Source/Runtime/Online/HTTP/Private/GenericPlatform/GenericPlatformHttp.cpp b/Engine/Source/Runtime/Online/HTTP/Private/GenericPlatform/GenericPlatformHttp.cpp
--- a/Engine/Source/Runtime/Online/HTTP/Private/GenericPlatform/GenericPlatformHttp.cpp
+++ b/Engine/Source/Runtime/Online/HTTP/Private/GenericPlatform/GenericPlatformHttp.cpp
@@ -181,7 +181,7 @@ FString FGenericPlatformHttp::UrlDecode(const FStringView EncodedString)
 	FTCHARToUTF8 Converter(EncodedString.GetData(), EncodedString.Len());
 	const UTF8CHAR* UTF8Data = (UTF8CHAR*)Converter.Get();	
 	
-	TArray<ANSICHAR> Data;
+	TArray<UTF8CHAR> Data;
 	Data.Reserve(EncodedString.Len());
 
 	for (int32 CharIdx = 0; CharIdx < Converter.Length();)
@@ -200,10 +200,10 @@ FString FGenericPlatformHttp::UrlDecode(const FStringView EncodedString)
 					Value += FParse::HexDigit(UTF8Data[CharIdx + 5]);
 					CharIdx += 6;
 
-					ANSICHAR Buffer[8] = { 0 };
-					ANSICHAR* BufferPtr = Buffer;
+					UTF8CHAR Buffer[8] = {};
+					UTF8CHAR* BufferPtr = Buffer;
 					const int32 Len = UE_ARRAY_COUNT(Buffer);
-					const int32 WrittenChars = FTCHARToUTF8_Convert::Utf8FromCodepoint(Value, BufferPtr, Len);
+					const int32 WrittenChars = FPlatformString::Convert(BufferPtr, Len, (UTF32CHAR*)&Value, 1) - BufferPtr;
 
 					Data.Append(Buffer, WrittenChars);
 				}
@@ -220,7 +220,7 @@ FString FGenericPlatformHttp::UrlDecode(const FStringView EncodedString)
 				Value = FParse::HexDigit(UTF8Data[CharIdx + 1]) << 4;
 				Value += FParse::HexDigit(UTF8Data[CharIdx + 2]);
 				CharIdx += 3;
-				Data.Add((ANSICHAR)(Value));
+				Data.Add((UTF8CHAR)(Value));
 			}
 			else
 			{
@@ -237,7 +237,7 @@ FString FGenericPlatformHttp::UrlDecode(const FStringView EncodedString)
 		}
 	}
 
-	Data.Add('\0');
+	Data.Add(UTF8TEXT('\0'));
 	return FString(UTF8_TO_TCHAR(Data.GetData()));
 }
 
diff --git a/Engine/Source/Runtime/Online/HTTP/Private/IXML/HttpIXml.cpp b/Engine/Source/Runtime/Online/HTTP/Private/IXML/HttpIXml.cpp
--- a/Engine/Source/Runtime/Online/HTTP/Private/IXML/HttpIXml.cpp
+++ b/Engine/Source/Runtime/Online/HTTP/Private/IXML/HttpIXml.cpp
@@ -137,10 +137,10 @@ void FHttpRequestIXML::SetContentAsString(const FString& ContentString)
 {
 	if ( ContentString.Len() )
 	{
-		int32 Utf8Length = FTCHARToUTF8_Convert::ConvertedLength(*ContentString, ContentString.Len());
+		int32 Utf8Length = FPlatformString::ConvertedLength<UTF8CHAR>(*ContentString, ContentString.Len());
 		TArray<uint8> Buffer;
 		Buffer.SetNumUninitialized(Utf8Length);
-		FTCHARToUTF8_Convert::Convert((ANSICHAR*)Buffer.GetData(), Buffer.Num(), *ContentString, ContentString.Len());
+		FPlatformString::Convert((UTF8CHAR*)Buffer.GetData(), Buffer.Num(), *ContentString, ContentString.Len());
 		Payload = MakeUnique<FRequestPayloadInMemory>(MoveTemp(Buffer));
 	}
 }
diff --git a/Engine/Source/Runtime/Online/HTTP/Private/NullHttp.cpp b/Engine/Source/Runtime/Online/HTTP/Private/NullHttp.cpp
--- a/Engine/Source/Runtime/Online/HTTP/Private/NullHttp.cpp
+++ b/Engine/Source/Runtime/Online/HTTP/Private/NullHttp.cpp
@@ -79,9 +79,9 @@ void FNullHttpRequest::SetContent(TArray<uint8>&& ContentPayload)
 
 void FNullHttpRequest::SetContentAsString(const FString& ContentString)
 {
-	int32 Utf8Length = FTCHARToUTF8_Convert::ConvertedLength(*ContentString, ContentString.Len());
+	int32 Utf8Length = FPlatformString::ConvertedLength<UTF8CHAR>(*ContentString, ContentString.Len());
 	Payload.SetNumUninitialized(Utf8Length);
-	FTCHARToUTF8_Convert::Convert((UTF8CHAR*)Payload.GetData(), Payload.Num(), *ContentString, ContentString.Len());
+	FPlatformString::Convert((UTF8CHAR*)Payload.GetData(), Payload.Num(), *ContentString, ContentString.Len());
 }
 
 bool FNullHttpRequest::SetContentAsStreamedFile(const FString& Filename)
diff --git a/Engine/Source/Runtime/Online/HTTPServer/Private/HttpConnectionRequestReadContext.cpp b/Engine/Source/Runtime/Online/HTTPServer/Private/HttpConnectionRequestReadContext.cpp
--- a/Engine/Source/Runtime/Online/HTTPServer/Private/HttpConnectionRequestReadContext.cpp
+++ b/Engine/Source/Runtime/Online/HTTPServer/Private/HttpConnectionRequestReadContext.cpp
@@ -324,7 +324,7 @@ FString FHttpConnectionRequestReadContext::UrlDecode(const FString &EncodedStrin
 	FTCHARToUTF8 Converter(*EncodedString);
 	const UTF8CHAR* UTF8Data = (UTF8CHAR*)Converter.Get();
 
-	TArray<ANSICHAR> Data;
+	TArray<UTF8CHAR> Data;
 	Data.Reserve(EncodedString.Len());
 
 	for (int32 CharIdx = 0; CharIdx < Converter.Length();)
@@ -343,10 +343,10 @@ FString FHttpConnectionRequestReadContext::UrlDecode(const FString &EncodedStrin
 					Value += FParse::HexDigit(UTF8Data[CharIdx + 5]);
 					CharIdx += 6;
 
-					ANSICHAR Buffer[8] = { 0 };
-					ANSICHAR* BufferPtr = Buffer;
+					UTF8CHAR Buffer[8] = { };
+					UTF8CHAR* BufferPtr = Buffer;
 					const int32 Len = UE_ARRAY_COUNT(Buffer);
-					const int32 WrittenChars = FTCHARToUTF8_Convert::Utf8FromCodepoint(Value, BufferPtr, Len);
+					const int32 WrittenChars = FPlatformString::Convert(BufferPtr, Len, (UTF32CHAR*)&Value, 1) - BufferPtr;
 
 					Data.Append(Buffer, WrittenChars);
 				}
@@ -363,7 +363,7 @@ FString FHttpConnectionRequestReadContext::UrlDecode(const FString &EncodedStrin
 				Value = FParse::HexDigit(UTF8Data[CharIdx + 1]) << 4;
 				Value += FParse::HexDigit(UTF8Data[CharIdx + 2]);
 				CharIdx += 3;
-				Data.Add((ANSICHAR)(Value));
+				Data.Add((UTF8CHAR)(Value));
 			}
 			else
 			{
@@ -380,6 +380,6 @@ FString FHttpConnectionRequestReadContext::UrlDecode(const FString &EncodedStrin
 		}
 	}
 
-	Data.Add('\0');
+	Data.Add(UTF8TEXT('\0'));
 	return FString(UTF8_TO_TCHAR(Data.GetData()));
 }
\ No newline at end of file
