Added TAttribute::CreateStatic() to disambiguate in cases where Create() was not correctly detecting that additional arguments were meant to be forwarded to the function.

diff --git a/Engine/Source/Runtime/Core/Public/Misc/Attribute.h b/Engine/Source/Runtime/Core/Public/Misc/Attribute.h
--- a/Engine/Source/Runtime/Core/Public/Misc/Attribute.h
+++ b/Engine/Source/Runtime/Core/Public/Misc/Attribute.h
@@ -117,6 +117,20 @@ public:
 		return TAttribute( FGetter::CreateStatic( InFuncPtr, Vars... ), bExplicitConstructor );
 	}
 
+	/**
+	 * Creates an attribute by binding an arbitrary function that will be called to generate this attribute's value on demand.
+	 * After binding, the attribute will no longer have a value that can be accessed directly, and instead the bound
+	 * function will always be called to generate the value.
+	 *
+	 * @param  InFuncPtr Member function to bind.  The function's structure (return value, arguments, etc) must match IBoundAttributeDelegate's definition.
+	 */
+	template <typename FuncPtrType, typename... VarTypes >
+	UE_NODISCARD static TAttribute CreateStatic(FuncPtrType&& InFuncPtr, VarTypes... Vars)
+	{
+		const bool bExplicitConstructor = true;
+		return TAttribute(FGetter::CreateStatic(InFuncPtr, MoveTemp(Vars)...), bExplicitConstructor);
+	}
+
 	/**
 	 * Helper function for creating TAttributes from a function pointer, accessed through a raw pointer
 	 */
