Fix CIS warning for implicit conversion

diff --git a/Engine/Source/Runtime/Core/Private/Unix/UnixSignalHeartBeat.cpp b/Engine/Source/Runtime/Core/Private/Unix/UnixSignalHeartBeat.cpp
--- a/Engine/Source/Runtime/Core/Private/Unix/UnixSignalHeartBeat.cpp
+++ b/Engine/Source/Runtime/Core/Private/Unix/UnixSignalHeartBeat.cpp
@@ -165,7 +165,7 @@ void FUnixSignalGameHitchHeartBeat::FrameStart(bool bSkipThisFrame)
 			FMemory::Memzero(HeartBeatTime);
 
 			long FullSeconds = static_cast<long>(HitchThresholdS);
-			long RemainderInNanoSeconds = FMath::Fmod(HitchThresholdS, 1.0) * 1000000000LL;
+			long RemainderInNanoSeconds = static_cast<long>(FMath::Fmod(HitchThresholdS, 1.0) * 1000000000.0);
 			HeartBeatTime.it_value.tv_sec = FullSeconds;
 			HeartBeatTime.it_value.tv_nsec = RemainderInNanoSeconds;
 
