Added an Equals() implementation for TBox2D.

All of the core types (eg: TVector, TMatrix, etc.) have an Equals() member function with a tolerance level, except TBox2D. This simply maintains parity across those types.

diff --git a/Engine/Source/Runtime/Core/Public/Math/Box2D.h b/Engine/Source/Runtime/Core/Public/Math/Box2D.h
--- a/Engine/Source/Runtime/Core/Public/Math/Box2D.h
+++ b/Engine/Source/Runtime/Core/Public/Math/Box2D.h
@@ -88,16 +88,28 @@ public:
 	}
 
 	/**
-	* Compares two boxes for inequality.
-	*
-	* @param Other The other box to compare with.
-	* @return true if the boxes are not equal, false otherwise.
-	*/
-	bool operator!=(const TBox2<T>& Other) const
+	 * Compares two boxes for inequality.
+	 *
+	 * @param Other The other box to compare with.
+	 * @return true if the boxes are not equal, false otherwise.
+	 */
+	bool operator!=( const TBox2<T>& Other ) const
 	{
 		return !(*this == Other);
 	}
 
+	/**
+	 * Checks for equality with error-tolerant comparison.
+	 *
+	 * @param Other The box to compare.
+	 * @param Tolerance Error tolerance.
+	 * @return true if the boxes are equal within specified tolerance, otherwise false.
+	 */
+	bool Equals( const TBox2<T>& Other, T Tolerance = UE_KINDA_SMALL_NUMBER ) const
+	{
+		return Min.Equals(Other.Min, Tolerance) && Max.Equals(Other.Max, Tolerance);
+	}
+
 	/**
 	 * Adds to this bounding box to include a given point.
 	 *
