Fix IterateDirectory() on various platforms so it's behavior is consistent (was not breaking iteration when desired, meaning the function's return value could be wrong as well).

diff --git a/Engine/Source/Runtime/Core/Private/Apple/ApplePlatformFile.cpp b/Engine/Source/Runtime/Core/Private/Apple/ApplePlatformFile.cpp
--- a/Engine/Source/Runtime/Core/Private/Apple/ApplePlatformFile.cpp
+++ b/Engine/Source/Runtime/Core/Private/Apple/ApplePlatformFile.cpp
@@ -767,7 +767,7 @@ bool FApplePlatformFile::IterateDirectoryCommon(const TCHAR* Directory, const TF
 	{
 		Result = true;
 		struct dirent *Entry;
-		while ((Entry = readdir(Handle)) != NULL)
+		while (Result && (Entry = readdir(Handle)) != NULL)
 		{
 			if (FCStringAnsi::Strcmp(Entry->d_name, ".") && FCStringAnsi::Strcmp(Entry->d_name, "..") && FCStringAnsi::Strcmp(Entry->d_name, ".DS_Store"))
 			{
diff --git a/Engine/Source/Runtime/Core/Private/IOS/IOSPlatformFile.cpp b/Engine/Source/Runtime/Core/Private/IOS/IOSPlatformFile.cpp
--- a/Engine/Source/Runtime/Core/Private/IOS/IOSPlatformFile.cpp
+++ b/Engine/Source/Runtime/Core/Private/IOS/IOSPlatformFile.cpp
@@ -1035,7 +1035,7 @@ bool FIOSPlatformFile::IterateDirectoryCommon(const TCHAR* Directory, const TFun
 	{
 		Result = true;
 		struct dirent *Entry;
-		while ((Entry = readdir(Handle)) != NULL)
+		while (Result && (Entry = readdir(Handle)) != NULL)
 		{
 			if (FCStringAnsi::Strcmp(Entry->d_name, ".") && FCStringAnsi::Strcmp(Entry->d_name, ".."))
 			{
diff --git a/Engine/Source/Runtime/Core/Private/Unix/UnixPlatformFile.cpp b/Engine/Source/Runtime/Core/Private/Unix/UnixPlatformFile.cpp
--- a/Engine/Source/Runtime/Core/Private/Unix/UnixPlatformFile.cpp
+++ b/Engine/Source/Runtime/Core/Private/Unix/UnixPlatformFile.cpp
@@ -1266,7 +1266,7 @@ bool FUnixPlatformFile::IterateDirectoryCommon(const TCHAR* Directory, const TFu
 	{
 		Result = true;
 		struct dirent* Entry;
-		while ((Entry = readdir(Handle)) != NULL)
+		while (Result && (Entry = readdir(Handle)) != NULL)
 		{
 			if (FCString::Strcmp(UTF8_TO_TCHAR(Entry->d_name), TEXT(".")) && FCString::Strcmp(UTF8_TO_TCHAR(Entry->d_name), TEXT("..")))
 			{
