Android - Hang detector can be toggled via configrules.
Android hang detector is disabled if no configrule exists.

diff --git a/Engine/Source/Runtime/Core/Private/Android/AndroidPlatformMisc.cpp b/Engine/Source/Runtime/Core/Private/Android/AndroidPlatformMisc.cpp
--- a/Engine/Source/Runtime/Core/Private/Android/AndroidPlatformMisc.cpp
+++ b/Engine/Source/Runtime/Core/Private/Android/AndroidPlatformMisc.cpp
@@ -2449,6 +2449,25 @@ FString* FAndroidMisc::GetConfigRulesVariable(const FString& Key)
 	return ConfigRulesVariables.Find(Key);
 }
 
+bool FAndroidMisc::AllowThreadHeartBeat()
+{
+	static uint32 AllowThreadHeartBeatOnce = -1;
+	if (AllowThreadHeartBeatOnce == -1)
+	{
+		const FString* AllowThreadHeartBeatConfigVar = GetConfigRulesVariable(TEXT("EnableThreadHeartBeat"));
+		if (AllowThreadHeartBeatConfigVar)
+		{
+			AllowThreadHeartBeatOnce = (uint32)AllowThreadHeartBeatConfigVar->Equals("true", ESearchCase::IgnoreCase);
+		}
+		else
+		{
+			AllowThreadHeartBeatOnce = 0;
+		}
+	}
+
+	return AllowThreadHeartBeatOnce == 1;
+}
+
 JNI_METHOD void Java_com_epicgames_unreal_GameActivity_nativeSetConfigRulesVariables(JNIEnv* jenv, jobject thiz, jobjectArray KeyValuePairs)
 {
 	int32 Count = jenv->GetArrayLength(KeyValuePairs);
diff --git a/Engine/Source/Runtime/Core/Public/Android/AndroidPlatformMisc.h b/Engine/Source/Runtime/Core/Public/Android/AndroidPlatformMisc.h
--- a/Engine/Source/Runtime/Core/Public/Android/AndroidPlatformMisc.h
+++ b/Engine/Source/Runtime/Core/Public/Android/AndroidPlatformMisc.h
@@ -57,10 +57,7 @@ struct CORE_API FAndroidMisc : public FGenericPlatformMisc
 
 public:
 
-	static bool AllowThreadHeartBeat()
-	{
-		return true;
-	}
+	static bool AllowThreadHeartBeat();
 
 	struct FCPUStatTime{
 		uint64_t			TotalTime;
