Direct FTCHARToUTF8_Convert usage replaced with StringCast or FPlatformString in Core.

diff --git a/Engine/Source/Runtime/Core/Private/Serialization/CompactBinaryWriter.cpp b/Engine/Source/Runtime/Core/Private/Serialization/CompactBinaryWriter.cpp
--- a/Engine/Source/Runtime/Core/Private/Serialization/CompactBinaryWriter.cpp
+++ b/Engine/Source/Runtime/Core/Private/Serialization/CompactBinaryWriter.cpp
@@ -398,7 +398,7 @@ void FCbWriter::AddString(const FUtf8StringView Value)
 void FCbWriter::AddString(const FWideStringView Value)
 {
 	BeginField();
-	const uint32 Size = uint32(FTCHARToUTF8_Convert::ConvertedLength(Value.GetData(), Value.Len()));
+	const uint32 Size = uint32(FPlatformString::ConvertedLength<UTF8CHAR>(Value.GetData(), Value.Len()));
 	const uint32 SizeByteCount = MeasureVarUInt(Size);
 	const int64 Offset = Data.AddUninitialized(SizeByteCount + Size);
 	uint8* StringData = Data.GetData() + Offset;
@@ -406,7 +406,7 @@ void FCbWriter::AddString(const FWideStringView Value)
 	StringData += SizeByteCount;
 	if (Size > 0)
 	{
-		FTCHARToUTF8_Convert::Convert(reinterpret_cast<UTF8CHAR*>(StringData), Size, Value.GetData(), Value.Len());
+		FPlatformString::Convert(reinterpret_cast<UTF8CHAR*>(StringData), Size, Value.GetData(), Value.Len());
 	}
 	EndField(ECbFieldType::String);
 }
