FTCHARToUTF8_Convert potentially deprecated.
Old UE4 namespace renamed.

diff --git a/Engine/Source/Runtime/Core/Public/Containers/StringConv.h b/Engine/Source/Runtime/Core/Public/Containers/StringConv.h
--- a/Engine/Source/Runtime/Core/Public/Containers/StringConv.h
+++ b/Engine/Source/Runtime/Core/Public/Containers/StringConv.h
@@ -173,7 +173,7 @@ namespace StringConv
 	FUnused GetLegacyFromType(...);
 }
 
-namespace UE4StringConv_Private
+namespace UE::Core::Private
 {
 	/**
 	 * This is a basic object which counts how many times it has been incremented
@@ -206,130 +206,132 @@ namespace UE4StringConv_Private
 	private:
 		int32 Counter;
 	};
-}
-
-// This should be replaced with Platform stuff when FPlatformString starts to know about UTF-8.
-class FTCHARToUTF8_Convert
-{
-	// Leave this until we convert all the call sites and can change FromType to this.
-	using IntendedToType = UTF8CHAR;
 
-public:
-	typedef TCHAR    FromType;
-	typedef ANSICHAR ToType;
-
-	/**
-	 * Convert Codepoint into UTF-8 characters.
-	 *
-	 * @param Codepoint Codepoint to expand into UTF-8 bytes
-	 * @param OutputIterator Output iterator to write UTF-8 bytes into
-	 * @param OutputIteratorByteSizeRemaining Maximum number of ANSI characters that can be written to OutputIterator
-	 * @return Number of characters written for Codepoint
-	 */
-	template <typename BufferType>
-	static int32 Utf8FromCodepoint(uint32 Codepoint, BufferType OutputIterator, uint32 OutputIteratorByteSizeRemaining)
+	// This should be replaced with Platform stuff when FPlatformString starts to know about UTF-8.
+	class FTCHARToUTF8_Convert
 	{
-		// Ensure we have at least one character in size to write
-		if (OutputIteratorByteSizeRemaining < sizeof(ToType))
-		{
-			return 0;
-		}
-
-		const BufferType OutputIteratorStartPosition = OutputIterator;
-
-		if (!StringConv::IsValidCodepoint(Codepoint))
-		{
-			Codepoint = UNICODE_BOGUS_CHAR_CODEPOINT;
-		}
-		else if (StringConv::IsHighSurrogate(Codepoint) || StringConv::IsLowSurrogate(Codepoint)) // UTF-8 Characters are not allowed to encode codepoints in the surrogate pair range
-		{
-			Codepoint = UNICODE_BOGUS_CHAR_CODEPOINT;
-		}
-
-		// Do the encoding...
-		if (Codepoint < 0x80)
-		{
-			*(OutputIterator++) = (ToType)Codepoint;
-		}
-		else if (Codepoint < 0x800)
+		// Leave this until we convert all the call sites and can change FromType to this.
+		using IntendedToType = UTF8CHAR;
+	
+	public:
+		typedef TCHAR    FromType;
+		typedef ANSICHAR ToType;
+	
+		/**
+		 * Convert Codepoint into UTF-8 characters.
+		 *
+		 * @param Codepoint Codepoint to expand into UTF-8 bytes
+		 * @param OutputIterator Output iterator to write UTF-8 bytes into
+		 * @param OutputIteratorByteSizeRemaining Maximum number of ANSI characters that can be written to OutputIterator
+		 * @return Number of characters written for Codepoint
+		 */
+		template <typename BufferType>
+		static int32 Utf8FromCodepoint(uint32 Codepoint, BufferType OutputIterator, uint32 OutputIteratorByteSizeRemaining)
 		{
-			if (OutputIteratorByteSizeRemaining >= 2)
+			// Ensure we have at least one character in size to write
+			if (OutputIteratorByteSizeRemaining < sizeof(ToType))
 			{
-				*(OutputIterator++) = (ToType)((Codepoint >> 6)         | 128 | 64);
-				*(OutputIterator++) = (ToType)((Codepoint       & 0x3F) | 128);
+				return 0;
 			}
-		}
-		else if (Codepoint < 0x10000)
-		{
-			if (OutputIteratorByteSizeRemaining >= 3)
+	
+			const BufferType OutputIteratorStartPosition = OutputIterator;
+	
+			if (!StringConv::IsValidCodepoint(Codepoint))
+			{
+				Codepoint = UNICODE_BOGUS_CHAR_CODEPOINT;
+			}
+			else if (StringConv::IsHighSurrogate(Codepoint) || StringConv::IsLowSurrogate(Codepoint)) // UTF-8 Characters are not allowed to encode codepoints in the surrogate pair range
+			{
+				Codepoint = UNICODE_BOGUS_CHAR_CODEPOINT;
+			}
+	
+			// Do the encoding...
+			if (Codepoint < 0x80)
+			{
+				*(OutputIterator++) = (ToType)Codepoint;
+			}
+			else if (Codepoint < 0x800)
 			{
-				*(OutputIterator++) = (ToType) ((Codepoint >> 12)        | 128 | 64 | 32);
-				*(OutputIterator++) = (ToType)(((Codepoint >> 6) & 0x3F) | 128);
-				*(OutputIterator++) = (ToType) ((Codepoint       & 0x3F) | 128);
+				if (OutputIteratorByteSizeRemaining >= 2)
+				{
+					*(OutputIterator++) = (ToType)((Codepoint >> 6)         | 128 | 64);
+					*(OutputIterator++) = (ToType)((Codepoint       & 0x3F) | 128);
+				}
 			}
+			else if (Codepoint < 0x10000)
+			{
+				if (OutputIteratorByteSizeRemaining >= 3)
+				{
+					*(OutputIterator++) = (ToType) ((Codepoint >> 12)        | 128 | 64 | 32);
+					*(OutputIterator++) = (ToType)(((Codepoint >> 6) & 0x3F) | 128);
+					*(OutputIterator++) = (ToType) ((Codepoint       & 0x3F) | 128);
+				}
+			}
+			else
+			{
+				if (OutputIteratorByteSizeRemaining >= 4)
+				{
+					*(OutputIterator++) = (ToType) ((Codepoint >> 18)         | 128 | 64 | 32 | 16);
+					*(OutputIterator++) = (ToType)(((Codepoint >> 12) & 0x3F) | 128);
+					*(OutputIterator++) = (ToType)(((Codepoint >> 6 ) & 0x3F) | 128);
+					*(OutputIterator++) = (ToType) ((Codepoint        & 0x3F) | 128);
+				}
+			}
+	
+			return UE_PTRDIFF_TO_INT32(OutputIterator - OutputIteratorStartPosition);
 		}
-		else
+	
+		/**
+		 * Converts a Source string into UTF8 and stores it in Dest.
+		 *
+		 * @param Dest      The destination output iterator. Usually ANSICHAR*, but you can supply your own output iterator.
+		 *                  One can determine the number of characters written by checking the offset of Dest when the function returns.
+		 * @param DestLen   The length of the destination buffer.
+		 * @param Source    The source string to convert.
+		 * @param SourceLen The length of the source string.
+		 * @return          The number of bytes written to Dest, up to DestLen, or -1 if the entire Source string could did not fit in DestLen bytes.
+		 */
+		template <
+			typename SrcBufferType,
+			std::enable_if_t<TIsCharEncodingCompatibleWith_V<SrcBufferType, FromType>>* = nullptr
+		>
+		static FORCEINLINE int32 Convert(IntendedToType* Dest, int32 DestLen, const SrcBufferType* Source, int32 SourceLen)
 		{
-			if (OutputIteratorByteSizeRemaining >= 4)
+			IntendedToType* Result = FPlatformString::Convert(Dest, DestLen, (const FromType*)Source, SourceLen);
+			if (!Result)
 			{
-				*(OutputIterator++) = (ToType) ((Codepoint >> 18)         | 128 | 64 | 32 | 16);
-				*(OutputIterator++) = (ToType)(((Codepoint >> 12) & 0x3F) | 128);
-				*(OutputIterator++) = (ToType)(((Codepoint >> 6 ) & 0x3F) | 128);
-				*(OutputIterator++) = (ToType) ((Codepoint        & 0x3F) | 128);
+				return -1;
 			}
+	
+			return (int32)(Result - Dest);
 		}
-
-		return UE_PTRDIFF_TO_INT32(OutputIterator - OutputIteratorStartPosition);
-	}
-
-	/**
-	 * Converts a Source string into UTF8 and stores it in Dest.
-	 *
-	 * @param Dest      The destination output iterator. Usually ANSICHAR*, but you can supply your own output iterator.
-	 *                  One can determine the number of characters written by checking the offset of Dest when the function returns.
-	 * @param DestLen   The length of the destination buffer.
-	 * @param Source    The source string to convert.
-	 * @param SourceLen The length of the source string.
-	 * @return          The number of bytes written to Dest, up to DestLen, or -1 if the entire Source string could did not fit in DestLen bytes.
-	 */
-	template <
-		typename SrcBufferType,
-		std::enable_if_t<TIsCharEncodingCompatibleWith_V<SrcBufferType, FromType>>* = nullptr
-	>
-	static FORCEINLINE int32 Convert(IntendedToType* Dest, int32 DestLen, const SrcBufferType* Source, int32 SourceLen)
-	{
-		IntendedToType* Result = FPlatformString::Convert(Dest, DestLen, (const FromType*)Source, SourceLen);
-		if (!Result)
+		template <
+			typename SrcBufferType,
+			std::enable_if_t<TIsCharEncodingCompatibleWith_V<SrcBufferType, FromType>>* = nullptr
+		>
+		static FORCEINLINE int32 Convert(ToType* Dest, int32 DestLen, const SrcBufferType* Source, int32 SourceLen)
 		{
-			return -1;
+			// This helper function is only for backwards compatibility and won't be necessary once IntendedToType becomes the ToType
+			return FTCHARToUTF8_Convert::Convert((IntendedToType*)Dest, DestLen, Source, SourceLen);
 		}
+	
+		/**
+		 * Determines the length of the converted string.
+		 *
+		 * @return The length of the string in UTF-8 code units.
+		 */
+		template <
+			typename SrcBufferType,
+			std::enable_if_t<TIsCharEncodingCompatibleWith_V<SrcBufferType, FromType>>* = nullptr
+		>
+		static FORCEINLINE int32 ConvertedLength(const SrcBufferType* Source, int32 SourceLen)
+		{
+			return FPlatformString::ConvertedLength<IntendedToType>((const FromType*)Source, SourceLen);
+		}
+	};
+}
 
-		return (int32)(Result - Dest);
-	}
-	template <
-		typename SrcBufferType,
-		std::enable_if_t<TIsCharEncodingCompatibleWith_V<SrcBufferType, FromType>>* = nullptr
-	>
-	static FORCEINLINE int32 Convert(ToType* Dest, int32 DestLen, const SrcBufferType* Source, int32 SourceLen)
-	{
-		// This helper function is only for backwards compatibility and won't be necessary once IntendedToType becomes the ToType
-		return FTCHARToUTF8_Convert::Convert((IntendedToType*)Dest, DestLen, Source, SourceLen);
-	}
-
-	/**
-	 * Determines the length of the converted string.
-	 *
-	 * @return The length of the string in UTF-8 code units.
-	 */
-	template <
-		typename SrcBufferType,
-		std::enable_if_t<TIsCharEncodingCompatibleWith_V<SrcBufferType, FromType>>* = nullptr
-	>
-	static FORCEINLINE int32 ConvertedLength(const SrcBufferType* Source, int32 SourceLen)
-	{
-		return FPlatformString::ConvertedLength<IntendedToType>((const FromType*)Source, SourceLen);
-	}
-};
+using FTCHARToUTF8_Convert /*UE_DEPRECATED(5.1, "FTCHARToUTF8_Convert has been deprecated in favor of FPlatformString::Convert and StringCast")*/ = UE::Core::Private::FTCHARToUTF8_Convert;
 
 class FUTF8ToTCHAR_Convert
 {
@@ -474,7 +476,7 @@ public:
 	>
 	static FORCEINLINE int32 ConvertedLength(const SrcBufferType* Source, int32 SourceLen)
 	{
-		UE4StringConv_Private::FCountingOutputIterator Dest;
+		UE::Core::Private::FCountingOutputIterator Dest;
 		const int32 DestLen = SourceLen * 2;
 		Convert_Impl(Dest, DestLen, (const FromType*)Source, SourceLen);
 
@@ -552,7 +554,7 @@ public:
 	>
 	static int32 ConvertedLength(const SrcBufferType* Source, const int32 SourceLen)
 	{
-		UE4StringConv_Private::FCountingOutputIterator Dest;
+		UE::Core::Private::FCountingOutputIterator Dest;
 		Convert_Impl(Dest, MAX_int32, (const FromType*)Source, SourceLen);
 
 		return Dest.GetCount();
@@ -962,7 +964,7 @@ struct TElementType<TStringPointer<FromType, ToType>>
  */
 
 // These should be replaced with StringCasts when FPlatformString starts to know about UTF-8.
-typedef TStringConversion<FTCHARToUTF8_Convert> FTCHARToUTF8;
+typedef TStringConversion<UE::Core::Private::FTCHARToUTF8_Convert> FTCHARToUTF8;
 typedef TStringConversion<FUTF8ToTCHAR_Convert> FUTF8ToTCHAR;
 
 // Usage of these should be replaced with StringCasts.
