Logging: Changed FOutputDevicesReadScopeLock to use release instead of relaxed when unlocking

The existing use of relaxed technically allows reads of the shared state to be reordered after the fetch_sub that unlocks the lock.

#rnx

diff --git a/Engine/Source/Runtime/Core/Private/Misc/OutputDeviceRedirector.cpp b/Engine/Source/Runtime/Core/Private/Misc/OutputDeviceRedirector.cpp
--- a/Engine/Source/Runtime/Core/Private/Misc/OutputDeviceRedirector.cpp
+++ b/Engine/Source/Runtime/Core/Private/Misc/OutputDeviceRedirector.cpp
@@ -231,7 +231,7 @@ public:
 	{
 		TRACE_CPUPROFILER_EVENT_SCOPE(FOutputDevicesReadScopeLock);
 		// A write lock has set the LSB. Cancel this read lock and wait for the write.
-		State.OutputDevicesLockState.fetch_sub(2, std::memory_order_relaxed);
+		State.OutputDevicesLockState.fetch_sub(2, std::memory_order_release);
 		// This read lock will wait until the write lock exits.
 		FReadScopeLock ScopeLock(State.OutputDevicesLock);
 		// Acquire on this read lock because the write may have mutated state that we read.
@@ -241,7 +241,7 @@ public:
 
 	FORCEINLINE ~FOutputDevicesReadScopeLock()
 	{
-		State.OutputDevicesLockState.fetch_sub(2, std::memory_order_relaxed);
+		State.OutputDevicesLockState.fetch_sub(2, std::memory_order_release);
 	}
 
 private:
