CSV profiler - persistent custom stats. Thread-safe helper classes for custom stats that persist across frames.

diff --git a/Engine/Source/Runtime/Core/Private/ProfilingDebugging/CsvProfiler.cpp b/Engine/Source/Runtime/Core/Private/ProfilingDebugging/CsvProfiler.cpp
--- a/Engine/Source/Runtime/Core/Private/ProfilingDebugging/CsvProfiler.cpp
+++ b/Engine/Source/Runtime/Core/Private/ProfilingDebugging/CsvProfiler.cpp
@@ -168,6 +168,107 @@ static bool GCsvExitOnCompletion = false;
 static thread_local bool GCsvThreadLocalWaitsEnabled = false;
 
 
+// A unique ID for a CSV stat, either ansi or FName
+union FCsvUniqueStatID
+{
+public:
+	FCsvUniqueStatID(const FCsvUniqueStatID& Src)
+	{
+		Hash = Src.Hash;
+	}
+	FCsvUniqueStatID(uint64 InStatIDRaw, int32 InCategoryIndex, bool bInIsFName, bool bInIsCountStat = false)
+	{
+		check(InCategoryIndex < CSV_MAX_CATEGORY_COUNT);
+		Fields.IsFName = bInIsFName ? 1 : 0;
+		Fields.FNameOrIndex = InStatIDRaw;
+		Fields.CategoryIndex = InCategoryIndex;
+		Fields.IsCountStat = bInIsCountStat ? 1 : 0;
+	}
+	FCsvUniqueStatID(const FName& Name, int32 InCategoryIndex)
+	{
+		check(InCategoryIndex < CSV_MAX_CATEGORY_COUNT);
+		Fields.FNameOrIndex = Name.ToUnstableInt();
+		Fields.CategoryIndex = InCategoryIndex;
+		Fields.IsFName = 1;
+		Fields.IsCountStat = 0;
+	}
+	struct
+	{
+		uint64 IsFName : 1;
+		uint64 IsCountStat : 1;
+		uint64 CategoryIndex : 11;
+		uint64 FNameOrIndex : 51;
+	} Fields;
+	uint64 Hash;
+};
+
+
+// Persistent custom stats
+struct FCsvPersistentCustomStats
+{
+	void RecordStats()
+	{
+		FScopeLock Lock(&Cs);
+		for (FCsvPersistentCustomStatBase* BaseStat : Stats)
+		{
+			switch(BaseStat->GetStatType())
+			{
+				case ECsvPersistentCustomStatType::Float:
+				{
+					RecordStat<float>(BaseStat);
+					break;
+				}
+				case ECsvPersistentCustomStatType::Int:
+				{
+					RecordStat<int32>(BaseStat);
+					break;
+				}
+			}
+		}
+	}
+
+
+	template<class T>
+	TCsvPersistentCustomStat<T>* GetOrCreatePersistentCustomStat(FName Name, int32 CategoryIndex, bool bResetEachFrame)
+	{
+		LLM_SCOPE(ELLMTag::CsvProfiler);
+		FScopeLock Lock(&Cs);
+		FCsvUniqueStatID Id(Name, CategoryIndex);
+		FCsvPersistentCustomStatBase** FindStat = StatLookup.Find(Id.Hash);
+		if (FindStat)
+		{
+			if (TCsvPersistentCustomStat<T>::GetClassStatType() == (*FindStat)->GetStatType())
+			{
+				return static_cast<TCsvPersistentCustomStat<T>*>(*FindStat);
+			}
+			UE_LOG(LogCsvProfiler, Fatal, TEXT("Error: Custom stat %s was already registered with a different type"), *Name.ToString());
+		}
+		// This will leak, and that's ok. These stats are intended to persist for the lifetime of the program
+		TCsvPersistentCustomStat<T>* NewStat = new TCsvPersistentCustomStat<T>(Name, CategoryIndex, bResetEachFrame);
+		StatLookup.Add(Id.Hash, NewStat);
+		Stats.Add(NewStat);
+		return NewStat;
+	}
+
+	template <class T>
+	void RecordStat(FCsvPersistentCustomStatBase* BaseStat)
+	{
+		TCsvPersistentCustomStat<T>* Stat = static_cast<TCsvPersistentCustomStat<T>*>(BaseStat);
+		FCsvProfiler::RecordCustomStat(Stat->Name, Stat->CategoryIndex, Stat->GetValue(), ECsvCustomStatOp::Set);
+		if (Stat->bResetEachFrame)
+		{
+			Stat->Set(0);
+		}
+	}
+
+	FCriticalSection Cs;
+	TMap<uint64, FCsvPersistentCustomStatBase*> StatLookup;
+	TArray<FCsvPersistentCustomStatBase*> Stats;
+};
+static FCsvPersistentCustomStats GCsvPersistentCustomStats;
+
+
+
 #if CSV_PROFILER_SUPPORT_NAMED_EVENTS
   #if PLATFORM_IMPLEMENTS_BeginNamedEventStatic
     #define CSV_PROFILER_BeginNamedEvent(Color,Text) FPlatformMisc::BeginNamedEventStatic(Color, Text)
@@ -1040,12 +1141,7 @@ public:
 		check(IsInCsvProcessingThread());
 
 		// Make a compound key
-		FUniqueID UniqueID;
-		check(InCategoryIndex < CSV_MAX_CATEGORY_COUNT);
-		UniqueID.Fields.IsFName = bInIsFName ? 1 : 0;
-		UniqueID.Fields.FNameOrIndex = InStatIDRaw;
-		UniqueID.Fields.CategoryIndex = InCategoryIndex;
-		UniqueID.Fields.IsCountStat = bInIsCountStat ? 1 : 0;
+		FCsvUniqueStatID UniqueID(InStatIDRaw, InCategoryIndex, bInIsFName, bInIsCountStat);
 
 		uint64 Hash = UniqueID.Hash;
 		int32 *IndexPtr = StatIDToIndex.Find(Hash);
@@ -1068,8 +1164,7 @@ public:
 				// With non-fname stats, the same string can appear with different pointers.
 				// We need to look up the stat in the ansi stat register to see if it's actually unique
 				uint32 AnsiNameIndex = FAnsiStringRegister::GetUniqueStringIndex((ANSICHAR*)InStatIDRaw);
-				FUniqueID AnsiUniqueID;
-				AnsiUniqueID.Hash = UniqueID.Hash;
+				FCsvUniqueStatID AnsiUniqueID(UniqueID);
 				AnsiUniqueID.Fields.FNameOrIndex = AnsiNameIndex;
 				int32 *AnsiIndexPtr = AnsiStringStatIDToIndex.Find(AnsiUniqueID.Hash);
 				if (AnsiIndexPtr)
@@ -1139,18 +1234,6 @@ protected:
 	TArray<FString> StatNames;
 	TArray<int32> StatCategoryIndices;
 	TArray<uint8> StatFlags;
-
-	union FUniqueID
-	{
-		struct
-		{
-			uint64 IsFName : 1;
-			uint64 IsCountStat : 1;
-			uint64 CategoryIndex : 11;
-			uint64 FNameOrIndex : 51;
-		} Fields;
-		uint64 Hash;
-	};
 };
 
 //-----------------------------------------------------------------------------
@@ -2862,6 +2945,8 @@ void FCsvProfiler::EndFrame()
 	check(IsInGameThread());
 	if (GCsvProfilerIsCapturing)
 	{
+		GCsvPersistentCustomStats.RecordStats();
+
 		QUICK_SCOPE_CYCLE_COUNTER(STAT_FCsvProfiler_EndFrame_Capturing);
 		if (NumFramesToCapture >= 0)
 		{
@@ -3697,11 +3782,29 @@ float FCsvProfiler::ProcessStatData()
 	return ElapsedMS;
 }
 
+TCsvPersistentCustomStat<int32>* FCsvProfiler::GetOrCreatePersistentCustomStatInt(FName Name, int32 CategoryIndex, bool bResetEachFrame)
+{
+	return GCsvPersistentCustomStats.GetOrCreatePersistentCustomStat<int32>(Name, CategoryIndex, bResetEachFrame);
+}
+
+TCsvPersistentCustomStat<float>* FCsvProfiler::GetOrCreatePersistentCustomStatFloat(FName Name, int32 CategoryIndex, bool bResetEachFrame)
+{
+	return GCsvPersistentCustomStats.GetOrCreatePersistentCustomStat<float>(Name, CategoryIndex, bResetEachFrame);
+}
+
 #if CSV_PROFILER_ALLOW_DEBUG_FEATURES
 
 // Simple benchmarking and debugging tests for the csv profiler. Enable with -csvtest, e.g -csvtest -csvcaptureframes=400
 void CSVTest()
 {
+	TCsvPersistentCustomStat<float>* PersistentStatFloat = FCsvProfiler::Get()->GetOrCreatePersistentCustomStatFloat(TEXT("PersistentStatFloat"));
+	PersistentStatFloat->Add(0.15f);
+	PersistentStatFloat->Sub(0.1f);
+
+	TCsvPersistentCustomStat<int>* PersistentStatInt = FCsvProfiler::Get()->GetOrCreatePersistentCustomStatInt(TEXT("PersistentStatInt"), CSV_CATEGORY_INDEX(CsvTest));
+	PersistentStatInt->Add(15);
+	PersistentStatInt->Sub(1);
+
 	uint32 FrameNumber = FCsvProfiler::Get()->GetCaptureFrameNumber();
 	CSV_SCOPED_TIMING_STAT(CsvTest, CsvTestStat);
 	CSV_CUSTOM_STAT(CsvTest, CaptureFrameNumber, int32(FrameNumber), ECsvCustomStatOp::Set);
diff --git a/Engine/Source/Runtime/Core/Public/ProfilingDebugging/CsvProfiler.h b/Engine/Source/Runtime/Core/Public/ProfilingDebugging/CsvProfiler.h
--- a/Engine/Source/Runtime/Core/Public/ProfilingDebugging/CsvProfiler.h
+++ b/Engine/Source/Runtime/Core/Public/ProfilingDebugging/CsvProfiler.h
@@ -117,6 +117,7 @@
 class FCsvProfilerFrame;
 class FCsvProfilerThreadData;
 class FCsvProfilerProcessingThread;
+struct FCsvPersistentCustomStats;
 class FName;
 
 enum class ECsvCustomStatOp : uint8
@@ -191,6 +192,115 @@ struct FCsvCaptureCommand
 	TSharedFuture<FString> Future;
 };
 
+//
+// Persistent custom stats
+// 
+enum class ECsvPersistentCustomStatType : uint8
+{
+	Float,
+	Int,
+	Count
+};
+
+class FCsvPersistentCustomStatBase
+{
+public:
+	FCsvPersistentCustomStatBase(FName InName, uint32 InCategoryIndex, bool bInResetEachFrame, ECsvPersistentCustomStatType InStatType)
+		: Name(InName), CategoryIndex(InCategoryIndex), bResetEachFrame(bInResetEachFrame), StatType(InStatType)
+	{
+	}
+
+	ECsvPersistentCustomStatType GetStatType() const
+	{
+		return StatType;
+	}
+protected:
+	FName Name;
+	uint32 CategoryIndex;
+	bool bResetEachFrame;
+	ECsvPersistentCustomStatType StatType;
+};
+
+
+template <class T> 
+class TCsvPersistentCustomStat : public FCsvPersistentCustomStatBase
+{
+	friend struct FCsvPersistentCustomStats;
+private:
+	std::atomic<T> Value;
+};
+
+
+template <>
+class TCsvPersistentCustomStat<float> : public FCsvPersistentCustomStatBase
+{
+	friend struct FCsvPersistentCustomStats;
+public:
+	TCsvPersistentCustomStat<float>(FName InName, uint32 InCategoryIndex, bool bInResetEachFrame)
+		: FCsvPersistentCustomStatBase(InName, InCategoryIndex, bInResetEachFrame, ECsvPersistentCustomStatType::Float)
+		, Value(0.0f)
+	{}
+
+	float Add(float Rhs)
+	{
+		float Previous = Value.load(std::memory_order_consume);
+		float Desired = Previous + Rhs;
+		while (!Value.compare_exchange_weak(Previous, Desired, std::memory_order_release, std::memory_order_consume))
+		{
+			Desired = Previous + Rhs;
+		}
+		return Desired;
+	}
+	float Sub(float Rhs)
+	{
+		return Add(-Rhs);
+	}
+	void Set(float NewVal)
+	{
+		return Value.store(NewVal, std::memory_order_relaxed);
+	}
+	float GetValue()
+	{
+		return Value.load();
+	}
+	static ECsvPersistentCustomStatType GetClassStatType() { return ECsvPersistentCustomStatType::Float; }
+protected:
+	std::atomic<float> Value;
+};
+
+
+template <>
+class TCsvPersistentCustomStat<int32> : public FCsvPersistentCustomStatBase
+{
+	friend struct FCsvPersistentCustomStats;
+public:
+	TCsvPersistentCustomStat<int32>(FName InName, uint32 InCategoryIndex, bool bInResetEachFrame)
+		: FCsvPersistentCustomStatBase(InName, InCategoryIndex, bInResetEachFrame, ECsvPersistentCustomStatType::Int)
+		, Value(0)
+	{}
+
+	int32 Add(int32 Rhs)
+	{
+		return Value.fetch_add(Rhs, std::memory_order_relaxed);
+	}
+	int32 Sub(int32 Rhs)
+	{
+		return Value.fetch_sub(Rhs, std::memory_order_relaxed);
+	}
+	void Set(int32 NewVal)
+	{
+		Value.store(NewVal, std::memory_order_relaxed);
+	}
+	int32 GetValue()
+	{
+		return Value.load();
+	}
+	static ECsvPersistentCustomStatType GetClassStatType() { return ECsvPersistentCustomStatType::Int; }
+protected:
+	std::atomic<int> Value;
+};
+
+
 /**
 * FCsvProfiler class. This manages recording and reporting all for CSV stats
 */
@@ -313,6 +423,11 @@ public:
 		RHIThreadId = InRHIThreadId;
 	}
 
+	// Persistent custom stat methods. These are pre-registered custom stats whose value can persist over multiple frames
+	// They are lower overhead than normal custom stats because accumulate ops don't require memory allocation
+	CORE_API TCsvPersistentCustomStat<int32>* GetOrCreatePersistentCustomStatInt(FName Name, int32 CategoryIndex = CSV_CATEGORY_INDEX_GLOBAL, bool bResetEachFrame = false);
+	CORE_API TCsvPersistentCustomStat<float>* GetOrCreatePersistentCustomStatFloat(FName Name, int32 CategoryIndex = CSV_CATEGORY_INDEX_GLOBAL, bool bResetEachFrame = false);
+
 private:
 	CORE_API static void VARARGS RecordEventfInternal(int32 CategoryIndex, const TCHAR* Fmt, ...);
 
