TagTrace: Handle LLM TagNames that implicitly declare their parent with directory separators (ParentA/ParentB/LeafName).
#rnx

diff --git a/Engine/Source/Runtime/Core/Private/ProfilingDebugging/TagTrace.cpp b/Engine/Source/Runtime/Core/Private/ProfilingDebugging/TagTrace.cpp
--- a/Engine/Source/Runtime/Core/Private/ProfilingDebugging/TagTrace.cpp
+++ b/Engine/Source/Runtime/Core/Private/ProfilingDebugging/TagTrace.cpp
@@ -2,6 +2,7 @@
 #include "ProfilingDebugging/TagTrace.h"
 
 #include "Experimental/Containers/GrowOnlyLockFreeHash.h"
+#include "Containers/StringView.h"
 #include "CoreTypes.h"
 #include "ProfilingDebugging/MemoryTrace.h"
 #include "HAL/LowLevelMemTracker.h"
@@ -242,7 +243,17 @@ int32 FTagTrace::AnnounceFNameTag(const FName& Name)
 	// First time encountering this name, announce it
 	ANSICHAR NameString[NAME_SIZE];
 	Name.GetPlainANSIString(NameString);
-	return AnnounceCustomTag(NameIndex, -1, NameString);
+
+	int32 ParentTag = -1;
+	FAnsiStringView NameView(NameString);
+	int32 LeafStart;
+	if (NameView.FindLastChar('/', LeafStart))
+	{
+		FName ParentName(NameView.Left(LeafStart));
+		ParentTag = AnnounceFNameTag(ParentName);
+	}
+
+	return AnnounceCustomTag(NameIndex, ParentTag, NameString);
 }
 
 ////////////////////////////////////////////////////////////////////////////////
