Update the physical and virtual memory for iOS based devices.  Since Apple has added new entitlements to allow for additional app memory (physical and virtual (file memmapped)), the Jetsam limits are no longer accurate.  Instead, calculate actual available/used memory from vm_stats and os_proc_available_memory.

diff --git a/Engine/Source/Runtime/Core/Private/Apple/ApplePlatformMemory.cpp b/Engine/Source/Runtime/Core/Private/Apple/ApplePlatformMemory.cpp
--- a/Engine/Source/Runtime/Core/Private/Apple/ApplePlatformMemory.cpp
+++ b/Engine/Source/Runtime/Core/Private/Apple/ApplePlatformMemory.cpp
@@ -337,44 +337,48 @@ FMalloc* FApplePlatformMemory::BaseAllocator()
 FPlatformMemoryStats FApplePlatformMemory::GetStats()
 {
 	const FPlatformMemoryConstants& MemoryConstants = FPlatformMemory::GetConstants();
-#if PLATFORM_IOS
-	const uint64 MaxVirtualMemory = 1ull << 34; // set to 16GB for now since IOS can see a maximum of 8GB
-#endif
 	static FPlatformMemoryStats MemoryStats;
 	
 	// Gather platform memory stats.
+	vm_statistics Stats;
+	mach_msg_type_number_t StatsSize = sizeof(Stats);
+	FMemory::Memset(&Stats, 0, StatsSize);
+	if (host_statistics(mach_host_self(), HOST_VM_INFO, (host_info_t)&Stats, &StatsSize))
+	{
+		UE_LOG(LogTemp, Warning, TEXT("Failed to fetch vm statistics"));
+	}
+
 	uint64_t FreeMem = 0;
 #if PLATFORM_IOS
 	FreeMem = os_proc_available_memory();
 #else
-    vm_statistics Stats;
-    mach_msg_type_number_t StatsSize = sizeof(Stats);
-    host_statistics(mach_host_self(), HOST_VM_INFO, (host_info_t)&Stats, &StatsSize);
     FreeMem = (Stats.free_count + Stats.inactive_count) * MemoryConstants.PageSize;
 #endif
+
 	MemoryStats.AvailablePhysical = FreeMem;
-	MemoryStats.AvailableVirtual = FreeMem;
+
+	// Calculate the number of available free pages on iOS/macOS. Apple considers "inactive_count" pages
+	// as pages that the app says it doesn't need anymore and can be recycled/freed, but could be reactivated 
+	// if the app does request those resources again.  Apple tries to maximize use of memory and 
+	// considers "free" pages a waste of resources.
+	MemoryStats.AvailableVirtual = (Stats.free_count + Stats.inactive_count) * MemoryConstants.PageSize;
 	
 	// Just get memory information for the process and report the working set instead
 	mach_task_basic_info_data_t TaskInfo;
 	mach_msg_type_number_t TaskInfoCount = MACH_TASK_BASIC_INFO_COUNT;
 	task_info( mach_task_self(), MACH_TASK_BASIC_INFO, (task_info_t)&TaskInfo, &TaskInfoCount );
-#if PLATFORM_IOS
-    MemoryStats.UsedPhysical = MemoryConstants.TotalPhysical - FreeMem;
-#else
-    MemoryStats.UsedPhysical = TaskInfo.resident_size;
-#endif
 
+    MemoryStats.UsedPhysical = TaskInfo.resident_size;
 	if(MemoryStats.UsedPhysical > MemoryStats.PeakUsedPhysical)
 	{
 		MemoryStats.PeakUsedPhysical = MemoryStats.UsedPhysical;
 	}
-	MemoryStats.UsedVirtual = TaskInfo.virtual_size;
-#if PLATFORM_IOS
-	if(MemoryStats.UsedVirtual > MemoryStats.PeakUsedVirtual || MemoryStats.PeakUsedVirtual > MaxVirtualMemory)
-#else
+
+	uint64 VMemUsed = (Stats.active_count +
+					   Stats.wire_count) *
+					  MemoryConstants.PageSize;
+	MemoryStats.UsedVirtual = VMemUsed;
 	if(MemoryStats.UsedVirtual > MemoryStats.PeakUsedVirtual)
-#endif
 	{
 		MemoryStats.PeakUsedVirtual = MemoryStats.UsedVirtual;
 	}
@@ -389,45 +393,50 @@ const FPlatformMemoryConstants& FApplePlatformMemory::GetConstants()
 	if( MemoryConstants.TotalPhysical == 0 )
 	{
 		// Gather platform memory constants.
-		
-		// Get memory.
 		int64 AvailablePhysical = 0;
 #if PLATFORM_IOS
-        AvailablePhysical = os_proc_available_memory();
-        
-        // quantize to the known jetsam limits, we should be within 50MB of the correct one
-        uint64 JetsamLimits[] = { 1520435200, 1939865600, 2201170740, 2252710350, 3006477100 }; // { 2GB, gimped 3GB, gimped 4GB, 3GB, 4GB
-        if (AvailablePhysical < JetsamLimits[0])
-        {
-            AvailablePhysical = JetsamLimits[0];
-        }
-        else if (AvailablePhysical < JetsamLimits[1])
-        {
-            AvailablePhysical = JetsamLimits[1];
-        }
-        else if (AvailablePhysical < JetsamLimits[2])
-        {
-            AvailablePhysical = JetsamLimits[2];
-        }
-        else if (AvailablePhysical < JetsamLimits[3])
-        {
-            AvailablePhysical = JetsamLimits[3];
-        }
-        else if (AvailablePhysical < JetsamLimits[4])
-        {
-            AvailablePhysical = JetsamLimits[4];
-        }
+		AvailablePhysical = os_proc_available_memory();
 #else
-        int Mib[] = {CTL_HW, HW_MEMSIZE};
-        size_t Length = sizeof(int64);
-        sysctl(Mib, 2, &AvailablePhysical, &Length, NULL, 0);
+		int Mib[] = {CTL_HW, HW_MEMSIZE};
+		size_t Length = sizeof(int64);
+		sysctl(Mib, 2, &AvailablePhysical, &Length, NULL, 0);
 #endif
+		
 		MemoryConstants.TotalPhysical = AvailablePhysical;
-		MemoryConstants.TotalVirtual = AvailablePhysical;
+		MemoryConstants.TotalVirtual = AvailablePhysical;	// Calculate true value below, but default to physical if vmstats call fails 
 		MemoryConstants.PageSize = (uint32)vm_page_size;
 		MemoryConstants.OsAllocationGranularity = (uint32)vm_page_size;
 		MemoryConstants.BinnedPageSize = FMath::Max((SIZE_T)65536, (SIZE_T)vm_page_size);
-		MemoryConstants.TotalPhysicalGB = (MemoryConstants.TotalPhysical + 1024 * 1024 * 1024 - 1) / 1024 / 1024 / 1024;
+		
+		// macOS reports the correct amount of physical memory, however iOS reports a lower amount of 
+		// actual physical memory. To work around this, we add 1Gb - 1b so it will be truncated 
+		// correctly and will not affect macOS
+		MemoryConstants.TotalPhysicalGB = ([NSProcessInfo processInfo].physicalMemory + (1024*1024*1024 - 1)) / 1024 / 1024 / 1024;
+
+		// Calculate total and available Virtual Memory
+		mach_port_t HostPort = mach_host_self();
+		mach_msg_type_number_t HostSize = sizeof(vm_statistics_data_t) / sizeof(integer_t);
+
+		// verify that actual device pagesize matches defined size in vm_page_size.h
+		vm_size_t PageSize = 0;
+		host_page_size(HostPort, &PageSize);
+		ensure(vm_page_size == PageSize);
+
+		vm_statistics_data_t vm_stat;
+		if (host_statistics(HostPort, HOST_VM_INFO, (host_info_t)&vm_stat, &HostSize) != KERN_SUCCESS)
+		{
+			UE_LOG(LogTemp, Warning, TEXT("Failed to fetch vm statistics"));
+			return MemoryConstants;
+		}
+
+		uint64 VMemUsed = (vm_stat.active_count +
+						  vm_stat.inactive_count +
+						  vm_stat.wire_count) *
+						 PageSize;
+		uint64 VMemFree = vm_stat.free_count * PageSize;
+		uint64 VMemTotal = VMemUsed + VMemFree;
+
+		MemoryConstants.TotalVirtual = VMemTotal;
 	}
 	
 	return MemoryConstants;
