Adding String/RemoveFrom.h

diff --git a/Engine/Source/Runtime/Core/Public/String/RemoveFrom.h b/Engine/Source/Runtime/Core/Public/String/RemoveFrom.h
new file mode 100644
--- /dev/null
+++ b/Engine/Source/Runtime/Core/Public/String/RemoveFrom.h
@@ -0,0 +1,21 @@
+// Copyright Epic Games, Inc. All Rights Reserved.
+
+#pragma once
+
+#include "CoreTypes.h"
+#include "Containers/StringFwd.h"
+
+namespace UE::String
+{
+	template <typename CharType>
+	[[nodiscard]] inline TStringView<CharType> RemoveFromStart(const TStringView<CharType> View, const TStringView<CharType> Prefix, ESearchCase::Type SearchCase = ESearchCase::CaseSensitive)
+	{
+		return View.StartsWith(Prefix, SearchCase) ? View.RightChop(Prefix.Len()) : View;
+	}
+
+	template <typename CharType>
+	[[nodiscard]] inline TStringView<CharType> RemoveFromEnd(const TStringView<CharType> View, const TStringView<CharType> Prefix, ESearchCase::Type SearchCase = ESearchCase::CaseSensitive)
+	{
+		return View.EndsWith(Prefix, SearchCase) ? View.LeftChop(Prefix.Len()) : View;
+	}
+} // UE::String
diff --git a/Engine/Source/Runtime/Core/Tests/String/RemoveFromTest.cpp b/Engine/Source/Runtime/Core/Tests/String/RemoveFromTest.cpp
new file mode 100644
--- /dev/null
+++ b/Engine/Source/Runtime/Core/Tests/String/RemoveFromTest.cpp
@@ -0,0 +1,79 @@
+// Copyright Epic Games, Inc. All Rights Reserved.
+
+#if WITH_LOW_LEVEL_TESTS
+
+#include "String/RemoveFrom.h"
+
+#include "Containers/StringView.h"
+#include "TestHarness.h"
+
+TEST_CASE("Core::String::RemoveFrom", "[Core][String][Smoke]")
+{
+	SECTION("RemoveFromStart")
+	{
+		CHECK(UE::String::RemoveFromStart(TEXTVIEW("AbCABCAbCABC"), TEXTVIEW("Ab")) == TEXTVIEW("CABCAbCABC"));
+		CHECK(UE::String::RemoveFromStart(TEXTVIEW("AbCABCAbCABC"), TEXTVIEW("Abc")) == TEXTVIEW("AbCABCAbCABC"));
+		CHECK(UE::String::RemoveFromStart(TEXTVIEW("AbCABCAbCABC"), TEXTVIEW("Abc"), ESearchCase::IgnoreCase) == TEXTVIEW("ABCAbCABC"));
+		CHECK(UE::String::RemoveFromStart(TEXTVIEW("AbCABCAbCABC"), TEXTVIEW("AbD"), ESearchCase::IgnoreCase) == TEXTVIEW("AbCABCAbCABC"));
+		
+		// Empty
+		CHECK(UE::String::RemoveFromStart(TEXTVIEW(""), TEXTVIEW("A")) == TEXTVIEW(""));
+		CHECK(UE::String::RemoveFromStart(TEXTVIEW(""), TEXTVIEW("ABC")) == TEXTVIEW(""));
+
+		// Exact
+		CHECK(UE::String::RemoveFromStart(TEXTVIEW("A"), TEXTVIEW("A")) == TEXTVIEW(""));
+		CHECK(UE::String::RemoveFromStart(TEXTVIEW("A"), TEXTVIEW("A"), ESearchCase::IgnoreCase) == TEXTVIEW(""));
+
+		// Duplicate
+		CHECK(UE::String::RemoveFromStart(TEXTVIEW("AA"), TEXTVIEW("A")) == TEXTVIEW("A"));
+		CHECK(UE::String::RemoveFromStart(TEXTVIEW("AA"), TEXTVIEW("A"), ESearchCase::IgnoreCase) == TEXTVIEW("A"));
+
+		// Different Case
+		CHECK(UE::String::RemoveFromStart(TEXTVIEW("A"), TEXTVIEW("a")) == TEXTVIEW("A"));
+		CHECK(UE::String::RemoveFromStart(TEXTVIEW("A"), TEXTVIEW("a"), ESearchCase::IgnoreCase) == TEXTVIEW(""));
+
+		// Substring
+		CHECK(UE::String::RemoveFromStart(TEXTVIEW("ABC"), TEXTVIEW("A")) == TEXTVIEW("BC"));
+		CHECK(UE::String::RemoveFromStart(TEXTVIEW("ABC"), TEXTVIEW("a")) == TEXTVIEW("ABC"));
+		CHECK(UE::String::RemoveFromStart(TEXTVIEW("ABC"), TEXTVIEW("a"), ESearchCase::IgnoreCase) == TEXTVIEW("BC"));
+
+		CHECK(UE::String::RemoveFromStart(TEXTVIEW("ABC"), TEXTVIEW("B")) == TEXTVIEW("ABC"));
+		CHECK(UE::String::RemoveFromStart(TEXTVIEW("ABC"), TEXTVIEW("C")) == TEXTVIEW("ABC"));
+		CHECK(UE::String::RemoveFromStart(TEXTVIEW("ABC"), TEXTVIEW("BC")) == TEXTVIEW("ABC"));
+	}
+
+	SECTION("RemoveFromEnd")
+	{
+		CHECK(UE::String::RemoveFromEnd(TEXTVIEW("AbCABCAbCABC"), TEXTVIEW("ABC")) == TEXTVIEW("AbCABCAbC"));
+		CHECK(UE::String::RemoveFromEnd(TEXTVIEW("AbCABCAbCABC"), TEXTVIEW("abc")) == TEXTVIEW("AbCABCAbCABC"));
+		CHECK(UE::String::RemoveFromEnd(TEXTVIEW("AbCABCAbCABC"), TEXTVIEW("abc"), ESearchCase::IgnoreCase) == TEXTVIEW("AbCABCAbC"));
+		CHECK(UE::String::RemoveFromEnd(TEXTVIEW("AbCABCAbCABC"), TEXTVIEW("abD"), ESearchCase::IgnoreCase) == TEXTVIEW("AbCABCAbCABC"));
+
+		// Empty
+		CHECK(UE::String::RemoveFromEnd(TEXTVIEW(""), TEXTVIEW("A")) == TEXTVIEW(""));
+		CHECK(UE::String::RemoveFromEnd(TEXTVIEW(""), TEXTVIEW("ABC")) == TEXTVIEW(""));
+
+		// Exact
+		CHECK(UE::String::RemoveFromEnd(TEXTVIEW("A"), TEXTVIEW("A")) == TEXTVIEW(""));
+		CHECK(UE::String::RemoveFromEnd(TEXTVIEW("A"), TEXTVIEW("A"), ESearchCase::IgnoreCase) == TEXTVIEW(""));
+
+		// Duplicate
+		CHECK(UE::String::RemoveFromEnd(TEXTVIEW("AA"), TEXTVIEW("A")) == TEXTVIEW("A"));
+		CHECK(UE::String::RemoveFromEnd(TEXTVIEW("AA"), TEXTVIEW("A"), ESearchCase::IgnoreCase) == TEXTVIEW("A"));
+
+		// Different Case
+		CHECK(UE::String::RemoveFromEnd(TEXTVIEW("A"), TEXTVIEW("a")) == TEXTVIEW("A"));
+		CHECK(UE::String::RemoveFromEnd(TEXTVIEW("A"), TEXTVIEW("a"), ESearchCase::IgnoreCase) == TEXTVIEW(""));
+
+		// Substring
+		CHECK(UE::String::RemoveFromEnd(TEXTVIEW("ABC"), TEXTVIEW("C")) == TEXTVIEW("AB"));
+		CHECK(UE::String::RemoveFromEnd(TEXTVIEW("ABC"), TEXTVIEW("c")) == TEXTVIEW("ABC"));
+		CHECK(UE::String::RemoveFromEnd(TEXTVIEW("ABC"), TEXTVIEW("c"), ESearchCase::IgnoreCase) == TEXTVIEW("AB"));
+
+		CHECK(UE::String::RemoveFromEnd(TEXTVIEW("ABC"), TEXTVIEW("A")) == TEXTVIEW("ABC"));
+		CHECK(UE::String::RemoveFromEnd(TEXTVIEW("ABC"), TEXTVIEW("B")) == TEXTVIEW("ABC"));
+		CHECK(UE::String::RemoveFromEnd(TEXTVIEW("ABC"), TEXTVIEW("AB")) == TEXTVIEW("ABC"));
+	}
+}
+
+#endif // WITH_LOW_LEVEL_TESTS
