Logging: Refined memory order used by the output device locks

#rnx

diff --git a/Engine/Source/Runtime/Core/Private/Misc/OutputDeviceRedirector.cpp b/Engine/Source/Runtime/Core/Private/Misc/OutputDeviceRedirector.cpp
--- a/Engine/Source/Runtime/Core/Private/Misc/OutputDeviceRedirector.cpp
+++ b/Engine/Source/Runtime/Core/Private/Misc/OutputDeviceRedirector.cpp
@@ -231,11 +231,11 @@ public:
 	{
 		TRACE_CPUPROFILER_EVENT_SCOPE(FOutputDevicesReadScopeLock);
 		// A write lock has set the LSB. Cancel this read lock and wait for the write.
-		State.OutputDevicesLockState.fetch_sub(2, std::memory_order_release);
+		State.OutputDevicesLockState.fetch_sub(2, std::memory_order_relaxed);
 		// This read lock will wait until the write lock exits.
 		FReadScopeLock ScopeLock(State.OutputDevicesLock);
-		// Acquire on this read lock because the write may have mutated state that we read.
-		uint32 LockState = State.OutputDevicesLockState.fetch_add(2, std::memory_order_acquire);
+		// This is relaxed because locking OutputDevicesLock has acquire semantics.
+		uint32 LockState = State.OutputDevicesLockState.fetch_add(2, std::memory_order_relaxed);
 		check((LockState & 1) == 0);
 	}
 
@@ -266,7 +266,7 @@ public:
 		State.OutputDevicesLock.WriteLock();
 		State.LockedThreadId.store(FPlatformTLS::GetCurrentThreadId(), std::memory_order_relaxed);
 		// Set the LSB to flag to read locks that a write lock is waiting.
-		uint32 LockState = State.OutputDevicesLockState.fetch_or(uint32(1), std::memory_order_relaxed);
+		uint32 LockState = State.OutputDevicesLockState.fetch_or(uint32(1), std::memory_order_acquire);
 		check((LockState & 1) == 0);
 		if (LockState > 1)
 		{
@@ -274,7 +274,7 @@ public:
 			do
 			{
 				FPlatformProcess::Sleep(0);
-				LockState = State.OutputDevicesLockState.load(std::memory_order_relaxed);
+				LockState = State.OutputDevicesLockState.load(std::memory_order_acquire);
 			}
 			while (LockState > 1);
 		}
