Change to looping in ConvertToUTF8 to make some other upcoming refactors easier.

diff --git a/Engine/Source/Runtime/Core/Private/GenericPlatform/GenericPlatformString.cpp b/Engine/Source/Runtime/Core/Private/GenericPlatform/GenericPlatformString.cpp
--- a/Engine/Source/Runtime/Core/Private/GenericPlatform/GenericPlatformString.cpp
+++ b/Engine/Source/Runtime/Core/Private/GenericPlatform/GenericPlatformString.cpp
@@ -184,14 +184,21 @@ namespace UE::Core::Private
 	}
 
 	template <typename DestBufferType, typename FromType>
-	static int32 ConvertToUTF8(DestBufferType& Dest, int32 DestLen, const FromType* Source, const int32 SourceLen)
+	static int32 ConvertToUTF8(DestBufferType& Dest, int32 DestLen, const FromType* Source, int32 SourceLen)
 	{
-		DestBufferType DestStartingPosition = Dest;
 		if constexpr (sizeof(FromType) == 4)
 		{
-			for (int32 i = 0; i < SourceLen; ++i)
+			DestBufferType DestStartingPosition = Dest;
+
+			for (;;)
 			{
-				uint32 Codepoint = static_cast<uint32>(Source[i]);
+				if (SourceLen <= 0)
+				{
+					return UE_PTRDIFF_TO_INT32(Dest - DestStartingPosition);
+				}
+
+				uint32 Codepoint = static_cast<uint32>(*Source++);
+				--SourceLen;
 
 				if (!WriteCodepointToBuffer(Codepoint, Dest, DestLen))
 				{
@@ -202,18 +209,26 @@ namespace UE::Core::Private
 		}
 		else
 		{
+			DestBufferType DestStartingPosition = Dest;
+
 			uint32 HighSurrogate = MAX_uint32;
 
-			for (int32 i = 0; i < SourceLen; ++i)
+			for (;;)
 			{
+				if (SourceLen <= 0)
+				{
+					return UE_PTRDIFF_TO_INT32(Dest - DestStartingPosition);
+				}
+
 				const bool bHighSurrogateIsSet = HighSurrogate != MAX_uint32;
-				uint32 Codepoint = static_cast<uint32>(Source[i]);
+				uint32 Codepoint = static_cast<uint32>(*Source++);
+				--SourceLen;
 
 				// Check if this character is a high-surrogate
 				if (IsHighSurrogate(Codepoint))
 				{
 					// Ensure we don't already have a high-surrogate set or end without a matching low-surrogate
-					if (bHighSurrogateIsSet || i == SourceLen - 1)
+					if (bHighSurrogateIsSet || SourceLen == 0)
 					{
 						// Already have a high-surrogate in this pair or string ends with lone high-surrogate
 						// Write our stored value (will be converted into bogus character)
@@ -259,8 +274,6 @@ namespace UE::Core::Private
 				}
 			}
 		}
-
-		return UE_PTRDIFF_TO_INT32(Dest - DestStartingPosition);
 	}
 
 	static uint32 CodepointFromUtf8(const UTF8CHAR*& SourceString, const uint32 SourceLengthRemaining)
