Improve CPU thermal sensor detection on Android

diff --git a/Engine/Config/Android/AndroidEngine.ini b/Engine/Config/Android/AndroidEngine.ini
--- a/Engine/Config/Android/AndroidEngine.ini
+++ b/Engine/Config/Android/AndroidEngine.ini
@@ -93,25 +93,6 @@ r.TemporalAA.R11G11B10History=0
 [OnlineSubsystemGooglePlay.Store]
 bSupportsInAppPurchasing=true
 
-[ThermalSensors]
-+SensorLocations=/sys/devices/virtual/thermal/thermal_zone1/temp
-+SensorLocations=/sys/devices/system/cpu/cpu0/cpufreq/cpu_temp
-+SensorLocations=/sys/class/thermal/thermal_zone0/temp
-+SensorLocations=/sys/class/thermal/thermal_zone1/temp
-+SensorLocations=/sys/class/hwmon/hwmon0/device/temp1_input
-+SensorLocations=/sys/class/hwmon/hwmon0/temp1_input
-+SensorLocations=/sys/devices/virtual/hwmon/hwmon1/temp1_input
-+SensorLocations=/sys/devices/system/cpu/cpu0/cpufreq/FakeShmoo_cpu_temp
-+SensorLocations=/sys/class/i2c-adapter/i2c-4/4-004c/temperature
-+SensorLocations=/sys/devices/platform/tegra-i2c.3/i2c-4/4-004c/temperature
-+SensorLocations=/sys/devices/platform/tegra-i2c.3/i2c-4/4-004c/ext_temperature
-+SensorLocations=/sys/devices/platform/omap/omap_temp_sensor.0/temperature
-+SensorLocations=/sys/devices/platform/tegra_tmon/temp1_input
-+SensorLocations=/sys/kernel/debug/tegra_thermal/temp_tj
-+SensorLocations=/sys/devices/platform/tegra-tsensor/tsensor_temperature
-+SensorLocations=/sys/devices/platform/s5p-tmu/temperature
-+SensorLocations=/sys/devices/platform/s5p-tmu/curr_temp
-+SensorLocations=/sys/htc/cpu_temp
 
 [/Script/Engine.RendererSettings]
 r.DistanceFields=0
diff --git a/Engine/Source/Runtime/Core/Private/Android/AndroidPlatformMisc.cpp b/Engine/Source/Runtime/Core/Private/Android/AndroidPlatformMisc.cpp
--- a/Engine/Source/Runtime/Core/Private/Android/AndroidPlatformMisc.cpp
+++ b/Engine/Source/Runtime/Core/Private/Android/AndroidPlatformMisc.cpp
@@ -173,6 +173,8 @@ static void InitCpuThermalSensor()
 	CVarAndroidCPUThermalSensorFilePath->SetOnChangedCallback(FConsoleVariableDelegate::CreateStatic(&OverrideCpuThermalSensorFileFromCVar));
 
 	uint32 Counter = 0;
+	const uint32 INVALID_INDEX = -1;
+	uint32 CPUSensorIndex = INVALID_INDEX;
 	while (true)
 	{
 		char Buf[256] = "";
@@ -188,6 +190,12 @@ static void InitCpuThermalSensor()
 			}
 			*Ptr = 0;
 
+			if (strstr(Buf, "cpu-") && CPUSensorIndex == INVALID_INDEX)
+			{
+				CPUSensorIndex = Counter;
+				FCStringAnsi::Sprintf(AndroidCpuThermalSensorFileBuf, "/sys/devices/virtual/thermal/thermal_zone%u/temp", Counter);
+			}
+
 			UE_LOG(LogAndroid, Display, TEXT("Detected thermal sensor `%s` at /sys/devices/virtual/thermal/thermal_zone%u/temp"), ANSI_TO_TCHAR(Buf), Counter);
 			++Counter;
 		}
@@ -213,7 +221,14 @@ static void InitCpuThermalSensor()
 		}
 	}
 
-	UE_LOG(LogAndroid, Display, TEXT("No CPU thermal sensor was detected. To manually override the sensor path set android.CPUThermalSensorFilePath CVar."));
+	if (CPUSensorIndex != INVALID_INDEX)
+	{
+		UE_LOG(LogAndroid, Display, TEXT("Selecting thermal sensor located at `%s`"), ANSI_TO_TCHAR(AndroidCpuThermalSensorFileBuf));
+	}
+	else
+	{
+		UE_LOG(LogAndroid, Display, TEXT("No CPU thermal sensor was detected. To manually override the sensor path set android.CPUThermalSensorFilePath CVar."));
+	}
 }
 
 void FAndroidMisc::RequestExit( bool Force )
diff --git a/Engine/Source/Runtime/Launch/Private/Android/AndroidJNI.cpp b/Engine/Source/Runtime/Launch/Private/Android/AndroidJNI.cpp
--- a/Engine/Source/Runtime/Launch/Private/Android/AndroidJNI.cpp
+++ b/Engine/Source/Runtime/Launch/Private/Android/AndroidJNI.cpp
@@ -2077,6 +2077,35 @@ JNI_METHOD void Java_com_epicgames_unreal_NativeCalls_RouteServiceIntent(JNIEnv*
 JNI_METHOD void Java_com_epicgames_unreal_GameActivity_nativeOnThermalStatusChangedListener(JNIEnv* jenv, jobject thiz, jint Status)
 {
 	FAndroidStats::OnThermalStatusChanged(Status);
+
+	FCoreDelegates::ETemperatureSeverity Severity;
+	switch (Status)
+	{
+	case 0:
+	case 1:
+		Severity = FCoreDelegates::ETemperatureSeverity::Good;
+		break;
+
+	case 2:
+		Severity = FCoreDelegates::ETemperatureSeverity::Bad;
+		break;
+
+	case 3:
+	case 4:
+		Severity = FCoreDelegates::ETemperatureSeverity::Serious;
+		break;
+
+	case 5:
+	case 6:
+		Severity = FCoreDelegates::ETemperatureSeverity::Critical;
+		break;
+
+	default:
+		Severity = FCoreDelegates::ETemperatureSeverity::Unknown;
+		break;
+	}
+
+	FCoreDelegates::OnTemperatureChange.Broadcast(Severity);
 }
 
 static void OnTrimMessage(int MemoryTrimValue)
