Deprecate FAutomationTestBase::SetSuccessState() and  FAutomationTestBase::GetSuccessState()

diff --git a/Engine/Source/Runtime/Core/Private/Misc/AutomationTest.cpp b/Engine/Source/Runtime/Core/Private/Misc/AutomationTest.cpp
--- a/Engine/Source/Runtime/Core/Private/Misc/AutomationTest.cpp
+++ b/Engine/Source/Runtime/Core/Private/Misc/AutomationTest.cpp
@@ -823,7 +823,7 @@ bool FAutomationTestFramework::InternalStopTest(FAutomationTestExecutionInfo& Ou
 	CurrentTest->ExpectedErrors.Empty();
 
 	// Set the success state of the test based on the above criteria
-	CurrentTest->SetSuccessState(bTestSuccessful);
+	CurrentTest->InternalSetSuccessState(bTestSuccessful);
 
 	OnTestEndEvent.Broadcast(CurrentTest);
 
@@ -1192,12 +1192,12 @@ bool FAutomationTestBase::HasMetExpectedErrors()
 	return HasMetAllExpectedErrors;
 }
 
-void FAutomationTestBase::SetSuccessState( bool bSuccessful )
+void FAutomationTestBase::InternalSetSuccessState( bool bSuccessful )
 {
 	ExecutionInfo.bSuccessful = bSuccessful;
 }
 
-bool FAutomationTestBase::GetSuccessState()
+bool FAutomationTestBase::GetLastExecutionSuccessState()
 {
 	return ExecutionInfo.bSuccessful;
 }
diff --git a/Engine/Source/Runtime/Core/Public/Misc/AutomationTest.h b/Engine/Source/Runtime/Core/Public/Misc/AutomationTest.h
--- a/Engine/Source/Runtime/Core/Public/Misc/AutomationTest.h
+++ b/Engine/Source/Runtime/Core/Public/Misc/AutomationTest.h
@@ -1348,16 +1348,21 @@ public:
 	bool HasMetExpectedErrors();
 
 	/**
-	 * Forcibly sets whether the test has succeeded or not
-	 *
-	 * @param	bSuccessful	true to mark the test successful, false to mark the test as failed
+	 * Return the last success state for this test
+	 */
+	bool GetLastExecutionSuccessState();
+
+	/**
+	 * [Deprecated] Use AddError(msg) instead to change the state of the test to a failure
 	 */
-	void SetSuccessState( bool bSuccessful );
+	UE_DEPRECATED(5.1, "Use AddError(msg) instead to change the state of the test to a failure.")
+	void SetSuccessState(bool bSuccessful) { }
 
 	/**
-	 * Return the test success state
+	 * [Deprecated] Return the last success state for this test
 	 */
-	bool GetSuccessState();
+	UE_DEPRECATED(5.1, "Use GetLastExecutionSuccessState instead.")
+	bool GetSuccessState() { return GetLastExecutionSuccessState(); }
 
 	/**
 	 * Populate the provided execution info object with the execution info contained within the test. Not particularly efficient,
@@ -1840,6 +1845,13 @@ private:
 	*/
 	bool IsExpectedError(const FString& Error);
 
+	/**
+	 * Sets whether the test has succeeded or not
+	 *
+	 * @param	bSuccessful	true to mark the test successful, false to mark the test as failed
+	 */
+	void InternalSetSuccessState(bool bSuccessful);
+
 	/* Errors to be expected while processing this test.*/
 	TArray< FAutomationExpectedError> ExpectedErrors;
 
