Fix for 'wrong' high surrogate output.

diff --git a/Engine/Source/Runtime/Core/Private/GenericPlatform/GenericPlatformString.cpp b/Engine/Source/Runtime/Core/Private/GenericPlatform/GenericPlatformString.cpp
--- a/Engine/Source/Runtime/Core/Private/GenericPlatform/GenericPlatformString.cpp
+++ b/Engine/Source/Runtime/Core/Private/GenericPlatform/GenericPlatformString.cpp
@@ -227,11 +227,21 @@ namespace UE::Core::Private
 				// Check if this character is a high-surrogate
 				if (IsHighSurrogate(Codepoint))
 				{
-					// Ensure we don't already have a high-surrogate set or end without a matching low-surrogate
-					if (bHighSurrogateIsSet || SourceLen == 0)
+					if (SourceLen == 0)
 					{
-						// Already have a high-surrogate in this pair or string ends with lone high-surrogate
-						// Write our stored value (will be converted into bogus character)
+						// String ends with lone high-surrogate - write out surrogate (will be converted into bogus character)
+						if (!WriteCodepointToBuffer(Codepoint, Dest, DestLen))
+						{
+							// Could not write data, bail out
+							return -1;
+						}
+
+						return UE_PTRDIFF_TO_INT32(Dest - DestStartingPosition);
+					}
+
+					if (bHighSurrogateIsSet)
+					{
+						// Already have a high-surrogate in this pair - write our stored value (will be converted into bogus character)
 						if (!WriteCodepointToBuffer(HighSurrogate, Dest, DestLen))
 						{
 							// Could not write data, bail out
