Ran a pass of IWYU on Core.

Things to note regarding this pass:
* No includes are currently being removed from public headers.
* Any private file that has a #if is currently not being updated.

diff --git a/Engine/Source/Runtime/Core/Internal/IO/IoDirectoryIndex.h b/Engine/Source/Runtime/Core/Internal/IO/IoDirectoryIndex.h
--- a/Engine/Source/Runtime/Core/Internal/IO/IoDirectoryIndex.h
+++ b/Engine/Source/Runtime/Core/Internal/IO/IoDirectoryIndex.h
@@ -2,12 +2,10 @@
 
 #pragma once
 
+#include "Containers/StringFwd.h"
 #include "CoreMinimal.h"
-#include "Containers/StringView.h"
 #include "Misc/AES.h"
 
-class FArchive;
-
 namespace IoDirectoryIndexUtils
 {
 	FString GetCommonRootPath(const TArray<FString>& Filenames);
diff --git a/Engine/Source/Runtime/Core/Private/Async/Async.cpp b/Engine/Source/Runtime/Core/Private/Async/Async.cpp
--- a/Engine/Source/Runtime/Core/Private/Async/Async.cpp
+++ b/Engine/Source/Runtime/Core/Private/Async/Async.cpp
@@ -2,6 +2,9 @@
 
 #include "Async/Async.h"
 
+#include "CoreMinimal.h"
+#include "HAL/ThreadSafeCounter.h"
+
 
 /* Private types
  *****************************************************************************/
diff --git a/Engine/Source/Runtime/Core/Private/Async/Fundamental/ReserveScheduler.cpp b/Engine/Source/Runtime/Core/Private/Async/Fundamental/ReserveScheduler.cpp
--- a/Engine/Source/Runtime/Core/Private/Async/Fundamental/ReserveScheduler.cpp
+++ b/Engine/Source/Runtime/Core/Private/Async/Fundamental/ReserveScheduler.cpp
@@ -1,7 +1,13 @@
 // Copyright Epic Games, Inc. All Rights Reserved.
 
 #include "Async/Fundamental/ReserveScheduler.h"
+
+#include "Async/Fundamental/LocalQueue.h"
+#include "Async/Fundamental/Scheduler.h"
+#include "HAL/PlatformProcess.h"
 #include "Misc/ScopeLock.h"
+#include "ProfilingDebugging/CpuProfilerTrace.h"
+#include "Trace/Trace.h"
 
 extern CORE_API bool GTaskGraphUseDynamicPrioritization;
 
diff --git a/Engine/Source/Runtime/Core/Private/Async/LazyEvent.cpp b/Engine/Source/Runtime/Core/Private/Async/LazyEvent.cpp
--- a/Engine/Source/Runtime/Core/Private/Async/LazyEvent.cpp
+++ b/Engine/Source/Runtime/Core/Private/Async/LazyEvent.cpp
@@ -1,7 +1,9 @@
 // Copyright Epic Games, Inc. All Rights Reserved.
 
 #include "Experimental/Async/LazyEvent.h"
+
 #include "HAL/PlatformProcess.h"
+#include "Math/NumericLimits.h"
 
 namespace UE::Core::Private
 {
diff --git a/Engine/Source/Runtime/Core/Private/Async/ParallelFor.cpp b/Engine/Source/Runtime/Core/Private/Async/ParallelFor.cpp
--- a/Engine/Source/Runtime/Core/Private/Async/ParallelFor.cpp
+++ b/Engine/Source/Runtime/Core/Private/Async/ParallelFor.cpp
@@ -2,6 +2,9 @@
 
 #include "Async/ParallelFor.h"
 
+#include "CoreTypes.h"
+#include "HAL/IConsoleManager.h"
+
 CORE_API int32 GParallelForBackgroundYieldingTimeoutMs = 8;
 static FAutoConsoleVariableRef CVarParallelForBackgroundYieldingTimeout(
 	TEXT("Async.ParallelFor.YieldingTimeout"),
diff --git a/Engine/Source/Runtime/Core/Private/Compression/CompressedBuffer.cpp b/Engine/Source/Runtime/Core/Private/Compression/CompressedBuffer.cpp
--- a/Engine/Source/Runtime/Core/Private/Compression/CompressedBuffer.cpp
+++ b/Engine/Source/Runtime/Core/Private/Compression/CompressedBuffer.cpp
@@ -3,15 +3,17 @@
 #include "Compression/CompressedBuffer.h"
 
 #include "Algo/Accumulate.h"
-#include "Algo/AnyOf.h"
 #include "Algo/ForEach.h"
 #include "Async/ParallelFor.h"
 #include "Compression/OodleDataCompression.h"
+#include "Containers/ArrayView.h"
 #include "Hash/Blake3.h"
 #include "IO/IoHash.h"
 #include "Math/UnrealMathUtility.h"
+#include "Memory/MemoryView.h"
 #include "Misc/ByteSwap.h"
 #include "Misc/Crc.h"
+#include "ProfilingDebugging/CpuProfilerTrace.h"
 #include "Serialization/Archive.h"
 
 THIRD_PARTY_INCLUDES_START
diff --git a/Engine/Source/Runtime/Core/Private/Compression/CompressionUtil.cpp b/Engine/Source/Runtime/Core/Private/Compression/CompressionUtil.cpp
--- a/Engine/Source/Runtime/Core/Private/Compression/CompressionUtil.cpp
+++ b/Engine/Source/Runtime/Core/Private/Compression/CompressionUtil.cpp
@@ -1,10 +1,15 @@
 // Copyright Epic Games, Inc. All Rights Reserved.
 
 #include "Compression/CompressionUtil.h"
+
+#include "Containers/UnrealString.h"
+#include "CoreGlobals.h"
+#include "CoreTypes.h"
+#include "Logging/LogCategory.h"
+#include "Logging/LogMacros.h"
 #include "Serialization/Archive.h"
 #include "UObject/NameTypes.h"
-#include "Logging/LogMacros.h"
-#include "CoreGlobals.h"
+#include "UObject/UnrealNames.h"
 
 namespace FCompressionUtil
 {
diff --git a/Engine/Source/Runtime/Core/Private/Compression/OodleDataCompression.cpp b/Engine/Source/Runtime/Core/Private/Compression/OodleDataCompression.cpp
--- a/Engine/Source/Runtime/Core/Private/Compression/OodleDataCompression.cpp
+++ b/Engine/Source/Runtime/Core/Private/Compression/OodleDataCompression.cpp
@@ -1,11 +1,21 @@
 // Copyright Epic Games, Inc. All Rights Reserved.
 
 #include "Compression/OodleDataCompression.h"
-#include "Misc/ConfigCacheIni.h"
-#include "Misc/ICompressionFormat.h"
-#include "HAL/PlatformMisc.h"
 
+#include "CoreGlobals.h"
+#include "HAL/CriticalSection.h"
+#include "HAL/PlatformMisc.h"
+#include "HAL/PlatformProperties.h"
+#include "HAL/UnrealMemory.h"
+#include "Logging/LogCategory.h"
+#include "Misc/AssertionMacros.h"
+#include "Misc/ConfigCacheIni.h"
+#include "ProfilingDebugging/CpuProfilerTrace.h"
+#include "Templates/UnrealTemplate.h"
 #include "oodle2.h"
+#include "oodle2base.h"
+
+struct ICompressionFormat;
 
 DEFINE_LOG_CATEGORY(OodleDataCompression);
 
diff --git a/Engine/Source/Runtime/Core/Private/Containers/Algo/FindSortedStringCaseInsensitive.cpp b/Engine/Source/Runtime/Core/Private/Containers/Algo/FindSortedStringCaseInsensitive.cpp
--- a/Engine/Source/Runtime/Core/Private/Containers/Algo/FindSortedStringCaseInsensitive.cpp
+++ b/Engine/Source/Runtime/Core/Private/Containers/Algo/FindSortedStringCaseInsensitive.cpp
@@ -1,8 +1,9 @@
 // Copyright Epic Games, Inc. All Rights Reserved.
 
 #include "Algo/FindSortedStringCaseInsensitive.h"
-#include "Misc/CString.h"
 
+#include "Misc/CString.h"
+#include "Misc/Char.h"
 
 namespace Algo
 {
diff --git a/Engine/Source/Runtime/Core/Private/Containers/BackgroundableTicker.cpp b/Engine/Source/Runtime/Core/Private/Containers/BackgroundableTicker.cpp
--- a/Engine/Source/Runtime/Core/Private/Containers/BackgroundableTicker.cpp
+++ b/Engine/Source/Runtime/Core/Private/Containers/BackgroundableTicker.cpp
@@ -1,8 +1,12 @@
 // Copyright Epic Games, Inc. All Rights Reserved.
 
 #include "Containers/BackgroundableTicker.h"
+
+#include "Containers/Ticker.h"
+#include "Delegates/DelegateInstancesImpl.h"
 #include "Misc/CoreDelegates.h"
 #include "Stats/Stats.h"
+#include "Stats/Stats2.h"
 
 FTSBackgroundableTicker& FTSBackgroundableTicker::GetCoreTicker()
 {
diff --git a/Engine/Source/Runtime/Core/Private/Containers/BitArray.cpp b/Engine/Source/Runtime/Core/Private/Containers/BitArray.cpp
--- a/Engine/Source/Runtime/Core/Private/Containers/BitArray.cpp
+++ b/Engine/Source/Runtime/Core/Private/Containers/BitArray.cpp
@@ -2,6 +2,8 @@
 
 #include "Containers/BitArray.h"
 
+#include "Math/NumericLimits.h"
+
 void FBitArrayMemory::MemmoveBitsWordOrder(uint32* StartDest, int32 DestOffset, const uint32* StartSource, int32 SourceOffset, uint32 NumBits)
 {
 	// Normalize Inputs
diff --git a/Engine/Source/Runtime/Core/Private/Containers/MemoryImageAllocator.cpp b/Engine/Source/Runtime/Core/Private/Containers/MemoryImageAllocator.cpp
--- a/Engine/Source/Runtime/Core/Private/Containers/MemoryImageAllocator.cpp
+++ b/Engine/Source/Runtime/Core/Private/Containers/MemoryImageAllocator.cpp
@@ -1,9 +1,12 @@
 // Copyright Epic Games, Inc. All Rights Reserved.
 
-#include "Containers/ContainerAllocationPolicies.h"
 #include "Containers/ResourceArray.h"
+#include "Containers/StringFwd.h"
+#include "CoreMinimal.h"
 #include "Misc/StringBuilder.h"
 #include "Serialization/MemoryImage.h"
+#include "Serialization/MemoryImageWriter.h"
+#include "Serialization/MemoryLayout.h"
 
 IMPLEMENT_ABSTRACT_TYPE_LAYOUT(FResourceArrayInterface);
 
diff --git a/Engine/Source/Runtime/Core/Private/Containers/StackTracker.cpp b/Engine/Source/Runtime/Core/Private/Containers/StackTracker.cpp
--- a/Engine/Source/Runtime/Core/Private/Containers/StackTracker.cpp
+++ b/Engine/Source/Runtime/Core/Private/Containers/StackTracker.cpp
@@ -1,9 +1,20 @@
 // Copyright Epic Games, Inc. All Rights Reserved.
 
 #include "Containers/StackTracker.h"
+
+#include "Containers/UnrealString.h"
+#include "CoreGlobals.h"
+#include "HAL/PlatformMemory.h"
 #include "HAL/PlatformStackWalk.h"
+#include "HAL/UnrealMemory.h"
+#include "Logging/LogCategory.h"
 #include "Logging/LogMacros.h"
-#include "CoreGlobals.h"
+#include "Math/UnrealMathUtility.h"
+#include "Misc/AssertionMacros.h"
+#include "Misc/CString.h"
+#include "Misc/Crc.h"
+#include "Misc/OutputDevice.h"
+#include "Templates/UnrealTemplate.h"
 
 
 DEFINE_LOG_CATEGORY_STATIC(LogStackTracker, Log, All);
diff --git a/Engine/Source/Runtime/Core/Private/Containers/String.cpp b/Engine/Source/Runtime/Core/Private/Containers/String.cpp
--- a/Engine/Source/Runtime/Core/Private/Containers/String.cpp
+++ b/Engine/Source/Runtime/Core/Private/Containers/String.cpp
@@ -1,19 +1,27 @@
 // Copyright Epic Games, Inc. All Rights Reserved.
 
+#include "Containers/Array.h"
+#include "Containers/StringConv.h"
 #include "Containers/UnrealString.h"
-
-#include "Containers/StringView.h"
 #include "CoreGlobals.h"
 #include "CoreTypes.h"
+#include "HAL/PlatformString.h"
 #include "HAL/UnrealMemory.h"
+#include "Logging/LogCategory.h"
 #include "Logging/LogMacros.h"
 #include "Math/NumericLimits.h"
 #include "Math/UnrealMathUtility.h"
 #include "Misc/AssertionMacros.h"
 #include "Misc/ByteSwap.h"
+#include "Misc/CString.h"
+#include "Misc/Char.h"
 #include "Misc/VarArgs.h"
+#include "Serialization/Archive.h"
 #include "String/HexToBytes.h"
+#include "Templates/MemoryOps.h"
+#include "Templates/RemoveReference.h"
 #include "Templates/UnrealTemplate.h"
+#include "UObject/NameTypes.h"
 
 /* FString implementation
  *****************************************************************************/
diff --git a/Engine/Source/Runtime/Core/Private/Experimental/Containers/HazardPointer.cpp b/Engine/Source/Runtime/Core/Private/Experimental/Containers/HazardPointer.cpp
--- a/Engine/Source/Runtime/Core/Private/Experimental/Containers/HazardPointer.cpp
+++ b/Engine/Source/Runtime/Core/Private/Experimental/Containers/HazardPointer.cpp
@@ -1,6 +1,8 @@
 // Copyright Epic Games, Inc. All Rights Reserved.
 
 #include "Experimental/Containers/HazardPointer.h"
+
+#include "Algo/Sort.h"
 #include "Misc/App.h"
 #include "Misc/ScopeLock.h"
 
diff --git a/Engine/Source/Runtime/Core/Private/Experimental/Coroutine/Coroutine.cpp b/Engine/Source/Runtime/Core/Private/Experimental/Coroutine/Coroutine.cpp
--- a/Engine/Source/Runtime/Core/Private/Experimental/Coroutine/Coroutine.cpp
+++ b/Engine/Source/Runtime/Core/Private/Experimental/Coroutine/Coroutine.cpp
@@ -2,6 +2,8 @@
 
 #include "Experimental/Coroutine/Coroutine.h"
 
+#include "Experimental/Coroutine/CoroLocalVariable.h"
+
 namespace CoroTask_Detail
 {
 
diff --git a/Engine/Source/Runtime/Core/Private/Experimental/Misc/ExecutionResource.cpp b/Engine/Source/Runtime/Core/Private/Experimental/Misc/ExecutionResource.cpp
--- a/Engine/Source/Runtime/Core/Private/Experimental/Misc/ExecutionResource.cpp
+++ b/Engine/Source/Runtime/Core/Private/Experimental/Misc/ExecutionResource.cpp
@@ -1,6 +1,9 @@
 // Copyright Epic Games, Inc. All Rights Reserved.
 #include "Experimental/Misc/ExecutionResource.h"
+
 #include "Containers/Array.h"
+#include "HAL/Platform.h"
+#include "Templates/UnrealTemplate.h"
 
 namespace ExecutionResourceImpl
 {
diff --git a/Engine/Source/Runtime/Core/Private/Features/ModularFeatures.cpp b/Engine/Source/Runtime/Core/Private/Features/ModularFeatures.cpp
--- a/Engine/Source/Runtime/Core/Private/Features/ModularFeatures.cpp
+++ b/Engine/Source/Runtime/Core/Private/Features/ModularFeatures.cpp
@@ -1,6 +1,9 @@
 // Copyright Epic Games, Inc. All Rights Reserved.
 
 #include "Features/ModularFeatures.h"
+
+#include "CoreTypes.h"
+#include "Misc/AssertionMacros.h"
 #include "Misc/ScopeLock.h"
 
 IModularFeatures& IModularFeatures::Get()
diff --git a/Engine/Source/Runtime/Core/Private/Features/ModularFeatures.h b/Engine/Source/Runtime/Core/Private/Features/ModularFeatures.h
--- a/Engine/Source/Runtime/Core/Private/Features/ModularFeatures.h
+++ b/Engine/Source/Runtime/Core/Private/Features/ModularFeatures.h
@@ -2,13 +2,12 @@
 
 #pragma once
 
-#include "CoreTypes.h"
 #include "Containers/Map.h"
-#include "UObject/NameTypes.h"
 #include "Delegates/Delegate.h"
+#include "Features/IModularFeature.h"
 #include "Features/IModularFeatures.h"
-
-class IModularFeature;
+#include "HAL/CriticalSection.h"
+#include "UObject/NameTypes.h"
 
 /**
  * Private implementation of modular features interface
diff --git a/Engine/Source/Runtime/Core/Private/HAL/MallocDoubleFreeFinder.cpp b/Engine/Source/Runtime/Core/Private/HAL/MallocDoubleFreeFinder.cpp
--- a/Engine/Source/Runtime/Core/Private/HAL/MallocDoubleFreeFinder.cpp
+++ b/Engine/Source/Runtime/Core/Private/HAL/MallocDoubleFreeFinder.cpp
@@ -4,14 +4,14 @@
 	FMallocDoubleFreeFinder.cpp: Memoory tracking allocator
 =============================================================================*/
 #include "HAL/MallocDoubleFreeFinder.h"
-#include "Logging/LogMacros.h"
-#include "Misc/OutputDeviceRedirector.h"
+
+#include "HAL/MemoryBase.h"
+#include "HAL/UnrealMemory.h"
 #include "Misc/Parse.h"
 #include "Misc/ScopeLock.h"
-#include "GenericPlatform/GenericPlatformProcess.h"
-#include "HAL/IConsoleManager.h"
-#include "HAL/PlatformMisc.h"
-#include "HAL/PlatformStackWalk.h"
+
+class FOutputDevice;
+class UWorld;
 
 CORE_API FMallocDoubleFreeFinder* GMallocDoubleFreeFinder;
 CORE_API bool GMallocDoubleFreeFinderEnabled = false;
diff --git a/Engine/Source/Runtime/Core/Private/HAL/MallocFrameProfiler.cpp b/Engine/Source/Runtime/Core/Private/HAL/MallocFrameProfiler.cpp
--- a/Engine/Source/Runtime/Core/Private/HAL/MallocFrameProfiler.cpp
+++ b/Engine/Source/Runtime/Core/Private/HAL/MallocFrameProfiler.cpp
@@ -4,14 +4,15 @@
 	FMallocFrameProfiler.cpp: Memoory tracking allocator
 =============================================================================*/
 #include "HAL/MallocFrameProfiler.h"
+
+#include "HAL/MemoryBase.h"
+#include "Logging/LogCategory.h"
 #include "Logging/LogMacros.h"
-#include "Misc/OutputDeviceRedirector.h"
 #include "Misc/Parse.h"
 #include "Misc/ScopeLock.h"
-#include "GenericPlatform/GenericPlatformProcess.h"
-#include "HAL/IConsoleManager.h"
-#include "HAL/PlatformMisc.h"
-#include "HAL/PlatformStackWalk.h"
+
+class FOutputDevice;
+class UWorld;
 
 DECLARE_LOG_CATEGORY_EXTERN(LogMallocFrameProfiler, Log, All);
 DEFINE_LOG_CATEGORY(LogMallocFrameProfiler);
diff --git a/Engine/Source/Runtime/Core/Private/Hash/Blake3.cpp b/Engine/Source/Runtime/Core/Private/Hash/Blake3.cpp
--- a/Engine/Source/Runtime/Core/Private/Hash/Blake3.cpp
+++ b/Engine/Source/Runtime/Core/Private/Hash/Blake3.cpp
@@ -2,9 +2,12 @@
 
 #include "Hash/Blake3.h"
 
+#include "Containers/Array.h"
+#include "Containers/ArrayView.h"
+#include "Containers/ContainersFwd.h"
 #include "Containers/UnrealString.h"
 #include "Memory/CompositeBuffer.h"
-
+#include "Memory/SharedBuffer.h"
 #include "blake3.h"
 
 ///////////////////////////////////////////////////////////////////////////////////////////////////////////////////////
diff --git a/Engine/Source/Runtime/Core/Private/Hash/xxhash.cpp b/Engine/Source/Runtime/Core/Private/Hash/xxhash.cpp
--- a/Engine/Source/Runtime/Core/Private/Hash/xxhash.cpp
+++ b/Engine/Source/Runtime/Core/Private/Hash/xxhash.cpp
@@ -1,8 +1,11 @@
 // Copyright Epic Games, Inc. All Rights Reserved.
 
 #include "Hash/xxhash.h"
+
+#include "Containers/ContainersFwd.h"
 #include "Memory/CompositeBuffer.h"
 #include "Memory/MemoryView.h"
+#include "Memory/SharedBuffer.h"
 
 THIRD_PARTY_INCLUDES_START
 #define XXH_INLINE_ALL
diff --git a/Engine/Source/Runtime/Core/Private/IO/IoBuffer.cpp b/Engine/Source/Runtime/Core/Private/IO/IoBuffer.cpp
--- a/Engine/Source/Runtime/Core/Private/IO/IoBuffer.cpp
+++ b/Engine/Source/Runtime/Core/Private/IO/IoBuffer.cpp
@@ -1,6 +1,10 @@
 // Copyright Epic Games, Inc. All Rights Reserved.
 
+#include "CoreMinimal.h"
 #include "IO/IoDispatcher.h"
+#include "Memory/MemoryFwd.h"
+#include "Memory/MemoryView.h"
+#include "Templates/RefCounting.h"
 
 //////////////////////////////////////////////////////////////////////////
 
diff --git a/Engine/Source/Runtime/Core/Private/IO/IoContainerId.cpp b/Engine/Source/Runtime/Core/Private/IO/IoContainerId.cpp
--- a/Engine/Source/Runtime/Core/Private/IO/IoContainerId.cpp
+++ b/Engine/Source/Runtime/Core/Private/IO/IoContainerId.cpp
@@ -1,7 +1,9 @@
 // Copyright Epic Games, Inc. All Rights Reserved.
 
 #include "IO/IoContainerId.h"
-#include "Serialization/StructuredArchive.h"
+
+#include "Hash/CityHash.h"
+#include "Serialization/StructuredArchiveSlots.h"
 
 FIoContainerId FIoContainerId::FromName(const FName& Name)
 {
diff --git a/Engine/Source/Runtime/Core/Private/IO/IoDirectoryIndex.cpp b/Engine/Source/Runtime/Core/Private/IO/IoDirectoryIndex.cpp
--- a/Engine/Source/Runtime/Core/Private/IO/IoDirectoryIndex.cpp
+++ b/Engine/Source/Runtime/Core/Private/IO/IoDirectoryIndex.cpp
@@ -1,10 +1,13 @@
 // Copyright Epic Games, Inc. All Rights Reserved.
 
 #include "IO/IoDirectoryIndex.h"
+
+#include "Containers/ArrayView.h"
+#include "Containers/StringView.h"
 #include "IO/IoDispatcher.h"
 #include "Misc/Paths.h"
-#include "Serialization/MemoryWriter.h"
 #include "Serialization/MemoryReader.h"
+#include "Serialization/MemoryWriter.h"
 
 namespace IoDirectoryIndexUtils
 {
diff --git a/Engine/Source/Runtime/Core/Private/IO/IoHash.cpp b/Engine/Source/Runtime/Core/Private/IO/IoHash.cpp
--- a/Engine/Source/Runtime/Core/Private/IO/IoHash.cpp
+++ b/Engine/Source/Runtime/Core/Private/IO/IoHash.cpp
@@ -2,6 +2,7 @@
 
 #include "IO/IoHash.h"
 
+#include "Containers/Array.h"
 #include "Containers/UnrealString.h"
 
 FString LexToString(const FIoHash& Hash)
diff --git a/Engine/Source/Runtime/Core/Private/IO/IoStatus.cpp b/Engine/Source/Runtime/Core/Private/IO/IoStatus.cpp
--- a/Engine/Source/Runtime/Core/Private/IO/IoStatus.cpp
+++ b/Engine/Source/Runtime/Core/Private/IO/IoStatus.cpp
@@ -1,5 +1,8 @@
 // Copyright Epic Games, Inc. All Rights Reserved.
 
+#include "Containers/StringFwd.h"
+#include "Containers/StringView.h"
+#include "CoreMinimal.h"
 #include "IO/IoDispatcher.h"
 
 //////////////////////////////////////////////////////////////////////////
diff --git a/Engine/Source/Runtime/Core/Private/IO/PackageId.cpp b/Engine/Source/Runtime/Core/Private/IO/PackageId.cpp
--- a/Engine/Source/Runtime/Core/Private/IO/PackageId.cpp
+++ b/Engine/Source/Runtime/Core/Private/IO/PackageId.cpp
@@ -1,7 +1,10 @@
 // Copyright Epic Games, Inc. All Rights Reserved.
 
 #include "IO/PackageId.h"
-#include "Serialization/StructuredArchive.h"
+
+#include "Hash/CityHash.h"
+#include "Serialization/StructuredArchiveAdapters.h"
+#include "Serialization/StructuredArchiveSlots.h"
 
 FPackageId FPackageId::FromName(const FName& Name)
 {
diff --git a/Engine/Source/Runtime/Core/Private/IO/PackageStore.cpp b/Engine/Source/Runtime/Core/Private/IO/PackageStore.cpp
--- a/Engine/Source/Runtime/Core/Private/IO/PackageStore.cpp
+++ b/Engine/Source/Runtime/Core/Private/IO/PackageStore.cpp
@@ -1,7 +1,9 @@
 // Copyright Epic Games, Inc. All Rights Reserved.
 
 #include "IO/PackageStore.h"
-#include "Serialization/StructuredArchive.h"
+
+#include "Containers/StringView.h"
+#include "IO/PackageId.h"
 #include "Serialization/CompactBinary.h"
 #include "Serialization/CompactBinaryWriter.h"
 
diff --git a/Engine/Source/Runtime/Core/Private/Internationalization/CamelCaseBreakIterator.cpp b/Engine/Source/Runtime/Core/Private/Internationalization/CamelCaseBreakIterator.cpp
--- a/Engine/Source/Runtime/Core/Private/Internationalization/CamelCaseBreakIterator.cpp
+++ b/Engine/Source/Runtime/Core/Private/Internationalization/CamelCaseBreakIterator.cpp
@@ -1,7 +1,9 @@
 // Copyright Epic Games, Inc. All Rights Reserved.
 
 #include "Internationalization/CamelCaseBreakIterator.h"
-#include "Internationalization/Text.h"
+
+#include "Misc/AssertionMacros.h"
+#include "Templates/UnrealTemplate.h"
 
 FCamelCaseBreakIterator::FCamelCaseBreakIterator()
 	: CurrentPosition(0)
diff --git a/Engine/Source/Runtime/Core/Private/Internationalization/CamelCaseBreakIterator.h b/Engine/Source/Runtime/Core/Private/Internationalization/CamelCaseBreakIterator.h
--- a/Engine/Source/Runtime/Core/Private/Internationalization/CamelCaseBreakIterator.h
+++ b/Engine/Source/Runtime/Core/Private/Internationalization/CamelCaseBreakIterator.h
@@ -2,9 +2,12 @@
 
 #pragma once
 
-#include "CoreTypes.h"
 #include "Containers/Array.h"
+#include "Containers/ContainerAllocationPolicies.h"
+#include "Containers/StringFwd.h"
+#include "Containers/StringView.h"
 #include "Containers/UnrealString.h"
+#include "CoreTypes.h"
 #include "Internationalization/IBreakIterator.h"
 
 /**
diff --git a/Engine/Source/Runtime/Core/Private/Internationalization/FastDecimalFormat.cpp b/Engine/Source/Runtime/Core/Private/Internationalization/FastDecimalFormat.cpp
--- a/Engine/Source/Runtime/Core/Private/Internationalization/FastDecimalFormat.cpp
+++ b/Engine/Source/Runtime/Core/Private/Internationalization/FastDecimalFormat.cpp
@@ -1,9 +1,15 @@
 // Copyright Epic Games, Inc. All Rights Reserved.
 
 #include "Internationalization/FastDecimalFormat.h"
+
+#include "Containers/StringFwd.h"
+#include "HAL/IConsoleManager.h"
+#include "HAL/UnrealMemory.h"
+#include "Math/UnrealMathUtility.h"
+#include "Misc/AssertionMacros.h"
 #include "Misc/EnumClassFlags.h"
 #include "Misc/StringBuilder.h"
-#include "HAL/IConsoleManager.h"
+#include "Templates/UnrealTemplate.h"
 
 namespace FastDecimalFormat
 {
diff --git a/Engine/Source/Runtime/Core/Private/Internationalization/GatherableTextData.cpp b/Engine/Source/Runtime/Core/Private/Internationalization/GatherableTextData.cpp
--- a/Engine/Source/Runtime/Core/Private/Internationalization/GatherableTextData.cpp
+++ b/Engine/Source/Runtime/Core/Private/Internationalization/GatherableTextData.cpp
@@ -1,7 +1,14 @@
 // Copyright Epic Games, Inc. All Rights Reserved.
 
 #include "Internationalization/GatherableTextData.h"
+
+#include "CoreTypes.h"
 #include "Serialization/StructuredArchive.h"
+#include "Serialization/StructuredArchiveAdapters.h"
+#include "Serialization/StructuredArchiveNameHelpers.h"
+#include "Serialization/StructuredArchiveSlots.h"
+
+class FArchive;
 
 FArchive& operator<<(FArchive& Archive, FTextSourceSiteContext& This)
 {
diff --git a/Engine/Source/Runtime/Core/Private/Internationalization/InternationalizationArchive.cpp b/Engine/Source/Runtime/Core/Private/Internationalization/InternationalizationArchive.cpp
--- a/Engine/Source/Runtime/Core/Private/Internationalization/InternationalizationArchive.cpp
+++ b/Engine/Source/Runtime/Core/Private/Internationalization/InternationalizationArchive.cpp
@@ -1,7 +1,10 @@
 // Copyright Epic Games, Inc. All Rights Reserved.
 
 #include "Internationalization/InternationalizationArchive.h"
+
+#include "Containers/Array.h"
 #include "Internationalization/InternationalizationMetadata.h"
+#include "Logging/LogMacros.h"
 
 DEFINE_LOG_CATEGORY_STATIC(LogGatherArchiveObject, Log, All);
 
diff --git a/Engine/Source/Runtime/Core/Private/Internationalization/InternationalizationManifest.cpp b/Engine/Source/Runtime/Core/Private/Internationalization/InternationalizationManifest.cpp
--- a/Engine/Source/Runtime/Core/Private/Internationalization/InternationalizationManifest.cpp
+++ b/Engine/Source/Runtime/Core/Private/Internationalization/InternationalizationManifest.cpp
@@ -1,7 +1,10 @@
 // Copyright Epic Games, Inc. All Rights Reserved.
 
 #include "Internationalization/InternationalizationManifest.h"
+
 #include "Internationalization/InternationalizationMetadata.h"
+#include "Logging/LogMacros.h"
+#include "Misc/CString.h"
 
 DEFINE_LOG_CATEGORY_STATIC(LogInternationalizationManifestObject, Log, All);
 
diff --git a/Engine/Source/Runtime/Core/Private/Internationalization/InternationalizationMetadata.cpp b/Engine/Source/Runtime/Core/Private/Internationalization/InternationalizationMetadata.cpp
--- a/Engine/Source/Runtime/Core/Private/Internationalization/InternationalizationMetadata.cpp
+++ b/Engine/Source/Runtime/Core/Private/Internationalization/InternationalizationMetadata.cpp
@@ -1,7 +1,18 @@
 // Copyright Epic Games, Inc. All Rights Reserved.
 
 #include "Internationalization/InternationalizationMetadata.h"
+
+#include "HAL/PlatformCrt.h"
+#include "Logging/LogCategory.h"
+#include "Logging/LogMacros.h"
+#include "Serialization/Archive.h"
 #include "Serialization/StructuredArchive.h"
+#include "Serialization/StructuredArchiveAdapters.h"
+#include "Serialization/StructuredArchiveNameHelpers.h"
+#include "Serialization/StructuredArchiveSlots.h"
+#include "Templates/ChooseClass.h"
+#include "Templates/Less.h"
+#include "Templates/Tuple.h"
 
 DEFINE_LOG_CATEGORY_STATIC(LogInternationalizationMetadata, Log, All);
 
diff --git a/Engine/Source/Runtime/Core/Private/Internationalization/LocalizationResourceTextSource.cpp b/Engine/Source/Runtime/Core/Private/Internationalization/LocalizationResourceTextSource.cpp
--- a/Engine/Source/Runtime/Core/Private/Internationalization/LocalizationResourceTextSource.cpp
+++ b/Engine/Source/Runtime/Core/Private/Internationalization/LocalizationResourceTextSource.cpp
@@ -1,12 +1,14 @@
 // Copyright Epic Games, Inc. All Rights Reserved.
 
 #include "Internationalization/LocalizationResourceTextSource.h"
-#include "Internationalization/TextLocalizationResource.h"
-#include "HAL/PlatformProperties.h"
+
 #include "HAL/FileManager.h"
-#include "Misc/Paths.h"
-#include "Misc/CoreDelegates.h"
+#include "HAL/PlatformProperties.h"
+#include "Internationalization/TextLocalizationResource.h"
 #include "Misc/ConfigCacheIni.h"
+#include "Misc/CoreDelegates.h"
+#include "Misc/EnumClassFlags.h"
+#include "Misc/Paths.h"
 
 bool FLocalizationResourceTextSource::GetNativeCultureName(const ELocalizedTextSourceCategory InCategory, FString& OutNativeCultureName)
 {
diff --git a/Engine/Source/Runtime/Core/Private/Internationalization/LocalizationResourceTextSource.h b/Engine/Source/Runtime/Core/Private/Internationalization/LocalizationResourceTextSource.h
--- a/Engine/Source/Runtime/Core/Private/Internationalization/LocalizationResourceTextSource.h
+++ b/Engine/Source/Runtime/Core/Private/Internationalization/LocalizationResourceTextSource.h
@@ -2,8 +2,13 @@
 
 #pragma once
 
+#include "Containers/ArrayView.h"
+#include "CoreMinimal.h"
 #include "CoreTypes.h"
 #include "Internationalization/ILocalizedTextSource.h"
+#include "Internationalization/LocalizedTextSourceTypes.h"
+
+class FTextLocalizationResource;
 
 /**
  * Implementation of a localized text source that loads data from Localization Resource (LocRes) files.
diff --git a/Engine/Source/Runtime/Core/Private/Internationalization/PolyglotTextData.cpp b/Engine/Source/Runtime/Core/Private/Internationalization/PolyglotTextData.cpp
--- a/Engine/Source/Runtime/Core/Private/Internationalization/PolyglotTextData.cpp
+++ b/Engine/Source/Runtime/Core/Private/Internationalization/PolyglotTextData.cpp
@@ -1,9 +1,12 @@
 // Copyright Epic Games, Inc. All Rights Reserved.
 
 #include "Internationalization/PolyglotTextData.h"
+
+#include "CoreTypes.h"
+#include "Internationalization/Internationalization.h"
 #include "Internationalization/TextLocalizationManager.h"
 #include "Internationalization/TextLocalizationResource.h"
-#include "Internationalization/Internationalization.h"
+#include "Misc/AssertionMacros.h"
 
 #define LOCTEXT_NAMESPACE "PolyglotTextData"
 
diff --git a/Engine/Source/Runtime/Core/Private/Internationalization/PolyglotTextSource.cpp b/Engine/Source/Runtime/Core/Private/Internationalization/PolyglotTextSource.cpp
--- a/Engine/Source/Runtime/Core/Private/Internationalization/PolyglotTextSource.cpp
+++ b/Engine/Source/Runtime/Core/Private/Internationalization/PolyglotTextSource.cpp
@@ -1,7 +1,16 @@
 // Copyright Epic Games, Inc. All Rights Reserved.
 
 #include "Internationalization/PolyglotTextSource.h"
+
+#include "Containers/Array.h"
+#include "Containers/UnrealString.h"
+#include "Internationalization/PolyglotTextData.h"
+#include "Internationalization/TextKey.h"
 #include "Internationalization/TextLocalizationResource.h"
+#include "Misc/AssertionMacros.h"
+#include "Misc/EnumClassFlags.h"
+#include "Templates/Tuple.h"
+#include "Templates/UnrealTemplate.h"
 
 bool FPolyglotTextSource::GetNativeCultureName(const ELocalizedTextSourceCategory InCategory, FString& OutNativeCultureName)
 {
diff --git a/Engine/Source/Runtime/Core/Private/Internationalization/StringTableCore.cpp b/Engine/Source/Runtime/Core/Private/Internationalization/StringTableCore.cpp
--- a/Engine/Source/Runtime/Core/Private/Internationalization/StringTableCore.cpp
+++ b/Engine/Source/Runtime/Core/Private/Internationalization/StringTableCore.cpp
@@ -1,10 +1,25 @@
 // Copyright Epic Games, Inc. All Rights Reserved.
 
 #include "Internationalization/StringTableCore.h"
-#include "Misc/ScopeLock.h"
-#include "Misc/FileHelper.h"
+
+#include "Containers/Array.h"
+#include "Containers/ContainerAllocationPolicies.h"
+#include "Containers/Set.h"
+#include "CoreGlobals.h"
+#include "CoreTypes.h"
+#include "Internationalization/LocKeyFuncs.h"
+#include "Internationalization/TextLocalizationManager.h"
+#include "Logging/LogCategory.h"
+#include "Misc/CString.h"
 #include "Misc/ConfigCacheIni.h"
+#include "Misc/FileHelper.h"
+#include "Misc/Parse.h"
+#include "Misc/ScopeLock.h"
 #include "Serialization/Csv/CsvParser.h"
+#include "Templates/ChooseClass.h"
+#include "Templates/Tuple.h"
+
+class UStringTable;
 
 DEFINE_LOG_CATEGORY(LogStringTable);
 
diff --git a/Engine/Source/Runtime/Core/Private/Internationalization/TextCache.cpp b/Engine/Source/Runtime/Core/Private/Internationalization/TextCache.cpp
--- a/Engine/Source/Runtime/Core/Private/Internationalization/TextCache.cpp
+++ b/Engine/Source/Runtime/Core/Private/Internationalization/TextCache.cpp
@@ -1,9 +1,14 @@
 // Copyright Epic Games, Inc. All Rights Reserved.
 
 #include "Internationalization/TextCache.h"
+
+#include "Containers/UnrealString.h"
+#include "CoreTypes.h"
+#include "HAL/LowLevelMemTracker.h"
+#include "Internationalization/TextKey.h"
+#include "Misc/CString.h"
 #include "Misc/LazySingleton.h"
 #include "Misc/ScopeLock.h"
-#include "HAL/LowLevelMemTracker.h"
 
 FTextCache& FTextCache::Get()
 {
diff --git a/Engine/Source/Runtime/Core/Private/Internationalization/TextCache.h b/Engine/Source/Runtime/Core/Private/Internationalization/TextCache.h
--- a/Engine/Source/Runtime/Core/Private/Internationalization/TextCache.h
+++ b/Engine/Source/Runtime/Core/Private/Internationalization/TextCache.h
@@ -2,10 +2,11 @@
 
 #pragma once
 
-#include "CoreTypes.h"
 #include "Containers/Map.h"
+#include "HAL/CriticalSection.h"
 #include "Internationalization/Text.h"
-#include "Internationalization/TextKey.h"
+
+class FTextId;
 
 /** Caches FText instances generated via the LOCTEXT macro to avoid repeated constructions */
 class FTextCache
diff --git a/Engine/Source/Runtime/Core/Private/Internationalization/TextFormatArgumentModifier.cpp b/Engine/Source/Runtime/Core/Private/Internationalization/TextFormatArgumentModifier.cpp
--- a/Engine/Source/Runtime/Core/Private/Internationalization/TextFormatArgumentModifier.cpp
+++ b/Engine/Source/Runtime/Core/Private/Internationalization/TextFormatArgumentModifier.cpp
@@ -1,10 +1,18 @@
 // Copyright Epic Games, Inc. All Rights Reserved.
 
 #include "Internationalization/TextFormatArgumentModifier.h"
-#include "Misc/Parse.h"
+
+#include "Containers/UnrealString.h"
 #include "Internationalization/Culture.h"
 #include "Internationalization/Internationalization.h"
 #include "Internationalization/TextFormatter.h"
+#include "Math/UnrealMathUtility.h"
+#include "Misc/CString.h"
+#include "Misc/Char.h"
+#include "Misc/Parse.h"
+#include "Templates/Less.h"
+#include "Templates/Tuple.h"
+#include "Templates/UnrealTemplate.h"
 
 bool ITextFormatArgumentModifier::ParseKeyValueArgs(const FTextFormatString& InArgsString, TSortedMap<FTextFormatString, FTextFormatString>& OutArgKeyValues, const TCHAR InValueSeparator, const TCHAR InArgSeparator)
 {
diff --git a/Engine/Source/Runtime/Core/Private/Internationalization/TextFormatArgumentModifier.h b/Engine/Source/Runtime/Core/Private/Internationalization/TextFormatArgumentModifier.h
--- a/Engine/Source/Runtime/Core/Private/Internationalization/TextFormatArgumentModifier.h
+++ b/Engine/Source/Runtime/Core/Private/Internationalization/TextFormatArgumentModifier.h
@@ -2,11 +2,15 @@
 
 #pragma once
 
+#include "Containers/Array.h"
+#include "Containers/SortedMap.h"
 #include "CoreTypes.h"
-#include "Containers/UnrealString.h"
-#include "Internationalization/Text.h"
+#include "Internationalization/CulturePointer.h"
 #include "Internationalization/ITextFormatArgumentModifier.h"
+#include "Internationalization/Text.h"
+#include "Templates/SharedPointer.h"
 
+class FString;
 struct FPrivateTextFormatArguments;
 
 /**
diff --git a/Engine/Source/Runtime/Core/Private/Internationalization/TextFormatter.cpp b/Engine/Source/Runtime/Core/Private/Internationalization/TextFormatter.cpp
--- a/Engine/Source/Runtime/Core/Private/Internationalization/TextFormatter.cpp
+++ b/Engine/Source/Runtime/Core/Private/Internationalization/TextFormatter.cpp
@@ -1,13 +1,32 @@
 // Copyright Epic Games, Inc. All Rights Reserved.
 
 #include "Internationalization/TextFormatter.h"
-#include "Misc/ScopeLock.h"
-#include "Misc/ScopeRWLock.h"
+
+#include "CoreGlobals.h"
+#include "CoreTypes.h"
+#include "HAL/PlatformProcess.h"
+#include "Internationalization/CulturePointer.h"
+#include "Internationalization/Internationalization.h"
 #include "Internationalization/TextFormatArgumentModifier.h"
 #include "Internationalization/TextHistory.h"
+#include "Logging/LogCategory.h"
+#include "Logging/LogMacros.h"
+#include "Misc/AssertionMacros.h"
+#include "Misc/CString.h"
+#include "Misc/Char.h"
+#include "Misc/EnumClassFlags.h"
 #include "Misc/ExpressionParser.h"
+#include "Misc/ExpressionParserTypes.inl"
+#include "Misc/Guid.h"
+#include "Misc/Optional.h"
+#include "Misc/ScopeLock.h"
+#include "Misc/ScopeRWLock.h"
 #include "Stats/Stats.h"
-#include "HAL/PlatformProcess.h"
+#include "Stats/Stats2.h"
+#include "Templates/Less.h"
+#include "Templates/Tuple.h"
+#include "Templates/UnrealTemplate.h"
+#include "Templates/ValueOrError.h"
 
 #define LOCTEXT_NAMESPACE "TextFormatter"
 
diff --git a/Engine/Source/Runtime/Core/Private/Internationalization/TextGeneratorBase.cpp b/Engine/Source/Runtime/Core/Private/Internationalization/TextGeneratorBase.cpp
--- a/Engine/Source/Runtime/Core/Private/Internationalization/TextGeneratorBase.cpp
+++ b/Engine/Source/Runtime/Core/Private/Internationalization/TextGeneratorBase.cpp
@@ -1,7 +1,11 @@
 // Copyright Epic Games, Inc. All Rights Reserved.
 
 #include "Internationalization/TextGeneratorBase.h"
+
+#include "CoreTypes.h"
+#include "Misc/AssertionMacros.h"
 #include "UObject/NameTypes.h"
+#include "UObject/UnrealNames.h"
 
 FName FTextGeneratorBase::GetTypeID() const
 {
diff --git a/Engine/Source/Runtime/Core/Private/Internationalization/TextGeneratorRegistry.cpp b/Engine/Source/Runtime/Core/Private/Internationalization/TextGeneratorRegistry.cpp
--- a/Engine/Source/Runtime/Core/Private/Internationalization/TextGeneratorRegistry.cpp
+++ b/Engine/Source/Runtime/Core/Private/Internationalization/TextGeneratorRegistry.cpp
@@ -1,7 +1,10 @@
 // Copyright Epic Games, Inc. All Rights Reserved.
 
 #include "Internationalization/TextGeneratorRegistry.h"
+
+#include "Misc/AssertionMacros.h"
 #include "Misc/ScopeLock.h"
+#include "Templates/UnrealTemplate.h"
 
 FTextGeneratorRegistry& FTextGeneratorRegistry::Get()
 {
diff --git a/Engine/Source/Runtime/Core/Private/Internationalization/TextGeneratorRegistry.h b/Engine/Source/Runtime/Core/Private/Internationalization/TextGeneratorRegistry.h
--- a/Engine/Source/Runtime/Core/Private/Internationalization/TextGeneratorRegistry.h
+++ b/Engine/Source/Runtime/Core/Private/Internationalization/TextGeneratorRegistry.h
@@ -1,8 +1,10 @@
 // Copyright Epic Games, Inc. All Rights Reserved.
 #pragma once
 
-#include "CoreTypes.h"
+#include "Containers/Map.h"
+#include "HAL/CriticalSection.h"
 #include "Internationalization/Text.h"
+#include "UObject/NameTypes.h"
 
 /** Singleton registry of text generator factory functions */
 class FTextGeneratorRegistry
diff --git a/Engine/Source/Runtime/Core/Private/Logging/LogCategory.cpp b/Engine/Source/Runtime/Core/Private/Logging/LogCategory.cpp
--- a/Engine/Source/Runtime/Core/Private/Logging/LogCategory.cpp
+++ b/Engine/Source/Runtime/Core/Private/Logging/LogCategory.cpp
@@ -1,8 +1,11 @@
 // Copyright Epic Games, Inc. All Rights Reserved.
 
 #include "Logging/LogCategory.h"
+
 #include "CoreGlobals.h"
+#include "CoreMinimal.h"
 #include "Logging/LogSuppressionInterface.h"
+#include "Logging/LogTrace.h"
 #include "Misc/CoreDelegates.h"
 #include "Misc/OutputDeviceRedirector.h"
 
diff --git a/Engine/Source/Runtime/Core/Private/Logging/LogScopedCategoryAndVerbosityOverride.cpp b/Engine/Source/Runtime/Core/Private/Logging/LogScopedCategoryAndVerbosityOverride.cpp
--- a/Engine/Source/Runtime/Core/Private/Logging/LogScopedCategoryAndVerbosityOverride.cpp
+++ b/Engine/Source/Runtime/Core/Private/Logging/LogScopedCategoryAndVerbosityOverride.cpp
@@ -1,6 +1,8 @@
 // Copyright Epic Games, Inc. All Rights Reserved.
 
 #include "Logging/LogScopedCategoryAndVerbosityOverride.h"
+
+#include "CoreTypes.h"
 #include "HAL/PlatformTLS.h"
 
 FLogScopedCategoryAndVerbosityOverride::FLogScopedCategoryAndVerbosityOverride(FName Category, ELogVerbosity::Type Verbosity)
diff --git a/Engine/Source/Runtime/Core/Private/Logging/TokenizedMessage.cpp b/Engine/Source/Runtime/Core/Private/Logging/TokenizedMessage.cpp
--- a/Engine/Source/Runtime/Core/Private/Logging/TokenizedMessage.cpp
+++ b/Engine/Source/Runtime/Core/Private/Logging/TokenizedMessage.cpp
@@ -1,6 +1,10 @@
 // Copyright Epic Games, Inc. All Rights Reserved.
 
 #include "Logging/TokenizedMessage.h"
+
+#include "CoreTypes.h"
+#include "Delegates/DelegateInstancesImpl.h"
+#include "HAL/PlatformCrt.h"
 #include "HAL/PlatformProcess.h"
 #include "Internationalization/Internationalization.h"
 
diff --git a/Engine/Source/Runtime/Core/Private/Math/ColorList.cpp b/Engine/Source/Runtime/Core/Private/Math/ColorList.cpp
--- a/Engine/Source/Runtime/Core/Private/Math/ColorList.cpp
+++ b/Engine/Source/Runtime/Core/Private/Math/ColorList.cpp
@@ -5,6 +5,9 @@
 =============================================================================*/
 
 #include "Math/ColorList.h"
+
+#include "Containers/UnrealString.h"
+#include "Logging/LogCategory.h"
 #include "Logging/LogMacros.h"
 
 DEFINE_LOG_CATEGORY_STATIC(LogColorList, Log, All);
diff --git a/Engine/Source/Runtime/Core/Private/Math/Sobol.cpp b/Engine/Source/Runtime/Core/Private/Math/Sobol.cpp
--- a/Engine/Source/Runtime/Core/Private/Math/Sobol.cpp
+++ b/Engine/Source/Runtime/Core/Private/Math/Sobol.cpp
@@ -3,6 +3,12 @@
 
 #include "Math/Sobol.h"
 
+#include "Math/IntPoint.h"
+#include "Math/IntVector.h"
+#include "Math/UnrealMathUtility.h"
+#include "Math/Vector.h"
+#include "Misc/AssertionMacros.h"
+
 float FSobol::Evaluate(int32 Index, int32 Dim, int32 Seed)
 {
 	check(Dim >= 0 && Dim <= MaxDimension);
diff --git a/Engine/Source/Runtime/Core/Private/Math/Sphere.cpp b/Engine/Source/Runtime/Core/Private/Math/Sphere.cpp
--- a/Engine/Source/Runtime/Core/Private/Math/Sphere.cpp
+++ b/Engine/Source/Runtime/Core/Private/Math/Sphere.cpp
@@ -5,8 +5,6 @@
 =============================================================================*/
 
 #include "Math/Sphere.h"
-#include "Math/Box.h"
-#include "Math/Transform.h"
 
 namespace UE {
 namespace Math {
diff --git a/Engine/Source/Runtime/Core/Private/Memory/CompositeBuffer.cpp b/Engine/Source/Runtime/Core/Private/Memory/CompositeBuffer.cpp
--- a/Engine/Source/Runtime/Core/Private/Memory/CompositeBuffer.cpp
+++ b/Engine/Source/Runtime/Core/Private/Memory/CompositeBuffer.cpp
@@ -4,7 +4,12 @@
 
 #include "Algo/Accumulate.h"
 #include "Algo/AllOf.h"
+#include "HAL/PlatformString.h"
+#include "Math/UnrealMathUtility.h"
+#include "Memory/MemoryView.h"
+#include "Misc/AssertionMacros.h"
 #include "Templates/Function.h"
+#include "Templates/RemoveReference.h"
 
 void FCompositeBuffer::Reset()
 {
diff --git a/Engine/Source/Runtime/Core/Private/Memory/SharedBuffer.cpp b/Engine/Source/Runtime/Core/Private/Memory/SharedBuffer.cpp
--- a/Engine/Source/Runtime/Core/Private/Memory/SharedBuffer.cpp
+++ b/Engine/Source/Runtime/Core/Private/Memory/SharedBuffer.cpp
@@ -3,6 +3,7 @@
 #include "Memory/SharedBuffer.h"
 
 #include "HAL/UnrealMemory.h"
+#include "Templates/RemoveReference.h"
 
 ///////////////////////////////////////////////////////////////////////////////////////////////////////////////////////
 
diff --git a/Engine/Source/Runtime/Core/Private/Memory/VirtualStackAllocator.cpp b/Engine/Source/Runtime/Core/Private/Memory/VirtualStackAllocator.cpp
--- a/Engine/Source/Runtime/Core/Private/Memory/VirtualStackAllocator.cpp
+++ b/Engine/Source/Runtime/Core/Private/Memory/VirtualStackAllocator.cpp
@@ -1,8 +1,9 @@
 // Copyright Epic Games, Inc. All Rights Reserved.
 
 #include "Memory/VirtualStackAllocator.h"
-#include "Templates/AlignmentTemplates.h"
+
 #include "HAL/LowLevelMemTracker.h"
+#include "HAL/PlatformMemory.h"
 
 LLM_DEFINE_TAG(VirtualStackAllocator);
 
diff --git a/Engine/Source/Runtime/Core/Private/Misc/AsyncTaskNotification.cpp b/Engine/Source/Runtime/Core/Private/Misc/AsyncTaskNotification.cpp
--- a/Engine/Source/Runtime/Core/Private/Misc/AsyncTaskNotification.cpp
+++ b/Engine/Source/Runtime/Core/Private/Misc/AsyncTaskNotification.cpp
@@ -1,6 +1,8 @@
 // Copyright Epic Games, Inc. All Rights Reserved.
 
 #include "Misc/AsyncTaskNotification.h"
+
+#include "Misc/AssertionMacros.h"
 #include "Misc/CoreAsyncTaskNotificationImpl.h"
 
 FAsyncTaskNotification::FAsyncTaskNotification(const FAsyncTaskNotificationConfig& InConfig)
diff --git a/Engine/Source/Runtime/Core/Private/Misc/Base64.cpp b/Engine/Source/Runtime/Core/Private/Misc/Base64.cpp
--- a/Engine/Source/Runtime/Core/Private/Misc/Base64.cpp
+++ b/Engine/Source/Runtime/Core/Private/Misc/Base64.cpp
@@ -1,7 +1,9 @@
 // Copyright Epic Games, Inc. All Rights Reserved.
 
 #include "Misc/Base64.h"
+
 #include "Containers/StringConv.h"
+#include "Misc/AssertionMacros.h"
 
 /** The table used to encode a 6 bit value as an ascii character */
 static const uint8 EncodingAlphabet[64] = 
diff --git a/Engine/Source/Runtime/Core/Private/Misc/BufferedOutputDevice.cpp b/Engine/Source/Runtime/Core/Private/Misc/BufferedOutputDevice.cpp
--- a/Engine/Source/Runtime/Core/Private/Misc/BufferedOutputDevice.cpp
+++ b/Engine/Source/Runtime/Core/Private/Misc/BufferedOutputDevice.cpp
@@ -1,7 +1,10 @@
 // Copyright Epic Games, Inc. All Rights Reserved.
 
 #include "Misc/BufferedOutputDevice.h"
+
+#include "CoreTypes.h"
 #include "Misc/ScopeLock.h"
+#include "Templates/UnrealTemplate.h"
 
 void FBufferedOutputDevice::Serialize(const TCHAR* InData, ELogVerbosity::Type Verbosity, const class FName& Category)
 {
diff --git a/Engine/Source/Runtime/Core/Private/Misc/Change.cpp b/Engine/Source/Runtime/Core/Private/Misc/Change.cpp
--- a/Engine/Source/Runtime/Core/Private/Misc/Change.cpp
+++ b/Engine/Source/Runtime/Core/Private/Misc/Change.cpp
@@ -1,8 +1,12 @@
 // Copyright Epic Games, Inc. All Rights Reserved.
 
 #include "Misc/Change.h"
+
 #include "Containers/UnrealString.h"
 #include "Misc/FeedbackContext.h"
+#include "Templates/EnableIf.h"
+
+class UObject;
 
 void FChange::PrintToLog( FFeedbackContext& FeedbackContext, const int32 IndentLevel )
 {
diff --git a/Engine/Source/Runtime/Core/Private/Misc/ComparisonUtility.cpp b/Engine/Source/Runtime/Core/Private/Misc/ComparisonUtility.cpp
--- a/Engine/Source/Runtime/Core/Private/Misc/ComparisonUtility.cpp
+++ b/Engine/Source/Runtime/Core/Private/Misc/ComparisonUtility.cpp
@@ -2,6 +2,10 @@
 
 #include "Misc/ComparisonUtility.h"
 
+#include "Containers/StringView.h"
+#include "CoreTypes.h"
+#include "Misc/CString.h"
+
 namespace UE { namespace ComparisonUtility {
 
 int32 CompareWithNumericSuffix(FName A, FName B)
diff --git a/Engine/Source/Runtime/Core/Private/Misc/ConfigManifest.cpp b/Engine/Source/Runtime/Core/Private/Misc/ConfigManifest.cpp
--- a/Engine/Source/Runtime/Core/Private/Misc/ConfigManifest.cpp
+++ b/Engine/Source/Runtime/Core/Private/Misc/ConfigManifest.cpp
@@ -1,19 +1,23 @@
 // Copyright Epic Games, Inc. All Rights Reserved.
 
 #include "Misc/ConfigManifest.h"
-#include "Misc/EngineVersionBase.h"
-#include "Containers/UnrealString.h"
-#include "GenericPlatform/GenericPlatformFile.h"
-#include "HAL/PlatformFileManager.h"
+
+#include "Containers/Map.h"
 #include "Containers/StringConv.h"
+#include "Containers/UnrealString.h"
 #include "HAL/FileManager.h"
-#include "Misc/Paths.h"
-#include "Misc/ConfigCacheIni.h"
+#include "HAL/PlatformFile.h"
+#include "HAL/PlatformFileManager.h"
+#include "HAL/PlatformProcess.h"
+#include "HAL/PlatformProperties.h"
 #include "Misc/App.h"
-
+#include "Misc/ConfigCacheIni.h"
 #include "Misc/EngineVersion.h"
-
+#include "Misc/EngineVersionBase.h"
+#include "Misc/Paths.h"
 #include "Runtime/Launch/Resources/Version.h"
+#include "Templates/Tuple.h"
+#include "Templates/UnrealTemplate.h"
 
 enum class EConfigManifestVersion
 {
diff --git a/Engine/Source/Runtime/Core/Private/Misc/ConfigManifest.h b/Engine/Source/Runtime/Core/Private/Misc/ConfigManifest.h
--- a/Engine/Source/Runtime/Core/Private/Misc/ConfigManifest.h
+++ b/Engine/Source/Runtime/Core/Private/Misc/ConfigManifest.h
@@ -3,7 +3,8 @@
 #pragma once
 
 #include "CoreTypes.h"
-#include "CoreFwd.h"
+
+class FConfigFile;
 
 enum class EConfigManifestVersion;
 
diff --git a/Engine/Source/Runtime/Core/Private/Misc/DefaultValueHelper.cpp b/Engine/Source/Runtime/Core/Private/Misc/DefaultValueHelper.cpp
--- a/Engine/Source/Runtime/Core/Private/Misc/DefaultValueHelper.cpp
+++ b/Engine/Source/Runtime/Core/Private/Misc/DefaultValueHelper.cpp
@@ -1,12 +1,17 @@
 // Copyright Epic Games, Inc. All Rights Reserved.
 
 #include "Misc/DefaultValueHelper.h"
+
+#include "Containers/Array.h"
 #include "Math/Color.h"
-#include "Math/Vector2D.h"
+#include "Math/Rotator.h"
 #include "Math/Vector.h"
+#include "Math/Vector2D.h"
 #include "Math/Vector4.h"
-#include "Math/Rotator.h"
-
+#include "Misc/AssertionMacros.h"
+#include "Misc/CString.h"
+#include "Misc/Char.h"
+#include "Misc/Parse.h"
 
 bool FDefaultValueHelper::Is(const FString& Source, const TCHAR* CompareStr)
 {
diff --git a/Engine/Source/Runtime/Core/Private/Misc/DelayedAutoRegister.cpp b/Engine/Source/Runtime/Core/Private/Misc/DelayedAutoRegister.cpp
--- a/Engine/Source/Runtime/Core/Private/Misc/DelayedAutoRegister.cpp
+++ b/Engine/Source/Runtime/Core/Private/Misc/DelayedAutoRegister.cpp
@@ -1,7 +1,13 @@
 // Copyright Epic Games, Inc. All Rights Reserved.
 
 #include "Misc/DelayedAutoRegister.h"
+
+#include "Containers/Set.h"
 #include "Delegates/Delegate.h"
+#include "Delegates/DelegateBase.h"
+#include "Delegates/DelegateInstancesImpl.h"
+#include "Delegates/IDelegateInstance.h"
+#include "Misc/AssertionMacros.h"
 
 DECLARE_MULTICAST_DELEGATE(FDelayedAutoRegisterDelegate);
 
diff --git a/Engine/Source/Runtime/Core/Private/Misc/EngineBuildSettings.cpp b/Engine/Source/Runtime/Core/Private/Misc/EngineBuildSettings.cpp
--- a/Engine/Source/Runtime/Core/Private/Misc/EngineBuildSettings.cpp
+++ b/Engine/Source/Runtime/Core/Private/Misc/EngineBuildSettings.cpp
@@ -1,6 +1,8 @@
 // Copyright Epic Games, Inc. All Rights Reserved.
 
 #include "Misc/EngineBuildSettings.h"
+
+#include "CoreTypes.h"
 #include "Misc/Paths.h"
 
 
diff --git a/Engine/Source/Runtime/Core/Private/Misc/EngineVersion.cpp b/Engine/Source/Runtime/Core/Private/Misc/EngineVersion.cpp
--- a/Engine/Source/Runtime/Core/Private/Misc/EngineVersion.cpp
+++ b/Engine/Source/Runtime/Core/Private/Misc/EngineVersion.cpp
@@ -1,13 +1,19 @@
 // Copyright Epic Games, Inc. All Rights Reserved.
 
 #include "Misc/EngineVersion.h"
+
+#include "BuildSettings.h"
+#include "Math/NumericLimits.h"
+#include "Misc/CString.h"
 #include "Misc/Guid.h"
 #include "Misc/LazySingleton.h"
-#include "Serialization/CustomVersion.h"
 #include "Runtime/Launch/Resources/Version.h"
+#include "Serialization/Archive.h"
+#include "Serialization/CustomVersion.h"
+#include "Serialization/StructuredArchiveAdapters.h"
+#include "Serialization/StructuredArchiveNameHelpers.h"
+#include "Serialization/StructuredArchiveSlots.h"
 #include "UObject/ReleaseObjectVersion.h"
-#include "BuildSettings.h"
-#include "CoreGlobals.h"
 
 FEngineVersionBase::FEngineVersionBase(uint16 InMajor, uint16 InMinor, uint16 InPatch, uint32 InChangelist)
 : Major(InMajor)
diff --git a/Engine/Source/Runtime/Core/Private/Misc/FeedbackContext.cpp b/Engine/Source/Runtime/Core/Private/Misc/FeedbackContext.cpp
--- a/Engine/Source/Runtime/Core/Private/Misc/FeedbackContext.cpp
+++ b/Engine/Source/Runtime/Core/Private/Misc/FeedbackContext.cpp
@@ -1,7 +1,11 @@
 // Copyright Epic Games, Inc. All Rights Reserved.
 
 #include "Misc/FeedbackContext.h"
+
+#include "CoreTypes.h"
 #include "HAL/PlatformTime.h"
+#include "Misc/AssertionMacros.h"
+#include "Misc/SlowTask.h"
 
 FFeedbackContext::FFeedbackContext()
 	: TreatWarningsAsErrors(0)
diff --git a/Engine/Source/Runtime/Core/Private/Misc/FeedbackContextMarkup.cpp b/Engine/Source/Runtime/Core/Private/Misc/FeedbackContextMarkup.cpp
--- a/Engine/Source/Runtime/Core/Private/Misc/FeedbackContextMarkup.cpp
+++ b/Engine/Source/Runtime/Core/Private/Misc/FeedbackContextMarkup.cpp
@@ -1,10 +1,15 @@
 // Copyright Epic Games, Inc. All Rights Reserved.
 
 #include "Misc/FeedbackContextMarkup.h"
-#include "Misc/AssertionMacros.h"
+
 #include "Containers/UnrealString.h"
 #include "HAL/PlatformProcess.h"
 #include "Internationalization/Text.h"
+#include "Logging/LogVerbosity.h"
+#include "Math/NumericLimits.h"
+#include "Misc/AssertionMacros.h"
+#include "Misc/CString.h"
+#include "Misc/Char.h"
 #include "Misc/FeedbackContext.h"
 
 bool FFeedbackContextMarkup::ParseCommand(const FString& Line, FFeedbackContext* Warn)
diff --git a/Engine/Source/Runtime/Core/Private/Misc/FrameRate.cpp b/Engine/Source/Runtime/Core/Private/Misc/FrameRate.cpp
--- a/Engine/Source/Runtime/Core/Private/Misc/FrameRate.cpp
+++ b/Engine/Source/Runtime/Core/Private/Misc/FrameRate.cpp
@@ -1,10 +1,18 @@
 // Copyright Epic Games, Inc. All Rights Reserved.
 
 #include "Misc/FrameRate.h"
+
 #include "Algo/BinarySearch.h"
 #include "Algo/Reverse.h"
-#include "Misc/ExpressionParser.h"
+#include "Containers/Array.h"
+#include "Containers/ContainerAllocationPolicies.h"
+#include "Internationalization/Internationalization.h"
 #include "Math/BasicMathExpressionEvaluator.h"
+#include "Misc/ExpressionParser.h"
+#include "Misc/ExpressionParserTypes.h"
+#include "Misc/ExpressionParserTypes.inl"
+#include "Misc/Guid.h"
+#include "Templates/UnrealTemplate.h"
 
 #define LOCTEXT_NAMESPACE "FFrameRate"
 
diff --git a/Engine/Source/Runtime/Core/Private/Misc/Guid.cpp b/Engine/Source/Runtime/Core/Private/Misc/Guid.cpp
--- a/Engine/Source/Runtime/Core/Private/Misc/Guid.cpp
+++ b/Engine/Source/Runtime/Core/Private/Misc/Guid.cpp
@@ -1,15 +1,22 @@
 // Copyright Epic Games, Inc. All Rights Reserved.
 
 #include "Misc/Guid.h"
+
 #include "Algo/Replace.h"
 #include "Containers/ArrayView.h"
 #include "Containers/StaticArray.h"
-#include "Misc/ByteSwap.h"
-#include "Misc/Parse.h"
+#include "HAL/PlatformMisc.h"
+#include "HAL/PlatformString.h"
 #include "Math/Int128.h"
-#include "UObject/PropertyPortFlags.h"
 #include "Misc/Base64.h"
+#include "Misc/ByteSwap.h"
+#include "Misc/CString.h"
+#include "Misc/Char.h"
+#include "Misc/Parse.h"
 #include "Misc/StringBuilder.h"
+#include "Serialization/StructuredArchiveNameHelpers.h"
+#include "Serialization/StructuredArchiveSlots.h"
+#include "UObject/PropertyPortFlags.h"
 
 
 /* FGuid interface
diff --git a/Engine/Source/Runtime/Core/Private/Misc/HierarchicalLogArchive.cpp b/Engine/Source/Runtime/Core/Private/Misc/HierarchicalLogArchive.cpp
--- a/Engine/Source/Runtime/Core/Private/Misc/HierarchicalLogArchive.cpp
+++ b/Engine/Source/Runtime/Core/Private/Misc/HierarchicalLogArchive.cpp
@@ -2,6 +2,9 @@
 
 #include "Misc/HierarchicalLogArchive.h"
 
+#include "Containers/StringFwd.h"
+#include "Misc/StringBuilder.h"
+
 FHierarchicalLogArchive::FHierarchicalLogArchive(FArchive& InInnerArchive)
 	: FArchiveProxy(InInnerArchive)
 	, Indentation(0)
diff --git a/Engine/Source/Runtime/Core/Private/Misc/LargeWorldRenderPosition.cpp b/Engine/Source/Runtime/Core/Private/Misc/LargeWorldRenderPosition.cpp
--- a/Engine/Source/Runtime/Core/Private/Misc/LargeWorldRenderPosition.cpp
+++ b/Engine/Source/Runtime/Core/Private/Misc/LargeWorldRenderPosition.cpp
@@ -1,7 +1,8 @@
 // Copyright Epic Games, Inc. All Rights Reserved.
 
 #include "Misc/LargeWorldRenderPosition.h"
-#include "EngineDefines.h"
+
+#include "Math/Matrix.inl"
 
 #define UE_LWC_RENDER_TILE_SIZE_MIN (262144.0)	// LWC_TODO: Make this smaller?
 
diff --git a/Engine/Source/Runtime/Core/Private/Misc/LocalTimestampDirectoryVisitor.cpp b/Engine/Source/Runtime/Core/Private/Misc/LocalTimestampDirectoryVisitor.cpp
--- a/Engine/Source/Runtime/Core/Private/Misc/LocalTimestampDirectoryVisitor.cpp
+++ b/Engine/Source/Runtime/Core/Private/Misc/LocalTimestampDirectoryVisitor.cpp
@@ -1,7 +1,10 @@
 // Copyright Epic Games, Inc. All Rights Reserved.
 
 #include "Misc/LocalTimestampDirectoryVisitor.h"
+
 #include "Misc/Paths.h"
+#include "Templates/RemoveReference.h"
+#include "Templates/UnrealTemplate.h"
 
 /* FLocalTimestampVisitor structors
  *****************************************************************************/
diff --git a/Engine/Source/Runtime/Core/Private/Misc/MapErrors.cpp b/Engine/Source/Runtime/Core/Private/Misc/MapErrors.cpp
--- a/Engine/Source/Runtime/Core/Private/Misc/MapErrors.cpp
+++ b/Engine/Source/Runtime/Core/Private/Misc/MapErrors.cpp
@@ -2,6 +2,8 @@
 
 #include "Misc/MapErrors.h"
 
+#include "CoreTypes.h"
+
 FLazyName FMapErrors::MatchingActorGUID(TEXT("MatchingActorGUID"));
 FLazyName FMapErrors::MatchingLightGUID(TEXT("MatchingLightGUID"));
 FLazyName FMapErrors::ActorLargeShadowCaster(TEXT("ActorLargeShadowCaster"));
diff --git a/Engine/Source/Runtime/Core/Private/Misc/MessageDialog.cpp b/Engine/Source/Runtime/Core/Private/Misc/MessageDialog.cpp
--- a/Engine/Source/Runtime/Core/Private/Misc/MessageDialog.cpp
+++ b/Engine/Source/Runtime/Core/Private/Misc/MessageDialog.cpp
@@ -1,15 +1,18 @@
 // Copyright Epic Games, Inc. All Rights Reserved.
 
 #include "Misc/MessageDialog.h"
-#include "Misc/CString.h"
-#include "Logging/LogMacros.h"
+
 #include "CoreGlobals.h"
-#include "Internationalization/Text.h"
+#include "CoreMinimal.h"
+#include "CoreTypes.h"
 #include "Internationalization/Internationalization.h"
-#include "Misc/OutputDeviceRedirector.h"
-#include "Misc/FeedbackContext.h"
-#include "Misc/CoreDelegates.h"
+#include "Internationalization/Text.h"
+#include "Logging/LogMacros.h"
 #include "Misc/App.h"
+#include "Misc/CString.h"
+#include "Misc/CoreDelegates.h"
+#include "Misc/FeedbackContext.h"
+#include "Misc/OutputDeviceRedirector.h"
 
 namespace
 {
diff --git a/Engine/Source/Runtime/Core/Private/Misc/NamePermissionList.cpp b/Engine/Source/Runtime/Core/Private/Misc/NamePermissionList.cpp
--- a/Engine/Source/Runtime/Core/Private/Misc/NamePermissionList.cpp
+++ b/Engine/Source/Runtime/Core/Private/Misc/NamePermissionList.cpp
@@ -1,7 +1,10 @@
 // Copyright Epic Games, Inc. All Rights Reserved.
 
 #include "Misc/NamePermissionList.h"
-#include "Misc/StringBuilder.h"
+
+#include "Templates/ChooseClass.h"
+#include "Templates/Tuple.h"
+#include "Templates/UnrealTemplate.h"
 
 
 bool FNamePermissionList::PassesFilter(const FName Item) const
diff --git a/Engine/Source/Runtime/Core/Private/Misc/OutputDevice.cpp b/Engine/Source/Runtime/Core/Private/Misc/OutputDevice.cpp
--- a/Engine/Source/Runtime/Core/Private/Misc/OutputDevice.cpp
+++ b/Engine/Source/Runtime/Core/Private/Misc/OutputDevice.cpp
@@ -1,12 +1,18 @@
 // Copyright Epic Games, Inc. All Rights Reserved.
 
 #include "Misc/OutputDevice.h"
+
 #include "Containers/UnrealString.h"
-#include "Logging/LogMacros.h"
+#include "CoreGlobals.h"
 #include "Internationalization/Text.h"
+#include "Logging/LogMacros.h"
 #include "Logging/LogScopedCategoryAndVerbosityOverride.h"
-#include "Misc/OutputDeviceHelper.h"
 #include "Misc/VarargsHelper.h"
+#include "Templates/UnrealTemplate.h"
+#include "UObject/NameTypes.h"
+#include "UObject/UnrealNames.h"
+
+class FOutputDeviceError;
 
 DEFINE_LOG_CATEGORY(LogOutputDevice);
 
diff --git a/Engine/Source/Runtime/Core/Private/Misc/OutputDeviceDebug.cpp b/Engine/Source/Runtime/Core/Private/Misc/OutputDeviceDebug.cpp
--- a/Engine/Source/Runtime/Core/Private/Misc/OutputDeviceDebug.cpp
+++ b/Engine/Source/Runtime/Core/Private/Misc/OutputDeviceDebug.cpp
@@ -1,7 +1,10 @@
 // Copyright Epic Games, Inc. All Rights Reserved.
 
 #include "Misc/OutputDeviceDebug.h"
+
+#include "Containers/StringFwd.h"
 #include "CoreGlobals.h"
+#include "HAL/PlatformMisc.h"
 #include "Misc/OutputDeviceHelper.h"
 #include "Misc/StringBuilder.h"
 
diff --git a/Engine/Source/Runtime/Core/Private/Misc/PathPermissionList.cpp b/Engine/Source/Runtime/Core/Private/Misc/PathPermissionList.cpp
--- a/Engine/Source/Runtime/Core/Private/Misc/PathPermissionList.cpp
+++ b/Engine/Source/Runtime/Core/Private/Misc/PathPermissionList.cpp
@@ -1,9 +1,20 @@
 // Copyright Epic Games, Inc. All Rights Reserved.
 
+#include "Containers/Array.h"
+#include "Containers/Map.h"
+#include "Containers/StringFwd.h"
+#include "Containers/StringView.h"
+#include "Containers/UnrealString.h"
+#include "Delegates/Delegate.h"
+#include "HAL/Platform.h"
+#include "Misc/AssertionMacros.h"
 #include "Misc/NamePermissionList.h"
-#include "Misc/StringBuilder.h"
 #include "Misc/PathViews.h"
-
+#include "Templates/ChooseClass.h"
+#include "Templates/Tuple.h"
+#include "Templates/UnrealTemplate.h"
+#include "UObject/NameTypes.h"
+#include "UObject/UnrealNames.h"
 
 bool FPathPermissionList::PassesFilter(const FStringView Item) const
 {
diff --git a/Engine/Source/Runtime/Core/Private/Misc/PathViews.cpp b/Engine/Source/Runtime/Core/Private/Misc/PathViews.cpp
--- a/Engine/Source/Runtime/Core/Private/Misc/PathViews.cpp
+++ b/Engine/Source/Runtime/Core/Private/Misc/PathViews.cpp
@@ -5,15 +5,20 @@
 #include "Algo/Find.h"
 #include "Algo/FindLast.h"
 #include "Algo/Replace.h"
-#include "Containers/UnrealString.h"
+#include "Containers/ArrayView.h"
 #include "Containers/StringView.h"
+#include "Containers/UnrealString.h"
+#include "HAL/PlatformMisc.h"
 #include "HAL/PlatformProcess.h"
+#include "Math/UnrealMathUtility.h"
+#include "Misc/AssertionMacros.h"
+#include "Misc/CString.h"
 #include "Misc/Char.h"
 #include "Misc/StringBuilder.h"
 #include "String/Find.h"
 #include "String/ParseTokens.h"
 #include "Templates/Function.h"
-#include "UObject/NameTypes.h"
+#include "Templates/UnrealTemplate.h"
 
 namespace UE4PathViews_Private
 {
diff --git a/Engine/Source/Runtime/Core/Private/Misc/QueuedThreadPoolWrapper.cpp b/Engine/Source/Runtime/Core/Private/Misc/QueuedThreadPoolWrapper.cpp
--- a/Engine/Source/Runtime/Core/Private/Misc/QueuedThreadPoolWrapper.cpp
+++ b/Engine/Source/Runtime/Core/Private/Misc/QueuedThreadPoolWrapper.cpp
@@ -1,10 +1,12 @@
 // Copyright Epic Games, Inc. All Rights Reserved.
 
 #include "Misc/QueuedThreadPoolWrapper.h"
-#include "Misc/ScopeRWLock.h"
+
+#include "HAL/PlatformProcess.h"
 #include "Misc/IQueuedWork.h"
+#include "Misc/QueuedThreadPool.h"
 #include "ProfilingDebugging/CountersTrace.h"
-#include "HAL/PlatformProcess.h"
+#include "ProfilingDebugging/CpuProfilerTrace.h"
 
 TRACE_DECLARE_INT_COUNTER(QueuedThreadPoolWrapperScheduledWorkAllocs, TEXT("QueuedThreadPoolWrapper/ScheduledWorkAllocs"));
 
diff --git a/Engine/Source/Runtime/Core/Private/Misc/RemoteConfigIni.cpp b/Engine/Source/Runtime/Core/Private/Misc/RemoteConfigIni.cpp
--- a/Engine/Source/Runtime/Core/Private/Misc/RemoteConfigIni.cpp
+++ b/Engine/Source/Runtime/Core/Private/Misc/RemoteConfigIni.cpp
@@ -1,13 +1,18 @@
 // Copyright Epic Games, Inc. All Rights Reserved.
 
 #include "Misc/RemoteConfigIni.h"
+
 #include "Async/AsyncWork.h"
+#include "CoreMinimal.h"
 #include "HAL/FileManager.h"
-#include "Misc/ScopeLock.h"
+#include "HAL/PlatformProcess.h"
+#include "HAL/PlatformTime.h"
+#include "Misc/App.h"
+#include "Misc/ConfigCacheIni.h"
 #include "Misc/FileHelper.h"
 #include "Misc/Paths.h"
-#include "Misc/ConfigCacheIni.h"
-#include "Misc/App.h"
+#include "Misc/ScopeLock.h"
+#include "Stats/Stats.h"
 
 // Globals
 FRemoteConfig GRemoteConfig;
diff --git a/Engine/Source/Runtime/Core/Private/Misc/SlowTask.cpp b/Engine/Source/Runtime/Core/Private/Misc/SlowTask.cpp
--- a/Engine/Source/Runtime/Core/Private/Misc/SlowTask.cpp
+++ b/Engine/Source/Runtime/Core/Private/Misc/SlowTask.cpp
@@ -1,8 +1,15 @@
 // Copyright Epic Games, Inc. All Rights Reserved.
 
 #include "Misc/SlowTask.h"
+
+#include "Containers/Array.h"
+#include "CoreTypes.h"
+#include "HAL/PlatformMisc.h"
 #include "HAL/PlatformTime.h"
+#include "Math/UnrealMathUtility.h"
+#include "Misc/AssertionMacros.h"
 #include "Misc/FeedbackContext.h"
+#include "Misc/SlowTaskStack.h"
 
 bool FSlowTask::ShouldCreateThrottledSlowTask()
 {
diff --git a/Engine/Source/Runtime/Core/Private/Misc/StringBuilder.cpp b/Engine/Source/Runtime/Core/Private/Misc/StringBuilder.cpp
--- a/Engine/Source/Runtime/Core/Private/Misc/StringBuilder.cpp
+++ b/Engine/Source/Runtime/Core/Private/Misc/StringBuilder.cpp
@@ -1,8 +1,10 @@
 // Copyright Epic Games, Inc. All Rights Reserved.
 
 #include "Misc/StringBuilder.h"
+
 #include "HAL/PlatformMath.h"
 #include "HAL/UnrealMemory.h"
+#include "Misc/CString.h"
 
 static inline uint64_t NextPowerOfTwo(uint64_t x)
 {
diff --git a/Engine/Source/Runtime/Core/Private/Misc/StringUtility.cpp b/Engine/Source/Runtime/Core/Private/Misc/StringUtility.cpp
--- a/Engine/Source/Runtime/Core/Private/Misc/StringUtility.cpp
+++ b/Engine/Source/Runtime/Core/Private/Misc/StringUtility.cpp
@@ -1,7 +1,10 @@
 // Copyright Epic Games, Inc. All Rights Reserved.
 
 #include "Misc/StringUtility.h"
+
+#include "CoreTypes.h"
 #include "Logging/LogMacros.h"
+#include "Misc/CString.h"
 
 
 DEFINE_LOG_CATEGORY_STATIC(LogStringUtility, Log, All);
diff --git a/Engine/Source/Runtime/Core/Private/Misc/TextFilterExpressionEvaluator.cpp b/Engine/Source/Runtime/Core/Private/Misc/TextFilterExpressionEvaluator.cpp
--- a/Engine/Source/Runtime/Core/Private/Misc/TextFilterExpressionEvaluator.cpp
+++ b/Engine/Source/Runtime/Core/Private/Misc/TextFilterExpressionEvaluator.cpp
@@ -1,7 +1,14 @@
 // Copyright Epic Games, Inc. All Rights Reserved.
 
 #include "Misc/TextFilterExpressionEvaluator.h"
+
+#include "Containers/SparseArray.h"
 #include "Math/BasicMathExpressionEvaluator.h"
+#include "Misc/CString.h"
+#include "Misc/ExpressionParserTypes.inl"
+#include "Misc/Guid.h"
+#include "Templates/Function.h"
+#include "Templates/ValueOrError.h"
 
 namespace TextFilterExpressionParser
 {
diff --git a/Engine/Source/Runtime/Core/Private/Misc/TextFilterUtils.cpp b/Engine/Source/Runtime/Core/Private/Misc/TextFilterUtils.cpp
--- a/Engine/Source/Runtime/Core/Private/Misc/TextFilterUtils.cpp
+++ b/Engine/Source/Runtime/Core/Private/Misc/TextFilterUtils.cpp
@@ -2,6 +2,11 @@
 
 #include "Misc/TextFilterUtils.h"
 
+#include "HAL/PlatformString.h"
+#include "Math/UnrealMathUtility.h"
+#include "Misc/Char.h"
+#include "Templates/UnrealTemplate.h"
+
 namespace TextFilterInternal
 {
 	template <typename CharType>
diff --git a/Engine/Source/Runtime/Core/Private/Misc/Timespan.cpp b/Engine/Source/Runtime/Core/Private/Misc/Timespan.cpp
--- a/Engine/Source/Runtime/Core/Private/Misc/Timespan.cpp
+++ b/Engine/Source/Runtime/Core/Private/Misc/Timespan.cpp
@@ -1,8 +1,13 @@
 // Copyright Epic Games, Inc. All Rights Reserved.
 
 #include "Misc/Timespan.h"
-#include "Templates/TypeHash.h"
+
+#include "Containers/Array.h"
 #include "Containers/UnrealString.h"
+#include "HAL/PlatformString.h"
+#include "Misc/CString.h"
+#include "Serialization/Archive.h"
+#include "Templates/TypeHash.h"
 #include "UObject/PropertyPortFlags.h"
 
 
diff --git a/Engine/Source/Runtime/Core/Private/Misc/WildcardString.cpp b/Engine/Source/Runtime/Core/Private/Misc/WildcardString.cpp
--- a/Engine/Source/Runtime/Core/Private/Misc/WildcardString.cpp
+++ b/Engine/Source/Runtime/Core/Private/Misc/WildcardString.cpp
@@ -2,6 +2,8 @@
 
 #include "Misc/WildcardString.h"
 
+#include "Misc/Char.h"
+
 /* FWildcardString static functions
  *****************************************************************************/
 
diff --git a/Engine/Source/Runtime/Core/Private/Modules/SimpleParse.cpp b/Engine/Source/Runtime/Core/Private/Modules/SimpleParse.cpp
--- a/Engine/Source/Runtime/Core/Private/Modules/SimpleParse.cpp
+++ b/Engine/Source/Runtime/Core/Private/Modules/SimpleParse.cpp
@@ -1,6 +1,8 @@
 // Copyright Epic Games, Inc. All Rights Reserved.
 
 #include "Modules/SimpleParse.h"
+
+#include "Containers/UnrealString.h"
 #include "Misc/StringBuilder.h"
 
 bool FSimpleParse::MatchZeroOrMoreWhitespace(const TCHAR*& InOutPtr)
diff --git a/Engine/Source/Runtime/Core/Private/Modules/SimpleParse.h b/Engine/Source/Runtime/Core/Private/Modules/SimpleParse.h
--- a/Engine/Source/Runtime/Core/Private/Modules/SimpleParse.h
+++ b/Engine/Source/Runtime/Core/Private/Modules/SimpleParse.h
@@ -2,8 +2,10 @@
 
 #pragma once
 
+#include "Containers/StringFwd.h"
 #include "CoreTypes.h"
-#include "Containers/UnrealString.h"
+
+class FString;
 
 struct FSimpleParse
 {
diff --git a/Engine/Source/Runtime/Core/Private/ProfilingDebugging/Histogram.cpp b/Engine/Source/Runtime/Core/Private/ProfilingDebugging/Histogram.cpp
--- a/Engine/Source/Runtime/Core/Private/ProfilingDebugging/Histogram.cpp
+++ b/Engine/Source/Runtime/Core/Private/ProfilingDebugging/Histogram.cpp
@@ -2,6 +2,9 @@
 
 #include "ProfilingDebugging/Histogram.h"
 
+#include "Logging/LogCategory.h"
+#include "Templates/Function.h"
+
 DEFINE_LOG_CATEGORY(LogHistograms);
 
 //////////////////////////////////////////////////////////////////////
diff --git a/Engine/Source/Runtime/Core/Private/ProfilingDebugging/ModuleDiagnostics.cpp b/Engine/Source/Runtime/Core/Private/ProfilingDebugging/ModuleDiagnostics.cpp
--- a/Engine/Source/Runtime/Core/Private/ProfilingDebugging/ModuleDiagnostics.cpp
+++ b/Engine/Source/Runtime/Core/Private/ProfilingDebugging/ModuleDiagnostics.cpp
@@ -2,6 +2,8 @@
 
 #include "ProfilingDebugging/ModuleDiagnostics.h"
 
+#include "Trace/Detail/Channel.h"
+
 ////////////////////////////////////////////////////////////////////////////////
 UE_TRACE_CHANNEL_DEFINE(ModuleChannel, "Module information needed for symbols resolution", true)
  
diff --git a/Engine/Source/Runtime/Core/Private/ProfilingDebugging/ModuleDiagnostics.h b/Engine/Source/Runtime/Core/Private/ProfilingDebugging/ModuleDiagnostics.h
--- a/Engine/Source/Runtime/Core/Private/ProfilingDebugging/ModuleDiagnostics.h
+++ b/Engine/Source/Runtime/Core/Private/ProfilingDebugging/ModuleDiagnostics.h
@@ -1,10 +1,11 @@
 // Copyright Epic Games, Inc. All Rights Reserved.
 #pragma once
 
-#include "CoreMinimal.h"
 #include "Trace/Trace.h"
 #include "Trace/Trace.inl"
 
+namespace UE { namespace Trace { class FChannel; } }
+
 ////////////////////////////////////////////////////////////////////////////////
 UE_TRACE_CHANNEL_EXTERN(ModuleChannel)
  
diff --git a/Engine/Source/Runtime/Core/Private/ProfilingDebugging/ResourceSize.cpp b/Engine/Source/Runtime/Core/Private/ProfilingDebugging/ResourceSize.cpp
--- a/Engine/Source/Runtime/Core/Private/ProfilingDebugging/ResourceSize.cpp
+++ b/Engine/Source/Runtime/Core/Private/ProfilingDebugging/ResourceSize.cpp
@@ -2,6 +2,11 @@
 
 #include "ProfilingDebugging/ResourceSize.h"
 
+#include "Containers/Array.h"
+#include "Containers/UnrealString.h"
+#include "Misc/AssertionMacros.h"
+#include "Misc/OutputDevice.h"
+
 FResourceSizeEx::FResourceSizeEx()
 	: ResourceSizeMode(EResourceSizeMode::Exclusive)
 {
diff --git a/Engine/Source/Runtime/Core/Private/Serialization/ArchiveLoadCompressedProxy.cpp b/Engine/Source/Runtime/Core/Private/Serialization/ArchiveLoadCompressedProxy.cpp
--- a/Engine/Source/Runtime/Core/Private/Serialization/ArchiveLoadCompressedProxy.cpp
+++ b/Engine/Source/Runtime/Core/Private/Serialization/ArchiveLoadCompressedProxy.cpp
@@ -1,9 +1,12 @@
 // Copyright Epic Games, Inc. All Rights Reserved.
 
 #include "Serialization/ArchiveLoadCompressedProxy.h"
-#include "Math/UnrealMathUtility.h"
-#include "HAL/UnrealMemory.h"
+
 #include "Containers/Array.h"
+#include "HAL/UnrealMemory.h"
+#include "Math/UnrealMathUtility.h"
+#include "Misc/AssertionMacros.h"
+#include "Misc/Compression.h"
 
 PRAGMA_DISABLE_UNSAFE_TYPECAST_WARNINGS
 
diff --git a/Engine/Source/Runtime/Core/Private/Serialization/ArchiveSaveCompressedProxy.cpp b/Engine/Source/Runtime/Core/Private/Serialization/ArchiveSaveCompressedProxy.cpp
--- a/Engine/Source/Runtime/Core/Private/Serialization/ArchiveSaveCompressedProxy.cpp
+++ b/Engine/Source/Runtime/Core/Private/Serialization/ArchiveSaveCompressedProxy.cpp
@@ -1,10 +1,13 @@
 // Copyright Epic Games, Inc. All Rights Reserved.
 
 #include "Serialization/ArchiveSaveCompressedProxy.h"
-#include "Math/UnrealMathUtility.h"
+
+#include "CoreGlobals.h"
 #include "HAL/UnrealMemory.h"
+#include "Logging/LogCategory.h"
 #include "Logging/LogMacros.h"
-#include "CoreGlobals.h"
+#include "Math/UnrealMathUtility.h"
+#include "Misc/Compression.h"
 
 PRAGMA_DISABLE_UNSAFE_TYPECAST_WARNINGS
 
diff --git a/Engine/Source/Runtime/Core/Private/Serialization/BitReader.cpp b/Engine/Source/Runtime/Core/Private/Serialization/BitReader.cpp
--- a/Engine/Source/Runtime/Core/Private/Serialization/BitReader.cpp
+++ b/Engine/Source/Runtime/Core/Private/Serialization/BitReader.cpp
@@ -1,9 +1,14 @@
 // Copyright Epic Games, Inc. All Rights Reserved.
 
 #include "Serialization/BitReader.h"
-#include "Math/UnrealMathUtility.h"
-#include "Logging/LogMacros.h"
+
 #include "CoreGlobals.h"
+#include "HAL/IConsoleManager.h"
+#include "Logging/LogCategory.h"
+#include "Logging/LogMacros.h"
+#include "Math/UnrealMathUtility.h"
+#include "Serialization/Archive.h"
+#include "Templates/UnrealTemplate.h"
 
 PRAGMA_DISABLE_UNSAFE_TYPECAST_WARNINGS
 
diff --git a/Engine/Source/Runtime/Core/Private/Serialization/BitWriter.cpp b/Engine/Source/Runtime/Core/Private/Serialization/BitWriter.cpp
--- a/Engine/Source/Runtime/Core/Private/Serialization/BitWriter.cpp
+++ b/Engine/Source/Runtime/Core/Private/Serialization/BitWriter.cpp
@@ -1,8 +1,13 @@
 // Copyright Epic Games, Inc. All Rights Reserved.
 
 #include "Serialization/BitWriter.h"
-#include "Logging/LogMacros.h"
+
 #include "CoreGlobals.h"
+#include "HAL/IConsoleManager.h"
+#include "HAL/UnrealMemory.h"
+#include "Logging/LogCategory.h"
+#include "Logging/LogMacros.h"
+#include "Serialization/Archive.h"
 
 PRAGMA_DISABLE_UNSAFE_TYPECAST_WARNINGS
 
diff --git a/Engine/Source/Runtime/Core/Private/Serialization/CompactBinary.cpp b/Engine/Source/Runtime/Core/Private/Serialization/CompactBinary.cpp
--- a/Engine/Source/Runtime/Core/Private/Serialization/CompactBinary.cpp
+++ b/Engine/Source/Runtime/Core/Private/Serialization/CompactBinary.cpp
@@ -2,18 +2,20 @@
 
 #include "Serialization/CompactBinary.h"
 
-#include "Hash/Blake3.h"
+#include "Containers/Array.h"
+#include "Containers/ContainerAllocationPolicies.h"
+#include "HAL/PlatformString.h"
 #include "Hash/xxhash.h"
 #include "Math/UnrealMathUtility.h"
 #include "Memory/CompositeBuffer.h"
-#include "Misc/ByteSwap.h"
+#include "Misc/CString.h"
 #include "Misc/DateTime.h"
 #include "Misc/Guid.h"
 #include "Misc/Timespan.h"
 #include "Serialization/Archive.h"
 #include "Serialization/CompactBinaryValue.h"
 #include "Serialization/VarInt.h"
-#include "String/BytesToHex.h"
+
 #include <atomic>
 
 ///////////////////////////////////////////////////////////////////////////////////////////////////////////////////////
diff --git a/Engine/Source/Runtime/Core/Private/Serialization/CompactBinaryPackage.cpp b/Engine/Source/Runtime/Core/Private/Serialization/CompactBinaryPackage.cpp
--- a/Engine/Source/Runtime/Core/Private/Serialization/CompactBinaryPackage.cpp
+++ b/Engine/Source/Runtime/Core/Private/Serialization/CompactBinaryPackage.cpp
@@ -3,10 +3,13 @@
 #include "Serialization/CompactBinaryPackage.h"
 
 #include "Algo/BinarySearch.h"
-#include "Compression/OodleDataCompression.h"
 #include "Memory/CompositeBuffer.h"
+#include "Memory/MemoryFwd.h"
+#include "Memory/MemoryView.h"
+#include "Serialization/Archive.h"
 #include "Serialization/CompactBinarySerialization.h"
 #include "Serialization/CompactBinaryWriter.h"
+#include "Templates/RemoveReference.h"
 
 ///////////////////////////////////////////////////////////////////////////////////////////////////////////////////////
 
diff --git a/Engine/Source/Runtime/Core/Private/Serialization/CompactBinarySerialization.cpp b/Engine/Source/Runtime/Core/Private/Serialization/CompactBinarySerialization.cpp
--- a/Engine/Source/Runtime/Core/Private/Serialization/CompactBinarySerialization.cpp
+++ b/Engine/Source/Runtime/Core/Private/Serialization/CompactBinarySerialization.cpp
@@ -2,19 +2,28 @@
 
 #include "Serialization/CompactBinarySerialization.h"
 
+#include "Containers/Array.h"
+#include "Containers/ContainerAllocationPolicies.h"
 #include "Containers/StringConv.h"
-#include "HAL/Platform.h"
+#include "Containers/StringView.h"
+#include "Containers/UnrealString.h"
+#include "HAL/PlatformString.h"
+#include "Memory/MemoryView.h"
 #include "Misc/AsciiSet.h"
+#include "Misc/AssertionMacros.h"
 #include "Misc/Base64.h"
 #include "Misc/DateTime.h"
 #include "Misc/Guid.h"
 #include "Misc/StringBuilder.h"
 #include "Misc/Timespan.h"
+#include "Serialization/Archive.h"
 #include "Serialization/CompactBinaryValidation.h"
 #include "Serialization/CompactBinaryValue.h"
 #include "Serialization/VarInt.h"
 #include "Templates/IdentityFunctor.h"
 #include "Templates/Invoke.h"
+#include "Templates/RemoveReference.h"
+#include "Templates/UnrealTemplate.h"
 
 ///////////////////////////////////////////////////////////////////////////////////////////////////////////////////////
 
diff --git a/Engine/Source/Runtime/Core/Private/Serialization/CompactBinaryValidation.cpp b/Engine/Source/Runtime/Core/Private/Serialization/CompactBinaryValidation.cpp
--- a/Engine/Source/Runtime/Core/Private/Serialization/CompactBinaryValidation.cpp
+++ b/Engine/Source/Runtime/Core/Private/Serialization/CompactBinaryValidation.cpp
@@ -2,8 +2,18 @@
 
 #include "Serialization/CompactBinaryValidation.h"
 
+#include "Algo/Sort.h"
 #include "Compression/CompressedBuffer.h"
+#include "Containers/Array.h"
+#include "Containers/ArrayView.h"
+#include "Containers/ContainerAllocationPolicies.h"
+#include "Containers/ContainersFwd.h"
+#include "Containers/StringFwd.h"
+#include "Containers/StringView.h"
+#include "HAL/PlatformMemory.h"
 #include "IO/IoHash.h"
+#include "Memory/MemoryView.h"
+#include "Memory/SharedBuffer.h"
 #include "Misc/ByteSwap.h"
 #include "Serialization/CompactBinary.h"
 #include "Serialization/CompactBinaryPackage.h"
diff --git a/Engine/Source/Runtime/Core/Private/Serialization/CompactBinaryWriter.cpp b/Engine/Source/Runtime/Core/Private/Serialization/CompactBinaryWriter.cpp
--- a/Engine/Source/Runtime/Core/Private/Serialization/CompactBinaryWriter.cpp
+++ b/Engine/Source/Runtime/Core/Private/Serialization/CompactBinaryWriter.cpp
@@ -2,12 +2,17 @@
 
 #include "Serialization/CompactBinaryWriter.h"
 
-#include "Containers/StringConv.h"
+#include "HAL/PlatformString.h"
+#include "HAL/UnrealMemory.h"
 #include "IO/IoHash.h"
+#include "Math/NumericLimits.h"
+#include "Math/UnrealMathUtility.h"
 #include "Memory/CompositeBuffer.h"
 #include "Memory/SharedBuffer.h"
+#include "Misc/AssertionMacros.h"
 #include "Misc/ByteSwap.h"
 #include "Misc/DateTime.h"
+#include "Misc/EnumClassFlags.h"
 #include "Misc/Guid.h"
 #include "Misc/Timespan.h"
 #include "Serialization/Archive.h"
diff --git a/Engine/Source/Runtime/Core/Private/Serialization/Csv/CsvParser.cpp b/Engine/Source/Runtime/Core/Private/Serialization/Csv/CsvParser.cpp
--- a/Engine/Source/Runtime/Core/Private/Serialization/Csv/CsvParser.cpp
+++ b/Engine/Source/Runtime/Core/Private/Serialization/Csv/CsvParser.cpp
@@ -2,6 +2,8 @@
 
 #include "Serialization/Csv/CsvParser.h"
 
+#include "Templates/UnrealTemplate.h"
+
 FCsvParser::FCsvParser(FString InSourceString)
 	: SourceString(MoveTemp(InSourceString))		
 {
diff --git a/Engine/Source/Runtime/Core/Private/Serialization/Csv/CsvParserTests.cpp b/Engine/Source/Runtime/Core/Private/Serialization/Csv/CsvParserTests.cpp
--- a/Engine/Source/Runtime/Core/Private/Serialization/Csv/CsvParserTests.cpp
+++ b/Engine/Source/Runtime/Core/Private/Serialization/Csv/CsvParserTests.cpp
@@ -1,7 +1,8 @@
 // Copyright Epic Games, Inc. All Rights Reserved.
 
-#include "CoreTypes.h"
 #include "Containers/UnrealString.h"
+#include "CoreMinimal.h"
+#include "CoreTypes.h"
 #include "Misc/AutomationTest.h"
 #include "Serialization/Csv/CsvParser.h"
 
diff --git a/Engine/Source/Runtime/Core/Private/Serialization/CustomVersion.cpp b/Engine/Source/Runtime/Core/Private/Serialization/CustomVersion.cpp
--- a/Engine/Source/Runtime/Core/Private/Serialization/CustomVersion.cpp
+++ b/Engine/Source/Runtime/Core/Private/Serialization/CustomVersion.cpp
@@ -5,11 +5,19 @@
 =============================================================================*/
 
 #include "Serialization/CustomVersion.h"
-#include "Serialization/StructuredArchive.h"
+
 #include "Algo/Sort.h"
+#include "Containers/ContainerAllocationPolicies.h"
 #include "Containers/Map.h"
+#include "HAL/CriticalSection.h"
+#include "Misc/AssertionMacros.h"
 #include "Misc/ScopeRWLock.h"
-#include "CoreGlobals.h"
+#include "Serialization/Archive.h"
+#include "Serialization/StructuredArchive.h"
+#include "Serialization/StructuredArchiveAdapters.h"
+#include "Serialization/StructuredArchiveNameHelpers.h"
+#include "Serialization/StructuredArchiveSlots.h"
+#include "UObject/UnrealNames.h"
 
 namespace
 {
diff --git a/Engine/Source/Runtime/Core/Private/Serialization/FileRegions.cpp b/Engine/Source/Runtime/Core/Private/Serialization/FileRegions.cpp
--- a/Engine/Source/Runtime/Core/Private/Serialization/FileRegions.cpp
+++ b/Engine/Source/Runtime/Core/Private/Serialization/FileRegions.cpp
@@ -2,6 +2,9 @@
 
 #include "Serialization/FileRegions.h"
 
+#include "Misc/AssertionMacros.h"
+#include "Serialization/Archive.h"
+
 const TCHAR* FFileRegion::RegionsFileExtension = TEXT(".uregs");
 
 // NOTE: This serialization function must match FileRegion::ReadRegionsFromFile in AutomationUtils/FileRegions.cs
diff --git a/Engine/Source/Runtime/Core/Private/Serialization/Formatters/BinaryArchiveFormatter.cpp b/Engine/Source/Runtime/Core/Private/Serialization/Formatters/BinaryArchiveFormatter.cpp
--- a/Engine/Source/Runtime/Core/Private/Serialization/Formatters/BinaryArchiveFormatter.cpp
+++ b/Engine/Source/Runtime/Core/Private/Serialization/Formatters/BinaryArchiveFormatter.cpp
@@ -1,7 +1,6 @@
 // Copyright Epic Games, Inc. All Rights Reserved.
 
 #include "Serialization/Formatters/BinaryArchiveFormatter.h"
-#include "Serialization/Archive.h"
 
 FBinaryArchiveFormatter::FBinaryArchiveFormatter(FArchive& InInner) 
 	: Inner(InInner)
diff --git a/Engine/Source/Runtime/Core/Private/Serialization/LargeMemoryData.cpp b/Engine/Source/Runtime/Core/Private/Serialization/LargeMemoryData.cpp
--- a/Engine/Source/Runtime/Core/Private/Serialization/LargeMemoryData.cpp
+++ b/Engine/Source/Runtime/Core/Private/Serialization/LargeMemoryData.cpp
@@ -1,9 +1,11 @@
 // Copyright Epic Games, Inc. All Rights Reserved.
 
 #include "Serialization/LargeMemoryData.h"
-#include "Logging/LogMacros.h"
-#include "CoreGlobals.h"
+
+#include "Containers/LockFreeList.h"
 #include "HAL/IConsoleManager.h"
+#include "HAL/UnrealMemory.h"
+#include "Math/UnrealMathUtility.h"
 
 
 int32 GPooledLargeMemoryData_MaxPoolLength = 2;
diff --git a/Engine/Source/Runtime/Core/Private/Serialization/LargeMemoryReader.cpp b/Engine/Source/Runtime/Core/Private/Serialization/LargeMemoryReader.cpp
--- a/Engine/Source/Runtime/Core/Private/Serialization/LargeMemoryReader.cpp
+++ b/Engine/Source/Runtime/Core/Private/Serialization/LargeMemoryReader.cpp
@@ -1,8 +1,11 @@
 // Copyright Epic Games, Inc. All Rights Reserved.
 
 #include "Serialization/LargeMemoryReader.h"
-#include "Logging/LogMacros.h"
+
 #include "CoreGlobals.h"
+#include "HAL/UnrealMemory.h"
+#include "Logging/LogCategory.h"
+#include "Logging/LogMacros.h"
 
 /*----------------------------------------------------------------------------
 	FLargeMemoryReader
diff --git a/Engine/Source/Runtime/Core/Private/Serialization/LargeMemoryWriter.cpp b/Engine/Source/Runtime/Core/Private/Serialization/LargeMemoryWriter.cpp
--- a/Engine/Source/Runtime/Core/Private/Serialization/LargeMemoryWriter.cpp
+++ b/Engine/Source/Runtime/Core/Private/Serialization/LargeMemoryWriter.cpp
@@ -1,8 +1,10 @@
 // Copyright Epic Games, Inc. All Rights Reserved.
 
 #include "Serialization/LargeMemoryWriter.h"
-#include "Logging/LogMacros.h"
+
 #include "CoreGlobals.h"
+#include "Logging/LogCategory.h"
+#include "Logging/LogMacros.h"
 
 /*----------------------------------------------------------------------------
 	FLargeMemoryWriter
diff --git a/Engine/Source/Runtime/Core/Private/Serialization/MappedName.cpp b/Engine/Source/Runtime/Core/Private/Serialization/MappedName.cpp
--- a/Engine/Source/Runtime/Core/Private/Serialization/MappedName.cpp
+++ b/Engine/Source/Runtime/Core/Private/Serialization/MappedName.cpp
@@ -1,7 +1,7 @@
 // Copyright Epic Games, Inc. All Rights Reserved.
 
 #include "Serialization/MappedName.h"
-#include "Serialization/Archive.h"
+
 #include "UObject/NameBatchSerialization.h"
 
 FArchive& operator<<(FArchive& Ar, FMappedName& MappedName)
diff --git a/Engine/Source/Runtime/Core/Private/Serialization/NameAsStringIndexProxyArchive.cpp b/Engine/Source/Runtime/Core/Private/Serialization/NameAsStringIndexProxyArchive.cpp
--- a/Engine/Source/Runtime/Core/Private/Serialization/NameAsStringIndexProxyArchive.cpp
+++ b/Engine/Source/Runtime/Core/Private/Serialization/NameAsStringIndexProxyArchive.cpp
@@ -1,11 +1,14 @@
 // Copyright Epic Games, Inc. All Rights Reserved.
 
 #include "Serialization/NameAsStringIndexProxyArchive.h"
-#include "Serialization/VarInt.h"
-#include "Containers/UnrealString.h"
-#include "UObject/NameTypes.h"
+
 #include "Containers/Array.h"
 #include "Containers/Set.h"
+#include "Containers/UnrealString.h"
+#include "CoreTypes.h"
+#include "Serialization/Archive.h"
+#include "Serialization/VarInt.h"
+#include "UObject/NameTypes.h"
 
 FArchive& FNameAsStringIndexProxyArchive::operator<<(class FName& N)
 {
diff --git a/Engine/Source/Runtime/Core/Private/Serialization/StructuredArchiveFormatter.cpp b/Engine/Source/Runtime/Core/Private/Serialization/StructuredArchiveFormatter.cpp
--- a/Engine/Source/Runtime/Core/Private/Serialization/StructuredArchiveFormatter.cpp
+++ b/Engine/Source/Runtime/Core/Private/Serialization/StructuredArchiveFormatter.cpp
@@ -1,9 +1,6 @@
 // Copyright Epic Games, Inc. All Rights Reserved.
 
 #include "Serialization/StructuredArchiveFormatter.h"
-#include "Containers/Array.h"
-#include "Containers/UnrealString.h"
-#include "UObject/NameTypes.h"
 
 //////////// FStructuredArchiveFormatter ////////////
 
diff --git a/Engine/Source/Runtime/Core/Private/Serialization/StructuredArchiveFromArchive.cpp b/Engine/Source/Runtime/Core/Private/Serialization/StructuredArchiveFromArchive.cpp
--- a/Engine/Source/Runtime/Core/Private/Serialization/StructuredArchiveFromArchive.cpp
+++ b/Engine/Source/Runtime/Core/Private/Serialization/StructuredArchiveFromArchive.cpp
@@ -1,8 +1,11 @@
 // Copyright Epic Games, Inc. All Rights Reserved.
 
+#include "Serialization/Formatters/BinaryArchiveFormatter.h"
 #include "Serialization/StructuredArchive.h"
-#include "Containers/Set.h"
-#include "Containers/UnrealString.h"
+#include "Serialization/StructuredArchiveAdapters.h"
+#include "Templates/MemoryOps.h"
+
+class FArchive;
 
 struct FStructuredArchiveFromArchive::FImpl
 {
diff --git a/Engine/Source/Runtime/Core/Private/Serialization/VarInt.cpp b/Engine/Source/Runtime/Core/Private/Serialization/VarInt.cpp
--- a/Engine/Source/Runtime/Core/Private/Serialization/VarInt.cpp
+++ b/Engine/Source/Runtime/Core/Private/Serialization/VarInt.cpp
@@ -2,6 +2,7 @@
 
 #include "Serialization/VarInt.h"
 
+#include "CoreTypes.h"
 #include "Serialization/Archive.h"
 
 int64 ReadVarIntFromArchive(FArchive& Ar)
diff --git a/Engine/Source/Runtime/Core/Private/String/BytesToHex.cpp b/Engine/Source/Runtime/Core/Private/String/BytesToHex.cpp
--- a/Engine/Source/Runtime/Core/Private/String/BytesToHex.cpp
+++ b/Engine/Source/Runtime/Core/Private/String/BytesToHex.cpp
@@ -1,8 +1,12 @@
 // Copyright Epic Games, Inc. All Rights Reserved.
 
 #include "String/BytesToHex.h"
+
+#include "Containers/Array.h"
+#include "Containers/ArrayView.h"
 #include "Containers/UnrealString.h"
 #include "Misc/StringBuilder.h"
+#include "Templates/UnrealTemplate.h"
 
 namespace UE::String
 {
diff --git a/Engine/Source/Runtime/Core/Private/String/Escape.cpp b/Engine/Source/Runtime/Core/Private/String/Escape.cpp
--- a/Engine/Source/Runtime/Core/Private/String/Escape.cpp
+++ b/Engine/Source/Runtime/Core/Private/String/Escape.cpp
@@ -2,6 +2,7 @@
 
 #include "String/Escape.h"
 
+#include "CoreTypes.h"
 #include "Misc/StringBuilder.h"
 #include "String/BytesToHex.h"
 
diff --git a/Engine/Source/Runtime/Core/Private/String/Find.cpp b/Engine/Source/Runtime/Core/Private/String/Find.cpp
--- a/Engine/Source/Runtime/Core/Private/String/Find.cpp
+++ b/Engine/Source/Runtime/Core/Private/String/Find.cpp
@@ -5,7 +5,12 @@
 #include "Algo/AllOf.h"
 #include "Algo/Find.h"
 #include "Algo/NoneOf.h"
+#include "Containers/Array.h"
+#include "Containers/ArrayView.h"
+#include "Containers/ContainerAllocationPolicies.h"
 #include "Containers/StringView.h"
+#include "Misc/AssertionMacros.h"
+#include "Misc/Char.h"
 
 namespace UE::String::Private
 {
diff --git a/Engine/Source/Runtime/Core/Private/String/LexFromString.cpp b/Engine/Source/Runtime/Core/Private/String/LexFromString.cpp
--- a/Engine/Source/Runtime/Core/Private/String/LexFromString.cpp
+++ b/Engine/Source/Runtime/Core/Private/String/LexFromString.cpp
@@ -2,7 +2,6 @@
 
 #include "String/LexFromString.h"
 
-#include "Containers/StringView.h"
 #include "Containers/UnrealString.h"
 #include "Misc/StringBuilder.h"
 
diff --git a/Engine/Source/Runtime/Core/Private/String/LineEndings.cpp b/Engine/Source/Runtime/Core/Private/String/LineEndings.cpp
--- a/Engine/Source/Runtime/Core/Private/String/LineEndings.cpp
+++ b/Engine/Source/Runtime/Core/Private/String/LineEndings.cpp
@@ -2,6 +2,10 @@
 
 #include "String/LineEndings.h"
 
+#include "CoreTypes.h"
+#include "Misc/CString.h"
+#include "Templates/UnrealTemplate.h"
+
 namespace UE::String
 {
 	CORE_API FString FromHostLineEndings(const FString& InString)
diff --git a/Engine/Source/Runtime/Core/Private/String/ParseLines.cpp b/Engine/Source/Runtime/Core/Private/String/ParseLines.cpp
--- a/Engine/Source/Runtime/Core/Private/String/ParseLines.cpp
+++ b/Engine/Source/Runtime/Core/Private/String/ParseLines.cpp
@@ -2,6 +2,8 @@
 
 #include "String/ParseLines.h"
 
+#include "Templates/Function.h"
+
 namespace UE::String
 {
 
diff --git a/Engine/Source/Runtime/Core/Private/String/ParseTokens.cpp b/Engine/Source/Runtime/Core/Private/String/ParseTokens.cpp
--- a/Engine/Source/Runtime/Core/Private/String/ParseTokens.cpp
+++ b/Engine/Source/Runtime/Core/Private/String/ParseTokens.cpp
@@ -7,7 +7,14 @@
 #include "Algo/NoneOf.h"
 #include "Algo/Transform.h"
 #include "Containers/Array.h"
+#include "Containers/ArrayView.h"
 #include "Containers/BitArray.h"
+#include "Containers/ContainerAllocationPolicies.h"
+#include "Misc/AssertionMacros.h"
+#include "Misc/CString.h"
+#include "Misc/Char.h"
+#include "Templates/Function.h"
+#include "Templates/UnrealTemplate.h"
 
 namespace UE::String::Private
 {
diff --git a/Engine/Source/Runtime/Core/Private/Tasks/Pipe.cpp b/Engine/Source/Runtime/Core/Private/Tasks/Pipe.cpp
--- a/Engine/Source/Runtime/Core/Private/Tasks/Pipe.cpp
+++ b/Engine/Source/Runtime/Core/Private/Tasks/Pipe.cpp
@@ -1,8 +1,8 @@
 // Copyright Epic Games, Inc. All Rights Reserved.
 
 #include "Tasks/Pipe.h"
+
 #include "Misc/ScopeExit.h"
-#include "Tasks/Task.h"
 
 namespace UE::Tasks
 {
diff --git a/Engine/Source/Runtime/Core/Private/Unix/UnixCriticalSection.cpp b/Engine/Source/Runtime/Core/Private/Unix/UnixCriticalSection.cpp
--- a/Engine/Source/Runtime/Core/Private/Unix/UnixCriticalSection.cpp
+++ b/Engine/Source/Runtime/Core/Private/Unix/UnixCriticalSection.cpp
@@ -1,16 +1,18 @@
 // Copyright Epic Games, Inc. All Rights Reserved.
 
 #include "Unix/UnixCriticalSection.h"
-#include "Misc/DateTime.h"
+
+#include "Containers/StringConv.h"
+#include "HAL/CriticalSection.h"
 #include "HAL/PlatformProcess.h"
 #include "HAL/PlatformTime.h"
-#include "Containers/StringConv.h"
-#include <spawn.h>
-#include <sys/wait.h>
-#include <sys/resource.h>
-#include <sys/ioctl.h>
+#include "Math/UnrealMathUtility.h"
+#include "Misc/AssertionMacros.h"
+
+#include <fcntl.h>
+#include <float.h>
 #include <sys/file.h>
-#include <asm/ioctls.h>
+#include <unistd.h>
 
 FUnixSystemWideCriticalSection::FUnixSystemWideCriticalSection(const FString& InName, FTimespan InTimeout)
 {
diff --git a/Engine/Source/Runtime/Core/Private/Unix/UnixPlatformOutputDevices.cpp b/Engine/Source/Runtime/Core/Private/Unix/UnixPlatformOutputDevices.cpp
--- a/Engine/Source/Runtime/Core/Private/Unix/UnixPlatformOutputDevices.cpp
+++ b/Engine/Source/Runtime/Core/Private/Unix/UnixPlatformOutputDevices.cpp
@@ -1,19 +1,19 @@
 // Copyright Epic Games, Inc. All Rights Reserved.
 
 #include "Unix/UnixPlatformOutputDevices.h"
-#include "Containers/StringConv.h"
-#include "Logging/LogMacros.h"
+
 #include "CoreGlobals.h"
-#include "Misc/Parse.h"
+#include "CoreMinimal.h"
+#include "CoreTypes.h"
 #include "Misc/CommandLine.h"
-#include "Misc/Paths.h"
-#include "Misc/OutputDeviceHelper.h"
-#include "Misc/CoreDelegates.h"
-#include "Misc/App.h"
 #include "Misc/OutputDeviceConsole.h"
 #include "Misc/OutputDeviceRedirector.h"
+#include "Misc/Parse.h"
+#include "Misc/Paths.h"
 #include "Unix/UnixErrorOutputDevice.h"
 
+class FOutputDeviceError;
+
 void FUnixOutputDevices::SetupOutputDevices()
 {
 	check(GLog);
diff --git a/Engine/Source/Runtime/Core/Public/Algo/TopologicalSort.h b/Engine/Source/Runtime/Core/Public/Algo/TopologicalSort.h
--- a/Engine/Source/Runtime/Core/Public/Algo/TopologicalSort.h
+++ b/Engine/Source/Runtime/Core/Public/Algo/TopologicalSort.h
@@ -3,6 +3,7 @@
 #pragma once
 
 #include "Algo/KahnTopologicalSort.h"
+#include "CoreMinimal.h"
 
 namespace Algo
 {
diff --git a/Engine/Source/Runtime/Core/Public/Async/Async.h b/Engine/Source/Runtime/Core/Public/Async/Async.h
--- a/Engine/Source/Runtime/Core/Public/Async/Async.h
+++ b/Engine/Source/Runtime/Core/Public/Async/Async.h
@@ -2,21 +2,24 @@
 
 #pragma once
 
-#include "CoreTypes.h"
 #include "Async/Future.h"
 #include "Async/TaskGraphInterfaces.h"
 #include "Containers/UnrealString.h"
+#include "CoreTypes.h"
+#include "HAL/PlatformAffinity.h"
 #include "HAL/PlatformProcess.h"
 #include "HAL/Runnable.h"
 #include "HAL/RunnableThread.h"
 #include "HAL/ThreadSafeCounter.h"
 #include "Misc/AssertionMacros.h"
 #include "Misc/CoreStats.h"
+#include "Misc/Fork.h"
 #include "Misc/IQueuedWork.h"
 #include "Misc/QueuedThreadPool.h"
-#include "Misc/Fork.h"
 #include "Stats/Stats.h"
+#include "Stats/Stats2.h"
 #include "Templates/Function.h"
+#include "Templates/UnrealTemplate.h"
 
 /**
  * Enumerates available asynchronous execution methods.
diff --git a/Engine/Source/Runtime/Core/Public/Async/Fundamental/ReserveScheduler.h b/Engine/Source/Runtime/Core/Public/Async/Fundamental/ReserveScheduler.h
--- a/Engine/Source/Runtime/Core/Public/Async/Fundamental/ReserveScheduler.h
+++ b/Engine/Source/Runtime/Core/Public/Async/Fundamental/ReserveScheduler.h
@@ -1,6 +1,11 @@
 // Copyright Epic Games, Inc. All Rights Reserved.
 
 #pragma once 
+#include "Async/Fundamental/TaskDelegate.h"
+#include "CoreMinimal.h"
+#include "HAL/Event.h"
+#include "HAL/PlatformAffinity.h"
+#include "HAL/Thread.h"
 #include "Scheduler.h"
 
 namespace LowLevelTasks
diff --git a/Engine/Source/Runtime/Core/Public/Async/Fundamental/Scheduler.h b/Engine/Source/Runtime/Core/Public/Async/Fundamental/Scheduler.h
--- a/Engine/Source/Runtime/Core/Public/Async/Fundamental/Scheduler.h
+++ b/Engine/Source/Runtime/Core/Public/Async/Fundamental/Scheduler.h
@@ -2,13 +2,18 @@
 
 #pragma once 
 #include "Async/Fundamental/Task.h"
-#include "HAL/Thread.h"
+#include "Async/Fundamental/TaskDelegate.h"
+#include "Containers/Array.h"
+#include "Containers/ArrayView.h"
+#include "CoreMinimal.h"
 #include "HAL/Event.h"
+#include "HAL/PlatformAffinity.h"
 #include "HAL/PlatformProcess.h"
+#include "HAL/Thread.h"
 #include "LocalQueue.h"
-#include "Containers/Array.h"
-#include "Templates/UniquePtr.h"
 #include "Templates/Function.h"
+#include "Templates/IsInvocable.h"
+#include "Templates/UniquePtr.h"
 
 namespace LowLevelTasks
 {
diff --git a/Engine/Source/Runtime/Core/Public/Async/ParallelFor.h b/Engine/Source/Runtime/Core/Public/Async/ParallelFor.h
--- a/Engine/Source/Runtime/Core/Public/Async/ParallelFor.h
+++ b/Engine/Source/Runtime/Core/Public/Async/ParallelFor.h
@@ -6,21 +6,32 @@
 
 #pragma once
 
-#include <atomic>
+#include "Async/Fundamental/Scheduler.h"
+#include "Async/Fundamental/Task.h"
+#include "Async/TaskGraphInterfaces.h"
+#include "Containers/ArrayView.h"
+#include "CoreMinimal.h"
 #include "CoreTypes.h"
-#include "Misc/AssertionMacros.h"
-#include "Math/UnrealMathUtility.h"
-#include "Templates/Function.h"
-#include "Templates/SharedPointer.h"
-#include "Templates/RefCounting.h"
+#include "Experimental/ConcurrentLinearAllocator.h"
+#include "HAL/Event.h"
+#include "HAL/LowLevelMemTracker.h"
+#include "HAL/PlatformTime.h"
 #include "HAL/ThreadSafeCounter.h"
-#include "Stats/Stats.h"
-#include "Async/TaskGraphInterfaces.h"
+#include "Math/UnrealMathUtility.h"
 #include "Misc/App.h"
+#include "Misc/AssertionMacros.h"
+#include "Misc/EnumClassFlags.h"
 #include "Misc/Fork.h"
-#include "Async/Fundamental/Task.h"
-#include "Async/Fundamental/Scheduler.h"
-#include "Experimental/ConcurrentLinearAllocator.h"
+#include "Misc/MemStack.h"
+#include "ProfilingDebugging/CpuProfilerTrace.h"
+#include "ProfilingDebugging/TagTrace.h"
+#include "Stats/Stats.h"
+#include "Stats/Stats2.h"
+#include "Templates/Function.h"
+#include "Templates/RefCounting.h"
+#include "Templates/SharedPointer.h"
+
+namespace UE { namespace LLMPrivate { class FTagData; } }
 
 extern CORE_API int32 GParallelForBackgroundYieldingTimeoutMs;
 
diff --git a/Engine/Source/Runtime/Core/Public/Async/TaskTrace.h b/Engine/Source/Runtime/Core/Public/Async/TaskTrace.h
--- a/Engine/Source/Runtime/Core/Public/Async/TaskTrace.h
+++ b/Engine/Source/Runtime/Core/Public/Async/TaskTrace.h
@@ -2,8 +2,12 @@
 
 #pragma once
 
-#include "Trace/Config.h"
 #include "Containers/Array.h"
+#include "HAL/Platform.h"
+#include "Trace/Config.h"
+#include "Trace/Trace.h"
+
+namespace UE { namespace Trace { class FChannel; } }
 
 #if !defined(UE_TASK_TRACE_ENABLED)
 #if UE_TRACE_ENABLED && !IS_PROGRAM && !UE_BUILD_SHIPPING
diff --git a/Engine/Source/Runtime/Core/Public/Compression/CompressedBuffer.h b/Engine/Source/Runtime/Core/Public/Compression/CompressedBuffer.h
--- a/Engine/Source/Runtime/Core/Public/Compression/CompressedBuffer.h
+++ b/Engine/Source/Runtime/Core/Public/Compression/CompressedBuffer.h
@@ -2,6 +2,7 @@
 
 #pragma once
 
+#include "CoreMinimal.h"
 #include "Memory/CompositeBuffer.h"
 #include "Memory/MemoryFwd.h"
 #include "Memory/SharedBuffer.h"
diff --git a/Engine/Source/Runtime/Core/Public/Compression/OodleDataCompression.h b/Engine/Source/Runtime/Core/Public/Compression/OodleDataCompression.h
--- a/Engine/Source/Runtime/Core/Public/Compression/OodleDataCompression.h
+++ b/Engine/Source/Runtime/Core/Public/Compression/OodleDataCompression.h
@@ -2,6 +2,7 @@
 
 #pragma once
 
+#include "Containers/UnrealString.h"
 #include "CoreTypes.h"
 #include "Logging/LogCategory.h"
 #include "Logging/LogMacros.h"
diff --git a/Engine/Source/Runtime/Core/Public/Compression/OodleDataCompressionUtil.h b/Engine/Source/Runtime/Core/Public/Compression/OodleDataCompressionUtil.h
--- a/Engine/Source/Runtime/Core/Public/Compression/OodleDataCompressionUtil.h
+++ b/Engine/Source/Runtime/Core/Public/Compression/OodleDataCompressionUtil.h
@@ -2,6 +2,9 @@
 
 #pragma once
 
+#include "Containers/Array.h"
+#include "Containers/ContainersFwd.h"
+#include "HAL/Platform.h"
 #include "OodleDataCompression.h"
 
 // OodleDataCompressionUtil : utilities for common Unreal actions built on top of OodleDataCompression
diff --git a/Engine/Source/Runtime/Core/Public/Containers/BitArray.h b/Engine/Source/Runtime/Core/Public/Containers/BitArray.h
--- a/Engine/Source/Runtime/Core/Public/Containers/BitArray.h
+++ b/Engine/Source/Runtime/Core/Public/Containers/BitArray.h
@@ -2,17 +2,25 @@
 
 #pragma once
 
+#include "Containers/ContainerAllocationPolicies.h"
 #include "CoreTypes.h"
-#include "Misc/AssertionMacros.h"
+#include "HAL/PlatformAtomics.h"
 #include "HAL/UnrealMemory.h"
-#include "Templates/UnrealTypeTraits.h"
-#include "Templates/UnrealTemplate.h"
-#include "Containers/ContainerAllocationPolicies.h"
-#include "Serialization/Archive.h"
-#include "Serialization/MemoryImageWriter.h"
 #include "Math/UnrealMathUtility.h"
+#include "Misc/AssertionMacros.h"
 #include "Misc/EnumClassFlags.h"
+#include "Serialization/Archive.h"
+#include "Serialization/MemoryImageWriter.h"
+#include "Serialization/MemoryLayout.h"
+#include "Templates/ChooseClass.h"
+#include "Templates/EnableIf.h"
+#include "Templates/Invoke.h"
+#include "Templates/UnrealTemplate.h"
+#include "Templates/UnrealTypeTraits.h"
+#include "Traits/IsVoidType.h"
 
+class FPointerTableBase;
+class FSHA1;
 template<typename Allocator > class TBitArray;
 
 // Functions for manipulating bit sets.
@@ -37,17 +45,13 @@ struct FBitSet
 	}
 };
 
-
 // Forward declaration.
+template<typename Allocator = FDefaultBitArrayAllocator,typename OtherAllocator = FDefaultBitArrayAllocator, bool Both=true>
+class TConstDualSetBitIterator;
 template<typename Allocator = FDefaultBitArrayAllocator>
 class TBitArray;
-
 template<typename Allocator = FDefaultBitArrayAllocator>
 class TConstSetBitIterator;
-
-template<typename Allocator = FDefaultBitArrayAllocator,typename OtherAllocator = FDefaultBitArrayAllocator, bool Both=true>
-class TConstDualSetBitIterator;
-
 template <typename AllocatorType, typename InDerivedType = void>
 class TScriptBitArray;
 
