Ran a pass of IWYU on Core.

diff --git a/Engine/Source/Runtime/Core/Public/Containers/HashTable.h b/Engine/Source/Runtime/Core/Public/Containers/HashTable.h
--- a/Engine/Source/Runtime/Core/Public/Containers/HashTable.h
+++ b/Engine/Source/Runtime/Core/Public/Containers/HashTable.h
@@ -2,17 +2,26 @@
 
 #pragma once
 
-#include "CoreTypes.h"
 #include "Containers/ContainerAllocationPolicies.h"
-#include "Serialization/MemoryImageWriter.h"
-#include "Serialization/MemoryLayout.h"
-#include "Serialization/MemoryImageWriter.h"
 #include "Containers/ContainerAllocationPolicies.h"
-#include "Misc/AssertionMacros.h"
+#include "CoreTypes.h"
+#include "HAL/PlatformAtomics.h"
 #include "HAL/UnrealMemory.h"
 #include "Math/UnrealMathUtility.h"
+#include "Misc/AssertionMacros.h"
+#include "Serialization/MemoryImageWriter.h"
+#include "Serialization/MemoryImageWriter.h"
+#include "Serialization/MemoryLayout.h"
+#include "Templates/ChooseClass.h"
+#include "Templates/UnrealTemplate.h"
+
 #include <initializer_list>
 
+class FMemoryImageWriter;
+class FMemoryUnfreezeContent;
+class FPointerTableBase;
+class FSHA1;
+
 static FORCEINLINE uint32 MurmurFinalize32(uint32 Hash)
 {
 	Hash ^= Hash >> 16;
diff --git a/Engine/Source/Runtime/Core/Public/Containers/LockFreeList.h b/Engine/Source/Runtime/Core/Public/Containers/LockFreeList.h
--- a/Engine/Source/Runtime/Core/Public/Containers/LockFreeList.h
+++ b/Engine/Source/Runtime/Core/Public/Containers/LockFreeList.h
@@ -2,15 +2,18 @@
 
 #pragma once
 
-#include "CoreTypes.h"
-#include "Misc/AssertionMacros.h"
-#include "Logging/LogMacros.h"
+#include "Containers/Array.h"
+#include "Containers/ContainersFwd.h"
 #include "CoreGlobals.h"
+#include "CoreTypes.h"
+#include "HAL/PlatformAtomics.h"
+#include "HAL/PlatformProcess.h"
 #include "HAL/ThreadSafeCounter.h"
+#include "Logging/LogMacros.h"
+#include "Misc/AssertionMacros.h"
 #include "Misc/NoopCounter.h"
-#include "Containers/ContainersFwd.h"
-#include "HAL/PlatformProcess.h"
 #include "Templates/AlignmentTemplates.h"
+#include "Templates/Function.h"
 
 CORE_API DECLARE_LOG_CATEGORY_EXTERN(LogLockFreeList, Log, All);
 
diff --git a/Engine/Source/Runtime/Core/Public/Containers/StackTracker.h b/Engine/Source/Runtime/Core/Public/Containers/StackTracker.h
--- a/Engine/Source/Runtime/Core/Public/Containers/StackTracker.h
+++ b/Engine/Source/Runtime/Core/Public/Containers/StackTracker.h
@@ -2,11 +2,14 @@
 
 #pragma once
 
-#include "CoreTypes.h"
-#include "Math/NumericLimits.h"
 #include "Containers/Array.h"
-#include "Containers/UnrealString.h"
 #include "Containers/Map.h"
+#include "Containers/UnrealString.h"
+#include "CoreTypes.h"
+#include "Math/NumericLimits.h"
+
+class FOutputDevice;
+class FString;
 
 
 /**
diff --git a/Engine/Source/Runtime/Core/Public/Containers/Ticker.h b/Engine/Source/Runtime/Core/Public/Containers/Ticker.h
--- a/Engine/Source/Runtime/Core/Public/Containers/Ticker.h
+++ b/Engine/Source/Runtime/Core/Public/Containers/Ticker.h
@@ -2,12 +2,15 @@
 
 #pragma once
 
-#include "CoreTypes.h"
 #include "Containers/Array.h"
-#include "Delegates/IDelegateInstance.h"
+#include "Containers/ContainerAllocationPolicies.h"
+#include "Containers/MpscQueue.h"
+#include "CoreTypes.h"
 #include "Delegates/Delegate.h"
+#include "Delegates/DelegateBase.h"
+#include "Delegates/IDelegateInstance.h"
+#include "Templates/Function.h"
 #include "Templates/SharedPointer.h"
-#include "Containers/MpscQueue.h"
 
 #include <atomic>
 
diff --git a/Engine/Source/Runtime/Core/Public/Containers/Union.h b/Engine/Source/Runtime/Core/Public/Containers/Union.h
--- a/Engine/Source/Runtime/Core/Public/Containers/Union.h
+++ b/Engine/Source/Runtime/Core/Public/Containers/Union.h
@@ -3,12 +3,15 @@
 #pragma once
 
 #include "CoreTypes.h"
+#include "HAL/PlatformCrt.h"
+#include "Logging/LogCategory.h"
+#include "Logging/LogMacros.h"
 #include "Misc/AssertionMacros.h"
+#include "Serialization/Archive.h"
 #include "Templates/AreTypesEqual.h"
-#include "Templates/UnrealTypeTraits.h"
 #include "Templates/TypeCompatibleBytes.h"
 #include "Templates/UnrealTemplate.h"
-#include "Logging/LogMacros.h"
+#include "Templates/UnrealTypeTraits.h"
 
 
 CORE_API DECLARE_LOG_CATEGORY_EXTERN(LogUnion, Log, All);
diff --git a/Engine/Source/Runtime/Core/Public/CoreGlobals.h b/Engine/Source/Runtime/Core/Public/CoreGlobals.h
--- a/Engine/Source/Runtime/Core/Public/CoreGlobals.h
+++ b/Engine/Source/Runtime/Core/Public/CoreGlobals.h
@@ -1,24 +1,29 @@
 // Copyright Epic Games, Inc. All Rights Reserved.
 #pragma once
 
-#include "CoreTypes.h"
 #include "Containers/UnrealString.h"
-#include "Misc/EnumClassFlags.h"
-#include "UObject/NameTypes.h"
-#include "Logging/LogMacros.h"
+#include "CoreTypes.h"
 #include "HAL/PlatformTLS.h"
-#include "Templates/Atomic.h"
+#include "Logging/LogMacros.h"
+#include "Misc/EnumClassFlags.h"
+#include "Misc/OutputDevice.h"
 #include "ProfilingDebugging/CpuProfilerTrace.h"
+#include "Templates/Atomic.h"
+#include "UObject/NameTypes.h"
 
 #include <atomic>
 
 class Error;
+class FChunkedFixedUObjectArray;
 class FConfigCacheIni;
+class FExec;
 class FFixedUObjectArray;
-class FChunkedFixedUObjectArray;
 class FOutputDeviceConsole;
 class FOutputDeviceRedirector;
+class FRunnableThread;
+class FText;
 class ITransaction;
+class UClass;
 
 CORE_API DECLARE_LOG_CATEGORY_EXTERN(LogHAL, Log, All);
 CORE_API DECLARE_LOG_CATEGORY_EXTERN(LogSerialization, Log, All);
@@ -739,6 +744,7 @@ void CORE_API SetEmitDrawEvents(bool EmitDrawEvents);
 
 /** Array to help visualize weak pointers in the debugger */
 class FChunkedFixedUObjectArray;
+
 extern CORE_API FChunkedFixedUObjectArray* GCoreObjectArrayForDebugVisualizers;
 
 /** Array to help visualize object paths in the debugger */
@@ -750,6 +756,7 @@ extern CORE_API UE::ObjectPath::Private::FStoredObjectPath* GCoreComplexObjectPa
 
 /** Array to help visualize object handles in the debugger */
 struct FObjectHandlePackageDebugData;
+
 extern CORE_API FObjectHandlePackageDebugData* GCoreObjectHandlePackageDebug;
 
 /** @return True if running cook-on-the-fly. */
diff --git a/Engine/Source/Runtime/Core/Public/Experimental/Coroutine/Coroutine.h b/Engine/Source/Runtime/Core/Public/Experimental/Coroutine/Coroutine.h
--- a/Engine/Source/Runtime/Core/Public/Experimental/Coroutine/Coroutine.h
+++ b/Engine/Source/Runtime/Core/Public/Experimental/Coroutine/Coroutine.h
@@ -1,17 +1,19 @@
 // Copyright Epic Games, Inc. All Rights Reserved.
 #pragma once
 
-#include <atomic>
-#include "CoreMinimal.h"
-#include "Templates/UnrealTemplate.h"
-#include "Templates/SharedPointer.h"
-#include "Misc/MemStack.h"
-#include "Async/TaskTrace.h"
-#include "ProfilingDebugging/CpuProfilerTrace.h"
 #include "Async/Fundamental/Scheduler.h"
-#include "CoroutineHandle.h"
+#include "Async/Fundamental/Task.h"
+#include "Async/TaskTrace.h"
+#include "CoreMinimal.h"
 #include "CoroLocalVariable.h"
+#include "CoroutineHandle.h"
 #include "Experimental/ConcurrentLinearAllocator.h"
+#include "Misc/MemStack.h"
+#include "ProfilingDebugging/CpuProfilerTrace.h"
+#include "Templates/SharedPointer.h"
+#include "Templates/UnrealTemplate.h"
+
+#include <atomic>
 
 #if WITH_CPP_COROUTINES
 
diff --git a/Engine/Source/Runtime/Core/Public/FileCache/FileCache.h b/Engine/Source/Runtime/Core/Public/FileCache/FileCache.h
--- a/Engine/Source/Runtime/Core/Public/FileCache/FileCache.h
+++ b/Engine/Source/Runtime/Core/Public/FileCache/FileCache.h
@@ -2,11 +2,14 @@
 
 #pragma once
 
+#include "Async/TaskGraphInterfaces.h"
 #include "CoreMinimal.h"
 #include "GenericPlatform/GenericPlatformFile.h"
-#include "Async/TaskGraphInterfaces.h"
+#include "HAL/PlatformFile.h"
 #include "Misc/MemoryReadStream.h"
 
+class IAsyncReadFileHandle;
+
 //
 struct FFileCachePreloadEntry
 {
diff --git a/Engine/Source/Runtime/Core/Public/Hash/Blake3.h b/Engine/Source/Runtime/Core/Public/Hash/Blake3.h
--- a/Engine/Source/Runtime/Core/Public/Hash/Blake3.h
+++ b/Engine/Source/Runtime/Core/Public/Hash/Blake3.h
@@ -4,14 +4,21 @@
 
 #include "Containers/StringFwd.h"
 #include "Containers/StringView.h"
+#include "Containers/UnrealString.h"
+#include "HAL/Platform.h"
+#include "HAL/PlatformString.h"
 #include "HAL/UnrealMemory.h"
 #include "Memory/MemoryFwd.h"
 #include "Memory/MemoryView.h"
+#include "Misc/AssertionMacros.h"
 #include "Serialization/Archive.h"
 #include "String/BytesToHex.h"
 #include "String/HexToBytes.h"
+#include "Templates/TypeCompatibleBytes.h"
 
 class FArchive;
+class FCompositeBuffer;
+template <typename CharType> class TStringBuilderBase;
 
 ///////////////////////////////////////////////////////////////////////////////////////////////////////////////////////
 
diff --git a/Engine/Source/Runtime/Core/Public/Hash/xxhash.h b/Engine/Source/Runtime/Core/Public/Hash/xxhash.h
--- a/Engine/Source/Runtime/Core/Public/Hash/xxhash.h
+++ b/Engine/Source/Runtime/Core/Public/Hash/xxhash.h
@@ -2,14 +2,18 @@
 
 #pragma once
 
-#include "CoreTypes.h"
+#include "Containers/ArrayView.h"
 #include "Containers/StringFwd.h"
+#include "CoreTypes.h"
 #include "HAL/UnrealMemory.h"
 #include "Memory/MemoryFwd.h"
 #include "Misc/ByteSwap.h"
 #include "Serialization/Archive.h"
 #include "String/BytesToHex.h"
 
+class FCompositeBuffer;
+template <typename CharType> class TStringBuilderBase;
+
 /** A 64-bit hash from XXH3. */
 struct FXxHash64
 {
diff --git a/Engine/Source/Runtime/Core/Public/IO/IoContainerHeader.h b/Engine/Source/Runtime/Core/Public/IO/IoContainerHeader.h
--- a/Engine/Source/Runtime/Core/Public/IO/IoContainerHeader.h
+++ b/Engine/Source/Runtime/Core/Public/IO/IoContainerHeader.h
@@ -2,11 +2,14 @@
 
 #pragma once
 
-#include "IO/PackageId.h"
+#include "CoreMinimal.h"
 #include "IO/IoContainerId.h"
+#include "IO/PackageId.h"
 #include "Serialization/MappedName.h"
+#include "UObject/NameBatchSerialization.h"
 
 class FArchive;
+class FSHAHash;
 
 /**
  * Package store entry array view.
diff --git a/Engine/Source/Runtime/Core/Public/IO/IoDispatcher.h b/Engine/Source/Runtime/Core/Public/IO/IoDispatcher.h
--- a/Engine/Source/Runtime/Core/Public/IO/IoDispatcher.h
+++ b/Engine/Source/Runtime/Core/Public/IO/IoDispatcher.h
@@ -2,41 +2,51 @@
 
 #pragma once
 
-#include "CoreTypes.h"
-#include "IO/IoContainerId.h"
+#include "Async/Future.h"
+#include "Async/TaskGraphInterfaces.h"
+#include "Containers/ArrayView.h"
+#include "Containers/StringFwd.h"
 #include "Containers/StringView.h"
 #include "Containers/UnrealString.h"
-#include "Logging/LogMacros.h"
-#include "Templates/RefCounting.h"
-#include "Templates/UnrealTemplate.h"
-#include "Templates/TypeCompatibleBytes.h"
+#include "CoreMinimal.h"
+#include "CoreTypes.h"
+#include "GenericPlatform/GenericPlatformFile.h"
 #include "HAL/PlatformAtomics.h"
-#include "Memory/MemoryView.h"
-#include "Misc/SecureHash.h"
+#include "HAL/PlatformFile.h"
+#include "IO/IoContainerId.h"
 #include "IO/IoHash.h"
+#include "Logging/LogMacros.h"
+#include "Memory/MemoryFwd.h"
+#include "Memory/MemoryView.h"
 #include "Misc/AES.h"
+#include "Misc/EnumClassFlags.h"
+#include "Misc/Guid.h"
 #include "Misc/IEngineCrypto.h"
+#include "Misc/SecureHash.h"
 #include "Serialization/FileRegions.h"
-#include "Async/TaskGraphInterfaces.h"
-#include "GenericPlatform/GenericPlatformFile.h"
-#include "Async/Future.h"
+#include "String/BytesToHex.h"
 #include "Tasks/Task.h"
+#include "Templates/RefCounting.h"
+#include "Templates/TypeCompatibleBytes.h"
+#include "Templates/UnrealTemplate.h"
 
-class FIoRequest;
-class FIoDispatcher;
-class FIoStoreReader;
-class FIoStoreEnvironment;
-
-class FIoRequestImpl;
+class FEvent;
 class FIoBatchImpl;
+class FIoDirectoryIndexReaderImpl;
+class FIoDispatcher;
 class FIoDispatcherImpl;
-class FIoStoreWriterContextImpl;
+class FIoRequest;
+class FIoRequestImpl;
+class FIoStoreEnvironment;
+class FIoStoreReader;
 class FIoStoreReaderImpl;
+class FIoStoreWriterContextImpl;
+class FPackageId;
 class IMappedFileHandle;
 class IMappedFileRegion;
-class FIoDirectoryIndexReaderImpl;
+struct FFileRegion;
 struct IIoDispatcherBackend;
-class FPackageId;
+template <typename CharType> class TStringBuilderBase;
 
 CORE_API DECLARE_LOG_CATEGORY_EXTERN(LogIoDispatcher, Log, All);
 
diff --git a/Engine/Source/Runtime/Core/Public/IO/IoHash.h b/Engine/Source/Runtime/Core/Public/IO/IoHash.h
--- a/Engine/Source/Runtime/Core/Public/IO/IoHash.h
+++ b/Engine/Source/Runtime/Core/Public/IO/IoHash.h
@@ -4,14 +4,21 @@
 
 #include "Containers/StringFwd.h"
 #include "Containers/StringView.h"
+#include "Containers/UnrealString.h"
+#include "HAL/Platform.h"
+#include "HAL/PlatformString.h"
 #include "HAL/UnrealMemory.h"
 #include "Hash/Blake3.h"
 #include "Memory/MemoryFwd.h"
 #include "Memory/MemoryView.h"
+#include "Misc/AssertionMacros.h"
 #include "Serialization/Archive.h"
 #include "String/BytesToHex.h"
 #include "String/HexToBytes.h"
 
+class FCompositeBuffer;
+template <typename CharType> class TStringBuilderBase;
+
 /**
  * Stores a BLAKE3-160 hash, taken from the first 20 bytes of a BLAKE3-256 hash.
  *
diff --git a/Engine/Source/Runtime/Core/Public/IO/PackageStore.h b/Engine/Source/Runtime/Core/Public/IO/PackageStore.h
--- a/Engine/Source/Runtime/Core/Public/IO/PackageStore.h
+++ b/Engine/Source/Runtime/Core/Public/IO/PackageStore.h
@@ -2,14 +2,16 @@
 
 #pragma once
 
+#include "Containers/ArrayView.h"
 #include "CoreMinimal.h"
-#include "PackageId.h"
+#include "Misc/EnumClassFlags.h"
 #include "Misc/SecureHash.h"
+#include "PackageId.h"
 
 class FArchive;
-class FStructuredArchiveSlot;
 class FCbObject;
 class FCbWriter;
+class FStructuredArchiveSlot;
 
 /**
  * Package export information.
diff --git a/Engine/Source/Runtime/Core/Public/Internationalization/FastDecimalFormat.h b/Engine/Source/Runtime/Core/Public/Internationalization/FastDecimalFormat.h
--- a/Engine/Source/Runtime/Core/Public/Internationalization/FastDecimalFormat.h
+++ b/Engine/Source/Runtime/Core/Public/Internationalization/FastDecimalFormat.h
@@ -2,10 +2,11 @@
 
 #pragma once
 
-#include "CoreTypes.h"
 #include "Containers/UnrealString.h"
+#include "CoreTypes.h"
 #include "Internationalization/Text.h"
 #include "Math/NumericLimits.h"
+#include "Misc/CString.h"
 #include "Templates/EnableIf.h"
 #include "Templates/IsFloatingPoint.h"
 #include "Templates/IsIntegral.h"
diff --git a/Engine/Source/Runtime/Core/Public/Internationalization/GatherableTextData.h b/Engine/Source/Runtime/Core/Public/Internationalization/GatherableTextData.h
--- a/Engine/Source/Runtime/Core/Public/Internationalization/GatherableTextData.h
+++ b/Engine/Source/Runtime/Core/Public/Internationalization/GatherableTextData.h
@@ -1,10 +1,13 @@
 // Copyright Epic Games, Inc. All Rights Reserved.
 #pragma once
  
-#include "CoreTypes.h"
 #include "Containers/Array.h"
 #include "Containers/UnrealString.h"
+#include "CoreTypes.h"
 #include "Internationalization/InternationalizationMetadata.h"
+#include "Serialization/StructuredArchive.h"
+
+class FArchive;
 
 struct FTextSourceSiteContext
 {
diff --git a/Engine/Source/Runtime/Core/Public/Internationalization/Internationalization.h b/Engine/Source/Runtime/Core/Public/Internationalization/Internationalization.h
--- a/Engine/Source/Runtime/Core/Public/Internationalization/Internationalization.h
+++ b/Engine/Source/Runtime/Core/Public/Internationalization/Internationalization.h
@@ -1,22 +1,25 @@
 // Copyright Epic Games, Inc. All Rights Reserved.
 #pragma once
 
-#include "CoreTypes.h"
 #include "Containers/Array.h"
 #include "Containers/UnrealString.h"
-#include "Templates/SharedPointer.h"
+#include "CoreTypes.h"
 #include "Delegates/Delegate.h"
-#include "Internationalization/Text.h"
+#include "Delegates/DelegateBase.h"
+#include "Internationalization/CulturePointer.h"
 #include "Internationalization/LocTesting.h"
-
+#include "Internationalization/Text.h"
+#include "Templates/SharedPointer.h"
+#include "Templates/Tuple.h"
 #include "Templates/UniqueObj.h"
+#include "UObject/NameTypes.h"
 
 #define LOC_DEFINE_REGION
 
 class FCulture;
-class ICustomCulture;
 class FICUInternationalization;
 class FLegacyInternationalization;
+class ICustomCulture;
 
 class FInternationalization
 {
diff --git a/Engine/Source/Runtime/Core/Public/Internationalization/InternationalizationArchive.h b/Engine/Source/Runtime/Core/Public/Internationalization/InternationalizationArchive.h
--- a/Engine/Source/Runtime/Core/Public/Internationalization/InternationalizationArchive.h
+++ b/Engine/Source/Runtime/Core/Public/Internationalization/InternationalizationArchive.h
@@ -2,15 +2,15 @@
 
 #pragma once
 
-#include "CoreTypes.h"
 #include "Containers/ContainerAllocationPolicies.h"
-#include "Misc/Crc.h"
-#include "Containers/UnrealString.h"
-#include "Containers/Set.h"
 #include "Containers/Map.h"
-#include "Templates/SharedPointer.h"
+#include "Containers/Set.h"
+#include "Containers/UnrealString.h"
+#include "CoreTypes.h"
 #include "Internationalization/InternationalizationManifest.h"
 #include "Internationalization/LocKeyFuncs.h"
+#include "Misc/Crc.h"
+#include "Templates/SharedPointer.h"
 
 class FLocMetadataObject;
 
diff --git a/Engine/Source/Runtime/Core/Public/Internationalization/InternationalizationManifest.h b/Engine/Source/Runtime/Core/Public/Internationalization/InternationalizationManifest.h
--- a/Engine/Source/Runtime/Core/Public/Internationalization/InternationalizationManifest.h
+++ b/Engine/Source/Runtime/Core/Public/Internationalization/InternationalizationManifest.h
@@ -1,15 +1,17 @@
 // Copyright Epic Games, Inc. All Rights Reserved.
 #pragma once
 
-#include "CoreTypes.h"
-#include "Containers/ContainerAllocationPolicies.h"
 #include "Containers/Array.h"
-#include "Misc/Crc.h"
-#include "Containers/UnrealString.h"
-#include "Containers/Set.h"
+#include "Containers/ContainerAllocationPolicies.h"
 #include "Containers/Map.h"
-#include "Templates/SharedPointer.h"
+#include "Containers/Set.h"
+#include "Containers/UnrealString.h"
+#include "CoreTypes.h"
 #include "Internationalization/LocKeyFuncs.h"
+#include "Misc/Crc.h"
+#include "Templates/SharedPointer.h"
+#include "Templates/UnrealTemplate.h"
+#include "UObject/NameTypes.h"
 
 class FLocMetadataObject;
 
diff --git a/Engine/Source/Runtime/Core/Public/Internationalization/InternationalizationMetadata.h b/Engine/Source/Runtime/Core/Public/Internationalization/InternationalizationMetadata.h
--- a/Engine/Source/Runtime/Core/Public/Internationalization/InternationalizationMetadata.h
+++ b/Engine/Source/Runtime/Core/Public/Internationalization/InternationalizationMetadata.h
@@ -1,13 +1,15 @@
 // Copyright Epic Games, Inc. All Rights Reserved.
 #pragma once
 
-#include "CoreTypes.h"
-#include "Misc/AssertionMacros.h"
 #include "Containers/Array.h"
-#include "Containers/UnrealString.h"
 #include "Containers/Map.h"
+#include "Containers/UnrealString.h"
+#include "CoreTypes.h"
+#include "Misc/AssertionMacros.h"
+#include "Serialization/StructuredArchive.h"
 #include "Templates/SharedPointer.h"
 
+class FArchive;
 class FLocMetadataObject;
 
 /**
diff --git a/Engine/Source/Runtime/Core/Public/Internationalization/PolyglotTextData.h b/Engine/Source/Runtime/Core/Public/Internationalization/PolyglotTextData.h
--- a/Engine/Source/Runtime/Core/Public/Internationalization/PolyglotTextData.h
+++ b/Engine/Source/Runtime/Core/Public/Internationalization/PolyglotTextData.h
@@ -2,10 +2,11 @@
 
 #pragma once
 
+#include "Containers/Array.h"
+#include "Containers/Map.h"
+#include "Containers/UnrealString.h"
 #include "CoreTypes.h"
 #include "Internationalization/LocalizedTextSourceTypes.h"
-#include "Containers/UnrealString.h"
-#include "Containers/Map.h"
 #include "Internationalization/Text.h"
 
 /**
diff --git a/Engine/Source/Runtime/Core/Public/Internationalization/StringTableCore.h b/Engine/Source/Runtime/Core/Public/Internationalization/StringTableCore.h
--- a/Engine/Source/Runtime/Core/Public/Internationalization/StringTableCore.h
+++ b/Engine/Source/Runtime/Core/Public/Internationalization/StringTableCore.h
@@ -1,14 +1,30 @@
 // Copyright Epic Games, Inc. All Rights Reserved.
 #pragma once
 
-#include <atomic>
-#include "Internationalization/StringTableCoreFwd.h"
 #include "Containers/Map.h"
-#include "UObject/NameTypes.h"
+#include "Containers/SparseArray.h"
 #include "Containers/UnrealString.h"
-#include "Internationalization/Text.h"
-#include "Internationalization/LocKeyFuncs.h"
+#include "HAL/CriticalSection.h"
 #include "Internationalization/Internationalization.h"
+#include "Internationalization/LocKeyFuncs.h"
+#include "Internationalization/LocalizedTextSourceTypes.h"
+#include "Internationalization/StringTableCoreFwd.h"
+#include "Internationalization/Text.h"
+#include "Internationalization/TextKey.h"
+#include "Logging/LogMacros.h"
+#include "Misc/AssertionMacros.h"
+#include "Serialization/Archive.h"
+#include "Serialization/StructuredArchive.h"
+#include "Templates/Function.h"
+#include "Templates/SharedPointer.h"
+#include "Templates/SharedPointerInternals.h"
+#include "Templates/UnrealTemplate.h"
+#include "UObject/NameTypes.h"
+
+#include <atomic>
+
+class FStringTable;
+class UStringTable;
 
 DECLARE_LOG_CATEGORY_EXTERN(LogStringTable, Log, All);
 
diff --git a/Engine/Source/Runtime/Core/Public/Internationalization/StringTableRegistry.h b/Engine/Source/Runtime/Core/Public/Internationalization/StringTableRegistry.h
--- a/Engine/Source/Runtime/Core/Public/Internationalization/StringTableRegistry.h
+++ b/Engine/Source/Runtime/Core/Public/Internationalization/StringTableRegistry.h
@@ -1,13 +1,17 @@
 // Copyright Epic Games, Inc. All Rights Reserved.
 #pragma once
 
-#include "Internationalization/StringTableCoreFwd.h"
 #include "Containers/Map.h"
-#include "UObject/NameTypes.h"
-#include "Delegates/Delegate.h"
 #include "Containers/UnrealString.h"
-#include "Internationalization/TextLocalizationManager.h"
+#include "CoreMinimal.h"
+#include "Delegates/Delegate.h"
 #include "Internationalization/LocKeyFuncs.h"
+#include "Internationalization/StringTableCoreFwd.h"
+#include "Internationalization/TextKey.h"
+#include "Internationalization/TextLocalizationManager.h"
+#include "UObject/NameTypes.h"
+
+class UStringTable;
 
 /** Singleton registry of known string table instances */
 class CORE_API FStringTableRegistry
diff --git a/Engine/Source/Runtime/Core/Public/Internationalization/TextFormatter.h b/Engine/Source/Runtime/Core/Public/Internationalization/TextFormatter.h
--- a/Engine/Source/Runtime/Core/Public/Internationalization/TextFormatter.h
+++ b/Engine/Source/Runtime/Core/Public/Internationalization/TextFormatter.h
@@ -2,14 +2,17 @@
 
 #pragma once
 
-#include "CoreTypes.h"
 #include "Containers/Array.h"
-#include "Containers/UnrealString.h"
-#include "Templates/Function.h"
 #include "Containers/SortedMap.h"
+#include "Containers/UnrealString.h"
+#include "CoreTypes.h"
+#include "HAL/CriticalSection.h"
+#include "Internationalization/ITextFormatArgumentModifier.h"
 #include "Internationalization/Text.h"
 #include "Misc/ExpressionParserTypes.h"
-#include "Internationalization/ITextFormatArgumentModifier.h"
+#include "Templates/Function.h"
+#include "Templates/SharedPointer.h"
+#include "Templates/SharedPointerInternals.h"
 
 struct FPrivateTextFormatArguments;
 
diff --git a/Engine/Source/Runtime/Core/Public/Internationalization/TextGeneratorBase.h b/Engine/Source/Runtime/Core/Public/Internationalization/TextGeneratorBase.h
--- a/Engine/Source/Runtime/Core/Public/Internationalization/TextGeneratorBase.h
+++ b/Engine/Source/Runtime/Core/Public/Internationalization/TextGeneratorBase.h
@@ -3,6 +3,8 @@
 
 #include "CoreTypes.h"
 #include "ITextGenerator.h"
+#include "Serialization/StructuredArchive.h"
+#include "UObject/NameTypes.h"
 
 /**
  * Base class implementation for ITextGenerator.
diff --git a/Engine/Source/Runtime/Core/Public/Internationalization/TextLocalizationManager.h b/Engine/Source/Runtime/Core/Public/Internationalization/TextLocalizationManager.h
--- a/Engine/Source/Runtime/Core/Public/Internationalization/TextLocalizationManager.h
+++ b/Engine/Source/Runtime/Core/Public/Internationalization/TextLocalizationManager.h
@@ -2,27 +2,29 @@
 
 #pragma once
 
-#include "CoreTypes.h"
-#include "Containers/ContainerAllocationPolicies.h"
 #include "Containers/Array.h"
 #include "Containers/ArrayView.h"
-#include "Misc/Crc.h"
-#include "Containers/UnrealString.h"
-#include "Containers/Set.h"
+#include "Containers/ContainerAllocationPolicies.h"
 #include "Containers/Map.h"
-#include "Templates/SharedPointer.h"
+#include "Containers/Set.h"
+#include "Containers/UnrealString.h"
+#include "CoreTypes.h"
 #include "Delegates/Delegate.h"
-#include "Internationalization/TextKey.h"
-#include "Internationalization/LocTesting.h"
+#include "HAL/CriticalSection.h"
 #include "Internationalization/LocKeyFuncs.h"
+#include "Internationalization/LocTesting.h"
 #include "Internationalization/LocalizedTextSourceTypes.h"
+#include "Internationalization/TextKey.h"
+#include "Misc/Crc.h"
+#include "Misc/EnumClassFlags.h"
+#include "Templates/SharedPointer.h"
 
 #include <atomic>
 
-struct FPolyglotTextData;
+class FTextLocalizationResource;
 class ILocalizedTextSource;
 class IPakFile;
-class FTextLocalizationResource;
+struct FPolyglotTextData;
 
 enum class ETextLocalizationManagerInitializedFlags : uint8
 {
diff --git a/Engine/Source/Runtime/Core/Public/Internationalization/TextLocalizationResource.h b/Engine/Source/Runtime/Core/Public/Internationalization/TextLocalizationResource.h
--- a/Engine/Source/Runtime/Core/Public/Internationalization/TextLocalizationResource.h
+++ b/Engine/Source/Runtime/Core/Public/Internationalization/TextLocalizationResource.h
@@ -2,10 +2,17 @@
 
 #pragma once 
 
-#include "CoreTypes.h"
+#include "Containers/Array.h"
+#include "Containers/Map.h"
 #include "Containers/SortedMap.h"
-#include "Internationalization/TextKey.h"
+#include "Containers/UnrealString.h"
+#include "CoreTypes.h"
 #include "Internationalization/LocalizedTextSourceTypes.h"
+#include "Internationalization/TextKey.h"
+#include "Misc/Crc.h"
+#include "Templates/SharedPointer.h"
+
+class FArchive;
 
 /** Utility class for working with Localization MetaData Resource (LocMeta) files. */
 class CORE_API FTextLocalizationMetaDataResource
diff --git a/Engine/Source/Runtime/Core/Public/Internationalization/TextNamespaceUtil.h b/Engine/Source/Runtime/Core/Public/Internationalization/TextNamespaceUtil.h
--- a/Engine/Source/Runtime/Core/Public/Internationalization/TextNamespaceUtil.h
+++ b/Engine/Source/Runtime/Core/Public/Internationalization/TextNamespaceUtil.h
@@ -2,10 +2,12 @@
 
 #pragma once
 
+#include "Containers/UnrealString.h"
 #include "CoreTypes.h"
 #include "Internationalization/Text.h"
 #include "Internationalization/TextNamespaceFwd.h"
-#include "Containers/UnrealString.h"
+
+class FArchive;
 
 namespace TextNamespaceUtil
 {
diff --git a/Engine/Source/Runtime/Core/Public/Logging/LogScopedCategoryAndVerbosityOverride.h b/Engine/Source/Runtime/Core/Public/Logging/LogScopedCategoryAndVerbosityOverride.h
--- a/Engine/Source/Runtime/Core/Public/Logging/LogScopedCategoryAndVerbosityOverride.h
+++ b/Engine/Source/Runtime/Core/Public/Logging/LogScopedCategoryAndVerbosityOverride.h
@@ -3,6 +3,7 @@
 #pragma once
 
 #include "CoreTypes.h"
+#include "Logging/LogVerbosity.h"
 #include "UObject/NameTypes.h"
 
 /** 
diff --git a/Engine/Source/Runtime/Core/Public/Logging/LogTrace.h b/Engine/Source/Runtime/Core/Public/Logging/LogTrace.h
--- a/Engine/Source/Runtime/Core/Public/Logging/LogTrace.h
+++ b/Engine/Source/Runtime/Core/Public/Logging/LogTrace.h
@@ -3,10 +3,14 @@
 #pragma once
 
 #include "CoreTypes.h"
-#include "Trace/Config.h"
+#include "Logging/LogVerbosity.h"
 #include "Templates/IsArrayOrRefOfTypeByPredicate.h"
+#include "Trace/Config.h"
+#include "Trace/Trace.h"
 #include "Traits/IsCharEncodingCompatibleWith.h"
 
+namespace UE { namespace Trace { class FChannel; } }
+
 #if !defined(LOGTRACE_ENABLED)
 #if UE_TRACE_ENABLED && !UE_BUILD_SHIPPING
 #define LOGTRACE_ENABLED 1
diff --git a/Engine/Source/Runtime/Core/Public/Logging/MessageLog.h b/Engine/Source/Runtime/Core/Public/Logging/MessageLog.h
--- a/Engine/Source/Runtime/Core/Public/Logging/MessageLog.h
+++ b/Engine/Source/Runtime/Core/Public/Logging/MessageLog.h
@@ -2,12 +2,15 @@
 
 #pragma once
 
-#include "CoreTypes.h"
 #include "Containers/Array.h"
-#include "Templates/SharedPointer.h"
+#include "CoreTypes.h"
 #include "Delegates/Delegate.h"
+#include "Delegates/DelegateBase.h"
 #include "Internationalization/Text.h"
+#include "Logging/LogVerbosity.h"
 #include "Logging/TokenizedMessage.h"
+#include "Templates/SharedPointer.h"
+#include "UObject/NameTypes.h"
 
 class IMessageLog;
 
diff --git a/Engine/Source/Runtime/Core/Public/Logging/TokenizedMessage.h b/Engine/Source/Runtime/Core/Public/Logging/TokenizedMessage.h
--- a/Engine/Source/Runtime/Core/Public/Logging/TokenizedMessage.h
+++ b/Engine/Source/Runtime/Core/Public/Logging/TokenizedMessage.h
@@ -2,15 +2,17 @@
 
 #pragma once
 
-#include "CoreTypes.h"
 #include "Containers/Array.h"
 #include "Containers/UnrealString.h"
-#include "UObject/NameTypes.h"
-#include "Templates/SharedPointer.h"
+#include "CoreTypes.h"
 #include "Delegates/Delegate.h"
+#include "Delegates/DelegateBase.h"
 #include "Internationalization/Text.h"
 #include "Misc/Attribute.h"
 #include "Misc/Guid.h"
+#include "Templates/SharedPointer.h"
+#include "UObject/NameTypes.h"
+#include "UObject/UnrealNames.h"
 
 /** The severity of the message type */
 namespace EMessageSeverity
diff --git a/Engine/Source/Runtime/Core/Public/Math/BasicMathExpressionEvaluator.h b/Engine/Source/Runtime/Core/Public/Math/BasicMathExpressionEvaluator.h
--- a/Engine/Source/Runtime/Core/Public/Math/BasicMathExpressionEvaluator.h
+++ b/Engine/Source/Runtime/Core/Public/Math/BasicMathExpressionEvaluator.h
@@ -3,10 +3,12 @@
 #pragma once
 
 #include "CoreTypes.h"
+#include "Internationalization/FastDecimalFormat.h"
+#include "Misc/ExpressionParserTypes.h"
 #include "Misc/Optional.h"
 #include "Templates/ValueOrError.h"
-#include "Misc/ExpressionParserTypes.h"
-#include "Internationalization/FastDecimalFormat.h"
+
+struct FDecimalNumberFormattingRules;
 
 #define DEFINE_EXPRESSION_OPERATOR_NODE(EXPORTAPI, TYPE, ...) \
 namespace ExpressionParser {\
diff --git a/Engine/Source/Runtime/Core/Public/Math/Box2D.h b/Engine/Source/Runtime/Core/Public/Math/Box2D.h
--- a/Engine/Source/Runtime/Core/Public/Math/Box2D.h
+++ b/Engine/Source/Runtime/Core/Public/Math/Box2D.h
@@ -2,12 +2,20 @@
 
 #pragma once
 
+#include "Containers/Array.h"
+#include "Containers/UnrealString.h"
 #include "CoreTypes.h"
-#include "Misc/AssertionMacros.h"
-#include "Misc/LargeWorldCoordinatesSerializer.h"
 #include "Math/UnrealMathUtility.h"
-#include "Containers/UnrealString.h"
 #include "Math/Vector2D.h"
+#include "Misc/AssertionMacros.h"
+#include "Misc/LargeWorldCoordinatesSerializer.h"
+#include "Serialization/Archive.h"
+#include "Templates/IsUECoreType.h"
+#include "Templates/UnrealTypeTraits.h"
+#include "UObject/NameTypes.h"
+#include "UObject/UnrealNames.h"
+
+template <typename T> struct TIsPODType;
 
 /**
  * Implements a rectangular 2D Box.
diff --git a/Engine/Source/Runtime/Core/Public/Math/Color.h b/Engine/Source/Runtime/Core/Public/Math/Color.h
--- a/Engine/Source/Runtime/Core/Public/Math/Color.h
+++ b/Engine/Source/Runtime/Core/Public/Math/Color.h
@@ -2,15 +2,27 @@
 
 #pragma once
 
+#include "Containers/Array.h"
+#include "Containers/UnrealString.h"
 #include "CoreTypes.h"
-#include "Misc/Crc.h"
+#include "Math/MathFwd.h"
 #include "Math/UnrealMathUtility.h"
-#include "Containers/UnrealString.h"
+#include "Misc/AssertionMacros.h"
+#include "Misc/Crc.h"
 #include "Misc/Parse.h"
-#include "Serialization/StructuredArchive.h"
+#include "Serialization/Archive.h"
 #include "Serialization/MemoryLayout.h"
+#include "Serialization/StructuredArchive.h"
+#include "Serialization/StructuredArchiveNameHelpers.h"
+#include "Serialization/StructuredArchiveSlots.h"
 
 class FFloat16Color;
+class FMemoryImageWriter;
+class FMemoryUnfreezeContent;
+class FPointerTableBase;
+class FSHA1;
+struct FColor;
+template <typename T> struct TIsPODType;
 
 /**
  * Enum for the different kinds of gamma spaces we expect to need to convert from/to.
diff --git a/Engine/Source/Runtime/Core/Public/Math/ColorList.h b/Engine/Source/Runtime/Core/Public/Math/ColorList.h
--- a/Engine/Source/Runtime/Core/Public/Math/ColorList.h
+++ b/Engine/Source/Runtime/Core/Public/Math/ColorList.h
@@ -2,12 +2,14 @@
 
 #pragma once
 
-#include "CoreTypes.h"
-#include "Misc/Exec.h"
 #include "Containers/Array.h"
-#include "Containers/UnrealString.h"
 #include "Containers/Map.h"
+#include "Containers/UnrealString.h"
+#include "CoreTypes.h"
 #include "Math/Color.h"
+#include "Misc/Exec.h"
+
+class FString;
 
 class CORE_API FColorList
 {
diff --git a/Engine/Source/Runtime/Core/Public/Math/GenericOctree.h b/Engine/Source/Runtime/Core/Public/Math/GenericOctree.h
--- a/Engine/Source/Runtime/Core/Public/Math/GenericOctree.h
+++ b/Engine/Source/Runtime/Core/Public/Math/GenericOctree.h
@@ -6,8 +6,11 @@
 
 #pragma once
 
+#include "Containers/ArrayView.h"
 #include "CoreMinimal.h"
 #include "GenericOctreePublic.h"
+#include "Math/MathFwd.h"
+#include "Math/UnrealMathSSE.h"
 #include "Templates/Models.h"
 
 /** A concise iteration over the children of an octree node. */
diff --git a/Engine/Source/Runtime/Core/Public/Math/SHMath.h b/Engine/Source/Runtime/Core/Public/Math/SHMath.h
--- a/Engine/Source/Runtime/Core/Public/Math/SHMath.h
+++ b/Engine/Source/Runtime/Core/Public/Math/SHMath.h
@@ -4,12 +4,16 @@
 
 #include "CoreTypes.h"
 #include "HAL/UnrealMemory.h"
-#include "Math/UnrealMathUtility.h"
 #include "Math/Color.h"
+#include "Math/MathFwd.h"
+#include "Math/UnrealMathSSE.h"
+#include "Math/UnrealMathUtility.h"
 #include "Math/Vector.h"
 #include "Math/Vector4.h"
 #include "Math/VectorRegister.h"
 
+class FArchive;
+
 //	Constants.
 extern CORE_API float NormalizationConstants[9];
 extern CORE_API int32 BasisL[9];
diff --git a/Engine/Source/Runtime/Core/Public/Math/Sobol.h b/Engine/Source/Runtime/Core/Public/Math/Sobol.h
--- a/Engine/Source/Runtime/Core/Public/Math/Sobol.h
+++ b/Engine/Source/Runtime/Core/Public/Math/Sobol.h
@@ -3,8 +3,9 @@
 #pragma once
 
 #include "CoreTypes.h"
-#include "Math/Vector2D.h"
+#include "Math/MathFwd.h"
 #include "Math/Vector.h"
+#include "Math/Vector2D.h"
 
 /**
  * Support for Sobol quasi-random numbers
diff --git a/Engine/Source/Runtime/Core/Public/Math/Sphere.h b/Engine/Source/Runtime/Core/Public/Math/Sphere.h
--- a/Engine/Source/Runtime/Core/Public/Math/Sphere.h
+++ b/Engine/Source/Runtime/Core/Public/Math/Sphere.h
@@ -3,16 +3,25 @@
 #pragma once
 
 #include "CoreTypes.h"
-#include "Math/UnrealMathUtility.h"
-#include "Math/Vector.h"
+#include "Math/MathFwd.h"
 #include "Math/Matrix.h"
 #include "Math/Transform.h"
+#include "Math/UnrealMathUtility.h"
+#include "Math/Vector.h"
+#include "Math/Vector4.h"
+#include "Misc/AssertionMacros.h"
+#include "Serialization/Archive.h"
+#include "Serialization/StructuredArchiveAdapters.h"
+#include "Templates/IsUECoreType.h"
+#include "Templates/UnrealTypeTraits.h"
+#include "UObject/ObjectVersion.h"
 
 /**
  * Implements a basic sphere.
  */
 namespace UE {
 namespace Math {
+template <typename T> struct TMatrix;
 
 template<typename T>
 struct TSphere
diff --git a/Engine/Source/Runtime/Core/Public/Math/Transform.h b/Engine/Source/Runtime/Core/Public/Math/Transform.h
--- a/Engine/Source/Runtime/Core/Public/Math/Transform.h
+++ b/Engine/Source/Runtime/Core/Public/Math/Transform.h
@@ -9,11 +9,23 @@
 #include "CoreTypes.h"
 #include "Math/Quat.h"
 #include "Math/ScalarRegister.h"
-#include "Math/TransformVectorized.h"
 #include "Math/TransformNonVectorized.h"
+#include "Math/TransformVectorized.h"
 #include "Misc/LargeWorldCoordinatesSerializer.h"
+#include "Serialization/MemoryLayout.h"
+#include "Serialization/StructuredArchiveAdapters.h"
+#include "Templates/IsUECoreType.h"
+#include "UObject/NameTypes.h"
+#include "UObject/UnrealNames.h"
 
 UE_DECLARE_LWC_TYPE(Transform, 3);
+class FArchive;
+class FMemoryImageWriter;
+class FMemoryUnfreezeContent;
+class FPointerTableBase;
+class FSHA1;
+template <typename T> struct TCanBulkSerialize;
+template <typename T> struct TIsPODType;
 
 template<> struct TIsPODType<FTransform3f> { enum { Value = true }; };
 template<> struct TIsPODType<FTransform3d> { enum { Value = true }; };
diff --git a/Engine/Source/Runtime/Core/Public/Math/TransformVectorized.h b/Engine/Source/Runtime/Core/Public/Math/TransformVectorized.h
--- a/Engine/Source/Runtime/Core/Public/Math/TransformVectorized.h
+++ b/Engine/Source/Runtime/Core/Public/Math/TransformVectorized.h
@@ -2,23 +2,33 @@
 
 #pragma once
 
-#include "CoreTypes.h"
-#include "Misc/AssertionMacros.h"
-#include "Math/UnrealMathUtility.h"
 #include "Containers/UnrealString.h"
+#include "CoreTypes.h"
 #include "Logging/LogMacros.h"
-#include "Math/Vector.h"
-#include "Math/VectorRegister.h"
-#include "Math/Rotator.h"
-#include "Math/Vector4.h"
+#include "Math/Axis.h"
+#include "Math/MathFwd.h"
 #include "Math/Matrix.h"
+#include "Math/Matrix.inl"
 #include "Math/Quat.h"
+#include "Math/Rotator.h"
 #include "Math/ScalarRegister.h"
+#include "Math/UnrealMathSSE.h"
+#include "Math/UnrealMathUtility.h"
+#include "Math/UnrealMathVectorConstants.h"
+#include "Math/Vector.h"
+#include "Math/Vector4.h"
+#include "Math/VectorRegister.h"
+#include "Misc/AssertionMacros.h"
+#include "Serialization/Archive.h"
+#include "Serialization/StructuredArchiveAdapters.h"
+#include "Templates/IsFloatingPoint.h"
+#include "Templates/UnrealTypeTraits.h"
+#include "UObject/NameTypes.h"
 
 #if ENABLE_VECTORIZED_TRANSFORM
 
-struct Z_Construct_UScriptStruct_FTransform3f_Statics;
 struct Z_Construct_UScriptStruct_FTransform3d_Statics;
+struct Z_Construct_UScriptStruct_FTransform3f_Statics;
 struct Z_Construct_UScriptStruct_FTransform_Statics;
 
 namespace UE
diff --git a/Engine/Source/Runtime/Core/Public/Math/UnitConversion.h b/Engine/Source/Runtime/Core/Public/Math/UnitConversion.h
--- a/Engine/Source/Runtime/Core/Public/Math/UnitConversion.h
+++ b/Engine/Source/Runtime/Core/Public/Math/UnitConversion.h
@@ -2,14 +2,16 @@
 
 #pragma once
 
-#include "CoreTypes.h"
 #include "Containers/Array.h"
 #include "Containers/UnrealString.h"
+#include "CoreTypes.h"
 #include "Delegates/Delegate.h"
-#include "Misc/Optional.h"
+#include "Delegates/DelegateBase.h"
 #include "Internationalization/Text.h"
+#include "Misc/Optional.h"
 #include "Templates/ValueOrError.h"
 
+class FText;
 template<typename NumericType> struct FNumericUnit;
 
 /** Enum *must* be zero-indexed and sequential. Must be grouped by relevance and ordered by magnitude. */
diff --git a/Engine/Source/Runtime/Core/Public/Math/UnrealMath.h b/Engine/Source/Runtime/Core/Public/Math/UnrealMath.h
--- a/Engine/Source/Runtime/Core/Public/Math/UnrealMath.h
+++ b/Engine/Source/Runtime/Core/Public/Math/UnrealMath.h
@@ -2,54 +2,54 @@
 
 #pragma once
 
-#include "Math/UnrealMathUtility.h"
+#include "Math/Axis.h"
+#include "Math/Box.h"
+#include "Math/Box2D.h"
+#include "Math/BoxSphereBounds.h"
+#include "Math/CapsuleShape.h"
+#include "Math/ClipProjectionMatrix.h"
 #include "Math/Color.h"
 #include "Math/ColorList.h"
-#include "Math/NumericLimits.h"
+#include "Math/ConvexHull2d.h"
+#include "Math/CurveEdInterface.h"
+#include "Math/Edge.h"
+#include "Math/Float16.h"
+#include "Math/Float16Color.h"
+#include "Math/Float32.h"
 #include "Math/IntPoint.h"
-#include "Math/IntVector.h"
-#include "Math/Vector2D.h"
 #include "Math/IntRect.h"
-#include "Math/Vector.h"
-#include "Math/Vector4.h"
-#include "Math/VectorRegister.h"
-#include "Math/TwoVectors.h"
-#include "Math/Ray.h"
-#include "Math/Edge.h"
+#include "Math/IntVector.h"
+#include "Math/InterpCurve.h"
+#include "Math/InterpCurvePoint.h"
+#include "Math/Interval.h"
+#include "Math/InverseRotationMatrix.h"
+#include "Math/Matrix.h"
+#include "Math/MirrorMatrix.h"
+#include "Math/NumericLimits.h"
+#include "Math/OrientedBox.h"
+#include "Math/OrthoMatrix.h"
+#include "Math/PerspectiveMatrix.h"
 #include "Math/Plane.h"
-#include "Math/Sphere.h"
-#include "Math/CapsuleShape.h"
-#include "Math/Rotator.h"
-#include "Math/RangeBound.h"
+#include "Math/Quat.h"
+#include "Math/QuatRotationTranslationMatrix.h"
 #include "Math/Range.h"
+#include "Math/RangeBound.h"
 #include "Math/RangeSet.h"
-#include "Math/Interval.h"
-#include "Math/Box.h"
-#include "Math/Box2D.h"
-#include "Math/BoxSphereBounds.h"
-#include "Math/OrientedBox.h"
-#include "Math/Axis.h"
-#include "Math/Matrix.h"
-#include "Math/RotationTranslationMatrix.h"
+#include "Math/Ray.h"
 #include "Math/RotationAboutPointMatrix.h"
-#include "Math/ScaleRotationTranslationMatrix.h"
 #include "Math/RotationMatrix.h"
-#include "Math/Quat.h"
-#include "Math/PerspectiveMatrix.h"
-#include "Math/OrthoMatrix.h"
-#include "Math/TranslationMatrix.h"
-#include "Math/QuatRotationTranslationMatrix.h"
-#include "Math/InverseRotationMatrix.h"
-#include "Math/ScaleMatrix.h"
-#include "Math/MirrorMatrix.h"
-#include "Math/ClipProjectionMatrix.h"
-#include "Math/InterpCurvePoint.h"
-#include "Math/InterpCurve.h"
-#include "Math/CurveEdInterface.h"
-#include "Math/Float32.h"
-#include "Math/Float16.h"
-#include "Math/Float16Color.h"
-#include "Math/Vector2DHalf.h"
+#include "Math/RotationTranslationMatrix.h"
+#include "Math/Rotator.h"
 #include "Math/ScalarRegister.h"
+#include "Math/ScaleMatrix.h"
+#include "Math/ScaleRotationTranslationMatrix.h"
+#include "Math/Sphere.h"
 #include "Math/Transform.h"
-#include "Math/ConvexHull2d.h"
+#include "Math/TranslationMatrix.h"
+#include "Math/TwoVectors.h"
+#include "Math/UnrealMathUtility.h"
+#include "Math/Vector.h"
+#include "Math/Vector2D.h"
+#include "Math/Vector2DHalf.h"
+#include "Math/Vector4.h"
+#include "Math/VectorRegister.h"
diff --git a/Engine/Source/Runtime/Core/Public/Math/UnrealMathSSE.h b/Engine/Source/Runtime/Core/Public/Math/UnrealMathSSE.h
--- a/Engine/Source/Runtime/Core/Public/Math/UnrealMathSSE.h
+++ b/Engine/Source/Runtime/Core/Public/Math/UnrealMathSSE.h
@@ -2,6 +2,8 @@
 
 #pragma once
 
+#include "HAL/Platform.h"
+#include "Math/MathFwd.h"
 #include "Math/UnrealMathUtility.h"
 
 #if !WITH_DIRECTXMATH && !PLATFORM_ENABLE_VECTORINTRINSICS_NEON && !(defined(__cplusplus_cli) && !PLATFORM_HOLOLENS) && PLATFORM_ENABLE_VECTORINTRINSICS
diff --git a/Engine/Source/Runtime/Core/Public/Memory/CompositeBuffer.h b/Engine/Source/Runtime/Core/Public/Memory/CompositeBuffer.h
--- a/Engine/Source/Runtime/Core/Public/Memory/CompositeBuffer.h
+++ b/Engine/Source/Runtime/Core/Public/Memory/CompositeBuffer.h
@@ -4,8 +4,16 @@
 
 #include "Containers/Array.h"
 #include "Containers/ArrayView.h"
+#include "Containers/ContainerAllocationPolicies.h"
+#include "Containers/ContainersFwd.h"
+#include "HAL/Platform.h"
+#include "Math/NumericLimits.h"
+#include "Memory/MemoryFwd.h"
 #include "Memory/MemoryView.h"
 #include "Memory/SharedBuffer.h"
+#include "Templates/Function.h"
+#include "Templates/UnrealTemplate.h"
+
 #include <type_traits>
 
 template <typename FuncType> class TFunctionRef;
diff --git a/Engine/Source/Runtime/Core/Public/Memory/SharedBuffer.h b/Engine/Source/Runtime/Core/Public/Memory/SharedBuffer.h
--- a/Engine/Source/Runtime/Core/Public/Memory/SharedBuffer.h
+++ b/Engine/Source/Runtime/Core/Public/Memory/SharedBuffer.h
@@ -2,10 +2,11 @@
 
 #pragma once
 
-#include "CoreTypes.h"
 #include "Containers/ContainersFwd.h"
+#include "CoreTypes.h"
 #include "Memory/MemoryFwd.h"
 #include "Memory/MemoryView.h"
+#include "Misc/AssertionMacros.h"
 #include "Misc/EnumClassFlags.h"
 #include "Templates/Invoke.h"
 #include "Templates/TypeHash.h"
@@ -14,6 +15,9 @@
 #include <atomic>
 #include <type_traits>
 
+template <typename T> struct TIsWeakPointerType;
+template <typename T> struct TIsZeroConstructType;
+
 ///////////////////////////////////////////////////////////////////////////////////////////////////////////////////////
 
 namespace UE::SharedBuffer::Private { struct FSharedOps; }
diff --git a/Engine/Source/Runtime/Core/Public/Misc/App.h b/Engine/Source/Runtime/Core/Public/Misc/App.h
--- a/Engine/Source/Runtime/Core/Public/Misc/App.h
+++ b/Engine/Source/Runtime/Core/Public/Misc/App.h
@@ -2,20 +2,23 @@
 
 #pragma once
 
-#include "CoreTypes.h"
-#include "Templates/UnrealTemplate.h"
 #include "Containers/Array.h"
 #include "Containers/UnrealString.h"
-#include "Misc/Parse.h"
-#include "UObject/NameTypes.h"
 #include "CoreGlobals.h"
+#include "CoreMinimal.h"
+#include "CoreTypes.h"
 #include "Delegates/Delegate.h"
-#include "Misc/Guid.h"
-#include "Misc/CoreMisc.h"
+#include "HAL/PlatformProcess.h"
 #include "Misc/CommandLine.h"
+#include "Misc/CoreMisc.h"
+#include "Misc/FrameRate.h"
+#include "Misc/Guid.h"
 #include "Misc/Optional.h"
+#include "Misc/Parse.h"
 #include "Misc/QualifiedFrameTime.h"
-#include "HAL/PlatformProcess.h"
+#include "Misc/Timecode.h"
+#include "Templates/UnrealTemplate.h"
+#include "UObject/NameTypes.h"
 
 /**
  * Provides information about the application.
diff --git a/Engine/Source/Runtime/Core/Public/Misc/AsyncTaskNotification.h b/Engine/Source/Runtime/Core/Public/Misc/AsyncTaskNotification.h
--- a/Engine/Source/Runtime/Core/Public/Misc/AsyncTaskNotification.h
+++ b/Engine/Source/Runtime/Core/Public/Misc/AsyncTaskNotification.h
@@ -3,12 +3,15 @@
 #pragma once
 
 #include "CoreTypes.h"
-#include "Misc/Attribute.h"
-#include "Logging/LogCategory.h"
+#include "Delegates/Delegate.h"
 #include "Internationalization/Text.h"
+#include "Logging/LogCategory.h"
+#include "Misc/Attribute.h"
+#include "Templates/SharedPointer.h"
 
-struct FSlateBrush;
 class IAsyncTaskNotificationImpl;
+struct FLogCategoryBase;
+struct FSlateBrush;
 
 /**
  * Configuration data for initializing an asynchronous task notification.
diff --git a/Engine/Source/Runtime/Core/Public/Misc/Base64.h b/Engine/Source/Runtime/Core/Public/Misc/Base64.h
--- a/Engine/Source/Runtime/Core/Public/Misc/Base64.h
+++ b/Engine/Source/Runtime/Core/Public/Misc/Base64.h
@@ -2,8 +2,9 @@
 
 #pragma once
 
-#include "CoreTypes.h"
+#include "Containers/Array.h"
 #include "Containers/UnrealString.h"
+#include "CoreTypes.h"
 #include "Misc/Timespan.h"
 
 /**
diff --git a/Engine/Source/Runtime/Core/Public/Misc/BufferedOutputDevice.h b/Engine/Source/Runtime/Core/Public/Misc/BufferedOutputDevice.h
--- a/Engine/Source/Runtime/Core/Public/Misc/BufferedOutputDevice.h
+++ b/Engine/Source/Runtime/Core/Public/Misc/BufferedOutputDevice.h
@@ -2,10 +2,13 @@
 
 #pragma once
 
+#include "Containers/Array.h"
 #include "CoreTypes.h"
+#include "HAL/CriticalSection.h"
+#include "Logging/LogVerbosity.h"
 #include "Misc/OutputDevice.h"
-#include "Containers/Array.h"
 #include "Misc/OutputDeviceRedirector.h"
+#include "UObject/NameTypes.h"
 
 /** Buffered output device. */
 class CORE_API FBufferedOutputDevice : public FOutputDevice
diff --git a/Engine/Source/Runtime/Core/Public/Misc/CString.h b/Engine/Source/Runtime/Core/Public/Misc/CString.h
--- a/Engine/Source/Runtime/Core/Public/Misc/CString.h
+++ b/Engine/Source/Runtime/Core/Public/Misc/CString.h
@@ -3,13 +3,14 @@
 #pragma once
 
 #include "CoreTypes.h"
-#include "Misc/VarArgs.h"
+#include "HAL/PlatformCrt.h"
+#include "HAL/PlatformString.h"
 #include "Misc/AssertionMacros.h"
 #include "Misc/Char.h"
-#include "HAL/PlatformString.h"
-#include "Templates/IsValidVariadicFunctionArg.h"
+#include "Misc/VarArgs.h"
 #include "Templates/AndOrNot.h"
 #include "Templates/IsArrayOrRefOfTypeByPredicate.h"
+#include "Templates/IsValidVariadicFunctionArg.h"
 #include "Traits/IsCharEncodingCompatibleWith.h"
 
 #define MAX_SPRINTF 1024
diff --git a/Engine/Source/Runtime/Core/Public/Misc/Change.h b/Engine/Source/Runtime/Core/Public/Misc/Change.h
--- a/Engine/Source/Runtime/Core/Public/Misc/Change.h
+++ b/Engine/Source/Runtime/Core/Public/Misc/Change.h
@@ -2,11 +2,16 @@
 
 #pragma once
 
-#include "Templates/UniquePtr.h"
 #include "Containers/Array.h"
+#include "Containers/UnrealString.h"
+#include "HAL/Platform.h"
+#include "Misc/AssertionMacros.h"
+#include "Templates/UniquePtr.h"
+#include "Templates/UnrealTemplate.h"
 
 class FFeedbackContext;
 class FReferenceCollector;
+class UObject;
 
 /**
  * FChange modifies a UObject and is meant to be used to implement undo/redo.
diff --git a/Engine/Source/Runtime/Core/Public/Misc/ComparisonUtility.h b/Engine/Source/Runtime/Core/Public/Misc/ComparisonUtility.h
--- a/Engine/Source/Runtime/Core/Public/Misc/ComparisonUtility.h
+++ b/Engine/Source/Runtime/Core/Public/Misc/ComparisonUtility.h
@@ -2,8 +2,10 @@
 
 #pragma once
 
-#include "UObject/NameTypes.h"
+#include "Containers/StringFwd.h"
 #include "Containers/StringView.h"
+#include "HAL/Platform.h"
+#include "UObject/NameTypes.h"
 
 namespace UE { namespace ComparisonUtility {
 
diff --git a/Engine/Source/Runtime/Core/Public/Misc/Compression.h b/Engine/Source/Runtime/Core/Public/Misc/Compression.h
--- a/Engine/Source/Runtime/Core/Public/Misc/Compression.h
+++ b/Engine/Source/Runtime/Core/Public/Misc/Compression.h
@@ -2,13 +2,15 @@
 
 #pragma once
 
+#include "CoreMinimal.h"
 #include "CoreTypes.h"
-#include "UObject/NameTypes.h"
-#include "Templates/Atomic.h"
-#include "Misc/CompressionFlags.h"
 #include "HAL/CriticalSection.h"
+#include "Misc/CompressionFlags.h"
+#include "Templates/Atomic.h"
+#include "UObject/NameTypes.h"
 
 class IMemoryReadStream;
+template <typename T> class TAtomic;
 
 // Define global current platform default to current platform.  
 // DEPRECATED, USE NAME_Zlib
diff --git a/Engine/Source/Runtime/Core/Public/Misc/ConfigCacheIni.h b/Engine/Source/Runtime/Core/Public/Misc/ConfigCacheIni.h
--- a/Engine/Source/Runtime/Core/Public/Misc/ConfigCacheIni.h
+++ b/Engine/Source/Runtime/Core/Public/Misc/ConfigCacheIni.h
@@ -6,21 +6,29 @@
 
 #pragma once
 
-#include "CoreTypes.h"
 #include "Containers/Array.h"
-#include "Containers/UnrealString.h"
 #include "Containers/Map.h"
 #include "Containers/StringFwd.h"
-#include "Math/Color.h"
-#include "UObject/NameTypes.h"
-#include "Logging/LogMacros.h"
-#include "Math/Vector2D.h"
+#include "Containers/UnrealString.h"
+#include "CoreMinimal.h"
+#include "CoreTypes.h"
 #include "Delegates/Delegate.h"
-#include "Math/Vector.h"
+#include "Logging/LogMacros.h"
+#include "Math/Color.h"
+#include "Math/MathFwd.h"
 #include "Math/Rotator.h"
+#include "Math/Vector.h"
+#include "Math/Vector2D.h"
 #include "Misc/AccessDetection.h"
 #include "Misc/Paths.h"
 #include "Serialization/StructuredArchive.h"
+#include "Serialization/StructuredArchiveAdapters.h"
+#include "Serialization/StructuredArchiveSlots.h"
+#include "UObject/NameTypes.h"
+
+class FConfigCacheIni;
+class FConfigFile;
+class IConsoleVariable;
 
 CORE_API DECLARE_LOG_CATEGORY_EXTERN(LogConfig, Log, All);
 
diff --git a/Engine/Source/Runtime/Core/Public/Misc/ConfigUtilities.h b/Engine/Source/Runtime/Core/Public/Misc/ConfigUtilities.h
--- a/Engine/Source/Runtime/Core/Public/Misc/ConfigUtilities.h
+++ b/Engine/Source/Runtime/Core/Public/Misc/ConfigUtilities.h
@@ -6,10 +6,14 @@
 
 #pragma once
 
-#include "CoreTypes.h"
-#include "Containers/UnrealString.h"
 #include "Containers/Map.h"
+#include "Containers/UnrealString.h"
+#include "CoreTypes.h"
 #include "Delegates/Delegate.h"
+#include "Templates/Function.h"
+
+class FString;
+class IConsoleVariable;
 
 
 namespace UE::ConfigUtilities
diff --git a/Engine/Source/Runtime/Core/Public/Misc/CoreAsyncTaskNotificationImpl.h b/Engine/Source/Runtime/Core/Public/Misc/CoreAsyncTaskNotificationImpl.h
--- a/Engine/Source/Runtime/Core/Public/Misc/CoreAsyncTaskNotificationImpl.h
+++ b/Engine/Source/Runtime/Core/Public/Misc/CoreAsyncTaskNotificationImpl.h
@@ -2,7 +2,18 @@
 
 #pragma once
 
+#include "Containers/Array.h"
+#include "Delegates/Delegate.h"
+#include "HAL/CriticalSection.h"
+#include "Internationalization/Text.h"
 #include "Misc/AsyncTaskNotification.h"
+#include "Misc/Attribute.h"
+#include "Templates/Function.h"
+#include "Templates/SharedPointer.h"
+#include "Templates/Tuple.h"
+#include "UObject/NameTypes.h"
+
+struct FLogCategoryBase;
 
 /**
  * Implementation interface of notifications for an on-going asynchronous task.
diff --git a/Engine/Source/Runtime/Core/Public/Misc/CoreDelegates.h b/Engine/Source/Runtime/Core/Public/Misc/CoreDelegates.h
--- a/Engine/Source/Runtime/Core/Public/Misc/CoreDelegates.h
+++ b/Engine/Source/Runtime/Core/Public/Misc/CoreDelegates.h
@@ -2,19 +2,29 @@
 
 #pragma once
 
-#include "CoreTypes.h"
+#include "Async/TaskGraphInterfaces.h"
 #include "Containers/Array.h"
-#include "Containers/UnrealString.h"
+#include "Containers/ArrayView.h"
 #include "Containers/Map.h"
+#include "Containers/UnrealString.h"
+#include "CoreMinimal.h"
+#include "CoreTypes.h"
 #include "Delegates/Delegate.h"
+#include "GenericPlatform/GenericPlatformFile.h"
+#include "HAL/PlatformFile.h"
 #include "Math/IntVector.h"
+#include "Math/MathFwd.h"
 #include "Misc/AES.h"
-#include "GenericPlatform/GenericPlatformFile.h"
-#include "Async/TaskGraphInterfaces.h"
+#include "Templates/SharedPointerInternals.h"
 
 class AActor;
 class Error;
+class FConfigFile;
+class FSHAHash;
 class IPakFile;
+class ITargetPlatform;
+struct FGuid;
+
 enum class EForkProcessRole : uint8;
 
 // delegates for hotfixes
diff --git a/Engine/Source/Runtime/Core/Public/Misc/CoreMisc.h b/Engine/Source/Runtime/Core/Public/Misc/CoreMisc.h
--- a/Engine/Source/Runtime/Core/Public/Misc/CoreMisc.h
+++ b/Engine/Source/Runtime/Core/Public/Misc/CoreMisc.h
@@ -2,16 +2,19 @@
 
 #pragma once
 
-#include "CoreTypes.h"
-#include "Misc/Exec.h"
 #include "Containers/Array.h"
-#include "Containers/UnrealString.h"
 #include "Containers/Map.h"
-#include "Templates/Function.h"
-#include "Math/IntPoint.h"
-#include "UObject/NameTypes.h"
+#include "Containers/UnrealString.h"
 #include "CoreGlobals.h"
+#include "CoreMinimal.h"
+#include "CoreTypes.h"
 #include "HAL/ThreadSingleton.h"
+#include "Math/IntPoint.h"
+#include "Misc/Exec.h"
+#include "Templates/Function.h"
+#include "UObject/NameTypes.h"
+
+class UWorld;
 
 /**
  * Exec handler that registers itself and is being routed via StaticExec.
diff --git a/Engine/Source/Runtime/Core/Public/Misc/Crc.h b/Engine/Source/Runtime/Core/Public/Misc/Crc.h
--- a/Engine/Source/Runtime/Core/Public/Misc/Crc.h
+++ b/Engine/Source/Runtime/Core/Public/Misc/Crc.h
@@ -3,9 +3,11 @@
 #pragma once
 
 #include "CoreTypes.h"
+#include "HAL/PlatformString.h"
 #include "Misc/AssertionMacros.h"
-#include "Misc/Char.h"
 #include "Misc/CString.h"
+#include "Misc/Char.h"
+#include "Templates/EnableIf.h"
 
 
 /** 
diff --git a/Engine/Source/Runtime/Core/Public/Misc/DataDrivenPlatformInfoRegistry.h b/Engine/Source/Runtime/Core/Public/Misc/DataDrivenPlatformInfoRegistry.h
--- a/Engine/Source/Runtime/Core/Public/Misc/DataDrivenPlatformInfoRegistry.h
+++ b/Engine/Source/Runtime/Core/Public/Misc/DataDrivenPlatformInfoRegistry.h
@@ -2,13 +2,16 @@
 
 #pragma once
 
-#include "CoreTypes.h"
-#include "Containers/UnrealString.h"
-#include "Containers/Map.h"
 #include "Containers/Array.h"
+#include "Containers/Map.h"
+#include "Containers/UnrealString.h"
+#include "CoreTypes.h"
+#include "Internationalization/Text.h"
 #include "Misc/Guid.h"
 #include "UObject/NameTypes.h"
-#include "Internationalization/Text.h"
+#include "UObject/UnrealNames.h"
+
+class FConfigFile;
 
 #define DDPI_HAS_EXTENDED_PLATFORMINFO_DATA (WITH_EDITOR || IS_PROGRAM)
 
diff --git a/Engine/Source/Runtime/Core/Public/Misc/DateTime.h b/Engine/Source/Runtime/Core/Public/Misc/DateTime.h
--- a/Engine/Source/Runtime/Core/Public/Misc/DateTime.h
+++ b/Engine/Source/Runtime/Core/Public/Misc/DateTime.h
@@ -2,9 +2,11 @@
 
 #pragma once
 
-#include "CoreTypes.h"
+#include "Containers/StringFwd.h"
 #include "Containers/UnrealString.h"
+#include "CoreTypes.h"
 #include "Misc/Timespan.h"
+#include "Serialization/Archive.h"
 #include "Serialization/StructuredArchive.h"
 #include "Templates/TypeHash.h"
 
diff --git a/Engine/Source/Runtime/Core/Public/Misc/DefaultValueHelper.h b/Engine/Source/Runtime/Core/Public/Misc/DefaultValueHelper.h
--- a/Engine/Source/Runtime/Core/Public/Misc/DefaultValueHelper.h
+++ b/Engine/Source/Runtime/Core/Public/Misc/DefaultValueHelper.h
@@ -2,10 +2,14 @@
 
 #pragma once
 
-#include "CoreTypes.h"
 #include "Containers/UnrealString.h"
+#include "CoreTypes.h"
+#include "Math/MathFwd.h"
 #include "Misc/Timespan.h"
 
+struct FColor;
+struct FLinearColor;
+
 //
 //	FDefaultValueHelper
 //
diff --git a/Engine/Source/Runtime/Core/Public/Misc/DelayedAutoRegister.h b/Engine/Source/Runtime/Core/Public/Misc/DelayedAutoRegister.h
--- a/Engine/Source/Runtime/Core/Public/Misc/DelayedAutoRegister.h
+++ b/Engine/Source/Runtime/Core/Public/Misc/DelayedAutoRegister.h
@@ -2,6 +2,7 @@
 
 #pragma once
 
+#include "HAL/Platform.h"
 #include "Templates/Function.h"
 
 enum class EDelayedRegisterRunPhase : uint8
diff --git a/Engine/Source/Runtime/Core/Public/Misc/EmbeddedCommunication.h b/Engine/Source/Runtime/Core/Public/Misc/EmbeddedCommunication.h
--- a/Engine/Source/Runtime/Core/Public/Misc/EmbeddedCommunication.h
+++ b/Engine/Source/Runtime/Core/Public/Misc/EmbeddedCommunication.h
@@ -2,11 +2,18 @@
 
 #pragma once
 
-#include "CoreTypes.h"
 #include "Containers/Array.h"
-#include "Containers/UnrealString.h"
 #include "Containers/Map.h"
+#include "Containers/UnrealString.h"
+#include "CoreTypes.h"
 #include "Delegates/Delegate.h"
+#include "Delegates/DelegateBase.h"
+#include "Delegates/IDelegateInstance.h"
+#include "Delegates/MulticastDelegateBase.h"
+#include "HAL/CriticalSection.h"
+#include "Templates/Function.h"
+#include "UObject/NameTypes.h"
+#include "UObject/UnrealNames.h"
 
 typedef TMap<FString, FString> FEmbeddedCommunicationMap;
 
diff --git a/Engine/Source/Runtime/Core/Public/Misc/EngineBuildSettings.h b/Engine/Source/Runtime/Core/Public/Misc/EngineBuildSettings.h
--- a/Engine/Source/Runtime/Core/Public/Misc/EngineBuildSettings.h
+++ b/Engine/Source/Runtime/Core/Public/Misc/EngineBuildSettings.h
@@ -2,8 +2,9 @@
 
 #pragma once
 
-#include "CoreTypes.h"
 #include "Containers/ContainersFwd.h"
+#include "Containers/UnrealString.h"
+#include "CoreTypes.h"
 
 class CORE_API FEngineBuildSettings
 {
diff --git a/Engine/Source/Runtime/Core/Public/Misc/Exec.h b/Engine/Source/Runtime/Core/Public/Misc/Exec.h
--- a/Engine/Source/Runtime/Core/Public/Misc/Exec.h
+++ b/Engine/Source/Runtime/Core/Public/Misc/Exec.h
@@ -5,6 +5,8 @@
 #include "CoreTypes.h"
 #include "Misc/AssertionMacros.h"
 
+class FOutputDevice;
+
 // Any object that is capable of taking commands.
 class CORE_API FExec
 {
diff --git a/Engine/Source/Runtime/Core/Public/Misc/FbxErrors.h b/Engine/Source/Runtime/Core/Public/Misc/FbxErrors.h
--- a/Engine/Source/Runtime/Core/Public/Misc/FbxErrors.h
+++ b/Engine/Source/Runtime/Core/Public/Misc/FbxErrors.h
@@ -3,9 +3,9 @@
 #pragma once
 
 #include "CoreTypes.h"
-#include "UObject/NameTypes.h"
-#include "Templates/SharedPointer.h"
 #include "Logging/TokenizedMessage.h"
+#include "Templates/SharedPointer.h"
+#include "UObject/NameTypes.h"
 
 
 /**
diff --git a/Engine/Source/Runtime/Core/Public/Misc/FeedbackContext.h b/Engine/Source/Runtime/Core/Public/Misc/FeedbackContext.h
--- a/Engine/Source/Runtime/Core/Public/Misc/FeedbackContext.h
+++ b/Engine/Source/Runtime/Core/Public/Misc/FeedbackContext.h
@@ -2,18 +2,24 @@
 
 #pragma once
 
-#include "CoreTypes.h"
-#include "Misc/OutputDevice.h"
 #include "Containers/Array.h"
 #include "Containers/UnrealString.h"
-#include "Templates/SharedPointer.h"
+#include "CoreGlobals.h"
+#include "CoreTypes.h"
+#include "HAL/CriticalSection.h"
+#include "HAL/PlatformCrt.h"
 #include "Internationalization/Text.h"
+#include "Misc/OutputDevice.h"
 #include "Misc/ScopeLock.h"
 #include "Misc/SlowTask.h"
 #include "Misc/SlowTaskStack.h"
+#include "Templates/SharedPointer.h"
+#include "Templates/UniquePtr.h"
+#include "Templates/UnrealTemplate.h"
 
 class FContextSupplier;
 class SBuildProgressWidget;
+struct FSlowTask;
 
 /** A context for displaying modal warning messages. */
 class CORE_API FFeedbackContext
diff --git a/Engine/Source/Runtime/Core/Public/Misc/FeedbackContextMarkup.h b/Engine/Source/Runtime/Core/Public/Misc/FeedbackContextMarkup.h
--- a/Engine/Source/Runtime/Core/Public/Misc/FeedbackContextMarkup.h
+++ b/Engine/Source/Runtime/Core/Public/Misc/FeedbackContextMarkup.h
@@ -2,8 +2,12 @@
 
 #pragma once
 
-#include "CoreTypes.h"
+#include "Containers/UnrealString.h"
 #include "CoreFwd.h"
+#include "CoreTypes.h"
+
+class FFeedbackContext;
+class FText;
 
 /**
  * A text-based markup language can be used to allow external processes to control the state 
diff --git a/Engine/Source/Runtime/Core/Public/Misc/FileHelper.h b/Engine/Source/Runtime/Core/Public/Misc/FileHelper.h
--- a/Engine/Source/Runtime/Core/Public/Misc/FileHelper.h
+++ b/Engine/Source/Runtime/Core/Public/Misc/FileHelper.h
@@ -2,13 +2,18 @@
 
 #pragma once
 
-#include "CoreTypes.h"
+#include "Containers/ArrayView.h"
+#include "Containers/StringFwd.h"
 #include "Containers/StringView.h"
 #include "Containers/UnrealString.h"
+#include "CoreMinimal.h"
+#include "CoreTypes.h"
 #include "HAL/FileManager.h"
-#include "Containers/ArrayView.h"
-#include "Misc/EnumClassFlags.h"
 #include "Math/Color.h"
+#include "Math/MathFwd.h"
+#include "Misc/EnumClassFlags.h"
+
+class IPlatformFile;
 
 /*-----------------------------------------------------------------------------
 	FFileHelper
diff --git a/Engine/Source/Runtime/Core/Public/Misc/FrameNumber.h b/Engine/Source/Runtime/Core/Public/Misc/FrameNumber.h
--- a/Engine/Source/Runtime/Core/Public/Misc/FrameNumber.h
+++ b/Engine/Source/Runtime/Core/Public/Misc/FrameNumber.h
@@ -9,6 +9,7 @@
 #include "Templates/UnrealTypeTraits.h"
 
 class FArchive;
+template <typename NumericType> struct TNumericLimits;
 
 /**
  * Typesafe 32-bit signed frame number. Defined in this way to prevent erroneous float->int conversions and afford type-safe operator overloading.
diff --git a/Engine/Source/Runtime/Core/Public/Misc/FrameRate.h b/Engine/Source/Runtime/Core/Public/Misc/FrameRate.h
--- a/Engine/Source/Runtime/Core/Public/Misc/FrameRate.h
+++ b/Engine/Source/Runtime/Core/Public/Misc/FrameRate.h
@@ -2,11 +2,15 @@
 
 #pragma once
 
-#include "Misc/FrameNumber.h"
-#include "Misc/FrameTime.h"
+#include "HAL/Platform.h"
 #include "Internationalization/Text.h"
+#include "Math/NumericLimits.h"
 #include "Math/Range.h"
+#include "Math/RangeBound.h"
+#include "Math/UnrealMathUtility.h"
 #include "Misc/ExpressionParserTypes.h"
+#include "Misc/FrameNumber.h"
+#include "Misc/FrameTime.h"
 #include "Templates/ValueOrError.h"
 
 struct FExpressionError;
diff --git a/Engine/Source/Runtime/Core/Public/Misc/Guid.h b/Engine/Source/Runtime/Core/Public/Misc/Guid.h
--- a/Engine/Source/Runtime/Core/Public/Misc/Guid.h
+++ b/Engine/Source/Runtime/Core/Public/Misc/Guid.h
@@ -2,18 +2,26 @@
 
 #pragma once
 
+#include "Containers/StringFwd.h"
+#include "Containers/UnrealString.h"
 #include "CoreTypes.h"
+#include "Hash/CityHash.h"
 #include "Misc/AssertionMacros.h"
 #include "Misc/Crc.h"
-#include "Containers/StringFwd.h"
-#include "Containers/UnrealString.h"
-#include "Serialization/StructuredArchive.h"
+#include "Serialization/Archive.h"
 #include "Serialization/MemoryLayout.h"
-#include "Hash/CityHash.h"
+#include "Serialization/StructuredArchive.h"
 
 class FArchive;
+class FMemoryImageWriter;
+class FMemoryUnfreezeContent;
 class FOutputDevice;
+class FPointerTableBase;
+class FSHA1;
 class UObject;
+template <typename CharType> class TStringBuilderBase;
+template <typename T> struct TCanBulkSerialize;
+template <typename T> struct TIsPODType;
 
 
 /**
diff --git a/Engine/Source/Runtime/Core/Public/Misc/InteractiveProcess.h b/Engine/Source/Runtime/Core/Public/Misc/InteractiveProcess.h
--- a/Engine/Source/Runtime/Core/Public/Misc/InteractiveProcess.h
+++ b/Engine/Source/Runtime/Core/Public/Misc/InteractiveProcess.h
@@ -2,15 +2,19 @@
 
 #pragma once
 
-#include "CoreTypes.h"
+#include "Containers/Array.h"
+#include "Containers/Queue.h"
 #include "Containers/UnrealString.h"
-#include "Misc/Timespan.h"
-#include "Logging/LogMacros.h"
+#include "CoreTypes.h"
 #include "Delegates/Delegate.h"
-#include "Misc/DateTime.h"
+#include "Delegates/DelegateBase.h"
 #include "HAL/PlatformProcess.h"
 #include "HAL/Runnable.h"
-#include "Containers/Queue.h"
+#include "Logging/LogMacros.h"
+#include "Misc/DateTime.h"
+#include "Misc/Timespan.h"
+
+class FRunnableThread;
 
 DECLARE_LOG_CATEGORY_EXTERN(LogInteractiveProcess, Log, All);
 
diff --git a/Engine/Source/Runtime/Core/Public/Misc/LargeWorldRenderPosition.h b/Engine/Source/Runtime/Core/Public/Misc/LargeWorldRenderPosition.h
--- a/Engine/Source/Runtime/Core/Public/Misc/LargeWorldRenderPosition.h
+++ b/Engine/Source/Runtime/Core/Public/Misc/LargeWorldRenderPosition.h
@@ -3,6 +3,7 @@
 #pragma once
 
 #include "CoreMinimal.h"
+#include "Math/MathFwd.h"
 #include "Misc/LargeWorldCoordinates.h"
 
 struct FLargeWorldRenderScalar
diff --git a/Engine/Source/Runtime/Core/Public/Misc/LocalTimestampDirectoryVisitor.h b/Engine/Source/Runtime/Core/Public/Misc/LocalTimestampDirectoryVisitor.h
--- a/Engine/Source/Runtime/Core/Public/Misc/LocalTimestampDirectoryVisitor.h
+++ b/Engine/Source/Runtime/Core/Public/Misc/LocalTimestampDirectoryVisitor.h
@@ -2,12 +2,15 @@
 
 #pragma once
 
-#include "CoreTypes.h"
 #include "Containers/Array.h"
-#include "Containers/UnrealString.h"
 #include "Containers/Map.h"
-#include "Misc/DateTime.h"
+#include "Containers/UnrealString.h"
+#include "CoreTypes.h"
 #include "GenericPlatform/GenericPlatformFile.h"
+#include "HAL/PlatformFile.h"
+#include "Misc/DateTime.h"
+
+struct FDateTime;
 
 /**
  * Visitor to gather local files with their timestamps.
diff --git a/Engine/Source/Runtime/Core/Public/Misc/MapErrors.h b/Engine/Source/Runtime/Core/Public/Misc/MapErrors.h
--- a/Engine/Source/Runtime/Core/Public/Misc/MapErrors.h
+++ b/Engine/Source/Runtime/Core/Public/Misc/MapErrors.h
@@ -3,9 +3,9 @@
 #pragma once
 
 #include "CoreTypes.h"
-#include "UObject/NameTypes.h"
-#include "Templates/SharedPointer.h"
 #include "Logging/TokenizedMessage.h"
+#include "Templates/SharedPointer.h"
+#include "UObject/NameTypes.h"
 
 /**
  * This file contains known map errors that can be referenced by name.
diff --git a/Engine/Source/Runtime/Core/Public/Misc/MemStack.h b/Engine/Source/Runtime/Core/Public/Misc/MemStack.h
--- a/Engine/Source/Runtime/Core/Public/Misc/MemStack.h
+++ b/Engine/Source/Runtime/Core/Public/Misc/MemStack.h
@@ -2,16 +2,17 @@
 
 #pragma once
 
+#include "Containers/ContainerAllocationPolicies.h"
+#include "Containers/LockFreeFixedSizeAllocator.h"
+#include "CoreMinimal.h"
 #include "CoreTypes.h"
-#include "Misc/AssertionMacros.h"
+#include "HAL/ThreadSafeCounter.h"
+#include "HAL/ThreadSingleton.h"
 #include "HAL/UnrealMemory.h"
-#include "Containers/ContainerAllocationPolicies.h"
 #include "Math/UnrealMathUtility.h"
-#include "Templates/AlignmentTemplates.h"
-#include "HAL/ThreadSingleton.h"
-#include "HAL/ThreadSafeCounter.h"
+#include "Misc/AssertionMacros.h"
 #include "Misc/NoopCounter.h"
-#include "Containers/LockFreeFixedSizeAllocator.h"
+#include "Templates/AlignmentTemplates.h"
 
 
 // Enums for specifying memory allocation type.
diff --git a/Engine/Source/Runtime/Core/Public/Misc/MemoryReadStream.h b/Engine/Source/Runtime/Core/Public/Misc/MemoryReadStream.h
--- a/Engine/Source/Runtime/Core/Public/Misc/MemoryReadStream.h
+++ b/Engine/Source/Runtime/Core/Public/Misc/MemoryReadStream.h
@@ -3,6 +3,7 @@
 #pragma once
 
 #include "CoreMinimal.h"
+#include "HAL/ThreadSafeCounter.h"
 #include "Templates/RefCounting.h"
 
 /**
diff --git a/Engine/Source/Runtime/Core/Public/Misc/MessageDialog.h b/Engine/Source/Runtime/Core/Public/Misc/MessageDialog.h
--- a/Engine/Source/Runtime/Core/Public/Misc/MessageDialog.h
+++ b/Engine/Source/Runtime/Core/Public/Misc/MessageDialog.h
@@ -4,6 +4,9 @@
 
 #include "CoreTypes.h"
 #include "GenericPlatform/GenericPlatformMisc.h"
+#include "HAL/PlatformMisc.h"
+
+class FText;
 
 /** 
  * FMessageDialog
diff --git a/Engine/Source/Runtime/Core/Public/Misc/MonitoredProcess.h b/Engine/Source/Runtime/Core/Public/Misc/MonitoredProcess.h
--- a/Engine/Source/Runtime/Core/Public/Misc/MonitoredProcess.h
+++ b/Engine/Source/Runtime/Core/Public/Misc/MonitoredProcess.h
@@ -2,14 +2,17 @@
 
 #pragma once
 
-#include "CoreTypes.h"
 #include "Containers/UnrealString.h"
-#include "Misc/Timespan.h"
+#include "CoreTypes.h"
 #include "Delegates/Delegate.h"
-#include "Misc/DateTime.h"
+#include "HAL/CriticalSection.h"
 #include "HAL/PlatformProcess.h"
 #include "HAL/Runnable.h"
+#include "Misc/DateTime.h"
 #include "Misc/SingleThreadRunnable.h"
+#include "Misc/Timespan.h"
+
+class FRunnableThread;
 
 /**
  * Declares a delegate that is executed when a monitored process completed.
diff --git a/Engine/Source/Runtime/Core/Public/Misc/NamePermissionList.h b/Engine/Source/Runtime/Core/Public/Misc/NamePermissionList.h
--- a/Engine/Source/Runtime/Core/Public/Misc/NamePermissionList.h
+++ b/Engine/Source/Runtime/Core/Public/Misc/NamePermissionList.h
@@ -2,10 +2,16 @@
 
 #pragma once
 
-#include "CoreTypes.h"
+#include "Containers/Array.h"
+#include "Containers/Map.h"
+#include "Containers/StringFwd.h"
 #include "Containers/StringView.h"
-#include "Templates/SharedPointer.h"
+#include "CoreTypes.h"
 #include "Delegates/Delegate.h"
+#include "Templates/SharedPointer.h"
+#include "UObject/NameTypes.h"
+
+class FString;
 
 /** List of owner names that requested a specific item filtered, allows unregistering specific set of changes by a given plugin or system */
 typedef TArray<FName> FPermissionListOwners;
diff --git a/Engine/Source/Runtime/Core/Public/Misc/NetworkVersion.h b/Engine/Source/Runtime/Core/Public/Misc/NetworkVersion.h
--- a/Engine/Source/Runtime/Core/Public/Misc/NetworkVersion.h
+++ b/Engine/Source/Runtime/Core/Public/Misc/NetworkVersion.h
@@ -2,10 +2,11 @@
 
 #pragma once
 
-#include "CoreTypes.h"
 #include "Containers/UnrealString.h"
-#include "Logging/LogMacros.h"
+#include "CoreTypes.h"
 #include "Delegates/Delegate.h"
+#include "Delegates/DelegateBase.h"
+#include "Logging/LogMacros.h"
 
 // The version number used for determining network compatibility. If zero, uses the engine compatible version.
 #define ENGINE_NET_VERSION  0
diff --git a/Engine/Source/Runtime/Core/Public/Misc/ObjectThumbnail.h b/Engine/Source/Runtime/Core/Public/Misc/ObjectThumbnail.h
--- a/Engine/Source/Runtime/Core/Public/Misc/ObjectThumbnail.h
+++ b/Engine/Source/Runtime/Core/Public/Misc/ObjectThumbnail.h
@@ -2,9 +2,11 @@
 
 #pragma once
 
-#include "CoreTypes.h"
 #include "Containers/Array.h"
 #include "Containers/Map.h"
+#include "CoreTypes.h"
+#include "Serialization/Archive.h"
+#include "Serialization/StructuredArchive.h"
 #include "UObject/NameTypes.h"
 
 /**
diff --git a/Engine/Source/Runtime/Core/Public/Misc/OutputDevice.h b/Engine/Source/Runtime/Core/Public/Misc/OutputDevice.h
--- a/Engine/Source/Runtime/Core/Public/Misc/OutputDevice.h
+++ b/Engine/Source/Runtime/Core/Public/Misc/OutputDevice.h
@@ -2,16 +2,19 @@
 
 #pragma once
 
-#include "CoreTypes.h"
 #include "CoreFwd.h"
+#include "CoreTypes.h"
 #include "Logging/LogVerbosity.h"
 #include "Misc/CoreMiscDefines.h"
 #include "Misc/VarArgs.h"
-#include "Templates/IsValidVariadicFunctionArg.h"
 #include "Templates/AndOrNot.h"
 #include "Templates/IsArrayOrRefOfTypeByPredicate.h"
+#include "Templates/IsValidVariadicFunctionArg.h"
 #include "Traits/IsCharEncodingCompatibleWith.h"
 
+class FString;
+class FText;
+
 #ifndef USE_DEBUG_LOGGING
 #define USE_DEBUG_LOGGING 1
 #endif
diff --git a/Engine/Source/Runtime/Core/Public/Misc/OutputDeviceArchiveWrapper.h b/Engine/Source/Runtime/Core/Public/Misc/OutputDeviceArchiveWrapper.h
--- a/Engine/Source/Runtime/Core/Public/Misc/OutputDeviceArchiveWrapper.h
+++ b/Engine/Source/Runtime/Core/Public/Misc/OutputDeviceArchiveWrapper.h
@@ -3,8 +3,11 @@
 #pragma once
 
 #include "CoreTypes.h"
-#include "Misc/OutputDevice.h"
+#include "Logging/LogVerbosity.h"
 #include "Misc/AssertionMacros.h"
+#include "Misc/OutputDevice.h"
+
+class FArchive;
 
 /**
 * Output device wrapping any kind of FArchive.  Note: Works in any build configuration.
diff --git a/Engine/Source/Runtime/Core/Public/Misc/OutputDeviceDebug.h b/Engine/Source/Runtime/Core/Public/Misc/OutputDeviceDebug.h
--- a/Engine/Source/Runtime/Core/Public/Misc/OutputDeviceDebug.h
+++ b/Engine/Source/Runtime/Core/Public/Misc/OutputDeviceDebug.h
@@ -3,7 +3,9 @@
 #pragma once
 
 #include "CoreTypes.h"
+#include "Logging/LogVerbosity.h"
 #include "Misc/OutputDevice.h"
+#include "UObject/NameTypes.h"
 
 class CORE_API FOutputDeviceDebug : public FOutputDevice
 {
diff --git a/Engine/Source/Runtime/Core/Public/Misc/OutputDeviceFile.h b/Engine/Source/Runtime/Core/Public/Misc/OutputDeviceFile.h
--- a/Engine/Source/Runtime/Core/Public/Misc/OutputDeviceFile.h
+++ b/Engine/Source/Runtime/Core/Public/Misc/OutputDeviceFile.h
@@ -2,15 +2,22 @@
 
 #pragma once
 
+#include "Containers/Array.h"
 #include "CoreTypes.h"
-#include "Misc/OutputDevice.h"
-#include "Templates/UniquePtr.h"
-#include "Templates/Function.h"
+#include "HAL/CriticalSection.h"
 #include "HAL/Runnable.h"
-#include "Serialization/Archive.h"
+#include "HAL/ThreadSafeCounter.h"
+#include "Logging/LogVerbosity.h"
+#include "Misc/OutputDevice.h"
 #include "Misc/ScopeLock.h"
 #include "Misc/SingleThreadRunnable.h"
-#include "Containers/Array.h"
+#include "Serialization/Archive.h"
+#include "Templates/Atomic.h"
+#include "Templates/Function.h"
+#include "Templates/UniquePtr.h"
+#include "UObject/NameTypes.h"
+
+class FRunnableThread;
 
 
 /** string added to the filename of timestamped backup log files */
@@ -182,6 +189,7 @@ private:
 
 	/** Internal data for category inclusion. Must be declared inside CPP file as it uses a TSet<FName> */
 	struct FCategoryInclusionInternal;
+
 	TUniquePtr<FCategoryInclusionInternal> CategoryInclusionInternal;
 
 	/** If true, existing files will not be backed up */
diff --git a/Engine/Source/Runtime/Core/Public/Misc/OutputDeviceMemory.h b/Engine/Source/Runtime/Core/Public/Misc/OutputDeviceMemory.h
--- a/Engine/Source/Runtime/Core/Public/Misc/OutputDeviceMemory.h
+++ b/Engine/Source/Runtime/Core/Public/Misc/OutputDeviceMemory.h
@@ -2,11 +2,12 @@
 
 #pragma once
 
+#include "Containers/Array.h"
 #include "CoreTypes.h"
+#include "HAL/CriticalSection.h"
+#include "Logging/LogVerbosity.h"
 #include "Misc/OutputDevice.h"
 #include "Serialization/Archive.h"
-#include "Containers/Array.h"
-#include "HAL/CriticalSection.h"
 
 /**
  * Memory output device. Logs only into pre-allocated memory buffer.
diff --git a/Engine/Source/Runtime/Core/Public/Misc/OutputDeviceNull.h b/Engine/Source/Runtime/Core/Public/Misc/OutputDeviceNull.h
--- a/Engine/Source/Runtime/Core/Public/Misc/OutputDeviceNull.h
+++ b/Engine/Source/Runtime/Core/Public/Misc/OutputDeviceNull.h
@@ -3,6 +3,7 @@
 #pragma once
 
 #include "CoreTypes.h"
+#include "Logging/LogVerbosity.h"
 #include "Misc/OutputDevice.h"
 
 // Null output device.
diff --git a/Engine/Source/Runtime/Core/Public/Misc/OutputDeviceRedirector.h b/Engine/Source/Runtime/Core/Public/Misc/OutputDeviceRedirector.h
--- a/Engine/Source/Runtime/Core/Public/Misc/OutputDeviceRedirector.h
+++ b/Engine/Source/Runtime/Core/Public/Misc/OutputDeviceRedirector.h
@@ -2,8 +2,9 @@
 
 #pragma once
 
-#include "CoreTypes.h"
 #include "Containers/Array.h"
+#include "CoreMinimal.h"
+#include "CoreTypes.h"
 #include "Misc/EnumClassFlags.h"
 #include "Misc/OutputDevice.h"
 #include "Templates/PimplPtr.h"
diff --git a/Engine/Source/Runtime/Core/Public/Misc/Parse.h b/Engine/Source/Runtime/Core/Public/Misc/Parse.h
--- a/Engine/Source/Runtime/Core/Public/Misc/Parse.h
+++ b/Engine/Source/Runtime/Core/Public/Misc/Parse.h
@@ -2,9 +2,14 @@
 
 #pragma once
 
-#include "CoreTypes.h"
 #include "Containers/StringFwd.h"
 #include "Containers/UnrealString.h"
+#include "CoreTypes.h"
+
+class FExec;
+class FName;
+class FOutputDevice;
+class FText;
 
 /*-----------------------------------------------------------------------------
 	Parsing functions.
diff --git a/Engine/Source/Runtime/Core/Public/Misc/PathViews.h b/Engine/Source/Runtime/Core/Public/Misc/PathViews.h
--- a/Engine/Source/Runtime/Core/Public/Misc/PathViews.h
+++ b/Engine/Source/Runtime/Core/Public/Misc/PathViews.h
@@ -2,11 +2,12 @@
 
 #pragma once
 
-#include "CoreTypes.h"
 #include "Containers/StringFwd.h"
+#include "CoreTypes.h"
 
 class FName;
 class FString;
+template <typename CharType> class TStringBuilderBase;
 template <typename FuncType> class TFunctionRef;
 
 class CORE_API FPathViews
diff --git a/Engine/Source/Runtime/Core/Public/Misc/Paths.h b/Engine/Source/Runtime/Core/Public/Misc/Paths.h
--- a/Engine/Source/Runtime/Core/Public/Misc/Paths.h
+++ b/Engine/Source/Runtime/Core/Public/Misc/Paths.h
@@ -2,9 +2,13 @@
 
 #pragma once
 
-#include "CoreTypes.h"
+#include "Containers/Array.h"
 #include "Containers/UnrealString.h"
+#include "CoreTypes.h"
 #include "HAL/CriticalSection.h"
+#include "Templates/UnrealTemplate.h"
+
+class FText;
 
 namespace UE::Core::Private
 {
diff --git a/Engine/Source/Runtime/Core/Public/Misc/PreLoadFile.h b/Engine/Source/Runtime/Core/Public/Misc/PreLoadFile.h
--- a/Engine/Source/Runtime/Core/Public/Misc/PreLoadFile.h
+++ b/Engine/Source/Runtime/Core/Public/Misc/PreLoadFile.h
@@ -5,6 +5,8 @@
 #include "CoreMinimal.h"
 #include "Misc/DelayedAutoRegister.h"
 
+class FEvent;
+
 struct CORE_API FPreLoadFile : public FDelayedAutoRegisterHelper
 {
 	FPreLoadFile(const TCHAR* InPath);
diff --git a/Engine/Source/Runtime/Core/Public/Misc/PreloadableFile.h b/Engine/Source/Runtime/Core/Public/Misc/PreloadableFile.h
--- a/Engine/Source/Runtime/Core/Public/Misc/PreloadableFile.h
+++ b/Engine/Source/Runtime/Core/Public/Misc/PreloadableFile.h
@@ -2,22 +2,24 @@
 
 #pragma once
 
-#include "CoreTypes.h"
 #include "Containers/Array.h"
 #include "Containers/Map.h"
+#include "Containers/StringFwd.h"
 #include "Containers/StringView.h"
 #include "Containers/UnrealString.h"
+#include "CoreTypes.h"
 #include "HAL/CriticalSection.h"
+#include "Misc/AssertionMacros.h"
 #include "Serialization/Archive.h"
 #include "Templates/Atomic.h"
 #include "Templates/Function.h"
 #include "Templates/SharedPointer.h"
 #include "Templates/UniquePtr.h"
 
+class FEvent;
 class IAsyncReadFileHandle;
 class IAsyncReadRequest;
 class IFileHandle;
-class FEvent;
 
 #define FPRELOADABLEFILE_TEST_ENABLED 0
 
diff --git a/Engine/Source/Runtime/Core/Public/Misc/QueuedThreadPoolWrapper.h b/Engine/Source/Runtime/Core/Public/Misc/QueuedThreadPoolWrapper.h
--- a/Engine/Source/Runtime/Core/Public/Misc/QueuedThreadPoolWrapper.h
+++ b/Engine/Source/Runtime/Core/Public/Misc/QueuedThreadPoolWrapper.h
@@ -2,16 +2,25 @@
 
 #pragma once
 
+#include "Async/Fundamental/Scheduler.h"
+#include "Async/Fundamental/Task.h"
+#include "Async/TaskGraphInterfaces.h"
 #include "CoreMinimal.h"
+#include "Experimental/ConcurrentLinearAllocator.h"
+#include "Experimental/Containers/FAAArrayQueue.h"
+#include "Experimental/Misc/ExecutionResource.h"
 #include "GenericPlatform/GenericPlatformAffinity.h"
+#include "HAL/PlatformAffinity.h"
+#include "HAL/PlatformProcess.h"
+#include "HAL/ThreadSafeCounter.h"
+#include "Misc/IQueuedWork.h"
+#include "Misc/MemStack.h"
+#include "Misc/ScopeLock.h"
 #include "QueuedThreadPool.h"
 #include "ScopeRWLock.h"
-#include "Async/TaskGraphInterfaces.h"
-#include "Misc/IQueuedWork.h"
-#include "Async/Fundamental/Scheduler.h"
-#include "Experimental/Misc/ExecutionResource.h"
-#include "Experimental/Containers/FAAArrayQueue.h"
-#include "Experimental/ConcurrentLinearAllocator.h"
+#include "Stats/Stats2.h"
+#include "Templates/Atomic.h"
+#include "Templates/RefCounting.h"
 
 #include <atomic>
 
diff --git a/Engine/Source/Runtime/Core/Public/Misc/RemoteConfigIni.h b/Engine/Source/Runtime/Core/Public/Misc/RemoteConfigIni.h
--- a/Engine/Source/Runtime/Core/Public/Misc/RemoteConfigIni.h
+++ b/Engine/Source/Runtime/Core/Public/Misc/RemoteConfigIni.h
@@ -2,13 +2,16 @@
 
 #pragma once
 
-#include "CoreTypes.h"
 #include "Containers/Array.h"
-#include "Containers/UnrealString.h"
 #include "Containers/Map.h"
+#include "Containers/UnrealString.h"
+#include "CoreTypes.h"
+#include "HAL/CriticalSection.h"
 #include "Misc/DateTime.h"
 #include "Stats/Stats.h"
+#include "Stats/Stats2.h"
 
+class FConfigFile;
 template<typename TTask> class FAsyncTask;
 
 /**
diff --git a/Engine/Source/Runtime/Core/Public/Misc/RuntimeErrors.h b/Engine/Source/Runtime/Core/Public/Misc/RuntimeErrors.h
--- a/Engine/Source/Runtime/Core/Public/Misc/RuntimeErrors.h
+++ b/Engine/Source/Runtime/Core/Public/Misc/RuntimeErrors.h
@@ -3,11 +3,14 @@
 #pragma once
 
 #include "CoreTypes.h"
-#include "Misc/AssertionMacros.h"
-#include "Logging/LogVerbosity.h"
+#include "Delegates/Delegate.h"
 #include "HAL/PlatformMisc.h"
+#include "Logging/LogVerbosity.h"
+#include "Misc/AssertionMacros.h"
 #include "Misc/CoreDelegates.h"
 
+class FText;
+
 // By default we bubble up runtime errors to the editor or log under the same circumstances as we assert with check/ensure
 #if !defined(UE_RAISE_RUNTIME_ERRORS)
 #define UE_RAISE_RUNTIME_ERRORS DO_CHECK
diff --git a/Engine/Source/Runtime/Core/Public/Misc/SecureHash.h b/Engine/Source/Runtime/Core/Public/Misc/SecureHash.h
--- a/Engine/Source/Runtime/Core/Public/Misc/SecureHash.h
+++ b/Engine/Source/Runtime/Core/Public/Misc/SecureHash.h
@@ -2,20 +2,26 @@
 
 #pragma once
 
-#include "CoreTypes.h"
-#include "HAL/UnrealMemory.h"
-#include "Containers/UnrealString.h"
+#include "Async/AsyncWork.h"
 #include "Containers/Map.h"
 #include "Containers/StringConv.h"
 #include "Containers/StringFwd.h"
 #include "Containers/StringView.h"
-#include "Stats/Stats.h"
-#include "Async/AsyncWork.h"
+#include "Containers/UnrealString.h"
+#include "CoreMinimal.h"
+#include "CoreTypes.h"
+#include "HAL/UnrealMemory.h"
 #include "Serialization/BufferReader.h"
+#include "Serialization/MemoryLayout.h"
+#include "Stats/Stats.h"
+#include "Stats/Stats2.h"
 #include "String/BytesToHex.h"
 #include "String/HexToBytes.h"
-#include "Serialization/MemoryLayout.h"
 
+class FMemoryImageWriter;
+class FMemoryUnfreezeContent;
+class FPointerTableBase;
+class FSHA1;
 struct FMD5Hash;
 
 /*-----------------------------------------------------------------------------
diff --git a/Engine/Source/Runtime/Core/Public/Misc/SlowTask.h b/Engine/Source/Runtime/Core/Public/Misc/SlowTask.h
--- a/Engine/Source/Runtime/Core/Public/Misc/SlowTask.h
+++ b/Engine/Source/Runtime/Core/Public/Misc/SlowTask.h
@@ -2,9 +2,12 @@
 
 #pragma once
 
+#include "CoreGlobals.h"
 #include "CoreTypes.h"
-#include "Misc/Optional.h"
 #include "Internationalization/Text.h"
+#include "Misc/Optional.h"
+
+class FFeedbackContext;
 
 /** Enum to specify a particular slow task section should be shown */
 enum class ESlowTaskVisibility
diff --git a/Engine/Source/Runtime/Core/Public/Misc/StringBuilder.h b/Engine/Source/Runtime/Core/Public/Misc/StringBuilder.h
--- a/Engine/Source/Runtime/Core/Public/Misc/StringBuilder.h
+++ b/Engine/Source/Runtime/Core/Public/Misc/StringBuilder.h
@@ -2,21 +2,28 @@
 
 #pragma once
 
-#include "CoreTypes.h"
 #include "Containers/StringFwd.h"
 #include "Containers/StringView.h"
+#include "CoreTypes.h"
+#include "HAL/PlatformString.h"
+#include "HAL/UnrealMemory.h"
+#include "Misc/AssertionMacros.h"
 #include "Misc/CString.h"
 #include "Templates/AndOrNot.h"
 #include "Templates/EnableIf.h"
 #include "Templates/IsArrayOrRefOfTypeByPredicate.h"
 #include "Templates/IsValidVariadicFunctionArg.h"
 #include "Templates/UnrealTemplate.h"
-#include "Traits/IsCharType.h"
+#include "Templates/UnrealTypeTraits.h"
 #include "Traits/IsCharEncodingCompatibleWith.h"
 #include "Traits/IsCharEncodingSimplyConvertibleTo.h"
+#include "Traits/IsCharType.h"
 #include "Traits/IsContiguousContainer.h"
+
 #include <type_traits>
 
+template <typename T> struct TIsContiguousContainer;
+
 /**
  * String Builder
  *
diff --git a/Engine/Source/Runtime/Core/Public/Misc/StringFormatter.h b/Engine/Source/Runtime/Core/Public/Misc/StringFormatter.h
--- a/Engine/Source/Runtime/Core/Public/Misc/StringFormatter.h
+++ b/Engine/Source/Runtime/Core/Public/Misc/StringFormatter.h
@@ -2,12 +2,14 @@
 
 #pragma once
 
+#include "Containers/UnrealString.h"
+#include "CoreMinimal.h"
 #include "CoreTypes.h"
 #include "Misc/AssertionMacros.h"
-#include "Containers/UnrealString.h"
-#include "Templates/ValueOrError.h"
 #include "Misc/ExpressionParserTypes.h"
+#include "Templates/ValueOrError.h"
 
+struct FStringFormatArg;
 template<typename KeyType,typename ValueType,typename SetAllocator ,typename KeyFuncs > class TMap;
 
 /** A string formatter is responsible for formatting string patterns using a set of named, or ordered arguments */
diff --git a/Engine/Source/Runtime/Core/Public/Misc/TextFilterExpressionEvaluator.h b/Engine/Source/Runtime/Core/Public/Misc/TextFilterExpressionEvaluator.h
--- a/Engine/Source/Runtime/Core/Public/Misc/TextFilterExpressionEvaluator.h
+++ b/Engine/Source/Runtime/Core/Public/Misc/TextFilterExpressionEvaluator.h
@@ -1,15 +1,19 @@
 // Copyright Epic Games, Inc. All Rights Reserved.
 #pragma once
 
-#include "CoreTypes.h"
-#include "Containers/UnrealString.h"
+#include "Containers/Array.h"
 #include "Containers/Map.h"
+#include "Containers/UnrealString.h"
+#include "CoreTypes.h"
 #include "Delegates/Delegate.h"
-#include "Misc/Optional.h"
+#include "Delegates/DelegateBase.h"
 #include "Internationalization/Text.h"
-#include "Misc/ExpressionParserTypes.h"
 #include "Misc/ExpressionParser.h"
+#include "Misc/ExpressionParserTypes.h"
+#include "Misc/Optional.h"
 #include "Misc/TextFilterUtils.h"
+#include "Templates/UnrealTemplate.h"
+#include "UObject/NameTypes.h"
 
 class FTextToken;
 
diff --git a/Engine/Source/Runtime/Core/Public/Misc/TextFilterUtils.h b/Engine/Source/Runtime/Core/Public/Misc/TextFilterUtils.h
--- a/Engine/Source/Runtime/Core/Public/Misc/TextFilterUtils.h
+++ b/Engine/Source/Runtime/Core/Public/Misc/TextFilterUtils.h
@@ -1,8 +1,11 @@
 // Copyright Epic Games, Inc. All Rights Reserved.
 #pragma once
 
-#include "CoreTypes.h"
+#include "Containers/Array.h"
 #include "Containers/UnrealString.h"
+#include "CoreTypes.h"
+#include "Misc/AssertionMacros.h"
+#include "Misc/CString.h"
 #include "UObject/NameTypes.h"
 
 /** Defines the comparison operators that can be used for a complex (key->value) comparison */
diff --git a/Engine/Source/Runtime/Core/Public/Misc/TimeGuard.h b/Engine/Source/Runtime/Core/Public/Misc/TimeGuard.h
--- a/Engine/Source/Runtime/Core/Public/Misc/TimeGuard.h
+++ b/Engine/Source/Runtime/Core/Public/Misc/TimeGuard.h
@@ -2,12 +2,12 @@
 
 #pragma once
 
-#include "CoreTypes.h"
 #include "Containers/Map.h"
 #include "CoreGlobals.h"
-#include "Misc/DateTime.h"
-#include "HAL/PlatformTime.h"
+#include "CoreTypes.h"
 #include "Delegates/Delegate.h"
+#include "HAL/PlatformTime.h"
+#include "Misc/DateTime.h"
 
 #ifndef DO_TIMEGUARD
 	// By default we are disabled, DO_TIMEGUARD can be set in XXX.Target.cs if so desired%
diff --git a/Engine/Source/Runtime/Core/Public/Misc/Timecode.h b/Engine/Source/Runtime/Core/Public/Misc/Timecode.h
--- a/Engine/Source/Runtime/Core/Public/Misc/Timecode.h
+++ b/Engine/Source/Runtime/Core/Public/Misc/Timecode.h
@@ -4,6 +4,8 @@
 
 #include "Containers/UnrealString.h"
 #include "GenericPlatform/GenericPlatformMath.h"
+#include "HAL/Platform.h"
+#include "Math/UnrealMathUtility.h"
 #include "Misc/FrameNumber.h"
 #include "Misc/FrameRate.h"
 #include "Misc/Timespan.h"
diff --git a/Engine/Source/Runtime/Core/Public/Misc/Timespan.h b/Engine/Source/Runtime/Core/Public/Misc/Timespan.h
--- a/Engine/Source/Runtime/Core/Public/Misc/Timespan.h
+++ b/Engine/Source/Runtime/Core/Public/Misc/Timespan.h
@@ -3,14 +3,15 @@
 #pragma once
 
 #include "CoreTypes.h"
+#include "Math/Interval.h"
 #include "Math/UnrealMathUtility.h"
 #include "Misc/AssertionMacros.h"
-#include "Math/Interval.h"
 
 class FArchive;
 class FOutputDevice;
 class FString;
 class UObject;
+template <typename ElementType> struct TIntervalTraits;
 
 
 /**
diff --git a/Engine/Source/Runtime/Core/Public/Misc/WildcardString.h b/Engine/Source/Runtime/Core/Public/Misc/WildcardString.h
--- a/Engine/Source/Runtime/Core/Public/Misc/WildcardString.h
+++ b/Engine/Source/Runtime/Core/Public/Misc/WildcardString.h
@@ -2,8 +2,9 @@
 
 #pragma once
 
-#include "CoreTypes.h"
 #include "Containers/UnrealString.h"
+#include "CoreTypes.h"
+#include "Misc/CString.h"
 
 /**
  * Implements a string with wild card pattern matching abilities.
diff --git a/Engine/Source/Runtime/Core/Public/Modules/ModuleManager.h b/Engine/Source/Runtime/Core/Public/Modules/ModuleManager.h
--- a/Engine/Source/Runtime/Core/Public/Modules/ModuleManager.h
+++ b/Engine/Source/Runtime/Core/Public/Modules/ModuleManager.h
@@ -2,21 +2,31 @@
 
 #pragma once
 
-#include "CoreTypes.h"
-#include "Misc/AssertionMacros.h"
-#include "Templates/UnrealTemplate.h"
 #include "Containers/Array.h"
-#include "Containers/UnrealString.h"
+#include "Containers/ContainerAllocationPolicies.h"
 #include "Containers/Map.h"
-#include "UObject/NameTypes.h"
-#include "Templates/SharedPointer.h"
-#include "Templates/UniquePtr.h"
+#include "Containers/UnrealString.h"
+#include "CoreTypes.h"
 #include "Delegates/Delegate.h"
-#include "Misc/Optional.h"
+#include "Delegates/DelegateBase.h"
+#include "HAL/CriticalSection.h"
+#include "Misc/AssertionMacros.h"
 #include "Misc/CoreMisc.h"
-#include "Modules/ModuleInterface.h"
-#include "Modules/Boilerplate/ModuleBoilerplate.h"
 #include "Misc/EnumClassFlags.h"
+#include "Misc/Optional.h"
+#include "Modules/Boilerplate/ModuleBoilerplate.h"
+#include "Modules/ModuleInterface.h"
+#include "Serialization/Archive.h"
+#include "Templates/Atomic.h"
+#include "Templates/SharedPointer.h"
+#include "Templates/SharedPointerInternals.h"
+#include "Templates/UniquePtr.h"
+#include "Templates/UnrealTemplate.h"
+#include "UObject/NameTypes.h"
+
+class FOutputDevice;
+class UClass;
+class UWorld;
 
 #if WITH_HOT_RELOAD
 	/** If true, we are reloading a class for HotReload */
diff --git a/Engine/Source/Runtime/Core/Public/ProfilingDebugging/ABTesting.h b/Engine/Source/Runtime/Core/Public/ProfilingDebugging/ABTesting.h
--- a/Engine/Source/Runtime/Core/Public/ProfilingDebugging/ABTesting.h
+++ b/Engine/Source/Runtime/Core/Public/ProfilingDebugging/ABTesting.h
@@ -6,10 +6,10 @@
 
 #pragma once
 
-#include "CoreTypes.h"
 #include "Containers/Array.h"
-#include "Containers/UnrealString.h"
 #include "Containers/BitArray.h"
+#include "Containers/UnrealString.h"
+#include "CoreTypes.h"
 #include "Math/RandomStream.h"
 
 #if !UE_BUILD_SHIPPING
diff --git a/Engine/Source/Runtime/Core/Public/ProfilingDebugging/AssetMetadataTrace.h b/Engine/Source/Runtime/Core/Public/ProfilingDebugging/AssetMetadataTrace.h
--- a/Engine/Source/Runtime/Core/Public/ProfilingDebugging/AssetMetadataTrace.h
+++ b/Engine/Source/Runtime/Core/Public/ProfilingDebugging/AssetMetadataTrace.h
@@ -2,6 +2,11 @@
 #pragma once
 
 #include "MetadataTrace.h"
+#include "ProfilingDebugging/StringsTrace.h"
+#include "Trace/Detail/EventNode.h"
+#include "Trace/Trace.h"
+
+namespace UE { namespace Trace { class FChannel; } }
 
 ////////////////////////////////////////////////////////////////////////////////////////////////////
 #if !defined(UE_TRACE_ASSET_METADATA_ENABLED) 
diff --git a/Engine/Source/Runtime/Core/Public/ProfilingDebugging/CallstackTrace.h b/Engine/Source/Runtime/Core/Public/ProfilingDebugging/CallstackTrace.h
--- a/Engine/Source/Runtime/Core/Public/ProfilingDebugging/CallstackTrace.h
+++ b/Engine/Source/Runtime/Core/Public/ProfilingDebugging/CallstackTrace.h
@@ -1,8 +1,9 @@
 // Copyright Epic Games, Inc. All Rights Reserved.
 #pragma once
-#include "TraceLog/Public/Trace/Config.h"
-#include "Math/NumericLimits.h"
+#include "HAL/Platform.h"
 #include "HAL/PlatformTLS.h"
+#include "Math/NumericLimits.h"
+#include "TraceLog/Public/Trace/Config.h"
 
 ////////////////////////////////////////////////////////////////////////////////
 #if !defined(UE_CALLSTACK_TRACE_ENABLED)
diff --git a/Engine/Source/Runtime/Core/Public/ProfilingDebugging/CookStats.h b/Engine/Source/Runtime/Core/Public/ProfilingDebugging/CookStats.h
--- a/Engine/Source/Runtime/Core/Public/ProfilingDebugging/CookStats.h
+++ b/Engine/Source/Runtime/Core/Public/ProfilingDebugging/CookStats.h
@@ -2,15 +2,22 @@
 
 #pragma once
 
-#include "CoreTypes.h"
-#include "Templates/UnrealTemplate.h"
-#include "Templates/Decay.h"
 #include "Containers/Array.h"
+#include "Containers/StringFwd.h"
 #include "Containers/UnrealString.h"
-#include "Templates/Function.h"
+#include "CoreGlobals.h"
+#include "CoreTypes.h"
 #include "Delegates/Delegate.h"
+#include "Delegates/DelegateBase.h"
+#include "Delegates/IDelegateInstance.h"
 #include "HAL/PlatformTime.h"
 #include "HAL/ThreadSafeCounter64.h"
+#include "Templates/Decay.h"
+#include "Templates/Function.h"
+#include "Templates/UnrealTemplate.h"
+#include "Trace/Trace.h"
+
+namespace UE { namespace Trace { class FChannel; } }
 
 #ifndef ENABLE_COOK_STATS
 #define ENABLE_COOK_STATS WITH_EDITOR
diff --git a/Engine/Source/Runtime/Core/Public/ProfilingDebugging/CountersTrace.h b/Engine/Source/Runtime/Core/Public/ProfilingDebugging/CountersTrace.h
--- a/Engine/Source/Runtime/Core/Public/ProfilingDebugging/CountersTrace.h
+++ b/Engine/Source/Runtime/Core/Public/ProfilingDebugging/CountersTrace.h
@@ -6,6 +6,8 @@
 #include "Trace/Config.h"
 #include "Trace/Trace.h"
 
+namespace UE { namespace Trace { class FChannel; } }
+
 #if !defined(COUNTERSTRACE_ENABLED)
 #if UE_TRACE_ENABLED && !UE_BUILD_SHIPPING
 #define COUNTERSTRACE_ENABLED 1
diff --git a/Engine/Source/Runtime/Core/Public/ProfilingDebugging/CpuProfilerTrace.h b/Engine/Source/Runtime/Core/Public/ProfilingDebugging/CpuProfilerTrace.h
--- a/Engine/Source/Runtime/Core/Public/ProfilingDebugging/CpuProfilerTrace.h
+++ b/Engine/Source/Runtime/Core/Public/ProfilingDebugging/CpuProfilerTrace.h
@@ -3,9 +3,10 @@
 #pragma once
 
 #include "CoreTypes.h"
-#include "Trace/Trace.h"
+#include "Trace/Config.h"
 #include "Trace/Detail/Channel.h"
 #include "Trace/Detail/Channel.inl"
+#include "Trace/Trace.h"
 
 #if !defined(CPUPROFILERTRACE_ENABLED)
 #if UE_TRACE_ENABLED && !UE_BUILD_SHIPPING
diff --git a/Engine/Source/Runtime/Core/Public/ProfilingDebugging/ExternalProfiler.h b/Engine/Source/Runtime/Core/Public/ProfilingDebugging/ExternalProfiler.h
--- a/Engine/Source/Runtime/Core/Public/ProfilingDebugging/ExternalProfiler.h
+++ b/Engine/Source/Runtime/Core/Public/ProfilingDebugging/ExternalProfiler.h
@@ -2,9 +2,10 @@
 
 #pragma once
 
+#include "Containers/StringConv.h"
 #include "CoreTypes.h"
-#include "UObject/NameTypes.h"
 #include "Features/IModularFeature.h"
+#include "UObject/NameTypes.h"
 
 #ifndef UE_EXTERNAL_PROFILING_ENABLED
 // temporarily turn off profiler on Switch because of a compiler issue (?) with the thread_local init bools (switching 1 to int helped, 1 did not)
diff --git a/Engine/Source/Runtime/Core/Public/ProfilingDebugging/Histogram.h b/Engine/Source/Runtime/Core/Public/ProfilingDebugging/Histogram.h
--- a/Engine/Source/Runtime/Core/Public/ProfilingDebugging/Histogram.h
+++ b/Engine/Source/Runtime/Core/Public/ProfilingDebugging/Histogram.h
@@ -8,16 +8,20 @@
 
 #pragma once
 
-#include "CoreTypes.h"
-#include "Misc/AssertionMacros.h"
 #include "Containers/Array.h"
-#include "Containers/UnrealString.h"
+#include "Containers/ArrayView.h"
 #include "Containers/Map.h"
+#include "Containers/UnrealString.h"
+#include "CoreTypes.h"
+#include "HAL/PlatformCrt.h"
 #include "Logging/LogMacros.h"
-#include "Containers/ArrayView.h"
+#include "Math/UnrealMathUtility.h"
+#include "Misc/AssertionMacros.h"
+#include "Templates/UnrealTemplate.h"
 
-struct FHistogramBuilder;
 struct FAnalyticsEventAttribute;
+struct FHistogramBuilder;
+template <typename FuncType> class TFunctionRef;
 
 DECLARE_LOG_CATEGORY_EXTERN(LogHistograms, Log, All);
 
diff --git a/Engine/Source/Runtime/Core/Public/ProfilingDebugging/InstanceCounter.h b/Engine/Source/Runtime/Core/Public/ProfilingDebugging/InstanceCounter.h
--- a/Engine/Source/Runtime/Core/Public/ProfilingDebugging/InstanceCounter.h
+++ b/Engine/Source/Runtime/Core/Public/ProfilingDebugging/InstanceCounter.h
@@ -1,11 +1,14 @@
 // Copyright Epic Games, Inc. All Rights Reserved.
 #pragma once
 
-#include "CoreTypes.h"
-#include "Containers/UnrealString.h"
 #include "Containers/Map.h"
+#include "Containers/UnrealString.h"
+#include "CoreTypes.h"
+#include "HAL/CriticalSection.h"
 #include "UObject/NameTypes.h"
 
+class FOutputDevice;
+
 /*
 	A helper object for counting instances of struct or classes for debugging
 	purposes, e.g. to get an absolute count of the number of "Foo"'s in existence
diff --git a/Engine/Source/Runtime/Core/Public/ProfilingDebugging/LoadTimeTracker.h b/Engine/Source/Runtime/Core/Public/ProfilingDebugging/LoadTimeTracker.h
--- a/Engine/Source/Runtime/Core/Public/ProfilingDebugging/LoadTimeTracker.h
+++ b/Engine/Source/Runtime/Core/Public/ProfilingDebugging/LoadTimeTracker.h
@@ -6,16 +6,18 @@
 
 #pragma once
 
-#include "CoreTypes.h"
 #include "Containers/Array.h"
 #include "Containers/Map.h"
-#include "UObject/NameTypes.h"
-#include "Trace/Trace.inl"
+#include "Containers/SparseArray.h"
+#include "CoreTypes.h"
+#include "Misc/Optional.h"
 #include "ProfilingDebugging/CpuProfilerTrace.h"
 #include "ProfilingDebugging/ScopedTimers.h"
 #include "Serialization/LoadTimeTrace.h"
 #include "Stats/Stats.h"
-#include "Misc/Optional.h"
+#include "Trace/Trace.h"
+#include "Trace/Trace.inl"
+#include "UObject/NameTypes.h"
 
 #ifndef ENABLE_LOADTIME_TRACKING
 	#define ENABLE_LOADTIME_TRACKING 0
diff --git a/Engine/Source/Runtime/Core/Public/ProfilingDebugging/MallocProfiler.h b/Engine/Source/Runtime/Core/Public/ProfilingDebugging/MallocProfiler.h
--- a/Engine/Source/Runtime/Core/Public/ProfilingDebugging/MallocProfiler.h
+++ b/Engine/Source/Runtime/Core/Public/ProfilingDebugging/MallocProfiler.h
@@ -5,15 +5,15 @@
 =============================================================================*/
 #pragma once
 
+#include "Containers/Array.h"
+#include "Containers/Map.h"
+#include "Containers/UnrealString.h"
 #include "CoreTypes.h"
-#include "Misc/AssertionMacros.h"
 #include "HAL/MemoryBase.h"
+#include "Misc/AssertionMacros.h"
+#include "Misc/CompressedGrowableBuffer.h"
 #include "Serialization/Archive.h"
-#include "Containers/Array.h"
-#include "Containers/UnrealString.h"
-#include "Containers/Map.h"
 #include "UObject/NameTypes.h"
-#include "Misc/CompressedGrowableBuffer.h"
 
 #if USE_MALLOC_PROFILER
 
diff --git a/Engine/Source/Runtime/Core/Public/ProfilingDebugging/MemoryTrace.h b/Engine/Source/Runtime/Core/Public/ProfilingDebugging/MemoryTrace.h
--- a/Engine/Source/Runtime/Core/Public/ProfilingDebugging/MemoryTrace.h
+++ b/Engine/Source/Runtime/Core/Public/ProfilingDebugging/MemoryTrace.h
@@ -1,8 +1,10 @@
 // Copyright Epic Games, Inc. All Rights Reserved.
 #pragma once
 
-#include "Trace/Trace.h"
+#include "HAL/Platform.h"
 #include "Misc/EnumClassFlags.h"
+#include "Trace/Config.h"
+#include "Trace/Trace.h"
 
 #if !defined(UE_MEMORY_TRACE_AVAILABLE)
 #   define UE_MEMORY_TRACE_AVAILABLE 0
diff --git a/Engine/Source/Runtime/Core/Public/ProfilingDebugging/MetadataTrace.h b/Engine/Source/Runtime/Core/Public/ProfilingDebugging/MetadataTrace.h
--- a/Engine/Source/Runtime/Core/Public/ProfilingDebugging/MetadataTrace.h
+++ b/Engine/Source/Runtime/Core/Public/ProfilingDebugging/MetadataTrace.h
@@ -1,8 +1,11 @@
 ﻿// Copyright Epic Games, Inc. All Rights Reserved.
 #pragma once
 
-#include "Trace/Trace.h"
 #include "ProfilingDebugging/StringsTrace.h"
+#include "Trace/Config.h"
+#include "Trace/Trace.h"
+
+namespace UE { namespace Trace { class FChannel; } }
 
 ////////////////////////////////////////////////////////////////////////////////////////////////////
 # if !defined(UE_TRACE_METADATA_ENABLED)
