Ran a pass of IWYU on Core.

diff --git a/Engine/Source/Runtime/Core/Public/ProfilingDebugging/MiscTrace.h b/Engine/Source/Runtime/Core/Public/ProfilingDebugging/MiscTrace.h
--- a/Engine/Source/Runtime/Core/Public/ProfilingDebugging/MiscTrace.h
+++ b/Engine/Source/Runtime/Core/Public/ProfilingDebugging/MiscTrace.h
@@ -2,9 +2,10 @@
 
 #pragma once
 
+#include "Containers/ContainersFwd.h"
 #include "CoreTypes.h"
-#include "Trace/Config.h"
 #include "ProfilingDebugging/FormatArgsTrace.h"
+#include "Trace/Config.h"
 
 #if !defined(MISCTRACE_ENABLED)
 #if UE_TRACE_ENABLED && !UE_BUILD_SHIPPING
diff --git a/Engine/Source/Runtime/Core/Public/ProfilingDebugging/ProfilingHelpers.h b/Engine/Source/Runtime/Core/Public/ProfilingDebugging/ProfilingHelpers.h
--- a/Engine/Source/Runtime/Core/Public/ProfilingDebugging/ProfilingHelpers.h
+++ b/Engine/Source/Runtime/Core/Public/ProfilingDebugging/ProfilingHelpers.h
@@ -8,9 +8,10 @@
 
 #pragma once
 
-#include "CoreTypes.h"
 #include "Containers/UnrealString.h"
+#include "CoreTypes.h"
 #include "Delegates/Delegate.h"
+#include "Delegates/DelegateBase.h"
 
 enum EStreamingStatus
 {
diff --git a/Engine/Source/Runtime/Core/Public/ProfilingDebugging/ResourceSize.h b/Engine/Source/Runtime/Core/Public/ProfilingDebugging/ResourceSize.h
--- a/Engine/Source/Runtime/Core/Public/ProfilingDebugging/ResourceSize.h
+++ b/Engine/Source/Runtime/Core/Public/ProfilingDebugging/ResourceSize.h
@@ -2,12 +2,14 @@
 
 #pragma once
 
+#include "Containers/Map.h"
+#include "Containers/SparseArray.h"
 #include "CoreTypes.h"
 #include "Misc/AssertionMacros.h"
-#include "Containers/Map.h"
 #include "UObject/NameTypes.h"
 
 class FName;
+class FOutputDevice;
 
 /** Indicate what types of resources should be included for calculating used memory */
 namespace EResourceSizeMode
diff --git a/Engine/Source/Runtime/Core/Public/ProfilingDebugging/StallDetector.h b/Engine/Source/Runtime/Core/Public/ProfilingDebugging/StallDetector.h
--- a/Engine/Source/Runtime/Core/Public/ProfilingDebugging/StallDetector.h
+++ b/Engine/Source/Runtime/Core/Public/ProfilingDebugging/StallDetector.h
@@ -2,10 +2,10 @@
 
 #pragma once
 
-#include "CoreTypes.h"
-
-#include "Containers/UnrealString.h"
+#include "Containers/Array.h"
 #include "Containers/Set.h"
+#include "Containers/UnrealString.h"
+#include "CoreTypes.h"
 #include "HAL/CriticalSection.h"
 #include "HAL/PreprocessorHelpers.h"
 #include "Logging/LogMacros.h"
diff --git a/Engine/Source/Runtime/Core/Public/ProfilingDebugging/TraceAuxiliary.h b/Engine/Source/Runtime/Core/Public/ProfilingDebugging/TraceAuxiliary.h
--- a/Engine/Source/Runtime/Core/Public/ProfilingDebugging/TraceAuxiliary.h
+++ b/Engine/Source/Runtime/Core/Public/ProfilingDebugging/TraceAuxiliary.h
@@ -2,10 +2,12 @@
 
 #pragma once
 
+#include "Containers/StringFwd.h"
 #include "CoreGlobals.h"
+#include "Delegates/Delegate.h"
 #include "HAL/Platform.h"
+#include "Logging/LogCategory.h"
 #include "Logging/LogMacros.h"
-#include "Delegates/Delegate.h"
 
 ////////////////////////////////////////////////////////////////////////////////
 class CORE_API FTraceAuxiliary
diff --git a/Engine/Source/Runtime/Core/Public/Serialization/Archive.h b/Engine/Source/Runtime/Core/Public/Serialization/Archive.h
--- a/Engine/Source/Runtime/Core/Public/Serialization/Archive.h
+++ b/Engine/Source/Runtime/Core/Public/Serialization/Archive.h
@@ -2,10 +2,11 @@
 
 #pragma once
 
-#include "CoreTypes.h"
 #include "CoreFwd.h"
+#include "CoreTypes.h"
 #include "HAL/PlatformProperties.h"
 #include "Internationalization/TextNamespaceFwd.h"
+#include "Math/MathFwd.h"
 #include "Misc/AssertionMacros.h"
 #include "Misc/CompressionFlags.h"
 #include "Misc/EngineVersionBase.h"
@@ -18,25 +19,28 @@
 #include "Templates/IsEnumClass.h"
 #include "Templates/IsSigned.h"
 #include "Templates/IsValidVariadicFunctionArg.h"
+#include "Templates/UnrealTemplate.h"
 #include "Traits/IsCharEncodingCompatibleWith.h"
 #include "UObject/ObjectVersion.h"
 
 class FArchive;
+class FBulkData;
 class FCustomVersionContainer;
 class FLinker;
 class FName;
+class FProperty;
 class FString;
 class FText;
 class ITargetPlatform;
 class UObject;
-class FProperty;
-class FBulkData;
 namespace UE::Serialization{ class FEditorBulkData; }
 struct FArchiveSerializedPropertyChain;
 template<class TEnum> class TEnumAsByte;
+
 typedef TFunction<bool (double RemainingTime)> FExternalReadCallback;
-struct FUObjectSerializeContext;
 class FField;
+struct FUObjectSerializeContext;
+
 enum class EFileRegionType : uint8;
 struct FArchiveCookContext;
 
diff --git a/Engine/Source/Runtime/Core/Public/Serialization/ArchiveLoadCompressedProxy.h b/Engine/Source/Runtime/Core/Public/Serialization/ArchiveLoadCompressedProxy.h
--- a/Engine/Source/Runtime/Core/Public/Serialization/ArchiveLoadCompressedProxy.h
+++ b/Engine/Source/Runtime/Core/Public/Serialization/ArchiveLoadCompressedProxy.h
@@ -2,10 +2,12 @@
 
 #pragma once
 
+#include "Containers/Array.h"
 #include "CoreTypes.h"
-#include "UObject/NameTypes.h"
-#include "Serialization/Archive.h"
 #include "Misc/Compression.h"
+#include "Misc/CompressionFlags.h"
+#include "Serialization/Archive.h"
+#include "UObject/NameTypes.h"
 
 /*----------------------------------------------------------------------------
 	FArchiveLoadCompressedProxy.
diff --git a/Engine/Source/Runtime/Core/Public/Serialization/ArchiveSaveCompressedProxy.h b/Engine/Source/Runtime/Core/Public/Serialization/ArchiveSaveCompressedProxy.h
--- a/Engine/Source/Runtime/Core/Public/Serialization/ArchiveSaveCompressedProxy.h
+++ b/Engine/Source/Runtime/Core/Public/Serialization/ArchiveSaveCompressedProxy.h
@@ -2,10 +2,12 @@
 
 #pragma once
 
+#include "Containers/Array.h"
 #include "CoreTypes.h"
-#include "UObject/NameTypes.h"
-#include "Serialization/Archive.h"
 #include "Misc/Compression.h"
+#include "Misc/CompressionFlags.h"
+#include "Serialization/Archive.h"
+#include "UObject/NameTypes.h"
 
 /*----------------------------------------------------------------------------
 	FArchiveSaveCompressedProxy.
diff --git a/Engine/Source/Runtime/Core/Public/Serialization/BitArchive.h b/Engine/Source/Runtime/Core/Public/Serialization/BitArchive.h
--- a/Engine/Source/Runtime/Core/Public/Serialization/BitArchive.h
+++ b/Engine/Source/Runtime/Core/Public/Serialization/BitArchive.h
@@ -3,8 +3,9 @@
 #pragma once
 
 #include "CoreTypes.h"
-#include "Serialization/Archive.h"
 #include "HAL/IConsoleManager.h"
+#include "Misc/AssertionMacros.h"
+#include "Serialization/Archive.h"
 
 /** CVar specifying the maximum serialization size for strings sent/received by the netcode */
 extern CORE_API TAutoConsoleVariable<int32> CVarMaxNetStringSize;
diff --git a/Engine/Source/Runtime/Core/Public/Serialization/BitReader.h b/Engine/Source/Runtime/Core/Public/Serialization/BitReader.h
--- a/Engine/Source/Runtime/Core/Public/Serialization/BitReader.h
+++ b/Engine/Source/Runtime/Core/Public/Serialization/BitReader.h
@@ -2,11 +2,13 @@
 
 #pragma once
 
+#include "Containers/Array.h"
 #include "CoreTypes.h"
-#include "Misc/AssertionMacros.h"
 #include "HAL/UnrealMemory.h"
+#include "Misc/AssertionMacros.h"
 #include "Serialization/BitArchive.h"
-#include "Containers/Array.h"
+
+class FArchive;
 
 CORE_API void appBitsCpy( uint8* Dest, int32 DestBit, uint8* Src, int32 SrcBit, int32 BitCount );
 
diff --git a/Engine/Source/Runtime/Core/Public/Serialization/BitWriter.h b/Engine/Source/Runtime/Core/Public/Serialization/BitWriter.h
--- a/Engine/Source/Runtime/Core/Public/Serialization/BitWriter.h
+++ b/Engine/Source/Runtime/Core/Public/Serialization/BitWriter.h
@@ -2,11 +2,13 @@
 
 #pragma once
 
+#include "Containers/Array.h"
 #include "CoreTypes.h"
+#include "Math/UnrealMathUtility.h"
 #include "Misc/AssertionMacros.h"
 #include "Serialization/BitArchive.h"
-#include "Containers/Array.h"
-#include "Math/UnrealMathUtility.h"
+
+class FArchive;
 
 /*-----------------------------------------------------------------------------
 	FBitWriter.
diff --git a/Engine/Source/Runtime/Core/Public/Serialization/CompactBinary.h b/Engine/Source/Runtime/Core/Public/Serialization/CompactBinary.h
--- a/Engine/Source/Runtime/Core/Public/Serialization/CompactBinary.h
+++ b/Engine/Source/Runtime/Core/Public/Serialization/CompactBinary.h
@@ -2,16 +2,24 @@
 
 #pragma once
 
-#include "CoreTypes.h"
 #include "Containers/StringFwd.h"
 #include "Containers/StringView.h"
+#include "CoreTypes.h"
+#include "HAL/UnrealMemory.h"
 #include "IO/IoHash.h"
+#include "Memory/CompositeBuffer.h"
 #include "Memory/MemoryFwd.h"
 #include "Memory/MemoryView.h"
 #include "Memory/SharedBuffer.h"
+#include "Misc/AssertionMacros.h"
 #include "Misc/EnumClassFlags.h"
 #include "String/BytesToHex.h"
+#include "Templates/Function.h"
 #include "Templates/IsTriviallyDestructible.h"
+#include "Templates/RemoveReference.h"
+#include "Templates/UnrealTemplate.h"
+
+template <typename CharType> class TStringBuilderBase;
 
 /**
  * This file declares a compact binary data format that is compatible with JSON and only slightly
@@ -64,7 +72,6 @@
  */
 
 class FArchive;
-class FIoHashBuilder;
 class FCbArrayView;
 class FCbField;
 class FCbFieldIterator;
@@ -73,6 +80,7 @@ class FCbFieldViewIterator;
 class FCbObjectId;
 class FCbObjectView;
 class FCbValue;
+class FIoHashBuilder;
 struct FDateTime;
 struct FGuid;
 struct FTimespan;
diff --git a/Engine/Source/Runtime/Core/Public/Serialization/CompactBinaryPackage.h b/Engine/Source/Runtime/Core/Public/Serialization/CompactBinaryPackage.h
--- a/Engine/Source/Runtime/Core/Public/Serialization/CompactBinaryPackage.h
+++ b/Engine/Source/Runtime/Core/Public/Serialization/CompactBinaryPackage.h
@@ -2,14 +2,17 @@
 
 #pragma once
 
-#include "CoreTypes.h"
 #include "Compression/CompressedBuffer.h"
 #include "Containers/Array.h"
+#include "Containers/ContainersFwd.h"
+#include "CoreTypes.h"
 #include "IO/IoHash.h"
 #include "Memory/CompositeBuffer.h"
 #include "Memory/SharedBuffer.h"
+#include "Misc/AssertionMacros.h"
 #include "Misc/TVariant.h"
 #include "Serialization/CompactBinary.h"
+#include "Templates/Function.h"
 #include "Templates/UnrealTemplate.h"
 
 class FArchive;
diff --git a/Engine/Source/Runtime/Core/Public/Serialization/CompactBinarySerialization.h b/Engine/Source/Runtime/Core/Public/Serialization/CompactBinarySerialization.h
--- a/Engine/Source/Runtime/Core/Public/Serialization/CompactBinarySerialization.h
+++ b/Engine/Source/Runtime/Core/Public/Serialization/CompactBinarySerialization.h
@@ -2,13 +2,17 @@
 
 #pragma once
 
-#include "CoreTypes.h"
 #include "Containers/StringFwd.h"
+#include "CoreTypes.h"
+#include "IO/IoHash.h"
+#include "Memory/MemoryFwd.h"
 #include "Memory/MemoryView.h"
 #include "Memory/SharedBuffer.h"
 #include "Serialization/CompactBinary.h"
+#include "Templates/Function.h"
 
 class FArchive;
+struct FGuid;
 
 ///////////////////////////////////////////////////////////////////////////////////////////////////////////////////////
 
diff --git a/Engine/Source/Runtime/Core/Public/Serialization/CompactBinaryValidation.h b/Engine/Source/Runtime/Core/Public/Serialization/CompactBinaryValidation.h
--- a/Engine/Source/Runtime/Core/Public/Serialization/CompactBinaryValidation.h
+++ b/Engine/Source/Runtime/Core/Public/Serialization/CompactBinaryValidation.h
@@ -3,6 +3,7 @@
 #pragma once
 
 #include "CoreTypes.h"
+#include "Memory/MemoryFwd.h"
 #include "Memory/MemoryView.h"
 #include "Misc/EnumClassFlags.h"
 #include "Serialization/CompactBinary.h"
diff --git a/Engine/Source/Runtime/Core/Public/Serialization/CompactBinaryWriter.h b/Engine/Source/Runtime/Core/Public/Serialization/CompactBinaryWriter.h
--- a/Engine/Source/Runtime/Core/Public/Serialization/CompactBinaryWriter.h
+++ b/Engine/Source/Runtime/Core/Public/Serialization/CompactBinaryWriter.h
@@ -2,11 +2,19 @@
 
 #pragma once
 
-#include "CoreTypes.h"
+#include "Containers/Array.h"
+#include "Containers/ContainersFwd.h"
+#include "Containers/StringFwd.h"
 #include "Containers/StringView.h"
+#include "CoreTypes.h"
+#include "HAL/PlatformCrt.h"
 #include "Memory/MemoryFwd.h"
+#include "Memory/MemoryView.h"
 #include "Serialization/CompactBinary.h"
 
+class FCompositeBuffer;
+class FSharedBuffer;
+
 ///////////////////////////////////////////////////////////////////////////////////////////////////////////////////////
 
 class FArchive;
diff --git a/Engine/Source/Runtime/Core/Public/Serialization/CompressedChunkInfo.h b/Engine/Source/Runtime/Core/Public/Serialization/CompressedChunkInfo.h
--- a/Engine/Source/Runtime/Core/Public/Serialization/CompressedChunkInfo.h
+++ b/Engine/Source/Runtime/Core/Public/Serialization/CompressedChunkInfo.h
@@ -2,8 +2,10 @@
 
 #pragma once
 
-#include "CoreTypes.h"
 #include "CoreFwd.h"
+#include "CoreTypes.h"
+
+class FArchive;
 
 /**
  * Implements a helper structure for compression support
diff --git a/Engine/Source/Runtime/Core/Public/Serialization/CustomVersion.h b/Engine/Source/Runtime/Core/Public/Serialization/CustomVersion.h
--- a/Engine/Source/Runtime/Core/Public/Serialization/CustomVersion.h
+++ b/Engine/Source/Runtime/Core/Public/Serialization/CustomVersion.h
@@ -2,15 +2,18 @@
 
 #pragma once
 
+#include "Containers/Array.h"
+#include "Containers/Set.h"
+#include "Containers/UnrealString.h"
 #include "CoreTypes.h"
-#include "Templates/UnrealTemplate.h"
 #include "Misc/Crc.h"
-#include "Containers/UnrealString.h"
-#include "Containers/Set.h"
-#include "UObject/NameTypes.h"
 #include "Misc/Guid.h"
+#include "Misc/Optional.h"
 #include "Serialization/StructuredArchive.h"
+#include "Templates/UnrealTemplate.h"
+#include "UObject/NameTypes.h"
 
+class FArchive;
 class FCustomVersionContainer;
 struct FCustomVersion;
 
diff --git a/Engine/Source/Runtime/Core/Public/Serialization/FileRegions.h b/Engine/Source/Runtime/Core/Public/Serialization/FileRegions.h
--- a/Engine/Source/Runtime/Core/Public/Serialization/FileRegions.h
+++ b/Engine/Source/Runtime/Core/Public/Serialization/FileRegions.h
@@ -2,9 +2,10 @@
 
 #pragma once
 
-#include "PixelFormat.h"
 #include "Containers/Array.h"
 #include "Containers/ArrayView.h"
+#include "HAL/Platform.h"
+#include "PixelFormat.h"
 
 // Describes the type of data found in a given file region.
 enum class EFileRegionType : uint8
diff --git a/Engine/Source/Runtime/Core/Public/Serialization/Formatters/BinaryArchiveFormatter.h b/Engine/Source/Runtime/Core/Public/Serialization/Formatters/BinaryArchiveFormatter.h
--- a/Engine/Source/Runtime/Core/Public/Serialization/Formatters/BinaryArchiveFormatter.h
+++ b/Engine/Source/Runtime/Core/Public/Serialization/Formatters/BinaryArchiveFormatter.h
@@ -2,9 +2,21 @@
 
 #pragma once
 
-#include "Serialization/StructuredArchiveFormatter.h"
-#include "Serialization/Archive.h"
 #include "Containers/Array.h"
+#include "HAL/Platform.h"
+#include "Serialization/Archive.h"
+#include "Serialization/StructuredArchiveFormatter.h"
+#include "Serialization/StructuredArchiveNameHelpers.h"
+
+class FName;
+class FString;
+class FText;
+class UObject;
+struct FLazyObjectPtr;
+struct FObjectPtr;
+struct FSoftObjectPath;
+struct FSoftObjectPtr;
+struct FWeakObjectPtr;
 
 class CORE_API FBinaryArchiveFormatter final : public FStructuredArchiveFormatter
 {
diff --git a/Engine/Source/Runtime/Core/Public/Serialization/LargeMemoryData.h b/Engine/Source/Runtime/Core/Public/Serialization/LargeMemoryData.h
--- a/Engine/Source/Runtime/Core/Public/Serialization/LargeMemoryData.h
+++ b/Engine/Source/Runtime/Core/Public/Serialization/LargeMemoryData.h
@@ -2,8 +2,10 @@
 
 #pragma once
 
-#include "CoreTypes.h"
 #include "Containers/LockFreeList.h"
+#include "CoreTypes.h"
+
+template <class T, int TPaddingForCacheContention> class TLockFreePointerListUnordered;
 
 /**
 * Data storage for the large memory reader and writer.
diff --git a/Engine/Source/Runtime/Core/Public/Serialization/LargeMemoryReader.h b/Engine/Source/Runtime/Core/Public/Serialization/LargeMemoryReader.h
--- a/Engine/Source/Runtime/Core/Public/Serialization/LargeMemoryReader.h
+++ b/Engine/Source/Runtime/Core/Public/Serialization/LargeMemoryReader.h
@@ -2,11 +2,12 @@
 
 #pragma once
 
-#include "CoreTypes.h"
 #include "Containers/UnrealString.h"
-#include "UObject/NameTypes.h"
+#include "CoreTypes.h"
 #include "Misc/EnumClassFlags.h"
 #include "Serialization/MemoryArchive.h"
+#include "UObject/NameTypes.h"
+#include "UObject/UnrealNames.h"
 
 enum class ELargeMemoryReaderFlags : uint8
 {
diff --git a/Engine/Source/Runtime/Core/Public/Serialization/LargeMemoryWriter.h b/Engine/Source/Runtime/Core/Public/Serialization/LargeMemoryWriter.h
--- a/Engine/Source/Runtime/Core/Public/Serialization/LargeMemoryWriter.h
+++ b/Engine/Source/Runtime/Core/Public/Serialization/LargeMemoryWriter.h
@@ -2,12 +2,13 @@
 
 #pragma once
 
-#include "CoreTypes.h"
 #include "Containers/UnrealString.h"
+#include "CoreTypes.h"
+#include "Memory/MemoryFwd.h"
 #include "Memory/MemoryView.h"
-#include "UObject/NameTypes.h"
-#include "Serialization/MemoryArchive.h"
 #include "Serialization/LargeMemoryData.h"
+#include "Serialization/MemoryArchive.h"
+#include "UObject/NameTypes.h"
 
 /**
 * Archive for storing a large amount of arbitrary data to memory
diff --git a/Engine/Source/Runtime/Core/Public/Serialization/LoadTimeTrace.h b/Engine/Source/Runtime/Core/Public/Serialization/LoadTimeTrace.h
--- a/Engine/Source/Runtime/Core/Public/Serialization/LoadTimeTrace.h
+++ b/Engine/Source/Runtime/Core/Public/Serialization/LoadTimeTrace.h
@@ -3,9 +3,11 @@
 #pragma once
 
 #include "CoreTypes.h"
+#include "ProfilingDebugging/FormatArgsTrace.h"
 #include "Trace/Config.h"
 #include "Trace/Trace.h"
-#include "ProfilingDebugging/FormatArgsTrace.h"
+
+namespace UE { namespace Trace { class FChannel; } }
 
 #if !defined(LOADTIMEPROFILERTRACE_ENABLED)
 #if UE_TRACE_ENABLED && !UE_BUILD_SHIPPING
diff --git a/Engine/Source/Runtime/Core/Public/Serialization/MappedName.h b/Engine/Source/Runtime/Core/Public/Serialization/MappedName.h
--- a/Engine/Source/Runtime/Core/Public/Serialization/MappedName.h
+++ b/Engine/Source/Runtime/Core/Public/Serialization/MappedName.h
@@ -2,6 +2,7 @@
 
 #pragma once
 
+#include "Containers/ArrayView.h"
 #include "CoreMinimal.h"
 #include "UObject/NameBatchSerialization.h"
 
diff --git a/Engine/Source/Runtime/Core/Public/Serialization/MemoryImage.h b/Engine/Source/Runtime/Core/Public/Serialization/MemoryImage.h
--- a/Engine/Source/Runtime/Core/Public/Serialization/MemoryImage.h
+++ b/Engine/Source/Runtime/Core/Public/Serialization/MemoryImage.h
@@ -2,15 +2,16 @@
 
 #pragma once
 
-#include "CoreTypes.h"
 #include "Containers/Array.h"
-#include "Containers/Set.h"
 #include "Containers/HashTable.h"
+#include "Containers/Set.h"
 #include "Containers/UnrealString.h"
+#include "CoreMinimal.h"
+#include "CoreTypes.h"
 #include "Misc/SecureHash.h"
 #include "Serialization/Archive.h"
-#include "Serialization/MemoryLayout.h"
 #include "Serialization/MemoryImageWriter.h"
+#include "Serialization/MemoryLayout.h"
 #include "Templates/RefCounting.h"
 #include "Traits/IsCharType.h"
 
diff --git a/Engine/Source/Runtime/Core/Public/Serialization/NameAsStringIndexProxyArchive.h b/Engine/Source/Runtime/Core/Public/Serialization/NameAsStringIndexProxyArchive.h
--- a/Engine/Source/Runtime/Core/Public/Serialization/NameAsStringIndexProxyArchive.h
+++ b/Engine/Source/Runtime/Core/Public/Serialization/NameAsStringIndexProxyArchive.h
@@ -2,11 +2,13 @@
 
 #pragma once
 
-#include "CoreTypes.h"
 #include "Containers/Array.h"
 #include "Containers/Set.h"
-#include "UObject/NameTypes.h"
+#include "CoreTypes.h"
 #include "Serialization/ArchiveProxy.h"
+#include "UObject/NameTypes.h"
+
+class FArchive;
 
 /**
  * Implements a proxy archive that serializes FNames as string data or an index (if the same name is repeated).
diff --git a/Engine/Source/Runtime/Core/Public/Serialization/StructuredArchive.h b/Engine/Source/Runtime/Core/Public/Serialization/StructuredArchive.h
--- a/Engine/Source/Runtime/Core/Public/Serialization/StructuredArchive.h
+++ b/Engine/Source/Runtime/Core/Public/Serialization/StructuredArchive.h
@@ -2,17 +2,22 @@
 
 #pragma once
 
+#include "Containers/Array.h"
+#include "Containers/ContainerAllocationPolicies.h"
 #include "CoreTypes.h"
-#include "Serialization/StructuredArchiveFwd.h"
-#include "Serialization/StructuredArchiveDefines.h"
+#include "Formatters/BinaryArchiveFormatter.h"
+#include "Serialization/Archive.h"
 #include "Serialization/StructuredArchiveAdapters.h"
+#include "Serialization/StructuredArchiveDefines.h"
 #include "Serialization/StructuredArchiveFormatter.h"
+#include "Serialization/StructuredArchiveFwd.h"
 #include "Serialization/StructuredArchiveNameHelpers.h"
+#include "Serialization/StructuredArchiveSlotBase.h"
 #include "Serialization/StructuredArchiveSlots.h"
-#include "Formatters/BinaryArchiveFormatter.h"
-#include "Containers/Array.h"
 #include "Templates/UniqueObj.h"
 
+class FStructuredArchive;
+
 namespace UE::StructuredArchive::Private
 {
 	enum class EElementType : unsigned char
diff --git a/Engine/Source/Runtime/Core/Public/Serialization/StructuredArchiveChildReader.h b/Engine/Source/Runtime/Core/Public/Serialization/StructuredArchiveChildReader.h
--- a/Engine/Source/Runtime/Core/Public/Serialization/StructuredArchiveChildReader.h
+++ b/Engine/Source/Runtime/Core/Public/Serialization/StructuredArchiveChildReader.h
@@ -2,9 +2,14 @@
 
 #pragma once
 
+#include "HAL/Platform.h"
 #include "Misc/Optional.h"
 #include "Serialization/StructuredArchive.h"
+#include "Serialization/StructuredArchiveSlots.h"
 
+class FStructuredArchive;
+class FStructuredArchiveFormatter;
+
 /**
  * FStructuredArchiveChildReader
  *
diff --git a/Engine/Source/Runtime/Core/Public/Serialization/StructuredArchiveFormatter.h b/Engine/Source/Runtime/Core/Public/Serialization/StructuredArchiveFormatter.h
--- a/Engine/Source/Runtime/Core/Public/Serialization/StructuredArchiveFormatter.h
+++ b/Engine/Source/Runtime/Core/Public/Serialization/StructuredArchiveFormatter.h
@@ -2,10 +2,17 @@
 
 #pragma once
 
-#include "CoreTypes.h"
+#include "Containers/Array.h"
 #include "CoreFwd.h"
+#include "CoreTypes.h"
 #include "Serialization/StructuredArchiveNameHelpers.h"
 
+class FArchive;
+class FName;
+class FString;
+class FText;
+class UObject;
+
 /**
  * Specifies the type of a value in a slot. Used by FContextFreeArchiveFormatter for introspection.
  */
diff --git a/Engine/Source/Runtime/Core/Public/Serialization/StructuredArchiveSlots.h b/Engine/Source/Runtime/Core/Public/Serialization/StructuredArchiveSlots.h
--- a/Engine/Source/Runtime/Core/Public/Serialization/StructuredArchiveSlots.h
+++ b/Engine/Source/Runtime/Core/Public/Serialization/StructuredArchiveSlots.h
@@ -2,20 +2,34 @@
 
 #pragma once
 
+#include "Containers/Array.h"
 #include "CoreTypes.h"
+#include "Formatters/BinaryArchiveFormatter.h"
 #include "Misc/Optional.h"
+#include "Serialization/Archive.h"
 #include "Serialization/StructuredArchiveFwd.h"
 #include "Serialization/StructuredArchiveNameHelpers.h"
 #include "Serialization/StructuredArchiveSlotBase.h"
-#include "Formatters/BinaryArchiveFormatter.h"
+#include "Templates/EnableIf.h"
+#include "Templates/IsEnumClass.h"
 
+class FName;
+class FString;
 class FStructuredArchive;
+class FStructuredArchiveArray;
 class FStructuredArchiveChildReader;
-class FStructuredArchiveSlot;
+class FStructuredArchiveMap;
 class FStructuredArchiveRecord;
-class FStructuredArchiveArray;
+class FStructuredArchiveSlot;
 class FStructuredArchiveStream;
-class FStructuredArchiveMap;
+class FText;
+class UObject;
+struct FLazyObjectPtr;
+struct FObjectPtr;
+struct FSoftObjectPath;
+struct FSoftObjectPtr;
+struct FWeakObjectPtr;
+template <class TEnum> class TEnumAsByte;
 
 namespace UE::StructuredArchive::Private
 {
diff --git a/Engine/Source/Runtime/Core/Public/Stats/Stats2.h b/Engine/Source/Runtime/Core/Public/Stats/Stats2.h
--- a/Engine/Source/Runtime/Core/Public/Stats/Stats2.h
+++ b/Engine/Source/Runtime/Core/Public/Stats/Stats2.h
@@ -2,27 +2,31 @@
 
 #pragma once
 
-#include "CoreTypes.h"
-#include "Templates/UnrealTemplate.h"
-#include "HAL/ThreadSafeCounter.h"
-#include "Math/NumericLimits.h"
-#include "HAL/ThreadSingleton.h"
-#include "HAL/LowLevelMemTracker.h"
 #include "Containers/Array.h"
-#include "Containers/UnrealString.h"
-#include "HAL/PlatformTime.h"
-#include "UObject/NameTypes.h"
-#include "Containers/LockFreeList.h"
 #include "Containers/ChunkedArray.h"
+#include "Containers/LockFreeList.h"
+#include "Containers/UnrealString.h"
+#include "CoreMinimal.h"
+#include "CoreTypes.h"
 #include "Delegates/Delegate.h"
-#include "Templates/Atomic.h"
+#include "HAL/LowLevelMemTracker.h"
+#include "HAL/PlatformTime.h"
+#include "HAL/ThreadSafeCounter.h"
+#include "HAL/ThreadSingleton.h"
 #include "Math/Color.h"
-#include "StatsCommon.h"
-#include "Templates/UniquePtr.h"
+#include "Math/NumericLimits.h"
+#include "Misc/EnumClassFlags.h"
 #include "ProfilingDebugging/CpuProfilerTrace.h"
 #include "ProfilingDebugging/MiscTrace.h"
+#include "StatsCommon.h"
 #include "StatsTrace.h"
-#include "Misc/EnumClassFlags.h"
+#include "Templates/Atomic.h"
+#include "Templates/UniquePtr.h"
+#include "Templates/UnrealTemplate.h"
+#include "Trace/Detail/Channel.h"
+#include "Trace/Detail/Channel.inl"
+#include "Trace/Trace.h"
+#include "UObject/NameTypes.h"
 
 class FScopeCycleCounter;
 class FThreadStats;
diff --git a/Engine/Source/Runtime/Core/Public/Stats/StatsData.h b/Engine/Source/Runtime/Core/Public/Stats/StatsData.h
--- a/Engine/Source/Runtime/Core/Public/Stats/StatsData.h
+++ b/Engine/Source/Runtime/Core/Public/Stats/StatsData.h
@@ -2,21 +2,25 @@
 
 #pragma once
 
-#include "CoreTypes.h"
-#include "Misc/AssertionMacros.h"
-#include "Templates/UnrealTemplate.h"
-#include "Math/NumericLimits.h"
 #include "Containers/Array.h"
-#include "Math/UnrealMathUtility.h"
-#include "Containers/UnrealString.h"
-#include "Containers/Set.h"
+#include "Containers/IndirectArray.h"
 #include "Containers/Map.h"
-#include "Misc/Parse.h"
-#include "UObject/NameTypes.h"
+#include "Containers/Set.h"
+#include "Containers/StringFwd.h"
+#include "Containers/UnrealString.h"
+#include "CoreMinimal.h"
+#include "CoreTypes.h"
 #include "Delegates/Delegate.h"
 #include "HAL/ThreadSafeCounter.h"
-#include "Containers/IndirectArray.h"
+#include "Math/NumericLimits.h"
+#include "Math/UnrealMathUtility.h"
+#include "Misc/AssertionMacros.h"
+#include "Misc/Parse.h"
 #include "Stats/Stats.h"
+#include "Stats/Stats2.h"
+#include "Stats/StatsCommon.h"
+#include "Templates/UnrealTemplate.h"
+#include "UObject/NameTypes.h"
 
 struct FEventData;
 
diff --git a/Engine/Source/Runtime/Core/Public/Stats/StatsFile.h b/Engine/Source/Runtime/Core/Public/Stats/StatsFile.h
--- a/Engine/Source/Runtime/Core/Public/Stats/StatsFile.h
+++ b/Engine/Source/Runtime/Core/Public/Stats/StatsFile.h
@@ -2,26 +2,29 @@
 
 #pragma once
 
-#include "CoreTypes.h"
-#include "Misc/AssertionMacros.h"
-#include "Misc/Compression.h"
+#include "Async/AsyncWork.h"
 #include "Containers/Array.h"
-#include "Containers/UnrealString.h"
-#include "Containers/Set.h"
+#include "Containers/ChunkedArray.h"
 #include "Containers/Map.h"
-#include "UObject/NameTypes.h"
-#include "Logging/LogMacros.h"
+#include "Containers/Set.h"
+#include "Containers/UnrealString.h"
+#include "CoreMinimal.h"
+#include "CoreTypes.h"
 #include "Delegates/IDelegateInstance.h"
-#include "Internationalization/Text.h"
+#include "HAL/ThreadSafeBool.h"
 #include "HAL/ThreadSafeCounter.h"
+#include "Internationalization/Text.h"
+#include "Logging/LogMacros.h"
+#include "Misc/AssertionMacros.h"
+#include "Misc/Compression.h"
 #include "Stats/Stats.h"
-#include "HAL/ThreadSafeBool.h"
-#include "Async/AsyncWork.h"
+#include "Stats/Stats2.h"
+#include "Stats/StatsCommon.h"
 #include "Stats/StatsData.h"
+#include "UObject/NameTypes.h"
 
 class FAsyncStatsWrite;
 struct FStatsReadFile;
-
 template< typename InElementType, typename KeyFuncs , typename Allocator > class TSet;
 template<typename KeyType,typename ValueType,typename SetAllocator ,typename KeyFuncs > class TMap;
 
diff --git a/Engine/Source/Runtime/Core/Public/Stats/StatsHierarchical.h b/Engine/Source/Runtime/Core/Public/Stats/StatsHierarchical.h
--- a/Engine/Source/Runtime/Core/Public/Stats/StatsHierarchical.h
+++ b/Engine/Source/Runtime/Core/Public/Stats/StatsHierarchical.h
@@ -2,10 +2,16 @@
 
 #pragma once
 
+#include "Containers/Array.h"
+#include "Containers/UnrealString.h"
 #include "CoreTypes.h"
-#include "UObject/NameTypes.h"
 #include "HAL/PlatformTime.h"
 #include "Logging/MessageLog.h"
+#include "Templates/SharedPointer.h"
+#include "UObject/NameTypes.h"
+
+class FMessageLog;
+class FStatsHierarchical;
 
 /**
  *	An element in a profiling / stats tree
diff --git a/Engine/Source/Runtime/Core/Public/Stats/StatsMallocProfilerProxy.h b/Engine/Source/Runtime/Core/Public/Stats/StatsMallocProfilerProxy.h
--- a/Engine/Source/Runtime/Core/Public/Stats/StatsMallocProfilerProxy.h
+++ b/Engine/Source/Runtime/Core/Public/Stats/StatsMallocProfilerProxy.h
@@ -6,6 +6,10 @@
 #include "HAL/MemoryBase.h"
 #include "HAL/ThreadSafeCounter.h"
 
+class FOutputDevice;
+class UWorld;
+struct FGenericMemoryStats;
+
 #if STATS
 
 /**
diff --git a/Engine/Source/Runtime/Core/Public/String/BytesToHex.h b/Engine/Source/Runtime/Core/Public/String/BytesToHex.h
--- a/Engine/Source/Runtime/Core/Public/String/BytesToHex.h
+++ b/Engine/Source/Runtime/Core/Public/String/BytesToHex.h
@@ -2,9 +2,10 @@
 
 #pragma once
 
-#include "CoreTypes.h"
 #include "Containers/ArrayView.h"
+#include "Containers/ContainersFwd.h"
 #include "Containers/StringFwd.h"
+#include "CoreTypes.h"
 
 namespace UE::String
 {
diff --git a/Engine/Source/Runtime/Core/Public/String/Find.h b/Engine/Source/Runtime/Core/Public/String/Find.h
--- a/Engine/Source/Runtime/Core/Public/String/Find.h
+++ b/Engine/Source/Runtime/Core/Public/String/Find.h
@@ -2,9 +2,11 @@
 
 #pragma once
 
-#include "CoreTypes.h"
 #include "Containers/ArrayView.h"
+#include "Containers/ContainersFwd.h"
 #include "Containers/StringFwd.h"
+#include "CoreTypes.h"
+#include "Misc/CString.h"
 
 namespace UE::String
 {
diff --git a/Engine/Source/Runtime/Core/Public/String/ParseLines.h b/Engine/Source/Runtime/Core/Public/String/ParseLines.h
--- a/Engine/Source/Runtime/Core/Public/String/ParseLines.h
+++ b/Engine/Source/Runtime/Core/Public/String/ParseLines.h
@@ -2,11 +2,14 @@
 
 #pragma once
 
-#include "CoreTypes.h"
+#include "Containers/StringFwd.h"
 #include "Containers/StringView.h"
+#include "CoreTypes.h"
 #include "Misc/EnumClassFlags.h"
 #include "Templates/Function.h"
 
+template <typename FuncType> class TFunctionRef;
+
 namespace UE::String
 {
 
diff --git a/Engine/Source/Runtime/Core/Public/String/ParseTokens.h b/Engine/Source/Runtime/Core/Public/String/ParseTokens.h
--- a/Engine/Source/Runtime/Core/Public/String/ParseTokens.h
+++ b/Engine/Source/Runtime/Core/Public/String/ParseTokens.h
@@ -2,12 +2,16 @@
 
 #pragma once
 
-#include "CoreTypes.h"
 #include "Containers/ArrayView.h"
+#include "Containers/ContainersFwd.h"
+#include "Containers/StringFwd.h"
 #include "Containers/StringView.h"
+#include "CoreTypes.h"
 #include "Misc/EnumClassFlags.h"
 #include "Templates/Function.h"
 
+template <typename FuncType> class TFunctionRef;
+
 namespace UE::String
 {
 
diff --git a/Engine/Source/Runtime/Core/Public/Tasks/Pipe.h b/Engine/Source/Runtime/Core/Public/Tasks/Pipe.h
--- a/Engine/Source/Runtime/Core/Public/Tasks/Pipe.h
+++ b/Engine/Source/Runtime/Core/Public/Tasks/Pipe.h
@@ -2,8 +2,12 @@
 
 #pragma once
 
-#include "Tasks/Task.h"
+#include "Async/Fundamental/Scheduler.h"
+#include "Async/Fundamental/Task.h"
+#include "CoreMinimal.h"
 #include "CoreTypes.h"
+#include "Tasks/Task.h"
+#include "Tasks/TaskPrivate.h"
 
 #include <atomic>
 
diff --git a/Engine/Source/Runtime/Core/Public/Templates/RefCounting.h b/Engine/Source/Runtime/Core/Public/Templates/RefCounting.h
--- a/Engine/Source/Runtime/Core/Public/Templates/RefCounting.h
+++ b/Engine/Source/Runtime/Core/Public/Templates/RefCounting.h
@@ -3,10 +3,12 @@
 #pragma once
 
 #include "CoreTypes.h"
+#include "HAL/PlatformAtomics.h"
+#include "HAL/ThreadSafeCounter.h"
 #include "Misc/AssertionMacros.h"
 #include "Serialization/Archive.h"
 #include "Serialization/MemoryLayout.h"
-#include "HAL/PlatformAtomics.h"
+#include "Templates/TypeHash.h"
 
 /** A virtual interface for ref counted objects to implement. */
 class IRefCountedObject
diff --git a/Engine/Source/Runtime/Core/Public/Virtualization/VirtualizationSystem.h b/Engine/Source/Runtime/Core/Public/Virtualization/VirtualizationSystem.h
--- a/Engine/Source/Runtime/Core/Public/Virtualization/VirtualizationSystem.h
+++ b/Engine/Source/Runtime/Core/Public/Virtualization/VirtualizationSystem.h
@@ -3,12 +3,26 @@
 #pragma once
 
 #include "Compression/CompressedBuffer.h"
+#include "Containers/Array.h"
+#include "Containers/ArrayView.h"
 #include "Containers/StringFwd.h"
-#include "IO/IoHash.h"
+#include "Containers/StringView.h"
+#include "Containers/UnrealString.h"
+#include "CoreTypes.h"
+#include "Delegates/Delegate.h"
+#include "Delegates/DelegateBase.h"
+#include "Delegates/IDelegateInstance.h"
 #include "Features/IModularFeature.h"
 #include "Features/IModularFeatures.h"
+#include "IO/IoHash.h"
+#include "Misc/ConfigCacheIni.h"
+#include "Templates/Function.h"
+#include "Templates/UniquePtr.h"
+#include "UObject/NameTypes.h"
 
 class FPackagePath;
+class FText;
+class UObject;
 
 namespace UE::Virtualization
 {
diff --git a/Engine/Source/Runtime/Core/Tests/Containers/TestUtils.cpp b/Engine/Source/Runtime/Core/Tests/Containers/TestUtils.cpp
--- a/Engine/Source/Runtime/Core/Tests/Containers/TestUtils.cpp
+++ b/Engine/Source/Runtime/Core/Tests/Containers/TestUtils.cpp
@@ -1,6 +1,9 @@
 // Copyright Epic Games, Inc. All Rights Reserved.
 
 #include "TestUtils.h"
+
+#include "Math/UnrealMathUtility.h"
+
 #include <numeric>
 
 namespace Test
diff --git a/Engine/Source/Runtime/Core/Tests/Containers/TestUtils.h b/Engine/Source/Runtime/Core/Tests/Containers/TestUtils.h
--- a/Engine/Source/Runtime/Core/Tests/Containers/TestUtils.h
+++ b/Engine/Source/Runtime/Core/Tests/Containers/TestUtils.h
@@ -3,7 +3,6 @@
 #pragma once
 
 #include "CoreMinimal.h"
-#include "Math/UnrealMathUtility.h"
 
 namespace Test
 {
