Fix for memory corruption when converting a null-terminated string via StringConv<UTF8CHAR>(Ptr) where the string ends with two high surrogate pairs.

diff --git a/Engine/Source/Runtime/Core/Private/GenericPlatform/GenericPlatformString.cpp b/Engine/Source/Runtime/Core/Private/GenericPlatform/GenericPlatformString.cpp
--- a/Engine/Source/Runtime/Core/Private/GenericPlatform/GenericPlatformString.cpp
+++ b/Engine/Source/Runtime/Core/Private/GenericPlatform/GenericPlatformString.cpp
@@ -227,26 +227,26 @@ namespace UE::Core::Private
 				// Check if this character is a high-surrogate
 				if (IsHighSurrogate(Codepoint))
 				{
-					if (SourceLen == 0)
+					if (bHighSurrogateIsSet)
 					{
-						// String ends with lone high-surrogate - write out surrogate (will be converted into bogus character)
-						if (!WriteCodepointToBuffer(Codepoint, Dest, DestLen))
+						// Already have a high-surrogate in this pair - write our stored value (will be converted into bogus character)
+						if (!WriteCodepointToBuffer(HighSurrogate, Dest, DestLen))
 						{
 							// Could not write data, bail out
 							return -1;
 						}
-
-						return UE_PTRDIFF_TO_INT32(Dest - DestStartingPosition);
 					}
 
-					if (bHighSurrogateIsSet)
+					if (SourceLen == 0)
 					{
-						// Already have a high-surrogate in this pair - write our stored value (will be converted into bogus character)
-						if (!WriteCodepointToBuffer(HighSurrogate, Dest, DestLen))
+						// String ends with lone high-surrogate - write out surrogate (will be converted into bogus character)
+						if (!WriteCodepointToBuffer(Codepoint, Dest, DestLen))
 						{
 							// Could not write data, bail out
 							return -1;
 						}
+
+						return UE_PTRDIFF_TO_INT32(Dest - DestStartingPosition);
 					}
 
 					// Store our code point for our next character
