Temp revert changes to FMath::{Abs, Min, Max} due to undesirable interactions w/ multi-dimensional vector types.
Will add back in changes after putting in type traits to check for this edge case.

diff --git a/Engine/Source/Runtime/Core/Public/GenericPlatform/GenericPlatformMath.h b/Engine/Source/Runtime/Core/Public/GenericPlatform/GenericPlatformMath.h
--- a/Engine/Source/Runtime/Core/Public/GenericPlatform/GenericPlatformMath.h
+++ b/Engine/Source/Runtime/Core/Public/GenericPlatform/GenericPlatformMath.h
@@ -919,7 +919,7 @@ struct FGenericPlatformMath
 	template< class T > 
 	static CONSTEXPR FORCEINLINE T Abs( const T A )
 	{
-		return (A < (T)0) ? -A : A;
+		return (A>=(T)0) ? A : -A;
 	}
 
 	/** Returns 1, 0, or -1 depending on relation of T to 0 */
@@ -933,14 +933,14 @@ struct FGenericPlatformMath
 	template< class T > 
 	static CONSTEXPR FORCEINLINE T Max( const T A, const T B )
 	{
-		return (B < A) ? A : B;
+		return (A>=B) ? A : B;
 	}
 
 	/** Returns lower value in a generic way */
 	template< class T > 
 	static CONSTEXPR FORCEINLINE T Min( const T A, const T B )
 	{
-		return (A < B) ? A : B;
+		return (A<=B) ? A : B;
 	}
 
 	// Allow mixing of float types to promote to highest precision type
