Deprecate platform chunk install's FCustomChunk API in favor of simpler 'named chunk' API.

diff --git a/Engine/Source/Runtime/Core/Private/GenericPlatform/GenericPlatformChunkInstall.cpp b/Engine/Source/Runtime/Core/Private/GenericPlatform/GenericPlatformChunkInstall.cpp
--- a/Engine/Source/Runtime/Core/Private/GenericPlatform/GenericPlatformChunkInstall.cpp
+++ b/Engine/Source/Runtime/Core/Private/GenericPlatform/GenericPlatformChunkInstall.cpp
@@ -1,5 +1,225 @@
 // Copyright Epic Games, Inc. All Rights Reserved.
 
 #include "GenericPlatform/GenericPlatformChunkInstall.h"
+#include "HAL/PlatformMisc.h"
 
 DEFINE_LOG_CATEGORY(LogChunkInstaller)
+
+
+
+/*
+ *  Functionality to to implement the new 'named chunk' API using the existing, deprecated 'FCustomChunk' API
+ *  Once the FCustomChunk API is removed then these will be replaced by platform native implementations
+ */
+
+PRAGMA_DISABLE_DEPRECATION_WARNINGS
+
+bool FNamedChunkPlatformChunkInstall::IsNamedChunkInProgress(const FName NamedChunk)
+{
+	TArray<FCustomChunk> CustomChunks = GetCustomChunksFromNamedChunk(NamedChunk);
+	return IsChunkInstallationPending(CustomChunks);
+}
+
+
+bool FNamedChunkPlatformChunkInstall::InstallNamedChunk(const FName NamedChunk)
+{
+	TArray<FCustomChunk> CustomChunks = GetCustomChunksFromNamedChunk(NamedChunk);
+	return InstallChunks(CustomChunks);
+}
+
+
+bool FNamedChunkPlatformChunkInstall::UninstallNamedChunk(const FName NamedChunk)
+{
+	TArray<FCustomChunk> CustomChunks = GetCustomChunksFromNamedChunk(NamedChunk);
+	return UninstallChunks(CustomChunks);
+}
+
+
+bool FNamedChunkPlatformChunkInstall::InstallNamedChunks(const TArrayView<FName>& NamedChunks)
+{
+	TArray<FCustomChunk> CustomChunks = GetCustomChunksFromNamedChunks(NamedChunks);
+	return InstallChunks(CustomChunks);
+}
+
+
+bool FNamedChunkPlatformChunkInstall::UninstallNamedChunks(const TArrayView<FName>& NamedChunks)
+{
+	TArray<FCustomChunk> CustomChunks = GetCustomChunksFromNamedChunks(NamedChunks);
+	return UninstallChunks(CustomChunks);
+}
+
+
+bool FNamedChunkPlatformChunkInstall::PrioritizeNamedChunk(const FName NamedChunk, EChunkPriority::Type Priority)
+{
+	FCustomChunk CustomChunk(TEXT(""), 0xFFFFFFFF, ECustomChunkType::OnDemandChunk);
+	if (TryGetCustomChunkFromNamedChunk(NamedChunk, CustomChunk))
+	{
+		return PrioritizeChunk(CustomChunk.ChunkID, Priority);
+	}
+
+	return false;
+}
+
+
+EChunkLocation::Type FNamedChunkPlatformChunkInstall::GetNamedChunkLocation(const FName NamedChunk)
+{
+	FCustomChunk CustomChunk(TEXT(""), 0xFFFFFFFF, ECustomChunkType::OnDemandChunk);
+	if (TryGetCustomChunkFromNamedChunk(NamedChunk, CustomChunk))
+	{
+		return GetChunkLocation(CustomChunk.ChunkID);
+	}
+
+	return EChunkLocation::DoesNotExist;
+}
+
+
+float FNamedChunkPlatformChunkInstall::GetNamedChunkProgress(const FName NamedChunk, EChunkProgressReportingType::Type ReportType)
+{
+	FCustomChunk CustomChunk(TEXT(""), 0xFFFFFFFF, ECustomChunkType::OnDemandChunk);
+	if (TryGetCustomChunkFromNamedChunk(NamedChunk, CustomChunk))
+	{
+		return GetCustomChunkProgress(CustomChunk, ReportType);
+	}
+
+	return 0.0f;
+}
+
+
+ENamedChunkType FNamedChunkPlatformChunkInstall::GetNamedChunkType(const FName NamedChunk) const
+{
+	FCustomChunk CustomChunk(TEXT(""), 0xFFFFFFFF, ECustomChunkType::OnDemandChunk);
+	if (TryGetCustomChunkFromNamedChunk(NamedChunk, CustomChunk))
+	{
+		if (CustomChunk.ChunkType == ECustomChunkType::OnDemandChunk)
+		{
+			return ENamedChunkType::OnDemand;
+		}
+		else if (CustomChunk.ChunkType == ECustomChunkType::LanguageChunk)
+		{
+			return ENamedChunkType::Language;
+		}
+	}
+
+	return ENamedChunkType::Invalid;
+}
+
+
+TArray<FName> FNamedChunkPlatformChunkInstall::GetNamedChunksByType(ENamedChunkType NamedChunkType) const
+{
+	TArray<FCustomChunk> CustomChunks;
+	if (NamedChunkType == ENamedChunkType::OnDemand)
+	{
+		CustomChunks = FPlatformMisc::GetCustomChunksByType(ECustomChunkType::OnDemandChunk);
+	}
+	else if (NamedChunkType == ENamedChunkType::Language)
+	{
+		CustomChunks = FPlatformMisc::GetCustomChunksByType(ECustomChunkType::LanguageChunk);
+	}
+
+	TArray<FName> NamedChunks = GetNamedChunksFromCustomChunks(CustomChunks);
+	return MoveTemp(NamedChunks);
+}
+
+
+TArray<FCustomChunk> FNamedChunkPlatformChunkInstall::GetCustomChunksFromNamedChunk(const FName NamedChunk) const
+{
+	TArray<FCustomChunk> CustomChunks;
+	
+	FCustomChunk CustomChunk(TEXT(""), 0xFFFFFFFF, ECustomChunkType::OnDemandChunk);
+	if (TryGetCustomChunkFromNamedChunk(NamedChunk, CustomChunk))
+	{
+		CustomChunks.Emplace(CustomChunk.ChunkTag, CustomChunk.ChunkID, CustomChunk.ChunkType, CustomChunk.ChunkTag2);
+	}
+
+	return MoveTemp(CustomChunks);
+}
+
+TArray<FCustomChunk> FNamedChunkPlatformChunkInstall::GetCustomChunksFromNamedChunks(const TArrayView<FName>& NamedChunks) const
+{
+	TArray<FCustomChunk> CustomChunks;
+
+	for (FName NamedChunk : NamedChunks)
+	{
+		CustomChunks.Append(GetCustomChunksFromNamedChunk(NamedChunk));
+	}
+
+	return MoveTemp(CustomChunks);
+}
+
+
+bool FNamedChunkPlatformChunkInstall::TryGetCustomChunkFromNamedChunk(const FName NamedChunk, FCustomChunk& OutCustomChunk) const
+{
+	for (const FCustomChunk& CustomChunk : FPlatformMisc::GetAllOnDemandChunks())
+	{
+		if (FName(CustomChunk.ChunkTag) == NamedChunk || FName(CustomChunk.ChunkTag2) == NamedChunk)
+		{
+			OutCustomChunk = FCustomChunk(CustomChunk.ChunkTag, CustomChunk.ChunkID, CustomChunk.ChunkType, CustomChunk.ChunkTag2);
+			return true;
+		}
+	}
+
+	for (const FCustomChunk& CustomChunk : FPlatformMisc::GetAllLanguageChunks())
+	{
+		if ( FName(CustomChunk.ChunkTag) == NamedChunk || FName(CustomChunk.ChunkTag2) == NamedChunk)
+		{
+			OutCustomChunk = FCustomChunk(CustomChunk.ChunkTag, CustomChunk.ChunkID, CustomChunk.ChunkType, CustomChunk.ChunkTag2);
+			return true;
+		}
+	}
+
+	UE_LOG(LogChunkInstaller, Warning, TEXT("Named chunk %s not found"), *NamedChunk.ToString());
+	return false;
+}
+
+
+TArray<FName> FNamedChunkPlatformChunkInstall::GetNamedChunksFromCustomChunks(const TArray<FCustomChunk>& CustomChunks) const
+{
+	TArray<FName> NamedChunks;
+	NamedChunks.Reserve(CustomChunks.Num());
+
+	for (const FCustomChunk& CustomChunk : CustomChunks)
+	{
+		FName ChunkName = GetCustomChunkName(CustomChunk);
+		if (ChunkName != NAME_None)
+		{
+			NamedChunks.Add(FName(ChunkName));
+		}
+	}
+
+	return MoveTemp(NamedChunks);
+}
+
+
+FName FNamedChunkPlatformChunkInstall::GetNamedChunkByPakChunkIndex(int32 InPakchunkIndex) const
+{
+	TArray<int32> PakchunkIndices;
+	PakchunkIndices.Add(InPakchunkIndex);
+
+	TArray<FCustomChunk> CustomChunks = FPlatformMisc::GetOnDemandChunksForPakchunkIndices(PakchunkIndices);
+	if (CustomChunks.Num() == 0)
+	{
+		return NAME_None;
+	}
+	
+	check(CustomChunks.Num() == 1); // a pak file cannot be in multiple chunks
+	return GetCustomChunkName(CustomChunks[0]);
+}
+
+
+FName FNamedChunkPlatformChunkInstall::GetCustomChunkName(const FCustomChunk& CustomChunk) const
+{
+	if (!CustomChunk.ChunkTag.IsEmpty())
+	{
+		return FName(CustomChunk.ChunkTag);
+	}
+	else if (!CustomChunk.ChunkTag2.IsEmpty())
+	{
+		return FName(CustomChunk.ChunkTag2);
+	}
+	else
+	{
+		return NAME_None;
+	}
+}
+
+PRAGMA_ENABLE_DEPRECATION_WARNINGS
diff --git a/Engine/Source/Runtime/Core/Private/GenericPlatform/GenericPlatformMisc.cpp b/Engine/Source/Runtime/Core/Private/GenericPlatform/GenericPlatformMisc.cpp
--- a/Engine/Source/Runtime/Core/Private/GenericPlatform/GenericPlatformMisc.cpp
+++ b/Engine/Source/Runtime/Core/Private/GenericPlatform/GenericPlatformMisc.cpp
@@ -1805,6 +1805,7 @@
 
 TArray<FCustomChunk> FGenericPlatformMisc::GetCustomChunksByType(ECustomChunkType DesiredChunkType)
 {
+	PRAGMA_DISABLE_DEPRECATION_WARNINGS
 	if (DesiredChunkType == ECustomChunkType::OnDemandChunk)
 	{
 		return GetAllOnDemandChunks();
@@ -1813,6 +1814,7 @@
 	{
 		return GetAllLanguageChunks();
 	}
+	PRAGMA_ENABLE_DEPRECATION_WARNINGS
 }
 
 FString FGenericPlatformMisc::LoadTextFileFromPlatformPackage(const FString& RelativePath)
diff --git a/Engine/Source/Runtime/Core/Public/GenericPlatform/GenericPlatformChunkInstall.h b/Engine/Source/Runtime/Core/Public/GenericPlatform/GenericPlatformChunkInstall.h
--- a/Engine/Source/Runtime/Core/Public/GenericPlatform/GenericPlatformChunkInstall.h
+++ b/Engine/Source/Runtime/Core/Public/GenericPlatform/GenericPlatformChunkInstall.h
@@ -87,11 +87,12 @@ enum class ECustomChunkType : uint8
 struct FCustomChunk
 {
 	FString ChunkTag;
+	FString ChunkTag2;
 	uint32	ChunkID;
 	ECustomChunkType ChunkType;
 
-	FCustomChunk(FString InTag, uint32 InID, ECustomChunkType InChunkType) :
-		ChunkTag(InTag), ChunkID(InID), ChunkType(InChunkType)
+	FCustomChunk(FString InTag, uint32 InID, ECustomChunkType InChunkType, FString InTag2 = TEXT("")) :
+		ChunkTag(InTag), ChunkTag2(InTag2), ChunkID(InID), ChunkType(InChunkType)
 	{}
 };
 
@@ -112,6 +113,13 @@ struct FCustomChunkMapping
 	{}
 };
 
+enum class ENamedChunkType : uint8
+{
+	Invalid,
+	OnDemand,
+	Language,
+};
+
 /**
 * Interface for platform specific chunk based install
 **/
@@ -199,30 +207,98 @@ public:
 	UE_DEPRECATED(4.18, "Call RemoveChunkInstallDelegate instead")
 	virtual void RemoveChunkInstallDelgate( uint32 ChunkID, FDelegateHandle Delegate ) = 0;
 
-	/**
-	* Check whether current platform supports intelligent chunk installation
-	* @return				whether Intelligent Install is supported
-	*/
+
+	UE_DEPRECATED(5.2, "Call SupportsNamedChunkInstall instead")
 	virtual bool SupportsIntelligentInstall() = 0;
 
-	/**
-	* Check whether installation of chunks are pending
-	* @return				whether installation task has been kicked
-	*/
+	UE_DEPRECATED(5.2, "Call IsNamedChunkInProgress instead")
 	virtual bool IsChunkInstallationPending(const TArray<FCustomChunk>& ChunkTagsID) = 0;
 
-	/**
-	* Install chunks with Intelligent Delivery API
-	* @return				whether installation task has been kicked
-	*/
+	UE_DEPRECATED(5.2, "Call InstallNamedChunks instead")
 	virtual bool InstallChunks(const TArray<FCustomChunk>& ChunkTagsID) = 0;
 
-	/**
-	* Uninstall chunks with Intelligent Delivery API
-	* @return				whether uninstallation task has been kicked
-	*/
+	UE_DEPRECATED(5.2, "Call UninstallNamedChunks instead")
 	virtual bool UninstallChunks(const TArray<FCustomChunk>& ChunkTagsID) = 0;
 
+
+	/**
+	 * Check whether current platform supports chunk installation by name
+	 * @return				whether Intelligent Install is supported
+	 */
+	virtual bool SupportsNamedChunkInstall() const = 0;
+
+	/**
+	 * Check whether the give chunk is being installed
+	 * @param NamedChunk	The name of the chunk
+	 * @return				whether installation task has been kicked
+	 */
+	virtual bool IsNamedChunkInProgress(const FName NamedChunk) = 0;
+
+	/**
+	 * Install the given named chunk
+	 * @param NamedChunk	The name of the chunk
+	 * @return				whether installation task has been kicked
+	 **/
+	virtual bool InstallNamedChunk(const FName NamedChunk) = 0;
+
+	/**
+	 * Uninstall the given named chunk
+	 * @param NamedChunk	The name of the chunk
+	 * @return				whether uninstallation task has been kicked
+	 **/
+	virtual bool UninstallNamedChunk(const FName NamedChunk) = 0;
+
+	/**
+	 * Install the given set of named chunks
+	 * @param NamedChunks	The names of the chunks to install
+	 * @return				whether installation task has been kicked
+	 **/
+	virtual bool InstallNamedChunks(const TArrayView<FName>& NamedChunks) = 0;
+
+	/**
+	 * Uninstall the given set of named chunks
+	 * @param NamedChunk	The names of the chunks to uninstall
+	 * @return				whether uninstallation task has been kicked
+	 **/
+	virtual bool UninstallNamedChunks(const TArrayView<FName>& NamedChunks) = 0;
+
+	/**
+	 * Get the current location of the given named chunk
+	 * @param NamedChunk	The name of the chunk
+	 * @return				Enum specifying whether the chunk is available to use, waiting to install, or does not exist.
+	 **/
+	virtual EChunkLocation::Type GetNamedChunkLocation(const FName NamedChunk) = 0;
+
+	/**
+	 * Get the current install progress of the given named chunk.  Let the user specify report type for platforms that support more than one.
+	 * @param NamedChunk	The name of the chunk
+	 * @param ReportType	The type of progress report you want.
+	 * @return				A value whose meaning is dependent on the ReportType param.
+	 **/
+	virtual float GetNamedChunkProgress(const FName NamedChunk, EChunkProgressReportingType::Type ReportType) = 0;
+
+	/**
+	 * Hint to the installer that we would like to prioritize a specific chunk
+	 * @param NamedChunk	The name of the chunk
+	 * @param Priority		The priority for the chunk.
+	 * @return				false if the operation is not allowed or the chunk doesn't exist, otherwise true.
+	 **/
+	virtual bool PrioritizeNamedChunk(const FName NamedChunk, EChunkPriority::Type Priority) = 0;
+
+	/** 
+	 * Query the type of the given named chunk
+	 * @param NamedChunk	The name of the chunk
+	 * @return				Enum indicating the type of chunk, if any
+	 */
+	virtual ENamedChunkType GetNamedChunkType(const FName NamedChunk) const = 0;
+
+	/**
+	 * Get a list of all the named chunks of the given type
+	 * @param				Enum indicating the type of chunk
+	 * @return				Array containing all named chunks of the given type
+	 */
+	virtual TArray<FName> GetNamedChunksByType(ENamedChunkType NamedChunkType) const = 0;
+
 protected:
 		/**
 		 * Get the current location of a chunk.
@@ -343,6 +419,61 @@ public:
 		return false;
 	}
 
+	virtual bool SupportsNamedChunkInstall() const override
+	{
+		return false;
+	}
+
+	virtual bool IsNamedChunkInProgress(const FName NamedChunk) override
+	{
+		return false;
+	}
+
+	virtual bool InstallNamedChunk(const FName NamedChunk) override
+	{
+		return false;
+	}
+
+	virtual bool UninstallNamedChunk(const FName NamedChunk) override
+	{
+		return false;
+	}
+
+	virtual bool InstallNamedChunks(const TArrayView<FName>& NamedChunks) override
+	{
+		return false;
+	}
+
+	virtual bool UninstallNamedChunks(const TArrayView<FName>& NamedChunks) override
+	{
+		return false;
+	}
+
+	virtual EChunkLocation::Type GetNamedChunkLocation(const FName NamedChunk) override
+	{
+		return EChunkLocation::NotAvailable;
+	}
+
+	virtual float GetNamedChunkProgress(const FName NamedChunk, EChunkProgressReportingType::Type ReportType) override
+	{
+		return 0.0f;
+	}
+
+	virtual bool PrioritizeNamedChunk(const FName NamedChunk, EChunkPriority::Type Priority) override
+	{
+		return false;
+	}
+
+	virtual ENamedChunkType GetNamedChunkType(const FName NamedChunk) const override
+	{
+		return ENamedChunkType::Invalid;
+	}
+
+	virtual TArray<FName> GetNamedChunksByType(ENamedChunkType NamedChunkType) const override
+	{
+		return TArray<FName>();
+	}
+
 protected:
 
 	/** Delegate called when installation succeeds or fails */
@@ -354,4 +485,36 @@ protected:
 	}
 };
 
+// temporary helper base for platform chunk installers to transition from FCustomChunk to named chunks
+class FNamedChunkPlatformChunkInstall : public FGenericPlatformChunkInstall
+{
+public:
+	virtual bool SupportsNamedChunkInstall() const override
+	{
+		return true;
+	}
+
+	virtual bool IsNamedChunkInProgress(const FName NamedChunk) override;
+	virtual bool InstallNamedChunk(const FName NamedChunk) override;
+	virtual bool UninstallNamedChunk(const FName NamedChunk) override;
+	virtual bool InstallNamedChunks(const TArrayView<FName>& NamedChunks) override;
+	virtual bool UninstallNamedChunks(const TArrayView<FName>& NamedChunks) override;
+
+	virtual EChunkLocation::Type GetNamedChunkLocation(const FName NamedChunk) override;
+	virtual float GetNamedChunkProgress(const FName NamedChunk, EChunkProgressReportingType::Type ReportType) override;
+	virtual bool PrioritizeNamedChunk(const FName NamedChunk, EChunkPriority::Type Priority) override;
+
+	virtual ENamedChunkType GetNamedChunkType(const FName NamedChunk) const override;
+	virtual TArray<FName> GetNamedChunksByType(ENamedChunkType NamedChunkType) const override;
+
+protected:
+	virtual float GetCustomChunkProgress(const FCustomChunk& CustomChunk, EChunkProgressReportingType::Type ReportType) = 0; // platform specializations need to implement this as it's missing in the existing api
+	bool TryGetCustomChunkFromNamedChunk(const FName NamedChunk, FCustomChunk& OutCustomChunk) const;
+	TArray<FCustomChunk> GetCustomChunksFromNamedChunk(const FName NamedChunk) const;
+	TArray<FCustomChunk> GetCustomChunksFromNamedChunks(const TArrayView<FName>& NamedChunks) const;
+	TArray<FName> GetNamedChunksFromCustomChunks(const TArray<FCustomChunk>& CustomChunks) const;
+	FName GetNamedChunkByPakChunkIndex(int32 InPakchunkIndex) const;
+	FName GetCustomChunkName(const FCustomChunk& CustomChunk) const;
+};
+
 PRAGMA_ENABLE_DEPRECATION_WARNINGS
\ No newline at end of file
diff --git a/Engine/Source/Runtime/Core/Public/GenericPlatform/GenericPlatformMisc.h b/Engine/Source/Runtime/Core/Public/GenericPlatform/GenericPlatformMisc.h
--- a/Engine/Source/Runtime/Core/Public/GenericPlatform/GenericPlatformMisc.h
+++ b/Engine/Source/Runtime/Core/Public/GenericPlatform/GenericPlatformMisc.h
@@ -1657,9 +1657,16 @@ public:
 
 	static bool RequestDeviceCheckToken(TFunction<void(const TArray<uint8>&)> QuerySucceededFunc, TFunction<void(const FString&, const FString&)> QueryFailedFunc);
 
+	UE_DEPRECATED(5.2, "Use named chunks instead")
 	static TArray<FCustomChunk> GetOnDemandChunksForPakchunkIndices(const TArray<int32>& PakchunkIndices);
+
+	UE_DEPRECATED(5.2, "Use IPlatformChunkInstall::GetNamedChunksByType instead")
 	static TArray<FCustomChunk> GetAllOnDemandChunks();
+
+	UE_DEPRECATED(5.2, "Use IPlatformChunkInstall::GetNamedChunksByType instead")
 	static TArray<FCustomChunk> GetAllLanguageChunks();
+
+	UE_DEPRECATED(5.2, "Use IPlatformChunkInstall::GetNamedChunksByType instead")
 	static TArray<FCustomChunk> GetCustomChunksByType(ECustomChunkType DesiredChunkType);
 
 	/**
