Modular builds: Memory: Tweaked the system initialization to make it work with modular builds on some platforms.

On some consoles we load dynamic modules very early during the process initialization and rely on the allocators being ready. It happens even before we have the opportunity to call FPlatformMemory::Init() explicitly, so need to have it called on-demand from FMemory_GCreateMalloc_ThreadUnsafe.

Updated FXXXPlatformMemory::Init on all supported platforms to handle multiple Init() calls gracefully. We do this to free our common code from having to explicitly check if the system has been initialized.

#rnx

diff --git a/Engine/Platforms/Hololens/Source/Runtime/Core/Private/HoloLensPlatformMemory.cpp b/Engine/Platforms/Hololens/Source/Runtime/Core/Private/HoloLensPlatformMemory.cpp
--- a/Engine/Platforms/Hololens/Source/Runtime/Core/Private/HoloLensPlatformMemory.cpp
+++ b/Engine/Platforms/Hololens/Source/Runtime/Core/Private/HoloLensPlatformMemory.cpp
@@ -40,6 +40,14 @@ return true;
 
 void FHoloLensPlatformMemory::Init()
 {
+	// Only allow this method to be called once
+	{
+		static bool bInitDone = false;
+		if (bInitDone)
+			return;
+		bInitDone = true;
+	}
+
 	FGenericPlatformMemory::Init();
 
 #if PLATFORM_32BITS
diff --git a/Engine/Source/Runtime/Core/Private/Android/AndroidPlatformMemory.cpp b/Engine/Source/Runtime/Core/Private/Android/AndroidPlatformMemory.cpp
--- a/Engine/Source/Runtime/Core/Private/Android/AndroidPlatformMemory.cpp
+++ b/Engine/Source/Runtime/Core/Private/Android/AndroidPlatformMemory.cpp
@@ -54,6 +54,14 @@ static int64 GetNativeHeapAllocatedSize()
 
 void FAndroidPlatformMemory::Init()
 {
+	// Only allow this method to be called once
+	{
+		static bool bInitDone = false;
+		if (bInitDone)
+			return;
+		bInitDone = true;
+	}
+
 	FGenericPlatformMemory::Init();
 
 	const FPlatformMemoryConstants& MemoryConstants = FPlatformMemory::GetConstants();
diff --git a/Engine/Source/Runtime/Core/Private/Apple/ApplePlatformMemory.cpp b/Engine/Source/Runtime/Core/Private/Apple/ApplePlatformMemory.cpp
--- a/Engine/Source/Runtime/Core/Private/Apple/ApplePlatformMemory.cpp
+++ b/Engine/Source/Runtime/Core/Private/Apple/ApplePlatformMemory.cpp
@@ -265,6 +265,14 @@ void FApplePlatformMemory::NanoMallocInit()
 
 void FApplePlatformMemory::Init()
 {
+	// Only allow this method to be called once
+	{
+		static bool bInitDone = false;
+		if (bInitDone)
+			return;
+		bInitDone = true;
+	}
+
 	FGenericPlatformMemory::Init();
     
 	LLM(AppleLLM::Initialise());
diff --git a/Engine/Source/Runtime/Core/Private/HAL/UnrealMemory.cpp b/Engine/Source/Runtime/Core/Private/HAL/UnrealMemory.cpp
--- a/Engine/Source/Runtime/Core/Private/HAL/UnrealMemory.cpp
+++ b/Engine/Source/Runtime/Core/Private/HAL/UnrealMemory.cpp
@@ -332,12 +332,17 @@ FConsoleCommandDelegate::CreateStatic(&FMemory::EnablePoisonTests)
 /** Helper function called on first allocation to create and initialize GMalloc */
 static int FMemory_GCreateMalloc_ThreadUnsafe()
 {
+#if PLATFORM_CONSOLE_DYNAMIC_LINK
+	FPlatformMemory::Init();
+#endif
+
+	GMalloc = FPlatformMemory::BaseAllocator();
+
 #if !PLATFORM_MAC
 	FPlatformMemoryStats Stats = FPlatformMemory::GetStats();
 	uint64 ProgramSize = Stats.UsedPhysical;
 #endif
 
-	GMalloc = FPlatformMemory::BaseAllocator();
 	// Setup malloc crash as soon as possible.
 	FPlatformMallocCrash::Get( GMalloc );
 
diff --git a/Engine/Source/Runtime/Core/Private/Unix/UnixPlatformMemory.cpp b/Engine/Source/Runtime/Core/Private/Unix/UnixPlatformMemory.cpp
--- a/Engine/Source/Runtime/Core/Private/Unix/UnixPlatformMemory.cpp
+++ b/Engine/Source/Runtime/Core/Private/Unix/UnixPlatformMemory.cpp
@@ -102,6 +102,14 @@ bool GMemoryRangeDecommitIsNoOp = (UE_SERVER == 0);
 
 void FUnixPlatformMemory::Init()
 {
+	// Only allow this method to be called once
+	{
+		static bool bInitDone = false;
+		if (bInitDone)
+			return;
+		bInitDone = true;
+	}
+
 	FGenericPlatformMemory::Init();
 
 	const FPlatformMemoryConstants& MemoryConstants = FPlatformMemory::GetConstants();
diff --git a/Engine/Source/Runtime/Core/Private/Windows/WindowsPlatformMemory.cpp b/Engine/Source/Runtime/Core/Private/Windows/WindowsPlatformMemory.cpp
--- a/Engine/Source/Runtime/Core/Private/Windows/WindowsPlatformMemory.cpp
+++ b/Engine/Source/Runtime/Core/Private/Windows/WindowsPlatformMemory.cpp
@@ -62,6 +62,14 @@ int WindowsAllocHook(int nAllocType, void *pvData,
 
 void FWindowsPlatformMemory::Init()
 {
+	// Only allow this method to be called once
+	{
+		static bool bInitDone = false;
+		if (bInitDone)
+			return;
+		bInitDone = true;
+	}
+
 	FGenericPlatformMemory::Init();
 
 #if PLATFORM_32BITS
