Added IWYU pragmas.

diff --git a/Engine/Source/Runtime/Core/Public/CoreTypes.h b/Engine/Source/Runtime/Core/Public/CoreTypes.h
--- a/Engine/Source/Runtime/Core/Public/CoreTypes.h
+++ b/Engine/Source/Runtime/Core/Public/CoreTypes.h
@@ -6,8 +6,10 @@
 	Low level includes.
 ----------------------------------------------------------------------------*/
 
+// IWYU pragma: begin_exports
 #include "HAL/Platform.h"
 #include "ProfilingDebugging/UMemoryDefines.h"
 #include "Misc/CoreMiscDefines.h"
 #include "Misc/CoreDefines.h"
+// IWYU pragma: end_exports
 
diff --git a/Engine/Source/Runtime/Core/Public/HAL/PlatformCrt.h b/Engine/Source/Runtime/Core/Public/HAL/PlatformCrt.h
--- a/Engine/Source/Runtime/Core/Public/HAL/PlatformCrt.h
+++ b/Engine/Source/Runtime/Core/Public/HAL/PlatformCrt.h
@@ -2,6 +2,7 @@
 
 #pragma once
 
+// IWYU pragma: begin_exports
 #include <new>
 #include <wchar.h>
 #include <stddef.h>
@@ -11,3 +12,4 @@
 #include <math.h>
 #include <float.h>
 #include <string.h>
+// IWYU pragma: end_exports
