Added additional driver detection method using the Windows Setup API to resolve an issue with display drivers being reported as Unknown

diff --git a/Engine/Source/Runtime/Core/Private/Windows/WindowsPlatformMisc.cpp b/Engine/Source/Runtime/Core/Private/Windows/WindowsPlatformMisc.cpp
--- a/Engine/Source/Runtime/Core/Private/Windows/WindowsPlatformMisc.cpp
+++ b/Engine/Source/Runtime/Core/Private/Windows/WindowsPlatformMisc.cpp
@@ -82,13 +82,14 @@ THIRD_PARTY_INCLUDES_END
 // this cvar can be removed once we have a single method that works well
 static TAutoConsoleVariable<int32> CVarDriverDetectionMethod(
 	TEXT("r.DriverDetectionMethod"),
-	4,
+	5,
 	TEXT("Defines which implementation is used to detect the GPU driver (to check for old drivers, logs and statistics)\n"
 	     "  0: Iterate available drivers in registry and choose the one with the same name, if in question use next method (happens)\n"
 	     "  1: Get the driver of the primary adapter (might not be correct when dealing with multiple adapters)\n"
 	     "  2: Use DirectX LUID (would be the best, not yet implemented)\n"
 	     "  3: Use Windows functions, use the primary device (might be wrong when API is using another adapter)\n"
-	     "  4: Use Windows functions, use names such as DirectX Device (newest, most promising)"),
+	     "  4: Use Windows functions, use names such as DirectX Device (newest, most promising)\n"
+	     "  5: Use Windows SetupAPI functions"),
 	ECVF_RenderThreadSafe);
 
 int32 GetOSVersionsHelper( TCHAR* OutOSVersionLabel, int32 OSVersionLabelLength, TCHAR* OutOSSubVersionLabel, int32 OSSubVersionLabelLength )
@@ -2432,6 +2433,174 @@ FString FWindowsPlatformMisc::GetPrimaryGPUBrand()
 	return PrimaryGPUBrand;
 }
 
+#define USE_SP_ALTPLATFORM_INFO_V1 0
+#define USE_SP_ALTPLATFORM_INFO_V3 1
+#define USE_SP_DRVINFO_DATA_V1 0
+#define USE_SP_BACKUP_QUEUE_PARAMS_V1 0
+#define USE_SP_INF_SIGNER_INFO_V1 0
+#include <SetupAPI.h>
+#include <initguid.h>
+#include <devguid.h>
+#include <devpkey.h>
+#undef USE_SP_ALTPLATFORM_INFO_V1
+#undef USE_SP_ALTPLATFORM_INFO_V3
+#undef USE_SP_DRVINFO_DATA_V1
+#undef USE_SP_BACKUP_QUEUE_PARAMS_V1
+#undef USE_SP_INF_SIGNER_INFO_V1
+
+static void GetVideoDriverDetailsFromSetup(const FString& DeviceName, FGPUDriverInfo& Out)
+{
+	
+	HDEVINFO hDevInfo = SetupDiGetClassDevs(&GUID_DEVCLASS_DISPLAY, NULL, NULL, DIGCF_PRESENT);
+
+	FString RegistryKey = "";
+	
+	if (hDevInfo != INVALID_HANDLE_VALUE)
+	{
+		DWORD DataType = 0;
+		
+		const uint32 BufferSize = 512;
+		TCHAR Buffer[BufferSize + 1] = { 0 };
+		
+		SP_DEVINFO_DATA DeviceInfoData;
+		DeviceInfoData.cbSize = sizeof(SP_DEVINFO_DATA);
+
+		bool bFound = false;
+		for (int32 Idx = 0; SetupDiEnumDeviceInfo(hDevInfo, Idx, &DeviceInfoData); Idx++)
+		{
+			// Get the device description, check if it matches the queried device name
+			if (SetupDiGetDeviceProperty(hDevInfo, &DeviceInfoData, &DEVPKEY_Device_DriverDesc, &DataType,
+				(PBYTE)Buffer, sizeof(Buffer), nullptr, 0))
+			{
+				if (DeviceName.Compare(Buffer) == 0)
+				{
+					Out.DeviceDescription = Buffer;
+					bFound = true;
+					ZeroMemory(Buffer, sizeof(Buffer));
+
+					// Retrieve the registry key for this device for 3rd party data
+					if (SetupDiGetDeviceProperty(hDevInfo, &DeviceInfoData, &DEVPKEY_Device_Driver, &DataType,
+						(PBYTE)Buffer, sizeof(Buffer), nullptr, 0))
+					{
+						RegistryKey = Buffer;
+						ZeroMemory(Buffer, sizeof(Buffer));
+					}
+					else
+					{
+						UE_LOG(LogWindows, Warning, TEXT("Failed to retrieve driver registry key for device %d"), Idx);
+					}
+					
+					break;
+				}
+				ZeroMemory(Buffer, sizeof(Buffer));
+			}
+			else
+			{
+				UE_LOG(LogWindows, Warning, TEXT("Failed to retrieve driver description for device %d"), Idx);
+			}
+		}
+
+		if (bFound)
+		{
+			// Get the provider name
+			if (SetupDiGetDeviceProperty(hDevInfo, &DeviceInfoData, &DEVPKEY_Device_DriverProvider, &DataType,
+				(PBYTE)Buffer, sizeof(Buffer), nullptr, 0))
+			{
+				Out.ProviderName = Buffer;
+				ZeroMemory(Buffer, sizeof(Buffer));
+			}
+			else
+			{
+				UE_LOG(LogWindows, Warning, TEXT("Failed to find provider name"));
+			}
+			// Get the internal driver version
+			if (SetupDiGetDeviceProperty(hDevInfo, &DeviceInfoData, &DEVPKEY_Device_DriverVersion, &DataType,
+				(PBYTE)Buffer, sizeof(Buffer), nullptr, 0))
+			{
+				Out.InternalDriverVersion = Buffer;
+				ZeroMemory(Buffer, sizeof(Buffer));
+			}
+			else
+			{
+				UE_LOG(LogWindows, Warning, TEXT("Failed to find internal driver version"));
+			}
+			// Get the driver date
+			FILETIME FileTime;
+			if (SetupDiGetDeviceProperty(hDevInfo, &DeviceInfoData, &DEVPKEY_Device_DriverDate, &DataType,
+				(PBYTE)&FileTime, sizeof(FILETIME), nullptr, 0))
+			{
+				SYSTEMTIME SystemTime;
+				FileTimeToSystemTime(&FileTime, &SystemTime);
+				GetDateFormat(LOCALE_USER_DEFAULT, DATE_SHORTDATE, &SystemTime, NULL, Buffer, UE_ARRAY_COUNT(Buffer));
+				Out.DriverDate = Buffer;
+				ZeroMemory(Buffer, sizeof(Buffer));
+			}
+			else
+			{
+				UE_LOG(LogWindows, Warning, TEXT("Failed to find driver date"));
+			}
+		}
+		else
+		{
+			UE_LOG(LogWindows, Warning, TEXT("Unable to find requested device '%s' using Setup API"), *DeviceName);
+		}
+		
+		SetupDiDestroyDeviceInfoList(hDevInfo);
+	}
+	else
+	{
+		UE_LOG(LogWindows, Warning, TEXT("Failed to initialize Setup API"));
+	}
+
+	if (!Out.ProviderName.IsEmpty())
+	{
+		if (Out.ProviderName.Find(TEXT("NVIDIA")) != INDEX_NONE)
+		{
+			Out.SetNVIDIA();
+		}
+		else if (Out.ProviderName.Find(TEXT("Advanced Micro Devices")) != INDEX_NONE)
+		{
+			Out.SetAMD();
+		}
+		else if (Out.ProviderName.Find(TEXT("Intel")) != INDEX_NONE)	// usually TEXT("Intel Corporation")
+		{
+			Out.SetIntel();
+		}
+	}
+
+	Out.UserDriverVersion = Out.InternalDriverVersion;
+
+	if(Out.IsNVIDIA())
+	{
+		Out.UserDriverVersion = Out.GetUnifiedDriverVersion();
+	}
+	else if(Out.IsAMD() && !RegistryKey.IsEmpty())
+	{
+		// Get the AMD specific information directly from the registry
+		// AMD AGS could be used instead, but retrieving the radeon software version cannot occur after a D3D Device
+		// has been created, and this function could be called at any time
+		
+		const FString Key = FString::Printf(TEXT("SYSTEM\\CurrentControlSet\\Control\\Class\\%s"), *RegistryKey);
+		
+		if(FWindowsPlatformMisc::QueryRegKey(HKEY_LOCAL_MACHINE, *Key, TEXT("Catalyst_Version"), Out.UserDriverVersion))
+		{
+			Out.UserDriverVersion = FString(TEXT("Catalyst ")) + Out.UserDriverVersion;
+		}
+
+		FString Edition;
+		if(FWindowsPlatformMisc::QueryRegKey(HKEY_LOCAL_MACHINE, *Key, TEXT("RadeonSoftwareEdition"), Edition))
+		{
+			FString Version;
+			if(FWindowsPlatformMisc::QueryRegKey(HKEY_LOCAL_MACHINE, *Key, TEXT("RadeonSoftwareVersion"), Version))
+			{
+				// e.g. TEXT("Crimson 15.12") or TEXT("Catalyst 14.1")
+				Out.UserDriverVersion = Edition + TEXT(" ") + Version;
+			}
+		}
+	}
+
+}
+
 static void GetVideoDriverDetails(const FString& Key, FGPUDriverInfo& Out)
 {
 	// https://msdn.microsoft.com/en-us/library/windows/hardware/ff569240(v=vs.85).aspx
@@ -2557,6 +2726,20 @@ FGPUDriverInfo FWindowsPlatformMisc::GetGPUDriverInfo(const FString& DeviceDescr
 
 	int32 Method = CVarDriverDetectionMethod.GetValueOnGameThread();
 
+	if (Method == 5)
+	{
+		UE_LOG(LogWindows, Log, TEXT("Gathering driver information using Windows Setup API"));
+		FGPUDriverInfo Local;
+		GetVideoDriverDetailsFromSetup(DeviceDescription, Local);
+
+		if(Local.IsValid() && Local.DeviceDescription == DeviceDescription)
+		{
+			return Local;
+		}
+
+		return Ret;
+	}
+	
 	if(Method == 3 || Method == 4)
 	{
 		UE_LOG(LogWindows, Log, TEXT("EnumDisplayDevices:"));
