Pass 2 on running IWYU on Core.

Things to note regarding this pass:
* No includes are currently being removed from public headers.
* Any private file that has a #if is currently not being updated.

diff --git a/Engine/Source/Runtime/Core/Private/IO/IoStatus.cpp b/Engine/Source/Runtime/Core/Private/IO/IoStatus.cpp
--- a/Engine/Source/Runtime/Core/Private/IO/IoStatus.cpp
+++ b/Engine/Source/Runtime/Core/Private/IO/IoStatus.cpp
@@ -4,6 +4,7 @@
 #include "Containers/StringView.h"
 #include "CoreMinimal.h"
 #include "IO/IoDispatcher.h"
+#include "Trace/Detail/Channel.h"
 
 //////////////////////////////////////////////////////////////////////////
 
diff --git a/Engine/Source/Runtime/Core/Private/Internationalization/InternationalizationMetadata.cpp b/Engine/Source/Runtime/Core/Private/Internationalization/InternationalizationMetadata.cpp
--- a/Engine/Source/Runtime/Core/Private/Internationalization/InternationalizationMetadata.cpp
+++ b/Engine/Source/Runtime/Core/Private/Internationalization/InternationalizationMetadata.cpp
@@ -13,6 +13,7 @@
 #include "Templates/ChooseClass.h"
 #include "Templates/Less.h"
 #include "Templates/Tuple.h"
+#include "Trace/Detail/Channel.h"
 
 DEFINE_LOG_CATEGORY_STATIC(LogInternationalizationMetadata, Log, All);
 
diff --git a/Engine/Source/Runtime/Core/Private/Internationalization/LocalizationResourceTextSource.cpp b/Engine/Source/Runtime/Core/Private/Internationalization/LocalizationResourceTextSource.cpp
--- a/Engine/Source/Runtime/Core/Private/Internationalization/LocalizationResourceTextSource.cpp
+++ b/Engine/Source/Runtime/Core/Private/Internationalization/LocalizationResourceTextSource.cpp
@@ -2,8 +2,8 @@
 
 #include "Internationalization/LocalizationResourceTextSource.h"
 
+#include "CoreTypes.h"
 #include "HAL/FileManager.h"
-#include "HAL/PlatformProperties.h"
 #include "Internationalization/TextLocalizationResource.h"
 #include "Misc/ConfigCacheIni.h"
 #include "Misc/CoreDelegates.h"
diff --git a/Engine/Source/Runtime/Core/Private/Internationalization/LocalizationResourceTextSource.h b/Engine/Source/Runtime/Core/Private/Internationalization/LocalizationResourceTextSource.h
--- a/Engine/Source/Runtime/Core/Private/Internationalization/LocalizationResourceTextSource.h
+++ b/Engine/Source/Runtime/Core/Private/Internationalization/LocalizationResourceTextSource.h
@@ -4,7 +4,6 @@
 
 #include "Containers/ArrayView.h"
 #include "CoreMinimal.h"
-#include "CoreTypes.h"
 #include "Internationalization/ILocalizedTextSource.h"
 #include "Internationalization/LocalizedTextSourceTypes.h"
 
diff --git a/Engine/Source/Runtime/Core/Private/Internationalization/PolyglotTextSource.h b/Engine/Source/Runtime/Core/Private/Internationalization/PolyglotTextSource.h
--- a/Engine/Source/Runtime/Core/Private/Internationalization/PolyglotTextSource.h
+++ b/Engine/Source/Runtime/Core/Private/Internationalization/PolyglotTextSource.h
@@ -4,8 +4,8 @@
 
 #include "CoreTypes.h"
 #include "Internationalization/ILocalizedTextSource.h"
-#include "Internationalization/PolyglotTextData.h"
 #include "Internationalization/LocKeyFuncs.h"
+#include "Internationalization/PolyglotTextData.h" // IWYU pragma: keep
 
 /**
  * Implementation of a localized text source that stores literal polyglot data registered at runtime.
diff --git a/Engine/Source/Runtime/Core/Private/Internationalization/StringTableCore.cpp b/Engine/Source/Runtime/Core/Private/Internationalization/StringTableCore.cpp
--- a/Engine/Source/Runtime/Core/Private/Internationalization/StringTableCore.cpp
+++ b/Engine/Source/Runtime/Core/Private/Internationalization/StringTableCore.cpp
@@ -15,9 +15,11 @@
 #include "Misc/FileHelper.h"
 #include "Misc/Parse.h"
 #include "Misc/ScopeLock.h"
+#include "Serialization/Archive.h"
 #include "Serialization/Csv/CsvParser.h"
 #include "Templates/ChooseClass.h"
 #include "Templates/Tuple.h"
+#include "Trace/Detail/Channel.h"
 
 class UStringTable;
 
diff --git a/Engine/Source/Runtime/Core/Private/Internationalization/Text.cpp b/Engine/Source/Runtime/Core/Private/Internationalization/Text.cpp
--- a/Engine/Source/Runtime/Core/Private/Internationalization/Text.cpp
+++ b/Engine/Source/Runtime/Core/Private/Internationalization/Text.cpp
@@ -7,6 +7,7 @@
 #include "UObject/DebugSerializationFlags.h"
 #include "Internationalization/Culture.h"
 #include "Internationalization/Internationalization.h"
+#include "Internationalization/ITextGenerator.h"
 #include "Internationalization/StringTableRegistry.h"
 
 #include "Internationalization/TextHistory.h"
diff --git a/Engine/Source/Runtime/Core/Private/Internationalization/TextCache.h b/Engine/Source/Runtime/Core/Private/Internationalization/TextCache.h
--- a/Engine/Source/Runtime/Core/Private/Internationalization/TextCache.h
+++ b/Engine/Source/Runtime/Core/Private/Internationalization/TextCache.h
@@ -4,6 +4,7 @@
 
 #include "Containers/Map.h"
 #include "HAL/CriticalSection.h"
+#include "HAL/Platform.h"
 #include "Internationalization/Text.h"
 
 class FTextId;
diff --git a/Engine/Source/Runtime/Core/Private/Internationalization/TextFormatter.cpp b/Engine/Source/Runtime/Core/Private/Internationalization/TextFormatter.cpp
--- a/Engine/Source/Runtime/Core/Private/Internationalization/TextFormatter.cpp
+++ b/Engine/Source/Runtime/Core/Private/Internationalization/TextFormatter.cpp
@@ -27,6 +27,7 @@
 #include "Templates/Tuple.h"
 #include "Templates/UnrealTemplate.h"
 #include "Templates/ValueOrError.h"
+#include "Trace/Detail/Channel.h"
 
 #define LOCTEXT_NAMESPACE "TextFormatter"
 
diff --git a/Engine/Source/Runtime/Core/Private/Internationalization/TextHistory.cpp b/Engine/Source/Runtime/Core/Private/Internationalization/TextHistory.cpp
--- a/Engine/Source/Runtime/Core/Private/Internationalization/TextHistory.cpp
+++ b/Engine/Source/Runtime/Core/Private/Internationalization/TextHistory.cpp
@@ -12,6 +12,7 @@
 #include "UObject/PropertyPortFlags.h"
 
 #include "Internationalization/FastDecimalFormat.h"
+#include "Internationalization/ITextGenerator.h"
 #include "Internationalization/TextFormatter.h"
 #include "Internationalization/TextChronoFormatter.h"
 #include "Internationalization/TextTransformer.h"
diff --git a/Engine/Source/Runtime/Core/Private/Internationalization/TextHistory.h b/Engine/Source/Runtime/Core/Private/Internationalization/TextHistory.h
--- a/Engine/Source/Runtime/Core/Private/Internationalization/TextHistory.h
+++ b/Engine/Source/Runtime/Core/Private/Internationalization/TextHistory.h
@@ -1,17 +1,22 @@
 // Copyright Epic Games, Inc. All Rights Reserved.
 #pragma once
 
-#include "CoreTypes.h"
 #include "Containers/Array.h"
 #include "Containers/UnrealString.h"
-#include "Templates/SharedPointer.h"
-#include "Misc/Optional.h"
-#include "Internationalization/Text.h"
+#include "CoreMinimal.h"
+#include "CoreTypes.h"
 #include "Internationalization/ITextData.h"
-#include "Internationalization/ITextGenerator.h"
+#include "Internationalization/LocalizedTextSourceTypes.h"
 #include "Internationalization/StringTableCoreFwd.h"
+#include "Internationalization/Text.h"
+#include "Internationalization/TextKey.h"
 #include "Misc/DateTime.h"
+#include "Misc/Optional.h"
+#include "Serialization/StructuredArchive.h"
+#include "Templates/SharedPointer.h"
+#include "Templates/SharedPointerInternals.h"
 
+class ITextGenerator;
 struct FDecimalNumberFormattingRules;
 
 enum class ETextHistoryType : int8
diff --git a/Engine/Source/Runtime/Core/Private/Logging/LogCategory.cpp b/Engine/Source/Runtime/Core/Private/Logging/LogCategory.cpp
--- a/Engine/Source/Runtime/Core/Private/Logging/LogCategory.cpp
+++ b/Engine/Source/Runtime/Core/Private/Logging/LogCategory.cpp
@@ -2,7 +2,6 @@
 
 #include "Logging/LogCategory.h"
 
-#include "CoreGlobals.h"
 #include "CoreMinimal.h"
 #include "Logging/LogSuppressionInterface.h"
 #include "Logging/LogTrace.h"
diff --git a/Engine/Source/Runtime/Core/Private/Math/ColorList.cpp b/Engine/Source/Runtime/Core/Private/Math/ColorList.cpp
--- a/Engine/Source/Runtime/Core/Private/Math/ColorList.cpp
+++ b/Engine/Source/Runtime/Core/Private/Math/ColorList.cpp
@@ -9,6 +9,7 @@
 #include "Containers/UnrealString.h"
 #include "Logging/LogCategory.h"
 #include "Logging/LogMacros.h"
+#include "Trace/Detail/Channel.h"
 
 DEFINE_LOG_CATEGORY_STATIC(LogColorList, Log, All);
 
diff --git a/Engine/Source/Runtime/Core/Private/Math/UnitConversion.cpp b/Engine/Source/Runtime/Core/Private/Math/UnitConversion.cpp
--- a/Engine/Source/Runtime/Core/Private/Math/UnitConversion.cpp
+++ b/Engine/Source/Runtime/Core/Private/Math/UnitConversion.cpp
@@ -1,10 +1,21 @@
 // Copyright Epic Games, Inc. All Rights Reserved.
 
 #include "Math/UnitConversion.h"
+
+#include "HAL/PlatformCrt.h"
 #include "Internationalization/Internationalization.h"
-#include "Misc/ExpressionParserTypes.h"
-#include "Misc/ExpressionParser.h"
+#include "Internationalization/Text.h"
 #include "Math/BasicMathExpressionEvaluator.h"
+#include "Math/UnitConversion.inl"
+#include "Math/UnrealMathUtility.h"
+#include "Misc/AssertionMacros.h"
+#include "Misc/CString.h"
+#include "Misc/ExpressionParser.h"
+#include "Misc/ExpressionParserTypes.h"
+#include "Misc/ExpressionParserTypes.inl"
+#include "Misc/Guid.h"
+#include "Templates/UniquePtr.h"
+#include "Templates/UnrealTemplate.h"
 
 PRAGMA_DISABLE_UNSAFE_TYPECAST_WARNINGS
 
diff --git a/Engine/Source/Runtime/Core/Private/Misc/MessageDialog.cpp b/Engine/Source/Runtime/Core/Private/Misc/MessageDialog.cpp
--- a/Engine/Source/Runtime/Core/Private/Misc/MessageDialog.cpp
+++ b/Engine/Source/Runtime/Core/Private/Misc/MessageDialog.cpp
@@ -13,6 +13,7 @@
 #include "Misc/CoreDelegates.h"
 #include "Misc/FeedbackContext.h"
 #include "Misc/OutputDeviceRedirector.h"
+#include "Trace/Detail/Channel.h"
 
 namespace
 {
diff --git a/Engine/Source/Runtime/Core/Private/Misc/NetworkVersion.cpp b/Engine/Source/Runtime/Core/Private/Misc/NetworkVersion.cpp
--- a/Engine/Source/Runtime/Core/Private/Misc/NetworkVersion.cpp
+++ b/Engine/Source/Runtime/Core/Private/Misc/NetworkVersion.cpp
@@ -1,12 +1,15 @@
 // Copyright Epic Games, Inc. All Rights Reserved.
 
 #include "Misc/NetworkVersion.h"
+
+#include "BuildSettings.h"
+#include "CoreMinimal.h"
+#include "HAL/IConsoleManager.h"
 #include "Misc/App.h"
+#include "Misc/CommandLine.h"
 #include "Misc/EngineVersion.h"
-#include "Runtime/Launch/Resources/Version.h"
-#include "Misc/NetworkGuid.h"
-#include "HAL/IConsoleManager.h"
-#include "BuildSettings.h"
+#include "Misc/NetworkGuid.h" // IWYU pragma: keep
+#include "Trace/Detail/Channel.h"
 
 DEFINE_LOG_CATEGORY( LogNetVersion );
 
diff --git a/Engine/Source/Runtime/Core/Private/Misc/PathPermissionList.cpp b/Engine/Source/Runtime/Core/Private/Misc/PathPermissionList.cpp
--- a/Engine/Source/Runtime/Core/Private/Misc/PathPermissionList.cpp
+++ b/Engine/Source/Runtime/Core/Private/Misc/PathPermissionList.cpp
@@ -10,6 +10,7 @@
 #include "Misc/AssertionMacros.h"
 #include "Misc/NamePermissionList.h"
 #include "Misc/PathViews.h"
+#include "Misc/StringBuilder.h"
 #include "Templates/ChooseClass.h"
 #include "Templates/Tuple.h"
 #include "Templates/UnrealTemplate.h"
diff --git a/Engine/Source/Runtime/Core/Private/Misc/PathViews.cpp b/Engine/Source/Runtime/Core/Private/Misc/PathViews.cpp
--- a/Engine/Source/Runtime/Core/Private/Misc/PathViews.cpp
+++ b/Engine/Source/Runtime/Core/Private/Misc/PathViews.cpp
@@ -18,7 +18,6 @@
 #include "String/Find.h"
 #include "String/ParseTokens.h"
 #include "Templates/Function.h"
-#include "Templates/UnrealTemplate.h"
 
 namespace UE4PathViews_Private
 {
diff --git a/Engine/Source/Runtime/Core/Private/ProfilingDebugging/Histogram.cpp b/Engine/Source/Runtime/Core/Private/ProfilingDebugging/Histogram.cpp
--- a/Engine/Source/Runtime/Core/Private/ProfilingDebugging/Histogram.cpp
+++ b/Engine/Source/Runtime/Core/Private/ProfilingDebugging/Histogram.cpp
@@ -4,6 +4,7 @@
 
 #include "Logging/LogCategory.h"
 #include "Templates/Function.h"
+#include "Trace/Detail/Channel.h"
 
 DEFINE_LOG_CATEGORY(LogHistograms);
 
diff --git a/Engine/Source/Runtime/Core/Private/Serialization/ArchiveSaveCompressedProxy.cpp b/Engine/Source/Runtime/Core/Private/Serialization/ArchiveSaveCompressedProxy.cpp
--- a/Engine/Source/Runtime/Core/Private/Serialization/ArchiveSaveCompressedProxy.cpp
+++ b/Engine/Source/Runtime/Core/Private/Serialization/ArchiveSaveCompressedProxy.cpp
@@ -8,6 +8,7 @@
 #include "Logging/LogMacros.h"
 #include "Math/UnrealMathUtility.h"
 #include "Misc/Compression.h"
+#include "Trace/Detail/Channel.h"
 
 PRAGMA_DISABLE_UNSAFE_TYPECAST_WARNINGS
 
diff --git a/Engine/Source/Runtime/Core/Private/Serialization/BitReader.cpp b/Engine/Source/Runtime/Core/Private/Serialization/BitReader.cpp
--- a/Engine/Source/Runtime/Core/Private/Serialization/BitReader.cpp
+++ b/Engine/Source/Runtime/Core/Private/Serialization/BitReader.cpp
@@ -9,6 +9,7 @@
 #include "Math/UnrealMathUtility.h"
 #include "Serialization/Archive.h"
 #include "Templates/UnrealTemplate.h"
+#include "Trace/Detail/Channel.h"
 
 PRAGMA_DISABLE_UNSAFE_TYPECAST_WARNINGS
 
diff --git a/Engine/Source/Runtime/Core/Private/Serialization/BitWriter.cpp b/Engine/Source/Runtime/Core/Private/Serialization/BitWriter.cpp
--- a/Engine/Source/Runtime/Core/Private/Serialization/BitWriter.cpp
+++ b/Engine/Source/Runtime/Core/Private/Serialization/BitWriter.cpp
@@ -8,6 +8,7 @@
 #include "Logging/LogCategory.h"
 #include "Logging/LogMacros.h"
 #include "Serialization/Archive.h"
+#include "Trace/Detail/Channel.h"
 
 PRAGMA_DISABLE_UNSAFE_TYPECAST_WARNINGS
 
diff --git a/Engine/Source/Runtime/Core/Private/Serialization/CompactBinarySerialization.cpp b/Engine/Source/Runtime/Core/Private/Serialization/CompactBinarySerialization.cpp
--- a/Engine/Source/Runtime/Core/Private/Serialization/CompactBinarySerialization.cpp
+++ b/Engine/Source/Runtime/Core/Private/Serialization/CompactBinarySerialization.cpp
@@ -20,6 +20,7 @@
 #include "Serialization/CompactBinaryValidation.h"
 #include "Serialization/CompactBinaryValue.h"
 #include "Serialization/VarInt.h"
+#include "Templates/Function.h"
 #include "Templates/IdentityFunctor.h"
 #include "Templates/Invoke.h"
 #include "Templates/RemoveReference.h"
diff --git a/Engine/Source/Runtime/Core/Private/Serialization/Csv/CsvParserTests.cpp b/Engine/Source/Runtime/Core/Private/Serialization/Csv/CsvParserTests.cpp
--- a/Engine/Source/Runtime/Core/Private/Serialization/Csv/CsvParserTests.cpp
+++ b/Engine/Source/Runtime/Core/Private/Serialization/Csv/CsvParserTests.cpp
@@ -1,6 +1,5 @@
 // Copyright Epic Games, Inc. All Rights Reserved.
 
-#include "Containers/UnrealString.h"
 #include "CoreMinimal.h"
 #include "CoreTypes.h"
 #include "Misc/AutomationTest.h"
diff --git a/Engine/Source/Runtime/Core/Private/Serialization/LargeMemoryReader.cpp b/Engine/Source/Runtime/Core/Private/Serialization/LargeMemoryReader.cpp
--- a/Engine/Source/Runtime/Core/Private/Serialization/LargeMemoryReader.cpp
+++ b/Engine/Source/Runtime/Core/Private/Serialization/LargeMemoryReader.cpp
@@ -6,6 +6,7 @@
 #include "HAL/UnrealMemory.h"
 #include "Logging/LogCategory.h"
 #include "Logging/LogMacros.h"
+#include "Trace/Detail/Channel.h"
 
 /*----------------------------------------------------------------------------
 	FLargeMemoryReader
diff --git a/Engine/Source/Runtime/Core/Private/Serialization/LargeMemoryWriter.cpp b/Engine/Source/Runtime/Core/Private/Serialization/LargeMemoryWriter.cpp
--- a/Engine/Source/Runtime/Core/Private/Serialization/LargeMemoryWriter.cpp
+++ b/Engine/Source/Runtime/Core/Private/Serialization/LargeMemoryWriter.cpp
@@ -5,6 +5,7 @@
 #include "CoreGlobals.h"
 #include "Logging/LogCategory.h"
 #include "Logging/LogMacros.h"
+#include "Trace/Detail/Channel.h"
 
 /*----------------------------------------------------------------------------
 	FLargeMemoryWriter
diff --git a/Engine/Source/Runtime/Core/Private/UObject/ObjectVersion.cpp b/Engine/Source/Runtime/Core/Private/UObject/ObjectVersion.cpp
--- a/Engine/Source/Runtime/Core/Private/UObject/ObjectVersion.cpp
+++ b/Engine/Source/Runtime/Core/Private/UObject/ObjectVersion.cpp
@@ -6,6 +6,7 @@
 
 #include "UObject/ObjectVersion.h"
 
+#include "Misc/AssertionMacros.h"
 #include "Serialization/Archive.h"
 #include "Serialization/CompactBinary.h"
 #include "Serialization/CompactBinaryWriter.h"
diff --git a/Engine/Source/Runtime/Core/Private/Unix/UnixPlatformOutputDevices.cpp b/Engine/Source/Runtime/Core/Private/Unix/UnixPlatformOutputDevices.cpp
--- a/Engine/Source/Runtime/Core/Private/Unix/UnixPlatformOutputDevices.cpp
+++ b/Engine/Source/Runtime/Core/Private/Unix/UnixPlatformOutputDevices.cpp
@@ -2,13 +2,12 @@
 
 #include "Unix/UnixPlatformOutputDevices.h"
 
-#include "CoreGlobals.h"
 #include "CoreMinimal.h"
 #include "CoreTypes.h"
+#include "Misc/Build.h"
 #include "Misc/CommandLine.h"
 #include "Misc/OutputDeviceConsole.h"
 #include "Misc/OutputDeviceRedirector.h"
-#include "Misc/Parse.h"
 #include "Misc/Paths.h"
 #include "Unix/UnixErrorOutputDevice.h"
 
diff --git a/Engine/Source/Runtime/Core/Public/CoreGlobals.h b/Engine/Source/Runtime/Core/Public/CoreGlobals.h
--- a/Engine/Source/Runtime/Core/Public/CoreGlobals.h
+++ b/Engine/Source/Runtime/Core/Public/CoreGlobals.h
@@ -5,6 +5,7 @@
 #include "CoreTypes.h"
 #include "HAL/PlatformTLS.h"
 #include "Logging/LogMacros.h"
+#include "Misc/Build.h"
 #include "Misc/EnumClassFlags.h"
 #include "Misc/OutputDevice.h"
 #include "ProfilingDebugging/CpuProfilerTrace.h"
diff --git a/Engine/Source/Runtime/Core/Public/IO/IoDispatcher.h b/Engine/Source/Runtime/Core/Public/IO/IoDispatcher.h
--- a/Engine/Source/Runtime/Core/Public/IO/IoDispatcher.h
+++ b/Engine/Source/Runtime/Core/Public/IO/IoDispatcher.h
@@ -19,6 +19,7 @@
 #include "Memory/MemoryFwd.h"
 #include "Memory/MemoryView.h"
 #include "Misc/AES.h"
+#include "Misc/Build.h"
 #include "Misc/EnumClassFlags.h"
 #include "Misc/Guid.h"
 #include "Misc/IEngineCrypto.h"
diff --git a/Engine/Source/Runtime/Core/Public/Internationalization/StringTableCore.h b/Engine/Source/Runtime/Core/Public/Internationalization/StringTableCore.h
--- a/Engine/Source/Runtime/Core/Public/Internationalization/StringTableCore.h
+++ b/Engine/Source/Runtime/Core/Public/Internationalization/StringTableCore.h
@@ -4,6 +4,7 @@
 #include "Containers/Map.h"
 #include "Containers/SparseArray.h"
 #include "Containers/UnrealString.h"
+#include "CoreMinimal.h"
 #include "HAL/CriticalSection.h"
 #include "Internationalization/Internationalization.h"
 #include "Internationalization/LocKeyFuncs.h"
diff --git a/Engine/Source/Runtime/Core/Public/Linux/LinuxPlatformProcess.h b/Engine/Source/Runtime/Core/Public/Linux/LinuxPlatformProcess.h
--- a/Engine/Source/Runtime/Core/Public/Linux/LinuxPlatformProcess.h
+++ b/Engine/Source/Runtime/Core/Public/Linux/LinuxPlatformProcess.h
@@ -6,6 +6,8 @@
 
 #pragma once
 
+#include "HAL/Platform.h"
+#include "HAL/PlatformProcess.h"
 #include "Unix/UnixPlatformProcess.h"
 
 /**
diff --git a/Engine/Source/Runtime/Core/Public/Logging/LogScopedVerbosityOverride.h b/Engine/Source/Runtime/Core/Public/Logging/LogScopedVerbosityOverride.h
--- a/Engine/Source/Runtime/Core/Public/Logging/LogScopedVerbosityOverride.h
+++ b/Engine/Source/Runtime/Core/Public/Logging/LogScopedVerbosityOverride.h
@@ -4,6 +4,7 @@
 
 #include "CoreTypes.h"
 #include "Logging/LogVerbosity.h"
+#include "Misc/Build.h"
 
 struct FLogCategoryBase;
 
diff --git a/Engine/Source/Runtime/Core/Public/Logging/LogTrace.h b/Engine/Source/Runtime/Core/Public/Logging/LogTrace.h
--- a/Engine/Source/Runtime/Core/Public/Logging/LogTrace.h
+++ b/Engine/Source/Runtime/Core/Public/Logging/LogTrace.h
@@ -3,7 +3,9 @@
 #pragma once
 
 #include "CoreTypes.h"
+#include "HAL/PreprocessorHelpers.h"
 #include "Logging/LogVerbosity.h"
+#include "Misc/Build.h"
 #include "Templates/IsArrayOrRefOfTypeByPredicate.h"
 #include "Trace/Config.h"
 #include "Trace/Trace.h"
diff --git a/Engine/Source/Runtime/Core/Public/Math/Box2D.h b/Engine/Source/Runtime/Core/Public/Math/Box2D.h
--- a/Engine/Source/Runtime/Core/Public/Math/Box2D.h
+++ b/Engine/Source/Runtime/Core/Public/Math/Box2D.h
@@ -8,6 +8,7 @@
 #include "Math/UnrealMathUtility.h"
 #include "Math/Vector2D.h"
 #include "Misc/AssertionMacros.h"
+#include "Misc/LargeWorldCoordinates.h"
 #include "Misc/LargeWorldCoordinatesSerializer.h"
 #include "Serialization/Archive.h"
 #include "Templates/IsUECoreType.h"
diff --git a/Engine/Source/Runtime/Core/Public/Math/Color.h b/Engine/Source/Runtime/Core/Public/Math/Color.h
--- a/Engine/Source/Runtime/Core/Public/Math/Color.h
+++ b/Engine/Source/Runtime/Core/Public/Math/Color.h
@@ -5,6 +5,7 @@
 #include "Containers/Array.h"
 #include "Containers/UnrealString.h"
 #include "CoreTypes.h"
+#include "HAL/PreprocessorHelpers.h"
 #include "Math/MathFwd.h"
 #include "Math/UnrealMathUtility.h"
 #include "Misc/AssertionMacros.h"
diff --git a/Engine/Source/Runtime/Core/Public/Math/GenericOctree.h b/Engine/Source/Runtime/Core/Public/Math/GenericOctree.h
--- a/Engine/Source/Runtime/Core/Public/Math/GenericOctree.h
+++ b/Engine/Source/Runtime/Core/Public/Math/GenericOctree.h
@@ -12,6 +12,7 @@
 #include "Math/MathFwd.h"
 #include "Math/UnrealMathSSE.h"
 #include "Templates/Models.h"
+#include "Trace/Detail/Channel.h"
 
 /** A concise iteration over the children of an octree node. */
 #define FOREACH_OCTREE_CHILD_NODE(ChildRef) \
diff --git a/Engine/Source/Runtime/Core/Public/Math/Sphere.h b/Engine/Source/Runtime/Core/Public/Math/Sphere.h
--- a/Engine/Source/Runtime/Core/Public/Math/Sphere.h
+++ b/Engine/Source/Runtime/Core/Public/Math/Sphere.h
@@ -10,6 +10,7 @@
 #include "Math/Vector.h"
 #include "Math/Vector4.h"
 #include "Misc/AssertionMacros.h"
+#include "Misc/LargeWorldCoordinates.h"
 #include "Serialization/Archive.h"
 #include "Serialization/StructuredArchiveAdapters.h"
 #include "Templates/IsUECoreType.h"
diff --git a/Engine/Source/Runtime/Core/Public/Math/Transform.h b/Engine/Source/Runtime/Core/Public/Math/Transform.h
--- a/Engine/Source/Runtime/Core/Public/Math/Transform.h
+++ b/Engine/Source/Runtime/Core/Public/Math/Transform.h
@@ -7,10 +7,12 @@
  */
 
 #include "CoreTypes.h"
+#include "HAL/PreprocessorHelpers.h"
 #include "Math/Quat.h"
 #include "Math/ScalarRegister.h"
 #include "Math/TransformNonVectorized.h"
 #include "Math/TransformVectorized.h"
+#include "Misc/LargeWorldCoordinates.h"
 #include "Misc/LargeWorldCoordinatesSerializer.h"
 #include "Serialization/MemoryLayout.h"
 #include "Serialization/StructuredArchiveAdapters.h"
diff --git a/Engine/Source/Runtime/Core/Public/Math/TransformVectorized.h b/Engine/Source/Runtime/Core/Public/Math/TransformVectorized.h
--- a/Engine/Source/Runtime/Core/Public/Math/TransformVectorized.h
+++ b/Engine/Source/Runtime/Core/Public/Math/TransformVectorized.h
@@ -19,6 +19,7 @@
 #include "Math/Vector4.h"
 #include "Math/VectorRegister.h"
 #include "Misc/AssertionMacros.h"
+#include "Misc/Build.h"
 #include "Serialization/Archive.h"
 #include "Serialization/StructuredArchiveAdapters.h"
 #include "Templates/IsFloatingPoint.h"
diff --git a/Engine/Source/Runtime/Core/Public/Math/UnitConversion.h b/Engine/Source/Runtime/Core/Public/Math/UnitConversion.h
--- a/Engine/Source/Runtime/Core/Public/Math/UnitConversion.h
+++ b/Engine/Source/Runtime/Core/Public/Math/UnitConversion.h
@@ -212,4 +212,4 @@ bool LexTryParseString(FNumericUnit<T>& OutValue, const TCHAR* String);
 
 
 // Include template definitions
-#include "Math/UnitConversion.inl"
+#include "Math/UnitConversion.inl" // IWYU pragma: export
diff --git a/Engine/Source/Runtime/Core/Public/Math/UnitConversion.inl b/Engine/Source/Runtime/Core/Public/Math/UnitConversion.inl
--- a/Engine/Source/Runtime/Core/Public/Math/UnitConversion.inl
+++ b/Engine/Source/Runtime/Core/Public/Math/UnitConversion.inl
@@ -1,5 +1,7 @@
 // Copyright Epic Games, Inc. All Rights Reserved.
 
+#pragma once
+
 /** Inline file for UnitConversion.h to separate the implementation from the header */
 
 
diff --git a/Engine/Source/Runtime/Core/Public/Misc/App.h b/Engine/Source/Runtime/Core/Public/Misc/App.h
--- a/Engine/Source/Runtime/Core/Public/Misc/App.h
+++ b/Engine/Source/Runtime/Core/Public/Misc/App.h
@@ -9,6 +9,7 @@
 #include "CoreTypes.h"
 #include "Delegates/Delegate.h"
 #include "HAL/PlatformProcess.h"
+#include "Misc/Build.h"
 #include "Misc/CommandLine.h"
 #include "Misc/CoreMisc.h"
 #include "Misc/FrameRate.h"
diff --git a/Engine/Source/Runtime/Core/Public/Misc/AutomationTest.h b/Engine/Source/Runtime/Core/Public/Misc/AutomationTest.h
--- a/Engine/Source/Runtime/Core/Public/Misc/AutomationTest.h
+++ b/Engine/Source/Runtime/Core/Public/Misc/AutomationTest.h
@@ -2,37 +2,44 @@
 
 #pragma once
 
-#include "CoreTypes.h"
-#include "Misc/AssertionMacros.h"
+#include "Async/Async.h"
+#include "Async/Future.h"
 #include "Containers/Array.h"
-#include "Math/UnrealMathUtility.h"
-#include "Containers/UnrealString.h"
-#include "Templates/Function.h"
 #include "Containers/Map.h"
-#include "Math/Color.h"
-#include "Templates/SharedPointer.h"
+#include "Containers/Queue.h"
+#include "Containers/UnrealString.h"
+#include "CoreMinimal.h"
+#include "CoreTypes.h"
 #include "Delegates/Delegate.h"
-#include "Misc/Optional.h"
-#include "HAL/PlatformTime.h"
-#include "HAL/ThreadSafeBool.h" 
-#include "HAL/PlatformStackWalk.h"
-#include "HAL/LowLevelMemTracker.h"
+#include "Delegates/DelegateInstancesImpl.h"
 #include "GenericPlatform/GenericPlatformStackWalk.h"
-#include "Containers/Queue.h"
+#include "HAL/CriticalSection.h"
+#include "HAL/LowLevelMemTracker.h"
+#include "HAL/PlatformProcess.h"
+#include "HAL/PlatformStackWalk.h"
+#include "HAL/PlatformTime.h"
+#include "HAL/ThreadSafeBool.h"
+#include "Internationalization/Regex.h"
+#include "Math/Color.h"
+#include "Math/Color.h"
+#include "Math/MathFwd.h"
+#include "Math/Rotator.h"
+#include "Math/UnrealMathUtility.h"
+#include "Math/Vector.h"
+#include "Misc/AssertionMacros.h"
+#include "Misc/AutomationEvent.h"
+#include "Misc/Build.h"
 #include "Misc/CString.h"
 #include "Misc/FeedbackContext.h"
-#include "Async/Future.h"
-#include "Async/Async.h"
 #include "Misc/Guid.h"
-#include "Math/Vector.h"
-#include "Math/Color.h"
-#include "Math/Rotator.h"
-#include "HAL/PlatformProcess.h"
-#include "HAL/CriticalSection.h"
-#include "Misc/AutomationEvent.h"
-#include "Internationalization/Regex.h"
+#include "Misc/Optional.h"
+#include "Templates/Function.h"
+#include "Templates/SharedPointer.h"
+
 #include <atomic>
 
+class FAutomationTestBase;
+
 #ifndef WITH_AUTOMATION_TESTS
 	#define WITH_AUTOMATION_TESTS (WITH_DEV_AUTOMATION_TESTS || WITH_PERF_AUTOMATION_TESTS)
 #endif
diff --git a/Engine/Source/Runtime/Core/Public/Misc/ConfigCacheIni.h b/Engine/Source/Runtime/Core/Public/Misc/ConfigCacheIni.h
--- a/Engine/Source/Runtime/Core/Public/Misc/ConfigCacheIni.h
+++ b/Engine/Source/Runtime/Core/Public/Misc/ConfigCacheIni.h
@@ -20,6 +20,7 @@
 #include "Math/Vector.h"
 #include "Math/Vector2D.h"
 #include "Misc/AccessDetection.h"
+#include "Misc/Build.h"
 #include "Misc/ConfigTypes.h"
 #include "Misc/Paths.h"
 #include "Serialization/StructuredArchive.h"
@@ -1339,4 +1340,4 @@ CORE_API void DeleteRecordedConfigReadsFromIni();
  * Helper function to deal with "True","False","Yes","No","On","Off"
  */
 UE_DEPRECATED(5.1, "Use UE::ConfigUtilities::ConvertValueFromHumanFriendlyValue")
-CORE_API const TCHAR* ConvertValueFromHumanFriendlyValue(const TCHAR* Value);
+CORE_API const TCHAR* ConvertValueFromHumanFriendlyValue(const TCHAR* Value);
\ No newline at end of file
diff --git a/Engine/Source/Runtime/Core/Public/Misc/CoreDelegates.h b/Engine/Source/Runtime/Core/Public/Misc/CoreDelegates.h
--- a/Engine/Source/Runtime/Core/Public/Misc/CoreDelegates.h
+++ b/Engine/Source/Runtime/Core/Public/Misc/CoreDelegates.h
@@ -15,6 +15,7 @@
 #include "Math/IntVector.h"
 #include "Math/MathFwd.h"
 #include "Misc/AES.h"
+#include "Misc/Build.h"
 #include "Templates/SharedPointerInternals.h"
 
 class AActor;
diff --git a/Engine/Source/Runtime/Core/Public/Misc/CoreMisc.h b/Engine/Source/Runtime/Core/Public/Misc/CoreMisc.h
--- a/Engine/Source/Runtime/Core/Public/Misc/CoreMisc.h
+++ b/Engine/Source/Runtime/Core/Public/Misc/CoreMisc.h
@@ -10,6 +10,7 @@
 #include "CoreTypes.h"
 #include "HAL/ThreadSingleton.h"
 #include "Math/IntPoint.h"
+#include "Misc/Build.h"
 #include "Misc/Exec.h"
 #include "Templates/Function.h"
 #include "UObject/NameTypes.h"
diff --git a/Engine/Source/Runtime/Core/Public/Misc/Guid.h b/Engine/Source/Runtime/Core/Public/Misc/Guid.h
--- a/Engine/Source/Runtime/Core/Public/Misc/Guid.h
+++ b/Engine/Source/Runtime/Core/Public/Misc/Guid.h
@@ -5,6 +5,7 @@
 #include "Containers/StringFwd.h"
 #include "Containers/UnrealString.h"
 #include "CoreTypes.h"
+#include "HAL/PreprocessorHelpers.h"
 #include "Hash/CityHash.h"
 #include "Misc/AssertionMacros.h"
 #include "Misc/Crc.h"
diff --git a/Engine/Source/Runtime/Core/Public/Misc/MTAccessDetector.h b/Engine/Source/Runtime/Core/Public/Misc/MTAccessDetector.h
--- a/Engine/Source/Runtime/Core/Public/Misc/MTAccessDetector.h
+++ b/Engine/Source/Runtime/Core/Public/Misc/MTAccessDetector.h
@@ -5,7 +5,9 @@
 #include "CoreGlobals.h"
 #include "CoreTypes.h"
 #include "HAL/PlatformTLS.h"
+#include "HAL/PreprocessorHelpers.h"
 #include "Misc/AssertionMacros.h"
+#include "Misc/Build.h"
 
 #define ENABLE_MT_DETECTOR DO_CHECK
 
diff --git a/Engine/Source/Runtime/Core/Public/Misc/MemStack.h b/Engine/Source/Runtime/Core/Public/Misc/MemStack.h
--- a/Engine/Source/Runtime/Core/Public/Misc/MemStack.h
+++ b/Engine/Source/Runtime/Core/Public/Misc/MemStack.h
@@ -11,6 +11,7 @@
 #include "HAL/UnrealMemory.h"
 #include "Math/UnrealMathUtility.h"
 #include "Misc/AssertionMacros.h"
+#include "Misc/Build.h"
 #include "Misc/NoopCounter.h"
 #include "Templates/AlignmentTemplates.h"
 
diff --git a/Engine/Source/Runtime/Core/Public/Misc/OutputDeviceMemory.h b/Engine/Source/Runtime/Core/Public/Misc/OutputDeviceMemory.h
--- a/Engine/Source/Runtime/Core/Public/Misc/OutputDeviceMemory.h
+++ b/Engine/Source/Runtime/Core/Public/Misc/OutputDeviceMemory.h
@@ -8,6 +8,7 @@
 #include "Logging/LogVerbosity.h"
 #include "Misc/OutputDevice.h"
 #include "Serialization/Archive.h"
+#include "UObject/NameTypes.h"
 
 /**
  * Memory output device. Logs only into pre-allocated memory buffer.
diff --git a/Engine/Source/Runtime/Core/Public/Misc/Parse.h b/Engine/Source/Runtime/Core/Public/Misc/Parse.h
--- a/Engine/Source/Runtime/Core/Public/Misc/Parse.h
+++ b/Engine/Source/Runtime/Core/Public/Misc/Parse.h
@@ -5,6 +5,7 @@
 #include "Containers/StringFwd.h"
 #include "Containers/UnrealString.h"
 #include "CoreTypes.h"
+#include "Misc/Build.h"
 
 class FExec;
 class FName;
diff --git a/Engine/Source/Runtime/Core/Public/Misc/RuntimeErrors.h b/Engine/Source/Runtime/Core/Public/Misc/RuntimeErrors.h
--- a/Engine/Source/Runtime/Core/Public/Misc/RuntimeErrors.h
+++ b/Engine/Source/Runtime/Core/Public/Misc/RuntimeErrors.h
@@ -7,6 +7,7 @@
 #include "HAL/PlatformMisc.h"
 #include "Logging/LogVerbosity.h"
 #include "Misc/AssertionMacros.h"
+#include "Misc/Build.h"
 #include "Misc/CoreDelegates.h"
 
 class FText;
diff --git a/Engine/Source/Runtime/Core/Public/Misc/SecureHash.h b/Engine/Source/Runtime/Core/Public/Misc/SecureHash.h
--- a/Engine/Source/Runtime/Core/Public/Misc/SecureHash.h
+++ b/Engine/Source/Runtime/Core/Public/Misc/SecureHash.h
@@ -10,6 +10,7 @@
 #include "Containers/UnrealString.h"
 #include "CoreMinimal.h"
 #include "CoreTypes.h"
+#include "HAL/PreprocessorHelpers.h"
 #include "HAL/UnrealMemory.h"
 #include "Serialization/BufferReader.h"
 #include "Serialization/MemoryLayout.h"
diff --git a/Engine/Source/Runtime/Core/Public/Misc/TimeGuard.h b/Engine/Source/Runtime/Core/Public/Misc/TimeGuard.h
--- a/Engine/Source/Runtime/Core/Public/Misc/TimeGuard.h
+++ b/Engine/Source/Runtime/Core/Public/Misc/TimeGuard.h
@@ -7,6 +7,7 @@
 #include "CoreTypes.h"
 #include "Delegates/Delegate.h"
 #include "HAL/PlatformTime.h"
+#include "Misc/Build.h"
 #include "Misc/DateTime.h"
 
 #ifndef DO_TIMEGUARD
diff --git a/Engine/Source/Runtime/Core/Public/Modules/ModuleManager.h b/Engine/Source/Runtime/Core/Public/Modules/ModuleManager.h
--- a/Engine/Source/Runtime/Core/Public/Modules/ModuleManager.h
+++ b/Engine/Source/Runtime/Core/Public/Modules/ModuleManager.h
@@ -10,7 +10,9 @@
 #include "Delegates/Delegate.h"
 #include "Delegates/DelegateBase.h"
 #include "HAL/CriticalSection.h"
+#include "HAL/PreprocessorHelpers.h"
 #include "Misc/AssertionMacros.h"
+#include "Misc/Build.h"
 #include "Misc/CoreMisc.h"
 #include "Misc/EnumClassFlags.h"
 #include "Misc/Optional.h"
@@ -24,6 +26,7 @@
 #include "Templates/UnrealTemplate.h"
 #include "UObject/NameTypes.h"
 
+class FArchive;
 class FOutputDevice;
 class UClass;
 class UWorld;
diff --git a/Engine/Source/Runtime/Core/Public/ProfilingDebugging/ABTesting.h b/Engine/Source/Runtime/Core/Public/ProfilingDebugging/ABTesting.h
--- a/Engine/Source/Runtime/Core/Public/ProfilingDebugging/ABTesting.h
+++ b/Engine/Source/Runtime/Core/Public/ProfilingDebugging/ABTesting.h
@@ -11,6 +11,7 @@
 #include "Containers/UnrealString.h"
 #include "CoreTypes.h"
 #include "Math/RandomStream.h"
+#include "Misc/Build.h"
 
 #if !UE_BUILD_SHIPPING
 #define ENABLE_ABTEST 1
diff --git a/Engine/Source/Runtime/Core/Public/ProfilingDebugging/CallstackTrace.h b/Engine/Source/Runtime/Core/Public/ProfilingDebugging/CallstackTrace.h
--- a/Engine/Source/Runtime/Core/Public/ProfilingDebugging/CallstackTrace.h
+++ b/Engine/Source/Runtime/Core/Public/ProfilingDebugging/CallstackTrace.h
@@ -3,6 +3,7 @@
 #include "HAL/Platform.h"
 #include "HAL/PlatformTLS.h"
 #include "Math/NumericLimits.h"
+#include "Misc/Build.h"
 #include "TraceLog/Public/Trace/Config.h"
 
 ////////////////////////////////////////////////////////////////////////////////
diff --git a/Engine/Source/Runtime/Core/Public/ProfilingDebugging/CookStats.h b/Engine/Source/Runtime/Core/Public/ProfilingDebugging/CookStats.h
--- a/Engine/Source/Runtime/Core/Public/ProfilingDebugging/CookStats.h
+++ b/Engine/Source/Runtime/Core/Public/ProfilingDebugging/CookStats.h
@@ -11,6 +11,7 @@
 #include "Delegates/DelegateBase.h"
 #include "Delegates/IDelegateInstance.h"
 #include "HAL/PlatformTime.h"
+#include "HAL/PreprocessorHelpers.h"
 #include "HAL/ThreadSafeCounter64.h"
 #include "Templates/Decay.h"
 #include "Templates/Function.h"
diff --git a/Engine/Source/Runtime/Core/Public/ProfilingDebugging/CountersTrace.h b/Engine/Source/Runtime/Core/Public/ProfilingDebugging/CountersTrace.h
--- a/Engine/Source/Runtime/Core/Public/ProfilingDebugging/CountersTrace.h
+++ b/Engine/Source/Runtime/Core/Public/ProfilingDebugging/CountersTrace.h
@@ -3,6 +3,8 @@
 #pragma once
 
 #include "CoreTypes.h"
+#include "HAL/PreprocessorHelpers.h"
+#include "Misc/Build.h"
 #include "Trace/Config.h"
 #include "Trace/Trace.h"
 
diff --git a/Engine/Source/Runtime/Core/Public/ProfilingDebugging/CpuProfilerTrace.h b/Engine/Source/Runtime/Core/Public/ProfilingDebugging/CpuProfilerTrace.h
--- a/Engine/Source/Runtime/Core/Public/ProfilingDebugging/CpuProfilerTrace.h
+++ b/Engine/Source/Runtime/Core/Public/ProfilingDebugging/CpuProfilerTrace.h
@@ -3,6 +3,8 @@
 #pragma once
 
 #include "CoreTypes.h"
+#include "HAL/PreprocessorHelpers.h"
+#include "Misc/Build.h"
 #include "Trace/Config.h"
 #include "Trace/Detail/Channel.h"
 #include "Trace/Detail/Channel.inl"
diff --git a/Engine/Source/Runtime/Core/Public/ProfilingDebugging/CsvProfiler.h b/Engine/Source/Runtime/Core/Public/ProfilingDebugging/CsvProfiler.h
--- a/Engine/Source/Runtime/Core/Public/ProfilingDebugging/CsvProfiler.h
+++ b/Engine/Source/Runtime/Core/Public/ProfilingDebugging/CsvProfiler.h
@@ -7,21 +7,31 @@
 
 #pragma once
 
-#include "CoreTypes.h"
-#include "Containers/Queue.h"
-#include "UObject/NameTypes.h"
-#include "Templates/UniquePtr.h"
 #include "Async/Future.h"
 #include "Async/TaskGraphInterfaces.h"
+#include "Containers/Queue.h"
+#include "CoreMinimal.h"
+#include "CoreTypes.h"
+#include "HAL/PreprocessorHelpers.h"
+#include "HAL/ThreadSafeCounter.h"
 #include "Misc/EnumClassFlags.h"
-#include "ProfilingDebugging/MiscTrace.h"
 #include "ProfilingDebugging/CsvProfilerConfig.h"
 #include "ProfilingDebugging/CsvProfilerTrace.h"
+#include "ProfilingDebugging/MiscTrace.h"
 #include "Templates/IsArrayOrRefOfTypeByPredicate.h"
+#include "Templates/IsValidVariadicFunctionArg.h"
+#include "Templates/UniquePtr.h"
 #include "Traits/IsCharEncodingCompatibleWith.h"
+#include "UObject/NameTypes.h"
 
 #include <atomic>
 
+class FCsvProfiler;
+class FEvent;
+class FScopedCsvStat;
+class FScopedCsvStatExclusive;
+struct FCsvDeclaredStat;
+
 #if CSV_PROFILER
 
 // Helpers
@@ -115,10 +125,10 @@
 #if CSV_PROFILER
 
 class FCsvProfilerFrame;
-class FCsvProfilerThreadData;
 class FCsvProfilerProcessingThread;
-struct FCsvPersistentCustomStats;
+class FCsvProfilerThreadData;
 class FName;
+struct FCsvPersistentCustomStats;
 
 enum class ECsvCustomStatOp : uint8
 {
diff --git a/Engine/Source/Runtime/Core/Public/ProfilingDebugging/CsvProfilerTrace.h b/Engine/Source/Runtime/Core/Public/ProfilingDebugging/CsvProfilerTrace.h
--- a/Engine/Source/Runtime/Core/Public/ProfilingDebugging/CsvProfilerTrace.h
+++ b/Engine/Source/Runtime/Core/Public/ProfilingDebugging/CsvProfilerTrace.h
@@ -3,8 +3,10 @@
 #pragma once
 
 #include "CoreTypes.h"
-#include "Trace/Config.h"
+#include "HAL/PreprocessorHelpers.h"
+#include "Misc/Build.h"
 #include "ProfilingDebugging/CsvProfilerConfig.h"
+#include "Trace/Config.h"
 
 // Note: Enabling this adds about 15% to the instrumentation overhead of the CSV profiler, even if Insights is not capturing
 #ifndef CSVPROFILERTRACE_ENABLED
diff --git a/Engine/Source/Runtime/Core/Public/ProfilingDebugging/ExternalProfiler.h b/Engine/Source/Runtime/Core/Public/ProfilingDebugging/ExternalProfiler.h
--- a/Engine/Source/Runtime/Core/Public/ProfilingDebugging/ExternalProfiler.h
+++ b/Engine/Source/Runtime/Core/Public/ProfilingDebugging/ExternalProfiler.h
@@ -5,6 +5,7 @@
 #include "Containers/StringConv.h"
 #include "CoreTypes.h"
 #include "Features/IModularFeature.h"
+#include "Misc/Build.h"
 #include "UObject/NameTypes.h"
 
 #ifndef UE_EXTERNAL_PROFILING_ENABLED
diff --git a/Engine/Source/Runtime/Core/Public/ProfilingDebugging/Histogram.h b/Engine/Source/Runtime/Core/Public/ProfilingDebugging/Histogram.h
--- a/Engine/Source/Runtime/Core/Public/ProfilingDebugging/Histogram.h
+++ b/Engine/Source/Runtime/Core/Public/ProfilingDebugging/Histogram.h
@@ -19,6 +19,8 @@
 #include "Misc/AssertionMacros.h"
 #include "Templates/UnrealTemplate.h"
 
+#include <initializer_list>
+
 struct FAnalyticsEventAttribute;
 struct FHistogramBuilder;
 template <typename FuncType> class TFunctionRef;
diff --git a/Engine/Source/Runtime/Core/Public/ProfilingDebugging/InstanceCounter.h b/Engine/Source/Runtime/Core/Public/ProfilingDebugging/InstanceCounter.h
--- a/Engine/Source/Runtime/Core/Public/ProfilingDebugging/InstanceCounter.h
+++ b/Engine/Source/Runtime/Core/Public/ProfilingDebugging/InstanceCounter.h
@@ -5,6 +5,8 @@
 #include "Containers/UnrealString.h"
 #include "CoreTypes.h"
 #include "HAL/CriticalSection.h"
+#include "HAL/PreprocessorHelpers.h"
+#include "Misc/Build.h"
 #include "UObject/NameTypes.h"
 
 class FOutputDevice;
diff --git a/Engine/Source/Runtime/Core/Public/ProfilingDebugging/LoadTimeTracker.h b/Engine/Source/Runtime/Core/Public/ProfilingDebugging/LoadTimeTracker.h
--- a/Engine/Source/Runtime/Core/Public/ProfilingDebugging/LoadTimeTracker.h
+++ b/Engine/Source/Runtime/Core/Public/ProfilingDebugging/LoadTimeTracker.h
@@ -10,6 +10,7 @@
 #include "Containers/Map.h"
 #include "Containers/SparseArray.h"
 #include "CoreTypes.h"
+#include "Misc/Build.h"
 #include "Misc/Optional.h"
 #include "ProfilingDebugging/CpuProfilerTrace.h"
 #include "ProfilingDebugging/ScopedTimers.h"
diff --git a/Engine/Source/Runtime/Core/Public/ProfilingDebugging/MiscTrace.h b/Engine/Source/Runtime/Core/Public/ProfilingDebugging/MiscTrace.h
--- a/Engine/Source/Runtime/Core/Public/ProfilingDebugging/MiscTrace.h
+++ b/Engine/Source/Runtime/Core/Public/ProfilingDebugging/MiscTrace.h
@@ -4,6 +4,8 @@
 
 #include "Containers/ContainersFwd.h"
 #include "CoreTypes.h"
+#include "HAL/PreprocessorHelpers.h"
+#include "Misc/Build.h"
 #include "ProfilingDebugging/FormatArgsTrace.h"
 #include "Trace/Config.h"
 
diff --git a/Engine/Source/Runtime/Core/Public/ProfilingDebugging/PlatformFileTrace.h b/Engine/Source/Runtime/Core/Public/ProfilingDebugging/PlatformFileTrace.h
--- a/Engine/Source/Runtime/Core/Public/ProfilingDebugging/PlatformFileTrace.h
+++ b/Engine/Source/Runtime/Core/Public/ProfilingDebugging/PlatformFileTrace.h
@@ -3,6 +3,7 @@
 #pragma once
 
 #include "CoreTypes.h"
+#include "Misc/Build.h"
 #include "Trace/Config.h"
 
 #if !UE_BUILD_SHIPPING
diff --git a/Engine/Source/Runtime/Core/Public/ProfilingDebugging/StallDetector.h b/Engine/Source/Runtime/Core/Public/ProfilingDebugging/StallDetector.h
--- a/Engine/Source/Runtime/Core/Public/ProfilingDebugging/StallDetector.h
+++ b/Engine/Source/Runtime/Core/Public/ProfilingDebugging/StallDetector.h
@@ -9,6 +9,7 @@
 #include "HAL/CriticalSection.h"
 #include "HAL/PreprocessorHelpers.h"
 #include "Logging/LogMacros.h"
+#include "Misc/Build.h"
 #include "Misc/ScopeLock.h"
 #include "ProfilingDebugging/CountersTrace.h"
 
diff --git a/Engine/Source/Runtime/Core/Public/ProfilingDebugging/TraceAuxiliary.h b/Engine/Source/Runtime/Core/Public/ProfilingDebugging/TraceAuxiliary.h
--- a/Engine/Source/Runtime/Core/Public/ProfilingDebugging/TraceAuxiliary.h
+++ b/Engine/Source/Runtime/Core/Public/ProfilingDebugging/TraceAuxiliary.h
@@ -8,6 +8,7 @@
 #include "HAL/Platform.h"
 #include "Logging/LogCategory.h"
 #include "Logging/LogMacros.h"
+#include "Misc/Build.h"
 
 ////////////////////////////////////////////////////////////////////////////////
 class CORE_API FTraceAuxiliary
diff --git a/Engine/Source/Runtime/Core/Public/Serialization/Archive.h b/Engine/Source/Runtime/Core/Public/Serialization/Archive.h
--- a/Engine/Source/Runtime/Core/Public/Serialization/Archive.h
+++ b/Engine/Source/Runtime/Core/Public/Serialization/Archive.h
@@ -8,6 +8,7 @@
 #include "Internationalization/TextNamespaceFwd.h"
 #include "Math/MathFwd.h"
 #include "Misc/AssertionMacros.h"
+#include "Misc/Build.h"
 #include "Misc/CompressionFlags.h"
 #include "Misc/EngineVersionBase.h"
 #include "Misc/VarArgs.h"
diff --git a/Engine/Source/Runtime/Core/Public/Serialization/BitWriter.h b/Engine/Source/Runtime/Core/Public/Serialization/BitWriter.h
--- a/Engine/Source/Runtime/Core/Public/Serialization/BitWriter.h
+++ b/Engine/Source/Runtime/Core/Public/Serialization/BitWriter.h
@@ -6,6 +6,7 @@
 #include "CoreTypes.h"
 #include "Math/UnrealMathUtility.h"
 #include "Misc/AssertionMacros.h"
+#include "Misc/Build.h"
 #include "Serialization/BitArchive.h"
 
 class FArchive;
diff --git a/Engine/Source/Runtime/Core/Public/Serialization/LoadTimeTrace.h b/Engine/Source/Runtime/Core/Public/Serialization/LoadTimeTrace.h
--- a/Engine/Source/Runtime/Core/Public/Serialization/LoadTimeTrace.h
+++ b/Engine/Source/Runtime/Core/Public/Serialization/LoadTimeTrace.h
@@ -3,6 +3,7 @@
 #pragma once
 
 #include "CoreTypes.h"
+#include "Misc/Build.h"
 #include "ProfilingDebugging/FormatArgsTrace.h"
 #include "Trace/Config.h"
 #include "Trace/Trace.h"
diff --git a/Engine/Source/Runtime/Core/Public/Serialization/MemoryImage.h b/Engine/Source/Runtime/Core/Public/Serialization/MemoryImage.h
--- a/Engine/Source/Runtime/Core/Public/Serialization/MemoryImage.h
+++ b/Engine/Source/Runtime/Core/Public/Serialization/MemoryImage.h
@@ -8,6 +8,7 @@
 #include "Containers/UnrealString.h"
 #include "CoreMinimal.h"
 #include "CoreTypes.h"
+#include "HAL/PreprocessorHelpers.h"
 #include "Misc/SecureHash.h"
 #include "Serialization/Archive.h"
 #include "Serialization/MemoryImageWriter.h"
diff --git a/Engine/Source/Runtime/Core/Public/Serialization/StructuredArchive.h b/Engine/Source/Runtime/Core/Public/Serialization/StructuredArchive.h
--- a/Engine/Source/Runtime/Core/Public/Serialization/StructuredArchive.h
+++ b/Engine/Source/Runtime/Core/Public/Serialization/StructuredArchive.h
@@ -6,6 +6,7 @@
 #include "Containers/ContainerAllocationPolicies.h"
 #include "CoreTypes.h"
 #include "Formatters/BinaryArchiveFormatter.h"
+#include "Misc/Build.h"
 #include "Serialization/Archive.h"
 #include "Serialization/StructuredArchiveAdapters.h"
 #include "Serialization/StructuredArchiveDefines.h"
diff --git a/Engine/Source/Runtime/Core/Public/Serialization/StructuredArchiveChildReader.h b/Engine/Source/Runtime/Core/Public/Serialization/StructuredArchiveChildReader.h
--- a/Engine/Source/Runtime/Core/Public/Serialization/StructuredArchiveChildReader.h
+++ b/Engine/Source/Runtime/Core/Public/Serialization/StructuredArchiveChildReader.h
@@ -3,6 +3,7 @@
 #pragma once
 
 #include "HAL/Platform.h"
+#include "Misc/Build.h"
 #include "Misc/Optional.h"
 #include "Serialization/StructuredArchive.h"
 #include "Serialization/StructuredArchiveSlots.h"
diff --git a/Engine/Source/Runtime/Core/Public/Serialization/StructuredArchiveSlots.h b/Engine/Source/Runtime/Core/Public/Serialization/StructuredArchiveSlots.h
--- a/Engine/Source/Runtime/Core/Public/Serialization/StructuredArchiveSlots.h
+++ b/Engine/Source/Runtime/Core/Public/Serialization/StructuredArchiveSlots.h
@@ -5,6 +5,7 @@
 #include "Containers/Array.h"
 #include "CoreTypes.h"
 #include "Formatters/BinaryArchiveFormatter.h"
+#include "Misc/Build.h"
 #include "Misc/Optional.h"
 #include "Serialization/Archive.h"
 #include "Serialization/StructuredArchiveFwd.h"
diff --git a/Engine/Source/Runtime/Core/Public/Stats/Stats2.h b/Engine/Source/Runtime/Core/Public/Stats/Stats2.h
--- a/Engine/Source/Runtime/Core/Public/Stats/Stats2.h
+++ b/Engine/Source/Runtime/Core/Public/Stats/Stats2.h
@@ -15,6 +15,7 @@
 #include "HAL/ThreadSingleton.h"
 #include "Math/Color.h"
 #include "Math/NumericLimits.h"
+#include "Misc/Build.h"
 #include "Misc/EnumClassFlags.h"
 #include "ProfilingDebugging/CpuProfilerTrace.h"
 #include "ProfilingDebugging/MiscTrace.h"
diff --git a/Engine/Source/Runtime/Core/Public/Stats/StatsData.h b/Engine/Source/Runtime/Core/Public/Stats/StatsData.h
--- a/Engine/Source/Runtime/Core/Public/Stats/StatsData.h
+++ b/Engine/Source/Runtime/Core/Public/Stats/StatsData.h
@@ -15,6 +15,7 @@
 #include "Math/NumericLimits.h"
 #include "Math/UnrealMathUtility.h"
 #include "Misc/AssertionMacros.h"
+#include "Misc/Build.h"
 #include "Misc/Parse.h"
 #include "Stats/Stats.h"
 #include "Stats/Stats2.h"
diff --git a/Engine/Source/Runtime/Core/Public/Stats/StatsFile.h b/Engine/Source/Runtime/Core/Public/Stats/StatsFile.h
--- a/Engine/Source/Runtime/Core/Public/Stats/StatsFile.h
+++ b/Engine/Source/Runtime/Core/Public/Stats/StatsFile.h
@@ -16,11 +16,13 @@
 #include "Internationalization/Text.h"
 #include "Logging/LogMacros.h"
 #include "Misc/AssertionMacros.h"
+#include "Misc/Build.h"
 #include "Misc/Compression.h"
 #include "Stats/Stats.h"
 #include "Stats/Stats2.h"
 #include "Stats/StatsCommon.h"
 #include "Stats/StatsData.h"
+#include "Trace/Detail/Channel.h"
 #include "UObject/NameTypes.h"
 
 class FAsyncStatsWrite;
diff --git a/Engine/Source/Runtime/Core/Public/Stats/StatsHierarchical.h b/Engine/Source/Runtime/Core/Public/Stats/StatsHierarchical.h
--- a/Engine/Source/Runtime/Core/Public/Stats/StatsHierarchical.h
+++ b/Engine/Source/Runtime/Core/Public/Stats/StatsHierarchical.h
@@ -6,7 +6,9 @@
 #include "Containers/UnrealString.h"
 #include "CoreTypes.h"
 #include "HAL/PlatformTime.h"
+#include "HAL/PreprocessorHelpers.h"
 #include "Logging/MessageLog.h"
+#include "Misc/Build.h"
 #include "Templates/SharedPointer.h"
 #include "UObject/NameTypes.h"
 
diff --git a/Engine/Source/Runtime/Core/Public/Stats/StatsMallocProfilerProxy.h b/Engine/Source/Runtime/Core/Public/Stats/StatsMallocProfilerProxy.h
--- a/Engine/Source/Runtime/Core/Public/Stats/StatsMallocProfilerProxy.h
+++ b/Engine/Source/Runtime/Core/Public/Stats/StatsMallocProfilerProxy.h
@@ -5,6 +5,7 @@
 #include "CoreTypes.h"
 #include "HAL/MemoryBase.h"
 #include "HAL/ThreadSafeCounter.h"
+#include "Misc/Build.h"
 
 class FOutputDevice;
 class UWorld;
diff --git a/Engine/Source/Runtime/Core/Public/Stats/StatsMisc.h b/Engine/Source/Runtime/Core/Public/Stats/StatsMisc.h
--- a/Engine/Source/Runtime/Core/Public/Stats/StatsMisc.h
+++ b/Engine/Source/Runtime/Core/Public/Stats/StatsMisc.h
@@ -5,7 +5,9 @@
 #include "Containers/UnrealString.h"
 #include "CoreTypes.h"
 #include "HAL/PlatformTime.h"
+#include "HAL/PreprocessorHelpers.h"
 #include "Misc/AssertionMacros.h"
+#include "Misc/Build.h"
 #include "Templates/UnrealTemplate.h"
 
 /**
diff --git a/Engine/Source/Runtime/Core/Public/Stats/StatsTrace.h b/Engine/Source/Runtime/Core/Public/Stats/StatsTrace.h
--- a/Engine/Source/Runtime/Core/Public/Stats/StatsTrace.h
+++ b/Engine/Source/Runtime/Core/Public/Stats/StatsTrace.h
@@ -3,6 +3,7 @@
 #pragma once
 
 #include "CoreTypes.h"
+#include "Misc/Build.h"
 #include "Trace/Config.h"
 
 #if !defined(STATSTRACE_ENABLED)
diff --git a/Engine/Source/Runtime/Core/Public/Templates/RefCounting.h b/Engine/Source/Runtime/Core/Public/Templates/RefCounting.h
--- a/Engine/Source/Runtime/Core/Public/Templates/RefCounting.h
+++ b/Engine/Source/Runtime/Core/Public/Templates/RefCounting.h
@@ -4,8 +4,10 @@
 
 #include "CoreTypes.h"
 #include "HAL/PlatformAtomics.h"
+#include "HAL/PreprocessorHelpers.h"
 #include "HAL/ThreadSafeCounter.h"
 #include "Misc/AssertionMacros.h"
+#include "Misc/Build.h"
 #include "Serialization/Archive.h"
 #include "Serialization/MemoryLayout.h"
 #include "Templates/TypeHash.h"
diff --git a/Engine/Source/Runtime/Core/Public/Unix/UnixCriticalSection.h b/Engine/Source/Runtime/Core/Public/Unix/UnixCriticalSection.h
--- a/Engine/Source/Runtime/Core/Public/Unix/UnixCriticalSection.h
+++ b/Engine/Source/Runtime/Core/Public/Unix/UnixCriticalSection.h
@@ -2,10 +2,11 @@
 
 #pragma once
 
+#include "Containers/UnrealString.h"
 #include "CoreTypes.h"
-#include "Misc/Timespan.h"
 #include "HAL/PThreadCriticalSection.h"
 #include "HAL/PThreadRWLock.h"
+#include "Misc/Timespan.h"
 
 /**
 	* Unix implementation of the FSystemWideCriticalSection. Uses exclusive file locking.
diff --git a/Engine/Source/Runtime/Core/Public/Unix/UnixErrorOutputDevice.h b/Engine/Source/Runtime/Core/Public/Unix/UnixErrorOutputDevice.h
--- a/Engine/Source/Runtime/Core/Public/Unix/UnixErrorOutputDevice.h
+++ b/Engine/Source/Runtime/Core/Public/Unix/UnixErrorOutputDevice.h
@@ -2,6 +2,7 @@
 
 #pragma once
 
+#include "CoreMinimal.h"
 #include "CoreTypes.h"
 #include "Misc/OutputDeviceError.h"
 
diff --git a/Engine/Source/Runtime/Core/Public/Unix/UnixPlatformCrashContext.h b/Engine/Source/Runtime/Core/Public/Unix/UnixPlatformCrashContext.h
--- a/Engine/Source/Runtime/Core/Public/Unix/UnixPlatformCrashContext.h
+++ b/Engine/Source/Runtime/Core/Public/Unix/UnixPlatformCrashContext.h
@@ -2,8 +2,10 @@
 
 #pragma once
 
+#include "CoreMinimal.h"
 #include "CoreTypes.h"
 #include "GenericPlatform/GenericPlatformCrashContext.h"
+#include "Templates/Atomic.h"
 
 /** Passed in through sigqueue for gathering of a callstack from a signal */
 struct ThreadStackUserData
diff --git a/Engine/Source/Runtime/Core/Public/Unix/UnixPlatformFile.h b/Engine/Source/Runtime/Core/Public/Unix/UnixPlatformFile.h
--- a/Engine/Source/Runtime/Core/Public/Unix/UnixPlatformFile.h
+++ b/Engine/Source/Runtime/Core/Public/Unix/UnixPlatformFile.h
@@ -6,11 +6,13 @@
 
 #pragma once
 
-#include "CoreTypes.h"
 #include "Containers/UnrealString.h"
-#include "Misc/DateTime.h"
+#include "CoreTypes.h"
 #include "GenericPlatform/GenericPlatformFile.h"
+#include "HAL/PlatformFile.h"
+#include "Misc/DateTime.h"
 
+class IMappedFileHandle;
 template <typename FuncType> class TFunctionRef;
 
 /**
diff --git a/Engine/Source/Runtime/Core/Public/Unix/UnixPlatformOutputDevices.h b/Engine/Source/Runtime/Core/Public/Unix/UnixPlatformOutputDevices.h
--- a/Engine/Source/Runtime/Core/Public/Unix/UnixPlatformOutputDevices.h
+++ b/Engine/Source/Runtime/Core/Public/Unix/UnixPlatformOutputDevices.h
@@ -6,9 +6,13 @@
 
 #pragma once
 
-#include "CoreTypes.h"
 #include "Containers/UnrealString.h"
+#include "CoreTypes.h"
 #include "GenericPlatform/GenericPlatformOutputDevices.h"
+#include "HAL/PlatformOutputDevices.h"
+
+class FOutputDevice;
+class FOutputDeviceError;
 
 struct CORE_API FUnixOutputDevices : public FGenericPlatformOutputDevices
 {
diff --git a/Engine/Source/Runtime/Core/Public/Unix/UnixPlatformProcess.h b/Engine/Source/Runtime/Core/Public/Unix/UnixPlatformProcess.h
--- a/Engine/Source/Runtime/Core/Public/Unix/UnixPlatformProcess.h
+++ b/Engine/Source/Runtime/Core/Public/Unix/UnixPlatformProcess.h
@@ -6,10 +6,12 @@
 
 #pragma once
 
-#include "CoreTypes.h"
-#include "Misc/AssertionMacros.h"
 #include "Containers/UnrealString.h"
+#include "CoreMinimal.h"
+#include "CoreTypes.h"
 #include "GenericPlatform/GenericPlatformProcess.h"
+#include "HAL/PlatformProcess.h"
+#include "Misc/AssertionMacros.h"
 
 class Error;
 
diff --git a/Engine/Source/Runtime/Core/Public/Unix/UnixPlatformStackWalk.h b/Engine/Source/Runtime/Core/Public/Unix/UnixPlatformStackWalk.h
--- a/Engine/Source/Runtime/Core/Public/Unix/UnixPlatformStackWalk.h
+++ b/Engine/Source/Runtime/Core/Public/Unix/UnixPlatformStackWalk.h
@@ -8,6 +8,7 @@
 
 #include "CoreTypes.h"
 #include "GenericPlatform/GenericPlatformStackWalk.h"
+#include "HAL/PlatformStackWalk.h"
 
 struct FGenericCrashContext;
 
diff --git a/Engine/Source/Runtime/Core/Public/Unix/UnixPlatformTime.h b/Engine/Source/Runtime/Core/Public/Unix/UnixPlatformTime.h
--- a/Engine/Source/Runtime/Core/Public/Unix/UnixPlatformTime.h
+++ b/Engine/Source/Runtime/Core/Public/Unix/UnixPlatformTime.h
@@ -8,6 +8,7 @@
 
 #include "CoreTypes.h"
 #include "GenericPlatform/GenericPlatformTime.h"
+#include "HAL/PlatformTime.h"
 
 /**
  * Unix implementation of the Time OS functions
diff --git a/Engine/Source/Runtime/Core/Public/Virtualization/VirtualizationSystem.h b/Engine/Source/Runtime/Core/Public/Virtualization/VirtualizationSystem.h
--- a/Engine/Source/Runtime/Core/Public/Virtualization/VirtualizationSystem.h
+++ b/Engine/Source/Runtime/Core/Public/Virtualization/VirtualizationSystem.h
@@ -18,6 +18,7 @@
 #include "Misc/ConfigCacheIni.h"
 #include "Templates/Function.h"
 #include "Templates/UniquePtr.h"
+#include "Templates/UnrealTemplate.h"
 #include "UObject/NameTypes.h"
 
 class FPackagePath;
diff --git a/Engine/Source/Runtime/Core/Tests/Containers/TestUtils.cpp b/Engine/Source/Runtime/Core/Tests/Containers/TestUtils.cpp
--- a/Engine/Source/Runtime/Core/Tests/Containers/TestUtils.cpp
+++ b/Engine/Source/Runtime/Core/Tests/Containers/TestUtils.cpp
@@ -2,8 +2,6 @@
 
 #include "TestUtils.h"
 
-#include "Math/UnrealMathUtility.h"
-
 #include <numeric>
 
 namespace Test
