Ran IWYU again after fixing some IWYU mapping issues.

diff --git a/Engine/Source/Runtime/Core/Internal/IO/IoStore.h b/Engine/Source/Runtime/Core/Internal/IO/IoStore.h
--- a/Engine/Source/Runtime/Core/Internal/IO/IoStore.h
+++ b/Engine/Source/Runtime/Core/Internal/IO/IoStore.h
@@ -2,8 +2,12 @@
 
 #pragma once
 
+#include "CoreMinimal.h"
+#include "IO/IoContainerId.h"
 #include "IO/IoDispatcher.h"
-#include "IO/IoDirectoryIndex.h"
+#include "Misc/EnumClassFlags.h"
+#include "Misc/Guid.h"
+#include "Misc/SecureHash.h"
 
 CORE_API DECLARE_LOG_CATEGORY_EXTERN(LogIoStore, Log, All);
 
diff --git a/Engine/Source/Runtime/Core/Private/IO/IoStore.cpp b/Engine/Source/Runtime/Core/Private/IO/IoStore.cpp
--- a/Engine/Source/Runtime/Core/Private/IO/IoStore.cpp
+++ b/Engine/Source/Runtime/Core/Private/IO/IoStore.cpp
@@ -25,6 +25,7 @@
 #include "Serialization/MemoryWriter.h"
 #include "ProfilingDebugging/CountersTrace.h"
 #include "Tasks/Task.h"
+#include "IO/IoDirectoryIndex.h"
 
 DEFINE_LOG_CATEGORY(LogIoStore);
 
diff --git a/Engine/Source/Runtime/Core/Private/Internationalization/TextHistory.h b/Engine/Source/Runtime/Core/Private/Internationalization/TextHistory.h
--- a/Engine/Source/Runtime/Core/Private/Internationalization/TextHistory.h
+++ b/Engine/Source/Runtime/Core/Private/Internationalization/TextHistory.h
@@ -1,19 +1,13 @@
 // Copyright Epic Games, Inc. All Rights Reserved.
 #pragma once
 
-#include "Containers/Array.h"
-#include "Containers/UnrealString.h"
 #include "CoreMinimal.h"
 #include "CoreTypes.h"
 #include "Internationalization/ITextData.h"
 #include "Internationalization/LocalizedTextSourceTypes.h"
 #include "Internationalization/StringTableCoreFwd.h"
-#include "Internationalization/Text.h"
 #include "Internationalization/TextKey.h"
-#include "Misc/DateTime.h"
-#include "Misc/Optional.h"
 #include "Serialization/StructuredArchive.h"
-#include "Templates/SharedPointer.h"
 #include "Templates/SharedPointerInternals.h"
 
 class ITextGenerator;
diff --git a/Engine/Source/Runtime/Core/Private/Misc/RemoteConfigIni.cpp b/Engine/Source/Runtime/Core/Private/Misc/RemoteConfigIni.cpp
--- a/Engine/Source/Runtime/Core/Private/Misc/RemoteConfigIni.cpp
+++ b/Engine/Source/Runtime/Core/Private/Misc/RemoteConfigIni.cpp
@@ -9,7 +9,6 @@
 #include "HAL/PlatformTime.h"
 #include "Misc/App.h"
 #include "Misc/ConfigCacheIni.h"
-#include "Misc/ConfigContext.h"
 #include "Misc/FileHelper.h"
 #include "Misc/Paths.h"
 #include "Misc/ScopeLock.h"
diff --git a/Engine/Source/Runtime/Core/Public/Logging/LogMacros.h b/Engine/Source/Runtime/Core/Public/Logging/LogMacros.h
--- a/Engine/Source/Runtime/Core/Public/Logging/LogMacros.h
+++ b/Engine/Source/Runtime/Core/Public/Logging/LogMacros.h
@@ -2,16 +2,18 @@
 
 #pragma once
 
-#include "CoreTypes.h"
-#include "Misc/VarArgs.h"
-#include "Misc/AssertionMacros.h"
 #include "Containers/UnrealString.h"
+#include "CoreTypes.h"
+#include "HAL/PreprocessorHelpers.h"
 #include "Logging/LogCategory.h"
 #include "Logging/LogScopedCategoryAndVerbosityOverride.h"
 #include "Logging/LogTrace.h"
-#include "Templates/IsValidVariadicFunctionArg.h"
+#include "Misc/AssertionMacros.h"
+#include "Misc/Build.h"
+#include "Misc/VarArgs.h"
 #include "Templates/AndOrNot.h"
 #include "Templates/IsArrayOrRefOfTypeByPredicate.h"
+#include "Templates/IsValidVariadicFunctionArg.h"
 #include "Traits/IsCharEncodingCompatibleWith.h"
 
 
diff --git a/Engine/Source/Runtime/Core/Public/Tasks/TaskPrivate.h b/Engine/Source/Runtime/Core/Public/Tasks/TaskPrivate.h
--- a/Engine/Source/Runtime/Core/Public/Tasks/TaskPrivate.h
+++ b/Engine/Source/Runtime/Core/Public/Tasks/TaskPrivate.h
@@ -2,20 +2,25 @@
 
 #pragma once
 
-#include "Async/Fundamental/Task.h"
 #include "Async/Fundamental/Scheduler.h"
+#include "Async/Fundamental/Task.h"
 #include "Async/TaskTrace.h"
-#include "Templates/Invoke.h"
-#include "Templates/RefCounting.h"
-#include "Templates/TypeCompatibleBytes.h"
-#include "Misc/Timeout.h"
-#include "Containers/LockFreeList.h"
 #include "Containers/LockFreeFixedSizeAllocator.h"
+#include "Containers/LockFreeList.h"
+#include "CoreTypes.h"
 #include "Experimental/ConcurrentLinearAllocator.h"
-#include "Templates/UnrealTemplate.h"
+#include "HAL/Event.h"
+#include "HAL/PlatformProcess.h"
+#include "HAL/Thread.h"
 #include "Math/NumericLimits.h"
 #include "Misc/ScopeLock.h"
-#include "CoreTypes.h"
+#include "Misc/Timeout.h"
+#include "ProfilingDebugging/CpuProfilerTrace.h"
+#include "Templates/Invoke.h"
+#include "Templates/RefCounting.h"
+#include "Templates/TypeCompatibleBytes.h"
+#include "Templates/UnrealTemplate.h"
+#include "Trace/Detail/Channel.h"
 
 #include <atomic>
 #include <type_traits>
