Switched to using a specific date format when getting driver date information to ensure compatibility with driver deny list expectations and resolve crashes resulting from some system time formats.

diff --git a/Engine/Source/Runtime/Core/Private/Windows/WindowsPlatformMisc.cpp b/Engine/Source/Runtime/Core/Private/Windows/WindowsPlatformMisc.cpp
--- a/Engine/Source/Runtime/Core/Private/Windows/WindowsPlatformMisc.cpp
+++ b/Engine/Source/Runtime/Core/Private/Windows/WindowsPlatformMisc.cpp
@@ -2531,9 +2531,7 @@ static void GetVideoDriverDetailsFromSetup(const FString& DeviceName, FGPUDriver
 			{
 				SYSTEMTIME SystemTime;
 				FileTimeToSystemTime(&FileTime, &SystemTime);
-				GetDateFormat(LOCALE_USER_DEFAULT, DATE_SHORTDATE, &SystemTime, NULL, Buffer, UE_ARRAY_COUNT(Buffer));
-				Out.DriverDate = Buffer;
-				ZeroMemory(Buffer, sizeof(Buffer));
+				Out.DriverDate = FString::Printf(TEXT("%d-%d-%d"), SystemTime.wMonth, SystemTime.wDay, SystemTime.wYear);
 			}
 			else
 			{
