Fixed compile issue.

diff --git a/Engine/Source/Runtime/Core/Public/Unix/UnixPlatformFile.h b/Engine/Source/Runtime/Core/Public/Unix/UnixPlatformFile.h
--- a/Engine/Source/Runtime/Core/Public/Unix/UnixPlatformFile.h
+++ b/Engine/Source/Runtime/Core/Public/Unix/UnixPlatformFile.h
@@ -9,7 +9,6 @@
 #include "Containers/UnrealString.h"
 #include "CoreTypes.h"
 #include "GenericPlatform/GenericPlatformFile.h"
-#include "HAL/PlatformFile.h"
 #include "Misc/DateTime.h"
 
 class IMappedFileHandle;
