Fixed some issues found when running includetool.

diff --git a/Engine/Source/Runtime/Core/Public/Linux/LinuxPlatformProcess.h b/Engine/Source/Runtime/Core/Public/Linux/LinuxPlatformProcess.h
--- a/Engine/Source/Runtime/Core/Public/Linux/LinuxPlatformProcess.h
+++ b/Engine/Source/Runtime/Core/Public/Linux/LinuxPlatformProcess.h
@@ -7,7 +7,6 @@
 #pragma once
 
 #include "HAL/Platform.h"
-#include "HAL/PlatformProcess.h"
 #include "Unix/UnixPlatformProcess.h"
 
 /**
diff --git a/Engine/Source/Runtime/Core/Public/Unix/UnixPlatformOutputDevices.h b/Engine/Source/Runtime/Core/Public/Unix/UnixPlatformOutputDevices.h
--- a/Engine/Source/Runtime/Core/Public/Unix/UnixPlatformOutputDevices.h
+++ b/Engine/Source/Runtime/Core/Public/Unix/UnixPlatformOutputDevices.h
@@ -9,7 +9,6 @@
 #include "Containers/UnrealString.h"
 #include "CoreTypes.h"
 #include "GenericPlatform/GenericPlatformOutputDevices.h"
-#include "HAL/PlatformOutputDevices.h"
 
 class FOutputDevice;
 class FOutputDeviceError;
diff --git a/Engine/Source/Runtime/Core/Public/Unix/UnixPlatformProcess.h b/Engine/Source/Runtime/Core/Public/Unix/UnixPlatformProcess.h
--- a/Engine/Source/Runtime/Core/Public/Unix/UnixPlatformProcess.h
+++ b/Engine/Source/Runtime/Core/Public/Unix/UnixPlatformProcess.h
@@ -10,7 +10,6 @@
 #include "CoreMinimal.h"
 #include "CoreTypes.h"
 #include "GenericPlatform/GenericPlatformProcess.h"
-#include "HAL/PlatformProcess.h"
 #include "Misc/AssertionMacros.h"
 
 class Error;
diff --git a/Engine/Source/Runtime/Core/Public/Unix/UnixPlatformStackWalk.h b/Engine/Source/Runtime/Core/Public/Unix/UnixPlatformStackWalk.h
--- a/Engine/Source/Runtime/Core/Public/Unix/UnixPlatformStackWalk.h
+++ b/Engine/Source/Runtime/Core/Public/Unix/UnixPlatformStackWalk.h
@@ -8,7 +8,6 @@
 
 #include "CoreTypes.h"
 #include "GenericPlatform/GenericPlatformStackWalk.h"
-#include "HAL/PlatformStackWalk.h"
 
 struct FGenericCrashContext;
 
diff --git a/Engine/Source/Runtime/Core/Public/Unix/UnixPlatformTime.h b/Engine/Source/Runtime/Core/Public/Unix/UnixPlatformTime.h
--- a/Engine/Source/Runtime/Core/Public/Unix/UnixPlatformTime.h
+++ b/Engine/Source/Runtime/Core/Public/Unix/UnixPlatformTime.h
@@ -8,7 +8,6 @@
 
 #include "CoreTypes.h"
 #include "GenericPlatform/GenericPlatformTime.h"
-#include "HAL/PlatformTime.h"
 
 /**
  * Unix implementation of the Time OS functions
