Updated Core code to not use CoreMinimal.h when possible.
Core public headers still will contain CoreMinimal.h if any included it before my IWYU changes.

diff --git a/Engine/Source/Runtime/Core/Internal/IO/IoDirectoryIndex.h b/Engine/Source/Runtime/Core/Internal/IO/IoDirectoryIndex.h
--- a/Engine/Source/Runtime/Core/Internal/IO/IoDirectoryIndex.h
+++ b/Engine/Source/Runtime/Core/Internal/IO/IoDirectoryIndex.h
@@ -2,10 +2,15 @@
 
 #pragma once
 
+#include "Containers/Array.h"
+#include "Containers/Map.h"
 #include "Containers/StringFwd.h"
-#include "CoreMinimal.h"
+#include "Containers/UnrealString.h"
+#include "HAL/Platform.h"
 #include "Misc/AES.h"
 
+class FArchive;
+
 namespace IoDirectoryIndexUtils
 {
 	FString GetCommonRootPath(const TArray<FString>& Filenames);
diff --git a/Engine/Source/Runtime/Core/Private/IO/IoBuffer.cpp b/Engine/Source/Runtime/Core/Private/IO/IoBuffer.cpp
--- a/Engine/Source/Runtime/Core/Private/IO/IoBuffer.cpp
+++ b/Engine/Source/Runtime/Core/Private/IO/IoBuffer.cpp
@@ -1,9 +1,11 @@
 // Copyright Epic Games, Inc. All Rights Reserved.
 
-#include "CoreMinimal.h"
+#include "HAL/Platform.h"
+#include "HAL/UnrealMemory.h"
 #include "IO/IoDispatcher.h"
 #include "Memory/MemoryFwd.h"
 #include "Memory/MemoryView.h"
+#include "Misc/AssertionMacros.h"
 #include "Templates/RefCounting.h"
 
 //////////////////////////////////////////////////////////////////////////
diff --git a/Engine/Source/Runtime/Core/Private/IO/IoContainerId.cpp b/Engine/Source/Runtime/Core/Private/IO/IoContainerId.cpp
--- a/Engine/Source/Runtime/Core/Private/IO/IoContainerId.cpp
+++ b/Engine/Source/Runtime/Core/Private/IO/IoContainerId.cpp
@@ -3,6 +3,9 @@
 #include "IO/IoContainerId.h"
 
 #include "Hash/CityHash.h"
+#include "Misc/AssertionMacros.h"
+#include "Misc/Char.h"
+#include "Serialization/Archive.h"
 #include "Serialization/StructuredArchiveSlots.h"
 
 FIoContainerId FIoContainerId::FromName(const FName& Name)
diff --git a/Engine/Source/Runtime/Core/Private/IO/IoDirectoryIndex.cpp b/Engine/Source/Runtime/Core/Private/IO/IoDirectoryIndex.cpp
--- a/Engine/Source/Runtime/Core/Private/IO/IoDirectoryIndex.cpp
+++ b/Engine/Source/Runtime/Core/Private/IO/IoDirectoryIndex.cpp
@@ -4,10 +4,17 @@
 
 #include "Containers/ArrayView.h"
 #include "Containers/StringView.h"
+#include "CoreTypes.h"
 #include "IO/IoDispatcher.h"
+#include "Misc/AssertionMacros.h"
+#include "Misc/CString.h"
 #include "Misc/Paths.h"
+#include "Serialization/Archive.h"
 #include "Serialization/MemoryReader.h"
 #include "Serialization/MemoryWriter.h"
+#include "Templates/AlignmentTemplates.h"
+#include "Templates/Function.h"
+#include "Templates/UnrealTemplate.h"
 
 namespace IoDirectoryIndexUtils
 {
diff --git a/Engine/Source/Runtime/Core/Private/IO/IoStatus.cpp b/Engine/Source/Runtime/Core/Private/IO/IoStatus.cpp
--- a/Engine/Source/Runtime/Core/Private/IO/IoStatus.cpp
+++ b/Engine/Source/Runtime/Core/Private/IO/IoStatus.cpp
@@ -2,8 +2,14 @@
 
 #include "Containers/StringFwd.h"
 #include "Containers/StringView.h"
-#include "CoreMinimal.h"
+#include "Containers/UnrealString.h"
+#include "HAL/Platform.h"
+#include "HAL/PlatformString.h"
+#include "HAL/UnrealMemory.h"
 #include "IO/IoDispatcher.h"
+#include "Logging/LogCategory.h"
+#include "Logging/LogMacros.h"
+#include "Math/UnrealMathUtility.h"
 #include "Trace/Detail/Channel.h"
 
 //////////////////////////////////////////////////////////////////////////
diff --git a/Engine/Source/Runtime/Core/Private/IO/PackageId.cpp b/Engine/Source/Runtime/Core/Private/IO/PackageId.cpp
--- a/Engine/Source/Runtime/Core/Private/IO/PackageId.cpp
+++ b/Engine/Source/Runtime/Core/Private/IO/PackageId.cpp
@@ -3,6 +3,7 @@
 #include "IO/PackageId.h"
 
 #include "Hash/CityHash.h"
+#include "Misc/Char.h"
 #include "Serialization/StructuredArchiveAdapters.h"
 #include "Serialization/StructuredArchiveSlots.h"
 
diff --git a/Engine/Source/Runtime/Core/Private/IO/PackageStore.cpp b/Engine/Source/Runtime/Core/Private/IO/PackageStore.cpp
--- a/Engine/Source/Runtime/Core/Private/IO/PackageStore.cpp
+++ b/Engine/Source/Runtime/Core/Private/IO/PackageStore.cpp
@@ -2,8 +2,14 @@
 
 #include "IO/PackageStore.h"
 
+#include "Containers/StringConv.h"
 #include "Containers/StringView.h"
+#include "Containers/UnrealString.h"
+#include "CoreGlobals.h"
+#include "HAL/Platform.h"
 #include "IO/PackageId.h"
+#include "Misc/AssertionMacros.h"
+#include "Serialization/Archive.h"
 #include "Serialization/CompactBinary.h"
 #include "Serialization/CompactBinaryWriter.h"
 
diff --git a/Engine/Source/Runtime/Core/Private/Internationalization/TextCache.cpp b/Engine/Source/Runtime/Core/Private/Internationalization/TextCache.cpp
--- a/Engine/Source/Runtime/Core/Private/Internationalization/TextCache.cpp
+++ b/Engine/Source/Runtime/Core/Private/Internationalization/TextCache.cpp
@@ -3,7 +3,6 @@
 #include "Internationalization/TextCache.h"
 
 #include "Containers/UnrealString.h"
-#include "CoreTypes.h"
 #include "HAL/LowLevelMemTracker.h"
 #include "Internationalization/TextKey.h"
 #include "Misc/CString.h"
diff --git a/Engine/Source/Runtime/Core/Private/Internationalization/TextHistory.h b/Engine/Source/Runtime/Core/Private/Internationalization/TextHistory.h
--- a/Engine/Source/Runtime/Core/Private/Internationalization/TextHistory.h
+++ b/Engine/Source/Runtime/Core/Private/Internationalization/TextHistory.h
@@ -1,14 +1,17 @@
 // Copyright Epic Games, Inc. All Rights Reserved.
 #pragma once
 
-#include "CoreMinimal.h"
 #include "CoreTypes.h"
+#include "HAL/CriticalSection.h"
+#include "Internationalization/CulturePointer.h"
 #include "Internationalization/ITextData.h"
 #include "Internationalization/LocalizedTextSourceTypes.h"
 #include "Internationalization/StringTableCoreFwd.h"
 #include "Internationalization/TextKey.h"
 #include "Serialization/StructuredArchive.h"
 #include "Templates/SharedPointerInternals.h"
+#include "Templates/UnrealTemplate.h"
+#include "UObject/NameTypes.h"
 
 class ITextGenerator;
 struct FDecimalNumberFormattingRules;
diff --git a/Engine/Source/Runtime/Core/Private/Logging/LogCategory.cpp b/Engine/Source/Runtime/Core/Private/Logging/LogCategory.cpp
--- a/Engine/Source/Runtime/Core/Private/Logging/LogCategory.cpp
+++ b/Engine/Source/Runtime/Core/Private/Logging/LogCategory.cpp
@@ -2,9 +2,13 @@
 
 #include "Logging/LogCategory.h"
 
-#include "CoreMinimal.h"
+#include "Containers/UnrealString.h"
+#include "CoreGlobals.h"
+#include "HAL/PlatformMisc.h"
 #include "Logging/LogSuppressionInterface.h"
 #include "Logging/LogTrace.h"
+#include "Math/UnrealMathUtility.h"
+#include "Misc/AssertionMacros.h"
 #include "Misc/CoreDelegates.h"
 #include "Misc/OutputDeviceRedirector.h"
 
diff --git a/Engine/Source/Runtime/Core/Private/Memory/VirtualStackAllocator.cpp b/Engine/Source/Runtime/Core/Private/Memory/VirtualStackAllocator.cpp
--- a/Engine/Source/Runtime/Core/Private/Memory/VirtualStackAllocator.cpp
+++ b/Engine/Source/Runtime/Core/Private/Memory/VirtualStackAllocator.cpp
@@ -4,6 +4,7 @@
 
 #include "HAL/LowLevelMemTracker.h"
 #include "HAL/PlatformMemory.h"
+#include "Templates/AlignmentTemplates.h"
 
 LLM_DEFINE_TAG(VirtualStackAllocator);
 
diff --git a/Engine/Source/Runtime/Core/Private/Misc/HierarchicalLogArchive.cpp b/Engine/Source/Runtime/Core/Private/Misc/HierarchicalLogArchive.cpp
--- a/Engine/Source/Runtime/Core/Private/Misc/HierarchicalLogArchive.cpp
+++ b/Engine/Source/Runtime/Core/Private/Misc/HierarchicalLogArchive.cpp
@@ -2,9 +2,12 @@
 
 #include "Misc/HierarchicalLogArchive.h"
 
+#include "Containers/StringConv.h"
 #include "Containers/StringFwd.h"
 #include "Misc/StringBuilder.h"
 
+class FArchive;
+
 FHierarchicalLogArchive::FHierarchicalLogArchive(FArchive& InInnerArchive)
 	: FArchiveProxy(InInnerArchive)
 	, Indentation(0)
diff --git a/Engine/Source/Runtime/Core/Private/Misc/LargeWorldRenderPosition.cpp b/Engine/Source/Runtime/Core/Private/Misc/LargeWorldRenderPosition.cpp
--- a/Engine/Source/Runtime/Core/Private/Misc/LargeWorldRenderPosition.cpp
+++ b/Engine/Source/Runtime/Core/Private/Misc/LargeWorldRenderPosition.cpp
@@ -2,7 +2,10 @@
 
 #include "Misc/LargeWorldRenderPosition.h"
 
+#include "HAL/Platform.h"
 #include "Math/Matrix.inl"
+#include "Math/TranslationMatrix.h"
+#include "Misc/AssertionMacros.h"
 
 #define UE_LWC_RENDER_TILE_SIZE_MIN (262144.0)	// LWC_TODO: Make this smaller?
 
diff --git a/Engine/Source/Runtime/Core/Private/Misc/MemoryReadStream.cpp b/Engine/Source/Runtime/Core/Private/Misc/MemoryReadStream.cpp
--- a/Engine/Source/Runtime/Core/Private/Misc/MemoryReadStream.cpp
+++ b/Engine/Source/Runtime/Core/Private/Misc/MemoryReadStream.cpp
@@ -2,6 +2,9 @@
 
 #include "Misc/MemoryReadStream.h"
 
+#include "HAL/UnrealMemory.h"
+#include "Math/UnrealMathUtility.h"
+
 class FMemoryReadStreamBuffer : public IMemoryReadStream
 {
 public:
diff --git a/Engine/Source/Runtime/Core/Private/Misc/MessageDialog.cpp b/Engine/Source/Runtime/Core/Private/Misc/MessageDialog.cpp
--- a/Engine/Source/Runtime/Core/Private/Misc/MessageDialog.cpp
+++ b/Engine/Source/Runtime/Core/Private/Misc/MessageDialog.cpp
@@ -2,13 +2,16 @@
 
 #include "Misc/MessageDialog.h"
 
+#include "Containers/UnrealString.h"
 #include "CoreGlobals.h"
-#include "CoreMinimal.h"
 #include "CoreTypes.h"
 #include "Internationalization/Internationalization.h"
 #include "Internationalization/Text.h"
+#include "Logging/LogCategory.h"
 #include "Logging/LogMacros.h"
+#include "Logging/LogVerbosity.h"
 #include "Misc/App.h"
+#include "Misc/AssertionMacros.h"
 #include "Misc/CString.h"
 #include "Misc/CoreDelegates.h"
 #include "Misc/FeedbackContext.h"
diff --git a/Engine/Source/Runtime/Core/Private/Misc/NetworkVersion.cpp b/Engine/Source/Runtime/Core/Private/Misc/NetworkVersion.cpp
--- a/Engine/Source/Runtime/Core/Private/Misc/NetworkVersion.cpp
+++ b/Engine/Source/Runtime/Core/Private/Misc/NetworkVersion.cpp
@@ -3,12 +3,16 @@
 #include "Misc/NetworkVersion.h"
 
 #include "BuildSettings.h"
-#include "CoreMinimal.h"
 #include "HAL/IConsoleManager.h"
+#include "Logging/LogCategory.h"
 #include "Misc/App.h"
+#include "Misc/AssertionMacros.h"
+#include "Misc/CString.h"
 #include "Misc/CommandLine.h"
+#include "Misc/Crc.h"
 #include "Misc/EngineVersion.h"
 #include "Misc/NetworkGuid.h" // IWYU pragma: keep
+#include "Misc/Parse.h"
 #include "Trace/Detail/Channel.h"
 
 DEFINE_LOG_CATEGORY( LogNetVersion );
diff --git a/Engine/Source/Runtime/Core/Private/Misc/QueuedThreadPoolWrapper.cpp b/Engine/Source/Runtime/Core/Private/Misc/QueuedThreadPoolWrapper.cpp
--- a/Engine/Source/Runtime/Core/Private/Misc/QueuedThreadPoolWrapper.cpp
+++ b/Engine/Source/Runtime/Core/Private/Misc/QueuedThreadPoolWrapper.cpp
@@ -7,6 +7,8 @@
 #include "Misc/QueuedThreadPool.h"
 #include "ProfilingDebugging/CountersTrace.h"
 #include "ProfilingDebugging/CpuProfilerTrace.h"
+#include "Templates/Tuple.h"
+#include "Templates/UnrealTemplate.h"
 
 TRACE_DECLARE_INT_COUNTER(QueuedThreadPoolWrapperScheduledWorkAllocs, TEXT("QueuedThreadPoolWrapper/ScheduledWorkAllocs"));
 
diff --git a/Engine/Source/Runtime/Core/Private/Misc/RemoteConfigIni.cpp b/Engine/Source/Runtime/Core/Private/Misc/RemoteConfigIni.cpp
--- a/Engine/Source/Runtime/Core/Private/Misc/RemoteConfigIni.cpp
+++ b/Engine/Source/Runtime/Core/Private/Misc/RemoteConfigIni.cpp
@@ -3,16 +3,21 @@
 #include "Misc/RemoteConfigIni.h"
 
 #include "Async/AsyncWork.h"
-#include "CoreMinimal.h"
+#include "CoreGlobals.h"
 #include "HAL/FileManager.h"
+#include "HAL/PlatformCrt.h"
 #include "HAL/PlatformProcess.h"
 #include "HAL/PlatformTime.h"
+#include "HAL/UnrealMemory.h"
 #include "Misc/App.h"
+#include "Misc/AssertionMacros.h"
+#include "Misc/CString.h"
 #include "Misc/ConfigCacheIni.h"
 #include "Misc/FileHelper.h"
 #include "Misc/Paths.h"
 #include "Misc/ScopeLock.h"
 #include "Stats/Stats.h"
+#include "Templates/TypeHash.h"
 
 // Globals
 FRemoteConfig GRemoteConfig;
diff --git a/Engine/Source/Runtime/Core/Private/Serialization/Csv/CsvParserTests.cpp b/Engine/Source/Runtime/Core/Private/Serialization/Csv/CsvParserTests.cpp
--- a/Engine/Source/Runtime/Core/Private/Serialization/Csv/CsvParserTests.cpp
+++ b/Engine/Source/Runtime/Core/Private/Serialization/Csv/CsvParserTests.cpp
@@ -1,9 +1,12 @@
 // Copyright Epic Games, Inc. All Rights Reserved.
 
-#include "CoreMinimal.h"
+#include "Containers/Array.h"
+#include "Containers/UnrealString.h"
 #include "CoreTypes.h"
 #include "Misc/AutomationTest.h"
+#include "Misc/CString.h"
 #include "Serialization/Csv/CsvParser.h"
+#include "Templates/UnrealTemplate.h"
 
 namespace CsvParser_Tests
 {
diff --git a/Engine/Source/Runtime/Core/Private/Serialization/MappedName.cpp b/Engine/Source/Runtime/Core/Private/Serialization/MappedName.cpp
--- a/Engine/Source/Runtime/Core/Private/Serialization/MappedName.cpp
+++ b/Engine/Source/Runtime/Core/Private/Serialization/MappedName.cpp
@@ -2,6 +2,7 @@
 
 #include "Serialization/MappedName.h"
 
+#include "Serialization/Archive.h"
 #include "UObject/NameBatchSerialization.h"
 
 FArchive& operator<<(FArchive& Ar, FMappedName& MappedName)
diff --git a/Engine/Source/Runtime/Core/Private/Tasks/Pipe.cpp b/Engine/Source/Runtime/Core/Private/Tasks/Pipe.cpp
--- a/Engine/Source/Runtime/Core/Private/Tasks/Pipe.cpp
+++ b/Engine/Source/Runtime/Core/Private/Tasks/Pipe.cpp
@@ -2,6 +2,7 @@
 
 #include "Tasks/Pipe.h"
 
+#include "Containers/Array.h"
 #include "Misc/ScopeExit.h"
 
 namespace UE::Tasks
diff --git a/Engine/Source/Runtime/Core/Private/Unix/UnixPlatformOutputDevices.cpp b/Engine/Source/Runtime/Core/Private/Unix/UnixPlatformOutputDevices.cpp
--- a/Engine/Source/Runtime/Core/Private/Unix/UnixPlatformOutputDevices.cpp
+++ b/Engine/Source/Runtime/Core/Private/Unix/UnixPlatformOutputDevices.cpp
@@ -2,12 +2,14 @@
 
 #include "Unix/UnixPlatformOutputDevices.h"
 
-#include "CoreMinimal.h"
+#include "CoreGlobals.h"
 #include "CoreTypes.h"
+#include "Misc/AssertionMacros.h"
 #include "Misc/Build.h"
 #include "Misc/CommandLine.h"
 #include "Misc/OutputDeviceConsole.h"
 #include "Misc/OutputDeviceRedirector.h"
+#include "Misc/Parse.h"
 #include "Misc/Paths.h"
 #include "Unix/UnixErrorOutputDevice.h"
 
diff --git a/Engine/Source/Runtime/Core/Public/Algo/TopologicalSort.h b/Engine/Source/Runtime/Core/Public/Algo/TopologicalSort.h
--- a/Engine/Source/Runtime/Core/Public/Algo/TopologicalSort.h
+++ b/Engine/Source/Runtime/Core/Public/Algo/TopologicalSort.h
@@ -3,7 +3,8 @@
 #pragma once
 
 #include "Algo/KahnTopologicalSort.h"
-#include "CoreMinimal.h"
+#include "HAL/Platform.h"
+#include "Templates/UnrealTemplate.h"
 
 namespace Algo
 {
diff --git a/Engine/Source/Runtime/Core/Public/IO/IoContainerHeader.h b/Engine/Source/Runtime/Core/Public/IO/IoContainerHeader.h
--- a/Engine/Source/Runtime/Core/Public/IO/IoContainerHeader.h
+++ b/Engine/Source/Runtime/Core/Public/IO/IoContainerHeader.h
@@ -2,7 +2,8 @@
 
 #pragma once
 
-#include "CoreMinimal.h"
+#include "Containers/Array.h"
+#include "HAL/Platform.h"
 #include "IO/IoContainerId.h"
 #include "IO/PackageId.h"
 #include "Serialization/MappedName.h"
diff --git a/Engine/Source/Runtime/Core/Public/IO/IoContainerId.h b/Engine/Source/Runtime/Core/Public/IO/IoContainerId.h
--- a/Engine/Source/Runtime/Core/Public/IO/IoContainerId.h
+++ b/Engine/Source/Runtime/Core/Public/IO/IoContainerId.h
@@ -3,6 +3,8 @@
 #pragma once
 
 #include "CoreMinimal.h"
+#include "HAL/Platform.h"
+#include "UObject/NameTypes.h"
 
 class FArchive;
 class FStructuredArchiveSlot;
diff --git a/Engine/Source/Runtime/Core/Public/IO/IoDispatcher.h b/Engine/Source/Runtime/Core/Public/IO/IoDispatcher.h
--- a/Engine/Source/Runtime/Core/Public/IO/IoDispatcher.h
+++ b/Engine/Source/Runtime/Core/Public/IO/IoDispatcher.h
@@ -4,32 +4,47 @@
 
 #include "Async/Future.h"
 #include "Async/TaskGraphInterfaces.h"
+#include "Containers/Array.h"
 #include "Containers/ArrayView.h"
+#include "Containers/Map.h"
 #include "Containers/StringFwd.h"
 #include "Containers/StringView.h"
 #include "Containers/UnrealString.h"
-#include "CoreMinimal.h"
 #include "CoreTypes.h"
+#include "Delegates/Delegate.h"
+#include "Delegates/DelegateBase.h"
+#include "Delegates/IDelegateInstance.h"
 #include "GenericPlatform/GenericPlatformFile.h"
+#include "HAL/CriticalSection.h"
 #include "HAL/PlatformAtomics.h"
 #include "HAL/PlatformFile.h"
+#include "HAL/UnrealMemory.h"
 #include "IO/IoContainerId.h"
 #include "IO/IoHash.h"
 #include "Logging/LogMacros.h"
+#include "Math/NumericLimits.h"
 #include "Memory/MemoryFwd.h"
 #include "Memory/MemoryView.h"
 #include "Misc/AES.h"
+#include "Misc/AssertionMacros.h"
 #include "Misc/Build.h"
+#include "Misc/ByteSwap.h"
 #include "Misc/EnumClassFlags.h"
 #include "Misc/Guid.h"
 #include "Misc/IEngineCrypto.h"
 #include "Misc/SecureHash.h"
+#include "Serialization/Archive.h"
 #include "Serialization/FileRegions.h"
 #include "String/BytesToHex.h"
 #include "Tasks/Task.h"
+#include "Templates/Function.h"
 #include "Templates/RefCounting.h"
+#include "Templates/SharedPointer.h"
 #include "Templates/TypeCompatibleBytes.h"
+#include "Templates/UniquePtr.h"
 #include "Templates/UnrealTemplate.h"
+#include "UObject/NameTypes.h"
+#include "UObject/UnrealNames.h"
 
 class FEvent;
 class FIoBatchImpl;
diff --git a/Engine/Source/Runtime/Core/Public/IO/PackageId.h b/Engine/Source/Runtime/Core/Public/IO/PackageId.h
--- a/Engine/Source/Runtime/Core/Public/IO/PackageId.h
+++ b/Engine/Source/Runtime/Core/Public/IO/PackageId.h
@@ -3,6 +3,9 @@
 #pragma once
 
 #include "CoreMinimal.h"
+#include "HAL/Platform.h"
+#include "Misc/AssertionMacros.h"
+#include "UObject/NameTypes.h"
 
 class FArchive;
 class FStructuredArchiveSlot;
diff --git a/Engine/Source/Runtime/Core/Public/IO/PackageStore.h b/Engine/Source/Runtime/Core/Public/IO/PackageStore.h
--- a/Engine/Source/Runtime/Core/Public/IO/PackageStore.h
+++ b/Engine/Source/Runtime/Core/Public/IO/PackageStore.h
@@ -2,11 +2,17 @@
 
 #pragma once
 
+#include "Containers/Array.h"
 #include "Containers/ArrayView.h"
 #include "CoreMinimal.h"
+#include "CoreTypes.h"
+#include "Delegates/Delegate.h"
+#include "Delegates/DelegateBase.h"
 #include "Misc/EnumClassFlags.h"
 #include "Misc/SecureHash.h"
 #include "PackageId.h"
+#include "Templates/SharedPointer.h"
+#include "UObject/NameTypes.h"
 
 class FArchive;
 class FCbObject;
diff --git a/Engine/Source/Runtime/Core/Public/Internationalization/TextLocalizationManager.h b/Engine/Source/Runtime/Core/Public/Internationalization/TextLocalizationManager.h
--- a/Engine/Source/Runtime/Core/Public/Internationalization/TextLocalizationManager.h
+++ b/Engine/Source/Runtime/Core/Public/Internationalization/TextLocalizationManager.h
@@ -10,6 +10,7 @@
 #include "Containers/UnrealString.h"
 #include "CoreTypes.h"
 #include "Delegates/Delegate.h"
+#include "Delegates/DelegateBase.h"
 #include "HAL/CriticalSection.h"
 #include "Internationalization/LocKeyFuncs.h"
 #include "Internationalization/LocTesting.h"
diff --git a/Engine/Source/Runtime/Core/Public/Logging/LogMacros.h b/Engine/Source/Runtime/Core/Public/Logging/LogMacros.h
--- a/Engine/Source/Runtime/Core/Public/Logging/LogMacros.h
+++ b/Engine/Source/Runtime/Core/Public/Logging/LogMacros.h
@@ -8,10 +8,12 @@
 #include "Logging/LogCategory.h"
 #include "Logging/LogScopedCategoryAndVerbosityOverride.h"
 #include "Logging/LogTrace.h"
+#include "Logging/LogVerbosity.h"
 #include "Misc/AssertionMacros.h"
 #include "Misc/Build.h"
 #include "Misc/VarArgs.h"
 #include "Templates/AndOrNot.h"
+#include "Templates/EnableIf.h"
 #include "Templates/IsArrayOrRefOfTypeByPredicate.h"
 #include "Templates/IsValidVariadicFunctionArg.h"
 #include "Traits/IsCharEncodingCompatibleWith.h"
diff --git a/Engine/Source/Runtime/Core/Public/Math/GenericOctree.h b/Engine/Source/Runtime/Core/Public/Math/GenericOctree.h
--- a/Engine/Source/Runtime/Core/Public/Math/GenericOctree.h
+++ b/Engine/Source/Runtime/Core/Public/Math/GenericOctree.h
@@ -6,12 +6,29 @@
 
 #pragma once
 
+#include "Containers/Array.h"
 #include "Containers/ArrayView.h"
+#include "Containers/ContainerAllocationPolicies.h"
+#include "CoreGlobals.h"
 #include "CoreMinimal.h"
+#include "CoreTypes.h"
 #include "GenericOctreePublic.h"
+#include "HAL/PlatformMisc.h"
+#include "Logging/LogCategory.h"
+#include "Logging/LogMacros.h"
+#include "Math/Box.h"
+#include "Math/BoxSphereBounds.h"
 #include "Math/MathFwd.h"
 #include "Math/UnrealMathSSE.h"
+#include "Math/UnrealMathUtility.h"
+#include "Math/Vector.h"
+#include "Math/Vector4.h"
+#include "Math/VectorRegister.h"
+#include "Misc/AssertionMacros.h"
+#include "Templates/EnableIf.h"
 #include "Templates/Models.h"
+#include "Templates/UnrealTemplate.h"
+#include "Templates/UnrealTypeTraits.h"
 #include "Trace/Detail/Channel.h"
 
 /** A concise iteration over the children of an octree node. */
@@ -1800,3 +1817,5 @@ public:
 };
 
 #include "GenericOctree.inl"
+
+#include <stddef.h>
diff --git a/Engine/Source/Runtime/Core/Public/Memory/VirtualStackAllocator.h b/Engine/Source/Runtime/Core/Public/Memory/VirtualStackAllocator.h
--- a/Engine/Source/Runtime/Core/Public/Memory/VirtualStackAllocator.h
+++ b/Engine/Source/Runtime/Core/Public/Memory/VirtualStackAllocator.h
@@ -3,7 +3,9 @@
 #pragma once
 
 #include "CoreMinimal.h"
+#include "HAL/Platform.h"
 #include "HAL/PlatformMemory.h"
+#include "Misc/AssertionMacros.h"
 
 class FVirtualStackAllocator;
 
diff --git a/Engine/Source/Runtime/Core/Public/Misc/App.h b/Engine/Source/Runtime/Core/Public/Misc/App.h
--- a/Engine/Source/Runtime/Core/Public/Misc/App.h
+++ b/Engine/Source/Runtime/Core/Public/Misc/App.h
@@ -5,11 +5,12 @@
 #include "Containers/Array.h"
 #include "Containers/UnrealString.h"
 #include "CoreGlobals.h"
-#include "CoreMinimal.h"
 #include "CoreTypes.h"
 #include "Delegates/Delegate.h"
+#include "HAL/PlatformMisc.h"
 #include "HAL/PlatformProcess.h"
 #include "Misc/Build.h"
+#include "Misc/CString.h"
 #include "Misc/CommandLine.h"
 #include "Misc/CoreMisc.h"
 #include "Misc/FrameRate.h"
diff --git a/Engine/Source/Runtime/Core/Public/Misc/ArchiveMD5.h b/Engine/Source/Runtime/Core/Public/Misc/ArchiveMD5.h
--- a/Engine/Source/Runtime/Core/Public/Misc/ArchiveMD5.h
+++ b/Engine/Source/Runtime/Core/Public/Misc/ArchiveMD5.h
@@ -2,8 +2,13 @@
 
 #pragma once
 
+#include "Containers/UnrealString.h"
 #include "CoreMinimal.h"
+#include "HAL/Platform.h"
+#include "Misc/AssertionMacros.h"
 #include "Misc/SecureHash.h"
+#include "Serialization/Archive.h"
+#include "UObject/NameTypes.h"
 
 /**
  * FArchive adapter for FMD5
diff --git a/Engine/Source/Runtime/Core/Public/Misc/AutomationTest.h b/Engine/Source/Runtime/Core/Public/Misc/AutomationTest.h
--- a/Engine/Source/Runtime/Core/Public/Misc/AutomationTest.h
+++ b/Engine/Source/Runtime/Core/Public/Misc/AutomationTest.h
@@ -7,11 +7,13 @@
 #include "Containers/Array.h"
 #include "Containers/Map.h"
 #include "Containers/Queue.h"
+#include "Containers/Set.h"
 #include "Containers/UnrealString.h"
-#include "CoreMinimal.h"
 #include "CoreTypes.h"
 #include "Delegates/Delegate.h"
+#include "Delegates/DelegateBase.h"
 #include "Delegates/DelegateInstancesImpl.h"
+#include "Delegates/IDelegateInstance.h"
 #include "GenericPlatform/GenericPlatformStackWalk.h"
 #include "HAL/CriticalSection.h"
 #include "HAL/LowLevelMemTracker.h"
@@ -20,6 +22,7 @@
 #include "HAL/PlatformTime.h"
 #include "HAL/ThreadSafeBool.h"
 #include "Internationalization/Regex.h"
+#include "Logging/LogVerbosity.h"
 #include "Math/Color.h"
 #include "Math/Color.h"
 #include "Math/MathFwd.h"
@@ -30,11 +33,17 @@
 #include "Misc/AutomationEvent.h"
 #include "Misc/Build.h"
 #include "Misc/CString.h"
+#include "Misc/Char.h"
+#include "Misc/DateTime.h"
 #include "Misc/FeedbackContext.h"
 #include "Misc/Guid.h"
 #include "Misc/Optional.h"
+#include "Misc/OutputDevice.h"
+#include "Misc/Timespan.h"
 #include "Templates/Function.h"
 #include "Templates/SharedPointer.h"
+#include "Templates/UnrealTemplate.h"
+#include "UObject/NameTypes.h"
 
 #include <atomic>
 
diff --git a/Engine/Source/Runtime/Core/Public/Misc/Compression.h b/Engine/Source/Runtime/Core/Public/Misc/Compression.h
--- a/Engine/Source/Runtime/Core/Public/Misc/Compression.h
+++ b/Engine/Source/Runtime/Core/Public/Misc/Compression.h
@@ -2,7 +2,8 @@
 
 #pragma once
 
-#include "CoreMinimal.h"
+#include "Containers/Map.h"
+#include "Containers/UnrealString.h"
 #include "CoreTypes.h"
 #include "HAL/CriticalSection.h"
 #include "Misc/CompressionFlags.h"
diff --git a/Engine/Source/Runtime/Core/Public/Misc/ConfigCacheIni.h b/Engine/Source/Runtime/Core/Public/Misc/ConfigCacheIni.h
--- a/Engine/Source/Runtime/Core/Public/Misc/ConfigCacheIni.h
+++ b/Engine/Source/Runtime/Core/Public/Misc/ConfigCacheIni.h
@@ -6,31 +6,43 @@
 
 #pragma once
 
+#include "Algo/Reverse.h"
 #include "Containers/Array.h"
 #include "Containers/Map.h"
+#include "Containers/Set.h"
+#include "Containers/SparseArray.h"
 #include "Containers/StringFwd.h"
 #include "Containers/UnrealString.h"
-#include "CoreMinimal.h"
+#include "CoreGlobals.h"
 #include "CoreTypes.h"
 #include "Delegates/Delegate.h"
+#include "Delegates/DelegateBase.h"
+#include "Internationalization/Text.h"
 #include "Logging/LogMacros.h"
 #include "Math/Color.h"
 #include "Math/MathFwd.h"
 #include "Math/Rotator.h"
 #include "Math/Vector.h"
 #include "Math/Vector2D.h"
+#include "Math/Vector4.h"
 #include "Misc/AccessDetection.h"
 #include "Misc/Build.h"
 #include "Misc/ConfigTypes.h"
 #include "Misc/Paths.h"
+#include "Serialization/Archive.h"
 #include "Serialization/StructuredArchive.h"
 #include "Serialization/StructuredArchiveAdapters.h"
 #include "Serialization/StructuredArchiveSlots.h"
+#include "Templates/ChooseClass.h"
+#include "Templates/Function.h"
+#include "Templates/UnrealTemplate.h"
 #include "UObject/NameTypes.h"
 
 class FConfigCacheIni;
 class FConfigFile;
+class FOutputDevice;
 class IConsoleVariable;
+struct FColor;
 
 CORE_API DECLARE_LOG_CATEGORY_EXTERN(LogConfig, Log, All);
 
diff --git a/Engine/Source/Runtime/Core/Public/Misc/CoreDelegates.h b/Engine/Source/Runtime/Core/Public/Misc/CoreDelegates.h
--- a/Engine/Source/Runtime/Core/Public/Misc/CoreDelegates.h
+++ b/Engine/Source/Runtime/Core/Public/Misc/CoreDelegates.h
@@ -6,22 +6,32 @@
 #include "Containers/Array.h"
 #include "Containers/ArrayView.h"
 #include "Containers/Map.h"
+#include "Containers/Set.h"
 #include "Containers/UnrealString.h"
-#include "CoreMinimal.h"
 #include "CoreTypes.h"
 #include "Delegates/Delegate.h"
+#include "Delegates/DelegateBase.h"
+#include "Delegates/IDelegateInstance.h"
 #include "GenericPlatform/GenericPlatformFile.h"
 #include "HAL/PlatformFile.h"
+#include "HAL/PlatformMisc.h"
+#include "Logging/LogCategory.h"
+#include "Logging/LogVerbosity.h"
 #include "Math/IntVector.h"
 #include "Math/MathFwd.h"
 #include "Misc/AES.h"
 #include "Misc/Build.h"
+#include "Misc/Optional.h"
+#include "Templates/Function.h"
+#include "Templates/SharedPointer.h"
 #include "Templates/SharedPointerInternals.h"
 
 class AActor;
 class Error;
 class FConfigFile;
+class FName;
 class FSHAHash;
+class FText;
 class IPakFile;
 class ITargetPlatform;
 struct FGuid;
diff --git a/Engine/Source/Runtime/Core/Public/Misc/CoreMisc.h b/Engine/Source/Runtime/Core/Public/Misc/CoreMisc.h
--- a/Engine/Source/Runtime/Core/Public/Misc/CoreMisc.h
+++ b/Engine/Source/Runtime/Core/Public/Misc/CoreMisc.h
@@ -3,18 +3,21 @@
 #pragma once
 
 #include "Containers/Array.h"
+#include "Containers/ContainerAllocationPolicies.h"
 #include "Containers/Map.h"
 #include "Containers/UnrealString.h"
 #include "CoreGlobals.h"
-#include "CoreMinimal.h"
 #include "CoreTypes.h"
+#include "HAL/PlatformProperties.h"
 #include "HAL/ThreadSingleton.h"
+#include "Logging/LogVerbosity.h"
 #include "Math/IntPoint.h"
 #include "Misc/Build.h"
 #include "Misc/Exec.h"
 #include "Templates/Function.h"
 #include "UObject/NameTypes.h"
 
+class FOutputDevice;
 class UWorld;
 
 /**
diff --git a/Engine/Source/Runtime/Core/Public/Misc/FileHelper.h b/Engine/Source/Runtime/Core/Public/Misc/FileHelper.h
--- a/Engine/Source/Runtime/Core/Public/Misc/FileHelper.h
+++ b/Engine/Source/Runtime/Core/Public/Misc/FileHelper.h
@@ -2,18 +2,23 @@
 
 #pragma once
 
+#include "Containers/Array.h"
 #include "Containers/ArrayView.h"
+#include "Containers/ContainersFwd.h"
 #include "Containers/StringFwd.h"
 #include "Containers/StringView.h"
 #include "Containers/UnrealString.h"
-#include "CoreMinimal.h"
 #include "CoreTypes.h"
 #include "HAL/FileManager.h"
 #include "Math/Color.h"
 #include "Math/MathFwd.h"
 #include "Misc/EnumClassFlags.h"
+#include "Templates/UnrealTemplate.h"
 
+class FArchive;
+class FText;
 class IPlatformFile;
+template <typename FuncType> class TFunctionRef;
 
 /*-----------------------------------------------------------------------------
 	FFileHelper
diff --git a/Engine/Source/Runtime/Core/Public/Misc/HierarchicalLogArchive.h b/Engine/Source/Runtime/Core/Public/Misc/HierarchicalLogArchive.h
--- a/Engine/Source/Runtime/Core/Public/Misc/HierarchicalLogArchive.h
+++ b/Engine/Source/Runtime/Core/Public/Misc/HierarchicalLogArchive.h
@@ -2,9 +2,14 @@
 
 #pragma once
 
+#include "Containers/UnrealString.h"
 #include "CoreMinimal.h"
+#include "HAL/Platform.h"
+#include "Misc/AssertionMacros.h"
 #include "Serialization/ArchiveProxy.h"
 
+class FArchive;
+
 struct CORE_API FHierarchicalLogArchive : private FArchiveProxy
 {
 public:
diff --git a/Engine/Source/Runtime/Core/Public/Misc/LargeWorldRenderPosition.h b/Engine/Source/Runtime/Core/Public/Misc/LargeWorldRenderPosition.h
--- a/Engine/Source/Runtime/Core/Public/Misc/LargeWorldRenderPosition.h
+++ b/Engine/Source/Runtime/Core/Public/Misc/LargeWorldRenderPosition.h
@@ -4,6 +4,10 @@
 
 #include "CoreMinimal.h"
 #include "Math/MathFwd.h"
+#include "Math/Matrix.h"
+#include "Math/UnrealMathUtility.h"
+#include "Math/Vector.h"
+#include "Math/Vector4.h"
 #include "Misc/LargeWorldCoordinates.h"
 
 struct FLargeWorldRenderScalar
diff --git a/Engine/Source/Runtime/Core/Public/Misc/MemStack.h b/Engine/Source/Runtime/Core/Public/Misc/MemStack.h
--- a/Engine/Source/Runtime/Core/Public/Misc/MemStack.h
+++ b/Engine/Source/Runtime/Core/Public/Misc/MemStack.h
@@ -4,8 +4,10 @@
 
 #include "Containers/ContainerAllocationPolicies.h"
 #include "Containers/LockFreeFixedSizeAllocator.h"
-#include "CoreMinimal.h"
+#include "CoreGlobals.h"
 #include "CoreTypes.h"
+#include "HAL/MemoryBase.h"
+#include "HAL/PlatformCrt.h"
 #include "HAL/ThreadSafeCounter.h"
 #include "HAL/ThreadSingleton.h"
 #include "HAL/UnrealMemory.h"
@@ -14,6 +16,7 @@
 #include "Misc/Build.h"
 #include "Misc/NoopCounter.h"
 #include "Templates/AlignmentTemplates.h"
+#include "Templates/UnrealTemplate.h"
 
 
 // Enums for specifying memory allocation type.
diff --git a/Engine/Source/Runtime/Core/Public/Misc/MemoryReadStream.h b/Engine/Source/Runtime/Core/Public/Misc/MemoryReadStream.h
--- a/Engine/Source/Runtime/Core/Public/Misc/MemoryReadStream.h
+++ b/Engine/Source/Runtime/Core/Public/Misc/MemoryReadStream.h
@@ -3,6 +3,7 @@
 #pragma once
 
 #include "CoreMinimal.h"
+#include "HAL/Platform.h"
 #include "HAL/ThreadSafeCounter.h"
 #include "Templates/RefCounting.h"
 
diff --git a/Engine/Source/Runtime/Core/Public/Misc/MonitoredProcess.h b/Engine/Source/Runtime/Core/Public/Misc/MonitoredProcess.h
--- a/Engine/Source/Runtime/Core/Public/Misc/MonitoredProcess.h
+++ b/Engine/Source/Runtime/Core/Public/Misc/MonitoredProcess.h
@@ -5,6 +5,7 @@
 #include "Containers/UnrealString.h"
 #include "CoreTypes.h"
 #include "Delegates/Delegate.h"
+#include "Delegates/DelegateBase.h"
 #include "HAL/CriticalSection.h"
 #include "HAL/PlatformProcess.h"
 #include "HAL/Runnable.h"
diff --git a/Engine/Source/Runtime/Core/Public/Misc/OutputDeviceAnsiError.h b/Engine/Source/Runtime/Core/Public/Misc/OutputDeviceAnsiError.h
--- a/Engine/Source/Runtime/Core/Public/Misc/OutputDeviceAnsiError.h
+++ b/Engine/Source/Runtime/Core/Public/Misc/OutputDeviceAnsiError.h
@@ -2,9 +2,10 @@
 
 #pragma once
 
-#include "CoreMinimal.h"
 #include "CoreTypes.h"
+#include "Logging/LogVerbosity.h"
 #include "Misc/OutputDeviceError.h"
+#include "UObject/NameTypes.h"
 
 class CORE_API FOutputDeviceAnsiError : public FOutputDeviceError
 {
diff --git a/Engine/Source/Runtime/Core/Public/Misc/OutputDeviceRedirector.h b/Engine/Source/Runtime/Core/Public/Misc/OutputDeviceRedirector.h
--- a/Engine/Source/Runtime/Core/Public/Misc/OutputDeviceRedirector.h
+++ b/Engine/Source/Runtime/Core/Public/Misc/OutputDeviceRedirector.h
@@ -3,8 +3,9 @@
 #pragma once
 
 #include "Containers/Array.h"
-#include "CoreMinimal.h"
+#include "Containers/ContainerAllocationPolicies.h"
 #include "CoreTypes.h"
+#include "Logging/LogVerbosity.h"
 #include "Misc/EnumClassFlags.h"
 #include "Misc/OutputDevice.h"
 #include "Templates/PimplPtr.h"
diff --git a/Engine/Source/Runtime/Core/Public/Misc/PreLoadFile.h b/Engine/Source/Runtime/Core/Public/Misc/PreLoadFile.h
--- a/Engine/Source/Runtime/Core/Public/Misc/PreLoadFile.h
+++ b/Engine/Source/Runtime/Core/Public/Misc/PreLoadFile.h
@@ -2,7 +2,9 @@
 
 #pragma once
 
+#include "Containers/UnrealString.h"
 #include "CoreMinimal.h"
+#include "HAL/Platform.h"
 #include "Misc/DelayedAutoRegister.h"
 
 class FEvent;
diff --git a/Engine/Source/Runtime/Core/Public/Misc/QueuedThreadPoolWrapper.h b/Engine/Source/Runtime/Core/Public/Misc/QueuedThreadPoolWrapper.h
--- a/Engine/Source/Runtime/Core/Public/Misc/QueuedThreadPoolWrapper.h
+++ b/Engine/Source/Runtime/Core/Public/Misc/QueuedThreadPoolWrapper.h
@@ -5,14 +5,21 @@
 #include "Async/Fundamental/Scheduler.h"
 #include "Async/Fundamental/Task.h"
 #include "Async/TaskGraphInterfaces.h"
+#include "Containers/Array.h"
+#include "Containers/Map.h"
 #include "CoreMinimal.h"
 #include "Experimental/ConcurrentLinearAllocator.h"
 #include "Experimental/Containers/FAAArrayQueue.h"
 #include "Experimental/Misc/ExecutionResource.h"
 #include "GenericPlatform/GenericPlatformAffinity.h"
+#include "HAL/CriticalSection.h"
+#include "HAL/Platform.h"
 #include "HAL/PlatformAffinity.h"
+#include "HAL/PlatformCrt.h"
 #include "HAL/PlatformProcess.h"
 #include "HAL/ThreadSafeCounter.h"
+#include "HAL/UnrealMemory.h"
+#include "Misc/AssertionMacros.h"
 #include "Misc/IQueuedWork.h"
 #include "Misc/MemStack.h"
 #include "Misc/ScopeLock.h"
@@ -20,6 +27,7 @@
 #include "ScopeRWLock.h"
 #include "Stats/Stats2.h"
 #include "Templates/Atomic.h"
+#include "Templates/Function.h"
 #include "Templates/RefCounting.h"
 
 #include <atomic>
diff --git a/Engine/Source/Runtime/Core/Public/Misc/RuntimeErrors.h b/Engine/Source/Runtime/Core/Public/Misc/RuntimeErrors.h
--- a/Engine/Source/Runtime/Core/Public/Misc/RuntimeErrors.h
+++ b/Engine/Source/Runtime/Core/Public/Misc/RuntimeErrors.h
@@ -4,6 +4,7 @@
 
 #include "CoreTypes.h"
 #include "Delegates/Delegate.h"
+#include "Delegates/DelegateBase.h"
 #include "HAL/PlatformMisc.h"
 #include "Logging/LogVerbosity.h"
 #include "Misc/AssertionMacros.h"
diff --git a/Engine/Source/Runtime/Core/Public/Misc/SecureHash.h b/Engine/Source/Runtime/Core/Public/Misc/SecureHash.h
--- a/Engine/Source/Runtime/Core/Public/Misc/SecureHash.h
+++ b/Engine/Source/Runtime/Core/Public/Misc/SecureHash.h
@@ -3,21 +3,25 @@
 #pragma once
 
 #include "Async/AsyncWork.h"
+#include "Containers/Array.h"
 #include "Containers/Map.h"
 #include "Containers/StringConv.h"
 #include "Containers/StringFwd.h"
 #include "Containers/StringView.h"
 #include "Containers/UnrealString.h"
-#include "CoreMinimal.h"
 #include "CoreTypes.h"
 #include "HAL/PreprocessorHelpers.h"
 #include "HAL/UnrealMemory.h"
+#include "Misc/AssertionMacros.h"
+#include "Misc/CString.h"
+#include "Serialization/Archive.h"
 #include "Serialization/BufferReader.h"
 #include "Serialization/MemoryLayout.h"
 #include "Stats/Stats.h"
 #include "Stats/Stats2.h"
 #include "String/BytesToHex.h"
 #include "String/HexToBytes.h"
+#include "Templates/UnrealTemplate.h"
 
 class FMemoryImageWriter;
 class FMemoryUnfreezeContent;
diff --git a/Engine/Source/Runtime/Core/Public/Misc/StringFormatter.h b/Engine/Source/Runtime/Core/Public/Misc/StringFormatter.h
--- a/Engine/Source/Runtime/Core/Public/Misc/StringFormatter.h
+++ b/Engine/Source/Runtime/Core/Public/Misc/StringFormatter.h
@@ -2,11 +2,13 @@
 
 #pragma once
 
+#include "Containers/Array.h"
+#include "Containers/Map.h"
 #include "Containers/UnrealString.h"
-#include "CoreMinimal.h"
 #include "CoreTypes.h"
 #include "Misc/AssertionMacros.h"
 #include "Misc/ExpressionParserTypes.h"
+#include "Templates/UnrealTemplate.h"
 #include "Templates/ValueOrError.h"
 
 struct FStringFormatArg;
diff --git a/Engine/Source/Runtime/Core/Public/ProfilingDebugging/CsvProfiler.h b/Engine/Source/Runtime/Core/Public/ProfilingDebugging/CsvProfiler.h
--- a/Engine/Source/Runtime/Core/Public/ProfilingDebugging/CsvProfiler.h
+++ b/Engine/Source/Runtime/Core/Public/ProfilingDebugging/CsvProfiler.h
@@ -9,18 +9,27 @@
 
 #include "Async/Future.h"
 #include "Async/TaskGraphInterfaces.h"
+#include "Containers/Array.h"
+#include "Containers/Map.h"
 #include "Containers/Queue.h"
-#include "CoreMinimal.h"
+#include "Containers/SparseArray.h"
+#include "Containers/UnrealString.h"
 #include "CoreTypes.h"
+#include "Delegates/Delegate.h"
+#include "Delegates/DelegateBase.h"
+#include "Delegates/IDelegateInstance.h"
+#include "HAL/CriticalSection.h"
 #include "HAL/PreprocessorHelpers.h"
 #include "HAL/ThreadSafeCounter.h"
 #include "Misc/EnumClassFlags.h"
 #include "ProfilingDebugging/CsvProfilerConfig.h"
 #include "ProfilingDebugging/CsvProfilerTrace.h"
 #include "ProfilingDebugging/MiscTrace.h"
+#include "Templates/AndOrNot.h"
 #include "Templates/IsArrayOrRefOfTypeByPredicate.h"
 #include "Templates/IsValidVariadicFunctionArg.h"
 #include "Templates/UniquePtr.h"
+#include "Templates/UnrealTemplate.h"
 #include "Traits/IsCharEncodingCompatibleWith.h"
 #include "UObject/NameTypes.h"
 
diff --git a/Engine/Source/Runtime/Core/Public/ProfilingDebugging/TraceAuxiliary.h b/Engine/Source/Runtime/Core/Public/ProfilingDebugging/TraceAuxiliary.h
--- a/Engine/Source/Runtime/Core/Public/ProfilingDebugging/TraceAuxiliary.h
+++ b/Engine/Source/Runtime/Core/Public/ProfilingDebugging/TraceAuxiliary.h
@@ -5,6 +5,8 @@
 #include "Containers/StringFwd.h"
 #include "CoreGlobals.h"
 #include "Delegates/Delegate.h"
+#include "Delegates/DelegateBase.h"
+#include "Delegates/IDelegateInstance.h"
 #include "HAL/Platform.h"
 #include "Logging/LogCategory.h"
 #include "Logging/LogMacros.h"
diff --git a/Engine/Source/Runtime/Core/Public/Serialization/MappedName.h b/Engine/Source/Runtime/Core/Public/Serialization/MappedName.h
--- a/Engine/Source/Runtime/Core/Public/Serialization/MappedName.h
+++ b/Engine/Source/Runtime/Core/Public/Serialization/MappedName.h
@@ -2,9 +2,18 @@
 
 #pragma once
 
+#include "Containers/Array.h"
 #include "Containers/ArrayView.h"
+#include "Containers/ContainersFwd.h"
 #include "CoreMinimal.h"
+#include "HAL/Platform.h"
+#include "Math/NumericLimits.h"
+#include "Misc/AssertionMacros.h"
 #include "UObject/NameBatchSerialization.h"
+#include "UObject/NameTypes.h"
+#include "UObject/UnrealNames.h"
+
+class FArchive;
 
 /**
  * Index and name number into a name map.
diff --git a/Engine/Source/Runtime/Core/Public/Serialization/MemoryImage.h b/Engine/Source/Runtime/Core/Public/Serialization/MemoryImage.h
--- a/Engine/Source/Runtime/Core/Public/Serialization/MemoryImage.h
+++ b/Engine/Source/Runtime/Core/Public/Serialization/MemoryImage.h
@@ -3,18 +3,32 @@
 #pragma once
 
 #include "Containers/Array.h"
+#include "Containers/ContainerAllocationPolicies.h"
 #include "Containers/HashTable.h"
+#include "Containers/Map.h"
 #include "Containers/Set.h"
 #include "Containers/UnrealString.h"
-#include "CoreMinimal.h"
 #include "CoreTypes.h"
+#include "HAL/MemoryBase.h"
+#include "HAL/PlatformString.h"
 #include "HAL/PreprocessorHelpers.h"
+#include "HAL/UnrealMemory.h"
+#include "Math/UnrealMathUtility.h"
+#include "Misc/AssertionMacros.h"
+#include "Misc/CString.h"
+#include "Misc/Crc.h"
 #include "Misc/SecureHash.h"
 #include "Serialization/Archive.h"
 #include "Serialization/MemoryImageWriter.h"
 #include "Serialization/MemoryLayout.h"
+#include "Templates/AlignmentTemplates.h"
+#include "Templates/EnableIf.h"
 #include "Templates/RefCounting.h"
+#include "Templates/TypeHash.h"
 #include "Traits/IsCharType.h"
+#include "UObject/NameTypes.h"
+
+template <typename T> struct TIsContiguousContainer;
 
 #if defined(WITH_RTTI) || defined(_CPPRTTI) || defined(__GXX_RTTI) || WITH_EDITOR
 #include <typeinfo>
diff --git a/Engine/Source/Runtime/Core/Public/Stats/Stats2.h b/Engine/Source/Runtime/Core/Public/Stats/Stats2.h
--- a/Engine/Source/Runtime/Core/Public/Stats/Stats2.h
+++ b/Engine/Source/Runtime/Core/Public/Stats/Stats2.h
@@ -4,34 +4,47 @@
 
 #include "Containers/Array.h"
 #include "Containers/ChunkedArray.h"
+#include "Containers/ContainerAllocationPolicies.h"
 #include "Containers/LockFreeList.h"
 #include "Containers/UnrealString.h"
-#include "CoreMinimal.h"
+#include "CoreGlobals.h"
 #include "CoreTypes.h"
 #include "Delegates/Delegate.h"
+#include "Delegates/DelegateBase.h"
+#include "HAL/CriticalSection.h"
 #include "HAL/LowLevelMemTracker.h"
+#include "HAL/PlatformMemory.h"
+#include "HAL/PlatformMisc.h"
+#include "HAL/PlatformTLS.h"
 #include "HAL/PlatformTime.h"
 #include "HAL/ThreadSafeCounter.h"
 #include "HAL/ThreadSingleton.h"
+#include "HAL/UnrealMemory.h"
 #include "Math/Color.h"
 #include "Math/NumericLimits.h"
+#include "Misc/AssertionMacros.h"
 #include "Misc/Build.h"
+#include "Misc/CString.h"
 #include "Misc/EnumClassFlags.h"
 #include "ProfilingDebugging/CpuProfilerTrace.h"
 #include "ProfilingDebugging/MiscTrace.h"
 #include "StatsCommon.h"
 #include "StatsTrace.h"
 #include "Templates/Atomic.h"
+#include "Templates/TypeCompatibleBytes.h"
 #include "Templates/UniquePtr.h"
 #include "Templates/UnrealTemplate.h"
 #include "Trace/Detail/Channel.h"
 #include "Trace/Detail/Channel.inl"
 #include "Trace/Trace.h"
 #include "UObject/NameTypes.h"
+#include "UObject/UnrealNames.h"
 
+class FOutputDevice;
 class FScopeCycleCounter;
 class FThreadStats;
 struct TStatId;
+template <typename T> struct TIsPODType;
 
 /**
 * This is thread-private information about the thread idle stats, which we always collect, even in final builds
diff --git a/Engine/Source/Runtime/Core/Public/Stats/StatsData.h b/Engine/Source/Runtime/Core/Public/Stats/StatsData.h
--- a/Engine/Source/Runtime/Core/Public/Stats/StatsData.h
+++ b/Engine/Source/Runtime/Core/Public/Stats/StatsData.h
@@ -8,9 +8,12 @@
 #include "Containers/Set.h"
 #include "Containers/StringFwd.h"
 #include "Containers/UnrealString.h"
-#include "CoreMinimal.h"
 #include "CoreTypes.h"
 #include "Delegates/Delegate.h"
+#include "Delegates/DelegateBase.h"
+#include "Delegates/IDelegateInstance.h"
+#include "HAL/PlatformCrt.h"
+#include "HAL/PlatformMemory.h"
 #include "HAL/ThreadSafeCounter.h"
 #include "Math/NumericLimits.h"
 #include "Math/UnrealMathUtility.h"
@@ -22,6 +25,7 @@
 #include "Stats/StatsCommon.h"
 #include "Templates/UnrealTemplate.h"
 #include "UObject/NameTypes.h"
+#include "UObject/UnrealNames.h"
 
 struct FEventData;
 
diff --git a/Engine/Source/Runtime/Core/Public/Stats/StatsFile.h b/Engine/Source/Runtime/Core/Public/Stats/StatsFile.h
--- a/Engine/Source/Runtime/Core/Public/Stats/StatsFile.h
+++ b/Engine/Source/Runtime/Core/Public/Stats/StatsFile.h
@@ -7,23 +7,27 @@
 #include "Containers/ChunkedArray.h"
 #include "Containers/Map.h"
 #include "Containers/Set.h"
+#include "Containers/SparseArray.h"
 #include "Containers/UnrealString.h"
-#include "CoreMinimal.h"
+#include "CoreGlobals.h"
 #include "CoreTypes.h"
 #include "Delegates/IDelegateInstance.h"
 #include "HAL/ThreadSafeBool.h"
 #include "HAL/ThreadSafeCounter.h"
 #include "Internationalization/Text.h"
+#include "Logging/LogCategory.h"
 #include "Logging/LogMacros.h"
 #include "Misc/AssertionMacros.h"
 #include "Misc/Build.h"
 #include "Misc/Compression.h"
+#include "Serialization/Archive.h"
 #include "Stats/Stats.h"
 #include "Stats/Stats2.h"
 #include "Stats/StatsCommon.h"
 #include "Stats/StatsData.h"
 #include "Trace/Detail/Channel.h"
 #include "UObject/NameTypes.h"
+#include "UObject/UnrealNames.h"
 
 class FAsyncStatsWrite;
 struct FStatsReadFile;
diff --git a/Engine/Source/Runtime/Core/Public/Tasks/Pipe.h b/Engine/Source/Runtime/Core/Public/Tasks/Pipe.h
--- a/Engine/Source/Runtime/Core/Public/Tasks/Pipe.h
+++ b/Engine/Source/Runtime/Core/Public/Tasks/Pipe.h
@@ -4,10 +4,12 @@
 
 #include "Async/Fundamental/Scheduler.h"
 #include "Async/Fundamental/Task.h"
-#include "CoreMinimal.h"
 #include "CoreTypes.h"
+#include "Misc/AssertionMacros.h"
 #include "Tasks/Task.h"
 #include "Tasks/TaskPrivate.h"
+#include "Templates/Invoke.h"
+#include "Templates/UnrealTemplate.h"
 
 #include <atomic>
 
diff --git a/Engine/Source/Runtime/Core/Public/Unix/UnixErrorOutputDevice.h b/Engine/Source/Runtime/Core/Public/Unix/UnixErrorOutputDevice.h
--- a/Engine/Source/Runtime/Core/Public/Unix/UnixErrorOutputDevice.h
+++ b/Engine/Source/Runtime/Core/Public/Unix/UnixErrorOutputDevice.h
@@ -2,9 +2,10 @@
 
 #pragma once
 
-#include "CoreMinimal.h"
 #include "CoreTypes.h"
+#include "Logging/LogVerbosity.h"
 #include "Misc/OutputDeviceError.h"
+#include "UObject/NameTypes.h"
 
 class CORE_API FUnixErrorOutputDevice : public FOutputDeviceError
 {
diff --git a/Engine/Source/Runtime/Core/Public/Unix/UnixPlatformCrashContext.h b/Engine/Source/Runtime/Core/Public/Unix/UnixPlatformCrashContext.h
--- a/Engine/Source/Runtime/Core/Public/Unix/UnixPlatformCrashContext.h
+++ b/Engine/Source/Runtime/Core/Public/Unix/UnixPlatformCrashContext.h
@@ -2,7 +2,8 @@
 
 #pragma once
 
-#include "CoreMinimal.h"
+#include "Containers/Array.h"
+#include "Containers/UnrealString.h"
 #include "CoreTypes.h"
 #include "GenericPlatform/GenericPlatformCrashContext.h"
 #include "Templates/Atomic.h"
diff --git a/Engine/Source/Runtime/Core/Tests/Containers/TestUtils.cpp b/Engine/Source/Runtime/Core/Tests/Containers/TestUtils.cpp
--- a/Engine/Source/Runtime/Core/Tests/Containers/TestUtils.cpp
+++ b/Engine/Source/Runtime/Core/Tests/Containers/TestUtils.cpp
@@ -2,6 +2,8 @@
 
 #include "TestUtils.h"
 
+#include "HAL/Platform.h"
+
 #include <numeric>
 
 namespace Test
diff --git a/Engine/Source/Runtime/Core/Tests/Containers/TestUtils.h b/Engine/Source/Runtime/Core/Tests/Containers/TestUtils.h
--- a/Engine/Source/Runtime/Core/Tests/Containers/TestUtils.h
+++ b/Engine/Source/Runtime/Core/Tests/Containers/TestUtils.h
@@ -2,7 +2,9 @@
 
 #pragma once
 
-#include "CoreMinimal.h"
+#include "Containers/Array.h"
+#include "CoreTypes.h"
+#include "Math/UnrealMathUtility.h"
 
 namespace Test
 {
