Fix non-unity error.

diff --git a/Engine/Source/Runtime/Core/Public/Misc/CompressedGrowableBuffer.h b/Engine/Source/Runtime/Core/Public/Misc/CompressedGrowableBuffer.h
--- a/Engine/Source/Runtime/Core/Public/Misc/CompressedGrowableBuffer.h
+++ b/Engine/Source/Runtime/Core/Public/Misc/CompressedGrowableBuffer.h
@@ -4,6 +4,7 @@
 
 #include "CoreTypes.h"
 #include "Containers/Array.h"
+#include "UObject/NameTypes.h"
 
 /**
  * Growable compressed buffer. Usage is to append frequently but only request and therefore decompress
