Same as CL 20991367:

Updated Core code to not use CoreMinimal.h when possible.
Core public headers still will contain CoreMinimal.h if any included it before my IWYU changes.

diff --git a/Engine/Source/Runtime/Core/Internal/IO/IoStore.h b/Engine/Source/Runtime/Core/Internal/IO/IoStore.h
--- a/Engine/Source/Runtime/Core/Internal/IO/IoStore.h
+++ b/Engine/Source/Runtime/Core/Internal/IO/IoStore.h
@@ -2,12 +2,16 @@
 
 #pragma once
 
-#include "CoreMinimal.h"
+#include "Containers/Array.h"
+#include "HAL/Platform.h"
+#include "HAL/UnrealMemory.h"
 #include "IO/IoContainerId.h"
 #include "IO/IoDispatcher.h"
+#include "Logging/LogMacros.h"
 #include "Misc/EnumClassFlags.h"
 #include "Misc/Guid.h"
 #include "Misc/SecureHash.h"
+#include "UObject/NameTypes.h"
 
 CORE_API DECLARE_LOG_CATEGORY_EXTERN(LogIoStore, Log, All);
 
diff --git a/Engine/Source/Runtime/Core/Private/Async/Async.cpp b/Engine/Source/Runtime/Core/Private/Async/Async.cpp
--- a/Engine/Source/Runtime/Core/Private/Async/Async.cpp
+++ b/Engine/Source/Runtime/Core/Private/Async/Async.cpp
@@ -2,7 +2,6 @@
 
 #include "Async/Async.h"
 
-#include "CoreMinimal.h"
 #include "HAL/ThreadSafeCounter.h"
 
 
diff --git a/Engine/Source/Runtime/Core/Private/Containers/MemoryImageAllocator.cpp b/Engine/Source/Runtime/Core/Private/Containers/MemoryImageAllocator.cpp
--- a/Engine/Source/Runtime/Core/Private/Containers/MemoryImageAllocator.cpp
+++ b/Engine/Source/Runtime/Core/Private/Containers/MemoryImageAllocator.cpp
@@ -2,7 +2,6 @@
 
 #include "Containers/ResourceArray.h"
 #include "Containers/StringFwd.h"
-#include "CoreMinimal.h"
 #include "Misc/StringBuilder.h"
 #include "Serialization/MemoryImage.h"
 #include "Serialization/MemoryImageWriter.h"
diff --git a/Engine/Source/Runtime/Core/Private/Internationalization/LocalizationResourceTextSource.h b/Engine/Source/Runtime/Core/Private/Internationalization/LocalizationResourceTextSource.h
--- a/Engine/Source/Runtime/Core/Private/Internationalization/LocalizationResourceTextSource.h
+++ b/Engine/Source/Runtime/Core/Private/Internationalization/LocalizationResourceTextSource.h
@@ -3,7 +3,6 @@
 #pragma once
 
 #include "Containers/ArrayView.h"
-#include "CoreMinimal.h"
 #include "Internationalization/ILocalizedTextSource.h"
 #include "Internationalization/LocalizedTextSourceTypes.h"
 
diff --git a/Engine/Source/Runtime/Core/Public/Async/Fundamental/ReserveScheduler.h b/Engine/Source/Runtime/Core/Public/Async/Fundamental/ReserveScheduler.h
--- a/Engine/Source/Runtime/Core/Public/Async/Fundamental/ReserveScheduler.h
+++ b/Engine/Source/Runtime/Core/Public/Async/Fundamental/ReserveScheduler.h
@@ -2,7 +2,6 @@
 
 #pragma once 
 #include "Async/Fundamental/TaskDelegate.h"
-#include "CoreMinimal.h"
 #include "HAL/Event.h"
 #include "HAL/PlatformAffinity.h"
 #include "HAL/Thread.h"
diff --git a/Engine/Source/Runtime/Core/Public/Async/Fundamental/Scheduler.h b/Engine/Source/Runtime/Core/Public/Async/Fundamental/Scheduler.h
--- a/Engine/Source/Runtime/Core/Public/Async/Fundamental/Scheduler.h
+++ b/Engine/Source/Runtime/Core/Public/Async/Fundamental/Scheduler.h
@@ -5,7 +5,6 @@
 #include "Async/Fundamental/TaskDelegate.h"
 #include "Containers/Array.h"
 #include "Containers/ArrayView.h"
-#include "CoreMinimal.h"
 #include "HAL/Event.h"
 #include "HAL/PlatformAffinity.h"
 #include "HAL/PlatformProcess.h"
diff --git a/Engine/Source/Runtime/Core/Public/Async/ParallelFor.h b/Engine/Source/Runtime/Core/Public/Async/ParallelFor.h
--- a/Engine/Source/Runtime/Core/Public/Async/ParallelFor.h
+++ b/Engine/Source/Runtime/Core/Public/Async/ParallelFor.h
@@ -10,7 +10,6 @@
 #include "Async/Fundamental/Task.h"
 #include "Async/TaskGraphInterfaces.h"
 #include "Containers/ArrayView.h"
-#include "CoreMinimal.h"
 #include "CoreTypes.h"
 #include "Experimental/ConcurrentLinearAllocator.h"
 #include "HAL/Event.h"
diff --git a/Engine/Source/Runtime/Core/Public/Compression/CompressedBuffer.h b/Engine/Source/Runtime/Core/Public/Compression/CompressedBuffer.h
--- a/Engine/Source/Runtime/Core/Public/Compression/CompressedBuffer.h
+++ b/Engine/Source/Runtime/Core/Public/Compression/CompressedBuffer.h
@@ -2,7 +2,6 @@
 
 #pragma once
 
-#include "CoreMinimal.h"
 #include "Memory/CompositeBuffer.h"
 #include "Memory/MemoryFwd.h"
 #include "Memory/SharedBuffer.h"
diff --git a/Engine/Source/Runtime/Core/Public/Containers/BackgroundableTicker.h b/Engine/Source/Runtime/Core/Public/Containers/BackgroundableTicker.h
--- a/Engine/Source/Runtime/Core/Public/Containers/BackgroundableTicker.h
+++ b/Engine/Source/Runtime/Core/Public/Containers/BackgroundableTicker.h
@@ -2,7 +2,6 @@
 
 #pragma once
 
-#include "CoreMinimal.h"
 #include "Ticker.h"
 
 /**
diff --git a/Engine/Source/Runtime/Core/Public/Internationalization/StringTableCore.h b/Engine/Source/Runtime/Core/Public/Internationalization/StringTableCore.h
--- a/Engine/Source/Runtime/Core/Public/Internationalization/StringTableCore.h
+++ b/Engine/Source/Runtime/Core/Public/Internationalization/StringTableCore.h
@@ -4,7 +4,6 @@
 #include "Containers/Map.h"
 #include "Containers/SparseArray.h"
 #include "Containers/UnrealString.h"
-#include "CoreMinimal.h"
 #include "HAL/CriticalSection.h"
 #include "Internationalization/Internationalization.h"
 #include "Internationalization/LocKeyFuncs.h"
diff --git a/Engine/Source/Runtime/Core/Public/Internationalization/StringTableRegistry.h b/Engine/Source/Runtime/Core/Public/Internationalization/StringTableRegistry.h
--- a/Engine/Source/Runtime/Core/Public/Internationalization/StringTableRegistry.h
+++ b/Engine/Source/Runtime/Core/Public/Internationalization/StringTableRegistry.h
@@ -3,7 +3,6 @@
 
 #include "Containers/Map.h"
 #include "Containers/UnrealString.h"
-#include "CoreMinimal.h"
 #include "Delegates/Delegate.h"
 #include "Internationalization/LocKeyFuncs.h"
 #include "Internationalization/StringTableCoreFwd.h"
diff --git a/Engine/Source/Runtime/Core/Public/Unix/UnixPlatformProcess.h b/Engine/Source/Runtime/Core/Public/Unix/UnixPlatformProcess.h
--- a/Engine/Source/Runtime/Core/Public/Unix/UnixPlatformProcess.h
+++ b/Engine/Source/Runtime/Core/Public/Unix/UnixPlatformProcess.h
@@ -6,10 +6,11 @@
 
 #pragma once
 
+#include "Containers/Array.h"
 #include "Containers/UnrealString.h"
-#include "CoreMinimal.h"
 #include "CoreTypes.h"
 #include "GenericPlatform/GenericPlatformProcess.h"
+#include "HAL/PlatformMisc.h"
 #include "Misc/AssertionMacros.h"
 
 class Error;
