Fixed nopch disableunity compile issue.

diff --git a/Engine/Source/Runtime/Core/Private/Internationalization/TextHistory.h b/Engine/Source/Runtime/Core/Private/Internationalization/TextHistory.h
--- a/Engine/Source/Runtime/Core/Private/Internationalization/TextHistory.h
+++ b/Engine/Source/Runtime/Core/Private/Internationalization/TextHistory.h
@@ -8,6 +8,8 @@
 #include "Internationalization/LocalizedTextSourceTypes.h"
 #include "Internationalization/StringTableCoreFwd.h"
 #include "Internationalization/TextKey.h"
+#include "Internationalization/Text.h"
+#include "Misc/DateTime.h"
 #include "Serialization/StructuredArchive.h"
 #include "Templates/SharedPointerInternals.h"
 #include "Templates/UnrealTemplate.h"
