Replace uses of master in OutputDeviceRedirector

diff --git a/Engine/Source/Developer/Datasmith/DatasmithExporter/Private/DatasmithExporterManager.cpp b/Engine/Source/Developer/Datasmith/DatasmithExporter/Private/DatasmithExporterManager.cpp
--- a/Engine/Source/Developer/Datasmith/DatasmithExporter/Private/DatasmithExporterManager.cpp
+++ b/Engine/Source/Developer/Datasmith/DatasmithExporter/Private/DatasmithExporterManager.cpp
@@ -105,7 +105,7 @@ namespace DatasmithGameThread
 		if ( GLog )
 		{
 			// Make sure Logger is set on the right thread
-			GLog->SetCurrentThreadAsMasterThread();
+			GLog->SetCurrentThreadAsMainLogThread();
 
 			if ( bSuppressLogs )
 			{
@@ -130,7 +130,7 @@ namespace DatasmithGameThread
 			// Set the current thread on the Logger
 			if ( GLog )
 			{
-				GLog->SetCurrentThreadAsMasterThread();
+				GLog->SetCurrentThreadAsMainLogThread();
 			}
 
 			// Re-attach task graph to new game thread
diff --git a/Engine/Source/Runtime/Core/Private/Mac/CocoaThread.cpp b/Engine/Source/Runtime/Core/Private/Mac/CocoaThread.cpp
--- a/Engine/Source/Runtime/Core/Private/Mac/CocoaThread.cpp
+++ b/Engine/Source/Runtime/Core/Private/Mac/CocoaThread.cpp
@@ -377,7 +377,7 @@ FCocoaRunLoopSource* FCocoaRunLoopSource::GameRunLoopSource = nullptr;
 	
 	if (GLog)
 	{
-		GLog->SetCurrentThreadAsMasterThread();
+		GLog->SetCurrentThreadAsMainLogThread();
 	}
 
 	[super main];
diff --git a/Engine/Source/Runtime/Core/Private/Misc/OutputDeviceRedirector.cpp b/Engine/Source/Runtime/Core/Private/Misc/OutputDeviceRedirector.cpp
--- a/Engine/Source/Runtime/Core/Private/Misc/OutputDeviceRedirector.cpp
+++ b/Engine/Source/Runtime/Core/Private/Misc/OutputDeviceRedirector.cpp
@@ -114,11 +114,11 @@ struct FOutputDeviceRedirectorState
 	std::atomic<uint32> OutputDevicesLockState = 0;
 	uint8 OutputDevicesLockPadding[CalculateRedirectorCacheLinePadding(sizeof(OutputDevicesLock) + sizeof(OutputDevicesLockState))]{};
 
-	/** A queue of lines logged by non-master threads. */
+	/** A queue of lines logged by non-main threads. */
 	TDepletableMpscQueue<FOutputDeviceLine, FOutputDeviceLinearAllocator> BufferedLines;
 	uint8 BufferedLinesPadding[CalculateRedirectorCacheLinePadding(sizeof(BufferedLines))]{};
 
-	/** Array of output devices to redirect to from the master thread. */
+	/** Array of output devices to redirect to from the main thread. */
 	TArray<FOutputDevice*> BufferedOutputDevices;
 
 	/** Array of output devices to redirect to from the calling thread. */
@@ -128,23 +128,23 @@ struct FOutputDeviceRedirectorState
 	TArray<FBufferedLine> BacklogLines;
 	FRWLock BacklogLock;
 
-	/** An optional dedicated master thread for logging to buffered output devices. */
+	/** An optional dedicated main thread for logging to buffered output devices. */
 	FThread Thread;
 
 	/** A lock to synchronize access to the thread. */
 	FRWLock ThreadLock;
 
-	/** A queue of events to trigger when the dedicated master thread is idle. */
+	/** A queue of events to trigger when the dedicated main thread is idle. */
 	TDepletableMpscQueue<FEvent*, FOutputDeviceLinearAllocator> ThreadIdleEvents;
 
-	/** An event to wake the dedicated master thread to process buffered lines. */
+	/** An event to wake the dedicated main thread to process buffered lines. */
 	std::atomic<FEvent*> ThreadWakeEvent = nullptr;
 
-	/** The ID of the thread holding the master lock. */
+	/** The ID of the thread holding the main lock. */
 	std::atomic<uint32> LockedThreadId = MAX_uint32;
 
-	/** The ID of the master thread. Logging from other threads will be buffered for processing by the master thread. */
-	std::atomic<uint32> MasterThreadId = FPlatformTLS::GetCurrentThreadId();
+	/** The ID of the main logging thread. Logging from other threads will be buffered for processing by the main thread. */
+	std::atomic<uint32> MainLogThreadId = FPlatformTLS::GetCurrentThreadId();
 
 	/** The ID of the panic thread, which is only set by Panic(). */
 	std::atomic<uint32> PanicThreadId = MAX_uint32;
@@ -156,9 +156,9 @@ struct FOutputDeviceRedirectorState
 	TBitArray<TInlineAllocator<1>> BufferedOutputDevicesCanBeUsedOnPanicThread;
 	TBitArray<TInlineAllocator<1>> UnbufferedOutputDevicesCanBeUsedOnPanicThread;
 
-	bool IsMasterThread(const uint32 ThreadId) const
+	bool IsMainLogThread(const uint32 ThreadId) const
 	{
-		return ThreadId == MasterThreadId.load(std::memory_order_relaxed);
+		return ThreadId == MainLogThreadId.load(std::memory_order_relaxed);
 	}
 
 	bool IsPanicThread(const uint32 ThreadId) const
@@ -212,7 +212,7 @@ struct FOutputDeviceRedirectorState
  * The read lock:
  * - Must be locked to read the OutputDevices arrays.
  * - Must be locked to write to unbuffered output devices.
- * - Must not be entered when the thread holds a write or master lock.
+ * - Must not be entered when the thread holds a write or main lock.
  */
 class FOutputDevicesReadScopeLock
 {
@@ -251,9 +251,9 @@ private:
 /**
  * A scoped lock for writers of the OutputDevices arrays.
  *
- * The write lock has the same access as the master lock, and:
+ * The write lock has the same access as the main lock, and:
  * - Must be locked to add or remove output devices.
- * - Must not be entered when the thread holds a read, write, or master lock.
+ * - Must not be entered when the thread holds a read, write, or main lock.
  */
 class FOutputDevicesWriteScopeLock
 {
@@ -294,20 +294,20 @@ private:
 };
 
 /**
- * A scoped lock for readers of the OutputDevices arrays that need to access master thread state.
+ * A scoped lock for readers of the OutputDevices arrays that need to access main logging thread's state.
  *
- * The master lock has the same access as the read lock, and:
- * - Must not be entered when the thread holds a write lock or master lock.
+ * The main lock has the same access as the read lock, and:
+ * - Must not be entered when the thread holds a write lock or main lock.
  * - Must check IsLocked() before performing restricted operations.
  * - Must be locked to write to buffered output devices.
  * - Must be locked while calling FlushBufferedLines().
  * - May be locked when the thread holds a read lock.
  * - When a panic thread is active, may only be locked from the panic thread.
  */
-class FOutputDevicesMasterScope
+class FOutputDevicesMainLogScope
 {
 public:
-	explicit FOutputDevicesMasterScope(FOutputDeviceRedirectorState& InState)
+	explicit FOutputDevicesMainLogScope(FOutputDeviceRedirectorState& InState)
 		: State(InState)
 	{
 		const uint32 ThreadId = FPlatformTLS::GetCurrentThreadId();
@@ -337,7 +337,7 @@ public:
 		}
 	}
 
-	FORCEINLINE ~FOutputDevicesMasterScope()
+	FORCEINLINE ~FOutputDevicesMainLogScope()
 	{
 		if (bLocked)
 		{
@@ -421,17 +421,17 @@ void FOutputDeviceRedirectorState::ThreadLoop()
 {
 	const uint32 ThreadId = FPlatformTLS::GetCurrentThreadId();
 
-	if (FOutputDevicesMasterScope Lock(*this); Lock.IsLocked())
+	if (FOutputDevicesMainLogScope Lock(*this); Lock.IsLocked())
 	{
-		MasterThreadId.store(ThreadId, std::memory_order_relaxed);
+		MainLogThreadId.store(ThreadId, std::memory_order_relaxed);
 	}
 
 	while (FEvent* WakeEvent = ThreadWakeEvent.load(std::memory_order_acquire))
 	{
 		WakeEvent->Wait();
-		while (!BufferedLines.IsEmpty() && IsMasterThread(ThreadId))
+		while (!BufferedLines.IsEmpty() && IsMainLogThread(ThreadId))
 		{
-			if (FOutputDevicesMasterScope Lock(*this); Lock.IsLocked())
+			if (FOutputDevicesMainLogScope Lock(*this); Lock.IsLocked())
 			{
 				FlushBufferedLines();
 			}
@@ -510,7 +510,7 @@ void FOutputDeviceRedirector::FlushThreadedLogs(EOutputDeviceRedirectorFlushOpti
 		return;
 	}
 
-	if (UE::Private::FOutputDevicesMasterScope Lock(*State); Lock.IsLocked())
+	if (UE::Private::FOutputDevicesMainLogScope Lock(*State); Lock.IsLocked())
 	{
 		State->FlushBufferedLines();
 	}
@@ -535,24 +535,24 @@ void FOutputDeviceRedirector::EnableBacklog(bool bEnable)
 	}
 }
 
-void FOutputDeviceRedirector::SetCurrentThreadAsMasterThread()
+void FOutputDeviceRedirector::SetCurrentThreadAsMainLogThread()
 {
 	const uint32 ThreadId = FPlatformTLS::GetCurrentThreadId();
 
-	if (UE::Private::FOutputDevicesMasterScope Lock(*State); !Lock.IsLocked() || State->MasterThreadId.load(std::memory_order_relaxed) == ThreadId)
+	if (UE::Private::FOutputDevicesMainLogScope Lock(*State); !Lock.IsLocked() || State->MainLogThreadId.load(std::memory_order_relaxed) == ThreadId)
 	{
 		return;
 	}
 	else
 	{
-		State->MasterThreadId.store(ThreadId, std::memory_order_relaxed);
+		State->MainLogThreadId.store(ThreadId, std::memory_order_relaxed);
 		State->FlushBufferedLines();
 	}
 
 	State->TryStopThread();
 }
 
-bool FOutputDeviceRedirector::TryStartDedicatedMasterThread()
+bool FOutputDeviceRedirector::TryStartDedicatedMainLogThread()
 {
 	return FApp::ShouldUseThreadingForPerformance() && State->TryStartThread();
 }
@@ -590,13 +590,13 @@ void FOutputDeviceRedirector::Serialize(const TCHAR* const Data, const ELogVerbo
 		State->BacklogLines.Emplace(Data, Category, Verbosity, RealTime);
 	}
 
-	// Serialize to buffered output devices from the master thread.
+	// Serialize to buffered output devices from the main logging thread.
 	// Lines are queued until buffered output devices are added to avoid missing early log lines.
-	if (State->IsMasterThread(ThreadId) && !State->BufferedOutputDevices.IsEmpty())
+	if (State->IsMainLogThread(ThreadId) && !State->BufferedOutputDevices.IsEmpty())
 	{
-		// Verify that this is the master thread again because another thread may have become
-		// the master thread between the previous check and the lock.
-		if (UE::Private::FOutputDevicesMasterScope MasterLock(*State); MasterLock.IsLocked() && State->IsMasterThread(ThreadId))
+		// Verify that this is the main thread again because another thread may have become
+		// the main thread between the previous check and the lock.
+		if (UE::Private::FOutputDevicesMainLogScope MainLock(*State); MainLock.IsLocked() && State->IsMainLogThread(ThreadId))
 		{
 			State->FlushBufferedLines();
 			State->BroadcastTo(ThreadId, State->BufferedOutputDevices, State->BufferedOutputDevicesCanBeUsedOnPanicThread,
@@ -610,7 +610,7 @@ void FOutputDeviceRedirector::Serialize(const TCHAR* const Data, const ELogVerbo
 		}
 	}
 
-	// Queue the line to serialize to buffered output devices from the master thread.
+	// Queue the line to serialize to buffered output devices from the main thread.
 	if (State->BufferedLines.EnqueueAndReturnWasEmpty(Data, Category, Verbosity, RealTime))
 	{
 		if (FEvent* WakeEvent = State->ThreadWakeEvent.load(std::memory_order_acquire))
@@ -637,7 +637,7 @@ void FOutputDeviceRedirector::RedirectLog(const FLazyName& Category, ELogVerbosi
 
 void FOutputDeviceRedirector::Flush()
 {
-	if (UE::Private::FOutputDevicesMasterScope Lock(*State); Lock.IsLocked())
+	if (UE::Private::FOutputDevicesMainLogScope Lock(*State); Lock.IsLocked())
 	{
 		State->FlushBufferedLines();
 		const uint32 ThreadId = FPlatformTLS::GetCurrentThreadId();
@@ -671,8 +671,8 @@ void FOutputDeviceRedirector::Panic()
 			FPlatformProcess::Yield();
 		}
 
-		// Make the panic thread the master thread. Neither thread can be changed after this point.
-		State->MasterThreadId.exchange(ThreadId, std::memory_order_relaxed);
+		// Make the panic thread the main thread. Neither thread can be changed after this point.
+		State->MainLogThreadId.exchange(ThreadId, std::memory_order_relaxed);
 
 		// Flush. Every log from the panic thread after this point will also flush.
 		Flush();
@@ -686,7 +686,7 @@ void FOutputDeviceRedirector::Panic()
 
 void FOutputDeviceRedirector::TearDown()
 {
-	SetCurrentThreadAsMasterThread();
+	SetCurrentThreadAsMainLogThread();
 
 	Flush();
 
diff --git a/Engine/Source/Runtime/Core/Private/Unix/UnixPlatformProcess.cpp b/Engine/Source/Runtime/Core/Private/Unix/UnixPlatformProcess.cpp
--- a/Engine/Source/Runtime/Core/Private/Unix/UnixPlatformProcess.cpp
+++ b/Engine/Source/Runtime/Core/Private/Unix/UnixPlatformProcess.cpp
@@ -1531,7 +1531,7 @@ FGenericPlatformProcess::EWaitAndForkResult FUnixPlatformProcess::WaitAndFork()
 
 				// Start up the log again
 				FPlatformOutputDevices::SetupOutputDevices();
-				GLog->SetCurrentThreadAsMasterThread();
+				GLog->SetCurrentThreadAsMainLogThread();
 
 				// Set the process name, if specified
 				if (ChildIdx > 0)
diff --git a/Engine/Source/Runtime/Core/Public/Misc/OutputDeviceRedirector.h b/Engine/Source/Runtime/Core/Public/Misc/OutputDeviceRedirector.h
--- a/Engine/Source/Runtime/Core/Public/Misc/OutputDeviceRedirector.h
+++ b/Engine/Source/Runtime/Core/Public/Misc/OutputDeviceRedirector.h
@@ -57,7 +57,7 @@ enum class EOutputDeviceRedirectorFlushOptions : uint32
 	/**
 		* Flush asynchronously when possible.
 		*
-		* When this flag is set and there is a dedicated master thread, the flush function returns immediately.
+		* When this flag is set and there is a dedicated main logging thread, the flush function returns immediately.
 		* Otherwise, the flush function does not return until the requested type of flush is complete.
 		*/
 	Async = 1 << 0,
@@ -122,17 +122,17 @@ public:
 	 */
 	void EnableBacklog(bool bEnable);
 
-	/** Sets the current thread to be the master thread redirects logs without buffering. */
-	void SetCurrentThreadAsMasterThread();
+	/** Sets the current thread to be the main thread, redirecting logs without buffering. */
+	void SetCurrentThreadAsMainLogThread();
 
 	/**
-	 * Starts a dedicated master thread that redirects logs that require buffering.
+	 * Starts a dedicated main thread that redirects logs that require buffering.
 	 *
 	 * A thread will not be started for certain configurations or platforms, or when threading is disabled.
 	 *
-	 * @return true if a dedicated master thread is running, false otherwise.
+	 * @return true if a dedicated main logging thread is running, false otherwise.
 	 */
-	bool TryStartDedicatedMasterThread();
+	bool TryStartDedicatedMainLogThread();
 
 	/**
 	 * Serializes the passed in data via all current output devices.
@@ -162,7 +162,7 @@ public:
 	 *
 	 * Only one thread can be the panic thread. Subsequent calls from other threads are ignored.
 	 * Only redirects logs to panic-safe output devices from this point forward.
-	 * Makes the calling thread the master thread as well.
+	 * Makes the calling thread the main log thread as well.
 	 * Flushes buffered logs to panic-safe output devices.
 	 * Flushes panic-safe output devices.
 	 */
diff --git a/Engine/Source/Runtime/Launch/Private/Android/LaunchAndroid.cpp b/Engine/Source/Runtime/Launch/Private/Android/LaunchAndroid.cpp
--- a/Engine/Source/Runtime/Launch/Private/Android/LaunchAndroid.cpp
+++ b/Engine/Source/Runtime/Launch/Private/Android/LaunchAndroid.cpp
@@ -544,7 +544,7 @@ int32 AndroidMain(struct android_app* state)
 
 	UE_LOG(LogAndroid, Display, TEXT("Passed PreInit()"));
 
-	GLog->SetCurrentThreadAsMasterThread();
+	GLog->SetCurrentThreadAsMainLogThread();
 
 	FAppEventManager::GetInstance()->SetEmptyQueueHandlerEvent(FPlatformProcess::GetSynchEventFromPool(false));
 
diff --git a/Engine/Source/Runtime/Launch/Private/IOS/LaunchIOS.cpp b/Engine/Source/Runtime/Launch/Private/IOS/LaunchIOS.cpp
--- a/Engine/Source/Runtime/Launch/Private/IOS/LaunchIOS.cpp
+++ b/Engine/Source/Runtime/Launch/Private/IOS/LaunchIOS.cpp
@@ -411,7 +411,7 @@ void FAppEntry::Init()
 	NSLog(@"%s", "Initializing ULD Communications in game mode\n");
 	GCommandSystem.Init();
 
-	GLog->SetCurrentThreadAsMasterThread();
+	GLog->SetCurrentThreadAsMainLogThread();
 	
 	// Send the launch local notification to the local notification service now that the engine module system has been initialized
 	if(gAppLaunchedWithLocalNotification)
diff --git a/Engine/Source/Runtime/Launch/Private/LaunchEngineLoop.cpp b/Engine/Source/Runtime/Launch/Private/LaunchEngineLoop.cpp
--- a/Engine/Source/Runtime/Launch/Private/LaunchEngineLoop.cpp
+++ b/Engine/Source/Runtime/Launch/Private/LaunchEngineLoop.cpp
@@ -1548,7 +1548,7 @@ int32 FEngineLoop::PreInitPreStartupScreen(const TCHAR* CmdLine)
 	// until initialization is far enough along to try to start a dedicated master thread.
 	if (GLog)
 	{
-		GLog->SetCurrentThreadAsMasterThread();
+		GLog->SetCurrentThreadAsMainLogThread();
 	}
 
 	// Command line option for enabling named events
@@ -1699,7 +1699,7 @@ int32 FEngineLoop::PreInitPreStartupScreen(const TCHAR* CmdLine)
 	// Try to start the dedicated master thread now that the command line is available.
 	if (!FParse::Param(FCommandLine::Get(), TEXT("NoLogThread")))
 	{
-		GLog->TryStartDedicatedMasterThread();
+		GLog->TryStartDedicatedMainLogThread();
 	}
 
 	// Initialize std out device as early as possible if requested in the command line
