Add Insight tags on a few config cache ini functions

diff --git a/Engine/Source/Runtime/Core/Private/Misc/ConfigCacheIni.cpp b/Engine/Source/Runtime/Core/Private/Misc/ConfigCacheIni.cpp
--- a/Engine/Source/Runtime/Core/Private/Misc/ConfigCacheIni.cpp
+++ b/Engine/Source/Runtime/Core/Private/Misc/ConfigCacheIni.cpp
@@ -494,6 +494,8 @@ static bool LoadConfigFileWrapper(const TCHAR* IniFile, FString& Contents, bool
  */
 static bool SaveConfigFileWrapper(const TCHAR* IniFile, const FString& Contents)
 {
+	TRACE_CPUPROFILER_EVENT_SCOPE(SaveConfigFileWrapper);
+
 	// let anyone that needs to save it, do so (counting how many did)
 	int32 SavedCount = 0;
 	FCoreDelegates::PreSaveConfigFileDelegate.Broadcast(IniFile, Contents, SavedCount);
@@ -2070,6 +2072,8 @@ TArray<FString> FConfigCacheIni::GetFilenames()
 
 void FConfigCacheIni::Flush(bool bRemoveFromCache, const FString& Filename )
 {
+	TRACE_CPUPROFILER_EVENT_SCOPE(FConfigCacheIni::Flush);
+
 	// never Flush temporary cache objects
 	if (Type != EConfigCacheType::Temporary)
 	{
