Added a few IWYU pragmas.

diff --git a/Engine/Source/Runtime/Core/Public/GenericPlatform/GenericPlatformTime.h b/Engine/Source/Runtime/Core/Public/GenericPlatform/GenericPlatformTime.h
--- a/Engine/Source/Runtime/Core/Public/GenericPlatform/GenericPlatformTime.h
+++ b/Engine/Source/Runtime/Core/Public/GenericPlatform/GenericPlatformTime.h
@@ -6,7 +6,7 @@
 #include "Containers/UnrealString.h"
 
 #if PLATFORM_HAS_BSD_TIME 
-	#include <sys/time.h>
+	#include <sys/time.h> // IWYU pragma: export
 #endif
 
 
diff --git a/Engine/Source/Runtime/Core/Public/HAL/Platform.h b/Engine/Source/Runtime/Core/Public/HAL/Platform.h
--- a/Engine/Source/Runtime/Core/Public/HAL/Platform.h
+++ b/Engine/Source/Runtime/Core/Public/HAL/Platform.h
@@ -1,5 +1,5 @@
 // Copyright Epic Games, Inc. All Rights Reserved.
-
+// IWYU pragma: begin_exports
 #pragma once
 
 #include "Misc/Build.h"
@@ -1169,3 +1169,5 @@ namespace UE::Core::Private
 #define UTF8TEXT(x) (UE::Core::Private::ToUTF8Literal(UTF8TEXT_PASTE(x)))
 
 #define WIDETEXT(str) WIDETEXT_PASTE(str)
+
+// IWYU pragma: end_exports
diff --git a/Engine/Source/Runtime/Core/Public/Linux/LinuxPlatformAffinity.h b/Engine/Source/Runtime/Core/Public/Linux/LinuxPlatformAffinity.h
--- a/Engine/Source/Runtime/Core/Public/Linux/LinuxPlatformAffinity.h
+++ b/Engine/Source/Runtime/Core/Public/Linux/LinuxPlatformAffinity.h
@@ -2,5 +2,5 @@
 
 #pragma once
 
-#include "GenericPlatform/GenericPlatformAffinity.h"
+#include "GenericPlatform/GenericPlatformAffinity.h" // IWYU pragma: export
 typedef FGenericPlatformAffinity FPlatformAffinity;
diff --git a/Engine/Source/Runtime/Core/Public/Linux/LinuxPlatformFile.h b/Engine/Source/Runtime/Core/Public/Linux/LinuxPlatformFile.h
--- a/Engine/Source/Runtime/Core/Public/Linux/LinuxPlatformFile.h
+++ b/Engine/Source/Runtime/Core/Public/Linux/LinuxPlatformFile.h
@@ -6,4 +6,4 @@
 
 #pragma once
 
-#include "Unix/UnixPlatformFile.h"
+#include "Unix/UnixPlatformFile.h" // IWYU pragma: export
diff --git a/Engine/Source/Runtime/Core/Public/Unix/UnixPlatformFile.h b/Engine/Source/Runtime/Core/Public/Unix/UnixPlatformFile.h
--- a/Engine/Source/Runtime/Core/Public/Unix/UnixPlatformFile.h
+++ b/Engine/Source/Runtime/Core/Public/Unix/UnixPlatformFile.h
@@ -8,7 +8,7 @@
 
 #include "Containers/UnrealString.h"
 #include "CoreTypes.h"
-#include "GenericPlatform/GenericPlatformFile.h"
+#include "GenericPlatform/GenericPlatformFile.h" // IWYU pragma: export
 #include "Misc/DateTime.h"
 
 class IMappedFileHandle;
diff --git a/Engine/Source/Runtime/Core/Public/Unix/UnixPlatformProcess.h b/Engine/Source/Runtime/Core/Public/Unix/UnixPlatformProcess.h
--- a/Engine/Source/Runtime/Core/Public/Unix/UnixPlatformProcess.h
+++ b/Engine/Source/Runtime/Core/Public/Unix/UnixPlatformProcess.h
@@ -9,7 +9,7 @@
 #include "Containers/Array.h"
 #include "Containers/UnrealString.h"
 #include "CoreTypes.h"
-#include "GenericPlatform/GenericPlatformProcess.h"
+#include "GenericPlatform/GenericPlatformProcess.h" // IWYU pragma: export
 #include "HAL/PlatformMisc.h"
 #include "Misc/AssertionMacros.h"
 
diff --git a/Engine/Source/Runtime/TraceLog/Public/Trace/Trace.inl b/Engine/Source/Runtime/TraceLog/Public/Trace/Trace.inl
--- a/Engine/Source/Runtime/TraceLog/Public/Trace/Trace.inl
+++ b/Engine/Source/Runtime/TraceLog/Public/Trace/Trace.inl
@@ -2,9 +2,11 @@
 
 #pragma once
 
+// IWYU pragma: begin_exports
 #include "CoreTypes.h"
 #include "Detail/Channel.inl"
 #include "Detail/EventNode.h"
 #include "Detail/Important/ImportantLogScope.inl"
 #include "Detail/LogScope.inl"
 #include "Detail/Writer.inl"
+// IWYU pragma: end_exports
