Add TestEqual for FTransform in AutomationTest

diff --git a/Engine/Source/Runtime/Core/Private/Misc/AutomationTest.cpp b/Engine/Source/Runtime/Core/Private/Misc/AutomationTest.cpp
--- a/Engine/Source/Runtime/Core/Private/Misc/AutomationTest.cpp
+++ b/Engine/Source/Runtime/Core/Private/Misc/AutomationTest.cpp
@@ -1364,6 +1364,16 @@ bool FAutomationTestBase::TestEqual(const TCHAR* What, const FVector Actual, con
 	return true;
 }
 
+bool FAutomationTestBase::TestEqual(const TCHAR* What, const FTransform Actual, const FTransform Expected, float Tolerance)
+{
+	if (!Expected.Equals(Actual, Tolerance))
+	{
+		AddError(FString::Printf(TEXT("Expected '%s' to be %s, but it was %s within tolerance %f."), What, *Expected.ToString(), *Actual.ToString(), Tolerance), 1);
+		return false;
+	}
+	return true;
+}
+
 bool FAutomationTestBase::TestEqual(const TCHAR* What, const FRotator Actual, const FRotator Expected, float Tolerance)
 {
 	if (!Expected.Equals(Actual, Tolerance))
diff --git a/Engine/Source/Runtime/Core/Public/Misc/AutomationTest.h b/Engine/Source/Runtime/Core/Public/Misc/AutomationTest.h
--- a/Engine/Source/Runtime/Core/Public/Misc/AutomationTest.h
+++ b/Engine/Source/Runtime/Core/Public/Misc/AutomationTest.h
@@ -1514,6 +1514,7 @@ public:
 	bool TestEqual(const TCHAR* What, float Actual, float Expected, float Tolerance = UE_KINDA_SMALL_NUMBER);
 	bool TestEqual(const TCHAR* What, double Actual, double Expected, double Tolerance = UE_KINDA_SMALL_NUMBER);
 	bool TestEqual(const TCHAR* What, FVector Actual, FVector Expected, float Tolerance = UE_KINDA_SMALL_NUMBER);
+	bool TestEqual(const TCHAR* What, FTransform Actual, FTransform Expected, float Tolerance = UE_KINDA_SMALL_NUMBER);
 	bool TestEqual(const TCHAR* What, FRotator Actual, FRotator Expected, float Tolerance = UE_KINDA_SMALL_NUMBER);
 	bool TestEqual(const TCHAR* What, FColor Actual, FColor Expected);
 	bool TestEqual(const TCHAR* What, FLinearColor Actual, FLinearColor Expected);
