Logging: Restored and deprecated a deleted function and tweaked the affected names

#rnx

diff --git a/Engine/Source/Developer/Datasmith/DatasmithExporter/Private/DatasmithExporterManager.cpp b/Engine/Source/Developer/Datasmith/DatasmithExporter/Private/DatasmithExporterManager.cpp
--- a/Engine/Source/Developer/Datasmith/DatasmithExporter/Private/DatasmithExporterManager.cpp
+++ b/Engine/Source/Developer/Datasmith/DatasmithExporter/Private/DatasmithExporterManager.cpp
@@ -105,7 +105,7 @@ namespace DatasmithGameThread
 		if ( GLog )
 		{
 			// Make sure Logger is set on the right thread
-			GLog->SetCurrentThreadAsMainLogThread();
+			GLog->SetCurrentThreadAsPrimaryThread();
 
 			if ( bSuppressLogs )
 			{
@@ -130,7 +130,7 @@ namespace DatasmithGameThread
 			// Set the current thread on the Logger
 			if ( GLog )
 			{
-				GLog->SetCurrentThreadAsMainLogThread();
+				GLog->SetCurrentThreadAsPrimaryThread();
 			}
 
 			// Re-attach task graph to new game thread
diff --git a/Engine/Source/Runtime/Core/Private/Mac/CocoaThread.cpp b/Engine/Source/Runtime/Core/Private/Mac/CocoaThread.cpp
--- a/Engine/Source/Runtime/Core/Private/Mac/CocoaThread.cpp
+++ b/Engine/Source/Runtime/Core/Private/Mac/CocoaThread.cpp
@@ -377,7 +377,7 @@ FCocoaRunLoopSource* FCocoaRunLoopSource::GameRunLoopSource = nullptr;
 	
 	if (GLog)
 	{
-		GLog->SetCurrentThreadAsMainLogThread();
+		GLog->SetCurrentThreadAsPrimaryThread();
 	}
 
 	[super main];
diff --git a/Engine/Source/Runtime/Core/Private/Misc/OutputDeviceRedirector.cpp b/Engine/Source/Runtime/Core/Private/Misc/OutputDeviceRedirector.cpp
--- a/Engine/Source/Runtime/Core/Private/Misc/OutputDeviceRedirector.cpp
+++ b/Engine/Source/Runtime/Core/Private/Misc/OutputDeviceRedirector.cpp
@@ -114,11 +114,11 @@ struct FOutputDeviceRedirectorState
 	std::atomic<uint32> OutputDevicesLockState = 0;
 	uint8 OutputDevicesLockPadding[CalculateRedirectorCacheLinePadding(sizeof(OutputDevicesLock) + sizeof(OutputDevicesLockState))]{};
 
-	/** A queue of lines logged by non-main threads. */
+	/** A queue of lines logged by non-primary threads. */
 	TDepletableMpscQueue<FOutputDeviceLine, FOutputDeviceLinearAllocator> BufferedLines;
 	uint8 BufferedLinesPadding[CalculateRedirectorCacheLinePadding(sizeof(BufferedLines))]{};
 
-	/** Array of output devices to redirect to from the main thread. */
+	/** Array of output devices to redirect to from the primary thread. */
 	TArray<FOutputDevice*> BufferedOutputDevices;
 
 	/** Array of output devices to redirect to from the calling thread. */
@@ -128,23 +128,23 @@ struct FOutputDeviceRedirectorState
 	TArray<FBufferedLine> BacklogLines;
 	FRWLock BacklogLock;
 
-	/** An optional dedicated main thread for logging to buffered output devices. */
+	/** An optional dedicated primary thread for logging to buffered output devices. */
 	FThread Thread;
 
 	/** A lock to synchronize access to the thread. */
 	FRWLock ThreadLock;
 
-	/** A queue of events to trigger when the dedicated main thread is idle. */
+	/** A queue of events to trigger when the dedicated primary thread is idle. */
 	TDepletableMpscQueue<FEvent*, FOutputDeviceLinearAllocator> ThreadIdleEvents;
 
-	/** An event to wake the dedicated main thread to process buffered lines. */
+	/** An event to wake the dedicated primary thread to process buffered lines. */
 	std::atomic<FEvent*> ThreadWakeEvent = nullptr;
 
-	/** The ID of the thread holding the main lock. */
+	/** The ID of the thread holding the primary lock. */
 	std::atomic<uint32> LockedThreadId = MAX_uint32;
 
-	/** The ID of the main logging thread. Logging from other threads will be buffered for processing by the main thread. */
-	std::atomic<uint32> MainLogThreadId = FPlatformTLS::GetCurrentThreadId();
+	/** The ID of the primary logging thread. Logging from other threads will be buffered for processing by the primary thread. */
+	std::atomic<uint32> PrimaryThreadId = FPlatformTLS::GetCurrentThreadId();
 
 	/** The ID of the panic thread, which is only set by Panic(). */
 	std::atomic<uint32> PanicThreadId = MAX_uint32;
@@ -156,9 +156,9 @@ struct FOutputDeviceRedirectorState
 	TBitArray<TInlineAllocator<1>> BufferedOutputDevicesCanBeUsedOnPanicThread;
 	TBitArray<TInlineAllocator<1>> UnbufferedOutputDevicesCanBeUsedOnPanicThread;
 
-	bool IsMainLogThread(const uint32 ThreadId) const
+	bool IsPrimaryThread(const uint32 ThreadId) const
 	{
-		return ThreadId == MainLogThreadId.load(std::memory_order_relaxed);
+		return ThreadId == PrimaryThreadId.load(std::memory_order_relaxed);
 	}
 
 	bool IsPanicThread(const uint32 ThreadId) const
@@ -212,7 +212,7 @@ struct FOutputDeviceRedirectorState
  * The read lock:
  * - Must be locked to read the OutputDevices arrays.
  * - Must be locked to write to unbuffered output devices.
- * - Must not be entered when the thread holds a write or main lock.
+ * - Must not be entered when the thread holds a write or primary lock.
  */
 class FOutputDevicesReadScopeLock
 {
@@ -251,9 +251,9 @@ private:
 /**
  * A scoped lock for writers of the OutputDevices arrays.
  *
- * The write lock has the same access as the main lock, and:
+ * The write lock has the same access as the primary lock, and:
  * - Must be locked to add or remove output devices.
- * - Must not be entered when the thread holds a read, write, or main lock.
+ * - Must not be entered when the thread holds a read, write, or primary lock.
  */
 class FOutputDevicesWriteScopeLock
 {
@@ -294,20 +294,20 @@ private:
 };
 
 /**
- * A scoped lock for readers of the OutputDevices arrays that need to access main logging thread's state.
+ * A scoped lock for exclusive access to the state of the primary log thread.
  *
- * The main lock has the same access as the read lock, and:
- * - Must not be entered when the thread holds a write lock or main lock.
+ * The primary lock has the same access as the read lock, and:
+ * - Must not be entered when the thread holds a write lock or primary lock.
  * - Must check IsLocked() before performing restricted operations.
  * - Must be locked to write to buffered output devices.
  * - Must be locked while calling FlushBufferedLines().
  * - May be locked when the thread holds a read lock.
- * - When a panic thread is active, may only be locked from the panic thread.
+ * - When a panic thread is active, locking will only succeed from the panic thread.
  */
-class FOutputDevicesMainLogScope
+class FOutputDevicesPrimaryScopeLock
 {
 public:
-	explicit FOutputDevicesMainLogScope(FOutputDeviceRedirectorState& InState)
+	explicit FOutputDevicesPrimaryScopeLock(FOutputDeviceRedirectorState& InState)
 		: State(InState)
 	{
 		const uint32 ThreadId = FPlatformTLS::GetCurrentThreadId();
@@ -337,7 +337,7 @@ public:
 		}
 	}
 
-	FORCEINLINE ~FOutputDevicesMainLogScope()
+	FORCEINLINE ~FOutputDevicesPrimaryScopeLock()
 	{
 		if (bLocked)
 		{
@@ -421,17 +421,17 @@ void FOutputDeviceRedirectorState::ThreadLoop()
 {
 	const uint32 ThreadId = FPlatformTLS::GetCurrentThreadId();
 
-	if (FOutputDevicesMainLogScope Lock(*this); Lock.IsLocked())
+	if (FOutputDevicesPrimaryScopeLock Lock(*this); Lock.IsLocked())
 	{
-		MainLogThreadId.store(ThreadId, std::memory_order_relaxed);
+		PrimaryThreadId.store(ThreadId, std::memory_order_relaxed);
 	}
 
 	while (FEvent* WakeEvent = ThreadWakeEvent.load(std::memory_order_acquire))
 	{
 		WakeEvent->Wait();
-		while (!BufferedLines.IsEmpty() && IsMainLogThread(ThreadId))
+		while (!BufferedLines.IsEmpty() && IsPrimaryThread(ThreadId))
 		{
-			if (FOutputDevicesMainLogScope Lock(*this); Lock.IsLocked())
+			if (FOutputDevicesPrimaryScopeLock Lock(*this); Lock.IsLocked())
 			{
 				FlushBufferedLines();
 			}
@@ -510,7 +510,7 @@ void FOutputDeviceRedirector::FlushThreadedLogs(EOutputDeviceRedirectorFlushOpti
 		return;
 	}
 
-	if (UE::Private::FOutputDevicesMainLogScope Lock(*State); Lock.IsLocked())
+	if (UE::Private::FOutputDevicesPrimaryScopeLock Lock(*State); Lock.IsLocked())
 	{
 		State->FlushBufferedLines();
 	}
@@ -535,24 +535,24 @@ void FOutputDeviceRedirector::EnableBacklog(bool bEnable)
 	}
 }
 
-void FOutputDeviceRedirector::SetCurrentThreadAsMainLogThread()
+void FOutputDeviceRedirector::SetCurrentThreadAsPrimaryThread()
 {
 	const uint32 ThreadId = FPlatformTLS::GetCurrentThreadId();
 
-	if (UE::Private::FOutputDevicesMainLogScope Lock(*State); !Lock.IsLocked() || State->MainLogThreadId.load(std::memory_order_relaxed) == ThreadId)
+	if (UE::Private::FOutputDevicesPrimaryScopeLock Lock(*State); !Lock.IsLocked() || State->PrimaryThreadId.load(std::memory_order_relaxed) == ThreadId)
 	{
 		return;
 	}
 	else
 	{
-		State->MainLogThreadId.store(ThreadId, std::memory_order_relaxed);
+		State->PrimaryThreadId.store(ThreadId, std::memory_order_relaxed);
 		State->FlushBufferedLines();
 	}
 
 	State->TryStopThread();
 }
 
-bool FOutputDeviceRedirector::TryStartDedicatedMainLogThread()
+bool FOutputDeviceRedirector::TryStartDedicatedPrimaryThread()
 {
 	return FApp::ShouldUseThreadingForPerformance() && State->TryStartThread();
 }
@@ -590,13 +590,13 @@ void FOutputDeviceRedirector::Serialize(const TCHAR* const Data, const ELogVerbo
 		State->BacklogLines.Emplace(Data, Category, Verbosity, RealTime);
 	}
 
-	// Serialize to buffered output devices from the main logging thread.
+	// Serialize to buffered output devices from the primary logging thread.
 	// Lines are queued until buffered output devices are added to avoid missing early log lines.
-	if (State->IsMainLogThread(ThreadId) && !State->BufferedOutputDevices.IsEmpty())
+	if (State->IsPrimaryThread(ThreadId) && !State->BufferedOutputDevices.IsEmpty())
 	{
-		// Verify that this is the main thread again because another thread may have become
-		// the main thread between the previous check and the lock.
-		if (UE::Private::FOutputDevicesMainLogScope MainLock(*State); MainLock.IsLocked() && State->IsMainLogThread(ThreadId))
+		// Verify that this is the primary thread again because another thread may have become
+		// the primary thread between the previous check and the lock.
+		if (UE::Private::FOutputDevicesPrimaryScopeLock PrimaryLock(*State); PrimaryLock.IsLocked() && State->IsPrimaryThread(ThreadId))
 		{
 			State->FlushBufferedLines();
 			State->BroadcastTo(ThreadId, State->BufferedOutputDevices, State->BufferedOutputDevicesCanBeUsedOnPanicThread,
@@ -610,7 +610,7 @@ void FOutputDeviceRedirector::Serialize(const TCHAR* const Data, const ELogVerbo
 		}
 	}
 
-	// Queue the line to serialize to buffered output devices from the main thread.
+	// Queue the line to serialize to buffered output devices from the primary thread.
 	if (State->BufferedLines.EnqueueAndReturnWasEmpty(Data, Category, Verbosity, RealTime))
 	{
 		if (FEvent* WakeEvent = State->ThreadWakeEvent.load(std::memory_order_acquire))
@@ -637,7 +637,7 @@ void FOutputDeviceRedirector::RedirectLog(const FLazyName& Category, ELogVerbosi
 
 void FOutputDeviceRedirector::Flush()
 {
-	if (UE::Private::FOutputDevicesMainLogScope Lock(*State); Lock.IsLocked())
+	if (UE::Private::FOutputDevicesPrimaryScopeLock Lock(*State); Lock.IsLocked())
 	{
 		State->FlushBufferedLines();
 		const uint32 ThreadId = FPlatformTLS::GetCurrentThreadId();
@@ -671,8 +671,8 @@ void FOutputDeviceRedirector::Panic()
 			FPlatformProcess::Yield();
 		}
 
-		// Make the panic thread the main thread. Neither thread can be changed after this point.
-		State->MainLogThreadId.exchange(ThreadId, std::memory_order_relaxed);
+		// Make the panic thread the primary thread. Neither thread can be changed after this point.
+		State->PrimaryThreadId.exchange(ThreadId, std::memory_order_relaxed);
 
 		// Flush. Every log from the panic thread after this point will also flush.
 		Flush();
@@ -686,7 +686,7 @@ void FOutputDeviceRedirector::Panic()
 
 void FOutputDeviceRedirector::TearDown()
 {
-	SetCurrentThreadAsMainLogThread();
+	SetCurrentThreadAsPrimaryThread();
 
 	Flush();
 
diff --git a/Engine/Source/Runtime/Core/Private/Unix/UnixPlatformProcess.cpp b/Engine/Source/Runtime/Core/Private/Unix/UnixPlatformProcess.cpp
--- a/Engine/Source/Runtime/Core/Private/Unix/UnixPlatformProcess.cpp
+++ b/Engine/Source/Runtime/Core/Private/Unix/UnixPlatformProcess.cpp
@@ -1531,7 +1531,7 @@ FGenericPlatformProcess::EWaitAndForkResult FUnixPlatformProcess::WaitAndFork()
 
 				// Start up the log again
 				FPlatformOutputDevices::SetupOutputDevices();
-				GLog->SetCurrentThreadAsMainLogThread();
+				GLog->SetCurrentThreadAsPrimaryThread();
 
 				// Set the process name, if specified
 				if (ChildIdx > 0)
diff --git a/Engine/Source/Runtime/Core/Public/Misc/OutputDeviceRedirector.h b/Engine/Source/Runtime/Core/Public/Misc/OutputDeviceRedirector.h
--- a/Engine/Source/Runtime/Core/Public/Misc/OutputDeviceRedirector.h
+++ b/Engine/Source/Runtime/Core/Public/Misc/OutputDeviceRedirector.h
@@ -57,7 +57,7 @@ enum class EOutputDeviceRedirectorFlushOptions : uint32
 	/**
 		* Flush asynchronously when possible.
 		*
-		* When this flag is set and there is a dedicated main logging thread, the flush function returns immediately.
+		* When this flag is set and there is a dedicated primary logging thread, the flush function returns immediately.
 		* Otherwise, the flush function does not return until the requested type of flush is complete.
 		*/
 	Async = 1 << 0,
@@ -122,17 +122,26 @@ public:
 	 */
 	void EnableBacklog(bool bEnable);
 
-	/** Sets the current thread to be the main thread, redirecting logs without buffering. */
-	void SetCurrentThreadAsMainLogThread();
+	/**
+	 * Sets the current thread to be the thread that redirects logs to buffered output devices.
+	 *
+	 * The current thread can redirect to buffered output devices without buffering, and becomes
+	 * responsible for flushing buffered logs from secondary threads. Logs from secondary threads
+	 * will not be redirected unless the current thread periodically flushes threaded logs.
+	 */
+	void SetCurrentThreadAsPrimaryThread();
+
+	UE_DEPRECATED(5.1, "Use SetCurrentThreadAsPrimaryThread().")
+	inline void SetCurrentThreadAsMasterThread() { SetCurrentThreadAsPrimaryThread(); }
 
 	/**
-	 * Starts a dedicated main thread that redirects logs that require buffering.
+	 * Starts a dedicated primary thread that redirects logs to buffered output devices.
 	 *
 	 * A thread will not be started for certain configurations or platforms, or when threading is disabled.
 	 *
-	 * @return true if a dedicated main logging thread is running, false otherwise.
+	 * @return true if a dedicated primary logging thread is running, false otherwise.
 	 */
-	bool TryStartDedicatedMainLogThread();
+	bool TryStartDedicatedPrimaryThread();
 
 	/**
 	 * Serializes the passed in data via all current output devices.
@@ -162,7 +171,7 @@ public:
 	 *
 	 * Only one thread can be the panic thread. Subsequent calls from other threads are ignored.
 	 * Only redirects logs to panic-safe output devices from this point forward.
-	 * Makes the calling thread the main log thread as well.
+	 * Makes the calling thread the primary log thread as well.
 	 * Flushes buffered logs to panic-safe output devices.
 	 * Flushes panic-safe output devices.
 	 */
diff --git a/Engine/Source/Runtime/Launch/Private/Android/LaunchAndroid.cpp b/Engine/Source/Runtime/Launch/Private/Android/LaunchAndroid.cpp
--- a/Engine/Source/Runtime/Launch/Private/Android/LaunchAndroid.cpp
+++ b/Engine/Source/Runtime/Launch/Private/Android/LaunchAndroid.cpp
@@ -544,7 +544,7 @@ int32 AndroidMain(struct android_app* state)
 
 	UE_LOG(LogAndroid, Display, TEXT("Passed PreInit()"));
 
-	GLog->SetCurrentThreadAsMainLogThread();
+	GLog->SetCurrentThreadAsPrimaryThread();
 
 	FAppEventManager::GetInstance()->SetEmptyQueueHandlerEvent(FPlatformProcess::GetSynchEventFromPool(false));
 
diff --git a/Engine/Source/Runtime/Launch/Private/IOS/LaunchIOS.cpp b/Engine/Source/Runtime/Launch/Private/IOS/LaunchIOS.cpp
--- a/Engine/Source/Runtime/Launch/Private/IOS/LaunchIOS.cpp
+++ b/Engine/Source/Runtime/Launch/Private/IOS/LaunchIOS.cpp
@@ -411,7 +411,7 @@ void FAppEntry::Init()
 	NSLog(@"%s", "Initializing ULD Communications in game mode\n");
 	GCommandSystem.Init();
 
-	GLog->SetCurrentThreadAsMainLogThread();
+	GLog->SetCurrentThreadAsPrimaryThread();
 	
 	// Send the launch local notification to the local notification service now that the engine module system has been initialized
 	if(gAppLaunchedWithLocalNotification)
diff --git a/Engine/Source/Runtime/Launch/Private/LaunchEngineLoop.cpp b/Engine/Source/Runtime/Launch/Private/LaunchEngineLoop.cpp
--- a/Engine/Source/Runtime/Launch/Private/LaunchEngineLoop.cpp
+++ b/Engine/Source/Runtime/Launch/Private/LaunchEngineLoop.cpp
@@ -1549,7 +1549,7 @@ int32 FEngineLoop::PreInitPreStartupScreen(const TCHAR* CmdLine)
 	// until initialization is far enough along to try to start a dedicated master thread.
 	if (GLog)
 	{
-		GLog->SetCurrentThreadAsMainLogThread();
+		GLog->SetCurrentThreadAsPrimaryThread();
 	}
 
 	// Command line option for enabling named events
@@ -1700,7 +1700,7 @@ int32 FEngineLoop::PreInitPreStartupScreen(const TCHAR* CmdLine)
 	// Try to start the dedicated master thread now that the command line is available.
 	if (!FParse::Param(FCommandLine::Get(), TEXT("NoLogThread")))
 	{
-		GLog->TryStartDedicatedMainLogThread();
+		GLog->TryStartDedicatedPrimaryThread();
 	}
 
 	// Initialize std out device as early as possible if requested in the command line
