Added more IWYU pragmas.

diff --git a/Engine/Source/Runtime/Core/Private/Misc/ConfigManifest.cpp b/Engine/Source/Runtime/Core/Private/Misc/ConfigManifest.cpp
--- a/Engine/Source/Runtime/Core/Private/Misc/ConfigManifest.cpp
+++ b/Engine/Source/Runtime/Core/Private/Misc/ConfigManifest.cpp
@@ -5,8 +5,8 @@
 #include "Containers/Map.h"
 #include "Containers/StringConv.h"
 #include "Containers/UnrealString.h"
+#include "GenericPlatform/GenericPlatformFile.h"
 #include "HAL/FileManager.h"
-#include "HAL/PlatformFile.h"
 #include "HAL/PlatformFileManager.h"
 #include "HAL/PlatformProcess.h"
 #include "HAL/PlatformProperties.h"
diff --git a/Engine/Source/Runtime/Core/Public/Clang/ClangPlatformCodeAnalysis.h b/Engine/Source/Runtime/Core/Public/Clang/ClangPlatformCodeAnalysis.h
--- a/Engine/Source/Runtime/Core/Public/Clang/ClangPlatformCodeAnalysis.h
+++ b/Engine/Source/Runtime/Core/Public/Clang/ClangPlatformCodeAnalysis.h
@@ -1,4 +1,5 @@
 // Copyright Epic Games, Inc. All Rights Reserved.
+// IWYU pragma: begin_exports
 
 #pragma once
 
@@ -130,3 +131,5 @@
 	#endif
 
 #endif
+
+// IWYU pragma: end_exports
diff --git a/Engine/Source/Runtime/Core/Public/FramePro/FramePro.h b/Engine/Source/Runtime/Core/Public/FramePro/FramePro.h
--- a/Engine/Source/Runtime/Core/Public/FramePro/FramePro.h
+++ b/Engine/Source/Runtime/Core/Public/FramePro/FramePro.h
@@ -1159,8 +1159,8 @@ namespace FramePro
 //------------------------------------------------------------------------
 
 //------------------------------------------------------------------------
-#include <stdarg.h>
-#include <time.h>
+#include <stdarg.h> // IWYU pragma: export
+#include <time.h> // IWYU pragma: export
 
 //------------------------------------------------------------------------
 //                         FRAMEPRO_PLATFORM_UE4
diff --git a/Engine/Source/Runtime/Core/Public/GenericPlatform/StandardPlatformString.h b/Engine/Source/Runtime/Core/Public/GenericPlatform/StandardPlatformString.h
--- a/Engine/Source/Runtime/Core/Public/GenericPlatform/StandardPlatformString.h
+++ b/Engine/Source/Runtime/Core/Public/GenericPlatform/StandardPlatformString.h
@@ -6,10 +6,13 @@
 #include "Misc/Char.h"
 #include "GenericPlatform/GenericPlatformStricmp.h"
 #include "GenericPlatform/GenericPlatformString.h"
+
+// IWYU pragma: begin_exports
 #include <wchar.h>
 #include <string.h>
 #include <stdlib.h>
 #include <stdio.h>
+// IWYU pragma: end_exports
 
 #if PLATFORM_USE_GENERIC_STRING_IMPLEMENTATION && !PLATFORM_TCHAR_IS_CHAR16
 
diff --git a/Engine/Source/Runtime/Core/Public/Linux/LinuxPlatformProcess.h b/Engine/Source/Runtime/Core/Public/Linux/LinuxPlatformProcess.h
--- a/Engine/Source/Runtime/Core/Public/Linux/LinuxPlatformProcess.h
+++ b/Engine/Source/Runtime/Core/Public/Linux/LinuxPlatformProcess.h
@@ -7,7 +7,7 @@
 #pragma once
 
 #include "HAL/Platform.h"
-#include "Unix/UnixPlatformProcess.h"
+#include "Unix/UnixPlatformProcess.h" // IWYU pragma: export
 
 /**
  * Linux implementation of the Process OS functions
diff --git a/Engine/Source/Runtime/Core/Public/Linux/LinuxSystemIncludes.h b/Engine/Source/Runtime/Core/Public/Linux/LinuxSystemIncludes.h
--- a/Engine/Source/Runtime/Core/Public/Linux/LinuxSystemIncludes.h
+++ b/Engine/Source/Runtime/Core/Public/Linux/LinuxSystemIncludes.h
@@ -1,4 +1,4 @@
 // Copyright Epic Games, Inc. All Rights Reserved.
 #pragma once
 
-#include "Unix/UnixSystemIncludes.h"
+#include "Unix/UnixSystemIncludes.h" // IWYU pragma: export
diff --git a/Engine/Source/Runtime/Core/Public/Math/GenericOctree.h b/Engine/Source/Runtime/Core/Public/Math/GenericOctree.h
--- a/Engine/Source/Runtime/Core/Public/Math/GenericOctree.h
+++ b/Engine/Source/Runtime/Core/Public/Math/GenericOctree.h
@@ -1816,6 +1816,6 @@ public:
 	}
 };
 
-#include "GenericOctree.inl"
+#include "GenericOctree.inl" // IWYU pragma: export
 
-#include <stddef.h>
+#include <stddef.h> // IWYU pragma: export
diff --git a/Engine/Source/Runtime/Core/Public/Math/UnrealMathSSE.h b/Engine/Source/Runtime/Core/Public/Math/UnrealMathSSE.h
--- a/Engine/Source/Runtime/Core/Public/Math/UnrealMathSSE.h
+++ b/Engine/Source/Runtime/Core/Public/Math/UnrealMathSSE.h
@@ -1,4 +1,5 @@
 // Copyright Epic Games, Inc. All Rights Reserved.
+// IWYU pragma: begin_exports
 
 #pragma once
 
@@ -3442,3 +3443,5 @@ FORCEINLINE VectorRegister4Int VectorFloatToInt(const VectorRegister4Double& A)
 #define VectorIntExpandLow16To32(V0)				_mm_unpacklo_epi16(V0, _mm_setzero_si128())
 
 #endif
+
+// IWYU pragma: end_exports
diff --git a/Engine/Source/Runtime/Core/Public/Misc/Launder.h b/Engine/Source/Runtime/Core/Public/Misc/Launder.h
--- a/Engine/Source/Runtime/Core/Public/Misc/Launder.h
+++ b/Engine/Source/Runtime/Core/Public/Misc/Launder.h
@@ -2,7 +2,7 @@
 #pragma once
 
 #if __cplusplus >= 201703L
-#include <new>
+#include <new> // IWYU pragma: export
 #define UE_LAUNDER(x) std::launder(x)
 #else
 #define UE_LAUNDER(x) (x)
diff --git a/Engine/Source/Runtime/Core/Public/Misc/LazySingleton.h b/Engine/Source/Runtime/Core/Public/Misc/LazySingleton.h
--- a/Engine/Source/Runtime/Core/Public/Misc/LazySingleton.h
+++ b/Engine/Source/Runtime/Core/Public/Misc/LazySingleton.h
@@ -2,7 +2,7 @@
 
 #pragma once
 
-#include <new>
+#include <new> // IWYU pragma: export
 
 /** Allows inline friend declaration without forward-declaring TLazySingleton */
 class FLazySingleton
diff --git a/Engine/Source/Runtime/Core/Public/Templates/Function.h b/Engine/Source/Runtime/Core/Public/Templates/Function.h
--- a/Engine/Source/Runtime/Core/Public/Templates/Function.h
+++ b/Engine/Source/Runtime/Core/Public/Templates/Function.h
@@ -17,7 +17,7 @@
 #include "Templates/IsPointer.h"
 #include "Templates/UnrealTemplate.h"
 #include "Math/UnrealMathUtility.h"
-#include <new>
+#include <new> // IWYU pragma: export
 
 
 // Disable visualization hack for shipping or test builds.
diff --git a/Engine/Source/Runtime/Core/Public/Templates/MemoryOps.h b/Engine/Source/Runtime/Core/Public/Templates/MemoryOps.h
--- a/Engine/Source/Runtime/Core/Public/Templates/MemoryOps.h
+++ b/Engine/Source/Runtime/Core/Public/Templates/MemoryOps.h
@@ -7,7 +7,7 @@
 #include "Templates/AreTypesEqual.h"
 #include "Templates/IsTriviallyCopyConstructible.h"
 #include "Templates/UnrealTypeTraits.h"
-#include <new>
+#include <new> // IWYU pragma: export
 
 #include "Templates/IsTriviallyCopyAssignable.h"
 #include "Templates/IsTriviallyDestructible.h"
diff --git a/Engine/Source/Runtime/Core/Public/Unix/UnixPlatformCrashContext.h b/Engine/Source/Runtime/Core/Public/Unix/UnixPlatformCrashContext.h
--- a/Engine/Source/Runtime/Core/Public/Unix/UnixPlatformCrashContext.h
+++ b/Engine/Source/Runtime/Core/Public/Unix/UnixPlatformCrashContext.h
@@ -6,8 +6,11 @@
 #include "Containers/UnrealString.h"
 #include "CoreTypes.h"
 #include "GenericPlatform/GenericPlatformCrashContext.h"
+#include "Unix/UnixSystemIncludes.h" // IWYU pragma: export
 #include "Templates/Atomic.h"
 
+#include <ucontext.h> // IWYU pragma: export
+
 /** Passed in through sigqueue for gathering of a callstack from a signal */
 struct ThreadStackUserData
 {
diff --git a/Engine/Source/Runtime/Core/Public/Unix/UnixPlatformMath.h b/Engine/Source/Runtime/Core/Public/Unix/UnixPlatformMath.h
--- a/Engine/Source/Runtime/Core/Public/Unix/UnixPlatformMath.h
+++ b/Engine/Source/Runtime/Core/Public/Unix/UnixPlatformMath.h
@@ -1,4 +1,5 @@
 // Copyright Epic Games, Inc. All Rights Reserved.
+// IWYU pragma: begin_exports
 
 /*=============================================================================================
 	LinuxPlatformMath.h: Linux platform Math functions
@@ -36,3 +37,4 @@ struct FLinuxPlatformMath : public TUnrealPlatformMathSSE4Base<FClangPlatformMat
 
 typedef FLinuxPlatformMath FPlatformMath;
 
+// IWYU pragma: end_exports
diff --git a/Engine/Source/Runtime/Core/Public/Unix/UnixPlatformProcess.h b/Engine/Source/Runtime/Core/Public/Unix/UnixPlatformProcess.h
--- a/Engine/Source/Runtime/Core/Public/Unix/UnixPlatformProcess.h
+++ b/Engine/Source/Runtime/Core/Public/Unix/UnixPlatformProcess.h
@@ -11,6 +11,7 @@
 #include "CoreTypes.h"
 #include "GenericPlatform/GenericPlatformProcess.h" // IWYU pragma: export
 #include "HAL/PlatformMisc.h"
+#include "Unix/UnixSystemIncludes.h" // IWYU pragma: export
 #include "Misc/AssertionMacros.h"
 
 class Error;
diff --git a/Engine/Source/Runtime/Core/Public/Unix/UnixPlatformTime.h b/Engine/Source/Runtime/Core/Public/Unix/UnixPlatformTime.h
--- a/Engine/Source/Runtime/Core/Public/Unix/UnixPlatformTime.h
+++ b/Engine/Source/Runtime/Core/Public/Unix/UnixPlatformTime.h
@@ -9,6 +9,8 @@
 #include "CoreTypes.h"
 #include "GenericPlatform/GenericPlatformTime.h"
 
+#include <time.h> // IWYU pragma: export
+
 /**
  * Unix implementation of the Time OS functions
  */
diff --git a/Engine/Source/Runtime/Core/Public/Unix/UnixSignalHeartBeat.h b/Engine/Source/Runtime/Core/Public/Unix/UnixSignalHeartBeat.h
--- a/Engine/Source/Runtime/Core/Public/Unix/UnixSignalHeartBeat.h
+++ b/Engine/Source/Runtime/Core/Public/Unix/UnixSignalHeartBeat.h
@@ -4,7 +4,7 @@
 #include "CoreGlobals.h"
 #include "CoreTypes.h"
 
-#include <time.h>
+#include <time.h> // IWYU pragma: export
 
 class CORE_API FUnixSignalGameHitchHeartBeat
 {
diff --git a/Engine/Source/Runtime/Core/Public/Unix/UnixSystemIncludes.h b/Engine/Source/Runtime/Core/Public/Unix/UnixSystemIncludes.h
--- a/Engine/Source/Runtime/Core/Public/Unix/UnixSystemIncludes.h
+++ b/Engine/Source/Runtime/Core/Public/Unix/UnixSystemIncludes.h
@@ -1,4 +1,6 @@
 // Copyright Epic Games, Inc. All Rights Reserved.
+// IWYU pragma: begin_exports
+
 #pragma once
 
 // Set up compiler pragmas, etc
@@ -53,3 +55,5 @@ typedef struct tagRECT RECT;
 #if defined(__GLIBC__) && (__GLIBC__ == 2) && (__GLIBC_MINOR__ < 17)
 #  define secure_getenv __secure_getenv
 #endif
+
+// IWYU pragma: end_exports
diff --git a/Engine/Source/Runtime/Core/Tests/Containers/TestUtils.cpp b/Engine/Source/Runtime/Core/Tests/Containers/TestUtils.cpp
--- a/Engine/Source/Runtime/Core/Tests/Containers/TestUtils.cpp
+++ b/Engine/Source/Runtime/Core/Tests/Containers/TestUtils.cpp
@@ -2,8 +2,6 @@
 
 #include "TestUtils.h"
 
-#include "HAL/Platform.h"
-
 #include <numeric>
 
 namespace Test
