Reworked GetVideoDriverDetailsFromSetup logging statements to no longer be warnings and to provide additional context. Failure to find driver information will now fall back to the old detection method.

diff --git a/Engine/Source/Runtime/Core/Private/Windows/WindowsPlatformMisc.cpp b/Engine/Source/Runtime/Core/Private/Windows/WindowsPlatformMisc.cpp
--- a/Engine/Source/Runtime/Core/Private/Windows/WindowsPlatformMisc.cpp
+++ b/Engine/Source/Runtime/Core/Private/Windows/WindowsPlatformMisc.cpp
@@ -2487,7 +2487,7 @@ static void GetVideoDriverDetailsFromSetup(const FString& DeviceName, FGPUDriver
 					}
 					else
 					{
-						UE_LOG(LogWindows, Warning, TEXT("Failed to retrieve driver registry key for device %d"), Idx);
+						UE_LOG(LogWindows, Log, TEXT("Failed to retrieve driver registry key for device %d"), Idx);
 					}
 					
 					break;
@@ -2496,7 +2496,7 @@ static void GetVideoDriverDetailsFromSetup(const FString& DeviceName, FGPUDriver
 			}
 			else
 			{
-				UE_LOG(LogWindows, Warning, TEXT("Failed to retrieve driver description for device %d"), Idx);
+				UE_LOG(LogWindows, Log, TEXT("Failed to retrieve driver description for device %d"), Idx);
 			}
 		}
 
@@ -2511,7 +2511,7 @@ static void GetVideoDriverDetailsFromSetup(const FString& DeviceName, FGPUDriver
 			}
 			else
 			{
-				UE_LOG(LogWindows, Warning, TEXT("Failed to find provider name"));
+				UE_LOG(LogWindows, Log, TEXT("Failed to find provider name"));
 			}
 			// Get the internal driver version
 			if (SetupDiGetDeviceProperty(hDevInfo, &DeviceInfoData, &DEVPKEY_Device_DriverVersion, &DataType,
@@ -2522,7 +2522,7 @@ static void GetVideoDriverDetailsFromSetup(const FString& DeviceName, FGPUDriver
 			}
 			else
 			{
-				UE_LOG(LogWindows, Warning, TEXT("Failed to find internal driver version"));
+				UE_LOG(LogWindows, Log, TEXT("Failed to find internal driver version"));
 			}
 			// Get the driver date
 			FILETIME FileTime;
@@ -2535,19 +2535,19 @@ static void GetVideoDriverDetailsFromSetup(const FString& DeviceName, FGPUDriver
 			}
 			else
 			{
-				UE_LOG(LogWindows, Warning, TEXT("Failed to find driver date"));
+				UE_LOG(LogWindows, Log, TEXT("Failed to find driver date"));
 			}
 		}
 		else
 		{
-			UE_LOG(LogWindows, Warning, TEXT("Unable to find requested device '%s' using Setup API"), *DeviceName);
+			UE_LOG(LogWindows, Log, TEXT("Unable to find requested device '%s' using Setup API."), *DeviceName);
 		}
 		
 		SetupDiDestroyDeviceInfoList(hDevInfo);
 	}
 	else
 	{
-		UE_LOG(LogWindows, Warning, TEXT("Failed to initialize Setup API"));
+		UE_LOG(LogWindows, Log, TEXT("Failed to initialize Setup API"));
 	}
 
 	if (!Out.ProviderName.IsEmpty())
@@ -2735,7 +2735,8 @@ FGPUDriverInfo FWindowsPlatformMisc::GetGPUDriverInfo(const FString& DeviceDescr
 			return Local;
 		}
 
-		return Ret;
+		UE_LOG(LogWindows, Log, TEXT("Failed to get driver data for device '%s' using Setup API. Switching to fallback method."), *DeviceDescription);
+		Method = 4; // Switch to method 4 as a fallback if method 5 fails
 	}
 	
 	if(Method == 3 || Method == 4)
