Some UTF-8 conversion tests for invalid encodings.
Fix for some non-UTF-8 codepoints inside a TCHAR sting literal.

diff --git a/Engine/Source/Runtime/Core/Private/Tests/Misc/TCharToUTF8.cpp b/Engine/Source/Runtime/Core/Private/Tests/Misc/TCharToUTF8.cpp
--- a/Engine/Source/Runtime/Core/Private/Tests/Misc/TCharToUTF8.cpp
+++ b/Engine/Source/Runtime/Core/Private/Tests/Misc/TCharToUTF8.cpp
@@ -100,9 +100,89 @@ bool FTCharToUTF8Test::RunTest(const FString& Parameters)
 		}
 	}
 
+	{
+		// Bad UTF-16 surrogates
+		{
+			const FTCHARToUTF8 EndWithHighSurrogate(UTF16_TO_TCHAR(u"Hello\xD83D"));
+			TestTrue(TEXT("UTF-16 string ending with a single high surrogate converted to UTF-8"), FCStringUtf8::Strcmp((const UTF8CHAR*)EndWithHighSurrogate.Get(), (const UTF8CHAR*)"Hello?") == 0);
+		}
+		{
+			const FTCHARToUTF8 EndWithTwoHighSurrogates(UTF16_TO_TCHAR(u"Hello\xD83D\xD83D"));
+			TestTrue(TEXT("UTF-16 string ending with two high surrogates converted to UTF-8"), FCStringUtf8::Strcmp((const UTF8CHAR*)EndWithTwoHighSurrogates.Get(), (const UTF8CHAR*)"Hello??") == 0);
+		}
+		{
+			const FTCHARToUTF8 EndWithThreeHighSurrogates(UTF16_TO_TCHAR(u"Hello\xD83D\xD83D\xD83D"));
+			TestTrue(TEXT("UTF-16 string ending with three high surrogates converted to UTF-8"), FCStringUtf8::Strcmp((const UTF8CHAR*)EndWithThreeHighSurrogates.Get(), (const UTF8CHAR*)"Hello???") == 0);
+		}
+		{
+			const FTCHARToUTF8 EndWithHighSurrogateAndPair(UTF16_TO_TCHAR(u"Hello\xD83D\xD83D\xDC69"));
+			TestTrue(TEXT("UTF-16 string ending with a high surrogate and a surrogate pair converted to UTF-8"), FCStringUtf8::Strcmp((const UTF8CHAR*)EndWithHighSurrogateAndPair.Get(), (const UTF8CHAR*)"Hello?\xf0\x9f\x91\xa9") == 0);
+		}
+		{
+			const FTCHARToUTF8 EndWithLowSurrogate(UTF16_TO_TCHAR(u"Hello\xDC69"));
+			TestTrue(TEXT("UTF-16 string ending with a single low surrogate converted to UTF-8"), FCStringUtf8::Strcmp((const UTF8CHAR*)EndWithLowSurrogate.Get(), (const UTF8CHAR*)"Hello?") == 0);
+		}
+		{
+			const FTCHARToUTF8 EndWithTwoLowSurrogates(UTF16_TO_TCHAR(u"Hello\xDC69\xDC69"));
+			TestTrue(TEXT("UTF-16 string ending with two low surrogates converted to UTF-8"), FCStringUtf8::Strcmp((const UTF8CHAR*)EndWithTwoLowSurrogates.Get(), (const UTF8CHAR*)"Hello??") == 0);
+		}
+		{
+			const FTCHARToUTF8 EndWithThreeLowSurrogates(UTF16_TO_TCHAR(u"Hello\xDC69\xDC69\xDC69"));
+			TestTrue(TEXT("UTF-16 string ending with three low surrogates converted to UTF-8"), FCStringUtf8::Strcmp((const UTF8CHAR*)EndWithThreeLowSurrogates.Get(), (const UTF8CHAR*)"Hello???") == 0);
+		}
+		{
+			const FTCHARToUTF8 EndWithLowSurrogateAndPair(UTF16_TO_TCHAR(u"Hello\xDC69\xD83D\xDC69"));
+			TestTrue(TEXT("UTF-16 string ending with a low surrogate and a surrogate pair converted to UTF-8"), FCStringUtf8::Strcmp((const UTF8CHAR*)EndWithLowSurrogateAndPair.Get(), (const UTF8CHAR*)"Hello?\xf0\x9f\x91\xa9") == 0);
+		}
+		{
+			const FTCHARToUTF8 MidHighSurrogate(UTF16_TO_TCHAR(u"Hello\xD83DWorld"));
+			TestTrue(TEXT("UTF-16 string containing a single high surrogate converted to UTF-8"), FCStringUtf8::Strcmp((const UTF8CHAR*)MidHighSurrogate.Get(), (const UTF8CHAR*)"Hello?World") == 0);
+		}
+		{
+			const FTCHARToUTF8 MidTwoHighSurrogates(UTF16_TO_TCHAR(u"Hello\xD83D\xD83DWorld"));
+			TestTrue(TEXT("UTF-16 string containing two high surrogates converted to UTF-8"), FCStringUtf8::Strcmp((const UTF8CHAR*)MidTwoHighSurrogates.Get(), (const UTF8CHAR*)"Hello??World") == 0);
+		}
+		{
+			const FTCHARToUTF8 MidThreeHighSurrogates(UTF16_TO_TCHAR(u"Hello\xD83D\xD83D\xD83DWorld"));
+			TestTrue(TEXT("UTF-16 string containing three high surrogates converted to UTF-8"), FCStringUtf8::Strcmp((const UTF8CHAR*)MidThreeHighSurrogates.Get(), (const UTF8CHAR*)"Hello???World") == 0);
+		}
+		{
+			const FTCHARToUTF8 MidHighSurrogateAndPair(UTF16_TO_TCHAR(u"Hello\xD83D\xD83D\xDC69World"));
+			TestTrue(TEXT("UTF-16 string containing a high surrogate and a surrogate pair converted to UTF-8"), FCStringUtf8::Strcmp((const UTF8CHAR*)MidHighSurrogateAndPair.Get(), (const UTF8CHAR*)"Hello?\xf0\x9f\x91\xa9World") == 0);
+		}
+		{
+			const FTCHARToUTF8 MidLowSurrogate(UTF16_TO_TCHAR(u"Hello\xDC69World"));
+			TestTrue(TEXT("UTF-16 string containing a single low surrogate converted to UTF-8"), FCStringUtf8::Strcmp((const UTF8CHAR*)MidLowSurrogate.Get(), (const UTF8CHAR*)"Hello?World") == 0);
+		}
+		{
+			const FTCHARToUTF8 MidTwoLowSurrogates(UTF16_TO_TCHAR(u"Hello\xDC69\xDC69World"));
+			TestTrue(TEXT("UTF-16 string containing two low surrogates converted to UTF-8"), FCStringUtf8::Strcmp((const UTF8CHAR*)MidTwoLowSurrogates.Get(), (const UTF8CHAR*)"Hello??World") == 0);
+		}
+		{
+			const FTCHARToUTF8 MidThreeLowSurrogates(UTF16_TO_TCHAR(u"Hello\xDC69\xDC69\xDC69World"));
+			TestTrue(TEXT("UTF-16 string containing three low surrogates converted to UTF-8"), FCStringUtf8::Strcmp((const UTF8CHAR*)MidThreeLowSurrogates.Get(), (const UTF8CHAR*)"Hello???World") == 0);
+		}
+		{
+			const FTCHARToUTF8 MidLowSurrogateAndPair(UTF16_TO_TCHAR(u"Hello\xDC69\xD83D\xDC69World"));
+			TestTrue(TEXT("UTF-16 string containing a low surrogate and a surrogate pair converted to UTF-8"), FCStringUtf8::Strcmp((const UTF8CHAR*)MidLowSurrogateAndPair.Get(), (const UTF8CHAR*)"Hello?\xf0\x9f\x91\xa9World") == 0);
+		}
+	}
+
 	{
 		// Verify that we handle invalid UTF-16 strings that ends in half a surrogate pair w/o crashing
-		TCHAR EndWithIllegalSurrogatePair[] = TEXT("ab\xD800");
+#if PLATFORM_TCHAR_IS_UTF8CHAR
+#ifdef _MSC_VER
+		// We need to do this because MSVC doesn't encode numeric escape sequences properly in u8"" literals:
+		// https://developercommunity.visualstudio.com/t/hex-escape-codes-in-a-utf8-literal-are-t/225847
+		char EndWithIllegalSurrogatePairImpl[] = "ab\xED\xA0\x80";
+		const UTF8CHAR* EndWithIllegalSurrogatePair = (const UTF8CHAR*)EndWithIllegalSurrogatePairImpl;
+#else
+		const TCHAR* EndWithIllegalSurrogatePair = TEXT("ab\xED\xA0\x80"));
+#endif // #ifdef _MSC_VER
+#else
+		const TCHAR* EndWithIllegalSurrogatePair = TEXT("ab\xD800");
+#endif // #if PLATFORM_TCHAR_IS_UTF8CHAR
+
 		TestEqual("IllegalSurrogatePair", TCHAR_TO_UTF8(EndWithIllegalSurrogatePair), "ab?");
 	}
 
