Core :  fixing type conversion warnings with CString.h, RemoveIf.h and RandomShuffle.h

diff --git a/Engine/Source/Runtime/Core/Public/Algo/RandomShuffle.h b/Engine/Source/Runtime/Core/Public/Algo/RandomShuffle.h
--- a/Engine/Source/Runtime/Core/Public/Algo/RandomShuffle.h
+++ b/Engine/Source/Runtime/Core/Public/Algo/RandomShuffle.h
@@ -23,7 +23,7 @@ namespace Algo
 		for (SizeType Index = 0; Index < Num - 1; ++Index)
 		{
 			// Get a random integer in [Index, Num)
-			const SizeType RandomIndex = Index + FMath::RandHelper64(Num - Index);
+			const SizeType RandomIndex = Index + (SizeType)FMath::RandHelper64(Num - Index);
 			if (RandomIndex != Index)
 			{
 				Swap(Data[Index], Data[RandomIndex]);
diff --git a/Engine/Source/Runtime/Core/Public/Algo/RemoveIf.h b/Engine/Source/Runtime/Core/Public/Algo/RemoveIf.h
--- a/Engine/Source/Runtime/Core/Public/Algo/RemoveIf.h
+++ b/Engine/Source/Runtime/Core/Public/Algo/RemoveIf.h
@@ -3,6 +3,7 @@
 #pragma once
 
 #include "Templates/Invoke.h"
+#include "Misc/CoreMiscDefines.h"
 
 namespace Algo
 {
@@ -31,7 +32,7 @@ namespace Algo
 			{
 				if (IterStart == IterEnd)
 				{
-					return IterStart - First;
+					return UE_PTRDIFF_TO_INT32(IterStart - First);
 				}
 
 				if (Invoke(Pred, *IterStart))
@@ -54,7 +55,7 @@ namespace Algo
 
 				if (IterStart == IterEnd)
 				{
-					return IterStart - First;
+					return UE_PTRDIFF_TO_INT32(IterStart - First);
 				}
 			}
 
@@ -88,7 +89,7 @@ namespace Algo
 		{
 			if (IterStart == Last)
 			{
-				return IterStart - First;
+				return UE_PTRDIFF_TO_INT32(IterStart - First);
 			}
 
 			if (Invoke(Pred, *IterStart))
@@ -106,7 +107,7 @@ namespace Algo
 		{
 			if (IterKeep == Last)
 			{
-				return IterStart - First;
+				return UE_PTRDIFF_TO_INT32(IterStart - First);
 			}
 
 			if (!Invoke(Pred, *IterKeep))
diff --git a/Engine/Source/Runtime/Core/Public/Misc/CString.h b/Engine/Source/Runtime/Core/Public/Misc/CString.h
--- a/Engine/Source/Runtime/Core/Public/Misc/CString.h
+++ b/Engine/Source/Runtime/Core/Public/Misc/CString.h
@@ -802,14 +802,14 @@ int32 TCString<T>::Strcspn( const CharType* String, const CharType* Mask )
 		{
 			if (*StringIt == *MaskIt)
 			{
-				return StringIt - String;
+				return UE_PTRDIFF_TO_INT32(StringIt - String);
 			}
 		}
 
 		++StringIt;
 	}
 
-	return StringIt - String;
+	return UE_PTRDIFF_TO_INT32(StringIt - String);
 }
 
 template <typename T> FORCEINLINE 
