#ue5 - Change float versions of VectorTan(), VectorATan(), and VectorATan2() to wrap to FMath functions when not using SVML, due to some compiler differences with old SSE2 extension library we were previously using.

diff --git a/Engine/Source/Runtime/Core/Public/Math/UnrealMathSSE.h b/Engine/Source/Runtime/Core/Public/Math/UnrealMathSSE.h
--- a/Engine/Source/Runtime/Core/Public/Math/UnrealMathSSE.h
+++ b/Engine/Source/Runtime/Core/Public/Math/UnrealMathSSE.h
@@ -3176,7 +3176,13 @@ FORCEINLINE VectorRegister4Float VectorTan(const VectorRegister4Float& X)
 #if UE_PLATFORM_MATH_USE_SVML
 	return _mm_tan_ps(X);
 #else
-	return SSE::tan_ps(X);
+	//return SSE::tan_ps(X);
+	AlignedFloat4 Floats(X);
+	Floats[0] = FMath::Tan(Floats[0]);
+	Floats[1] = FMath::Tan(Floats[1]);
+	Floats[2] = FMath::Tan(Floats[2]);
+	Floats[3] = FMath::Tan(Floats[3]);
+	return Floats.ToVectorRegister();
 #endif
 }
 
@@ -3261,7 +3267,13 @@ FORCEINLINE VectorRegister4Float VectorATan(const VectorRegister4Float& X)
 #if UE_PLATFORM_MATH_USE_SVML
 	return _mm_atan_ps(X);
 #else
-	return SSE::atan_ps(X);
+	//return SSE::atan_ps(X);
+	AlignedFloat4 Floats(X);
+	Floats[0] = FMath::Atan(Floats[0]);
+	Floats[1] = FMath::Atan(Floats[1]);
+	Floats[2] = FMath::Atan(Floats[2]);
+	Floats[3] = FMath::Atan(Floats[3]);
+	return Floats.ToVectorRegister();
 #endif
 }
 
@@ -3286,7 +3298,14 @@ FORCEINLINE VectorRegister4Float VectorATan2(const VectorRegister4Float& Y, cons
 #if UE_PLATFORM_MATH_USE_SVML
 	return _mm_atan2_ps(Y, X);
 #else
-	return SSE::atan2_ps(Y, X);
+	//return SSE::atan2_ps(Y, X);
+	AlignedFloat4 FloatsY(Y);
+	AlignedFloat4 FloatsX(X);
+	FloatsY[0] = FMath::Atan2(FloatsY[0], FloatsX[0]);
+	FloatsY[1] = FMath::Atan2(FloatsY[1], FloatsX[1]);
+	FloatsY[2] = FMath::Atan2(FloatsY[2], FloatsX[2]);
+	FloatsY[3] = FMath::Atan2(FloatsY[3], FloatsX[3]);
+	return FloatsY.ToVectorRegister();
 #endif
 }
 
