CompactBinarySerialization: Add operator<< and LoadFromCompactBinary functions for FName, FString, and TArray.
Add them for TMap and TSet as well, but keep them private for now since the write does not sort the keys and hence has determinism problems if used in persistent storage.
#rnx

diff --git a/Engine/Source/Editor/UnrealEd/Private/Cooker/CompactBinaryTCP.h b/Engine/Source/Editor/UnrealEd/Private/Cooker/CompactBinaryTCP.h
--- a/Engine/Source/Editor/UnrealEd/Private/Cooker/CompactBinaryTCP.h
+++ b/Engine/Source/Editor/UnrealEd/Private/Cooker/CompactBinaryTCP.h
@@ -6,7 +6,8 @@
 #include "Containers/ArrayView.h"
 #include "Memory/SharedBuffer.h"
 #include "Misc/Guid.h"
-#include "Serialization/CompactBinary.h"
+#include "Serialization/CompactBinarySerialization.h"
+#include "Serialization/CompactBinaryWriter.h"
 
 class FSocket;
 
@@ -126,3 +127,77 @@ private:
 };
 
 }
+
+template <typename KeyType, typename KeyFuncs, typename Allocator>
+inline bool LoadFromCompactBinary(FCbFieldView Field, TSet<KeyType, KeyFuncs, Allocator>& OutValue)
+{
+	OutValue.Reset();
+	OutValue.Reserve(Field.AsArrayView().Num());
+	bool bOk = !Field.HasError();
+	for (const FCbFieldView& ElementField : Field)
+	{
+		KeyType Key;
+		if (LoadFromCompactBinary(ElementField, Key))
+		{
+			OutValue.Add(MoveTemp(Key));
+		}
+		else
+		{
+			bOk = false;
+		}
+	}
+	return bOk;
+}
+
+template <typename KeyType, typename KeyFuncs, typename Allocator,
+	std::void_t<decltype(std::declval<FCbWriter&>() << std::declval<const KeyType&>())>* = nullptr>
+inline FCbWriter& operator<<(FCbWriter& Writer, const TSet<KeyType, KeyFuncs, Allocator>& Value)
+{
+	Writer.BeginArray();
+	for (const KeyType& Key : Value)
+	{
+		Writer << Key;
+	}
+	Writer.EndArray();
+	return Writer;
+}
+
+template <typename KeyType, typename ValueType, typename SetAllocator, typename KeyFuncs>
+inline bool LoadFromCompactBinary(FCbFieldView Field, TMap<KeyType, ValueType, SetAllocator, KeyFuncs>& OutValue)
+{
+	OutValue.Reset();
+	OutValue.Reserve(Field.AsArrayView().Num());
+	bool bOk = !Field.HasError();
+	for (const FCbFieldView& PairField : Field)
+	{
+		bOk &= PairField.IsObject();
+		KeyType Key;
+		if (LoadFromCompactBinary(PairField["K"], Key))
+		{
+			ValueType& Value = OutValue.FindOrAdd(MoveTemp(Key));
+			bOk = LoadFromCompactBinary(PairField["V"], Value) & bOk;
+		}
+		else
+		{
+			bOk = false;
+		}
+	}
+	return bOk;
+}
+
+template <typename KeyType, typename ValueType, typename SetAllocator, typename KeyFuncs,
+	std::void_t<decltype(std::declval<FCbWriter&>() << std::declval<const KeyType&>())>* = nullptr,
+	std::void_t<decltype(std::declval<FCbWriter&>() << std::declval<const ValueType&>())>* = nullptr>
+inline FCbWriter& operator<<(FCbWriter& Writer, const TMap<KeyType, ValueType, SetAllocator, KeyFuncs>& Value)
+{
+	Writer.BeginArray();
+	for (const auto& Pair : Value)
+	{
+		Writer.BeginObject();
+		Writer << "K" << Pair.Key;
+		Writer << "V" << Pair.Value;
+		Writer.EndObject();
+	}
+	Writer.EndArray();
+	return Writer;
+}
diff --git a/Engine/Source/Runtime/Core/Private/Serialization/CompactBinarySerialization.cpp b/Engine/Source/Runtime/Core/Private/Serialization/CompactBinarySerialization.cpp
--- a/Engine/Source/Runtime/Core/Private/Serialization/CompactBinarySerialization.cpp
+++ b/Engine/Source/Runtime/Core/Private/Serialization/CompactBinarySerialization.cpp
@@ -25,6 +25,7 @@
 #include "Templates/Invoke.h"
 #include "Templates/RemoveReference.h"
 #include "Templates/UnrealTemplate.h"
+#include "UObject/NameTypes.h"
 
 ///////////////////////////////////////////////////////////////////////////////////////////////////////////////////////
 
@@ -287,6 +288,18 @@ bool LoadFromCompactBinary(FCbFieldView Field, FWideStringBuilderBase& OutValue)
 	return !Field.HasError();
 }
 
+bool LoadFromCompactBinary(FCbFieldView Field, FString& OutValue)
+{
+	OutValue = FString(Field.AsString());
+	return !Field.HasError();
+}
+
+bool LoadFromCompactBinary(FCbFieldView Field, FName& OutValue)
+{
+	OutValue = FName(Field.AsString());
+	return !Field.HasError();
+}
+
 bool LoadFromCompactBinary(FCbFieldView Field, FGuid& OutValue)
 {
 	OutValue = Field.AsUuid();
diff --git a/Engine/Source/Runtime/Core/Private/Serialization/CompactBinaryWriter.cpp b/Engine/Source/Runtime/Core/Private/Serialization/CompactBinaryWriter.cpp
--- a/Engine/Source/Runtime/Core/Private/Serialization/CompactBinaryWriter.cpp
+++ b/Engine/Source/Runtime/Core/Private/Serialization/CompactBinaryWriter.cpp
@@ -19,6 +19,7 @@
 #include "Serialization/CompactBinaryPackage.h"
 #include "Serialization/CompactBinarySerialization.h"
 #include "Serialization/VarInt.h"
+#include "UObject/NameTypes.h"
 
 enum class FCbWriter::EStateFlags : uint8
 {
@@ -628,4 +629,10 @@ FCbWriter& operator<<(FCbWriter& Writer, const FTimespan Value)
 	return Writer;
 }
 
+FCbWriter& operator<<(FCbWriter& Writer, FName Value)
+{
+	Writer << WriteToUtf8String<FName::StringBufferSize>(Value).ToView();
+	return Writer;
+}
+
 ///////////////////////////////////////////////////////////////////////////////////////////////////////////////////////
diff --git a/Engine/Source/Runtime/Core/Private/UObject/UnrealNames.cpp b/Engine/Source/Runtime/Core/Private/UObject/UnrealNames.cpp
--- a/Engine/Source/Runtime/Core/Private/UObject/UnrealNames.cpp
+++ b/Engine/Source/Runtime/Core/Private/UObject/UnrealNames.cpp
@@ -2282,7 +2282,7 @@ void FNameEntry::AppendNameToString(FString& Out) const
 	Out.Append(EntryToCString(*this, Temp), Header.Len);
 }
 
-void FNameEntry::AppendNameToString(FStringBuilderBase& Out) const
+void FNameEntry::AppendNameToString(FWideStringBuilderBase& Out) const
 {
 	checkName(Header.Len != 0);
 	const int32 Offset = Out.AddUninitialized(Header.Len);
@@ -2299,6 +2299,19 @@ void FNameEntry::AppendNameToString(FStringBuilderBase& Out) const
 	}
 }
 
+void FNameEntry::AppendNameToString(FUtf8StringBuilderBase& Out) const
+{
+	checkName(Header.Len != 0);
+	if (Header.bIsWide)
+	{
+		Out << FWideStringView(WideName, Header.Len);
+	}
+	else
+	{
+		Out << FAnsiStringView(AnsiName, Header.Len);
+	}
+}
+
 void FNameEntry::AppendAnsiNameToString(FAnsiStringBuilderBase& Out) const
 {
 	checkName(Header.Len != 0);
@@ -3338,7 +3351,13 @@ void FName::ToString(FString& Out) const
 #endif // UE_FNAME_OUTLINE_NUMBER
 }
 
-void FName::ToString(FStringBuilderBase& Out) const
+void FName::ToString(FWideStringBuilderBase& Out) const
+{
+	Out.Reset();
+	AppendString(Out);
+}
+
+void FName::ToString(FUtf8StringBuilderBase& Out) const
 {
 	Out.Reset();
 	AppendString(Out);
@@ -3399,10 +3418,11 @@ void FName::AppendString(FString& Out) const
 	}
 }
 
-void FName::AppendString(FStringBuilderBase& Out) const
+template <typename StringBuilderType>
+void FName::AppendStringInternal(StringBuilderType& Out) const
 {
 #if UE_FNAME_OUTLINE_NUMBER
-	FNameEntryId Id = GetDisplayIndexInternal();	
+	FNameEntryId Id = GetDisplayIndexInternal();
 	const FNameEntry* ThisNameEntry = ResolveEntry(GetDisplayIndexInternal());
 
 	if (ThisNameEntry->Header.Len != 0)
@@ -3414,7 +3434,7 @@ void FName::AppendString(FStringBuilderBase& Out) const
 	{
 		const FNameEntry* BaseEntry = ResolveEntry(ThisNameEntry->NumberedName.Id);
 		BaseEntry->AppendNameToString(Out);
-		Out << TEXT('_') << NAME_INTERNAL_TO_EXTERNAL(ThisNameEntry->GetNumber());
+		Out << '_' << NAME_INTERNAL_TO_EXTERNAL(ThisNameEntry->GetNumber());
 	}
 
 #else // UE_FNAME_OUTLINE_NUMBER
@@ -3423,11 +3443,21 @@ void FName::AppendString(FStringBuilderBase& Out) const
 	const int32 InternalNumber = GetNumber();
 	if (InternalNumber != NAME_NO_NUMBER_INTERNAL)
 	{
-		Out << TEXT('_') << NAME_INTERNAL_TO_EXTERNAL(InternalNumber);
+		Out << '_' << NAME_INTERNAL_TO_EXTERNAL(InternalNumber);
 	}
 #endif // UE_FNAME_OUTLINE_NUMBER
 }
 
+void FName::AppendString(FWideStringBuilderBase& Out) const
+{
+	AppendStringInternal(Out);
+}
+
+void FName::AppendString(FUtf8StringBuilderBase& Out) const
+{
+	AppendStringInternal(Out);
+}
+
 bool FName::TryAppendAnsiString(FAnsiStringBuilderBase& Out) const
 {
 	const FNameEntry* const NameEntry = GetDisplayNameEntry();
@@ -3543,7 +3573,13 @@ bool FName::IsValidXName(const FStringView& InName, const FString& InInvalidChar
 	return true;
 }
 
-FStringBuilderBase& operator<<(FStringBuilderBase& Builder, FNameEntryId Id)
+FWideStringBuilderBase& operator<<(FWideStringBuilderBase& Builder, FNameEntryId Id)
+{
+	FName::GetEntry(Id)->AppendNameToString(Builder);
+	return Builder;
+}
+
+FUtf8StringBuilderBase& operator<<(FUtf8StringBuilderBase& Builder, FNameEntryId Id)
 {
 	FName::GetEntry(Id)->AppendNameToString(Builder);
 	return Builder;
diff --git a/Engine/Source/Runtime/Core/Public/Serialization/CompactBinarySerialization.h b/Engine/Source/Runtime/Core/Public/Serialization/CompactBinarySerialization.h
--- a/Engine/Source/Runtime/Core/Public/Serialization/CompactBinarySerialization.h
+++ b/Engine/Source/Runtime/Core/Public/Serialization/CompactBinarySerialization.h
@@ -2,6 +2,7 @@
 
 #pragma once
 
+#include "Containers/ContainersFwd.h"
 #include "Containers/StringFwd.h"
 #include "CoreTypes.h"
 #include "IO/IoHash.h"
@@ -12,6 +13,7 @@
 #include "Templates/Function.h"
 
 class FArchive;
+class FName;
 struct FGuid;
 
 ///////////////////////////////////////////////////////////////////////////////////////////////////////////////////////
@@ -92,6 +94,8 @@ CORE_API FArchive& operator<<(FArchive& Ar, FCbObject& Object);
 
 CORE_API bool LoadFromCompactBinary(FCbFieldView Field, FUtf8StringBuilderBase& OutValue);
 CORE_API bool LoadFromCompactBinary(FCbFieldView Field, FWideStringBuilderBase& OutValue);
+CORE_API bool LoadFromCompactBinary(FCbFieldView Field, FString& OutValue);
+CORE_API bool LoadFromCompactBinary(FCbFieldView Field, FName& OutValue);
 
 inline bool LoadFromCompactBinary(FCbFieldView Field, int8& OutValue, const int8 Default = 0)
 {
@@ -174,6 +178,18 @@ inline bool LoadFromCompactBinary(FCbFieldView Field, FCbObjectId& OutValue, con
 	return !Field.HasError();
 }
 
+template <typename T, typename Allocator>
+inline bool LoadFromCompactBinary(FCbFieldView Field, TArray<T, Allocator>& OutValue)
+{
+	OutValue.Reset(Field.AsArrayView().Num());
+	bool bOk = !Field.HasError();
+	for (const FCbFieldView& ElementField : Field)
+	{
+		bOk = LoadFromCompactBinary(ElementField, OutValue.Emplace_GetRef()) & bOk;
+	}
+	return bOk;
+}
+
 ///////////////////////////////////////////////////////////////////////////////////////////////////////////////////////
 
 /** Convert the object to JSON in a multi-line format with indentation. */
diff --git a/Engine/Source/Runtime/Core/Public/Serialization/CompactBinaryWriter.h b/Engine/Source/Runtime/Core/Public/Serialization/CompactBinaryWriter.h
--- a/Engine/Source/Runtime/Core/Public/Serialization/CompactBinaryWriter.h
+++ b/Engine/Source/Runtime/Core/Public/Serialization/CompactBinaryWriter.h
@@ -19,6 +19,7 @@ class FSharedBuffer;
 
 class FArchive;
 class FCbAttachment;
+class FName;
 struct FDateTime;
 struct FGuid;
 struct FIoHash;
@@ -460,4 +461,19 @@ inline FCbWriter& operator<<(FCbWriter& Writer, const FCbObjectId& Value)
 	return Writer;
 }
 
+CORE_API FCbWriter& operator<<(FCbWriter& Writer, FName Value);
+
+template <typename T, typename Allocator,
+	std::void_t<decltype(std::declval<FCbWriter&>() << std::declval<const T&>())>* = nullptr>
+inline FCbWriter& operator<<(FCbWriter& Writer, const TArray<T, Allocator>& Value)
+{
+	Writer.BeginArray();
+	for (const T& Element : Value)
+	{
+		Writer << Element;
+	}
+	Writer.EndArray();
+	return Writer;
+}
+
 ///////////////////////////////////////////////////////////////////////////////////////////////////////////////////////
diff --git a/Engine/Source/Runtime/Core/Public/UObject/NameTypes.h b/Engine/Source/Runtime/Core/Public/UObject/NameTypes.h
--- a/Engine/Source/Runtime/Core/Public/UObject/NameTypes.h
+++ b/Engine/Source/Runtime/Core/Public/UObject/NameTypes.h
@@ -283,7 +283,8 @@ public:
 	CORE_API void AppendNameToString(FString& OutString) const;
 
 	/** Appends name to string builder. */
-	CORE_API void AppendNameToString(FStringBuilderBase& OutString) const;
+	CORE_API void AppendNameToString(FWideStringBuilderBase& OutString) const;
+	CORE_API void AppendNameToString(FUtf8StringBuilderBase& OutString) const;
 
 	/** Appends name to string builder. Entry must not be wide. */
 	CORE_API void AppendAnsiNameToString(FAnsiStringBuilderBase& OutString) const;
@@ -583,7 +584,8 @@ public:
 	 * 
 	 * @param Out StringBuilder to fill with the string representation of the name
 	 */
-	void ToString(FStringBuilderBase& Out) const;
+	void ToString(FWideStringBuilderBase& Out) const;
+	void ToString(FUtf8StringBuilderBase& Out) const;
 
 	/**
 	 * Get the number of characters, excluding null-terminator, that ToString() would yield
@@ -622,7 +624,8 @@ public:
 	 * 
 	 * @param Out StringBuilder to append with the string representation of the name
 	 */
-	void AppendString(FStringBuilderBase& Out) const;
+	void AppendString(FWideStringBuilderBase& Out) const;
+	void AppendString(FUtf8StringBuilderBase& Out) const;
 
 	/**
 	 * Converts an ANSI FName to a readable format appended to the string builder.
@@ -1110,6 +1113,9 @@ private:
 	friend bool operator==(FName Lhs, FScriptName Rhs);
 	friend bool operator==(FName Lhs, FMemoryImageName Rhs);
 
+	template <typename StringBufferType>
+	FORCEINLINE void AppendStringInternal(StringBufferType& Out) const;
+
 	FORCEINLINE FNameEntryId GetDisplayIndexFast() const
 	{
 #if WITH_CASE_PRESERVING_NAME
@@ -1587,13 +1593,20 @@ FORCEINLINE void LexFromString(FName& Name, const TCHAR* Str)
 	Name = FName(Str);
 }
 
-inline FStringBuilderBase& operator<<(FStringBuilderBase& Builder, const FName& Name)
+inline FWideStringBuilderBase& operator<<(FWideStringBuilderBase& Builder, const FName& Name)
+{
+	Name.AppendString(Builder);
+	return Builder;
+}
+
+inline FUtf8StringBuilderBase& operator<<(FUtf8StringBuilderBase& Builder, const FName& Name)
 {
 	Name.AppendString(Builder);
 	return Builder;
 }
 
-CORE_API FStringBuilderBase& operator<<(FStringBuilderBase& Builder, FNameEntryId Id);
+CORE_API FWideStringBuilderBase& operator<<(FWideStringBuilderBase& Builder, FNameEntryId Id);
+CORE_API FUtf8StringBuilderBase& operator<<(FUtf8StringBuilderBase& Builder, FNameEntryId Id);
 
 /**
  * Equality operator with CharType* on left hand side and FName on right hand side
