Fix includes being too broad preventing the usage of MTAccessDetector in containers without causing an include loop

diff --git a/Engine/Source/Runtime/Core/Public/Misc/MTAccessDetector.h b/Engine/Source/Runtime/Core/Public/Misc/MTAccessDetector.h
--- a/Engine/Source/Runtime/Core/Public/Misc/MTAccessDetector.h
+++ b/Engine/Source/Runtime/Core/Public/Misc/MTAccessDetector.h
@@ -2,17 +2,19 @@
 
 #pragma once
 
-#include "CoreGlobals.h"
 #include "CoreTypes.h"
 #include "HAL/PlatformTLS.h"
 #include "HAL/PreprocessorHelpers.h"
 #include "Misc/AssertionMacros.h"
 #include "Misc/Build.h"
+#include <atomic>
 
 #define ENABLE_MT_DETECTOR DO_CHECK
 
 #if ENABLE_MT_DETECTOR
 
+extern CORE_API bool GIsAutomationTesting;
+
 /**
  * Read write multithread access detector, will check on concurrent write/write and read/write access, but will not on concurrent read access.
  * Note this detector is not re-entrant, see FRWRecursiveAccessDetector and FRWFullyRecursiveAccessDetector. 
