Fix for incorrect memory ordering in TSharedPtr::IsUnique() in the typical use case.

diff --git a/Engine/Source/Runtime/Core/Public/Templates/SharedPointer.h b/Engine/Source/Runtime/Core/Public/Templates/SharedPointer.h
--- a/Engine/Source/Runtime/Core/Public/Templates/SharedPointer.h
+++ b/Engine/Source/Runtime/Core/Public/Templates/SharedPointer.h
@@ -514,7 +514,11 @@ public:
 	/**
 	 * Returns true if this is the only shared reference to this object.  Note that there may be
 	 * outstanding weak references left.
-	 * IMPORTANT: Not necessarily fast!  Should only be used for debugging purposes!
+	 *
+	 * IMPORTANT: This has different behavior to GetSharedReferenceCount() == 1 in a multithreaded
+	 *            context.  The expectation is that this will be used when a user wants exclusive
+	 *            write-access to an otherwise-immutable object.  Care still needs to be taken when
+	 *            pinning TWeakPtrs to make new shared references.
 	 *
 	 * @return  True if there is only one shared reference to the object, and this is it!
 	 */
diff --git a/Engine/Source/Runtime/Core/Public/Templates/SharedPointerInternals.h b/Engine/Source/Runtime/Core/Public/Templates/SharedPointerInternals.h
--- a/Engine/Source/Runtime/Core/Public/Templates/SharedPointerInternals.h
+++ b/Engine/Source/Runtime/Core/Public/Templates/SharedPointerInternals.h
@@ -112,6 +112,24 @@ namespace SharedPointerInternals
 			}
 		}
 
+		/** Checks if there is exactly one reference left to the object. */
+		FORCEINLINE bool IsUnique() const
+		{
+			if constexpr (Mode == ESPMode::ThreadSafe)
+			{
+				// This is equivalent to https://en.cppreference.com/w/cpp/memory/shared_ptr/unique,
+				// however instead of deprecating it, we implement it with an acquire, which should
+				// suit our use cases.
+
+				// This reference count may be accessed by multiple threads
+				return SharedReferenceCount.load(std::memory_order_acquire) == 1;
+			}
+			else
+			{
+				return SharedReferenceCount == 1;
+			}
+		}
+
 		/** Adds a shared reference to this counter */
 		FORCEINLINE void AddSharedReference()
 		{
@@ -647,7 +665,7 @@ namespace SharedPointerInternals
 		 */
 		FORCEINLINE const bool IsUnique() const
 		{
-			return GetSharedReferenceCount() == 1;
+			return ReferenceController != nullptr && ReferenceController->IsUnique();
 		}
 
 	private:
