Migrate pixel format size calculations to Core so that texture build workers (and misc programs) can touch them without having to depend on RenderCore. By popular demand, brough over the RHI specific capabilities stuff as well.

diff --git a/Engine/Source/Developer/TextureBuild/Private/TextureBuildFunction.cpp b/Engine/Source/Developer/TextureBuild/Private/TextureBuildFunction.cpp
--- a/Engine/Source/Developer/TextureBuild/Private/TextureBuildFunction.cpp
+++ b/Engine/Source/Developer/TextureBuild/Private/TextureBuildFunction.cpp
@@ -30,19 +30,6 @@ DEFINE_LOG_CATEGORY_STATIC(LogTextureBuildFunction, Log, All);
 // by generating a new version.
 static const FGuid TextureDerivedDataVersion(TEXT("19ffd7ab-ea5c-4a18-b6f4-6b0b5eedd606"));
 
-#ifndef CASE_ENUM_TO_TEXT
-#define CASE_ENUM_TO_TEXT(txt) case txt: return TEXT(#txt);
-#endif
-
-static const TCHAR* GetPixelFormatString(EPixelFormat InPixelFormat)
-{
-	switch (InPixelFormat)
-	{
-		FOREACH_ENUM_EPIXELFORMAT(CASE_ENUM_TO_TEXT)
-	default:
-		return TEXT("PF_Unknown");
-	}
-}
 
 static void ReadCbField(FCbFieldView Field, bool& OutValue) { OutValue = Field.AsBool(OutValue); }
 static void ReadCbField(FCbFieldView Field, int32& OutValue) { OutValue = Field.AsInt32(OutValue); }
diff --git a/Engine/Source/Runtime/Core/Private/Misc/PixelFormat.cpp b/Engine/Source/Runtime/Core/Private/Misc/PixelFormat.cpp
new file mode 100644
--- /dev/null
+++ b/Engine/Source/Runtime/Core/Private/Misc/PixelFormat.cpp
@@ -0,0 +1,253 @@
+// Copyright Epic Games, Inc. All Rights Reserved.
+
+#include "PixelFormat.h"
+#include "Math/IntVector.h"
+
+//
+//	Pixel format information.
+//
+
+FPixelFormatInfo::FPixelFormatInfo(
+	EPixelFormat InUnrealFormat,
+	const TCHAR* InName,
+	int32 InBlockSizeX,
+	int32 InBlockSizeY,
+	int32 InBlockSizeZ,
+	int32 InBlockBytes,
+	int32 InNumComponents,
+	bool  InSupported)
+	: Name(InName)
+	, UnrealFormat(InUnrealFormat)
+	, BlockSizeX(InBlockSizeX)
+	, BlockSizeY(InBlockSizeY)
+	, BlockSizeZ(InBlockSizeZ)
+	, BlockBytes(InBlockBytes)
+	, NumComponents(InNumComponents)
+	, Supported(InSupported)
+	, bIs24BitUnormDepthStencil(true)
+{
+}
+
+FPixelFormatInfo    GPixelFormats[PF_MAX] =
+{
+	//               UnrealFormat           Name                  BlockSizeX  BlockSizeY  BlockSizeZ  BlockBytes  NumComponents    Supported
+	FPixelFormatInfo(PF_Unknown,            TEXT("unknown"),               0,          0,          0,          0,          0,              0),
+	FPixelFormatInfo(PF_A32B32G32R32F,      TEXT("A32B32G32R32F"),         1,          1,          1,          16,         4,              1),
+	FPixelFormatInfo(PF_B8G8R8A8,           TEXT("B8G8R8A8"),              1,          1,          1,          4,          4,              1),
+	FPixelFormatInfo(PF_G8,                 TEXT("G8"),                    1,          1,          1,          1,          1,              1),
+	FPixelFormatInfo(PF_G16,                TEXT("G16"),                   1,          1,          1,          2,          1,              1),
+	FPixelFormatInfo(PF_DXT1,               TEXT("DXT1"),                  4,          4,          1,          8,          3,              1),
+	FPixelFormatInfo(PF_DXT3,               TEXT("DXT3"),                  4,          4,          1,          16,         4,              1),
+	FPixelFormatInfo(PF_DXT5,               TEXT("DXT5"),                  4,          4,          1,          16,         4,              1),
+	FPixelFormatInfo(PF_UYVY,               TEXT("UYVY"),                  2,          1,          1,          4,          4,              0),
+	FPixelFormatInfo(PF_FloatRGB,           TEXT("FloatRGB"),              1,          1,          1,          4,          3,              1),
+	FPixelFormatInfo(PF_FloatRGBA,          TEXT("FloatRGBA"),             1,          1,          1,          8,          4,              1),
+	FPixelFormatInfo(PF_DepthStencil,       TEXT("DepthStencil"),          1,          1,          1,          4,          1,              0),
+	FPixelFormatInfo(PF_ShadowDepth,        TEXT("ShadowDepth"),           1,          1,          1,          4,          1,              0),
+	FPixelFormatInfo(PF_R32_FLOAT,          TEXT("R32_FLOAT"),             1,          1,          1,          4,          1,              1),
+	FPixelFormatInfo(PF_G16R16,             TEXT("G16R16"),                1,          1,          1,          4,          2,              1),
+	FPixelFormatInfo(PF_G16R16F,            TEXT("G16R16F"),               1,          1,          1,          4,          2,              1),
+	FPixelFormatInfo(PF_G16R16F_FILTER,     TEXT("G16R16F_FILTER"),        1,          1,          1,          4,          2,              1),
+	FPixelFormatInfo(PF_G32R32F,            TEXT("G32R32F"),               1,          1,          1,          8,          2,              1),
+	FPixelFormatInfo(PF_A2B10G10R10,        TEXT("A2B10G10R10"),           1,          1,          1,          4,          4,              1),
+	FPixelFormatInfo(PF_A16B16G16R16,       TEXT("A16B16G16R16"),          1,          1,          1,          8,          4,              1),
+	FPixelFormatInfo(PF_D24,                TEXT("D24"),                   1,          1,          1,          4,          1,              1),
+	FPixelFormatInfo(PF_R16F,               TEXT("PF_R16F"),               1,          1,          1,          2,          1,              1),
+	FPixelFormatInfo(PF_R16F_FILTER,        TEXT("PF_R16F_FILTER"),        1,          1,          1,          2,          1,              1),
+	FPixelFormatInfo(PF_BC5,                TEXT("BC5"),                   4,          4,          1,          16,         2,              1),
+	FPixelFormatInfo(PF_V8U8,               TEXT("V8U8"),                  1,          1,          1,          2,          2,              1),
+	FPixelFormatInfo(PF_A1,                 TEXT("A1"),                    1,          1,          1,          1,          1,              0),
+	FPixelFormatInfo(PF_FloatR11G11B10,     TEXT("FloatR11G11B10"),        1,          1,          1,          4,          3,              0),
+	FPixelFormatInfo(PF_A8,                 TEXT("A8"),                    1,          1,          1,          1,          1,              1),
+	FPixelFormatInfo(PF_R32_UINT,           TEXT("R32_UINT"),              1,          1,          1,          4,          1,              1),
+	FPixelFormatInfo(PF_R32_SINT,           TEXT("R32_SINT"),              1,          1,          1,          4,          1,              1),
+
+	// IOS Support
+	FPixelFormatInfo(PF_PVRTC2,             TEXT("PVRTC2"),                8,          4,          1,          8,          4,              0),
+	FPixelFormatInfo(PF_PVRTC4,             TEXT("PVRTC4"),                4,          4,          1,          8,          4,              0),
+
+	FPixelFormatInfo(PF_R16_UINT,           TEXT("R16_UINT"),              1,          1,          1,          2,          1,              1),
+	FPixelFormatInfo(PF_R16_SINT,           TEXT("R16_SINT"),              1,          1,          1,          2,          1,              1),
+	FPixelFormatInfo(PF_R16G16B16A16_UINT,  TEXT("R16G16B16A16_UINT"),     1,          1,          1,          8,          4,              1),
+	FPixelFormatInfo(PF_R16G16B16A16_SINT,  TEXT("R16G16B16A16_SINT"),     1,          1,          1,          8,          4,              1),
+	FPixelFormatInfo(PF_R5G6B5_UNORM,       TEXT("R5G6B5_UNORM"),          1,          1,          1,          2,          3,              0),
+	FPixelFormatInfo(PF_R8G8B8A8,           TEXT("R8G8B8A8"),              1,          1,          1,          4,          4,              1),
+	FPixelFormatInfo(PF_A8R8G8B8,           TEXT("A8R8G8B8"),              1,          1,          1,          4,          4,              1),
+	FPixelFormatInfo(PF_BC4,                TEXT("BC4"),                   4,          4,          1,          8,          1,              1),
+	FPixelFormatInfo(PF_R8G8,               TEXT("R8G8"),                  1,          1,          1,          2,          2,              1),
+
+	FPixelFormatInfo(PF_ATC_RGB,            TEXT("ATC_RGB"),               4,          4,          1,          8,          3,              0),
+	FPixelFormatInfo(PF_ATC_RGBA_E,         TEXT("ATC_RGBA_E"),            4,          4,          1,          16,         4,              0),
+	FPixelFormatInfo(PF_ATC_RGBA_I,         TEXT("ATC_RGBA_I"),            4,          4,          1,          16,         4,              0),
+	FPixelFormatInfo(PF_X24_G8,             TEXT("X24_G8"),                1,          1,          1,          1,          1,              0),
+	FPixelFormatInfo(PF_ETC1,               TEXT("ETC1"),                  4,          4,          1,          8,          3,              0),
+	FPixelFormatInfo(PF_ETC2_RGB,           TEXT("ETC2_RGB"),              4,          4,          1,          8,          3,              0),
+	FPixelFormatInfo(PF_ETC2_RGBA,          TEXT("ETC2_RGBA"),             4,          4,          1,          16,         4,              0),
+	FPixelFormatInfo(PF_R32G32B32A32_UINT,  TEXT("PF_R32G32B32A32_UINT"),  1,          1,          1,          16,         4,              1),
+	FPixelFormatInfo(PF_R16G16_UINT,        TEXT("PF_R16G16_UINT"),        1,          1,          1,          4,          4,              1),
+
+	// ASTC support
+	FPixelFormatInfo(PF_ASTC_4x4,           TEXT("ASTC_4x4"),              4,          4,          1,          16,         4,              0),
+	FPixelFormatInfo(PF_ASTC_6x6,           TEXT("ASTC_6x6"),              6,          6,          1,          16,         4,              0),
+	FPixelFormatInfo(PF_ASTC_8x8,           TEXT("ASTC_8x8"),              8,          8,          1,          16,         4,              0),
+	FPixelFormatInfo(PF_ASTC_10x10,         TEXT("ASTC_10x10"),            10,         10,         1,          16,         4,              0),
+	FPixelFormatInfo(PF_ASTC_12x12,         TEXT("ASTC_12x12"),            12,         12,         1,          16,         4,              0),
+	
+	FPixelFormatInfo(PF_BC6H,               TEXT("BC6H"),                  4,          4,          1,          16,         3,              1),
+	FPixelFormatInfo(PF_BC7,                TEXT("BC7"),                   4,          4,          1,          16,         4,              1),
+	FPixelFormatInfo(PF_R8_UINT,            TEXT("R8_UINT"),               1,          1,          1,          1,          1,              1),
+	FPixelFormatInfo(PF_L8,                 TEXT("L8"),                    1,          1,          1,          1,          1,              0),
+	FPixelFormatInfo(PF_XGXR8,              TEXT("XGXR8"),                 1,          1,          1,          4,          4,              1),
+	FPixelFormatInfo(PF_R8G8B8A8_UINT,      TEXT("R8G8B8A8_UINT"),         1,          1,          1,          4,          4,              1),
+	FPixelFormatInfo(PF_R8G8B8A8_SNORM,     TEXT("R8G8B8A8_SNORM"),        1,          1,          1,          4,          4,              1),
+
+	FPixelFormatInfo(PF_R16G16B16A16_UNORM, TEXT("R16G16B16A16_UINT"),     1,          1,          1,          8,          4,              1),
+	FPixelFormatInfo(PF_R16G16B16A16_SNORM, TEXT("R16G16B16A16_SINT"),     1,          1,          1,          8,          4,              1),
+	FPixelFormatInfo(PF_PLATFORM_HDR_0,     TEXT("PLATFORM_HDR_0"),        0,          0,          0,          0,          0,              0),
+	FPixelFormatInfo(PF_PLATFORM_HDR_1,     TEXT("PLATFORM_HDR_1"),        0,          0,          0,          0,          0,              0),
+	FPixelFormatInfo(PF_PLATFORM_HDR_2,     TEXT("PLATFORM_HDR_2"),        0,          0,          0,          0,          0,              0),
+
+	// NV12 contains 2 textures: R8 luminance plane followed by R8G8 1/4 size chrominance plane.
+	// BlockSize/BlockBytes/NumComponents values don't make much sense for this format, so set them all to one.
+	FPixelFormatInfo(PF_NV12,               TEXT("NV12"),                  1,          1,          1,          1,          1,              0),
+
+	FPixelFormatInfo(PF_R32G32_UINT,        TEXT("PF_R32G32_UINT"),        1,          1,          1,          8,          2,              1),
+
+	FPixelFormatInfo(PF_ETC2_R11_EAC,       TEXT("PF_ETC2_R11_EAC"),       4,          4,          1,          8,          1,              0),
+	FPixelFormatInfo(PF_ETC2_RG11_EAC,      TEXT("PF_ETC2_RG11_EAC"),      4,          4,          1,          16,         2,              0),
+	FPixelFormatInfo(PF_R8,                 TEXT("R8"),                    1,          1,          1,          1,          1,              1),
+    FPixelFormatInfo(PF_B5G5R5A1_UNORM,     TEXT("B5G5R5A1_UNORM"),        1,          1,          1,          2,          4,              0),
+
+	// ASTC HDR support
+	FPixelFormatInfo(PF_ASTC_4x4_HDR,       TEXT("ASTC_4x4_HDR"),          4,          4,          1,          16,         4,              0),
+	FPixelFormatInfo(PF_ASTC_6x6_HDR,       TEXT("ASTC_6x6_HDR"),          6,          6,          1,          16,         4,              0),
+	FPixelFormatInfo(PF_ASTC_8x8_HDR,       TEXT("ASTC_8x8_HDR"),          8,          8,          1,          16,         4,              0),
+	FPixelFormatInfo(PF_ASTC_10x10_HDR,     TEXT("ASTC_10x10_HDR"),        10,         10,         1,          16,         4,              0),
+	FPixelFormatInfo(PF_ASTC_12x12_HDR,     TEXT("ASTC_12x12_HDR"),        12,         12,         1,          16,         4,              0),
+
+	FPixelFormatInfo(PF_G16R16_SNORM,       TEXT("G16R16_SNORM"),          1,          1,          1,          4,          2,              1),
+	FPixelFormatInfo(PF_R8G8_UINT,          TEXT("R8G8_UINT"),             1,          1,          1,          2,          2,              1),
+	FPixelFormatInfo(PF_R32G32B32_UINT,     TEXT("R32G32B32_UINT"),        1,          1,          1,          12,         3,              1),
+	FPixelFormatInfo(PF_R32G32B32_SINT,     TEXT("R32G32B32_SINT"),        1,          1,          1,          12,         3,              1),
+	FPixelFormatInfo(PF_R32G32B32F,         TEXT("R32G32B32F"),            1,          1,          1,          12,         3,              1),
+	FPixelFormatInfo(PF_R8_SINT,            TEXT("R8_SINT"),               1,          1,          1,          1,          1,              1),
+	FPixelFormatInfo(PF_R64_UINT,			TEXT("R64_UINT"),              1,          1,          1,          8,          1,              0),
+};
+
+
+uint64 FPixelFormatInfo::Get3DImageSizeInBytes(uint32 InWidth, uint32 InHeight, uint32 InDepth) const
+{
+	const uint64 WidthInBlocks = GetBlockCountForWidth(InWidth);
+	const uint64 HeightInBlocks = GetBlockCountForHeight(InHeight);
+
+	// no format currently has block requirement for depth
+	return InDepth * WidthInBlocks * HeightInBlocks * BlockBytes;
+}
+
+uint64 FPixelFormatInfo::Get3DTextureMipSizeInBytes(uint32 InTextureWidth, uint32 InTextureHeight, uint32 InTextureDepth, uint32 InMipIndex) const
+{
+	int32 MipWidth = FMath::Max((int32)InTextureWidth >> InMipIndex, 1);
+	int32 MipHeight = FMath::Max((int32)InTextureHeight >> InMipIndex, 1);
+	int32 MipDepth = FMath::Max((int32)InTextureDepth >> InMipIndex, 1);	
+	return Get3DImageSizeInBytes(MipWidth, MipHeight, MipDepth);
+}
+
+uint64 FPixelFormatInfo::Get3DTextureSizeInBytes(uint32 InTextureWidth, uint32 InTextureHeight, uint32 InTextureDepth, uint32 InMipCount) const
+{
+	uint64 Size = 0;
+	int32 MipWidth = InTextureWidth;
+	int32 MipHeight = InTextureHeight;
+	int32 MipDepth = InTextureDepth;
+	for (uint32 MipIndex = 0; MipIndex < InMipCount; ++MipIndex)
+	{
+		Size += Get3DImageSizeInBytes(MipWidth, MipHeight, MipDepth);
+		MipWidth = FMath::Max(MipWidth >> 1, 1);
+		MipHeight = FMath::Max(MipHeight >> 1, 1);
+		MipDepth = FMath::Max(MipDepth >> 1, 1);
+	}
+	return Size;
+}
+
+uint64 FPixelFormatInfo::GetBlockCountForWidth(uint32 InWidth) const
+{
+	// note: FTexture2DResource applies a 2 block min for PVRTC that is not applied here
+	if (BlockSizeX > 0)
+	{
+		return (InWidth + BlockSizeX - 1) / BlockSizeX;
+	}
+	else
+	{
+		return 0;
+	}
+}
+
+
+uint64 FPixelFormatInfo::GetBlockCountForHeight(uint32 InHeight) const
+{
+	// note: FTexture2DResource applies a 2 block min for PVRTC that is not applied here
+	if (BlockSizeY > 0)
+	{
+		return (InHeight + BlockSizeY - 1) / BlockSizeY;
+	}
+	else
+	{
+		return 0;
+	}
+}
+
+uint64 FPixelFormatInfo::Get2DImageSizeInBytes(uint32 InWidth, uint32 InHeight) const
+{
+	const uint64 WidthInBlocks = GetBlockCountForWidth(InWidth);
+	const uint64 HeightInBlocks = GetBlockCountForHeight(InHeight);
+	return WidthInBlocks * HeightInBlocks * BlockBytes;
+}
+
+uint64 FPixelFormatInfo::Get2DTextureMipSizeInBytes(uint32 InTextureWidth, uint32 InTextureHeight, uint32 InMipIndex) const
+{
+	int32 MipWidth = FMath::Max((int32)InTextureWidth >> InMipIndex, 1);
+	int32 MipHeight = FMath::Max((int32)InTextureHeight >> InMipIndex, 1);
+	return Get2DImageSizeInBytes(MipWidth, MipHeight);
+}
+
+uint64 FPixelFormatInfo::Get2DTextureSizeInBytes(uint32 InTextureWidth, uint32 InTextureHeight, uint32 InMipCount) const
+{
+	uint64 Size = 0;
+	int32 MipWidth = InTextureWidth;
+	int32 MipHeight = InTextureHeight;
+	for ( uint32 MipIndex=0; MipIndex < InMipCount; ++MipIndex )
+	{
+		Size += Get2DImageSizeInBytes(MipWidth, MipHeight);
+		MipWidth = FMath::Max(MipWidth >> 1, 1);
+		MipHeight = FMath::Max(MipHeight >> 1, 1);
+	}
+	return Size;
+}
+
+
+/** Helper functions for text output of texture properties... */
+#ifndef CASE_ENUM_TO_TEXT
+#define CASE_ENUM_TO_TEXT(txt) case txt: return TEXT(#txt);
+#endif
+
+#ifndef TEXT_TO_ENUM
+#define TEXT_TO_ENUM(eVal, txt)		if (FCString::Stricmp(TEXT(#eVal), txt) == 0)	return eVal;
+#endif
+
+const TCHAR* GetPixelFormatString(EPixelFormat InPixelFormat)
+{
+	switch (InPixelFormat)
+	{
+		FOREACH_ENUM_EPIXELFORMAT(CASE_ENUM_TO_TEXT)
+	default:
+		return TEXT("PF_Unknown");
+	}
+}
+
+EPixelFormat GetPixelFormatFromString(const TCHAR* InPixelFormatStr)
+{
+#define TEXT_TO_PIXELFORMAT(f) TEXT_TO_ENUM(f, InPixelFormatStr);
+	FOREACH_ENUM_EPIXELFORMAT(TEXT_TO_PIXELFORMAT)
+#undef TEXT_TO_PIXELFORMAT
+		return PF_Unknown;
+}
+
diff --git a/Engine/Source/Runtime/Core/Public/PixelFormat.h b/Engine/Source/Runtime/Core/Public/PixelFormat.h
--- a/Engine/Source/Runtime/Core/Public/PixelFormat.h
+++ b/Engine/Source/Runtime/Core/Public/PixelFormat.h
@@ -10,6 +10,8 @@
 #undef PF_MAX
 #endif
 
+enum class EPixelFormatCapabilities : uint32;
+
 enum EPixelFormat : uint8
 {
 	PF_Unknown              =0,
@@ -210,3 +212,75 @@ FORCEINLINE bool IsHDR(EPixelFormat PixelFormat)
 {
 	return PixelFormat == PF_FloatRGBA || PixelFormat == PF_BC6H || PixelFormat == PF_R16F || PixelFormat == PF_R32_FLOAT || PixelFormat == PF_A32B32G32R32F;
 }
+
+/** 
+* Information about a pixel format. The majority of this structure is valid after static init, however RHI does keep
+* some state in here that is initialized by that module and should not be used by general use programs that don't
+* have RHI (so noted in comments).
+*/
+struct FPixelFormatInfo
+{
+	FPixelFormatInfo() = delete;
+	FPixelFormatInfo(
+		EPixelFormat InUnrealFormat,
+		const TCHAR* InName,
+		int32 InBlockSizeX,
+		int32 InBlockSizeY,
+		int32 InBlockSizeZ,
+		int32 InBlockBytes,
+		int32 InNumComponents,
+		bool  InSupported);
+
+	const TCHAR*				Name;
+	EPixelFormat				UnrealFormat;
+	int32						BlockSizeX;
+	int32						BlockSizeY;
+	int32						BlockSizeZ;
+	int32						BlockBytes;
+	int32						NumComponents;
+
+	/** Per platform cabilities for the format (initialized by RHI module - invalid otherwise) */
+	EPixelFormatCapabilities	Capabilities{ 0 };
+
+	/** Platform specific converted format (initialized by RHI module - invalid otherwise) */
+	uint32						PlatformFormat{ 0 };
+
+	/** Whether the texture format is supported on the current platform/ rendering combination */
+	uint8						Supported : 1;
+	uint8						bIs24BitUnormDepthStencil : 1;	// If false, 32 bit float is assumed (initialized by RHI module - invalid otherwise)
+	
+	/** 
+	* Get 2D/3D image/texture size in bytes. This is for storage of the encoded image data, and does not adjust
+	* for any GPU alignment/padding constraints. It is also not valid for tiled or packed mip tails (i.e. cooked mips 
+	* for consoles). Only use these when you know you're working with bog standard textures/images in block based pixel formats.
+	*/
+	CORE_API uint64 Get2DImageSizeInBytes(uint32 InWidth, uint32 InHeight) const;
+	CORE_API uint64 Get2DTextureMipSizeInBytes(uint32 InTextureWidth, uint32 InTextureHeight, uint32 InMipIndex) const;
+	CORE_API uint64 Get2DTextureSizeInBytes(uint32 InTextureWidth, uint32 InTextureHeight, uint32 InMipCount) const;
+	CORE_API uint64 Get3DImageSizeInBytes(uint32 InWidth, uint32 InHeight, uint32 InDepth) const;
+	CORE_API uint64 Get3DTextureMipSizeInBytes(uint32 InTextureWidth, uint32 InTextureHeight, uint32 InTextureDepth, uint32 InMipIndex) const;
+	CORE_API uint64 Get3DTextureSizeInBytes(uint32 InTextureWidth, uint32 InTextureHeight, uint32 InTextureDepth, uint32 InMipCount) const;
+	
+	/**
+	* Get the number of compressed blocks necessary to hold the given dimensions.
+	*/
+	CORE_API uint64 GetBlockCountForWidth(uint32 InWidth) const;
+	CORE_API uint64 GetBlockCountForHeight(uint32 InHeight) const;
+};
+
+
+/**
+ * enum to string
+ * Note this is not the FPixelFormatInfo::Name
+ * @return e.g. "PF_B8G8R8A8"
+ */
+CORE_API const TCHAR* GetPixelFormatString(EPixelFormat InPixelFormat);
+/**
+ * string to enum (not case sensitive)
+ * Note this is not the FPixelFormatInfo::Name
+ * @param InPixelFormatStr e.g. "PF_B8G8R8A8", must not not be 0
+ */
+CORE_API EPixelFormat GetPixelFormatFromString(const TCHAR* InPixelFormatStr);
+
+
+extern CORE_API FPixelFormatInfo GPixelFormats[PF_MAX];		// Maps members of EPixelFormat to a FPixelFormatInfo describing the format.
diff --git a/Engine/Source/Runtime/Engine/Private/Slate/SlateTextures.cpp b/Engine/Source/Runtime/Engine/Private/Slate/SlateTextures.cpp
--- a/Engine/Source/Runtime/Engine/Private/Slate/SlateTextures.cpp
+++ b/Engine/Source/Runtime/Engine/Private/Slate/SlateTextures.cpp
@@ -69,8 +69,8 @@ void FSlateTexture2DRHIRef::InitDynamicRHI()
 			uint8* DestTextureData = (uint8*)RHILockTexture2D(ShaderResource, 0, RLM_WriteOnly, Stride, false);
 			const uint8* SourceTextureData = TextureData->GetRawBytes().GetData();
 
-			const uint32 BlocksX = CalcTextureMipWidthInBlocks(Width, PixelFormat, 0);
-			const uint32 BlocksY = CalcTextureMipWidthInBlocks(Height, PixelFormat, 0);
+			const uint32 BlocksX = GPixelFormats[PixelFormat].GetBlockCountForWidth(Width);
+			const uint32 BlocksY = GPixelFormats[PixelFormat].GetBlockCountForHeight(Height);
 			const uint32 DataStride = BlocksX * GPixelFormats[PixelFormat].BlockBytes;
 
 			if (Stride == DataStride)
diff --git a/Engine/Source/Runtime/RHI/Private/RHI.cpp b/Engine/Source/Runtime/RHI/Private/RHI.cpp
--- a/Engine/Source/Runtime/RHI/Private/RHI.cpp
+++ b/Engine/Source/Runtime/RHI/Private/RHI.cpp
@@ -2663,138 +2663,6 @@ int32 CalculateMSAASampleArrayIndex(int32 NumSamples, int32 SampleIndex)
 	return NumSamples - 1 + SampleIndex;
 }
 
-//
-//	Pixel format information.
-//
-
-FPixelFormatInfo::FPixelFormatInfo(
-	EPixelFormat InUnrealFormat,
-	const TCHAR* InName,
-	int32 InBlockSizeX,
-	int32 InBlockSizeY,
-	int32 InBlockSizeZ,
-	int32 InBlockBytes,
-	int32 InNumComponents,
-	bool  InSupported)
-	: Name(InName)
-	, UnrealFormat(InUnrealFormat)
-	, BlockSizeX(InBlockSizeX)
-	, BlockSizeY(InBlockSizeY)
-	, BlockSizeZ(InBlockSizeZ)
-	, BlockBytes(InBlockBytes)
-	, NumComponents(InNumComponents)
-	, Supported(InSupported)
-	, bIs24BitUnormDepthStencil(true)
-{
-}
-
-FPixelFormatInfo    GPixelFormats[PF_MAX] =
-{
-	//               UnrealFormat                 Name                 BlockSizeX  BlockSizeY  BlockSizeZ  BlockBytes  NumComponents    Supported
-
-	FPixelFormatInfo(PF_Unknown,            TEXT("unknown"),               0,          0,          0,          0,          0,              0),
-	FPixelFormatInfo(PF_A32B32G32R32F,      TEXT("A32B32G32R32F"),         1,          1,          1,          16,         4,              1),
-	FPixelFormatInfo(PF_B8G8R8A8,           TEXT("B8G8R8A8"),              1,          1,          1,          4,          4,              1),
-	FPixelFormatInfo(PF_G8,                 TEXT("G8"),                    1,          1,          1,          1,          1,              1),
-	FPixelFormatInfo(PF_G16,                TEXT("G16"),                   1,          1,          1,          2,          1,              1),
-	FPixelFormatInfo(PF_DXT1,               TEXT("DXT1"),                  4,          4,          1,          8,          3,              1),
-	FPixelFormatInfo(PF_DXT3,               TEXT("DXT3"),                  4,          4,          1,          16,         4,              1),
-	FPixelFormatInfo(PF_DXT5,               TEXT("DXT5"),                  4,          4,          1,          16,         4,              1),
-	FPixelFormatInfo(PF_UYVY,               TEXT("UYVY"),                  2,          1,          1,          4,          4,              0),
-	FPixelFormatInfo(PF_FloatRGB,           TEXT("FloatRGB"),              1,          1,          1,          4,          3,              1),
-	FPixelFormatInfo(PF_FloatRGBA,          TEXT("FloatRGBA"),             1,          1,          1,          8,          4,              1),
-	FPixelFormatInfo(PF_DepthStencil,       TEXT("DepthStencil"),          1,          1,          1,          4,          1,              0),
-	FPixelFormatInfo(PF_ShadowDepth,        TEXT("ShadowDepth"),           1,          1,          1,          4,          1,              0),
-	FPixelFormatInfo(PF_R32_FLOAT,          TEXT("R32_FLOAT"),             1,          1,          1,          4,          1,              1),
-	FPixelFormatInfo(PF_G16R16,             TEXT("G16R16"),                1,          1,          1,          4,          2,              1),
-	FPixelFormatInfo(PF_G16R16F,            TEXT("G16R16F"),               1,          1,          1,          4,          2,              1),
-	FPixelFormatInfo(PF_G16R16F_FILTER,     TEXT("G16R16F_FILTER"),        1,          1,          1,          4,          2,              1),
-	FPixelFormatInfo(PF_G32R32F,            TEXT("G32R32F"),               1,          1,          1,          8,          2,              1),
-	FPixelFormatInfo(PF_A2B10G10R10,        TEXT("A2B10G10R10"),           1,          1,          1,          4,          4,              1),
-	FPixelFormatInfo(PF_A16B16G16R16,       TEXT("A16B16G16R16"),          1,          1,          1,          8,          4,              1),
-	FPixelFormatInfo(PF_D24,                TEXT("D24"),                   1,          1,          1,          4,          1,              1),
-	FPixelFormatInfo(PF_R16F,               TEXT("PF_R16F"),               1,          1,          1,          2,          1,              1),
-	FPixelFormatInfo(PF_R16F_FILTER,        TEXT("PF_R16F_FILTER"),        1,          1,          1,          2,          1,              1),
-	FPixelFormatInfo(PF_BC5,                TEXT("BC5"),                   4,          4,          1,          16,         2,              1),
-	FPixelFormatInfo(PF_V8U8,               TEXT("V8U8"),                  1,          1,          1,          2,          2,              1),
-	FPixelFormatInfo(PF_A1,                 TEXT("A1"),                    1,          1,          1,          1,          1,              0),
-	FPixelFormatInfo(PF_FloatR11G11B10,     TEXT("FloatR11G11B10"),        1,          1,          1,          4,          3,              0),
-	FPixelFormatInfo(PF_A8,                 TEXT("A8"),                    1,          1,          1,          1,          1,              1),
-	FPixelFormatInfo(PF_R32_UINT,           TEXT("R32_UINT"),              1,          1,          1,          4,          1,              1),
-	FPixelFormatInfo(PF_R32_SINT,           TEXT("R32_SINT"),              1,          1,          1,          4,          1,              1),
-
-	// IOS Support
-	FPixelFormatInfo(PF_PVRTC2,             TEXT("PVRTC2"),                8,          4,          1,          8,          4,              0),
-	FPixelFormatInfo(PF_PVRTC4,             TEXT("PVRTC4"),                4,          4,          1,          8,          4,              0),
-
-	FPixelFormatInfo(PF_R16_UINT,           TEXT("R16_UINT"),              1,          1,          1,          2,          1,              1),
-	FPixelFormatInfo(PF_R16_SINT,           TEXT("R16_SINT"),              1,          1,          1,          2,          1,              1),
-	FPixelFormatInfo(PF_R16G16B16A16_UINT,  TEXT("R16G16B16A16_UINT"),     1,          1,          1,          8,          4,              1),
-	FPixelFormatInfo(PF_R16G16B16A16_SINT,  TEXT("R16G16B16A16_SINT"),     1,          1,          1,          8,          4,              1),
-	FPixelFormatInfo(PF_R5G6B5_UNORM,       TEXT("R5G6B5_UNORM"),          1,          1,          1,          2,          3,              0),
-	FPixelFormatInfo(PF_R8G8B8A8,           TEXT("R8G8B8A8"),              1,          1,          1,          4,          4,              1),
-	FPixelFormatInfo(PF_A8R8G8B8,           TEXT("A8R8G8B8"),              1,          1,          1,          4,          4,              1),
-	FPixelFormatInfo(PF_BC4,                TEXT("BC4"),                   4,          4,          1,          8,          1,              1),
-	FPixelFormatInfo(PF_R8G8,               TEXT("R8G8"),                  1,          1,          1,          2,          2,              1),
-
-	FPixelFormatInfo(PF_ATC_RGB,            TEXT("ATC_RGB"),               4,          4,          1,          8,          3,              0),
-	FPixelFormatInfo(PF_ATC_RGBA_E,         TEXT("ATC_RGBA_E"),            4,          4,          1,          16,         4,              0),
-	FPixelFormatInfo(PF_ATC_RGBA_I,         TEXT("ATC_RGBA_I"),            4,          4,          1,          16,         4,              0),
-	FPixelFormatInfo(PF_X24_G8,             TEXT("X24_G8"),                1,          1,          1,          1,          1,              0),
-	FPixelFormatInfo(PF_ETC1,               TEXT("ETC1"),                  4,          4,          1,          8,          3,              0),
-	FPixelFormatInfo(PF_ETC2_RGB,           TEXT("ETC2_RGB"),              4,          4,          1,          8,          3,              0),
-	FPixelFormatInfo(PF_ETC2_RGBA,          TEXT("ETC2_RGBA"),             4,          4,          1,          16,         4,              0),
-	FPixelFormatInfo(PF_R32G32B32A32_UINT,  TEXT("PF_R32G32B32A32_UINT"),  1,          1,          1,          16,         4,              1),
-	FPixelFormatInfo(PF_R16G16_UINT,        TEXT("PF_R16G16_UINT"),        1,          1,          1,          4,          4,              1),
-
-	// ASTC support
-	FPixelFormatInfo(PF_ASTC_4x4,           TEXT("ASTC_4x4"),              4,          4,          1,          16,         4,              0),
-	FPixelFormatInfo(PF_ASTC_6x6,           TEXT("ASTC_6x6"),              6,          6,          1,          16,         4,              0),
-	FPixelFormatInfo(PF_ASTC_8x8,           TEXT("ASTC_8x8"),              8,          8,          1,          16,         4,              0),
-	FPixelFormatInfo(PF_ASTC_10x10,         TEXT("ASTC_10x10"),            10,         10,         1,          16,         4,              0),
-	FPixelFormatInfo(PF_ASTC_12x12,         TEXT("ASTC_12x12"),            12,         12,         1,          16,         4,              0),
-	
-	FPixelFormatInfo(PF_BC6H,               TEXT("BC6H"),                  4,          4,          1,          16,         3,              1),
-	FPixelFormatInfo(PF_BC7,                TEXT("BC7"),                   4,          4,          1,          16,         4,              1),
-	FPixelFormatInfo(PF_R8_UINT,            TEXT("R8_UINT"),               1,          1,          1,          1,          1,              1),
-	FPixelFormatInfo(PF_L8,                 TEXT("L8"),                    1,          1,          1,          1,          1,              0),
-	FPixelFormatInfo(PF_XGXR8,              TEXT("XGXR8"),                 1,          1,          1,          4,          4,              1),
-	FPixelFormatInfo(PF_R8G8B8A8_UINT,      TEXT("R8G8B8A8_UINT"),         1,          1,          1,          4,          4,              1),
-	FPixelFormatInfo(PF_R8G8B8A8_SNORM,     TEXT("R8G8B8A8_SNORM"),        1,          1,          1,          4,          4,              1),
-
-	FPixelFormatInfo(PF_R16G16B16A16_UNORM, TEXT("R16G16B16A16_UINT"),     1,          1,          1,          8,          4,              1),
-	FPixelFormatInfo(PF_R16G16B16A16_SNORM, TEXT("R16G16B16A16_SINT"),     1,          1,          1,          8,          4,              1),
-	FPixelFormatInfo(PF_PLATFORM_HDR_0,     TEXT("PLATFORM_HDR_0"),        0,          0,          0,          0,          0,              0),
-	FPixelFormatInfo(PF_PLATFORM_HDR_1,     TEXT("PLATFORM_HDR_1"),        0,          0,          0,          0,          0,              0),
-	FPixelFormatInfo(PF_PLATFORM_HDR_2,     TEXT("PLATFORM_HDR_2"),        0,          0,          0,          0,          0,              0),
-
-	// NV12 contains 2 textures: R8 luminance plane followed by R8G8 1/4 size chrominance plane.
-	// BlockSize/BlockBytes/NumComponents values don't make much sense for this format, so set them all to one.
-	FPixelFormatInfo(PF_NV12,               TEXT("NV12"),                  1,          1,          1,          1,          1,              0),
-
-	FPixelFormatInfo(PF_R32G32_UINT,        TEXT("PF_R32G32_UINT"),        1,          1,          1,          8,          2,              1),
-
-	FPixelFormatInfo(PF_ETC2_R11_EAC,       TEXT("PF_ETC2_R11_EAC"),       4,          4,          1,          8,          1,              0),
-	FPixelFormatInfo(PF_ETC2_RG11_EAC,      TEXT("PF_ETC2_RG11_EAC"),      4,          4,          1,          16,         2,              0),
-	FPixelFormatInfo(PF_R8,                 TEXT("R8"),                    1,          1,          1,          1,          1,              1),
-    FPixelFormatInfo(PF_B5G5R5A1_UNORM,     TEXT("B5G5R5A1_UNORM"),        1,          1,          1,          2,          4,              0),
-
-	// ASTC HDR support
-	FPixelFormatInfo(PF_ASTC_4x4_HDR,           TEXT("ASTC_4x4_HDR"),      4,          4,          1,          16,         4,              0),
-	FPixelFormatInfo(PF_ASTC_6x6_HDR,           TEXT("ASTC_6x6_HDR"),      6,          6,          1,          16,         4,              0),
-	FPixelFormatInfo(PF_ASTC_8x8_HDR,           TEXT("ASTC_8x8_HDR"),      8,          8,          1,          16,         4,              0),
-	FPixelFormatInfo(PF_ASTC_10x10_HDR,         TEXT("ASTC_10x10_HDR"),    10,         10,         1,          16,         4,              0),
-	FPixelFormatInfo(PF_ASTC_12x12_HDR,         TEXT("ASTC_12x12_HDR"),    12,         12,         1,          16,         4,              0),
-
-	FPixelFormatInfo(PF_G16R16_SNORM,       TEXT("G16R16_SNORM"),          1,          1,          1,          4,          2,              1),
-	FPixelFormatInfo(PF_R8G8_UINT,          TEXT("R8G8_UINT"),             1,          1,          1,          2,          2,              1),
-	FPixelFormatInfo(PF_R32G32B32_UINT,     TEXT("R32G32B32_UINT"),        1,          1,          1,          12,         3,              1),
-	FPixelFormatInfo(PF_R32G32B32_SINT,     TEXT("R32G32B32_SINT"),        1,          1,          1,          12,         3,              1),
-	FPixelFormatInfo(PF_R32G32B32F,         TEXT("R32G32B32F"),            1,          1,          1,          12,         3,              1),
-	FPixelFormatInfo(PF_R8_SINT,            TEXT("R8_SINT"),               1,          1,          1,          1,          1,              1),
-	FPixelFormatInfo(PF_R64_UINT,			TEXT("R64_UINT"),              1,          1,          1,          8,          1,              0),
-};
-
 void RHIInitDefaultPixelFormatCapabilities()
 {
 	for (FPixelFormatInfo& Info : GPixelFormats)
@@ -2838,19 +2706,19 @@ static struct FValidatePixelFormats
 SIZE_T CalculateImageBytes(uint32 SizeX,uint32 SizeY,uint32 SizeZ,uint8 Format)
 {
 	if ( Format == PF_A1 )
-	{
+	{		
 		// The number of bytes needed to store all 1 bit pixels in a line is the width of the image divided by the number of bits in a byte
-		uint32 BytesPerLine = SizeX / 8;
+		uint32 BytesPerLine = (SizeX + 7) / 8;
 		// The number of actual bytes in a 1 bit image is the bytes per line of pixels times the number of lines
 		return sizeof(uint8) * BytesPerLine * SizeY;
 	}
 	else if( SizeZ > 0 )
 	{
-		return static_cast<SIZE_T>(SizeX / GPixelFormats[Format].BlockSizeX) * (SizeY / GPixelFormats[Format].BlockSizeY) * (SizeZ / GPixelFormats[Format].BlockSizeZ) * GPixelFormats[Format].BlockBytes;
+		return GPixelFormats[Format].Get3DImageSizeInBytes(SizeX, SizeY, SizeZ);
 	}
 	else
 	{
-		return static_cast<SIZE_T>(SizeX / GPixelFormats[Format].BlockSizeX) * (SizeY / GPixelFormats[Format].BlockSizeY) * GPixelFormats[Format].BlockBytes;
+		return GPixelFormats[Format].Get2DImageSizeInBytes(SizeX, SizeY);
 	}
 }
 
diff --git a/Engine/Source/Runtime/RHI/Public/RHI.h b/Engine/Source/Runtime/RHI/Public/RHI.h
--- a/Engine/Source/Runtime/RHI/Public/RHI.h
+++ b/Engine/Source/Runtime/RHI/Public/RHI.h
@@ -843,41 +843,6 @@ enum class EPixelFormatCapabilities : uint32
 };
 ENUM_CLASS_FLAGS(EPixelFormatCapabilities);
 
-/** Information about a pixel format. */
-struct FPixelFormatInfo
-{
-	FPixelFormatInfo() = delete;
-	FPixelFormatInfo(
-		EPixelFormat InUnrealFormat,
-		const TCHAR* InName,
-		int32 InBlockSizeX,
-		int32 InBlockSizeY,
-		int32 InBlockSizeZ,
-		int32 InBlockBytes,
-		int32 InNumComponents,
-		bool  InSupported);
-
-	const TCHAR*				Name;
-	EPixelFormat				UnrealFormat;
-	int32						BlockSizeX;
-	int32						BlockSizeY;
-	int32						BlockSizeZ;
-	int32						BlockBytes;
-	int32						NumComponents;
-
-	/** Per platform cabilities for the format */
-	EPixelFormatCapabilities	Capabilities{ EPixelFormatCapabilities::None };
-
-	/** Platform specific converted format */
-	uint32						PlatformFormat{ 0 };
-
-	/** Whether the texture format is supported on the current platform/ rendering combination	*/
-	uint8						Supported : 1;
-	uint8						bIs24BitUnormDepthStencil : 1;	// If false, 32 bit float is assumed
-};
-
-extern RHI_API FPixelFormatInfo GPixelFormats[PF_MAX];		// Maps members of EPixelFormat to a FPixelFormatInfo describing the format.
-
 /** Initialize the 'best guess' pixel format capabilities. Platform formats and support must be filled out before calling this. */
 extern RHI_API void RHIInitDefaultPixelFormatCapabilities();
 
@@ -896,10 +861,10 @@ inline bool RHIIsTypedUAVStoreSupported(EPixelFormat InFormat)
 	return RHIPixelFormatHasCapabilities(InFormat, EPixelFormatCapabilities::TypedUAVStore);
 }
 
-//
-//	CalculateImageBytes
-//
-
+/**
+* Returns the memory required to store an image in the given pixel format (EPixelFormat). Use
+* GPixelFormats[Format].Get2D/3DImageSizeInBytes instead, unless you need PF_A1.
+*/
 extern RHI_API SIZE_T CalculateImageBytes(uint32 SizeX,uint32 SizeY,uint32 SizeZ,uint8 Format);
 
 /** Called once per frame only from within an RHI. */
diff --git a/Engine/Source/Runtime/RenderCore/Private/RenderUtils.cpp b/Engine/Source/Runtime/RenderCore/Private/RenderUtils.cpp
--- a/Engine/Source/Runtime/RenderCore/Private/RenderUtils.cpp
+++ b/Engine/Source/Runtime/RenderCore/Private/RenderUtils.cpp
@@ -736,6 +736,7 @@ FArchive& operator<<(FArchive& Ar,FPackedPosition& N)
 
 void CalcMipMapExtent3D( uint32 TextureSizeX, uint32 TextureSizeY, uint32 TextureSizeZ, EPixelFormat Format, uint32 MipIndex, uint32& OutXExtent, uint32& OutYExtent, uint32& OutZExtent )
 {
+	// UE-159189 to explain/fix why this forces min size of block size
 	OutXExtent = FMath::Max<uint32>(TextureSizeX >> MipIndex, GPixelFormats[Format].BlockSizeX);
 	OutYExtent = FMath::Max<uint32>(TextureSizeY >> MipIndex, GPixelFormats[Format].BlockSizeY);
 	OutZExtent = FMath::Max<uint32>(TextureSizeZ >> MipIndex, GPixelFormats[Format].BlockSizeZ);
@@ -743,84 +744,40 @@ void CalcMipMapExtent3D( uint32 TextureSizeX, uint32 TextureSizeY, uint32 Textur
 
 SIZE_T CalcTextureMipMapSize3D( uint32 TextureSizeX, uint32 TextureSizeY, uint32 TextureSizeZ, EPixelFormat Format, uint32 MipIndex )
 {
-	uint32 XExtent;
-	uint32 YExtent;
-	uint32 ZExtent;
-	CalcMipMapExtent3D(TextureSizeX, TextureSizeY, TextureSizeZ, Format, MipIndex, XExtent, YExtent, ZExtent);
-
-	// Offset MipExtent to round up result
-	XExtent += GPixelFormats[Format].BlockSizeX - 1;
-	YExtent += GPixelFormats[Format].BlockSizeY - 1;
-	ZExtent += GPixelFormats[Format].BlockSizeZ - 1;
-
-	const uint32 XPitch = (XExtent / GPixelFormats[Format].BlockSizeX) * GPixelFormats[Format].BlockBytes;
-	const uint32 NumRows = YExtent / GPixelFormats[Format].BlockSizeY;
-	const uint32 NumLayers = ZExtent / GPixelFormats[Format].BlockSizeZ;
-
-	return static_cast<SIZE_T>(NumLayers) * NumRows * XPitch;
+	return GPixelFormats[Format].Get3DTextureMipSizeInBytes(TextureSizeX, TextureSizeY, TextureSizeZ, MipIndex);
 }
 
 SIZE_T CalcTextureSize3D( uint32 SizeX, uint32 SizeY, uint32 SizeZ, EPixelFormat Format, uint32 MipCount )
 {
-	SIZE_T Size = 0;
-	for ( uint32 MipIndex=0; MipIndex < MipCount; ++MipIndex )
-	{
-		Size += CalcTextureMipMapSize3D(SizeX,SizeY,SizeZ,Format,MipIndex);
-	}
-	return Size;
+	return GPixelFormats[Format].Get3DTextureSizeInBytes(SizeX, SizeY, SizeZ, MipCount);
 }
 
 FIntPoint CalcMipMapExtent( uint32 TextureSizeX, uint32 TextureSizeY, EPixelFormat Format, uint32 MipIndex )
 {
+	// UE-159189 to explain/fix why this forces min size of block size
 	return FIntPoint(FMath::Max<uint32>(TextureSizeX >> MipIndex, GPixelFormats[Format].BlockSizeX), FMath::Max<uint32>(TextureSizeY >> MipIndex, GPixelFormats[Format].BlockSizeY));
 }
 
 SIZE_T CalcTextureMipWidthInBlocks(uint32 TextureSizeX, EPixelFormat Format, uint32 MipIndex)
 {
-	// note: FTexture2DResource applies a 2 block min for PVRTC that is not applied here
-	const uint32 BlockSizeX = GPixelFormats[Format].BlockSizeX;
-	if (BlockSizeX > 0)
-	{
-		const uint32 WidthInTexels = FMath::Max<uint32>(TextureSizeX >> MipIndex, 1);
-		const uint32 WidthInBlocks = (WidthInTexels + BlockSizeX - 1) / BlockSizeX;
-		return WidthInBlocks;
-	}
-	else
-	{
-		return 0;
-	}
+	const uint32 WidthInTexels = FMath::Max<uint32>(TextureSizeX >> MipIndex, 1);
+	return GPixelFormats[Format].GetBlockCountForWidth(WidthInTexels);
 }
 
 SIZE_T CalcTextureMipHeightInBlocks(uint32 TextureSizeY, EPixelFormat Format, uint32 MipIndex)
 {
-	const uint32 BlockSizeY = GPixelFormats[Format].BlockSizeY;
-	if (BlockSizeY > 0)
-	{
-		const uint32 HeightInTexels = FMath::Max<uint32>(TextureSizeY >> MipIndex, 1);
-		const uint32 HeightInBlocks = (HeightInTexels + BlockSizeY - 1) / BlockSizeY;
-		return HeightInBlocks;
-	}
-	else
-	{
-		return 0;
-	}
+	const uint32 HeightInTexels = FMath::Max<uint32>(TextureSizeY >> MipIndex, 1);
+	return GPixelFormats[Format].GetBlockCountForHeight(HeightInTexels);
 }
 
 SIZE_T CalcTextureMipMapSize( uint32 TextureSizeX, uint32 TextureSizeY, EPixelFormat Format, uint32 MipIndex )
 {
-	const SIZE_T WidthInBlocks = CalcTextureMipWidthInBlocks(TextureSizeX, Format, MipIndex);
-	const SIZE_T HeightInBlocks = CalcTextureMipHeightInBlocks(TextureSizeY, Format, MipIndex);
-	return WidthInBlocks * HeightInBlocks * GPixelFormats[Format].BlockBytes;
+	return GPixelFormats[Format].Get2DTextureMipSizeInBytes(TextureSizeX, TextureSizeY, MipIndex);
 }
 
 SIZE_T CalcTextureSize( uint32 SizeX, uint32 SizeY, EPixelFormat Format, uint32 MipCount )
 {
-	SIZE_T Size = 0;
-	for ( uint32 MipIndex=0; MipIndex < MipCount; ++MipIndex )
-	{
-		Size += CalcTextureMipMapSize(SizeX,SizeY,Format,MipIndex);
-	}
-	return Size;
+	return GPixelFormats[Format].Get2DTextureSizeInBytes(SizeX, SizeY, MipCount);
 }
 
 void CopyTextureData2D(const void* Source,void* Dest,uint32 SizeY,EPixelFormat Format,uint32 SourceStride,uint32 DestStride)
@@ -864,33 +821,6 @@ void CopyTextureData2D(const void* Source,void* Dest,uint32 SizeY,EPixelFormat F
 	}
 }
 
-/** Helper functions for text output of texture properties... */
-#ifndef CASE_ENUM_TO_TEXT
-#define CASE_ENUM_TO_TEXT(txt) case txt: return TEXT(#txt);
-#endif
-
-#ifndef TEXT_TO_ENUM
-#define TEXT_TO_ENUM(eVal, txt)		if (FCString::Stricmp(TEXT(#eVal), txt) == 0)	return eVal;
-#endif
-
-const TCHAR* GetPixelFormatString(EPixelFormat InPixelFormat)
-{
-	switch (InPixelFormat)
-	{
-		FOREACH_ENUM_EPIXELFORMAT(CASE_ENUM_TO_TEXT)
-	default:
-		return TEXT("PF_Unknown");
-	}	
-}
-
-EPixelFormat GetPixelFormatFromString(const TCHAR* InPixelFormatStr)
-{
-#define TEXT_TO_PIXELFORMAT(f) TEXT_TO_ENUM(f, InPixelFormatStr);
-	FOREACH_ENUM_EPIXELFORMAT(TEXT_TO_PIXELFORMAT)
-#undef TEXT_TO_PIXELFORMAT
-	return PF_Unknown;
-}
-
 EPixelFormatChannelFlags GetPixelFormatValidChannels(EPixelFormat InPixelFormat)
 {
 	static constexpr EPixelFormatChannelFlags PixelFormatToChannelFlags[] =
diff --git a/Engine/Source/Runtime/RenderCore/Public/RenderUtils.h b/Engine/Source/Runtime/RenderCore/Public/RenderUtils.h
--- a/Engine/Source/Runtime/RenderCore/Public/RenderUtils.h
+++ b/Engine/Source/Runtime/RenderCore/Public/RenderUtils.h
@@ -348,6 +348,8 @@ enum EConstructTextureFlags
 /**
  * Calculates the amount of memory used for a single mip-map of a texture 3D.
  *
+ * Use GPixelFormats[Format].Get3DTextureMipSizeInBytes() instead.
+ * 
  * @param TextureSizeX		Number of horizontal texels (for the base mip-level)
  * @param TextureSizeY		Number of vertical texels (for the base mip-level)
  * @param TextureSizeZ		Number of slices (for the base mip-level)
@@ -359,6 +361,8 @@ RENDERCORE_API SIZE_T CalcTextureMipMapSize3D( uint32 TextureSizeX, uint32 Textu
 /**
  * Calculates the extent of a mip.
  *
+ * Incorrectly forces min mip size to be block dimensions: UE-159189
+ * 
  * @param TextureSizeX		Number of horizontal texels (for the base mip-level)
  * @param TextureSizeY		Number of vertical texels (for the base mip-level)
  * @param TextureSizeZ		Number of depth texels (for the base mip-level)
@@ -373,6 +377,8 @@ RENDERCORE_API void CalcMipMapExtent3D( uint32 TextureSizeX, uint32 TextureSizeY
 /**
  * Calculates the extent of a mip.
  *
+ * Incorrectly forces min mip size to be block dimensions: UE-159189
+ * 
  * @param TextureSizeX		Number of horizontal texels (for the base mip-level)
  * @param TextureSizeY		Number of vertical texels (for the base mip-level)
  * @param Format	Texture format
@@ -387,6 +393,7 @@ RENDERCORE_API FIntPoint CalcMipMapExtent( uint32 TextureSizeX, uint32 TextureSi
  * @param Format			Texture format
  * @param MipIndex			The index of the mip-map to compute the size of.
  */
+UE_DEPRECATED(5.1, "See GPixelFormats in PixelFormat.h for analogous functions")
 RENDERCORE_API SIZE_T CalcTextureMipWidthInBlocks(uint32 TextureSizeX, EPixelFormat Format, uint32 MipIndex);
 
 /**
@@ -396,10 +403,13 @@ RENDERCORE_API SIZE_T CalcTextureMipWidthInBlocks(uint32 TextureSizeX, EPixelFor
  * @param Format			Texture format
  * @param MipIndex			The index of the mip-map to compute the size of.
  */
+UE_DEPRECATED(5.1, "See GPixelFormats in PixelFormat.h for analogous functions")
 RENDERCORE_API SIZE_T CalcTextureMipHeightInBlocks(uint32 TextureSizeY, EPixelFormat Format, uint32 MipIndex);
 
 /**
  * Calculates the amount of memory used for a single mip-map of a texture.
+ * 
+ * Use GPixelFormats[Format].Get2DTextureMipSizeInBytes() instead.
  *
  * @param TextureSizeX		Number of horizontal texels (for the base mip-level)
  * @param TextureSizeY		Number of vertical texels (for the base mip-level)
@@ -411,6 +421,8 @@ RENDERCORE_API SIZE_T CalcTextureMipMapSize( uint32 TextureSizeX, uint32 Texture
 /**
  * Calculates the amount of memory used for a texture.
  *
+ * Use GPixelFormats[Format].Get2DTextureSizeInBytes() instead.
+ * 
  * @param SizeX		Number of horizontal texels (for the base mip-level)
  * @param SizeY		Number of vertical texels (for the base mip-level)
  * @param Format	Texture format
@@ -421,6 +433,8 @@ RENDERCORE_API SIZE_T CalcTextureSize( uint32 SizeX, uint32 SizeY, EPixelFormat
 /**
  * Calculates the amount of memory used for a texture.
  *
+ * Use GPixelFormats[Format].Get3DTextureSizeInBytes() instead.
+ * 
  * @param SizeX		Number of horizontal texels (for the base mip-level)
  * @param SizeY		Number of vertical texels (for the base mip-level)
  * @param SizeY		Number of depth texels (for the base mip-level)
@@ -440,19 +454,6 @@ RENDERCORE_API SIZE_T CalcTextureSize3D( uint32 SizeX, uint32 SizeY, uint32 Size
  */
 RENDERCORE_API void CopyTextureData2D(const void* Source,void* Dest,uint32 SizeY,EPixelFormat Format,uint32 SourceStride,uint32 DestStride);
 
-/**
- * enum to string
- *
- * @return e.g. "PF_B8G8R8A8"
- */
-RENDERCORE_API const TCHAR* GetPixelFormatString(EPixelFormat InPixelFormat);
-/**
- * string to enum (not case sensitive)
- *
- * @param InPixelFormatStr e.g. "PF_B8G8R8A8", must not not be 0
- */
-RENDERCORE_API EPixelFormat GetPixelFormatFromString(const TCHAR* InPixelFormatStr);
-
 /**
  *  Returns the valid channels for this pixel format
  * 
