Move the pixel format validator over so that it's initialized after the pixel formats array.
Fixes CIS boot issue

diff --git a/Engine/Source/Runtime/Core/Private/Misc/PixelFormat.cpp b/Engine/Source/Runtime/Core/Private/Misc/PixelFormat.cpp
--- a/Engine/Source/Runtime/Core/Private/Misc/PixelFormat.cpp
+++ b/Engine/Source/Runtime/Core/Private/Misc/PixelFormat.cpp
@@ -251,3 +251,14 @@ EPixelFormat GetPixelFormatFromString(const TCHAR* InPixelFormatStr)
 		return PF_Unknown;
 }
 
+static struct FValidatePixelFormats
+{
+	FValidatePixelFormats()
+	{
+		for (int32 Index = 0; Index < UE_ARRAY_COUNT(GPixelFormats); ++Index)
+		{
+			// Make sure GPixelFormats has an entry for every unreal format
+			checkf((EPixelFormat)Index == GPixelFormats[Index].UnrealFormat, TEXT("Missing entry for EPixelFormat %d"), (int32)Index);
+		}
+	}
+} ValidatePixelFormats;
diff --git a/Engine/Source/Runtime/RHI/Private/RHI.cpp b/Engine/Source/Runtime/RHI/Private/RHI.cpp
--- a/Engine/Source/Runtime/RHI/Private/RHI.cpp
+++ b/Engine/Source/Runtime/RHI/Private/RHI.cpp
@@ -2687,18 +2687,6 @@ void RHIInitDefaultPixelFormatCapabilities()
 	}
 }
 
-static struct FValidatePixelFormats
-{
-	FValidatePixelFormats()
-	{
-		for (int32 Index = 0; Index < UE_ARRAY_COUNT(GPixelFormats); ++Index)
-		{
-			// Make sure GPixelFormats has an entry for every unreal format
-			checkf((EPixelFormat)Index == GPixelFormats[Index].UnrealFormat, TEXT("Missing entry for EPixelFormat %d"), (int32)Index);
-		}
-	}
-} ValidatePixelFormats;
-
 //
 //	CalculateImageBytes
 //
