Added a few IWYU pragma comments.

diff --git a/Engine/Source/Developer/DesktopPlatform/Private/DesktopPlatformPrivate.h b/Engine/Source/Developer/DesktopPlatform/Private/DesktopPlatformPrivate.h
--- a/Engine/Source/Developer/DesktopPlatform/Private/DesktopPlatformPrivate.h
+++ b/Engine/Source/Developer/DesktopPlatform/Private/DesktopPlatformPrivate.h
@@ -4,6 +4,7 @@
 
 #include "CoreMinimal.h"
 
+// IWYU pragma: begin_exports
 #if PLATFORM_WINDOWS
 #include "Windows/DesktopPlatformWindows.h"
 #elif PLATFORM_MAC
@@ -13,5 +14,6 @@
 #else
 #include "DesktopPlatformStub.h"
 #endif
+// IWYU pragma: end_exports
 
 DECLARE_LOG_CATEGORY_EXTERN(LogDesktopPlatform, Log, All);
diff --git a/Engine/Source/Runtime/Core/Public/GenericPlatform/GenericPlatformMemory.h b/Engine/Source/Runtime/Core/Public/GenericPlatform/GenericPlatformMemory.h
--- a/Engine/Source/Runtime/Core/Public/GenericPlatform/GenericPlatformMemory.h
+++ b/Engine/Source/Runtime/Core/Public/GenericPlatform/GenericPlatformMemory.h
@@ -7,15 +7,19 @@
 
 #pragma once
 
+#include "CoreFwd.h" // IWYU pragma: keep
 #include "CoreTypes.h"
-#include "CoreFwd.h"
-#include <wchar.h>
-#include <string.h>
+#include "HAL/PlatformString.h"
 
-struct FPlatformMemoryStats;
+#include <string.h>
+#include <wchar.h>
 
+class FMalloc;
+class FOutputDevice;
+class FString;
 /** Holds generic memory stats, internally implemented as a map. */
 struct FGenericMemoryStats;
+struct FPlatformMemoryStats;
 
 /**
 * Platform-dependent "bucket" for memory size, where Default is the normal, or possibly the largest.
diff --git a/Engine/Source/Runtime/Core/Public/Unix/UnixPlatformTLS.h b/Engine/Source/Runtime/Core/Public/Unix/UnixPlatformTLS.h
--- a/Engine/Source/Runtime/Core/Public/Unix/UnixPlatformTLS.h
+++ b/Engine/Source/Runtime/Core/Public/Unix/UnixPlatformTLS.h
@@ -7,7 +7,7 @@
 #pragma once
 
 #include "CoreTypes.h"
-#include "Unix/UnixSystemIncludes.h"
+#include "Unix/UnixSystemIncludes.h" // IWYU pragma: export
 #include "Misc/AssertionMacros.h"
 #include "GenericPlatform/GenericPlatformTLS.h"
 
