instrumented ParallelFor for TasksInsights. now markers for ParallelFor tasks show task lifetime events, and most importantly point to which ParallelFor they belong. ParallelFor markers don't point to their tasks yet.

diff --git a/Engine/Source/Runtime/Core/Public/Async/ParallelFor.h b/Engine/Source/Runtime/Core/Public/Async/ParallelFor.h
--- a/Engine/Source/Runtime/Core/Public/Async/ParallelFor.h
+++ b/Engine/Source/Runtime/Core/Public/Async/ParallelFor.h
@@ -108,6 +108,7 @@ namespace ParallelForImpl
 	inline void ParallelForInternal(const TCHAR* DebugName, int32 Num, int32 MinBatchSize, BodyType Body, PreWorkType CurrentThreadWorkToDoBeforeHelping, EParallelForFlags Flags, const TArrayView<ContextType>& Contexts)
 	{
 		SCOPE_CYCLE_COUNTER(STAT_ParallelFor);
+		TRACE_CPUPROFILER_EVENT_SCOPE(ParallelFor);
 		check(Num >= 0);
 
 		int32 NumWorkers = GetNumberOfThreadTasks(Num, MinBatchSize, Flags);
@@ -179,6 +180,12 @@ namespace ParallelForImpl
 			Priority = LowLevelTasks::ETaskPriority::BackgroundNormal;
 		}
 		
+		struct FTracedTask
+		{
+			LowLevelTasks::FTask Task;
+			std::atomic<TaskTrace::FId> TraceId = TaskTrace::InvalidId;
+		};
+
 		//shared data between tasks
 		struct alignas(PLATFORM_CACHE_LINE_SIZE) FParallelForData : public TConcurrentLinearObject<FParallelForData, FTaskGraphBlockAllocationTag>, public FThreadSafeRefCountedObject
 		{
@@ -198,6 +205,15 @@ namespace ParallelForImpl
 #endif
 				LLM(InheritedLLMTag = FLowLevelMemTracker::bIsDisabled ? nullptr : FLowLevelMemTracker::Get().GetActiveTagData(ELLMTracker::Default));
 			}
+
+			~FParallelForData()
+			{
+				for (FTracedTask& Task : Tasks)
+				{
+					TaskTrace::Destroyed(Task.TraceId);
+				}
+			}
+
 			std::atomic_int BatchItem  { 0 };
 			std::atomic_int IncompleteBatches { 0 };
 			int32 Num;
@@ -210,7 +226,8 @@ namespace ParallelForImpl
 			const TArrayView<ContextType>& Contexts;
 			const BodyType& Body;
 			FEventRef& FinishedSignal;
-			TArray<LowLevelTasks::FTask, TConcurrentLinearArrayAllocator<FTaskGraphBlockAllocationTag>> Tasks;
+
+			TArray<FTracedTask, TConcurrentLinearArrayAllocator<FTaskGraphBlockAllocationTag>> Tasks;
 		};
 		using FDataHandle = TRefCountPtr<FParallelForData>;
 
@@ -258,11 +275,25 @@ namespace ParallelForImpl
 
 				if(DebugName == nullptr)
 				{
-					LowLevelTasks::FTask& Task = Data->Tasks[WorkerIndex];
+					LowLevelTasks::FTask& Task = Data->Tasks[WorkerIndex].Task;
 					DebugName = Task.GetDebugName();
 				}
-				TRACE_CPUPROFILER_EVENT_SCOPE_TEXT(DebugName);
+				
+				TaskTrace::FId TraceId = TaskTrace::InvalidId;
+				if (!bIsMaster)
+				{
+					TraceId = Data->Tasks[WorkerIndex].TraceId;
+					TaskTrace::Started(TraceId);
+				}
+				ON_SCOPE_EXIT
+				{
+					if (!bIsMaster)
+					{
+						TaskTrace::Completed(TraceId);
+					}
+				};
 
+				TRACE_CPUPROFILER_EVENT_SCOPE_TEXT(DebugName);
 
 				auto Now = [] { return FTimespan::FromSeconds(FPlatformTime::Seconds()); };
 				FTimespan Start = FTimespan::MinValue();
@@ -336,13 +367,22 @@ namespace ParallelForImpl
 
 			static void LaunchTask(const TCHAR* DebugName, FDataHandle&& InData, int32 InWorkerIndex, LowLevelTasks::ETaskPriority InPriority)
 			{
-				LowLevelTasks::FTask& Task = InData->Tasks[InWorkerIndex];
+				FTracedTask& TracedTask = InData->Tasks[InWorkerIndex];
 				if(DebugName == nullptr)
 				{
-					DebugName = Task.GetDebugName();
+					DebugName = TracedTask.Task.GetDebugName();
 				}
-				Task.Init(DebugName, InPriority, FParallelExecutor(MoveTemp(InData), InWorkerIndex, InPriority));
-				verify(LowLevelTasks::TryLaunch(Task, LowLevelTasks::EQueuePreference::GlobalQueuePreference));
+
+				if (TracedTask.TraceId != TaskTrace::InvalidId) // reused task
+				{
+					TaskTrace::Destroyed(TracedTask.TraceId);
+				}
+				TracedTask.TraceId = TaskTrace::GenerateTaskId();
+				TaskTrace::Launched(TracedTask.TraceId, DebugName, false, ENamedThreads::AnyThread);
+
+				TracedTask.Task.Init(DebugName, InPriority, FParallelExecutor(MoveTemp(InData), InWorkerIndex, InPriority));
+				verify(LowLevelTasks::TryLaunch(TracedTask.Task, LowLevelTasks::EQueuePreference::GlobalQueuePreference));
+
 			}
 		};
 
@@ -364,9 +404,9 @@ namespace ParallelForImpl
 		// Cancel tasks that have not yet launched since they will otherwise waste time on worker threads
 		{
 			TRACE_CPUPROFILER_EVENT_SCOPE(ParallelFor.Cancel);
-			for (LowLevelTasks::FTask& Task : LocalExecutor.GetData()->Tasks)
+			for (FTracedTask& TracedTask : LocalExecutor.GetData()->Tasks)
 			{
-				Task.TryCancel(LowLevelTasks::ECancellationFlags::PrelaunchCancellation);
+				TracedTask.Task.TryCancel(LowLevelTasks::ECancellationFlags::PrelaunchCancellation);
 			}
 		}
 
@@ -403,7 +443,7 @@ namespace ParallelForImpl
 **/
 inline void ParallelFor(int32 Num, TFunctionRef<void(int32)> Body, bool bForceSingleThread, bool bPumpRenderingThread=false)
 {
-	ParallelForImpl::ParallelForInternal(TEXT("ParallelFor"), Num, 1, Body, [](){},
+	ParallelForImpl::ParallelForInternal(TEXT("ParallelFor Task"), Num, 1, Body, [](){},
 		(bForceSingleThread ? EParallelForFlags::ForceSingleThread : EParallelForFlags::None) | 
 		(bPumpRenderingThread ? EParallelForFlags::PumpRenderingThread : EParallelForFlags::None), TArrayView<TYPE_OF_NULLPTR>());
 }
@@ -418,7 +458,7 @@ inline void ParallelFor(int32 Num, TFunctionRef<void(int32)> Body, bool bForceSi
 template<typename FunctionType>
 inline void ParallelForTemplate(int32 Num, const FunctionType& Body, EParallelForFlags Flags = EParallelForFlags::None)
 {
-	ParallelForImpl::ParallelForInternal(TEXT("ParallelFor"), Num, 1, Body, [](){}, Flags, TArrayView<TYPE_OF_NULLPTR>());
+	ParallelForImpl::ParallelForInternal(TEXT("ParallelFor Task"), Num, 1, Body, [](){}, Flags, TArrayView<TYPE_OF_NULLPTR>());
 }
 
 /**
@@ -448,7 +488,7 @@ inline void ParallelForTemplate(const TCHAR* DebugName, int32 Num, int32 MinBatc
 **/
 inline void ParallelFor(int32 Num, TFunctionRef<void(int32)> Body, EParallelForFlags Flags = EParallelForFlags::None)
 {
-	ParallelForImpl::ParallelForInternal(TEXT("ParallelFor"), Num, 1, Body, [](){}, Flags, TArrayView<TYPE_OF_NULLPTR>());
+	ParallelForImpl::ParallelForInternal(TEXT("ParallelFor Task"), Num, 1, Body, [](){}, Flags, TArrayView<TYPE_OF_NULLPTR>());
 }
 
 /** 
@@ -478,7 +518,7 @@ inline void ParallelFor(const TCHAR* DebugName, int32 Num, int32 MinBatchSize, T
 **/
 inline void ParallelForWithPreWork(int32 Num, TFunctionRef<void(int32)> Body, TFunctionRef<void()> CurrentThreadWorkToDoBeforeHelping, bool bForceSingleThread, bool bPumpRenderingThread = false)
 {
-	ParallelForImpl::ParallelForInternal(TEXT("ParallelFor"), Num, 1, Body, CurrentThreadWorkToDoBeforeHelping,
+	ParallelForImpl::ParallelForInternal(TEXT("ParallelFor Task"), Num, 1, Body, CurrentThreadWorkToDoBeforeHelping,
 		(bForceSingleThread ? EParallelForFlags::ForceSingleThread : EParallelForFlags::None) |
 		(bPumpRenderingThread ? EParallelForFlags::PumpRenderingThread : EParallelForFlags::None), TArrayView<TYPE_OF_NULLPTR>());
 }
@@ -493,7 +533,7 @@ inline void ParallelForWithPreWork(int32 Num, TFunctionRef<void(int32)> Body, TF
 **/
 inline void ParallelForWithPreWork(int32 Num, TFunctionRef<void(int32)> Body, TFunctionRef<void()> CurrentThreadWorkToDoBeforeHelping, EParallelForFlags Flags = EParallelForFlags::None)
 {
-	ParallelForImpl::ParallelForInternal(TEXT("ParallelFor"), Num, 1, Body, CurrentThreadWorkToDoBeforeHelping, Flags, TArrayView<TYPE_OF_NULLPTR>());
+	ParallelForImpl::ParallelForInternal(TEXT("ParallelFor Task"), Num, 1, Body, CurrentThreadWorkToDoBeforeHelping, Flags, TArrayView<TYPE_OF_NULLPTR>());
 }
 
 /** 
@@ -600,7 +640,7 @@ inline void ParallelForWithTaskContext(TArray<ContextType, ContextAllocatorType>
 		{
 			new(&OutContexts[ContextIndex]) ContextType(ContextConstructor(ContextIndex, NumContexts));
 		}
-		ParallelForImpl::ParallelForInternal(TEXT("ParallelFor"), Num, 1, Body, [](){}, Flags, TArrayView<ContextType>(OutContexts));
+		ParallelForImpl::ParallelForInternal(TEXT("ParallelFor Task"), Num, 1, Body, [](){}, Flags, TArrayView<ContextType>(OutContexts));
 	}
 }
 
@@ -622,7 +662,7 @@ inline void ParallelForWithTaskContext(TArray<ContextType, ContextAllocatorType>
 		const int32 NumContexts = ParallelForImpl::GetNumberOfThreadTasks(Num, 1, Flags);
 		OutContexts.Reset();
 		OutContexts.AddDefaulted(NumContexts);
-		ParallelForImpl::ParallelForInternal(TEXT("ParallelFor"), Num, 1, Body, [](){}, Flags, TArrayView<ContextType>(OutContexts));
+		ParallelForImpl::ParallelForInternal(TEXT("ParallelFor Task"), Num, 1, Body, [](){}, Flags, TArrayView<ContextType>(OutContexts));
 	}
 }
 
@@ -696,7 +736,7 @@ inline void ParallelForWithExistingTaskContext(TArrayView<ContextType> Contexts,
 {
 	if (Num > 0)
 	{
-		ParallelForImpl::ParallelForInternal(TEXT("ParallelFor"), Num, MinBatchSize, Body, [](){}, Flags, Contexts);
+		ParallelForImpl::ParallelForInternal(TEXT("ParallelFor Task"), Num, MinBatchSize, Body, [](){}, Flags, Contexts);
 	}
 }
 
