Modular builds: ModuleManager: Fixed DebugGame configuration which failed to find the main game's module.

The problem was caused by the fact that in DebugGame we end up with two manifests in the executable's directory: one for the engine's modules (Development configuration) and one for the game's modules (DebugGame configuration). We need to load both to be able to find all the modules.

The editor has a separate path for mounting game's directories with executables but it would fail in the standalone game.

#rnx

diff --git a/Engine/Source/Runtime/Core/Private/Modules/ModuleManager.cpp b/Engine/Source/Runtime/Core/Private/Modules/ModuleManager.cpp
--- a/Engine/Source/Runtime/Core/Private/Modules/ModuleManager.cpp
+++ b/Engine/Source/Runtime/Core/Private/Modules/ModuleManager.cpp
@@ -1129,6 +1129,17 @@ void FModuleManager::FindModulePaths(const TCHAR* NamePattern, TMap<FName, FStri
 
 		// Add the engine directory to the cache - only needs to be cached once as the contents do not change at runtime
 		FindModulePathsInDirectory(FPlatformProcess::GetModulesDirectory(), false, ModulePathsCache);
+
+#if !WITH_EDITOR
+		// DebugGame is a hybrid configuration where the engine is in Development and the game Debug.
+		// As such, it generates two separate module manifests, one for the engine modules (Development) and
+		// one for the game (Debug). We need to load both as otherwise some modules won't be found.
+		// We exclude this code in the editor as it handles the game's libraries separately via PendingGameBinariesDirectories.
+		if (FApp::GetBuildConfiguration() == EBuildConfiguration::DebugGame)
+		{
+			FindModulePathsInDirectory(FPlatformProcess::GetModulesDirectory(), true, ModulePathsCache);
+		}
+#endif
 	}
 
 	// If any entries have been added to the PendingEngineBinariesDirectories or PendingGameBinariesDirectories arrays, add any
