Fixed an assertion failure if a scope heartbeat guard failed to suspend due to uninitialized singleton and then later tried to resume it once it was initialized.

diff --git a/Engine/Source/Runtime/Core/Public/HAL/ThreadHeartBeat.h b/Engine/Source/Runtime/Core/Public/HAL/ThreadHeartBeat.h
--- a/Engine/Source/Runtime/Core/Public/HAL/ThreadHeartBeat.h
+++ b/Engine/Source/Runtime/Core/Public/HAL/ThreadHeartBeat.h
@@ -267,20 +267,26 @@ struct FSlowHeartBeatScope
 {
 private:
 	bool bSuspendedAllThreads;
+	bool bSuspended;
 public:
 	FORCEINLINE FSlowHeartBeatScope(bool bAllThreads = false)
 		: bSuspendedAllThreads(bAllThreads)
+		, bSuspended(false)
 	{
 		if (FThreadHeartBeat* HB = FThreadHeartBeat::GetNoInit())
 		{
+			bSuspended = true;
 			HB->SuspendHeartBeat(bSuspendedAllThreads);
 		}
 	}
 	FORCEINLINE ~FSlowHeartBeatScope()
 	{
-		if (FThreadHeartBeat* HB = FThreadHeartBeat::GetNoInit())
+		if (bSuspended)
 		{
-			HB->ResumeHeartBeat(bSuspendedAllThreads);
+			if (FThreadHeartBeat* HB = FThreadHeartBeat::GetNoInit())
+			{
+				HB->ResumeHeartBeat(bSuspendedAllThreads);
+			}
 		}
 	}
 };
@@ -288,19 +294,26 @@ public:
 /** Simple scope object to put at the top of a function to monitor it completes in a timely fashion */
 struct FFunctionHeartBeatScope
 {
+private:
+	bool bStartedMonitor;
 public:
 	FORCEINLINE FFunctionHeartBeatScope()
+		: bStartedMonitor(false)
 	{
 		if (FThreadHeartBeat* HB = FThreadHeartBeat::GetNoInit())
 		{
+			bStartedMonitor = true;
 			HB->MonitorFunctionStart();
 		}
 	}
 	FORCEINLINE ~FFunctionHeartBeatScope()
 	{
-		if (FThreadHeartBeat* HB = FThreadHeartBeat::GetNoInit())
+		if (bStartedMonitor)
 		{
-			HB->MonitorFunctionEnd();
+			if (FThreadHeartBeat* HB = FThreadHeartBeat::GetNoInit())
+			{
+				HB->MonitorFunctionEnd();
+			}
 		}
 	}
 };
