Simplify code flow in ConvertToUTF8.

diff --git a/Engine/Source/Runtime/Core/Private/GenericPlatform/GenericPlatformString.cpp b/Engine/Source/Runtime/Core/Private/GenericPlatform/GenericPlatformString.cpp
--- a/Engine/Source/Runtime/Core/Private/GenericPlatform/GenericPlatformString.cpp
+++ b/Engine/Source/Runtime/Core/Private/GenericPlatform/GenericPlatformString.cpp
@@ -211,8 +211,6 @@ namespace UE::Core::Private
 		{
 			DestBufferType DestStartingPosition = Dest;
 
-			uint32 HighSurrogate = MAX_uint32;
-
 			for (;;)
 			{
 				if (SourceLen <= 0)
@@ -220,26 +218,15 @@ namespace UE::Core::Private
 					return UE_PTRDIFF_TO_INT32(Dest - DestStartingPosition);
 				}
 
-				const bool bHighSurrogateIsSet = HighSurrogate != MAX_uint32;
 				uint32 Codepoint = static_cast<uint32>(*Source++);
 				--SourceLen;
 
 				// Check if this character is a high-surrogate
 				if (IsHighSurrogate(Codepoint))
 				{
-					if (bHighSurrogateIsSet)
-					{
-						// Already have a high-surrogate in this pair - write our stored value (will be converted into bogus character)
-						if (!WriteCodepointToBuffer(HighSurrogate, Dest, DestLen))
-						{
-							// Could not write data, bail out
-							return -1;
-						}
-					}
-
 					if (SourceLen == 0)
 					{
-						// String ends with lone high-surrogate - write out surrogate (will be converted into bogus character)
+						// String ends with lone high-surrogate - write it out (will be converted into bogus character)
 						if (!WriteCodepointToBuffer(Codepoint, Dest, DestLen))
 						{
 							// Could not write data, bail out
@@ -249,32 +236,17 @@ namespace UE::Core::Private
 						return UE_PTRDIFF_TO_INT32(Dest - DestStartingPosition);
 					}
 
-					// Store our code point for our next character
-					HighSurrogate = Codepoint;
-					continue;
-				}
+					// Read next codepoint
+					uint32 NextCodepoint = static_cast<uint32>(*Source);
 
-				// If our High Surrogate is set, check if this character is the matching low-surrogate
-				if (bHighSurrogateIsSet)
-				{
-					if (IsLowSurrogate(Codepoint))
-					{
-						const uint32 LowSurrogate = Codepoint;
-						// Combine our high and low surrogates together to a single Unicode codepoint
-						Codepoint = EncodeSurrogate((uint16)HighSurrogate, (uint16)LowSurrogate);
-					}
-					else
+					// If it's a low surrogate, combine it with the current high surrogate,
+					// otherwise just leave the high surrogate to be written out by itself (as a bogus character)
+					if (IsLowSurrogate(NextCodepoint))
 					{
-						// Did not find matching low-surrogate, write out a bogus character for our stored HighSurrogate
-						if (!WriteCodepointToBuffer(HighSurrogate, Dest, DestLen))
-						{
-							// Could not write data, bail out
-							return -1;
-						}
+						Codepoint = EncodeSurrogate((uint16)Codepoint, (uint16)NextCodepoint);
+						++Source;
+						--SourceLen;
 					}
-
-					// Reset our high-surrogate now that we've used (or discarded) its value
-					HighSurrogate = MAX_uint32;
 				}
 
 				if (!WriteCodepointToBuffer(Codepoint, Dest, DestLen))
