Reduced FScopeCycleCounter complexity by undoing the description string addition. Replaced the one use case with a dynamic stat string.

diff --git a/Engine/Source/Runtime/Core/Public/HAL/PlatformMisc.h b/Engine/Source/Runtime/Core/Public/HAL/PlatformMisc.h
--- a/Engine/Source/Runtime/Core/Public/HAL/PlatformMisc.h
+++ b/Engine/Source/Runtime/Core/Public/HAL/PlatformMisc.h
@@ -36,6 +36,39 @@
 
 #include "ProfilingDebugging/CpuProfilerTrace.h"
 
+class FScopedNamedEventConditional
+{
+public:
+	FScopedNamedEventConditional(const struct FColor& Color, const TCHAR* Text, bool bCondition)
+		: bStarted(bCondition)
+	{
+		if (bCondition)
+		{
+			FPlatformMisc::BeginNamedEvent(Color, Text);
+		}
+	}
+
+	FScopedNamedEventConditional(const struct FColor& Color, const ANSICHAR* Text, bool bCondition)
+		: bStarted(bCondition)
+	{
+		if (bCondition)
+		{
+			FPlatformMisc::BeginNamedEvent(Color, Text);
+		}
+	}
+
+	~FScopedNamedEventConditional()
+	{
+		if (bStarted)
+		{
+			FPlatformMisc::EndNamedEvent();
+		}
+	}
+
+private:
+	bool bStarted;
+};
+
 class CORE_API FScopedNamedEvent
 {
 public:
@@ -136,6 +169,10 @@ public:
 	FScopedNamedEvent ANONYMOUS_VARIABLE(NamedEvent_)(Color, *FString::Printf(Format, __VA_ARGS__));\
 	TRACE_CPUPROFILER_EVENT_SCOPE_TEXT(*FString::Printf(Format, __VA_ARGS__));
 
+#define SCOPED_NAMED_EVENT_TCHAR_CONDITIONAL(Text, Color, bCondition)\
+	FScopedNamedEventConditional ANONYMOUS_VARIABLE(NamedEvent_)(Color, Text, (bCondition));\
+	TRACE_CPUPROFILER_EVENT_SCOPE_TEXT_CONDITIONAL(Text, (bCondition));
+
 #define SCOPED_PROFILER_COLOR(Color)			 FScopedProfilerColor    ANONYMOUS_VARIABLE(ProfilerColor_##Name##_)(Color);
 
 
@@ -174,6 +211,7 @@ public:
 #define SCOPED_NAMED_EVENT_TCHAR(...)
 #define SCOPED_NAMED_EVENT_TEXT(...)
 #define SCOPED_NAMED_EVENT_F(...)
+#define SCOPED_NAMED_EVENT_TCHAR_CONDITIONAL(...)
 #define SCOPED_PROFILER_COLOR(...)
 
 #endif
diff --git a/Engine/Source/Runtime/Core/Public/ProfilingDebugging/CpuProfilerTrace.h b/Engine/Source/Runtime/Core/Public/ProfilingDebugging/CpuProfilerTrace.h
--- a/Engine/Source/Runtime/Core/Public/ProfilingDebugging/CpuProfilerTrace.h
+++ b/Engine/Source/Runtime/Core/Public/ProfilingDebugging/CpuProfilerTrace.h
@@ -296,6 +296,9 @@ struct FCpuProfilerTrace
 #define TRACE_CPUPROFILER_EVENT_SCOPE_TEXT(Name) \
 	FCpuProfilerTrace::FDynamicEventScope PREPROCESSOR_JOIN(__CpuProfilerEventScope, __LINE__)(Name, true, __FILE__, __LINE__);
 
+#define TRACE_CPUPROFILER_EVENT_SCOPE_TEXT_CONDITIONAL(Name, Condition) \
+	FCpuProfilerTrace::FDynamicEventScope PREPROCESSOR_JOIN(__CpuProfilerEventScope, __LINE__)(Name, (Condition), __FILE__, __LINE__);
+
 #else
 
 #define TRACE_CPUPROFILER_EVENT_DECLARE(DeclName)
@@ -307,5 +310,6 @@ struct FCpuProfilerTrace
 #define TRACE_CPUPROFILER_EVENT_SCOPE_ON_CHANNEL(Name, Channel)
 #define TRACE_CPUPROFILER_EVENT_SCOPE_TEXT_ON_CHANNEL(Name, Channel)
 #define TRACE_CPUPROFILER_EVENT_SCOPE_TEXT(Name)
+#define TRACE_CPUPROFILER_EVENT_SCOPE_TEXT_CONDITIONAL(Name, Condition)
 
 #endif
diff --git a/Engine/Source/Runtime/Core/Public/Stats/Stats.h b/Engine/Source/Runtime/Core/Public/Stats/Stats.h
--- a/Engine/Source/Runtime/Core/Public/Stats/Stats.h
+++ b/Engine/Source/Runtime/Core/Public/Stats/Stats.h
@@ -59,19 +59,9 @@ public:
 	 * Pushes the specified stat onto the hierarchy for this thread. Starts
 	 * the timing of the cycles used
 	 */
-	FORCEINLINE_STATS FScopeCycleCounter( TStatId StatId, EStatFlags StatFlags, bool bAlways = false
-#if CPUPROFILERTRACE_ENABLED
-		// Optional verbose description added for CPU profiler, without affecting the stat name.  Should be invariant
-		// across frames, to avoid bloating memory in the name map used in the CPU profiler.
-		, const TCHAR* OptionalVerboseDescription = nullptr
-#endif
-	)
+	FORCEINLINE_STATS FScopeCycleCounter( TStatId StatId, EStatFlags StatFlags, bool bAlways = false)
 	{
-		Start( StatId, StatFlags, bAlways
-#if CPUPROFILERTRACE_ENABLED
-			, OptionalVerboseDescription
-#endif
-			);
+		Start( StatId, StatFlags, bAlways);
 	}
 
 	FORCEINLINE_STATS FScopeCycleCounter(TStatId StatId, bool bAlways = false)
@@ -255,10 +245,6 @@ FORCEINLINE void StatsMasterEnableSubtract(int32 Value = 1)
 	FScopeCycleCounter StatNamedEventsScope_##Stat(bCondition ? ANSI_TO_PROFILING(#Stat) : nullptr); \
 	SCOPE_CYCLE_COUNTER_TO_TRACE(#Stat, Stat, bCondition);
 
-#define SCOPE_CYCLE_COUNTER_VERBOSE(Stat, VerboseDescription) \
-	FScopeCycleCounter StatNamedEventsScope_##Stat(TStatId(ANSI_TO_PROFILING(#Stat))); \
-	SCOPE_CYCLE_COUNTER_TO_TRACE(#Stat, Stat, true);
-
 #define RETURN_QUICK_DECLARE_CYCLE_STAT(StatId,GroupId) return TStatId(ANSI_TO_PROFILING(#StatId));
 
 #define GET_STATID(Stat) (TStatId(ANSI_TO_PROFILING(#Stat)))
@@ -299,9 +285,6 @@ public:
 #define CONDITIONAL_SCOPE_CYCLE_COUNTER(Stat,bCondition) \
 	FLightweightStatScope LightweightStatScope_##Stat(bCondition ? TEXT(#Stat) : nullptr);
 
-#define SCOPE_CYCLE_COUNTER_VERBOSE(Stat,VerboseDescription) \
-	FLightweightStatScope LightweightStatScope_##Stat(TEXT(#Stat));
-
 #define RETURN_QUICK_DECLARE_CYCLE_STAT(StatId,GroupId) return TStatId();
 #define GET_STATID(Stat) (TStatId())
 
@@ -310,7 +293,6 @@ public:
 #define QUICK_SCOPE_CYCLE_COUNTER(Stat)
 #define DECLARE_SCOPE_CYCLE_COUNTER(CounterName,StatId,GroupId)
 #define CONDITIONAL_SCOPE_CYCLE_COUNTER(Stat,bCondition)
-#define SCOPE_CYCLE_COUNTER_VERBOSE(Stat,VerboseDescription)
 #define RETURN_QUICK_DECLARE_CYCLE_STAT(StatId,GroupId) return TStatId();
 #define GET_STATID(Stat) (TStatId())
 
diff --git a/Engine/Source/Runtime/Core/Public/Stats/Stats2.h b/Engine/Source/Runtime/Core/Public/Stats/Stats2.h
--- a/Engine/Source/Runtime/Core/Public/Stats/Stats2.h
+++ b/Engine/Source/Runtime/Core/Public/Stats/Stats2.h
@@ -1580,13 +1580,7 @@ public:
 	 * Pushes the specified stat onto the hierarchy for this thread. Starts
 	 * the timing of the cycles used
 	 */
-	FORCEINLINE_STATS void Start(TStatId InStatId, EStatFlags InStatFlags, bool bAlways = false
-#if CPUPROFILERTRACE_ENABLED
-		// Optional verbose description added for CPU profiler, without affecting the stat name.  Should be invariant
-		// across frames, to avoid bloating memory in the name map used in the CPU profiler.
-		, const TCHAR* OptionalVerboseDescription = nullptr
-#endif
-		)
+	FORCEINLINE_STATS void Start(TStatId InStatId, EStatFlags InStatFlags, bool bAlways = false)
 	{
 		FMinimalName StatMinimalName = InStatId.GetMinimalName(EMemoryOrder::Relaxed);
 		if (StatMinimalName.IsNone())
@@ -1608,37 +1602,7 @@ public:
 #if CPUPROFILERTRACE_ENABLED
 			if (UE_TRACE_CHANNELEXPR_IS_ENABLED(CpuChannel))
 			{
-				// We only support ANSI strings, since the CPU profiler internally uses that in its event type.  We could
-				// in the future add a separate wide string profiler event, but it involves duplicating a fair amount of
-				// code, so this works for now.
-				if (OptionalVerboseDescription && FCString::IsPureAnsi(OptionalVerboseDescription))
-				{
-					TArray<ANSICHAR, TNonRelocatableInlineAllocator<256>> VerboseNameANSI;
-
-					const char* StatDescriptionANSI = InStatId.GetStatDescriptionANSI();
-					int32 StatLen = FCStringAnsi::Strlen(StatDescriptionANSI);
-					int32 DescriptionLen = FCString::Strlen(OptionalVerboseDescription);
-
-					VerboseNameANSI.AddUninitialized(StatLen + DescriptionLen + 4);		// 3 character separator plus null terminator
-					char* Buffer = VerboseNameANSI.GetData();
-					char* BufferEnd = Buffer + VerboseNameANSI.Num();
-
-					FMemory::Memcpy(Buffer, StatDescriptionANSI, StatLen*sizeof(ANSICHAR));	Buffer += StatLen;
-					FMemory::Memcpy(Buffer, " - ", 3*sizeof(ANSICHAR));						Buffer += 3;
-					for (const TCHAR* DescriptionLetter = OptionalVerboseDescription; *DescriptionLetter; DescriptionLetter++, Buffer++)
-					{
-						*Buffer = (ANSICHAR)*DescriptionLetter;
-					}
-					check(Buffer + 1 == BufferEnd);
-
-					*Buffer = 0;		// null terminate
-
-					FCpuProfilerTrace::OutputBeginDynamicEvent(VerboseNameANSI.GetData());
-				}
-				else
-				{
-					FCpuProfilerTrace::OutputBeginDynamicEvent(InStatId.GetStatDescriptionANSI()); //todo: Could we use FName index as event id?
-				}
+				FCpuProfilerTrace::OutputBeginDynamicEvent(InStatId.GetStatDescriptionANSI()); //todo: Could we use FName index as event id?
 				EmittedEvent |= TraceEvent;
 			}
 #endif
@@ -2056,14 +2020,6 @@ struct FStat_##StatName\
 #define CONDITIONAL_SCOPE_CYCLE_COUNTER(Stat,bCondition) \
 	FScopeCycleCounter CycleCount_##Stat(bCondition ? GET_STATID(Stat) : TStatId(), GET_STATFLAGS(Stat));
 
-#if CPUPROFILERTRACE_ENABLED
-#define SCOPE_CYCLE_COUNTER_VERBOSE(Stat,VerboseDescription) \
-	FScopeCycleCounter CycleCount_##Stat(GET_STATID(Stat), GET_STATFLAGS(Stat), false, VerboseDescription);
-#else
-#define SCOPE_CYCLE_COUNTER_VERBOSE(Stat,VerboseDescription) \
-	FScopeCycleCounter CycleCount_##Stat(GET_STATID(Stat), GET_STATFLAGS(Stat));
-#endif
-
 #define SCOPE_SECONDS_ACCUMULATOR(Stat) \
 	FSimpleScopeSecondsStat SecondsAccum_##Stat(GET_STATID(Stat));
 
diff --git a/Engine/Source/Runtime/Renderer/Private/SceneRendering.cpp b/Engine/Source/Runtime/Renderer/Private/SceneRendering.cpp
--- a/Engine/Source/Runtime/Renderer/Private/SceneRendering.cpp
+++ b/Engine/Source/Runtime/Renderer/Private/SceneRendering.cpp
@@ -4281,7 +4281,8 @@ static void RenderViewFamilies_RenderThread(FRHICommandListImmediate& RHICmdList
 		// We need to execute the pre-render view extensions before we do any view dependent work.
 		FSceneRenderer::ViewExtensionPreRender_RenderThread(GraphBuilder, SceneRenderer);
 
-		SCOPE_CYCLE_COUNTER_VERBOSE(STAT_TotalSceneRenderingTime, ViewFamily.ProfileDescription.IsEmpty() ? nullptr : *ViewFamily.ProfileDescription);
+		SCOPE_CYCLE_COUNTER(STAT_TotalSceneRenderingTime);
+		SCOPED_NAMED_EVENT_TCHAR_CONDITIONAL(*ViewFamily.ProfileDescription, FColor::Red, !ViewFamily.ProfileDescription.IsEmpty());
 		const uint64 FamilyRenderStart = FPlatformTime::Cycles64();
 
 #if WITH_DEBUG_VIEW_MODES
