Always log out config cvar settings when the last write to the cvar was from the constructor.

diff --git a/Engine/Source/Runtime/Core/Private/Misc/ConfigUtilities.cpp b/Engine/Source/Runtime/Core/Private/Misc/ConfigUtilities.cpp
--- a/Engine/Source/Runtime/Core/Private/Misc/ConfigUtilities.cpp
+++ b/Engine/Source/Runtime/Core/Private/Misc/ConfigUtilities.cpp
@@ -62,6 +62,7 @@ void OnSetCVarFromIniEntry(const TCHAR *IniFile, const TCHAR *Key, const TCHAR*
 			int32 IntValue = CVar->IsVariableInt() ? CVar->GetInt() : 0;
 			float FloatValue = CVar->IsVariableFloat() ? CVar->GetFloat() : 0.0f;
 			FString StringValue = CVar->IsVariableString() ? CVar->GetString() : FString();
+			bool bFirstSet = (CVar->GetFlags() & ECVF_SetByMask) == ECVF_SetByConstructor;
 #endif
 			if (SetBy == ECVF_SetByMask)
 			{
@@ -72,22 +73,25 @@ void OnSetCVarFromIniEntry(const TCHAR *IniFile, const TCHAR *Key, const TCHAR*
 				CVar->Set(Value, (EConsoleVariableFlags)SetBy);
 			}
 #if !NO_LOGGING
-			bool bChanged = false;
-			if (CVar->IsVariableBool())
+			bool bChanged = bFirstSet;
+			if(!bChanged)
 			{
-				bChanged = BoolValue != CVar->GetBool();
-			}
-			else if (CVar->IsVariableInt())
-			{
-				bChanged = IntValue != CVar->GetInt();
-			}
-			else if (CVar->IsVariableFloat())
-			{
-				bChanged = FloatValue != CVar->GetFloat();
-			}
-			else if (CVar->IsVariableString())
-			{
-				bChanged = StringValue != CVar->GetString();
+				if (CVar->IsVariableBool())
+				{
+					bChanged = BoolValue != CVar->GetBool();
+				}
+				else if (CVar->IsVariableInt())
+				{
+					bChanged = IntValue != CVar->GetInt();
+				}
+				else if (CVar->IsVariableFloat())
+				{
+					bChanged = FloatValue != CVar->GetFloat();
+				}
+				else if (CVar->IsVariableString())
+				{
+					bChanged = StringValue != CVar->GetString();
+				}
 			}
 			UE_CLOG(!bNoLogging && bChanged, LogConfig, Log, TEXT("Set CVar [[%s:%s]]"), Key, Value);
 #endif
