- Pass in a better name to the FConfigFile::OverrideFromCommandline from the ConfigContext,. because the one that was being passed was not set when loading another platform's config system from the editor

diff --git a/Engine/Source/Runtime/Core/Private/Misc/ConfigContext.cpp b/Engine/Source/Runtime/Core/Private/Misc/ConfigContext.cpp
--- a/Engine/Source/Runtime/Core/Private/Misc/ConfigContext.cpp
+++ b/Engine/Source/Runtime/Core/Private/Misc/ConfigContext.cpp
@@ -655,7 +655,7 @@ bool FConfigContext::GenerateDestIniFile()
 
 #if ALLOW_INI_OVERRIDE_FROM_COMMANDLINE
 	// process any commandline overrides
-	FConfigFile::OverrideFromCommandline(ConfigFile, DestIniFilename);
+	FConfigFile::OverrideFromCommandline(ConfigFile, BaseIniName);
 #endif
 
 	bool bForceRegenerate = false;
