Fix for long pathnames resulting in UE failing to load dynamic libraries on macOS.

#rnx

diff --git a/Engine/Source/Runtime/Core/Private/Mac/MacPlatformProcess.cpp b/Engine/Source/Runtime/Core/Private/Mac/MacPlatformProcess.cpp
--- a/Engine/Source/Runtime/Core/Private/Mac/MacPlatformProcess.cpp
+++ b/Engine/Source/Runtime/Core/Private/Mac/MacPlatformProcess.cpp
@@ -74,7 +74,7 @@ void* FMacPlatformProcess::GetDllHandle( const TCHAR* Filename )
 
 	check(Filename);
 
-	NSString* DylibPath = [NSString stringWithUTF8String:TCHAR_TO_UTF8(Filename)];
+	NSString* DylibPath = FString(Filename).GetNSString();
 	NSString* ExecutableFolder = [[[NSBundle mainBundle] executablePath] stringByDeletingLastPathComponent];
 	void* Handle = nullptr;
 
@@ -86,7 +86,7 @@ void* FMacPlatformProcess::GetDllHandle( const TCHAR* Filename )
 		if (!Handle)
 		{
 			// If it's not a absolute or relative path, try to find the file in the app bundle
-			DylibPath = [ExecutableFolder stringByAppendingPathComponent:FString(Filename).GetNSString()];
+			DylibPath = [ExecutableFolder stringByAppendingPathComponent:DylibPath];
 			Handle = GetDllHandleImpl(DylibPath, ExecutableFolder);
 		}
 	}
@@ -96,7 +96,7 @@ void* FMacPlatformProcess::GetDllHandle( const TCHAR* Filename )
 		if (![FileManager fileExistsAtPath:DylibPath])
 		{
 			// If it's not a absolute or relative path, try to find the file in the app bundle
-			DylibPath = [ExecutableFolder stringByAppendingPathComponent:FString(Filename).GetNSString()];
+			DylibPath = [ExecutableFolder stringByAppendingPathComponent:DylibPath];
 		}
 		Handle = GetDllHandleImpl(DylibPath, ExecutableFolder);
 	}
