Added IWYU pragmas.

diff --git a/Engine/Source/Editor/UnrealEd/Public/Tools/ToolCompatible.h b/Engine/Source/Editor/UnrealEd/Public/Tools/ToolCompatible.h
--- a/Engine/Source/Editor/UnrealEd/Public/Tools/ToolCompatible.h
+++ b/Engine/Source/Editor/UnrealEd/Public/Tools/ToolCompatible.h
@@ -65,4 +65,4 @@ protected:
 
 };
 
-#include "Tools/ToolCompatible.inl"
+#include "Tools/ToolCompatible.inl" // IWYU pragma: export
diff --git a/Engine/Source/Programs/UnrealLightmass/Private/Lighting/LightingSystem.h b/Engine/Source/Programs/UnrealLightmass/Private/Lighting/LightingSystem.h
--- a/Engine/Source/Programs/UnrealLightmass/Private/Lighting/LightingSystem.h
+++ b/Engine/Source/Programs/UnrealLightmass/Private/Lighting/LightingSystem.h
@@ -2824,7 +2824,7 @@ extern bool IsLightBehindSurface(const FVector4f& TrianglePoint, const FVector4f
  */
 extern TBitArray<> CullBackfacingLights(bool bTwoSidedMaterial, const FVector4f& TrianglePoint, const FVector4f& TriangleNormal, const TArray<FLight*>& Lights);
 
-#include "LightingSystem.inl"
+#include "LightingSystem.inl" // IWYU pragma: export
 
 } //namespace Lightmass
 
diff --git a/Engine/Source/Runtime/Core/Public/Delegates/Delegate.h b/Engine/Source/Runtime/Core/Public/Delegates/Delegate.h
--- a/Engine/Source/Runtime/Core/Public/Delegates/Delegate.h
+++ b/Engine/Source/Runtime/Core/Public/Delegates/Delegate.h
@@ -7,8 +7,8 @@
 #include "UObject/NameTypes.h"
 #include "Templates/SharedPointer.h"
 #include "UObject/WeakObjectPtrTemplates.h"
-#include "Delegates/MulticastDelegateBase.h"
-#include "Delegates/IntegerSequence.h"
+#include "Delegates/MulticastDelegateBase.h" // IWYU pragma: export
+#include "Delegates/IntegerSequence.h" // IWYU pragma: export
 
 /**
  *  C++ DELEGATES
@@ -478,12 +478,14 @@ namespace UE4Delegates_Private
 #define __Delegate_h__
 #define FUNC_INCLUDING_INLINE_IMPL
 
+// IWYU pragma: begin_exports
 #if !UE_BUILD_DOCS
 	#include "Delegates/DelegateInstanceInterface.h"
 	#include "Delegates/DelegateInstancesImpl.h"
 	#include "Delegates/DelegateSignatureImpl.inl"
 	#include "Delegates/DelegateCombinations.h"
 #endif
+// IWYU pragma: end_exports
 
 // No longer allowed to include DelegateSignatureImpl.inl
 #undef FUNC_INCLUDING_INLINE_IMPL
diff --git a/Engine/Source/Runtime/Core/Public/Delegates/MulticastDelegateBase.h b/Engine/Source/Runtime/Core/Public/Delegates/MulticastDelegateBase.h
--- a/Engine/Source/Runtime/Core/Public/Delegates/MulticastDelegateBase.h
+++ b/Engine/Source/Runtime/Core/Public/Delegates/MulticastDelegateBase.h
@@ -6,8 +6,8 @@
 #include "Containers/ContainerAllocationPolicies.h"
 #include "Containers/Array.h"
 #include "Math/UnrealMathUtility.h"
-#include "Delegates/IDelegateInstance.h"
-#include "Delegates/DelegateBase.h"
+#include "Delegates/IDelegateInstance.h" // IWYU pragma: export
+#include "Delegates/DelegateBase.h" // IWYU pragma: export
 
 #if !defined(NUM_MULTICAST_DELEGATE_INLINE_ENTRIES) || NUM_MULTICAST_DELEGATE_INLINE_ENTRIES == 0
 	typedef FHeapAllocator FMulticastInvocationListAllocatorType;
diff --git a/Engine/Source/Runtime/Core/Public/Math/Matrix.h b/Engine/Source/Runtime/Core/Public/Math/Matrix.h
--- a/Engine/Source/Runtime/Core/Public/Math/Matrix.h
+++ b/Engine/Source/Runtime/Core/Public/Math/Matrix.h
@@ -623,7 +623,7 @@ static inline void Inverse4x4( double* dst, const FArg* src )
 	}
 }
 
-#include "Math/Matrix.inl"
+#include "Math/Matrix.inl" // IWYU pragma: export
 
 #ifdef _MSC_VER
 #pragma warning (pop)
diff --git a/Engine/Source/Runtime/Core/Public/Misc/ExpressionParserTypes.h b/Engine/Source/Runtime/Core/Public/Misc/ExpressionParserTypes.h
--- a/Engine/Source/Runtime/Core/Public/Misc/ExpressionParserTypes.h
+++ b/Engine/Source/Runtime/Core/Public/Misc/ExpressionParserTypes.h
@@ -580,4 +580,4 @@ private:
 	TMap<FGuid, FOpParameters>	BinaryOperators;
 };
 
-#include "Misc/ExpressionParserTypes.inl"
+#include "Misc/ExpressionParserTypes.inl" // IWYU pragma: export
diff --git a/Engine/Source/Runtime/CoreUObject/Private/UObject/LinkerPlaceholderBase.h b/Engine/Source/Runtime/CoreUObject/Private/UObject/LinkerPlaceholderBase.h
--- a/Engine/Source/Runtime/CoreUObject/Private/UObject/LinkerPlaceholderBase.h
+++ b/Engine/Source/Runtime/CoreUObject/Private/UObject/LinkerPlaceholderBase.h
@@ -367,7 +367,7 @@ private:
 };
 
 // Templatized implementation
-#include "UObject/LinkerPlaceholderBase.inl"
+#include "UObject/LinkerPlaceholderBase.inl" // IWYU pragma: export
 
 /*******************************************************************************
  * TLinkerImportPlaceholder<> Specializations
diff --git a/Engine/Source/Runtime/Engine/Classes/Kismet/KismetMathLibrary.h b/Engine/Source/Runtime/Engine/Classes/Kismet/KismetMathLibrary.h
--- a/Engine/Source/Runtime/Engine/Classes/Kismet/KismetMathLibrary.h
+++ b/Engine/Source/Runtime/Engine/Classes/Kismet/KismetMathLibrary.h
@@ -4299,5 +4299,5 @@ private:
 
 // Conditionally inlined
 #if KISMET_MATH_INLINE_ENABLED
-#include "KismetMathLibrary.inl"
+#include "KismetMathLibrary.inl" // IWYU pragma: export
 #endif
diff --git a/Engine/Source/Runtime/Experimental/Chaos/Public/Chaos/Joint/JointSolverConstraints.h b/Engine/Source/Runtime/Experimental/Chaos/Public/Chaos/Joint/JointSolverConstraints.h
--- a/Engine/Source/Runtime/Experimental/Chaos/Public/Chaos/Joint/JointSolverConstraints.h
+++ b/Engine/Source/Runtime/Experimental/Chaos/Public/Chaos/Joint/JointSolverConstraints.h
@@ -498,4 +498,4 @@ namespace Chaos
 	};
 }
 
-#include "Chaos/Joint/JointSolverConstraints.inl"
+#include "Chaos/Joint/JointSolverConstraints.inl" // IWYU pragma: export
diff --git a/Engine/Source/Runtime/HeadMountedDisplay/Public/MotionDelayBuffer.h b/Engine/Source/Runtime/HeadMountedDisplay/Public/MotionDelayBuffer.h
--- a/Engine/Source/Runtime/HeadMountedDisplay/Public/MotionDelayBuffer.h
+++ b/Engine/Source/Runtime/HeadMountedDisplay/Public/MotionDelayBuffer.h
@@ -119,4 +119,4 @@ private:
 };
 
 // Templatized implementation
-#include "MotionDelayBuffer.inl"
+#include "MotionDelayBuffer.inl" // IWYU pragma: export
diff --git a/Engine/Source/Runtime/RHI/Public/RHICommandList.h b/Engine/Source/Runtime/RHI/Public/RHICommandList.h
--- a/Engine/Source/Runtime/RHI/Public/RHICommandList.h
+++ b/Engine/Source/Runtime/RHI/Public/RHICommandList.h
@@ -5904,4 +5904,4 @@ private:
 
 #undef RHICOMMAND_CALLSTACK
 
-#include "RHICommandList.inl"
+#include "RHICommandList.inl" // IWYU pragma: export
diff --git a/Engine/Source/Runtime/Renderer/Private/DynamicPrimitiveDrawing.h b/Engine/Source/Runtime/Renderer/Private/DynamicPrimitiveDrawing.h
--- a/Engine/Source/Runtime/Renderer/Private/DynamicPrimitiveDrawing.h
+++ b/Engine/Source/Runtime/Renderer/Private/DynamicPrimitiveDrawing.h
@@ -75,5 +75,5 @@ private:
 	FBatchedElements& GetElements(uint8 DepthPriorityGroup) const;
 };
 
-#include "DynamicPrimitiveDrawing.inl"
+#include "DynamicPrimitiveDrawing.inl" // IWYU pragma: export
 
diff --git a/Engine/Source/Runtime/Renderer/Private/ScenePrivate.h b/Engine/Source/Runtime/Renderer/Private/ScenePrivate.h
--- a/Engine/Source/Runtime/Renderer/Private/ScenePrivate.h
+++ b/Engine/Source/Runtime/Renderer/Private/ScenePrivate.h
@@ -3534,5 +3534,5 @@ inline bool ShouldIncludeDomainInMeshPass(EMaterialDomain Domain)
 	return Domain != MD_Volume;
 }
 
-#include "BasePassRendering.inl"
+#include "BasePassRendering.inl" // IWYU pragma: export
 
diff --git a/Engine/Source/Runtime/Renderer/Private/ScreenPass.h b/Engine/Source/Runtime/Renderer/Private/ScreenPass.h
--- a/Engine/Source/Runtime/Renderer/Private/ScreenPass.h
+++ b/Engine/Source/Runtime/Renderer/Private/ScreenPass.h
@@ -670,4 +670,4 @@ void AddDownsampleDepthPass(
 	FScreenPassRenderTarget Output,
 	EDownsampleDepthFilter DownsampleDepthFilter);
 
-#include "ScreenPass.inl"
+#include "ScreenPass.inl" // IWYU pragma: export
diff --git a/Engine/Source/Runtime/SlateCore/Public/Input/DragAndDrop.h b/Engine/Source/Runtime/SlateCore/Public/Input/DragAndDrop.h
--- a/Engine/Source/Runtime/SlateCore/Public/Input/DragAndDrop.h
+++ b/Engine/Source/Runtime/SlateCore/Public/Input/DragAndDrop.h
@@ -271,4 +271,4 @@ protected:
 	FVector2D DecoratorPosition;
 };
 
-#include "Input/DragAndDrop.inl"
+#include "Input/DragAndDrop.inl" // IWYU pragma: export
diff --git a/Engine/Source/Runtime/SlateCore/Public/Layout/WidgetPath.h b/Engine/Source/Runtime/SlateCore/Public/Layout/WidgetPath.h
--- a/Engine/Source/Runtime/SlateCore/Public/Layout/WidgetPath.h
+++ b/Engine/Source/Runtime/SlateCore/Public/Layout/WidgetPath.h
@@ -286,4 +286,4 @@ public:
 	TWeakPtr< SWindow > Window;
 };
 
-#include "Layout/WidgetPath.inl"
+#include "Layout/WidgetPath.inl" // IWYU pragma: export
diff --git a/Engine/Source/Runtime/SlateCore/Public/Types/SlateAttribute.h b/Engine/Source/Runtime/SlateCore/Public/Types/SlateAttribute.h
--- a/Engine/Source/Runtime/SlateCore/Public/Types/SlateAttribute.h
+++ b/Engine/Source/Runtime/SlateCore/Public/Types/SlateAttribute.h
@@ -252,10 +252,11 @@ struct TSlateDeprecatedTAttribute
 	bool IdenticalTo(const TAttribute<ObjectType>& InOther) const { return false; }
 };
 
-
+// IWYU pragma: begin_exports
 #include "Types/Attributes/SlateAttributeDefinition.inl"
 #include "Types/Attributes/SlateAttributeBase.inl"
 
 #include "Types/Attributes/SlateAttributeContained.inl"
 #include "Types/Attributes/SlateAttributeManaged.inl"
 #include "Types/Attributes/SlateAttributeMember.inl"
+// IWYU pragma: end_exports
