Link against the correct Superluminal Performance API when building targets that have `bDebugBuildsActuallyUseDebugCRT` set to `true` in `Debug` builds.

diff --git a/Engine/Source/Runtime/Core/Core.Build.cs b/Engine/Source/Runtime/Core/Core.Build.cs
--- a/Engine/Source/Runtime/Core/Core.Build.cs
+++ b/Engine/Source/Runtime/Core/Core.Build.cs
@@ -218,9 +218,16 @@ public class Core : ModuleRules
 			if (Target.Configuration != UnrealTargetConfiguration.Shipping &&
 				File.Exists(Path.Combine(SuperluminalApiDir, "include/Superluminal/PerformanceAPI_capi.h")))
 			{
+				if (Target.bDebugBuildsActuallyUseDebugCRT == true && Target.Configuration == UnrealTargetConfiguration.Debug)
+				{
+					PublicAdditionalLibraries.Add(Path.Combine(SuperluminalLibDir, "PerformanceAPI_MDd.lib"));
+				}
+				else
+				{
+					PublicAdditionalLibraries.Add(Path.Combine(SuperluminalLibDir, "PerformanceAPI_MD.lib"));
+				}
+				PublicDefinitions.Add("WITH_SUPERLUMINAL_PROFILER=1");					  
 				PublicSystemIncludePaths.Add(Path.Combine(SuperluminalApiDir, "include/"));
-				PublicAdditionalLibraries.Add(Path.Combine(SuperluminalLibDir, "PerformanceAPI_MD.lib"));
-				PublicDefinitions.Add("WITH_SUPERLUMINAL_PROFILER=1");
 			}
 			else
 			{
