Adding minidump output on hangs in the crash reporter.

diff --git a/Engine/Source/Runtime/Core/Private/Windows/WindowsPlatformCrashContext.cpp b/Engine/Source/Runtime/Core/Private/Windows/WindowsPlatformCrashContext.cpp
--- a/Engine/Source/Runtime/Core/Private/Windows/WindowsPlatformCrashContext.cpp
+++ b/Engine/Source/Runtime/Core/Private/Windows/WindowsPlatformCrashContext.cpp
@@ -192,11 +192,8 @@ void FWindowsPlatformCrashContext::CopyPlatformSpecificFiles(const TCHAR* Output
 
 	// Save minidump
 	LPEXCEPTION_POINTERS ExceptionInfo = (LPEXCEPTION_POINTERS)Context;
-	if (ExceptionInfo != nullptr)
-	{
-		const FString MinidumpFileName = FPaths::Combine(OutputDirectory, FGenericCrashContext::UEMinidumpName);
-		WriteMinidump(ProcessHandle.Get(), CrashedThreadId, *this, *MinidumpFileName, ExceptionInfo);
-	}
+	const FString MinidumpFileName = FPaths::Combine(OutputDirectory, FGenericCrashContext::UEMinidumpName);
+	WriteMinidump(ProcessHandle.Get(), CrashedThreadId, *this, *MinidumpFileName, ExceptionInfo);
 
 	// If present, include the crash video
 	const FString CrashVideoPath = FPaths::ProjectLogDir() / TEXT("CrashVideo.avi");
