Macro UE_SCOPED_TIMER now includes total time when triggered several times.

#rnx

diff --git a/Engine/Source/Runtime/Core/Public/ProfilingDebugging/ScopedTimers.h b/Engine/Source/Runtime/Core/Public/ProfilingDebugging/ScopedTimers.h
--- a/Engine/Source/Runtime/Core/Public/ProfilingDebugging/ScopedTimers.h
+++ b/Engine/Source/Runtime/Core/Public/ProfilingDebugging/ScopedTimers.h
@@ -4,6 +4,7 @@
 
 #include "CoreTypes.h"
 #include "Containers/UnrealString.h"
+#include "Containers/StringFwd.h"
 #include "CoreGlobals.h"
 #include "HAL/PlatformTime.h"
 #include "HAL/ThreadSafeCounter.h"
@@ -203,36 +204,45 @@ public:
  * doesn't have to call Start() and Stop() manually) using a custom
  * output function.
  */
-template<class F>
+template<class Func>
 class FScopedDurationTimeCustomLogger
 {
 public:
-	explicit FScopedDurationTimeCustomLogger(const TCHAR* InTitle, F InFunc)
+	explicit FScopedDurationTimeCustomLogger(const TCHAR* InTitle, double& InTotalTime, Func InLogFunc)
 		: Title(InTitle)
-		, Func(InFunc)
-		, Accumulator(0.0)
-		, Timer(Accumulator)
+		, LogFunc(InLogFunc)
+		, LocalTime(-FPlatformTime::Seconds())
+		, TotalTime(InTotalTime)
 	{
-		Func(*FString::Printf(TEXT("%s started..."), InTitle));
-		Timer.Start();
+		LogFunc(*FString::Printf(TEXT("%s started..."), InTitle));
 	}
 
 	~FScopedDurationTimeCustomLogger()
 	{
-		Timer.Stop();
-		Func(*FString::Printf(TEXT("%s took %s"), *Title, *FPlatformTime::PrettyTime(Accumulator)));
+		LocalTime += FPlatformTime::Seconds();
+		TotalTime += LocalTime;
+
+		TStringBuilder<1024> Msg;
+		Msg.Appendf(TEXT("%s took %s"), *Title, *FPlatformTime::PrettyTime(LocalTime));
+
+		if (TotalTime > LocalTime)
+		{
+			Msg.Appendf(TEXT(" (total: %s)"), *FPlatformTime::PrettyTime(TotalTime));
+		}
+
+		LogFunc(Msg.ToString());
 	}
 
 private:
 	FString Title;
-	F Func;
-	double Accumulator;
-	FDurationTimer Timer;
+	Func LogFunc;
+	double LocalTime;
+	double& TotalTime;
 };
 
 #if NO_LOGGING
 #define UE_SCOPED_TIMER(Title, Category, Verbosity)
 #else
 #define UE_SCOPED_TIMER(Title, Category, Verbosity) \
-	FScopedDurationTimeCustomLogger BODY_MACRO_COMBINE(Scoped,Timer,_,__LINE__)(Title, [](const TCHAR* Msg) { UE_LOG(Category, Verbosity, TEXT("%s"), Msg); })
+	static double BODY_MACRO_COMBINE(Scoped,TimerTotal,_,__LINE__) = 0.0; FScopedDurationTimeCustomLogger BODY_MACRO_COMBINE(Scoped,Timer,_,__LINE__)(Title, BODY_MACRO_COMBINE(Scoped,TimerTotal,_,__LINE__), [](const TCHAR* Msg) { UE_LOG(Category, Verbosity, TEXT("%s"), Msg) })
 #endif
