Removing unused raytracing specific code from MemoryLayout and the associated DDPI entry

diff --git a/Engine/Config/Windows/DataDrivenPlatformInfo.ini b/Engine/Config/Windows/DataDrivenPlatformInfo.ini
--- a/Engine/Config/Windows/DataDrivenPlatformInfo.ini
+++ b/Engine/Config/Windows/DataDrivenPlatformInfo.ini
@@ -4,7 +4,6 @@ IniParent=VulkanPC
 PlatformNameAliases=Win64
 TargetSettingsIniSectionName=/Script/WindowsTargetPlatform.WindowsTargetSettings
 Freezing_b32Bit=false
-Freezing_bWithRayTracing=true
 Freezing_MaxFieldAlignment32=4
 Freezing_MaxFieldAlignment64=8
 Freezing_bAlignBases=true
diff --git a/Engine/Source/Runtime/Core/Private/Misc/DataDrivenPlatformInfoRegistry.cpp b/Engine/Source/Runtime/Core/Private/Misc/DataDrivenPlatformInfoRegistry.cpp
--- a/Engine/Source/Runtime/Core/Private/Misc/DataDrivenPlatformInfoRegistry.cpp
+++ b/Engine/Source/Runtime/Core/Private/Misc/DataDrivenPlatformInfoRegistry.cpp
@@ -312,7 +312,6 @@ static void LoadDDPIIniSettings(const FConfigFile& IniFile, FDataDrivenPlatformI
 	DDPIGetUInt(IniFile, Info.Freezing_b32Bit ? TEXT("Freezing_MaxFieldAlignment32") : TEXT("Freezing_MaxFieldAlignment64"), Info.Freezing_MaxFieldAlignment);
 	DDPIGetBool(IniFile, TEXT("Freezing_bForce64BitMemoryImagePointers"), Info.Freezing_bForce64BitMemoryImagePointers);
 	DDPIGetBool(IniFile, TEXT("Freezing_bAlignBases"), Info.Freezing_bAlignBases);
-	DDPIGetBool(IniFile, TEXT("Freezing_bWithRayTracing"), Info.Freezing_bWithRayTracing);
 
 	DDPIGetGuid(IniFile, TEXT("GlobalIdentifier"), Info.GlobalIdentifier);
 	checkf(Info.GlobalIdentifier != FGuid(), TEXT("Platform %s didn't have a valid GlobalIdentifier set in DataDrivenPlatformInfo.ini"), *PlatformName.ToString());
diff --git a/Engine/Source/Runtime/Core/Private/Serialization/MemoryImage.cpp b/Engine/Source/Runtime/Core/Private/Serialization/MemoryImage.cpp
--- a/Engine/Source/Runtime/Core/Private/Serialization/MemoryImage.cpp
+++ b/Engine/Source/Runtime/Core/Private/Serialization/MemoryImage.cpp
@@ -65,7 +65,6 @@ void FPlatformTypeLayoutParameters::InitializeForPlatform(const FString& Platfor
 
 	Flags = Flag_Initialized;
 	if (bHasEditorOnlyData) Flags |= Flag_WithEditorOnly;
-	if (PlatformInfo.Freezing_bWithRayTracing) Flags |= Flag_WithRaytracing;
 	if (PlatformInfo.Freezing_b32Bit) Flags |= Flag_Is32Bit;
 	if (PlatformInfo.Freezing_bAlignBases) Flags |= Flag_AlignBases;
 
@@ -76,7 +75,6 @@ void FPlatformTypeLayoutParameters::InitializeForCurrent()
 {
 	Flags = Flag_Initialized;
 	if (WITH_EDITORONLY_DATA) Flags |= Flag_WithEditorOnly;
-	if (WITH_RAYTRACING) Flags |= Flag_WithRaytracing;
 	if (PLATFORM_32BITS) Flags |= Flag_Is32Bit;
 
 	check(GetRawPointerSize() == sizeof(void*));
@@ -346,16 +344,11 @@ void Freeze::ExtractBitFieldValue(const void* Value, uint32 SrcBitOffset, uint32
 bool Freeze::IncludeField(const FFieldLayoutDesc* FieldDesc, const FPlatformTypeLayoutParameters& LayoutParams)
 {
 	const bool bIsEditorOnly = (FieldDesc->Flags & EFieldLayoutFlags::WithEditorOnly) != 0u;
-	const bool bIsRayTracing = (FieldDesc->Flags & EFieldLayoutFlags::WithRayTracing) != 0u;
 
 	if (bIsEditorOnly && !LayoutParams.WithEditorOnly())
 	{
 		return false;
 	}
-	if (bIsRayTracing && !LayoutParams.WithRaytracing())
-	{
-		return false;
-	}
 
 	return true;
 }
diff --git a/Engine/Source/Runtime/Core/Public/Misc/DataDrivenPlatformInfoRegistry.h b/Engine/Source/Runtime/Core/Public/Misc/DataDrivenPlatformInfoRegistry.h
--- a/Engine/Source/Runtime/Core/Public/Misc/DataDrivenPlatformInfoRegistry.h
+++ b/Engine/Source/Runtime/Core/Public/Misc/DataDrivenPlatformInfoRegistry.h
@@ -140,7 +140,6 @@ struct FDataDrivenPlatformInfo
 	bool Freezing_b32Bit = false;
 	bool Freezing_bForce64BitMemoryImagePointers = false;
 	bool Freezing_bAlignBases = false;
-	bool Freezing_bWithRayTracing = false;
 
 	// True if this platform has a non-generic gamepad specifically associated with it
 	bool bHasDedicatedGamepad = false;
diff --git a/Engine/Source/Runtime/Core/Public/Serialization/MemoryLayout.h b/Engine/Source/Runtime/Core/Public/Serialization/MemoryLayout.h
--- a/Engine/Source/Runtime/Core/Public/Serialization/MemoryLayout.h
+++ b/Engine/Source/Runtime/Core/Public/Serialization/MemoryLayout.h
@@ -22,15 +22,6 @@ class ITargetPlatform;
 struct FTypeLayoutDesc;
 struct FPlatformTypeLayoutParameters;
 
-// Duplicated from RHIDefinitions.h [AND MUST MATCH THE Freezing_bWithRayTracing in DataDrivenPlatformInfo.ini]
-#ifndef WITH_RAYTRACING
-#if (PLATFORM_WINDOWS && PLATFORM_64BITS)
-#define WITH_RAYTRACING 1
-#else
-#define WITH_RAYTRACING 0
-#endif
-#endif
-
 /*#define UE_STATIC_ONLY(T) \
 	T() = delete; \
 	~T() = delete; \
@@ -493,16 +484,6 @@ PRAGMA_ENABLE_DEPRECATION_WARNINGS
 #define LAYOUT_BITFIELD_EDITORONLY(T, Name, BitFieldSize, ...)
 #endif
 
-#if WITH_RAYTRACING
-#define LAYOUT_FIELD_RAYTRACING(T, Name, ...) PREPROCESSOR_REMOVE_OPTIONAL_PARENS(T) Name; INTERNAL_LAYOUT_FIELD(T, Name, STRUCT_OFFSET(DerivedType, Name), EFieldLayoutFlags::MakeFlagsRayTracing(__VA_ARGS__), 1u, 0u, __COUNTER__)
-#define LAYOUT_FIELD_INITIALIZED_RAYTRACING(T, Name, Value, ...) PREPROCESSOR_REMOVE_OPTIONAL_PARENS(T) Name = Value; INTERNAL_LAYOUT_FIELD(T, Name, STRUCT_OFFSET(DerivedType, Name), EFieldLayoutFlags::MakeFlagsRayTracing(__VA_ARGS__), 1u, 0u, __COUNTER__)
-#define LAYOUT_ARRAY_RAYTRACING(T, Name, NumArray, ...) PREPROCESSOR_REMOVE_OPTIONAL_PARENS(T) Name[NumArray]; INTERNAL_LAYOUT_FIELD(T, Name, STRUCT_OFFSET(DerivedType, Name), EFieldLayoutFlags::MakeFlagsRayTracing(__VA_ARGS__), NumArray, 0u, __COUNTER__)
-#else
-#define LAYOUT_FIELD_RAYTRACING(T, Name, ...)
-#define LAYOUT_FIELD_INITIALIZED_RAYTRACING(T, Name, Value, ...)
-#define LAYOUT_ARRAY_RAYTRACING(T, Name, NumArray, ...)
-#endif
-
 #define INTERNAL_LAYOUT_INTERFACE_PREFIX_NonVirtual(...) __VA_ARGS__
 #define INTERNAL_LAYOUT_INTERFACE_PREFIX_Virtual(...) __VA_ARGS__ virtual
 #define INTERNAL_LAYOUT_INTERFACE_PREFIX_Abstract(...) virtual
@@ -818,7 +799,6 @@ struct FPlatformTypeLayoutParameters
 		Flag_Is32Bit = (1 << 1),
 		Flag_AlignBases = (1 << 2),
 		Flag_WithEditorOnly = (1 << 3),
-		Flag_WithRaytracing = (1 << 4),
 	};
 
 	LAYOUT_FIELD_INITIALIZED(uint32, MaxFieldAlignment, 0xffffffff);
@@ -828,7 +808,6 @@ struct FPlatformTypeLayoutParameters
 	inline bool Is32Bit() const { return (Flags & Flag_Is32Bit) != 0u; }
 	inline bool HasAlignBases() const { return (Flags & Flag_AlignBases) != 0u; }
 	inline bool WithEditorOnly() const { return (Flags & Flag_WithEditorOnly) != 0u; }
-	inline bool WithRaytracing() const { return (Flags & Flag_WithRaytracing) != 0u; }
 
 	// May need dedicated flag for this, if we need to support case-preserving names in non-editor builds
 	inline bool WithCasePreservingFName() const { return WithEditorOnly(); }
