LLT Catch2-style macros: REQUIRE_(NO)ENSURE and CHECK_(NO)ENSURE
Based on existing Catch2 macros and reusing FDebug number of ensures counter.

diff --git a/Engine/Source/Developer/LowLevelTestsRunner/Private/TestListeners/ConsoleListener.cpp b/Engine/Source/Developer/LowLevelTestsRunner/Private/TestListeners/ConsoleListener.cpp
--- a/Engine/Source/Developer/LowLevelTestsRunner/Private/TestListeners/ConsoleListener.cpp
+++ b/Engine/Source/Developer/LowLevelTestsRunner/Private/TestListeners/ConsoleListener.cpp
@@ -33,7 +33,7 @@ private:
 		// I tried this but unfortunately it didn't work, passing that option makes catch complain with the error message:
 		// "Verbosity level not supported by this reporter"
 
-		if (LogVerbosity <= DesiredLogVerbosity)
+		if (LogVerbosity <= DesiredLogVerbosity && !GLog->GetSuppressEventTag())
 		{
 			std::cout << *FText::FromName(Category).ToString() << "(" << ToString(LogVerbosity) << ")" << ": " << V << "\n";
 		}
diff --git a/Engine/Source/Developer/LowLevelTestsRunner/Public/TestMacros/Assertions.h b/Engine/Source/Developer/LowLevelTestsRunner/Public/TestMacros/Assertions.h
new file mode 100644
--- /dev/null
+++ b/Engine/Source/Developer/LowLevelTestsRunner/Public/TestMacros/Assertions.h
@@ -0,0 +1,39 @@
+#pragma once
+
+#include "Misc/AssertionMacros.h"
+
+#if defined(CATCH_CONFIG_PREFIX_ALL) && !defined(CATCH_CONFIG_DISABLE)
+    #define CATCH_REQUIRE_ENSURE( ... ) INTERNAL_CATCH_ENSURE( "CATCH_REQUIRE_ENSURE", Catch::ResultDisposition::Normal, __VA_ARGS__ )
+    #define CATCH_REQUIRE_NOENSURE( ... ) INTERNAL_CATCH_NO_ENSURE( "CATCH_REQUIRE_NOENSURE", Catch::ResultDisposition::Normal, __VA_ARGS__ )
+    #define CATCH_CHECK_ENSURE( ... ) INTERNAL_CATCH_ENSURE( "CATCH_CHECK_ENSURE", Catch::ResultDisposition::ContinueOnFailure, __VA_ARGS__ )
+    #define CATCH_CHECK_NOENSURE( ... ) INTERNAL_CATCH_NO_ENSURE( "CATCH_CHECK_NOENSURE", Catch::ResultDisposition::ContinueOnFailure, __VA_ARGS__ )
+#else
+    #define REQUIRE_ENSURE( ... ) INTERNAL_CATCH_ENSURE( "REQUIRE_ENSURE", Catch::ResultDisposition::Normal, __VA_ARGS__ )
+    #define REQUIRE_NOENSURE( ... ) INTERNAL_CATCH_NO_ENSURE( "REQUIRE_NOENSURE", Catch::ResultDisposition::Normal, __VA_ARGS__ )
+    #define CHECK_ENSURE( ... ) INTERNAL_CATCH_ENSURE( "CHECK_ENSURE", Catch::ResultDisposition::ContinueOnFailure, __VA_ARGS__ )
+    #define CHECK_NOENSURE( ... ) INTERNAL_CATCH_NO_ENSURE( "CHECK_NOENSURE", Catch::ResultDisposition::ContinueOnFailure, __VA_ARGS__ )
+#endif
+
+#define INTERNAL_CATCH_ENSURE( macroName, resultDisposition, ... ) \
+ do { \
+		SIZE_T NumEnsureFailuresBefore = FDebug::GetNumEnsureFailures(); \
+        Catch::AssertionHandler catchAssertionHandler( macroName##_catch_sr, CATCH_INTERNAL_LINEINFO, CATCH_INTERNAL_STRINGIFY(__VA_ARGS__), resultDisposition ); \
+        INTERNAL_CATCH_TRY { \
+            static_cast<void>(__VA_ARGS__); \
+            bool bEncounteredEnsure = (FDebug::GetNumEnsureFailures() - NumEnsureFailuresBefore) > 0; \
+			catchAssertionHandler.handleExpr( Catch::Decomposer() <= bEncounteredEnsure ); \
+        } INTERNAL_CATCH_CATCH( catchAssertionHandler ) \
+        INTERNAL_CATCH_REACT( catchAssertionHandler ) \
+    } while( false )
+
+#define INTERNAL_CATCH_NO_ENSURE( macroName, resultDisposition, ... ) \
+ do { \
+		SIZE_T NumEnsureFailuresBefore = FDebug::GetNumEnsureFailures(); \
+        Catch::AssertionHandler catchAssertionHandler( macroName##_catch_sr, CATCH_INTERNAL_LINEINFO, CATCH_INTERNAL_STRINGIFY(__VA_ARGS__), resultDisposition ); \
+        INTERNAL_CATCH_TRY { \
+            static_cast<void>(__VA_ARGS__); \
+            bool bEncounteredEnsure = (FDebug::GetNumEnsureFailures() - NumEnsureFailuresBefore) > 0; \
+			catchAssertionHandler.handleExpr( Catch::Decomposer() <= !(bEncounteredEnsure) ); \
+        } INTERNAL_CATCH_CATCH( catchAssertionHandler ) \
+        INTERNAL_CATCH_REACT( catchAssertionHandler ) \
+    } while( false )
diff --git a/Engine/Source/Runtime/Core/Tests/Misc/AssertionMacrosTest.cpp b/Engine/Source/Runtime/Core/Tests/Misc/AssertionMacrosTest.cpp
new file mode 100644
--- /dev/null
+++ b/Engine/Source/Runtime/Core/Tests/Misc/AssertionMacrosTest.cpp
@@ -0,0 +1,119 @@
+// Copyright Epic Games, Inc. All Rights Reserved.
+
+#if WITH_LOW_LEVEL_TESTS
+
+#include "Misc/AssertionMacros.h"
+
+#include "TestHarness.h"
+#include "TestMacros/Assertions.h"
+
+static void OneEnsureNonFailed();
+static void OneEnsureNonFailedDifferent();
+static void OneEnsureFailed();
+static void OneEnsureFailedDifferent();
+static void OneEnsureFailedDifferentAgain();
+static void MultipleEnsuresNonFailed();
+static void MultipleEnsuresFailed();
+static void MultipleEnsuresFailedDifferent();
+
+/**
+ * @brief Ensure test case, tests Catch2-style macros that check or require whether at least one ensure was triggered or not.
+ * Note that each different ensure expression is counted once, if the same ensure fails twice it's only counted the first time it fails.
+ */
+TEST_CASE("Core::Misc::AssertionMacros", "[Core][Misc][AssertionMacros][Ensure]")
+{
+	GLog->SetSuppressEventTag(true);
+	SECTION("Ensure not triggered with REQUIRE_NOENSURE")
+	{
+		REQUIRE_NOENSURE(OneEnsureNonFailed());
+	}
+
+	SECTION("Ensure triggered with REQUIRE_ENSURE")
+	{
+		REQUIRE_ENSURE(OneEnsureFailed());
+	}
+
+	SECTION("Ensure not triggered with multiple CHECK_NOENSURE")
+	{
+		SIZE_T EnsuresBefore = FDebug::GetNumEnsureFailures();
+		CHECK_NOENSURE(OneEnsureNonFailed());
+		CHECK_NOENSURE(OneEnsureNonFailedDifferent());
+		CHECK(FDebug::GetNumEnsureFailures() == EnsuresBefore);
+	}
+
+	SECTION("Ensure triggered multiple CHECK_ENSURE")
+	{
+		SIZE_T EnsuresBefore = FDebug::GetNumEnsureFailures();
+		CHECK_ENSURE(OneEnsureFailedDifferent());
+		CHECK_ENSURE(OneEnsureFailedDifferentAgain());
+		CHECK(FDebug::GetNumEnsureFailures() == EnsuresBefore + 2);
+	}
+
+	SECTION("Multiple ensures not triggered with one REQUIRE_NOENSURE")
+	{
+		REQUIRE_NOENSURE(MultipleEnsuresNonFailed());
+	}
+
+	SECTION("At least one of multiple ensures triggered with REQUIRE_ENSURE")
+	{
+		REQUIRE_ENSURE(MultipleEnsuresFailed());
+	}
+
+	SECTION("Multiple ensures not triggered with CHECK_NOENSURE")
+	{
+		CHECK_NOENSURE(MultipleEnsuresNonFailed());
+	}
+
+	SECTION("At least one of multiple ensures triggered with CHECK_ENSURE")
+	{
+		CHECK_ENSURE(MultipleEnsuresFailedDifferent());
+	}
+	GLog->SetSuppressEventTag(false);
+}
+
+static void OneEnsureNonFailed()
+{
+	ensure(true);
+}
+
+static void OneEnsureNonFailedDifferent()
+{
+	ensure(1 == 1);
+}
+
+static void OneEnsureFailed()
+{
+	ensure(false);
+}
+
+static void OneEnsureFailedDifferent()
+{
+	ensure(1 == 2);
+}
+
+static void OneEnsureFailedDifferentAgain()
+{
+	ensure(2 == 3);
+}
+
+static void MultipleEnsuresNonFailed()
+{
+	ensure(2 == 2);
+	ensure(3 == 3);
+}
+
+static void MultipleEnsuresFailed()
+{
+	ensure(2 == 2);
+	ensure(3 == 4);
+	ensure(4 == 4);
+}
+
+static void MultipleEnsuresFailedDifferent()
+{
+	ensure(2 == 2);
+	ensure(4 == 5);
+	ensure(3 == 3);
+}
+
+#endif
