Fixed deprecation warning

#rnx

diff --git a/Engine/Source/Runtime/Core/Public/Containers/StringView.h b/Engine/Source/Runtime/Core/Public/Containers/StringView.h
--- a/Engine/Source/Runtime/Core/Public/Containers/StringView.h
+++ b/Engine/Source/Runtime/Core/Public/Containers/StringView.h
@@ -337,7 +337,7 @@ public:
 	 *
 	 * @returns True if index is valid. False otherwise.
 	 */
-	UE_NODISCARD FORCEINLINE bool IsValidIndex(SizeType Index) const
+	UE_NODISCARD FORCEINLINE bool IsValidIndex(int32 Index) const
 	{
 		return Index >= 0 && Index < Len();
 	}
