Fixed compile issues found by compiling editor with -NoPCH -DisableUnity.

diff --git a/Engine/Source/Developer/DirectoryWatcher/Private/Linux/DirectoryWatchRequestLinux.cpp b/Engine/Source/Developer/DirectoryWatcher/Private/Linux/DirectoryWatchRequestLinux.cpp
--- a/Engine/Source/Developer/DirectoryWatcher/Private/Linux/DirectoryWatchRequestLinux.cpp
+++ b/Engine/Source/Developer/DirectoryWatcher/Private/Linux/DirectoryWatchRequestLinux.cpp
@@ -4,6 +4,10 @@
 #include "HAL/FileManager.h"
 #include "DirectoryWatcherPrivate.h"
 
+#if !UE_BUILD_SHIPPING
+#include "Misc/CoreMisc.h"
+#endif
+
 // To see inotify watch events:
 //   TestPAL dirwatcher -LogCmds="LogDirectoryWatcher VeryVerbose"
 
diff --git a/Engine/Source/Developer/SlateFontDialog/Private/SlateFontDialogModule.cpp b/Engine/Source/Developer/SlateFontDialog/Private/SlateFontDialogModule.cpp
--- a/Engine/Source/Developer/SlateFontDialog/Private/SlateFontDialogModule.cpp
+++ b/Engine/Source/Developer/SlateFontDialog/Private/SlateFontDialogModule.cpp
@@ -1,6 +1,8 @@
 // Copyright Epic Games, Inc. All Rights Reserved.
 
 #include "SlateFontDialogModule.h"
+
+#include "Modules/ModuleManager.h"
 #include "SlateFontDlgWindow.h"
 
 bool FSlateFontDialogModule::OpenFontDialog(FString& OutFontName, float& OutHeight, EFontImportFlags& OutFlags)
diff --git a/Engine/Source/Developer/SlateFontDialog/Private/SlateFontDlgWindow.cpp b/Engine/Source/Developer/SlateFontDialog/Private/SlateFontDlgWindow.cpp
--- a/Engine/Source/Developer/SlateFontDialog/Private/SlateFontDlgWindow.cpp
+++ b/Engine/Source/Developer/SlateFontDialog/Private/SlateFontDlgWindow.cpp
@@ -4,12 +4,18 @@
 
 #include "Fonts/CompositeFont.h"
 
+#include "Framework/Docking/TabManager.h"
+
 #include "HAL/FileManager.h"
 #include "HAL/PlatformFileManager.h"
 
+#include "SlateOptMacros.h"
+
 #include "Widgets/Input/SNumericEntryBox.h"
+#include "Widgets/Input/SButton.h"
 #include "Widgets/Colors/SColorPicker.h"
 #include "Widgets/Colors/SColorBlock.h"
+#include "Widgets/Layout/SGridPanel.h"
 
 #include <fontconfig.h>
 
diff --git a/Engine/Source/Runtime/Core/Private/Internationalization/TextHistory.h b/Engine/Source/Runtime/Core/Private/Internationalization/TextHistory.h
--- a/Engine/Source/Runtime/Core/Private/Internationalization/TextHistory.h
+++ b/Engine/Source/Runtime/Core/Private/Internationalization/TextHistory.h
@@ -1,12 +1,14 @@
 // Copyright Epic Games, Inc. All Rights Reserved.
 #pragma once
 
+#include "Containers/Array.h"
 #include "CoreTypes.h"
 #include "HAL/CriticalSection.h"
 #include "Internationalization/CulturePointer.h"
 #include "Internationalization/ITextData.h"
 #include "Internationalization/LocalizedTextSourceTypes.h"
 #include "Internationalization/StringTableCoreFwd.h"
+#include "Internationalization/Text.h"
 #include "Internationalization/TextKey.h"
 #include "Internationalization/Text.h"
 #include "Misc/DateTime.h"
diff --git a/Engine/Source/Runtime/Core/Private/Unix/UnixPlatformMisc.cpp b/Engine/Source/Runtime/Core/Private/Unix/UnixPlatformMisc.cpp
--- a/Engine/Source/Runtime/Core/Private/Unix/UnixPlatformMisc.cpp
+++ b/Engine/Source/Runtime/Core/Private/Unix/UnixPlatformMisc.cpp
@@ -24,6 +24,7 @@
 #include "Unix/UnixPlatformCrashContext.h"
 #include "Misc/ConfigCacheIni.h"
 #include "GenericPlatform/GenericPlatformChunkInstall.h"
+#include "Misc/OutputDeviceRedirector.h"
 
 #if PLATFORM_HAS_CPUID
 	#include <cpuid.h>
diff --git a/Engine/Source/Runtime/Core/Public/Templates/TypeHash.h b/Engine/Source/Runtime/Core/Public/Templates/TypeHash.h
--- a/Engine/Source/Runtime/Core/Public/Templates/TypeHash.h
+++ b/Engine/Source/Runtime/Core/Public/Templates/TypeHash.h
@@ -6,6 +6,8 @@
 #include "Templates/IsEnum.h"
 #include "Misc/Crc.h"
 
+#include <stdint.h>
+
 /**
  * Combines two hash values to get a third.
  * Note - this function is not commutative.
diff --git a/Engine/Source/Runtime/Engine/Private/MaterialInstanceActor.cpp b/Engine/Source/Runtime/Engine/Private/MaterialInstanceActor.cpp
--- a/Engine/Source/Runtime/Engine/Private/MaterialInstanceActor.cpp
+++ b/Engine/Source/Runtime/Engine/Private/MaterialInstanceActor.cpp
@@ -2,7 +2,9 @@
 
 #include "Materials/MaterialInstanceActor.h"
 #include "Components/BillboardComponent.h"
+#include "Engine/Texture2D.h"
 #include "Misc/MessageDialog.h"
+#include "UObject/ConstructorHelpers.h"
 
 /**
  * Construct a list of static actor names.
