Undo single //UE5/Main/Engine/Source/Runtime/Core/Private/Mac/MacPlatformProcess.cpp revision 21 from CL 21090810
This change was used for debugging and was triggering Horde warnings

diff --git a/Engine/Source/Runtime/Core/Private/Mac/MacPlatformProcess.cpp b/Engine/Source/Runtime/Core/Private/Mac/MacPlatformProcess.cpp
--- a/Engine/Source/Runtime/Core/Private/Mac/MacPlatformProcess.cpp
+++ b/Engine/Source/Runtime/Core/Private/Mac/MacPlatformProcess.cpp
@@ -60,7 +60,7 @@ static void* GetDllHandleImpl(NSString* DylibPath, NSString* ExecutableFolder)
 		Handle = dlopen([DylibPath fileSystemRepresentation], RTLD_LAZY | RTLD_LOCAL);
 	}
 	
-	if (!Handle)
+	if (!Handle && FParse::Param(FCommandLine::Get(), TEXT("dllerrors")))
 	{
 		UE_LOG(LogMac, Warning, TEXT("dlopen failed: %s"), ANSI_TO_TCHAR(dlerror()));
 	}
