Possible fix for compile issue.  I CNR the issue but this should fix it.

diff --git a/Engine/Source/Runtime/Core/Public/ProfilingDebugging/TagTrace.h b/Engine/Source/Runtime/Core/Public/ProfilingDebugging/TagTrace.h
--- a/Engine/Source/Runtime/Core/Public/ProfilingDebugging/TagTrace.h
+++ b/Engine/Source/Runtime/Core/Public/ProfilingDebugging/TagTrace.h
@@ -3,7 +3,8 @@
 #include "CoreTypes.h"
 #include "Trace/Config.h"
 #include "Trace/Detail/LogScope.h"
-#include "UObject/NameTypes.h"
+
+class FName;
 
 ////////////////////////////////////////////////////////////////////////////////
 // Fwd declare ELLMTag
