Refactor of CodepointFromUtf8() to support per-byte iteration.

diff --git a/Engine/Source/Runtime/Core/Private/GenericPlatform/GenericPlatformString.cpp b/Engine/Source/Runtime/Core/Private/GenericPlatform/GenericPlatformString.cpp
--- a/Engine/Source/Runtime/Core/Private/GenericPlatform/GenericPlatformString.cpp
+++ b/Engine/Source/Runtime/Core/Private/GenericPlatform/GenericPlatformString.cpp
@@ -258,230 +258,133 @@ namespace UE::Core::Private
 		}
 	}
 
-	static uint32 CodepointFromUtf8(const UTF8CHAR*& SourceString, const uint32 SourceLengthRemaining)
+	static bool ReadTrailingOctet(uint32& OutOctet, const UTF8CHAR*& Ptr, uint32& Size)
 	{
-		checkSlow(SourceLengthRemaining > 0);
+		// Ensure our string has enough characters to read from
+		if (Size == 0)
+		{
+			return false;
+		}
 
-		const UTF8CHAR* OctetPtr = SourceString;
+		uint32 Octet = (uint32)(uint8)*Ptr;
+		++Ptr;
+		--Size;
 
-		uint32 Codepoint = 0;
-		uint32 Octet = (uint32) ((uint8) *SourceString);
-		uint32 Octet2, Octet3, Octet4;
+		// Format isn't 10xxxxxx?
+		if ((Octet & 192) != 128)
+		{
+			return false;
+		}
+
+		OutOctet = Octet;
+		return true;
+	};
+
+	static uint32 CodepointFromUtf8(const UTF8CHAR*& SourceString, uint32 SourceLengthRemaining)
+	{
+		checkSlow(SourceLengthRemaining > 0);
+
+		uint32 Octet = (uint32)(uint8)*SourceString;
+		++SourceString;
+		--SourceLengthRemaining;
 
 		if (Octet < 128)  // one octet char: 0 to 127
 		{
-			++SourceString;  // read all octets - skip to next possible start of codepoint.
-
-			Codepoint = Octet;
+			return Octet;
 		}
-		else if (Octet < 192)  // bad (starts with 10xxxxxx).
+
+		if (Octet < 192)  // bad (starts with 10xxxxxx).
 		{
 			// Apparently each of these is supposed to be flagged as a bogus
 			//  char, instead of just resyncing to the next valid codepoint.
 
-			++SourceString;  // Sequence was not valid UTF-8. Skip the first byte and continue.
+			// Sequence was not valid UTF-8. Skip the first byte and continue.
 			return UNICODE_BOGUS_CHAR_CODEPOINT;
 		}
-		else if (Octet < 224)  // two octets
-		{
-			// Ensure our string has enough characters to read from
-			if (SourceLengthRemaining < 2)
-			{
-				// Skip to end and write out a single char (we always have room for at least 1 char)
-				SourceString += SourceLengthRemaining;
-				return UNICODE_BOGUS_CHAR_CODEPOINT;
-			}
 
-			Octet -= (128+64);
-			Octet2 = (uint32) ((uint8) *(++OctetPtr));
-			if ((Octet2 & (128 + 64)) != 128)  // Format isn't 10xxxxxx?
-			{
-				++SourceString;  // Sequence was not valid UTF-8. Skip the first byte and continue.
-				return UNICODE_BOGUS_CHAR_CODEPOINT;
-			}
-
-			SourceString += 2;  // read all octets - skip to next possible start of codepoint.
-
-			Codepoint = ((Octet << 6) | (Octet2 - 128));
-			if (Codepoint < 0x80 || Codepoint > 0x7FF)
-			{
-				return UNICODE_BOGUS_CHAR_CODEPOINT;
-			}
-		}
-		else if (Octet < 240)  // three octets
+		uint32 Octet2;
+		if (!ReadTrailingOctet(Octet2, SourceString, SourceLengthRemaining))
 		{
-			// Ensure our string has enough characters to read from
-			if (SourceLengthRemaining < 3)
-			{
-				// Skip to end and write out a single char (we always have room for at least 1 char)
-				SourceString += SourceLengthRemaining;
-				return UNICODE_BOGUS_CHAR_CODEPOINT;
-			}
+			return UNICODE_BOGUS_CHAR_CODEPOINT;
+		}
 
-			Octet -= (128+64+32);
-			Octet2 = (uint32) ((uint8) *(++OctetPtr));
-			if ((Octet2 & (128+64)) != 128)  // Format isn't 10xxxxxx?
+		if (Octet < 224)  // two octets
+		{
+			uint32 Codepoint = ((Octet - 192) << 6) | (Octet2 - 128);
+			if (Codepoint < 0x80 || Codepoint > 0x7FF)
 			{
-				++SourceString;  // Sequence was not valid UTF-8. Skip the first byte and continue.
 				return UNICODE_BOGUS_CHAR_CODEPOINT;
 			}
 
-			Octet3 = (uint32) ((uint8) *(++OctetPtr));
-			if ((Octet3 & (128+64)) != 128)  // Format isn't 10xxxxxx?
-			{
-				++SourceString;  // Sequence was not valid UTF-8. Skip the first byte and continue.
-				return UNICODE_BOGUS_CHAR_CODEPOINT;
-			}
+			// Should only reach here after parsing a legal codepoint
+			return Codepoint;
+		}
 
-			SourceString += 3;  // read all octets - skip to next possible start of codepoint.
+		uint32 Octet3;
+		if (!ReadTrailingOctet(Octet3, SourceString, SourceLengthRemaining))
+		{
+			return UNICODE_BOGUS_CHAR_CODEPOINT;
+		}
 
-			Codepoint = ( ((Octet << 12)) | ((Octet2-128) << 6) | ((Octet3-128)) );
+		if (Octet < 240)  // three octets
+		{
+			uint32 Codepoint = ((Octet - 224) << 12) | ((Octet2 - 128) << 6) | (Octet3 - 128);
 
 			// UTF-8 characters cannot be in the UTF-16 surrogates range.  0xFFFE and 0xFFFF are illegal, too, so we check them at the edge.
 			if (Codepoint < 0x800 || Codepoint > 0xFFFD || IsHighSurrogate(Codepoint) || IsLowSurrogate(Codepoint))
 			{
 				return UNICODE_BOGUS_CHAR_CODEPOINT;
 			}
+
+			// Should only reach here after parsing a legal codepoint
+			return Codepoint;
 		}
-		else if (Octet < 248)  // four octets
+
+		uint32 Octet4;
+		if (!ReadTrailingOctet(Octet4, SourceString, SourceLengthRemaining))
 		{
-			// Ensure our string has enough characters to read from
-			if (SourceLengthRemaining < 4)
-			{
-				// Skip to end and write out a single char (we always have room for at least 1 char)
-				SourceString += SourceLengthRemaining;
-				return UNICODE_BOGUS_CHAR_CODEPOINT;
-			}
+			return UNICODE_BOGUS_CHAR_CODEPOINT;
+		}
 
+		if (Octet < 248)  // four octets
+		{
 			Octet -= (128+64+32+16);
-			Octet2 = (uint32) ((uint8) *(++OctetPtr));
-			if ((Octet2 & (128+64)) != 128)  // Format isn't 10xxxxxx?
-			{
-				++SourceString;  // Sequence was not valid UTF-8. Skip the first byte and continue.
-				return UNICODE_BOGUS_CHAR_CODEPOINT;
-			}
 
-			Octet3 = (uint32) ((uint8) *(++OctetPtr));
-			if ((Octet3 & (128+64)) != 128)  // Format isn't 10xxxxxx?
-			{
-				++SourceString;  // Sequence was not valid UTF-8. Skip the first byte and continue.
-				return UNICODE_BOGUS_CHAR_CODEPOINT;
-			}
-
-			Octet4 = (uint32) ((uint8) *(++OctetPtr));
-			if ((Octet4 & (128+64)) != 128)  // Format isn't 10xxxxxx?
+			uint32 Codepoint = (Octet << 18) | ((Octet2 - 128) << 12) | ((Octet3 - 128) << 6) | (Octet4 - 128);
+			if (Codepoint < 0x10000 || Codepoint > 0x10FFFF)
 			{
-				++SourceString;  // Sequence was not valid UTF-8. Skip the first byte and continue.
 				return UNICODE_BOGUS_CHAR_CODEPOINT;
 			}
 
-			SourceString += 4;  // read all octets - skip to next possible start of codepoint.
+			// Should only reach here after parsing a legal codepoint
+			return Codepoint;
+		}
 
-			Codepoint = ( ((Octet << 18)) | ((Octet2 - 128) << 12) |
-						((Octet3 - 128) << 6) | ((Octet4 - 128)) );
-			if (Codepoint < 0x10000 || Codepoint > 0x10FFFF)
-			{
-				return UNICODE_BOGUS_CHAR_CODEPOINT;
-			}
+		uint32 Octet5;
+		if (!ReadTrailingOctet(Octet5, SourceString, SourceLengthRemaining))
+		{
+			return UNICODE_BOGUS_CHAR_CODEPOINT;
 		}
+
 		// Five and six octet sequences became illegal in rfc3629.
 		//  We throw the codepoint away, but parse them to make sure we move
 		//  ahead the right number of bytes and don't overflow the buffer.
-		else if (Octet < 252)  // five octets
-		{
-			// Ensure our string has enough characters to read from
-			if (SourceLengthRemaining < 5)
-			{
-				// Skip to end and write out a single char (we always have room for at least 1 char)
-				SourceString += SourceLengthRemaining;
-				return UNICODE_BOGUS_CHAR_CODEPOINT;
-			}
-
-			Octet = (uint32) ((uint8) *(++OctetPtr));
-			if ((Octet & (128+64)) != 128)  // Format isn't 10xxxxxx?
-			{
-				++SourceString;  // Sequence was not valid UTF-8. Skip the first byte and continue.
-				return UNICODE_BOGUS_CHAR_CODEPOINT;
-			}
-
-			Octet = (uint32) ((uint8) *(++OctetPtr));
-			if ((Octet & (128+64)) != 128)  // Format isn't 10xxxxxx?
-			{
-				++SourceString;  // Sequence was not valid UTF-8. Skip the first byte and continue.
-				return UNICODE_BOGUS_CHAR_CODEPOINT;
-			}
-
-			Octet = (uint32) ((uint8) *(++OctetPtr));
-			if ((Octet & (128+64)) != 128)  // Format isn't 10xxxxxx?
-			{
-				++SourceString;  // Sequence was not valid UTF-8. Skip the first byte and continue.
-				return UNICODE_BOGUS_CHAR_CODEPOINT;
-			}
-
-			Octet = (uint32) ((uint8) *(++OctetPtr));
-			if ((Octet & (128+64)) != 128)  // Format isn't 10xxxxxx?
-			{
-				++SourceString;  // Sequence was not valid UTF-8. Skip the first byte and continue.
-				return UNICODE_BOGUS_CHAR_CODEPOINT;
-			}
-
-			SourceString += 5;  // read all octets - skip to next possible start of codepoint.
 
+		if (Octet < 252)  // five octets
+		{
+			// Skip to end and write out a single char (we always have room for at least 1 char)
 			return UNICODE_BOGUS_CHAR_CODEPOINT;
 		}
 
-		else  // six octets
-		{
-			// Ensure our string has enough characters to read from
-			if (SourceLengthRemaining < 6)
-			{
-				// Skip to end and write out a single char (we always have room for at least 1 char)
-				SourceString += SourceLengthRemaining;
-				return UNICODE_BOGUS_CHAR_CODEPOINT;
-			}
-
-			Octet = (uint32) ((uint8) *(++OctetPtr));
-			if ((Octet & (128+64)) != 128)  // Format isn't 10xxxxxx?
-			{
-				++SourceString;  // Sequence was not valid UTF-8. Skip the first byte and continue.
-				return UNICODE_BOGUS_CHAR_CODEPOINT;
-			}
-
-			Octet = (uint32) ((uint8) *(++OctetPtr));
-			if ((Octet & (128+64)) != 128)  // Format isn't 10xxxxxx?
-			{
-				++SourceString;  // Sequence was not valid UTF-8. Skip the first byte and continue.
-				return UNICODE_BOGUS_CHAR_CODEPOINT;
-			}
-
-			Octet = (uint32) ((uint8) *(++OctetPtr));
-			if ((Octet & (128+64)) != 128)  // Format isn't 10xxxxxx?
-			{
-				++SourceString;  // Sequence was not valid UTF-8. Skip the first byte and continue.
-				return UNICODE_BOGUS_CHAR_CODEPOINT;
-			}
-
-			Octet = (uint32) ((uint8) *(++OctetPtr));
-			if ((Octet & (128+64)) != 128)  // Format isn't 10xxxxxx?
-			{
-				++SourceString;  // Sequence was not valid UTF-8. Skip the first byte and continue.
-				return UNICODE_BOGUS_CHAR_CODEPOINT;
-			}
-
-			Octet = (uint32) ((uint8) *(++OctetPtr));
-			if ((Octet & (128+64)) != 128)  // Format isn't 10xxxxxx?
-			{
-				++SourceString;  // Sequence was not valid UTF-8. Skip the first byte and continue.
-				return UNICODE_BOGUS_CHAR_CODEPOINT;
-			}
-
-			SourceString += 6;  // read all octets - skip to next possible start of codepoint.
+		// six octets
 
+		uint32 Octet6;
+		if (!ReadTrailingOctet(Octet6, SourceString, SourceLengthRemaining))
+		{
 			return UNICODE_BOGUS_CHAR_CODEPOINT;
 		}
 
-		// Should only reach here after parsing a legal codepoint
-		return Codepoint;
+		return UNICODE_BOGUS_CHAR_CODEPOINT;
 	}
 
 	/**
