BuildVersion: Fixed GetFileNameForCurrentExecutable for Linux

#rnx

diff --git a/Engine/Source/Runtime/Core/Private/Modules/BuildVersion.cpp b/Engine/Source/Runtime/Core/Private/Modules/BuildVersion.cpp
--- a/Engine/Source/Runtime/Core/Private/Modules/BuildVersion.cpp
+++ b/Engine/Source/Runtime/Core/Private/Modules/BuildVersion.cpp
@@ -51,10 +51,10 @@ FString FBuildVersion::GetDefaultFileName()
 FString FBuildVersion::GetFileNameForCurrentExecutable()
 {
 	FString AppExecutableName = FPlatformProcess::ExecutableName();
-#if PLATFORM_WINDOWS || PLATFORM_MAC
-	if(AppExecutableName.EndsWith(TEXT("-Cmd")))
+#if PLATFORM_DESKTOP
+	if (AppExecutableName.EndsWith(TEXT("-Cmd")))
 	{
-		AppExecutableName.LeftInline(AppExecutableName.Len() - 4, false);
+		AppExecutableName.LeftChopInline(4, /*bAllowShrinking*/ false);
 	}
 #endif
 #if UE_BUILD_DEVELOPMENT
