Tasks System: minor cleanup, moved Wait methods to the cpp file

diff --git a/Engine/Source/Runtime/Core/Private/Tasks/TaskPrivate.cpp b/Engine/Source/Runtime/Core/Private/Tasks/TaskPrivate.cpp
--- a/Engine/Source/Runtime/Core/Private/Tasks/TaskPrivate.cpp
+++ b/Engine/Source/Runtime/Core/Private/Tasks/TaskPrivate.cpp
@@ -46,6 +46,93 @@ namespace UE::Tasks
 			LowLevelTasks::FScheduler::Get().TryLaunch(LowLevelTask, LowLevelTasks::EQueuePreference::GlobalQueuePreference, /*bWakeUpWorker=*/ true);
 		}
 
+		void FTaskBase::Wait()
+		{
+			if (IsCompleted())
+			{
+				return;
+			}
+
+			TaskTrace::FWaitingScope WaitingScope(GetTraceId());
+			TRACE_CPUPROFILER_EVENT_SCOPE(Tasks::Wait);
+
+			if (!IsAwaitable())
+			{
+				UE_LOG(LogTemp, Fatal, TEXT("Deadlock detected! A task can't be waited here, e.g. because it's being executed by the currect thread"));
+				return;
+			}
+
+			if (TryRetractAndExecute())
+			{
+				return;
+			}
+
+			// if we are on a named thread, handle waiting in TaskGraph-specific style
+			if (TryWaitOnNamedThread(*this))
+			{
+				return;
+			}
+
+			FEventRef CompletionEvent;
+			auto WaitingTaskBody = [&CompletionEvent] { CompletionEvent->Trigger(); };
+			using FWaitingTask = TExecutableTask<decltype(WaitingTaskBody)>;
+
+			// the task is stored on the stack as we can guarantee that it's out of the system by the end of the call
+			FWaitingTask WaitingTask{ TEXT("Waiting Task"), MoveTemp(WaitingTaskBody), ETaskPriority::Default /* doesn't matter*/, EExtendedTaskPriority::Inline };
+			WaitingTask.AddPrerequisites(*this);
+
+			if (WaitingTask.TryLaunch())
+			{	// was executed inline
+				check(WaitingTask.IsCompleted());
+			}
+			else
+			{
+				CompletionEvent->Wait();
+			}
+
+			// the waiting task will be destroyed leaving this scope, wait for the internal reference to it to be released
+			while (WaitingTask.GetRefCount() != 1)
+			{
+				FPlatformProcess::Yield();
+			}
+		}
+
+		bool FTaskBase::Wait(FTimespan InTimeout)
+		{
+			TaskTrace::FWaitingScope WaitingScope(GetTraceId());
+			TRACE_CPUPROFILER_EVENT_SCOPE(Tasks::Wait);
+
+			FTimeout Timeout{ InTimeout };
+
+			if (TryRetractAndExecute())
+			{
+				return true;
+			}
+
+			if (GetCurrentTask() == this)
+			{
+				UE_LOG(LogTemp, Fatal, TEXT("A task waiting for itself detected"));
+				return true;
+			}
+
+			// the event must be alive for the task and this function lifetime, we don't know which one will be finished first as waiting can 
+			// time out before the waiting task is completed
+			FSharedEventRef CompletionEvent;
+			auto WaitingTaskBody = [CompletionEvent] { CompletionEvent->Trigger(); };
+			using FWaitingTask = TExecutableTask<decltype(WaitingTaskBody)>;
+
+			TRefCountPtr<FWaitingTask> WaitingTask{ FWaitingTask::Create(TEXT("Waiting Task"), MoveTemp(WaitingTaskBody), ETaskPriority::Default /* doesn't matter*/, EExtendedTaskPriority::Inline), /*bAddRef=*/ false };
+			WaitingTask->AddPrerequisites(*this);
+
+			if (WaitingTask->TryLaunch())
+			{	// was executed inline
+				check(WaitingTask->IsCompleted());
+				return true;
+			}
+
+			return CompletionEvent->Wait((uint32)FMath::Clamp<int64>(Timeout.GetRemainingTime().GetTicks() / ETimespan::TicksPerMillisecond, 0, MAX_uint32));
+		}
+
 		FTaskBase* FTaskBase::TryPushIntoPipe()
 		{
 			return GetPipe()->PushIntoPipe(*this);
diff --git a/Engine/Source/Runtime/Core/Public/Tasks/TaskPrivate.h b/Engine/Source/Runtime/Core/Public/Tasks/TaskPrivate.h
--- a/Engine/Source/Runtime/Core/Public/Tasks/TaskPrivate.h
+++ b/Engine/Source/Runtime/Core/Public/Tasks/TaskPrivate.h
@@ -434,8 +434,8 @@ namespace UE::Tasks
 			// `Wait(FTimespan::Zero())` still tries to retract and execute the task, use `IsCompleted()` to check for completeness. 
 			// The version w/o timeout is slightly more efficient.
 			// @return true if the task is completed
-			void Wait();
-			bool Wait(FTimespan Timeout);
+			CORE_API void Wait();
+			CORE_API bool Wait(FTimespan Timeout);
 
 			// waits until the task is completed while executing other tasks
 			void BusyWait()
@@ -871,57 +871,6 @@ namespace UE::Tasks
 			}
 		};
 
-		inline void FTaskBase::Wait()
-		{
-			if (IsCompleted())
-			{
-				return;
-			}
-
-			TaskTrace::FWaitingScope WaitingScope(GetTraceId());
-			TRACE_CPUPROFILER_EVENT_SCOPE(Tasks::Wait);
-
-			if (!IsAwaitable())
-			{
-				UE_LOG(LogTemp, Fatal, TEXT("Deadlock detected! A task can't be waited here, e.g. because it's being executed by the currect thread"));
-				return;
-			}
-
-			if (TryRetractAndExecute())
-			{
-				return;
-			}
-
-			// if we are on a named thread, handle waiting in TaskGraph-specific style
-			if (TryWaitOnNamedThread(*this))
-			{
-				return;
-			}
-
-			FEventRef CompletionEvent;
-			auto WaitingTaskBody = [&CompletionEvent] { CompletionEvent->Trigger(); };
-			using FWaitingTask = TExecutableTask<decltype(WaitingTaskBody)>;
-
-			// the task is stored on the stack as we can guarantee that it's out of the system by the end of the call
-			FWaitingTask WaitingTask{ TEXT("Waiting Task"), MoveTemp(WaitingTaskBody), ETaskPriority::Default /* doesn't matter*/, EExtendedTaskPriority::Inline};
-			WaitingTask.AddPrerequisites(*this);
-
-			if (WaitingTask.TryLaunch())
-			{	// was executed inline
-				check(WaitingTask.IsCompleted());
-			}
-			else
-			{
-				CompletionEvent->Wait();
-			}
-
-			// the waiting task will be destroyed leaving this scope, wait for the internal reference to it to be released
-			while (WaitingTask.GetRefCount() != 1)
-			{
-				FPlatformProcess::Yield();
-			}
-		}
-
 		using FTaskEventBaseAllocator = TLockFreeFixedSizeAllocator_TLSCache<sizeof(FTaskEventBase), PLATFORM_CACHE_LINE_SIZE>;
 		CORE_API extern FTaskEventBaseAllocator TaskEventBaseAllocator;
 
@@ -935,43 +884,6 @@ namespace UE::Tasks
 			TaskEventBaseAllocator.Free(Ptr);
 		}
 
-
-		inline bool FTaskBase::Wait(FTimespan InTimeout)
-		{
-			TaskTrace::FWaitingScope WaitingScope(GetTraceId());
-			TRACE_CPUPROFILER_EVENT_SCOPE(Tasks::Wait);
-
-			FTimeout Timeout{ InTimeout };
-
-			if (TryRetractAndExecute())
-			{
-				return true;
-			}
-
-			if (GetCurrentTask() == this)
-			{
-				UE_LOG(LogTemp, Fatal, TEXT("A task waiting for itself detected"));
-				return true;
-			}
-
-			// the event must be alive for the task and this function lifetime, we don't know which one will be finished first as waiting can 
-			// time out before the waiting task is completed
-			FSharedEventRef CompletionEvent;
-			auto WaitingTaskBody = [CompletionEvent] { CompletionEvent->Trigger(); };
-			using FWaitingTask = TExecutableTask<decltype(WaitingTaskBody)>;
-
-			TRefCountPtr<FWaitingTask> WaitingTask{ FWaitingTask::Create(TEXT("Waiting Task"), MoveTemp(WaitingTaskBody), ETaskPriority::Default /* doesn't matter*/, EExtendedTaskPriority::Inline), /*bAddRef=*/ false};
-			WaitingTask->AddPrerequisites(*this);
-
-			if (WaitingTask->TryLaunch())
-			{	// was executed inline
-				check(WaitingTask->IsCompleted());
-				return true;
-			}
-
-			return CompletionEvent->Wait((uint32)FMath::Clamp<int64>(Timeout.GetRemainingTime().GetTicks() / ETimespan::TicksPerMillisecond, 0, MAX_uint32));
-		}
-
 		// task retraction of multiple tasks. 
 		// @return true if all tasks are completed
 		template<typename TaskCollectionType>
@@ -983,7 +895,7 @@ namespace UE::Tasks
 			{
 				if (Task.IsValid() && !Task.Pimpl->TryRetractAndExecute())
 				{
-					bResult = false;
+					bResult = false; // do not stop here to let this thread to help in executing tasks as much as possible, as it's waiting for their completion anyway
 				}
 			}
 
@@ -1003,7 +915,7 @@ namespace UE::Tasks
 			{
 				if (Task.IsValid() && !Task.Pimpl->TryRetractAndExecute())
 				{
-					bResult = false;
+					bResult = false;  // do not stop here to let this thread to help in executing tasks as much as possible, as it's waiting for their completion anyway
 				}
 
 				if (Timeout)
