Tasks System: cleaned up some includes

diff --git a/Engine/Source/Runtime/Core/Public/Async/TaskTrace.h b/Engine/Source/Runtime/Core/Public/Async/TaskTrace.h
--- a/Engine/Source/Runtime/Core/Public/Async/TaskTrace.h
+++ b/Engine/Source/Runtime/Core/Public/Async/TaskTrace.h
@@ -6,6 +6,7 @@
 #include "HAL/Platform.h"
 #include "Trace/Config.h"
 #include "Trace/Trace.h"
+#include "Trace/Detail/Channel.h"
 
 namespace UE { namespace Trace { class FChannel; } }
 
diff --git a/Engine/Source/Runtime/Core/Public/Tasks/TaskPrivate.h b/Engine/Source/Runtime/Core/Public/Tasks/TaskPrivate.h
--- a/Engine/Source/Runtime/Core/Public/Tasks/TaskPrivate.h
+++ b/Engine/Source/Runtime/Core/Public/Tasks/TaskPrivate.h
@@ -13,20 +13,18 @@
 #include "HAL/PlatformProcess.h"
 #include "HAL/Thread.h"
 #include "Math/NumericLimits.h"
-#include "Misc/ScopeLock.h"
 #include "Misc/Timeout.h"
 #include "ProfilingDebugging/CpuProfilerTrace.h"
 #include "Templates/Invoke.h"
 #include "Templates/RefCounting.h"
 #include "Templates/TypeCompatibleBytes.h"
 #include "Templates/UnrealTemplate.h"
-#include "Trace/Detail/Channel.h"
 
 #include <atomic>
 #include <type_traits>
 
 #if !defined(TASKGRAPH_NEW_FRONTEND)
-#define TASKGRAPH_NEW_FRONTEND 0
+#define TASKGRAPH_NEW_FRONTEND 1
 #endif
 
 namespace UE::Tasks
@@ -747,47 +745,22 @@ namespace UE::Tasks
 			TTypeCompatibleBytes<ResultType> ResultStorage;
 		};
 
-		struct FTaskBlockAllocationTag : FDefaultBlockAllocationTag
-		{
-			static constexpr uint32 BlockSize = 64 * 1024;
-			static constexpr bool AllowOversizedBlocks = false;
-			static constexpr bool RequiresAccurateSize = false;
-			static constexpr bool InlineBlockAllocation = true;
-			static constexpr const char* TagName = "TaskLinearAllocator";
-
-			using Allocator = TBlockAllocationCache<BlockSize, FAlignedAllocator>;
-		};
-
 		// Task implementation that can be executed, as it stores task body. Generic version (for tasks that return non-void results).
 		// In most cases it should be allocated on the heap and used with TRefCountPtr, e.g. @see FTaskHandle. With care, can be allocated on the stack, e.g. see 
 		// WaitingTask in FTaskBase::Wait().
 		// Implements memory allocation from a pooled fixed-size allocator tuned for the everage UE task size
 		template<typename TaskBodyType, typename ResultType = TInvokeResult_T<TaskBodyType>, typename Enable = void>
-		class TExecutableTask final : public TConcurrentLinearObject<TExecutableTask<TaskBodyType>, FTaskBlockAllocationTag>, public TTaskWithResult<ResultType>
+		class TExecutableTaskBase : public TTaskWithResult<ResultType>
 		{
-			UE_NONCOPYABLE(TExecutableTask);
+			UE_NONCOPYABLE(TExecutableTaskBase);
 
 		public:
-			static TExecutableTask* Create(const TCHAR* DebugName, TaskBodyType TaskBody, ETaskPriority Priority, EExtendedTaskPriority InExtendedPriority = EExtendedTaskPriority::None)
-			{
-				return new TExecutableTask(DebugName, MoveTemp(TaskBody), Priority, InExtendedPriority);
-			}
-
-			TExecutableTask(const TCHAR* InDebugName, TaskBodyType&& TaskBody, ETaskPriority InPriority, EExtendedTaskPriority InExtendedPriority):
-				TTaskWithResult<ResultType>(InDebugName, InPriority, InExtendedPriority, 2)
-				// 2 init refs: one for the initial reference (we don't increment it on passing to `TRefCountPtr`), and one for the internal 
-				// reference that keeps the task alive while it's in the system. is released either on task completion or by the scheduler after
-				// trying to execute the task
-			{
-				new(&TaskBodyStorage) TaskBodyType(MoveTemp(TaskBody));
-			}
-
 			virtual bool TryExecuteTask() override
 			{
 				return FTaskBase::TryExecute(
 					[](FTaskBase& Task) 
 					{ 
-						TExecutableTask& This = static_cast<TExecutableTask&>(Task);
+						TExecutableTaskBase& This = static_cast<TExecutableTaskBase&>(Task);
 						new(&This.ResultStorage) ResultType{ Invoke(*This.TaskBodyStorage.GetTypedPtr()) }; 
 
 						// destroy the task body as soon as we are done with it, as it can have captured data sensitive to destruction order
@@ -796,38 +769,33 @@ namespace UE::Tasks
 				);
 			}
 
+		protected:
+			TExecutableTaskBase(const TCHAR* InDebugName, TaskBodyType&& TaskBody, ETaskPriority InPriority, EExtendedTaskPriority InExtendedPriority) :
+				TTaskWithResult<ResultType>(InDebugName, InPriority, InExtendedPriority, 2)
+				// 2 init refs: one for the initial reference (we don't increment it on passing to `TRefCountPtr`), and one for the internal 
+				// reference that keeps the task alive while it's in the system. is released either on task completion or by the scheduler after
+				// trying to execute the task
+			{
+				new(&TaskBodyStorage) TaskBodyType(MoveTemp(TaskBody));
+			}
+
 		private:
 			TTypeCompatibleBytes<TaskBodyType> TaskBodyStorage;
 		};
 
 		// a specialization for tasks that don't return results
 		template<typename TaskBodyType>
-		class TExecutableTask<TaskBodyType, typename TEnableIf<TIsSame<TInvokeResult_T<TaskBodyType>, void>::Value>::Type> final : public TConcurrentLinearObject<TExecutableTask<TaskBodyType>, FTaskBlockAllocationTag>, public FTaskBase
+		class TExecutableTaskBase<TaskBodyType, typename TEnableIf<TIsSame<TInvokeResult_T<TaskBodyType>, void>::Value>::Type> : public FTaskBase
 		{
-			UE_NONCOPYABLE(TExecutableTask);
+			UE_NONCOPYABLE(TExecutableTaskBase);
 
 		public:
-			static TExecutableTask* Create(const TCHAR* DebugName, TaskBodyType TaskBody, ETaskPriority Priority, EExtendedTaskPriority InExtendedPriority = EExtendedTaskPriority::None)
-			{
-				return new TExecutableTask(DebugName, MoveTemp(TaskBody), Priority, InExtendedPriority);
-			}
-
-			TExecutableTask(const TCHAR* InDebugName, TaskBodyType&& TaskBody, ETaskPriority InPriority, EExtendedTaskPriority InExtendedPriority) :
-				FTaskBase(2)
-				// 2 init refs: one for the initial reference (we don't increment it on passing to `TRefCountPtr`), and one for the internal 
-				// reference that keeps the task alive while it's in the system. is released either on task completion or by the scheduler after
-				// trying to execute the task
-			{
-				Init(InDebugName, InPriority, InExtendedPriority);
-				new(&TaskBodyStorage) TaskBodyType(MoveTemp(TaskBody));
-			}
-
 			virtual bool TryExecuteTask() override
 			{
 				return TryExecute(
 					[](FTaskBase& Task)
 					{
-						TExecutableTask& This = static_cast<TExecutableTask&>(Task);
+						TExecutableTaskBase& This = static_cast<TExecutableTaskBase&>(Task);
 						Invoke(*This.TaskBodyStorage.GetTypedPtr()); 
 
 						// destroy the task body as soon as we are done with it, as it can have captured data sensitive to destruction order
@@ -836,10 +804,52 @@ namespace UE::Tasks
 				);
 			}
 
+		protected:
+			TExecutableTaskBase(const TCHAR* InDebugName, TaskBodyType&& TaskBody, ETaskPriority InPriority, EExtendedTaskPriority InExtendedPriority) :
+				FTaskBase(2) // 2 init refs: one for the initial reference (we don't increment it on passing to `TRefCountPtr`), and one for the internal 
+				// reference that keeps the task alive while it's in the system. is released either on task completion or by the scheduler after
+				// trying to execute the task
+			{
+				Init(InDebugName, InPriority, InExtendedPriority);
+				new(&TaskBodyStorage) TaskBodyType(MoveTemp(TaskBody));
+			}
+
 		private:
 			TTypeCompatibleBytes<TaskBodyType> TaskBodyStorage;
 		};
 
+		static constexpr int32 SmallTaskSize = 256;
+		using FExecutableTaskAllocator = TLockFreeFixedSizeAllocator_TLSCache<SmallTaskSize, PLATFORM_CACHE_LINE_SIZE>;
+		CORE_API extern FExecutableTaskAllocator SmallTaskAllocator;
+
+		// a separate derived class to add "small task" allocation optimization to both base class specializations
+		template<typename TaskBodyType>
+		class TExecutableTask final : public TExecutableTaskBase<TaskBodyType>
+		{
+		public:
+			TExecutableTask(const TCHAR* InDebugName, TaskBodyType&& TaskBody, ETaskPriority InPriority, EExtendedTaskPriority InExtendedPriority)
+				: TExecutableTaskBase<TaskBodyType>(InDebugName, MoveTemp(TaskBody), InPriority, InExtendedPriority)
+			{
+			}
+
+			// a helper that deduces the template argument
+			static TExecutableTask* Create(const TCHAR* InDebugName, TaskBodyType&& TaskBody, ETaskPriority InPriority, EExtendedTaskPriority InExtendedPriority)
+			{
+				return new TExecutableTask(InDebugName, MoveTemp(TaskBody), InPriority, InExtendedPriority);
+			}
+
+			void* operator new(size_t Size)
+			{
+				return Size <= SmallTaskSize ? SmallTaskAllocator.Allocate() : GMalloc->Malloc(sizeof(TExecutableTask), PLATFORM_CACHE_LINE_SIZE);
+			}
+
+			void operator delete(void* Ptr, size_t Size)
+			{
+				DestructItem((TExecutableTask*)Ptr);
+				Size <= SmallTaskSize ? SmallTaskAllocator.Free(Ptr) : GMalloc->Free(Ptr);
+			}
+		};
+
 		// waiting on named threads that replicates TaskGraph logic
 		// returns true if called on a named thread
 		CORE_API bool TryWaitOnNamedThread(FTaskBase& Task);
