[Insights] Fixed tracing of FNames; bug where FNameEntryHandle was created invalid (initialized with Block + 1).

diff --git a/Engine/Source/Runtime/Core/Private/UObject/UnrealNames.cpp b/Engine/Source/Runtime/Core/Private/UObject/UnrealNames.cpp
--- a/Engine/Source/Runtime/Core/Private/UObject/UnrealNames.cpp
+++ b/Engine/Source/Runtime/Core/Private/UObject/UnrealNames.cpp
@@ -1584,7 +1584,7 @@ private:
 	FNamePoolShard<ENameCase::IgnoreCase> ComparisonShards[FNamePoolShards];
 
 #if UE_TRACE_ENABLED
-	using FTracedBitSet = TNameAtomicBitSet<(1 << FNameBlockOffsetBits)>;
+	using FTracedBitSet = TNameAtomicBitSet<FNameBlockOffsets>;
 	FTracedBitSet TracedNames[FNameMaxBlocks];
 #endif
 	
@@ -1905,17 +1905,17 @@ UE::Trace::FEventRef32 FNamePool::Trace(const FNameEntryId& EntryId)
 
 void FNamePool::RetraceAll() const
 {
-	uint32 Block = 0;
-	for (const FTracedBitSet& Set : TracedNames)
+	for (uint32 Block = 0; Block < FNameMaxBlocks; ++Block)
 	{
-		++Block;
+		const FTracedBitSet& Set = TracedNames[Block];
+
 		Set.ForEachSetBit([&](uint32 Index)
 		{
 			const FNameEntryHandle Handle(Block, Index);
-			const FNameEntry& DisplayEntry = Resolve(Handle);
 			const FNameEntryId EntryId = Handle;
 			const uint32 DefinitionId = EntryId.ToUnstableInt(); 
-				
+			const FNameEntry& DisplayEntry = Resolve(Handle);
+
 			if (DisplayEntry.IsWide())
 			{
 				TCHAR Chars[NAME_SIZE];
@@ -1933,7 +1933,6 @@ void FNamePool::RetraceAll() const
 					<< FName.DisplayAnsi(Chars, Size);
 			}
 		});
-
 	}
 }
 #endif //UE_TRACE_ENABLED
