[Insights]
  - Fixed Trace.SnapshotFile not working when trace is already connected.
  - Fixed Trace.SnapshotFile to also issue on connection callback. This allows writing events that are not cached but important for the cache (like FName deduplication events).
  - Fixed Trace.SnapshotFile to emit sync events at the end of trace in order for analyzer to properly process all traced events.
  - TraceLog: Fixed CallbackOnConnect when main thread has multiple local buffers (bug where only partial events were traced).

diff --git a/Engine/Source/Runtime/Core/Private/ProfilingDebugging/TraceAuxiliary.cpp b/Engine/Source/Runtime/Core/Private/ProfilingDebugging/TraceAuxiliary.cpp
--- a/Engine/Source/Runtime/Core/Private/ProfilingDebugging/TraceAuxiliary.cpp
+++ b/Engine/Source/Runtime/Core/Private/ProfilingDebugging/TraceAuxiliary.cpp
@@ -97,7 +97,7 @@ public:
 	void					UpdateCsvStats() const;
 	void					StartWorkerThread();
 	void					StartEndFramePump();
-	bool					WriteSnapshot(const TCHAR* FilePath);
+	bool					WriteSnapshot(const TCHAR* InFilePath, const FTraceAuxiliary::FLogCategoryAlias& LogCategory);
 
 	// True if this is parent process with forking requested before forking.
 	bool					IsParentProcessAndPreFork();
@@ -534,15 +534,30 @@ bool FTraceAuxiliaryImpl::WriteToFile(const TCHAR* Path, const FTraceAuxiliary::
 }
 
 ////////////////////////////////////////////////////////////////////////////////
-bool FTraceAuxiliaryImpl::WriteSnapshot(const TCHAR* FilePath)
+bool FTraceAuxiliaryImpl::WriteSnapshot(const TCHAR* InFilePath, const FTraceAuxiliary::FLogCategoryAlias& LogCategory)
 {
+	double StartTime = FPlatformTime::Seconds();
+
 	FString NativePath;
-	if (!FinalizeFilePath(FilePath, NativePath, LogConsoleResponse))
+	if (!FinalizeFilePath(InFilePath, NativePath, LogCategory))
 	{
 		return false;
 	}
 
-	return UE::Trace::WriteSnapshotTo(*NativePath);
+	UE_LOG_REF(LogCategory, Log, TEXT("Writing trace snapshot to '%s'..."), *NativePath);
+
+	const bool bResult = UE::Trace::WriteSnapshotTo(*NativePath);
+
+	if (bResult)
+	{
+		UE_LOG_REF(LogCategory, Display, TEXT("Trace snapshot generated in %.3f seconds to \"%s\"."), FPlatformTime::Seconds() - StartTime, *NativePath);
+	}
+	else
+	{
+		UE_LOG_REF(LogCategory, Error, TEXT("Failed to trace snapshot to \"%s\"."), *NativePath);
+	}
+
+	return bResult;
 }
 
 ////////////////////////////////////////////////////////////////////////////////
@@ -865,15 +880,7 @@ static void TraceAuxiliarySnapshotFile(const TArray<FString>& Args)
 		return;
 	}
 
-	double StartTime = FPlatformTime::Seconds();
-	if (GTraceAuxiliary.WriteSnapshot(FilePath))
-	{
-		UE_LOG(LogConsoleResponse, Display, TEXT("Trace snapshot generated in %f seconds."), FPlatformTime::Seconds() - StartTime);
-	}
-	else
-	{
-		UE_LOG(LogConsoleResponse, Error, TEXT("Trace snapshot failed."));
-	}
+	GTraceAuxiliary.WriteSnapshot(FilePath, LogConsoleResponse);
 }
 
 ////////////////////////////////////////////////////////////////////////////////
@@ -1291,12 +1298,13 @@ bool FTraceAuxiliary::Resume()
 }
 
 ////////////////////////////////////////////////////////////////////////////////
-bool FTraceAuxiliary::WriteSnapshot(const TCHAR* FilePath)
+bool FTraceAuxiliary::WriteSnapshot(const TCHAR* InFilePath)
 {
 #if UE_TRACE_ENABLED
-	return GTraceAuxiliary.WriteSnapshot(FilePath);
-#endif
+	return GTraceAuxiliary.WriteSnapshot(InFilePath, LogCore);
+#else
 	return true;
+#endif
 }
 
 ////////////////////////////////////////////////////////////////////////////////
diff --git a/Engine/Source/Runtime/TraceLog/Private/Trace/TlsBuffer.cpp b/Engine/Source/Runtime/TraceLog/Private/Trace/TlsBuffer.cpp
--- a/Engine/Source/Runtime/TraceLog/Private/Trace/TlsBuffer.cpp
+++ b/Engine/Source/Runtime/TraceLog/Private/Trace/TlsBuffer.cpp
@@ -219,11 +219,74 @@ void Writer_DrainBuffers()
 }
 
 ////////////////////////////////////////////////////////////////////////////////
-bool Writer_DrainLocalBuffer(uint32 ThreadId)
+void Writer_DrainLocalBuffers()
 {
-	return Writer_DrainBuffer(ThreadId, GTlsWriteBuffer);
+	if (GTlsWriteBuffer == &GNullWriteBuffer)
+	{
+		return;
+	}
+
+	const uint32 LocalThreadId = GTlsWriteBuffer->ThreadId;
+
+	struct FRetireList
+	{
+		FWriteBuffer* __restrict Head = nullptr;
+		FWriteBuffer* __restrict Tail = nullptr;
+
+		void Insert(FWriteBuffer* __restrict Buffer)
+		{
+			Buffer->NextBuffer = Head;
+			Head = Buffer;
+			Tail = (Tail != nullptr) ? Tail : Head;
+		}
+	};
+
+	FRetireList RetireList;
+
+	FWriteBuffer* __restrict ActiveThreadList = GActiveThreadList;
+	GActiveThreadList = nullptr;
+
+	{
+		FWriteBuffer* __restrict Buffer = ActiveThreadList;
+		FWriteBuffer* __restrict NextThread;
+
+		// For each thread...
+		for (; Buffer != nullptr; Buffer = NextThread)
+		{
+			NextThread = Buffer->NextThread;
+			uint32 ThreadId = Buffer->ThreadId;
+
+			if (ThreadId == LocalThreadId)
+			{
+				// For each of the thread's buffers...
+				for (FWriteBuffer* __restrict NextBuffer; Buffer != nullptr; Buffer = NextBuffer)
+				{
+					if (Writer_DrainBuffer(ThreadId, Buffer))
+					{
+						break;
+					}
+
+					// Retire the buffer
+					NextBuffer = Buffer->NextBuffer;
+					RetireList.Insert(Buffer);
+				}
+			}
+
+			if (Buffer != nullptr)
+			{
+				Buffer->NextThread = GActiveThreadList;
+				GActiveThreadList = Buffer;
+			}
+		}
+	}
+
+	// Put the retirees we found back into the system again.
+	if (RetireList.Head != nullptr)
+	{
+		Writer_FreeBlockListToPool(RetireList.Head, RetireList.Tail);
+	}
 }
-	
+
 ////////////////////////////////////////////////////////////////////////////////
 void Writer_EndThreadBuffer()
 {
diff --git a/Engine/Source/Runtime/TraceLog/Private/Trace/Writer.cpp b/Engine/Source/Runtime/TraceLog/Private/Trace/Writer.cpp
--- a/Engine/Source/Runtime/TraceLog/Private/Trace/Writer.cpp
+++ b/Engine/Source/Runtime/TraceLog/Private/Trace/Writer.cpp
@@ -51,8 +51,8 @@ void			Writer_CacheOnConnect();
 void			Writer_CallbackOnConnect();
 void			Writer_InitializePool();
 void			Writer_ShutdownPool();
-bool			Writer_DrainLocalBuffer(uint32);
 void			Writer_DrainBuffers();
+void			Writer_DrainLocalBuffers();
 void			Writer_EndThreadBuffer();
 uint32			Writer_GetControlPort();
 void			Writer_UpdateControl();
@@ -447,7 +447,7 @@ static bool Writer_UpdateConnection()
 	{
 		--GPendingDataHandle;
 
-		if (GPendingDataHandle == (~0ull -CloseInertia))
+		if (GPendingDataHandle == (~0ull - CloseInertia))
 		{
 			Writer_Close();
 			GPendingDataHandle = 0;
@@ -481,9 +481,11 @@ static bool Writer_UpdateConnection()
 
 	// Send cached events (i.e. importants) 
 	Writer_CacheOnConnect();
+
 	// Issue on connection callback. This allows writing events that are
 	// not cached but important for the cache
 	Writer_CallbackOnConnect();
+
 	// Finally write the events in the tail buffer
 	Writer_TailOnConnect();
 
@@ -551,11 +553,14 @@ void Writer_CallbackOnConnect()
 	}
 	
 	Writer_TailPause(true);
+
 	// Issue callback. We assume any events emitted here are not marked as
 	// important and emitted on this thread.
 	GOnConnection();
+
 	// Drain only the events for this thread passing tail as you go.
-	Writer_DrainLocalBuffer(ETransportTid::Bias);
+	Writer_DrainLocalBuffers();
+
 	Writer_TailPause(false);
 }
 
@@ -907,6 +912,7 @@ bool Writer_WriteSnapshotTo(const ANSICHAR* Path)
 
 	{
 		TStashGlobal DataHandle(GDataHandle);
+		TStashGlobal PendingDataHandle(GPendingDataHandle);
 		TStashGlobal SyncPacketCountdown(GSyncPacketCountdown, GNumSyncPackets);
 		TStashGlobal TraceStatistics(GTraceStatistics);
 
@@ -920,12 +926,23 @@ bool Writer_WriteSnapshotTo(const ANSICHAR* Path)
 		// The first events we will send are ones that describe the trace's events
 		Writer_DescribeEvents(FEventNode::Read());
 
-		// Send cached events (i.e. importants) and the tail of recent events
+		// Send cached events (i.e. importants)
 		Writer_CacheOnConnect();
+
+		// Issue on connection callback. This allows writing events that are
+		// not cached but important for the cache
+		Writer_CallbackOnConnect();
+
+		// Finally write the events in the tail buffer
 		Writer_TailOnConnect();
 
 		// Send sync packets to help parsers digest any out-of-order events
-		Writer_SendSync();
+		GSyncPacketCountdown = GNumSyncPackets;
+		while (GSyncPacketCountdown > 0)
+		{
+			Writer_SendSync();
+		}
+
 		Writer_Close();
 	}
 
