partially reverted 21177799 that took much more changes than it was intended

diff --git a/Engine/Source/Runtime/Core/Public/Tasks/TaskPrivate.h b/Engine/Source/Runtime/Core/Public/Tasks/TaskPrivate.h
--- a/Engine/Source/Runtime/Core/Public/Tasks/TaskPrivate.h
+++ b/Engine/Source/Runtime/Core/Public/Tasks/TaskPrivate.h
@@ -24,7 +24,7 @@
 #include <type_traits>
 
 #if !defined(TASKGRAPH_NEW_FRONTEND)
-#define TASKGRAPH_NEW_FRONTEND 1
+#define TASKGRAPH_NEW_FRONTEND 0
 #endif
 
 namespace UE::Tasks
@@ -745,32 +745,33 @@ namespace UE::Tasks
 			TTypeCompatibleBytes<ResultType> ResultStorage;
 		};
 
+		struct FTaskBlockAllocationTag : FDefaultBlockAllocationTag
+		{
+			static constexpr uint32 BlockSize = 64 * 1024;
+			static constexpr bool AllowOversizedBlocks = false;
+			static constexpr bool RequiresAccurateSize = false;
+			static constexpr bool InlineBlockAllocation = true;
+			static constexpr const char* TagName = "TaskLinearAllocator";
+
+			using Allocator = TBlockAllocationCache<BlockSize, FAlignedAllocator>;
+		};
+
 		// Task implementation that can be executed, as it stores task body. Generic version (for tasks that return non-void results).
 		// In most cases it should be allocated on the heap and used with TRefCountPtr, e.g. @see FTaskHandle. With care, can be allocated on the stack, e.g. see 
 		// WaitingTask in FTaskBase::Wait().
 		// Implements memory allocation from a pooled fixed-size allocator tuned for the everage UE task size
 		template<typename TaskBodyType, typename ResultType = TInvokeResult_T<TaskBodyType>, typename Enable = void>
-		class TExecutableTaskBase : public TTaskWithResult<ResultType>
+		class TExecutableTask final : public TConcurrentLinearObject<TExecutableTask<TaskBodyType>, FTaskBlockAllocationTag>, public TTaskWithResult<ResultType>
 		{
-			UE_NONCOPYABLE(TExecutableTaskBase);
+			UE_NONCOPYABLE(TExecutableTask);
 
 		public:
-			virtual bool TryExecuteTask() override
+			static TExecutableTask* Create(const TCHAR* DebugName, TaskBodyType TaskBody, ETaskPriority Priority, EExtendedTaskPriority InExtendedPriority = EExtendedTaskPriority::None)
 			{
-				return FTaskBase::TryExecute(
-					[](FTaskBase& Task) 
-					{ 
-						TExecutableTaskBase& This = static_cast<TExecutableTaskBase&>(Task);
-						new(&This.ResultStorage) ResultType{ Invoke(*This.TaskBodyStorage.GetTypedPtr()) }; 
-
-						// destroy the task body as soon as we are done with it, as it can have captured data sensitive to destruction order
-						DestructItem(This.TaskBodyStorage.GetTypedPtr());
-					}
-				);
+				return new TExecutableTask(DebugName, MoveTemp(TaskBody), Priority, InExtendedPriority);
 			}
 
-		protected:
-			TExecutableTaskBase(const TCHAR* InDebugName, TaskBodyType&& TaskBody, ETaskPriority InPriority, EExtendedTaskPriority InExtendedPriority) :
+			TExecutableTask(const TCHAR* InDebugName, TaskBodyType&& TaskBody, ETaskPriority InPriority, EExtendedTaskPriority InExtendedPriority) :
 				TTaskWithResult<ResultType>(InDebugName, InPriority, InExtendedPriority, 2)
 				// 2 init refs: one for the initial reference (we don't increment it on passing to `TRefCountPtr`), and one for the internal 
 				// reference that keeps the task alive while it's in the system. is released either on task completion or by the scheduler after
@@ -779,24 +780,13 @@ namespace UE::Tasks
 				new(&TaskBodyStorage) TaskBodyType(MoveTemp(TaskBody));
 			}
 
-		private:
-			TTypeCompatibleBytes<TaskBodyType> TaskBodyStorage;
-		};
-
-		// a specialization for tasks that don't return results
-		template<typename TaskBodyType>
-		class TExecutableTaskBase<TaskBodyType, typename TEnableIf<TIsSame<TInvokeResult_T<TaskBodyType>, void>::Value>::Type> : public FTaskBase
-		{
-			UE_NONCOPYABLE(TExecutableTaskBase);
-
-		public:
 			virtual bool TryExecuteTask() override
 			{
-				return TryExecute(
+				return FTaskBase::TryExecute(
 					[](FTaskBase& Task)
 					{
-						TExecutableTaskBase& This = static_cast<TExecutableTaskBase&>(Task);
-						Invoke(*This.TaskBodyStorage.GetTypedPtr()); 
+						TExecutableTask& This = static_cast<TExecutableTask&>(Task);
+						new(&This.ResultStorage) ResultType{ Invoke(*This.TaskBodyStorage.GetTypedPtr()) };
 
 						// destroy the task body as soon as we are done with it, as it can have captured data sensitive to destruction order
 						DestructItem(This.TaskBodyStorage.GetTypedPtr());
@@ -804,50 +794,48 @@ namespace UE::Tasks
 				);
 			}
 
-		protected:
-			TExecutableTaskBase(const TCHAR* InDebugName, TaskBodyType&& TaskBody, ETaskPriority InPriority, EExtendedTaskPriority InExtendedPriority) :
-				FTaskBase(2) // 2 init refs: one for the initial reference (we don't increment it on passing to `TRefCountPtr`), and one for the internal 
-				// reference that keeps the task alive while it's in the system. is released either on task completion or by the scheduler after
-				// trying to execute the task
-			{
-				Init(InDebugName, InPriority, InExtendedPriority);
-				new(&TaskBodyStorage) TaskBodyType(MoveTemp(TaskBody));
-			}
-
 		private:
 			TTypeCompatibleBytes<TaskBodyType> TaskBodyStorage;
 		};
 
-		static constexpr int32 SmallTaskSize = 256;
-		using FExecutableTaskAllocator = TLockFreeFixedSizeAllocator_TLSCache<SmallTaskSize, PLATFORM_CACHE_LINE_SIZE>;
-		CORE_API extern FExecutableTaskAllocator SmallTaskAllocator;
-
-		// a separate derived class to add "small task" allocation optimization to both base class specializations
+		// a specialization for tasks that don't return results
 		template<typename TaskBodyType>
-		class TExecutableTask final : public TExecutableTaskBase<TaskBodyType>
+		class TExecutableTask<TaskBodyType, typename TEnableIf<TIsSame<TInvokeResult_T<TaskBodyType>, void>::Value>::Type> final : public TConcurrentLinearObject<TExecutableTask<TaskBodyType>, FTaskBlockAllocationTag>, public FTaskBase
 		{
+			UE_NONCOPYABLE(TExecutableTask);
+
 		public:
-			TExecutableTask(const TCHAR* InDebugName, TaskBodyType&& TaskBody, ETaskPriority InPriority, EExtendedTaskPriority InExtendedPriority)
-				: TExecutableTaskBase<TaskBodyType>(InDebugName, MoveTemp(TaskBody), InPriority, InExtendedPriority)
+			static TExecutableTask* Create(const TCHAR* DebugName, TaskBodyType TaskBody, ETaskPriority Priority, EExtendedTaskPriority InExtendedPriority = EExtendedTaskPriority::None)
 			{
+				return new TExecutableTask(DebugName, MoveTemp(TaskBody), Priority, InExtendedPriority);
 			}
 
-			// a helper that deduces the template argument
-			static TExecutableTask* Create(const TCHAR* InDebugName, TaskBodyType&& TaskBody, ETaskPriority InPriority, EExtendedTaskPriority InExtendedPriority)
+			TExecutableTask(const TCHAR* InDebugName, TaskBodyType&& TaskBody, ETaskPriority InPriority, EExtendedTaskPriority InExtendedPriority) :
+				FTaskBase(2)
+				// 2 init refs: one for the initial reference (we don't increment it on passing to `TRefCountPtr`), and one for the internal 
+				// reference that keeps the task alive while it's in the system. is released either on task completion or by the scheduler after
+				// trying to execute the task
 			{
-				return new TExecutableTask(InDebugName, MoveTemp(TaskBody), InPriority, InExtendedPriority);
+				Init(InDebugName, InPriority, InExtendedPriority);
+				new(&TaskBodyStorage) TaskBodyType(MoveTemp(TaskBody));
 			}
 
-			void* operator new(size_t Size)
+			virtual bool TryExecuteTask() override
 			{
-				return Size <= SmallTaskSize ? SmallTaskAllocator.Allocate() : GMalloc->Malloc(sizeof(TExecutableTask), PLATFORM_CACHE_LINE_SIZE);
-			}
+				return TryExecute(
+					[](FTaskBase& Task)
+					{
+						TExecutableTask& This = static_cast<TExecutableTask&>(Task);
+						Invoke(*This.TaskBodyStorage.GetTypedPtr());
 
-			void operator delete(void* Ptr, size_t Size)
-			{
-				DestructItem((TExecutableTask*)Ptr);
-				Size <= SmallTaskSize ? SmallTaskAllocator.Free(Ptr) : GMalloc->Free(Ptr);
+						// destroy the task body as soon as we are done with it, as it can have captured data sensitive to destruction order
+						DestructItem(This.TaskBodyStorage.GetTypedPtr());
+					}
+				);
 			}
+
+		private:
+			TTypeCompatibleBytes<TaskBodyType> TaskBodyStorage;
 		};
 
 		// waiting on named threads that replicates TaskGraph logic
@@ -937,4 +925,4 @@ namespace UE::Tasks
 			return bResult;
 		}
 	}
-}
\ No newline at end of file
+}
