Logging: Fixed AppendFormatLogLine to work when TCHAR is UTF8CHAR

#rnx

diff --git a/Engine/Source/Runtime/Core/Private/Misc/OutputDeviceHelper.cpp b/Engine/Source/Runtime/Core/Private/Misc/OutputDeviceHelper.cpp
--- a/Engine/Source/Runtime/Core/Private/Misc/OutputDeviceHelper.cpp
+++ b/Engine/Source/Runtime/Core/Private/Misc/OutputDeviceHelper.cpp
@@ -8,7 +8,7 @@
 #include "Misc/StringBuilder.h"
 
 void FOutputDeviceHelper::AppendFormatLogLine(
-	FStringBuilderBase& Format,
+	FWideStringBuilderBase& Format,
 	const ELogVerbosity::Type Verbosity,
 	const FName& Category,
 	const TCHAR* const Message,
@@ -21,22 +21,22 @@ void FOutputDeviceHelper::AppendFormatLogLine(
 		case ELogTimes::SinceGStartTime:
 		{
 			const double RealTime = (Time < 0.0) ? (FPlatformTime::Seconds() - GStartTime) : Time;
-			Format.Appendf(TEXT("[%07.2f][%3llu]"), RealTime, GFrameCounter % 1000);
+			Format.Appendf(WIDETEXT("[%07.2f][%3llu]"), RealTime, GFrameCounter % 1000);
 			break;
 		}
 
 		case ELogTimes::UTC:
 			FDateTime::UtcNow().ToString(TEXT("[%Y.%m.%d-%H.%M.%S:%s]"), Format);
-			Format.Appendf(TEXT("[%3llu]"), GFrameCounter % 1000);
+			Format.Appendf(WIDETEXT("[%3llu]"), GFrameCounter % 1000);
 			break;
 
 		case ELogTimes::Local:
 			FDateTime::Now().ToString(TEXT("[%Y.%m.%d-%H.%M.%S:%s]"), Format);
-			Format.Appendf(TEXT("[%3llu]"), GFrameCounter % 1000);
+			Format.Appendf(WIDETEXT("[%3llu]"), GFrameCounter % 1000);
 			break;
 
 		case ELogTimes::Timecode:
-			Format.Appendf(TEXT("[%s][%3llu]"), *FApp::GetTimecode().ToString(), GFrameCounter % 1000);
+			Format.Appendf(WIDETEXT("[%s][%3llu]"), *FApp::GetTimecode().ToString(), GFrameCounter % 1000);
 			break;
 
 		default:
@@ -53,19 +53,19 @@ void FOutputDeviceHelper::AppendFormatLogLine(
 	if (bShowCategory)
 	{
 		Category.AppendString(Format);
-		Format.Append(TEXT(": "));
+		Format.Append(WIDETEXT(": "));
 
 		if (GPrintLogVerbosity && Verbosity != ELogVerbosity::Log)
 		{
 			Format.Append(ToString(Verbosity));
-			Format.Append(TEXT(": "));
+			Format.Append(WIDETEXT(": "));
 		}
 	}
 	else if (GPrintLogVerbosity && Verbosity != ELogVerbosity::Log)
 	{
 #if !HACK_HEADER_GENERATOR
 		Format.Append(ToString(Verbosity));
-		Format.Append(TEXT(": "));
+		Format.Append(WIDETEXT(": "));
 #endif
 	}
 
@@ -84,7 +84,7 @@ void FOutputDeviceHelper::AppendFormatLogLine(
 	const double Time,
 	int32* const OutCategoryIndex)
 {
-	TStringBuilder<128> Prefix;
+	TWideStringBuilder<128> Prefix;
 	AppendFormatLogLine(Prefix, Verbosity, Category, nullptr, LogTime, Time, OutCategoryIndex);
 	Format.Append(Prefix);
 
diff --git a/Engine/Source/Runtime/Core/Public/Misc/OutputDeviceHelper.h b/Engine/Source/Runtime/Core/Public/Misc/OutputDeviceHelper.h
--- a/Engine/Source/Runtime/Core/Public/Misc/OutputDeviceHelper.h
+++ b/Engine/Source/Runtime/Core/Public/Misc/OutputDeviceHelper.h
@@ -19,7 +19,7 @@ class FString;
 struct CORE_API FOutputDeviceHelper
 {
 	/** Append a formatted log line to the string builder. */
-	static void AppendFormatLogLine(FStringBuilderBase& Output, ELogVerbosity::Type Verbosity, const FName& Category, const TCHAR* Message = nullptr, ELogTimes::Type LogTime = ELogTimes::None, double Time = -1.0, int32* OutCategoryIndex = nullptr);
+	static void AppendFormatLogLine(FWideStringBuilderBase& Output, ELogVerbosity::Type Verbosity, const FName& Category, const TCHAR* Message = nullptr, ELogTimes::Type LogTime = ELogTimes::None, double Time = -1.0, int32* OutCategoryIndex = nullptr);
 	static void AppendFormatLogLine(FUtf8StringBuilderBase& Output, ELogVerbosity::Type Verbosity, const FName& Category, const TCHAR* Message = nullptr, ELogTimes::Type LogTime = ELogTimes::None, double Time = -1.0, int32* OutCategoryIndex = nullptr);
 
 	/**
