[Runtime Core]
* TrackedActivity - New class that will be used to track what is going on in running process. This is being used by the new console window. This is designed to be used mainly for developers and not end consumers

diff --git a/Engine/Source/Runtime/Core/Private/Misc/TrackedActivity.cpp b/Engine/Source/Runtime/Core/Private/Misc/TrackedActivity.cpp
new file mode 100644
--- /dev/null
+++ b/Engine/Source/Runtime/Core/Private/Misc/TrackedActivity.cpp
@@ -0,0 +1,230 @@
+// Copyright Epic Games, Inc. All Rights Reserved.
+
+#include "Misc/TrackedActivity.h"
+
+
+// Singleton class that keep track of all the existing activities and their stacks
+class FTrackedActivityManager
+{
+public:
+	static FTrackedActivityManager& Get()
+	{
+		static FTrackedActivityManager Manager;
+		return Manager;
+	}
+
+	struct FActivity;
+	using ELight = FTrackedActivity::ELight;
+	using EType = FTrackedActivity::EType;
+
+	struct FInfo
+	{
+		FString Status;
+		ELight Light;
+		bool bShowParent;
+	};
+
+	struct FActivity
+	{
+		FString Name;
+		TArray<FInfo> Stack;
+		EType Type;
+		uint32 Id;
+	};
+
+	FCriticalSection ActivitiesCs;
+	TArray<FActivity*> Activities;
+	TUniqueFunction<void(FTrackedActivity::EEvent Event, const FTrackedActivity::FInfo& Info)> EventFunc;
+	uint32 EventFuncMaxDepth = ~0u;
+	uint32 IdCounter = 1u;
+
+	FActivity* Create(const TCHAR* Name, const TCHAR* Status, ELight Light, EType Type)
+	{
+		FActivity* Activity = new FActivity();
+		Activity->Name = Name;
+		Activity->Type = Type;
+		Activity->Stack.Add({ Status, Light, false });
+
+		FScopeLock _(&ActivitiesCs);
+		Activity->Id = IdCounter++;
+		Activities.Add(Activity);
+		SendEvent(FTrackedActivity::EEvent::Added, *Activity);
+		return Activity;
+	}
+
+	void Destroy(FActivity* Activity)
+	{
+		check(Activity->Stack.Num() == 1);
+		{
+			FScopeLock _(&ActivitiesCs);
+			Activities.Remove(Activity);
+			SendEvent(FTrackedActivity::EEvent::Removed, *Activity);
+		}
+		delete Activity;
+	}
+
+	uint32 Push(FActivity& Activity, const TCHAR* Status, bool bShowParent, ELight Light)
+	{
+		FScopeLock _(&ActivitiesCs);
+		uint32 Index = Activity.Stack.Num();
+		Activity.Stack.Add({ Status, Light, bShowParent });
+		SendEvent(FTrackedActivity::EEvent::Changed, Activity);
+		return Index;
+	}
+
+	void Pop(FActivity& Activity)
+	{
+		FScopeLock _(&ActivitiesCs);
+		Activity.Stack.SetNum(Activity.Stack.Num() - 1);
+		SendEvent(FTrackedActivity::EEvent::Changed, Activity);
+	}
+
+	void Update(FActivity& Activity, const TCHAR* Status, uint32 Index)
+	{
+		FScopeLock _(&ActivitiesCs);
+		FTrackedActivityManager::FInfo& Info = Index == ~0u ? Activity.Stack.Last() : Activity.Stack[Index];
+		Info.Status = Status;
+		SendEvent(FTrackedActivity::EEvent::Changed, Activity);
+	}
+
+	void Update(FActivity& Activity, const TCHAR* Status, ELight Light, uint32 Index)
+	{
+		FScopeLock _(&ActivitiesCs);
+		FTrackedActivityManager::FInfo& Info = Index == ~0u ? Activity.Stack.Last() : Activity.Stack[Index];
+		Info.Status = Status;
+		Info.Light = Light;
+		SendEvent(FTrackedActivity::EEvent::Changed, Activity);
+	}
+
+	void Update(FActivity& Activity, ELight Light, uint32 Index)
+	{
+		FScopeLock _(&ActivitiesCs);
+		FTrackedActivityManager::FInfo& Info = Index == ~0u ? Activity.Stack.Last() : Activity.Stack[Index];
+		Info.Light = Light;
+		SendEvent(FTrackedActivity::EEvent::Changed, Activity);
+	}
+
+	FTrackedActivity::FInfo GetInfo(const FActivity& Activity, uint32 MaxDepth, TStringBuilder<256>& Temp)
+	{
+		uint32 StackIndex = FMath::Min<uint32>(Activity.Stack.Num(), MaxDepth) - 1;
+		auto& Info = Activity.Stack[StackIndex];
+
+		const TCHAR* Status = *Info.Status;
+		ELight Light = Info.Light;
+
+		if (Info.bShowParent)
+		{
+			Temp.Append(Activity.Stack[StackIndex - 1].Status);
+			Temp.Append(Status);
+			Status = *Temp;
+		}
+
+		while (Light == ELight::Inherit)
+		{
+			if (StackIndex == 0)
+			{
+				Light = ELight::None;
+				break;
+			}
+			Light = Activity.Stack[--StackIndex].Light;
+		}
+
+		return { *Activity.Name, Status, Light, Activity.Type, Activity.Id };
+	}
+
+	void TraverseActivities(const TFunction<void(const FTrackedActivity::FInfo& Info)>& Func)
+	{
+		FScopeLock _(&ActivitiesCs);
+
+		for (auto A : Activities)
+		{
+			TStringBuilder<256> Temp;
+			Func(GetInfo(*A, ~0u, Temp));
+		}
+	}
+
+	void RegisterEventListener(TUniqueFunction<void(FTrackedActivity::EEvent Event, const FTrackedActivity::FInfo& Info)>&& Func, uint32 MaxDepth)
+	{
+		checkf(!EventFunc, TEXT("TrackedActivity system is only supporting one event listener at the time. Please add support"));
+		FScopeLock _(&ActivitiesCs);
+
+		EventFunc = MoveTemp(Func);
+		EventFuncMaxDepth = MaxDepth;
+	}
+
+	void SendEvent(FTrackedActivity::EEvent Event, FActivity& Activity)
+	{
+		if (!EventFunc)
+			return;
+		TStringBuilder<256> Temp;
+		EventFunc(Event, GetInfo(Activity, EventFuncMaxDepth, Temp));
+	}
+};
+
+
+FTrackedActivity::FTrackedActivity(const TCHAR* Name, const TCHAR* Status, ELight Light, EType Type)
+{
+	Internal = FTrackedActivityManager::Get().Create(Name, Status, Light, Type);
+}
+
+FTrackedActivity::~FTrackedActivity()
+{
+	FTrackedActivityManager::Get().Destroy((FTrackedActivityManager::FActivity*)Internal);
+}
+
+uint32 FTrackedActivity::Push(const TCHAR* Status, bool bShowParent, ELight Light)
+{
+	return FTrackedActivityManager::Get().Push(*(FTrackedActivityManager::FActivity*)Internal, Status, bShowParent, Light);
+}
+
+void FTrackedActivity::Pop()
+{
+	FTrackedActivityManager::Get().Pop(*(FTrackedActivityManager::FActivity*)Internal);
+}
+
+void FTrackedActivity::Update(const TCHAR* Status, uint32 Index)
+{
+	FTrackedActivityManager::Get().Update(*(FTrackedActivityManager::FActivity*)Internal, Status, Index);
+}
+
+void FTrackedActivity::Update(const TCHAR* Status, ELight Light, uint32 Index)
+{
+	FTrackedActivityManager::Get().Update(*(FTrackedActivityManager::FActivity*)Internal, Status, Light, Index);
+}
+
+void FTrackedActivity::Update(ELight Light, uint32 Index)
+{
+	FTrackedActivityManager::Get().Update(*(FTrackedActivityManager::FActivity*)Internal, Light, Index);
+}
+
+FTrackedActivity& FTrackedActivity::GetEngineActivity()
+{
+	static FTrackedActivity A(TEXT("Status"), TEXT("Unknown"));
+	return A;
+}
+
+FTrackedActivity& FTrackedActivity::GetIOActivity()
+{
+	static FTrackedActivity A(TEXT("I/O"), TEXT("Idle"));
+	return A;
+}
+
+void FTrackedActivity::RegisterEventListener(TUniqueFunction<void(EEvent Event, const FInfo& Info)>&& Func, uint32 MaxDepth)
+{
+	FTrackedActivityManager::Get().RegisterEventListener(MoveTemp(Func), MaxDepth);
+}
+
+void FTrackedActivity::TraverseActivities(const TFunction<void(const FInfo& Info)>& Func)
+{
+	FTrackedActivityManager::Get().TraverseActivities(Func);
+}
+
+FTrackedActivityScope::FTrackedActivityScope(FTrackedActivity& A, const TCHAR* Status, bool bShowParent, FTrackedActivity::ELight Light)
+:	Activity(A)
+{
+	Activity.Push(Status, bShowParent, Light);
+}
+FTrackedActivityScope::~FTrackedActivityScope()
+{
+	Activity.Pop();
+}
diff --git a/Engine/Source/Runtime/Core/Public/Misc/TrackedActivity.h b/Engine/Source/Runtime/Core/Public/Misc/TrackedActivity.h
new file mode 100644
--- /dev/null
+++ b/Engine/Source/Runtime/Core/Public/Misc/TrackedActivity.h
@@ -0,0 +1,119 @@
+// Copyright Epic Games, Inc. All Rights Reserved.
+
+#pragma once
+
+#include "CoreTypes.h"
+#include "Templates/Function.h"
+
+/**
+ * Tracked Activity is used to be able to visualize on a semi-high level what is going on in the process.
+ * It is very useful when wanting to show the status of online, or if the runtime is waiting on loading something
+ * When new console is enabled tracked activities show at the bottom of the window under the log
+ * Tracked Activities can be created/updated/destroyed on multiple threads
+ */
+class CORE_API FTrackedActivity
+{
+public:
+	/**
+	* Enum to specify status light for an activity
+	* Is used by the new console to show a colored dot in front of the tracked activity entry
+	*/
+	enum class ELight
+	{
+		None,
+		Red,
+		Yellow,
+		Green,
+		Inherit,
+	};
+
+	/**
+	* Enum to specify status type of activity
+	* In the new console, 'Activity' will make sure tracked activity appears to the left of the window
+	* 'Info' appears to the right and use less width. Typically use 'Info' for things that are more static
+	*/
+	enum class EType
+	{
+		Activity,
+		Info,
+		Debug,
+	};
+
+	/** Ctor/Dtor */
+	FTrackedActivity(const TCHAR* Name, const TCHAR* Status, ELight Light = ELight::Inherit, EType Type = EType::Activity);
+	~FTrackedActivity();
+
+
+	/**
+	* Pushes new status on to tracked activity, will require a pop to get back to previous status
+	* ShowParent can be used to make sure Status of parent is visible in front of status of pushed scope.
+	*/
+	uint32	Push(const TCHAR* Status, bool bShowParent = false, ELight Light = ELight::Inherit);
+	void	Pop();
+
+	/** Updates status. If Index is ~0u entry at top of stack will be updated (which is the one showing) */
+	void	Update(const TCHAR* Status, uint32 Index = ~0u);
+	void	Update(const TCHAR* Status, ELight Light, uint32 Index = ~0u);
+	void	Update(ELight Light, uint32 Index = ~0u);
+
+
+	/**
+	* Process Engine Activity. General status of the Engine. Used by Engine initialization etc.
+	* By default, status light will stay yellow until main threads hits Tick update, then it turns green.
+	*/
+	static FTrackedActivity& GetEngineActivity();
+
+	/** I/O Activity.Shows current I / O operation.If plugin / game have their own I / O, scopes will need to be manually added */
+	static FTrackedActivity& GetIOActivity();
+
+
+
+	struct FInfo
+	{
+		const TCHAR* Name;
+		const TCHAR* Status;
+		ELight Light; // Can not be "Inherit"
+		EType Type;
+		uint32 Id; // Unique Id for Activity (Is thread safe counter starting at 1)
+	};
+
+	/**
+	* Traverses all FTrackedActivities in the order they were added.
+	*/
+	static void TraverseActivities(const TFunction<void(const FInfo& Info)>& Func);
+
+
+
+	/** Enum used for Event listener to identify type of activity change */
+	enum class EEvent : int32
+	{
+		Added = 1,
+		Removed = 2,
+		Changed = 3
+	};
+
+	/** Register listener that can track adds, removes and changes when they happen.
+	* Listener will be called from the same thread as the event happens, so make sure listener is threadsafe.
+	*/
+	static void RegisterEventListener(TUniqueFunction<void(EEvent Event, const FInfo& Info)>&& Func, uint32 MaxDepth = ~0u);
+
+
+private:
+	FTrackedActivity(const FTrackedActivity& O) = delete;
+	FTrackedActivity& operator=(const FTrackedActivity& O) = delete;
+
+	void* Internal;
+};
+
+
+/**
+* RAII class that calls push in ctor and pop in dtor.
+*/
+class CORE_API FTrackedActivityScope
+{
+public:
+	FTrackedActivityScope(FTrackedActivity& Activity, const TCHAR* Status, bool bShowParent = false, FTrackedActivity::ELight Light = FTrackedActivity::ELight::Inherit);
+	~FTrackedActivityScope();
+private:
+	FTrackedActivity& Activity;
+};
