TrackedActivity.cpp - Fixed unity build

diff --git a/Engine/Source/Runtime/Core/Private/Misc/TrackedActivity.cpp b/Engine/Source/Runtime/Core/Private/Misc/TrackedActivity.cpp
--- a/Engine/Source/Runtime/Core/Private/Misc/TrackedActivity.cpp
+++ b/Engine/Source/Runtime/Core/Private/Misc/TrackedActivity.cpp
@@ -1,7 +1,7 @@
 // Copyright Epic Games, Inc. All Rights Reserved.
 
 #include "Misc/TrackedActivity.h"
-
+#include "Containers/UnrealString.h"
 
 // Singleton class that keep track of all the existing activities and their stacks
 class FTrackedActivityManager
