TrackedActivity.cpp - Fixed unity build

diff --git a/Engine/Source/Runtime/Core/Private/Misc/TrackedActivity.cpp b/Engine/Source/Runtime/Core/Private/Misc/TrackedActivity.cpp
--- a/Engine/Source/Runtime/Core/Private/Misc/TrackedActivity.cpp
+++ b/Engine/Source/Runtime/Core/Private/Misc/TrackedActivity.cpp
@@ -2,6 +2,9 @@
 
 #include "Misc/TrackedActivity.h"
 #include "Containers/UnrealString.h"
+#include "HAL/CriticalSection.h"
+#include "Misc/ScopeLock.h"
+#include "Misc/StringBuilder.h"
 
 // Singleton class that keep track of all the existing activities and their stacks
 class FTrackedActivityManager
