Attempting to fix a compile issue that I CNR.

diff --git a/Engine/Source/Runtime/Core/Public/Containers/UnrealString.h b/Engine/Source/Runtime/Core/Public/Containers/UnrealString.h
--- a/Engine/Source/Runtime/Core/Public/Containers/UnrealString.h
+++ b/Engine/Source/Runtime/Core/Public/Containers/UnrealString.h
@@ -5,6 +5,7 @@
 
 #pragma once
 
+#include "Containers/Array.h"
 #include "CoreTypes.h"
 #include "Misc/VarArgs.h"
 #include "Misc/OutputDevice.h"
