Fixed issue compiling shipping.

diff --git a/Engine/Source/Runtime/Core/Public/Containers/UnrealString.h b/Engine/Source/Runtime/Core/Public/Containers/UnrealString.h
--- a/Engine/Source/Runtime/Core/Public/Containers/UnrealString.h
+++ b/Engine/Source/Runtime/Core/Public/Containers/UnrealString.h
@@ -5,7 +5,6 @@
 
 #pragma once
 
-#include "Containers/Array.h"
 #include "CoreTypes.h"
 #include "Misc/VarArgs.h"
 #include "Misc/OutputDevice.h"
diff --git a/Engine/Source/Runtime/Core/Public/HAL/LowLevelMemTracker.h b/Engine/Source/Runtime/Core/Public/HAL/LowLevelMemTracker.h
--- a/Engine/Source/Runtime/Core/Public/HAL/LowLevelMemTracker.h
+++ b/Engine/Source/Runtime/Core/Public/HAL/LowLevelMemTracker.h
@@ -2,14 +2,8 @@
 
 #pragma once
 
-#include "Containers/UnrealString.h"
 #include "CoreTypes.h"
-#include "HAL/PlatformCrt.h"
-#include "HAL/PlatformMisc.h"
 #include "ProfilingDebugging/TagTrace.h"
-#include "UObject/UnrealNames.h"
-
-#include <atomic>
 
 #ifndef ALLOW_LOW_LEVEL_MEM_TRACKER_IN_TEST
 	#define ALLOW_LOW_LEVEL_MEM_TRACKER_IN_TEST 0
@@ -48,9 +42,14 @@
 #include "Containers/Array.h"
 #include "Containers/ArrayView.h"
 #include "HAL/CriticalSection.h"
+#include "HAL/PlatformCrt.h"
+#include "HAL/PlatformMisc.h"
 #include "Templates/AlignmentTemplates.h"
 #include "Templates/UnrealTemplate.h"
 #include "UObject/NameTypes.h"
+#include "UObject/UnrealNames.h"
+
+#include <atomic>
 
 #if DO_CHECK
 
