Disable PLATFORM_ALLOW_ALLOCATIONS_IN_FASYNCWRITER_SERIALIZEBUFFERTOARCHIVE on additional platforms which use FPlatformMallocCrash during crash handling

diff --git a/Engine/Platforms/Hololens/Source/Runtime/Core/Public/HoloLensPlatform.h b/Engine/Platforms/Hololens/Source/Runtime/Core/Public/HoloLensPlatform.h
--- a/Engine/Platforms/Hololens/Source/Runtime/Core/Public/HoloLensPlatform.h
+++ b/Engine/Platforms/Hololens/Source/Runtime/Core/Public/HoloLensPlatform.h
@@ -74,6 +74,7 @@ typedef FHoloLensTypes FPlatformTypes;
 //#define EXCEPTION_EXECUTE_HANDLER                           1
 //#define PLATFORM_COMPILER_HAS_DEFAULT_FUNCTION_TEMPLATE_ARGUMENTS	1
 #define PLATFORM_USE_GENERIC_STRING_IMPLEMENTATION			0  // GenericWidePlatformString will fail on HoloLens in UE5
+#define PLATFORM_ALLOW_ALLOCATIONS_IN_FASYNCWRITER_SERIALIZEBUFFERTOARCHIVE	0
 
 #define PLATFORM_BREAK() (__nop(), __debugbreak())
 
diff --git a/Engine/Source/Runtime/Core/Private/Misc/OutputDeviceFile.cpp b/Engine/Source/Runtime/Core/Private/Misc/OutputDeviceFile.cpp
--- a/Engine/Source/Runtime/Core/Private/Misc/OutputDeviceFile.cpp
+++ b/Engine/Source/Runtime/Core/Private/Misc/OutputDeviceFile.cpp
@@ -83,7 +83,7 @@ void FAsyncWriter::FlushArchiveAndResetTimer()
 void FAsyncWriter::SerializeBufferToArchive()
 {
 	QUICK_SCOPE_CYCLE_COUNTER(STAT_FAsyncWriter_SerializeBufferToArchive);
-	// Unix and PS4 use FPlatformMallocCrash during a crash, which means this function is not allowed to perform any allocations
+	// Some platforms use FPlatformMallocCrash during a crash, which means this function is not allowed to perform any allocations
 	// or else it will deadlock when flushing logs during crash handling. Ideally scoped named events would be disabled while crashing.
 	// GIsCriticalError is not always true when crashing (i.e. the case of a GPF) so there is no way to know to skip this behavior only when crashing
 #if PLATFORM_ALLOW_ALLOCATIONS_IN_FASYNCWRITER_SERIALIZEBUFFERTOARCHIVE
diff --git a/Engine/Source/Runtime/Core/Public/Android/AndroidPlatform.h b/Engine/Source/Runtime/Core/Public/Android/AndroidPlatform.h
--- a/Engine/Source/Runtime/Core/Public/Android/AndroidPlatform.h
+++ b/Engine/Source/Runtime/Core/Public/Android/AndroidPlatform.h
@@ -57,12 +57,13 @@ typedef FAndroidTypes FPlatformTypes;
 #define PLATFORM_UI_HAS_MOBILE_SCROLLBARS			1
 #define PLATFORM_UI_NEEDS_TOOLTIPS					0
 #define PLATFORM_UI_NEEDS_FOCUS_OUTLINES			0
-#define PLATFORM_SUPPORTS_EARLY_MOVIE_PLAYBACK		1 // movies will start before engine is initalized
+#define PLATFORM_SUPPORTS_EARLY_MOVIE_PLAYBACK		1 // movies will start before engine is initialized
 #define PLATFORM_SUPPORTS_GEOMETRY_SHADERS			0
 #define PLATFORM_SUPPORTS_VIRTUAL_TEXTURE_STREAMING	1
 #define PLATFORM_USE_ANSI_POSIX_MALLOC				1
 #define PLATFORM_USE_MINIMAL_HANG_DETECTION			1
 #define PLATFORM_CODE_SECTION(Name)					__attribute__((section(Name)))
+#define PLATFORM_ALLOW_ALLOCATIONS_IN_FASYNCWRITER_SERIALIZEBUFFERTOARCHIVE	0
 
 #define PLATFORM_GLOBAL_LOG_CATEGORY				LogAndroid
 #define PLATFORM_RHITHREAD_DEFAULT_BYPASS			0
