Add support for tpause instruction in YieldCycle

diff --git a/Engine/Source/Programs/UnrealBuildTool/Platform/Windows/VCToolChain.cs b/Engine/Source/Programs/UnrealBuildTool/Platform/Windows/VCToolChain.cs
--- a/Engine/Source/Programs/UnrealBuildTool/Platform/Windows/VCToolChain.cs
+++ b/Engine/Source/Programs/UnrealBuildTool/Platform/Windows/VCToolChain.cs
@@ -365,6 +365,9 @@ namespace UnrealBuildTool
 
 				// This matches Microsoft's default support floor for SSE.
 				Arguments.Add("-mssse3");
+
+				// Use tpause on supported processors.
+				Arguments.Add("-mwaitpkg");
 			}
 
 			// Compile into an .obj file, and skip linking.
diff --git a/Engine/Source/Runtime/Core/Private/Windows/WindowsPlatformMisc.cpp b/Engine/Source/Runtime/Core/Private/Windows/WindowsPlatformMisc.cpp
--- a/Engine/Source/Runtime/Core/Private/Windows/WindowsPlatformMisc.cpp
+++ b/Engine/Source/Runtime/Core/Private/Windows/WindowsPlatformMisc.cpp
@@ -2145,6 +2145,9 @@ public:
 			CPUInfo = Info[0];
 			CPUInfo2 = Info[2];
 			CacheLineSize = QueryCacheLineSize();
+			int ExtendedFeatures[4];
+			QueryCPUExtendedFeatures(ExtendedFeatures);
+			CPUExtendedFeatures2 = ExtendedFeatures[2];
 		}
 	}
 
@@ -2198,6 +2201,16 @@ public:
 		return CPUIDStaticCache.CPUInfo2;
 	}
 
+	/**
+	* Gets __cpuidex CPU features.
+	*
+	* @returns CPU info unsigned int queried using __cpuidex.
+	*/
+	static uint32 GetCPUExtendedFeatures2()
+	{
+		return CPUIDStaticCache.CPUExtendedFeatures2;
+	}
+
 	/**
 	 * Gets cache line size.
 	 *
@@ -2305,6 +2318,16 @@ private:
 		__cpuid(Args, 1);
 	}
 
+	/**
+	 * Queries CPU info using __cpuid instruction.
+	 *
+	 * @returns CPU info unsigned int queried using __cpuidex.
+	 */
+	static void QueryCPUExtendedFeatures(int Args[4])
+	{
+		__cpuidex(Args, 7, 0);
+	}
+
 	/**
 	 * Queries cache line size using __cpuid instruction.
 	 *
@@ -2338,6 +2361,7 @@ private:
 	/** CPU info from __cpuid. */
 	uint32 CPUInfo;
 	uint32 CPUInfo2;
+	uint32 CPUExtendedFeatures2;
 
 	/** CPU cache line size. */
 	int32 CacheLineSize;
@@ -3041,6 +3065,12 @@ bool FWindowsPlatformMisc::NeedsNonoptionalCPUFeaturesCheck()
 	return PLATFORM_ENABLE_POPCNT_INTRINSIC;
 }
 
+bool FWindowsPlatformMisc::HasTimedPauseCPUFeature()
+{
+	// Check for waitpkg is bit 5
+	return (FCPUIDQueriedData::GetCPUExtendedFeatures2() & (1 << 5)) != 0;
+}
+
 int32 FWindowsPlatformMisc::GetCacheLineSize()
 {
 	return FCPUIDQueriedData::GetCacheLineSize();
diff --git a/Engine/Source/Runtime/Core/Public/GenericPlatform/GenericPlatformProcess.h b/Engine/Source/Runtime/Core/Public/GenericPlatform/GenericPlatformProcess.h
--- a/Engine/Source/Runtime/Core/Public/GenericPlatform/GenericPlatformProcess.h
+++ b/Engine/Source/Runtime/Core/Public/GenericPlatform/GenericPlatformProcess.h
@@ -16,7 +16,7 @@ class FEvent;
 
 ////////////////////////////////////////////////////////////////////////////////
 #if PLATFORM_CPU_X86_FAMILY
-#include <emmintrin.h>
+#include <immintrin.h>
 #endif
 
 ////////////////////////////////////////////////////////////////////////////////
@@ -788,10 +788,28 @@ struct CORE_API FGenericPlatformProcess
 		uint64 start = ReadCycleCounter();
 		//some 32bit implementations return 0 for __builtin_readcyclecounter just to be on the safe side we protect against this.
 		Cycles = start != 0 ? Cycles : 0;
-		do
+
+#if PLATFORM_WINDOWS
+		if (FPlatformMisc::HasTimedPauseCPUFeature())
 		{
-			Yield();
-		} while((ReadCycleCounter() - start) < Cycles);
+			uint64 PauseCycles = ReadCycleCounter() + Cycles;
+#if defined(_MSC_VER) && !defined(__clang__)
+			_tpause(0, PauseCycles);
+#elif __has_builtin(__builtin_ia32_tpause)
+			__builtin_ia32_tpause(0, (uint32)(PauseCycles >> 32), (uint32)PauseCycles);
+#else
+#	error Unsupported architecture!
+#endif
+		}
+		else
+#endif
+		{
+			do
+			{
+				Yield();
+			} while ((ReadCycleCounter() - start) < Cycles);
+		}
+
 #else
 		// We can't read cycle counter from user mode on these platform
 		for (uint64 i = 0; i < Cycles; i++)
diff --git a/Engine/Source/Runtime/Core/Public/Windows/WindowsPlatformMisc.h b/Engine/Source/Runtime/Core/Public/Windows/WindowsPlatformMisc.h
--- a/Engine/Source/Runtime/Core/Public/Windows/WindowsPlatformMisc.h
+++ b/Engine/Source/Runtime/Core/Public/Windows/WindowsPlatformMisc.h
@@ -206,6 +206,8 @@ struct CORE_API FWindowsPlatformMisc
 	static bool HasNonoptionalCPUFeatures();
 	/** @return whether to check for specific CPU compatibility or not */
 	static bool NeedsNonoptionalCPUFeaturesCheck();
+	/** @return whether this cpu has timed pause instruction support or not */
+	static bool HasTimedPauseCPUFeature();
 
 	/** 
 	 * Provides a simpler interface for fetching and cleanup of registry value queries
