Deprecate TVector2 operator <, >, <=, >=

2D vectors are not canonically ordered and having these operators results in very unfortunate interactions with generic code. In particular FMath::{Min,Max,Abs} on TVector2s compile but do not behave as one would expect.

Add explicit comparison methods as a replacement and deprecate uses of TVector2 overloaded relational operators since it's confusing and error-prone. Fix existing uses to avoid deprecation warnings. Several of these fixes are behavioral changes to fix what I am reasonably certain are bugs in the existing code - in particular, uses of FMath::{Min,Max,Abs} on TVector2s that almost certainly actually want TVector2::{Min,Max,GetAbs} instead.

diff --git a/Engine/Plugins/Experimental/LandscapePatch/Source/LandscapePatch/Private/LandscapeTexturePatch.cpp b/Engine/Plugins/Experimental/LandscapePatch/Source/LandscapePatch/Private/LandscapeTexturePatch.cpp
--- a/Engine/Plugins/Experimental/LandscapePatch/Source/LandscapePatch/Private/LandscapeTexturePatch.cpp
+++ b/Engine/Plugins/Experimental/LandscapePatch/Source/LandscapePatch/Private/LandscapeTexturePatch.cpp
@@ -934,7 +934,7 @@ FVector2D ULandscapeTexturePatch::GetFullUnscaledWorldSize() const
 
 	// UnscaledPatchCoverage is meant to represent the distance between the centers of the extremal pixels.
 	// That distance in pixels is Resolution-1.
-	FVector2D TargetPixelSize(UnscaledPatchCoverage / FMath::Max(Resolution - 1, FVector2D(1, 1)));
+	FVector2D TargetPixelSize(UnscaledPatchCoverage / FVector2D::Max(Resolution - 1, FVector2D(1, 1)));
 	return TargetPixelSize * Resolution;
 }
 
diff --git a/Engine/Plugins/Experimental/LandscapePatch/Source/LandscapePatch/Private/LandscapeTexturePatchBase.cpp b/Engine/Plugins/Experimental/LandscapePatch/Source/LandscapePatch/Private/LandscapeTexturePatchBase.cpp
--- a/Engine/Plugins/Experimental/LandscapePatch/Source/LandscapePatch/Private/LandscapeTexturePatchBase.cpp
+++ b/Engine/Plugins/Experimental/LandscapePatch/Source/LandscapePatch/Private/LandscapeTexturePatchBase.cpp
@@ -97,7 +97,7 @@ FVector2D ULandscapeTexturePatchBase::GetFullUnscaledWorldSizeForResolution(cons
 {
 	// UnscaledPatchCoverage is meant to represent the distance between the centers of the extremal pixels.
 	// That distance in pixels is TextureSize-1.
-	FVector2D TargetPixelSize(UnscaledPatchCoverage / FMath::Max(ResolutionIn - 1, FVector2D(1, 1)));
+	FVector2D TargetPixelSize(UnscaledPatchCoverage / FVector2D::Max(ResolutionIn - 1, FVector2D(1, 1)));
 	return TargetPixelSize * ResolutionIn;
 }
 
diff --git a/Engine/Plugins/FX/Niagara/Source/Niagara/Private/NiagaraDebugHud.cpp b/Engine/Plugins/FX/Niagara/Source/Niagara/Private/NiagaraDebugHud.cpp
--- a/Engine/Plugins/FX/Niagara/Source/Niagara/Private/NiagaraDebugHud.cpp
+++ b/Engine/Plugins/FX/Niagara/Source/Niagara/Private/NiagaraDebugHud.cpp
@@ -2802,7 +2802,7 @@ void FNiagaraDebugHud::DrawMessages(class FNiagaraWorldManager* WorldManager, cl
 		Message.Lifetime -= DeltaSeconds;
 		if (Message.Lifetime > 0.0f)
 		{
-			BackgroundSize = FMath::Max(BackgroundSize, GetStringSize(Font, *Message.Message));
+			BackgroundSize = FVector2D::Max(BackgroundSize, GetStringSize(Font, *Message.Message));
 		}
 		else
 		{
diff --git a/Engine/Plugins/Runtime/nDisplay/Source/DisplayClusterConfigurator/Private/Views/OutputMapping/Alignment/DisplayClusterConfiguratorHostNodeArrangementHelper.cpp b/Engine/Plugins/Runtime/nDisplay/Source/DisplayClusterConfigurator/Private/Views/OutputMapping/Alignment/DisplayClusterConfiguratorHostNodeArrangementHelper.cpp
--- a/Engine/Plugins/Runtime/nDisplay/Source/DisplayClusterConfigurator/Private/Views/OutputMapping/Alignment/DisplayClusterConfiguratorHostNodeArrangementHelper.cpp
+++ b/Engine/Plugins/Runtime/nDisplay/Source/DisplayClusterConfigurator/Private/Views/OutputMapping/Alignment/DisplayClusterConfiguratorHostNodeArrangementHelper.cpp
@@ -266,7 +266,7 @@ void FDisplayClusterConfiguratorHostNodeArrangementHelper::PlaceNodes(const TArr
 				}
 			}
 
-			FVector2D Size = HostNode->TransformSizeToGlobal(FMath::Max(ChildBounds.Max + HostNode->GetHostOrigin(), FVector2D::ZeroVector));
+			FVector2D Size = HostNode->TransformSizeToGlobal(FVector2D::Max(ChildBounds.Max + HostNode->GetHostOrigin(), FVector2D::ZeroVector));
 
 			HostNode->NodeWidth = Size.X;
 			HostNode->NodeHeight = Size.Y;
diff --git a/Engine/Plugins/Runtime/nDisplay/Source/DisplayClusterConfigurator/Private/Views/OutputMapping/EdNodes/DisplayClusterConfiguratorBaseNode.cpp b/Engine/Plugins/Runtime/nDisplay/Source/DisplayClusterConfigurator/Private/Views/OutputMapping/EdNodes/DisplayClusterConfiguratorBaseNode.cpp
--- a/Engine/Plugins/Runtime/nDisplay/Source/DisplayClusterConfigurator/Private/Views/OutputMapping/EdNodes/DisplayClusterConfiguratorBaseNode.cpp
+++ b/Engine/Plugins/Runtime/nDisplay/Source/DisplayClusterConfigurator/Private/Views/OutputMapping/EdNodes/DisplayClusterConfiguratorBaseNode.cpp
@@ -547,7 +547,7 @@ FVector2D UDisplayClusterConfiguratorBaseNode::FindNonOverlappingSize(UDisplayCl
 	const float AspectRatio = OriginalSlotSize.X / OriginalSlotSize.Y;
 	FVector2D SizeChange = BestSize - OriginalSlotSize;
 
-	if (SizeChange < FVector2D::ZeroVector)
+	if (SizeChange.GetMax() < 0)
 	{
 		return BestSize;
 	}
@@ -713,7 +713,7 @@ FVector2D UDisplayClusterConfiguratorBaseNode::FindNonOverlappingSizeFromParent(
 	const FVector2D NodeSize = GetNodeSize();
 
 	// If desired size is smaller in both dimensions to the slot's current size, can return it immediately, as shrinking a slot can't cause any new intersections.
-	if (BestSize - NodeSize < FVector2D::ZeroVector)
+	if (BestSize.ComponentwiseAllLessThan(NodeSize))
 	{
 		return BestSize;
 	}
@@ -753,7 +753,7 @@ FVector2D UDisplayClusterConfiguratorBaseNode::FindBoundedSizeFromParent(const F
 	FVector2D SizeChange = BestSize - NodeSize;
 
 	// If desired size is smaller in both dimensions to the slot's current size, can return it immediately, as shrinking a slot can't cause any bound exceeding.
-	if (SizeChange < FVector2D::ZeroVector)
+	if (SizeChange.GetMax() < 0)
 	{
 		return BestSize;
 	}
@@ -842,7 +842,7 @@ FVector2D UDisplayClusterConfiguratorBaseNode::FindBoundedSizeFromChildren(const
 	FVector2D SizeChange = BestSize - NodeSize;
 
 	// If desired size is bigger in both dimensions to the slot's current size, can return it immediately, as growing a slot can't cause any bound exceeding.
-	if (SizeChange > FVector2D::ZeroVector)
+	if (SizeChange.GetMin() > 0)
 	{
 		return BestSize;
 	}
diff --git a/Engine/Plugins/Runtime/nDisplay/Source/DisplayClusterConfigurator/Private/Views/OutputMapping/GraphNodes/SDisplayClusterConfiguratorWindowNode.cpp b/Engine/Plugins/Runtime/nDisplay/Source/DisplayClusterConfigurator/Private/Views/OutputMapping/GraphNodes/SDisplayClusterConfiguratorWindowNode.cpp
--- a/Engine/Plugins/Runtime/nDisplay/Source/DisplayClusterConfigurator/Private/Views/OutputMapping/GraphNodes/SDisplayClusterConfiguratorWindowNode.cpp
+++ b/Engine/Plugins/Runtime/nDisplay/Source/DisplayClusterConfigurator/Private/Views/OutputMapping/GraphNodes/SDisplayClusterConfiguratorWindowNode.cpp
@@ -525,7 +525,7 @@ bool SDisplayClusterConfiguratorWindowNode::CanShowInfoWidget() const
 	TSharedRef<IDisplayClusterConfiguratorViewOutputMapping> OutputMapping = Toolkit->GetViewOutputMapping();
 	FVector2D NodeSize = GetSize();
 
-	return IsNodeVisible() && OutputMapping->GetOutputMappingSettings().bShowWindowInfo && NodeSize > FVector2D::ZeroVector;
+	return IsNodeVisible() && OutputMapping->GetOutputMappingSettings().bShowWindowInfo && NodeSize.GetMin() > 0;
 }
 
 bool SDisplayClusterConfiguratorWindowNode::CanShowCornerImageWidget() const
@@ -536,7 +536,7 @@ bool SDisplayClusterConfiguratorWindowNode::CanShowCornerImageWidget() const
 	TSharedRef<IDisplayClusterConfiguratorViewOutputMapping> OutputMapping = Toolkit->GetViewOutputMapping();
 	FVector2D NodeSize = GetSize();
 
-	return IsNodeVisible() && OutputMapping->GetOutputMappingSettings().bShowWindowCornerImage && NodeSize > FVector2D::ZeroVector;
+	return IsNodeVisible() && OutputMapping->GetOutputMappingSettings().bShowWindowCornerImage && NodeSize.GetMin() > 0;
 }
 
 bool SDisplayClusterConfiguratorWindowNode::IsClusterNodeLocked() const
diff --git a/Engine/Plugins/VirtualProduction/CameraCalibration/Source/CameraCalibrationEditor/Private/AssetEditor/Viewport/SimulcamEditorViewportClient.cpp b/Engine/Plugins/VirtualProduction/CameraCalibration/Source/CameraCalibrationEditor/Private/AssetEditor/Viewport/SimulcamEditorViewportClient.cpp
--- a/Engine/Plugins/VirtualProduction/CameraCalibration/Source/CameraCalibrationEditor/Private/AssetEditor/Viewport/SimulcamEditorViewportClient.cpp
+++ b/Engine/Plugins/VirtualProduction/CameraCalibration/Source/CameraCalibrationEditor/Private/AssetEditor/Viewport/SimulcamEditorViewportClient.cpp
@@ -124,7 +124,8 @@ bool FSimulcamEditorViewportClient::InputKey(const FInputKeyEventArgs& InEventAr
 			const FGeometry& MyGeometry = SimulcamEditorViewportWeakPtr.Pin()->GetTickSpaceGeometry();
 			const FVector2D LocalMouse = FVector2D(InEventArgs.Viewport->GetMouseX(), InEventArgs.Viewport->GetMouseY());
 			// check if we are under the viewport, otherwise the capture system will blindly trigger the PointerEvent
-			if (LocalMouse >= FVector2D(0, 0) && LocalMouse < MyGeometry.GetAbsoluteSize())
+			if (LocalMouse.ComponentwiseAllGreaterOrEqual(FVector2D(0, 0)) &&
+				LocalMouse.ComponentwiseAllLessThan(MyGeometry.GetAbsoluteSize()))
 			{
 				if (UTexture* Texture = SimulcamViewportWeakPtr.Pin()->GetTexture())
 				{
diff --git a/Engine/Plugins/VirtualProduction/DMX/DMXPixelMapping/Source/DMXPixelMappingEditor/Private/DragDrop/DMXPixelMappingGroupChildDragDropHelper.cpp b/Engine/Plugins/VirtualProduction/DMX/DMXPixelMapping/Source/DMXPixelMappingEditor/Private/DragDrop/DMXPixelMappingGroupChildDragDropHelper.cpp
--- a/Engine/Plugins/VirtualProduction/DMX/DMXPixelMapping/Source/DMXPixelMappingEditor/Private/DragDrop/DMXPixelMappingGroupChildDragDropHelper.cpp
+++ b/Engine/Plugins/VirtualProduction/DMX/DMXPixelMapping/Source/DMXPixelMappingEditor/Private/DragDrop/DMXPixelMappingGroupChildDragDropHelper.cpp
@@ -65,11 +65,11 @@ TSharedPtr<FDMXPixelMappingGroupChildDragDropHelper> FDMXPixelMappingGroupChildD
 
 		// Clamp the size of each child to the a valid size
 		FVector2D IndividualChildSize = OutputComponent->GetSize();
-		if (IndividualChildSize > NewHelper->ParentSize)
+		if (IndividualChildSize.ComponentwiseAllGreaterThan(NewHelper->ParentSize))
 		{
 			OutputComponent->SetSize(NewHelper->ParentSize - FVector2D(NewHelper->ParentSize.X / 10.f, NewHelper->ParentSize.X / 10.f)); // Same X and Y so it's evenly aligned
 		}
-		else if (IndividualChildSize < FVector2D(1.f, 1.f))
+		else if (IndividualChildSize.GetMax() < 1.f)
 		{
 			OutputComponent->SetSize(FVector2D(1.f, 1.f));
 		}
diff --git a/Engine/Source/Developer/MaterialUtilities/Private/MaterialUtilities.cpp b/Engine/Source/Developer/MaterialUtilities/Private/MaterialUtilities.cpp
--- a/Engine/Source/Developer/MaterialUtilities/Private/MaterialUtilities.cpp
+++ b/Engine/Source/Developer/MaterialUtilities/Private/MaterialUtilities.cpp
@@ -2695,7 +2695,7 @@ void FMaterialUtilities::GeneratedBinnedTextureSquares(const FVector2D Destinati
 			for (int32 AreaIndex = 0; AreaIndex < UnusedAreas.Num(); ++AreaIndex)
 			{
 				const FUnusedArea& UnusedArea = UnusedAreas[AreaIndex];
-				if (UnusedArea.GetSize() >= TextureSize)
+				if (UnusedArea.GetSize().ComponentwiseAllGreaterOrEqual(TextureSize))
 				{
 					const float Remainder = UnusedArea.GetArea() - TextureSurface;
 					if (Remainder < RemainingArea && Remainder >= 0)
diff --git a/Engine/Source/Editor/GraphEditor/Private/SGraphPanel.cpp b/Engine/Source/Editor/GraphEditor/Private/SGraphPanel.cpp
--- a/Engine/Source/Editor/GraphEditor/Private/SGraphPanel.cpp
+++ b/Engine/Source/Editor/GraphEditor/Private/SGraphPanel.cpp
@@ -1817,7 +1817,8 @@ bool SGraphPanel::IsNodeTitleVisible(const class UEdGraphNode* Node, bool bReque
 
 bool SGraphPanel::IsRectVisible(const FVector2D &TopLeft, const FVector2D &BottomRight)
 {
-	return TopLeft >= PanelCoordToGraphCoord( FVector2D::ZeroVector ) && BottomRight <= PanelCoordToGraphCoord( CachedAllottedGeometryScaledSize );
+	return TopLeft.ComponentwiseAllGreaterOrEqual( PanelCoordToGraphCoord( FVector2D::ZeroVector )) && 
+		BottomRight.ComponentwiseAllLessOrEqual( PanelCoordToGraphCoord( CachedAllottedGeometryScaledSize ) );
 }
 
 bool SGraphPanel::JumpToRect(const FVector2D &TopLeft, const FVector2D &BottomRight)
diff --git a/Engine/Source/Editor/Kismet/Private/Debugging/SKismetDebugTreeView.cpp b/Engine/Source/Editor/Kismet/Private/Debugging/SKismetDebugTreeView.cpp
--- a/Engine/Source/Editor/Kismet/Private/Debugging/SKismetDebugTreeView.cpp
+++ b/Engine/Source/Editor/Kismet/Private/Debugging/SKismetDebugTreeView.cpp
@@ -2926,7 +2926,7 @@ protected:
 
 		if (NameWidget && ValWidget)
 		{
-			return FMath::Max((*NameWidget)->GetDesiredSize(), (*ValWidget)->GetDesiredSize()) * FVector2D(2.0f, 1.0f);
+			return FVector2D::Max((*NameWidget)->GetDesiredSize(), (*ValWidget)->GetDesiredSize()) * FVector2D(2.0f, 1.0f);
 		}
 
 		return STableRow<FDebugTreeItemPtr>::ComputeDesiredSize(LayoutScaleMultiplier);
diff --git a/Engine/Source/Editor/Persona/Private/SAnimNotifyPanel.cpp b/Engine/Source/Editor/Persona/Private/SAnimNotifyPanel.cpp
--- a/Engine/Source/Editor/Persona/Private/SAnimNotifyPanel.cpp
+++ b/Engine/Source/Editor/Persona/Private/SAnimNotifyPanel.cpp
@@ -1622,7 +1622,8 @@ bool SAnimNotifyNode::HitTest(const FGeometry& AllottedGeometry, FVector2D Mouse
 	FVector2D Position = GetWidgetPosition();
 	FVector2D Size = GetSize();
 
-	return MouseLocalPose >= Position && MouseLocalPose <= (Position + Size);
+	return MouseLocalPose.ComponentwiseAllGreaterOrEqual(Position)
+		&& MouseLocalPose.ComponentwiseAllLessOrEqual(Position + Size);
 }
 
 ENotifyStateHandleHit::Type SAnimNotifyNode::DurationHandleHitTest(const FVector2D& CursorTrackPosition) const
@@ -1641,7 +1642,8 @@ ENotifyStateHandleHit::Type SAnimNotifyNode::DurationHandleHitTest(const FVector
 
 		FVector2D MouseRelativePosition(CursorTrackPosition - GetWidgetPosition());
 
-		if(MouseRelativePosition > NotifyNodePosition && MouseRelativePosition < (NotifyNodePosition + NotifyNodeSize))
+		if(MouseRelativePosition.ComponentwiseAllGreaterThan(NotifyNodePosition) &&
+			MouseRelativePosition.ComponentwiseAllLessThan(NotifyNodePosition + NotifyNodeSize))
 		{
 			// Definitely inside the duration box, need to see which handle we hit if any
 			if(MouseRelativePosition.X <= (NotifyNodePosition.X + ScrubHandleSize.X))
diff --git a/Engine/Source/Runtime/Core/Public/Math/Vector2D.h b/Engine/Source/Runtime/Core/Public/Math/Vector2D.h
--- a/Engine/Source/Runtime/Core/Public/Math/Vector2D.h
+++ b/Engine/Source/Runtime/Core/Public/Math/Vector2D.h
@@ -231,36 +231,85 @@ public:
 	bool operator!=(const TVector2<T>& V) const;
 
 	/**
-	* Checks whether both components of this vector are less than another.
+	* Deprecated comparison operator. Use ComponentwiseAllLessThan.
 	*
 	* @param Other The vector to compare against.
 	* @return true if this is the smaller vector, otherwise false.
 	*/
-	bool operator<(const TVector2<T>& Other) const;
+	UE_DEPRECATED(5.1, "TVector2 comparison operators are deprecated. Use ComponentwiseAllLessThan. For componentwise min/max/abs, use TVector2::{Min,Max,GetAbs}, FMath::{Min,Max,Abs} compute something different.")
+	bool operator<(const TVector2<T>& Other) const
+	{
+		return ComponentwiseAllLessThan(Other);
+	}
 
 	/**
-	* Checks whether both components of this vector are greater than another.
+	* Deprecated comparison operator. Use ComponentwiseAllGreaterThan.
 	*
 	* @param Other The vector to compare against.
 	* @return true if this is the larger vector, otherwise false.
 	*/
-	bool operator>(const TVector2<T>& Other) const;
+	UE_DEPRECATED(5.1, "TVector2 comparison operators are deprecated. Use ComponentwiseAllGreaterThan. For componentwise min/max/abs, use TVector2::{Min,Max,GetAbs}, FMath::{Min,Max,Abs} compute something different.")
+	bool operator>(const TVector2<T>& Other) const
+	{
+		return ComponentwiseAllGreaterThan(Other);
+	}
 
 	/**
-	* Checks whether both components of this vector are less than or equal to another.
+	* Deprecated comparison operator. Use ComponentwiseAllLessOrEqual.
 	*
 	* @param Other The vector to compare against.
 	* @return true if this vector is less than or equal to the other vector, otherwise false.
 	*/
-	bool operator<=(const TVector2<T>& Other) const;
+	UE_DEPRECATED(5.1, "TVector2 comparison operators are deprecated. Use ComponentwiseAllLessOrEqual. For componentwise min/max/abs, use TVector2::{Min,Max,GetAbs}, FMath::{Min,Max,Abs} compute something different.")
+	bool operator<=(const TVector2<T>& Other) const
+	{
+		return ComponentwiseAllLessOrEqual(Other);
+	}
 
 	/**
-	* Checks whether both components of this vector are greater than or equal to another.
+	* Deprecated comparison operator. Use ComponentwiseAllGreaterOrEqual.
 	*
 	* @param Other The vector to compare against.
 	* @return true if this vector is greater than or equal to the other vector, otherwise false.
 	*/
-	bool operator>=(const TVector2<T>& Other) const;
+	UE_DEPRECATED(5.1, "TVector2 comparison operators are deprecated. Use ComponentwiseAllGreaterOrEqual. For componentwise min/max/abs, use TVector2::{Min,Max,GetAbs}, FMath::{Min,Max,Abs} compute something different.")
+	bool operator>=(const TVector2<T>& Other) const
+	{
+		return ComponentwiseAllGreaterOrEqual(Other);
+	}
+
+	/**
+	* Checks whether both components of this vector are less than another.
+	*
+	* @param Other The vector to compare against.
+	* @return true if both components of this are less than Other, otherwise false.
+	*/
+	bool ComponentwiseAllLessThan(const TVector2<T>& Other) const;
+
+	/**
+	* Checks whether both components of this vector are greater than another.
+	*
+	* @param Other The vector to compare against.
+	* @return true if both components of this are greater than Other, otherwise false.
+	*/
+	bool ComponentwiseAllGreaterThan(const TVector2<T>& Other) const;
+
+	/**
+	* Checks whether both components of this vector are less than or equal to another.
+	*
+	* @param Other The vector to compare against.
+	* @return true if both components of this are less than or equal to Other, otherwise false.
+	*/
+	bool ComponentwiseAllLessOrEqual(const TVector2<T>& Other) const;
+
+	/**
+	* Checks whether both components of this vector are greater than or equal to another.
+	*
+	* @param Other The vector to compare against.
+	* @return true if both components of this are greater than or equal to Other, otherwise false.
+	*/
+	bool ComponentwiseAllGreaterOrEqual(const TVector2<T>& Other) const;
+
 
 	/**
 	* Gets a negated copy of the vector.
@@ -890,25 +939,25 @@ FORCEINLINE bool TVector2<T>::operator!=(const TVector2<T>& V) const
 }
 
 template<typename T>
-FORCEINLINE bool TVector2<T>::operator<(const TVector2<T>& Other) const
+FORCEINLINE bool TVector2<T>::ComponentwiseAllLessThan(const TVector2<T>& Other) const
 {
 	return X < Other.X && Y < Other.Y;
 }
 
 template<typename T>
-FORCEINLINE bool TVector2<T>::operator>(const TVector2<T>& Other) const
+FORCEINLINE bool TVector2<T>::ComponentwiseAllGreaterThan(const TVector2<T>& Other) const
 {
 	return X > Other.X && Y > Other.Y;
 }
 
 template<typename T>
-FORCEINLINE bool TVector2<T>::operator<=(const TVector2<T>& Other) const
+FORCEINLINE bool TVector2<T>::ComponentwiseAllLessOrEqual(const TVector2<T>& Other) const
 {
 	return X <= Other.X && Y <= Other.Y;
 }
 
 template<typename T>
-FORCEINLINE bool TVector2<T>::operator>=(const TVector2<T>& Other) const
+FORCEINLINE bool TVector2<T>::ComponentwiseAllGreaterOrEqual(const TVector2<T>& Other) const
 {
 	return X >= Other.X && Y >= Other.Y;
 }
diff --git a/Engine/Source/Runtime/Engine/Private/Animation/BlendSpaceHelpers.cpp b/Engine/Source/Runtime/Engine/Private/Animation/BlendSpaceHelpers.cpp
--- a/Engine/Source/Runtime/Engine/Private/Animation/BlendSpaceHelpers.cpp
+++ b/Engine/Source/Runtime/Engine/Private/Animation/BlendSpaceHelpers.cpp
@@ -653,7 +653,7 @@ static FVector2D ClosestPointOnSegment2D(const FVector2D& Point, const FVector2D
 void FBlendSpaceGrid::GenerateGridElements(const TArray<FVertex>& SamplePoints, const TArray<FTriangle*>& TriangleList)
 {
 	check (NumGridDivisions.X > 0 && NumGridDivisions.Y > 0 );
-	check (GridMax > GridMin);
+	check (GridMax.ComponentwiseAllGreaterThan(GridMin));
 
 	const int32 TotalNumGridPoints = NumGridPointsForAxis.X * NumGridPointsForAxis.Y;
 
diff --git a/Engine/Source/Runtime/Experimental/Chaos/Public/Chaos/UniformGrid.h b/Engine/Source/Runtime/Experimental/Chaos/Public/Chaos/UniformGrid.h
--- a/Engine/Source/Runtime/Experimental/Chaos/Public/Chaos/UniformGrid.h
+++ b/Engine/Source/Runtime/Experimental/Chaos/Public/Chaos/UniformGrid.h
@@ -75,7 +75,7 @@ class CHAOS_API TUniformGridBase
 			MCells += TVector<T, d>(2 * static_cast<T>(GhostCells));
 		}
 
-		if (MDx >= TVector<T, d>(UE_SMALL_NUMBER))
+		if (MDx.Min() >= UE_SMALL_NUMBER)
 		{
 			const TVector<T, d> MinToDXRatio = MMinCorner / MDx;
 			for (int32 Axis = 0; Axis < d; ++Axis)
diff --git a/Engine/Source/Runtime/Renderer/Private/RayTracing/RayTracingDeferredReflections.cpp b/Engine/Source/Runtime/Renderer/Private/RayTracing/RayTracingDeferredReflections.cpp
--- a/Engine/Source/Runtime/Renderer/Private/RayTracing/RayTracingDeferredReflections.cpp
+++ b/Engine/Source/Runtime/Renderer/Private/RayTracing/RayTracingDeferredReflections.cpp
@@ -468,8 +468,8 @@ void FDeferredShadingSceneRenderer::RenderRayTracingDeferredReflections(
 	if (bSpatialResolve && (ResolutionFraction != 1.0f || CVarRayTracingReflectionsHorizontalResolutionScale.GetValueOnAnyThread() != 1.0))
 	{
 		float ResolutionFractionX = FMath::Clamp(CVarRayTracingReflectionsHorizontalResolutionScale.GetValueOnAnyThread(), 0.25f, 4.0f);
-		FVector2D ResolutionFloat = FMath::Max(FVector2D(4.0f), FVector2D(RayTracingResolution) * FVector2D(ResolutionFractionX, 1.0f) * ResolutionFraction);
-		FVector2D BufferSizeFloat = FMath::Max(FVector2D(4.0f), FVector2D(RayTracingBufferSize) * FVector2D(ResolutionFractionX, 1.0f) * ResolutionFraction);
+		FVector2D ResolutionFloat = FVector2D::Max(FVector2D(4.0f), FVector2D(RayTracingResolution) * FVector2D(ResolutionFractionX, 1.0f) * ResolutionFraction);
+		FVector2D BufferSizeFloat = FVector2D::Max(FVector2D(4.0f), FVector2D(RayTracingBufferSize) * FVector2D(ResolutionFractionX, 1.0f) * ResolutionFraction);
 
 		RayTracingResolution.X = (int32)FMath::CeilToFloat(ResolutionFloat.X);
 		RayTracingResolution.Y = (int32)FMath::CeilToFloat(ResolutionFloat.Y);
diff --git a/Engine/Source/Runtime/Slate/Private/Widgets/Layout/SBackgroundBlur.cpp b/Engine/Source/Runtime/Slate/Private/Widgets/Layout/SBackgroundBlur.cpp
--- a/Engine/Source/Runtime/Slate/Private/Widgets/Layout/SBackgroundBlur.cpp
+++ b/Engine/Source/Runtime/Slate/Private/Widgets/Layout/SBackgroundBlur.cpp
@@ -113,7 +113,7 @@ int32 SBackgroundBlur::OnPaint(const FPaintArgs& Args, const FGeometry& Allotted
 	const bool bUsingLowQualityFallbackBrush = IsUsingLowQualityFallbackBrush();
 
 	int32 PostFXLayerId = LayerId;
-	if (bAllowBackgroundBlur && AllottedGeometry.GetLocalSize() > FVector2D::ZeroVector)
+	if (bAllowBackgroundBlur && AllottedGeometry.GetLocalSize().GetMin() > 0)
 	{
 		if (!bUsingLowQualityFallbackBrush)
 		{
diff --git a/Engine/Source/Runtime/Slate/Public/Framework/Application/SWindowTitleBar.h b/Engine/Source/Runtime/Slate/Public/Framework/Application/SWindowTitleBar.h
--- a/Engine/Source/Runtime/Slate/Public/Framework/Application/SWindowTitleBar.h
+++ b/Engine/Source/Runtime/Slate/Public/Framework/Application/SWindowTitleBar.h
@@ -446,7 +446,7 @@ protected:
 
 		if (CenterContentAlignment == HAlign_Center)
 		{
-			LeftSize = FMath::Max(LeftSize, RightSize);
+			LeftSize = FVector2D::Max(LeftSize, RightSize);
 			RightSize = LeftSize;
 		}
 
diff --git a/Engine/Source/Runtime/SlateRHIRenderer/Private/SlateRHIRenderingPolicy.cpp b/Engine/Source/Runtime/SlateRHIRenderer/Private/SlateRHIRenderingPolicy.cpp
--- a/Engine/Source/Runtime/SlateRHIRenderer/Private/SlateRHIRenderingPolicy.cpp
+++ b/Engine/Source/Runtime/SlateRHIRenderer/Private/SlateRHIRenderingPolicy.cpp
@@ -310,8 +310,8 @@ static bool UpdateScissorRect(
 				const FVector2f ViewSize((float) SizeXY.X, (float) SizeXY.Y);
 
 				// Clamp scissor rect to BackBuffer size
-				const FVector2f TopLeft     = FMath::Min(FMath::Max(ScissorRect.TopLeft     + ViewTranslation2D, FVector2f(0.0f, 0.0f)), ViewSize);
-				const FVector2f BottomRight = FMath::Min(FMath::Max(ScissorRect.BottomRight + ViewTranslation2D, FVector2f(0.0f, 0.0f)), ViewSize);
+				const FVector2f TopLeft     = FVector2f::Min(FVector2f::Max(ScissorRect.TopLeft     + ViewTranslation2D, FVector2f(0.0f, 0.0f)), ViewSize);
+				const FVector2f BottomRight = FVector2f::Min(FVector2f::Max(ScissorRect.BottomRight + ViewTranslation2D, FVector2f(0.0f, 0.0f)), ViewSize);
 				
 				RHICmdList.SetScissorRect(true, TopLeft.X, TopLeft.Y, BottomRight.X, BottomRight.Y);
 
