Implementation for FAndroidMisc::GetMaxRefreshRate

diff --git a/Engine/Source/Runtime/Core/Private/Android/AndroidPlatformMisc.cpp b/Engine/Source/Runtime/Core/Private/Android/AndroidPlatformMisc.cpp
--- a/Engine/Source/Runtime/Core/Private/Android/AndroidPlatformMisc.cpp
+++ b/Engine/Source/Runtime/Core/Private/Android/AndroidPlatformMisc.cpp
@@ -3033,6 +3033,17 @@ int32 FAndroidMisc::GetNativeDisplayRefreshRate()
 
 }
 
+int32 FAndroidMisc::GetMaxRefreshRate()
+{
+	int32 MaxRefreshRate = 60;
+	TArray<int32> RefreshRates = GetSupportedNativeDisplayRefreshRates();
+	if (RefreshRates.Num() > 0) {
+		RefreshRates.Sort();
+		MaxRefreshRate = RefreshRates[RefreshRates.Num()-1];
+	}
+	return MaxRefreshRate;
+}
+
 FORCEINLINE bool ValueOutsideThreshold(float Value, float BaseLine, float Threshold)
 {
 	return Value > BaseLine * (1.0f + Threshold)
diff --git a/Engine/Source/Runtime/Core/Public/Android/AndroidPlatformMisc.h b/Engine/Source/Runtime/Core/Public/Android/AndroidPlatformMisc.h
--- a/Engine/Source/Runtime/Core/Public/Android/AndroidPlatformMisc.h
+++ b/Engine/Source/Runtime/Core/Public/Android/AndroidPlatformMisc.h
@@ -289,6 +289,8 @@ public:
 
 	static int32 GetNativeDisplayRefreshRate();
 
+	static int32 GetMaxRefreshRate();
+
 	/**
 	 * Returns whether or not a 16 bit index buffer should be promoted to 32 bit on load, needed for some Android devices
 	 */
