Changing exception code for abnormal shutdowns triggered by abort and terminate from ensure to assert.

diff --git a/Engine/Source/Runtime/Core/Private/Windows/WindowsPlatformCrashContext.cpp b/Engine/Source/Runtime/Core/Private/Windows/WindowsPlatformCrashContext.cpp
--- a/Engine/Source/Runtime/Core/Private/Windows/WindowsPlatformCrashContext.cpp
+++ b/Engine/Source/Runtime/Core/Private/Windows/WindowsPlatformCrashContext.cpp
@@ -980,7 +980,7 @@ static void AbortHandler(int Signal)
 		// Force an exception to invoke the crash reporter
 		FAssertInfo Info(TEXT("Abort signal received"), PLATFORM_RETURN_ADDRESS());
 		ULONG_PTR Arguments[] = { (ULONG_PTR)&Info };
-		::RaiseException(EnsureExceptionCode, 0, UE_ARRAY_COUNT(Arguments), Arguments);
+		::RaiseException(AssertExceptionCode, 0, UE_ARRAY_COUNT(Arguments), Arguments);
 	}
 	__except(ReportCrash(GetExceptionInformation()))
 	{
