[Core]
* TrackedActivity - Fixed so "Status" activity always appear before "I/O" activity

diff --git a/Engine/Source/Runtime/Core/Private/Misc/TrackedActivity.cpp b/Engine/Source/Runtime/Core/Private/Misc/TrackedActivity.cpp
--- a/Engine/Source/Runtime/Core/Private/Misc/TrackedActivity.cpp
+++ b/Engine/Source/Runtime/Core/Private/Misc/TrackedActivity.cpp
@@ -208,6 +208,8 @@ FTrackedActivity& FTrackedActivity::GetEngineActivity()
 
 FTrackedActivity& FTrackedActivity::GetIOActivity()
 {
+	GetEngineActivity(); // Ugly, but just to make sure Status show first
+
 	static FTrackedActivity A(TEXT("I/O"), TEXT("Idle"));
 	return A;
 }
