Add an ini file option which when set to true allows package submissions to source control to go ahead even if the virtualization process failed.

#rnx

### Problem
- When rolling our virtualization to projects we have run into a number of set up problems which has caused users to fail to submit packages from the editor as when the virtualization failed we prevented the overall submit from going through.
- Although we will eventually work through all of these problems and add better fault tolerance for now it is annoying for the content team so we need an easy way for a project to allow package submits to be completed even if the virtualization process part of the submit failed.
- Worth noting that the virtualization process failing should never affect the state of the package files, the only downside is that the packages may not be virtualized.

### Fix
- The virtualization/rehydration of packages now return an enum instead of a bool although at the moment we are only returning Success and Failed, but it will allow us to expand in the future
- Combined the code calling the virtualization process from both the source control window and the source control changelist window into a single utility function.
-- Now the two windows only diverge on how they report failure to the user.

- This new utility function now checks to see if the virtualization manager suggests if failing the virtualization process should block any package submission to source control or not.
- I am not happy with this solution, but I haven't come up with anything that feels right, this solution will at least work.

diff --git a/Engine/Config/BaseEngine.ini b/Engine/Config/BaseEngine.ini
--- a/Engine/Config/BaseEngine.ini
+++ b/Engine/Config/BaseEngine.ini
@@ -1241,6 +1241,7 @@ EnablePushToBackend=False
 EnableCacheAfterPull=True
 MinPayloadLength=0
 BackendGraph=ContentVirtualizationBackendGraph_None
+AllowSubmitIfVirtualizationFailed=false
 FilterMode=OptOut
 FilterEngineContent=True
 FilterEnginePluginContent=True
diff --git a/Engine/Source/Developer/Virtualization/Private/VirtualizationManager.cpp b/Engine/Source/Developer/Virtualization/Private/VirtualizationManager.cpp
--- a/Engine/Source/Developer/Virtualization/Private/VirtualizationManager.cpp
+++ b/Engine/Source/Developer/Virtualization/Private/VirtualizationManager.cpp
@@ -267,6 +267,7 @@ FVirtualizationManager::FVirtualizationManager()
 	, FilteringMode(EPackageFilterMode::OptOut)
 	, bFilterEngineContent(true)
 	, bFilterEnginePluginContent(true)
+	, bAllowSubmitIfVirtualizationFailed(false)
 {
 }
 
@@ -392,6 +393,11 @@ bool FVirtualizationManager::IsDisabledForObject(const UObject* OwnerObject) con
 	return DisabledAssetTypes.Find(ClassName) != nullptr;
 }
 
+bool FVirtualizationManager::AllowSubmitIfVirtualizationFailed() const
+{
+	return bAllowSubmitIfVirtualizationFailed;
+}
+
 bool FVirtualizationManager::PushData(const FIoHash& Id, const FCompressedBuffer& Payload, EStorageType StorageType, const FString& Context)
 {
 	FPushRequest Request(Id, Payload, Context);
@@ -630,23 +636,23 @@ EQueryResult FVirtualizationManager::QueryPayloadStatuses(TArrayView<const FIoHa
 	return EQueryResult::Success;
 }
 
-bool FVirtualizationManager::TryVirtualizePackages(const TArray<FString>& FilesToVirtualize, TArray<FText>& OutDescriptionTags, TArray<FText>& OutErrors)
+EVirtualizationResult FVirtualizationManager::TryVirtualizePackages(const TArray<FString>& FilesToVirtualize, TArray<FText>& OutDescriptionTags, TArray<FText>& OutErrors)
 {
 	OutDescriptionTags.Reset();
 	OutErrors.Reset();
 
 	UE::Virtualization::VirtualizePackages(FilesToVirtualize, OutDescriptionTags, OutErrors);
 
-	return OutErrors.IsEmpty();
+	return OutErrors.IsEmpty() ? EVirtualizationResult::Success : EVirtualizationResult::Failed;
 }
 
-bool FVirtualizationManager::TryRehydratePackages(const TArray<FString>& Packages, TArray<FText>& OutErrors)
+ERehydrationResult FVirtualizationManager::TryRehydratePackages(const TArray<FString>& Packages, TArray<FText>& OutErrors)
 {
 	OutErrors.Reset();
 
 	UE::Virtualization::RehydratePackages(Packages, OutErrors);
 
-	return OutErrors.IsEmpty();
+	return OutErrors.IsEmpty() ? ERehydrationResult::Success : ERehydrationResult::Failed;
 }
 
 void FVirtualizationManager::DumpStats() const
@@ -817,6 +823,8 @@ void FVirtualizationManager::ApplySettingsFromConfigFiles(const FConfigFile& Con
 		UE_LOG(LogVirtualization, Error, TEXT("Failed to load [Core.VirtualizationModule].FilterEnginePluginContent from config file!"));
 	}
 
+
+	// Optional
 	TArray<FString> DisabledAssetTypesFromIni;
 	if (ConfigFile.GetArray(TEXT("Core.ContentVirtualization"), TEXT("DisabledAsset"), DisabledAssetTypesFromIni) > 0 ||
 		ConfigFile.GetArray(TEXT("Core.VirtualizationModule"), TEXT("DisabledAsset"), DisabledAssetTypesFromIni) > 0)
@@ -830,6 +838,18 @@ void FVirtualizationManager::ApplySettingsFromConfigFiles(const FConfigFile& Con
 		}	
 	}
 
+	// Optional
+	bool bAllowSubmitIfVirtualizationFailedFromIni = true;
+	if (ConfigFile.GetBool(TEXT("Core.VirtualizationModule"), TEXT("AllowSubmitIfVirtualizationFailed"), bAllowSubmitIfVirtualizationFailedFromIni))
+	{
+		bAllowSubmitIfVirtualizationFailed = bAllowSubmitIfVirtualizationFailedFromIni;
+		UE_LOG(LogVirtualization, Display, TEXT("\tAllowSubmitIfVirtualizationFailed : %s"), bAllowSubmitIfVirtualizationFailed ? TEXT("true") : TEXT("false"));
+	}
+	else
+	{
+		UE_LOG(LogVirtualization, Error, TEXT("Failed to load [Core.VirtualizationModule].AllowSubmitIfVirtualizationFailed from config file!"));
+	}
+
 	// Check for any legacy settings and print them out (easier to do this in one block rather than one and time)
 	if (const FConfigSection* LegacySection = ConfigFile.Find(TEXT("Core.ContentVirtualization")))
 	{
diff --git a/Engine/Source/Developer/Virtualization/Private/VirtualizationManager.h b/Engine/Source/Developer/Virtualization/Private/VirtualizationManager.h
--- a/Engine/Source/Developer/Virtualization/Private/VirtualizationManager.h
+++ b/Engine/Source/Developer/Virtualization/Private/VirtualizationManager.h
@@ -98,7 +98,7 @@ class IVirtualizationBackendFactory;
 enum class EPackageFilterMode : uint8
 {
 	/** Packages will be virtualized by default and must be opted out by the use of UVirtualizationFilterSettings::ExcludePackagePaths */
-	OptOut,
+	OptOut = 0,
 	/** Packages will not be virtualized by default and must be opted in by the user of UVirtualizationFilterSettings::IncludePackagePaths */
 	OptIn
 };
@@ -130,6 +130,8 @@ private:
 	virtual bool IsPushingEnabled(EStorageType StorageType) const override;
 
 	virtual bool IsDisabledForObject(const UObject* Owner) const override;
+
+	virtual bool AllowSubmitIfVirtualizationFailed() const override;
 	
 	virtual bool PushData(const FIoHash& Id, const FCompressedBuffer& Payload, EStorageType StorageType, const FString& Context) override;
 	virtual bool PushData(TArrayView<FPushRequest> Requests, EStorageType StorageType) override;
@@ -138,8 +140,8 @@ private:
 
 	virtual EQueryResult QueryPayloadStatuses(TArrayView<const FIoHash> Ids, EStorageType StorageType, TArray<EPayloadStatus>& OutStatuses) override;
 
-	virtual bool TryVirtualizePackages(const TArray<FString>& FilesToVirtualize, TArray<FText>& OutDescriptionTags, TArray<FText>& OutErrors) override;
-	virtual bool TryRehydratePackages(const TArray<FString>& Packages, TArray<FText>& OutErrors) override;
+	virtual EVirtualizationResult TryVirtualizePackages(const TArray<FString>& FilesToVirtualize, TArray<FText>& OutDescriptionTags, TArray<FText>& OutErrors) override;
+	virtual ERehydrationResult TryRehydratePackages(const TArray<FString>& Packages, TArray<FText>& OutErrors) override;
 
 	virtual void DumpStats() const override;
 
@@ -232,6 +234,9 @@ private:
 	
 	/** Should payloads in engine plugin content packages before filtered out and never virtualized */
 	bool bFilterEnginePluginContent;
+
+	/** Should file submits be allowed to continue if a call to TryVirtualizePackages fails */
+	bool bAllowSubmitIfVirtualizationFailed;
 	
 private:
 
diff --git a/Engine/Source/Editor/DerivedDataEditor/Private/SVirtualAssetsStatistics.cpp b/Engine/Source/Editor/DerivedDataEditor/Private/SVirtualAssetsStatistics.cpp
--- a/Engine/Source/Editor/DerivedDataEditor/Private/SVirtualAssetsStatistics.cpp
+++ b/Engine/Source/Editor/DerivedDataEditor/Private/SVirtualAssetsStatistics.cpp
@@ -36,10 +36,8 @@ namespace UE::Virtualization::Rehydration
  */
 void RehydratePackages(TArray<FString> SelectedFiles)
 {
-	
-
 	TArray<FText> Errors;
-	if (IVirtualizationSystem::Get().TryRehydratePackages(SelectedFiles, Errors))
+	if (IVirtualizationSystem::Get().TryRehydratePackages(SelectedFiles, Errors) == ERehydrationResult::Success)
 	{
 		// TODO: At some point ::TryRehydratePackages will return more detail info about the process
 		// when it does we should make a better job at logging it.
diff --git a/Engine/Source/Editor/SourceControlWindows/Private/SSourceControlChangelists.cpp b/Engine/Source/Editor/SourceControlWindows/Private/SSourceControlChangelists.cpp
--- a/Engine/Source/Editor/SourceControlWindows/Private/SSourceControlChangelists.cpp
+++ b/Engine/Source/Editor/SourceControlWindows/Private/SSourceControlChangelists.cpp
@@ -36,7 +36,6 @@
 #include "SSourceControlSubmit.h"
 #include "Framework/Application/SlateApplication.h"
 #include "Framework/Notifications/NotificationManager.h"
-#include "Virtualization/VirtualizationSystem.h"
 
 #define LOCTEXT_NAMESPACE "SourceControlChangelist"
 
@@ -934,53 +933,21 @@ static bool GetOnPresubmitResult(FSourceControlChangelistStatePtr Changelist, FC
 		LocalFilepathList.Add(State->GetFilename());
 	}
 
-	TArray<FText> PayloadErrors;
-	TArray<FText> DescriptionTags;
-
-PRAGMA_DISABLE_DEPRECATION_WARNINGS
-	ISourceControlModule::Get().GetOnPreSubmitFinalize().Broadcast(LocalFilepathList, DescriptionTags, PayloadErrors);
-PRAGMA_ENABLE_DEPRECATION_WARNINGS
-
-	UE::Virtualization::IVirtualizationSystem& System = UE::Virtualization::IVirtualizationSystem::Get();
-	if (System.IsEnabled())
+	FText FailureMsg;
+	if (!TryToVirtualizeFilesToSubmit(LocalFilepathList, Description.Description, FailureMsg))
 	{
-		if (System.TryVirtualizePackages(LocalFilepathList, DescriptionTags, PayloadErrors))
-		{
-			if (!DescriptionTags.IsEmpty())
-			{
-				FTextBuilder NewDescription;
-				NewDescription.AppendLine(Description.Description);
-				NewDescription.AppendLine(); // Add a gap between the user input and any automated description tags we need to add
-
-				for (const FText& Line : DescriptionTags)
-				{
-					NewDescription.AppendLine(Line);
-				}
-
-				Description.Description = NewDescription.ToText();
-			}
-		}
-		else
-		{
-			for (const FText& Error : PayloadErrors)
-			{
-				FMessageLog("SourceControl").Error(Error);
-			}
-
-			// Setup the notification for operation feedback
-			FText FailureMsg(LOCTEXT("SCC_Virtualization_Failed", "Virtualized payloads failed to submit."));
-			FNotificationInfo Info(FailureMsg);
+		// Setup the notification for operation feedback
+		FNotificationInfo Info(FailureMsg);
 
-			Info.Text = LOCTEXT("SCC_Checkin_Failed", "Failed to check in files!");
-			Info.ExpireDuration = 8.0f;
-			Info.HyperlinkText = LOCTEXT("SCC_Checkin_ShowLog", "Show Message Log");
-			Info.Hyperlink = FSimpleDelegate::CreateLambda([]() { FMessageLog("SourceControl").Open(EMessageSeverity::Error, true); });
+		Info.Text = LOCTEXT("SCC_Checkin_Failed", "Failed to check in files!");
+		Info.ExpireDuration = 8.0f;
+		Info.HyperlinkText = LOCTEXT("SCC_Checkin_ShowLog", "Show Message Log");
+		Info.Hyperlink = FSimpleDelegate::CreateLambda([]() { FMessageLog("SourceControl").Open(EMessageSeverity::Error, true); });
 
-			TSharedPtr<SNotificationItem> Notification = FSlateNotificationManager::Get().AddNotification(Info);
-			Notification->SetCompletionState(SNotificationItem::CS_Fail);
+		TSharedPtr<SNotificationItem> Notification = FSlateNotificationManager::Get().AddNotification(Info);
+		Notification->SetCompletionState(SNotificationItem::CS_Fail);
 
-			return false;
-		}
+		return false;
 	}
 
 	return true;
diff --git a/Engine/Source/Editor/SourceControlWindows/Private/SSourceControlSubmit.cpp b/Engine/Source/Editor/SourceControlWindows/Private/SSourceControlSubmit.cpp
--- a/Engine/Source/Editor/SourceControlWindows/Private/SSourceControlSubmit.cpp
+++ b/Engine/Source/Editor/SourceControlWindows/Private/SSourceControlSubmit.cpp
@@ -1,6 +1,7 @@
 // Copyright Epic Games, Inc. All Rights Reserved.
 
 #include "SSourceControlSubmit.h"
+
 #include "ISourceControlOperation.h"
 #include "SourceControlOperations.h"
 #include "ISourceControlProvider.h"
@@ -23,12 +24,66 @@
 #include "Styling/AppStyle.h"
 #include "AssetToolsModule.h"
 #include "AssetRegistry/AssetRegistryModule.h"
-
+#include "Virtualization/VirtualizationSystem.h"
+#include "Logging/MessageLog.h"
 
 #if SOURCE_CONTROL_WITH_SLATE
 
 #define LOCTEXT_NAMESPACE "SSourceControlSubmit"
 
+bool TryToVirtualizeFilesToSubmit(const TArray<FString>& FilesToSubmit, FText& Description, FText& OutFailureMsg)
+{
+	// TODO: Once this is removed and not deprecated move the following arrays inside of the System.IsEnabled() scope
+	TArray<FText> PayloadErrors;
+	TArray<FText> DescriptionTags;
+
+PRAGMA_DISABLE_DEPRECATION_WARNINGS
+	ISourceControlModule::Get().GetOnPreSubmitFinalize().Broadcast(FilesToSubmit, DescriptionTags, PayloadErrors);
+PRAGMA_ENABLE_DEPRECATION_WARNINGS
+
+	UE::Virtualization::IVirtualizationSystem& System = UE::Virtualization::IVirtualizationSystem::Get();
+	if (!System.IsEnabled())
+	{
+		return true;
+	}
+
+	UE::Virtualization::EVirtualizationResult Result = System.TryVirtualizePackages(FilesToSubmit, DescriptionTags, PayloadErrors);
+	if (Result == UE::Virtualization::EVirtualizationResult::Success)
+	{
+		FTextBuilder NewDescription;
+		NewDescription.AppendLine(Description);
+
+		for (const FText& Line : DescriptionTags)
+		{
+			NewDescription.AppendLine(Line);
+		}
+
+		Description = NewDescription.ToText();
+
+		return true;
+	}
+	else if (System.AllowSubmitIfVirtualizationFailed())
+	{
+		for (const FText& Error : PayloadErrors)
+		{
+			FMessageLog("SourceControl").Warning(Error);
+		}
+
+		// Even though the virtualization process had problems we should continue submitting
+		return true;
+	}
+	else
+	{
+		for (const FText& Error : PayloadErrors)
+		{
+			FMessageLog("SourceControl").Error(Error);
+		}
+
+		OutFailureMsg = LOCTEXT("SCC_Virtualization_Failed", "Failed to virtualize the files being submitted!");
+
+		return false;
+	}
+}
 
 namespace SSourceControlSubmitWidgetDefs
 {
diff --git a/Engine/Source/Editor/SourceControlWindows/Private/SSourceControlSubmit.h b/Engine/Source/Editor/SourceControlWindows/Private/SSourceControlSubmit.h
--- a/Engine/Source/Editor/SourceControlWindows/Private/SSourceControlSubmit.h
+++ b/Engine/Source/Editor/SourceControlWindows/Private/SSourceControlSubmit.h
@@ -40,6 +40,8 @@ public:
 	FText Description;
 };
 
+bool TryToVirtualizeFilesToSubmit(const TArray<FString>& FilesToSubmit, FText& Description, FText& OutFailureMsg);
+
 class SSourceControlSubmitWidget : public SCompoundWidget
 {
 public:
diff --git a/Engine/Source/Editor/SourceControlWindows/Private/SourceControlWindows.cpp b/Engine/Source/Editor/SourceControlWindows/Private/SourceControlWindows.cpp
--- a/Engine/Source/Editor/SourceControlWindows/Private/SourceControlWindows.cpp
+++ b/Engine/Source/Editor/SourceControlWindows/Private/SourceControlWindows.cpp
@@ -13,7 +13,6 @@
 #include "Logging/TokenizedMessage.h"
 #include "Misc/MessageDialog.h"
 #include "Widgets/Notifications/SNotificationList.h"
-#include "Virtualization/VirtualizationSystem.h"
 
 
 #if SOURCE_CONTROL_WITH_SLATE
@@ -294,58 +293,17 @@ bool FSourceControlWindows::PromptForCheckin(FCheckinResultInfo& OutResultInfo,
 		return bSuccess;
 	}
 
-	//~~~~~~~~~~~~~~~~~~~~~~~~~~~~~~~~~~~~~~~~~~
-	// Allow any finalization pre-submit processes to run
-	
-	// TODO: Once this is removed and not deprecated move the following arrays inside of the System.IsEnabled() scope
-	TArray<FText> PayloadErrors;
-	TArray<FText> DescriptionTags;
-
-PRAGMA_DISABLE_DEPRECATION_WARNINGS
-	ISourceControlModule::Get().GetOnPreSubmitFinalize().Broadcast(CombinedFileList, DescriptionTags, PayloadErrors);
-PRAGMA_ENABLE_DEPRECATION_WARNINGS
-
-	UE::Virtualization::IVirtualizationSystem& System = UE::Virtualization::IVirtualizationSystem::Get();
-	if (System.IsEnabled())
+	FText VirtualizationFailureMsg;
+	if (!TryToVirtualizeFilesToSubmit(CombinedFileList, Description.Description, VirtualizationFailureMsg))
 	{
-		if (System.TryVirtualizePackages(CombinedFileList, DescriptionTags, PayloadErrors))
-		{
-			if (!DescriptionTags.IsEmpty())
-			{
-				FTextBuilder NewDescription;
-				NewDescription.AppendLine(Description.Description);
-
-				for (const FText& Line : DescriptionTags)
-				{
-					NewDescription.AppendLine(Line);
-				}
-
-				Description.Description = NewDescription.ToText();
-			}
-		}
-		else
-		{
-			for (const FText& Error : PayloadErrors)
-			{
-				FMessageLog("SourceControl").Error(Error);
-			}
-
-			// TODO: We rely on the calling code to issue a noticeable error message to the user but this is 
-			// currently handled very inconsistently in the engine and the return values often ignored.
-			// So print the error to the log file as well for now but we should try and unify how the errors
-			// are reported.
-			FText FailureMsg(LOCTEXT("SCC_Virtualization_Failed", "Virtualized payloads failed to submit."));
-			FMessageLog("SourceControl").Notify(FailureMsg);
+		FMessageLog("SourceControl").Notify(VirtualizationFailureMsg);
 
-			OutResultInfo.Result = ECommandResult::Failed;
-			OutResultInfo.Description = FailureMsg;
+		OutResultInfo.Result = ECommandResult::Failed;
+		OutResultInfo.Description = VirtualizationFailureMsg;
 
-			return false;
-		}
+		return false;
 	}
 
-	//~~~~~~~~~~~~~~~~~~~~~~~~~~~~~~~~~~~~~~~~~~
-
 	// Check in files
 	TSharedRef<FCheckIn, ESPMode::ThreadSafe> CheckInOperation = ISourceControlOperation::Create<FCheckIn>();
 	CheckInOperation->SetDescription(Description.Description);
diff --git a/Engine/Source/Runtime/Core/Private/Virtualization/VirtualizationSystem.cpp b/Engine/Source/Runtime/Core/Private/Virtualization/VirtualizationSystem.cpp
--- a/Engine/Source/Runtime/Core/Private/Virtualization/VirtualizationSystem.cpp
+++ b/Engine/Source/Runtime/Core/Private/Virtualization/VirtualizationSystem.cpp
@@ -47,6 +47,11 @@ public:
 		return false;
 	}
 
+	virtual bool AllowSubmitIfVirtualizationFailed() const override
+	{
+		return false;
+	}
+
 	virtual bool PushData(const FIoHash& Id, const FCompressedBuffer& Payload, EStorageType StorageType, const FString& Context) override
 	{
 		return false;
@@ -69,23 +74,23 @@ public:
 		return EQueryResult::Failure_NotImplemented;
 	}
 
-	virtual bool TryVirtualizePackages(const TArray<FString>& FilesToVirtualize, TArray<FText>& OutDescriptionTags, TArray<FText>& OutErrors) override
+	virtual EVirtualizationResult TryVirtualizePackages(const TArray<FString>& FilesToVirtualize, TArray<FText>& OutDescriptionTags, TArray<FText>& OutErrors) override
 	{
 		OutDescriptionTags.Reset();
 		OutErrors.Reset();
 
 		OutErrors.Add(FText::FromString(TEXT("Calling ::TryVirtualizePackages on FNullVirtualizationSystem")));
 
-		return false;
+		return EVirtualizationResult::Failed;
 	}
 
-	virtual bool TryRehydratePackages(const TArray<FString>& Packages, TArray<FText>& OutErrors) override
+	virtual ERehydrationResult TryRehydratePackages(const TArray<FString>& Packages, TArray<FText>& OutErrors) override
 	{
 		OutErrors.Reset();
 
 		OutErrors.Add(FText::FromString(TEXT("Calling ::TryRehydratePackages on FNullVirtualizationSystem")));
 
-		return false;
+		return ERehydrationResult::Failed;
 	}
 
 	virtual void DumpStats() const override
diff --git a/Engine/Source/Runtime/Core/Public/Virtualization/VirtualizationSystem.h b/Engine/Source/Runtime/Core/Public/Virtualization/VirtualizationSystem.h
--- a/Engine/Source/Runtime/Core/Public/Virtualization/VirtualizationSystem.h
+++ b/Engine/Source/Runtime/Core/Public/Virtualization/VirtualizationSystem.h
@@ -81,6 +81,24 @@ enum class EPayloadStatus : int8
 	FoundAll
 };
 
+/** The result of the virtualization process */
+enum class EVirtualizationResult : uint8
+{
+	/** The virtualization process ran with no problems to completion. */
+	Success = 0,
+	/** The process failed before completion and nothing was virtualized */
+	Failed
+};
+
+/** The result of the re-hydration process */
+enum class ERehydrationResult : uint8
+{
+	/** The re-hydration process ran with no problems to completion. */
+	Success = 0,
+	/** The process failed before completion and nothing was re-hydrate */
+	Failed,
+};
+
 /** 
  * This interface can be implemented and passed to a FPushRequest as a way of providing the payload 
  * to the virtualization system for a push operation but deferring the loading of the payload from 
@@ -342,6 +360,19 @@ public:
 	 */
 	virtual bool IsDisabledForObject(const UObject* Owner) const = 0;
 
+	/** 
+	 * Poll to see if the virtualization process failing when submitting a collection of files to source
+	 * control should block that submit or allow it to continue.
+	 * 
+	 * NOTE: This is a bit of an odd ball option to be included in the VirtualizationSystem API but at
+	 * the moment we don't really have a better place for it. This might be removed or replaced in future
+	 * releases.
+	 * 
+	 * @return		True if the calling code should continue with submitting files if the virtualization
+	 *				process failed. If it returns false then the calling code should prevent the submit.
+	 */
+	virtual bool AllowSubmitIfVirtualizationFailed() const = 0;
+
 	/**
 	 * Push a payload to the virtualization backends.
 	 *
@@ -414,9 +445,10 @@ public:
 	 *								can be assumed that the process will return false. Note that the array will be emptied before the process
 	 *								is run and will not contain any pre-existing entries.
 	 * 
-	 * @return True if the process succeeded and false if it did not. If this returns false then OutErrors should contain at least one entry
+	 * @return						A EVirtualizationResult enum with the status of the process. If the status is not EVirtualizationResult::Success
+	 *								then the parameter OutErrors should contain at least one entry.
 	 */
-	virtual bool TryVirtualizePackages(const TArray<FString>& FilesToVirtualize, TArray<FText>& OutDescriptionTags, TArray<FText>& OutErrors) = 0;
+	virtual EVirtualizationResult TryVirtualizePackages(const TArray<FString>& FilesToVirtualize, TArray<FText>& OutDescriptionTags, TArray<FText>& OutErrors) = 0;
 	
 	/**
 	 * Runs the re-hydration process on a set of packages. This involves downloading virtualized payloads and placing them back in the trailer of
@@ -428,9 +460,10 @@ public:
 	 *					can be assumed that the process will return false. Note that the array will be emptied before the process
 	 *					is run and will not contain any pre-existing entries.
 	 * 
-	 * @return True if the process succeeded and false if it did not. If this returns false then OutErrors should contain at least one entry
+	 * @return	A ERehydrationResult enum with the status of the process. If the status indicates any sort of failure then OutErrors should
+	 *			contain at least one entry.
 	 */
-	virtual bool TryRehydratePackages(const TArray<FString>& Packages, TArray<FText>& OutErrors) = 0;
+	virtual ERehydrationResult TryRehydratePackages(const TArray<FString>& Packages, TArray<FText>& OutErrors) = 0;
 
 	/** When called the system should write any performance stats that it has been gathering to the log file */
 	virtual void DumpStats() const = 0;
