Logging: Reduced stack usage of assertions, crashes, ensures, and fatal logs

- Stack required by failed assertions, fatal logs, and low-level fatal errors decreased by 24,832 bytes.
- Stack required by failed ensures decreased by 18,688 bytes.
- Stack required by non-fatal logs decreased by 8,192 bytes.

diff --git a/Engine/Source/Runtime/Core/Private/Logging/LogMacros.cpp b/Engine/Source/Runtime/Core/Private/Logging/LogMacros.cpp
--- a/Engine/Source/Runtime/Core/Private/Logging/LogMacros.cpp
+++ b/Engine/Source/Runtime/Core/Private/Logging/LogMacros.cpp
@@ -10,66 +10,41 @@
 #include "Stats/Stats.h"
 #include "ProfilingDebugging/CsvProfiler.h"
 
-void StaticFailDebug( const TCHAR* Error, const ANSICHAR* File, int32 Line, void* ProgramCounter, const TCHAR* Description, bool bIsEnsure );
-
-/** Statics to prevent FMsg::Logf from allocating too much stack memory. */
-static FCriticalSection* GetMsgLogfStaticBufferGuard()
-{
-	static FCriticalSection CS;
-	return &CS;
-}
-
-/** Increased from 4096 to fix crashes in the renderthread without autoreporter. */
-static TCHAR							MsgLogfStaticBuffer[8192];
+void StaticFailDebugV(const TCHAR* Error, const ANSICHAR* Expression, const ANSICHAR* File, int32 Line, bool bIsEnsure, void* ProgramCounter, const TCHAR* DescriptionFormat, va_list DescriptionArgs);
 
 CSV_DEFINE_CATEGORY(FMsgLogf, true);
 
 void FMsg::LogfImpl(const ANSICHAR* File, int32 Line, const FLogCategoryName& Category, ELogVerbosity::Type Verbosity, const TCHAR* Fmt, ...)
 {
 #if !NO_LOGGING
-	if (Verbosity != ELogVerbosity::Fatal)
+	if (LIKELY(Verbosity != ELogVerbosity::Fatal))
 	{
-		// SetColour is routed to GWarn just like the other verbosities and handled in the 
+		// SetColour is routed to GWarn just like the other verbosities and handled in the
 		// device that does the actual printing.
-		FOutputDevice* LogDevice = NULL;
+		FOutputDevice* LogOverride = nullptr;
 		switch (Verbosity)
 		{
 		case ELogVerbosity::Error:
 		case ELogVerbosity::Warning:
 		case ELogVerbosity::Display:
 		case ELogVerbosity::SetColor:
-			if (GWarn)
-			{
-				LogDevice = GWarn;
-				break;
-			}
+			LogOverride = GWarn;
+			break;
 		default:
-		{
-			LogDevice = GLog;
+			break;
 		}
-		break;
-		}
-		GROWABLE_LOGF(LogDevice->Log(Category, Verbosity, Buffer))
+		GROWABLE_LOGF(LogOverride	? LogOverride->Log(Category, Verbosity, Buffer)
+									: GLog->RedirectLog(Category, Verbosity, Buffer))
 	}
 	else
 	{
-		// Keep Message buffer small, in some cases, this code is executed with 16KB stack.
-		TCHAR Message[4096];
-		{
-			// Simulate Sprintf_s
-			// @todo: implement platform independent sprintf_S
-			// We're using one big shared static buffer here, so guard against re-entry
-			FScopeLock MsgLock(GetMsgLogfStaticBufferGuard());
-			// Print to a large static buffer so we can keep the stack allocation below 16K
-			GET_VARARGS(MsgLogfStaticBuffer, UE_ARRAY_COUNT(MsgLogfStaticBuffer), UE_ARRAY_COUNT(MsgLogfStaticBuffer) - 1, Fmt, Fmt);
-			// Copy the message to the stack-allocated buffer)
-			FCString::Strncpy(Message, MsgLogfStaticBuffer, UE_ARRAY_COUNT(Message) - 1);
-			Message[UE_ARRAY_COUNT(Message) - 1] = TEXT('\0');
-		}
-
-		const int32 NumStackFramesToIgnore = 1;
-		StaticFailDebug(TEXT("Fatal error:"), File, Line, PLATFORM_RETURN_ADDRESS(), Message, false);
-		FDebug::AssertFailed("", File, Line, Message);
+		va_list Args;
+		va_start(Args, Fmt);
+		StaticFailDebugV(TEXT("Fatal error:"), "", File, Line, /*bIsEnsure*/ false, PLATFORM_RETURN_ADDRESS(), Fmt, Args);
+		va_end(Args);
+		va_start(Args, Fmt);
+		FDebug::AssertFailedV("", File, Line, Fmt, Args);
+		va_end(Args);
 	}
 #endif
 }
@@ -80,11 +55,11 @@ void FMsg::Logf_InternalImpl(const ANSICHAR* File, int32 Line, const FLogCategor
 	QUICK_SCOPE_CYCLE_COUNTER(STAT_FMsgLogf);
 	CSV_CUSTOM_STAT(FMsgLogf, FMsgLogfCount, 1, ECsvCustomStatOp::Accumulate);
 
-	if (Verbosity != ELogVerbosity::Fatal)
+	if (LIKELY(Verbosity != ELogVerbosity::Fatal))
 	{
-		// SetColour is routed to GWarn just like the other verbosities and handled in the 
+		// SetColour is routed to GWarn just like the other verbosities and handled in the
 		// device that does the actual printing.
-		FOutputDevice* LogOverride = NULL;
+		FOutputDevice* LogOverride = nullptr;
 		switch (Verbosity)
 		{
 		case ELogVerbosity::Error:
@@ -92,30 +67,19 @@ void FMsg::Logf_InternalImpl(const ANSICHAR* File, int32 Line, const FLogCategor
 		case ELogVerbosity::Display:
 		case ELogVerbosity::SetColor:
 			LogOverride = GWarn;
+			break;
 		default:
-		break;
+			break;
 		}
 		GROWABLE_LOGF(LogOverride	? LogOverride->Log(Category, Verbosity, Buffer)
 									: GLog->RedirectLog(Category, Verbosity, Buffer))
 	}
 	else
 	{
-		// Keep Message buffer small, in some cases, this code is executed with 16KB stack.
-		TCHAR Message[4096];
-		{
-			// Simulate Sprintf_s
-			// @todo: implement platform independent sprintf_S
-			// We're using one big shared static buffer here, so guard against re-entry
-			FScopeLock MsgLock(GetMsgLogfStaticBufferGuard());
-			// Print to a large static buffer so we can keep the stack allocation below 16K
-			GET_VARARGS(MsgLogfStaticBuffer, UE_ARRAY_COUNT(MsgLogfStaticBuffer), UE_ARRAY_COUNT(MsgLogfStaticBuffer) - 1, Fmt, Fmt);
-			// Copy the message to the stack-allocated buffer)
-			FCString::Strncpy(Message, MsgLogfStaticBuffer, UE_ARRAY_COUNT(Message) - 1);
-			Message[UE_ARRAY_COUNT(Message) - 1] = TEXT('\0');
-		}
-
-		const int32 NumStackFramesToIgnore = 1;
-		StaticFailDebug(TEXT("Fatal error:"), File, Line, PLATFORM_RETURN_ADDRESS(), Message, false);
+		va_list Args;
+		va_start(Args, Fmt);
+		StaticFailDebugV(TEXT("Fatal error:"), "", File, Line, /*bIsEnsure*/ false, PLATFORM_RETURN_ADDRESS(), Fmt, Args);
+		va_end(Args);
 	}
 #endif
 }
diff --git a/Engine/Source/Runtime/Core/Private/Misc/AssertionMacros.cpp b/Engine/Source/Runtime/Core/Private/Misc/AssertionMacros.cpp
--- a/Engine/Source/Runtime/Core/Private/Misc/AssertionMacros.cpp
+++ b/Engine/Source/Runtime/Core/Private/Misc/AssertionMacros.cpp
@@ -37,7 +37,7 @@ namespace
 	{
 		static FCriticalSection FailDebugCriticalSection;
 		return FailDebugCriticalSection;
-	}	
+	}
 
 	struct FTempCommandLineScope
 	{
@@ -67,7 +67,7 @@ namespace
 	};
 }
 
-#define FILE_LINE_DESC_ANSI " [File:%s] [Line: %i] "
+#define FILE_LINE_DESC_ANSI " [File:%hs] [Line: %i] "
 
 /*
 	Ensure behavior
@@ -75,7 +75,7 @@ namespace
 	* ensure() macro calls OptionallyLogFormattedEnsureMessageReturningFalse 
 	* OptionallyLogFormattedEnsureMessageReturningFalse calls EnsureFailed()
 	* EnsureFailed() -
-		* Formats the ensure failure and calls StaticFailDebug to populate the global error info (without callstack)
+		* Formats the ensure failure and calls StaticFailDebugV to populate the global error info (without callstack)
 		* Prints the script callstack (if any)
 		* Halts if a debugger is attached 
 		* If not, logs the callstack and attempts to submit an error report
@@ -84,8 +84,8 @@ namespace
 	Check behavior
 
 	* check() macro calls LogAssertFailedMessage
-	* LogAssertFailedMessage formats the assertion message and calls StaticFailDebug
-	* StaticFailDebug populates global error info with the failure message and if supported (AllowsCallStackDumpDuringAssert) the callstack
+	* LogAssertFailedMessage formats the assertion message and calls StaticFailDebugV
+	* StaticFailDebugV populates global error info with the failure message and if supported (AllowsCallStackDumpDuringAssert) the callstack
 	* If a debugger is attached execution halts
 	* If not FDebug::AssertFailed is called
 	* FDebug::AssertFailed logs the assert message and description to GError
@@ -96,8 +96,8 @@ namespace
 	Fatal-error behavior
 
 	* The UE_LOG macro calls FMsg::Logf which checks for "Fatal" verbosity
-	* FMsg::Logf formats the failure message and calls StaticFailDebug
-	* StaticFailDebug populates global error info with the failure message and if supported (AllowsCallStackDumpDuringAssert) the callstack
+	* FMsg::Logf formats the failure message and calls StaticFailDebugV
+	* StaticFailDebugV populates global error info with the failure message and if supported (AllowsCallStackDumpDuringAssert) the callstack
 	* FDebug::AssertFailed is then called, and from this point behavior is identical to an assert but with a different message
 
 */
@@ -136,67 +136,122 @@ static void AssertFailedImplV(const ANSICHAR* Expr, const ANSICHAR* File, int32
 
 	if (GError)
 	{
+		TCHAR DescriptionString[4096];
+		FCString::GetVarArgs(DescriptionString, UE_ARRAY_COUNT(DescriptionString), Format, Args);
 		GError->SetErrorProgramCounter(ProgramCounter);
+		GError->Logf(TEXT("Assertion failed: %hs" FILE_LINE_DESC_ANSI "\n%s\n"), Expr, File, Line, DescriptionString);
 	}
+}
 
-	TCHAR DescriptionString[4096];
-	FCString::GetVarArgs(DescriptionString, UE_ARRAY_COUNT(DescriptionString), Format, Args);
+class FErrorHistWriter
+{
+public:
+	UE_NONCOPYABLE(FErrorHistWriter);
 
-	TCHAR ErrorString[MAX_SPRINTF];
-	FCString::Sprintf(ErrorString, TEXT("%s"), ANSI_TO_TCHAR(Expr));
-	if (GError)
+	inline FErrorHistWriter()
 	{
-		GError->Logf(TEXT("Assertion failed: %s" FILE_LINE_DESC_ANSI "\n%s\n"), ErrorString, ANSI_TO_TCHAR(File), Line, DescriptionString);
+		GetFailDebugCriticalSection().Lock();
+	}
+
+	inline ~FErrorHistWriter()
+	{
+		Terminate();
+		GetFailDebugCriticalSection().Unlock();
+	}
+
+	inline void Terminate()
+	{
+		GErrorHist[Index] = TCHAR(0);
+	}
+
+	template <typename CharType>
+	void Append(const CharType* Text)
+	{
+		const int32 TextLen = TCString<CharType>::Strlen(Text);
+		const int32 RequiredLen = FPlatformString::ConvertedLength<TCHAR>(Text, TextLen);
+		if (Index + RequiredLen < Capacity)
+		{
+			FPlatformString::Convert(GErrorHist + Index, RequiredLen, Text, TextLen);
+			Index += RequiredLen;
+		}
+	}
+
+	void AppendV(const TCHAR* Format, va_list Args)
+	{
+		const int32 Len = FPlatformString::GetVarArgs(GErrorHist + Index, Capacity - Index, Format, Args);
+		Index = FMath::Clamp(Index, Index + Len, Capacity - 1);
+	}
+
+	void Appendf(const TCHAR* Format, ...)
+	{
+		va_list Args;
+		va_start(Args, Format);
+		AppendV(Format, Args);
+		va_end(Args);
 	}
-}
+
+private:
+	int32 Index = 0;
+
+	static constexpr int32 Capacity = UE_ARRAY_COUNT(GErrorHist);
+};
 
 /**
- *	Prints error to the debug output, 
- *	prompts for the remote debugging if there is not debugger, breaks into the debugger 
- *	and copies the error into the global error message.
+ * Prints error to the debug output and copies the error into the global error message.
  */
-FORCENOINLINE void StaticFailDebug( const TCHAR* Error, const ANSICHAR* File, int32 Line, void* ProgramCounter, const TCHAR* Description, bool bIsEnsure )
+FORCENOINLINE void StaticFailDebugV(
+	const TCHAR* Error,
+	const ANSICHAR* Expression,
+	const ANSICHAR* File,
+	int32 Line,
+	bool bIsEnsure,
+	void* ProgramCounter,
+	const TCHAR* DescriptionFormat,
+	va_list DescriptionArgs)
 {
-	// Print out the blueprint callstack
 	PrintScriptCallstack();
 
-	TCHAR DescriptionAndTrace[4096];
-
-	FCString::Strncpy(DescriptionAndTrace, Description, UE_ARRAY_COUNT(DescriptionAndTrace) - 1);
-
 	// some platforms (Windows, Mac, Linux) generate this themselves by throwing an exception and capturing
 	// the backtrace later on
-	if (FPlatformProperties::AllowsCallStackDumpDuringAssert() && bIsEnsure == false)
+	ANSICHAR StackTrace[4096] = "";
+	if (FPlatformProperties::AllowsCallStackDumpDuringAssert() && !bIsEnsure)
 	{
-		ANSICHAR StackTrace[4096];
-		StackTrace[0] = 0;
 		FPlatformStackWalk::StackWalkAndDump(StackTrace, UE_ARRAY_COUNT(StackTrace), ProgramCounter);
-
-		FCString::Strncat(DescriptionAndTrace, TEXT("\n"), UE_ARRAY_COUNT(DescriptionAndTrace) - 1);
-		FCString::Strncat(DescriptionAndTrace, ANSI_TO_TCHAR(StackTrace), UE_ARRAY_COUNT(DescriptionAndTrace) - 1);
 	}
 
-	FScopeLock Lock( &GetFailDebugCriticalSection());
-	FPlatformMisc::LowLevelOutputDebugStringf(TEXT("%s" FILE_LINE_DESC_ANSI "\n%s\n"), Error, ANSI_TO_TCHAR(File), Line, Description);
+	FErrorHistWriter Writer;
+	Writer.Append(Error);
+	Writer.Append(Expression);
+	Writer.Appendf(TEXT(FILE_LINE_DESC_ANSI "\n"), File, Line);
+	Writer.AppendV(DescriptionFormat, DescriptionArgs);
+	Writer.Append(TEXT("\n"));
+	Writer.Terminate();
+
+	FPlatformMisc::LowLevelOutputDebugString(GErrorHist);
 
-	// Copy the detailed error into the error message.
-	TCHAR ErrorMessage[4096];
-	if (FCString::Snprintf( ErrorMessage, UE_ARRAY_COUNT( ErrorMessage ), TEXT( "%s" FILE_LINE_DESC_ANSI "\n%s\n" ), Error, ANSI_TO_TCHAR( File ), Line, DescriptionAndTrace) < 0)
+	if (*StackTrace)
 	{
-		// Description and callstack was too long to fit in GErrorMessage. Use only description
-		FCString::Snprintf( ErrorMessage, UE_ARRAY_COUNT( ErrorMessage ), TEXT( "%s" FILE_LINE_DESC_ANSI "\n%s\n<< callstack too long >>" ), Error, ANSI_TO_TCHAR( File ), Line, Description);
+		Writer.Append(StackTrace);
+		Writer.Append(TEXT("\n"));
 	}
 
-	// Copy the error message to the error history.
-	FCString::Strncpy( GErrorHist, ErrorMessage, UE_ARRAY_COUNT( GErrorHist ) );
-	FCString::Strncat( GErrorHist, TEXT( "\r\n\r\n" ), UE_ARRAY_COUNT( GErrorHist ) );
+	Writer.Append(TEXT("\r\n\r\n"));
+	Writer.Terminate();
 
-	if (GError != nullptr)
+	if (GError)
 	{
 		GError->SetErrorProgramCounter(ProgramCounter);
 	}
 }
 
+FORCENOINLINE void VARARGS StaticFailDebug(const TCHAR* Error, const ANSICHAR* Expression, const ANSICHAR* File, int32 Line,
+	bool bIsEnsure, void* ProgramCounter, const TCHAR* DescriptionFormat, ...)
+{
+	va_list DescriptionArgs;
+	va_start(DescriptionArgs, DescriptionFormat);
+	StaticFailDebugV(Error, Expression, File, Line, bIsEnsure, ProgramCounter, DescriptionFormat, DescriptionArgs);
+	va_end(DescriptionArgs);
+}
 
 /// track thread asserts
 bool FDebug::HasAsserted()
@@ -302,13 +357,7 @@ void FDebug::LogAssertFailedMessageImplV(const ANSICHAR* Expr, const ANSICHAR* F
 	// Ignore this assert if we're already forcibly shutting down because of a critical error.
 	if( !GIsCriticalError )
 	{
-		TCHAR DescriptionString[4096];
-		FCString::GetVarArgs( DescriptionString, UE_ARRAY_COUNT(DescriptionString), Fmt, Args );
-
-		TCHAR ErrorString[MAX_SPRINTF];
-		FCString::Sprintf( ErrorString, TEXT( "Assertion failed: %s" ), ANSI_TO_TCHAR( Expr ) );
-
-		StaticFailDebug( ErrorString, File, Line, ProgramCounter, DescriptionString, false );
+		StaticFailDebugV(TEXT("Assertion failed: "), Expr, File, Line, /*bIsEnsure*/ false, ProgramCounter, Fmt, Args);
 	}
 }
 
@@ -350,11 +399,7 @@ FORCENOINLINE void FDebug::EnsureFailed(const ANSICHAR* Expr, const ANSICHAR* Fi
 	// Should we spin here?
 	FPlatformAtomics::InterlockedIncrement(&ActiveEnsureCount);
 
-	// Print initial debug message for this error
-	TCHAR ErrorString[MAX_SPRINTF];
-	FCString::Sprintf(ErrorString,TEXT("Ensure condition failed: %s"),ANSI_TO_TCHAR(Expr));
-
-	StaticFailDebug( ErrorString, File, Line, ProgramCounter, Msg, true );
+	StaticFailDebug(TEXT("Ensure condition failed: "), Expr, File, Line, /*bIsEnsure*/ true, ProgramCounter, TEXT("%s"), Msg);
 
 	// Is there a debugger attached?  If not we'll submit an error report.
 	if (FPlatformMisc::IsDebuggerPresent() && !GAlwaysReportCrash)
@@ -362,11 +407,11 @@ FORCENOINLINE void FDebug::EnsureFailed(const ANSICHAR* Expr, const ANSICHAR* Fi
 #if !NO_LOGGING
 		if (GEnsuresAreErrors)
 		{
-			UE_LOG(LogOutputDevice, Error, TEXT("%s" FILE_LINE_DESC_ANSI "\n%s\n"), ErrorString, ANSI_TO_TCHAR(File), Line, Msg);
+			UE_LOG(LogOutputDevice, Error, TEXT("Ensure condition failed: %hs" FILE_LINE_DESC_ANSI "\n%s\n"), Expr, File, Line, Msg);
 		}
 		else
 		{
-			UE_LOG(LogOutputDevice, Warning, TEXT("%s" FILE_LINE_DESC_ANSI "\n%s\n"), ErrorString, ANSI_TO_TCHAR(File), Line, Msg);
+			UE_LOG(LogOutputDevice, Warning, TEXT("Ensure condition failed: %hs" FILE_LINE_DESC_ANSI "\n%s\n"), Expr, File, Line, Msg);
 		}
 #endif
 	}
@@ -377,7 +422,7 @@ FORCENOINLINE void FDebug::EnsureFailed(const ANSICHAR* Expr, const ANSICHAR* Fi
 
 		// Create a final string that we'll output to the log (and error history buffer)
 		TCHAR ErrorMsg[16384];
-		FCString::Snprintf(ErrorMsg, UE_ARRAY_COUNT(ErrorMsg), TEXT("Ensure condition failed: %s [File:%s] [Line: %i]" LINE_TERMINATOR_ANSI "%s" LINE_TERMINATOR_ANSI "Stack: " LINE_TERMINATOR_ANSI), ANSI_TO_TCHAR(Expr), ANSI_TO_TCHAR(File), Line, Msg);
+		FCString::Snprintf(ErrorMsg, UE_ARRAY_COUNT(ErrorMsg), TEXT("Ensure condition failed: %hs " FILE_LINE_DESC_ANSI LINE_TERMINATOR_ANSI "%s" LINE_TERMINATOR_ANSI "Stack: " LINE_TERMINATOR_ANSI), Expr, File, Line, Msg);
 
 		// No debugger attached, so generate a call stack and submit a crash report
 		// Walk the stack and dump it to the allocated memory.
@@ -468,11 +513,11 @@ FORCENOINLINE void FDebug::EnsureFailed(const ANSICHAR* Expr, const ANSICHAR* Fi
 #if !NO_LOGGING
 			if (GEnsuresAreErrors)
 			{
-				UE_LOG(LogOutputDevice, Error, TEXT("%s [File:%s] [Line: %i]"), ErrorString, ANSI_TO_TCHAR(File), Line);
+				UE_LOG(LogOutputDevice, Error, TEXT("Ensure condition failed: %hs " FILE_LINE_DESC_ANSI), Expr, File, Line);
 			}
 			else
 			{
-				UE_LOG(LogOutputDevice, Warning, TEXT("%s [File:%s] [Line: %i]"), ErrorString, ANSI_TO_TCHAR(File), Line);
+				UE_LOG(LogOutputDevice, Warning, TEXT("Ensure condition failed: %hs " FILE_LINE_DESC_ANSI), Expr, File, Line);
 			}
 #endif
 		}
@@ -540,6 +585,11 @@ void VARARGS FDebug::AssertFailed(const ANSICHAR* Expr, const ANSICHAR* File, in
 	va_end(Args);
 }
 
+void FDebug::AssertFailedV(const ANSICHAR* Expr, const ANSICHAR* File, int32 Line, const TCHAR* Format, va_list Args)
+{
+	AssertFailedImplV(Expr, File, Line, nullptr, Format, Args);
+}
+
 void FDebug::ProcessFatalError(void* ProgramCounter)
 {
 	// This is not perfect because another thread might crash and be handled before this assert
@@ -570,10 +620,10 @@ FORCENOINLINE bool VARARGS FDebug::OptionallyLogFormattedEnsureMessageReturningF
 
 FORCENOINLINE void VARARGS LowLevelFatalErrorHandler(const ANSICHAR* File, int32 Line, void* ProgramCounter, const TCHAR* Format, ...)
 {
-	TCHAR DescriptionString[4096];
-	GET_VARARGS( DescriptionString, UE_ARRAY_COUNT(DescriptionString), UE_ARRAY_COUNT(DescriptionString)-1, Format, Format );
-
-	StaticFailDebug(TEXT("LowLevelFatalError"), File, Line, ProgramCounter, DescriptionString, false);
+	va_list Args;
+	va_start(Args, Format);
+	StaticFailDebugV(TEXT("LowLevelFatalError"), "", File, Line, /*bIsEnsure*/ false, ProgramCounter, Format, Args);
+	va_end(Args);
 }
 
 void FDebug::DumpStackTraceToLog(const ELogVerbosity::Type LogVerbosity)
diff --git a/Engine/Source/Runtime/Core/Public/Misc/AssertionMacros.h b/Engine/Source/Runtime/Core/Public/Misc/AssertionMacros.h
--- a/Engine/Source/Runtime/Core/Public/Misc/AssertionMacros.h
+++ b/Engine/Source/Runtime/Core/Public/Misc/AssertionMacros.h
@@ -46,6 +46,7 @@ struct CORE_API FDebug
 {
 	/** Logs final assert message and exits the program. */
 	static void VARARGS AssertFailed(const ANSICHAR* Expr, const ANSICHAR* File, int32 Line, const TCHAR* Format = TEXT(""), ...);
+	static void AssertFailedV(const ANSICHAR* Expr, const ANSICHAR* File, int32 Line, const TCHAR* Format, va_list Args);
 
 	/** Triggers a fatal error, using the error formatted to GErrorHist via a previous call to FMsg*/
 	static void ProcessFatalError(void* ProgramCounter);
