Trim Low level allocator memory when OS might go into OOM

#rnx

diff --git a/Engine/Source/Runtime/Core/Private/HAL/MallocBinned2.cpp b/Engine/Source/Runtime/Core/Private/HAL/MallocBinned2.cpp
--- a/Engine/Source/Runtime/Core/Private/HAL/MallocBinned2.cpp
+++ b/Engine/Source/Runtime/Core/Private/HAL/MallocBinned2.cpp
@@ -10,6 +10,7 @@
 #include "HAL/MemoryMisc.h"
 #include "HAL/PlatformMisc.h"
 #include "Misc/App.h"
+#include "Misc/CoreDelegates.h"
 #include "HAL/MallocTimer.h"
 #include "ProfilingDebugging/CsvProfiler.h"
 #include "FramePro/FrameProProfiler.h"
@@ -797,6 +798,17 @@ FMallocBinned2::~FMallocBinned2()
 {
 }
 
+void FMallocBinned2::OnMallocInitialized() 
+{
+#if UE_USE_VERYLARGEPAGEALLOCATOR
+	FCoreDelegates::GetLowLevelAllocatorMemoryTrimDelegate().AddLambda([this]()
+		{
+			CachedOSPageAllocator.FreeAll(&Mutex);
+		}
+	);
+#endif
+}
+
 void FMallocBinned2::OnPreFork()
 {
 #if BINNED2_FORK_SUPPORT
@@ -1511,4 +1523,4 @@ void FMallocBinned2::UpdateStats()
 }
 
 
-PRAGMA_RESTORE_UNSAFE_TYPECAST_WARNINGS
+PRAGMA_RESTORE_UNSAFE_TYPECAST_WARNINGS
\ No newline at end of file
diff --git a/Engine/Source/Runtime/Core/Private/HAL/UnrealMemory.cpp b/Engine/Source/Runtime/Core/Private/HAL/UnrealMemory.cpp
--- a/Engine/Source/Runtime/Core/Private/HAL/UnrealMemory.cpp
+++ b/Engine/Source/Runtime/Core/Private/HAL/UnrealMemory.cpp
@@ -218,6 +218,11 @@ public:
 		return UsedMalloc->GetDescriptiveName();
 	}
 
+	virtual void OnMallocInitialized() override
+	{
+		UsedMalloc->OnMallocInitialized();
+	}
+
 	virtual void OnPreFork() override
 	{
 		UsedMalloc->OnPreFork();
@@ -414,6 +419,9 @@ static int FMemory_GCreateMalloc_ThreadUnsafe()
 
 	GMalloc = FMallocDoubleFreeFinder::OverrideIfEnabled(GMalloc);
 	GMalloc = FMallocFrameProfiler::OverrideIfEnabled(GMalloc);
+
+	GMalloc->OnMallocInitialized();
+
 	return 0;
 }
 
@@ -717,4 +725,4 @@ SIZE_T FMemory::GetUsedPersistentAuxiliary()
 
 #if !INLINE_FMEMORY_OPERATION && !PLATFORM_USES_FIXED_GMalloc_CLASS
 #include "HAL/FMemory.inl"
-#endif
+#endif
\ No newline at end of file
diff --git a/Engine/Source/Runtime/Core/Private/Misc/CoreDelegates.cpp b/Engine/Source/Runtime/Core/Private/Misc/CoreDelegates.cpp
--- a/Engine/Source/Runtime/Core/Private/Misc/CoreDelegates.cpp
+++ b/Engine/Source/Runtime/Core/Private/Misc/CoreDelegates.cpp
@@ -224,6 +224,12 @@ FSimpleMulticastDelegate& FCoreDelegates::GetMemoryTrimDelegate()
 	return OnMemoryTrim;
 }
 
+FSimpleMulticastDelegate& FCoreDelegates::GetLowLevelAllocatorMemoryTrimDelegate()
+{
+	static FSimpleMulticastDelegate OnLowLevelAllocatorMemoryTrim;
+	return OnLowLevelAllocatorMemoryTrim;
+}
+
 /**	 Implemented as a function to address global ctor issues */
 FSimpleMulticastDelegate& FCoreDelegates::GetOutOfMemoryDelegate()
 {
diff --git a/Engine/Source/Runtime/Core/Private/ProfilingDebugging/MemoryTrace.cpp b/Engine/Source/Runtime/Core/Private/ProfilingDebugging/MemoryTrace.cpp
--- a/Engine/Source/Runtime/Core/Private/ProfilingDebugging/MemoryTrace.cpp
+++ b/Engine/Source/Runtime/Core/Private/ProfilingDebugging/MemoryTrace.cpp
@@ -173,7 +173,8 @@ private:
 	virtual bool			ValidateHeap() override										{ return InnerMalloc->ValidateHeap(); }
 	virtual bool			GetAllocationSize(void* Address, SIZE_T &SizeOut) override	{ return InnerMalloc->GetAllocationSize(Address, SizeOut); }
 	virtual void			SetupTLSCachesOnCurrentThread() override					{ return InnerMalloc->SetupTLSCachesOnCurrentThread(); }
-	virtual void			OnPreFork() override										{ InnerMalloc->OnPreFork();}
+	virtual void			OnMallocInitialized() override								{ InnerMalloc->OnMallocInitialized(); }
+	virtual void			OnPreFork() override										{ InnerMalloc->OnPreFork(); }
 	virtual void			OnPostFork() override										{ InnerMalloc->OnPostFork(); }
 
 	FMalloc*				InnerMalloc;
diff --git a/Engine/Source/Runtime/Core/Public/HAL/MallocBinned2.h b/Engine/Source/Runtime/Core/Public/HAL/MallocBinned2.h
--- a/Engine/Source/Runtime/Core/Public/HAL/MallocBinned2.h
+++ b/Engine/Source/Runtime/Core/Public/HAL/MallocBinned2.h
@@ -763,6 +763,7 @@ if (NewSize <= BINNED2_MAX_SMALL_POOL_SIZE && Alignment <= BINNED2_MINIMUM_ALIGN
 	virtual void ClearAndDisableTLSCachesOnCurrentThread() override;
 	virtual const TCHAR* GetDescriptiveName() override;
 	virtual void UpdateStats() override;
+	virtual void OnMallocInitialized() override;
 	virtual void OnPreFork() override;
 	virtual void OnPostFork() override;
 	// End FMalloc interface.
diff --git a/Engine/Source/Runtime/Core/Public/HAL/MallocCallstackHandler.h b/Engine/Source/Runtime/Core/Public/HAL/MallocCallstackHandler.h
--- a/Engine/Source/Runtime/Core/Public/HAL/MallocCallstackHandler.h
+++ b/Engine/Source/Runtime/Core/Public/HAL/MallocCallstackHandler.h
@@ -132,6 +132,11 @@ public:
 	{
 		return UsedMalloc->GetDescriptiveName();
 	}
+		
+	virtual void OnMallocInitialized() override
+	{
+		UsedMalloc->OnMallocInitialized();
+	}
 
 	virtual void OnPreFork() override
 	{
diff --git a/Engine/Source/Runtime/Core/Public/HAL/MallocPoisonProxy.h b/Engine/Source/Runtime/Core/Public/HAL/MallocPoisonProxy.h
--- a/Engine/Source/Runtime/Core/Public/HAL/MallocPoisonProxy.h
+++ b/Engine/Source/Runtime/Core/Public/HAL/MallocPoisonProxy.h
@@ -144,6 +144,11 @@ public:
 		UsedMalloc->ClearAndDisableTLSCachesOnCurrentThread();
 	}
 
+	virtual void OnMallocInitialized() override
+	{
+		UsedMalloc->OnMallocInitialized();
+	}
+
 	virtual void OnPreFork() override
 	{
 		UsedMalloc->OnPreFork();
diff --git a/Engine/Source/Runtime/Core/Public/HAL/MallocThreadSafeProxy.h b/Engine/Source/Runtime/Core/Public/HAL/MallocThreadSafeProxy.h
--- a/Engine/Source/Runtime/Core/Public/HAL/MallocThreadSafeProxy.h
+++ b/Engine/Source/Runtime/Core/Public/HAL/MallocThreadSafeProxy.h
@@ -125,6 +125,11 @@ public:
 		return true; 
 	}
 
+	virtual void OnMallocInitialized() override
+	{
+		UsedMalloc->OnMallocInitialized();
+	}
+
 	virtual void OnPreFork() override
 	{
 		UsedMalloc->OnPreFork();
diff --git a/Engine/Source/Runtime/Core/Public/HAL/MemoryBase.h b/Engine/Source/Runtime/Core/Public/HAL/MemoryBase.h
--- a/Engine/Source/Runtime/Core/Public/HAL/MemoryBase.h
+++ b/Engine/Source/Runtime/Core/Public/HAL/MemoryBase.h
@@ -209,6 +209,11 @@ public:
 		return TEXT("Unspecified allocator");
 	}
 
+	/**
+	 * Notifies the malloc implementation that initialization of all allocators in GMalloc is complete, so it's safe to initialize any extra features that require "regular" allocations
+	 */
+	virtual void OnMallocInitialized() {}
+
 	/**
 	 * Notifies the malloc implementation that the process is about to fork. May be used to trim caches etc.
 	 */
diff --git a/Engine/Source/Runtime/Core/Public/Misc/CoreDelegates.h b/Engine/Source/Runtime/Core/Public/Misc/CoreDelegates.h
--- a/Engine/Source/Runtime/Core/Public/Misc/CoreDelegates.h
+++ b/Engine/Source/Runtime/Core/Public/Misc/CoreDelegates.h
@@ -596,6 +596,9 @@ public:
 	// Caller is responsible for flushing rendering etc. See UEngine::TrimMemory
 	static FSimpleMulticastDelegate& GetMemoryTrimDelegate();
 
+	// Called to request that low level allocator free whatever memory they are able to. 
+	static FSimpleMulticastDelegate& GetLowLevelAllocatorMemoryTrimDelegate();
+
 	// Called when OOM event occurs, after backup memory has been freed, so there's some hope of being effective
 	static FSimpleMulticastDelegate& GetOutOfMemoryDelegate();
 
diff --git a/Engine/Source/Runtime/Core/Public/ProfilingDebugging/TraceMalloc.h b/Engine/Source/Runtime/Core/Public/ProfilingDebugging/TraceMalloc.h
--- a/Engine/Source/Runtime/Core/Public/ProfilingDebugging/TraceMalloc.h
+++ b/Engine/Source/Runtime/Core/Public/ProfilingDebugging/TraceMalloc.h
@@ -16,6 +16,11 @@ public:
 	UE_DEPRECATED(5.1, "ShouldTrace is deprecated. Details should be internal now.")
 	static bool ShouldTrace();
 
+	virtual void OnMallocInitialized() override
+	{
+		WrappedMalloc->OnMallocInitialized();
+	}
+
 	virtual void OnPreFork() override
 	{
 		WrappedMalloc->OnPreFork();
diff --git a/Engine/Source/Runtime/Core/Public/Stats/StatsMallocProfilerProxy.h b/Engine/Source/Runtime/Core/Public/Stats/StatsMallocProfilerProxy.h
--- a/Engine/Source/Runtime/Core/Public/Stats/StatsMallocProfilerProxy.h
+++ b/Engine/Source/Runtime/Core/Public/Stats/StatsMallocProfilerProxy.h
@@ -133,6 +133,11 @@ public:
 		return UsedMalloc->GetDescriptiveName(); 
 	}
 
+	virtual void OnMallocInitialized() override
+	{
+		UsedMalloc->OnMallocInitialized();
+	}
+
 	virtual void OnPreFork() override
 	{
 		UsedMalloc->OnPreFork();
