- Removed the idea of Required .ini files (Base.ini). It never helped, and just causes issues with Programs.

diff --git a/Engine/Source/Runtime/Core/Private/Misc/ConfigContext.cpp b/Engine/Source/Runtime/Core/Private/Misc/ConfigContext.cpp
--- a/Engine/Source/Runtime/Core/Private/Misc/ConfigContext.cpp
+++ b/Engine/Source/Runtime/Core/Private/Misc/ConfigContext.cpp
@@ -601,21 +601,12 @@ static bool LoadIniFileHierarchy(const FConfigFileHierarchy& HierarchyToLoad, FC
 	for (const TPair<int32, FString>& HierarchyIt : HierarchyToLoad)
 	{
 		bool bDoCombine = (HierarchyIt.Key != 0);
-		bool bIsRequired = (HierarchyIt.Key == 0);
 		const FString& IniFileName = HierarchyIt.Value;
 
-		// Spit out friendly error if there was a problem locating .inis (e.g. bad command line parameter or missing folder, ...).
+		// skip non-existant files
 		if (IsUsingLocalIniFile(*IniFileName, nullptr) && !DoesConfigFileExistWrapper(*IniFileName, IniCacheSet))
 		{
-			if (bIsRequired)
-			{
-				UE_LOG(LogConfig, Error, TEXT("Couldn't locate '%s' which is required to run '%s'"), *IniFileName, FApp::GetProjectName());
-				return false;
-			}
-			else
-			{
-				continue;
-			}
+			continue;
 		}
 
 		bool bDoEmptyConfig = false;
