FThreadHeartBeat
- Fix CheckpointSuspendCount undercount

diff --git a/Engine/Source/Runtime/Core/Private/HAL/ThreadHeartBeat.cpp b/Engine/Source/Runtime/Core/Private/HAL/ThreadHeartBeat.cpp
--- a/Engine/Source/Runtime/Core/Private/HAL/ThreadHeartBeat.cpp
+++ b/Engine/Source/Runtime/Core/Private/HAL/ThreadHeartBeat.cpp
@@ -792,6 +792,7 @@ void FThreadHeartBeat::SuspendHeartBeat(bool bAllThreads)
 		if (!bAllThreads)
 		{
 			FScopeLock HeartBeatLock(&CheckpointHeartBeatCritical);
+			check(CheckpointSuspendCount >= 0);
 			++CheckpointSuspendCount;
 			for (TPair<FName, FHeartBeatInfo>& HeartBeatEntry : CheckpointHeartBeat)
 			{
