SharedBuffer: Fixed a memory leak in the tests

#rnx

diff --git a/Engine/Source/Runtime/Core/Private/Tests/Memory/SharedBufferTest.cpp b/Engine/Source/Runtime/Core/Private/Tests/Memory/SharedBufferTest.cpp
--- a/Engine/Source/Runtime/Core/Private/Tests/Memory/SharedBufferTest.cpp
+++ b/Engine/Source/Runtime/Core/Private/Tests/Memory/SharedBufferTest.cpp
@@ -273,11 +273,11 @@ bool FSharedBufferTest::RunTest(const FString& Parameters)
 	{
 		bool bDeleted = false;
 		constexpr uint64 Size = 64;
-		const void* Data = new uint8[Size];
+		uint8 Data[Size]{};
 		FSharedBuffer Buffer = FSharedBuffer::TakeOwnership(Data, Size, [&bDeleted](void*, uint64) { bDeleted = true; });
 		TestTrue(TEXT("FSharedBuffer::TakeOwnership().IsOwned()"), Buffer.IsOwned());
 		TestEqual(TEXT("FSharedBuffer::TakeOwnership().GetSize()"), Buffer.GetSize(), Size);
-		TestEqual(TEXT("FSharedBuffer::TakeOwnership().GetData()"), Buffer.GetData(), Data);
+		TestEqual(TEXT("FSharedBuffer::TakeOwnership().GetData()"), Buffer.GetData(), static_cast<const void*>(Data));
 		Buffer.Reset();
 		TestTrue(TEXT("FSharedBuffer::TakeOwnership() Deleted"), bDeleted);
 	}
