StringBuilder: Added AppendChars(TStringBuilderBase&) to work with algorithms

Can be used as Algo::Transform(View, AppendChars(Builder), FChar::ToLower)

diff --git a/Engine/Source/Runtime/Core/Public/Misc/StringBuilder.h b/Engine/Source/Runtime/Core/Public/Misc/StringBuilder.h
--- a/Engine/Source/Runtime/Core/Public/Misc/StringBuilder.h
+++ b/Engine/Source/Runtime/Core/Public/Misc/StringBuilder.h
@@ -504,3 +504,19 @@ template <int32 BufferSize> using WriteToString = TWriteToString<TCHAR, BufferSi
 template <int32 BufferSize> using WriteToAnsiString = TWriteToString<ANSICHAR, BufferSize>;
 template <int32 BufferSize> using WriteToWideString = TWriteToString<WIDECHAR, BufferSize>;
 template <int32 BufferSize> using WriteToUtf8String = TWriteToString<UTF8CHAR, BufferSize>;
+
+/**
+ * Returns an object that can be used as the output container for algorithms by appending to the builder.
+ *
+ * Example: Algo::Transform(StringView, AppendChars(Builder), FChar::ToLower)
+ */
+template <typename CharType>
+auto AppendChars(TStringBuilderBase<CharType>& Builder)
+{
+	struct FAppendChar
+	{
+		TStringBuilderBase<CharType>& Builder;
+		inline void Add(CharType Char) { Builder.AppendChar(Char); }
+	};
+	return FAppendChar{Builder};
+}
