Fix an issue with CL: 21203967
Switchin to use a FAutoConsoleVariableRef so that the local static variable can be sampled rather than using the GetValueOnGameThread method which can ensure if sampled incorrectly.

diff --git a/Engine/Source/Runtime/Core/Private/Misc/AutomationTest.cpp b/Engine/Source/Runtime/Core/Private/Misc/AutomationTest.cpp
--- a/Engine/Source/Runtime/Core/Private/Misc/AutomationTest.cpp
+++ b/Engine/Source/Runtime/Core/Private/Misc/AutomationTest.cpp
@@ -18,12 +18,15 @@
 
 DEFINE_LOG_CATEGORY_STATIC(LogAutomationTest, Warning, All);
 
-static TAutoConsoleVariable<bool> CVarAutomationCaptureLogEvents(
-	TEXT("Automation.CaptureLogEvents"),
-	1,
-	TEXT("Consider warning/error log events during a test as impacting the test itself"),
-	ECVF_ReadOnly
-);
+namespace AutomationTest
+{
+	static bool bCaptureLogEvents = true;
+	static FAutoConsoleVariableRef CVarAutomationCaptureLogEvents(
+		TEXT("Automation.CaptureLogEvents"),
+		bCaptureLogEvents,
+		TEXT("Consider warning/error log events during a test as impacting the test itself"));
+};
+
 
 CORE_API const TMap<FString, EAutomationTestFlags::Type>& EAutomationTestFlags::GetTestFlagsMap()
 {
@@ -75,7 +78,7 @@ CORE_API ELogVerbosity::Type GetAutomationLogLevel(ELogVerbosity::Type LogVerbos
 	static TArray<FString> SuppressedLogCategories;
 	static FAutomationTestBase* LastTest = nullptr;
 
-	if (CVarAutomationCaptureLogEvents.GetValueOnGameThread() == false)
+	if (AutomationTest::bCaptureLogEvents == false)
 	{
 		return ELogVerbosity::NoLogging;
 	}
