Extending IAutomationLatentCommand to allow for commands with retries and delays.
* Added pure virtual Latent Command Child class IAutomationLatentCommandWithRetriesAndDelays for built-in command retry and delays in between runs
* Added helper function to IAutomationLatentCommand to get current run time for better reporting

andrew.grant

diff --git a/Engine/Source/Runtime/Core/Public/Misc/AutomationTest.h b/Engine/Source/Runtime/Core/Public/Misc/AutomationTest.h
--- a/Engine/Source/Runtime/Core/Public/Misc/AutomationTest.h
+++ b/Engine/Source/Runtime/Core/Public/Misc/AutomationTest.h
@@ -47,6 +47,7 @@
 
 #include <atomic>
 
+DEFINE_LOG_CATEGORY_STATIC(LogLatentCommands, Log, All);
 class FAutomationTestBase;
 
 #ifndef WITH_AUTOMATION_TESTS
@@ -511,6 +512,7 @@ private:
 		{
 			StartTime = FPlatformTime::Seconds();
 		}
+
 		return Update();
 	}
 
@@ -521,12 +523,145 @@ protected:
 	{
 	}
 
+	// Gets current run time for the command for reporting purposes.
+	double GetCurrentRunTime() const
+	{
+		if (StartTime == 0.0)
+		{
+			return 0.0;
+		}
+
+		return FPlatformTime::Seconds() - StartTime;
+	}
+
 	/** For timers, track the first time this ticks */
 	double StartTime;
 
 	friend class FAutomationTestFramework;
 };
 
+// Extension of IAutomationLatentCommand with delays between attempts
+// if initial command does not succeed. Has a max retry count that can be 
+// overridden. Default is 10 retries with a 1 second delay in between.
+class IAutomationLatentCommandWithRetriesAndDelays : public IAutomationLatentCommand
+{	
+public:
+	virtual ~IAutomationLatentCommandWithRetriesAndDelays() {}
+
+	// Base Update override with delay logic built in. Can be further overridden in child
+	// classes
+	virtual bool Update() override
+	{
+		if (IsDelayTimerRunning())
+		{
+			return false;
+		}
+
+		// pre-increment iteration so that its 1 based instead of 0 based for readability
+		CurrentIteration++;
+
+		if (!CanRetry())
+		{
+			// Retry count has been exceeded
+			UE_LOG(LogLatentCommands, Error, TEXT("%s All Retries Have Been Exhausted. Command Has Failed"), *GetCommandName());
+#if WITH_EDITOR
+			FPlatformMisc::RequestExitWithStatus(false, -1);
+#endif
+			return true;
+		}
+
+		ResetDelayTimer();
+
+		// auto-logging for limited retries and infinite retries
+		if (MaxRetries == 0) 
+		{
+			UE_LOG(LogLatentCommands, Log, TEXT("%s Executing Attempt %d."), *GetCommandName(), CurrentIteration);
+		}
+		else
+		{
+			UE_LOG(LogLatentCommands, Log, TEXT("%s Executing Attempt %d of %d."), *GetCommandName(), CurrentIteration, MaxRetries);
+		}
+
+		if (Execute())
+		{
+			// completion log message, logs total time taken.
+			UE_LOG(LogLatentCommands, Log, TEXT("%s Completed Successfully, total run time: %f."), *GetCommandName(), GetCurrentRunTime());
+			return true;
+		}
+
+		return false;
+	}
+
+	// Pure virtual method that must be overridden.
+	// This is the actual command logic.
+	virtual bool Execute() = 0;
+
+private:
+	// To keep track of which iteration we are on. 1 based
+	// To allow for more readable logs. Attempt 1 of N instead
+	// of 0 of N.
+	int32 CurrentIteration = 0;
+
+protected:
+	// default constructor
+	IAutomationLatentCommandWithRetriesAndDelays() {}
+
+	// parameterized constructor
+	IAutomationLatentCommandWithRetriesAndDelays(const FString InCommandClassName, const int32 InMaxRetries, const double InWaitTimeBetweenRuns)
+		:CommandClassName(InCommandClassName)
+		,MaxRetries(InMaxRetries)
+		,DelayTimeInSeconds(InWaitTimeBetweenRuns)
+	{
+		// set first run start time so that we don't start off in a delay. Its fine if its negative.
+		DelayStartTime = FPlatformTime::Seconds() - InWaitTimeBetweenRuns;
+	}
+
+	// Resets the Delay Timer by setting start time to now (In game time).
+	void ResetDelayTimer()
+	{
+		DelayStartTime = FPlatformTime::Seconds();
+	}
+
+	// Determines if timer is running, pausing execution in a non-blocking way.
+	bool IsDelayTimerRunning() const
+	{
+		// time elapsed < Delay time
+		return (FPlatformTime::Seconds() - DelayStartTime) < DelayTimeInSeconds;
+	}
+
+	// Determines if we should execute the command. If MaxRetries is set to 0, 
+	// indicates infinite retries.
+	bool CanRetry() const
+	{
+		return (CurrentIteration <= MaxRetries) || MaxRetries == 0;
+	}
+
+	// Returns the command name 
+	FString GetCommandName(const FString& TestName = FString(TEXT(""))) const
+	{
+		if (!TestName.IsEmpty())
+		{
+			return FString::Printf(TEXT("Test: Unknown - Command: %s - "), *CommandClassName);
+		}
+		
+		return FString::Printf(TEXT("Test: %s - Command: %s - "), *TestName, *CommandClassName);
+	}
+
+	// Command Name stored for easy logging
+	const FString CommandClassName = "UnknownCommand";
+	
+	// Times to retry command before reporting command failure.
+	// Defaults to 10, but can be overridden.
+	const int32 MaxRetries = 10;
+	
+	// Time in between UpdateDelayed calls. 
+	// Default is 1 second but can be overridden.
+	const double DelayTimeInSeconds = 1.0;
+	
+	// Time that the Delay Timer started.
+	double DelayStartTime = 0.0;
+};
+
 /**
  * A simple latent command that runs the provided function on another thread
  */
@@ -2928,7 +3063,7 @@ private:
 };
 
 
-//////////////////////////////////////////////////////////////////////////
+//////////////////////////////////////////////////
 // Latent command definition macros
 
 #define DEFINE_LATENT_AUTOMATION_COMMAND(CommandName)	\
@@ -3075,6 +3210,91 @@ class EXPORT_API CommandName : public IAutomationLatentCommand \
 #define DEFINE_ENGINE_LATENT_AUTOMATION_COMMAND_ONE_PARAMETER(CommandName,ParamType,ParamName)	\
 	DEFINE_EXPORTED_LATENT_AUTOMATION_COMMAND_ONE_PARAMETER(ENGINE_API, CommandName, ParamType, ParamName)
 
+#define DEFINE_EXPORTED_LATENT_AUTOMATION_COMMAND_WITH_RETRIES(EXPORT_API,CommandName,RetryCount,WaitTimeBetweenRuns)	\
+class EXPORT_API CommandName : public IAutomationLatentCommandWithRetriesAndDelays \
+	{ \
+	public: \
+	CommandName(int32 InRetryCount, double InWaitTimeBetweenRuns) \
+	: IAutomationLatentCommandWithRetriesAndDelays(#CommandName, InTest, InRetryCount, InWaitTimeBetweenRuns) \
+		{} \
+		virtual ~CommandName() \
+		{} \
+		virtual bool Execute() override; \
+	private: \
+}
+
+#define DEFINE_EXPORTED_LATENT_AUTOMATION_COMMAND_WITH_RETRIES_ONE_PARAMETER(EXPORT_API,CommandName,RetryCount,WaitTimeBetweenRuns,ParamType,ParamName)	\
+class EXPORT_API CommandName : public IAutomationLatentCommandWithRetriesAndDelays \
+	{ \
+	public: \
+	CommandName(int32 InRetryCount, double InWaitTimeBetweenRuns, ParamType ParamName) \
+	: IAutomationLatentCommandWithRetriesAndDelays(#CommandName, InRetryCount, InWaitTimeBetweenRuns) \
+	, ParamName(ParamName) \
+		{} \
+		virtual ~CommandName() \
+		{} \
+		virtual bool Execute() override; \
+	private: \
+	ParamType ParamName; \
+}
+
+#define DEFINE_EXPORTED_LATENT_AUTOMATION_COMMAND_WITH_RETRIES_TWO_PARAMETERS(EXPORT_API,CommandName,RetryCount,WaitTimeBetweenRuns,ParamType0,ParamName0,ParamType1,ParamName1)	\
+class EXPORT_API CommandName : public IAutomationLatentCommandWithRetriesAndDelays \
+	{ \
+	public: \
+	CommandName(int32 InRetryCount, double InWaitTimeBetweenRuns, ParamType0 ParamName0, ParamType1 ParamName1) \
+	: IAutomationLatentCommandWithRetriesAndDelays(#CommandName, InRetryCount, InWaitTimeBetweenRuns) \
+	, ParamName0(ParamName0) \
+	, ParamName1(ParamName1) \
+		{} \
+		virtual ~CommandName() \
+		{} \
+		virtual bool Execute() override; \
+	private: \
+	ParamType0 ParamName0; \
+	ParamType1 ParamName1; \
+}
+
+#define DEFINE_EXPORTED_LATENT_AUTOMATION_COMMAND_WITH_RETRIES_THREE_PARAMETERS(EXPORT_API,CommandName,RetryCount,WaitTimeBetweenRuns,ParamType0,ParamName0,ParamType1,ParamName1,ParamType2,ParamName2)	\
+class EXPORT_API CommandName : public IAutomationLatentCommandWithRetriesAndDelays \
+	{ \
+	public: \
+	CommandName(int32 InRetryCount, double InWaitTimeBetweenRuns, ParamType0 ParamName0, ParamType1 ParamName1, ParamType2 ParamName2) \
+	: IAutomationLatentCommandWithRetriesAndDelays(#CommandName, InRetryCount, InWaitTimeBetweenRuns) \
+	, ParamName0(ParamName0) \
+	, ParamName1(ParamName1) \
+	, ParamName2(ParamName2) \
+		{} \
+		virtual ~CommandName() \
+		{} \
+		virtual bool Execute() override; \
+	private: \
+	ParamType0 ParamName0; \
+	ParamType1 ParamName1; \
+	ParamType2 ParamName2; \
+}
+
+#define DEFINE_EXPORTED_LATENT_AUTOMATION_COMMAND_WITH_RETRIES_FOUR_PARAMETERS(EXPORT_API,CommandName,RetryCount,WaitTimeBetweenRuns,ParamType0,ParamName0,ParamType1,ParamName1,ParamType2,ParamName2,ParamType3,ParamName3)	\
+class EXPORT_API CommandName : public IAutomationLatentCommandWithRetriesAndDelays \
+	{ \
+	public: \
+	CommandName(int32 InRetryCount, double InWaitTimeBetweenRuns, ParamType0 ParamName0, ParamType1 ParamName1, ParamType2 ParamName2, ParamType3 ParamName3) \
+	: IAutomationLatentCommandWithRetriesAndDelays(#CommandName, InRetryCount, InWaitTimeBetweenRuns) \
+	, ParamName0(ParamName0) \
+	, ParamName1(ParamName1) \
+	, ParamName2(ParamName2) \
+	, ParamName3(ParamName3) \
+		{} \
+		virtual ~CommandName() \
+		{} \
+		virtual bool Execute() override; \
+	private: \
+	ParamType0 ParamName0; \
+	ParamType1 ParamName1; \
+	ParamType2 ParamName2; \
+	ParamType3 ParamName3; \
+}
+
 //macro to simply the syntax for enqueueing a latent command
 #define ADD_LATENT_AUTOMATION_COMMAND(ClassDeclaration) FAutomationTestFramework::Get().EnqueueLatentCommand(MakeShareable(new ClassDeclaration));
 
@@ -3465,7 +3685,7 @@ public: \
 		return false;\
 	}
 
-//////////////////////////////////////////////////////////////////////////
+//////////////////////////////////////////////////
 // Basic Latent Commands
 
 /**
