Remove CreateFile from PostWindowsApi because some things rely on it leaking out

diff --git a/Engine/Source/Runtime/Core/Public/Windows/PostWindowsApi.h b/Engine/Source/Runtime/Core/Public/Windows/PostWindowsApi.h
--- a/Engine/Source/Runtime/Core/Public/Windows/PostWindowsApi.h
+++ b/Engine/Source/Runtime/Core/Public/Windows/PostWindowsApi.h
@@ -27,7 +27,6 @@ THIRD_PARTY_INCLUDES_END
 #undef CopyFile
 #undef CreateDesktop
 #undef CreateDirectory
-#undef CreateFile
 #undef CreateFont
 #undef DeleteFile
 #undef DrawText
