comment correction

diff --git a/Engine/Source/Runtime/Core/Public/Tasks/Pipe.h b/Engine/Source/Runtime/Core/Public/Tasks/Pipe.h
--- a/Engine/Source/Runtime/Core/Public/Tasks/Pipe.h
+++ b/Engine/Source/Runtime/Core/Public/Tasks/Pipe.h
@@ -54,7 +54,7 @@ namespace UE::Tasks
 
 		// launches a task in the pipe
 		// @param InDebugName helps to identify the task in debugger and profiler
-		// @param TaskBody a callable with no parameters whose return value is ignored, usually a lambda but can be also a functor object 
+		// @param TaskBody a callable with no parameters, usually a lambda but can be also a functor object 
 		// or a pointer to a function. TaskBody can return results.
 		// @Priority - task priority, can affect task scheduling once it's passed the pipe
 		// @return Task instance that can be used to wait for task completion or to obtain the result of task execution
@@ -78,7 +78,7 @@ namespace UE::Tasks
 
 		// launches a task in the pipe, with multiple prerequisites that must be completed before the task is scheduled
 		// @param InDebugName helps to identify the task in debugger and profiler
-		// @param TaskBody a callable with no parameters whose return value is ignored, usually a lambda but can be also a functor object 
+		// @param TaskBody a callable with no parameters, usually a lambda but can be also a functor object 
 		// or a pointer to a function. TaskBody can return results.
 		// @Priority - task priority, can affect task scheduling once it's passed the pipe
 		// @return Task instance that can be used to wait for task completion or to obtain the result of task execution
