fix compile error when missing DECLARE_STATS_GROUP_SORTBYNAME if STATS==0

diff --git a/Engine/Source/Runtime/Core/Public/Stats/Stats.h b/Engine/Source/Runtime/Core/Public/Stats/Stats.h
--- a/Engine/Source/Runtime/Core/Public/Stats/Stats.h
+++ b/Engine/Source/Runtime/Core/Public/Stats/Stats.h
@@ -325,6 +325,7 @@ public:
 
 #define DECLARE_STATS_GROUP(GroupDesc,GroupId,GroupCat)
 #define DECLARE_STATS_GROUP_VERBOSE(GroupDesc,GroupId,GroupCat)
+#define DECLARE_STATS_GROUP_SORTBYNAME(GroupDesc,GroupId,GroupCat)
 #define DECLARE_STATS_GROUP_MAYBE_COMPILED_OUT(GroupDesc,GroupId,GroupCat,CompileIn)
 
 #define SET_CYCLE_COUNTER(Stat,Cycles)
