PR #9256: Add getters for individual worker thread retrieval (Contributed by [REDACTED_PERSON])

diff --git a/Engine/Source/Runtime/Core/Private/Async/TaskGraph.cpp b/Engine/Source/Runtime/Core/Private/Async/TaskGraph.cpp
--- a/Engine/Source/Runtime/Core/Private/Async/TaskGraph.cpp
+++ b/Engine/Source/Runtime/Core/Private/Async/TaskGraph.cpp
@@ -1541,14 +1541,23 @@ PRAGMA_ENABLE_DEPRECATION_WARNINGS
 		}
 	}
 
-
-	virtual	int32 GetNumWorkerThreads() final override
+	virtual int32 GetNumWorkerThreads() final override
 	{
 		int32 Result = (NumThreads - NumNamedThreads) / NumTaskThreadSets - GNumWorkerThreadsToIgnore;
 		check(Result > 0); // can't tune it to zero task threads
 		return Result;
 	}
 
+	virtual int32 GetNumForegroundThreads() final override
+	{
+		return bCreatedHiPriorityThreads ? NumTaskThreadsPerSet : 0;
+	}
+
+	virtual int32 GetNumBackgroundThreads() final override
+	{
+		return bCreatedBackgroundPriorityThreads ? NumTaskThreadsPerSet : 0;
+	}
+
 	virtual bool IsCurrentThreadKnown() final override
 	{
 		return FPlatformTLS::GetTlsValue(PerThreadIDTLSSlot) != nullptr;
@@ -1883,6 +1892,10 @@ class FTaskGraphCompatibilityImplementation final : public FTaskGraphInterface
 	int32				NumNamedThreads;
 	int32				NumWorkerThreads;
 
+	/** Individual foreground and background workers. **/
+	int32				NumBackgroundWorkers;
+	int32				NumForegroundWorkers;
+
 	TArray<FWorkerThread> NamedThreads;
 
 	FThreadSafeCounter	ReentrancyCheck;
@@ -1903,8 +1916,8 @@ public:
 				GNumForegroundWorkers = 1;
 			}
 
-			int32 NumBackgroundWorkers = FMath::Max(1, NumWorkerThreads - FMath::Min<int>(GNumForegroundWorkers, NumWorkerThreads));
-			int32 NumForegroundWorkers =  FMath::Max(1, NumWorkerThreads - NumBackgroundWorkers);
+			NumBackgroundWorkers = FMath::Max(1, NumWorkerThreads - FMath::Min<int>(GNumForegroundWorkers, NumWorkerThreads));
+			NumForegroundWorkers =  FMath::Max(1, NumWorkerThreads - NumBackgroundWorkers);
 
 			LowLevelTasks::FScheduler::Get().StartWorkers(NumForegroundWorkers, NumBackgroundWorkers, FForkProcessHelper::IsForkedMultithreadInstance() ? FThread::Forkable : FThread::NonForkable, FPlatformAffinity::GetTaskThreadPriority(), FPlatformAffinity::GetTaskBPThreadPriority());
 
@@ -1994,16 +2007,16 @@ public:
 	{
 		if (FTaskGraphInterface::IsMultithread())
 		{
-			int32 NumBackgroundWorkers = FMath::Max(1, NumWorkerThreads - FMath::Min<int>(GNumForegroundWorkers, NumWorkerThreads));
-			int32 NumWorkers =  FMath::Max(1, NumWorkerThreads - NumBackgroundWorkers);
+			NumBackgroundWorkers = FMath::Max(1, NumWorkerThreads - FMath::Min<int>(GNumForegroundWorkers, NumWorkerThreads));
+			NumForegroundWorkers =  FMath::Max(1, NumWorkerThreads - NumBackgroundWorkers);
 
 			LowLevelTasks::FScheduler::Get().StopWorkers();
-			LowLevelTasks::FScheduler::Get().StartWorkers(NumWorkers, NumBackgroundWorkers, FForkProcessHelper::IsForkedMultithreadInstance() ? FThread::Forkable : FThread::NonForkable, Pri, FPlatformAffinity::GetTaskBPThreadPriority());
+			LowLevelTasks::FScheduler::Get().StartWorkers(NumForegroundWorkers, NumBackgroundWorkers, FForkProcessHelper::IsForkedMultithreadInstance() ? FThread::Forkable : FThread::NonForkable, Pri, FPlatformAffinity::GetTaskBPThreadPriority());
 
 			if (bReserveWorkersEnabled)
 			{
 				LowLevelTasks::FReserveScheduler::Get().StopWorkers();
-				LowLevelTasks::FReserveScheduler::Get().StartWorkers(LowLevelTasks::FScheduler::Get(), NumWorkers + NumBackgroundWorkers, FForkProcessHelper::IsForkedMultithreadInstance() ? FThread::Forkable : FThread::NonForkable, FPlatformAffinity::GetTaskBPThreadPriority());
+				LowLevelTasks::FReserveScheduler::Get().StartWorkers(LowLevelTasks::FScheduler::Get(), NumForegroundWorkers + NumBackgroundWorkers, FForkProcessHelper::IsForkedMultithreadInstance() ? FThread::Forkable : FThread::NonForkable, FPlatformAffinity::GetTaskBPThreadPriority());
 			}
 		}
 	}
@@ -2075,6 +2088,16 @@ private:
 		return LowLevelTasks::FScheduler::Get().GetNumWorkers();
 	}
 
+	virtual	int32 GetNumForegroundThreads() final override
+	{
+		return NumForegroundWorkers;
+	}
+
+	virtual	int32 GetNumBackgroundThreads() final override
+	{
+		return NumBackgroundWorkers;
+	}
+
 	bool IsCurrentThreadKnown() final override
 	{
 		return FPlatformTLS::GetTlsValue(PerThreadIDTLSSlot) != nullptr || LowLevelTasks::FTask::GetActiveTask() != nullptr;
@@ -2372,8 +2395,8 @@ private:
 			return; // once enabled, reserve workers can't be disabled
 		}
 
-		int32 NumBackgroundWorkers = FMath::Max(1, NumWorkerThreads - FMath::Min<int>(GNumForegroundWorkers, NumWorkerThreads));
-		int32 NumForegroundWorkers = FMath::Max(1, NumWorkerThreads - NumBackgroundWorkers);
+		NumBackgroundWorkers = FMath::Max(1, NumWorkerThreads - FMath::Min<int>(GNumForegroundWorkers, NumWorkerThreads));
+		NumForegroundWorkers = FMath::Max(1, NumWorkerThreads - NumBackgroundWorkers);
 
 		check(GConfig);
 		
diff --git a/Engine/Source/Runtime/Core/Public/Async/TaskGraphInterfaces.h b/Engine/Source/Runtime/Core/Public/Async/TaskGraphInterfaces.h
--- a/Engine/Source/Runtime/Core/Public/Async/TaskGraphInterfaces.h
+++ b/Engine/Source/Runtime/Core/Public/Async/TaskGraphInterfaces.h
@@ -374,6 +374,18 @@ public:
 	**/
 	virtual	int32 GetNumWorkerThreads() = 0;
 
+	/**
+		Return the number of foreground worker threads.
+		If the old backend is used, return the number of high-pri workers (0 if high-pri workers are disabled).
+	**/
+	virtual	int32 GetNumForegroundThreads() = 0;
+
+	/**
+		Return the number of background worker threads.
+		If the old backend is used, return the number of background workers (0 if background workers are disabled).
+	**/
+	virtual	int32 GetNumBackgroundThreads() = 0;
+
 	/** Return true if the given named thread is processing tasks. This is only a "guess" if you ask for a thread other than yourself because that can change before the function returns. **/
 	virtual bool IsThreadProcessingTasks(ENamedThreads::Type ThreadToCheck) = 0;
 
