Fix String Tables not loading Metadata

diff --git a/Engine/Source/Runtime/Core/Private/Internationalization/StringTableCore.cpp b/Engine/Source/Runtime/Core/Private/Internationalization/StringTableCore.cpp
--- a/Engine/Source/Runtime/Core/Private/Internationalization/StringTableCore.cpp
+++ b/Engine/Source/Runtime/Core/Private/Internationalization/StringTableCore.cpp
@@ -381,7 +381,7 @@ void FStringTable::Serialize(FArchive& Ar)
 
 			KeysToMetaData.Reset();
 			KeysToMetaData.Reserve(TmpKeysToMetaData.Num());
-			for (auto& TmpKeyToMetaDataPair : KeysToMetaData)
+			for (auto& TmpKeyToMetaDataPair : TmpKeysToMetaData)
 			{
 				KeysToMetaData.Add(TmpKeyToMetaDataPair.Key, MoveTemp(TmpKeyToMetaDataPair.Value));
 			}
