Removal of uncalled ANSI->UTF8 conversions.
Fix for incorrect sign extension in to-UTF8 conversions.
Improved surrogate testing implementations.

diff --git a/Engine/Source/Runtime/Core/Private/GenericPlatform/GenericPlatformString.cpp b/Engine/Source/Runtime/Core/Private/GenericPlatform/GenericPlatformString.cpp
--- a/Engine/Source/Runtime/Core/Private/GenericPlatform/GenericPlatformString.cpp
+++ b/Engine/Source/Runtime/Core/Private/GenericPlatform/GenericPlatformString.cpp
@@ -77,16 +77,22 @@ namespace UE::Core::Private
 		return true;
 	}
 
+	/** Is the provided Codepoint within the range of the high-or low surrogates? */
+	static FORCEINLINE bool IsSurrogate(const uint32 Codepoint)
+	{
+		return (Codepoint & 0xFFFFF800) == 0xD800;
+	}
+
 	/** Is the provided Codepoint within the range of the high-surrogates? */
 	static FORCEINLINE bool IsHighSurrogate(const uint32 Codepoint)
 	{
-		return Codepoint >= HIGH_SURROGATE_START_CODEPOINT && Codepoint <= HIGH_SURROGATE_END_CODEPOINT;
+		return (Codepoint & 0xFFFFFC00) == 0xD800;
 	}
 
 	/** Is the provided Codepoint within the range of the low-surrogates? */
 	static FORCEINLINE bool IsLowSurrogate(const uint32 Codepoint)
 	{
-		return Codepoint >= LOW_SURROGATE_START_CODEPOINT && Codepoint <= LOW_SURROGATE_END_CODEPOINT;
+		return (Codepoint & 0xFFFFFC00) == 0xDC00;
 	}
 
 	static FORCEINLINE uint32 EncodeSurrogate(const uint16 HighSurrogate, const uint16 LowSurrogate)
@@ -128,7 +134,7 @@ namespace UE::Core::Private
 		{
 			Codepoint = (uint32)UNICODE_BOGUS_CHAR_CODEPOINT;
 		}
-		else if (IsHighSurrogate(Codepoint) || IsLowSurrogate(Codepoint)) // UTF-8 Characters are not allowed to encode codepoints in the surrogate pair range
+		else if (IsSurrogate(Codepoint)) // UTF-8 Characters are not allowed to encode codepoints in the surrogate pair range
 		{
 			Codepoint = (uint32)UNICODE_BOGUS_CHAR_CODEPOINT;
 		}
@@ -186,6 +192,8 @@ namespace UE::Core::Private
 	template <typename DestBufferType, typename FromType>
 	static int32 ConvertToUTF8(DestBufferType& Dest, int32 DestLen, const FromType* Source, int32 SourceLen)
 	{
+		using UnsignedFromType = std::make_unsigned_t<FromType>;
+
 		if constexpr (sizeof(FromType) == 4)
 		{
 			DestBufferType DestStartingPosition = Dest;
@@ -197,7 +205,7 @@ namespace UE::Core::Private
 					return UE_PTRDIFF_TO_INT32(Dest - DestStartingPosition);
 				}
 
-				uint32 Codepoint = static_cast<uint32>(*Source++);
+				uint32 Codepoint = (uint32)(UnsignedFromType)*Source++;
 				--SourceLen;
 
 				if (!WriteCodepointToBuffer(Codepoint, Dest, DestLen))
@@ -218,7 +226,7 @@ namespace UE::Core::Private
 					return UE_PTRDIFF_TO_INT32(Dest - DestStartingPosition);
 				}
 
-				uint32 Codepoint = static_cast<uint32>(*Source++);
+				uint32 Codepoint = (uint32)(UnsignedFromType)*Source++;
 				--SourceLen;
 
 				// Check if this character is a high-surrogate
@@ -237,7 +245,7 @@ namespace UE::Core::Private
 					}
 
 					// Read next codepoint
-					uint32 NextCodepoint = static_cast<uint32>(*Source);
+					uint32 NextCodepoint = (uint32)(UnsignedFromType)*Source;
 
 					// If it's a low surrogate, combine it with the current high surrogate,
 					// otherwise just leave the high surrogate to be written out by itself (as a bogus character)
@@ -331,7 +339,7 @@ namespace UE::Core::Private
 			uint32 Codepoint = ((Octet - 224) << 12) | ((Octet2 - 128) << 6) | (Octet3 - 128);
 
 			// UTF-8 characters cannot be in the UTF-16 surrogates range.  0xFFFE and 0xFFFF are illegal, too, so we check them at the edge.
-			if (Codepoint < 0x800 || Codepoint > 0xFFFD || IsHighSurrogate(Codepoint) || IsLowSurrogate(Codepoint))
+			if (Codepoint < 0x800 || Codepoint > 0xFFFD || IsSurrogate(Codepoint))
 			{
 				return UNICODE_BOGUS_CHAR_CODEPOINT;
 			}
@@ -484,12 +492,6 @@ namespace UE::Core::Private
 	 *
 	 * @return The length of the string in UTF-16 code units.
 	 */
-	int32 GetConvertedLength(const UTF8CHAR*, const ANSICHAR* Source, int32 SourceLen)
-	{
-		TCountingOutputIterator<UTF8CHAR> Dest;
-		int32 Result = ConvertToUTF8(Dest, INT32_MAX, Source, SourceLen);
-		return Result;
-	}
 	int32 GetConvertedLength(const UTF8CHAR*, const WIDECHAR* Source, int32 SourceLen)
 	{
 		TCountingOutputIterator<UTF8CHAR> Dest;
@@ -527,14 +529,6 @@ namespace UE::Core::Private
 		return Result;
 	}
 
-	UTF8CHAR* Convert(UTF8CHAR* Dest, int32 DestLen, const ANSICHAR* Src, int32 SrcLen)
-	{
-		if (ConvertToUTF8(Dest, DestLen, Src, SrcLen) == -1)
-		{
-			return nullptr;
-		}
-		return Dest;
-	}
 	UTF8CHAR* Convert(UTF8CHAR* Dest, int32 DestLen, const WIDECHAR* Src, int32 SrcLen)
 	{
 		if (ConvertToUTF8(Dest, DestLen, Src, SrcLen) == -1)
diff --git a/Engine/Source/Runtime/Core/Public/GenericPlatform/GenericPlatformString.h b/Engine/Source/Runtime/Core/Public/GenericPlatform/GenericPlatformString.h
--- a/Engine/Source/Runtime/Core/Public/GenericPlatform/GenericPlatformString.h
+++ b/Engine/Source/Runtime/Core/Public/GenericPlatform/GenericPlatformString.h
@@ -15,7 +15,6 @@
 namespace UE::Core::Private
 {
 	// The Dest parameter is just used for overload resolution
-	CORE_API int32 GetConvertedLength(const UTF8CHAR* Dest, const ANSICHAR*  Src, int32 SrcLen);
 	CORE_API int32 GetConvertedLength(const UTF8CHAR* Dest, const WIDECHAR*  Src, int32 SrcLen);
 	CORE_API int32 GetConvertedLength(const UTF8CHAR* Dest, const UCS2CHAR*  Src, int32 SrcLen);
 	CORE_API int32 GetConvertedLength(const UTF8CHAR* Dest, const UTF32CHAR* Src, int32 SrcLen);
@@ -23,7 +22,6 @@ namespace UE::Core::Private
 	CORE_API int32 GetConvertedLength(const WIDECHAR* Dest, const UTF8CHAR*  Src, int32 SrcLen);
 	CORE_API int32 GetConvertedLength(const UCS2CHAR* Dest, const UTF8CHAR*  Src, int32 SrcLen);
 
-	CORE_API UTF8CHAR* Convert(UTF8CHAR* Dest, int32 DestLen, const ANSICHAR*  Src, int32 SrcLen);
 	CORE_API UTF8CHAR* Convert(UTF8CHAR* Dest, int32 DestLen, const WIDECHAR*  Src, int32 SrcLen);
 	CORE_API UTF8CHAR* Convert(UTF8CHAR* Dest, int32 DestLen, const UCS2CHAR*  Src, int32 SrcLen);
 	CORE_API UTF8CHAR* Convert(UTF8CHAR* Dest, int32 DestLen, const UTF32CHAR* Src, int32 SrcLen);
