[Core] change to disable breaking in the debugger when running the "ensure" tests

diff --git a/Engine/Source/Runtime/Core/Tests/Misc/AssertionMacrosTest.cpp b/Engine/Source/Runtime/Core/Tests/Misc/AssertionMacrosTest.cpp
--- a/Engine/Source/Runtime/Core/Tests/Misc/AssertionMacrosTest.cpp
+++ b/Engine/Source/Runtime/Core/Tests/Misc/AssertionMacrosTest.cpp
@@ -22,6 +22,7 @@ static void MultipleEnsuresFailedDifferent();
  */
 TEST_CASE("Core::Misc::AssertionMacros", "[Core][Misc][AssertionMacros][Ensure]")
 {
+	TGuardValue<bool> IgnoreDebugger(GIgnoreDebugger, true);
 	GLog->SetSuppressEventTag(true);
 	SECTION("Ensure not triggered with REQUIRE_NOENSURE")
 	{
