Extract exit reasons for all processes withing the same app on Android and separate OOM crashes from Abnormal Terminations into it's own category in Crash Reporter
#[fyi] Allan.Bentham

diff --git a/Engine/Source/Runtime/Core/Private/GenericPlatform/GenericPlatformCrashContext.cpp b/Engine/Source/Runtime/Core/Private/GenericPlatform/GenericPlatformCrashContext.cpp
--- a/Engine/Source/Runtime/Core/Private/GenericPlatform/GenericPlatformCrashContext.cpp
+++ b/Engine/Source/Runtime/Core/Private/GenericPlatform/GenericPlatformCrashContext.cpp
@@ -104,6 +104,7 @@ const TCHAR* const FGenericCrashContext::CrashTypeStall = TEXT("Stall");
 const TCHAR* const FGenericCrashContext::CrashTypeGPU = TEXT("GPUCrash");
 const TCHAR* const FGenericCrashContext::CrashTypeHang = TEXT("Hang");
 const TCHAR* const FGenericCrashContext::CrashTypeAbnormalShutdown = TEXT("AbnormalShutdown");
+const TCHAR* const FGenericCrashContext::CrashTypeOutOfMemory = TEXT("OutOfMemory");
 
 const TCHAR* const FGenericCrashContext::EngineModeExUnknown = TEXT("Unset");
 const TCHAR* const FGenericCrashContext::EngineModeExDirty = TEXT("Dirty");
@@ -1046,6 +1047,8 @@ const TCHAR* FGenericCrashContext::GetCrashTypeString(ECrashContextType Type)
 		return CrashTypeAssert;
 	case ECrashContextType::AbnormalShutdown:
 		return CrashTypeAbnormalShutdown;
+	case ECrashContextType::OutOfMemory:
+		return CrashTypeOutOfMemory;
 	default:
 		return CrashTypeCrash;
 	}
diff --git a/Engine/Source/Runtime/Core/Public/GenericPlatform/GenericPlatformCrashContext.h b/Engine/Source/Runtime/Core/Public/GenericPlatform/GenericPlatformCrashContext.h
--- a/Engine/Source/Runtime/Core/Public/GenericPlatform/GenericPlatformCrashContext.h
+++ b/Engine/Source/Runtime/Core/Public/GenericPlatform/GenericPlatformCrashContext.h
@@ -294,6 +294,7 @@ public:
 	static const TCHAR* const CrashTypeGPU;
 	static const TCHAR* const CrashTypeHang;
 	static const TCHAR* const CrashTypeAbnormalShutdown;
+	static const TCHAR* const CrashTypeOutOfMemory;
 
 	static const TCHAR* const EngineModeExUnknown;
 	static const TCHAR* const EngineModeExDirty;
