[FramePro]
* Fixed so exiting doesnt wait 10 seconds for a timeout to happen when FramePro connect thread was never started

diff --git a/Engine/Source/Runtime/Core/Private/FramePro/FramePro.cpp b/Engine/Source/Runtime/Core/Private/FramePro/FramePro.cpp
--- a/Engine/Source/Runtime/Core/Private/FramePro/FramePro.cpp
+++ b/Engine/Source/Runtime/Core/Private/FramePro/FramePro.cpp
@@ -8463,7 +8463,8 @@ namespace FramePro
 	//------------------------------------------------------------------------
 	void Thread::WaitForThreadToTerminate(int timeout)
 	{
-		m_ThreadTerminatedEvent.Wait(timeout);
+		if (m_Created)
+			m_ThreadTerminatedEvent.Wait(timeout);
 	}
 }
 
