[RuntimeCore]
* Changed so FTrackedActivity can be used with shared pointers.

diff --git a/Engine/Source/Runtime/Core/Private/Misc/TrackedActivity.cpp b/Engine/Source/Runtime/Core/Private/Misc/TrackedActivity.cpp
--- a/Engine/Source/Runtime/Core/Private/Misc/TrackedActivity.cpp
+++ b/Engine/Source/Runtime/Core/Private/Misc/TrackedActivity.cpp
@@ -202,16 +202,16 @@ void FTrackedActivity::Update(ELight Light, uint32 Index)
 
 FTrackedActivity& FTrackedActivity::GetEngineActivity()
 {
-	static FTrackedActivity A(TEXT("Status"), TEXT("Unknown"));
-	return A;
+	static TSharedPtr<FTrackedActivity> A(MakeShared<FTrackedActivity>(TEXT("Status"), TEXT("Unknown")));
+	return *A;
 }
 
 FTrackedActivity& FTrackedActivity::GetIOActivity()
 {
 	GetEngineActivity(); // Ugly, but just to make sure Status show first
 
-	static FTrackedActivity A(TEXT("I/O"), TEXT("Idle"));
-	return A;
+	static TSharedPtr<FTrackedActivity> A(MakeShared<FTrackedActivity>(TEXT("I/O"), TEXT("Idle")));
+	return *A;
 }
 
 void FTrackedActivity::RegisterEventListener(TUniqueFunction<void(EEvent Event, const FInfo& Info)>&& Func, uint32 MaxDepth)
diff --git a/Engine/Source/Runtime/Core/Public/Misc/TrackedActivity.h b/Engine/Source/Runtime/Core/Public/Misc/TrackedActivity.h
--- a/Engine/Source/Runtime/Core/Public/Misc/TrackedActivity.h
+++ b/Engine/Source/Runtime/Core/Public/Misc/TrackedActivity.h
@@ -4,6 +4,7 @@
 
 #include "CoreTypes.h"
 #include "Templates/Function.h"
+#include "Templates/SharedPointer.h"
 
 /**
  * Tracked Activity is used to be able to visualize on a semi-high level what is going on in the process.
@@ -11,7 +12,7 @@
  * When new console is enabled tracked activities show at the bottom of the window under the log
  * Tracked Activities can be created/updated/destroyed on multiple threads
  */
-class CORE_API FTrackedActivity
+class CORE_API FTrackedActivity : public TSharedFromThis<FTrackedActivity>
 {
 public:
 	/**
