Enable UE_DEBUG_BREAK in shipping config if the target compiled WITH_EDITOR

diff --git a/Engine/Source/Runtime/Core/Public/GenericPlatform/GenericPlatformMisc.h b/Engine/Source/Runtime/Core/Public/GenericPlatform/GenericPlatformMisc.h
--- a/Engine/Source/Runtime/Core/Public/GenericPlatform/GenericPlatformMisc.h
+++ b/Engine/Source/Runtime/Core/Public/GenericPlatform/GenericPlatformMisc.h
@@ -29,7 +29,7 @@ enum class ECustomChunkType : uint8;
 template <typename FuncType>
 class TFunction;
 
-#if UE_BUILD_SHIPPING
+#if UE_BUILD_SHIPPING && !WITH_EDITOR
 #define UE_DEBUG_BREAK() ((void)0)
 #else
 #define UE_DEBUG_BREAK() ((void)(FPlatformMisc::IsDebuggerPresent() && ([] () { UE_DEBUG_BREAK_IMPL(); } (), 1)))
