Added FloatCastChecked/FloatFitsIn to provide runtime validation of acceptable precision loss when casting between floating point types.

diff --git a/Engine/Source/Runtime/Core/Public/Templates/UnrealTemplate.h b/Engine/Source/Runtime/Core/Public/Templates/UnrealTemplate.h
--- a/Engine/Source/Runtime/Core/Public/Templates/UnrealTemplate.h
+++ b/Engine/Source/Runtime/Core/Public/Templates/UnrealTemplate.h
@@ -180,7 +180,24 @@ constexpr bool IntFitsIn(InType In)
 template<typename OutType, typename InType>
 OutType IntCastChecked(InType In)
 {
-	check(IntFitsIn<OutType>(In));
+	checkf(IntFitsIn<OutType>(In), TEXT("Loss of data caused by narrowing conversion"));
+	return static_cast<OutType>(In);
+}
+
+/** Test if value can make a static_cast roundtrip via OutType whilst maintaining precision */
+template<typename OutType, typename InType>
+constexpr bool FloatFitsIn(InType In, InType Precision)
+{
+	static_assert(std::is_floating_point_v<InType> && std::is_floating_point_v<OutType>, "Only floating point supported");
+	
+	OutType Out = static_cast<OutType>(In);
+	return std::abs(static_cast<InType>(Out) - In) <= Precision;
+}
+
+template<typename OutType, typename InType>
+OutType FloatCastChecked(InType In, InType Precision)
+{
+	checkf(FloatFitsIn<OutType>(In, Precision), TEXT("Loss of data caused by narrowing conversion"));
 	return static_cast<OutType>(In);
 }
 
@@ -687,4 +704,4 @@ template <
 FORCEINLINE decltype(auto) ForwardAsBase(typename TRemoveReference<T>::Type&& Obj)
 {
 	return (TCopyQualifiersAndRefsFromTo_T<T&&, Base>)Obj;
-}
+}
\ No newline at end of file
