Fixed unacceptable words in Stats and Profiling system

diff --git a/Engine/Plugins/NetcodeUnitTest/NetcodeUnitTest/Source/NetcodeUnitTest/Private/NUTUtilProfiler.cpp b/Engine/Plugins/NetcodeUnitTest/NetcodeUnitTest/Source/NetcodeUnitTest/Private/NUTUtilProfiler.cpp
--- a/Engine/Plugins/NetcodeUnitTest/NetcodeUnitTest/Source/NetcodeUnitTest/Private/NUTUtilProfiler.cpp
+++ b/Engine/Plugins/NetcodeUnitTest/NetcodeUnitTest/Source/NetcodeUnitTest/Private/NUTUtilProfiler.cpp
@@ -18,7 +18,7 @@ void FFrameProfiler::Start()
 	if (!bActive)
 	{
 		bActive = true;
-		StatsMasterEnableAdd();
+		StatsPrimaryEnableAdd();
 
 		const FStatsThreadState& Stats = FStatsThreadState::GetLocalState();
 
@@ -34,7 +34,7 @@ void FFrameProfiler::Stop()
 
 		Stats.NewFrameDelegate.Remove(OnNewFrameDelegateHandle);
 
-		StatsMasterEnableSubtract();
+		StatsPrimaryEnableSubtract();
 		bActive = false;
 
 
diff --git a/Engine/Plugins/Performance/PerformanceMonitor/Source/PerformanceMonitor/Private/PerformanceMonitor.cpp b/Engine/Plugins/Performance/PerformanceMonitor/Source/PerformanceMonitor/Private/PerformanceMonitor.cpp
--- a/Engine/Plugins/Performance/PerformanceMonitor/Source/PerformanceMonitor/Private/PerformanceMonitor.cpp
+++ b/Engine/Plugins/Performance/PerformanceMonitor/Source/PerformanceMonitor/Private/PerformanceMonitor.cpp
@@ -289,7 +289,7 @@ void FPerformanceMonitorModule::StartRecordingPerfTimers(FString FileNameToUse,
 			bRequiresCutsceneStart = GatheredBool;
 		}
 	}
-	FThreadStats::MasterEnableAdd(1);
+	FThreadStats::PrimaryEnableAdd(1);
 	FAutomationTestFramework::Get().SetCaptureStack(false);
 	FStatsThreadState& Stats = FStatsThreadState::GetLocalState();
 	// Set up our delegate to gather data from the stats thread for safe consumption on game thread.
@@ -464,7 +464,7 @@ void FPerformanceMonitorModule::StopRecordingPerformanceTimers()
 		RecordData();
 	}
 
-	FThreadStats::MasterEnableSubtract(1);
+	FThreadStats::PrimaryEnableSubtract(1);
 	FStatsThreadState& Stats = FStatsThreadState::GetLocalState();
 	Stats.NewFrameDelegate.RemoveAll(this);
 	bRecording = false;
diff --git a/Engine/Plugins/VirtualProduction/StageMonitoring/Source/StageDataProvider/Private/FramePerformanceProvider.cpp b/Engine/Plugins/VirtualProduction/StageMonitoring/Source/StageDataProvider/Private/FramePerformanceProvider.cpp
--- a/Engine/Plugins/VirtualProduction/StageMonitoring/Source/StageDataProvider/Private/FramePerformanceProvider.cpp
+++ b/Engine/Plugins/VirtualProduction/StageMonitoring/Source/StageDataProvider/Private/FramePerformanceProvider.cpp
@@ -105,13 +105,13 @@ void FFramePerformanceProvider::EnableHitchDetection(bool bShouldEnable)
 			CachedHitchSettings = GetDefault<UStageMonitoringSettings>()->ProviderSettings.HitchDetectionSettings;
 
 			// Subscribe to Stats provider to verify hitches
-			StatsMasterEnableAdd();
+			StatsPrimaryEnableAdd();
 			FStatsThreadState& Stats = FStatsThreadState::GetLocalState();
 			Stats.NewFrameDelegate.AddRaw(this, &FFramePerformanceProvider::CheckHitches);
 		}
 		else
 		{
-			StatsMasterEnableSubtract();
+			StatsPrimaryEnableSubtract();
 			FStatsThreadState& Stats = FStatsThreadState::GetLocalState();
 			Stats.NewFrameDelegate.RemoveAll(this);
 		}
diff --git a/Engine/Source/Developer/ProfilerService/Private/ProfilerServiceManager.cpp b/Engine/Source/Developer/ProfilerService/Private/ProfilerServiceManager.cpp
--- a/Engine/Source/Developer/ProfilerService/Private/ProfilerServiceManager.cpp
+++ b/Engine/Source/Developer/ProfilerService/Private/ProfilerServiceManager.cpp
@@ -131,7 +131,7 @@ void FProfilerServiceManager::AddNewFrameHandleStatsPipe()
 #if	STATS
 	const FStatsThreadState& Stats = FStatsThreadState::GetLocalState();
 	NewFrameDelegateHandle = Stats.NewFrameDelegate.AddRaw( this, &FProfilerServiceManager::HandleNewFrame );
-	StatsMasterEnableAdd();
+	StatsPrimaryEnableAdd();
 	MetadataSize = 0;
 #endif //STATS
 }
@@ -142,7 +142,7 @@ void FProfilerServiceManager::RemoveNewFrameHandleStatsPipe()
 #if	STATS
 	const FStatsThreadState& Stats = FStatsThreadState::GetLocalState();
 	Stats.NewFrameDelegate.Remove( NewFrameDelegateHandle );
-	StatsMasterEnableSubtract();
+	StatsPrimaryEnableSubtract();
 	MetadataSize = 0;
 #endif //STATS
 }
diff --git a/Engine/Source/Developer/RealtimeProfiler/Private/RealtimeProfiler.cpp b/Engine/Source/Developer/RealtimeProfiler/Private/RealtimeProfiler.cpp
--- a/Engine/Source/Developer/RealtimeProfiler/Private/RealtimeProfiler.cpp
+++ b/Engine/Source/Developer/RealtimeProfiler/Private/RealtimeProfiler.cpp
@@ -21,12 +21,12 @@ void FRealtimeProfiler::MakeWindow()
 FRealtimeProfiler::FRealtimeProfiler()
 {
 	MakeWindow();	
-	StatsMasterEnableAdd();
+	StatsPrimaryEnableAdd();
 }
 
 FRealtimeProfiler::~FRealtimeProfiler()
 {
-	StatsMasterEnableSubtract();
+	StatsPrimaryEnableSubtract();
 }
 
 bool FRealtimeProfiler::IsProfiling()
diff --git a/Engine/Source/Runtime/Core/Private/Stats/Stats2.cpp b/Engine/Source/Runtime/Core/Private/Stats/Stats2.cpp
--- a/Engine/Source/Runtime/Core/Private/Stats/Stats2.cpp
+++ b/Engine/Source/Runtime/Core/Private/Stats/Stats2.cpp
@@ -197,18 +197,18 @@ void FStats::AdvanceFrame( bool bDiscardCallstack, const FOnAdvanceRenderingThre
 	TRACE_CPUPROFILER_EVENT_SCOPE(FStats::AdvanceFrame);
 	LLM_SCOPE(ELLMTag::Stats);
 	check( IsInGameThread() );
-	static int32 MasterDisableChangeTagStartFrame = -1;
+	static int32 PrimaryDisableChangeTagStartFrame = -1;
 	int64 Frame = ++GameThreadStatsFrame;
 
 	if( bDiscardCallstack )
 	{
 		FThreadStats::FrameDataIsIncomplete(); // we won't collect call stack stats this frame
 	}
-	if( MasterDisableChangeTagStartFrame == -1 )
+	if( PrimaryDisableChangeTagStartFrame == -1 )
 	{
-		MasterDisableChangeTagStartFrame = FThreadStats::PrimaryDisableChangeTag();
+		PrimaryDisableChangeTagStartFrame = FThreadStats::PrimaryDisableChangeTag();
 	}
-	if( !FThreadStats::IsCollectingData() || MasterDisableChangeTagStartFrame != FThreadStats::PrimaryDisableChangeTag() )
+	if( !FThreadStats::IsCollectingData() || PrimaryDisableChangeTagStartFrame != FThreadStats::PrimaryDisableChangeTag() )
 	{
 		Frame = -Frame; // mark this as a bad frame
 	}
@@ -220,7 +220,7 @@ void FStats::AdvanceFrame( bool bDiscardCallstack, const FOnAdvanceRenderingThre
 
 	if( AdvanceRenderingThreadStatsDelegate.IsBound() )
 	{
-		AdvanceRenderingThreadStatsDelegate.Execute( bDiscardCallstack, Frame, MasterDisableChangeTagStartFrame );
+		AdvanceRenderingThreadStatsDelegate.Execute( bDiscardCallstack, Frame, PrimaryDisableChangeTagStartFrame );
 	}
 	else
 	{
@@ -230,7 +230,7 @@ void FStats::AdvanceFrame( bool bDiscardCallstack, const FOnAdvanceRenderingThre
 
 	FThreadStats::ExplicitFlush( bDiscardCallstack );
 	FThreadStats::WaitForStats();
-	MasterDisableChangeTagStartFrame = FThreadStats::PrimaryDisableChangeTag();
+	PrimaryDisableChangeTagStartFrame = FThreadStats::PrimaryDisableChangeTag();
 #endif
 }
 
@@ -439,7 +439,7 @@ public:
 	virtual void SetHighPerformanceEnableForGroup(FName Group, bool Enable) override
 	{
 		FScopeLock ScopeLock(&SynchronizationObject);
-		FThreadStats::MasterDisableChangeTagLockAdd();
+		FThreadStats::PrimaryDisableChangeTagLockAdd();
 		FGroupEnable* Found = HighPerformanceEnable.Find(Group);
 		if (Found)
 		{
@@ -459,13 +459,13 @@ public:
 				}
 			}
 		}
-		FThreadStats::MasterDisableChangeTagLockSubtract();
+		FThreadStats::PrimaryDisableChangeTagLockSubtract();
 	}
 
 	virtual void SetHighPerformanceEnableForAllGroups(bool Enable) override
 	{
 		FScopeLock ScopeLock(&SynchronizationObject);
-		FThreadStats::MasterDisableChangeTagLockAdd();
+		FThreadStats::PrimaryDisableChangeTagLockAdd();
 		for (auto It = HighPerformanceEnable.CreateIterator(); It; ++It)
 		{
 			It.Value().CurrentEnable = Enable;
@@ -484,12 +484,12 @@ public:
 				}
 			}
 		}
-		FThreadStats::MasterDisableChangeTagLockSubtract();
+		FThreadStats::PrimaryDisableChangeTagLockSubtract();
 	}
 	virtual void ResetHighPerformanceEnableForAllGroups() override
 	{
 		FScopeLock ScopeLock(&SynchronizationObject);
-		FThreadStats::MasterDisableChangeTagLockAdd();
+		FThreadStats::PrimaryDisableChangeTagLockAdd();
 		for (auto It = HighPerformanceEnable.CreateIterator(); It; ++It)
 		{
 			It.Value().CurrentEnable = It.Value().DefaultEnable;
@@ -508,7 +508,7 @@ public:
 				}
 			}
 		}
-		FThreadStats::MasterDisableChangeTagLockSubtract();
+		FThreadStats::PrimaryDisableChangeTagLockSubtract();
 	}
 
 	virtual TStatId GetHighPerformanceEnableForStat(FName StatShortName, const char* InGroup, const char* InCategory, bool bDefaultEnable, bool bShouldClearEveryFrame, EStatDataType::Type InStatType, TCHAR const* InDescription, bool bCycleStat, bool bSortByName, FPlatformMemory::EMemoryCounterRegion MemoryRegion = FPlatformMemory::MCR_Invalid) override
@@ -1244,11 +1244,11 @@ void FThreadStatsPool::ReturnToPool( FThreadStats* Instance )
 -----------------------------------------------------------------------------*/
 
 uint32 FThreadStats::TlsSlot = 0;
-FThreadSafeCounter FThreadStats::MasterEnableCounter;
-FThreadSafeCounter FThreadStats::MasterEnableUpdateNumber;
-FThreadSafeCounter FThreadStats::MasterDisableChangeTagLock;
-bool FThreadStats::bMasterEnable = false;
-bool FThreadStats::bMasterDisableForever = false;
+FThreadSafeCounter FThreadStats::PrimaryEnableCounter;
+FThreadSafeCounter FThreadStats::PrimaryEnableUpdateNumber;
+FThreadSafeCounter FThreadStats::PrimaryDisableChangeTagLock;
+bool FThreadStats::bPrimaryEnable = false;
+bool FThreadStats::bPrimaryDisableForever = false;
 bool FThreadStats::bIsRawStatsActive = false;
 
 FThreadStats::FThreadStats():
@@ -1282,19 +1282,19 @@ FThreadStats::FThreadStats( EConstructor ):
 
 void FThreadStats::CheckEnable()
 {
-	bool bOldMasterEnable(bMasterEnable);
-	bool bNewMasterEnable( WillEverCollectData() && (!IsRunningCommandlet() || FStats::EnabledForCommandlet()) && IsThreadingReady() && (MasterEnableCounter.GetValue()) );
-	if (bMasterEnable != bNewMasterEnable)
+	bool bOldPrimaryEnable(bPrimaryEnable);
+	bool bNewPrimaryEnable( WillEverCollectData() && (!IsRunningCommandlet() || FStats::EnabledForCommandlet()) && IsThreadingReady() && (PrimaryEnableCounter.GetValue()) );
+	if (bPrimaryEnable != bNewPrimaryEnable)
 	{
-		MasterDisableChangeTagLockAdd();
-		bMasterEnable = bNewMasterEnable;
-		MasterDisableChangeTagLockSubtract();
+		PrimaryDisableChangeTagLockAdd();
+		bPrimaryEnable = bNewPrimaryEnable;
+		PrimaryDisableChangeTagLockSubtract();
 	}
 }
 
 void FThreadStats::Flush( bool bHasBrokenCallstacks /*= false*/, bool bForceFlush /*= false*/ )
 {
-	if (bMasterDisableForever)
+	if (bPrimaryDisableForever)
 	{
 		Packet.StatMessages.Empty();
 		return;
@@ -1595,7 +1595,7 @@ void FThreadStats::StopThread()
 		// If we are writing stats data, stop it now.
 		DirectStatsCommand( TEXT( "stat stopfile" ), true );
 
-		FThreadStats::MasterDisableForever();
+		FThreadStats::PrimaryDisableForever();
 
 		WaitForStats();
 	}
@@ -1630,7 +1630,7 @@ void FThreadStats::WaitForStats()
 	}
 
 	check(IsInGameThread());
-	if (IsThreadingReady() && !bMasterDisableForever)
+	if (IsThreadingReady() && !bPrimaryDisableForever)
 	{
 		DECLARE_CYCLE_STAT(TEXT("WaitForStats"), STAT_WaitForStats, STATGROUP_Engine);
 
diff --git a/Engine/Source/Runtime/Core/Private/Stats/StatsCommand.cpp b/Engine/Source/Runtime/Core/Private/Stats/StatsCommand.cpp
--- a/Engine/Source/Runtime/Core/Private/Stats/StatsCommand.cpp
+++ b/Engine/Source/Runtime/Core/Private/Stats/StatsCommand.cpp
@@ -494,7 +494,7 @@ void DumpEventsOnce( int64 TargetFrame, float DumpEventsCullMS, bool bDisplayAll
 {
 	FStatsThreadState& Stats = FStatsThreadState::GetLocalState();
 	DumpEvents( TargetFrame, DumpEventsCullMS, bDisplayAllThreads );
-	StatsMasterEnableSubtract();
+	StatsPrimaryEnableSubtract();
 	Stats.NewFrameDelegate.Remove( DumpEventsDelegateHandle );
 }
 
@@ -1055,12 +1055,12 @@ struct FHUDGroupManager
 		{
 			bEnabled = true;
 			NewFrameDelegateHandle = Stats.NewFrameDelegate.AddRaw( this, &FHUDGroupManager::NewFrame );
-			StatsMasterEnableAdd();
+			StatsPrimaryEnableAdd();
 		}
 		else if( !EnabledGroups.Num() && bEnabled )
 		{
 			Stats.NewFrameDelegate.Remove( NewFrameDelegateHandle );
-			StatsMasterEnableSubtract();
+			StatsPrimaryEnableSubtract();
 			bEnabled = false;
 
 			DECLARE_CYCLE_STAT(TEXT("FSimpleDelegateGraphTask.StatsToGame"),
@@ -1674,7 +1674,7 @@ static void DumpFrame(int64 Frame)
 	check(Latest > 0);
 	DumpHistoryFrame(Stats, Latest, DumpCull, GMaxDepth, *GNameFilter);
 	Stats.NewFrameDelegate.Remove(GDumpFrameDelegateHandle);
-	StatsMasterEnableSubtract();
+	StatsPrimaryEnableSubtract();
 }
 
 static void DumpCPU(int64 Frame)
@@ -1684,7 +1684,7 @@ static void DumpCPU(int64 Frame)
 	check(Latest > 0);
 	DumpCPUSummary(Stats, Latest);
 	Stats.NewFrameDelegate.Remove(GDumpCPUDelegateHandle);
-	StatsMasterEnableSubtract();
+	StatsPrimaryEnableSubtract();
 }
 
 static struct FDumpMultiple* DumpMultiple = NULL;
@@ -1707,7 +1707,7 @@ struct FDumpMultiple
 		, NumFramesToGo(0)
 		, Stack(NULL)
 	{
-		StatsMasterEnableAdd();
+		StatsPrimaryEnableAdd();
 		NewFrameDelegateHandle = Stats.NewFrameDelegate.AddRaw(this, &FDumpMultiple::NewFrame);
 	}
 
@@ -1753,7 +1753,7 @@ struct FDumpMultiple
 			Stack = NULL;
 		}
 		Stats.NewFrameDelegate.Remove(NewFrameDelegateHandle);
-		StatsMasterEnableSubtract();
+		StatsPrimaryEnableSubtract();
 		DumpMultiple = NULL;
 	}
 
@@ -1801,14 +1801,14 @@ struct FDumpSpam
 		, NumPackets(0)
 	{
 		FThreadStats::EnableRawStats();
-		StatsMasterEnableAdd();
+		StatsPrimaryEnableAdd();
 		NewRawStatPacketDelegateHandle = Stats.NewRawStatPacket.AddRaw(this, &FDumpSpam::NewFrame);
 	}
 
 	~FDumpSpam()
 	{
 		FThreadStats::DisableRawStats();
-		StatsMasterEnableSubtract();
+		StatsPrimaryEnableSubtract();
 		UE_LOG(LogStats, Log, TEXT("------------------ %d packets, %d total messages ---------------"), NumPackets, TotalCount);
 
 		Counts.ValueSort(TGreater<int32>());
@@ -1912,7 +1912,7 @@ static void StatCmd(FString InCmd, bool bStatCommand, FOutputDevice* Ar /*= null
 		FParse::Value(Cmd, TEXT("DEPTH="), GMaxDepth);
 		if (FParse::Command(&Cmd, TEXT("DUMPFRAME")))
 		{
-			StatsMasterEnableAdd();
+			StatsPrimaryEnableAdd();
 			GDumpFrameDelegateHandle = Stats.NewFrameDelegate.AddStatic(&DumpFrame);
 		}
 		else if (FParse::Command(&Cmd, TEXT("DUMPNONFRAME")))
@@ -1924,7 +1924,7 @@ static void StatCmd(FString InCmd, bool bStatCommand, FOutputDevice* Ar /*= null
 		}
 		else if (FParse::Command(&Cmd, TEXT("DUMPCPU")))
 		{
-			StatsMasterEnableAdd();
+			StatsPrimaryEnableAdd();
 			GDumpCPUDelegateHandle = Stats.NewFrameDelegate.AddStatic(&DumpCPU);
 		}
 		else if (FParse::Command(&Cmd, TEXT("STOP")))
@@ -2000,14 +2000,14 @@ static void StatCmd(FString InCmd, bool bStatCommand, FOutputDevice* Ar /*= null
 			bToggle = !bToggle;
 			if (bToggle)
 			{
-				StatsMasterEnableAdd();
+				StatsPrimaryEnableAdd();
 				HitchIndex = 0;
 				TotalHitchTime = 0.0;
 				DumpHitchDelegateHandle = Stats.NewFrameDelegate.AddStatic(&DumpHitch);
 			}
 			else
 			{
-				StatsMasterEnableSubtract();
+				StatsPrimaryEnableSubtract();
 				Stats.NewFrameDelegate.Remove(DumpHitchDelegateHandle);
 				UE_LOG(LogStats, Log, TEXT("**************************** %d hitches	%8.0fms total hitch time"), HitchIndex, TotalHitchTime);
 			}
@@ -2022,7 +2022,7 @@ static void StatCmd(FString InCmd, bool bStatCommand, FOutputDevice* Ar /*= null
 			FParse::Value(Cmd, TEXT("MS="), DumpEventsCullMS);
 			const bool bDisplayAllThreads = FParse::Param(Cmd, TEXT("all"));
 
-			StatsMasterEnableAdd();
+			StatsPrimaryEnableAdd();
 			DumpEventsDelegateHandle = Stats.NewFrameDelegate.AddStatic(&DumpEventsOnce, DumpEventsCullMS, bDisplayAllThreads);
 		}
 		else if (FParse::Command(&Cmd, TEXT("STARTFILE")))
@@ -2070,7 +2070,7 @@ static void StatCmd(FString InCmd, bool bStatCommand, FOutputDevice* Ar /*= null
 		}
 		else if (FParse::Command(&Cmd, TEXT("testdisable")))
 		{
-			FThreadStats::MasterDisableForever();
+			FThreadStats::PrimaryDisableForever();
 		}
 		else if (FParse::Command(&Cmd, TEXT("none")))
 		{
diff --git a/Engine/Source/Runtime/Core/Private/Stats/StatsFile.cpp b/Engine/Source/Runtime/Core/Private/Stats/StatsFile.cpp
--- a/Engine/Source/Runtime/Core/Private/Stats/StatsFile.cpp
+++ b/Engine/Source/Runtime/Core/Private/Stats/StatsFile.cpp
@@ -246,7 +246,7 @@ void IStatsWriteFile::Start( const FString& InFilename )
 		ArchiveFilename = Filename;
 		WriteHeader();
 		SetDataDelegate( true );
-		StatsMasterEnableAdd();
+		StatsPrimaryEnableAdd();
 	}
 }
 
@@ -255,7 +255,7 @@ void IStatsWriteFile::Stop()
 {
 	if (IsValid())
 	{
-		StatsMasterEnableSubtract();
+		StatsPrimaryEnableSubtract();
 		SetDataDelegate( false );
 		SendTask();
 		SendTask();
diff --git a/Engine/Source/Runtime/Core/Public/Stats/Stats.h b/Engine/Source/Runtime/Core/Public/Stats/Stats.h
--- a/Engine/Source/Runtime/Core/Public/Stats/Stats.h
+++ b/Engine/Source/Runtime/Core/Public/Stats/Stats.h
@@ -78,13 +78,25 @@ public:
 
 };
 
+FORCEINLINE void StatsPrimaryEnableAdd(int32 Value = 1)
+{
+	FThreadStats::PrimaryEnableAdd(Value);
+}
+FORCEINLINE void StatsPrimaryEnableSubtract(int32 Value = 1)
+{
+	FThreadStats::PrimaryEnableSubtract(Value);
+}
+
+UE_DEPRECATED(5.1, "Use StatsPrimaryEnableAdd instead")
 FORCEINLINE void StatsMasterEnableAdd(int32 Value = 1)
 {
-	FThreadStats::MasterEnableAdd(Value);
+	StatsPrimaryEnableAdd(Value);
 }
+
+UE_DEPRECATED(5.1, "Use StatsPrimaryEnableSubtract instead")
 FORCEINLINE void StatsMasterEnableSubtract(int32 Value = 1)
 {
-	FThreadStats::MasterEnableSubtract(Value);
+	StatsPrimaryEnableSubtract(Value);
 }
 
 #else	//STATS
@@ -208,10 +220,10 @@ public:
 };
 #endif
 
-FORCEINLINE void StatsMasterEnableAdd(int32 Value = 1)
+FORCEINLINE void StatsPrimaryEnableAdd(int32 Value = 1)
 {
 }
-FORCEINLINE void StatsMasterEnableSubtract(int32 Value = 1)
+FORCEINLINE void StatsPrimaryEnableSubtract(int32 Value = 1)
 {
 }
 
diff --git a/Engine/Source/Runtime/Core/Public/Stats/Stats2.h b/Engine/Source/Runtime/Core/Public/Stats/Stats2.h
--- a/Engine/Source/Runtime/Core/Public/Stats/Stats2.h
+++ b/Engine/Source/Runtime/Core/Public/Stats/Stats2.h
@@ -101,7 +101,7 @@ CORE_API bool DirectStatsCommand(const TCHAR* Cmd, bool bBlockForCompletion = fa
 struct CORE_API FStats
 {
 	/** Delegate to fire every time we need to advance the stats for the rendering thread. */
-	DECLARE_DELEGATE_ThreeParams( FOnAdvanceRenderingThreadStats, bool /*bDiscardCallstack*/, int64 /*StatsFrame*/, int32 /*MasterDisableChangeTagStartFrame*/ );
+	DECLARE_DELEGATE_ThreeParams( FOnAdvanceRenderingThreadStats, bool /*bDiscardCallstack*/, int64 /*StatsFrame*/, int32 /*PrimaryDisableChangeTagStartFrame*/ );
 
 	/** Advances stats for the current frame. */
 	static void AdvanceFrame( bool bDiscardCallstack, const FOnAdvanceRenderingThreadStats& AdvanceRenderingThreadStatsDelegate = FOnAdvanceRenderingThreadStats() );
@@ -1252,17 +1252,17 @@ class FThreadStats : FNoncopyable
 	friend struct FThreadStatsPool;
 
 	/** Used to control when we are collecting stats. User of the stats system increment and decrement this counter as they need data. **/
-	CORE_API static FThreadSafeCounter MasterEnableCounter;
-	/** Every time bMasterEnable changes, we update this. This is used to determine frames that have complete data. **/
-	CORE_API static FThreadSafeCounter MasterEnableUpdateNumber;
-	/** while bMasterEnable (or other things affecting stat collection) is chaning, we lock this. This is used to determine frames that have complete data. **/
-	CORE_API static FThreadSafeCounter MasterDisableChangeTagLock;
+	CORE_API static FThreadSafeCounter PrimaryEnableCounter;
+	/** Every time bPrimaryEnable changes, we update this. This is used to determine frames that have complete data. **/
+	CORE_API static FThreadSafeCounter PrimaryEnableUpdateNumber;
+	/** while bPrimaryEnable (or other things affecting stat collection) is chaning, we lock this. This is used to determine frames that have complete data. **/
+	CORE_API static FThreadSafeCounter PrimaryDisableChangeTagLock;
 	/** TLS slot that holds a FThreadStats. **/
 	CORE_API static uint32 TlsSlot;
-	/** Computed by CheckEnable, the current "master control" for stats collection, based on MasterEnableCounter and a few other things. **/
-	CORE_API static bool bMasterEnable;
+	/** Computed by CheckEnable, the current "primary control" for stats collection, based on PrimaryEnableCounter and a few other things. **/
+	CORE_API static bool bPrimaryEnable;
 	/** Set to permanently disable the stats system. **/
-	CORE_API static bool bMasterDisableForever;
+	CORE_API static bool bPrimaryDisableForever;
 	/** True if we running in the raw stats mode, all stats processing is disabled, captured stats messages are written in timely manner, memory overhead is minimal. */
 	CORE_API static bool bIsRawStatsActive;
 
@@ -1456,7 +1456,7 @@ public:
 	/** Return true if we are currently collecting data **/
 	static FORCEINLINE_STATS bool IsCollectingData()
 	{
-		return bMasterEnable;
+		return bPrimaryEnable;
 	}
 	static FORCEINLINE_STATS bool IsCollectingData(TStatId StatId)
 	{
@@ -1467,7 +1467,7 @@ public:
 	/** Return true if we are currently collecting data **/
 	static FORCEINLINE_STATS bool WillEverCollectData()
 	{
-		return !bMasterDisableForever;
+		return !bPrimaryDisableForever;
 	}
 
 	/** Return true if the threading is ready **/
@@ -1477,59 +1477,94 @@ public:
 	}
 
 	/** Indicate that you would like the system to begin collecting data, if it isn't already collecting data. Think reference count. **/
-	static FORCEINLINE_STATS void MasterEnableAdd(int32 Value = 1)
+	static FORCEINLINE_STATS void PrimaryEnableAdd(int32 Value = 1)
 	{
-		MasterEnableCounter.Add(Value);
+		PrimaryEnableCounter.Add(Value);
 		CheckEnable();
 	}
 
 	/** Indicate that you no longer need stat data, if nobody else needs stat data, then no stat data will be collected. Think reference count. **/
-	static FORCEINLINE_STATS void MasterEnableSubtract(int32 Value = 1)
+	static FORCEINLINE_STATS void PrimaryEnableSubtract(int32 Value = 1)
 	{
-		MasterEnableCounter.Subtract(Value);
+		PrimaryEnableCounter.Subtract(Value);
 		CheckEnable();
 	}
 
 	/** Indicate that you no longer need stat data, forever. **/
-	static FORCEINLINE_STATS void MasterDisableForever()
+	static FORCEINLINE_STATS void PrimaryDisableForever()
 	{
-		bMasterDisableForever = true;
+		bPrimaryDisableForever = true;
 		CheckEnable();
 	}
 
 	/** This is called before we start to change something that will invalidate. **/
-	static FORCEINLINE_STATS void MasterDisableChangeTagLockAdd(int32 Value = 1)
+	static FORCEINLINE_STATS void PrimaryDisableChangeTagLockAdd(int32 Value = 1)
 	{
-		MasterDisableChangeTagLock.Add(Value);
+		PrimaryDisableChangeTagLock.Add(Value);
 		FPlatformMisc::MemoryBarrier();
-		MasterEnableUpdateNumber.Increment();
+		PrimaryEnableUpdateNumber.Increment();
 	}
 
 	/** Indicate that you no longer need stat data, if nobody else needs stat data, then no stat data will be collected. Think reference count. **/
-	static FORCEINLINE_STATS void MasterDisableChangeTagLockSubtract(int32 Value = 1)
+	static FORCEINLINE_STATS void PrimaryDisableChangeTagLockSubtract(int32 Value = 1)
 	{
 		FPlatformMisc::MemoryBarrier();
-		MasterEnableUpdateNumber.Increment();
+		PrimaryEnableUpdateNumber.Increment();
 		FPlatformMisc::MemoryBarrier();
-		MasterDisableChangeTagLock.Subtract(Value);
+		PrimaryDisableChangeTagLock.Subtract(Value);
 	}
 
-	/** Everytime master enable changes, this number increases. This is used to determine full frames. **/
+	/** Everytime primary enable changes, this number increases. This is used to determine full frames. **/
 	static FORCEINLINE_STATS int32 PrimaryDisableChangeTag()
 	{
-		if (MasterDisableChangeTagLock.GetValue())
+		if (PrimaryDisableChangeTagLock.GetValue())
 		{
 			// while locked we are continually invalid, so we will just keep giving unique numbers
-			return MasterEnableUpdateNumber.Increment();
+			return PrimaryEnableUpdateNumber.Increment();
 		}
-		return MasterEnableUpdateNumber.GetValue();
+		return PrimaryEnableUpdateNumber.GetValue();
+	}
+
+	/** Indicate that you would like the system to begin collecting data, if it isn't already collecting data. Think reference count. **/
+	UE_DEPRECATED(5.1, "Use PrimaryEnableAdd instead")
+	static FORCEINLINE_STATS void MasterEnableAdd(int32 Value = 1)
+	{
+		PrimaryEnableAdd(Value);
+	}
+
+	/** Indicate that you no longer need stat data, if nobody else needs stat data, then no stat data will be collected. Think reference count. **/
+	UE_DEPRECATED(5.1, "Use PrimaryEnableSubtract instead")
+	static FORCEINLINE_STATS void MasterEnableSubtract(int32 Value = 1)
+	{
+		PrimaryEnableSubtract(Value);
+	}
+
+	/** Indicate that you no longer need stat data, forever. **/
+	UE_DEPRECATED(5.1, "Use PrimaryDisableForever instead")
+	static FORCEINLINE_STATS void MasterDisableForever()
+	{
+		PrimaryDisableForever();
+	}
+
+	/** This is called before we start to change something that will invalidate. **/
+	UE_DEPRECATED(5.1, "Use PrimaryDisableChangeTagLockAdd instead")
+	static FORCEINLINE_STATS void MasterDisableChangeTagLockAdd(int32 Value = 1)
+	{
+		PrimaryDisableChangeTagLockAdd(Value);
+	}
+
+	/** Indicate that you no longer need stat data, if nobody else needs stat data, then no stat data will be collected. Think reference count. **/
+	UE_DEPRECATED(5.1, "Use PrimaryDisableChangeTagLockSubtract instead")
+	static FORCEINLINE_STATS void MasterDisableChangeTagLockSubtract(int32 Value = 1)
+	{
+		PrimaryDisableChangeTagLockSubtract(Value);
 	}
 
 	/** Call this if something disrupts data gathering. For example when the render thread is killed, data is abandoned.**/
 	static FORCEINLINE_STATS void FrameDataIsIncomplete()
 	{
 		FPlatformMisc::MemoryBarrier();
-		MasterEnableUpdateNumber.Increment();
+		PrimaryEnableUpdateNumber.Increment();
 		FPlatformMisc::MemoryBarrier();
 	}
 
diff --git a/Engine/Source/Runtime/Engine/Private/UnrealEngine.cpp b/Engine/Source/Runtime/Engine/Private/UnrealEngine.cpp
--- a/Engine/Source/Runtime/Engine/Private/UnrealEngine.cpp
+++ b/Engine/Source/Runtime/Engine/Private/UnrealEngine.cpp
@@ -17116,7 +17116,7 @@ bool UEngine::ToggleStatNamedEvents(UWorld* World, FCommonViewportClient* Viewpo
 	{
 		if (GCycleStatsShouldEmitNamedEvents == 0)
 		{
-			StatsMasterEnableAdd();
+			StatsPrimaryEnableAdd();
 		}
 		GCycleStatsShouldEmitNamedEvents++;
 	}
@@ -17125,7 +17125,7 @@ bool UEngine::ToggleStatNamedEvents(UWorld* World, FCommonViewportClient* Viewpo
 	{
 		if (GCycleStatsShouldEmitNamedEvents == 1)
 		{
-			StatsMasterEnableSubtract();
+			StatsPrimaryEnableSubtract();
 		}
 		GCycleStatsShouldEmitNamedEvents = FMath::Max(0, GCycleStatsShouldEmitNamedEvents - 1);
 	}
diff --git a/Engine/Source/Runtime/Launch/Private/LaunchEngineLoop.cpp b/Engine/Source/Runtime/Launch/Private/LaunchEngineLoop.cpp
--- a/Engine/Source/Runtime/Launch/Private/LaunchEngineLoop.cpp
+++ b/Engine/Source/Runtime/Launch/Private/LaunchEngineLoop.cpp
@@ -1964,7 +1964,7 @@ int32 FEngineLoop::PreInitPreStartupScreen(const TCHAR* CmdLine)
 		// Leave the stats enabled.
 		if (!FStats::EnabledForCommandlet())
 		{
-			FThreadStats::MasterDisableForever();
+			FThreadStats::PrimaryDisableForever();
 		}
 #endif
 
