Implement CA_ASSUME define for clang static analysis

#rnx

diff --git a/Engine/Source/Runtime/Core/Public/Clang/ClangPlatformCodeAnalysis.h b/Engine/Source/Runtime/Core/Public/Clang/ClangPlatformCodeAnalysis.h
--- a/Engine/Source/Runtime/Core/Public/Clang/ClangPlatformCodeAnalysis.h
+++ b/Engine/Source/Runtime/Core/Public/Clang/ClangPlatformCodeAnalysis.h
@@ -36,6 +36,10 @@
 
 #if USING_CODE_ANALYSIS
 
+	// A fake function marked with noreturn that acts as a marker for CA_ASSUME to ensure the
+	// static analyzer doesn't take an analysis path that is assumed not to be navigable.
+	void CA_AssumeNoReturn() __attribute__((analyzer_noreturn));
+
 	// Input argument
 	// Example:  void SetValue( CA_IN bool bReadable );
 	#define CA_IN
@@ -69,8 +73,7 @@
 	#define CA_SUPPRESS( WarningNumber )
 
 	// Tells the code analysis engine to assume the statement to be true.  Useful for suppressing false positive warnings.
-	// NOTE: We use a double operator not here to avoid issues with passing certain class objects directly into __analysis_assume (which may cause a bogus compiler warning)
-	#define CA_ASSUME( Expr )
+	#define CA_ASSUME( Expr )  (__builtin_expect(!bool(Expr), 0) ? CA_AssumeNoReturn() : (void)0)
 
 	// Does a simple 'if (Condition)', but disables warnings about using constants in the condition.  Helps with some macro expansions.
 	#define CA_CONSTANT_IF(Condition) if (Condition)
