Remove duplicate include
#rnx

diff --git a/Engine/Source/Runtime/Core/Private/Internationalization/TextHistory.h b/Engine/Source/Runtime/Core/Private/Internationalization/TextHistory.h
--- a/Engine/Source/Runtime/Core/Private/Internationalization/TextHistory.h
+++ b/Engine/Source/Runtime/Core/Private/Internationalization/TextHistory.h
@@ -10,7 +10,6 @@
 #include "Internationalization/StringTableCoreFwd.h"
 #include "Internationalization/Text.h"
 #include "Internationalization/TextKey.h"
-#include "Internationalization/Text.h"
 #include "Misc/DateTime.h"
 #include "Serialization/StructuredArchive.h"
 #include "Templates/SharedPointerInternals.h"
