Logging: Avoid allocations in FAsyncWriter::SerializeBufferToArchive on platforms that use FPlatformMallocCrash

diff --git a/Engine/Source/Runtime/Core/Private/Misc/OutputDeviceFile.cpp b/Engine/Source/Runtime/Core/Private/Misc/OutputDeviceFile.cpp
--- a/Engine/Source/Runtime/Core/Private/Misc/OutputDeviceFile.cpp
+++ b/Engine/Source/Runtime/Core/Private/Misc/OutputDeviceFile.cpp
@@ -82,11 +82,11 @@ void FAsyncWriter::FlushArchiveAndResetTimer()
 /** [WRITER THREAD] Serialize the contents of the ring buffer to disk */
 void FAsyncWriter::SerializeBufferToArchive()
 {
-	QUICK_SCOPE_CYCLE_COUNTER(STAT_FAsyncWriter_SerializeBufferToArchive);
 	// Some platforms use FPlatformMallocCrash during a crash, which means this function is not allowed to perform any allocations
-	// or else it will deadlock when flushing logs during crash handling. Ideally scoped named events would be disabled while crashing.
-	// GIsCriticalError is not always true when crashing (i.e. the case of a GPF) so there is no way to know to skip this behavior only when crashing
+	// or else it will deadlock when flushing logs during crash handling. Ideally scoped counters and events would be disabled while crashing.
+	// GIsCriticalError is not always true when crashing (i.e. the case of a GPF) so there is no way to know to skip this behavior only when crashing.
 #if PLATFORM_ALLOW_ALLOCATIONS_IN_FASYNCWRITER_SERIALIZEBUFFERTOARCHIVE
+	QUICK_SCOPE_CYCLE_COUNTER(STAT_FAsyncWriter_SerializeBufferToArchive);
 	SCOPED_NAMED_EVENT(FAsyncWriter_SerializeBufferToArchive, FColor::Cyan);
 #endif
 	while (SerializeRequestCounter.GetValue() > 0)
