Fix for KeyChain Automated Test. It was not clear how this had not failed in past.

diff --git a/Engine/Source/Runtime/Core/Public/Misc/KeyChainUtilities.h b/Engine/Source/Runtime/Core/Public/Misc/KeyChainUtilities.h
--- a/Engine/Source/Runtime/Core/Public/Misc/KeyChainUtilities.h
+++ b/Engine/Source/Runtime/Core/Public/Misc/KeyChainUtilities.h
@@ -53,6 +53,7 @@ public:
 		
 		Other.SetSigningKey(InvalidRSAKeyHandle);
 		Other.SetPrincipalEncryptionKey(nullptr);
+		Other.SetEncryptionKeys(TMap<FGuid, FNamedAESKey>());
 	}
 
 	FKeyChain& operator=(const FKeyChain& Other)
@@ -88,6 +89,7 @@ public:
 
 		Other.SetSigningKey(InvalidRSAKeyHandle);
 		Other.SetPrincipalEncryptionKey(nullptr);
+		Other.SetEncryptionKeys(TMap<FGuid, FNamedAESKey>());
 
 		return *this;
 	}
