Fix const in FAutomationTestBase::TestEqual function signatures

diff --git a/Engine/Source/Runtime/Core/Private/Misc/AutomationTest.cpp b/Engine/Source/Runtime/Core/Private/Misc/AutomationTest.cpp
--- a/Engine/Source/Runtime/Core/Private/Misc/AutomationTest.cpp
+++ b/Engine/Source/Runtime/Core/Private/Misc/AutomationTest.cpp
@@ -1409,7 +1409,7 @@ bool FAutomationTestBase::TestEqual(const TCHAR* What, const FColor Actual, cons
 	return true;
 }
 
-bool FAutomationTestBase::TestEqual(const TCHAR* What, FLinearColor Actual, FLinearColor Expected)
+bool FAutomationTestBase::TestEqual(const TCHAR* What, const FLinearColor Actual, const FLinearColor Expected)
 {
 	if (Expected != Actual)
 	{
diff --git a/Engine/Source/Runtime/Core/Public/Misc/AutomationTest.h b/Engine/Source/Runtime/Core/Public/Misc/AutomationTest.h
--- a/Engine/Source/Runtime/Core/Public/Misc/AutomationTest.h
+++ b/Engine/Source/Runtime/Core/Public/Misc/AutomationTest.h
@@ -1641,47 +1641,47 @@ public:
 
 public:
 
-	bool TestEqual(const TCHAR* What, int32 Actual, int32 Expected);
-	bool TestEqual(const TCHAR* What, int64 Actual, int64 Expected);
+	bool TestEqual(const TCHAR* What, const int32 Actual, const int32 Expected);
+	bool TestEqual(const TCHAR* What, const int64 Actual, const int64 Expected);
 #if PLATFORM_64BITS
-	bool TestEqual(const TCHAR* What, SIZE_T Actual, SIZE_T Expected);
+	bool TestEqual(const TCHAR* What, const SIZE_T Actual, const SIZE_T Expected);
 #endif
-	bool TestEqual(const TCHAR* What, float Actual, float Expected, float Tolerance = UE_KINDA_SMALL_NUMBER);
-	bool TestEqual(const TCHAR* What, double Actual, double Expected, double Tolerance = UE_KINDA_SMALL_NUMBER);
-	bool TestEqual(const TCHAR* What, FVector Actual, FVector Expected, float Tolerance = UE_KINDA_SMALL_NUMBER);
-	bool TestEqual(const TCHAR* What, FTransform Actual, FTransform Expected, float Tolerance = UE_KINDA_SMALL_NUMBER);
-	bool TestEqual(const TCHAR* What, FRotator Actual, FRotator Expected, float Tolerance = UE_KINDA_SMALL_NUMBER);
-	bool TestEqual(const TCHAR* What, FColor Actual, FColor Expected);
-	bool TestEqual(const TCHAR* What, FLinearColor Actual, FLinearColor Expected);
+	bool TestEqual(const TCHAR* What, const float Actual, const float Expected, float Tolerance = UE_KINDA_SMALL_NUMBER);
+	bool TestEqual(const TCHAR* What, const double Actual, const double Expected, double Tolerance = UE_KINDA_SMALL_NUMBER);
+	bool TestEqual(const TCHAR* What, const FVector Actual, const FVector Expected, float Tolerance = UE_KINDA_SMALL_NUMBER);
+	bool TestEqual(const TCHAR* What, const FTransform Actual, const FTransform Expected, float Tolerance = UE_KINDA_SMALL_NUMBER);
+	bool TestEqual(const TCHAR* What, const FRotator Actual, const FRotator Expected, float Tolerance = UE_KINDA_SMALL_NUMBER);
+	bool TestEqual(const TCHAR* What, const FColor Actual, const FColor Expected);
+	bool TestEqual(const TCHAR* What, const FLinearColor Actual, const FLinearColor Expected);
 	bool TestEqual(const TCHAR* What, const TCHAR* Actual, const TCHAR* Expected);
 	bool TestEqualInsensitive(const TCHAR* What, const TCHAR* Actual, const TCHAR* Expected);
 
-	bool TestEqual(const FString& What, int32 Actual, int32 Expected)
+	bool TestEqual(const FString& What, const int32 Actual, const int32 Expected)
 	{
 		return TestEqual(*What, Actual, Expected);
 	}
 
-	bool TestEqual(const FString& What, float Actual, float Expected, float Tolerance = UE_KINDA_SMALL_NUMBER)
+	bool TestEqual(const FString& What, const float Actual, const float Expected, float Tolerance = UE_KINDA_SMALL_NUMBER)
 	{
 		return TestEqual(*What, Actual, Expected, Tolerance);
 	}
 
-	bool TestEqual(const FString& What, double Actual, double Expected, double Tolerance = UE_KINDA_SMALL_NUMBER)
+	bool TestEqual(const FString& What, const double Actual, const double Expected, double Tolerance = UE_KINDA_SMALL_NUMBER)
 	{
 		return TestEqual(*What, Actual, Expected, Tolerance);
 	}
 
-	bool TestEqual(const FString& What, FVector Actual, FVector Expected, float Tolerance = UE_KINDA_SMALL_NUMBER)
+	bool TestEqual(const FString& What, const FVector Actual, const FVector Expected, float Tolerance = UE_KINDA_SMALL_NUMBER)
 	{
 		return TestEqual(*What, Actual, Expected, Tolerance);
 	}
 
-	bool TestEqual(const FString& What, FRotator Actual, FRotator Expected, float Tolerance = UE_KINDA_SMALL_NUMBER)
+	bool TestEqual(const FString& What, const FRotator Actual, const FRotator Expected, float Tolerance = UE_KINDA_SMALL_NUMBER)
 	{
 		return TestEqual(*What, Actual, Expected, Tolerance);
 	}
 
-	bool TestEqual(const FString& What, FColor Actual, FColor Expected)
+	bool TestEqual(const FString& What, const FColor Actual, const FColor Expected)
 	{
 		return TestEqual(*What, Actual, Expected);
 	}
@@ -1818,7 +1818,7 @@ public:
 	 *
 	 * @see TestNull
 	 */
-	template<typename ValueType> bool TestNotNull(const TCHAR* What, ValueType* Pointer)
+	template<typename ValueType> bool TestNotNull(const TCHAR* What, const ValueType* Pointer)
 	{
 		if (Pointer == nullptr)
 		{
@@ -1828,7 +1828,7 @@ public:
 		return true;
 	}
 
-	template<typename ValueType> bool TestNotNull(const FString& What, ValueType* Pointer)
+	template<typename ValueType> bool TestNotNull(const FString& What, const ValueType* Pointer)
 	{
 		return TestNotNull(*What, Pointer);
 	}
