Removed ignored calculated values from TBigInt::ToInt(), FGenericPlatformMallocCrash::InitializeSmallPools and FArchive::SerializeCompressedNew.

diff --git a/Engine/Source/Runtime/Core/Private/GenericPlatform/GenericPlatformMallocCrash.cpp b/Engine/Source/Runtime/Core/Private/GenericPlatform/GenericPlatformMallocCrash.cpp
--- a/Engine/Source/Runtime/Core/Private/GenericPlatform/GenericPlatformMallocCrash.cpp
+++ b/Engine/Source/Runtime/Core/Private/GenericPlatform/GenericPlatformMallocCrash.cpp
@@ -521,12 +521,9 @@ uint32 FGenericPlatformMallocCrash::CalculateBookkeepingPoolTotalSize() const
 
 void FGenericPlatformMallocCrash::InitializeSmallPools()
 {
-	uint32 MallocCrashOverhead = 0;
-
 	for( uint32 Index = 0; Index < NUM_POOLS; ++Index )
 	{
 		FMallocCrashPool* NewPool = new (&FMallocCrashPool::GetPool(Index)) FMallocCrashPool(GetPoolDesc(Index), *this);
-		MallocCrashOverhead += NewPool->AllocatedMemory;
 	}
 
 	check(SmallMemoryPoolOffset<=CalculateSmallPoolTotalSize());
diff --git a/Engine/Source/Runtime/Core/Private/Serialization/Archive.cpp b/Engine/Source/Runtime/Core/Private/Serialization/Archive.cpp
--- a/Engine/Source/Runtime/Core/Private/Serialization/Archive.cpp
+++ b/Engine/Source/Runtime/Core/Private/Serialization/Archive.cpp
@@ -712,7 +712,6 @@ void FArchive::SerializeCompressedNew(void* V, int64 Length, FName CompressionFo
 		//	assume it was CompressionFormatToDecodeOldV1Files (usually Zlib)
 		FName CompressionFormatToDecode = CompressionFormatToDecodeOldV1Files;
 
-		bool bHeaderWasValid=false;
 		bool bWasByteSwapped=false;
 		bool bReadCompressionFormat=false;
 		
@@ -725,20 +724,17 @@ void FArchive::SerializeCompressedNew(void* V, int64 Length, FName CompressionFo
 		if ( PackageFileTag.CompressedSize == PACKAGE_FILE_TAG )
 		{
 			// v1 header, not swapped
-			bHeaderWasValid = true;
 		}
 		else if ( PackageFileTag.CompressedSize == PACKAGE_FILE_TAG_SWAPPED ||
 			PackageFileTag.CompressedSize == BYTESWAP_ORDER64((uint64)PACKAGE_FILE_TAG) )
 		{
 			// v1 header, swapped
-			bHeaderWasValid = true;
 			bWasByteSwapped = true;
 		}
 		else if ( PackageFileTag.CompressedSize == ARCHIVE_V2_HEADER_TAG ||
 			PackageFileTag.CompressedSize == BYTESWAP_ORDER64((uint64)ARCHIVE_V2_HEADER_TAG) )
 		{
 			// v2 header
-			bHeaderWasValid = true;
 			bWasByteSwapped = ( PackageFileTag.CompressedSize != ARCHIVE_V2_HEADER_TAG );
 			bReadCompressionFormat = true;
 
diff --git a/Engine/Source/Runtime/Core/Public/Math/BigInt.h b/Engine/Source/Runtime/Core/Public/Math/BigInt.h
--- a/Engine/Source/Runtime/Core/Public/Math/BigInt.h
+++ b/Engine/Source/Runtime/Core/Public/Math/BigInt.h
@@ -961,17 +961,6 @@ public:
 	 */
 	int64 ToInt() const
 	{
-		int64 Result;
-		if (!IsNegative())
-		{
-			Result = (int64)Bits[0] + ((int64)Bits[1] << BitsPerWord);
-		}
-		else
-		{
-			BigInt Positive(*this);
-			Positive.Negate();
-			Result = (int64)Positive.Bits[0] + ((int64)Positive.Bits[1] << BitsPerWord);
-		}
 		return (int64)Bits[0] + ((int64)Bits[1] << BitsPerWord);
 	}
 
