Explicitly allocate pages for hook trampolines.

Hooking a Windows API works by patching in a "jmp" instruction. The instructions that are overwritten by this patch are moved elsewhere (and appended with another jmp) such that the original function can still be called - I called this a trampoline. Prior to this change, trampolines were written into the null-filled slack space at the end of the code section. This however turned out to be too small as more hooks were added and as Windows' DLLs mutated. This change instead uses a page of address space which it allocates just before the DLL being patched (proximity is important as the jmp instruction used only has 32 bit signed displacement).

#rnx

diff --git a/Engine/Source/Runtime/Core/Private/ProfilingDebugging/Microsoft/WindowsMemoryTrace.cpp b/Engine/Source/Runtime/Core/Private/ProfilingDebugging/Microsoft/WindowsMemoryTrace.cpp
--- a/Engine/Source/Runtime/Core/Private/ProfilingDebugging/Microsoft/WindowsMemoryTrace.cpp
+++ b/Engine/Source/Runtime/Core/Private/ProfilingDebugging/Microsoft/WindowsMemoryTrace.cpp
@@ -38,46 +38,42 @@ static_assert(sizeof(FAddrPack) == sizeof(uint64), "");
 class FTextSectionEditor
 {
 public:
-							FTextSectionEditor(void* InBase);
 							~FTextSectionEditor();
-	template <typename T>
-	T*						Hook(T* Target, T* HookFunction);
+	template <typename T> T*Hook(T* Target, T* HookFunction);
 
 private:
+	struct FTrampolineBlock
+	{
+		FTrampolineBlock*	Next;
+		uint32				Size;
+		uint32				Used;
+	};
+
 	static void*			GetActualAddress(void* Function);
-	uint8*					AllocateTrampoline(unsigned int Size);
+	FTrampolineBlock*		AllocateTrampolineBlock(void* Reference);
+	uint8*					AllocateTrampoline(void* Reference, unsigned int Size);
 	void*					HookImpl(void* Target, void* HookFunction);
-	uint8*					TrampolineTail;
-	void*					Base;
-	size_t					Size;
-	DWORD					Protection;
+	FTrampolineBlock*		HeadBlock = nullptr;
 };
 
-////////////////////////////////////////////////////////////////////////////////
-FTextSectionEditor::FTextSectionEditor(void* InBase)
-{
-	InBase = GetActualAddress(InBase);
-
-	MEMORY_BASIC_INFORMATION MemInfo;
-	VirtualQuery(InBase, &MemInfo, sizeof(MemInfo));
-	Base = MemInfo.BaseAddress;
-	Size = MemInfo.RegionSize;
-
-	VirtualProtect(Base, Size, PAGE_EXECUTE_READWRITE, &Protection);
-
-	TrampolineTail = (uint8*)Base + Size;
-}
-
 ////////////////////////////////////////////////////////////////////////////////
 FTextSectionEditor::~FTextSectionEditor()
 {
-	VirtualProtect(Base, Size, Protection, &Protection);
-	FlushInstructionCache(GetCurrentProcess(), Base, Size);
+	for (FTrampolineBlock* Block = HeadBlock; Block != nullptr; Block = Block->Next)
+	{
+		DWORD Unused;
+		VirtualProtect(Block, Block->Size, PAGE_EXECUTE_READ, &Unused);
+	}
+
+	FlushInstructionCache(GetCurrentProcess(), nullptr, 0);
 }
 
 ////////////////////////////////////////////////////////////////////////////////
 void* FTextSectionEditor::GetActualAddress(void* Function)
 {
+	// Follow a jmp instruction (0xff/4 only for now) at function and returns
+	// where it would jmp to.
+
 	uint8* Addr = (uint8*)Function;
 	int Offset = unsigned(Addr[0] & 0xf0) == 0x40; // REX prefix
 	if (Addr[Offset + 0] == 0xff && Addr[Offset + 1] == 0x25)
@@ -89,19 +85,73 @@ void* FTextSectionEditor::GetActualAddress(void* Function)
 }
 
 ////////////////////////////////////////////////////////////////////////////////
-uint8* FTextSectionEditor::AllocateTrampoline(unsigned int PatchSize)
+FTextSectionEditor::FTrampolineBlock* FTextSectionEditor::AllocateTrampolineBlock(void* Reference)
 {
-	static const int TrampolineSize = 24;
-	uint8* NextTail = TrampolineTail - TrampolineSize;
-	do 
+	static const SIZE_T BlockSize = 0x10000; // 64KB is Windows' canonical granularity
+
+	// Find the start of the main allocation that mapped Reference
+	MEMORY_BASIC_INFORMATION MemInfo;
+	VirtualQuery(Reference, &MemInfo, sizeof(MemInfo));
+	auto* Ptr = (uint8*)(MemInfo.AllocationBase);
+
+	// Step backwards one block at a time and try and allocate that address
+	while (true)
 	{
-		check(*NextTail == 0);
-		++NextTail;
+		Ptr -= BlockSize;
+		if (VirtualAlloc(Ptr, BlockSize, MEM_RESERVE|MEM_COMMIT, PAGE_READWRITE) != nullptr)
+		{
+			break;
+		}
+
+		UPTRINT Distance = UPTRINT(Reference) - UPTRINT(Ptr);
+		if (Distance >= 1ull << 31)
+		{
+			check(!"Failed to allocate trampoline blocks for memory tracing hooks");
+		}
 	}
-	while (NextTail != TrampolineTail);
 
-	TrampolineTail -= TrampolineSize;
-	return TrampolineTail;
+	auto* Block = (FTrampolineBlock*)Ptr;
+	Block->Next = HeadBlock;
+	Block->Size = BlockSize;
+	Block->Used = sizeof(FTrampolineBlock);
+	HeadBlock = Block;
+
+	return Block;
+}
+
+////////////////////////////////////////////////////////////////////////////////
+uint8* FTextSectionEditor::AllocateTrampoline(void* Reference, unsigned int Size)
+{
+	// Try and find a block that's within 2^31 bytes before Reference
+	FTrampolineBlock* Block;
+	for (Block = HeadBlock; Block != nullptr; Block = Block->Next)
+	{
+		UPTRINT Distance = UPTRINT(Reference) - UPTRINT(Block);
+		if (Distance < 1ull << 31)
+		{
+			break;
+		}
+	}
+
+	// If we didn't find a block then we need to allocate a new one.
+	if (Block == nullptr)
+	{
+		Block = AllocateTrampolineBlock(Reference);
+	}
+
+	// Allocate space for the trampoline.
+	uint32 NextUsed = Block->Used + Size;
+	if (NextUsed > Block->Size)
+	{
+		// Block is full. We could allocate a new block here but as it is not
+		// expected that so many hooks will be made this path shouldn't happen
+		check(!"Unable to allocate memory for memory tracing's hooks");
+	}
+
+	uint8* Out = (uint8*)Block + Block->Used;
+	Block->Used = NextUsed;
+
+	return Out;
 }
 
 ////////////////////////////////////////////////////////////////////////////////
@@ -116,7 +166,9 @@ void* FTextSectionEditor::HookImpl(void* Target, void* HookFunction)
 {
 	Target = GetActualAddress(Target);
 
-	const uint8* Start = (uint8*)Target;
+	// Very rudimentary x86_64 instruction length decoding that only supports op
+	// code ranges (0x80,0x8b) and (0x50,0x5f). Enough for simple prologues
+	uint8* __restrict Start = (uint8*)Target;
 	const uint8* Read = Start;
 	do
 	{
@@ -138,9 +190,11 @@ void* FTextSectionEditor::HookImpl(void* Target, void* HookFunction)
 	}
 	while (Read - Start < 6);
 
+	static const int TrampolineSize = 24;
 	int PatchSize = int(Read - Start);
-	uint8* TrampolinePtr = AllocateTrampoline(PatchSize);
+	uint8* TrampolinePtr = AllocateTrampoline(Start, PatchSize + TrampolineSize);
 
+	// Write the trampoline
 	*(void**)TrampolinePtr = HookFunction;
 
 	uint8* PatchJmp = TrampolinePtr + sizeof(void*);
@@ -150,10 +204,20 @@ void* FTextSectionEditor::HookImpl(void* Target, void* HookFunction)
 	*PatchJmp = 0xe9;
 	*(int32*)(PatchJmp + 1) = int32(intptr_t(Start + PatchSize) - intptr_t(PatchJmp)) - 5;
 
+	// Need to make the text section writeable
+	DWORD ProtPrev;
+	UPTRINT ProtBase = UPTRINT(Target) & ~0x0fff;						// 0x0fff is mask of VM page size
+	SIZE_T ProtSize = ((ProtBase + 16 + 0x1000) & ~0x0fff) - ProtBase;	// 16 is enough for one x86 instruction
+	VirtualProtect((void*)ProtBase, ProtSize, PAGE_EXECUTE_READWRITE, &ProtPrev);
+
+	// Patch function to jmp to the hook
 	uint16* HookJmp = (uint16*)Target;
 	HookJmp[0] = 0x25ff;
 	*(int32*)(HookJmp + 1) = int32(intptr_t(TrampolinePtr) - intptr_t(HookJmp + 3));
 
+	// Put the protection back the way it was
+	VirtualProtect((void*)ProtBase, ProtSize, ProtPrev, &ProtPrev);
+
 	return PatchJmp - PatchSize;
 }
 
@@ -193,45 +257,42 @@ void FVirtualWinApiHooks::Initialize(bool bInLight)
 {
 	bLight = bInLight;
 
-	{
-		FTextSectionEditor Editor((void*)VirtualAlloc);
-		VmAllocOrig = Editor.Hook(VirtualAlloc, &FVirtualWinApiHooks::VmAlloc);
-		VmFreeOrig = Editor.Hook(VirtualFree, &FVirtualWinApiHooks::VmFree); 
-		
-		{
-			FTextSectionEditor EditorEx((void*)VirtualAllocEx);
-			VmAllocExOrig = Editor.Hook(VirtualAllocEx, &FVirtualWinApiHooks::VmAllocEx);
-			VmFreeExOrig = Editor.Hook(VirtualFreeEx, &FVirtualWinApiHooks::VmFreeEx);
-		}
+	FTextSectionEditor Editor;
+
+	// Note that hooking alloc functions is done last as applying the hook can
+	// allocate some memory pages.
+
+	VmFreeOrig = Editor.Hook(VirtualFree, &FVirtualWinApiHooks::VmFree); 
+	VmFreeExOrig = Editor.Hook(VirtualFreeEx, &FVirtualWinApiHooks::VmFreeEx);
 
 #if PLATFORM_WINDOWS
 #if (NTDDI_VERSION >= NTDDI_WIN10_RS4)
-		{
-			FTextSectionEditor EditorAlloc2((void*)VirtualAlloc2);
-			VmAlloc2Orig = Editor.Hook(VirtualAlloc2, &FVirtualWinApiHooks::VmAlloc2);
-		}
+	{
+		VmAlloc2Orig = Editor.Hook(VirtualAlloc2, &FVirtualWinApiHooks::VmAlloc2);
+	}
 #else // NTDDI_VERSION
+	{
+		VmAlloc2Orig = nullptr;
+		HINSTANCE DllInstance;
+		DllInstance = LoadLibrary(TEXT("kernelbase.dll"));
+		if (DllInstance != NULL)
 		{
-			VmAlloc2Orig = nullptr;
-			HINSTANCE DllInstance;
-			DllInstance = LoadLibrary(TEXT("kernelbase.dll"));
-			if (DllInstance != NULL)
-			{
 #pragma warning(push)
 #pragma warning(disable: 4191) // 'type cast': unsafe conversion from 'FARPROC' to 'FVirtualWinApiHooks::FnVirtualAlloc2'
-				VmAlloc2Orig = (FnVirtualAlloc2)GetProcAddress(DllInstance, "VirtualAlloc2");
+			VmAlloc2Orig = (FnVirtualAlloc2)GetProcAddress(DllInstance, "VirtualAlloc2");
 #pragma warning(pop)
-				FreeLibrary(DllInstance);
-			}
-			if (VmAlloc2Orig)
-			{
-				FTextSectionEditor EditorAlloc2((void*)VmAlloc2Orig);
-				VmAlloc2Orig = Editor.Hook(VmAlloc2Orig, &FVirtualWinApiHooks::VmAlloc2);
-			}
+			FreeLibrary(DllInstance);
+		}
+		if (VmAlloc2Orig)
+		{
+			VmAlloc2Orig = Editor.Hook(VmAlloc2Orig, &FVirtualWinApiHooks::VmAlloc2);
 		}
+	}
 #endif // NTDDI_VERSION
 #endif // PLATFORM_WINDOWS
-	}
+
+	VmAllocExOrig = Editor.Hook(VirtualAllocEx, &FVirtualWinApiHooks::VmAllocEx);
+	VmAllocOrig = Editor.Hook(VirtualAlloc, &FVirtualWinApiHooks::VmAlloc);
 }
 
 ////////////////////////////////////////////////////////////////////////////////
