Fix inconsistency in Trace.File and Trace.Send functionality when -traceautostart is false.

diff --git a/Engine/Source/Runtime/Core/Private/ProfilingDebugging/TraceAuxiliary.cpp b/Engine/Source/Runtime/Core/Private/ProfilingDebugging/TraceAuxiliary.cpp
--- a/Engine/Source/Runtime/Core/Private/ProfilingDebugging/TraceAuxiliary.cpp
+++ b/Engine/Source/Runtime/Core/Private/ProfilingDebugging/TraceAuxiliary.cpp
@@ -1207,6 +1207,7 @@ static bool StartFromCommandlineArguments(const TCHAR* CommandLine)
 
 	if (!GTraceAutoStart)
 	{
+		GTraceAuxiliary.AddCommandlineChannels(*Channels);
 		return false;
 	}
 
@@ -1412,12 +1413,9 @@ void FTraceAuxiliary::Initialize(const TCHAR* CommandLine)
 	}
 	UE::Trace::Initialize(Desc);
 
-	if (GTraceAutoStart)
-	{
-		// Workaround for the fact that even if StartFromCommandlineArguments will enable channels
-		// specified by the commandline, UE::Trace::Initialize will reset all channels.
-		GTraceAuxiliary.EnableCommandlineChannelsPostInitialize();
-	}
+	// Workaround for the fact that even if StartFromCommandlineArguments will enable channels
+	// specified by the commandline, UE::Trace::Initialize will reset all channels.
+	GTraceAuxiliary.EnableCommandlineChannelsPostInitialize();
 	
 	// Setup known on connection callbacks
 	OnConnection.AddStatic(FStringTrace::OnConnection);
