#ue4 - Add UE::LWC::FloatToIntCastChecked<> for floating point to integer conversion. Accepts float/double input and requires template argument for return type (int32 and int64).

diff --git a/Engine/Source/Runtime/Core/Private/Tests/Math/UnrealMathTest.cpp b/Engine/Source/Runtime/Core/Private/Tests/Math/UnrealMathTest.cpp
--- a/Engine/Source/Runtime/Core/Private/Tests/Math/UnrealMathTest.cpp
+++ b/Engine/Source/Runtime/Core/Private/Tests/Math/UnrealMathTest.cpp
@@ -3244,6 +3244,32 @@ bool FVectorRegisterAbstractionTest::RunTest(const FString& Parameters)
 		Big3 = FMath::Clamp(I1, Big1, I2);
 		I = FMath::Clamp<int32>(Big1, I1, I2);
 #endif
+
+		I = UE::LWC::FloatToIntCastChecked<int32>(F1);
+		I = UE::LWC::FloatToIntCastChecked<int32>(D1);
+
+		Big1 = UE::LWC::FloatToIntCastChecked<int64>(F1);
+		Big1 = UE::LWC::FloatToIntCastChecked<int64>(D1);
+
+		// Valid (implicit int32->int64)
+		Big1 = UE::LWC::FloatToIntCastChecked<int32>(F1);
+		Big1 = UE::LWC::FloatToIntCastChecked<int32>(D1);
+
+		// Valid (implicit int32->float then int32->int64
+		Big1 = UE::LWC::FloatToIntCastChecked<int32>(I);
+
+#if MATHTEST_CHECK_INVALID_OVERLOAD_VARIANTS
+		// Shouldn't match a function
+		I = UE::LWC::FloatToIntCastChecked(F1);
+		I = UE::LWC::FloatToIntCastChecked(D1);
+		Big1 = UE::LWC::FloatToIntCastChecked(F1);
+		Big1 = UE::LWC::FloatToIntCastChecked(D1);
+		I = UE::LWC::FloatToIntCastChecked(I);
+		Big1 = UE::LWC::FloatToIntCastChecked(Big1);
+		// Truncation warnings
+		I = UE::LWC::FloatToIntCastChecked<int64>(F1);
+		I = UE::LWC::FloatToIntCastChecked<int64>(D1);
+#endif
 	}
 
 	// Sin, Cos, Tan tests
diff --git a/Engine/Source/Runtime/Core/Public/Math/UnrealMathUtility.h b/Engine/Source/Runtime/Core/Public/Math/UnrealMathUtility.h
--- a/Engine/Source/Runtime/Core/Public/Math/UnrealMathUtility.h
+++ b/Engine/Source/Runtime/Core/Public/Math/UnrealMathUtility.h
@@ -2490,5 +2490,44 @@ namespace LWC
 		}
 	}
 
+	/*
+	 * Floating point to integer conversions
+	 */
+
+	// Generic float type to int type, to enable specializations below.
+	template<typename OutIntType, typename InFloatType>
+	FORCEINLINE OutIntType FloatToIntCastChecked(InFloatType FloatValue)
+	{
+		return (OutIntType)(FloatValue);
+	}
+
+	// float->int32
+	template<>
+	FORCEINLINE int32 FloatToIntCastChecked(float FloatValue)
+	{
+		return FMath::TruncToInt32(FloatValue);
+	}
+
+	// float->int64
+	template<>
+	FORCEINLINE int64 FloatToIntCastChecked(float FloatValue)
+	{
+		return FMath::TruncToInt64(FloatValue);
+	}
+
+	// double->int32
+	template<>
+	FORCEINLINE int32 FloatToIntCastChecked(double FloatValue)
+	{
+		return FMath::TruncToInt32(FloatValue);
+	}
+
+	// double->int64
+	template<>
+	FORCEINLINE int64 FloatToIntCastChecked(double FloatValue)
+	{
+		return FMath::TruncToInt64(FloatValue);
+	}
+
 } // namespace LWC
 } // namespace UE
