Provide a commandline option '-VA-Disable' that when used will disable the asset virtualization system and mount the null implementation instead.

#rnx

diff --git a/Engine/Source/Runtime/Core/Private/Virtualization/VirtualizationSystem.cpp b/Engine/Source/Runtime/Core/Private/Virtualization/VirtualizationSystem.cpp
--- a/Engine/Source/Runtime/Core/Private/Virtualization/VirtualizationSystem.cpp
+++ b/Engine/Source/Runtime/Core/Private/Virtualization/VirtualizationSystem.cpp
@@ -162,11 +162,18 @@ void Initialize(const FInitParams& InitParams)
 	FName SystemName;
 
 #if UE_DISABLE_VIRTUALIZATION_SYSTEM == 0
-	FString RawSystemName;
-	if (InitParams.ConfigFile.GetString(TEXT("Core.ContentVirtualization"), TEXT("SystemName"), RawSystemName))
+	if (FParse::Param(FCommandLine::Get(), TEXT("VA-Disable")))
 	{
-		SystemName = FName(RawSystemName);
-		UE_LOG(LogVirtualization, Display, TEXT("VirtualizationSystem name found in ini file: %s"), *RawSystemName);
+		UE_LOG(LogVirtualization, Display, TEXT("The virtualization system has been disabled by the command line"));
+	}
+	else
+	{
+		FString RawSystemName;
+		if (InitParams.ConfigFile.GetString(TEXT("Core.ContentVirtualization"), TEXT("SystemName"), RawSystemName))
+		{
+			SystemName = FName(RawSystemName);
+			UE_LOG(LogVirtualization, Display, TEXT("VirtualizationSystem name found in ini file: %s"), *RawSystemName);
+		}
 	}
 #else
 	UE_LOG(LogVirtualization, Display, TEXT("The virtualization system has been disabled by code"));
