Bring RES files in line with EM
#rnx

diff --git a/Engine/Source/Runtime/Core/Public/Stats/Stats2.h b/Engine/Source/Runtime/Core/Public/Stats/Stats2.h
--- a/Engine/Source/Runtime/Core/Public/Stats/Stats2.h
+++ b/Engine/Source/Runtime/Core/Public/Stats/Stats2.h
@@ -10,6 +10,7 @@
 #include "CoreGlobals.h"
 #include "CoreTypes.h"
 #include "Delegates/Delegate.h"
+#include "Delegates/DelegateBase.h"
 #include "HAL/CriticalSection.h"
 #include "HAL/LowLevelMemTracker.h"
 #include "HAL/PlatformCrt.h"
