[Core]
* Added macro to be used to place system trackedactivity scopes

diff --git a/Engine/Source/Runtime/Core/Private/Misc/CoreGlobals.cpp b/Engine/Source/Runtime/Core/Private/Misc/CoreGlobals.cpp
--- a/Engine/Source/Runtime/Core/Private/Misc/CoreGlobals.cpp
+++ b/Engine/Source/Runtime/Core/Private/Misc/CoreGlobals.cpp
@@ -7,6 +7,7 @@
 #include "HAL/IConsoleManager.h"
 #include "Modules/ModuleManager.h"
 #include "Misc/CoreStats.h"
+#include "Misc/TrackedActivity.h"
 #include "Misc/Compression.h"
 #include "Misc/LazySingleton.h"
 #include "Misc/CommandLine.h"
@@ -450,6 +451,25 @@ static struct FBootTimingStart
 	}
 } GBootTimingStart;
 
+FEngineTrackedActivityScope::FEngineTrackedActivityScope(const TCHAR* Fmt, ...)
+{
+	va_list Args;
+	va_start(Args, Fmt);
+	TCHAR Str[4096];
+	FCString::GetVarArgs(Str, UE_ARRAY_COUNT(Str), Fmt, Args);
+	FTrackedActivity::GetEngineActivity().Push(Str);
+	va_end(Args);
+}
+
+FEngineTrackedActivityScope::FEngineTrackedActivityScope(const FString& Text)
+{
+	FTrackedActivity::GetEngineActivity().Push(*Text);
+}
+
+FEngineTrackedActivityScope::~FEngineTrackedActivityScope()
+{
+	FTrackedActivity::GetEngineActivity().Pop();
+}
 
 #define USE_BOOT_PROFILING 0
 
diff --git a/Engine/Source/Runtime/Core/Public/CoreGlobals.h b/Engine/Source/Runtime/Core/Public/CoreGlobals.h
--- a/Engine/Source/Runtime/Core/Public/CoreGlobals.h
+++ b/Engine/Source/Runtime/Core/Public/CoreGlobals.h
@@ -80,8 +80,17 @@ struct CORE_API FScopedBootTiming
 	~FScopedBootTiming();
 };
 
+struct CORE_API FEngineTrackedActivityScope
+{
+	FEngineTrackedActivityScope(const TCHAR* Fmt, ...);
+	FEngineTrackedActivityScope(const FString& Str);
+	~FEngineTrackedActivityScope();
+};
+
 
 #define SCOPED_BOOT_TIMING(x) TRACE_CPUPROFILER_EVENT_SCOPE_STR(x); FScopedBootTiming ANONYMOUS_VARIABLE(BootTiming_)(x);
+#define UE_SCOPED_ENGINE_ACTIVITY(Fmt, ...) FEngineTrackedActivityScope ANONYMOUS_VARIABLE(EngineActivity_)(Fmt, __VA_ARGS__);
+
 
 #define GLog GetGlobalLogSingleton()
 extern CORE_API FConfigCacheIni* GConfig;
