[Runtime]
* Fixed so UE_SCOPED_ENGINE_ACTIVITY works with empty variadic arguments

diff --git a/Engine/Source/Runtime/Core/Public/CoreGlobals.h b/Engine/Source/Runtime/Core/Public/CoreGlobals.h
--- a/Engine/Source/Runtime/Core/Public/CoreGlobals.h
+++ b/Engine/Source/Runtime/Core/Public/CoreGlobals.h
@@ -89,7 +89,7 @@ struct CORE_API FEngineTrackedActivityScope
 
 
 #define SCOPED_BOOT_TIMING(x) TRACE_CPUPROFILER_EVENT_SCOPE_STR(x); FScopedBootTiming ANONYMOUS_VARIABLE(BootTiming_)(x);
-#define UE_SCOPED_ENGINE_ACTIVITY(Fmt, ...) FEngineTrackedActivityScope ANONYMOUS_VARIABLE(EngineActivity_)(Fmt, __VA_ARGS__);
+#define UE_SCOPED_ENGINE_ACTIVITY(Fmt, ...) FEngineTrackedActivityScope ANONYMOUS_VARIABLE(EngineActivity_)(Fmt, ## __VA_ARGS__);
 
 
 #define GLog GetGlobalLogSingleton()
