#ue5 - Make generic UE::LWC::FloatToIntCastChecked() strict in the input and return types to require float input and int output.

diff --git a/Engine/Source/Runtime/Core/Private/Tests/Math/UnrealMathTest.cpp b/Engine/Source/Runtime/Core/Private/Tests/Math/UnrealMathTest.cpp
--- a/Engine/Source/Runtime/Core/Private/Tests/Math/UnrealMathTest.cpp
+++ b/Engine/Source/Runtime/Core/Private/Tests/Math/UnrealMathTest.cpp
@@ -3255,9 +3255,6 @@ bool FVectorRegisterAbstractionTest::RunTest(const FString& Parameters)
 		Big1 = UE::LWC::FloatToIntCastChecked<int32>(F1);
 		Big1 = UE::LWC::FloatToIntCastChecked<int32>(D1);
 
-		// Valid (implicit int32->float then int32->int64
-		Big1 = UE::LWC::FloatToIntCastChecked<int32>(I);
-
 #if MATHTEST_CHECK_INVALID_OVERLOAD_VARIANTS
 		// Shouldn't match a function
 		I = UE::LWC::FloatToIntCastChecked(F1);
@@ -3269,6 +3266,8 @@ bool FVectorRegisterAbstractionTest::RunTest(const FString& Parameters)
 		// Truncation warnings
 		I = UE::LWC::FloatToIntCastChecked<int64>(F1);
 		I = UE::LWC::FloatToIntCastChecked<int64>(D1);
+		// Possibly invalid, depending on the generic base FloatToIntCastChecked, since this doesn't match a specialization
+		Big1 = UE::LWC::FloatToIntCastChecked<int32>(I);
 #endif
 	}
 
diff --git a/Engine/Source/Runtime/Core/Public/Math/UnrealMathUtility.h b/Engine/Source/Runtime/Core/Public/Math/UnrealMathUtility.h
--- a/Engine/Source/Runtime/Core/Public/Math/UnrealMathUtility.h
+++ b/Engine/Source/Runtime/Core/Public/Math/UnrealMathUtility.h
@@ -2498,6 +2498,8 @@ namespace LWC
 	template<typename OutIntType, typename InFloatType>
 	FORCEINLINE OutIntType FloatToIntCastChecked(InFloatType FloatValue)
 	{
+		static_assert(std::is_floating_point_v<InFloatType>, "Only floating point input type supported!");
+		static_assert(std::is_integral_v<OutIntType>, "Only integral output type supported!");
 		return (OutIntType)(FloatValue);
 	}
 
