[Insights] Added few log messages re trace system initialization.

diff --git a/Engine/Source/Runtime/Core/Private/ProfilingDebugging/TraceAuxiliary.cpp b/Engine/Source/Runtime/Core/Private/ProfilingDebugging/TraceAuxiliary.cpp
--- a/Engine/Source/Runtime/Core/Private/ProfilingDebugging/TraceAuxiliary.cpp
+++ b/Engine/Source/Runtime/Core/Private/ProfilingDebugging/TraceAuxiliary.cpp
@@ -1351,11 +1351,15 @@ void FTraceAuxiliary::Initialize(const TCHAR* CommandLine)
 #endif
 
 #if UE_TRACE_ENABLED
+	UE_LOG(LogCore, Log, TEXT("Initializing trace..."));
+
 	FParse::Bool(CommandLine, TEXT("-traceautostart="), GTraceAutoStart);
 	UE_LOG(LogCore, Verbose, TEXT("Trace auto start = %d."), GTraceAutoStart);
 
 	if (GTraceAuxiliary.IsParentProcessAndPreFork())
 	{
+		UE_LOG(LogCore, Log, TEXT("Trace initialization skipped for parent process (pre fork)."));
+
 		GTraceAuxiliary.DisableChannels(nullptr);
 
 		// Set our post fork callback up and return - children will pass through and Initialize when they're created.
@@ -1401,7 +1405,7 @@ void FTraceAuxiliary::Initialize(const TCHAR* CommandLine)
 	// seems odd to do this before initialising Trace, but it is done this way
 	// to support disabling the "important" cache without losing any events.
 	StartFromCommandlineArguments(CommandLine);
-	
+
 	// Initialize Trace
 	UE::Trace::FInitializeDesc Desc;
 	Desc.bUseWorkerThread = false;
@@ -1416,7 +1420,7 @@ void FTraceAuxiliary::Initialize(const TCHAR* CommandLine)
 	// Workaround for the fact that even if StartFromCommandlineArguments will enable channels
 	// specified by the commandline, UE::Trace::Initialize will reset all channels.
 	GTraceAuxiliary.EnableCommandlineChannelsPostInitialize();
-	
+
 	// Setup known on connection callbacks
 	OnConnection.AddStatic(FStringTrace::OnConnection);
 
@@ -1449,7 +1453,14 @@ void FTraceAuxiliary::Initialize(const TCHAR* CommandLine)
 	}
 
 	UE::Trace::ThreadRegister(TEXT("GameThread"), FPlatformTLS::GetCurrentThreadId(), -1);
-#endif
+
+	UE_LOG(LogCore, Log, TEXT("Finished trace initialization."));
+
+#else //UE_TRACE_ENABLED
+
+	UE_LOG(LogCore, Log, TEXT("Trace is not enabled (see UE_TRACE_ENABLED)."));
+
+#endif //UE_TRACE_ENABLED
 }
 
 ////////////////////////////////////////////////////////////////////////////////
