Add a number of ways for the VA system to be changed to lazy initialize on first use.

#rnx

- There are some usecases that require the VA system to initialize the first time it is accessed (usually the first time we attempt to pull a virtualized payload) rather than be initialized in the program start up. This change provides three different methods to achieve this:
-- Setting the define 'UE_VIRTUALIZATION_SYSTEM_LAZY_INIT' to 1 in a programs .target.cs
-- Setting [Core.ContentVirtualization]LazyInit=true in the Engine ini file
-- Running with the commandline option -VA-LazyInit

- If we detect that the source control backend is being initialized on a background thread we do not try to run the FConnect operation. The backend will still work but this does reduce the potential error checking on initialization. This is done because the FConnect operation currently only works on the main thread and to change this would be a bigger work item than we can schedule at the moment.
- UE::Virtualization::Initialize functions now take a EInitializationFlags enum as a parameter. This enum allows the call to ignore all lazy init settings and force the initialization immediately. This is useful for programs like the Virtualization standalone tool which just needs to start the system when needed.
-- The call to ::Initialize in LaunchEngineLoop passes in None and does not ignore lazy initialization.
-- Calls to ::Initialize in the UnrealVirtualizationTool however all use EInitializationFlags::ForceInitialize and ignore lazy initialization settings.
- Fixed an odd bug in UE::Virtualization::Initialize where the error path (if the config file cannot be found) was using a different start up code path.
- Add asserts when assigning to GVirtualizationSystem to make sure that it is null. This is not 100% safe but should catch some potential threading issues, if any.
- Add an assert after lazy initialization (IVirtualizationSystem::Get) to make sure that GVirtualizationSystem was assigned a valid object.
- Improve how we check for legacy values in [Core.ContentVirtualization]. We now support multiple allowed values.
- Added a way to poll if a VA system has been initialize yet or not, this allows us to avoid initializing a VA system if one has not yet been created and we try to:
-- Dump VA profiling stats after cooking
-- Send VA stats to studio analytics
- Note that currently using lazy init loading will probably cause the VA statistics panel not to work, this will be fixed in future work where we will allow the panel to register for a callback when the system is initialized.

diff --git a/Engine/Config/BaseEngine.ini b/Engine/Config/BaseEngine.ini
--- a/Engine/Config/BaseEngine.ini
+++ b/Engine/Config/BaseEngine.ini
@@ -1241,6 +1241,7 @@ AllowVirtualizationOptOut =False
 ; Sets which virtualization system to use, valid options are None, Default (use the Epic provided system) or a third party system
 [Core.ContentVirtualization]
 SystemName=None
+LazyInit=False
 
 [Core.VirtualizationModule]
 EnablePayloadVirtualization=True
diff --git a/Engine/Source/Developer/Virtualization/Private/VirtualizationManager.cpp b/Engine/Source/Developer/Virtualization/Private/VirtualizationManager.cpp
--- a/Engine/Source/Developer/Virtualization/Private/VirtualizationManager.cpp
+++ b/Engine/Source/Developer/Virtualization/Private/VirtualizationManager.cpp
@@ -906,20 +906,32 @@ void FVirtualizationManager::ApplySettingsFromConfigFiles(const FConfigFile& Con
 	}
 
 	// Check for any legacy settings and print them out (easier to do this in one block rather than one and time)
-	if (const FConfigSection* LegacySection = ConfigFile.Find(LegacyConfigSection))
 	{
-		if (LegacySection->Num() > 1)
+		// Entries that are allows to be in [Core.ContentVirtualization
+		TArrayView<const TCHAR* const> AllowedEntries = { TEXT("SystemName") , TEXT("LazyInit") };
+		
+		TArray<FString> LegacyEntries;	
+		if (const FConfigSection* LegacySection = ConfigFile.Find(LegacyConfigSection))
 		{
-			UE_LOG(LogVirtualization, Warning, TEXT("\tFound %d legacy ini file settings under [Core.ContentVirtualization] that should be moved to [Core.VirtualizationModule]"), LegacySection->Num() - 1);
-			for (auto Iterator : *LegacySection)
+			for (const TPair<FName, FConfigValue>& It : *LegacySection)
 			{
-				FString Name = Iterator.Key.ToString();
-				if (Name != TEXT("SystemName"))
+				FString Name = It.Key.ToString();
+			
+				if (!AllowedEntries.Contains(Name))
 				{
-					UE_LOG(LogVirtualization, Warning, TEXT("\t\t%s"), *Iterator.Key.ToString());
+					LegacyEntries.Add(MoveTemp(Name));
 				}
 			}
 		}
+
+		if (!LegacyEntries.IsEmpty())
+		{
+			UE_LOG(LogVirtualization, Warning, TEXT("\tFound %d legacy ini file settings under [Core.ContentVirtualization] that should be moved to [Core.VirtualizationModule]"), LegacyEntries.Num());
+			for (const FString& LegacyEntry : LegacyEntries)
+			{
+				UE_LOG(LogVirtualization, Warning, TEXT("\t\t%s"), *LegacyEntry);
+			}
+		}
 	}
 }
 
diff --git a/Engine/Source/Developer/Virtualization/Private/VirtualizationSourceControlBackend.cpp b/Engine/Source/Developer/Virtualization/Private/VirtualizationSourceControlBackend.cpp
--- a/Engine/Source/Developer/Virtualization/Private/VirtualizationSourceControlBackend.cpp
+++ b/Engine/Source/Developer/Virtualization/Private/VirtualizationSourceControlBackend.cpp
@@ -246,22 +246,25 @@ bool FSourceControlBackend::Initialize(const FString& ConfigEntry)
 
 	SCCProvider->Init(true);
  
-	// Note that if the connect is failing then we expect it to fail here rather than in the subsequent attempts to get the meta info file
-	TSharedRef<FConnect, ESPMode::ThreadSafe> ConnectCommand = ISourceControlOperation::Create<FConnect>();
-	if (SCCProvider->Execute(ConnectCommand, FString(), EConcurrency::Synchronous) != ECommandResult::Succeeded)
-	{		
-		FTextBuilder Errors;
-		for (const FText& Msg : ConnectCommand->GetResultInfo().ErrorMessages)
+	if (IsInGameThread())
+	{
+		// Note that if the connect is failing then we expect it to fail here rather than in the subsequent attempts to get the meta info file
+		TSharedRef<FConnect, ESPMode::ThreadSafe> ConnectCommand = ISourceControlOperation::Create<FConnect>();
+		if (SCCProvider->Execute(ConnectCommand, FString(), EConcurrency::Synchronous) != ECommandResult::Succeeded)
 		{
-			Errors.AppendLine(Msg);
-		}
+			FTextBuilder Errors;
+			for (const FText& Msg : ConnectCommand->GetResultInfo().ErrorMessages)
+			{
+				Errors.AppendLine(Msg);
+			}
 
-		FMessageLog Log("LogVirtualization");
-		Log.Warning(	FText::Format(LOCTEXT("FailedSourceControlConnection", "Failed to connect to source control backend with the following errors:\n{0}\nThe source control backend had trouble connecting!\nTrying logging in with the 'p4 login' command or by using p4vs/UnrealGameSync."),
-					Errors.ToText()));
-		
-		OnConnectionError();
-		return true;
+			FMessageLog Log("LogVirtualization");
+			Log.Warning(FText::Format(LOCTEXT("FailedSourceControlConnection", "Failed to connect to source control backend with the following errors:\n{0}\nThe source control backend had trouble connecting!\nTrying logging in with the 'p4 login' command or by using p4vs/UnrealGameSync."),
+				Errors.ToText()));
+
+			OnConnectionError();
+			return true;
+		}
 	}
 
 	// When a source control depot is set up a file named 'payload_metainfo.txt' should be submitted to it's root.
diff --git a/Engine/Source/Editor/DerivedDataEditor/Private/SVirtualAssetsStatistics.cpp b/Engine/Source/Editor/DerivedDataEditor/Private/SVirtualAssetsStatistics.cpp
--- a/Engine/Source/Editor/DerivedDataEditor/Private/SVirtualAssetsStatistics.cpp
+++ b/Engine/Source/Editor/DerivedDataEditor/Private/SVirtualAssetsStatistics.cpp
@@ -144,10 +144,14 @@ void AddContextMenuEntry()
 SVirtualAssetsStatisticsDialog::SVirtualAssetsStatisticsDialog()
 {
 	using namespace UE::Virtualization;
-
+	
+	// TODO - need a way to make this work once the system is initialized
 	// Register our VA notification delegate with the event
-	IVirtualizationSystem& System = IVirtualizationSystem::Get();
-	System.GetNotificationEvent().AddRaw(this, &SVirtualAssetsStatisticsDialog::OnNotificationEvent);
+	if (IVirtualizationSystem::IsInitialized())
+	{
+		IVirtualizationSystem& System = IVirtualizationSystem::Get();
+		System.GetNotificationEvent().AddRaw(this, &SVirtualAssetsStatisticsDialog::OnNotificationEvent);
+	}
 
 	UE::Virtualization::Rehydration::AddContextMenuEntry();
 }
@@ -313,10 +317,6 @@ TSharedRef<SWidget> SVirtualAssetsStatisticsDialog::GetGridPanel()
 {
 	using namespace UE::Virtualization;
 
-	IVirtualizationSystem& System = IVirtualizationSystem::Get();
-
-	TSharedRef<SGridPanel> Panel = SNew(SGridPanel);
-
 	const float RowMargin = 0.0f;
 	const float TitleMargin = 10.0f;
 	const float ColumnMargin = 10.0f;
@@ -324,6 +324,26 @@ TSharedRef<SWidget> SVirtualAssetsStatisticsDialog::GetGridPanel()
 	const FSlateFontInfo TitleFont = FCoreStyle::GetDefaultFontStyle("Bold", 10);
 	const double BytesToMegaBytes = 1.0 / (1024.0 * 1024.0);
 
+	TSharedRef<SGridPanel> Panel = SNew(SGridPanel);
+
+	// TODO - need a way to make this work once the system is initialized
+	if (!IVirtualizationSystem::IsInitialized())
+	{
+		Panel->AddSlot(0, 0)
+			[
+				SNew(STextBlock)
+				.Margin(FMargin(ColumnMargin, RowMargin))
+				.ColorAndOpacity(TitleColor)
+				.Font(TitleFont)
+				.Justification(ETextJustify::Center)
+				.Text(LOCTEXT("LazyInit", "Virtual Assets are lazy initialized for this project, cannot display stats"))
+			];
+
+		return Panel;
+	}
+
+	IVirtualizationSystem& System = IVirtualizationSystem::Get();	
+
 	if (System.IsEnabled() == false)
 	{
 		Panel->AddSlot(0,0)
diff --git a/Engine/Source/Editor/UnrealEd/Private/Commandlets/CookCommandlet.cpp b/Engine/Source/Editor/UnrealEd/Private/Commandlets/CookCommandlet.cpp
--- a/Engine/Source/Editor/UnrealEd/Private/Commandlets/CookCommandlet.cpp
+++ b/Engine/Source/Editor/UnrealEd/Private/Commandlets/CookCommandlet.cpp
@@ -155,7 +155,7 @@ namespace DetailedCookStats
 			if (GConfig->GetString(TEXT("CookAnalytics"), TEXT("APIServer"), APIServerET, GEngineIni))
 			{
 				FString AppId(TEXT("Cook"));
- 				bool bUseLegacyCookProtocol = !GConfig->GetString(TEXT("CookAnalytics"), TEXT("AppId"), AppId, GEngineIni);
+				bool bUseLegacyCookProtocol = !GConfig->GetString(TEXT("CookAnalytics"), TEXT("AppId"), AppId, GEngineIni);
 
 				// Optionally create an analytics provider to send stats to for central collection.
 				TSharedPtr<IAnalyticsProviderET> CookAnalytics = FAnalyticsET::Get().CreateAnalyticsProvider(FAnalyticsET::Config(AppId, APIServerET, FString(), bUseLegacyCookProtocol));
@@ -241,7 +241,7 @@ namespace DetailedCookStats
 
 		/** this functor will take a collected cooker stat and log it out using some custom formatting based on known stats that are collected.. */
 		auto LogStatsFunc = [&DDCResourceUsageStats, &DDCSummaryStats, &CookProfileData, &StatCategories, &StatsInCategories]
-							(const FString& StatName, const TArray<FCookStatsManager::StringKeyValue>& StatAttributes)
+		(const FString& StatName, const TArray<FCookStatsManager::StringKeyValue>& StatAttributes)
 		{
 			// Some stats will use custom formatting to make a visibly pleasing summary.
 			bool bStatUsedCustomFormatting = false;
@@ -379,9 +379,9 @@ namespace DetailedCookStats
 				SortedDDCResourceUsageStats.Emplace(Stat);
 			}
 			SortedDDCResourceUsageStats.Sort([](const FDDCResourceUsageStat& LHS, const FDDCResourceUsageStat& RHS)
-			{
-				return LHS.TotalTimeSec > RHS.TotalTimeSec;
-			});
+				{
+					return LHS.TotalTimeSec > RHS.TotalTimeSec;
+				});
 
 			UE_LOG(LogCookCommandlet, Display, TEXT(""));
 			UE_LOG(LogCookCommandlet, Display, TEXT("DDC Resource Stats"));
@@ -396,7 +396,10 @@ namespace DetailedCookStats
 
 		DumpBuildDependencyTrackerStats();
 
-		UE::Virtualization::IVirtualizationSystem::Get().DumpStats();
+		if (UE::Virtualization::IVirtualizationSystem::IsInitialized())
+		{
+			UE::Virtualization::IVirtualizationSystem::Get().DumpStats();
+		}
 	}
 }
 #endif
diff --git a/Engine/Source/Programs/UnrealVirtualizationTool/Private/Commands/RehydrateCommand.cpp b/Engine/Source/Programs/UnrealVirtualizationTool/Private/Commands/RehydrateCommand.cpp
--- a/Engine/Source/Programs/UnrealVirtualizationTool/Private/Commands/RehydrateCommand.cpp
+++ b/Engine/Source/Programs/UnrealVirtualizationTool/Private/Commands/RehydrateCommand.cpp
@@ -145,7 +145,7 @@ bool FRehydrateCommand::Run(const TArray<FProject>& Projects)
 		Project.RegisterMountPoints();
 
 		UE::Virtualization::FInitParams InitParams(ProjectName, EngineConfigWithProject);
-		UE::Virtualization::Initialize(InitParams);
+		UE::Virtualization::Initialize(InitParams, UE::Virtualization::EInitializationFlags::ForceInitialize);
 		
 		ON_SCOPE_EXIT
 		{
diff --git a/Engine/Source/Programs/UnrealVirtualizationTool/Private/UnrealVirtualizationToolApp.cpp b/Engine/Source/Programs/UnrealVirtualizationTool/Private/UnrealVirtualizationToolApp.cpp
--- a/Engine/Source/Programs/UnrealVirtualizationTool/Private/UnrealVirtualizationToolApp.cpp
+++ b/Engine/Source/Programs/UnrealVirtualizationTool/Private/UnrealVirtualizationToolApp.cpp
@@ -240,7 +240,7 @@ bool FUnrealVirtualizationToolApp::Run()
 				Project.RegisterMountPoints();
 
 				UE::Virtualization::FInitParams InitParams(ProjectName, EngineConfigWithProject);
-				UE::Virtualization::Initialize(InitParams);
+				UE::Virtualization::Initialize(InitParams, UE::Virtualization::EInitializationFlags::ForceInitialize);
 
 				TArray<FString> Packages = Project.GetAllPackages();
 
diff --git a/Engine/Source/Runtime/Core/Private/Virtualization/VirtualizationSystem.cpp b/Engine/Source/Runtime/Core/Private/Virtualization/VirtualizationSystem.cpp
--- a/Engine/Source/Runtime/Core/Private/Virtualization/VirtualizationSystem.cpp
+++ b/Engine/Source/Runtime/Core/Private/Virtualization/VirtualizationSystem.cpp
@@ -5,14 +5,20 @@
 #include "CoreGlobals.h"
 #include "Misc/App.h"
 #include "Misc/ConfigCacheIni.h"
+#include "Misc/ScopeLock.h"
 
 // Can be defined as 1 by programs target.cs files to disable the virtualization system
 // entirely. This could be removed in the future if we move the virtualization module
 // to be a plugin
 #ifndef UE_DISABLE_VIRTUALIZATION_SYSTEM
 	#define UE_DISABLE_VIRTUALIZATION_SYSTEM 0
-#endif
+#endif //UE_DISABLE_VIRTUALIZATION_SYSTEM
 
+// Can be defined as 1 by programs target.cs files to change the initialization of
+// the virtualization system to be a lazy initialization.
+#ifndef UE_VIRTUALIZATION_SYSTEM_LAZY_INIT
+	#define UE_VIRTUALIZATION_SYSTEM_LAZY_INIT 0
+#endif //UE_VIRTUALIZATION_SYSTEM_LAZY_INIT
 namespace UE::Virtualization
 {
 
@@ -133,32 +139,63 @@ Private::IVirtualizationSystemFactory* FindFactory(FName SystemName)
 	return nullptr;
 }
 
-void Initialize()
+void Initialize(EInitializationFlags Flags)
 {
-	FApp::GetProjectName();
-
 	const FConfigFile* ConfigFile = GConfig->Find(GEngineIni);
 
 	if (ConfigFile != nullptr)
 	{
 		FInitParams InitParams(FApp::GetProjectName() , *ConfigFile);
-		Initialize(InitParams);
+		Initialize(InitParams, Flags);
 	}
 	else
 	{
 		UE_LOG(LogVirtualization, Error, TEXT("Unable to find a valid engine config file when trying to create the virtualization system"));
 
-		GVirtualizationSystem = MakeUnique<FNullVirtualizationSystem>();
-
 		FConfigFile EmptyConfigFile;
-		FInitParams DummyParams(TEXT(""), EmptyConfigFile);
+		FInitParams InitParams(TEXT(""), EmptyConfigFile);
 
-		GVirtualizationSystem->Initialize(DummyParams);
+		Initialize(InitParams, Flags);
 	}	
 }
 
-void Initialize(const FInitParams& InitParams)
+void Initialize(const FInitParams& InitParams, EInitializationFlags Flags)
 {
+	// If we are not forcing the initialization check to see if we should lazy
+	// initialize or not.
+	if (!EnumHasAnyFlags(Flags, EInitializationFlags::ForceInitialize))
+	{
+#if UE_VIRTUALIZATION_SYSTEM_LAZY_INIT
+		return;
+#else
+		if (FParse::Param(FCommandLine::Get(), TEXT("VA-LazyInit")))
+		{
+			return;
+		}
+		else
+		{
+			bool bLazyInit = false;
+			InitParams.ConfigFile.GetBool(TEXT("Core.ContentVirtualization"), TEXT("LazyInit"), bLazyInit);
+
+			if (bLazyInit)
+			{
+				return;
+			}
+		}
+#endif //UE_VIRTUALIZATION_SYSTEM_LAZY_INIT
+	}
+
+	// Only allow one thread to initialize the system at a time
+	static FCriticalSection InitCS;
+
+	FScopeLock _(&InitCS);
+	
+	if (IVirtualizationSystem::IsInitialized())
+	{
+		// Another thread initialized the system first
+		return;
+	}
+
 	FName SystemName;
 
 #if UE_DISABLE_VIRTUALIZATION_SYSTEM == 0
@@ -184,13 +221,17 @@ void Initialize(const FInitParams& InitParams)
 		Private::IVirtualizationSystemFactory* SystemFactory = FindFactory(SystemName);
 		if (SystemFactory != nullptr)
 		{
-			GVirtualizationSystem = SystemFactory->Create();
-			check(GVirtualizationSystem.IsValid()); // It is assumed that create will always return a valid pointer
+			TUniquePtr<IVirtualizationSystem> NewSystem = SystemFactory->Create();
+			check(NewSystem != nullptr); // It is assumed that create will always return a valid pointer
 
-			if (!GVirtualizationSystem->Initialize(InitParams))
+			if (NewSystem->Initialize(InitParams))
+			{
+				check(!IVirtualizationSystem::IsInitialized());
+				GVirtualizationSystem = MoveTemp(NewSystem);
+			}
+			else
 			{
 				UE_LOG(LogVirtualization, Error, TEXT("Initialization of the virtualization system '%s' failed, falling back to the default implementation"), *SystemName.ToString());
-				GVirtualizationSystem.Reset();
 			}
 		}
 		else
@@ -202,8 +243,11 @@ void Initialize(const FInitParams& InitParams)
 	// If we found no system to create so we will use the fallback system
 	if (!GVirtualizationSystem.IsValid())
 	{
-		GVirtualizationSystem = MakeUnique<FNullVirtualizationSystem>();
-		GVirtualizationSystem->Initialize(InitParams);
+		TUniquePtr<IVirtualizationSystem> NullSystem = MakeUnique<FNullVirtualizationSystem>();
+		NullSystem->Initialize(InitParams);
+
+		check(!IVirtualizationSystem::IsInitialized());
+		GVirtualizationSystem = MoveTemp(NullSystem);
 	}
 }
 
@@ -213,13 +257,20 @@ void Shutdown()
 	UE_LOG(LogVirtualization, Log, TEXT("UE::Virtualization was shutdown"));
 }
 
+
+bool IVirtualizationSystem::IsInitialized()
+{
+	return GVirtualizationSystem != nullptr;
+}
+
 IVirtualizationSystem& IVirtualizationSystem::Get()
 {
 	// For now allow Initialize to be called directly if it was not called explicitly.
-	if (GVirtualizationSystem == nullptr)
+	if (!IsInitialized())
 	{
-		UE_LOG(LogVirtualization, Warning, TEXT("UE::Virtualization::Initialize was not called before UE::Virtualization::IVirtualizationSystem::Get()!"));
-		UE::Virtualization::Initialize();
+		// We need to ForceInitialize at this point to make sure any lazy init flags are ignored
+		UE::Virtualization::Initialize(EInitializationFlags::ForceInitialize);
+		check(IsInitialized());
 	}
 
 	return *GVirtualizationSystem;
diff --git a/Engine/Source/Runtime/Core/Public/Virtualization/VirtualizationSystem.h b/Engine/Source/Runtime/Core/Public/Virtualization/VirtualizationSystem.h
--- a/Engine/Source/Runtime/Core/Public/Virtualization/VirtualizationSystem.h
+++ b/Engine/Source/Runtime/Core/Public/Virtualization/VirtualizationSystem.h
@@ -289,6 +289,18 @@ struct FInitParams
 	/** The config file to load the settings from (will default to GEngineIni) */
 	const FConfigFile& ConfigFile;
 };
+
+/** Optional flags that change how the VA system is initialized */
+enum class EInitializationFlags : uint32
+{
+	/** No flags are set */
+	None			= 0,
+	/** Forces the initialization to occur, ignoring and of the lazy initialization flags */
+	ForceInitialize = 1 << 0
+};
+
+ENUM_CLASS_FLAGS(EInitializationFlags);
+
 /**
  * Creates the global IVirtualizationSystem if it has not already been set up. This can be called explicitly
  * during process start up but it will also be called by IVirtualizationSystem::Get if it detects that the
@@ -296,12 +308,12 @@ struct FInitParams
  * 
  * This version will use the default values of FInitParams.
  */
-CORE_API void Initialize();
+CORE_API void Initialize(EInitializationFlags Flags);
 
 /**
  * This version of ::Initialize takes parameters via the FInitParams structure.
  */
-CORE_API void Initialize(const FInitParams& InitParams);
+CORE_API void Initialize(const FInitParams& InitParams, EInitializationFlags Flags);
 
 /**
  * Shutdowns the global IVirtualizationSystem if it exists. 
@@ -344,7 +356,13 @@ public:
 	 */
 	virtual bool Initialize(const FInitParams& InitParams) = 0;
 
-	/** Gain access to the current virtualization system active for the project */
+	/** Returns true if a virtualization system has been initialized and false if not */
+	static bool IsInitialized();
+
+	/** 
+	 * Gain access to the current virtualization system active for the project. If the system has not yet been 
+	 * initialized then calling this method will initialize it.
+	 */
 	static IVirtualizationSystem& Get();
 
 	/** Poll to see if content virtualization is enabled or not. */
diff --git a/Engine/Source/Runtime/Engine/Private/Analytics/StudioAnalytics.cpp b/Engine/Source/Runtime/Engine/Private/Analytics/StudioAnalytics.cpp
--- a/Engine/Source/Runtime/Engine/Private/Analytics/StudioAnalytics.cpp
+++ b/Engine/Source/Runtime/Engine/Private/Analytics/StudioAnalytics.cpp
@@ -239,60 +239,63 @@ void FStudioAnalytics::FireEvent_Loading(const FString& LoadingName, double Seco
 			}
 
 			// Grab the Virtualization stats
-			IVirtualizationSystem& System = IVirtualizationSystem::Get();
-
-			FPayloadActivityInfo PayloadActivityInfo = System.GetAccumualtedPayloadActivityInfo();
-
-			const FString BaseName = TEXT("Virtualization");
-
-			{
-				FString AttrName = BaseName + TEXT(".Enabled");
-				Attributes.Emplace(MoveTemp(AttrName), System.IsEnabled());
-			}
-
-			{
-				FString AttrName = BaseName + TEXT(".Cache.TimeSpent");
-				Attributes.Emplace(MoveTemp(AttrName), (double)PayloadActivityInfo.Cache.CyclesSpent * FPlatformTime::GetSecondsPerCycle());
-			}
-
-			{
-				FString AttrName = BaseName + TEXT(".Cache.PayloadCount");
-				Attributes.Emplace(MoveTemp(AttrName), (double)PayloadActivityInfo.Cache.PayloadCount);
-			}
-
-			{
-				FString AttrName = BaseName + TEXT(".Cache.TotalBytes");
-				Attributes.Emplace(MoveTemp(AttrName), (double)PayloadActivityInfo.Cache.TotalBytes);
-			}
-
-			{
-				FString AttrName = BaseName + TEXT(".Push.TimeSpent");
-				Attributes.Emplace(MoveTemp(AttrName), (double)PayloadActivityInfo.Push.CyclesSpent * FPlatformTime::GetSecondsPerCycle());
-			}
-
-			{
-				FString AttrName = BaseName + TEXT(".Push.PayloadCount");
-				Attributes.Emplace(MoveTemp(AttrName), (double)PayloadActivityInfo.Push.PayloadCount);
-			}
-
-			{
-				FString AttrName = BaseName + TEXT(".Push.TotalBytes");
-				Attributes.Emplace(MoveTemp(AttrName), (double)PayloadActivityInfo.Push.TotalBytes);
-			}
-
-			{
-				FString AttrName = BaseName + TEXT(".Pull.TimeSpent");
-				Attributes.Emplace(MoveTemp(AttrName), (double)PayloadActivityInfo.Pull.CyclesSpent * FPlatformTime::GetSecondsPerCycle());
-			}
-
-			{
-				FString AttrName = BaseName + TEXT(".Pull.PayloadCount");
-				Attributes.Emplace(MoveTemp(AttrName), (double)PayloadActivityInfo.Pull.PayloadCount);
-			}
-
-			{
-				FString AttrName = BaseName + TEXT(".Pull.TotalBytes");
-				Attributes.Emplace(MoveTemp(AttrName), (double)PayloadActivityInfo.Pull.TotalBytes);
+			if (IVirtualizationSystem::IsInitialized())
+			{
+				IVirtualizationSystem& System = IVirtualizationSystem::Get();
+
+				FPayloadActivityInfo PayloadActivityInfo = System.GetAccumualtedPayloadActivityInfo();
+
+				const FString BaseName = TEXT("Virtualization");
+
+				{
+					FString AttrName = BaseName + TEXT(".Enabled");
+					Attributes.Emplace(MoveTemp(AttrName), System.IsEnabled());
+				}
+
+				{
+					FString AttrName = BaseName + TEXT(".Cache.TimeSpent");
+					Attributes.Emplace(MoveTemp(AttrName), (double)PayloadActivityInfo.Cache.CyclesSpent * FPlatformTime::GetSecondsPerCycle());
+				}
+
+				{
+					FString AttrName = BaseName + TEXT(".Cache.PayloadCount");
+					Attributes.Emplace(MoveTemp(AttrName), (double)PayloadActivityInfo.Cache.PayloadCount);
+				}
+
+				{
+					FString AttrName = BaseName + TEXT(".Cache.TotalBytes");
+					Attributes.Emplace(MoveTemp(AttrName), (double)PayloadActivityInfo.Cache.TotalBytes);
+				}
+
+				{
+					FString AttrName = BaseName + TEXT(".Push.TimeSpent");
+					Attributes.Emplace(MoveTemp(AttrName), (double)PayloadActivityInfo.Push.CyclesSpent * FPlatformTime::GetSecondsPerCycle());
+				}
+
+				{
+					FString AttrName = BaseName + TEXT(".Push.PayloadCount");
+					Attributes.Emplace(MoveTemp(AttrName), (double)PayloadActivityInfo.Push.PayloadCount);
+				}
+
+				{
+					FString AttrName = BaseName + TEXT(".Push.TotalBytes");
+					Attributes.Emplace(MoveTemp(AttrName), (double)PayloadActivityInfo.Push.TotalBytes);
+				}
+
+				{
+					FString AttrName = BaseName + TEXT(".Pull.TimeSpent");
+					Attributes.Emplace(MoveTemp(AttrName), (double)PayloadActivityInfo.Pull.CyclesSpent * FPlatformTime::GetSecondsPerCycle());
+				}
+
+				{
+					FString AttrName = BaseName + TEXT(".Pull.PayloadCount");
+					Attributes.Emplace(MoveTemp(AttrName), (double)PayloadActivityInfo.Pull.PayloadCount);
+				}
+
+				{
+					FString AttrName = BaseName + TEXT(".Pull.TotalBytes");
+					Attributes.Emplace(MoveTemp(AttrName), (double)PayloadActivityInfo.Pull.TotalBytes);
+				}
 			}
 		}
 
diff --git a/Engine/Source/Runtime/Launch/Private/LaunchEngineLoop.cpp b/Engine/Source/Runtime/Launch/Private/LaunchEngineLoop.cpp
--- a/Engine/Source/Runtime/Launch/Private/LaunchEngineLoop.cpp
+++ b/Engine/Source/Runtime/Launch/Private/LaunchEngineLoop.cpp
@@ -2891,7 +2891,7 @@ PRAGMA_ENABLE_DEPRECATION_WARNINGS
 #if WITH_EDITOR
 		// The virtualization system will be initialized when ever IVirtualizationSystem::Get is called
 		// but here we try to explicitly initialize it.
-		UE::Virtualization::Initialize();
+		UE::Virtualization::Initialize(UE::Virtualization::EInitializationFlags::None);
 #endif //WITH_EDITOR
 
 		check(!GDistanceFieldAsyncQueue);
