Fixed memory corruption due to low level allocator trim delegate not thread safe

#rnx

diff --git a/Engine/Source/Runtime/Core/Private/HAL/MallocBinned2.cpp b/Engine/Source/Runtime/Core/Private/HAL/MallocBinned2.cpp
--- a/Engine/Source/Runtime/Core/Private/HAL/MallocBinned2.cpp
+++ b/Engine/Source/Runtime/Core/Private/HAL/MallocBinned2.cpp
@@ -803,6 +803,7 @@ void FMallocBinned2::OnMallocInitialized()
 #if UE_USE_VERYLARGEPAGEALLOCATOR
 	FCoreDelegates::GetLowLevelAllocatorMemoryTrimDelegate().AddLambda([this]()
 		{
+			FScopeLock Lock(&Mutex);
 			CachedOSPageAllocator.FreeAll(&Mutex);
 		}
 	);
