[config] Fixed unsafe multi-threaded access to non-const static paths string during config init. Particular case crashing for me was UserSettingsDir getting called by PerformBasicReplacements.
#rnx

diff --git a/Engine/Source/Runtime/Core/Private/Misc/ConfigCacheIni.cpp b/Engine/Source/Runtime/Core/Private/Misc/ConfigCacheIni.cpp
--- a/Engine/Source/Runtime/Core/Private/Misc/ConfigCacheIni.cpp
+++ b/Engine/Source/Runtime/Core/Private/Misc/ConfigCacheIni.cpp
@@ -4124,9 +4124,11 @@ static FCriticalSection GConfigForPlatformLock;
 #if WITH_EDITOR
 void FConfigCacheIni::AsyncInitializeConfigForPlatforms()
 {
-	// since GConfig just became ready, this will cache something new on first use, and we don't want two threads going at it
+	// make sure any (non-const static) paths the worker threads will use are already initialized
 	FPaths::ProjectDir();
-	FConfigCacheIni::GetCustomConfigString();
+	FPlatformMisc::GeneratedConfigDir(); // also inits FPaths::ProjectSavedDir
+	FConfigContext::EnsureRequiredGlobalPathsHaveBeenInitialized();
+	FPlatformProcess::ApplicationSettingsDir();
 
 	// pre-create all platforms so that the loop below doesn't reallocate anything in the map
 	for (auto Pair : FDataDrivenPlatformInfoRegistry::GetAllPlatformInfos())
diff --git a/Engine/Source/Runtime/Core/Private/Misc/ConfigContext.cpp b/Engine/Source/Runtime/Core/Private/Misc/ConfigContext.cpp
--- a/Engine/Source/Runtime/Core/Private/Misc/ConfigContext.cpp
+++ b/Engine/Source/Runtime/Core/Private/Misc/ConfigContext.cpp
@@ -874,3 +874,8 @@ int32 FConfigFileHierarchy::AddDynamicLayer(const FString& Filename)
 	Emplace(Key, Filename);
 	return Key;
 }
+
+void FConfigContext::EnsureRequiredGlobalPathsHaveBeenInitialized()
+{
+	PerformBasicReplacements(TEXT(""), TEXT("")); // requests user directories and FConfigCacheIni::GetCustomConfigString
+}
diff --git a/Engine/Source/Runtime/Core/Public/Misc/ConfigContext.h b/Engine/Source/Runtime/Core/Public/Misc/ConfigContext.h
--- a/Engine/Source/Runtime/Core/Public/Misc/ConfigContext.h
+++ b/Engine/Source/Runtime/Core/Public/Misc/ConfigContext.h
@@ -117,6 +117,10 @@ public:
 		return Context;
 	}
 
+	/**
+	 * Call to make before attempting parallel config init
+	 */
+	static void EnsureRequiredGlobalPathsHaveBeenInitialized();
 
 	/**
 	 * Use the context to perform the actual load operation. Note that this is where you specify the Ini name (for instance "Engine"), meaning
