Added ASan CSV Metadata

#rnx

diff --git a/Engine/Source/Runtime/Core/Private/ProfilingDebugging/CsvProfiler.cpp b/Engine/Source/Runtime/Core/Private/ProfilingDebugging/CsvProfiler.cpp
--- a/Engine/Source/Runtime/Core/Private/ProfilingDebugging/CsvProfiler.cpp
+++ b/Engine/Source/Runtime/Core/Private/ProfilingDebugging/CsvProfiler.cpp
@@ -2755,6 +2755,7 @@ FCsvProfiler::FCsvProfiler()
 	SetMetadataInternal(TEXT("CPU"), *FPlatformMisc::GetDeviceMakeAndModel());
 	SetMetadataInternal(TEXT("PGOEnabled"), FPlatformMisc::IsPGOEnabled() ? TEXT("1") : TEXT("0"));
 	SetMetadataInternal(TEXT("LoginID"), *FPlatformMisc::GetLoginId());
+	SetMetadataInternal(TEXT("ASan"), USING_ADDRESS_SANITISER ? TEXT("1") : TEXT("0"));
 
 	// Set the device ID if the platform supports it
 	FString DeviceID = FPlatformMisc::GetDeviceId();
