UE-161022 - Editor exits when max retries reached by a IAutomationLatentCommandWithRetriesAndDelays command

- IAutomationLatentCommandWithRetriesAndDelays no longer requests exit in editor, but instead dequeues all commands and adds an error to the existing test.
- Fixed found bug with the macro with no params
- Due to the fact that the interface as well as the FAutomationTestFramework singleton are all defined AutomationTest.h, I had to move IAutomationLatentCommandWithRetriesAndDelays' class declaration after FAutomationTestFramework's.

diff --git a/Engine/Source/Runtime/Core/Public/Misc/AutomationTest.h b/Engine/Source/Runtime/Core/Public/Misc/AutomationTest.h
--- a/Engine/Source/Runtime/Core/Public/Misc/AutomationTest.h
+++ b/Engine/Source/Runtime/Core/Public/Misc/AutomationTest.h
@@ -540,128 +540,6 @@
 	friend class FAutomationTestFramework;
 };
 
-// Extension of IAutomationLatentCommand with delays between attempts
-// if initial command does not succeed. Has a max retry count that can be 
-// overridden. Default is 10 retries with a 1 second delay in between.
-class IAutomationLatentCommandWithRetriesAndDelays : public IAutomationLatentCommand
-{	
-public:
-	virtual ~IAutomationLatentCommandWithRetriesAndDelays() {}
-
-	// Base Update override with delay logic built in. Can be further overridden in child
-	// classes
-	virtual bool Update() override
-	{
-		if (IsDelayTimerRunning())
-		{
-			return false;
-		}
-
-		// pre-increment iteration so that its 1 based instead of 0 based for readability
-		CurrentIteration++;
-
-		if (!CanRetry())
-		{
-			// Retry count has been exceeded
-			UE_LOG(LogLatentCommands, Error, TEXT("%s All Retries Have Been Exhausted. Command Has Failed"), *GetCommandName());
-#if WITH_EDITOR
-			FPlatformMisc::RequestExitWithStatus(false, -1);
-#endif
-			return true;
-		}
-
-		ResetDelayTimer();
-
-		// auto-logging for limited retries and infinite retries
-		if (MaxRetries == 0) 
-		{
-			UE_LOG(LogLatentCommands, Log, TEXT("%s Executing Attempt %d."), *GetCommandName(), CurrentIteration);
-		}
-		else
-		{
-			UE_LOG(LogLatentCommands, Log, TEXT("%s Executing Attempt %d of %d."), *GetCommandName(), CurrentIteration, MaxRetries);
-		}
-
-		if (Execute())
-		{
-			// completion log message, logs total time taken.
-			UE_LOG(LogLatentCommands, Log, TEXT("%s Completed Successfully, total run time: %f."), *GetCommandName(), GetCurrentRunTime());
-			return true;
-		}
-
-		return false;
-	}
-
-	// Pure virtual method that must be overridden.
-	// This is the actual command logic.
-	virtual bool Execute() = 0;
-
-private:
-	// To keep track of which iteration we are on. 1 based
-	// To allow for more readable logs. Attempt 1 of N instead
-	// of 0 of N.
-	int32 CurrentIteration = 0;
-
-protected:
-	// default constructor
-	IAutomationLatentCommandWithRetriesAndDelays() {}
-
-	// parameterized constructor
-	IAutomationLatentCommandWithRetriesAndDelays(const FString InCommandClassName, const int32 InMaxRetries, const double InWaitTimeBetweenRuns)
-		:CommandClassName(InCommandClassName)
-		,MaxRetries(InMaxRetries)
-		,DelayTimeInSeconds(InWaitTimeBetweenRuns)
-	{
-		// set first run start time so that we don't start off in a delay. Its fine if its negative.
-		DelayStartTime = FPlatformTime::Seconds() - InWaitTimeBetweenRuns;
-	}
-
-	// Resets the Delay Timer by setting start time to now (In game time).
-	void ResetDelayTimer()
-	{
-		DelayStartTime = FPlatformTime::Seconds();
-	}
-
-	// Determines if timer is running, pausing execution in a non-blocking way.
-	bool IsDelayTimerRunning() const
-	{
-		// time elapsed < Delay time
-		return (FPlatformTime::Seconds() - DelayStartTime) < DelayTimeInSeconds;
-	}
-
-	// Determines if we should execute the command. If MaxRetries is set to 0, 
-	// indicates infinite retries.
-	bool CanRetry() const
-	{
-		return (CurrentIteration <= MaxRetries) || MaxRetries == 0;
-	}
-
-	// Returns the command name 
-	FString GetCommandName(const FString& TestName = FString(TEXT(""))) const
-	{
-		if (!TestName.IsEmpty())
-		{
-			return FString::Printf(TEXT("Test: Unknown - Command: %s - "), *CommandClassName);
-		}
-		
-		return FString::Printf(TEXT("Test: %s - Command: %s - "), *TestName, *CommandClassName);
-	}
-
-	// Command Name stored for easy logging
-	const FString CommandClassName = "UnknownCommand";
-	
-	// Times to retry command before reporting command failure.
-	// Defaults to 10, but can be overridden.
-	const int32 MaxRetries = 10;
-	
-	// Time in between UpdateDelayed calls. 
-	// Default is 1 second but can be overridden.
-	const double DelayTimeInSeconds = 1.0;
-	
-	// Time that the Delay Timer started.
-	double DelayStartTime = 0.0;
-};
-
 /**
  * A simple latent command that runs the provided function on another thread
  */
@@ -3215,7 +3093,7 @@
 	{ \
 	public: \
 	CommandName(int32 InRetryCount, double InWaitTimeBetweenRuns) \
-	: IAutomationLatentCommandWithRetriesAndDelays(#CommandName, InTest, InRetryCount, InWaitTimeBetweenRuns) \
+	: IAutomationLatentCommandWithRetriesAndDelays(#CommandName, InRetryCount, InWaitTimeBetweenRuns) \
 		{} \
 		virtual ~CommandName() \
 		{} \
@@ -3771,3 +3649,177 @@
 	float Timeout;
 };
 
+// Extension of IAutomationLatentCommand with delays between attempts
+// if initial command does not succeed. Has a max retry count that can be 
+// overridden. Default is 10 retries with a 1 second delay in between.
+// To protect against misuse of unlimited retries, has a Max execution 
+// time of 5 minutes. This can be overridden, but not recommended. Only 
+// do this if you're absolutely certain it will not cause an infinite loop
+class IAutomationLatentCommandWithRetriesAndDelays : public IAutomationLatentCommand
+{
+public:
+	virtual ~IAutomationLatentCommandWithRetriesAndDelays() {}
+
+	// Base Update override with delay logic built in. Can be further overridden in child
+	// classes
+	virtual bool Update() override
+	{
+		if (bHasUnlimitedRetries)
+		{
+			// command has unlimited retries, so need to check if max run time has been exceeded
+			if (HasExceededMaxTotalRunTime())
+			{
+				// Command run time has exceeded max total run time. Stop further commands and log error so the test fails.
+				FAutomationTestFramework::Get().DequeueAllCommands();
+				FString ErrorMessageText = FString::Printf(TEXT("%s has failed due to exceeding the max allowed run time of %f seconds. \
+This may be due to an error, or having a single command attempt to do too many things. If this is not due to an error, consider breaking \
+up this command into multiple, smaller commands."), *GetTestAndCommandName(), MaxTotalRunTimeInSeconds);
+				// Must log here so that Gauntlet can also pick up the error for its report. Otherwise, it will only show up in the Automation log.
+				UE_LOG(LogLatentCommands, Error, TEXT("%s"), *ErrorMessageText);
+				GetCurrentTest()->AddError(ErrorMessageText);
+				return true;
+			}
+		}
+
+		if (IsDelayTimerRunning())
+		{
+			return false;
+		}
+
+		// pre-increment iteration so that its 1 based instead of 0 based for readability
+		CurrentIteration++;
+
+		if (!CanRetry())
+		{
+			// Retry count has been exceeded, dequeue further commands and log error to test.
+			FAutomationTestFramework::Get().DequeueAllCommands();
+			// Must log here so that Gauntlet can also pick up the error for its report. Otherwise, it will only show up in the Automation log.
+			FString ErrorMessageText = FString::Printf(TEXT("%s Latent command with retries and delays has failed after %d retries"), *GetTestAndCommandName(), MaxRetries);
+			UE_LOG(LogLatentCommands, Error, TEXT("%s"), *ErrorMessageText);
+			GetCurrentTest()->AddError(ErrorMessageText);
+			return true;
+		}
+
+		ResetDelayTimer();
+
+		// auto-logging for limited retries and unlimited retries
+		if (bHasUnlimitedRetries)
+		{
+			UE_LOG(LogLatentCommands, Log, TEXT("%s Executing Attempt %d."), *GetTestAndCommandName(), CurrentIteration);
+		}
+		else
+		{
+			UE_LOG(LogLatentCommands, Log, TEXT("%s Executing Attempt %d of %d."), *GetTestAndCommandName(), CurrentIteration, MaxRetries);
+		}
+
+		if (Execute())
+		{
+			// completion log message, logs total time taken.
+			UE_LOG(LogLatentCommands, Log, TEXT("%s Completed Successfully, total run time: %f seconds."), *GetTestAndCommandName(), GetCurrentRunTime());
+
+			return true;
+		}
+
+		return false;
+	}
+
+	// Pure virtual method that must be overridden.
+	// This is the actual command logic.
+	virtual bool Execute() = 0;
+
+private:
+	// To keep track of which iteration we are on. 1 based
+	// To allow for more readable logs. Attempt 1 of N instead
+	// of 0 of N.
+	int32 CurrentIteration = 0;
+
+protected:
+	// default constructor
+	IAutomationLatentCommandWithRetriesAndDelays() {}
+
+	// parameterized constructor
+	IAutomationLatentCommandWithRetriesAndDelays(const FString InCommandClassName, const int32 InMaxRetries, const double InWaitTimeBetweenRuns)
+		:CommandClassName(InCommandClassName)
+		, MaxRetries(InMaxRetries < 0 ? 0 : InMaxRetries)
+		, bHasUnlimitedRetries(InMaxRetries == 0)
+		, DelayTimeInSeconds(InWaitTimeBetweenRuns)
+	{
+		// set first run start time so that we don't start off in a delay. Its fine if its negative.
+		DelayStartTime = FPlatformTime::Seconds() - InWaitTimeBetweenRuns;
+
+		// warn if command has unlimited retries MaxTotalRunTime
+		if (bHasUnlimitedRetries)
+		{
+			// Must log here so that Gauntlet can also pick up the error for its report. Otherwise, it will only show up in the Automation log.
+			UE_LOG(LogLatentCommands, Warning, TEXT("%s has been set to Unlimited retries. Will be using MaxTotalRunTime to prevent \
+running forever. Default time is set at 300 seconds. If this is not enough time, make sure to override this in your Execute() loop."), *GetTestAndCommandName());
+		}
+	}
+
+	// Resets the Delay Timer by setting start time to now (In game time).
+	void ResetDelayTimer()
+	{
+		DelayStartTime = FPlatformTime::Seconds();
+	}
+
+	// Determines if timer is running, pausing execution in a non-blocking way.
+	bool IsDelayTimerRunning() const
+	{
+		// time elapsed < Delay time
+		return (FPlatformTime::Seconds() - DelayStartTime) < DelayTimeInSeconds;
+	}
+
+	// Returns if we have exceeded max allowed total run time for this latent command
+	double HasExceededMaxTotalRunTime()
+	{
+		return MaxTotalRunTimeInSeconds - GetCurrentRunTime() < 0;
+	}
+
+	// Determines if we should execute the command. If MaxRetries is set to 0, 
+	// indicates unlimited retries.
+	bool CanRetry() const
+	{
+		return (CurrentIteration <= MaxRetries) || MaxRetries == 0;
+	}
+
+	// Returns the command name 
+	FString GetTestAndCommandName() const
+	{
+		// Otherwise use the provided string
+		return FString::Printf(TEXT("Test: %s - Command: %s - "), *GetCurrentTest()->GetTestFullName(), *CommandClassName);
+	}
+
+	FAutomationTestBase* GetCurrentTest() const
+	{
+		return FAutomationTestFramework::Get().GetCurrentTest();
+	}
+
+	void OverrideMaxTotalRunTimeInSeconds(double OverrideValue)
+	{
+		MaxTotalRunTimeInSeconds = OverrideValue;
+	}
+
+	// Command Name stored for easy logging
+	const FString CommandClassName = "UnknownCommand";
+
+	// Times to retry command before reporting command failure.
+	// Defaults to 10, but can be overridden.
+	const int32 MaxRetries = 10;
+
+	// Indicates that this latent command has unlimited retries
+	// Used to implement the MaxTotalRunTimeInSeconds logic
+	const bool bHasUnlimitedRetries = false;
+
+	// Time in between UpdateDelayed calls. 
+	// Default is 1 second but can be overridden.
+	const double DelayTimeInSeconds = 1.0;
+
+	// Time that the Delay Timer started.
+	double DelayStartTime = 0.0;	
+
+private:
+	// Max total run time for command 
+	// Default is 5 minutes, but can be overridden
+	double MaxTotalRunTimeInSeconds = 300.0;
+	
+};
