UE-161022 - Editor exits when max retries reached by a IAutomationLatentCommandWithRetriesAndDelays command

- IAutomationLatentCommandWithRetriesAndDelays no longer requests exit in editor, but instead dequeues all commands and adds an error to the existing test.
- Fixed found bug with the macro with no params
- Due to the fact that the interface as well as the FAutomationTestFramework singleton are all defined AutomationTest.h, I had to move IAutomationLatentCommandWithRetriesAndDelays' class declaration after FAutomationTestFramework's.
- Updated GetCommandName(FString&) to GetTestAndCommandName() now that I can get the current test's name without having to pass it in from the outside. Updated all calls to use new method signature.
**Updated** Due to PS4 Error in Horde, changed return value for HasExceededMaxTotalRunTime() to bool.

diff --git a/Engine/Source/Runtime/Core/Public/Misc/AutomationTest.h b/Engine/Source/Runtime/Core/Public/Misc/AutomationTest.h
--- a/Engine/Source/Runtime/Core/Public/Misc/AutomationTest.h
+++ b/Engine/Source/Runtime/Core/Public/Misc/AutomationTest.h
@@ -3770,7 +3770,7 @@ running forever. Default time is set at 300 seconds. If this is not enough time,
 	}
 
 	// Returns if we have exceeded max allowed total run time for this latent command
-	double HasExceededMaxTotalRunTime()
+	bool HasExceededMaxTotalRunTime()
 	{
 		return MaxTotalRunTimeInSeconds - GetCurrentRunTime() < 0;
 	}
