Hooked up JWT verification process (RS256 and PS256 algorithms), without encryption key caching.  Ported Base64url conversion from EOS SDK.

diff --git a/Engine/Plugins/Experimental/JWT/JWT.uplugin b/Engine/Plugins/Experimental/JWT/JWT.uplugin
--- a/Engine/Plugins/Experimental/JWT/JWT.uplugin
+++ b/Engine/Plugins/Experimental/JWT/JWT.uplugin
@@ -21,5 +21,12 @@
 			"Type" : "Runtime",
 			"LoadingPhase" : "Default"
 		}
+	],
+	"Plugins": [
+		{
+			"Name": "PlatformCrypto",
+			"Enabled": true
+		}
 	]
+
 }
diff --git a/Engine/Plugins/Experimental/JWT/Source/JWT/JWT.Build.cs b/Engine/Plugins/Experimental/JWT/Source/JWT/JWT.Build.cs
--- a/Engine/Plugins/Experimental/JWT/Source/JWT/JWT.Build.cs
+++ b/Engine/Plugins/Experimental/JWT/Source/JWT/JWT.Build.cs
@@ -1,20 +1,30 @@
 // Copyright Epic Games, Inc. All Rights Reserved.
 
+using System.IO;
+
 namespace UnrealBuildTool.Rules
 {
 	public class JWT : ModuleRules
 	{
 		public JWT(ReadOnlyTargetRules Target) : base(Target)
 		{
+			OptimizeCode = CodeOptimization.InShippingBuildsOnly;
+
 			PublicIncludePaths.AddRange(
 				new string[] {
 					// ... add public include paths required here ...
+					Path.Combine(EngineDirectory, "Plugins/Experimental/PlatformCrypto/Source/PlatformCrypto/Public/"),
+					Path.Combine(EngineDirectory, "Plugins/Experimental/PlatformCrypto/Source/PlatformCryptoOpenSSL/Public/"),
+					Path.Combine(EngineDirectory, "Plugins/Experimental/PlatformCrypto/Source/PlatformCryptoTypes/Public/"),
 				}
 			);
 
 			PrivateIncludePaths.AddRange(
 				new string[] {
 					// ... add other private include paths required here ...
+					Path.Combine(EngineDirectory, "Plugins/Experimental/PlatformCrypto/Source/PlatformCrypto/Private/"),
+					Path.Combine(EngineDirectory, "Plugins/Experimental/PlatformCrypto/Source/PlatformCryptoOpenSSL/Private/"),
+					Path.Combine(EngineDirectory, "Plugins/Experimental/PlatformCrypto/Source/PlatformCryptoTypes/Private/"),
 				}
 			);
 
@@ -24,6 +34,9 @@ namespace UnrealBuildTool.Rules
 					"Core",
 					// ... add other public dependencies that you statically link with here ...
 					"Json",
+					"PlatformCrypto",
+					"PlatformCryptoOpenSSL",
+					"PlatformCryptoTypes",
 				}
 			);
 
@@ -31,6 +44,9 @@ namespace UnrealBuildTool.Rules
 				new string[]
 				{
 					// ... add private dependencies that you statically link with here ...
+					"PlatformCrypto",
+					"PlatformCryptoOpenSSL",
+					"PlatformCryptoTypes",
 				}
 			);
 
diff --git a/Engine/Plugins/Experimental/JWT/Source/JWT/Private/IJwt.cpp b/Engine/Plugins/Experimental/JWT/Source/JWT/Private/IJwt.cpp
--- a/Engine/Plugins/Experimental/JWT/Source/JWT/Private/IJwt.cpp
+++ b/Engine/Plugins/Experimental/JWT/Source/JWT/Private/IJwt.cpp
@@ -32,14 +32,14 @@ void FJwt::ShutdownModule()
 
 namespace UE::JWT
 {
-	TOptional<FJsonWebToken> FromString(const FStringView InEncodedJsonWebToken, const bool bIsSignatureEncoded)
+	TOptional<FJsonWebToken> FromString(const FStringView InEncodedJsonWebToken)
 	{
-		return FJsonWebToken::FromString(InEncodedJsonWebToken, bIsSignatureEncoded);
+		return FJsonWebToken::FromString(InEncodedJsonWebToken);
 	}
 
 
-	bool FromString(const FStringView InEncodedJsonWebToken, FJsonWebToken& OutJsonWebToken, const bool bIsSignatureEncoded)
+	bool FromString(const FStringView InEncodedJsonWebToken, FJsonWebToken& OutJsonWebToken)
 	{
-		return FJsonWebToken::FromString(InEncodedJsonWebToken, OutJsonWebToken, bIsSignatureEncoded);
+		return FJsonWebToken::FromString(InEncodedJsonWebToken, OutJsonWebToken);
 	}
 }
diff --git a/Engine/Plugins/Experimental/JWT/Source/JWT/Private/JsonWebToken.cpp b/Engine/Plugins/Experimental/JWT/Source/JWT/Private/JsonWebToken.cpp
--- a/Engine/Plugins/Experimental/JWT/Source/JWT/Private/JsonWebToken.cpp
+++ b/Engine/Plugins/Experimental/JWT/Source/JWT/Private/JsonWebToken.cpp
@@ -2,8 +2,11 @@
 
 #include "JsonWebToken.h"
 #include "Algo/Count.h"
+#include "Algo/Reverse.h"
 #include "Misc/Base64.h"
 #include "Serialization/JsonSerializer.h"
+#include "IPlatformCrypto.h"
+#include "PlatformCryptoIncludes.h"
 
 // JWT header field names
 const TCHAR *const FJsonWebToken::HEADER_TYPE = TEXT("typ");
@@ -20,6 +23,10 @@ const TCHAR *const FJsonWebToken::CLAIM_EXPIRATION = TEXT("exp");
 const TCHAR *const FJsonWebToken::CLAIM_SUBJECT = TEXT("sub");
 const TCHAR *const FJsonWebToken::CLAIM_AUDIENCE = TEXT("aud");
 
+// Expected/supported values
+constexpr const TCHAR* TYPE_VALUE_JWT = TEXT("JWT");
+constexpr const TCHAR* ALGORITHM_VALUE_RS256 = TEXT("RS256"); // RSASSA using SHA-256 digest
+constexpr const TCHAR* ALGORITHM_VALUE_PS256 = TEXT("PS256"); // RSASSA-PSS using SHA-256 digest, and MGF1 with SHA-256
 
 namespace
 {
@@ -67,31 +74,31 @@ namespace
 		return true;
 	}
 
-	bool Base64Decode(const FStringView InSource, FString& OutDest)
+	bool Base64UrlDecode(const FStringView InSource, FString& OutDest)
 	{
 		if (InSource.IsEmpty())
 		{
 			return false;
 		}
 
-		return FBase64::Decode(FString(InSource), OutDest);
+		return FBase64::Decode(FString(InSource), OutDest, EBase64Mode::UrlSafe);
 	}
 
-	bool Base64Decode(const FStringView InSource, TArray<uint8>& OutDest)
+	bool Base64UrlDecode(const FStringView InSource, TArray<uint8>& OutDest)
 	{
 		if (InSource.IsEmpty())
 		{
 			return false;
 		}
 
-		return FBase64::Decode(FString(InSource), OutDest);
+		return FBase64::Decode(FString(InSource), OutDest, EBase64Mode::UrlSafe);
 	}
 
 	bool StringViewToBytes(const FStringView In, TArray<uint8>& OutBytes, const bool IsEncoded)
 	{
 		if (IsEncoded)
 		{
-			return Base64Decode(In, OutBytes);
+			return Base64UrlDecode(In, OutBytes);
 		}
 
 		OutBytes.Reserve(In.Len());
@@ -130,7 +137,7 @@ void FJsonWebToken::DumpJsonObject(const FJsonObject& InJsonObject)
 			continue;
 		}
 
-		// Present the JSON value in a printable form; TryGetString() is overridden for
+		// Present the JSON value in a printable form; TryGetString() is overloaded for
 		// non-compound types (bool, number, string).
 		FString OutString;
 		if (Value->TryGetString(OutString))
@@ -155,7 +162,7 @@ void FJsonWebToken::DumpJsonObject(const FJsonObject& InJsonObject)
 			continue;
 		}
 
-		UE_LOG(LogJwt, VeryVerbose, TEXT("  %s [unknown type])"), *Name);
+		UE_LOG(LogJwt, VeryVerbose, TEXT("  %s [unknown type]"), *Name);
 	}
 }
 
@@ -188,7 +195,7 @@ TSharedPtr<FJsonObject> FJsonWebToken::ParseEncodedJson(const FStringView InEnco
 	TSharedPtr<FJsonObject> ParsedObj;
 
 	FString DecodedJson;
-	if (Base64Decode(InEncodedJson, DecodedJson))
+	if (Base64UrlDecode(InEncodedJson, DecodedJson))
 	{
 		ParsedObj = FromJson(DecodedJson);
 	}
@@ -196,9 +203,9 @@ TSharedPtr<FJsonObject> FJsonWebToken::ParseEncodedJson(const FStringView InEnco
 	return ParsedObj;
 }
 
-bool FJsonWebToken::FromString(const FStringView InEncodedJsonWebToken, FJsonWebToken& OutJsonWebToken, const bool bIsSignatureEncoded)
+bool FJsonWebToken::FromString(const FStringView InEncodedJsonWebToken, FJsonWebToken& OutJsonWebToken)
 {
-	TOptional<FJsonWebToken> JsonWebToken = FromString(InEncodedJsonWebToken, bIsSignatureEncoded);
+	TOptional<FJsonWebToken> JsonWebToken = FromString(InEncodedJsonWebToken);
 	if (!JsonWebToken.IsSet())
 	{
 		return false;
@@ -208,7 +215,7 @@ bool FJsonWebToken::FromString(const FStringView InEncodedJsonWebToken, FJsonWeb
 	return true;
 }
 
-TOptional<FJsonWebToken> FJsonWebToken::FromString(const FStringView InEncodedJsonWebToken, const bool bIsSignatureEncoded)
+TOptional<FJsonWebToken> FJsonWebToken::FromString(const FStringView InEncodedJsonWebToken)
 {
 	// Check for the correct number of dots.
 	if (Algo::Count(InEncodedJsonWebToken, TEXT('.')) != 2)
@@ -250,7 +257,7 @@ TOptional<FJsonWebToken> FJsonWebToken::FromString(const FStringView InEncodedJs
 	{
 		UE_LOG(LogJwt, VeryVerbose, TEXT("[FJsonWebToken::FromString] Decoding JWT signature."));
 		TArray<uint8> SignatureBytes;
-		if (!StringViewToBytes(SignaturePart, SignatureBytes, bIsSignatureEncoded))
+		if (!Base64UrlDecode(SignaturePart, SignatureBytes))
 		{
 			UE_LOG(LogJwt, Verbose, TEXT("[FJsonWebToken::FromString] Failed to decode the signature."));
 			return {};
@@ -306,10 +313,147 @@ bool FJsonWebToken::GetClaim(const FStringView InName, TSharedPtr<FJsonValue>& O
 	return OutClaim.IsValid();
 }
 
+bool FJsonWebToken::GetIssuer(FString& OutValue) const
+{
+	return Payload->TryGetStringField(CLAIM_ISSUER, OutValue);
+}
+
+bool FJsonWebToken::GetIssuedAt(double& OutValue) const
+{
+	return Payload->TryGetNumberField(CLAIM_ISSUED_AT, OutValue);
+}
+
+bool FJsonWebToken::GetExpiration(double& OutValue) const
+{
+	return Payload->TryGetNumberField(CLAIM_EXPIRATION, OutValue);
+}
+
+bool FJsonWebToken::GetSubject(FString& OutValue) const
+{
+	return Payload->TryGetStringField(CLAIM_SUBJECT, OutValue);
+}
+
+bool FJsonWebToken::GetAudience(FString& OutValue) const
+{
+	return Payload->TryGetStringField(CLAIM_AUDIENCE, OutValue);
+}
+
+bool IsSupportedAlgorithm(const FString& Algorithm)
+{
+	return (ALGORITHM_VALUE_RS256 == Algorithm) || (ALGORITHM_VALUE_PS256 == Algorithm);
+}
 
-bool FJsonWebToken::Verify() const
+bool DigestVerify(FEncryptionContextOpenSSL& EncryptionContext, const FString& Algorithm, const TArrayView<const uint8> Message, TArray<uint8>& Signature, FRSAKeyHandle Key)
 {
-	// Signature verification is yet to be implemented, so return true for now.
-	// TODO: Add support for RSA crypto algorithm (PlatformCrypto plugin) and key management so that signature verification can be performed.
+	if (ALGORITHM_VALUE_RS256 == Algorithm)
+	{
+		return EncryptionContext.DigestVerify_RS256(Message, Signature, Key);
+	}
+
+	if (ALGORITHM_VALUE_PS256 == Algorithm)
+	{
+		return EncryptionContext.DigestVerify_PS256(Message, Signature, Key);
+	}
+
+	return false;
+}
+
+bool FJsonWebToken::Verify(const TArrayView<const uint8> InKeyExponent, const TArrayView<const uint8> InKeyModulus) const
+{
+	if (!Signature.IsSet())
+	{
+		UE_LOG(LogJwt, Verbose, TEXT("[FJsonWebToken::Verify] Failed because signature is empty"));
+		return false;
+	}
+
+	FString Algorithm;
+	if (!GetAlgorithm(Algorithm))
+	{
+		UE_LOG(LogJwt, Verbose, TEXT("[FJsonWebToken::Verify] Failed to find signing algorithm"));
+		return false;
+	}
+
+	if (!IsSupportedAlgorithm(Algorithm))
+	{
+		UE_LOG(LogJwt, Verbose, TEXT("[FJsonWebToken::Verify] Signing algorithm is not currently supported: '%s'"), *Algorithm);
+		return false;
+	}
+
+	TUniquePtr<FEncryptionContext> EncryptionContext = IPlatformCrypto::Get().CreateContext();
+	if (!EncryptionContext)
+	{
+		UE_LOG(LogJwt, Verbose, TEXT("[FJsonWebToken::Verify] Failed to create encryption context"));
+		return false;
+	}
+
+	// Create the public key from the exponent and modulus.
+	// InKeyExponent and InKeyModulus must be in little-endian format (required by PlatformCrypto module).
+	const static TArrayView<const uint8> EmptyPrivateKeyExponent;
+	FRSAKeyHandle RsaKeyHandle = EncryptionContext->CreateKey_RSA(InKeyExponent, EmptyPrivateKeyExponent, InKeyModulus);
+	if (!RsaKeyHandle)
+	{
+		UE_LOG(LogJwt, Verbose, TEXT("[FJsonWebToken::Verify] Failed to create RSA key"));
+		return false;
+	}
+
+	// Split out the token string into header, payload, and signature parts.
+	FStringView EncodedHeader;
+	FStringView EncodedPayload;
+	FStringView EncodedSignature;
+	if (!SplitEncodedJsonWebTokenString(EncodedJsonWebToken, EncodedHeader, EncodedPayload, EncodedSignature))
+	{
+		return false;
+	}
+
+	// Decode JWT signature from Base64Url form to raw bytes.
+	TArray<uint8> DecodedSignatureBytes;
+	Base64UrlDecode(EncodedSignature, DecodedSignatureBytes); // Big-endian.
+
+	// Verify the digest hash against the signature.
+	const FString EncodedContent = FString::Printf(TEXT("%s.%s"), *FString(EncodedHeader), *FString(EncodedPayload));
+	const TArray<uint8> EncodedContentBytes(reinterpret_cast<uint8*>(TCHAR_TO_ANSI(*EncodedContent)), EncodedContent.Len());
+
+	// Select the appropriate algorithm and verify the JWT content against the signature with it.
+	const bool bIsVerified = DigestVerify(*EncryptionContext.Get(), Algorithm, EncodedContentBytes, DecodedSignatureBytes, RsaKeyHandle);
+
+	EncryptionContext->DestroyKey_RSA(RsaKeyHandle);
+
+	if (!bIsVerified)
+	{
+		UE_LOG(LogJwt, Verbose, TEXT("[FJsonWebToken::Verify] Failed to verify signature"));
+		return false;
+	}
+
 	return true;
 }
+
+bool FJsonWebToken::Verify(const FStringView InKeyExponent, const FStringView InKeyModulus, const bool bIsBase64UrlEncoded, const bool bIsBigEndian) const
+{
+	TArray<uint8> ExponentBytes;
+	if (!StringViewToBytes(InKeyExponent, ExponentBytes, bIsBase64UrlEncoded))
+	{
+		UE_LOG(LogJwt, Verbose, TEXT("[FJsonWebToken::Verify] Failed to parse the public key exponent."));
+		return false;
+	}
+
+	TArray<uint8> ModulusBytes;
+	if (!StringViewToBytes(InKeyModulus, ModulusBytes, bIsBase64UrlEncoded))
+	{
+		UE_LOG(LogJwt, Verbose, TEXT("[FJsonWebToken::Verify] Failed to parse the key modulus."));
+		return false;
+	}
+
+	// When retrieved from the web, the key will be in big-endian form.  PlatformCrypto requires little-endian, so reverse here.
+	if (bIsBigEndian)
+	{
+		Algo::Reverse(ExponentBytes);
+		Algo::Reverse(ModulusBytes);
+	}
+
+	return Verify(ExponentBytes, ModulusBytes);
+}
+
+bool FJsonWebToken::Verify(const FStringView InKeyExponent, const FStringView InKeyModulus) const
+{
+	return Verify(InKeyExponent, InKeyModulus, true, true);
+}
diff --git a/Engine/Plugins/Experimental/JWT/Source/JWT/Public/IJwt.h b/Engine/Plugins/Experimental/JWT/Source/JWT/Public/IJwt.h
--- a/Engine/Plugins/Experimental/JWT/Source/JWT/Public/IJwt.h
+++ b/Engine/Plugins/Experimental/JWT/Source/JWT/Public/IJwt.h
@@ -46,7 +46,7 @@ public:
 
 namespace UE::JWT
 {
-	JWT_API TOptional<FJsonWebToken> FromString(const FStringView InEncodedJsonWebToken, const bool bIsSignatureEncoded);
+	JWT_API TOptional<FJsonWebToken> FromString(const FStringView InEncodedJsonWebToken);
 
-	JWT_API bool FromString(const FStringView InEncodedJWT, FJsonWebToken& OutJsonWebToken, const bool bIsSignatureEncoded);
+	JWT_API bool FromString(const FStringView InEncodedJWT, FJsonWebToken& OutJsonWebToken);
 }
diff --git a/Engine/Plugins/Experimental/JWT/Source/JWT/Public/JsonWebToken.h b/Engine/Plugins/Experimental/JWT/Source/JWT/Public/JsonWebToken.h
--- a/Engine/Plugins/Experimental/JWT/Source/JWT/Public/JsonWebToken.h
+++ b/Engine/Plugins/Experimental/JWT/Source/JWT/Public/JsonWebToken.h
@@ -13,15 +13,28 @@ public:
 	 * Creates a JWT from the provided string.
 	 * The string must consist of 3 base64 url encoded parts: a header, payload, and signature.
 	 * The parts must be split by a period character.
-	 * The signature part is optional. If the signature is excluded, the string must still contain a period character in its place.
+	 * The signature part is optional. If the signature is excluded, the string must still contain
+	 * a period character in its place.
 	 * Valid formats: "header.payload.signature" and "header.payload."
 	 *
 	 * @param InJsonWebTokenString The string to decode.
-	 * @return An optional set with a FJsonWebToken if the JWT was successfully decoded.
+	 * @return An optional set with a FJsonWebToken if the JWT was generated successfully.
 	 */
-	static TOptional<FJsonWebToken> FromString(const FStringView InEncodedJsonWebToken, const bool bIsSignatureEncoded);
+	static TOptional<FJsonWebToken> FromString(const FStringView InEncodedJsonWebToken);
 
-	static bool FromString(const FStringView InEncodedJsonWebToken, FJsonWebToken& OutJsonWebToken, const bool bIsSignatureEncoded);
+	/**
+	 * Creates a JWT from the provided string.
+	 * The string must consist of 3 base64 url encoded parts: a header, payload, and signature.
+	 * The parts must be split by a period character.
+	 * The signature part is optional. If the signature is excluded, the string must still contain
+	 * a period character in its place.
+	 * Valid formats: "header.payload.signature" and "header.payload."
+	 *
+	 * @param InJsonWebTokenString The string to decode.
+	 * @param OutJsonWebToken The generated JWT (FJsonWebToken instance), if input string was decoded successfully.
+	 * @return True if JWT was generated successfully, otherwise false.
+	 */
+	static bool FromString(const FStringView InEncodedJsonWebToken, FJsonWebToken& OutJsonWebToken);
 
 	/**
 	 * Gets the type.
@@ -71,14 +84,82 @@ public:
 	}
 
 	/**
-	 * Perform verification of the encoded header and encoded payload.
-	 * The is done using the cryptographic algorithm specified in the header, and a secret value obtained
-	 * by key-lookup, to generate a signature.  The generated signature must be identical to the signature
-	 * provided in the JWT for verification to succeed.
+	 * Gets the issuer (`iss`) claim from the payload, if present.
 	 * 
+	 * @param OutValue the issuer value, if successful.
+	 * @return True on success (claim exists), otherwise false.
+	 */
+	bool GetIssuer(FString& OutValue) const;
+
+	/**
+	 * Gets the issued at (`iat`) claim from the payload, if present.
+	 *
+	 * @param OutValue the issued at value, if successful.
+	 * @return True on success (claim exists), otherwise false.
+	 */
+	bool GetIssuedAt(double& OutValue) const;
+
+	/**
+	 * Gets the expiriation time (`exp`) claim from the payload, if present.
+	 *
+	 * @param OutValue the expiration time value, if successful.
+	 * @return True on success (claim exists), otherwise false.
+	 */
+	bool GetExpiration(double& OutValue) const;
+
+	/**
+	 * Gets the subject (`sub`) claim from the payload, if present.
+	 *
+	 * @param OutValue the subject value, if successful.
+	 * @return True on success (claim exists), otherwise false.
+	 */
+	bool GetSubject(FString& OutValue) const;
+
+	/**
+	 * Gets the audience (`aud`) claim from the payload, if present.
+	 *
+	 * @param OutValue the audience value, if successful.
+	 * @return True on success (claim exists), otherwise false.
+	 */
+	bool GetAudience(FString& OutValue) const;
+
+	/**
+	 * Verifies the signature against the header and content, using the given public key info, which
+	 * may be obtained from a JWK object.
+	 * Assumes that input arrays are in little-endian byte order.
+	 * The cryptographic algorithm used for verification is specified in this JsonWebToken's header.
+	 *
+	 * @see GetAlgorithm
+	 *
+	 * @param InKeyExponent The e (exponent) value of the public key.
+	 * @param InKeyModulus The n (modulus) value of the public key.
+	 * @return true if the header and payload were verified successfully (signature match), otherwise false.
+	 */
+	bool Verify(const TArrayView<const uint8> InKeyExponent, const TArrayView<const uint8> InKeyModulus) const;
+
+	/**
+	 * Verifies the signature against the header and content, using the given public key info, which
+	 * may be obtained from a JWK object.
+	 * If specified, will decode the key from Base64Url form, and/or convert to little-endian.
+	 *
+	 * @param InKeyExponent The e (exponent) value of the public key.
+	 * @param InKeyModulus The n (modulus) value of the public key.
+	 * @param bIsBase64UrlEncoded Set to true if InKeyExponent and InKeyModulus are Base64Url encoded (will decode).
+	 * @param bIsBigEndian Set to true if InKeyExponent and InKeyModulus and in big-endian byte form (will reverse).
+	 * @return true if the header and payload were verified successfully (signature match), otherwise false.
+	 */
+	bool Verify(const FStringView InKeyExponent, const FStringView InKeyModulus, const bool bIsBase64UrlEncoded, const bool bIsBigEndian) const;
+
+	/**
+	 * Verifies the signature against the header and content, using the given public key info, which
+	 * may be obtained from a JWK object.
+	 * Assumes Base64Url-encoded big-endian byte format for the key, as it would arrive from the web.
+	 *
+	 * @param InKeyExponent The e (exponent) value of the public key.
+	 * @param InKeyModulus The n (modulus) value of the public key.
 	 * @return true if the header and payload were verified successfully (signature match), otherwise false.
 	 */
-	bool Verify() const;
+	bool Verify(const FStringView InKeyExponent, const FStringView InKeyModulus) const;
 
 private:
 	FJsonWebToken(const FStringView InEncodedJsonWebToken, const TSharedRef<FJsonObject>& InHeaderPtr,
diff --git a/Engine/Plugins/Experimental/PlatformCrypto/Source/PlatformCryptoOpenSSL/Private/EncryptionContextOpenSSL.cpp b/Engine/Plugins/Experimental/PlatformCrypto/Source/PlatformCryptoOpenSSL/Private/EncryptionContextOpenSSL.cpp
--- a/Engine/Plugins/Experimental/PlatformCrypto/Source/PlatformCryptoOpenSSL/Private/EncryptionContextOpenSSL.cpp
+++ b/Engine/Plugins/Experimental/PlatformCrypto/Source/PlatformCryptoOpenSSL/Private/EncryptionContextOpenSSL.cpp
@@ -5,6 +5,10 @@
 #include "PlatformCryptoAesEncryptorsOpenSSL.h"
 #include "PlatformCryptoAesDecryptorsOpenSSL.h"
 
+THIRD_PARTY_INCLUDES_START
+#include <openssl/err.h>
+THIRD_PARTY_INCLUDES_END
+
 #include "Misc/AssertionMacros.h"
 
 DEFINE_LOG_CATEGORY(LogPlatformCryptoOpenSSL);
@@ -323,25 +327,103 @@ bool FEncryptionContextOpenSSL::DigestSign_RS256(const TArrayView<const uint8> M
 	return RSA_sign(NID_sha256, Message.GetData(), Message.Num(), Signature.GetData(), &BytesWritten, (RSA*)Key) == 1;
 }
 
-bool FEncryptionContextOpenSSL::DigestVerify_PS256(const TArrayView<const char> Message, const TArrayView<const uint8> Signature, const TArrayView<const uint8> PKCS1Key)
+// Verify a message against the provided signature, using the PS256 algorithm.
+// PS256 is a JWT abbrieviation which translates to "RSASSA-PSS using SHA-256 digest, and MGF1 with SHA-256"
+bool FEncryptionContextOpenSSL::DigestVerify_PS256(const TArrayView<const uint8> Message, const TArrayView<const uint8> Signature, FRSAKeyHandle Key)
 {
-	FScopedEVPMDContext Context;
+	TArray<uint8> MessageDigest;
+	if (!CalcSHA256(Message, MessageDigest))
+	{
+		UE_LOG(LogPlatformCryptoOpenSSL, Verbose, TEXT("FEncryptionContextOpenSSL::DigestVerify_PS256 failed: error in generate digest"));
+		return false;
+	}
 
-	if (Context.Get() == nullptr)
+	if (!ensure(SHA256_DIGEST_LENGTH == MessageDigest.Num()))
 	{
+		UE_LOG(LogPlatformCryptoOpenSSL, Verbose, TEXT("FEncryptionContextOpenSSL::DigestVerify_PS256 failed: invalid digest length"));
 		return false;
 	}
 
-	const unsigned char* PKCS1KeyData = PKCS1Key.GetData();
-	RSA* RsaKey = d2i_RSAPublicKey(nullptr, &PKCS1KeyData, PKCS1Key.Num());
-	EVP_PKEY* PKey = EVP_PKEY_new();
-	EVP_PKEY_assign_RSA(PKey, RsaKey);
+	return DigestVerifyPreHashed_PS256(MessageDigest.GetData(), Signature.GetData(), Signature.Num(), Key);
+}
 
-	EVP_PKEY_CTX* KeyContext = nullptr;
-	EVP_DigestVerifyInit(Context.Get(), &KeyContext, EVP_sha256(), nullptr, PKey);
-	EVP_PKEY_CTX_set_rsa_padding(KeyContext, RSA_PKCS1_PSS_PADDING);
-	EVP_DigestVerifyUpdate(Context.Get(), Message.GetData(), Message.Num());
-	return EVP_DigestVerifyFinal(Context.Get(), const_cast<uint8*>(Signature.GetData()), Signature.Num()) == 1;
+// Verify a message digest (hash) against the provided signature, using the PS256 algorithm.
+// PS256 is a JWT abbrieviation which translates to "RSASSA-PSS using SHA-256 digest, and MGF1 with SHA-256"
+bool FEncryptionContextOpenSSL::DigestVerifyPreHashed_PS256(const uint8* const MessageDigest, const TArrayView<const uint8> Signature, FRSAKeyHandle Key)
+{
+	return DigestVerifyPreHashed_PS256(MessageDigest, Signature.GetData(), Signature.Num(), Key);
+}
+
+bool FEncryptionContextOpenSSL::DigestVerifyPreHashed_PS256(const uint8* const MessageDigest, const uint8* const Signature, const uint32 SignatureLen, FRSAKeyHandle Key)
+{
+	RSA *RsaKey = (RSA*)Key;
+
+	// Create public key (PKEY) and copy the relevant RSA key data into it.
+	EVP_PKEY *PubKey = EVP_PKEY_new();
+	const bool bPubKeySetOk = 1 == EVP_PKEY_set1_RSA(PubKey, RsaKey);
+	if (!bPubKeySetOk)
+	{
+		UE_LOG(LogPlatformCryptoOpenSSL, Verbose, TEXT("FEncryptionContextOpenSSL::DigestVerifyPreHashed_PS256 failed: error setting public key"));
+		EVP_PKEY_free(PubKey);
+		return false;
+	}
+
+	// Create a context for this public key operation, to hold any internal/intermediate variables.
+	EVP_PKEY_CTX *PubKeyContext = EVP_PKEY_CTX_new(PubKey, nullptr);
+	if (!PubKeyContext)
+	{
+		UE_LOG(LogPlatformCryptoOpenSSL, Verbose, TEXT("FEncryptionContextOpenSSL::DigestVerifyPreHashed_PS256 failed: error allocating public key context"));
+		EVP_PKEY_free(PubKey);
+		return false;
+	}
+
+	// Set up the context of this operation.
+
+	bool bInitOk = true;
+	bInitOk = bInitOk && (1 == EVP_PKEY_verify_init(PubKeyContext)); // set that we want to perform a verification operation
+	bInitOk = bInitOk && (1 == EVP_PKEY_CTX_set_signature_md(PubKeyContext, EVP_sha256())); // set digest algorithm used in signature operation
+
+	// PSS-specific context params.
+	bInitOk = bInitOk && (1 == EVP_PKEY_CTX_set_rsa_padding(PubKeyContext, RSA_PKCS1_PSS_PADDING)); // set padding mode (PKCS1 PSS)
+	bInitOk = bInitOk && (1 == EVP_PKEY_CTX_set_rsa_mgf1_md(PubKeyContext, EVP_sha256())); // set digest algorithm used by mask generation function (MGF1)
+	// Salt length options:
+	// - RSA_PSS_SALTLEN_DIGEST (-1) : Salt length matches digest length
+	// - RSA_PSS_SALTLEN_AUTO (-2) : Auto-detect salt length (for verification only)
+	// - RSA_PSS_SALTLEN_MAX (-3) : Set salt length to maximum possible
+	// - Positive values : Fixed-size salt length
+	bInitOk = bInitOk && (1 == EVP_PKEY_CTX_set_rsa_pss_saltlen(PubKeyContext, RSA_PSS_SALTLEN_AUTO)); // set PSS salt length
+
+	if (!bInitOk)
+	{
+		UE_LOG(LogPlatformCryptoOpenSSL, Verbose, TEXT("FEncryptionContextOpenSSL::DigestVerifyPreHashed_PS256 failed in initializing public key context (%u): %s"),
+			ERR_get_error(), UTF8_TO_TCHAR(ERR_error_string(ERR_get_error(), nullptr)));
+
+		EVP_PKEY_CTX_free(PubKeyContext);
+		EVP_PKEY_free(PubKey);
+		return false;
+	}
+
+	// Verify against pre-hashed input.
+	const int VerifyResult = EVP_PKEY_verify(PubKeyContext, Signature, SignatureLen, MessageDigest, SHA256_DIGEST_LENGTH);
+	const bool bVerifyOk = (1 == VerifyResult);
+	if (!bVerifyOk)
+	{
+		if (0 == VerifyResult)
+		{
+			UE_LOG(LogPlatformCryptoOpenSSL, Verbose, TEXT("FEncryptionContextOpenSSL::DigestVerifyPreHashed_PS256 signature is invalid (%u): %s"),
+				ERR_get_error(), UTF8_TO_TCHAR(ERR_error_string(ERR_get_error(), nullptr)));
+		}
+		else
+		{
+			UE_LOG(LogPlatformCryptoOpenSSL, Verbose, TEXT("FEncryptionContextOpenSSL::DigestVerifyPreHashed_PS256 error verifying signature (%u): %s"),
+				ERR_get_error(), UTF8_TO_TCHAR(ERR_error_string(ERR_get_error(), nullptr)));
+		}
+	}
+
+	EVP_PKEY_CTX_free(PubKeyContext);
+	EVP_PKEY_free(PubKey);
+
+	return bVerifyOk;
 }
 
 bool FEncryptionContextOpenSSL::DigestVerify_RS256(const TArrayView<const uint8> Message, const TArrayView<const uint8> Signature, FRSAKeyHandle Key)
@@ -421,25 +503,57 @@ FRSAKeyHandle FEncryptionContextOpenSSL::CreateKey_RSA(const TArrayView<const ui
 	BIGNUM* BN_PrivateExponent = PrivateExponent.Num() > 0 ? BN_new() : nullptr;
 	BIGNUM* BN_Modulus = BN_new();
 
+	bool bCreateKeyFailed = false;
+
 	if (PublicExponent.Num())
 	{
 		LoadBinaryIntoBigNum(PublicExponent.GetData(), PublicExponent.Num(), BN_PublicExponent);
+
+		if (!BN_is_odd(BN_PublicExponent))
+		{
+			bCreateKeyFailed = true;
+			UE_LOG(LogPlatformCryptoOpenSSL, Warning, TEXT("FEncryptionContextOpenSSL::CreateKey_RSA: public exponent is not an odd number"));
+		}
 	}
 
 	if (PrivateExponent.Num())
 	{
 		LoadBinaryIntoBigNum(PrivateExponent.GetData(), PrivateExponent.Num(), BN_PrivateExponent);
+
+		if (!BN_is_odd(BN_PrivateExponent))
+		{
+			bCreateKeyFailed = true;
+			UE_LOG(LogPlatformCryptoOpenSSL, Warning, TEXT("FEncryptionContextOpenSSL::CreateKey_RSA: private exponent is not an odd number"));
+		}
 	}
 
 	LoadBinaryIntoBigNum(Modulus.GetData(), Modulus.Num(), BN_Modulus);
+	if (!BN_is_odd(BN_Modulus))
+	{
+		bCreateKeyFailed = true;
+		UE_LOG(LogPlatformCryptoOpenSSL, Warning, TEXT("FEncryptionContextOpenSSL::CreateKey_RSA: modulus is not an odd number"));
+	}
+
 #if USE_LEGACY_OPENSSL
 	NewKey->n = BN_Modulus;
 	NewKey->e = BN_PublicExponent;
 	NewKey->d = BN_PrivateExponent;
 #else
-	RSA_set0_key(NewKey, BN_Modulus, BN_PublicExponent, BN_PrivateExponent);
+	const int Result = RSA_set0_key(NewKey, BN_Modulus, BN_PublicExponent, BN_PrivateExponent);
+	if (!Result)
+	{
+		bCreateKeyFailed = true;
+		UE_LOG(LogPlatformCryptoOpenSSL, Warning, TEXT("FEncryptionContextOpenSSL::CreateKey_RSA failed (%u): %s"),
+			ERR_get_error(), UTF8_TO_TCHAR(ERR_error_string(ERR_get_error(), nullptr)));
+	}
 #endif
 
+	if (bCreateKeyFailed)
+	{
+		RSA_free(NewKey);
+		NewKey = nullptr;
+	}
+
 	return NewKey;
 }
 
@@ -528,3 +642,29 @@ FSHA256Hasher FEncryptionContextOpenSSL::CreateSHA256Hasher()
 {
 	return FSHA256Hasher();
 }
+
+bool FEncryptionContextOpenSSL::CalcSHA256(const TArrayView<const uint8> Source, TArray<uint8>& OutHash)
+{
+	FSHA256Hasher Hasher = CreateSHA256Hasher();
+	if (Hasher.Init() != EPlatformCryptoResult::Success)
+	{
+		UE_LOG(LogPlatformCryptoOpenSSL, Verbose, TEXT("FEncryptionContextOpenSSL::CalcSHA256: create failed"));
+		return false;
+	}
+
+	if (Hasher.Update(Source) != EPlatformCryptoResult::Success)
+	{
+		UE_LOG(LogPlatformCryptoOpenSSL, Verbose, TEXT("FEncryptionContextOpenSSL::CalcSHA256: update failed"));
+		return false;
+	}
+
+	OutHash.Empty();
+	OutHash.AddUninitialized(FSHA256Hasher::OutputByteLength);
+	if (Hasher.Finalize(OutHash) != EPlatformCryptoResult::Success)
+	{
+		UE_LOG(LogPlatformCryptoOpenSSL, Verbose, TEXT("FEncryptionContextOpenSSL::CalcSHA256: finalize failed"));
+		return false;
+	}
+
+	return true;
+}
diff --git a/Engine/Plugins/Experimental/PlatformCrypto/Source/PlatformCryptoOpenSSL/Public/EncryptionContextOpenSSL.h b/Engine/Plugins/Experimental/PlatformCrypto/Source/PlatformCryptoOpenSSL/Public/EncryptionContextOpenSSL.h
--- a/Engine/Plugins/Experimental/PlatformCrypto/Source/PlatformCryptoOpenSSL/Public/EncryptionContextOpenSSL.h
+++ b/Engine/Plugins/Experimental/PlatformCrypto/Source/PlatformCryptoOpenSSL/Public/EncryptionContextOpenSSL.h
@@ -76,9 +76,11 @@ public:
 
 	bool DigestSign_RS256(const TArrayView<const uint8> Message, TArray<uint8>& Signature, FRSAKeyHandle Key);
 
-	bool DigestVerify_PS256(const TArrayView<const char> Message, const TArrayView<const uint8> Signature, const TArrayView<const uint8> PKCS1Key);
 	bool DigestVerify_RS256(const TArrayView<const uint8> Message, const TArrayView<const uint8> Signature, FRSAKeyHandle Key);
 
+	bool DigestVerify_PS256(const TArrayView<const uint8> Message, const TArrayView<const uint8> Signature, FRSAKeyHandle Key);
+	bool DigestVerifyPreHashed_PS256(const uint8* const MessageDigest, const TArrayView<const uint8> Signature, FRSAKeyHandle Key);
+
 	bool GenerateKey_RSA(const int32 InNumKeyBits, TArray<uint8>& OutPublicExponent, TArray<uint8>& OutPrivateExponent, TArray<uint8>& OutModulus);
 
 	FRSAKeyHandle CreateKey_RSA(const TArrayView<const uint8> PublicExponent, const TArrayView<const uint8> PrivateExponent, const TArrayView<const uint8> Modulus);
@@ -94,6 +96,13 @@ public:
 	EPlatformCryptoResult CreatePseudoRandomBytes(const TArrayView<uint8> OutData);
 
 	FSHA256Hasher CreateSHA256Hasher();
+
+	bool CalcSHA256(const TArrayView<const uint8> Source, TArray<uint8>& OutHash);
+
+private:
+
+	bool DigestVerify_PS256(const uint8* const Message, const uint32 MessageLen, const uint8* const Signature, const uint32 SignatureLen, FRSAKeyHandle Key);
+	bool DigestVerifyPreHashed_PS256(const uint8* const MessageDigest, const uint8* const Signature, const uint32 SignatureLen, FRSAKeyHandle Key);
 };
 
 typedef FEncryptionContextOpenSSL FEncryptionContext;
diff --git a/Engine/Source/Runtime/Core/Private/Misc/Base64.cpp b/Engine/Source/Runtime/Core/Private/Misc/Base64.cpp
--- a/Engine/Source/Runtime/Core/Private/Misc/Base64.cpp
+++ b/Engine/Source/Runtime/Core/Private/Misc/Base64.cpp
@@ -1,67 +1,85 @@
 // Copyright Epic Games, Inc. All Rights Reserved.
 
 #include "Misc/Base64.h"
-
 #include "Containers/StringConv.h"
 #include "Misc/AssertionMacros.h"
 
-/** The table used to encode a 6 bit value as an ascii character */
-static const uint8 EncodingAlphabet[64] = 
-{ 
-	'A', 'B', 'C', 'D', 'E', 'F', 'G', 'H', 'I', 'J', 'K', 'L', 'M', 'N', 'O', 'P', 
-	'Q', 'R', 'S', 'T', 'U', 'V', 'W', 'X', 'Y', 'Z', 'a', 'b', 'c', 'd', 'e', 'f', 
-	'g', 'h', 'i', 'j', 'k', 'l', 'm', 'n', 'o', 'p', 'q', 'r', 's', 't', 'u', 'v', 
-	'w', 'x', 'y', 'z', '0', '1', '2', '3', '4', '5', '6', '7', '8', '9', '+', '/' 
-};
-
-/** The table used to convert an ascii character into a 6 bit value */
-static const uint8 DecodingAlphabet[256] = 
-{ 
-	0xFF, 0xFF, 0xFF, 0xFF, 0xFF, 0xFF, 0xFF, 0xFF, 0xFF, 0xFF, 0xFF, 0xFF, 0xFF, 0xFF, 0xFF, 0xFF, // 0x00-0x0f
-	0xFF, 0xFF, 0xFF, 0xFF, 0xFF, 0xFF, 0xFF, 0xFF, 0xFF, 0xFF, 0xFF, 0xFF, 0xFF, 0xFF, 0xFF, 0xFF, // 0x10-0x1f
-	0xFF, 0xFF, 0xFF, 0xFF, 0xFF, 0xFF, 0xFF, 0xFF, 0xFF, 0xFF, 0xFF, 0x3E, 0xFF, 0xFF, 0xFF, 0x3F, // 0x20-0x2f
-	0x34, 0x35, 0x36, 0x37, 0x38, 0x39, 0x3A, 0x3B, 0x3C, 0x3D, 0xFF, 0xFF, 0xFF, 0xFF, 0xFF, 0xFF, // 0x30-0x3f
-	0xFF, 0x00, 0x01, 0x02, 0x03, 0x04, 0x05, 0x06, 0x07, 0x08, 0x09, 0x0A, 0x0B, 0x0C, 0x0D, 0x0E, // 0x40-0x4f
-	0x0F, 0x10, 0x11, 0x12, 0x13, 0x14, 0x15, 0x16, 0x17, 0x18, 0x19, 0xFF, 0xFF, 0xFF, 0xFF, 0xFF, // 0x50-0x5f
-	0xFF, 0x1A, 0x1B, 0x1C, 0x1D, 0x1E, 0x1F, 0x20, 0x21, 0x22, 0x23, 0x24, 0x25, 0x26, 0x27, 0x28, // 0x60-0x6f
-	0x29, 0x2A, 0x2B, 0x2C, 0x2D, 0x2E, 0x2F, 0x30, 0x31, 0x32, 0x33, 0xFF, 0xFF, 0xFF, 0xFF, 0xFF, // 0x70-0x7f
-	0xFF, 0xFF, 0xFF, 0xFF, 0xFF, 0xFF, 0xFF, 0xFF, 0xFF, 0xFF, 0xFF, 0xFF, 0xFF, 0xFF, 0xFF, 0xFF, // 0x80-0x8f
-	0xFF, 0xFF, 0xFF, 0xFF, 0xFF, 0xFF, 0xFF, 0xFF, 0xFF, 0xFF, 0xFF, 0xFF, 0xFF, 0xFF, 0xFF, 0xFF, // 0x90-0x9f
-	0xFF, 0xFF, 0xFF, 0xFF, 0xFF, 0xFF, 0xFF, 0xFF, 0xFF, 0xFF, 0xFF, 0xFF, 0xFF, 0xFF, 0xFF, 0xFF, // 0xa0-0xaf
-	0xFF, 0xFF, 0xFF, 0xFF, 0xFF, 0xFF, 0xFF, 0xFF, 0xFF, 0xFF, 0xFF, 0xFF, 0xFF, 0xFF, 0xFF, 0xFF, // 0xb0-0xbf
-	0xFF, 0xFF, 0xFF, 0xFF, 0xFF, 0xFF, 0xFF, 0xFF, 0xFF, 0xFF, 0xFF, 0xFF, 0xFF, 0xFF, 0xFF, 0xFF, // 0xc0-0xcf
-	0xFF, 0xFF, 0xFF, 0xFF, 0xFF, 0xFF, 0xFF, 0xFF, 0xFF, 0xFF, 0xFF, 0xFF, 0xFF, 0xFF, 0xFF, 0xFF, // 0xd0-0xdf
-	0xFF, 0xFF, 0xFF, 0xFF, 0xFF, 0xFF, 0xFF, 0xFF, 0xFF, 0xFF, 0xFF, 0xFF, 0xFF, 0xFF, 0xFF, 0xFF, // 0xe0-0xef
-	0xFF, 0xFF, 0xFF, 0xFF, 0xFF, 0xFF, 0xFF, 0xFF, 0xFF, 0xFF, 0xFF, 0xFF, 0xFF, 0xFF, 0xFF, 0xFF  // 0xf0-0xff
-};
-
-FString FBase64::Encode(const FString& Source)
+namespace
+{
+	/** The table used to encode a 6 bit value as an ascii character using Table 1 from RFC 4648 ("base64") */
+	static constexpr uint8 GBase64EncodingAlphabet[64] =
+	{
+		'A', 'B', 'C', 'D', 'E', 'F', 'G', 'H', 'I', 'J', 'K', 'L', 'M', 'N', 'O', 'P',
+		'Q', 'R', 'S', 'T', 'U', 'V', 'W', 'X', 'Y', 'Z', 'a', 'b', 'c', 'd', 'e', 'f',
+		'g', 'h', 'i', 'j', 'k', 'l', 'm', 'n', 'o', 'p', 'q', 'r', 's', 't', 'u', 'v',
+		'w', 'x', 'y', 'z', '0', '1', '2', '3', '4', '5', '6', '7', '8', '9', '+', '/'
+	};
+	/** The table used to encode a 6 bit value as an ascii character using Table 2 from RFC 4648 ("base64url") */
+	static constexpr uint8 GBase64UrlSafeEncodingAlphabet[64] =
+	{
+		'A', 'B', 'C', 'D', 'E', 'F', 'G', 'H', 'I', 'J', 'K', 'L', 'M', 'N', 'O', 'P',
+		'Q', 'R', 'S', 'T', 'U', 'V', 'W', 'X', 'Y', 'Z', 'a', 'b', 'c', 'd', 'e', 'f',
+		'g', 'h', 'i', 'j', 'k', 'l', 'm', 'n', 'o', 'p', 'q', 'r', 's', 't', 'u', 'v',
+		'w', 'x', 'y', 'z', '0', '1', '2', '3', '4', '5', '6', '7', '8', '9', '-', '_'
+	};
+
+	/** The table used to convert an ascii character into a 6 bit value using Table 1 from RFC 4648 ("base64") */
+	static constexpr uint8 GBase64DecodingAlphabet[128] =
+	{
+		0xFF, 0xFF, 0xFF, 0xFF, 0xFF, 0xFF, 0xFF, 0xFF, 0xFF, 0xFF, 0xFF, 0xFF, 0xFF, 0xFF, 0xFF, 0xFF, // 0x00-0x0f
+		0xFF, 0xFF, 0xFF, 0xFF, 0xFF, 0xFF, 0xFF, 0xFF, 0xFF, 0xFF, 0xFF, 0xFF, 0xFF, 0xFF, 0xFF, 0xFF, // 0x10-0x1f
+		0xFF, 0xFF, 0xFF, 0xFF, 0xFF, 0xFF, 0xFF, 0xFF, 0xFF, 0xFF, 0xFF, 0x3E, 0xFF, 0xFF, 0xFF, 0x3F, // 0x20-0x2f
+		0x34, 0x35, 0x36, 0x37, 0x38, 0x39, 0x3A, 0x3B, 0x3C, 0x3D, 0xFF, 0xFF, 0xFF, 0xFF, 0xFF, 0xFF, // 0x30-0x3f
+		0xFF, 0x00, 0x01, 0x02, 0x03, 0x04, 0x05, 0x06, 0x07, 0x08, 0x09, 0x0A, 0x0B, 0x0C, 0x0D, 0x0E, // 0x40-0x4f
+		0x0F, 0x10, 0x11, 0x12, 0x13, 0x14, 0x15, 0x16, 0x17, 0x18, 0x19, 0xFF, 0xFF, 0xFF, 0xFF, 0xFF, // 0x50-0x5f
+		0xFF, 0x1A, 0x1B, 0x1C, 0x1D, 0x1E, 0x1F, 0x20, 0x21, 0x22, 0x23, 0x24, 0x25, 0x26, 0x27, 0x28, // 0x60-0x6f
+		0x29, 0x2A, 0x2B, 0x2C, 0x2D, 0x2E, 0x2F, 0x30, 0x31, 0x32, 0x33, 0xFF, 0xFF, 0xFF, 0xFF, 0xFF  // 0x70-0x7f
+	};
+
+	/** The table used to convert an ascii character into a 6 bit value using Table 2 from RFC 4648 ("base64url") */
+	static constexpr uint8 GBase64UrlSafeDecodingAlphabet[128] =
+	{
+		0xFF, 0xFF, 0xFF, 0xFF, 0xFF, 0xFF, 0xFF, 0xFF, 0xFF, 0xFF, 0xFF, 0xFF, 0xFF, 0xFF, 0xFF, 0xFF, // 0x00-0x0f
+		0xFF, 0xFF, 0xFF, 0xFF, 0xFF, 0xFF, 0xFF, 0xFF, 0xFF, 0xFF, 0xFF, 0xFF, 0xFF, 0xFF, 0xFF, 0xFF, // 0x10-0x1f
+		0xFF, 0xFF, 0xFF, 0xFF, 0xFF, 0xFF, 0xFF, 0xFF, 0xFF, 0xFF, 0xFF, 0xFF, 0xFF, 0x3E, 0xFF, 0xFF, // 0x20-0x2f
+		0x34, 0x35, 0x36, 0x37, 0x38, 0x39, 0x3A, 0x3B, 0x3C, 0x3D, 0xFF, 0xFF, 0xFF, 0xFF, 0xFF, 0xFF, // 0x30-0x3f
+		0xFF, 0x00, 0x01, 0x02, 0x03, 0x04, 0x05, 0x06, 0x07, 0x08, 0x09, 0x0A, 0x0B, 0x0C, 0x0D, 0x0E, // 0x40-0x4f
+		0x0F, 0x10, 0x11, 0x12, 0x13, 0x14, 0x15, 0x16, 0x17, 0x18, 0x19, 0xFF, 0xFF, 0xFF, 0xFF, 0x3F, // 0x50-0x5f
+		0xFF, 0x1A, 0x1B, 0x1C, 0x1D, 0x1E, 0x1F, 0x20, 0x21, 0x22, 0x23, 0x24, 0x25, 0x26, 0x27, 0x28, // 0x60-0x6f
+		0x29, 0x2A, 0x2B, 0x2C, 0x2D, 0x2E, 0x2F, 0x30, 0x31, 0x32, 0x33, 0xFF, 0xFF, 0xFF, 0xFF, 0xFF  // 0x70-0x7f
+	};
+}
+
+FString FBase64::Encode(const FString& Source, EBase64Mode Mode)
 {
-	return Encode((uint8*)TCHAR_TO_ANSI(*Source), Source.Len());
+	return Encode((uint8*)TCHAR_TO_ANSI(*Source), Source.Len(), Mode);
 }
 
-FString FBase64::Encode(const TArray<uint8>& Source)
+FString FBase64::Encode(const TArray<uint8>& Source, EBase64Mode Mode)
 {
-	return Encode((uint8*)Source.GetData(), Source.Num());
+	return Encode((uint8*)Source.GetData(), Source.Num(), Mode);
 }
 
-FString FBase64::Encode(const uint8* Source, uint32 Length)
+FString FBase64::Encode(const uint8* Source, uint32 Length, EBase64Mode Mode)
 {
 	uint32 ExpectedLength = GetEncodedDataSize(Length);
-	
+
 	FString OutBuffer;
 
 	TArray<TCHAR, FString::AllocatorType>& OutCharArray = OutBuffer.GetCharArray();
 	OutCharArray.SetNum(ExpectedLength + 1);
-	int64 EncodedLength = Encode(Source, Length, OutCharArray.GetData());
+	int64 EncodedLength = Encode(Source, Length, OutCharArray.GetData(), Mode);
 	verify(EncodedLength == OutBuffer.Len());
 
 	return OutBuffer;
 }
 
-template<typename CharType> uint32 FBase64::Encode(const uint8* Source, uint32 Length, CharType* Dest)
+template<typename CharType> uint32 FBase64::Encode(const uint8* Source, uint32 Length, CharType* Dest, EBase64Mode Mode)
 {
+	check(Mode == EBase64Mode::Standard || Mode == EBase64Mode::UrlSafe);
+
 	CharType* EncodedBytes = Dest;
+	const uint8* const EncodingAlphabet = (Mode == EBase64Mode::UrlSafe ? GBase64UrlSafeEncodingAlphabet : GBase64EncodingAlphabet);
 
 	// Loop through the buffer converting 3 bytes of binary data at a time
 	while (Length >= 3)
@@ -126,16 +144,16 @@ template<typename CharType> uint32 FBase64::Encode(const uint8* Source, uint32 L
 	return UE_PTRDIFF_TO_UINT32(EncodedBytes - Dest);
 }
 
-template CORE_API uint32 FBase64::Encode<ANSICHAR>(const uint8* Source, uint32 Length, ANSICHAR* Dest);
-template CORE_API uint32 FBase64::Encode<WIDECHAR>(const uint8* Source, uint32 Length, WIDECHAR* Dest);
+template CORE_API uint32 FBase64::Encode<ANSICHAR>(const uint8* Source, uint32 Length, ANSICHAR* Dest, EBase64Mode Mode);
+template CORE_API uint32 FBase64::Encode<WIDECHAR>(const uint8* Source, uint32 Length, WIDECHAR* Dest, EBase64Mode Mode);
 
-bool FBase64::Decode(const FString& Source, FString& OutDest)
+bool FBase64::Decode(const FString& Source, FString& OutDest, EBase64Mode Mode)
 {
 	uint32 ExpectedLength = GetDecodedDataSize(Source);
 
 	TArray<ANSICHAR> TempDest;
 	TempDest.AddZeroed(ExpectedLength + 1);
-	if(!Decode(*Source, Source.Len(), (uint8*)TempDest.GetData()))
+	if (!Decode(*Source, Source.Len(), (uint8*)TempDest.GetData(), Mode))
 	{
 		return false;
 	}
@@ -143,35 +161,39 @@ bool FBase64::Decode(const FString& Source, FString& OutDest)
 	return true;
 }
 
-bool FBase64::Decode(const FString& Source, TArray<uint8>& OutDest)
+bool FBase64::Decode(const FString& Source, TArray<uint8>& OutDest, EBase64Mode Mode)
 {
 	OutDest.SetNum(GetDecodedDataSize(Source));
-	return Decode(*Source, Source.Len(), OutDest.GetData());
+	return Decode(*Source, Source.Len(), OutDest.GetData(), Mode);
 }
 
-template<typename CharType> bool FBase64::Decode(const CharType* Source, uint32 Length, uint8* Dest)
+template<typename CharType> bool FBase64::Decode(const CharType* Source, uint32 Length, uint8* Dest, EBase64Mode Mode)
 {
+	check(Mode == EBase64Mode::Standard || Mode == EBase64Mode::UrlSafe);
+
 	// Remove the trailing '=' characters, so we can handle padded and non-padded input the same
-	while(Length > 0 && Source[Length - 1] == '=')
+	while (Length > 0 && Source[Length - 1] == '=')
 	{
 		Length--;
 	}
 
 	// Make sure the length is valid. Only lengths modulo 4 of 0, 2, and 3 are valid.
-	if((Length & 3) == 1)
+	if ((Length & 3) == 1)
 	{
 		return false;
 	}
 
+	const uint8* const DecodingAlphabet = (Mode == EBase64Mode::UrlSafe ? GBase64UrlSafeDecodingAlphabet : GBase64DecodingAlphabet);
+
 	// Convert all the full chunks of data
-	for(; Length >= 4; Length -= 4)
+	for (; Length >= 4; Length -= 4)
 	{
 		// Decode the next 4 BYTEs
 		uint32 OriginalTriplet = 0;
 		for (int32 Index = 0; Index < 4; Index++)
 		{
 			uint32 SourceChar = (uint32)*Source++;
-			if(SourceChar >= 256)
+			if (SourceChar >= 128)
 			{
 				return false;
 			}
@@ -195,14 +217,14 @@ template<typename CharType> bool FBase64::Decode(const CharType* Source, uint32
 	}
 
 	// Convert the last chunk of data
-	if(Length > 0)
+	if (Length > 0)
 	{
 		// Decode the next 4 BYTEs, or up to the end of the input buffer
 		uint32 OriginalTriplet = 0;
 		for (uint32 Index = 0; Index < Length; Index++)
 		{
 			uint32 SourceChar = (uint32)*Source++;
-			if(SourceChar >= 256)
+			if (SourceChar >= 128)
 			{
 				return false;
 			}
@@ -213,14 +235,14 @@ template<typename CharType> bool FBase64::Decode(const CharType* Source, uint32
 			}
 			OriginalTriplet = (OriginalTriplet << 6) | DecodedValue;
 		}
-		for(int32 Index = Length; Index < 4; Index++)
+		for (int32 Index = Length; Index < 4; Index++)
 		{
 			OriginalTriplet = (OriginalTriplet << 6);
 		}
 
 		// Now we can tear the uint32 into bytes
 		OriginalTriplet >>= 8;
-		if(Length >= 3)
+		if (Length >= 3)
 		{
 			Dest[1] = OriginalTriplet & 0xFF;
 		}
@@ -230,8 +252,8 @@ template<typename CharType> bool FBase64::Decode(const CharType* Source, uint32
 	return true;
 }
 
-template CORE_API bool FBase64::Decode<ANSICHAR>(const ANSICHAR* Source, uint32 Length, uint8* Dest);
-template CORE_API bool FBase64::Decode<WIDECHAR>(const WIDECHAR* Source, uint32 Length, uint8* Dest);
+template CORE_API bool FBase64::Decode<ANSICHAR>(const ANSICHAR* Source, uint32 Length, uint8* Dest, EBase64Mode Mode);
+template CORE_API bool FBase64::Decode<WIDECHAR>(const WIDECHAR* Source, uint32 Length, uint8* Dest, EBase64Mode Mode);
 
 uint32 FBase64::GetDecodedDataSize(const FString& Source)
 {
@@ -241,21 +263,21 @@ uint32 FBase64::GetDecodedDataSize(const FString& Source)
 template<typename CharType> uint32 FBase64::GetDecodedDataSize(const CharType* Source, uint32 Length)
 {
 	uint32 NumBytes = 0;
-	if(Length > 0)
+	if (Length > 0)
 	{
 		// Get the source length without the trailing padding characters
-		while(Length > 0 && Source[Length - 1] == '=')
+		while (Length > 0 && Source[Length - 1] == '=')
 		{
 			Length--;
 		}
 
 		// Get the lower bound for number of bytes, excluding the last partial chunk
 		NumBytes += (Length / 4) * 3;
-		if((Length & 3) == 3)
+		if ((Length & 3) == 3)
 		{
 			NumBytes += 2;
 		}
-		if((Length & 3) == 2)
+		if ((Length & 3) == 2)
 		{
 			NumBytes++;
 		}
diff --git a/Engine/Source/Runtime/Core/Public/Misc/Base64.h b/Engine/Source/Runtime/Core/Public/Misc/Base64.h
--- a/Engine/Source/Runtime/Core/Public/Misc/Base64.h
+++ b/Engine/Source/Runtime/Core/Public/Misc/Base64.h
@@ -7,6 +7,15 @@
 #include "CoreTypes.h"
 #include "Misc/Timespan.h"
 
+/** Mode to use for base64 encoding and decoding */
+enum class EBase64Mode : uint8
+{
+	/** Use the standard set of character mappings. (Table 1 from RFC 4648, known as "base64") */
+	Standard,
+	/** Use the "URL and Filename safe" set of character mappings. (Table 2 from RFC 4648, known as "base64url") */
+	UrlSafe
+};
+
 /**
  * Class for encoding/decoding Base64 data (RFC 4648)
  */
@@ -16,29 +25,32 @@ struct CORE_API FBase64
 	 * Encodes a FString into a Base64 string
 	 *
 	 * @param Source The string data to convert
+	 * @param Mode The mode to use for encoding. Default is EBase64Mode::Standard
 	 *
 	 * @return A string that encodes the binary data in a way that can be safely transmitted via various Internet protocols
 	 */
-	static FString Encode(const FString& Source);
+	static FString Encode(const FString& Source, EBase64Mode Mode = EBase64Mode::Standard);
 
 	/**
 	 * Encodes a binary uint8 array into a Base64 string
 	 *
 	 * @param Source The binary data to convert
+	 * @param Mode The mode to use for encoding. Default is EBase64Mode::Standard
 	 *
 	 * @return A string that encodes the binary data in a way that can be safely transmitted via various Internet protocols
 	 */
-	static FString Encode(const TArray<uint8>& Source);
+	static FString Encode(const TArray<uint8>& Source, EBase64Mode Mode = EBase64Mode::Standard);
 
 	/**
 	 * Encodes the source into a Base64 string
 	 *
 	 * @param Source The binary data to encode
 	 * @param Length Length of the binary data to be encoded
+	 * @param Mode The mode to use for encoding. Default is EBase64Mode::Standard
 	 *
 	 * @return Base64 encoded string containing the binary data.
 	 */
-	static FString Encode(const uint8* Source, uint32 Length);
+	static FString Encode(const uint8* Source, uint32 Length, EBase64Mode Mode = EBase64Mode::Standard);
 
 	/**
 	 * Encodes the source into a Base64 string, storing it in a preallocated buffer.
@@ -46,10 +58,11 @@ struct CORE_API FBase64
 	 * @param Source The binary data to encode
 	 * @param Length Length of the binary data to be encoded
 	 * @param Dest Buffer to receive the encoded data. Must be large enough to contain the entire output data (see GetEncodedDataSize()).
+	 * @param Mode The mode to use for encoding. Default is EBase64Mode::Standard
 	 *
 	 * @return The length of the encoded data
 	 */
-	template<typename CharType> static uint32 Encode(const uint8* Source, uint32 Length, CharType* Dest);
+	template<typename CharType> static uint32 Encode(const uint8* Source, uint32 Length, CharType* Dest, EBase64Mode Mode = EBase64Mode::Standard);
 
 	/**
 	* Get the encoded data size for the given number of bytes.
@@ -68,16 +81,22 @@ struct CORE_API FBase64
 	 *
 	 * @param Source The Base64 encoded string
 	 * @param OutDest Receives the decoded string data
+	 * @param Mode The mode to use for decoding. Default is EBase64Mode::Standard
+	 *
+	 * @return true if the buffer was decoded, false if it was invalid.
 	 */
-	static bool Decode(const FString& Source, FString& OutDest);
+	static bool Decode(const FString& Source, FString& OutDest, EBase64Mode Mode = EBase64Mode::Standard);
 
 	/**
 	 * Decodes a Base64 string into an array of bytes
 	 *
 	 * @param Source The Base64 encoded string
 	 * @param Dest Array to receive the decoded data
+	 * @param Mode The mode to use for decoding. Default is EBase64Mode::Standard
+	 *
+	 * @return true if the buffer was decoded, false if it was invalid.
 	 */
-	static bool Decode(const FString& Source, TArray<uint8>& Dest);
+	static bool Decode(const FString& Source, TArray<uint8>& Dest, EBase64Mode Mode = EBase64Mode::Standard);
 
 	/**
 	 * Decodes a Base64 string into a preallocated buffer
@@ -85,10 +104,11 @@ struct CORE_API FBase64
 	 * @param Source The Base64 encoded string
 	 * @param Length Length of the Base64 encoded string
 	 * @param Dest Buffer to receive the decoded data
+	 * @param Mode The mode to use for decoding. Default is EBase64Mode::Standard
 	 *
 	 * @return true if the buffer was decoded, false if it was invalid.
 	 */
-	template<typename CharType> static bool Decode(const CharType* Source, uint32 Length, uint8* Dest);
+	template<typename CharType> static bool Decode(const CharType* Source, uint32 Length, uint8* Dest, EBase64Mode Mode = EBase64Mode::Standard);
 
 	/**
 	* Determine the decoded data size for the incoming base64 encoded string
