FGenericPlatformString::Convert and ::ConvertedLength with null-terminator.
Missing UTF16TEXT macro added.
Unit tests for FGenericPlatformString conversion functions.
MakeUniqueForOverwrite added.

diff --git a/Engine/Source/Runtime/Core/Private/GenericPlatform/GenericPlatformString.cpp b/Engine/Source/Runtime/Core/Private/GenericPlatform/GenericPlatformString.cpp
--- a/Engine/Source/Runtime/Core/Private/GenericPlatform/GenericPlatformString.cpp
+++ b/Engine/Source/Runtime/Core/Private/GenericPlatform/GenericPlatformString.cpp
@@ -189,8 +189,40 @@ namespace UE::Core::Private
 		return true;
 	}
 
-	template <typename DestBufferType, typename FromType>
-	static int32 ConvertToUTF8(DestBufferType& Dest, int32 DestLen, const FromType* Source, int32 SourceLen)
+	struct FNullTerminal
+	{
+		explicit FNullTerminal() = default;
+	};
+
+	template <typename Pointer>
+	bool IsRangeEmpty(Pointer& Ptr, int32& Len)
+	{
+		return Len <= 0;
+	}
+
+	template <typename Pointer>
+	bool IsRangeEmpty(Pointer& Ptr, FNullTerminal)
+	{
+		return *Ptr == '\0';
+	}
+
+	template <typename Pointer>
+	void PopFront(Pointer& Ptr, int32& Len)
+	{
+		checkfSlow(Len > 0, TEXT("Trying to pop past end end of the range"));
+		++Ptr;
+		--Len;
+	}
+
+	template <typename Pointer>
+	void PopFront(Pointer& Ptr, FNullTerminal)
+	{
+		checkfSlow(*Ptr != '\0', TEXT("Trying to pop the null terminator of the string"));
+		++Ptr;
+	}
+
+	template <typename DestBufferType, typename FromType, typename SourceEndType>
+	static int32 ConvertToUTF8(DestBufferType& Dest, int32 DestLen, const FromType* Source, SourceEndType SourceEnd)
 	{
 		using UnsignedFromType = std::make_unsigned_t<FromType>;
 
@@ -200,19 +232,19 @@ namespace UE::Core::Private
 
 			for (;;)
 			{
-				if (SourceLen <= 0)
+				if (IsRangeEmpty(Source, SourceEnd))
 				{
 					return UE_PTRDIFF_TO_INT32(Dest - DestStartingPosition);
 				}
 
-				uint32 Codepoint = (uint32)(UnsignedFromType)*Source++;
-				--SourceLen;
-
+				uint32 Codepoint = (uint32)(UnsignedFromType)*Source;
 				if (!WriteCodepointToBuffer(Codepoint, Dest, DestLen))
 				{
 					// Could not write data, bail out
 					return -1;
 				}
+
+				PopFront(Source, SourceEnd);
 			}
 		}
 		else
@@ -221,18 +253,18 @@ namespace UE::Core::Private
 
 			for (;;)
 			{
-				if (SourceLen <= 0)
+				if (IsRangeEmpty(Source, SourceEnd))
 				{
 					return UE_PTRDIFF_TO_INT32(Dest - DestStartingPosition);
 				}
 
-				uint32 Codepoint = (uint32)(UnsignedFromType)*Source++;
-				--SourceLen;
+				uint32 Codepoint = (uint32)(UnsignedFromType)*Source;
+				PopFront(Source, SourceEnd);
 
 				// Check if this character is a high-surrogate
 				if (IsHighSurrogate(Codepoint))
 				{
-					if (SourceLen == 0)
+					if (IsRangeEmpty(Source, SourceEnd))
 					{
 						// String ends with lone high-surrogate - write it out (will be converted into bogus character)
 						if (!WriteCodepointToBuffer(Codepoint, Dest, DestLen))
@@ -252,8 +284,7 @@ namespace UE::Core::Private
 					if (IsLowSurrogate(NextCodepoint))
 					{
 						Codepoint = EncodeSurrogate((uint16)Codepoint, (uint16)NextCodepoint);
-						++Source;
-						--SourceLen;
+						PopFront(Source, SourceEnd);
 					}
 				}
 
@@ -266,17 +297,17 @@ namespace UE::Core::Private
 		}
 	}
 
-	static bool ReadTrailingOctet(uint32& OutOctet, const UTF8CHAR*& Ptr, uint32& Size)
+	template <typename SourceEndType>
+	static bool ReadTrailingOctet(uint32& OutOctet, const UTF8CHAR*& Ptr, SourceEndType& SourceEnd)
 	{
 		// Ensure our string has enough characters to read from
-		if (Size == 0)
+		if (IsRangeEmpty(Ptr, SourceEnd))
 		{
 			return false;
 		}
 
 		uint32 Octet = (uint32)(uint8)*Ptr;
-		++Ptr;
-		--Size;
+		PopFront(Ptr, SourceEnd);
 
 		// Format isn't 10xxxxxx?
 		if ((Octet & 192) != 128)
@@ -288,13 +319,13 @@ namespace UE::Core::Private
 		return true;
 	};
 
-	static uint32 CodepointFromUtf8(const UTF8CHAR*& SourceString, uint32 SourceLengthRemaining)
+	template <typename SourceEndType>
+	static uint32 CodepointFromUtf8(const UTF8CHAR*& SourceString, SourceEndType& SourceEnd)
 	{
-		checkSlow(SourceLengthRemaining > 0);
+		checkSlow(!IsRangeEmpty(SourceString, SourceEnd));
 
 		uint32 Octet = (uint32)(uint8)*SourceString;
-		++SourceString;
-		--SourceLengthRemaining;
+		PopFront(SourceString, SourceEnd);
 
 		if (Octet < 128)  // one octet char: 0 to 127
 		{
@@ -311,7 +342,7 @@ namespace UE::Core::Private
 		}
 
 		uint32 Octet2;
-		if (!ReadTrailingOctet(Octet2, SourceString, SourceLengthRemaining))
+		if (!ReadTrailingOctet(Octet2, SourceString, SourceEnd))
 		{
 			return UNICODE_BOGUS_CHAR_CODEPOINT;
 		}
@@ -329,7 +360,7 @@ namespace UE::Core::Private
 		}
 
 		uint32 Octet3;
-		if (!ReadTrailingOctet(Octet3, SourceString, SourceLengthRemaining))
+		if (!ReadTrailingOctet(Octet3, SourceString, SourceEnd))
 		{
 			return UNICODE_BOGUS_CHAR_CODEPOINT;
 		}
@@ -349,7 +380,7 @@ namespace UE::Core::Private
 		}
 
 		uint32 Octet4;
-		if (!ReadTrailingOctet(Octet4, SourceString, SourceLengthRemaining))
+		if (!ReadTrailingOctet(Octet4, SourceString, SourceEnd))
 		{
 			return UNICODE_BOGUS_CHAR_CODEPOINT;
 		}
@@ -369,7 +400,7 @@ namespace UE::Core::Private
 		}
 
 		uint32 Octet5;
-		if (!ReadTrailingOctet(Octet5, SourceString, SourceLengthRemaining))
+		if (!ReadTrailingOctet(Octet5, SourceString, SourceEnd))
 		{
 			return UNICODE_BOGUS_CHAR_CODEPOINT;
 		}
@@ -387,7 +418,7 @@ namespace UE::Core::Private
 		// six octets
 
 		uint32 Octet6;
-		if (!ReadTrailingOctet(Octet6, SourceString, SourceLengthRemaining))
+		if (!ReadTrailingOctet(Octet6, SourceString, SourceEnd))
 		{
 			return UNICODE_BOGUS_CHAR_CODEPOINT;
 		}
@@ -398,15 +429,13 @@ namespace UE::Core::Private
 	/**
 	 * Read Source string, converting the data from UTF-8 into UTF-16, and placing these in the Destination
 	 */
-	template <typename DestType, typename DestBufferType>
-	static int32 ConvertFromUTF8(DestBufferType& ConvertedBuffer, int32 DestLen, const UTF8CHAR* Source, const int32 SourceLen)
+	template <typename DestType, typename DestBufferType, typename SourceEndType>
+	static int32 ConvertFromUTF8(DestBufferType& ConvertedBuffer, int32 DestLen, const UTF8CHAR* Source, SourceEndType SourceEnd)
 	{
 		DestBufferType DestStartingPosition = ConvertedBuffer;
 
-		const UTF8CHAR* SourceEnd = Source + SourceLen;
-
 		const uint64 ExtendedCharMask = 0x8080808080808080;
-		while (Source < SourceEnd)
+		while (!IsRangeEmpty(Source, SourceEnd))
 		{
 			if (DestLen == 0)
 			{
@@ -415,35 +444,40 @@ namespace UE::Core::Private
 
 			// In case we're given an unaligned pointer, we'll
 			// fallback to the slow path until properly aligned.
-			if (IsAligned(Source, 8))
+			// But we can only do that if we know how much buffer we have left.
+			if constexpr (std::is_integral_v<SourceEndType>)
 			{
-				// Fast path for most common case
-				while (Source < SourceEnd - 8 && DestLen >= 8)
+				if (IsAligned(Source, 8))
 				{
-					// Detect any extended characters 8 chars at a time
-					if ((*(const uint64*)Source) & ExtendedCharMask)
+					// Fast path for most common case
+					while (SourceEnd >= 8 && DestLen >= 8)
 					{
-						// Move to slow path since we got extended characters to process
-						break;
-					}
+						// Detect any extended characters 8 chars at a time
+						if ((*(const uint64*)Source) & ExtendedCharMask)
+						{
+							// Move to slow path since we got extended characters to process
+							break;
+						}
 
-					// This should get unrolled on most compiler
-					// ROI of diminished return to vectorize this as we 
-					// would have to deal with alignment, endianness and
-					// rewrite the iterators to support bulk writes
-					for (int32 Index = 0; Index < 8; ++Index)
-					{
-						*(ConvertedBuffer++) = (DestType)(uint8)*(Source++);
+						// This should get unrolled on most compiler
+						// ROI of diminished return to vectorize this as we 
+						// would have to deal with alignment, endianness and
+						// rewrite the iterators to support bulk writes
+						for (int32 Index = 0; Index < 8; ++Index)
+						{
+							*(ConvertedBuffer++) = (DestType)(uint8)*(Source++);
+						}
+						SourceEnd -= 8;
+						DestLen -= 8;
 					}
-					DestLen -= 8;
 				}
 			}
 
 			// Slow path for extended characters
-			while (Source < SourceEnd && DestLen > 0)
+			while (!IsRangeEmpty(Source, SourceEnd) && DestLen > 0)
 			{
 				// Read our codepoint, advancing the source pointer
-				uint32 Codepoint = CodepointFromUtf8(Source, UE_PTRDIFF_TO_UINT32(SourceEnd - Source));
+				uint32 Codepoint = CodepointFromUtf8(Source, SourceEnd);
 
 				if constexpr (sizeof(DestType) != 4)
 				{
@@ -476,15 +510,18 @@ namespace UE::Core::Private
 				*(ConvertedBuffer++) = (DestType)Codepoint;
 				--DestLen;
 
-				// Return to the fast path once aligned and back to simple ASCII chars
-				if (Codepoint < 128 && IsAligned(Source, 8))
+				if constexpr (std::is_integral_v<SourceEndType>)
 				{
-					break;
+					// Return to the fast path once aligned and back to simple ASCII chars
+					if (Codepoint < 128 && IsAligned(Source, 8))
+					{
+						break;
+					}
 				}
 			}
 		}
 
-		return UE_PTRDIFF_TO_INT32(ConvertedBuffer - DestStartingPosition);;
+		return UE_PTRDIFF_TO_INT32(ConvertedBuffer - DestStartingPosition);
 	}
 
 	/**
@@ -492,36 +529,72 @@ namespace UE::Core::Private
 	 *
 	 * @return The length of the string in UTF-16 code units.
 	 */
+	int32 GetConvertedLength(const UTF8CHAR*, const WIDECHAR* Source)
+	{
+		TCountingOutputIterator<UTF8CHAR> Dest;
+		int32 Result = ConvertToUTF8(Dest, INT32_MAX, Source, FNullTerminal{});
+		return Result;
+	}
 	int32 GetConvertedLength(const UTF8CHAR*, const WIDECHAR* Source, int32 SourceLen)
 	{
 		TCountingOutputIterator<UTF8CHAR> Dest;
 		int32 Result = ConvertToUTF8(Dest, INT32_MAX, Source, SourceLen);
 		return Result;
 	}
+	int32 GetConvertedLength(const UTF8CHAR*, const UCS2CHAR* Source)
+	{
+		TCountingOutputIterator<UTF8CHAR> Dest;
+		int32 Result = ConvertToUTF8(Dest, INT32_MAX, Source, FNullTerminal{});
+		return Result;
+	}
 	int32 GetConvertedLength(const UTF8CHAR*, const UCS2CHAR* Source, int32 SourceLen)
 	{
 		TCountingOutputIterator<UTF8CHAR> Dest;
 		int32 Result = ConvertToUTF8(Dest, INT32_MAX, Source, SourceLen);
 		return Result;
 	}
+	int32 GetConvertedLength(const UTF8CHAR*, const UTF32CHAR* Source)
+	{
+		TCountingOutputIterator<UTF8CHAR> Dest;
+		int32 Result = ConvertToUTF8(Dest, INT32_MAX, Source, FNullTerminal{});
+		return Result;
+	}
 	int32 GetConvertedLength(const UTF8CHAR*, const UTF32CHAR* Source, int32 SourceLen)
 	{
 		TCountingOutputIterator<UTF8CHAR> Dest;
 		int32 Result = ConvertToUTF8(Dest, INT32_MAX, Source, SourceLen);
 		return Result;
 	}
+	int32 GetConvertedLength(const ANSICHAR*, const UTF8CHAR* Source)
+	{
+		TCountingOutputIterator<ANSICHAR> Dest;
+		int32 Result = ConvertFromUTF8<ANSICHAR>(Dest, INT32_MAX, Source, FNullTerminal{});
+		return Result;
+	}
 	int32 GetConvertedLength(const ANSICHAR*, const UTF8CHAR* Source, int32 SourceLen)
 	{
 		TCountingOutputIterator<ANSICHAR> Dest;
 		int32 Result = ConvertFromUTF8<ANSICHAR>(Dest, INT32_MAX, Source, SourceLen);
 		return Result;
 	}
+	int32 GetConvertedLength(const WIDECHAR*, const UTF8CHAR* Source)
+	{
+		TCountingOutputIterator<WIDECHAR> Dest;
+		int32 Result = ConvertFromUTF8<WIDECHAR>(Dest, INT32_MAX, Source, FNullTerminal{});
+		return Result;
+	}
 	int32 GetConvertedLength(const WIDECHAR*, const UTF8CHAR* Source, int32 SourceLen)
 	{
 		TCountingOutputIterator<WIDECHAR> Dest;
 		int32 Result = ConvertFromUTF8<WIDECHAR>(Dest, INT32_MAX, Source, SourceLen);
 		return Result;
 	}
+	int32 GetConvertedLength(const UCS2CHAR*, const UTF8CHAR* Source)
+	{
+		TCountingOutputIterator<UCS2CHAR> Dest;
+		int32 Result = ConvertFromUTF8<UCS2CHAR>(Dest, INT32_MAX, Source, FNullTerminal{});
+		return Result;
+	}
 	int32 GetConvertedLength(const UCS2CHAR*, const UTF8CHAR* Source, int32 SourceLen)
 	{
 		TCountingOutputIterator<UCS2CHAR> Dest;
@@ -529,6 +602,14 @@ namespace UE::Core::Private
 		return Result;
 	}
 
+	UTF8CHAR* Convert(UTF8CHAR* Dest, int32 DestLen, const WIDECHAR* Src)
+	{
+		if (ConvertToUTF8(Dest, DestLen, Src, FNullTerminal{}) == -1)
+		{
+			return nullptr;
+		}
+		return Dest;
+	}
 	UTF8CHAR* Convert(UTF8CHAR* Dest, int32 DestLen, const WIDECHAR* Src, int32 SrcLen)
 	{
 		if (ConvertToUTF8(Dest, DestLen, Src, SrcLen) == -1)
@@ -537,6 +618,14 @@ namespace UE::Core::Private
 		}
 		return Dest;
 	}
+	UTF8CHAR* Convert(UTF8CHAR* Dest, int32 DestLen, const UCS2CHAR* Src)
+	{
+		if (ConvertToUTF8(Dest, DestLen, Src, FNullTerminal{}) == -1)
+		{
+			return nullptr;
+		}
+		return Dest;
+	}
 	UTF8CHAR* Convert(UTF8CHAR* Dest, int32 DestLen, const UCS2CHAR* Src, int32 SrcLen)
 	{
 		if (ConvertToUTF8(Dest, DestLen, Src, SrcLen) == -1)
@@ -545,6 +634,14 @@ namespace UE::Core::Private
 		}
 		return Dest;
 	}
+	UTF8CHAR* Convert(UTF8CHAR* Dest, int32 DestLen, const UTF32CHAR* Src)
+	{
+		if (ConvertToUTF8(Dest, DestLen, Src, FNullTerminal{}) == -1)
+		{
+			return nullptr;
+		}
+		return Dest;
+	}
 	UTF8CHAR* Convert(UTF8CHAR* Dest, int32 DestLen, const UTF32CHAR* Src, int32 SrcLen)
 	{
 		if (ConvertToUTF8(Dest, DestLen, Src, SrcLen) == -1)
@@ -553,6 +650,14 @@ namespace UE::Core::Private
 		}
 		return Dest;
 	}
+	ANSICHAR* Convert(ANSICHAR* Dest, int32 DestLen, const UTF8CHAR* Src)
+	{
+		if (ConvertFromUTF8<ANSICHAR>(Dest, DestLen, Src, FNullTerminal{}) == -1)
+		{
+			return nullptr;
+		}
+		return Dest;
+	}
 	ANSICHAR* Convert(ANSICHAR* Dest, int32 DestLen, const UTF8CHAR* Src, int32 SrcLen)
 	{
 		if (ConvertFromUTF8<ANSICHAR>(Dest, DestLen, Src, SrcLen) == -1)
@@ -561,6 +666,14 @@ namespace UE::Core::Private
 		}
 		return Dest;
 	}
+	WIDECHAR* Convert(WIDECHAR* Dest, int32 DestLen, const UTF8CHAR* Src)
+	{
+		if (ConvertFromUTF8<WIDECHAR>(Dest, DestLen, Src, FNullTerminal{}) == -1)
+		{
+			return nullptr;
+		}
+		return Dest;
+	}
 	WIDECHAR* Convert(WIDECHAR* Dest, int32 DestLen, const UTF8CHAR* Src, int32 SrcLen)
 	{
 		if (ConvertFromUTF8<WIDECHAR>(Dest, DestLen, Src, SrcLen) == -1)
@@ -569,6 +682,14 @@ namespace UE::Core::Private
 		}
 		return Dest;
 	}
+	UCS2CHAR* Convert(UCS2CHAR* Dest, int32 DestLen, const UTF8CHAR* Src)
+	{
+		if (ConvertFromUTF8<UCS2CHAR>(Dest, DestLen, Src, FNullTerminal{}) == -1)
+		{
+			return nullptr;
+		}
+		return Dest;
+	}
 	UCS2CHAR* Convert(UCS2CHAR* Dest, int32 DestLen, const UTF8CHAR* Src, int32 SrcLen)
 	{
 		if (ConvertFromUTF8<UCS2CHAR>(Dest, DestLen, Src, SrcLen) == -1)
@@ -577,49 +698,61 @@ namespace UE::Core::Private
 		}
 		return Dest;
 	}
-}
-
-template <typename DestEncoding, typename SourceEncoding>
-void FGenericPlatformString::LogBogusChars(const SourceEncoding* Src, int32 SrcSize)
-{
-	static_assert(TIsFixedWidthCharEncoding_V<SourceEncoding>, "Currently unimplemented for non-fixed-width source conversions");
 
-	FString SrcStr;
-	bool    bFoundBogusChars = false;
-	for (; SrcSize; --SrcSize)
+	template <typename DestEncoding, typename SourceEncoding, typename SourceEndType>
+	void LogBogusCharsImpl(const SourceEncoding* Src, SourceEndType SourceEnd)
 	{
-		SourceEncoding SrcCh = *Src++;
-		if (!CanConvertCodepoint<DestEncoding>(SrcCh))
-		{
-			SrcStr += FString::Printf(TEXT("[0x%X]"), (int32)SrcCh);
-			bFoundBogusChars = true;
-		}
-		else if (CanConvertCodepoint<TCHAR>(SrcCh))
+		static_assert(TIsFixedWidthCharEncoding_V<SourceEncoding>, "Currently unimplemented for non-fixed-width source conversions");
+
+		FString SrcStr;
+		bool    bFoundBogusChars = false;
+		for (; !IsRangeEmpty(Src, SourceEnd); PopFront(Src, SourceEnd))
 		{
-			if (TChar<SourceEncoding>::IsLinebreak(SrcCh))
+			SourceEncoding SrcCh = *Src;
+			if (!FGenericPlatformString::CanConvertCodepoint<DestEncoding>(SrcCh))
 			{
-				if (bFoundBogusChars)
+				SrcStr += FString::Printf(TEXT("[0x%X]"), (int32)SrcCh);
+				bFoundBogusChars = true;
+			}
+			else if (FGenericPlatformString::CanConvertCodepoint<TCHAR>(SrcCh))
+			{
+				if (TChar<SourceEncoding>::IsLinebreak(SrcCh))
 				{
-					TrimStringAndLogBogusCharsError(SrcStr, GetEncodingTypeName<SourceEncoding>(), GetEncodingTypeName<DestEncoding>());
-					bFoundBogusChars = false;
+					if (bFoundBogusChars)
+					{
+						TrimStringAndLogBogusCharsError(SrcStr, FGenericPlatformString::GetEncodingTypeName<SourceEncoding>(), FGenericPlatformString::GetEncodingTypeName<DestEncoding>());
+						bFoundBogusChars = false;
+					}
+					SrcStr.Empty();
+				}
+				else
+				{
+					SrcStr.AppendChar((TCHAR)SrcCh);
 				}
-				SrcStr.Empty();
 			}
 			else
 			{
-				SrcStr.AppendChar((TCHAR)SrcCh);
+				SrcStr.AppendChar((TCHAR)'?');
 			}
 		}
-		else
+
+		if (bFoundBogusChars)
 		{
-			SrcStr.AppendChar((TCHAR)'?');
+			TrimStringAndLogBogusCharsError(SrcStr, FGenericPlatformString::GetEncodingTypeName<SourceEncoding>(), FGenericPlatformString::GetEncodingTypeName<DestEncoding>());
 		}
 	}
+}
 
-	if (bFoundBogusChars)
-	{
-		TrimStringAndLogBogusCharsError(SrcStr, GetEncodingTypeName<SourceEncoding>(), GetEncodingTypeName<DestEncoding>());
-	}
+template <typename DestEncoding, typename SourceEncoding>
+void FGenericPlatformString::LogBogusChars(const SourceEncoding* Src)
+{
+	UE::Core::Private::LogBogusCharsImpl<DestEncoding>(Src, UE::Core::Private::FNullTerminal{});
+}
+
+template <typename DestEncoding, typename SourceEncoding>
+void FGenericPlatformString::LogBogusChars(const SourceEncoding* Src, int32 SrcSize)
+{
+	UE::Core::Private::LogBogusCharsImpl<DestEncoding>(Src, SrcSize);
 }
 
 namespace GenericPlatformStringPrivate
@@ -659,17 +792,28 @@ int32 FGenericPlatformString::Strncmp(const WIDECHAR* Str1, const UTF8CHAR* Str2
 int32 FGenericPlatformString::Strncmp(const UTF8CHAR* Str1, const UTF8CHAR* Str2, SIZE_T Count) { return GenericPlatformStringPrivate::StrncmpImpl(Str1, Str2, Count); }
 
 #if !UE_BUILD_DOCS
+template CORE_API void FGenericPlatformString::LogBogusChars<ANSICHAR, WIDECHAR>(const WIDECHAR* Src);
 template CORE_API void FGenericPlatformString::LogBogusChars<ANSICHAR, WIDECHAR>(const WIDECHAR* Src, int32 SrcSize);
+template CORE_API void FGenericPlatformString::LogBogusChars<ANSICHAR, UCS2CHAR>(const UCS2CHAR* Src);
 template CORE_API void FGenericPlatformString::LogBogusChars<ANSICHAR, UCS2CHAR>(const UCS2CHAR* Src, int32 SrcSize);
+template CORE_API void FGenericPlatformString::LogBogusChars<WIDECHAR, ANSICHAR>(const ANSICHAR* Src);
 template CORE_API void FGenericPlatformString::LogBogusChars<WIDECHAR, ANSICHAR>(const ANSICHAR* Src, int32 SrcSize);
+template CORE_API void FGenericPlatformString::LogBogusChars<WIDECHAR, UCS2CHAR>(const UCS2CHAR* Src);
 template CORE_API void FGenericPlatformString::LogBogusChars<WIDECHAR, UCS2CHAR>(const UCS2CHAR* Src, int32 SrcSize);
+template CORE_API void FGenericPlatformString::LogBogusChars<UCS2CHAR, ANSICHAR>(const ANSICHAR* Src);
 template CORE_API void FGenericPlatformString::LogBogusChars<UCS2CHAR, ANSICHAR>(const ANSICHAR* Src, int32 SrcSize);
+template CORE_API void FGenericPlatformString::LogBogusChars<UCS2CHAR, WIDECHAR>(const WIDECHAR* Src);
 template CORE_API void FGenericPlatformString::LogBogusChars<UCS2CHAR, WIDECHAR>(const WIDECHAR* Src, int32 SrcSize);
+template CORE_API void FGenericPlatformString::LogBogusChars<UTF8CHAR, ANSICHAR>(const ANSICHAR* Src);
 template CORE_API void FGenericPlatformString::LogBogusChars<UTF8CHAR, ANSICHAR>(const ANSICHAR* Src, int32 SrcSize);
+template CORE_API void FGenericPlatformString::LogBogusChars<UTF8CHAR, WIDECHAR>(const WIDECHAR* Src);
 template CORE_API void FGenericPlatformString::LogBogusChars<UTF8CHAR, WIDECHAR>(const WIDECHAR* Src, int32 SrcSize);
+template CORE_API void FGenericPlatformString::LogBogusChars<UTF8CHAR, UCS2CHAR>(const UCS2CHAR* Src);
 template CORE_API void FGenericPlatformString::LogBogusChars<UTF8CHAR, UCS2CHAR>(const UCS2CHAR* Src, int32 SrcSize);
 #if PLATFORM_TCHAR_IS_CHAR16
+template CORE_API void FGenericPlatformString::LogBogusChars<wchar_t, char16_t>(const char16_t* Src);
 template CORE_API void FGenericPlatformString::LogBogusChars<wchar_t, char16_t>(const char16_t* Src, int32 SrcSize);
+template CORE_API void FGenericPlatformString::LogBogusChars<char16_t, wchar_t>(const wchar_t* Src);
 template CORE_API void FGenericPlatformString::LogBogusChars<char16_t, wchar_t>(const wchar_t* Src, int32 SrcSize);
 #endif
 #endif
diff --git a/Engine/Source/Runtime/Core/Private/Tests/Misc/PlatformStringConvertTests.cpp b/Engine/Source/Runtime/Core/Private/Tests/Misc/PlatformStringConvertTests.cpp
new file mode 100644
--- /dev/null
+++ b/Engine/Source/Runtime/Core/Private/Tests/Misc/PlatformStringConvertTests.cpp
@@ -0,0 +1,100 @@
+// Copyright Epic Games, Inc. All Rights Reserved.
+
+#include "CoreTypes.h"
+#include "Misc/AutomationTest.h"
+#include <type_traits>
+
+#if WITH_DEV_AUTOMATION_TESTS
+
+IMPLEMENT_SIMPLE_AUTOMATION_TEST(FStringConvertTest, "System.Core.Misc.StringConvertTests", EAutomationTestFlags::ApplicationContextMask | EAutomationTestFlags::SmokeFilter)
+
+bool FStringConvertTest::RunTest(const FString& Parameters)
+{
+	auto CheckNullTerminatedStringConversion = [this](const TCHAR* TestName, const auto* From, const auto* To)
+	{
+		using ToType = std::decay_t<decltype(*To)>;
+
+		// Basic version of Strcmp for arbitrary types
+		auto Strcmp = [](const auto* Lhs, const auto* Rhs)
+		{
+			for (;;)
+			{
+				if (!*Lhs)
+				{
+					return *Rhs ? -1 : 0;
+				}
+
+				if (!*Rhs)
+				{
+					return 1;
+				}
+
+				if (*Lhs != *Rhs)
+				{
+					return *Lhs < *Rhs ? -1 : 1;
+				}
+
+				++Lhs;
+				++Rhs;
+			}
+		};
+
+		int32 ConvertedLength = FPlatformString::ConvertedLength<ToType>(From);
+
+		// Allocate a buffer to hold the converted string and set it with 0xCD
+		TUniquePtr<ToType[]> UniquePtr = MakeUniqueForOverwrite<ToType[]>(ConvertedLength + 8);
+		ToType* Ptr = UniquePtr.Get();
+		FMemory::Memset(Ptr, 0xCD, (ConvertedLength + 8) * sizeof(ToType));
+
+		ToType* End = FPlatformString::Convert(Ptr + 4, ConvertedLength, From);
+
+		TestTrue(FString::Printf(TEXT("%s converted length matches"), TestName), End == Ptr + 4 + ConvertedLength);
+
+		for (int i = 0; i != 4 * sizeof(ToType); ++i)
+		{
+			TestTrue(FString::Printf(TEXT("%s guard block is preserved"), TestName), ((unsigned char*)Ptr)[i] == 0xCD && ((unsigned char*)End)[i] == 0xCD);
+		}
+
+		TestTrue(FString::Printf(TEXT("%s converted string matches"), TestName), Strcmp(To, Ptr + 4) == 0);
+	};
+
+	// We need to do UTF-8 literals like below because MSVC doesn't encode numeric escape sequences properly in u8"" literals:
+	// https://developercommunity.visualstudio.com/t/hex-escape-codes-in-a-utf8-literal-are-t/225847
+	//
+	// We cast to UTF16CHAR* because UTF16TEXT doesn't currently result in an array of UTF16CHAR.
+
+	// Freely interconvertible
+	{
+		const UTF16CHAR* UnrealEngineIsGreat16 = (const UTF16CHAR*)UTF16TEXT("\xD835\xDE50\x043F\xD835\xDDCB\xD835\xDE26\xD835\xDD52\x006C\x0020\xD835\xDC6C\xD835\xDD2B\xFF47\xD835\xDCBE\xD835\xDD93\xD835\xDD56\x0020\xD835\xDE92\xFF53\x0020\xD835\xDC6E\xD835\xDE9B\x212E\xD835\xDF36\xFF54\x0021");
+		const UTF8CHAR*  UnrealEngineIsGreat8  = (const UTF8CHAR*)"\xF0\x9D\x99\x90\xD0\xBF\xF0\x9D\x97\x8B\xF0\x9D\x98\xA6\xF0\x9D\x95\x92\x6C\x20\xF0\x9D\x91\xAC\xF0\x9D\x94\xAB\xEF\xBD\x87\xF0\x9D\x92\xBE\xF0\x9D\x96\x93\xF0\x9D\x95\x96\x20\xF0\x9D\x9A\x92\xEF\xBD\x93\x20\xF0\x9D\x91\xAE\xF0\x9D\x9A\x9B\xE2\x84\xAE\xF0\x9D\x9C\xB6\xEF\xBD\x94\x21";
+		CheckNullTerminatedStringConversion(TEXT("UnrealEngineIsGreat UTF-16 -> UTF-8"), UnrealEngineIsGreat16, UnrealEngineIsGreat8);
+		CheckNullTerminatedStringConversion(TEXT("UnrealEngineIsGreat UTF-8 -> UTF-16"), UnrealEngineIsGreat8, UnrealEngineIsGreat16);
+
+		const UTF16CHAR* Internationalization16 = (const UTF16CHAR*)UTF16TEXT("\x0049\x00F1\x0074\x00EB\x0072\x006E\x00E2\x0074\x0069\x00F4\x006E\x00E0\x006C\x0069\x007A\x00E6\x0074\x0069\x00F8\x006E");
+		const UTF8CHAR*  Internationalization8  = (const UTF8CHAR*)"\x49\xC3\xB1\x74\xC3\xAB\x72\x6E\xC3\xA2\x74\x69\xC3\xB4\x6E\xC3\xA0\x6C\x69\x7A\xC3\xA6\x74\x69\xC3\xB8\x6E";
+		CheckNullTerminatedStringConversion(TEXT("Internationalization UTF-16 -> UTF-8"), Internationalization16, Internationalization8);
+		CheckNullTerminatedStringConversion(TEXT("Internationalization UTF-8 -> UTF-16"), Internationalization8, Internationalization16);
+
+		const UTF16CHAR* WomanEmoji16 = (const UTF16CHAR*)UTF16TEXT("\xD83D\xDC69");
+		const UTF8CHAR*  WomanEmoji8  = (const UTF8CHAR*)"\xf0\x9f\x91\xa9";
+		CheckNullTerminatedStringConversion(TEXT("WomanEmoji UTF-16 -> UTF-8"), WomanEmoji16, WomanEmoji8);
+		CheckNullTerminatedStringConversion(TEXT("WomanEmoji UTF-8 -> UTF-16"), WomanEmoji8, WomanEmoji16);
+
+		const UTF16CHAR* LanguageProcessingChinese16 = (const UTF16CHAR*)UTF16TEXT("\x8bed\x8a00\x5904\x7406");
+		const UTF8CHAR*  LanguageProcessingChinese8  = (const UTF8CHAR*)"\xe8\xaf\xad\xe8\xa8\x80\xe5\xa4\x84\xe7\x90\x86";
+		CheckNullTerminatedStringConversion(TEXT("LanguageProcessingChinese UTF-16 -> UTF-8"), LanguageProcessingChinese16, LanguageProcessingChinese8);
+		CheckNullTerminatedStringConversion(TEXT("LanguageProcessingChinese UTF-8 -> UTF-16"), LanguageProcessingChinese8, LanguageProcessingChinese16);
+	}
+
+	// Broken UTF-16 encodings
+	{
+		CheckNullTerminatedStringConversion(TEXT("Double UTF-16 high surrogate"),   (const UTF16CHAR*)UTF16TEXT("q\xD83D\xD83Dq"), (const UTF8CHAR*)"q??q");
+		CheckNullTerminatedStringConversion(TEXT("Double UTF-16 low surrogate"),    (const UTF16CHAR*)UTF16TEXT("q\xDC69\xDC69q"), (const UTF8CHAR*)"q??q");
+		CheckNullTerminatedStringConversion(TEXT("Orphaned UTF-16 high surrogate"), (const UTF16CHAR*)UTF16TEXT("q\xDC69q"),       (const UTF8CHAR*)"q?q");
+		CheckNullTerminatedStringConversion(TEXT("Orphaned UTF-16 low surrogate"),  (const UTF16CHAR*)UTF16TEXT("q\xDC69q"),       (const UTF8CHAR*)"q?q");
+	}
+
+	return !HasAnyErrors();
+}
+
+#endif //WITH_DEV_AUTOMATION_TESTS
diff --git a/Engine/Source/Runtime/Core/Public/GenericPlatform/GenericPlatformString.h b/Engine/Source/Runtime/Core/Public/GenericPlatform/GenericPlatformString.h
--- a/Engine/Source/Runtime/Core/Public/GenericPlatform/GenericPlatformString.h
+++ b/Engine/Source/Runtime/Core/Public/GenericPlatform/GenericPlatformString.h
@@ -15,18 +15,30 @@
 namespace UE::Core::Private
 {
 	// The Dest parameter is just used for overload resolution
+	CORE_API int32 GetConvertedLength(const UTF8CHAR* Dest, const WIDECHAR*  Src);
 	CORE_API int32 GetConvertedLength(const UTF8CHAR* Dest, const WIDECHAR*  Src, int32 SrcLen);
+	CORE_API int32 GetConvertedLength(const UTF8CHAR* Dest, const UCS2CHAR*  Src);
 	CORE_API int32 GetConvertedLength(const UTF8CHAR* Dest, const UCS2CHAR*  Src, int32 SrcLen);
+	CORE_API int32 GetConvertedLength(const UTF8CHAR* Dest, const UTF32CHAR* Src);
 	CORE_API int32 GetConvertedLength(const UTF8CHAR* Dest, const UTF32CHAR* Src, int32 SrcLen);
+	CORE_API int32 GetConvertedLength(const ANSICHAR* Dest, const UTF8CHAR*  Src);
 	CORE_API int32 GetConvertedLength(const ANSICHAR* Dest, const UTF8CHAR*  Src, int32 SrcLen);
+	CORE_API int32 GetConvertedLength(const WIDECHAR* Dest, const UTF8CHAR*  Src);
 	CORE_API int32 GetConvertedLength(const WIDECHAR* Dest, const UTF8CHAR*  Src, int32 SrcLen);
+	CORE_API int32 GetConvertedLength(const UCS2CHAR* Dest, const UTF8CHAR*  Src);
 	CORE_API int32 GetConvertedLength(const UCS2CHAR* Dest, const UTF8CHAR*  Src, int32 SrcLen);
 
+	CORE_API UTF8CHAR* Convert(UTF8CHAR* Dest, int32 DestLen, const WIDECHAR*  Src);
 	CORE_API UTF8CHAR* Convert(UTF8CHAR* Dest, int32 DestLen, const WIDECHAR*  Src, int32 SrcLen);
+	CORE_API UTF8CHAR* Convert(UTF8CHAR* Dest, int32 DestLen, const UCS2CHAR*  Src);
 	CORE_API UTF8CHAR* Convert(UTF8CHAR* Dest, int32 DestLen, const UCS2CHAR*  Src, int32 SrcLen);
+	CORE_API UTF8CHAR* Convert(UTF8CHAR* Dest, int32 DestLen, const UTF32CHAR* Src);
 	CORE_API UTF8CHAR* Convert(UTF8CHAR* Dest, int32 DestLen, const UTF32CHAR* Src, int32 SrcLen);
+	CORE_API ANSICHAR* Convert(ANSICHAR* Dest, int32 DestLen, const UTF8CHAR*  Src);
 	CORE_API ANSICHAR* Convert(ANSICHAR* Dest, int32 DestLen, const UTF8CHAR*  Src, int32 SrcLen);
+	CORE_API WIDECHAR* Convert(WIDECHAR* Dest, int32 DestLen, const UTF8CHAR*  Src);
 	CORE_API WIDECHAR* Convert(WIDECHAR* Dest, int32 DestLen, const UTF8CHAR*  Src, int32 SrcLen);
+	CORE_API UCS2CHAR* Convert(UCS2CHAR* Dest, int32 DestLen, const UTF8CHAR*  Src);
 	CORE_API UCS2CHAR* Convert(UCS2CHAR* Dest, int32 DestLen, const UTF8CHAR*  Src, int32 SrcLen);
 }
 
@@ -170,6 +182,101 @@ struct FGenericPlatformString : public FGenericPlatformStricmp
 		enum { Value = TIsFixedWidthCharEncoding_V<EncodingA> && TIsFixedWidthCharEncoding_V<EncodingB> && sizeof(EncodingA) == sizeof(EncodingB) };
 	};
 
+
+	/**
+	 * Converts the null-terminated Src string range from SourceEncoding to DestEncoding and writes it to the [Dest, Dest+DestSize) range, including a null terminator.
+	 * If the Dest range is not big enough to hold the converted output, NULL is returned.  In this case, nothing should be assumed about the contents of Dest.
+	 *
+	 * @param Dest      The start of the destination buffer.
+	 * @param DestSize  The size of the destination buffer.
+	 * @param Src       The start of the string to convert.
+	 * @return          A pointer to one past the last-written element.
+	 */
+	template <typename SourceEncoding, typename DestEncoding>
+	static FORCEINLINE DestEncoding* Convert(DestEncoding* Dest, int32 DestSize, const SourceEncoding* Src)
+	{
+		if constexpr (TIsCharEncodingSimplyConvertibleTo_V<SourceEncoding, DestEncoding>)
+		{
+			for (;;)
+			{
+				if (DestSize == 0)
+				{
+					return nullptr;
+				}
+
+				if (!(*Dest++ = (DestEncoding)*Src++))
+				{
+					return Dest;
+				}
+
+				--DestSize;
+			}
+		}
+		else if constexpr (TIsFixedWidthCharEncoding_V<SourceEncoding> && TIsFixedWidthCharEncoding_V<DestEncoding>)
+		{
+			DestEncoding* DestCopy     = Dest;
+			DestEncoding* SrcCopy      = Src;
+			int32         DestSizeCopy = DestSize;
+
+			bool bInvalidChars = false;
+			for (;;)
+			{
+				if (DestSize == 0)
+				{
+					Dest = nullptr;
+					break;
+				}
+
+				SourceEncoding SrcCh = *Src++;
+				*Dest++ = (DestEncoding)SrcCh;
+				if (!SrcCh)
+				{
+					break;
+				}
+				bInvalidChars |= !CanConvertCodepoint<DestEncoding>(SrcCh);
+
+				--DestSize;
+			}
+
+			if (bInvalidChars)
+			{
+				for (;;)
+				{
+					if (DestSizeCopy == 0)
+					{
+						break;
+					}
+
+					SourceEncoding SrcCh = *SrcCopy++;
+					if (!SrcCh)
+					{
+						break;
+					}
+					if (!CanConvertCodepoint<DestEncoding>(SrcCh))
+					{
+						*DestCopy = UNICODE_BOGUS_CHAR_CODEPOINT;
+					}
+					++DestCopy;
+
+					--DestSizeCopy;
+				}
+
+				LogBogusChars<DestEncoding>(Src);
+			}
+
+			return Dest;
+		}
+		else
+		{
+			DestEncoding* Result = UE::Core::Private::Convert(Dest, DestSize, Src);
+			if (Result)
+			{
+				*Result++ = (DestEncoding)0;
+			}
+			return Result;
+		}
+	}
+
 	/**
 	 * Converts the [Src, Src+SrcSize) string range from SourceEncoding to DestEncoding and writes it to the [Dest, Dest+DestSize) range.
 	 * The Src range should contain a null terminator if a null terminator is required in the output.
@@ -237,6 +344,32 @@ struct FGenericPlatformString : public FGenericPlatformStricmp
 	}
 
 
+	/**
+	 * Returns the required buffer length for the null-terminated Src string when converted to the DestChar encoding.
+	 * The returned length includes the space for the null terminator (equivalent to strlen+1).
+	 * 
+	 * @param  Src  The start of the string to convert.
+	 * @return      The number of DestChar elements that Src will be converted into.
+	 */
+	template <typename DestEncoding, typename SourceEncoding>
+	static int32 ConvertedLength(const SourceEncoding* Src)
+	{
+		if constexpr (TIsCharEncodingSimplyConvertibleTo_V<SourceEncoding, DestEncoding> || (TIsFixedWidthCharEncoding_V<SourceEncoding> && TIsFixedWidthCharEncoding_V<DestEncoding>))
+		{
+			int32 Result = 0;
+			while (*Src)
+			{
+				++Src;
+				++Result;
+			}
+			return Result + 1;
+		}
+		else
+		{
+			return UE::Core::Private::GetConvertedLength((DestEncoding*)nullptr, Src) + 1;
+		}
+	}
+
 	/**
 	 * Returns the required buffer length for the [Src, Src+SrcSize) string when converted to the DestChar encoding.
 	 * The Src range should contain a null terminator if a null terminator is required in the output.
@@ -280,6 +413,15 @@ private:
 	static CORE_API void* Memcpy(void* Dest, const void* Src, SIZE_T Count);
 
 
+	/**
+	 * Logs a message about bogus characters which were detected during string conversion.
+	 *
+	 * @param Src     Pointer to the null-terminated string being converted.
+	 */
+	template <typename DestEncoding, typename SourceEncoding>
+	static CORE_API void LogBogusChars(const SourceEncoding* Src);
+
+
 	/**
 	 * Logs a message about bogus characters which were detected during string conversion.
 	 *
diff --git a/Engine/Source/Runtime/Core/Public/HAL/Platform.h b/Engine/Source/Runtime/Core/Public/HAL/Platform.h
--- a/Engine/Source/Runtime/Core/Public/HAL/Platform.h
+++ b/Engine/Source/Runtime/Core/Public/HAL/Platform.h
@@ -1171,7 +1171,7 @@ namespace UE::Core::Private
 }
 
 #define UTF8TEXT(x) (UE::Core::Private::ToUTF8Literal(UTF8TEXT_PASTE(x)))
-
+#define UTF16TEXT(x) UTF16TEXT_PASTE(x)
 #define WIDETEXT(str) WIDETEXT_PASTE(str)
 
 // IWYU pragma: end_exports
diff --git a/Engine/Source/Runtime/Core/Public/Templates/UniquePtr.h b/Engine/Source/Runtime/Core/Public/Templates/UniquePtr.h
--- a/Engine/Source/Runtime/Core/Public/Templates/UniquePtr.h
+++ b/Engine/Source/Runtime/Core/Public/Templates/UniquePtr.h
@@ -742,6 +742,7 @@ struct TIsBitwiseConstructible<TUniquePtr<T>, T*>
 
 /**
  * Allocates a new object of type T with the given arguments and returns it as a TUniquePtr.  Disabled for array-type TUniquePtrs.
+ * The object is value-initialized, which will call a user-defined default constructor if it exists, but a trivial type will be zeroed.
  *
  * @param Args The arguments to pass to the constructor of T.
  *
@@ -753,8 +754,23 @@ FORCEINLINE typename TEnableIf<!TIsArray<T>::Value, TUniquePtr<T>>::Type MakeUni
 	return TUniquePtr<T>(new T(Forward<TArgs>(Args)...));
 }
 
+/**
+ * Allocates a new object of type T with the given arguments and returns it as a TUniquePtr.  Disabled for array-type TUniquePtrs.
+ * The object is default-initialized, which will call a user-defined default constructor if it exists, but a trivial type will be uninitialized.
+ *
+ * @param Args The arguments to pass to the constructor of T.
+ *
+ * @return A TUniquePtr which points to a newly-constructed T with the specified Args.
+ */
+template <typename T, typename... TArgs>
+FORCEINLINE typename TEnableIf<!TIsArray<T>::Value, TUniquePtr<T>>::Type MakeUniqueForOverwrite()
+{
+	return TUniquePtr<T>(new T);
+}
+
 /**
  * Allocates a new array of type T with the given size and returns it as a TUniquePtr.  Only enabled for array-type TUniquePtrs.
+ * Elements are value-initialized, which will call a user-defined default constructor if it exists, but causes trivial types to be zeroed.
  *
  * @param Size The size of the array to allocate.
  *
@@ -767,6 +783,21 @@ FORCEINLINE typename TEnableIf<TIsUnboundedArray<T>::Value, TUniquePtr<T>>::Type
 	return TUniquePtr<T>(new ElementType[Size]());
 }
 
+/**
+ * Allocates a new array of type T with the given size and returns it as a TUniquePtr.  Only enabled for array-type TUniquePtrs.
+ * Elements are default-initialized, which will call a user-defined default constructor if it exists, but causes trivial types to be left uninitialized.
+ *
+ * @param Size The size of the array to allocate.
+ *
+ * @return A TUniquePtr which points to a newly-constructed T array of the specified Size.
+ */
+template <typename T>
+FORCEINLINE typename TEnableIf<TIsUnboundedArray<T>::Value, TUniquePtr<T>>::Type MakeUniqueForOverwrite(SIZE_T Size)
+{
+	typedef typename TRemoveExtent<T>::Type ElementType;
+	return TUniquePtr<T>(new ElementType[Size]);
+}
+
 /**
  * Overload to cause a compile error when MakeUnique<T[N]> is attempted.  Use MakeUnique<T>(N) instead.
  */
