Add cvars for controlling which VA system is mounted during initialization and if it should be lazy initialized or not.

#rnx

- We currently provide ways to change these options via code define, config file options and command line options but it has been requested that we also support this level of customization via cvars as well.
- The cvar 'VA.DisableSystem' when enabled will act as though [Core.ContentVirtualization]SystemName is set to 'None' and disable the VA system entirely by mounting the null reference implementation.
- The cvar 'VA.LazyInitSystem' when enabled will act as though [Core.ContentVirtualization]LazyInit is set to true and the VA system will not be initialized until first used.
- The logic for both of these was getting a little verbose, so I moved them to their own functions to make ::Initialize easier to read.

diff --git a/Engine/Source/Runtime/Core/Private/Virtualization/VirtualizationSystem.cpp b/Engine/Source/Runtime/Core/Private/Virtualization/VirtualizationSystem.cpp
--- a/Engine/Source/Runtime/Core/Private/Virtualization/VirtualizationSystem.cpp
+++ b/Engine/Source/Runtime/Core/Private/Virtualization/VirtualizationSystem.cpp
@@ -3,6 +3,7 @@
 #include "Virtualization/VirtualizationSystem.h"
 
 #include "CoreGlobals.h"
+#include "HAL/IConsoleManager.h"
 #include "Misc/App.h"
 #include "Misc/ConfigCacheIni.h"
 #include "Misc/ScopeLock.h"
@@ -36,13 +37,23 @@
 namespace UE::Virtualization
 {
 
+static TAutoConsoleVariable<bool> CVarDisableSystem(
+	TEXT("VA.DisableSystem"),
+	false,
+	TEXT("When true the VA system will be disabled as though 'SystemName' was 'None'"));
+
+static TAutoConsoleVariable<bool> CVarLazyInitSystem(
+	TEXT("VA.LazyInitSystem"),
+	false,
+	TEXT("When true the VA system will be lazy initialized on first use"));
+
 /** Default implementation to be used when the system is disabled */
 class FNullVirtualizationSystem : public IVirtualizationSystem
 {
 public:
 	FNullVirtualizationSystem()
 	{
-		UE_LOG(LogVirtualization, Log, TEXT("FNullVirtualizationSystem mounted, virtualization will be disabled"));
+		UE_LOG(LogVirtualization, Display, TEXT("FNullVirtualizationSystem mounted, virtualization will be disabled"));
 	}
 
 	virtual ~FNullVirtualizationSystem() = default;
@@ -165,6 +176,74 @@ Private::IVirtualizationSystemFactory* FindFactory(FName SystemName)
 	return nullptr;
 }
 
+/** Utility determining if the virtualization system should be initialized immediately or on first use */
+bool ShouldLazyInitializeSystem(const FConfigFile& ConfigFile)
+{
+#if UE_VIRTUALIZATION_SYSTEM_LAZY_INIT
+	UE_LOG(LogVirtualization, Display, TEXT("The virtualization system will lazy initialize due to code"));
+	return true;
+#else
+	if (FParse::Param(FCommandLine::Get(), TEXT("VA-LazyInit")))
+	{
+		UE_LOG(LogVirtualization, Display, TEXT("The virtualization system will lazy initialize due to the command line"));
+		return true;
+	}
+	else if (CVarLazyInitSystem.GetValueOnAnyThread())
+	{
+		UE_LOG(LogVirtualization, Display, TEXT("The virtualization system will lazy initialize due to a cvar"));
+		return true;
+	}
+	else
+	{
+		bool bLazyInit = false;
+		ConfigFile.GetBool(TEXT("Core.ContentVirtualization"), TEXT("LazyInit"), bLazyInit);
+
+		if (bLazyInit)
+		{
+			UE_LOG(LogVirtualization, Display, TEXT("The virtualization system will lazy initialize due to the  ini file option"));
+			return true;
+		}
+		else
+		{
+			return false;
+		}
+	}
+#endif //UE_VIRTUALIZATION_SYSTEM_LAZY_INIT
+}
+
+/*** Utility determining which virtualization system should be mounted during initialization */
+FName FindSystemToMount(const FConfigFile& ConfigFile)
+{
+#if UE_DISABLE_VIRTUALIZATION_SYSTEM
+	UE_LOG(LogVirtualization, Display, TEXT("The virtualization system has been disabled by code"));
+	return FName();
+#else
+	if (FParse::Param(FCommandLine::Get(), TEXT("VA-Disable")))
+	{
+		UE_LOG(LogVirtualization, Display, TEXT("The virtualization system has been disabled by the command line"));
+		return FName();
+	}
+	else if (CVarDisableSystem.GetValueOnAnyThread())
+	{
+		UE_LOG(LogVirtualization, Display, TEXT("The virtualization system has been disabled by cvar"));
+		return FName();
+	}
+	else
+	{
+		FString SystemName;
+		if (ConfigFile.GetString(TEXT("Core.ContentVirtualization"), TEXT("SystemName"), SystemName))
+		{
+			UE_LOG(LogVirtualization, Display, TEXT("VirtualizationSystem name found in ini file: %s"), *SystemName);
+			return FName(SystemName);	
+		}
+		else
+		{
+			return FName();
+		}
+	}
+#endif //UE_DISABLE_VIRTUALIZATION_SYSTEM
+}
+
 void Initialize(EInitializationFlags Flags)
 {
 	const FConfigFile* ConfigFile = GConfig->Find(GEngineIni);
@@ -191,24 +270,10 @@ void Initialize(const FInitParams& InitParams, EInitializationFlags Flags)
 	// initialize or not.
 	if (!EnumHasAnyFlags(Flags, EInitializationFlags::ForceInitialize))
 	{
-#if UE_VIRTUALIZATION_SYSTEM_LAZY_INIT
-		return;
-#else
-		if (FParse::Param(FCommandLine::Get(), TEXT("VA-LazyInit")))
+		if (ShouldLazyInitializeSystem(InitParams.ConfigFile))
 		{
 			return;
 		}
-		else
-		{
-			bool bLazyInit = false;
-			InitParams.ConfigFile.GetBool(TEXT("Core.ContentVirtualization"), TEXT("LazyInit"), bLazyInit);
-
-			if (bLazyInit)
-			{
-				return;
-			}
-		}
-#endif //UE_VIRTUALIZATION_SYSTEM_LAZY_INIT
 	}
 
 	// Only allow one thread to initialize the system at a time
@@ -222,25 +287,7 @@ void Initialize(const FInitParams& InitParams, EInitializationFlags Flags)
 		return;
 	}
 
-	FName SystemName;
-
-#if UE_DISABLE_VIRTUALIZATION_SYSTEM == 0
-	if (FParse::Param(FCommandLine::Get(), TEXT("VA-Disable")))
-	{
-		UE_LOG(LogVirtualization, Display, TEXT("The virtualization system has been disabled by the command line"));
-	}
-	else
-	{
-		FString RawSystemName;
-		if (InitParams.ConfigFile.GetString(TEXT("Core.ContentVirtualization"), TEXT("SystemName"), RawSystemName))
-		{
-			SystemName = FName(RawSystemName);
-			UE_LOG(LogVirtualization, Display, TEXT("VirtualizationSystem name found in ini file: %s"), *RawSystemName);
-		}
-	}
-#else
-	UE_LOG(LogVirtualization, Display, TEXT("The virtualization system has been disabled by code"));
-#endif //UE_DISABLE_VIRTUALIZATION_SYSTEM
+	FName SystemName = FindSystemToMount(InitParams.ConfigFile);
 
 	if (!SystemName.IsNone())
 	{
