Add new LLM subtags for UI & increase UI LLM coverage.

diff --git a/Engine/Plugins/Experimental/CommonUI/Source/CommonUI/Public/UITag.h b/Engine/Plugins/Experimental/CommonUI/Source/CommonUI/Public/UITag.h
--- a/Engine/Plugins/Experimental/CommonUI/Source/CommonUI/Public/UITag.h
+++ b/Engine/Plugins/Experimental/CommonUI/Source/CommonUI/Public/UITag.h
@@ -3,6 +3,7 @@
 
 #include "GameplayTagContainer.h"
 #include "GameplayTagsManager.h"
+#include "Trace/SlateMemoryTags.h"
 
 #include "UITag.generated.h"
 
@@ -72,7 +73,8 @@ class TTypedTagStaticImpl
 	}
 
 	TTypedTagStaticImpl()
-	{
+	{ 
+		LLM_SCOPE(ELLMTag::UI);
 		UGameplayTagsManager::OnLastChanceToAddNativeTags().AddLambda([this]()
 			{
 				StaticImpl.RootTag = UGameplayTagsManager::Get().AddNativeGameplayTag(TagT::GetRootTagStr());
diff --git a/Engine/Source/Runtime/Core/Private/HAL/LowLevelMemTracker.cpp b/Engine/Source/Runtime/Core/Private/HAL/LowLevelMemTracker.cpp
--- a/Engine/Source/Runtime/Core/Private/HAL/LowLevelMemTracker.cpp
+++ b/Engine/Source/Runtime/Core/Private/HAL/LowLevelMemTracker.cpp
@@ -161,10 +161,28 @@ void FLowLevelMemTracker::DumpToLog(EDumpFormat DumpFormat, FOutputDevice* Outpu
 
 DECLARE_LLM_MEMORY_STAT(TEXT("LLM Overhead"), STAT_LLMOverheadTotal, STATGROUP_LLMOverhead);
 
+/*
+* LLM Summary stats referenced by ELLMTagNames
+*/
 DEFINE_STAT(STAT_EngineSummaryLLM);
 DEFINE_STAT(STAT_ProjectSummaryLLM);
 DEFINE_STAT(STAT_NetworkingSummaryLLM);
 DEFINE_STAT(STAT_AudioSummaryLLM);
+DEFINE_STAT(STAT_TrackedTotalSummaryLLM);
+DEFINE_STAT(STAT_MeshesSummaryLLM);
+DEFINE_STAT(STAT_PhysicsSummaryLLM);
+DEFINE_STAT(STAT_PhysXSummaryLLM);
+DEFINE_STAT(STAT_ChaosSummaryLLM);
+DEFINE_STAT(STAT_UObjectSummaryLLM);
+DEFINE_STAT(STAT_AnimationSummaryLLM);
+DEFINE_STAT(STAT_StaticMeshSummaryLLM);
+DEFINE_STAT(STAT_MaterialsSummaryLLM);
+DEFINE_STAT(STAT_ParticlesSummaryLLM);
+DEFINE_STAT(STAT_NiagaraSummaryLLM);
+DEFINE_STAT(STAT_UISummaryLLM);
+DEFINE_STAT(STAT_NavigationSummaryLLM);
+DEFINE_STAT(STAT_TexturesSummaryLLM);
+DEFINE_STAT(STAT_MediaStreamingSummaryLLM);
 
 /*
  * LLM stats referenced by ELLMTagNames
@@ -278,25 +296,6 @@ DECLARE_LLM_MEMORY_STAT(TEXT("MMIO"), STAT_PlatformMMIOLLM, STATGROUP_LLMPlatfor
 DECLARE_LLM_MEMORY_STAT(TEXT("VirtualMemory"), STAT_PlatformVMLLM, STATGROUP_LLMPlatform);
 DECLARE_LLM_MEMORY_STAT(TEXT("CustomName"), STAT_CustomName, STATGROUP_LLMFULL);
 
-/*
-* LLM Summary stats referenced by ELLMTagNames
-*/
-DECLARE_LLM_MEMORY_STAT(TEXT("Total"), STAT_TrackedTotalSummaryLLM, STATGROUP_LLM);
-DECLARE_LLM_MEMORY_STAT(TEXT("Meshes"), STAT_MeshesSummaryLLM, STATGROUP_LLM);
-DECLARE_LLM_MEMORY_STAT(TEXT("Physics"), STAT_PhysicsSummaryLLM, STATGROUP_LLM);
-DECLARE_LLM_MEMORY_STAT(TEXT("PhysX"), STAT_PhysXSummaryLLM, STATGROUP_LLM);
-DECLARE_LLM_MEMORY_STAT(TEXT("Chaos"), STAT_ChaosSummaryLLM, STATGROUP_LLM);
-DECLARE_LLM_MEMORY_STAT(TEXT("UObject"), STAT_UObjectSummaryLLM, STATGROUP_LLM);
-DECLARE_LLM_MEMORY_STAT(TEXT("Animation"), STAT_AnimationSummaryLLM, STATGROUP_LLM);
-DECLARE_LLM_MEMORY_STAT(TEXT("StaticMesh"), STAT_StaticMeshSummaryLLM, STATGROUP_LLM);
-DECLARE_LLM_MEMORY_STAT(TEXT("Materials"), STAT_MaterialsSummaryLLM, STATGROUP_LLM);
-DECLARE_LLM_MEMORY_STAT(TEXT("Particles"), STAT_ParticlesSummaryLLM, STATGROUP_LLM);
-DECLARE_LLM_MEMORY_STAT(TEXT("Niagara"), STAT_NiagaraSummaryLLM, STATGROUP_LLM);
-DECLARE_LLM_MEMORY_STAT(TEXT("UI"), STAT_UISummaryLLM, STATGROUP_LLM);
-DECLARE_LLM_MEMORY_STAT(TEXT("Navigation"), STAT_NavigationSummaryLLM, STATGROUP_LLM);
-DECLARE_LLM_MEMORY_STAT(TEXT("Textures"), STAT_TexturesSummaryLLM, STATGROUP_LLM);
-DECLARE_LLM_MEMORY_STAT(TEXT("MediaStreaming"), STAT_MediaStreamingSummaryLLM, STATGROUP_LLM);
-
 namespace UE
 {
 namespace LLMPrivate
diff --git a/Engine/Source/Runtime/Core/Public/HAL/LowLevelMemStats.h b/Engine/Source/Runtime/Core/Public/HAL/LowLevelMemStats.h
--- a/Engine/Source/Runtime/Core/Public/HAL/LowLevelMemStats.h
+++ b/Engine/Source/Runtime/Core/Public/HAL/LowLevelMemStats.h
@@ -22,10 +22,28 @@
 	DECLARE_STATS_GROUP(TEXT("LLM Overhead"), STATGROUP_LLMOverhead, STATCAT_Advanced);
 	DECLARE_STATS_GROUP(TEXT("LLM Assets"), STATGROUP_LLMAssets, STATCAT_Advanced);
 
+	/*
+	* LLM Summary stats referenced by ELLMTagNames
+	*/
 	DECLARE_LLM_MEMORY_STAT_EXTERN(TEXT("Engine"), STAT_EngineSummaryLLM, STATGROUP_LLM, CORE_API);
 	DECLARE_LLM_MEMORY_STAT_EXTERN(TEXT("Project"), STAT_ProjectSummaryLLM, STATGROUP_LLM, CORE_API);
 	DECLARE_LLM_MEMORY_STAT_EXTERN(TEXT("Networking"), STAT_NetworkingSummaryLLM, STATGROUP_LLM, CORE_API);
 	DECLARE_LLM_MEMORY_STAT_EXTERN(TEXT("Audio"), STAT_AudioSummaryLLM, STATGROUP_LLM, CORE_API);
+	DECLARE_LLM_MEMORY_STAT_EXTERN(TEXT("Total"), STAT_TrackedTotalSummaryLLM, STATGROUP_LLM, CORE_API);
+	DECLARE_LLM_MEMORY_STAT_EXTERN(TEXT("Meshes"), STAT_MeshesSummaryLLM, STATGROUP_LLM, CORE_API);
+	DECLARE_LLM_MEMORY_STAT_EXTERN(TEXT("Physics"), STAT_PhysicsSummaryLLM, STATGROUP_LLM, CORE_API);
+	DECLARE_LLM_MEMORY_STAT_EXTERN(TEXT("PhysX"), STAT_PhysXSummaryLLM, STATGROUP_LLM, CORE_API);
+	DECLARE_LLM_MEMORY_STAT_EXTERN(TEXT("Chaos"), STAT_ChaosSummaryLLM, STATGROUP_LLM, CORE_API);
+	DECLARE_LLM_MEMORY_STAT_EXTERN(TEXT("UObject"), STAT_UObjectSummaryLLM, STATGROUP_LLM, CORE_API);
+	DECLARE_LLM_MEMORY_STAT_EXTERN(TEXT("Animation"), STAT_AnimationSummaryLLM, STATGROUP_LLM, CORE_API);
+	DECLARE_LLM_MEMORY_STAT_EXTERN(TEXT("StaticMesh"), STAT_StaticMeshSummaryLLM, STATGROUP_LLM, CORE_API);
+	DECLARE_LLM_MEMORY_STAT_EXTERN(TEXT("Materials"), STAT_MaterialsSummaryLLM, STATGROUP_LLM, CORE_API);
+	DECLARE_LLM_MEMORY_STAT_EXTERN(TEXT("Particles"), STAT_ParticlesSummaryLLM, STATGROUP_LLM, CORE_API);
+	DECLARE_LLM_MEMORY_STAT_EXTERN(TEXT("Niagara"), STAT_NiagaraSummaryLLM, STATGROUP_LLM, CORE_API);
+	DECLARE_LLM_MEMORY_STAT_EXTERN(TEXT("UI"), STAT_UISummaryLLM, STATGROUP_LLM, CORE_API);
+	DECLARE_LLM_MEMORY_STAT_EXTERN(TEXT("Navigation"), STAT_NavigationSummaryLLM, STATGROUP_LLM, CORE_API);
+	DECLARE_LLM_MEMORY_STAT_EXTERN(TEXT("Textures"), STAT_TexturesSummaryLLM, STATGROUP_LLM, CORE_API);
+	DECLARE_LLM_MEMORY_STAT_EXTERN(TEXT("MediaStreaming"), STAT_MediaStreamingSummaryLLM, STATGROUP_LLM, CORE_API);
 
 #else
 	#define DECLARE_LLM_MEMORY_STAT(CounterName,StatId,GroupId)
diff --git a/Engine/Source/Runtime/Slate/Private/Framework/Application/SlateApplication.cpp b/Engine/Source/Runtime/Slate/Private/Framework/Application/SlateApplication.cpp
--- a/Engine/Source/Runtime/Slate/Private/Framework/Application/SlateApplication.cpp
+++ b/Engine/Source/Runtime/Slate/Private/Framework/Application/SlateApplication.cpp
@@ -39,12 +39,11 @@
 #include "Framework/Application/IInputProcessor.h"
 #include "GenericPlatform/ITextInputMethodSystem.h"
 #include "Framework/Docking/TabCommands.h"
-#include "HAL/LowLevelMemTracker.h"
 #include "Math/UnitConversion.h"
-#include "HAL/LowLevelMemTracker.h"
 #include "ProfilingDebugging/CsvProfiler.h"
 #include "ProfilingDebugging/StallDetector.h"
 #include "Types/ReflectionMetadata.h"
+#include "Trace/SlateMemoryTags.h"
 #include "Trace/SlateTrace.h"
 #include "Styling/StarshipCoreStyle.h"
 #include "Styling/UMGCoreStyle.h"
@@ -1449,7 +1448,7 @@ static const TCHAR* LexToString(ESlateTickType TickType)
 
 void FSlateApplication::Tick(ESlateTickType TickType)
 {
-	LLM_SCOPE(ELLMTag::UI);
+	LLM_SCOPE_BYTAG(UI_Slate);
 
 	SCOPE_TIME_GUARD(TEXT("FSlateApplication::Tick"));
 	CSV_SCOPED_TIMING_STAT_EXCLUSIVE(UI);
diff --git a/Engine/Source/Runtime/Slate/Private/Framework/Commands/InputBindingManager.cpp b/Engine/Source/Runtime/Slate/Private/Framework/Commands/InputBindingManager.cpp
--- a/Engine/Source/Runtime/Slate/Private/Framework/Commands/InputBindingManager.cpp
+++ b/Engine/Source/Runtime/Slate/Private/Framework/Commands/InputBindingManager.cpp
@@ -2,6 +2,7 @@
 
 #include "Framework/Commands/InputBindingManager.h"
 #include "HAL/FileManager.h"
+#include "HAL/IConsoleManager.h"
 #include "Misc/Paths.h"
 #include "Misc/ConfigCacheIni.h"
 #include "Policies/CondensedJsonPrintPolicy.h"
@@ -10,8 +11,8 @@
 #include "Serialization/JsonReader.h"
 #include "Serialization/JsonSerializer.h"
 #include "SlateGlobals.h"
+#include "Trace/SlateMemoryTags.h"
 #include "Misc/RemoteConfigIni.h"
-#include "HAL/IConsoleManager.h"
 
 /* FUserDefinedChords helper class
  *****************************************************************************/
@@ -409,6 +410,8 @@ void FInputBindingManager::CreateInputCommand( const TSharedRef<FBindingContext>
 	// The command name should be valid
 	check( InCommandInfo->CommandName != NAME_None );
 
+	LLM_SCOPE_BYTAG(UI_Slate);
+
 	// Should not have already created a chord for this command
 	for (uint32 i = 0; i < static_cast<uint8>(EMultipleKeyBindingIndex::NumChords); ++i)
 	{
@@ -518,6 +521,8 @@ void FInputBindingManager::RemoveInputCommand(const TSharedRef<FBindingContext>&
 	// The command name should be valid
 	check(InUICommandInfo->CommandName != NAME_None);
 
+	LLM_SCOPE_BYTAG(UI_Slate);
+
 	const FName ContextName = InBindingContext->GetContextName();
 
 	FContextEntry& ContextEntry = ContextMap.FindOrAdd(ContextName);
diff --git a/Engine/Source/Runtime/Slate/Private/Framework/Commands/UICommandInfo.cpp b/Engine/Source/Runtime/Slate/Private/Framework/Commands/UICommandInfo.cpp
--- a/Engine/Source/Runtime/Slate/Private/Framework/Commands/UICommandInfo.cpp
+++ b/Engine/Source/Runtime/Slate/Private/Framework/Commands/UICommandInfo.cpp
@@ -1,11 +1,13 @@
 // Copyright Epic Games, Inc. All Rights Reserved.
 
 #include "Framework/Commands/UICommandInfo.h"
+
 #include "Styling/SlateColor.h"
 #include "Widgets/DeclarativeSyntaxSupport.h"
 #include "Widgets/SBoxPanel.h"
 #include "Styling/CoreStyle.h"
 #include "Framework/Commands/InputBindingManager.h"
+#include "Trace/SlateMemoryTags.h"
 #include "Widgets/Text/STextBlock.h"
 #include "Widgets/SToolTip.h"
 
@@ -15,6 +17,7 @@ FOnBindingContextChanged FBindingContext::CommandsChanged;
 
 FUICommandInfoDecl FBindingContext::NewCommand( const FName InCommandName, const FText& InCommandLabel, const FText& InCommandDesc )
 {
+	LLM_SCOPE_BYTAG(UI_Slate);
 	return FUICommandInfoDecl( this->AsShared(), InCommandName, InCommandLabel, InCommandDesc );
 }
 
@@ -53,6 +56,7 @@ const FText& FBindingContext::GetBundleLabel(const FName Name)
 FUICommandInfoDecl::FUICommandInfoDecl( const TSharedRef<FBindingContext>& InContext, const FName InCommandName, const FText& InLabel, const FText& InDesc, const FName InBundle)
 	: Context( InContext )
 {
+	LLM_SCOPE_BYTAG(UI_Slate);
 	Info = MakeShareable( new FUICommandInfo( InContext->GetContextName() ) );
 	Info->CommandName = InCommandName;
 	Info->Label = InLabel;
@@ -108,6 +112,8 @@ void FUICommandInfo::MakeCommandInfo( const TSharedRef<class FBindingContext>& I
 {
 	ensureMsgf( !InCommandLabel.IsEmpty(), TEXT("Command labels cannot be empty") );
 
+	LLM_SCOPE_BYTAG(UI_Slate);
+
 	OutCommand = MakeShareable( new FUICommandInfo( InContext->GetContextName() ) );
 	OutCommand->CommandName = InCommandName;
 	OutCommand->Label = InCommandLabel;
diff --git a/Engine/Source/Runtime/Slate/Private/Framework/Docking/TabManager.cpp b/Engine/Source/Runtime/Slate/Private/Framework/Docking/TabManager.cpp
--- a/Engine/Source/Runtime/Slate/Private/Framework/Docking/TabManager.cpp
+++ b/Engine/Source/Runtime/Slate/Private/Framework/Docking/TabManager.cpp
@@ -19,6 +19,7 @@
 #include "Framework/Docking/SDockingTabWell.h"
 #include "Framework/Docking/LayoutExtender.h"
 #include "Misc/NamePermissionList.h"
+#include "Trace/SlateMemoryTags.h"
 #include "HAL/PlatformApplicationMisc.h"
 #if PLATFORM_MAC
 #include "../MultiBox/Mac/MacMenu.h"
@@ -959,6 +960,8 @@ FTabSpawnerEntry& FTabManager::RegisterTabSpawner(const FName TabId, const FOnSp
 	ensure(!TabSpawner.Contains(TabId));
 	ensure(!FGlobalTabmanager::Get()->IsLegacyTabType(TabId));
 
+	LLM_SCOPE_BYTAG(UI_Slate);
+
 	TSharedRef<FTabSpawnerEntry> NewSpawnerEntry = MakeShareable(new FTabSpawnerEntry(TabId, OnSpawnTab, CanSpawnTab));
 	TabSpawner.Add(TabId, NewSpawnerEntry);
 	return NewSpawnerEntry.Get();
@@ -2478,6 +2481,9 @@ FTabSpawnerEntry& FGlobalTabmanager::RegisterNomadTabSpawner(const FName TabId,
 {
 	// Sanity check
 	ensure(!IsLegacyTabType(TabId));
+
+	LLM_SCOPE_BYTAG(UI_Slate);
+
 	// Remove TabId if it was previously loaded. This allows re-loading the Editor UI layout without restarting the whole Editor (Window->Load Layout)
 	if (NomadTabSpawner->Contains(TabId))
 	{
diff --git a/Engine/Source/Runtime/Slate/Private/Framework/Text/SlateTextLayout.cpp b/Engine/Source/Runtime/Slate/Private/Framework/Text/SlateTextLayout.cpp
--- a/Engine/Source/Runtime/Slate/Private/Framework/Text/SlateTextLayout.cpp
+++ b/Engine/Source/Runtime/Slate/Private/Framework/Text/SlateTextLayout.cpp
@@ -10,9 +10,11 @@
 #include "Framework/Text/ISlateLineHighlighter.h"
 #include "Framework/Text/SlateTextRun.h"
 #include "Framework/Text/SlatePasswordRun.h"
+#include "Trace/SlateMemoryTags.h"
 
 TSharedRef< FSlateTextLayout > FSlateTextLayout::Create(SWidget* InOwner, FTextBlockStyle InDefaultTextStyle)
 {
+	LLM_SCOPE_BYTAG(UI_Text);
 	TSharedRef< FSlateTextLayout > Layout = MakeShareable( new FSlateTextLayout(InOwner, MoveTemp(InDefaultTextStyle)) );
 	Layout->AggregateChildren();
 
diff --git a/Engine/Source/Runtime/Slate/Public/Framework/Commands/UICommandInfo.h b/Engine/Source/Runtime/Slate/Public/Framework/Commands/UICommandInfo.h
--- a/Engine/Source/Runtime/Slate/Public/Framework/Commands/UICommandInfo.h
+++ b/Engine/Source/Runtime/Slate/Public/Framework/Commands/UICommandInfo.h
@@ -7,6 +7,7 @@
 #include "InputCoreTypes.h"
 #include "Layout/Visibility.h"
 #include "Textures/SlateIcon.h"
+#include "Trace/SlateMemoryTags.h"
 #include "Framework/Commands/InputChord.h"
 #include "UICommandInfo.generated.h"
 
@@ -198,6 +199,8 @@ public:
 		, UserInterfaceType( EUserInterfaceActionType::Button )
 		, bUseLongDisplayName( true )
 	{
+		LLM_SCOPE_BYTAG(UI_Slate);
+
 		ActiveChords.Empty(2);
 		ActiveChords.Add(TSharedRef<FInputChord>(new FInputChord));
 		ActiveChords.Add(TSharedRef<FInputChord>(new FInputChord));
diff --git a/Engine/Source/Runtime/SlateCore/Private/Fonts/FontCacheCompositeFont.cpp b/Engine/Source/Runtime/SlateCore/Private/Fonts/FontCacheCompositeFont.cpp
--- a/Engine/Source/Runtime/SlateCore/Private/Fonts/FontCacheCompositeFont.cpp
+++ b/Engine/Source/Runtime/SlateCore/Private/Fonts/FontCacheCompositeFont.cpp
@@ -6,12 +6,12 @@
 #include "Fonts/FontBulkData.h"
 #include "Misc/FileHelper.h"
 #include "Algo/BinarySearch.h"
-#include "HAL/LowLevelMemTracker.h"
 #include "Internationalization/Culture.h"
 #include "Internationalization/Internationalization.h"
 #include "Application/SlateApplicationBase.h"
 #include "Async/Async.h"
 #include "HAL/PlatformProcess.h"
+#include "Trace/SlateMemoryTags.h"
 
 static bool GAsyncFontLazyLoad = false;
 FAutoConsoleVariableRef CVarAsyncLazyLoad(
@@ -386,7 +386,7 @@ const FFontData& FCompositeFontCache::GetFontDataForCodepoint(const FSlateFontIn
 
 TSharedPtr<FFreeTypeFace> FCompositeFontCache::GetFontFace(const FFontData& InFontData)
 {
-	LLM_SCOPE(ELLMTag::UI);
+	LLM_SCOPE_BYTAG(UI_Text);
 
 	TSharedPtr<FFreeTypeFace> FaceAndMemory = FontFaceMap.FindRef(InFontData);
 	if (!FaceAndMemory.IsValid() && InFontData.HasFont())
diff --git a/Engine/Source/Runtime/SlateCore/Private/Fonts/FontCacheFreeType.cpp b/Engine/Source/Runtime/SlateCore/Private/Fonts/FontCacheFreeType.cpp
--- a/Engine/Source/Runtime/SlateCore/Private/Fonts/FontCacheFreeType.cpp
+++ b/Engine/Source/Runtime/SlateCore/Private/Fonts/FontCacheFreeType.cpp
@@ -4,12 +4,12 @@
 #include "SlateGlobals.h"
 #include "HAL/PlatformFile.h"
 #include "HAL/PlatformFileManager.h"
-#include "HAL/LowLevelMemTracker.h"
 #include "HAL/IConsoleManager.h"
 #include "Misc/FileHelper.h"
 #include "Application/SlateApplicationBase.h"
 #include "Async/Async.h"
 #include "HAL/PlatformProcess.h"
+#include "Trace/SlateMemoryTags.h"
 
 #include <limits>
 
@@ -30,7 +30,7 @@ namespace FreeTypeMemory
 
 static void* Alloc(FT_Memory Memory, long Size)
 {
-	LLM_SCOPE(ELLMTag::UI);
+	LLM_SCOPE_BYTAG(UI_Text);
 	void* Result = FMemory::Malloc(Size);
 
 #if STATS
@@ -43,7 +43,7 @@ static void* Alloc(FT_Memory Memory, long Size)
 
 static void* Realloc(FT_Memory Memory, long CurSize, long NewSize, void* Block)
 {
-	LLM_SCOPE(ELLMTag::UI);
+	LLM_SCOPE_BYTAG(UI_Text);
 
 #if STATS
 	long DeltaNewSize = NewSize - CurSize;
diff --git a/Engine/Source/Runtime/SlateCore/Private/Fonts/FontCacheHarfBuzz.cpp b/Engine/Source/Runtime/SlateCore/Private/Fonts/FontCacheHarfBuzz.cpp
--- a/Engine/Source/Runtime/SlateCore/Private/Fonts/FontCacheHarfBuzz.cpp
+++ b/Engine/Source/Runtime/SlateCore/Private/Fonts/FontCacheHarfBuzz.cpp
@@ -4,7 +4,7 @@
 #include "Fonts/FontCache.h"
 #include "Fonts/FontCacheFreeType.h"
 #include "Fonts/SlateFontRenderer.h"
-#include "HAL/LowLevelMemTracker.h"
+#include "Trace/SlateMemoryTags.h"
 
 #if WITH_HARFBUZZ
 
@@ -13,13 +13,13 @@ extern "C"
 
 void* HarfBuzzMalloc(size_t InSizeBytes)
 {
-	LLM_SCOPE(ELLMTag::UI);
+	LLM_SCOPE_BYTAG(UI_Text);
 	return FMemory::Malloc(InSizeBytes);
 }
 
 void* HarfBuzzCalloc(size_t InNumItems, size_t InItemSizeBytes)
 {
-	LLM_SCOPE(ELLMTag::UI);
+	LLM_SCOPE_BYTAG(UI_Text);
 	const size_t AllocSizeBytes = InNumItems * InItemSizeBytes;
 	if (AllocSizeBytes > 0)
 	{
@@ -32,7 +32,7 @@ void* HarfBuzzCalloc(size_t InNumItems, size_t InItemSizeBytes)
 
 void* HarfBuzzRealloc(void* InPtr, size_t InSizeBytes)
 {
-	LLM_SCOPE(ELLMTag::UI);
+	LLM_SCOPE_BYTAG(UI_Text);
 	return FMemory::Realloc(InPtr, InSizeBytes);
 }
 
diff --git a/Engine/Source/Runtime/SlateCore/Private/Textures/TextureAtlas.cpp b/Engine/Source/Runtime/SlateCore/Private/Textures/TextureAtlas.cpp
--- a/Engine/Source/Runtime/SlateCore/Private/Textures/TextureAtlas.cpp
+++ b/Engine/Source/Runtime/SlateCore/Private/Textures/TextureAtlas.cpp
@@ -4,7 +4,7 @@
 #include "Stats/Stats.h"
 #include "Textures/SlateShaderResource.h"
 #include "Textures/SlateTextureData.h"
-#include "HAL/LowLevelMemTracker.h"
+#include "Trace/SlateMemoryTags.h"
 #include "Misc/MemStack.h"
 
 #include <limits>
@@ -113,7 +113,7 @@ void FSlateTextureAtlas::MarkTextureDirty()
 
 void FSlateTextureAtlas::InitAtlasData()
 {
-	LLM_SCOPE(ELLMTag::UI);
+	LLM_SCOPE_BYTAG(UI_Texture);
 
 	check(AtlasEmptySlots == NULL && AtlasData.Num() == 0);
 
diff --git a/Engine/Source/Runtime/SlateCore/Private/Trace/SlateMemoryTags.cpp b/Engine/Source/Runtime/SlateCore/Private/Trace/SlateMemoryTags.cpp
new file mode 100644
--- /dev/null
+++ b/Engine/Source/Runtime/SlateCore/Private/Trace/SlateMemoryTags.cpp
@@ -0,0 +1,19 @@
+// Copyright Epic Games, Inc. All Rights Reserved.
+
+#include "Trace/SlateMemoryTags.h"
+
+#if ENABLE_LOW_LEVEL_MEM_TRACKER
+
+LLM_DEFINE_TAG(UI_Slate, "UI Slate", TEXT("UI"), GET_STATFNAME(STAT_UI_SlateLLM), GET_STATFNAME(STAT_UISummaryLLM));
+LLM_DEFINE_TAG(UI_UMG, "UI UMG", TEXT("UI"), GET_STATFNAME(STAT_UI_UMGLLM), GET_STATFNAME(STAT_UISummaryLLM));
+LLM_DEFINE_TAG(UI_Text, "UI Text", TEXT("UI"), GET_STATFNAME(STAT_UI_TextLLM), GET_STATFNAME(STAT_UISummaryLLM));
+LLM_DEFINE_TAG(UI_Texture, "UI Texture", TEXT("UI"), GET_STATFNAME(STAT_UI_TextureLLM), GET_STATFNAME(STAT_UISummaryLLM));
+LLM_DEFINE_TAG(UI_Style, "UI Style", TEXT("UI"), GET_STATFNAME(STAT_UI_StyleLLM), GET_STATFNAME(STAT_UISummaryLLM));
+
+DEFINE_STAT(STAT_UI_SlateLLM);
+DEFINE_STAT(STAT_UI_UMGLLM);
+DEFINE_STAT(STAT_UI_TextLLM);
+DEFINE_STAT(STAT_UI_TextureLLM);
+DEFINE_STAT(STAT_UI_StyleLLM);
+
+#endif // ENABLE_LOW_LEVEL_MEM_TRACKER
\ No newline at end of file
diff --git a/Engine/Source/Runtime/SlateCore/Public/Styling/SlateStyle.h b/Engine/Source/Runtime/SlateCore/Public/Styling/SlateStyle.h
--- a/Engine/Source/Runtime/SlateCore/Public/Styling/SlateStyle.h
+++ b/Engine/Source/Runtime/SlateCore/Public/Styling/SlateStyle.h
@@ -15,6 +15,7 @@
 #include "Brushes/SlateRoundedBoxBrush.h"
 #include "Brushes/SlateImageBrush.h"
 #include "Brushes/SlateDynamicImageBrush.h"
+#include "Trace/SlateMemoryTags.h"
 
 class UTexture2D;
 struct FSlateWidgetStyle;
@@ -81,6 +82,7 @@ public:
 	template< typename DefinitionType >
 	FORCENOINLINE void Set(const FName PropertyName, const DefinitionType& InStyleDefintion)
 	{
+		LLM_SCOPE_BYTAG(UI_Style);
 		WidgetStyleValues.Add(PropertyName, MakeShared<DefinitionType>(InStyleDefintion));
 	}
 
@@ -91,6 +93,7 @@ public:
 	 */
 	FORCENOINLINE void Set(const FName PropertyName, const float InFloat)
 	{
+		LLM_SCOPE_BYTAG(UI_Style);
 		FloatValues.Add(PropertyName, InFloat);
 	}
 
@@ -101,6 +104,7 @@ public:
 	 */
 	FORCENOINLINE void Set(const FName PropertyName, const FVector2D InVector)
 	{
+		LLM_SCOPE_BYTAG(UI_Style);
 		Vector2DValues.Add(PropertyName, InVector);
 	}
 
@@ -111,11 +115,13 @@ public:
 	 */
 	FORCENOINLINE void Set(const FName PropertyName, const FLinearColor& InColor)
 	{
+		LLM_SCOPE_BYTAG(UI_Style);
 		ColorValues.Add(PropertyName, InColor);
 	}
 
 	FORCENOINLINE void Set(const FName PropertyName, const FColor& InColor)
 	{
+		LLM_SCOPE_BYTAG(UI_Style);
 		ColorValues.Add(PropertyName, InColor);
 	}
 
@@ -126,6 +132,7 @@ public:
 	 */
 	FORCENOINLINE void Set(const FName PropertyName, const FSlateColor& InColor)
 	{
+		LLM_SCOPE_BYTAG(UI_Style);
 		SlateColorValues.Add(PropertyName, InColor);
 	}
 
@@ -136,6 +143,7 @@ public:
 	 */
 	FORCENOINLINE void Set(const FName PropertyName, const FMargin& InMargin)
 	{
+		LLM_SCOPE_BYTAG(UI_Style);
 		MarginValues.Add(PropertyName, InMargin);
 	}
 
@@ -147,6 +155,9 @@ public:
 	template<typename BrushType>
 	FORCENOINLINE void Set(const FName PropertyName, BrushType* InBrush)
 	{
+		// @TODO: This scope may not work, the memory was allocated from the caller not in this scope
+		// Need some way to capure in parent scope?
+		LLM_SCOPE_BYTAG(UI_Style);
 		BrushResources.Add(PropertyName, InBrush);
 	}
 	/**
@@ -157,6 +168,7 @@ public:
 	 */
 	FORCENOINLINE void Set(FName PropertyName, const FSlateSound& InSound)
 	{
+		LLM_SCOPE_BYTAG(UI_Style);
 		Sounds.Add(PropertyName, InSound);
 	}
 
@@ -168,6 +180,7 @@ public:
 	 */
 	FORCENOINLINE void Set(FName PropertyName, const FSlateFontInfo& InFontInfo)
 	{
+		LLM_SCOPE_BYTAG(UI_Style);
 		FontInfoResources.Add(PropertyName, InFontInfo);
 	}
 
diff --git a/Engine/Source/Runtime/SlateCore/Public/Trace/SlateMemoryTags.h b/Engine/Source/Runtime/SlateCore/Public/Trace/SlateMemoryTags.h
new file mode 100644
--- /dev/null
+++ b/Engine/Source/Runtime/SlateCore/Public/Trace/SlateMemoryTags.h
@@ -0,0 +1,30 @@
+// Copyright Epic Games, Inc. All Rights Reserved.
+
+#pragma once
+
+#include "HAL/LowLevelMemTracker.h"
+#include "HAL/LowLevelMemStats.h"
+
+/**
+ * This file contains LLM tags for subcategories of the UI LLM tag. 
+ * 
+ * For more into on each macro, please see: LowLevelMemTracker
+ */
+
+#if ENABLE_LOW_LEVEL_MEM_TRACKER
+
+// Declare LLM tags, note that these can only be used by modules that include SlateCore
+LLM_DECLARE_TAG_API(UI_Slate, SLATECORE_API);
+LLM_DECLARE_TAG_API(UI_UMG, SLATECORE_API);
+LLM_DECLARE_TAG_API(UI_Text, SLATECORE_API);
+LLM_DECLARE_TAG_API(UI_Texture, SLATECORE_API);
+LLM_DECLARE_TAG_API(UI_Style, SLATECORE_API);
+
+// Declare stat groups per tag, definition and association with UI tag is done in .cpp
+DECLARE_LLM_MEMORY_STAT_EXTERN(TEXT("UI_Slate"), STAT_UI_SlateLLM, STATGROUP_LLMFULL, SLATECORE_API);
+DECLARE_LLM_MEMORY_STAT_EXTERN(TEXT("UI_UMG"), STAT_UI_UMGLLM, STATGROUP_LLMFULL, SLATECORE_API);
+DECLARE_LLM_MEMORY_STAT_EXTERN(TEXT("UI_Text"), STAT_UI_TextLLM, STATGROUP_LLMFULL, SLATECORE_API);
+DECLARE_LLM_MEMORY_STAT_EXTERN(TEXT("UI_Texture"), STAT_UI_TextureLLM, STATGROUP_LLMFULL, SLATECORE_API);
+DECLARE_LLM_MEMORY_STAT_EXTERN(TEXT("UI_Style"), STAT_UI_StyleLLM, STATGROUP_LLMFULL, SLATECORE_API);
+
+#endif // ENABLE_LOW_LEVEL_MEM_TRACKER
\ No newline at end of file
diff --git a/Engine/Source/Runtime/SlateCore/Public/Widgets/DeclarativeSyntaxSupport.h b/Engine/Source/Runtime/SlateCore/Public/Widgets/DeclarativeSyntaxSupport.h
--- a/Engine/Source/Runtime/SlateCore/Public/Widgets/DeclarativeSyntaxSupport.h
+++ b/Engine/Source/Runtime/SlateCore/Public/Widgets/DeclarativeSyntaxSupport.h
@@ -10,6 +10,7 @@
 #include "Rendering/SlateRenderTransform.h"
 #include "GenericPlatform/ICursor.h"
 #include "Types/ISlateMetaData.h"
+#include "Trace/SlateMemoryTags.h"
 #include "Widgets/SNullWidget.h"
 #include "Widgets/Accessibility/SlateWidgetAccessibleTypes.h"
 
@@ -947,5 +948,6 @@ struct TSlateDecl
 template<typename WidgetType, typename RequiredArgsPayloadType>
 TSlateDecl<WidgetType, RequiredArgsPayloadType> MakeTDecl( const ANSICHAR* InType, const ANSICHAR* InFile, int32 OnLine, RequiredArgsPayloadType&& InRequiredArgs )
 {
+	LLM_SCOPE_BYTAG(UI_Slate);
 	return TSlateDecl<WidgetType, RequiredArgsPayloadType>(InType, InFile, OnLine, Forward<RequiredArgsPayloadType>(InRequiredArgs));
 }
diff --git a/Engine/Source/Runtime/UMG/Private/Components/Widget.cpp b/Engine/Source/Runtime/UMG/Private/Components/Widget.cpp
--- a/Engine/Source/Runtime/UMG/Private/Components/Widget.cpp
+++ b/Engine/Source/Runtime/UMG/Private/Components/Widget.cpp
@@ -28,8 +28,8 @@
 #include "Blueprint/WidgetTree.h"
 #include "UMGStyle.h"
 #include "Types/ReflectionMetadata.h"
+#include "Trace/SlateMemoryTags.h"
 #include "Serialization/PropertyLocalizationDataGathering.h"
-#include "HAL/LowLevelMemTracker.h"
 #include "Components/NamedSlotInterface.h"
 
 #define LOCTEXT_NAMESPACE "UMG"
@@ -889,7 +889,7 @@ void UWidget::OnWidgetRebuilt()
 
 TSharedRef<SWidget> UWidget::TakeWidget()
 {
-	LLM_SCOPE(ELLMTag::UI);
+	LLM_SCOPE_BYTAG(UI_UMG);
 
 	return TakeWidget_Private( []( UUserWidget* Widget, TSharedRef<SWidget> Content ) -> TSharedPtr<SObjectWidget> {
 		       return SNew( SObjectWidget, Widget )[ Content ];
diff --git a/Engine/Source/Runtime/UMG/Private/UserWidget.cpp b/Engine/Source/Runtime/UMG/Private/UserWidget.cpp
--- a/Engine/Source/Runtime/UMG/Private/UserWidget.cpp
+++ b/Engine/Source/Runtime/UMG/Private/UserWidget.cpp
@@ -6,6 +6,7 @@
 #include "Sound/SoundBase.h"
 #include "Sound/SlateSound.h"
 #include "Framework/Application/SlateApplication.h"
+#include "Trace/SlateMemoryTags.h"
 #include "Widgets/Layout/SSpacer.h"
 #include "Widgets/Layout/SConstraintCanvas.h"
 #include "Components/NamedSlot.h"
@@ -1390,6 +1391,7 @@ void UUserWidget::NativeOnInitialized()
 
 void UUserWidget::NativePreConstruct()
 {
+	LLM_SCOPE_BYTAG(UI_UMG);
 	const bool bIsDesignTime = IsDesignTime();
 	if (UWidgetBlueprintGeneratedClass* BPClass = Cast<UWidgetBlueprintGeneratedClass>(GetClass()))
 	{
@@ -1404,6 +1406,7 @@ void UUserWidget::NativePreConstruct()
 
 void UUserWidget::NativeConstruct()
 {
+	LLM_SCOPE_BYTAG(UI_UMG);
 	if (UWidgetBlueprintGeneratedClass* BPClass = Cast<UWidgetBlueprintGeneratedClass>(GetClass()))
 	{
 		BPClass->ForEachExtension([this](UWidgetBlueprintGeneratedClassExtension* Extension)
@@ -2042,6 +2045,8 @@ UUserWidget* UUserWidget::CreateWidgetInstance(UWorld& World, TSubclassOf<UUserW
 
 UUserWidget* UUserWidget::CreateInstanceInternal(UObject* Outer, TSubclassOf<UUserWidget> UserWidgetClass, FName InstanceName, UWorld* World, ULocalPlayer* LocalPlayer)
 {
+	LLM_SCOPE_BYTAG(UI_UMG);
+
 	//CSV_SCOPED_TIMING_STAT(Slate, CreateWidget);
 
 #if !(UE_BUILD_SHIPPING || UE_BUILD_TEST)
