Add R9G9B9E5 format on platforms that support it

diff --git a/Engine/Source/Runtime/Apple/MetalRHI/Private/MetalRHI.cpp b/Engine/Source/Runtime/Apple/MetalRHI/Private/MetalRHI.cpp
--- a/Engine/Source/Runtime/Apple/MetalRHI/Private/MetalRHI.cpp
+++ b/Engine/Source/Runtime/Apple/MetalRHI/Private/MetalRHI.cpp
@@ -737,6 +737,8 @@ FMetalDynamicRHI::FMetalDynamicRHI(ERHIFeatureLevel::Type RequestedFeatureLevel)
 	GMetalBufferFormats[PF_R32G32B32F			] = { mtlpp::PixelFormat::Invalid, (uint8)EMetalBufferFormat::Unknown };
 	GMetalBufferFormats[PF_R8_SINT				] = { mtlpp::PixelFormat::R8Sint, (uint8)EMetalBufferFormat::R8Sint };
 	GMetalBufferFormats[PF_R64_UINT				] = { mtlpp::PixelFormat::Invalid, (uint8)EMetalBufferFormat::Unknown };
+	GMetalBufferFormats[PF_R9G9B9EXP5			] = { mtlpp::PixelFormat::Invalid, (uint8)EMetalBufferFormat::Unknown };
+	static_assert(PF_MAX == 86, "Please setup GMetalBufferFormats properly for the new pixel format");
 
 	// Initialize the platform pixel format map.
 	GPixelFormats[PF_Unknown			].PlatformFormat	= (uint32)mtlpp::PixelFormat::Invalid;
@@ -967,6 +969,8 @@ FMetalDynamicRHI::FMetalDynamicRHI(ERHIFeatureLevel::Type RequestedFeatureLevel)
 	GPixelFormats[PF_R8_SINT			].PlatformFormat	= (uint32)mtlpp::PixelFormat::R8Sint;
 	GPixelFormats[PF_R64_UINT			].PlatformFormat	= (uint32)mtlpp::PixelFormat::Invalid;
 	GPixelFormats[PF_R64_UINT			].Supported			= false;
+	GPixelFormats[PF_R9G9B9EXP5		    ].PlatformFormat    = (uint32)mtlpp::PixelFormat::Invalid;
+	GPixelFormats[PF_R9G9B9EXP5		    ].Supported			= false;
 
 #if METAL_DEBUG_OPTIONS
 	for (uint32 i = 0; i < PF_MAX; i++)
diff --git a/Engine/Source/Runtime/Core/Private/Misc/PixelFormat.cpp b/Engine/Source/Runtime/Core/Private/Misc/PixelFormat.cpp
--- a/Engine/Source/Runtime/Core/Private/Misc/PixelFormat.cpp
+++ b/Engine/Source/Runtime/Core/Private/Misc/PixelFormat.cpp
@@ -132,6 +132,7 @@ FPixelFormatInfo    GPixelFormats[PF_MAX] =
 	FPixelFormatInfo(PF_R32G32B32F,         TEXT("R32G32B32F"),            1,          1,          1,          12,         3,              1),
 	FPixelFormatInfo(PF_R8_SINT,            TEXT("R8_SINT"),               1,          1,          1,          1,          1,              1),
 	FPixelFormatInfo(PF_R64_UINT,			TEXT("R64_UINT"),              1,          1,          1,          8,          1,              0),
+	FPixelFormatInfo(PF_R9G9B9EXP5,			TEXT("R9G9B9EXP5"),			   1,		   1,		   1,		   4,		   4,			   0),
 };
 
 
diff --git a/Engine/Source/Runtime/Core/Public/Math/PackedVector.h b/Engine/Source/Runtime/Core/Public/Math/PackedVector.h
--- a/Engine/Source/Runtime/Core/Public/Math/PackedVector.h
+++ b/Engine/Source/Runtime/Core/Public/Math/PackedVector.h
@@ -6,6 +6,7 @@
 #include "Math/NumericLimits.h"
 #include "Math/UnrealMathUtility.h"
 #include "Math/Color.h"
+#include "Templates/TypeCompatibleBytes.h"
 
 /** 
  * 3 component vector corresponding to DXGI_FORMAT_R11G11B10_FLOAT. 
@@ -304,4 +305,121 @@ inline FLinearColor FFixedRGBASigned8::ToLinearColor() const
 {
 	const float Scale = 1.0f / MAX_int8;
 	return FLinearColor(R * Scale, G * Scale, B * Scale, A * Scale);
+}
+
+/**
+ * 3 component vector corresponding to PF_R9G9B9EXP5.
+ */
+class FFloat3PackedSE
+{
+public:
+
+	union
+	{
+		struct
+		{
+			uint32 RMantissa : 9;
+			uint32 GMantissa : 9;
+			uint32 BMantissa : 9;
+			uint32 SharedExponent : 5;
+		};
+		uint32 EncodedValue;
+	};
+
+	FFloat3PackedSE() {}
+
+	explicit FFloat3PackedSE(const FLinearColor& Src);
+	explicit FFloat3PackedSE(uint32 InEncodedValue) : EncodedValue(InEncodedValue) {}
+
+	FLinearColor ToLinearColor() const;
+};
+
+inline FFloat3PackedSE::FFloat3PackedSE(const FLinearColor& Src)
+{
+
+	//s=sign, e = exponent, m = mantissa
+	//32 bit floating point:
+	// s:e:m 1:8:23   
+	// if e > 0					float = (-1)^sign X 2^(e-127) X (1.0 + m)
+	// if e == 0	&& m!= 0		float = (-1)^sign X 2^(-126) X (0.0 + m)   (subnormal)
+	// if e == 0xff && m==0        float = (-1)^sign X infinity					
+	// if e == 0xff && m!=0         NaN											
+	//
+	struct PackedFloat
+	{
+		union
+		{
+			struct
+			{
+				uint32_t m : 23;
+				uint32_t e : 8;
+				uint32_t s : 1;
+			};
+			float v;
+		};
+		bool IsSubnormal() { return (e == 0 && m != 0 ); }
+		uint32 Mantissa() { return m + (IsSubnormal() ? 0 : (1 << 23) ); }
+		int32 Exponent() { return e - 127; }
+		bool IsInfinity() { return (e == 0xff) && (m == 0); }
+		bool IsNaN() { return (e == 0xff) && (m != 0); }
+	};
+
+	PackedFloat RGBPacked[3];
+	RGBPacked[0] = BitCast<PackedFloat, float>(Src.R);
+	RGBPacked[1] = BitCast<PackedFloat, float>(Src.G);
+	RGBPacked[2] = BitCast<PackedFloat, float>(Src.B);
+
+	uint32 Mantissas[3];
+	int32 Exponents[3];
+
+	for (uint32 i = 0; i < 3; ++i)
+	{
+		Mantissas[i] = RGBPacked[i].Mantissa();
+		Exponents[i] = RGBPacked[i].Exponent();
+
+		// 9995 uses (0.0 + m) instead of floating point's (1.0 + m) , so let's fix that here.
+		if (!RGBPacked[i].IsSubnormal()) //don't need to fix subnormal (it is already 0.m based)
+		{
+			// example 2^2 X 1.5 ---> 2^3 X 0.75.  Exponent++, and mantissa/=2
+			Exponents[i]++;
+			Mantissas[i] /= 2;
+		}
+		if (  (Exponents[i] < -15 && RGBPacked[i].v != 0.0f) || (RGBPacked[i].v < 0.0f) || RGBPacked[i].IsNaN()) //underflow or negative or NaN
+		{
+			//As per DirectX implementation, underflow or negative or NaN clamp to a lowest possible exponent and mantissa
+			Exponents[i] = -15;
+			Mantissas[i] = 0;
+		}
+		else if (Exponents[i] > 15)  //overflow or infinity
+		{
+			Exponents[i] = 16; //Match behavior of XMStoreFloat3SE
+			Mantissas[i] = 0x7fffff; //as close to 1 as possible
+		}
+	}
+
+	//exponent now guaranteed to be from -15 to +15
+	int32 NewExponent = FMath::Max<int32>(Exponents[0], FMath::Max<int32>(Exponents[1], Exponents[2]));
+
+	for (uint32 i = 0; i < 3; ++i)
+	{
+		Mantissas[i] = (Mantissas[i] >> ((uint32)(NewExponent - Exponents[i]))); // this can conceivably go to zero if the diff between max and min Exponents > 8 !
+		//Mantissa is still 23 bits Need to make it 9 bit
+		Mantissas[i]	= Mantissas[i] >> (23 - 9); // now 9 bit
+	}
+	
+	check( Mantissas[0] < (1 << 23) && Mantissas[1] < (1 << 23) && Mantissas[2] < (1 << 23) );
+
+	EncodedValue = (((NewExponent + 15) & 0x01f) << (9 + 9 + 9)) |
+				    ((Mantissas[2]      & 0x1ff) << (0 + 9 + 9)) |
+				    ((Mantissas[1]      & 0x1ff) << (0 + 0 + 9)) |
+				    ((Mantissas[0]      & 0x1ff) << (0 + 0 + 0)) ;
+
+}
+
+inline FLinearColor FFloat3PackedSE::ToLinearColor() const
+{
+    int32 SharedExponent8Bits = 0x33800000 + (SharedExponent << 23);
+    float Scale = BitCast<float, int32>(SharedExponent8Bits);
+
+	return FLinearColor( Scale * float(RMantissa), Scale * float(GMantissa), Scale * float(BMantissa),1.0f );
 }
\ No newline at end of file
diff --git a/Engine/Source/Runtime/Core/Public/PixelFormat.h b/Engine/Source/Runtime/Core/Public/PixelFormat.h
--- a/Engine/Source/Runtime/Core/Public/PixelFormat.h
+++ b/Engine/Source/Runtime/Core/Public/PixelFormat.h
@@ -99,7 +99,8 @@ enum EPixelFormat : uint8
 	PF_R32G32B32F			=82,
 	PF_R8_SINT				=83,	
 	PF_R64_UINT				=84,
-	PF_MAX					=85,
+	PF_R9G9B9EXP5			=85,
+	PF_MAX					=86,
 };
 #define FOREACH_ENUM_EPIXELFORMAT(op) \
 	op(PF_Unknown) \
@@ -186,8 +187,8 @@ enum EPixelFormat : uint8
 	op(PF_R32G32B32_SINT) \
 	op(PF_R32G32B32F) \
 	op(PF_R8_SINT) \
-	op(PF_R64_UINT)
-
+	op(PF_R64_UINT) \
+	op(PF_R9G9B9EXP5)
 // Defines which channel is valid for each pixel format
 enum class EPixelFormatChannelFlags : uint8
 {
diff --git a/Engine/Source/Runtime/CoreUObject/Public/UObject/NoExportTypes.h b/Engine/Source/Runtime/CoreUObject/Public/UObject/NoExportTypes.h
--- a/Engine/Source/Runtime/CoreUObject/Public/UObject/NoExportTypes.h
+++ b/Engine/Source/Runtime/CoreUObject/Public/UObject/NoExportTypes.h
@@ -286,6 +286,7 @@ enum EPixelFormat
 	PF_R32G32B32F,
 	PF_R8_SINT,
 	PF_R64_UINT,
+	PF_R9G9B9EXP5,
 	PF_MAX,
 };
 
diff --git a/Engine/Source/Runtime/D3D12RHI/Private/D3D12RHI.cpp b/Engine/Source/Runtime/D3D12RHI/Private/D3D12RHI.cpp
--- a/Engine/Source/Runtime/D3D12RHI/Private/D3D12RHI.cpp
+++ b/Engine/Source/Runtime/D3D12RHI/Private/D3D12RHI.cpp
@@ -187,6 +187,8 @@ FD3D12DynamicRHI::FD3D12DynamicRHI(const TArray<TSharedPtr<FD3D12Adapter>>& Chos
 	GPixelFormats[PF_R32G32B32F		].PlatformFormat = DXGI_FORMAT_R32G32B32_FLOAT;
 	GPixelFormats[PF_R8_SINT		].PlatformFormat = DXGI_FORMAT_R8_SINT;
 
+	GPixelFormats[PF_R9G9B9EXP5	    ].PlatformFormat = DXGI_FORMAT_R9G9B9E5_SHAREDEXP;
+
 	// MS - Not doing any feature level checks. D3D12 currently supports these limits.
 	// However this may need to be revisited if new feature levels are introduced with different HW requirement
 	GSupportsSeparateRenderTargetBlendState = true;
diff --git a/Engine/Source/Runtime/D3D12RHI/Private/D3D12RenderTarget.cpp b/Engine/Source/Runtime/D3D12RHI/Private/D3D12RenderTarget.cpp
--- a/Engine/Source/Runtime/D3D12RHI/Private/D3D12RenderTarget.cpp
+++ b/Engine/Source/Runtime/D3D12RHI/Private/D3D12RenderTarget.cpp
@@ -65,6 +65,7 @@ static uint32 ComputeBytesPerPixel(DXGI_FORMAT Format)
 	case DXGI_FORMAT_R16G16_SINT:
 	case DXGI_FORMAT_D32_FLOAT:
 	case DXGI_FORMAT_R32_SINT:
+	case DXGI_FORMAT_R9G9B9E5_SHAREDEXP:
 		BytesPerPixel = 4;
 		break;
 	case DXGI_FORMAT_R16G16B16A16_FLOAT:
@@ -856,6 +857,9 @@ static void ConvertDXGIToFColor(DXGI_FORMAT Format, uint32 Width, uint32 Height,
 	case DXGI_FORMAT_R11G11B10_FLOAT:
 		ConvertRawR11G11B10DataToFColor(Width, Height, In, SrcPitch, Out, bLinearToGamma);
 		break;
+	case DXGI_FORMAT_R9G9B9E5_SHAREDEXP:
+		ConvertRawR9G9B9E5DataToFColor(Width, Height, In, SrcPitch, Out, bLinearToGamma);
+		break;
 	case DXGI_FORMAT_R32G32B32A32_FLOAT:
 		ConvertRawR32G32B32A32DataToFColor(Width, Height, In, SrcPitch, Out, bLinearToGamma);
 		break;
diff --git a/Engine/Source/Runtime/RHI/Public/RHISurfaceDataConversion.h b/Engine/Source/Runtime/RHI/Public/RHISurfaceDataConversion.h
--- a/Engine/Source/Runtime/RHI/Public/RHISurfaceDataConversion.h
+++ b/Engine/Source/Runtime/RHI/Public/RHISurfaceDataConversion.h
@@ -219,6 +219,26 @@ static inline void ConvertRawR11G11B10DataToFColor(uint32 Width, uint32 Height,
 	}
 }
 
+static inline void ConvertRawR9G9B9E5DataToFColor(uint32 Width, uint32 Height, uint8* In, uint32 SrcPitch, FColor* Out, bool LinearToGamma)
+{
+	check(sizeof(FFloat3PackedSE) == sizeof(uint32));
+
+	for (uint32 Y = 0; Y < Height; Y++)
+	{
+		FFloat3PackedSE* SrcPtr = (FFloat3PackedSE*)(In + Y * SrcPitch);
+		FColor* DestPtr = Out + Y * Width;
+
+		for (uint32 X = 0; X < Width; X++)
+		{
+			FLinearColor Value = (*SrcPtr).ToLinearColor();
+
+			*DestPtr = Value.ToFColor(LinearToGamma);
+			++SrcPtr;
+			++DestPtr;
+		}
+	}
+}
+
 static inline void ConvertRawR32G32B32A32DataToFColor(uint32 Width, uint32 Height, uint8 *In, uint32 SrcPitch, FColor* Out, bool LinearToGamma)
 {
 	FPlane MinValue(0.0f, 0.0f, 0.0f, 0.0f);
diff --git a/Engine/Source/Runtime/RenderCore/Private/RenderUtils.cpp b/Engine/Source/Runtime/RenderCore/Private/RenderUtils.cpp
--- a/Engine/Source/Runtime/RenderCore/Private/RenderUtils.cpp
+++ b/Engine/Source/Runtime/RenderCore/Private/RenderUtils.cpp
@@ -910,6 +910,7 @@ EPixelFormatChannelFlags GetPixelFormatValidChannels(EPixelFormat InPixelFormat)
 		EPixelFormatChannelFlags::RGB,		// PF_R32G32B32F
 		EPixelFormatChannelFlags::R,		// PF_R8_SINT
 		EPixelFormatChannelFlags::R,		// PF_R64_UINT
+		EPixelFormatChannelFlags::RGB,		// PF_R9G9B9EXP5
 	};
 	static_assert(UE_ARRAY_COUNT(PixelFormatToChannelFlags) == (uint8)PF_MAX, "Missing pixel format");
 	return (InPixelFormat < PF_MAX) ? PixelFormatToChannelFlags[(uint8)InPixelFormat] : EPixelFormatChannelFlags::None;
diff --git a/Engine/Source/Runtime/Renderer/Private/PostProcess/PostProcessTonemap.cpp b/Engine/Source/Runtime/Renderer/Private/PostProcess/PostProcessTonemap.cpp
--- a/Engine/Source/Runtime/Renderer/Private/PostProcess/PostProcessTonemap.cpp
+++ b/Engine/Source/Runtime/Renderer/Private/PostProcess/PostProcessTonemap.cpp
@@ -871,7 +871,8 @@ FScreenPassTexture AddTonemapPass(FRDGBuilder& GraphBuilder, const FViewInfo& Vi
 		// If this is a stereo view, there's a good chance we need alpha out of the tonemapper
 		// @todo: Remove this once Oculus fix the bug in their runtime that requires alpha here.
 		const bool bIsStereo = IStereoRendering::IsStereoEyeView(View);
-		FRHIBlendState* BlendState = Inputs.bWriteAlphaChannel || bIsStereo ? FScreenPassPipelineState::FDefaultBlendState::GetRHI() : TStaticBlendStateWriteMask<CW_RGB>::GetRHI();
+		const bool bFormatNeedsAlphaWrite = Output.Texture->Desc.Format == PF_R9G9B9EXP5;
+		FRHIBlendState* BlendState = (Inputs.bWriteAlphaChannel || bIsStereo || bFormatNeedsAlphaWrite) ? FScreenPassPipelineState::FDefaultBlendState::GetRHI() : TStaticBlendStateWriteMask<CW_RGB>::GetRHI();
 		FRHIDepthStencilState* DepthStencilState = FScreenPassPipelineState::FDefaultDepthStencilState::GetRHI();
 
 		EScreenPassDrawFlags DrawFlags = EScreenPassDrawFlags::AllowHMDHiddenAreaMask;
diff --git a/Engine/Source/Runtime/Renderer/Private/SystemTextures.cpp b/Engine/Source/Runtime/Renderer/Private/SystemTextures.cpp
--- a/Engine/Source/Runtime/Renderer/Private/SystemTextures.cpp
+++ b/Engine/Source/Runtime/Renderer/Private/SystemTextures.cpp
@@ -11,6 +11,7 @@
 #include "RenderTargetPool.h"
 #include "ClearQuad.h"
 #include "LTC.h"
+#include "Math/PackedVector.h"
 
 /*-----------------------------------------------------------------------------
 SystemTextures
@@ -1169,11 +1170,20 @@ void FormatData5551(const TInType& In, uint8* Out, uint32& OutByteCount)
 	OutByteCount = 2;
 }
 
+template<typename TInType>
+void FormatData9995(const TInType& In, uint8* Out, uint32& OutByteCount)
+{
+	FFloat3PackedSE PackedFloat(FLinearColor(In[0], In[1], In[2], 1.0f));
+	uint32* OutTyped = (uint32*)Out;
+	OutByteCount = 4;
+	*OutTyped = PackedFloat.EncodedValue;
+}
+
 template<typename TInType>
 void InitializeData(const TInType& InData, EPixelFormat InFormat, uint8* OutData, uint32& OutByteCount)
 {
 	// If a new format is added insure that it is either supported here, or at least flagged as not supported
-	static_assert(PF_MAX == 85);
+	static_assert(PF_MAX == 86);
 
 	switch (InFormat)
 	{
@@ -1236,6 +1246,7 @@ void InitializeData(const TInType& InData, EPixelFormat InFormat, uint8* OutData
 		case PF_A2B10G10R10:			{ FormatData1010102<TInType>(InData, OutData, OutByteCount); } break;
 		case PF_FloatR11G11B10:			{ FormatData111110<TInType>	(InData, OutData, OutByteCount); } break;
 		case PF_B5G5R5A1_UNORM:         { FormatData5551<TInType>(InData, OutData, OutByteCount); } break;
+		case PF_R9G9B9EXP5:				{ FormatData9995<TInType>(InData, OutData, OutByteCount); }	break;
 			return;
 
 		// Not supported
