[Insights] Removed log when trace is not enabled.

diff --git a/Engine/Source/Runtime/Core/Private/ProfilingDebugging/TraceAuxiliary.cpp b/Engine/Source/Runtime/Core/Private/ProfilingDebugging/TraceAuxiliary.cpp
--- a/Engine/Source/Runtime/Core/Private/ProfilingDebugging/TraceAuxiliary.cpp
+++ b/Engine/Source/Runtime/Core/Private/ProfilingDebugging/TraceAuxiliary.cpp
@@ -1455,11 +1455,6 @@ void FTraceAuxiliary::Initialize(const TCHAR* CommandLine)
 	UE::Trace::ThreadRegister(TEXT("GameThread"), FPlatformTLS::GetCurrentThreadId(), -1);
 
 	UE_LOG(LogCore, Log, TEXT("Finished trace initialization."));
-
-#else //UE_TRACE_ENABLED
-
-	UE_LOG(LogCore, Log, TEXT("Trace is not enabled (see UE_TRACE_ENABLED)."));
-
 #endif //UE_TRACE_ENABLED
 }
 
