improved delegates extensibility

diff --git a/Engine/Source/Runtime/Core/Public/Delegates/DelegateInstanceInterface.h b/Engine/Source/Runtime/Core/Public/Delegates/DelegateInstanceInterface.h
--- a/Engine/Source/Runtime/Core/Public/Delegates/DelegateInstanceInterface.h
+++ b/Engine/Source/Runtime/Core/Public/Delegates/DelegateInstanceInterface.h
@@ -16,9 +16,9 @@ template <typename RetType, typename... ArgTypes, typename UserPolicy>
 struct IBaseDelegateInstance<RetType(ArgTypes...), UserPolicy> : public UserPolicy::FDelegateInstanceExtras
 {
 	/**
-	 * Emplaces a copy of the delegate instance into the FDelegateBase.
+	 * Emplaces a copy of the delegate instance into the given base delegate.
 	 */
-	virtual void CreateCopy(FDelegateBase& Base) = 0;
+	virtual void CreateCopy(typename UserPolicy::FDelegateExtras& Base) = 0;
 
 	/**
 	 * Execute the delegate.  If the function pointer is not valid, an error will occur.
diff --git a/Engine/Source/Runtime/Core/Public/Delegates/DelegateInstancesImpl.h b/Engine/Source/Runtime/Core/Public/Delegates/DelegateInstancesImpl.h
--- a/Engine/Source/Runtime/Core/Public/Delegates/DelegateInstancesImpl.h
+++ b/Engine/Source/Runtime/Core/Public/Delegates/DelegateInstancesImpl.h
@@ -73,6 +73,7 @@ private:
 	using Super             = TCommonDelegateInstanceState<WrappedRetValType(ParamTypes...), UserPolicy, VarTypes...>;
 	using RetValType        = typename Super::RetValType;
 	using UnwrappedThisType = TBaseUFunctionDelegateInstance<UserClass, RetValType(ParamTypes...), UserPolicy, VarTypes...>;
+	using DelegateBaseType = typename UserPolicy::FDelegateExtras;
 
 	static_assert(UE4Delegates_Private::IsUObjectPtr((UserClass*)nullptr), "You cannot use UFunction delegates with non UObject classes.");
 
@@ -136,7 +137,7 @@ public:
 
 	// IBaseDelegateInstance interface
 
-	void CreateCopy(FDelegateBase& Base) final
+	void CreateCopy(DelegateBaseType& Base) final
 	{
 		new (Base) UnwrappedThisType(*(UnwrappedThisType*)this);
 	}
@@ -177,7 +178,7 @@ public:
 	 * @param InFunctionName The name of the function call.
 	 * @return The new delegate.
 	 */
-	FORCEINLINE static void Create(FDelegateBase& Base, UserClass* InUserObject, const FName& InFunctionName, VarTypes... Vars)
+	FORCEINLINE static void Create(DelegateBaseType& Base, UserClass* InUserObject, const FName& InFunctionName, VarTypes... Vars)
 	{
 		new (Base) UnwrappedThisType(InUserObject, InFunctionName, Vars...);
 	}
@@ -211,6 +212,7 @@ private:
 	using Super             = TCommonDelegateInstanceState<WrappedRetValType(ParamTypes...), UserPolicy, VarTypes...>;
 	using RetValType        = typename Super::RetValType;
 	using UnwrappedThisType = TBaseSPMethodDelegateInstance<bConst, UserClass, SPMode, RetValType(ParamTypes...), UserPolicy, VarTypes...>;
+	using DelegateBaseType = typename UserPolicy::FDelegateExtras;
 
 public:
 	using FMethodPtr = typename TMemFunPtrType<bConst, UserClass, RetValType(ParamTypes..., VarTypes...)>::Type;
@@ -270,7 +272,7 @@ public:
 
 	// IBaseDelegateInstance interface
 
-	void CreateCopy(FDelegateBase& Base) final
+	void CreateCopy(DelegateBaseType& Base) final
 	{
 		new (Base) UnwrappedThisType(*(UnwrappedThisType*)this);
 	}
@@ -324,7 +326,7 @@ public:
 	 * @param InFunc Member function pointer to your class method.
 	 * @return The new delegate.
 	 */
-	FORCEINLINE static void Create(FDelegateBase& Base, const TSharedPtr<UserClass, SPMode>& InUserObjectRef, FMethodPtr InFunc, VarTypes... Vars)
+	FORCEINLINE static void Create(DelegateBaseType& Base, const TSharedPtr<UserClass, SPMode>& InUserObjectRef, FMethodPtr InFunc, VarTypes... Vars)
 	{
 		new (Base) UnwrappedThisType(InUserObjectRef, InFunc, Vars...);
 	}
@@ -338,7 +340,7 @@ public:
 	 * @param InFunc  Member function pointer to your class method.
 	 * @return The new delegate.
 	 */
-	FORCEINLINE static void Create(FDelegateBase& Base, UserClass* InUserObject, FMethodPtr InFunc, VarTypes... Vars)
+	FORCEINLINE static void Create(DelegateBaseType& Base, UserClass* InUserObject, FMethodPtr InFunc, VarTypes... Vars)
 	{
 		// We expect the incoming InUserObject to derived from TSharedFromThis.
 		TSharedRef<UserClass, SPMode> UserObjectRef = StaticCastSharedRef<UserClass>(InUserObject->AsShared());
@@ -370,6 +372,7 @@ private:
 	using Super             = TCommonDelegateInstanceState<WrappedRetValType(ParamTypes...), UserPolicy, VarTypes...>;
 	using RetValType        = typename Super::RetValType;
 	using UnwrappedThisType = TBaseRawMethodDelegateInstance<bConst, UserClass, RetValType(ParamTypes...), UserPolicy, VarTypes...>;
+	using DelegateBaseType = typename UserPolicy::FDelegateExtras;
 
 public:
 	using FMethodPtr = typename TMemFunPtrType<bConst, UserClass, RetValType(ParamTypes..., VarTypes...)>::Type;
@@ -436,7 +439,7 @@ public:
 
 	// IBaseDelegateInstance interface
 
-	void CreateCopy(FDelegateBase& Base) final
+	void CreateCopy(DelegateBaseType& Base) final
 	{
 		new (Base) UnwrappedThisType(*(UnwrappedThisType*)this);
 	}
@@ -481,7 +484,7 @@ public:
 	 * @param InFunc Member function pointer to your class method.
 	 * @return The new delegate.
 	 */
-	FORCEINLINE static void Create(FDelegateBase& Base, UserClass* InUserObject, FMethodPtr InFunc, VarTypes... Vars)
+	FORCEINLINE static void Create(DelegateBaseType& Base, UserClass* InUserObject, FMethodPtr InFunc, VarTypes... Vars)
 	{
 		new (Base) UnwrappedThisType(InUserObject, InFunc, Vars...);
 	}
@@ -508,6 +511,7 @@ private:
 	using Super             = TCommonDelegateInstanceState<WrappedRetValType(ParamTypes...), UserPolicy, VarTypes...>;
 	using RetValType        = typename Super::RetValType;
 	using UnwrappedThisType = TBaseUObjectMethodDelegateInstance<bConst, UserClass, RetValType(ParamTypes...), UserPolicy, VarTypes...>;
+	using DelegateBaseType = typename UserPolicy::FDelegateExtras;
 
 	static_assert(UE4Delegates_Private::IsUObjectPtr((UserClass*)nullptr), "You cannot use UObject method delegates with raw pointers.");
 
@@ -574,7 +578,7 @@ public:
 
 	// IBaseDelegateInstance interface
 
-	void CreateCopy(FDelegateBase& Base) final
+	void CreateCopy(DelegateBaseType& Base) final
 	{
 		new (Base) UnwrappedThisType(*(UnwrappedThisType*)this);
 	}
@@ -625,7 +629,7 @@ public:
 	 * @param InFunc Member function pointer to your class method.
 	 * @return The new delegate.
 	 */
-	FORCEINLINE static void Create(FDelegateBase& Base, UserClass* InUserObject, FMethodPtr InFunc, VarTypes... Vars)
+	FORCEINLINE static void Create(DelegateBaseType& Base, UserClass* InUserObject, FMethodPtr InFunc, VarTypes... Vars)
 	{
 		new (Base) UnwrappedThisType(InUserObject, InFunc, Vars...);
 	}
@@ -653,6 +657,7 @@ private:
 	using Super             = TCommonDelegateInstanceState<WrappedRetValType(ParamTypes...), UserPolicy, VarTypes...>;
 	using RetValType        = typename Super::RetValType;
 	using UnwrappedThisType = TBaseStaticDelegateInstance<RetValType(ParamTypes...), UserPolicy, VarTypes...>;
+	using DelegateBaseType = typename UserPolicy::FDelegateExtras;
 
 public:
 	using FFuncPtr = RetValType(*)(ParamTypes..., VarTypes...);
@@ -711,7 +716,7 @@ public:
 
 	// IBaseDelegateInstance interface
 
-	void CreateCopy(FDelegateBase& Base) final
+	void CreateCopy(DelegateBaseType& Base) final
 	{
 		new (Base) UnwrappedThisType(*(UnwrappedThisType*)this);
 	}
@@ -742,7 +747,7 @@ public:
 	 * @param InFunc Static function pointer.
 	 * @return The new delegate.
 	 */
-	FORCEINLINE static void Create(FDelegateBase& Base, FFuncPtr InFunc, VarTypes... Vars)
+	FORCEINLINE static void Create(DelegateBaseType& Base, FFuncPtr InFunc, VarTypes... Vars)
 	{
 		new (Base) UnwrappedThisType(InFunc, Vars...);
 	}
@@ -768,6 +773,7 @@ private:
 	using Super             = TCommonDelegateInstanceState<WrappedRetValType(ParamTypes...), UserPolicy, VarTypes...>;
 	using RetValType        = typename Super::RetValType;
 	using UnwrappedThisType = TBaseFunctorDelegateInstance<RetValType(ParamTypes...), UserPolicy, FunctorType, VarTypes...>;
+	using DelegateBaseType = typename UserPolicy::FDelegateExtras;
 
 public:
 	TBaseFunctorDelegateInstance(const FunctorType& InFunctor, VarTypes... Vars)
@@ -823,7 +829,7 @@ public:
 
 public:
 	// IBaseDelegateInstance interface
-	void CreateCopy(FDelegateBase& Base) final
+	void CreateCopy(DelegateBaseType& Base) final
 	{
 		new (Base) UnwrappedThisType(*(UnwrappedThisType*)this);
 	}
@@ -848,11 +854,11 @@ public:
 	 * @param InFunctor C++ functor
 	 * @return The new delegate.
 	 */
-	FORCEINLINE static void Create(FDelegateBase& Base, const FunctorType& InFunctor, VarTypes... Vars)
+	FORCEINLINE static void Create(DelegateBaseType& Base, const FunctorType& InFunctor, VarTypes... Vars)
 	{
 		new (Base) UnwrappedThisType(InFunctor, Vars...);
 	}
-	FORCEINLINE static void Create(FDelegateBase& Base, FunctorType&& InFunctor, VarTypes... Vars)
+	FORCEINLINE static void Create(DelegateBaseType& Base, FunctorType&& InFunctor, VarTypes... Vars)
 	{
 		new (Base) UnwrappedThisType(MoveTemp(InFunctor), Vars...);
 	}
@@ -880,6 +886,7 @@ private:
 	using Super             = TCommonDelegateInstanceState<WrappedRetValType(ParamTypes...), UserPolicy, VarTypes...>;
 	using RetValType        = typename Super::RetValType;
 	using UnwrappedThisType = TWeakBaseFunctorDelegateInstance<UserClass, RetValType(ParamTypes...), UserPolicy, FunctorType, VarTypes...>;
+	using DelegateBaseType = typename UserPolicy::FDelegateExtras;
 
 public:
 	TWeakBaseFunctorDelegateInstance(UserClass* InContextObject, const FunctorType& InFunctor, VarTypes... Vars)
@@ -940,7 +947,7 @@ public:
 
 public:
 	// IBaseDelegateInstance interface
-	void CreateCopy(FDelegateBase& Base) final
+	void CreateCopy(DelegateBaseType& Base) final
 	{
 		new (Base) UnwrappedThisType(*(UnwrappedThisType*)this);
 	}
@@ -968,11 +975,11 @@ public:
 	 * @param InFunctor C++ functor
 	 * @return The new delegate.
 	 */
-	FORCEINLINE static void Create(FDelegateBase& Base, UserClass* InContextObject, const FunctorType& InFunctor, VarTypes... Vars)
+	FORCEINLINE static void Create(DelegateBaseType& Base, UserClass* InContextObject, const FunctorType& InFunctor, VarTypes... Vars)
 	{
 		new (Base) UnwrappedThisType(InContextObject, InFunctor, Vars...);
 	}
-	FORCEINLINE static void Create(FDelegateBase& Base, UserClass* InContextObject, FunctorType&& InFunctor, VarTypes... Vars)
+	FORCEINLINE static void Create(DelegateBaseType& Base, UserClass* InContextObject, FunctorType&& InFunctor, VarTypes... Vars)
 	{
 		new (Base) UnwrappedThisType(InContextObject, MoveTemp(InFunctor), Vars...);
 	}
diff --git a/Engine/Source/Runtime/Core/Public/Delegates/DelegateSignatureImpl.inl b/Engine/Source/Runtime/Core/Public/Delegates/DelegateSignatureImpl.inl
--- a/Engine/Source/Runtime/Core/Public/Delegates/DelegateSignatureImpl.inl
+++ b/Engine/Source/Runtime/Core/Public/Delegates/DelegateSignatureImpl.inl
@@ -664,11 +664,6 @@ private:
 	// Make sure TMulticastDelegateBase's protected functions are not accidentally exposed through the TMulticastDelegate API
 	using Super::AddDelegateInstance;
 	using Super::RemoveDelegateInstance;
-	using Super::CompactInvocationList;
-	using Super::GetInvocationList;
-	using Super::LockInvocationList;
-	using Super::UnlockInvocationList;
-	using Super::GetInvocationListLockCount;
 	using Super::GetDelegateInstanceProtectedHelper;
 
 public:
@@ -679,13 +674,7 @@ public:
 	 */
 	FDelegateHandle Add(FDelegate&& InNewDelegate)
 	{
-		FDelegateHandle Result;
-		if (Super::GetDelegateInstanceProtectedHelper(InNewDelegate))
-		{
-			Result = Super::AddDelegateInstance(MoveTemp(InNewDelegate));
-		}
-
-		return Result;
+		return Super::AddDelegateInstance(MoveTemp(InNewDelegate));
 	}
 
 	/**
@@ -695,13 +684,7 @@ public:
 	 */
 	FDelegateHandle Add(const FDelegate& InNewDelegate)
 	{
-		FDelegateHandle Result;
-		if (Super::GetDelegateInstanceProtectedHelper(InNewDelegate))
-		{
-			Result = Super::AddDelegateInstance(CopyTemp(InNewDelegate));
-		}
-
-		return Result;
+		return Super::AddDelegateInstance(CopyTemp(InNewDelegate));
 	}
 
 	/**
@@ -928,19 +911,8 @@ public:
 	{
 		if (&Other != this)
 		{
-			Super::Clear();
-
-			for (const TDelegateBase<UserPolicy>& OtherDelegateRef : Other.GetInvocationList())
-			{
-				if (IDelegateInstance* OtherInstance = Super::GetDelegateInstanceProtectedHelper(OtherDelegateRef))
-				{
-					FDelegate TempDelegate;
-					((DelegateInstanceInterfaceType*)OtherInstance)->CreateCopy(TempDelegate);
-					Super::AddDelegateInstance(MoveTemp(TempDelegate));
-				}
-			}
+			Super::template CopyFrom<DelegateInstanceInterfaceType, FDelegate>(Other);
 		}
-
 		return *this;
 	}
 
@@ -951,31 +923,7 @@ public:
 	 */
 	void Broadcast(ParamTypes... Params) const
 	{
-		bool NeedsCompaction = false;
-
-		Super::LockInvocationList();
-		{
-			const InvocationListType& LocalInvocationList = Super::GetInvocationList();
-
-			// call bound functions in reverse order, so we ignore any instances that may be added by callees
-			for (int32 InvocationListIndex = LocalInvocationList.Num() - 1; InvocationListIndex >= 0; --InvocationListIndex)
-			{
-				// this down-cast is OK! allows for managing invocation list in the base class without requiring virtual functions
-				const FDelegate& DelegateBase = (const FDelegate&)LocalInvocationList[InvocationListIndex];
-
-				IDelegateInstance* DelegateInstanceInterface = Super::GetDelegateInstanceProtectedHelper(DelegateBase);
-				if (DelegateInstanceInterface == nullptr || !((DelegateInstanceInterfaceType*)DelegateInstanceInterface)->ExecuteIfSafe(Params...))
-				{
-					NeedsCompaction = true;
-				}
-			}
-		}
-		Super::UnlockInvocationList();
-
-		if (NeedsCompaction)
-		{
-			const_cast<TMulticastDelegate*>(this)->CompactInvocationList();
-		}
+		Super::template Broadcast<DelegateInstanceInterfaceType, FDelegate, ParamTypes...>(Params...);
 	}
 };
 
diff --git a/Engine/Source/Runtime/Core/Public/Delegates/MulticastDelegateBase.h b/Engine/Source/Runtime/Core/Public/Delegates/MulticastDelegateBase.h
--- a/Engine/Source/Runtime/Core/Public/Delegates/MulticastDelegateBase.h
+++ b/Engine/Source/Runtime/Core/Public/Delegates/MulticastDelegateBase.h
@@ -141,6 +141,53 @@ protected:
 	{ }
 
 protected:
+	template<typename DelegateInstanceInterfaceType, typename DelegateType>
+	void CopyFrom(const TMulticastDelegateBase& Other)
+	{
+		checkSlow(&Other != this);
+
+		Clear();
+
+		for (const TDelegateBase<UserPolicy>& OtherDelegateRef : Other.GetInvocationList())
+		{
+			if (IDelegateInstance* OtherInstance = GetDelegateInstanceProtectedHelper(OtherDelegateRef))
+			{
+				DelegateType TempDelegate;
+				((DelegateInstanceInterfaceType*)OtherInstance)->CreateCopy(TempDelegate);
+				AddDelegateInstance(MoveTemp(TempDelegate));
+			}
+		}
+	}
+
+	template<typename DelegateInstanceInterfaceType, typename DelegateBaseType, typename... ParamTypes>
+	void Broadcast(ParamTypes... Params) const
+	{
+		bool NeedsCompaction = false;
+
+		LockInvocationList();
+		{
+			const InvocationListType& LocalInvocationList = GetInvocationList();
+
+			// call bound functions in reverse order, so we ignore any instances that may be added by callees
+			for (int32 InvocationListIndex = LocalInvocationList.Num() - 1; InvocationListIndex >= 0; --InvocationListIndex)
+			{
+				// this down-cast is OK! allows for managing invocation list in the base class without requiring virtual functions
+				const DelegateBaseType& DelegateBase = (const DelegateBaseType&)LocalInvocationList[InvocationListIndex];
+
+				IDelegateInstance* DelegateInstanceInterface = GetDelegateInstanceProtectedHelper(DelegateBase);
+				if (DelegateInstanceInterface == nullptr || !((DelegateInstanceInterfaceType*)DelegateInstanceInterface)->ExecuteIfSafe(Params...))
+				{
+					NeedsCompaction = true;
+				}
+			}
+		}
+		UnlockInvocationList();
+
+		if (NeedsCompaction)
+		{
+			const_cast<TMulticastDelegateBase*>(this)->CompactInvocationList();
+		}
+	}
 
 	/**
 	 * Adds the given delegate instance to the invocation list.
@@ -149,10 +196,14 @@ protected:
 	 */
 	inline FDelegateHandle AddDelegateInstance(TDelegateBase<UserPolicy>&& NewDelegateBaseRef)
 	{
-		// compact but obey threshold of when this will trigger
-		CompactInvocationList(true);
-		FDelegateHandle Result = NewDelegateBaseRef.GetHandle();
-		InvocationList.Add(MoveTemp(NewDelegateBaseRef));
+		FDelegateHandle Result;
+		if (GetDelegateInstanceProtectedHelper(NewDelegateBaseRef))
+		{
+			// compact but obey threshold of when this will trigger
+			CompactInvocationList(true);
+			Result = NewDelegateBaseRef.GetHandle();
+			InvocationList.Add(MoveTemp(NewDelegateBaseRef));
+		}
 		return Result;
 	}
 
@@ -180,6 +231,17 @@ protected:
 		return false;
 	}
 
+protected:
+	/**
+	 * Helper function for derived classes of TMulticastDelegateBase to get at the delegate instance.
+	 */
+	template <typename DelegateType>
+	static FORCEINLINE auto* GetDelegateInstanceProtectedHelper(const DelegateType& Base)
+	{
+		return Base.GetDelegateInstanceProtected();
+	}
+
+private:
 	/**
 	 * Removes any expired or deleted functions from the invocation list.
 	 *
@@ -260,16 +322,6 @@ protected:
 		return InvocationListLockCount;
 	}
 
-protected:
-	/**
-	 * Helper function for derived classes of TMulticastDelegateBase to get at the delegate instance.
-	 */
-	template <typename DelegateType>
-	static FORCEINLINE auto* GetDelegateInstanceProtectedHelper(const DelegateType& Base)
-	{
-		return Base.GetDelegateInstanceProtected();
-	}
-
 private:
 
 	/** Holds the collection of delegate instances to invoke. */
