added support for custom allocator for TSpscQueue and TMpscQueue

diff --git a/Engine/Source/Runtime/Core/Public/Containers/MpscQueue.h b/Engine/Source/Runtime/Core/Public/Containers/MpscQueue.h
--- a/Engine/Source/Runtime/Core/Public/Containers/MpscQueue.h
+++ b/Engine/Source/Runtime/Core/Public/Containers/MpscQueue.h
@@ -12,7 +12,7 @@
  * Fast multi-producer/single-consumer unbounded concurrent queue.
  * Based on http://www.1024cores.net/home/lock-free-algorithms/queues/non-intrusive-mpsc-node-based-queue
  */
-template<typename T>
+template<typename T, typename AllocatorType = FMemory>
 class TMpscQueue final
 {
 public:
@@ -22,7 +22,7 @@ public:
 
 	TMpscQueue()
 	{
-		FNode* Sentinel = new FNode;
+		FNode* Sentinel = new(AllocatorType::Malloc(sizeof(FNode), alignof(FNode))) FNode;
 		Head.store(Sentinel, std::memory_order_relaxed);
 		Tail = Sentinel;
 	}
@@ -32,7 +32,7 @@ public:
 		FNode* Next = Tail->Next.load(std::memory_order_relaxed);
 
 		// sentinel's value is already destroyed
-		delete Tail;
+		AllocatorType::Free(Tail);
 
 		while (Next != nullptr)
 		{
@@ -40,14 +40,14 @@ public:
 			Next = Tail->Next.load(std::memory_order_relaxed);
 
 			DestructItem((ElementType*)&Tail->Value);
-			delete Tail;
+			AllocatorType::Free(Tail);
 		}
 	}
 
 	template <typename... ArgTypes>
 	void Enqueue(ArgTypes&&... Args)
 	{
-		FNode* New = new FNode;
+		FNode* New = new(AllocatorType::Malloc(sizeof(FNode), alignof(FNode))) FNode;
 		new (&New->Value) ElementType(Forward<ArgTypes>(Args)...);
 
 		FNode* Prev = Head.exchange(New, std::memory_order_acq_rel);
@@ -67,7 +67,7 @@ public:
 		TOptional<ElementType> Res{ MoveTemp(*ValuePtr) };
 		DestructItem(ValuePtr);
 
-		delete Tail; // current sentinel
+		AllocatorType::Free(Tail); // current sentinel
 
 		Tail = Next; // new sentinel
 		return Res;
diff --git a/Engine/Source/Runtime/Core/Public/Containers/SpscQueue.h b/Engine/Source/Runtime/Core/Public/Containers/SpscQueue.h
--- a/Engine/Source/Runtime/Core/Public/Containers/SpscQueue.h
+++ b/Engine/Source/Runtime/Core/Public/Containers/SpscQueue.h
@@ -12,7 +12,7 @@
  * Fast single-producer/single-consumer unbounded concurrent queue. Doesn't free memory until destruction but recycles consumed items.
  * Based on http://www.1024cores.net/home/lock-free-algorithms/queues/unbounded-spsc-queue
  */
-template<typename T>
+template<typename T, typename AllocatorType = FMemory>
 class TSpscQueue final
 {
 public:
@@ -22,7 +22,7 @@ public:
 
 	TSpscQueue()
 	{
-		FNode* Node = new FNode;
+		FNode* Node = new(AllocatorType::Malloc(sizeof(FNode), alignof(FNode))) FNode;
 		Tail.store(Node, std::memory_order_relaxed);
 		Head = First = TailCopy = Node;
 	}
@@ -38,7 +38,7 @@ public:
 		{
 			FNode* Next = Node->Next.load(std::memory_order_relaxed);
 			bContinue = Node != LocalTail;
-			delete Node;
+			AllocatorType::Free(Node);
 			Node = Next;
 		} while (bContinue);
 
@@ -47,7 +47,7 @@ public:
 		{
 			FNode* Next = Node->Next.load(std::memory_order_relaxed);
 			DestructItem((ElementType*)&Node->Value);
-			delete Node;
+			AllocatorType::Free(Node);
 			Node = Next;
 		}
 	}
@@ -56,7 +56,7 @@ public:
 	void Enqueue(ArgTypes&&... Args)
 	{
 		FNode* Node = AllocNode();
-		new (&Node->Value) ElementType(Forward<ArgTypes>(Args)...);
+		new(&Node->Value) ElementType(Forward<ArgTypes>(Args)...);
 
 		Head->Next.store(Node, std::memory_order_release);
 		Head = Node;
@@ -147,7 +147,7 @@ private:
 			return AllocFromCache();
 		}
 
-		return new FNode();
+		return new(AllocatorType::Malloc(sizeof(FNode), alignof(FNode))) FNode();
 	}
 
 private:
