Added PLATFORM_HAS_ASAN_INCLUDE define to be used instead of _has_include(<sanitizer/asan_interface.h>) to be able to disable, if needed on specific platform, the include of that header file even if the header file is present

#rnx

diff --git a/Engine/Source/Runtime/Core/Public/Experimental/ConcurrentLinearAllocator.h b/Engine/Source/Runtime/Core/Public/Experimental/ConcurrentLinearAllocator.h
--- a/Engine/Source/Runtime/Core/Public/Experimental/ConcurrentLinearAllocator.h
+++ b/Engine/Source/Runtime/Core/Public/Experimental/ConcurrentLinearAllocator.h
@@ -21,7 +21,7 @@
 #define SUPPORTS_VERY_LARGE_ALIGNMENTS 1
 #endif
 
-#if __has_include(<sanitizer/asan_interface.h>)
+#if PLATFORM_HAS_ASAN_INCLUDE
 #include <sanitizer/asan_interface.h>
 #if defined(__SANITIZE_ADDRESS__)
 #define IS_ASAN_ENABLED 1
diff --git a/Engine/Source/Runtime/Core/Public/HAL/Platform.h b/Engine/Source/Runtime/Core/Public/HAL/Platform.h
--- a/Engine/Source/Runtime/Core/Public/HAL/Platform.h
+++ b/Engine/Source/Runtime/Core/Public/HAL/Platform.h
@@ -115,6 +115,10 @@
 	#define USING_ADDRESS_SANITISER 0
 #endif
 
+#ifndef PLATFORM_HAS_ASAN_INCLUDE
+	#define PLATFORM_HAS_ASAN_INCLUDE __has_include(<sanitizer/asan_interface.h>)
+#endif
+
 //---------------------------------------------------------
 // Include main platform setup header (XXX/XXXPlatform.h)
 //---------------------------------------------------------
