Prevent a possible division by zero in FHistogram::GetAverageOfAllMeasures
Reported on UDN case 416970

diff --git a/Engine/Source/Runtime/Core/Public/ProfilingDebugging/Histogram.h b/Engine/Source/Runtime/Core/Public/ProfilingDebugging/Histogram.h
--- a/Engine/Source/Runtime/Core/Public/ProfilingDebugging/Histogram.h
+++ b/Engine/Source/Runtime/Core/Public/ProfilingDebugging/Histogram.h
@@ -130,7 +130,7 @@ struct CORE_API FHistogram
 	/** Returns the average of all measurements (essentially a shortcut for Sum/Count). */
 	inline double GetAverageOfAllMeasures() const
 	{
-		return SumOfAllMeasures / (double)CountOfAllMeasures;
+		return (CountOfAllMeasures > 0) ? (SumOfAllMeasures / (double)CountOfAllMeasures) : 0 ;
 	}
 
 	/** Returns the minimum of all measurements. */
