Adding OnActivatedByProtocol Core Delegate for platforms that support launching via URI or other protocol.

 - only the protocol parameter string is stored - any URI/protocol prefix is stripped for consistency between platforms.

diff --git a/Engine/Source/Runtime/Core/Private/Misc/CoreDelegates.cpp b/Engine/Source/Runtime/Core/Private/Misc/CoreDelegates.cpp
--- a/Engine/Source/Runtime/Core/Private/Misc/CoreDelegates.cpp
+++ b/Engine/Source/Runtime/Core/Private/Misc/CoreDelegates.cpp
@@ -197,6 +197,8 @@ FCoreDelegates::FIsLoadingMovieCurrentlyPlaying FCoreDelegates::IsLoadingMovieCu
 
 FCoreDelegates::FShouldLaunchUrl FCoreDelegates::ShouldLaunchUrl;
 
+FCoreDelegates::FOnActivatedByProtocol FCoreDelegates::OnActivatedByProtocol;
+
 FCoreDelegates::FOnGCFinishDestroyTimeExtended FCoreDelegates::OnGCFinishDestroyTimeExtended;
 
 FCoreDelegates::FAccesExtraBinaryConfigData FCoreDelegates::AccessExtraBinaryConfigData;
diff --git a/Engine/Source/Runtime/Core/Public/Misc/CoreDelegates.h b/Engine/Source/Runtime/Core/Public/Misc/CoreDelegates.h
--- a/Engine/Source/Runtime/Core/Public/Misc/CoreDelegates.h
+++ b/Engine/Source/Runtime/Core/Public/Misc/CoreDelegates.h
@@ -631,6 +631,10 @@ public:
 	DECLARE_DELEGATE_RetVal_OneParam(bool, FShouldLaunchUrl, const TCHAR* /* URL */);
 	static FShouldLaunchUrl ShouldLaunchUrl;
 
+	// Callback when the application has been activated by protocol (with optional user id, depending on the platform)
+	DECLARE_MULTICAST_DELEGATE_TwoParams(FOnActivatedByProtocol, const FString& /* parameter */, FPlatformUserId /* user id = PLATFORMUSERID_NONE */ );
+	static FOnActivatedByProtocol OnActivatedByProtocol;
+
 	/** Sent when GC finish destroy takes more time than expected */
 	DECLARE_MULTICAST_DELEGATE_OneParam(FOnGCFinishDestroyTimeExtended, const FString&);
 	static FOnGCFinishDestroyTimeExtended OnGCFinishDestroyTimeExtended;
