Delegates support TObjectPtr.

diff --git a/Engine/Source/Runtime/Core/Public/Delegates/DelegateSignatureImpl.inl b/Engine/Source/Runtime/Core/Public/Delegates/DelegateSignatureImpl.inl
--- a/Engine/Source/Runtime/Core/Public/Delegates/DelegateSignatureImpl.inl
+++ b/Engine/Source/Runtime/Core/Public/Delegates/DelegateSignatureImpl.inl
@@ -24,6 +24,10 @@ class IDelegateInstance;
 struct FWeakObjectPtr;
 template <typename FuncType, typename UserPolicy> struct IBaseDelegateInstance;
 template<typename UserPolicy> class TMulticastDelegateBase;
+template <typename T> struct TObjectPtr;
+
+template <typename T>
+T* ToRawPtr(const TObjectPtr<T>& Ptr);
 
 /**
  * Unicast delegate template class.
@@ -292,6 +296,13 @@ public:
 		TBaseUFunctionDelegateInstance<UObjectTemplate, FuncType, UserPolicy, VarTypes...>::Create(Result, InUserObject, InFunctionName, Vars...);
 		return Result;
 	}
+	template <typename UObjectTemplate, typename... VarTypes>
+	UE_NODISCARD inline static TDelegate<RetValType(ParamTypes...), UserPolicy> CreateUFunction(TObjectPtr<UObjectTemplate> InUserObject, const FName& InFunctionName, VarTypes... Vars)
+	{
+		TDelegate<RetValType(ParamTypes...), UserPolicy> Result;
+		TBaseUFunctionDelegateInstance<UObjectTemplate, FuncType, UserPolicy, VarTypes...>::Create(Result, ToRawPtr(InUserObject), InFunctionName, Vars...);
+		return Result;
+	}
 
 	/**
 	 * Static: Creates a UObject-based member function delegate.
@@ -315,6 +326,22 @@ public:
 		TBaseUObjectMethodDelegateInstance<true, const UserClass, FuncType, UserPolicy, VarTypes...>::Create(Result, InUserObject, InFunc, Vars...);
 		return Result;
 	}
+	template <typename UserClass, typename... VarTypes>
+	UE_NODISCARD inline static TDelegate<RetValType(ParamTypes...), UserPolicy> CreateUObject(TObjectPtr<UserClass> InUserObject, typename TMemFunPtrType<false, UserClass, RetValType (ParamTypes..., VarTypes...)>::Type InFunc, VarTypes... Vars)
+	{
+		static_assert(!TIsConst<UserClass>::Value, "Attempting to bind a delegate with a const object pointer and non-const member function.");
+
+		TDelegate<RetValType(ParamTypes...), UserPolicy> Result;
+		TBaseUObjectMethodDelegateInstance<false, UserClass, FuncType, UserPolicy, VarTypes...>::Create(Result, ToRawPtr(InUserObject), InFunc, Vars...);
+		return Result;
+	}
+	template <typename UserClass, typename... VarTypes>
+	UE_NODISCARD inline static TDelegate<RetValType(ParamTypes...), UserPolicy> CreateUObject(TObjectPtr<const UserClass> InUserObject, typename TMemFunPtrType<true, UserClass, RetValType (ParamTypes..., VarTypes...)>::Type InFunc, VarTypes... Vars)
+	{
+		TDelegate<RetValType(ParamTypes...), UserPolicy> Result;
+		TBaseUObjectMethodDelegateInstance<true, const UserClass, FuncType, UserPolicy, VarTypes...>::Create(Result, ToRawPtr(InUserObject), InFunc, Vars...);
+		return Result;
+	}
 
 public:
 
@@ -548,6 +575,11 @@ public:
 	{
 		*this = CreateUFunction(InUserObject, InFunctionName, Vars...);
 	}
+	template <typename UObjectTemplate, typename... VarTypes>
+	inline void BindUFunction(TObjectPtr<UObjectTemplate> InUserObject, const FName& InFunctionName, VarTypes... Vars)
+	{
+		*this = CreateUFunction(InUserObject, InFunctionName, Vars...);
+	}
 
 	/**
 	 * Binds a UObject-based member function delegate.
@@ -567,6 +599,18 @@ public:
 	{
 		*this = CreateUObject(InUserObject, InFunc, Vars...);
 	}
+	template <typename UserClass, typename... VarTypes>
+	inline void BindUObject(TObjectPtr<UserClass> InUserObject, typename TMemFunPtrType<false, UserClass, RetValType (ParamTypes..., VarTypes...)>::Type InFunc, VarTypes... Vars)
+	{
+		static_assert(!TIsConst<UserClass>::Value, "Attempting to bind a delegate with a const object pointer and non-const member function.");
+
+		*this = CreateUObject(InUserObject, InFunc, Vars...);
+	}
+	template <typename UserClass, typename... VarTypes>
+	inline void BindUObject(TObjectPtr<const UserClass> InUserObject, typename TMemFunPtrType<true, UserClass, RetValType (ParamTypes..., VarTypes...)>::Type InFunc, VarTypes... Vars)
+	{
+		*this = CreateUObject(InUserObject, InFunc, Vars...);
+	}
 
 public:
 	/**
@@ -844,6 +888,11 @@ public:
 	{
 		return Add(FDelegate::CreateUFunction(InUserObject, InFunctionName, Vars...));
 	}
+	template <typename UObjectTemplate, typename... VarTypes>
+	inline FDelegateHandle AddUFunction(TObjectPtr<UObjectTemplate> InUserObject, const FName& InFunctionName, VarTypes... Vars)
+	{
+		return Add(FDelegate::CreateUFunction(InUserObject, InFunctionName, Vars...));
+	}
 
 	/**
 	 * Adds a UObject-based member function delegate.
@@ -865,6 +914,18 @@ public:
 	{
 		return Add(FDelegate::CreateUObject(InUserObject, InFunc, Vars...));
 	}
+	template <typename UserClass, typename... VarTypes>
+	inline FDelegateHandle AddUObject(TObjectPtr<UserClass> InUserObject, typename TMemFunPtrType<false, UserClass, void (ParamTypes..., VarTypes...)>::Type InFunc, VarTypes... Vars)
+	{
+		static_assert(!TIsConst<UserClass>::Value, "Attempting to bind a delegate with a const object pointer and non-const member function.");
+
+		return Add(FDelegate::CreateUObject(InUserObject, InFunc, Vars...));
+	}
+	template <typename UserClass, typename... VarTypes>
+	inline FDelegateHandle AddUObject(TObjectPtr<const UserClass> InUserObject, typename TMemFunPtrType<true, UserClass, void (ParamTypes..., VarTypes...)>::Type InFunc, VarTypes... Vars)
+	{
+		return Add(FDelegate::CreateUObject(InUserObject, InFunc, Vars...));
+	}
 
 public:
 
@@ -987,6 +1048,11 @@ public:
 
 		ensureMsgf(this->IsBound(), TEXT("Unable to bind delegate to '%s' (function might not be marked as a UFUNCTION or object may be pending kill)"), *InFunctionName.ToString());
 	}
+	template< class UserClass >
+	void __Internal_BindDynamic( TObjectPtr<UserClass> InUserObject, typename TMethodPtrResolver< UserClass >::FMethodPtr InMethodPtr, FName InFunctionName )
+	{
+		__Internal_BindDynamic(ToRawPtr(InUserObject), InMethodPtr, InFunctionName);
+	}
 
 	friend uint32 GetTypeHash(const TBaseDynamicDelegate& Key)
 	{
@@ -1046,6 +1112,11 @@ public:
 
 		return this->Contains( InUserObject, InFunctionName );
 	}
+	template< class UserClass >
+	bool __Internal_IsAlreadyBound( TObjectPtr<UserClass> InUserObject, typename FDelegate::template TMethodPtrResolver< UserClass >::FMethodPtr InMethodPtr, FName InFunctionName ) const
+	{
+		return __Internal_IsAlreadyBound(ToRawPtr(InUserObject), InMethodPtr, InFunctionName);
+	}
 
 	/**
 	 * Binds a UObject instance and a UObject method address to this multi-cast delegate.
@@ -1069,6 +1140,11 @@ public:
 
 		this->Add( NewDelegate );
 	}
+	template< class UserClass >
+	void __Internal_AddDynamic( TObjectPtr<UserClass> InUserObject, typename FDelegate::template TMethodPtrResolver< UserClass >::FMethodPtr InMethodPtr, FName InFunctionName )
+	{
+		__Internal_AddDynamic(ToRawPtr(InUserObject), InMethodPtr, InFunctionName);
+	}
 
 	/**
 	 * Binds a UObject instance and a UObject method address to this multi-cast delegate, but only if it hasn't been bound before.
@@ -1092,6 +1168,11 @@ public:
 
 		this->AddUnique( NewDelegate );
 	}
+	template< class UserClass >
+	void __Internal_AddUniqueDynamic( TObjectPtr<UserClass> InUserObject, typename FDelegate::template TMethodPtrResolver< UserClass >::FMethodPtr InMethodPtr, FName InFunctionName )
+	{
+		__Internal_AddUniqueDynamic(ToRawPtr(InUserObject), InMethodPtr, InFunctionName);
+	}
 
 	/**
 	 * Unbinds a UObject instance and a UObject method address from this multi-cast delegate.
@@ -1112,6 +1193,11 @@ public:
 
 		this->Remove( InUserObject, InFunctionName );
 	}
+	template< class UserClass >
+	void __Internal_RemoveDynamic( TObjectPtr<UserClass> InUserObject, typename FDelegate::template TMethodPtrResolver< UserClass >::FMethodPtr InMethodPtr, FName InFunctionName )
+	{
+		__Internal_RemoveDynamic(ToRawPtr(InUserObject), InMethodPtr, InFunctionName);
+	}
 
 	// NOTE:  Broadcast() method must be defined in derived classes
 };
