Check the IsValid flag when comparing bounding boxes, and more generally take the flag more often into account

diff --git a/Engine/Source/Runtime/Core/Public/Math/Box.h b/Engine/Source/Runtime/Core/Public/Math/Box.h
--- a/Engine/Source/Runtime/Core/Public/Math/Box.h
+++ b/Engine/Source/Runtime/Core/Public/Math/Box.h
@@ -104,26 +104,30 @@ public:
 
 	/**
 	 * Compares two boxes for equality.
+	 * 
+	 * Returns true if both bounding boxes are invalid. Returns false if one of the bounding boxes is invalid.
 	 *
 	 * @return true if the boxes are equal, false otherwise.
 	 */
 	FORCEINLINE bool operator==( const TBox<T>& Other ) const
 	{
-		return (Min == Other.Min) && (Max == Other.Max);
+		return (!IsValid && !Other.IsValid) || ((IsValid && Other.IsValid) && (Min == Other.Min) && (Max == Other.Max));
 	}
 
 	/**
 	 * Compares two boxes for inequality.
-	 *
+	 * 
 	 * @return false if the boxes are equal, true otherwise.
 	 */
 	FORCEINLINE bool operator!=( const TBox<T>& Other) const
 	{
-		return (Min != Other.Min) || (Max != Other.Max);
+		return !(*this == Other);
 	}
 
 	/**
 	 * Check against another box for equality, within specified error limits.
+	 * 
+	 * Returns true if both bounding boxes are invalid. Returns false if one of the bounding boxes is invalid.
 	 *
 	 * @param Other The box to check against.
 	 * @param Tolerance Error tolerance.
@@ -131,7 +135,7 @@ public:
 	 */
 	bool Equals(const TBox<T>& Other, T Tolerance=UE_KINDA_SMALL_NUMBER) const
 	{
-		return Min.Equals(Other.Min, Tolerance) && Max.Equals(Other.Max, Tolerance);
+		return (!IsValid && !Other.IsValid) || ((IsValid && Other.IsValid) && Min.Equals(Other.Min, Tolerance) && Max.Equals(Other.Max, Tolerance));
 	}
 
 	/**
@@ -221,6 +225,7 @@ public:
 	 */
 	UE_NODISCARD FORCEINLINE TBox<T> ExpandBy(T W) const
 	{
+		checkSlow(IsValid);
 		return TBox<T>(Min - TVector<T>(W, W, W), Max + TVector<T>(W, W, W));
 	}
 
@@ -232,6 +237,7 @@ public:
 	*/
 	UE_NODISCARD FORCEINLINE TBox<T> ExpandBy(const TVector<T>& V) const
 	{
+		checkSlow(IsValid);
 		return TBox<T>(Min - V, Max + V);
 	}
 
@@ -244,6 +250,7 @@ public:
 	*/
 	UE_NODISCARD TBox<T> ExpandBy(const TVector<T>& Neg, const TVector<T>& Pos) const
 	{
+		checkSlow(IsValid);
 		return TBox<T>(Min - Neg, Max + Pos);
 	}
 
@@ -255,6 +262,7 @@ public:
 	 */
 	UE_NODISCARD FORCEINLINE TBox<T> ShiftBy( const TVector<T>& Offset ) const
 	{
+		checkSlow(IsValid);
 		return TBox<T>(Min + Offset, Max + Offset);
 	}
 
@@ -278,6 +286,7 @@ public:
 	 */
 	FORCEINLINE TVector<T> GetCenter() const
 	{
+		checkSlow(IsValid);
 		return TVector<T>((Min + Max) * 0.5f);
 	}
 
@@ -310,6 +319,7 @@ public:
 	 */
 	FORCEINLINE TVector<T> GetExtent() const
 	{
+		checkSlow(IsValid);
 		return 0.5f * (Max - Min);
 	}
 
@@ -333,6 +343,7 @@ public:
 	 */
 	FORCEINLINE TVector<T> GetSize() const
 	{
+		checkSlow(IsValid);
 		return (Max - Min);
 	}
 
@@ -344,6 +355,7 @@ public:
 	 */
 	FORCEINLINE T GetVolume() const
 	{
+		checkSlow(IsValid);
 		return (Max.X - Min.X) * (Max.Y - Min.Y) * (Max.Z - Min.Z);
 	}
 
@@ -397,6 +409,7 @@ public:
 	 */
 	FORCEINLINE bool IsInside( const TVector<T>& In ) const
 	{
+		checkSlow(IsValid);
 		return ((In.X > Min.X) && (In.X < Max.X) && (In.Y > Min.Y) && (In.Y < Max.Y) && (In.Z > Min.Z) && (In.Z < Max.Z));
 	}
 
@@ -409,6 +422,7 @@ public:
 	 */
 	FORCEINLINE bool IsInsideOrOn( const TVector<T>& In ) const
 	{
+		checkSlow(IsValid);
 		return ((In.X >= Min.X) && (In.X <= Max.X) && (In.Y >= Min.Y) && (In.Y <= Max.Y) && (In.Z >= Min.Z) && (In.Z <= Max.Z));
 	}
 
@@ -432,6 +446,7 @@ public:
 	 */
 	FORCEINLINE bool IsInsideXY( const TVector<T>& In ) const
 	{
+		checkSlow(IsValid);
 		return ((In.X > Min.X) && (In.X < Max.X) && (In.Y > Min.Y) && (In.Y < Max.Y));
 	}
 
@@ -444,6 +459,7 @@ public:
 	 */
 	FORCEINLINE bool IsInsideOrOnXY(const FVector& In) const
 	{
+		checkSlow(IsValid);
 		return ((In.X >= Min.X) && (In.X <= Max.X) && (In.Y >= Min.Y) && (In.Y <= Max.Y));
 	}
 
@@ -606,6 +622,8 @@ FORCEINLINE TBox<T>& TBox<T>::operator+=( const TBox<T>& Other )
 template<typename T>
 FORCEINLINE TVector<T> TBox<T>::GetClosestPointTo( const TVector<T>& Point ) const
 {
+	checkSlow(IsValid);
+
 	// start by considering the point inside the box
 	TVector<T> ClosestPoint = Point;
 
@@ -646,6 +664,8 @@ FORCEINLINE TVector<T> TBox<T>::GetClosestPointTo( const TVector<T>& Point ) con
 template<typename T>
 FORCEINLINE bool TBox<T>::Intersect( const TBox<T>& Other ) const
 {
+	checkSlow(IsValid && Other.IsValid);
+
 	if ((Min.X > Other.Max.X) || (Other.Min.X > Max.X))
 	{
 		return false;
@@ -668,6 +688,8 @@ FORCEINLINE bool TBox<T>::Intersect( const TBox<T>& Other ) const
 template<typename T>
 FORCEINLINE bool TBox<T>::IntersectXY( const TBox<T>& Other ) const
 {
+	checkSlow(IsValid && Other.IsValid);
+
 	if ((Min.X > Other.Max.X) || (Other.Min.X > Max.X))
 	{
 		return false;
@@ -693,7 +715,7 @@ template<typename T>
 TBox<T> TBox<T>::TransformBy(const TMatrix<T>& M) const
 {
 	// if we are not valid, return another invalid box.
-	if (!IsValid)
+	if (!ensure(IsValid))
 	{
 		return TBox<T>(ForceInit);
 	}
@@ -741,6 +763,7 @@ TBox<T> TBox<T>::TransformBy(const TTransform<T>& M) const
 template<typename T>
 void TBox<T>::GetVertices(TVector<T>(&Vertices)[8]) const
 {
+	checkSlow(IsValid);
 	Vertices[0] = TVector<T>(Min);
 	Vertices[1] = TVector<T>(Min.X, Min.Y, Max.Z);
 	Vertices[2] = TVector<T>(Min.X, Max.Y, Min.Z);
@@ -846,6 +869,7 @@ inline bool FMath::PointBoxIntersection
 	const UE::Math::TBox<FReal>&		Box
 	)
 {
+	checkSlow(Box.IsValid);
 	return (Point.X >= Box.Min.X && Point.X <= Box.Max.X &&
 			Point.Y >= Box.Min.Y && Point.Y <= Box.Max.Y &&
 			Point.Z >= Box.Min.Z && Point.Z <= Box.Max.Z);
@@ -873,6 +897,8 @@ inline bool FMath::LineBoxIntersection
 	const UE::Math::TVector<FReal>&	OneOverStartToEnd
 	)
 {
+	checkSlow(Box.IsValid);
+
 	UE::Math::TVector<FReal>	Time;
 	bool	bStartIsOutside = false;
 
@@ -1003,6 +1029,7 @@ inline bool FMath::LineBoxIntersection
 template<typename FReal>
 inline bool FMath::SphereAABBIntersection(const UE::Math::TVector<FReal>& SphereCenter, const FReal RadiusSquared, const UE::Math::TBox<FReal>& AABB)
 {
+	checkSlow(AABB.IsValid);
 	// Accumulates the distance as we iterate axis
 	FReal DistSquared = 0.f;
 	// Check each axis for min/max and add the distance accordingly
diff --git a/Engine/Source/Runtime/Engine/Public/MeshCardRepresentation.h b/Engine/Source/Runtime/Engine/Public/MeshCardRepresentation.h
--- a/Engine/Source/Runtime/Engine/Public/MeshCardRepresentation.h
+++ b/Engine/Source/Runtime/Engine/Public/MeshCardRepresentation.h
@@ -115,12 +115,11 @@ public:
 		return WorldOBB;
 	}
 
-	FBox GetBox()
+	FBox GetBox() const
 	{
-		FBox Box;
-		Box.Min = FVector(AxisX.GetAbs() * -Extent.X + AxisY.GetAbs() * -Extent.Y + AxisZ.GetAbs() * -Extent.Z + Origin);
-		Box.Max = FVector(AxisX.GetAbs() * +Extent.X + AxisY.GetAbs() * +Extent.Y + AxisZ.GetAbs() * +Extent.Z + Origin);
-		return Box;
+		FVector BoxMin(AxisX.GetAbs() * -Extent.X + AxisY.GetAbs() * -Extent.Y + AxisZ.GetAbs() * -Extent.Z + Origin);
+		FVector BoxMax(AxisX.GetAbs() * +Extent.X + AxisY.GetAbs() * +Extent.Y + AxisZ.GetAbs() * +Extent.Z + Origin);
+		return FBox(BoxMin, BoxMax);
 	}
 
 	friend FArchive& operator<<(FArchive& Ar, FLumenCardOBB& Data)
