Replace __DATE__ and __TIME__ in Core module with BuildSettings::GetBuildDate() and GetBuildTime()

#rnx

diff --git a/Engine/Source/Runtime/Core/Private/Misc/App.cpp b/Engine/Source/Runtime/Core/Private/Misc/App.cpp
--- a/Engine/Source/Runtime/Core/Private/Misc/App.cpp
+++ b/Engine/Source/Runtime/Core/Private/Misc/App.cpp
@@ -336,9 +336,9 @@ void FApp::PrintStartupLogMessages()
 	UE_LOG(LogInit, Log, TEXT("OS: %s (%s), CPU: %s, GPU: %s"), *OSLabel, *OSVersion, *FPlatformMisc::GetCPUBrand(), *FPlatformMisc::GetPrimaryGPUBrand());
 
 #if PLATFORM_64BITS
-	UE_LOG(LogInit, Log, TEXT("Compiled (64-bit): %s %s"), ANSI_TO_TCHAR(__DATE__), ANSI_TO_TCHAR(__TIME__));
+	UE_LOG(LogInit, Log, TEXT("Compiled (64-bit): %s %s"), BuildSettings::GetBuildDate(), BuildSettings::GetBuildTime());
 #else
-	UE_LOG(LogInit, Log, TEXT("Compiled (32-bit): %s %s"), ANSI_TO_TCHAR(__DATE__), ANSI_TO_TCHAR(__TIME__));
+	UE_LOG(LogInit, Log, TEXT("Compiled (32-bit): %s %s"), BuildSettings::GetBuildDate(), BuildSettings::GetBuildTime());
 #endif
 
 	// Print compiler version info
