Removed including CoreSharedPCH.h.

diff --git a/Engine/Source/Runtime/Core/Private/ProfilingDebugging/StringTrace.cpp b/Engine/Source/Runtime/Core/Private/ProfilingDebugging/StringTrace.cpp
--- a/Engine/Source/Runtime/Core/Private/ProfilingDebugging/StringTrace.cpp
+++ b/Engine/Source/Runtime/Core/Private/ProfilingDebugging/StringTrace.cpp
@@ -1,7 +1,15 @@
 ﻿// Copyright Epic Games, Inc. All Rights Reserved.
+
 #include "ProfilingDebugging/StringsTrace.h"
-#include "CoreSharedPCH.h"
+
+#include "Containers/Set.h"
 #include "HAL/CriticalSection.h"
+#include "HAL/PlatformString.h"
+#include "Misc/ScopeLock.h"
+#include "ProfilingDebugging/StringsTrace.h"
+#include "Trace/Trace.h"
+#include "UObject/NameTypes.h"
+#include "UObject/UnrealNames.h"
 
 //////////////////////////////////////////////////////////////////////////////////////////
 UE_TRACE_EVENT_DEFINE(Strings, StaticString);
