Remove checks for bounding box validity for now, they're triggering too often in debug mode

diff --git a/Engine/Source/Runtime/Core/Public/Math/Box.h b/Engine/Source/Runtime/Core/Public/Math/Box.h
--- a/Engine/Source/Runtime/Core/Public/Math/Box.h
+++ b/Engine/Source/Runtime/Core/Public/Math/Box.h
@@ -225,7 +225,6 @@ public:
 	 */
 	UE_NODISCARD FORCEINLINE TBox<T> ExpandBy(T W) const
 	{
-		checkSlow(IsValid);
 		return TBox<T>(Min - TVector<T>(W, W, W), Max + TVector<T>(W, W, W));
 	}
 
@@ -237,7 +236,6 @@ public:
 	*/
 	UE_NODISCARD FORCEINLINE TBox<T> ExpandBy(const TVector<T>& V) const
 	{
-		checkSlow(IsValid);
 		return TBox<T>(Min - V, Max + V);
 	}
 
@@ -250,7 +248,6 @@ public:
 	*/
 	UE_NODISCARD TBox<T> ExpandBy(const TVector<T>& Neg, const TVector<T>& Pos) const
 	{
-		checkSlow(IsValid);
 		return TBox<T>(Min - Neg, Max + Pos);
 	}
 
@@ -262,7 +259,6 @@ public:
 	 */
 	UE_NODISCARD FORCEINLINE TBox<T> ShiftBy( const TVector<T>& Offset ) const
 	{
-		checkSlow(IsValid);
 		return TBox<T>(Min + Offset, Max + Offset);
 	}
 
@@ -286,7 +282,6 @@ public:
 	 */
 	FORCEINLINE TVector<T> GetCenter() const
 	{
-		checkSlow(IsValid);
 		return TVector<T>((Min + Max) * 0.5f);
 	}
 
@@ -319,7 +314,6 @@ public:
 	 */
 	FORCEINLINE TVector<T> GetExtent() const
 	{
-		checkSlow(IsValid);
 		return 0.5f * (Max - Min);
 	}
 
@@ -343,7 +337,6 @@ public:
 	 */
 	FORCEINLINE TVector<T> GetSize() const
 	{
-		checkSlow(IsValid);
 		return (Max - Min);
 	}
 
@@ -355,7 +348,6 @@ public:
 	 */
 	FORCEINLINE T GetVolume() const
 	{
-		checkSlow(IsValid);
 		return (Max.X - Min.X) * (Max.Y - Min.Y) * (Max.Z - Min.Z);
 	}
 
@@ -409,7 +401,6 @@ public:
 	 */
 	FORCEINLINE bool IsInside( const TVector<T>& In ) const
 	{
-		checkSlow(IsValid);
 		return ((In.X > Min.X) && (In.X < Max.X) && (In.Y > Min.Y) && (In.Y < Max.Y) && (In.Z > Min.Z) && (In.Z < Max.Z));
 	}
 
@@ -422,7 +413,6 @@ public:
 	 */
 	FORCEINLINE bool IsInsideOrOn( const TVector<T>& In ) const
 	{
-		checkSlow(IsValid);
 		return ((In.X >= Min.X) && (In.X <= Max.X) && (In.Y >= Min.Y) && (In.Y <= Max.Y) && (In.Z >= Min.Z) && (In.Z <= Max.Z));
 	}
 
@@ -446,7 +436,6 @@ public:
 	 */
 	FORCEINLINE bool IsInsideXY( const TVector<T>& In ) const
 	{
-		checkSlow(IsValid);
 		return ((In.X > Min.X) && (In.X < Max.X) && (In.Y > Min.Y) && (In.Y < Max.Y));
 	}
 
@@ -459,7 +448,6 @@ public:
 	 */
 	FORCEINLINE bool IsInsideOrOnXY(const FVector& In) const
 	{
-		checkSlow(IsValid);
 		return ((In.X >= Min.X) && (In.X <= Max.X) && (In.Y >= Min.Y) && (In.Y <= Max.Y));
 	}
 
@@ -622,8 +610,6 @@ FORCEINLINE TBox<T>& TBox<T>::operator+=( const TBox<T>& Other )
 template<typename T>
 FORCEINLINE TVector<T> TBox<T>::GetClosestPointTo( const TVector<T>& Point ) const
 {
-	checkSlow(IsValid);
-
 	// start by considering the point inside the box
 	TVector<T> ClosestPoint = Point;
 
@@ -664,8 +650,6 @@ FORCEINLINE TVector<T> TBox<T>::GetClosestPointTo( const TVector<T>& Point ) con
 template<typename T>
 FORCEINLINE bool TBox<T>::Intersect( const TBox<T>& Other ) const
 {
-	checkSlow(IsValid && Other.IsValid);
-
 	if ((Min.X > Other.Max.X) || (Other.Min.X > Max.X))
 	{
 		return false;
@@ -688,8 +672,6 @@ FORCEINLINE bool TBox<T>::Intersect( const TBox<T>& Other ) const
 template<typename T>
 FORCEINLINE bool TBox<T>::IntersectXY( const TBox<T>& Other ) const
 {
-	checkSlow(IsValid && Other.IsValid);
-
 	if ((Min.X > Other.Max.X) || (Other.Min.X > Max.X))
 	{
 		return false;
@@ -715,7 +697,7 @@ template<typename T>
 TBox<T> TBox<T>::TransformBy(const TMatrix<T>& M) const
 {
 	// if we are not valid, return another invalid box.
-	if (!ensure(IsValid))
+	if (!IsValid)
 	{
 		return TBox<T>(ForceInit);
 	}
@@ -763,7 +745,6 @@ TBox<T> TBox<T>::TransformBy(const TTransform<T>& M) const
 template<typename T>
 void TBox<T>::GetVertices(TVector<T>(&Vertices)[8]) const
 {
-	checkSlow(IsValid);
 	Vertices[0] = TVector<T>(Min);
 	Vertices[1] = TVector<T>(Min.X, Min.Y, Max.Z);
 	Vertices[2] = TVector<T>(Min.X, Max.Y, Min.Z);
@@ -869,7 +850,6 @@ inline bool FMath::PointBoxIntersection
 	const UE::Math::TBox<FReal>&		Box
 	)
 {
-	checkSlow(Box.IsValid);
 	return (Point.X >= Box.Min.X && Point.X <= Box.Max.X &&
 			Point.Y >= Box.Min.Y && Point.Y <= Box.Max.Y &&
 			Point.Z >= Box.Min.Z && Point.Z <= Box.Max.Z);
@@ -897,8 +877,6 @@ inline bool FMath::LineBoxIntersection
 	const UE::Math::TVector<FReal>&	OneOverStartToEnd
 	)
 {
-	checkSlow(Box.IsValid);
-
 	UE::Math::TVector<FReal>	Time;
 	bool	bStartIsOutside = false;
 
@@ -1029,7 +1007,6 @@ inline bool FMath::LineBoxIntersection
 template<typename FReal>
 inline bool FMath::SphereAABBIntersection(const UE::Math::TVector<FReal>& SphereCenter, const FReal RadiusSquared, const UE::Math::TBox<FReal>& AABB)
 {
-	checkSlow(AABB.IsValid);
 	// Accumulates the distance as we iterate axis
 	FReal DistSquared = 0.f;
 	// Check each axis for min/max and add the distance accordingly
