Don't pass nullptr into ExceptionPointers in MiniDumpWriteDump

diff --git a/Engine/Source/Runtime/Core/Private/Windows/WindowsPlatformCrashContext.cpp b/Engine/Source/Runtime/Core/Private/Windows/WindowsPlatformCrashContext.cpp
--- a/Engine/Source/Runtime/Core/Private/Windows/WindowsPlatformCrashContext.cpp
+++ b/Engine/Source/Runtime/Core/Private/Windows/WindowsPlatformCrashContext.cpp
@@ -136,7 +136,7 @@ namespace {
 			MinidumpType = (MINIDUMP_TYPE)(MiniDumpWithFullMemory | MiniDumpWithFullMemoryInfo | MiniDumpWithHandleData | MiniDumpWithThreadInfo | MiniDumpWithUnloadedModules);
 		}
 
-		const BOOL Result = MiniDumpWriteDump(Process, GetProcessId(Process), FileHandle, MinidumpType, &DumpExceptionInfo, &CrashContextStreamInformation, NULL);
+		const BOOL Result = MiniDumpWriteDump(Process, GetProcessId(Process), FileHandle, MinidumpType, ExceptionInfo ? &DumpExceptionInfo : NULL, &CrashContextStreamInformation, NULL);
 		CloseHandle(FileHandle);
 
 		return Result == TRUE;
