Fix TNonRelocatableInlineAllocator leaking heap allocations

Reported by [REDACTED_PERSON]

diff --git a/Engine/Source/Runtime/Core/Public/Containers/ContainerAllocationPolicies.h b/Engine/Source/Runtime/Core/Public/Containers/ContainerAllocationPolicies.h
--- a/Engine/Source/Runtime/Core/Public/Containers/ContainerAllocationPolicies.h
+++ b/Engine/Source/Runtime/Core/Public/Containers/ContainerAllocationPolicies.h
@@ -819,6 +819,14 @@ public:
 			UE_STATIC_DEPRECATE(5.0, alignof(ElementType) > __STDCPP_DEFAULT_NEW_ALIGNMENT__, "TNonRelocatableInlineAllocator uses GMalloc's default alignment, which is lower than the element type's alignment - please consider a different approach");
 		}
 
+		~ForElementType()
+		{
+			if (HasAllocation())
+			{
+				FMemory::Free(Data);
+			}
+		}
+
 		/**
 		 * Moves the state of another allocator into this one.
 		 * Assumes that the allocator is currently empty, i.e. memory may be allocated but any existing elements have already been destructed (if necessary).
