Fix type promotion errors

diff --git a/Engine/Source/Runtime/Core/Private/HAL/DiskUtilizationTracker.cpp b/Engine/Source/Runtime/Core/Private/HAL/DiskUtilizationTracker.cpp
--- a/Engine/Source/Runtime/Core/Private/HAL/DiskUtilizationTracker.cpp
+++ b/Engine/Source/Runtime/Core/Private/HAL/DiskUtilizationTracker.cpp
@@ -104,11 +104,11 @@ void FDiskUtilizationTracker::FinishRead()
 			if ((ThrottledThroughputBS > 0.0f) && (LongTermStats.GetReadThrougputBS() > ThrottledThroughputBS))
 			{
 				const double IOTime = double(FPlatformTime::Cycles64() - ReadStartCycle) * FPlatformTime::GetSecondsPerCycle64();
-				const double ThrottledIOTime = ((LongTermStats.TotalBytesRead + InFlightBytes) / ThrottledThroughputBS) - LongTermStats.TotalIOTime;
+				const double ThrottledIOTime = (double(LongTermStats.TotalBytesRead + InFlightBytes) / ThrottledThroughputBS) - LongTermStats.TotalIOTime;
 
 				if (IOTime < ThrottledIOTime)
 				{
-					FPlatformProcess::Sleep(ThrottledIOTime - IOTime);
+					FPlatformProcess::Sleep(float(ThrottledIOTime - IOTime));
 				}
 			}
 #endif // !UE_BUILD_SHIPPING
