FSHAHash - Add CompactBinary serialization functions.
#rnx

diff --git a/Engine/Source/Runtime/Core/Private/Misc/SecureHash.cpp b/Engine/Source/Runtime/Core/Private/Misc/SecureHash.cpp
--- a/Engine/Source/Runtime/Core/Private/Misc/SecureHash.cpp
+++ b/Engine/Source/Runtime/Core/Private/Misc/SecureHash.cpp
@@ -1,16 +1,18 @@
 // Copyright Epic Games, Inc. All Rights Reserved.
 
 #include "Misc/SecureHash.h"
-#include "Misc/StringBuilder.h"
+
 #include "HAL/FileManager.h"
 #include "Misc/Paths.h"
+#include "Misc/StringBuilder.h"
+#include "Serialization/CompactBinaryWriter.h"
 
 #if defined(_M_AMD64) || defined(__x86_64__)
 
 #include <immintrin.h>
 #if PLATFORM_COMPILER_CLANG
-#include <shaintrin.h>
 #include <cpuid.h>
+#include <shaintrin.h>
 #endif
 
 #define UE_PLATFORM_SHA_X86 1
@@ -648,6 +650,24 @@ void LexFromString(FSHAHash& InHash, const TCHAR* InString)
 	InHash.FromString(InString);
 }
 
+FCbWriter& operator<<(FCbWriter& Writer, const FSHAHash& Hash)
+{
+	Writer.AddBinary(FMemoryView(&Hash.Hash, sizeof(Hash.Hash)));
+	return Writer;
+}
+
+bool LoadFromCompactBinary(FCbFieldView Field, FSHAHash& OutHash)
+{
+	FMemoryView Bytes = Field.AsBinaryView();
+	if (Bytes.GetSize() != sizeof(OutHash.Hash))
+	{
+		OutHash = FSHAHash();
+		return false;
+	}
+	FMemory::Memcpy(OutHash.Hash, Bytes.GetData(), Bytes.GetSize());
+	return true;
+}
+
 void Freeze::IntrinsicToString(const FSHAHash& Object, const FTypeLayoutDesc& TypeDesc, const FPlatformTypeLayoutParameters& LayoutParams, FMemoryToStringContext& OutContext)
 {
 	OutContext.String->Appendf(TEXT("%s\n"), *Object.ToString());
diff --git a/Engine/Source/Runtime/Core/Public/Misc/SecureHash.h b/Engine/Source/Runtime/Core/Public/Misc/SecureHash.h
--- a/Engine/Source/Runtime/Core/Public/Misc/SecureHash.h
+++ b/Engine/Source/Runtime/Core/Public/Misc/SecureHash.h
@@ -24,6 +24,8 @@
 #include "String/HexToBytes.h"
 #include "Templates/UnrealTemplate.h"
 
+class FCbFieldView;
+class FCbWriter;
 class FMemoryImageWriter;
 class FMemoryUnfreezeContent;
 class FPointerTableBase;
@@ -279,6 +281,8 @@ DECLARE_INTRINSIC_TYPE_LAYOUT(FSHAHash);
 
 inline FStringBuilderBase& operator<<(FStringBuilderBase& Builder, const FSHAHash& Hash) { UE::String::BytesToHex(Hash.Hash, Builder); return Builder; }
 inline FAnsiStringBuilderBase& operator<<(FAnsiStringBuilderBase& Builder, const FSHAHash& Hash) { UE::String::BytesToHex(Hash.Hash, Builder); return Builder; }
+CORE_API FCbWriter& operator<<(FCbWriter& Writer, const FSHAHash& Hash);
+CORE_API bool LoadFromCompactBinary(FCbFieldView Field, FSHAHash& OutHash);
 
 class CORE_API FSHA1
 {
