Removed IWYU pragmas.

diff --git a/Engine/Source/Runtime/Core/Public/Delegates/Delegate.h b/Engine/Source/Runtime/Core/Public/Delegates/Delegate.h
--- a/Engine/Source/Runtime/Core/Public/Delegates/Delegate.h
+++ b/Engine/Source/Runtime/Core/Public/Delegates/Delegate.h
@@ -478,14 +478,12 @@ namespace UE4Delegates_Private
 #define __Delegate_h__
 #define FUNC_INCLUDING_INLINE_IMPL
 
-// IWYU pragma: begin_exports
 #if !UE_BUILD_DOCS
 	#include "Delegates/DelegateInstanceInterface.h"
 	#include "Delegates/DelegateInstancesImpl.h"
 	#include "Delegates/DelegateSignatureImpl.inl"
 	#include "Delegates/DelegateCombinations.h"
 #endif
-// IWYU pragma: end_exports
 
 // No longer allowed to include DelegateSignatureImpl.inl
 #undef FUNC_INCLUDING_INLINE_IMPL
