Fix to ensure platform version of plugin config files are loaded correctly.

diff --git a/Engine/Source/Runtime/Core/Public/Misc/ConfigHierarchy.h b/Engine/Source/Runtime/Core/Public/Misc/ConfigHierarchy.h
--- a/Engine/Source/Runtime/Core/Public/Misc/ConfigHierarchy.h
+++ b/Engine/Source/Runtime/Core/Public/Misc/ConfigHierarchy.h
@@ -91,19 +91,23 @@ static FConfigLayerExpansion GConfigExpansions[] =
 	{
 		TEXT("{ENGINE}/Config/{PLATFORM}/"),	TEXT("{EXTENGINE}/Config/"),	
 		TEXT("{PROJECT}/Config/{PLATFORM}/"),	TEXT("{EXTPROJECT}/Config/"), 
-		EConfigExpansionFlags::ForUncooked | EConfigExpansionFlags::ForCooked 
+		EConfigExpansionFlags::ForUncooked | EConfigExpansionFlags::ForCooked | EConfigExpansionFlags::ForPlugin
 	},
 
 	// Platform Extensions in Restricted Locations
+	// 
+	// Regarding the commented EConfigExpansionFlags::ForPlugin expansions: in the interest of keeping plugin ini scanning fast,
+	// we disable these expansions for plugins because they are not used by Epic, and are unlikely to be used by licensees. If
+	// we can make scanning fast (caching what directories exist, etc), then we could turn this back on to be future-proof.
 	{
 		TEXT("{ENGINE}/Config/{PLATFORM}/"),	TEXT("{ENGINE}/Restricted/NotForLicensees/Platforms/{PLATFORM}/Config/"),	
 		TEXT("{PROJECT}/Config/{PLATFORM}/"),	TEXT("{RESTRICTEDPROJECT_NFL}/Platforms/{PLATFORM}/Config/"), 
-		EConfigExpansionFlags::ForUncooked | EConfigExpansionFlags::ForCooked
+		EConfigExpansionFlags::ForUncooked | EConfigExpansionFlags::ForCooked // | EConfigExpansionFlags::ForPlugin 
 	},
 	{
 		TEXT("{ENGINE}/Config/{PLATFORM}/"),	TEXT("{ENGINE}/Restricted/NoRedist/Platforms/{PLATFORM}/Config/"),			
 		TEXT("{PROJECT}/Config/{PLATFORM}/"),	TEXT("{RESTRICTEDPROJECT_NR}/Platforms/{PLATFORM}/Config/"), 
-		EConfigExpansionFlags::ForUncooked 
+		EConfigExpansionFlags::ForUncooked // | EConfigExpansionFlags::ForPlugin
 	},
 };
 
