TaskGraph - Use a local scratch variable for tasks popped in DispatchSubsequents() instead of reusing a referenced one that's passed in. This was causing a crash because the passed in array was being cleared by another thread before we were finished using them.

diff --git a/Engine/Source/Runtime/Core/Private/Async/TaskGraph.cpp b/Engine/Source/Runtime/Core/Private/Async/TaskGraph.cpp
--- a/Engine/Source/Runtime/Core/Private/Async/TaskGraph.cpp
+++ b/Engine/Source/Runtime/Core/Private/Async/TaskGraph.cpp
@@ -2561,14 +2561,14 @@ void FGraphEvent::DispatchSubsequents(TArray<FBaseGraphTask*>& NewTasks, ENamedT
 	}
 
 	bool bWakeUpWorker = false;
-	SubsequentList.PopAllAndClose(NewTasks);
-	for (int32 Index = NewTasks.Num() - 1; Index >= 0 ; Index--) // reverse the order since PopAll is implicitly backwards
+	TArray<FBaseGraphTask*> PoppedTasks;
+	SubsequentList.PopAllAndClose(PoppedTasks);
+	for (int32 Index = PoppedTasks.Num() - 1; Index >= 0 ; Index--) // reverse the order since PopAll is implicitly backwards
 	{
-		FBaseGraphTask* NewTask = NewTasks[Index];
+		FBaseGraphTask* NewTask = PoppedTasks[Index];
 		checkThreadGraph(NewTask);
 		NewTask->ConditionalQueueTask(CurrentThreadIfKnown, bWakeUpWorker);
 	}
-	NewTasks.Reset();
 
 	TaskTrace::Completed(GetTraceId());
 }
