[Runtime Core]
* TrackedActivity - Moved light enum to outside the class in order to be able to forward declare it

diff --git a/Engine/Source/Runtime/Core/Public/Misc/TrackedActivity.h b/Engine/Source/Runtime/Core/Public/Misc/TrackedActivity.h
--- a/Engine/Source/Runtime/Core/Public/Misc/TrackedActivity.h
+++ b/Engine/Source/Runtime/Core/Public/Misc/TrackedActivity.h
@@ -6,6 +6,19 @@
 #include "Templates/Function.h"
 #include "Templates/SharedPointer.h"
 
+/**
+* Enum to specify status light for an activity
+* Is used by the new console to show a colored dot in front of the tracked activity entry
+*/
+enum class ETrackedActivityLight : uint8
+{
+	None,
+	Red,
+	Yellow,
+	Green,
+	Inherit,
+};
+
 /**
  * Tracked Activity is used to be able to visualize on a semi-high level what is going on in the process.
  * It is very useful when wanting to show the status of online, or if the runtime is waiting on loading something
@@ -15,19 +28,6 @@
 class CORE_API FTrackedActivity : public TSharedFromThis<FTrackedActivity>
 {
 public:
-	/**
-	* Enum to specify status light for an activity
-	* Is used by the new console to show a colored dot in front of the tracked activity entry
-	*/
-	enum class ELight
-	{
-		None,
-		Red,
-		Yellow,
-		Green,
-		Inherit,
-	};
-
 	/**
 	* Enum to specify status type of activity
 	* In the new console, 'Activity' will make sure tracked activity appears to the left of the window
@@ -40,6 +40,8 @@ public:
 		Debug,
 	};
 
+	using ELight = ETrackedActivityLight;
+
 	/** Ctor
 	* @param Name Name of tracked activity
 	* @param Status Initial status of tracked activity
