MemoryView: Fixed CopyFrom to avoid calling Memcpy with a null pointer or size of 0

Migrated the tests to LowLevelTests at the same time.

#rnx

diff --git a/Engine/Source/Runtime/Core/Private/Tests/Memory/MemoryViewTest.cpp b/Engine/Source/Runtime/Core/Private/Tests/Memory/MemoryViewTest.cpp
deleted file mode 100644
diff --git a/Engine/Source/Runtime/Core/Public/Memory/MemoryView.h b/Engine/Source/Runtime/Core/Public/Memory/MemoryView.h
--- a/Engine/Source/Runtime/Core/Public/Memory/MemoryView.h
+++ b/Engine/Source/Runtime/Core/Public/Memory/MemoryView.h
@@ -215,7 +215,10 @@ public:
 	{
 		checkf(InView.Size <= Size, TEXT("Failed to copy from a view of %" UINT64_FMT " bytes "
 			"to a view of %" UINT64_FMT " bytes."), InView.Size, Size);
-		FMemory::Memcpy(Data, InView.Data, InView.Size);
+		if (InView.Size)
+		{
+			FMemory::Memcpy(Data, InView.Data, InView.Size);
+		}
 		return RightChop(InView.Size);
 	}
 
diff --git a/Engine/Source/Runtime/Core/Tests/Memory/MemoryViewTest.cpp b/Engine/Source/Runtime/Core/Tests/Memory/MemoryViewTest.cpp
new file mode 100644
--- /dev/null
+++ b/Engine/Source/Runtime/Core/Tests/Memory/MemoryViewTest.cpp
@@ -0,0 +1,273 @@
+// Copyright Epic Games, Inc. All Rights Reserved.
+
+#if WITH_LOW_LEVEL_TESTS
+
+#include "Memory/MemoryView.h"
+
+#include "Containers/ArrayView.h"
+#include "TestHarness.h"
+
+#include <type_traits>
+
+TEST_CASE("Core::Memory::MemoryView", "[Core][Memory][Smoke]")
+{
+	auto TestMemoryView = [](const FMemoryView& View, const void* Data, uint64 Size)
+	{
+		CHECK(View.GetData() == Data);
+		CHECK(View.GetDataEnd() == static_cast<const uint8*>(Data) + Size);
+		CHECK(View.GetSize() == Size);
+		CHECK(View.IsEmpty() == (Size == 0));
+	};
+
+	auto TestMutableMemoryView = [](const FMutableMemoryView& View, void* Data, uint64 Size)
+	{
+		CHECK(View.GetData() == Data);
+		CHECK(View.GetDataEnd() == static_cast<uint8*>(Data) + Size);
+		CHECK(View.GetSize() == Size);
+		CHECK(View.IsEmpty() == (Size == 0));
+	};
+
+	struct
+	{
+		uint8 BeforeByteArray[4];
+		uint8 ByteArray[16]{};
+		uint8 AfterByteArray[4];
+	} ByteArrayContainer;
+
+	uint8 (&ByteArray)[16] = ByteArrayContainer.ByteArray;
+	uint32 IntArray[12]{};
+
+	SECTION("Constructor/Assignment Availability")
+	{
+		STATIC_CHECK(std::is_trivially_copyable<FMemoryView>::value);
+		STATIC_CHECK(std::is_trivially_copy_constructible<FMemoryView>::value);
+		STATIC_CHECK(std::is_trivially_move_constructible<FMemoryView>::value);
+		STATIC_CHECK(std::is_trivially_copy_assignable<FMemoryView>::value);
+		STATIC_CHECK(std::is_trivially_move_assignable<FMemoryView>::value);
+		STATIC_CHECK(std::is_trivially_destructible<FMemoryView>::value);
+
+		STATIC_CHECK(std::is_trivially_copyable<FMutableMemoryView>::value);
+		STATIC_CHECK(std::is_trivially_copy_constructible<FMutableMemoryView>::value);
+		STATIC_CHECK(std::is_trivially_move_constructible<FMutableMemoryView>::value);
+		STATIC_CHECK(std::is_trivially_copy_assignable<FMutableMemoryView>::value);
+		STATIC_CHECK(std::is_trivially_move_assignable<FMutableMemoryView>::value);
+		STATIC_CHECK(std::is_trivially_destructible<FMutableMemoryView>::value);
+
+		STATIC_CHECK(std::is_constructible<FMemoryView, const FMutableMemoryView&>::value);
+		STATIC_CHECK(!std::is_constructible<FMutableMemoryView, const FMemoryView&>::value);
+		STATIC_CHECK(!std::is_constructible<FMutableMemoryView, const void*, uint64>::value);
+
+		STATIC_CHECK(std::is_assignable<FMemoryView, const FMutableMemoryView&>::value);
+		STATIC_CHECK(!std::is_assignable<FMutableMemoryView, const FMemoryView&>::value);
+	}
+
+	SECTION("Empty Views")
+	{
+		TestMemoryView(FMemoryView(), nullptr, 0);
+		TestMemoryView(FMutableMemoryView(), nullptr, 0);
+		TestMutableMemoryView(FMutableMemoryView(), nullptr, 0);
+	}
+
+	SECTION("Construction from Type[], TArrayView, (Type*, uint64), (Type*, Type*)")
+	{
+		TestMemoryView(MakeMemoryView(AsConst(IntArray)), IntArray, sizeof(IntArray));
+		TestMemoryView(MakeMemoryView(MakeArrayView(AsConst(IntArray))), IntArray, sizeof(IntArray));
+		TestMemoryView(MakeMemoryView(AsConst(IntArray), sizeof(IntArray)), IntArray, sizeof(IntArray));
+		TestMemoryView(MakeMemoryView(AsConst(IntArray), AsConst(IntArray) + 6), IntArray, sizeof(*IntArray) * 6);
+		TestMutableMemoryView(MakeMemoryView(IntArray), IntArray, sizeof(IntArray));
+		TestMutableMemoryView(MakeMemoryView(MakeArrayView(IntArray)), IntArray, sizeof(IntArray));
+		TestMutableMemoryView(MakeMemoryView(IntArray, sizeof(IntArray)), IntArray, sizeof(IntArray));
+		TestMutableMemoryView(MakeMemoryView(IntArray, IntArray + 6), IntArray, sizeof(*IntArray) * 6);
+	}
+
+	SECTION("Construction from std::initializer_list")
+	{
+		//MakeMemoryView({1, 2, 3}); // fail because the type must be deduced
+		std::initializer_list<uint8> InitializerList{1, 2, 3};
+		TestMemoryView(MakeMemoryView(InitializerList), GetData(InitializerList), GetNum(InitializerList) * sizeof(uint8));
+	}
+
+	SECTION("Reset")
+	{
+		FMutableMemoryView View = MakeMemoryView(IntArray);
+		View.Reset();
+		CHECK(View == FMutableMemoryView());
+	}
+
+	SECTION("Left")
+	{
+		STATIC_CHECK(MakeMemoryView(IntArray).Left(0).IsEmpty());
+		STATIC_CHECK(MakeMemoryView(IntArray).Left(1) == MakeMemoryView(IntArray, 1));
+		STATIC_CHECK(MakeMemoryView(IntArray).Left(sizeof(IntArray)) == MakeMemoryView(IntArray));
+		STATIC_CHECK(MakeMemoryView(IntArray).Left(sizeof(IntArray) + 1) == MakeMemoryView(IntArray));
+		STATIC_CHECK(MakeMemoryView(IntArray).Left(MAX_uint64) == MakeMemoryView(IntArray));
+	}
+
+	SECTION("LeftChop")
+	{
+		STATIC_CHECK(MakeMemoryView(IntArray).LeftChop(0) == MakeMemoryView(IntArray));
+		STATIC_CHECK(MakeMemoryView(IntArray).LeftChop(1) == MakeMemoryView(IntArray, sizeof(IntArray) - 1));
+		STATIC_CHECK(MakeMemoryView(IntArray).LeftChop(sizeof(IntArray)).IsEmpty());
+		STATIC_CHECK(MakeMemoryView(IntArray).LeftChop(sizeof(IntArray) + 1).IsEmpty());
+		STATIC_CHECK(MakeMemoryView(IntArray).LeftChop(MAX_uint64).IsEmpty());
+	}
+
+	SECTION("Right")
+	{
+		CHECK(MakeMemoryView(IntArray).Right(0) == FMutableMemoryView());
+		CHECK(MakeMemoryView(IntArray).Right(1) == MakeMemoryView(reinterpret_cast<uint8*>(IntArray) + sizeof(IntArray) - 1, 1));
+		CHECK(MakeMemoryView(IntArray).Right(sizeof(IntArray)) == MakeMemoryView(IntArray));
+		CHECK(MakeMemoryView(IntArray).Right(sizeof(IntArray) + 1) == MakeMemoryView(IntArray));
+		CHECK(MakeMemoryView(IntArray).Right(MAX_uint64) == MakeMemoryView(IntArray));
+	}
+
+	SECTION("RightChop")
+	{
+		CHECK(MakeMemoryView(IntArray).RightChop(0) == MakeMemoryView(IntArray));
+		CHECK(MakeMemoryView(IntArray).RightChop(1) == MakeMemoryView(reinterpret_cast<uint8*>(IntArray) + 1, sizeof(IntArray) - 1));
+		CHECK(MakeMemoryView(IntArray).RightChop(sizeof(IntArray)) == FMutableMemoryView());
+		CHECK(MakeMemoryView(IntArray).RightChop(sizeof(IntArray) + 1) == FMutableMemoryView());
+		CHECK(MakeMemoryView(IntArray).RightChop(MAX_uint64) == FMutableMemoryView());
+	}
+
+	SECTION("Mid")
+	{
+		CHECK(MakeMemoryView(IntArray).Mid(0) == MakeMemoryView(IntArray));
+		CHECK(MakeMemoryView(IntArray).Mid(1) == MakeMemoryView(reinterpret_cast<uint8*>(IntArray) + 1, sizeof(IntArray) - 1));
+		CHECK(MakeMemoryView(IntArray).Mid(sizeof(IntArray)) == FMutableMemoryView());
+		CHECK(MakeMemoryView(IntArray).Mid(sizeof(IntArray) + 1) == FMutableMemoryView());
+		CHECK(MakeMemoryView(IntArray).Mid(MAX_uint64) == FMutableMemoryView());
+		CHECK(MakeMemoryView(IntArray).Mid(0, 0) == FMutableMemoryView());
+		CHECK(MakeMemoryView(IntArray).Mid(0, 1) == MakeMemoryView(IntArray, 1));
+		CHECK(MakeMemoryView(IntArray).Mid(1, sizeof(IntArray) - 2) == MakeMemoryView(reinterpret_cast<uint8*>(IntArray) + 1, sizeof(IntArray) - 2));
+		CHECK(MakeMemoryView(IntArray).Mid(1, sizeof(IntArray) - 1) == MakeMemoryView(reinterpret_cast<uint8*>(IntArray) + 1, sizeof(IntArray) - 1));
+		CHECK(MakeMemoryView(IntArray).Mid(1, sizeof(IntArray)) == MakeMemoryView(reinterpret_cast<uint8*>(IntArray) + 1, sizeof(IntArray) - 1));
+		CHECK(MakeMemoryView(IntArray).Mid(0, MAX_uint64) == MakeMemoryView(IntArray));
+		CHECK(MakeMemoryView(IntArray).Mid(MAX_uint64, MAX_uint64) == FMutableMemoryView());
+	}
+
+	SECTION("Contains")
+	{
+		CHECK(FMemoryView().Contains(FMutableMemoryView()));
+		CHECK(FMutableMemoryView().Contains(FMemoryView()));
+		CHECK(MakeMemoryView(ByteArray).Contains(MakeMemoryView(ByteArray)));
+		CHECK(MakeMemoryView(ByteArray).Contains(MakeMemoryView(ByteArray + 1, 15)));
+		CHECK(MakeMemoryView(ByteArray).Contains(MakeMemoryView(ByteArray, 15)));
+		CHECK(MakeMemoryView(ByteArray).Contains(MakeMemoryView(ByteArray + 1, 14)));
+		CHECK(MakeMemoryView(ByteArray).Contains(MakeMemoryView(ByteArray, 0)));
+		CHECK(MakeMemoryView(ByteArray).Contains(MakeMemoryView(ByteArray + 8, 0)));
+		CHECK(MakeMemoryView(ByteArray).Contains(MakeMemoryView(ByteArray + 16, 0)));
+		CHECK(!MakeMemoryView(ByteArray).Contains(MakeMemoryView(ByteArrayContainer.BeforeByteArray + 3, 0)));
+		CHECK(!MakeMemoryView(ByteArray).Contains(MakeMemoryView(ByteArrayContainer.AfterByteArray + 1, 0)));
+		CHECK(!MakeMemoryView(ByteArray).Contains(MakeMemoryView(ByteArrayContainer.BeforeByteArray + 3, 1)));
+		CHECK(!MakeMemoryView(ByteArray).Contains(MakeMemoryView(ByteArray + 16, 1)));
+		CHECK(!MakeMemoryView(ByteArray).Contains(MakeMemoryView(ByteArrayContainer.BeforeByteArray + 3, 17)));
+		CHECK(!MakeMemoryView(ByteArray).Contains(MakeMemoryView(ByteArray, 17)));
+		CHECK(!MakeMemoryView(ByteArray).Contains(MakeMemoryView(ByteArrayContainer.BeforeByteArray + 3, 18)));
+		CHECK(!MakeMemoryView(ByteArray).Contains(MakeMemoryView(ByteArrayContainer.BeforeByteArray + 3, 8)));
+		CHECK(!MakeMemoryView(ByteArray).Contains(MakeMemoryView(ByteArray + 9, 8)));
+	}
+
+	SECTION("Intersects")
+	{
+		CHECK(MakeMemoryView(ByteArray).Intersects(MakeMemoryView(ByteArray)));
+		CHECK(MakeMemoryView(ByteArray).Intersects(MakeMemoryView(ByteArray + 1, 15)));
+		CHECK(MakeMemoryView(ByteArray).Intersects(MakeMemoryView(ByteArray, 15)));
+		CHECK(MakeMemoryView(ByteArray).Intersects(MakeMemoryView(ByteArray + 1, 14)));
+		CHECK(MakeMemoryView(ByteArray).Intersects(MakeMemoryView(ByteArrayContainer.BeforeByteArray + 3, 8)));
+		CHECK(MakeMemoryView(ByteArray).Intersects(MakeMemoryView(ByteArray + 9, 8)));
+		CHECK(MakeMemoryView(ByteArray).Intersects(MakeMemoryView(ByteArrayContainer.BeforeByteArray + 3, 17)));
+		CHECK(MakeMemoryView(ByteArray).Intersects(MakeMemoryView(ByteArray, 17)));
+		CHECK(MakeMemoryView(ByteArray).Intersects(MakeMemoryView(ByteArrayContainer.BeforeByteArray + 3, 18)));
+		CHECK(MakeMemoryView(ByteArray).Intersects(MakeMemoryView(ByteArray + 8, 0)));
+		CHECK(!FMemoryView().Intersects(FMutableMemoryView()));
+		CHECK(!FMutableMemoryView().Intersects(FMemoryView()));
+		CHECK(!MakeMemoryView(ByteArray).Intersects(MakeMemoryView(ByteArray, 0)));
+		CHECK(!MakeMemoryView(ByteArray).Intersects(MakeMemoryView(ByteArray + 16, 0)));
+		CHECK(!MakeMemoryView(ByteArray).Intersects(MakeMemoryView(ByteArrayContainer.BeforeByteArray + 3, 0)));
+		CHECK(!MakeMemoryView(ByteArray).Intersects(MakeMemoryView(ByteArrayContainer.AfterByteArray + 1, 0)));
+		CHECK(!MakeMemoryView(ByteArray).Intersects(MakeMemoryView(ByteArrayContainer.BeforeByteArray + 3, 1)));
+		CHECK(!MakeMemoryView(ByteArray).Intersects(MakeMemoryView(ByteArray + 16, 1)));
+	}
+
+	SECTION("CompareBytes")
+	{
+		const uint8 CompareBytes[8] = { 5, 4, 6, 2, 4, 7, 1, 3 };
+		CHECK(FMemoryView().CompareBytes(FMutableMemoryView()) == 0);
+		CHECK(FMutableMemoryView().CompareBytes(FMemoryView()) == 0);
+		CHECK(MakeMemoryView(ByteArray).CompareBytes(MakeMemoryView(ByteArray)) == 0);
+		CHECK(MakeMemoryView(ByteArray, 8).CompareBytes(MakeMemoryView(ByteArray + 8, 8)) == 0);
+		CHECK(FMemoryView().CompareBytes(MakeMemoryView(ByteArray)) < 0);
+		CHECK(MakeMemoryView(ByteArray).CompareBytes(FMemoryView()) > 0);
+		CHECK(MakeMemoryView(ByteArray, 8).CompareBytes(MakeMemoryView(ByteArray)) < 0);
+		CHECK(MakeMemoryView(ByteArray).CompareBytes(MakeMemoryView(ByteArray, 8)) > 0);
+		CHECK(MakeMemoryView(IntArray, 8).CompareBytes(MakeMemoryView(ByteArray)) < 0);
+		CHECK(MakeMemoryView(ByteArray).CompareBytes(MakeMemoryView(IntArray, 8)) > 0);
+		CHECK(MakeMemoryView(ByteArray, 4).CompareBytes(MakeMemoryView(ByteArray, 8)) < 0);
+		CHECK(MakeMemoryView(ByteArray, 8).CompareBytes(MakeMemoryView(ByteArray, 4)) > 0);
+		CHECK(MakeMemoryView(CompareBytes, 2).CompareBytes(MakeMemoryView(CompareBytes + 2, 2)) < 0);
+		CHECK(MakeMemoryView(CompareBytes, 3).CompareBytes(MakeMemoryView(CompareBytes + 3, 3)) > 0);
+	}
+
+	SECTION("EqualBytes")
+	{
+		const uint8 CompareBytes[8] = { 5, 4, 6, 2, 4, 7, 1, 3 };
+		CHECK(FMemoryView().EqualBytes(FMutableMemoryView()));
+		CHECK(FMutableMemoryView().EqualBytes(FMemoryView()));
+		CHECK(MakeMemoryView(ByteArray).EqualBytes(MakeMemoryView(ByteArray)));
+		CHECK(MakeMemoryView(ByteArray, 8).EqualBytes(MakeMemoryView(ByteArray + 8, 8)));
+		CHECK(!MakeMemoryView(ByteArray, 8).EqualBytes(MakeMemoryView(ByteArray, 4)));
+		CHECK(!MakeMemoryView(ByteArray, 4).EqualBytes(MakeMemoryView(ByteArray, 8)));
+		CHECK(!MakeMemoryView(CompareBytes, 4).EqualBytes(MakeMemoryView(CompareBytes + 4, 4)));
+	}
+
+	SECTION("Equals")
+	{
+		CHECK(FMemoryView().Equals(FMemoryView()));
+		CHECK(FMemoryView().Equals(FMutableMemoryView()));
+		CHECK(FMutableMemoryView().Equals(FMemoryView()));
+		CHECK(FMutableMemoryView().Equals(FMutableMemoryView()));
+		CHECK(MakeMemoryView(IntArray).Equals(MakeMemoryView(AsConst(IntArray))));
+		CHECK(!MakeMemoryView(IntArray).Equals(MakeMemoryView(IntArray + 1, sizeof(IntArray) - sizeof(*IntArray))));
+		CHECK(!MakeMemoryView(IntArray).Equals(MakeMemoryView(IntArray, sizeof(*IntArray))));
+		CHECK(!MakeMemoryView(IntArray).Equals(FMutableMemoryView()));
+	}
+
+	SECTION("operator==")
+	{
+		STATIC_CHECK(MakeMemoryView(ByteArrayContainer.ByteArray) == MakeMemoryView(ByteArrayContainer.ByteArray)); //-V501
+		STATIC_CHECK(MakeMemoryView(ByteArrayContainer.ByteArray) == MakeMemoryView(AsConst(ByteArrayContainer.ByteArray)));
+		STATIC_CHECK(MakeMemoryView(AsConst(ByteArrayContainer.ByteArray)) == MakeMemoryView(ByteArrayContainer.ByteArray));
+		STATIC_CHECK(MakeMemoryView(AsConst(ByteArrayContainer.ByteArray)) == MakeMemoryView(AsConst(ByteArrayContainer.ByteArray))); //-V501
+	}
+
+	SECTION("operator!=")
+	{
+		STATIC_CHECK(MakeMemoryView(ByteArrayContainer.ByteArray) != MakeMemoryView(IntArray));
+		STATIC_CHECK(MakeMemoryView(ByteArrayContainer.ByteArray) != MakeMemoryView(AsConst(IntArray)));
+		STATIC_CHECK(MakeMemoryView(AsConst(ByteArrayContainer.ByteArray)) != MakeMemoryView(IntArray));
+		STATIC_CHECK(MakeMemoryView(AsConst(ByteArrayContainer.ByteArray)) != MakeMemoryView(AsConst(IntArray)));
+	}
+
+	SECTION("operator+=")
+	{
+		CHECK((MakeMemoryView(ByteArray) += 0) == MakeMemoryView(ByteArray));
+		CHECK((MakeMemoryView(ByteArray) += 8) == MakeMemoryView(ByteArray + 8, 8));
+		CHECK((MakeMemoryView(ByteArray) += 16) == MakeMemoryView(ByteArray + 16, 0));
+		CHECK((MakeMemoryView(ByteArray) += 32) == MakeMemoryView(ByteArray + 16, 0));
+	}
+
+	SECTION("operator+")
+	{
+		CHECK((MakeMemoryView(ByteArray) + 0) == MakeMemoryView(ByteArray));
+		CHECK((0 + MakeMemoryView(ByteArray)) == MakeMemoryView(ByteArray));
+		CHECK((MakeMemoryView(ByteArray) + 8) == MakeMemoryView(ByteArray + 8, 8));
+		CHECK((8 + MakeMemoryView(ByteArray)) == MakeMemoryView(ByteArray + 8, 8));
+		CHECK((MakeMemoryView(ByteArray) + 16) == MakeMemoryView(ByteArray + 16, 0));
+		CHECK((16 + MakeMemoryView(ByteArray)) == MakeMemoryView(ByteArray + 16, 0));
+		CHECK((MakeMemoryView(ByteArray) + 32) == MakeMemoryView(ByteArray + 16, 0));
+		CHECK((32 + MakeMemoryView(ByteArray)) == MakeMemoryView(ByteArray + 16, 0));
+	}
+}
+
+#endif // WITH_LOW_LEVEL_TESTS
