UE-154432
Fix crash handler for Apple Silicon
The infinite loop workaround is no longer needed. Arm64 is now capable of generating the same level of crash report as x86_64 counterpart.

diff --git a/Engine/Source/Runtime/Core/Private/Mac/MacErrorOutputDevice.cpp b/Engine/Source/Runtime/Core/Private/Mac/MacErrorOutputDevice.cpp
--- a/Engine/Source/Runtime/Core/Private/Mac/MacErrorOutputDevice.cpp
+++ b/Engine/Source/Runtime/Core/Private/Mac/MacErrorOutputDevice.cpp
@@ -98,16 +98,6 @@ void FMacErrorOutputDevice::HandleError()
 	UE_LOG(LogMac, Log, TEXT("%s"), GErrorHist);
     
     GLog->Panic();
-    
-#if PLATFORM_MAC_ARM64
-    // stack dumping seems broken on arm, so for the time being spin so we can dump callstacks
-    // via activity monitor to find broken things.
-    while (true)
-    {
-        FPlatformMisc::LowLevelOutputDebugString(TEXT("Spinning after fatal error.."));
-        FPlatformProcess::Sleep( 1.0 );
-    }
-#endif
 
 	HandleErrorRestoreUI();
 
