Mac: Apply UE_RELATIVE_BASE_DIR to BaseDir for console apps.
This will ensure correct Engine/Binaries/Mac path is returned for these types of apps.

diff --git a/Engine/Source/Runtime/Core/Private/Mac/MacPlatformProcess.cpp b/Engine/Source/Runtime/Core/Private/Mac/MacPlatformProcess.cpp
--- a/Engine/Source/Runtime/Core/Private/Mac/MacPlatformProcess.cpp
+++ b/Engine/Source/Runtime/Core/Private/Mac/MacPlatformProcess.cpp
@@ -1041,6 +1041,12 @@ const TCHAR* FMacPlatformProcess::BaseDir()
 					BasePath = [BasePath stringByDeletingLastPathComponent];
 				}
 			}
+#ifdef UE_RELATIVE_BASE_DIR
+			else
+			{
+				BasePath = [BasePath stringByAppendingPathComponent : @UE_RELATIVE_BASE_DIR];
+			}
+#endif
 
 			FCString::Strcpy(Result, MAC_MAX_PATH, *FString(BasePath));
 			FCString::Strcat(Result, TEXT("/"));
