Disable the altstack for the thread before freeing it. If not can leave the kernel with a dangling pointer to a memory region that is no longer valid

Also add protection pages to the left and right of our altstack to catch buffer over/under flows

diff --git a/Engine/Source/Runtime/Core/Private/Unix/UnixPlatformCrashContext.cpp b/Engine/Source/Runtime/Core/Private/Unix/UnixPlatformCrashContext.cpp
--- a/Engine/Source/Runtime/Core/Private/Unix/UnixPlatformCrashContext.cpp
+++ b/Engine/Source/Runtime/Core/Private/Unix/UnixPlatformCrashContext.cpp
@@ -1055,16 +1055,34 @@ void *FRunnableThreadUnix::MainThreadSignalHandlerStack = nullptr;
 
 void *FRunnableThreadUnix::AllocCrashHandlerStack()
 {
+	SIZE_T PageSize = FPlatformMemory::GetConstants().PageSize;
 	uint64 StackBufferSize = GetCrashHandlerStackSize();
 
-	return mmap(nullptr, StackBufferSize, PROT_READ | PROT_WRITE, MAP_PRIVATE | MAP_ANON, -1, 0);
+	// grab two extra pages to protect the left/right boundary of this stack
+	void* Ptr = mmap(nullptr, StackBufferSize + PageSize * 2, PROT_READ | PROT_WRITE, MAP_PRIVATE | MAP_ANON, -1, 0);
+
+	// protect the left most page, and the right most page in case of a buffer under/over flow
+	mprotect(Ptr, PageSize, PROT_NONE);
+	mprotect(reinterpret_cast<uint8*>(Ptr) + StackBufferSize + PageSize, PageSize, PROT_NONE);
+
+	return reinterpret_cast<uint8*>(Ptr) + PageSize;
 }
 
 void FRunnableThreadUnix::FreeCrashHandlerStack(void *StackBuffer)
 {
 	if (StackBuffer)
 	{
-		munmap(StackBuffer, GetCrashHandlerStackSize());
+		SIZE_T PageSize = FPlatformMemory::GetConstants().PageSize;
+		uint64 StackTraceSize = GetCrashHandlerStackSize();
+
+		// disable our current altstack so the kernel is not left with a dangling pointer, so we can then free the memory
+		stack_t CurrentSignalHandlerStack;
+		FMemory::Memzero(CurrentSignalHandlerStack);
+		CurrentSignalHandlerStack.ss_flags = SS_DISABLE;
+		sigaltstack(&CurrentSignalHandlerStack, nullptr);
+
+		// we added an extra PageSize when allocating in ::AllocCrashHandlerStack, lets return back to the start here
+		munmap(reinterpret_cast<uint8*>(StackBuffer) - PageSize, GetCrashHandlerStackSize() + PageSize * 2);
 	}
 }
 
@@ -1081,6 +1099,9 @@ uint64 FRunnableThreadUnix::GetCrashHandlerStackSize()
 	{
 		GCrashHandlerStackSize = EConstants::CrashHandlerStackSizeMin;
 	}
+
+	check(IsAligned(GCrashHandlerStackSize, FPlatformMemory::GetConstants().PageSize));
+
 	return GCrashHandlerStackSize;
 }
 
