CreateProfileFilename will now prefix the current process id to the filename when WITH_SERVER_CODE

diff --git a/Engine/Source/Runtime/Core/Private/ProfilingDebugging/ProfilingHelpers.cpp b/Engine/Source/Runtime/Core/Private/ProfilingDebugging/ProfilingHelpers.cpp
--- a/Engine/Source/Runtime/Core/Private/ProfilingDebugging/ProfilingHelpers.cpp
+++ b/Engine/Source/Runtime/Core/Private/ProfilingDebugging/ProfilingHelpers.cpp
@@ -10,6 +10,7 @@
 #include "ProfilingDebugging/ProfilingHelpers.h"
 #include "Misc/DateTime.h"
 #include "HAL/FileManager.h"
+#include "HAL/PlatformProcess.h"
 #include "Misc/Parse.h"
 #include "Misc/CommandLine.h"
 #include "Misc/Paths.h"
@@ -190,7 +191,12 @@ FString CreateProfileFilename( const FString& InFilename, const FString& InFileE
 	}
 	NameOfProfile.RightInline(MaxFilenameLen, false);
 
-	FString FileNameWithExtension = FString::Printf( TEXT("%s%s"), *NameOfProfile, *InFileExtension );
+#if WITH_SERVER_CODE
+	FString FileNameWithExtension = FString::Printf(TEXT("Pid%d_%s%s"), FPlatformProcess::GetCurrentProcessId(), *NameOfProfile, *InFileExtension);
+#else
+	FString FileNameWithExtension = FString::Printf(TEXT("%s%s"), *NameOfProfile, *InFileExtension);
+#endif
+
 	FileNameWithExtension.RightInline(MaxFilenameLen, false);
 
 	FString Filename;
