* Added TPagedArray: a fixed-block chained array container similar to TChunkedArray. The main difference with TChunkedArray is that TPagedArray doesn't default construct its page elements, they are allocated and constructed as requested.
* Added System::Core::Containers::TPagedArray unit tests (26) to LowLevelUnitTests.
* Added debug visualizer.

#tests passed TPagedArray unit tests.

diff --git a/Engine/Source/Runtime/Core/Public/Containers/PagedArray.h b/Engine/Source/Runtime/Core/Public/Containers/PagedArray.h
new file mode 100644
--- /dev/null
+++ b/Engine/Source/Runtime/Core/Public/Containers/PagedArray.h
@@ -0,0 +1,641 @@
+// Copyright Epic Games, Inc. All Rights Reserved.
+
+#pragma once
+
+#include "Containers/Array.h"
+#include "CoreTypes.h"
+#include "GenericPlatform/GenericPlatformMath.h"
+#include "IteratorAdapter.h"
+#include "Math/UnrealMathUtility.h"
+#include "Misc/AssertionMacros.h"
+#include "Templates/UnrealTemplate.h"  // For GetData, GetNum
+
+#include <cstring>
+#include <initializer_list>
+#include <type_traits>
+
+template <typename InElementType, int32 InPageSizeInBytes = 16384, typename InAllocatorType = FDefaultAllocator>
+class TPagedArray;
+
+namespace UE::PagedArray::Private
+{
+
+/**
+ * Page traits struct for TPagedArray.
+ */
+template <typename InElementType, int32 InPageSizeInBytes>
+struct TPageTraits
+{
+	static_assert(InPageSizeInBytes >= sizeof(InElementType), "Page size must be greater or equal than element type's");
+
+	using ElementType = InElementType;
+
+	static constexpr int32 Size = InPageSizeInBytes;
+	static constexpr int32 Capacity = Size / sizeof(ElementType);
+};
+
+/**
+ * TPagedArray iterator base class for TIteratorAdapter.
+ */
+template <typename InElementType, typename InPageType, typename InPageTraits>
+class TIteratorBase
+{
+public:
+	using ElementType = InElementType;
+	using PageType = InPageType;
+	using PageTraits = InPageTraits;
+	using SizeType = typename PageType::SizeType;
+
+	TIteratorBase() = default;
+
+protected:
+	/**
+	 * Internal parameter constructor (used exclusively by the container).
+	 */
+	TIteratorBase(InPageType* Data, SizeType Offset)
+		: Data(Data), PageIndex(Offset / PageTraits::Capacity), PageOffset(Offset % PageTraits::Capacity)
+	{
+	}
+
+	FORCEINLINE ElementType& Dereference() const
+	{
+		return Data[PageIndex][PageOffset];
+	}
+
+	FORCEINLINE void Increment()
+	{
+		if (++PageOffset == PageTraits::Capacity)
+		{
+			PageOffset = 0;
+			++PageIndex;
+		}
+	}
+
+	FORCEINLINE bool Equals(const TIteratorBase& Other) const
+	{
+		return Data == Other.Data && PageIndex == Other.PageIndex && PageOffset == Other.PageOffset;
+	}
+
+private:
+	PageType* Data = nullptr;
+	SizeType PageIndex = 0;
+	SizeType PageOffset = 0;
+};
+
+template <typename InElementType, typename InPageType, typename InPageTraits>
+using TIterator = TIteratorAdapter<TIteratorBase<InElementType, InPageType, InPageTraits>>;
+
+}  // namespace UE::PagedArray::Private
+
+/**
+ * Fixed size block allocated container class.
+ *
+ * This container mimics the behavior of an array without fulfilling the contiguous storage guarantee. Instead, it
+ * allocates memory in pages. Each page is a block of memory of the parameter page size in bytes. Elements within a page
+ * are guaranteed to be contiguous. Paged arrays have the following advantages v normal arrays:
+ * 1. It's more suited for large sized arrays as they don't require a single contiguous memory block. This make them
+ *    less vulnerable to memory fragmentation.
+ * 2. The contained elements' addresses are persistent as long as the container isn't resized.
+ **/
+template <typename InElementType, int32 InPageSizeInBytes, typename InAllocatorType>
+class TPagedArray
+{
+	template <typename AnyElementType, int32 AnyPageSizeInBytes, typename AnyAllocatorType>
+	friend class TPagedArray;
+
+	using PageType = TArray<InElementType, InAllocatorType>;
+	using PageTraits = UE::PagedArray::Private::TPageTraits<InElementType, InPageSizeInBytes>;
+
+public:
+	using AllocatorType = InAllocatorType;
+	using SizeType = typename InAllocatorType::SizeType;
+	using ElementType = InElementType;
+	using ElementAllocatorType = typename TChooseClass<
+		AllocatorType::NeedsElementType,
+		typename AllocatorType::template ForElementType<ElementType>,
+		typename AllocatorType::ForAnyElementType>::Result;
+
+private:
+	static constexpr SizeType GetPageIndex(SizeType Index)
+	{
+		return Index / PageTraits::Capacity;
+	}
+
+	static constexpr SizeType GetPageOffset(SizeType Index)
+	{
+		return Index % PageTraits::Capacity;
+	}
+
+	static constexpr SizeType NumRequiredPages(SizeType Count)
+	{
+		const SizeType Div = Count / PageTraits::Capacity;
+		return Count % PageTraits::Capacity == 0 ? Div : Div + 1;
+	}
+
+public:
+	using ConstIteratorType = UE::PagedArray::Private::TIterator<const ElementType, const PageType, PageTraits>;
+	using IteratorType = UE::PagedArray::Private::TIterator<ElementType, PageType, PageTraits>;
+
+	static constexpr SizeType MaxPerPage()
+	{
+		return PageTraits::Capacity;
+	}
+
+	TPagedArray() = default;
+
+	TPagedArray(TPagedArray&& Other) : Pages(MoveTemp(Other.Pages)), Count(Other.Count)
+	{
+		Other.Count = 0;
+	}
+
+	TPagedArray(const TPagedArray& Other) = default;
+
+	TPagedArray(std::initializer_list<ElementType> InList)
+	{
+		CopyToEmpty(InList.begin(), static_cast<SizeType>(InList.size()));
+	}
+
+	TPagedArray(const ElementType* InSource, SizeType InSize)
+	{
+		check(InSource || !InSize);
+		CopyToEmpty(InSource, InSize);
+	}
+
+	~TPagedArray() = default;
+
+	TPagedArray& operator=(TPagedArray&& Other)
+	{
+		if (this != &Other)
+		{
+			Pages = MoveTemp(Other.Pages);
+			Count = Other.Count;
+			Other.Count = 0;
+		}
+		return *this;
+	}
+
+	TPagedArray& operator=(const TPagedArray& Other) = default;
+
+	TPagedArray& operator=(std::initializer_list<ElementType> InList)
+	{
+		Reset();
+		CopyToEmpty(InList.begin(), static_cast<SizeType>(InList.size()));
+		return *this;
+	}
+
+	FORCEINLINE const ElementType& operator[](SizeType Index) const
+	{
+		CheckValidIndex(Index);
+		return Pages[GetPageIndex(Index)][GetPageOffset(Index)];
+	}
+
+	FORCEINLINE ElementType& operator[](SizeType Index)
+	{
+		CheckValidIndex(Index);
+		return Pages[GetPageIndex(Index)][GetPageOffset(Index)];
+	}
+
+	FORCEINLINE SizeType Max() const
+	{
+		return Pages.Num() * PageTraits::Capacity;
+	}
+
+	FORCEINLINE SizeType Num() const
+	{
+		return Count;
+	}
+
+	FORCEINLINE SizeType NumPages() const
+	{
+		return Pages.Num();
+	}
+
+	FORCEINLINE bool IsEmpty() const
+	{
+		return !Count;
+	}
+
+	FORCEINLINE const ElementType& Last() const
+	{
+		CheckValidIndex(0);
+		const SizeType LastIndex = Num() - 1;
+		return Pages[GetPageIndex(LastIndex)][GetPageOffset(LastIndex)];
+	}
+
+	FORCEINLINE ElementType& Last()
+	{
+		CheckValidIndex(0);
+		const SizeType LastIndex = Num() - 1;
+		return Pages[GetPageIndex(LastIndex)][GetPageOffset(LastIndex)];
+	}
+
+	/**
+	 * Resizes array to the parameter number of elements.
+	 * The allow shrinking parameter indicates whether the container page allocation can be reduced if possible.
+	 */
+	void SetNum(SizeType NewNum, bool bAllowShrinking = true)
+	{
+		const SizeType RequiredPageCount = NumRequiredPages(NewNum);
+		if (NewNum > Num())
+		{
+			Pages.SetNum(RequiredPageCount);
+			SizeType PageIndex = 0;
+			SizeType PendingCount = NewNum;
+			while (PendingCount >= PageTraits::Capacity)
+			{
+				Pages[PageIndex++].SetNum(PageTraits::Capacity);
+				PendingCount -= PageTraits::Capacity;
+			}
+			if (PendingCount)
+			{
+				Pages[PageIndex].Reserve(PageTraits::Capacity);
+				Pages[PageIndex].SetNum(PendingCount);
+			}
+		}
+		else if (NewNum < Num())
+		{
+			SizeType PendingCount = Num() - NewNum;
+			Pages.SetNum(RequiredPageCount, bAllowShrinking);
+			if (const SizeType Mod = NewNum % PageTraits::Capacity)
+			{
+				Pages.Last().SetNum(Mod);
+			}
+		}
+		Count = NewNum;
+	}
+
+	/*
+	 * Fills the memory used by the container elements with the parameter byte value.
+	 */
+	void SetMem(uint8 ByteValue)
+	{
+		for (PageType& Page : Pages)
+		{
+			if (Page.IsEmpty())
+			{
+				break;
+			}
+			std::memset(Page.GetData(), ByteValue, Page.Num() * sizeof(ElementType));
+		}
+	}
+
+	/*
+	 * Sets the memory used by the container elements to zero.
+	 */
+	FORCEINLINE void SetZero()
+	{
+		SetMem(0);
+	}
+
+	/*
+	 * Appends the parameter contiguous range to this container.
+	 */
+	void Append(const ElementType* InSource, SizeType InSize)
+	{
+		check(InSource || !InSize);
+		Grow(Count + InSize);
+		CopyUnchecked(InSource, InSize);
+	}
+
+	template <typename ContainerType, std::enable_if_t<TIsContiguousContainer<ContainerType>::Value>* = nullptr>
+	FORCEINLINE void Append(ContainerType&& Container)
+	{
+		Append(GetData(Container), GetNum(Container));
+	}
+
+	/*
+	 * Appends a compatible paged array to this container.
+	 * Note TPagedArray doesn't meet TIsContiguousContainer traits so it won't use the contiguous container overload.
+	 */
+	template <int32 OtherPageSizeInBytes, typename OtherAllocator>
+	void Append(const TPagedArray<ElementType, OtherPageSizeInBytes, OtherAllocator>& Other)
+	{
+		using OtherType = TPagedArray<ElementType, OtherPageSizeInBytes, OtherAllocator>;
+		using OtherPageType = typename OtherType::PageType;
+		Grow(Count + Other.Count);
+		for (const OtherPageType& Page : Other.Pages)
+		{
+			CopyUnchecked(Page.GetData(), Page.Num());
+		}
+	}
+
+	/**
+	 * Assigns the parameter contiguous range to this container.
+	 * Any pre-existing content is removed.
+	 */
+	void Assign(const ElementType* InSource, SizeType InSize)
+	{
+		Reset();
+		Append(InSource, InSize);
+	}
+
+	template <typename ContainerType, std::enable_if_t<TIsContiguousContainer<ContainerType>::Value>* = nullptr>
+	void Assign(ContainerType&& Container)
+	{
+		Reset();
+		Append(GetData(Container), GetNum(Container));
+	}
+
+	/**
+	 * Assigns a compatible paged array to this container.
+	 * Any pre-existing content is removed.
+	 * Note TPagedArray doesn't meet TIsContiguousContainer traits so it won't use the contiguous container overload.
+	 */
+	template <int32 OtherPageSizeInBytes, typename OtherAllocator>
+	void Assign(const TPagedArray<ElementType, OtherPageSizeInBytes, OtherAllocator>& Other)
+	{
+		Reset();
+		Append(Other);
+	}
+
+	/*
+	 * Constructs an element in place using the parameter arguments and adds it at the back of the container.
+	 * This method returns a reference to the constructed element.
+	 */
+	template <typename... ArgsType>
+	ElementType& Emplace(ArgsType&&... Args)
+	{
+		GrowIfRequired();
+		ElementType& Result = Pages.Last().Emplace_GetRef(Forward<ArgsType>(Args)...);
+		++Count;
+		return Result;
+	}
+
+	/*
+	 * Adds the parameter element at the back of the container.
+	 */
+	FORCEINLINE void Add(const ElementType& Element)
+	{
+		Emplace(Element);
+	}
+
+	FORCEINLINE void Add(ElementType&& Element)
+	{
+		Emplace(MoveTempIfPossible(Element));
+	}
+
+	FORCEINLINE void Push(const ElementType& Element)
+	{
+		Emplace(Element);
+	}
+
+	FORCEINLINE void Push(ElementType&& Element)
+	{
+		Emplace(MoveTempIfPossible(Element));
+	}
+
+	/*
+	 * Removes the last element in the container.
+	 */
+	void Pop(bool bAllowShrinking = true)
+	{
+		CheckValidIndex(0);
+		const SizeType LastIndex = Num() - 1;
+		const SizeType LastPageIndex = GetPageIndex(LastIndex);
+		const SizeType LastIndexInPage = GetPageOffset(LastIndex);
+		Pages[LastPageIndex].RemoveAt(LastIndexInPage, 1, false);
+		if (bAllowShrinking && LastIndexInPage == 0)
+		{
+			Pages.SetNum(LastPageIndex);
+		}
+		--Count;
+	}
+
+	/**
+	 * Removes the element at the parameter index position and swaps the last element if existent to the same position
+	 * to ensure the range is contiguous. This method provides efficient removal O(1) but it doesn't preserve the
+	 * insertion order.
+	 */
+	void RemoveAtSwap(SizeType Index, bool bAllowShrinking = true)
+	{
+		CheckValidIndex(Index);
+		const SizeType TargetPageIndex = GetPageIndex(Index);
+		const SizeType TargetIndexInPage = GetPageOffset(Index);
+		const SizeType LastIndex = Num() - 1;
+		const SizeType LastPageIndex = GetPageIndex(LastIndex);
+		if (TargetPageIndex == LastPageIndex)
+		{
+			Pages[TargetPageIndex].RemoveAtSwap(TargetIndexInPage, 1, false);
+			if (bAllowShrinking && Pages[TargetPageIndex].IsEmpty())
+			{
+				Pages.SetNum(TargetPageIndex);
+			}
+		}
+		else
+		{
+			const SizeType LastIndexInPage = GetPageOffset(LastIndex);
+			Pages[TargetPageIndex][TargetIndexInPage] = MoveTempIfPossible(Pages[LastPageIndex][LastIndexInPage]);
+			Pages[LastPageIndex].RemoveAt(LastIndexInPage, 1, false);
+		}
+		--Count;
+	}
+
+	/*
+	 * Empties the container effectively destroying all contained elements and releases any acquired storage.
+	 */
+	void Empty()
+	{
+		Pages.Empty();
+		Count = 0;
+	}
+
+	/*
+	 * Destroys all contained elements but doesn't release the container's storage.
+	 */
+	void Reset()
+	{
+		for (PageType& Page : Pages)
+		{
+			Page.Reset();
+		}
+		Count = 0;
+	}
+
+	/*
+	 * Reserves storage for the parameter element count.
+	 */
+	void Reserve(SizeType InCount)
+	{
+		if (Max() < InCount)
+		{
+			Grow(InCount);
+		}
+	}
+
+	/*
+	 * Copies this container elements into the parameter destination array.
+	 */
+	template <typename AnyAllocator>
+	void ToArray(TArray<ElementType, AnyAllocator>& OutDestination) const&
+	{
+		OutDestination.Reset();
+		OutDestination.Reserve(Count);
+		for (const PageType& Page : Pages)
+		{
+			if (Page.IsEmpty())
+			{
+				break;
+			}
+			OutDestination.Append(Page);
+		}
+	}
+
+	/*
+	 * Moves this container elements into the parameter destination array.
+	 */
+	template <typename AnyAllocator>
+	void ToArray(TArray<ElementType, AnyAllocator>& OutDestination) &&
+	{
+		OutDestination.Reset();
+		OutDestination.Reserve(Count);
+		for (PageType& Page : Pages)
+		{
+			if (Page.IsEmpty())
+			{
+				break;
+			}
+			OutDestination.Append(MoveTemp(Page));
+		}
+		Count = 0;
+	}
+
+	/** STL IteratorType model compliance methods.*/
+	FORCEINLINE ConstIteratorType begin() const
+	{
+		return GetIterator(0);
+	}
+	FORCEINLINE IteratorType begin()
+	{
+		return GetIterator(0);
+	}
+	FORCEINLINE ConstIteratorType end() const
+	{
+		return GetIterator(Count);
+	}
+	FORCEINLINE IteratorType end()
+	{
+		return GetIterator(Count);
+	}
+
+	// Friend operators
+	friend bool operator==(const TPagedArray& Left, const TPagedArray& Right)
+	{
+		if (Left.Num() != Right.Num())
+		{
+			return false;
+		}
+		const SizeType MinPageCount = NumRequiredPages(Left.Count);
+		for (SizeType PageIndex = 0; PageIndex < MinPageCount; ++PageIndex)
+		{
+			if (Left.Pages[PageIndex] != Right.Pages[PageIndex])
+			{
+				return false;
+			}
+		}
+		return true;
+	}
+
+	FORCEINLINE friend bool operator!=(const TPagedArray& Left, const TPagedArray& Right)
+	{
+		return !(Left == Right);
+	}
+
+private:
+	TArray<PageType> Pages;
+	SizeType Count = 0;
+
+	FORCEINLINE ConstIteratorType GetIterator(SizeType Offset) const
+	{
+		return ConstIteratorType(InPlace, Pages.GetData(), Offset);
+	}
+
+	FORCEINLINE IteratorType GetIterator(SizeType Offset)
+	{
+		return IteratorType(InPlace, Pages.GetData(), Offset);
+	}
+
+	/*
+	 * Adds a new page and reserves its capacity predefined by its traits.
+	 */
+	PageType& AddPage()
+	{
+		PageType& Page = Pages.Emplace_GetRef();
+		Page.Reserve(PageTraits::Capacity);
+		return Page;
+	}
+
+	/*
+	 * Copies the parameter element array of the parameter size into this container.
+	 * This method assumes no previously existing content.
+	 */
+	void CopyToEmpty(const ElementType* InSource, SizeType InSize)
+	{
+		checkSlow(!Count);
+		if (InSize)
+		{
+			Grow(InSize);
+			CopyUnchecked(InSource, InSize);
+		}
+	}
+
+	/*
+	 * Copies the parameter element array of the parameter size into this container.
+	 */
+	void CopyUnchecked(const ElementType* InSource, SizeType InSize)
+	{
+		if (InSize)
+		{
+			SizeType PageIndex = GetPageIndex(Count);
+			SizeType PendingCount = InSize;
+			if (const SizeType PageOffset = GetPageOffset(Count))
+			{
+				const SizeType AppendCount = FMath::Min(PageTraits::Capacity - PageOffset, PendingCount);
+				Pages[PageIndex++].Append(InSource, AppendCount);
+				InSource += AppendCount;
+				PendingCount -= AppendCount;
+			}
+			while (PendingCount >= PageTraits::Capacity)
+			{
+				Pages[PageIndex++].Append(InSource, PageTraits::Capacity);
+				InSource += PageTraits::Capacity;
+				PendingCount -= PageTraits::Capacity;
+			}
+			if (PendingCount)
+			{
+				Pages[PageIndex].Append(InSource, PendingCount);
+			}
+			Count += InSize;
+		}
+	}
+
+	/*
+	 * Grows the container's storage to the parameter capacity value allocating the required page count.
+	 * This method preserves existing data.
+	 */
+	void Grow(SizeType InCapacity)
+	{
+		checkSlow(Max() < InCapacity);
+		const SizeType RequiredPageCount = NumRequiredPages(InCapacity);
+		Pages.Reserve(RequiredPageCount);
+		for (SizeType Index = Pages.Num(); Index < RequiredPageCount; ++Index)
+		{
+			AddPage();
+		}
+	}
+
+	/* Grows the container by adding a new page of the parameter size if full.
+	 * This method preserves existing data.
+	 */
+	void GrowIfRequired()
+	{
+		if (Count == Max())
+		{
+			AddPage();
+		}
+	}
+
+	FORCEINLINE void CheckValidIndex(SizeType Index) const
+	{
+		checkf((Index >= 0) & (Index < Count), TEXT("Parameter index %d exceeds container size %d"), Index, Count);
+	}
+};
diff --git a/Engine/Source/Runtime/Core/Tests/Containers/PagedArrayTest.cpp b/Engine/Source/Runtime/Core/Tests/Containers/PagedArrayTest.cpp
new file mode 100644
--- /dev/null
+++ b/Engine/Source/Runtime/Core/Tests/Containers/PagedArrayTest.cpp
@@ -0,0 +1,1188 @@
+// Copyright Epic Games, Inc. All Rights Reserved.
+
+#if WITH_TESTS
+
+#include "Containers/PagedArray.h"
+#include "CoreMinimal.h"
+#include "TestUtils.h"
+#include "TokenTest.h"
+
+#include "Tests/TestHarness.h"
+
+namespace PagedArray
+{
+using TestArrayType = TPagedArray<int32Token, 128>;
+
+//---------------------------------------------------------------------------------------------------------------------
+// Forward declaration test
+//---------------------------------------------------------------------------------------------------------------------
+
+struct FSomeForwardDeclaredStruct;
+
+struct FDummyTestStruct
+{
+	TPagedArray<FSomeForwardDeclaredStruct> Data;
+};
+
+//---------------------------------------------------------------------------------------------------------------------
+// Unit tests
+//---------------------------------------------------------------------------------------------------------------------
+
+TEST_CASE("System::Core::Containers::TPagedArray::Reserve without data", "[SmokeFilter][Core][Containers][PagedArray]")
+{
+	int32Token::Reset();
+	using ArrayType = TPagedArray<int32Token, 128>;
+	ArrayType PagedArray;
+	CHECK(!PagedArray.Max());
+	CHECK(!PagedArray.Num());
+	CHECK(PagedArray.IsEmpty());
+	PagedArray.Reserve(ArrayType::MaxPerPage());
+	CHECK(PagedArray.Max() == ArrayType::MaxPerPage());
+	CHECK(!PagedArray.Num());
+	CHECK(PagedArray.IsEmpty());
+	CHECK(int32Token::EvenConstructionDestructionCalls(0));
+}
+
+TEST_CASE(
+	"System::Core::Containers::TPagedArray::Reserve Emplace single element",
+	"[SmokeFilter][Core][Containers][PagedArray]")
+{
+	int32Token::Reset();
+	{
+		using ArrayType = TPagedArray<int32Token, 128>;
+		ArrayType PagedArray;
+		CHECK(!PagedArray.Max());
+		CHECK(!PagedArray.Num());
+
+		PagedArray.Reserve(0);	// innocuous
+		PagedArray.Reserve(1);
+		PagedArray.Emplace(0);
+		CHECK(PagedArray.Max() >= 1);
+		CHECK(PagedArray.Max() == ArrayType::MaxPerPage());
+		CHECK(PagedArray.Num() == 1);
+
+		PagedArray.Reserve(ArrayType::MaxPerPage());
+		CHECK(PagedArray.Max() == ArrayType::MaxPerPage());
+		CHECK(PagedArray.Num() == 1);
+
+		PagedArray.Reserve(1 + ArrayType::MaxPerPage());
+		CHECK(PagedArray.Max() == ArrayType::MaxPerPage() * 2);
+		CHECK(PagedArray.Num() == 1);
+	}
+	CHECK(int32Token::EvenConstructionDestructionCalls(1));
+}
+
+TEST_CASE(
+	"System::Core::Containers::TPagedArray::Add",
+	"[SmokeFilter][Core][Containers][PagedArray]")
+{
+	int32Token::Reset();
+	{
+		using ArrayType = TPagedArray<int32Token, 128>;
+		ArrayType PagedArray;
+		CHECK(!PagedArray.Max());
+		CHECK(!PagedArray.Num());
+
+		int32Token Temp(13);
+		PagedArray.Add(Temp);
+		CHECK(PagedArray.Max() == ArrayType::MaxPerPage());
+		CHECK(PagedArray.Num() == 1);
+
+		PagedArray.Add(3);
+		CHECK(PagedArray.Num() == 2);
+	}
+	CHECK(int32Token::EvenConstructionDestructionCalls(4));
+}
+
+TEST_CASE("System::Core::Containers::TPagedArray::Reset", "[SmokeFilter][Core][Containers][PagedArray]")
+{
+	int32Token::Reset();
+	{
+		using ArrayType = TPagedArray<int32Token, 128>;
+		ArrayType PagedArray;
+		PagedArray.Reset();	 // Should be innocuous
+		CHECK(!PagedArray.Max());
+		CHECK(!PagedArray.Num());
+		CHECK(PagedArray.IsEmpty());
+
+		PagedArray.Emplace(0);
+		CHECK(PagedArray.Max());
+		CHECK(PagedArray.Max() == ArrayType::MaxPerPage());
+		CHECK(PagedArray.Num() == 1);
+		CHECK(!PagedArray.IsEmpty());
+
+		PagedArray.Reset();
+		CHECK(PagedArray.Max() == ArrayType::MaxPerPage());
+		CHECK(!PagedArray.Num());
+		CHECK(PagedArray.IsEmpty());
+
+		PagedArray.Reserve(ArrayType::MaxPerPage() * 6);
+		CHECK(PagedArray.Max() == ArrayType::MaxPerPage() * 6);
+		CHECK(!PagedArray.Num());
+		CHECK(PagedArray.IsEmpty());
+		PagedArray.Reset();
+		CHECK(PagedArray.Max() == ArrayType::MaxPerPage() * 6);
+	}
+	CHECK(int32Token::EvenConstructionDestructionCalls(1));
+}
+
+TEST_CASE(
+	"System::Core::Containers::TPagedArray::Reset multiple-page plus consistent address check",
+	"[SmokeFilter][Core][Containers][PagedArray]")
+{
+	int32Token::Reset();
+	{
+		using ArrayType = TPagedArray<int32Token, 16>;
+		ArrayType PagedArray = {0, 1, 2, 3, 4, 5, 6, 7, 8, 9};
+		CHECK(PagedArray.Max() == ArrayType::MaxPerPage() * 3);
+
+		TArray<const int32Token*> PtrAddresses;
+		for (const int32Token& Value : PagedArray)
+		{
+			PtrAddresses.Emplace(&Value);
+		}
+
+		PagedArray.Reset();
+		CHECK(!PagedArray.Num());
+		CHECK(PagedArray.IsEmpty());
+		CHECK(PagedArray.Max() == ArrayType::MaxPerPage() * 3);
+
+		// Validate same address after Reset plus re-adding the same elements
+		PagedArray = {0, 1, 2, 3, 4, 5, 6, 7, 8, 9};
+		int32 Index = 0;
+		for (const int32Token& Value : PagedArray)
+		{
+			CHECK(PtrAddresses[Index++] == &Value);
+		}
+	}
+	CHECK(int32Token::EvenConstructionDestructionCalls(10 * 4));
+}
+
+TEST_CASE("System::Core::Containers::TPagedArray::Empty", "[SmokeFilter][Core][Containers][PagedArray]")
+{
+	int32Token::Reset();
+	{
+		using ArrayType = TPagedArray<int32Token, 128>;
+		ArrayType PagedArray;
+		PagedArray.Empty();	 // Should be innocuous
+		CHECK(!PagedArray.Max());
+		CHECK(!PagedArray.Num());
+		CHECK(PagedArray.IsEmpty());
+	}
+	CHECK(int32Token::EvenConstructionDestructionCalls(0));
+}
+
+TEST_CASE(
+	"System::Core::Containers::TPagedArray::Empty after single Emplace",
+	"[SmokeFilter][Core][Containers][PagedArray]")
+{
+	int32Token::Reset();
+	{
+		using ArrayType = TPagedArray<int32Token, 128>;
+		ArrayType PagedArray;
+		PagedArray.Emplace(0);
+		CHECK(PagedArray.Max());
+		CHECK(PagedArray.Num() == 1);
+		CHECK(!PagedArray.IsEmpty());
+		PagedArray.Empty();
+		CHECK(!PagedArray.Max());
+		CHECK(!PagedArray.Num());
+		CHECK(PagedArray.IsEmpty());
+	}
+	CHECK(int32Token::EvenConstructionDestructionCalls(1));
+}
+
+TEST_CASE("System::Core::Containers::TPagedArray::Indexing", "[SmokeFilter][Core][Containers][PagedArray]")
+{
+	using ArrayType = TPagedArray<int32Token, 128>;
+	int32Token::Reset();
+	{
+		ArrayType PagedArray;
+		// Single-page
+		PagedArray.Reserve(ArrayType::MaxPerPage());
+		CHECK(PagedArray.Max() == ArrayType::MaxPerPage());
+		for (int32 i = 0; i < ArrayType::MaxPerPage(); ++i)
+		{
+			PagedArray.Emplace(i);
+			CHECK(PagedArray.Max() == ArrayType::MaxPerPage());
+		}
+		for (int32 i = 0; i < PagedArray.Num(); ++i)
+		{
+			CHECK(PagedArray[i] == i);
+		}
+		// Multi-page
+		for (int32 i = 0; i < ArrayType::MaxPerPage(); ++i)
+		{
+			PagedArray.Emplace(i + ArrayType::MaxPerPage());
+			CHECK(PagedArray.Max() == ArrayType::MaxPerPage() * 2);
+		}
+		for (int32 i = 0; i < PagedArray.Num(); ++i)
+		{
+			CHECK(PagedArray[i] == i);
+		}
+	}
+	CHECK(int32Token::EvenConstructionDestructionCalls(ArrayType::MaxPerPage() * 2));
+}
+
+TEST_CASE("System::Core::Containers::TPagedArray::Iteration", "[SmokeFilter][Core][Containers][PagedArray]")
+{
+	using ArrayType = TPagedArray<int32Token, 128>;
+	int32Token::Reset();
+	{
+		ArrayType PagedArray;
+		// Single-page
+		PagedArray.Reserve(ArrayType::MaxPerPage());
+		CHECK(PagedArray.Max() == ArrayType::MaxPerPage());
+		for (int32 i = 0; i < ArrayType::MaxPerPage(); ++i)
+		{
+			PagedArray.Emplace(i);
+			CHECK(PagedArray.Max() == ArrayType::MaxPerPage());
+		}
+		int32 CheckValue = 0;
+		for (const int32Token& Value : PagedArray)
+		{
+			CHECK(Value == CheckValue++);
+		}
+		// Multi-page
+		for (int32 i = 0; i < ArrayType::MaxPerPage(); ++i)
+		{
+			PagedArray.Emplace(i + ArrayType::MaxPerPage());
+			CHECK(PagedArray.Max() == ArrayType::MaxPerPage() * 2);
+		}
+		CheckValue = 0;
+		for (const int32Token& Value : PagedArray)
+		{
+			CHECK(Value == CheckValue++);
+		}
+	}
+	CHECK(int32Token::EvenConstructionDestructionCalls(ArrayType::MaxPerPage() * 2));
+}
+
+TEST_CASE("System::Core::Containers::TPagedArray::Last", "[SmokeFilter][Core][Containers][PagedArray]")
+{
+	using ArrayType = TPagedArray<int32Token, 128>;
+	int32Token::Reset();
+	{
+		ArrayType PagedArray;
+		for (int32 i = 0; i < ArrayType::MaxPerPage() * 2; ++i)
+		{
+			PagedArray.Emplace(i);
+			CHECK(PagedArray.Last() == i);
+		}
+	}
+	CHECK(int32Token::EvenConstructionDestructionCalls(ArrayType::MaxPerPage() * 2));
+}
+
+TEST_CASE("System::Core::Containers::TPagedArray::SetNum", "[SmokeFilter][Core][Containers][PagedArray]")
+{
+	using ArrayType = TPagedArray<int32Token, 128>;
+	int32Token::Reset();
+	{
+		ArrayType PagedArray;
+
+		// Zero: should be innocuous
+		PagedArray.SetNum(0);
+		CHECK(!PagedArray.NumPages());
+		CHECK(PagedArray.IsEmpty());
+
+		// Single non-full page
+		PagedArray.SetNum(ArrayType::MaxPerPage() - 1);
+		CHECK(PagedArray.Num() == ArrayType::MaxPerPage() - 1);
+		CHECK(PagedArray.NumPages() == 1);
+		CHECK(!PagedArray.IsEmpty());
+
+		// Single full page
+		PagedArray.SetNum(ArrayType::MaxPerPage());
+		CHECK(PagedArray.Num() == ArrayType::MaxPerPage());
+		CHECK(PagedArray.NumPages() == 1);
+		CHECK(!PagedArray.IsEmpty());
+
+		// Full page plus 1
+		PagedArray.SetNum(ArrayType::MaxPerPage() + 1);
+		CHECK(PagedArray.Num() == ArrayType::MaxPerPage() + 1);
+		CHECK(PagedArray.NumPages() == 2);
+		CHECK(!PagedArray.IsEmpty());
+
+		// Back to single full page
+		PagedArray.SetNum(ArrayType::MaxPerPage());
+		CHECK(PagedArray.Num() == ArrayType::MaxPerPage());
+		CHECK(PagedArray.NumPages() == 1);
+		CHECK(!PagedArray.IsEmpty());
+
+		// Back to single non-full page
+		PagedArray.SetNum(ArrayType::MaxPerPage() - 1);
+		CHECK(PagedArray.Num() == ArrayType::MaxPerPage() - 1);
+		CHECK(PagedArray.NumPages() == 1);
+		CHECK(!PagedArray.IsEmpty());
+
+		// Back to zero
+		PagedArray.SetNum(0);
+		CHECK(!PagedArray.NumPages());
+		CHECK(PagedArray.IsEmpty());
+	}
+	CHECK(int32Token::EvenConstructionDestructionCalls(ArrayType::MaxPerPage() + 1));
+	int32Token::Reset();
+	{
+		ArrayType PagedArray;
+
+		// Add multiple pages in one go (non-full page)
+		PagedArray.SetNum(ArrayType::MaxPerPage() * 4 - 1);
+		CHECK(PagedArray.Num() == ArrayType::MaxPerPage() * 4 - 1);
+		CHECK(PagedArray.NumPages() == 4);
+		CHECK(!PagedArray.IsEmpty());
+
+		// Add multiple pages in one go (full page)
+		PagedArray.SetNum(ArrayType::MaxPerPage() * 6);
+		CHECK(PagedArray.Num() == ArrayType::MaxPerPage() * 6);
+		CHECK(PagedArray.NumPages() == 6);
+		CHECK(!PagedArray.IsEmpty());
+
+		// Verify content
+		int32 CheckValue = 0;
+		for (int32Token& Value : PagedArray)
+		{
+			Value = CheckValue++;
+		}
+
+		// Remove multiple pages in one go (non-full page)
+		PagedArray.SetNum(ArrayType::MaxPerPage() * 3 - 1);
+		CHECK(PagedArray.Num() == ArrayType::MaxPerPage() * 3 - 1);
+		CHECK(PagedArray.NumPages() == 3);
+		CHECK(!PagedArray.IsEmpty());
+		CheckValue = 0;
+		for (const int32Token& Value : PagedArray)
+		{
+			CHECK(Value == CheckValue++);
+		}
+
+		// Remove multiple pages in one go (full page)
+		PagedArray.SetNum(ArrayType::MaxPerPage());
+		CHECK(PagedArray.Num() == ArrayType::MaxPerPage());
+		CHECK(PagedArray.NumPages() == 1);
+		CHECK(!PagedArray.IsEmpty());
+		CheckValue = 0;
+		for (const int32Token& Value : PagedArray)
+		{
+			CHECK(Value == CheckValue++);
+		}
+	}
+	CHECK(int32Token::EvenConstructionDestructionCalls(ArrayType::MaxPerPage() * 6 * 2));
+}
+
+TEST_CASE("System::Core::Containers::TPagedArray::SetNum after Reset", "[SmokeFilter][Core][Containers][PagedArray]")
+{
+	using ArrayType = TPagedArray<int32Token, 16>;
+	int32Token::Reset();
+	{
+		ArrayType PagedArray;
+
+		PagedArray.SetNum(ArrayType::MaxPerPage() * 3);
+		CHECK(PagedArray.Num() == ArrayType::MaxPerPage() * 3);
+		CHECK(PagedArray.NumPages() == 3);
+
+		// Reset and set same number of elements
+		PagedArray.Reset();
+		PagedArray.SetNum(ArrayType::MaxPerPage() * 3);
+		CHECK(PagedArray.Num() == ArrayType::MaxPerPage() * 3);
+		CHECK(PagedArray.NumPages() == 3);
+
+		// Verify content
+		for (const int32Token& Value : PagedArray)
+		{
+			CHECK(Value == 0);
+		}
+
+		// Write content memory
+		for (int32Token& Value : PagedArray)
+		{
+			Value = 13;
+		}
+
+		// Reset and set a greater number of elements
+		PagedArray.Reset();
+		PagedArray.SetNum(ArrayType::MaxPerPage() * 5);
+		CHECK(PagedArray.Num() == ArrayType::MaxPerPage() * 5);
+		CHECK(PagedArray.NumPages() == 5);
+
+		// Verify content
+		for (const int32Token& Value : PagedArray)
+		{
+			CHECK(Value == 0);
+		}
+	}
+	CHECK(int32Token::EvenConstructionDestructionCalls(ArrayType::MaxPerPage() * (3 * 3 + 5)));
+}
+
+TEST_CASE(
+	"System::Core::Containers::TPagedArray::SetNum small page size",
+	"[SmokeFilter][Core][Containers][PagedArray]")
+{
+	using ArrayType = TPagedArray<int32Token, 16>;
+	int32Token::Reset();
+	{
+		ArrayType PagedArray;
+
+		// Add multiple pages in one go (non-full page)
+		PagedArray.SetNum(ArrayType::MaxPerPage() * 4 - 1);
+		CHECK(PagedArray.Num() == ArrayType::MaxPerPage() * 4 - 1);
+		CHECK(PagedArray.NumPages() == 4);
+		CHECK(!PagedArray.IsEmpty());
+
+		// Verify content
+		int32 CheckValue = 0;
+		for (int32Token& Value : PagedArray)
+		{
+			Value = CheckValue++;
+		}
+
+		// Add multiple pages in one go (full page)
+		PagedArray.SetNum(ArrayType::MaxPerPage() * 6);
+		CHECK(PagedArray.Num() == ArrayType::MaxPerPage() * 6);
+		CHECK(PagedArray.NumPages() == 6);
+		CHECK(!PagedArray.IsEmpty());
+
+		CheckValue = 0;
+		for (int32 i = 0; i < ArrayType::MaxPerPage() * 4 - 1; ++i)
+		{
+			CHECK(PagedArray[i] == CheckValue++);
+		}
+	}
+	CHECK(int32Token::EvenConstructionDestructionCalls(ArrayType::MaxPerPage() * 10 - 1));
+}
+
+TEST_CASE("System::Core::Containers::TPagedArray::SetZero", "[SmokeFilter][Core][Containers][PagedArray]")
+{
+	using ArrayType = TPagedArray<int32Token, 16>;
+	int32Token::Reset();
+	{
+		ArrayType PagedArray;
+
+		// Single page
+		{
+			// Write content memory
+			PagedArray.SetNum(ArrayType::MaxPerPage() - 1);
+			for (int32Token& Value : PagedArray)
+			{
+				Value = 13;
+			}
+			PagedArray.SetZero();
+	
+			// Verify content
+			CHECK(PagedArray.Num() == ArrayType::MaxPerPage() - 1);
+			for (const int32Token& Value : PagedArray)
+			{
+				CHECK(Value == 0);
+			}
+		}
+		// Multiple pages
+		{
+			// Write content memory
+			PagedArray.SetNum(ArrayType::MaxPerPage() * 3 + 1); // This adds 2 pages + 2 elements
+			for (int32Token& Value : PagedArray)
+			{
+				Value = 13;
+			}
+			PagedArray.SetZero();
+
+			// Verify content
+			for (const int32Token& Value : PagedArray)
+			{
+				CHECK(Value == 0);
+			}
+			
+			// Verify content
+			CHECK(PagedArray.Num() == ArrayType::MaxPerPage() * 3 + 1);
+			for (const int32Token& Value : PagedArray)
+			{
+				CHECK(Value == 0);
+			}
+		}
+	}
+	CHECK(int32Token::EvenConstructionDestructionCalls(ArrayType::MaxPerPage() * 7 + 1));
+}
+
+
+TEST_CASE("System::Core::Containers::TPagedArray::RemoveAtSwap", "[SmokeFilter][Core][Containers][PagedArray]")
+{
+	using ArrayType = TPagedArray<int32Token, 128>;
+	int32Token::Reset();
+	{
+		// Single page
+		ArrayType PagedArray;
+		PagedArray.SetNum(ArrayType::MaxPerPage());
+		int32 CheckValue = 0;
+		for (int32Token& Value : PagedArray)
+		{
+			Value = CheckValue++;
+		}
+
+		while (PagedArray.Num() > 1)
+		{
+			const int32 Count = PagedArray.Num();
+			const int32 Index = FMath::RandRange(0, Count - 1);
+			const int32Token Last = PagedArray.Last();
+			PagedArray.RemoveAtSwap(Index);
+			CHECK(Count == PagedArray.Num() + 1);
+			if (Index < PagedArray.Num())
+			{
+				CHECK(PagedArray[Index] == Last);
+			}
+		}
+
+		// Test page shrinking after removing last element in page
+		PagedArray.RemoveAtSwap(0, false); 
+		CHECK(PagedArray.Max() == ArrayType::MaxPerPage());
+		PagedArray.Emplace(0);
+		PagedArray.RemoveAtSwap(0);
+		CHECK(!PagedArray.Max());
+	}
+	// One full page for SetNum, another for assignment and another minus one for RemoveAtSwap.
+	CHECK(int32Token::EvenConstructionDestructionCalls(ArrayType::MaxPerPage() * 3));
+	int32Token::Reset();
+	{
+		// Multiple pages
+		ArrayType PagedArray;
+		PagedArray.SetNum(ArrayType::MaxPerPage() * 3 + 1);
+		int32 CheckValue = 0;
+		for (int32Token& Value : PagedArray)
+		{
+			Value = CheckValue++;
+		}
+
+		while (PagedArray.Num() > 1)
+		{
+			const int32 Count = PagedArray.Num();
+			const int32 Index = FMath::RandRange(0, Count - 1);
+			const int32Token Last = PagedArray.Last();
+			PagedArray.RemoveAtSwap(Index);
+			CHECK(Count == PagedArray.Num() + 1);
+			if (Index < PagedArray.Num())
+			{
+				CHECK(PagedArray[Index] == Last);
+			}
+		}
+		PagedArray.RemoveAtSwap(0);
+	}
+	CHECK(int32Token::EvenConstructionDestructionCalls(ArrayType::MaxPerPage() * 9 + 2));
+}
+
+TEST_CASE("System::Core::Containers::TPagedArray::Pop", "[SmokeFilter][Core][Containers][PagedArray]")
+{
+	using ArrayType = TPagedArray<int32Token, 128>;
+	int32Token::Reset();
+	{
+		// Single page
+		ArrayType PagedArray;
+		PagedArray.Emplace(0);
+		PagedArray.Pop();
+		CHECK(PagedArray.IsEmpty());
+
+		PagedArray.SetNum(ArrayType::MaxPerPage());
+		int32 CheckValue = 0;
+		for (int32Token& Value : PagedArray)
+		{
+			Value = CheckValue++;
+		}
+		PagedArray.Emplace(0);
+		PagedArray.Pop();
+		CHECK(PagedArray.Num() == ArrayType::MaxPerPage());
+
+		CheckValue = 0;
+		for (int32Token& Value : PagedArray)
+		{
+			CHECK(Value == CheckValue++);
+		}
+	}
+	CHECK(int32Token::EvenConstructionDestructionCalls(ArrayType::MaxPerPage() * 2 + 2));
+}
+
+TEST_CASE("System::Core::Containers::TPagedArray::Copy", "[SmokeFilter][Core][Containers][PagedArray]")
+{
+	using ArrayType = TPagedArray<int32Token, 128>;
+	int32Token::Reset();
+	{
+		ArrayType PagedArray;
+
+		// Single-page
+		PagedArray.Reserve(ArrayType::MaxPerPage());
+		CHECK(PagedArray.Max() == ArrayType::MaxPerPage());
+		for (int32 i = 0; i < ArrayType::MaxPerPage(); ++i)
+		{
+			PagedArray.Emplace(i);
+		}
+		ArrayType OtherArray(PagedArray);
+		CHECK(OtherArray.Num() == PagedArray.Num());
+		for (int32 i = 0; i < OtherArray.Num(); ++i)
+		{
+			CHECK(OtherArray[i] == i);
+		}
+		ArrayType AnotherArray;
+		CHECK(!AnotherArray.Num());
+		AnotherArray = PagedArray;
+		CHECK(AnotherArray.Num() == PagedArray.Num());
+		for (int32 i = 0; i < AnotherArray.Num(); ++i)
+		{
+			CHECK(AnotherArray[i] == i);
+		}
+
+		// Multi-page
+		for (int32 i = 0; i < ArrayType::MaxPerPage(); ++i)
+		{
+			PagedArray.Emplace(i + ArrayType::MaxPerPage());
+			CHECK(PagedArray.Max() == ArrayType::MaxPerPage() * 2);
+		}
+		ArrayType YetAnotherArray(PagedArray);
+		CHECK(YetAnotherArray.Num() == PagedArray.Num());
+		for (int32 i = 0; i < YetAnotherArray.Num(); ++i)
+		{
+			CHECK(YetAnotherArray[i] == i);
+		}
+
+		AnotherArray = PagedArray;
+		CHECK(AnotherArray.Num() == PagedArray.Num());
+		for (int32 i = 0; i < AnotherArray.Num(); ++i)
+		{
+			CHECK(AnotherArray[i] == i);
+		}
+	}
+	CHECK(int32Token::EvenConstructionDestructionCalls(ArrayType::MaxPerPage() * 8));
+}
+
+TEST_CASE("System::Core::Containers::TPagedArray::Move", "[SmokeFilter][Core][Containers][PagedArray]")
+{
+	using ArrayType = TPagedArray<int32Token, 128>;
+	int32Token::Reset();
+	{
+		TestArrayType PagedArray;
+
+		// Single-page
+		PagedArray.Reserve(ArrayType::MaxPerPage());
+		CHECK(PagedArray.Max() == ArrayType::MaxPerPage());
+		for (int32 i = 0; i < ArrayType::MaxPerPage(); ++i)
+		{
+			PagedArray.Emplace(i);
+		}
+		ArrayType OtherArray(MoveTemp(PagedArray));
+		CHECK(OtherArray.Num() == ArrayType::MaxPerPage());
+		CHECK(PagedArray.IsEmpty());
+		for (int32 i = 0; i < OtherArray.Num(); ++i)
+		{
+			CHECK(OtherArray[i] == i);
+		}
+		ArrayType AnotherArray;
+		CHECK(PagedArray.IsEmpty());
+		AnotherArray = MoveTemp(OtherArray);
+		CHECK(OtherArray.IsEmpty());
+		CHECK(AnotherArray.Num() == ArrayType::MaxPerPage());
+		for (int32 i = 0; i < AnotherArray.Num(); ++i)
+		{
+			CHECK(AnotherArray[i] == i);
+		}
+
+		// Multi-page
+		for (int32 i = 0; i < ArrayType::MaxPerPage() * 3; ++i)
+		{
+			PagedArray.Emplace(i);
+		}
+		ArrayType YetAnotherArray(MoveTemp(PagedArray));
+		CHECK(YetAnotherArray.Num() == ArrayType::MaxPerPage() * 3);
+		CHECK(PagedArray.IsEmpty());
+		for (int32 i = 0; i < YetAnotherArray.Num(); ++i)
+		{
+			CHECK(YetAnotherArray[i] == i);
+		}
+
+		AnotherArray = MoveTemp(YetAnotherArray);
+		CHECK(YetAnotherArray.IsEmpty());
+		CHECK(AnotherArray.Num() == ArrayType::MaxPerPage() * 3);
+		for (int32 i = 0; i < AnotherArray.Num(); ++i)
+		{
+			CHECK(AnotherArray[i] == i);
+		}
+	}
+	CHECK(int32Token::EvenConstructionDestructionCalls(ArrayType::MaxPerPage() * 4));
+}
+
+TEST_CASE(
+	"System::Core::Containers::TPagedArray::Construct from std initializer_list",
+	"[SmokeFilter][Core][Containers][PagedArray]")
+{
+	using ArrayType = TPagedArray<int32Token, 16>;
+	int32Token::Reset();
+	{
+		// Single-page
+		{
+			ArrayType PagedArray({0, 1, 2, 3});
+			CHECK(PagedArray.Num() == 4);
+			CHECK(PagedArray.Max() == ArrayType::MaxPerPage());
+			int32 CheckValue = 0;
+			for (const int32Token& Value : PagedArray)
+			{
+				CHECK(Value == CheckValue++);
+			}
+			CHECK(PagedArray.Num() == CheckValue);
+		}
+
+		// Multiple-page
+		{
+			ArrayType PagedArray({0, 1, 2, 3, 4, 5, 6, 7, 8, 9});
+			CHECK(PagedArray.Num() == 10);
+			CHECK(PagedArray.Max() == ArrayType::MaxPerPage() * 3);
+			int32 CheckValue = 0;
+			for (const int32Token& Value : PagedArray)
+			{
+				CHECK(Value == CheckValue++);
+			}
+			CHECK(PagedArray.Num() == CheckValue);
+		}
+	}
+	CHECK(int32Token::EvenConstructionDestructionCalls((4 + 10) * 2));
+}
+
+TEST_CASE(
+	"System::Core::Containers::TPagedArray::Assign from std initializer_list",
+	"[SmokeFilter][Core][Containers][PagedArray]")
+{
+	using ArrayType = TPagedArray<int32Token, 16>;
+	int32Token::Reset();
+	{
+		// Single-page
+		{
+			ArrayType PagedArray;
+			CHECK(PagedArray.IsEmpty());
+			PagedArray = {0, 1, 2, 3};
+			CHECK(PagedArray.Num() == 4);
+			CHECK(PagedArray.Max() == ArrayType::MaxPerPage());
+			int32 CheckValue = 0;
+			for (const int32Token& Value : PagedArray)
+			{
+				CHECK(Value == CheckValue++);
+			}
+			CHECK(PagedArray.Num() == CheckValue);
+		}
+
+		// Multiple-page
+		{
+			ArrayType PagedArray;
+			CHECK(PagedArray.IsEmpty());
+			PagedArray = {0, 1, 2, 3, 4, 5, 6, 7, 8, 9};
+			CHECK(PagedArray.Num() == 10);
+			CHECK(PagedArray.Max() == ArrayType::MaxPerPage() * 3);
+			int32 CheckValue = 0;
+			for (const int32Token& Value : PagedArray)
+			{
+				CHECK(Value == CheckValue++);
+			}
+			CHECK(PagedArray.Num() == CheckValue);
+		}
+	}
+	CHECK(int32Token::EvenConstructionDestructionCalls((4 + 10) * 2));
+}
+
+TEST_CASE("System::Core::Containers::TPagedArray::Comparison", "[SmokeFilter][Core][Containers][PagedArray]")
+{
+	using ArrayType = TPagedArray<int32Token, 128>;
+	int32Token::Reset();
+	{
+		ArrayType PagedArray, OtherPagedArray;
+		CHECK(PagedArray == OtherPagedArray);
+
+		for (int32 i = 0; i < ArrayType::MaxPerPage(); ++i)
+		{
+			PagedArray.Emplace(i);
+			OtherPagedArray.Emplace(i);
+		}
+		CHECK(PagedArray == OtherPagedArray);
+
+		// Verify diffirent content
+		PagedArray.Emplace(0);
+		CHECK(PagedArray.Max() == ArrayType::MaxPerPage() * 2);
+		CHECK(OtherPagedArray.Max() == ArrayType::MaxPerPage());
+		CHECK(PagedArray != OtherPagedArray);
+
+		// Verify same content with same page allocation
+		PagedArray.Pop();
+		CHECK(PagedArray.Max() == ArrayType::MaxPerPage());
+		CHECK(OtherPagedArray.Max() == ArrayType::MaxPerPage());
+		CHECK(PagedArray == OtherPagedArray);
+
+		// Verify same content with same  different page allocation
+		PagedArray.Emplace(0);
+		PagedArray.Pop(false);
+		CHECK(PagedArray.Max() == ArrayType::MaxPerPage() * 2);
+		CHECK(OtherPagedArray.Max() == ArrayType::MaxPerPage());
+		CHECK(PagedArray == OtherPagedArray);
+	}
+	CHECK(int32Token::EvenConstructionDestructionCalls(ArrayType::MaxPerPage() * 2 + 2));
+}
+
+TEST_CASE("System::Core::Containers::TPagedArray::Assign from TArray", "[SmokeFilter][Core][Containers][PagedArray]")
+{
+	using ArrayType = TPagedArray<int32Token, 16>;
+	int32Token::Reset();
+	{
+		// Single-page
+		{
+			TArray<int32Token> SourceArray({0, 1, 2, 3});
+
+			ArrayType PagedArray;
+			CHECK(PagedArray.IsEmpty());
+			PagedArray.Assign(TArray<int32Token>{}); // innocuous.
+			PagedArray.Assign(SourceArray);
+			CHECK(PagedArray.Num() == 4);
+			CHECK(PagedArray.Max() == ArrayType::MaxPerPage());
+			int32 CheckValue = 0;
+			for (const int32Token& Value : PagedArray)
+			{
+				CHECK(Value == CheckValue++);
+			}
+			CHECK(PagedArray.Num() == CheckValue);
+		}
+
+		// Multiple-page
+		{
+			const TArray<int32Token> SourceArray({0, 1, 2, 3, 4, 5, 6, 7, 8, 9});
+
+			ArrayType PagedArray;
+			CHECK(PagedArray.IsEmpty());
+			PagedArray.Assign(SourceArray);
+			CHECK(PagedArray.Num() == 10);
+			CHECK(PagedArray.Max() == ArrayType::MaxPerPage() * 3);
+			int32 CheckValue = 0;
+			for (const int32Token& Value : PagedArray)
+			{
+				CHECK(Value == CheckValue++);
+			}
+			CHECK(PagedArray.Num() == CheckValue);
+		}
+	}
+	CHECK(int32Token::EvenConstructionDestructionCalls((4 + 10) * 3));
+}
+
+TEST_CASE("System::Core::Containers::TPagedArray::Copy to TArray", "[SmokeFilter][Core][Containers][PagedArray]")
+{
+	using ArrayType = TPagedArray<int32Token, 16>;
+	int32Token::Reset();
+	{
+		ArrayType PagedArray;
+		TArray<int32Token> DestinationArray;
+
+		// Empty copy
+		PagedArray.ToArray(DestinationArray);  // innocuos
+		CHECK(PagedArray.IsEmpty());
+		CHECK(DestinationArray.IsEmpty());
+
+		// Single full page copy
+		PagedArray = {0, 1, 2, 3};
+		PagedArray.ToArray(DestinationArray);
+		CHECK(DestinationArray.Num() == 4);
+		int32 CheckValue = 0;
+		for (const int32Token& Value : DestinationArray)
+		{
+			CHECK(Value == CheckValue++);
+		}
+
+		// Single full page copy from multiple allocated pages
+		PagedArray.Reserve(20);
+		CHECK(PagedArray.Max() == ArrayType::MaxPerPage() * 5);
+		PagedArray.ToArray(DestinationArray);
+		CHECK(DestinationArray.Num() == 4);
+		CheckValue = 0;
+		for (const int32Token& Value : DestinationArray)
+		{
+			CHECK(Value == CheckValue++);
+		}
+
+		// Multiple page copy
+		PagedArray = {0, 1, 2, 3, 4, 5, 6, 7, 8, 9};
+		PagedArray.ToArray(DestinationArray);
+		CHECK(DestinationArray.Num() == 10);
+		CheckValue = 0;
+		for (const int32Token& Value : DestinationArray)
+		{
+			CHECK(Value == CheckValue++);
+		}
+	}
+	CHECK(int32Token::EvenConstructionDestructionCalls(4 * 4 + 10 * 3));
+}
+
+TEST_CASE("System::Core::Containers::TPagedArray::Move to TArray", "[SmokeFilter][Core][Containers][PagedArray]")
+{
+	using ArrayType = TPagedArray<int32Token, 16>;
+	int32Token::Reset();
+	{
+		ArrayType PagedArray;
+		TArray<int32Token> DestinationArray;
+
+		// Empty copy
+		MoveTemp(PagedArray).ToArray(DestinationArray);  // innocuos
+		CHECK(PagedArray.IsEmpty());
+		CHECK(DestinationArray.IsEmpty());
+
+		// Single full page copy
+		PagedArray = {0, 1, 2, 3};
+		MoveTemp(PagedArray).ToArray(DestinationArray);
+		CHECK(DestinationArray.Num() == 4);
+		int32 CheckValue = 0;
+		for (const int32Token& Value : DestinationArray)
+		{
+			CHECK(Value == CheckValue++);
+		}
+		CHECK(PagedArray.IsEmpty());
+
+		// Single full page copy from multiple allocated pages
+		PagedArray = {0, 1, 2, 3};
+		PagedArray.Reserve(20);
+		CHECK(PagedArray.Max() == ArrayType::MaxPerPage() * 5);
+		MoveTemp(PagedArray).ToArray(DestinationArray);
+		CHECK(DestinationArray.Num() == 4);
+		CheckValue = 0;
+		for (const int32Token& Value : DestinationArray)
+		{
+			CHECK(Value == CheckValue++);
+		}
+		CHECK(PagedArray.IsEmpty());
+
+		// Multiple page copy
+		PagedArray = {0, 1, 2, 3, 4, 5, 6, 7, 8, 9};
+		MoveTemp(PagedArray).ToArray(DestinationArray);
+		CHECK(DestinationArray.Num() == 10);
+		CheckValue = 0;
+		for (const int32Token& Value : DestinationArray)
+		{
+			CHECK(Value == CheckValue++);
+		}
+		CHECK(PagedArray.IsEmpty());
+	}
+	CHECK(int32Token::EvenConstructionDestructionCalls(4 * 4 + 10 * 2));
+}
+
+TEST_CASE(
+	"System::Core::Containers::TPagedArray::Assign from compatible paged array",
+	"[SmokeFilter][Core][Containers][PagedArray]")
+{
+	using ArrayType = TPagedArray<int32Token, 16>;
+	using OtherArrayType = TPagedArray<int32Token, 32>;
+	int32Token::Reset();
+	{
+		// Single-page smaller to larger page size
+		{
+			ArrayType PagedArray({0, 1, 2, 3});
+			OtherArrayType OtherArray;
+			OtherArray.Assign(PagedArray);
+			CHECK(PagedArray.NumPages() == 1);
+			CHECK(OtherArray.NumPages() == 1);
+
+			CHECK(OtherArray.Num() == 4);
+			int32 CheckValue = 0;
+			for (const int32Token& Value : OtherArray)
+			{
+				CHECK(Value == CheckValue++);
+			}
+			CHECK(OtherArray.Num() == CheckValue);
+		}
+		// Single-page larger to smaller page size
+		{
+			OtherArrayType OtherArray({0, 1, 2, 3, 4, 5, 6, 7});
+			ArrayType PagedArray;
+			PagedArray.Assign(OtherArray);
+			CHECK(PagedArray.NumPages() == 2);
+			CHECK(OtherArray.NumPages() == 1);
+
+			CHECK(OtherArray.Num() == 8);
+			int32 CheckValue = 0;
+			for (const int32Token& Value : OtherArray)
+			{
+				CHECK(Value == CheckValue++);
+			}
+			CHECK(OtherArray.Num() == CheckValue);
+		}
+		// Multiple-page smaller to larger page size
+		{
+			ArrayType PagedArray({0, 1, 2, 3, 4, 5, 6, 7});
+			OtherArrayType OtherArray;
+			OtherArray.Assign(PagedArray);
+			CHECK(PagedArray.NumPages() == 2);
+			CHECK(OtherArray.NumPages() == 1);
+
+			CHECK(OtherArray.Num() == 8);
+			int32 CheckValue = 0;
+			for (const int32Token& Value : OtherArray)
+			{
+				CHECK(Value == CheckValue++);
+			}
+			CHECK(OtherArray.Num() == CheckValue);
+		}
+		// Multiple-page larger to smaller page size
+		{
+			OtherArrayType OtherArray({0, 1, 2, 3, 4, 5, 6, 7, 8, 9});
+			ArrayType PagedArray;
+			PagedArray.Assign(OtherArray);
+			CHECK(PagedArray.NumPages() == 3);
+			CHECK(OtherArray.NumPages() == 2);
+
+			CHECK(OtherArray.Num() == 10);
+			int32 CheckValue = 0;
+			for (const int32Token& Value : OtherArray)
+			{
+				CHECK(Value == CheckValue++);
+			}
+			CHECK(OtherArray.Num() == CheckValue);
+		}
+	}
+	CHECK(int32Token::EvenConstructionDestructionCalls((10 + 8 + 8 + 4) * 3));
+}
+
+TEST_CASE(
+	"System::Core::Containers::TPagedArray::Append from compatible paged array",
+	"[SmokeFilter][Core][Containers][PagedArray]")
+{
+	using ArrayType = TPagedArray<int32Token, 16>;
+	using OtherArrayType = TPagedArray<int32Token, 32>;
+	int32Token::Reset();
+	{
+		// Single-page smaller to larger page size
+		{
+			ArrayType PagedArray;
+			const ArrayType SourceArray({ 0, 1, 2, 3 });
+			const OtherArrayType OtherSourceArray({ 4, 5, 6, 7 });
+			CHECK(SourceArray.NumPages() == 1);
+			CHECK(OtherSourceArray.NumPages() == 1);
+
+			PagedArray.Append(SourceArray);
+			CHECK(PagedArray.NumPages() == 1);
+			CHECK(PagedArray.Num() == 4);
+
+			PagedArray.Append(OtherSourceArray);
+			CHECK(PagedArray.NumPages() == 2);
+			CHECK(PagedArray.Num() == 8);
+
+			// Validate content
+			int32 CheckValue = 0;
+			for (const int32Token& Value : PagedArray)
+			{
+				CHECK(Value == CheckValue++);
+			}
+			CHECK(PagedArray.Num() == CheckValue);
+		}
+		// Single-page larger to smaller page size
+		{
+			OtherArrayType OtherArray;
+			const ArrayType SourceArray({0, 1, 2, 3});
+			const OtherArrayType OtherSourceArray({4, 5, 6, 7});
+			CHECK(SourceArray.NumPages() == 1);
+			CHECK(OtherSourceArray.NumPages() == 1);
+
+			OtherArray.Append(SourceArray);
+			CHECK(OtherArray.NumPages() == 1);
+			CHECK(OtherArray.Num() == 4);
+
+			OtherArray.Append(OtherSourceArray);
+			CHECK(OtherArray.NumPages() == 1);
+			CHECK(OtherArray.Num() == 8);
+
+			// Validate content
+			int32 CheckValue = 0;
+			for (const int32Token& Value : OtherArray)
+			{
+				CHECK(Value == CheckValue++);
+			}
+			CHECK(OtherArray.Num() == CheckValue);
+		}
+		// Multiple-page smaller to larger page size
+		{
+			ArrayType PagedArray;
+			const ArrayType SourceArray({0, 1, 2, 3, 4, 5, 6});
+			const OtherArrayType OtherSourceArray({7, 8, 9, 10, 11, 12, 13, 14, 15});
+			CHECK(SourceArray.NumPages() == 2);
+			CHECK(OtherSourceArray.NumPages() == 2);
+
+			PagedArray.Append(SourceArray);
+			CHECK(PagedArray.NumPages() == 2);
+			CHECK(PagedArray.Num() == 7);
+
+			PagedArray.Append(OtherSourceArray);
+			CHECK(PagedArray.NumPages() == 4);
+			CHECK(PagedArray.Num() == 16);
+
+			// Validate content
+			int32 CheckValue = 0;
+			for (const int32Token& Value : PagedArray)
+			{
+				CHECK(Value == CheckValue++);
+			}
+			CHECK(PagedArray.Num() == CheckValue);
+		}
+		// Multiple-page larger to smaller page size
+		{
+			OtherArrayType OtherArray;
+			const ArrayType SourceArray({0, 1, 2, 3, 4, 5});
+			const OtherArrayType OtherSourceArray({6, 7, 8, 9, 10, 11, 12, 13, 14, 15});
+			CHECK(SourceArray.NumPages() == 2);
+			CHECK(OtherSourceArray.NumPages() == 2);
+
+			OtherArray.Append(SourceArray);
+			CHECK(OtherArray.NumPages() == 1);
+			CHECK(OtherArray.Num() == 6);
+
+			OtherArray.Append(OtherSourceArray);
+			CHECK(OtherArray.NumPages() == 2);
+			CHECK(OtherArray.Num() == 16);
+
+			// Validate content
+			int32 CheckValue = 0;
+			for (const int32Token& Value : OtherArray)
+			{
+				CHECK(Value == CheckValue++);
+			}
+			CHECK(OtherArray.Num() == CheckValue);
+		}
+	}
+	CHECK(int32Token::EvenConstructionDestructionCalls((8 + 16) * 6));
+}
+
+TEST_CASE("System::Core::Containers::TPagedArray::Iterator arithmetic", "[SmokeFilter][Core][Containers][PagedArray]")
+{
+	using ArrayType = TPagedArray<int32Token, 16>;
+	int32Token::Reset();
+	{
+		ArrayType PagedArray;
+
+		// Single page
+		PagedArray.Emplace(13);
+		PagedArray.Emplace(42);
+		PagedArray.Emplace(19);
+		CHECK(PagedArray.Max() == ArrayType::MaxPerPage());
+
+		auto It = PagedArray.begin();
+		CHECK(**It == 13);
+		CHECK(*It == int32Token(13));
+		CHECK(It->Value == 13);
+		auto It2 = It++;
+		CHECK(It2 != It);
+		CHECK(It2->Value == 13);
+		CHECK(**It == 42);
+		CHECK(*It == int32Token(42));
+		CHECK(It->Value == 42);
+
+		PagedArray.Emplace(73);
+		PagedArray.Emplace(171);
+		CHECK(PagedArray.Max() == ArrayType::MaxPerPage() * 2);
+		It = PagedArray.begin();
+		It++;
+		It++;
+		It++;
+		It2 = It++;
+		CHECK(It2 != It);
+		CHECK(It2->Value == 73);
+		CHECK(**It == 171);
+		CHECK(*It == int32Token(171));
+		CHECK(It->Value == 171);
+	}
+	// 5 from container, 3 manual instantiations
+	CHECK(int32Token::EvenConstructionDestructionCalls(5 + 3));
+}
+
+}  // namespace PagedArray
+
+#endif
