Exposed new PlatformMisc functions to disable/enable the screen timeout feature.

diff --git a/Engine/Source/Runtime/Core/Public/GenericPlatform/GenericPlatformMisc.h b/Engine/Source/Runtime/Core/Public/GenericPlatform/GenericPlatformMisc.h
--- a/Engine/Source/Runtime/Core/Public/GenericPlatform/GenericPlatformMisc.h
+++ b/Engine/Source/Runtime/Core/Public/GenericPlatform/GenericPlatformMisc.h
@@ -1816,6 +1816,14 @@ public:
 		return false;
 	}
 
+	static void DisableScreenTimeout()
+	{
+	}
+
+	static void EnableScreenTimeout()
+	{
+	}
+
 protected:
 	/** Whether the user should be prompted to allow for a remote debugger to be attached */
 	static bool bShouldPromptForRemoteDebugging;
