Moved new PlatformMisc functions to disable/enable the screen timeout feature out of #if !UE_BUILD_SHIPPING block.

diff --git a/Engine/Source/Runtime/Core/Public/GenericPlatform/GenericPlatformMisc.h b/Engine/Source/Runtime/Core/Public/GenericPlatform/GenericPlatformMisc.h
--- a/Engine/Source/Runtime/Core/Public/GenericPlatform/GenericPlatformMisc.h
+++ b/Engine/Source/Runtime/Core/Public/GenericPlatform/GenericPlatformMisc.h
@@ -1807,21 +1807,21 @@ public:
 		return false;
 	}
 
-#if !UE_BUILD_SHIPPING
-	/**
-	 * Returns any platform specific warning messages we want printed on screen
-	 */
-	static bool GetPlatformScreenWarnings(TArray<FText>& PlatformScreenWarnings)
+	static void DisableScreenTimeout()
 	{
-		return false;
 	}
 
-	static void DisableScreenTimeout()
+	static void EnableScreenTimeout()
 	{
 	}
 
-	static void EnableScreenTimeout()
+#if !UE_BUILD_SHIPPING
+	/**
+	 * Returns any platform specific warning messages we want printed on screen
+	 */
+	static bool GetPlatformScreenWarnings(TArray<FText>& PlatformScreenWarnings)
 	{
+		return false;
 	}
 
 protected:
