Fix hang caused by deletion of FGenericSizeRequest inside completion callback, which forces a wait that can't complete

diff --git a/Engine/Source/Runtime/Core/Private/Misc/PreLoadFile.cpp b/Engine/Source/Runtime/Core/Private/Misc/PreLoadFile.cpp
--- a/Engine/Source/Runtime/Core/Private/Misc/PreLoadFile.cpp
+++ b/Engine/Source/Runtime/Core/Private/Misc/PreLoadFile.cpp
@@ -74,7 +74,6 @@ void FPreLoadFile::KickOffRead()
 			{
 				Data = ReadRequest->GetReadResults();
 				CompletionEvent->Trigger();
-				delete AsyncReadHandle;
 			};
 			AsyncReadHandle->ReadRequest(0, FileSize, EAsyncIOPriorityAndFlags::AIOP_High, &ReadCallbackFunction);
 		}
@@ -85,11 +84,7 @@ void FPreLoadFile::KickOffRead()
 
 			FileSize = -1;
 			CompletionEvent->Trigger();
-			delete AsyncReadHandle;
 		}
-
-		delete SizeRequestHandle;
-		SizeRequestHandle = nullptr;
 	};
 
 	AsyncReadHandle = FPlatformFileManager::Get().GetPlatformFile().OpenAsyncRead(*Path);
@@ -115,7 +110,6 @@ void FPreLoadFile::KickOffRead()
 		}
 		// it's possible pak files with the file weren't mounted yet, so note that we didn't find it, and try again in TakeOwnership time
 		bFailedToOpenInKickOff = true;
-
 	}
 #endif
 }	
@@ -146,7 +140,14 @@ void* FPreLoadFile::TakeOwnershipOfLoadedData(int64* OutFileSize)
 			CompletionEvent->Wait();
 		}
 	}
-	
+
+#if PLATFORM_CAN_ASYNC_PRELOAD_FILES
+	delete SizeRequestHandle;
+	delete AsyncReadHandle;
+	SizeRequestHandle = nullptr;
+	AsyncReadHandle = nullptr;
+#endif
+
 	FPlatformProcess::ReturnSynchEventToPool(CompletionEvent);
 	CompletionEvent = nullptr;
 
