UE4 namespace fix.

diff --git a/Engine/Source/Runtime/Core/Private/Delegates/DelegateHandle.cpp b/Engine/Source/Runtime/Core/Private/Delegates/DelegateHandle.cpp
--- a/Engine/Source/Runtime/Core/Private/Delegates/DelegateHandle.cpp
+++ b/Engine/Source/Runtime/Core/Private/Delegates/DelegateHandle.cpp
@@ -4,7 +4,7 @@
 #include "Delegates/IDelegateInstance.h"
 #include "Templates/Atomic.h"
 
-namespace UE4Delegates_Private
+namespace UE::Delegates::Private
 {
 	TAtomic<uint64> GNextID(1);
 }
@@ -12,13 +12,13 @@ namespace UE4Delegates_Private
 uint64 FDelegateHandle::GenerateNewID()
 {
 	// Just increment a counter to generate an ID.
-	uint64 Result = ++UE4Delegates_Private::GNextID;
+	uint64 Result = ++UE::Delegates::Private::GNextID;
 
 	// Check for the next-to-impossible event that we wrap round to 0, because we reserve 0 for null delegates.
 	if (Result == 0)
 	{
 		// Increment it again - it might not be zero, so don't just assign it to 1.
-		Result = ++UE4Delegates_Private::GNextID;
+		Result = ++UE::Delegates::Private::GNextID;
 	}
 
 	return Result;
diff --git a/Engine/Source/Runtime/Core/Public/Delegates/Delegate.h b/Engine/Source/Runtime/Core/Public/Delegates/Delegate.h
--- a/Engine/Source/Runtime/Core/Public/Delegates/Delegate.h
+++ b/Engine/Source/Runtime/Core/Public/Delegates/Delegate.h
@@ -406,7 +406,7 @@ class DynamicMulticastDelegateClassName : public TBaseDynamicMulticastDelegate<T
 
 #else
 
-	#define STATIC_FUNCTION_FNAME(str) UE4Delegates_Private::GetTrimmedMemberFunctionName(str)
+	#define STATIC_FUNCTION_FNAME(str) UE::Delegates::Private::GetTrimmedMemberFunctionName(str)
 
 #endif
 
@@ -452,7 +452,7 @@ class DynamicMulticastDelegateClassName : public TBaseDynamicMulticastDelegate<T
 #define IsAlreadyBound( UserObject, FuncName ) __Internal_IsAlreadyBound( UserObject, FuncName, STATIC_FUNCTION_FNAME( TEXT( #FuncName ) ) )
 
 
-namespace UE4Delegates_Private
+namespace UE::Delegates::Private
 {
 	/**
 	 * Returns the root function name from a string representing a member function pointer.
diff --git a/Engine/Source/Runtime/Core/Public/Delegates/DelegateInstancesImpl.h b/Engine/Source/Runtime/Core/Public/Delegates/DelegateInstancesImpl.h
--- a/Engine/Source/Runtime/Core/Public/Delegates/DelegateInstancesImpl.h
+++ b/Engine/Source/Runtime/Core/Public/Delegates/DelegateInstancesImpl.h
@@ -23,7 +23,7 @@ class FDelegateBase;
 class FDelegateHandle;
 enum class ESPMode : uint8;
 
-namespace UE4Delegates_Private
+namespace UE::Delegates::Private
 {
 	constexpr bool IsUObjectPtr(const volatile UObjectBase*) { return true; }
 	constexpr bool IsUObjectPtr(...)                         { return false; }
@@ -75,7 +75,7 @@ private:
 	using UnwrappedThisType = TBaseUFunctionDelegateInstance<UserClass, RetValType(ParamTypes...), UserPolicy, VarTypes...>;
 	using DelegateBaseType = typename UserPolicy::FDelegateExtras;
 
-	static_assert(UE4Delegates_Private::IsUObjectPtr((UserClass*)nullptr), "You cannot use UFunction delegates with non UObject classes.");
+	static_assert(UE::Delegates::Private::IsUObjectPtr((UserClass*)nullptr), "You cannot use UFunction delegates with non UObject classes.");
 
 public:
 	TBaseUFunctionDelegateInstance(UserClass* InUserObject, const FName& InFunctionName, VarTypes... Vars)
@@ -367,7 +367,7 @@ template <bool bConst, class UserClass, typename WrappedRetValType, typename...
 class TBaseRawMethodDelegateInstance<bConst, UserClass, WrappedRetValType(ParamTypes...), UserPolicy, VarTypes...> : public TCommonDelegateInstanceState<WrappedRetValType(ParamTypes...), UserPolicy, VarTypes...>
 {
 private:
-	static_assert(!UE4Delegates_Private::IsUObjectPtr((UserClass*)nullptr), "You cannot use raw method delegates with UObjects.");
+	static_assert(!UE::Delegates::Private::IsUObjectPtr((UserClass*)nullptr), "You cannot use raw method delegates with UObjects.");
 
 	using Super             = TCommonDelegateInstanceState<WrappedRetValType(ParamTypes...), UserPolicy, VarTypes...>;
 	using RetValType        = typename Super::RetValType;
@@ -513,7 +513,7 @@ private:
 	using UnwrappedThisType = TBaseUObjectMethodDelegateInstance<bConst, UserClass, RetValType(ParamTypes...), UserPolicy, VarTypes...>;
 	using DelegateBaseType = typename UserPolicy::FDelegateExtras;
 
-	static_assert(UE4Delegates_Private::IsUObjectPtr((UserClass*)nullptr), "You cannot use UObject method delegates with raw pointers.");
+	static_assert(UE::Delegates::Private::IsUObjectPtr((UserClass*)nullptr), "You cannot use UObject method delegates with raw pointers.");
 
 public:
 	using FMethodPtr = typename TMemFunPtrType<bConst, UserClass, RetValType(ParamTypes..., VarTypes...)>::Type;
