[Insights] Improved performance of OutputBeginDynamicEvent by using a better hash function for the (dynamic scope name) hash maps.

Benchmarks/tests:
Using scope names from a large project (~220K string names), the average cost of adding an entry to TMap improved 2.4x (from 1.2us to 0.5us) and worst case improved 3.8x (from 80ms to 21ms).
Using Lyra (up to 5000 strings per TMap) the worst case improved 3.6x (from 11ms to 3ms).

diff --git a/Engine/Source/Runtime/Core/Private/ProfilingDebugging/CpuProfilerTrace.cpp b/Engine/Source/Runtime/Core/Private/ProfilingDebugging/CpuProfilerTrace.cpp
--- a/Engine/Source/Runtime/Core/Private/ProfilingDebugging/CpuProfilerTrace.cpp
+++ b/Engine/Source/Runtime/Core/Private/ProfilingDebugging/CpuProfilerTrace.cpp
@@ -61,7 +61,12 @@ struct FCpuProfilerTraceInternal
 
 		static FORCEINLINE uint32 GetKeyHash(const CharType* Key)
 		{
-			return FCrc::Strihash_DEPRECATED(Key);
+			uint32 Hash = 0;
+			for (const CharType* P = Key; *P; ++P)
+			{
+				Hash = ((Hash << 13) | (Hash >> 19)) ^ uint32(*P);
+			}
+			return Hash;
 		}
 
 		static FORCEINLINE KeyInitType GetSetKey(ElementInitType Element)
