PR #8666: Fix crash when exporting files larger than 2GB

diff --git a/Engine/Source/Runtime/Core/Public/Serialization/BufferArchive.h b/Engine/Source/Runtime/Core/Public/Serialization/BufferArchive.h
--- a/Engine/Source/Runtime/Core/Public/Serialization/BufferArchive.h
+++ b/Engine/Source/Runtime/Core/Public/Serialization/BufferArchive.h
@@ -11,11 +11,16 @@
 /**
  * Buffer archiver.
  */
-class FBufferArchive : public FMemoryWriter, public TArray<uint8>
+template <int IndexSize>
+class TBufferArchive : public TMemoryWriter<IndexSize>, public TArray<uint8, TSizedDefaultAllocator<IndexSize>>
 {
+	static_assert(IndexSize == 32 || IndexSize == 64, "Only 32-bit and 64-bit index sizes supported");
+
 public:
-	FBufferArchive( bool bIsPersistent = false, const FName InArchiveName = NAME_None )
-	: FMemoryWriter( *static_cast<TArray<uint8>*>(this), bIsPersistent, false, InArchiveName )
+	using ArrayType = TArray<uint8, TSizedDefaultAllocator<IndexSize>>;
+
+	TBufferArchive( bool bIsPersistent = false, const FName InArchiveName = NAME_None )
+	: TMemoryWriter<IndexSize>( *static_cast<ArrayType*>(this), bIsPersistent, false, InArchiveName )
 	{}
 	/**
   	 * Returns the name of the Archive.  Useful for getting the name of the package a struct or object
@@ -23,11 +28,47 @@ public:
 	 *
 	 * This is overridden for the specific Archive Types
 	 **/
-	virtual FString GetArchiveName() const { return *FString::Printf( TEXT("FBufferArchive %s"), *ArchiveName.ToString()); }
+	virtual FString GetArchiveName() const
+	{
+		if constexpr (IndexSize == 64)
+		{
+			return TEXT("FBufferArchive64 ") + this->ArchiveName.ToString();
+		}
+		else
+		{
+			return TEXT("FBufferArchive ") + this->ArchiveName.ToString();
+		}
+	}
+};
+
+// FBufferArchive and FBufferArchive64 are implemented as derived classes rather than aliases
+// so that forward declarations will work.
+
+class FBufferArchive : public TBufferArchive<32>
+{
+	using Super = TBufferArchive<32>;
+
+public:
+	using Super::Super;
+};
+
+class FBufferArchive64 : public TBufferArchive<64>
+{
+	using Super = TBufferArchive<64>;
+
+public:
+	using Super::Super;
 };
 
+
 template <>
 struct TIsContiguousContainer<FBufferArchive>
 {
-	static constexpr bool Value = TIsContiguousContainer<TArray<uint8>>::Value;
+	static constexpr bool Value = TIsContiguousContainer<FBufferArchive::ArrayType>::Value;
+};
+
+template <>
+struct TIsContiguousContainer<FBufferArchive64>
+{
+	static constexpr bool Value = TIsContiguousContainer<FBufferArchive64::ArrayType>::Value;
 };
diff --git a/Engine/Source/Runtime/Engine/Private/UnrealExporter.cpp b/Engine/Source/Runtime/Engine/Private/UnrealExporter.cpp
--- a/Engine/Source/Runtime/Engine/Private/UnrealExporter.cpp
+++ b/Engine/Source/Runtime/Engine/Private/UnrealExporter.cpp
@@ -374,7 +374,7 @@ bool UExporter::RunAssetExportTask(class UAssetExportTask* Task)
 		const int32 FileCount = Exporter->GetFileCount(Task->Object);
 		for( int32 i = 0; i < FileCount; i++ )
 		{
-			FBufferArchive Buffer;
+			FBufferArchive64 Buffer;
 			if(ExportToArchive(Task->Object, Exporter, Buffer, *Extension, i))
 			{
 				FString UniqueFilename = Exporter->GetUniqueFilename(Task->Object, *Task->Filename, i, FileCount);
