Updated TPagedArray:
* Added Append/Assign std::initializer_list overloads.
* Implemented std::initializer_list assignment operator overload in term of Assign.
* Fixed Append debug assertion: it now properly uses Reserve instead of Grow: added a comment to Grow as this method expects the container capacity to be smaller than the value requested.
* Added unit tests for both Append and Emplace that force the container to Grow.

#tests passed PagedArray unit test in Debug and Test: verified no assertions are triggered.

diff --git a/Engine/Source/Runtime/Core/Public/Containers/PagedArray.h b/Engine/Source/Runtime/Core/Public/Containers/PagedArray.h
--- a/Engine/Source/Runtime/Core/Public/Containers/PagedArray.h
+++ b/Engine/Source/Runtime/Core/Public/Containers/PagedArray.h
@@ -178,8 +178,7 @@ public:
 
 	TPagedArray& operator=(std::initializer_list<ElementType> InList)
 	{
-		Reset();
-		CopyToEmpty(InList.begin(), static_cast<SizeType>(InList.size()));
+		Assign(InList);
 		return *this;
 	}
 
@@ -293,10 +292,15 @@ public:
 	void Append(const ElementType* InSource, SizeType InSize)
 	{
 		check(InSource || !InSize);
-		Grow(Count + InSize);
+		GrowIfRequired(Count + InSize);
 		CopyUnchecked(InSource, InSize);
 	}
 
+	void Append(std::initializer_list<ElementType> InList)
+	{
+		Append(InList.begin(), static_cast<SizeType>(InList.size()));
+	}
+
 	template <typename ContainerType, std::enable_if_t<TIsContiguousContainer<ContainerType>::Value>* = nullptr>
 	FORCEINLINE void Append(ContainerType&& Container)
 	{
@@ -312,7 +316,7 @@ public:
 	{
 		using OtherType = TPagedArray<ElementType, OtherPageSizeInBytes, OtherAllocator>;
 		using OtherPageType = typename OtherType::PageType;
-		Grow(Count + Other.Count);
+		GrowIfRequired(Count + Other.Count);
 		for (const OtherPageType& Page : Other.Pages)
 		{
 			CopyUnchecked(Page.GetData(), Page.Num());
@@ -329,6 +333,12 @@ public:
 		Append(InSource, InSize);
 	}
 
+	void Assign(std::initializer_list<ElementType> InList)
+	{
+		Reset();
+		Append(InList.begin(), static_cast<SizeType>(InList.size()));
+	}
+
 	template <typename ContainerType, std::enable_if_t<TIsContiguousContainer<ContainerType>::Value>* = nullptr>
 	void Assign(ContainerType&& Container)
 	{
@@ -454,12 +464,10 @@ public:
 	/*
 	 * Reserves storage for the parameter element count.
 	 */
-	void Reserve(SizeType InCount)
+	FORCEINLINE void Reserve(SizeType InCount)
 	{
-		if (Max() < InCount)
-		{
-			Grow(InCount);
-		}
+		check(InCount >= 0);
+		GrowIfRequired(InCount);
 	}
 
 	/*
@@ -610,6 +618,7 @@ private:
 
 	/*
 	 * Grows the container's storage to the parameter capacity value allocating the required page count.
+	 * This method assumes the container's current capacity is smaller than the parameter value.
 	 * This method preserves existing data.
 	 */
 	void Grow(SizeType InCapacity)
@@ -623,7 +632,8 @@ private:
 		}
 	}
 
-	/* Grows the container by adding a new page of the parameter size if full.
+	/*
+	 * Grows the container by adding a new page if full.
 	 * This method preserves existing data.
 	 */
 	void GrowIfRequired()
@@ -634,6 +644,18 @@ private:
 		}
 	}
 
+	/*
+	 * Grows the container's storage to meet the parameter capacity value.
+	 * This method preserves existing data.
+	 */
+	void GrowIfRequired(SizeType InCapacity)
+	{
+		if (Max() < InCapacity)
+		{
+			Grow(InCapacity);
+		}
+	}
+
 	FORCEINLINE void CheckValidIndex(SizeType Index) const
 	{
 		checkf((Index >= 0) & (Index < Count), TEXT("Parameter index %d exceeds container size %d"), Index, Count);
diff --git a/Engine/Source/Runtime/Core/Tests/Containers/PagedArrayTest.cpp b/Engine/Source/Runtime/Core/Tests/Containers/PagedArrayTest.cpp
--- a/Engine/Source/Runtime/Core/Tests/Containers/PagedArrayTest.cpp
+++ b/Engine/Source/Runtime/Core/Tests/Containers/PagedArrayTest.cpp
@@ -72,9 +72,7 @@ TEST_CASE(
 	CHECK(int32Token::EvenConstructionDestructionCalls(1));
 }
 
-TEST_CASE(
-	"System::Core::Containers::TPagedArray::Add",
-	"[SmokeFilter][Core][Containers][PagedArray]")
+TEST_CASE("System::Core::Containers::TPagedArray::Add", "[SmokeFilter][Core][Containers][PagedArray]")
 {
 	int32Token::Reset();
 	{
@@ -94,6 +92,33 @@ TEST_CASE(
 	CHECK(int32Token::EvenConstructionDestructionCalls(4));
 }
 
+TEST_CASE(
+	"System::Core::Containers::TPagedArray::Emplace multiple elements",
+	"[SmokeFilter][Core][Containers][PagedArray]")
+{
+	int32Token::Reset();
+	{
+		static_assert(sizeof(int32Token) < 16);
+		using ArrayType = TPagedArray<int32Token, 16>;
+		ArrayType PagedArray;
+		CHECK(!PagedArray.Max());
+		CHECK(!PagedArray.Num());
+
+		const int32 MaxCount = 16 / sizeof(int32Token);
+		int32 CheckValue = 0;
+		while (CheckValue < MaxCount)
+		{
+			PagedArray.Emplace(CheckValue++);
+			CHECK(PagedArray.Max() == ArrayType::MaxPerPage());
+			CHECK(PagedArray.Num() == CheckValue);
+		}
+		PagedArray.Emplace(CheckValue++);
+		CHECK(PagedArray.Max() == ArrayType::MaxPerPage() * 2);
+		CHECK(PagedArray.Num() == CheckValue);
+	}
+	CHECK(int32Token::EvenConstructionDestructionCalls(5));
+}
+
 TEST_CASE("System::Core::Containers::TPagedArray::Reset", "[SmokeFilter][Core][Containers][PagedArray]")
 {
 	int32Token::Reset();
@@ -466,7 +491,7 @@ TEST_CASE("System::Core::Containers::TPagedArray::SetZero", "[SmokeFilter][Core]
 				Value = 13;
 			}
 			PagedArray.SetZero();
-	
+
 			// Verify content
 			CHECK(PagedArray.Num() == ArrayType::MaxPerPage() - 1);
 			for (const int32Token& Value : PagedArray)
@@ -477,7 +502,7 @@ TEST_CASE("System::Core::Containers::TPagedArray::SetZero", "[SmokeFilter][Core]
 		// Multiple pages
 		{
 			// Write content memory
-			PagedArray.SetNum(ArrayType::MaxPerPage() * 3 + 1); // This adds 2 pages + 2 elements
+			PagedArray.SetNum(ArrayType::MaxPerPage() * 3 + 1);	 // This adds 2 pages + 2 elements
 			for (int32Token& Value : PagedArray)
 			{
 				Value = 13;
@@ -489,7 +514,7 @@ TEST_CASE("System::Core::Containers::TPagedArray::SetZero", "[SmokeFilter][Core]
 			{
 				CHECK(Value == 0);
 			}
-			
+
 			// Verify content
 			CHECK(PagedArray.Num() == ArrayType::MaxPerPage() * 3 + 1);
 			for (const int32Token& Value : PagedArray)
@@ -501,7 +526,6 @@ TEST_CASE("System::Core::Containers::TPagedArray::SetZero", "[SmokeFilter][Core]
 	CHECK(int32Token::EvenConstructionDestructionCalls(ArrayType::MaxPerPage() * 7 + 1));
 }
 
-
 TEST_CASE("System::Core::Containers::TPagedArray::RemoveAtSwap", "[SmokeFilter][Core][Containers][PagedArray]")
 {
 	using ArrayType = TPagedArray<int32Token, 128>;
@@ -530,7 +554,7 @@ TEST_CASE("System::Core::Containers::TPagedArray::RemoveAtSwap", "[SmokeFilter][
 		}
 
 		// Test page shrinking after removing last element in page
-		PagedArray.RemoveAtSwap(0, false); 
+		PagedArray.RemoveAtSwap(0, false);
 		CHECK(PagedArray.Max() == ArrayType::MaxPerPage());
 		PagedArray.Emplace(0);
 		PagedArray.RemoveAtSwap(0);
@@ -827,7 +851,7 @@ TEST_CASE("System::Core::Containers::TPagedArray::Assign from TArray", "[SmokeFi
 
 			ArrayType PagedArray;
 			CHECK(PagedArray.IsEmpty());
-			PagedArray.Assign(TArray<int32Token>{}); // innocuous.
+			PagedArray.Assign(TArray<int32Token>{});  // innocuous.
 			PagedArray.Assign(SourceArray);
 			CHECK(PagedArray.Num() == 4);
 			CHECK(PagedArray.Max() == ArrayType::MaxPerPage());
@@ -915,7 +939,7 @@ TEST_CASE("System::Core::Containers::TPagedArray::Move to TArray", "[SmokeFilter
 		TArray<int32Token> DestinationArray;
 
 		// Empty copy
-		MoveTemp(PagedArray).ToArray(DestinationArray);  // innocuos
+		MoveTemp(PagedArray).ToArray(DestinationArray);	 // innocuos
 		CHECK(PagedArray.IsEmpty());
 		CHECK(DestinationArray.IsEmpty());
 
@@ -1044,8 +1068,8 @@ TEST_CASE(
 		// Single-page smaller to larger page size
 		{
 			ArrayType PagedArray;
-			const ArrayType SourceArray({ 0, 1, 2, 3 });
-			const OtherArrayType OtherSourceArray({ 4, 5, 6, 7 });
+			const ArrayType SourceArray({0, 1, 2, 3});
+			const OtherArrayType OtherSourceArray({4, 5, 6, 7});
 			CHECK(SourceArray.NumPages() == 1);
 			CHECK(OtherSourceArray.NumPages() == 1);
 
@@ -1141,6 +1165,38 @@ TEST_CASE(
 	CHECK(int32Token::EvenConstructionDestructionCalls((8 + 16) * 6));
 }
 
+TEST_CASE("System::Core::Containers::TPagedArray::Append multiple", "[SmokeFilter][Core][Containers][PagedArray]")
+{
+	using ArrayType = TPagedArray<int32Token, 16>;
+	int32Token::Reset();
+	{
+		ArrayType PagedArray;
+		PagedArray.Append({0, 1});
+		CHECK(PagedArray.NumPages() == 1);
+		CHECK(PagedArray.Num() == 2);
+
+		PagedArray.Append({2, 3});
+		CHECK(PagedArray.NumPages() == 1);
+		CHECK(PagedArray.Num() == 4);
+
+		for (int32 i = 4; i < 8; ++i)
+		{
+			PagedArray.Append({i});
+			CHECK(PagedArray.NumPages() == 2);
+			CHECK(PagedArray.Num() == i + 1);
+		}
+
+		// Validate content
+		int32 CheckValue = 0;
+		for (const int32Token& Value : PagedArray)
+		{
+			CHECK(Value == CheckValue++);
+		}
+		CHECK(PagedArray.Num() == CheckValue);
+	}
+	CHECK(int32Token::EvenConstructionDestructionCalls(8 * 2));
+}
+
 TEST_CASE("System::Core::Containers::TPagedArray::Iterator arithmetic", "[SmokeFilter][Core][Containers][PagedArray]")
 {
 	using ArrayType = TPagedArray<int32Token, 16>;
