TIdentity_T alias added.

diff --git a/Engine/Source/Runtime/Core/Public/Templates/Identity.h b/Engine/Source/Runtime/Core/Public/Templates/Identity.h
--- a/Engine/Source/Runtime/Core/Public/Templates/Identity.h
+++ b/Engine/Source/Runtime/Core/Public/Templates/Identity.h
@@ -10,9 +10,14 @@
  *
  * template <typename T>
  * void Func2(typename TIdentity<T>::Type Val); // Must be called like Func<int>(123)
+ *
+ * Equivalent to C++20's std::type_identity.
  */
 template <typename T>
 struct TIdentity
 {
 	typedef T Type;
 };
+
+template <typename T>
+using TIdentity_T = typename TIdentity<T>::Type;
