KahnTopologicalSort: Fix ordering of vertices in cycles when the graph contains multiple cycles that are not mutually reachable.
#rnx

diff --git a/Engine/Source/Runtime/Core/Private/Tests/Algo/TopologicalSort.cpp b/Engine/Source/Runtime/Core/Private/Tests/Algo/TopologicalSort.cpp
--- a/Engine/Source/Runtime/Core/Private/Tests/Algo/TopologicalSort.cpp
+++ b/Engine/Source/Runtime/Core/Private/Tests/Algo/TopologicalSort.cpp
@@ -16,27 +16,27 @@ bool FSystemCoreAlgoTopologicalSortTest::RunTest(const FString& Parameters)
 
 		// Test the sort when each node depends on the previous one
 		bool bHasSucceeded = TopologicalSort(Array, [](int32 Element) { return Element > 1 ? TArray<int32>{ Element - 1 } : TArray<int32>{}; });
-		TestEqual(TEXT("TopologicalSort did not sort correctly"), bHasSucceeded, true);
-		TestEqual(TEXT("TopologicalSort did not sort correctly"), Array, TArray<int32>{1, 2, 3});
+		TestEqual(TEXT("TopologicalSort returned true when each node depends on the previous one"), bHasSucceeded, true);
+		TestEqual(TEXT("TopologicalSort sorted correctly when each node depends on the previous one"), Array, TArray<int32>{1, 2, 3});
 	}
 	{
 		TArray<int32> Array{ 1, 2, 3 };
 		
 		// Test the sort when each node depends on the next one
 		bool bHasSucceeded = TopologicalSort(Array, [](int32 Element) { return Element < 3 ? TArray<int32>{ Element + 1 } : TArray<int32>{}; });
-		TestEqual(TEXT("TopologicalSort did not sort correctly"), bHasSucceeded, true);
-		TestEqual(TEXT("TopologicalSort did not sort correctly"), Array, TArray<int32>{3, 2, 1});
+		TestEqual(TEXT("TopologicalSort returned true when each node depends on the next one"), bHasSucceeded, true);
+		TestEqual(TEXT("TopologicalSort sorted correctly when each node depends on the next one"), Array, TArray<int32>{3, 2, 1});
 	}
 	{
 		TArray<int32> Array{ 1, 2 };
 		
 		// Test the sort with a cycle between 1 and 2
 		bool bHasSucceeded = TopologicalSort(Array, [](int32 Element) { return TArray<int32>{ 1 + Element % 2 }; }, ETopologicalSort::None);
-		TestEqual(TEXT("TopologicalSort should not have succeeded when a cycle is detected"), bHasSucceeded, false);
-		TestEqual(TEXT("TopologicalSort should not have modified the array when failing"), Array, TArray<int32>{1, 2});
+		TestEqual(TEXT("TopologicalSort returned false when a cycle is detected and AllowCycles is not specified"), bHasSucceeded, false);
+		TestEqual(TEXT("TopologicalSort does not modify the array when a cycle is detected and AllowCycles is not specified"), Array, TArray<int32>{1, 2});
 
 		bHasSucceeded = TopologicalSort(Array, [](int32 Element) { return TArray<int32>{ 1 + Element % 2 }; }, ETopologicalSort::AllowCycles);
-		TestEqual(TEXT("TopologicalSort should succeed when a cycle is detected but AllowCycles is specified"), bHasSucceeded, true);
+		TestEqual(TEXT("TopologicalSort returned true when a cycle is detected but AllowCycles is specified"), bHasSucceeded, true);
 	}
 	{
 		TArray<int32> Array;
@@ -47,8 +47,8 @@ bool FSystemCoreAlgoTopologicalSortTest::RunTest(const FString& Parameters)
 
 		// Make sure node 500 makes it on top if every other node depends on it
 		bool bHasSucceeded = TopologicalSort(Array, [](int32 Element) { return Element == 500 ? TArray<int32>{} : TArray<int32>{ 500 }; });
-		TestEqual(TEXT("TopologicalSort did not sort correctly"), bHasSucceeded, true);
-		TestEqual(TEXT("TopologicalSort did not sort correctly"), Array[0], 500);
+		TestEqual(TEXT("TopologicalSort returned true when every node has a dependency on a single node"), bHasSucceeded, true);
+		TestEqual(TEXT("TopologicalSort sorted correctly when every node has a dependency on a single node"), Array[0], 500);
 	}
 	{
 		TArray<int32> Array{ 1, 2, 3, 4, 5, 6, 7, 8, 9, 10 };
@@ -77,14 +77,14 @@ bool FSystemCoreAlgoTopologicalSortTest::RunTest(const FString& Parameters)
 		Links.Add(5, 2);
 
 		bool bHasSucceeded = TopologicalSort(Array, [&Links](int32 Element) { TArray<int32> Dependencies; Links.MultiFind(Element, Dependencies); return Dependencies; });
-		TestEqual(TEXT("TopologicalSort did not sort correctly"), bHasSucceeded, true);
+		TestEqual(TEXT("TopologicalSort returned true on sketched-out-example-1"), bHasSucceeded, true);
 
 		// There might be multiple valid answers, so test each condition separately to make sure they are all met.
 		for (const auto& Pair : Links)
 		{
 			int32 ChildIndex = Array.IndexOfByKey(Pair.Key);
 			int32 ParentIndex = Array.IndexOfByKey(Pair.Value);
-			TestEqual(TEXT("TopologicalSort did not sort correctly"), ParentIndex < ChildIndex, true);
+			TestEqual(TEXT("TopologicalSort sorted correctly on sketched-out-example-1"), ParentIndex < ChildIndex, true);
 		}
 	}
 
@@ -100,8 +100,9 @@ bool FSystemCoreAlgoTopologicalSortTest::RunTest(const FString& Parameters)
 		bool bHasSucceeded = TopologicalSort(Array,
 			[&Links](int32 Element) { TArray<int32> Dependencies; Links.MultiFind(Element, Dependencies); return Dependencies; },
 			ETopologicalSort::AllowCycles);
-		TestEqual(TEXT("TopologicalSort did not sort correctly"), bHasSucceeded, true);
-		TestTrue(TEXT("TopologicalSort did not sort correctly"),
+		TestEqual(TEXT("TopologicalSort returned true with a cycle in the root and with the root cycle depending on a chain of non-cycle verts"),
+			bHasSucceeded, true);
+		TestTrue(TEXT("TopologicalSort sorted correctly with a cycle in the root and with the root cycle depending on a chain of non-cycle verts"),
 			Array == TArray<int32>{4, 3, 2, 1} || Array == TArray<int32>{4, 3, 1, 2});
 	}
 	{
@@ -117,8 +118,9 @@ bool FSystemCoreAlgoTopologicalSortTest::RunTest(const FString& Parameters)
 			[&Links](int32 Element) { TArray<int32> Dependencies; Links.MultiFind(Element, Dependencies); return Dependencies; },
 			ETopologicalSort::AllowCycles);
 
-		TestEqual(TEXT("TopologicalSort did not sort correctly"), bHasSucceeded, true);
-		TestTrue(TEXT("TopologicalSort did not sort correctly"),
+		TestEqual(TEXT("TopologicalSort returned true with a cycle in the root and with the root cycle depending on a chain of non-cycle verts, submitted in revese"),
+			bHasSucceeded, true);
+		TestTrue(TEXT("TopologicalSort sorted correctly with a cycle in the root and with the root cycle depending on a chain of non-cycle verts, submitted in revese"),
 			Array == TArray<int32>{1, 2, 3, 4} || Array == TArray<int32>{1, 2, 4, 3});
 	}
 	{
@@ -134,13 +136,13 @@ bool FSystemCoreAlgoTopologicalSortTest::RunTest(const FString& Parameters)
 			[&Links](int32 Element) { TArray<int32> Dependencies; Links.MultiFind(Element, Dependencies); return Dependencies; },
 			ETopologicalSort::AllowCycles);
 
-		TestEqual(TEXT("TopologicalSort did not sort correctly"), bHasSucceeded, true);
+		TestEqual(TEXT("TopologicalSort returned true with a cycle at a leaf and a chain from the root depending on that cycle"), bHasSucceeded, true);
 		// There might be multiple valid answers, so test each condition separately to make sure they are all met.
-		TestTrue(TEXT("TopologicalSort did not sort correctly"), Array.IndexOfByKey(2) < Array.IndexOfByKey(1));
-		TestTrue(TEXT("TopologicalSort did not sort correctly"), Array[0] == 3 || Array[0] == 4);
+		TestTrue(TEXT("TopologicalSort sorted correctly with a cycle at a leaf and a chain from the root depending on that cycle"), Array.IndexOfByKey(2) < Array.IndexOfByKey(1));
+		TestTrue(TEXT("TopologicalSort sorted correctly with a cycle at a leaf and a chain from the root depending on that cycle"), Array[0] == 3 || Array[0] == 4);
 	}
 	{
-		// Test the sort with a cycle in the root and with the root cycle depending on a chain of non-cycle verts, submitted in revese
+		// Test the sort with a cycle in the root and with the root cycle depending on a chain of non-cycle verts, submitted in reverse
 		TArray<int32> Array{ 1, 2, 3, 4 };
 		TMultiMap<int32, int32> Links;
 		Links.Add(4, 3);
@@ -152,10 +154,13 @@ bool FSystemCoreAlgoTopologicalSortTest::RunTest(const FString& Parameters)
 			[&Links](int32 Element) { TArray<int32> Dependencies; Links.MultiFind(Element, Dependencies); return Dependencies; },
 			ETopologicalSort::AllowCycles);
 
-		TestEqual(TEXT("TopologicalSort did not sort correctly"), bHasSucceeded, true);
+		TestEqual(TEXT("TopologicalSort returned true with a cycle in the root and with the root cycle depending on a chain of non-cycle verts, submitted in reverse"),
+			bHasSucceeded, true);
 		// There might be multiple valid answers, so test each condition separately to make sure they are all met.
-		TestTrue(TEXT("TopologicalSort did not sort correctly"), Array.IndexOfByKey(3) < Array.IndexOfByKey(4));
-		TestTrue(TEXT("TopologicalSort did not sort correctly"), Array[0] == 1 || Array[0] == 2);
+		TestTrue(TEXT("TopologicalSort sorted correctly with a cycle in the root and with the root cycle depending on a chain of non-cycle verts, submitted in reverse"),
+			Array.IndexOfByKey(3) < Array.IndexOfByKey(4));
+		TestTrue(TEXT("TopologicalSort sorted correctly with a cycle in the root and with the root cycle depending on a chain of non-cycle verts, submitted in reverse"),
+			Array[0] == 1 || Array[0] == 2);
 	}
 	{
 		// Verify that when breaking a cycle a member of the cycle is selected rather than an element that depends on but is not in the cycle
@@ -174,9 +179,94 @@ bool FSystemCoreAlgoTopologicalSortTest::RunTest(const FString& Parameters)
 			[&Links](int32 Element) { TArray<int32> Dependencies; Links.MultiFind(Element, Dependencies); return Dependencies; },
 			ETopologicalSort::AllowCycles);
 
-		TestEqual(TEXT("TopologicalSort did not sort correctly"), bHasSucceeded, true);
+		TestEqual(TEXT("TopologicalSort returned true with a vertex dependent upon a cycle"), bHasSucceeded, true);
 		// There might be multiple valid answers, so test each condition separately to make sure they are all met.
-		TestTrue(TEXT("TopologicalSort did not sort correctly"), Array.IndexOfByKey(2) < Array.IndexOfByKey(1));
+		TestTrue(TEXT("TopologicalSort sorted correctly with a vertex dependent upon a cycle"), Array.IndexOfByKey(2) < Array.IndexOfByKey(1));
+	}
+	{
+		// Verify that when two cycles are present and one cycle depends on the other, the independent cycle is taken
+		TMultiMap<int32, int32> Links;
+		// 1 -> 2 -> 3 -> 4 -> 2
+		//           |
+		//           v
+		//      5 -> 6 -> 7 -> 5
+		Links.Add(1, 2);
+		Links.Add(2, 3);
+		Links.Add(3, 4);
+		Links.Add(3, 6);
+		Links.Add(4, 2);
+		Links.Add(5, 6);
+		Links.Add(6, 7);
+		Links.Add(7, 5);
+
+		int32 TrialIndex = 0;
+		while (TrialIndex >= 0)
+		{
+			TArray<int32> Array;
+			switch (TrialIndex++)
+			{
+			case 0:
+				Array = { 1, 2, 3, 4, 5, 6, 7 };
+				break;
+			case 1:
+				Array = { 7, 6, 5, 4, 3, 2, 1 };
+				break;
+			default:
+				TrialIndex = -1;
+				break;
+			};
+			if (TrialIndex < 0)
+			{
+				break;
+			}
+
+			bool bHasSucceeded = TopologicalSort(Array,
+				[&Links](int32 Element) { TArray<int32> Dependencies; Links.MultiFind(Element, Dependencies); return Dependencies; },
+				ETopologicalSort::AllowCycles);
+
+			TestEqual(TEXT("TopologicalSort returned true for one cycle dependent upon another"), bHasSucceeded, true);
+			// Vertex 1 comes last because it depends on the top cycle which depends on the bottom cycle
+			TestTrue(TEXT("TopologicalSort sorted correctly for one cycle dependent upon another"), Array.Last() == 1);
+			// Test that all elements of the bottom cycle come before all elements of the top cycle
+			int32 LastIndexOfBottom = FMath::Max3(Array.IndexOfByKey(5), Array.IndexOfByKey(6), Array.IndexOfByKey(7));
+			int32 FirstIndexOfTop = FMath::Min3(Array.IndexOfByKey(2), Array.IndexOfByKey(3), Array.IndexOfByKey(4));
+			TestTrue(TEXT("TopologicalSort sorted correctly for one cycle dependent upon another"), LastIndexOfBottom < FirstIndexOfTop);
+		}
+	}
+	{
+		// Verify that when a short cycle is embedded in a larger mutually reachable set, and that set has a
+		// dependency, from a vertex outside of the short cycle, to a separate mutually reachable set, that the more
+		// independent mutually reachable set is reported as more leaf-ward
+		TArray<int32> Array = { 0,1,2,3,4,5 };
+		TMultiMap<int32, int32> Links;
+		// 5 -> 0 -> (1 -> 2 -> 1)
+		//      |          |
+		//      |          v
+		//      |          5
+		//      v
+		//     (3 -> 4 -> 3)
+		Links.Add(0, 1);
+		Links.Add(1, 2);
+		Links.Add(2, 1);
+		// This edge creates the longer cycle and is later in the edge list from 2 and goes to a higher vertex than
+		// the edge that creates the shorter cycle 
+		Links.Add(2, 5);
+		// This edge that leads to the more independency cycle is later in the edge list from 0 than the edge list
+		// to the short cycle
+		Links.Add(0, 3);
+		Links.Add(3, 4);
+		Links.Add(4, 3);
+		Links.Add(5, 0);
+
+		bool bHasSucceeded = TopologicalSort(Array,
+			[&Links](int32 Element) { TArray<int32> Dependencies; Links.MultiFind(Element, Dependencies); return Dependencies; },
+			ETopologicalSort::AllowCycles);
+
+		TestEqual(TEXT("TopologicalSort returned true for MutuallyReachableSet Problem1"), bHasSucceeded, true);
+		// 3,4 comes first
+		TestTrue(TEXT("TopologicalSort sorted correctly for MutuallyReachableSet Problem1"),
+			(Array[0] == 3 && Array[1] == 4) || (Array[0] == 4 && Array[1] == 3));
+		// Any order of 0,1,2,5 is valid since they are mutually reachable, so there's nothing else to test
 	}
 
 	return true;
diff --git a/Engine/Source/Runtime/Core/Public/Algo/KahnTopologicalSort.h b/Engine/Source/Runtime/Core/Public/Algo/KahnTopologicalSort.h
--- a/Engine/Source/Runtime/Core/Public/Algo/KahnTopologicalSort.h
+++ b/Engine/Source/Runtime/Core/Public/Algo/KahnTopologicalSort.h
@@ -5,6 +5,7 @@
 #include "Algo/Sort.h"
 #include "Algo/Unique.h"
 #include "Containers/Array.h"
+#include "Containers/ArrayView.h"
 #include "Containers/Map.h"
 #include "Containers/Set.h"
 #include "Misc/EnumClassFlags.h"
@@ -20,6 +21,16 @@ namespace AlgoImpl
 	 */
 	typedef int32 FKahnHandle;
 
+	/** Scratch variables shared across multiple calls to FindMostIndependentMutuallyReachableVertexSet */
+	struct FMutuallyReachableVertexSetContext
+	{
+		TSet<FKahnHandle> VisitedReferences;
+		TSet<FKahnHandle> VisitedDependencies;
+		TSet<FKahnHandle> MaximalMRVS;
+		TSet<FKahnHandle> Culprits;
+		TArray<FKahnHandle> Stack;
+		TArray<FKahnHandle> Cycle;
+	};
 	/** Some variables shared with subfunctions */
 	struct FKahnContext
 	{
@@ -27,15 +38,14 @@ namespace AlgoImpl
 		TArray<TArray<FKahnHandle>> Referencers;
 		TArray<TArray<FKahnHandle>> Dependencies;
 		TArray<int32> DependencyCount;
-		TSet<FKahnHandle> CycleVisited;
-		TArray<FKahnHandle> CycleStack;
+		TOptional<FMutuallyReachableVertexSetContext> MRVSContext;
 	};
 
 	// Helper functions
 	template <typename RangeType, typename GetElementDependenciesType>
 	void KahnTopologicalSort_CreateWorkingGraph(FKahnContext& Context, RangeType& UniqueRange,
 		GetElementDependenciesType GetElementDependencies, TSet<FKahnHandle>& OutInitialIndependents);
-	FKahnHandle KahnTopologicalSort_PickMinimumCycleVertex(FKahnContext& Context);
+	const TSet<FKahnHandle>& FindMostIndependentMutuallyReachableVertexSet(FKahnContext& Context);
 }
 
 namespace Algo
@@ -87,11 +97,18 @@ namespace Algo
 					return false;
 				}
 
-				// In the presence of a cycle, pick a vertex that minimizes the number of unvisited dependencies.
-				FKahnHandle Vertex = KahnTopologicalSort_PickMinimumCycleVertex(Context);
-				IndependentVertices.Add(Vertex);
-				// Mark that we should no longer consider Vertex when looking for dependency-count==0 vertices.
-				DependencyCount[Vertex] = INDEX_NONE;
+				// When all remaining vertices are in cycles, find the maximal MutuallyReachableVertexSet
+				//  (a cycle is an MRVS) that is independent of all the other MRVS's.
+				const TSet<FKahnHandle>& MRVS = FindMostIndependentMutuallyReachableVertexSet(Context);
+				check(!MRVS.IsEmpty());
+				for (FKahnHandle Vertex : MRVS)
+				{
+					// Add all the vertices in the MutuallyReachableVertexSet to the SortedRange in arbitrary order
+					IndependentVertices.Add(Vertex);
+					// Mark that we should no longer consider any vertices in the MRVS when looking for NewIndependentVertices;
+					// they have already been removed
+					DependencyCount[Vertex] = INDEX_NONE;
+				}
 			}
 
 			NewIndependentVertices.Reset();
@@ -200,75 +217,164 @@ namespace AlgoImpl
 	}
 
 	/**
-	 * Called when there is a cycle (aka no vertices are independent). It finds a cycle and returns
-	 * one of the vertices with the minimum number of dependencies in that cycle.
+	 * Called when there is a MutuallyReachableVertexSet (aka no vertices are independent). It returns the most-independent
+	 * maximal MutuallyReachableVertexSet.
 	 */
-	inline FKahnHandle KahnTopologicalSort_PickMinimumCycleVertex(FKahnContext& Context)
+	inline const TSet<FKahnHandle>& FindMostIndependentMutuallyReachableVertexSet(FKahnContext& Context)
 	{
-		TSet<FKahnHandle>& CycleVisited = Context.CycleVisited;
-		TArray<FKahnHandle>& CycleStack = Context.CycleStack;
-		TArray<TArray<FKahnHandle>>& Dependencies = Context.Dependencies;
-		TSet<FKahnHandle>& RemainingVertices = Context.RemainingVertices;
-		TArray<int32>& DependencyCount = Context.DependencyCount;
-
-		int32 NumElements = Dependencies.Num();
-		CycleVisited.Reset();
-		CycleStack.Reset();
-		FKahnHandle Current = INDEX_NONE;
-		for (FKahnHandle RemainingVertex : RemainingVertices)
+		if (!Context.MRVSContext.IsSet())
 		{
-			Current = RemainingVertex;
-			break;
+			int32 NumVertices = Context.RemainingVertices.Num();
+			FMutuallyReachableVertexSetContext& InitContext = Context.MRVSContext.Emplace();
+			InitContext.VisitedReferences.Reserve(NumVertices);
+			InitContext.VisitedDependencies.Reserve(NumVertices);
+			InitContext.Culprits.Reserve(NumVertices);
+			InitContext.Stack.Reserve(NumVertices);
+			InitContext.Cycle.Reserve(NumVertices);
 		}
-		check(Current != INDEX_NONE);
+		const TArray<TArray<FKahnHandle>>& Dependencies = Context.Dependencies;
+		const TArray<TArray<FKahnHandle>>& Referencers = Context.Referencers;
+		FMutuallyReachableVertexSetContext& MRVSContext = *Context.MRVSContext;
+		TSet<FKahnHandle>& VisitedReferences = MRVSContext.VisitedReferences;
+		TSet<FKahnHandle>& VisitedDependencies = MRVSContext.VisitedDependencies;
+		TSet<FKahnHandle>& MaximalMRVS = MRVSContext.MaximalMRVS;
+		TSet<FKahnHandle>& Culprits = MRVSContext.Culprits;
+		TArray<FKahnHandle>& Stack = MRVSContext.Stack;
+		TArray<FKahnHandle>& Cycle = MRVSContext.Cycle;
 
-		// Find a cycle by arbitrarily following dependencies until we revisit a vertex.
-		CycleVisited.Add(Current);
-		CycleStack.Add(Current);
-		bool bAlreadyExists = false;
-		while (!bAlreadyExists)
+		// Copy Remaining Vertices into Culprits; we will be modifying the list
+		// as we search for the most-independent MRVS.
+		Culprits.Reset();
+		Culprits.Append(Context.RemainingVertices);
+
+		// Start from an arbitrary vertex
+		check(!Culprits.IsEmpty());
+		FKahnHandle StartingVertex = *Culprits.CreateIterator();
+
+		int32 StartingVertexLoopCount = 0;
+		while (StartingVertex != INDEX_NONE)
 		{
-			// Assert a dependency is found. PickMinimumCycleVertex is only called when every vertex has dependencies.
-			TArray<FKahnHandle>& ElementDependencies = Dependencies[Current];
-			FKahnHandle NextVertex = INDEX_NONE;
-			for (int32 Index = 0; Index < ElementDependencies.Num();)
+			// Find a cycle by arbitrarily following dependencies until we revisit a vertex.
+			VisitedReferences.Reset();
+			Stack.Reset();
+			VisitedReferences.Add(StartingVertex);
+			Stack.Add(StartingVertex);
+			FKahnHandle CycleWitnessVertex = StartingVertex;
+			bool bCycleWitnessWasAlreadyInStack = false;
+			while (!bCycleWitnessWasAlreadyInStack)
 			{
-				FKahnHandle TestVertex = ElementDependencies[Index];
-				if (RemainingVertices.Contains(TestVertex))
+				FKahnHandle NextVertex = INDEX_NONE;
+				for (FKahnHandle Dependency : Dependencies[CycleWitnessVertex])
 				{
-					NextVertex = TestVertex;
-					break;
+					if (Culprits.Contains(Dependency)) // Only consider edges to remaining vertices
+					{
+						NextVertex = Dependency;
+						break;
+					}
 				}
-				else
+				// Assert a dependency is found. This function is only called when every vertex has remaining dependencies.
+				check(NextVertex != INDEX_NONE);
+				CycleWitnessVertex = NextVertex;
+				VisitedReferences.Add(CycleWitnessVertex, &bCycleWitnessWasAlreadyInStack);
+				Stack.Add(CycleWitnessVertex);
+			}
+
+			// The cycle is everything on the stack after the first occurrence of CycleWitnessVertex.
+			// Cycle stacks will be like [7,7] or [7,3,2,7], always starting and ending with a vertex
+			// Pop the second occurrence of CycleWitnessVertex off the stack
+			check(Stack.Num() >= 2);
+			Stack.Pop(false /* bAllowShrinking */);
+			// Find the start of the Cycle
+			int32 StartIndex = Stack.Num() - 1;
+			while (StartIndex >= 0 && Stack[StartIndex] != CycleWitnessVertex)
+			{
+				--StartIndex;
+			}
+			check(StartIndex >= 0);
+			Cycle.Reset();
+			Cycle.Append(TArrayView<FKahnHandle>(Stack).Mid(StartIndex));
+
+			// We now have a cycle, which is an MRVS that may or may not be maximal. Expand it to a maximal MRVS by
+			// intersecting everything reachable from it with everything from which it is reachable.
+
+			// Find all references to the cycle. We start from VisitedReferences and Stack that we created when
+			// finding the cycle, since everything in that set is a referencer of the cycle
+			while (!Stack.IsEmpty())
+			{
+				FKahnHandle Vertex = Stack.Pop(false /* bAllowShrinking */);
+				for (FKahnHandle Referencer : Referencers[Vertex])
 				{
-					ElementDependencies.RemoveAtSwap(Index);
+					if (Culprits.Contains(Referencer)) // Only consider edges to remaining vertices
+					{
+						bool bAlreadyVisited;
+						VisitedReferences.Add(Referencer, &bAlreadyVisited);
+						if (!bAlreadyVisited)
+						{
+							Stack.Add(Referencer);
+						}
+					}
 				}
 			}
-			check(NextVertex != INDEX_NONE);
-			Current = NextVertex;
-			CycleVisited.Add(Current, &bAlreadyExists);
-			CycleStack.Add(Current);
-		}
 
-		// The cycle is everything on the stack after the first occurrence of Current.
-		// Pick the minimum-dependency-count of the cycle vertices.
-		int32 NumStack = CycleStack.Num();
-		int32 MinCount = DependencyCount[Current];
-		FKahnHandle MinVertex = Current;
-		for (int32 IndexInStack = NumStack - 2; IndexInStack >= 0; --IndexInStack)
-		{
-			FKahnHandle VertexInStack = CycleStack[IndexInStack];
-			if (VertexInStack == Current)
+			// Find all dependencies from the cycle.
+			VisitedDependencies.Reset();
+			Stack.Reset();
+			for (FKahnHandle Vertex : Cycle)
+			{
+				VisitedDependencies.Add(Vertex);
+				Stack.Add(Vertex);
+			}
+			while (!Stack.IsEmpty())
+			{
+				FKahnHandle Vertex = Stack.Pop(false /* bAllowShrinking */);
+				for (FKahnHandle Dependency : Dependencies[Vertex])
+				{
+					if (Culprits.Contains(Dependency)) // Only consider edges to remaining vertices
+					{
+						bool bAlreadyVisited;
+						VisitedDependencies.Add(Dependency, &bAlreadyVisited);
+						if (!bAlreadyVisited)
+						{
+							Stack.Add(Dependency);
+						}
+					}
+				}
+			}
+
+			MaximalMRVS = VisitedDependencies.Intersect(VisitedReferences);
+
+			// We now have a maximal MRVS, but there may be multiple maximal MRVS's, and we need to find the most independent
+			// one. Remove all elements of VisitedReferences from the remaining vertices in Culprits. We will not need to consider
+			// them when searching the graph for the more independent MRVS. Once removed, look for any dependencies from
+			// the MaximalMRVS to a remaining culprit. If any exist, then we can follow them to find a new maximal MRVS that
+			// is more independent than the current maximal MRVS.
+			for (FKahnHandle Vertex : VisitedReferences)
 			{
-				break;
+				Culprits.Remove(Vertex);
 			}
-			int32 Count = DependencyCount[VertexInStack];
-			if (Count < MinCount)
+
+			StartingVertex = INDEX_NONE;
+			for (FKahnHandle Vertex : MaximalMRVS)
 			{
-				MinCount = Count;
-				MinVertex = VertexInStack;
+				for (FKahnHandle Dependency : Dependencies[Vertex])
+				{
+					if (Culprits.Contains(Dependency)) // Only consider edges to remaining vertices
+					{
+						StartingVertex = Dependency;
+						break;
+					}
+				}
+				if (StartingVertex != INDEX_NONE)
+				{
+					break;
+				}
 			}
+
+			// If we found a new StartingVertex, the loop will continue and follow it to the new MaximalMRVS
+			// Otherwise the current MaximalMRVS is our returned value
+			checkf(StartingVertexLoopCount++ < Context.RemainingVertices.Num(), TEXT("Infinite loop in FindMostIndependentMutuallyReachableVertexSet"));
 		}
-		return MinVertex;
+
+		return MaximalMRVS;
 	}
-}
\ No newline at end of file
+}
diff --git a/Engine/Source/Runtime/Core/Public/Algo/TopologicalSort.h b/Engine/Source/Runtime/Core/Public/Algo/TopologicalSort.h
--- a/Engine/Source/Runtime/Core/Public/Algo/TopologicalSort.h
+++ b/Engine/Source/Runtime/Core/Public/Algo/TopologicalSort.h
@@ -9,7 +9,9 @@
 namespace Algo
 {
 	/**
-	 * Sorts the given range in leaf to root order.
+	 * Sorts the given range in leaf to root order:
+	 * For every pair of elements (A,B) where SortedIndex(A) < SortedIndex(B),
+	 * either B is not reachable from A, or A and B are mutually reachable.
 	 *
 	 * @param UniqueRange A range with element type T
 	 *        Type T must support GetTypeHash and copy+move constructors; T being pointertype is recommended.
@@ -19,8 +21,9 @@ namespace Algo
 	 *            const RangeType<T>& GetElementDependencies(const T& Element)
 	 *        It should return the leaf-ward vertices of directed edges from the root-wards Element.
 	 * @param Flags
-	 *        ETopologicalSort::AllowCycles: If present, cycles will be allowed; an arbitrary vertex in the cycle
-	 *            will be chosen first. If not present, the presence of a cycle will cause a failure to sort.
+	 *        ETopologicalSort::AllowCycles: If present, cycles will be allowed, with the caveat that the sorted
+	 *            ordering of MutuallyReachable vertices in the sorted range is arbitrary.
+	 *            If not present, the presence of a cycle will cause a failure to sort.
 	 * 
 	 * @return True if succeeded, false if failed due to cycles.
 	 */
