Thread-safe multicast delegates (locked by a mutex). Supports Broadcast() and Add()/Remove() on different threads and concurrently.

diff --git a/Engine/Source/Runtime/Core/Private/Tests/Delegates/TSDelegatesTest.cpp b/Engine/Source/Runtime/Core/Private/Tests/Delegates/TSDelegatesTest.cpp
new file mode 100644
--- /dev/null
+++ b/Engine/Source/Runtime/Core/Private/Tests/Delegates/TSDelegatesTest.cpp
@@ -0,0 +1,171 @@
+// Copyright Epic Games, Inc. All Rights Reserved.
+
+#include "CoreTypes.h"
+#include "Misc/AutomationTest.h"
+#include "Delegates/Delegate.h"
+#include "Tasks/Task.h"
+
+#if WITH_DEV_AUTOMATION_TESTS
+
+namespace UE { namespace TSDelegatesTest 
+{
+	using namespace Tasks;
+
+	void Func() {}
+
+	void OneParamFunc(bool* bExecuted)
+	{
+		*bExecuted = true;
+	}
+
+	void AnotherOneParamFunc(bool* bExecuted)
+	{
+		*bExecuted = true;
+	}
+
+	IMPLEMENT_SIMPLE_AUTOMATION_TEST(FTSMulticastDelegateBasicTest, "System.Core.Delegates.MulticastBasic", EAutomationTestFlags::ApplicationContextMask | EAutomationTestFlags::EngineFilter);
+
+	bool FTSMulticastDelegateBasicTest::RunTest(const FString& Parameters)
+	{
+		DECLARE_TS_MULTICAST_DELEGATE(FTestDelegate);
+
+		{	// single thread
+			FTestDelegate TestDelegate;
+			bool bExecuted = false;
+			FDelegateHandle DelegateHandle = TestDelegate.AddStatic(&OneParamFunc, &bExecuted);
+			TestDelegate.AddStatic(&AnotherOneParamFunc, &bExecuted);
+			TestDelegate.Broadcast();
+			check(bExecuted);
+			TestDelegate.Remove(DelegateHandle);
+		}
+
+		{	// concurrent Add/Remove and Execute
+			FTestDelegate TestDelegate;
+			bool bExecuted = false;
+			FTask AddStaticTask = Launch(UE_SOURCE_LOCATION, [&TestDelegate, &bExecuted] { TestDelegate.AddStatic(&AnotherOneParamFunc, &bExecuted); });
+			FTask BroadcastTask = Launch(UE_SOURCE_LOCATION, [&TestDelegate] { TestDelegate.Broadcast(); });
+			FDelegateHandle DelegateHandle = TestDelegate.AddStatic(&OneParamFunc, &bExecuted);
+			TestDelegate.Remove(DelegateHandle);
+			Wait(TArray{ AddStaticTask, BroadcastTask });
+		}
+
+		{ // same as above, but Add/Remove from different threads (not concurrently)
+			FTestDelegate TestDelegate;
+			bool bExecuted = false;
+			TTask<FDelegateHandle> AddStaticTask = Launch(UE_SOURCE_LOCATION, [&TestDelegate, &bExecuted] { return TestDelegate.AddStatic(&OneParamFunc, &bExecuted); });
+			FTask RemoveTask = Launch(UE_SOURCE_LOCATION, [&TestDelegate, &AddStaticTask] { TestDelegate.Remove(AddStaticTask.GetResult()); }, AddStaticTask);
+			TestDelegate.Broadcast();
+			RemoveTask.Wait();
+			TestDelegate.Remove(AddStaticTask.GetResult());
+		}
+
+		{	// recursive locking: removing a delegate handle from inside broadcasting
+			FTestDelegate TestDelegate;
+			FDelegateHandle DelegateHandle;
+			auto DelegateLambda = [&TestDelegate, &DelegateHandle] { TestDelegate.Remove(DelegateHandle); };
+			DelegateHandle = TestDelegate.AddLambda(MoveTemp(DelegateLambda));
+			check(TestDelegate.IsBound());
+			TestDelegate.Broadcast();
+			check(!TestDelegate.IsBound());
+		}
+
+		// test that all macros compile
+		{
+			DECLARE_TS_MULTICAST_DELEGATE_OneParam(FDelegate, int);
+			FDelegate Delegate;
+			Delegate.AddLambda([](int) {});
+			Delegate.Broadcast(0);
+		}
+		{
+			DECLARE_TS_MULTICAST_DELEGATE_TwoParams(FDelegate, int, int);
+			FDelegate Delegate;
+			Delegate.AddLambda([](int, int) {});
+			Delegate.Broadcast(0, 0);
+		}
+		{
+			DECLARE_TS_MULTICAST_DELEGATE_ThreeParams(FDelegate, int, int, int);
+			FDelegate Delegate;
+			Delegate.AddLambda([](int, int, int) {});
+			Delegate.Broadcast(0, 0, 0);
+		}
+		{
+			DECLARE_TS_MULTICAST_DELEGATE_FourParams(FDelegate, int, int, int, int);
+			FDelegate Delegate;
+			Delegate.AddLambda([](int, int, int, int) {});
+			Delegate.Broadcast(0, 0, 0, 0);
+		}
+		{
+			DECLARE_TS_MULTICAST_DELEGATE_FiveParams(FDelegate, int, int, int, int, int);
+			FDelegate Delegate;
+			Delegate.AddLambda([](int, int, int, int, int) {});
+			Delegate.Broadcast(0, 0, 0, 0, 0);
+		}
+		{
+			DECLARE_TS_MULTICAST_DELEGATE_SixParams(FDelegate, int, int, int, int, int, int);
+			FDelegate Delegate;
+			Delegate.AddLambda([](int, int, int, int, int, int) {});
+			Delegate.Broadcast(0, 0, 0, 0, 0, 0);
+		}
+		{
+			DECLARE_TS_MULTICAST_DELEGATE_SevenParams(FDelegate, int, int, int, int, int, int, int);
+			FDelegate Delegate;
+			Delegate.AddLambda([](int, int, int, int, int, int, int) {});
+			Delegate.Broadcast(0, 0, 0, 0, 0, 0, 0);
+		}
+		{
+			DECLARE_TS_MULTICAST_DELEGATE_EightParams(FDelegate, int, int, int, int, int, int, int, int);
+			FDelegate Delegate;
+			Delegate.AddLambda([](int, int, int, int, int, int, int, int) {});
+			Delegate.Broadcast(0, 0, 0, 0, 0, 0, 0, 0);
+		}
+		{
+			DECLARE_TS_MULTICAST_DELEGATE_NineParams(FDelegate, int, int, int, int, int, int, int, int, int);
+			FDelegate Delegate;
+			Delegate.AddLambda([](int, int, int, int, int, int, int, int, int) {});
+			Delegate.Broadcast(0, 0, 0, 0, 0, 0, 0, 0, 0);
+		}
+
+		return true;
+	}
+
+	IMPLEMENT_SIMPLE_AUTOMATION_TEST(FTSMulticastDelegateStressTest, "System.Core.Delegates.MulticastStress", EAutomationTestFlags::ApplicationContextMask | EAutomationTestFlags::EngineFilter);
+
+	bool FTSMulticastDelegateStressTest::RunTest(const FString& Parameters)
+	{
+		DECLARE_TS_MULTICAST_DELEGATE(FTestDelegate);
+
+		{	// one thread repeatedly adds/removes delegates, another broadcasts
+			FTestDelegate Delegate;
+			std::atomic<bool> bQuit{ false };
+
+			FTask Binding = Launch(UE_SOURCE_LOCATION,
+				[&bQuit, &Delegate]
+				{
+					while (!bQuit)
+					{
+						FDelegateHandle Handle = Delegate.AddStatic(&Func);
+						Delegate.Remove(Handle);
+					}
+				}
+			);
+
+			FTask Executing = Launch(UE_SOURCE_LOCATION,
+				[&bQuit, &Delegate]
+				{
+					while (!bQuit)
+					{
+						Delegate.Broadcast();
+					}
+				}
+			);
+
+			FPlatformProcess::Sleep(1.0);
+			bQuit = true;
+			Wait(TArray<FTask>{ Binding, Executing });
+		}
+
+		return true;
+	}
+}}
+
+#endif
\ No newline at end of file
diff --git a/Engine/Source/Runtime/Core/Public/Delegates/Delegate.h b/Engine/Source/Runtime/Core/Public/Delegates/Delegate.h
--- a/Engine/Source/Runtime/Core/Public/Delegates/Delegate.h
+++ b/Engine/Source/Runtime/Core/Public/Delegates/Delegate.h
@@ -8,6 +8,7 @@
 #include "Templates/SharedPointer.h"
 #include "UObject/WeakObjectPtrTemplates.h"
 #include "Delegates/MulticastDelegateBase.h" // IWYU pragma: export
+#include "Delegates/TSMulticastDelegateBase.h"
 #include "Delegates/IntegerSequence.h" // IWYU pragma: export
 
 /**
@@ -210,6 +211,10 @@
 #define FUNC_DECLARE_MULTICAST_DELEGATE( MulticastDelegateName, ReturnType, ... ) \
 	typedef TMulticastDelegate<ReturnType(__VA_ARGS__)> MulticastDelegateName;
 
+  /** Declares a broadcast thread-safe delegate that can bind to multiple native functions simultaneously */
+#define FUNC_DECLARE_TS_MULTICAST_DELEGATE( MulticastDelegateName, ReturnType, ... ) \
+	typedef TMulticastDelegate<ReturnType(__VA_ARGS__), FDefaultTSDelegateUserPolicy> MulticastDelegateName;
+
 /**
  * Declares a multicast delegate that is meant to only be activated from OwningType
  *
@@ -494,6 +499,7 @@ namespace UE::Delegates::Private
 // Simple delegate used by various utilities such as timers
 DECLARE_DELEGATE( FSimpleDelegate );
 DECLARE_MULTICAST_DELEGATE( FSimpleMulticastDelegate );
+DECLARE_TS_MULTICAST_DELEGATE( FTSSimpleMulticastDelegate );
 
 // Legacy typedefs
 template <typename RetType, typename... ArgTypes>
diff --git a/Engine/Source/Runtime/Core/Public/Delegates/DelegateBase.h b/Engine/Source/Runtime/Core/Public/Delegates/DelegateBase.h
--- a/Engine/Source/Runtime/Core/Public/Delegates/DelegateBase.h
+++ b/Engine/Source/Runtime/Core/Public/Delegates/DelegateBase.h
@@ -32,6 +32,9 @@ struct FWeakObjectPtr;
 template <typename UserPolicy>
 class TMulticastDelegateBase;
 
+template <typename UserPolicy>
+class TTSMulticastDelegateBase;
+
 ALIAS_TEMPLATE_TYPE_LAYOUT(template<typename ElementType>, FDelegateAllocatorType::ForElementType<ElementType>, void*);
 
 /**
@@ -167,6 +170,14 @@ struct FDefaultDelegateUserPolicy
 	using FMulticastDelegateExtras = TMulticastDelegateBase<FDefaultDelegateUserPolicy>;
 };
 
+struct FDefaultTSDelegateUserPolicy
+{
+	// see `FDefaultDelegateUserPolicy` for documentation
+	using FDelegateInstanceExtras = IDelegateInstance;
+	using FDelegateExtras = FDelegateBase;
+	using FMulticastDelegateExtras = TTSMulticastDelegateBase<FDefaultTSDelegateUserPolicy>;
+};
+
 template <typename UserPolicy>
 class TDelegateBase : public UserPolicy::FDelegateExtras
 {
diff --git a/Engine/Source/Runtime/Core/Public/Delegates/DelegateCombinations.h b/Engine/Source/Runtime/Core/Public/Delegates/DelegateCombinations.h
--- a/Engine/Source/Runtime/Core/Public/Delegates/DelegateCombinations.h
+++ b/Engine/Source/Runtime/Core/Public/Delegates/DelegateCombinations.h
@@ -22,6 +22,9 @@
 /** Declares a broadcast delegate that can bind to multiple native functions simultaneously */
 #define DECLARE_MULTICAST_DELEGATE( DelegateName ) FUNC_DECLARE_MULTICAST_DELEGATE( DelegateName, void )
 
+/** Declares a broadcast thread-safe delegate that can bind to multiple native functions simultaneously */
+#define DECLARE_TS_MULTICAST_DELEGATE( DelegateName ) FUNC_DECLARE_TS_MULTICAST_DELEGATE( DelegateName, void )
+
 /**
  * Declares a multicast delegate that is meant to only be activated from OwningType
  * NOTE: This behavior is not enforced and this type should be considered deprecated for new delegates, use normal multicast instead
@@ -44,6 +47,7 @@
 // Multiple-parameter versions of above delegate types:
 #define DECLARE_DELEGATE_OneParam( DelegateName, Param1Type ) FUNC_DECLARE_DELEGATE( DelegateName, void, Param1Type )
 #define DECLARE_MULTICAST_DELEGATE_OneParam( DelegateName, Param1Type ) FUNC_DECLARE_MULTICAST_DELEGATE( DelegateName, void, Param1Type )
+#define DECLARE_TS_MULTICAST_DELEGATE_OneParam( DelegateName, Param1Type ) FUNC_DECLARE_TS_MULTICAST_DELEGATE( DelegateName, void, Param1Type )
 #define DECLARE_EVENT_OneParam( OwningType, EventName, Param1Type ) FUNC_DECLARE_EVENT( OwningType, EventName, void, Param1Type )
 #define DECLARE_DYNAMIC_DELEGATE_OneParam( DelegateName, Param1Type, Param1Name ) BODY_MACRO_COMBINE(CURRENT_FILE_ID,_,__LINE__,_DELEGATE) FUNC_DECLARE_DYNAMIC_DELEGATE( FWeakObjectPtr, DelegateName, DelegateName##_DelegateWrapper, FUNC_CONCAT( Param1Type InParam1 ), FUNC_CONCAT( *this, InParam1 ), void, Param1Type )
 #define DECLARE_DYNAMIC_MULTICAST_DELEGATE_OneParam( DelegateName, Param1Type, Param1Name ) BODY_MACRO_COMBINE(CURRENT_FILE_ID,_,__LINE__,_DELEGATE) FUNC_DECLARE_DYNAMIC_MULTICAST_DELEGATE( FWeakObjectPtr, DelegateName, DelegateName##_DelegateWrapper, FUNC_CONCAT( Param1Type InParam1 ), FUNC_CONCAT( *this, InParam1 ), void, Param1Type )
@@ -52,6 +56,7 @@
 
 #define DECLARE_DELEGATE_TwoParams( DelegateName, Param1Type, Param2Type ) FUNC_DECLARE_DELEGATE( DelegateName, void, Param1Type, Param2Type )
 #define DECLARE_MULTICAST_DELEGATE_TwoParams( DelegateName, Param1Type, Param2Type ) FUNC_DECLARE_MULTICAST_DELEGATE( DelegateName, void, Param1Type, Param2Type )
+#define DECLARE_TS_MULTICAST_DELEGATE_TwoParams( DelegateName, Param1Type, Param2Type ) FUNC_DECLARE_TS_MULTICAST_DELEGATE( DelegateName, void, Param1Type, Param2Type )
 #define DECLARE_EVENT_TwoParams( OwningType, EventName, Param1Type, Param2Type ) FUNC_DECLARE_EVENT( OwningType, EventName, void, Param1Type, Param2Type )
 #define DECLARE_DYNAMIC_DELEGATE_TwoParams( DelegateName, Param1Type, Param1Name, Param2Type, Param2Name ) BODY_MACRO_COMBINE(CURRENT_FILE_ID,_,__LINE__,_DELEGATE) FUNC_DECLARE_DYNAMIC_DELEGATE( FWeakObjectPtr, DelegateName, DelegateName##_DelegateWrapper, FUNC_CONCAT( Param1Type InParam1, Param2Type InParam2 ), FUNC_CONCAT( *this, InParam1, InParam2 ), void, Param1Type, Param2Type )
 #define DECLARE_DYNAMIC_MULTICAST_DELEGATE_TwoParams( DelegateName, Param1Type, Param1Name, Param2Type, Param2Name ) BODY_MACRO_COMBINE(CURRENT_FILE_ID,_,__LINE__,_DELEGATE) FUNC_DECLARE_DYNAMIC_MULTICAST_DELEGATE( FWeakObjectPtr, DelegateName, DelegateName##_DelegateWrapper, FUNC_CONCAT( Param1Type InParam1, Param2Type InParam2 ), FUNC_CONCAT( *this, InParam1, InParam2 ), void, Param1Type, Param2Type )
@@ -60,6 +65,7 @@
 
 #define DECLARE_DELEGATE_ThreeParams( DelegateName, Param1Type, Param2Type, Param3Type ) FUNC_DECLARE_DELEGATE( DelegateName, void, Param1Type, Param2Type, Param3Type )
 #define DECLARE_MULTICAST_DELEGATE_ThreeParams( DelegateName, Param1Type, Param2Type, Param3Type ) FUNC_DECLARE_MULTICAST_DELEGATE( DelegateName, void, Param1Type, Param2Type, Param3Type )
+#define DECLARE_TS_MULTICAST_DELEGATE_ThreeParams( DelegateName, Param1Type, Param2Type, Param3Type ) FUNC_DECLARE_TS_MULTICAST_DELEGATE( DelegateName, void, Param1Type, Param2Type, Param3Type )
 #define DECLARE_EVENT_ThreeParams( OwningType, EventName, Param1Type, Param2Type, Param3Type ) FUNC_DECLARE_EVENT( OwningType, EventName, void, Param1Type, Param2Type, Param3Type )
 #define DECLARE_DYNAMIC_DELEGATE_ThreeParams( DelegateName, Param1Type, Param1Name, Param2Type, Param2Name, Param3Type, Param3Name ) BODY_MACRO_COMBINE(CURRENT_FILE_ID,_,__LINE__,_DELEGATE) FUNC_DECLARE_DYNAMIC_DELEGATE( FWeakObjectPtr, DelegateName, DelegateName##_DelegateWrapper, FUNC_CONCAT( Param1Type InParam1, Param2Type InParam2, Param3Type InParam3 ), FUNC_CONCAT( *this, InParam1, InParam2, InParam3 ), void, Param1Type, Param2Type, Param3Type )
 #define DECLARE_DYNAMIC_MULTICAST_DELEGATE_ThreeParams( DelegateName, Param1Type, Param1Name, Param2Type, Param2Name, Param3Type, Param3Name ) BODY_MACRO_COMBINE(CURRENT_FILE_ID,_,__LINE__,_DELEGATE) FUNC_DECLARE_DYNAMIC_MULTICAST_DELEGATE( FWeakObjectPtr, DelegateName, DelegateName##_DelegateWrapper, FUNC_CONCAT( Param1Type InParam1, Param2Type InParam2, Param3Type InParam3 ), FUNC_CONCAT( *this, InParam1, InParam2, InParam3 ), void, Param1Type, Param2Type, Param3Type )
@@ -68,6 +74,7 @@
 
 #define DECLARE_DELEGATE_FourParams( DelegateName, Param1Type, Param2Type, Param3Type, Param4Type ) FUNC_DECLARE_DELEGATE( DelegateName, void, Param1Type, Param2Type, Param3Type, Param4Type )
 #define DECLARE_MULTICAST_DELEGATE_FourParams( DelegateName, Param1Type, Param2Type, Param3Type, Param4Type ) FUNC_DECLARE_MULTICAST_DELEGATE( DelegateName, void, Param1Type, Param2Type, Param3Type, Param4Type )
+#define DECLARE_TS_MULTICAST_DELEGATE_FourParams( DelegateName, Param1Type, Param2Type, Param3Type, Param4Type ) FUNC_DECLARE_TS_MULTICAST_DELEGATE( DelegateName, void, Param1Type, Param2Type, Param3Type, Param4Type )
 #define DECLARE_EVENT_FourParams( OwningType, EventName, Param1Type, Param2Type, Param3Type, Param4Type ) FUNC_DECLARE_EVENT( OwningType, EventName, void, Param1Type, Param2Type, Param3Type, Param4Type )
 #define DECLARE_DYNAMIC_DELEGATE_FourParams( DelegateName, Param1Type, Param1Name, Param2Type, Param2Name, Param3Type, Param3Name, Param4Type, Param4Name ) BODY_MACRO_COMBINE(CURRENT_FILE_ID,_,__LINE__,_DELEGATE) FUNC_DECLARE_DYNAMIC_DELEGATE( FWeakObjectPtr, DelegateName, DelegateName##_DelegateWrapper, FUNC_CONCAT( Param1Type InParam1, Param2Type InParam2, Param3Type InParam3, Param4Type InParam4 ), FUNC_CONCAT( *this, InParam1, InParam2, InParam3, InParam4 ), void, Param1Type, Param2Type, Param3Type, Param4Type )
 #define DECLARE_DYNAMIC_MULTICAST_DELEGATE_FourParams( DelegateName, Param1Type, Param1Name, Param2Type, Param2Name, Param3Type, Param3Name, Param4Type, Param4Name ) BODY_MACRO_COMBINE(CURRENT_FILE_ID,_,__LINE__,_DELEGATE) FUNC_DECLARE_DYNAMIC_MULTICAST_DELEGATE( FWeakObjectPtr, DelegateName, DelegateName##_DelegateWrapper, FUNC_CONCAT( Param1Type InParam1, Param2Type InParam2, Param3Type InParam3, Param4Type InParam4 ), FUNC_CONCAT( *this, InParam1, InParam2, InParam3, InParam4 ), void, Param1Type, Param2Type, Param3Type, Param4Type )
@@ -76,6 +83,7 @@
 
 #define DECLARE_DELEGATE_FiveParams( DelegateName, Param1Type, Param2Type, Param3Type, Param4Type, Param5Type ) FUNC_DECLARE_DELEGATE( DelegateName, void, Param1Type, Param2Type, Param3Type, Param4Type, Param5Type )
 #define DECLARE_MULTICAST_DELEGATE_FiveParams( DelegateName, Param1Type, Param2Type, Param3Type, Param4Type, Param5Type ) FUNC_DECLARE_MULTICAST_DELEGATE( DelegateName, void, Param1Type, Param2Type, Param3Type, Param4Type, Param5Type )
+#define DECLARE_TS_MULTICAST_DELEGATE_FiveParams( DelegateName, Param1Type, Param2Type, Param3Type, Param4Type, Param5Type ) FUNC_DECLARE_TS_MULTICAST_DELEGATE( DelegateName, void, Param1Type, Param2Type, Param3Type, Param4Type, Param5Type )
 #define DECLARE_EVENT_FiveParams( OwningType, EventName, Param1Type, Param2Type, Param3Type, Param4Type, Param5Type ) FUNC_DECLARE_EVENT( OwningType, EventName, void, Param1Type, Param2Type, Param3Type, Param4Type, Param5Type )
 #define DECLARE_DYNAMIC_DELEGATE_FiveParams( DelegateName, Param1Type, Param1Name, Param2Type, Param2Name, Param3Type, Param3Name, Param4Type, Param4Name, Param5Type, Param5Name ) BODY_MACRO_COMBINE(CURRENT_FILE_ID,_,__LINE__,_DELEGATE) FUNC_DECLARE_DYNAMIC_DELEGATE( FWeakObjectPtr, DelegateName, DelegateName##_DelegateWrapper, FUNC_CONCAT( Param1Type InParam1, Param2Type InParam2, Param3Type InParam3, Param4Type InParam4, Param5Type InParam5 ), FUNC_CONCAT( *this, InParam1, InParam2, InParam3, InParam4, InParam5 ), void, Param1Type, Param2Type, Param3Type, Param4Type, Param5Type )
 #define DECLARE_DYNAMIC_MULTICAST_DELEGATE_FiveParams( DelegateName, Param1Type, Param1Name, Param2Type, Param2Name, Param3Type, Param3Name, Param4Type, Param4Name, Param5Type, Param5Name ) BODY_MACRO_COMBINE(CURRENT_FILE_ID,_,__LINE__,_DELEGATE) FUNC_DECLARE_DYNAMIC_MULTICAST_DELEGATE( FWeakObjectPtr, DelegateName, DelegateName##_DelegateWrapper, FUNC_CONCAT( Param1Type InParam1, Param2Type InParam2, Param3Type InParam3, Param4Type InParam4, Param5Type InParam5 ), FUNC_CONCAT( *this, InParam1, InParam2, InParam3, InParam4, InParam5 ), void, Param1Type, Param2Type, Param3Type, Param4Type, Param5Type )
@@ -84,6 +92,7 @@
 
 #define DECLARE_DELEGATE_SixParams( DelegateName, Param1Type, Param2Type, Param3Type, Param4Type, Param5Type, Param6Type ) FUNC_DECLARE_DELEGATE( DelegateName, void, Param1Type, Param2Type, Param3Type, Param4Type, Param5Type, Param6Type )
 #define DECLARE_MULTICAST_DELEGATE_SixParams( DelegateName, Param1Type, Param2Type, Param3Type, Param4Type, Param5Type, Param6Type ) FUNC_DECLARE_MULTICAST_DELEGATE( DelegateName, void, Param1Type, Param2Type, Param3Type, Param4Type, Param5Type, Param6Type )
+#define DECLARE_TS_MULTICAST_DELEGATE_SixParams( DelegateName, Param1Type, Param2Type, Param3Type, Param4Type, Param5Type, Param6Type ) FUNC_DECLARE_TS_MULTICAST_DELEGATE( DelegateName, void, Param1Type, Param2Type, Param3Type, Param4Type, Param5Type, Param6Type )
 #define DECLARE_EVENT_SixParams( OwningType, EventName, Param1Type, Param2Type, Param3Type, Param4Type, Param5Type, Param6Type ) FUNC_DECLARE_EVENT( OwningType, EventName, void, Param1Type, Param2Type, Param3Type, Param4Type, Param5Type, Param6Type )
 #define DECLARE_DYNAMIC_DELEGATE_SixParams( DelegateName, Param1Type, Param1Name, Param2Type, Param2Name, Param3Type, Param3Name, Param4Type, Param4Name, Param5Type, Param5Name, Param6Type, Param6Name ) BODY_MACRO_COMBINE(CURRENT_FILE_ID,_,__LINE__,_DELEGATE) FUNC_DECLARE_DYNAMIC_DELEGATE( FWeakObjectPtr, DelegateName, DelegateName##_DelegateWrapper, FUNC_CONCAT( Param1Type InParam1, Param2Type InParam2, Param3Type InParam3, Param4Type InParam4, Param5Type InParam5, Param6Type InParam6 ), FUNC_CONCAT( *this, InParam1, InParam2, InParam3, InParam4, InParam5, InParam6 ), void, Param1Type, Param2Type, Param3Type, Param4Type, Param5Type, Param6Type )
 #define DECLARE_DYNAMIC_MULTICAST_DELEGATE_SixParams( DelegateName, Param1Type, Param1Name, Param2Type, Param2Name, Param3Type, Param3Name, Param4Type, Param4Name, Param5Type, Param5Name, Param6Type, Param6Name ) BODY_MACRO_COMBINE(CURRENT_FILE_ID,_,__LINE__,_DELEGATE) FUNC_DECLARE_DYNAMIC_MULTICAST_DELEGATE( FWeakObjectPtr, DelegateName, DelegateName##_DelegateWrapper, FUNC_CONCAT( Param1Type InParam1, Param2Type InParam2, Param3Type InParam3, Param4Type InParam4, Param5Type InParam5, Param6Type InParam6 ), FUNC_CONCAT( *this, InParam1, InParam2, InParam3, InParam4, InParam5, InParam6 ), void, Param1Type, Param2Type, Param3Type, Param4Type, Param5Type, Param6Type )
@@ -92,6 +101,7 @@
 
 #define DECLARE_DELEGATE_SevenParams( DelegateName, Param1Type, Param2Type, Param3Type, Param4Type, Param5Type, Param6Type, Param7Type ) FUNC_DECLARE_DELEGATE( DelegateName, void, Param1Type, Param2Type, Param3Type, Param4Type, Param5Type, Param6Type, Param7Type )
 #define DECLARE_MULTICAST_DELEGATE_SevenParams( DelegateName, Param1Type, Param2Type, Param3Type, Param4Type, Param5Type, Param6Type, Param7Type ) FUNC_DECLARE_MULTICAST_DELEGATE( DelegateName, void, Param1Type, Param2Type, Param3Type, Param4Type, Param5Type, Param6Type, Param7Type )
+#define DECLARE_TS_MULTICAST_DELEGATE_SevenParams( DelegateName, Param1Type, Param2Type, Param3Type, Param4Type, Param5Type, Param6Type, Param7Type ) FUNC_DECLARE_TS_MULTICAST_DELEGATE( DelegateName, void, Param1Type, Param2Type, Param3Type, Param4Type, Param5Type, Param6Type, Param7Type )
 #define DECLARE_EVENT_SevenParams( OwningType, EventName, Param1Type, Param2Type, Param3Type, Param4Type, Param5Type, Param6Type, Param7Type ) FUNC_DECLARE_EVENT( OwningType, EventName, void, Param1Type, Param2Type, Param3Type, Param4Type, Param5Type, Param6Type, Param7Type )
 #define DECLARE_DYNAMIC_DELEGATE_SevenParams( DelegateName, Param1Type, Param1Name, Param2Type, Param2Name, Param3Type, Param3Name, Param4Type, Param4Name, Param5Type, Param5Name, Param6Type, Param6Name, Param7Type, Param7Name ) BODY_MACRO_COMBINE(CURRENT_FILE_ID,_,__LINE__,_DELEGATE) FUNC_DECLARE_DYNAMIC_DELEGATE( FWeakObjectPtr, DelegateName, DelegateName##_DelegateWrapper, FUNC_CONCAT( Param1Type InParam1, Param2Type InParam2, Param3Type InParam3, Param4Type InParam4, Param5Type InParam5, Param6Type InParam6, Param7Type InParam7 ), FUNC_CONCAT( *this, InParam1, InParam2, InParam3, InParam4, InParam5, InParam6, InParam7 ), void, Param1Type, Param2Type, Param3Type, Param4Type, Param5Type, Param6Type, Param7Type )
 #define DECLARE_DYNAMIC_MULTICAST_DELEGATE_SevenParams( DelegateName, Param1Type, Param1Name, Param2Type, Param2Name, Param3Type, Param3Name, Param4Type, Param4Name, Param5Type, Param5Name, Param6Type, Param6Name, Param7Type, Param7Name ) BODY_MACRO_COMBINE(CURRENT_FILE_ID,_,__LINE__,_DELEGATE) FUNC_DECLARE_DYNAMIC_MULTICAST_DELEGATE( FWeakObjectPtr, DelegateName, DelegateName##_DelegateWrapper, FUNC_CONCAT( Param1Type InParam1, Param2Type InParam2, Param3Type InParam3, Param4Type InParam4, Param5Type InParam5, Param6Type InParam6, Param7Type InParam7 ), FUNC_CONCAT( *this, InParam1, InParam2, InParam3, InParam4, InParam5, InParam6, InParam7 ), void, Param1Type, Param2Type, Param3Type, Param4Type, Param5Type, Param6Type, Param7Type )
@@ -100,6 +110,7 @@
 
 #define DECLARE_DELEGATE_EightParams( DelegateName, Param1Type, Param2Type, Param3Type, Param4Type, Param5Type, Param6Type, Param7Type, Param8Type ) FUNC_DECLARE_DELEGATE( DelegateName, void, Param1Type, Param2Type, Param3Type, Param4Type, Param5Type, Param6Type, Param7Type, Param8Type )
 #define DECLARE_MULTICAST_DELEGATE_EightParams( DelegateName, Param1Type, Param2Type, Param3Type, Param4Type, Param5Type, Param6Type, Param7Type, Param8Type ) FUNC_DECLARE_MULTICAST_DELEGATE( DelegateName, void, Param1Type, Param2Type, Param3Type, Param4Type, Param5Type, Param6Type, Param7Type, Param8Type )
+#define DECLARE_TS_MULTICAST_DELEGATE_EightParams( DelegateName, Param1Type, Param2Type, Param3Type, Param4Type, Param5Type, Param6Type, Param7Type, Param8Type ) FUNC_DECLARE_TS_MULTICAST_DELEGATE( DelegateName, void, Param1Type, Param2Type, Param3Type, Param4Type, Param5Type, Param6Type, Param7Type, Param8Type )
 #define DECLARE_EVENT_EightParams( OwningType, EventName, Param1Type, Param2Type, Param3Type, Param4Type, Param5Type, Param6Type, Param7Type, Param8Type ) FUNC_DECLARE_EVENT( OwningType, EventName, void, Param1Type, Param2Type, Param3Type, Param4Type, Param5Type, Param6Type, Param7Type, Param8Type )
 #define DECLARE_DYNAMIC_DELEGATE_EightParams( DelegateName, Param1Type, Param1Name, Param2Type, Param2Name, Param3Type, Param3Name, Param4Type, Param4Name, Param5Type, Param5Name, Param6Type, Param6Name, Param7Type, Param7Name, Param8Type, Param8Name ) BODY_MACRO_COMBINE(CURRENT_FILE_ID,_,__LINE__,_DELEGATE) FUNC_DECLARE_DYNAMIC_DELEGATE( FWeakObjectPtr, DelegateName, DelegateName##_DelegateWrapper, FUNC_CONCAT( Param1Type InParam1, Param2Type InParam2, Param3Type InParam3, Param4Type InParam4, Param5Type InParam5, Param6Type InParam6, Param7Type InParam7, Param8Type InParam8 ), FUNC_CONCAT( *this, InParam1, InParam2, InParam3, InParam4, InParam5, InParam6, InParam7, InParam8 ), void, Param1Type, Param2Type, Param3Type, Param4Type, Param5Type, Param6Type, Param7Type, Param8Type )
 #define DECLARE_DYNAMIC_MULTICAST_DELEGATE_EightParams( DelegateName, Param1Type, Param1Name, Param2Type, Param2Name, Param3Type, Param3Name, Param4Type, Param4Name, Param5Type, Param5Name, Param6Type, Param6Name, Param7Type, Param7Name, Param8Type, Param8Name ) BODY_MACRO_COMBINE(CURRENT_FILE_ID,_,__LINE__,_DELEGATE) FUNC_DECLARE_DYNAMIC_MULTICAST_DELEGATE( FWeakObjectPtr, DelegateName, DelegateName##_DelegateWrapper, FUNC_CONCAT( Param1Type InParam1, Param2Type InParam2, Param3Type InParam3, Param4Type InParam4, Param5Type InParam5, Param6Type InParam6, Param7Type InParam7, Param8Type InParam8 ), FUNC_CONCAT( *this, InParam1, InParam2, InParam3, InParam4, InParam5, InParam6, InParam7, InParam8 ), void, Param1Type, Param2Type, Param3Type, Param4Type, Param5Type, Param6Type, Param7Type, Param8Type )
@@ -108,6 +119,7 @@
 
 #define DECLARE_DELEGATE_NineParams( DelegateName, Param1Type, Param2Type, Param3Type, Param4Type, Param5Type, Param6Type, Param7Type, Param8Type, Param9Type ) FUNC_DECLARE_DELEGATE( DelegateName, void, Param1Type, Param2Type, Param3Type, Param4Type, Param5Type, Param6Type, Param7Type, Param8Type, Param9Type )
 #define DECLARE_MULTICAST_DELEGATE_NineParams( DelegateName, Param1Type, Param2Type, Param3Type, Param4Type, Param5Type, Param6Type, Param7Type, Param8Type, Param9Type ) FUNC_DECLARE_MULTICAST_DELEGATE( DelegateName, void, Param1Type, Param2Type, Param3Type, Param4Type, Param5Type, Param6Type, Param7Type, Param8Type, Param9Type )
+#define DECLARE_TS_MULTICAST_DELEGATE_NineParams( DelegateName, Param1Type, Param2Type, Param3Type, Param4Type, Param5Type, Param6Type, Param7Type, Param8Type, Param9Type ) FUNC_DECLARE_TS_MULTICAST_DELEGATE( DelegateName, void, Param1Type, Param2Type, Param3Type, Param4Type, Param5Type, Param6Type, Param7Type, Param8Type, Param9Type )
 #define DECLARE_EVENT_NineParams( OwningType, EventName, Param1Type, Param2Type, Param3Type, Param4Type, Param5Type, Param6Type, Param7Type, Param8Type, Param9Type ) FUNC_DECLARE_EVENT( OwningType, EventName, void, Param1Type, Param2Type, Param3Type, Param4Type, Param5Type, Param6Type, Param7Type, Param8Type, Param9Type )
 #define DECLARE_DYNAMIC_DELEGATE_NineParams( DelegateName, Param1Type, Param1Name, Param2Type, Param2Name, Param3Type, Param3Name, Param4Type, Param4Name, Param5Type, Param5Name, Param6Type, Param6Name, Param7Type, Param7Name, Param8Type, Param8Name, Param9Type, Param9Name ) BODY_MACRO_COMBINE(CURRENT_FILE_ID,_,__LINE__,_DELEGATE) FUNC_DECLARE_DYNAMIC_DELEGATE( FWeakObjectPtr, DelegateName, DelegateName##_DelegateWrapper, FUNC_CONCAT( Param1Type InParam1, Param2Type InParam2, Param3Type InParam3, Param4Type InParam4, Param5Type InParam5, Param6Type InParam6, Param7Type InParam7, Param8Type InParam8, Param9Type InParam9 ), FUNC_CONCAT( *this, InParam1, InParam2, InParam3, InParam4, InParam5, InParam6, InParam7, InParam8, InParam9 ), void, Param1Type, Param2Type, Param3Type, Param4Type, Param5Type, Param6Type, Param7Type, Param8Type, Param9Type )
 #define DECLARE_DYNAMIC_MULTICAST_DELEGATE_NineParams( DelegateName, Param1Type, Param1Name, Param2Type, Param2Name, Param3Type, Param3Name, Param4Type, Param4Name, Param5Type, Param5Name, Param6Type, Param6Name, Param7Type, Param7Name, Param8Type, Param8Name, Param9Type, Param9Name ) BODY_MACRO_COMBINE(CURRENT_FILE_ID,_,__LINE__,_DELEGATE) FUNC_DECLARE_DYNAMIC_MULTICAST_DELEGATE( FWeakObjectPtr, DelegateName, DelegateName##_DelegateWrapper, FUNC_CONCAT( Param1Type InParam1, Param2Type InParam2, Param3Type InParam3, Param4Type InParam4, Param5Type InParam5, Param6Type InParam6, Param7Type InParam7, Param8Type InParam8, Param9Type InParam9 ), FUNC_CONCAT( *this, InParam1, InParam2, InParam3, InParam4, InParam5, InParam6, InParam7, InParam8, InParam9 ), void, Param1Type, Param2Type, Param3Type, Param4Type, Param5Type, Param6Type, Param7Type, Param8Type, Param9Type )
diff --git a/Engine/Source/Runtime/Core/Public/Delegates/DelegateSignatureImpl.inl b/Engine/Source/Runtime/Core/Public/Delegates/DelegateSignatureImpl.inl
--- a/Engine/Source/Runtime/Core/Public/Delegates/DelegateSignatureImpl.inl
+++ b/Engine/Source/Runtime/Core/Public/Delegates/DelegateSignatureImpl.inl
@@ -23,12 +23,17 @@ class FDelegateHandle;
 class IDelegateInstance;
 struct FWeakObjectPtr;
 template <typename FuncType, typename UserPolicy> struct IBaseDelegateInstance;
-template<typename UserPolicy> class TMulticastDelegateBase;
 template <typename T> struct TObjectPtr;
 
 template <typename T>
 T* ToRawPtr(const TObjectPtr<T>& Ptr);
 
+template <typename To, typename From>
+To* Cast(From* Src);
+
+template<typename UserPolicy> class TMulticastDelegateBase;
+template<typename UserPolicy> class TTSMulticastDelegateBase;
+
 /**
  * Unicast delegate template class.
  *
@@ -53,13 +58,16 @@ class TDelegate<InRetValType(ParamTypes...), UserPolicy> : public TDelegateBase<
 	using DelegateInstanceType  = typename UserPolicy::FDelegateInstanceExtras;
 	using MulticastDelegateType = typename UserPolicy::FMulticastDelegateExtras;
 
-	static_assert(std::is_convertible<DelegateType*,          FDelegateBase*                     >::value, "UserPolicy::FDelegateExtras should publicly inherit FDelegateBase");
-	static_assert(std::is_convertible<DelegateInstanceType*,  IDelegateInstance*                 >::value, "UserPolicy::FDelegateInstanceExtras should publicly inherit IDelegateInstance");
-	static_assert(std::is_convertible<MulticastDelegateType*, TMulticastDelegateBase<UserPolicy>*>::value, "UserPolicy::FMulticastDelegateExtras should publicly inherit TMulticastDelegateBase<UserPolicy>");
+	static_assert(std::is_convertible_v<DelegateType*, FDelegateBase*>, "UserPolicy::FDelegateExtras should publicly inherit FDelegateBase");
+	static_assert(std::is_convertible_v<DelegateInstanceType*, IDelegateInstance*>, "UserPolicy::FDelegateInstanceExtras should publicly inherit IDelegateInstance");
+	static_assert(std::is_convertible_v<MulticastDelegateType*, TMulticastDelegateBase<UserPolicy>*> || std::is_convertible_v<MulticastDelegateType*, TTSMulticastDelegateBase<UserPolicy>*>, "UserPolicy::FMulticastDelegateExtras should publicly inherit TMulticastDelegateBase<UserPolicy> or TTSMulticastDelegateBase<UserPolicy>");
 
 	template <typename>
 	friend class TMulticastDelegateBase;
 
+	template <typename>
+	friend class TTSMulticastDelegateBase;
+
 public:
 	// Make sure FDelegateBase's public functions are publicly exposed through the TDelegate API
 	using Super::Unbind;
@@ -74,66 +82,67 @@ public:
 	typedef InRetValType TFuncType(ParamTypes...);
 
 	/* Helper typedefs for getting a member function pointer type for the delegate with a given payload */
+	template <typename... VarTypes>                     using TFuncPtr        = RetValType(*)(ParamTypes..., VarTypes...);
 	template <typename UserClass, typename... VarTypes> using TMethodPtr      = typename TMemFunPtrType<false, UserClass, RetValType(ParamTypes..., VarTypes...)>::Type;
 	template <typename UserClass, typename... VarTypes> using TConstMethodPtr = typename TMemFunPtrType<true,  UserClass, RetValType(ParamTypes..., VarTypes...)>::Type;
 
 	/** Declare the user's shared pointer-based delegate instance types. */
-	template <typename UserClass                                                                            > struct TSPMethodDelegate                 : TBaseSPMethodDelegateInstance<false, UserClass, ESPMode::ThreadSafe, FuncType, UserPolicy                                        > { typedef TBaseSPMethodDelegateInstance<false, UserClass, ESPMode::ThreadSafe, FuncType, UserPolicy                                        > Super; TSPMethodDelegate                (const TSharedRef< UserClass, ESPMode::ThreadSafe >& InUserObject, typename Super::FMethodPtr InMethodPtr                                                            ) : Super(InUserObject, InMethodPtr                        ) {} };
-	template <typename UserClass                                                                            > struct TSPMethodDelegate_Const           : TBaseSPMethodDelegateInstance<true , UserClass, ESPMode::ThreadSafe, FuncType, UserPolicy                                        > { typedef TBaseSPMethodDelegateInstance<true , UserClass, ESPMode::ThreadSafe, FuncType, UserPolicy                                        > Super; TSPMethodDelegate_Const          (const TSharedRef< UserClass, ESPMode::ThreadSafe >& InUserObject, typename Super::FMethodPtr InMethodPtr                                                            ) : Super(InUserObject, InMethodPtr                        ) {} };
-	template <typename UserClass, typename Var1Type                                                         > struct TSPMethodDelegate_OneVar          : TBaseSPMethodDelegateInstance<false, UserClass, ESPMode::ThreadSafe, FuncType, UserPolicy, Var1Type                              > { typedef TBaseSPMethodDelegateInstance<false, UserClass, ESPMode::ThreadSafe, FuncType, UserPolicy, Var1Type                              > Super; TSPMethodDelegate_OneVar         (const TSharedRef< UserClass, ESPMode::ThreadSafe >& InUserObject, typename Super::FMethodPtr InMethodPtr, Var1Type Var1                                             ) : Super(InUserObject, InMethodPtr, Var1                  ) {} };
-	template <typename UserClass, typename Var1Type                                                         > struct TSPMethodDelegate_OneVar_Const    : TBaseSPMethodDelegateInstance<true , UserClass, ESPMode::ThreadSafe, FuncType, UserPolicy, Var1Type                              > { typedef TBaseSPMethodDelegateInstance<true , UserClass, ESPMode::ThreadSafe, FuncType, UserPolicy, Var1Type                              > Super; TSPMethodDelegate_OneVar_Const   (const TSharedRef< UserClass, ESPMode::ThreadSafe >& InUserObject, typename Super::FMethodPtr InMethodPtr, Var1Type Var1                                             ) : Super(InUserObject, InMethodPtr, Var1                  ) {} };
-	template <typename UserClass, typename Var1Type, typename Var2Type                                      > struct TSPMethodDelegate_TwoVars         : TBaseSPMethodDelegateInstance<false, UserClass, ESPMode::ThreadSafe, FuncType, UserPolicy, Var1Type, Var2Type                    > { typedef TBaseSPMethodDelegateInstance<false, UserClass, ESPMode::ThreadSafe, FuncType, UserPolicy, Var1Type, Var2Type                    > Super; TSPMethodDelegate_TwoVars        (const TSharedRef< UserClass, ESPMode::ThreadSafe >& InUserObject, typename Super::FMethodPtr InMethodPtr, Var1Type Var1, Var2Type Var2                              ) : Super(InUserObject, InMethodPtr, Var1, Var2            ) {} };
-	template <typename UserClass, typename Var1Type, typename Var2Type                                      > struct TSPMethodDelegate_TwoVars_Const   : TBaseSPMethodDelegateInstance<true , UserClass, ESPMode::ThreadSafe, FuncType, UserPolicy, Var1Type, Var2Type                    > { typedef TBaseSPMethodDelegateInstance<true , UserClass, ESPMode::ThreadSafe, FuncType, UserPolicy, Var1Type, Var2Type                    > Super; TSPMethodDelegate_TwoVars_Const  (const TSharedRef< UserClass, ESPMode::ThreadSafe >& InUserObject, typename Super::FMethodPtr InMethodPtr, Var1Type Var1, Var2Type Var2                              ) : Super(InUserObject, InMethodPtr, Var1, Var2            ) {} };
-	template <typename UserClass, typename Var1Type, typename Var2Type, typename Var3Type                   > struct TSPMethodDelegate_ThreeVars       : TBaseSPMethodDelegateInstance<false, UserClass, ESPMode::ThreadSafe, FuncType, UserPolicy, Var1Type, Var2Type, Var3Type          > { typedef TBaseSPMethodDelegateInstance<false, UserClass, ESPMode::ThreadSafe, FuncType, UserPolicy, Var1Type, Var2Type, Var3Type          > Super; TSPMethodDelegate_ThreeVars      (const TSharedRef< UserClass, ESPMode::ThreadSafe >& InUserObject, typename Super::FMethodPtr InMethodPtr, Var1Type Var1, Var2Type Var2, Var3Type Var3               ) : Super(InUserObject, InMethodPtr, Var1, Var2, Var3      ) {} };
-	template <typename UserClass, typename Var1Type, typename Var2Type, typename Var3Type                   > struct TSPMethodDelegate_ThreeVars_Const : TBaseSPMethodDelegateInstance<true , UserClass, ESPMode::ThreadSafe, FuncType, UserPolicy, Var1Type, Var2Type, Var3Type          > { typedef TBaseSPMethodDelegateInstance<true , UserClass, ESPMode::ThreadSafe, FuncType, UserPolicy, Var1Type, Var2Type, Var3Type          > Super; TSPMethodDelegate_ThreeVars_Const(const TSharedRef< UserClass, ESPMode::ThreadSafe >& InUserObject, typename Super::FMethodPtr InMethodPtr, Var1Type Var1, Var2Type Var2, Var3Type Var3               ) : Super(InUserObject, InMethodPtr, Var1, Var2, Var3      ) {} };
-	template <typename UserClass, typename Var1Type, typename Var2Type, typename Var3Type, typename Var4Type> struct TSPMethodDelegate_FourVars        : TBaseSPMethodDelegateInstance<false, UserClass, ESPMode::ThreadSafe, FuncType, UserPolicy, Var1Type, Var2Type, Var3Type, Var4Type> { typedef TBaseSPMethodDelegateInstance<false, UserClass, ESPMode::ThreadSafe, FuncType, UserPolicy, Var1Type, Var2Type, Var3Type, Var4Type> Super; TSPMethodDelegate_FourVars       (const TSharedRef< UserClass, ESPMode::ThreadSafe >& InUserObject, typename Super::FMethodPtr InMethodPtr, Var1Type Var1, Var2Type Var2, Var3Type Var3, Var4Type Var4) : Super(InUserObject, InMethodPtr, Var1, Var2, Var3, Var4) {} };
-	template <typename UserClass, typename Var1Type, typename Var2Type, typename Var3Type, typename Var4Type> struct TSPMethodDelegate_FourVars_Const  : TBaseSPMethodDelegateInstance<true , UserClass, ESPMode::ThreadSafe, FuncType, UserPolicy, Var1Type, Var2Type, Var3Type, Var4Type> { typedef TBaseSPMethodDelegateInstance<true , UserClass, ESPMode::ThreadSafe, FuncType, UserPolicy, Var1Type, Var2Type, Var3Type, Var4Type> Super; TSPMethodDelegate_FourVars_Const (const TSharedRef< UserClass, ESPMode::ThreadSafe >& InUserObject, typename Super::FMethodPtr InMethodPtr, Var1Type Var1, Var2Type Var2, Var3Type Var3, Var4Type Var4) : Super(InUserObject, InMethodPtr, Var1, Var2, Var3, Var4) {} };
+	template <typename UserClass                                                                            > struct UE_DEPRECATED(5.2, "DelegateType::TSPMethodDelegate<UserClass, ...>::FMethodPtr has been deprecated - use DelegateType::TMethodPtr<UserClass, ...> instead")                      TSPMethodDelegate                 : TBaseSPMethodDelegateInstance<false, UserClass, ESPMode::ThreadSafe, FuncType, UserPolicy                                        > { typedef TBaseSPMethodDelegateInstance<false, UserClass, ESPMode::ThreadSafe, FuncType, UserPolicy                                        > Super; TSPMethodDelegate                (const TSharedRef< UserClass, ESPMode::ThreadSafe >& InUserObject, typename Super::FMethodPtr InMethodPtr                                                            ) : Super(InUserObject, InMethodPtr                        ) {} };
+	template <typename UserClass                                                                            > struct UE_DEPRECATED(5.2, "DelegateType::TSPMethodDelegate_Const<UserClass, ...>::FMethodPtr has been deprecated - use DelegateType::TConstMethodPtr<UserClass, ...> instead")           TSPMethodDelegate_Const           : TBaseSPMethodDelegateInstance<true , UserClass, ESPMode::ThreadSafe, FuncType, UserPolicy                                        > { typedef TBaseSPMethodDelegateInstance<true , UserClass, ESPMode::ThreadSafe, FuncType, UserPolicy                                        > Super; TSPMethodDelegate_Const          (const TSharedRef< UserClass, ESPMode::ThreadSafe >& InUserObject, typename Super::FMethodPtr InMethodPtr                                                            ) : Super(InUserObject, InMethodPtr                        ) {} };
+	template <typename UserClass, typename Var1Type                                                         > struct UE_DEPRECATED(5.2, "DelegateType::TSPMethodDelegate_OneVar<UserClass, ...>::FMethodPtr has been deprecated - use DelegateType::TMethodPtr<UserClass, ...> instead")               TSPMethodDelegate_OneVar          : TBaseSPMethodDelegateInstance<false, UserClass, ESPMode::ThreadSafe, FuncType, UserPolicy, Var1Type                              > { typedef TBaseSPMethodDelegateInstance<false, UserClass, ESPMode::ThreadSafe, FuncType, UserPolicy, Var1Type                              > Super; TSPMethodDelegate_OneVar         (const TSharedRef< UserClass, ESPMode::ThreadSafe >& InUserObject, typename Super::FMethodPtr InMethodPtr, Var1Type Var1                                             ) : Super(InUserObject, InMethodPtr, Var1                  ) {} };
+	template <typename UserClass, typename Var1Type                                                         > struct UE_DEPRECATED(5.2, "DelegateType::TSPMethodDelegate_OneVar_Const<UserClass, ...>::FMethodPtr has been deprecated - use DelegateType::TConstMethodPtr<UserClass, ...> instead")    TSPMethodDelegate_OneVar_Const    : TBaseSPMethodDelegateInstance<true , UserClass, ESPMode::ThreadSafe, FuncType, UserPolicy, Var1Type                              > { typedef TBaseSPMethodDelegateInstance<true , UserClass, ESPMode::ThreadSafe, FuncType, UserPolicy, Var1Type                              > Super; TSPMethodDelegate_OneVar_Const   (const TSharedRef< UserClass, ESPMode::ThreadSafe >& InUserObject, typename Super::FMethodPtr InMethodPtr, Var1Type Var1                                             ) : Super(InUserObject, InMethodPtr, Var1                  ) {} };
+	template <typename UserClass, typename Var1Type, typename Var2Type                                      > struct UE_DEPRECATED(5.2, "DelegateType::TSPMethodDelegate_TwoVars<UserClass, ...>::FMethodPtr has been deprecated - use DelegateType::TMethodPtr<UserClass, ...> instead")              TSPMethodDelegate_TwoVars         : TBaseSPMethodDelegateInstance<false, UserClass, ESPMode::ThreadSafe, FuncType, UserPolicy, Var1Type, Var2Type                    > { typedef TBaseSPMethodDelegateInstance<false, UserClass, ESPMode::ThreadSafe, FuncType, UserPolicy, Var1Type, Var2Type                    > Super; TSPMethodDelegate_TwoVars        (const TSharedRef< UserClass, ESPMode::ThreadSafe >& InUserObject, typename Super::FMethodPtr InMethodPtr, Var1Type Var1, Var2Type Var2                              ) : Super(InUserObject, InMethodPtr, Var1, Var2            ) {} };
+	template <typename UserClass, typename Var1Type, typename Var2Type                                      > struct UE_DEPRECATED(5.2, "DelegateType::TSPMethodDelegate_TwoVars_Const<UserClass, ...>::FMethodPtr has been deprecated - use DelegateType::TConstMethodPtr<UserClass, ...> instead")   TSPMethodDelegate_TwoVars_Const   : TBaseSPMethodDelegateInstance<true , UserClass, ESPMode::ThreadSafe, FuncType, UserPolicy, Var1Type, Var2Type                    > { typedef TBaseSPMethodDelegateInstance<true , UserClass, ESPMode::ThreadSafe, FuncType, UserPolicy, Var1Type, Var2Type                    > Super; TSPMethodDelegate_TwoVars_Const  (const TSharedRef< UserClass, ESPMode::ThreadSafe >& InUserObject, typename Super::FMethodPtr InMethodPtr, Var1Type Var1, Var2Type Var2                              ) : Super(InUserObject, InMethodPtr, Var1, Var2            ) {} };
+	template <typename UserClass, typename Var1Type, typename Var2Type, typename Var3Type                   > struct UE_DEPRECATED(5.2, "DelegateType::TSPMethodDelegate_ThreeVars<UserClass, ...>::FMethodPtr has been deprecated - use DelegateType::TMethodPtr<UserClass, ...> instead")            TSPMethodDelegate_ThreeVars       : TBaseSPMethodDelegateInstance<false, UserClass, ESPMode::ThreadSafe, FuncType, UserPolicy, Var1Type, Var2Type, Var3Type          > { typedef TBaseSPMethodDelegateInstance<false, UserClass, ESPMode::ThreadSafe, FuncType, UserPolicy, Var1Type, Var2Type, Var3Type          > Super; TSPMethodDelegate_ThreeVars      (const TSharedRef< UserClass, ESPMode::ThreadSafe >& InUserObject, typename Super::FMethodPtr InMethodPtr, Var1Type Var1, Var2Type Var2, Var3Type Var3               ) : Super(InUserObject, InMethodPtr, Var1, Var2, Var3      ) {} };
+	template <typename UserClass, typename Var1Type, typename Var2Type, typename Var3Type                   > struct UE_DEPRECATED(5.2, "DelegateType::TSPMethodDelegate_ThreeVars_Const<UserClass, ...>::FMethodPtr has been deprecated - use DelegateType::TConstMethodPtr<UserClass, ...> instead") TSPMethodDelegate_ThreeVars_Const : TBaseSPMethodDelegateInstance<true , UserClass, ESPMode::ThreadSafe, FuncType, UserPolicy, Var1Type, Var2Type, Var3Type          > { typedef TBaseSPMethodDelegateInstance<true , UserClass, ESPMode::ThreadSafe, FuncType, UserPolicy, Var1Type, Var2Type, Var3Type          > Super; TSPMethodDelegate_ThreeVars_Const(const TSharedRef< UserClass, ESPMode::ThreadSafe >& InUserObject, typename Super::FMethodPtr InMethodPtr, Var1Type Var1, Var2Type Var2, Var3Type Var3               ) : Super(InUserObject, InMethodPtr, Var1, Var2, Var3      ) {} };
+	template <typename UserClass, typename Var1Type, typename Var2Type, typename Var3Type, typename Var4Type> struct UE_DEPRECATED(5.2, "DelegateType::TSPMethodDelegate_FourVars<UserClass, ...>::FMethodPtr has been deprecated - use DelegateType::TMethodPtr<UserClass, ...> instead")             TSPMethodDelegate_FourVars        : TBaseSPMethodDelegateInstance<false, UserClass, ESPMode::ThreadSafe, FuncType, UserPolicy, Var1Type, Var2Type, Var3Type, Var4Type> { typedef TBaseSPMethodDelegateInstance<false, UserClass, ESPMode::ThreadSafe, FuncType, UserPolicy, Var1Type, Var2Type, Var3Type, Var4Type> Super; TSPMethodDelegate_FourVars       (const TSharedRef< UserClass, ESPMode::ThreadSafe >& InUserObject, typename Super::FMethodPtr InMethodPtr, Var1Type Var1, Var2Type Var2, Var3Type Var3, Var4Type Var4) : Super(InUserObject, InMethodPtr, Var1, Var2, Var3, Var4) {} };
+	template <typename UserClass, typename Var1Type, typename Var2Type, typename Var3Type, typename Var4Type> struct UE_DEPRECATED(5.2, "DelegateType::TSPMethodDelegate_FourVars_Const<UserClass, ...>::FMethodPtr has been deprecated - use DelegateType::TConstMethodPtr<UserClass, ...> instead")  TSPMethodDelegate_FourVars_Const  : TBaseSPMethodDelegateInstance<true , UserClass, ESPMode::ThreadSafe, FuncType, UserPolicy, Var1Type, Var2Type, Var3Type, Var4Type> { typedef TBaseSPMethodDelegateInstance<true , UserClass, ESPMode::ThreadSafe, FuncType, UserPolicy, Var1Type, Var2Type, Var3Type, Var4Type> Super; TSPMethodDelegate_FourVars_Const (const TSharedRef< UserClass, ESPMode::ThreadSafe >& InUserObject, typename Super::FMethodPtr InMethodPtr, Var1Type Var1, Var2Type Var2, Var3Type Var3, Var4Type Var4) : Super(InUserObject, InMethodPtr, Var1, Var2, Var3, Var4) {} };
 
 	/** Declare the user's shared pointer-based delegate instance types. */
-	template <typename UserClass                                                                            > struct TThreadSafeSPMethodDelegate                 : TBaseSPMethodDelegateInstance<false, UserClass, ESPMode::ThreadSafe, FuncType, UserPolicy                                        > { typedef TBaseSPMethodDelegateInstance<false, UserClass, ESPMode::ThreadSafe, FuncType, UserPolicy                                        > Super; TThreadSafeSPMethodDelegate                (const TSharedRef<UserClass, ESPMode::ThreadSafe>& InUserObject, typename Super::FMethodPtr InMethodPtr                                                            ) : Super(InUserObject, InMethodPtr                        ) {} };
-	template <typename UserClass                                                                            > struct TThreadSafeSPMethodDelegate_Const           : TBaseSPMethodDelegateInstance<true , UserClass, ESPMode::ThreadSafe, FuncType, UserPolicy                                        > { typedef TBaseSPMethodDelegateInstance<true , UserClass, ESPMode::ThreadSafe, FuncType, UserPolicy                                        > Super; TThreadSafeSPMethodDelegate_Const          (const TSharedRef<UserClass, ESPMode::ThreadSafe>& InUserObject, typename Super::FMethodPtr InMethodPtr                                                            ) : Super(InUserObject, InMethodPtr                        ) {} };
-	template <typename UserClass, typename Var1Type                                                         > struct TThreadSafeSPMethodDelegate_OneVar          : TBaseSPMethodDelegateInstance<false, UserClass, ESPMode::ThreadSafe, FuncType, UserPolicy, Var1Type                              > { typedef TBaseSPMethodDelegateInstance<false, UserClass, ESPMode::ThreadSafe, FuncType, UserPolicy, Var1Type                              > Super; TThreadSafeSPMethodDelegate_OneVar         (const TSharedRef<UserClass, ESPMode::ThreadSafe>& InUserObject, typename Super::FMethodPtr InMethodPtr, Var1Type Var1                                             ) : Super(InUserObject, InMethodPtr, Var1                  ) {} };
-	template <typename UserClass, typename Var1Type                                                         > struct TThreadSafeSPMethodDelegate_OneVar_Const    : TBaseSPMethodDelegateInstance<true , UserClass, ESPMode::ThreadSafe, FuncType, UserPolicy, Var1Type                              > { typedef TBaseSPMethodDelegateInstance<true , UserClass, ESPMode::ThreadSafe, FuncType, UserPolicy, Var1Type                              > Super; TThreadSafeSPMethodDelegate_OneVar_Const   (const TSharedRef<UserClass, ESPMode::ThreadSafe>& InUserObject, typename Super::FMethodPtr InMethodPtr, Var1Type Var1                                             ) : Super(InUserObject, InMethodPtr, Var1                  ) {} };
-	template <typename UserClass, typename Var1Type, typename Var2Type                                      > struct TThreadSafeSPMethodDelegate_TwoVars         : TBaseSPMethodDelegateInstance<false, UserClass, ESPMode::ThreadSafe, FuncType, UserPolicy, Var1Type, Var2Type                    > { typedef TBaseSPMethodDelegateInstance<false, UserClass, ESPMode::ThreadSafe, FuncType, UserPolicy, Var1Type, Var2Type                    > Super; TThreadSafeSPMethodDelegate_TwoVars        (const TSharedRef<UserClass, ESPMode::ThreadSafe>& InUserObject, typename Super::FMethodPtr InMethodPtr, Var1Type Var1, Var2Type Var2                              ) : Super(InUserObject, InMethodPtr, Var1, Var2            ) {} };
-	template <typename UserClass, typename Var1Type, typename Var2Type                                      > struct TThreadSafeSPMethodDelegate_TwoVars_Const   : TBaseSPMethodDelegateInstance<true , UserClass, ESPMode::ThreadSafe, FuncType, UserPolicy, Var1Type, Var2Type                    > { typedef TBaseSPMethodDelegateInstance<true , UserClass, ESPMode::ThreadSafe, FuncType, UserPolicy, Var1Type, Var2Type                    > Super; TThreadSafeSPMethodDelegate_TwoVars_Const  (const TSharedRef<UserClass, ESPMode::ThreadSafe>& InUserObject, typename Super::FMethodPtr InMethodPtr, Var1Type Var1, Var2Type Var2                              ) : Super(InUserObject, InMethodPtr, Var1, Var2            ) {} };
-	template <typename UserClass, typename Var1Type, typename Var2Type, typename Var3Type                   > struct TThreadSafeSPMethodDelegate_ThreeVars       : TBaseSPMethodDelegateInstance<false, UserClass, ESPMode::ThreadSafe, FuncType, UserPolicy, Var1Type, Var2Type, Var3Type          > { typedef TBaseSPMethodDelegateInstance<false, UserClass, ESPMode::ThreadSafe, FuncType, UserPolicy, Var1Type, Var2Type, Var3Type          > Super; TThreadSafeSPMethodDelegate_ThreeVars      (const TSharedRef<UserClass, ESPMode::ThreadSafe>& InUserObject, typename Super::FMethodPtr InMethodPtr, Var1Type Var1, Var2Type Var2, Var3Type Var3               ) : Super(InUserObject, InMethodPtr, Var1, Var2, Var3      ) {} };
-	template <typename UserClass, typename Var1Type, typename Var2Type, typename Var3Type                   > struct TThreadSafeSPMethodDelegate_ThreeVars_Const : TBaseSPMethodDelegateInstance<true , UserClass, ESPMode::ThreadSafe, FuncType, UserPolicy, Var1Type, Var2Type, Var3Type          > { typedef TBaseSPMethodDelegateInstance<true , UserClass, ESPMode::ThreadSafe, FuncType, UserPolicy, Var1Type, Var2Type, Var3Type          > Super; TThreadSafeSPMethodDelegate_ThreeVars_Const(const TSharedRef<UserClass, ESPMode::ThreadSafe>& InUserObject, typename Super::FMethodPtr InMethodPtr, Var1Type Var1, Var2Type Var2, Var3Type Var3               ) : Super(InUserObject, InMethodPtr, Var1, Var2, Var3      ) {} };
-	template <typename UserClass, typename Var1Type, typename Var2Type, typename Var3Type, typename Var4Type> struct TThreadSafeSPMethodDelegate_FourVars        : TBaseSPMethodDelegateInstance<false, UserClass, ESPMode::ThreadSafe, FuncType, UserPolicy, Var1Type, Var2Type, Var3Type, Var4Type> { typedef TBaseSPMethodDelegateInstance<false, UserClass, ESPMode::ThreadSafe, FuncType, UserPolicy, Var1Type, Var2Type, Var3Type, Var4Type> Super; TThreadSafeSPMethodDelegate_FourVars       (const TSharedRef<UserClass, ESPMode::ThreadSafe>& InUserObject, typename Super::FMethodPtr InMethodPtr, Var1Type Var1, Var2Type Var2, Var3Type Var3, Var4Type Var4) : Super(InUserObject, InMethodPtr, Var1, Var2, Var3, Var4) {} };
-	template <typename UserClass, typename Var1Type, typename Var2Type, typename Var3Type, typename Var4Type> struct TThreadSafeSPMethodDelegate_FourVars_Const  : TBaseSPMethodDelegateInstance<true , UserClass, ESPMode::ThreadSafe, FuncType, UserPolicy, Var1Type, Var2Type, Var3Type, Var4Type> { typedef TBaseSPMethodDelegateInstance<true , UserClass, ESPMode::ThreadSafe, FuncType, UserPolicy, Var1Type, Var2Type, Var3Type, Var4Type> Super; TThreadSafeSPMethodDelegate_FourVars_Const (const TSharedRef<UserClass, ESPMode::ThreadSafe>& InUserObject, typename Super::FMethodPtr InMethodPtr, Var1Type Var1, Var2Type Var2, Var3Type Var3, Var4Type Var4) : Super(InUserObject, InMethodPtr, Var1, Var2, Var3, Var4) {} };
+	template <typename UserClass                                                                            > struct UE_DEPRECATED(5.2, "DelegateType::TThreadSafeSPMethodDelegate<UserClass, ...>::FMethodPtr has been deprecated - use DelegateType::TMethodPtr<UserClass, ...> instead")                      TThreadSafeSPMethodDelegate                 : TBaseSPMethodDelegateInstance<false, UserClass, ESPMode::ThreadSafe, FuncType, UserPolicy                                        > { typedef TBaseSPMethodDelegateInstance<false, UserClass, ESPMode::ThreadSafe, FuncType, UserPolicy                                        > Super; TThreadSafeSPMethodDelegate                (const TSharedRef<UserClass, ESPMode::ThreadSafe>& InUserObject, typename Super::FMethodPtr InMethodPtr                                                            ) : Super(InUserObject, InMethodPtr                        ) {} };
+	template <typename UserClass                                                                            > struct UE_DEPRECATED(5.2, "DelegateType::TThreadSafeSPMethodDelegate_Const<UserClass, ...>::FMethodPtr has been deprecated - use DelegateType::TConstMethodPtr<UserClass, ...> instead")           TThreadSafeSPMethodDelegate_Const           : TBaseSPMethodDelegateInstance<true , UserClass, ESPMode::ThreadSafe, FuncType, UserPolicy                                        > { typedef TBaseSPMethodDelegateInstance<true , UserClass, ESPMode::ThreadSafe, FuncType, UserPolicy                                        > Super; TThreadSafeSPMethodDelegate_Const          (const TSharedRef<UserClass, ESPMode::ThreadSafe>& InUserObject, typename Super::FMethodPtr InMethodPtr                                                            ) : Super(InUserObject, InMethodPtr                        ) {} };
+	template <typename UserClass, typename Var1Type                                                         > struct UE_DEPRECATED(5.2, "DelegateType::TThreadSafeSPMethodDelegate_OneVar<UserClass, ...>::FMethodPtr has been deprecated - use DelegateType::TMethodPtr<UserClass, ...> instead")               TThreadSafeSPMethodDelegate_OneVar          : TBaseSPMethodDelegateInstance<false, UserClass, ESPMode::ThreadSafe, FuncType, UserPolicy, Var1Type                              > { typedef TBaseSPMethodDelegateInstance<false, UserClass, ESPMode::ThreadSafe, FuncType, UserPolicy, Var1Type                              > Super; TThreadSafeSPMethodDelegate_OneVar         (const TSharedRef<UserClass, ESPMode::ThreadSafe>& InUserObject, typename Super::FMethodPtr InMethodPtr, Var1Type Var1                                             ) : Super(InUserObject, InMethodPtr, Var1                  ) {} };
+	template <typename UserClass, typename Var1Type                                                         > struct UE_DEPRECATED(5.2, "DelegateType::TThreadSafeSPMethodDelegate_OneVar_Const<UserClass, ...>::FMethodPtr has been deprecated - use DelegateType::TConstMethodPtr<UserClass, ...> instead")    TThreadSafeSPMethodDelegate_OneVar_Const    : TBaseSPMethodDelegateInstance<true , UserClass, ESPMode::ThreadSafe, FuncType, UserPolicy, Var1Type                              > { typedef TBaseSPMethodDelegateInstance<true , UserClass, ESPMode::ThreadSafe, FuncType, UserPolicy, Var1Type                              > Super; TThreadSafeSPMethodDelegate_OneVar_Const   (const TSharedRef<UserClass, ESPMode::ThreadSafe>& InUserObject, typename Super::FMethodPtr InMethodPtr, Var1Type Var1                                             ) : Super(InUserObject, InMethodPtr, Var1                  ) {} };
+	template <typename UserClass, typename Var1Type, typename Var2Type                                      > struct UE_DEPRECATED(5.2, "DelegateType::TThreadSafeSPMethodDelegate_TwoVars<UserClass, ...>::FMethodPtr has been deprecated - use DelegateType::TMethodPtr<UserClass, ...> instead")              TThreadSafeSPMethodDelegate_TwoVars         : TBaseSPMethodDelegateInstance<false, UserClass, ESPMode::ThreadSafe, FuncType, UserPolicy, Var1Type, Var2Type                    > { typedef TBaseSPMethodDelegateInstance<false, UserClass, ESPMode::ThreadSafe, FuncType, UserPolicy, Var1Type, Var2Type                    > Super; TThreadSafeSPMethodDelegate_TwoVars        (const TSharedRef<UserClass, ESPMode::ThreadSafe>& InUserObject, typename Super::FMethodPtr InMethodPtr, Var1Type Var1, Var2Type Var2                              ) : Super(InUserObject, InMethodPtr, Var1, Var2            ) {} };
+	template <typename UserClass, typename Var1Type, typename Var2Type                                      > struct UE_DEPRECATED(5.2, "DelegateType::TThreadSafeSPMethodDelegate_TwoVars_Const<UserClass, ...>::FMethodPtr has been deprecated - use DelegateType::TConstMethodPtr<UserClass, ...> instead")   TThreadSafeSPMethodDelegate_TwoVars_Const   : TBaseSPMethodDelegateInstance<true , UserClass, ESPMode::ThreadSafe, FuncType, UserPolicy, Var1Type, Var2Type                    > { typedef TBaseSPMethodDelegateInstance<true , UserClass, ESPMode::ThreadSafe, FuncType, UserPolicy, Var1Type, Var2Type                    > Super; TThreadSafeSPMethodDelegate_TwoVars_Const  (const TSharedRef<UserClass, ESPMode::ThreadSafe>& InUserObject, typename Super::FMethodPtr InMethodPtr, Var1Type Var1, Var2Type Var2                              ) : Super(InUserObject, InMethodPtr, Var1, Var2            ) {} };
+	template <typename UserClass, typename Var1Type, typename Var2Type, typename Var3Type                   > struct UE_DEPRECATED(5.2, "DelegateType::TThreadSafeSPMethodDelegate_ThreeVars<UserClass, ...>::FMethodPtr has been deprecated - use DelegateType::TMethodPtr<UserClass, ...> instead")            TThreadSafeSPMethodDelegate_ThreeVars       : TBaseSPMethodDelegateInstance<false, UserClass, ESPMode::ThreadSafe, FuncType, UserPolicy, Var1Type, Var2Type, Var3Type          > { typedef TBaseSPMethodDelegateInstance<false, UserClass, ESPMode::ThreadSafe, FuncType, UserPolicy, Var1Type, Var2Type, Var3Type          > Super; TThreadSafeSPMethodDelegate_ThreeVars      (const TSharedRef<UserClass, ESPMode::ThreadSafe>& InUserObject, typename Super::FMethodPtr InMethodPtr, Var1Type Var1, Var2Type Var2, Var3Type Var3               ) : Super(InUserObject, InMethodPtr, Var1, Var2, Var3      ) {} };
+	template <typename UserClass, typename Var1Type, typename Var2Type, typename Var3Type                   > struct UE_DEPRECATED(5.2, "DelegateType::TThreadSafeSPMethodDelegate_ThreeVars_Const<UserClass, ...>::FMethodPtr has been deprecated - use DelegateType::TConstMethodPtr<UserClass, ...> instead") TThreadSafeSPMethodDelegate_ThreeVars_Const : TBaseSPMethodDelegateInstance<true , UserClass, ESPMode::ThreadSafe, FuncType, UserPolicy, Var1Type, Var2Type, Var3Type          > { typedef TBaseSPMethodDelegateInstance<true , UserClass, ESPMode::ThreadSafe, FuncType, UserPolicy, Var1Type, Var2Type, Var3Type          > Super; TThreadSafeSPMethodDelegate_ThreeVars_Const(const TSharedRef<UserClass, ESPMode::ThreadSafe>& InUserObject, typename Super::FMethodPtr InMethodPtr, Var1Type Var1, Var2Type Var2, Var3Type Var3               ) : Super(InUserObject, InMethodPtr, Var1, Var2, Var3      ) {} };
+	template <typename UserClass, typename Var1Type, typename Var2Type, typename Var3Type, typename Var4Type> struct UE_DEPRECATED(5.2, "DelegateType::TThreadSafeSPMethodDelegate_FourVars<UserClass, ...>::FMethodPtr has been deprecated - use DelegateType::TMethodPtr<UserClass, ...> instead")             TThreadSafeSPMethodDelegate_FourVars        : TBaseSPMethodDelegateInstance<false, UserClass, ESPMode::ThreadSafe, FuncType, UserPolicy, Var1Type, Var2Type, Var3Type, Var4Type> { typedef TBaseSPMethodDelegateInstance<false, UserClass, ESPMode::ThreadSafe, FuncType, UserPolicy, Var1Type, Var2Type, Var3Type, Var4Type> Super; TThreadSafeSPMethodDelegate_FourVars       (const TSharedRef<UserClass, ESPMode::ThreadSafe>& InUserObject, typename Super::FMethodPtr InMethodPtr, Var1Type Var1, Var2Type Var2, Var3Type Var3, Var4Type Var4) : Super(InUserObject, InMethodPtr, Var1, Var2, Var3, Var4) {} };
+	template <typename UserClass, typename Var1Type, typename Var2Type, typename Var3Type, typename Var4Type> struct UE_DEPRECATED(5.2, "DelegateType::TThreadSafeSPMethodDelegate_FourVars_Const<UserClass, ...>::FMethodPtr has been deprecated - use DelegateType::TConstMethodPtr<UserClass, ...> instead")  TThreadSafeSPMethodDelegate_FourVars_Const  : TBaseSPMethodDelegateInstance<true , UserClass, ESPMode::ThreadSafe, FuncType, UserPolicy, Var1Type, Var2Type, Var3Type, Var4Type> { typedef TBaseSPMethodDelegateInstance<true , UserClass, ESPMode::ThreadSafe, FuncType, UserPolicy, Var1Type, Var2Type, Var3Type, Var4Type> Super; TThreadSafeSPMethodDelegate_FourVars_Const (const TSharedRef<UserClass, ESPMode::ThreadSafe>& InUserObject, typename Super::FMethodPtr InMethodPtr, Var1Type Var1, Var2Type Var2, Var3Type Var3, Var4Type Var4) : Super(InUserObject, InMethodPtr, Var1, Var2, Var3, Var4) {} };
 
 	/** Declare the user's C++ pointer-based delegate instance types. */
-	template <typename UserClass                                                                            > struct TRawMethodDelegate                 : TBaseRawMethodDelegateInstance<false, UserClass, FuncType, UserPolicy                                        > { typedef TBaseRawMethodDelegateInstance<false, UserClass, FuncType, UserPolicy                                        > Super; TRawMethodDelegate                (UserClass* InUserObject, typename Super::FMethodPtr InMethodPtr                                                            ) : Super(InUserObject, InMethodPtr                        ) {} };
-	template <typename UserClass                                                                            > struct TRawMethodDelegate_Const           : TBaseRawMethodDelegateInstance<true , UserClass, FuncType, UserPolicy                                        > { typedef TBaseRawMethodDelegateInstance<true , UserClass, FuncType, UserPolicy                                        > Super; TRawMethodDelegate_Const          (UserClass* InUserObject, typename Super::FMethodPtr InMethodPtr                                                            ) : Super(InUserObject, InMethodPtr                        ) {} };
-	template <typename UserClass, typename Var1Type                                                         > struct TRawMethodDelegate_OneVar          : TBaseRawMethodDelegateInstance<false, UserClass, FuncType, UserPolicy, Var1Type                              > { typedef TBaseRawMethodDelegateInstance<false, UserClass, FuncType, UserPolicy, Var1Type                              > Super; TRawMethodDelegate_OneVar         (UserClass* InUserObject, typename Super::FMethodPtr InMethodPtr, Var1Type Var1                                             ) : Super(InUserObject, InMethodPtr, Var1                  ) {} };
-	template <typename UserClass, typename Var1Type                                                         > struct TRawMethodDelegate_OneVar_Const    : TBaseRawMethodDelegateInstance<true , UserClass, FuncType, UserPolicy, Var1Type                              > { typedef TBaseRawMethodDelegateInstance<true , UserClass, FuncType, UserPolicy, Var1Type                              > Super; TRawMethodDelegate_OneVar_Const   (UserClass* InUserObject, typename Super::FMethodPtr InMethodPtr, Var1Type Var1                                             ) : Super(InUserObject, InMethodPtr, Var1                  ) {} };
-	template <typename UserClass, typename Var1Type, typename Var2Type                                      > struct TRawMethodDelegate_TwoVars         : TBaseRawMethodDelegateInstance<false, UserClass, FuncType, UserPolicy, Var1Type, Var2Type                    > { typedef TBaseRawMethodDelegateInstance<false, UserClass, FuncType, UserPolicy, Var1Type, Var2Type                    > Super; TRawMethodDelegate_TwoVars        (UserClass* InUserObject, typename Super::FMethodPtr InMethodPtr, Var1Type Var1, Var2Type Var2                              ) : Super(InUserObject, InMethodPtr, Var1, Var2            ) {} };
-	template <typename UserClass, typename Var1Type, typename Var2Type                                      > struct TRawMethodDelegate_TwoVars_Const   : TBaseRawMethodDelegateInstance<true , UserClass, FuncType, UserPolicy, Var1Type, Var2Type                    > { typedef TBaseRawMethodDelegateInstance<true , UserClass, FuncType, UserPolicy, Var1Type, Var2Type                    > Super; TRawMethodDelegate_TwoVars_Const  (UserClass* InUserObject, typename Super::FMethodPtr InMethodPtr, Var1Type Var1, Var2Type Var2                              ) : Super(InUserObject, InMethodPtr, Var1, Var2            ) {} };
-	template <typename UserClass, typename Var1Type, typename Var2Type, typename Var3Type                   > struct TRawMethodDelegate_ThreeVars       : TBaseRawMethodDelegateInstance<false, UserClass, FuncType, UserPolicy, Var1Type, Var2Type, Var3Type          > { typedef TBaseRawMethodDelegateInstance<false, UserClass, FuncType, UserPolicy, Var1Type, Var2Type, Var3Type          > Super; TRawMethodDelegate_ThreeVars      (UserClass* InUserObject, typename Super::FMethodPtr InMethodPtr, Var1Type Var1, Var2Type Var2, Var3Type Var3               ) : Super(InUserObject, InMethodPtr, Var1, Var2, Var3      ) {} };
-	template <typename UserClass, typename Var1Type, typename Var2Type, typename Var3Type                   > struct TRawMethodDelegate_ThreeVars_Const : TBaseRawMethodDelegateInstance<true , UserClass, FuncType, UserPolicy, Var1Type, Var2Type, Var3Type          > { typedef TBaseRawMethodDelegateInstance<true , UserClass, FuncType, UserPolicy, Var1Type, Var2Type, Var3Type          > Super; TRawMethodDelegate_ThreeVars_Const(UserClass* InUserObject, typename Super::FMethodPtr InMethodPtr, Var1Type Var1, Var2Type Var2, Var3Type Var3               ) : Super(InUserObject, InMethodPtr, Var1, Var2, Var3      ) {} };
-	template <typename UserClass, typename Var1Type, typename Var2Type, typename Var3Type, typename Var4Type> struct TRawMethodDelegate_FourVars        : TBaseRawMethodDelegateInstance<false, UserClass, FuncType, UserPolicy, Var1Type, Var2Type, Var3Type, Var4Type> { typedef TBaseRawMethodDelegateInstance<false, UserClass, FuncType, UserPolicy, Var1Type, Var2Type, Var3Type, Var4Type> Super; TRawMethodDelegate_FourVars       (UserClass* InUserObject, typename Super::FMethodPtr InMethodPtr, Var1Type Var1, Var2Type Var2, Var3Type Var3, Var4Type Var4) : Super(InUserObject, InMethodPtr, Var1, Var2, Var3, Var4) {} };
-	template <typename UserClass, typename Var1Type, typename Var2Type, typename Var3Type, typename Var4Type> struct TRawMethodDelegate_FourVars_Const  : TBaseRawMethodDelegateInstance<true , UserClass, FuncType, UserPolicy, Var1Type, Var2Type, Var3Type, Var4Type> { typedef TBaseRawMethodDelegateInstance<true , UserClass, FuncType, UserPolicy, Var1Type, Var2Type, Var3Type, Var4Type> Super; TRawMethodDelegate_FourVars_Const (UserClass* InUserObject, typename Super::FMethodPtr InMethodPtr, Var1Type Var1, Var2Type Var2, Var3Type Var3, Var4Type Var4) : Super(InUserObject, InMethodPtr, Var1, Var2, Var3, Var4) {} };
+	template <typename UserClass                                                                            > struct UE_DEPRECATED(5.2, "DelegateType::TRawMethodDelegate<UserClass, ...>::FMethodPtr has been deprecated - use DelegateType::TMethodPtr<UserClass, ...> instead")                      TRawMethodDelegate                 : TBaseRawMethodDelegateInstance<false, UserClass, FuncType, UserPolicy                                        > { typedef TBaseRawMethodDelegateInstance<false, UserClass, FuncType, UserPolicy                                        > Super; TRawMethodDelegate                (UserClass* InUserObject, typename Super::FMethodPtr InMethodPtr                                                            ) : Super(InUserObject, InMethodPtr                        ) {} };
+	template <typename UserClass                                                                            > struct UE_DEPRECATED(5.2, "DelegateType::TRawMethodDelegate_Const<UserClass, ...>::FMethodPtr has been deprecated - use DelegateType::TConstMethodPtr<UserClass, ...> instead")           TRawMethodDelegate_Const           : TBaseRawMethodDelegateInstance<true , UserClass, FuncType, UserPolicy                                        > { typedef TBaseRawMethodDelegateInstance<true , UserClass, FuncType, UserPolicy                                        > Super; TRawMethodDelegate_Const          (UserClass* InUserObject, typename Super::FMethodPtr InMethodPtr                                                            ) : Super(InUserObject, InMethodPtr                        ) {} };
+	template <typename UserClass, typename Var1Type                                                         > struct UE_DEPRECATED(5.2, "DelegateType::TRawMethodDelegate_OneVar<UserClass, ...>::FMethodPtr has been deprecated - use DelegateType::TMethodPtr<UserClass, ...> instead")               TRawMethodDelegate_OneVar          : TBaseRawMethodDelegateInstance<false, UserClass, FuncType, UserPolicy, Var1Type                              > { typedef TBaseRawMethodDelegateInstance<false, UserClass, FuncType, UserPolicy, Var1Type                              > Super; TRawMethodDelegate_OneVar         (UserClass* InUserObject, typename Super::FMethodPtr InMethodPtr, Var1Type Var1                                             ) : Super(InUserObject, InMethodPtr, Var1                  ) {} };
+	template <typename UserClass, typename Var1Type                                                         > struct UE_DEPRECATED(5.2, "DelegateType::TRawMethodDelegate_OneVar_Const<UserClass, ...>::FMethodPtr has been deprecated - use DelegateType::TConstMethodPtr<UserClass, ...> instead")    TRawMethodDelegate_OneVar_Const    : TBaseRawMethodDelegateInstance<true , UserClass, FuncType, UserPolicy, Var1Type                              > { typedef TBaseRawMethodDelegateInstance<true , UserClass, FuncType, UserPolicy, Var1Type                              > Super; TRawMethodDelegate_OneVar_Const   (UserClass* InUserObject, typename Super::FMethodPtr InMethodPtr, Var1Type Var1                                             ) : Super(InUserObject, InMethodPtr, Var1                  ) {} };
+	template <typename UserClass, typename Var1Type, typename Var2Type                                      > struct UE_DEPRECATED(5.2, "DelegateType::TRawMethodDelegate_TwoVars<UserClass, ...>::FMethodPtr has been deprecated - use DelegateType::TMethodPtr<UserClass, ...> instead")              TRawMethodDelegate_TwoVars         : TBaseRawMethodDelegateInstance<false, UserClass, FuncType, UserPolicy, Var1Type, Var2Type                    > { typedef TBaseRawMethodDelegateInstance<false, UserClass, FuncType, UserPolicy, Var1Type, Var2Type                    > Super; TRawMethodDelegate_TwoVars        (UserClass* InUserObject, typename Super::FMethodPtr InMethodPtr, Var1Type Var1, Var2Type Var2                              ) : Super(InUserObject, InMethodPtr, Var1, Var2            ) {} };
+	template <typename UserClass, typename Var1Type, typename Var2Type                                      > struct UE_DEPRECATED(5.2, "DelegateType::TRawMethodDelegate_TwoVars_Const<UserClass, ...>::FMethodPtr has been deprecated - use DelegateType::TConstMethodPtr<UserClass, ...> instead")   TRawMethodDelegate_TwoVars_Const   : TBaseRawMethodDelegateInstance<true , UserClass, FuncType, UserPolicy, Var1Type, Var2Type                    > { typedef TBaseRawMethodDelegateInstance<true , UserClass, FuncType, UserPolicy, Var1Type, Var2Type                    > Super; TRawMethodDelegate_TwoVars_Const  (UserClass* InUserObject, typename Super::FMethodPtr InMethodPtr, Var1Type Var1, Var2Type Var2                              ) : Super(InUserObject, InMethodPtr, Var1, Var2            ) {} };
+	template <typename UserClass, typename Var1Type, typename Var2Type, typename Var3Type                   > struct UE_DEPRECATED(5.2, "DelegateType::TRawMethodDelegate_ThreeVars<UserClass, ...>::FMethodPtr has been deprecated - use DelegateType::TMethodPtr<UserClass, ...> instead")            TRawMethodDelegate_ThreeVars       : TBaseRawMethodDelegateInstance<false, UserClass, FuncType, UserPolicy, Var1Type, Var2Type, Var3Type          > { typedef TBaseRawMethodDelegateInstance<false, UserClass, FuncType, UserPolicy, Var1Type, Var2Type, Var3Type          > Super; TRawMethodDelegate_ThreeVars      (UserClass* InUserObject, typename Super::FMethodPtr InMethodPtr, Var1Type Var1, Var2Type Var2, Var3Type Var3               ) : Super(InUserObject, InMethodPtr, Var1, Var2, Var3      ) {} };
+	template <typename UserClass, typename Var1Type, typename Var2Type, typename Var3Type                   > struct UE_DEPRECATED(5.2, "DelegateType::TRawMethodDelegate_ThreeVars_Const<UserClass, ...>::FMethodPtr has been deprecated - use DelegateType::TConstMethodPtr<UserClass, ...> instead") TRawMethodDelegate_ThreeVars_Const : TBaseRawMethodDelegateInstance<true , UserClass, FuncType, UserPolicy, Var1Type, Var2Type, Var3Type          > { typedef TBaseRawMethodDelegateInstance<true , UserClass, FuncType, UserPolicy, Var1Type, Var2Type, Var3Type          > Super; TRawMethodDelegate_ThreeVars_Const(UserClass* InUserObject, typename Super::FMethodPtr InMethodPtr, Var1Type Var1, Var2Type Var2, Var3Type Var3               ) : Super(InUserObject, InMethodPtr, Var1, Var2, Var3      ) {} };
+	template <typename UserClass, typename Var1Type, typename Var2Type, typename Var3Type, typename Var4Type> struct UE_DEPRECATED(5.2, "DelegateType::TRawMethodDelegate_FourVars<UserClass, ...>::FMethodPtr has been deprecated - use DelegateType::TMethodPtr<UserClass, ...> instead")             TRawMethodDelegate_FourVars        : TBaseRawMethodDelegateInstance<false, UserClass, FuncType, UserPolicy, Var1Type, Var2Type, Var3Type, Var4Type> { typedef TBaseRawMethodDelegateInstance<false, UserClass, FuncType, UserPolicy, Var1Type, Var2Type, Var3Type, Var4Type> Super; TRawMethodDelegate_FourVars       (UserClass* InUserObject, typename Super::FMethodPtr InMethodPtr, Var1Type Var1, Var2Type Var2, Var3Type Var3, Var4Type Var4) : Super(InUserObject, InMethodPtr, Var1, Var2, Var3, Var4) {} };
+	template <typename UserClass, typename Var1Type, typename Var2Type, typename Var3Type, typename Var4Type> struct UE_DEPRECATED(5.2, "DelegateType::TRawMethodDelegate_FourVars_Const<UserClass, ...>::FMethodPtr has been deprecated - use DelegateType::TConstMethodPtr<UserClass, ...> instead")  TRawMethodDelegate_FourVars_Const  : TBaseRawMethodDelegateInstance<true , UserClass, FuncType, UserPolicy, Var1Type, Var2Type, Var3Type, Var4Type> { typedef TBaseRawMethodDelegateInstance<true , UserClass, FuncType, UserPolicy, Var1Type, Var2Type, Var3Type, Var4Type> Super; TRawMethodDelegate_FourVars_Const (UserClass* InUserObject, typename Super::FMethodPtr InMethodPtr, Var1Type Var1, Var2Type Var2, Var3Type Var3, Var4Type Var4) : Super(InUserObject, InMethodPtr, Var1, Var2, Var3, Var4) {} };
 	
 	/** Declare the user's UFunction-based delegate instance types. */
-	template <typename UObjectTemplate                                                                            > struct TUFunctionDelegateBinding           : TBaseUFunctionDelegateInstance<UObjectTemplate, FuncType, UserPolicy                                        > { typedef TBaseUFunctionDelegateInstance<UObjectTemplate, FuncType, UserPolicy                                        > Super; TUFunctionDelegateBinding          (UObjectTemplate* InUserObject, const FName& InFunctionName                                                            ) : Super(InUserObject, InFunctionName                        ) {} };
-	template <typename UObjectTemplate, typename Var1Type                                                         > struct TUFunctionDelegateBinding_OneVar    : TBaseUFunctionDelegateInstance<UObjectTemplate, FuncType, UserPolicy, Var1Type                              > { typedef TBaseUFunctionDelegateInstance<UObjectTemplate, FuncType, UserPolicy, Var1Type                              > Super; TUFunctionDelegateBinding_OneVar   (UObjectTemplate* InUserObject, const FName& InFunctionName, Var1Type Var1                                             ) : Super(InUserObject, InFunctionName, Var1                  ) {} };
-	template <typename UObjectTemplate, typename Var1Type, typename Var2Type                                      > struct TUFunctionDelegateBinding_TwoVars   : TBaseUFunctionDelegateInstance<UObjectTemplate, FuncType, UserPolicy, Var1Type, Var2Type                    > { typedef TBaseUFunctionDelegateInstance<UObjectTemplate, FuncType, UserPolicy, Var1Type, Var2Type                    > Super; TUFunctionDelegateBinding_TwoVars  (UObjectTemplate* InUserObject, const FName& InFunctionName, Var1Type Var1, Var2Type Var2                              ) : Super(InUserObject, InFunctionName, Var1, Var2            ) {} };
-	template <typename UObjectTemplate, typename Var1Type, typename Var2Type, typename Var3Type                   > struct TUFunctionDelegateBinding_ThreeVars : TBaseUFunctionDelegateInstance<UObjectTemplate, FuncType, UserPolicy, Var1Type, Var2Type, Var3Type          > { typedef TBaseUFunctionDelegateInstance<UObjectTemplate, FuncType, UserPolicy, Var1Type, Var2Type, Var3Type          > Super; TUFunctionDelegateBinding_ThreeVars(UObjectTemplate* InUserObject, const FName& InFunctionName, Var1Type Var1, Var2Type Var2, Var3Type Var3               ) : Super(InUserObject, InFunctionName, Var1, Var2, Var3      ) {} };
-	template <typename UObjectTemplate, typename Var1Type, typename Var2Type, typename Var3Type, typename Var4Type> struct TUFunctionDelegateBinding_FourVars  : TBaseUFunctionDelegateInstance<UObjectTemplate, FuncType, UserPolicy, Var1Type, Var2Type, Var3Type, Var4Type> { typedef TBaseUFunctionDelegateInstance<UObjectTemplate, FuncType, UserPolicy, Var1Type, Var2Type, Var3Type, Var4Type> Super; TUFunctionDelegateBinding_FourVars (UObjectTemplate* InUserObject, const FName& InFunctionName, Var1Type Var1, Var2Type Var2, Var3Type Var3, Var4Type Var4) : Super(InUserObject, InFunctionName, Var1, Var2, Var3, Var4) {} };
+	template <typename UObjectTemplate                                                                            > struct UE_DEPRECATED(5.2, "DelegateType::TUFunctionDelegateBinding has been deprecated")           TUFunctionDelegateBinding           : TBaseUFunctionDelegateInstance<UObjectTemplate, FuncType, UserPolicy                                        > { typedef TBaseUFunctionDelegateInstance<UObjectTemplate, FuncType, UserPolicy                                        > Super; TUFunctionDelegateBinding          (UObjectTemplate* InUserObject, const FName& InFunctionName                                                            ) : Super(InUserObject, InFunctionName                        ) {} };
+	template <typename UObjectTemplate, typename Var1Type                                                         > struct UE_DEPRECATED(5.2, "DelegateType::TUFunctionDelegateBinding_OneVar has been deprecated")    TUFunctionDelegateBinding_OneVar    : TBaseUFunctionDelegateInstance<UObjectTemplate, FuncType, UserPolicy, Var1Type                              > { typedef TBaseUFunctionDelegateInstance<UObjectTemplate, FuncType, UserPolicy, Var1Type                              > Super; TUFunctionDelegateBinding_OneVar   (UObjectTemplate* InUserObject, const FName& InFunctionName, Var1Type Var1                                             ) : Super(InUserObject, InFunctionName, Var1                  ) {} };
+	template <typename UObjectTemplate, typename Var1Type, typename Var2Type                                      > struct UE_DEPRECATED(5.2, "DelegateType::TUFunctionDelegateBinding_TwoVars has been deprecated")   TUFunctionDelegateBinding_TwoVars   : TBaseUFunctionDelegateInstance<UObjectTemplate, FuncType, UserPolicy, Var1Type, Var2Type                    > { typedef TBaseUFunctionDelegateInstance<UObjectTemplate, FuncType, UserPolicy, Var1Type, Var2Type                    > Super; TUFunctionDelegateBinding_TwoVars  (UObjectTemplate* InUserObject, const FName& InFunctionName, Var1Type Var1, Var2Type Var2                              ) : Super(InUserObject, InFunctionName, Var1, Var2            ) {} };
+	template <typename UObjectTemplate, typename Var1Type, typename Var2Type, typename Var3Type                   > struct UE_DEPRECATED(5.2, "DelegateType::TUFunctionDelegateBinding_ThreeVars has been deprecated") TUFunctionDelegateBinding_ThreeVars : TBaseUFunctionDelegateInstance<UObjectTemplate, FuncType, UserPolicy, Var1Type, Var2Type, Var3Type          > { typedef TBaseUFunctionDelegateInstance<UObjectTemplate, FuncType, UserPolicy, Var1Type, Var2Type, Var3Type          > Super; TUFunctionDelegateBinding_ThreeVars(UObjectTemplate* InUserObject, const FName& InFunctionName, Var1Type Var1, Var2Type Var2, Var3Type Var3               ) : Super(InUserObject, InFunctionName, Var1, Var2, Var3      ) {} };
+	template <typename UObjectTemplate, typename Var1Type, typename Var2Type, typename Var3Type, typename Var4Type> struct UE_DEPRECATED(5.2, "DelegateType::TUFunctionDelegateBinding_FourVars has been deprecated")  TUFunctionDelegateBinding_FourVars  : TBaseUFunctionDelegateInstance<UObjectTemplate, FuncType, UserPolicy, Var1Type, Var2Type, Var3Type, Var4Type> { typedef TBaseUFunctionDelegateInstance<UObjectTemplate, FuncType, UserPolicy, Var1Type, Var2Type, Var3Type, Var4Type> Super; TUFunctionDelegateBinding_FourVars (UObjectTemplate* InUserObject, const FName& InFunctionName, Var1Type Var1, Var2Type Var2, Var3Type Var3, Var4Type Var4) : Super(InUserObject, InFunctionName, Var1, Var2, Var3, Var4) {} };
 
 	/** Declare the user's UObject-based delegate instance types. */
-	template <typename UserClass                                                                            > struct TUObjectMethodDelegate                 : TBaseUObjectMethodDelegateInstance<false, UserClass, FuncType, UserPolicy                                        > { typedef TBaseUObjectMethodDelegateInstance<false, UserClass, FuncType, UserPolicy                                        > Super; TUObjectMethodDelegate                (UserClass* InUserObject, typename Super::FMethodPtr InMethodPtr                                                            ) : Super(InUserObject, InMethodPtr                        ) {} };
-	template <typename UserClass                                                                            > struct TUObjectMethodDelegate_Const           : TBaseUObjectMethodDelegateInstance<true , UserClass, FuncType, UserPolicy                                        > { typedef TBaseUObjectMethodDelegateInstance<true , UserClass, FuncType, UserPolicy                                        > Super; TUObjectMethodDelegate_Const          (UserClass* InUserObject, typename Super::FMethodPtr InMethodPtr                                                            ) : Super(InUserObject, InMethodPtr                        ) {} };
-	template <typename UserClass, typename Var1Type                                                         > struct TUObjectMethodDelegate_OneVar          : TBaseUObjectMethodDelegateInstance<false, UserClass, FuncType, UserPolicy, Var1Type                              > { typedef TBaseUObjectMethodDelegateInstance<false, UserClass, FuncType, UserPolicy, Var1Type                              > Super; TUObjectMethodDelegate_OneVar         (UserClass* InUserObject, typename Super::FMethodPtr InMethodPtr, Var1Type Var1                                             ) : Super(InUserObject, InMethodPtr, Var1                  ) {} };
-	template <typename UserClass, typename Var1Type                                                         > struct TUObjectMethodDelegate_OneVar_Const    : TBaseUObjectMethodDelegateInstance<true , UserClass, FuncType, UserPolicy, Var1Type                              > { typedef TBaseUObjectMethodDelegateInstance<true , UserClass, FuncType, UserPolicy, Var1Type                              > Super; TUObjectMethodDelegate_OneVar_Const   (UserClass* InUserObject, typename Super::FMethodPtr InMethodPtr, Var1Type Var1                                             ) : Super(InUserObject, InMethodPtr, Var1                  ) {} };
-	template <typename UserClass, typename Var1Type, typename Var2Type                                      > struct TUObjectMethodDelegate_TwoVars         : TBaseUObjectMethodDelegateInstance<false, UserClass, FuncType, UserPolicy, Var1Type, Var2Type                    > { typedef TBaseUObjectMethodDelegateInstance<false, UserClass, FuncType, UserPolicy, Var1Type, Var2Type                    > Super; TUObjectMethodDelegate_TwoVars        (UserClass* InUserObject, typename Super::FMethodPtr InMethodPtr, Var1Type Var1, Var2Type Var2                              ) : Super(InUserObject, InMethodPtr, Var1, Var2            ) {} };
-	template <typename UserClass, typename Var1Type, typename Var2Type                                      > struct TUObjectMethodDelegate_TwoVars_Const   : TBaseUObjectMethodDelegateInstance<true , UserClass, FuncType, UserPolicy, Var1Type, Var2Type                    > { typedef TBaseUObjectMethodDelegateInstance<true , UserClass, FuncType, UserPolicy, Var1Type, Var2Type                    > Super; TUObjectMethodDelegate_TwoVars_Const  (UserClass* InUserObject, typename Super::FMethodPtr InMethodPtr, Var1Type Var1, Var2Type Var2                              ) : Super(InUserObject, InMethodPtr, Var1, Var2            ) {} };
-	template <typename UserClass, typename Var1Type, typename Var2Type, typename Var3Type                   > struct TUObjectMethodDelegate_ThreeVars       : TBaseUObjectMethodDelegateInstance<false, UserClass, FuncType, UserPolicy, Var1Type, Var2Type, Var3Type          > { typedef TBaseUObjectMethodDelegateInstance<false, UserClass, FuncType, UserPolicy, Var1Type, Var2Type, Var3Type          > Super; TUObjectMethodDelegate_ThreeVars      (UserClass* InUserObject, typename Super::FMethodPtr InMethodPtr, Var1Type Var1, Var2Type Var2, Var3Type Var3               ) : Super(InUserObject, InMethodPtr, Var1, Var2, Var3      ) {} };
-	template <typename UserClass, typename Var1Type, typename Var2Type, typename Var3Type                   > struct TUObjectMethodDelegate_ThreeVars_Const : TBaseUObjectMethodDelegateInstance<true , UserClass, FuncType, UserPolicy, Var1Type, Var2Type, Var3Type          > { typedef TBaseUObjectMethodDelegateInstance<true , UserClass, FuncType, UserPolicy, Var1Type, Var2Type, Var3Type          > Super; TUObjectMethodDelegate_ThreeVars_Const(UserClass* InUserObject, typename Super::FMethodPtr InMethodPtr, Var1Type Var1, Var2Type Var2, Var3Type Var3               ) : Super(InUserObject, InMethodPtr, Var1, Var2, Var3      ) {} };
-	template <typename UserClass, typename Var1Type, typename Var2Type, typename Var3Type, typename Var4Type> struct TUObjectMethodDelegate_FourVars        : TBaseUObjectMethodDelegateInstance<false, UserClass, FuncType, UserPolicy, Var1Type, Var2Type, Var3Type, Var4Type> { typedef TBaseUObjectMethodDelegateInstance<false, UserClass, FuncType, UserPolicy, Var1Type, Var2Type, Var3Type, Var4Type> Super; TUObjectMethodDelegate_FourVars       (UserClass* InUserObject, typename Super::FMethodPtr InMethodPtr, Var1Type Var1, Var2Type Var2, Var3Type Var3, Var4Type Var4) : Super(InUserObject, InMethodPtr, Var1, Var2, Var3, Var4) {} };
-	template <typename UserClass, typename Var1Type, typename Var2Type, typename Var3Type, typename Var4Type> struct TUObjectMethodDelegate_FourVars_Const  : TBaseUObjectMethodDelegateInstance<true , UserClass, FuncType, UserPolicy, Var1Type, Var2Type, Var3Type, Var4Type> { typedef TBaseUObjectMethodDelegateInstance<true , UserClass, FuncType, UserPolicy, Var1Type, Var2Type, Var3Type, Var4Type> Super; TUObjectMethodDelegate_FourVars_Const (UserClass* InUserObject, typename Super::FMethodPtr InMethodPtr, Var1Type Var1, Var2Type Var2, Var3Type Var3, Var4Type Var4) : Super(InUserObject, InMethodPtr, Var1, Var2, Var3, Var4) {} };
+	template <typename UserClass                                                                            > struct UE_DEPRECATED(5.2, "DelegateType::TUObjectMethodDelegate<UserClass, ...>::FMethodPtr has been deprecated - use DelegateType::TMethodPtr<UserClass, ...> instead")                      TUObjectMethodDelegate                 : TBaseUObjectMethodDelegateInstance<false, UserClass, FuncType, UserPolicy                                        > { typedef TBaseUObjectMethodDelegateInstance<false, UserClass, FuncType, UserPolicy                                        > Super; TUObjectMethodDelegate                (UserClass* InUserObject, typename Super::FMethodPtr InMethodPtr                                                            ) : Super(InUserObject, InMethodPtr                        ) {} };
+	template <typename UserClass                                                                            > struct UE_DEPRECATED(5.2, "DelegateType::TUObjectMethodDelegate_Const<UserClass, ...>::FMethodPtr has been deprecated - use DelegateType::TConstMethodPtr<UserClass, ...> instead")           TUObjectMethodDelegate_Const           : TBaseUObjectMethodDelegateInstance<true , UserClass, FuncType, UserPolicy                                        > { typedef TBaseUObjectMethodDelegateInstance<true , UserClass, FuncType, UserPolicy                                        > Super; TUObjectMethodDelegate_Const          (UserClass* InUserObject, typename Super::FMethodPtr InMethodPtr                                                            ) : Super(InUserObject, InMethodPtr                        ) {} };
+	template <typename UserClass, typename Var1Type                                                         > struct UE_DEPRECATED(5.2, "DelegateType::TUObjectMethodDelegate_OneVar<UserClass, ...>::FMethodPtr has been deprecated - use DelegateType::TMethodPtr<UserClass, ...> instead")               TUObjectMethodDelegate_OneVar          : TBaseUObjectMethodDelegateInstance<false, UserClass, FuncType, UserPolicy, Var1Type                              > { typedef TBaseUObjectMethodDelegateInstance<false, UserClass, FuncType, UserPolicy, Var1Type                              > Super; TUObjectMethodDelegate_OneVar         (UserClass* InUserObject, typename Super::FMethodPtr InMethodPtr, Var1Type Var1                                             ) : Super(InUserObject, InMethodPtr, Var1                  ) {} };
+	template <typename UserClass, typename Var1Type                                                         > struct UE_DEPRECATED(5.2, "DelegateType::TUObjectMethodDelegate_OneVar_Const<UserClass, ...>::FMethodPtr has been deprecated - use DelegateType::TConstMethodPtr<UserClass, ...> instead")    TUObjectMethodDelegate_OneVar_Const    : TBaseUObjectMethodDelegateInstance<true , UserClass, FuncType, UserPolicy, Var1Type                              > { typedef TBaseUObjectMethodDelegateInstance<true , UserClass, FuncType, UserPolicy, Var1Type                              > Super; TUObjectMethodDelegate_OneVar_Const   (UserClass* InUserObject, typename Super::FMethodPtr InMethodPtr, Var1Type Var1                                             ) : Super(InUserObject, InMethodPtr, Var1                  ) {} };
+	template <typename UserClass, typename Var1Type, typename Var2Type                                      > struct UE_DEPRECATED(5.2, "DelegateType::TUObjectMethodDelegate_TwoVars<UserClass, ...>::FMethodPtr has been deprecated - use DelegateType::TMethodPtr<UserClass, ...> instead")              TUObjectMethodDelegate_TwoVars         : TBaseUObjectMethodDelegateInstance<false, UserClass, FuncType, UserPolicy, Var1Type, Var2Type                    > { typedef TBaseUObjectMethodDelegateInstance<false, UserClass, FuncType, UserPolicy, Var1Type, Var2Type                    > Super; TUObjectMethodDelegate_TwoVars        (UserClass* InUserObject, typename Super::FMethodPtr InMethodPtr, Var1Type Var1, Var2Type Var2                              ) : Super(InUserObject, InMethodPtr, Var1, Var2            ) {} };
+	template <typename UserClass, typename Var1Type, typename Var2Type                                      > struct UE_DEPRECATED(5.2, "DelegateType::TUObjectMethodDelegate_TwoVars_Const<UserClass, ...>::FMethodPtr has been deprecated - use DelegateType::TConstMethodPtr<UserClass, ...> instead")   TUObjectMethodDelegate_TwoVars_Const   : TBaseUObjectMethodDelegateInstance<true , UserClass, FuncType, UserPolicy, Var1Type, Var2Type                    > { typedef TBaseUObjectMethodDelegateInstance<true , UserClass, FuncType, UserPolicy, Var1Type, Var2Type                    > Super; TUObjectMethodDelegate_TwoVars_Const  (UserClass* InUserObject, typename Super::FMethodPtr InMethodPtr, Var1Type Var1, Var2Type Var2                              ) : Super(InUserObject, InMethodPtr, Var1, Var2            ) {} };
+	template <typename UserClass, typename Var1Type, typename Var2Type, typename Var3Type                   > struct UE_DEPRECATED(5.2, "DelegateType::TUObjectMethodDelegate_ThreeVars<UserClass, ...>::FMethodPtr has been deprecated - use DelegateType::TMethodPtr<UserClass, ...> instead")            TUObjectMethodDelegate_ThreeVars       : TBaseUObjectMethodDelegateInstance<false, UserClass, FuncType, UserPolicy, Var1Type, Var2Type, Var3Type          > { typedef TBaseUObjectMethodDelegateInstance<false, UserClass, FuncType, UserPolicy, Var1Type, Var2Type, Var3Type          > Super; TUObjectMethodDelegate_ThreeVars      (UserClass* InUserObject, typename Super::FMethodPtr InMethodPtr, Var1Type Var1, Var2Type Var2, Var3Type Var3               ) : Super(InUserObject, InMethodPtr, Var1, Var2, Var3      ) {} };
+	template <typename UserClass, typename Var1Type, typename Var2Type, typename Var3Type                   > struct UE_DEPRECATED(5.2, "DelegateType::TUObjectMethodDelegate_ThreeVars_Const<UserClass, ...>::FMethodPtr has been deprecated - use DelegateType::TConstMethodPtr<UserClass, ...> instead") TUObjectMethodDelegate_ThreeVars_Const : TBaseUObjectMethodDelegateInstance<true , UserClass, FuncType, UserPolicy, Var1Type, Var2Type, Var3Type          > { typedef TBaseUObjectMethodDelegateInstance<true , UserClass, FuncType, UserPolicy, Var1Type, Var2Type, Var3Type          > Super; TUObjectMethodDelegate_ThreeVars_Const(UserClass* InUserObject, typename Super::FMethodPtr InMethodPtr, Var1Type Var1, Var2Type Var2, Var3Type Var3               ) : Super(InUserObject, InMethodPtr, Var1, Var2, Var3      ) {} };
+	template <typename UserClass, typename Var1Type, typename Var2Type, typename Var3Type, typename Var4Type> struct UE_DEPRECATED(5.2, "DelegateType::TUObjectMethodDelegate_FourVars<UserClass, ...>::FMethodPtr has been deprecated - use DelegateType::TMethodPtr<UserClass, ...> instead")             TUObjectMethodDelegate_FourVars        : TBaseUObjectMethodDelegateInstance<false, UserClass, FuncType, UserPolicy, Var1Type, Var2Type, Var3Type, Var4Type> { typedef TBaseUObjectMethodDelegateInstance<false, UserClass, FuncType, UserPolicy, Var1Type, Var2Type, Var3Type, Var4Type> Super; TUObjectMethodDelegate_FourVars       (UserClass* InUserObject, typename Super::FMethodPtr InMethodPtr, Var1Type Var1, Var2Type Var2, Var3Type Var3, Var4Type Var4) : Super(InUserObject, InMethodPtr, Var1, Var2, Var3, Var4) {} };
+	template <typename UserClass, typename Var1Type, typename Var2Type, typename Var3Type, typename Var4Type> struct UE_DEPRECATED(5.2, "DelegateType::TUObjectMethodDelegate_FourVars_Const<UserClass, ...>::FMethodPtr has been deprecated - use DelegateType::TConstMethodPtr<UserClass, ...> instead")  TUObjectMethodDelegate_FourVars_Const  : TBaseUObjectMethodDelegateInstance<true , UserClass, FuncType, UserPolicy, Var1Type, Var2Type, Var3Type, Var4Type> { typedef TBaseUObjectMethodDelegateInstance<true , UserClass, FuncType, UserPolicy, Var1Type, Var2Type, Var3Type, Var4Type> Super; TUObjectMethodDelegate_FourVars_Const (UserClass* InUserObject, typename Super::FMethodPtr InMethodPtr, Var1Type Var1, Var2Type Var2, Var3Type Var3, Var4Type Var4) : Super(InUserObject, InMethodPtr, Var1, Var2, Var3, Var4) {} };
 	
 	/** Declare the user's static function pointer delegate instance types. */
-	template <typename... VarTypes > struct FStaticDelegate : TBaseStaticDelegateInstance<FuncType, UserPolicy, VarTypes... > { typedef TBaseStaticDelegateInstance<FuncType, UserPolicy, VarTypes... > Super; FStaticDelegate ( typename Super::FFuncPtr InFuncPtr, VarTypes... Vars ) : Super(InFuncPtr, Vars... ) {} };
+	template <typename... VarTypes > struct UE_DEPRECATED(5.2, "DelegateType::FStaticDelegate<...>::FFuncPtr has been deprecated - use DelegateType::TFuncPtr<...> instead") FStaticDelegate : TBaseStaticDelegateInstance<FuncType, UserPolicy, VarTypes... > { typedef TBaseStaticDelegateInstance<FuncType, UserPolicy, VarTypes... > Super; FStaticDelegate ( typename Super::FFuncPtr InFuncPtr, VarTypes... Vars ) : Super(InFuncPtr, Vars... ) {} };
 
 public:
 
@@ -372,10 +381,7 @@ public:
 	 *
 	 * @param Other The delegate object to move from.
 	 */
-	inline TDelegate(TDelegate&& Other)
-	{
-		*this = MoveTemp(Other);
-	}
+	TDelegate(TDelegate&& Other) = default;
 
 	/**
 	 * Creates and initializes a new instance from an existing delegate object.
@@ -392,25 +398,7 @@ public:
 	 *
 	 * @param	OtherDelegate	Delegate object to copy from
 	 */
-	inline TDelegate& operator=(TDelegate&& Other)
-	{
-		if (&Other != this)
-		{
-			// this down-cast is OK! allows for managing invocation list in the base class without requiring virtual functions
-			DelegateInstanceInterfaceType* OtherInstance = Other.GetDelegateInstanceProtected();
-
-			if (OtherInstance != nullptr)
-			{
-				OtherInstance->CreateCopy(*this);
-			}
-			else
-			{
-				Unbind();
-			}
-		}
-
-		return *this;
-	}
+	TDelegate& operator=(TDelegate&& Other) = default;
 
 	/**
 	 * Assignment operator.
@@ -1041,7 +1029,7 @@ public:
 
 		// NOTE: If you hit a compile error on the following line, it means you're trying to use a non-UObject type
 		//       with this delegate, which is not supported
-		this->Object = InUserObject;
+		this->Object = Cast<UObject>(InUserObject);
 
 		// Store the function name.  The incoming function name was generated by a macro and includes the method's class name.
 		this->FunctionName = InFunctionName;
diff --git a/Engine/Source/Runtime/Core/Public/Delegates/TSMulticastDelegateBase.h b/Engine/Source/Runtime/Core/Public/Delegates/TSMulticastDelegateBase.h
new file mode 100644
--- /dev/null
+++ b/Engine/Source/Runtime/Core/Public/Delegates/TSMulticastDelegateBase.h
@@ -0,0 +1,108 @@
+// Copyright Epic Games, Inc. All Rights Reserved.
+
+#pragma once
+
+#include "CoreTypes.h"
+#include "Misc/ScopeLock.h"
+#include "Delegates/MulticastDelegateBase.h"
+
+/**
+ * Abstract base class for thread-safe multicast delegates.
+ */
+template<typename UserPolicy>
+class TTSMulticastDelegateBase : protected TMulticastDelegateBase<UserPolicy>
+{
+protected:
+	using Super = TMulticastDelegateBase<UserPolicy>;
+	using UserPolicyType = UserPolicy;
+
+public:
+	/** Removes all functions from this delegate's invocation list. */
+	void Clear()
+	{
+		FScopeLock Lock(&CS);
+		Super::Clear();
+	}
+
+	/**
+	 * Checks to see if any functions are bound to this multi-cast delegate.
+	 *
+	 * @return true if any functions are bound, false otherwise.
+	 */
+	inline bool IsBound() const
+	{
+		FScopeLock Lock(&CS);
+		return Super::IsBound();
+	}
+
+	/**
+	 * Checks to see if any functions are bound to the given user object.
+	 *
+	 * @return	True if any functions are bound to InUserObject, false otherwise.
+	 */
+	inline bool IsBoundToObject(void const* InUserObject) const
+	{
+		FScopeLock Lock(&CS);
+		return Super::IsBoundToObject(InUserObject);
+	}
+
+	/**
+	 * Removes all functions from this multi-cast delegate's invocation list that are bound to the specified UserObject.
+	 * Note that the order of the delegates may not be preserved!
+	 *
+	 * @param InUserObject The object to remove all delegates for.
+	 * @return  The number of delegates successfully removed.
+	 */
+	int32 RemoveAll(const void* InUserObject)
+	{
+		FScopeLock Lock(&CS);
+		return Super::RemoveAll(InUserObject);
+	}
+
+protected:
+
+	/** Hidden default constructor. */
+	TTSMulticastDelegateBase() = default;
+
+protected:
+	template<typename DelegateInstanceInterfaceType, typename DelegateType>
+	void CopyFrom(const TTSMulticastDelegateBase& Other)
+	{
+		FScopeLock Lock(&CS);
+		Super::template CopyFrom<DelegateInstanceInterfaceType, DelegateType>(Other);
+		return *this;
+	}
+
+	template<typename DelegateInstanceInterfaceType, typename DelegateBaseType, typename... ParamTypes>
+	void Broadcast(ParamTypes... Params) const
+	{
+		FScopeLock Lock(&const_cast<TTSMulticastDelegateBase*>(this)->CS);
+		Super::template Broadcast<DelegateInstanceInterfaceType, DelegateBaseType, ParamTypes...>(Params...);
+	}
+
+	/**
+	 * Adds the given delegate instance to the invocation list.
+	 *
+	 * @param NewDelegateBaseRef The delegate instance to add.
+	 */
+	inline FDelegateHandle AddDelegateInstance(TDelegateBase<UserPolicy>&& NewDelegateBaseRef)
+	{
+		FScopeLock Lock(&CS);
+		return Super::AddDelegateInstance(MoveTemp(NewDelegateBaseRef));
+	}
+
+	/**
+	 * Removes a function from this multi-cast delegate's invocation list (performance is O(N)).
+	 *
+	 * @param Handle The handle of the delegate instance to remove.
+	 * @return  true if the delegate was successfully removed.
+	 */
+	bool RemoveDelegateInstance(FDelegateHandle Handle)
+	{
+		FScopeLock Lock(&CS);
+		return Super::RemoveDelegateInstance(Handle);
+	}
+
+private:
+	mutable FCriticalSection CS;
+};
\ No newline at end of file
