removing changes that slipped in from another unsubmitted CL and broke the build

diff --git a/Engine/Source/Runtime/Core/Public/Delegates/DelegateSignatureImpl.inl b/Engine/Source/Runtime/Core/Public/Delegates/DelegateSignatureImpl.inl
--- a/Engine/Source/Runtime/Core/Public/Delegates/DelegateSignatureImpl.inl
+++ b/Engine/Source/Runtime/Core/Public/Delegates/DelegateSignatureImpl.inl
@@ -28,9 +28,6 @@ template <typename T> struct TObjectPtr;
 template <typename T>
 T* ToRawPtr(const TObjectPtr<T>& Ptr);
 
-template <typename To, typename From>
-To* Cast(From* Src);
-
 template<typename UserPolicy> class TMulticastDelegateBase;
 template<typename UserPolicy> class TTSMulticastDelegateBase;
 
@@ -82,67 +79,66 @@ public:
 	typedef InRetValType TFuncType(ParamTypes...);
 
 	/* Helper typedefs for getting a member function pointer type for the delegate with a given payload */
-	template <typename... VarTypes>                     using TFuncPtr        = RetValType(*)(ParamTypes..., VarTypes...);
-	template <typename UserClass, typename... VarTypes> using TMethodPtr      = typename TMemFunPtrType<false, UserClass, RetValType(ParamTypes..., VarTypes...)>::Type;
-	template <typename UserClass, typename... VarTypes> using TConstMethodPtr = typename TMemFunPtrType<true,  UserClass, RetValType(ParamTypes..., VarTypes...)>::Type;
+	template <typename UserClass, typename... VarTypes> using TMethodPtr = typename TMemFunPtrType<false, UserClass, RetValType(ParamTypes..., VarTypes...)>::Type;
+	template <typename UserClass, typename... VarTypes> using TConstMethodPtr = typename TMemFunPtrType<true, UserClass, RetValType(ParamTypes..., VarTypes...)>::Type;
 
 	/** Declare the user's shared pointer-based delegate instance types. */
-	template <typename UserClass                                                                            > struct UE_DEPRECATED(5.2, "DelegateType::TSPMethodDelegate<UserClass, ...>::FMethodPtr has been deprecated - use DelegateType::TMethodPtr<UserClass, ...> instead")                      TSPMethodDelegate                 : TBaseSPMethodDelegateInstance<false, UserClass, ESPMode::ThreadSafe, FuncType, UserPolicy                                        > { typedef TBaseSPMethodDelegateInstance<false, UserClass, ESPMode::ThreadSafe, FuncType, UserPolicy                                        > Super; TSPMethodDelegate                (const TSharedRef< UserClass, ESPMode::ThreadSafe >& InUserObject, typename Super::FMethodPtr InMethodPtr                                                            ) : Super(InUserObject, InMethodPtr                        ) {} };
-	template <typename UserClass                                                                            > struct UE_DEPRECATED(5.2, "DelegateType::TSPMethodDelegate_Const<UserClass, ...>::FMethodPtr has been deprecated - use DelegateType::TConstMethodPtr<UserClass, ...> instead")           TSPMethodDelegate_Const           : TBaseSPMethodDelegateInstance<true , UserClass, ESPMode::ThreadSafe, FuncType, UserPolicy                                        > { typedef TBaseSPMethodDelegateInstance<true , UserClass, ESPMode::ThreadSafe, FuncType, UserPolicy                                        > Super; TSPMethodDelegate_Const          (const TSharedRef< UserClass, ESPMode::ThreadSafe >& InUserObject, typename Super::FMethodPtr InMethodPtr                                                            ) : Super(InUserObject, InMethodPtr                        ) {} };
-	template <typename UserClass, typename Var1Type                                                         > struct UE_DEPRECATED(5.2, "DelegateType::TSPMethodDelegate_OneVar<UserClass, ...>::FMethodPtr has been deprecated - use DelegateType::TMethodPtr<UserClass, ...> instead")               TSPMethodDelegate_OneVar          : TBaseSPMethodDelegateInstance<false, UserClass, ESPMode::ThreadSafe, FuncType, UserPolicy, Var1Type                              > { typedef TBaseSPMethodDelegateInstance<false, UserClass, ESPMode::ThreadSafe, FuncType, UserPolicy, Var1Type                              > Super; TSPMethodDelegate_OneVar         (const TSharedRef< UserClass, ESPMode::ThreadSafe >& InUserObject, typename Super::FMethodPtr InMethodPtr, Var1Type Var1                                             ) : Super(InUserObject, InMethodPtr, Var1                  ) {} };
-	template <typename UserClass, typename Var1Type                                                         > struct UE_DEPRECATED(5.2, "DelegateType::TSPMethodDelegate_OneVar_Const<UserClass, ...>::FMethodPtr has been deprecated - use DelegateType::TConstMethodPtr<UserClass, ...> instead")    TSPMethodDelegate_OneVar_Const    : TBaseSPMethodDelegateInstance<true , UserClass, ESPMode::ThreadSafe, FuncType, UserPolicy, Var1Type                              > { typedef TBaseSPMethodDelegateInstance<true , UserClass, ESPMode::ThreadSafe, FuncType, UserPolicy, Var1Type                              > Super; TSPMethodDelegate_OneVar_Const   (const TSharedRef< UserClass, ESPMode::ThreadSafe >& InUserObject, typename Super::FMethodPtr InMethodPtr, Var1Type Var1                                             ) : Super(InUserObject, InMethodPtr, Var1                  ) {} };
-	template <typename UserClass, typename Var1Type, typename Var2Type                                      > struct UE_DEPRECATED(5.2, "DelegateType::TSPMethodDelegate_TwoVars<UserClass, ...>::FMethodPtr has been deprecated - use DelegateType::TMethodPtr<UserClass, ...> instead")              TSPMethodDelegate_TwoVars         : TBaseSPMethodDelegateInstance<false, UserClass, ESPMode::ThreadSafe, FuncType, UserPolicy, Var1Type, Var2Type                    > { typedef TBaseSPMethodDelegateInstance<false, UserClass, ESPMode::ThreadSafe, FuncType, UserPolicy, Var1Type, Var2Type                    > Super; TSPMethodDelegate_TwoVars        (const TSharedRef< UserClass, ESPMode::ThreadSafe >& InUserObject, typename Super::FMethodPtr InMethodPtr, Var1Type Var1, Var2Type Var2                              ) : Super(InUserObject, InMethodPtr, Var1, Var2            ) {} };
-	template <typename UserClass, typename Var1Type, typename Var2Type                                      > struct UE_DEPRECATED(5.2, "DelegateType::TSPMethodDelegate_TwoVars_Const<UserClass, ...>::FMethodPtr has been deprecated - use DelegateType::TConstMethodPtr<UserClass, ...> instead")   TSPMethodDelegate_TwoVars_Const   : TBaseSPMethodDelegateInstance<true , UserClass, ESPMode::ThreadSafe, FuncType, UserPolicy, Var1Type, Var2Type                    > { typedef TBaseSPMethodDelegateInstance<true , UserClass, ESPMode::ThreadSafe, FuncType, UserPolicy, Var1Type, Var2Type                    > Super; TSPMethodDelegate_TwoVars_Const  (const TSharedRef< UserClass, ESPMode::ThreadSafe >& InUserObject, typename Super::FMethodPtr InMethodPtr, Var1Type Var1, Var2Type Var2                              ) : Super(InUserObject, InMethodPtr, Var1, Var2            ) {} };
-	template <typename UserClass, typename Var1Type, typename Var2Type, typename Var3Type                   > struct UE_DEPRECATED(5.2, "DelegateType::TSPMethodDelegate_ThreeVars<UserClass, ...>::FMethodPtr has been deprecated - use DelegateType::TMethodPtr<UserClass, ...> instead")            TSPMethodDelegate_ThreeVars       : TBaseSPMethodDelegateInstance<false, UserClass, ESPMode::ThreadSafe, FuncType, UserPolicy, Var1Type, Var2Type, Var3Type          > { typedef TBaseSPMethodDelegateInstance<false, UserClass, ESPMode::ThreadSafe, FuncType, UserPolicy, Var1Type, Var2Type, Var3Type          > Super; TSPMethodDelegate_ThreeVars      (const TSharedRef< UserClass, ESPMode::ThreadSafe >& InUserObject, typename Super::FMethodPtr InMethodPtr, Var1Type Var1, Var2Type Var2, Var3Type Var3               ) : Super(InUserObject, InMethodPtr, Var1, Var2, Var3      ) {} };
-	template <typename UserClass, typename Var1Type, typename Var2Type, typename Var3Type                   > struct UE_DEPRECATED(5.2, "DelegateType::TSPMethodDelegate_ThreeVars_Const<UserClass, ...>::FMethodPtr has been deprecated - use DelegateType::TConstMethodPtr<UserClass, ...> instead") TSPMethodDelegate_ThreeVars_Const : TBaseSPMethodDelegateInstance<true , UserClass, ESPMode::ThreadSafe, FuncType, UserPolicy, Var1Type, Var2Type, Var3Type          > { typedef TBaseSPMethodDelegateInstance<true , UserClass, ESPMode::ThreadSafe, FuncType, UserPolicy, Var1Type, Var2Type, Var3Type          > Super; TSPMethodDelegate_ThreeVars_Const(const TSharedRef< UserClass, ESPMode::ThreadSafe >& InUserObject, typename Super::FMethodPtr InMethodPtr, Var1Type Var1, Var2Type Var2, Var3Type Var3               ) : Super(InUserObject, InMethodPtr, Var1, Var2, Var3      ) {} };
-	template <typename UserClass, typename Var1Type, typename Var2Type, typename Var3Type, typename Var4Type> struct UE_DEPRECATED(5.2, "DelegateType::TSPMethodDelegate_FourVars<UserClass, ...>::FMethodPtr has been deprecated - use DelegateType::TMethodPtr<UserClass, ...> instead")             TSPMethodDelegate_FourVars        : TBaseSPMethodDelegateInstance<false, UserClass, ESPMode::ThreadSafe, FuncType, UserPolicy, Var1Type, Var2Type, Var3Type, Var4Type> { typedef TBaseSPMethodDelegateInstance<false, UserClass, ESPMode::ThreadSafe, FuncType, UserPolicy, Var1Type, Var2Type, Var3Type, Var4Type> Super; TSPMethodDelegate_FourVars       (const TSharedRef< UserClass, ESPMode::ThreadSafe >& InUserObject, typename Super::FMethodPtr InMethodPtr, Var1Type Var1, Var2Type Var2, Var3Type Var3, Var4Type Var4) : Super(InUserObject, InMethodPtr, Var1, Var2, Var3, Var4) {} };
-	template <typename UserClass, typename Var1Type, typename Var2Type, typename Var3Type, typename Var4Type> struct UE_DEPRECATED(5.2, "DelegateType::TSPMethodDelegate_FourVars_Const<UserClass, ...>::FMethodPtr has been deprecated - use DelegateType::TConstMethodPtr<UserClass, ...> instead")  TSPMethodDelegate_FourVars_Const  : TBaseSPMethodDelegateInstance<true , UserClass, ESPMode::ThreadSafe, FuncType, UserPolicy, Var1Type, Var2Type, Var3Type, Var4Type> { typedef TBaseSPMethodDelegateInstance<true , UserClass, ESPMode::ThreadSafe, FuncType, UserPolicy, Var1Type, Var2Type, Var3Type, Var4Type> Super; TSPMethodDelegate_FourVars_Const (const TSharedRef< UserClass, ESPMode::ThreadSafe >& InUserObject, typename Super::FMethodPtr InMethodPtr, Var1Type Var1, Var2Type Var2, Var3Type Var3, Var4Type Var4) : Super(InUserObject, InMethodPtr, Var1, Var2, Var3, Var4) {} };
+	template <typename UserClass                                                                            > struct TSPMethodDelegate : TBaseSPMethodDelegateInstance<false, UserClass, ESPMode::ThreadSafe, FuncType, UserPolicy                                        > { typedef TBaseSPMethodDelegateInstance<false, UserClass, ESPMode::ThreadSafe, FuncType, UserPolicy                                        > Super; TSPMethodDelegate(const TSharedRef< UserClass, ESPMode::ThreadSafe >& InUserObject, typename Super::FMethodPtr InMethodPtr) : Super(InUserObject, InMethodPtr) {} };
+	template <typename UserClass                                                                            > struct TSPMethodDelegate_Const : TBaseSPMethodDelegateInstance<true, UserClass, ESPMode::ThreadSafe, FuncType, UserPolicy                                        > { typedef TBaseSPMethodDelegateInstance<true, UserClass, ESPMode::ThreadSafe, FuncType, UserPolicy                                        > Super; TSPMethodDelegate_Const(const TSharedRef< UserClass, ESPMode::ThreadSafe >& InUserObject, typename Super::FMethodPtr InMethodPtr) : Super(InUserObject, InMethodPtr) {} };
+	template <typename UserClass, typename Var1Type                                                         > struct TSPMethodDelegate_OneVar : TBaseSPMethodDelegateInstance<false, UserClass, ESPMode::ThreadSafe, FuncType, UserPolicy, Var1Type                              > { typedef TBaseSPMethodDelegateInstance<false, UserClass, ESPMode::ThreadSafe, FuncType, UserPolicy, Var1Type                              > Super; TSPMethodDelegate_OneVar(const TSharedRef< UserClass, ESPMode::ThreadSafe >& InUserObject, typename Super::FMethodPtr InMethodPtr, Var1Type Var1) : Super(InUserObject, InMethodPtr, Var1) {} };
+	template <typename UserClass, typename Var1Type                                                         > struct TSPMethodDelegate_OneVar_Const : TBaseSPMethodDelegateInstance<true, UserClass, ESPMode::ThreadSafe, FuncType, UserPolicy, Var1Type                              > { typedef TBaseSPMethodDelegateInstance<true, UserClass, ESPMode::ThreadSafe, FuncType, UserPolicy, Var1Type                              > Super; TSPMethodDelegate_OneVar_Const(const TSharedRef< UserClass, ESPMode::ThreadSafe >& InUserObject, typename Super::FMethodPtr InMethodPtr, Var1Type Var1) : Super(InUserObject, InMethodPtr, Var1) {} };
+	template <typename UserClass, typename Var1Type, typename Var2Type                                      > struct TSPMethodDelegate_TwoVars : TBaseSPMethodDelegateInstance<false, UserClass, ESPMode::ThreadSafe, FuncType, UserPolicy, Var1Type, Var2Type                    > { typedef TBaseSPMethodDelegateInstance<false, UserClass, ESPMode::ThreadSafe, FuncType, UserPolicy, Var1Type, Var2Type                    > Super; TSPMethodDelegate_TwoVars(const TSharedRef< UserClass, ESPMode::ThreadSafe >& InUserObject, typename Super::FMethodPtr InMethodPtr, Var1Type Var1, Var2Type Var2) : Super(InUserObject, InMethodPtr, Var1, Var2) {} };
+	template <typename UserClass, typename Var1Type, typename Var2Type                                      > struct TSPMethodDelegate_TwoVars_Const : TBaseSPMethodDelegateInstance<true, UserClass, ESPMode::ThreadSafe, FuncType, UserPolicy, Var1Type, Var2Type                    > { typedef TBaseSPMethodDelegateInstance<true, UserClass, ESPMode::ThreadSafe, FuncType, UserPolicy, Var1Type, Var2Type                    > Super; TSPMethodDelegate_TwoVars_Const(const TSharedRef< UserClass, ESPMode::ThreadSafe >& InUserObject, typename Super::FMethodPtr InMethodPtr, Var1Type Var1, Var2Type Var2) : Super(InUserObject, InMethodPtr, Var1, Var2) {} };
+	template <typename UserClass, typename Var1Type, typename Var2Type, typename Var3Type                   > struct TSPMethodDelegate_ThreeVars : TBaseSPMethodDelegateInstance<false, UserClass, ESPMode::ThreadSafe, FuncType, UserPolicy, Var1Type, Var2Type, Var3Type          > { typedef TBaseSPMethodDelegateInstance<false, UserClass, ESPMode::ThreadSafe, FuncType, UserPolicy, Var1Type, Var2Type, Var3Type          > Super; TSPMethodDelegate_ThreeVars(const TSharedRef< UserClass, ESPMode::ThreadSafe >& InUserObject, typename Super::FMethodPtr InMethodPtr, Var1Type Var1, Var2Type Var2, Var3Type Var3) : Super(InUserObject, InMethodPtr, Var1, Var2, Var3) {} };
+	template <typename UserClass, typename Var1Type, typename Var2Type, typename Var3Type                   > struct TSPMethodDelegate_ThreeVars_Const : TBaseSPMethodDelegateInstance<true, UserClass, ESPMode::ThreadSafe, FuncType, UserPolicy, Var1Type, Var2Type, Var3Type          > { typedef TBaseSPMethodDelegateInstance<true, UserClass, ESPMode::ThreadSafe, FuncType, UserPolicy, Var1Type, Var2Type, Var3Type          > Super; TSPMethodDelegate_ThreeVars_Const(const TSharedRef< UserClass, ESPMode::ThreadSafe >& InUserObject, typename Super::FMethodPtr InMethodPtr, Var1Type Var1, Var2Type Var2, Var3Type Var3) : Super(InUserObject, InMethodPtr, Var1, Var2, Var3) {} };
+	template <typename UserClass, typename Var1Type, typename Var2Type, typename Var3Type, typename Var4Type> struct TSPMethodDelegate_FourVars : TBaseSPMethodDelegateInstance<false, UserClass, ESPMode::ThreadSafe, FuncType, UserPolicy, Var1Type, Var2Type, Var3Type, Var4Type> { typedef TBaseSPMethodDelegateInstance<false, UserClass, ESPMode::ThreadSafe, FuncType, UserPolicy, Var1Type, Var2Type, Var3Type, Var4Type> Super; TSPMethodDelegate_FourVars(const TSharedRef< UserClass, ESPMode::ThreadSafe >& InUserObject, typename Super::FMethodPtr InMethodPtr, Var1Type Var1, Var2Type Var2, Var3Type Var3, Var4Type Var4) : Super(InUserObject, InMethodPtr, Var1, Var2, Var3, Var4) {} };
+	template <typename UserClass, typename Var1Type, typename Var2Type, typename Var3Type, typename Var4Type> struct TSPMethodDelegate_FourVars_Const : TBaseSPMethodDelegateInstance<true, UserClass, ESPMode::ThreadSafe, FuncType, UserPolicy, Var1Type, Var2Type, Var3Type, Var4Type> { typedef TBaseSPMethodDelegateInstance<true, UserClass, ESPMode::ThreadSafe, FuncType, UserPolicy, Var1Type, Var2Type, Var3Type, Var4Type> Super; TSPMethodDelegate_FourVars_Const(const TSharedRef< UserClass, ESPMode::ThreadSafe >& InUserObject, typename Super::FMethodPtr InMethodPtr, Var1Type Var1, Var2Type Var2, Var3Type Var3, Var4Type Var4) : Super(InUserObject, InMethodPtr, Var1, Var2, Var3, Var4) {} };
 
 	/** Declare the user's shared pointer-based delegate instance types. */
-	template <typename UserClass                                                                            > struct UE_DEPRECATED(5.2, "DelegateType::TThreadSafeSPMethodDelegate<UserClass, ...>::FMethodPtr has been deprecated - use DelegateType::TMethodPtr<UserClass, ...> instead")                      TThreadSafeSPMethodDelegate                 : TBaseSPMethodDelegateInstance<false, UserClass, ESPMode::ThreadSafe, FuncType, UserPolicy                                        > { typedef TBaseSPMethodDelegateInstance<false, UserClass, ESPMode::ThreadSafe, FuncType, UserPolicy                                        > Super; TThreadSafeSPMethodDelegate                (const TSharedRef<UserClass, ESPMode::ThreadSafe>& InUserObject, typename Super::FMethodPtr InMethodPtr                                                            ) : Super(InUserObject, InMethodPtr                        ) {} };
-	template <typename UserClass                                                                            > struct UE_DEPRECATED(5.2, "DelegateType::TThreadSafeSPMethodDelegate_Const<UserClass, ...>::FMethodPtr has been deprecated - use DelegateType::TConstMethodPtr<UserClass, ...> instead")           TThreadSafeSPMethodDelegate_Const           : TBaseSPMethodDelegateInstance<true , UserClass, ESPMode::ThreadSafe, FuncType, UserPolicy                                        > { typedef TBaseSPMethodDelegateInstance<true , UserClass, ESPMode::ThreadSafe, FuncType, UserPolicy                                        > Super; TThreadSafeSPMethodDelegate_Const          (const TSharedRef<UserClass, ESPMode::ThreadSafe>& InUserObject, typename Super::FMethodPtr InMethodPtr                                                            ) : Super(InUserObject, InMethodPtr                        ) {} };
-	template <typename UserClass, typename Var1Type                                                         > struct UE_DEPRECATED(5.2, "DelegateType::TThreadSafeSPMethodDelegate_OneVar<UserClass, ...>::FMethodPtr has been deprecated - use DelegateType::TMethodPtr<UserClass, ...> instead")               TThreadSafeSPMethodDelegate_OneVar          : TBaseSPMethodDelegateInstance<false, UserClass, ESPMode::ThreadSafe, FuncType, UserPolicy, Var1Type                              > { typedef TBaseSPMethodDelegateInstance<false, UserClass, ESPMode::ThreadSafe, FuncType, UserPolicy, Var1Type                              > Super; TThreadSafeSPMethodDelegate_OneVar         (const TSharedRef<UserClass, ESPMode::ThreadSafe>& InUserObject, typename Super::FMethodPtr InMethodPtr, Var1Type Var1                                             ) : Super(InUserObject, InMethodPtr, Var1                  ) {} };
-	template <typename UserClass, typename Var1Type                                                         > struct UE_DEPRECATED(5.2, "DelegateType::TThreadSafeSPMethodDelegate_OneVar_Const<UserClass, ...>::FMethodPtr has been deprecated - use DelegateType::TConstMethodPtr<UserClass, ...> instead")    TThreadSafeSPMethodDelegate_OneVar_Const    : TBaseSPMethodDelegateInstance<true , UserClass, ESPMode::ThreadSafe, FuncType, UserPolicy, Var1Type                              > { typedef TBaseSPMethodDelegateInstance<true , UserClass, ESPMode::ThreadSafe, FuncType, UserPolicy, Var1Type                              > Super; TThreadSafeSPMethodDelegate_OneVar_Const   (const TSharedRef<UserClass, ESPMode::ThreadSafe>& InUserObject, typename Super::FMethodPtr InMethodPtr, Var1Type Var1                                             ) : Super(InUserObject, InMethodPtr, Var1                  ) {} };
-	template <typename UserClass, typename Var1Type, typename Var2Type                                      > struct UE_DEPRECATED(5.2, "DelegateType::TThreadSafeSPMethodDelegate_TwoVars<UserClass, ...>::FMethodPtr has been deprecated - use DelegateType::TMethodPtr<UserClass, ...> instead")              TThreadSafeSPMethodDelegate_TwoVars         : TBaseSPMethodDelegateInstance<false, UserClass, ESPMode::ThreadSafe, FuncType, UserPolicy, Var1Type, Var2Type                    > { typedef TBaseSPMethodDelegateInstance<false, UserClass, ESPMode::ThreadSafe, FuncType, UserPolicy, Var1Type, Var2Type                    > Super; TThreadSafeSPMethodDelegate_TwoVars        (const TSharedRef<UserClass, ESPMode::ThreadSafe>& InUserObject, typename Super::FMethodPtr InMethodPtr, Var1Type Var1, Var2Type Var2                              ) : Super(InUserObject, InMethodPtr, Var1, Var2            ) {} };
-	template <typename UserClass, typename Var1Type, typename Var2Type                                      > struct UE_DEPRECATED(5.2, "DelegateType::TThreadSafeSPMethodDelegate_TwoVars_Const<UserClass, ...>::FMethodPtr has been deprecated - use DelegateType::TConstMethodPtr<UserClass, ...> instead")   TThreadSafeSPMethodDelegate_TwoVars_Const   : TBaseSPMethodDelegateInstance<true , UserClass, ESPMode::ThreadSafe, FuncType, UserPolicy, Var1Type, Var2Type                    > { typedef TBaseSPMethodDelegateInstance<true , UserClass, ESPMode::ThreadSafe, FuncType, UserPolicy, Var1Type, Var2Type                    > Super; TThreadSafeSPMethodDelegate_TwoVars_Const  (const TSharedRef<UserClass, ESPMode::ThreadSafe>& InUserObject, typename Super::FMethodPtr InMethodPtr, Var1Type Var1, Var2Type Var2                              ) : Super(InUserObject, InMethodPtr, Var1, Var2            ) {} };
-	template <typename UserClass, typename Var1Type, typename Var2Type, typename Var3Type                   > struct UE_DEPRECATED(5.2, "DelegateType::TThreadSafeSPMethodDelegate_ThreeVars<UserClass, ...>::FMethodPtr has been deprecated - use DelegateType::TMethodPtr<UserClass, ...> instead")            TThreadSafeSPMethodDelegate_ThreeVars       : TBaseSPMethodDelegateInstance<false, UserClass, ESPMode::ThreadSafe, FuncType, UserPolicy, Var1Type, Var2Type, Var3Type          > { typedef TBaseSPMethodDelegateInstance<false, UserClass, ESPMode::ThreadSafe, FuncType, UserPolicy, Var1Type, Var2Type, Var3Type          > Super; TThreadSafeSPMethodDelegate_ThreeVars      (const TSharedRef<UserClass, ESPMode::ThreadSafe>& InUserObject, typename Super::FMethodPtr InMethodPtr, Var1Type Var1, Var2Type Var2, Var3Type Var3               ) : Super(InUserObject, InMethodPtr, Var1, Var2, Var3      ) {} };
-	template <typename UserClass, typename Var1Type, typename Var2Type, typename Var3Type                   > struct UE_DEPRECATED(5.2, "DelegateType::TThreadSafeSPMethodDelegate_ThreeVars_Const<UserClass, ...>::FMethodPtr has been deprecated - use DelegateType::TConstMethodPtr<UserClass, ...> instead") TThreadSafeSPMethodDelegate_ThreeVars_Const : TBaseSPMethodDelegateInstance<true , UserClass, ESPMode::ThreadSafe, FuncType, UserPolicy, Var1Type, Var2Type, Var3Type          > { typedef TBaseSPMethodDelegateInstance<true , UserClass, ESPMode::ThreadSafe, FuncType, UserPolicy, Var1Type, Var2Type, Var3Type          > Super; TThreadSafeSPMethodDelegate_ThreeVars_Const(const TSharedRef<UserClass, ESPMode::ThreadSafe>& InUserObject, typename Super::FMethodPtr InMethodPtr, Var1Type Var1, Var2Type Var2, Var3Type Var3               ) : Super(InUserObject, InMethodPtr, Var1, Var2, Var3      ) {} };
-	template <typename UserClass, typename Var1Type, typename Var2Type, typename Var3Type, typename Var4Type> struct UE_DEPRECATED(5.2, "DelegateType::TThreadSafeSPMethodDelegate_FourVars<UserClass, ...>::FMethodPtr has been deprecated - use DelegateType::TMethodPtr<UserClass, ...> instead")             TThreadSafeSPMethodDelegate_FourVars        : TBaseSPMethodDelegateInstance<false, UserClass, ESPMode::ThreadSafe, FuncType, UserPolicy, Var1Type, Var2Type, Var3Type, Var4Type> { typedef TBaseSPMethodDelegateInstance<false, UserClass, ESPMode::ThreadSafe, FuncType, UserPolicy, Var1Type, Var2Type, Var3Type, Var4Type> Super; TThreadSafeSPMethodDelegate_FourVars       (const TSharedRef<UserClass, ESPMode::ThreadSafe>& InUserObject, typename Super::FMethodPtr InMethodPtr, Var1Type Var1, Var2Type Var2, Var3Type Var3, Var4Type Var4) : Super(InUserObject, InMethodPtr, Var1, Var2, Var3, Var4) {} };
-	template <typename UserClass, typename Var1Type, typename Var2Type, typename Var3Type, typename Var4Type> struct UE_DEPRECATED(5.2, "DelegateType::TThreadSafeSPMethodDelegate_FourVars_Const<UserClass, ...>::FMethodPtr has been deprecated - use DelegateType::TConstMethodPtr<UserClass, ...> instead")  TThreadSafeSPMethodDelegate_FourVars_Const  : TBaseSPMethodDelegateInstance<true , UserClass, ESPMode::ThreadSafe, FuncType, UserPolicy, Var1Type, Var2Type, Var3Type, Var4Type> { typedef TBaseSPMethodDelegateInstance<true , UserClass, ESPMode::ThreadSafe, FuncType, UserPolicy, Var1Type, Var2Type, Var3Type, Var4Type> Super; TThreadSafeSPMethodDelegate_FourVars_Const (const TSharedRef<UserClass, ESPMode::ThreadSafe>& InUserObject, typename Super::FMethodPtr InMethodPtr, Var1Type Var1, Var2Type Var2, Var3Type Var3, Var4Type Var4) : Super(InUserObject, InMethodPtr, Var1, Var2, Var3, Var4) {} };
+	template <typename UserClass                                                                            > struct TThreadSafeSPMethodDelegate : TBaseSPMethodDelegateInstance<false, UserClass, ESPMode::ThreadSafe, FuncType, UserPolicy                                        > { typedef TBaseSPMethodDelegateInstance<false, UserClass, ESPMode::ThreadSafe, FuncType, UserPolicy                                        > Super; TThreadSafeSPMethodDelegate(const TSharedRef<UserClass, ESPMode::ThreadSafe>& InUserObject, typename Super::FMethodPtr InMethodPtr) : Super(InUserObject, InMethodPtr) {} };
+	template <typename UserClass                                                                            > struct TThreadSafeSPMethodDelegate_Const : TBaseSPMethodDelegateInstance<true, UserClass, ESPMode::ThreadSafe, FuncType, UserPolicy                                        > { typedef TBaseSPMethodDelegateInstance<true, UserClass, ESPMode::ThreadSafe, FuncType, UserPolicy                                        > Super; TThreadSafeSPMethodDelegate_Const(const TSharedRef<UserClass, ESPMode::ThreadSafe>& InUserObject, typename Super::FMethodPtr InMethodPtr) : Super(InUserObject, InMethodPtr) {} };
+	template <typename UserClass, typename Var1Type                                                         > struct TThreadSafeSPMethodDelegate_OneVar : TBaseSPMethodDelegateInstance<false, UserClass, ESPMode::ThreadSafe, FuncType, UserPolicy, Var1Type                              > { typedef TBaseSPMethodDelegateInstance<false, UserClass, ESPMode::ThreadSafe, FuncType, UserPolicy, Var1Type                              > Super; TThreadSafeSPMethodDelegate_OneVar(const TSharedRef<UserClass, ESPMode::ThreadSafe>& InUserObject, typename Super::FMethodPtr InMethodPtr, Var1Type Var1) : Super(InUserObject, InMethodPtr, Var1) {} };
+	template <typename UserClass, typename Var1Type                                                         > struct TThreadSafeSPMethodDelegate_OneVar_Const : TBaseSPMethodDelegateInstance<true, UserClass, ESPMode::ThreadSafe, FuncType, UserPolicy, Var1Type                              > { typedef TBaseSPMethodDelegateInstance<true, UserClass, ESPMode::ThreadSafe, FuncType, UserPolicy, Var1Type                              > Super; TThreadSafeSPMethodDelegate_OneVar_Const(const TSharedRef<UserClass, ESPMode::ThreadSafe>& InUserObject, typename Super::FMethodPtr InMethodPtr, Var1Type Var1) : Super(InUserObject, InMethodPtr, Var1) {} };
+	template <typename UserClass, typename Var1Type, typename Var2Type                                      > struct TThreadSafeSPMethodDelegate_TwoVars : TBaseSPMethodDelegateInstance<false, UserClass, ESPMode::ThreadSafe, FuncType, UserPolicy, Var1Type, Var2Type                    > { typedef TBaseSPMethodDelegateInstance<false, UserClass, ESPMode::ThreadSafe, FuncType, UserPolicy, Var1Type, Var2Type                    > Super; TThreadSafeSPMethodDelegate_TwoVars(const TSharedRef<UserClass, ESPMode::ThreadSafe>& InUserObject, typename Super::FMethodPtr InMethodPtr, Var1Type Var1, Var2Type Var2) : Super(InUserObject, InMethodPtr, Var1, Var2) {} };
+	template <typename UserClass, typename Var1Type, typename Var2Type                                      > struct TThreadSafeSPMethodDelegate_TwoVars_Const : TBaseSPMethodDelegateInstance<true, UserClass, ESPMode::ThreadSafe, FuncType, UserPolicy, Var1Type, Var2Type                    > { typedef TBaseSPMethodDelegateInstance<true, UserClass, ESPMode::ThreadSafe, FuncType, UserPolicy, Var1Type, Var2Type                    > Super; TThreadSafeSPMethodDelegate_TwoVars_Const(const TSharedRef<UserClass, ESPMode::ThreadSafe>& InUserObject, typename Super::FMethodPtr InMethodPtr, Var1Type Var1, Var2Type Var2) : Super(InUserObject, InMethodPtr, Var1, Var2) {} };
+	template <typename UserClass, typename Var1Type, typename Var2Type, typename Var3Type                   > struct TThreadSafeSPMethodDelegate_ThreeVars : TBaseSPMethodDelegateInstance<false, UserClass, ESPMode::ThreadSafe, FuncType, UserPolicy, Var1Type, Var2Type, Var3Type          > { typedef TBaseSPMethodDelegateInstance<false, UserClass, ESPMode::ThreadSafe, FuncType, UserPolicy, Var1Type, Var2Type, Var3Type          > Super; TThreadSafeSPMethodDelegate_ThreeVars(const TSharedRef<UserClass, ESPMode::ThreadSafe>& InUserObject, typename Super::FMethodPtr InMethodPtr, Var1Type Var1, Var2Type Var2, Var3Type Var3) : Super(InUserObject, InMethodPtr, Var1, Var2, Var3) {} };
+	template <typename UserClass, typename Var1Type, typename Var2Type, typename Var3Type                   > struct TThreadSafeSPMethodDelegate_ThreeVars_Const : TBaseSPMethodDelegateInstance<true, UserClass, ESPMode::ThreadSafe, FuncType, UserPolicy, Var1Type, Var2Type, Var3Type          > { typedef TBaseSPMethodDelegateInstance<true, UserClass, ESPMode::ThreadSafe, FuncType, UserPolicy, Var1Type, Var2Type, Var3Type          > Super; TThreadSafeSPMethodDelegate_ThreeVars_Const(const TSharedRef<UserClass, ESPMode::ThreadSafe>& InUserObject, typename Super::FMethodPtr InMethodPtr, Var1Type Var1, Var2Type Var2, Var3Type Var3) : Super(InUserObject, InMethodPtr, Var1, Var2, Var3) {} };
+	template <typename UserClass, typename Var1Type, typename Var2Type, typename Var3Type, typename Var4Type> struct TThreadSafeSPMethodDelegate_FourVars : TBaseSPMethodDelegateInstance<false, UserClass, ESPMode::ThreadSafe, FuncType, UserPolicy, Var1Type, Var2Type, Var3Type, Var4Type> { typedef TBaseSPMethodDelegateInstance<false, UserClass, ESPMode::ThreadSafe, FuncType, UserPolicy, Var1Type, Var2Type, Var3Type, Var4Type> Super; TThreadSafeSPMethodDelegate_FourVars(const TSharedRef<UserClass, ESPMode::ThreadSafe>& InUserObject, typename Super::FMethodPtr InMethodPtr, Var1Type Var1, Var2Type Var2, Var3Type Var3, Var4Type Var4) : Super(InUserObject, InMethodPtr, Var1, Var2, Var3, Var4) {} };
+	template <typename UserClass, typename Var1Type, typename Var2Type, typename Var3Type, typename Var4Type> struct TThreadSafeSPMethodDelegate_FourVars_Const : TBaseSPMethodDelegateInstance<true, UserClass, ESPMode::ThreadSafe, FuncType, UserPolicy, Var1Type, Var2Type, Var3Type, Var4Type> { typedef TBaseSPMethodDelegateInstance<true, UserClass, ESPMode::ThreadSafe, FuncType, UserPolicy, Var1Type, Var2Type, Var3Type, Var4Type> Super; TThreadSafeSPMethodDelegate_FourVars_Const(const TSharedRef<UserClass, ESPMode::ThreadSafe>& InUserObject, typename Super::FMethodPtr InMethodPtr, Var1Type Var1, Var2Type Var2, Var3Type Var3, Var4Type Var4) : Super(InUserObject, InMethodPtr, Var1, Var2, Var3, Var4) {} };
 
 	/** Declare the user's C++ pointer-based delegate instance types. */
-	template <typename UserClass                                                                            > struct UE_DEPRECATED(5.2, "DelegateType::TRawMethodDelegate<UserClass, ...>::FMethodPtr has been deprecated - use DelegateType::TMethodPtr<UserClass, ...> instead")                      TRawMethodDelegate                 : TBaseRawMethodDelegateInstance<false, UserClass, FuncType, UserPolicy                                        > { typedef TBaseRawMethodDelegateInstance<false, UserClass, FuncType, UserPolicy                                        > Super; TRawMethodDelegate                (UserClass* InUserObject, typename Super::FMethodPtr InMethodPtr                                                            ) : Super(InUserObject, InMethodPtr                        ) {} };
-	template <typename UserClass                                                                            > struct UE_DEPRECATED(5.2, "DelegateType::TRawMethodDelegate_Const<UserClass, ...>::FMethodPtr has been deprecated - use DelegateType::TConstMethodPtr<UserClass, ...> instead")           TRawMethodDelegate_Const           : TBaseRawMethodDelegateInstance<true , UserClass, FuncType, UserPolicy                                        > { typedef TBaseRawMethodDelegateInstance<true , UserClass, FuncType, UserPolicy                                        > Super; TRawMethodDelegate_Const          (UserClass* InUserObject, typename Super::FMethodPtr InMethodPtr                                                            ) : Super(InUserObject, InMethodPtr                        ) {} };
-	template <typename UserClass, typename Var1Type                                                         > struct UE_DEPRECATED(5.2, "DelegateType::TRawMethodDelegate_OneVar<UserClass, ...>::FMethodPtr has been deprecated - use DelegateType::TMethodPtr<UserClass, ...> instead")               TRawMethodDelegate_OneVar          : TBaseRawMethodDelegateInstance<false, UserClass, FuncType, UserPolicy, Var1Type                              > { typedef TBaseRawMethodDelegateInstance<false, UserClass, FuncType, UserPolicy, Var1Type                              > Super; TRawMethodDelegate_OneVar         (UserClass* InUserObject, typename Super::FMethodPtr InMethodPtr, Var1Type Var1                                             ) : Super(InUserObject, InMethodPtr, Var1                  ) {} };
-	template <typename UserClass, typename Var1Type                                                         > struct UE_DEPRECATED(5.2, "DelegateType::TRawMethodDelegate_OneVar_Const<UserClass, ...>::FMethodPtr has been deprecated - use DelegateType::TConstMethodPtr<UserClass, ...> instead")    TRawMethodDelegate_OneVar_Const    : TBaseRawMethodDelegateInstance<true , UserClass, FuncType, UserPolicy, Var1Type                              > { typedef TBaseRawMethodDelegateInstance<true , UserClass, FuncType, UserPolicy, Var1Type                              > Super; TRawMethodDelegate_OneVar_Const   (UserClass* InUserObject, typename Super::FMethodPtr InMethodPtr, Var1Type Var1                                             ) : Super(InUserObject, InMethodPtr, Var1                  ) {} };
-	template <typename UserClass, typename Var1Type, typename Var2Type                                      > struct UE_DEPRECATED(5.2, "DelegateType::TRawMethodDelegate_TwoVars<UserClass, ...>::FMethodPtr has been deprecated - use DelegateType::TMethodPtr<UserClass, ...> instead")              TRawMethodDelegate_TwoVars         : TBaseRawMethodDelegateInstance<false, UserClass, FuncType, UserPolicy, Var1Type, Var2Type                    > { typedef TBaseRawMethodDelegateInstance<false, UserClass, FuncType, UserPolicy, Var1Type, Var2Type                    > Super; TRawMethodDelegate_TwoVars        (UserClass* InUserObject, typename Super::FMethodPtr InMethodPtr, Var1Type Var1, Var2Type Var2                              ) : Super(InUserObject, InMethodPtr, Var1, Var2            ) {} };
-	template <typename UserClass, typename Var1Type, typename Var2Type                                      > struct UE_DEPRECATED(5.2, "DelegateType::TRawMethodDelegate_TwoVars_Const<UserClass, ...>::FMethodPtr has been deprecated - use DelegateType::TConstMethodPtr<UserClass, ...> instead")   TRawMethodDelegate_TwoVars_Const   : TBaseRawMethodDelegateInstance<true , UserClass, FuncType, UserPolicy, Var1Type, Var2Type                    > { typedef TBaseRawMethodDelegateInstance<true , UserClass, FuncType, UserPolicy, Var1Type, Var2Type                    > Super; TRawMethodDelegate_TwoVars_Const  (UserClass* InUserObject, typename Super::FMethodPtr InMethodPtr, Var1Type Var1, Var2Type Var2                              ) : Super(InUserObject, InMethodPtr, Var1, Var2            ) {} };
-	template <typename UserClass, typename Var1Type, typename Var2Type, typename Var3Type                   > struct UE_DEPRECATED(5.2, "DelegateType::TRawMethodDelegate_ThreeVars<UserClass, ...>::FMethodPtr has been deprecated - use DelegateType::TMethodPtr<UserClass, ...> instead")            TRawMethodDelegate_ThreeVars       : TBaseRawMethodDelegateInstance<false, UserClass, FuncType, UserPolicy, Var1Type, Var2Type, Var3Type          > { typedef TBaseRawMethodDelegateInstance<false, UserClass, FuncType, UserPolicy, Var1Type, Var2Type, Var3Type          > Super; TRawMethodDelegate_ThreeVars      (UserClass* InUserObject, typename Super::FMethodPtr InMethodPtr, Var1Type Var1, Var2Type Var2, Var3Type Var3               ) : Super(InUserObject, InMethodPtr, Var1, Var2, Var3      ) {} };
-	template <typename UserClass, typename Var1Type, typename Var2Type, typename Var3Type                   > struct UE_DEPRECATED(5.2, "DelegateType::TRawMethodDelegate_ThreeVars_Const<UserClass, ...>::FMethodPtr has been deprecated - use DelegateType::TConstMethodPtr<UserClass, ...> instead") TRawMethodDelegate_ThreeVars_Const : TBaseRawMethodDelegateInstance<true , UserClass, FuncType, UserPolicy, Var1Type, Var2Type, Var3Type          > { typedef TBaseRawMethodDelegateInstance<true , UserClass, FuncType, UserPolicy, Var1Type, Var2Type, Var3Type          > Super; TRawMethodDelegate_ThreeVars_Const(UserClass* InUserObject, typename Super::FMethodPtr InMethodPtr, Var1Type Var1, Var2Type Var2, Var3Type Var3               ) : Super(InUserObject, InMethodPtr, Var1, Var2, Var3      ) {} };
-	template <typename UserClass, typename Var1Type, typename Var2Type, typename Var3Type, typename Var4Type> struct UE_DEPRECATED(5.2, "DelegateType::TRawMethodDelegate_FourVars<UserClass, ...>::FMethodPtr has been deprecated - use DelegateType::TMethodPtr<UserClass, ...> instead")             TRawMethodDelegate_FourVars        : TBaseRawMethodDelegateInstance<false, UserClass, FuncType, UserPolicy, Var1Type, Var2Type, Var3Type, Var4Type> { typedef TBaseRawMethodDelegateInstance<false, UserClass, FuncType, UserPolicy, Var1Type, Var2Type, Var3Type, Var4Type> Super; TRawMethodDelegate_FourVars       (UserClass* InUserObject, typename Super::FMethodPtr InMethodPtr, Var1Type Var1, Var2Type Var2, Var3Type Var3, Var4Type Var4) : Super(InUserObject, InMethodPtr, Var1, Var2, Var3, Var4) {} };
-	template <typename UserClass, typename Var1Type, typename Var2Type, typename Var3Type, typename Var4Type> struct UE_DEPRECATED(5.2, "DelegateType::TRawMethodDelegate_FourVars_Const<UserClass, ...>::FMethodPtr has been deprecated - use DelegateType::TConstMethodPtr<UserClass, ...> instead")  TRawMethodDelegate_FourVars_Const  : TBaseRawMethodDelegateInstance<true , UserClass, FuncType, UserPolicy, Var1Type, Var2Type, Var3Type, Var4Type> { typedef TBaseRawMethodDelegateInstance<true , UserClass, FuncType, UserPolicy, Var1Type, Var2Type, Var3Type, Var4Type> Super; TRawMethodDelegate_FourVars_Const (UserClass* InUserObject, typename Super::FMethodPtr InMethodPtr, Var1Type Var1, Var2Type Var2, Var3Type Var3, Var4Type Var4) : Super(InUserObject, InMethodPtr, Var1, Var2, Var3, Var4) {} };
-	
+	template <typename UserClass                                                                            > struct TRawMethodDelegate : TBaseRawMethodDelegateInstance<false, UserClass, FuncType, UserPolicy                                        > { typedef TBaseRawMethodDelegateInstance<false, UserClass, FuncType, UserPolicy                                        > Super; TRawMethodDelegate(UserClass* InUserObject, typename Super::FMethodPtr InMethodPtr) : Super(InUserObject, InMethodPtr) {} };
+	template <typename UserClass                                                                            > struct TRawMethodDelegate_Const : TBaseRawMethodDelegateInstance<true, UserClass, FuncType, UserPolicy                                        > { typedef TBaseRawMethodDelegateInstance<true, UserClass, FuncType, UserPolicy                                        > Super; TRawMethodDelegate_Const(UserClass* InUserObject, typename Super::FMethodPtr InMethodPtr) : Super(InUserObject, InMethodPtr) {} };
+	template <typename UserClass, typename Var1Type                                                         > struct TRawMethodDelegate_OneVar : TBaseRawMethodDelegateInstance<false, UserClass, FuncType, UserPolicy, Var1Type                              > { typedef TBaseRawMethodDelegateInstance<false, UserClass, FuncType, UserPolicy, Var1Type                              > Super; TRawMethodDelegate_OneVar(UserClass* InUserObject, typename Super::FMethodPtr InMethodPtr, Var1Type Var1) : Super(InUserObject, InMethodPtr, Var1) {} };
+	template <typename UserClass, typename Var1Type                                                         > struct TRawMethodDelegate_OneVar_Const : TBaseRawMethodDelegateInstance<true, UserClass, FuncType, UserPolicy, Var1Type                              > { typedef TBaseRawMethodDelegateInstance<true, UserClass, FuncType, UserPolicy, Var1Type                              > Super; TRawMethodDelegate_OneVar_Const(UserClass* InUserObject, typename Super::FMethodPtr InMethodPtr, Var1Type Var1) : Super(InUserObject, InMethodPtr, Var1) {} };
+	template <typename UserClass, typename Var1Type, typename Var2Type                                      > struct TRawMethodDelegate_TwoVars : TBaseRawMethodDelegateInstance<false, UserClass, FuncType, UserPolicy, Var1Type, Var2Type                    > { typedef TBaseRawMethodDelegateInstance<false, UserClass, FuncType, UserPolicy, Var1Type, Var2Type                    > Super; TRawMethodDelegate_TwoVars(UserClass* InUserObject, typename Super::FMethodPtr InMethodPtr, Var1Type Var1, Var2Type Var2) : Super(InUserObject, InMethodPtr, Var1, Var2) {} };
+	template <typename UserClass, typename Var1Type, typename Var2Type                                      > struct TRawMethodDelegate_TwoVars_Const : TBaseRawMethodDelegateInstance<true, UserClass, FuncType, UserPolicy, Var1Type, Var2Type                    > { typedef TBaseRawMethodDelegateInstance<true, UserClass, FuncType, UserPolicy, Var1Type, Var2Type                    > Super; TRawMethodDelegate_TwoVars_Const(UserClass* InUserObject, typename Super::FMethodPtr InMethodPtr, Var1Type Var1, Var2Type Var2) : Super(InUserObject, InMethodPtr, Var1, Var2) {} };
+	template <typename UserClass, typename Var1Type, typename Var2Type, typename Var3Type                   > struct TRawMethodDelegate_ThreeVars : TBaseRawMethodDelegateInstance<false, UserClass, FuncType, UserPolicy, Var1Type, Var2Type, Var3Type          > { typedef TBaseRawMethodDelegateInstance<false, UserClass, FuncType, UserPolicy, Var1Type, Var2Type, Var3Type          > Super; TRawMethodDelegate_ThreeVars(UserClass* InUserObject, typename Super::FMethodPtr InMethodPtr, Var1Type Var1, Var2Type Var2, Var3Type Var3) : Super(InUserObject, InMethodPtr, Var1, Var2, Var3) {} };
+	template <typename UserClass, typename Var1Type, typename Var2Type, typename Var3Type                   > struct TRawMethodDelegate_ThreeVars_Const : TBaseRawMethodDelegateInstance<true, UserClass, FuncType, UserPolicy, Var1Type, Var2Type, Var3Type          > { typedef TBaseRawMethodDelegateInstance<true, UserClass, FuncType, UserPolicy, Var1Type, Var2Type, Var3Type          > Super; TRawMethodDelegate_ThreeVars_Const(UserClass* InUserObject, typename Super::FMethodPtr InMethodPtr, Var1Type Var1, Var2Type Var2, Var3Type Var3) : Super(InUserObject, InMethodPtr, Var1, Var2, Var3) {} };
+	template <typename UserClass, typename Var1Type, typename Var2Type, typename Var3Type, typename Var4Type> struct TRawMethodDelegate_FourVars : TBaseRawMethodDelegateInstance<false, UserClass, FuncType, UserPolicy, Var1Type, Var2Type, Var3Type, Var4Type> { typedef TBaseRawMethodDelegateInstance<false, UserClass, FuncType, UserPolicy, Var1Type, Var2Type, Var3Type, Var4Type> Super; TRawMethodDelegate_FourVars(UserClass* InUserObject, typename Super::FMethodPtr InMethodPtr, Var1Type Var1, Var2Type Var2, Var3Type Var3, Var4Type Var4) : Super(InUserObject, InMethodPtr, Var1, Var2, Var3, Var4) {} };
+	template <typename UserClass, typename Var1Type, typename Var2Type, typename Var3Type, typename Var4Type> struct TRawMethodDelegate_FourVars_Const : TBaseRawMethodDelegateInstance<true, UserClass, FuncType, UserPolicy, Var1Type, Var2Type, Var3Type, Var4Type> { typedef TBaseRawMethodDelegateInstance<true, UserClass, FuncType, UserPolicy, Var1Type, Var2Type, Var3Type, Var4Type> Super; TRawMethodDelegate_FourVars_Const(UserClass* InUserObject, typename Super::FMethodPtr InMethodPtr, Var1Type Var1, Var2Type Var2, Var3Type Var3, Var4Type Var4) : Super(InUserObject, InMethodPtr, Var1, Var2, Var3, Var4) {} };
+
 	/** Declare the user's UFunction-based delegate instance types. */
-	template <typename UObjectTemplate                                                                            > struct UE_DEPRECATED(5.2, "DelegateType::TUFunctionDelegateBinding has been deprecated")           TUFunctionDelegateBinding           : TBaseUFunctionDelegateInstance<UObjectTemplate, FuncType, UserPolicy                                        > { typedef TBaseUFunctionDelegateInstance<UObjectTemplate, FuncType, UserPolicy                                        > Super; TUFunctionDelegateBinding          (UObjectTemplate* InUserObject, const FName& InFunctionName                                                            ) : Super(InUserObject, InFunctionName                        ) {} };
-	template <typename UObjectTemplate, typename Var1Type                                                         > struct UE_DEPRECATED(5.2, "DelegateType::TUFunctionDelegateBinding_OneVar has been deprecated")    TUFunctionDelegateBinding_OneVar    : TBaseUFunctionDelegateInstance<UObjectTemplate, FuncType, UserPolicy, Var1Type                              > { typedef TBaseUFunctionDelegateInstance<UObjectTemplate, FuncType, UserPolicy, Var1Type                              > Super; TUFunctionDelegateBinding_OneVar   (UObjectTemplate* InUserObject, const FName& InFunctionName, Var1Type Var1                                             ) : Super(InUserObject, InFunctionName, Var1                  ) {} };
-	template <typename UObjectTemplate, typename Var1Type, typename Var2Type                                      > struct UE_DEPRECATED(5.2, "DelegateType::TUFunctionDelegateBinding_TwoVars has been deprecated")   TUFunctionDelegateBinding_TwoVars   : TBaseUFunctionDelegateInstance<UObjectTemplate, FuncType, UserPolicy, Var1Type, Var2Type                    > { typedef TBaseUFunctionDelegateInstance<UObjectTemplate, FuncType, UserPolicy, Var1Type, Var2Type                    > Super; TUFunctionDelegateBinding_TwoVars  (UObjectTemplate* InUserObject, const FName& InFunctionName, Var1Type Var1, Var2Type Var2                              ) : Super(InUserObject, InFunctionName, Var1, Var2            ) {} };
-	template <typename UObjectTemplate, typename Var1Type, typename Var2Type, typename Var3Type                   > struct UE_DEPRECATED(5.2, "DelegateType::TUFunctionDelegateBinding_ThreeVars has been deprecated") TUFunctionDelegateBinding_ThreeVars : TBaseUFunctionDelegateInstance<UObjectTemplate, FuncType, UserPolicy, Var1Type, Var2Type, Var3Type          > { typedef TBaseUFunctionDelegateInstance<UObjectTemplate, FuncType, UserPolicy, Var1Type, Var2Type, Var3Type          > Super; TUFunctionDelegateBinding_ThreeVars(UObjectTemplate* InUserObject, const FName& InFunctionName, Var1Type Var1, Var2Type Var2, Var3Type Var3               ) : Super(InUserObject, InFunctionName, Var1, Var2, Var3      ) {} };
-	template <typename UObjectTemplate, typename Var1Type, typename Var2Type, typename Var3Type, typename Var4Type> struct UE_DEPRECATED(5.2, "DelegateType::TUFunctionDelegateBinding_FourVars has been deprecated")  TUFunctionDelegateBinding_FourVars  : TBaseUFunctionDelegateInstance<UObjectTemplate, FuncType, UserPolicy, Var1Type, Var2Type, Var3Type, Var4Type> { typedef TBaseUFunctionDelegateInstance<UObjectTemplate, FuncType, UserPolicy, Var1Type, Var2Type, Var3Type, Var4Type> Super; TUFunctionDelegateBinding_FourVars (UObjectTemplate* InUserObject, const FName& InFunctionName, Var1Type Var1, Var2Type Var2, Var3Type Var3, Var4Type Var4) : Super(InUserObject, InFunctionName, Var1, Var2, Var3, Var4) {} };
+	template <typename UObjectTemplate                                                                            > struct TUFunctionDelegateBinding : TBaseUFunctionDelegateInstance<UObjectTemplate, FuncType, UserPolicy                                        > { typedef TBaseUFunctionDelegateInstance<UObjectTemplate, FuncType, UserPolicy                                        > Super; TUFunctionDelegateBinding(UObjectTemplate* InUserObject, const FName& InFunctionName) : Super(InUserObject, InFunctionName) {} };
+	template <typename UObjectTemplate, typename Var1Type                                                         > struct TUFunctionDelegateBinding_OneVar : TBaseUFunctionDelegateInstance<UObjectTemplate, FuncType, UserPolicy, Var1Type                              > { typedef TBaseUFunctionDelegateInstance<UObjectTemplate, FuncType, UserPolicy, Var1Type                              > Super; TUFunctionDelegateBinding_OneVar(UObjectTemplate* InUserObject, const FName& InFunctionName, Var1Type Var1) : Super(InUserObject, InFunctionName, Var1) {} };
+	template <typename UObjectTemplate, typename Var1Type, typename Var2Type                                      > struct TUFunctionDelegateBinding_TwoVars : TBaseUFunctionDelegateInstance<UObjectTemplate, FuncType, UserPolicy, Var1Type, Var2Type                    > { typedef TBaseUFunctionDelegateInstance<UObjectTemplate, FuncType, UserPolicy, Var1Type, Var2Type                    > Super; TUFunctionDelegateBinding_TwoVars(UObjectTemplate* InUserObject, const FName& InFunctionName, Var1Type Var1, Var2Type Var2) : Super(InUserObject, InFunctionName, Var1, Var2) {} };
+	template <typename UObjectTemplate, typename Var1Type, typename Var2Type, typename Var3Type                   > struct TUFunctionDelegateBinding_ThreeVars : TBaseUFunctionDelegateInstance<UObjectTemplate, FuncType, UserPolicy, Var1Type, Var2Type, Var3Type          > { typedef TBaseUFunctionDelegateInstance<UObjectTemplate, FuncType, UserPolicy, Var1Type, Var2Type, Var3Type          > Super; TUFunctionDelegateBinding_ThreeVars(UObjectTemplate* InUserObject, const FName& InFunctionName, Var1Type Var1, Var2Type Var2, Var3Type Var3) : Super(InUserObject, InFunctionName, Var1, Var2, Var3) {} };
+	template <typename UObjectTemplate, typename Var1Type, typename Var2Type, typename Var3Type, typename Var4Type> struct TUFunctionDelegateBinding_FourVars : TBaseUFunctionDelegateInstance<UObjectTemplate, FuncType, UserPolicy, Var1Type, Var2Type, Var3Type, Var4Type> { typedef TBaseUFunctionDelegateInstance<UObjectTemplate, FuncType, UserPolicy, Var1Type, Var2Type, Var3Type, Var4Type> Super; TUFunctionDelegateBinding_FourVars(UObjectTemplate* InUserObject, const FName& InFunctionName, Var1Type Var1, Var2Type Var2, Var3Type Var3, Var4Type Var4) : Super(InUserObject, InFunctionName, Var1, Var2, Var3, Var4) {} };
 
 	/** Declare the user's UObject-based delegate instance types. */
-	template <typename UserClass                                                                            > struct UE_DEPRECATED(5.2, "DelegateType::TUObjectMethodDelegate<UserClass, ...>::FMethodPtr has been deprecated - use DelegateType::TMethodPtr<UserClass, ...> instead")                      TUObjectMethodDelegate                 : TBaseUObjectMethodDelegateInstance<false, UserClass, FuncType, UserPolicy                                        > { typedef TBaseUObjectMethodDelegateInstance<false, UserClass, FuncType, UserPolicy                                        > Super; TUObjectMethodDelegate                (UserClass* InUserObject, typename Super::FMethodPtr InMethodPtr                                                            ) : Super(InUserObject, InMethodPtr                        ) {} };
-	template <typename UserClass                                                                            > struct UE_DEPRECATED(5.2, "DelegateType::TUObjectMethodDelegate_Const<UserClass, ...>::FMethodPtr has been deprecated - use DelegateType::TConstMethodPtr<UserClass, ...> instead")           TUObjectMethodDelegate_Const           : TBaseUObjectMethodDelegateInstance<true , UserClass, FuncType, UserPolicy                                        > { typedef TBaseUObjectMethodDelegateInstance<true , UserClass, FuncType, UserPolicy                                        > Super; TUObjectMethodDelegate_Const          (UserClass* InUserObject, typename Super::FMethodPtr InMethodPtr                                                            ) : Super(InUserObject, InMethodPtr                        ) {} };
-	template <typename UserClass, typename Var1Type                                                         > struct UE_DEPRECATED(5.2, "DelegateType::TUObjectMethodDelegate_OneVar<UserClass, ...>::FMethodPtr has been deprecated - use DelegateType::TMethodPtr<UserClass, ...> instead")               TUObjectMethodDelegate_OneVar          : TBaseUObjectMethodDelegateInstance<false, UserClass, FuncType, UserPolicy, Var1Type                              > { typedef TBaseUObjectMethodDelegateInstance<false, UserClass, FuncType, UserPolicy, Var1Type                              > Super; TUObjectMethodDelegate_OneVar         (UserClass* InUserObject, typename Super::FMethodPtr InMethodPtr, Var1Type Var1                                             ) : Super(InUserObject, InMethodPtr, Var1                  ) {} };
-	template <typename UserClass, typename Var1Type                                                         > struct UE_DEPRECATED(5.2, "DelegateType::TUObjectMethodDelegate_OneVar_Const<UserClass, ...>::FMethodPtr has been deprecated - use DelegateType::TConstMethodPtr<UserClass, ...> instead")    TUObjectMethodDelegate_OneVar_Const    : TBaseUObjectMethodDelegateInstance<true , UserClass, FuncType, UserPolicy, Var1Type                              > { typedef TBaseUObjectMethodDelegateInstance<true , UserClass, FuncType, UserPolicy, Var1Type                              > Super; TUObjectMethodDelegate_OneVar_Const   (UserClass* InUserObject, typename Super::FMethodPtr InMethodPtr, Var1Type Var1                                             ) : Super(InUserObject, InMethodPtr, Var1                  ) {} };
-	template <typename UserClass, typename Var1Type, typename Var2Type                                      > struct UE_DEPRECATED(5.2, "DelegateType::TUObjectMethodDelegate_TwoVars<UserClass, ...>::FMethodPtr has been deprecated - use DelegateType::TMethodPtr<UserClass, ...> instead")              TUObjectMethodDelegate_TwoVars         : TBaseUObjectMethodDelegateInstance<false, UserClass, FuncType, UserPolicy, Var1Type, Var2Type                    > { typedef TBaseUObjectMethodDelegateInstance<false, UserClass, FuncType, UserPolicy, Var1Type, Var2Type                    > Super; TUObjectMethodDelegate_TwoVars        (UserClass* InUserObject, typename Super::FMethodPtr InMethodPtr, Var1Type Var1, Var2Type Var2                              ) : Super(InUserObject, InMethodPtr, Var1, Var2            ) {} };
-	template <typename UserClass, typename Var1Type, typename Var2Type                                      > struct UE_DEPRECATED(5.2, "DelegateType::TUObjectMethodDelegate_TwoVars_Const<UserClass, ...>::FMethodPtr has been deprecated - use DelegateType::TConstMethodPtr<UserClass, ...> instead")   TUObjectMethodDelegate_TwoVars_Const   : TBaseUObjectMethodDelegateInstance<true , UserClass, FuncType, UserPolicy, Var1Type, Var2Type                    > { typedef TBaseUObjectMethodDelegateInstance<true , UserClass, FuncType, UserPolicy, Var1Type, Var2Type                    > Super; TUObjectMethodDelegate_TwoVars_Const  (UserClass* InUserObject, typename Super::FMethodPtr InMethodPtr, Var1Type Var1, Var2Type Var2                              ) : Super(InUserObject, InMethodPtr, Var1, Var2            ) {} };
-	template <typename UserClass, typename Var1Type, typename Var2Type, typename Var3Type                   > struct UE_DEPRECATED(5.2, "DelegateType::TUObjectMethodDelegate_ThreeVars<UserClass, ...>::FMethodPtr has been deprecated - use DelegateType::TMethodPtr<UserClass, ...> instead")            TUObjectMethodDelegate_ThreeVars       : TBaseUObjectMethodDelegateInstance<false, UserClass, FuncType, UserPolicy, Var1Type, Var2Type, Var3Type          > { typedef TBaseUObjectMethodDelegateInstance<false, UserClass, FuncType, UserPolicy, Var1Type, Var2Type, Var3Type          > Super; TUObjectMethodDelegate_ThreeVars      (UserClass* InUserObject, typename Super::FMethodPtr InMethodPtr, Var1Type Var1, Var2Type Var2, Var3Type Var3               ) : Super(InUserObject, InMethodPtr, Var1, Var2, Var3      ) {} };
-	template <typename UserClass, typename Var1Type, typename Var2Type, typename Var3Type                   > struct UE_DEPRECATED(5.2, "DelegateType::TUObjectMethodDelegate_ThreeVars_Const<UserClass, ...>::FMethodPtr has been deprecated - use DelegateType::TConstMethodPtr<UserClass, ...> instead") TUObjectMethodDelegate_ThreeVars_Const : TBaseUObjectMethodDelegateInstance<true , UserClass, FuncType, UserPolicy, Var1Type, Var2Type, Var3Type          > { typedef TBaseUObjectMethodDelegateInstance<true , UserClass, FuncType, UserPolicy, Var1Type, Var2Type, Var3Type          > Super; TUObjectMethodDelegate_ThreeVars_Const(UserClass* InUserObject, typename Super::FMethodPtr InMethodPtr, Var1Type Var1, Var2Type Var2, Var3Type Var3               ) : Super(InUserObject, InMethodPtr, Var1, Var2, Var3      ) {} };
-	template <typename UserClass, typename Var1Type, typename Var2Type, typename Var3Type, typename Var4Type> struct UE_DEPRECATED(5.2, "DelegateType::TUObjectMethodDelegate_FourVars<UserClass, ...>::FMethodPtr has been deprecated - use DelegateType::TMethodPtr<UserClass, ...> instead")             TUObjectMethodDelegate_FourVars        : TBaseUObjectMethodDelegateInstance<false, UserClass, FuncType, UserPolicy, Var1Type, Var2Type, Var3Type, Var4Type> { typedef TBaseUObjectMethodDelegateInstance<false, UserClass, FuncType, UserPolicy, Var1Type, Var2Type, Var3Type, Var4Type> Super; TUObjectMethodDelegate_FourVars       (UserClass* InUserObject, typename Super::FMethodPtr InMethodPtr, Var1Type Var1, Var2Type Var2, Var3Type Var3, Var4Type Var4) : Super(InUserObject, InMethodPtr, Var1, Var2, Var3, Var4) {} };
-	template <typename UserClass, typename Var1Type, typename Var2Type, typename Var3Type, typename Var4Type> struct UE_DEPRECATED(5.2, "DelegateType::TUObjectMethodDelegate_FourVars_Const<UserClass, ...>::FMethodPtr has been deprecated - use DelegateType::TConstMethodPtr<UserClass, ...> instead")  TUObjectMethodDelegate_FourVars_Const  : TBaseUObjectMethodDelegateInstance<true , UserClass, FuncType, UserPolicy, Var1Type, Var2Type, Var3Type, Var4Type> { typedef TBaseUObjectMethodDelegateInstance<true , UserClass, FuncType, UserPolicy, Var1Type, Var2Type, Var3Type, Var4Type> Super; TUObjectMethodDelegate_FourVars_Const (UserClass* InUserObject, typename Super::FMethodPtr InMethodPtr, Var1Type Var1, Var2Type Var2, Var3Type Var3, Var4Type Var4) : Super(InUserObject, InMethodPtr, Var1, Var2, Var3, Var4) {} };
-	
+	template <typename UserClass                                                                            > struct TUObjectMethodDelegate : TBaseUObjectMethodDelegateInstance<false, UserClass, FuncType, UserPolicy                                        > { typedef TBaseUObjectMethodDelegateInstance<false, UserClass, FuncType, UserPolicy                                        > Super; TUObjectMethodDelegate(UserClass* InUserObject, typename Super::FMethodPtr InMethodPtr) : Super(InUserObject, InMethodPtr) {} };
+	template <typename UserClass                                                                            > struct TUObjectMethodDelegate_Const : TBaseUObjectMethodDelegateInstance<true, UserClass, FuncType, UserPolicy                                        > { typedef TBaseUObjectMethodDelegateInstance<true, UserClass, FuncType, UserPolicy                                        > Super; TUObjectMethodDelegate_Const(UserClass* InUserObject, typename Super::FMethodPtr InMethodPtr) : Super(InUserObject, InMethodPtr) {} };
+	template <typename UserClass, typename Var1Type                                                         > struct TUObjectMethodDelegate_OneVar : TBaseUObjectMethodDelegateInstance<false, UserClass, FuncType, UserPolicy, Var1Type                              > { typedef TBaseUObjectMethodDelegateInstance<false, UserClass, FuncType, UserPolicy, Var1Type                              > Super; TUObjectMethodDelegate_OneVar(UserClass* InUserObject, typename Super::FMethodPtr InMethodPtr, Var1Type Var1) : Super(InUserObject, InMethodPtr, Var1) {} };
+	template <typename UserClass, typename Var1Type                                                         > struct TUObjectMethodDelegate_OneVar_Const : TBaseUObjectMethodDelegateInstance<true, UserClass, FuncType, UserPolicy, Var1Type                              > { typedef TBaseUObjectMethodDelegateInstance<true, UserClass, FuncType, UserPolicy, Var1Type                              > Super; TUObjectMethodDelegate_OneVar_Const(UserClass* InUserObject, typename Super::FMethodPtr InMethodPtr, Var1Type Var1) : Super(InUserObject, InMethodPtr, Var1) {} };
+	template <typename UserClass, typename Var1Type, typename Var2Type                                      > struct TUObjectMethodDelegate_TwoVars : TBaseUObjectMethodDelegateInstance<false, UserClass, FuncType, UserPolicy, Var1Type, Var2Type                    > { typedef TBaseUObjectMethodDelegateInstance<false, UserClass, FuncType, UserPolicy, Var1Type, Var2Type                    > Super; TUObjectMethodDelegate_TwoVars(UserClass* InUserObject, typename Super::FMethodPtr InMethodPtr, Var1Type Var1, Var2Type Var2) : Super(InUserObject, InMethodPtr, Var1, Var2) {} };
+	template <typename UserClass, typename Var1Type, typename Var2Type                                      > struct TUObjectMethodDelegate_TwoVars_Const : TBaseUObjectMethodDelegateInstance<true, UserClass, FuncType, UserPolicy, Var1Type, Var2Type                    > { typedef TBaseUObjectMethodDelegateInstance<true, UserClass, FuncType, UserPolicy, Var1Type, Var2Type                    > Super; TUObjectMethodDelegate_TwoVars_Const(UserClass* InUserObject, typename Super::FMethodPtr InMethodPtr, Var1Type Var1, Var2Type Var2) : Super(InUserObject, InMethodPtr, Var1, Var2) {} };
+	template <typename UserClass, typename Var1Type, typename Var2Type, typename Var3Type                   > struct TUObjectMethodDelegate_ThreeVars : TBaseUObjectMethodDelegateInstance<false, UserClass, FuncType, UserPolicy, Var1Type, Var2Type, Var3Type          > { typedef TBaseUObjectMethodDelegateInstance<false, UserClass, FuncType, UserPolicy, Var1Type, Var2Type, Var3Type          > Super; TUObjectMethodDelegate_ThreeVars(UserClass* InUserObject, typename Super::FMethodPtr InMethodPtr, Var1Type Var1, Var2Type Var2, Var3Type Var3) : Super(InUserObject, InMethodPtr, Var1, Var2, Var3) {} };
+	template <typename UserClass, typename Var1Type, typename Var2Type, typename Var3Type                   > struct TUObjectMethodDelegate_ThreeVars_Const : TBaseUObjectMethodDelegateInstance<true, UserClass, FuncType, UserPolicy, Var1Type, Var2Type, Var3Type          > { typedef TBaseUObjectMethodDelegateInstance<true, UserClass, FuncType, UserPolicy, Var1Type, Var2Type, Var3Type          > Super; TUObjectMethodDelegate_ThreeVars_Const(UserClass* InUserObject, typename Super::FMethodPtr InMethodPtr, Var1Type Var1, Var2Type Var2, Var3Type Var3) : Super(InUserObject, InMethodPtr, Var1, Var2, Var3) {} };
+	template <typename UserClass, typename Var1Type, typename Var2Type, typename Var3Type, typename Var4Type> struct TUObjectMethodDelegate_FourVars : TBaseUObjectMethodDelegateInstance<false, UserClass, FuncType, UserPolicy, Var1Type, Var2Type, Var3Type, Var4Type> { typedef TBaseUObjectMethodDelegateInstance<false, UserClass, FuncType, UserPolicy, Var1Type, Var2Type, Var3Type, Var4Type> Super; TUObjectMethodDelegate_FourVars(UserClass* InUserObject, typename Super::FMethodPtr InMethodPtr, Var1Type Var1, Var2Type Var2, Var3Type Var3, Var4Type Var4) : Super(InUserObject, InMethodPtr, Var1, Var2, Var3, Var4) {} };
+	template <typename UserClass, typename Var1Type, typename Var2Type, typename Var3Type, typename Var4Type> struct TUObjectMethodDelegate_FourVars_Const : TBaseUObjectMethodDelegateInstance<true, UserClass, FuncType, UserPolicy, Var1Type, Var2Type, Var3Type, Var4Type> { typedef TBaseUObjectMethodDelegateInstance<true, UserClass, FuncType, UserPolicy, Var1Type, Var2Type, Var3Type, Var4Type> Super; TUObjectMethodDelegate_FourVars_Const(UserClass* InUserObject, typename Super::FMethodPtr InMethodPtr, Var1Type Var1, Var2Type Var2, Var3Type Var3, Var4Type Var4) : Super(InUserObject, InMethodPtr, Var1, Var2, Var3, Var4) {} };
+
 	/** Declare the user's static function pointer delegate instance types. */
-	template <typename... VarTypes > struct UE_DEPRECATED(5.2, "DelegateType::FStaticDelegate<...>::FFuncPtr has been deprecated - use DelegateType::TFuncPtr<...> instead") FStaticDelegate : TBaseStaticDelegateInstance<FuncType, UserPolicy, VarTypes... > { typedef TBaseStaticDelegateInstance<FuncType, UserPolicy, VarTypes... > Super; FStaticDelegate ( typename Super::FFuncPtr InFuncPtr, VarTypes... Vars ) : Super(InFuncPtr, Vars... ) {} };
+	template <typename... VarTypes > struct FStaticDelegate : TBaseStaticDelegateInstance<FuncType, UserPolicy, VarTypes... > { typedef TBaseStaticDelegateInstance<FuncType, UserPolicy, VarTypes... > Super; FStaticDelegate(typename Super::FFuncPtr InFuncPtr, VarTypes... Vars) : Super(InFuncPtr, Vars...) {} };
 
 public:
 
@@ -381,7 +377,10 @@ public:
 	 *
 	 * @param Other The delegate object to move from.
 	 */
-	TDelegate(TDelegate&& Other) = default;
+	inline TDelegate(TDelegate&& Other)
+	{
+		*this = MoveTemp(Other);
+	}
 
 	/**
 	 * Creates and initializes a new instance from an existing delegate object.
@@ -398,7 +397,25 @@ public:
 	 *
 	 * @param	OtherDelegate	Delegate object to copy from
 	 */
-	TDelegate& operator=(TDelegate&& Other) = default;
+	inline TDelegate& operator=(TDelegate&& Other)
+	{
+		if (&Other != this)
+		{
+			// this down-cast is OK! allows for managing invocation list in the base class without requiring virtual functions
+			DelegateInstanceInterfaceType* OtherInstance = Other.GetDelegateInstanceProtected();
+
+			if (OtherInstance != nullptr)
+			{
+				OtherInstance->CreateCopy(*this);
+			}
+			else
+			{
+				Unbind();
+			}
+		}
+
+		return *this;
+	}
 
 	/**
 	 * Assignment operator.
@@ -1029,7 +1046,7 @@ public:
 
 		// NOTE: If you hit a compile error on the following line, it means you're trying to use a non-UObject type
 		//       with this delegate, which is not supported
-		this->Object = Cast<UObject>(InUserObject);
+		this->Object = InUserObject;
 
 		// Store the function name.  The incoming function name was generated by a macro and includes the method's class name.
 		this->FunctionName = InFunctionName;
