Mac extra built commandlets (bBuileAdditionalConsoleApp=true) are placed 3 directories up relative to the non-commandlet packaged app.
#rnx

diff --git a/Engine/Source/Programs/UnrealBuildTool/Configuration/UEBuildTarget.cs b/Engine/Source/Programs/UnrealBuildTool/Configuration/UEBuildTarget.cs
--- a/Engine/Source/Programs/UnrealBuildTool/Configuration/UEBuildTarget.cs
+++ b/Engine/Source/Programs/UnrealBuildTool/Configuration/UEBuildTarget.cs
@@ -3967,6 +3967,11 @@ namespace UnrealBuildTool
 			if (RelativeBaseDir != null)
 			{
 				GlobalCompileEnvironment.Definitions.Add(String.Format("UE_RELATIVE_BASE_DIR=\"{0}/\"", RelativeBaseDir));
+				if (Rules.bBuildAdditionalConsoleApp)
+				{
+					string? CmdletRelativeBaseDir = GetRelativeBaseDir(GetExecutableDir(), Platform);
+					GlobalCompileEnvironment.Definitions.Add(String.Format("UE_CMDLET_RELATIVE_BASE_DIR=\"{0}/\"", CmdletRelativeBaseDir));
+				}
 			}
 
 			bool bCompileDevTests = (Configuration != UnrealTargetConfiguration.Test && Configuration != UnrealTargetConfiguration.Shipping);
diff --git a/Engine/Source/Runtime/Core/Private/Mac/MacPlatformProcess.cpp b/Engine/Source/Runtime/Core/Private/Mac/MacPlatformProcess.cpp
--- a/Engine/Source/Runtime/Core/Private/Mac/MacPlatformProcess.cpp
+++ b/Engine/Source/Runtime/Core/Private/Mac/MacPlatformProcess.cpp
@@ -1029,7 +1029,20 @@ const TCHAR* FMacPlatformProcess::BaseDir()
 #ifdef UE_RELATIVE_BASE_DIR
 			else
 			{
-				BasePath = [BasePath stringByAppendingPathComponent : @UE_RELATIVE_BASE_DIR];
+				// Get executable path
+				NSString* ExecutablePath = [[NSBundle mainBundle] executablePath];
+
+				// When building an extra console app i.e. a commandlet, the Mac executable can  be placed in a distinct folder, usually 3 directories up compared to the packaged app
+				if ([[ExecutablePath lowercaseString] hasSuffix : @"-cmd"] )
+				{
+#ifdef UE_CMDLET_RELATIVE_BASE_DIR
+					BasePath = [BasePath stringByAppendingPathComponent : @UE_CMDLET_RELATIVE_BASE_DIR];
+#endif
+				}
+				else
+				{
+					BasePath = [BasePath stringByAppendingPathComponent : @UE_RELATIVE_BASE_DIR];
+				}
 			}
 #endif
 
