added multithreaded access detector to delegates, to detect concurrent access. the intention is to enable it for a horde nightly test run once we fix all known cases

diff --git a/Engine/Source/Runtime/Core/Public/Delegates/Delegate.h b/Engine/Source/Runtime/Core/Public/Delegates/Delegate.h
--- a/Engine/Source/Runtime/Core/Public/Delegates/Delegate.h
+++ b/Engine/Source/Runtime/Core/Public/Delegates/Delegate.h
@@ -2,6 +2,23 @@
 
 #pragma once
 
+//#define UE_DETECT_DELEGATES_RACE_CONDITIONS 1
+
+#if !defined(UE_DETECT_DELEGATES_RACE_CONDITIONS)
+	#define UE_DETECT_DELEGATES_RACE_CONDITIONS 0
+#endif
+
+#if UE_DETECT_DELEGATES_RACE_CONDITIONS
+	#include "Misc/MTAccessDetector.h"
+	#define UE_DELEGATES_MT_ACCESS_DETECTOR UE_MT_DECLARE_RW_FULLY_RECURSIVE_ACCESS_DETECTOR
+	#define UE_DELEGATES_MT_SCOPED_WRITE_ACCESS UE_MT_SCOPED_WRITE_ACCESS
+	#define UE_DELEGATES_MT_SCOPED_READ_ACCESS UE_MT_SCOPED_READ_ACCESS
+#else
+	#define UE_DELEGATES_MT_ACCESS_DETECTOR(...)
+	#define UE_DELEGATES_MT_SCOPED_WRITE_ACCESS(...)
+	#define UE_DELEGATES_MT_SCOPED_READ_ACCESS(...)
+#endif
+
 #include "CoreTypes.h"
 #include "Misc/AssertionMacros.h"
 #include "UObject/NameTypes.h"
diff --git a/Engine/Source/Runtime/Core/Public/Delegates/DelegateBase.h b/Engine/Source/Runtime/Core/Public/Delegates/DelegateBase.h
--- a/Engine/Source/Runtime/Core/Public/Delegates/DelegateBase.h
+++ b/Engine/Source/Runtime/Core/Public/Delegates/DelegateBase.h
@@ -79,6 +79,8 @@ protected:
 	 */
 	FDelegateBase& operator=(FDelegateBase&& Other)
 	{
+		UE_DELEGATES_MT_SCOPED_WRITE_ACCESS(AccessDetector);
+
 		Unbind();
 		DelegateAllocator.MoveToEmpty(Other.DelegateAllocator);
 		DelegateSize = Other.DelegateSize;
@@ -91,6 +93,8 @@ protected:
 	 */
 	FORCEINLINE void Unbind()
 	{
+		UE_DELEGATES_MT_SCOPED_WRITE_ACCESS(AccessDetector);
+
 		if (IDelegateInstance* Ptr = GetDelegateInstanceProtected())
 		{
 			Ptr->~IDelegateInstance();
@@ -107,6 +111,8 @@ protected:
 	 */
 	FORCEINLINE IDelegateInstance* GetDelegateInstanceProtected() const
 	{
+		UE_DELEGATES_MT_SCOPED_READ_ACCESS(AccessDetector);
+
 		return DelegateSize ? (IDelegateInstance*)DelegateAllocator.GetAllocation() : nullptr;
 	}
 
@@ -115,6 +121,8 @@ private:
 
 	void* Allocate(int32 Size)
 	{
+		UE_DELEGATES_MT_SCOPED_WRITE_ACCESS(AccessDetector);
+
 		if (IDelegateInstance* CurrentInstance = GetDelegateInstanceProtected())
 		{
 			CurrentInstance->~IDelegateInstance();
@@ -133,6 +141,8 @@ private:
 private:
 	FDelegateAllocatorType::ForElementType<FAlignedInlineDelegateType> DelegateAllocator;
 	int32 DelegateSize;
+
+	UE_DELEGATES_MT_ACCESS_DETECTOR(AccessDetector);
 };
 
 inline void* operator new(size_t Size, FDelegateBase& Base)
diff --git a/Engine/Source/Runtime/Core/Public/Delegates/MulticastDelegateBase.h b/Engine/Source/Runtime/Core/Public/Delegates/MulticastDelegateBase.h
--- a/Engine/Source/Runtime/Core/Public/Delegates/MulticastDelegateBase.h
+++ b/Engine/Source/Runtime/Core/Public/Delegates/MulticastDelegateBase.h
@@ -30,6 +30,8 @@ public:
 	/** Removes all functions from this delegate's invocation list. */
 	void Clear( )
 	{
+		UE_DELEGATES_MT_SCOPED_WRITE_ACCESS(AccessDetector);
+
 		for (TDelegateBase<UserPolicy>& DelegateBaseRef : InvocationList)
 		{
 			DelegateBaseRef.Unbind();
@@ -45,6 +47,8 @@ public:
 	 */
 	inline bool IsBound( ) const
 	{
+		UE_DELEGATES_MT_SCOPED_READ_ACCESS(AccessDetector);
+	
 		for (const TDelegateBase<UserPolicy>& DelegateBaseRef : InvocationList)
 		{
 			if (DelegateBaseRef.GetDelegateInstanceProtected())
@@ -62,6 +66,8 @@ public:
 	 */
 	inline bool IsBoundToObject( void const* InUserObject ) const
 	{
+		UE_DELEGATES_MT_SCOPED_READ_ACCESS(AccessDetector);
+
 		for (const TDelegateBase<UserPolicy>& DelegateBaseRef : InvocationList)
 		{
 			IDelegateInstance* DelegateInstance = DelegateBaseRef.GetDelegateInstanceProtected();
@@ -83,6 +89,8 @@ public:
 	 */
 	int32 RemoveAll( const void* InUserObject )
 	{
+		UE_DELEGATES_MT_SCOPED_WRITE_ACCESS(AccessDetector);
+
 		int32 Result = 0;
 		if (InvocationListLockCount > 0)
 		{
@@ -144,6 +152,8 @@ protected:
 	template<typename DelegateInstanceInterfaceType, typename DelegateType>
 	void CopyFrom(const TMulticastDelegateBase& Other)
 	{
+		UE_DELEGATES_MT_SCOPED_WRITE_ACCESS(AccessDetector);
+
 		checkSlow(&Other != this);
 
 		Clear();
@@ -162,6 +172,8 @@ protected:
 	template<typename DelegateInstanceInterfaceType, typename DelegateBaseType, typename... ParamTypes>
 	void Broadcast(ParamTypes... Params) const
 	{
+		UE_DELEGATES_MT_SCOPED_WRITE_ACCESS(AccessDetector);
+
 		bool NeedsCompaction = false;
 
 		LockInvocationList();
@@ -196,6 +208,8 @@ protected:
 	 */
 	inline FDelegateHandle AddDelegateInstance(TDelegateBase<UserPolicy>&& NewDelegateBaseRef)
 	{
+		UE_DELEGATES_MT_SCOPED_WRITE_ACCESS(AccessDetector);
+
 		FDelegateHandle Result;
 		if (GetDelegateInstanceProtectedHelper(NewDelegateBaseRef))
 		{
@@ -215,6 +229,8 @@ protected:
 	 */
 	bool RemoveDelegateInstance(FDelegateHandle Handle)
 	{
+		UE_DELEGATES_MT_SCOPED_WRITE_ACCESS(AccessDetector);
+
 		for (int32 InvocationListIndex = 0; InvocationListIndex < InvocationList.Num(); ++InvocationListIndex)
 		{
 			TDelegateBase<UserPolicy>& DelegateBase = InvocationList[InvocationListIndex];
@@ -332,4 +348,6 @@ private:
 
 	/** Holds a lock counter for the invocation list. */
 	mutable int32 InvocationListLockCount;
+
+	UE_DELEGATES_MT_ACCESS_DETECTOR(AccessDetector);
 };
diff --git a/Engine/Source/Runtime/Core/Public/UObject/ScriptDelegates.h b/Engine/Source/Runtime/Core/Public/UObject/ScriptDelegates.h
--- a/Engine/Source/Runtime/Core/Public/UObject/ScriptDelegates.h
+++ b/Engine/Source/Runtime/Core/Public/UObject/ScriptDelegates.h
@@ -58,6 +58,8 @@ public:
 	 */
 	void BindUFunction( UObject* InObject, const FName& InFunctionName )
 	{
+		UE_DELEGATES_MT_SCOPED_WRITE_ACCESS(AccessDetector);
+
 		Object = InObject;
 		FunctionName = InFunctionName;
 	}
@@ -69,6 +71,8 @@ public:
 	 */
 	inline bool IsBound() const
 	{
+		UE_DELEGATES_MT_SCOPED_READ_ACCESS(AccessDetector);
+
 		return IsBound_Internal<UObject>();
 	}
 
@@ -79,6 +83,8 @@ public:
 	 */
 	inline bool IsBoundToObject(void const* InUserObject) const
 	{
+		UE_DELEGATES_MT_SCOPED_READ_ACCESS(AccessDetector);
+
 		return InUserObject && (InUserObject == GetUObject());
 	}
 
@@ -89,6 +95,8 @@ public:
 	 */
 	bool IsBoundToObjectEvenIfUnreachable(void const* InUserObject) const
 	{
+		UE_DELEGATES_MT_SCOPED_READ_ACCESS(AccessDetector);
+
 		return InUserObject && InUserObject == GetUObjectEvenIfUnreachable();
 	}
 
@@ -99,6 +107,8 @@ public:
 	 */
 	inline bool IsCompactable() const
 	{
+		UE_DELEGATES_MT_SCOPED_READ_ACCESS(AccessDetector);
+
 		return FunctionName == NAME_None || !Object.Get(true);
 	}
 
@@ -107,6 +117,8 @@ public:
 	 */
 	void Unbind()
 	{
+		UE_DELEGATES_MT_SCOPED_WRITE_ACCESS(AccessDetector);
+
 		Object = nullptr;
 		FunctionName = NAME_None;
 	}
@@ -129,6 +141,8 @@ public:
 	{
 		if( IsBound() )
 		{
+			UE_DELEGATES_MT_SCOPED_READ_ACCESS(AccessDetector);
+
 			return ((UObjectTemplate*)GetUObject())->GetPathName() + TEXT(".") + GetFunctionName().ToString();
 		}
 		return TEXT( "<Unbound>" );
@@ -137,6 +151,8 @@ public:
 	/** Delegate serialization */
 	friend FArchive& operator<<( FArchive& Ar, TScriptDelegate& D )
 	{
+		UE_DELEGATES_MT_SCOPED_READ_ACCESS(AccessDetector);
+
 		Ar << D.Object << D.FunctionName;
 		return Ar;
 	}
@@ -144,6 +160,8 @@ public:
 	/** Delegate serialization */
 	friend void operator<<(FStructuredArchive::FSlot Slot, TScriptDelegate& D)
 	{
+		UE_DELEGATES_MT_SCOPED_READ_ACCESS(AccessDetector);
+
 		FStructuredArchive::FRecord Record = Slot.EnterRecord();
 		Record << SA_VALUE(TEXT("Object"), D.Object) << SA_VALUE(TEXT("FunctionName"),D.FunctionName);
 	}
@@ -151,16 +169,22 @@ public:
 	/** Comparison operators */
 	FORCEINLINE bool operator==( const TScriptDelegate& Other ) const
 	{
+		UE_DELEGATES_MT_SCOPED_READ_ACCESS(AccessDetector);
+
 		return Object == Other.Object && FunctionName == Other.FunctionName;
 	}
 
 	FORCEINLINE bool operator!=( const TScriptDelegate& Other ) const
 	{
+		UE_DELEGATES_MT_SCOPED_READ_ACCESS(AccessDetector);
+
 		return Object != Other.Object || FunctionName != Other.FunctionName;
 	}
 
 	void operator=( const TScriptDelegate& Other )
 	{
+		UE_DELEGATES_MT_SCOPED_WRITE_ACCESS(AccessDetector);
+
 		Object = Other.Object;
 		FunctionName = Other.FunctionName;
 	}
@@ -172,6 +196,8 @@ public:
 	 */
 	UObject* GetUObject()
 	{
+		UE_DELEGATES_MT_SCOPED_WRITE_ACCESS(AccessDetector);
+
 		// Downcast UObjectBase to UObject
 		return static_cast< UObject* >( Object.Get() );
 	}
@@ -183,6 +209,8 @@ public:
 	 */
 	const UObject* GetUObject() const
 	{
+		UE_DELEGATES_MT_SCOPED_READ_ACCESS(AccessDetector);
+
 		// Downcast UObjectBase to UObject
 		return static_cast< const UObject* >( Object.Get() );
 	}
@@ -194,6 +222,8 @@ public:
 	 */
 	UObject* GetUObjectEvenIfUnreachable()
 	{
+		UE_DELEGATES_MT_SCOPED_WRITE_ACCESS(AccessDetector);
+
 		// Downcast UObjectBase to UObject
 		return static_cast< UObject* >( Object.GetEvenIfUnreachable() );
 	}
@@ -205,6 +235,8 @@ public:
 	 */
 	const UObject* GetUObjectEvenIfUnreachable() const
 	{
+		UE_DELEGATES_MT_SCOPED_READ_ACCESS(AccessDetector);
+
 		// Downcast UObjectBase to UObject
 		return static_cast< const UObject* >( Object.GetEvenIfUnreachable() );
 	}
@@ -216,6 +248,8 @@ public:
 	 */
 	FName GetFunctionName() const
 	{
+		UE_DELEGATES_MT_SCOPED_READ_ACCESS(AccessDetector);
+
 		return FunctionName;
 	}
 
@@ -230,6 +264,8 @@ public:
 	template <class UObjectTemplate>
 	void ProcessDelegate( void* Parameters ) const
 	{
+		UE_DELEGATES_MT_SCOPED_WRITE_ACCESS(AccessDetector);
+
 		checkf( Object.IsValid() != false, TEXT( "ProcessDelegate() called with no object bound to delegate!" ) );
 		checkf( FunctionName != NAME_None, TEXT( "ProcessDelegate() called with no function name set!" ) );
 
@@ -249,6 +285,8 @@ public:
 
 	friend uint32 GetTypeHash(const TScriptDelegate& Delegate)
 	{
+		UE_DELEGATES_MT_SCOPED_READ_ACCESS(AccessDetector);
+
 		return HashCombine(GetTypeHash(Delegate.Object), GetTypeHash(Delegate.GetFunctionName()));
 	}
 
@@ -260,6 +298,8 @@ protected:
 	/** Name of the function to call on the bound object */
 	FName FunctionName;
 
+	UE_DELEGATES_MT_ACCESS_DETECTOR(AccessDetector);
+
 	// 
 	friend class FCallDelegateHelper;
 };
@@ -290,6 +330,8 @@ public:
 	 */
 	inline bool IsBound() const
 	{
+		UE_DELEGATES_MT_SCOPED_READ_ACCESS(AccessDetector);
+
 		return InvocationList.Num() > 0;
 	}
 
@@ -301,6 +343,8 @@ public:
 	 */
 	bool Contains( const TScriptDelegate<TWeakPtr>& InDelegate ) const
 	{
+		UE_DELEGATES_MT_SCOPED_READ_ACCESS(AccessDetector);
+
 		return InvocationList.Contains( InDelegate );
 	}
 
@@ -313,6 +357,8 @@ public:
 	 */
 	bool Contains( const UObject* InObject, FName InFunctionName ) const
 	{
+		UE_DELEGATES_MT_SCOPED_READ_ACCESS(AccessDetector);
+
 		return InvocationList.ContainsByPredicate( [=]( const TScriptDelegate<TWeakPtr>& Delegate ){
 			return Delegate.GetFunctionName() == InFunctionName && Delegate.IsBoundToObjectEvenIfUnreachable(InObject);
 		} );
@@ -325,6 +371,8 @@ public:
 	 */
 	void Add( const TScriptDelegate<TWeakPtr>& InDelegate )
 	{
+		UE_DELEGATES_MT_SCOPED_WRITE_ACCESS(AccessDetector);
+
 		// First check for any objects that may have expired
 		CompactInvocationList();
 
@@ -340,6 +388,8 @@ public:
 	 */
 	void AddUnique( const TScriptDelegate<TWeakPtr>& InDelegate )
 	{
+		UE_DELEGATES_MT_SCOPED_WRITE_ACCESS(AccessDetector);
+
 		// Add the delegate, if possible
 		AddUniqueInternal( InDelegate );
 
@@ -355,6 +405,8 @@ public:
 	 */
 	void Remove( const TScriptDelegate<TWeakPtr>& InDelegate )
 	{
+		UE_DELEGATES_MT_SCOPED_WRITE_ACCESS(AccessDetector);
+
 		// Remove the delegate
 		RemoveInternal( InDelegate );
 
@@ -371,6 +423,8 @@ public:
 	 */
 	void Remove( const UObject* InObject, FName InFunctionName )
 	{
+		UE_DELEGATES_MT_SCOPED_WRITE_ACCESS(AccessDetector);
+
 		// Remove the delegate
 		RemoveInternal( InObject, InFunctionName );
 
@@ -388,6 +442,8 @@ public:
 	 */
 	void RemoveAll(const UObject* Object)
 	{
+		UE_DELEGATES_MT_SCOPED_WRITE_ACCESS(AccessDetector);
+
 		for (int32 BindingIndex = InvocationList.Num() - 1; BindingIndex >= 0; --BindingIndex)
 		{
 			const TScriptDelegate<TWeakPtr>& Binding = InvocationList[BindingIndex];
@@ -404,6 +460,8 @@ public:
 	 */
 	void Clear()
 	{
+		UE_DELEGATES_MT_SCOPED_WRITE_ACCESS(AccessDetector);
+
 		InvocationList.Empty();
 	}
 
@@ -415,6 +473,8 @@ public:
 	template <typename UObjectTemplate>
 	inline FString ToString() const
 	{
+		UE_DELEGATES_MT_SCOPED_READ_ACCESS(AccessDetector);
+
 		if( IsBound() )
 		{
 			FString AllDelegatesString = TEXT( "[" );
@@ -437,6 +497,8 @@ public:
 	/** Multi-cast delegate serialization */
 	friend FArchive& operator<<( FArchive& Ar, TMulticastScriptDelegate<TWeakPtr>& D )
 	{
+		UE_DELEGATES_MT_SCOPED_WRITE_ACCESS(AccessDetector);
+
 		if( Ar.IsSaving() )
 		{
 			// When saving the delegate, clean up the list to make sure there are no bad object references
@@ -456,6 +518,8 @@ public:
 
 	friend void operator<<(FStructuredArchive::FSlot Slot, TMulticastScriptDelegate<TWeakPtr>& D)
 	{
+		UE_DELEGATES_MT_SCOPED_WRITE_ACCESS(AccessDetector);
+
 		FArchive& UnderlyingArchive = Slot.GetUnderlyingArchive();
 
 		if (UnderlyingArchive.IsSaving())
@@ -483,6 +547,8 @@ public:
 	template <class UObjectTemplate>
 	void ProcessMulticastDelegate(void* Parameters) const
 	{
+		UE_DELEGATES_MT_SCOPED_WRITE_ACCESS(AccessDetector);
+
 		if( InvocationList.Num() > 0 )
 		{
 			// Create a copy of the invocation list, just in case the list is modified by one of the callbacks during the broadcast
@@ -513,6 +579,8 @@ public:
 	 */
 	TArray< UObject* > GetAllObjects() const
 	{
+		UE_DELEGATES_MT_SCOPED_READ_ACCESS(AccessDetector);
+
 		TArray< UObject* > OutputList;
 		for( typename FInvocationList::TIterator CurDelegate( InvocationList ); CurDelegate; ++CurDelegate )
 		{
@@ -532,6 +600,8 @@ public:
 	 */
 	TArray< UObject* > GetAllObjectsEvenIfUnreachable() const
 	{
+		UE_DELEGATES_MT_SCOPED_READ_ACCESS(AccessDetector);
+
 		TArray< UObject* > OutputList;
 		for( typename FInvocationList::TIterator CurDelegate( InvocationList ); CurDelegate; ++CurDelegate )
 		{
@@ -622,6 +692,8 @@ protected:
 	/** Ordered list functions to invoke when the Broadcast function is called */
 	mutable FInvocationList InvocationList;		// Mutable so that we can housekeep list even with 'const' broadcasts
 
+	UE_DELEGATES_MT_ACCESS_DETECTOR(AccessDetector);
+
 	// Declare ourselves as a friend of FMulticastDelegateProperty so that it can access our function list
 	friend class FMulticastDelegateProperty;
 	friend class FMulticastInlineDelegateProperty;
