FCoreDelegates: deprecations cleanup #1

diff --git a/Engine/Source/Runtime/Core/Private/Misc/CoreDelegates.cpp b/Engine/Source/Runtime/Core/Private/Misc/CoreDelegates.cpp
--- a/Engine/Source/Runtime/Core/Private/Misc/CoreDelegates.cpp
+++ b/Engine/Source/Runtime/Core/Private/Misc/CoreDelegates.cpp
@@ -58,14 +58,6 @@ FCoreDelegates::FOnActorLabelChanged FCoreDelegates::OnActorLabelChanged;
 FCoreDelegates::FRegisterMovieStreamerDelegate FCoreDelegates::RegisterMovieStreamerDelegate;
 FCoreDelegates::FUnRegisterMovieStreamerDelegate FCoreDelegates::UnRegisterMovieStreamerDelegate;
 
-PRAGMA_DISABLE_DEPRECATION_WARNINGS
-FCoreDelegates::FRegisterEncryptionKeyDelegate& FCoreDelegates::GetRegisterEncryptionKeyDelegate()
-{
-	static FRegisterEncryptionKeyDelegate RegisterEncryptionKeyDelegate;
-	return RegisterEncryptionKeyDelegate;
-}
-PRAGMA_ENABLE_DEPRECATION_WARNINGS
-
 FCoreDelegates::FRegisterEncryptionKeyMulticastDelegate& FCoreDelegates::GetRegisterEncryptionKeyMulticastDelegate()
 {
 	static FRegisterEncryptionKeyMulticastDelegate RegisterEncryptionKeyDelegate;
diff --git a/Engine/Source/Runtime/Core/Public/Misc/CoreDelegates.h b/Engine/Source/Runtime/Core/Public/Misc/CoreDelegates.h
--- a/Engine/Source/Runtime/Core/Public/Misc/CoreDelegates.h
+++ b/Engine/Source/Runtime/Core/Public/Misc/CoreDelegates.h
@@ -130,9 +130,6 @@ public:
 	// Callback for handling accepting invitations - generally for engine code
 	DECLARE_MULTICAST_DELEGATE_TwoParams(FOnInviteAccepted, const FString&, const FString&);
 
-	UE_DEPRECATED(4.26, "FCoreDelegates::FRegisterEncryptionKeyDelegate is deprecated; use FRegisterEncryptionKeyMulticastDelegate instead")
-	DECLARE_DELEGATE_TwoParams(FRegisterEncryptionKeyDelegate, const FGuid&, const FAES::FAESKey&);
-
 	// Callback for registering a new encryption key
 	DECLARE_MULTICAST_DELEGATE_TwoParams(FRegisterEncryptionKeyMulticastDelegate, const FGuid&, const FAES::FAESKey&);
 
@@ -235,10 +232,6 @@ public:
 	// Called when an actor label is changed
 	static FOnActorLabelChanged OnActorLabelChanged;
 
-	UE_DEPRECATED(4.26, "FCoreDelegates::GetRegisterEncryptionKeyDelegate is deprecated; use GetRegisterEncryptionKeyMulticastDelegate instead")
-	PRAGMA_DISABLE_DEPRECATION_WARNINGS // Don't warn about FRegisterEncryptionKeyDelegate
-	static FRegisterEncryptionKeyDelegate& GetRegisterEncryptionKeyDelegate();
-	PRAGMA_ENABLE_DEPRECATION_WARNINGS
 	static FRegisterEncryptionKeyMulticastDelegate& GetRegisterEncryptionKeyMulticastDelegate();
 	static FPakEncryptionKeyDelegate& GetPakEncryptionKeyDelegate();
 	static FPakSigningKeysDelegate& GetPakSigningKeysDelegate();
diff --git a/Engine/Source/Runtime/Core/Public/Misc/KeyChainUtilities.h b/Engine/Source/Runtime/Core/Public/Misc/KeyChainUtilities.h
--- a/Engine/Source/Runtime/Core/Public/Misc/KeyChainUtilities.h
+++ b/Engine/Source/Runtime/Core/Public/Misc/KeyChainUtilities.h
@@ -219,12 +219,6 @@ namespace KeyChainUtilities
 		{
 			if (Key.Key.IsValid())
 			{
-				// Deprecated version
-				PRAGMA_DISABLE_DEPRECATION_WARNINGS
-				FCoreDelegates::GetRegisterEncryptionKeyDelegate().ExecuteIfBound(Key.Key, Key.Value.Key);
-				PRAGMA_ENABLE_DEPRECATION_WARNINGS
-
-				// New version
 				FCoreDelegates::GetRegisterEncryptionKeyMulticastDelegate().Broadcast(Key.Key, Key.Value.Key);
 			}
 		}
diff --git a/Engine/Source/Runtime/PakFile/Private/IPlatformFilePak.cpp b/Engine/Source/Runtime/PakFile/Private/IPlatformFilePak.cpp
--- a/Engine/Source/Runtime/PakFile/Private/IPlatformFilePak.cpp
+++ b/Engine/Source/Runtime/PakFile/Private/IPlatformFilePak.cpp
@@ -174,12 +174,6 @@ static void TestRegisterEncryptionKey(const TArray<FString>& Args)
 				check(KeyBytes.Num() == sizeof(FAES::FAESKey));
 				FMemory::Memcpy(EncryptionKey.Key, &KeyBytes[0], sizeof(EncryptionKey.Key));
 
-				// Deprecated version
-				PRAGMA_DISABLE_DEPRECATION_WARNINGS
-				FCoreDelegates::GetRegisterEncryptionKeyDelegate().ExecuteIfBound(EncryptionKeyGuid, EncryptionKey);
-				PRAGMA_ENABLE_DEPRECATION_WARNINGS
-
-				// New version
 				FCoreDelegates::GetRegisterEncryptionKeyMulticastDelegate().Broadcast(EncryptionKeyGuid, EncryptionKey);
 			}
 		}
