FCoreDelegates: deprecations cleanup #2

diff --git a/Engine/Source/Runtime/Core/Private/Misc/CoreDelegates.cpp b/Engine/Source/Runtime/Core/Private/Misc/CoreDelegates.cpp
--- a/Engine/Source/Runtime/Core/Private/Misc/CoreDelegates.cpp
+++ b/Engine/Source/Runtime/Core/Private/Misc/CoreDelegates.cpp
@@ -28,18 +28,11 @@ FCoreDelegates::FOnPreMainInit& FCoreDelegates::GetPreMainInitDelegate()
 }
 
 FCoreDelegates::FOnMountAllPakFiles FCoreDelegates::OnMountAllPakFiles;
-PRAGMA_DISABLE_DEPRECATION_WARNINGS
-FCoreDelegates::FOnMountPak FCoreDelegates::OnMountPak;
-PRAGMA_ENABLE_DEPRECATION_WARNINGS
 FCoreDelegates::FMountPak FCoreDelegates::MountPak;
 FCoreDelegates::FOnUnmountPak FCoreDelegates::OnUnmountPak;
 FCoreDelegates::FOnOptimizeMemoryUsageForMountedPaks FCoreDelegates::OnOptimizeMemoryUsageForMountedPaks;
 
 FCoreDelegates::FOnPakFileMounted2 FCoreDelegates::OnPakFileMounted2;
-PRAGMA_DISABLE_DEPRECATION_WARNINGS
-FCoreDelegates::FOnPakFileMounted FCoreDelegates::OnPakFileMounted;
-FCoreDelegates::FPakFileMountedDelegate FCoreDelegates::PakFileMountedCallback;
-PRAGMA_ENABLE_DEPRECATION_WARNINGS
 FCoreDelegates::FNewFileAddedDelegate FCoreDelegates::NewFileAddedDelegate;
 FCoreDelegates::FNoPakFilesMountedDelegate FCoreDelegates::NoPakFilesMountedDelegate;
 FCoreDelegates::FOnFileOpenedForReadFromPakFile FCoreDelegates::OnFileOpenedForReadFromPakFile;
diff --git a/Engine/Source/Runtime/Core/Public/Misc/CoreDelegates.h b/Engine/Source/Runtime/Core/Public/Misc/CoreDelegates.h
--- a/Engine/Source/Runtime/Core/Public/Misc/CoreDelegates.h
+++ b/Engine/Source/Runtime/Core/Public/Misc/CoreDelegates.h
@@ -81,9 +81,6 @@ public:
 	// delegate type for prompting the pak system to mount all pak files, which haven't already been mounted, from all default locations
 	DECLARE_DELEGATE_RetVal_OneParam(int32, FOnMountAllPakFiles, const TArray<FString>&);
 
-	// deprecated delegate type for prompting the pak system to mount a new pak
-	DECLARE_DELEGATE_RetVal_ThreeParams(bool, FOnMountPak, const FString&, int32, IPlatformFile::FDirectoryVisitor*);
-
 	// delegate type for prompting the pak system to mount a new pak
 	DECLARE_DELEGATE_RetVal_TwoParams(IPakFile*, FMountPak, const FString&, int32);
 
@@ -93,12 +90,6 @@ public:
 	// delegate type for prompting the pak system to optimize memory for mounted paks
 	DECLARE_DELEGATE(FOnOptimizeMemoryUsageForMountedPaks);
 
-	// deprecated delegate for handling when a new pak file is successfully mounted passes in the name of the mounted pak file
-	DECLARE_MULTICAST_DELEGATE_OneParam(FPakFileMountedDelegate, const TCHAR*);
-
-	// deprecated delegate for handling when a new pak file is successfully mounted passes in the name of the pak file and its chunk ID (or INDEX_NONE)
-	DECLARE_MULTICAST_DELEGATE_TwoParams(FOnPakFileMounted, const TCHAR*, const int32);
-
 	// delegate for handling when a new pak file is successfully mounted
 	DECLARE_MULTICAST_DELEGATE_OneParam(FOnPakFileMounted2, const IPakFile&);
 
@@ -187,9 +178,6 @@ public:
 	// Callback to prompt the pak system to mount a pak file
 	static FMountPak MountPak;
 
-	UE_DEPRECATED(4.26, "OnMountPak is deprecated; use MountPak instead.")
-	static FOnMountPak OnMountPak;
-
 	// Callback to prompt the pak system to unmount a pak file.
 	static FOnUnmountPak OnUnmountPak;
 
@@ -199,12 +187,6 @@ public:
 	// After a pakfile is mounted this is called
 	static FOnPakFileMounted2 OnPakFileMounted2;
 
-	UE_DEPRECATED(4.26, "FCoreDelegates::OnPakFileMounted is deprecated; use OnPakFileMounted2 instead")
-	static FOnPakFileMounted OnPakFileMounted;
-
-	UE_DEPRECATED(4.25, "FCoreDelegates::PakFileMountedCallback is deprecated. Use FCoreDelegates::OnPakFileMounted2 instead.")
-	static FPakFileMountedDelegate PakFileMountedCallback;
-
 	// After a file is added this is called
 	DECLARE_MULTICAST_DELEGATE_OneParam(FNewFileAddedDelegate, const FString&);
 	static FNewFileAddedDelegate NewFileAddedDelegate;
diff --git a/Engine/Source/Runtime/PakFile/Private/IPlatformFilePak.cpp b/Engine/Source/Runtime/PakFile/Private/IPlatformFilePak.cpp
--- a/Engine/Source/Runtime/PakFile/Private/IPlatformFilePak.cpp
+++ b/Engine/Source/Runtime/PakFile/Private/IPlatformFilePak.cpp
@@ -7281,9 +7281,6 @@ FPakPlatformFile::~FPakPlatformFile()
 	FCoreDelegates::OnFEngineLoopInitComplete.RemoveAll(this);
 
 	FCoreDelegates::OnMountAllPakFiles.Unbind();
-	PRAGMA_DISABLE_DEPRECATION_WARNINGS
-	FCoreDelegates::OnMountPak.Unbind();
-	PRAGMA_ENABLE_DEPRECATION_WARNINGS
 	FCoreDelegates::MountPak.Unbind();
 	FCoreDelegates::OnUnmountPak.Unbind();
 	FCoreDelegates::OnOptimizeMemoryUsageForMountedPaks.Unbind();
@@ -7510,9 +7507,6 @@ bool FPakPlatformFile::Initialize(IPlatformFile* Inner, const TCHAR* CmdLine)
 #endif // !UE_BUILD_SHIPPING
 
 	FCoreDelegates::OnMountAllPakFiles.BindRaw(this, &FPakPlatformFile::MountAllPakFiles);
-	PRAGMA_DISABLE_DEPRECATION_WARNINGS
-	FCoreDelegates::OnMountPak.BindRaw(this, &FPakPlatformFile::HandleOnMountPakDelegate);
-	PRAGMA_ENABLE_DEPRECATION_WARNINGS
 	FCoreDelegates::MountPak.BindRaw(this, &FPakPlatformFile::HandleMountPakDelegate);
 	FCoreDelegates::OnUnmountPak.BindRaw(this, &FPakPlatformFile::HandleUnmountPakDelegate);
 	FCoreDelegates::OnOptimizeMemoryUsageForMountedPaks.BindRaw(this, &FPakPlatformFile::OptimizeMemoryUsageForMountedPaks);
@@ -7805,10 +7799,6 @@ bool FPakPlatformFile::Mount(const TCHAR* InPakFilename, uint32 PakOrder, const
 
 		if (bPakSuccess)
 		{
-			PRAGMA_DISABLE_DEPRECATION_WARNINGS
-			FCoreDelegates::PakFileMountedCallback.Broadcast(InPakFilename);
-			FCoreDelegates::OnPakFileMounted.Broadcast(InPakFilename, Pak->PakchunkIndex);
-			PRAGMA_ENABLE_DEPRECATION_WARNINGS
 			double OnPakFileMounted2Time = 0.0;
 			{
 				FScopedDurationTimer Timer(OnPakFileMounted2Time);
