FCoreDelegates: deprecation cleanup #3

diff --git a/Engine/Source/Runtime/Core/Private/GenericPlatform/GenericPlatformCrashContext.cpp b/Engine/Source/Runtime/Core/Private/GenericPlatform/GenericPlatformCrashContext.cpp
--- a/Engine/Source/Runtime/Core/Private/GenericPlatform/GenericPlatformCrashContext.cpp
+++ b/Engine/Source/Runtime/Core/Private/GenericPlatform/GenericPlatformCrashContext.cpp
@@ -271,13 +271,8 @@ void FGenericCrashContext::Initialize()
 		NCached::Set(NCached::Session.GameStateName, *InGameStateName);
 	});
 
-PRAGMA_DISABLE_DEPRECATION_WARNINGS
 	FCoreDelegates::CrashOverrideParamsChanged.AddLambda([](const FCrashOverrideParameters& InParams)
 	{
-		if (InParams.bSetCrashReportClientMessageText)
-		{
-			NCached::Set(NCached::Session.CrashReportClientRichText, *InParams.CrashReportClientMessageText);
-		}
 		if (InParams.bSetGameNameSuffix)
 		{
 			NCached::Set(NCached::Session.GameName, *(FString(TEXT("UE-")) + FApp::GetProjectName() + InParams.GameNameSuffix));
@@ -292,7 +287,6 @@ PRAGMA_DISABLE_DEPRECATION_WARNINGS
 		}
 		SerializeTempCrashContextToFile();
 	});
-PRAGMA_ENABLE_DEPRECATION_WARNINGS
 
 	FCoreDelegates::OnPostEngineInit.AddLambda([] 
 	{
diff --git a/Engine/Source/Runtime/Core/Public/Misc/CoreDelegates.h b/Engine/Source/Runtime/Core/Public/Misc/CoreDelegates.h
--- a/Engine/Source/Runtime/Core/Public/Misc/CoreDelegates.h
+++ b/Engine/Source/Runtime/Core/Public/Misc/CoreDelegates.h
@@ -56,8 +56,6 @@ struct FTestHotFixPayload
 // Parameters passed to CrashOverrideParamsChanged used to customize crash report client behavior/appearance. If the corresponding bool is not true, this value will not be stored.
 struct FCrashOverrideParameters
 {
-	UE_DEPRECATED(4.21, "CrashReportClientMessageText should now be set through the CrashReportClientRichText property in the [CrashContextProperties] section of DefaultEngine.ini.")
-	FString CrashReportClientMessageText;
 	/** Appended to the end of GameName (which is retreived from FApp::GetGameName). */
 	FString GameNameSuffix;
 	/** Default this to true for backward compatibility before these bools were added. */
