Array: Fixed unsafe type cast warnings when moving from an array with a larger size type

diff --git a/Engine/Source/Runtime/Core/Public/Containers/Array.h b/Engine/Source/Runtime/Core/Public/Containers/Array.h
--- a/Engine/Source/Runtime/Core/Public/Containers/Array.h
+++ b/Engine/Source/Runtime/Core/Public/Containers/Array.h
@@ -513,8 +513,8 @@ private:
 			ToArray.AllocatorInstance.MoveToEmpty(FromArray.AllocatorInstance);
 		}
 
-		ToArray  .ArrayNum = FromArray.ArrayNum;
-		ToArray  .ArrayMax = FromArray.ArrayMax;
+		ToArray  .ArrayNum = (SizeType)FromArray.ArrayNum;
+		ToArray  .ArrayMax = (SizeType)FromArray.ArrayMax;
 
 		// Ensure the destination container could hold the source range (when the allocator size types shrink)
 		checkf(ToArray.ArrayNum == FromArray.ArrayNum && ToArray.ArrayMax == FromArray.ArrayMax, TEXT("Data lost when moving to a container with a more constrained size type"));
