DDC: Set UnsafeTypeCastWarningLevel to Error

#rnx

diff --git a/Engine/Source/Developer/DerivedDataCache/DerivedDataCache.Build.cs b/Engine/Source/Developer/DerivedDataCache/DerivedDataCache.Build.cs
--- a/Engine/Source/Developer/DerivedDataCache/DerivedDataCache.Build.cs
+++ b/Engine/Source/Developer/DerivedDataCache/DerivedDataCache.Build.cs
@@ -19,5 +19,7 @@ public class DerivedDataCache : ModuleRules
 		// Platform-specific opt-in
 		PrivateDefinitions.Add($"WITH_HTTP_DDC_BACKEND=1");
 		PrivateDefinitions.Add($"WITH_S3_DDC_BACKEND={(Target.Platform == UnrealTargetPlatform.Win64 ? 1 : 0)}");
+
+		UnsafeTypeCastWarningLevel = WarningLevel.Error;
 	}
 }
diff --git a/Engine/Source/Developer/DerivedDataCache/Private/Http/CurlHttpClient.cpp b/Engine/Source/Developer/DerivedDataCache/Private/Http/CurlHttpClient.cpp
--- a/Engine/Source/Developer/DerivedDataCache/Private/Http/CurlHttpClient.cpp
+++ b/Engine/Source/Developer/DerivedDataCache/Private/Http/CurlHttpClient.cpp
@@ -1219,7 +1219,7 @@ size_t FCurlHttpResponse::CurlDebug(CURL* Curl, curl_infotype DebugInfoType, cha
 
 size_t FCurlHttpResponse::CurlHeader(void* Ptr, size_t SizeInBlocks, size_t BlockSizeInBytes, void* Param)
 {
-	const size_t HeaderSize = SizeInBlocks * BlockSizeInBytes;
+	const int32 HeaderSize = IntCastChecked<int32>(SizeInBlocks * BlockSizeInBytes);
 	return ((FCurlHttpResponse*)Param)->WriteHeader(MakeStringView((const ANSICHAR*)Ptr, HeaderSize)) ? HeaderSize : 0;
 }
 
diff --git a/Engine/Source/Runtime/Core/Public/Containers/StringView.h b/Engine/Source/Runtime/Core/Public/Containers/StringView.h
--- a/Engine/Source/Runtime/Core/Public/Containers/StringView.h
+++ b/Engine/Source/Runtime/Core/Public/Containers/StringView.h
@@ -365,10 +365,10 @@ constexpr inline auto MakeStringView(CharPtrOrRangeType&& CharPtrOrRange) -> dec
 	return TStringView(Forward<CharPtrOrRangeType>(CharPtrOrRange));
 }
 
-template <typename CharPtrType, typename SizeType>
-constexpr inline auto MakeStringView(CharPtrType&& CharPtr, SizeType&& Size) -> decltype(TStringView(Forward<CharPtrType>(CharPtr), Forward<SizeType>(Size)))
+template <typename CharPtrType>
+constexpr inline auto MakeStringView(CharPtrType&& CharPtr, int32 Size) -> decltype(TStringView(Forward<CharPtrType>(CharPtr), Size))
 {
-	return TStringView(Forward<CharPtrType>(CharPtr), Forward<SizeType>(Size));
+	return TStringView(Forward<CharPtrType>(CharPtr), Size);
 }
 
 template <typename CharType>
