TOptional<TNonNullPtr>: Add support for NullOpt

diff --git a/Engine/Source/Runtime/Core/Public/Templates/NonNullPointer.h b/Engine/Source/Runtime/Core/Public/Templates/NonNullPointer.h
--- a/Engine/Source/Runtime/Core/Public/Templates/NonNullPointer.h
+++ b/Engine/Source/Runtime/Core/Public/Templates/NonNullPointer.h
@@ -181,6 +181,12 @@ public:
 	{
 	}
 
+	/** Construct an OptionalType with an invalid value. */
+	TOptional(FNullOpt)
+		: TOptional()
+	{
+	}
+
 	TOptional& operator=(const TOptional& Other)
 	{
 		Pointer = Other.Pointer;
