fixed MTAccessDetector usage in TScriptDelegate

diff --git a/Engine/Source/Runtime/Core/Public/UObject/ScriptDelegates.h b/Engine/Source/Runtime/Core/Public/UObject/ScriptDelegates.h
--- a/Engine/Source/Runtime/Core/Public/UObject/ScriptDelegates.h
+++ b/Engine/Source/Runtime/Core/Public/UObject/ScriptDelegates.h
@@ -151,7 +151,7 @@ public:
 	/** Delegate serialization */
 	friend FArchive& operator<<( FArchive& Ar, TScriptDelegate& D )
 	{
-		UE_DELEGATES_MT_SCOPED_READ_ACCESS(AccessDetector);
+		UE_DELEGATES_MT_SCOPED_READ_ACCESS(D.AccessDetector);
 
 		Ar << D.Object << D.FunctionName;
 		return Ar;
@@ -160,7 +160,7 @@ public:
 	/** Delegate serialization */
 	friend void operator<<(FStructuredArchive::FSlot Slot, TScriptDelegate& D)
 	{
-		UE_DELEGATES_MT_SCOPED_READ_ACCESS(AccessDetector);
+		UE_DELEGATES_MT_SCOPED_READ_ACCESS(D.AccessDetector);
 
 		FStructuredArchive::FRecord Record = Slot.EnterRecord();
 		Record << SA_VALUE(TEXT("Object"), D.Object) << SA_VALUE(TEXT("FunctionName"),D.FunctionName);
@@ -285,7 +285,7 @@ public:
 
 	friend uint32 GetTypeHash(const TScriptDelegate& Delegate)
 	{
-		UE_DELEGATES_MT_SCOPED_READ_ACCESS(AccessDetector);
+		UE_DELEGATES_MT_SCOPED_READ_ACCESS(Delegate.AccessDetector);
 
 		return HashCombine(GetTypeHash(Delegate.Object), GetTypeHash(Delegate.GetFunctionName()));
 	}
@@ -497,7 +497,7 @@ public:
 	/** Multi-cast delegate serialization */
 	friend FArchive& operator<<( FArchive& Ar, TMulticastScriptDelegate<TWeakPtr>& D )
 	{
-		UE_DELEGATES_MT_SCOPED_WRITE_ACCESS(AccessDetector);
+		UE_DELEGATES_MT_SCOPED_WRITE_ACCESS(D.AccessDetector);
 
 		if( Ar.IsSaving() )
 		{
@@ -518,7 +518,7 @@ public:
 
 	friend void operator<<(FStructuredArchive::FSlot Slot, TMulticastScriptDelegate<TWeakPtr>& D)
 	{
-		UE_DELEGATES_MT_SCOPED_WRITE_ACCESS(AccessDetector);
+		UE_DELEGATES_MT_SCOPED_WRITE_ACCESS(D.AccessDetector);
 
 		FArchive& UnderlyingArchive = Slot.GetUnderlyingArchive();
 
