config-related delegates made thread-safe as they are accessed concurrently during async config loading

diff --git a/Engine/Source/Runtime/Core/Private/Misc/CoreDelegates.cpp b/Engine/Source/Runtime/Core/Private/Misc/CoreDelegates.cpp
--- a/Engine/Source/Runtime/Core/Private/Misc/CoreDelegates.cpp
+++ b/Engine/Source/Runtime/Core/Private/Misc/CoreDelegates.cpp
@@ -130,7 +130,7 @@ FCoreDelegates::FCountPreLoadConfigFileRespondersDelegate FCoreDelegates::CountP
 FCoreDelegates::FPreLoadConfigFileDelegate FCoreDelegates::PreLoadConfigFileDelegate;
 FCoreDelegates::FPreSaveConfigFileDelegate FCoreDelegates::PreSaveConfigFileDelegate;
 FCoreDelegates::FOnFConfigFileCreated FCoreDelegates::OnFConfigCreated;
-FCoreDelegates::FOnFConfigFileCreated FCoreDelegates::OnFConfigDeleted;
+FCoreDelegates::FOnFConfigFileDeleted FCoreDelegates::OnFConfigDeleted;
 FCoreDelegates::FOnConfigValueRead FCoreDelegates::OnConfigValueRead;
 FCoreDelegates::FOnConfigSectionRead FCoreDelegates::OnConfigSectionRead;
 FCoreDelegates::FOnConfigSectionRead FCoreDelegates::OnConfigSectionNameRead;
diff --git a/Engine/Source/Runtime/Core/Public/Misc/CoreDelegates.h b/Engine/Source/Runtime/Core/Public/Misc/CoreDelegates.h
--- a/Engine/Source/Runtime/Core/Public/Misc/CoreDelegates.h
+++ b/Engine/Source/Runtime/Core/Public/Misc/CoreDelegates.h
@@ -327,26 +327,26 @@ public:
 
 
 
-	DECLARE_MULTICAST_DELEGATE_TwoParams(FCountPreLoadConfigFileRespondersDelegate, const TCHAR* /*IniFilename*/, int32& /*ResponderCount*/);
-	DECLARE_MULTICAST_DELEGATE_TwoParams(FPreLoadConfigFileDelegate, const TCHAR* /*IniFilename*/, FString& /*LoadedContents*/);
-	DECLARE_MULTICAST_DELEGATE_ThreeParams(FPreSaveConfigFileDelegate, const TCHAR* /*IniFilename*/, const FString& /*ContentsToSave*/, int32& /*SavedCount*/);
+	DECLARE_TS_MULTICAST_DELEGATE_TwoParams(FCountPreLoadConfigFileRespondersDelegate, const TCHAR* /*IniFilename*/, int32& /*ResponderCount*/);
+	DECLARE_TS_MULTICAST_DELEGATE_TwoParams(FPreLoadConfigFileDelegate, const TCHAR* /*IniFilename*/, FString& /*LoadedContents*/);
+	DECLARE_TS_MULTICAST_DELEGATE_ThreeParams(FPreSaveConfigFileDelegate, const TCHAR* /*IniFilename*/, const FString& /*ContentsToSave*/, int32& /*SavedCount*/);
 	static FCountPreLoadConfigFileRespondersDelegate CountPreLoadConfigFileRespondersDelegate;
 	static FPreLoadConfigFileDelegate PreLoadConfigFileDelegate;
 	static FPreSaveConfigFileDelegate PreSaveConfigFileDelegate;
 
-	DECLARE_MULTICAST_DELEGATE_OneParam(FOnFConfigFileCreated, const FConfigFile *);
-	DECLARE_MULTICAST_DELEGATE_OneParam(FOnFConfigFileDeleted, const FConfigFile *);
+	DECLARE_TS_MULTICAST_DELEGATE_OneParam(FOnFConfigFileCreated, const FConfigFile *);
+	DECLARE_TS_MULTICAST_DELEGATE_OneParam(FOnFConfigFileDeleted, const FConfigFile *);
 	static FOnFConfigFileCreated OnFConfigCreated;
 	static FOnFConfigFileDeleted OnFConfigDeleted;
 
-	DECLARE_MULTICAST_DELEGATE_ThreeParams(FOnConfigValueRead, const TCHAR* /*IniFilename*/, const TCHAR* /*SectionName*/, const TCHAR* /*Key*/);
+	DECLARE_TS_MULTICAST_DELEGATE_ThreeParams(FOnConfigValueRead, const TCHAR* /*IniFilename*/, const TCHAR* /*SectionName*/, const TCHAR* /*Key*/);
 	static FOnConfigValueRead OnConfigValueRead;
 
-	DECLARE_MULTICAST_DELEGATE_TwoParams(FOnConfigSectionRead, const TCHAR* /*IniFilename*/, const TCHAR* /*SectionName*/);
+	DECLARE_TS_MULTICAST_DELEGATE_TwoParams(FOnConfigSectionRead, const TCHAR* /*IniFilename*/, const TCHAR* /*SectionName*/);
 	static FOnConfigSectionRead OnConfigSectionRead;
 	static FOnConfigSectionRead OnConfigSectionNameRead;
 
-	DECLARE_MULTICAST_DELEGATE_TwoParams(FOnConfigSectionsChanged, const FString& /*IniFilename*/, const TSet<FString>& /*SectionNames*/);
+	DECLARE_TS_MULTICAST_DELEGATE_TwoParams(FOnConfigSectionsChanged, const FString& /*IniFilename*/, const TSet<FString>& /*SectionNames*/);
 	static FOnConfigSectionsChanged OnConfigSectionsChanged;
 
 	DECLARE_MULTICAST_DELEGATE_FourParams(FOnApplyCVarFromIni, const TCHAR* /*SectionName*/, const TCHAR* /*IniFilename*/, uint32 /*SetBy*/, bool /*bAllowCheating*/);
@@ -549,7 +549,7 @@ public:
 	static FOnPreMainInit& GetPreMainInitDelegate();
 	
 	/** Sent when GConfig is finished initializing */
-	DECLARE_MULTICAST_DELEGATE(FConfigReadyForUse);
+	DECLARE_TS_MULTICAST_DELEGATE(FConfigReadyForUse);
 	static FConfigReadyForUse ConfigReadyForUse;
 
 	/** Callback for notifications regarding changes of the rendering thread. */
@@ -637,7 +637,7 @@ public:
 		{
 		}
 	};
-	DECLARE_MULTICAST_DELEGATE_OneParam(FAccesExtraBinaryConfigData, FExtraBinaryConfigData&);
+	DECLARE_TS_MULTICAST_DELEGATE_OneParam(FAccesExtraBinaryConfigData, FExtraBinaryConfigData&);
 	static FAccesExtraBinaryConfigData AccessExtraBinaryConfigData;
 
 	using FAttachShaderReadRequestFunc = TFunctionRef<class FIoRequest(const class FIoChunkId&, FGraphEventRef)>;
