Fix using FPaths::DirectoryExists in crash handling code as it depends on a global which may be dead at time of crash/assert

Also move to less allocations, and a static buffer. We still use TCHAR_TO_UTF8 allocations but will look at fixing up later

diff --git a/Engine/Source/Runtime/Core/Private/Unix/UnixPlatformStackWalk.cpp b/Engine/Source/Runtime/Core/Private/Unix/UnixPlatformStackWalk.cpp
--- a/Engine/Source/Runtime/Core/Private/Unix/UnixPlatformStackWalk.cpp
+++ b/Engine/Source/Runtime/Core/Private/Unix/UnixPlatformStackWalk.cpp
@@ -331,7 +331,17 @@ namespace
 #if WITH_LOW_LEVEL_TESTS
 				// Low level tests live one level above the base directory in a folder <ModuleName>Tests
 				// Sometimes this folder can also be just <ModuleName> if the target was compiled with the tests
-				if (FPaths::DirectoryExists(FPaths::Combine(ModuleSymbolPath, TCHAR_TO_UTF8(*FPaths::GetBaseFilename(out_SymbolInfo.ModuleName)))))
+				// TODO: This code needs work as its only hardcoded to allows finding the *.sym for Development config.
+				// Debug/Test/Shipping/ASan configs all fail here
+				ANSICHAR ModuleDirectory[UNIX_MAX_PATH + 1];
+
+				FCStringAnsi::Strcpy(ModuleDirectory, ModuleSymbolPath);
+				FCStringAnsi::Strcat(ModuleDirectory, "/");
+				FCStringAnsi::Strcat(ModuleDirectory, TCHAR_TO_UTF8(*FPaths::GetBaseFilename(out_SymbolInfo.ModuleName)));
+
+				// use stat instead of FPaths::DirectoryExists as it calls into a static global which may be dead at exit time
+				struct stat StatInfo;
+				if (stat(ModuleDirectory, &StatInfo) == 0)
 				{
 					FCStringAnsi::Strcat(ModuleSymbolPath, "/");
 					FCStringAnsi::Strcat(ModuleSymbolPath, TCHAR_TO_UTF8(*FPaths::GetBaseFilename(out_SymbolInfo.ModuleName)));
