Serialization function for TVariants and addition of TVariant unit tests

In order for a TVariant to be serializable, each type in its template parameter pack must:
    1. Have a default constructor. This is required because when reading the type from an archive, it must be default constructed before being loaded.
    2. Implement the `FArchive& operator<<(FArchive&, T&)` function. This is required to serialize the actual type that's stored in TVariant.

diff --git a/Engine/Source/Runtime/Core/Private/Tests/Misc/TVariantTest.cpp b/Engine/Source/Runtime/Core/Private/Tests/Misc/TVariantTest.cpp
new file mode 100644
--- /dev/null
+++ b/Engine/Source/Runtime/Core/Private/Tests/Misc/TVariantTest.cpp
@@ -0,0 +1,536 @@
+// Copyright Epic Games, Inc. All Rights Reserved.
+
+#include "CoreTypes.h"
+#include "Templates/SharedPointer.h"
+#include "Templates/UniquePtr.h"
+#include "Misc/AutomationTest.h"
+#include "Misc/TVariant.h"
+#include "Serialization/MemoryReader.h"
+#include "Serialization/MemoryWriter.h"
+
+#if WITH_DEV_AUTOMATION_TESTS
+
+namespace UE::Core::Private
+{
+	// TIsAllVariant
+	static_assert(TIsAllVariant<TVariant<int, float>>::Value, "Expected TIsAllVariant<TVariant> to be true");
+	static_assert(TIsAllVariant<TVariant<int, float>, TVariant<float, int>>::Value, "Expected TIsAllVariant<TVariant, TVariant> to be true");
+	static_assert(TIsAllVariant<TVariant<int, float>, const TVariant<float, int>>::Value, "Expected TIsAllVariant<TVariant, const TVariant> to be true");
+	static_assert(TIsAllVariant<>::Value, "Expected TIsAllVariant<> for base case to be true");
+	static_assert(TIsAllVariant<bool>::Value == false, "Expected TIsAllVariant<non-TVariant> to be false");
+	static_assert(TIsAllVariant<TVariant<int, float>, bool>::Value == false, "Expected TIsAllVariant<TVariant, non-TVariant> to be false");
+	static_assert(TIsAllVariant<TVariant<int, float>, bool, TVariant<float, int>>::Value == false, "Expected TIsAllVariant<TVariant, non-TVariant, TVariant> to be false");
+
+	// TIsVariant
+	static_assert(TIsVariant<TVariant<int, float>>::Value, "Expected TIsVariant<TVariant> to be true");
+	static_assert(TIsVariant<const TVariant<int, float>>::Value, "Expected TIsVariant<const TVariant> to be true");
+	static_assert(TIsVariant<TVariant<int, float>&>::Value, "Expected TIsVariant<TVariant&> to be true");
+	static_assert(TIsVariant<const TVariant<int, float>&>::Value, "Expected TIsVariant<const TVariant&> to be true");
+	static_assert(TIsVariant<TVariant<int, float>&&>::Value, "Expected TIsVariant<TVariant&&> to be true");
+	static_assert(TIsVariant<bool>::Value == false, "Expected TIsVariant<non-TVariant> to be false");
+
+	// TVariantSize
+	static_assert(TVariantSize<TVariant<int, float>>::Value == 2, "Expected TVariantSize<TVariant> to be 2");
+	static_assert(TVariantSize<const TVariant<int, float>>::Value == 2, "Expected TVariantSize<const TVariant> to be 2");
+	static_assert(TVariantSize<TVariant<int, float>&>::Value == 2, "Expected TVariantSize<TVariant&> to be 2");
+	static_assert(TVariantSize<const TVariant<int, float>&>::Value == 2, "Expected TVariantSize<const TVariant&> to be 2");
+	static_assert(TVariantSize<TVariant<int, float>&&>::Value == 2, "Expected TVariantSize<TVariant&&> to be 2");
+} // namespace UE::Core::Private
+
+struct FNonDefaultConstructible
+{
+	int _A;
+	float _B;
+	FNonDefaultConstructible(int A, float B) : _A(A), _B(B) {}
+
+	bool operator==(const FNonDefaultConstructible& Other) const
+	{
+		return _A == Other._A && _B == Other._B;
+	}
+};
+
+IMPLEMENT_SIMPLE_AUTOMATION_TEST(FTVariantInPlaceConstructionTest, "System.Core.Misc.TVariant.InPlaceConstruction", EAutomationTestFlags::ApplicationContextMask | EAutomationTestFlags::EngineFilter)
+bool FTVariantInPlaceConstructionTest::RunTest(const FString&)
+{
+	// TVariant supports in-place construction of a type in its parameter pack
+	using FVariant = TVariant<FNonDefaultConstructible>;
+	
+	FVariant Var{ TInPlaceType<FNonDefaultConstructible>(), 30, 40.0f };
+	check(Var.IsType<FNonDefaultConstructible>());
+	check(Var.Get<FNonDefaultConstructible>() == FNonDefaultConstructible(30, 40.0f));
+
+	return true;
+}
+
+IMPLEMENT_SIMPLE_AUTOMATION_TEST(FTVariantDefaultConstructorTest, "System.Core.Misc.TVariant.DefaultConstruction", EAutomationTestFlags::ApplicationContextMask | EAutomationTestFlags::EngineFilter)
+bool FTVariantDefaultConstructorTest::RunTest(const FString&)
+{
+	// TVariant default constructor constructs the first type in the parameter pack
+	TVariant<int, float, FString> Int;
+	check(Int.IsType<int>());
+	check(Int.Get<int>() == 0);
+	
+	TVariant<FString, int, float> String;
+	check(String.IsType<FString>());
+	check(String.Get<FString>() == FString{});
+	
+	return true;
+}
+
+IMPLEMENT_SIMPLE_AUTOMATION_TEST(FTVariantEmptyVariantStateTest, "System.Core.Misc.TVariant.EmptyVariantState", EAutomationTestFlags::ApplicationContextMask | EAutomationTestFlags::EngineFilter)
+bool FTVariantEmptyVariantStateTest::RunTest(const FString&)
+{
+	// TVariant default constructor works for the provided FEmptyVariantState
+	TVariant<FEmptyVariantState, FNonDefaultConstructible> Var;
+	check(Var.IsType<FEmptyVariantState>());
+	
+	return true;
+}
+
+IMPLEMENT_SIMPLE_AUTOMATION_TEST(FTVariantCopyTest, "System.Core.Misc.TVariant.Copy", EAutomationTestFlags::ApplicationContextMask | EAutomationTestFlags::EngineFilter)
+bool FTVariantCopyTest::RunTest(const FString&)
+{
+	// TVariants can be copied
+	using FVariant = TVariant<int, float, double>;
+	FVariant Int{ TInPlaceType<int>(), 40 };
+	FVariant Copy{ Int };
+
+	check(Int.IsType<int>());
+	check(Copy.IsType<int>());
+	check(Int.Get<int>() == Copy.Get<int>());
+
+	return true;
+}
+
+IMPLEMENT_SIMPLE_AUTOMATION_TEST(FTVariantMoveTest, "System.Core.Misc.TVariant.Move", EAutomationTestFlags::ApplicationContextMask | EAutomationTestFlags::EngineFilter)
+bool FTVariantMoveTest::RunTest(const FString&)
+{
+	// TVariants can be moved
+	using FVariant = TVariant<int, float, double>;
+	FVariant Original{ TInPlaceType<int>(), 40 };
+	FVariant NewLoc{ MoveTemp(Original) };
+
+	check(NewLoc.IsType<int>());
+	check(NewLoc.Get<int>() == 40);
+
+	return true;
+}
+
+IMPLEMENT_SIMPLE_AUTOMATION_TEST(FTVariantAssignTest, "System.Core.Misc.TVariant.Assign", EAutomationTestFlags::ApplicationContextMask | EAutomationTestFlags::EngineFilter)
+bool FTVariantAssignTest::RunTest(const FString&)
+{
+	// TVariants can be assigned from TVariants of the same template parameter pack
+
+	using FVariant = TVariant<int, float, double>;
+
+	// Copy Assign
+	{
+		FVariant One{ TInPlaceType<double>(), 20.0 };
+		FVariant Two{ TInPlaceType<int>(), 40 };
+
+		One = Two;
+		check(One.IsType<int>());
+		check(One.Get<int>() == 40);
+
+		check(Two.IsType<int>());
+		check(Two.Get<int>() == 40);
+	}
+
+	// Move Assign
+	{
+		FVariant One{ TInPlaceType<double>(), 20.0 };
+		FVariant Two{ TInPlaceType<int>(), 40 };
+
+		One = MoveTemp(Two);
+		check(One.IsType<int>());
+		check(One.Get<int>() == 40);
+	}
+
+	return true;
+}
+
+IMPLEMENT_SIMPLE_AUTOMATION_TEST(FTVariantConstGetTest, "System.Core.Misc.TVariant.ConstGet", EAutomationTestFlags::ApplicationContextMask | EAutomationTestFlags::EngineFilter)
+bool FTVariantConstGetTest::RunTest(const FString&)
+{
+	// const TVariant values can be gotten
+	const TVariant<int, float, double> ConstVar{ TInPlaceType<float>(), 20.0f };
+	check(ConstVar.IsType<float>());
+	check(ConstVar.Get<float>() == 20.0f);
+
+	return true;
+}
+
+IMPLEMENT_SIMPLE_AUTOMATION_TEST(FTVariantTryGetTest, "System.Core.Misc.TVariant.TryGet", EAutomationTestFlags::ApplicationContextMask | EAutomationTestFlags::EngineFilter)
+bool FTVariantTryGetTest::RunTest(const FString&)
+{
+	// TVariant supports conditional get
+
+	// mutable
+	{
+		TVariant<int, float, double> Int{ TInPlaceType<int>(), 20 };
+
+		check(Int.TryGet<double>() == nullptr);
+
+		int* Value = Int.TryGet<int>();
+		check(Value != nullptr);
+		check(*Value == 20);
+	}
+
+	// const
+	{
+		const TVariant<int, float, double> Int{ TInPlaceType<int>(), 20 };
+
+		check(Int.TryGet<double>() == nullptr);
+
+		const int* Value = Int.TryGet<int>();
+		check(Value != nullptr);
+		check(*Value == 20);
+	}
+
+	return true;
+}
+
+IMPLEMENT_SIMPLE_AUTOMATION_TEST(FTVariantSharedPtrTest, "System.Core.Misc.TVariant.SharedPtr", EAutomationTestFlags::ApplicationContextMask | EAutomationTestFlags::EngineFilter)
+bool FTVariantSharedPtrTest::RunTest(const FString&)
+{
+	// TVariant can successfully hold a TSharedPtr that goes out of scope
+	
+	// Assignment
+	{
+		TVariant<TSharedPtr<int>> Var;
+		{
+			TSharedPtr<int> Ptr = MakeShared<int>(20);
+			Var.Set<TSharedPtr<int>>(Ptr);
+		}
+
+		check(Var.IsType<TSharedPtr<int>>());
+
+		TSharedPtr<int>& PtrRef = Var.Get<TSharedPtr<int>>();
+		check(PtrRef.IsValid());
+		check(*PtrRef == 20);
+	}
+
+	// Construction
+	{
+		TUniquePtr<TVariant<TSharedPtr<int>>> Var = nullptr;
+		{
+			TSharedPtr<int> Ptr = MakeShared<int>(20);
+			Var = MakeUnique<TVariant<TSharedPtr<int>>>(TInPlaceType<TSharedPtr<int>>(), Ptr);
+		}
+
+		check(Var->IsType<TSharedPtr<int>>());
+
+		TSharedPtr<int>& PtrRef = Var->Get<TSharedPtr<int>>();
+		check(PtrRef.IsValid());
+		check(*PtrRef == 20);
+	}
+
+	return true;
+}
+
+IMPLEMENT_SIMPLE_AUTOMATION_TEST(FTVariantUniquePtrTest, "System.Core.Misc.TVariant.UniquePtr", EAutomationTestFlags::ApplicationContextMask | EAutomationTestFlags::EngineFilter)
+bool FTVariantUniquePtrTest::RunTest(const FString&)
+{
+	// TVariant can successfully hold a TUniquePtr
+
+	// Assignment
+	{
+		TVariant<TUniquePtr<int>> Var;
+		Var.Set<TUniquePtr<int>>(MakeUnique<int>(20));
+
+		check(Var.IsType<TUniquePtr<int>>());
+
+		TUniquePtr<int>& PtrRef = Var.Get<TUniquePtr<int>>();
+		check(*PtrRef == 20);
+	}
+
+	// Construction
+	{
+		TVariant<TUniquePtr<int>> Var{ TInPlaceType<TUniquePtr<int>>(), MakeUnique<int>(20) };
+		check(Var.IsType<TUniquePtr<int>>());
+
+		TUniquePtr<int>& PtrRef = Var.Get<TUniquePtr<int>>();
+		check(*PtrRef == 20);
+	}
+
+	return true;
+}
+
+IMPLEMENT_SIMPLE_AUTOMATION_TEST(FTVariantMoveAssignNonCopyableTest, "System.Core.Misc.TVariant.MoveAssignNonCopyable", EAutomationTestFlags::ApplicationContextMask | EAutomationTestFlags::EngineFilter)
+bool FTVariantMoveAssignNonCopyableTest::RunTest(const FString&)
+{
+	// TVariant holding a non-copyable can be move assigned
+	struct FNoCopy
+	{
+		int Data = 0;
+		FNoCopy() = default;
+		explicit FNoCopy(int _Data) : Data(_Data) {}
+		FNoCopy(const FNoCopy&) = delete;
+		FNoCopy(FNoCopy&& Other) = default;
+		FNoCopy& operator=(const FNoCopy&) = delete;
+		FNoCopy& operator=(FNoCopy&& Other) = default;
+	};
+
+	TVariant<FNoCopy> Var{ TInPlaceType<FNoCopy>(), 20 };
+
+	TVariant<FNoCopy> Moved;
+	Moved = MoveTemp(Var);
+	check(Moved.Get<FNoCopy>().Data == 20);
+
+	return true;
+}
+
+IMPLEMENT_SIMPLE_AUTOMATION_TEST(FTVariantEmplaceTest, "System.Core.Misc.TVariant.Emplace", EAutomationTestFlags::ApplicationContextMask | EAutomationTestFlags::EngineFilter)
+bool FTVariantEmplaceTest::RunTest(const FString&)
+{
+	// TVariant can emplace values directly
+	TVariant<int, float, FString> Var;
+
+	Var.Emplace<int>(20);
+	check(Var.IsType<int>());
+	check(Var.Get<int>() == 20);
+
+	Var.Emplace<float>(30.0f);
+	check(Var.IsType<float>());
+	check(Var.Get<float>() == 30.0f);
+
+	Var.Emplace<FString>(TEXT("FOO"));
+	check(Var.IsType<FString>());
+	check(Var.Get<FString>() == FString{ TEXT("FOO") });
+
+	return true;
+}
+
+IMPLEMENT_SIMPLE_AUTOMATION_TEST(FTVariantEmplaceNonMoveNonCopyTest, "System.Core.Misc.TVariant.EmplaceNonMoveNonCopy", EAutomationTestFlags::ApplicationContextMask | EAutomationTestFlags::EngineFilter)
+bool FTVariantEmplaceNonMoveNonCopyTest::RunTest(const FString&)
+{
+	// TVariant can emplace values that are neither moveable nor copyable
+	struct FNonMoveNonCopy
+	{
+		int Data;
+		explicit FNonMoveNonCopy(int _Data) : Data(_Data) {}
+		~FNonMoveNonCopy() = default;
+		FNonMoveNonCopy(FNonMoveNonCopy&&) = delete;
+		FNonMoveNonCopy(const FNonMoveNonCopy&) = delete;
+		FNonMoveNonCopy& operator=(FNonMoveNonCopy&&) = delete;
+		FNonMoveNonCopy& operator=(const FNonMoveNonCopy&) = delete;
+	};
+
+	TVariant<FEmptyVariantState, FNonMoveNonCopy> Var;
+	Var.Emplace<FNonMoveNonCopy>(20);
+	check(Var.Get<FNonMoveNonCopy>().Data == 20);
+
+	return true;
+}
+
+IMPLEMENT_SIMPLE_AUTOMATION_TEST(FTVariantSelfAssignTest, "System.Core.Misc.TVariant.SelfAssign", EAutomationTestFlags::ApplicationContextMask | EAutomationTestFlags::EngineFilter)
+bool FTVariantSelfAssignTest::RunTest(const FString&)
+{
+	// TVariant self-assignment is harmless
+	TVariant<int, float, double> Var{ TInPlaceType<int>(), 20 };
+
+	const TVariant<int, float, double>& Ref = Var; // trick clang into thinking we're not assigning the same variable to itself.
+
+	Var = Ref;
+	check(Var.Get<int>() == 20);
+
+	Var = MoveTemp(Var);
+	check(Var.Get<int>() == 20);
+
+	return true;
+}
+
+IMPLEMENT_SIMPLE_AUTOMATION_TEST(FTVariantSwitchPatternMatchingTest, "System.Core.Misc.TVariant.SwitchPatternMatching", EAutomationTestFlags::ApplicationContextMask | EAutomationTestFlags::EngineFilter)
+bool FTVariantSwitchPatternMatchingTest::RunTest(const FString&)
+{
+	// TVariant supports switch-based 'pattern matching
+	using FVariant = TVariant<int, float, double>;
+	check(FVariant::IndexOfType<int>() == 0);
+	check(FVariant::IndexOfType<float>() == 1);
+	check(FVariant::IndexOfType<double>() == 2);
+
+	SIZE_T IndexTaken = (SIZE_T)-1;
+	auto Tester = [&IndexTaken](const FVariant& Var)
+	{
+		switch (IndexTaken = Var.GetIndex())
+		{
+			case FVariant::IndexOfType<int>():
+			{
+				int Value = Var.Get<int>();
+				check(Value == 20);
+				break;
+			}
+
+			case FVariant::IndexOfType<float>():
+			{
+				float Value = Var.Get<float>();
+				check(Value == 30.0f);
+				break;
+			}
+
+			case FVariant::IndexOfType<double>():
+			{
+				double Value = Var.Get<double>();
+				check(Value == 10.0);
+				break;
+			}
+
+			default:
+				check(!"Should not get here");
+				break;
+		}
+	};
+
+	FVariant Var;
+
+	Var.Set<int>(20);
+	Tester(Var);
+	check(IndexTaken == FVariant::IndexOfType<int>());
+
+	Var.Set<float>(30.0f);
+	Tester(Var);
+	check(IndexTaken == FVariant::IndexOfType<float>());
+
+	Var.Set<double>(10.0);
+	Tester(Var);
+	check(IndexTaken == FVariant::IndexOfType<double>());
+
+	return true;
+}
+
+
+IMPLEMENT_SIMPLE_AUTOMATION_TEST(FTVariantDestructorTest, "System.Core.Misc.TVariant.Destructor", EAutomationTestFlags::ApplicationContextMask | EAutomationTestFlags::EngineFilter)
+bool FTVariantDestructorTest::RunTest(const FString&)
+{
+	// Destructor test
+	struct FTestItem
+	{
+		int* DestructorCallCount;
+		FTestItem(int* _DestructorCallCount) : DestructorCallCount(_DestructorCallCount) {}
+		~FTestItem()
+		{
+			(*DestructorCallCount)++;
+		}
+	};
+
+	int DestructorCallCount = 0;
+	{
+		TVariant<TUniquePtr<FTestItem>> MovedFrom{ TInPlaceType<TUniquePtr<FTestItem>>(), MakeUnique<FTestItem>(&DestructorCallCount) };
+		TVariant<TUniquePtr<FTestItem>> MovedInto{ MoveTemp(MovedFrom) };
+	}
+	check(DestructorCallCount == 1);
+
+	DestructorCallCount = 0;
+	{
+		TVariant<TSharedPtr<FTestItem>> MovedFrom{ TInPlaceType<TSharedPtr<FTestItem>>(), MakeShared<FTestItem>(&DestructorCallCount) };
+		TVariant<TSharedPtr<FTestItem>> MovedInto{ MoveTemp(MovedFrom) };
+	}
+	check(DestructorCallCount == 1);
+
+	return true;
+}
+
+IMPLEMENT_SIMPLE_AUTOMATION_TEST(FTVariantVisitTest, "System.Core.Misc.TVariant.Visit", EAutomationTestFlags::ApplicationContextMask | EAutomationTestFlags::EngineFilter)
+bool FTVariantVisitTest::RunTest(const FString&)
+{
+	// Visit calls the correct function overload
+	enum class EVisitType
+	{
+		IntFloat,
+		IntInt,
+		IntStr,
+		FloatInt,
+		FloatFloat,
+		FloatStr,
+		StrInt,
+		StrFloat,
+		StrStr,
+	};
+	struct FCallable
+	{
+		EVisitType operator()(int, float) { return EVisitType::IntFloat; }
+		EVisitType operator()(int, int) { return EVisitType::IntInt; }
+		EVisitType operator()(int, const FString&) { return EVisitType::IntStr; }
+		EVisitType operator()(float, int) { return EVisitType::FloatInt; }
+		EVisitType operator()(float, float) { return EVisitType::FloatFloat; }
+		EVisitType operator()(float, const FString&) { return EVisitType::FloatStr; }
+		EVisitType operator()(const FString&, int) { return EVisitType::StrInt; }
+		EVisitType operator()(const FString&, float) { return EVisitType::StrFloat; }
+		EVisitType operator()(const FString&, const FString&) { return EVisitType::StrStr; }
+	};
+
+	TVariant<int, float> IntV(TInPlaceType<int>(), 0);
+	const TVariant<int, float, FString> FloatV(TInPlaceType<float>(), 0.0f);
+	const TVariant<FString> StrV(TInPlaceType<FString>(), TEXT(""));
+
+#define CHECK_VISIT(Type1, Type2) check(Visit(FCallable{}, Type1##V, Type2##V) == EVisitType::Type1##Type2)
+	CHECK_VISIT(Int, Int);
+	CHECK_VISIT(Int, Float);
+	CHECK_VISIT(Int, Str);
+	CHECK_VISIT(Float, Int);
+	CHECK_VISIT(Float, Float);
+	CHECK_VISIT(Float, Str);
+	CHECK_VISIT(Str, Int);
+	CHECK_VISIT(Str, Float);
+	CHECK_VISIT(Str, Str);
+#undef CHECK_VISIT
+
+	return true;
+}
+
+IMPLEMENT_SIMPLE_AUTOMATION_TEST(FTVariantVisitorRvalueReferenceParamTest, "System.Core.Misc.TVariant.VisitorRvalueReferenceParam", EAutomationTestFlags::ApplicationContextMask | EAutomationTestFlags::EngineFilter)
+bool FTVariantVisitorRvalueReferenceParamTest::RunTest(const FString&)
+{
+	// Visitor implementation can have rvalue reference parameters
+	struct FAnonOne {};
+	struct FAnonTwo {};
+
+	using FVariant = TVariant<FAnonOne, FAnonTwo>;
+
+	struct FAnonVisitor
+	{
+		void operator()(FAnonOne&&)
+		{
+		}
+
+		void operator()(FAnonTwo&&)
+		{
+		}
+	};
+
+	FVariant MoveThisVariant;
+	Visit(FAnonVisitor{}, MoveTemp(MoveThisVariant));
+
+	// No checks, this is a compilation test.
+	return true;
+}
+
+IMPLEMENT_SIMPLE_AUTOMATION_TEST(FTVariantSerializationTest, "System.Core.Misc.TVariant.Serialization", EAutomationTestFlags::ApplicationContextMask | EAutomationTestFlags::EngineFilter)
+bool FTVariantSerializationTest::RunTest(const FString&)
+{
+	// TVariants can be serialized
+	TArray<uint8> Data;
+	FMemoryWriter Writer(Data);
+
+	TVariant<int, FString> VarData;
+	VarData.Emplace<int>(25);
+
+	Writer << VarData;
+
+	VarData.Emplace<FString>(TEXT("Foo"));
+	Writer << VarData;
+
+
+	FMemoryReader Reader(Data);
+	Reader << VarData;
+	check(VarData.IsType<int>());
+	check(VarData.Get<int>() == 25);
+
+	Reader << VarData;
+	check(VarData.IsType<FString>());
+	check(VarData.Get<FString>() == FString(TEXT("Foo")));
+
+	return true;
+}
+
+#endif //WITH_DEV_AUTOMATION_TESTS
diff --git a/Engine/Source/Runtime/Core/Public/Misc/TVariant.h b/Engine/Source/Runtime/Core/Public/Misc/TVariant.h
--- a/Engine/Source/Runtime/Core/Public/Misc/TVariant.h
+++ b/Engine/Source/Runtime/Core/Public/Misc/TVariant.h
@@ -18,6 +18,12 @@ struct TInPlaceType {};
  */
 struct FEmptyVariantState {};
 
+/** Allow FEmptyVariantState to be used with FArchive serialization */
+inline FArchive& operator<<(FArchive& Ar, FEmptyVariantState&)
+{
+	return Ar;
+}
+
 /**
  * A type-safe union based loosely on std::variant. This flavor of variant requires that all the types in the declaring template parameter pack be unique.
  * Attempting to use the value of a Get() when the underlying type is different leads to undefined behavior.
@@ -236,4 +242,35 @@ decltype(auto) Visit(Func&& Callable, Variants&&... Args)
 		TMakeIntegerSequence<SIZE_T, sizeof...(Variants)>{},
 		Forward<Variants>(Args)...
 	);
-}
\ No newline at end of file
+}
+
+/**
+ * Serialization function for TVariants. 
+ *
+ * In order for a TVariant to be serializable, each type in its template parameter pack must:
+ *   1. Have a default constructor. This is required because when reading the type from an archive, it must be default constructed before being loaded.
+ *   2. Implement the `FArchive& operator<<(FArchive&, T&)` function. This is required to serialize the actual type that's stored in TVariant.
+ */
+template <typename... Ts>
+inline FArchive& operator<<(FArchive& Ar, TVariant<Ts...>& Variant)
+{
+	static_assert(sizeof...(Ts) < 256, "TVariant serialization assumes that the stored index can fit in 8 bits");
+	if (Ar.IsLoading())
+	{
+		uint8 Index;
+		Ar << Index;
+		check(Index < sizeof...(Ts));
+
+		UE::Core::Private::TVariantLoadFromArchiveLookup<Ts...>::Load((SIZE_T)Index, Ar, Variant);
+	}
+	else
+	{
+		uint8 Index = (uint8)Variant.GetIndex();
+		Ar << Index;
+		Visit([&Ar](auto& StoredValue)
+		{
+			Ar << StoredValue;
+		}, Variant);
+	}
+	return Ar;
+}
diff --git a/Engine/Source/Runtime/Core/Public/Misc/TVariantMeta.h b/Engine/Source/Runtime/Core/Public/Misc/TVariantMeta.h
--- a/Engine/Source/Runtime/Core/Public/Misc/TVariantMeta.h
+++ b/Engine/Source/Runtime/Core/Public/Misc/TVariantMeta.h
@@ -8,6 +8,7 @@
 #include "Templates/UnrealTypeTraits.h"
 #include "Delegates/IntegerSequence.h"
 #include "Templates/AndOrNot.h"
+#include "Concepts/Insertable.h"
 
 #include "Misc/AssertionMacros.h"
 
@@ -220,6 +221,35 @@ namespace Private
 		}
 	};
 
+	/** A utility for loading a specific type from FArchive into a TVariant */
+	template <typename T, typename VariantType>
+	struct TVariantLoadFromArchiveCaller
+	{
+		/** Default construct the type and load it from the FArchive */
+		static void Load(FArchive& Ar, VariantType& OutVariant)
+		{
+			OutVariant.template Emplace<T>();
+			Ar << OutVariant.template Get<T>();
+		}
+	};
+
+	/** A utility for loading a type from FArchive based on an index into a template parameter pack. */
+	template <typename... Ts>
+	struct TVariantLoadFromArchiveLookup
+	{
+		using VariantType = TVariant<Ts...>;
+		static_assert((std::is_default_constructible<Ts>::value && ...), "Each type in TVariant template parameter pack must be default constructible in order to use FArchive serialization");
+		static_assert((TModels<CInsertable<FArchive&>, Ts>::Value && ...), "Each type in TVariant template parameter pack must be able to use operator<< with an FArchive");
+
+		/** Load the type at the specified index from the FArchive and emplace it into the TVariant */
+		static void Load(SIZE_T TypeIndex, FArchive& Ar, VariantType& OutVariant)
+		{
+			static constexpr void(*Loaders[])(FArchive&, VariantType&) = { &TVariantLoadFromArchiveCaller<Ts, VariantType>::Load... };
+			check(TypeIndex < UE_ARRAY_COUNT(Loaders));
+			Loaders[TypeIndex](Ar, OutVariant);
+		}
+	};
+
 	/** Determine if the type with the provided index in the template parameter pack is the same */
 	template <typename LookupType, typename... Ts>
 	struct TIsType
