Fix Issue 209713: Static analyzer warnings in TVariantMeta.h

diff --git a/Engine/Source/Runtime/Core/Public/Misc/TVariantMeta.h b/Engine/Source/Runtime/Core/Public/Misc/TVariantMeta.h
--- a/Engine/Source/Runtime/Core/Public/Misc/TVariantMeta.h
+++ b/Engine/Source/Runtime/Core/Public/Misc/TVariantMeta.h
@@ -68,8 +68,8 @@ namespace Private
 	{
 		static constexpr SIZE_T MaxOf(const SIZE_T Sizes[])
 		{
-			SIZE_T MaxSize = Sizes[0];
-			for (int32 Itr = 1; Itr < sizeof...(Ts); ++Itr)
+			SIZE_T MaxSize = 0;
+			for (int32 Itr = 0; Itr < sizeof...(Ts); ++Itr)
 			{
 				if (Sizes[Itr] > MaxSize)
 				{
