Reverted cl 21893507 to fix CIS issues: Serialization function for TVariants and addition of TVariant unit tests

diff --git a/Engine/Source/Runtime/Core/Private/Tests/Misc/TVariantTest.cpp b/Engine/Source/Runtime/Core/Private/Tests/Misc/TVariantTest.cpp
deleted file mode 100644
diff --git a/Engine/Source/Runtime/Core/Public/Misc/TVariant.h b/Engine/Source/Runtime/Core/Public/Misc/TVariant.h
--- a/Engine/Source/Runtime/Core/Public/Misc/TVariant.h
+++ b/Engine/Source/Runtime/Core/Public/Misc/TVariant.h
@@ -18,12 +18,6 @@ struct TInPlaceType {};
  */
 struct FEmptyVariantState {};
 
-/** Allow FEmptyVariantState to be used with FArchive serialization */
-inline FArchive& operator<<(FArchive& Ar, FEmptyVariantState&)
-{
-	return Ar;
-}
-
 /**
  * A type-safe union based loosely on std::variant. This flavor of variant requires that all the types in the declaring template parameter pack be unique.
  * Attempting to use the value of a Get() when the underlying type is different leads to undefined behavior.
@@ -242,35 +236,4 @@ decltype(auto) Visit(Func&& Callable, Variants&&... Args)
 		TMakeIntegerSequence<SIZE_T, sizeof...(Variants)>{},
 		Forward<Variants>(Args)...
 	);
-}
-
-/**
- * Serialization function for TVariants. 
- *
- * In order for a TVariant to be serializable, each type in its template parameter pack must:
- *   1. Have a default constructor. This is required because when reading the type from an archive, it must be default constructed before being loaded.
- *   2. Implement the `FArchive& operator<<(FArchive&, T&)` function. This is required to serialize the actual type that's stored in TVariant.
- */
-template <typename... Ts>
-inline FArchive& operator<<(FArchive& Ar, TVariant<Ts...>& Variant)
-{
-	static_assert(sizeof...(Ts) < 256, "TVariant serialization assumes that the stored index can fit in 8 bits");
-	if (Ar.IsLoading())
-	{
-		uint8 Index;
-		Ar << Index;
-		check(Index < sizeof...(Ts));
-
-		UE::Core::Private::TVariantLoadFromArchiveLookup<Ts...>::Load((SIZE_T)Index, Ar, Variant);
-	}
-	else
-	{
-		uint8 Index = (uint8)Variant.GetIndex();
-		Ar << Index;
-		Visit([&Ar](auto& StoredValue)
-		{
-			Ar << StoredValue;
-		}, Variant);
-	}
-	return Ar;
-}
+}
\ No newline at end of file
diff --git a/Engine/Source/Runtime/Core/Public/Misc/TVariantMeta.h b/Engine/Source/Runtime/Core/Public/Misc/TVariantMeta.h
--- a/Engine/Source/Runtime/Core/Public/Misc/TVariantMeta.h
+++ b/Engine/Source/Runtime/Core/Public/Misc/TVariantMeta.h
@@ -8,7 +8,6 @@
 #include "Templates/UnrealTypeTraits.h"
 #include "Delegates/IntegerSequence.h"
 #include "Templates/AndOrNot.h"
-#include "Concepts/Insertable.h"
 
 #include "Misc/AssertionMacros.h"
 
@@ -68,8 +67,8 @@ namespace Private
 	{
 		static constexpr SIZE_T MaxOf(const SIZE_T Sizes[])
 		{
-			SIZE_T MaxSize = 0;
-			for (int32 Itr = 0; Itr < sizeof...(Ts); ++Itr)
+			SIZE_T MaxSize = Sizes[0];
+			for (int32 Itr = 1; Itr < sizeof...(Ts); ++Itr)
 			{
 				if (Sizes[Itr] > MaxSize)
 				{
@@ -221,35 +220,6 @@ namespace Private
 		}
 	};
 
-	/** A utility for loading a specific type from FArchive into a TVariant */
-	template <typename T, typename VariantType>
-	struct TVariantLoadFromArchiveCaller
-	{
-		/** Default construct the type and load it from the FArchive */
-		static void Load(FArchive& Ar, VariantType& OutVariant)
-		{
-			OutVariant.template Emplace<T>();
-			Ar << OutVariant.template Get<T>();
-		}
-	};
-
-	/** A utility for loading a type from FArchive based on an index into a template parameter pack. */
-	template <typename... Ts>
-	struct TVariantLoadFromArchiveLookup
-	{
-		using VariantType = TVariant<Ts...>;
-		static_assert((std::is_default_constructible<Ts>::value && ...), "Each type in TVariant template parameter pack must be default constructible in order to use FArchive serialization");
-		static_assert((TModels<CInsertable<FArchive&>, Ts>::Value && ...), "Each type in TVariant template parameter pack must be able to use operator<< with an FArchive");
-
-		/** Load the type at the specified index from the FArchive and emplace it into the TVariant */
-		static void Load(SIZE_T TypeIndex, FArchive& Ar, VariantType& OutVariant)
-		{
-			static constexpr void(*Loaders[])(FArchive&, VariantType&) = { &TVariantLoadFromArchiveCaller<Ts, VariantType>::Load... };
-			check(TypeIndex < UE_ARRAY_COUNT(Loaders));
-			Loaders[TypeIndex](Ar, OutVariant);
-		}
-	};
-
 	/** Determine if the type with the provided index in the template parameter pack is the same */
 	template <typename LookupType, typename... Ts>
 	struct TIsType
