Add InsertPlatformFile to FPlatformFileManager to support adding an IPlatformFile after initial startup (for example, from a plugin)

diff --git a/Engine/Source/Runtime/Core/Private/HAL/PlatformFileManager.cpp b/Engine/Source/Runtime/Core/Private/HAL/PlatformFileManager.cpp
--- a/Engine/Source/Runtime/Core/Private/HAL/PlatformFileManager.cpp
+++ b/Engine/Source/Runtime/Core/Private/HAL/PlatformFileManager.cpp
@@ -115,6 +115,39 @@ void FPlatformFileManager::RemovePlatformFile(IPlatformFile* PlatformFileToRemov
 	}
 }
 
+bool FPlatformFileManager::InsertPlatformFile(IPlatformFile* NewPlatformFile)
+{
+	check(TopmostPlatformFile != nullptr);
+	check(NewPlatformFile != nullptr);
+
+	if (FindPlatformFile(NewPlatformFile->GetName()))
+	{
+		return false;
+	}
+
+	if (NewPlatformFile->GetLowerLevel() == nullptr)
+	{
+		return false; // Physical layer must be at the bottom
+	}
+
+	if (NewPlatformFile->GetLowerLevel() == TopmostPlatformFile)
+	{
+		SetPlatformFile(*NewPlatformFile);
+		return true;
+	}
+
+	for (IPlatformFile* ChainElement = TopmostPlatformFile; ChainElement; ChainElement = ChainElement->GetLowerLevel())
+	{
+		if (ChainElement->GetLowerLevel() == NewPlatformFile->GetLowerLevel())
+		{
+			ChainElement->SetLowerLevel(NewPlatformFile);
+			return true;
+		}
+	}
+
+	return false;
+}
+
 void FPlatformFileManager::InitializeNewAsyncIO()
 {
 	// Removed the cached file wrapper because it doesn't work well with EDL
diff --git a/Engine/Source/Runtime/Core/Public/HAL/PlatformFileManager.h b/Engine/Source/Runtime/Core/Public/HAL/PlatformFileManager.h
--- a/Engine/Source/Runtime/Core/Public/HAL/PlatformFileManager.h
+++ b/Engine/Source/Runtime/Core/Public/HAL/PlatformFileManager.h
@@ -79,9 +79,19 @@ public:
 	* THIS IS EXTREMELY DANGEROUS AFTER THE ENGINE HAS BEEN INITIALIZED AS WE MAY BE MODIFYING
 	* THE WRAPPER CHAIN WHILE THINGS ARE BEING LOADED
 	*
-	* @param Name of the platform file to create.
-	* @return Platform file instance of the platform file type was found, nullptr otherwise.
+	* @param The platform file to remove.
 	*/
 	void RemovePlatformFile(IPlatformFile* PlatformFileToRemove);
 
+	/**
+	* Inserts a new platform file into the platform file wrapper chain.
+	* The file is inserted before NewPlatformFile->GetLowerLevel().
+	*
+	* THIS IS EXTREMELY DANGEROUS AFTER THE ENGINE HAS BEEN INITIALIZED AS WE MAY BE MODIFYING
+	* THE WRAPPER CHAIN WHILE THINGS ARE BEING LOADED
+	*
+	* @param The platform file to insert.
+	* @return true if the platform file was inserted.
+	*/
+	bool InsertPlatformFile(IPlatformFile* NewPlatformFile);
 };
