added copy/move constructor and assignment operators for MTAccessDetector so it can be used in copyable/moveable types

diff --git a/Engine/Source/Runtime/Core/Public/Misc/MTAccessDetector.h b/Engine/Source/Runtime/Core/Public/Misc/MTAccessDetector.h
--- a/Engine/Source/Runtime/Core/Public/Misc/MTAccessDetector.h
+++ b/Engine/Source/Runtime/Core/Public/Misc/MTAccessDetector.h
@@ -31,7 +31,31 @@ public:
 
 	~FRWAccessDetector()
 	{
-		checkf(AtomicValue == 0, TEXT("Detector cannot be destroyed while other threads has not release all access"))
+		checkf(AtomicValue.load(std::memory_order_relaxed) == 0, TEXT("Detector cannot be destroyed while other threads access it"))
+	}
+
+	FRWAccessDetector(FRWAccessDetector&& Other)
+	{
+		checkf(Other.AtomicValue.load(std::memory_order_relaxed) == 0, TEXT("Detector cannot be \"moved out\" while other threads access it"));
+	}
+
+	FRWAccessDetector& operator=(FRWAccessDetector&& Other)
+	{
+		checkf(AtomicValue.load(std::memory_order_relaxed) == 0, TEXT("Detector cannot be modified while other threads access it"));
+		checkf(Other.AtomicValue.load(std::memory_order_relaxed) == 0, TEXT("Detector cannot be \"moved out\" while other threads access it"));
+		return *this;
+	}
+
+	FRWAccessDetector(const FRWAccessDetector& Other)
+	{
+		checkf(Other.AtomicValue.load(std::memory_order_relaxed) == 0, TEXT("Detector cannot be copied while other threads access it"));
+	}
+
+	FRWAccessDetector& operator=(const FRWAccessDetector& Other)
+	{
+		checkf(AtomicValue.load(std::memory_order_relaxed) == 0, TEXT("Detector cannot be modified while other threads access it"));
+		checkf(Other.AtomicValue.load(std::memory_order_relaxed) == 0, TEXT("Detector cannot be copied while other threads access it"));
+		return *this;
 	}
 
 	/**
