Fixed NEON version of VectorRound

diff --git a/Engine/Source/Runtime/Core/Public/Math/UnrealMathNeon.h b/Engine/Source/Runtime/Core/Public/Math/UnrealMathNeon.h
--- a/Engine/Source/Runtime/Core/Public/Math/UnrealMathNeon.h
+++ b/Engine/Source/Runtime/Core/Public/Math/UnrealMathNeon.h
@@ -2481,7 +2481,7 @@ FORCEINLINE void VectorSinCos(  VectorRegister4Float* RESTRICT VSinAngles, Vecto
 
 	// Quotient = round(A/2pi)
 	VectorRegister4Float Quotient = VectorMultiply(*VAngles, GlobalVectorConstants::OneOverTwoPi);
-	Quotient = vcvtq_f32_s32(vcvtnq_s32_f32(Quotient)); // round to nearest even is the default rounding mode but that's fine here.
+	Quotient = vrndnq_f32(Quotient); // round to nearest even is the default rounding mode but that's fine here.
 
 	// X = A - 2pi * Quotient
 	VectorRegister4Float X = VectorNegateMultiplyAdd(GlobalVectorConstants::TwoPi, Quotient, *VAngles);
@@ -3025,7 +3025,7 @@ FORCEINLINE VectorRegister4Int VectorFloatToInt(const VectorRegister4Double& A)
 #define VectorShiftLeftImm(Vec, ImmAmt)             vshlq_n_s32(Vec, ImmAmt)
 #define VectorShiftRightImmArithmetic(Vec, ImmAmt)  vshrq_n_s32(Vec, ImmAmt)
 #define VectorShiftRightImmLogical(Vec, ImmAmt)     vshrq_n_u32(Vec, ImmAmt)
-#define VectorRound(Vec)							vreinterpretq_f32_s32(vcvtnq_s32_f32(Vec))
+#define VectorRound(Vec)							vrndnq_f32(Vec)
 
 
 inline VectorRegister4i VectorIntExpandLow16To32(VectorRegister4i V) {
