Fixed a missing #include in WindowsPlatformMisc.cpp

Symbols from psapi.h are needed even when FORCE_ANSI_ALLOCATOR=1.

#rnx

diff --git a/Engine/Source/Runtime/Core/Private/Windows/WindowsPlatformMisc.cpp b/Engine/Source/Runtime/Core/Private/Windows/WindowsPlatformMisc.cpp
--- a/Engine/Source/Runtime/Core/Private/Windows/WindowsPlatformMisc.cpp
+++ b/Engine/Source/Runtime/Core/Private/Windows/WindowsPlatformMisc.cpp
@@ -60,12 +60,10 @@ THIRD_PARTY_INCLUDES_END
 
 #include "Modules/ModuleManager.h"
 
-#if !FORCE_ANSI_ALLOCATOR
-	#include "Windows/AllowWindowsPlatformTypes.h"
-		#include <Psapi.h>
-	#include "Windows/HideWindowsPlatformTypes.h"
-	#pragma comment(lib, "psapi.lib")
-#endif
+#include "Windows/AllowWindowsPlatformTypes.h"
+	#include <Psapi.h>
+#include "Windows/HideWindowsPlatformTypes.h"
+#pragma comment(lib, "psapi.lib")
 
 #include <fcntl.h>
 #include <io.h>
