Unshelved from Michael.Sartain pending changelist '21459195':

Implement quick workaround to allow munmaps to happen with Mutex unlocked
We can't define UE_ALLOW_OSMEMORYLOCKFREE on Linux as we'll deadlock with the OsAllocatorCacheLock and Mutex locks
Munmap calls on Linux can take quite a bit of time (few ms in some cases). This should allow these to happen without the Mutex and remove some hitching.

diff --git a/Engine/Source/Runtime/Core/Private/HAL/Allocators/CachedOSPageAllocator.cpp b/Engine/Source/Runtime/Core/Private/HAL/Allocators/CachedOSPageAllocator.cpp
--- a/Engine/Source/Runtime/Core/Private/HAL/Allocators/CachedOSPageAllocator.cpp
+++ b/Engine/Source/Runtime/Core/Private/HAL/Allocators/CachedOSPageAllocator.cpp
@@ -9,7 +9,7 @@
 
 void* FCachedOSPageAllocator::AllocateImpl(SIZE_T Size, uint32 CachedByteLimit, FFreePageBlock* First, FFreePageBlock* Last, uint32& FreedPageBlocksNum, SIZE_T& CachedTotal, FCriticalSection* Mutex)
 {
-	if (!FPlatformMemory::BinnedPlatformHasMemoryPoolForThisSize(Size) && (Size <= CachedByteLimit / 4))
+	if (!IsOSAllocation(Size, CachedByteLimit))
 	{
 		if (First != Last)
 		{
@@ -95,7 +95,7 @@ void* FCachedOSPageAllocator::AllocateImpl(SIZE_T Size, uint32 CachedByteLimit,
 
 void FCachedOSPageAllocator::FreeImpl(void* Ptr, SIZE_T Size, uint32 NumCacheBlocks, uint32 CachedByteLimit, FFreePageBlock* First, uint32& FreedPageBlocksNum, SIZE_T& CachedTotal, FCriticalSection* Mutex, bool ThreadIsTimeCritical)
 {
-	if (FPlatformMemory::BinnedPlatformHasMemoryPoolForThisSize(Size) || Size > CachedByteLimit / 4)
+	if (IsOSAllocation(Size, CachedByteLimit))
 	{
 #if UE_ALLOW_OSMEMORYLOCKFREE
 		FScopeUnlock ScopeUnlock(Mutex);
diff --git a/Engine/Source/Runtime/Core/Private/HAL/Allocators/PooledVirtualMemoryAllocator.cpp b/Engine/Source/Runtime/Core/Private/HAL/Allocators/PooledVirtualMemoryAllocator.cpp
--- a/Engine/Source/Runtime/Core/Private/HAL/Allocators/PooledVirtualMemoryAllocator.cpp
+++ b/Engine/Source/Runtime/Core/Private/HAL/Allocators/PooledVirtualMemoryAllocator.cpp
@@ -41,14 +41,25 @@ FPooledVirtualMemoryAllocator::FPooledVirtualMemoryAllocator()
 	}
 }
 
-void* FPooledVirtualMemoryAllocator::Allocate(SIZE_T Size, uint32 /*AllocationHint = 0*/, FCriticalSection* /*Mutex = nullptr*/)
+void* FPooledVirtualMemoryAllocator::Allocate(SIZE_T Size, uint32 AllocationHint, FCriticalSection* Mutex)
 {
 	if (Size > Limits::MaxAllocationSizeToPool)
 	{
-		// do not report to LLM here, the platform functions will do that
-		FScopeLock Lock(&OsAllocatorCacheLock);
-		void* Ptr = OsAllocatorCache.Allocate(Size);
-		return Ptr;
+		if (OsAllocatorCache.IsOSAllocation(Size))
+		{
+			// do not report to LLM here, the platform functions will do that
+#if UE_ALLOW_OSMEMORYLOCKFREE
+			FScopeUnlock FScopeUnlock(Mutex);
+#endif
+			return FPlatformMemory::BinnedAllocFromOS(Size);
+		}
+		else
+		{
+			// do not report to LLM here, the platform functions will do that
+			FScopeLock Lock(&OsAllocatorCacheLock);
+			void* Ptr = OsAllocatorCache.Allocate(Size);
+			return Ptr;
+		}
 	}
 	else
 	{
@@ -100,13 +111,24 @@ void* FPooledVirtualMemoryAllocator::Allocate(SIZE_T Size, uint32 /*AllocationHi
 	}
 };
 
-void FPooledVirtualMemoryAllocator::Free(void* Ptr, SIZE_T Size, FCriticalSection* /*Mutex = nullptr*/, bool /*ThreadIsTimeCritical = false*/)
+void FPooledVirtualMemoryAllocator::Free(void* Ptr, SIZE_T Size, FCriticalSection* Mutex, bool /*ThreadIsTimeCritical = false*/)
 {
 	if (Size > Limits::MaxAllocationSizeToPool)
 	{
-		// do not report to LLM here, the platform functions will do that
-		FScopeLock Lock(&OsAllocatorCacheLock);
-		OsAllocatorCache.Free(Ptr, Size);
+		if (OsAllocatorCache.IsOSAllocation(Size))
+		{
+			// do not report to LLM here, the platform functions will do that
+#if UE_ALLOW_OSMEMORYLOCKFREE
+			FScopeUnlock FScopeUnlock(Mutex);
+#endif
+			FPlatformMemory::BinnedFreeToOS(Ptr, Size);
+		}
+		else
+		{
+			// do not report to LLM here, the platform functions will do that
+			FScopeLock Lock(&OsAllocatorCacheLock);
+			OsAllocatorCache.Free(Ptr, Size);
+		}
 	}
 	else
 	{
@@ -259,4 +281,4 @@ uint64 FPooledVirtualMemoryAllocator::GetCachedFreeTotal()
 
 	return TotalFree;
 }
-#endif
\ No newline at end of file
+#endif
diff --git a/Engine/Source/Runtime/Core/Public/HAL/Allocators/CachedOSPageAllocator.h b/Engine/Source/Runtime/Core/Public/HAL/Allocators/CachedOSPageAllocator.h
--- a/Engine/Source/Runtime/Core/Public/HAL/Allocators/CachedOSPageAllocator.h
+++ b/Engine/Source/Runtime/Core/Public/HAL/Allocators/CachedOSPageAllocator.h
@@ -23,6 +23,11 @@ protected:
 	void* AllocateImpl(SIZE_T Size, uint32 CachedByteLimit, FFreePageBlock* First, FFreePageBlock* Last, uint32& FreedPageBlocksNum, SIZE_T& CachedTotal, FCriticalSection* Mutex);
 	void FreeImpl(void* Ptr, SIZE_T Size, uint32 NumCacheBlocks, uint32 CachedByteLimit, FFreePageBlock* First, uint32& FreedPageBlocksNum, SIZE_T& CachedTotal, FCriticalSection* Mutex, bool ThreadIsTimeCritical);
 	void FreeAllImpl(FFreePageBlock* First, uint32& FreedPageBlocksNum, SIZE_T& CachedTotal, FCriticalSection* Mutex);
+
+	static bool IsOSAllocation(SIZE_T Size, uint32 CachedByteLimit)
+	{
+		return (FPlatformMemory::BinnedPlatformHasMemoryPoolForThisSize(Size) || Size > CachedByteLimit / 4);
+	}
 };
 
 template <uint32 NumCacheBlocks, uint32 CachedByteLimit>
@@ -56,6 +61,11 @@ struct TCachedOSPageAllocator : private FCachedOSPageAllocator
 		return CachedTotal;
 	}
 
+	bool IsOSAllocation(SIZE_T Size)
+	{
+		return FCachedOSPageAllocator::IsOSAllocation(Size, CachedByteLimit);
+	}
+
 private:
 	FFreePageBlock FreedPageBlocks[NumCacheBlocks*2];
 	SIZE_T         CachedTotal;
