Fixed non-unity build in CachedOSPageAllocator

#tests compiled
#rnx

diff --git a/Engine/Source/Runtime/Core/Public/HAL/Allocators/CachedOSPageAllocator.h b/Engine/Source/Runtime/Core/Public/HAL/Allocators/CachedOSPageAllocator.h
--- a/Engine/Source/Runtime/Core/Public/HAL/Allocators/CachedOSPageAllocator.h
+++ b/Engine/Source/Runtime/Core/Public/HAL/Allocators/CachedOSPageAllocator.h
@@ -4,6 +4,7 @@
 
 #include "CoreTypes.h"
 #include "HAL/CriticalSection.h"
+#include "HAL/PlatformMemory.h"
 
 struct FCachedOSPageAllocator
 {
