Word replacement mapping for FName::NameToDisplayString

diff --git a/Engine/Source/Runtime/Core/Private/UObject/UnrealNames.cpp b/Engine/Source/Runtime/Core/Private/UObject/UnrealNames.cpp
--- a/Engine/Source/Runtime/Core/Private/UObject/UnrealNames.cpp
+++ b/Engine/Source/Runtime/Core/Private/UObject/UnrealNames.cpp
@@ -36,6 +36,7 @@
 #include "Hash/CityHash.h"
 #include "Templates/AlignmentTemplates.h"
 #include "Templates/Greater.h"
+#include "Misc/AsciiSet.h"
 
 PRAGMA_DISABLE_UNSAFE_TYPECAST_WARNINGS
 
@@ -2401,6 +2402,9 @@ FString FNameEntrySerialized::GetPlainNameString() const
 	FName statics.
 -----------------------------------------------------------------------------*/
 
+static TArray<FString> NameToDisplayStringExemptions;
+static FRWLock NameToDisplayStringExemptionLock;
+
 int32 FName::GetNameEntryMemorySize()
 {
 	return GetNamePool().NumBlocks() * FNameEntryAllocator::BlockSizeBytes;
@@ -2464,6 +2468,14 @@ FString FName::NameToDisplayString( const FString& InDisplayName, const bool bIs
 	// Copy the characters out so that we can modify the string in place
 	const TArray< TCHAR, FString::AllocatorType >& Chars = InDisplayName.GetCharArray();
 
+	// Characters that can follow the end of a word
+	constexpr FAsciiSet EndOfWord(" .,;:?!)");
+
+	FReadScopeLock Lock(NameToDisplayStringExemptionLock);
+
+	const FStringView DisplayNameStringView = InDisplayName;
+	const int32 DisplayNameLength = InDisplayName.Len();
+
 	// This is used to indicate that we are in a run of uppercase letter and/or digits.  The code attempts to keep
 	// these characters together as breaking them up often looks silly (i.e. "Draw Scale 3 D" as opposed to "Draw Scale 3D"
 	bool bInARun = false;
@@ -2493,6 +2505,38 @@ FString FName::NameToDisplayString( const FString& InDisplayName, const bool bIs
 			}
 		}
 
+		// If this is the start of a word, check if its an exempt word from formatting
+		if (CharIndex == 0 || bWasSpace || bWasOpenParen)
+		{
+			bool bExemptionFound = false;
+
+			for (const FString& Exemption : NameToDisplayStringExemptions)
+			{
+				if (DisplayNameStringView.Mid(CharIndex).StartsWith(Exemption, ESearchCase::CaseSensitive)) // exempt word
+				{
+					bExemptionFound = true;
+
+					OutDisplayName += Exemption;
+
+					const int32 ExceptionLength = Exemption.Len();
+
+					if (CharIndex + ExceptionLength < DisplayNameLength && !EndOfWord.Contains(Chars[CharIndex + ExceptionLength]))
+					{
+						OutDisplayName += TEXT(' ');
+						bWasSpace = true;
+					}
+
+					CharIndex += ExceptionLength - 1;
+					break;
+				}
+			}
+
+			if (bExemptionFound)
+			{
+				continue;
+			}
+		}
+
 		// If the current character is upper case or a digit, and the previous character wasn't, then we need to insert a space if there wasn't one previously
 		// We don't do this for numerical expressions, for example "-1.2" should not be formatted as "- 1. 2"
 		if( (bUpperCase || (bIsDigit && !bWasMinusSign)) && !bInARun && !bWasOpenParen && !bWasNumber)
@@ -2588,6 +2632,20 @@ FString FName::NameToDisplayString( const FString& InDisplayName, const bool bIs
 	return OutDisplayName;
 }
 
+void FName::AddNameToDisplayStringExemption(const FString& InExemption)
+{
+	FWriteScopeLock Lock(NameToDisplayStringExemptionLock);
+
+	NameToDisplayStringExemptions.Add(InExemption);
+}
+
+void FName::RemoveNameToDisplayStringExemption(const FString& InExemption)
+{
+	FWriteScopeLock Lock(NameToDisplayStringExemptionLock);
+
+	NameToDisplayStringExemptions.Remove(InExemption);
+}
+
 const EName* FName::ToEName() const
 {
 	return GetNamePoolPostInit().FindEName(ComparisonIndex);
diff --git a/Engine/Source/Runtime/Core/Public/UObject/NameTypes.h b/Engine/Source/Runtime/Core/Public/UObject/NameTypes.h
--- a/Engine/Source/Runtime/Core/Public/UObject/NameTypes.h
+++ b/Engine/Source/Runtime/Core/Public/UObject/NameTypes.h
@@ -1056,6 +1056,14 @@ PRAGMA_ENABLE_DEPRECATION_WARNINGS
 	 */
 	static FString NameToDisplayString( const FString& InDisplayName, const bool bIsBool );
 
+	/**
+	 * Add/remove an exemption to the formatting applied by NameToDisplayString.
+	 * Example: exempt the compound word "MetaHuman" to ensure its not reformatted
+	 * as "Meta Human".
+	 */
+	static void AddNameToDisplayStringExemption(const FString& InExemption);
+	static void RemoveNameToDisplayStringExemption(const FString& InExemption);
+
 	/** Get the EName that this FName represents or nullptr */
 	const EName* ToEName() const;
 
