- Added a system to remap .ini section and key names. This will convert names on load, so saving out the section will use new names, and should not cause issues (we do not force update because the files are usually checked in, so we warn in logs and in the editor on start to tell folks to update).
- Deprecation warnings will show up in C# log, C++ log and in editor as a load warning
- Next changelist will use the system to make some replacements

diff --git a/Engine/Config/ConfigRedirects.ini b/Engine/Config/ConfigRedirects.ini
new file mode 100644
--- /dev/null
+++ b/Engine/Config/ConfigRedirects.ini
@@ -0,0 +1,27 @@
+;; This file contains a list of remaps/redirects for config file Sections and Keys, so Epic can make changes
+;; to the Engine without breaking licensee .ini files. 
+;;
+;; The config system will remap the names on load, so any saving will use the new names, but with .ini files
+;; often being checked into source control, and ini saving can be differential against sections in other files,
+;; it's safest to have a human fixup the ini files.
+;; 
+;; To that end, any remaps that happen will be Warnings in C# utilites, C++ console commands, and in the Editor Message Log.
+;; 
+;; The main intention of this file is to allow Epic to change section/key names without breaking licensee config files. So,
+;; it's not expected that licensees would need to modify this file - instead, licensees can fix up their config files directly
+;; instead of modifying a config file like this one.
+;; 
+;; This first section contains a list of section names to be remapped. All other sections are the name of a section that contains
+;; one of more keys to remap/redirect, followed up the old and new name of the keys to remap
+
+[SectionNameRemap]
+;OldSection=NewSection
+
+
+;; This will look in sections named [Staging] and remap WhitelistConfigFiles to AllowedConfigFiles, etc
+[Staging]
+;WhitelistConfigFiles=AllowedConfigFiles
+;BlacklistConfigFiles=DisallowedConfigFiles
+;WhitelistDirectories=AllowedDirectories
+;BlacklistLocalizationTargets=DisallowedLocalizationTargets
+
diff --git a/Engine/Plugins/Online/OnlineFramework/Source/Hotfix/Private/OnlineHotfixManager.cpp b/Engine/Plugins/Online/OnlineFramework/Source/Hotfix/Private/OnlineHotfixManager.cpp
--- a/Engine/Plugins/Online/OnlineFramework/Source/Hotfix/Private/OnlineHotfixManager.cpp
+++ b/Engine/Plugins/Online/OnlineFramework/Source/Hotfix/Private/OnlineHotfixManager.cpp
@@ -912,7 +912,7 @@ bool UOnlineHotfixManager::HotfixIniFile(const FString& FileName, const FString&
 	// Store the original file so we can undo this later
 	FConfigFileBackup& BackupFile = BackupIniFile(FileName, ConfigFile);
 	// Merge the string into the config file
-	ConfigFile->CombineFromBuffer(IniData);
+	ConfigFile->CombineFromBuffer(IniData, FileName);
 	TArray<UClass*> Classes;
 	TArray<UObject*> PerObjectConfigObjects;
 	int32 StartIndex = 0;
diff --git a/Engine/Source/Editor/AddContentDialog/Private/ContentSourceProviders/FeaturePack/FeaturePackContentSource.cpp b/Engine/Source/Editor/AddContentDialog/Private/ContentSourceProviders/FeaturePack/FeaturePackContentSource.cpp
--- a/Engine/Source/Editor/AddContentDialog/Private/ContentSourceProviders/FeaturePack/FeaturePackContentSource.cpp
+++ b/Engine/Source/Editor/AddContentDialog/Private/ContentSourceProviders/FeaturePack/FeaturePackContentSource.cpp
@@ -437,7 +437,7 @@ bool FFeaturePackContentSource::GetAdditionalFilesForPack(TArray<FString>& FileL
 bool FFeaturePackContentSource::ExtractListOfAdditionalFiles(const FString& InConfigFileAsString,TArray<FString>& InFileList, bool& bContainsSource)
 {
 	FConfigFile PackConfig;
-	PackConfig.ProcessInputFileContents(InConfigFileAsString);
+	PackConfig.ProcessInputFileContents(InConfigFileAsString, TEXT("Uknown, see FFeaturePackContentSource::ExtractListOfAdditionalFiles"));
 	FConfigSection* AdditionalFilesSection = PackConfig.Find("AdditionalFilesToAdd");
 	
 	bContainsSource = false;
diff --git a/Engine/Source/Editor/UnrealEd/Private/Factories/PackFactory.cpp b/Engine/Source/Editor/UnrealEd/Private/Factories/PackFactory.cpp
--- a/Engine/Source/Editor/UnrealEd/Private/Factories/PackFactory.cpp
+++ b/Engine/Source/Editor/UnrealEd/Private/Factories/PackFactory.cpp
@@ -184,7 +184,7 @@ namespace PackFactoryHelper
 	void ProcessPackConfig(const FString& ConfigString, FPackConfigParameters& ConfigParameters)
 	{
 		FConfigFile PackConfig;
-		PackConfig.ProcessInputFileContents(ConfigString);
+		PackConfig.ProcessInputFileContents(ConfigString, TEXT("Unknown (see PackFactoryHelper::ProcessPackConfig)"));
 
 		// Input Settings
 		static FArrayProperty* ActionMappingsProp = FindFieldChecked<FArrayProperty>(UInputSettings::StaticClass(), UInputSettings::GetActionMappingsPropertyName());
diff --git a/Engine/Source/Programs/UnrealBuildTool/System/ConfigFile.cs b/Engine/Source/Programs/UnrealBuildTool/System/ConfigFile.cs
--- a/Engine/Source/Programs/UnrealBuildTool/System/ConfigFile.cs
+++ b/Engine/Source/Programs/UnrealBuildTool/System/ConfigFile.cs
@@ -139,6 +139,65 @@ namespace UnrealBuildTool
 		/// </summary>
 		Dictionary<string, ConfigFileSection> Sections = new Dictionary<string, ConfigFileSection>(StringComparer.InvariantCultureIgnoreCase);
 
+
+		// Remap of config names/sections
+		static private Dictionary<string, string> SectionNameRemap = new();
+		static private Dictionary<string, Dictionary<string, string>> SectionKeyRemap = new();
+		static private HashSet<string> WarnedKeys = new(StringComparer.InvariantCultureIgnoreCase);
+
+		static private string Remap(Dictionary<string, string>? Remap, string Key, string Context)
+		{
+			if (Remap != null)
+			{
+				string? Remapped;
+				if (Remap.TryGetValue(Key, out Remapped))
+				{
+					if (!WarnedKeys.Contains(Key))
+					{
+						WarnedKeys.Add(Key);
+						Log.Logger.LogWarning($"DEPRECATION: '{Key}', {Context}, has been deprecated. Using '{Remapped}' instead. It is recommended you update your .ini files as soon as possible, and replace {Key} with {Remapped}");
+					}
+					return Remapped;
+				}
+			}
+			return Key;
+		}
+
+
+		/// <summary>
+		/// Static (class) constructor which is called before any ConfigFile instance is created
+		/// </summary>
+		static ConfigFile()
+		{
+			// read the special ConfigRemap.ini file into sections
+			FileReference ConfigRemapFile = FileReference.Combine(UnrealBuildBase.Unreal.EngineDirectory, "Config", "ConfigRedirects.ini");
+			Dictionary<string, ConfigFileSection> Sections = new(StringComparer.InvariantCultureIgnoreCase);
+			ReadIntoSections(ConfigRemapFile, Sections, ConfigLineAction.Set);
+			
+			// walk over the sections, where all but the special SectionNameRemap section is a section of keys to remap in that same section
+			foreach (KeyValuePair<string, ConfigFileSection> Pair in Sections)
+			{
+				// remember a remap for section names
+				if (Pair.Key.Equals("SectionNameRemap", StringComparison.InvariantCultureIgnoreCase))
+				{
+					foreach (ConfigLine Line in Pair.Value.Lines)
+					{
+						SectionNameRemap.Add(Line.Key, Line.Value);
+					}
+				}
+				else
+				{
+					// any other section is rmembered by the section name here, and each key/value pair is a remap for the given section
+					Dictionary<string, string> KeyRemap = new(StringComparer.InvariantCultureIgnoreCase);
+					SectionKeyRemap.Add(Pair.Key, KeyRemap);
+					foreach (ConfigLine Line in Pair.Value.Lines)
+					{
+						KeyRemap.Add(Line.Key, Line.Value);
+					}
+				}
+			}
+		}
+
 		/// <summary>
 		/// Constructs a new, empty config file
 		/// </summary>
@@ -153,11 +212,18 @@ namespace UnrealBuildTool
 		/// <param name="Location">File to read from</param>
 		/// <param name="DefaultAction">The default action to take when encountering arrays without a '+' prefix</param>
 		public ConfigFile(FileReference Location, ConfigLineAction DefaultAction = ConfigLineAction.Set)
+		{
+			ReadIntoSections(Location, Sections, DefaultAction);
+		}
+
+		private static void ReadIntoSections(FileReference Location, Dictionary<string, ConfigFileSection> Sections, ConfigLineAction DefaultAction)
 		{
 			using (StreamReader Reader = new StreamReader(Location.FullName))
 			{
 				ConfigFileSection? CurrentSection = null;
-				for(; ;)
+				Dictionary<string, string>? CurrentRemap = null;
+
+				for (; ;)
 				{
 					string? Line = Reader.ReadLine();
 					if (Line == null)
@@ -210,7 +276,12 @@ namespace UnrealBuildTool
 								if(Line[EndIdx - 1] == ']')
 								{
 									string SectionName = Line.Substring(StartIdx + 1, EndIdx - StartIdx - 2);
-									if(!Sections.TryGetValue(SectionName, out CurrentSection))
+									
+									// lookup remaps
+									SectionName = Remap(SectionNameRemap, SectionName, $"which is a config section in '{Location.FullName}'");
+									SectionKeyRemap.TryGetValue(SectionName, out CurrentRemap);
+
+									if (!Sections.TryGetValue(SectionName, out CurrentSection))
 									{
 										CurrentSection = new ConfigFileSection(SectionName);
 										Sections.Add(SectionName, CurrentSection);
@@ -222,7 +293,7 @@ namespace UnrealBuildTool
 							// Otherwise add it to the current section or add a new one
 							if(CurrentSection != null)
 							{
-								TryAddConfigLine(CurrentSection, Line, StartIdx, EndIdx, DefaultAction);
+								TryAddConfigLine(CurrentSection, CurrentRemap, Location.FullName, Line, StartIdx, EndIdx, DefaultAction);
 								break;
 							}
 
@@ -244,10 +315,15 @@ namespace UnrealBuildTool
 
 			// Break into individual settings of the form [Section]:Key=Value
 			string[] SettingLines = IniText.Split(new char[] { ',' });
+			Dictionary<string, string>? CurrentRemap;
 			foreach (string Setting in SettingLines)
 			{
 				// Locate and break off the section name
 				string SectionName = Setting.Remove(Setting.IndexOf(':')).Trim(new char[] { '[', ']' });
+				// lookup remaps
+				Remap(SectionNameRemap, SectionName, $"which is a config section (found in a string in code - search your .ini files for it, the source file is unknown)");
+				SectionKeyRemap.TryGetValue(SectionName, out CurrentRemap);
+
 				if (SectionName.Length > 0)
 				{
 					ConfigFileSection? CurrentSection = null;
@@ -260,7 +336,7 @@ namespace UnrealBuildTool
 					if (CurrentSection != null)
 					{
 						string IniKeyValue = Setting.Substring(Setting.IndexOf(':') + 1);
-						TryAddConfigLine(CurrentSection, IniKeyValue, 0, IniKeyValue.Length, DefaultAction);
+						TryAddConfigLine(CurrentSection, CurrentRemap, "unknown source file", IniKeyValue, 0, IniKeyValue.Length, DefaultAction);
 					}
 				}
 			}
@@ -270,12 +346,14 @@ namespace UnrealBuildTool
 		/// Try to parse a key/value pair from the given line, and add it to a config section
 		/// </summary>
 		/// <param name="Section">The section to receive the parsed config line</param>
+		/// <param name="KeyRemap">Optional map to change names of keys on read</param> 
+		/// <param name="Filename">Optional map to change names of keys on read</param> 
 		/// <param name="Line">Text to parse</param>
 		/// <param name="StartIdx">Index of the first non-whitespace character in this line</param>
 		/// <param name="EndIdx">Index of the last (exclusive) non-whitespace character in this line</param>
 		/// <param name="DefaultAction">The default action to take if '+' or '-' is not specified on a config line</param>
 		/// <returns>True if the line was parsed correctly, false otherwise</returns>
-		static bool TryAddConfigLine(ConfigFileSection Section, string Line, int StartIdx, int EndIdx, ConfigLineAction DefaultAction)
+		static bool TryAddConfigLine(ConfigFileSection Section, Dictionary<string, string>? KeyRemap, string Filename, string Line, int StartIdx, int EndIdx, ConfigLineAction DefaultAction)
 		{
 			// Find the '=' character separating key and value
 			int EqualsIdx = Line.IndexOf('=', StartIdx, EndIdx - StartIdx);
@@ -343,6 +421,10 @@ namespace UnrealBuildTool
 			// Add it to the config section
 			string Key = Line.Substring(KeyStartIdx, KeyEndIdx - KeyStartIdx);
 			string Value = Line.Substring(ValueStartIdx, ValueEndIdx - ValueStartIdx);
+
+			// remap the key if needed
+			Key = Remap(KeyRemap, Key, $"which is a config key in section [{Section.Name}], in '{Filename}'");
+
 			Section.Lines.Add(new ConfigLine(Action, Key, Value));
 			return true;
 		}
diff --git a/Engine/Source/Runtime/Core/Private/Misc/ConfigCacheIni.cpp b/Engine/Source/Runtime/Core/Private/Misc/ConfigCacheIni.cpp
--- a/Engine/Source/Runtime/Core/Private/Misc/ConfigCacheIni.cpp
+++ b/Engine/Source/Runtime/Core/Private/Misc/ConfigCacheIni.cpp
@@ -26,14 +26,19 @@
 #include "Serialization/LargeMemoryReader.h"
 #include "Async/Async.h"
 #include "Misc/OutputDeviceRedirector.h"
+#include "Logging/MessageLog.h"
 #include <limits>
 
 namespace
 {
 	FString CurrentIniVersionStr = TEXT("CurrentIniVersion");
+
+	TMap<FString, FString> SectionRemap;
+	TMap<FString, TMap<FString, FString>> KeyRemap;
 }
 
 DEFINE_LOG_CATEGORY(LogConfig);
+#define LOCTEXT_NAMESPACE "ConfigCache"
 
 /*-----------------------------------------------------------------------------
 FConfigValue
@@ -614,7 +619,7 @@ bool FConfigFile::Combine(const FString& Filename)
 			Combine(FPaths::GetPath(Filename) / TheLine);
 		}
 
-		CombineFromBuffer(Text);
+		CombineFromBuffer(Text, Filename);
 		return true;
 	}
 	else
@@ -634,14 +639,94 @@ ValueType& FindOrAddHeterogeneous(TMap<KeyType, ValueType>& Map, const AltKeyTyp
 	return Existing ? *Existing : Map.Emplace(KeyType(Key));
 }
 
-void FConfigFile::CombineFromBuffer(const FString& Buffer)
+namespace
+{
+
+// either show an editor warning, or just write to log for non-editor
+void LogOrEditorWarning(const FText& Msg, const FString& PartialKey, const FString& File)
+{
+	if (GIsEditor)
+	{
+		static TSet<FString> AlreadyWarnedKeys;
+		
+		FString AbsPath = IFileManager::Get().ConvertToAbsolutePathForExternalAppForRead(*File);
+		
+		// make sure we haven't warned about this yet
+		FString Key = PartialKey + AbsPath;
+		if (AlreadyWarnedKeys.Contains(Key))
+		{
+			return;
+		}
+		AlreadyWarnedKeys.Add(Key);
+		
+		FMessageLog EditorErrors("EditorErrors");
+		TSharedRef<FTokenizedMessage> Message = EditorErrors.Message(EMessageSeverity::Warning);
+		if (File.EndsWith(TEXT(".ini")))
+		{
+			Message->AddToken(FURLToken::Create(FString::Printf(TEXT("file://%s"), *AbsPath), LOCTEXT("DeprecatedConfig_URLCLick", "Click to open file")));
+		}
+		Message->AddToken(FTextToken::Create(Msg));
+		EditorErrors.Notify();
+	}
+
+	// always spit to log
+	UE_LOG(LogConfig, Warning, TEXT("%s"), *Msg.ToString());
+}
+
+// warn about a section name that's deprecated
+void WarnAboutSectionRemap(const FString& OldValue, const FString& NewValue, const FString& File)
+{
+	FFormatNamedArguments Arguments;
+	Arguments.Add(TEXT("OldValue"), FText::FromString(OldValue));
+	Arguments.Add(TEXT("NewValue"), FText::FromString(NewValue));
+	Arguments.Add(TEXT("File"), FText::FromString(File));
+	FText Msg = FText::Format(LOCTEXT("DeprecatedConfig", "Found a deprecated ini section name in {File}. Search for [{OldValue}] and replace with [{NewValue}]"), Arguments);
+	
+	FString Key = OldValue;
+	if (!IsInGameThread())
+	{
+		AsyncTask( ENamedThreads::GameThread, [Msg, Key, File]() { LogOrEditorWarning(Msg, Key, File); });
+	}
+	else
+	{
+		LogOrEditorWarning(Msg, Key, File);
+	}
+}
+
+// warn about a key that's deprecated
+static void WarnAboutKeyRemap(const FString& OldValue, const FString& NewValue, const FString& Section, const FString& File)
+{
+	FFormatNamedArguments Arguments;
+	Arguments.Add(TEXT("OldValue"), FText::FromString(OldValue));
+	Arguments.Add(TEXT("NewValue"), FText::FromString(NewValue));
+	Arguments.Add(TEXT("Section"), FText::FromString(Section));
+	Arguments.Add(TEXT("File"), FText::FromString(File));
+	FText Msg = FText::Format(LOCTEXT("DeprecatedConfig", "Found a deprecated ini section name in {File}. Search for [{OldValue}] and replace with [{NewValue}]"), Arguments);
+	
+	FString Key = OldValue+Section;
+	if (!IsInGameThread())
+	{
+		AsyncTask( ENamedThreads::GameThread, [Msg, Key, File]() { LogOrEditorWarning(Msg, Key, File); });
+	}
+	else
+	{
+		LogOrEditorWarning(Msg, Key, File);
+	}
+}
+
+}
+
+void FConfigFile::CombineFromBuffer(const FString& Buffer, const FString& FileHint)
 {
 	const TCHAR* Ptr = *Buffer;
 	FConfigSection* CurrentSection = nullptr;
-	TStringBuilder<64> CurrentSectionName;
+	FString CurrentSectionName;
+	TMap<FString, FString>* CurrentKeyRemap = nullptr;
 	TStringBuilder<128> TheLine;
 	FString ProcessedValue;
 	bool Done = false;
+	
+	
 	while( !Done )
 	{
 		// Advance past new line characters
@@ -675,9 +760,24 @@ void FConfigFile::CombineFromBuffer(const FString& Buffer)
 			// If we don't have an existing section by this name, add one
 			CurrentSection = FindOrAddSection( Start );
 			CurrentSectionName = Start;
+			
+			// lookup to see if there is an entry in the SectionName remap
+			const FString* FoundRemap;
+			if ((FoundRemap = SectionRemap.Find(CurrentSectionName)) != nullptr)
+			{
+				// show warning in editor
+				WarnAboutSectionRemap(CurrentSectionName, *FoundRemap, FileHint);
+				
+				CurrentSectionName = *FoundRemap;
+			}
+			// look to see if there is a set of key remaps for this section
+			CurrentKeyRemap = KeyRemap.Find(CurrentSectionName);
 
 			// make sure the CurrentSection has any of the special ArrayOfStructKeys added
-			FixupArrayOfStructKeysForSection(CurrentSection, Start, PerObjectConfigArrayOfStructKeys);
+			if (PerObjectConfigArrayOfStructKeys.Num() > 0)
+			{
+				FixupArrayOfStructKeysForSection(CurrentSection, CurrentSectionName, PerObjectConfigArrayOfStructKeys);
+			}
 		}
 
 		// Otherwise, if we're currently inside a section, and we haven't reached the end of the stream
@@ -726,6 +826,20 @@ void FConfigFile::CombineFromBuffer(const FString& Buffer)
 					Start[FCString::Strlen(Start)-1] = TEXT('\0');
 				}
 
+				const TCHAR* KeyName = Start;
+				// look up for key remap
+				if (CurrentKeyRemap != nullptr)
+				{
+					const FString* FoundRemap;
+					if ((FoundRemap = CurrentKeyRemap->Find(KeyName)) != nullptr)
+					{
+						WarnAboutKeyRemap(KeyName, *FoundRemap, CurrentSectionName, FileHint);
+
+						// the Remap will not ever reallocate, so we can just point right into the FString
+						KeyName = **FoundRemap;
+					}
+				}
+				
 				// Strip leading whitespace from the property value
 				while ( *Value && FChar::IsWhitespace(*Value) )
 				{
@@ -750,7 +864,7 @@ void FConfigFile::CombineFromBuffer(const FString& Buffer)
 					ProcessedValue = Value;
 				}
 
-				const FName Key(Start);
+				const FName Key(KeyName);
 				if (Cmd == '+')
 				{
 					// Add if not already present.
@@ -778,7 +892,7 @@ void FConfigFile::CombineFromBuffer(const FString& Buffer)
 				else if (Cmd == '*')
 				{
 					// track a key to show uniqueness for arrays of structs
-					TMap<FName, FString>& POCKeys = FindOrAddHeterogeneous(PerObjectConfigArrayOfStructKeys, CurrentSectionName.ToView());
+					TMap<FName, FString>& POCKeys = FindOrAddHeterogeneous(PerObjectConfigArrayOfStructKeys, CurrentSectionName);
 					POCKeys.Add(Key, MoveTemp(ProcessedValue));
 				}
 				else
@@ -818,10 +932,12 @@ void FConfigFile::CombineFromBuffer(const FString& Buffer)
  * 
  * @param Contents Contents of the .ini file
  */
-void FConfigFile::ProcessInputFileContents(FStringView Contents)
+void FConfigFile::ProcessInputFileContents(FStringView Contents, const FString& FileHint)
 {
 	const TCHAR* Ptr = Contents.Len() > 0 ? Contents.GetData() : nullptr;
 	FConfigSection* CurrentSection = nullptr;
+	FString CurrentSectionName;
+	TMap<FString, FString>* CurrentKeyRemap = nullptr;
 	TStringBuilder<128> TheLine;
 	bool Done = false;
 	while( !Done && Ptr != nullptr )
@@ -853,8 +969,20 @@ void FConfigFile::ProcessInputFileContents(FStringView Contents)
 			Start++;
 			Start[FCString::Strlen(Start)-1] = TEXT('\0');
 
+			// lookup to see if there is an entry in the SectionName remap
+			CurrentSectionName = Start;
+			const FString* FoundRemap;
+			if ((FoundRemap = SectionRemap.Find(CurrentSectionName)) != nullptr)
+			{
+				WarnAboutSectionRemap(CurrentSectionName, *FoundRemap, FileHint);
+				
+				CurrentSectionName = *FoundRemap;
+			}
+			// look to see if there is a set of key remaps for this section
+			CurrentKeyRemap = KeyRemap.Find(CurrentSectionName);
+
 			// If we don't have an existing section by this name, add one
-			CurrentSection = FindOrAddSection( Start );
+			CurrentSection = FindOrAddSection(CurrentSectionName);
 		}
 
 		// Otherwise, if we're currently inside a section, and we haven't reached the end of the stream
@@ -882,6 +1010,20 @@ void FConfigFile::ProcessInputFileContents(FStringView Contents)
 				while( *Start && FChar::IsWhitespace(Start[FCString::Strlen(Start)-1]) )
 					Start[FCString::Strlen(Start)-1] = TEXT('\0');
 
+				const TCHAR* KeyName = Start;
+				// look up for key remap
+				if (CurrentKeyRemap != nullptr)
+				{
+					const FString* FoundRemap;
+					if ((FoundRemap = CurrentKeyRemap->Find(KeyName)) != nullptr)
+					{
+						WarnAboutKeyRemap(KeyName, *FoundRemap, CurrentSectionName, FileHint);
+						
+						// the Remap will not ever reallocate, so we can just point right into the FString
+						KeyName = **FoundRemap;
+					}
+				}
+
 				// Strip leading whitespace from the property value
 				while ( *Value && FChar::IsWhitespace(*Value) )
 					Value++;
@@ -897,12 +1039,12 @@ void FConfigFile::ProcessInputFileContents(FStringView Contents)
 					FParse::QuotedString(Value, ProcessedValue);
 
 					// Add this pair to the current FConfigSection
-					CurrentSection->Add(Start, *ProcessedValue);
+					CurrentSection->Add(KeyName, *ProcessedValue);
 				}
 				else
 				{
 					// Add this pair to the current FConfigSection
-					CurrentSection->Add(Start, Value);
+					CurrentSection->Add(KeyName, Value);
 				}
 			}
 		}
@@ -930,7 +1072,7 @@ void FConfigFile::Read( const FString& Filename )
 		if (LoadConfigFileWrapper(*FinalFileName, Text, bFoundOverride))
 		{
 			// process the contents of the string
-			ProcessInputFileContents(Text);
+			ProcessInputFileContents(Text, Filename);
 		}
 		else
 		{
@@ -1160,11 +1302,11 @@ void FConfigFile::AddDynamicLayerToHierarchy(const FString& Filename)
 	if (SourceConfigFile)
 	{
 		SourceConfigFile->SourceIniHierarchy.AddDynamicLayer(Filename);
-		SourceConfigFile->CombineFromBuffer(ConfigContent);
+		SourceConfigFile->CombineFromBuffer(ConfigContent, Filename);
 	}
 
 	SourceIniHierarchy.AddDynamicLayer(Filename);
-	CombineFromBuffer(ConfigContent);
+	CombineFromBuffer(ConfigContent, Filename);
 
 	// Disable saving since dynamic layers are only for runtime
 	NoSave = true;
@@ -3560,6 +3702,33 @@ static void LoadRemainingConfigFiles(FConfigContext& Context)
 	}
 }
 
+static void InitializeConfigRemap()
+{
+	// read in the single remap file
+	FConfigFile RemapFile;
+	FConfigContext Context = FConfigContext::ReadSingleIntoLocalFile(RemapFile);
+	Context.Load(*FPaths::Combine(FPaths::EngineDir(), TEXT("Config/ConfigRedirects.ini")));
+	
+	for (const TPair<FString, FConfigSection>& Section : RemapFile)
+	{
+		if (Section.Key == TEXT("SectionNameRemap"))
+		{
+			for (const TPair<FName, FConfigValue>& Line : Section.Value)
+			{
+				SectionRemap.Add(Line.Key.ToString(), Line.Value.GetSavedValue());
+			}
+		}
+		else
+		{
+			TMap<FString, FString>& KeyRemaps = KeyRemap.FindOrAdd(Section.Key);
+			for (const TPair<FName, FConfigValue>& Line : Section.Value)
+			{
+				KeyRemaps.Add(Line.Key.ToString(), Line.Value.GetSavedValue());
+			}
+		}
+	}
+}
+
 void FConfigCacheIni::InitializeConfigSystem()
 {
 	// assign the G***Ini strings for the known ini's
@@ -3568,6 +3737,9 @@ void FConfigCacheIni::InitializeConfigSystem()
 		ENUMERATE_KNOWN_INI_FILES(ASSIGN_GLOBAL_INI_STRING);
 	#undef ASSIGN_GLOBAL_INI_STRING
 
+	
+	InitializeConfigRemap();
+	
 #if WITH_EDITOR
 	AsyncInitializeConfigForPlatforms();
 #endif
@@ -4273,3 +4445,5 @@ const TCHAR* ConvertValueFromHumanFriendlyValue(const TCHAR* Value)
 {
 	return UE::ConfigUtilities::ConvertValueFromHumanFriendlyValue(Value);
 }
+
+#undef LOCTEXT_NAMESPACE
diff --git a/Engine/Source/Runtime/Core/Private/Misc/DataDrivenPlatformInfoRegistry.cpp b/Engine/Source/Runtime/Core/Private/Misc/DataDrivenPlatformInfoRegistry.cpp
--- a/Engine/Source/Runtime/Core/Private/Misc/DataDrivenPlatformInfoRegistry.cpp
+++ b/Engine/Source/Runtime/Core/Private/Misc/DataDrivenPlatformInfoRegistry.cpp
@@ -83,7 +83,7 @@ bool FDataDrivenPlatformInfoRegistry::LoadDataDrivenIniFile(int32 Index, FConfig
 	FString IniContents;
 	if (FFileHelper::LoadFileToString(IniContents, *IniFilenames[Index]))
 	{
-		IniFile.ProcessInputFileContents(IniContents);
+		IniFile.ProcessInputFileContents(IniContents, IniFilenames[Index]);
 
 		// platform extension paths are different (engine/platforms/platform/config, not engine/config/platform)
 		if (IniFilenames[Index].StartsWith(FPaths::EnginePlatformExtensionsDir()))
diff --git a/Engine/Source/Runtime/Core/Private/Misc/RemoteConfigIni.cpp b/Engine/Source/Runtime/Core/Private/Misc/RemoteConfigIni.cpp
--- a/Engine/Source/Runtime/Core/Private/Misc/RemoteConfigIni.cpp
+++ b/Engine/Source/Runtime/Core/Private/Misc/RemoteConfigIni.cpp
@@ -493,11 +493,11 @@ void ProcessIniContents(const TCHAR* FilenameToLoad, const TCHAR* IniFileName, F
 		
 		if (bDoCombine)
 		{
-			Config->CombineFromBuffer(RemoteInfo->Buffer);
+			Config->CombineFromBuffer(RemoteInfo->Buffer, FilenameToLoad);
 		}
 		else
 		{
-			Config->ProcessInputFileContents(RemoteInfo->Buffer);
+			Config->ProcessInputFileContents(RemoteInfo->Buffer, FilenameToLoad);
 		}
 	}
 }
diff --git a/Engine/Source/Runtime/Core/Public/Misc/ConfigCacheIni.h b/Engine/Source/Runtime/Core/Public/Misc/ConfigCacheIni.h
--- a/Engine/Source/Runtime/Core/Public/Misc/ConfigCacheIni.h
+++ b/Engine/Source/Runtime/Core/Public/Misc/ConfigCacheIni.h
@@ -402,7 +402,12 @@ public:
 	bool operator!=( const FConfigFile& Other ) const;
 
 	CORE_API bool Combine( const FString& Filename);
-	CORE_API void CombineFromBuffer(const FString& Buffer);
+	CORE_API void CombineFromBuffer(const FString& Buffer, const FString& FileHint);
+	UE_DEPRECATED(5.1, "Use CombineFromBuffer that takes FileHint")
+	CORE_API void CombineFromBuffer(const FString& Buffer)
+	{
+		CombineFromBuffer(Buffer, TEXT("Unknown file, using deprecated function"));
+	}
 	CORE_API void Read( const FString& Filename );
 
 	/** Whether to write a temp file then move it to it's destination when saving. */
@@ -511,7 +516,12 @@ public:
 	 * @param Filename Name of the .ini file the contents came from
 	 * @param Contents Contents of the .ini file
 	 */
-	CORE_API void ProcessInputFileContents(FStringView Contents);
+	CORE_API void ProcessInputFileContents(FStringView Contents, const FString& FileHint);
+	UE_DEPRECATED(5.1, "Use ProcessInputFileContents that takes FileHint")
+	CORE_API void ProcessInputFileContents(const FString& Buffer)
+	{
+		ProcessInputFileContents(Buffer, TEXT("Unknown file, using deprecated function"));
+	}
 
 
 	/** Adds any properties that exist in InSourceFile that this config file is missing */
@@ -1352,4 +1362,4 @@ CORE_API void DeleteRecordedConfigReadsFromIni();
  * Helper function to deal with "True","False","Yes","No","On","Off"
  */
 UE_DEPRECATED(5.1, "Use UE::ConfigUtilities::ConvertValueFromHumanFriendlyValue")
-CORE_API const TCHAR* ConvertValueFromHumanFriendlyValue(const TCHAR* Value);
\ No newline at end of file
+CORE_API const TCHAR* ConvertValueFromHumanFriendlyValue(const TCHAR* Value);
