FCoreDelegates cleanup: removed delegate macros

diff --git a/Engine/Source/Runtime/ApplicationCore/Public/GenericPlatform/GenericPlatformInputDeviceMapper.h b/Engine/Source/Runtime/ApplicationCore/Public/GenericPlatform/GenericPlatformInputDeviceMapper.h
--- a/Engine/Source/Runtime/ApplicationCore/Public/GenericPlatform/GenericPlatformInputDeviceMapper.h
+++ b/Engine/Source/Runtime/ApplicationCore/Public/GenericPlatform/GenericPlatformInputDeviceMapper.h
@@ -225,7 +225,7 @@ protected:
 	*/
 	virtual void UnbindCoreDelegates();
 
-	/** Callback for when FCoreDelegates::FOnUserLoginChangedEvent is broadcasted */
+	/** Callback for when FCoreDelegates::OnUserLoginChangedEvent is broadcasted */
 	virtual void OnUserLoginChangedEvent(bool bLoggedIn, int32 UserId, int32 UserIndex) = 0;
 	
 	/**
@@ -296,7 +296,7 @@ public:
 protected:
 
 	/** 
-	* Callback for when FCoreDelegates::FOnUserLoginChangedEvent is broadcasted.
+	* Callback for when FCoreDelegates::OnUserLoginChangedEvent is broadcasted.
 	* If the user has logged out, then remap any input devices that the user had to the "Unpaired" user on this platform (if bUnpairInputDevicesWhenLoggingOut is true)
 	*/
 	virtual void OnUserLoginChangedEvent(bool bLoggedIn, int32 RawPlatformUserId, int32 UserIndex) override;
diff --git a/Engine/Source/Runtime/AssetRegistry/Private/AssetRegistry.cpp b/Engine/Source/Runtime/AssetRegistry/Private/AssetRegistry.cpp
--- a/Engine/Source/Runtime/AssetRegistry/Private/AssetRegistry.cpp
+++ b/Engine/Source/Runtime/AssetRegistry/Private/AssetRegistry.cpp
@@ -840,7 +840,7 @@ void UAssetRegistryImpl::InitializeEvents(UE::AssetRegistry::Impl::FInitializeCo
 
 	if (Context.bRedirectorsNeedSubscribe)
 	{
-		FCoreDelegates::FResolvePackageNameDelegate PackageResolveDelegate;
+		TDelegate<bool(const FString&, FString&)> PackageResolveDelegate;
 		PackageResolveDelegate.BindUObject(this, &UAssetRegistryImpl::OnResolveRedirect);
 		FCoreDelegates::PackageNameResolvers.Add(PackageResolveDelegate);
 	}
diff --git a/Engine/Source/Runtime/Core/Private/Android/AndroidPlatformMisc.cpp b/Engine/Source/Runtime/Core/Private/Android/AndroidPlatformMisc.cpp
--- a/Engine/Source/Runtime/Core/Private/Android/AndroidPlatformMisc.cpp
+++ b/Engine/Source/Runtime/Core/Private/Android/AndroidPlatformMisc.cpp
@@ -593,7 +593,7 @@ void FAndroidMisc::PlatformTearDown()
 	StopTraceMarkers();
 #endif
 
-	auto RemoveBinding = [](FCoreDelegates::FApplicationLifetimeDelegate& ApplicationLifetimeDelegate, FDelegateHandle& DelegateBinding)
+	auto RemoveBinding = [](TMulticastDelegate<void()>& ApplicationLifetimeDelegate, FDelegateHandle& DelegateBinding)
 	{
 		if (DelegateBinding.IsValid())
 		{
diff --git a/Engine/Source/Runtime/Core/Private/IO/IoStore.cpp b/Engine/Source/Runtime/Core/Private/IO/IoStore.cpp
--- a/Engine/Source/Runtime/Core/Private/IO/IoStore.cpp
+++ b/Engine/Source/Runtime/Core/Private/IO/IoStore.cpp
@@ -115,7 +115,7 @@ static FRSAKeyHandle GetPublicSigningKey()
 	static bool bInitializedPublicKey = false;
 	if (!bInitializedPublicKey)
 	{
-		FCoreDelegates::FPakSigningKeysDelegate& Delegate = FCoreDelegates::GetPakSigningKeysDelegate();
+		TDelegate<void(TArray<uint8>&, TArray<uint8>&)>& Delegate = FCoreDelegates::GetPakSigningKeysDelegate();
 		if (Delegate.IsBound())
 		{
 			TArray<uint8> Exponent;
diff --git a/Engine/Source/Runtime/Core/Private/Misc/CoreDelegates.cpp b/Engine/Source/Runtime/Core/Private/Misc/CoreDelegates.cpp
--- a/Engine/Source/Runtime/Core/Private/Misc/CoreDelegates.cpp
+++ b/Engine/Source/Runtime/Core/Private/Misc/CoreDelegates.cpp
@@ -9,10 +9,10 @@
 //////////////////////////////////////////////////////////////////////////
 // FCoreDelegates
 
-TArray<FCoreDelegates::FHotFixDelegate> FCoreDelegates::HotFixDelegates;
-TArray<FCoreDelegates::FResolvePackageNameDelegate> FCoreDelegates::PackageNameResolvers;
+TArray<TDelegate<void(void*, int32)>> FCoreDelegates::HotFixDelegates;
+TArray<TDelegate<bool(const FString&, FString&)>> FCoreDelegates::PackageNameResolvers;
 
-FCoreDelegates::FHotFixDelegate& FCoreDelegates::GetHotfixDelegate(EHotfixDelegates::Type HotFix)
+TDelegate<void(void*, int32)>& FCoreDelegates::GetHotfixDelegate(EHotfixDelegates::Type HotFix)
 {
 	if (HotFix >= HotFixDelegates.Num())
 	{
@@ -21,51 +21,51 @@ FCoreDelegates::FHotFixDelegate& FCoreDelegates::GetHotfixDelegate(EHotfixDelega
 	return HotFixDelegates[HotFix];
 }
 
-FCoreDelegates::FOnPreMainInit& FCoreDelegates::GetPreMainInitDelegate()
+TMulticastDelegate<void()>& FCoreDelegates::GetPreMainInitDelegate()
 {
-	static FOnPreMainInit StaticDelegate;
+	static TMulticastDelegate<void()> StaticDelegate;
 	return StaticDelegate;
 }
 
-FCoreDelegates::FOnMountAllPakFiles FCoreDelegates::OnMountAllPakFiles;
-FCoreDelegates::FMountPak FCoreDelegates::MountPak;
-FCoreDelegates::FOnUnmountPak FCoreDelegates::OnUnmountPak;
-FCoreDelegates::FOnOptimizeMemoryUsageForMountedPaks FCoreDelegates::OnOptimizeMemoryUsageForMountedPaks;
+TDelegate<int32(const TArray<FString>&)> FCoreDelegates::OnMountAllPakFiles;
+TDelegate<IPakFile*(const FString&, int32)> FCoreDelegates::MountPak;
+TDelegate<bool(const FString&)> FCoreDelegates::OnUnmountPak;
+TDelegate<void()> FCoreDelegates::OnOptimizeMemoryUsageForMountedPaks;
 
-FCoreDelegates::FOnPakFileMounted2 FCoreDelegates::OnPakFileMounted2;
-FCoreDelegates::FNewFileAddedDelegate FCoreDelegates::NewFileAddedDelegate;
-FCoreDelegates::FNoPakFilesMountedDelegate FCoreDelegates::NoPakFilesMountedDelegate;
-FCoreDelegates::FOnFileOpenedForReadFromPakFile FCoreDelegates::OnFileOpenedForReadFromPakFile;
+TMulticastDelegate<void(const IPakFile&)> FCoreDelegates::OnPakFileMounted2;
+TMulticastDelegate<void(const FString&)> FCoreDelegates::NewFileAddedDelegate;
+TMulticastDelegate<void()> FCoreDelegates::NoPakFilesMountedDelegate;
+TMulticastDelegate<void(const TCHAR*, const TCHAR*)> FCoreDelegates::OnFileOpenedForReadFromPakFile;
 
-FCoreDelegates::FOnUserLoginChangedEvent FCoreDelegates::OnUserLoginChangedEvent;
+TMulticastDelegate<void(bool, int32, int32)> FCoreDelegates::OnUserLoginChangedEvent;
 PRAGMA_DISABLE_DEPRECATION_WARNINGS
-FCoreDelegates::FOnUserControllerConnectionChange FCoreDelegates::OnControllerConnectionChange;
-FCoreDelegates::FOnUserControllerPairingChange FCoreDelegates::OnControllerPairingChange;
+TMulticastDelegate<void(bool, FPlatformUserId, int32)> FCoreDelegates::OnControllerConnectionChange;
+TMulticastDelegate<void(int32, FPlatformUserId, FPlatformUserId)> FCoreDelegates::OnControllerPairingChange;
 PRAGMA_ENABLE_DEPRECATION_WARNINGS
-FCoreDelegates::FOnSafeFrameChangedEvent FCoreDelegates::OnSafeFrameChangedEvent;
-FCoreDelegates::FOnHandleSystemEnsure FCoreDelegates::OnHandleSystemEnsure;
-FCoreDelegates::FOnHandleSystemError FCoreDelegates::OnHandleSystemError;
+TMulticastDelegate<void()> FCoreDelegates::OnSafeFrameChangedEvent;
+TMulticastDelegate<void()> FCoreDelegates::OnHandleSystemEnsure;
+TMulticastDelegate<void()> FCoreDelegates::OnHandleSystemError;
 
-FCoreDelegates::FOnActorLabelChanged FCoreDelegates::OnActorLabelChanged;
+TMulticastDelegate<void(AActor*)> FCoreDelegates::OnActorLabelChanged;
 
-FCoreDelegates::FRegisterMovieStreamerDelegate FCoreDelegates::RegisterMovieStreamerDelegate;
-FCoreDelegates::FUnRegisterMovieStreamerDelegate FCoreDelegates::UnRegisterMovieStreamerDelegate;
+TMulticastDelegate<void(FCoreDelegates::FMovieStreamerPtr)> FCoreDelegates::RegisterMovieStreamerDelegate;
+TMulticastDelegate<void(FCoreDelegates::FMovieStreamerPtr)> FCoreDelegates::UnRegisterMovieStreamerDelegate;
 
-FCoreDelegates::FRegisterEncryptionKeyMulticastDelegate& FCoreDelegates::GetRegisterEncryptionKeyMulticastDelegate()
+TMulticastDelegate<void(const FGuid&, const FAES::FAESKey&)>& FCoreDelegates::GetRegisterEncryptionKeyMulticastDelegate()
 {
-	static FRegisterEncryptionKeyMulticastDelegate RegisterEncryptionKeyDelegate;
+	static TMulticastDelegate<void(const FGuid&, const FAES::FAESKey&)> RegisterEncryptionKeyDelegate;
 	return RegisterEncryptionKeyDelegate;
 }
 
-FCoreDelegates::FPakEncryptionKeyDelegate& FCoreDelegates::GetPakEncryptionKeyDelegate()
+TDelegate<void(uint8[32])>& FCoreDelegates::GetPakEncryptionKeyDelegate()
 {
-	static FPakEncryptionKeyDelegate PakEncryptionKeyDelegate;
+	static TDelegate<void(uint8[32])> PakEncryptionKeyDelegate;
 	return PakEncryptionKeyDelegate;
 }
 
-FCoreDelegates::FPakSigningKeysDelegate& FCoreDelegates::GetPakSigningKeysDelegate()
+TDelegate<void(TArray<uint8>&, TArray<uint8>&)>& FCoreDelegates::GetPakSigningKeysDelegate()
 {
-	static FPakSigningKeysDelegate PakSigningKeysDelegate;
+	static TDelegate<void(TArray<uint8>&, TArray<uint8>&)> PakSigningKeysDelegate;
 	return PakSigningKeysDelegate;
 }
 
@@ -76,7 +76,7 @@ FCoreDelegates::FPakSigningKeysDelegate& FCoreDelegates::GetPakSigningKeysDelega
     FSimpleMulticastDelegate FCoreDelegates::PostSlateModal;
 #endif	//WITH_EDITOR
 #if ALLOW_OTHER_PLATFORM_CONFIG
-	FCoreDelegates::FGatherDeviceProfileCVars FCoreDelegates::GatherDeviceProfileCVars;
+	TDelegate<FCoreDelegates::FCVarKeyValueMap(const FString&)> FCoreDelegates::GatherDeviceProfileCVars;
 #endif
 FSimpleMulticastDelegate FCoreDelegates::OnShutdownAfterError;
 FSimpleMulticastDelegate FCoreDelegates::OnInit;
@@ -87,116 +87,117 @@ FSimpleMulticastDelegate FCoreDelegates::OnFEngineLoopInitComplete;
 FSimpleMulticastDelegate FCoreDelegates::OnExit;
 FSimpleMulticastDelegate FCoreDelegates::OnPreExit;
 FSimpleMulticastDelegate FCoreDelegates::OnEnginePreExit;
-FCoreDelegates::FGatherAdditionalLocResPathsDelegate FCoreDelegates::GatherAdditionalLocResPathsCallback;
+TMulticastDelegate<void(TArray<FString>&)> FCoreDelegates::GatherAdditionalLocResPathsCallback;
 FSimpleMulticastDelegate FCoreDelegates::ColorPickerChanged;
 FSimpleMulticastDelegate FCoreDelegates::OnBeginFrame;
 FSimpleMulticastDelegate FCoreDelegates::OnSamplingInput;
 FSimpleMulticastDelegate FCoreDelegates::OnEndFrame;
 FSimpleMulticastDelegate FCoreDelegates::OnBeginFrameRT;
 FSimpleMulticastDelegate FCoreDelegates::OnEndFrameRT;
-FCoreDelegates::FOnModalMessageBox FCoreDelegates::ModalErrorMessage;
-FCoreDelegates::FOnInviteAccepted FCoreDelegates::OnInviteAccepted;
-FCoreDelegates::FWorldOriginOffset FCoreDelegates::PreWorldOriginOffset;
-FCoreDelegates::FWorldOriginOffset FCoreDelegates::PostWorldOriginOffset;
-FCoreDelegates::FStarvedGameLoop FCoreDelegates::StarvedGameLoop;
-FCoreDelegates::FOnTemperatureChange FCoreDelegates::OnTemperatureChange;
-FCoreDelegates::FOnLowPowerMode FCoreDelegates::OnLowPowerMode;
-
-FCoreDelegates::FApplicationLifetimeDelegate FCoreDelegates::ApplicationWillDeactivateDelegate;
-FCoreDelegates::FApplicationLifetimeDelegate FCoreDelegates::ApplicationHasReactivatedDelegate;
-FCoreDelegates::FApplicationLifetimeDelegate FCoreDelegates::ApplicationWillEnterBackgroundDelegate;
-FCoreDelegates::FApplicationLifetimeDelegate FCoreDelegates::ApplicationHasEnteredForegroundDelegate;
-FCoreDelegates::FApplicationLifetimeDelegate FCoreDelegates::ApplicationWillTerminateDelegate;
-FCoreDelegates::FApplicationLifetimeDelegate FCoreDelegates::ApplicationShouldUnloadResourcesDelegate;
-FCoreDelegates::FBackgroundTickDelegate FCoreDelegates::MobileBackgroundTickDelegate;
-
-FCoreDelegates::FApplicationStartupArgumentsDelegate FCoreDelegates::ApplicationReceivedStartupArgumentsDelegate;
-
-FCoreDelegates::FUserMusicInterruptDelegate FCoreDelegates::UserMusicInterruptDelegate;
-FCoreDelegates::FAudioRouteChangedDelegate FCoreDelegates::AudioRouteChangedDelegate;
-FCoreDelegates::FAudioMuteDelegate FCoreDelegates::AudioMuteDelegate;
-FCoreDelegates::FApplicationRequestAudioState FCoreDelegates::ApplicationRequestAudioState;
-
-FCoreDelegates::FApplicationRegisteredForRemoteNotificationsDelegate FCoreDelegates::ApplicationRegisteredForRemoteNotificationsDelegate;
-FCoreDelegates::FApplicationRegisteredForUserNotificationsDelegate FCoreDelegates::ApplicationRegisteredForUserNotificationsDelegate;
-FCoreDelegates::FApplicationFailedToRegisterForRemoteNotificationsDelegate FCoreDelegates::ApplicationFailedToRegisterForRemoteNotificationsDelegate;
-FCoreDelegates::FApplicationReceivedRemoteNotificationDelegate FCoreDelegates::ApplicationReceivedRemoteNotificationDelegate;
-FCoreDelegates::FApplicationReceivedLocalNotificationDelegate FCoreDelegates::ApplicationReceivedLocalNotificationDelegate;
-
-FCoreDelegates::FApplicationPerformFetchDelegate FCoreDelegates::ApplicationPerformFetchDelegate;
-FCoreDelegates::FApplicationBackgroundSessionEventDelegate FCoreDelegates::ApplicationBackgroundSessionEventDelegate;
-
-FCoreDelegates::FCountPreLoadConfigFileRespondersDelegate FCoreDelegates::CountPreLoadConfigFileRespondersDelegate;
-FCoreDelegates::FPreLoadConfigFileDelegate FCoreDelegates::PreLoadConfigFileDelegate;
-FCoreDelegates::FPreSaveConfigFileDelegate FCoreDelegates::PreSaveConfigFileDelegate;
-FCoreDelegates::FOnFConfigFileCreated FCoreDelegates::OnFConfigCreated;
-FCoreDelegates::FOnFConfigFileDeleted FCoreDelegates::OnFConfigDeleted;
-FCoreDelegates::FOnConfigValueRead FCoreDelegates::OnConfigValueRead;
-FCoreDelegates::FOnConfigSectionRead FCoreDelegates::OnConfigSectionRead;
-FCoreDelegates::FOnConfigSectionRead FCoreDelegates::OnConfigSectionNameRead;
-FCoreDelegates::FOnConfigSectionsChanged FCoreDelegates::OnConfigSectionsChanged;
-FCoreDelegates::FOnApplyCVarFromIni FCoreDelegates::OnApplyCVarFromIni;
-FCoreDelegates::FOnSystemResolutionChanged FCoreDelegates::OnSystemResolutionChanged;
+TDelegate<EAppReturnType::Type(EAppMsgType::Type, const FText&, const FText&)> FCoreDelegates::ModalErrorMessage;
+TMulticastDelegate<void(const FString&, const FString&)> FCoreDelegates::OnInviteAccepted;
+TMulticastDelegate<void(class UWorld*, FIntVector, FIntVector)> FCoreDelegates::PreWorldOriginOffset;
+TMulticastDelegate<void(class UWorld*, FIntVector, FIntVector)> FCoreDelegates::PostWorldOriginOffset;
+TDelegate<void()> FCoreDelegates::StarvedGameLoop;
+TMulticastDelegate<void(FCoreDelegates::ETemperatureSeverity)> FCoreDelegates::OnTemperatureChange;
+TMulticastDelegate<void(bool)> FCoreDelegates::OnLowPowerMode;
+
+TMulticastDelegate<void()> FCoreDelegates::ApplicationWillDeactivateDelegate;
+TMulticastDelegate<void()> FCoreDelegates::ApplicationHasReactivatedDelegate;
+TMulticastDelegate<void()> FCoreDelegates::ApplicationWillEnterBackgroundDelegate;
+TMulticastDelegate<void()> FCoreDelegates::ApplicationHasEnteredForegroundDelegate;
+TMulticastDelegate<void()> FCoreDelegates::ApplicationWillTerminateDelegate;
+TMulticastDelegate<void()> FCoreDelegates::ApplicationShouldUnloadResourcesDelegate;
+TMulticastDelegate<void(float)> FCoreDelegates::MobileBackgroundTickDelegate;
+
+TMulticastDelegate<void(const TArray<FString>&)> FCoreDelegates::ApplicationReceivedStartupArgumentsDelegate;
+
+TMulticastDelegate<void(bool)> FCoreDelegates::UserMusicInterruptDelegate;
+TMulticastDelegate<void(bool)> FCoreDelegates::AudioRouteChangedDelegate;
+TMulticastDelegate<void(bool, int)> FCoreDelegates::AudioMuteDelegate;
+TMulticastDelegate<void()> FCoreDelegates::ApplicationRequestAudioState;
+
+TMulticastDelegate<void(TArray<uint8>)> FCoreDelegates::ApplicationRegisteredForRemoteNotificationsDelegate;
+TMulticastDelegate<void(int)> FCoreDelegates::ApplicationRegisteredForUserNotificationsDelegate;
+TMulticastDelegate<void(FString)> FCoreDelegates::ApplicationFailedToRegisterForRemoteNotificationsDelegate;
+TMulticastDelegate<void(FString, int)> FCoreDelegates::ApplicationReceivedRemoteNotificationDelegate;
+TMulticastDelegate<void(FString, int, int)> FCoreDelegates::ApplicationReceivedLocalNotificationDelegate;
+
+TMulticastDelegate<void()> FCoreDelegates::ApplicationPerformFetchDelegate;
+TMulticastDelegate<void(FString)> FCoreDelegates::ApplicationBackgroundSessionEventDelegate;
+
+TTSMulticastDelegate<void(const TCHAR*, int32&)> FCoreDelegates::CountPreLoadConfigFileRespondersDelegate;
+TTSMulticastDelegate<void(const TCHAR*, FString&)> FCoreDelegates::PreLoadConfigFileDelegate;
+TTSMulticastDelegate<void(const TCHAR*, const FString&, int32&)> FCoreDelegates::PreSaveConfigFileDelegate;
+TTSMulticastDelegate<void(const FConfigFile*)> FCoreDelegates::OnFConfigCreated;
+TTSMulticastDelegate<void(const FConfigFile*)> FCoreDelegates::OnFConfigDeleted;
+TTSMulticastDelegate<void(const TCHAR*, const TCHAR*, const TCHAR*)> FCoreDelegates::OnConfigValueRead;
+TTSMulticastDelegate<void(const TCHAR*, const TCHAR*)> FCoreDelegates::OnConfigSectionRead;
+TTSMulticastDelegate<void(const TCHAR*, const TCHAR*)> FCoreDelegates::OnConfigSectionNameRead;
+TTSMulticastDelegate<void(const FString&, const TSet<FString>&)> FCoreDelegates::OnConfigSectionsChanged;
+
+TMulticastDelegate<void(const TCHAR*, const TCHAR*, uint32, bool)> FCoreDelegates::OnApplyCVarFromIni;
+TMulticastDelegate<void(uint32, uint32)> FCoreDelegates::OnSystemResolutionChanged;
 
 #if WITH_EDITOR
-FCoreDelegates::FOnTargetPlatformChangedSupportedFormats FCoreDelegates::OnTargetPlatformChangedSupportedFormats;
-FCoreDelegates::FOnFeatureLevelDisabled FCoreDelegates::OnFeatureLevelDisabled;
+TMulticastDelegate<void(const ITargetPlatform*)> FCoreDelegates::OnTargetPlatformChangedSupportedFormats;
+TMulticastDelegate<void(int, const FName&)> FCoreDelegates::OnFeatureLevelDisabled;
 #endif 
 
-FCoreDelegates::FStatCheckEnabled FCoreDelegates::StatCheckEnabled;
-FCoreDelegates::FStatEnabled FCoreDelegates::StatEnabled;
-FCoreDelegates::FStatDisabled FCoreDelegates::StatDisabled;
-FCoreDelegates::FStatDisableAll FCoreDelegates::StatDisableAll;
+TMulticastDelegate<void(const TCHAR*, bool&, bool&)> FCoreDelegates::StatCheckEnabled;
+TMulticastDelegate<void(const TCHAR*)> FCoreDelegates::StatEnabled;
+TMulticastDelegate<void(const TCHAR*)> FCoreDelegates::StatDisabled;
+TMulticastDelegate<void(const bool)> FCoreDelegates::StatDisableAll;
 
-FCoreDelegates::FApplicationLicenseChange FCoreDelegates::ApplicationLicenseChange;
-FCoreDelegates::FPlatformChangedLaptopMode FCoreDelegates::PlatformChangedLaptopMode;
+TMulticastDelegate<void()> FCoreDelegates::ApplicationLicenseChange;
+TMulticastDelegate<void(EConvertibleLaptopMode)> FCoreDelegates::PlatformChangedLaptopMode;
 
-FCoreDelegates::FVRHeadsetRecenter FCoreDelegates::VRHeadsetTrackingInitializingAndNeedsHMDToBeTrackedDelegate;
-FCoreDelegates::FVRHeadsetRecenter FCoreDelegates::VRHeadsetTrackingInitializedDelegate;
-FCoreDelegates::FVRHeadsetRecenter FCoreDelegates::VRHeadsetRecenter;
-FCoreDelegates::FVRHeadsetLost FCoreDelegates::VRHeadsetLost;
-FCoreDelegates::FVRHeadsetReconnected FCoreDelegates::VRHeadsetReconnected;
-FCoreDelegates::FVRHeadsetConnectCanceled FCoreDelegates::VRHeadsetConnectCanceled;
-FCoreDelegates::FVRHeadsetPutOnHead FCoreDelegates::VRHeadsetPutOnHead;
-FCoreDelegates::FVRHeadsetRemovedFromHead FCoreDelegates::VRHeadsetRemovedFromHead;
-FCoreDelegates::FVRControllerRecentered FCoreDelegates::VRControllerRecentered;
+TMulticastDelegate<void()> FCoreDelegates::VRHeadsetTrackingInitializingAndNeedsHMDToBeTrackedDelegate;
+TMulticastDelegate<void()> FCoreDelegates::VRHeadsetTrackingInitializedDelegate;
+TMulticastDelegate<void()> FCoreDelegates::VRHeadsetRecenter;
+TMulticastDelegate<void()> FCoreDelegates::VRHeadsetLost;
+TMulticastDelegate<void()> FCoreDelegates::VRHeadsetReconnected;
+TMulticastDelegate<void()> FCoreDelegates::VRHeadsetConnectCanceled;
+TMulticastDelegate<void()> FCoreDelegates::VRHeadsetPutOnHead;
+TMulticastDelegate<void()> FCoreDelegates::VRHeadsetRemovedFromHead;
+TMulticastDelegate<void()> FCoreDelegates::VRControllerRecentered;
 
-FCoreDelegates::FOnUserActivityStringChanged FCoreDelegates::UserActivityStringChanged;
-FCoreDelegates::FOnGameSessionIDChange FCoreDelegates::GameSessionIDChanged;
-FCoreDelegates::FOnGameStateClassChange FCoreDelegates::GameStateClassChanged;
-FCoreDelegates::FOnCrashOverrideParamsChanged FCoreDelegates::CrashOverrideParamsChanged;
-FCoreDelegates::FOnIsVanillaProductChanged FCoreDelegates::IsVanillaProductChanged;
+TMulticastDelegate<void(const FString&)> FCoreDelegates::UserActivityStringChanged;
+TMulticastDelegate<void(const FString&)> FCoreDelegates::GameSessionIDChanged;
+TMulticastDelegate<void(const FString&)> FCoreDelegates::GameStateClassChanged;
+TMulticastDelegate<void(const FCrashOverrideParameters&)> FCoreDelegates::CrashOverrideParamsChanged;
+TMulticastDelegate<void(bool)> FCoreDelegates::IsVanillaProductChanged;
 
-FCoreDelegates::FOnAsyncLoadingFlush FCoreDelegates::OnAsyncLoadingFlush;
-FCoreDelegates::FOnAsyncLoadingFlushUpdate FCoreDelegates::OnAsyncLoadingFlushUpdate;
-FCoreDelegates::FOnAsyncLoadPackage FCoreDelegates::OnAsyncLoadPackage;
-FCoreDelegates::FOnSyncLoadPackage FCoreDelegates::OnSyncLoadPackage;
-FCoreDelegates::FRenderingThreadChanged FCoreDelegates::PostRenderingThreadCreated;
-FCoreDelegates::FRenderingThreadChanged FCoreDelegates::PreRenderingThreadDestroyed;
+TMulticastDelegate<void()> FCoreDelegates::OnAsyncLoadingFlush;
+TMulticastDelegate<void()> FCoreDelegates::OnAsyncLoadingFlushUpdate;
+TMulticastDelegate<void(const FString&)> FCoreDelegates::OnAsyncLoadPackage;
+TMulticastDelegate<void(const FString&)> FCoreDelegates::OnSyncLoadPackage;
+TMulticastDelegate<void()> FCoreDelegates::PostRenderingThreadCreated;
+TMulticastDelegate<void()> FCoreDelegates::PreRenderingThreadDestroyed;
 
-FCoreDelegates::FApplicationReceivedOnScreenOrientationChangedNotificationDelegate FCoreDelegates::ApplicationReceivedScreenOrientationChangedNotificationDelegate;
+TMulticastDelegate<void(int32)> FCoreDelegates::ApplicationReceivedScreenOrientationChangedNotificationDelegate;
 
-FCoreDelegates::FConfigReadyForUse FCoreDelegates::ConfigReadyForUse;
+TTSMulticastDelegate<void()> FCoreDelegates::ConfigReadyForUse;
 
-FCoreDelegates::FIsLoadingMovieCurrentlyPlaying FCoreDelegates::IsLoadingMovieCurrentlyPlaying;
+TDelegate<bool()> FCoreDelegates::IsLoadingMovieCurrentlyPlaying;
 
-FCoreDelegates::FShouldLaunchUrl FCoreDelegates::ShouldLaunchUrl;
+TDelegate<bool(const TCHAR*)> FCoreDelegates::ShouldLaunchUrl;
 
-FCoreDelegates::FOnActivatedByProtocol FCoreDelegates::OnActivatedByProtocol;
+TMulticastDelegate<void(const FString&, FPlatformUserId)> FCoreDelegates::OnActivatedByProtocol;
 
-FCoreDelegates::FOnGCFinishDestroyTimeExtended FCoreDelegates::OnGCFinishDestroyTimeExtended;
+TMulticastDelegate<void(const FString&)> FCoreDelegates::OnGCFinishDestroyTimeExtended;
 
-FCoreDelegates::FAccesExtraBinaryConfigData FCoreDelegates::AccessExtraBinaryConfigData;
+TTSMulticastDelegate<void(FCoreDelegates::FExtraBinaryConfigData&)> FCoreDelegates::AccessExtraBinaryConfigData;
 
-FCoreDelegates::FPreloadPackageShaderMaps FCoreDelegates::PreloadPackageShaderMaps;
-FCoreDelegates::FReleasePreloadedPackageShaderMaps FCoreDelegates::ReleasePreloadedPackageShaderMaps;
-FCoreDelegates::FOnLogVerbosityChanged FCoreDelegates::OnLogVerbosityChanged;
+TDelegate<void(TArrayView<const FSHAHash>, FCoreDelegates::FAttachShaderReadRequestFunc)> FCoreDelegates::PreloadPackageShaderMaps;
+TDelegate<void(TArrayView<const FSHAHash>)> FCoreDelegates::ReleasePreloadedPackageShaderMaps;
+TMulticastDelegate<void(const FLogCategoryName&, ELogVerbosity::Type, ELogVerbosity::Type)> FCoreDelegates::OnLogVerbosityChanged;
 
 PRAGMA_DISABLE_DEPRECATION_WARNINGS
-FCoreDelegates::FCreatePackageStore FCoreDelegates::CreatePackageStore;
+TDelegate<TSharedPtr<class IPackageStore>()> FCoreDelegates::CreatePackageStore;
 PRAGMA_ENABLE_DEPRECATION_WARNINGS
 
-FCoreDelegates::FApplicationNetworkInitializationChanged FCoreDelegates::ApplicationNetworkInitializationChanged;
+TMulticastDelegate<void(bool)> FCoreDelegates::ApplicationNetworkInitializationChanged;
 
 PRAGMA_DISABLE_DEPRECATION_WARNINGS
 FCrashOverrideParameters::~FCrashOverrideParameters()
@@ -225,13 +226,13 @@ FSimpleMulticastDelegate& FCoreDelegates::GetOutOfMemoryDelegate()
 }
 
 /**	 Implemented as a function to address global ctor issues */
-FCoreDelegates::FGPUOutOfMemoryDelegate& FCoreDelegates::GetGPUOutOfMemoryDelegate()
+TMulticastDelegate<void(const uint64, const uint64)>& FCoreDelegates::GetGPUOutOfMemoryDelegate()
 {
-	static FGPUOutOfMemoryDelegate OnGPUOOM;
+	static TMulticastDelegate<void(const uint64, const uint64)> OnGPUOOM;
 	return OnGPUOOM;
 }
 
-FCoreDelegates::FGetOnScreenMessagesDelegate FCoreDelegates::OnGetOnScreenMessages;
+TMulticastDelegate<void(FCoreDelegates::FSeverityMessageMap&)> FCoreDelegates::OnGetOnScreenMessages;
 
 typedef void(*TSigningKeyFunc)(TArray<uint8>&, TArray<uint8>&);
 typedef void(*TEncryptionKeyFunc)(unsigned char[32]);
@@ -254,5 +255,5 @@ CORE_API void RegisterEncryptionKeyCallback(TEncryptionKeyFunc InCallback)
 
 FSimpleMulticastDelegate FCoreDelegates::OnParentBeginFork;
 FSimpleMulticastDelegate FCoreDelegates::OnParentPreFork;
-FCoreDelegates::FProcessForkDelegate FCoreDelegates::OnPostFork;
+TMulticastDelegate<void(EForkProcessRole /* ProcessRole */)> FCoreDelegates::OnPostFork;
 FSimpleMulticastDelegate FCoreDelegates::OnChildEndFramePostFork;
\ No newline at end of file
diff --git a/Engine/Source/Runtime/Core/Public/Delegates/DelegateSignatureImpl.inl b/Engine/Source/Runtime/Core/Public/Delegates/DelegateSignatureImpl.inl
--- a/Engine/Source/Runtime/Core/Public/Delegates/DelegateSignatureImpl.inl
+++ b/Engine/Source/Runtime/Core/Public/Delegates/DelegateSignatureImpl.inl
@@ -976,6 +976,9 @@ public:
 	}
 };
 
+template <typename DelegateSignature>
+using TTSMulticastDelegate = TMulticastDelegate<DelegateSignature, FDefaultTSDelegateUserPolicy>;
+
 
 /**
  * Dynamic delegate template class (UObject-based, serializable).  You'll use the various DECLARE_DYNAMIC_DELEGATE
diff --git a/Engine/Source/Runtime/Core/Public/Misc/CoreDelegates.h b/Engine/Source/Runtime/Core/Public/Misc/CoreDelegates.h
--- a/Engine/Source/Runtime/Core/Public/Misc/CoreDelegates.h
+++ b/Engine/Source/Runtime/Core/Public/Misc/CoreDelegates.h
@@ -70,151 +70,79 @@ struct FCrashOverrideParameters
 class CORE_API FCoreDelegates
 {
 public:
-	//hot fix delegate
-	DECLARE_DELEGATE_TwoParams(FHotFixDelegate, void *, int32);
-
-	// Callback for object property modifications
-	DECLARE_MULTICAST_DELEGATE_OneParam(FOnActorLabelChanged, AActor*);
-
-	// delegate type for prompting the pak system to mount all pak files, which haven't already been mounted, from all default locations
-	DECLARE_DELEGATE_RetVal_OneParam(int32, FOnMountAllPakFiles, const TArray<FString>&);
-
-	// delegate type for prompting the pak system to mount a new pak
-	DECLARE_DELEGATE_RetVal_TwoParams(IPakFile*, FMountPak, const FString&, int32);
-
-	// delegate type for prompting the pak system to unmount a pak
-	DECLARE_DELEGATE_RetVal_OneParam(bool, FOnUnmountPak, const FString&);
-
-	// delegate type for prompting the pak system to optimize memory for mounted paks
-	DECLARE_DELEGATE(FOnOptimizeMemoryUsageForMountedPaks);
-
-	// delegate for handling when a new pak file is successfully mounted
-	DECLARE_MULTICAST_DELEGATE_OneParam(FOnPakFileMounted2, const IPakFile&);
-
-	// delegate to let other systems no that no paks were mounted, in case something wants to handle that case
-	DECLARE_MULTICAST_DELEGATE(FNoPakFilesMountedDelegate);
-
-	/** delegate type for opening a modal message box ( Params: EAppMsgType::Type MessageType, const FText& Text, const FText& Title ) */
-	DECLARE_DELEGATE_RetVal_ThreeParams(EAppReturnType::Type, FOnModalMessageBox, EAppMsgType::Type, const FText&, const FText&);
-
-	// Callback for handling an ensure
-	DECLARE_MULTICAST_DELEGATE(FOnHandleSystemEnsure);
-
-	// Callback for handling an error
-	DECLARE_MULTICAST_DELEGATE(FOnHandleSystemError);
-
-	typedef TSharedPtr<class IMovieStreamer, ESPMode::ThreadSafe> FMovieStreamerPtr;
-    // Delegate used to register a movie streamer with any movie player modules that bind to this delegate
-    DECLARE_MULTICAST_DELEGATE_OneParam(FRegisterMovieStreamerDelegate, FMovieStreamerPtr);
-
-    // Delegate used to un-register a movie streamer with any movie player modules that bind to this delegate
-    DECLARE_MULTICAST_DELEGATE_OneParam(FUnRegisterMovieStreamerDelegate, FMovieStreamerPtr);
-
-	// Callback for handling user login/logout.  first int is UserID, second int is UserIndex
-	DECLARE_MULTICAST_DELEGATE_ThreeParams(FOnUserLoginChangedEvent, bool, int32, int32);
-
-	// Callback for handling safe frame area size changes
-	DECLARE_MULTICAST_DELEGATE(FOnSafeFrameChangedEvent);
-
-	// Callback for handling accepting invitations - generally for engine code
-	DECLARE_MULTICAST_DELEGATE_TwoParams(FOnInviteAccepted, const FString&, const FString&);
-
-	// Callback for registering a new encryption key
-	DECLARE_MULTICAST_DELEGATE_TwoParams(FRegisterEncryptionKeyMulticastDelegate, const FGuid&, const FAES::FAESKey&);
-
-	// Callback for accessing pak encryption key, if it exists
-	DECLARE_DELEGATE_OneParam(FPakEncryptionKeyDelegate, uint8[32]);
-
-	// Callback for gathering pak signing keys, if they exist
-	DECLARE_DELEGATE_TwoParams(FPakSigningKeysDelegate, TArray<uint8>&, TArray<uint8>&);
-
-	// Callback for handling the Controller connection / disconnection
-	// first param is true for a connection, false for a disconnection.
-	// second param is UserID, third is UserIndex / ControllerId.
-	DECLARE_MULTICAST_DELEGATE_ThreeParams(FOnUserControllerConnectionChange, bool, FPlatformUserId, int32);
-
-	// Callback for handling a Controller pairing change
-	// first param is controller index
-	// second param is NewUserPlatformId, third is OldUserPlatformId.
-	DECLARE_MULTICAST_DELEGATE_ThreeParams(FOnUserControllerPairingChange, int32 /*ControllerIndex*/, FPlatformUserId /*NewUserPlatformId*/, FPlatformUserId /*OldUserPlatformId*/);
-
 	// Callback for platform handling when flushing async loads.
-	DECLARE_MULTICAST_DELEGATE(FOnAsyncLoadingFlush);
-	static FOnAsyncLoadingFlush OnAsyncLoadingFlush;
+	static TMulticastDelegate<void()> OnAsyncLoadingFlush;
 
 	// Callback for a game thread interruption point when a async load flushing. Used to updating UI during long loads.
-	DECLARE_MULTICAST_DELEGATE(FOnAsyncLoadingFlushUpdate);
-	static FOnAsyncLoadingFlushUpdate OnAsyncLoadingFlushUpdate;
+	static TMulticastDelegate<void()> OnAsyncLoadingFlushUpdate;
 
 	// Callback on the game thread when an async load is started. This goes off before the packages has finished loading
-	DECLARE_MULTICAST_DELEGATE_OneParam(FOnAsyncLoadPackage, const FString&);
-	static FOnAsyncLoadPackage OnAsyncLoadPackage;
+	static TMulticastDelegate<void(const FString&)> OnAsyncLoadPackage;
 
-	DECLARE_MULTICAST_DELEGATE_OneParam(FOnSyncLoadPackage, const FString&);
-	static FOnSyncLoadPackage OnSyncLoadPackage;
+	static TMulticastDelegate<void(const FString&)> OnSyncLoadPackage;
 
 	// get a hotfix delegate
-	static FHotFixDelegate& GetHotfixDelegate(EHotfixDelegates::Type HotFix);
+	static TDelegate<void(void*, int32)>& GetHotfixDelegate(EHotfixDelegates::Type HotFix);
 
 	// Callback when a user logs in/out of the platform.
-	static FOnUserLoginChangedEvent OnUserLoginChangedEvent;
+	static TMulticastDelegate<void(bool, int32, int32)> OnUserLoginChangedEvent;
 
 	// Callback when controllers disconnected / reconnected
 	UE_DEPRECATED(5.1, "OnControllerConnectionChange, use IPlatformInputDeviceMapper::GetOnInputDeviceConnectionChange() instead")
-	static FOnUserControllerConnectionChange OnControllerConnectionChange;
+	static TMulticastDelegate<void(bool, FPlatformUserId, int32)> OnControllerConnectionChange;
 
 	// Callback when a single controller pairing changes
 	UE_DEPRECATED(5.1, "OnControllerPairingChange, use IPlatformInputDeviceMapper::GetOnInputDevicePairingChange() instead")
-	static FOnUserControllerPairingChange OnControllerPairingChange;
+	static TMulticastDelegate<void(int32 ControllerIndex, FPlatformUserId NewUserPlatformId, FPlatformUserId OldUserPlatformId)> OnControllerPairingChange;
 
 	// Callback when a user changes the safe frame size
-	static FOnSafeFrameChangedEvent OnSafeFrameChangedEvent;
+	static TMulticastDelegate<void()> OnSafeFrameChangedEvent;
 
 	// Callback for mounting all the pak files in default locations
-	static FOnMountAllPakFiles OnMountAllPakFiles;
+	static TDelegate<int32(const TArray<FString>&)> OnMountAllPakFiles;
 
 	// Callback to prompt the pak system to mount a pak file
-	static FMountPak MountPak;
+	static TDelegate<IPakFile*(const FString&, int32)> MountPak;
 
 	// Callback to prompt the pak system to unmount a pak file.
-	static FOnUnmountPak OnUnmountPak;
+	static TDelegate<bool(const FString&)> OnUnmountPak;
 
 	// Callback to optimize memeory for currently mounted paks
-	static FOnOptimizeMemoryUsageForMountedPaks OnOptimizeMemoryUsageForMountedPaks;
+	static TDelegate<void()> OnOptimizeMemoryUsageForMountedPaks;
 
 	// After a pakfile is mounted this is called
-	static FOnPakFileMounted2 OnPakFileMounted2;
+	static TMulticastDelegate<void(const IPakFile&)> OnPakFileMounted2;
 
 	// After a file is added this is called
-	DECLARE_MULTICAST_DELEGATE_OneParam(FNewFileAddedDelegate, const FString&);
-	static FNewFileAddedDelegate NewFileAddedDelegate;
+	static TMulticastDelegate<void(const FString&)> NewFileAddedDelegate;
 
 	// After an attempt to mount all pak files, but none wre found, this is called
-	static FNoPakFilesMountedDelegate NoPakFilesMountedDelegate;
+	static TMulticastDelegate<void()> NoPakFilesMountedDelegate;
 
 	// When a file is opened for read from a pak file
-	DECLARE_MULTICAST_DELEGATE_TwoParams(FOnFileOpenedForReadFromPakFile, const TCHAR* /*PakFile*/, const TCHAR* /*FileName*/);
-	static FOnFileOpenedForReadFromPakFile OnFileOpenedForReadFromPakFile;
+	static TMulticastDelegate<void(const TCHAR* PakFile, const TCHAR* FileName)> OnFileOpenedForReadFromPakFile;
+
+	typedef TSharedPtr<class IMovieStreamer, ESPMode::ThreadSafe> FMovieStreamerPtr;
 
     // Delegate used to register a movie streamer with any movie player modules that bind to this delegate
     // Designed to be called when a platform specific movie streamer plugin starts up so that it doesn't need to implement a register for all movie player plugins
-    static FRegisterMovieStreamerDelegate RegisterMovieStreamerDelegate;
+    static TMulticastDelegate<void(FMovieStreamerPtr)> RegisterMovieStreamerDelegate;
     // Delegate used to un-register a movie streamer with any movie player modules that bind to this delegate
     // Designed to be called when a platform specific movie streamer plugin shuts down so that it doesn't need to implement a register for all movie player plugins
-    static FUnRegisterMovieStreamerDelegate UnRegisterMovieStreamerDelegate;
+    static TMulticastDelegate<void(FMovieStreamerPtr)> UnRegisterMovieStreamerDelegate;
 
 	// Callback when an ensure has occurred
-	static FOnHandleSystemEnsure OnHandleSystemEnsure;
+	static TMulticastDelegate<void()> OnHandleSystemEnsure;
 
 	// Callback when an error (crash) has occurred
-	static FOnHandleSystemError OnHandleSystemError;
+	static TMulticastDelegate<void()> OnHandleSystemError;
 
 	// Called when an actor label is changed
-	static FOnActorLabelChanged OnActorLabelChanged;
+	static TMulticastDelegate<void(AActor*)> OnActorLabelChanged;
 
-	static FRegisterEncryptionKeyMulticastDelegate& GetRegisterEncryptionKeyMulticastDelegate();
-	static FPakEncryptionKeyDelegate& GetPakEncryptionKeyDelegate();
-	static FPakSigningKeysDelegate& GetPakSigningKeysDelegate();
+	static TMulticastDelegate<void(const FGuid&, const FAES::FAESKey&)>& GetRegisterEncryptionKeyMulticastDelegate();
+	static TDelegate<void(uint8[32])>& GetPakEncryptionKeyDelegate();
+	static TDelegate<void(TArray<uint8>&, TArray<uint8>&)>& GetPakSigningKeysDelegate();
 
 	
 
@@ -236,8 +164,7 @@ public:
 #if ALLOW_OTHER_PLATFORM_CONFIG
 	// Called when the CVar (ConsoleManager) needs to retrieve CVars for a deviceprofile for another platform - this dramatically simplifies module dependencies
 	typedef TMap<FName, FString> FCVarKeyValueMap;
-	DECLARE_DELEGATE_RetVal_OneParam(FCVarKeyValueMap, FGatherDeviceProfileCVars, const FString&  /*DeviceProfileName*/);
-	static FGatherDeviceProfileCVars GatherDeviceProfileCVars;
+	static TDelegate<FCVarKeyValueMap(const FString& DeviceProfileName)> GatherDeviceProfileCVars;
 #endif
 
 	// Called when an error occurred.
@@ -268,17 +195,16 @@ public:
 	static FSimpleMulticastDelegate OnEnginePreExit;
 
 	/** Delegate for gathering up additional localization paths that are unknown to the engine core (such as plugins) */
-	DECLARE_MULTICAST_DELEGATE_OneParam(FGatherAdditionalLocResPathsDelegate, TArray<FString>&);
-	static FGatherAdditionalLocResPathsDelegate GatherAdditionalLocResPathsCallback;
+	static TMulticastDelegate<void(TArray<FString>&)> GatherAdditionalLocResPathsCallback;
 
 	/** Color picker color has changed, please refresh as needed*/
 	static FSimpleMulticastDelegate ColorPickerChanged;
 
 	/** requests to open a message box */
-	static FOnModalMessageBox ModalErrorMessage;
+	static TDelegate<EAppReturnType::Type(EAppMsgType::Type, const FText&, const FText&)> ModalErrorMessage;
 
 	/** Called when the user accepts an invitation to the current game */
-	static FOnInviteAccepted OnInviteAccepted;
+	static TMulticastDelegate<void(const FString&, const FString&)> OnInviteAccepted;
 
 	// Called at the beginning of a frame
 	static FSimpleMulticastDelegate OnBeginFrame;
@@ -296,15 +222,13 @@ public:
 	static FSimpleMulticastDelegate OnEndFrameRT;
 
 
-	DECLARE_MULTICAST_DELEGATE_ThreeParams(FWorldOriginOffset, class UWorld*, FIntVector, FIntVector);
 	/** called before world origin shifting */
-	static FWorldOriginOffset PreWorldOriginOffset;
+	static TMulticastDelegate<void(class UWorld*, FIntVector, FIntVector)> PreWorldOriginOffset;
 	/** called after world origin shifting */
-	static FWorldOriginOffset PostWorldOriginOffset;
+	static TMulticastDelegate<void(class UWorld*, FIntVector, FIntVector)> PostWorldOriginOffset;
 
 	/** called when the main loop would otherwise starve. */
-	DECLARE_DELEGATE(FStarvedGameLoop);
-	static FStarvedGameLoop StarvedGameLoop;
+	static TDelegate<void()> StarvedGameLoop;
 
 	// IOS-style temperature updates, allowing game to scale down to let temp drop (to avoid thermal throttling on mobile, for instance) */
 	// There is a parellel enum in ApplicationLifecycleComponent
@@ -318,252 +242,197 @@ public:
 
 		NumSeverities,
 	};
-	DECLARE_MULTICAST_DELEGATE_OneParam(FOnTemperatureChange, ETemperatureSeverity);
-	static FOnTemperatureChange OnTemperatureChange;
+	static TMulticastDelegate<void(ETemperatureSeverity)> OnTemperatureChange;
 
 	/** Called when the OS goes into low power mode */
-	DECLARE_MULTICAST_DELEGATE_OneParam(FOnLowPowerMode, bool);
-	static FOnLowPowerMode OnLowPowerMode;
+	static TMulticastDelegate<void(bool)> OnLowPowerMode;
 
 
+	static TTSMulticastDelegate<void(const TCHAR* IniFilename, int32& ResponderCount)> CountPreLoadConfigFileRespondersDelegate;
+	static TTSMulticastDelegate<void(const TCHAR* IniFilename, FString& LoadedContents)> PreLoadConfigFileDelegate;
+	static TTSMulticastDelegate<void(const TCHAR* IniFilename, const FString& ContentsToSave, int32& SavedCount)> PreSaveConfigFileDelegate;
 
-	DECLARE_TS_MULTICAST_DELEGATE_TwoParams(FCountPreLoadConfigFileRespondersDelegate, const TCHAR* /*IniFilename*/, int32& /*ResponderCount*/);
-	DECLARE_TS_MULTICAST_DELEGATE_TwoParams(FPreLoadConfigFileDelegate, const TCHAR* /*IniFilename*/, FString& /*LoadedContents*/);
-	DECLARE_TS_MULTICAST_DELEGATE_ThreeParams(FPreSaveConfigFileDelegate, const TCHAR* /*IniFilename*/, const FString& /*ContentsToSave*/, int32& /*SavedCount*/);
-	static FCountPreLoadConfigFileRespondersDelegate CountPreLoadConfigFileRespondersDelegate;
-	static FPreLoadConfigFileDelegate PreLoadConfigFileDelegate;
-	static FPreSaveConfigFileDelegate PreSaveConfigFileDelegate;
+	static TTSMulticastDelegate<void(const FConfigFile*)> OnFConfigCreated;
+	static TTSMulticastDelegate<void(const FConfigFile*)> OnFConfigDeleted;
 
-	DECLARE_TS_MULTICAST_DELEGATE_OneParam(FOnFConfigFileCreated, const FConfigFile *);
-	DECLARE_TS_MULTICAST_DELEGATE_OneParam(FOnFConfigFileDeleted, const FConfigFile *);
-	static FOnFConfigFileCreated OnFConfigCreated;
-	static FOnFConfigFileDeleted OnFConfigDeleted;
+	static TTSMulticastDelegate<void(const TCHAR* IniFilename, const TCHAR* SectionName, const TCHAR* Key)> OnConfigValueRead;
 
-	DECLARE_TS_MULTICAST_DELEGATE_ThreeParams(FOnConfigValueRead, const TCHAR* /*IniFilename*/, const TCHAR* /*SectionName*/, const TCHAR* /*Key*/);
-	static FOnConfigValueRead OnConfigValueRead;
+	static TTSMulticastDelegate<void(const TCHAR* IniFilename, const TCHAR* SectionName)> OnConfigSectionRead;
+	static TTSMulticastDelegate<void(const TCHAR* IniFilename, const TCHAR* SectionName)> OnConfigSectionNameRead;
 
-	DECLARE_TS_MULTICAST_DELEGATE_TwoParams(FOnConfigSectionRead, const TCHAR* /*IniFilename*/, const TCHAR* /*SectionName*/);
-	static FOnConfigSectionRead OnConfigSectionRead;
-	static FOnConfigSectionRead OnConfigSectionNameRead;
+	static TTSMulticastDelegate<void(const FString& IniFilename, const TSet<FString>& SectionNames)> OnConfigSectionsChanged;
 
-	DECLARE_TS_MULTICAST_DELEGATE_TwoParams(FOnConfigSectionsChanged, const FString& /*IniFilename*/, const TSet<FString>& /*SectionNames*/);
-	static FOnConfigSectionsChanged OnConfigSectionsChanged;
+	static TMulticastDelegate<void(const TCHAR* SectionName, const TCHAR* IniFilename, uint32 SetBy, bool bAllowCheating)> OnApplyCVarFromIni;
 
-	DECLARE_MULTICAST_DELEGATE_FourParams(FOnApplyCVarFromIni, const TCHAR* /*SectionName*/, const TCHAR* /*IniFilename*/, uint32 /*SetBy*/, bool /*bAllowCheating*/);
-	static FOnApplyCVarFromIni OnApplyCVarFromIni;
-
-	DECLARE_MULTICAST_DELEGATE_TwoParams(FOnSystemResolutionChanged, uint32 /*ResX*/, uint32 /*ResY*/);
-	static FOnSystemResolutionChanged OnSystemResolutionChanged;
+	static TMulticastDelegate<void(uint32 ResX, uint32 ResY)> OnSystemResolutionChanged;
 
 #if WITH_EDITOR
 	// called when a target platform changes it's return value of supported formats.  This is so anything caching those results can reset (like cached shaders for cooking)
-	DECLARE_MULTICAST_DELEGATE_OneParam(FOnTargetPlatformChangedSupportedFormats, const ITargetPlatform*); 
-	static FOnTargetPlatformChangedSupportedFormats OnTargetPlatformChangedSupportedFormats;
+	static TMulticastDelegate<void(const ITargetPlatform*)> OnTargetPlatformChangedSupportedFormats;
 
 	// Called when a feature level is disabled by the user.
-	DECLARE_MULTICAST_DELEGATE_TwoParams(FOnFeatureLevelDisabled, int, const FName&);
-	static FOnFeatureLevelDisabled OnFeatureLevelDisabled;
+	static TMulticastDelegate<void(int, const FName&)> OnFeatureLevelDisabled;
 #endif
 
 	/** IOS-style application lifecycle delegates */
-	DECLARE_MULTICAST_DELEGATE(FApplicationLifetimeDelegate);
 
 	// This is called when the application is about to be deactivated (e.g., due to a phone call or SMS or the sleep button).
 	// The game should be paused if possible, etc...
-	static FApplicationLifetimeDelegate ApplicationWillDeactivateDelegate;
+	static TMulticastDelegate<void()> ApplicationWillDeactivateDelegate;
 
 	// Called when the application has been reactivated (reverse any processing done in the Deactivate delegate)
-	static FApplicationLifetimeDelegate ApplicationHasReactivatedDelegate;
+	static TMulticastDelegate<void()> ApplicationHasReactivatedDelegate;
 
 	// This is called when the application is being backgrounded (e.g., due to switching
 	// to another app or closing it via the home button)
 	// The game should release shared resources, save state, etc..., since it can be
 	// terminated from the background state without any further warning.
-	static FApplicationLifetimeDelegate ApplicationWillEnterBackgroundDelegate; // for instance, hitting the home button
+	static TMulticastDelegate<void()> ApplicationWillEnterBackgroundDelegate; // for instance, hitting the home button
 
 	// Called when the application is returning to the foreground (reverse any processing done in the EnterBackground delegate)
-	static FApplicationLifetimeDelegate ApplicationHasEnteredForegroundDelegate;
+	static TMulticastDelegate<void()> ApplicationHasEnteredForegroundDelegate;
 
 	// This *may* be called when the application is getting terminated by the OS.
 	// There is no guarantee that this will ever be called on a mobile device,
 	// save state when ApplicationWillEnterBackgroundDelegate is called instead.
-	static FApplicationLifetimeDelegate ApplicationWillTerminateDelegate;
+	static TMulticastDelegate<void()> ApplicationWillTerminateDelegate;
 
 	// Called when in the background, if the OS is giving CPU time to the device. It is very likely
 	// this will never be called due to mobile OS backgrounded CPU restrictions. But if, for instance,
 	// VOIP is active on iOS, the will be getting called
-	DECLARE_MULTICAST_DELEGATE_OneParam(FBackgroundTickDelegate, float /*DeltaTime*/);
-	static FBackgroundTickDelegate MobileBackgroundTickDelegate;
+	static TMulticastDelegate<void(float DeltaTime)> MobileBackgroundTickDelegate;
 
 	// Called when the OS needs control of the music (parameter is true) or when the OS returns
 	// control of the music to the application (parameter is false). This can happen due to a
 	// phone call or timer or other OS-level event. This is currently triggered only on iOS
 	// devices.
-	DECLARE_MULTICAST_DELEGATE_OneParam(FUserMusicInterruptDelegate, bool);
-	static FUserMusicInterruptDelegate UserMusicInterruptDelegate;
+	static TMulticastDelegate<void(bool)> UserMusicInterruptDelegate;
 	
 	// [iOS only] Called when the mute switch is detected as changed or when the
 	// volume changes. Parameter 1 is the mute switch state (true is muted, false is
 	// unmuted). Parameter 2 is the volume as an integer from 0 to 100.
-	DECLARE_MULTICAST_DELEGATE_TwoParams(FAudioMuteDelegate, bool, int);
-	static FAudioMuteDelegate AudioMuteDelegate;
+	static TMulticastDelegate<void(bool, int)> AudioMuteDelegate;
 	
 	// [iOS only] Called when the audio device changes
 	// For instance, when the headphones are plugged in or removed
-	DECLARE_MULTICAST_DELEGATE_OneParam(FAudioRouteChangedDelegate, bool);
-	static FAudioRouteChangedDelegate AudioRouteChangedDelegate;
+	static TMulticastDelegate<void(bool)> AudioRouteChangedDelegate;
 
 	// Generally, events triggering UserMusicInterruptDelegate or AudioMuteDelegate happen only
 	// when a change occurs. When a system comes online needing the current audio state but the
 	// event has already been broadcast, calling ApplicationRequestAudioState will force the
 	// UserMusicInterruptDelegate and AudioMuteDelegate to be called again if the low-level
 	// application layer supports it. Currently, this is available only on iOS.
-	DECLARE_MULTICAST_DELEGATE(FApplicationRequestAudioState);
-	static FApplicationRequestAudioState ApplicationRequestAudioState;
+	static TMulticastDelegate<void()> ApplicationRequestAudioState;
 	
 	// Called when the OS is running low on resources and asks the application to free up any cached resources, drop graphics quality etc.
-	static FApplicationLifetimeDelegate ApplicationShouldUnloadResourcesDelegate;
-
-	DECLARE_MULTICAST_DELEGATE_OneParam(FApplicationStartupArgumentsDelegate, const TArray<FString>&);
+	static TMulticastDelegate<void()> ApplicationShouldUnloadResourcesDelegate;
 
 	// Called with arguments passed to the application on statup, perhaps meta data passed on by another application which launched this one.
-	static FApplicationStartupArgumentsDelegate ApplicationReceivedStartupArgumentsDelegate;
+	static TMulticastDelegate<void(const TArray<FString>&)> ApplicationReceivedStartupArgumentsDelegate;
 
 	/** IOS-style push notification delegates */
-	DECLARE_MULTICAST_DELEGATE_OneParam(FApplicationRegisteredForRemoteNotificationsDelegate, TArray<uint8>);
-	DECLARE_MULTICAST_DELEGATE_OneParam(FApplicationRegisteredForUserNotificationsDelegate, int);
-	DECLARE_MULTICAST_DELEGATE_OneParam(FApplicationFailedToRegisterForRemoteNotificationsDelegate, FString);
-	DECLARE_MULTICAST_DELEGATE_TwoParams(FApplicationReceivedRemoteNotificationDelegate, FString, int);
-	DECLARE_MULTICAST_DELEGATE_ThreeParams(FApplicationReceivedLocalNotificationDelegate, FString, int, int);
-    DECLARE_MULTICAST_DELEGATE(FApplicationPerformFetchDelegate);
-    DECLARE_MULTICAST_DELEGATE_OneParam(FApplicationBackgroundSessionEventDelegate, FString);
 
 	// called when the user grants permission to register for remote notifications
-	static FApplicationRegisteredForRemoteNotificationsDelegate ApplicationRegisteredForRemoteNotificationsDelegate;
+	static TMulticastDelegate<void(TArray<uint8>)> ApplicationRegisteredForRemoteNotificationsDelegate;
 
 	// called when the user grants permission to register for notifications
-	static FApplicationRegisteredForUserNotificationsDelegate ApplicationRegisteredForUserNotificationsDelegate;
+	static TMulticastDelegate<void(int)> ApplicationRegisteredForUserNotificationsDelegate;
 
 	// called when the application fails to register for remote notifications
-	static FApplicationFailedToRegisterForRemoteNotificationsDelegate ApplicationFailedToRegisterForRemoteNotificationsDelegate;
+	static TMulticastDelegate<void(FString)> ApplicationFailedToRegisterForRemoteNotificationsDelegate;
 
 	// called when the application receives a remote notification
-	static FApplicationReceivedRemoteNotificationDelegate ApplicationReceivedRemoteNotificationDelegate;
+	static TMulticastDelegate<void(FString, int)> ApplicationReceivedRemoteNotificationDelegate;
 
 	// called when the application receives a local notification
-	static FApplicationReceivedLocalNotificationDelegate ApplicationReceivedLocalNotificationDelegate;
+	static TMulticastDelegate<void(FString, int, int)> ApplicationReceivedLocalNotificationDelegate;
 
     // called when the application receives notice to perform a background fetch
-    static FApplicationPerformFetchDelegate ApplicationPerformFetchDelegate;
+    static TMulticastDelegate<void()> ApplicationPerformFetchDelegate;
 
     // called when the application receives notice that a background download has completed
-    static FApplicationBackgroundSessionEventDelegate ApplicationBackgroundSessionEventDelegate;
+    static TMulticastDelegate<void(FString)> ApplicationBackgroundSessionEventDelegate;
 
 	/** Sent when a device screen orientation changes */
-	DECLARE_MULTICAST_DELEGATE_OneParam(FApplicationReceivedOnScreenOrientationChangedNotificationDelegate, int32);
-	static FApplicationReceivedOnScreenOrientationChangedNotificationDelegate ApplicationReceivedScreenOrientationChangedNotificationDelegate;
+	static TMulticastDelegate<void(int32)> ApplicationReceivedScreenOrientationChangedNotificationDelegate;
 
 	/** Checks to see if the stat is already enabled */
-	DECLARE_MULTICAST_DELEGATE_ThreeParams(FStatCheckEnabled, const TCHAR*, bool&, bool&);
-	static FStatCheckEnabled StatCheckEnabled;
+	static TMulticastDelegate<void(const TCHAR*, bool&, bool&)> StatCheckEnabled;
 
 	/** Sent after each stat is enabled */
-	DECLARE_MULTICAST_DELEGATE_OneParam(FStatEnabled, const TCHAR*);
-	static FStatEnabled StatEnabled;
+	static TMulticastDelegate<void(const TCHAR*)> StatEnabled;
 
 	/** Sent after each stat is disabled */
-	DECLARE_MULTICAST_DELEGATE_OneParam(FStatDisabled, const TCHAR*);
-	static FStatDisabled StatDisabled;
+	static TMulticastDelegate<void(const TCHAR*)> StatDisabled;
 
 	/** Sent when all stats need to be disabled */
-	DECLARE_MULTICAST_DELEGATE_OneParam(FStatDisableAll, const bool);
-	static FStatDisableAll StatDisableAll;
+	static TMulticastDelegate<void(const bool)> StatDisableAll;
 
 	// Called when an application is notified that the application license info has been updated.
 	// The new license data should be polled and steps taken based on the results (i.e. halt application if license is no longer valid).
-	DECLARE_MULTICAST_DELEGATE(FApplicationLicenseChange);
-	static FApplicationLicenseChange ApplicationLicenseChange;
+	static TMulticastDelegate<void()> ApplicationLicenseChange;
 
 	/** Sent when the platform changed its laptop mode (for convertible laptops).*/
-	DECLARE_MULTICAST_DELEGATE_OneParam(FPlatformChangedLaptopMode, EConvertibleLaptopMode);
-	static FPlatformChangedLaptopMode PlatformChangedLaptopMode;
+	static TMulticastDelegate<void(EConvertibleLaptopMode)> PlatformChangedLaptopMode;
 
 	/** Sent when the platform needs the user to fix headset tracking on startup (Most platforms do not need this.) */
-	DECLARE_MULTICAST_DELEGATE(FVRHeadsetTrackingInitializingAndNeedsHMDToBeTrackedDelegate);
-	static FVRHeadsetTrackingInitializingAndNeedsHMDToBeTrackedDelegate VRHeadsetTrackingInitializingAndNeedsHMDToBeTrackedDelegate;
+	static TMulticastDelegate<void()> VRHeadsetTrackingInitializingAndNeedsHMDToBeTrackedDelegate;
 
 	/** Sent when the platform finds that needed headset tracking on startup has completed (Most platforms do not need this.) */
-	DECLARE_MULTICAST_DELEGATE(FVRHeadsetTrackingInitializedDelegate);
-	static FVRHeadsetTrackingInitializedDelegate VRHeadsetTrackingInitializedDelegate;
+	static TMulticastDelegate<void()> VRHeadsetTrackingInitializedDelegate;
 
 	/** Sent when the platform requests a low-level VR recentering */
-	DECLARE_MULTICAST_DELEGATE(FVRHeadsetRecenter);
-	static FVRHeadsetRecenter VRHeadsetRecenter;
+	static TMulticastDelegate<void()> VRHeadsetRecenter;
 
 	/** Sent when connection to VR HMD is lost */
-	DECLARE_MULTICAST_DELEGATE(FVRHeadsetLost);
-	static FVRHeadsetLost VRHeadsetLost;
+	static TMulticastDelegate<void()> VRHeadsetLost;
 
 	/** Sent when connection to VR HMD is restored */
-	DECLARE_MULTICAST_DELEGATE(FVRHeadsetReconnected);
-	static FVRHeadsetReconnected VRHeadsetReconnected;
+	static TMulticastDelegate<void()> VRHeadsetReconnected;
 
 	/** Sent when connection to VR HMD connection is refused by the player */
-	DECLARE_MULTICAST_DELEGATE(FVRHeadsetConnectCanceled);
-	static FVRHeadsetConnectCanceled VRHeadsetConnectCanceled;
+	static TMulticastDelegate<void()> VRHeadsetConnectCanceled;
 
 	/** Sent when the VR HMD detects that it has been put on by the player. */
-	DECLARE_MULTICAST_DELEGATE(FVRHeadsetPutOnHead);
-	static FVRHeadsetPutOnHead VRHeadsetPutOnHead;
+	static TMulticastDelegate<void()> VRHeadsetPutOnHead;
 
 	/** Sent when the VR HMD detects that it has been taken off by the player. */
-	DECLARE_MULTICAST_DELEGATE(FVRHeadsetRemovedFromHead);
-	static FVRHeadsetRemovedFromHead VRHeadsetRemovedFromHead;
+	static TMulticastDelegate<void()> VRHeadsetRemovedFromHead;
 
 	/** Sent when a 3DOF VR controller is recentered */
-	DECLARE_MULTICAST_DELEGATE(FVRControllerRecentered);
-	static FVRControllerRecentered VRControllerRecentered;
+	static TMulticastDelegate<void()> VRControllerRecentered;
 
 	/** Sent when application code changes the user activity hint string for analytics, crash reports, etc */
-	DECLARE_MULTICAST_DELEGATE_OneParam(FOnUserActivityStringChanged, const FString&);
-	static FOnUserActivityStringChanged UserActivityStringChanged;
+	static TMulticastDelegate<void(const FString&)> UserActivityStringChanged;
 
 	/** Sent when application code changes the currently active game session. The exact semantics of this will vary between games but it is useful for analytics, crash reports, etc  */
-	DECLARE_MULTICAST_DELEGATE_OneParam(FOnGameSessionIDChange, const FString&);
-	static FOnGameSessionIDChange GameSessionIDChanged;
+	static TMulticastDelegate<void(const FString&)> GameSessionIDChanged;
 
 	/** Sent when application code changes game state. The exact semantics of this will vary between games but it is useful for analytics, crash reports, etc  */
-	DECLARE_MULTICAST_DELEGATE_OneParam(FOnGameStateClassChange, const FString&);
-	static FOnGameStateClassChange GameStateClassChanged;
+	static TMulticastDelegate<void(const FString&)> GameStateClassChanged;
 
 	/** Sent by application code to set params that customize crash reporting behavior. */
-	DECLARE_MULTICAST_DELEGATE_OneParam(FOnCrashOverrideParamsChanged, const FCrashOverrideParameters&);
-	static FOnCrashOverrideParamsChanged CrashOverrideParamsChanged;
+	static TMulticastDelegate<void(const FCrashOverrideParameters&)> CrashOverrideParamsChanged;
 	
 	/** Sent by engine code when the "vanilla" status of the engine changes */
-	DECLARE_MULTICAST_DELEGATE_OneParam(FOnIsVanillaProductChanged, bool);
-	static FOnIsVanillaProductChanged IsVanillaProductChanged;
+	static TMulticastDelegate<void(bool)> IsVanillaProductChanged;
 
 	// Callback for platform specific very early init code.
-	DECLARE_MULTICAST_DELEGATE(FOnPreMainInit);
-	static FOnPreMainInit& GetPreMainInitDelegate();
+	static TMulticastDelegate<void()>& GetPreMainInitDelegate();
 	
 	/** Sent when GConfig is finished initializing */
-	DECLARE_TS_MULTICAST_DELEGATE(FConfigReadyForUse);
-	static FConfigReadyForUse ConfigReadyForUse;
+	static TTSMulticastDelegate<void()> ConfigReadyForUse;
 
 	/** Callback for notifications regarding changes of the rendering thread. */
-	DECLARE_MULTICAST_DELEGATE(FRenderingThreadChanged)
 
 	/** Sent just after the rendering thread has been created. */
-	static FRenderingThreadChanged PostRenderingThreadCreated;
+	static TMulticastDelegate<void()> PostRenderingThreadCreated;
 	/* Sent just before the rendering thread is destroyed. */
-	static FRenderingThreadChanged PreRenderingThreadDestroyed;
+	static TMulticastDelegate<void()> PreRenderingThreadDestroyed;
 
 	// Callback to allow custom resolution of package names. Arguments are InRequestedName, OutResolvedName.
 	// Should return True of resolution occured.
-	DECLARE_DELEGATE_RetVal_TwoParams(bool, FResolvePackageNameDelegate, const FString&, FString&);
-	static TArray<FResolvePackageNameDelegate> PackageNameResolvers;
+	static TArray<TDelegate<bool(const FString&, FString&)>> PackageNameResolvers;
 
 	// Called to request that systems free whatever memory they are able to. Called early in LoadMap.
 	// Caller is responsible for flushing rendering etc. See UEngine::TrimMemory
@@ -576,8 +445,7 @@ public:
 	static FSimpleMulticastDelegate& GetOutOfMemoryDelegate();
 
 	// Called from TerminateOnOutOfMemory in D3D11Util.cpp/D3D12Util.cpp
-	DECLARE_MULTICAST_DELEGATE_TwoParams(FGPUOutOfMemoryDelegate, const uint64, const uint64);
-	static FGPUOutOfMemoryDelegate& GetGPUOutOfMemoryDelegate();
+	static TMulticastDelegate<void(const uint64, const uint64)>& GetGPUOutOfMemoryDelegate();
 
 	enum class EOnScreenMessageSeverity : uint8
 	{
@@ -593,28 +461,22 @@ public:
 	// {
 	//		OutMessages.Add(FCoreDelegates::EOnScreenMessageSeverity::Info, FText::Format(LOCTEXT("MyMessage", "My Status: {0}"), SomeStatus));
 	// }
-	DECLARE_MULTICAST_DELEGATE_OneParam(FGetOnScreenMessagesDelegate, FSeverityMessageMap&);
-	static FGetOnScreenMessagesDelegate OnGetOnScreenMessages;
+	static TMulticastDelegate<void(FSeverityMessageMap&)> OnGetOnScreenMessages;
 
-	DECLARE_DELEGATE_RetVal(bool, FIsLoadingMovieCurrentlyPlaying)
-	static FIsLoadingMovieCurrentlyPlaying IsLoadingMovieCurrentlyPlaying;
+	static TDelegate<bool()> IsLoadingMovieCurrentlyPlaying;
 
 	// Callback to allow user code to prevent url from being launched from FPlatformProcess::LaunchURL. Used to apply http allow list
 	// Return true for to launch the url
-	DECLARE_DELEGATE_RetVal_OneParam(bool, FShouldLaunchUrl, const TCHAR* /* URL */);
-	static FShouldLaunchUrl ShouldLaunchUrl;
+	static TDelegate<bool(const TCHAR* URL)> ShouldLaunchUrl;
 
 	// Callback when the application has been activated by protocol (with optional user id, depending on the platform)
-	DECLARE_MULTICAST_DELEGATE_TwoParams(FOnActivatedByProtocol, const FString& /* parameter */, FPlatformUserId /* user id = PLATFORMUSERID_NONE */ );
-	static FOnActivatedByProtocol OnActivatedByProtocol;
+	static TMulticastDelegate<void(const FString& Parameter, FPlatformUserId UserId /*= PLATFORMUSERID_NONE*/)> OnActivatedByProtocol;
 
 	/** Sent when GC finish destroy takes more time than expected */
-	DECLARE_MULTICAST_DELEGATE_OneParam(FOnGCFinishDestroyTimeExtended, const FString&);
-	static FOnGCFinishDestroyTimeExtended OnGCFinishDestroyTimeExtended;
+	static TMulticastDelegate<void(const FString&)> OnGCFinishDestroyTimeExtended;
 
 	/** Called when the application's network initializes or shutdowns on platforms where the network stack is not always available */
-	DECLARE_MULTICAST_DELEGATE_OneParam(FApplicationNetworkInitializationChanged, bool /*bIsNetworkInitialized*/);
-	static FApplicationNetworkInitializationChanged ApplicationNetworkInitializationChanged;
+	static TMulticastDelegate<void(bool bIsNetworkInitialized)> ApplicationNetworkInitializationChanged;
 
 	// Callback to let code read or write specialized binary data that is generated at Stage time, for optimizing data right before 
 	// final game data is being written to disk
@@ -637,41 +499,135 @@ public:
 		{
 		}
 	};
-	DECLARE_TS_MULTICAST_DELEGATE_OneParam(FAccesExtraBinaryConfigData, FExtraBinaryConfigData&);
-	static FAccesExtraBinaryConfigData AccessExtraBinaryConfigData;
+	static TTSMulticastDelegate<void(FExtraBinaryConfigData&)> AccessExtraBinaryConfigData;
 
 	using FAttachShaderReadRequestFunc = TFunctionRef<class FIoRequest(const class FIoChunkId&, FGraphEventRef)>;
-	DECLARE_DELEGATE_TwoParams(FPreloadPackageShaderMaps, TArrayView<const FSHAHash>, FAttachShaderReadRequestFunc);
-	static FPreloadPackageShaderMaps PreloadPackageShaderMaps;
-	DECLARE_DELEGATE_OneParam(FReleasePreloadedPackageShaderMaps, TArrayView<const FSHAHash>);
-	static FReleasePreloadedPackageShaderMaps ReleasePreloadedPackageShaderMaps;
+	static TDelegate<void(TArrayView<const FSHAHash>, FAttachShaderReadRequestFunc)> PreloadPackageShaderMaps;
+	static TDelegate<void(TArrayView<const FSHAHash>)> ReleasePreloadedPackageShaderMaps;
 	/** Called when the verbosity of a log category is changed */
-	DECLARE_MULTICAST_DELEGATE_ThreeParams(FOnLogVerbosityChanged, const FLogCategoryName& /* CategoryName */, ELogVerbosity::Type /* OldVerbosity */, ELogVerbosity::Type /* NewVerbosity */);
-	static FOnLogVerbosityChanged OnLogVerbosityChanged;
+	static TMulticastDelegate<void(const FLogCategoryName& CategoryName, ELogVerbosity::Type OldVerbosity, ELogVerbosity::Type NewVerbosity)> OnLogVerbosityChanged;
 
 	UE_DEPRECATED(5.1, "Use FPackageStore::Mount() instead")
-	DECLARE_DELEGATE_RetVal(TSharedPtr<class IPackageStore>, FCreatePackageStore);
-	UE_DEPRECATED(5.1, "Use FPackageStore::Mount() instead")
-	PRAGMA_DISABLE_DEPRECATION_WARNINGS
-	static FCreatePackageStore CreatePackageStore;
-	PRAGMA_ENABLE_DEPRECATION_WARNINGS
+	static TDelegate<TSharedPtr<class IPackageStore>()> CreatePackageStore;
 
 	// Called immediately before the parent process will start responding to signals to fork
 	static FSimpleMulticastDelegate OnParentBeginFork;
 	// Called each time immediately before the parent process forks itself
 	static FSimpleMulticastDelegate OnParentPreFork;
 
-	DECLARE_MULTICAST_DELEGATE_OneParam(FProcessForkDelegate, EForkProcessRole /* ProcessRole */);
 	// Called immediately after the process spawned a fork
-	static FProcessForkDelegate OnPostFork;
+	static TMulticastDelegate<void(EForkProcessRole ProcessRole)> OnPostFork;
 	// Called at the end of the frame where the process spawned a fork
 	static FSimpleMulticastDelegate OnChildEndFramePostFork;
 
 private:
 
 	// Callbacks for hotfixes
-	static TArray<FHotFixDelegate> HotFixDelegates;
+	static TArray<TDelegate<void(void*, int32)>> HotFixDelegates;
 
 	// This class is only for namespace use
-	FCoreDelegates() {}
+	FCoreDelegates() = default;
+
+public:	// deprecated delegate type aliases
+	using FHotFixDelegate UE_DEPRECATED(5.2, "Use template instantiation instead as `TDelegate<void(void*, int32)>`") = TDelegate<void(void*, int32)>;
+	using FOnActorLabelChanged UE_DEPRECATED(5.2, "Use template instantiation instead as `TMulticastDelegate<void(AActor*)>`") = TMulticastDelegate<void(AActor*)>;
+	using FOnMountAllPakFiles UE_DEPRECATED(5.2, "Use template instantiation instead as `TDelegate<int32(const TArray<FString>&)>`") = TDelegate<int32(const TArray<FString>&)>;
+	using FMountPak UE_DEPRECATED(5.2, "Use template instantiation instead as `TDelegate<IPakFile*(const FString&, int32)>`") = TDelegate<IPakFile*(const FString&, int32)>;
+	using FOnUnmountPak UE_DEPRECATED(5.2, "Use template instantiation instead as `TDelegate<bool(const FString&)>`") = TDelegate<bool(const FString&)>;
+	using FOnOptimizeMemoryUsageForMountedPaks UE_DEPRECATED(5.2, "Use template instantiation instead as `TDelegate<void()>`") = TDelegate<void()>;
+	using FOnPakFileMounted2 UE_DEPRECATED(5.2, "Use template instantiation instead as `TMulticastDelegate<void(const IPakFile&)>`") = TMulticastDelegate<void(const IPakFile&)>;
+	using FNoPakFilesMountedDelegate UE_DEPRECATED(5.2, "Use template instantiation instead as `TMulticastDelegate<void()>`") = TMulticastDelegate<void()>;
+	using FOnModalMessageBox UE_DEPRECATED(5.2, "Use template instantiation instead as `TDelegate<EAppReturnType::Type(EAppMsgType::Type, const FText&, const FText&)>`") = TDelegate<EAppReturnType::Type(EAppMsgType::Type, const FText&, const FText&)>;
+	using FOnHandleSystemEnsure UE_DEPRECATED(5.2, "Use template instantiation instead as `TMulticastDelegate<void()>`") = TMulticastDelegate<void()>;
+	using FOnHandleSystemError UE_DEPRECATED(5.2, "Use template instantiation instead as `TMulticastDelegate<void()>`") = TMulticastDelegate<void()>;
+	using FRegisterMovieStreamerDelegate UE_DEPRECATED(5.2, "Use template instantiation instead as `TMulticastDelegate<void(FMovieStreamerPtr)>`") = TMulticastDelegate<void(FMovieStreamerPtr)>;
+	using FUnRegisterMovieStreamerDelegate UE_DEPRECATED(5.2, "Use template instantiation instead as `TMulticastDelegate<void(FMovieStreamerPtr)>`") = TMulticastDelegate<void(FMovieStreamerPtr)>;
+	using FOnUserLoginChangedEvent UE_DEPRECATED(5.2, "Use template instantiation instead as `TMulticastDelegate<void(bool, int32, int32)>`") = TMulticastDelegate<void(bool, int32, int32)>;
+	using FOnSafeFrameChangedEvent UE_DEPRECATED(5.2, "Use template instantiation instead as `TMulticastDelegate<void()>`") = TMulticastDelegate<void()>;
+	using FOnInviteAccepted UE_DEPRECATED(5.2, "Use template instantiation instead as `TMulticastDelegate<void(const FString&, const FString&)>`") = TMulticastDelegate<void(const FString&, const FString&)>;
+	using FRegisterEncryptionKeyMulticastDelegate UE_DEPRECATED(5.2, "Use template instantiation instead as `TMulticastDelegate<void(const FGuid&, const FAES::FAESKey&)>`") = TMulticastDelegate<void(const FGuid&, const FAES::FAESKey&)>;
+	using FPakEncryptionKeyDelegate UE_DEPRECATED(5.2, "Use template instantiation instead as `TDelegate<void(uint8[32])>`") = TDelegate<void(uint8[32])>;
+	using FPakSigningKeysDelegate UE_DEPRECATED(5.2, "Use template instantiation instead as `TDelegate<void(TArray<uint8>&, TArray<uint8>&)>`") = TDelegate<void(TArray<uint8>&, TArray<uint8>&)>;
+	using FOnUserControllerConnectionChange UE_DEPRECATED(5.2, "Use template instantiation instead as `TMulticastDelegate<void(bool, FPlatformUserId, int32)>`") = TMulticastDelegate<void(bool, FPlatformUserId, int32)>;
+	using FOnUserControllerPairingChange UE_DEPRECATED(5.2, "Use template instantiation instead as `TMulticastDelegate<void(int32 ControllerIndex, FPlatformUserId NewUserPlatformId, FPlatformUserId OldUserPlatformId)>`") = TMulticastDelegate<void(int32 ControllerIndex, FPlatformUserId NewUserPlatformId, FPlatformUserId OldUserPlatformId)>;
+	using FOnAsyncLoadingFlush UE_DEPRECATED(5.2, "Use template instantiation instead as `TMulticastDelegate<void()>`") = TMulticastDelegate<void()>;
+	using FOnAsyncLoadingFlushUpdate UE_DEPRECATED(5.2, "Use template instantiation instead as `TMulticastDelegate<void()>`") = TMulticastDelegate<void()>;
+	using FOnAsyncLoadPackage UE_DEPRECATED(5.2, "Use template instantiation instead as `TMulticastDelegate<void(const FString&)>`") = TMulticastDelegate<void(const FString&)>;
+	using FOnSyncLoadPackage UE_DEPRECATED(5.2, "Use template instantiation instead as `TMulticastDelegate<void(const FString&)>`") = TMulticastDelegate<void(const FString&)>;
+	using FNewFileAddedDelegate UE_DEPRECATED(5.2, "Use template instantiation instead as `TMulticastDelegate<void(const FString&)>`") = TMulticastDelegate<void(const FString&)>;
+	using FOnFileOpenedForReadFromPakFile UE_DEPRECATED(5.2, "Use template instantiation instead as `TMulticastDelegate<void(const TCHAR* PakFile, const TCHAR* FileName)>`") = TMulticastDelegate<void(const TCHAR* PakFile, const TCHAR* FileName)>;
+#if ALLOW_OTHER_PLATFORM_CONFIG
+	using FGatherDeviceProfileCVars UE_DEPRECATED(5.2, "Use template instantiation instead as `TDelegate<FCVarKeyValueMap(const FString& DeviceProfileName)>`") = TDelegate<FCVarKeyValueMap(const FString& DeviceProfileName)>;
+#endif
+	using FGatherAdditionalLocResPathsDelegate UE_DEPRECATED(5.2, "Use template instantiation instead as `TMulticastDelegate<void(TArray<FString>&)>`") = TMulticastDelegate<void(TArray<FString>&)>;
+	using FWorldOriginOffset UE_DEPRECATED(5.2, "Use template instantiation instead as `TMulticastDelegate<void(class UWorld*, FIntVector, FIntVector)>`") = TMulticastDelegate<void(class UWorld*, FIntVector, FIntVector)>;
+	using FStarvedGameLoop UE_DEPRECATED(5.2, "Use template instantiation instead as `TDelegate<void()>`") = TDelegate<void()>;
+	using FOnTemperatureChange UE_DEPRECATED(5.2, "Use template instantiation instead as `TMulticastDelegate<void(ETemperatureSeverity)>`") = TMulticastDelegate<void(ETemperatureSeverity)>;
+	using FOnLowPowerMode UE_DEPRECATED(5.2, "Use template instantiation instead as `TMulticastDelegate<void(bool)>`") = TMulticastDelegate<void(bool)>;
+	using FCountPreLoadConfigFileRespondersDelegate UE_DEPRECATED(5.2, "Use template instantiation instead as `TMulticastDelegate<void(const TCHAR* IniFilename, int32& ResponderCount)>`") = TMulticastDelegate<void(const TCHAR* IniFilename, int32& ResponderCount)>;
+	using FPreLoadConfigFileDelegate UE_DEPRECATED(5.2, "Use template instantiation instead as `TMulticastDelegate<void(const TCHAR* IniFilename, FString& LoadedContents)>`") = TMulticastDelegate<void(const TCHAR* IniFilename, FString& LoadedContents)>;
+	using FPreSaveConfigFileDelegate UE_DEPRECATED(5.2, "Use template instantiation instead as `TMulticastDelegate<void(const TCHAR* IniFilename, const FString& ContentsToSave, int32& SavedCount)>`") = TMulticastDelegate<void(const TCHAR* IniFilename, const FString& ContentsToSave, int32& SavedCount)>;
+	using FOnFConfigFileCreated UE_DEPRECATED(5.2, "Use template instantiation instead as `TMulticastDelegate<void(const FConfigFile*)>`") = TMulticastDelegate<void(const FConfigFile*)>;
+	using FOnFConfigFileDeleted UE_DEPRECATED(5.2, "Use template instantiation instead as `TMulticastDelegate<void(const FConfigFile*)>`") = TMulticastDelegate<void(const FConfigFile*)>;
+	using FOnConfigValueRead UE_DEPRECATED(5.2, "Use template instantiation instead as `TMulticastDelegate<void(const TCHAR* IniFilename, const TCHAR* SectionName, const TCHAR* Key)>`") = TMulticastDelegate<void(const TCHAR* IniFilename, const TCHAR* SectionName, const TCHAR* Key)>;
+	using FOnConfigSectionRead UE_DEPRECATED(5.2, "Use template instantiation instead as `TMulticastDelegate<void(const TCHAR* IniFilename, const TCHAR* SectionName)>`") = TMulticastDelegate<void(const TCHAR* IniFilename, const TCHAR* SectionName)>;
+	using FOnConfigSectionsChanged UE_DEPRECATED(5.2, "Use template instantiation instead as `TMulticastDelegate<void(const FString& IniFilename, const TSet<FString>& SectionNames)>`") = TMulticastDelegate<void(const FString& IniFilename, const TSet<FString>& SectionNames)>;
+	using FOnApplyCVarFromIni UE_DEPRECATED(5.2, "Use template instantiation instead as `TMulticastDelegate<void(const TCHAR* SectionName, const TCHAR* IniFilename, uint32 SetBy, bool bAllowCheating)>`") = TMulticastDelegate<void(const TCHAR* SectionName, const TCHAR* IniFilename, uint32 SetBy, bool bAllowCheating)>;
+	using FOnSystemResolutionChanged UE_DEPRECATED(5.2, "Use template instantiation instead as `TMulticastDelegate<void(uint32 ResX, uint32 ResY)>`") = TMulticastDelegate<void(uint32 ResX, uint32 ResY)>;
+	using FOnTargetPlatformChangedSupportedFormats UE_DEPRECATED(5.2, "Use template instantiation instead as `TMulticastDelegate<void(const ITargetPlatform*)>`") = TMulticastDelegate<void(const ITargetPlatform*)>;
+	using FOnFeatureLevelDisabled UE_DEPRECATED(5.2, "Use template instantiation instead as `TMulticastDelegate<void(int, const FName&)>`") = TMulticastDelegate<void(int, const FName&)>;
+	using FApplicationLifetimeDelegate UE_DEPRECATED(5.2, "Use template instantiation instead as `TMulticastDelegate<void()>`") = TMulticastDelegate<void()>;
+	using FBackgroundTickDelegate UE_DEPRECATED(5.2, "Use template instantiation instead as `TMulticastDelegate<void(float DeltaTime)>`") = TMulticastDelegate<void(float DeltaTime)>;
+	using FUserMusicInterruptDelegate UE_DEPRECATED(5.2, "Use template instantiation instead as `TMulticastDelegate<void(bool)>`") = TMulticastDelegate<void(bool)>;
+	using FAudioMuteDelegate UE_DEPRECATED(5.2, "Use template instantiation instead as `TMulticastDelegate<void(bool, int)>`") = TMulticastDelegate<void(bool, int)>;
+	using FAudioRouteChangedDelegate UE_DEPRECATED(5.2, "Use template instantiation instead as `TMulticastDelegate<void(bool)>`") = TMulticastDelegate<void(bool)>;
+	using FApplicationRequestAudioState UE_DEPRECATED(5.2, "Use template instantiation instead as `TMulticastDelegate<void()>`") = TMulticastDelegate<void()>;
+	using FApplicationStartupArgumentsDelegate UE_DEPRECATED(5.2, "Use template instantiation instead as `TMulticastDelegate<void(const TArray<FString>&)>`") = TMulticastDelegate<void(const TArray<FString>&)>;
+	using FApplicationRegisteredForRemoteNotificationsDelegate UE_DEPRECATED(5.2, "Use template instantiation instead as `TMulticastDelegate<void(TArray<uint8>)>`") = TMulticastDelegate<void(TArray<uint8>)>;
+	using FApplicationRegisteredForUserNotificationsDelegate UE_DEPRECATED(5.2, "Use template instantiation instead as `TMulticastDelegate<void(int)>`") = TMulticastDelegate<void(int)>;
+	using FApplicationFailedToRegisterForRemoteNotificationsDelegate UE_DEPRECATED(5.2, "Use template instantiation instead as `TMulticastDelegate<void(FString)>`") = TMulticastDelegate<void(FString)>;
+	using FApplicationReceivedRemoteNotificationDelegate UE_DEPRECATED(5.2, "Use template instantiation instead as `TMulticastDelegate<void(FString, int)>`") = TMulticastDelegate<void(FString, int)>;
+	using FApplicationReceivedLocalNotificationDelegate UE_DEPRECATED(5.2, "Use template instantiation instead as `TMulticastDelegate<void(FString, int, int)>`") = TMulticastDelegate<void(FString, int, int)>;
+	using FApplicationPerformFetchDelegate UE_DEPRECATED(5.2, "Use template instantiation instead as `TMulticastDelegate<void()>`") = TMulticastDelegate<void()>;
+	using FApplicationBackgroundSessionEventDelegate UE_DEPRECATED(5.2, "Use template instantiation instead as `TMulticastDelegate<void(FString)>`") = TMulticastDelegate<void(FString)>;
+	using FApplicationReceivedOnScreenOrientationChangedNotificationDelegate UE_DEPRECATED(5.2, "Use template instantiation instead as `TMulticastDelegate<void(int32)>`") = TMulticastDelegate<void(int32)>;
+	using FStatCheckEnabled UE_DEPRECATED(5.2, "Use template instantiation instead as `TMulticastDelegate<void(const TCHAR*, bool&, bool&)>`") = TMulticastDelegate<void(const TCHAR*, bool&, bool&)>;
+	using FStatEnabled UE_DEPRECATED(5.2, "Use template instantiation instead as `TMulticastDelegate<void(const TCHAR*)>`") = TMulticastDelegate<void(const TCHAR*)>;
+	using FStatDisabled UE_DEPRECATED(5.2, "Use template instantiation instead as `TMulticastDelegate<void(const TCHAR*)>`") = TMulticastDelegate<void(const TCHAR*)>;
+	using FStatDisableAll UE_DEPRECATED(5.2, "Use template instantiation instead as `TMulticastDelegate<void(const bool)>`") = TMulticastDelegate<void(const bool)>;
+	using FApplicationLicenseChange UE_DEPRECATED(5.2, "Use template instantiation instead as `TMulticastDelegate<void()>`") = TMulticastDelegate<void()>;
+	using FPlatformChangedLaptopMode UE_DEPRECATED(5.2, "Use template instantiation instead as `TMulticastDelegate<void(EConvertibleLaptopMode)>`") = TMulticastDelegate<void(EConvertibleLaptopMode)>;
+	using FVRHeadsetTrackingInitializingAndNeedsHMDToBeTrackedDelegate UE_DEPRECATED(5.2, "Use template instantiation instead as `TMulticastDelegate<void()>`") = TMulticastDelegate<void()>;
+	using FVRHeadsetTrackingInitializedDelegate UE_DEPRECATED(5.2, "Use template instantiation instead as `TMulticastDelegate<void()>`") = TMulticastDelegate<void()>;
+	using FVRHeadsetRecenter UE_DEPRECATED(5.2, "Use template instantiation instead as `TMulticastDelegate<void()>`") = TMulticastDelegate<void()>;
+	using FVRHeadsetLost UE_DEPRECATED(5.2, "Use template instantiation instead as `TMulticastDelegate<void()>`") = TMulticastDelegate<void()>;
+	using FVRHeadsetReconnected UE_DEPRECATED(5.2, "Use template instantiation instead as `TMulticastDelegate<void()>`") = TMulticastDelegate<void()>;
+	using FVRHeadsetConnectCanceled UE_DEPRECATED(5.2, "Use template instantiation instead as `TMulticastDelegate<void()>`") = TMulticastDelegate<void()>;
+	using FVRHeadsetPutOnHead UE_DEPRECATED(5.2, "Use template instantiation instead as `TMulticastDelegate<void()>`") = TMulticastDelegate<void()>;
+	using FVRHeadsetRemovedFromHead UE_DEPRECATED(5.2, "Use template instantiation instead as `TMulticastDelegate<void()>`") = TMulticastDelegate<void()>;
+	using FVRControllerRecentered UE_DEPRECATED(5.2, "Use template instantiation instead as `TMulticastDelegate<void()>`") = TMulticastDelegate<void()>;
+	using FOnUserActivityStringChanged UE_DEPRECATED(5.2, "Use template instantiation instead as `TMulticastDelegate<void(const FString&)>`") = TMulticastDelegate<void(const FString&)>;
+	using FOnGameSessionIDChange UE_DEPRECATED(5.2, "Use template instantiation instead as `TMulticastDelegate<void(const FString&)>`") = TMulticastDelegate<void(const FString&)>;
+	using FOnGameStateClassChange UE_DEPRECATED(5.2, "Use template instantiation instead as `TMulticastDelegate<void(const FString&)>`") = TMulticastDelegate<void(const FString&)>;
+	using FOnCrashOverrideParamsChanged UE_DEPRECATED(5.2, "Use template instantiation instead as `TMulticastDelegate<void(const FCrashOverrideParameters&)>`") = TMulticastDelegate<void(const FCrashOverrideParameters&)>;
+	using FOnIsVanillaProductChanged UE_DEPRECATED(5.2, "Use template instantiation instead as `TMulticastDelegate<void(bool)>`") = TMulticastDelegate<void(bool)>;
+	using FOnPreMainInit UE_DEPRECATED(5.2, "Use template instantiation instead as `TMulticastDelegate<void()>`") = TMulticastDelegate<void()>;
+	using FConfigReadyForUse UE_DEPRECATED(5.2, "Use template instantiation instead as `TMulticastDelegate<void()>`") = TMulticastDelegate<void()>;
+	using FRenderingThreadChanged UE_DEPRECATED(5.2, "Use template instantiation instead as `TMulticastDelegate<void()>`") = TMulticastDelegate<void()>;
+	using FResolvePackageNameDelegate UE_DEPRECATED(5.2, "Use template instantiation instead as `TMulticastDelegate<bool(const FString&, FString&)>`") = TMulticastDelegate<bool(const FString&, FString&)>;
+	using FGPUOutOfMemoryDelegate UE_DEPRECATED(5.2, "Use template instantiation instead as `TMulticastDelegate<void(const uint64, const uint64)>`") = TMulticastDelegate<void(const uint64, const uint64)>;
+	using FGetOnScreenMessagesDelegate UE_DEPRECATED(5.2, "Use template instantiation instead as `TMulticastDelegate<void(FSeverityMessageMap&)>`") = TMulticastDelegate<void(FSeverityMessageMap&)>;
+	using FIsLoadingMovieCurrentlyPlaying UE_DEPRECATED(5.2, "Use template instantiation instead as `TDelegate<bool()>`") = TDelegate<bool()>;
+	using FShouldLaunchUrl UE_DEPRECATED(5.2, "Use template instantiation instead as `TDelegate<bool(const TCHAR* URL)>`") = TDelegate<bool(const TCHAR* URL)>;
+	using FOnActivatedByProtocol UE_DEPRECATED(5.2, "Use template instantiation instead as `TMulticastDelegate<void(const FString& Parameter, FPlatformUserId UserId /*= PLATFORMUSERID_NONE*/)>`") = TMulticastDelegate<void(const FString& Parameter, FPlatformUserId UserId /*= PLATFORMUSERID_NONE*/)>;
+	using FOnGCFinishDestroyTimeExtended UE_DEPRECATED(5.2, "Use template instantiation instead as `TMulticastDelegate<void(const FString&)>`") = TMulticastDelegate<void(const FString&)>;
+	using FApplicationNetworkInitializationChanged UE_DEPRECATED(5.2, "Use template instantiation instead as `TMulticastDelegate<void(bool bIsNetworkInitialized)>`") = TMulticastDelegate<void(bool bIsNetworkInitialized)>;
+	using FAccesExtraBinaryConfigData UE_DEPRECATED(5.2, "Use template instantiation instead as `TMulticastDelegate<void(FExtraBinaryConfigData&)>`") = TMulticastDelegate<void(FExtraBinaryConfigData&)>;
+	using FPreloadPackageShaderMaps UE_DEPRECATED(5.2, "Use template instantiation instead as `TDelegate<void(TArrayView<const FSHAHash>, FAttachShaderReadRequestFunc)>`") = TDelegate<void(TArrayView<const FSHAHash>, FAttachShaderReadRequestFunc)>;
+	using FReleasePreloadedPackageShaderMaps UE_DEPRECATED(5.2, "Use template instantiation instead as `TDelegate<void(TArrayView<const FSHAHash>)>`") = TDelegate<void(TArrayView<const FSHAHash>)>;
+	using FOnLogVerbosityChanged UE_DEPRECATED(5.2, "Use template instantiation instead as `TMulticastDelegate<void(const FLogCategoryName& CategoryName, ELogVerbosity::Type OldVerbosity, ELogVerbosity::Type NewVerbosity)>`") = TMulticastDelegate<void(const FLogCategoryName& CategoryName, ELogVerbosity::Type OldVerbosity, ELogVerbosity::Type NewVerbosity)>;
+PRAGMA_DISABLE_DEPRECATION_WARNINGS
+	using FCreatePackageStore UE_DEPRECATED(5.2, "Use template instantiation instead as `TDelegate<TSharedPtr<class IPackageStore>()>`") = TDelegate<TSharedPtr<class IPackageStore>()>;
+PRAGMA_ENABLE_DEPRECATION_WARNINGS
+	using FProcessForkDelegate UE_DEPRECATED(5.2, "Use template instantiation instead as `TMulticastDelegate<void(EForkProcessRole ProcessRole)>`") = TMulticastDelegate<void(EForkProcessRole ProcessRole)>;
 };
diff --git a/Engine/Source/Runtime/PakFile/Private/IPlatformFilePak.cpp b/Engine/Source/Runtime/PakFile/Private/IPlatformFilePak.cpp
--- a/Engine/Source/Runtime/PakFile/Private/IPlatformFilePak.cpp
+++ b/Engine/Source/Runtime/PakFile/Private/IPlatformFilePak.cpp
@@ -367,7 +367,7 @@ TSharedPtr<const struct FPakSignatureFile, ESPMode::ThreadSafe> FPakPlatformFile
 	static bool bInitializedPublicKey = false;
 	if (!bInitializedPublicKey)
 	{
-		FCoreDelegates::FPakSigningKeysDelegate& Delegate = FCoreDelegates::GetPakSigningKeysDelegate();
+		TDelegate<void(TArray<uint8>&, TArray<uint8>&)>& Delegate = FCoreDelegates::GetPakSigningKeysDelegate();
 		if (Delegate.IsBound())
 		{
 			TArray<uint8> Exponent;
diff --git a/Engine/Source/Runtime/Renderer/Private/RendererOnScreenNotification.h b/Engine/Source/Runtime/Renderer/Private/RendererOnScreenNotification.h
--- a/Engine/Source/Runtime/Renderer/Private/RendererOnScreenNotification.h
+++ b/Engine/Source/Runtime/Renderer/Private/RendererOnScreenNotification.h
@@ -91,7 +91,7 @@ private:
 private:
 	FCriticalSection DelgateCS;
 	FCriticalSection MessageCS;
-	FCoreDelegates::FGetOnScreenMessagesDelegate ProxyDelegate;
+	TMulticastDelegate<void(FCoreDelegates::FSeverityMessageMap&)> ProxyDelegate;
 	FDelegateHandle BaseDelegateHandle;
 	FCoreDelegates::FSeverityMessageMap Messages;
 	FCoreDelegates::FSeverityMessageMap MessagesTmp;
diff --git a/Engine/Source/Runtime/Windows/AudioMixerXAudio2/Private/AudioMixerPlatformXAudio2.cpp b/Engine/Source/Runtime/Windows/AudioMixerXAudio2/Private/AudioMixerPlatformXAudio2.cpp
--- a/Engine/Source/Runtime/Windows/AudioMixerXAudio2/Private/AudioMixerPlatformXAudio2.cpp
+++ b/Engine/Source/Runtime/Windows/AudioMixerXAudio2/Private/AudioMixerPlatformXAudio2.cpp
@@ -699,7 +699,7 @@ namespace Audio
 
 #if PLATFORM_NEEDS_SUSPEND_ON_BACKGROUND
 
-		auto UnregisterLambda = [](FDelegateHandle& InHandle, FCoreDelegates::FApplicationLifetimeDelegate& InDelegate) 
+		auto UnregisterLambda = [](FDelegateHandle& InHandle, TMulticastDelegate<void()>& InDelegate)
 		{
 			if (InHandle.IsValid())
 			{
