Math: Added GetTypeHash for FQuat and FTransform

diff --git a/Engine/Source/Runtime/Core/Public/Math/Quat.h b/Engine/Source/Runtime/Core/Public/Math/Quat.h
--- a/Engine/Source/Runtime/Core/Public/Math/Quat.h
+++ b/Engine/Source/Runtime/Core/Public/Math/Quat.h
@@ -1350,7 +1350,20 @@ FORCEINLINE bool TQuat<T>::ContainsNaN() const
 			!FMath::IsFinite(W)
 	);
 }
-
+	
+/**
+ * Creates a hash value from an FQuat.
+ *
+ * @param Quat the quat to create a hash value for
+ * @return The hash value from the components
+ */
+template<typename T>
+FORCEINLINE uint32 GetTypeHash(const TQuat<T>& Quat)
+{
+	// Note: this assumes there's no padding in Quat that could contain uncompared data.
+	static_assert(sizeof(TQuat<T>) == sizeof(T[4]), "Unexpected padding in TQuat");
+	return FCrc::MemCrc_DEPRECATED(&Quat, sizeof(Quat));
+}
 
 } // namespace UE::Math
 } // namespace UE
diff --git a/Engine/Source/Runtime/Core/Public/Math/Transform.h b/Engine/Source/Runtime/Core/Public/Math/Transform.h
--- a/Engine/Source/Runtime/Core/Public/Math/Transform.h
+++ b/Engine/Source/Runtime/Core/Public/Math/Transform.h
@@ -58,4 +58,24 @@ inline bool FTransform3d::SerializeFromMismatchedTag(FName StructTag, FArchive&
 		return false;
 	}
 	return UE_SERIALIZE_VARIANT_FROM_MISMATCHED_TAG(Ar, Transform, Transform3d, Transform3f);
+}
+
+namespace UE::Math
+{
+	/**
+	 * Creates a hash value from an FTransform.
+	 *
+	 * @param Transform the transform to create a hash value for
+	 * @return The hash value from the components
+	 */
+	template<typename T>
+	inline uint32 GetTypeHash(const TTransform<T>& Transform)
+	{
+		uint32 Hash = GetTypeHash(Transform.GetTranslation());
+		Hash = HashCombineFast(Hash, GetTypeHash(Transform.GetRotation()));
+		Hash = HashCombineFast(Hash, GetTypeHash(Transform.GetScale3D()));
+
+		return Hash;
+	}
+	
 }
\ No newline at end of file
