CompactBinary: Fixed LoadCompactBinary to gracefully handle read failures and sizes larger than the archive

diff --git a/Engine/Source/Runtime/Core/Private/Serialization/CompactBinarySerialization.cpp b/Engine/Source/Runtime/Core/Private/Serialization/CompactBinarySerialization.cpp
--- a/Engine/Source/Runtime/Core/Private/Serialization/CompactBinarySerialization.cpp
+++ b/Engine/Source/Runtime/Core/Private/Serialization/CompactBinarySerialization.cpp
@@ -184,39 +184,39 @@ FCbField LoadCompactBinary(FArchive& Ar, FCbBufferAllocator Allocator)
 	ECbFieldType FieldType;
 	uint64 FieldSize = 1;
 
-	// Read in small increments until the total field size is known, to avoid reading too far.
-	for (;;)
+	for (const int64 StartPos = Ar.Tell(); !Ar.IsError() && FieldSize > 0;)
 	{
+		// Read in small increments until the total field size is known, to avoid reading too far.
 		const int32 ReadSize = int32(FieldSize - HeaderBytes.Num());
 		const int32 ReadOffset = HeaderBytes.AddUninitialized(ReadSize);
 		Ar.Serialize(HeaderBytes.GetData() + ReadOffset, ReadSize);
-		if (TryMeasureCompactBinary(MakeMemoryView(HeaderBytes), FieldType, FieldSize))
+
+		if (!Ar.IsError() && TryMeasureCompactBinary(MakeMemoryView(HeaderBytes), FieldType, FieldSize))
 		{
+			if (FieldSize <= uint64(Ar.TotalSize() - StartPos))
+			{
+				FUniqueBuffer Buffer = Allocator(FieldSize);
+				checkf(Buffer.GetSize() == FieldSize, TEXT("Allocator returned a buffer of size %" UINT64_FMT " bytes "
+					"when %" UINT64_FMT " bytes were requested."), Buffer.GetSize(), FieldSize);
+
+				FMutableMemoryView View = Buffer.GetView().CopyFrom(MakeMemoryView(HeaderBytes));
+				if (!View.IsEmpty())
+				{
+					// Read the remainder of the field.
+					Ar.Serialize(View.GetData(), static_cast<int64>(View.GetSize()));
+				}
+
+				if (!Ar.IsError() && ValidateCompactBinary(Buffer, ECbValidateMode::Default) == ECbValidateError::None)
+				{
+					return FCbField(Buffer.MoveToShared());
+				}
+			}
 			break;
 		}
-		if (FieldSize == 0)
-		{
-			Ar.SetError();
-			return FCbField();
-		}
 	}
 
-	// Allocate the buffer, copy the header, and read the remainder of the field.
-	FUniqueBuffer Buffer = Allocator(FieldSize);
-	checkf(Buffer.GetSize() == FieldSize,
-		TEXT("Allocator returned a buffer of size %" UINT64_FMT " bytes when %" UINT64_FMT " bytes were requested."),
-		Buffer.GetSize(), FieldSize);
-	FMutableMemoryView View = Buffer.GetView().CopyFrom(MakeMemoryView(HeaderBytes));
-	if (!View.IsEmpty())
-	{
-		Ar.Serialize(View.GetData(), static_cast<int64>(View.GetSize()));
-	}
-	if (ValidateCompactBinary(Buffer, ECbValidateMode::Default) != ECbValidateError::None)
-	{
-		Ar.SetError();
-		return FCbField();
-	}
-	return FCbField(Buffer.MoveToShared());
+	Ar.SetError();
+	return FCbField();
 }
 
 ///////////////////////////////////////////////////////////////////////////////////////////////////////////////////////
diff --git a/Engine/Source/Runtime/Core/Tests/Serialization/CompactBinaryPackageTest.cpp b/Engine/Source/Runtime/Core/Tests/Serialization/CompactBinaryPackageTest.cpp
--- a/Engine/Source/Runtime/Core/Tests/Serialization/CompactBinaryPackageTest.cpp
+++ b/Engine/Source/Runtime/Core/Tests/Serialization/CompactBinaryPackageTest.cpp
@@ -26,13 +26,13 @@ TEST_CASE("Core::Serialization::CompactBinary::Attachment", "[Core][CompactBinar
 		CHECK(ValidateCompactBinaryAttachment(MakeMemoryView(WriteAr), ECbValidateMode::All) == ECbValidateError::None);
 
 		FCbAttachment FromFields;
-		FromFields.TryLoad(Fields);
+		CHECK(FromFields.TryLoad(Fields));
 		CHECK_FALSE(Fields);
 		CHECK(FromFields == Attachment);
 
 		FCbAttachment FromArchive;
 		FMemoryReader ReadAr(WriteAr);
-		FromArchive.TryLoad(ReadAr);
+		CHECK(FromArchive.TryLoad(ReadAr));
 		CHECK(ReadAr.AtEnd());
 		CHECK(FromArchive == Attachment);
 	};
@@ -144,7 +144,7 @@ TEST_CASE("Core::Serialization::CompactBinary::Attachment", "[Core][CompactBinar
 		FCbFieldIterator Fields = Writer.Save();
 		FCbFieldIterator FieldsView = FCbFieldIterator::MakeRangeView(FCbFieldViewIterator(Fields));
 
-		Attachment.TryLoad(FieldsView);
+		CHECK(Attachment.TryLoad(FieldsView));
 		CHECK_FALSE(Attachment.IsNull());
 		CHECK(Attachment);
 		CHECK_FALSE(FieldsView.GetOuterBuffer().GetView().Contains(Attachment.AsBinary().GetView()));
@@ -167,7 +167,7 @@ TEST_CASE("Core::Serialization::CompactBinary::Attachment", "[Core][CompactBinar
 		FCbFieldIterator Fields = Writer.Save();
 		FCbFieldIterator FieldsView = FCbFieldIterator::MakeRangeView(FCbFieldViewIterator(Fields));
 
-		Attachment.TryLoad(FieldsView);
+		CHECK(Attachment.TryLoad(FieldsView));
 		CHECK_FALSE(Attachment.IsNull());
 		CHECK(Attachment);
 		CHECK_FALSE(FieldsView.GetOuterBuffer().GetView().Contains(Attachment.AsCompressedBinary().GetCompressed().ToShared().GetView()));
@@ -194,7 +194,7 @@ TEST_CASE("Core::Serialization::CompactBinary::Attachment", "[Core][CompactBinar
 		FCbFieldIterator Fields = Writer.Save();
 		FCbFieldIterator FieldsView = FCbFieldIterator::MakeRangeView(FCbFieldViewIterator(Fields));
 
-		Attachment.TryLoad(FieldsView);
+		CHECK(Attachment.TryLoad(FieldsView));
 		FMemoryView View;
 		CHECK_FALSE(Attachment.IsNull());
 		CHECK(Attachment);
@@ -262,13 +262,13 @@ TEST_CASE("Core::Serialization::CompactBinary::Package", "[Core][CompactBinary][
 		CHECK(ValidateCompactBinaryPackage(MakeMemoryView(WriteAr), ECbValidateMode::All) == ECbValidateError::None);
 
 		FCbPackage FromFields;
-		FromFields.TryLoad(Fields);
+		CHECK(FromFields.TryLoad(Fields));
 		CHECK_FALSE(Fields);
 		CHECK(FromFields == Package);
 
 		FCbPackage FromArchive;
 		FMemoryReader ReadAr(WriteAr);
-		FromArchive.TryLoad(ReadAr);
+		CHECK(FromArchive.TryLoad(ReadAr));
 		CHECK(ReadAr.AtEnd());
 		CHECK(FromArchive == Package);
 	};
@@ -470,7 +470,7 @@ TEST_CASE("Core::Serialization::CompactBinary::Package", "[Core][CompactBinary][
 
 		FCbFieldIterator Fields = Writer.Save();
 		FCbPackage FromFields;
-		FromFields.TryLoad(Fields);
+		CHECK(FromFields.TryLoad(Fields));
 
 		const FCbAttachment* const Level2Attachment = FromFields.FindAttachment(Level2Hash);
 		const FCbAttachment* const Level3Attachment = FromFields.FindAttachment(Level3Hash);
@@ -494,7 +494,7 @@ TEST_CASE("Core::Serialization::CompactBinary::Package", "[Core][CompactBinary][
 		Writer.Save(WriteAr);
 		FCbPackage FromArchive;
 		FMemoryReader ReadAr(WriteAr);
-		FromArchive.TryLoad(ReadAr);
+		CHECK(FromArchive.TryLoad(ReadAr));
 
 		Writer.Reset();
 		FromArchive.Save(Writer);
@@ -548,4 +548,55 @@ TEST_CASE("Core::Serialization::CompactBinary::Package", "[Core][CompactBinary][
 	}
 }
 
+TEST_CASE("Core::Serialization::CompactBinary::InvalidPackage", "[Core][CompactBinary][Smoke]")
+{
+	const auto TestLoad = [](std::initializer_list<uint8> RawData, FCbBufferAllocator Allocator = FUniqueBuffer::Alloc)
+	{
+		const FMemoryView RawView = MakeMemoryView(RawData);
+		FCbPackage FromArchive;
+		FMemoryReaderView ReadAr(RawView);
+		CHECK_FALSE(FromArchive.TryLoad(ReadAr, Allocator));
+	};
+
+	const auto AllocFail = [](uint64) -> FUniqueBuffer
+	{
+		FAIL_CHECK("Allocation is not expected");
+		return FUniqueBuffer();
+	};
+
+	SECTION("Empty")
+	{
+		TestLoad({}, AllocFail);
+	}
+
+	SECTION("Invalid Initial Field")
+	{
+		TestLoad({uint8(ECbFieldType::None)});
+		TestLoad({uint8(ECbFieldType::Array)});
+		TestLoad({uint8(ECbFieldType::UniformArray)});
+		TestLoad({uint8(ECbFieldType::Binary)});
+		TestLoad({uint8(ECbFieldType::String)});
+		TestLoad({uint8(ECbFieldType::IntegerPositive)});
+		TestLoad({uint8(ECbFieldType::IntegerNegative)});
+		TestLoad({uint8(ECbFieldType::Float32)});
+		TestLoad({uint8(ECbFieldType::Float64)});
+		TestLoad({uint8(ECbFieldType::BoolFalse)});
+		TestLoad({uint8(ECbFieldType::BoolTrue)});
+		TestLoad({uint8(ECbFieldType::ObjectAttachment)});
+		TestLoad({uint8(ECbFieldType::BinaryAttachment)});
+		TestLoad({uint8(ECbFieldType::Uuid)});
+		TestLoad({uint8(ECbFieldType::DateTime)});
+		TestLoad({uint8(ECbFieldType::TimeSpan)});
+		TestLoad({uint8(ECbFieldType::ObjectId)});
+		TestLoad({uint8(ECbFieldType::CustomById)});
+		TestLoad({uint8(ECbFieldType::CustomByName)});
+	}
+
+	SECTION("Size Out Of Bounds")
+	{
+		TestLoad({uint8(ECbFieldType::Object), 1}, AllocFail);
+		TestLoad({uint8(ECbFieldType::Object), 0xff, 0x12, 0x34, 0x56, 0x78, 0x9a, 0xbc, 0xde, 0xf0}, AllocFail);
+	}
+}
+
 #endif // WITH_LOW_LEVEL_TESTS
