Fixed leak in page cache due to multiple symbols. Pivoted to use the MemStack page allocator to share pages between memstack and CLA.

diff --git a/Engine/Source/Runtime/Core/Public/Experimental/ConcurrentLinearAllocator.h b/Engine/Source/Runtime/Core/Public/Experimental/ConcurrentLinearAllocator.h
--- a/Engine/Source/Runtime/Core/Public/Experimental/ConcurrentLinearAllocator.h
+++ b/Engine/Source/Runtime/Core/Public/Experimental/ConcurrentLinearAllocator.h
@@ -10,6 +10,7 @@
 #include "Templates/UniquePtr.h"
 #include "Templates/UnrealTypeTraits.h"
 #include "Containers/LockFreeFixedSizeAllocator.h"
+#include "Misc/MemStack.h"
 
 #ifdef USE_MALLOC_BINNED3
 #if USE_MALLOC_BINNED3
@@ -141,17 +142,15 @@ public:
 template <uint32 BlockSize, typename Allocator = FOsAllocator>
 class TBlockAllocationLockFreeCache
 {
-private:
-	static TLockFreeFixedSizeAllocator<BlockSize, PLATFORM_CACHE_LINE_SIZE> BlockAllocator;
-
 public:
+	static_assert(BlockSize == FPageAllocator::PageSize, "Only 64k pages are supported with this cache.");
 	static constexpr bool SupportsAlignment = Allocator::SupportsAlignment;
 
 	FORCEINLINE static void* Malloc(SIZE_T Size, uint32 Alignment)
 	{
 		if (Size == BlockSize)
 		{
-			return BlockAllocator.Allocate();
+			return FPageAllocator::Get().Alloc();
 		}
 		else
 		{
@@ -163,7 +162,7 @@ public:
 	{
 		if (Size == BlockSize)
 		{
-			BlockAllocator.Free(Pointer);
+			FPageAllocator::Get().Free(Pointer);
 		}
 		else
 		{
@@ -172,9 +171,6 @@ public:
 	}
 };
 
-template <uint32 BlockSize, typename Allocator>
-TLockFreeFixedSizeAllocator<BlockSize, PLATFORM_CACHE_LINE_SIZE> TBlockAllocationLockFreeCache<BlockSize, Allocator>::BlockAllocator;
-
 struct FDefaultBlockAllocationTag 
 {
 	static constexpr uint32 BlockSize = 64 * 1024;		// Blocksize used to allocate from
