Reenable trace worker thread to fix hitching on the GT.

diff --git a/Engine/Source/Runtime/Core/Private/ProfilingDebugging/TraceAuxiliary.cpp b/Engine/Source/Runtime/Core/Private/ProfilingDebugging/TraceAuxiliary.cpp
--- a/Engine/Source/Runtime/Core/Private/ProfilingDebugging/TraceAuxiliary.cpp
+++ b/Engine/Source/Runtime/Core/Private/ProfilingDebugging/TraceAuxiliary.cpp
@@ -717,9 +717,7 @@ static void TraceAuxiliaryFile(const TArray<FString>& Args)
 		return;
 	}
 
-	FTraceAuxiliary::Options Opts;
-	Opts.bNoWorkerThread = true;
-	if (FTraceAuxiliary::Start(FTraceAuxiliary::EConnectionType::File, Filepath, Channels, &Opts, LogConsoleResponse))
+	if (FTraceAuxiliary::Start(FTraceAuxiliary::EConnectionType::File, Filepath, Channels, nullptr, LogConsoleResponse))
 	{
 		TraceAuxiliaryConnectEpilogue();
 	}
