Added F prefix to Options struct name

#rnx

diff --git a/Engine/Source/Runtime/Core/Private/ProfilingDebugging/TraceAuxiliary.cpp b/Engine/Source/Runtime/Core/Private/ProfilingDebugging/TraceAuxiliary.cpp
--- a/Engine/Source/Runtime/Core/Private/ProfilingDebugging/TraceAuxiliary.cpp
+++ b/Engine/Source/Runtime/Core/Private/ProfilingDebugging/TraceAuxiliary.cpp
@@ -1157,7 +1157,7 @@ static bool StartFromCommandlineArguments(const TCHAR* CommandLine)
 	FTraceAuxiliary::EConnectionType Type = FTraceAuxiliary::EConnectionType::None;
 
 	// Setup options
-	FTraceAuxiliary::Options Opts;
+	FTraceAuxiliary::FOptions Opts;
 	Opts.bTruncateFile = FParse::Param(CommandLine, TEXT("tracefiletrunc"));
 
 	if (FParse::Param(FCommandLine::Get(), TEXT("nothreading")) || FParse::Param(FCommandLine::Get(), TEXT("notracethreading")))
@@ -1219,7 +1219,7 @@ static bool StartFromCommandlineArguments(const TCHAR* CommandLine)
 FTraceAuxiliary::FOnConnection FTraceAuxiliary::OnConnection;
 
 ////////////////////////////////////////////////////////////////////////////////
-bool FTraceAuxiliary::Start(EConnectionType Type, const TCHAR* Target, const TCHAR* Channels, Options* Options, const FLogCategoryAlias& LogCategory)
+bool FTraceAuxiliary::Start(EConnectionType Type, const TCHAR* Target, const TCHAR* Channels, FOptions* Options, const FLogCategoryAlias& LogCategory)
 {
 #if UE_TRACE_ENABLED
 	if (GTraceAuxiliary.IsParentProcessAndPreFork())
diff --git a/Engine/Source/Runtime/Core/Public/ProfilingDebugging/TraceAuxiliary.h b/Engine/Source/Runtime/Core/Public/ProfilingDebugging/TraceAuxiliary.h
--- a/Engine/Source/Runtime/Core/Public/ProfilingDebugging/TraceAuxiliary.h
+++ b/Engine/Source/Runtime/Core/Public/ProfilingDebugging/TraceAuxiliary.h
@@ -45,7 +45,7 @@ public:
 		None,
 	};
 
-	struct Options
+	struct FOptions
 	{
 		/** When set, trace will not start a worker thread, instead it is updated from end frame delegate. */
 		bool bNoWorkerThread = false;
@@ -63,7 +63,7 @@ public:
 	 * @param LogCategory Log channel to output messages to. Default set to 'Core'.
 	 * @return True when successfully starting the trace, false if the data connection could not be made.
 	 */
-	static bool Start(EConnectionType Type, const TCHAR* Target, const TCHAR* Channels = TEXT("default"), Options* Options = nullptr, const FLogCategoryAlias& LogCategory = LogCore);
+	static bool Start(EConnectionType Type, const TCHAR* Target, const TCHAR* Channels = TEXT("default"), FOptions* Options = nullptr, const FLogCategoryAlias& LogCategory = LogCore);
 
 	/**
 	 * Stop tracing.
