Probably fix for PVS Studio warning about taking the size of a pointer.

diff --git a/Engine/Source/Runtime/Core/Public/Delegates/DelegateSignatureImpl.inl b/Engine/Source/Runtime/Core/Public/Delegates/DelegateSignatureImpl.inl
--- a/Engine/Source/Runtime/Core/Public/Delegates/DelegateSignatureImpl.inl
+++ b/Engine/Source/Runtime/Core/Public/Delegates/DelegateSignatureImpl.inl
@@ -44,7 +44,7 @@ template<typename UserPolicy> class TTSMulticastDelegateBase;
 template <typename DelegateSignature, typename UserPolicy = FDefaultDelegateUserPolicy>
 class TDelegate
 {
-	static_assert(sizeof(DelegateSignature) == 0, "Expected a function signature for the delegate template parameter");
+	static_assert(sizeof(UserPolicy) == 0, "Expected a function signature for the delegate template parameter");
 };
 
 template <typename InRetValType, typename... ParamTypes, typename UserPolicy>
