Replaced the (commented out) line making the write pipe non-blocking with a comment

diff --git a/Engine/Source/Runtime/Core/Private/Mac/MacPlatformProcess.cpp b/Engine/Source/Runtime/Core/Private/Mac/MacPlatformProcess.cpp
--- a/Engine/Source/Runtime/Core/Private/Mac/MacPlatformProcess.cpp
+++ b/Engine/Source/Runtime/Core/Private/Mac/MacPlatformProcess.cpp
@@ -1316,8 +1316,8 @@ bool FMacPlatformProcess::CreatePipe( void*& ReadPipe, void*& WritePipe, bool bW
 	int pipefd[2];
 	pipe(pipefd);
 
+	// The read pipe should be non-blocking, but the write pipe SHOULD block
 	fcntl(pipefd[0], F_SETFL, O_NONBLOCK);
-	//fcntl(pipefd[1], F_SETFL, O_NONBLOCK);
 
 	// create an NSFileHandle from the descriptor
 	ReadPipe = [[NSFileHandle alloc] initWithFileDescriptor: pipefd[0]];
