Typo fix.

diff --git a/Engine/Source/Runtime/Core/Public/Misc/GeneratedTypeName.h b/Engine/Source/Runtime/Core/Public/Misc/GeneratedTypeName.h
--- a/Engine/Source/Runtime/Core/Public/Misc/GeneratedTypeName.h
+++ b/Engine/Source/Runtime/Core/Public/Misc/GeneratedTypeName.h
@@ -105,7 +105,7 @@ namespace UETypeName_Private
  * Caveats:
  * - The strings are compiler-dependent and thus non-portable, and so shouldn't be saved or relied upon as a form of identity, e.g. the example above returns "class FString" on MSVC.
  * - Default template parameters are also handled differently by different compilers, sometimes ignored, sometimes not.
- * - Only the concrete type is known to the compiler, aliases are ignored, e.g. GetTypeName<TCHAR>() typically returns "wchar_t".
+ * - Only the concrete type is known to the compiler, aliases are ignored, e.g. GetGeneratedTypeName<TCHAR>() typically returns "wchar_t".
  */
 template <typename T>
 inline const TCHAR* GetGeneratedTypeName()
