fix modular build

#rnx

diff --git a/Engine/Source/Runtime/Core/Public/GenericPlatform/GenericPlatformChunkInstall.h b/Engine/Source/Runtime/Core/Public/GenericPlatform/GenericPlatformChunkInstall.h
--- a/Engine/Source/Runtime/Core/Public/GenericPlatform/GenericPlatformChunkInstall.h
+++ b/Engine/Source/Runtime/Core/Public/GenericPlatform/GenericPlatformChunkInstall.h
@@ -474,7 +474,7 @@ protected:
 };
 
 // temporary helper base for platform chunk installers to transition from FCustomChunk to named chunks
-class FNamedChunkPlatformChunkInstall : public FGenericPlatformChunkInstall
+class CORE_API FNamedChunkPlatformChunkInstall : public FGenericPlatformChunkInstall
 {
 public:
 	virtual bool SupportsNamedChunkInstall() const override
