Trace
* Added trace file created callback:  Useful to gather all trace files created during execution
* Added config option to disable or resize tail buffers via [Trace.Config]:TailSizeBytes
* Added config option for worker thread sleep time via [Trace.Config]:SleepTimeInMS
* Fixed init code so it never creates a workerthread in builds that don't support multithreading

diff --git a/Engine/Source/Runtime/Core/Private/ProfilingDebugging/TraceAuxiliary.cpp b/Engine/Source/Runtime/Core/Private/ProfilingDebugging/TraceAuxiliary.cpp
--- a/Engine/Source/Runtime/Core/Private/ProfilingDebugging/TraceAuxiliary.cpp
+++ b/Engine/Source/Runtime/Core/Private/ProfilingDebugging/TraceAuxiliary.cpp
@@ -56,6 +56,7 @@
 ////////////////////////////////////////////////////////////////////////////////
 const TCHAR* GDefaultChannels = TEXT("cpu,gpu,frame,log,bookmark");
 const TCHAR* GMemoryChannels = TEXT("memtag,memalloc,callstack,module");
+const TCHAR* GTraceConfigSection = TEXT("Trace.Config");
 
 ////////////////////////////////////////////////////////////////////////////////
 CSV_DEFINE_CATEGORY(Trace, true);
@@ -516,7 +517,11 @@ bool FTraceAuxiliaryImpl::WriteToFile(const TCHAR* Path, const FTraceAuxiliary::
 		return false;
 	}
 
-	if (!UE::Trace::WriteTo(*NativePath))
+	if (UE::Trace::WriteTo(*NativePath))
+	{
+		FTraceAuxiliary::OnTraceFileCreated.Broadcast(*NativePath);
+	}
+	else
 	{
 		if (FPathViews::Equals(NativePath, FStringView(Path)))
 		{
@@ -646,6 +651,30 @@ void OnConnectionCallback()
 	FTraceAuxiliary::OnConnection.Broadcast();
 }
 
+////////////////////////////////////////////////////////////////////////////////
+static void SetupInitFromConfig(UE::Trace::FInitializeDesc& OutDesc)
+{
+	if (!GConfig)
+	{
+		return;
+	}
+
+	int32 SleepTimeConfig = 0;
+	if (GConfig->GetInt(GTraceConfigSection, TEXT("SleepTimeInMS"), SleepTimeConfig, GEngineIni))
+	{
+		if (SleepTimeConfig > 0)
+		{
+			OutDesc.ThreadSleepTimeInMS = SleepTimeConfig;
+		}
+	}
+
+	int32 TailSizeBytesConfig = 0;
+	if (GConfig->GetInt(GTraceConfigSection, TEXT("TailSizeBytes"), TailSizeBytesConfig, GEngineIni))
+	{
+		OutDesc.TailSizeBytes = TailSizeBytesConfig;
+	}
+}
+
 ////////////////////////////////////////////////////////////////////////////////
 static void TraceAuxiliaryConnectEpilogue()
 {
@@ -1160,7 +1189,9 @@ static bool StartFromCommandlineArguments(const TCHAR* CommandLine)
 	FTraceAuxiliary::FOptions Opts;
 	Opts.bTruncateFile = FParse::Param(CommandLine, TEXT("tracefiletrunc"));
 
-	if (FParse::Param(FCommandLine::Get(), TEXT("nothreading")) || FParse::Param(FCommandLine::Get(), TEXT("notracethreading")))
+	const bool bWorkerThreadAllowed = FGenericPlatformProcess::SupportsMultithreading() || FForkProcessHelper::IsForkedMultithreadInstance();
+
+	if (!bWorkerThreadAllowed || FParse::Param(FCommandLine::Get(), TEXT("notracethreading")))
 	{
 		Opts.bNoWorkerThread = true;
 	}
@@ -1217,6 +1248,7 @@ static bool StartFromCommandlineArguments(const TCHAR* CommandLine)
 
 ////////////////////////////////////////////////////////////////////////////////
 FTraceAuxiliary::FOnConnection FTraceAuxiliary::OnConnection;
+FTraceAuxiliary::FOnTraceFileCreated FTraceAuxiliary::OnTraceFileCreated;
 
 ////////////////////////////////////////////////////////////////////////////////
 bool FTraceAuxiliary::Start(EConnectionType Type, const TCHAR* Target, const TCHAR* Channels, FOptions* Options, const FLogCategoryAlias& LogCategory)
@@ -1236,7 +1268,10 @@ bool FTraceAuxiliary::Start(EConnectionType Type, const TCHAR* Target, const TCH
 	// Make sure the worker thread is started unless explicitly opt out.
 	if (!Options || !Options->bNoWorkerThread)
 	{
-		GTraceAuxiliary.StartWorkerThread();
+		if (FGenericPlatformProcess::SupportsMultithreading() || FForkProcessHelper::IsForkedMultithreadInstance())
+		{
+			GTraceAuxiliary.StartWorkerThread();
+		}
 	}
 
 	if (Channels)
@@ -1406,6 +1441,8 @@ void FTraceAuxiliary::Initialize(const TCHAR* CommandLine)
 
 	// Initialize Trace
 	UE::Trace::FInitializeDesc Desc;
+	SetupInitFromConfig(Desc);
+	
 	Desc.bUseWorkerThread = false;
 	Desc.bUseImportantCache = (FParse::Param(CommandLine, TEXT("tracenocache")) == false);
 	Desc.OnConnectionFunc = &OnConnectionCallback;
@@ -1499,6 +1536,13 @@ void FTraceAuxiliary::EnableChannels()
 #endif
 }
 
+void FTraceAuxiliary::DisableChannels(const TCHAR* Channels)
+{
+#if UE_TRACE_ENABLED
+	GTraceAuxiliary.DisableChannels(Channels);
+#endif
+}
+
 const TCHAR* FTraceAuxiliary::GetTraceDestination()
 {
 #if UE_TRACE_ENABLED
diff --git a/Engine/Source/Runtime/Core/Public/ProfilingDebugging/TraceAuxiliary.h b/Engine/Source/Runtime/Core/Public/ProfilingDebugging/TraceAuxiliary.h
--- a/Engine/Source/Runtime/Core/Public/ProfilingDebugging/TraceAuxiliary.h
+++ b/Engine/Source/Runtime/Core/Public/ProfilingDebugging/TraceAuxiliary.h
@@ -28,6 +28,9 @@ public:
 	 */
 	DECLARE_MULTICAST_DELEGATE(FOnConnection);
 
+	/** Callback whenever a trace file is created */
+	DECLARE_MULTICAST_DELEGATE_OneParam(FOnTraceFileCreated, const TCHAR* FileNameAndPath);
+
 	enum class EConnectionType
 	{
 		/**
@@ -116,6 +119,12 @@ public:
 	 */
 	static void EnableChannels();
 
+	/**
+	* Disable channels to stop recording traces with them.
+	* @param Channels List of channels (or a preset) to disable. If null it will disable all active channels.
+	*/
+	static void DisableChannels(const TCHAR* Channels = nullptr);
+
 	/**
 	 *  Returns the destination string that is currently being traced to.
 	 *  Contains either a file path or network address. Points to an empty string if tracing is disabled.
@@ -146,5 +155,12 @@ public:
 	 * @param Callback Delegate to call on new connections.
 	 */
 	static FOnConnection OnConnection;
+
+	/**
+	 * Delegate that triggers when a trace file is created. Useful if you need to collect all trace files
+     * created during a session.
+	 * The full file name and path is passed to the delegate.
+	 */
+	static FOnTraceFileCreated OnTraceFileCreated;
 	
 };
diff --git a/Engine/Source/Runtime/CoreUObject/Public/UObject/CoreNet.h b/Engine/Source/Runtime/CoreUObject/Public/UObject/CoreNet.h
--- a/Engine/Source/Runtime/CoreUObject/Public/UObject/CoreNet.h
+++ b/Engine/Source/Runtime/CoreUObject/Public/UObject/CoreNet.h
@@ -336,7 +336,6 @@ template <> struct TIsZeroConstructType<FLifetimeProperty> { enum { Value = true
 
 GENERATE_MEMBER_FUNCTION_CHECK(GetLifetimeReplicatedProps, void, const, TArray<FLifetimeProperty>&)
 
-// Consider adding UE_NET_TRACE_ENABLE to build config, for now we use the UE_TRACE_ENABLED as NetTrace is not support unless tracing is enabled
 #if UE_TRACE_ENABLED
 /**
  * We pass a NetTraceCollector along with the NetBitWriter in order avoid modifying all API`s where we want to be able to collect Network stats
diff --git a/Engine/Source/Runtime/TraceLog/Private/Trace/Writer.cpp b/Engine/Source/Runtime/TraceLog/Private/Trace/Writer.cpp
--- a/Engine/Source/Runtime/TraceLog/Private/Trace/Writer.cpp
+++ b/Engine/Source/Runtime/TraceLog/Private/Trace/Writer.cpp
@@ -587,6 +587,7 @@ void Writer_CallbackOnConnect()
 ////////////////////////////////////////////////////////////////////////////////
 static UPTRINT			GWorkerThread;		// = 0;
 static volatile bool	GWorkerThreadQuit;	// = false;
+static uint32			GSleepTimeInMS = 17;
 static volatile unsigned int	GUpdateInProgress = 1;	// Don't allow updates until initialized
 
 ////////////////////////////////////////////////////////////////////////////////
@@ -630,8 +631,7 @@ static void Writer_WorkerThread()
 	{
 		Writer_WorkerUpdate();
 
-		const uint32 SleepMs = 17;
-		ThreadSleep(SleepMs);
+		ThreadSleep(GSleepTimeInMS);
 	}
 }
 
@@ -759,6 +759,11 @@ void Writer_Initialize(const FInitializeDesc& Desc)
 		Writer_InitializeCache();
 	}
 
+	if (Desc.ThreadSleepTimeInMS != 0)
+	{
+		GSleepTimeInMS = Desc.ThreadSleepTimeInMS;
+	}
+
 	if (Desc.bUseWorkerThread)
 	{
 		Writer_WorkerCreate();
@@ -784,11 +789,8 @@ void Writer_Update()
 	{
 		Writer_WorkerUpdate();
 	}
-
 }
 
-
-
 ////////////////////////////////////////////////////////////////////////////////
 bool Writer_SendTo(const ANSICHAR* Host, uint32 Port)
 {
diff --git a/Engine/Source/Runtime/TraceLog/Public/Trace/Trace.h b/Engine/Source/Runtime/TraceLog/Public/Trace/Trace.h
--- a/Engine/Source/Runtime/TraceLog/Public/Trace/Trace.h
+++ b/Engine/Source/Runtime/TraceLog/Public/Trace/Trace.h
@@ -81,7 +81,8 @@ using OnConnectFunc = void(void);
 
 struct FInitializeDesc
 {
-	uint32			TailSizeBytes		= 4 << 20;
+	uint32			TailSizeBytes		= 4 << 20;  // Set to 0 to disable the tail buffer
+	uint32			ThreadSleepTimeInMS = 0;
 	bool			bUseWorkerThread	= true;
 	bool			bUseImportantCache	= true;
 	OnConnectFunc*	OnConnectionFunc	= nullptr;
