Adding new placeholder base class for platform chunk installers that are implementing the new Named Chunks system, allowing support for the legacy FCustomChunk system.

- also firming up the const correctness in the new chunk api

#rnx

diff --git a/Engine/Source/Runtime/Core/Private/GenericPlatform/GenericPlatformChunkInstall.cpp b/Engine/Source/Runtime/Core/Private/GenericPlatform/GenericPlatformChunkInstall.cpp
--- a/Engine/Source/Runtime/Core/Private/GenericPlatform/GenericPlatformChunkInstall.cpp
+++ b/Engine/Source/Runtime/Core/Private/GenericPlatform/GenericPlatformChunkInstall.cpp
@@ -35,14 +35,14 @@ bool FNamedChunkPlatformChunkInstall::UninstallNamedChunk(const FName NamedChunk
 }
 
 
-bool FNamedChunkPlatformChunkInstall::InstallNamedChunks(const TArrayView<FName>& NamedChunks)
+bool FNamedChunkPlatformChunkInstall::InstallNamedChunks(const TArrayView<const FName>& NamedChunks)
 {
 	TArray<FCustomChunk> CustomChunks = GetCustomChunksFromNamedChunks(NamedChunks);
 	return InstallChunks(CustomChunks);
 }
 
 
-bool FNamedChunkPlatformChunkInstall::UninstallNamedChunks(const TArrayView<FName>& NamedChunks)
+bool FNamedChunkPlatformChunkInstall::UninstallNamedChunks(const TArrayView<const FName>& NamedChunks)
 {
 	TArray<FCustomChunk> CustomChunks = GetCustomChunksFromNamedChunks(NamedChunks);
 	return UninstallChunks(CustomChunks);
@@ -134,7 +134,7 @@ TArray<FCustomChunk> FNamedChunkPlatformChunkInstall::GetCustomChunksFromNamedCh
 	return MoveTemp(CustomChunks);
 }
 
-TArray<FCustomChunk> FNamedChunkPlatformChunkInstall::GetCustomChunksFromNamedChunks(const TArrayView<FName>& NamedChunks) const
+TArray<FCustomChunk> FNamedChunkPlatformChunkInstall::GetCustomChunksFromNamedChunks(const TArrayView<const FName>& NamedChunks) const
 {
 	TArray<FCustomChunk> CustomChunks;
 
@@ -222,4 +222,62 @@ FName FNamedChunkPlatformChunkInstall::GetCustomChunkName(const FCustomChunk& Cu
 	}
 }
 
+
+
+
+
+
+
+
+
+
+
+
+
+
+
+
+
+
+bool FGenericPlatformChunkInstall_WithEmulatedCustomChunks::IsChunkInstallationPending(const TArray<FCustomChunk>& ChunkTagsID)
+{
+	for (const FCustomChunk& CustomChunk : ChunkTagsID)
+	{
+		if (IsNamedChunkInProgress(FName(CustomChunk.ChunkTag)) || IsNamedChunkInProgress(FName(CustomChunk.ChunkTag2)) )
+		{
+			return true;
+		}
+	}
+
+	return false;
+}
+
+bool FGenericPlatformChunkInstall_WithEmulatedCustomChunks::InstallChunks(const TArray<FCustomChunk>& ChunkTagsID)
+{
+	for (const FCustomChunk& CustomChunk : ChunkTagsID)
+	{
+		if (!InstallNamedChunk(FName(CustomChunk.ChunkTag)) && !InstallNamedChunk(FName(CustomChunk.ChunkTag2)))
+		{
+			return false;
+		}
+	}
+
+	return true;
+}
+
+bool FGenericPlatformChunkInstall_WithEmulatedCustomChunks::UninstallChunks(const TArray<FCustomChunk>& ChunkTagsID)
+{
+	for (const FCustomChunk& CustomChunk : ChunkTagsID)
+	{
+		if (!UninstallNamedChunk(FName(CustomChunk.ChunkTag)) && !UninstallNamedChunk(FName(CustomChunk.ChunkTag2)))
+		{
+			return false;
+		}
+	}
+
+	return true;
+}
+
+
+
 PRAGMA_ENABLE_DEPRECATION_WARNINGS
diff --git a/Engine/Source/Runtime/Core/Public/GenericPlatform/GenericPlatformChunkInstall.h b/Engine/Source/Runtime/Core/Public/GenericPlatform/GenericPlatformChunkInstall.h
--- a/Engine/Source/Runtime/Core/Public/GenericPlatform/GenericPlatformChunkInstall.h
+++ b/Engine/Source/Runtime/Core/Public/GenericPlatform/GenericPlatformChunkInstall.h
@@ -251,14 +251,14 @@ public:
 	 * @param NamedChunks	The names of the chunks to install
 	 * @return				whether installation task has been kicked
 	 **/
-	virtual bool InstallNamedChunks(const TArrayView<FName>& NamedChunks) = 0;
+	virtual bool InstallNamedChunks(const TArrayView<const FName>& NamedChunks) = 0;
 
 	/**
 	 * Uninstall the given set of named chunks
 	 * @param NamedChunk	The names of the chunks to uninstall
 	 * @return				whether uninstallation task has been kicked
 	 **/
-	virtual bool UninstallNamedChunks(const TArrayView<FName>& NamedChunks) = 0;
+	virtual bool UninstallNamedChunks(const TArrayView<const FName>& NamedChunks) = 0;
 
 	/**
 	 * Get the current location of the given named chunk
@@ -419,20 +419,22 @@ public:
 
 	virtual bool InstallNamedChunk(const FName NamedChunk) override
 	{
-		return false;
+		TArrayView<const FName> NamedChunks = { NamedChunk };
+		return InstallNamedChunks(NamedChunks);
 	}
 
 	virtual bool UninstallNamedChunk(const FName NamedChunk) override
 	{
-		return false;
+		TArrayView<const FName> NamedChunks = { NamedChunk };
+		return UninstallNamedChunks(NamedChunks);
 	}
 
-	virtual bool InstallNamedChunks(const TArrayView<FName>& NamedChunks) override
+	virtual bool InstallNamedChunks(const TArrayView<const FName>& NamedChunks) override
 	{
 		return false;
 	}
 
-	virtual bool UninstallNamedChunks(const TArrayView<FName>& NamedChunks) override
+	virtual bool UninstallNamedChunks(const TArrayView<const FName>& NamedChunks) override
 	{
 		return false;
 	}
@@ -485,8 +487,8 @@ public:
 	virtual bool IsNamedChunkInProgress(const FName NamedChunk) override;
 	virtual bool InstallNamedChunk(const FName NamedChunk) override;
 	virtual bool UninstallNamedChunk(const FName NamedChunk) override;
-	virtual bool InstallNamedChunks(const TArrayView<FName>& NamedChunks) override;
-	virtual bool UninstallNamedChunks(const TArrayView<FName>& NamedChunks) override;
+	virtual bool InstallNamedChunks(const TArrayView<const FName>& NamedChunks) override;
+	virtual bool UninstallNamedChunks(const TArrayView<const FName>& NamedChunks) override;
 
 	virtual EChunkLocation::Type GetNamedChunkLocation(const FName NamedChunk) override;
 	virtual float GetNamedChunkProgress(const FName NamedChunk, EChunkProgressReportingType::Type ReportType) override;
@@ -499,10 +501,20 @@ protected:
 	virtual float GetCustomChunkProgress(const FCustomChunk& CustomChunk, EChunkProgressReportingType::Type ReportType) = 0; // platform specializations need to implement this as it's missing in the existing api
 	bool TryGetCustomChunkFromNamedChunk(const FName NamedChunk, FCustomChunk& OutCustomChunk) const;
 	TArray<FCustomChunk> GetCustomChunksFromNamedChunk(const FName NamedChunk) const;
-	TArray<FCustomChunk> GetCustomChunksFromNamedChunks(const TArrayView<FName>& NamedChunks) const;
+	TArray<FCustomChunk> GetCustomChunksFromNamedChunks(const TArrayView<const FName>& NamedChunks) const;
 	TArray<FName> GetNamedChunksFromCustomChunks(const TArray<FCustomChunk>& CustomChunks) const;
 	FName GetNamedChunkByPakChunkIndex(int32 InPakchunkIndex) const;
 	FName GetCustomChunkName(const FCustomChunk& CustomChunk) const;
 };
 
+// temporary helper base class for platform chunk installers that have implemented named chunk support to provide FCustomChunk emulation
+class CORE_API FGenericPlatformChunkInstall_WithEmulatedCustomChunks : public FGenericPlatformChunkInstall
+{
+public:
+	virtual bool SupportsIntelligentInstall() override final { return true; }
+	virtual bool IsChunkInstallationPending(const TArray<FCustomChunk>& ChunkTagsID) override final;
+	virtual bool InstallChunks(const TArray<FCustomChunk>& ChunkTagsID)  override final;
+	virtual bool UninstallChunks(const TArray<FCustomChunk>& ChunkTagsID) override final;
+};
+
 PRAGMA_ENABLE_DEPRECATION_WARNINGS
\ No newline at end of file
