fixed false check if TScriptDelegate is destroyed from inside its callback

diff --git a/Engine/Source/Runtime/Core/Public/UObject/ScriptDelegates.h b/Engine/Source/Runtime/Core/Public/UObject/ScriptDelegates.h
--- a/Engine/Source/Runtime/Core/Public/UObject/ScriptDelegates.h
+++ b/Engine/Source/Runtime/Core/Public/UObject/ScriptDelegates.h
@@ -264,20 +264,26 @@ public:
 	template <class UObjectTemplate>
 	void ProcessDelegate( void* Parameters ) const
 	{
-		UE_DELEGATES_MT_SCOPED_WRITE_ACCESS(AccessDetector);
+		UObjectTemplate* ObjectPtr;
+		UFunction* Function;
+
+		{	// to avoid MT access check if the delegate is deleted from inside of its callback, we don't cover the callback execution
+			// by access protection scope
+			UE_DELEGATES_MT_SCOPED_WRITE_ACCESS(AccessDetector);
 
-		checkf( Object.IsValid() != false, TEXT( "ProcessDelegate() called with no object bound to delegate!" ) );
-		checkf( FunctionName != NAME_None, TEXT( "ProcessDelegate() called with no function name set!" ) );
+			checkf(Object.IsValid() != false, TEXT("ProcessDelegate() called with no object bound to delegate!"));
+			checkf(FunctionName != NAME_None, TEXT("ProcessDelegate() called with no function name set!"));
 
-		// Object was pending kill, so we cannot execute the delegate.  Note that it's important to assert
-		// here and not simply continue execution, as memory may be left uninitialized if the delegate is
-		// not able to execute, resulting in much harder-to-detect code errors.  Users should always make
-		// sure IsBound() returns true before calling ProcessDelegate()!
-		UObjectTemplate* ObjectPtr = static_cast< UObjectTemplate* >( Object.Get() );	// Down-cast
-		checkSlow( IsValid(ObjectPtr) );
+			// Object was pending kill, so we cannot execute the delegate.  Note that it's important to assert
+			// here and not simply continue execution, as memory may be left uninitialized if the delegate is
+			// not able to execute, resulting in much harder-to-detect code errors.  Users should always make
+			// sure IsBound() returns true before calling ProcessDelegate()!
+			ObjectPtr = static_cast<UObjectTemplate*>(Object.Get());	// Down-cast
+			checkSlow(IsValid(ObjectPtr));
 
-		// Object *must* implement the specified function
-		UFunction* Function = ObjectPtr->FindFunctionChecked( FunctionName );
+			// Object *must* implement the specified function
+			Function = ObjectPtr->FindFunctionChecked(FunctionName);
+		}
 
 		// Execute the delegate!
 		ObjectPtr->ProcessEvent(Function, Parameters);
