switched back to lock-free allocator for tasks and removed CreateCompletionHandle()

diff --git a/Engine/Source/Runtime/Core/Private/Async/TaskGraph.cpp b/Engine/Source/Runtime/Core/Private/Async/TaskGraph.cpp
--- a/Engine/Source/Runtime/Core/Private/Async/TaskGraph.cpp
+++ b/Engine/Source/Runtime/Core/Private/Async/TaskGraph.cpp
@@ -2481,6 +2481,7 @@ bool FTaskGraphInterface::IsMultithread()
 #if TASKGRAPH_NEW_FRONTEND
 
 FGraphEventImplAllocator GraphEventImplAllocator;
+FGraphTaskAllocator SmallTaskAllocator;
 
 #else
 
diff --git a/Engine/Source/Runtime/Core/Private/Tasks/TaskPrivate.cpp b/Engine/Source/Runtime/Core/Private/Tasks/TaskPrivate.cpp
--- a/Engine/Source/Runtime/Core/Private/Tasks/TaskPrivate.cpp
+++ b/Engine/Source/Runtime/Core/Private/Tasks/TaskPrivate.cpp
@@ -11,6 +11,7 @@ namespace UE::Tasks
 {
 	namespace Private
 	{
+		FExecutableTaskAllocator SmallTaskAllocator;
 		FTaskEventBaseAllocator TaskEventBaseAllocator;
 
 		void FTaskBase::Schedule()
diff --git a/Engine/Source/Runtime/Core/Private/Tests/Async/TaskGraphTest.cpp b/Engine/Source/Runtime/Core/Private/Tests/Async/TaskGraphTest.cpp
--- a/Engine/Source/Runtime/Core/Private/Tests/Async/TaskGraphTest.cpp
+++ b/Engine/Source/Runtime/Core/Private/Tests/Async/TaskGraphTest.cpp
@@ -1434,20 +1434,29 @@ PRAGMA_ENABLE_DEPRECATION_WARNINGS
 		return true;
 	}
 
-	IMPLEMENT_SIMPLE_AUTOMATION_TEST(FTaskGraphCreateCompletionHandleTest, "System.Core.Async.TaskGraph.CreateCompletionHandle", EAutomationTestFlags::ApplicationContextMask | EAutomationTestFlags::EngineFilter);
+	IMPLEMENT_SIMPLE_AUTOMATION_TEST(FTaskGraphTaskDestructionTest, "System.Core.Async.TaskGraph.TaskDestruction", EAutomationTestFlags::ApplicationContextMask | EAutomationTestFlags::EngineFilter | EAutomationTestFlags::Disabled);
 
-	bool FTaskGraphCreateCompletionHandleTest::RunTest(const FString& Parameters)
+	bool FTaskGraphTaskDestructionTest::RunTest(const FString& Parameters)
 	{
-		FGraphEventRef UnblockTask = FGraphEvent::CreateGraphEvent(); // to block initially the following task
-		FGraphEventRef Task = FFunctionGraphTask::CreateAndDispatchWhenReady([] {}, TStatId{}, UnblockTask); // supposedly long-living task
-		FGraphEventRef CompletionHandle = Task->CreateCompletionHandle();
-		// check that the completion handle is not signalling as the task is blocked
-		FPlatformProcess::Sleep(0.1f);
-		check(!CompletionHandle->IsComplete());
-		// unblock the task
-		UnblockTask->DispatchSubsequents();
-		// wating for the completion handle instead of waiting for the task, should succeed
-		CompletionHandle->Wait();
+		struct FDestructionTest
+		{
+			explicit FDestructionTest(bool* bDestroyedIn)
+				: bDestroyed(bDestroyedIn)
+			{}
+
+			~FDestructionTest()
+			{
+				*bDestroyed = true;
+			}
+
+			bool* bDestroyed;
+		};
+
+		bool bDestroyed = false;
+
+		FFunctionGraphTask::CreateAndDispatchWhenReady([DestructionTest = FDestructionTest{ &bDestroyed }]{})->Wait();
+
+		check(bDestroyed);
 
 		return true;
 	}
diff --git a/Engine/Source/Runtime/Core/Private/Tests/Tasks/TasksTest.cpp b/Engine/Source/Runtime/Core/Private/Tests/Tasks/TasksTest.cpp
--- a/Engine/Source/Runtime/Core/Private/Tests/Tasks/TasksTest.cpp
+++ b/Engine/Source/Runtime/Core/Private/Tests/Tasks/TasksTest.cpp
@@ -723,7 +723,7 @@ namespace UE { namespace TasksTests
 		Branches.Reserve(NumBranches);
 		for (uint64 BranchIndex = 0; BranchIndex != NumBranches; ++BranchIndex)
 		{
-			Branches.Add(Launch(UE_SOURCE_LOCATION, Branch));
+			Branches.Add(Launch(UE_SOURCE_LOCATION, MoveTemp(Branch)));
 		}
 		TArray<FTask> BranchTasks;
 		BranchTasks.Reserve(NumBranches);
@@ -1320,24 +1320,6 @@ namespace UE { namespace TasksTests
 
 		return true;
 	}
-
-	IMPLEMENT_SIMPLE_AUTOMATION_TEST(FTasksCreateCompletionHandleTest, "System.Core.Async.Tasks.CreateCompletionHandle", EAutomationTestFlags::ApplicationContextMask | EAutomationTestFlags::EngineFilter);
-
-	bool FTasksCreateCompletionHandleTest::RunTest(const FString& Parameters)
-	{
-		FTaskEvent UnblockTask{ UE_SOURCE_LOCATION }; // to block initially the following task
-		FTask Task = Launch(UE_SOURCE_LOCATION, [] {}, Prerequisites(UnblockTask)); // supposedly long-living task
-		FTask CompletionHandle = Task.CreateCompletionHandle();
-		// check that the completion handle is not signalling as the task is blocked
-		FPlatformProcess::Sleep(0.1f);
-		check(!CompletionHandle.IsCompleted());
-		// unblock the task
-		UnblockTask.Trigger();
-		// wating for the completion handle instead of waiting for the task, should succeed
-		check(CompletionHandle.Wait(FTimespan::FromMilliseconds(100)));
-
-		return true;
-	}
 }}
 
 #endif // WITH_DEV_AUTOMATION_TESTS
diff --git a/Engine/Source/Runtime/Core/Public/Async/TaskGraphInterfaces.h b/Engine/Source/Runtime/Core/Public/Async/TaskGraphInterfaces.h
--- a/Engine/Source/Runtime/Core/Public/Async/TaskGraphInterfaces.h
+++ b/Engine/Source/Runtime/Core/Public/Async/TaskGraphInterfaces.h
@@ -641,32 +641,15 @@ public:
 
 		return (ENamedThreads::Type)ConversionMap[(int32)Priority - (int32)EExtendedTaskPriority::GameThreadNormalPri];
 	}
-
-	// see Tasks::FTaskHandle::CreateCompletionHandle description
-	FGraphEventRef CreateCompletionHandle()
-	{
-		if (IsCompleted())
-		{
-			return {};
-		}
-
-		// CreateGraphEvent() uses an allocator that doesn't have an issue with long-living allocs
-		FGraphEventRef CompletionHandle = FGraphEvent::CreateGraphEvent();
-		if (!CompletionHandle->AddPrerequisites(*this))
-		{
-			return {}; // too late, the task is already completed
-		}
-
-		// trigger the completion handle so the only thing that holds it from signalling is the task itself
-		CompletionHandle->DispatchSubsequents();
-
-		return CompletionHandle;
-	}
 };
 
+static constexpr int32 SmallTaskSize = 256;
+using FGraphTaskAllocator = TLockFreeFixedSizeAllocator_TLSCache<SmallTaskSize, PLATFORM_CACHE_LINE_SIZE>;
+CORE_API extern FGraphTaskAllocator SmallTaskAllocator;
+
 // the new task implementation integrated into the old task API
 template<typename TTask>
-class TGraphTask : public TConcurrentLinearObject<TGraphTask<TTask>, FTaskGraphBlockAllocationTag>, public FBaseGraphTask
+class TGraphTask : public FBaseGraphTask
 {
 public:
 	/**
@@ -732,11 +715,14 @@ public:
 		return FConstructor(Prerequisites);
 	}
 
-	virtual ~TGraphTask() override
+	FORCENOINLINE virtual ~TGraphTask() override
 	{
 		DestructItem(TaskStorage.GetTypedPtr());
 	}
 
+	static void* operator new(size_t Size);
+	static void operator delete(void* Ptr, size_t Size);
+
 private:
 	explicit TGraphTask(const FGraphEventArray* InPrerequisites)
 		: FBaseGraphTask(InPrerequisites)
@@ -767,6 +753,18 @@ private:
 	TTypeCompatibleBytes<TTask> TaskStorage;
 };
 
+template<typename TTask>
+void* TGraphTask<TTask>::operator new(size_t Size)
+{
+	return Size <= SmallTaskSize ? SmallTaskAllocator.Allocate() : GMalloc->Malloc(sizeof(TGraphTask), PLATFORM_CACHE_LINE_SIZE);
+}
+
+template<typename TTask>
+void TGraphTask<TTask>::operator delete(void* Ptr, size_t Size)
+{
+	Size <= SmallTaskSize ? SmallTaskAllocator.Free(Ptr) : GMalloc->Free(Ptr);
+}
+
 // an adaptation of FBaseGraphTask to be used as a standalone FGraphEvent
 class FGraphEventImpl : public FBaseGraphTask
 {
@@ -809,6 +807,17 @@ inline FGraphEventRef FBaseGraphTask::CreateGraphEvent()
 
 #else // TASKGRAPH_NEW_FRONTEND
 
+struct FTaskBlockAllocationTag : FDefaultBlockAllocationTag
+{
+	static constexpr uint32 BlockSize = 64 * 1024;
+	static constexpr bool AllowOversizedBlocks = false;
+	static constexpr bool RequiresAccurateSize = false;
+	static constexpr bool InlineBlockAllocation = true;
+	static constexpr const char* TagName = "TaskLinearAllocator";
+
+	using Allocator = TBlockAllocationCache<BlockSize, FAlignedAllocator>;
+};
+
 /** 
  *	Base class for all tasks. 
  *	Tasks go through a very specific life stage progression, and this is verified.
diff --git a/Engine/Source/Runtime/Core/Public/Tasks/Task.h b/Engine/Source/Runtime/Core/Public/Tasks/Task.h
--- a/Engine/Source/Runtime/Core/Public/Tasks/Task.h
+++ b/Engine/Source/Runtime/Core/Public/Tasks/Task.h
@@ -180,35 +180,6 @@ namespace UE::Tasks
 				return IsValid() && Pimpl->IsAwaitable();
 			}
 
-			// Creates and returns a "task event" that can be used to wait for the task completion.
-			// Regular tasks are allocated by a fast allocator that doesn't handle well long-living tasks, as such tasks can keep an entire memory page alive, 
-			// thus pushing up total memory consumption. In most cases such tasks are stored as class members, or as global vars.
-			// hovewer, task events use a different allocator and doesn't cause this issue, and so can be used for long-living tasks.
-			// Make sure to profile and identify that you indeed have "long-living task" problem before using this function, 
-			// as otherwise it would be a needless overhead.
-			// Task events don't support execution result. if you do need execution result, most probably you can reset your task right after getting its execution
-			// result, so it's not long-living anymore. 
-			FTaskHandle CreateCompletionHandle()
-			{
-				if (!IsValid() || IsCompleted())
-				{
-					return {};
-				}
-
-				// `FTaskEventBase` uses an allocator that doesn't have an issue with long-living allocs
-				Private::FTaskEventBase* CompletionHandle{ Private::FTaskEventBase::Create(TEXT("CompletionHandle")) };
-				if (!CompletionHandle->AddPrerequisites(*Pimpl))
-				{
-					delete CompletionHandle;
-					return {}; // too late, the task is already completed
-				}
-
-				CompletionHandle->AddRef(); // internal reference that is released when the handle is completed
-				// trigger the completion handle so the only thing that holds it from signalling is the task itself
-				CompletionHandle->TryLaunch();
-				return FTaskHandle{ CompletionHandle };
-			}
-
 		protected:
 			TRefCountPtr<FTaskBase> Pimpl;
 		};
diff --git a/Engine/Source/Runtime/Core/Public/Tasks/TaskPrivate.h b/Engine/Source/Runtime/Core/Public/Tasks/TaskPrivate.h
--- a/Engine/Source/Runtime/Core/Public/Tasks/TaskPrivate.h
+++ b/Engine/Source/Runtime/Core/Public/Tasks/TaskPrivate.h
@@ -756,33 +756,32 @@ namespace UE::Tasks
 			TTypeCompatibleBytes<ResultType> ResultStorage;
 		};
 
-		struct FTaskBlockAllocationTag : FDefaultBlockAllocationTag
-		{
-			static constexpr uint32 BlockSize = 64 * 1024;
-			static constexpr bool AllowOversizedBlocks = false;
-			static constexpr bool RequiresAccurateSize = false;
-			static constexpr bool InlineBlockAllocation = true;
-			static constexpr const char* TagName = "TaskLinearAllocator";
-
-			using Allocator = TBlockAllocationCache<BlockSize, FAlignedAllocator>;
-		};
-
 		// Task implementation that can be executed, as it stores task body. Generic version (for tasks that return non-void results).
 		// In most cases it should be allocated on the heap and used with TRefCountPtr, e.g. @see FTaskHandle. With care, can be allocated on the stack, e.g. see 
 		// WaitingTask in FTaskBase::Wait().
 		// Implements memory allocation from a pooled fixed-size allocator tuned for the everage UE task size
 		template<typename TaskBodyType, typename ResultType = TInvokeResult_T<TaskBodyType>, typename Enable = void>
-		class TExecutableTask final : public TConcurrentLinearObject<TExecutableTask<TaskBodyType>, FTaskBlockAllocationTag>, public TTaskWithResult<ResultType>
+		class TExecutableTaskBase : public TTaskWithResult<ResultType>
 		{
-			UE_NONCOPYABLE(TExecutableTask);
+			UE_NONCOPYABLE(TExecutableTaskBase);
 
 		public:
-			static TExecutableTask* Create(const TCHAR* DebugName, TaskBodyType TaskBody, ETaskPriority Priority, EExtendedTaskPriority InExtendedPriority = EExtendedTaskPriority::None)
+			virtual bool TryExecuteTask() override
 			{
-				return new TExecutableTask(DebugName, MoveTemp(TaskBody), Priority, InExtendedPriority);
+				return FTaskBase::TryExecute(
+					[](FTaskBase& Task)
+					{
+						TExecutableTaskBase& This = static_cast<TExecutableTaskBase&>(Task);
+						new(&This.ResultStorage) ResultType{ Invoke(*This.TaskBodyStorage.GetTypedPtr()) };
+
+						// destroy the task body as soon as we are done with it, as it can have captured data sensitive to destruction order
+						DestructItem(This.TaskBodyStorage.GetTypedPtr());
+					}
+				);
 			}
 
-			TExecutableTask(const TCHAR* InDebugName, TaskBodyType&& TaskBody, ETaskPriority InPriority, EExtendedTaskPriority InExtendedPriority) :
+		protected:
+			TExecutableTaskBase(const TCHAR* InDebugName, TaskBodyType&& TaskBody, ETaskPriority InPriority, EExtendedTaskPriority InExtendedPriority) :
 				TTaskWithResult<ResultType>(InDebugName, InPriority, InExtendedPriority, 2)
 				// 2 init refs: one for the initial reference (we don't increment it on passing to `TRefCountPtr`), and one for the internal 
 				// reference that keeps the task alive while it's in the system. is released either on task completion or by the scheduler after
@@ -791,13 +790,24 @@ namespace UE::Tasks
 				new(&TaskBodyStorage) TaskBodyType(MoveTemp(TaskBody));
 			}
 
+		private:
+			TTypeCompatibleBytes<TaskBodyType> TaskBodyStorage;
+		};
+
+		// a specialization for tasks that don't return results
+		template<typename TaskBodyType>
+		class TExecutableTaskBase<TaskBodyType, typename TEnableIf<TIsSame<TInvokeResult_T<TaskBodyType>, void>::Value>::Type> : public FTaskBase
+		{
+			UE_NONCOPYABLE(TExecutableTaskBase);
+
+		public:
 			virtual bool TryExecuteTask() override
 			{
-				return FTaskBase::TryExecute(
+				return TryExecute(
 					[](FTaskBase& Task)
 					{
-						TExecutableTask& This = static_cast<TExecutableTask&>(Task);
-						new(&This.ResultStorage) ResultType{ Invoke(*This.TaskBodyStorage.GetTypedPtr()) };
+						TExecutableTaskBase& This = static_cast<TExecutableTaskBase&>(Task);
+						Invoke(*This.TaskBodyStorage.GetTypedPtr());
 
 						// destroy the task body as soon as we are done with it, as it can have captured data sensitive to destruction order
 						DestructItem(This.TaskBodyStorage.GetTypedPtr());
@@ -805,48 +815,50 @@ namespace UE::Tasks
 				);
 			}
 
+		protected:
+			TExecutableTaskBase(const TCHAR* InDebugName, TaskBodyType&& TaskBody, ETaskPriority InPriority, EExtendedTaskPriority InExtendedPriority) :
+				FTaskBase(2) // 2 init refs: one for the initial reference (we don't increment it on passing to `TRefCountPtr`), and one for the internal 
+				// reference that keeps the task alive while it's in the system. is released either on task completion or by the scheduler after
+				// trying to execute the task
+			{
+				Init(InDebugName, InPriority, InExtendedPriority);
+				new(&TaskBodyStorage) TaskBodyType(MoveTemp(TaskBody));
+			}
+
 		private:
 			TTypeCompatibleBytes<TaskBodyType> TaskBodyStorage;
 		};
 
-		// a specialization for tasks that don't return results
+		static constexpr int32 SmallTaskSize = 256;
+		using FExecutableTaskAllocator = TLockFreeFixedSizeAllocator_TLSCache<SmallTaskSize, PLATFORM_CACHE_LINE_SIZE>;
+		CORE_API extern FExecutableTaskAllocator SmallTaskAllocator;
+
+		// a separate derived class to add "small task" allocation optimization to both base class specializations
 		template<typename TaskBodyType>
-		class TExecutableTask<TaskBodyType, typename TEnableIf<TIsSame<TInvokeResult_T<TaskBodyType>, void>::Value>::Type> final : public TConcurrentLinearObject<TExecutableTask<TaskBodyType>, FTaskBlockAllocationTag>, public FTaskBase
+		class TExecutableTask final : public TExecutableTaskBase<TaskBodyType>
 		{
-			UE_NONCOPYABLE(TExecutableTask);
-
 		public:
-			static TExecutableTask* Create(const TCHAR* DebugName, TaskBodyType TaskBody, ETaskPriority Priority, EExtendedTaskPriority InExtendedPriority = EExtendedTaskPriority::None)
+			TExecutableTask(const TCHAR* InDebugName, TaskBodyType&& TaskBody, ETaskPriority InPriority, EExtendedTaskPriority InExtendedPriority)
+				: TExecutableTaskBase<TaskBodyType>(InDebugName, MoveTemp(TaskBody), InPriority, InExtendedPriority)
 			{
-				return new TExecutableTask(DebugName, MoveTemp(TaskBody), Priority, InExtendedPriority);
 			}
 
-			TExecutableTask(const TCHAR* InDebugName, TaskBodyType&& TaskBody, ETaskPriority InPriority, EExtendedTaskPriority InExtendedPriority) :
-				FTaskBase(2)
-				// 2 init refs: one for the initial reference (we don't increment it on passing to `TRefCountPtr`), and one for the internal 
-				// reference that keeps the task alive while it's in the system. is released either on task completion or by the scheduler after
-				// trying to execute the task
+			// a helper that deduces the template argument
+			static TExecutableTask* Create(const TCHAR* InDebugName, TaskBodyType&& TaskBody, ETaskPriority InPriority, EExtendedTaskPriority InExtendedPriority)
 			{
-				Init(InDebugName, InPriority, InExtendedPriority);
-				new(&TaskBodyStorage) TaskBodyType(MoveTemp(TaskBody));
+				return new TExecutableTask(InDebugName, MoveTemp(TaskBody), InPriority, InExtendedPriority);
 			}
 
-			virtual bool TryExecuteTask() override
+			void* operator new(size_t Size)
 			{
-				return TryExecute(
-					[](FTaskBase& Task)
-					{
-						TExecutableTask& This = static_cast<TExecutableTask&>(Task);
-						Invoke(*This.TaskBodyStorage.GetTypedPtr());
-
-						// destroy the task body as soon as we are done with it, as it can have captured data sensitive to destruction order
-						DestructItem(This.TaskBodyStorage.GetTypedPtr());
-					}
-				);
+				return Size <= SmallTaskSize ? SmallTaskAllocator.Allocate() : GMalloc->Malloc(sizeof(TExecutableTask), PLATFORM_CACHE_LINE_SIZE);
 			}
 
-		private:
-			TTypeCompatibleBytes<TaskBodyType> TaskBodyStorage;
+			void operator delete(void* Ptr, size_t Size)
+			{
+				DestructItem((TExecutableTask*)Ptr);
+				Size <= SmallTaskSize ? SmallTaskAllocator.Free(Ptr) : GMalloc->Free(Ptr);
+			}
 		};
 
 		// waiting on named threads that replicates TaskGraph logic
diff --git a/Engine/Source/Runtime/Engine/Private/VT/VirtualTextureTranscodeCache.cpp b/Engine/Source/Runtime/Engine/Private/VT/VirtualTextureTranscodeCache.cpp
--- a/Engine/Source/Runtime/Engine/Private/VT/VirtualTextureTranscodeCache.cpp
+++ b/Engine/Source/Runtime/Engine/Private/VT/VirtualTextureTranscodeCache.cpp
@@ -431,9 +431,7 @@ FVTTranscodeTileHandle FVirtualTextureTranscodeCache::SubmitTask(
 
 	if (bNeedToLaunchTask)
 	{
-		FGraphEventRef Task = TGraphTask<FTranscodeTask>::CreateTask(Prerequisites).ConstructAndDispatchWhenReady(StagingBuffer, InParams);
-		// this reference can live long, store completion handle instead of a reference to the task to reduce peak mem usage 
-		TaskEntry.GraphEvent = Task->CreateCompletionHandle();
+		TaskEntry.GraphEvent = TGraphTask<FTranscodeTask>::CreateTask(Prerequisites).ConstructAndDispatchWhenReady(StagingBuffer, InParams);
 	}
 	else
 	{
diff --git a/Engine/Source/Runtime/Experimental/Chaos/Private/Chaos/PBDRigidsEvolution.cpp b/Engine/Source/Runtime/Experimental/Chaos/Private/Chaos/PBDRigidsEvolution.cpp
--- a/Engine/Source/Runtime/Experimental/Chaos/Private/Chaos/PBDRigidsEvolution.cpp
+++ b/Engine/Source/Runtime/Experimental/Chaos/Private/Chaos/PBDRigidsEvolution.cpp
@@ -781,9 +781,7 @@ namespace Chaos
 			if (bCanStartAsyncTasks)
 			{
 				// we run the task for both starting a new accel structure as well as for the time-slicing
-				FGraphEventRef Task = TGraphTask<FChaosAccelerationStructureTask>::CreateTask().ConstructAndDispatchWhenReady(*SpatialCollectionFactory, SpatialAccelerationCache, AsyncInternalAcceleration, AsyncExternalAcceleration, ForceFullBuild, bIsSingleThreaded, bNeedsReset);
-				// this reference can live long, store completion handle instead of a reference to the task to reduce peak mem usage 
-				AccelerationStructureTaskComplete = Task->CreateCompletionHandle();
+				AccelerationStructureTaskComplete = TGraphTask<FChaosAccelerationStructureTask>::CreateTask().ConstructAndDispatchWhenReady(*SpatialCollectionFactory, SpatialAccelerationCache, AsyncInternalAcceleration, AsyncExternalAcceleration, ForceFullBuild, bIsSingleThreaded, bNeedsReset);
 			}
 		}
 		else
diff --git a/Engine/Source/Runtime/Landscape/Private/LandscapeRender.cpp b/Engine/Source/Runtime/Landscape/Private/LandscapeRender.cpp
--- a/Engine/Source/Runtime/Landscape/Private/LandscapeRender.cpp
+++ b/Engine/Source/Runtime/Landscape/Private/LandscapeRender.cpp
@@ -615,7 +615,7 @@ void FLandscapeSceneViewExtension::PreRenderView_RenderThread(FRDGBuilder& Graph
 	// Kick the job once all views have been collected.
 	if (!LandscapeRenderSystems.IsEmpty() && LandscapeViews.Num() == InView.Family->Views.Num())
 	{
-		const auto ComputeLODs = [this]
+		auto ComputeLODs = [this]
 		{
 			TRACE_CPUPROFILER_EVENT_SCOPE(FLandscapeRenderSystem::ComputeLODs);
 			FOptionalTaskTagScope Scope(ETaskTag::EParallelRenderingThread);
@@ -646,7 +646,7 @@ void FLandscapeSceneViewExtension::PreRenderView_RenderThread(FRDGBuilder& Graph
 
 		if (GIsThreadedRendering && GLandscapeUseAsyncTasksForLODComputation)
 		{
-			LandscapeSetupTask = UE::Tasks::Launch(UE_SOURCE_LOCATION, ComputeLODs, LowLevelTasks::ETaskPriority::Normal);
+			LandscapeSetupTask = UE::Tasks::Launch(UE_SOURCE_LOCATION, MoveTemp(ComputeLODs), LowLevelTasks::ETaskPriority::Normal);
 		}
 		else
 		{
diff --git a/Engine/Source/Runtime/RenderCore/Private/RenderingThread.cpp b/Engine/Source/Runtime/RenderCore/Private/RenderingThread.cpp
--- a/Engine/Source/Runtime/RenderCore/Private/RenderingThread.cpp
+++ b/Engine/Source/Runtime/RenderCore/Private/RenderingThread.cpp
@@ -986,11 +986,8 @@ void StartRenderCommandFenceBundler()
 	STAT_FNullGraphTask_FenceRenderCommandBundled,
 		STATGROUP_TaskGraphTasks);
 
-	FGraphEventRef Task = TGraphTask<FNullGraphTask>::CreateTask(&Prereqs, ENamedThreads::GameThread).ConstructAndDispatchWhenReady(
+	BundledCompletionEvent = TGraphTask<FNullGraphTask>::CreateTask(&Prereqs, ENamedThreads::GameThread).ConstructAndDispatchWhenReady(
 		GET_STATID(STAT_FNullGraphTask_FenceRenderCommandBundled), ENamedThreads::GetRenderThread());
-	// this reference can live long, store completion handle instead of a reference to the task to reduce peak mem usage 
-	BundledCompletionEvent = Task->CreateCompletionHandle();
-	
 
 	StartBatchedRelease();
 }
