Use proper function to make array views.

#rnx

diff --git a/Engine/Source/Runtime/Core/Public/GenericPlatform/GenericPlatformChunkInstall.h b/Engine/Source/Runtime/Core/Public/GenericPlatform/GenericPlatformChunkInstall.h
--- a/Engine/Source/Runtime/Core/Public/GenericPlatform/GenericPlatformChunkInstall.h
+++ b/Engine/Source/Runtime/Core/Public/GenericPlatform/GenericPlatformChunkInstall.h
@@ -419,14 +419,12 @@ public:
 
 	virtual bool InstallNamedChunk(const FName NamedChunk) override
 	{
-		TArrayView<const FName> NamedChunks = { NamedChunk };
-		return InstallNamedChunks(NamedChunks);
+		return InstallNamedChunks(MakeArrayView(&NamedChunk,1));
 	}
 
 	virtual bool UninstallNamedChunk(const FName NamedChunk) override
 	{
-		TArrayView<const FName> NamedChunks = { NamedChunk };
-		return UninstallNamedChunks(NamedChunks);
+		return UninstallNamedChunks(MakeArrayView(&NamedChunk,1));
 	}
 
 	virtual bool InstallNamedChunks(const TArrayView<const FName>& NamedChunks) override
