Fixed issue with TArray::BulkSerialize, when the buffer size can exceed 2Gb.
Even if the array index type is int32, the data buffer can exceed 2Gb.
BulkSerialize requires a cast to ensure the result of NewArrayNum * SerializedElementSize is a 64-bit value.
This should not adversely affect assets previously serialized with BulkSerialize: when saving (but not cooking), the slow path is taken, meaning all the data will be written out correctly.
During cooking and transacting (less permanent operations), the data would have been written and read equally erroneously, so recooking may be required.

diff --git a/Engine/Source/Runtime/Core/Public/Containers/Array.h b/Engine/Source/Runtime/Core/Public/Containers/Array.h
--- a/Engine/Source/Runtime/Core/Public/Containers/Array.h
+++ b/Engine/Source/Runtime/Core/Public/Containers/Array.h
@@ -1326,13 +1326,13 @@ public:
 				Ar << NewArrayNum;
 				Empty(NewArrayNum);
 				AddUninitialized(NewArrayNum);
-				Ar.Serialize(GetData(), NewArrayNum * SerializedElementSize);
+				Ar.Serialize(GetData(), (int64)NewArrayNum * (int64)SerializedElementSize);
 			}
 			else if (Ar.IsSaving())
 			{
 				SizeType ArrayCount = Num();
 				Ar << ArrayCount;
-				Ar.Serialize(GetData(), ArrayCount * SerializedElementSize);
+				Ar.Serialize(GetData(), (int64)ArrayCount * (int64)SerializedElementSize);
 			}
 		}
 	}
