Add a way to determine if we are in a retraction context

diff --git a/Engine/Source/Runtime/Core/Private/HAL/ThreadingBase.cpp b/Engine/Source/Runtime/Core/Private/HAL/ThreadingBase.cpp
--- a/Engine/Source/Runtime/Core/Private/HAL/ThreadingBase.cpp
+++ b/Engine/Source/Runtime/Core/Private/HAL/ThreadingBase.cpp
@@ -287,7 +287,10 @@ PRAGMA_DISABLE_DEPRECATION_WARNINGS
 		: FTaskTagScope::IsCurrentTag(ETaskTag::ERenderingThread);
 
 #if !UE_BUILD_SHIPPING && !UE_BUILD_TEST
-	if (!LowLevelTasks::FSchedulerTls::IsBusyWaiting() && !CoroTask_Detail::FCoroLocalState::IsCoroLaunchedTask())
+	// in all these cases the `newValue` can be wrong, so the "ensure" is disabled
+	if (!LowLevelTasks::FSchedulerTls::IsBusyWaiting() && 
+		!CoroTask_Detail::FCoroLocalState::IsCoroLaunchedTask() && 
+		!UE::Tasks::Private::IsThreadRetractingTask())
 	{
 		const uint32 CurrentThreadId = FPlatformTLS::GetCurrentThreadId();
 		bool oldValue = ((GRenderThreadId == 0) || bLocalIsLoadingThreadSuspended) ? (CurrentThreadId == GGameThreadId) || FTaskTagScope::IsRunningDuringStaticInit() : (CurrentThreadId == GRenderThreadId);
diff --git a/Engine/Source/Runtime/Core/Private/Tasks/TaskPrivate.cpp b/Engine/Source/Runtime/Core/Private/Tasks/TaskPrivate.cpp
--- a/Engine/Source/Runtime/Core/Private/Tasks/TaskPrivate.cpp
+++ b/Engine/Source/Runtime/Core/Private/Tasks/TaskPrivate.cpp
@@ -47,6 +47,140 @@ namespace UE::Tasks
 			LowLevelTasks::FScheduler::Get().TryLaunch(LowLevelTask, LowLevelTasks::EQueuePreference::GlobalQueuePreference, /*bWakeUpWorker=*/ true);
 		}
 
+
+		thread_local uint32 TaskRetractionRecursion = 0;
+
+		bool IsThreadRetractingTask()
+		{
+			return TaskRetractionRecursion != 0;
+		}
+
+		struct FThreadLocalRetractionScope
+		{
+			FThreadLocalRetractionScope()
+			{
+				checkSlow(TaskRetractionRecursion != TNumericLimits<decltype(TaskRetractionRecursion)>::Max() - 1);
+				++TaskRetractionRecursion;
+			}
+
+			~FThreadLocalRetractionScope()
+			{
+				checkSlow(TaskRetractionRecursion != 0);
+				--TaskRetractionRecursion;
+			}
+		};
+
+		bool FTaskBase::TryRetractAndExecute(uint32 RecursionDepth/* = 0*/)
+		{
+			TRACE_CPUPROFILER_EVENT_SCOPE(TaskRetraction);
+
+			if (IsCompleted())
+			{
+				return true;
+			}
+
+#if TASKGRAPH_NEW_FRONTEND
+			// task retraction is not supported for named thread tasks
+			if (IsNamedThreadTask())
+			{
+				return false;
+			}
+#endif
+
+			// avoid stack overflow. is not expected in a real-life cases but happens in stress tests
+			if (RecursionDepth == 200)
+			{
+				return false;
+			}
+			++RecursionDepth;
+
+			// returns false if the task has passed "pre-scheduling" state: all (if any) prerequisites are completed
+			auto IsLockedByPrerequisites = [this]
+			{
+				uint32 LocalNumLocks = NumLocks.load(std::memory_order_relaxed); // the order doesn't matter as this "happens before" task execution
+				return LocalNumLocks != 0 && LocalNumLocks < ExecutionFlag;
+			};
+
+			if (IsLockedByPrerequisites())
+			{
+				// try to unlock the task. even if (some or all) prerequisites retraction fails we still proceed to try helping with other prerequisites or this task execution
+
+				// prerequisites are "consumed" here even if their retraction fails. this means that once prerequisite retraction failed, it won't be performed again. 
+				// this can be potentially improved by using a different container for prerequisites
+				while (FTaskBase* Prerequisite = Prerequisites.Pop())
+				{
+					// ignore if retraction failed, as this thread still can try to help with other prerequisites instead of being blocked in waiting
+					Prerequisite->TryRetractAndExecute(RecursionDepth);
+					Prerequisite->Release();
+				}
+			}
+
+			{
+				FThreadLocalRetractionScope ThreadLocalRetractionScope;
+
+				// next we try to execute the task, despite we haven't verified that the task is unlocked. trying to obtain execution permission will fail in this case
+
+				if (ExtendedPriority == EExtendedTaskPriority::TaskEvent)
+				{
+					if (!TrySetExecutionFlag())
+					{
+						return false;
+					}
+
+					// task events have nothing to execute, and so can't have nested task, just close it
+					Close();
+					ReleaseInternalReference();
+					return true;
+				}
+
+				if (!TryExecuteTask())
+				{
+					return false; // still locked by prerequisites, or another thread managed to set execution flag first, or we're inside this task execution
+					// we could try to help with nested tasks execution (the task execution could already spawned a couple of nested tasks sitting in the queue). 
+					// it's unclear how important this is, but this would definitely lead to more complicated impl. we can revisit this once we see such instances in profiler captures
+				}
+			}
+
+			TRACE_CPUPROFILER_EVENT_SCOPE(SuccessfulTaskRetraction);
+
+			// the task was launched so the scheduler will handle the internal reference held by low-level task
+
+			if (IsCompleted()) // still can be hold back by nested tasks, this is an optional early out for better perf
+			{
+				return true;
+			}
+
+			// retract nested tasks, if any
+			{
+				// keep trying retracting all nested tasks even if some of them fail, so the current worker can contribute instead of being blocked
+				bool bSucceeded = true;
+				// prerequisites are "consumed" here even if their retraction fails. this means that once prerequisite retraction failed, it won't be performed again. 
+				// this can be potentially improved by using a different container for prerequisites
+				while (FTaskBase* Prerequisite = Prerequisites.Pop())
+				{
+					if (!Prerequisite->TryRetractAndExecute(RecursionDepth))
+					{
+						bSucceeded = false;
+					}
+					Prerequisite->Release();
+				}
+
+				if (!bSucceeded)
+				{
+					return false;
+				}
+			}
+
+			// it happens that all nested tasks are completed and are in the process of completing the parent (this task) concurrently, 
+			// but the flag is not set yet. wait for it to maintain postconditions
+			while (!IsCompleted())
+			{
+				FPlatformProcess::Yield();
+			}
+
+			return true;
+		}
+
 		void FTaskBase::Wait()
 		{
 			if (IsCompleted())
diff --git a/Engine/Source/Runtime/Core/Public/Tasks/TaskPrivate.h b/Engine/Source/Runtime/Core/Public/Tasks/TaskPrivate.h
--- a/Engine/Source/Runtime/Core/Public/Tasks/TaskPrivate.h
+++ b/Engine/Source/Runtime/Core/Public/Tasks/TaskPrivate.h
@@ -86,6 +86,10 @@ namespace UE::Tasks
 		// sets the current task and returns the previous current task
 		CORE_API FTaskBase* ExchangeCurrentTask(FTaskBase* Task);
 
+		// Returns true if called from inside a task that is being retracted
+		UE_DEPRECATED(5.2, "You should not use this function as it exists only to patch another system and can be removed any time.")
+		CORE_API bool IsThreadRetractingTask();
+
 		// An abstract base class for task implementation. 
 		// Implements internal logic of task prerequisites, nested tasks and deep task retraction.
 		// Implements intrusive ref-counting and so can be used with TRefCountPtr.
@@ -305,112 +309,7 @@ namespace UE::Tasks
 			// Tries to pull out the task from the system and execute it. If the task is locked by either prerequisites or nested tasks, tries to retract and execute them recursively. 
 			// @return true if task is completed, not necessarily by retraction. If the task is being executed (or its dependency) in parallel, it doesn't wait for task completion and 
 			// returns false immediately.
-			bool TryRetractAndExecute(uint32 RecursionDepth = 0)
-			{
-				TRACE_CPUPROFILER_EVENT_SCOPE(TaskRetraction);
-
-				if (IsCompleted())
-				{
-					return true;
-				}
-
-#if TASKGRAPH_NEW_FRONTEND
-				// task retraction is not supported for named thread tasks
-				if (IsNamedThreadTask())
-				{
-					return false;
-				}
-#endif
-
-				// avoid stack overflow. is not expected in a real-life cases but happens in stress tests
-				if (RecursionDepth == 200)
-				{
-					return false;
-				}
-				++RecursionDepth;
-
-				// returns false if the task has passed "pre-scheduling" state: all (if any) prerequisites are completed
-				auto IsLockedByPrerequisites = [this]
-				{
-					uint32 LocalNumLocks = NumLocks.load(std::memory_order_relaxed); // the order doesn't matter as this "happens before" task execution
-					return LocalNumLocks != 0 && LocalNumLocks < ExecutionFlag;
-				};
-
-				if (IsLockedByPrerequisites())
-				{
-					// try to unlock the task. even if (some or all) prerequisites retraction fails we still proceed to try helping with other prerequisites or this task execution
-
-					// prerequisites are "consumed" here even if their retraction fails. this means that once prerequisite retraction failed, it won't be performed again. 
-					// this can be potentially improved by using a different container for prerequisites
-					while (FTaskBase* Prerequisite = Prerequisites.Pop())
-					{
-						// ignore if retraction failed, as this thread still can try to help with other prerequisites instead of being blocked in waiting
-						Prerequisite->TryRetractAndExecute(RecursionDepth);
-						Prerequisite->Release();
-					}
-				}
-
-				// next we try to execute the task, despite we haven't verified that the task is unlocked. trying to obtain execution permission will fail in this case
-
-				if (ExtendedPriority == EExtendedTaskPriority::TaskEvent)
-				{
-					if (!TrySetExecutionFlag())
-					{
-						return false;
-					}
-
-					// task events have nothing to execute, and so can't have nested task, just close it
-					Close();
-					ReleaseInternalReference();
-					return true;
-				}
-
-				if (!TryExecuteTask())
-				{
-					return false; // still locked by prerequisites, or another thread managed to set execution flag first, or we're inside this task execution
-					// we could try to help with nested tasks execution (the task execution could already spawned a couple of nested tasks sitting in the queue). 
-					// it's unclear how important this is, but this would definitely lead to more complicated impl. we can revisit this once we see such instances in profiler captures
-				}
-
-				TRACE_CPUPROFILER_EVENT_SCOPE(SuccessfulTaskRetraction);
-
-				// the task was launched so the scheduler will handle the internal reference held by low-level task
-
-				if (IsCompleted()) // still can be hold back by nested tasks, this is an optional early out for better perf
-				{
-					return true;
-				}
-
-				// retract nested tasks, if any
-				{
-					// keep trying retracting all nested tasks even if some of them fail, so the current worker can contribute instead of being blocked
-					bool bSucceeded = true;
-					// prerequisites are "consumed" here even if their retraction fails. this means that once prerequisite retraction failed, it won't be performed again. 
-					// this can be potentially improved by using a different container for prerequisites
-					while (FTaskBase* Prerequisite = Prerequisites.Pop())
-					{
-						if (!Prerequisite->TryRetractAndExecute(RecursionDepth))
-						{
-							bSucceeded = false;
-						}
-						Prerequisite->Release();
-					}
-
-					if (!bSucceeded)
-					{
-						return false;
-					}
-				}
-
-				// it happens that all nested tasks are completed and are in the process of completing the parent (this task) concurrently, 
-				// but the flag is not set yet. wait for it to maintain postconditions
-				while (!IsCompleted())
-				{
-					FPlatformProcess::Yield();
-				}
-
-				return true;
-			}
+			CORE_API bool TryRetractAndExecute(uint32 RecursionDepth = 0);
 
 			// releases internal reference and maintains low-level task state. must be called iff the task was never launched, otherwise 
 			// the scheduler will do this in due course
