Move FThreadStatsPool::Get() static member function to cpp file to avoid linker problem complaining about:
error: undefined symbol: guard variable for FThreadStatsPool::Get()::Singleton

diff --git a/Engine/Source/Runtime/Core/Private/Stats/Stats2.cpp b/Engine/Source/Runtime/Core/Private/Stats/Stats2.cpp
--- a/Engine/Source/Runtime/Core/Private/Stats/Stats2.cpp
+++ b/Engine/Source/Runtime/Core/Private/Stats/Stats2.cpp
@@ -1219,6 +1219,12 @@ FThreadStatsPool::FThreadStatsPool()
 	}
 }
 
+FThreadStatsPool& FThreadStatsPool::Get()
+{
+	static FThreadStatsPool Singleton;
+	return Singleton;
+}
+
 FThreadStats* FThreadStatsPool::GetFromPool()
 {
 	LLM_SCOPE(ELLMTag::Stats);
diff --git a/Engine/Source/Runtime/Core/Public/Stats/Stats2.h b/Engine/Source/Runtime/Core/Public/Stats/Stats2.h
--- a/Engine/Source/Runtime/Core/Public/Stats/Stats2.h
+++ b/Engine/Source/Runtime/Core/Public/Stats/Stats2.h
@@ -1273,11 +1273,7 @@ public:
 	CORE_API FThreadStatsPool();
 
 	/** Singleton accessor. */
-	CORE_API static FThreadStatsPool& Get()
-	{
-		static FThreadStatsPool Singleton;
-		return Singleton;
-	}
+	CORE_API static FThreadStatsPool& Get();
 
 	/** Gets an instance from the pool and call the default constructor on it. */
 	CORE_API FThreadStats* GetFromPool();
