Process control:
- Implemented priority control on Unix platform
- Refactored RunnableThreadUnix implementation
- Bump resource limits for process scheduling priority (nice) if possible

diff --git a/Engine/Source/Programs/UnrealBuildTool/Configuration/TargetRules.cs b/Engine/Source/Programs/UnrealBuildTool/Configuration/TargetRules.cs
--- a/Engine/Source/Programs/UnrealBuildTool/Configuration/TargetRules.cs
+++ b/Engine/Source/Programs/UnrealBuildTool/Configuration/TargetRules.cs
@@ -1599,6 +1599,13 @@ namespace UnrealBuildTool
 		/// </summary>
 		public bool bEnableCppCoroutinesForEvaluation = false;
 
+		/// <summary>
+		/// Whether to enable engine's ability to set process priority on runtime.
+		/// This option requires some environment setup on Linux, that's why it's disabled by default.
+		/// Project has to opt-in this feature as it has to guarantee correct setup.
+		/// </summary>
+		public bool bEnableProcessPriorityControl = false;
+
 		/// <summary>
 		/// If true, then enable memory profiling in the build (defines USE_MALLOC_PROFILER=1 and forces bOmitFramePointers=false).
 		/// </summary>
@@ -3222,6 +3229,11 @@ namespace UnrealBuildTool
 			get { return Inner.bEnableCppCoroutinesForEvaluation; }
 		}
 
+		public bool bEnableProcessPriorityControl
+		{
+			get { return Inner.bEnableProcessPriorityControl; }
+		}
+
 		public bool bUseMallocProfiler
 		{
 			get { return Inner.bUseMallocProfiler; }
diff --git a/Engine/Source/Programs/UnrealBuildTool/Configuration/UEBuildTarget.cs b/Engine/Source/Programs/UnrealBuildTool/Configuration/UEBuildTarget.cs
--- a/Engine/Source/Programs/UnrealBuildTool/Configuration/UEBuildTarget.cs
+++ b/Engine/Source/Programs/UnrealBuildTool/Configuration/UEBuildTarget.cs
@@ -4219,6 +4219,15 @@ namespace UnrealBuildTool
 				GlobalCompileEnvironment.Definitions.Add("WITH_CPP_COROUTINES=0");
 			}
 
+			if (Rules.bEnableProcessPriorityControl)
+			{
+				GlobalCompileEnvironment.Definitions.Add("WITH_PROCESS_PRIORITY_CONTROL=1");
+			}
+			else
+			{
+				GlobalCompileEnvironment.Definitions.Add("WITH_PROCESS_PRIORITY_CONTROL=0");
+			}
+
 			// Define the custom config if specified, this should not be set in an editor build
 			if (!bUseSharedBuildEnvironment && !String.IsNullOrEmpty(Rules.CustomConfig))
 			{
diff --git a/Engine/Source/Runtime/Core/Private/Unix/UnixPlatformProcess.cpp b/Engine/Source/Runtime/Core/Private/Unix/UnixPlatformProcess.cpp
--- a/Engine/Source/Runtime/Core/Private/Unix/UnixPlatformProcess.cpp
+++ b/Engine/Source/Runtime/Core/Private/Unix/UnixPlatformProcess.cpp
@@ -2057,3 +2057,66 @@ void FUnixPlatformProcess::OnChildEndFramePostFork()
 
 	FCoreDelegates::OnChildEndFramePostFork.Broadcast();
 }
+
+int32 FUnixPlatformProcess::TranslateThreadPriority(EThreadPriority Priority)
+{
+	// In general, the range is -20 to 19 (negative is highest, positive is lowest)
+	int32 NiceLevel = 0;
+	switch (Priority)
+	{
+		case TPri_TimeCritical:
+			NiceLevel = -20;
+			break;
+
+		case TPri_Highest:
+			NiceLevel = -15;
+			break;
+
+		case TPri_AboveNormal:
+			NiceLevel = -10;
+			break;
+
+		case TPri_Normal:
+			NiceLevel = 0;
+			break;
+
+		case TPri_SlightlyBelowNormal:
+			NiceLevel = 3;
+			break;
+
+		case TPri_BelowNormal:
+			NiceLevel = 5;
+			break;
+
+		case TPri_Lowest:
+			NiceLevel = 10;		// 19 is a total starvation
+			break;
+
+		default:
+			UE_LOG(LogHAL, Fatal, TEXT("Unknown Priority passed to FRunnableThreadPThread::TranslateThreadPriority()"));
+			return 0;
+	}
+
+	// note: a non-privileged process can only go as low as RLIMIT_NICE
+	return NiceLevel;
+}
+
+void FUnixPlatformProcess::SetThreadNiceValue(uint32_t ThreadId, int32 NiceValue)
+{
+	// We still try to set priority, but failure is not considered as an error
+	if (setpriority(PRIO_PROCESS, ThreadId, NiceValue) != 0 && WITH_PROCESS_PRIORITY_CONTROL)
+	{
+		// Unfortunately this is going to be a frequent occurence given that by default Unix doesn't allow raising priorities.
+		// NOTE: In WSL run "sudo prlimit --nice=40 --pid $$" to promote current shell to change nice values.
+		int ErrNo = errno;
+		UE_LOG(LogHAL, Error, TEXT("Can't set nice to %d. Reason = %s. Do you have CAP_SYS_NICE capability?"), NiceValue, ANSI_TO_TCHAR(strerror(ErrNo)));
+	}
+}
+
+void FUnixPlatformProcess::SetThreadPriority(EThreadPriority NewPriority)
+{
+	uint32 ThreadId = FPlatformTLS::GetCurrentThreadId();
+	int32 NiceValue = FUnixPlatformProcess::TranslateThreadPriority(NewPriority);
+
+	FUnixPlatformProcess::SetThreadNiceValue(ThreadId, NiceValue);
+}
diff --git a/Engine/Source/Runtime/Core/Private/Unix/UnixPlatformRunnableThread.cpp b/Engine/Source/Runtime/Core/Private/Unix/UnixPlatformRunnableThread.cpp
new file mode 100644
--- /dev/null
+++ b/Engine/Source/Runtime/Core/Private/Unix/UnixPlatformRunnableThread.cpp
@@ -0,0 +1,189 @@
+// Copyright Epic Games, Inc. All Rights Reserved.
+
+#include "Unix/UnixPlatformRunnableThread.h"
+#include "Unix/UnixPlatformProcess.h"
+#include "HAL/UnrealMemory.h"
+#include "Templates/UnrealTemplate.h"
+#include "Containers/UnrealString.h"
+#include "Containers/StringConv.h"
+#include "Logging/LogMacros.h"
+#include "Misc/AssertionMacros.h"
+
+#include <sys/resource.h>
+
+bool FRunnableThreadUnix::SetupSignalHandlerStack(void* StackBuffer, const size_t StackBufferSize, void** OutStackGuardPageAddress)
+{
+	if (!StackBuffer)
+	{
+		return false;
+	}
+
+	// Added by CL 11188846 for ASan, TSan, and UBSan
+	// #jira UE-62784 UE-62803 UE-62804
+	if (FORCE_ANSI_ALLOCATOR)
+	{
+		return true;
+	}
+
+	// find an address close to begin of the stack and protect it
+	uint64 StackGuardPage = reinterpret_cast<uint64>(StackBuffer);
+
+	// align by page
+	const uint64 PageSize = sysconf(_SC_PAGESIZE);
+	const uint64 Remainder = StackGuardPage % PageSize;
+	if (Remainder != 0)
+	{
+		StackGuardPage += (PageSize - Remainder);
+		checkf(StackGuardPage % PageSize == 0, TEXT("StackGuardPage is not aligned on page size"));
+	}
+
+	checkf(StackGuardPage + PageSize - reinterpret_cast<uint64>(StackBuffer) < StackBufferSize,
+		TEXT("Stack size is too small for the extra guard page!"));
+
+	void* StackGuardPageAddr = reinterpret_cast<void*>(StackGuardPage);
+	if (FPlatformMemory::PageProtect(StackGuardPageAddr, PageSize, true, false))
+	{
+		if (OutStackGuardPageAddress)
+		{
+			*OutStackGuardPageAddress = StackGuardPageAddr;
+		}
+	}
+	else
+	{
+		// cannot use UE_LOG - can run into deadlocks in output device code
+		fprintf(stderr, "Unable to set a guard page on the alt stack\n");
+	}
+
+	// set up the buffer to be used as stack
+	stack_t SignalHandlerStack;
+	FMemory::Memzero(SignalHandlerStack);
+	SignalHandlerStack.ss_sp = StackBuffer;
+	SignalHandlerStack.ss_size = StackBufferSize;
+
+	bool bSuccess = (sigaltstack(&SignalHandlerStack, nullptr) == 0);
+	if (!bSuccess)
+	{
+		int ErrNo = errno;
+		// cannot use UE_LOG - can run into deadlocks in output device code
+		fprintf(stderr, "Unable to set alternate stack for crash handler, sigaltstack() failed with errno=%d (%s)\n", ErrNo, strerror(ErrNo));
+	}
+
+	return bSuccess;
+}
+
+int32 FRunnableThreadUnix::TranslateThreadPriority(EThreadPriority Priority)
+{
+	return FUnixPlatformProcess::TranslateThreadPriority(Priority);
+}
+
+void FRunnableThreadUnix::SetThreadPriority(EThreadPriority NewPriority)
+{
+	// always set priority to avoid second guessing
+	ThreadPriority = NewPriority;
+	SetThreadPriority(Thread, NewPriority);
+}
+
+void FRunnableThreadUnix::SetThreadPriority(pthread_t InThread, EThreadPriority NewPriority)
+{
+	// NOTE: InThread is ignored, but we can use ThreadID that maps to SYS_ttid
+	int32 Prio = TranslateThreadPriority(NewPriority);
+
+	// Unix implements thread priorities the same way as process priorities, while on Windows they are relative to process priority.
+	// We want Windows behavior, since sometimes we set the whole process to a lower priority and would like its threads to avoid raising it
+	// (even if RTLIMIT_NICE allows it) - example is ShaderCompileWorker that need to run in the background.
+	//
+	// Thusly we remember the baseline value that the process has at the moment of first priority change and set thread priority relative to it.
+	// This is of course subject to race conditions and other problems (e.g. in case main thread changes its priority after the fact), but it's the best we have.
+
+	if (!bGotBaselineNiceValue)
+	{
+		// checking errno is necessary since -1 is a valid priority to return from getpriority()
+		errno = 0;
+		int CurrentPriority = getpriority(PRIO_PROCESS, getpid());
+		// if getting priority wasn't successful, don't change the baseline value (will be 0 - i.e. normal - by default)
+		if (CurrentPriority != -1 || errno == 0)
+		{
+			BaselineNiceValue = CurrentPriority;
+			bGotBaselineNiceValue = true;
+		}
+	}
+
+	int ModifiedPriority = FMath::Clamp(BaselineNiceValue + Prio, -20, 19);
+
+	FUnixPlatformProcess::SetThreadNiceValue(ThreadID, ModifiedPriority);
+}
+
+void FRunnableThreadUnix::PreRun()
+{
+	FString SizeLimitedThreadName = ThreadName;
+
+	if (SizeLimitedThreadName.Len() > EConstants::UnixThreadNameLimit)
+	{
+		// first, attempt to cut out common and meaningless substrings
+		SizeLimitedThreadName = SizeLimitedThreadName.Replace(TEXT("Thread"), TEXT(""));
+		SizeLimitedThreadName = SizeLimitedThreadName.Replace(TEXT("Runnable"), TEXT(""));
+
+		// if still larger
+		if (SizeLimitedThreadName.Len() > EConstants::UnixThreadNameLimit)
+		{
+			FString Temp = SizeLimitedThreadName;
+
+			// cut out the middle and replace with a substitute
+			const TCHAR Dash[] = TEXT("-");
+			const int32 DashLen = UE_ARRAY_COUNT(Dash) - 1;
+			int NumToLeave = (EConstants::UnixThreadNameLimit - DashLen) / 2;
+
+			SizeLimitedThreadName = Temp.Left(EConstants::UnixThreadNameLimit - (NumToLeave + DashLen));
+			SizeLimitedThreadName += Dash;
+			SizeLimitedThreadName += Temp.Right(NumToLeave);
+
+			check(SizeLimitedThreadName.Len() <= EConstants::UnixThreadNameLimit);
+		}
+	}
+
+	int ErrCode = pthread_setname_np(Thread, TCHAR_TO_ANSI(*SizeLimitedThreadName));
+	if (ErrCode != 0)
+	{
+		UE_LOG(LogHAL, Warning, TEXT("pthread_setname_np(, '%s') failed with error %d (%s)."), *ThreadName, ErrCode, ANSI_TO_TCHAR(strerror(ErrCode)));
+	}
+
+	// set the alternate stack for handling crashes due to stack overflow
+	check(ThreadCrashHandlingStack == nullptr);
+	ThreadCrashHandlingStack = AllocCrashHandlerStack();
+	SetupSignalHandlerStack(ThreadCrashHandlingStack, FRunnableThreadUnix::GetCrashHandlerStackSize(), &StackGuardPageAddress);
+}
+
+void FRunnableThreadUnix::PostRun()
+{
+	if (StackGuardPageAddress != nullptr)
+	{
+		// we protected one page only
+		const uint64 PageSize = sysconf(_SC_PAGESIZE);
+
+		if (!FPlatformMemory::PageProtect(StackGuardPageAddress, PageSize, true, true))
+		{
+			UE_LOG(LogCore, Error, TEXT("Unable to remove a guard page from the alt stack"));
+		}
+
+		StackGuardPageAddress = nullptr;
+	}
+
+	if (ThreadCrashHandlingStack != nullptr)
+	{
+		FreeCrashHandlerStack(ThreadCrashHandlingStack);
+		ThreadCrashHandlingStack = nullptr;
+	}
+}
+
+uint32 FRunnableThreadUnix::AdjustStackSize(uint32 InStackSize)
+{
+	InStackSize = FRunnableThreadPThread::AdjustStackSize(InStackSize);
+
+	// If it's set, make sure it's at least 128 KB or stack allocations (e.g. in Logf) may fail
+	if (InStackSize && InStackSize < 128 * 1024)
+	{
+		InStackSize = 128 * 1024;
+	}
+
+	return InStackSize;
+}
diff --git a/Engine/Source/Runtime/Core/Public/Misc/Build.h b/Engine/Source/Runtime/Core/Public/Misc/Build.h
--- a/Engine/Source/Runtime/Core/Public/Misc/Build.h
+++ b/Engine/Source/Runtime/Core/Public/Misc/Build.h
@@ -468,3 +468,8 @@
 
 // Controls if the config system can stores configs for other platforms than the running one
 #define ALLOW_OTHER_PLATFORM_CONFIG		WITH_UNREAL_DEVELOPER_TOOLS
+
+// Controls whether or not process will control OS scheduler priority
+#ifndef WITH_PROCESS_PRIORITY_CONTROL
+	#define WITH_PROCESS_PRIORITY_CONTROL 0
+#endif
diff --git a/Engine/Source/Runtime/Core/Public/Unix/UnixPlatformProcess.h b/Engine/Source/Runtime/Core/Public/Unix/UnixPlatformProcess.h
--- a/Engine/Source/Runtime/Core/Public/Unix/UnixPlatformProcess.h
+++ b/Engine/Source/Runtime/Core/Public/Unix/UnixPlatformProcess.h
@@ -301,6 +301,9 @@ struct CORE_API FUnixPlatformProcess : public FGenericPlatformProcess
 	static bool LaunchFileInDefaultExternalApplication( const TCHAR* FileName, const TCHAR* Parms = NULL, ELaunchVerb::Type Verb = ELaunchVerb::Open, bool bPromptToOpenOnFailure = true);
 	static bool IsFirstInstance();
 	static void OnChildEndFramePostFork();
+	static int32 TranslateThreadPriority(EThreadPriority Priority);
+	static void SetThreadNiceValue( uint32_t ThreadId, int32 NiceValue );
+	static void SetThreadPriority( EThreadPriority NewPriority );
 
 	/**
 	 * @brief Releases locks that we held for IsFirstInstance check
diff --git a/Engine/Source/Runtime/Core/Public/Unix/UnixPlatformRunnableThread.h b/Engine/Source/Runtime/Core/Public/Unix/UnixPlatformRunnableThread.h
--- a/Engine/Source/Runtime/Core/Public/Unix/UnixPlatformRunnableThread.h
+++ b/Engine/Source/Runtime/Core/Public/Unix/UnixPlatformRunnableThread.h
@@ -7,14 +7,8 @@
 #pragma once
 
 #include "CoreTypes.h"
-#include "Misc/AssertionMacros.h"
-#include "HAL/UnrealMemory.h"
-#include "Templates/UnrealTemplate.h"
-#include "Containers/UnrealString.h"
-#include "Containers/StringConv.h"
-#include "Logging/LogMacros.h"
-#include "Runtime/Core/Private/HAL/PThreadRunnableThread.h"
-#include <sys/resource.h>
+#include "HAL/PThreadRunnableThread.h"
+#include <signal.h> // for SIGSTKSZ
 
 class Error;
 
@@ -82,232 +76,28 @@ public:
 	 *
 	 * @return true if setting the alt stack succeeded. Inability to set guard page will not affect success of the operation.
 	 */
-	static bool SetupSignalHandlerStack(void* StackBuffer, const size_t StackBufferSize, void** OutStackGuardPageAddress)
-	{
-		if (!StackBuffer)
-		{
-			return false;
-		}
-
-		// Added by CL 11188846 for ASan, TSan, and UBSan
-		// #jira UE-62784 UE-62803 UE-62804
-		if (FORCE_ANSI_ALLOCATOR)
-		{
-			return true;
-		}
-
-		// find an address close to begin of the stack and protect it
-		uint64 StackGuardPage = reinterpret_cast<uint64>(StackBuffer);
-
-		// align by page
-		const uint64 PageSize = sysconf(_SC_PAGESIZE);
-		const uint64 Remainder = StackGuardPage % PageSize;
-		if (Remainder != 0)
-		{
-			StackGuardPage += (PageSize - Remainder);
-			checkf(StackGuardPage % PageSize == 0, TEXT("StackGuardPage is not aligned on page size"));
-		}
-
-		checkf(StackGuardPage + PageSize - reinterpret_cast<uint64>(StackBuffer) < StackBufferSize,
-			TEXT("Stack size is too small for the extra guard page!"));
-
-		void* StackGuardPageAddr = reinterpret_cast<void*>(StackGuardPage);
-		if (FPlatformMemory::PageProtect(StackGuardPageAddr, PageSize, true, false))
-		{
-			if (OutStackGuardPageAddress)
-			{
-				*OutStackGuardPageAddress = StackGuardPageAddr;
-			}
-		}
-		else
-		{
-			// cannot use UE_LOG - can run into deadlocks in output device code
-			fprintf(stderr, "Unable to set a guard page on the alt stack\n");
-		}
-
-		// set up the buffer to be used as stack
-		stack_t SignalHandlerStack;
-		FMemory::Memzero(SignalHandlerStack);
-		SignalHandlerStack.ss_sp = StackBuffer;
-		SignalHandlerStack.ss_size = StackBufferSize;
-
-		bool bSuccess = (sigaltstack(&SignalHandlerStack, nullptr) == 0);
-		if (!bSuccess)
-		{
-			int ErrNo = errno;
-			// cannot use UE_LOG - can run into deadlocks in output device code
-			fprintf(stderr, "Unable to set alternate stack for crash handler, sigaltstack() failed with errno=%d (%s)\n", ErrNo, strerror(ErrNo));
-		}
-
-		return bSuccess;
-	}
+	static bool SetupSignalHandlerStack(void* StackBuffer, const size_t StackBufferSize, void** OutStackGuardPageAddress);
 
 protected:
 
 	/** on Unix, this translates to ranges of setpriority(). Note that not all range may be available*/
-	virtual int32 TranslateThreadPriority(EThreadPriority Priority)
-	{
-		// In general, the range is -20 to 19 (negative is highest, positive is lowest)
-		int32 NiceLevel = 0;
-		switch (Priority)
-		{
-			case TPri_TimeCritical:
-				NiceLevel = -20;
-				break;
-
-			case TPri_Highest:
-				NiceLevel = -15;
-				break;
-
-			case TPri_AboveNormal:
-				NiceLevel = -10;
-				break;
-
-			case TPri_Normal:
-				NiceLevel = 0;
-				break;
-
-			case TPri_SlightlyBelowNormal:
-				NiceLevel = 3;
-				break;
-
-			case TPri_BelowNormal:
-				NiceLevel = 5;
-				break;
+	int32 TranslateThreadPriority(EThreadPriority Priority) override;
 
-			case TPri_Lowest:
-				NiceLevel = 10;		// 19 is a total starvation
-				break;
+	void SetThreadPriority(EThreadPriority NewPriority) override;
 
-			default:
-				UE_LOG(LogHAL, Fatal, TEXT("Unknown Priority passed to FRunnableThreadPThread::TranslateThreadPriority()"));
-				return 0;
-		}
-
-		// note: a non-privileged process can only go as low as RLIMIT_NICE
-		return NiceLevel;
-	}
-
-	virtual void SetThreadPriority(EThreadPriority NewPriority) override
-	{
-		// always set priority to avoid second guessing
-		ThreadPriority = NewPriority;
-		SetThreadPriority(Thread, NewPriority);
-	}
-
-	virtual void SetThreadPriority(pthread_t InThread, EThreadPriority NewPriority)
-	{
-		// NOTE: InThread is ignored, but we can use ThreadID that maps to SYS_ttid
-		int32 Prio = TranslateThreadPriority(NewPriority);
-
-		// Unix implements thread priorities the same way as process priorities, while on Windows they are relative to process priority.
-		// We want Windows behavior, since sometimes we set the whole process to a lower priority and would like its threads to avoid raising it
-		// (even if RTLIMIT_NICE allows it) - example is ShaderCompileWorker that need to run in the background.
-		//
-		// Thusly we remember the baseline value that the process has at the moment of first priority change and set thread priority relative to it.
-		// This is of course subject to race conditions and other problems (e.g. in case main thread changes its priority after the fact), but it's the best we have.
-
-		if (!bGotBaselineNiceValue)
-		{
-			// checking errno is necessary since -1 is a valid priority to return from getpriority()
-			errno = 0;
-			int CurrentPriority = getpriority(PRIO_PROCESS, getpid());
-			// if getting priority wasn't successful, don't change the baseline value (will be 0 - i.e. normal - by default)
-			if (CurrentPriority != -1 || errno == 0)
-			{
-				BaselineNiceValue = CurrentPriority;
-				bGotBaselineNiceValue = true;
-			}
-		}
-
-		int ModifiedPriority = FMath::Clamp(BaselineNiceValue + Prio, -20, 19);
-		if (setpriority(PRIO_PROCESS, ThreadID, ModifiedPriority) != 0)
-		{
-			// Unfortunately this is going to be a frequent occurence given that by default Unix doesn't allow raising priorities.
-			// Don't issue a warning here
-		}
-	}
+	void SetThreadPriority(pthread_t InThread, EThreadPriority NewPriority) override;
 
 private:
 
 	/**
 	 * Allows a platform subclass to setup anything needed on the thread before running the Run function
 	 */
-	virtual void PreRun() override
-	{
-		FString SizeLimitedThreadName = ThreadName;
+	void PreRun() override;
 
-		if (SizeLimitedThreadName.Len() > EConstants::UnixThreadNameLimit)
-		{
-			// first, attempt to cut out common and meaningless substrings
-			SizeLimitedThreadName = SizeLimitedThreadName.Replace(TEXT("Thread"), TEXT(""));
-			SizeLimitedThreadName = SizeLimitedThreadName.Replace(TEXT("Runnable"), TEXT(""));
-
-			// if still larger
-			if (SizeLimitedThreadName.Len() > EConstants::UnixThreadNameLimit)
-			{
-				FString Temp = SizeLimitedThreadName;
-
-				// cut out the middle and replace with a substitute
-				const TCHAR Dash[] = TEXT("-");
-				const int32 DashLen = UE_ARRAY_COUNT(Dash) - 1;
-				int NumToLeave = (EConstants::UnixThreadNameLimit - DashLen) / 2;
-
-				SizeLimitedThreadName = Temp.Left(EConstants::UnixThreadNameLimit - (NumToLeave + DashLen));
-				SizeLimitedThreadName += Dash;
-				SizeLimitedThreadName += Temp.Right(NumToLeave);
-
-				check(SizeLimitedThreadName.Len() <= EConstants::UnixThreadNameLimit);
-			}
-		}
-
-		int ErrCode = pthread_setname_np(Thread, TCHAR_TO_ANSI(*SizeLimitedThreadName));
-		if (ErrCode != 0)
-		{
-			UE_LOG(LogHAL, Warning, TEXT("pthread_setname_np(, '%s') failed with error %d (%s)."), *ThreadName, ErrCode, ANSI_TO_TCHAR(strerror(ErrCode)));
-		}
-
-		// set the alternate stack for handling crashes due to stack overflow
-		check(ThreadCrashHandlingStack == nullptr);
-		ThreadCrashHandlingStack = AllocCrashHandlerStack();
-		SetupSignalHandlerStack(ThreadCrashHandlingStack, FRunnableThreadUnix::GetCrashHandlerStackSize(), &StackGuardPageAddress);
-	}
-
-	virtual void PostRun() override
-	{
-		if (StackGuardPageAddress != nullptr)
-		{
-			// we protected one page only
-			const uint64 PageSize = sysconf(_SC_PAGESIZE);
-
-			if (!FPlatformMemory::PageProtect(StackGuardPageAddress, PageSize, true, true))
-			{
-				UE_LOG(LogCore, Error, TEXT("Unable to remove a guard page from the alt stack"));
-			}
-
-			StackGuardPageAddress = nullptr;
-		}
-
-		if (ThreadCrashHandlingStack != nullptr)
-		{
-			FreeCrashHandlerStack(ThreadCrashHandlingStack);
-			ThreadCrashHandlingStack = nullptr;
-		}
-	}
+	void PostRun() override;
 
 	/**
 	 * Allows platforms to adjust stack size
 	 */
-	virtual uint32 AdjustStackSize(uint32 InStackSize)
-	{
-		InStackSize = FRunnableThreadPThread::AdjustStackSize(InStackSize);
-
-		// If it's set, make sure it's at least 128 KB or stack allocations (e.g. in Logf) may fail
-		if (InStackSize && InStackSize < 128 * 1024)
-		{
-			InStackSize = 128 * 1024;
-		}
-
-		return InStackSize;
-	}
+	uint32 AdjustStackSize(uint32 InStackSize) override;
 };
diff --git a/Engine/Source/Runtime/Unix/UnixCommonStartup/Private/UnixCommonStartup.cpp b/Engine/Source/Runtime/Unix/UnixCommonStartup/Private/UnixCommonStartup.cpp
--- a/Engine/Source/Runtime/Unix/UnixCommonStartup/Private/UnixCommonStartup.cpp
+++ b/Engine/Source/Runtime/Unix/UnixCommonStartup/Private/UnixCommonStartup.cpp
@@ -183,6 +183,12 @@ static bool IncreasePerProcessLimits()
 		}
 	}
 
+	if constexpr (WITH_PROCESS_PRIORITY_CONTROL)
+	{
+		printf("Increasing per-process limit for scheduling priority.\n");
+		SetResourceMaxHardLimit(RLIMIT_NICE);
+	}
+
 	return true;
 }
 
