Trace:
* Replaced OnTraceFileStarted delegate with OnTraceStarted
* Added OnTraceStopped delegate
* Both events pass if the recording was for a network or file type.

diff --git a/Engine/Source/Runtime/Core/Private/ProfilingDebugging/TraceAuxiliary.cpp b/Engine/Source/Runtime/Core/Private/ProfilingDebugging/TraceAuxiliary.cpp
--- a/Engine/Source/Runtime/Core/Private/ProfilingDebugging/TraceAuxiliary.cpp
+++ b/Engine/Source/Runtime/Core/Private/ProfilingDebugging/TraceAuxiliary.cpp
@@ -129,6 +129,7 @@ private:
 	typedef TMap<uint32, FChannelEntry, TInlineSetAllocator<128>> ChannelSet;
 	ChannelSet				CommandlineChannels;
 	FString					TraceDest;
+	FTraceAuxiliary::EConnectionType TraceType = FTraceAuxiliary::EConnectionType::None;
 	EState					State = EState::Stopped;
 	bool					bTruncateFile = false;
 	bool					bWorkerThreadStarted = false;
@@ -289,6 +290,8 @@ bool FTraceAuxiliaryImpl::Connect(ETraceConnectType Type, const TCHAR* Parameter
 			{
 				UE_LOG_REF(LogCategory, Error, TEXT("Trace failed to connect (trace server: %s)!"), Parameter ? Parameter : TEXT(""));
 			}
+
+			TraceType = FTraceAuxiliary::EConnectionType::Network;
 		}
 
 		else if (Type == ETraceConnectType::File)
@@ -302,6 +305,13 @@ bool FTraceAuxiliaryImpl::Connect(ETraceConnectType Type, const TCHAR* Parameter
 			{
 				UE_LOG_REF(LogCategory, Error, TEXT("Trace failed to connect (file: \"%s\")!"), Parameter ? Parameter : TEXT(""));
 			}
+
+			TraceType = FTraceAuxiliary::EConnectionType::File;
+		}
+
+		if (bConnected)
+		{
+			FTraceAuxiliary::OnTraceStarted.Broadcast(TraceType, TraceDest);
 		}
 	}
 
@@ -327,7 +337,10 @@ bool FTraceAuxiliaryImpl::Stop()
 		return false;
 	}
 
+	FTraceAuxiliary::OnTraceStopped.Broadcast(TraceType, TraceDest);
+
 	State = EState::Stopped;
+	TraceType = FTraceAuxiliary::EConnectionType::None;
 	TraceDest.Reset();
 	return true;
 }
@@ -517,11 +530,7 @@ bool FTraceAuxiliaryImpl::WriteToFile(const TCHAR* Path, const FTraceAuxiliary::
 		return false;
 	}
 
-	if (UE::Trace::WriteTo(*NativePath))
-	{
-		FTraceAuxiliary::OnTraceFileCreated.Broadcast(*NativePath);
-	}
-	else
+	if (!UE::Trace::WriteTo(*NativePath))
 	{
 		if (FPathViews::Equals(NativePath, FStringView(Path)))
 		{
@@ -1248,7 +1257,8 @@ static bool StartFromCommandlineArguments(const TCHAR* CommandLine)
 
 ////////////////////////////////////////////////////////////////////////////////
 FTraceAuxiliary::FOnConnection FTraceAuxiliary::OnConnection;
-FTraceAuxiliary::FOnTraceFileCreated FTraceAuxiliary::OnTraceFileCreated;
+FTraceAuxiliary::FOnTraceStarted FTraceAuxiliary::OnTraceStarted;
+FTraceAuxiliary::FOnTraceStopped FTraceAuxiliary::OnTraceStopped;
 
 ////////////////////////////////////////////////////////////////////////////////
 bool FTraceAuxiliary::Start(EConnectionType Type, const TCHAR* Target, const TCHAR* Channels, FOptions* Options, const FLogCategoryAlias& LogCategory)
diff --git a/Engine/Source/Runtime/Core/Public/ProfilingDebugging/TraceAuxiliary.h b/Engine/Source/Runtime/Core/Public/ProfilingDebugging/TraceAuxiliary.h
--- a/Engine/Source/Runtime/Core/Public/ProfilingDebugging/TraceAuxiliary.h
+++ b/Engine/Source/Runtime/Core/Public/ProfilingDebugging/TraceAuxiliary.h
@@ -23,15 +23,7 @@ public:
 	typedef FLogCategoryBase FLogCategoryAlias;
 #endif
 
-	/**
-	 * Callback type when a new connection is established.
-	 */
-	DECLARE_MULTICAST_DELEGATE(FOnConnection);
-
-	/** Callback whenever a trace file is created */
-	DECLARE_MULTICAST_DELEGATE_OneParam(FOnTraceFileCreated, const TCHAR* FileNameAndPath);
-
-	enum class EConnectionType
+	enum class EConnectionType : uint8
 	{
 		/**
 		 * Connect to a trace server. Target is IP address or hostname.
@@ -41,13 +33,28 @@ public:
 		 * Write to a file. Target string is filename. Absolute or relative current working directory.
 		 * If target is null the current date and time is used.
 		 */
-		File,
-		/**
-		 * Don't connect, just start tracing to memory.
-		 */
-		None,
+		 File,
+		 /**
+		  * Don't connect, just start tracing to memory.
+		  */
+		  None,
 	};
 
+	/**
+	 * Callback type when a new connection is established.
+	 */
+	DECLARE_MULTICAST_DELEGATE(FOnConnection);
+
+	/** Callback whenever a trace is started */
+	DECLARE_MULTICAST_DELEGATE_TwoParams(FOnTraceStarted, FTraceAuxiliary::EConnectionType TraceType, const FString& TraceDestination);
+
+	/** 
+	* Callback whenever a trace recording is stopped. 
+	* TraceType tells what kind of trace it is.
+	* TraceDestination will be either the the filename and path for a file trace or the network connection for a network trace
+	*/
+	DECLARE_MULTICAST_DELEGATE_TwoParams(FOnTraceStopped, FTraceAuxiliary::EConnectionType TraceType, const FString& TraceDestination);
+
 	struct FOptions
 	{
 		/** When set, trace will not start a worker thread, instead it is updated from end frame delegate. */
@@ -157,10 +164,15 @@ public:
 	static FOnConnection OnConnection;
 
 	/**
-	 * Delegate that triggers when a trace file is created. Useful if you need to collect all trace files
-     * created during a session.
-	 * The full file name and path is passed to the delegate.
+	 * Delegate that triggers when a trace session is started.
+	 * The type of recording and the destination (filepath or network) is passed to the delegate.
+	 */
+	static FOnTraceStarted OnTraceStarted;
+
+	/**
+	 * Delegate that triggers when a trace has finished recording. Useful if you need to collect all completed trace files in a session.
+	 * The type of recording and the destination (filepath or network) is passed to the delegate.
 	 */
-	static FOnTraceFileCreated OnTraceFileCreated;
+	static FOnTraceStopped OnTraceStopped;
 	
 };
