Workaround linux compiler failing to expand out nested UE_LOG_REF/CLINKAGE when NO_LOGGING is true and there are no parameters to the log message
#horde 214762
#rnx

diff --git a/Engine/Source/Runtime/Core/Public/Logging/LogMacros.h b/Engine/Source/Runtime/Core/Public/Logging/LogMacros.h
--- a/Engine/Source/Runtime/Core/Public/Logging/LogMacros.h
+++ b/Engine/Source/Runtime/Core/Public/Logging/LogMacros.h
@@ -114,9 +114,22 @@ private:
 			UE_LOG_EXPAND_IS_FATAL(Verbosity, CA_ASSUME(false);, PREPROCESSOR_NOTHING) \
 		} \
 	}
-	#define UE_LOG_REF(ConstLogCategoryRef, Verbosity, Format, ...) UE_LOG(ConstLogCategoryRef, Verbosity, Format, __VA_ARGS__ )
-
-	#define UE_LOG_CLINKAGE(CategoryName, Verbosity, Format, ...) UE_LOG(CategoryName, Verbosity, Format, __VA_ARGS__ )
+	#define UE_LOG_REF(ConstLogCategoryRef, Verbosity, Format, ...) \
+	{ \
+		if (ELogVerbosity::Verbosity == ELogVerbosity::Fatal) \
+		{ \
+			LowLevelFatalError(Format, ##__VA_ARGS__); \
+			UE_LOG_EXPAND_IS_FATAL(Verbosity, CA_ASSUME(false);, PREPROCESSOR_NOTHING) \
+		} \
+	}
+	#define UE_LOG_CLINKAGE(CategoryName, Verbosity, Format, ...) \
+	{ \
+		if (ELogVerbosity::Verbosity == ELogVerbosity::Fatal) \
+		{ \
+			LowLevelFatalError(Format, ##__VA_ARGS__); \
+			UE_LOG_EXPAND_IS_FATAL(Verbosity, CA_ASSUME(false);, PREPROCESSOR_NOTHING) \
+		} \
+	}
 
 	// Conditional logging (fatal errors only).
 	#define UE_CLOG(Condition, CategoryName, Verbosity, Format, ...) \
