Fixed TraceAuxilary starting Trace's thread too early.

Trace's thread can't be started before Trace has been initialised as some internals are not ready prior to that (i.e. the cache and the tail). TraceAux attempts to control the worker thread and indirectly started it too early when processing command line arguments. This would cause data races between the worker thread and initialisation.

#rnx

diff --git a/Engine/Source/Runtime/Core/Private/ProfilingDebugging/TraceAuxiliary.cpp b/Engine/Source/Runtime/Core/Private/ProfilingDebugging/TraceAuxiliary.cpp
--- a/Engine/Source/Runtime/Core/Private/ProfilingDebugging/TraceAuxiliary.cpp
+++ b/Engine/Source/Runtime/Core/Private/ProfilingDebugging/TraceAuxiliary.cpp
@@ -1171,7 +1171,7 @@ UE_TRACE_EVENT_BEGIN(Diagnostics, Session2, NoSync|Important)
 UE_TRACE_EVENT_END()
 
 ////////////////////////////////////////////////////////////////////////////////
-static bool StartFromCommandlineArguments(const TCHAR* CommandLine)
+static bool StartFromCommandlineArguments(const TCHAR* CommandLine, bool& bOutStartWorkerThread)
 {
 #if UE_TRACE_ENABLED
 
@@ -1249,6 +1249,12 @@ static bool StartFromCommandlineArguments(const TCHAR* CommandLine)
 		return false;
 	}
 
+	// Trace's worker thread should really only be started by Trace itself as
+	// order is important. At the very least it must be done after Trace is
+	// initialised. It isn't yet here so we defer it.
+	bOutStartWorkerThread = !Opts.bNoWorkerThread;
+	Opts.bNoWorkerThread = true;
+
 	// Finally start tracing to the requested connection
 	return FTraceAuxiliary::Start(Type, Target, *Channels, &Opts);
 #endif
@@ -1447,13 +1453,14 @@ void FTraceAuxiliary::Initialize(const TCHAR* CommandLine)
 	// Attempt to send trace data somewhere from the command line. It perhaps
 	// seems odd to do this before initialising Trace, but it is done this way
 	// to support disabling the "important" cache without losing any events.
-	StartFromCommandlineArguments(CommandLine);
+	bool bShouldStartWorkerThread = false;
+	StartFromCommandlineArguments(CommandLine, bShouldStartWorkerThread);
 
 	// Initialize Trace
 	UE::Trace::FInitializeDesc Desc;
 	SetupInitFromConfig(Desc);
 	
-	Desc.bUseWorkerThread = false;
+	Desc.bUseWorkerThread = bShouldStartWorkerThread;
 	Desc.bUseImportantCache = (FParse::Param(CommandLine, TEXT("tracenocache")) == false);
 	Desc.OnConnectionFunc = &OnConnectionCallback;
 	if (FParse::Value(CommandLine, TEXT("-tracetailmb="), Desc.TailSizeBytes))
