Added "stats.NamedEvents 0/1" and "stats.VerboseNamedEvents 0/1" console commands to be able to force enable or disable the NamedEvents. The existing "stat NamedEvents" and "stat VerboseNamedEvents" are still available with "toggle state" behaviour. The "stats." prefix was choosen to mach the convention of other related console commands (like "stats.AutoEnableNamedEventsWhenProfiling").

diff --git a/Engine/Source/Runtime/Core/Private/Misc/CoreMisc.cpp b/Engine/Source/Runtime/Core/Private/Misc/CoreMisc.cpp
--- a/Engine/Source/Runtime/Core/Private/Misc/CoreMisc.cpp
+++ b/Engine/Source/Runtime/Core/Private/Misc/CoreMisc.cpp
@@ -527,3 +527,54 @@ void FAutoNamedEventsToggler::Update(bool bIsProfiling)
 		bSetNamedEventsEnabled = false;
 	}
 }
+
+static FAutoConsoleCommand GSetStatsNamedEventsCmd(
+	TEXT("stats.NamedEvents"),
+	TEXT("<on/off> Enables or disables the Named Events."),
+	FConsoleCommandWithWorldArgsAndOutputDeviceDelegate::CreateLambda([](const TArray<FString>& Args, UWorld*, FOutputDevice& Ar)
+	{
+		if (Args.Num() > 0)
+		{
+			bool bShouldEnableNamedEvents = FCString::ToBool(*Args[0]);
+			if (bShouldEnableNamedEvents)
+			{
+				if (GCycleStatsShouldEmitNamedEvents == 0)
+				{
+					GCycleStatsShouldEmitNamedEvents = 1;
+				}
+			}
+			else
+			{
+				if (GCycleStatsShouldEmitNamedEvents > 0)
+				{
+					GCycleStatsShouldEmitNamedEvents = 0;
+				}
+			}
+		}
+		else
+		{
+			Ar.Logf(TEXT("Expected 0/1 argument"));
+		}
+	})
+);
+
+static FAutoConsoleCommand GSetStatsVerboseNamedEventsCmd(
+	TEXT("stats.VerboseNamedEvents"),
+	TEXT("<on/off> Enables or disables the Verbose Named Events."),
+	FConsoleCommandWithWorldArgsAndOutputDeviceDelegate::CreateLambda([](const TArray<FString>& Args, UWorld*, FOutputDevice& Ar)
+	{
+		if (Args.Num() > 0)
+		{
+			GShouldEmitVerboseNamedEvents = FCString::ToBool(*Args[0]);
+
+			if (GShouldEmitVerboseNamedEvents && GCycleStatsShouldEmitNamedEvents == 0)
+			{
+				GCycleStatsShouldEmitNamedEvents = 1;
+			}
+		}
+		else
+		{
+			Ar.Logf(TEXT("Expected on/off argument"));
+		}
+	})
+);
