Adding deprecated GetCustomChunksByType function in base chunk installer to make some upcoming deprecations easier.

#rnx

diff --git a/Engine/Source/Runtime/Core/Public/GenericPlatform/GenericPlatformChunkInstall.h b/Engine/Source/Runtime/Core/Public/GenericPlatform/GenericPlatformChunkInstall.h
--- a/Engine/Source/Runtime/Core/Public/GenericPlatform/GenericPlatformChunkInstall.h
+++ b/Engine/Source/Runtime/Core/Public/GenericPlatform/GenericPlatformChunkInstall.h
@@ -218,6 +218,9 @@ public:
 	UE_DEPRECATED(5.2, "Call UninstallNamedChunks instead")
 	virtual bool UninstallChunks(const TArray<FCustomChunk>& ChunkTagsID) = 0;
 
+	UE_DEPRECATED(5.2, "Call GetNamedChunksByType instead")
+	virtual TArray<FCustomChunk> GetCustomChunksByType(ECustomChunkType DesiredChunkType) = 0;
+
 
 	/**
 	 * Check whether current platform supports chunk installation by name
@@ -407,6 +410,11 @@ public:
 		return false;
 	}
 
+	virtual TArray<FCustomChunk> GetCustomChunksByType(ECustomChunkType DesiredChunkType) override
+	{
+		return TArray<FCustomChunk>();
+	}
+
 	virtual bool SupportsNamedChunkInstall() const override
 	{
 		return false;
