Fixed IMappedFileHandle::MapRegion to map files\regions smaller than a page size

diff --git a/Engine/Source/Runtime/Core/Private/IOS/IOSPlatformFile.cpp b/Engine/Source/Runtime/Core/Private/IOS/IOSPlatformFile.cpp
--- a/Engine/Source/Runtime/Core/Private/IOS/IOSPlatformFile.cpp
+++ b/Engine/Source/Runtime/Core/Private/IOS/IOSPlatformFile.cpp
@@ -434,12 +434,13 @@ public:
 	
 };
 
+static SIZE_T FileMappingAlignment = FPlatformMemory::GetConstants().PageSize;
+
 class FIOSMappedFileHandle final : public IMappedFileHandle
 {
 	const uint8* MappedPtr;
 	FString Filename;
 	int32 NumOutstandingRegions;
-	int32 Alignment;
 	int FileHandle;
 	
 public:
@@ -453,7 +454,6 @@ public:
 		, NumOutstandingRegions(0)
 		, FileHandle(InFileHandle)
 	{
-		Alignment = sysconf(_SC_PAGE_SIZE);
 	}
 
 	~FIOSMappedFileHandle()
@@ -474,15 +474,9 @@ public:
 		// const uint8* MapPtr = (const uint8 *)mmap(NULL, BytesToMap, PROT_READ, MAP_PRIVATE, FileHandle, Offset);
 		//		const uint8* MapPtr = (const uint8 *)mmap(NULL, BytesToMap, PROT_READ, MAP_SHARED, FileHandle, Offset);
 		
-		int64 AlignedOffset = AlignDown(Offset, Alignment);
-		int64 AlignedSize = Align(BytesToMap + Offset - AlignedOffset, Alignment);
-		
-		// if we are about to go off the end, let's not
-		if (AlignedOffset + AlignedSize > GetFileSize())
-		{
-			UE_LOG(LogIOS, Warning, TEXT("Mapping fell off the end, did we need to actually abort? [%lld + %lld > %lld]"), AlignedOffset, AlignedSize, GetFileSize());
-			return nullptr;
-		}
+		const int64 AlignedOffset = AlignDown(Offset, FileMappingAlignment);
+		//File mapping can extend beyond file size. It's OK, kernel will just fill any leftover page data with zeros
+		const int64 AlignedSize = Align(BytesToMap + Offset - AlignedOffset, FileMappingAlignment);
 		
 		const uint8* AlignedMapPtr = (const uint8 *)mmap(NULL, AlignedSize, PROT_READ, MAP_PRIVATE, FileHandle, AlignedOffset);
 		if (AlignedMapPtr == (const uint8*)-1 || AlignedMapPtr == nullptr)
diff --git a/Engine/Source/Runtime/Core/Private/Unix/UnixPlatformFile.cpp b/Engine/Source/Runtime/Core/Private/Unix/UnixPlatformFile.cpp
--- a/Engine/Source/Runtime/Core/Private/Unix/UnixPlatformFile.cpp
+++ b/Engine/Source/Runtime/Core/Private/Unix/UnixPlatformFile.cpp
@@ -702,6 +702,8 @@ public:
 	virtual ~FUnixMappedFileRegion();
 };
 
+static SIZE_T FileMappingAlignment = FPlatformMemory::GetConstants().PageSize;
+
 class FUnixMappedFileHandle final : public IMappedFileHandle
 {
 public:
@@ -712,7 +714,6 @@ public:
 		, NumOutstandingRegions(0)
 		, FileHandle(InFileHandle)
 	{
-		Alignment = sysconf(_SC_PAGE_SIZE);
 	}
 
 	virtual ~FUnixMappedFileHandle() override
@@ -728,14 +729,9 @@ public:
 		BytesToMap = FMath::Min<int64>(BytesToMap, GetFileSize() - Offset);
 		check(BytesToMap > 0); // don't map zero bytes
 
-		const int64 AlignedOffset = AlignDown(Offset, Alignment);
-		int64 AlignedSize = Align(BytesToMap + Offset - AlignedOffset, Alignment);
-
-		// if we are about to go off the end, let's just go to the end of the file
-		if (AlignedOffset + AlignedSize > GetFileSize())
-		{
-			AlignedSize = GetFileSize() - AlignedOffset;
-		}
+		const int64 AlignedOffset = AlignDown(Offset, FileMappingAlignment);
+		//File mapping can extend beyond file size. It's OK, kernel will just fill any leftover page data with zeros
+		const int64 AlignedSize = Align(BytesToMap + Offset - AlignedOffset, FileMappingAlignment);
 
 		int Flags = MAP_PRIVATE;
 		if (bPreloadHint)
@@ -773,7 +769,6 @@ private:
 	const uint8* MappedPtr;
 	FString Filename;
 	int32 NumOutstandingRegions;
-	int32 Alignment;
 	int FileHandle;
 };
 
diff --git a/Engine/Source/Runtime/PakFile/Private/IPlatformFilePak.cpp b/Engine/Source/Runtime/PakFile/Private/IPlatformFilePak.cpp
--- a/Engine/Source/Runtime/PakFile/Private/IPlatformFilePak.cpp
+++ b/Engine/Source/Runtime/PakFile/Private/IPlatformFilePak.cpp
@@ -4680,12 +4680,12 @@ public:
 	}
 	virtual IMappedFileRegion* MapRegion(int64 Offset = 0, int64 BytesToMap = MAX_int64, bool bPreloadHint = false) override
 	{
-		check(Offset + OffsetInPak < PakSize); // don't map zero bytes and don't map off the end of the (real) file
+		//check(Offset + OffsetInPak < PakSize); // don't map zero bytes and don't map off the end of the (real) file
 		check(Offset < GetFileSize()); // don't map zero bytes and don't map off the end of the (virtual) file
 		BytesToMap = FMath::Min<int64>(BytesToMap, GetFileSize() - Offset);
 		check(BytesToMap > 0); // don't map zero bytes
-		check(Offset + BytesToMap <= GetFileSize()); // don't map zero bytes and don't map off the end of the (virtual) file
-		check(Offset + OffsetInPak + BytesToMap <= PakSize); // don't map zero bytes and don't map off the end of the (real) file
+		//check(Offset + BytesToMap <= GetFileSize()); // don't map zero bytes and don't map off the end of the (virtual) file
+		//check(Offset + OffsetInPak + BytesToMap <= PakSize); // don't map zero bytes and don't map off the end of the (real) file
 		return LowerLevel->MapRegion(Offset + OffsetInPak, BytesToMap, bPreloadHint);
 	}
 };
