PR #9577: Added CORE_API to FComplexStatUtils struct for games usage (Contributed by [REDACTED_PERSON])

diff --git a/Engine/Source/Runtime/Core/Public/Stats/StatsData.h b/Engine/Source/Runtime/Core/Public/Stats/StatsData.h
--- a/Engine/Source/Runtime/Core/Public/Stats/StatsData.h
+++ b/Engine/Source/Runtime/Core/Public/Stats/StatsData.h
@@ -816,7 +816,7 @@ struct CORE_API FStatsUtils
 /**
 * Contains helpers functions to manage complex stat messages.
 */
-struct FComplexStatUtils
+struct CORE_API FComplexStatUtils
 {
 	/** Accumulates a stat message into a complex stat message. */
 	static void AddAndMinMax( FComplexStatMessage& Dest, const FStatMessage& Item, EComplexStatField::Type SumIndex, EComplexStatField::Type MaxIndex, EComplexStatField::Type MinIndex);
