Deprecated all global vars of thread-safe delegates from FCoreDelegates and replaced them by functions returning local static singleton.

thread-safe delegates are not zero-initializable and so can't be used as global vars because they are vulnerable to static initialization order fiasco.

diff --git a/Engine/Plugins/AI/MLAdapter/Source/MLAdapter/Private/Managers/MLAdapterManager.cpp b/Engine/Plugins/AI/MLAdapter/Source/MLAdapter/Private/Managers/MLAdapterManager.cpp
--- a/Engine/Plugins/AI/MLAdapter/Source/MLAdapter/Private/Managers/MLAdapterManager.cpp
+++ b/Engine/Plugins/AI/MLAdapter/Source/MLAdapter/Private/Managers/MLAdapterManager.cpp
@@ -332,7 +332,7 @@ void UMLAdapterManager::BindToDelegates()
 {	
 	// Commented out possible other useful delegates
 	//	FCoreDelegates::GameStateClassChanged;
-	//	FCoreDelegates::ConfigReadyForUse;
+	//	FCoreDelegates::TSConfigReadyForUse();
 	//  FWorldDelegates::OnPostWorldCreation;
 	//	FWorldDelegates::OnPreWorldInitialization; 
 
diff --git a/Engine/Plugins/Online/EOSShared/Source/EOSShared/Private/EOSSDKManager.cpp b/Engine/Plugins/Online/EOSShared/Source/EOSShared/Private/EOSSDKManager.cpp
--- a/Engine/Plugins/Online/EOSShared/Source/EOSShared/Private/EOSSDKManager.cpp
+++ b/Engine/Plugins/Online/EOSShared/Source/EOSShared/Private/EOSSDKManager.cpp
@@ -225,7 +225,7 @@ EOS_EResult FEOSSDKManager::Initialize()
 		{
 			bInitialized = true;
 
-			FCoreDelegates::OnConfigSectionsChanged.AddRaw(this, &FEOSSDKManager::OnConfigSectionsChanged);
+			FCoreDelegates::TSOnConfigSectionsChanged().AddRaw(this, &FEOSSDKManager::OnConfigSectionsChanged);
 			LoadConfig();
 
 #if !NO_LOGGING
@@ -448,7 +448,7 @@ void FEOSSDKManager::Shutdown()
 			ReleaseReleasedPlatforms();
 		}
 
-		FCoreDelegates::OnConfigSectionsChanged.RemoveAll(this);
+		FCoreDelegates::TSOnConfigSectionsChanged().RemoveAll(this);
 
 #if !NO_LOGGING
 		FCoreDelegates::OnLogVerbosityChanged.RemoveAll(this);
diff --git a/Engine/Plugins/Online/OnlineFramework/Source/Hotfix/Private/OnlineHotfixManager.cpp b/Engine/Plugins/Online/OnlineFramework/Source/Hotfix/Private/OnlineHotfixManager.cpp
--- a/Engine/Plugins/Online/OnlineFramework/Source/Hotfix/Private/OnlineHotfixManager.cpp
+++ b/Engine/Plugins/Online/OnlineFramework/Source/Hotfix/Private/OnlineHotfixManager.cpp
@@ -977,7 +977,7 @@ bool UOnlineHotfixManager::HotfixIniFile(const FString& FileName, const FString&
 
 				if (bIsEngineIni)
 				{
-					// TODO replace all of this with bindees to FCoreDelegates::OnConfigSectionsChanged
+					// TODO replace all of this with bindees to FCoreDelegates::TSOnConfigSectionsChanged()
 					const TCHAR* LogConfigSection = TEXT("[Core.Log]");
 					const TCHAR* ConsoleVariableSection = TEXT("[ConsoleVariables]");
 					const TCHAR* HttpSection = TEXT("[HTTP"); // note "]" omitted on purpose since we want a partial match
@@ -1102,7 +1102,10 @@ bool UOnlineHotfixManager::HotfixIniFile(const FString& FileName, const FString&
 	}
 
 	const FString ConfigFileName = ConfigFile->Name.ToString();
+	FCoreDelegates::TSOnConfigSectionsChanged().Broadcast(ConfigFileName, UpdatedSectionNames);
+PRAGMA_DISABLE_DEPRECATION_WARNINGS
 	FCoreDelegates::OnConfigSectionsChanged.Broadcast(ConfigFileName, UpdatedSectionNames);
+PRAGMA_ENABLE_DEPRECATION_WARNINGS
 
 	// Reload log suppression if configs changed
 	if (bUpdateLogSuppression)
diff --git a/Engine/Source/Developer/PakFileUtilities/Private/PakFileUtilities.cpp b/Engine/Source/Developer/PakFileUtilities/Private/PakFileUtilities.cpp
--- a/Engine/Source/Developer/PakFileUtilities/Private/PakFileUtilities.cpp
+++ b/Engine/Source/Developer/PakFileUtilities/Private/PakFileUtilities.cpp
@@ -5088,7 +5088,10 @@ bool MakeBinaryConfig(const TCHAR* CmdLine)
 
 	// allow delegates to modify the config data with some tagged binary data
 	FCoreDelegates::FExtraBinaryConfigData ExtraData(Config, true);
+	FCoreDelegates::TSAccessExtraBinaryConfigData().Broadcast(ExtraData);
+PRAGMA_DISABLE_DEPRECATION_WARNINGS
 	FCoreDelegates::AccessExtraBinaryConfigData.Broadcast(ExtraData);
+PRAGMA_ENABLE_DEPRECATION_WARNINGS
 
 	// write it all out!
 	TArray<uint8> FileContent;
diff --git a/Engine/Source/Editor/UnrealEd/Private/Commandlets/MakeBinaryConfigCommandlet.cpp b/Engine/Source/Editor/UnrealEd/Private/Commandlets/MakeBinaryConfigCommandlet.cpp
--- a/Engine/Source/Editor/UnrealEd/Private/Commandlets/MakeBinaryConfigCommandlet.cpp
+++ b/Engine/Source/Editor/UnrealEd/Private/Commandlets/MakeBinaryConfigCommandlet.cpp
@@ -99,7 +99,10 @@ int32 UMakeBinaryConfigCommandlet::Main(const FString& Params)
 
 	// allow delegates to modify the config data with some tagged binary data
 	FCoreDelegates::FExtraBinaryConfigData ExtraData(Config, true);
+	FCoreDelegates::TSAccessExtraBinaryConfigData().Broadcast(ExtraData);
+PRAGMA_DISABLE_DEPRECATION_WARNINGS
 	FCoreDelegates::AccessExtraBinaryConfigData.Broadcast(ExtraData);
+PRAGMA_ENABLE_DEPRECATION_WARNINGS
 
 	// write it all out!
 	TArray<uint8> FileContent;
diff --git a/Engine/Source/Editor/UnrealEd/Private/CookOnTheFlyServer.cpp b/Engine/Source/Editor/UnrealEd/Private/CookOnTheFlyServer.cpp
--- a/Engine/Source/Editor/UnrealEd/Private/CookOnTheFlyServer.cpp
+++ b/Engine/Source/Editor/UnrealEd/Private/CookOnTheFlyServer.cpp
@@ -454,8 +454,8 @@ UCookOnTheFlyServer::~UCookOnTheFlyServer()
 {
 	ClearPackageStoreContexts();
 
-	FCoreDelegates::OnFConfigCreated.RemoveAll(this);
-	FCoreDelegates::OnFConfigDeleted.RemoveAll(this);
+	FCoreDelegates::TSOnFConfigCreated().RemoveAll(this);
+	FCoreDelegates::TSOnFConfigDeleted().RemoveAll(this);
 	FCoreUObjectDelegates::GetPreGarbageCollectDelegate().RemoveAll(this);
 	FCoreUObjectDelegates::GetPostGarbageCollect().RemoveAll(this);
 	GetTargetPlatformManager()->GetOnTargetPlatformsInvalidatedDelegate().RemoveAll(this);
@@ -5764,8 +5764,8 @@ void UCookOnTheFlyServer::Initialize( ECookMode::Type DesiredCookMode, ECookInit
 		FCoreDelegates::OnTargetPlatformChangedSupportedFormats.AddUObject(this, &UCookOnTheFlyServer::OnTargetPlatformChangedSupportedFormats);
 	}
 
-	FCoreDelegates::OnFConfigCreated.AddUObject(this, &UCookOnTheFlyServer::OnFConfigCreated);
-	FCoreDelegates::OnFConfigDeleted.AddUObject(this, &UCookOnTheFlyServer::OnFConfigDeleted);
+	FCoreDelegates::TSOnFConfigCreated().AddUObject(this, &UCookOnTheFlyServer::OnFConfigCreated);
+	FCoreDelegates::TSOnFConfigDeleted().AddUObject(this, &UCookOnTheFlyServer::OnFConfigDeleted);
 
 	GetTargetPlatformManager()->GetOnTargetPlatformsInvalidatedDelegate().AddUObject(this, &UCookOnTheFlyServer::OnTargetPlatformsInvalidated);
 
diff --git a/Engine/Source/Runtime/Core/Private/Async/TaskGraph.cpp b/Engine/Source/Runtime/Core/Private/Async/TaskGraph.cpp
--- a/Engine/Source/Runtime/Core/Private/Async/TaskGraph.cpp
+++ b/Engine/Source/Runtime/Core/Private/Async/TaskGraph.cpp
@@ -1925,7 +1925,7 @@ public:
 			if (GConfig == nullptr)
 			{
 				// postpone starting reserve workers until GConfig is initialized, to know if reserve workers are disabled
-				FCoreDelegates::ConfigReadyForUse.AddRaw(this, &FTaskGraphCompatibilityImplementation::StartReserveWorkers);
+				FCoreDelegates::TSConfigReadyForUse().AddRaw(this, &FTaskGraphCompatibilityImplementation::StartReserveWorkers);
 			}
 			else
 			{
@@ -1962,7 +1962,7 @@ public:
 
 	~FTaskGraphCompatibilityImplementation() override
 	{
-		FCoreDelegates::ConfigReadyForUse.RemoveAll(this);
+		FCoreDelegates::TSConfigReadyForUse().RemoveAll(this);
 
 		for (auto& Callback : ShutdownCallbacks)
 		{
diff --git a/Engine/Source/Runtime/Core/Private/GenericPlatform/GenericPlatformCrashContext.cpp b/Engine/Source/Runtime/Core/Private/GenericPlatform/GenericPlatformCrashContext.cpp
--- a/Engine/Source/Runtime/Core/Private/GenericPlatform/GenericPlatformCrashContext.cpp
+++ b/Engine/Source/Runtime/Core/Private/GenericPlatform/GenericPlatformCrashContext.cpp
@@ -299,7 +299,7 @@ void FGenericCrashContext::Initialize()
 		NCached::Set(NCached::Session.EngineModeEx, bIsVanilla ? FGenericCrashContext::EngineModeExVanilla : FGenericCrashContext::EngineModeExDirty);
 	});
 
-	FCoreDelegates::ConfigReadyForUse.AddStatic(FGenericCrashContext::InitializeFromConfig);
+	FCoreDelegates::TSConfigReadyForUse().AddStatic(FGenericCrashContext::InitializeFromConfig);
 
 	FCoreDelegates::OnPostFork.AddLambda([](EForkProcessRole Role)
 	{
diff --git a/Engine/Source/Runtime/Core/Private/Misc/ConfigCacheIni.cpp b/Engine/Source/Runtime/Core/Private/Misc/ConfigCacheIni.cpp
--- a/Engine/Source/Runtime/Core/Private/Misc/ConfigCacheIni.cpp
+++ b/Engine/Source/Runtime/Core/Private/Misc/ConfigCacheIni.cpp
@@ -422,9 +422,12 @@ static void FixupArrayOfStructKeysForSection(FConfigSection* Section, const FStr
  */
 /*static*/ bool DoesConfigFileExistWrapper(const TCHAR* IniFile, const TSet<FString>* IniCacheSet)
 {
-	// will any delegates return contents via PreLoadConfigFileDelegate()?
+	// will any delegates return contents via TSPreLoadConfigFileDelegate()?
 	int32 ResponderCount = 0;
+	FCoreDelegates::TSCountPreLoadConfigFileRespondersDelegate().Broadcast(IniFile, ResponderCount);
+PRAGMA_DISABLE_DEPRECATION_WARNINGS
 	FCoreDelegates::CountPreLoadConfigFileRespondersDelegate.Broadcast(IniFile, ResponderCount);
+PRAGMA_ENABLE_DEPRECATION_WARNINGS
 
 	if (ResponderCount > 0)
 	{
@@ -470,7 +473,10 @@ static void FixupArrayOfStructKeysForSection(FConfigSection* Section, const FStr
 static bool LoadConfigFileWrapper(const TCHAR* IniFile, FString& Contents, bool bIsOverride = false)
 {
 	// let other systems load the file instead of the standard load below
+	FCoreDelegates::TSPreLoadConfigFileDelegate().Broadcast(IniFile, Contents);
+PRAGMA_DISABLE_DEPRECATION_WARNINGS
 	FCoreDelegates::PreLoadConfigFileDelegate.Broadcast(IniFile, Contents);
+PRAGMA_ENABLE_DEPRECATION_WARNINGS
 
 	// if this loaded any text, we are done, and we won't override the contents with standard ini file data
 	if (Contents.Len())
@@ -503,7 +509,10 @@ static bool SaveConfigFileWrapper(const TCHAR* IniFile, const FString& Contents)
 
 	// let anyone that needs to save it, do so (counting how many did)
 	int32 SavedCount = 0;
+	FCoreDelegates::TSPreSaveConfigFileDelegate().Broadcast(IniFile, Contents, SavedCount);
+PRAGMA_DISABLE_DEPRECATION_WARNINGS
 	FCoreDelegates::PreSaveConfigFileDelegate.Broadcast(IniFile, Contents, SavedCount);
+PRAGMA_ENABLE_DEPRECATION_WARNINGS
 
 	// save it even if a delegate did as well
 	bool bLocalWriteSucceeded = false;
@@ -540,11 +549,10 @@ FConfigFile::FConfigFile()
 , Name( NAME_None )
 , SourceConfigFile(nullptr)
 {
-
-	if (FCoreDelegates::OnFConfigCreated.IsBound())
-	{
-		FCoreDelegates::OnFConfigCreated.Broadcast(this);
-	}
+	FCoreDelegates::TSOnFConfigCreated().Broadcast(this);
+PRAGMA_DISABLE_DEPRECATION_WARNINGS
+	FCoreDelegates::OnFConfigCreated.Broadcast(this);
+PRAGMA_ENABLE_DEPRECATION_WARNINGS
 }
 
 FConfigFile::~FConfigFile()
@@ -553,11 +561,10 @@ FConfigFile::~FConfigFile()
 
 	if ( !GExitPurge )
 	{
-		// FCoreDelegates::OnFConfigDeleted may already be deleted
-		if ( FCoreDelegates::OnFConfigDeleted.IsBound() )
-		{
-			FCoreDelegates::OnFConfigDeleted.Broadcast(this);
-		}
+		FCoreDelegates::TSOnFConfigDeleted().Broadcast(this);
+PRAGMA_DISABLE_DEPRECATION_WARNINGS
+		FCoreDelegates::OnFConfigDeleted.Broadcast(this);
+PRAGMA_ENABLE_DEPRECATION_WARNINGS
 	}
 
 	delete SourceConfigFile;
@@ -2490,7 +2497,10 @@ bool FConfigCacheIni::GetString( const TCHAR* Section, const TCHAR* Key, FString
 	}
 	Value = ConfigValue->GetValue();
 
+	FCoreDelegates::TSOnConfigValueRead().Broadcast(*Filename, Section, Key);
+PRAGMA_DISABLE_DEPRECATION_WARNINGS
 	FCoreDelegates::OnConfigValueRead.Broadcast(*Filename, Section, Key);
+PRAGMA_ENABLE_DEPRECATION_WARNINGS
 
 	return true;
 }
@@ -2521,7 +2531,10 @@ bool FConfigCacheIni::GetText( const TCHAR* Section, const TCHAR* Key, FText& Va
 		return false;
 	}
 
+	FCoreDelegates::TSOnConfigValueRead().Broadcast(*Filename, Section, Key);
+PRAGMA_DISABLE_DEPRECATION_WARNINGS
 	FCoreDelegates::OnConfigValueRead.Broadcast(*Filename, Section, Key);
+PRAGMA_ENABLE_DEPRECATION_WARNINGS
 
 	return true;
 }
@@ -2546,7 +2559,10 @@ bool FConfigCacheIni::GetSection( const TCHAR* Section, TArray<FString>& Result,
 		Result.Add(FString::Printf(TEXT("%s=%s"), *It.Key().ToString(), *It.Value().GetValue()));
 	}
 
+	FCoreDelegates::TSOnConfigSectionRead().Broadcast(*Filename, Section);
+PRAGMA_DISABLE_DEPRECATION_WARNINGS
 	FCoreDelegates::OnConfigSectionRead.Broadcast(*Filename, Section);
+PRAGMA_ENABLE_DEPRECATION_WARNINGS
 
 	return true;
 }
@@ -2571,7 +2587,7 @@ FConfigSection* FConfigCacheIni::GetSectionPrivate( const TCHAR* Section, const
 
 	if (Sec)
 	{
-		FCoreDelegates::OnConfigSectionRead.Broadcast(*Filename, Section);
+		FCoreDelegates::TSOnConfigSectionRead().Broadcast(*Filename, Section);
 	}
 
 	return Sec;
@@ -2588,7 +2604,10 @@ bool FConfigCacheIni::DoesSectionExist(const TCHAR* Section, const FString& File
 
 	if (bReturnVal)
 	{
+		FCoreDelegates::TSOnConfigSectionNameRead().Broadcast(*Filename, Section);
+PRAGMA_DISABLE_DEPRECATION_WARNINGS
 		FCoreDelegates::OnConfigSectionNameRead.Broadcast(*Filename, Section);
+PRAGMA_ENABLE_DEPRECATION_WARNINGS
 	}
 
 	return bReturnVal;
@@ -2749,7 +2768,10 @@ bool FConfigCacheIni::GetSectionNames( const FString& Filename, TArray<FString>&
 		{
 			out_SectionNames.Add(It.Key());
 
+			FCoreDelegates::TSOnConfigSectionNameRead().Broadcast(*Filename, *It.Key());
+PRAGMA_DISABLE_DEPRECATION_WARNINGS
 			FCoreDelegates::OnConfigSectionNameRead.Broadcast(*Filename, *It.Key());
+PRAGMA_ENABLE_DEPRECATION_WARNINGS
 		}
 		bResult = true;
 	}
@@ -2792,7 +2814,10 @@ bool FConfigCacheIni::GetPerObjectConfigSections( const FString& Filename, const
 					out_SectionNames.Insert(SectionName,0);
 					bResult = true;
 
+					FCoreDelegates::TSOnConfigSectionNameRead().Broadcast(*Filename, *SectionName);
+PRAGMA_DISABLE_DEPRECATION_WARNINGS
 					FCoreDelegates::OnConfigSectionNameRead.Broadcast(*Filename, *SectionName);
+PRAGMA_ENABLE_DEPRECATION_WARNINGS
 				}
 			}
 		}
@@ -2943,7 +2968,10 @@ int32 FConfigCacheIni::GetArray
 
 	if (out_Arr.Num())
 	{
+		FCoreDelegates::TSOnConfigValueRead().Broadcast(*Filename, Section, Key);
+PRAGMA_DISABLE_DEPRECATION_WARNINGS
 		FCoreDelegates::OnConfigValueRead.Broadcast(*Filename, Section, Key);
+PRAGMA_ENABLE_DEPRECATION_WARNINGS
 	}
 
 	return out_Arr.Num();
@@ -3659,9 +3687,12 @@ bool FConfigCacheIni::CreateGConfigFromSaved(const TCHAR* Filename)
 	MemoryReader << Names;// << ExtraData.Data;
 
 	// now let the delegates pull their data out, after GConfig is set up
-//	FCoreDelegates::AccessExtraBinaryConfigData.Broadcast(ExtraData);
+//	FCoreDelegates::TSAccessExtraBinaryConfigData().Broadcast(ExtraData);
 
+	FCoreDelegates::TSConfigReadyForUse().Broadcast();
+PRAGMA_DISABLE_DEPRECATION_WARNINGS
 	FCoreDelegates::ConfigReadyForUse.Broadcast();
+PRAGMA_ENABLE_DEPRECATION_WARNINGS
 
 	FMemory::Free(PreloadedData);
 	return true;
@@ -3774,7 +3805,10 @@ void FConfigCacheIni::InitializeConfigSystem()
 			GConfig = new FConfigCacheIni(EConfigCacheType::Temporary);
 			GConfig->SerializeStateForBootstrap_Impl(MemoryReader);
 			GConfig->bIsReadyForUse = true;
+			FCoreDelegates::TSConfigReadyForUse().Broadcast();
+PRAGMA_DISABLE_DEPRECATION_WARNINGS
 			FCoreDelegates::ConfigReadyForUse.Broadcast();
+PRAGMA_ENABLE_DEPRECATION_WARNINGS
 			return;
 		}
 		else
@@ -3824,7 +3858,10 @@ void FConfigCacheIni::InitializeConfigSystem()
 	// now we can make use of GConfig
 	GConfig->bIsReadyForUse = true;
 
+	FCoreDelegates::TSConfigReadyForUse().Broadcast();
+PRAGMA_DISABLE_DEPRECATION_WARNINGS
 	FCoreDelegates::ConfigReadyForUse.Broadcast();
+PRAGMA_ENABLE_DEPRECATION_WARNINGS
 }
 
 const FString& FConfigCacheIni::GetCustomConfigString()
diff --git a/Engine/Source/Runtime/Core/Private/Misc/ConfigUtilities.cpp b/Engine/Source/Runtime/Core/Private/Misc/ConfigUtilities.cpp
--- a/Engine/Source/Runtime/Core/Private/Misc/ConfigUtilities.cpp
+++ b/Engine/Source/Runtime/Core/Private/Misc/ConfigUtilities.cpp
@@ -335,16 +335,16 @@ private:
 public:
 	FConfigHistoryHelper()
 	{
-		FCoreDelegates::OnConfigValueRead.AddRaw(this, &FConfigHistoryHelper::OnConfigValueRead);
-		FCoreDelegates::OnConfigSectionRead.AddRaw(this, &FConfigHistoryHelper::OnConfigSectionRead);
-		FCoreDelegates::OnConfigSectionNameRead.AddRaw(this, &FConfigHistoryHelper::OnConfigSectionNameRead);
+		FCoreDelegates::TSOnConfigValueRead().AddRaw(this, &FConfigHistoryHelper::OnConfigValueRead);
+		FCoreDelegates::TSOnConfigSectionRead().AddRaw(this, &FConfigHistoryHelper::OnConfigSectionRead);
+		FCoreDelegates::TSOnConfigSectionNameRead().AddRaw(this, &FConfigHistoryHelper::OnConfigSectionNameRead);
 	}
 
 	~FConfigHistoryHelper()
 	{
-		FCoreDelegates::OnConfigValueRead.RemoveAll(this);
-		FCoreDelegates::OnConfigSectionRead.RemoveAll(this);
-		FCoreDelegates::OnConfigSectionNameRead.RemoveAll(this);
+		FCoreDelegates::TSOnConfigValueRead().RemoveAll(this);
+		FCoreDelegates::TSOnConfigSectionRead().RemoveAll(this);
+		FCoreDelegates::TSOnConfigSectionNameRead().RemoveAll(this);
 	}
 
 	void DumpHistory()
diff --git a/Engine/Source/Runtime/Core/Private/Misc/CoreDelegates.cpp b/Engine/Source/Runtime/Core/Private/Misc/CoreDelegates.cpp
--- a/Engine/Source/Runtime/Core/Private/Misc/CoreDelegates.cpp
+++ b/Engine/Source/Runtime/Core/Private/Misc/CoreDelegates.cpp
@@ -126,15 +126,71 @@ TMulticastDelegate<void(FString, int, int)> FCoreDelegates::ApplicationReceivedL
 TMulticastDelegate<void()> FCoreDelegates::ApplicationPerformFetchDelegate;
 TMulticastDelegate<void(FString)> FCoreDelegates::ApplicationBackgroundSessionEventDelegate;
 
-TTSMulticastDelegate<void(const TCHAR*, int32&)> FCoreDelegates::CountPreLoadConfigFileRespondersDelegate;
-TTSMulticastDelegate<void(const TCHAR*, FString&)> FCoreDelegates::PreLoadConfigFileDelegate;
-TTSMulticastDelegate<void(const TCHAR*, const FString&, int32&)> FCoreDelegates::PreSaveConfigFileDelegate;
-TTSMulticastDelegate<void(const FConfigFile*)> FCoreDelegates::OnFConfigCreated;
-TTSMulticastDelegate<void(const FConfigFile*)> FCoreDelegates::OnFConfigDeleted;
-TTSMulticastDelegate<void(const TCHAR*, const TCHAR*, const TCHAR*)> FCoreDelegates::OnConfigValueRead;
-TTSMulticastDelegate<void(const TCHAR*, const TCHAR*)> FCoreDelegates::OnConfigSectionRead;
-TTSMulticastDelegate<void(const TCHAR*, const TCHAR*)> FCoreDelegates::OnConfigSectionNameRead;
-TTSMulticastDelegate<void(const FString&, const TSet<FString>&)> FCoreDelegates::OnConfigSectionsChanged;
+PRAGMA_DISABLE_DEPRECATION_WARNINGS
+TMulticastDelegate<void(const TCHAR*, int32&)> FCoreDelegates::CountPreLoadConfigFileRespondersDelegate;
+TMulticastDelegate<void(const TCHAR*, FString&)> FCoreDelegates::PreLoadConfigFileDelegate;
+TMulticastDelegate<void(const TCHAR*, const FString&, int32&)> FCoreDelegates::PreSaveConfigFileDelegate;
+TMulticastDelegate<void(const FConfigFile*)> FCoreDelegates::OnFConfigCreated;
+TMulticastDelegate<void(const FConfigFile*)> FCoreDelegates::OnFConfigDeleted;
+TMulticastDelegate<void(const TCHAR*, const TCHAR*, const TCHAR*)> FCoreDelegates::OnConfigValueRead;
+TMulticastDelegate<void(const TCHAR*, const TCHAR*)> FCoreDelegates::OnConfigSectionRead;
+TMulticastDelegate<void(const TCHAR*, const TCHAR*)> FCoreDelegates::OnConfigSectionNameRead;
+TMulticastDelegate<void(const FString&, const TSet<FString>&)> FCoreDelegates::OnConfigSectionsChanged;
+PRAGMA_ENABLE_DEPRECATION_WARNINGS
+
+TTSMulticastDelegate<void(const TCHAR*, int32&)>& FCoreDelegates::TSCountPreLoadConfigFileRespondersDelegate()
+{
+	static TTSMulticastDelegate<void(const TCHAR*, int32&)> Singleton;
+	return Singleton;
+}
+
+TTSMulticastDelegate<void(const TCHAR*, FString&)>& FCoreDelegates::TSPreLoadConfigFileDelegate()
+{
+	static TTSMulticastDelegate<void(const TCHAR*, FString&)> Singleton;
+	return Singleton;
+}
+
+TTSMulticastDelegate<void(const TCHAR*, const FString&, int32&)>& FCoreDelegates::TSPreSaveConfigFileDelegate()
+{
+	static TTSMulticastDelegate<void(const TCHAR*, const FString&, int32&)> Singleton;
+	return Singleton;
+}
+
+TTSMulticastDelegate<void(const FConfigFile*)>& FCoreDelegates::TSOnFConfigCreated()
+{
+	static TTSMulticastDelegate<void(const FConfigFile*)> Singleton;
+	return Singleton;
+}
+
+TTSMulticastDelegate<void(const FConfigFile*)>& FCoreDelegates::TSOnFConfigDeleted()
+{
+	static TTSMulticastDelegate<void(const FConfigFile*)> Singleton;
+	return Singleton;
+}
+
+TTSMulticastDelegate<void(const TCHAR*, const TCHAR*, const TCHAR*)>& FCoreDelegates::TSOnConfigValueRead()
+{
+	static TTSMulticastDelegate<void(const TCHAR*, const TCHAR*, const TCHAR*)> Singleton;
+	return Singleton;
+}
+
+TTSMulticastDelegate<void(const TCHAR*, const TCHAR*)>& FCoreDelegates::TSOnConfigSectionRead()
+{
+	static TTSMulticastDelegate<void(const TCHAR*, const TCHAR*)> Singleton;
+	return Singleton;
+}
+
+TTSMulticastDelegate<void(const TCHAR*, const TCHAR*)>& FCoreDelegates::TSOnConfigSectionNameRead()
+{
+	static TTSMulticastDelegate<void(const TCHAR*, const TCHAR*)> Singleton;
+	return Singleton;
+}
+
+TTSMulticastDelegate<void(const FString&, const TSet<FString>&)>& FCoreDelegates::TSOnConfigSectionsChanged()
+{
+	static TTSMulticastDelegate<void(const FString&, const TSet<FString>&)> Singleton;
+	return Singleton;
+}
 
 TMulticastDelegate<void(const TCHAR*, const TCHAR*, uint32, bool)> FCoreDelegates::OnApplyCVarFromIni;
 TMulticastDelegate<void(uint32, uint32)> FCoreDelegates::OnSystemResolutionChanged;
@@ -177,7 +233,15 @@ TMulticastDelegate<void()> FCoreDelegates::PreRenderingThreadDestroyed;
 
 TMulticastDelegate<void(int32)> FCoreDelegates::ApplicationReceivedScreenOrientationChangedNotificationDelegate;
 
-TTSMulticastDelegate<void()> FCoreDelegates::ConfigReadyForUse;
+PRAGMA_DISABLE_DEPRECATION_WARNINGS
+TMulticastDelegate<void()> FCoreDelegates::ConfigReadyForUse;
+PRAGMA_ENABLE_DEPRECATION_WARNINGS
+
+TTSMulticastDelegate<void()>& FCoreDelegates::TSConfigReadyForUse()
+{
+	static TTSMulticastDelegate<void()> Singleton;
+	return Singleton;
+}
 
 TDelegate<bool()> FCoreDelegates::IsLoadingMovieCurrentlyPlaying;
 
@@ -187,7 +251,15 @@ TMulticastDelegate<void(const FString&, FPlatformUserId)> FCoreDelegates::OnActi
 
 TMulticastDelegate<void(const FString&)> FCoreDelegates::OnGCFinishDestroyTimeExtended;
 
-TTSMulticastDelegate<void(FCoreDelegates::FExtraBinaryConfigData&)> FCoreDelegates::AccessExtraBinaryConfigData;
+PRAGMA_DISABLE_DEPRECATION_WARNINGS
+TMulticastDelegate<void(FCoreDelegates::FExtraBinaryConfigData&)> FCoreDelegates::AccessExtraBinaryConfigData;
+PRAGMA_ENABLE_DEPRECATION_WARNINGS
+
+TTSMulticastDelegate<void(FCoreDelegates::FExtraBinaryConfigData&)>& FCoreDelegates::TSAccessExtraBinaryConfigData()
+{
+	static TTSMulticastDelegate<void(FCoreDelegates::FExtraBinaryConfigData&)> Singleton;
+	return Singleton;
+}
 
 TDelegate<void(TArrayView<const FSHAHash>, FCoreDelegates::FAttachShaderReadRequestFunc)> FCoreDelegates::PreloadPackageShaderMaps;
 TDelegate<void(TArrayView<const FSHAHash>)> FCoreDelegates::ReleasePreloadedPackageShaderMaps;
@@ -208,7 +280,7 @@ PRAGMA_ENABLE_DEPRECATION_WARNINGS
 /**	 Implemented as a function to address global ctor issues */
 FSimpleMulticastDelegate& FCoreDelegates::GetMemoryTrimDelegate()
 {
-	static FSimpleMulticastDelegate OnMemoryTrim;;
+	static FSimpleMulticastDelegate OnMemoryTrim;
 	return OnMemoryTrim;
 }
 
diff --git a/Engine/Source/Runtime/Core/Public/Misc/CoreDelegates.h b/Engine/Source/Runtime/Core/Public/Misc/CoreDelegates.h
--- a/Engine/Source/Runtime/Core/Public/Misc/CoreDelegates.h
+++ b/Engine/Source/Runtime/Core/Public/Misc/CoreDelegates.h
@@ -247,20 +247,41 @@ public:
 	/** Called when the OS goes into low power mode */
 	static TMulticastDelegate<void(bool)> OnLowPowerMode;
 
+	UE_DEPRECATED(5.2, "This delegate is not thread-safe. Please use FCoreDelegates::TSCountPreLoadConfigFileRespondersDelegate()")
+	static TMulticastDelegate<void(const TCHAR* IniFilename, int32& ResponderCount)> CountPreLoadConfigFileRespondersDelegate;
+	static TTSMulticastDelegate<void(const TCHAR* IniFilename, int32& ResponderCount)>& TSCountPreLoadConfigFileRespondersDelegate();
 
-	static TTSMulticastDelegate<void(const TCHAR* IniFilename, int32& ResponderCount)> CountPreLoadConfigFileRespondersDelegate;
-	static TTSMulticastDelegate<void(const TCHAR* IniFilename, FString& LoadedContents)> PreLoadConfigFileDelegate;
-	static TTSMulticastDelegate<void(const TCHAR* IniFilename, const FString& ContentsToSave, int32& SavedCount)> PreSaveConfigFileDelegate;
+	UE_DEPRECATED(5.2, "This delegate is not thread-safe. Please use FCoreDelegates::TSPreLoadConfigFileDelegate()")
+	static TMulticastDelegate<void(const TCHAR* IniFilename, FString& LoadedContents)> PreLoadConfigFileDelegate;
+	static TTSMulticastDelegate<void(const TCHAR* IniFilename, FString& LoadedContents)>& TSPreLoadConfigFileDelegate();
 
-	static TTSMulticastDelegate<void(const FConfigFile*)> OnFConfigCreated;
-	static TTSMulticastDelegate<void(const FConfigFile*)> OnFConfigDeleted;
+	UE_DEPRECATED(5.2, "This delegate is not thread-safe. Please use FCoreDelegates::TSPreSaveConfigFileDelegate()")
+	static TMulticastDelegate<void(const TCHAR* IniFilename, const FString& ContentsToSave, int32& SavedCount)> PreSaveConfigFileDelegate;
+	static TTSMulticastDelegate<void(const TCHAR* IniFilename, const FString& ContentsToSave, int32& SavedCount)>& TSPreSaveConfigFileDelegate();
 
-	static TTSMulticastDelegate<void(const TCHAR* IniFilename, const TCHAR* SectionName, const TCHAR* Key)> OnConfigValueRead;
+	UE_DEPRECATED(5.2, "This delegate is not thread-safe. Please use FCoreDelegates::TSOnFConfigCreated()")
+	static TMulticastDelegate<void(const FConfigFile*)> OnFConfigCreated;
+	static TTSMulticastDelegate<void(const FConfigFile*)>& TSOnFConfigCreated();
 
-	static TTSMulticastDelegate<void(const TCHAR* IniFilename, const TCHAR* SectionName)> OnConfigSectionRead;
-	static TTSMulticastDelegate<void(const TCHAR* IniFilename, const TCHAR* SectionName)> OnConfigSectionNameRead;
+	UE_DEPRECATED(5.2, "This delegate is not thread-safe. Please use FCoreDelegates::TSOnFConfigDeleted()")
+	static TMulticastDelegate<void(const FConfigFile*)> OnFConfigDeleted;
+	static TTSMulticastDelegate<void(const FConfigFile*)>& TSOnFConfigDeleted();
 
-	static TTSMulticastDelegate<void(const FString& IniFilename, const TSet<FString>& SectionNames)> OnConfigSectionsChanged;
+	UE_DEPRECATED(5.2, "This delegate is not thread-safe. Please use FCoreDelegates::TSOnConfigValueRead()")
+	static TMulticastDelegate<void(const TCHAR* IniFilename, const TCHAR* SectionName, const TCHAR* Key)> OnConfigValueRead;
+	static TTSMulticastDelegate<void(const TCHAR* IniFilename, const TCHAR* SectionName, const TCHAR* Key)>& TSOnConfigValueRead();
+
+	UE_DEPRECATED(5.2, "This delegate is not thread-safe. Please use FCoreDelegates::TSOnConfigSectionRead()")
+	static TMulticastDelegate<void(const TCHAR* IniFilename, const TCHAR* SectionName)> OnConfigSectionRead;
+	static TTSMulticastDelegate<void(const TCHAR* IniFilename, const TCHAR* SectionName)>& TSOnConfigSectionRead();
+
+	UE_DEPRECATED(5.2, "This delegate is not thread-safe. Please use FCoreDelegates::TSOnConfigSectionNameRead()")
+	static TMulticastDelegate<void(const TCHAR* IniFilename, const TCHAR* SectionName)> OnConfigSectionNameRead;
+	static TTSMulticastDelegate<void(const TCHAR* IniFilename, const TCHAR* SectionName)>& TSOnConfigSectionNameRead();
+
+	UE_DEPRECATED(5.2, "This delegate is not thread-safe. Please use FCoreDelegates::TSOnConfigSectionsChanged()")
+	static TMulticastDelegate<void(const FString& IniFilename, const TSet<FString>& SectionNames)> OnConfigSectionsChanged;
+	static TTSMulticastDelegate<void(const FString& IniFilename, const TSet<FString>& SectionNames)>& TSOnConfigSectionsChanged();
 
 	static TMulticastDelegate<void(const TCHAR* SectionName, const TCHAR* IniFilename, uint32 SetBy, bool bAllowCheating)> OnApplyCVarFromIni;
 
@@ -421,7 +442,9 @@ public:
 	static TMulticastDelegate<void()>& GetPreMainInitDelegate();
 	
 	/** Sent when GConfig is finished initializing */
-	static TTSMulticastDelegate<void()> ConfigReadyForUse;
+	UE_DEPRECATED(5.2, "This delegate is not thread-safe. Please use FCoreDelegates::TSConfigReadyForUse()")
+	static TMulticastDelegate<void()> ConfigReadyForUse;
+	static TTSMulticastDelegate<void()>& TSConfigReadyForUse();
 
 	/** Callback for notifications regarding changes of the rendering thread. */
 
@@ -499,7 +522,10 @@ public:
 		{
 		}
 	};
-	static TTSMulticastDelegate<void(FExtraBinaryConfigData&)> AccessExtraBinaryConfigData;
+	
+	UE_DEPRECATED(5.2, "This delegate is not thread-safe. Please use FCoreDelegates::TSAccessExtraBinaryConfigData()")
+	static TMulticastDelegate<void(FExtraBinaryConfigData&)> AccessExtraBinaryConfigData;
+	static TTSMulticastDelegate<void(FExtraBinaryConfigData&)>& TSAccessExtraBinaryConfigData();
 
 	using FAttachShaderReadRequestFunc = TFunctionRef<class FIoRequest(const class FIoChunkId&, FGraphEventRef)>;
 	static TDelegate<void(TArrayView<const FSHAHash>, FAttachShaderReadRequestFunc)> PreloadPackageShaderMaps;
diff --git a/Engine/Source/Runtime/Online/HTTP/Private/HttpRetrySystem.cpp b/Engine/Source/Runtime/Online/HTTP/Private/HttpRetrySystem.cpp
--- a/Engine/Source/Runtime/Online/HTTP/Private/HttpRetrySystem.cpp
+++ b/Engine/Source/Runtime/Online/HTTP/Private/HttpRetrySystem.cpp
@@ -644,12 +644,12 @@ FHttpRetrySystem::FManager::FHttpLogVerbosityTracker& FHttpRetrySystem::FManager
 FHttpRetrySystem::FManager::FHttpLogVerbosityTracker::FHttpLogVerbosityTracker()
 {
 	UpdateSettingsFromConfig();
-	FCoreDelegates::OnConfigSectionsChanged.AddRaw(this, &FHttpLogVerbosityTracker::OnConfigSectionsChanged);
+	FCoreDelegates::TSOnConfigSectionsChanged().AddRaw(this, &FHttpLogVerbosityTracker::OnConfigSectionsChanged);
 }
 
 FHttpRetrySystem::FManager::FHttpLogVerbosityTracker::~FHttpLogVerbosityTracker()
 {
-	FCoreDelegates::OnConfigSectionsChanged.RemoveAll(this);
+	FCoreDelegates::TSOnConfigSectionsChanged().RemoveAll(this);
 }
 
 void FHttpRetrySystem::FManager::FHttpLogVerbosityTracker::IncrementRetriedRequests()
