Remove return on void function

diff --git a/Engine/Source/Runtime/Core/Public/Delegates/TSMulticastDelegateBase.h b/Engine/Source/Runtime/Core/Public/Delegates/TSMulticastDelegateBase.h
--- a/Engine/Source/Runtime/Core/Public/Delegates/TSMulticastDelegateBase.h
+++ b/Engine/Source/Runtime/Core/Public/Delegates/TSMulticastDelegateBase.h
@@ -70,7 +70,6 @@ protected:
 	{
 		FScopeLock Lock(&CS);
 		Super::template CopyFrom<DelegateInstanceInterfaceType, DelegateType>(Other);
-		return *this;
 	}
 
 	template<typename DelegateInstanceInterfaceType, typename DelegateBaseType, typename... ParamTypes>
