Fixing bug where LruCache::FindByPredicate and PsoLruCache::FindByPredicate are not returning a pointer but a value when an item is found.

diff --git a/Engine/Source/Runtime/Core/Public/Containers/LruCache.h b/Engine/Source/Runtime/Core/Public/Containers/LruCache.h
--- a/Engine/Source/Runtime/Core/Public/Containers/LruCache.h
+++ b/Engine/Source/Runtime/Core/Public/Containers/LruCache.h
@@ -383,7 +383,7 @@ public:
 	 * Find the value of an entry using a predicate.
 	 *
 	 * @param Pred The predicate functor to apply to each entry.
-	 * @return A value for which the predicate returned true, or nullptr if not found.
+	 * @return Pointer to value for which the predicate returned true, or nullptr if not found.
 	 * @see ContainsByPredicate, FilterByPredicate, RemoveByPredicate
 	 */
 	template<typename Predicate>
@@ -393,7 +393,7 @@ public:
 		{
 			if (Pred(Entry->Key, Entry->Value))
 			{
-				return Entry->Value;
+				return &Entry->Value;
 			}
 		}
 
diff --git a/Engine/Source/Runtime/RHI/Public/PsoLruCache.h b/Engine/Source/Runtime/RHI/Public/PsoLruCache.h
--- a/Engine/Source/Runtime/RHI/Public/PsoLruCache.h
+++ b/Engine/Source/Runtime/RHI/Public/PsoLruCache.h
@@ -273,7 +273,7 @@ public:
 	 * Find the value of an entry using a predicate.
 	 *
 	 * @param Pred The predicate functor to apply to each entry.
-	 * @return A value for which the predicate returned true, or nullptr if not found.
+	 * @return Pointer to value for which the predicate returned true, or nullptr if not found.
 	 * @see ContainsByPredicate, FilterByPredicate, RemoveByPredicate
 	 */
 	template<typename Predicate>
@@ -283,7 +283,7 @@ public:
 		{
 			if (Pred(Entry->Key, Entry->Value))
 			{
-				return Entry->Value;
+				return &Entry->Value;
 			}
 		}
 
