Skip old assumption checks when in retraction context for game, slate and parallel rendering threads

diff --git a/Engine/Source/Runtime/Core/Private/HAL/ThreadingBase.cpp b/Engine/Source/Runtime/Core/Private/HAL/ThreadingBase.cpp
--- a/Engine/Source/Runtime/Core/Private/HAL/ThreadingBase.cpp
+++ b/Engine/Source/Runtime/Core/Private/HAL/ThreadingBase.cpp
@@ -177,9 +177,12 @@ CORE_API bool IsInGameThread()
 {
 	if (GIsGameThreadIdInitialized)
 	{
+PRAGMA_DISABLE_DEPRECATION_WARNINGS
 		bool newValue = FTaskTagScope::IsCurrentTag(ETaskTag::EGameThread) || FTaskTagScope::IsRunningDuringStaticInit();
 #if !UE_BUILD_SHIPPING && !UE_BUILD_TEST
-		if (!LowLevelTasks::FSchedulerTls::IsBusyWaiting() && !CoroTask_Detail::FCoroLocalState::IsCoroLaunchedTask())
+		if (!LowLevelTasks::FSchedulerTls::IsBusyWaiting() &&
+			!CoroTask_Detail::FCoroLocalState::IsCoroLaunchedTask() &&
+			!UE::Tasks::Private::IsThreadRetractingTask())
 		{
 			const uint32 CurrentThreadId = FPlatformTLS::GetCurrentThreadId();
 			bool oldValue = CurrentThreadId == GGameThreadId;
@@ -188,6 +191,7 @@ CORE_API bool IsInGameThread()
 		}
 #endif
 		return newValue;
+PRAGMA_ENABLE_DEPRECATION_WARNINGS
 	}
 
 	return true;
@@ -200,10 +204,13 @@ CORE_API bool IsInParallelGameThread()
 
 CORE_API bool IsInSlateThread()
 {
+PRAGMA_DISABLE_DEPRECATION_WARNINGS
 	// If this explicitly is a slate thread, not just the main thread running slate
 	bool newValue = FTaskTagScope::IsCurrentTag(ETaskTag::ESlateThread);
 #if !UE_BUILD_SHIPPING && !UE_BUILD_TEST
-	if (!LowLevelTasks::FSchedulerTls::IsBusyWaiting() && !CoroTask_Detail::FCoroLocalState::IsCoroLaunchedTask())
+	if (!LowLevelTasks::FSchedulerTls::IsBusyWaiting() &&
+		!CoroTask_Detail::FCoroLocalState::IsCoroLaunchedTask() &&
+		!UE::Tasks::Private::IsThreadRetractingTask())
 	{
 		bool oldValue = GSlateLoadingThreadId != 0 && FPlatformTLS::GetCurrentThreadId() == GSlateLoadingThreadId;
 		ensureMsgf(oldValue == newValue, TEXT("oldValue(%i) newValue(%i) If this check fails make sure that there is a FTaskTagScope(ETaskTag::ESlateThread) as as deep as possible on the current callstack, you can see the current value in ActiveNamedThreads(%x)"), oldValue, newValue, FTaskTagScope::GetCurrentTag());
@@ -211,6 +218,7 @@ CORE_API bool IsInSlateThread()
 	}
 #endif
 	return newValue;
+PRAGMA_ENABLE_DEPRECATION_WARNINGS
 }
 
 // tasks pipe that arranges audio tasks execution one after another so no synchronisation between them is required
@@ -244,7 +252,9 @@ PRAGMA_DISABLE_DEPRECATION_WARNINGS
 		? FTaskTagScope::IsCurrentTag(ETaskTag::EGameThread)
 		: FTaskTagScope::IsCurrentTag(ETaskTag::EAudioThread);
 #if !UE_BUILD_SHIPPING && !UE_BUILD_TEST
-	if (!LowLevelTasks::FSchedulerTls::IsBusyWaiting() && !CoroTask_Detail::FCoroLocalState::IsCoroLaunchedTask())
+	if (!LowLevelTasks::FSchedulerTls::IsBusyWaiting() &&
+		!CoroTask_Detail::FCoroLocalState::IsCoroLaunchedTask() &&
+		!UE::Tasks::Private::IsThreadRetractingTask())
 	{
 		bool oldValue = FPlatformTLS::GetCurrentThreadId() == ((nullptr == GAudioThread || GIsAudioThreadSuspended.load(std::memory_order_relaxed)) ? GGameThreadId : GAudioThread->GetThreadID());
 		ensureMsgf(oldValue == newValue, TEXT("oldValue(%i) newValue(%i) If this check fails make sure that there is a FTaskTagScope(ETaskTag::EAudioThread) as deep as possible on the current callstack, you can see the current value in ActiveNamedThreads(%x)"), oldValue, newValue, FTaskTagScope::GetCurrentTag());
@@ -287,7 +297,6 @@ PRAGMA_DISABLE_DEPRECATION_WARNINGS
 		: FTaskTagScope::IsCurrentTag(ETaskTag::ERenderingThread);
 
 #if !UE_BUILD_SHIPPING && !UE_BUILD_TEST
-	// in all these cases the `newValue` can be wrong, so the "ensure" is disabled
 	if (!LowLevelTasks::FSchedulerTls::IsBusyWaiting() && 
 		!CoroTask_Detail::FCoroLocalState::IsCoroLaunchedTask() && 
 		!UE::Tasks::Private::IsThreadRetractingTask())
@@ -321,7 +330,9 @@ PRAGMA_DISABLE_DEPRECATION_WARNINGS
 	}
 
 #if !UE_BUILD_SHIPPING && !UE_BUILD_TEST
-	if (!LowLevelTasks::FSchedulerTls::IsBusyWaiting() && !CoroTask_Detail::FCoroLocalState::IsCoroLaunchedTask())
+	if (!LowLevelTasks::FSchedulerTls::IsBusyWaiting() &&
+		!CoroTask_Detail::FCoroLocalState::IsCoroLaunchedTask() &&
+		!UE::Tasks::Private::IsThreadRetractingTask())
 	{
 		const uint32 CurrentThreadId = FPlatformTLS::GetCurrentThreadId();
 		bool oldValue = ((GRenderThreadId == 0) || bLocalIsLoadingThreadSuspended) ?  true : CurrentThreadId != GGameThreadId;
