fixed System.Core.Tasks.Dependencies unit test

diff --git a/Engine/Source/Runtime/Core/Private/Tests/Tasks/TasksTest.cpp b/Engine/Source/Runtime/Core/Private/Tests/Tasks/TasksTest.cpp
--- a/Engine/Source/Runtime/Core/Private/Tests/Tasks/TasksTest.cpp
+++ b/Engine/Source/Runtime/Core/Private/Tests/Tasks/TasksTest.cpp
@@ -693,45 +693,37 @@ namespace UE { namespace TasksTests
 
 	IMPLEMENT_SIMPLE_AUTOMATION_TEST(FTasksDependenciesTest, "System.Core.Tasks.Dependencies", EAutomationTestFlags::ApplicationContextMask | EAutomationTestFlags::EngineFilter | EAutomationTestFlags::Disabled);
 
-	template<uint64 NumBranches, uint64 NumLoops, uint64 NumTasks>
+	template<uint64 NumBranches, uint64 NumTasks>
 	void DependenciesPerfTest()
 	{
-		auto Branch = []
+
+		TArray<TTask<FTaskEvent>> Branches;
+		Branches.Reserve(NumBranches);
+		for (uint64 BranchIndex = 0; BranchIndex != NumBranches; ++BranchIndex)
 		{
-			FTask Joiner;
-			for (uint64 LoopIndex = 0; LoopIndex != NumLoops; ++LoopIndex)
+			auto Branch = []
 			{
 				TArray<FTask> Tasks;
 				Tasks.Reserve(NumTasks);
 				for (uint64 TaskIndex = 0; TaskIndex != NumTasks; ++TaskIndex)
 				{
-					if (Joiner.IsValid())
-					{
-						Tasks.Add(Launch(UE_SOURCE_LOCATION, [] { /*FPlatformProcess::YieldCycles(100000);*/ }, Joiner));
-					}
-					else
-					{
-						Tasks.Add(Launch(UE_SOURCE_LOCATION, [] { /*FPlatformProcess::YieldCycles(100000);*/ }));
-					}
+					Tasks.Add(Launch(UE_SOURCE_LOCATION, [] { /*FPlatformProcess::YieldCycles(100000);*/ }));
 				}
-				Joiner = Launch(UE_SOURCE_LOCATION, [] {}, Tasks);
-			}
-			return Joiner;
-		};
+				FTaskEvent Joiner{ UE_SOURCE_LOCATION };
+				Joiner.AddPrerequisites(Tasks);
+				Joiner.Trigger();
+				return Joiner;
+			};
 
-		TArray<TTask<FTask>> Branches;
-		Branches.Reserve(NumBranches);
-		for (uint64 BranchIndex = 0; BranchIndex != NumBranches; ++BranchIndex)
-		{
 			Branches.Add(Launch(UE_SOURCE_LOCATION, MoveTemp(Branch)));
 		}
-		TArray<FTask> BranchTasks;
+		TArray<FTaskEvent> BranchTasks;
 		BranchTasks.Reserve(NumBranches);
-		for (TTask<FTask>& Task : Branches)
+		for (TTask<FTaskEvent>& Task : Branches)
 		{
 			BranchTasks.Add(Task.GetResult());
 		}
-		Wait(Branches);
+		Wait(BranchTasks);
 	}
 
 	bool FTasksDependenciesTest::RunTest(const FString& Parameters)
@@ -870,7 +862,7 @@ namespace UE { namespace TasksTests
 			Task.Wait();
 		}
 
-		UE_BENCHMARK(5, DependenciesPerfTest<50, 5, 50>);
+		UE_BENCHMARK(5, DependenciesPerfTest<150, 150>);
 
 		return true;
 	}
