enabled System.Core.Tasks.Dependencies test after fixing it in 22109874

diff --git a/Engine/Source/Runtime/Core/Private/Tests/Tasks/TasksTest.cpp b/Engine/Source/Runtime/Core/Private/Tests/Tasks/TasksTest.cpp
--- a/Engine/Source/Runtime/Core/Private/Tests/Tasks/TasksTest.cpp
+++ b/Engine/Source/Runtime/Core/Private/Tests/Tasks/TasksTest.cpp
@@ -691,7 +691,7 @@ namespace UE { namespace TasksTests
 		return true;
 	}
 
-	IMPLEMENT_SIMPLE_AUTOMATION_TEST(FTasksDependenciesTest, "System.Core.Tasks.Dependencies", EAutomationTestFlags::ApplicationContextMask | EAutomationTestFlags::EngineFilter | EAutomationTestFlags::Disabled);
+	IMPLEMENT_SIMPLE_AUTOMATION_TEST(FTasksDependenciesTest, "System.Core.Tasks.Dependencies", EAutomationTestFlags::ApplicationContextMask | EAutomationTestFlags::EngineFilter);
 
 	template<uint64 NumBranches, uint64 NumTasks>
 	void DependenciesPerfTest()
