FThreadManager crashes if its `Threads` collection is modified while it's iterating over it. This happens when running with `-nothreading`. Fixed.

diff --git a/Engine/Source/Runtime/Core/Private/HAL/ThreadingBase.cpp b/Engine/Source/Runtime/Core/Private/HAL/ThreadingBase.cpp
--- a/Engine/Source/Runtime/Core/Private/HAL/ThreadingBase.cpp
+++ b/Engine/Source/Runtime/Core/Private/HAL/ThreadingBase.cpp
@@ -543,17 +543,16 @@ void FThreadManager::Tick()
 	{
 		QUICK_SCOPE_CYCLE_COUNTER(STAT_FSingleThreadManager_Tick);
 
-		FScopeLock ThreadsLock(&ThreadsCritical);
-
-		// Tick all registered fake threads.
-		for (TPair<uint32, FRunnableThread*>& ThreadPair : Threads)
-		{
-			// Only fake and forkable threads are ticked by the ThreadManager
-			if( ThreadPair.Value->GetThreadType() != FRunnableThread::ThreadType::Real )
+		ForEachThread(
+			[] (uint32 ThreadId, FRunnableThread* Thread)
 			{
-				ThreadPair.Value->Tick();
+				// Only fake and forkable threads are ticked by the ThreadManager
+				if (Thread->GetThreadType() != FRunnableThread::ThreadType::Real)
+				{
+					Thread->Tick();
+				}
 			}
-		}
+		);
 	}
 }
 
@@ -602,19 +601,23 @@ void FThreadManager::GetAllThreadStackBackTraces(TArray<FThreadStackBackTrace>&
 	StackTraces.Empty(NumThreads);
 	GetAllThreadStackBackTraces_ProcessSingle(CurThreadId, GGameThreadId, TEXT("GameThread"), StackTraces.AddDefaulted_GetRef());
 
-	for (const TPair<uint32, FRunnableThread*>& Pair : Threads)
-	{
-		const uint32 Id = Pair.Key;
-		const FString& Name = Pair.Value->GetThreadName();
-		GetAllThreadStackBackTraces_ProcessSingle(CurThreadId, Id, *Name, StackTraces.AddDefaulted_GetRef());
-	}
+	ForEachThread(
+		[CurThreadId, &StackTraces] (uint32 ThreadId, FRunnableThread* Thread)
+		{
+			const FString& Name = Thread->GetThreadName();
+			GetAllThreadStackBackTraces_ProcessSingle(CurThreadId, ThreadId, *Name, StackTraces.AddDefaulted_GetRef());
+		}
+	);
 }
 #endif
 
-void FThreadManager::ForEachThread(TFunction<void(uint32, class FRunnableThread*)> Func)
+void FThreadManager::ForEachThread(TFunction<void(uint32, FRunnableThread*)> Func)
 {
 	FScopeLock Lock(&ThreadsCritical);
-	for (const TPair<uint32, FRunnableThread*>& Pair : Threads)
+	// threads can be added or removed while iterating over them, thus invalidating the iterator, so we iterate over the copy of threads collection
+	FThreads ThreadsCopy = Threads;
+
+	for (const TPair<uint32, FRunnableThread*>& Pair : ThreadsCopy)
 	{
 		Func(Pair.Key, Pair.Value);
 	}
@@ -629,14 +632,15 @@ FThreadManager& FThreadManager::Get()
 TArray<FRunnableThread*> FThreadManager::GetForkableThreads()
 {
 	TArray<FRunnableThread*> ForkableThreads;
-	FScopeLock Lock(&ThreadsCritical);
-	for (const TPair<uint32, FRunnableThread*>& Pair : Threads)
-	{
-		if (Pair.Value->GetThreadType() == FRunnableThread::ThreadType::Forkable)
+	ForEachThread(
+		[&ForkableThreads] (uint32 ThreadId, FRunnableThread* Thread)
 		{
-			ForkableThreads.Add(Pair.Value);
+			if (Thread->GetThreadType() == FRunnableThread::ThreadType::Forkable)
+			{
+				ForkableThreads.Add(Thread);
+			}
 		}
-	}
+	);
 
 	return ForkableThreads;
 }
diff --git a/Engine/Source/Runtime/Core/Public/HAL/ThreadManager.h b/Engine/Source/Runtime/Core/Public/HAL/ThreadManager.h
--- a/Engine/Source/Runtime/Core/Public/HAL/ThreadManager.h
+++ b/Engine/Source/Runtime/Core/Public/HAL/ThreadManager.h
@@ -11,6 +11,8 @@
 	#define PLATFORM_SUPPORTS_ALL_THREAD_BACKTRACES (PLATFORM_WINDOWS || PLATFORM_MAC)
 #endif
 
+class FRunnableThread;
+
 /**
  * Manages runnables and runnable threads.
  */
@@ -19,8 +21,9 @@ class CORE_API FThreadManager
 	/** Critical section for ThreadList */
 	FCriticalSection ThreadsCritical;
 
+	using FThreads = TMap<uint32, FRunnableThread*, TInlineSetAllocator<256>>;
 	/** List of thread objects to be ticked. */
-	TMap<uint32, class FRunnableThread*, TInlineSetAllocator<256>> Threads;
+	FThreads Threads;
 
 public:
 
@@ -30,7 +33,7 @@ public:
 	* @param Thread thread object.
 	* @see RemoveThread
 	*/
-	void AddThread(uint32 ThreadId, class FRunnableThread* Thread);
+	void AddThread(uint32 ThreadId, FRunnableThread* Thread);
 
 	/**
 	* Used internally to remove thread object.
@@ -38,7 +41,7 @@ public:
 	* @param Thread thread object to be removed.
 	* @see AddThread
 	*/
-	void RemoveThread(class FRunnableThread* Thread);
+	void RemoveThread(FRunnableThread* Thread);
 
 	/** Get the number of registered threads */
 	int32 NumThreads() const { return Threads.Num(); }
@@ -79,7 +82,7 @@ public:
 	 * Enumerate each thread.
 	 *
 	 */
-	void ForEachThread(TFunction<void(uint32, class FRunnableThread*)> Func);
+	void ForEachThread(TFunction<void(uint32 ThreadId, FRunnableThread* Thread)> Func);
 
 	/**
 	 * Access to the singleton object.
