Put critical section around android manifest file access.

diff --git a/Engine/Source/Runtime/Core/Private/Android/AndroidPlatformFile.cpp b/Engine/Source/Runtime/Core/Private/Android/AndroidPlatformFile.cpp
--- a/Engine/Source/Runtime/Core/Private/Android/AndroidPlatformFile.cpp
+++ b/Engine/Source/Runtime/Core/Private/Android/AndroidPlatformFile.cpp
@@ -417,6 +417,7 @@ private:
 	bool bInitialized;
 	FString ManifestFileName;
 	TMap<FString, FDateTime> ManifestEntries;
+	FCriticalSection ManifestEntriesCS;
 public:
 
 	FAndroidFileManifestReader( const FString& InManifestFileName ) : ManifestFileName(InManifestFileName), bInitialized(false)
@@ -425,6 +426,8 @@ public:
 
 	bool GetFileTimeStamp( const FString& FileName, FDateTime& DateTime ) 
 	{
+		FScopeLock Lock(&ManifestEntriesCS);
+
 		if ( bInitialized == false )
 		{
 			Read();
@@ -449,6 +452,8 @@ public:
 
 	bool SetFileTimeStamp( const FString& FileName, const FDateTime& DateTime )
 	{
+		FScopeLock Lock(&ManifestEntriesCS);
+
 		if (bInitialized == false)
 		{
 			Read();
@@ -472,6 +477,8 @@ public:
 
 	bool DeleteFileTimeStamp(const FString& FileName)
 	{
+		FScopeLock Lock(&ManifestEntriesCS);
+
 		if (bInitialized == false)
 		{
 			Read();
@@ -496,6 +503,8 @@ public:
 	// read manifest from disk
 	void Read()
 	{
+		FScopeLock Lock(&ManifestEntriesCS);
+
 		// Local filepaths are directly in the deployment directory.
 		static const FString &BasePath = GetFileBasePath();
 		const FString ManifestPath = BasePath + ManifestFileName;
@@ -571,7 +580,8 @@ public:
 
 	void Write()
 	{
-		
+		FScopeLock Lock(&ManifestEntriesCS);
+
 		// Local filepaths are directly in the deployment directory.
 		static const FString &BasePath = GetFileBasePath();
 		const FString ManifestPath = BasePath + ManifestFileName;
