Fix nopch build fail

diff --git a/Engine/Source/Runtime/Core/Private/Android/AndroidPlatformFile.cpp b/Engine/Source/Runtime/Core/Private/Android/AndroidPlatformFile.cpp
--- a/Engine/Source/Runtime/Core/Private/Android/AndroidPlatformFile.cpp
+++ b/Engine/Source/Runtime/Core/Private/Android/AndroidPlatformFile.cpp
@@ -20,6 +20,7 @@
 #include "Android/AndroidJava.h"
 #include "Containers/Map.h"
 #include <limits>
+#include "Misc/ScopeLock.h"
 
 #if PLATFORM_USE_PLATFORM_FILE_MANAGED_STORAGE_WRAPPER
 #include "HAL/IPlatformFileManagedStorageWrapper.h"
