Add an option to return the mount point with its slash in FPathViews::GetMountPointNameFromPath
Add unit test for that function

diff --git a/Engine/Source/Runtime/Core/Private/Misc/PathViews.cpp b/Engine/Source/Runtime/Core/Private/Misc/PathViews.cpp
--- a/Engine/Source/Runtime/Core/Private/Misc/PathViews.cpp
+++ b/Engine/Source/Runtime/Core/Private/Misc/PathViews.cpp
@@ -681,31 +681,33 @@ void FPathViews::AppendPath(FStringBuilderBase& InOutPath, FStringView AppendPat
 	}
 }
 
-FStringView FPathViews::GetMountPointNameFromPath(const FStringView InPath, bool* bOutHadClassesPrefix /*= nullptr*/)
+FStringView FPathViews::GetMountPointNameFromPath(const FStringView InPath, bool* bOutHadClassesPrefix /*= nullptr*/, bool bInWithoutSlashes /*= true*/)
 {
 	FStringView MountPointStringView;
-
-	int32 SecondForwardSlash = INDEX_NONE;
-	if (FStringView(InPath.GetData() + 1, InPath.Len() - 1).FindChar(TEXT('/'), SecondForwardSlash))
+	if (InPath.Len() > 0 && InPath[0] == '/')
 	{
-		MountPointStringView = FStringView(InPath.GetData() + 1, SecondForwardSlash);
-	}
-	else
-	{
-		MountPointStringView = FStringView(InPath.GetData() + 1, InPath.Len() - 1);
-	}
+		int32 SecondForwardSlash = InPath.Len();
+		int32 WithoutSlashes = bInWithoutSlashes ? 1 : 0;
+		if (FStringView(InPath.GetData() + 1, InPath.Len() - 1).FindChar(TEXT('/'), SecondForwardSlash))
+		{
+			MountPointStringView = FStringView(InPath.GetData() + WithoutSlashes, SecondForwardSlash + 1 - WithoutSlashes);
+		}
+		else
+		{
+			MountPointStringView = FStringView(InPath.GetData() + WithoutSlashes, InPath.Len() - WithoutSlashes);
+		}
 
-	static const FString ClassesPrefix = TEXT("Classes_");
-	const bool bHasClassesPrefix = MountPointStringView.StartsWith(ClassesPrefix);
-	if (bHasClassesPrefix)
-	{
-		MountPointStringView.RightInline(MountPointStringView.Len() - ClassesPrefix.Len());
-	}
+		static const FString ClassesPrefix = TEXT("Classes_");
+		const bool bHasClassesPrefix = MountPointStringView.StartsWith(ClassesPrefix);
+		if (bHasClassesPrefix)
+		{
+			MountPointStringView.RightInline(MountPointStringView.Len() - ClassesPrefix.Len());
+		}
 
-	if (bOutHadClassesPrefix)
-	{
-		*bOutHadClassesPrefix = bHasClassesPrefix;
+		if (bOutHadClassesPrefix)
+		{
+			*bOutHadClassesPrefix = bHasClassesPrefix;
+		}
 	}
-
 	return MountPointStringView;
 }
\ No newline at end of file
diff --git a/Engine/Source/Runtime/Core/Private/Tests/Misc/PathViewsTest.cpp b/Engine/Source/Runtime/Core/Private/Tests/Misc/PathViewsTest.cpp
--- a/Engine/Source/Runtime/Core/Private/Tests/Misc/PathViewsTest.cpp
+++ b/Engine/Source/Runtime/Core/Private/Tests/Misc/PathViewsTest.cpp
@@ -708,6 +708,43 @@ bool FPathViewsAppendPathTest::RunTest(const FString& InParameters)
 	return true;
 }
 
+
+IMPLEMENT_CUSTOM_SIMPLE_AUTOMATION_TEST(FPathViewsGetMountPointNameTest, FPathViewsTest, "System.Core.Misc.PathViews.GetMountPointNameFromPath", EAutomationTestFlags::ApplicationContextMask | EAutomationTestFlags::SmokeFilter)
+bool FPathViewsGetMountPointNameTest::RunTest(const FString& InParameters)
+{
+	auto RunGetMountPointName = [this](const TCHAR* Base, const TCHAR* ExpectedWithoutSlash, const TCHAR* ExpectedWithSlash)
+	{
+		TStringBuilder<128> BaseBuilder;
+		BaseBuilder << Base;
+		bool bHadClassPrefix = false;
+		FStringView Result = FPathViews::GetMountPointNameFromPath(Base, &bHadClassPrefix, true);
+		if (Result != FStringView(ExpectedWithoutSlash))
+		{
+			AddError(FString::Printf(TEXT("GetMountPointNameFromPath(\"%s\", &bHadClassPrefix, true) == \"%.*s\", expected \"%s\"."),
+				Base, Result.Len(), Result.GetData(), ExpectedWithoutSlash));
+		}
+		Result = FPathViews::GetMountPointNameFromPath(Base, &bHadClassPrefix, false);
+		if (Result != FStringView(ExpectedWithSlash))
+		{
+			AddError(FString::Printf(TEXT("AppendPath(\"%s\", \"%s\") == \"%.*s\", expected \"%s\"."),
+				Base, Result.Len(), Result.GetData(), ExpectedWithSlash));
+		}
+	};
+
+	RunGetMountPointName(TEXT(""), TEXT(""), TEXT(""));
+	RunGetMountPointName(TEXT("Root"), TEXT(""), TEXT(""));
+	RunGetMountPointName(TEXT("/"), TEXT(""), TEXT("/"));
+	RunGetMountPointName(TEXT("/Root"), TEXT("Root"), TEXT("/Root"));
+	RunGetMountPointName(TEXT("\\Root"), TEXT(""), TEXT(""));
+	RunGetMountPointName(TEXT("/Root/"), TEXT("Root"), TEXT("/Root"));
+	RunGetMountPointName(TEXT("/Root/A/B/"), TEXT("Root"), TEXT("/Root"));
+	RunGetMountPointName(TEXT("/Classes/A/"), TEXT("Classes"), TEXT("/Classes"));
+	RunGetMountPointName(TEXT("/Classes_A/B/"), TEXT("A"), TEXT("/Classes_A"));
+
+	return true;
+}
+
+
 IMPLEMENT_CUSTOM_SIMPLE_AUTOMATION_TEST(FPathViewsToAbsoluteTest, FPathViewsTest, "System.Core.Misc.PathViews.ToAbsolute", EAutomationTestFlags::ApplicationContextMask | EAutomationTestFlags::SmokeFilter)
 bool FPathViewsToAbsoluteTest::RunTest(const FString& InParameters)
 {
diff --git a/Engine/Source/Runtime/Core/Public/Misc/PathViews.h b/Engine/Source/Runtime/Core/Public/Misc/PathViews.h
--- a/Engine/Source/Runtime/Core/Public/Misc/PathViews.h
+++ b/Engine/Source/Runtime/Core/Public/Misc/PathViews.h
@@ -340,9 +340,10 @@ public:
 	static void AppendPath(FStringBuilderBase& InOutPath, FStringView AppendPath);
 
 	/**
-	 * Returns the name of the first folder in a path
-	 * Removes starting forward slash and Classes_ prefix
-	 * Example: "/Classes_A/Textures" returns "A" and sets bOutHadClassesPrefix=true
+	 * Returns the name of the mount point in a path
+	 * Removes starting forward slash and Classes_ prefix if bInWithoutSlashes is true
+	 * Example: "/Classes_A/Textures" returns "A" and sets bOutHadClassesPrefix=true if bInWithoutSlashes is true
+	 *           returns "/Classes_A" otherwise and sets bOutHadClassesPrefix=false
 	 */
-	static FStringView GetMountPointNameFromPath(const FStringView InPath, bool* bOutHadClassesPrefix = nullptr);
+	static FStringView GetMountPointNameFromPath(const FStringView InPath, bool* bOutHadClassesPrefix = nullptr, bool bInWithoutSlashes = true);
 };
