Prevent log spam in FUnixPlatformProcess::SetThreadNiceValue.

#tests - Compile and run

diff --git a/Engine/Source/Runtime/Core/Private/Unix/UnixPlatformProcess.cpp b/Engine/Source/Runtime/Core/Private/Unix/UnixPlatformProcess.cpp
--- a/Engine/Source/Runtime/Core/Private/Unix/UnixPlatformProcess.cpp
+++ b/Engine/Source/Runtime/Core/Private/Unix/UnixPlatformProcess.cpp
@@ -2106,10 +2106,15 @@ void FUnixPlatformProcess::SetThreadNiceValue(uint32_t ThreadId, int32 NiceValue
 	// We still try to set priority, but failure is not considered as an error
 	if (setpriority(PRIO_PROCESS, ThreadId, NiceValue) != 0 && WITH_PROCESS_PRIORITY_CONTROL)
 	{
-		// Unfortunately this is going to be a frequent occurence given that by default Unix doesn't allow raising priorities.
-		// NOTE: In WSL run "sudo prlimit --nice=40 --pid $$" to promote current shell to change nice values.
-		int ErrNo = errno;
-		UE_LOG(LogHAL, Error, TEXT("Can't set nice to %d. Reason = %s. Do you have CAP_SYS_NICE capability?"), NiceValue, ANSI_TO_TCHAR(strerror(ErrNo)));
+		static bool bIsLogged = false;
+		if (!bIsLogged)
+		{
+			bIsLogged = true;
+			// Unfortunately this is going to be a frequent occurence given that by default Unix doesn't allow raising priorities.
+			// NOTE: In WSL run "sudo prlimit --nice=40 --pid $$" to promote current shell to change nice values.
+			int ErrNo = errno;
+			UE_LOG(LogHAL, Error, TEXT("Can't set nice to %d. Reason = %s. Do you have CAP_SYS_NICE capability?"), NiceValue, ANSI_TO_TCHAR(strerror(ErrNo)));
+		}
 	}
 }
 
