Allow TSparseArray resize via EmplaceAt.

This is a part of a larger change, 21332740, which fixes TSparseArray::EmplaceAt, which is needed for physics user data.

Slack discussion here: https://epic.slack.com/archives/C024GKBSL/p1663795752638719

diff --git a/Engine/Source/Runtime/Core/Public/Containers/SparseArray.h b/Engine/Source/Runtime/Core/Public/Containers/SparseArray.h
--- a/Engine/Source/Runtime/Core/Public/Containers/SparseArray.h
+++ b/Engine/Source/Runtime/Core/Public/Containers/SparseArray.h
@@ -247,7 +247,7 @@ public:
 	FORCEINLINE int32 EmplaceAt(int32 Index, ArgsType&&... Args)
 	{
 		FSparseArrayAllocationInfo Allocation;
-		if(!AllocationFlags[Index])
+		if(!AllocationFlags.IsValidIndex(Index) || !AllocationFlags[Index])
 		{
 			Allocation = InsertUninitialized(Index);
 		}
