Fix for VM bytecode generation which was failing in a very small number of cases of the form:
A = GetFirstValue();
B = A;
A = GetSecondValue();
C = A; // <-- this would get GetSecondValue() assigned to C instead of GetFirstValue()

Adds an SSA pass to help reduce overwriting registers, adds logic to remove self assignments, and adds logic to stop propagating when values are dirtied

From a sampling of scripts, about 2% are changed (where 1-3 unnecessary ops have been deleted).

diff --git a/Engine/Source/Runtime/Core/Private/UObject/FortniteMainBranchObjectVersion.cpp b/Engine/Source/Runtime/Core/Private/UObject/FortniteMainBranchObjectVersion.cpp
--- a/Engine/Source/Runtime/Core/Private/UObject/FortniteMainBranchObjectVersion.cpp
+++ b/Engine/Source/Runtime/Core/Private/UObject/FortniteMainBranchObjectVersion.cpp
@@ -11,9 +11,9 @@ TMap<FGuid, FGuid> FFortniteMainBranchObjectVersion::GetSystemGuids()
 	SystemGuids.Add(DevGuids.MATERIALSHADERMAP_DERIVEDDATA_VER, FGuid("1DA85DA5733C4A489A1BE117CDC5ACCC"));
 	SystemGuids.Add(DevGuids.NANITE_DERIVEDDATA_VER, FGuid("93379C2F29624E74AFB14FB13E327FA2"));
 	SystemGuids.Add(DevGuids.NIAGARASHADERMAP_DERIVEDDATA_VER, FGuid("2289D5116CF94BC0AFEAEC541468E645"));
-	SystemGuids.Add(DevGuids.Niagara_LatestScriptCompileVersion, FGuid("04A14D478F6449BC92BD142968DCB2BB"));
+	SystemGuids.Add(DevGuids.Niagara_LatestScriptCompileVersion, FGuid("7379EA86433D4FD0857C15AFCAB9EC1D"));
 	SystemGuids.Add(DevGuids.SkeletalMeshDerivedDataVersion, FGuid("ACF593EAAE354FCBB34CF44F5AA1BFD2"));
 	SystemGuids.Add(DevGuids.STATICMESH_DERIVEDDATA_VER, FGuid("2B18DC81AAF64015A7901B5A3F8394C0"));
-
+	
 	return SystemGuids;
 }
