Fix crash when accessing FAutomationTestBase.ExpectedError

#rnx

diff --git a/Engine/Source/Runtime/Core/Private/Misc/AutomationTest.cpp b/Engine/Source/Runtime/Core/Private/Misc/AutomationTest.cpp
--- a/Engine/Source/Runtime/Core/Private/Misc/AutomationTest.cpp
+++ b/Engine/Source/Runtime/Core/Private/Misc/AutomationTest.cpp
@@ -835,7 +835,10 @@ bool FAutomationTestFramework::InternalStopTest(FAutomationTestExecutionInfo& Ou
 	// 3) Did we meet any errors that were expected with this test
 	bTestSuccessful = bTestSuccessful && !CurrentTest->HasAnyErrors() && CurrentTest->HasMetExpectedErrors();
 
-	CurrentTest->ExpectedErrors.Empty();
+	{
+		FScopeLock Lock(&CurrentTest->ActionCS);
+		CurrentTest->ExpectedErrors.Empty();
+	}
 
 	// Set the success state of the test based on the above criteria
 	CurrentTest->InternalSetSuccessState(bTestSuccessful);
@@ -1226,6 +1229,7 @@ void FAutomationTestBase::AddExpectedError(FString ExpectedErrorPattern, EAutoma
 {
 	if (Occurrences >= 0)
 	{
+		FScopeLock Lock(&ActionCS);
 		// If we already have an error matching string in our list, let's not add it again.
 		FAutomationExpectedError* FoundEntry = ExpectedErrors.FindByPredicate(
 			[ExpectedErrorPattern](const FAutomationExpectedError& InItem) 
@@ -1471,6 +1475,7 @@ bool FAutomationTestBase::TestNull(const TCHAR* What, const void* Pointer)
 
 bool FAutomationTestBase::IsExpectedError(const FString& Error)
 {
+	FScopeLock Lock(&ActionCS);
 	for (auto& EError : ExpectedErrors)
 	{
 		FRegexMatcher ErrorMatcher(EError.ErrorPattern, Error);
