Checked narrow cast for world positions

diff --git a/Engine/Source/Runtime/Core/Public/Math/Vector.h b/Engine/Source/Runtime/Core/Public/Math/Vector.h
--- a/Engine/Source/Runtime/Core/Public/Math/Vector.h
+++ b/Engine/Source/Runtime/Core/Public/Math/Vector.h
@@ -2654,6 +2654,21 @@ inline TVector<T> TVector2<T>::SphericalToUnitCartesian() const
 }
 
 } // namespace UE::Math
+	
+namespace LWC
+{
+// Validated narrowing cast for world positions. FVector -> FVector3f
+FORCEINLINE FVector3f NarrowWorldPositionChecked(const FVector& WorldPosition)
+{
+	FVector3f Narrowed;
+	constexpr FVector::FReal Precision = 1/16.0;
+	Narrowed.X = FloatCastChecked<float>(WorldPosition.X, Precision);
+	Narrowed.Y = FloatCastChecked<float>(WorldPosition.Y, Precision);
+	Narrowed.Z = FloatCastChecked<float>(WorldPosition.Z, Precision);
+	return Narrowed;
+}
+} // namespace UE::LWC
+
 } // namespace UE
 
 #if PLATFORM_ENABLE_VECTORINTRINSICS
