MemScope events do not need to be synchronised.

Something that is scoped by nature only pertains to the current thread and thus is implicitly synchronised. Furthermore, in memory-scope's case the alloc-free events are synchronised anyway which makes it redundant for the scoped events.

#rnx

diff --git a/Engine/Source/Runtime/Core/Private/ProfilingDebugging/TagTrace.cpp b/Engine/Source/Runtime/Core/Private/ProfilingDebugging/TagTrace.cpp
--- a/Engine/Source/Runtime/Core/Private/ProfilingDebugging/TagTrace.cpp
+++ b/Engine/Source/Runtime/Core/Private/ProfilingDebugging/TagTrace.cpp
@@ -21,11 +21,11 @@ UE_TRACE_EVENT_BEGIN(Memory, TagSpec, Important|NoSync)
 	UE_TRACE_EVENT_FIELD(UE::Trace::AnsiString, Display)
 UE_TRACE_EVENT_END()
 
-UE_TRACE_EVENT_BEGIN(Memory, MemoryScope)
+UE_TRACE_EVENT_BEGIN(Memory, MemoryScope, NoSync)
 	UE_TRACE_EVENT_FIELD(int32, Tag)
 UE_TRACE_EVENT_END()
 
-UE_TRACE_EVENT_BEGIN(Memory, MemoryScopePtr)
+UE_TRACE_EVENT_BEGIN(Memory, MemoryScopePtr, NoSync)
 	UE_TRACE_EVENT_FIELD(uint64, Ptr)
 UE_TRACE_EVENT_END()
 
