Adding Architecture to the log init info due to recent (ahem) issues with running the x64 build when it was believed the arm64 was running.

diff --git a/Engine/Source/Runtime/Core/Private/Misc/App.cpp b/Engine/Source/Runtime/Core/Private/Misc/App.cpp
--- a/Engine/Source/Runtime/Core/Private/Misc/App.cpp
+++ b/Engine/Source/Runtime/Core/Private/Misc/App.cpp
@@ -341,6 +341,14 @@ void FApp::PrintStartupLogMessages()
 	UE_LOG(LogInit, Log, TEXT("Compiled (32-bit): %s %s"), BuildSettings::GetBuildDate(), BuildSettings::GetBuildTime());
 #endif
 
+#if PLATFORM_CPU_ARM_FAMILY
+	UE_LOG(LogInit, Log, TEXT("Architecture: arm64"));
+#elif PLATFORM_CPU_X86_FAMILY
+	UE_LOG(LogInit, Log, TEXT("Architecture: x64"));
+#else
+#error No architecture defined!
+#endif // x64/arm
+
 	// Print compiler version info
 #if defined(__clang__)
 	UE_LOG(LogInit, Log, TEXT("Compiled with Clang: %s"), ANSI_TO_TCHAR(__clang_version__));
