unit-test for task pipes running with `-nothreading`

diff --git a/Engine/Source/Runtime/Core/Private/Tests/Tasks/TasksTest.cpp b/Engine/Source/Runtime/Core/Private/Tests/Tasks/TasksTest.cpp
--- a/Engine/Source/Runtime/Core/Private/Tests/Tasks/TasksTest.cpp
+++ b/Engine/Source/Runtime/Core/Private/Tests/Tasks/TasksTest.cpp
@@ -1312,6 +1312,24 @@ namespace UE { namespace TasksTests
 
 		return true;
 	}
+
+	// test Pipe support for `-nothreading` config (`FPlatformProcess::SupportsMultithreading()` is false)
+	IMPLEMENT_SIMPLE_AUTOMATION_TEST(FTasksPipeNoThreadingTest, "System.Core.Tasks.PipeNoThreading", EAutomationTestFlags::ApplicationContextMask | EAutomationTestFlags::EngineFilter);
+
+	bool FTasksPipeNoThreadingTest::RunTest(const FString& Parameters)
+	{
+		if (FPlatformProcess::SupportsMultithreading())
+		{
+			return true; // run with `-nothreading` for this test to make any sense
+		}
+
+		FPipe Pipe{ UE_SOURCE_LOCATION };
+		Pipe.Launch(UE_SOURCE_LOCATION, [] {});
+		FTask FinalTask = Pipe.Launch(UE_SOURCE_LOCATION, [] {});
+		check(FinalTask.Wait(FTimespan::FromMilliseconds(1.0)));
+
+		return true;
+	}
 }}
 
 #endif // WITH_DEV_AUTOMATION_TESTS
