Fix flaw in config path expansion caused by missing trailing slash when computing path relative to the Engine root directory.  This caused the trailing segment of the relative root to get trimmed inside of MakePathRelativeTo.

diff --git a/Engine/Source/Runtime/Core/Private/Misc/ConfigContext.cpp b/Engine/Source/Runtime/Core/Private/Misc/ConfigContext.cpp
--- a/Engine/Source/Runtime/Core/Private/Misc/ConfigContext.cpp
+++ b/Engine/Source/Runtime/Core/Private/Misc/ConfigContext.cpp
@@ -79,7 +79,7 @@ void FConfigContext::CachePaths()
 		if (FPaths::IsUnderDirectory(ProjectRootDir, EngineRootDir))
 		{
 			FString RelativeDir = ProjectRootDir;
-			FPaths::MakePathRelativeTo(RelativeDir, *EngineRootDir);
+			FPaths::MakePathRelativeTo(RelativeDir, *(EngineRootDir + TEXT("/")));
 			ProjectNotForLicenseesDir = FPaths::Combine(EngineRootDir, TEXT("Restricted/NotForLicensees"), RelativeDir);
 			ProjectNoRedistDir = FPaths::Combine(EngineRootDir, TEXT("Restricted/NoRedist"), RelativeDir);
 		}
