OodleDataCompression: Changed scratch buffer count and size to match usage by the engine

diff --git a/Engine/Source/Runtime/Core/Private/Compression/OodleDataCompression.cpp b/Engine/Source/Runtime/Core/Private/Compression/OodleDataCompression.cpp
--- a/Engine/Source/Runtime/Core/Private/Compression/OodleDataCompression.cpp
+++ b/Engine/Source/Runtime/Core/Private/Compression/OodleDataCompression.cpp
@@ -232,16 +232,17 @@ struct OodleScratchBuffers
 		// enough decoder scratch for any compressor & buffer size.
 		// note "InCompressor" is what we want to Encode with but we may be asked to decode other compressors!
 		OO_SINTa DecoderMemorySizeNeeded = OodleLZDecoder_MemorySizeNeeded(OodleLZ_Compressor_Invalid, -1);
-		OO_SINTa EncoderMemorySizeNeeded = OodleLZ_GetCompressScratchMemBound(OodleLZ_Compressor_Mermaid,OodleLZ_CompressionLevel_VeryFast,64*1024,NULL);
-		// DecoderMemorySizeNeeded is ~ 460000 , EncoderMemorySizeNeeded is ~ 470000
-		OodleScratchMemorySize = DecoderMemorySizeNeeded > EncoderMemorySizeNeeded ? DecoderMemorySizeNeeded : EncoderMemorySizeNeeded;
 
 		int32 BufferCount;
 		
 		if ( FPlatformProperties::RequiresCookedData() )
 		{
 			// runtime game or similar environment
-			
+
+			OO_SINTa EncoderMemorySizeNeeded = OodleLZ_GetCompressScratchMemBound(OodleLZ_Compressor_Mermaid, OodleLZ_CompressionLevel_VeryFast, 64*1024, nullptr);
+			// DecoderMemorySizeNeeded is ~ 450000 , EncoderMemorySizeNeeded is ~ 470000
+			OodleScratchMemorySize = DecoderMemorySizeNeeded > EncoderMemorySizeNeeded ? DecoderMemorySizeNeeded : EncoderMemorySizeNeeded;
+
 			BufferCount = 2;
 
 			// "PreallocatedBufferCount" is not a great name
@@ -257,8 +258,8 @@ struct OodleScratchBuffers
 				GConfig->GetInt(TEXT("OodleDataCompressionFormat"), TEXT("PreallocatedBufferCount"), BufferCount, GEngineIni);
 				if (BufferCount < 0)
 				{
-					// negative means one per core
-					BufferCount = FPlatformMisc::NumberOfCores();
+					// negative means one per worker thread
+					BufferCount = FPlatformMisc::NumberOfWorkerThreadsToSpawn();
 				}
 			}
 		}
@@ -266,10 +267,14 @@ struct OodleScratchBuffers
 		{
 			// tools, like UnrealPak or DDC commandlets
 
-			// allow one scratch buffer per core
+			OO_SINTa EncoderMemorySizeNeeded = OodleLZ_GetCompressScratchMemBound(OodleLZ_Compressor_Mermaid, OodleLZ_CompressionLevel_VeryFast, 256*1024, nullptr);
+			// DecoderMemorySizeNeeded is ~ 450000 , EncoderMemorySizeNeeded is ~ 1200000
+			OodleScratchMemorySize = DecoderMemorySizeNeeded > EncoderMemorySizeNeeded ? DecoderMemorySizeNeeded : EncoderMemorySizeNeeded;
+
+			// allow one scratch buffer per worker thread
 			// they will only be allocated on demand if we actually reach that level of parallelism
 			//	 so commandlets that don't use parallel compression don't waste memory
-			BufferCount = FPlatformMisc::NumberOfCores();
+			BufferCount = FPlatformMisc::NumberOfWorkerThreadsToSpawn();
 		}
 
 		OodleScratchBufferCount = BufferCount;
