Consolidated duplicated code that checks for -trace=memory,memalloc.

I had originally planned to also use this in TraceAux for but it later became apparent that I didn't need it. A single point where the command line is checked is more robust nonetheless

#rnx

diff --git a/Engine/Source/Runtime/Core/Private/ProfilingDebugging/MemoryTrace.cpp b/Engine/Source/Runtime/Core/Private/ProfilingDebugging/MemoryTrace.cpp
--- a/Engine/Source/Runtime/Core/Private/ProfilingDebugging/MemoryTrace.cpp
+++ b/Engine/Source/Runtime/Core/Private/ProfilingDebugging/MemoryTrace.cpp
@@ -1,6 +1,7 @@
 // Copyright Epic Games, Inc. All Rights Reserved.
 
 #include "ProfilingDebugging/MemoryTrace.h"
+#include "Containers/StringView.h"
 #include "HAL/PlatformTime.h"
 #include "ProfilingDebugging/CallstackTrace.h"
 #include "ProfilingDebugging/TagTrace.h"
@@ -279,7 +280,39 @@ protected:
 static FUndestructed<FTraceMalloc> GTraceMalloc;
 
 ////////////////////////////////////////////////////////////////////////////////
-FMalloc* MemoryTrace_CreateInternal(FMalloc* InMalloc)
+template <typename ArgCharType>
+static bool MemoryTrace_ShouldEnable(int32 ArgC, const ArgCharType* const* ArgV)
+{
+	for (int32 ArgIndex = 1; ArgIndex < ArgC; ArgIndex++)
+	{
+		const ArgCharType* Arg = ArgV[ArgIndex];
+		TStringView<ArgCharType> ArgView(Arg);
+		if (ArgView.SubStr(0, 7).Equals("-trace=", ESearchCase::IgnoreCase))
+		{
+			const ArgCharType* Start = Arg + 7;
+			const ArgCharType* End = Arg + TCString<ArgCharType>::Strlen(Arg);
+
+			for (const ArgCharType* c = Start; c < End + 1; ++c)
+			{
+				if (*c == '\0' || *c == ',')
+				{
+					TStringView<ArgCharType> View(Start, uint32(c - Start));
+					if (View.Equals("memalloc", ESearchCase::IgnoreCase) || View.Equals("memory", ESearchCase::IgnoreCase))
+					{
+						return true;
+					}
+
+					Start = c + 1;
+				}
+			}
+		}
+	}
+
+	return false;
+}
+
+////////////////////////////////////////////////////////////////////////////////
+static FMalloc* MemoryTrace_CreateInternal(FMalloc* InMalloc)
 {
 	// Some OSes (i.e. Windows) will terminate all threads except the main
 	// one as part of static deinit. However we may receive more memory
@@ -304,6 +337,28 @@ FMalloc* MemoryTrace_CreateInternal(FMalloc* InMalloc)
 	return &SMallocWrapper;
 }
 
+////////////////////////////////////////////////////////////////////////////////
+FMalloc* MemoryTrace_CreateInternal(FMalloc* InMalloc, int32 ArgC, const WIDECHAR* const* ArgV)
+{
+	if (!MemoryTrace_ShouldEnable(ArgC, ArgV))
+	{
+		return nullptr;
+	}
+
+	return MemoryTrace_CreateInternal(InMalloc);
+}
+
+////////////////////////////////////////////////////////////////////////////////
+FMalloc* MemoryTrace_CreateInternal(FMalloc* InMalloc, int32 ArgC, const ANSICHAR* const* ArgV)
+{
+	if (!MemoryTrace_ShouldEnable(ArgC, ArgV))
+	{
+		return nullptr;
+	}
+
+	return MemoryTrace_CreateInternal(InMalloc);
+}
+
 ////////////////////////////////////////////////////////////////////////////////
 void MemoryTrace_Initialize()
 {
diff --git a/Engine/Source/Runtime/Core/Private/ProfilingDebugging/Microsoft/WindowsMemoryTrace.cpp b/Engine/Source/Runtime/Core/Private/ProfilingDebugging/Microsoft/WindowsMemoryTrace.cpp
--- a/Engine/Source/Runtime/Core/Private/ProfilingDebugging/Microsoft/WindowsMemoryTrace.cpp
+++ b/Engine/Source/Runtime/Core/Private/ProfilingDebugging/Microsoft/WindowsMemoryTrace.cpp
@@ -17,11 +17,12 @@
 #include "Trace/Trace.inl"
 
 #include "Windows/AllowWindowsPlatformTypes.h"
+#include <shellapi.h>
 #include <winternl.h>
 #include <intrin.h>
 
 ////////////////////////////////////////////////////////////////////////////////
-FMalloc* MemoryTrace_CreateInternal(FMalloc*);
+FMalloc* MemoryTrace_CreateInternal(FMalloc*, int32, const WIDECHAR* const*);
 
 ////////////////////////////////////////////////////////////////////////////////
 struct FAddrPack
@@ -358,42 +359,23 @@ LPVOID WINAPI FVirtualWinApiHooks::VmAlloc2(HANDLE Process, LPVOID BaseAddress,
 }
 #endif // defined(PLATFORM_SUPPORTS_TRACE_WIN32_VIRTUAL_MEMORY_HOOKS)
 
-
 ////////////////////////////////////////////////////////////////////////////////
 FMalloc* MemoryTrace_Create(FMalloc* InMalloc)
 {
 	bool bEnabled = false;
-	const TCHAR* CmdLine = ::GetCommandLineW();
-	if (const TCHAR* TraceArg = FCString::Strstr(CmdLine, TEXT("-trace=")))
-	{
-		const TCHAR* Start = TraceArg + 7;
-		const TCHAR* End = Start;
-		for (; *End != ' ' && *End != '\0'; ++End);
 
-		for (const TCHAR* c = Start; c < End + 1; ++c)
-		{
-			if (*c == ' ' || *c == '\0' || *c == ',')
-			{
-				FStringView View(Start, uint32(c - Start));
-				if (View.Equals(TEXT("memalloc"), ESearchCase::IgnoreCase) || View.Equals(TEXT("memory"), ESearchCase::IgnoreCase))
-				{
-					bEnabled = true;
-					break;
-				}
-
-				Start = c + 1;
-			}
-		}
-	}
+	int ArgC = 0;
+	const WIDECHAR* CmdLine = ::GetCommandLineW();
+	const WIDECHAR* const* ArgV = ::CommandLineToArgvW(CmdLine, &ArgC);
+	FMalloc* OutMalloc = MemoryTrace_CreateInternal(InMalloc, ArgC, ArgV);
+	::LocalFree(HLOCAL(ArgV));
 
-	if (bEnabled)
+	if (OutMalloc != nullptr)
 	{
-		FMalloc* OutMalloc = MemoryTrace_CreateInternal(InMalloc);
-		
-	#if defined(PLATFORM_SUPPORTS_TRACE_WIN32_VIRTUAL_MEMORY_HOOKS)
+#if defined(PLATFORM_SUPPORTS_TRACE_WIN32_VIRTUAL_MEMORY_HOOKS)
 		FVirtualWinApiHooks::Initialize(false);
-	#endif // defined(PLATFORM_SUPPORTS_TRACE_WIN32_VIRTUAL_MEMORY_HOOKS)
-		
+#endif
+
 		return OutMalloc;
 	}
 
