Added function to query if memory tracing is enabled

#rnx

diff --git a/Engine/Source/Runtime/Core/Private/ProfilingDebugging/MemoryTrace.cpp b/Engine/Source/Runtime/Core/Private/ProfilingDebugging/MemoryTrace.cpp
--- a/Engine/Source/Runtime/Core/Private/ProfilingDebugging/MemoryTrace.cpp
+++ b/Engine/Source/Runtime/Core/Private/ProfilingDebugging/MemoryTrace.cpp
@@ -386,6 +386,12 @@ void MemoryTrace_Initialize()
 #endif
 }
 
+////////////////////////////////////////////////////////////////////////////////
+bool MemoryTrace_IsActive()
+{
+	return GTraceAllowed;
+}
+
 ////////////////////////////////////////////////////////////////////////////////
 void MemoryTrace_EnableTracePump()
 {
