Speed up trace's worker thread if memory tracing is enabled.

There is no backstory behind the current sleep time that Trace uses (other than it being prime to try and avoid phase with other systems). It's largely a quick choice to originally getting things stood up that's never been revisited. The aspirations where that the update happens as infrequently as possible to avoid interference with the process being instrumented. Memory tracing is overly chatty now with up to four or five synchronised events per allocation. While a higher frequency worker thread takes away from game thread time, it reduces resource code and makes less bursty IO requests.

#rnx

diff --git a/Engine/Source/Runtime/Core/Private/ProfilingDebugging/TraceAuxiliary.cpp b/Engine/Source/Runtime/Core/Private/ProfilingDebugging/TraceAuxiliary.cpp
--- a/Engine/Source/Runtime/Core/Private/ProfilingDebugging/TraceAuxiliary.cpp
+++ b/Engine/Source/Runtime/Core/Private/ProfilingDebugging/TraceAuxiliary.cpp
@@ -1467,6 +1467,26 @@ void FTraceAuxiliary::Initialize(const TCHAR* CommandLine)
 	{
 		Desc.TailSizeBytes <<= 20;
 	}
+
+	// Memory tracing is very chatty. To reduce load on trace we'll speed up the
+	// worker thread so it can clear events faster.
+	extern bool MemoryTrace_IsActive();
+	if (MemoryTrace_IsActive())
+	{
+		int32 SleepTimeMs = 5;
+		if (GConfig)
+		{
+			GConfig->GetInt(GTraceConfigSection, TEXT("SleepTimeWhenMemoryTracingInMS"), SleepTimeMs, GEngineIni);
+		}
+
+		if (Desc.ThreadSleepTimeInMS)
+		{
+			SleepTimeMs = FMath::Min<uint32>(Desc.ThreadSleepTimeInMS, SleepTimeMs);
+		}
+
+		Desc.ThreadSleepTimeInMS = SleepTimeMs;
+	}
+
 	UE::Trace::Initialize(Desc);
 
 	// Workaround for the fact that even if StartFromCommandlineArguments will enable channels
