Fixed link error on Linux

#rnx

diff --git a/Engine/Source/Runtime/Core/Private/ProfilingDebugging/MemoryTrace.cpp b/Engine/Source/Runtime/Core/Private/ProfilingDebugging/MemoryTrace.cpp
--- a/Engine/Source/Runtime/Core/Private/ProfilingDebugging/MemoryTrace.cpp
+++ b/Engine/Source/Runtime/Core/Private/ProfilingDebugging/MemoryTrace.cpp
@@ -644,6 +644,14 @@ void MemoryTrace_UnmarkAllocAsHeap(uint64 Address, HeapId Heap)
 		<< HeapUnmarkAlloc.Heap(Heap);
 }
 
+#else // UE_MEMORY_TRACE_ENABLED
+
+/////////////////////////////////////////////////////////////////////////////
+bool MemoryTrace_IsActive()
+{
+	return false;
+}
+
 #endif // UE_MEMORY_TRACE_ENABLED
 
 
