Add missing LLM tracking tags. Went from >40000 untagged allocations to about 4400 in QA game. Most of the untagged allocations were CVars. The remaining untracked allocations are largely from dynamic initializers and static/const data.

diff --git a/Engine/Plugins/FX/Niagara/Source/Niagara/Private/NiagaraModule.cpp b/Engine/Plugins/FX/Niagara/Source/Niagara/Private/NiagaraModule.cpp
--- a/Engine/Plugins/FX/Niagara/Source/Niagara/Private/NiagaraModule.cpp
+++ b/Engine/Plugins/FX/Niagara/Source/Niagara/Private/NiagaraModule.cpp
@@ -213,6 +213,7 @@ void INiagaraModule::StartupModule()
 #if VECTORVM_SUPPORTS_EXPERIMENTAL
 	InitVectorVM();
 #endif
+	LLM_SCOPE(ELLMTag::Niagara);
 	FNiagaraTypeDefinition::Init();
 	FNiagaraRenderViewDataManager::Init();
 
diff --git a/Engine/Source/Runtime/Core/Private/HAL/ConsoleManager.cpp b/Engine/Source/Runtime/Core/Private/HAL/ConsoleManager.cpp
--- a/Engine/Source/Runtime/Core/Private/HAL/ConsoleManager.cpp
+++ b/Engine/Source/Runtime/Core/Private/HAL/ConsoleManager.cpp
@@ -1431,21 +1431,25 @@ public:
 
 IConsoleVariable* FConsoleManager::RegisterConsoleVariable(const TCHAR* Name, bool DefaultValue, const TCHAR* Help, uint32 Flags)
 {
+	LLM_SCOPE_BYNAME(TEXT("EngineMisc/ConsoleCommands"));
 	return AddConsoleObject(Name, new FConsoleVariable<bool>(DefaultValue, Help, (EConsoleVariableFlags)Flags))->AsVariable();
 }
 
 IConsoleVariable* FConsoleManager::RegisterConsoleVariable(const TCHAR* Name, int32 DefaultValue, const TCHAR* Help, uint32 Flags)
 {
+	LLM_SCOPE_BYNAME(TEXT("EngineMisc/ConsoleCommands"));
 	return AddConsoleObject(Name, new FConsoleVariable<int32>(DefaultValue, Help, (EConsoleVariableFlags)Flags))->AsVariable();
 }
 
 IConsoleVariable* FConsoleManager::RegisterConsoleVariable(const TCHAR* Name, float DefaultValue, const TCHAR* Help, uint32 Flags)
 {
+	LLM_SCOPE_BYNAME(TEXT("EngineMisc/ConsoleCommands"));
 	return AddConsoleObject(Name, new FConsoleVariable<float>(DefaultValue, Help, (EConsoleVariableFlags)Flags))->AsVariable();
 }
 
 IConsoleVariable* FConsoleManager::RegisterConsoleVariable(const TCHAR* Name, const TCHAR* DefaultValue, const TCHAR* Help, uint32 Flags)
 {
+	LLM_SCOPE_BYNAME(TEXT("EngineMisc/ConsoleCommands"));
 	return RegisterConsoleVariable(Name, FString(DefaultValue), Help, Flags);
 }
 
@@ -1453,66 +1457,79 @@ IConsoleVariable* FConsoleManager::RegisterConsoleVariable(const TCHAR* Name, co
 {
 	// not supported
 	check((Flags & (uint32)ECVF_RenderThreadSafe) == 0);
+	LLM_SCOPE_BYNAME(TEXT("EngineMisc/ConsoleCommands"));
 	return AddConsoleObject(Name, new FConsoleVariable<FString>(DefaultValue, Help, (EConsoleVariableFlags)Flags))->AsVariable();
 }
 
 IConsoleVariable* FConsoleManager::RegisterConsoleVariableRef(const TCHAR* Name, bool& RefValue, const TCHAR* Help, uint32 Flags)
 {
+	LLM_SCOPE_BYNAME(TEXT("EngineMisc/ConsoleCommands"));
 	return AddConsoleObject(Name, new FConsoleVariableRef<bool>(RefValue, Help, (EConsoleVariableFlags)Flags))->AsVariable();
 }
 
 IConsoleVariable* FConsoleManager::RegisterConsoleVariableRef(const TCHAR* Name, int32& RefValue, const TCHAR* Help, uint32 Flags)
 {
+	LLM_SCOPE_BYNAME(TEXT("EngineMisc/ConsoleCommands"));
 	return AddConsoleObject(Name, new FConsoleVariableRef<int32>(RefValue, Help, (EConsoleVariableFlags)Flags))->AsVariable();
 }
 
 IConsoleVariable* FConsoleManager::RegisterConsoleVariableRef(const TCHAR* Name, float& RefValue, const TCHAR* Help, uint32 Flags)
 {
+	LLM_SCOPE_BYNAME(TEXT("EngineMisc/ConsoleCommands"));
 	return AddConsoleObject(Name, new FConsoleVariableRef<float>(RefValue, Help, (EConsoleVariableFlags)Flags))->AsVariable();
 }
 
 IConsoleVariable* FConsoleManager::RegisterConsoleVariableRef(const TCHAR* Name, FString& RefValue, const TCHAR* Help, uint32 Flags)
 {
+	LLM_SCOPE_BYNAME(TEXT("EngineMisc/ConsoleCommands"));
 	return AddConsoleObject(Name, new FConsoleVariableStringRef(RefValue, Help, (EConsoleVariableFlags)Flags))->AsVariable();
 }
 
 IConsoleCommand* FConsoleManager::RegisterConsoleCommand(const TCHAR* Name, const TCHAR* Help, const FConsoleCommandDelegate& Command, uint32 Flags)
 {
+	LLM_SCOPE_BYNAME(TEXT("EngineMisc/ConsoleCommands"));
 	return AddConsoleObject(Name, new FConsoleCommand(Command, Help, (EConsoleVariableFlags)Flags))->AsCommand();
 }
 
 IConsoleCommand* FConsoleManager::RegisterConsoleCommand(const TCHAR* Name, const TCHAR* Help, uint32 Flags)
 {
+	LLM_SCOPE_BYNAME(TEXT("EngineMisc/ConsoleCommands"));
 	return AddConsoleObject(Name, new FConsoleCommandExec(Help, (EConsoleVariableFlags)Flags))->AsCommand();
 }
 
 IConsoleCommand* FConsoleManager::RegisterConsoleCommand(const TCHAR* Name, const TCHAR* Help, const FConsoleCommandWithArgsDelegate& Command, uint32 Flags)
 {
+	LLM_SCOPE_BYNAME(TEXT("EngineMisc/ConsoleCommands"));
 	return AddConsoleObject(Name, new FConsoleCommandWithArgs(Command, Help, (EConsoleVariableFlags)Flags))->AsCommand();
 }
 
 IConsoleCommand* FConsoleManager::RegisterConsoleCommand(const TCHAR* Name, const TCHAR* Help, const FConsoleCommandWithWorldDelegate& Command, uint32 Flags)
 {
+	LLM_SCOPE_BYNAME(TEXT("EngineMisc/ConsoleCommands"));
 	return AddConsoleObject(Name, new FConsoleCommandWithWorld(Command, Help, (EConsoleVariableFlags)Flags))->AsCommand();
 }
 
 IConsoleCommand* FConsoleManager::RegisterConsoleCommand(const TCHAR* Name, const TCHAR* Help, const FConsoleCommandWithWorldAndArgsDelegate& Command, uint32 Flags)
 {
+	LLM_SCOPE_BYNAME(TEXT("EngineMisc/ConsoleCommands"));
 	return AddConsoleObject(Name, new FConsoleCommandWithWorldAndArgs(Command, Help, (EConsoleVariableFlags)Flags))->AsCommand();
 }
 
 IConsoleCommand* FConsoleManager::RegisterConsoleCommand(const TCHAR* Name, const TCHAR* Help, const FConsoleCommandWithArgsAndOutputDeviceDelegate& Command, uint32 Flags)
 {
+	LLM_SCOPE_BYNAME(TEXT("EngineMisc/ConsoleCommands"));
 	return AddConsoleObject(Name, new FConsoleCommanWithArgsAndOutputDevice(Command, Help, (EConsoleVariableFlags)Flags))->AsCommand();
 }
 
 IConsoleCommand* FConsoleManager::RegisterConsoleCommand(const TCHAR* Name, const TCHAR* Help, const FConsoleCommandWithWorldArgsAndOutputDeviceDelegate& Command, uint32 Flags)
 {
+	LLM_SCOPE_BYNAME(TEXT("EngineMisc/ConsoleCommands"));
 	return AddConsoleObject(Name, new FConsoleCommandWithWorldArgsAndOutputDevice(Command, Help, (EConsoleVariableFlags)Flags))->AsCommand();
 }
 
 IConsoleCommand* FConsoleManager::RegisterConsoleCommand(const TCHAR* Name, const TCHAR* Help, const FConsoleCommandWithOutputDeviceDelegate& Command, uint32 Flags)
 {
+	LLM_SCOPE_BYNAME(TEXT("EngineMisc/ConsoleCommands"));
 	return AddConsoleObject(Name, new FConsoleCommandWithOutputDevice(Command, Help, (EConsoleVariableFlags)Flags))->AsCommand();
 }
 
@@ -1943,6 +1960,7 @@ IConsoleObject* FConsoleManager::AddConsoleObject(const TCHAR* Name, IConsoleObj
 	check(Name);
 	check(*Name != 0);
 	check(Obj);
+	LLM_SCOPE_BYNAME(TEXT("EngineMisc/ConsoleCommands"));
 
 	FScopeLock ScopeLock( &ConsoleObjectsSynchronizationObject ); // we will lock on the entire add process
 	IConsoleObject* ExistingObj = ConsoleObjects.FindRef(Name);
@@ -2174,6 +2192,7 @@ IConsoleManager* IConsoleManager::Singleton;
 void IConsoleManager::SetupSingleton()
 {
 	check(!Singleton);
+	LLM_SCOPE_BYNAME(TEXT("EngineMisc/ConsoleCommands"));
 	if (!Singleton)
 	{
 		Singleton = new FConsoleManager; // we will leak this
@@ -2449,7 +2468,7 @@ void FConsoleManager::Test()
 void CreateConsoleVariables()
 {
 #if !NO_CVARS
-
+	LLM_SCOPE_BYNAME(TEXT("EngineMisc/ConsoleCommands"));
 	// this registeres to a reference, so we cannot use TAutoConsoleVariable
 	IConsoleManager::Get().RegisterConsoleVariableRef(TEXT("r.DumpingMovie"),
 		GIsDumpingMovie,
diff --git a/Engine/Source/Runtime/Core/Private/HAL/Thread.cpp b/Engine/Source/Runtime/Core/Private/HAL/Thread.cpp
--- a/Engine/Source/Runtime/Core/Private/HAL/Thread.cpp
+++ b/Engine/Source/Runtime/Core/Private/HAL/Thread.cpp
@@ -62,7 +62,7 @@ private:
 	{
 		// the thread can be started before `RunnableThread` member is initialized
 		//check(RunnableThread.IsValid());
-
+		LLM_SCOPE(ELLMTag::TaskGraphTasksMisc);
 		ThreadFunction();
 
 		return 0;
diff --git a/Engine/Source/Runtime/Core/Private/Modules/ModuleManager.cpp b/Engine/Source/Runtime/Core/Private/Modules/ModuleManager.cpp
--- a/Engine/Source/Runtime/Core/Private/Modules/ModuleManager.cpp
+++ b/Engine/Source/Runtime/Core/Private/Modules/ModuleManager.cpp
@@ -369,6 +369,7 @@ void FModuleManager::RefreshModuleFilenameFromManifest(const FName InModuleName)
 
 IModuleInterface* FModuleManager::LoadModule(const FName InModuleName, ELoadModuleFlags InLoadModuleFlags)
 {
+	LLM_SCOPE_BYNAME(TEXT("Modules"));
 	// We allow an already loaded module to be returned in other threads to simplify
 	// parallel processing scenarios but they must have been loaded from the main thread beforehand.
 	if(!IsInGameThread())
diff --git a/Engine/Source/Runtime/D3D12RHI/Private/D3D12DirectCommandListManager.cpp b/Engine/Source/Runtime/D3D12RHI/Private/D3D12DirectCommandListManager.cpp
--- a/Engine/Source/Runtime/D3D12RHI/Private/D3D12DirectCommandListManager.cpp
+++ b/Engine/Source/Runtime/D3D12RHI/Private/D3D12DirectCommandListManager.cpp
@@ -75,12 +75,14 @@ void FD3D12CommandContext::RHIWriteGPUFence(FRHIGPUFence* FenceRHI)
 
 FGPUFenceRHIRef FD3D12DynamicRHI::RHICreateGPUFence(const FName& Name)
 {
+	LLM_SCOPE_BYNAME(TEXT("RHIMisc/CreateGPUFence"));
 	return new FD3D12GPUFence(Name);
 }
 
 FStagingBufferRHIRef FD3D12DynamicRHI::RHICreateStagingBuffer()
 {
 	// Don't know the device yet - will be decided at copy time (lazy creation)
+	LLM_SCOPE_BYNAME(TEXT("RHIMisc/CreateStagingBuffer"));
 	return new FD3D12StagingBuffer(nullptr);
 }
 
diff --git a/Engine/Source/Runtime/D3D12RHI/Private/D3D12Submission.cpp b/Engine/Source/Runtime/D3D12RHI/Private/D3D12Submission.cpp
--- a/Engine/Source/Runtime/D3D12RHI/Private/D3D12Submission.cpp
+++ b/Engine/Source/Runtime/D3D12RHI/Private/D3D12Submission.cpp
@@ -225,7 +225,7 @@ FD3D12DynamicRHI::FProcessResult FD3D12DynamicRHI::ProcessSubmissionQueue()
 {
 	SCOPED_NAMED_EVENT_TEXT("SubmissionQueue_Process", FColor::Turquoise);
 	SCOPE_CYCLE_COUNTER(STAT_D3D12Submit);
-
+	LLM_SCOPE_BYNAME(TEXT("RHIMisc/ProcessSubmissionQueue"));
 	FProcessResult Result;
 
 	//
@@ -605,7 +605,7 @@ FD3D12CommandList* FD3D12DynamicRHI::GenerateBarrierCommandListAndUpdateState(FD
 uint64 FD3D12Queue::ExecutePayload()
 {
 	TRACE_CPUPROFILER_EVENT_SCOPE(ExecuteCommandList);
-
+	LLM_SCOPE_BYNAME(TEXT("RHIMisc/ExecuteCommandLists"));
 	check(PayloadToSubmit && this == &PayloadToSubmit->Queue);
 	PayloadToSubmit->PreExecute();
 
@@ -773,6 +773,7 @@ void FD3D12DynamicRHI::ProcessInterruptQueueUntil(FGraphEvent* GraphEvent)
 FD3D12DynamicRHI::FProcessResult FD3D12DynamicRHI::ProcessInterruptQueue()
 {
 	SCOPED_NAMED_EVENT_TEXT("InterruptQueue_Process", FColor::Yellow);
+	LLM_SCOPE_BYNAME(TEXT("RHIMisc/ProcessInterruptQueue"));
 
 	FProcessResult Result;
 	ForEachQueue([&](FD3D12Queue& CurrentQueue)
diff --git a/Engine/Source/Runtime/D3D12RHI/Private/D3D12Submission.h b/Engine/Source/Runtime/D3D12RHI/Private/D3D12Submission.h
--- a/Engine/Source/Runtime/D3D12RHI/Private/D3D12Submission.h
+++ b/Engine/Source/Runtime/D3D12RHI/Private/D3D12Submission.h
@@ -69,6 +69,7 @@ class FD3D12SyncPoint final : public FThreadSafeRefCountedObject
 public:
 	static FD3D12SyncPointRef Create(ED3D12SyncPointType Type)
 	{
+		LLM_SCOPE_BYNAME(TEXT("RHIMisc/CreateSyncPoint"));
 		return new FD3D12SyncPoint(Type);
 	}
 
diff --git a/Engine/Source/Runtime/Engine/Private/Materials/MaterialShared.cpp b/Engine/Source/Runtime/Engine/Private/Materials/MaterialShared.cpp
--- a/Engine/Source/Runtime/Engine/Private/Materials/MaterialShared.cpp
+++ b/Engine/Source/Runtime/Engine/Private/Materials/MaterialShared.cpp
@@ -5158,6 +5158,7 @@ void FMaterialAttributeDefinitionMap::InitializeAttributeMap()
 	check(!bIsInitialized);
 	bIsInitialized = true;
 	const bool bHideAttribute = true;
+	LLM_SCOPE(ELLMTag::Materials);
 
 	// All types plus default/missing attribute
 	AttributeMap.Empty(MP_MAX + 1);
diff --git a/Engine/Source/Runtime/Launch/Private/LaunchEngineLoop.cpp b/Engine/Source/Runtime/Launch/Private/LaunchEngineLoop.cpp
--- a/Engine/Source/Runtime/Launch/Private/LaunchEngineLoop.cpp
+++ b/Engine/Source/Runtime/Launch/Private/LaunchEngineLoop.cpp
@@ -4171,7 +4171,7 @@ bool FEngineLoop::LoadStartupCoreModules()
 bool FEngineLoop::LoadStartupModules()
 {
 	FScopedSlowTask SlowTask(3);
-
+	LLM_SCOPE_BYNAME(TEXT("Modules"));
 	SlowTask.EnterProgressFrame(1);
 	// Load any modules that want to be loaded before default modules are loaded up.
 	if (!IProjectManager::Get().LoadModulesForProject(ELoadingPhase::PreDefault) || !IPluginManager::Get().LoadModulesForEnabledPlugins(ELoadingPhase::PreDefault))
diff --git a/Engine/Source/Runtime/MoviePlayer/Private/DefaultGameMoviePlayer.cpp b/Engine/Source/Runtime/MoviePlayer/Private/DefaultGameMoviePlayer.cpp
--- a/Engine/Source/Runtime/MoviePlayer/Private/DefaultGameMoviePlayer.cpp
+++ b/Engine/Source/Runtime/MoviePlayer/Private/DefaultGameMoviePlayer.cpp
@@ -165,6 +165,7 @@ void FDefaultGameMoviePlayer::Initialize(FSlateRenderer& InSlateRenderer, TShare
 		return;
 	}
 
+	LLM_SCOPE_BYNAME(TEXT("SlateMoviePlayer"));
 	UE_LOG(LogMoviePlayer, Log, TEXT("Initializing movie player"));
 
 	FDefaultGameMoviePlayer* InMoviePlayer = this;
diff --git a/Engine/Source/Runtime/Projects/Private/PluginManager.cpp b/Engine/Source/Runtime/Projects/Private/PluginManager.cpp
--- a/Engine/Source/Runtime/Projects/Private/PluginManager.cpp
+++ b/Engine/Source/Runtime/Projects/Private/PluginManager.cpp
@@ -2140,6 +2140,7 @@ bool FPluginManager::LoadModulesForEnabledPlugins( const ELoadingPhase::Type Loa
 	else
 	{
 		FScopedSlowTask SlowTask((float)AllPlugins.Num());
+		LLM_SCOPE_BYNAME(TEXT("Modules"));
 
 		// Load plugins!
 		for (const FDiscoveredPluginMap::ElementType& PluginPair : AllPlugins)
diff --git a/Engine/Source/Runtime/RHI/Private/PipelineStateCache.cpp b/Engine/Source/Runtime/RHI/Private/PipelineStateCache.cpp
--- a/Engine/Source/Runtime/RHI/Private/PipelineStateCache.cpp
+++ b/Engine/Source/Runtime/RHI/Private/PipelineStateCache.cpp
@@ -1572,7 +1572,8 @@ uint64 PipelineStateCache::RetrieveGraphicsPipelineStateSortKey(const FGraphicsP
 }
 
 FComputePipelineState* PipelineStateCache::GetAndOrCreateComputePipelineState(FRHIComputeCommandList& RHICmdList, FRHIComputeShader* ComputeShader, bool bFromFileCache)
-{	
+{
+	LLM_SCOPE(ELLMTag::PSO);
 	bool DoAsyncCompile = IsAsyncCompilationAllowed(RHICmdList, bFromFileCache);
 
 	FComputePipelineState* OutCachedState = nullptr;
diff --git a/Engine/Source/Runtime/RHI/Public/DynamicRHI.h b/Engine/Source/Runtime/RHI/Public/DynamicRHI.h
--- a/Engine/Source/Runtime/RHI/Public/DynamicRHI.h
+++ b/Engine/Source/Runtime/RHI/Public/DynamicRHI.h
@@ -1345,6 +1345,7 @@ FORCEINLINE FUniformBufferLayoutRHIRef RHICreateUniformBufferLayout(const FRHIUn
 
 FORCEINLINE FUniformBufferRHIRef RHICreateUniformBuffer(const void* Contents, const FRHIUniformBufferLayout* Layout, EUniformBufferUsage Usage, EUniformBufferValidation Validation = EUniformBufferValidation::ValidateResources)
 {
+	LLM_SCOPE_BYNAME(TEXT("RHIMisc/CreateUniformBuffer"));
 	return GDynamicRHI->RHICreateUniformBuffer(Contents, Layout, Usage, Validation);
 }
 
diff --git a/Engine/Source/Runtime/RenderCore/Private/RenderResource.cpp b/Engine/Source/Runtime/RenderCore/Private/RenderResource.cpp
--- a/Engine/Source/Runtime/RenderCore/Private/RenderResource.cpp
+++ b/Engine/Source/Runtime/RenderCore/Private/RenderResource.cpp
@@ -327,6 +327,7 @@ FRenderResource::~FRenderResource()
 
 void BeginInitResource(FRenderResource* Resource)
 {
+	LLM_SCOPE(ELLMTag::SceneRender);
 	ENQUEUE_RENDER_COMMAND(InitCommand)(
 		[Resource](FRHICommandListImmediate& RHICmdList)
 		{
diff --git a/Engine/Source/Runtime/RenderCore/Private/ShaderParameterMetadata.cpp b/Engine/Source/Runtime/RenderCore/Private/ShaderParameterMetadata.cpp
--- a/Engine/Source/Runtime/RenderCore/Private/ShaderParameterMetadata.cpp
+++ b/Engine/Source/Runtime/RenderCore/Private/ShaderParameterMetadata.cpp
@@ -769,7 +769,7 @@ void FShaderParametersMetadata::InitializeLayout(FRHIUniformBufferLayoutInitiali
 		GetLayoutHashStructMap().Emplace(LayoutInitializer.GetHash(), this);
 	}
 
-	Layout = new FRHIUniformBufferLayout(LayoutInitializer);
+	Layout = RHICreateUniformBufferLayout(LayoutInitializer);
 }
 
 void FShaderParametersMetadata::GetNestedStructs(TArray<const FShaderParametersMetadata*>& OutNestedStructs) const
