Fix compilation error on Mac where ELLMTag was undefined in Thread.cpp.

diff --git a/Engine/Source/Runtime/Core/Private/HAL/Thread.cpp b/Engine/Source/Runtime/Core/Private/HAL/Thread.cpp
--- a/Engine/Source/Runtime/Core/Private/HAL/Thread.cpp
+++ b/Engine/Source/Runtime/Core/Private/HAL/Thread.cpp
@@ -5,6 +5,7 @@
 #include "HAL/RunnableThread.h"
 #include "HAL/PlatformTLS.h"
 #include "HAL/ThreadSafeBool.h"
+#include "HAL/LowLevelMemTracker.h"
 #include "Misc/AssertionMacros.h"
 #include "Templates/UnrealTemplate.h"
 #include "Templates/UniquePtr.h"
