fixup robomerge stomp

#rnx

diff --git a/Engine/Source/Runtime/Core/Private/Async/TaskGraph.cpp b/Engine/Source/Runtime/Core/Private/Async/TaskGraph.cpp
--- a/Engine/Source/Runtime/Core/Private/Async/TaskGraph.cpp
+++ b/Engine/Source/Runtime/Core/Private/Async/TaskGraph.cpp
@@ -1925,7 +1925,7 @@ public:
 			if (GConfig == nullptr)
 			{
 				// postpone starting reserve workers until GConfig is initialized, to know if reserve workers are disabled
-				FCoreDelegates::ConfigReadyForUse.AddRaw(this, &FTaskGraphCompatibilityImplementation::StartReserveWorkers);
+				FCoreDelegates::TSConfigReadyForUse().AddRaw(this, &FTaskGraphCompatibilityImplementation::StartReserveWorkers);
 			}
 			else
 			{
@@ -1962,7 +1962,7 @@ public:
 
 	~FTaskGraphCompatibilityImplementation() override
 	{
-		FCoreDelegates::ConfigReadyForUse.RemoveAll(this);
+		FCoreDelegates::TSConfigReadyForUse().RemoveAll(this);
 
 		for (auto& Callback : ShutdownCallbacks)
 		{
