OodleDataCompression: Added memory tracking

#rnx

diff --git a/Engine/Source/Runtime/Core/Private/Compression/OodleDataCompression.cpp b/Engine/Source/Runtime/Core/Private/Compression/OodleDataCompression.cpp
--- a/Engine/Source/Runtime/Core/Private/Compression/OodleDataCompression.cpp
+++ b/Engine/Source/Runtime/Core/Private/Compression/OodleDataCompression.cpp
@@ -4,6 +4,7 @@
 
 #include "CoreGlobals.h"
 #include "HAL/CriticalSection.h"
+#include "HAL/LowLevelMemTracker.h"
 #include "HAL/PlatformMisc.h"
 #include "HAL/PlatformProperties.h"
 #include "HAL/UnrealMemory.h"
@@ -18,6 +19,7 @@
 struct ICompressionFormat;
 
 DEFINE_LOG_CATEGORY(OodleDataCompression);
+LLM_DEFINE_TAG(OodleData);
 
 extern ICompressionFormat * CreateOodleDataCompressionFormat();
 
@@ -212,17 +214,12 @@ int64 CORE_API GetMaximumCompressedSize(int64 InUncompressedSize)
 struct OodleScratchBuffers
 {
 	OO_SINTa OodleScratchMemorySize = 0;
-	int32 OodleScratchBufferCount =0;
+	int32 OodleScratchBufferCount = 0;
 
-	struct OodleScratchBuffer
+	struct alignas(PLATFORM_CACHE_LINE_SIZE) OodleScratchBuffer
 	{
 		FCriticalSection OodleScratchMemoryMutex;
-		void * OodleScratchMemory = nullptr;
-		char pad[64];
-
-		OodleScratchBuffer()
-		{
-		}
+		void* OodleScratchMemory = nullptr;
 	};
 
 	OodleScratchBuffer* OodleScratches = nullptr;
@@ -258,8 +255,8 @@ struct OodleScratchBuffers
 				GConfig->GetInt(TEXT("OodleDataCompressionFormat"), TEXT("PreallocatedBufferCount"), BufferCount, GEngineIni);
 				if (BufferCount < 0)
 				{
-					// negative means one per worker thread
-					BufferCount = FPlatformMisc::NumberOfWorkerThreadsToSpawn();
+					// negative means one per logical core
+					BufferCount = FPlatformMisc::NumberOfCoresIncludingHyperthreads();
 				}
 			}
 		}
@@ -271,15 +268,16 @@ struct OodleScratchBuffers
 			// DecoderMemorySizeNeeded is ~ 450000 , EncoderMemorySizeNeeded is ~ 1200000
 			OodleScratchMemorySize = DecoderMemorySizeNeeded > EncoderMemorySizeNeeded ? DecoderMemorySizeNeeded : EncoderMemorySizeNeeded;
 
-			// allow one scratch buffer per worker thread
+			// allow one scratch buffer per logical core
 			// they will only be allocated on demand if we actually reach that level of parallelism
 			//	 so commandlets that don't use parallel compression don't waste memory
-			BufferCount = FPlatformMisc::NumberOfWorkerThreadsToSpawn();
+			BufferCount = FPlatformMisc::NumberOfCoresIncludingHyperthreads();
 		}
 
 		OodleScratchBufferCount = BufferCount;
 		if (OodleScratchBufferCount)
 		{
+			LLM_SCOPE_BYTAG(OodleData);
 			OodleScratches = new OodleScratchBuffer[OodleScratchBufferCount];
 		}
 		
@@ -338,6 +336,7 @@ struct OodleScratchBuffers
 			{
 				if (OodleScratches[i].OodleScratchMemory == nullptr)
 				{
+					LLM_SCOPE_BYTAG(OodleData);
 					// allocate on first use
 					OodleScratches[i].OodleScratchMemory = FMemory::Malloc(OodleScratchMemorySize);
 					if (OodleScratches[i].OodleScratchMemory == nullptr)
@@ -378,6 +377,8 @@ struct OodleScratchBuffers
 			return OODLELZ_FAILED;
 		}
 
+		LLM_SCOPE_BYTAG(OodleData);
+
 		// allocate memory for the decoder so that Oodle doesn't allocate anything internally
 		void * DecoderMemory = FMemory::Malloc(DecoderMemorySize);
 		if (DecoderMemory == NULL) 
@@ -425,6 +426,7 @@ struct OodleScratchBuffers
 			{
 				if (OodleScratches[i].OodleScratchMemory == nullptr)
 				{
+					LLM_SCOPE_BYTAG(OodleData);
 					// allocate on first use
 					OodleScratches[i].OodleScratchMemory = FMemory::Malloc(OodleScratchMemorySize);
 					if (OodleScratches[i].OodleScratchMemory == nullptr)
@@ -522,6 +524,7 @@ bool CORE_API Decompress(
 					
 static void* OODLE_CALLBACK OodleAlloc(OO_SINTa Size, OO_S32 Alignment)
 {
+	LLM_SCOPE_BYTAG(OodleData);
 	return FMemory::Malloc(SIZE_T(Size), uint32(Alignment));
 }
 
