ChunkInstall: Fixed link errors in modular builds on some platforms.

The problem started occurring when some platform-specific implementations of chunk install were moved into separate modules (LogChunkInstaller log category wasn't dllexported).

#rnx

diff --git a/Engine/Source/Runtime/Core/Public/GenericPlatform/GenericPlatformChunkInstall.h b/Engine/Source/Runtime/Core/Public/GenericPlatform/GenericPlatformChunkInstall.h
--- a/Engine/Source/Runtime/Core/Public/GenericPlatform/GenericPlatformChunkInstall.h
+++ b/Engine/Source/Runtime/Core/Public/GenericPlatform/GenericPlatformChunkInstall.h
@@ -19,7 +19,7 @@
 
 class IPlatformChunkInstall;
 
-DECLARE_LOG_CATEGORY_EXTERN(LogChunkInstaller, Log, All);
+CORE_API DECLARE_LOG_CATEGORY_EXTERN(LogChunkInstaller, Log, All);
 
 namespace EChunkLocation
 {
