DDC: Fixing compile error related to LLM tagging.

diff --git a/Engine/Source/Runtime/Core/Private/ProfilingDebugging/TagTrace.cpp b/Engine/Source/Runtime/Core/Private/ProfilingDebugging/TagTrace.cpp
--- a/Engine/Source/Runtime/Core/Private/ProfilingDebugging/TagTrace.cpp
+++ b/Engine/Source/Runtime/Core/Private/ProfilingDebugging/TagTrace.cpp
@@ -66,10 +66,12 @@ FMemScope::FMemScope(const FName& InName, bool bShouldActivate /*= true*/)
 ////////////////////////////////////////////////////////////////////////////////
 FMemScope::FMemScope(const UE::LLMPrivate::FTagData* TagData, bool bShouldActivate /*= true*/)
 {
+#if ENABLE_LOW_LEVEL_MEM_TRACKER
 	if (UE_TRACE_CHANNELEXPR_IS_ENABLED(MemAllocChannel) & bShouldActivate && TagData)
 	{
 		ActivateScope(MemoryTrace_AnnounceFNameTag(TagData->GetName()));
 	}
+#endif // ENABLE_LOW_LEVEL_MEM_TRACKER
 }
 
 ////////////////////////////////////////////////////////////////////////////////
