Optimized PointerHash

diff --git a/Engine/Source/Runtime/Core/Public/Templates/TypeHash.h b/Engine/Source/Runtime/Core/Public/Templates/TypeHash.h
--- a/Engine/Source/Runtime/Core/Public/Templates/TypeHash.h
+++ b/Engine/Source/Runtime/Core/Public/Templates/TypeHash.h
@@ -8,6 +8,23 @@
 
 #include <stdint.h>
 
+
+namespace UE
+{
+	namespace Private
+	{
+		FORCEINLINE uint32 MurmurFinalize32(uint32 Hash)
+		{
+			Hash ^= Hash >> 16;
+			Hash *= 0x85ebca6b;
+			Hash ^= Hash >> 13;
+			Hash *= 0xc2b2ae35;
+			Hash ^= Hash >> 16;
+			return Hash;
+		}
+	}
+}
+
 /**
  * Combines two hash values to get a third.
  * Note - this function is not commutative.
@@ -48,18 +65,18 @@ inline uint32 HashCombineFast(uint32 A, uint32 B)
 	return HashCombine(A, B);
 }
 
-inline uint32 PointerHash(const void* Key,uint32 C = 0)
+inline uint32 PointerHash(const void* Key)
 {
-#if PLATFORM_64BITS
 	// Ignoring the lower 4 bits since they are likely zero anyway.
 	// Higher bits are more significant in 64 bit builds.
-	auto PtrInt = reinterpret_cast<UPTRINT>(Key) >> 4;
-#else
-	auto PtrInt = reinterpret_cast<UPTRINT>(Key);
-#endif
+	const UPTRINT PtrInt = reinterpret_cast<UPTRINT>(Key) >> 4;
+	return UE::Private::MurmurFinalize32((uint32)PtrInt);
+}
 
+inline uint32 PointerHash(const void* Key, uint32 C)
+{
 	// we can use HashCombineFast here because pointers are non-persistent
-	return HashCombineFast((uint32)PtrInt, C);
+	return HashCombineFast(PointerHash(Key), C);
 }
 
 
