Engine module public header truncation warning fixes

diff --git a/Engine/Source/Runtime/Core/Private/Math/UnrealMath.cpp b/Engine/Source/Runtime/Core/Private/Math/UnrealMath.cpp
--- a/Engine/Source/Runtime/Core/Private/Math/UnrealMath.cpp
+++ b/Engine/Source/Runtime/Core/Private/Math/UnrealMath.cpp
@@ -3234,26 +3234,6 @@ float FMath::FixedTurn(float InCurrent, float InDesired, float InDeltaRate)
 	return FRotator3f::ClampAxis(result);
 }
 
-float FMath::ClampAngle(float AngleDegrees, float MinAngleDegrees, float MaxAngleDegrees)
-{
-	float const MaxDelta = FRotator3f::ClampAxis(MaxAngleDegrees - MinAngleDegrees) * 0.5f;			// 0..180
-	float const RangeCenter = FRotator3f::ClampAxis(MinAngleDegrees + MaxDelta);						// 0..360
-	float const DeltaFromCenter = FRotator3f::NormalizeAxis(AngleDegrees - RangeCenter);				// -180..180
-
-	// maybe clamp to nearest edge
-	if (DeltaFromCenter > MaxDelta)
-	{
-		return FRotator3f::NormalizeAxis(RangeCenter + MaxDelta);
-	}
-	else if (DeltaFromCenter < -MaxDelta)
-	{
-		return FRotator3f::NormalizeAxis(RangeCenter - MaxDelta);
-	}
-
-	// already in range, just return it
-	return FRotator3f::NormalizeAxis(AngleDegrees);
-}
-
 void FMath::ApplyScaleToFloat(float& Dst, const FVector& DeltaScale, float Magnitude)
 {
 	const float Multiplier = ( DeltaScale.X > 0.0f || DeltaScale.Y > 0.0f || DeltaScale.Z > 0.0f ) ? Magnitude : -Magnitude;
diff --git a/Engine/Source/Runtime/Core/Public/Math/InterpCurve.h b/Engine/Source/Runtime/Core/Public/Math/InterpCurve.h
--- a/Engine/Source/Runtime/Core/Public/Math/InterpCurve.h
+++ b/Engine/Source/Runtime/Core/Public/Math/InterpCurve.h
@@ -482,14 +482,14 @@ T FInterpCurve<T>::EvalSecondDerivative(const float InVal, const T& Default) con
 
 
 template< class T >
-float FInterpCurve<T>::InaccurateFindNearest(const T &PointInSpace, float& OutDistanceSq) const
+float FInterpCurve<T>::InaccurateFindNearest(const T &PointInSpace, float& OutDistanceSq) const							// LWC_TODO: Precision loss
 {
 	float OutSegment;
 	return InaccurateFindNearest(PointInSpace, OutDistanceSq, OutSegment);
 }
 
 template< class T >
-float FInterpCurve<T>::InaccurateFindNearest(const T &PointInSpace, float& OutDistanceSq, float& OutSegment) const
+float FInterpCurve<T>::InaccurateFindNearest(const T &PointInSpace, float& OutDistanceSq, float& OutSegment) const		// LWC_TODO: Precision loss
 {
 	const int32 NumPoints = Points.Num();
 	const int32 NumSegments = bIsLooped ? NumPoints : NumPoints - 1;
@@ -517,7 +517,7 @@ float FInterpCurve<T>::InaccurateFindNearest(const T &PointInSpace, float& OutDi
 
 	if (NumPoints == 1)
 	{
-		OutDistanceSq = (PointInSpace - Points[0].OutVal).SizeSquared();
+		OutDistanceSq = static_cast<float>((PointInSpace - Points[0].OutVal).SizeSquared());
 		OutSegment = 0;
 		return Points[0].InVal;
 	}
@@ -527,7 +527,7 @@ float FInterpCurve<T>::InaccurateFindNearest(const T &PointInSpace, float& OutDi
 
 
 template< class T >
-float FInterpCurve<T>::InaccurateFindNearestOnSegment(const T& PointInSpace, int32 PtIdx, float& OutSquaredDistance) const
+float FInterpCurve<T>::InaccurateFindNearestOnSegment(const T& PointInSpace, int32 PtIdx, float& OutSquaredDistance) const		// LWC_TODO: Precision loss
 {
 	const int32 NumPoints = Points.Num();
 	const int32 LastPoint = NumPoints - 1;
@@ -538,8 +538,8 @@ float FInterpCurve<T>::InaccurateFindNearestOnSegment(const T& PointInSpace, int
 
 	if (CIM_Constant == Points[PtIdx].InterpMode)
 	{
-		const float Distance1 = (Points[PtIdx].OutVal - PointInSpace).SizeSquared();
-		const float Distance2 = (Points[NextPtIdx].OutVal - PointInSpace).SizeSquared();
+		const float Distance1 = static_cast<float>((Points[PtIdx].OutVal - PointInSpace).SizeSquared());
+		const float Distance2 = static_cast<float>((Points[NextPtIdx].OutVal - PointInSpace).SizeSquared());
 		if (Distance1 < Distance2)
 		{
 			OutSquaredDistance = Distance1;
@@ -553,10 +553,10 @@ float FInterpCurve<T>::InaccurateFindNearestOnSegment(const T& PointInSpace, int
 	if (CIM_Linear == Points[PtIdx].InterpMode)
 	{
 		// like in function: FMath::ClosestPointOnLine
-		const float A = (Points[PtIdx].OutVal - PointInSpace) | (Points[NextPtIdx].OutVal - Points[PtIdx].OutVal);
-		const float B = (Points[NextPtIdx].OutVal - Points[PtIdx].OutVal).SizeSquared();
+		const float A = static_cast<float>((Points[PtIdx].OutVal - PointInSpace) | (Points[NextPtIdx].OutVal - Points[PtIdx].OutVal));
+		const float B = static_cast<float>((Points[NextPtIdx].OutVal - Points[PtIdx].OutVal).SizeSquared());
 		const float V = FMath::Clamp(-A / B, 0.f, 1.f);
-		OutSquaredDistance = (FMath::Lerp(Points[PtIdx].OutVal, Points[NextPtIdx].OutVal, V) - PointInSpace).SizeSquared();
+		OutSquaredDistance = static_cast<float>((FMath::Lerp(Points[PtIdx].OutVal, Points[NextPtIdx].OutVal, V) - PointInSpace).SizeSquared());
 		return V * Diff + Points[PtIdx].InVal;
 	}
 
@@ -587,14 +587,14 @@ float FInterpCurve<T>::InaccurateFindNearestOnSegment(const T& PointInSpace, int
 			{
 				const T LastBestTangent = FMath::CubicInterpDerivative(Points[PtIdx].OutVal, Points[PtIdx].LeaveTangent * Diff, Points[NextPtIdx].OutVal, Points[NextPtIdx].ArriveTangent * Diff, ValuesT[point]);
 				const T Delta = (PointInSpace - FoundPoint);
-				float Move = (LastBestTangent | Delta) / LastBestTangent.SizeSquared();
+				float Move = static_cast<float>((LastBestTangent | Delta) / LastBestTangent.SizeSquared());
 				Move = FMath::Clamp(Move, -LastMove*Scale, LastMove*Scale);
 				ValuesT[point] += Move;
 				ValuesT[point] = FMath::Clamp(ValuesT[point], 0.0f, 1.0f);
 				LastMove = FMath::Abs(Move);
 				FoundPoint = FMath::CubicInterp(Points[PtIdx].OutVal, Points[PtIdx].LeaveTangent * Diff, Points[NextPtIdx].OutVal, Points[NextPtIdx].ArriveTangent * Diff, ValuesT[point]);
 			}
-			DistancesSq[point] = (FoundPoint - PointInSpace).SizeSquared();
+			DistancesSq[point] = static_cast<float>((FoundPoint - PointInSpace).SizeSquared());
 			ValuesT[point] = ValuesT[point] * Diff + Points[PtIdx].InVal;
 		}
 
diff --git a/Engine/Source/Runtime/Core/Public/Math/Rotator.h b/Engine/Source/Runtime/Core/Public/Math/Rotator.h
--- a/Engine/Source/Runtime/Core/Public/Math/Rotator.h
+++ b/Engine/Source/Runtime/Core/Public/Math/Rotator.h
@@ -915,3 +915,24 @@ FORCEINLINE_DEBUGGABLE UE::Math::TRotator<T> FMath::LerpRange(const UE::Math::TR
 	// Similar to Lerp, but does not take the shortest path. Allows interpolation over more than 180 degrees.
 	return (A * ((T)1.0 - (T)Alpha) + B * Alpha).GetNormalized();
 }
+
+template<typename T>
+FORCEINLINE_DEBUGGABLE T FMath::ClampAngle(T AngleDegrees, T MinAngleDegrees, T MaxAngleDegrees)
+{
+	const T MaxDelta = UE::Math::TRotator<T>::ClampAxis(MaxAngleDegrees - MinAngleDegrees) * 0.5f;			// 0..180
+	const T RangeCenter = UE::Math::TRotator<T>::ClampAxis(MinAngleDegrees + MaxDelta);						// 0..360
+	const T DeltaFromCenter = UE::Math::TRotator<T>::NormalizeAxis(AngleDegrees - RangeCenter);				// -180..180
+
+	// maybe clamp to nearest edge
+	if (DeltaFromCenter > MaxDelta)
+	{
+		return UE::Math::TRotator<T>::NormalizeAxis(RangeCenter + MaxDelta);
+	}
+	else if (DeltaFromCenter < -MaxDelta)
+	{
+		return UE::Math::TRotator<T>::NormalizeAxis(RangeCenter - MaxDelta);
+	}
+
+	// already in range, just return it
+	return UE::Math::TRotator<T>::NormalizeAxis(AngleDegrees);
+}
\ No newline at end of file
diff --git a/Engine/Source/Runtime/Core/Public/Math/UnrealMathUtility.h b/Engine/Source/Runtime/Core/Public/Math/UnrealMathUtility.h
--- a/Engine/Source/Runtime/Core/Public/Math/UnrealMathUtility.h
+++ b/Engine/Source/Runtime/Core/Public/Math/UnrealMathUtility.h
@@ -822,7 +822,10 @@ public:
 	 * @param MaxAngleDegrees	"to" angle that defines the end of the range of valid angles
 	 * @return Returns clamped angle in the range -180..180.
 	 */
-	UE_NODISCARD static CORE_API float ClampAngle(float AngleDegrees, float MinAngleDegrees, float MaxAngleDegrees);
+	template<typename T>
+	UE_NODISCARD static T ClampAngle(T AngleDegrees, T MinAngleDegrees, T MaxAngleDegrees);
+
+	RESOLVE_FLOAT_AMBIGUITY_3_ARGS(ClampAngle);
 
 	/** Find the smallest angle between two headings (in degrees) */
 	template<typename T, typename T2, TEMPLATE_REQUIRES(TOr<TIsFloatingPoint<T>, TIsFloatingPoint<T2>>::Value)>
diff --git a/Engine/Source/Runtime/Core/Public/Misc/HashBuilder.h b/Engine/Source/Runtime/Core/Public/Misc/HashBuilder.h
--- a/Engine/Source/Runtime/Core/Public/Misc/HashBuilder.h
+++ b/Engine/Source/Runtime/Core/Public/Misc/HashBuilder.h
@@ -23,7 +23,7 @@ public:
 
 	void AppendRaw(const void* Data, int64 Num)
 	{
-		Hash = FCrc::MemCrc32(Data, Num, Hash);
+		Hash = FCrc::MemCrc32(Data, static_cast<int32>(Num), Hash);		// TODO: Update MemCrc32 to take an int64 Length?
 	}
 
 	template <typename T>
diff --git a/Engine/Source/Runtime/Core/Public/Templates/ResolveTypeAmbiguity.h b/Engine/Source/Runtime/Core/Public/Templates/ResolveTypeAmbiguity.h
--- a/Engine/Source/Runtime/Core/Public/Templates/ResolveTypeAmbiguity.h
+++ b/Engine/Source/Runtime/Core/Public/Templates/ResolveTypeAmbiguity.h
@@ -112,9 +112,9 @@
 																				TNot<TIsFloatingPoint<Arg3>> \
 																				>::Value \
 																			)> UE_DEPRECATED(5.0, "Arguments cause function resolution ambiguity, expected a floating point type.") \
-	static FORCEINLINE float Func(Arg1 X, Arg2 Y) \
+	static FORCEINLINE float Func(Arg1 X, Arg2 Y, Arg3 Z) \
 	{ \
-		return Func((float)X, (float)Y); \
+		return Func((float)X, (float)Y, (float)Z); \
 	}
 
 // Mixing float and double types with any other type results in same type as "X * Y", which is promoted to the floating point type with the highest precision of the argument types.
diff --git a/Engine/Source/Runtime/Engine/Classes/Animation/AnimCompress.h b/Engine/Source/Runtime/Engine/Classes/Animation/AnimCompress.h
--- a/Engine/Source/Runtime/Engine/Classes/Animation/AnimCompress.h
+++ b/Engine/Source/Runtime/Engine/Classes/Animation/AnimCompress.h
@@ -201,9 +201,9 @@ private:
 	bool bEnabled;
 	bool bUsed;
 	bool bPerformedCompression;
-	int32 TotalRaw;
-	int32 TotalBeforeCompressed;
-	int32 TotalAfterCompressed;
+	int64 TotalRaw;
+	int64 TotalBeforeCompressed;
+	int64 TotalAfterCompressed;
 	int32 NumberOfAnimations;
 
 	// Total time spent compressing animations
diff --git a/Engine/Source/Runtime/Engine/Classes/Animation/AnimCurveCompressionCodec_UniformIndexable.h b/Engine/Source/Runtime/Engine/Classes/Animation/AnimCurveCompressionCodec_UniformIndexable.h
--- a/Engine/Source/Runtime/Engine/Classes/Animation/AnimCurveCompressionCodec_UniformIndexable.h
+++ b/Engine/Source/Runtime/Engine/Classes/Animation/AnimCurveCompressionCodec_UniformIndexable.h
@@ -29,7 +29,7 @@ public:
 		return CompressedBuffer || RawCurve;
 	}
 
-	float GetNumSamples() const { return NumSamples; }
+	int32 GetNumSamples() const { return NumSamples; }
 
 	float GetValue(int32 SampleIndex) const;
 
diff --git a/Engine/Source/Runtime/Engine/Classes/Animation/BuiltInAttributeTypes.h b/Engine/Source/Runtime/Engine/Classes/Animation/BuiltInAttributeTypes.h
--- a/Engine/Source/Runtime/Engine/Classes/Animation/BuiltInAttributeTypes.h
+++ b/Engine/Source/Runtime/Engine/Classes/Animation/BuiltInAttributeTypes.h
@@ -72,8 +72,8 @@ struct FIntegerAnimationAttribute
 
 	void Interpolate(const FIntegerAnimationAttribute& Attribute, float Alpha)
 	{
-		Value *= (1.f - Alpha);
-		Value += (Attribute.Value * Alpha);
+		Value = FMath::TruncToInt32(Value * (1.f - Alpha));
+		Value += FMath::TruncToInt32(Attribute.Value * Alpha);
 	}
 };
 
diff --git a/Engine/Source/Runtime/Engine/Classes/Components/DirectionalLightComponent.h b/Engine/Source/Runtime/Engine/Classes/Components/DirectionalLightComponent.h
--- a/Engine/Source/Runtime/Engine/Classes/Components/DirectionalLightComponent.h
+++ b/Engine/Source/Runtime/Engine/Classes/Components/DirectionalLightComponent.h
@@ -329,7 +329,7 @@ class ENGINE_API UDirectionalLightComponent : public ULightComponent
 	}
 	virtual uint8 GetAtmosphereSunLightIndex() const override
 	{
-		return AtmosphereSunLightIndex;
+		return static_cast<uint8>(AtmosphereSunLightIndex);
 	}
 	virtual FLinearColor GetAtmosphereSunDiskColorScale() const override
 	{
diff --git a/Engine/Source/Runtime/Engine/Classes/Engine/NetSerialization.h b/Engine/Source/Runtime/Engine/Classes/Engine/NetSerialization.h
--- a/Engine/Source/Runtime/Engine/Classes/Engine/NetSerialization.h
+++ b/Engine/Source/Runtime/Engine/Classes/Engine/NetSerialization.h
@@ -294,7 +294,7 @@ bool WriteFixedCompressedFloat(const T Value, FArchive& Ar)
 	using Details = TFixedCompressedFloatDetails<MaxValue, NumBits>;
 
 	bool clamp = false;
-	int32 ScaledValue;
+	int64 ScaledValue;
 	if ( MaxValue > Details::MaxBitValue )
 	{
 		// We have to scale this down
diff --git a/Engine/Source/Runtime/Engine/Classes/Sound/QuartzQuantizationUtilities.cpp b/Engine/Source/Runtime/Engine/Classes/Sound/QuartzQuantizationUtilities.cpp
--- a/Engine/Source/Runtime/Engine/Classes/Sound/QuartzQuantizationUtilities.cpp
+++ b/Engine/Source/Runtime/Engine/Classes/Sound/QuartzQuantizationUtilities.cpp
@@ -133,11 +133,11 @@ void FQuartLatencyTracker::PushSingleResult(const double& InResult)
 {
 	if (++NumEntries == 0)
 	{
-		LifetimeAverage = InResult;
+		LifetimeAverage = (float)InResult;
 	}
 	else
 	{
-		LifetimeAverage = (LifetimeAverage * (NumEntries - 1) + InResult) / NumEntries;
+		LifetimeAverage = (LifetimeAverage * (NumEntries - 1) + (float)InResult) / NumEntries;
 	}
 
 	Min = FMath::Min(Min, (float)InResult);
@@ -166,31 +166,31 @@ namespace Audio
 
 	void FQuartzClockTickRate::SetFramesPerTick(int32 InNewFramesPerTick)
 	{
-		if (InNewFramesPerTick < 1.0)
+		if (InNewFramesPerTick < 1)
 		{
 			UE_LOG(LogAudioQuartz, Warning, TEXT("Quartz Metronme requires at least 1 frame per tick, clamping request"));
-			InNewFramesPerTick = 1.0;
+			InNewFramesPerTick = 1;
 		}
 
 		FramesPerTick = (double)InNewFramesPerTick;
 		RecalculateDurationsBasedOnFramesPerTick();
 	}
 
-	void FQuartzClockTickRate::SetMillisecondsPerTick(float InNewMillisecondsPerTick)
+	void FQuartzClockTickRate::SetMillisecondsPerTick(double InNewMillisecondsPerTick)
 	{
-		FramesPerTick = FMath::Max(1.0, ((double)InNewMillisecondsPerTick * SampleRate) / 1000.0);
+		FramesPerTick = FMath::Max(1.0, (InNewMillisecondsPerTick * SampleRate) / 1000.0);
 		RecalculateDurationsBasedOnFramesPerTick();
 	}
 
-	void FQuartzClockTickRate::SetThirtySecondNotesPerMinute(float InNewThirtySecondNotesPerMinute)
+	void FQuartzClockTickRate::SetThirtySecondNotesPerMinute(double InNewThirtySecondNotesPerMinute)
 	{
 		check(InNewThirtySecondNotesPerMinute > 0);
 
-		FramesPerTick = FMath::Max(1.0, (60. * SampleRate ) / (double)InNewThirtySecondNotesPerMinute);
+		FramesPerTick = FMath::Max(1.0, (60. * SampleRate ) / InNewThirtySecondNotesPerMinute);
 		RecalculateDurationsBasedOnFramesPerTick();
 	}
 
-	void FQuartzClockTickRate::SetBeatsPerMinute(float InNewBeatsPerMinute)
+	void FQuartzClockTickRate::SetBeatsPerMinute(double InNewBeatsPerMinute)
 	{
 		// same as 1/32nd notes,
 		// except there are 1/8th the number of quarter notes than thirty-second notes in a minute
@@ -199,15 +199,15 @@ namespace Audio
 		// FramesPerTick = 1/8 * (60.f / (InNewBeatsPerMinute)) * SampleRate;
 		// (60.0 / 8.0) = 7.5f
 
-		FramesPerTick = FMath::Max(1.0, (7.5 * SampleRate) / (double)InNewBeatsPerMinute);
+		FramesPerTick = FMath::Max(1.0, (7.5 * SampleRate) / InNewBeatsPerMinute);
 		RecalculateDurationsBasedOnFramesPerTick();
 	}
 
-	void FQuartzClockTickRate::SetSampleRate(float InNewSampleRate)
+	void FQuartzClockTickRate::SetSampleRate(double InNewSampleRate)
 	{
 		check(InNewSampleRate >= 0);
 
-		FramesPerTick = FMath::Max(1.0, ((double)InNewSampleRate / SampleRate) * FramesPerTick);
+		FramesPerTick = FMath::Max(1.0, (InNewSampleRate / SampleRate) * FramesPerTick);
 		SampleRate = InNewSampleRate;
 
 		RecalculateDurationsBasedOnFramesPerTick();
@@ -371,7 +371,7 @@ namespace Audio
 	}
 
 
-	void FQuartzClockTickRate::SetSecondsPerTick(float InNewSecondsPerTick)
+	void FQuartzClockTickRate::SetSecondsPerTick(double InNewSecondsPerTick)
 	{
 		SetMillisecondsPerTick(InNewSecondsPerTick * 1000.0);
 	}
diff --git a/Engine/Source/Runtime/Engine/Classes/Sound/QuartzQuantizationUtilities.h b/Engine/Source/Runtime/Engine/Classes/Sound/QuartzQuantizationUtilities.h
--- a/Engine/Source/Runtime/Engine/Classes/Sound/QuartzQuantizationUtilities.h
+++ b/Engine/Source/Runtime/Engine/Classes/Sound/QuartzQuantizationUtilities.h
@@ -290,28 +290,28 @@ namespace Audio
 		// Setters
 		void SetFramesPerTick(int32 InNewFramesPerTick);
 
-		void SetMillisecondsPerTick(float InNewMillisecondsPerTick);
+		void SetMillisecondsPerTick(double InNewMillisecondsPerTick);
 
-		void SetSecondsPerTick(float InNewSecondsPerTick);
+		void SetSecondsPerTick(double InNewSecondsPerTick);
 
-		void SetThirtySecondNotesPerMinute(float InNewThirtySecondNotesPerMinute);
+		void SetThirtySecondNotesPerMinute(double InNewThirtySecondNotesPerMinute);
 
-		void SetBeatsPerMinute(float InNewBeatsPerMinute);
+		void SetBeatsPerMinute(double InNewBeatsPerMinute);
 
-		void SetSampleRate(float InNewSampleRate);
+		void SetSampleRate(double InNewSampleRate);
 
 		// Getters
 		double GetFramesPerTick() const { return FramesPerTick; }
 
-		float GetMillisecondsPerTick() const { return MillisecondsPerTick; }
+		double GetMillisecondsPerTick() const { return MillisecondsPerTick; }
 
-		float GetSecondsPerTick() const { return SecondsPerTick; }
+		double GetSecondsPerTick() const { return SecondsPerTick; }
 
-		float GetThirtySecondNotesPerMinute() const { return ThirtySecondNotesPerMinute; }
+		double GetThirtySecondNotesPerMinute() const { return ThirtySecondNotesPerMinute; }
 
-		float GetBeatsPerMinute() const { return BeatsPerMinute; }
+		double GetBeatsPerMinute() const { return BeatsPerMinute; }
 
-		float GetSampleRate() const { return SampleRate; }
+		double GetSampleRate() const { return SampleRate; }
 
 		double GetFramesPerDuration(EQuartzCommandQuantization InDuration) const;
 
diff --git a/Engine/Source/Runtime/Engine/Private/Rendering/SkeletalMeshLODImporterData.cpp b/Engine/Source/Runtime/Engine/Private/Rendering/SkeletalMeshLODImporterData.cpp
--- a/Engine/Source/Runtime/Engine/Private/Rendering/SkeletalMeshLODImporterData.cpp
+++ b/Engine/Source/Runtime/Engine/Private/Rendering/SkeletalMeshLODImporterData.cpp
@@ -1037,7 +1037,7 @@ void FWedgePosition::FindMatchingPositionWegdeIndexes(const FVector3f &Position,
 		//No possible match
 		return;
 	}
-	FWedgePositionHelper::FIndexAndZ PositionIndexAndZ(INDEX_NONE, (FVector)Position);
+	FWedgePositionHelper::FIndexAndZ PositionIndexAndZ(INDEX_NONE, Position);
 	int32 SortedIndex = SortedPositions.Num()/2;
 	int32 StartIndex = 0;
 	int32 LastTopIndex = SortedPositions.Num();
@@ -1155,7 +1155,7 @@ void FWedgePosition::FillWedgePosition(
 	OutOverlappingPosition.SortedPositions.Reserve(NumWedges);
 	for (int32 WedgeIndex = 0; WedgeIndex < NumWedges; WedgeIndex++)
 	{
-		new(OutOverlappingPosition.SortedPositions)FWedgePositionHelper::FIndexAndZ(WedgeIndex, (FVector)OutOverlappingPosition.Points[OutOverlappingPosition.Wedges[WedgeIndex].VertexIndex]);
+		new(OutOverlappingPosition.SortedPositions)FWedgePositionHelper::FIndexAndZ(WedgeIndex, OutOverlappingPosition.Points[OutOverlappingPosition.Wedges[WedgeIndex].VertexIndex]);
 	}
 
 	// Sort the vertices by z value
diff --git a/Engine/Source/Runtime/Engine/Public/ActorPartition/ActorPartitionSubsystem.h b/Engine/Source/Runtime/Engine/Public/ActorPartition/ActorPartitionSubsystem.h
--- a/Engine/Source/Runtime/Engine/Public/ActorPartition/ActorPartitionSubsystem.h
+++ b/Engine/Source/Runtime/Engine/Public/ActorPartition/ActorPartitionSubsystem.h
@@ -140,14 +140,14 @@ public:
 		{
 			return FBox(
 				FVector(
-					InCellCoord.X * InGridSize,
-					InCellCoord.Y * InGridSize,
-					InCellCoord.Z * InGridSize
+					static_cast<FVector::FReal>(InCellCoord.X * InGridSize),
+					static_cast<FVector::FReal>(InCellCoord.Y * InGridSize),
+					static_cast<FVector::FReal>(InCellCoord.Z * InGridSize)
 				),
 				FVector(
-					InCellCoord.X * InGridSize + InGridSize,
-					InCellCoord.Y * InGridSize + InGridSize,
-					InCellCoord.Z * InGridSize + InGridSize
+					static_cast<FVector::FReal>(InCellCoord.X * InGridSize + InGridSize),
+					static_cast<FVector::FReal>(InCellCoord.Y * InGridSize + InGridSize),
+					static_cast<FVector::FReal>(InCellCoord.Z * InGridSize + InGridSize)
 				)
 			);
 		}
diff --git a/Engine/Source/Runtime/Engine/Public/CanvasRender.h b/Engine/Source/Runtime/Engine/Public/CanvasRender.h
--- a/Engine/Source/Runtime/Engine/Public/CanvasRender.h
+++ b/Engine/Source/Runtime/Engine/Public/CanvasRender.h
@@ -16,11 +16,11 @@ public:
 		GraphBuilder.AddPass(Forward<FRDGEventName>(PassName), PassParameters, ERDGPassFlags::Raster,
 			[LocalScissorRect = ScissorRect, LocalViewportRect = ViewportRect, LocalExecuteLambda = Forward<ExecuteLambdaType&&>(ExecuteLambda)](FRHICommandListImmediate& RHICmdList)
 		{
-			RHICmdList.SetViewport(LocalViewportRect.Min.X, LocalViewportRect.Min.Y, 0.0f, LocalViewportRect.Max.X, LocalViewportRect.Max.Y, 1.0f);
+			RHICmdList.SetViewport(static_cast<float>(LocalViewportRect.Min.X), static_cast<float>(LocalViewportRect.Min.Y), 0.0f, static_cast<float>(LocalViewportRect.Max.X), static_cast<float>(LocalViewportRect.Max.Y), 1.0f);
 
 			if (LocalScissorRect.Area() > 0)
 			{
-				RHICmdList.SetScissorRect(true, LocalScissorRect.Min.X, LocalScissorRect.Min.Y, LocalScissorRect.Max.X, LocalScissorRect.Max.Y);
+				RHICmdList.SetScissorRect(true, static_cast<uint32>(LocalScissorRect.Min.X), static_cast<uint32>(LocalScissorRect.Min.Y), static_cast<uint32>(LocalScissorRect.Max.X), static_cast<uint32>(LocalScissorRect.Max.Y));
 			}
 		
 			LocalExecuteLambda(RHICmdList);
diff --git a/Engine/Source/Runtime/Engine/Public/ChartCreation.h b/Engine/Source/Runtime/Engine/Public/ChartCreation.h
--- a/Engine/Source/Runtime/Engine/Public/ChartCreation.h
+++ b/Engine/Source/Runtime/Engine/Public/ChartCreation.h
@@ -282,7 +282,7 @@ public:
 
 	double GetPercentMissedVSync(int32 TargetFPS) const
 	{
-		const int64 TotalTargetFrames = TargetFPS * GetTotalTime();
+		const int64 TotalTargetFrames = static_cast<int64>(TargetFPS * GetTotalTime());
 		const int64 MissedFrames = FMath::Max<int64>(TotalTargetFrames - GetNumFrames(), 0);
 		return ((MissedFrames * 100.0) / (double)TotalTargetFrames);
 	}
diff --git a/Engine/Source/Runtime/Engine/Public/Collision.h b/Engine/Source/Runtime/Engine/Public/Collision.h
--- a/Engine/Source/Runtime/Engine/Public/Collision.h
+++ b/Engine/Source/Runtime/Engine/Public/Collision.h
@@ -176,7 +176,7 @@ FORCEINLINE bool LineCheckWithTriangle(FHitResult& Result,const FVector& V1,cons
 	FVector	Edge1 = V3 - V1,
 		Edge2 = V2 - V1,
 		P = Direction ^ Edge2;
-	float	Determinant = Edge1 | P;
+	FVector::FReal	Determinant = Edge1 | P;
 
 	if(Determinant < UE_DELTA)
 	{
@@ -184,7 +184,7 @@ FORCEINLINE bool LineCheckWithTriangle(FHitResult& Result,const FVector& V1,cons
 	}
 
 	FVector	T = Start - V1;
-	float	U = T | P;
+	FVector::FReal	U = T | P;
 
 	if(U < 0.0f || U > Determinant)
 	{
@@ -192,14 +192,14 @@ FORCEINLINE bool LineCheckWithTriangle(FHitResult& Result,const FVector& V1,cons
 	}
 
 	FVector	Q = T ^ Edge1;
-	float	V = Direction | Q;
+	FVector::FReal	V = Direction | Q;
 
 	if(V < 0.0f || U + V > Determinant)
 	{
 		return false;
 	}
 
-	float	Time = (Edge2 | Q) / Determinant;
+	FVector::FReal	Time = (Edge2 | Q) / Determinant;
 
 	if(Time < 0.0f || Time > Result.Time)
 	{
@@ -207,7 +207,7 @@ FORCEINLINE bool LineCheckWithTriangle(FHitResult& Result,const FVector& V1,cons
 	}
 
 	Result.Normal = ((V3-V2)^(V2-V1)).GetSafeNormal();
-	Result.Time = ((V1 - Start)|Result.Normal) / (Result.Normal|Direction);
+	Result.Time = static_cast<float>(((V1 - Start)|Result.Normal) / (Result.Normal|Direction));							// LWC_TODO: precision loss. Make FHitResult::Time/Distance doubles?
 
 	return true;
 }
diff --git a/Engine/Source/Runtime/Engine/Public/GPUSkinCache.h b/Engine/Source/Runtime/Engine/Public/GPUSkinCache.h
--- a/Engine/Source/Runtime/Engine/Public/GPUSkinCache.h
+++ b/Engine/Source/Runtime/Engine/Public/GPUSkinCache.h
@@ -327,7 +327,7 @@ public:
 
 		inline uint32 GetNumBytes() const
 		{
-			return Allocation->GetNumBytes();
+			return IntCastChecked<uint32>(Allocation->GetNumBytes());
 		}
 
 		FSkinCacheRWBuffer* Find(const FVertexBufferAndSRV& BoneBuffer, uint32 Revision)
diff --git a/Engine/Source/Runtime/Engine/Public/HLSLTree/HLSLTreeCommon.h b/Engine/Source/Runtime/Engine/Public/HLSLTree/HLSLTreeCommon.h
--- a/Engine/Source/Runtime/Engine/Public/HLSLTree/HLSLTreeCommon.h
+++ b/Engine/Source/Runtime/Engine/Public/HLSLTree/HLSLTreeCommon.h
@@ -219,7 +219,7 @@ class FExpressionSwitchBase : public FExpression
 public:
 	static constexpr int8 MaxInputs = 8;
 
-	FExpressionSwitchBase(TConstArrayView<const FExpression*> InInputs) : NumInputs(InInputs.Num())
+	FExpressionSwitchBase(TConstArrayView<const FExpression*> InInputs) : NumInputs(static_cast<int8>(InInputs.Num()))
 	{
 		check(InInputs.Num() <= MaxInputs);
 		for (int32 i = 0; i < InInputs.Num(); ++i)
diff --git a/Engine/Source/Runtime/Engine/Public/InGamePerformanceTracker.h b/Engine/Source/Runtime/Engine/Public/InGamePerformanceTracker.h
--- a/Engine/Source/Runtime/Engine/Public/InGamePerformanceTracker.h
+++ b/Engine/Source/Runtime/Engine/Public/InGamePerformanceTracker.h
@@ -88,7 +88,7 @@ public:
 
 	FORCEINLINE double GetAverageTimeSeconds() const
 	{
-		return FPlatformTime::ToSeconds(History.GetAverageCycles());
+		return FPlatformTime::ToSeconds64(History.GetAverageCycles());
 	}
 	
 	/** Enters a timed section of code. Can be re-entrant and will take the first entry point time. Only callable from game thread. */
diff --git a/Engine/Source/Runtime/Engine/Public/Instances/InstancedPlacementHash.h b/Engine/Source/Runtime/Engine/Public/Instances/InstancedPlacementHash.h
--- a/Engine/Source/Runtime/Engine/Public/Instances/InstancedPlacementHash.h
+++ b/Engine/Source/Runtime/Engine/Public/Instances/InstancedPlacementHash.h
@@ -57,7 +57,7 @@ private:
 
 	FVector MakeLocation(FKey CellKey) const
 	{
-		return FVector(CellKey.X << HashCellBits, CellKey.Y << HashCellBits, CellKey.Z << HashCellBits);
+		return FVector(FVector::FReal(CellKey.X << HashCellBits), FVector::FReal(CellKey.Y << HashCellBits), FVector::FReal(CellKey.Z << HashCellBits));
 	}
 
 public:
@@ -100,11 +100,11 @@ public:
 		FKey MinKey = MakeKey(InBox.Min);
 		FKey MaxKey = MakeKey(InBox.Max);
 
-		for (int32 z = MinKey.Z; z <= MaxKey.Z; ++z)
+		for (int64 z = MinKey.Z; z <= MaxKey.Z; ++z)
 		{
-			for (int32 y = MinKey.Y; y <= MaxKey.Y; y++)
+			for (int64 y = MinKey.Y; y <= MaxKey.Y; y++)
 			{
-				for (int32 x = MinKey.X; x <= MaxKey.X; x++)
+				for (int64 x = MinKey.X; x <= MaxKey.X; x++)
 				{
 					auto* SetPtr = CellMap.Find(FKey(x, y, z));
 					if (SetPtr)
diff --git a/Engine/Source/Runtime/Engine/Public/MaterialShared.h b/Engine/Source/Runtime/Engine/Public/MaterialShared.h
--- a/Engine/Source/Runtime/Engine/Public/MaterialShared.h
+++ b/Engine/Source/Runtime/Engine/Public/MaterialShared.h
@@ -3386,7 +3386,7 @@ public:
 	{
 		Locs.AddUninitialized();
 		int64 ResourceOffset = this->Tell();
-		Locs.Last().Offset = ResourceOffset;
+		Locs.Last().Offset = IntCastChecked<uint32>(ResourceOffset);
 	}
 
 	inline void EndSerializingMaterialResource(const FMaterialResource& Resource)
diff --git a/Engine/Source/Runtime/Engine/Public/Physics/SimpleSuspension.h b/Engine/Source/Runtime/Engine/Public/Physics/SimpleSuspension.h
--- a/Engine/Source/Runtime/Engine/Public/Physics/SimpleSuspension.h
+++ b/Engine/Source/Runtime/Engine/Public/Physics/SimpleSuspension.h
@@ -67,11 +67,11 @@ const FSimpleSuspensionRaycast SimpleSuspensionRaycastGroundPlane = [](const FVe
 	// Do an intersection with the z-plane at 0,0,0.
 	const FVector Norm = FVector::UpVector;
 	const FVector Diff = RayEnd - RayBegin;
-	const float ProjA = FVector::DotProduct(RayBegin, Norm);
-	const float ProjDiff = FVector::DotProduct(Diff, Norm);
+	const FVector::FReal ProjA = FVector::DotProduct(RayBegin, Norm);
+	const FVector::FReal ProjDiff = FVector::DotProduct(Diff, Norm);
 	if (FMath::Abs(ProjDiff) > UE_SMALL_NUMBER)
 	{
-		const float Lambda = -ProjA / ProjDiff;
+		const FVector::FReal Lambda = -ProjA / ProjDiff;
 		if (Lambda > 0.f && Lambda <= 1.f)
 		{
 			OutRayHitLocation = RayBegin + (Lambda * Diff);
diff --git a/Engine/Source/Runtime/Engine/Public/PrecomputedLightVolume.h b/Engine/Source/Runtime/Engine/Public/PrecomputedLightVolume.h
--- a/Engine/Source/Runtime/Engine/Public/PrecomputedLightVolume.h
+++ b/Engine/Source/Runtime/Engine/Public/PrecomputedLightVolume.h
@@ -188,7 +188,7 @@ public:
 		return bAddedToScene;
 	}
 
-	float GetNodeLevelExtent(int32 Level) const
+	double GetNodeLevelExtent(int32 Level) const
 	{
 		return OctreeForRendering->GetNodeLevelExtent(Level);
 	}
diff --git a/Engine/Source/Runtime/Engine/Public/PrimitiveUniformShaderParameters.h b/Engine/Source/Runtime/Engine/Public/PrimitiveUniformShaderParameters.h
--- a/Engine/Source/Runtime/Engine/Public/PrimitiveUniformShaderParameters.h
+++ b/Engine/Source/Runtime/Engine/Public/PrimitiveUniformShaderParameters.h
@@ -213,10 +213,10 @@ public:
 	inline FPrimitiveUniformShaderParametersBuilder& WorldBounds(const FBoxSphereBounds& InWorldBounds)
 	{
 		AbsoluteObjectWorldPosition = InWorldBounds.Origin;
-		ObjectRadius = InWorldBounds.SphereRadius;
-		Parameters.ObjectBoundsX = InWorldBounds.BoxExtent.X;
-		Parameters.ObjectBoundsY = InWorldBounds.BoxExtent.Y;
-		Parameters.ObjectBoundsZ = InWorldBounds.BoxExtent.Z;
+		ObjectRadius = static_cast<float>(InWorldBounds.SphereRadius);											//LWC_TODO: Precision loss
+		Parameters.ObjectBoundsX = static_cast<float>(InWorldBounds.BoxExtent.X);
+		Parameters.ObjectBoundsY = static_cast<float>(InWorldBounds.BoxExtent.Y);
+		Parameters.ObjectBoundsZ = static_cast<float>(InWorldBounds.BoxExtent.Z);
 		return *this;
 	}
 
diff --git a/Engine/Source/Runtime/Engine/Public/Raster.h b/Engine/Source/Runtime/Engine/Public/Raster.h
--- a/Engine/Source/Runtime/Engine/Public/Raster.h
+++ b/Engine/Source/Runtime/Engine/Public/Raster.h
@@ -21,7 +21,7 @@ public:
 	void DrawTriangle(const InterpolantType& I0, const InterpolantType& I1, const InterpolantType& I2, const FVector2D& P0, const FVector2D& P1, const FVector2D& P2, bool BackFacing)
 	{
 		InterpolantType	Interpolants[3] = { I0, I1, I2 };
-		FVector2D		Points[3] = { P0, P1, P2 };
+		FVector2f		Points[3] = { FVector2f(P0), FVector2f(P1), FVector2f(P2) };
 
 		// Find the top point.
 
diff --git a/Engine/Source/Runtime/Engine/Public/Rendering/SkeletalMeshLODImporterData.h b/Engine/Source/Runtime/Engine/Public/Rendering/SkeletalMeshLODImporterData.h
--- a/Engine/Source/Runtime/Engine/Public/Rendering/SkeletalMeshLODImporterData.h
+++ b/Engine/Source/Runtime/Engine/Public/Rendering/SkeletalMeshLODImporterData.h
@@ -629,7 +629,7 @@ namespace FWedgePositionHelper
 		FIndexAndZ() {}
 
 		/** Initialization constructor. */
-		FIndexAndZ(int32 InIndex, FVector V)
+		FIndexAndZ(int32 InIndex, FVector3f V)
 		{
 			Z = 0.30f * V.X + 0.33f * V.Y + 0.37f * V.Z;
 			Index = InIndex;
diff --git a/Engine/Source/Runtime/Engine/Public/ScreenRendering.h b/Engine/Source/Runtime/Engine/Public/ScreenRendering.h
--- a/Engine/Source/Runtime/Engine/Public/ScreenRendering.h
+++ b/Engine/Source/Runtime/Engine/Public/ScreenRendering.h
@@ -43,7 +43,7 @@ public:
 	virtual void InitRHI() override
 	{
 		FVertexDeclarationElementList Elements;
-		uint32 Stride = sizeof(FScreenVertex);
+		uint16 Stride = sizeof(FScreenVertex);
 		Elements.Add(FVertexElement(0,STRUCT_OFFSET(FScreenVertex,Position),VET_Float2,0,Stride));
 		Elements.Add(FVertexElement(0,STRUCT_OFFSET(FScreenVertex,UV),VET_Float2,1,Stride));
 		VertexDeclarationRHI = RHICreateVertexDeclaration(Elements);
diff --git a/Engine/Source/Runtime/Engine/Public/ShadowMap.h b/Engine/Source/Runtime/Engine/Public/ShadowMap.h
--- a/Engine/Source/Runtime/Engine/Public/ShadowMap.h
+++ b/Engine/Source/Runtime/Engine/Public/ShadowMap.h
@@ -87,8 +87,8 @@ public:
 	uint32 GetSizeY() const { return SizeY; }
 
 	// USurface interface
-	virtual float GetSurfaceWidth() const { return SizeX; }
-	virtual float GetSurfaceHeight() const { return SizeY; }
+	virtual float GetSurfaceWidth() const { return (float)SizeX; }
+	virtual float GetSurfaceHeight() const { return (float)SizeY; }
 
 	enum ShadowMapDataType {
 		UNKNOWN,
diff --git a/Engine/Source/Runtime/Engine/Public/StaticMeshResources.h b/Engine/Source/Runtime/Engine/Public/StaticMeshResources.h
--- a/Engine/Source/Runtime/Engine/Public/StaticMeshResources.h
+++ b/Engine/Source/Runtime/Engine/Public/StaticMeshResources.h
@@ -1267,7 +1267,7 @@ public:
 			SetInstanceTransformInternal<float>(InstanceIndex, InstanceTransform);
 		}
 
-		SetInstanceLightMapDataInternal(InstanceIndex, FVector4f(LightmapUVBias.X, LightmapUVBias.Y, ShadowmapUVBias.X, ShadowmapUVBias.Y));
+		SetInstanceLightMapDataInternal(InstanceIndex, FVector4f((float)LightmapUVBias.X, (float)LightmapUVBias.Y, (float)ShadowmapUVBias.X, (float)ShadowmapUVBias.Y));
 
 		for (int32 i = 0; i < NumCustomDataFloats; ++i)
 		{
@@ -1300,7 +1300,7 @@ public:
 
 	FORCEINLINE void SetInstanceLightMapData(int32 InstanceIndex, const FVector2D& LightmapUVBias, const FVector2D& ShadowmapUVBias)
 	{
-		SetInstanceLightMapDataInternal(InstanceIndex, FVector4f(LightmapUVBias.X, LightmapUVBias.Y, ShadowmapUVBias.X, ShadowmapUVBias.Y));
+		SetInstanceLightMapDataInternal(InstanceIndex, FVector4f((float)LightmapUVBias.X, (float)LightmapUVBias.Y, (float)ShadowmapUVBias.X, (float)ShadowmapUVBias.Y));
 	}
 	
 	FORCEINLINE void SetInstanceCustomData(int32 InstanceIndex, int32 Index, float CustomData)
diff --git a/Engine/Source/Runtime/Engine/Public/StereoLayerRendering.h b/Engine/Source/Runtime/Engine/Public/StereoLayerRendering.h
--- a/Engine/Source/Runtime/Engine/Public/StereoLayerRendering.h
+++ b/Engine/Source/Runtime/Engine/Public/StereoLayerRendering.h
@@ -44,10 +44,10 @@ public:
 		if (InUVAdjust.IsBound())
 		{
 			FVector4f UVAdjust;
-			UVAdjust.X = UVRect.Min.X;
-			UVAdjust.Y = UVRect.Min.Y;
-			UVAdjust.Z = UVRect.Max.X - UVRect.Min.X;
-			UVAdjust.W = UVRect.Max.Y - UVRect.Min.Y;
+			UVAdjust.X = static_cast<float>(UVRect.Min.X);
+			UVAdjust.Y = static_cast<float>(UVRect.Min.Y);
+			UVAdjust.Z = static_cast<float>(UVRect.Max.X - UVRect.Min.X);
+			UVAdjust.W = static_cast<float>(UVRect.Max.Y - UVRect.Min.Y);
 			SetShaderValue(RHICmdList, VS, InUVAdjust, UVAdjust);
 		}
 
diff --git a/Engine/Source/Runtime/Engine/Public/TextureLayout.h b/Engine/Source/Runtime/Engine/Public/TextureLayout.h
--- a/Engine/Source/Runtime/Engine/Public/TextureLayout.h
+++ b/Engine/Source/Runtime/Engine/Public/TextureLayout.h
@@ -43,7 +43,7 @@ public:
 		bPowerOfTwoSize(bInPowerOfTwoSize),
 		bAlignByFour(bInAlignByFour)
 	{
-		new(Nodes) FTextureLayoutNode(0, 0, MaxSizeX, MaxSizeY);
+		new(Nodes) FTextureLayoutNode(0, 0, IntCastChecked<uint16>(MaxSizeX), IntCastChecked<uint16>(MaxSizeY));
 	}
 
 	/**
@@ -327,16 +327,16 @@ private:
 				new(Nodes) FTextureLayoutNode(
 					CurrentNode.MinX,
 					CurrentNode.MinY,
-					ElementSizeX,
+					IntCastChecked<uint16>(ElementSizeX),
 					CurrentNode.SizeY
 					);
 
 				// Create a second child to contain the leftover area in the X direction
 				Nodes[NodeIndex].ChildB = Nodes.Num();
 				new(Nodes) FTextureLayoutNode(
-					CurrentNode.MinX + ElementSizeX,
+					IntCastChecked<uint16>(CurrentNode.MinX + ElementSizeX),
 					CurrentNode.MinY,
-					CurrentNode.SizeX - ElementSizeX,
+					IntCastChecked<uint16>(CurrentNode.SizeX - ElementSizeX),
 					CurrentNode.SizeY
 					);
 			}
@@ -347,15 +347,15 @@ private:
 					CurrentNode.MinX,
 					CurrentNode.MinY,
 					CurrentNode.SizeX,
-					ElementSizeY
+					IntCastChecked<uint16>(ElementSizeY)
 					);
 
 				Nodes[NodeIndex].ChildB = Nodes.Num();
 				new(Nodes) FTextureLayoutNode(
 					CurrentNode.MinX,
-					CurrentNode.MinY + ElementSizeY,
+					IntCastChecked<uint16>(CurrentNode.MinY + ElementSizeY),
 					CurrentNode.SizeX,
-					CurrentNode.SizeY - ElementSizeY
+					IntCastChecked<uint16>(CurrentNode.SizeY - ElementSizeY)
 					);
 			}
 
diff --git a/Engine/Source/Runtime/Engine/Public/TextureLayout3d.h b/Engine/Source/Runtime/Engine/Public/TextureLayout3d.h
--- a/Engine/Source/Runtime/Engine/Public/TextureLayout3d.h
+++ b/Engine/Source/Runtime/Engine/Public/TextureLayout3d.h
@@ -36,8 +36,8 @@ public:
 		bAllowShrink(bInAllowShrink)
 	{
 		check(MaxSizeX < USHRT_MAX && MaxSizeY < USHRT_MAX && MaxSizeZ < USHRT_MAX);
-		new(Nodes) FTextureLayoutNode3d(0, 0, 0, MaxSizeX, MaxSizeY, MaxSizeZ, INDEX_NONE);
-		new (UnusedLeaves) FUnusedLeaf(0, 0, 0, MaxSizeX, MaxSizeY, MaxSizeZ, 0);
+		new(Nodes) FTextureLayoutNode3d(0, 0, 0, IntCastChecked<uint16>(MaxSizeX), IntCastChecked<uint16>(MaxSizeY), IntCastChecked<uint16>(MaxSizeZ), INDEX_NONE);
+		new (UnusedLeaves) FUnusedLeaf(0, 0, 0, IntCastChecked<uint16>(MaxSizeX), IntCastChecked<uint16>(MaxSizeY), IntCastChecked<uint16>(MaxSizeZ), 0);
 	}
 
 	/**
@@ -376,17 +376,17 @@ private:
 						CurrentNode.MinX,
 						CurrentNode.MinY,
 						CurrentNode.MinZ,
-						ElementSizeX,
+						IntCastChecked<uint16>(ElementSizeX),
 						CurrentNode.SizeY,
 						CurrentNode.SizeZ,
 						NodeIndex);
 
 					// Create a second child to contain the leftover area in the X direction
 					new(&Nodes[ChildBNodeIndex]) FTextureLayoutNode3d(
-						CurrentNode.MinX + ElementSizeX,
+						IntCastChecked<uint16>(CurrentNode.MinX + ElementSizeX),
 						CurrentNode.MinY,
 						CurrentNode.MinZ,
-						CurrentNode.SizeX - ElementSizeX,
+						IntCastChecked<uint16>(CurrentNode.SizeX - ElementSizeX),
 						CurrentNode.SizeY,
 						CurrentNode.SizeZ,
 						NodeIndex);
@@ -399,16 +399,16 @@ private:
 						CurrentNode.MinZ,
 						CurrentNode.SizeX,
 						CurrentNode.SizeY,
-						ElementSizeZ,
+						IntCastChecked<uint16>(ElementSizeZ),
 						NodeIndex);
 
 					new(&Nodes[ChildBNodeIndex]) FTextureLayoutNode3d(
 						CurrentNode.MinX,
 						CurrentNode.MinY,
-						CurrentNode.MinZ + ElementSizeZ,
+						IntCastChecked<uint16>(CurrentNode.MinZ + ElementSizeZ),
 						CurrentNode.SizeX,
 						CurrentNode.SizeY,
-						CurrentNode.SizeZ - ElementSizeZ,
+						IntCastChecked<uint16>(CurrentNode.SizeZ - ElementSizeZ),
 						NodeIndex);
 				}
 			}
@@ -421,16 +421,16 @@ private:
 						CurrentNode.MinY,
 						CurrentNode.MinZ,
 						CurrentNode.SizeX,
-						ElementSizeY,
+						IntCastChecked<uint16>(ElementSizeY),
 						CurrentNode.SizeZ,
 						NodeIndex);
 
 					new(&Nodes[ChildBNodeIndex]) FTextureLayoutNode3d(
 						CurrentNode.MinX,
-						CurrentNode.MinY + ElementSizeY,
+						IntCastChecked<uint16>(CurrentNode.MinY + ElementSizeY),
 						CurrentNode.MinZ,
 						CurrentNode.SizeX,
-						CurrentNode.SizeY - ElementSizeY,
+						IntCastChecked<uint16>(CurrentNode.SizeY - ElementSizeY),
 						CurrentNode.SizeZ,
 						NodeIndex);
 				}
@@ -442,16 +442,16 @@ private:
 						CurrentNode.MinZ,
 						CurrentNode.SizeX,
 						CurrentNode.SizeY,
-						ElementSizeZ,
+						IntCastChecked<uint16>(ElementSizeZ),
 						NodeIndex);
 
 					new(&Nodes[ChildBNodeIndex]) FTextureLayoutNode3d(
 						CurrentNode.MinX,
 						CurrentNode.MinY,
-						CurrentNode.MinZ + ElementSizeZ,
+						IntCastChecked<uint16>(CurrentNode.MinZ + ElementSizeZ),
 						CurrentNode.SizeX,
 						CurrentNode.SizeY,
-						CurrentNode.SizeZ - ElementSizeZ,
+						IntCastChecked<uint16>(CurrentNode.SizeZ - ElementSizeZ),
 						NodeIndex);
 				}
 			}
diff --git a/Engine/Source/Runtime/Engine/Public/WorldPartition/HLOD/HLODLayer.h b/Engine/Source/Runtime/Engine/Public/WorldPartition/HLOD/HLODLayer.h
--- a/Engine/Source/Runtime/Engine/Public/WorldPartition/HLOD/HLODLayer.h
+++ b/Engine/Source/Runtime/Engine/Public/WorldPartition/HLOD/HLODLayer.h
@@ -56,7 +56,7 @@ public:
 	const UHLODBuilderSettings* GetHLODBuilderSettings() const { return HLODBuilderSettings; }
 	FName GetRuntimeGrid(uint32 InHLODLevel) const;
 	int32 GetCellSize() const { return !bIsSpatiallyLoaded ? 0 : CellSize; }
-	float GetLoadingRange() const { return !bIsSpatiallyLoaded ? WORLD_MAX : LoadingRange; }
+	double GetLoadingRange() const { return !bIsSpatiallyLoaded ? WORLD_MAX : LoadingRange; }
 	const TSoftObjectPtr<UHLODLayer>& GetParentLayer() const;
 	const void SetParentLayer(const TSoftObjectPtr<UHLODLayer>& InParentLayer);
 	bool IsSpatiallyLoaded() const { return bIsSpatiallyLoaded; }
diff --git a/Engine/Source/Runtime/Engine/Public/WorldPartition/WorldPartitionEditorSpatialHash.h b/Engine/Source/Runtime/Engine/Public/WorldPartition/WorldPartitionEditorSpatialHash.h
--- a/Engine/Source/Runtime/Engine/Public/WorldPartition/WorldPartitionEditorSpatialHash.h
+++ b/Engine/Source/Runtime/Engine/Public/WorldPartition/WorldPartitionEditorSpatialHash.h
@@ -95,19 +95,19 @@ class ENGINE_API UWorldPartitionEditorSpatialHash : public UWorldPartitionEditor
 		check(InCellCoord.Level >= 0);
 		const int64 CellSizeForLevel = (int64)CellSize * (1LL << InCellCoord.Level);
 		const FVector Min = FVector(
-			InCellCoord.X * CellSizeForLevel, 
-			InCellCoord.Y * CellSizeForLevel, 
-			InCellCoord.Z * CellSizeForLevel
+			static_cast<FVector::FReal>(InCellCoord.X * CellSizeForLevel), 
+			static_cast<FVector::FReal>(InCellCoord.Y * CellSizeForLevel), 
+			static_cast<FVector::FReal>(InCellCoord.Z * CellSizeForLevel)
 		);
-		const FVector Max = Min + FVector(CellSizeForLevel);
+		const FVector Max = Min + FVector(static_cast<double>(CellSizeForLevel));
 		return FBox(Min, Max);
 	}
 
 	inline int32 GetLevelForBox(const FBox& Box) const
 	{
 		const FVector Extent = Box.GetExtent();
-		const float MaxLength = Extent.GetMax() * 2.0;
-		return FMath::CeilToInt(FMath::Max<float>(FMath::Log2(MaxLength / CellSize), 0));
+		const FVector::FReal MaxLength = Extent.GetMax() * 2.0;
+		return FMath::CeilToInt32(FMath::Max<FVector::FReal>(FMath::Log2(MaxLength / CellSize), 0));
 	}
 
 	inline int32 ForEachIntersectingCells(const FBox& InBounds, int32 Level, TFunctionRef<void(const FCellCoord&)> InOperation) const
@@ -117,11 +117,11 @@ class ENGINE_API UWorldPartitionEditorSpatialHash : public UWorldPartitionEditor
 		FCellCoord MinCellCoords(GetCellCoords(InBounds.Min, Level));
 		FCellCoord MaxCellCoords(GetCellCoords(InBounds.Max, Level));
 
-		for (int32 z=MinCellCoords.Z; z<=MaxCellCoords.Z; z++)
+		for (int64 z=MinCellCoords.Z; z<=MaxCellCoords.Z; z++)
 		{
-			for (int32 y=MinCellCoords.Y; y<=MaxCellCoords.Y; y++)
+			for (int64 y=MinCellCoords.Y; y<=MaxCellCoords.Y; y++)
 			{
-				for (int32 x=MinCellCoords.X; x<=MaxCellCoords.X; x++)
+				for (int64 x=MinCellCoords.X; x<=MaxCellCoords.X; x++)
 				{
 					InOperation(FCellCoord(x, y, z, Level));
 					NumCells++;
