LLM: Fix compile error when ENABLE_LOW_LEVEL_MEM_TRACKER is defined but the platform does not allow it.
#rnx

diff --git a/Engine/Source/Runtime/Core/Public/HAL/LowLevelMemTracker.h b/Engine/Source/Runtime/Core/Public/HAL/LowLevelMemTracker.h
--- a/Engine/Source/Runtime/Core/Public/HAL/LowLevelMemTracker.h
+++ b/Engine/Source/Runtime/Core/Public/HAL/LowLevelMemTracker.h
@@ -15,6 +15,7 @@
 
 // *** enable/disable LLM here ***
 #if !defined(ENABLE_LOW_LEVEL_MEM_TRACKER) || !LLM_ENABLED_ON_PLATFORM 
+	#undef ENABLE_LOW_LEVEL_MEM_TRACKER
 	#define ENABLE_LOW_LEVEL_MEM_TRACKER (LLM_ENABLED_ON_PLATFORM && !UE_BUILD_SHIPPING && (!UE_BUILD_TEST || ALLOW_LOW_LEVEL_MEM_TRACKER_IN_TEST) && WITH_ENGINE && 1)
 #endif
 
