CookMultiprocess and PlatformMisc: Add experimental -PhysicalCoreLimit commandline argument to allow setting a limit on physical cores without losing information about the number of logical cores per physical core. Move parsing of -corelimit and -physicalcorelimit into a single function in PlatformMisc.
Change ShaderCompiler OnMachineResourcesChanged to allow specifying the number of logical cores and physicalcores rather than reading it from the (immutable-per-process) FPlatformMisc::NumberOfCores.
Change the CookDirector to use physicalcorelimit and OnMachineResourcesChanged.
#rnx

diff --git a/Engine/Source/Editor/UnrealEd/Private/Cooker/CookDirector.cpp b/Engine/Source/Editor/UnrealEd/Private/Cooker/CookDirector.cpp
--- a/Engine/Source/Editor/UnrealEd/Private/Cooker/CookDirector.cpp
+++ b/Engine/Source/Editor/UnrealEd/Private/Cooker/CookDirector.cpp
@@ -483,25 +483,22 @@ void FCookDirector::ActivateMachineResourceReduction()
 		COTFS.MemoryMinFreeVirtual / 1024 / 1024, COTFS.MemoryMinFreePhysical / 1024 / 1024);
 
 	// Set CoreLimit for updating workerthreads in this process and passing to the commandline for workers
-	// CoreLimit changes both NumberOfCoresIncludingHyperthreads() and NumberOfCores to return the max 
 	int32 NumProcesses = RequestedCookWorkerCount + 1;
 	int32 NumberOfCores = FPlatformMisc::NumberOfCores();
 	int32 HyperThreadCount = FPlatformMisc::NumberOfCoresIncludingHyperthreads();
 	int32 NumberOfHyperThreadsPerCore = HyperThreadCount / NumberOfCores;
 	CoreLimit = FMath::Max(NumberOfCores / NumProcesses, 1);
+	int32 CoreIncludingHyperthreadsLimit = CoreLimit * NumberOfHyperThreadsPerCore;
 	int32 NumberOfWorkers = FMath::Max(CoreLimit - 1, 1) * NumberOfHyperThreadsPerCore;
 
 	// Update the number of Cores and WorkerThreads for this process
 	check(IsInGameThread());
-	// MPCOOKTODO: UnComment when SetCoreLimit is implemented
-	//FPlatformMisc::SetCoreLimit(CoreLimit);
 	int32 NumBackgroundWorkers = FMath::Max(1, NumberOfWorkers - FMath::Min<int32>(GNumForegroundWorkers, NumberOfWorkers));
 	int32 NumForegroundWorkers = FMath::Max(1, NumberOfWorkers - NumBackgroundWorkers);
 	LowLevelTasks::FScheduler::Get().RestartWorkers(NumForegroundWorkers, NumBackgroundWorkers);
 
 	// Update the number of ShaderCompilerWorkers that can be launched
-	// MPCOOKTODO: UnComment once OnMachineResourcesChanges is implemented
-	// GShaderCompilingManager->OnMachineResourcesChanged();
+	GShaderCompilingManager->OnMachineResourcesChanged(CoreLimit, CoreIncludingHyperthreadsLimit);
 
 	UE_LOG(LogCook, Display, TEXT("CookMultiprocess changed number of cores from %d to %d."),
 		NumberOfCores, FPlatformMisc::NumberOfCores());
@@ -646,7 +643,7 @@ FString FCookDirector::GetWorkerCommandLine(FWorkerId WorkerId)
 				Token.StartsWith(TEXT("-CookWorkerId=")) ||
 				Token.StartsWith(TEXT("-ShowCookWorker")) ||
 				Token.StartsWith(TEXT("-CoreLimit")) ||
-				Token.StartsWith(TEXT("-CoreLimitCores")) ||
+				Token.StartsWith(TEXT("-PhysicalCoreLimit")) ||
 				Token.StartsWith(TEXT("-CoreLimitHyperThreads"))
 				)
 			{
@@ -665,8 +662,7 @@ FString FCookDirector::GetWorkerCommandLine(FWorkerId WorkerId)
 	Tokens.Add(FString::Printf(TEXT("-CookWorkerId=%d"), WorkerId.GetRemoteIndex()));
 	if (CoreLimit > 0)
 	{
-		// MPCOOKTODO: Replace CoreLimit with CoreLimitCores once it is implemented
-		Tokens.Add(FString::Printf(TEXT("-CoreLimit=%d"), CoreLimit));
+		Tokens.Add(FString::Printf(TEXT("-PhysicalCoreLimit=%d"), CoreLimit));
 	}
 
 	return FString::Join(Tokens, TEXT(" "));
diff --git a/Engine/Source/Runtime/Core/Private/GenericPlatform/GenericPlatformMisc.cpp b/Engine/Source/Runtime/Core/Private/GenericPlatform/GenericPlatformMisc.cpp
--- a/Engine/Source/Runtime/Core/Private/GenericPlatform/GenericPlatformMisc.cpp
+++ b/Engine/Source/Runtime/Core/Private/GenericPlatform/GenericPlatformMisc.cpp
@@ -1514,11 +1514,61 @@ bool FGenericPlatformMisc::AllowThreadHeartBeat()
 	return bHeartbeat;
 }
 
+int32 FGenericPlatformMisc::NumberOfCores()
+{
+	return 1;
+}
+
 int32 FGenericPlatformMisc::NumberOfCoresIncludingHyperthreads()
 {
 	return FPlatformMisc::NumberOfCores();
 }
 
+void FGenericPlatformMisc::GetConfiguredCoreLimits(int32 PlatformNumPhysicalCores, int32 PlatformNumLogicalCores,
+	bool& bOutFullyInitialized, int32& OutPhysicalCoreLimit, int32& OutLogicalCoreLimit,
+	bool& bOutSetPhysicalCountToLogicalCount)
+{
+	// If CommandLine is not yet initialized, silently return default values. Callers will need to handle calling again.
+	if (!FCommandLine::IsInitialized())
+	{
+		bOutFullyInitialized = false;
+		OutPhysicalCoreLimit = 0;
+		OutLogicalCoreLimit = 0;
+		bOutSetPhysicalCountToLogicalCount = false;
+		return;
+	}
+
+	int32 PhysicalCoreLimit = 0;
+	int32 LogicalCoreLimit = 0;
+	int32 LegacyCoreLimit = 0;
+	bool bSetPhysicalCountToLogicalCount = false;
+
+	const TCHAR* CommandLine = FCommandLine::Get();
+	FParse::Value(CommandLine, TEXT("-physicalcorelimit="), PhysicalCoreLimit); // DEPRECATION_WARNING: physicalcorelimit is experimental and may be changed in a future release without deprecation
+	FParse::Value(CommandLine, TEXT("-corelimit="), LegacyCoreLimit);
+	bSetPhysicalCountToLogicalCount = FParse::Param(CommandLine, TEXT("usehyperthreading"));
+	if (bSetPhysicalCountToLogicalCount)
+	{
+		LogicalCoreLimit = PhysicalCoreLimit;
+	}
+	else
+	{
+		LogicalCoreLimit = PlatformNumPhysicalCores > 0 ?
+			(PhysicalCoreLimit * PlatformNumLogicalCores) / PlatformNumPhysicalCores :
+			PhysicalCoreLimit;
+	}
+	if (LegacyCoreLimit > 0)
+	{
+		PhysicalCoreLimit = FMath::Min(PhysicalCoreLimit, LegacyCoreLimit);
+		LogicalCoreLimit = FMath::Min(LogicalCoreLimit, LegacyCoreLimit);
+	}
+
+	bOutFullyInitialized = true;
+	OutPhysicalCoreLimit = PhysicalCoreLimit;
+	OutLogicalCoreLimit = LogicalCoreLimit;
+	bOutSetPhysicalCountToLogicalCount = bSetPhysicalCountToLogicalCount;
+}
+
 FProcessorGroupDesc InternalGetProcessorGroupDesc()
 {
 	FProcessorGroupDesc Desc;
diff --git a/Engine/Source/Runtime/Core/Private/Windows/WindowsPlatformMisc.cpp b/Engine/Source/Runtime/Core/Private/Windows/WindowsPlatformMisc.cpp
--- a/Engine/Source/Runtime/Core/Private/Windows/WindowsPlatformMisc.cpp
+++ b/Engine/Source/Runtime/Core/Private/Windows/WindowsPlatformMisc.cpp
@@ -1766,30 +1766,30 @@ static void QueryCpuInformation(FProcessorGroupDesc& OutGroupDesc, uint32& OutNu
 int32 FWindowsPlatformMisc::NumberOfCores()
 {
 	static int32 CoreCount = 0;
-	if (CoreCount == 0)
+	if (CoreCount > 0)
 	{
-		FProcessorGroupDesc GroupDesc;
-		uint32 NumaNodeCount = 0;
-		uint32 NumCores = 0;
-		uint32 LogicalProcessorCount = 0;
-		QueryCpuInformation(GroupDesc, NumaNodeCount, NumCores, LogicalProcessorCount);
+		return CoreCount;
+	}
 
-		if (FCommandLine::IsInitialized() && FParse::Param(FCommandLine::Get(), TEXT("usehyperthreading")))
-		{
-			CoreCount = LogicalProcessorCount;
-		}
-		else
-		{
-			CoreCount = NumCores;
-		}
+	FProcessorGroupDesc GroupDesc;
+	uint32 NumaNodeCount = 0;
+	uint32 NumCores = 0;
+	uint32 LogicalProcessorCount = 0;
+	QueryCpuInformation(GroupDesc, NumaNodeCount, NumCores, LogicalProcessorCount);
 
-		// Optionally limit number of threads (we don't necessarily scale super well with very high core counts)
+	bool bLimitsInitialized;
+	int32 CoreLimitCores;
+	int32 CoreLimitHyperThreads;
+	bool bSetCoreCountToHyperThreadCount;
+	GetConfiguredCoreLimits(NumCores, LogicalProcessorCount, bLimitsInitialized, CoreLimitCores,
+		CoreLimitHyperThreads, bSetCoreCountToHyperThreadCount);
 
-		int32 LimitCount = 32768;
-		if (FCommandLine::IsInitialized() && FParse::Value(FCommandLine::Get(), TEXT("-corelimit="), LimitCount))
-		{
-			CoreCount = FMath::Min(CoreCount, LimitCount);
-		}
+	CoreCount = bSetCoreCountToHyperThreadCount ? LogicalProcessorCount : NumCores;
+
+	// Optionally limit number of threads (we don't necessarily scale super well with very high core counts)
+	if (CoreLimitCores > 0)
+	{
+		CoreCount = FMath::Min(CoreCount, CoreLimitCores);
 	}
 
 	return CoreCount;
@@ -1814,23 +1814,30 @@ const FProcessorGroupDesc& FWindowsPlatformMisc::GetProcessorGroupDesc()
 int32 FWindowsPlatformMisc::NumberOfCoresIncludingHyperthreads()
 {
 	static int32 CoreCount = 0;
-	if (CoreCount == 0)
+	if (CoreCount > 0)
 	{
-		FProcessorGroupDesc GroupDesc;
-		uint32 NumaNodeCount = 0;
-		uint32 NumCores = 0;
-		uint32 LogicalProcessorCount = 0;
-		QueryCpuInformation(GroupDesc, NumaNodeCount, NumCores, LogicalProcessorCount);
+		return CoreCount;
+	}
+
+	FProcessorGroupDesc GroupDesc;
+	uint32 NumaNodeCount = 0;
+	uint32 NumCores = 0;
+	uint32 LogicalProcessorCount = 0;
+	QueryCpuInformation(GroupDesc, NumaNodeCount, NumCores, LogicalProcessorCount);
 
-		CoreCount = LogicalProcessorCount;
+	bool bLimitsInitialized;
+	int32 CoreLimitCores;
+	int32 CoreLimitHyperThreads;
+	bool bSetCoreCountToHyperThreadCount;
+	GetConfiguredCoreLimits(NumCores, LogicalProcessorCount, bLimitsInitialized, CoreLimitCores,
+		CoreLimitHyperThreads, bSetCoreCountToHyperThreadCount);
 
-		// Optionally limit number of threads (we don't necessarily scale super well with very high core counts)
+	CoreCount = LogicalProcessorCount;
 
-		int32 LimitCount = 32768;
-		if (FCommandLine::IsInitialized() && FParse::Value(FCommandLine::Get(), TEXT("-corelimit="), LimitCount))
-		{
-			CoreCount = FMath::Min(CoreCount, LimitCount);
-		}
+	// Optionally limit number of threads (we don't necessarily scale super well with very high core counts)
+	if (CoreLimitHyperThreads > 0)
+	{
+		CoreCount = FMath::Min(CoreCount, CoreLimitHyperThreads);
 	}
 
 	return CoreCount;
diff --git a/Engine/Source/Runtime/Core/Public/GenericPlatform/GenericPlatformMisc.h b/Engine/Source/Runtime/Core/Public/GenericPlatform/GenericPlatformMisc.h
--- a/Engine/Source/Runtime/Core/Public/GenericPlatform/GenericPlatformMisc.h
+++ b/Engine/Source/Runtime/Core/Public/GenericPlatform/GenericPlatformMisc.h
@@ -1095,10 +1095,7 @@ public:
 	/**
 	 * return the number of hardware CPU cores
 	 */
-	static int32 NumberOfCores()
-	{
-		return 1;
-	}
+	static int32 NumberOfCores();
 
 	/**
 	* @return a description of all the groups in the current system and their affinities
@@ -1833,6 +1830,23 @@ protected:
 
 	static EDeviceScreenOrientation AllowedDeviceOrientation;
 
+protected:
+	/**
+	 * Parse CoreLimit from CommandLine or from the SetCoreLimit value. The real number of physical cores and logical
+	 * cores must be provided so that CommandLine arguments can be translated.
+	 * 
+	 * @param bOutFullyInitialized True if commandline was available. If false, values are set to defaults.
+	 * @param OutPhysicalCoreLimit The maximum number of physical cores to report in NumberOfCores.
+	 *        Non-negative. 0 indicates unlimited.
+	 * @param OutLogicalCoreLimit The maximum number of logical cores to report in NumberOfCoresIncludingHyperthreads.
+	 *        Non-negative. 0 indicates unlimited.
+	 * @param bOutSetPhysicalCountToLogicalCount Whether NumberOfCores should be overridden to return
+	 *        NumberOfCoresIncludingHyperthreads.
+	 */
+	static void GetConfiguredCoreLimits(int32 PlatformNumPhysicalCores, int32 PlatformNumLogicalCores,
+		bool& bOutFullyInitialized, int32& OutPhysicalCoreLimit, int32& OutLogicalCoreLimit,
+		bool& bOutSetPhysicalCountToLogicalCount);
+
 private:
 	struct FStaticData;
 };
diff --git a/Engine/Source/Runtime/Engine/Private/ShaderCompiler/ShaderCompiler.cpp b/Engine/Source/Runtime/Engine/Private/ShaderCompiler/ShaderCompiler.cpp
--- a/Engine/Source/Runtime/Engine/Private/ShaderCompiler/ShaderCompiler.cpp
+++ b/Engine/Source/Runtime/Engine/Private/ShaderCompiler/ShaderCompiler.cpp
@@ -3703,7 +3703,7 @@ FShaderCompilingManager::FShaderCompilingManager() :
 	FPaths::NormalizeDirectoryName(AbsoluteDebugInfoDirectory);
 	AbsoluteShaderDebugInfoDirectory = AbsoluteDebugInfoDirectory;
 
-	CalculateNumberOfCompilingThreads();
+	CalculateNumberOfCompilingThreads(FPlatformMisc::NumberOfCores(), FPlatformMisc::NumberOfCoresIncludingHyperthreads());
 
 	TUniquePtr<FShaderCompileThreadRunnableBase> RemoteCompileThread;
 	const bool bCanUseRemoteCompiling = bAllowCompilingThroughWorkers && ShaderCompiler::IsRemoteCompilingAllowed() && AllTargetPlatformSupportsRemoteShaderCompiling();
@@ -3771,9 +3771,9 @@ FShaderCompilingManager::~FShaderCompilingManager()
 	FAssetCompilingManager::Get().UnregisterManager(this);
 }
 
-void FShaderCompilingManager::CalculateNumberOfCompilingThreads()
+void FShaderCompilingManager::CalculateNumberOfCompilingThreads(int32 NumberOfCores, int32 NumberOfCoresIncludingHyperthreads)
 {
-	const int32 NumVirtualCores = FPlatformMisc::NumberOfCoresIncludingHyperthreads();
+	const int32 NumVirtualCores = NumberOfCoresIncludingHyperthreads;
 
 	int32 NumUnusedShaderCompilingThreads;
 	verify(GConfig->GetInt(TEXT("DevOptions.Shaders"), TEXT("NumUnusedShaderCompilingThreads"), NumUnusedShaderCompilingThreads, GEngineIni));
@@ -3856,7 +3856,7 @@ void FShaderCompilingManager::CalculateNumberOfCompilingThreads()
 
 				bool bUseVirtualCores = true;
 				GConfig->GetBool(TEXT("DevOptions.Shaders"), TEXT("bUseVirtualCores"), bUseVirtualCores, GEngineIni);
-				uint32 MaxNumCoresToUse = bUseVirtualCores ? NumVirtualCores : FPlatformMisc::NumberOfCores();
+				uint32 MaxNumCoresToUse = bUseVirtualCores ? NumVirtualCores : NumberOfCores;
 				NumShaderCompilingThreads = FMath::Clamp<uint32>(NumShaderCompilingThreads, 1, MaxNumCoresToUse - 1);
 				NumShaderCompilingThreadsDuringGame = FMath::Min<int32>(NumShaderCompilingThreads, NumShaderCompilingThreadsDuringGame);
 			}
@@ -3880,9 +3880,9 @@ void FShaderCompilingManager::CalculateNumberOfCompilingThreads()
 	NumShaderCompilingThreadsDuringGame = FMath::Min<int32>(NumShaderCompilingThreadsDuringGame, NumShaderCompilingThreads);
 }
 
-void FShaderCompilingManager::OnMachineResourcesChanged()
+void FShaderCompilingManager::OnMachineResourcesChanged(int32 NumberOfCores, int32 NumberOfCoresIncludingHyperthreads)
 {
-	CalculateNumberOfCompilingThreads();
+	CalculateNumberOfCompilingThreads(NumberOfCores, NumberOfCoresIncludingHyperthreads);
 	for (const auto& Thread : Threads)
 	{
 		Thread->OnMachineResourcesChanged();
diff --git a/Engine/Source/Runtime/Engine/Public/ShaderCompiler.h b/Engine/Source/Runtime/Engine/Public/ShaderCompiler.h
--- a/Engine/Source/Runtime/Engine/Public/ShaderCompiler.h
+++ b/Engine/Source/Runtime/Engine/Public/ShaderCompiler.h
@@ -935,7 +935,7 @@ private:
 	FAsyncCompilationNotification Notification;
 
 	/** Calculate NumShaderCompilingThreads, during construction or OnMachineResourcesChanged */
-	void CalculateNumberOfCompilingThreads();
+	void CalculateNumberOfCompilingThreads(int32 NumberOfCores, int32 NumberOfCoresIncludingHyperthreads);
 
 	/** Launches the worker, returns the launched process handle. */
 	FProcHandle LaunchWorker(const FString& WorkingDirectory, uint32 ProcessId, uint32 ThreadId, const FString& WorkerInputFile, const FString& WorkerOutputFile);
@@ -976,7 +976,7 @@ public:
 	ENGINE_API ~FShaderCompilingManager();
 
 	/** Called by external systems that have updated the number of worker threads available. */
-	ENGINE_API void OnMachineResourcesChanged();
+	ENGINE_API void OnMachineResourcesChanged(int32 NumberOfCores, int32 NumberOfCoresIncludingHyperthreads);
 
 	ENGINE_API int32 GetNumPendingJobs() const;
 	ENGINE_API int32 GetNumOutstandingJobs() const;
