Fix incorrect statistics being reported when virtualizing payloads. We no longer count a payload that was already stored in a backend as being uploaded there.

#rnx

- Removed the status enum from FPushRequest and replaced it with a class that represents the push result instead.
-- This lets us return more context about problems. At the moment it is only being used to return the filter reason if a payload is filtered out. In the future we could return per payload error messages etc.
- Now that we return more detailed info about the payload push we no longer report payloads that were already in the backend as being pushed which could cause some very misleading stats to be shown.
- Changed the backends to use the new results system.
-- Not all paths were correctly setting the state in the file system backend.
-- The source control backend does not set the result in all paths, which will leave the results in the "pending" state. Which will still mark the payload as not uploaded and so is safe to do.
- Removed the pushing of a single method from IVirtualizationBackend as it was only used in one place internally, so we might as well just change that code to call the batch push overload.

diff --git a/Engine/Source/Developer/Virtualization/Private/PackageVirtualizationProcess.cpp b/Engine/Source/Developer/Virtualization/Private/PackageVirtualizationProcess.cpp
--- a/Engine/Source/Developer/Virtualization/Private/PackageVirtualizationProcess.cpp
+++ b/Engine/Source/Developer/Virtualization/Private/PackageVirtualizationProcess.cpp
@@ -278,11 +278,11 @@ void VirtualizePackages(const TArray<FString>& FilesToSubmit, TArray<FText>& Out
 		int64 TotalPayloadsCached = 0;
 		for (Virtualization::FPushRequest& Request : PayloadsToSubmit)
 		{
-			TotalPayloadsCached += Request.GetStatus() == FPushRequest::EStatus::Success ? 1 : 0;
+			TotalPayloadsCached += Request.GetResult().WasPushed() ? 1 : 0;
 
 			// TODO: This really shouldn't be required, fix when we allow both pushes to be done in the same call
 			// Reset the status for the persistent storage push
-			Request.SetStatus(Virtualization::FPushRequest::EStatus::Failed);
+			Request.ResetResult();
 		}
 		UE_LOG(LogVirtualization, Display, TEXT("Pushed %" INT64_FMT " payload(s) to cached storage"), TotalPayloadsCached);
 	}
@@ -301,7 +301,7 @@ void VirtualizePackages(const TArray<FString>& FilesToSubmit, TArray<FText>& Out
 		int64 TotalPayloadsVirtualized = 0;
 		for (const Virtualization::FPushRequest& Request : PayloadsToSubmit)
 		{
-			TotalPayloadsVirtualized += Request.GetStatus() == FPushRequest::EStatus::Success ? 1 : 0;
+			TotalPayloadsVirtualized += Request.GetResult().WasPushed() ? 1 : 0;
 		}
 		UE_LOG(LogVirtualization, Display, TEXT("Pushed %" INT64_FMT " payload(s) to EStorageType::Persistent storage"), TotalPayloadsVirtualized);
 	}
@@ -314,7 +314,7 @@ void VirtualizePackages(const TArray<FString>& FilesToSubmit, TArray<FText>& Out
 			const Virtualization::FPushRequest& Request = PayloadsToSubmit[PackageInfo.PayloadIndex + Index];
 			check(Request.GetIdentifier() == PackageInfo.LocalPayloads[Index]);
 
-			if (Request.GetStatus() == Virtualization::FPushRequest::EStatus::Success)
+			if (Request.GetResult().IsVirtualized())
 			{
 				if (PackageInfo.Trailer.UpdatePayloadAsVirtualized(Request.GetIdentifier()))
 				{
diff --git a/Engine/Source/Developer/Virtualization/Private/VirtualizationDDCBackend.cpp b/Engine/Source/Developer/Virtualization/Private/VirtualizationDDCBackend.cpp
--- a/Engine/Source/Developer/Virtualization/Private/VirtualizationDDCBackend.cpp
+++ b/Engine/Source/Developer/Virtualization/Private/VirtualizationDDCBackend.cpp
@@ -91,7 +91,7 @@ bool FDDCBackend::PushData(TArrayView<FPushRequest> Requests)
 	{
 		if (DoesPayloadExist(Request.GetIdentifier()))
 		{
-			Request.SetStatus(FPushRequest::EStatus::Success);
+			Request.SetResult(FPushResult::GetAsAlreadyExists());
 		}
 		else
 		{
@@ -112,11 +112,11 @@ bool FDDCBackend::PushData(TArrayView<FPushRequest> Requests)
 			{
 				if (Response.Status == UE::DerivedData::EStatus::Ok)
 				{
-					Request.SetStatus(FPushRequest::EStatus::Success);
+					Request.SetResult(FPushResult::GetAsPushed());
 				}
 				else
 				{
-					Request.SetStatus(FPushRequest::EStatus::Failed);
+					Request.SetResult(FPushResult::GetAsError());
 					bWasSuccess = false;
 				}
 			};
diff --git a/Engine/Source/Developer/Virtualization/Private/VirtualizationFileBackend.cpp b/Engine/Source/Developer/Virtualization/Private/VirtualizationFileBackend.cpp
--- a/Engine/Source/Developer/Virtualization/Private/VirtualizationFileBackend.cpp
+++ b/Engine/Source/Developer/Virtualization/Private/VirtualizationFileBackend.cpp
@@ -71,7 +71,7 @@ bool FFileSystemBackend::PushData(TArrayView<FPushRequest> Requests)
 		if (DoesPayloadExist(PayloadId))
 		{
 			UE_LOG(LogVirtualization, Verbose, TEXT("[%s] Already has a copy of the payload '%s'."), *GetDebugName(), *LexToString(PayloadId));
-			Request.SetStatus(FPushRequest::EStatus::Success);
+			Request.SetResult(FPushResult::GetAsAlreadyExists());
 
 			continue;
 		}
@@ -97,7 +97,7 @@ bool FFileSystemBackend::PushData(TArrayView<FPushRequest> Requests)
 				SystemErrorMsg.ToString());
 
 			ErrorCount++;
-			Request.SetStatus(FPushRequest::EStatus::Failed);
+			Request.SetResult(FPushResult::GetAsError());
 
 			continue;
 		}
@@ -122,7 +122,7 @@ bool FFileSystemBackend::PushData(TArrayView<FPushRequest> Requests)
 			IFileManager::Get().Delete(*TempFilePath, true, false, true);  // Clean up the temp file if it is still around but do not failure cases to the user
 
 			ErrorCount++;
-			Request.SetStatus(FPushRequest::EStatus::Failed);
+			Request.SetResult(FPushResult::GetAsError());
 
 			continue;
 		}
@@ -131,7 +131,11 @@ bool FFileSystemBackend::PushData(TArrayView<FPushRequest> Requests)
 		CreateFilePath(PayloadId, FilePath);
 
 		// If the file already exists we don't need to replace it, we will also do our own error logging.
-		if (!IFileManager::Get().Move(FilePath.ToString(), *TempFilePath, /*Replace*/ false, /*EvenIfReadOnly*/ false, /*Attributes*/ false, /*bDoNotRetryOrError*/ true))
+		if (IFileManager::Get().Move(FilePath.ToString(), *TempFilePath, /*Replace*/ false, /*EvenIfReadOnly*/ false, /*Attributes*/ false, /*bDoNotRetryOrError*/ true))
+		{
+			Request.SetResult(FPushResult::GetAsPushed());
+		}
+		else
 		{
 			// Store the error message in case we need to display it
 			TStringBuilder<MAX_SPRINTF> SystemErrorMsg;
@@ -144,6 +148,7 @@ bool FFileSystemBackend::PushData(TArrayView<FPushRequest> Requests)
 			if (DoesPayloadExist(PayloadId))
 			{
 				UE_LOG(LogVirtualization, Verbose, TEXT("[%s] Already has a copy of the payload '%s'."), *GetDebugName(), *LexToString(PayloadId));
+				Request.SetResult(FPushResult::GetAsAlreadyExists());
 				continue;
 			}
 			else
@@ -155,7 +160,7 @@ bool FFileSystemBackend::PushData(TArrayView<FPushRequest> Requests)
 					SystemErrorMsg.ToString());
 
 				ErrorCount++;
-				Request.SetStatus(FPushRequest::EStatus::Failed);
+				Request.SetResult(FPushResult::GetAsError());
 			}
 		}
 	}
diff --git a/Engine/Source/Developer/Virtualization/Private/VirtualizationManager.cpp b/Engine/Source/Developer/Virtualization/Private/VirtualizationManager.cpp
--- a/Engine/Source/Developer/Virtualization/Private/VirtualizationManager.cpp
+++ b/Engine/Source/Developer/Virtualization/Private/VirtualizationManager.cpp
@@ -439,7 +439,7 @@ bool FVirtualizationManager::PushData(TArrayView<FPushRequest> Requests, EStorag
 		FPushRequest& Request = Requests[Index];
 		if (Request.GetIdentifier().IsZero() || Request.GetPayloadSize() == 0)
 		{
-			Request.SetStatus(FPushRequest::EStatus::Invalid);
+			Request.SetResult(FPushResult::GetAsInvalid());
 			continue;
 		}
 
@@ -451,7 +451,7 @@ bool FVirtualizationManager::PushData(TArrayView<FPushRequest> Requests, EStorag
 					Request.GetPayloadSize(),
 					MinPayloadLength);
 
-			Request.SetStatus(FPushRequest::EStatus::BelowMinSize);
+			Request.SetResult(FPushResult::GetAsFiltered(EPayloadFilterReason::MinSize));
 			continue;
 		}
 
@@ -461,7 +461,7 @@ bool FVirtualizationManager::PushData(TArrayView<FPushRequest> Requests, EStorag
 					*LexToString(Request.GetIdentifier()), 
 					*Request.GetContext());
 			
-			Request.SetStatus(FPushRequest::EStatus::ExcludedByPackagPath);
+			Request.SetResult(FPushResult::GetAsFiltered(EPayloadFilterReason::Path));
 			continue;
 		}
 
@@ -533,7 +533,7 @@ bool FVirtualizationManager::PushData(TArrayView<FPushRequest> Requests, EStorag
 		const int32 MappingIndex = OriginalToValidatedRequest[Index];
 		if (MappingIndex != INDEX_NONE)
 		{
-			Requests[Index].SetStatus(ValidatedRequests[MappingIndex].GetStatus());
+			Requests[Index].SetResult(ValidatedRequests[MappingIndex].GetResult());
 		}
 	}
 
@@ -1524,9 +1524,15 @@ bool FVirtualizationManager::TryCacheDataToBackend(IVirtualizationBackend& Backe
 {
 	COOK_STAT(FCookStats::FScopedStatsCounter Timer(Profiling::GetCacheStats(Backend)));
 	
-	if (Backend.PushData(Id, Payload, FString()))
+	FPushRequest Request(Id, Payload, FString());
+
+	if (Backend.PushData(MakeArrayView(&Request, 1)))
 	{
-		COOK_STAT(Timer.AddHit(Payload.GetCompressedSize()));
+		if (Request.GetResult().WasPushed())
+		{
+			COOK_STAT(Timer.AddHit(Payload.GetCompressedSize()));
+		}
+
 		return true;
 	}
 	else
@@ -1552,8 +1558,7 @@ bool FVirtualizationManager::TryPushDataToBackend(IVirtualizationBackend& Backen
 
 		for (const FPushRequest& Request : Requests)
 		{
-			// TODO: Don't add a hit if the payload was already uploaded
-			if (Request.GetStatus() == FPushRequest::EStatus::Success)
+			if (Request.GetResult().WasPushed())
 			{
 				Stats.Accumulate(FCookStats::CallStats::EHitOrMiss::Hit, FCookStats::CallStats::EStatType::Counter, 1l, bIsInGameThread);	
 				Stats.Accumulate(FCookStats::CallStats::EHitOrMiss::Hit, FCookStats::CallStats::EStatType::Bytes, Request.GetPayloadSize(), bIsInGameThread);
diff --git a/Engine/Source/Developer/Virtualization/Private/VirtualizationSourceControlBackend.cpp b/Engine/Source/Developer/Virtualization/Private/VirtualizationSourceControlBackend.cpp
--- a/Engine/Source/Developer/Virtualization/Private/VirtualizationSourceControlBackend.cpp
+++ b/Engine/Source/Developer/Virtualization/Private/VirtualizationSourceControlBackend.cpp
@@ -433,7 +433,7 @@ bool FSourceControlBackend::PushData(TArrayView<FPushRequest> Requests)
 	// We create a new array of requests that we know are unique along with a map of the request 
 	// results that we can use to set the original requests statuses at the end.
 
-	TMap<FIoHash, FPushRequest::EStatus> PayloadStatusMap;
+	TMap<FIoHash, FPushResult> PayloadStatusMap;
 	PayloadStatusMap.Reserve(Requests.Num());
 
 	TArray<const FPushRequest*> UniqueRequests;
@@ -446,7 +446,7 @@ bool FSourceControlBackend::PushData(TArrayView<FPushRequest> Requests)
 		{
 			if (!PayloadStatusMap.Contains(Request.GetIdentifier()))
 			{
-				PayloadStatusMap.Add(Request.GetIdentifier(), FPushRequest::EStatus::Failed);
+				PayloadStatusMap.Add(Request.GetIdentifier(), FPushResult::GetAsError());
 				UniqueRequests.Add(&Request);
 			}
 		}
@@ -495,7 +495,7 @@ bool FSourceControlBackend::PushData(TArrayView<FPushRequest> Requests)
 
 				if (PayloadResults[Index])
 				{
-					PayloadStatusMap[Request->GetIdentifier()] = FPushRequest::EStatus::Success;
+					PayloadStatusMap[Request->GetIdentifier()] = FPushResult::GetAsAlreadyExists();
 				}
 				else
 				{
@@ -511,7 +511,7 @@ bool FSourceControlBackend::PushData(TArrayView<FPushRequest> Requests)
 	{
 		for (FPushRequest& Request : Requests)
 		{
-			Request.SetStatus(PayloadStatusMap[Request.GetIdentifier()]);
+			Request.SetResult(PayloadStatusMap[Request.GetIdentifier()]);
 		}
 
 		return true;
@@ -731,10 +731,9 @@ bool FSourceControlBackend::PushData(TArrayView<FPushRequest> Requests)
 			}
 		}
 
-		// TODO: We really should be setting a more fine grain status for each request, or not bother with the status at all
 		for (const FPushRequest* Request : RequestBatch)
 		{
-			PayloadStatusMap[Request->GetIdentifier()] = FPushRequest::EStatus::Success;
+			PayloadStatusMap[Request->GetIdentifier()] = FPushResult::GetAsPushed();
 		}
 
 		// Try to clean up the files from this batch
@@ -751,7 +750,7 @@ bool FSourceControlBackend::PushData(TArrayView<FPushRequest> Requests)
 	// Finally set all of the request statuses
 	for (FPushRequest& Request : Requests)
 	{
-		Request.SetStatus(PayloadStatusMap[Request.GetIdentifier()]);
+		Request.SetResult(PayloadStatusMap[Request.GetIdentifier()]);
 	}
 
 	return true;
diff --git a/Engine/Source/Developer/Virtualization/Public/IVirtualizationBackend.h b/Engine/Source/Developer/Virtualization/Public/IVirtualizationBackend.h
--- a/Engine/Source/Developer/Virtualization/Public/IVirtualizationBackend.h
+++ b/Engine/Source/Developer/Virtualization/Public/IVirtualizationBackend.h
@@ -100,22 +100,6 @@ public:
 		return ConnectionStatus;
 	}
 
-	/**
-	 * The backend will attempt to store the given payload by what ever method the backend uses.
-	 * NOTE: It is assumed that the virtualization manager will run all appropriate validation
-	 * on the payload and it's id and that the inputs to PushData can be trusted.
-	 * 
-	 * @param Id		The Id of the payload
-	 * @param Payload	A potentially compressed buffer representing the payload
-	 * @return			The result of the push operation
-	 */
-	bool PushData(const FIoHash& Id, const FCompressedBuffer& Payload, const FString& PackageContext)
-	{
-		FPushRequest Request(Id, Payload, PackageContext);
-
-		return PushData(MakeArrayView(&Request, 1));
-	}
-
 	/**
 	 * @return True if all of the requests succeeded, false if one of more failed
 	 */
diff --git a/Engine/Source/Runtime/Core/Public/Virtualization/VirtualizationSystem.h b/Engine/Source/Runtime/Core/Public/Virtualization/VirtualizationSystem.h
--- a/Engine/Source/Runtime/Core/Public/Virtualization/VirtualizationSystem.h
+++ b/Engine/Source/Runtime/Core/Public/Virtualization/VirtualizationSystem.h
@@ -131,6 +131,111 @@ public:
 	virtual uint64 GetPayloadSize(const FIoHash& Identifier) = 0;
 };
 
+/** The result opf a push operation on a single payload */
+struct FPushResult
+{
+	/** 
+	 * Records the status of the push, negative values indicate an error, positive values indicate 
+	 * that the payload is stored in the target backend(s) and a zero value means that the operation 
+	 * has not yet run 
+	 */
+	enum class EStatus : int8
+	{
+		/** The push operation caused an error */
+		Error			= -3,
+		/** The push operation was run on an invalid payload id */
+		Invalid			= -2,
+		/** The payload was rejected by the filtering system */
+		Filtered		= -1,
+		/** The push operation has not yet been run */
+		Pending			= 0,
+		/** The payload was already present in the target backend(s) */
+		AlreadyExisted	= 1,
+		/** The payload was pushed to the backend(s) */
+		Pushed			= 2
+	};
+
+	FPushResult() = default;
+	~FPushResult() = default;
+
+	static FPushResult GetAsError()
+	{
+		return FPushResult(EStatus::Error);
+	}
+
+	static FPushResult GetAsInvalid()
+	{
+		return FPushResult(EStatus::Invalid);
+	}
+
+	static FPushResult GetAsFiltered(EPayloadFilterReason Reason)
+	{
+		return FPushResult(EStatus::Filtered, Reason);
+	}
+
+	static FPushResult GetAsAlreadyExists()
+	{
+		return FPushResult(EStatus::AlreadyExisted);
+	}
+
+	static FPushResult GetAsPushed()
+	{
+		return FPushResult(EStatus::Pushed);
+	}
+
+	/** Returns true if the payload was actually uploaded to the target backend(s) */
+	bool WasPushed() const
+	{
+		return Status == EStatus::Pushed;
+	}
+
+	/** Returns true if the payload is stored in the target backend(s) after the operation completed */
+	bool IsVirtualized() const
+	{
+		return Status > EStatus::Pending;
+	}
+
+	/** Returns true if the payload was rejected by the filtering system */
+	bool IsFiltered() const
+	{
+		if (Status == EStatus::Filtered)
+		{
+			// Quick sanity check, if the payload was filtered there must be a reason
+			check(FilterReason != EPayloadFilterReason::None);
+			return true;
+		}
+		else
+		{
+			return false;
+		}
+	}
+
+	/** Returns why the payload was rejected by the filtering system (if there was one) */
+	EPayloadFilterReason GetFilterReason() const
+	{
+		return FilterReason;
+	}
+
+private:
+
+	FPushResult(EStatus InStatus)
+		: Status(InStatus)
+	{
+		check(Status != EStatus::Pending);
+	}
+
+	FPushResult(EStatus InStatus, EPayloadFilterReason InReason)
+		: Status(InStatus)
+		, FilterReason(InReason)
+	{
+		check(Status != EStatus::Pending);
+		check(FilterReason != EPayloadFilterReason::None);
+	}
+
+	EStatus Status = EStatus::Pending;
+	EPayloadFilterReason FilterReason = EPayloadFilterReason::None;
+};
+
 /** 
  * Data structure representing a request to push a payload to a backend storage system. 
  * Note that a request can either before for payload already in memory (in which case the
@@ -139,20 +244,6 @@ public:
 */
 struct FPushRequest
 {
-	enum class EStatus
-	{
-		/** The request failed, or was not reached because of the failure of an earlier request */
-		Failed,
-		/** The payload does not have a valid identifier or is empty */
-		Invalid,
-		/** The payload is below the minimum length required for virtualization */
-		BelowMinSize,
-		/** The payload is owned by a package that is excluded from virtualization by path filtering */
-		ExcludedByPackagPath,
-		/** The payload in the request is now present in all backends */
-		Success
-	};
-
 	FPushRequest() = delete;
 	~FPushRequest() = default;
 
@@ -207,12 +298,6 @@ struct FPushRequest
 		return Identifier;
 	}
 
-	/** Returns the current status of the request */
-	EStatus GetStatus() const
-	{
-		return Status;
-	}
-
 	/** Returns the size of the payload when it was on disk */
 	uint64 GetPayloadSize() const
 	{
@@ -245,10 +330,19 @@ struct FPushRequest
 		return Context;
 	}
 
-	/** Allows the status of the request to be set, this should only be done by the virtualization backends */
-	void SetStatus(EStatus InStatus)
+	void ResetResult()
+	{
+		Result = FPushResult();
+	}
+
+	void SetResult(FPushResult InResult)
+	{
+		Result = InResult;
+	}
+
+	const FPushResult& GetResult() const
 	{
-		Status = InStatus;
+		return Result;
 	}
 
 private:
@@ -264,8 +358,7 @@ private:
 	/** A string containing context for the payload, typically a package name */
 	FString Context;
 
-	/** Once the request has been processed this value will contains the results */
-	EStatus Status = EStatus::Failed;
+	FPushResult Result;
 };
 
 /** 
