LLM: Fix bug in LLMMap Iterator that corrupted the tag of each allocation when tags are resorted.
#rnx

diff --git a/Engine/Source/Runtime/Core/Private/HAL/LowLevelMemoryUtils.h b/Engine/Source/Runtime/Core/Private/HAL/LowLevelMemoryUtils.h
--- a/Engine/Source/Runtime/Core/Private/HAL/LowLevelMemoryUtils.h
+++ b/Engine/Source/Runtime/Core/Private/HAL/LowLevelMemoryUtils.h
@@ -558,7 +558,8 @@ public:
 
 		bool operator!=(const FIterator& Other) const
 		{
-			return FBaseIterator::MapIndex != Other.FBaseIterator::MapIndex;
+			return FBaseIterator::StripeIndex != Other.FBaseIterator::StripeIndex || 
+				FBaseIterator::MapIndex != Other.FBaseIterator::MapIndex;
 		}
 
 		FTuple operator*() const
@@ -595,14 +596,16 @@ public:
 
 		bool operator!=(const FConstIterator& Other) const
 		{
-			return FBaseIterator::MapIndex != Other.FBaseIterator::MapIndex;
+			return FBaseIterator::StripeIndex != Other.FBaseIterator::StripeIndex || 
+				FBaseIterator::MapIndex != Other.FBaseIterator::MapIndex;
 		}
 
 		FConstTuple operator*() const
 		{
 			ThisType& LocalMap = FBaseIterator::MapRef;
-			SizeType KeyIndex = LocalMap.Map[FBaseIterator::MapIndex];
-			return FConstTuple(LocalMap.Keys[KeyIndex], LocalMap.Values1[KeyIndex], LocalMap.Values2[KeyIndex]);
+			StripeData& Stripe = LocalMap.MapStripes[FBaseIterator::StripeIndex];
+			SizeType KeyIndex = Stripe.Map[FBaseIterator::MapIndex];
+			return FConstTuple(Stripe.Keys[KeyIndex], Stripe.Values1[KeyIndex], Stripe.Values2[KeyIndex]);
 		}
 	};
 
