Moved IN and OUT defines to HAL/Platform.h

diff --git a/Engine/Source/Runtime/Core/Public/Android/AndroidSystemIncludes.h b/Engine/Source/Runtime/Core/Public/Android/AndroidSystemIncludes.h
--- a/Engine/Source/Runtime/Core/Public/Android/AndroidSystemIncludes.h
+++ b/Engine/Source/Runtime/Core/Public/Android/AndroidSystemIncludes.h
@@ -39,6 +39,3 @@ struct tagRECT
 	int32 bottom;
 };
 typedef struct tagRECT RECT;
-
-#define OUT
-#define IN
\ No newline at end of file
diff --git a/Engine/Source/Runtime/Core/Public/HAL/Platform.h b/Engine/Source/Runtime/Core/Public/HAL/Platform.h
--- a/Engine/Source/Runtime/Core/Public/HAL/Platform.h
+++ b/Engine/Source/Runtime/Core/Public/HAL/Platform.h
@@ -802,6 +802,12 @@
 #ifndef CONSTEXPR
 	#define CONSTEXPR constexpr
 #endif
+#ifndef IN
+	#define IN
+#endif
+#ifndef OUT
+	#define OUT
+#endif
 
 // String constants
 #ifndef LINE_TERMINATOR
diff --git a/Engine/Source/Runtime/Core/Public/IOS/IOSSystemIncludes.h b/Engine/Source/Runtime/Core/Public/IOS/IOSSystemIncludes.h
--- a/Engine/Source/Runtime/Core/Public/IOS/IOSSystemIncludes.h
+++ b/Engine/Source/Runtime/Core/Public/IOS/IOSSystemIncludes.h
@@ -69,9 +69,6 @@ struct tagRECT
 };
 typedef struct tagRECT RECT;
 
-#define OUT
-#define IN
-
 /*----------------------------------------------------------------------------
  Memory. On Mac OS X malloc allocates memory aligned to 16 bytes.
  ----------------------------------------------------------------------------*/
diff --git a/Engine/Source/Runtime/Core/Public/Mac/MacSystemIncludes.h b/Engine/Source/Runtime/Core/Public/Mac/MacSystemIncludes.h
--- a/Engine/Source/Runtime/Core/Public/Mac/MacSystemIncludes.h
+++ b/Engine/Source/Runtime/Core/Public/Mac/MacSystemIncludes.h
@@ -59,9 +59,6 @@ struct tagRECT
 };
 typedef struct tagRECT RECT;
 
-#define OUT
-#define IN
-
 /*----------------------------------------------------------------------------
 Memory. On Mac OS X malloc allocates memory aligned to 16 bytes.
 ----------------------------------------------------------------------------*/
diff --git a/Engine/Source/Runtime/Core/Public/Unix/UnixSystemIncludes.h b/Engine/Source/Runtime/Core/Public/Unix/UnixSystemIncludes.h
--- a/Engine/Source/Runtime/Core/Public/Unix/UnixSystemIncludes.h
+++ b/Engine/Source/Runtime/Core/Public/Unix/UnixSystemIncludes.h
@@ -1,11 +1,12 @@
 // Copyright Epic Games, Inc. All Rights Reserved.
-// IWYU pragma: begin_exports
 
 #pragma once
 
 // Set up compiler pragmas, etc
 
 #include "CoreTypes.h"
+
+// IWYU pragma: begin_exports
 #include "Unix/UnixPlatformCompilerSetup.h"
 
 #include <string.h>
@@ -48,9 +49,6 @@ struct tagRECT
 };
 typedef struct tagRECT RECT;
 
-#define OUT
-#define IN
-
 // In glibc 2.17, __secure_getenv was renamed to secure_getenv.
 #if defined(__GLIBC__) && (__GLIBC__ == 2) && (__GLIBC_MINOR__ < 17)
 #  define secure_getenv __secure_getenv
diff --git a/Engine/Source/Runtime/Core/Public/Windows/WindowsSystemIncludes.h b/Engine/Source/Runtime/Core/Public/Windows/WindowsSystemIncludes.h
--- a/Engine/Source/Runtime/Core/Public/Windows/WindowsSystemIncludes.h
+++ b/Engine/Source/Runtime/Core/Public/Windows/WindowsSystemIncludes.h
@@ -42,12 +42,3 @@ THIRD_PARTY_INCLUDES_START
 using namespace vc_attributes;
 #endif
 THIRD_PARTY_INCLUDES_END
-
-#ifndef OUT
-#define OUT
-#endif
-
-#ifndef IN
-#define IN
-#endif
-
