VectorIntMultiply SSE4.1

Add SSE4.1 support to VectorIntMultiply, so that it can use a single instruction to multiply two vectors of ints when SSE4.1 is available.

diff --git a/Engine/Source/Runtime/Core/Private/Tests/Math/UnrealMathTest.cpp b/Engine/Source/Runtime/Core/Private/Tests/Math/UnrealMathTest.cpp
--- a/Engine/Source/Runtime/Core/Private/Tests/Math/UnrealMathTest.cpp
+++ b/Engine/Source/Runtime/Core/Private/Tests/Math/UnrealMathTest.cpp
@@ -986,6 +986,20 @@ FORCENOINLINE void LogTest(const TCHAR* TestName, bool bHasPassed)
 	}
 }
 
+// Specialization for 'int32' so it checks the correct global state filled by the various comparision validation functions
+template<>
+FORCENOINLINE void LogTest<int32>(const TCHAR* TestName, bool bHasPassed)
+{
+	if (bHasPassed == false)
+	{
+		int32* GScratchI = static_cast<int32*>(static_cast<void*>(GScratch));
+		int32& GSumI = *static_cast<int32*>(static_cast<void*>(&GSum));
+		UE_LOG(LogUnrealMathTest, Warning, TEXT("%s <int32>: %s"), bHasPassed ? TEXT("PASSED") : TEXT("FAILED"), TestName);
+		UE_LOG(LogUnrealMathTest, Warning, TEXT("Bad(%d): (%d %d %d %d) (%d %d %d %d)"), GSumI, GScratchI[0], GScratchI[1], GScratchI[2], GScratchI[3], GScratchI[4], GScratchI[5], GScratchI[6], GScratchI[7]);
+		CheckPassing(false);
+	}
+}
+
 // Specialization for 'float' so it checks the correct global state filled by the various comparision validation functions
 template<>
 FORCENOINLINE void LogTest<float>(const TCHAR* TestName, bool bHasPassed)
@@ -2223,6 +2237,7 @@ bool FVectorRegisterAbstractionTest::RunTest(const FString& Parameters)
 	float F1 = 1.f;
 	uint32 U1 = *(uint32 *)&F1;
 	VectorRegister4Float V0, V1, V2, V3;
+	VectorRegister4Int VI0, VI1;
 	float Float0, Float1, Float2, Float3;
 
 	// Using a union as we need to do a bitwise cast of 0xFFFFFFFF into a float for NaN.
@@ -2338,6 +2353,13 @@ bool FVectorRegisterAbstractionTest::RunTest(const FString& Parameters)
 	V1 = MakeVectorRegisterFloat(Float0, Float1, Float2, Float3);
 	LogTest<float>(TEXT("VectorGetComponent<float>"), TestVectorsEqual(V0, V1));
 
+	//VectorIntMultiply
+	VI0 = MakeVectorRegisterInt(1, 2, -3, 4);
+	VI1 = MakeVectorRegisterInt(2, -4, -6, -8);
+	VI1 = VectorIntMultiply(VI0, VI1);
+	VI0 = MakeVectorRegisterInt(2, -8, 18, -32);
+	LogTest<int32>(TEXT("VectorIntMultiply"), TestVectorsEqualBitwise(VectorCastIntToFloat(VI0), VectorCastIntToFloat(VI1)));
+
 	V0 = MakeVectorRegister( 1.0f, -2.0f, 3.0f, -4.0f );
 	V1 = VectorAbs( V0 );
 	V0 = MakeVectorRegister( 1.0f, 2.0f, 3.0f, 4.0f );
diff --git a/Engine/Source/Runtime/Core/Public/Math/UnrealMathSSE.h b/Engine/Source/Runtime/Core/Public/Math/UnrealMathSSE.h
--- a/Engine/Source/Runtime/Core/Public/Math/UnrealMathSSE.h
+++ b/Engine/Source/Runtime/Core/Public/Math/UnrealMathSSE.h
@@ -3363,10 +3363,14 @@ FORCEINLINE VectorRegister4Int VectorIntSelect(const VectorRegister4Int& Mask, c
 
 FORCEINLINE VectorRegister4Int VectorIntMultiply(const VectorRegister4Int& A, const VectorRegister4Int& B)
 {
+#if UE_PLATFORM_MATH_USE_SSE4_1
+	return _mm_mullo_epi32(A, B);
+#else
 	//SSE2 doesn't have a multiply op for 4 32bit ints. Ugh.
 	__m128i Temp0 = _mm_mul_epu32(A, B);
 	__m128i Temp1 = _mm_mul_epu32(_mm_srli_si128(A, 4), _mm_srli_si128(B, 4));
 	return _mm_unpacklo_epi32(_mm_shuffle_epi32(Temp0, _MM_SHUFFLE(0, 0, 2, 0)), _mm_shuffle_epi32(Temp1, _MM_SHUFFLE(0, 0, 2, 0)));
+#endif
 }
 
 #define VectorIntNegate(A) VectorIntSubtract( GlobalVectorConstants::IntZero, A)
