VectorRoundToIntHalfToEven

Support for converting a vector of floats to a vector of ints with the round-to-nearest-or-even rounding mode.

diff --git a/Engine/Source/Runtime/Core/Private/Tests/Math/UnrealMathTest.cpp b/Engine/Source/Runtime/Core/Private/Tests/Math/UnrealMathTest.cpp
--- a/Engine/Source/Runtime/Core/Private/Tests/Math/UnrealMathTest.cpp
+++ b/Engine/Source/Runtime/Core/Private/Tests/Math/UnrealMathTest.cpp
@@ -3756,6 +3756,142 @@ bool FInterpolationFunctionTests::RunTest(const FString&)
 	return true;
 }
 
+class FMathVectorTestParameter
+{
+public:
+	FMathVectorTestParameter(float F) : Vec(VectorSet1(F))
+	{
+	}
+
+	FMathVectorTestParameter(int I) : Vec(VectorCastIntToFloat(VectorIntSet1(I)))
+	{
+	}
+
+	FMathVectorTestParameter(const VectorRegister4Float& F) : Vec(F)
+	{
+	}
+
+	FMathVectorTestParameter(const VectorRegister4Int& I) : Vec(VectorCastIntToFloat(I))
+	{
+	}
+
+	friend bool operator ==(const FMathVectorTestParameter& lhs, const FMathVectorTestParameter& rhs)
+	{
+		return TestVectorsEqualBitwise(lhs.Vec, rhs.Vec);
+	}
+
+	friend bool operator !=(const FMathVectorTestParameter& lhs, const FMathVectorTestParameter& rhs)
+	{
+		return !(lhs == rhs);
+	}
+
+private:
+	VectorRegister4Float Vec;
+};
+
+template<class VectorRoundTests, auto VectorRounder, auto ScalarRounder>
+class TTestEqualAfterVectorRounding
+{
+public:
+	explicit TTestEqualAfterVectorRounding(VectorRoundTests& InRoundTests) : RoundTests(InRoundTests)
+	{
+	}
+
+	template<class ExpectedType>
+	bool operator ()(const TCHAR* What, float Actual, ExpectedType Expected, bool bTestScalar = true) const
+	{
+		// Compare the result of the vector implementation with the equivalent scalar implementation.
+		if (bTestScalar)
+			RoundTests.TestEqual(What, FMathVectorTestParameter(VectorRounder(VectorSet1(Actual))), FMathVectorTestParameter(ScalarRounder(Actual)));
+
+		// Compare the result of the vector implementation with the given expected value.
+		RoundTests.TestEqual(What, FMathVectorTestParameter(VectorRounder(VectorSet1(Actual))), FMathVectorTestParameter(Expected));
+		return true;
+	}
+
+private:
+	VectorRoundTests& RoundTests;
+};
+
+FORCEINLINE VectorRegister4Int VectorRoundToIntHalfToEvenTest(const VectorRegister4Float& Vec)
+{
+	return VectorRoundToIntHalfToEven(Vec);
+}
+
+IMPLEMENT_SIMPLE_AUTOMATION_TEST(FMathVectorRoundToIntHalfToEvenTests, "System.Core.Math.RoundToIntHalfToEven Vector", EAutomationTestFlags::ApplicationContextMask | EAutomationTestFlags::SmokeFilter)
+bool FMathVectorRoundToIntHalfToEvenTests::RunTest(const FString& Parameters)
+{
+	const bool bSkipScalar = false;
+
+	struct FMathVectorRoundToIntHalfToEvenTestsRoundToInt
+	{
+		static int RoundToInt(float F)
+		{
+			return int(FMath::RoundHalfToEven(F));
+		}
+	};
+	TTestEqualAfterVectorRounding<FMathVectorRoundToIntHalfToEvenTests, &VectorRoundToIntHalfToEvenTest, &FMathVectorRoundToIntHalfToEvenTestsRoundToInt::RoundToInt> TestEqualAfterVectorRounding(*this);
+
+	TestEqualAfterVectorRounding(TEXT("VectorRound32-Zero"), 0.0f, 0);
+	TestEqualAfterVectorRounding(TEXT("VectorRound32-One"), 1.0f, 1);
+	TestEqualAfterVectorRounding(TEXT("VectorRound32-LessHalf"), 1.4f, 1);
+	TestEqualAfterVectorRounding(TEXT("VectorRound32-NegGreaterHalf"), -1.4f, -1);
+	TestEqualAfterVectorRounding(TEXT("VectorRound32-LessNearHalf"), 1.4999999f, 1);
+	TestEqualAfterVectorRounding(TEXT("VectorRound32-NegGreaterNearHalf"), -1.4999999f, -1);
+	TestEqualAfterVectorRounding(TEXT("VectorRound32-Half"), 1.5f, 2);
+	TestEqualAfterVectorRounding(TEXT("VectorRound32-NegHalf"), -1.5f, -2);
+	TestEqualAfterVectorRounding(TEXT("VectorRound32-GreaterNearHalf"), 1.5000001f, 2);
+	TestEqualAfterVectorRounding(TEXT("VectorRound32-NegLesserNearHalf"), -1.5000001f, -2);
+	TestEqualAfterVectorRounding(TEXT("VectorRound32-GreaterThanHalf"), 1.6f, 2);
+	TestEqualAfterVectorRounding(TEXT("VectorRound32-NegLesserThanHalf"), -1.6f, -2);
+
+	TestEqualAfterVectorRounding(TEXT("VectorRound32-TwoToOneBitPrecision"), 4194303.25f, 4194303);
+	TestEqualAfterVectorRounding(TEXT("VectorRound32-TwoToOneBitPrecision"), 4194303.5f, 4194304);
+	TestEqualAfterVectorRounding(TEXT("VectorRound32-TwoToOneBitPrecision"), 4194303.75f, 4194304);
+	TestEqualAfterVectorRounding(TEXT("VectorRound32-TwoToOneBitPrecision"), 4194304.0f, 4194304);
+	TestEqualAfterVectorRounding(TEXT("VectorRound32-TwoToOneBitPrecision"), 4194304.5f, 4194304);
+	TestEqualAfterVectorRounding(TEXT("VectorRound32-NegTwoToOneBitPrecision"), -4194303.25f, -4194303);
+	TestEqualAfterVectorRounding(TEXT("VectorRound32-NegTwoToOneBitPrecision"), -4194303.5f, -4194304);
+	TestEqualAfterVectorRounding(TEXT("VectorRound32-NegTwoToOneBitPrecision"), -4194303.75f, -4194304);
+	TestEqualAfterVectorRounding(TEXT("VectorRound32-NegTwoToOneBitPrecision"), -4194304.0f, -4194304);
+	TestEqualAfterVectorRounding(TEXT("VectorRound32-NegTwoToOneBitPrecision"), -4194304.5f, -4194304);
+
+	TestEqualAfterVectorRounding(TEXT("VectorRound32-OneToZeroBitPrecision"), 8388607.0f, 8388607);
+	TestEqualAfterVectorRounding(TEXT("VectorRound32-OneToZeroBitPrecision"), 8388607.5f, 8388608);
+	TestEqualAfterVectorRounding(TEXT("VectorRound32-OneToZeroBitPrecision"), 8388608.0f, 8388608);
+	TestEqualAfterVectorRounding(TEXT("VectorRound32-OneToZeroBitPrecision"), 8388608.5f, 8388608);
+	//TestEqualAfterVectorRounding(TEXT("VectorRound32-OneToZeroBitPrecision"), 8388609.0f, 8388609, bSkipScalar); // FMath::RoundHalfToEven incorrectly rounds 8388609.0f to 8388610.0f, so skip it.
+	TestEqualAfterVectorRounding(TEXT("VectorRound32-OneToZeroBitPrecision"), 8388609.5f, 8388610);
+
+	TestEqualAfterVectorRounding(TEXT("VectorRound32-NegOneToZeroBitPrecision"), -8388607.0f, -8388607);
+	TestEqualAfterVectorRounding(TEXT("VectorRound32-NegOneToZeroBitPrecision"), -8388607.5f, -8388608);
+	TestEqualAfterVectorRounding(TEXT("VectorRound32-NegOneToZeroBitPrecision"), -8388608.0f, -8388608);
+	TestEqualAfterVectorRounding(TEXT("VectorRound32-NegOneToZeroBitPrecision"), -8388608.5f, -8388608);
+	//TestEqualAfterVectorRounding(TEXT("VectorRound32-NegOneToZeroBitPrecision"), -8388609.0f, -8388609, bSkipScalar); // FMath::RoundHalfToEven incorrectly rounds -8388609.0f to -8388610.0f, so skip it.
+	TestEqualAfterVectorRounding(TEXT("VectorRound32-NegOneToZeroBitPrecision"), -8388609.5f, -8388610);
+
+	//TestEqualAfterVectorRounding(TEXT("VectorRound32-ZeroBitPrecision"), 16777215.0f, 16777215, bSkipScalar); // FMath::RoundHalfToEven incorrectly rounds 16777215.0f to 16777216.0f, so skip it.
+	TestEqualAfterVectorRounding(TEXT("VectorRound32-ZeroBitPrecision"), 16777215.5f, 16777216);
+	TestEqualAfterVectorRounding(TEXT("VectorRound32-ZeroBitPrecision"), 16777216.0f, 16777216);
+	TestEqualAfterVectorRounding(TEXT("VectorRound32-ZeroBitPrecision"), 16777216.5f, 16777216);
+	TestEqualAfterVectorRounding(TEXT("VectorRound32-ZeroBitPrecision"), 16777217.0f, 16777216);
+	TestEqualAfterVectorRounding(TEXT("VectorRound32-ZeroBitPrecision"), 16777217.5f, 16777218);
+	//TestEqualAfterVectorRounding(TEXT("VectorRound32-NegZeroBitPrecision"), -16777215.0f, -16777215, bSkipScalar); // FMath::RoundHalfToEven incorrectly rounds -16777215.0f to -16777216.0f, so skip it.
+	TestEqualAfterVectorRounding(TEXT("VectorRound32-NegZeroBitPrecision"), -16777215.5f, -16777216);
+	TestEqualAfterVectorRounding(TEXT("VectorRound32-NegZeroBitPrecision"), -16777216.0f, -16777216);
+	TestEqualAfterVectorRounding(TEXT("VectorRound32-NegZeroBitPrecision"), -16777216.5f, -16777216);
+	TestEqualAfterVectorRounding(TEXT("VectorRound32-NegZeroBitPrecision"), -16777217.0f, -16777216);
+	TestEqualAfterVectorRounding(TEXT("VectorRound32-NegZeroBitPrecision"), -16777217.5f, -16777218);
+
+	TestEqualAfterVectorRounding(TEXT("VectorRound32-FloatMax"), FLT_MAX, INT32_MIN);
+	TestEqualAfterVectorRounding(TEXT("VectorRound32-FloatMax"), -FLT_MAX, INT32_MIN);
+
+	TestEqualAfterVectorRounding(TEXT("VectorRound32-FloatMin"), FLT_MIN, 0);
+	TestEqualAfterVectorRounding(TEXT("VectorRound32-FloatMin"), -FLT_MIN, 0);
+
+	return true;
+}
+
 IMPLEMENT_SIMPLE_AUTOMATION_TEST(FMathRoundHalfToZeroTests, "System.Core.Math.Round HalfToZero", EAutomationTestFlags::ApplicationContextMask | EAutomationTestFlags::SmokeFilter)
 bool FMathRoundHalfToZeroTests::RunTest(const FString& Parameters)
 {
diff --git a/Engine/Source/Runtime/Core/Public/Math/UnrealMathDirectX.h b/Engine/Source/Runtime/Core/Public/Math/UnrealMathDirectX.h
--- a/Engine/Source/Runtime/Core/Public/Math/UnrealMathDirectX.h
+++ b/Engine/Source/Runtime/Core/Public/Math/UnrealMathDirectX.h
@@ -1199,6 +1199,11 @@ FORCEINLINE VectorRegister4Int VectorIntAbs(const VectorRegister4Int& A)
 #define VectorIntToFloat(A) _mm_cvtepi32_ps(A)
 #define VectorFloatToInt(A) _mm_cvttps_epi32(A)
 
+FORCEINLINE VectorRegister4Int VectorRoundToIntHalfToEven(const VectorRegister4Float& Vec)
+{
+	return _mm_cvtps_epi32(Vec);
+}
+
 //Loads and stores
 
 /**
diff --git a/Engine/Source/Runtime/Core/Public/Math/UnrealMathFPU.h b/Engine/Source/Runtime/Core/Public/Math/UnrealMathFPU.h
--- a/Engine/Source/Runtime/Core/Public/Math/UnrealMathFPU.h
+++ b/Engine/Source/Runtime/Core/Public/Math/UnrealMathFPU.h
@@ -2280,6 +2280,16 @@ FORCEINLINE VectorRegister4Double VectorRound(const VectorRegister4Double& Vec)
 }
 
 
+FORCEINLINE VectorRegister4Int VectorRoundToIntHalfToEven(const VectorRegister4Float& A)
+{
+	return MakeVectorRegisterInt(
+		(int32)FMath::RoundHalfToEven(A.V[0]),
+		(int32)FMath::RoundHalfToEven(A.V[1]),
+		(int32)FMath::RoundHalfToEven(A.V[2]),
+		(int32)FMath::RoundHalfToEven(A.V[3]));
+}
+
+
 FORCEINLINE VectorRegister4Float VectorMod(const VectorRegister4Float& X, const VectorRegister4Float& Y)
 {
 	return MakeVectorRegisterFloat(
diff --git a/Engine/Source/Runtime/Core/Public/Math/UnrealMathNeon.h b/Engine/Source/Runtime/Core/Public/Math/UnrealMathNeon.h
--- a/Engine/Source/Runtime/Core/Public/Math/UnrealMathNeon.h
+++ b/Engine/Source/Runtime/Core/Public/Math/UnrealMathNeon.h
@@ -3027,6 +3027,10 @@ FORCEINLINE VectorRegister4Int VectorFloatToInt(const VectorRegister4Double& A)
 #define VectorShiftRightImmLogical(Vec, ImmAmt)     vshrq_n_u32(Vec, ImmAmt)
 #define VectorRound(Vec)							vrndnq_f32(Vec)
 
+FORCEINLINE VectorRegister4Int VectorRoundToIntHalfToEven(const VectorRegister4Float& Vec)
+{
+	return vcvtnq_s32_f32(Vec);
+}
 
 inline VectorRegister4i VectorIntExpandLow16To32(VectorRegister4i V) {
 	int16x4x2_t res = vzip_s16(vget_low_u16(V), vdup_n_u16(0));
diff --git a/Engine/Source/Runtime/Core/Public/Math/UnrealMathSSE.h b/Engine/Source/Runtime/Core/Public/Math/UnrealMathSSE.h
--- a/Engine/Source/Runtime/Core/Public/Math/UnrealMathSSE.h
+++ b/Engine/Source/Runtime/Core/Public/Math/UnrealMathSSE.h
@@ -2785,6 +2785,11 @@ FORCEINLINE VectorRegister4Float VectorRound(const VectorRegister4Float &Vec)
 #endif
 }
 
+FORCEINLINE VectorRegister4Int VectorRoundToIntHalfToEven(const VectorRegister4Float& Vec)
+{
+	return _mm_cvtps_epi32(Vec);
+}
+
 FORCEINLINE VectorRegister4Float VectorCeil(const VectorRegister4Float& V)
 {
 #if UE_PLATFORM_MATH_USE_SSE4_1
