remove temporary 5.1 scaffolding class used for deprecating FCustomChunk and introducing Named Chunks.

#rnx

diff --git a/Engine/Source/Runtime/Core/Private/GenericPlatform/GenericPlatformChunkInstall.cpp b/Engine/Source/Runtime/Core/Private/GenericPlatform/GenericPlatformChunkInstall.cpp
--- a/Engine/Source/Runtime/Core/Private/GenericPlatform/GenericPlatformChunkInstall.cpp
+++ b/Engine/Source/Runtime/Core/Private/GenericPlatform/GenericPlatformChunkInstall.cpp
@@ -8,236 +8,11 @@ DEFINE_LOG_CATEGORY(LogChunkInstaller)
 
 
 /*
- *  Functionality to to implement the new 'named chunk' API using the existing, deprecated 'FCustomChunk' API
- *  Once the FCustomChunk API is removed then these will be replaced by platform native implementations
+ *  Functionality to to implement the deprecated 'FCustomChunk' API using the new 'named chunk' API
  */
 
 PRAGMA_DISABLE_DEPRECATION_WARNINGS
 
-bool FNamedChunkPlatformChunkInstall::IsNamedChunkInProgress(const FName NamedChunk)
-{
-	TArray<FCustomChunk> CustomChunks = GetCustomChunksFromNamedChunk(NamedChunk);
-	return IsChunkInstallationPending(CustomChunks);
-}
-
-
-bool FNamedChunkPlatformChunkInstall::InstallNamedChunk(const FName NamedChunk)
-{
-	TArray<FCustomChunk> CustomChunks = GetCustomChunksFromNamedChunk(NamedChunk);
-	return InstallChunks(CustomChunks);
-}
-
-
-bool FNamedChunkPlatformChunkInstall::UninstallNamedChunk(const FName NamedChunk)
-{
-	TArray<FCustomChunk> CustomChunks = GetCustomChunksFromNamedChunk(NamedChunk);
-	return UninstallChunks(CustomChunks);
-}
-
-
-bool FNamedChunkPlatformChunkInstall::InstallNamedChunks(const TArrayView<const FName>& NamedChunks)
-{
-	TArray<FCustomChunk> CustomChunks = GetCustomChunksFromNamedChunks(NamedChunks);
-	return InstallChunks(CustomChunks);
-}
-
-
-bool FNamedChunkPlatformChunkInstall::UninstallNamedChunks(const TArrayView<const FName>& NamedChunks)
-{
-	TArray<FCustomChunk> CustomChunks = GetCustomChunksFromNamedChunks(NamedChunks);
-	return UninstallChunks(CustomChunks);
-}
-
-
-bool FNamedChunkPlatformChunkInstall::PrioritizeNamedChunk(const FName NamedChunk, EChunkPriority::Type Priority)
-{
-	FCustomChunk CustomChunk(TEXT(""), 0xFFFFFFFF, ECustomChunkType::OnDemandChunk);
-	if (TryGetCustomChunkFromNamedChunk(NamedChunk, CustomChunk))
-	{
-		return PrioritizeChunk(CustomChunk.ChunkID, Priority);
-	}
-
-	return false;
-}
-
-
-EChunkLocation::Type FNamedChunkPlatformChunkInstall::GetNamedChunkLocation(const FName NamedChunk)
-{
-	FCustomChunk CustomChunk(TEXT(""), 0xFFFFFFFF, ECustomChunkType::OnDemandChunk);
-	if (TryGetCustomChunkFromNamedChunk(NamedChunk, CustomChunk))
-	{
-		return GetChunkLocation(CustomChunk.ChunkID);
-	}
-
-	return EChunkLocation::DoesNotExist;
-}
-
-
-float FNamedChunkPlatformChunkInstall::GetNamedChunkProgress(const FName NamedChunk, EChunkProgressReportingType::Type ReportType)
-{
-	FCustomChunk CustomChunk(TEXT(""), 0xFFFFFFFF, ECustomChunkType::OnDemandChunk);
-	if (TryGetCustomChunkFromNamedChunk(NamedChunk, CustomChunk))
-	{
-		return GetCustomChunkProgress(CustomChunk, ReportType);
-	}
-
-	return 0.0f;
-}
-
-
-ENamedChunkType FNamedChunkPlatformChunkInstall::GetNamedChunkType(const FName NamedChunk) const
-{
-	FCustomChunk CustomChunk(TEXT(""), 0xFFFFFFFF, ECustomChunkType::OnDemandChunk);
-	if (TryGetCustomChunkFromNamedChunk(NamedChunk, CustomChunk))
-	{
-		if (CustomChunk.ChunkType == ECustomChunkType::OnDemandChunk)
-		{
-			return ENamedChunkType::OnDemand;
-		}
-		else if (CustomChunk.ChunkType == ECustomChunkType::LanguageChunk)
-		{
-			return ENamedChunkType::Language;
-		}
-	}
-
-	return ENamedChunkType::Invalid;
-}
-
-
-TArray<FName> FNamedChunkPlatformChunkInstall::GetNamedChunksByType(ENamedChunkType NamedChunkType) const
-{
-	TArray<FCustomChunk> CustomChunks;
-	if (NamedChunkType == ENamedChunkType::OnDemand)
-	{
-		CustomChunks = FPlatformMisc::GetCustomChunksByType(ECustomChunkType::OnDemandChunk);
-	}
-	else if (NamedChunkType == ENamedChunkType::Language)
-	{
-		CustomChunks = FPlatformMisc::GetCustomChunksByType(ECustomChunkType::LanguageChunk);
-	}
-
-	TArray<FName> NamedChunks = GetNamedChunksFromCustomChunks(CustomChunks);
-	return MoveTemp(NamedChunks);
-}
-
-
-TArray<FCustomChunk> FNamedChunkPlatformChunkInstall::GetCustomChunksFromNamedChunk(const FName NamedChunk) const
-{
-	TArray<FCustomChunk> CustomChunks;
-	
-	FCustomChunk CustomChunk(TEXT(""), 0xFFFFFFFF, ECustomChunkType::OnDemandChunk);
-	if (TryGetCustomChunkFromNamedChunk(NamedChunk, CustomChunk))
-	{
-		CustomChunks.Emplace(CustomChunk.ChunkTag, CustomChunk.ChunkID, CustomChunk.ChunkType, CustomChunk.ChunkTag2);
-	}
-
-	return MoveTemp(CustomChunks);
-}
-
-TArray<FCustomChunk> FNamedChunkPlatformChunkInstall::GetCustomChunksFromNamedChunks(const TArrayView<const FName>& NamedChunks) const
-{
-	TArray<FCustomChunk> CustomChunks;
-
-	for (FName NamedChunk : NamedChunks)
-	{
-		CustomChunks.Append(GetCustomChunksFromNamedChunk(NamedChunk));
-	}
-
-	return MoveTemp(CustomChunks);
-}
-
-
-bool FNamedChunkPlatformChunkInstall::TryGetCustomChunkFromNamedChunk(const FName NamedChunk, FCustomChunk& OutCustomChunk) const
-{
-	for (const FCustomChunk& CustomChunk : FPlatformMisc::GetAllOnDemandChunks())
-	{
-		if (FName(CustomChunk.ChunkTag) == NamedChunk || FName(CustomChunk.ChunkTag2) == NamedChunk)
-		{
-			OutCustomChunk = FCustomChunk(CustomChunk.ChunkTag, CustomChunk.ChunkID, CustomChunk.ChunkType, CustomChunk.ChunkTag2);
-			return true;
-		}
-	}
-
-	for (const FCustomChunk& CustomChunk : FPlatformMisc::GetAllLanguageChunks())
-	{
-		if ( FName(CustomChunk.ChunkTag) == NamedChunk || FName(CustomChunk.ChunkTag2) == NamedChunk)
-		{
-			OutCustomChunk = FCustomChunk(CustomChunk.ChunkTag, CustomChunk.ChunkID, CustomChunk.ChunkType, CustomChunk.ChunkTag2);
-			return true;
-		}
-	}
-
-	UE_LOG(LogChunkInstaller, Warning, TEXT("Named chunk %s not found"), *NamedChunk.ToString());
-	return false;
-}
-
-
-TArray<FName> FNamedChunkPlatformChunkInstall::GetNamedChunksFromCustomChunks(const TArray<FCustomChunk>& CustomChunks) const
-{
-	TArray<FName> NamedChunks;
-	NamedChunks.Reserve(CustomChunks.Num());
-
-	for (const FCustomChunk& CustomChunk : CustomChunks)
-	{
-		FName ChunkName = GetCustomChunkName(CustomChunk);
-		if (ChunkName != NAME_None)
-		{
-			NamedChunks.Add(FName(ChunkName));
-		}
-	}
-
-	return MoveTemp(NamedChunks);
-}
-
-
-FName FNamedChunkPlatformChunkInstall::GetNamedChunkByPakChunkIndex(int32 InPakchunkIndex) const
-{
-	TArray<int32> PakchunkIndices;
-	PakchunkIndices.Add(InPakchunkIndex);
-
-	TArray<FCustomChunk> CustomChunks = FPlatformMisc::GetOnDemandChunksForPakchunkIndices(PakchunkIndices);
-	if (CustomChunks.Num() == 0)
-	{
-		return NAME_None;
-	}
-	
-	check(CustomChunks.Num() == 1); // a pak file cannot be in multiple chunks
-	return GetCustomChunkName(CustomChunks[0]);
-}
-
-
-FName FNamedChunkPlatformChunkInstall::GetCustomChunkName(const FCustomChunk& CustomChunk) const
-{
-	if (!CustomChunk.ChunkTag.IsEmpty())
-	{
-		return FName(CustomChunk.ChunkTag);
-	}
-	else if (!CustomChunk.ChunkTag2.IsEmpty())
-	{
-		return FName(CustomChunk.ChunkTag2);
-	}
-	else
-	{
-		return NAME_None;
-	}
-}
-
-
-
-
-
-
-
-
-
-
-
-
-
-
-
-
-
 
 bool FGenericPlatformChunkInstall_WithEmulatedCustomChunks::IsChunkInstallationPending(const TArray<FCustomChunk>& ChunkTagsID)
 {
diff --git a/Engine/Source/Runtime/Core/Public/GenericPlatform/GenericPlatformChunkInstall.h b/Engine/Source/Runtime/Core/Public/GenericPlatform/GenericPlatformChunkInstall.h
--- a/Engine/Source/Runtime/Core/Public/GenericPlatform/GenericPlatformChunkInstall.h
+++ b/Engine/Source/Runtime/Core/Public/GenericPlatform/GenericPlatformChunkInstall.h
@@ -481,37 +481,6 @@ protected:
 	}
 };
 
-// temporary helper base for platform chunk installers to transition from FCustomChunk to named chunks
-class CORE_API FNamedChunkPlatformChunkInstall : public FGenericPlatformChunkInstall
-{
-public:
-	virtual bool SupportsNamedChunkInstall() const override
-	{
-		return true;
-	}
-
-	virtual bool IsNamedChunkInProgress(const FName NamedChunk) override;
-	virtual bool InstallNamedChunk(const FName NamedChunk) override;
-	virtual bool UninstallNamedChunk(const FName NamedChunk) override;
-	virtual bool InstallNamedChunks(const TArrayView<const FName>& NamedChunks) override;
-	virtual bool UninstallNamedChunks(const TArrayView<const FName>& NamedChunks) override;
-
-	virtual EChunkLocation::Type GetNamedChunkLocation(const FName NamedChunk) override;
-	virtual float GetNamedChunkProgress(const FName NamedChunk, EChunkProgressReportingType::Type ReportType) override;
-	virtual bool PrioritizeNamedChunk(const FName NamedChunk, EChunkPriority::Type Priority) override;
-
-	virtual ENamedChunkType GetNamedChunkType(const FName NamedChunk) const override;
-	virtual TArray<FName> GetNamedChunksByType(ENamedChunkType NamedChunkType) const override;
-
-protected:
-	virtual float GetCustomChunkProgress(const FCustomChunk& CustomChunk, EChunkProgressReportingType::Type ReportType) = 0; // platform specializations need to implement this as it's missing in the existing api
-	bool TryGetCustomChunkFromNamedChunk(const FName NamedChunk, FCustomChunk& OutCustomChunk) const;
-	TArray<FCustomChunk> GetCustomChunksFromNamedChunk(const FName NamedChunk) const;
-	TArray<FCustomChunk> GetCustomChunksFromNamedChunks(const TArrayView<const FName>& NamedChunks) const;
-	TArray<FName> GetNamedChunksFromCustomChunks(const TArray<FCustomChunk>& CustomChunks) const;
-	FName GetNamedChunkByPakChunkIndex(int32 InPakchunkIndex) const;
-	FName GetCustomChunkName(const FCustomChunk& CustomChunk) const;
-};
 
 // temporary helper base class for platform chunk installers that have implemented named chunk support to provide FCustomChunk emulation
 class CORE_API FGenericPlatformChunkInstall_WithEmulatedCustomChunks : public FGenericPlatformChunkInstall
