InheritedContext

grouped inherited LLM tag, mem tag and trace metadata into FInheritedContext to simplify their usage, avoid forgetting using all of them together, and to hide the ugly conditional compilation. Added FInheritedContext to higher level tasks instead of the low-level tasks to allow batch task launch to capture the inherited context once per batch

diff --git a/Engine/Source/Developer/DerivedDataCache/Private/DerivedDataCache.cpp b/Engine/Source/Developer/DerivedDataCache/Private/DerivedDataCache.cpp
--- a/Engine/Source/Developer/DerivedDataCache/Private/DerivedDataCache.cpp
+++ b/Engine/Source/Developer/DerivedDataCache/Private/DerivedDataCache.cpp
@@ -19,7 +19,6 @@
 #include "DerivedDataRequestOwner.h"
 #include "Experimental/Async/LazyEvent.h"
 #include "Features/IModularFeatures.h"
-#include "HAL/LowLevelMemTracker.h"
 #include "HAL/ThreadSafeCounter.h"
 #include "Misc/CommandLine.h"
 #include "Misc/CoreMisc.h"
@@ -333,7 +332,6 @@ public:
 		: Owner(InOwner)
 		, TaskBody(MoveTemp(InTaskBody))
 	{
-		LLM_IF_ENABLED(MemTag = FLowLevelMemTracker::Get().GetActiveTagData(ELLMTracker::Default));
 		Owner.Begin(this);
 		DoneEvent.Reset();
 		GCacheThreadPool->AddQueuedWork(this, ConvertToQueuedWorkPriority(Owner.GetPriority()));
@@ -342,7 +340,6 @@ public:
 private:
 	inline void Execute()
 	{
-		LLM_SCOPE(MemTag);
 		FScopeCycleCounter Scope(GetStatId(), /*bAlways*/ true);
 		Owner.End(this, [this]
 		{
@@ -408,7 +405,6 @@ private:
 	IRequestOwner& Owner;
 	TUniqueFunction<void ()> TaskBody;
 	FLazyEvent DoneEvent{EEventMode::ManualReset};
-	LLM(const UE::LLMPrivate::FTagData* MemTag = nullptr);
 };
 
 void LaunchTaskInCacheThreadPool(IRequestOwner& Owner, TUniqueFunction<void ()>&& TaskBody)
diff --git a/Engine/Source/Developer/DerivedDataCache/Private/DerivedDataRequestOwner.cpp b/Engine/Source/Developer/DerivedDataCache/Private/DerivedDataRequestOwner.cpp
--- a/Engine/Source/Developer/DerivedDataCache/Private/DerivedDataRequestOwner.cpp
+++ b/Engine/Source/Developer/DerivedDataCache/Private/DerivedDataRequestOwner.cpp
@@ -7,7 +7,6 @@
 #include "DerivedDataRequestTypes.h"
 #include "Experimental/Async/LazyEvent.h"
 #include "HAL/CriticalSection.h"
-#include "HAL/LowLevelMemTracker.h"
 #include "Misc/ScopeExit.h"
 #include "Misc/ScopeRWLock.h"
 #include "Tasks/Task.h"
@@ -299,17 +298,14 @@ void IRequestOwner::LaunchTask(const TCHAR* DebugName, TUniqueFunction<void ()>&
 		void Cancel() final { Task.Wait(); }
 		void Wait() final { Task.Wait(); }
 		FTask Task;
-		LLM(const UE::LLMPrivate::FTagData* MemTag = nullptr);
 	};
 
 	Tasks::FTaskEvent TaskEvent(TEXT("LaunchTaskRequest"));
 	FTaskRequest* Request = new FTaskRequest;
-	LLM_IF_ENABLED(Request->MemTag = FLowLevelMemTracker::Get().GetActiveTagData(ELLMTracker::Default));
 	Request->Task = Launch(
 		DebugName,
 		[this, Request, TaskBody = MoveTemp(TaskBody)]
 		{
-			LLM_SCOPE(Request->MemTag);
 			End(Request, TaskBody);
 		},
 		TaskEvent,
diff --git a/Engine/Source/Developer/DerivedDataCache/Private/HttpCacheStore.cpp b/Engine/Source/Developer/DerivedDataCache/Private/HttpCacheStore.cpp
--- a/Engine/Source/Developer/DerivedDataCache/Private/HttpCacheStore.cpp
+++ b/Engine/Source/Developer/DerivedDataCache/Private/HttpCacheStore.cpp
@@ -21,7 +21,6 @@
 #include "Dom/JsonObject.h"
 #include "HAL/CriticalSection.h"
 #include "HAL/IConsoleManager.h"
-#include "HAL/LowLevelMemTracker.h"
 #include "HAL/PlatformProcess.h"
 #include "Http/HttpClient.h"
 #include "IO/IoHash.h"
@@ -568,9 +567,7 @@ public:
 		: Owner(InOwner)
 		, BaseReceiver(InOperation, this)
 		, OperationComplete(MoveTemp(InOperationComplete))
-	{
-		LLM_IF_ENABLED(MemTag = FLowLevelMemTracker::Get().GetActiveTagData(ELLMTracker::Default));
-	}
+	{}
 
 private:
 	// IRequest Interface
@@ -583,7 +580,6 @@ private:
 
 	IHttpReceiver* OnCreate(IHttpResponse& LocalResponse) final
 	{
-		LLM_SCOPE(MemTag);
 		Monitor = LocalResponse.GetMonitor();
 		Owner->Begin(this);
 		return &BaseReceiver;
@@ -591,7 +587,6 @@ private:
 
 	IHttpReceiver* OnComplete(IHttpResponse& LocalResponse) final
 	{
-		LLM_SCOPE(MemTag);
 		Owner->End(this, [Self = this]
 		{
 			FHttpOperation* Operation = Self->BaseReceiver.GetOperation();
@@ -618,7 +613,6 @@ private:
 	FHttpOperationReceiver BaseReceiver;
 	TUniqueFunction<void ()> OperationComplete;
 	TRefCountPtr<IHttpResponseMonitor> Monitor;
-	LLM(const UE::LLMPrivate::FTagData* MemTag = nullptr);
 };
 
 void FHttpCacheStore::FHttpOperation::Send()
diff --git a/Engine/Source/Runtime/Core/Private/Async/InheritedContext.cpp b/Engine/Source/Runtime/Core/Private/Async/InheritedContext.cpp
new file mode 100644
--- /dev/null
+++ b/Engine/Source/Runtime/Core/Private/Async/InheritedContext.cpp
@@ -0,0 +1,24 @@
+// Copyright Epic Games, Inc. All Rights Reserved.
+
+#include "Async/InheritedContext.h"
+
+namespace UE
+{
+	// this method must be defined in a .cpp file to avoid exporting dependency on trace module
+	UE_NODISCARD FInheritedContextScope FInheritedContextBase::RestoreInheritedContext()
+	{
+		UE_TRACE_METADATA_RESTORE_STACK(InheritedMetadataId);
+
+		return FInheritedContextScope{
+#if ENABLE_LOW_LEVEL_MEM_TRACKER
+				InheritedLLMTag
+	#if UE_MEMORY_TAGS_TRACE_ENABLED
+				,
+	#endif
+#endif
+#if UE_MEMORY_TAGS_TRACE_ENABLED
+				InheritedMemTag
+#endif
+		};
+	}
+}
\ No newline at end of file
diff --git a/Engine/Source/Runtime/Core/Private/IO/IoDispatcher.cpp b/Engine/Source/Runtime/Core/Private/IO/IoDispatcher.cpp
--- a/Engine/Source/Runtime/Core/Private/IO/IoDispatcher.cpp
+++ b/Engine/Source/Runtime/Core/Private/IO/IoDispatcher.cpp
@@ -1075,8 +1075,7 @@ void FIoRequestImpl::CreateBuffer(uint64 Size)
 	}
 	else
 	{
-		LLM_SCOPE(InheritedLLMTag);
-		UE_MEMSCOPE(InheritedTraceTag);
+		UE::FInheritedContextScope InheritedContextScope = RestoreInheritedContext();
 		TRACE_CPUPROFILER_EVENT_SCOPE(AllocMemoryForRequest);
 		Buffer.Emplace(Size);
 	}
diff --git a/Engine/Source/Runtime/Core/Public/Async/AsyncWork.h b/Engine/Source/Runtime/Core/Public/Async/AsyncWork.h
--- a/Engine/Source/Runtime/Core/Public/Async/AsyncWork.h
+++ b/Engine/Source/Runtime/Core/Public/Async/AsyncWork.h
@@ -13,10 +13,9 @@
 #include "Stats/Stats.h"
 #include "HAL/Event.h"
 #include "HAL/PlatformProcess.h"
-#include "HAL/LowLevelMemTracker.h"
+#include "Async/InheritedContext.h"
 #include "Misc/IQueuedWork.h"
 #include "Misc/QueuedThreadPool.h"
-#include "ProfilingDebugging/TagTrace.h"
 #include "Async/Fundamental/Scheduler.h"
 
 /**
@@ -58,28 +57,19 @@
 
 **/
 template<typename TTask>
-class FAutoDeleteAsyncTask
-	: private IQueuedWork
+class FAutoDeleteAsyncTask : 
+	private UE::FInheritedContextBase, 
+	private IQueuedWork
 {
 	/** User job embedded in this task */
 	TTask Task;
-	/** optional LLM tag */
-	LLM(const UE::LLMPrivate::FTagData* InheritedLLMTag);
-#if UE_MEMORY_TAGS_TRACE_ENABLED
-	/** optional Trace tag */
-	int32 InheritedTraceTag;
-#endif
 
 	/* Generic start function, not called directly
 	 * @param bForceSynchronous if true, this job will be started synchronously, now, on this thread
 	 **/
 	void Start(bool bForceSynchronous, FQueuedThreadPool* InQueuedPool, EQueuedWorkPriority InPriority = EQueuedWorkPriority::Normal)
 	{
-		LLM(InheritedLLMTag = FLowLevelMemTracker::bIsDisabled ? nullptr : FLowLevelMemTracker::Get().GetActiveTagData(ELLMTracker::Default));
-#if UE_MEMORY_TAGS_TRACE_ENABLED
-		InheritedTraceTag = MemoryTrace_GetActiveTag();
-#endif
-
+		CaptureInheritedContext();
 		FPlatformMisc::MemoryBarrier();
 		FQueuedThreadPool* QueuedPool = InQueuedPool;
 		if (bForceSynchronous)
@@ -102,8 +92,7 @@ class FAutoDeleteAsyncTask
 	 **/
 	void DoWork()
 	{
-		LLM_SCOPE(InheritedLLMTag);
-		UE_MEMSCOPE(InheritedTraceTag);
+		UE::FInheritedContextScope InheritedContextScope = RestoreInheritedContext();
 		FScopeCycleCounter Scope(Task.GetStatId(), true);
 
 		Task.DoWork();
@@ -215,7 +204,8 @@ public:
 **/
 
 class FAsyncTaskBase
-	: private IQueuedWork
+	: private UE::FInheritedContextBase
+	, private IQueuedWork
 {
 	/** Thread safe counter that indicates WORK completion, no necessarily finalization of the job */
 	FThreadSafeCounter	WorkNotFinishedCounter;
@@ -231,12 +221,6 @@ class FAsyncTaskBase
 	int64 RequiredMemory = -1;
 	/** StatId used for FScopeCycleCounter */
 	TStatId StatId;
-	/** optional LLM tag */
-	LLM(const UE::LLMPrivate::FTagData* InheritedLLMTag);
-	/** Memory trace tag */
-#if UE_MEMORY_TAGS_TRACE_ENABLED
-	int32 InheritedTraceTag;
-#endif
 
 	/* Internal function to destroy the completion event
 	**/
@@ -256,12 +240,10 @@ class FAsyncTaskBase
 	**/
 	void Start(bool bForceSynchronous, FQueuedThreadPool* InQueuedPool, EQueuedWorkPriority InQueuedWorkPriority, EQueuedWorkFlags InQueuedWorkFlags, int64 InRequiredMemory)
 	{
+		CaptureInheritedContext();
+
 		FScopeCycleCounter Scope(StatId, true);
 		DECLARE_SCOPE_CYCLE_COUNTER( TEXT( "FAsyncTask::Start" ), STAT_FAsyncTask_Start, STATGROUP_ThreadPoolAsyncTasks );
-		LLM(InheritedLLMTag = FLowLevelMemTracker::bIsDisabled ? nullptr : FLowLevelMemTracker::Get().GetActiveTagData(ELLMTracker::Default));
-#if UE_MEMORY_TAGS_TRACE_ENABLED
-		InheritedTraceTag = MemoryTrace_GetActiveTag();
-#endif
 		RequiredMemory = InRequiredMemory;
 
 		FPlatformMisc::MemoryBarrier();
@@ -296,8 +278,7 @@ class FAsyncTaskBase
 	**/
 	void DoWork()
 	{
-		LLM_SCOPE(InheritedLLMTag);
-		UE_MEMSCOPE(InheritedTraceTag);
+		UE::FInheritedContextScope InheritedContextScope = RestoreInheritedContext();
 		FScopeCycleCounter Scope(StatId, true);
 
 		DoTaskWork();
diff --git a/Engine/Source/Runtime/Core/Public/Async/InheritedContext.h b/Engine/Source/Runtime/Core/Public/Async/InheritedContext.h
new file mode 100644
--- /dev/null
+++ b/Engine/Source/Runtime/Core/Public/Async/InheritedContext.h
@@ -0,0 +1,93 @@
+// Copyright Epic Games, Inc. All Rights Reserved.
+
+#pragma once
+
+#include "HAL/LowLevelMemTracker.h"
+#include "ProfilingDebugging/MetadataTrace.h"
+#include "ProfilingDebugging/TagTrace.h"
+
+namespace UE
+{
+	// Restores an inherited contex for the current scope.
+	// An instance must be obtained by calling `FInheritedContextBase::RestoreInheritedContext()`
+	class FInheritedContextScope
+	{
+	private:
+		UE_NONCOPYABLE(FInheritedContextScope);
+
+		friend class FInheritedContextBase; // allow construction only by `FInheritedContextBase`
+
+		FInheritedContextScope(
+#if ENABLE_LOW_LEVEL_MEM_TRACKER
+			const UE::LLMPrivate::FTagData* InInheritedLLMTag
+	#if UE_MEMORY_TAGS_TRACE_ENABLED
+			,
+	#endif
+#endif
+#if UE_MEMORY_TAGS_TRACE_ENABLED
+			int32 InInheritedMemTag
+#endif
+		)
+#if ENABLE_LOW_LEVEL_MEM_TRACKER
+			: LLMScope(InInheritedLLMTag, false /* bIsStatTag */, ELLMTagSet::None, ELLMTracker::Default)
+	#if UE_MEMORY_TAGS_TRACE_ENABLED
+			,
+	#endif
+#endif
+#if UE_MEMORY_TAGS_TRACE_ENABLED
+	#if !ENABLE_LOW_LEVEL_MEM_TRACKER
+			:
+	#endif
+			MemScope(InInheritedMemTag)
+#endif
+		{
+		}
+
+	private:
+#if ENABLE_LOW_LEVEL_MEM_TRACKER
+		FLLMScope LLMScope;
+#endif
+
+#if UE_MEMORY_TAGS_TRACE_ENABLED
+		FMemScope MemScope;
+#endif
+	};
+
+	// this class extends the inherited context (see private members for what the inherited context is) to cover async execution. 
+	// Is intended to be used as a base class, if the inherited context is compiled out it takes 0 space
+	class FInheritedContextBase
+	{
+	public:
+		// must be called in the inherited context, e.g. on launching an async task
+		void CaptureInheritedContext()
+		{
+#if ENABLE_LOW_LEVEL_MEM_TRACKER
+			InheritedLLMTag = FLowLevelMemTracker::bIsDisabled ? nullptr : FLowLevelMemTracker::Get().GetActiveTagData(ELLMTracker::Default);
+#endif
+
+#if UE_MEMORY_TAGS_TRACE_ENABLED
+			InheritedMemTag = MemoryTrace_GetActiveTag();
+#endif
+
+#if UE_TRACE_METADATA_ENABLED
+			InheritedMetadataId = UE_TRACE_METADATA_SAVE_STACK();
+#endif
+		}
+
+		// must be called where the inherited context should be restored, e.g. at the start of an async task execution
+		UE_NODISCARD CORE_API FInheritedContextScope RestoreInheritedContext();
+
+	private:
+#if ENABLE_LOW_LEVEL_MEM_TRACKER
+		const UE::LLMPrivate::FTagData* InheritedLLMTag;
+#endif
+
+#if UE_MEMORY_TAGS_TRACE_ENABLED
+		int32 InheritedMemTag;
+#endif
+
+#if UE_TRACE_METADATA_ENABLED
+		uint32 InheritedMetadataId;
+#endif
+	};
+}
\ No newline at end of file
diff --git a/Engine/Source/Runtime/Core/Public/Async/ParallelFor.h b/Engine/Source/Runtime/Core/Public/Async/ParallelFor.h
--- a/Engine/Source/Runtime/Core/Public/Async/ParallelFor.h
+++ b/Engine/Source/Runtime/Core/Public/Async/ParallelFor.h
@@ -194,8 +194,13 @@ namespace ParallelForImpl
 		};
 
 		//shared data between tasks
-		struct alignas(PLATFORM_CACHE_LINE_SIZE) FParallelForData : public TConcurrentLinearObject<FParallelForData, FTaskGraphBlockAllocationTag>, public FThreadSafeRefCountedObject
+		struct alignas(PLATFORM_CACHE_LINE_SIZE) FParallelForData 
+			: public TConcurrentLinearObject<FParallelForData, FTaskGraphBlockAllocationTag>
+			, public FThreadSafeRefCountedObject
+			, private UE::FInheritedContextBase
 		{
+			using UE::FInheritedContextBase::RestoreInheritedContext;
+
 			FParallelForData(int32 InNum, int32 InBatchSize, int32 InNumBatches, int32 InNumWorkers, const TArrayView<ContextType>& InContexts, const BodyType& InBody, FEventRef& InFinishedSignal)
 				: Num(InNum)
 				, BatchSize(InBatchSize)
@@ -207,10 +212,7 @@ namespace ParallelForImpl
 				IncompleteBatches.store(NumBatches, std::memory_order_relaxed);
 				Tasks.AddDefaulted(InNumWorkers);
 
-#if UE_MEMORY_TAGS_TRACE_ENABLED
-				InheritedTraceTag = MemoryTrace_GetActiveTag();
-#endif
-				LLM(InheritedLLMTag = FLowLevelMemTracker::bIsDisabled ? nullptr : FLowLevelMemTracker::Get().GetActiveTagData(ELLMTracker::Default));
+				CaptureInheritedContext();
 			}
 
 			~FParallelForData()
@@ -226,10 +228,6 @@ namespace ParallelForImpl
 			int32 Num;
 			int32 BatchSize;
 			int32 NumBatches;
-#if UE_MEMORY_TAGS_TRACE_ENABLED
-			int32 InheritedTraceTag;
-#endif
-			LLM(const UE::LLMPrivate::FTagData* InheritedLLMTag);
 			const TArrayView<ContextType>& Contexts;
 			const BodyType& Body;
 			FEventRef& FinishedSignal;
@@ -276,8 +274,7 @@ namespace ParallelForImpl
 
 			inline bool operator()(const bool bIsMaster = false, const TCHAR* DebugName = nullptr) const noexcept
 			{
-				LLM_SCOPE(Data->InheritedLLMTag);
-				UE_MEMSCOPE(Data->InheritedTraceTag);
+				UE::FInheritedContextScope InheritedContextScope = Data->RestoreInheritedContext();
 				FMemMark Mark(FMemStack::Get());
 
 				if(DebugName == nullptr)
diff --git a/Engine/Source/Runtime/Core/Public/Async/TaskGraphInterfaces.h b/Engine/Source/Runtime/Core/Public/Async/TaskGraphInterfaces.h
--- a/Engine/Source/Runtime/Core/Public/Async/TaskGraphInterfaces.h
+++ b/Engine/Source/Runtime/Core/Public/Async/TaskGraphInterfaces.h
@@ -30,6 +30,7 @@
 
 #include "Async/TaskTrace.h"
 #include "Tasks/TaskPrivate.h"
+#include "Async/InheritedContext.h"
 
 #if !defined(STATS)
 #error "STATS must be defined as either zero or one."
@@ -823,7 +824,7 @@ struct FTaskBlockAllocationTag : FDefaultBlockAllocationTag
  *	Tasks go through a very specific life stage progression, and this is verified.
  **/
 
-class FBaseGraphTask
+class FBaseGraphTask : private UE::FInheritedContextBase
 {
 protected:
 	/** 
@@ -835,13 +836,7 @@ protected:
 		, NumberOfPrerequistitesOutstanding(InNumberOfPrerequistitesOutstanding + 1) // + 1 is not a prerequisite, it is a lock to prevent it from executing while it is getting prerequisites, one it is safe to execute, call PrerequisitesComplete
 	{
 		checkThreadGraph(LifeStage.Increment() == int32(LS_Contructed));
-#if UE_MEMORY_TAGS_TRACE_ENABLED
-		InheritedTraceTag = MemoryTrace_GetActiveTag();
-#endif
-#if UE_TRACE_METADATA_ENABLED
-		InheritedMetadataId = UE_TRACE_METADATA_SAVE_STACK();
-#endif
-		LLM(InheritedLLMTag = FLowLevelMemTracker::bIsDisabled ? nullptr : FLowLevelMemTracker::Get().GetActiveTagData(ELLMTracker::Default));
+		CaptureInheritedContext();
 	}
 	/** 
 	 *	Sets the desired execution thread. This is not part of the constructor because this information may not be known quite yet duiring construction.
@@ -952,10 +947,9 @@ private:
 	 **/
 	FORCEINLINE void Execute(TArray<FBaseGraphTask*>& NewTasks, ENamedThreads::Type CurrentThread, bool bDeleteOnCompletion)
 	{
-		LLM_SCOPE(InheritedLLMTag);
-		UE_MEMSCOPE(InheritedTraceTag);
-		UE_TRACE_METADATA_RESTORE_STACK(InheritedMetadataId);
 		checkThreadGraph(LifeStage.Increment() == int32(LS_Executing));
+
+		UE::FInheritedContextScope InheritedContextScope = RestoreInheritedContext();
 		ExecuteTask(NewTasks, CurrentThread, bDeleteOnCompletion);
 	}
 
@@ -994,13 +988,6 @@ private:
 	FThreadSafeCounter			LifeStage;
 
 #endif
-#if UE_MEMORY_TAGS_TRACE_ENABLED
-	int32 InheritedTraceTag;
-#endif
-#if UE_TRACE_METADATA_ENABLED
-	uint32 InheritedMetadataId;
-#endif
-	LLM(const UE::LLMPrivate::FTagData* InheritedLLMTag);
 
 #if UE_TASK_TRACE_ENABLED
 	TaskTrace::FId TraceId;
diff --git a/Engine/Source/Runtime/Core/Public/IO/IoDispatcherBackend.h b/Engine/Source/Runtime/Core/Public/IO/IoDispatcherBackend.h
--- a/Engine/Source/Runtime/Core/Public/IO/IoDispatcherBackend.h
+++ b/Engine/Source/Runtime/Core/Public/IO/IoDispatcherBackend.h
@@ -3,9 +3,8 @@
 #pragma once
 
 #include "IoDispatcher.h"
-#include "HAL/LowLevelMemTracker.h"
+#include "Async/InheritedContext.h"
 #include "Misc/Optional.h"
-#include "ProfilingDebugging/TagTrace.h"
 #include "ProfilingDebugging/CountersTrace.h"
 #include "ProfilingDebugging/CsvProfiler.h"
 
@@ -13,15 +12,11 @@
 
 struct FIoContainerHeader;
 
-class FIoRequestImpl
+class FIoRequestImpl : private UE::FInheritedContextBase
 {
 public:
 	FIoRequestImpl* NextRequest = nullptr;
 	void* BackendData = nullptr;
-	LLM(const UE::LLMPrivate::FTagData* InheritedLLMTag);
-#if UE_MEMORY_TAGS_TRACE_ENABLED
-	int32 InheritedTraceTag;
-#endif
 	FIoChunkId ChunkId;
 	FIoReadOptions Options;
 	int32 Priority = 0;
@@ -29,10 +24,7 @@ public:
 	FIoRequestImpl(class FIoRequestAllocator& InAllocator)
 		: Allocator(InAllocator)
 	{
-		LLM(InheritedLLMTag = FLowLevelMemTracker::bIsDisabled ? nullptr : FLowLevelMemTracker::Get().GetActiveTagData(ELLMTracker::Default));
-#if UE_MEMORY_TAGS_TRACE_ENABLED
-		InheritedTraceTag = MemoryTrace_GetActiveTag();
-#endif
+		CaptureInheritedContext();
 	}
 
 	bool IsCancelled() const
diff --git a/Engine/Source/Runtime/Core/Public/ProfilingDebugging/TagTrace.h b/Engine/Source/Runtime/Core/Public/ProfilingDebugging/TagTrace.h
--- a/Engine/Source/Runtime/Core/Public/ProfilingDebugging/TagTrace.h
+++ b/Engine/Source/Runtime/Core/Public/ProfilingDebugging/TagTrace.h
@@ -62,7 +62,7 @@ template<typename TagType>
 class FDefaultMemScope : public FMemScope
 {
 public:
-	CORE_API FDefaultMemScope(TagType InTag)
+	FDefaultMemScope(TagType InTag)
 		: FMemScope(InTag, MemoryTrace_GetActiveTag() == 0)
 	{
 		
diff --git a/Engine/Source/Runtime/Core/Public/Tasks/TaskPrivate.h b/Engine/Source/Runtime/Core/Public/Tasks/TaskPrivate.h
--- a/Engine/Source/Runtime/Core/Public/Tasks/TaskPrivate.h
+++ b/Engine/Source/Runtime/Core/Public/Tasks/TaskPrivate.h
@@ -30,6 +30,7 @@
 #include "Templates/TypeCompatibleBytes.h"
 #include "Templates/UnrealTemplate.h"
 #include "Templates/UnrealTypeTraits.h"
+#include "Async/InheritedContext.h"
 
 #include <atomic>
 #include <type_traits>
@@ -94,7 +95,7 @@ namespace UE::Tasks
 		// Implements internal logic of task prerequisites, nested tasks and deep task retraction.
 		// Implements intrusive ref-counting and so can be used with TRefCountPtr.
 		// It doesn't store task body, instead it expects a derived class to provide a task body as a parameter to `TryExecute` method. @see TExecutableTask
-		class FTaskBase
+		class FTaskBase : private UE::FInheritedContextBase
 		{
 			UE_NONCOPYABLE(FTaskBase);
 
@@ -152,6 +153,8 @@ namespace UE::Tasks
 					}
 				);
 				ExtendedPriority = InExtendedPriority;
+
+				CaptureInheritedContext();
 			}
 
 			virtual ~FTaskBase()
@@ -161,8 +164,8 @@ namespace UE::Tasks
 			}
 
 			// will be called to execute the task, must be implemented by a derived class that should call `FTaskBase::TryExecute` and pass the task body
-			// @see TExecutableTask::TryExecuteTask
-			virtual bool TryExecuteTask() = 0;
+			// @see TExecutableTask::TryExecuteTaskVirtual
+			virtual bool TryExecuteTaskVirtual() = 0;
 
 		public:
 			// returns true if it's valid to wait for the task completion.
@@ -591,8 +594,11 @@ namespace UE::Tasks
 			CORE_API void StartPipeExecution();
 			CORE_API void FinishPipeExecution();
 
-			// the task is already executed but it's not completed yet. This method sets completion flag if there're no pending nested tasks.
-			// The task can be deleted as the result of this call.
+			bool TryExecuteTask()
+			{
+				UE::FInheritedContextScope InheritedContextScope = RestoreInheritedContext();
+				return TryExecuteTaskVirtual();
+			}
 
 		private:
 			EExtendedTaskPriority ExtendedPriority; // internal priorities, if any
@@ -665,7 +671,7 @@ namespace UE::Tasks
 			UE_NONCOPYABLE(TExecutableTaskBase);
 
 		public:
-			virtual bool TryExecuteTask() override
+			virtual bool TryExecuteTaskVirtual() override
 			{
 				return FTaskBase::TryExecute(
 					[](FTaskBase& Task)
@@ -700,7 +706,7 @@ namespace UE::Tasks
 			UE_NONCOPYABLE(TExecutableTaskBase);
 
 		public:
-			virtual bool TryExecuteTask() override
+			virtual bool TryExecuteTaskVirtual() override
 			{
 				return TryExecute(
 					[](FTaskBase& Task)
@@ -784,7 +790,7 @@ namespace UE::Tasks
 				Init(InDebugName, ETaskPriority::Normal, EExtendedTaskPriority::TaskEvent);
 			}
 
-			virtual bool TryExecuteTask() override
+			virtual bool TryExecuteTaskVirtual() override
 			{
 				checkNoEntry(); // never executed because it doesn't have a task body
 				return true;
