fixed server compilation that happens to have UE_MEMORY_TAGS_TRACE_ENABLED enabled but not UE_TRACE_ENABLED

diff --git a/Engine/Source/Runtime/Core/Private/Async/InheritedContext.cpp b/Engine/Source/Runtime/Core/Private/Async/InheritedContext.cpp
--- a/Engine/Source/Runtime/Core/Private/Async/InheritedContext.cpp
+++ b/Engine/Source/Runtime/Core/Private/Async/InheritedContext.cpp
@@ -12,11 +12,11 @@ namespace UE
 		return FInheritedContextScope{
 #if ENABLE_LOW_LEVEL_MEM_TRACKER
 				InheritedLLMTag
-	#if UE_MEMORY_TAGS_TRACE_ENABLED
+	#if UE_MEMORY_TAGS_TRACE_ENABLED && UE_TRACE_ENABLED
 				,
 	#endif
 #endif
-#if UE_MEMORY_TAGS_TRACE_ENABLED
+#if UE_MEMORY_TAGS_TRACE_ENABLED && UE_TRACE_ENABLED
 				InheritedMemTag
 #endif
 		};
diff --git a/Engine/Source/Runtime/Core/Public/Async/InheritedContext.h b/Engine/Source/Runtime/Core/Public/Async/InheritedContext.h
--- a/Engine/Source/Runtime/Core/Public/Async/InheritedContext.h
+++ b/Engine/Source/Runtime/Core/Public/Async/InheritedContext.h
@@ -20,21 +20,21 @@ namespace UE
 		FInheritedContextScope(
 #if ENABLE_LOW_LEVEL_MEM_TRACKER
 			const UE::LLMPrivate::FTagData* InInheritedLLMTag
-	#if UE_MEMORY_TAGS_TRACE_ENABLED
+	#if UE_MEMORY_TAGS_TRACE_ENABLED && UE_TRACE_ENABLED
 			,
 	#endif
 #endif
-#if UE_MEMORY_TAGS_TRACE_ENABLED
+#if UE_MEMORY_TAGS_TRACE_ENABLED && UE_TRACE_ENABLED
 			int32 InInheritedMemTag
 #endif
 		)
 #if ENABLE_LOW_LEVEL_MEM_TRACKER
 			: LLMScope(InInheritedLLMTag, false /* bIsStatTag */, ELLMTagSet::None, ELLMTracker::Default)
-	#if UE_MEMORY_TAGS_TRACE_ENABLED
+	#if UE_MEMORY_TAGS_TRACE_ENABLED && UE_TRACE_ENABLED
 			,
 	#endif
 #endif
-#if UE_MEMORY_TAGS_TRACE_ENABLED
+#if UE_MEMORY_TAGS_TRACE_ENABLED && UE_TRACE_ENABLED
 	#if !ENABLE_LOW_LEVEL_MEM_TRACKER
 			:
 	#endif
@@ -48,7 +48,7 @@ namespace UE
 		FLLMScope LLMScope;
 #endif
 
-#if UE_MEMORY_TAGS_TRACE_ENABLED
+#if UE_MEMORY_TAGS_TRACE_ENABLED && UE_TRACE_ENABLED
 		FMemScope MemScope;
 #endif
 	};
@@ -65,7 +65,7 @@ namespace UE
 			InheritedLLMTag = FLowLevelMemTracker::bIsDisabled ? nullptr : FLowLevelMemTracker::Get().GetActiveTagData(ELLMTracker::Default);
 #endif
 
-#if UE_MEMORY_TAGS_TRACE_ENABLED
+#if UE_MEMORY_TAGS_TRACE_ENABLED && UE_TRACE_ENABLED
 			InheritedMemTag = MemoryTrace_GetActiveTag();
 #endif
 
@@ -82,7 +82,7 @@ namespace UE
 		const UE::LLMPrivate::FTagData* InheritedLLMTag;
 #endif
 
-#if UE_MEMORY_TAGS_TRACE_ENABLED
+#if UE_MEMORY_TAGS_TRACE_ENABLED && UE_TRACE_ENABLED
 		int32 InheritedMemTag;
 #endif
 
