moved a check for successful thread creation and added a message that is supposed to help diagnosing an underlying problem

diff --git a/Engine/Source/Runtime/Core/Private/HAL/Thread.cpp b/Engine/Source/Runtime/Core/Private/HAL/Thread.cpp
--- a/Engine/Source/Runtime/Core/Private/HAL/Thread.cpp
+++ b/Engine/Source/Runtime/Core/Private/HAL/Thread.cpp
@@ -29,8 +29,8 @@ public:
 		, RunnableThread(IsForkable == FThread::Forkable ? FForkProcessHelper::CreateForkableThread(this, ThreadName, StackSize, ThreadPriority, ThreadAffinity.ThreadAffinityMask)
 			: FRunnableThread::Create(this, ThreadName, StackSize, ThreadPriority, ThreadAffinity.ThreadAffinityMask))
 	{
+		checkf(IsJoinable(), TEXT("Failed to create a thread. Check if it ran out of memory or a similar \"out of system resources\" issue"));
 		RunnableThread->SetThreadAffinity(ThreadAffinity);
-		check(IsJoinable());
 	}
 
 	// Provides a reference to self. This can't be done by `SharedFromThis` in the constructor because the instance is not constructed yet
