Fix corelimit parsing; it was being ignored without physicalcorelimit.
#rnx

diff --git a/Engine/Source/Runtime/Core/Private/GenericPlatform/GenericPlatformMisc.cpp b/Engine/Source/Runtime/Core/Private/GenericPlatform/GenericPlatformMisc.cpp
--- a/Engine/Source/Runtime/Core/Private/GenericPlatform/GenericPlatformMisc.cpp
+++ b/Engine/Source/Runtime/Core/Private/GenericPlatform/GenericPlatformMisc.cpp
@@ -1559,8 +1559,8 @@ void FGenericPlatformMisc::GetConfiguredCoreLimits(int32 PlatformNumPhysicalCore
 	}
 	if (LegacyCoreLimit > 0)
 	{
-		PhysicalCoreLimit = FMath::Min(PhysicalCoreLimit, LegacyCoreLimit);
-		LogicalCoreLimit = FMath::Min(LogicalCoreLimit, LegacyCoreLimit);
+		PhysicalCoreLimit = PhysicalCoreLimit == 0 ? LegacyCoreLimit : FMath::Min(PhysicalCoreLimit, LegacyCoreLimit);
+		LogicalCoreLimit = LogicalCoreLimit == 0 ? LegacyCoreLimit : FMath::Min(LogicalCoreLimit, LegacyCoreLimit);
 	}
 
 	bOutFullyInitialized = true;
diff --git a/Engine/Source/Runtime/Core/Private/Windows/WindowsPlatformMisc.cpp b/Engine/Source/Runtime/Core/Private/Windows/WindowsPlatformMisc.cpp
--- a/Engine/Source/Runtime/Core/Private/Windows/WindowsPlatformMisc.cpp
+++ b/Engine/Source/Runtime/Core/Private/Windows/WindowsPlatformMisc.cpp
@@ -1778,18 +1778,18 @@ int32 FWindowsPlatformMisc::NumberOfCores()
 	QueryCpuInformation(GroupDesc, NumaNodeCount, NumCores, LogicalProcessorCount);
 
 	bool bLimitsInitialized;
-	int32 CoreLimitCores;
-	int32 CoreLimitHyperThreads;
-	bool bSetCoreCountToHyperThreadCount;
-	GetConfiguredCoreLimits(NumCores, LogicalProcessorCount, bLimitsInitialized, CoreLimitCores,
-		CoreLimitHyperThreads, bSetCoreCountToHyperThreadCount);
+	int32 PhysicalCoreLimit;
+	int32 LogicalCoreLimit;
+	bool bSetPhysicalCountToLogicalCount;
+	GetConfiguredCoreLimits(NumCores, LogicalProcessorCount, bLimitsInitialized, PhysicalCoreLimit,
+		LogicalCoreLimit, bSetPhysicalCountToLogicalCount);
 
-	CoreCount = bSetCoreCountToHyperThreadCount ? LogicalProcessorCount : NumCores;
+	CoreCount = bSetPhysicalCountToLogicalCount ? LogicalProcessorCount : NumCores;
 
 	// Optionally limit number of threads (we don't necessarily scale super well with very high core counts)
-	if (CoreLimitCores > 0)
+	if (PhysicalCoreLimit > 0)
 	{
-		CoreCount = FMath::Min(CoreCount, CoreLimitCores);
+		CoreCount = FMath::Min(CoreCount, PhysicalCoreLimit);
 	}
 
 	return CoreCount;
@@ -1826,18 +1826,18 @@ int32 FWindowsPlatformMisc::NumberOfCoresIncludingHyperthreads()
 	QueryCpuInformation(GroupDesc, NumaNodeCount, NumCores, LogicalProcessorCount);
 
 	bool bLimitsInitialized;
-	int32 CoreLimitCores;
-	int32 CoreLimitHyperThreads;
-	bool bSetCoreCountToHyperThreadCount;
-	GetConfiguredCoreLimits(NumCores, LogicalProcessorCount, bLimitsInitialized, CoreLimitCores,
-		CoreLimitHyperThreads, bSetCoreCountToHyperThreadCount);
+	int32 PhysicalCoreLimit;
+	int32 LogicalCoreLimit;
+	bool bSetPhysicalCountToLogicalCount;
+	GetConfiguredCoreLimits(NumCores, LogicalProcessorCount, bLimitsInitialized, PhysicalCoreLimit,
+		LogicalCoreLimit, bSetPhysicalCountToLogicalCount);
 
 	CoreCount = LogicalProcessorCount;
 
 	// Optionally limit number of threads (we don't necessarily scale super well with very high core counts)
-	if (CoreLimitHyperThreads > 0)
+	if (LogicalCoreLimit > 0)
 	{
-		CoreCount = FMath::Min(CoreCount, CoreLimitHyperThreads);
+		CoreCount = FMath::Min(CoreCount, LogicalCoreLimit);
 	}
 
 	return CoreCount;
