- Add shader parameter binding for velocity

#rnx

diff --git a/Engine/Plugins/FX/Niagara/Source/Niagara/Private/DataInterface/NiagaraDataInterfaceActorComponent.cpp b/Engine/Plugins/FX/Niagara/Source/Niagara/Private/DataInterface/NiagaraDataInterfaceActorComponent.cpp
--- a/Engine/Plugins/FX/Niagara/Source/Niagara/Private/DataInterface/NiagaraDataInterfaceActorComponent.cpp
+++ b/Engine/Plugins/FX/Niagara/Source/Niagara/Private/DataInterface/NiagaraDataInterfaceActorComponent.cpp
@@ -253,6 +253,7 @@ void UNiagaraDataInterfaceActorComponent::SetShaderParameters(const FNiagaraData
 	ShaderParameters->Matrix	= (FMatrix44f)InstanceData.CachedTransform.ToMatrixWithScale();
 	ShaderParameters->Rotation	= (FQuat4f)InstanceData.CachedTransform.GetRotation();
 	ShaderParameters->Scale		= (FVector3f)InstanceData.CachedTransform.GetScale3D();
+	ShaderParameters->Velocity	= InstanceData.CachedVelocity;
 }
 
 bool UNiagaraDataInterfaceActorComponent::InitPerInstanceData(void* PerInstanceData, FNiagaraSystemInstance* SystemInstance)
diff --git a/Engine/Plugins/FX/Niagara/Source/Niagara/Private/DataInterface/NiagaraDataInterfaceActorComponent.h b/Engine/Plugins/FX/Niagara/Source/Niagara/Private/DataInterface/NiagaraDataInterfaceActorComponent.h
--- a/Engine/Plugins/FX/Niagara/Source/Niagara/Private/DataInterface/NiagaraDataInterfaceActorComponent.h
+++ b/Engine/Plugins/FX/Niagara/Source/Niagara/Private/DataInterface/NiagaraDataInterfaceActorComponent.h
@@ -20,6 +20,7 @@ class NIAGARA_API UNiagaraDataInterfaceActorComponent : public UNiagaraDataInter
 		SHADER_PARAMETER(FMatrix44f,	Matrix)
 		SHADER_PARAMETER(FQuat4f,		Rotation)
 		SHADER_PARAMETER(FVector3f,		Scale)
+		SHADER_PARAMETER(FVector3f,		Velocity)
 	END_SHADER_PARAMETER_STRUCT();
 
 public:
diff --git a/Engine/Source/Runtime/Core/Private/UObject/FortniteMainBranchObjectVersion.cpp b/Engine/Source/Runtime/Core/Private/UObject/FortniteMainBranchObjectVersion.cpp
--- a/Engine/Source/Runtime/Core/Private/UObject/FortniteMainBranchObjectVersion.cpp
+++ b/Engine/Source/Runtime/Core/Private/UObject/FortniteMainBranchObjectVersion.cpp
@@ -11,7 +11,7 @@ TMap<FGuid, FGuid> FFortniteMainBranchObjectVersion::GetSystemGuids()
 	SystemGuids.Add(DevGuids.MATERIALSHADERMAP_DERIVEDDATA_VER, FGuid("1DA85DA5733C4A489A1BE117CDC5ACCC"));
 	SystemGuids.Add(DevGuids.NANITE_DERIVEDDATA_VER, FGuid("DBF5D892C84143D692AF77A9E9752124"));
 	SystemGuids.Add(DevGuids.NIAGARASHADERMAP_DERIVEDDATA_VER, FGuid("2289D5116CF94BC0AFEAEC541468E645"));
-	SystemGuids.Add(DevGuids.Niagara_LatestScriptCompileVersion, FGuid("A731B02963264F99A7CFD5B4AC060E4C"));
+	SystemGuids.Add(DevGuids.Niagara_LatestScriptCompileVersion, FGuid("A9222A671841874BA54368C57E3BCE82"));
 	SystemGuids.Add(DevGuids.SkeletalMeshDerivedDataVersion, FGuid("ACF593EAAE354FCBB34CF44F5AA1BFD2"));
 	SystemGuids.Add(DevGuids.STATICMESH_DERIVEDDATA_VER, FGuid("18E028FCF1BC469385806AA5F8AE07CC"));
 	return SystemGuids;
