removed a "check" because it relies on an insufficiently strict memory order ("release" on the previous line). reported by TSAN. the check is not important enough to pessimise the algo for it

diff --git a/Engine/Source/Runtime/Core/Public/Containers/DepletableMpscQueue.h b/Engine/Source/Runtime/Core/Public/Containers/DepletableMpscQueue.h
--- a/Engine/Source/Runtime/Core/Public/Containers/DepletableMpscQueue.h
+++ b/Engine/Source/Runtime/Core/Public/Containers/DepletableMpscQueue.h
@@ -65,7 +65,6 @@ namespace UE
 
 			FNode* Prev = Tail.exchange(New, std::memory_order_release); // `release` to make sure the new node is fully constructed before
 			// it becomes visible to the consumer
-			check(Prev->Next.load(std::memory_order_relaxed) == nullptr); // `Tail` is assigned before its Next
 
 			Prev->Next.store(New, std::memory_order_relaxed);
 
