fixed double destruction in Tasks System that was causing repeating "task destroyed" traces in TasksInsight

diff --git a/Engine/Source/Runtime/Core/Public/Tasks/TaskPrivate.h b/Engine/Source/Runtime/Core/Public/Tasks/TaskPrivate.h
--- a/Engine/Source/Runtime/Core/Public/Tasks/TaskPrivate.h
+++ b/Engine/Source/Runtime/Core/Public/Tasks/TaskPrivate.h
@@ -754,14 +754,13 @@ namespace UE::Tasks
 				return new TExecutableTask(InDebugName, MoveTemp(TaskBody), InPriority, InExtendedPriority);
 			}
 
-			void* operator new(size_t Size)
+			static void* operator new(size_t Size)
 			{
 				return Size <= SmallTaskSize ? SmallTaskAllocator.Allocate() : GMalloc->Malloc(sizeof(TExecutableTask), PLATFORM_CACHE_LINE_SIZE);
 			}
 
-			void operator delete(void* Ptr, size_t Size)
+			static void operator delete(void* Ptr, size_t Size)
 			{
-				DestructItem((TExecutableTask*)Ptr);
 				Size <= SmallTaskSize ? SmallTaskAllocator.Free(Ptr) : GMalloc->Free(Ptr);
 			}
 		};
