Add OnCSVProfileEndRequested delegate that fires right before we end a capture, allowing final information to be written.

diff --git a/Engine/Source/Runtime/Core/Private/ProfilingDebugging/CsvProfiler.cpp b/Engine/Source/Runtime/Core/Private/ProfilingDebugging/CsvProfiler.cpp
--- a/Engine/Source/Runtime/Core/Private/ProfilingDebugging/CsvProfiler.cpp
+++ b/Engine/Source/Runtime/Core/Private/ProfilingDebugging/CsvProfiler.cpp
@@ -3169,6 +3169,9 @@ TSharedFuture<FString> FCsvProfiler::EndCapture(FGraphEventRef EventToSignal)
 
 	check(IsInGameThread());
 
+	// Fire before we copy the metadata so it gives other systems a chance to write any final information.
+	OnCSVProfileEndRequestedDelegate.Broadcast();
+
 	TPromise<FString>* Completion = new TPromise<FString>([EventToSignal]()
 	{
 		if (EventToSignal)
diff --git a/Engine/Source/Runtime/Core/Public/ProfilingDebugging/CsvProfiler.h b/Engine/Source/Runtime/Core/Public/ProfilingDebugging/CsvProfiler.h
--- a/Engine/Source/Runtime/Core/Public/ProfilingDebugging/CsvProfiler.h
+++ b/Engine/Source/Runtime/Core/Public/ProfilingDebugging/CsvProfiler.h
@@ -426,6 +426,10 @@ public:
 	DECLARE_MULTICAST_DELEGATE(FOnCSVProfileFirstFrame);
 	FOnCSVProfileFirstFrame& OnCSVProfileFirstFrame() { return OnCSVProfileFirstFrameDelegate; }
 
+	/** Called when the capture is requested to end, allowing any final information to be written (eg. metadata). */
+	DECLARE_MULTICAST_DELEGATE(FOnCSVProfileEndRequested);
+	FOnCSVProfileEndRequested& OnCSVProfileEndRequested() { return OnCSVProfileEndRequestedDelegate; }
+
 	DECLARE_MULTICAST_DELEGATE(FOnCSVProfileEnd);
 	FOnCSVProfileEnd& OnCSVProfileEnd() { return OnCSVProfileEndDelegate; }
 	
@@ -483,6 +487,7 @@ private:
 
 	FOnCSVProfileStart OnCSVProfileStartDelegate;
 	FOnCSVProfileFirstFrame OnCSVProfileFirstFrameDelegate;
+	FOnCSVProfileEndRequested OnCSVProfileEndRequestedDelegate;
 	FOnCSVProfileEnd OnCSVProfileEndDelegate;
 	
 	FOnCSVProfileFinished OnCSVProfileFinishedDelegate;
