[Insights] Added "clear metadata scope" (UE_TRACE_METADATA_CLEAR_SCOPE) for the Trace system memory allocation events.

diff --git a/Engine/Source/Runtime/Core/Private/ProfilingDebugging/MemoryTrace.cpp b/Engine/Source/Runtime/Core/Private/ProfilingDebugging/MemoryTrace.cpp
--- a/Engine/Source/Runtime/Core/Private/ProfilingDebugging/MemoryTrace.cpp
+++ b/Engine/Source/Runtime/Core/Private/ProfilingDebugging/MemoryTrace.cpp
@@ -4,6 +4,7 @@
 #include "Containers/StringView.h"
 #include "HAL/PlatformTime.h"
 #include "ProfilingDebugging/CallstackTrace.h"
+#include "ProfilingDebugging/MetadataTrace.h"
 #include "ProfilingDebugging/TagTrace.h"
 #include "ProfilingDebugging/TraceMalloc.h"
 #include "Trace/Detail/EventNode.h"
@@ -680,6 +681,7 @@ void* FTraceMalloc::Malloc(SIZE_T Count, uint32 Alignment)
 	const uint32 AlignmentPow2 = uint32(FPlatformMath::CountTrailingZeros(Alignment));
 	const uint32 Alignment_SizeLower = (AlignmentPow2 << SizeShift) | uint32(Size & ((1 << SizeShift) - 1));
 
+	UE_TRACE_METADATA_CLEAR_SCOPE();
 	UE_MEMSCOPE(TRACE_TAG);
 	
 	UE_TRACE_LOG(Memory, Alloc, MemAllocChannel)
@@ -704,6 +706,7 @@ void* FTraceMalloc::Realloc(void* Original, SIZE_T Count, uint32 Alignment)
 	const uint32 AlignmentPow2 = uint32(FPlatformMath::CountTrailingZeros(Alignment));
 	const uint32 Alignment_SizeLower = (AlignmentPow2 << SizeShift) | uint32(Size & ((1 << SizeShift) - 1));
 
+	UE_TRACE_METADATA_CLEAR_SCOPE();
 	UE_MEMSCOPE(TRACE_TAG);
 	
 	UE_TRACE_LOG(Memory, ReallocFree, MemAllocChannel)
