Added support for specifying localization targets that should only be loaded in a cooked editor

diff --git a/Engine/Source/Runtime/Core/Private/Internationalization/LocalizationResourceTextSource.cpp b/Engine/Source/Runtime/Core/Private/Internationalization/LocalizationResourceTextSource.cpp
--- a/Engine/Source/Runtime/Core/Private/Internationalization/LocalizationResourceTextSource.cpp
+++ b/Engine/Source/Runtime/Core/Private/Internationalization/LocalizationResourceTextSource.cpp
@@ -47,24 +47,41 @@ void FLocalizationResourceTextSource::GetLocalizedCultureNames(const ELocalizati
 
 void FLocalizationResourceTextSource::LoadLocalizedResources(const ELocalizationLoadFlags InLoadFlags, TArrayView<const FString> InPrioritizedCultures, FTextLocalizationResource& InOutNativeResource, FTextLocalizationResource& InOutLocalizedResource)
 {
-	auto GetGameLocalizationPaths = [this]()
+	auto AppendChunkedLocalizationPaths = [this](TArray<FString>& OutLocalizationPaths)
 	{
-		TArray<FString> LocalizationPaths = FPaths::GetGameLocalizationPaths();
-
 		if (ChunkIds.Num() > 0)
 		{
-			const TArray<FString> ChunkedLocalizationTargets = GetChunkedLocalizationTargets();
+			// Filter the list of chunked targets against the list of targets we're actually loading data for
+			TArray<FString> ChunkedLocalizationTargets = GetChunkedLocalizationTargets();
+			ChunkedLocalizationTargets.RemoveAll([&OutLocalizationPaths](const FString& LocalizationTarget)
+			{
+				// Note: We only allow game localization targets to be chunked, and the layout is assumed to follow our standard pattern (as used by the localization dashboard and FLocTextHelper)
+				return !OutLocalizationPaths.Contains(FPaths::ProjectContentDir() / TEXT("Localization") / LocalizationTarget);
+			});
+
 			for (const FString& LocalizationTarget : ChunkedLocalizationTargets)
 			{
 				for (const int32 ChunkId : ChunkIds)
 				{
 					// Note: We only allow game localization targets to be chunked, and the layout is assumed to follow our standard pattern (as used by the localization dashboard and FLocTextHelper)
 					const FString LocalizationTargetForChunk = TextLocalizationResourceUtil::GetLocalizationTargetNameForChunkId(LocalizationTarget, ChunkId);
-					LocalizationPaths.Add(FPaths::ProjectContentDir() / TEXT("Localization") / LocalizationTargetForChunk);
+					OutLocalizationPaths.Add(FPaths::ProjectContentDir() / TEXT("Localization") / LocalizationTargetForChunk);
 				}
 			}
 		}
+	};
 
+	auto GetGameLocalizationPaths = [&AppendChunkedLocalizationPaths]()
+	{
+		TArray<FString> LocalizationPaths = FPaths::GetGameLocalizationPaths();
+		AppendChunkedLocalizationPaths(LocalizationPaths);
+		return LocalizationPaths;
+	};
+
+	auto GetCookedEditorLocalizationPaths = [&AppendChunkedLocalizationPaths]()
+	{
+		TArray<FString> LocalizationPaths = FPaths::GetCookedEditorLocalizationPaths();
+		AppendChunkedLocalizationPaths(LocalizationPaths);
 		return LocalizationPaths;
 	};
 
@@ -85,6 +102,9 @@ void FLocalizationResourceTextSource::LoadLocalizedResources(const ELocalization
 	if (ShouldLoadEditor(InLoadFlags))
 	{
 		EditorLocalizationPaths += FPaths::GetEditorLocalizationPaths();
+#if UE_IS_COOKED_EDITOR
+		EditorLocalizationPaths += GetCookedEditorLocalizationPaths();
+#endif
 		EditorLocalizationPaths += FPaths::GetToolTipLocalizationPaths();
 
 		bool bShouldUseLocalizedPropertyNames = false;
@@ -226,7 +246,7 @@ void FLocalizationResourceTextSource::LoadLocalizedResourcesFromPaths(TArrayView
 				}
 
 				const FString LocResFilename = FPaths::GetBaseFilename(LocalizationPath) + TEXT(".locres");
-				LoadLocalizationResourcesForCulture(InOutLocalizedResource, LocalizationPath, PrioritizedCultureName, LocResFilename, BaseResourcePriority + CultureIndex);
+				LoadLocalizationResourcesForCulture(InOutLocalizedResource, LocalizationPath, PrioritizedCultureName, LocResFilename, BaseResourcePriority + CultureIndex + 1);
 			}
 		}
 	}
diff --git a/Engine/Source/Runtime/Core/Private/Internationalization/PolyglotTextSource.cpp b/Engine/Source/Runtime/Core/Private/Internationalization/PolyglotTextSource.cpp
--- a/Engine/Source/Runtime/Core/Private/Internationalization/PolyglotTextSource.cpp
+++ b/Engine/Source/Runtime/Core/Private/Internationalization/PolyglotTextSource.cpp
@@ -153,7 +153,7 @@ void FPolyglotTextSource::AddPolyglotDataToResource(const FPolyglotTextData& InP
 					InPolyglotTextData.GetKey(),
 					InPolyglotTextData.GetNativeString(),
 					LocalizedString,
-					BaseResourcePriority + LocalizedPriority
+					BaseResourcePriority + LocalizedPriority + 1
 					);
 			}
 		}
diff --git a/Engine/Source/Runtime/Core/Private/Misc/Paths.cpp b/Engine/Source/Runtime/Core/Private/Misc/Paths.cpp
--- a/Engine/Source/Runtime/Core/Private/Misc/Paths.cpp
+++ b/Engine/Source/Runtime/Core/Private/Misc/Paths.cpp
@@ -37,6 +37,7 @@ struct FPaths::FStaticData
 	FString         UserFolder;
 	TArray<FString> EngineLocalizationPaths;
 	TArray<FString> EditorLocalizationPaths;
+	TArray<FString> CookedEditorLocalizationPaths;
 	TArray<FString> PropertyNameLocalizationPaths;
 	TArray<FString> ToolTipLocalizationPaths;
 	TArray<FString> GameLocalizationPaths;
@@ -51,6 +52,7 @@ struct FPaths::FStaticData
 	bool bUserFolderInitialized                    = false;
 	bool bEngineLocalizationPathsInitialized       = false;
 	bool bEditorLocalizationPathsInitialized       = false;
+	bool bCookedEditorLocalizationPathsInitialized = false;
 	bool bPropertyNameLocalizationPathsInitialized = false;
 	bool bToolTipLocalizationPathsInitialized      = false;
 	bool bGameLocalizationPathsInitialized         = false;
@@ -564,6 +566,22 @@ const TArray<FString>& FPaths::GetEditorLocalizationPaths()
 	return StaticData.EditorLocalizationPaths;
 }
 
+const TArray<FString>& FPaths::GetCookedEditorLocalizationPaths()
+{
+	FStaticData& StaticData = TLazySingleton<FStaticData>::Get();
+
+	if (!StaticData.bCookedEditorLocalizationPathsInitialized)
+	{
+		if (GConfig && GConfig->IsReadyForUse())
+		{
+			GConfig->GetArray(TEXT("Internationalization"), TEXT("CookedLocalizationPaths"), StaticData.CookedEditorLocalizationPaths, GEditorIni);
+			StaticData.bCookedEditorLocalizationPathsInitialized = true;
+		}
+	}
+
+	return StaticData.CookedEditorLocalizationPaths;
+}
+
 const TArray<FString>& FPaths::GetPropertyNameLocalizationPaths()
 {
 	FStaticData& StaticData = TLazySingleton<FStaticData>::Get();
diff --git a/Engine/Source/Runtime/Core/Public/Misc/Paths.h b/Engine/Source/Runtime/Core/Public/Misc/Paths.h
--- a/Engine/Source/Runtime/Core/Public/Misc/Paths.h
+++ b/Engine/Source/Runtime/Core/Public/Misc/Paths.h
@@ -356,6 +356,11 @@ public:
 	 */
 	static const TArray<FString>& GetEditorLocalizationPaths();
 
+	/**
+	 * Returns a list of cooked editor-specific localization paths
+	 */
+	static const TArray<FString>& GetCookedEditorLocalizationPaths();
+
 	/** 
 	 * Returns a list of property name localization paths
 	 */
