Modules - Better stats around loading the actual DLLs with modules.  It's just nice to know what's going on in those blank spots when doing module loading.

diff --git a/Engine/Source/Runtime/Core/Private/Modules/ModuleManager.cpp b/Engine/Source/Runtime/Core/Private/Modules/ModuleManager.cpp
--- a/Engine/Source/Runtime/Core/Private/Modules/ModuleManager.cpp
+++ b/Engine/Source/Runtime/Core/Private/Modules/ModuleManager.cpp
@@ -569,7 +569,11 @@ IModuleInterface* FModuleManager::LoadModuleWithFailureReason(const FName InModu
 		// Skip this check if file manager has not yet been initialized
 		if (FPaths::FileExists(ModuleFileToLoad))
 		{
-			ModuleInfo->Handle = FPlatformProcess::GetDllHandle(*ModuleFileToLoad);
+			{
+				TRACE_CPUPROFILER_EVENT_SCOPE(FPlatformProcess::GetDllHandle);
+				ModuleInfo->Handle = FPlatformProcess::GetDllHandle(*ModuleFileToLoad);
+			}
+			
 			if (ModuleInfo->Handle != nullptr)
 			{
 				{
diff --git a/Engine/Source/Runtime/Core/Private/Windows/WindowsPlatformProcess.cpp b/Engine/Source/Runtime/Core/Private/Windows/WindowsPlatformProcess.cpp
--- a/Engine/Source/Runtime/Core/Private/Windows/WindowsPlatformProcess.cpp
+++ b/Engine/Source/Runtime/Core/Private/Windows/WindowsPlatformProcess.cpp
@@ -119,18 +119,21 @@ void* FWindowsPlatformProcess::GetDllHandle( const TCHAR* FileName )
 	}
 
 	// Load the DLL, avoiding windows dialog boxes if missing
+	static const bool CMDLINE_dllerrors = FParse::Param(::GetCommandLineW(), TEXT("dllerrors"));
+	static const bool CMDLINE_unattended = FParse::Param(::GetCommandLineW(), TEXT("unattended"));
+	
 	DWORD ErrorMode = 0;
-	if(!FParse::Param(::GetCommandLineW(), TEXT("dllerrors")))
+	if(!CMDLINE_dllerrors)
 	{
 		ErrorMode |= SEM_NOOPENFILEERRORBOX;
-		if(FParse::Param(::GetCommandLineW(), TEXT("unattended")))
+		if(CMDLINE_unattended)
 		{
 			ErrorMode |= SEM_FAILCRITICALERRORS | SEM_NOGPFAULTERRORBOX;
 		}
 	}
 
 	DWORD PrevErrorMode = 0;
-	BOOL bHavePrevErrorMode = ::SetThreadErrorMode(ErrorMode, &PrevErrorMode);
+	const BOOL bHavePrevErrorMode = ::SetThreadErrorMode(ErrorMode, &PrevErrorMode);
 
 	// Load the DLL, avoiding windows dialog boxes if missing
 	void* Handle = LoadLibraryWithSearchPaths(FileName, SearchPaths);
@@ -2001,23 +2004,36 @@ void *FWindowsPlatformProcess::LoadLibraryWithSearchPaths(const FString& FileNam
 		// Load all the missing dependencies first
 		for (int32 Idx = 0; Idx < ImportFileNames.Num(); Idx++)
 		{
-			if (GetModuleHandle(*ImportFileNames[Idx]) == nullptr)
+			const FString& ImportFileName = ImportFileNames[Idx];
+			
+			if (!GetModuleHandle(*ImportFileName))
 			{
-				if(LoadLibrary(*ImportFileNames[Idx]))
+				const void* DependencyHandle = [&ImportFileName]() 
+				{
+					TRACE_CPUPROFILER_EVENT_SCOPE(Windows::LoadLibrary);
+					return LoadLibrary(*ImportFileName);
+				}();
+				
+				if (DependencyHandle)
 				{
-					UE_LOG(LogWindows, Verbose, TEXT("Preloaded '%s'"), *ImportFileNames[Idx]);
+					UE_LOG(LogWindows, Verbose, TEXT("Preloaded '%s'"), *ImportFileName);
 				}
 				else
 				{
-					UE_LOG(LogWindows, Log, TEXT("Failed to preload '%s' (GetLastError=%d)"), *ImportFileNames[Idx], GetLastError());
-					LogImportDiagnostics(ImportFileNames[Idx], SearchPaths);
+					UE_LOG(LogWindows, Log, TEXT("Failed to preload '%s' (GetLastError=%d)"), *ImportFileName, GetLastError());
+					LogImportDiagnostics(ImportFileName, SearchPaths);
 				}
 			}
 		}
 	}
 
 	// Try to load the actual library
-	void* Handle = LoadLibrary(*FullFileName);
+	void* Handle = [FullFileName]() 
+	{
+		TRACE_CPUPROFILER_EVENT_SCOPE(Windows::LoadLibrary);
+		return LoadLibrary(*FullFileName);
+	}();
+	
 	if(Handle)
 	{
 		UE_LOG(LogWindows, Verbose, TEXT("Loaded %s"), *FullFileName);
