Added a registerable callback for when a named chunk is installed.

 - AddNamedChunkInstallDelegate / RemoveNamedChunkInstallDelegate

diff --git a/Engine/Source/Runtime/Core/Public/GenericPlatform/GenericPlatformChunkInstall.h b/Engine/Source/Runtime/Core/Public/GenericPlatform/GenericPlatformChunkInstall.h
--- a/Engine/Source/Runtime/Core/Public/GenericPlatform/GenericPlatformChunkInstall.h
+++ b/Engine/Source/Runtime/Core/Public/GenericPlatform/GenericPlatformChunkInstall.h
@@ -82,6 +82,11 @@ DECLARE_DELEGATE_OneParam(FPlatformChunkInstallCompleteDelegate, uint32);
 DECLARE_DELEGATE_TwoParams(FPlatformChunkInstallDelegate, uint32, bool);
 DECLARE_MULTICAST_DELEGATE_TwoParams(FPlatformChunkInstallMultiDelegate, uint32, bool);
 
+/** Delegate called when a Named Chunk either successfully installs or fails to install, bool is success */
+DECLARE_DELEGATE_TwoParams(FPlatformNamedChunkInstallDelegate, FName, bool);
+DECLARE_MULTICAST_DELEGATE_TwoParams(FPlatformNamedChunkInstallMultiDelegate, FName, bool);
+
+
 enum class ECustomChunkType : uint8
 {
 	OnDemandChunk,
@@ -300,6 +305,19 @@ public:
 	 */
 	virtual TArray<FName> GetNamedChunksByType(ENamedChunkType NamedChunkType) const = 0;
 
+	/** 
+	 * Request a delegate callback on named chunk install completion or failure. Request may not be respected.
+	 * @param Delegate		The delegate to call when any named chunk is installed or fails to install
+	 * @return				Handle to the bound delegate
+	 */
+	virtual FDelegateHandle AddNamedChunkInstallDelegate( FPlatformNamedChunkInstallDelegate Delegate ) = 0;
+
+	/**
+	 * Remove a delegate callback on named chunk install completion.
+	 * @param Delegate		The delegate to remove.
+	 */
+	virtual void RemoveNamedChunkInstallDelegate( FDelegateHandle Delegate ) = 0;
+
 protected:
 		/**
 		 * Get the current location of a chunk.
@@ -470,10 +488,22 @@ public:
 		return TArray<FName>();
 	}
 
+	virtual FDelegateHandle AddNamedChunkInstallDelegate(FPlatformNamedChunkInstallDelegate Delegate) override
+	{
+		return NamedChunkInstallDelegate.Add(Delegate);
+	}
+
+	virtual void RemoveNamedChunkInstallDelegate(FDelegateHandle Delegate) override
+	{
+		NamedChunkInstallDelegate.Remove(Delegate);
+	}
+
+
 protected:
 
-	/** Delegate called when installation succeeds or fails */
+	/** Delegates called when installation succeeds or fails */
 	FPlatformChunkInstallMultiDelegate InstallDelegate;
+	FPlatformNamedChunkInstallMultiDelegate NamedChunkInstallDelegate;
 
 	virtual EChunkLocation::Type GetChunkLocation(uint32 ChunkID) override
 	{
