Temporary workaround for creating processes that use two unique pipes

diff --git a/Engine/Source/Runtime/Core/Private/Unix/UnixPlatformProcess.cpp b/Engine/Source/Runtime/Core/Private/Unix/UnixPlatformProcess.cpp
--- a/Engine/Source/Runtime/Core/Private/Unix/UnixPlatformProcess.cpp
+++ b/Engine/Source/Runtime/Core/Private/Unix/UnixPlatformProcess.cpp
@@ -526,8 +526,8 @@ bool FUnixPlatformProcess::CreatePipe(void*& ReadPipe, void*& WritePipe, bool bW
 		return false;
 	}
 
-	ReadPipe = new FPipeHandle(PipeFd[ 0 ]);
-	WritePipe = new FPipeHandle(PipeFd[ 1 ]);
+	ReadPipe = new FPipeHandle(PipeFd[ 0 ], PipeFd[ 1 ]);
+	WritePipe = new FPipeHandle(PipeFd[ 1 ], PipeFd[ 0 ]);
 
 	return true;
 }
@@ -958,13 +958,27 @@ FProcHandle FUnixPlatformProcess::CreateProc(const TCHAR* URL, const TCHAR* Parm
 
 		if (PipeWriteChild)
 		{
+			const FPipeHandle* PipeReadHandle = reinterpret_cast<const FPipeHandle*>(PipeReadChild);
 			const FPipeHandle* PipeWriteHandle = reinterpret_cast<const FPipeHandle*>(PipeWriteChild);
+
+			// If using unique read and write pipes, close the other end of the write pipe
+			if (PipeReadChild && PipeWriteHandle->GetPairHandle() != PipeReadHandle->GetHandle())
+			{
+				posix_spawn_file_actions_addclose(&FileActions, PipeWriteHandle->GetPairHandle());
+			}
 			posix_spawn_file_actions_adddup2(&FileActions, PipeWriteHandle->GetHandle(), STDOUT_FILENO);
 		}
 
 		if (PipeReadChild)
 		{
 			const FPipeHandle* PipeReadHandle = reinterpret_cast<const FPipeHandle*>(PipeReadChild);
+			const FPipeHandle* PipeWriteHandle = reinterpret_cast<const FPipeHandle*>(PipeWriteChild);
+
+			// If using unique read and write pipes, close the other end of the read pipe
+			if (PipeWriteChild && PipeReadHandle->GetPairHandle() != PipeWriteHandle->GetHandle())
+			{
+				posix_spawn_file_actions_addclose(&FileActions, PipeReadHandle->GetPairHandle());
+			}
 			posix_spawn_file_actions_adddup2(&FileActions, PipeReadHandle->GetHandle(), STDIN_FILENO);
 		}
 
diff --git a/Engine/Source/Runtime/Core/Public/Unix/UnixPlatformProcess.h b/Engine/Source/Runtime/Core/Public/Unix/UnixPlatformProcess.h
--- a/Engine/Source/Runtime/Core/Public/Unix/UnixPlatformProcess.h
+++ b/Engine/Source/Runtime/Core/Public/Unix/UnixPlatformProcess.h
@@ -162,8 +162,9 @@ struct FProcHandle
 /** Wrapper around Unix file descriptors */
 struct FPipeHandle
 {
-	FPipeHandle(int Fd)
+	FPipeHandle(int Fd, int PairFd)
 		:	PipeDesc(Fd)
+		,	PairDesc(PairFd)
 	{
 	}
 
@@ -187,9 +188,18 @@ struct FPipeHandle
 		return PipeDesc;
 	}
 
+	/**
+	 * Returns the raw file handle of the other endpoint of PipeDesc
+	 */
+	int GetPairHandle() const
+	{
+		return PairDesc;
+	}
+
 protected:
 
 	int	PipeDesc;
+	int	PairDesc;
 };
 
 /**
