fixed a race reported and analysed by Carl-Magnus Nordin:
T1: Tail.Store(Sentinel)
T2: Prev = Tail.Load()  (Prev = Sentinel)
T2: Prev.Next = [some node]  (Sentinel.Next = [some node])
T1: Sentinel.Next = null
which can lead to losing nodes

diff --git a/Engine/Source/Runtime/Core/Public/Containers/DepletableMpscQueue.h b/Engine/Source/Runtime/Core/Public/Containers/DepletableMpscQueue.h
--- a/Engine/Source/Runtime/Core/Public/Containers/DepletableMpscQueue.h
+++ b/Engine/Source/Runtime/Core/Public/Containers/DepletableMpscQueue.h
@@ -63,8 +63,9 @@ namespace UE
 			// switch `Tail` to the new node and only then link the old tail to the new one. The list is not fully linked between these ops,
 			// this is explicitly handled by the consumer by waiting for the link
 
-			FNode* Prev = Tail.exchange(New, std::memory_order_release); // `release` to make sure the new node is fully constructed before
-			// it becomes visible to the consumer
+			FNode* Prev = Tail.exchange(New, std::memory_order_acq_rel); // "acquire" to sync with `Prev->Next` initialisation from a concurrent calls,
+			// `release` to make sure the new node is fully constructed before it becomes visible to the consumer or a concurrent enqueueing
+			check(Prev->Next.load(std::memory_order_relaxed) == nullptr); // `Tail` is assigned before its Next
 
 			Prev->Next.store(New, std::memory_order_relaxed);
 
@@ -91,8 +92,8 @@ namespace UE
 			// reset the queue to the empty state. this redirects producers to start from `Sentinel` again.
 			// take note of the tail on resetting it because the list can be still not fully linked and so `Node.Next == nullptr` can't be 
 			// used to detect the end of the list
-			FNode* Last = Tail.exchange(&Sentinel, std::memory_order_acquire); // `acquire` to sync with producers' tail modifications,
-			// `Sentinel.Next = nullptr` must happen before modifying `Tail` but it's important only for the single consumer so doesn't need extra sync
+			FNode* Last = Tail.exchange(&Sentinel, std::memory_order_acq_rel); // `acquire` to sync with producers' tail modifications, and
+			// "release" to force `Sentinel.Next = nullptr` happening before modifying `Tail
 			check(Last->Next.load(std::memory_order_relaxed) == nullptr); // `Tail` is assigned before its Next
 			// the previously queued items are detached from the instance (as a linked list, though potentially not fully linked yet)
 
