#ue5 - Add range checks to UE::LWC::FloatToIntCastChecked() and validation code to UnrealMathTests startup tests.

diff --git a/Engine/Source/Runtime/Core/Private/Tests/Math/UnrealMathTest.cpp b/Engine/Source/Runtime/Core/Private/Tests/Math/UnrealMathTest.cpp
--- a/Engine/Source/Runtime/Core/Private/Tests/Math/UnrealMathTest.cpp
+++ b/Engine/Source/Runtime/Core/Private/Tests/Math/UnrealMathTest.cpp
@@ -3291,6 +3291,45 @@ bool FVectorRegisterAbstractionTest::RunTest(const FString& Parameters)
 		// Possibly invalid, depending on the generic base FloatToIntCastChecked, since this doesn't match a specialization
 		Big1 = UE::LWC::FloatToIntCastChecked<int32>(I);
 #endif
+
+		// Test FloatToIntCastChecked() range limits
+
+		// float->int32
+		I = UE::LWC::FloatToIntCastChecked<int32>(float(TNumericLimits<int32>::Max() - 64));
+		I = UE::LWC::FloatToIntCastChecked<int32>(float(TNumericLimits<int32>::Lowest()));
+
+		// float->int64
+		Big1 = UE::LWC::FloatToIntCastChecked<int64>(float(TNumericLimits<int32>::Max()));
+		Big1 = UE::LWC::FloatToIntCastChecked<int64>(float(TNumericLimits<int32>::Lowest()));
+		Big1 = UE::LWC::FloatToIntCastChecked<int64>(float(TNumericLimits<int64>::Max() - (int64)549755813888)); // 2^63 - 1 - 2^39
+		Big1 = UE::LWC::FloatToIntCastChecked<int64>(float(TNumericLimits<int64>::Lowest()));
+		
+		// double->int32
+		I = UE::LWC::FloatToIntCastChecked<int32>(double(TNumericLimits<int32>::Max()));
+		I = UE::LWC::FloatToIntCastChecked<int32>(double(TNumericLimits<int32>::Lowest()));
+
+		// double->int64
+		Big1 = UE::LWC::FloatToIntCastChecked<int64>(double(TNumericLimits<int64>::Max() - 512));
+		Big1 = UE::LWC::FloatToIntCastChecked<int64>(double(TNumericLimits<int64>::Lowest()));
+
+#if MATHTEST_CHECK_INVALID_OVERLOAD_VARIANTS
+		// These should all assert
+		I = UE::LWC::FloatToIntCastChecked<int32>(float(TNumericLimits<int32>::Max()));
+		I = UE::LWC::FloatToIntCastChecked<int32>(TNumericLimits<float>::Max());
+		I = UE::LWC::FloatToIntCastChecked<int32>(TNumericLimits<float>::Lowest());
+
+		I = UE::LWC::FloatToIntCastChecked<int32>(TNumericLimits<double>::Max());
+		I = UE::LWC::FloatToIntCastChecked<int32>(TNumericLimits<double>::Lowest());
+
+		Big1 = UE::LWC::FloatToIntCastChecked<int64>(float(TNumericLimits<int64>::Max()));
+		Big1 = UE::LWC::FloatToIntCastChecked<int64>(TNumericLimits<float>::Max());
+		Big1 = UE::LWC::FloatToIntCastChecked<int64>(TNumericLimits<float>::Lowest());
+
+		Big1 = UE::LWC::FloatToIntCastChecked<int64>(double(TNumericLimits<int64>::Max()));
+		Big1 = UE::LWC::FloatToIntCastChecked<int64>(TNumericLimits<double>::Max());
+		Big1 = UE::LWC::FloatToIntCastChecked<int64>(TNumericLimits<double>::Lowest());
+#endif
+
 	}
 
 	// Sin, Cos, Tan tests
diff --git a/Engine/Source/Runtime/Core/Public/Math/UnrealMathUtility.h b/Engine/Source/Runtime/Core/Public/Math/UnrealMathUtility.h
--- a/Engine/Source/Runtime/Core/Public/Math/UnrealMathUtility.h
+++ b/Engine/Source/Runtime/Core/Public/Math/UnrealMathUtility.h
@@ -2510,6 +2510,8 @@ namespace LWC
 	template<>
 	FORCEINLINE int32 FloatToIntCastChecked(float FloatValue)
 	{
+		// floats over 2^31 - 1 - 64 overflow int32 after conversion.
+		checkf(FloatValue >= float(TNumericLimits<int32>::Lowest()) && FloatValue <= float(TNumericLimits<int32>::Max() - 64), TEXT("Input value %f will exceed int32 limits"), FloatValue);
 		return FMath::TruncToInt32(FloatValue);
 	}
 
@@ -2517,6 +2519,8 @@ namespace LWC
 	template<>
 	FORCEINLINE int64 FloatToIntCastChecked(float FloatValue)
 	{
+		// floats over 2^63 - 1 - 2^39 overflow int64 after conversion.
+		checkf(FloatValue >= float(TNumericLimits<int64>::Lowest()) && FloatValue <= float(TNumericLimits<int64>::Max() - (int64)549755813888), TEXT("Input value %f will exceed int64 limits"), FloatValue);
 		return FMath::TruncToInt64(FloatValue);
 	}
 
@@ -2524,6 +2528,7 @@ namespace LWC
 	template<>
 	FORCEINLINE int32 FloatToIntCastChecked(double FloatValue)
 	{
+		checkf(FloatValue >= double(TNumericLimits<int32>::Lowest()) && FloatValue <= double(TNumericLimits<int32>::Max()), TEXT("Input value %f will exceed int32 limits"), FloatValue);
 		return FMath::TruncToInt32(FloatValue);
 	}
 
@@ -2531,6 +2536,8 @@ namespace LWC
 	template<>
 	FORCEINLINE int64 FloatToIntCastChecked(double FloatValue)
 	{
+		// doubles over 2^63 - 1 - 512 overflow int64 after conversion.
+		checkf(FloatValue >= double(TNumericLimits<int64>::Lowest()) && FloatValue <= double(TNumericLimits<int64>::Max() - 512), TEXT("Input value %f will exceed int64 limits"), FloatValue);
 		return FMath::TruncToInt64(FloatValue);
 	}
 
