Fix potential race condition on multicast delegate used from multiple threads

diff --git a/Engine/Source/Runtime/Core/Public/ProfilingDebugging/TraceAuxiliary.h b/Engine/Source/Runtime/Core/Public/ProfilingDebugging/TraceAuxiliary.h
--- a/Engine/Source/Runtime/Core/Public/ProfilingDebugging/TraceAuxiliary.h
+++ b/Engine/Source/Runtime/Core/Public/ProfilingDebugging/TraceAuxiliary.h
@@ -43,17 +43,17 @@ public:
 	/**
 	 * Callback type when a new connection is established.
 	 */
-	DECLARE_MULTICAST_DELEGATE(FOnConnection);
+	DECLARE_TS_MULTICAST_DELEGATE(FOnConnection);
 
 	/** Callback whenever a trace is started */
-	DECLARE_MULTICAST_DELEGATE_TwoParams(FOnTraceStarted, FTraceAuxiliary::EConnectionType TraceType, const FString& TraceDestination);
+	DECLARE_TS_MULTICAST_DELEGATE_TwoParams(FOnTraceStarted, FTraceAuxiliary::EConnectionType TraceType, const FString& TraceDestination);
 
 	/** 
 	* Callback whenever a trace recording is stopped. 
 	* TraceType tells what kind of trace it is.
 	* TraceDestination will be either the the filename and path for a file trace or the network connection for a network trace
 	*/
-	DECLARE_MULTICAST_DELEGATE_TwoParams(FOnTraceStopped, FTraceAuxiliary::EConnectionType TraceType, const FString& TraceDestination);
+	DECLARE_TS_MULTICAST_DELEGATE_TwoParams(FOnTraceStopped, FTraceAuxiliary::EConnectionType TraceType, const FString& TraceDestination);
 
 	struct FOptions
 	{
