Implement FAutomationTestBase::IsLatentCommandQueueEmpty()
Return whether there is no latent command in queue

diff --git a/Engine/Source/Runtime/Core/Public/Misc/AutomationTest.h b/Engine/Source/Runtime/Core/Public/Misc/AutomationTest.h
--- a/Engine/Source/Runtime/Core/Public/Misc/AutomationTest.h
+++ b/Engine/Source/Runtime/Core/Public/Misc/AutomationTest.h
@@ -934,6 +934,14 @@ public:
 	 */
 	void DequeueAllCommands();
 
+	/**
+	 * Whether there is no latent command in queue
+	 */
+	bool IsLatentCommandQueueEmpty() const
+	{
+		return LatentCommands.IsEmpty();
+	}
+
 	/**
 	 * Load any modules that are not loaded by default and have test classes in them
 	 */
