LLM: Add missing FileSystem scope for the single threaded dedicated server scenario.
Previously these allocations were incorrectly tracked as AsyncLoading memory.

#rnx

diff --git a/Engine/Source/Runtime/Core/Private/IO/IoDispatcher.cpp b/Engine/Source/Runtime/Core/Private/IO/IoDispatcher.cpp
--- a/Engine/Source/Runtime/Core/Private/IO/IoDispatcher.cpp
+++ b/Engine/Source/Runtime/Core/Private/IO/IoDispatcher.cpp
@@ -550,6 +550,7 @@ public:
 		}
 		else
 		{
+			LLM_SCOPE(ELLMTag::FileSystem);
 			ProcessIncomingRequests();
 			while (PendingIoRequestsCount > 0)
 			{
