Fix racy and uninitialized traceid in FBaseGraphTask, FTaskBase and FGraphEvent

diff --git a/Engine/Source/Runtime/Core/Public/Async/TaskGraphInterfaces.h b/Engine/Source/Runtime/Core/Public/Async/TaskGraphInterfaces.h
--- a/Engine/Source/Runtime/Core/Public/Async/TaskGraphInterfaces.h
+++ b/Engine/Source/Runtime/Core/Public/Async/TaskGraphInterfaces.h
@@ -890,7 +890,7 @@ protected:
 	TaskTrace::FId GetTraceId() const
 	{
 #if UE_TASK_TRACE_ENABLED
-		return TraceId;
+		return TraceId.load(std::memory_order_relaxed);
 #else
 		return TaskTrace::InvalidId;
 #endif
@@ -990,7 +990,7 @@ private:
 #endif
 
 #if UE_TASK_TRACE_ENABLED
-	TaskTrace::FId TraceId;
+	std::atomic<TaskTrace::FId> TraceId { TaskTrace::InvalidId };
 #endif
 };
 
@@ -1099,7 +1099,7 @@ public:
 	TaskTrace::FId GetTraceId() const
 	{
 #if UE_TASK_TRACE_ENABLED
-		return TraceId;
+		return TraceId.load(std::memory_order_relaxed);
 #else
 		return TaskTrace::InvalidId;
 #endif
@@ -1184,7 +1184,7 @@ private:
 #endif
 
 #if UE_TASK_TRACE_ENABLED
-	TaskTrace::FId TraceId = TaskTrace::GenerateTaskId();
+	std::atomic<TaskTrace::FId> TraceId { TaskTrace::GenerateTaskId() };
 #endif
 };
 
diff --git a/Engine/Source/Runtime/Core/Public/Tasks/TaskPrivate.h b/Engine/Source/Runtime/Core/Public/Tasks/TaskPrivate.h
--- a/Engine/Source/Runtime/Core/Public/Tasks/TaskPrivate.h
+++ b/Engine/Source/Runtime/Core/Public/Tasks/TaskPrivate.h
@@ -398,7 +398,7 @@ namespace UE::Tasks
 			TaskTrace::FId GetTraceId() const
 			{
 #if UE_TASK_TRACE_ENABLED
-				return TraceId;
+				return TraceId.load(std::memory_order_relaxed);
 #else
 				return TaskTrace::InvalidId;
 #endif
@@ -614,7 +614,7 @@ namespace UE::Tasks
 			std::atomic<uint32> NumLocks{ NumInitialLocks };
 
 #if UE_TASK_TRACE_ENABLED
-			TaskTrace::FId TraceId = TaskTrace::GenerateTaskId();
+			std::atomic<TaskTrace::FId> TraceId { TaskTrace::GenerateTaskId() };
 #endif
 
 			// the task is completed when its subsequents list is closed
