Add a way that failing rety automation commands can kill the application. This is useful for when running  a disconnected horde job and we just want the editor to shutdown upon failure.

diff --git a/Engine/Source/Runtime/Core/Public/Misc/AutomationTest.h b/Engine/Source/Runtime/Core/Public/Misc/AutomationTest.h
--- a/Engine/Source/Runtime/Core/Public/Misc/AutomationTest.h
+++ b/Engine/Source/Runtime/Core/Public/Misc/AutomationTest.h
@@ -3668,6 +3668,15 @@ class IAutomationLatentCommandWithRetriesAndDelays : public IAutomationLatentCom
 public:
 	virtual ~IAutomationLatentCommandWithRetriesAndDelays() {}
 
+	virtual void CommandFailedDueToError(const FString& ErrorMessage)
+	{
+		// Stop further commands and log error so the test fails.
+		FAutomationTestFramework::Get().DequeueAllCommands();
+		// Must log here so that Gauntlet can also pick up the error for its report. Otherwise, it will only show up in the Automation log.
+		UE_LOG(LogLatentCommands, Error, TEXT("%s"), *ErrorMessage);
+		GetCurrentTest()->AddError(ErrorMessage);
+	}
+
 	// Base Update override with delay logic built in. Can be further overridden in child
 	// classes
 	virtual bool Update() override
@@ -3677,14 +3686,11 @@ public:
 			// command has unlimited retries, so need to check if max run time has been exceeded
 			if (HasExceededMaxTotalRunTime())
 			{
-				// Command run time has exceeded max total run time. Stop further commands and log error so the test fails.
-				FAutomationTestFramework::Get().DequeueAllCommands();
+				// Command run time has exceeded max total run time.
 				FString ErrorMessageText = FString::Printf(TEXT("%s has failed due to exceeding the max allowed run time of %f seconds. \
 This may be due to an error, or having a single command attempt to do too many things. If this is not due to an error, consider breaking \
 up this command into multiple, smaller commands."), *GetTestAndCommandName(), MaxTotalRunTimeInSeconds);
-				// Must log here so that Gauntlet can also pick up the error for its report. Otherwise, it will only show up in the Automation log.
-				UE_LOG(LogLatentCommands, Error, TEXT("%s"), *ErrorMessageText);
-				GetCurrentTest()->AddError(ErrorMessageText);
+				CommandFailedDueToError(ErrorMessageText);
 				return true;
 			}
 		}
@@ -3699,12 +3705,9 @@ up this command into multiple, smaller commands."), *GetTestAndCommandName(), Ma
 
 		if (!CanRetry())
 		{
-			// Retry count has been exceeded, dequeue further commands and log error to test.
-			FAutomationTestFramework::Get().DequeueAllCommands();
 			// Must log here so that Gauntlet can also pick up the error for its report. Otherwise, it will only show up in the Automation log.
 			FString ErrorMessageText = FString::Printf(TEXT("%s Latent command with retries and delays has failed after %d retries"), *GetTestAndCommandName(), MaxRetries);
-			UE_LOG(LogLatentCommands, Error, TEXT("%s"), *ErrorMessageText);
-			GetCurrentTest()->AddError(ErrorMessageText);
+			CommandFailedDueToError(ErrorMessageText);
 			return true;
 		}
 
