Tagged the large memory allocation(s) from FReserveScheduler::StartWorkers (reserving worker threads) with "EngineMisc" LLM tag. Previously it was tagged randomly by the first system that triggers the initialization (usually being the "ConfigSystem" LLM tag).

diff --git a/Engine/Source/Runtime/Core/Private/Async/Fundamental/ReserveScheduler.cpp b/Engine/Source/Runtime/Core/Private/Async/Fundamental/ReserveScheduler.cpp
--- a/Engine/Source/Runtime/Core/Private/Async/Fundamental/ReserveScheduler.cpp
+++ b/Engine/Source/Runtime/Core/Private/Async/Fundamental/ReserveScheduler.cpp
@@ -4,6 +4,7 @@
 
 #include "Async/Fundamental/LocalQueue.h"
 #include "Async/Fundamental/Scheduler.h"
+#include "HAL/LowLevelMemTracker.h"
 #include "HAL/PlatformProcess.h"
 #include "Misc/ScopeLock.h"
 #include "ProfilingDebugging/CpuProfilerTrace.h"
@@ -78,6 +79,8 @@ void FReserveScheduler::StartWorkers(FScheduler& MainScheduler, uint32 NumWorker
 	uint32 OldActiveWorkers = ActiveWorkers.load(std::memory_order_relaxed);
 	if(OldActiveWorkers == 0 && FPlatformProcess::SupportsMultithreading() && ActiveWorkers.compare_exchange_strong(OldActiveWorkers, NumWorkers, std::memory_order_relaxed))
 	{
+		LLM_SCOPE_BYNAME(TEXT("EngineMisc/WorkerThreads"));
+
 		FScopeLock Lock(&WorkerThreadsCS);
 		check(!WorkerThreads.Num());
 		check(!WorkerLocalQueues.Num());
