[Insights] Cleared metadata stack for ProgramSize allocs.

diff --git a/Engine/Source/Runtime/Core/Private/ProfilingDebugging/Microsoft/WindowsModuleDiagnostics.cpp b/Engine/Source/Runtime/Core/Private/ProfilingDebugging/Microsoft/WindowsModuleDiagnostics.cpp
--- a/Engine/Source/Runtime/Core/Private/ProfilingDebugging/Microsoft/WindowsModuleDiagnostics.cpp
+++ b/Engine/Source/Runtime/Core/Private/ProfilingDebugging/Microsoft/WindowsModuleDiagnostics.cpp
@@ -7,6 +7,7 @@
 #include "Trace/Trace.inl"
 #include "ProfilingDebugging/ModuleDiagnostics.h"
 #include "ProfilingDebugging/MemoryTrace.h"
+#include "ProfilingDebugging/MetadataTrace.h"
 #include "ProfilingDebugging/TagTrace.h"
 #include "HAL/LowLevelMemTracker.h"
 
@@ -201,15 +202,16 @@ void FModuleTrace::OnDllLoaded(const UNICODE_STRING& Name, UPTRINT Base)
 		<< ModuleLoad.Size(OptionalHeader.SizeOfImage)
 		<< ModuleLoad.ImageId(ImageId.GetData(), ImageId.Num());
 
+#if UE_MEMORY_TRACE_ENABLED
 	{
-#if ENABLE_LOW_LEVEL_MEM_TRACKER
-		UE_MEMSCOPE(ELLMTag::ProgramSize);
-#endif
+		UE_TRACE_METADATA_CLEAR_SCOPE();
+		LLM(UE_MEMSCOPE(ELLMTag::ProgramSize));
 		MemoryTrace_Alloc(Base, OptionalHeader.SizeOfImage, 4 * 1024);
 		MemoryTrace_MarkAllocAsHeap(Base, ProgramHeapId);
 		MemoryTrace_Alloc(Base, OptionalHeader.SizeOfImage, 4 * 1024);
 	}
-	
+#endif // UE_MEMORY_TRACE_ENABLED
+
 	for (SubscribeFunc Subscriber : Subscribers)
 	{
 		Subscriber(true, (void*)Base, Name.Buffer);
