Fix automation crash.

diff --git a/Engine/Source/Runtime/Core/Private/Misc/AutomationTest.cpp b/Engine/Source/Runtime/Core/Private/Misc/AutomationTest.cpp
--- a/Engine/Source/Runtime/Core/Private/Misc/AutomationTest.cpp
+++ b/Engine/Source/Runtime/Core/Private/Misc/AutomationTest.cpp
@@ -208,7 +208,7 @@ void FAutomationTestFramework::FAutomationTestMessageFilter::Serialize(const TCH
 		FScopeLock Lock(&ActionCS);
 		if ((Verbosity == ELogVerbosity::Warning) || (Verbosity == ELogVerbosity::Error))
 		{
-			if (LocalCurTest->IsExpectedError(FString(V)))
+			if (LocalCurTest && LocalCurTest->IsExpectedError(FString(V)))
 			{
 				Verbosity = ELogVerbosity::Verbose;
 			}
