Global function IsRunningDLCCookCommandlet() which determines if we're running the cook commandlet with -dlcname specified

diff --git a/Engine/Source/Runtime/Core/Private/Misc/CoreGlobals.cpp b/Engine/Source/Runtime/Core/Private/Misc/CoreGlobals.cpp
--- a/Engine/Source/Runtime/Core/Private/Misc/CoreGlobals.cpp
+++ b/Engine/Source/Runtime/Core/Private/Misc/CoreGlobals.cpp
@@ -168,6 +168,7 @@ bool GIsRunningUnattendedScript = false;
 
 #if WITH_EDITOR
 bool					PRIVATE_GIsRunningCookCommandlet	= false;				/** Whether this executable is running the cook commandlet */
+bool					PRIVATE_GIsRunningDLCCookCommandlet = false;				/** Whether this executable is running the cook commandlet on a DLC plugin */
 #endif
 
 #if WITH_ENGINE
diff --git a/Engine/Source/Runtime/Core/Public/CoreGlobals.h b/Engine/Source/Runtime/Core/Public/CoreGlobals.h
--- a/Engine/Source/Runtime/Core/Public/CoreGlobals.h
+++ b/Engine/Source/Runtime/Core/Public/CoreGlobals.h
@@ -241,6 +241,7 @@ extern CORE_API bool GFirstFrameIntraFrameDebugging;
 
 #if WITH_EDITOR
 extern CORE_API bool PRIVATE_GIsRunningCookCommandlet;
+extern CORE_API bool PRIVATE_GIsRunningDLCCookCommandlet;
 #endif
 
 /**
@@ -267,6 +268,18 @@ FORCEINLINE bool IsRunningCookCommandlet()
 #endif
 }
 
+/**
+* Check to see if this executable is running the cookcommandlet
+*/
+FORCEINLINE bool IsRunningDLCCookCommandlet()
+{
+#if WITH_EDITOR
+	return PRIVATE_GIsRunningDLCCookCommandlet;
+#else
+	return false;
+#endif
+}
+
 /** Returns running commandlet name 
  * 
  */
diff --git a/Engine/Source/Runtime/Launch/Private/LaunchEngineLoop.cpp b/Engine/Source/Runtime/Launch/Private/LaunchEngineLoop.cpp
--- a/Engine/Source/Runtime/Launch/Private/LaunchEngineLoop.cpp
+++ b/Engine/Source/Runtime/Launch/Private/LaunchEngineLoop.cpp
@@ -1949,6 +1949,7 @@ int32 FEngineLoop::PreInitPreStartupScreen(const TCHAR* CmdLine)
 		if (CommandletName == TEXTVIEW("cookcommandlet"))
 		{
 			PRIVATE_GIsRunningCookCommandlet = true;
+			PRIVATE_GIsRunningDLCCookCommandlet = (FCString::Strfind(FCommandLine::Get(), TEXT("-dlcname=")) != nullptr);
 		}
 #endif
 #if WITH_ENGINE
