Addition of OnNetworkConnectionStatusChanged event to FCoreDelegates
Addition of GetNetworkConnectionStatus function to FGenericPlatformMisc
Platform related implementations

diff --git a/Engine/Source/Runtime/Core/Private/GenericPlatform/GenericPlatformMisc.cpp b/Engine/Source/Runtime/Core/Private/GenericPlatform/GenericPlatformMisc.cpp
--- a/Engine/Source/Runtime/Core/Private/GenericPlatform/GenericPlatformMisc.cpp
+++ b/Engine/Source/Runtime/Core/Private/GenericPlatform/GenericPlatformMisc.cpp
@@ -38,6 +38,7 @@
 #include "Misc/UProjectInfo.h"
 #include "Internationalization/Culture.h"
 #include "ProfilingDebugging/CsvProfiler.h"
+#include "Misc/CoreDelegates.h"
 
 #if UE_ENABLE_ICU
 	THIRD_PARTY_INCLUDES_START
@@ -443,6 +444,39 @@ FString FSHA256Signature::ToString() const
 	return LocalHashStr;
 }
 
+const TCHAR* LexToString(ENetworkConnectionStatus EnumVal)
+{
+	switch (EnumVal)
+	{
+	case ENetworkConnectionStatus::Unknown:		return TEXT("Unknown");
+	case ENetworkConnectionStatus::Disabled:	return TEXT("Disabled");
+	case ENetworkConnectionStatus::Local:		return TEXT("Local");
+	case ENetworkConnectionStatus::Connected:	return TEXT("Connected");
+	}
+
+	checkNoEntry();
+	return TEXT("Unknown");
+}
+
+ENetworkConnectionStatus FGenericPlatformMisc::GetNetworkConnectionStatus()
+{
+	return CurrentNetworkConnectionStatus;
+}
+
+void FGenericPlatformMisc::SetNetworkConnectionStatus(ENetworkConnectionStatus NewNetworkConnectionStatus)
+{
+	const ENetworkConnectionStatus OldNetworkConnectionStatus = CurrentNetworkConnectionStatus;
+
+	if (OldNetworkConnectionStatus != NewNetworkConnectionStatus)
+	{
+		CurrentNetworkConnectionStatus = NewNetworkConnectionStatus;
+
+		QUICK_SCOPE_CYCLE_COUNTER(STAT_FGenericPlatformMisc_SetNetworkConnectionStatus);
+
+		FCoreDelegates::OnNetworkConnectionStatusChanged.Broadcast(OldNetworkConnectionStatus, NewNetworkConnectionStatus);
+	}
+}
+
 /* ENetworkConnectionType interface
  *****************************************************************************/
 
@@ -486,6 +520,8 @@ FString FSHA256Signature::ToString() const
 
 EDeviceScreenOrientation FGenericPlatformMisc::AllowedDeviceOrientation = EDeviceScreenOrientation::Unknown;
 
+ENetworkConnectionStatus FGenericPlatformMisc::CurrentNetworkConnectionStatus = ENetworkConnectionStatus::Connected;
+
 struct FGenericPlatformMisc::FStaticData
 {
 	FString         RootDir;
diff --git a/Engine/Source/Runtime/Core/Private/Misc/CoreDelegates.cpp b/Engine/Source/Runtime/Core/Private/Misc/CoreDelegates.cpp
--- a/Engine/Source/Runtime/Core/Private/Misc/CoreDelegates.cpp
+++ b/Engine/Source/Runtime/Core/Private/Misc/CoreDelegates.cpp
@@ -270,6 +270,7 @@ TDelegate<TSharedPtr<class IPackageStore>()> FCoreDelegates::CreatePackageStore;
 PRAGMA_ENABLE_DEPRECATION_WARNINGS
 
 TMulticastDelegate<void(bool)> FCoreDelegates::ApplicationNetworkInitializationChanged;
+TMulticastDelegate<void(ENetworkConnectionStatus, ENetworkConnectionStatus)> FCoreDelegates::OnNetworkConnectionStatusChanged;
 
 PRAGMA_DISABLE_DEPRECATION_WARNINGS
 FCrashOverrideParameters::~FCrashOverrideParameters()
diff --git a/Engine/Source/Runtime/Core/Public/GenericPlatform/GenericPlatformMisc.h b/Engine/Source/Runtime/Core/Public/GenericPlatform/GenericPlatformMisc.h
--- a/Engine/Source/Runtime/Core/Public/GenericPlatform/GenericPlatformMisc.h
+++ b/Engine/Source/Runtime/Core/Public/GenericPlatform/GenericPlatformMisc.h
@@ -336,6 +336,20 @@ enum class EMobileHapticsType : uint8
 	ImpactHeavy,
 };
 
+/** Possible connection states */
+enum class ENetworkConnectionStatus : uint8
+{
+	/** Default state */
+	Unknown = 0,
+	/** No network connection or network device disabled */
+	Disabled,
+	/** Ad-hoc Wifi network or LAN with no external connection */
+	Local,
+	/** Connected to the network */
+	Connected
+};
+CORE_API const TCHAR* LexToString(ENetworkConnectionStatus EnumVal);
+
 enum class ENetworkConnectionType : uint8
 {
 	/**
@@ -1484,6 +1498,16 @@ public:
 		return false;
 	}
 
+	/**
+	 * Returns the current status for the network connection
+	 */
+	static ENetworkConnectionStatus GetNetworkConnectionStatus();
+
+	/**
+	 * Updates the current status for the network connection
+	 */
+	static void SetNetworkConnectionStatus(ENetworkConnectionStatus NewNetworkConnectionStatus);
+
 	/**
 	 * Returns whether WiFi connection is currently active
 	 */
@@ -1830,6 +1854,8 @@ protected:
 
 	static EDeviceScreenOrientation AllowedDeviceOrientation;
 
+	static ENetworkConnectionStatus CurrentNetworkConnectionStatus;
+
 protected:
 	/**
 	 * Parse CoreLimit from CommandLine or from the SetCoreLimit value. The real number of physical cores and logical
diff --git a/Engine/Source/Runtime/Core/Public/Misc/CoreDelegates.h b/Engine/Source/Runtime/Core/Public/Misc/CoreDelegates.h
--- a/Engine/Source/Runtime/Core/Public/Misc/CoreDelegates.h
+++ b/Engine/Source/Runtime/Core/Public/Misc/CoreDelegates.h
@@ -501,6 +501,14 @@ public:
 	/** Called when the application's network initializes or shutdowns on platforms where the network stack is not always available */
 	static TMulticastDelegate<void(bool bIsNetworkInitialized)> ApplicationNetworkInitializationChanged;
 
+	/**
+	 * Called when the connection state as reported by the platform changes
+	 *
+	 * @param LastConnectionState last state of the connection
+	 * @param ConnectionState current state of the connection
+	 */
+	static TMulticastDelegate<void(ENetworkConnectionStatus LastConnectionState, ENetworkConnectionStatus ConnectionState)> OnNetworkConnectionStatusChanged;
+
 	// Callback to let code read or write specialized binary data that is generated at Stage time, for optimizing data right before 
 	// final game data is being written to disk
 	// The TMap is a map of an identifier for owner of the data, and a boolean where true means the data is being generated (ie editor), and false 
