Fix for FName Error message and adding logging to SPropertyMenuActorPicker if the user has chosen an Actor with an invalid NamePrivate (FName).

diff --git a/Engine/Source/Editor/PropertyEditor/Private/UserInterface/PropertyEditor/SPropertyMenuActorPicker.cpp b/Engine/Source/Editor/PropertyEditor/Private/UserInterface/PropertyEditor/SPropertyMenuActorPicker.cpp
--- a/Engine/Source/Editor/PropertyEditor/Private/UserInterface/PropertyEditor/SPropertyMenuActorPicker.cpp
+++ b/Engine/Source/Editor/PropertyEditor/Private/UserInterface/PropertyEditor/SPropertyMenuActorPicker.cpp
@@ -18,6 +18,7 @@
 #include "HAL/PlatformApplicationMisc.h"
 #include "SceneOutlinerPublicTypes.h"
 #include "ActorTreeItem.h"
+#include "PropertyNode.h"
 
 #define LOCTEXT_NAMESPACE "PropertyEditor"
 
@@ -189,6 +190,25 @@ void SPropertyMenuActorPicker::OnClear()
 
 void SPropertyMenuActorPicker::OnActorSelected( AActor* InActor )
 {
+	FText OutErrorMsg;
+
+	if (IsValid(InActor) && !FName::IsValidXName(InActor->GetName(), "aeoiu",  &OutErrorMsg))
+	{
+		FFormatNamedArguments Args;
+		Args.Add("ActorName", FText::FromString(FName::SanitizeWhitespace(InActor->GetName())));
+		Args.Add("ActorLabel", FText::FromString(FName::SanitizeWhitespace(InActor->GetActorLabel())));
+
+		FText Error = OutErrorMsg.IsEmpty() ?
+		     LOCTEXT("InvalidCharactersDefaultErrorMessage", "Names do not support the following characters: \"\' ,\\n\\r\\t") : OutErrorMsg;
+		
+		Args.Add("ErrorMessage", Error);
+		const FText LogMessage = FText::Format(LOCTEXT("InvalidActorName", "The chosen actor, {ActorLabel}, has an invalid name of {ActorName}.\n{ErrorMessage}"), Args);
+		
+		UE_LOG(LogPropertyNode, Warning, TEXT("%s"), *LogMessage.ToString());
+	}
+
+
+	
 	SetValue(InActor);
 	OnClose.ExecuteIfBound();
 }
diff --git a/Engine/Source/Runtime/Core/Private/UObject/NameTypesTests.cpp b/Engine/Source/Runtime/Core/Private/UObject/NameTypesTests.cpp
new file mode 100644
--- /dev/null
+++ b/Engine/Source/Runtime/Core/Private/UObject/NameTypesTests.cpp
@@ -0,0 +1,84 @@
+// Copyright Epic Games, Inc. All Rights Reserved.
+
+#include "CoreMinimal.h"
+#include "Misc/AutomationTest.h"
+
+
+
+namespace FNameInvalidCharsTest
+{
+
+	bool RunTest(const FString InvalidChars,
+		const FString& StringViewToTest,
+		const FString& ExpectedErrorMessage,
+		bool bExpectedIsValid,
+		FAutomationTestBase& TestBase)
+	{
+		FText OutMessage;
+		const bool bActualIsValid = FName::IsValidXName(FStringView(StringViewToTest), InvalidChars, &OutMessage);
+		const bool bErrorMessageResult = TestBase.TestEqual("Generated error message", OutMessage.ToString(), ExpectedErrorMessage);
+		const bool bValidityResult = TestBase.TestEqual("validity return value from FName::IsValidXName", bActualIsValid, bExpectedIsValid);
+
+		return bErrorMessageResult && bValidityResult;
+	}
+}
+
+IMPLEMENT_SIMPLE_AUTOMATION_TEST(DoubleBadCharactersTest, "System.Core.NameTypes.Test repeated invalid characters", EAutomationTestFlags::EditorContext | EAutomationTestFlags::EngineFilter);
+bool DoubleBadCharactersTest::RunTest(const FString& Parameters)
+{
+	return FNameInvalidCharsTest::RunTest(	
+		" \n\"\r\''\t\x000B",
+		TEXT(" \t\r\n'\"'  \r\n\'\"' \x000B"),
+		TEXT(
+			"Name may not contain whitespace characters (space, line feed [\\n], carriage return [\\r], tab) or the following characters: \" \' U+000B"
+			),
+		false,
+		*this);
+}
+
+IMPLEMENT_SIMPLE_AUTOMATION_TEST(SingleInvalidWhitespaceCharacterTest, "System.Core.NameTypes.Test single whitespace character", EAutomationTestFlags::EditorContext | EAutomationTestFlags::EngineFilter);
+bool SingleInvalidWhitespaceCharacterTest::RunTest(const FString& Parameters)
+{
+	return FNameInvalidCharsTest::RunTest(
+		" \n\"\r\''",
+		TEXT("test\nstring"),
+		TEXT(
+			"Name may not contain whitespace characters (line feed [\\n])"
+		),
+		false,
+		*this);
+}
+
+IMPLEMENT_SIMPLE_AUTOMATION_TEST(SingleInvalidNonWhitespaceCharacterTest, "System.Core.NameTypes.Test single non-whitespace character", EAutomationTestFlags::EditorContext | EAutomationTestFlags::EngineFilter);
+bool SingleInvalidNonWhitespaceCharacterTest::RunTest(const FString& Parameters)
+{
+	return FNameInvalidCharsTest::RunTest(
+		" \n\"\r\''",
+		TEXT("test\nstring"),
+		TEXT("Name may not contain whitespace characters (line feed [\\n])"),
+		false,
+		*this);
+}
+
+IMPLEMENT_SIMPLE_AUTOMATION_TEST(InvalidNonWhitespaceCharactersTest, "System.Core.NameTypes.Test invalid whitespace characters", EAutomationTestFlags::EditorContext | EAutomationTestFlags::EngineFilter);
+bool InvalidNonWhitespaceCharactersTest::RunTest(const FString& Parameters)
+{
+	return FNameInvalidCharsTest::RunTest(
+		"\";:QWop \n",
+		TEXT("test\";:QWERts \ntringyuiopQWERtstringyuiopqwertstringyuiop"),
+		TEXT(
+		"Name may not contain whitespace characters (space, line feed [\\n]) or the following characters: \" ; : Q W o p"),
+		false,
+		*this);
+}
+
+IMPLEMENT_SIMPLE_AUTOMATION_TEST(ValidFNameTest, "System.Core.NameTypes.Test valid FName", EAutomationTestFlags::EditorContext | EAutomationTestFlags::EngineFilter);
+bool ValidFNameTest::RunTest(const FString& Parameters)
+{
+	return FNameInvalidCharsTest::RunTest(
+		"\"\n\r",
+		TEXT("ValidName"),
+		TEXT(""),
+		true,
+		*this);
+}
\ No newline at end of file
diff --git a/Engine/Source/Runtime/Core/Private/UObject/UnrealNames.cpp b/Engine/Source/Runtime/Core/Private/UObject/UnrealNames.cpp
--- a/Engine/Source/Runtime/Core/Private/UObject/UnrealNames.cpp
+++ b/Engine/Source/Runtime/Core/Private/UObject/UnrealNames.cpp
@@ -3602,27 +3602,82 @@ bool FName::IsValidXName(const FStringView& InName, const FString& InInvalidChar
 		return true;
 	}
 
+	enum class EInvalidCharacters
+	{
+		Space = 32,
+		LineFeed = 10,
+		CarriageReturn = 13,
+		Tab = 9
+	};
+	
 	// See if the name contains invalid characters.
 	FString MatchedInvalidChars;
+	TArray<FText> MatchedInvalidWhitespaceChars;
 	TSet<TCHAR> AlreadyMatchedInvalidChars;
+		
 	for (const TCHAR InvalidChar : InInvalidChars)
 	{
 		int32 InvalidCharIndex = INDEX_NONE;
 		if (!AlreadyMatchedInvalidChars.Contains(InvalidChar) && InName.FindChar(InvalidChar, InvalidCharIndex))
 		{
-			MatchedInvalidChars.AppendChar(InvalidChar);
+			switch (static_cast<EInvalidCharacters>(InvalidChar))
+			{
+			case EInvalidCharacters::Space:
+				MatchedInvalidWhitespaceChars.Add(NSLOCTEXT("Core", "FNameSpaceName", "space"));
+				break;
+			case EInvalidCharacters::CarriageReturn:
+				MatchedInvalidWhitespaceChars.Add(NSLOCTEXT("Core", "FNameCarriageReturnName", "carriage return [\\r]"));
+				break;
+			case EInvalidCharacters::LineFeed:
+				MatchedInvalidWhitespaceChars.Add(NSLOCTEXT("Core", "FNameLineFeedName", "line feed [\\n]"));
+				break;
+			case EInvalidCharacters::Tab:
+				MatchedInvalidWhitespaceChars.Add(NSLOCTEXT("Core", "FNameTabName", "tab"));
+				break;
+			default:
+				MatchedInvalidChars.AppendChar(InvalidChar);
+			}
 			AlreadyMatchedInvalidChars.Add(InvalidChar);
 		}
 	}
 
-	if (MatchedInvalidChars.Len())
+	const bool bHasMatchedInvalidWhitespaceChars = MatchedInvalidWhitespaceChars.Num() > 0;
+	const bool bHasMatchedInvalidChars = MatchedInvalidChars.Len() > 0;
+
+	if (bHasMatchedInvalidChars || bHasMatchedInvalidWhitespaceChars)
 	{
 		if (OutReason)
 		{
+			FText InvalidWhiteSpaceCharsFText =
+				FText::Join(FText::FromString(TEXT(", ")), MatchedInvalidWhitespaceChars);
 			FFormatNamedArguments Args;
-			Args.Add(TEXT("ErrorCtx"), (InErrorCtx) ? *InErrorCtx : NSLOCTEXT("Core", "NameDefaultErrorCtx", "Name"));
-			Args.Add(TEXT("IllegalNameCharacters"), FText::FromString(MatchedInvalidChars));
-			*OutReason = FText::Format(NSLOCTEXT("Core", "NameContainsInvalidCharacters", "{ErrorCtx} may not contain the following characters: {IllegalNameCharacters}"), Args);
+			Args.Add(TEXT("ErrorCtx"), InErrorCtx ? *InErrorCtx : NSLOCTEXT("Core", "NameDefaultErrorCtx", "Name"));
+			
+			if (bHasMatchedInvalidWhitespaceChars)
+			{
+				Args.Add(TEXT("InvalidWhitespaceChars"), InvalidWhiteSpaceCharsFText);
+				*OutReason =  FText::Format(NSLOCTEXT("Core", "NameContainsInvalidWhitespaceCharacters", "Name may not contain whitespace characters ({InvalidWhitespaceChars})"), Args);	
+			}
+			if (bHasMatchedInvalidChars)
+			{
+				FString InvalidCharacters;
+				
+				for (TCHAR ch : MatchedInvalidChars)
+				{
+					if (FText::IsWhitespace(ch))
+					{
+						InvalidCharacters.Append(FString::Printf(TEXT(" U+%04X"), ch));
+						continue;
+					}
+					InvalidCharacters.Append(FString::Printf(TEXT(" %c"), ch));
+				}
+				Args.Add(TEXT("InvalidCharacters"), FText::FromString(InvalidCharacters));
+
+				*OutReason = bHasMatchedInvalidWhitespaceChars ?
+					FText::Format(NSLOCTEXT("Core", "NameContainsInvalidMixedCharacters", "{ErrorCtx} may not contain whitespace characters ({InvalidWhitespaceChars}) or the following characters:{InvalidCharacters}"), Args) :
+					FText::Format(NSLOCTEXT("Core", "NameContainsInvalidNonwhitespaceCharacters", "{ErrorCtx} may not contain the following characters:{InvalidCharacters}"), Args) ;
+			}
+
 		}
 		return false;
 	}
@@ -3630,6 +3685,21 @@ bool FName::IsValidXName(const FStringView& InName, const FString& InInvalidChar
 	return true;
 }
 
+FString FName::SanitizeWhitespace(const FString& FNameString)
+{
+	FString ReturnStr;
+
+	for (const TCHAR CurrentChar : FNameString)
+	{
+		if (!FText::IsWhitespace(CurrentChar) || (CurrentChar == ' ' || CurrentChar =='\t'))
+		{
+			ReturnStr.AppendChar(CurrentChar);
+		}	
+	}
+
+	return ReturnStr;
+}
+
 FWideStringBuilderBase& operator<<(FWideStringBuilderBase& Builder, FNameEntryId Id)
 {
 	FName::GetEntry(Id)->AppendNameToString(Builder);
diff --git a/Engine/Source/Runtime/Core/Public/UObject/NameTypes.h b/Engine/Source/Runtime/Core/Public/UObject/NameTypes.h
--- a/Engine/Source/Runtime/Core/Public/UObject/NameTypes.h
+++ b/Engine/Source/Runtime/Core/Public/UObject/NameTypes.h
@@ -769,6 +769,16 @@ public:
 		return IsValidXName(*this, INVALID_LONGPACKAGE_CHARACTERS, &OutReason);
 	}
 
+	/**
+	 * Printing FNames in logging or on screen can be problematic when they contain Whitespace characters such as \n and \r,
+	 * so this will return an FName based upon the calling FName, but with any Whitespace characters potentially problematic for
+	 * showing in a log or on screen omitted.
+	 *
+	* @return the new FName based upon the calling FName, but with any Whitespace characters potentially problematic for
+	 * showing in a log or on screen omitted.
+	 */
+	static FString SanitizeWhitespace(const FString& FNameString);
+	
 	/**
 	 * Compares name to passed in one. Sort is alphabetical ascending.
 	 *
