[Core] temporarily disabled test

diff --git a/Engine/Source/Runtime/Core/Tests/Misc/AssertionMacrosTest.cpp b/Engine/Source/Runtime/Core/Tests/Misc/AssertionMacrosTest.cpp
--- a/Engine/Source/Runtime/Core/Tests/Misc/AssertionMacrosTest.cpp
+++ b/Engine/Source/Runtime/Core/Tests/Misc/AssertionMacrosTest.cpp
@@ -114,7 +114,7 @@ TEST_CASE("Core::Misc::TestEnsure", "[Core][Misc][AssertionMacros][Ensure]")
 	}
 }
 
-TEST_CASE("Core::Misc::TestCheck", "[Core][Misc][AssertionMacros][Check]")
+DISABLED_TEST_CASE("Core::Misc::TestCheck", "[Core][Misc][AssertionMacros][Check]")
 {
 	REQUIRE_CHECK(check(1 == 2));
 	REQUIRE_CHECK_MSG("1 == 2", checkf(1 == 2, TEXT("Error Message")));
