Fix emplace failing to deduce initializer lists.

diff --git a/Engine/Source/Runtime/Core/Public/Containers/Map.h b/Engine/Source/Runtime/Core/Public/Containers/Map.h
--- a/Engine/Source/Runtime/Core/Public/Containers/Map.h
+++ b/Engine/Source/Runtime/Core/Public/Containers/Map.h
@@ -411,7 +411,7 @@ public:
 	 * @param InKey The key to associate the value with.
 	 * @param InValue The value to associate with the key.
 	 * @return A reference to the value as stored in the map. The reference is only valid until the next change to any key in the map.	 */
-	template <typename InitKeyType, typename InitValueType>
+	template <typename InitKeyType = KeyType, typename InitValueType = ValueType>
 	ValueType& Emplace(InitKeyType&& InKey, InitValueType&& InValue)
 	{
 		const FSetElementId PairId = Pairs.Emplace(TPairInitializer<InitKeyType&&, InitValueType&&>(Forward<InitKeyType>(InKey), Forward<InitValueType>(InValue)));
@@ -420,7 +420,7 @@ public:
 	}
 
 	/** See Emplace() and class documentation section on ByHash() functions */
-	template <typename InitKeyType, typename InitValueType>
+	template <typename InitKeyType = KeyType, typename InitValueType = ValueType>
 	ValueType& EmplaceByHash(uint32 KeyHash, InitKeyType&& InKey, InitValueType&& InValue)
 	{
 		const FSetElementId PairId = Pairs.EmplaceByHash(KeyHash, TPairInitializer<InitKeyType&&, InitValueType&&>(Forward<InitKeyType>(InKey), Forward<InitValueType>(InValue)));
@@ -434,7 +434,7 @@ public:
 	 * @param InKey The key to associate the value with.
 	 * @return A reference to the value as stored in the map. The reference is only valid until the next change to any key in the map.
 	 */
-	template <typename InitKeyType>
+	template <typename InitKeyType = KeyType>
 	ValueType& Emplace(InitKeyType&& InKey)
 	{
 		const FSetElementId PairId = Pairs.Emplace(TKeyInitializer<InitKeyType&&>(Forward<InitKeyType>(InKey)));
@@ -443,7 +443,7 @@ public:
 	}
 
 	/** See Emplace() and class documentation section on ByHash() functions */
-	template <typename InitKeyType>
+	template <typename InitKeyType = KeyType>
 	ValueType& EmplaceByHash(uint32 KeyHash, InitKeyType&& InKey)
 	{
 		const FSetElementId PairId = Pairs.EmplaceByHash(KeyHash, TKeyInitializer<InitKeyType&&>(Forward<InitKeyType>(InKey)));
diff --git a/Engine/Source/Runtime/Core/Public/Containers/Set.h b/Engine/Source/Runtime/Core/Public/Containers/Set.h
--- a/Engine/Source/Runtime/Core/Public/Containers/Set.h
+++ b/Engine/Source/Runtime/Core/Public/Containers/Set.h
@@ -717,7 +717,7 @@ public:
 	 * @param	bIsAlreadyInSetPtr	[out]	Optional pointer to bool that will be set depending on whether element is already in set
 	 * @return	A handle to the element stored in the set.
 	 */
-	template <typename ArgsType>
+	template <typename ArgsType = ElementType>
 	FSetElementId Emplace(ArgsType&& Args, bool* bIsAlreadyInSetPtr = nullptr)
 	{
 		// Create a new element.
@@ -741,7 +741,7 @@ public:
 	 * @param	bIsAlreadyInSetPtr	[out]	Optional pointer to bool that will be set depending on whether element is already in set
 	 * @return	A handle to the element stored in the set.
 	 */
-	template <typename ArgsType>
+	template <typename ArgsType = ElementType>
 	FSetElementId EmplaceByHash(uint32 KeyHash, ArgsType&& Args, bool* bIsAlreadyInSetPtr = nullptr)
 	{
 		// Create a new element.
diff --git a/Engine/Source/Runtime/Core/Public/Containers/SortedMap.h b/Engine/Source/Runtime/Core/Public/Containers/SortedMap.h
--- a/Engine/Source/Runtime/Core/Public/Containers/SortedMap.h
+++ b/Engine/Source/Runtime/Core/Public/Containers/SortedMap.h
@@ -187,7 +187,7 @@ public:
 	 * @param InValue - The value to associate with the key.
 	 * @return A reference to the value as stored in the map (only valid until the next change to any key in the map).
 	 */
-	template <typename InitKeyType, typename InitValueType>
+	template <typename InitKeyType = KeyType, typename InitValueType = ValueType>
 	ValueType& Emplace(InitKeyType&& InKey, InitValueType&& InValue)
 	{
 		ElementType* DataPtr = AllocateMemoryForEmplace(InKey);
@@ -203,7 +203,7 @@ public:
 	 * @param InKey The key to associate the value with.
 	 * @return A reference to the value as stored in the map (only valid until the next change to any key in the map).
 	 */
-	template <typename InitKeyType>
+	template <typename InitKeyType = KeyType>
 	ValueType& Emplace(InitKeyType&& InKey)
 	{
 		ElementType* DataPtr = AllocateMemoryForEmplace(InKey);
