TSet & TMap: Make sure any dynamic hash allocation is released when Empty() is called without additional slack.
Memory improvements on top of performance improvements in 11535961.
Allow the hash allocation to shrink for calls to Empty(), Shrink(), Compact() and CompactStable().
Let TSparseArray::Compact() shrink the data array just as for Shrink() and CompactStable().
Note: The private function TSet::ShouldRehash() is only called from Empty() and ConditionalRehash(), and ConditionalRehash() with bAllowShrinking=true is only called from Relax() via Shrink(), so there should be no impact on other/growth scenarios.
Modify GetAllocatedSize() for TSet, TMap and TSparseArray to only return any dynamically allocated size from the underlying allocators and change return type from uint32 to SIZE_T (just as for TArray and FContainerAllocatorInterface).

diff --git a/Engine/Source/Runtime/ApplicationCore/Private/Windows/Accessibility/WindowsUIAManager.cpp b/Engine/Source/Runtime/ApplicationCore/Private/Windows/Accessibility/WindowsUIAManager.cpp
--- a/Engine/Source/Runtime/ApplicationCore/Private/Windows/Accessibility/WindowsUIAManager.cpp
+++ b/Engine/Source/Runtime/ApplicationCore/Private/Windows/Accessibility/WindowsUIAManager.cpp
@@ -312,20 +312,20 @@ void FWindowsUIAManager::DumpAccessibilityStats() const
 	// This isn't exactly right since some ControlProviders will be TextRangeProviders, but it should be close
 	const uint32 NumControlProviders = ProviderList.Num() - NumWidgetProviders;
 
-	const uint32 SizeOfWidgetProvider = sizeof(FWindowsUIAWidgetProvider);
-	const uint32 SizeOfControlProvider = sizeof(FWindowsUIAControlProvider);
-	const uint32 SizeOfCachedWidgetProviders = CachedWidgetProviders.GetAllocatedSize();
-	const uint32 SizeOfProviderList = ProviderList.GetAllocatedSize();
-	const uint32 CacheSize = NumWidgetProviders * SizeOfWidgetProvider + NumControlProviders * SizeOfControlProvider + SizeOfCachedWidgetProviders + SizeOfProviderList;
+	const SIZE_T SizeOfWidgetProvider = sizeof(FWindowsUIAWidgetProvider);
+	const SIZE_T SizeOfControlProvider = sizeof(FWindowsUIAControlProvider);
+	const SIZE_T SizeOfCachedWidgetProviders = CachedWidgetProviders.GetAllocatedSize();
+	const SIZE_T SizeOfProviderList = ProviderList.GetAllocatedSize();
+	const SIZE_T CacheSize = NumWidgetProviders * SizeOfWidgetProvider + NumControlProviders * SizeOfControlProvider + SizeOfCachedWidgetProviders + SizeOfProviderList;
 
 	UE_LOG(LogAccessibility, Log, TEXT("Dumping Windows accessibility stats:"));
 	UE_LOG(LogAccessibility, Log, TEXT("Number of Widget Providers: %i"), NumWidgetProviders);
 	UE_LOG(LogAccessibility, Log, TEXT("Number of non-Widget Providers: %i"), NumControlProviders);
-	UE_LOG(LogAccessibility, Log, TEXT("Size of FWindowsUIAWidgetProvider: %u"), SizeOfWidgetProvider);
-	UE_LOG(LogAccessibility, Log, TEXT("Size of FWindowsUIAControlProvider: %u"), SizeOfControlProvider);
-	UE_LOG(LogAccessibility, Log, TEXT("Size of WidgetProvider* map: %u"), SizeOfCachedWidgetProviders);
-	UE_LOG(LogAccessibility, Log, TEXT("Size of all Provider* set: %u"), SizeOfProviderList);
-	UE_LOG(LogAccessibility, Log, TEXT("Memory stored in cache: %u kb"), CacheSize / 1000);
+	UE_LOG(LogAccessibility, Log, TEXT("Size of FWindowsUIAWidgetProvider: %" SIZE_T_FMT), SizeOfWidgetProvider);
+	UE_LOG(LogAccessibility, Log, TEXT("Size of FWindowsUIAControlProvider: %" SIZE_T_FMT), SizeOfControlProvider);
+	UE_LOG(LogAccessibility, Log, TEXT("Size of WidgetProvider* map: %" SIZE_T_FMT), SizeOfCachedWidgetProviders);
+	UE_LOG(LogAccessibility, Log, TEXT("Size of all Provider* set: %" SIZE_T_FMT), SizeOfProviderList);
+	UE_LOG(LogAccessibility, Log, TEXT("Memory stored in cache: %" SIZE_T_FMT " kb"), CacheSize / 1000);
 }
 #endif
 
diff --git a/Engine/Source/Runtime/AssetRegistry/Private/AssetRegistry.cpp b/Engine/Source/Runtime/AssetRegistry/Private/AssetRegistry.cpp
--- a/Engine/Source/Runtime/AssetRegistry/Private/AssetRegistry.cpp
+++ b/Engine/Source/Runtime/AssetRegistry/Private/AssetRegistry.cpp
@@ -4172,9 +4172,9 @@ void FAssetRegistryImpl::GetAllocatedSize(bool bLogDetailed, SIZE_T& StateSize,
 
 	if (bIsTempCachingEnabled && !bIsTempCachingAlwaysEnabled)
 	{
-		uint32 TempCacheMem = TempCachedInheritanceBuffer.GetAllocatedSize();
+		SIZE_T TempCacheMem = TempCachedInheritanceBuffer.GetAllocatedSize();
 		StaticSize += TempCacheMem;
-		UE_LOG(LogAssetRegistry, Warning, TEXT("Asset Registry Temp caching enabled, wasting memory: %dk"), TempCacheMem / 1024);
+		UE_LOG(LogAssetRegistry, Warning, TEXT("Asset Registry Temp caching enabled, wasting memory: %lldk"), TempCacheMem / 1024);
 	}
 
 	if (GlobalGatherer.IsValid())
@@ -5970,7 +5970,7 @@ bool FClassInheritanceBuffer::IsUpToDate(uint64 CurrentClassesVersionNumber) con
 	return !bDirty && RegisteredClassesVersionNumber == CurrentClassesVersionNumber;
 }
 
-uint32 FClassInheritanceBuffer::GetAllocatedSize() const
+SIZE_T FClassInheritanceBuffer::GetAllocatedSize() const
 {
 	return InheritanceMap.GetAllocatedSize() + ReverseInheritanceMap.GetAllocatedSize();
 }
diff --git a/Engine/Source/Runtime/AssetRegistry/Private/AssetRegistryImpl.h b/Engine/Source/Runtime/AssetRegistry/Private/AssetRegistryImpl.h
--- a/Engine/Source/Runtime/AssetRegistry/Private/AssetRegistryImpl.h
+++ b/Engine/Source/Runtime/AssetRegistry/Private/AssetRegistryImpl.h
@@ -92,7 +92,7 @@ namespace Impl
 		/** Delete data and free allocations. */
 		void Clear();
 		/** Report size of dynamic allocations. */
-		uint32 GetAllocatedSize() const;
+		SIZE_T GetAllocatedSize() const;
 	};
 
 	/** Status of gathering, returned from the Tick function */
diff --git a/Engine/Source/Runtime/AssetRegistry/Public/AssetRegistry/PathTree.h b/Engine/Source/Runtime/AssetRegistry/Public/AssetRegistry/PathTree.h
--- a/Engine/Source/Runtime/AssetRegistry/Public/AssetRegistry/PathTree.h
+++ b/Engine/Source/Runtime/AssetRegistry/Public/AssetRegistry/PathTree.h
@@ -39,9 +39,9 @@ public:
 		return ParentPathToChildPaths.Num();
 	}
 
-	uint32 GetAllocatedSize(void) const
+	SIZE_T GetAllocatedSize(void) const
 	{
-		uint32 AllocatedSize = ParentPathToChildPaths.GetAllocatedSize() + ChildPathToParentPath.GetAllocatedSize();
+		SIZE_T AllocatedSize = ParentPathToChildPaths.GetAllocatedSize() + ChildPathToParentPath.GetAllocatedSize();
 
 		for (const TPair<FName, TSet<FName>>& Pair : ParentPathToChildPaths)
 		{
diff --git a/Engine/Source/Runtime/Core/Private/Containers/ContainersTest.cpp b/Engine/Source/Runtime/Core/Private/Containers/ContainersTest.cpp
--- a/Engine/Source/Runtime/Core/Private/Containers/ContainersTest.cpp
+++ b/Engine/Source/Runtime/Core/Private/Containers/ContainersTest.cpp
@@ -312,45 +312,87 @@ namespace
 	{
 		ContainerType Cont;
 
+		int32 MaxNum = 0;
+		SIZE_T MaxAllocatedSize = 0;
+		const SIZE_T InitialAllocatedSize = Cont.GetAllocatedSize();
+
 		ContainerTestStats.Reset();
+
+		auto CheckContainer = [&Cont]()
+		{
+			CheckContainerNum(Cont);
+			CheckContainerEnds(Cont);
+			CheckContainerElements(Cont);
+			CheckContainerSelfEquality(Cont);
+		};
+
+		// Test Add and Remove
 		// Subtract one to account for temporaries that will be created during an Add
 		for (int32 Count = 0; Count < MAX_TEST_OBJECTS - 1; Count += MAX_TEST_OBJECTS_STEP)
 		{
 			for (int32 N = 0; N != Count; ++N)
 			{
 				Cont.Add(GenerateTestKey<KeyType>(N), FContainerTestValueType(TEXT("New Value")));
-				CheckContainerNum(Cont);
-				CheckContainerEnds(Cont);
-				CheckContainerElements(Cont);
-				CheckContainerSelfEquality(Cont);
+				CheckContainer();
 			}
+			MaxNum = Cont.Num();
+			MaxAllocatedSize = Cont.GetAllocatedSize();
 
 			for (int32 N = 0; N != Count; ++N)
 			{
 				Cont.Remove(GenerateTestKey<KeyType>(N));
-				CheckContainerNum(Cont);
-				CheckContainerEnds(Cont);
-				CheckContainerElements(Cont);
-				CheckContainerSelfEquality(Cont);
+				CheckContainer();
 			}
 
+			check(Cont.IsEmpty());
+
 			for (int32 N = 0; N != Count; ++N)
 			{
 				Cont.Add(GenerateTestKey<KeyType>((Count - 1) - N), FContainerTestValueType(TEXT("New Value")));
-				CheckContainerNum(Cont);
-				CheckContainerEnds(Cont);
-				CheckContainerElements(Cont);
-				CheckContainerSelfEquality(Cont);
+				CheckContainer();
 			}
 
 			for (int32 N = 0; N != Count; ++N)
 			{
 				Cont.Remove(GenerateTestKey<KeyType>(N));
-				CheckContainerNum(Cont);
-				CheckContainerEnds(Cont);
-				CheckContainerElements(Cont);
-				CheckContainerSelfEquality(Cont);
+				CheckContainer();
+			}
+
+			check(Cont.IsEmpty());
+		}
+
+		// Test Empty and Shrink 
+		{
+			// Test releasing memory allocations
+			Cont.Empty();
+			CheckContainer();
+			check(Cont.GetAllocatedSize() == InitialAllocatedSize);
+
+			// Test integrity after re-growing container to MaxNum elements again
+			for (int32 N = 0; N < MaxNum; ++N)
+			{
+				Cont.Add(GenerateTestKey<KeyType>(N), FContainerTestValueType(TEXT("New Value")));
+			}
+			CheckContainer();
+			check(Cont.GetAllocatedSize() == MaxAllocatedSize);
+
+			// Test data integrity while removing and shrinking continously
+			{
+				SIZE_T PrevAllocatedSize = Cont.GetAllocatedSize();
+				for (int32 N = MaxNum - 1; N >= MaxNum / 4; --N)
+				{
+					Cont.Remove(GenerateTestKey<KeyType>(N));
+					Cont.Shrink();
+					CheckContainer();
+					check(Cont.GetAllocatedSize() <= PrevAllocatedSize);
+					PrevAllocatedSize = Cont.GetAllocatedSize();
+				}
 			}
+
+			// Test removing and releasing remaining elements
+			Cont.Empty();
+			check(Cont.IsEmpty());
+			check(Cont.GetAllocatedSize() == InitialAllocatedSize);
 		}
 	}
 
@@ -511,6 +553,8 @@ IMPLEMENT_SIMPLE_AUTOMATION_TEST(FContainersSmokeTest, "System.Core.Containers.S
 bool FContainersSmokeTest::RunTest( const FString& Parameters )
 {
 	RunContainerTests<TMap<int32, FContainerTestValueType>, int32>();
+	RunContainerTests<TMap<int32, FContainerTestValueType, TInlineSetAllocator<32>>, int32>();
+	RunContainerTests<TMap<int32, FContainerTestValueType, TFixedSetAllocator<64>>, int32>();
 
 	return true;
 }
@@ -521,7 +565,9 @@ bool FContainersFullTest::RunTest(const FString& Parameters)
 	RunContainerTests<TMap<int32, FContainerTestValueType>, int32>();
 	RunContainerTests<TMap<FName, FContainerTestValueType>, FName>();
 	RunContainerTests<TMap<FString, FContainerTestValueType>, FString>();
+	RunContainerTests<TMap<int32, FContainerTestValueType, TInlineSetAllocator<32>>, int32>();
 	RunContainerTests<TMap<int32, FContainerTestValueType, TInlineSetAllocator<64>>, int32>();
+	RunContainerTests<TMap<int32, FContainerTestValueType, TFixedSetAllocator<64>>, int32>();
 	RunContainerTests<TMap<FString, FContainerTestValueType, FDefaultSetAllocator, FCaseSensitiveLookupKeyFuncs<FContainerTestValueType>>, FString>();
 
 	RunContainerTests<TSortedMap<int32, FContainerTestValueType>, int32>();
diff --git a/Engine/Source/Runtime/Core/Public/Containers/Map.h b/Engine/Source/Runtime/Core/Public/Containers/Map.h
--- a/Engine/Source/Runtime/Core/Public/Containers/Map.h
+++ b/Engine/Source/Runtime/Core/Public/Containers/Map.h
@@ -349,7 +349,7 @@ public:
 	 * @return Number of bytes allocated by this container.
 	 * @see CountBytes
 	 */
-	FORCEINLINE uint32 GetAllocatedSize() const
+	FORCEINLINE SIZE_T GetAllocatedSize() const
 	{
 		return Pairs.GetAllocatedSize();
 	}
diff --git a/Engine/Source/Runtime/Core/Public/Containers/Set.h b/Engine/Source/Runtime/Core/Public/Containers/Set.h
--- a/Engine/Source/Runtime/Core/Public/Containers/Set.h
+++ b/Engine/Source/Runtime/Core/Public/Containers/Set.h
@@ -488,6 +488,7 @@ public:
 	{
 		if (Elements.Compact())
 		{
+			HashSize = Allocator::GetNumberOfHashBuckets(Elements.Num());
 			Rehash();
 		}
 	}
@@ -497,6 +498,7 @@ public:
 	{
 		if (Elements.CompactStable())
 		{
+			HashSize = Allocator::GetNumberOfHashBuckets(Elements.Num());
 			Rehash();
 		}
 	}
@@ -535,9 +537,9 @@ public:
 	 * Only returns the size of allocations made directly by the container, not the elements themselves.
 	 * @return number of bytes allocated by this container
 	 */
-	FORCEINLINE uint32 GetAllocatedSize( void ) const
+	FORCEINLINE SIZE_T GetAllocatedSize( void ) const
 	{
-		return Elements.GetAllocatedSize() + (HashSize * sizeof(FSetElementId));
+		return Elements.GetAllocatedSize() + Hash.GetAllocatedSize(HashSize, sizeof(FSetElementId));
 	}
 
 	/** Tracks the container's memory use through an archive. */
@@ -1525,10 +1527,8 @@ private:
 	FORCEINLINE bool ShouldRehash(int32 NumHashedElements,int32 DesiredHashSize,bool bAllowShrinking = false) const
 	{
 		// If the hash hasn't been created yet, or is smaller than the desired hash size, rehash.
-		return (NumHashedElements > 0 &&
-				(!HashSize ||
-				HashSize < DesiredHashSize ||
-				(HashSize > DesiredHashSize && bAllowShrinking)));
+		// If shrinking is allowed and the hash is bigger than the desired hash size, rehash.
+		return ((NumHashedElements > 0 && HashSize < DesiredHashSize) || (bAllowShrinking && HashSize > DesiredHashSize));
 	}
 
 	/**
diff --git a/Engine/Source/Runtime/Core/Public/Containers/SparseArray.h b/Engine/Source/Runtime/Core/Public/Containers/SparseArray.h
--- a/Engine/Source/Runtime/Core/Public/Containers/SparseArray.h
+++ b/Engine/Source/Runtime/Core/Public/Containers/SparseArray.h
@@ -478,7 +478,7 @@ public:
 			}
 
 			// Truncate unallocated elements at the end of the data array.
-			Data.RemoveAt(FirstIndexToRemove,Data.Num() - FirstIndexToRemove);
+			Data.RemoveAt(FirstIndexToRemove,Data.Num() - FirstIndexToRemove, false);
 			AllocationFlags.RemoveAt(FirstIndexToRemove,AllocationFlags.Num() - FirstIndexToRemove);
 		}
 
@@ -524,12 +524,15 @@ public:
 			FreeIndex = NextFreeIndex;
 		}
 
-		Data           .RemoveAt(TargetIndex, NumFree);
+		Data           .RemoveAt(TargetIndex, NumFree, false);
 		AllocationFlags.RemoveAt(TargetIndex, NumFree);
 
 		NumFreeIndices = 0;
 		FirstFreeIndex = -1;
 
+		// Shrink the data array.
+		Data.Shrink();
+
 		return bResult;
 	}
 
@@ -601,10 +604,9 @@ public:
 	 * Only returns the size of allocations made directly by the container, not the elements themselves.
 	 * @return number of bytes allocated by this container
 	 */
-	uint32 GetAllocatedSize( void ) const
+	SIZE_T GetAllocatedSize( void ) const
 	{
-		return	(Data.Num() + Data.GetSlack()) * sizeof(FElementOrFreeListLink) +
-				AllocationFlags.GetAllocatedSize();
+		return Data.GetAllocatedSize() + AllocationFlags.GetAllocatedSize();
 	}
 
 	/** Tracks the container's memory use through an archive. */
