Fix a whole class of TSAN false positive by making the boolean used for polling request atomic
  - That way, TSAN knows that destructing a request after just checking the bool is safe
  - TSAN doesn't support barriers, so we have to do that to make it understand everything it safe
  - Remove other TSAN_SAFE that are not needed anymore

diff --git a/Engine/Source/Runtime/Core/Public/Async/AsyncFileHandle.h b/Engine/Source/Runtime/Core/Public/Async/AsyncFileHandle.h
--- a/Engine/Source/Runtime/Core/Public/Async/AsyncFileHandle.h
+++ b/Engine/Source/Runtime/Core/Public/Async/AsyncFileHandle.h
@@ -29,10 +29,10 @@ protected:
 		uint8* Memory;
 	};
 	FAsyncFileCallBack Callback;
-	bool bDataIsReady;
-	bool bCompleteAndCallbackCalled;
-	bool bCompleteSync;
-	bool bCanceled;
+	TSAN_ATOMIC(bool) bDataIsReady;
+	TSAN_ATOMIC(bool) bCompleteAndCallbackCalled;
+	TSAN_ATOMIC(bool) bCompleteSync;
+	TSAN_ATOMIC(bool) bCanceled;
 	const bool bSizeRequest;
 	const bool bUserSuppliedMemory;
 public:
@@ -63,7 +63,7 @@ public:
 	}
 
 	/* Not legal to destroy the request until it is complete. */
-	virtual ~IAsyncReadRequest() TSAN_SAFE
+	virtual ~IAsyncReadRequest()
 	{
 		// check(bCompleteAndCallbackCalled && (bSizeRequest || !Memory)); // must be complete, and if it was a read request, the memory should be gone
 		UE_CLOG(!(bCompleteAndCallbackCalled && (bSizeRequest || !Memory)),
@@ -76,7 +76,7 @@ public:
 	* Nonblocking poll of the state of completion.
 	* @return true if the request is complete
 	**/
-	FORCEINLINE bool PollCompletion() TSAN_SAFE
+	FORCEINLINE bool PollCompletion()
 	{
 		return bCompleteAndCallbackCalled;
 	}
@@ -125,7 +125,7 @@ public:
 	* Return the bytes of a completed read request. Not legal to call unless the request is complete.
 	* @return Returned memory block which if non-null contains the bytes read. Caller owns the memory block and must call FMemory::Free on it when done. Can be null if the file was not found or could not be read or the request was cancelled, or the request had AIOP_FLAG_PRECACHE.
 	**/
-	FORCEINLINE uint8* GetReadResults() TSAN_SAFE
+	FORCEINLINE uint8* GetReadResults()
 	{
 		check(bDataIsReady && !bSizeRequest);
 		uint8* Result = Memory;
@@ -151,7 +151,7 @@ protected:
 	/** Cancel the request. This is a non-blocking async call and so does not ensure completion! **/
 	virtual void CancelImpl() = 0;
 
-	void SetDataComplete() TSAN_SAFE
+	void SetDataComplete()
 	{
 		bDataIsReady = true;
 		FPlatformMisc::MemoryBarrier();
@@ -162,13 +162,13 @@ protected:
 		FPlatformMisc::MemoryBarrier();
 	}
 
-	void SetAllComplete() TSAN_SAFE
+	void SetAllComplete()
 	{
 		bCompleteAndCallbackCalled = true;
 		FPlatformMisc::MemoryBarrier();
 	}
 
-	void SetComplete() TSAN_SAFE
+	void SetComplete()
 	{
 		SetDataComplete();
 		SetAllComplete();
