[Insights] Added tracing of callstack also for allocation free events.

diff --git a/Engine/Source/Runtime/Core/Private/ProfilingDebugging/MemoryTrace.cpp b/Engine/Source/Runtime/Core/Private/ProfilingDebugging/MemoryTrace.cpp
--- a/Engine/Source/Runtime/Core/Private/ProfilingDebugging/MemoryTrace.cpp
+++ b/Engine/Source/Runtime/Core/Private/ProfilingDebugging/MemoryTrace.cpp
@@ -19,33 +19,33 @@ UE_TRACE_CHANNEL_DEFINE(MemAllocChannel, "Memory allocations", true)
 #if UE_MEMORY_TRACE_ENABLED
 
 ////////////////////////////////////////////////////////////////////////////////
-void 	MemoryTrace_InitTags(FMalloc*);
-void	MemoryTrace_EnableTracePump();
+void MemoryTrace_InitTags(FMalloc*);
+void MemoryTrace_EnableTracePump();
 
 ////////////////////////////////////////////////////////////////////////////////
 namespace
 {
 	// Controls how often time markers are emitted (default every 4095 allocation)
-	constexpr uint32 MarkerSamplePeriod	= (4 << 10) - 1;
-	
+	constexpr uint32 MarkerSamplePeriod = (4 << 10) - 1;
+
 	// Number of bits shifted bits to SizeLower
 	constexpr uint32 SizeShift = 3;
-	
+
 	// Counter to track when time marker is emitted
-	std::atomic<uint32>	GMarkerCounter(0);
-	
+	std::atomic<uint32> GMarkerCounter(0);
+
 	// If enabled also pumps the Trace system itself. Used on process shutdown
 	// when worker thread has been killed, but memory events still occurs.
-	bool				GDoPumpTrace;
-	
+	bool GDoPumpTrace;
+
 	// Temporarily disables any internal operation that causes allocations. Used to
 	// avoid recursive behaviour when memory tracing needs to allocate memory through
 	// TraceMalloc
-	thread_local bool	GDoNotAllocateInTrace;
-	
+	thread_local bool GDoNotAllocateInTrace;
+
 	// Set on initialization, on some platforms we hook allocator functions very early
 	// before Trace has the ability to allocate memory.
-	bool				GTraceAllowed; 
+	bool GTraceAllowed;
 }
 
 ////////////////////////////////////////////////////////////////////////////////
@@ -93,15 +93,18 @@ UE_TRACE_EVENT_END()
 
 UE_TRACE_EVENT_BEGIN(Memory, Free)
 	UE_TRACE_EVENT_FIELD(uint64, Address)
+	UE_TRACE_EVENT_FIELD(uint32, CallstackId)
 	UE_TRACE_EVENT_FIELD(uint8, RootHeap)
 UE_TRACE_EVENT_END()
 
 UE_TRACE_EVENT_BEGIN(Memory, FreeSystem)
 	UE_TRACE_EVENT_FIELD(uint64, Address)
+	UE_TRACE_EVENT_FIELD(uint32, CallstackId)
 UE_TRACE_EVENT_END()
 
 UE_TRACE_EVENT_BEGIN(Memory, FreeVideo)
 	UE_TRACE_EVENT_FIELD(uint64, Address)
+	UE_TRACE_EVENT_FIELD(uint32, CallstackId)
 UE_TRACE_EVENT_END()
 
 UE_TRACE_EVENT_BEGIN(Memory, ReallocAlloc)
@@ -154,32 +157,32 @@ class FMallocWrapper
 	: public FMalloc
 {
 public:
-							FMallocWrapper(FMalloc* InMalloc);
+	FMallocWrapper(FMalloc* InMalloc);
 
 private:
 	struct FCookie
 	{
-		uint64				Tag  : 16;
-		uint64				Bias : 8;
-		uint64				Size : 40;
+		uint64 Tag  : 16;
+		uint64 Bias : 8;
+		uint64 Size : 40;
 	};
 
-	static uint32			GetActualAlignment(SIZE_T Size, uint32 Alignment);
-	virtual void*			Malloc(SIZE_T Size, uint32 Alignment) override;
-	virtual void*			Realloc(void* PrevAddress, SIZE_T NewSize, uint32 Alignment) override;
-	virtual void			Free(void* Address) override;
-	virtual bool			IsInternallyThreadSafe() const override						{ return InnerMalloc->IsInternallyThreadSafe(); }
-	virtual void			UpdateStats() override										{ InnerMalloc->UpdateStats(); }
-	virtual void			GetAllocatorStats(FGenericMemoryStats& Out) override		{ InnerMalloc->GetAllocatorStats(Out); }
-	virtual void			DumpAllocatorStats(FOutputDevice& Ar) override				{ InnerMalloc->DumpAllocatorStats(Ar); }
-	virtual bool			ValidateHeap() override										{ return InnerMalloc->ValidateHeap(); }
-	virtual bool			GetAllocationSize(void* Address, SIZE_T &SizeOut) override	{ return InnerMalloc->GetAllocationSize(Address, SizeOut); }
-	virtual void			SetupTLSCachesOnCurrentThread() override					{ return InnerMalloc->SetupTLSCachesOnCurrentThread(); }
-	virtual void			OnMallocInitialized() override								{ InnerMalloc->OnMallocInitialized(); }
-	virtual void			OnPreFork() override										{ InnerMalloc->OnPreFork(); }
-	virtual void			OnPostFork() override										{ InnerMalloc->OnPostFork(); }
-
-	FMalloc*				InnerMalloc;
+	static uint32 GetActualAlignment(SIZE_T Size, uint32 Alignment);
+	virtual void* Malloc(SIZE_T Size, uint32 Alignment) override;
+	virtual void* Realloc(void* PrevAddress, SIZE_T NewSize, uint32 Alignment) override;
+	virtual void Free(void* Address) override;
+	virtual bool IsInternallyThreadSafe() const override                    { return InnerMalloc->IsInternallyThreadSafe(); }
+	virtual void UpdateStats() override                                     { InnerMalloc->UpdateStats(); }
+	virtual void GetAllocatorStats(FGenericMemoryStats& Out) override       { InnerMalloc->GetAllocatorStats(Out); }
+	virtual void DumpAllocatorStats(FOutputDevice& Ar) override             { InnerMalloc->DumpAllocatorStats(Ar); }
+	virtual bool ValidateHeap() override                                    { return InnerMalloc->ValidateHeap(); }
+	virtual bool GetAllocationSize(void* Address, SIZE_T &SizeOut) override { return InnerMalloc->GetAllocationSize(Address, SizeOut); }
+	virtual void SetupTLSCachesOnCurrentThread() override                   { return InnerMalloc->SetupTLSCachesOnCurrentThread(); }
+	virtual void OnMallocInitialized() override                             { InnerMalloc->OnMallocInitialized(); }
+	virtual void OnPreFork() override                                       { InnerMalloc->OnPreFork(); }
+	virtual void OnPostFork() override                                      { InnerMalloc->OnPostFork(); }
+
+	FMalloc* InnerMalloc;
 };
 
 ////////////////////////////////////////////////////////////////////////////////
@@ -269,8 +272,8 @@ public:
 		return bIsConstructed;
 	}
 
-	T* operator & ()	{ return (T*)Buffer; }
-	T* operator -> ()	{ return (T*)Buffer; }
+	T* operator & ()  { return (T*)Buffer; }
+	T* operator -> () { return (T*)Buffer; }
 
 protected:
 	uint8 Buffer[sizeof(T)];
@@ -365,7 +368,7 @@ void MemoryTrace_Initialize()
 {
 	// At this point we initialized the system to allow tracing.
 	GTraceAllowed = true;
-	
+
 	UE_TRACE_LOG(Memory, Init, MemAllocChannel)
 		<< Init.MarkerPeriod(MarkerSamplePeriod + 1)
 		<< Init.Version(MemoryTraceVersion)
@@ -418,16 +421,17 @@ void MemoryTrace_UpdateInternal()
 ////////////////////////////////////////////////////////////////////////////////
 void MemoryTrace_Alloc(uint64 Address, uint64 Size, uint32 Alignment, HeapId RootHeap)
 {
-	check(RootHeap < 16);
 	if (!GTraceAllowed)
 	{
 		return;
 	}
-		
+
+	check(RootHeap < 16);
+
 	const uint32 AlignmentPow2 = uint32(FPlatformMath::CountTrailingZeros(Alignment));
 	const uint32 Alignment_SizeLower = (AlignmentPow2 << SizeShift) | uint32(Size & ((1 << SizeShift) - 1));
 	const uint32 CallstackId = GDoNotAllocateInTrace ? 0 : CallstackTrace_GetCurrentId();
-	
+
 	switch (RootHeap)
 	{
 		case EMemoryTraceRootHeap::SystemMemory:
@@ -468,30 +472,36 @@ void MemoryTrace_Alloc(uint64 Address, uint64 Size, uint32 Alignment, HeapId Roo
 ////////////////////////////////////////////////////////////////////////////////
 void MemoryTrace_Free(uint64 Address, HeapId RootHeap)
 {
-	check(RootHeap < 16);
 	if (!GTraceAllowed)
 	{
 		return;
 	}
-	
+
+	check(RootHeap < 16);
+
+	const uint32 CallstackId = GDoNotAllocateInTrace ? 0 : CallstackTrace_GetCurrentId();
+
 	switch (RootHeap)
 	{
 		case EMemoryTraceRootHeap::SystemMemory:
 			{
 				UE_TRACE_LOG(Memory, FreeSystem, MemAllocChannel)
-					<< FreeSystem.Address(uint64(Address));
+					<< FreeSystem.Address(uint64(Address))
+					<< FreeSystem.CallstackId(CallstackId);
 				break;
 			}
 		case EMemoryTraceRootHeap::VideoMemory:
 			{
 				UE_TRACE_LOG(Memory, FreeVideo, MemAllocChannel)
-					<< FreeVideo.Address(uint64(Address));
+					<< FreeVideo.Address(uint64(Address))
+					<< FreeVideo.CallstackId(CallstackId);
 				break;
 			}
 		default:
 			{
 				UE_TRACE_LOG(Memory, Free, MemAllocChannel)
 					<< Free.Address(uint64(Address))
+					<< Free.CallstackId(CallstackId)
 					<< Free.RootHeap(uint8(RootHeap));
 				break;
 			}
@@ -503,12 +513,13 @@ void MemoryTrace_Free(uint64 Address, HeapId RootHeap)
 ////////////////////////////////////////////////////////////////////////////////
 void MemoryTrace_ReallocAlloc(uint64 Address, uint64 Size, uint32 Alignment, HeapId RootHeap)
 {
-	check(RootHeap < 16);
 	if (!GTraceAllowed)
 	{
 		return;
 	}
-	
+
+	check(RootHeap < 16);
+
 	const uint32 AlignmentPow2 = uint32(FPlatformMath::CountTrailingZeros(Alignment));
 	const uint32 Alignment_SizeLower = (AlignmentPow2 << SizeShift) | uint32(Size & ((1 << SizeShift) - 1));
 	const uint32 CallstackId = GDoNotAllocateInTrace ? 0 : CallstackTrace_GetCurrentId();
@@ -543,12 +554,13 @@ void MemoryTrace_ReallocAlloc(uint64 Address, uint64 Size, uint32 Alignment, Hea
 ////////////////////////////////////////////////////////////////////////////////
 void MemoryTrace_ReallocFree(uint64 Address, HeapId RootHeap)
 {
-	check(RootHeap < 16);
 	if (!GTraceAllowed)
 	{
 		return;
 	}
 
+	check(RootHeap < 16);
+
 	switch (RootHeap)
 	{
 		case EMemoryTraceRootHeap::SystemMemory:
@@ -571,13 +583,13 @@ void MemoryTrace_ReallocFree(uint64 Address, HeapId RootHeap)
 }
 
 ////////////////////////////////////////////////////////////////////////////////
-HeapId MemoryTrace_HeapSpec(HeapId ParentId, const TCHAR* Name, EMemoryTraceHeapFlags Flags) 
+HeapId MemoryTrace_HeapSpec(HeapId ParentId, const TCHAR* Name, EMemoryTraceHeapFlags Flags)
 {
 	if (!GTraceAllowed)
 	{
 		return 0;
 	}
-	
+
 	static std::atomic<HeapId> HeapIdCount(EMemoryTraceRootHeap::EndReserved + 1); //Reserve indexes for root heaps
 	const HeapId Id = HeapIdCount.fetch_add(1);
 	const uint32 NameLen = FCString::Strlen(Name);
@@ -600,7 +612,7 @@ HeapId MemoryTrace_RootHeapSpec(const TCHAR* Name, EMemoryTraceHeapFlags Flags)
 	{
 		return 0;
 	}
-	
+
 	static std::atomic<HeapId> RootHeapCount(0);
 	const HeapId Id = RootHeapCount.fetch_add(1);
 	check(Id <= EMemoryTraceRootHeap::EndReserved);
@@ -624,7 +636,7 @@ void MemoryTrace_MarkAllocAsHeap(uint64 Address, HeapId Heap, EMemoryTraceHeapAl
 	{
 		return;
 	}
-	
+
 	UE_TRACE_LOG(Memory, HeapMarkAlloc, MemAllocChannel)
 		<< HeapMarkAlloc.Address(uint64(Address))
 		<< HeapMarkAlloc.Heap(Heap)
@@ -638,7 +650,7 @@ void MemoryTrace_UnmarkAllocAsHeap(uint64 Address, HeapId Heap)
 	{
 		return;
 	}
-	
+
 	// Sets all flags to zero
 	UE_TRACE_LOG(Memory, HeapUnmarkAlloc, MemAllocChannel)
 		<< HeapUnmarkAlloc.Address(uint64(Address))
@@ -683,14 +695,14 @@ void* FTraceMalloc::Malloc(SIZE_T Count, uint32 Alignment)
 
 	UE_TRACE_METADATA_CLEAR_SCOPE();
 	UE_MEMSCOPE(TRACE_TAG);
-	
+
 	UE_TRACE_LOG(Memory, Alloc, MemAllocChannel)
 		<< Alloc.CallstackId(0)
 		<< Alloc.Address(uint64(NewPtr))
 		<< Alloc.RootHeap(uint8(EMemoryTraceRootHeap::SystemMemory))
 		<< Alloc.Size(uint32(Size >> SizeShift))
 		<< Alloc.AlignmentPow2_SizeLower(uint8(Alignment_SizeLower));
-	
+
 	return NewPtr;
 #else
 	return WrappedMalloc->Malloc(Count, Alignment);
@@ -708,23 +720,23 @@ void* FTraceMalloc::Realloc(void* Original, SIZE_T Count, uint32 Alignment)
 
 	UE_TRACE_METADATA_CLEAR_SCOPE();
 	UE_MEMSCOPE(TRACE_TAG);
-	
+
 	UE_TRACE_LOG(Memory, ReallocFree, MemAllocChannel)
 		<< ReallocFree.Address(uint64(Original))
 		<< ReallocFree.RootHeap(uint8(EMemoryTraceRootHeap::SystemMemory));
-	
+
 	{
 		TGuardValue<bool> _(GDoNotAllocateInTrace, true);
 		NewPtr = WrappedMalloc->Realloc(Original, Count, Alignment);
 	}
-	
+
 	UE_TRACE_LOG(Memory, ReallocAlloc, MemAllocChannel)
 		<< ReallocAlloc.CallstackId(0)
 		<< ReallocAlloc.Address(uint64(NewPtr))
 		<< ReallocAlloc.RootHeap(uint8(EMemoryTraceRootHeap::SystemMemory))
 		<< ReallocAlloc.Size(uint32(Size >> SizeShift))
 		<< ReallocAlloc.AlignmentPow2_SizeLower(uint8(Alignment_SizeLower));
-	
+
 	return NewPtr;
 #else
 	return WrappedMalloc->Realloc(Original, Count, Alignment);
@@ -738,14 +750,14 @@ void FTraceMalloc::Free(void* Original)
 	UE_TRACE_LOG(Memory, Free, MemAllocChannel)
 		<< Free.Address(uint64(Original))
 		<< Free.RootHeap(uint8(EMemoryTraceRootHeap::SystemMemory));
-	
+
 	{
 		TGuardValue<bool> _(GDoNotAllocateInTrace, true);
 		WrappedMalloc->Free(Original);
 	}
 #else
 	WrappedMalloc->Free(Original);
-#endif
+#endif //UE_MEMORY_TRACE_ENABLED
 }
 
 /////////////////////////////////////////////////////////////////////////////
@@ -755,5 +767,5 @@ bool FTraceMalloc::ShouldTrace()
 	return !GDoNotAllocateInTrace;
 #else
 	return true;
-#endif
+#endif //UE_MEMORY_TRACE_ENABLED
 }
