Update modular feature not found assert to log name of the modular feature not found
#rnx

diff --git a/Engine/Source/Runtime/Core/Private/Features/ModularFeatures.cpp b/Engine/Source/Runtime/Core/Private/Features/ModularFeatures.cpp
--- a/Engine/Source/Runtime/Core/Private/Features/ModularFeatures.cpp
+++ b/Engine/Source/Runtime/Core/Private/Features/ModularFeatures.cpp
@@ -48,7 +48,7 @@ IModularFeature* FModularFeatures::GetModularFeatureImplementation( const FName
 		++CurrentIndex;
 	}
 
-	check( ModularFeature != nullptr );
+	checkf( ModularFeature != nullptr, TEXT("ModularFeature not found: '%s', Index: %d"), *Type.ToString(), Index);
 	return ModularFeature;
 }
 
