Security guards for TArray.

diff --git a/Engine/Source/Runtime/Core/Public/Containers/Array.h b/Engine/Source/Runtime/Core/Public/Containers/Array.h
--- a/Engine/Source/Runtime/Core/Public/Containers/Array.h
+++ b/Engine/Source/Runtime/Core/Public/Containers/Array.h
@@ -26,6 +26,7 @@
 #include "Templates/Sorting.h"
 #include "Templates/AlignmentTemplates.h"
 #include "Templates/IsConstructible.h"
+#include "Templates/MakeUnsigned.h"
 #include "Traits/ElementType.h"
 
 #include <type_traits>
@@ -354,6 +355,15 @@ public:
 	typedef InElementType ElementType;
 	typedef InAllocatorType AllocatorType;
 
+private:
+	using USizeType = typename TMakeUnsigned<SizeType>::Type;
+
+	FORCENOINLINE static void OnInvalidNum(USizeType NewNum)
+	{
+		FPlatformMemory::OnOutOfMemory(NewNum * sizeof(ElementType), alignof(ElementType));
+	}
+
+public:
 	UE_DEPRECATED(5.0, "TArray::Allocator type is deprecated, please use TArray::AllocatorType instead.")
 	typedef InAllocatorType Allocator;
 
@@ -382,9 +392,14 @@ public:
 	 */
 	FORCEINLINE TArray(const ElementType* Ptr, SizeType Count)
 	{
+		if (Count < 0)
+		{
+			OnInvalidNum((USizeType)Count);
+		}
+
 		check(Ptr != nullptr || Count == 0);
 
-		CopyToEmpty(Ptr, Count, 0, 0);
+		CopyToEmpty(Ptr, Count, 0);
 	}
 
 	template <typename OtherElementType, typename OtherSizeType>
@@ -398,7 +413,7 @@ public:
 		// This is not strictly legal, as std::initializer_list's iterators are not guaranteed to be pointers, but
 		// this appears to be the case on all of our implementations.  Also, if it's not true on a new implementation,
 		// it will fail to compile rather than behave badly.
-		CopyToEmpty(InitList.begin(), (SizeType)InitList.size(), 0, 0);
+		CopyToEmpty(InitList.begin(), (SizeType)InitList.size(), 0);
 	}
 
 	/**
@@ -413,7 +428,7 @@ public:
 	>
 	FORCEINLINE explicit TArray(const TArray<OtherElementType, OtherAllocator>& Other)
 	{
-		CopyToEmpty(Other.GetData(), Other.Num(), 0, 0);
+		CopyToEmpty(Other.GetData(), Other.Num(), 0);
 	}
 
 	/**
@@ -423,7 +438,7 @@ public:
 	 */
 	FORCEINLINE TArray(const TArray& Other)
 	{
-		CopyToEmpty(Other.GetData(), Other.Num(), 0, 0);
+		CopyToEmpty(Other.GetData(), Other.Num(), 0);
 	}
 
 	/**
@@ -435,7 +450,7 @@ public:
 	 */
 	FORCEINLINE TArray(const TArray& Other, SizeType ExtraSlack)
 	{
-		CopyToEmpty(Other.GetData(), Other.Num(), 0, ExtraSlack);
+		CopyToEmptyWithSlack(Other.GetData(), Other.Num(), 0, ExtraSlack);
 	}
 
 	/**
@@ -450,7 +465,7 @@ public:
 		// This is not strictly legal, as std::initializer_list's iterators are not guaranteed to be pointers, but
 		// this appears to be the case on all of our implementations.  Also, if it's not true on a new implementation,
 		// it will fail to compile rather than behave badly.
-		CopyToEmpty(InitList.begin(), (SizeType)InitList.size(), ArrayMax, 0);
+		CopyToEmpty(InitList.begin(), (SizeType)InitList.size(), ArrayMax);
 		return *this;
 	}
 
@@ -466,7 +481,7 @@ public:
 	TArray& operator=(const TArray<ElementType, OtherAllocatorType>& Other)
 	{
 		DestructItems(GetData(), ArrayNum);
-		CopyToEmpty(Other.GetData(), Other.Num(), ArrayMax, 0);
+		CopyToEmpty(Other.GetData(), Other.Num(), ArrayMax);
 		return *this;
 	}
 
@@ -481,7 +496,7 @@ public:
 		if (this != &Other)
 		{
 			DestructItems(GetData(), ArrayNum);
-			CopyToEmpty(Other.GetData(), Other.Num(), ArrayMax, 0);
+			CopyToEmpty(Other.GetData(), Other.Num(), ArrayMax);
 		}
 		return *this;
 	}
@@ -501,6 +516,8 @@ private:
 	template <typename FromArrayType, typename ToArrayType>
 	static FORCEINLINE std::enable_if_t<UE4Array_Private::TCanMoveTArrayPointersBetweenArrayTypes<FromArrayType, ToArrayType>::Value> MoveOrCopy(ToArrayType& ToArray, FromArrayType& FromArray, SizeType PrevMax)
 	{
+		static_assert(std::is_same_v<TArray, ToArrayType>, "MoveOrCopy is expected to be called with the current array type as the destination");
+
 		using FromAllocatorType = typename FromArrayType::AllocatorType;
 		using ToAllocatorType   = typename ToArrayType::AllocatorType;
 
@@ -517,7 +534,13 @@ private:
 		ToArray  .ArrayMax = (SizeType)FromArray.ArrayMax;
 
 		// Ensure the destination container could hold the source range (when the allocator size types shrink)
-		checkf(ToArray.ArrayNum == FromArray.ArrayNum && ToArray.ArrayMax == FromArray.ArrayMax, TEXT("Data lost when moving to a container with a more constrained size type"));
+		if constexpr (sizeof(USizeType) < sizeof(typename FromArrayType::USizeType))
+		{
+			if (ToArray.ArrayNum != FromArray.ArrayNum || ToArray.ArrayMax != FromArray.ArrayMax)
+			{
+				OnInvalidNum((USizeType)ToArray.ArrayNum);
+			}
+		}
 
 		FromArray.ArrayNum = 0;
 		FromArray.ArrayMax = FromArray.AllocatorInstance.GetInitialCapacity();
@@ -536,7 +559,7 @@ private:
 	template <typename FromArrayType, typename ToArrayType>
 	static FORCEINLINE std::enable_if_t<!UE4Array_Private::TCanMoveTArrayPointersBetweenArrayTypes<FromArrayType, ToArrayType>::Value> MoveOrCopy(ToArrayType& ToArray, FromArrayType& FromArray, SizeType PrevMax)
 	{
-		ToArray.CopyToEmpty(FromArray.GetData(), FromArray.Num(), PrevMax, 0);
+		ToArray.CopyToEmpty(FromArray.GetData(), FromArray.Num(), PrevMax);
 	}
 
 	/**
@@ -554,7 +577,16 @@ private:
 	{
 		MoveOrCopy(ToArray, FromArray, PrevMax);
 
-		ToArray.Reserve(ToArray.ArrayNum + ExtraSlack);
+		USizeType LocalArrayNum = (USizeType)ToArray.ArrayNum;
+		USizeType NewMax        = (USizeType)LocalArrayNum + (USizeType)ExtraSlack;
+
+		// This should only happen when we've underflowed or overflowed SizeType
+		if ((SizeType)NewMax < LocalArrayNum)
+		{
+			OnInvalidNum((USizeType)ExtraSlack);
+		}
+
+		ToArray.Reserve(NewMax);
 	}
 
 	/**
@@ -570,7 +602,7 @@ private:
 	template <typename FromArrayType, typename ToArrayType>
 	static FORCEINLINE std::enable_if_t<!UE4Array_Private::TCanMoveTArrayPointersBetweenArrayTypes<FromArrayType, ToArrayType>::Value> MoveOrCopyWithSlack(ToArrayType& ToArray, FromArrayType& FromArray, SizeType PrevMax, SizeType ExtraSlack)
 	{
-		ToArray.CopyToEmpty(FromArray.GetData(), FromArray.Num(), PrevMax, ExtraSlack);
+		ToArray.CopyToEmptyWithSlack(FromArray.GetData(), FromArray.Num(), PrevMax, ExtraSlack);
 	}
 
 public:
@@ -612,10 +644,6 @@ public:
 	>
 	TArray(TArray<OtherElementType, AllocatorType>&& Other, SizeType ExtraSlack)
 	{
-		// We don't implement move semantics for general OtherAllocators, as there's no way
-		// to tell if they're compatible with the current one.  Probably going to be a pretty
-		// rare requirement anyway.
-
 		MoveOrCopyWithSlack(*this, Other, 0, ExtraSlack);
 	}
 
@@ -1357,33 +1385,80 @@ public:
 	 * @param Count Number of elements to add.
 	 * @returns Number of elements in array before addition.
 	 */
-	FORCEINLINE SizeType AddUninitialized(SizeType Count = 1)
+	FORCEINLINE SizeType AddUninitialized()
+	{
+		CheckInvariants();
+
+		const USizeType OldNum = (USizeType)ArrayNum;
+		const USizeType NewNum = OldNum + (USizeType)1;
+		ArrayNum = (SizeType)NewNum;
+		if (NewNum > (USizeType)ArrayMax)
+		{
+			ResizeGrow((SizeType)OldNum);
+		}
+		return OldNum;
+	}
+	FORCEINLINE SizeType AddUninitialized(SizeType Count)
 	{
 		CheckInvariants();
 		checkSlow(Count >= 0);
 
-		const SizeType OldNum = ArrayNum;
-		if ((ArrayNum += Count) > ArrayMax)
+		const USizeType OldNum = (USizeType)ArrayNum;
+		const USizeType NewNum = OldNum + (USizeType)Count;
+		ArrayNum = (SizeType)NewNum;
+
+#if DO_GUARD_SLOW
+		if (NewNum > (USizeType)ArrayMax)
+#else
+		// SECURITY - This check will guard against negative counts too, in case the checkSlow(Count >= 0) above is compiled out.
+		// However, it results in slightly worse code generation.
+		if ((USizeType)Count > (USizeType)ArrayMax - OldNum)
+#endif
 		{
-			ResizeGrow(OldNum);
+			ResizeGrow((SizeType)OldNum);
 		}
 		return OldNum;
 	}
 
 private:
+	void InsertUninitializedImpl(SizeType Index)
+	{
+		CheckInvariants();
+		checkSlow((Index >= 0) & (Index <= ArrayNum));
+
+		const USizeType OldNum = (USizeType)ArrayNum;
+		const USizeType NewNum = OldNum + (USizeType)1;
+		ArrayNum = (SizeType)NewNum;
+
+		if (NewNum > (USizeType)ArrayMax)
+		{
+			ResizeGrow((SizeType)OldNum);
+		}
+		ElementType* Data = GetData() + Index;
+		RelocateConstructItems<ElementType>(Data + 1, Data, OldNum - Index);
+	}
 	template <typename OtherSizeType>
 	void InsertUninitializedImpl(SizeType Index, OtherSizeType Count)
 	{
 		CheckInvariants();
 		checkSlow((Count >= 0) & (Index >= 0) & (Index <= ArrayNum));
 
-		SizeType NewNum = Count;
-		checkf((OtherSizeType)NewNum == Count, TEXT("Invalid number of elements to add to this array type: %lld"), (long long)NewNum);
+		SizeType ConvertedCount = Count;
+		checkf((OtherSizeType)ConvertedCount == Count, TEXT("Invalid number of elements to add to this array type: %lld"), (long long)ConvertedCount);
+
+		const USizeType OldNum = (USizeType)ArrayNum;
+		const USizeType NewNum = OldNum + (USizeType)Count;
+		ArrayNum = (SizeType)NewNum;
 
-		const SizeType OldNum = ArrayNum;
-		if ((ArrayNum += Count) > ArrayMax)
+#if DO_GUARD_SLOW
+		if (NewNum > (USizeType)ArrayMax)
+#else
+		// SECURITY - This check will guard against negative counts too, in case the checkSlow(Count >= 0) above is compiled out.
+		// However, it results in slightly worse code generation.
+		if ((USizeType)Count > (USizeType)ArrayMax - OldNum)
+#endif
 		{
-			ResizeGrow(OldNum);
+			ResizeGrow((SizeType)OldNum);
 		}
 		ElementType* Data = GetData() + Index;
 		RelocateConstructItems<ElementType>(Data + Count, Data, OldNum - Index);
@@ -1402,7 +1477,11 @@ public:
 	 * @param Count Number of elements to add.
 	 * @see Insert, InsertZeroed, InsertDefaulted
 	 */
-	FORCEINLINE void InsertUninitialized(SizeType Index, SizeType Count = 1)
+	FORCEINLINE void InsertUninitialized(SizeType Index)
+	{
+		InsertUninitializedImpl(Index);
+	}
+	FORCEINLINE void InsertUninitialized(SizeType Index, SizeType Count)
 	{
 		InsertUninitializedImpl(Index, Count);
 	}
@@ -1419,7 +1498,12 @@ public:
 	 * @param Count Number of elements to add.
 	 * @see Insert, InsertUninitialized, InsertDefaulted
 	 */
-	void InsertZeroed(SizeType Index, SizeType Count = 1)
+	void InsertZeroed(SizeType Index)
+	{
+		InsertUninitializedImpl(Index);
+		FMemory::Memzero(GetData() + Index, sizeof(ElementType));
+	}
+	void InsertZeroed(SizeType Index, SizeType Count)
 	{
 		InsertUninitializedImpl(Index, Count);
 		FMemory::Memzero(GetData() + Index, Count * sizeof(ElementType));
@@ -1452,7 +1536,12 @@ public:
 	 * @param Count Number of elements to add.
 	 * @see Insert, InsertUninitialized, InsertZeroed
 	 */
-	void InsertDefaulted(SizeType Index, SizeType Count = 1)
+	void InsertDefaulted(SizeType Index)
+	{
+		InsertUninitializedImpl(Index);
+		DefaultConstructItems<ElementType>(GetData() + Index, 1);
+	}
+	void InsertDefaulted(SizeType Index, SizeType Count)
 	{
 		InsertUninitializedImpl(Index, Count);
 		DefaultConstructItems<ElementType>(GetData() + Index, Count);
@@ -1579,7 +1668,7 @@ public:
 
 		// construct a copy in place at Index (this new operator will insert at 
 		// Index, then construct that memory with Item)
-		InsertUninitializedImpl(Index, 1);
+		InsertUninitializedImpl(Index);
 		new(GetData() + Index) ElementType(MoveTempIfPossible(Item));
 		return Index;
 	}
@@ -1598,7 +1687,7 @@ public:
 
 		// construct a copy in place at Index (this new operator will insert at 
 		// Index, then construct that memory with Item)
-		InsertUninitializedImpl(Index, 1);
+		InsertUninitializedImpl(Index);
 		new(GetData() + Index) ElementType(Item);
 		return Index;
 	}
@@ -1618,7 +1707,7 @@ public:
 
 		// construct a copy in place at Index (this new operator will insert at 
 		// Index, then construct that memory with Item)
-		InsertUninitializedImpl(Index, 1);
+		InsertUninitializedImpl(Index);
 		ElementType* Ptr = GetData() + Index;
 		new(Ptr) ElementType(MoveTempIfPossible(Item));
 		return *Ptr;
@@ -1638,7 +1727,7 @@ public:
 
 		// construct a copy in place at Index (this new operator will insert at 
 		// Index, then construct that memory with Item)
-		InsertUninitializedImpl(Index, 1);
+		InsertUninitializedImpl(Index);
 		ElementType* Ptr = GetData() + Index;
 		new(Ptr) ElementType(Item);
 		return *Ptr;
@@ -1774,6 +1863,11 @@ public:
 	 */
 	void Reset(SizeType NewSize = 0)
 	{
+		if (NewSize < 0)
+		{
+			OnInvalidNum((USizeType)NewSize);
+		}
+
 		// If we have space to hold the excepted size, then don't reallocate
 		if (NewSize <= ArrayMax)
 		{
@@ -1793,6 +1887,11 @@ public:
 	 */
 	void Empty(SizeType Slack = 0)
 	{
+		if (Slack < 0)
+		{
+			OnInvalidNum((USizeType)Slack);
+		}
+
 		DestructItems(GetData(), ArrayNum);
 
 		checkSlow(Slack >= 0);
@@ -1818,6 +1917,10 @@ public:
 			const SizeType Index = AddUninitialized(Diff);
 			DefaultConstructItems<ElementType>((uint8*)AllocatorInstance.GetAllocation() + Index * sizeof(ElementType), Diff);
 		}
+		else if (NewNum < 0)
+		{
+			OnInvalidNum((USizeType)NewNum);
+		}
 		else if (NewNum < Num())
 		{
 			RemoveAt(NewNum, Num() - NewNum, bAllowShrinking);
@@ -1837,6 +1940,10 @@ public:
 		{
 			AddZeroed(NewNum - Num());
 		}
+		else if (NewNum < 0)
+		{
+			OnInvalidNum((USizeType)NewNum);
+		}
 		else if (NewNum < Num())
 		{
 			RemoveAt(NewNum, Num() - NewNum, bAllowShrinking);
@@ -1854,6 +1961,10 @@ public:
 		{
 			AddUninitialized(NewNum - Num());
 		}
+		else if (NewNum < 0)
+		{
+			OnInvalidNum((USizeType)NewNum);
+		}
 		else if (NewNum < Num())
 		{
 			RemoveAt(NewNum, Num() - NewNum, bAllowShrinking);
@@ -2038,7 +2149,7 @@ public:
 	template <typename... ArgsType>
 	FORCEINLINE SizeType Emplace(ArgsType&&... Args)
 	{
-		const SizeType Index = AddUninitialized(1);
+		const SizeType Index = AddUninitialized();
 
 		// If this fails to compile when trying to call Emplace with a non-public constructor,
 		// do not make TArray a friend.
@@ -2075,7 +2186,7 @@ public:
 	template <typename... ArgsType>
 	FORCEINLINE ElementType& Emplace_GetRef(ArgsType&&... Args)
 	{
-		const SizeType Index = AddUninitialized(1);
+		const SizeType Index = AddUninitialized();
 		ElementType* Ptr = GetData() + Index;
 		new(Ptr) ElementType(Forward<ArgsType>(Args)...);
 		return *Ptr;
@@ -2178,7 +2289,13 @@ public:
 	 * @return Index to the first of the new items.
 	 * @see Add, AddDefaulted, AddUnique, Append, Insert
 	 */
-	SizeType AddZeroed(SizeType Count = 1)
+	SizeType AddZeroed()
+	{
+		const SizeType Index = AddUninitialized();
+		FMemory::Memzero((uint8*)AllocatorInstance.GetAllocation() + Index * sizeof(ElementType), sizeof(ElementType));
+		return Index;
+	}
+	SizeType AddZeroed(SizeType Count)
 	{
 		const SizeType Index = AddUninitialized(Count);
 		FMemory::Memzero((uint8*)AllocatorInstance.GetAllocation() + Index*sizeof(ElementType), Count*sizeof(ElementType));
@@ -2198,7 +2315,7 @@ public:
 	 */
 	ElementType& AddZeroed_GetRef()
 	{
-		const SizeType Index = AddUninitialized(1);
+		const SizeType Index = AddUninitialized();
 		ElementType* Ptr = GetData() + Index;
 		FMemory::Memzero(Ptr, sizeof(ElementType));
 		return *Ptr;
@@ -2212,7 +2329,13 @@ public:
 	 * @return Index to the first of the new items.
 	 * @see Add, AddZeroed, AddUnique, Append, Insert
 	 */
-	SizeType AddDefaulted(SizeType Count = 1)
+	SizeType AddDefaulted()
+	{
+		const SizeType Index = AddUninitialized();
+		DefaultConstructItems<ElementType>((uint8*)AllocatorInstance.GetAllocation() + Index * sizeof(ElementType), 1);
+		return Index;
+	}
+	SizeType AddDefaulted(SizeType Count)
 	{
 		const SizeType Index = AddUninitialized(Count);
 		DefaultConstructItems<ElementType>((uint8*)AllocatorInstance.GetAllocation() + Index * sizeof(ElementType), Count);
@@ -2228,7 +2351,7 @@ public:
 	 */
 	ElementType& AddDefaulted_GetRef()
 	{
-		const SizeType Index = AddUninitialized(1);
+		const SizeType Index = AddUninitialized();
 		ElementType* Ptr = GetData() + Index;
 		DefaultConstructItems<ElementType>(Ptr, 1);
 		return *Ptr;
@@ -2292,7 +2415,7 @@ public:
 	{
 		TContainerElementTypeCompatibility<ElementType>::CopyingFromOtherType();
 		DestructItems(GetData(), ArrayNum);
-		CopyToEmpty(Other.GetData(), Other.Num(), ArrayMax, 0);
+		CopyToEmpty(Other.GetData(), Other.Num(), ArrayMax);
 		return *this;
 	}
 
@@ -2400,7 +2523,10 @@ public:
 	 * @returns Index of the element in the array.
 	 * @see Add, AddDefaulted, AddZeroed, Append, Insert
 	 */
-	FORCEINLINE SizeType AddUnique(ElementType&& Item) { return AddUniqueImpl(MoveTempIfPossible(Item)); }
+	FORCEINLINE SizeType AddUnique(ElementType&& Item)
+	{
+		return AddUniqueImpl(MoveTempIfPossible(Item));
+	}
 
 	/**
 	 * Adds unique element to array if it doesn't exist.
@@ -2409,7 +2535,10 @@ public:
 	 * @returns Index of the element in the array.
 	 * @see Add, AddDefaulted, AddZeroed, Append, Insert
 	 */
-	FORCEINLINE SizeType AddUnique(const ElementType& Item) { return AddUniqueImpl(Item); }
+	FORCEINLINE SizeType AddUnique(const ElementType& Item)
+	{
+		return AddUniqueImpl(Item);
+	}
 
 	/**
 	 * Reserves memory such that the array can contain at least Number elements.
@@ -2420,7 +2549,11 @@ public:
 	FORCEINLINE void Reserve(SizeType Number)
 	{
 		checkSlow(Number >= 0);
-		if (Number > ArrayMax)
+		if (Number < 0)
+		{
+			OnInvalidNum((USizeType)Number);
+		}
+		else if (Number > ArrayMax)
 		{
 			ResizeTo(Number);
 		}
@@ -2873,7 +3006,14 @@ private:
 
 	FORCENOINLINE void ResizeGrow(SizeType OldNum)
 	{
-		ArrayMax = AllocatorCalculateSlackGrow(ArrayNum, ArrayMax);
+		SizeType LocalArrayNum = ArrayNum;
+
+		// This should only happen when we've underflowed or overflowed SizeType in the caller
+		if (LocalArrayNum < OldNum)
+		{
+			OnInvalidNum((USizeType)LocalArrayNum - (USizeType)OldNum);
+		}
+		ArrayMax = AllocatorCalculateSlackGrow(LocalArrayNum, ArrayMax);
 		AllocatorResizeAllocation(OldNum, ArrayMax);
 	}
 	FORCENOINLINE void ResizeShrink()
@@ -2916,6 +3056,31 @@ private:
 	}
 
 
+	/**
+	 * Copies data from one array into this array. Uses the fast path if the
+	 * data in question does not need a constructor.
+	 *
+	 * @param Source The source array to copy
+	 * @param PrevMax The previous allocated size
+	 */
+	template <typename OtherElementType, typename OtherSizeType>
+	void CopyToEmpty(const OtherElementType* OtherData, OtherSizeType OtherNum, SizeType PrevMax)
+	{
+		SizeType NewNum = (SizeType)OtherNum;
+		checkf((OtherSizeType)NewNum == OtherNum, TEXT("Invalid number of elements to add to this array type: %lld"), (long long)NewNum);
+
+		ArrayNum = NewNum;
+		if (OtherNum || PrevMax)
+		{
+			ResizeForCopy(NewNum, PrevMax);
+			ConstructItems<ElementType>(GetData(), OtherData, OtherNum);
+		}
+		else
+		{
+			ArrayMax = AllocatorInstance.GetInitialCapacity();
+		}
+	}
+
 	/**
 	 * Copies data from one array into this array. Uses the fast path if the
 	 * data in question does not need a constructor.
@@ -2923,19 +3088,25 @@ private:
 	 * @param Source The source array to copy
 	 * @param PrevMax The previous allocated size
 	 * @param ExtraSlack Additional amount of memory to allocate at
-	 *                   the end of the buffer. Counted in elements. Zero by
-	 *                   default.
+	 *                   the end of the buffer. Counted in elements.
 	 */
 	template <typename OtherElementType, typename OtherSizeType>
-	void CopyToEmpty(const OtherElementType* OtherData, OtherSizeType OtherNum, SizeType PrevMax, SizeType ExtraSlack)
+	void CopyToEmptyWithSlack(const OtherElementType* OtherData, OtherSizeType OtherNum, SizeType PrevMax, SizeType ExtraSlack)
 	{
 		SizeType NewNum = (SizeType)OtherNum;
 		checkf((OtherSizeType)NewNum == OtherNum, TEXT("Invalid number of elements to add to this array type: %lld"), (long long)NewNum);
 
-		checkSlow(ExtraSlack >= 0);
 		ArrayNum = NewNum;
 		if (OtherNum || ExtraSlack || PrevMax)
 		{
+			USizeType NewMax = NewNum + ExtraSlack;
+
+			// This should only happen when we've underflowed or overflowed SizeType
+			if ((SizeType)NewMax < NewNum)
+			{
+				OnInvalidNum((USizeType)NewMax);
+			}
+
 			ResizeForCopy(NewNum + ExtraSlack, PrevMax);
 			ConstructItems<ElementType>(GetData(), OtherData, OtherNum);
 		}
@@ -3377,7 +3548,7 @@ template <typename InElementType, typename InAllocatorType> struct TIsTArray<con
 template <typename T,typename AllocatorType> void* operator new( size_t Size, TArray<T,AllocatorType>& Array )
 {
 	check(Size == sizeof(T));
-	const auto Index = Array.AddUninitialized(1);
+	const auto Index = Array.AddUninitialized();
 	return &Array[Index];
 }
 template <typename T,typename AllocatorType> void* operator new( size_t Size, TArray<T,AllocatorType>& Array, typename TArray<T, AllocatorType>::SizeType Index )
diff --git a/Engine/Source/Runtime/Core/Public/Containers/ArrayView.h b/Engine/Source/Runtime/Core/Public/Containers/ArrayView.h
--- a/Engine/Source/Runtime/Core/Public/Containers/ArrayView.h
+++ b/Engine/Source/Runtime/Core/Public/Containers/ArrayView.h
@@ -803,7 +803,7 @@ template<typename InElementType, typename InAllocatorType>
 template<typename OtherElementType, typename OtherSizeType>
 FORCEINLINE TArray<InElementType, InAllocatorType>::TArray(const TArrayView<OtherElementType, OtherSizeType>& Other)
 {
-	CopyToEmpty(Other.GetData(), Other.Num(), 0, 0);
+	CopyToEmpty(Other.GetData(), Other.Num(), 0);
 }
 
 template<typename InElementType, typename InAllocatorType>
@@ -811,6 +811,6 @@ template<typename OtherElementType, typename OtherSizeType>
 FORCEINLINE TArray<InElementType, InAllocatorType>& TArray<InElementType, InAllocatorType>::operator=(const TArrayView<OtherElementType, OtherSizeType>& Other)
 {
 	DestructItems(GetData(), ArrayNum);
-	CopyToEmpty(Other.GetData(), Other.Num(), ArrayMax, 0);
+	CopyToEmpty(Other.GetData(), Other.Num(), ArrayMax);
 	return *this;
 }
