Fix writing obfuscated names to utf8 string builder.

diff --git a/Engine/Source/Runtime/Core/Private/UObject/UnrealNames.cpp b/Engine/Source/Runtime/Core/Private/UObject/UnrealNames.cpp
--- a/Engine/Source/Runtime/Core/Private/UObject/UnrealNames.cpp
+++ b/Engine/Source/Runtime/Core/Private/UObject/UnrealNames.cpp
@@ -2304,11 +2304,15 @@ void FNameEntry::AppendNameToString(FUtf8StringBuilderBase& Out) const
 	checkName(Header.Len != 0);
 	if (Header.bIsWide)
 	{
-		Out << FWideStringView(WideName, Header.Len);
+		FNameBuffer DecodeBuffer;
+		FNameStringView View = MakeView(DecodeBuffer);
+		Out << FWideStringView(View.Wide, View.Len);
 	}
 	else
 	{
-		Out << FAnsiStringView(AnsiName, Header.Len);
+		const int32 Offset = Out.AddUninitialized(Header.Len);
+		UTF8CHAR* OutChars = Out.GetData() + Offset;
+		CopyUnterminatedName(reinterpret_cast<ANSICHAR*>(OutChars));
 	}
 }
 
