Checked integer template.

Adds new FCheckedSignedInt<SignedType> template with named instantiations FCheckedInt32 and FCheckedInt64 that perform range, overflow-checked integer arithmetic and know they're invalid when an error occured anywhere along the way. They also test for illegal shift amounts and division by zero.

This class is intended for processing integers from untrusted sources, especially for size computations, where overflows might otherwise be used to craft buffer overflow exploits.

diff --git a/Engine/Source/Runtime/Core/Private/Tests/Math/UnrealMathTest.cpp b/Engine/Source/Runtime/Core/Private/Tests/Math/UnrealMathTest.cpp
--- a/Engine/Source/Runtime/Core/Private/Tests/Math/UnrealMathTest.cpp
+++ b/Engine/Source/Runtime/Core/Private/Tests/Math/UnrealMathTest.cpp
@@ -17,6 +17,7 @@
 #include "Math/Quat.h"
 #include "Math/QuatRotationTranslationMatrix.h"
 #include "Math/Color.h"
+#include "Math/CheckedInt.h"
 #include "Misc/AutomationTest.h"
 #include "Async/ParallelFor.h"
 #include "Misc/ScopeLock.h"
@@ -4808,4 +4809,137 @@ bool FColorConversionHeavyTest::RunTest(const FString& Parameters)
 	return true;
 }
 
+IMPLEMENT_SIMPLE_AUTOMATION_TEST(FMathCheckedIntTest, "System.Core.Math.CheckedInt", EAutomationTestFlags::ApplicationContextMask | EAutomationTestFlags::EngineFilter)
+bool FMathCheckedIntTest::RunTest(const FString& Parameters)
+{
+	// Do an exhaustive test of all checked signed int operations with int8s. This is a small enough space of possibilities
+	// that it's easy to try all options exhaustively, and it's likewise easy to compute reference results in int32 arithmetic
+	// without having to worry about overflows. TCheckedSignedInt uses the same logic and implementation for all type
+	// instantiations (as of this writing), so this is a good way to get coverage on all relevant code paths.
+	using FCheckedInt8 = TCheckedSignedInt<int8>;
+	int NumErrorsFound = 0;
+	const int NumErrorsMax = 100; // Stop reporting errors after this in case something is completely busted.
+
+	// Helper functions to log test results
+	auto CheckArithResult = [this, &NumErrorsFound, NumErrorsMax](const TCHAR* Operation, int32 A, int32 B, FCheckedInt8 CheckedResult, int32 UncheckedValue)
+	{
+		FCheckedInt8 UncheckedResult{ UncheckedValue };
+
+		// If either operand was invalid, the unchecked result is also invalid
+		if (!FCheckedInt8(A).IsValid() || !FCheckedInt8(B).IsValid())
+		{
+			UncheckedResult = FCheckedInt8::Invalid();
+		}
+
+		if (CheckedResult != UncheckedResult && NumErrorsFound < NumErrorsMax)
+		{
+			NumErrorsFound++;
+
+			const int UncheckedNum = UncheckedResult.Get(-128);
+			const int CheckedNum = CheckedResult.Get(-128);
+
+			AddError(FString::Printf(TEXT("Arith %s A=%d B=%d Expected=(%d,Valid=%d) Got=(%d,Valid=%d)"),
+				Operation, A, B, UncheckedNum, UncheckedResult.IsValid(), CheckedNum, CheckedResult.IsValid()));
+		}
+	};
+
+	auto CheckCompareResult = [this, &NumErrorsFound, NumErrorsMax](const TCHAR* Operation, int32 A, int32 B, bool bAcceptInvalid, bool bTestResult, bool bRefResult)
+	{
+		// If either operands was invalid, override the ref result
+		if (!FCheckedInt8(A).IsValid() || !FCheckedInt8(B).IsValid())
+			bRefResult = bAcceptInvalid;
+
+		if (bTestResult != bRefResult && NumErrorsFound < NumErrorsMax)
+		{
+			NumErrorsFound++;
+
+			AddError(FString::Printf(TEXT("Cmp %s A=%d B=%d ref=%d test=%d\n"), Operation, A, B, bRefResult, bTestResult));
+		}
+	};
+
+	// Note: we include 128 in the ranges here to explicitly try out-of-range (invalid) values.
+	for (int32 A = -128; A <= 128; ++A)
+	{
+		for (int32 B = -128; B <= 128; ++B)
+		{
+			FCheckedInt8 CheckedA{ A };
+			FCheckedInt8 CheckedB{ B };
+
+			// Compute reference results for the operations that are not defined everywhere
+			int32 ADivB = -256;
+			int32 AModB = -256;
+
+			// -128 / -1 is undefined in int8 arithmetic because it overflows; when we're
+			// computing it with int32s, it of course works just fine, but we want to make sure
+			// our test results are correct.
+			if (B != 0 && (A != -128 || B != -1))
+			{
+				ADivB = A / B;
+				AModB = A % B;
+			}
+
+			// Shifts are only defined for in-range shift amounts, which for our checked ints
+			// is [0,NumBits-1]. This is tighter than usual C++ arithmetic for the FCheckedInt8 we
+			// use for testing since normal C++ promotes to int and shifts down (so usually [0,31]
+			// count limit) so we need to test for these explicitly.
+			int32 ALshB = -256;
+			int32 ARshB = -256;
+			if (B >= 0 && B <= 7)
+			{
+				ALshB = A << B;
+				ARshB = A >> B;
+			}
+
+			CheckArithResult(TEXT("add1"), A, B, CheckedA + CheckedB, A + B);
+			CheckArithResult(TEXT("add2"), A, B, CheckedA + B, A + B);
+			CheckArithResult(TEXT("add3"), A, B, A + CheckedB, A + B);
+			CheckArithResult(TEXT("sub1"), A, B, CheckedA - CheckedB, A - B);
+			CheckArithResult(TEXT("sub2"), A, B, CheckedA - B, A - B);
+			CheckArithResult(TEXT("sub3"), A, B, A - CheckedB, A - B);
+			CheckArithResult(TEXT("mul1"), A, B, CheckedA * CheckedB, A * B);
+			CheckArithResult(TEXT("mul2"), A, B, CheckedA * B, A * B);
+			CheckArithResult(TEXT("mul3"), A, B, A * CheckedB, A * B);
+			CheckArithResult(TEXT("div1"), A, B, CheckedA / CheckedB, ADivB);
+			CheckArithResult(TEXT("div2"), A, B, CheckedA / B, ADivB);
+			CheckArithResult(TEXT("div3"), A, B, A / CheckedB, ADivB);
+			CheckArithResult(TEXT("mod1"), A, B, CheckedA % CheckedB, AModB);
+			CheckArithResult(TEXT("mod2"), A, B, CheckedA % B, AModB);
+			CheckArithResult(TEXT("mod3"), A, B, A % CheckedB, AModB);
+			CheckArithResult(TEXT("lsh1"), A, B, CheckedA << CheckedB, ALshB);
+			CheckArithResult(TEXT("lsh2"), A, B, CheckedA << B, ALshB);
+			CheckArithResult(TEXT("lsh3"), A, B, A << CheckedB, ALshB);
+			CheckArithResult(TEXT("rsh1"), A, B, CheckedA >> CheckedB, ARshB);
+			CheckArithResult(TEXT("rsh2"), A, B, CheckedA >> B, ARshB);
+			CheckArithResult(TEXT("rsh3"), A, B, A >> CheckedB, ARshB);
+			CheckArithResult(TEXT("neg"), A, A, -CheckedA, -A);
+			CheckArithResult(TEXT("abs"), A, A, CheckedA.Abs(), FMath::Abs(A));
+			
+			CheckCompareResult(TEXT("valid_lt1"), A, B, false, CheckedA.ValidAndLessThan(CheckedB), A < B);
+			CheckCompareResult(TEXT("valid_lt2"), A, B, false, CheckedA.ValidAndLessThan(B), A < B);
+			CheckCompareResult(TEXT("valid_le1"), A, B, false, CheckedA.ValidAndLessOrEqual(CheckedB), A <= B);
+			CheckCompareResult(TEXT("valid_le1"), A, B, false, CheckedA.ValidAndLessOrEqual(B), A <= B);
+			CheckCompareResult(TEXT("valid_gt1"), A, B, false, CheckedA.ValidAndGreaterThan(CheckedB), A > B);
+			CheckCompareResult(TEXT("valid_gt1"), A, B, false, CheckedA.ValidAndGreaterThan(B), A > B);
+			CheckCompareResult(TEXT("valid_ge1"), A, B, false, CheckedA.ValidAndGreaterOrEqual(CheckedB), A >= B);
+			CheckCompareResult(TEXT("valid_ge2"), A, B, false, CheckedA.ValidAndGreaterOrEqual(B), A >= B);
+			
+			CheckCompareResult(TEXT("invalid_lt1"), A, B, true, CheckedA.InvalidOrLessThan(CheckedB), A < B);
+			CheckCompareResult(TEXT("invalid_lt2"), A, B, true, CheckedA.InvalidOrLessThan(B), A < B);
+			CheckCompareResult(TEXT("invalid_le1"), A, B, true, CheckedA.InvalidOrLessOrEqual(CheckedB), A <= B);
+			CheckCompareResult(TEXT("invalid_le2"), A, B, true, CheckedA.InvalidOrLessOrEqual(B), A <= B);
+			CheckCompareResult(TEXT("invalid_gt1"), A, B, true, CheckedA.InvalidOrGreaterThan(CheckedB), A > B);
+			CheckCompareResult(TEXT("invalid_gt2"), A, B, true, CheckedA.InvalidOrGreaterThan(B), A > B);
+			CheckCompareResult(TEXT("invalid_ge1"), A, B, true, CheckedA.InvalidOrGreaterOrEqual(CheckedB), A >= B);
+			CheckCompareResult(TEXT("invalid_ge2"), A, B, true, CheckedA.InvalidOrGreaterOrEqual(B), A >= B);
+		}
+	}
+
+	if (NumErrorsFound > NumErrorsMax)
+	{
+		AddError(FString::Printf(TEXT("Too many errors (%d), stopped reporting after %d instances.\n"), NumErrorsFound, NumErrorsMax));
+	}
+
+	return NumErrorsFound == 0;
+}
+
 #endif //WITH_DEV_AUTOMATION_TESTS
diff --git a/Engine/Source/Runtime/Core/Public/Math/CheckedInt.h b/Engine/Source/Runtime/Core/Public/Math/CheckedInt.h
new file mode 100644
--- /dev/null
+++ b/Engine/Source/Runtime/Core/Public/Math/CheckedInt.h
@@ -0,0 +1,425 @@
+// Copyright Epic Games, Inc. All Rights Reserved.
+
+#pragma once
+
+#include "CoreTypes.h"
+#include "NumericLimits.h"
+#include <type_traits>
+
+/**
+ * Overflow- and error-checked integer. For integer arithmetic on data from untrusted sources (like imported files),
+ * especially when doing size computations. Also checks for division by zero and invalid shift amounts.
+ *
+ * You're not meant to use this directly. Use FCheckedInt32 or FCheckedInt64 (defined below).
+ *
+ * This is a template meant to be instantiated on top of regular basic integer types. The code is written
+ * so the logic is integer-size agnostic and uses just regular C++ arithmetic operations. It is assumed
+ * to run on a two's complement integer platform (which is all we support, and as of C++20 is contractual).
+ * You should generally use the specializations FCheckedInt32 and FCheckedInt64 below.
+ *
+ * Checked integers keep both the integer value and a "valid" flag. Default-constructed checked ints
+ * are invalid, and checked integers constructed from an integer value are valid and hold that value.
+ * Checked integers are somewhat analogous to a TOptional<SignedType> in semantics, and borrow some of
+ * the function names.
+ *
+ * The main feature of checked integers is that all arithmetic on them is overflow-checked. Any arithmetic
+ * involving checked integers results in a checked integer, and any arithmetic involving invalid values,
+ * or arithmetic resulting in overflows or other errors (such as division by zero) likewise results in
+ * an invalid value. The idea is that integer arithmetic using checked integers should be possible to
+ * write very straightforwardly and without having to consider any of these special cases; if any error
+ * occurred along the way, the result will be invalid. These invalid values can then be checked for and
+ * handled right when the result is converted back to a regular integer.
+ *
+ * Some compilers provide built-ins for overflow-checked integer arithmetic for some types. We could
+ * eventually use this (it's especially interesting for multiplications, since our current overflow-checking
+ * algorithm is fairly expensive), but a big benefit of the current approach is that it uses nothing but
+ * regular arithmetic and is type-agnostic. In particular, this makes it possible to check this implementation
+ * exhaustively against a known-good reference for small integer types such as int8. It is much trickier and
+ * more subtle to do good testing for larger integer types where that brute-force approach is not practical.
+ * As-is, the current approach is not the fastest, but it's not presently intended to be used in contexts
+ * where speed of arithmetic operations is a major concern.
+ */
+template<typename SignedType>
+class TCheckedSignedInt
+{
+private:
+	static_assert(std::is_integral<SignedType>::value && std::is_signed<SignedType>::value, "Only defined for signed ints");
+	typedef typename std::make_unsigned<SignedType>::type UnsignedType;
+
+	static const SignedType MinValue = TNumericLimits<SignedType>::Min();
+	static const SignedType MaxValue = TNumericLimits<SignedType>::Max();
+	static const SignedType NumBits = SignedType((sizeof(SignedType) / sizeof(int8_t)) * 8); // Using sizeof to guess the bit count, ugh.
+	static const UnsignedType UnsignedMSB = (UnsignedType)MinValue; // Assuming two's complement
+
+	SignedType Value = 0;
+	bool bIsValid = false;
+
+public:
+	/** Construct a TCheckedSignedInt with an invalid value. */
+	TCheckedSignedInt() = default;
+
+	/** Construct a TCheckedSignedInt from a regular signed integer value. If it's out of range, it results in an invalid value. */
+	template<typename ValueType>
+	explicit TCheckedSignedInt(ValueType InValue, typename std::enable_if<std::is_integral<ValueType>::value && std::is_signed<ValueType>::value>::type* = 0)
+	{
+		if (InValue >= MinValue && InValue <= MaxValue)
+		{
+			Value = (SignedType)InValue;
+			bIsValid = true;
+		}
+		else
+		{
+			Value = 0;
+			bIsValid = false;
+		}
+	}
+
+	/** Construct a TCheckedSignedInt from a regular unsigned integer value. If it's out of range, it results in an invalid value. */
+	template<typename ValueType>
+	explicit TCheckedSignedInt(ValueType InValue, typename std::enable_if<std::is_integral<ValueType>::value && std::is_unsigned<ValueType>::value>::type* = 0)
+	{
+		if (InValue <= UnsignedType(MaxValue))
+		{
+			Value = (SignedType)InValue;
+			bIsValid = true;
+		}
+		else
+		{
+			Value = 0;
+			bIsValid = false;
+		}
+	}
+
+	/** Copy-construct a TCheckedSignedInt from another of matching type. */
+	TCheckedSignedInt(const TCheckedSignedInt& Other) = default;
+
+	/** Assign a TCheckedSignedInt to another. */
+	TCheckedSignedInt& operator=(const TCheckedSignedInt& Other) = default;
+
+	/** @return Returns an explicitly invalid value. */
+	static TCheckedSignedInt Invalid() { return TCheckedSignedInt(); }
+
+	/** @return true if current value is valid (assigned and no overflows or other errors occurred), false otherwise. */
+	bool IsValid() const { return bIsValid; }
+
+	/** @return The current value. Must check whether value IsValid() first. */
+	const SignedType GetValue() const
+	{
+		checkf(IsValid(), "Must check IsValid() before calling GetValue() on a TCheckedSignedInt, or alternatively use Get() with a default value.");
+		// Note our Invalid() placeholder has initialized Value=0, and we make sure to always
+		// initialize invalid values as Invalid() so even when compiled with checks off, this
+		// is a defined value.
+		return Value;
+	}
+
+	/** @return The value if valid, DefaultValue otherwise. */
+	const SignedType Get(const SignedType DefaultValue) const
+	{
+		return IsValid() ? Value : DefaultValue;
+	}
+
+	/** @return true if *this and Other are either both invalid or both valid and have the same value, false otherwise. */
+	bool operator ==(const TCheckedSignedInt Other) const
+	{
+		if (bIsValid != Other.bIsValid)
+		{
+			return false;
+		}
+
+		return !bIsValid || Value == Other.Value;
+	}
+
+	/** @return true if *this and Other either have different "valid" states or are both valid and have different values, false otherwise (logical negation of ==). */
+	bool operator !=(const TCheckedSignedInt Other) const
+	{
+		if (bIsValid != Other.bIsValid)
+		{
+			return true;
+		}
+
+		return bIsValid && Value != Other.Value;
+	}
+
+	// There are intentionally no overloads for the ordered comparison operators, because we have
+	// to decide what to do about validity as well. Instead, do this.
+
+	/** @return true if *this and Other are both valid so they can be compared. */
+	bool ComparisonValid(const TCheckedSignedInt Other) const { return bIsValid && Other.bIsValid; }
+
+	/** @return true if *this and Other are both valid and *this is less than Other. */
+	template<typename ValueType>
+	bool ValidAndLessThan(const ValueType Other) const { TCheckedSignedInt CheckedOther{ Other }; return ComparisonValid(CheckedOther) && Value < CheckedOther.Value; }
+	/** @return true if *this and Other are both valid and *this is less than or equal to Other. */
+	template<typename ValueType>
+	bool ValidAndLessOrEqual(const ValueType Other) const { TCheckedSignedInt CheckedOther{ Other }; return ComparisonValid(CheckedOther) && Value <= CheckedOther.Value; }
+	/** @return true if *this and Other are both valid and *this is greater than Other. */
+	template<typename ValueType>
+	bool ValidAndGreaterThan(const ValueType Other) const { TCheckedSignedInt CheckedOther{ Other }; return ComparisonValid(CheckedOther) && Value > CheckedOther.Value; }
+	/** @return true if *this and Other are both valid and *this is greater than or equal to Other. */
+	template<typename ValueType>
+	bool ValidAndGreaterOrEqual(const ValueType Other) const { TCheckedSignedInt CheckedOther{ Other }; return ComparisonValid(CheckedOther) && Value >= CheckedOther.Value; }
+
+	/** @return true if either of *this or Other are invalid or *this is less than Other. */
+	template<typename ValueType>
+	bool InvalidOrLessThan(const ValueType Other) const { TCheckedSignedInt CheckedOther{ Other }; return !ComparisonValid(CheckedOther) || Value < CheckedOther.Value; }
+	/** @return true if either of *this or Other are invalid or *this is less than or equal to Other. */
+	template<typename ValueType>
+	bool InvalidOrLessOrEqual(const ValueType Other) const { TCheckedSignedInt CheckedOther{ Other }; return !ComparisonValid(CheckedOther) || Value <= CheckedOther.Value; }
+	/** @return true if either of *this or Other are invalid or *this is greater than Other. */
+	template<typename ValueType>
+	bool InvalidOrGreaterThan(const ValueType Other) const { TCheckedSignedInt CheckedOther{ Other }; return !ComparisonValid(CheckedOther) || Value > CheckedOther.Value; }
+	/** @return true if either of *this or Other are invalid or *this is greater than or equal to Other. */
+	template<typename ValueType>
+	bool InvalidOrGreaterOrEqual(const ValueType Other) const { TCheckedSignedInt CheckedOther{ Other }; return !ComparisonValid(CheckedOther) || Value >= CheckedOther.Value; }
+
+	// Arithmetic operations
+
+	/** @return The negated value. */
+	TCheckedSignedInt operator-() const
+	{
+		// Unary negation (for two's complement) overflows iff the operand is MinValue.
+		return (bIsValid && Value > MinValue) ? TCheckedSignedInt(-Value) : Invalid();
+	}
+
+	/** @return The sum of the two operands. */
+	TCheckedSignedInt operator +(const TCheckedSignedInt Other) const
+	{
+		// Any sum involving an invalid value is invalid.
+		if (!bIsValid || !Other.bIsValid)
+		{
+			return Invalid();
+		}
+
+		// This follows Hacker's Delight, Chapter 2-12
+		// Signed->unsigned conversion and unsigned addition have defined behavior always
+		const UnsignedType UnsignedA = (UnsignedType)Value;
+		const UnsignedType UnsignedB = (UnsignedType)Other.Value;
+		const UnsignedType UnsignedSum = UnsignedA + UnsignedB;
+
+		// Check for signed overflow.
+		// The underlying logic here is pretty simple: if A and B had opposite signs, their sum can't
+		// overflow. If they had the same sign and the sum has the opposite value in the sign bit, we
+		// had an overflow. (See Hacker's Delight Chapter 2-12 for more details.)
+		if ((UnsignedSum ^ UnsignedA) & (UnsignedSum ^ UnsignedB) & UnsignedMSB)
+		{
+			return Invalid();
+		}
+
+		return TCheckedSignedInt(Value + Other.Value);
+	}
+
+	/** @return The difference between the two operands. */
+	TCheckedSignedInt operator -(const TCheckedSignedInt Other) const
+	{
+		// Any difference involving an invalid value is invalid.
+		if (!bIsValid || !Other.bIsValid)
+		{
+			return Invalid();
+		}
+
+		// This follows Hacker's Delight, Chapter 2-12
+		// Signed->unsigned conversion and unsigned subtraction have defined behavior always
+		const UnsignedType UnsignedA = (UnsignedType)Value;
+		const UnsignedType UnsignedB = (UnsignedType)Other.Value;
+		const UnsignedType UnsignedDiff = UnsignedA - UnsignedB;
+
+		// Check for signed overflow.
+		// If A and B have the same sign, the difference can't overflow. Therefore, we test for cases
+		// where the sign bit differs meaning ((UnsignedA ^ UnsignedB) & UnsignedMSB) != 0, and
+		// simultaneously the sign of the difference differs from the sign of the minuend (which should
+		// keep its sign when we're subtracting a value of the opposite sign), meaning
+		// ((UnsignedDiff ^ UnsignedA) & UnsignedMSB) != 0. Combining the two yields:
+		if ((UnsignedA ^ UnsignedB) & (UnsignedDiff ^ UnsignedA) & UnsignedMSB)
+		{
+			return Invalid();
+		}
+
+		return TCheckedSignedInt(Value - Other.Value);
+	}
+
+	/** @return The product of the two operands. */
+	TCheckedSignedInt operator *(const TCheckedSignedInt Other) const
+	{
+		// Any product involving invalid values is invalid.
+		if (!bIsValid || !Other.bIsValid)
+		{
+			return Invalid();
+		}
+
+		// Handle the case where the second factor is 0 specially (why will become clear in a minute).
+		if (Other.Value == 0)
+		{
+			// Anything times 0 is 0.
+			return TCheckedSignedInt(0);
+		}
+
+		// The overflow check is annoying and expensive, but the basic idea is fairly simple:
+		// reduce to an unsigned check of the absolute values. (Again the basic algorithm is
+		// in Hacker's Delight, Chapter 2-12).
+		//
+		// We need the absolute value of the product to be <=MaxValue when the result is positive
+		// (signs of factors same) and <= -MinValue = MaxValue + 1 if the result is negative
+		// (signs of factors opposite).
+		UnsignedType UnsignedA = (UnsignedType)Value;
+		UnsignedType UnsignedB = (UnsignedType)Other.Value;
+		bool bSignsDiffer = false;
+
+		// Determine the unsigned absolute values of A and B carefully (note we can't negate signed
+		// Value or Other.Value, because negating MinValue is UB). We can however subtract their
+		// unsigned values from 0 if the original value was less than zero. While doing this, also
+		// keep track of the sign parity.
+		if (Value < 0)
+		{
+			UnsignedA = UnsignedType(0) - UnsignedA;
+			bSignsDiffer = !bSignsDiffer;
+		}
+
+		if (Other.Value < 0)
+		{
+			UnsignedB = UnsignedType(0) - UnsignedB;
+			bSignsDiffer = !bSignsDiffer;
+		}
+
+		// Determine the unsigned product bound we need based on whether the signs were same or different.
+		const UnsignedType ProductBound = UnsignedType(MaxValue) + (bSignsDiffer ? 1 : 0);
+
+		// We're now in the unsigned case, 0 <= UnsignedA, 0 < UnsignedB (we established b != 0), and for
+		// there not to be overflows we need
+		//   a * b <= ProductBound
+		// <=> a <= ProductBound/b
+		// <=> a <= floor(ProductBound/b)   since a is integer
+		return (UnsignedA <= ProductBound / UnsignedB) ? TCheckedSignedInt(Value * Other.Value) : Invalid();
+	}
+
+	/** @return The quotient when dividing *this by Other. */
+	TCheckedSignedInt operator /(const TCheckedSignedInt Other) const
+	{
+		// Any quotient involving invalid values is invalid.
+		if (!bIsValid || !Other.bIsValid)
+		{
+			return Invalid();
+		}
+
+		// Luckily for us, division generally makes things smaller, so there's only two things to watch
+		// out for: division by zero is not allowed, and division of MinValue by -1 would give -MinValue
+		// which overflows. All other combinations are fine.
+		if (Other.Value == 0 || (Value == MinValue && Other.Value == -1))
+		{
+			return Invalid();
+		}
+
+		return TCheckedSignedInt(Value / Other.Value);
+	}
+
+	/** @return The remainder when dividing *this by Other. */
+	TCheckedSignedInt operator %(const TCheckedSignedInt Other) const
+	{
+		// Any quotient involving invalid values is invalid.
+		if (!bIsValid || !Other.bIsValid)
+		{
+			return Invalid();
+		}
+
+		// Same error cases as for division.
+		if (Other.Value == 0 || (Value == MinValue && Other.Value == -1))
+		{
+			return Invalid();
+		}
+
+		return TCheckedSignedInt(Value % Other.Value);
+	}
+
+	/** @return This value bitwise left-shifted by the operand. */
+	TCheckedSignedInt operator <<(const TCheckedSignedInt Other) const
+	{
+		// Any shift involving invalid values is invalid.
+		if (!bIsValid || !Other.bIsValid)
+		{
+			return Invalid();
+		}
+
+		// Left-shifts by negative values or >= the width of the type are always invalid.
+		if (Other.Value < 0 || Other.Value >= NumBits)
+		{
+			return Invalid();
+		}
+
+		const int ShiftAmount = Other.Value;
+
+		// Once again, taking our overflow-prone expression and using algebra to find
+		// a form that doesn't overflow:
+		//
+		// MinValue <= a * 2^b <= MaxValue
+		// <=> MinValue * 2^(-b) <= a <= MaxValue * 2^(-b)
+		//
+		// The LHS is exact because MinValue is -2^(NumBits - 1) for two's complement,
+		// and we just ensured that 0 <= b < NumBits (with b integer).
+		//
+		// The RHS has a fractional part whereas a is integer; therefore, we can
+		// substitute floor(MaxValue * 2^(-b)) for the RHS without changing the result.
+		//
+		// And that gives us our test!
+		return ((MinValue >> ShiftAmount) <= Value && Value <= (MaxValue >> ShiftAmount)) ? TCheckedSignedInt(Value << ShiftAmount) : Invalid();
+	}
+
+	/** @return This value bitwise right-shifted by the operand. */
+	TCheckedSignedInt operator >>(const TCheckedSignedInt Other) const
+	{
+		// Any shift involving invalid values is invalid.
+		if (!bIsValid || !Other.bIsValid)
+		{
+			return Invalid();
+		}
+
+		// Right-shifts by negative values or >= the width of the type are always invalid.
+		if (Other.Value < 0 || Other.Value >= NumBits)
+		{
+			return Invalid();
+		}
+
+		// Right-shifts don't have any overflow conditions, so we're good!
+		return TCheckedSignedInt(Value >> Other.Value);
+	}
+
+	/** @return The absolute value of *this. */
+	TCheckedSignedInt Abs() const
+	{
+		if (!bIsValid)
+		{
+			return Invalid();
+		}
+
+		// Note the absolute value of MinValue overflows, so this is not completely trivial.
+		// Can't just use TCheckedSignedInt(abs(Value)) here!
+		return (Value < 0) ? -*this : *this;
+	}
+
+	// Mixed-type operators and assignment operators reduce to the base operators systematically
+#define UE_CHECKED_SIGNED_INT_IMPL_BINARY_OPERATOR(OP) \
+	/* Mixed-type expressions that coerce both operands to checked ints */ \
+	template<typename OperandType> \
+	TCheckedSignedInt operator OP(OperandType InB) const { return *this OP TCheckedSignedInt(InB); } \
+	template<typename OperandType> \
+	friend TCheckedSignedInt operator OP(OperandType InA, TCheckedSignedInt InB) { return TCheckedSignedInt(InA) OP InB; } \
+	/* Assignment operators, direct and mixed */ \
+	TCheckedSignedInt& operator OP##=(TCheckedSignedInt InB) { return *this = *this OP InB; } \
+	template<typename OperandType> \
+	TCheckedSignedInt& operator OP##=(OperandType InB) { return *this = *this OP TCheckedSignedInt(InB); } \
+	/* end */
+
+	UE_CHECKED_SIGNED_INT_IMPL_BINARY_OPERATOR(+)
+	UE_CHECKED_SIGNED_INT_IMPL_BINARY_OPERATOR(-)
+	UE_CHECKED_SIGNED_INT_IMPL_BINARY_OPERATOR(*)
+	UE_CHECKED_SIGNED_INT_IMPL_BINARY_OPERATOR(/)
+	UE_CHECKED_SIGNED_INT_IMPL_BINARY_OPERATOR(%)
+	UE_CHECKED_SIGNED_INT_IMPL_BINARY_OPERATOR(<<)
+	UE_CHECKED_SIGNED_INT_IMPL_BINARY_OPERATOR(>>)
+
+#undef UE_CHECKED_SIGNED_INT_IMPL_BINARY_OPERATOR
+};
+
+/** Checked 32-bit integer class. Used to deal with integer data from untrusted sources in size computations etc. */
+using FCheckedInt32 = TCheckedSignedInt<int32>;
+
+/** Checked 64-bit integer class. Used to deal with integer data from untrusted sources in size computations etc. */
+using FCheckedInt64 = TCheckedSignedInt<int64>;
