Fix Clang warning for comparison of different types.

diff --git a/Engine/Source/Runtime/Core/Public/Misc/TVariantMeta.h b/Engine/Source/Runtime/Core/Public/Misc/TVariantMeta.h
--- a/Engine/Source/Runtime/Core/Public/Misc/TVariantMeta.h
+++ b/Engine/Source/Runtime/Core/Public/Misc/TVariantMeta.h
@@ -76,7 +76,7 @@ namespace Private
 		static constexpr SIZE_T MaxOf(const SIZE_T Sizes[])
 		{
 			SIZE_T MaxSize = 0;
-			for (int32 Itr = 0; Itr < sizeof...(Ts); ++Itr)
+			for (SIZE_T Itr = 0; Itr < sizeof...(Ts); ++Itr)
 			{
 				if (Sizes[Itr] > MaxSize)
 				{
