Further improved the cpu timing scopes for module loading (now using cpu timing events with name of module as metadata).

diff --git a/Engine/Source/Runtime/Core/Private/Modules/ModuleManager.cpp b/Engine/Source/Runtime/Core/Private/Modules/ModuleManager.cpp
--- a/Engine/Source/Runtime/Core/Private/Modules/ModuleManager.cpp
+++ b/Engine/Source/Runtime/Core/Private/Modules/ModuleManager.cpp
@@ -16,6 +16,8 @@
 #include "Misc/FileHelper.h"
 #include "Serialization/MemoryReader.h"
 #include "Serialization/MemoryWriter.h"
+#include "Trace/Trace.h"
+#include "Trace/Trace.inl"
 
 DEFINE_LOG_CATEGORY_STATIC(LogModuleManager, Log, All);
 
@@ -395,6 +397,17 @@ IModuleInterface& FModuleManager::LoadModuleChecked( const FName InModuleName )
 	return *Module;
 }
 
+UE_TRACE_EVENT_BEGIN(Cpu, LoadModule, NoSync)
+	UE_TRACE_EVENT_FIELD(UE::Trace::WideString, Name)
+UE_TRACE_EVENT_END()
+
+UE_TRACE_EVENT_BEGIN(Cpu, FPlatformProcess_GetDllHandle, NoSync)
+	UE_TRACE_EVENT_FIELD(UE::Trace::WideString, Name)
+UE_TRACE_EVENT_END()
+
+UE_TRACE_EVENT_BEGIN(Cpu, StartupModule, NoSync)
+	UE_TRACE_EVENT_FIELD(UE::Trace::WideString, Name)
+UE_TRACE_EVENT_END()
 
 IModuleInterface* FModuleManager::LoadModuleWithFailureReason(const FName InModuleName, EModuleLoadResult& OutFailureReason, ELoadModuleFlags InLoadModuleFlags)
 {
@@ -430,9 +443,10 @@ IModuleInterface* FModuleManager::LoadModuleWithFailureReason(const FName InModu
 //	ensureMsgf(IsInGameThread(), TEXT("ModuleManager: Attempting to load '%s' outside the main thread.  Please call LoadModule on the main/game thread only.  You can use GetModule or GetModuleChecked instead, those are safe to call outside the game thread."), *InModuleName.ToString());
 
 	UE_SCOPED_ENGINE_ACTIVITY(TEXT("Loading Module %s"), *InModuleName.ToString());
-	SCOPED_BOOT_TIMING("LoadModule");
-	TRACE_CPUPROFILER_EVENT_SCOPE_TEXT(*FString::Printf(TEXT("LoadModule_%s"), *InModuleName.ToString()));
+	FScopedBootTiming BootTimingScope("LoadModule");
 	DECLARE_SCOPE_CYCLE_COUNTER(TEXT("Module Load"), STAT_ModuleLoad, STATGROUP_LoadTime);
+	UE_TRACE_LOG_SCOPED_T(Cpu, LoadModule, CpuChannel)
+		<< LoadModule.Name(*InModuleName.ToString());
 
 #if	STATS
 	// This is fine here, we only load a handful of modules.
@@ -482,6 +496,7 @@ IModuleInterface* FModuleManager::LoadModuleWithFailureReason(const FName InModu
 		{
 			FScopedBootTiming BootScope("LoadModule  - ", InModuleName);
 			TRACE_LOADTIME_REQUEST_GROUP_SCOPE(TEXT("LoadModule - %s"), *InModuleName.ToString());
+
 #if USE_PER_MODULE_UOBJECT_BOOTSTRAP || WITH_VERSE
 			{
 				// Defer String Table find/load during CDO registration, as it may happen 
@@ -491,8 +506,13 @@ IModuleInterface* FModuleManager::LoadModuleWithFailureReason(const FName InModu
 				ProcessLoadedObjectsCallback.Broadcast(InModuleName, bCanProcessNewlyLoadedObjects);
 			}
 #endif
+
 			// Startup the module
-			ModuleInfo->Module->StartupModule();
+			{
+				UE_TRACE_LOG_SCOPED_T(Cpu, StartupModule, CpuChannel)
+					<< StartupModule.Name(*InModuleName.ToString());
+				ModuleInfo->Module->StartupModule();
+			}
 
 			// The module might try to load other dependent modules in StartupModule. In this case, we want those modules shut down AFTER this one because we may still depend on the module at shutdown.
 			ModuleInfo->LoadOrder = FModuleInfo::CurrentLoadOrder++;
@@ -570,7 +590,8 @@ IModuleInterface* FModuleManager::LoadModuleWithFailureReason(const FName InModu
 		if (FPaths::FileExists(ModuleFileToLoad))
 		{
 			{
-				TRACE_CPUPROFILER_EVENT_SCOPE(FPlatformProcess::GetDllHandle);
+				UE_TRACE_LOG_SCOPED_T(Cpu, FPlatformProcess_GetDllHandle, CpuChannel)
+					<< FPlatformProcess_GetDllHandle.Name(*ModuleFileToLoad);
 				ModuleInfo->Handle = FPlatformProcess::GetDllHandle(*ModuleFileToLoad);
 			}
 			
@@ -606,12 +627,14 @@ IModuleInterface* FModuleManager::LoadModuleWithFailureReason(const FName InModu
 						// Initialize the module!
 						ModuleInfo->Module = TUniquePtr<IModuleInterface>(InitializeModuleFunctionPtr());
 
-						if ( ModuleInfo->Module.IsValid() )
+						if (ModuleInfo->Module.IsValid())
 						{
-							TRACE_CPUPROFILER_EVENT_SCOPE_TEXT(*(InModuleName.ToString() + TEXT("::StartupModule")));
-
 							// Startup the module
-							ModuleInfo->Module->StartupModule();
+							{
+								UE_TRACE_LOG_SCOPED_T(Cpu, StartupModule, CpuChannel)
+									<< StartupModule.Name(*InModuleName.ToString());
+								ModuleInfo->Module->StartupModule();
+							}
 
 							// The module might try to load other dependent modules in StartupModule. In this case, we want those modules shut down AFTER this one because we may still depend on the module at shutdown.
 							ModuleInfo->LoadOrder = FModuleInfo::CurrentLoadOrder++;
diff --git a/Engine/Source/Runtime/Core/Private/Windows/WindowsPlatformProcess.cpp b/Engine/Source/Runtime/Core/Private/Windows/WindowsPlatformProcess.cpp
--- a/Engine/Source/Runtime/Core/Private/Windows/WindowsPlatformProcess.cpp
+++ b/Engine/Source/Runtime/Core/Private/Windows/WindowsPlatformProcess.cpp
@@ -26,6 +26,8 @@
 #include "Misc/CoreDelegates.h"
 #include "Misc/Fork.h"
 #include "ProfilingDebugging/CsvProfiler.h"
+#include "Trace/Trace.h"
+#include "Trace/Trace.inl"
 
 #include "Windows/AllowWindowsPlatformTypes.h"
 	#include <shellapi.h>
@@ -1914,6 +1916,10 @@ static bool ResolveImport(const FString& Name, const TArray<FString>& SearchPath
 	return false;
 }
 
+UE_TRACE_EVENT_BEGIN(Cpu, ResolveMissingImports, NoSync)
+	UE_TRACE_EVENT_FIELD(UE::Trace::WideString, Name)
+UE_TRACE_EVENT_END()
+
 /**
  * Resolve all the imports for the given library, searching through a set of directories.
  *
@@ -1924,6 +1930,9 @@ static bool ResolveImport(const FString& Name, const TArray<FString>& SearchPath
  */
 static void ResolveMissingImportsRecursive(const FString& FileName, const TArray<FString>& SearchPaths, TArray<FString>& ImportFileNames, TSet<FString>& VisitedImportNames)
 {
+	UE_TRACE_LOG_SCOPED_T(Cpu, ResolveMissingImports, CpuChannel)
+		<< ResolveMissingImports.Name(*FileName);
+
 	// Read the imports for this library
 	TArray<FString> ImportNames;
 	if(ReadLibraryImports(*FileName, ImportNames))
@@ -1982,6 +1991,10 @@ static void LogImportDiagnostics(const FString& FileName, const TArray<FString>&
 	}
 }
 
+UE_TRACE_EVENT_BEGIN(Cpu, Windows_LoadLibrary, NoSync)
+	UE_TRACE_EVENT_FIELD(UE::Trace::WideString, Name)
+UE_TRACE_EVENT_END()
+
 void *FWindowsPlatformProcess::LoadLibraryWithSearchPaths(const FString& FileName, const TArray<FString>& SearchPaths)
 {
 	UE_SCOPED_IO_ACTIVITY(*WriteToString<256>("Loading Dll ", FileName));
@@ -2010,7 +2023,8 @@ void *FWindowsPlatformProcess::LoadLibraryWithSearchPaths(const FString& FileNam
 			{
 				const void* DependencyHandle = [&ImportFileName]() 
 				{
-					TRACE_CPUPROFILER_EVENT_SCOPE(Windows::LoadLibrary);
+					UE_TRACE_LOG_SCOPED_T(Cpu, Windows_LoadLibrary, CpuChannel)
+						<< Windows_LoadLibrary.Name(*ImportFileName);
 					return LoadLibrary(*ImportFileName);
 				}();
 				
@@ -2030,7 +2044,8 @@ void *FWindowsPlatformProcess::LoadLibraryWithSearchPaths(const FString& FileNam
 	// Try to load the actual library
 	void* Handle = [FullFileName]() 
 	{
-		TRACE_CPUPROFILER_EVENT_SCOPE(Windows::LoadLibrary);
+		UE_TRACE_LOG_SCOPED_T(Cpu, Windows_LoadLibrary, CpuChannel)
+			<< Windows_LoadLibrary.Name(*FullFileName);
 		return LoadLibrary(*FullFileName);
 	}();
 	
