Add UE_CALL_ONCE varargs support

diff --git a/Engine/Source/Runtime/Core/Public/Misc/CoreMiscDefines.h b/Engine/Source/Runtime/Core/Public/Misc/CoreMiscDefines.h
--- a/Engine/Source/Runtime/Core/Public/Misc/CoreMiscDefines.h
+++ b/Engine/Source/Runtime/Core/Public/Misc/CoreMiscDefines.h
@@ -136,8 +136,8 @@ enum EInPlace {InPlace};
 	#define ANONYMOUS_VARIABLE( Name ) PREPROCESSOR_JOIN(Name, __LINE__)
 #endif
 
-/** Thread-safe call once helper, similar to std::call_once without the std::once_flag */
-#define UE_CALL_ONCE(Func) static int32 ANONYMOUS_VARIABLE(ThreadSafeOnce) = ((Func)(), 1)
+/** Thread-safe call once helper for void functions, similar to std::call_once without the std::once_flag */
+#define UE_CALL_ONCE(Func, ...) static int32 ANONYMOUS_VARIABLE(ThreadSafeOnce) = ((Func)(__VA_ARGS__), 1)
 
 /**
  * Macro for marking up deprecated code, functions and types.
