Add list of lights to GPU scene, mirroring Scene::Lights. Each element contains the view-independent properties of the light.

diff --git a/Engine/Shaders/Private/LightData.ush b/Engine/Shaders/Private/LightData.ush
--- a/Engine/Shaders/Private/LightData.ush
+++ b/Engine/Shaders/Private/LightData.ush
@@ -73,7 +73,7 @@ struct FLocalLightData
 	// Raw value form light buffers
 	float4 LightPositionAndInvRadius;
 	float4 LightColorAndFalloffExponent;
-	float4 SpotAnglesAndSourceRadiusPacked;
+	float4 SpotAnglesAndIdAndSourceRadiusPacked;
 	float4 LightDirectionAndShadowMask;
 	float4 LightTangentAndSpecularScale;
 	float4 RectDataAndVirtualShadowMapId;
@@ -86,6 +86,8 @@ struct FLocalLightData
 	bool bIsSimpleLight;
 	/** Virtual shadow map ID or INDEX_NONE if not present. */
 	int  VirtualShadowMapId;
+	/** Index into GPUScene lights */
+	int LightSceneId;
 };
 
 ///////////////////////////////////////////////////////////////////////////////////////////////////
@@ -199,22 +201,22 @@ bool UnpackCastShadow(uint InShadowMapChannelMaskPacked)
 
 float UnpackVolumetricScatteringIntensity(FLocalLightData In)
 {
-	return f16tof32(asuint(In.SpotAnglesAndSourceRadiusPacked.w) >> 16);
+	return f16tof32(asuint(In.SpotAnglesAndIdAndSourceRadiusPacked.w) >> 16);
 }
 
 float UnpackLightSourceRadius(FLocalLightData In)
 {
-	return f16tof32(asuint(In.SpotAnglesAndSourceRadiusPacked.z) & 0xFFFF);
+	return f16tof32(asuint(In.SpotAnglesAndIdAndSourceRadiusPacked.z) & 0xFFFF);
 }
 
 float UnpackLightSoftSourceRadius(FLocalLightData In)
 {
-	return f16tof32(asuint(In.SpotAnglesAndSourceRadiusPacked.z) >> 16);
+	return f16tof32(asuint(In.SpotAnglesAndIdAndSourceRadiusPacked.z) >> 16);
 }
 
 float UnpackLightSourceLength(FLocalLightData In)
 {
-	return f16tof32(asuint(In.SpotAnglesAndSourceRadiusPacked.w));
+	return f16tof32(asuint(In.SpotAnglesAndIdAndSourceRadiusPacked.w));
 }
 
 float GetLightSourceRadius(FLocalLightData In)
@@ -224,7 +226,9 @@ float GetLightSourceRadius(FLocalLightData In)
 
 float2 GetLightSpotAngles(FLocalLightData In)
 {
-	return In.SpotAnglesAndSourceRadiusPacked.xy;
+	uint SpotAnglesPacked = asuint(In.SpotAnglesAndIdAndSourceRadiusPacked.x);
+	float2 SpotAngles = float2(f16tof32(SpotAnglesPacked & 0xFFFF), f16tof32(SpotAnglesPacked >> 16));
+	return SpotAngles;
 }
 
 ///////////////////////////////////////////////////////////////////////////////////////////////////
@@ -256,7 +260,7 @@ FDeferredLightData ConvertToDeferredLight(
 	Out.Color					= In.LightColorAndFalloffExponent.xyz;
 	Out.FalloffExponent			= In.LightColorAndFalloffExponent.w;
 	Out.Direction				= In.LightDirectionAndShadowMask.xyz;
-	Out.SpotAngles				= In.SpotAnglesAndSourceRadiusPacked.xy;
+	Out.SpotAngles				= GetLightSpotAngles(In);
 	Out.SourceRadius			= UnpackLightSourceRadius(In);
 	Out.SourceLength			= UnpackLightSourceLength(In);
 	Out.Tangent					= In.LightTangentAndSpecularScale.xyz;
@@ -343,7 +347,7 @@ FLightShaderParameters ConvertFromLocal(const FLocalLightData In)
 	Out.FalloffExponent			= In.LightColorAndFalloffExponent.w;
 	Out.Direction				= In.LightDirectionAndShadowMask.xyz;
 	Out.Tangent					= In.LightTangentAndSpecularScale.xyz;
-	Out.SpotAngles				= In.SpotAnglesAndSourceRadiusPacked.xy;
+	Out.SpotAngles				= GetLightSpotAngles(In);
 	Out.SpecularScale			= In.LightTangentAndSpecularScale.w;
 	Out.SourceRadius			= UnpackLightSourceRadius(In);
 	Out.SoftSourceRadius		= UnpackLightSoftSourceRadius(In);
diff --git a/Engine/Shaders/Private/LightGridCommon.ush b/Engine/Shaders/Private/LightGridCommon.ush
--- a/Engine/Shaders/Private/LightGridCommon.ush
+++ b/Engine/Shaders/Private/LightGridCommon.ush
@@ -162,16 +162,17 @@ FLocalLightData GetLocalLightData_Internal(
 {
 	FLocalLightData Out = (FLocalLightData)0;
 
-	Out.bClusteredDeferredSupported		= LocalLightIndex <  ClusteredDeferredSupportedEndIndex;
-	Out.bLumenLightSupported			= LocalLightIndex >= LumenSupportedStartIndex;
-	Out.bIsSimpleLight					= LocalLightIndex <  SimpleLightsEndIndex;
-	Out.LightPositionAndInvRadius		= InData0;
-	Out.LightColorAndFalloffExponent	= InData1;
-	Out.LightDirectionAndShadowMask		= InData2;
-	Out.SpotAnglesAndSourceRadiusPacked	= InData3;
-	Out.LightTangentAndSpecularScale	= InData4;
-	Out.RectDataAndVirtualShadowMapId	= InData5;
-	Out.VirtualShadowMapId				= int(Out.RectDataAndVirtualShadowMapId.z);
+	Out.bClusteredDeferredSupported				= LocalLightIndex <  ClusteredDeferredSupportedEndIndex;
+	Out.bLumenLightSupported					= LocalLightIndex >= LumenSupportedStartIndex;
+	Out.bIsSimpleLight							= LocalLightIndex <  SimpleLightsEndIndex;
+	Out.LightPositionAndInvRadius				= InData0;
+	Out.LightColorAndFalloffExponent			= InData1;
+	Out.LightDirectionAndShadowMask				= InData2;
+	Out.SpotAnglesAndIdAndSourceRadiusPacked	= InData3;
+	Out.LightTangentAndSpecularScale			= InData4;
+	Out.RectDataAndVirtualShadowMapId			= InData5;
+	Out.VirtualShadowMapId						= int(Out.RectDataAndVirtualShadowMapId.z);
+	Out.LightSceneId							= int(Out.SpotAnglesAndIdAndSourceRadiusPacked.y);
 
 	return Out;
 }
diff --git a/Engine/Shaders/Private/SceneData.ush b/Engine/Shaders/Private/SceneData.ush
--- a/Engine/Shaders/Private/SceneData.ush
+++ b/Engine/Shaders/Private/SceneData.ush
@@ -4,6 +4,7 @@
 
 #include "LargeWorldCoordinates.ush"
 #include "OctahedralCommon.ush"
+#include "/Engine/Shared/LightSceneData.h"
 
 #ifndef USE_GLOBAL_GPU_SCENE_DATA
 	#define USE_GLOBAL_GPU_SCENE_DATA 0
@@ -1213,3 +1214,24 @@ FSceneDataIntermediates GetSceneDataIntermediates()
 	return Intermediates;
 }
 #endif //VF_USE_PRIMITIVE_SCENE_DATA
+
+//
+// GPU Lights
+//
+
+#if USE_GLOBAL_GPU_SCENE_DATA
+StructuredBuffer<FLightSceneData> GPUSceneLights;
+#elif USE_GLOBAL_GPU_SCENE_DATA_RW
+RWStructuredBuffer<FLightSceneData> GPUSceneLightsRW;
+#endif
+
+FLightSceneData GetLightSceneData(int LightId)
+{
+#if USE_GLOBAL_GPU_SCENE_DATA
+	return GPUSceneLights[LightId];
+#elif USE_GLOBAL_GPU_SCENE_DATA_RW
+	return GPUSceneLightsRW[LightId];
+#else
+	return (FLightSceneData)0; //TODO
+#endif
+}
diff --git a/Engine/Shaders/Shared/HLSLTypeAliases.h b/Engine/Shaders/Shared/HLSLTypeAliases.h
new file mode 100644
--- /dev/null
+++ b/Engine/Shaders/Shared/HLSLTypeAliases.h
@@ -0,0 +1,35 @@
+// Copyright Epic Games, Inc. All Rights Reserved.
+
+#pragma once
+
+/*================================================================================================
+	HLSLTypeAliases.h: include this file in headers that are used in both C++ and HLSL
+	to define type aliases that map HLSL types to C++ UE types.
+=================================================================================================*/
+
+#include "Math/MathFwd.h"
+#include "Misc/LargeWorldRenderPosition.h"
+
+namespace UE::HLSL
+{
+	using int2 = FIntVector2;
+	using int3 = FIntVector3;
+	using int4 = FIntVector4;
+
+	using uint = uint32;
+	using uint2 = FUintVector2;
+	using uint3 = FUintVector3;
+	using uint4 = FUintVector4;
+
+	//dword
+	//half
+
+	using float2 = FVector2f;
+	using float3 = FVector3f;
+	using float4 = FVector4f;
+
+	using float4x4 = FMatrix44f;
+
+	using FLWCVector3 = TLargeWorldRenderPosition<float>;
+}
+
diff --git a/Engine/Shaders/Shared/LightSceneData.h b/Engine/Shaders/Shared/LightSceneData.h
new file mode 100644
--- /dev/null
+++ b/Engine/Shaders/Shared/LightSceneData.h
@@ -0,0 +1,71 @@
+// Copyright Epic Games, Inc. All Rights Reserved.
+
+#pragma once
+
+#ifdef __cplusplus
+#include "HLSLTypeAliases.h"
+
+namespace UE::HLSL
+{
+#endif
+
+/**
+ * Has a 1:1 mapping with FLightRenderParameters, but unlike FLightShaderParameters, this is view-independent
+ */
+struct FLightSceneData
+{
+	// Position of the light in world space.
+	FLWCVector3 WorldPosition;
+
+	// 1 / light's falloff radius from Position.
+	float InvRadius;
+
+	// Color of the light.
+	float4 Color;
+
+	// The exponent for the falloff of the light intensity from the distance.
+	float FalloffExponent;
+
+	// Direction of the light if applies.
+	float3 Direction;
+
+	// Factor to applies on the specular.
+	float SpecularScale;
+
+	// One tangent of the light if applies.
+	// Note: BiTangent is on purpose not stored for memory optimisation purposes.
+	float3 Tangent;
+
+	// Radius of the point light.
+	float SourceRadius;
+
+	// Dimensions of the light, for spot light, but also
+	float2 SpotAngles;
+
+	// Radius of the soft source.
+	float SoftSourceRadius;
+
+	// Other dimensions of the light source for rect light specifically.
+	float SourceLength;
+
+	// Barn door angle for rect light
+	float RectLightBarnCosAngle;
+
+	// Barn door length for rect light
+	float RectLightBarnLength;
+
+	// Rect. light atlas transformation
+	float2 RectLightAtlasUVOffset;
+	float2 RectLightAtlasUVScale;
+	float RectLightAtlasMaxLevel;
+
+	float InverseExposureBlend;
+
+	// Extra fields
+	uint LightTypeAndShadowMapChannelMaskPacked;
+};
+
+#ifdef __cplusplus
+} // namespace
+using FLightSceneData = UE::HLSL::FLightSceneData;
+#endif
diff --git a/Engine/Source/Runtime/Core/Private/Misc/LargeWorldRenderPosition.cpp b/Engine/Source/Runtime/Core/Private/Misc/LargeWorldRenderPosition.cpp
--- a/Engine/Source/Runtime/Core/Private/Misc/LargeWorldRenderPosition.cpp
+++ b/Engine/Source/Runtime/Core/Private/Misc/LargeWorldRenderPosition.cpp
@@ -22,12 +22,14 @@
 // Normally offsets should be within +/-TileSizeDivideBy2, but we often rebase multiple quantities off a single tile origin
 #define UE_LWC_RENDER_MAX_OFFSET (2097152.0*0.5)
 
-double FLargeWorldRenderScalar::GetTileSize()
+template<typename TScalar>
+TScalar TLargeWorldRenderScalar<TScalar>::GetTileSize()
 {
-	return UE_LWC_RENDER_TILE_SIZE;
+	return static_cast<TScalar>(UE_LWC_RENDER_TILE_SIZE);
 }
 
-FVector3f FLargeWorldRenderScalar::GetTileFor(FVector InPosition)
+template<typename TScalar>
+FVector3f TLargeWorldRenderScalar<TScalar>::GetTileFor(FVector InPosition)
 {
 	if constexpr (UE_LWC_RENDER_TILE_SIZE == 0)
 	{
@@ -58,37 +60,44 @@ FMatrix CheckMatrixInTileOffsetRange(const FMatrix& Matrix)
 	return Matrix;
 }
 
-FMatrix44f FLargeWorldRenderScalar::SafeCastMatrix(const FMatrix& Matrix)
+template<typename TScalar>
+FMatrix44f TLargeWorldRenderScalar<TScalar>::SafeCastMatrix(const FMatrix& Matrix)
 {
 	return FMatrix44f(CheckMatrixInTileOffsetRange(Matrix));
 }
 
-FMatrix44f FLargeWorldRenderScalar::MakeToRelativeWorldMatrix(const FVector Origin, const FMatrix& ToWorld)
+template<typename TScalar>
+FMatrix44f TLargeWorldRenderScalar<TScalar>::MakeToRelativeWorldMatrix(const FVector Origin, const FMatrix& ToWorld)
 {
 	return FMatrix44f(MakeToRelativeWorldMatrixDouble(Origin, ToWorld));
 }
 
-FMatrix FLargeWorldRenderScalar::MakeToRelativeWorldMatrixDouble(const FVector Origin, const FMatrix& ToWorld)
+template<typename TScalar>
+FMatrix TLargeWorldRenderScalar<TScalar>::MakeToRelativeWorldMatrixDouble(const FVector Origin, const FMatrix& ToWorld)
 {
 	return CheckMatrixInTileOffsetRange(ToWorld * FTranslationMatrix(-Origin));
 }
 
-FMatrix44f FLargeWorldRenderScalar::MakeFromRelativeWorldMatrix(const FVector Origin, const FMatrix& FromWorld)
+template<typename TScalar>
+FMatrix44f TLargeWorldRenderScalar<TScalar>::MakeFromRelativeWorldMatrix(const FVector Origin, const FMatrix& FromWorld)
 {
 	return FMatrix44f(MakeFromRelativeWorldMatrixDouble(Origin, FromWorld));
 }
 
-FMatrix FLargeWorldRenderScalar::MakeFromRelativeWorldMatrixDouble(const FVector Origin, const FMatrix& FromWorld)
+template<typename TScalar>
+FMatrix TLargeWorldRenderScalar<TScalar>::MakeFromRelativeWorldMatrixDouble(const FVector Origin, const FMatrix& FromWorld)
 {
 	return CheckMatrixInTileOffsetRange(FTranslationMatrix(Origin) * FromWorld);
 }
 
-FMatrix44f FLargeWorldRenderScalar::MakeClampedToRelativeWorldMatrix(const FVector Origin, const FMatrix& ToWorld)
+template<typename TScalar>
+FMatrix44f TLargeWorldRenderScalar<TScalar>::MakeClampedToRelativeWorldMatrix(const FVector Origin, const FMatrix& ToWorld)
 {
 	return FMatrix44f(MakeClampedToRelativeWorldMatrixDouble(Origin, ToWorld));
 }
 
-FMatrix FLargeWorldRenderScalar::MakeClampedToRelativeWorldMatrixDouble(const FVector Origin, const FMatrix& ToWorld)
+template<typename TScalar>
+FMatrix TLargeWorldRenderScalar<TScalar>::MakeClampedToRelativeWorldMatrixDouble(const FVector Origin, const FMatrix& ToWorld)
 {
 	const double OriginMax = UE_LWC_RENDER_MAX_OFFSET;
 
@@ -104,12 +113,16 @@ FMatrix FLargeWorldRenderScalar::MakeClampedToRelativeWorldMatrixDouble(const FV
 	return ClampedToRelativeWorld;
 }
 
-void FLargeWorldRenderScalar::Validate(double InAbsolute)
+template<typename TScalar>
+void TLargeWorldRenderScalar<TScalar>::Validate(double InAbsolute)
 {
 	const double Tolerance = 0.01; // TODO LWC - How precise do we need to be?
 	const double CheckAbsolute = GetAbsolute();
 	const double Delta = FMath::Abs(CheckAbsolute - InAbsolute);
 
-	ensureMsgf(Delta < Tolerance, TEXT("Bad FLargeWorldRenderScalar (%g) vs (%g)"),
+	ensureMsgf(Delta < Tolerance, TEXT("Bad TLargeWorldRenderScalar<TScalar> (%g) vs (%g)"),
 		InAbsolute, CheckAbsolute);
 }
+
+template struct TLargeWorldRenderScalar<double>;
+template struct TLargeWorldRenderScalar<float>;
\ No newline at end of file
diff --git a/Engine/Source/Runtime/Core/Public/Misc/LargeWorldRenderPosition.h b/Engine/Source/Runtime/Core/Public/Misc/LargeWorldRenderPosition.h
--- a/Engine/Source/Runtime/Core/Public/Misc/LargeWorldRenderPosition.h
+++ b/Engine/Source/Runtime/Core/Public/Misc/LargeWorldRenderPosition.h
@@ -10,31 +10,35 @@
 #include "Math/Vector4.h"
 #include "Misc/LargeWorldCoordinates.h"
 
-struct FLargeWorldRenderScalar
+template <typename TScalar>
+struct TLargeWorldRenderScalar
 {
+private:
+	using VectorType = UE::Math::TVector<TScalar>;
+
 public:
-	// This should become inline at some point, but keeping it in cpp file for now to make iteration faster when experementing with tile size
-	CORE_API static double GetTileSize();
+	// This should become inline at some point, but keeping it in cpp file for now to make iteration faster when experimenting with tile size
+	CORE_API static TScalar GetTileSize();
 	CORE_API static FVector3f GetTileFor(FVector InPosition);
 
-	static double MakeTile(double InValue)
+	static TScalar MakeTile(double InValue)
 	{
-		return FMath::FloorToDouble(InValue / GetTileSize() + 0.5);
+		return static_cast<TScalar>(FMath::FloorToDouble(InValue / GetTileSize() + 0.5));
 	}
 
-	static double MakeQuantizedTile(double InValue, double InQuantization)
+	static TScalar MakeQuantizedTile(double InValue, double InQuantization)
 	{
-		return FMath::FloorToDouble((InValue / GetTileSize()) * InQuantization + 0.5) / InQuantization;
+		return static_cast<TScalar>(FMath::FloorToDouble((InValue / GetTileSize()) * InQuantization + 0.5) / InQuantization);
 	}
 
-	static FVector MakeTile(const FVector& InValue)
+	static VectorType MakeTile(const FVector& InValue)
 	{
-		return FVector(MakeTile(InValue.X), MakeTile(InValue.Y), MakeTile(InValue.Z));
+		return VectorType(MakeTile(InValue.X), MakeTile(InValue.Y), MakeTile(InValue.Z));
 	}
 
-	static FVector MakeQuantizedTile(const FVector& InValue, double InQuantization)
+	static VectorType MakeQuantizedTile(const FVector& InValue, double InQuantization)
 	{
-		return FVector(MakeQuantizedTile(InValue.X, InQuantization), MakeQuantizedTile(InValue.Y, InQuantization), MakeQuantizedTile(InValue.Z, InQuantization));
+		return VectorType(MakeQuantizedTile(InValue.X, InQuantization), MakeQuantizedTile(InValue.Y, InQuantization), MakeQuantizedTile(InValue.Z, InQuantization));
 	}
 
 	CORE_API static FMatrix44f SafeCastMatrix(const FMatrix& Matrix);
@@ -45,65 +49,104 @@ public:
 	CORE_API static FMatrix44f MakeClampedToRelativeWorldMatrix(const FVector Origin, const FMatrix& ToWorld);
 	CORE_API static FMatrix    MakeClampedToRelativeWorldMatrixDouble(const FVector Origin, const FMatrix& ToWorld);
 
-	float GetTile() const { return (float)Tile; }
-	float GetOffset() const { return (float)Offset; }
-	double GetTileAsDouble() const { return Tile; }
-	double GetOffsetAsDouble() const { return Offset; }
-	double GetTileOffset() const { return Tile * GetTileSize(); }
+	template <typename TResult = float>
+	TResult GetTile() const { return static_cast<TResult>(Tile); }
+	template <typename TResult = float>
+	TResult GetOffset() const { return static_cast<TResult>(Offset); }
+	UE_DEPRECATED(5.2, "GetTileAsDouble is deprecated, please use GetTile<double>() instead.")
+	double GetTileAsDouble() const { return GetTile<double>(); }
+	UE_DEPRECATED(5.2, "GetOffsetAsDouble is deprecated, please use GetOffset<double>() instead.")
+	double GetOffsetAsDouble() const { return GetOffset<double>(); }
+	double GetTileOffset() const { return static_cast<double>(Tile) * GetTileSize(); }
 	double GetAbsolute() const { return GetTileOffset() + Offset; }
 
-	FLargeWorldRenderScalar() : Tile(0.0), Offset(0.0) {}
-	FLargeWorldRenderScalar(float InTile, float InOffset) : Tile(InTile), Offset(InOffset) {}
+	TLargeWorldRenderScalar() : Tile(static_cast<TScalar>(0.0)), Offset(static_cast<TScalar>(0.0)) {}
+	TLargeWorldRenderScalar(TScalar InTile, TScalar InOffset) : Tile(static_cast<TScalar>(InTile)), Offset(static_cast<TScalar>(InOffset)) {}
+
+	template<typename TInputScalar = double>
+	explicit TLargeWorldRenderScalar(const TLargeWorldRenderScalar<TInputScalar>& In)
+		: TLargeWorldRenderScalar(In.Tile, In.Offset)
+	{ }
 
-	FLargeWorldRenderScalar(double InAbsolute)
+	TLargeWorldRenderScalar(double InAbsolute)
 	{
 		// Tiles are centered on the origin
 		Tile = MakeTile(InAbsolute);
-		Offset = InAbsolute - GetTileOffset();
+		Offset = static_cast<TScalar>(InAbsolute - GetTileOffset());
 		Validate(InAbsolute);
 	}
 
 private:
 	CORE_API void Validate(double InAbsolute);
 
-	double Tile;
-	double Offset;
+	TScalar Tile;
+	TScalar Offset;
 };
+using FLargeWorldRenderScalar = TLargeWorldRenderScalar<double>;
 
-struct FLargeWorldRenderPosition
+template<typename TScalar>
+struct TLargeWorldRenderPosition
 {
+private:
+	using LWCScalarType = TLargeWorldRenderScalar<TScalar>;
+	using VectorType = UE::Math::TVector<TScalar>;
+
 public:
-	FVector3f GetTile() const { return FVector3f(X.GetTile(), Y.GetTile(), Z.GetTile()); }
-	FVector GetTileOffset() const { return FVector(X.GetTileOffset(), Y.GetTileOffset(), Z.GetTileOffset()); }
-	FVector3f GetOffset() const { return FVector3f(X.GetOffset(), Y.GetOffset(), Z.GetOffset()); }
-	FVector GetAbsolute() const { return FVector(X.GetAbsolute(), Y.GetAbsolute(), Z.GetAbsolute()); }
-
-	explicit FLargeWorldRenderPosition(const FVector3f& InWorldPosition)
-		: X(InWorldPosition.X)
-		, Y(InWorldPosition.Y)
-		, Z(InWorldPosition.Z)
-	{}
+	template<typename TResult = float>
+	UE::Math::TVector<TResult> GetTile() const { return UE::Math::TVector<TResult>(static_cast<TResult>(Tile.X), static_cast<TResult>(Tile.Y), static_cast<TResult>(Tile.Z)); }
 	
-	explicit FLargeWorldRenderPosition(const FVector3d& InWorldPosition)
-		: X(InWorldPosition.X)
-		, Y(InWorldPosition.Y)
-		, Z(InWorldPosition.Z)
-	{}
-
-	explicit FLargeWorldRenderPosition(const FVector4d& InWorldPosition)
-		: X(InWorldPosition.X)
-		, Y(InWorldPosition.Y)
-		, Z(InWorldPosition.Z)
-	{}
-
-	explicit FLargeWorldRenderPosition(const FVector3f& InTilePosition, const FVector3f& InRelativePosition)
-		: X(InTilePosition.X, InRelativePosition.X)
-		, Y(InTilePosition.Y, InRelativePosition.Y)
-		, Z(InTilePosition.Z, InRelativePosition.Z)
-	{}
+	FVector GetTileOffset() const
+	{ 
+		LWCScalarType X(Tile.X, Offset.X);
+		LWCScalarType Y(Tile.Y, Offset.Y);
+		LWCScalarType Z(Tile.Z, Offset.Z);
+		return FVector(X.GetTileOffset(), Y.GetTileOffset(), Z.GetTileOffset());
+	}
+
+	template<typename TResult = float>
+	UE::Math::TVector<TResult> GetOffset() const { return UE::Math::TVector<TResult>(static_cast<TResult>(Offset.X), static_cast<TResult>(Offset.Y), static_cast<TResult>(Offset.Z)); }
+	
+	FVector GetAbsolute() const
+	{
+		LWCScalarType X(Tile.X, Offset.X);
+		LWCScalarType Y(Tile.Y, Offset.Y);
+		LWCScalarType Z(Tile.Z, Offset.Z);
+		return FVector(X.GetAbsolute(), Y.GetAbsolute(), Z.GetAbsolute()); 
+	}
+
+	template<typename TInputScalar = double>
+	explicit TLargeWorldRenderPosition(const UE::Math::TVector<TInputScalar>& InWorldPosition)
+	{
+		LWCScalarType X(InWorldPosition.X);
+		LWCScalarType Y(InWorldPosition.Y);
+		LWCScalarType Z(InWorldPosition.Z);
+		Tile = VectorType(X.template GetTile<TScalar>(), Y.template GetTile<TScalar>(), Z.template GetTile<TScalar>());
+		Offset = VectorType(X.template GetOffset<TScalar>(), Y.template GetOffset<TScalar>(), Z.template GetOffset<TScalar>());
+	}
+
+	template<typename TInputScalar = double>
+	explicit TLargeWorldRenderPosition(const UE::Math::TVector4<TInputScalar>& InWorldPosition)
+		: TLargeWorldRenderPosition(UE::Math::TVector<TInputScalar>(InWorldPosition.X, InWorldPosition.Y, InWorldPosition.Z))
+	{ }
+
+	template<typename TInputScalar = double>
+	explicit TLargeWorldRenderPosition(const UE::Math::TVector<TInputScalar>& InTilePosition, const UE::Math::TVector<TInputScalar>& InRelativePosition)
+	{
+		Tile = VectorType(InTilePosition.X, InTilePosition.Y, InTilePosition.Z);
+		Offset = VectorType(InRelativePosition.X, InRelativePosition.Y, InRelativePosition.Z);
+	}
+
+	template<typename TInputScalar>
+	explicit TLargeWorldRenderPosition(const TLargeWorldRenderPosition<TInputScalar>& In)
+		: TLargeWorldRenderPosition(In.Tile, In.Offset)
+	{ }
+
+	TLargeWorldRenderPosition()
+		: Tile(0, 0, 0), Offset(0, 0, 0)
+	{ }
 
 private:
-	FLargeWorldRenderScalar X;
-	FLargeWorldRenderScalar Y;
-	FLargeWorldRenderScalar Z;
+	VectorType Tile;
+	VectorType Offset;
 };
+using FLargeWorldRenderPosition = TLargeWorldRenderPosition<double>;
diff --git a/Engine/Source/Runtime/Renderer/Private/BasePassRendering.h b/Engine/Source/Runtime/Renderer/Private/BasePassRendering.h
--- a/Engine/Source/Runtime/Renderer/Private/BasePassRendering.h
+++ b/Engine/Source/Runtime/Renderer/Private/BasePassRendering.h
@@ -51,7 +51,7 @@ public:
 	FVector4f LightPositionAndInvRadius;
 	FVector4f LightColorAndFalloffExponent;
 	FVector4f LightDirectionAndShadowMapChannelMask;
-	FVector4f SpotAnglesAndSourceRadiusPacked;
+	FVector4f SpotAnglesAndIdAndSourceRadiusPacked;
 	FVector4f LightTangentAndSpecularScale;
 	FVector4f RectDataAndVirtualShadowMapId;
 };
diff --git a/Engine/Source/Runtime/Renderer/Private/GPUScene.cpp b/Engine/Source/Runtime/Renderer/Private/GPUScene.cpp
--- a/Engine/Source/Runtime/Renderer/Private/GPUScene.cpp
+++ b/Engine/Source/Runtime/Renderer/Private/GPUScene.cpp
@@ -22,6 +22,7 @@
 #include "HAL/LowLevelMemStats.h"
 #include "InstanceUniformShaderParameters.h"
 #include "ShaderPrint.h"
+#include "LightSceneData.h"
 
 #define LOG_INSTANCE_ALLOCATIONS 0
 
@@ -595,6 +596,58 @@ void FGPUScene::EndRender()
 	BufferState = {};
 }
 
+void FGPUScene::UpdateGPULights(FRDGBuilder& GraphBuilder, FScene& Scene)
+{
+	TArray<FLightSceneData, SceneRenderingAllocator> LightData{};
+	LightData.SetNum(Scene.Lights.Num(), false);
+
+	static const auto AllowStaticLightingVar = IConsoleManager::Get().FindTConsoleVariableDataInt(TEXT("r.AllowStaticLighting"));
+	const bool bAllowStaticLighting = (!AllowStaticLightingVar || AllowStaticLightingVar->GetValueOnRenderThread() != 0);
+
+	for (int32 Index = 0; Index < Scene.Lights.Num(); ++Index)
+	{
+		if (Scene.Lights.IsAllocated(Index))
+		{
+			InitLightData(Scene.Lights[Index], bAllowStaticLighting, LightData[Index]);
+		}
+	}
+
+	GraphBuilder.QueueBufferUpload<FLightSceneData>(BufferState.LightDataBuffer, LightData);
+}
+
+void FGPUScene::InitLightData(const FLightSceneInfoCompact& LightInfoCompact, bool bAllowStaticLighting, FLightSceneData& DataOut)
+{
+	const FLightSceneInfo& LightInfo = *LightInfoCompact.LightSceneInfo;
+	const FLightSceneProxy& LightProxy = *LightInfo.Proxy;
+
+	FLightRenderParameters LightParams;
+	LightProxy.GetLightShaderParameters(LightParams);
+
+	if (LightProxy.IsInverseSquared())
+	{
+		LightParams.FalloffExponent = 0;
+	}
+
+	// FLightRenderParameters fields
+	DataOut.WorldPosition = TLargeWorldRenderPosition<float>{ LightParams.WorldPosition };
+	DataOut.InvRadius = LightParams.InvRadius;
+	DataOut.Color = LightParams.Color;
+	DataOut.FalloffExponent = LightParams.FalloffExponent;
+	DataOut.Direction = LightParams.Direction;
+	DataOut.SpecularScale = LightParams.SpecularScale;
+	DataOut.Tangent = LightParams.Tangent;
+	DataOut.SourceRadius = LightParams.SourceRadius;
+	DataOut.SpotAngles = LightParams.SpotAngles;
+	DataOut.SoftSourceRadius = LightParams.SoftSourceRadius;
+	DataOut.SourceLength = LightParams.SourceLength;
+	DataOut.RectLightBarnCosAngle = LightParams.RectLightBarnCosAngle;
+	DataOut.RectLightBarnLength = LightParams.RectLightBarnLength;
+	DataOut.RectLightAtlasUVOffset = LightParams.RectLightAtlasUVOffset;
+	DataOut.RectLightAtlasUVScale = LightParams.RectLightAtlasUVScale;
+	DataOut.RectLightAtlasMaxLevel = LightParams.RectLightAtlasMaxLevel;
+	DataOut.InverseExposureBlend = LightParams.InverseExposureBlend;
+	DataOut.LightTypeAndShadowMapChannelMaskPacked = LightInfo.PackLightTypeAndShadowMapChannelMask(bAllowStaticLighting);
+}
 
 void FGPUScene::UpdateInternal(FRDGBuilder& GraphBuilder, FScene& Scene, FRDGExternalAccessQueue& ExternalAccessQueue)
 {
@@ -678,6 +731,8 @@ void FGPUScene::UpdateInternal(FRDGBuilder& GraphBuilder, FScene& Scene, FRDGExt
 		SCOPE_CYCLE_COUNTER(STAT_UpdateGPUSceneTime);
 
 		UploadGeneral<FUploadDataSourceAdapterScenePrimitives>(GraphBuilder, Scene, ExternalAccessQueue, Adapter);
+
+		UpdateGPULights(GraphBuilder, Scene);
 	}
 
 	UseExternalAccessMode(ExternalAccessQueue, ERHIAccess::SRVMask, ERHIPipeline::All);
@@ -729,11 +784,15 @@ void FGPUScene::UpdateBufferState(FRDGBuilder& GraphBuilder, FScene& Scene, cons
 	const uint32 LightMapDataBufferSize = FMath::RoundUpToPowerOfTwo(FMath::Max(LightmapDataAllocator.GetMaxSize(), InitialBufferSize));
 	BufferState.LightmapDataBuffer = ResizeStructuredBufferIfNeeded(GraphBuilder, LightmapDataBuffer, LightMapDataBufferSize * sizeof(FLightmapSceneShaderData::Data), TEXT("GPUScene.LightmapData"));
 	BufferState.LightMapDataBufferSize = LightMapDataBufferSize;
+	
+	const uint32 LightDataBufferSize = FMath::RoundUpToPowerOfTwo(FMath::Max(Scene.Lights.Num(), InitialBufferSize));
+	BufferState.LightDataBuffer = ResizeStructuredBufferIfNeeded(GraphBuilder, LightDataBuffer, LightDataBufferSize * sizeof(FLightSceneData), TEXT("GPUScene.LightData"));
 
 	ShaderParameters.GPUSceneInstanceSceneData = GraphBuilder.CreateSRV(BufferState.InstanceSceneDataBuffer);
 	ShaderParameters.GPUSceneInstancePayloadData = GraphBuilder.CreateSRV(BufferState.InstancePayloadDataBuffer);
 	ShaderParameters.GPUScenePrimitiveSceneData = GraphBuilder.CreateSRV(BufferState.PrimitiveBuffer);
 	ShaderParameters.GPUSceneLightmapData = GraphBuilder.CreateSRV(BufferState.LightmapDataBuffer);
+	ShaderParameters.GPUSceneLightData = GraphBuilder.CreateSRV(BufferState.LightDataBuffer);
 	ShaderParameters.InstanceDataSOAStride = InstanceSceneDataSOAStride;
 	ShaderParameters.NumScenePrimitives = NumScenePrimitives;
 	ShaderParameters.NumInstances = InstanceSceneDataAllocator.GetMaxSize();
diff --git a/Engine/Source/Runtime/Renderer/Private/GPUScene.h b/Engine/Source/Runtime/Renderer/Private/GPUScene.h
--- a/Engine/Source/Runtime/Renderer/Private/GPUScene.h
+++ b/Engine/Source/Runtime/Renderer/Private/GPUScene.h
@@ -9,10 +9,12 @@
 #include "GrowOnlySpanAllocator.h"
 #include "InstanceCulling/InstanceCullingLoadBalancer.h"
 #include "MeshBatch.h"
+#include "LightSceneData.h"
 
 class FRHICommandList;
 class FScene;
 class FViewInfo;
+class FLightSceneInfoCompact;
 
 UE_DEPRECATED(5.0, "Use GPUScene::AddPrimitiveToUpdate instead.4")
 extern RENDERER_API void AddPrimitiveToUpdateGPU(FScene& Scene, int32 PrimitiveId);
@@ -26,6 +28,7 @@ BEGIN_SHADER_PARAMETER_STRUCT(FGPUSceneResourceParameters, )
 	SHADER_PARAMETER_RDG_BUFFER_SRV(StructuredBuffer<float4>, GPUSceneInstancePayloadData)
 	SHADER_PARAMETER_RDG_BUFFER_SRV(StructuredBuffer<float4>, GPUScenePrimitiveSceneData)
 	SHADER_PARAMETER_RDG_BUFFER_SRV(StructuredBuffer<float4>, GPUSceneLightmapData)
+	SHADER_PARAMETER_RDG_BUFFER_SRV(StructuredBuffer<FGPULight>, GPUSceneLightData)
 	SHADER_PARAMETER(uint32, InstanceDataSOAStride)
 	SHADER_PARAMETER(uint32, GPUSceneFrameNumber)
 	SHADER_PARAMETER(int32, NumInstances)
@@ -174,6 +177,8 @@ struct FGPUSceneBufferState
 	FRDGBuffer* InstanceBVHBuffer = nullptr;
 	FRDGBuffer* LightmapDataBuffer = nullptr;
 	uint32 LightMapDataBufferSize = 0;
+
+	FRDGBuffer* LightDataBuffer = nullptr;
 };
 
 class FGPUScene
@@ -343,6 +348,9 @@ public:
 	TRefCountPtr<FRDGPooledBuffer> LightmapDataBuffer;
 	FRDGAsyncScatterUploadBuffer        LightmapUploadBuffer;
 
+	/** View-independent GPU light data */
+	TRefCountPtr<FRDGPooledBuffer> LightDataBuffer;
+
 	struct FInstanceRange
 	{
 		uint32 InstanceSceneDataOffset;
@@ -403,6 +411,13 @@ private:
 	template<typename FUploadDataSourceAdapter>
 	void UploadGeneral(FRDGBuilder& GraphBuilder, FScene& Scene, FRDGExternalAccessQueue& ExternalAccessQueue, const FUploadDataSourceAdapter& UploadDataSourceAdapter);
 
+	/**
+	 * Upload scene light data to gpu
+	 */
+	void UpdateGPULights(FRDGBuilder& GraphBuilder, FScene& Scene);
+
+	static void InitLightData(const FLightSceneInfoCompact& LightInfoCompact, bool bAllowStaticLighting, FLightSceneData& DataOut);
+
 	void UploadDynamicPrimitiveShaderDataForViewInternal(FRDGBuilder& GraphBuilder, FScene& Scene, FViewInfo& View, FRDGExternalAccessQueue& ExternalAccessQueue, bool bIsShadowView);
 
 	void UpdateInternal(FRDGBuilder& GraphBuilder, FScene& Scene, FRDGExternalAccessQueue& ExternalAccessQueue);
diff --git a/Engine/Source/Runtime/Renderer/Private/LightGridInjection.cpp b/Engine/Source/Runtime/Renderer/Private/LightGridInjection.cpp
--- a/Engine/Source/Runtime/Renderer/Private/LightGridInjection.cpp
+++ b/Engine/Source/Runtime/Renderer/Private/LightGridInjection.cpp
@@ -314,10 +314,13 @@ static void PackLocalLightData(
 	const uint32 RectPackedY = 0;
 	const uint32 RectPackedW = FFloat16(-2.f).Encoded;
 
+	// Pack spot angles in 2x f16
+	uint32 SpotAnglesPacked = PackRG16(-2, 1);
+
 	Out.LightPositionAndInvRadius				= FVector4f(LightTranslatedWorldPosition, 1.0f / FMath::Max(SimpleLight.Radius, KINDA_SMALL_NUMBER));
 	Out.LightColorAndFalloffExponent			= FVector4f((FVector3f)SimpleLight.Color * FLightRenderParameters::GetLightExposureScale(View.GetLastEyeAdaptationExposure(), SimpleLight.InverseExposureBlend), SimpleLight.Exponent);
 	Out.LightDirectionAndShadowMapChannelMask	= FVector4f(FVector3f(1, 0, 0), FMath::AsFloat(ShadowMapChannelMask));
-	Out.SpotAnglesAndSourceRadiusPacked			= FVector4f(-2, 1, FMath::AsFloat(PackedZ), FMath::AsFloat(PackedW));
+	Out.SpotAnglesAndIdAndSourceRadiusPacked	= FVector4f(FMath::AsFloat(SpotAnglesPacked), INDEX_NONE, FMath::AsFloat(PackedZ), FMath::AsFloat(PackedW));
 	Out.LightTangentAndSpecularScale			= FVector4f(1.0f, 0.0f, 0.0f, 1.0f);
 	Out.RectDataAndVirtualShadowMapId			= FVector4f(FMath::AsFloat(RectPackedX), FMath::AsFloat(RectPackedY), -1, FMath::AsFloat(RectPackedW));
 }
@@ -327,6 +330,7 @@ static void PackLocalLightData(
 	const FViewInfo& View,
 	const FLightRenderParameters& LightParameters,
 	const uint32 LightTypeAndShadowMapChannelMaskPacked,
+	const int32 LightSceneId,
 	const int32 VirtualShadowMapId,
 	const float VolumetricScatteringIntensity)
 {
@@ -346,19 +350,21 @@ static void PackLocalLightData(
 	RectPackedW |= uint32(FMath::Clamp(LightParameters.RectLightBarnCosAngle,  0.f, 1.0f) * 0x3FF) << 16;	// 10 bits
 	RectPackedW |= uint32(FMath::Clamp(LightParameters.RectLightAtlasMaxLevel, 0.f, 63.f)) << 26;			//  6 bits
 
+	// Pack spot angles in 2x f16
+	uint32 SpotAnglesPacked = PackRG16(LightParameters.SpotAngles.X, LightParameters.SpotAngles.Y);
+
 	// NOTE: This cast of VirtualShadowMapId to float is not ideal, but bitcast has issues here with INDEX_NONE -> NaN
 	// and 32-bit floats have enough mantissa to cover all reasonable numbers here for now.
 	Out.LightPositionAndInvRadius				= FVector4f(LightTranslatedWorldPosition, LightParameters.InvRadius);
 	Out.LightColorAndFalloffExponent			= FVector4f(LightParameters.Color, LightParameters.FalloffExponent);
 	Out.LightDirectionAndShadowMapChannelMask	= FVector4f(LightParameters.Direction, FMath::AsFloat(LightTypeAndShadowMapChannelMaskPacked));
-	Out.SpotAnglesAndSourceRadiusPacked			= FVector4f(LightParameters.SpotAngles.X, LightParameters.SpotAngles.Y, FMath::AsFloat(PackedZ), FMath::AsFloat(PackedW));
+	Out.SpotAnglesAndIdAndSourceRadiusPacked	= FVector4f(FMath::AsFloat(SpotAnglesPacked), LightSceneId, FMath::AsFloat(PackedZ), FMath::AsFloat(PackedW));
 	Out.LightTangentAndSpecularScale			= FVector4f(LightParameters.Tangent, LightParameters.SpecularScale);
 	Out.RectDataAndVirtualShadowMapId			= FVector4f(FMath::AsFloat(RectPackedX), FMath::AsFloat(RectPackedY), float(VirtualShadowMapId), FMath::AsFloat(RectPackedW));
 
 	checkSlow(int32(Out.RectDataAndVirtualShadowMapId.Z) == VirtualShadowMapId);
 }
 
-
 void FSceneRenderer::ComputeLightGrid(FRDGBuilder& GraphBuilder, bool bCullLightsToGrid, FSortedLightSetSceneInfo& SortedLightSet)
 {
 	RDG_CSV_STAT_EXCLUSIVE_SCOPE(GraphBuilder, ComputeLightGrid);
@@ -483,30 +489,7 @@ void FSceneRenderer::ComputeLightGrid(FRDGBuilder& GraphBuilder, bool bCullLight
 						LightParameters.Color *= LightProxy->GetIndirectLightingScale();
 					}
 
-					int32 ShadowMapChannel = LightProxy->GetShadowMapChannel();
-					int32 DynamicShadowMapChannel = LightSceneInfo->GetDynamicShadowMapChannel();
-
-					if (!bAllowStaticLighting)
-					{
-						ShadowMapChannel = INDEX_NONE;
-					}
-
-					// Static shadowing uses ShadowMapChannel, dynamic shadows are packed into light attenuation using DynamicShadowMapChannel
-					uint32 LightTypeAndShadowMapChannelMaskPacked =
-						(ShadowMapChannel == 0 ? 1 : 0) |
-						(ShadowMapChannel == 1 ? 2 : 0) |
-						(ShadowMapChannel == 2 ? 4 : 0) |
-						(ShadowMapChannel == 3 ? 8 : 0) |
-						(DynamicShadowMapChannel == 0 ? 16 : 0) |
-						(DynamicShadowMapChannel == 1 ? 32 : 0) |
-						(DynamicShadowMapChannel == 2 ? 64 : 0) |
-						(DynamicShadowMapChannel == 3 ? 128 : 0);
-
-					LightTypeAndShadowMapChannelMaskPacked |= LightProxy->GetLightingChannelMask() << 8;
-					// pack light type in this uint32 as well
-					LightTypeAndShadowMapChannelMaskPacked |= SortedLightInfo.SortKey.Fields.LightType << 16;
-					const uint32 CastShadows = LightProxy->CastsDynamicShadow() ? 1 : 0;
-					LightTypeAndShadowMapChannelMaskPacked |= CastShadows << (16 + LightType_NumBits);
+					uint32 LightTypeAndShadowMapChannelMaskPacked = LightSceneInfo->PackLightTypeAndShadowMapChannelMask(bAllowStaticLighting);
 
 					const bool bDynamicShadows = ViewFamily.EngineShowFlags.DynamicShadows && VisibleLightInfos.IsValidIndex(LightSceneInfo->Id);
 					const int32 VirtualShadowMapId = bDynamicShadows ? VisibleLightInfos[LightSceneInfo->Id].GetVirtualShadowMapId( &View ) : INDEX_NONE;
@@ -541,7 +524,7 @@ void FSceneRenderer::ComputeLightGrid(FRDGBuilder& GraphBuilder, bool bCullLight
 							VolumetricScatteringIntensity = 0;
 						}
 
-						PackLocalLightData(LightData, View, LightParameters, LightTypeAndShadowMapChannelMaskPacked, VirtualShadowMapId, VolumetricScatteringIntensity);
+						PackLocalLightData(LightData, View, LightParameters, LightTypeAndShadowMapChannelMaskPacked, LightSceneInfo->Id, VirtualShadowMapId, VolumetricScatteringIntensity);
 
 						const FSphere BoundingSphere = LightProxy->GetBoundingSphere();
 						const float Distance = View.ViewMatrices.GetViewMatrix().TransformPosition(BoundingSphere.Center).Z + BoundingSphere.W;
diff --git a/Engine/Source/Runtime/Renderer/Private/LightSceneInfo.cpp b/Engine/Source/Runtime/Renderer/Private/LightSceneInfo.cpp
--- a/Engine/Source/Runtime/Renderer/Private/LightSceneInfo.cpp
+++ b/Engine/Source/Runtime/Renderer/Private/LightSceneInfo.cpp
@@ -298,6 +298,39 @@ bool FLightSceneInfo::ShouldRecordShadowSubjectsForMobile() const
 	return bShouldRecordShadowSubjectsForMobile;
 }
 
+uint32 FLightSceneInfo::PackLightTypeAndShadowMapChannelMask(bool bAllowStaticLighting) const
+{
+	uint32 Result = 0;
+
+	// Light type and shadow map
+	int32 ShadowMapChannel = Proxy->GetShadowMapChannel();
+	int32 CurrentDynamicShadowMapChannel = GetDynamicShadowMapChannel();
+
+	if (!bAllowStaticLighting)
+	{
+		ShadowMapChannel = INDEX_NONE;
+	}
+
+	// Static shadowing uses ShadowMapChannel, dynamic shadows are packed into light attenuation using DynamicShadowMapChannel
+	Result =
+		(ShadowMapChannel == 0 ? 1 : 0) |
+		(ShadowMapChannel == 1 ? 2 : 0) |
+		(ShadowMapChannel == 2 ? 4 : 0) |
+		(ShadowMapChannel == 3 ? 8 : 0) |
+		(CurrentDynamicShadowMapChannel == 0 ? 16 : 0) |
+		(CurrentDynamicShadowMapChannel == 1 ? 32 : 0) |
+		(CurrentDynamicShadowMapChannel == 2 ? 64 : 0) |
+		(CurrentDynamicShadowMapChannel == 3 ? 128 : 0);
+
+	Result |= Proxy->GetLightingChannelMask() << 8;
+	// pack light type in this uint32 as well
+	Result |= ((uint32)Proxy->GetLightType()) << 16;
+	const uint32 CastShadows = Proxy->CastsDynamicShadow() ? 1 : 0;
+	Result |= CastShadows << (16 + LightType_NumBits);
+
+	return Result;
+}
+
 /** Determines whether two bounding spheres intersect. */
 FORCEINLINE bool AreSpheresNotIntersecting(
 	const VectorRegister& A_XYZ,
diff --git a/Engine/Source/Runtime/Renderer/Private/LightSceneInfo.h b/Engine/Source/Runtime/Renderer/Private/LightSceneInfo.h
--- a/Engine/Source/Runtime/Renderer/Private/LightSceneInfo.h
+++ b/Engine/Source/Runtime/Renderer/Private/LightSceneInfo.h
@@ -346,6 +346,8 @@ public:
 	bool SetupMobileMovableLocalLightShadowParameters(const FViewInfo& View, const TArray<FVisibleLightInfo, SceneRenderingAllocator>& VisibleLightInfos, FMobileMovableLocalLightShadowParameters& MobileMovableLocalLightShadowParameters) const;
 
 	bool ShouldRecordShadowSubjectsForMobile() const;
+
+	uint32 PackLightTypeAndShadowMapChannelMask(bool bAllowStaticLighting) const;
 };
 
 /** Defines how the light is stored in the scene's light octree. */
