Platform.h: Remove lines duplicated through a merge error.
#rnx

diff --git a/Engine/Source/Runtime/Core/Public/HAL/Platform.h b/Engine/Source/Runtime/Core/Public/HAL/Platform.h
--- a/Engine/Source/Runtime/Core/Public/HAL/Platform.h
+++ b/Engine/Source/Runtime/Core/Public/HAL/Platform.h
@@ -498,10 +498,6 @@
 	#define PLATFORM_USE_ANSI_MEMALIGN							0
 #endif
 
-#ifndef PLATFORM_USE_ANSI_POSIX_MALLOC
-	#define PLATFORM_USE_ANSI_POSIX_MALLOC						0
-#endif
-
 #ifndef PLATFORM_IS_ANSI_MALLOC_THREADSAFE
 	#define PLATFORM_IS_ANSI_MALLOC_THREADSAFE					0
 #endif
@@ -514,7 +510,6 @@
 	#define PLATFORM_SUPPORTS_VORBIS_CODEC						1
 #endif
 
-
 #ifndef PLATFORM_USE_MINIMAL_HANG_DETECTION
 	#define PLATFORM_USE_MINIMAL_HANG_DETECTION					0
 #endif
@@ -534,6 +529,7 @@
 #ifndef PLATFORM_HAS_FPlatformVirtualMemoryBlock
 	#define	PLATFORM_HAS_FPlatformVirtualMemoryBlock 1
 #endif
+
 #ifndef PLATFORM_BYPASS_PAK_PRECACHE
 	#define PLATFORM_BYPASS_PAK_PRECACHE 0
 #endif
@@ -542,58 +538,10 @@
 	#define PLATFORM_SUPPORTS_FLIP_TRACKING 0
 #endif
 
-#ifndef PLATFORM_USE_FULL_TASK_GRAPH
-	#define PLATFORM_USE_FULL_TASK_GRAPH						1
-#endif
-
-#ifndef PLATFORM_USE_ANSI_POSIX_MALLOC
-	#define PLATFORM_USE_ANSI_POSIX_MALLOC						0
-#endif
-
-#ifndef PLATFORM_USE_ANSI_MEMALIGN
-	#define PLATFORM_USE_ANSI_MEMALIGN							0
-#endif
-
-#ifndef PLATFORM_USE_ANSI_POSIX_MALLOC
-	#define PLATFORM_USE_ANSI_POSIX_MALLOC						0
-#endif
-
-#ifndef PLATFORM_IS_ANSI_MALLOC_THREADSAFE
-	#define PLATFORM_IS_ANSI_MALLOC_THREADSAFE					0
-#endif
-
-#ifndef PLATFORM_SUPPORTS_OPUS_CODEC
-	#define PLATFORM_SUPPORTS_OPUS_CODEC						1
-#endif
-
-#ifndef PLATFORM_SUPPORTS_VORBIS_CODEC
-	#define PLATFORM_SUPPORTS_VORBIS_CODEC						1
-#endif
-
-#ifndef PLATFORM_USE_MINIMAL_HANG_DETECTION
-	#define PLATFORM_USE_MINIMAL_HANG_DETECTION					0
-#endif
-
-#ifndef PLATFORM_USE_GENERIC_STRING_IMPLEMENTATION
-	#define PLATFORM_USE_GENERIC_STRING_IMPLEMENTATION			1
-#endif
-
 #ifndef PLATFORM_USE_SHOWFLAGS_ALWAYS_BITFIELD
 	#define	PLATFORM_USE_SHOWFLAGS_ALWAYS_BITFIELD				1
 #endif
 
-#ifndef PLATFORM_SUPPORTS_LLM
-	#define PLATFORM_SUPPORTS_LLM								1
-#endif
-
-#ifndef PLATFORM_ALLOW_ALLOCATIONS_IN_FASYNCWRITER_SERIALIZEBUFFERTOARCHIVE
-	#define	PLATFORM_ALLOW_ALLOCATIONS_IN_FASYNCWRITER_SERIALIZEBUFFERTOARCHIVE 1
-#endif
-
-#ifndef PLATFORM_HAS_FPlatformVirtualMemoryBlock
-	#define	PLATFORM_HAS_FPlatformVirtualMemoryBlock 1
-#endif
-
 #ifndef PLATFORM_SUPPORTS_LANDSCAPE_VISUAL_MESH_LOD_STREAMING
 	#define PLATFORM_SUPPORTS_LANDSCAPE_VISUAL_MESH_LOD_STREAMING 0 UE_DEPRECATED_MACRO(5.1, "PLATFORM_SUPPORTS_LANDSCAPE_VISUAL_MESH_LOD_STREAMING has been deprecated and should be replaced with 0.")
 #endif
