Core: FParse: Fixed one of the token overloads to not clamp results to 1024 bytes.

#rnx

diff --git a/Engine/Source/Runtime/Core/Private/Misc/Parse.cpp b/Engine/Source/Runtime/Core/Private/Misc/Parse.cpp
--- a/Engine/Source/Runtime/Core/Private/Misc/Parse.cpp
+++ b/Engine/Source/Runtime/Core/Private/Misc/Parse.cpp
@@ -794,7 +794,7 @@ bool FParse::Token( const TCHAR*& Str, TCHAR* Result, int32 MaxLen, bool UseEsca
 		while( *Str && *Str!=TEXT('"') && (Len+1)<MaxLen )
 		{
 			TCHAR c = *Str++;
-			if( c=='\\' && UseEscape )
+			if( c==TEXT('\\') && UseEscape )
 			{
 				// Get escape.
 				c = *Str++;
@@ -855,12 +855,12 @@ bool FParse::Token( const TCHAR*& Str, TCHAR* Result, int32 MaxLen, bool UseEsca
 		}
 	}
 	Result[Len] = TCHAR('\0');
-	return Len != TCHAR('\0');
+	return Len != 0;
 }
 
 bool FParse::Token( const TCHAR*& Str, FString& Arg, bool UseEscape )
 {
-	Arg.Empty();
+	Arg.Reset();
 
 	// Skip preceeding spaces and tabs.
 	while( FChar::IsWhitespace(*Str) )
@@ -933,20 +933,21 @@ bool FParse::Token( const TCHAR*& Str, FString& Arg, bool UseEscape )
 }
 FString FParse::Token( const TCHAR*& Str, bool UseEscape )
 {
-	TCHAR Buffer[1024];
-	if (FParse::Token(Str, Buffer, UE_ARRAY_COUNT(Buffer), UseEscape))
-	{
-		return Buffer;
-	}
-	else
-	{
-		return TEXT("");
-	}
+	FString Token;
+
+	// Preallocate some memory to avoid constant reallocations.
+	Token.Reserve(1023);
+
+	FParse::Token(Str, Token, UseEscape);
+	
+	Token.Shrink();
+
+	return MoveTemp(Token);
 }
 
 bool FParse::AlnumToken(const TCHAR*& Str, FString& Arg)
 {
-	Arg.Empty();
+	Arg.Reset();
 
 	// Skip preceeding spaces and tabs.
 	while (FChar::IsWhitespace(*Str))
