Modular builds: moved the early memory init for consoles into the appropriate per-platform memory code for better per-platform control.

 - fixes a reentrant GCreateMalloc hang on some console platforms.

#rnx

diff --git a/Engine/Source/Runtime/Core/Private/HAL/UnrealMemory.cpp b/Engine/Source/Runtime/Core/Private/HAL/UnrealMemory.cpp
--- a/Engine/Source/Runtime/Core/Private/HAL/UnrealMemory.cpp
+++ b/Engine/Source/Runtime/Core/Private/HAL/UnrealMemory.cpp
@@ -337,10 +337,6 @@ FConsoleCommandDelegate::CreateStatic(&FMemory::EnablePoisonTests)
 /** Helper function called on first allocation to create and initialize GMalloc */
 static int FMemory_GCreateMalloc_ThreadUnsafe()
 {
-#if PLATFORM_CONSOLE_DYNAMIC_LINK
-	FPlatformMemory::Init();
-#endif
-
 	GMalloc = FPlatformMemory::BaseAllocator();
 
 #if !PLATFORM_MAC
