Submit pull request for missing include.
PR #9673: Fix missing header for FText (Contributed by [REDACTED_PERSON])

diff --git a/Engine/Source/Runtime/Core/Public/Misc/ITransaction.h b/Engine/Source/Runtime/Core/Public/Misc/ITransaction.h
--- a/Engine/Source/Runtime/Core/Public/Misc/ITransaction.h
+++ b/Engine/Source/Runtime/Core/Public/Misc/ITransaction.h
@@ -4,6 +4,7 @@
 
 #include "CoreTypes.h"
 #include "Misc/Guid.h"
+#include "Internationalization/Text.h"
 #include "UObject/UObjectHierarchyFwd.h"
 #include "Change.h"
 
