Add support for -corelimit on unix

#rnx

diff --git a/Engine/Source/Runtime/Core/Private/Unix/UnixPlatformMisc.cpp b/Engine/Source/Runtime/Core/Private/Unix/UnixPlatformMisc.cpp
--- a/Engine/Source/Runtime/Core/Private/Unix/UnixPlatformMisc.cpp
+++ b/Engine/Source/Runtime/Core/Private/Unix/UnixPlatformMisc.cpp
@@ -737,6 +737,12 @@ int32 FUnixPlatformMisc::NumberOfCores()
 			}
 		}
 
+		int32 LimitCount = 32768;
+		if (FCommandLine::IsInitialized() && FParse::Value(FCommandLine::Get(), TEXT("-corelimit="), LimitCount))
+		{
+			NumberOfCores = FMath::Min(NumberOfCores, LimitCount);
+		}
+
 		// never allow it to be less than 1, we are running on something
 		NumberOfCores = FMath::Max(1, NumberOfCores);
 	}
@@ -760,7 +766,13 @@ int32 FUnixPlatformMisc::NumberOfCoresIncludingHyperthreads()
 		}
 		else
 		{
-			return CPU_COUNT(&AvailableCpusMask);
+			NumCoreIds = CPU_COUNT(&AvailableCpusMask);
+		}
+
+		int32 LimitCount = 32768;
+		if (FCommandLine::IsInitialized() && FParse::Value(FCommandLine::Get(), TEXT("-corelimit="), LimitCount))
+		{
+			NumCoreIds = FMath::Min(NumCoreIds, LimitCount);
 		}
 	}
 
