Use IsRunningGame instead of !GIsEditor when filtering tests for Client or Server

#rnx

diff --git a/Engine/Source/Runtime/Core/Private/Misc/AutomationTest.cpp b/Engine/Source/Runtime/Core/Private/Misc/AutomationTest.cpp
--- a/Engine/Source/Runtime/Core/Private/Misc/AutomationTest.cpp
+++ b/Engine/Source/Runtime/Core/Private/Misc/AutomationTest.cpp
@@ -584,8 +584,8 @@ void FAutomationTestFramework::GetValidTestNames( TArray<FAutomationTestInfo>& T
 
 	// Determine required application type (Editor, Game, or Commandlet)
 	const bool bRunningEditor = GIsEditor && !IsRunningCommandlet();
-	const bool bRunningGame = !GIsEditor || IsRunningGame();
-	const bool bRunningServer = !GIsEditor || IsRunningDedicatedServer();
+	const bool bRunningClient = IsRunningGame() || IsRunningClientOnly();
+	const bool bRunningServer = IsRunningGame() || IsRunningDedicatedServer();
 	const bool bRunningCommandlet = IsRunningCommandlet();
 
 	//application flags
@@ -594,7 +594,7 @@ void FAutomationTestFramework::GetValidTestNames( TArray<FAutomationTestInfo>& T
 	{
 		ApplicationSupportFlags |= EAutomationTestFlags::EditorContext;
 	}
-	if ( bRunningGame )
+	if ( bRunningClient )
 	{
 		ApplicationSupportFlags |= EAutomationTestFlags::ClientContext;
 	}
