Add FormatToString to create temporary strings similar to WriteToString, using a printf-style format string.

diff --git a/Engine/Source/Runtime/Core/Public/Misc/StringBuilder.h b/Engine/Source/Runtime/Core/Public/Misc/StringBuilder.h
--- a/Engine/Source/Runtime/Core/Public/Misc/StringBuilder.h
+++ b/Engine/Source/Runtime/Core/Public/Misc/StringBuilder.h
@@ -318,7 +318,7 @@ public:
 		return *this;
 	}
 
-private:
+protected:
 	template <typename SrcEncoding>
 	using TIsCharEncodingCompatibleWithCharType = TIsCharEncodingCompatibleWith<SrcEncoding, CharType>;
 
@@ -505,6 +505,36 @@ template <int32 BufferSize> using WriteToAnsiString = TWriteToString<ANSICHAR, B
 template <int32 BufferSize> using WriteToWideString = TWriteToString<WIDECHAR, BufferSize>;
 template <int32 BufferSize> using WriteToUtf8String = TWriteToString<UTF8CHAR, BufferSize>;
 
+/**
+ * A function-like type that creates a TStringBuilder by formatting its arguments with printf-style format string
+ *
+ * Example:
+ * For void Action(FStringView) -> Action(FormatToString<64>(TEXT("%x, %llx"), Arg1, Arg2));
+ */
+template <typename CharType, int32 BufferSize>
+class TFormatToString : public TStringBuilderWithBuffer<CharType, BufferSize>
+{
+	using TStringBuilderBase<CharType>::TIsCharEncodingCompatibleWithCharType;
+public:
+	template <typename FmtType, typename... ArgTypes,
+		std::enable_if_t<TIsArrayOrRefOfTypeByPredicate<FmtType, TIsCharEncodingCompatibleWithCharType>::Value>* = nullptr>
+	explicit TFormatToString(const FmtType& Format, ArgTypes&&... Args)
+	{
+		TStringBuilderBase<CharType>::Appendf(Format, Forward<ArgTypes>(Args)...);
+	}
+};
+
+template <typename CharType, int32 BufferSize>
+struct TIsContiguousContainer<TFormatToString <CharType, BufferSize>>
+{
+	static constexpr bool Value = true;
+};
+
+template <int32 BufferSize> using FormatToString = TFormatToString<TCHAR, BufferSize>;
+template <int32 BufferSize> using FormatToAnsiString = TFormatToString<ANSICHAR, BufferSize>;
+template <int32 BufferSize> using FormatToWideString = TFormatToString<WIDECHAR, BufferSize>;
+template <int32 BufferSize> using FormatToUtf8String = TFormatToString<UTF8CHAR, BufferSize>;
+
 /**
  * Returns an object that can be used as the output container for algorithms by appending to the builder.
  *
