Implemented FindTConsoleVariableDataBool

diff --git a/Engine/Source/Runtime/Core/Public/HAL/IConsoleManager.h b/Engine/Source/Runtime/Core/Public/HAL/IConsoleManager.h
--- a/Engine/Source/Runtime/Core/Public/HAL/IConsoleManager.h
+++ b/Engine/Source/Runtime/Core/Public/HAL/IConsoleManager.h
@@ -912,6 +912,18 @@ struct CORE_API IConsoleManager
 	 */
 	virtual FString FindConsoleObjectName(const IConsoleObject* Obj) const = 0;
 
+	/**
+	 * Find a typed console variable (faster access to the value, no virtual function call)
+	 * @param Name must not be 0
+	 * @return 0 if the object wasn't found
+	 */
+	TConsoleVariableData<bool>* FindTConsoleVariableDataBool(const TCHAR* Name) const
+	{
+		IConsoleVariable* P = FindConsoleVariable(Name);
+
+		return P ? P->AsVariableBool() : 0;
+	}
+
 	/**
 	 * Find a typed console variable (faster access to the value, no virtual function call)
 	 * @param Name must not be 0
