Csv Profiler: public accessor for a copy of current metadata for embedding inside other profiling/debugging output files.

diff --git a/Engine/Source/Runtime/Core/Private/ProfilingDebugging/CsvProfiler.cpp b/Engine/Source/Runtime/Core/Private/ProfilingDebugging/CsvProfiler.cpp
--- a/Engine/Source/Runtime/Core/Private/ProfilingDebugging/CsvProfiler.cpp
+++ b/Engine/Source/Runtime/Core/Private/ProfilingDebugging/CsvProfiler.cpp
@@ -3449,6 +3449,13 @@ void FCsvProfiler::SetMetadata(const TCHAR* Key, const TCHAR* Value)
 	FCsvProfiler::Get()->SetMetadataInternal(Key, Value, true);
 }
 
+TMap<FString, FString> FCsvProfiler::GetMetadataMapCopy()
+{
+	LLM_SCOPE(ELLMTag::CsvProfiler);
+	FScopeLock Lock(&MetadataCS);
+	return MetadataMap;
+}
+
 void FCsvProfiler::SetMetadataInternal(const TCHAR* Key, const TCHAR* Value, bool bSanitize)
 {
 	// Always gather CSV metadata, even if we're not currently capturing.
diff --git a/Engine/Source/Runtime/Core/Public/ProfilingDebugging/CsvProfiler.h b/Engine/Source/Runtime/Core/Public/ProfilingDebugging/CsvProfiler.h
--- a/Engine/Source/Runtime/Core/Public/ProfilingDebugging/CsvProfiler.h
+++ b/Engine/Source/Runtime/Core/Public/ProfilingDebugging/CsvProfiler.h
@@ -413,6 +413,7 @@ public:
 	CORE_API void SetDeviceProfileName(FString InDeviceProfileName);
 
 	CORE_API FString GetOutputFilename() const { return OutputFilename; }
+	CORE_API TMap<FString, FString> GetMetadataMapCopy();
 
 	CORE_API static bool IsWaitTrackingEnabledOnCurrentThread();
 
