Replaced the OOM in TArray::OnInvalidNum() with a UE_LOG(Fatal).

diff --git a/Engine/Source/Runtime/Core/Private/Containers/Array.cpp b/Engine/Source/Runtime/Core/Private/Containers/Array.cpp
new file mode 100644
--- /dev/null
+++ b/Engine/Source/Runtime/Core/Private/Containers/Array.cpp
@@ -0,0 +1,9 @@
+// Copyright Epic Games, Inc. All Rights Reserved.
+
+#include "Containers/Array.h"
+#include "CoreGlobals.h"
+
+void UE4Array_Private::OnInvalidArrayNum(const TCHAR* ArrayNameSuffix, unsigned long long NewNum)
+{
+	UE_LOG(LogCore, Fatal, TEXT("Trying to resize TArray%s to an invalid size of %llu"), ArrayNameSuffix, NewNum);
+}
diff --git a/Engine/Source/Runtime/Core/Public/Containers/Array.h b/Engine/Source/Runtime/Core/Public/Containers/Array.h
--- a/Engine/Source/Runtime/Core/Public/Containers/Array.h
+++ b/Engine/Source/Runtime/Core/Public/Containers/Array.h
@@ -330,6 +330,8 @@ namespace UE4Array_Private
 	{
 		enum { Value = TArrayElementsAreCompatible<ElementType, TElementType_T<RangeType>>::Value };
 	};
+
+	CORE_API void OnInvalidArrayNum(const TCHAR* ArrayNameSuffix, unsigned long long NewNum);
 }
 
 
@@ -360,7 +362,12 @@ private:
 
 	FORCENOINLINE static void OnInvalidNum(USizeType NewNum)
 	{
-		FPlatformMemory::OnOutOfMemory(NewNum * sizeof(ElementType), alignof(ElementType));
+		const TCHAR* ArrayNameSuffix = TEXT("");
+		if constexpr (sizeof(SizeType) == 8)
+		{
+			ArrayNameSuffix = TEXT("64");
+		}
+		UE4Array_Private::OnInvalidArrayNum(ArrayNameSuffix, (unsigned long long)NewNum);
 	}
 
 public:
