- Bump Niagara version to flush out missing compile hash for data interface

#rnx

diff --git a/Engine/Source/Runtime/Core/Private/UObject/FortniteMainBranchObjectVersion.cpp b/Engine/Source/Runtime/Core/Private/UObject/FortniteMainBranchObjectVersion.cpp
--- a/Engine/Source/Runtime/Core/Private/UObject/FortniteMainBranchObjectVersion.cpp
+++ b/Engine/Source/Runtime/Core/Private/UObject/FortniteMainBranchObjectVersion.cpp
@@ -11,7 +11,7 @@ TMap<FGuid, FGuid> FFortniteMainBranchObjectVersion::GetSystemGuids()
 	SystemGuids.Add(DevGuids.MATERIALSHADERMAP_DERIVEDDATA_VER, FGuid("1DA85DA5733C4A489A1BE117CDC5ACCC"));
 	SystemGuids.Add(DevGuids.NANITE_DERIVEDDATA_VER, FGuid("B62CDCAF66AA45289DDB82E276E7C3E7"));
 	SystemGuids.Add(DevGuids.NIAGARASHADERMAP_DERIVEDDATA_VER, FGuid("2289D5116CF94BC0AFEAEC541468E645"));
-	SystemGuids.Add(DevGuids.Niagara_LatestScriptCompileVersion, FGuid("ECAC41C5265C441F8FF141A1E4FEF577"));
+	SystemGuids.Add(DevGuids.Niagara_LatestScriptCompileVersion, FGuid("1F88A64AFB06BD4EACA3246E75D77770"));
 	SystemGuids.Add(DevGuids.SkeletalMeshDerivedDataVersion, FGuid("ACF593EAAE354FCBB34CF44F5AA1BFD2"));
 	SystemGuids.Add(DevGuids.STATICMESH_DERIVEDDATA_VER, FGuid("18E028FCF1BC469385806AA5F8AE07CC"));
 	return SystemGuids;
