Temporary workaround for _tpause instruction failure on certain CPU architectures that are expected to support while Intel continues to investigate internally.

diff --git a/Engine/Source/Runtime/Core/Private/Windows/WindowsPlatformMisc.cpp b/Engine/Source/Runtime/Core/Private/Windows/WindowsPlatformMisc.cpp
--- a/Engine/Source/Runtime/Core/Private/Windows/WindowsPlatformMisc.cpp
+++ b/Engine/Source/Runtime/Core/Private/Windows/WindowsPlatformMisc.cpp
@@ -3076,8 +3076,12 @@ bool FWindowsPlatformMisc::NeedsNonoptionalCPUFeaturesCheck()
 
 bool FWindowsPlatformMisc::HasTimedPauseCPUFeature()
 {
-	// Check for waitpkg is bit 5
-	return (FCPUIDQueriedData::GetCPUExtendedFeatures2() & (1 << 5)) != 0;
+	// Awaiting investigation from Intel on why _tpause instructions are failing on certain supported architectures
+	// e.g. a consistent crash on a 12700 i7 CPU
+	return false;
+
+	//// Check for waitpkg is bit 5
+	//return (FCPUIDQueriedData::GetCPUExtendedFeatures2() & (1 << 5)) != 0;
 }
 
 int32 FWindowsPlatformMisc::GetCacheLineSize()
