Fix TSAN warning in BeginEventCommon by using proper atomics

#rnx

diff --git a/Engine/Source/Runtime/Core/Public/ProfilingDebugging/CpuProfilerTrace.h b/Engine/Source/Runtime/Core/Public/ProfilingDebugging/CpuProfilerTrace.h
--- a/Engine/Source/Runtime/Core/Public/ProfilingDebugging/CpuProfilerTrace.h
+++ b/Engine/Source/Runtime/Core/Public/ProfilingDebugging/CpuProfilerTrace.h
@@ -3,6 +3,7 @@
 #pragma once
 
 #include "CoreTypes.h"
+#include "HAL/PlatformAtomics.h"
 #include "HAL/PreprocessorHelpers.h"
 #include "Misc/Build.h"
 #include "Trace/Config.h"
@@ -161,9 +162,9 @@ struct FCpuProfilerTrace
 		{
 			if (bEnabled)
 			{
-				if (InOutSpecId == 0)
+				if (FPlatformAtomics::AtomicRead_Relaxed((volatile int32*)&InOutSpecId) == 0)
 				{
-					InOutSpecId = FCpuProfilerTrace::OutputEventType(InEventString, File, Line);
+					FPlatformAtomics::AtomicStore_Relaxed((volatile int32*)&InOutSpecId, FCpuProfilerTrace::OutputEventType(InEventString, File, Line));
 				}
 				OutputBeginEvent(InOutSpecId);
 			}
@@ -173,9 +174,9 @@ struct FCpuProfilerTrace
 		{
 			if (bEnabled)
 			{
-				if (InOutSpecId == 0)
+				if (FPlatformAtomics::AtomicRead_Relaxed((volatile int32*)&InOutSpecId) == 0)
 				{
-					InOutSpecId = FCpuProfilerTrace::OutputEventType(InEventString, File, Line);
+					FPlatformAtomics::AtomicStore_Relaxed((volatile int32*)&InOutSpecId, FCpuProfilerTrace::OutputEventType(InEventString, File, Line));
 				}
 				OutputBeginEvent(InOutSpecId);
 			}
