Removed vestigates of TTypeWrapper, which was deleted in CL# 13634331.

diff --git a/Engine/Source/Runtime/Core/Public/Delegates/DelegateInstancesImpl.h b/Engine/Source/Runtime/Core/Public/Delegates/DelegateInstancesImpl.h
--- a/Engine/Source/Runtime/Core/Public/Delegates/DelegateInstancesImpl.h
+++ b/Engine/Source/Runtime/Core/Public/Delegates/DelegateInstancesImpl.h
@@ -30,14 +30,8 @@ namespace UE::Delegates::Private
 }
 
 template <typename FuncType, typename UserPolicy, typename... VarTypes>
-class TCommonDelegateInstanceState;
-
-template <typename InRetValType, typename... ParamTypes, typename UserPolicy, typename... VarTypes>
-class TCommonDelegateInstanceState<InRetValType(ParamTypes...), UserPolicy, VarTypes...> : IBaseDelegateInstance<InRetValType(ParamTypes...), UserPolicy>
+class TCommonDelegateInstanceState : IBaseDelegateInstance<FuncType, UserPolicy>
 {
-public:
-	using RetValType = InRetValType;
-
 public:
 	explicit TCommonDelegateInstanceState(VarTypes... Vars)
 		: Payload(Vars...)
@@ -66,13 +60,11 @@ protected:
 template <class UserClass, typename FuncType, typename UserPolicy, typename... VarTypes>
 class TBaseUFunctionDelegateInstance;
 
-template <class UserClass, typename WrappedRetValType, typename... ParamTypes, typename UserPolicy, typename... VarTypes>
-class TBaseUFunctionDelegateInstance<UserClass, WrappedRetValType(ParamTypes...), UserPolicy, VarTypes...> : public TCommonDelegateInstanceState<WrappedRetValType(ParamTypes...), UserPolicy, VarTypes...>
+template <class UserClass, typename RetValType, typename... ParamTypes, typename UserPolicy, typename... VarTypes>
+class TBaseUFunctionDelegateInstance<UserClass, RetValType(ParamTypes...), UserPolicy, VarTypes...> : public TCommonDelegateInstanceState<RetValType(ParamTypes...), UserPolicy, VarTypes...>
 {
 private:
-	using Super             = TCommonDelegateInstanceState<WrappedRetValType(ParamTypes...), UserPolicy, VarTypes...>;
-	using RetValType        = typename Super::RetValType;
-	using UnwrappedThisType = TBaseUFunctionDelegateInstance<UserClass, RetValType(ParamTypes...), UserPolicy, VarTypes...>;
+	using Super            = TCommonDelegateInstanceState<RetValType(ParamTypes...), UserPolicy, VarTypes...>;
 	using DelegateBaseType = typename UserPolicy::FDelegateExtras;
 
 	static_assert(UE::Delegates::Private::IsUObjectPtr((UserClass*)nullptr), "You cannot use UFunction delegates with non UObject classes.");
@@ -139,7 +131,7 @@ public:
 
 	void CreateCopy(DelegateBaseType& Base) final
 	{
-		new (Base) UnwrappedThisType(*(UnwrappedThisType*)this);
+		new (Base) TBaseUFunctionDelegateInstance(*this);
 	}
 
 	RetValType Execute(ParamTypes... Params) const final
@@ -180,7 +172,7 @@ public:
 	 */
 	FORCEINLINE static void Create(DelegateBaseType& Base, UserClass* InUserObject, const FName& InFunctionName, VarTypes... Vars)
 	{
-		new (Base) UnwrappedThisType(InUserObject, InFunctionName, Vars...);
+		new (Base) TBaseUFunctionDelegateInstance(InUserObject, InFunctionName, Vars...);
 	}
 
 public:
@@ -205,13 +197,11 @@ public:
 template <bool bConst, class UserClass, ESPMode SPMode, typename FuncType, typename UserPolicy, typename... VarTypes>
 class TBaseSPMethodDelegateInstance;
 
-template <bool bConst, class UserClass, ESPMode SPMode, typename WrappedRetValType, typename... ParamTypes, typename UserPolicy, typename... VarTypes>
-class TBaseSPMethodDelegateInstance<bConst, UserClass, SPMode, WrappedRetValType(ParamTypes...), UserPolicy, VarTypes...> : public TCommonDelegateInstanceState<WrappedRetValType(ParamTypes...), UserPolicy, VarTypes...>
+template <bool bConst, class UserClass, ESPMode SPMode, typename RetValType, typename... ParamTypes, typename UserPolicy, typename... VarTypes>
+class TBaseSPMethodDelegateInstance<bConst, UserClass, SPMode, RetValType(ParamTypes...), UserPolicy, VarTypes...> : public TCommonDelegateInstanceState<RetValType(ParamTypes...), UserPolicy, VarTypes...>
 {
 private:
-	using Super             = TCommonDelegateInstanceState<WrappedRetValType(ParamTypes...), UserPolicy, VarTypes...>;
-	using RetValType        = typename Super::RetValType;
-	using UnwrappedThisType = TBaseSPMethodDelegateInstance<bConst, UserClass, SPMode, RetValType(ParamTypes...), UserPolicy, VarTypes...>;
+	using Super            = TCommonDelegateInstanceState<RetValType(ParamTypes...), UserPolicy, VarTypes...>;
 	using DelegateBaseType = typename UserPolicy::FDelegateExtras;
 
 public:
@@ -274,7 +264,7 @@ public:
 
 	void CreateCopy(DelegateBaseType& Base) final
 	{
-		new (Base) UnwrappedThisType(*(UnwrappedThisType*)this);
+		new (Base) TBaseSPMethodDelegateInstance(*this);
 	}
 
 	RetValType Execute(ParamTypes... Params) const final
@@ -328,7 +318,7 @@ public:
 	 */
 	FORCEINLINE static void Create(DelegateBaseType& Base, const TSharedPtr<UserClass, SPMode>& InUserObjectRef, FMethodPtr InFunc, VarTypes... Vars)
 	{
-		new (Base) UnwrappedThisType(InUserObjectRef, InFunc, Vars...);
+		new (Base) TBaseSPMethodDelegateInstance(InUserObjectRef, InFunc, Vars...);
 	}
 
 	/**
@@ -363,15 +353,13 @@ protected:
 template <bool bConst, class UserClass, typename FuncType, typename UserPolicy, typename... VarTypes>
 class TBaseRawMethodDelegateInstance;
 
-template <bool bConst, class UserClass, typename WrappedRetValType, typename... ParamTypes, typename UserPolicy, typename... VarTypes>
-class TBaseRawMethodDelegateInstance<bConst, UserClass, WrappedRetValType(ParamTypes...), UserPolicy, VarTypes...> : public TCommonDelegateInstanceState<WrappedRetValType(ParamTypes...), UserPolicy, VarTypes...>
+template <bool bConst, class UserClass, typename RetValType, typename... ParamTypes, typename UserPolicy, typename... VarTypes>
+class TBaseRawMethodDelegateInstance<bConst, UserClass, RetValType(ParamTypes...), UserPolicy, VarTypes...> : public TCommonDelegateInstanceState<RetValType(ParamTypes...), UserPolicy, VarTypes...>
 {
 private:
 	static_assert(!UE::Delegates::Private::IsUObjectPtr((UserClass*)nullptr), "You cannot use raw method delegates with UObjects.");
 
-	using Super             = TCommonDelegateInstanceState<WrappedRetValType(ParamTypes...), UserPolicy, VarTypes...>;
-	using RetValType        = typename Super::RetValType;
-	using UnwrappedThisType = TBaseRawMethodDelegateInstance<bConst, UserClass, RetValType(ParamTypes...), UserPolicy, VarTypes...>;
+	using Super            = TCommonDelegateInstanceState<RetValType(ParamTypes...), UserPolicy, VarTypes...>;
 	using DelegateBaseType = typename UserPolicy::FDelegateExtras;
 
 public:
@@ -441,7 +429,7 @@ public:
 
 	void CreateCopy(DelegateBaseType& Base) final
 	{
-		new (Base) UnwrappedThisType(*(UnwrappedThisType*)this);
+		new (Base) TBaseRawMethodDelegateInstance(*this);
 	}
 
 	RetValType Execute(ParamTypes... Params) const final
@@ -486,7 +474,7 @@ public:
 	 */
 	FORCEINLINE static void Create(DelegateBaseType& Base, UserClass* InUserObject, FMethodPtr InFunc, VarTypes... Vars)
 	{
-		new (Base) UnwrappedThisType(InUserObject, InFunc, Vars...);
+		new (Base) TBaseRawMethodDelegateInstance(InUserObject, InFunc, Vars...);
 	}
 
 protected:
@@ -504,13 +492,11 @@ protected:
 template <bool bConst, class UserClass, typename FuncType, typename UserPolicy, typename... VarTypes>
 class TBaseUObjectMethodDelegateInstance;
 
-template <bool bConst, class UserClass, typename WrappedRetValType, typename... ParamTypes, typename UserPolicy, typename... VarTypes>
-class TBaseUObjectMethodDelegateInstance<bConst, UserClass, WrappedRetValType(ParamTypes...), UserPolicy, VarTypes...> : public TCommonDelegateInstanceState<WrappedRetValType(ParamTypes...), UserPolicy, VarTypes...>
+template <bool bConst, class UserClass, typename RetValType, typename... ParamTypes, typename UserPolicy, typename... VarTypes>
+class TBaseUObjectMethodDelegateInstance<bConst, UserClass, RetValType(ParamTypes...), UserPolicy, VarTypes...> : public TCommonDelegateInstanceState<RetValType(ParamTypes...), UserPolicy, VarTypes...>
 {
 private:
-	using Super             = TCommonDelegateInstanceState<WrappedRetValType(ParamTypes...), UserPolicy, VarTypes...>;
-	using RetValType        = typename Super::RetValType;
-	using UnwrappedThisType = TBaseUObjectMethodDelegateInstance<bConst, UserClass, RetValType(ParamTypes...), UserPolicy, VarTypes...>;
+	using Super            = TCommonDelegateInstanceState<RetValType(ParamTypes...), UserPolicy, VarTypes...>;
 	using DelegateBaseType = typename UserPolicy::FDelegateExtras;
 
 	static_assert(UE::Delegates::Private::IsUObjectPtr((UserClass*)nullptr), "You cannot use UObject method delegates with raw pointers.");
@@ -580,7 +566,7 @@ public:
 
 	void CreateCopy(DelegateBaseType& Base) final
 	{
-		new (Base) UnwrappedThisType(*(UnwrappedThisType*)this);
+		new (Base) TBaseUObjectMethodDelegateInstance(*this);
 	}
 
 	RetValType Execute(ParamTypes... Params) const final
@@ -631,7 +617,7 @@ public:
 	 */
 	FORCEINLINE static void Create(DelegateBaseType& Base, UserClass* InUserObject, FMethodPtr InFunc, VarTypes... Vars)
 	{
-		new (Base) UnwrappedThisType(InUserObject, InFunc, Vars...);
+		new (Base) TBaseUObjectMethodDelegateInstance(InUserObject, InFunc, Vars...);
 	}
 
 protected:
@@ -650,13 +636,11 @@ protected:
 template <typename FuncType, typename UserPolicy, typename... VarTypes>
 class TBaseStaticDelegateInstance;
 
-template <typename WrappedRetValType, typename... ParamTypes, typename UserPolicy, typename... VarTypes>
-class TBaseStaticDelegateInstance<WrappedRetValType(ParamTypes...), UserPolicy, VarTypes...> : public TCommonDelegateInstanceState<WrappedRetValType(ParamTypes...), UserPolicy, VarTypes...>
+template <typename RetValType, typename... ParamTypes, typename UserPolicy, typename... VarTypes>
+class TBaseStaticDelegateInstance<RetValType(ParamTypes...), UserPolicy, VarTypes...> : public TCommonDelegateInstanceState<RetValType(ParamTypes...), UserPolicy, VarTypes...>
 {
 private:
-	using Super             = TCommonDelegateInstanceState<WrappedRetValType(ParamTypes...), UserPolicy, VarTypes...>;
-	using RetValType        = typename Super::RetValType;
-	using UnwrappedThisType = TBaseStaticDelegateInstance<RetValType(ParamTypes...), UserPolicy, VarTypes...>;
+	using Super            = TCommonDelegateInstanceState<RetValType(ParamTypes...), UserPolicy, VarTypes...>;
 	using DelegateBaseType = typename UserPolicy::FDelegateExtras;
 
 public:
@@ -718,7 +702,7 @@ public:
 
 	void CreateCopy(DelegateBaseType& Base) final
 	{
-		new (Base) UnwrappedThisType(*(UnwrappedThisType*)this);
+		new (Base) TBaseStaticDelegateInstance(*this);
 	}
 
 	RetValType Execute(ParamTypes... Params) const final
@@ -749,7 +733,7 @@ public:
 	 */
 	FORCEINLINE static void Create(DelegateBaseType& Base, FFuncPtr InFunc, VarTypes... Vars)
 	{
-		new (Base) UnwrappedThisType(InFunc, Vars...);
+		new (Base) TBaseStaticDelegateInstance(InFunc, Vars...);
 	}
 
 private:
@@ -764,15 +748,13 @@ private:
 template <typename FuncType, typename UserPolicy, typename FunctorType, typename... VarTypes>
 class TBaseFunctorDelegateInstance;
 
-template <typename WrappedRetValType, typename... ParamTypes, typename UserPolicy, typename FunctorType, typename... VarTypes>
-class TBaseFunctorDelegateInstance<WrappedRetValType(ParamTypes...), UserPolicy, FunctorType, VarTypes...> : public TCommonDelegateInstanceState<WrappedRetValType(ParamTypes...), UserPolicy, VarTypes...>
+template <typename RetValType, typename... ParamTypes, typename UserPolicy, typename FunctorType, typename... VarTypes>
+class TBaseFunctorDelegateInstance<RetValType(ParamTypes...), UserPolicy, FunctorType, VarTypes...> : public TCommonDelegateInstanceState<RetValType(ParamTypes...), UserPolicy, VarTypes...>
 {
 private:
 	static_assert(TAreTypesEqual<FunctorType, typename TRemoveReference<FunctorType>::Type>::Value, "FunctorType cannot be a reference");
 
-	using Super             = TCommonDelegateInstanceState<WrappedRetValType(ParamTypes...), UserPolicy, VarTypes...>;
-	using RetValType        = typename Super::RetValType;
-	using UnwrappedThisType = TBaseFunctorDelegateInstance<RetValType(ParamTypes...), UserPolicy, FunctorType, VarTypes...>;
+	using Super            = TCommonDelegateInstanceState<RetValType(ParamTypes...), UserPolicy, VarTypes...>;
 	using DelegateBaseType = typename UserPolicy::FDelegateExtras;
 
 public:
@@ -831,7 +813,7 @@ public:
 	// IBaseDelegateInstance interface
 	void CreateCopy(DelegateBaseType& Base) final
 	{
-		new (Base) UnwrappedThisType(*(UnwrappedThisType*)this);
+		new (Base) TBaseFunctorDelegateInstance(*this);
 	}
 
 	RetValType Execute(ParamTypes... Params) const final
@@ -856,11 +838,11 @@ public:
 	 */
 	FORCEINLINE static void Create(DelegateBaseType& Base, const FunctorType& InFunctor, VarTypes... Vars)
 	{
-		new (Base) UnwrappedThisType(InFunctor, Vars...);
+		new (Base) TBaseFunctorDelegateInstance(InFunctor, Vars...);
 	}
 	FORCEINLINE static void Create(DelegateBaseType& Base, FunctorType&& InFunctor, VarTypes... Vars)
 	{
-		new (Base) UnwrappedThisType(MoveTemp(InFunctor), Vars...);
+		new (Base) TBaseFunctorDelegateInstance(MoveTemp(InFunctor), Vars...);
 	}
 
 private:
@@ -877,15 +859,13 @@ private:
 template <typename UserClass, typename FuncType, typename UserPolicy, typename FunctorType, typename... VarTypes>
 class TWeakBaseFunctorDelegateInstance;
 
-template <typename UserClass, typename WrappedRetValType, typename... ParamTypes, typename UserPolicy, typename FunctorType, typename... VarTypes>
-class TWeakBaseFunctorDelegateInstance<UserClass, WrappedRetValType(ParamTypes...), UserPolicy, FunctorType, VarTypes...> : public TCommonDelegateInstanceState<WrappedRetValType(ParamTypes...), UserPolicy, VarTypes...>
+template <typename UserClass, typename RetValType, typename... ParamTypes, typename UserPolicy, typename FunctorType, typename... VarTypes>
+class TWeakBaseFunctorDelegateInstance<UserClass, RetValType(ParamTypes...), UserPolicy, FunctorType, VarTypes...> : public TCommonDelegateInstanceState<RetValType(ParamTypes...), UserPolicy, VarTypes...>
 {
 private:
 	static_assert(TAreTypesEqual<FunctorType, typename TRemoveReference<FunctorType>::Type>::Value, "FunctorType cannot be a reference");
 
-	using Super             = TCommonDelegateInstanceState<WrappedRetValType(ParamTypes...), UserPolicy, VarTypes...>;
-	using RetValType        = typename Super::RetValType;
-	using UnwrappedThisType = TWeakBaseFunctorDelegateInstance<UserClass, RetValType(ParamTypes...), UserPolicy, FunctorType, VarTypes...>;
+	using Super            = TCommonDelegateInstanceState<RetValType(ParamTypes...), UserPolicy, VarTypes...>;
 	using DelegateBaseType = typename UserPolicy::FDelegateExtras;
 
 public:
@@ -949,7 +929,7 @@ public:
 	// IBaseDelegateInstance interface
 	void CreateCopy(DelegateBaseType& Base) final
 	{
-		new (Base) UnwrappedThisType(*(UnwrappedThisType*)this);
+		new (Base) TWeakBaseFunctorDelegateInstance(*this);
 	}
 
 	RetValType Execute(ParamTypes... Params) const final
@@ -977,11 +957,11 @@ public:
 	 */
 	FORCEINLINE static void Create(DelegateBaseType& Base, UserClass* InContextObject, const FunctorType& InFunctor, VarTypes... Vars)
 	{
-		new (Base) UnwrappedThisType(InContextObject, InFunctor, Vars...);
+		new (Base) TWeakBaseFunctorDelegateInstance(InContextObject, InFunctor, Vars...);
 	}
 	FORCEINLINE static void Create(DelegateBaseType& Base, UserClass* InContextObject, FunctorType&& InFunctor, VarTypes... Vars)
 	{
-		new (Base) UnwrappedThisType(InContextObject, MoveTemp(InFunctor), Vars...);
+		new (Base) TWeakBaseFunctorDelegateInstance(InContextObject, MoveTemp(InFunctor), Vars...);
 	}
 
 private:
