LowLevelMemTracker: Whitespace and comments normalization to match Core team coding standard.
#rnx

diff --git a/Engine/Source/Runtime/Core/Private/HAL/LowLevelMemTracker.cpp b/Engine/Source/Runtime/Core/Private/HAL/LowLevelMemTracker.cpp
--- a/Engine/Source/Runtime/Core/Private/HAL/LowLevelMemTracker.cpp
+++ b/Engine/Source/Runtime/Core/Private/HAL/LowLevelMemTracker.cpp
@@ -3,6 +3,7 @@
 #include "HAL/LowLevelMemTracker.h"
 
 #if ENABLE_LOW_LEVEL_MEM_TRACKER
+
 #include "HAL/LowLevelMemStats.h"
 #include "HAL/FileManager.h"
 #include "HAL/IConsoleManager.h"
@@ -19,8 +20,6 @@
 #include "Templates/Atomic.h"
 #include "Trace/Trace.inl"
 
-
-
 UE_TRACE_CHANNEL(MemTagChannel, "Memory overview", true)
 
 UE_TRACE_EVENT_BEGIN(LLM, TagsSpec, NoSync|Important)
@@ -48,7 +47,8 @@ UE_TRACE_EVENT_END()
 	CSV_DEFINE_CATEGORY(LLMPlatform, true);
 #endif
 
-TAutoConsoleVariable<int32> CVarLLMTrackPeaks(TEXT("LLM.TrackPeaks"), 0, TEXT("Track peak memory in each category since process start rather than current frame's value."));
+TAutoConsoleVariable<int32> CVarLLMTrackPeaks(TEXT("LLM.TrackPeaks"), 0,
+	TEXT("Track peak memory in each category since process start rather than current frame's value."));
 
 TAutoConsoleVariable<int32> CVarLLMWriteInterval(
 	TEXT("LLM.LLMWriteInterval"),
@@ -77,7 +77,8 @@ FAutoConsoleCommand LLMSnapshot(
 FAutoConsoleCommand DumpLLM(
 	TEXT("DumpLLM"),
 	TEXT("Logs out the current and peak sizes of all tracked LLM tags"),
-	FConsoleCommandWithWorldArgsAndOutputDeviceDelegate::CreateLambda([](const TArray<FString>& Args, UWorld* InWorld, FOutputDevice& Ar)
+	FConsoleCommandWithWorldArgsAndOutputDeviceDelegate::CreateLambda(
+		[](const TArray<FString>& Args, UWorld* InWorld, FOutputDevice& Ar)
 	{
 		bool bCSV = false;
 		for (const FString& Arg : Args)
@@ -87,7 +88,8 @@ FAutoConsoleCommand DumpLLM(
 				bCSV = true;
 			}
 		}
-		FLowLevelMemTracker::Get().DumpToLog(bCSV ? FLowLevelMemTracker::EDumpFormat::CSV : FLowLevelMemTracker::EDumpFormat::PlainText, &Ar);
+		FLowLevelMemTracker::Get().DumpToLog(
+			bCSV ? FLowLevelMemTracker::EDumpFormat::CSV : FLowLevelMemTracker::EDumpFormat::PlainText, &Ar);
 	}));
 
 void FLowLevelMemTracker::DumpToLog(EDumpFormat DumpFormat, FOutputDevice* OutputDevice)
@@ -161,9 +163,7 @@ void FLowLevelMemTracker::DumpToLog(EDumpFormat DumpFormat, FOutputDevice* Outpu
 
 DECLARE_LLM_MEMORY_STAT(TEXT("LLM Overhead"), STAT_LLMOverheadTotal, STATGROUP_LLMOverhead);
 
-/*
-* LLM Summary stats referenced by ELLMTagNames
-*/
+// LLM Summary stats referenced by ELLMTagNames
 DEFINE_STAT(STAT_EngineSummaryLLM);
 DEFINE_STAT(STAT_ProjectSummaryLLM);
 DEFINE_STAT(STAT_NetworkingSummaryLLM);
@@ -184,9 +184,7 @@ DEFINE_STAT(STAT_NavigationSummaryLLM);
 DEFINE_STAT(STAT_TexturesSummaryLLM);
 DEFINE_STAT(STAT_MediaStreamingSummaryLLM);
 
-/*
- * LLM stats referenced by ELLMTagNames
- */
+// LLM stats referenced by ELLMTagNames
 DECLARE_LLM_MEMORY_STAT(TEXT("Total"), STAT_TotalLLM, STATGROUP_LLMFULL);
 DECLARE_LLM_MEMORY_STAT(TEXT("Untracked"), STAT_UntrackedLLM, STATGROUP_LLMFULL);
 DECLARE_LLM_MEMORY_STAT(TEXT("Total"), STAT_PlatformTotalLLM, STATGROUP_LLMPlatform);
@@ -296,13 +294,10 @@ DECLARE_LLM_MEMORY_STAT(TEXT("MMIO"), STAT_PlatformMMIOLLM, STATGROUP_LLMPlatfor
 DECLARE_LLM_MEMORY_STAT(TEXT("VirtualMemory"), STAT_PlatformVMLLM, STATGROUP_LLMPlatform);
 DECLARE_LLM_MEMORY_STAT(TEXT("CustomName"), STAT_CustomName, STATGROUP_LLMFULL);
 
-namespace UE
-{
-namespace LLMPrivate
+namespace UE::LLMPrivate
 {
-	/**
-	 * FLLMCsvWriter: class for writing out the LLM tag sizes to a csv file every few seconds
-	 */
+
+/** FLLMCsvWriter: class for writing out the LLM tag sizes to a csv file every few seconds. */
 	class FLLMCsvWriter
 	{
 	public:
@@ -312,7 +307,9 @@ namespace LLMPrivate
 		void SetTracker(ELLMTracker InTracker);
 		void Clear();
 
-		void Publish(FLowLevelMemTracker& LLMRef, const FTrackerTagSizeMap& TagSizes, const FTagData* OverrideTrackedTotalTagData, const FTagData* OverrideUntaggedTagData, int64 TrackedTotal, bool bTrackPeaks);
+	void Publish(FLowLevelMemTracker& LLMRef, const FTrackerTagSizeMap& TagSizes,
+		const FTagData* OverrideTrackedTotalTagData, const FTagData* OverrideUntaggedTagData, int64 TrackedTotal,
+		bool bTrackPeaks);
 
 		void OnPreFork();
 
@@ -323,7 +320,9 @@ namespace LLMPrivate
 		bool CreateArchive();
 		bool UpdateColumns(const FTrackerTagSizeMap& TagSizes);
 		void WriteHeader(const FTagData* OverrideTrackedTotalTagData, const FTagData* OverrideUntaggedTagData);
-		void AddRow(FLowLevelMemTracker& LLMRef, const FTrackerTagSizeMap& TagSizes, const FTagData* OverrideTrackedTotalName, const FTagData* OverrideUntaggedName, int64 TrackedTotal, bool bTrackPeaks);
+	void AddRow(FLowLevelMemTracker& LLMRef, const FTrackerTagSizeMap& TagSizes,
+		const FTagData* OverrideTrackedTotalName, const FTagData* OverrideUntaggedName, int64 TrackedTotal,
+		bool bTrackPeaks);
 
 		FConstTagDataArray Columns;
 		TFastPointerLLMSet<const FTagData*> ExistingColumns;
@@ -333,16 +332,16 @@ namespace LLMPrivate
 		ELLMTracker Tracker;
 	};
 
-	/**
-	 * Outputs the LLM tags and sizes to TraceLog events.
-	 */
+/** Outputs the LLM tags and sizes to TraceLog events. */
 	class FLLMTraceWriter
 	{
 	public:
 		FLLMTraceWriter();
 		void SetTracker(ELLMTracker InTracker);
 		void Clear();
-		void Publish(FLowLevelMemTracker& LLMRef, const FTrackerTagSizeMap& TagSizes, const FTagData* OverrideTrackedTotalName, const FTagData* OverrideUntaggedName, int64 TrackedTotal, bool bTrackPeaks);
+	void Publish(FLowLevelMemTracker& LLMRef, const FTrackerTagSizeMap& TagSizes,
+		const FTagData* OverrideTrackedTotalName, const FTagData* OverrideUntaggedName, int64 TrackedTotal,
+		bool bTrackPeaks);
 		static void TraceGenericTags(FLowLevelMemTracker& LLMRef);
 
 	private:
@@ -353,16 +352,16 @@ namespace LLMPrivate
 		bool					bTrackerSpecSent = false;
 	};
 
-	/**
-	 * FLLMCsvWriter: class for writing out LLM stats to the Csv Profiler
-	 */
+/** FLLMCsvWriter: class for writing out LLM stats to the Csv Profiler. */
 	class FLLMCsvProfilerWriter
 	{
 	public:
 		FLLMCsvProfilerWriter();
 		void Clear();
 		void SetTracker(ELLMTracker InTracker);
-		void Publish(FLowLevelMemTracker& LLMRef, const FTrackerTagSizeMap& TagSizes, const FTagData* OverrideTrackedTotalName, const FTagData* OverrideUntaggedName, int64 TrackedTotal, bool bTrackPeaks);
+	void Publish(FLowLevelMemTracker& LLMRef, const FTrackerTagSizeMap& TagSizes,
+		const FTagData* OverrideTrackedTotalName, const FTagData* OverrideUntaggedName, int64 TrackedTotal,
+		bool bTrackPeaks);
 	protected:
 		void RecordTagToCsv(int32 CsvCategoryIndex, const FTagData* TagData, int64 TagSize);
 #if LLM_CSV_PROFILER_WRITER_ENABLED 
@@ -372,7 +371,7 @@ namespace LLMPrivate
 	};
 
 
-	/** per thread state in an LLMTracker */
+/** Per-thread state in an LLMTracker. */
 	class FLLMThreadState
 	{
 	public:
@@ -388,8 +387,10 @@ namespace LLMPrivate
 		void PopAssetTag();
 		const FTagData* GetTopAssetTag();
 #endif
-		void TrackAllocation(const void* Ptr, int64 Size, ELLMTracker Tracker, ELLMAllocType AllocType, const FTagData* TagData, const FTagData* AssetTagData, bool bTrackInMemPro);
-		void TrackFree(const void* Ptr, int64 Size, ELLMTracker Tracker, ELLMAllocType AllocType, const FTagData* TagData, const FTagData* AssetTagData, bool bTrackInMemPro);
+	void TrackAllocation(const void* Ptr, int64 Size, ELLMTracker Tracker, ELLMAllocType AllocType,
+		const FTagData* TagData, const FTagData* AssetTagData, bool bTrackInMemPro);
+	void TrackFree(const void* Ptr, int64 Size, ELLMTracker Tracker, ELLMAllocType AllocType,
+		const FTagData* TagData, const FTagData* AssetTagData, bool bTrackInMemPro);
 		void TrackMoved(const void* Dest, const void* Source, int64 Size, ELLMTracker Tracker, const FTagData* TagData);
 		void IncrTag(const FTagData* Tag, int64 Amount);
 
@@ -412,9 +413,7 @@ namespace LLMPrivate
 		int64 AllocTypeAmounts[(int32)ELLMAllocType::Count];
 	};
 
-	/*
-	 * this is really the main LLM class. It owns the thread state objects.
-	 */
+/** The main LLM implementation class. It owns the thread state objects. */
 	class FLLMTracker
 	{
 	public:
@@ -433,8 +432,10 @@ namespace LLMPrivate
 		void PushAssetTag(const FTagData* Tag);
 		void PopAssetTag();
 #endif
-		void TrackAllocation(const void* Ptr, int64 Size, FName DefaultTag, ELLMTracker Tracker, ELLMAllocType AllocType, bool bTrackInMemPro);
-		void TrackAllocation(const void* Ptr, int64 Size, ELLMTag DefaultTag, ELLMTracker Tracker, ELLMAllocType AllocType, bool bTrackInMemPro);
+	void TrackAllocation(const void* Ptr, int64 Size, FName DefaultTag, ELLMTracker Tracker, ELLMAllocType AllocType,
+		bool bTrackInMemPro);
+	void TrackAllocation(const void* Ptr, int64 Size, ELLMTag DefaultTag, ELLMTracker Tracker, ELLMAllocType AllocType,
+		bool bTrackInMemPro);
 		void TrackFree(const void* Ptr, ELLMTracker Tracker, ELLMAllocType AllocType, bool bTrackInMemPro);
 		void OnAllocMoved(const void* Dest, const void* Source, ELLMTracker Tracker, ELLMAllocType AllocType);
 
@@ -506,7 +507,8 @@ namespace LLMPrivate
 			const FTagData* AssetTag = nullptr;
 #endif
 #if LLM_ENABLED_FULL_TAGS
-			// Even with arbitrary tags we are still partially compressed - the allocation records the tag's index (4 bytes) rather than the full tag pointer (8 bytes) or the tag's name (12 bytes)
+		// Even with arbitrary tags we are still partially compressed - the allocation records the tag's index
+		// (4 bytes) rather than the full tag pointer (8 bytes) or the tag's name (12 bytes).
 			int32 Tag = 0;
 #else
 			ELLMTag Tag = ELLMTag::Untagged;
@@ -539,7 +541,8 @@ namespace LLMPrivate
 
 		FLLMThreadState* GetOrCreateState();
 		FLLMThreadState* GetState();
-		void TrackAllocation(const void* Ptr, int64 Size, const FTagData* ActiveTagData, ELLMTracker Tracker, ELLMAllocType AllocType, FLLMThreadState* State, bool bTrackInMemPro);
+	void TrackAllocation(const void* Ptr, int64 Size, const FTagData* ActiveTagData, ELLMTracker Tracker,
+		ELLMAllocType AllocType, FLLMThreadState* State, bool bTrackInMemPro);
 
 		FLowLevelMemTracker& LLMRef;
 
@@ -571,8 +574,8 @@ namespace LLMPrivate
 	const TCHAR* ToString(ETagReferenceSource ReferenceSource);
 	void SetMemoryStatByFName(FName Name, int64 Amount);
 	void ValidateUniqueName(FStringView UniqueName);
-}
-}
+
+} // UE::LLMPrivate
 
 static FName GetTagName_CustomName()
 {
@@ -706,7 +709,8 @@ bool FLowLevelMemTracker::IsEnabled()
 }
 
 FLowLevelMemTracker* FLowLevelMemTracker::TrackerInstance = nullptr;
-bool FLowLevelMemTracker::bIsDisabled = false; // must start off enabled because allocations happen before the command line enables/disables us
+// LLM must start off enabled because allocations happen before the command line enables/disables us
+bool FLowLevelMemTracker::bIsDisabled = false;
 
 static const TCHAR* InvalidLLMTagName = TEXT("?");
 
@@ -782,9 +786,12 @@ void FLowLevelMemTracker::BootstrapInitialise()
 	}
 
 	BootstrapTagDatas();
-	static_assert((uint8)ELLMTracker::Max == 2, "You added a tracker, without updating FLowLevelMemTracker::BootstrapInitialise (and probably need to update macros)");
-	GetTracker(ELLMTracker::Platform)->SetTotalTags(FindOrAddTagData(ELLMTag::PlatformUntaggedTotal), FindOrAddTagData(ELLMTag::PlatformTrackedTotal));
-	GetTracker(ELLMTracker::Default)->SetTotalTags(FindOrAddTagData(ELLMTag::UntaggedTotal), FindOrAddTagData(ELLMTag::TrackedTotal));
+	static_assert((uint8)ELLMTracker::Max == 2,
+		"You added a tracker, without updating FLowLevelMemTracker::BootstrapInitialise (and probably need to update macros)");
+	GetTracker(ELLMTracker::Platform)->SetTotalTags(FindOrAddTagData(ELLMTag::PlatformUntaggedTotal),
+		FindOrAddTagData(ELLMTag::PlatformTrackedTotal));
+	GetTracker(ELLMTracker::Default)->SetTotalTags(FindOrAddTagData(ELLMTag::UntaggedTotal),
+		FindOrAddTagData(ELLMTag::TrackedTotal));
 
 
 	// calculate program size early on... the platform can call SetProgramSize later if it sees fit
@@ -834,21 +841,23 @@ void FLowLevelMemTracker::UpdateStatsPerFrame(const TCHAR* LogName)
 	{
 		if (bFirstTimeUpdating)
 		{
-			// UpdateStatsPerFrame is usually called from the game thread, but can sometimes be called from the async loading thread, so enter a lock for it
+			// UpdateStatsPerFrame is usually called from the game thread, but can sometimes be called from the
+			// async loading thread, so enter a lock for it
 			FScopeLock UpdateScopeLock(&UpdateLock);
 			if (bFirstTimeUpdating)
 			{
-				// Write the saved overhead value to the stats system; this allows us to see the overhead that is always there even when disabled
-				// (unless the #define completely removes support, of course)
+				// Write the saved overhead value to the stats system; this allows us to see the overhead that is
+				// always there even when disabled (unless the #define completely removes support, of course)
 				bFirstTimeUpdating = false;
-				// Don't call Update since we have cleared the Trackers by this point. But do publish the recorded values
+				// Don't call Update; Trackers no longer exist. But do publish the recorded values.
 				PublishDataPerFrame(LogName);
 			}
 		}
 		return;
 	}
 
-	// UpdateStatsPerFrame is usually called from the game thread, but can sometimes be called from the async loading thread, so enter a lock for it
+	// UpdateStatsPerFrame is usually called from the game thread, but can sometimes be called from the
+	// async loading thread, so enter a lock for it
 	FScopeLock UpdateScopeLock(&UpdateLock);
 	BootstrapInitialise();
 
@@ -873,7 +882,8 @@ void FLowLevelMemTracker::Tick()
 	{
 		return;
 	}
-	// TickInternal is usually called from the game thread, but can sometimes be called from the async loading thread, so enter a lock for it
+	// TickInternal is usually called from the game thread, but can sometimes be called from the async loading thread,
+	// so enter a lock for it.
 	FScopeLock UpdateScopeLock(&UpdateLock);
 	TickInternal();
 }
@@ -884,9 +894,11 @@ void FLowLevelMemTracker::TickInternal()
 
 	if (bFullyInitialised)
 	{
-		// We call tick when not fully initialised to get the overhead when disabled. When not initialised, we have to avoid the portion of the tick that uses tags.
+		// We call tick when not fully initialised to get the overhead when disabled. When not initialised, we have to
+		// avoid the portion of the tick that uses tags.
 
-		// get the platform to update any custom tags; this must be done before the aggregation that occurs in GetTracker()->Update
+		// Get the platform to update any custom tags; this must be done before the aggregation that occurs in
+		// GetTracker()->Update.
 		FPlatformMemory::UpdateCustomLLMTags();
 
 		UpdateTags();
@@ -906,7 +918,8 @@ void FLowLevelMemTracker::TickInternal()
 	int64 FMallocUnused = FMallocPlatformAmount - FMallocAmount;
 	if (FMallocPlatformAmount == 0)
 	{
-		// We do not have instrumentation for this allocator, and so can not calculate how much memory it is using internally. Set unused to 0 for this case.
+		// We do not have instrumentation for this allocator, and so can not calculate how much memory it is using
+		// internally. Set unused to 0 for this case.
 		FMallocUnused = 0;
 	}
 	DefaultTracker.SetTagAmountInUpdate(FindOrAddTagData(ELLMTag::FMallocUnused), FMallocUnused, true);
@@ -915,25 +928,29 @@ void FLowLevelMemTracker::TickInternal()
 	MemoryUsageCurrentOverhead = StaticOverhead + Allocator.GetTotal();
 	PlatformTracker.SetTagAmountInUpdate(FindOrAddTagData(ELLMTag::PlatformOverhead), MemoryUsageCurrentOverhead, true);
 
-	// calculate memory the platform thinks we have allocated, compared to what we have tracked, including the program memory
+	// Compare memory the platform thinks we have allocated to what we have tracked, including the program memory
 	FPlatformMemoryStats PlatformStats = FPlatformMemory::GetStats();
 #if PLATFORM_DESKTOP
-	int64 PlatformProcessMemory = static_cast<int64>(PlatformStats.UsedVirtual);  // virtual is working set + paged out memory
+	// virtual is working set + paged out memory.
+	int64 PlatformProcessMemory = static_cast<int64>(PlatformStats.UsedVirtual);
 #elif PLATFORM_ANDROID || PLATFORM_IOS || UE_SERVER
 	int64 PlatformProcessMemory = static_cast<int64>(PlatformStats.UsedPhysical);
 #else
-	int64 PlatformProcessMemory = static_cast<int64>(PlatformStats.TotalPhysical) - static_cast<int64>(PlatformStats.AvailablePhysical);
+	int64 PlatformProcessMemory = static_cast<int64>(PlatformStats.TotalPhysical) -
+		static_cast<int64>(PlatformStats.AvailablePhysical);
 #endif
 	int64 PlatformTrackedTotal = PlatformTracker.GetTrackedTotal();
 	MemoryUsagePlatformTotalUntracked = FMath::Max<int64>(0, PlatformProcessMemory - PlatformTrackedTotal);
 
 	PlatformTracker.SetTagAmountInUpdate(FindOrAddTagData(ELLMTag::PlatformTotal), PlatformProcessMemory, false);
-	PlatformTracker.SetTagAmountInUpdate(FindOrAddTagData(ELLMTag::PlatformUntracked), MemoryUsagePlatformTotalUntracked, false);
-	PlatformTracker.SetTagAmountInUpdate(FindOrAddTagData(ELLMTag::PlatformOSAvailable), PlatformStats.AvailablePhysical, false);
+	PlatformTracker.SetTagAmountInUpdate(FindOrAddTagData(ELLMTag::PlatformUntracked),
+		MemoryUsagePlatformTotalUntracked, false);
+	PlatformTracker.SetTagAmountInUpdate(FindOrAddTagData(ELLMTag::PlatformOSAvailable),
+		PlatformStats.AvailablePhysical, false);
 
 	int64 TrackedTotal = DefaultTracker.GetTrackedTotal();
-	// remove the MemoryUsageCurrentOverhead from the "Total" for the default LLM as it's not something anyone needs to investigate when finding what to reduce
-	// the platform LLM will have the info 
+	// remove the MemoryUsageCurrentOverhead from the "Total" for the default LLM as it's not something anyone needs
+	// to investigate when finding what to reduce the platform LLM will have the info 
 	int64 DefaultProcessMemory = PlatformProcessMemory - MemoryUsageCurrentOverhead;
 	int64 DefaultUntracked = FMath::Max<int64>(0, DefaultProcessMemory - TrackedTotal);
 	DefaultTracker.SetTagAmountInUpdate(FindOrAddTagData(ELLMTag::Total), DefaultProcessMemory, false);
@@ -1044,7 +1061,8 @@ void FLowLevelMemTracker::PublishDataPerFrame(const TCHAR* LogName)
 			static bool bWarningGiven = false;
 			if (!bWarningGiven)
 			{
-				UE_LOG(LogHAL, Warning, TEXT("Attempted to enable LLM.TrackPeaks, but LLM_ENABLED_TRACK_PEAK_MEMORY is not defined to 1. You will need to enable the define"));
+				UE_LOG(LogHAL, Warning,
+					TEXT("Attempted to enable LLM.TrackPeaks, but LLM_ENABLED_TRACK_PEAK_MEMORY is not defined to 1. You will need to enable the define"));
 				bWarningGiven = true;
 			}
 		}
@@ -1076,19 +1094,20 @@ void FLowLevelMemTracker::PublishDataPerFrame(const TCHAR* LogName)
 
 	if (LogName != nullptr)
 	{
-		FPlatformMisc::LowLevelOutputDebugStringf(TEXT("---> Untracked memory at %s = %.2f mb\n"), LogName, (double)MemoryUsagePlatformTotalUntracked / (1024.0 * 1024.0));
+		FPlatformMisc::LowLevelOutputDebugStringf(TEXT("---> Untracked memory at %s = %.2f mb\n"),
+			LogName, (double)MemoryUsagePlatformTotalUntracked / (1024.0 * 1024.0));
 	}
 }
 
 void FLowLevelMemTracker::InitialiseProgramSize()
 {
-	if (!ProgramSize)
+	if (ProgramSize)
 	{
 		FPlatformMemoryStats Stats = FPlatformMemory::GetStats();
 		ProgramSize = Stats.TotalPhysical - Stats.AvailablePhysical;
 
-		Trackers[static_cast<int32>(ELLMTracker::Platform)]->TrackMemory(ELLMTag::ProgramSizePlatform, ProgramSize, ELLMAllocType::System);
-		Trackers[static_cast<int32>(ELLMTracker::Default)]->TrackMemory(ELLMTag::ProgramSize, ProgramSize, ELLMAllocType::System);
+		GetTracker(ELLMTracker::Platform)->TrackMemory(ELLMTag::ProgramSizePlatform, ProgramSize, ELLMAllocType::System);
+		GetTracker(ELLMTracker::Default)->TrackMemory(ELLMTag::ProgramSize, ProgramSize, ELLMAllocType::System);
 	}
 }
 
@@ -1142,7 +1161,8 @@ void FLowLevelMemTracker::ProcessCommandLine(const TCHAR* CmdLine)
 		LLMCheck(bIsDisabled);
 		if (!bShouldDisable)
 		{
-			UE_LOG(LogInit, Log, TEXT("LLM - Ignoring request to enable LLM; it is not available on the current platform"));
+			UE_LOG(LogInit, Log,
+				TEXT("LLM - Ignoring request to enable LLM; it is not available on the current platform"));
 		}
 		return;
 	}
@@ -1150,7 +1170,8 @@ void FLowLevelMemTracker::ProcessCommandLine(const TCHAR* CmdLine)
 
 	if (bShouldDisable)
 	{
-		// Before we shutdown, update once so we can publish the overhead-when-disabled later during the first call to UpdateStatsPerFrame
+		// Before we shutdown, update once so we can publish the overhead-when-disabled later during the first
+		// call to UpdateStatsPerFrame.
 		if (!bIsDisabled)
 		{
 			Tick();
@@ -1171,7 +1192,8 @@ void FLowLevelMemTracker::ProcessCommandLine(const TCHAR* CmdLine)
 
 	// activate tag sets (we ignore None set, it's always on)
 	FString SetList;
-	static_assert((uint8)ELLMTagSet::Max == 3, "You added a tagset, without updating FLowLevelMemTracker::ProcessCommandLine");
+	static_assert((uint8)ELLMTagSet::Max == 3,
+		"You added a tagset, without updating FLowLevelMemTracker::ProcessCommandLine");
 	if (FParse::Value(CmdLine, TEXT("LLMTAGSETS="), SetList, false /* bShouldStopOnSeparator */))
 	{
 		TArray<FString> Sets;
@@ -1180,10 +1202,13 @@ void FLowLevelMemTracker::ProcessCommandLine(const TCHAR* CmdLine)
 		{
 			if (Set == TEXT("Assets"))
 			{
-#if LLM_ALLOW_ASSETS_TAGS // asset tracking has a per-thread memory overhead, so we have a #define to completely disable it - warn if we don't match
+#if LLM_ALLOW_ASSETS_TAGS
 				ActiveSets[static_cast<int32>(ELLMTagSet::Assets)] = true;
 #else
-				UE_LOG(LogInit, Warning, TEXT("Attempted to use LLM to track assets, but LLM_ALLOW_ASSETS_TAGS is not defined to 1. You will need to enable the define"));
+				// Asset tracking has a per-thread memory overhead, so we have a #define to completely disable it.
+				// Warn if we try to toggle it on at runtime but it is not defined.
+				UE_LOG(LogInit, Warning,
+					TEXT("Attempted to use LLM to track assets, but LLM_ALLOW_ASSETS_TAGS is not defined to 1. You will need to enable the define"));
 #endif
 			}
 			else if (Set == TEXT("AssetClasses"))
@@ -1200,7 +1225,8 @@ void FLowLevelMemTracker::ProcessCommandLine(const TCHAR* CmdLine)
 		CVarLLMTrackPeaks->Set(TrackPeaks);
 	}
 
-	UE_LOG(LogInit, Log, TEXT("LLM enabled CsvWriter: %s TraceWriter: %s"), bCsvWriterEnabled ? TEXT("on") : TEXT("off"), bTraceWriterEnabled ? TEXT("on") : TEXT("off"));
+	UE_LOG(LogInit, Log, TEXT("LLM enabled CsvWriter: %s TraceWriter: %s"),
+		bCsvWriterEnabled ? TEXT("on") : TEXT("off"), bTraceWriterEnabled ? TEXT("on") : TEXT("off"));
 }
 
 // Return the total amount of memory being tracked
@@ -1215,7 +1241,8 @@ uint64 FLowLevelMemTracker::GetTotalTrackedMemory(ELLMTracker Tracker)
 	return static_cast<uint64>(GetTracker(Tracker)->GetTrackedTotal());
 }
 
-void FLowLevelMemTracker::OnLowLevelAlloc(ELLMTracker Tracker, const void* Ptr, uint64 Size, ELLMTag DefaultTag, ELLMAllocType AllocType, bool bTrackInMemPro)
+void FLowLevelMemTracker::OnLowLevelAlloc(ELLMTracker Tracker, const void* Ptr, uint64 Size, ELLMTag DefaultTag,
+	ELLMAllocType AllocType, bool bTrackInMemPro)
 {
 	if (bIsDisabled)
 	{
@@ -1223,10 +1250,12 @@ void FLowLevelMemTracker::OnLowLevelAlloc(ELLMTracker Tracker, const void* Ptr,
 	}
 	BootstrapInitialise();
 
-	GetTracker(Tracker)->TrackAllocation(Ptr, static_cast<int64>(Size), DefaultTag, Tracker, AllocType, bTrackInMemPro);
+	GetTracker(Tracker)->TrackAllocation(Ptr, static_cast<int64>(Size), DefaultTag, Tracker,
+		AllocType, bTrackInMemPro);
 }
 
-void FLowLevelMemTracker::OnLowLevelAlloc(ELLMTracker Tracker, const void* Ptr, uint64 Size, FName DefaultTag, ELLMAllocType AllocType, bool bTrackInMemPro)
+void FLowLevelMemTracker::OnLowLevelAlloc(ELLMTracker Tracker, const void* Ptr, uint64 Size, FName DefaultTag,
+	ELLMAllocType AllocType, bool bTrackInMemPro)
 {
 	if (bIsDisabled)
 	{
@@ -1234,10 +1263,12 @@ void FLowLevelMemTracker::OnLowLevelAlloc(ELLMTracker Tracker, const void* Ptr,
 	}
 	BootstrapInitialise();
 
-	GetTracker(Tracker)->TrackAllocation(Ptr, static_cast<int64>(Size), DefaultTag, Tracker, AllocType, bTrackInMemPro);
+	GetTracker(Tracker)->TrackAllocation(Ptr, static_cast<int64>(Size), DefaultTag, Tracker,
+		AllocType, bTrackInMemPro);
 }
 
-void FLowLevelMemTracker::OnLowLevelFree(ELLMTracker Tracker, const void* Ptr, ELLMAllocType AllocType, bool bTrackInMemPro)
+void FLowLevelMemTracker::OnLowLevelFree(ELLMTracker Tracker, const void* Ptr,
+	ELLMAllocType AllocType, bool bTrackInMemPro)
 {
 	if (bIsDisabled)
 	{
@@ -1251,7 +1282,8 @@ void FLowLevelMemTracker::OnLowLevelFree(ELLMTracker Tracker, const void* Ptr, E
 	}
 }
 
-void FLowLevelMemTracker::OnLowLevelAllocMoved(ELLMTracker Tracker, const void* Dest, const void* Source, ELLMAllocType AllocType)
+void FLowLevelMemTracker::OnLowLevelAllocMoved(ELLMTracker Tracker, const void* Dest, const void* Source,
+	ELLMAllocType AllocType)
 {
 	if (bIsDisabled)
 	{
@@ -1288,7 +1320,8 @@ bool FLowLevelMemTracker::Exec(const TCHAR* Cmd, FOutputDevice& Ar)
 			}
 
 			UpdateStatsPerFrame(TEXT("Before spam"));
-			FPlatformMisc::LowLevelOutputDebugStringf(TEXT("----> Spamming %d allocations, from %d..%d bytes\n"), NumAllocs, MaxSize/2, MaxSize);
+			FPlatformMisc::LowLevelOutputDebugStringf(TEXT("----> Spamming %d allocations, from %d..%d bytes\n"),
+				NumAllocs, MaxSize/2, MaxSize);
 
 			TArray<void*> Spam;
 			Spam.Reserve(NumAllocs);
@@ -1335,7 +1368,8 @@ bool FLowLevelMemTracker::ShouldReduceThreads()
 		return false;
 	}
 	BootstrapInitialise();
-	LLMCheckf(bConfigurationComplete, TEXT("ShouldReduceThreads has been called too early, before we processed the configuration settings required for it."));
+	LLMCheckf(bConfigurationComplete,
+		TEXT("ShouldReduceThreads has been called too early, before we processed the configuration settings required for it."));
 
 	return IsTagSetActive(ELLMTagSet::Assets) || IsTagSetActive(ELLMTagSet::AssetClasses);
 #else
@@ -1348,7 +1382,8 @@ static bool IsAssetTagForAssets(ELLMTagSet Set)
 	return Set == ELLMTagSet::Assets || Set == ELLMTagSet::AssetClasses;
 }
 
-void FLowLevelMemTracker::RegisterCustomTagInternal(int32 Tag, const TCHAR* InDisplayName, FName StatName, FName SummaryStatName, int32 ParentTag)
+void FLowLevelMemTracker::RegisterCustomTagInternal(int32 Tag, const TCHAR* InDisplayName, FName StatName,
+	FName SummaryStatName, int32 ParentTag)
 {
 	using namespace UE::LLMPrivate;
 
@@ -1360,10 +1395,12 @@ void FLowLevelMemTracker::RegisterCustomTagInternal(int32 Tag, const TCHAR* InDi
 	ELLMTag EnumTag = static_cast<ELLMTag>(Tag);
 	FName ParentName = ParentTag >= 0 ? LLMGetTagUniqueName(static_cast<ELLMTag>(ParentTag)) : NAME_None;
 
-	RegisterTagData(LLMGetTagUniqueName(EnumTag), DisplayName, ParentName, StatName, SummaryStatName, true, EnumTag, false, ETagReferenceSource::CustomEnumTag);
+	RegisterTagData(LLMGetTagUniqueName(EnumTag), DisplayName, ParentName, StatName, SummaryStatName, true,
+		EnumTag, false, ETagReferenceSource::CustomEnumTag);
 }
 
-void FLowLevelMemTracker::RegisterPlatformTag(int32 Tag, const TCHAR* Name, FName StatName, FName SummaryStatName, int32 ParentTag)
+void FLowLevelMemTracker::RegisterPlatformTag(int32 Tag, const TCHAR* Name, FName StatName, FName SummaryStatName,
+	int32 ParentTag)
 {
 	MemoryTrace_AnnounceCustomTag(Tag, ParentTag, Name);
 
@@ -1373,11 +1410,13 @@ void FLowLevelMemTracker::RegisterPlatformTag(int32 Tag, const TCHAR* Name, FNam
 	}
 	BootstrapInitialise();
 
-	LLMCheck(Tag >= static_cast<int32>(ELLMTag::PlatformTagStart) && Tag <= static_cast<int32>(ELLMTag::PlatformTagEnd));
+	LLMCheck(Tag >= static_cast<int32>(ELLMTag::PlatformTagStart) &&
+		Tag <= static_cast<int32>(ELLMTag::PlatformTagEnd));
 	RegisterCustomTagInternal(Tag, Name, StatName, SummaryStatName, ParentTag);
 }
 
-void FLowLevelMemTracker::RegisterProjectTag(int32 Tag, const TCHAR* Name, FName StatName, FName SummaryStatName, int32 ParentTag)
+void FLowLevelMemTracker::RegisterProjectTag(int32 Tag, const TCHAR* Name, FName StatName, FName SummaryStatName,
+	int32 ParentTag)
 {
 	MemoryTrace_AnnounceCustomTag(Tag, ParentTag, Name);
 
@@ -1404,8 +1443,8 @@ void FLowLevelMemTracker::BootstrapTagDatas()
 {
 	using namespace UE::LLMPrivate;
 
-	// While bootstrapping we are not allowed to construct any FNames because the FName system may not yet have been constructed.
-	// Construct not-fully-initialized TagDatas for the central list of ELLMTags
+	// While bootstrapping we are not allowed to construct any FNames because the FName system may not yet have been
+	// constructed. Construct not-fully-initialized TagDatas for the central list of ELLMTags.
 	{
 		FWriteScopeLock ScopeLock(TagDataLock);
 		bIsBootstrapping = true;
@@ -1429,9 +1468,9 @@ void FLowLevelMemTracker::BootstrapTagDatas()
 		LLM_ENUM_GENERIC_TAGS(REGISTER_ELLMTAG);
 #undef REGISTER_ELLMTAG
 
-		// The CustomName tag is an adapter for connecting LLM_SCOPE_BYNAME tags with platforms that used the ELLMTag-based reporting.
-		// We want to hide this procedural tags in systems that use FName-based reporting; if it is displayed it confusingly just displays 
-		// a sum of every LLM_SCOPE_BYNAME tag.
+		// The CustomName tag is an adapter for connecting LLM_SCOPE_BYNAME tags with platforms that used the
+		// ELLMTag-based reporting. We want to hide this procedural tag in systems that use FName-based reporting;
+		// if it is displayed it confusingly just displays a sum of every LLM_SCOPE_BYNAME tag.
 		TagDataEnumMap[static_cast<int32>(ELLMTag::CustomName)]->SetIsReportable(false);
 	}
 }
@@ -1444,7 +1483,8 @@ void FLowLevelMemTracker::FinishInitialise()
 	}
 	bFullyInitialised = true;
 	// Make sure that FNames and Malloc have already been initialised, since we will use them during InitialiseTagDatas
-	// We force this by calling LLMGetTagUniqueName, which initializes FNames internally, and will therein trigger FName system construction, which will itself trigger Malloc construction
+	// We force this by calling LLMGetTagUniqueName, which initializes FNames internally, and will therein trigger
+	// FName system construction, which will itself trigger Malloc construction.
 	(void)LLMGetTagUniqueName(ELLMTag::Untagged);
 	InitialiseTagDatas();
 }
@@ -1454,7 +1494,8 @@ void FLowLevelMemTracker::InitialiseTagDatas_SetLLMTagNames()
 	using namespace UE::LLMPrivate;
 
 	TStringBuilder<256> NameBuffer;
-	// Load all the names for the central list of ELLMTags (recording the allocations the FName system makes for the construction of the names)
+	// Load all the names for the central list of ELLMTags (recording the allocations the FName system makes for the
+	// construction of the names).
 #define SET_ELLMTAG_NAMES(Enum,Str,Stat,Group,ParentTag) \
 	{ \
 		ELLMTag EnumTag = ELLMTag::Enum; \
@@ -1467,7 +1508,8 @@ void FLowLevelMemTracker::InitialiseTagDatas_SetLLMTagNames()
 		TagData->SetDisplayName(Str); \
 		TagData->SetStatName(Stat); \
 		TagData->SetSummaryStatName(Group); \
-		TagData->SetParentName(static_cast<int32>(ParentTag) == -1 ? NAME_None : LLMGetTagUniqueName(static_cast<ELLMTag>(ParentTag))); \
+		TagData->SetParentName(static_cast<int32>(ParentTag) == -1 ? \
+			NAME_None : LLMGetTagUniqueName(static_cast<ELLMTag>(ParentTag))); \
 	}
 	LLM_ENUM_GENERIC_TAGS(SET_ELLMTAG_NAMES);
 #undef SET_ELLMTAG_NAMES
@@ -1508,7 +1550,8 @@ void FLowLevelMemTracker::InitialiseTagDatas()
 
 	InitialiseTagDatas_FinishRegister();
 
-	// Construct the remaining startup tags; allocations when constructing these tags are known to consist only of the central list of ELLMTags so we do not need to bootstrap
+	// Construct the remaining startup tags; allocations when constructing these tags are known to consist only of the
+	// central list of ELLMTags so we do not need to bootstrap.
 	{
 		// Construct LLM_DECLARE_TAGs
 		FLLMTagDeclaration* List = FLLMTagDeclaration::GetList();
@@ -1523,7 +1566,8 @@ void FLowLevelMemTracker::InitialiseTagDatas()
 	// now let the platform add any custom tags
 	FPlatformMemory::RegisterCustomLLMTags();
 
-	// All parents in the ELLMTags and the initial modules' list of LLM_DEFINE_TAG must be contained within that same set, so we can FinishConstruct them now, which we do in UpdateTags
+	// All parents in the ELLMTags and the initial modules' list of LLM_DEFINE_TAG must be contained within that same
+	// set, so we can FinishConstruct them now, which we do in UpdateTags.
 	bTagAdded = true;
 	UpdateTags();
 }
@@ -1550,23 +1594,31 @@ void FLowLevelMemTracker::ClearTagDatas()
 void FLowLevelMemTracker::RegisterTagDeclaration(FLLMTagDeclaration& TagDeclaration)
 {
 	TagDeclaration.ConstructUniqueName();
-	RegisterTagData(TagDeclaration.UniqueName, TagDeclaration.DisplayName, TagDeclaration.ParentTagName, TagDeclaration.StatName, TagDeclaration.SummaryStatName, false, ELLMTag::CustomName, false, UE::LLMPrivate::ETagReferenceSource::Declare);
+	RegisterTagData(TagDeclaration.UniqueName, TagDeclaration.DisplayName, TagDeclaration.ParentTagName,
+		TagDeclaration.StatName, TagDeclaration.SummaryStatName, false, ELLMTag::CustomName, false,
+		UE::LLMPrivate::ETagReferenceSource::Declare);
 }
 
-UE::LLMPrivate::FTagData& FLowLevelMemTracker::RegisterTagData(FName Name, FName DisplayName, FName ParentName, FName StatName, FName SummaryStatName, bool bHasEnumTag, ELLMTag EnumTag, bool bIsStatTag, UE::LLMPrivate::ETagReferenceSource ReferenceSource)
+UE::LLMPrivate::FTagData& FLowLevelMemTracker::RegisterTagData(FName Name, FName DisplayName, FName ParentName,
+	FName StatName, FName SummaryStatName, bool bHasEnumTag, ELLMTag EnumTag, bool bIsStatTag,
+	UE::LLMPrivate::ETagReferenceSource ReferenceSource)
 {
-	LLMCheckf(!bIsBootstrapping, TEXT("A tag outside of LLM_ENUM_GENERIC_TAGS was requested from LLM_SCOPE or allocation while bootstrapping the names for LLM_ENUM_GENERIC_TAGS, this is not supported."));
+	LLMCheckf(!bIsBootstrapping,
+		TEXT("A tag outside of LLM_ENUM_GENERIC_TAGS was requested from LLM_SCOPE or allocation while bootstrapping the names for LLM_ENUM_GENERIC_TAGS, this is not supported."));
 
 	using namespace UE::LLMPrivate;
 
-	TStringBuilder<256> NameBuffer; // If this allocates, that is okay. Set it to something small-as-possible-to-avoid-normally-allocating to prevent adding a lot of stack space in the calling LLM_SCOPE code
+	// If this allocates, that is okay. Set it to something small-as-possible-to-avoid-normally-allocating to prevent
+	// adding a lot of stack space in the calling LLM_SCOPE code.
+	TStringBuilder<256> NameBuffer;
 	Name.ToString(NameBuffer);
 
 	if (bHasEnumTag)
 	{
 		ValidateUniqueName(NameBuffer);
-		// EnumTags can specify DisplayName (if they are central or if CustomTag registration provided it); if not, they set DisplayName = UniqueName
-		// Enum tags only specify ParentName explicitly; if no ParentName is provided, they have no parent
+		// EnumTags can specify DisplayName (if they are central or if CustomTag registration provided it);
+		// if not, they set DisplayName = UniqueName.
+		// Enum tags only specify ParentName explicitly; if no ParentName is provided, they have no parent.
 		if (DisplayName.IsNone())
 		{
 			DisplayName = Name;
@@ -1574,9 +1626,10 @@ UE::LLMPrivate::FTagData& FLowLevelMemTracker::RegisterTagData(FName Name, FName
 	}
 	else if (bIsStatTag)
 	{
-		// Stat tag unique names do not have to be validated, because they are never used as parent tags
-		// Stat tag unique names are of the form //GroupName//StatUniqueName///StatDisplayName///<OtherData>. We set LLM UniqueName = <TheEntireString> and LLM DisplayName = StatDisplayName
-		// Stat tags do not specify their parent, and their parent is set to the CustomName aggregator
+		// Stat tag unique names do not have to be validated, because they are never used as parent tags.
+		// Stat tag unique names are of the form //GroupName//StatUniqueName///StatDisplayName///<OtherData>. 
+		// We set LLM UniqueName = <TheEntireString> and LLM DisplayName = StatDisplayName.
+		// Stat tags do not specify their parent, and their parent is set to the CustomName aggregator.
 		LLMCheck(DisplayName.IsNone());
 		LLMCheck(ParentName.IsNone());
 		DisplayName = Name;
@@ -1596,11 +1649,12 @@ UE::LLMPrivate::FTagData& FLowLevelMemTracker::RegisterTagData(FName Name, FName
 	else
 	{
 		ValidateUniqueName(NameBuffer);
-		// Normal defined-by-name tags supply unique names of the form Grandparent/.../Parent/Name
-		// ParentName and  DisplayName can be provided
+		// Normal defined-by-name tags supply unique names of the form Grandparent/.../Parent/Name.
+		// ParentName and  DisplayName can be provided.
 
-		// If both ParentName and /Parent/ are supplied, it is an error if they do not match
-		// All custom name tags have to be children of an ELLMTag. If no parent is set, it defaults to the the proxy parent CustomName
+		// If both ParentName and /Parent/ are supplied, it is an error if they do not match.
+		// All custom name tags have to be children of an ELLMTag. If no parent is set, it defaults to the proxy
+		// parent CustomName.
 		const TCHAR* LeafStart = NameBuffer.ToString();
 		while (true)
 		{
@@ -1611,15 +1665,19 @@ UE::LLMPrivate::FTagData& FLowLevelMemTracker::RegisterTagData(FName Name, FName
 			}
 			LeafStart = NextDivider + 1;
 		}
-		LLMCheckf(LeafStart[0] != '\0', TEXT("Invalid LLM custom name tag '%s'. Tag names must not end with /."), NameBuffer.ToString());
+		LLMCheckf(LeafStart[0] != '\0', TEXT("Invalid LLM custom name tag '%s'. Tag names must not end with /."),
+			NameBuffer.ToString());
 		if (LeafStart != NameBuffer.ToString())
 		{
-			FName ParsedParentName = FName(FStringView(NameBuffer.ToString(), UE_PTRDIFF_TO_INT32(LeafStart - 1 - NameBuffer.ToString())));
+			FName ParsedParentName = FName(FStringView(NameBuffer.ToString(),
+				UE_PTRDIFF_TO_INT32(LeafStart - 1 - NameBuffer.ToString())));
 			if (!ParentName.IsNone() && ParentName != ParsedParentName)
 			{
 				TStringBuilder<128> ParentBuffer;
 				ParentName.ToString(ParentBuffer);
-				LLMCheckf(false, TEXT("Invalid LLM tag: parent supplied in tag declaration is '%s', which does not match the parent parsed from the tag unique name '%s'"), ParentBuffer.ToString(), NameBuffer.ToString());
+				LLMCheckf(false,
+					TEXT("Invalid LLM tag: parent supplied in tag declaration is '%s', which does not match the parent parsed from the tag unique name '%s'"),
+					ParentBuffer.ToString(), NameBuffer.ToString());
 			}
 			ParentName = ParsedParentName;
 		}
@@ -1628,7 +1686,7 @@ UE::LLMPrivate::FTagData& FLowLevelMemTracker::RegisterTagData(FName Name, FName
 			ParentName = GetTagName_CustomName();
 		}
 
-		// Display name is set to the leaf /Name portion of the unique name, and is overridden if DisplayName is provided
+		// Display name is set to the leaf portion of the unique name, and is overridden if DisplayName is provided.
 		if (DisplayName.IsNone())
 		{
 			DisplayName = FName(LeafStart);
@@ -1639,15 +1697,14 @@ UE::LLMPrivate::FTagData& FLowLevelMemTracker::RegisterTagData(FName Name, FName
 	FTagData*& TagDataForName = TagDataNameMap->FindOrAdd(Name, nullptr);
 	if (TagDataForName)
 	{
-		// The tag already exists; this can happen because two formal registrations
-		// have tried to add the same name, but it can also happen due to a race condition
-		// when auto-adding a tag from LLM_SCOPE, if the tag is added by another thread
-		// in between FindOrAddTagData's check of TagDataNameMap->Find and now.
-		// Note that it is not valid for an LLM_SCOPE to be called before a formal registration
-		// (e.g. LLM_DECLARE_TAG). If a formal registration exists for a tag, it must precede its use
-		// in any LLM_SCOPE calls.
+		// The tag already exists; this can happen because two formal registrations have tried to add the same name,
+		// but it can also happen due to a race condition when auto-adding a tag from LLM_SCOPE, if the tag is added
+		// by another thread in between FindOrAddTagData's check of TagDataNameMap->Find and now. Note that it is not
+		// valid for an LLM_SCOPE to be called before a formal registration (e.g. LLM_DECLARE_TAG). If a formal
+		// registration exists for a tag, it must precede its use in any LLM_SCOPE calls.
 		if ((ParentName.IsNone() && TagDataForName->GetParent() != nullptr) ||
-			(!ParentName.IsNone() && (TagDataForName->GetParent() == nullptr || TagDataForName->GetParent()->GetName() != ParentName)) ||
+			(!ParentName.IsNone() &&
+				(TagDataForName->GetParent() == nullptr || TagDataForName->GetParent()->GetName() != ParentName)) ||
 			StatName != TagDataForName->GetStatName() ||
 			SummaryStatName != TagDataForName->GetSummaryStatName() ||
 			bHasEnumTag != TagDataForName->HasEnumTag() ||
@@ -1678,11 +1735,13 @@ UE::LLMPrivate::FTagData& FLowLevelMemTracker::RegisterTagData(FName Name, FName
 	FTagData* TagData;
 	if (ParentName.IsNone() || ParentData)
 	{
-		TagData = Allocator.New<FTagData>(Name, DisplayName, ParentData, StatName, SummaryStatName, bHasEnumTag, EnumTag, ReferenceSource);
+		TagData = Allocator.New<FTagData>(Name, DisplayName, ParentData, StatName, SummaryStatName, bHasEnumTag,
+			EnumTag, ReferenceSource);
 	}
 	else
 	{
-		TagData = Allocator.New<FTagData>(Name, DisplayName, ParentName, StatName, SummaryStatName, bHasEnumTag, EnumTag, ReferenceSource);
+		TagData = Allocator.New<FTagData>(Name, DisplayName, ParentName, StatName, SummaryStatName, bHasEnumTag,
+			EnumTag, ReferenceSource);
 	}
 	TagData->SetIndex(TagDatas->Num());
 	TagDatas->Add(TagData);
@@ -1705,35 +1764,44 @@ UE::LLMPrivate::FTagData& FLowLevelMemTracker::RegisterTagData(FName Name, FName
 	return *TagData;
 }
 
-void FLowLevelMemTracker::ReportDuplicateTagName(UE::LLMPrivate::FTagData* TagData, UE::LLMPrivate::ETagReferenceSource ReferenceSource)
+void FLowLevelMemTracker::ReportDuplicateTagName(UE::LLMPrivate::FTagData* TagData,
+	UE::LLMPrivate::ETagReferenceSource ReferenceSource)
 {
 	using namespace UE::LLMPrivate;
 
 	// We're inside the LLM lock, so do not allow these GetName calls to allocate memory from FMalloc.
-	if (ReferenceSource == ETagReferenceSource::FunctionAPI || ReferenceSource == ETagReferenceSource::Scope || ReferenceSource == ETagReferenceSource::ImplicitParent)
+	if (ReferenceSource == ETagReferenceSource::FunctionAPI || ReferenceSource == ETagReferenceSource::Scope ||
+		ReferenceSource == ETagReferenceSource::ImplicitParent)
 	{
-		LLMCheckf(false, TEXT("LLM Error: Unexpected call to RegisterTagData(%s) from LLM_SCOPE or function call when the tag already exists."),
+		LLMCheckf(false,
+			TEXT("LLM Error: Unexpected call to RegisterTagData(%s) from LLM_SCOPE or function call when the tag already exists."),
 			*WriteToString<FName::StringBufferSize>(TagData->GetName()));
 	}
-	else if (TagData->GetReferenceSource() == ETagReferenceSource::FunctionAPI || TagData->GetReferenceSource() == ETagReferenceSource::Scope)
+	else if (TagData->GetReferenceSource() == ETagReferenceSource::FunctionAPI ||
+		TagData->GetReferenceSource() == ETagReferenceSource::Scope)
 	{
 		LLMCheckf(false, TEXT("LLM Error: Tag %s parsed from %s after it was already used in an LLM_SCOPE or LLM api call."),
 			*WriteToString<FName::StringBufferSize>(TagData->GetName()), ToString(ReferenceSource));
 	}
 	else if (TagData->GetReferenceSource() == ETagReferenceSource::ImplicitParent)
 	{
-		LLMCheckf(false, TEXT("LLM Error: Tag %s parsed from %s after it was already used as an implicit parent in another tag. Add LLM_DEFINE_TAG(%s) in cpp, or move it to the same module as the child tag using it, so that it will be defined before the child tag tries to use it."),
+		LLMCheckf(false,
+			TEXT("LLM Error: Tag %s parsed from %s after it was already used as an implicit parent in another tag. ")
+			TEXT("Add LLM_DEFINE_TAG(% s) in cpp, or move it to the same module as the child tag using it, so that it will be defined before the child tag tries to use it."),
 			*WriteToString<FName::StringBufferSize>(TagData->GetName()), ToString(ReferenceSource),
 			*WriteToString<FName::StringBufferSize>(TagData->GetName()));
 	}
 	else
 	{
-		LLMCheckf(false, TEXT("LLM Error: Multiple occurrences of a unique tag name %s in ELLMTag or LLM_DEFINE_TAG. First occurrence: %s. Second occurrence: %s."),
-			*WriteToString<FName::StringBufferSize>(TagData->GetName()), ToString(TagData->GetReferenceSource()), ToString(ReferenceSource));
+		LLMCheckf(false, TEXT("LLM Error: Multiple occurrences of a unique tag name %s in ELLMTag or LLM_DEFINE_TAG. ")
+			TEXT("First occurrence : % s.Second occurrence : % s."),
+			*WriteToString<FName::StringBufferSize>(TagData->GetName()), ToString(TagData->GetReferenceSource()),
+			ToString(ReferenceSource));
 	}
 }
 
-void FLowLevelMemTracker::FinishConstruct(UE::LLMPrivate::FTagData* TagData, UE::LLMPrivate::ETagReferenceSource ReferenceSource)
+void FLowLevelMemTracker::FinishConstruct(UE::LLMPrivate::FTagData* TagData,
+	UE::LLMPrivate::ETagReferenceSource ReferenceSource)
 {
 	using namespace UE::LLMPrivate;
 	// Caller is responsible for holding a ReadLock (NOT a WriteLock) on TagDataLock
@@ -1761,13 +1829,14 @@ void FLowLevelMemTracker::FinishConstruct(UE::LLMPrivate::FTagData* TagData, UE:
 			FTagData** ParentDataPtr = TagDataNameMap->Find(ParentName);
 			if (!ParentDataPtr)
 			{
-				// We have to drop the ReadLock to call RegisterTagData, which takes a WriteLock
+				// We have to drop the ReadLock to call RegisterTagData, which takes a WriteLock.
 				TagDataLock.ReadUnlock();
-				FTagData* ParentTagData = &RegisterTagData(ParentName, NAME_None, NAME_None, NAME_None, NAME_None, false, ELLMTag::CustomName, false, ETagReferenceSource::ImplicitParent);
+				FTagData* ParentTagData = &RegisterTagData(ParentName, NAME_None, NAME_None, NAME_None, NAME_None,
+					false, ELLMTag::CustomName, false, ETagReferenceSource::ImplicitParent);
 				TagDataLock.ReadLock();
 				if (TagData->IsFinishConstructed())
 				{
-					// Another thread got in and finished construction while we were outside of the lock
+					// Another thread got in and finished construction while we were outside of the lock.
 					return;
 				}
 				ParentDataPtr = TagDataNameMap->Find(ParentName);
@@ -1782,7 +1851,8 @@ void FLowLevelMemTracker::FinishConstruct(UE::LLMPrivate::FTagData* TagData, UE:
 	FTagData* ParentData = const_cast<FTagData*>(TagData->GetParent());
 	if (ParentData)
 	{
-		// Make sure the parent chain is FinishConstructed as well, since GetContainingEnum or GetDisplayPath will be called and walk up the parent chain
+		// Make sure the parent chain is FinishConstructed as well, since GetContainingEnum or GetDisplayPath will be
+		// called and walk up the parent chain.
 		FinishConstruct(ParentData, ReferenceSource);
 	}
 }
@@ -1825,7 +1895,8 @@ bool FLowLevelMemTracker::FindTagByName( const TCHAR* Name, uint64& OutTag ) con
 	{
 		return false;
 	}
-	LLMCheck(bFullyInitialised); // Cannot call BootstrapInitialise and FinishInitialise without shenanigans because this function is const
+	// Cannot call BootstrapInitialise and FinishInitialise without shenanigans because this function is const.
+	LLMCheck(bFullyInitialised);
 
 	if (Name != nullptr)
 	{
@@ -1866,7 +1937,8 @@ const TCHAR* FLowLevelMemTracker::FindTagName(uint64 Tag) const
 	{
 		return nullptr;
 	}
-	LLMCheck(bInitialisedTracking); // Cannot call BootstrapInitialise without shenanigans because this function is const
+	// Cannot call BootstrapInitialise without shenanigans because this function is const.
+	LLMCheck(bInitialisedTracking);
 
 	static TMap<uint64, FString> FoundTags;
 	FString* Cached = FoundTags.Find(Tag);
@@ -1894,7 +1966,8 @@ FName FLowLevelMemTracker::FindTagDisplayName(uint64 Tag) const
 	{
 		return NAME_None;
 	}
-	LLMCheck(bInitialisedTracking); // Cannot call BootstrapInitialise without shenanigans because this function is const
+	// Cannot call BootstrapInitialise without shenanigans because this function is const.
+	LLMCheck(bInitialisedTracking);
 
 	int32 Index = static_cast<int32>(Tag);
 	if (0 <= Index && Index < LLM_CUSTOM_TAG_START)
@@ -1935,7 +2008,8 @@ int64 FLowLevelMemTracker::GetTagAmountForTracker(ELLMTracker Tracker, ELLMTag T
 	return GetTagAmountForTracker(Tracker, FindTagData(Tag), bPeakAmount);
 }
 
-int64 FLowLevelMemTracker::GetTagAmountForTracker(ELLMTracker Tracker, const UE::LLMPrivate::FTagData* TagData, bool bPeakAmount/* = false*/)
+int64 FLowLevelMemTracker::GetTagAmountForTracker(ELLMTracker Tracker, const UE::LLMPrivate::FTagData* TagData,
+	bool bPeakAmount/* = false*/)
 {
 	using namespace UE::LLMPrivate;
 
@@ -2014,7 +2088,9 @@ uint64 FLowLevelMemTracker::DumpTag( ELLMTracker Tracker, const char* FileName,
 	const FTagData* TagData = GetActiveTagData(Tracker);
 	if (TagData)
 	{
-		FPlatformMisc::LowLevelOutputDebugStringf(TEXT("LLM TAG: %s (%lld) @ %s:%d\n"), *TagData->GetDisplayName().ToString(), TagData->GetContainingEnum(), FileName ? ANSI_TO_TCHAR(FileName) : TEXT("?"), LineNumber);
+		FPlatformMisc::LowLevelOutputDebugStringf(TEXT("LLM TAG: %s (%lld) @ %s:%d\n"),
+			*TagData->GetDisplayName().ToString(), TagData->GetContainingEnum(),
+			FileName ? ANSI_TO_TCHAR(FileName) : TEXT("?"), LineNumber);
 		return static_cast<uint64>(TagData->GetContainingEnum());
 	}
 	else
@@ -2036,7 +2112,8 @@ void FLowLevelMemTracker::PublishDataSingleFrame()
 	}	
 }
 
-const UE::LLMPrivate::FTagData* FLowLevelMemTracker::FindOrAddTagData(ELLMTag EnumTag, UE::LLMPrivate::ETagReferenceSource ReferenceSource)
+const UE::LLMPrivate::FTagData* FLowLevelMemTracker::FindOrAddTagData(ELLMTag EnumTag,
+	UE::LLMPrivate::ETagReferenceSource ReferenceSource)
 {
 	using namespace UE::LLMPrivate;
 
@@ -2053,11 +2130,12 @@ const UE::LLMPrivate::FTagData* FLowLevelMemTracker::FindOrAddTagData(ELLMTag En
 		}
 	}
 
-	// If we have not initialised tags yet, we have to initialise now to potentially create the custom ELLMTag that we are trying to read
+	// If we have not initialized tags yet initialize now to potentially create the custom ELLMTag we are reading.
 	if (!bFullyInitialised)
 	{
 		FinishInitialise();
-		// Reeneter this function so that we retry the find above; note we avoid infinite recursion because bFullyInitialised is now true
+		// Reeneter this function so that we retry the find above; we avoid infinite recursion because
+		// bFullyInitialised is now true.
 		return FindOrAddTagData(EnumTag, ReferenceSource);
 	}
 	LLMCheckf(!bIsBootstrapping, TEXT("LLM Error: Invalid use of custom ELLMTag when initialising tags."));
@@ -2065,14 +2143,16 @@ const UE::LLMPrivate::FTagData* FLowLevelMemTracker::FindOrAddTagData(ELLMTag En
 	// Add the new Tag
 	FName TagName = LLMGetTagUniqueName(EnumTag);
 	{
-		FTagData* TagData = &RegisterTagData(TagName, NAME_None, NAME_None, NAME_None, NAME_None, true, EnumTag, false, ReferenceSource);
+		FTagData* TagData = &RegisterTagData(TagName, NAME_None, NAME_None, NAME_None, NAME_None, true, EnumTag,
+			false, ReferenceSource);
 		FReadScopeLock ScopeLock(TagDataLock);
 		FinishConstruct(TagData, ReferenceSource);
 		return TagData;
 	}
 }
 
-const UE::LLMPrivate::FTagData* FLowLevelMemTracker::FindOrAddTagData(FName TagName, bool bIsStatTag, UE::LLMPrivate::ETagReferenceSource ReferenceSource)
+const UE::LLMPrivate::FTagData* FLowLevelMemTracker::FindOrAddTagData(FName TagName, bool bIsStatTag,
+	UE::LLMPrivate::ETagReferenceSource ReferenceSource)
 {
 	using namespace UE::LLMPrivate;
 
@@ -2087,18 +2167,20 @@ const UE::LLMPrivate::FTagData* FLowLevelMemTracker::FindOrAddTagData(FName TagN
 		}
 	}
 
-	// If we have not initialised tags yet, we have to initialise now to potentially create the TagName that we are trying to read
+	// If we have not initialized tags yet initialize now to potentially create the custom ELLMTag we are reading.
 	if (!bFullyInitialised)
 	{
 		FinishInitialise();
-		// Reeneter this function so that we retry the find above; note we avoid infinite recursion because bFullyInitialised is now true
+		// Reeneter this function so that we retry the find above; note we avoid infinite recursion because
+		// bFullyInitialised is now true.
 		return FindOrAddTagData(TagName, bIsStatTag, ReferenceSource);
 	}
 	LLMCheckf(!bIsBootstrapping, TEXT("LLM Error: Invalid use of FName tag when initialising tags."));
 
 	// Add the new Tag
 	FName StatName = bIsStatTag ? TagName : NAME_None;
-	FTagData* TagData = &RegisterTagData(TagName, NAME_None, NAME_None, StatName, NAME_None, false, ELLMTag::CustomName, bIsStatTag, ReferenceSource);
+	FTagData* TagData = &RegisterTagData(TagName, NAME_None, NAME_None, StatName, NAME_None, false,
+		ELLMTag::CustomName, bIsStatTag, ReferenceSource);
 	{
 		FReadScopeLock ScopeLock(TagDataLock);
 		FinishConstruct(TagData, ReferenceSource);
@@ -2106,7 +2188,8 @@ const UE::LLMPrivate::FTagData* FLowLevelMemTracker::FindOrAddTagData(FName TagN
 	}
 }
 
-const UE::LLMPrivate::FTagData* FLowLevelMemTracker::FindTagData(ELLMTag EnumTag, UE::LLMPrivate::ETagReferenceSource ReferenceSource)
+const UE::LLMPrivate::FTagData* FLowLevelMemTracker::FindTagData(ELLMTag EnumTag,
+	UE::LLMPrivate::ETagReferenceSource ReferenceSource)
 {
 	using namespace UE::LLMPrivate;
 
@@ -2126,7 +2209,8 @@ const UE::LLMPrivate::FTagData* FLowLevelMemTracker::FindTagData(ELLMTag EnumTag
 	}
 }
 
-const UE::LLMPrivate::FTagData* FLowLevelMemTracker::FindTagData(FName TagName, UE::LLMPrivate::ETagReferenceSource ReferenceSource)
+const UE::LLMPrivate::FTagData* FLowLevelMemTracker::FindTagData(FName TagName,
+	UE::LLMPrivate::ETagReferenceSource ReferenceSource)
 {
 	using namespace UE::LLMPrivate;
 
@@ -2149,7 +2233,8 @@ void FLLMScope::Init(ELLMTag TagEnum, bool bInIsStatTag, ELLMTagSet InTagSet, EL
 {
 	LLMCheck(!bInIsStatTag && InTagSet == ELLMTagSet::None);
 	FLowLevelMemTracker& LLMRef = FLowLevelMemTracker::Get();
-	// We have to check bIsDisabled again after calling Get, because the constructor is called from Get, and will set bIsDisabled=false if the platform doesn't support it
+	// We have to check bIsDisabled again after calling Get, because the constructor is called from Get, and will set 
+	// bIsDisabled=false if the platform doesn't support it.
 	if (FLowLevelMemTracker::bIsDisabled)
 	{
 		bEnabled = false;
@@ -2168,7 +2253,8 @@ void FLLMScope::Init(ELLMTag TagEnum, bool bInIsStatTag, ELLMTagSet InTagSet, EL
 void FLLMScope::Init(FName TagName, bool bInIsStatTag, ELLMTagSet InTagSet, ELLMTracker InTracker)
 {
 	FLowLevelMemTracker& LLMRef = FLowLevelMemTracker::Get();
-	// We have to check bIsDisabled again after calling Get, because the constructor is called from Get, and will set bIsDisabled=false if the platform doesn't support it
+	// We have to check bIsDisabled again after calling Get, because the constructor is called from Get, and will set
+	// bIsDisabled=false if the platform doesn't support it.
 	if (FLowLevelMemTracker::bIsDisabled)
 	{
 		bEnabled = false;
@@ -2197,11 +2283,13 @@ void FLLMScope::Init(FName TagName, bool bInIsStatTag, ELLMTagSet InTagSet, ELLM
 	}
 }
 
-void FLLMScope::Init(const UE::LLMPrivate::FTagData* TagData, bool bInIsStatTag, ELLMTagSet InTagSet, ELLMTracker InTracker)
+void FLLMScope::Init(const UE::LLMPrivate::FTagData* TagData, bool bInIsStatTag, ELLMTagSet InTagSet,
+	ELLMTracker InTracker)
 {
 	LLMCheck(!bInIsStatTag && InTagSet == ELLMTagSet::None);
 	FLowLevelMemTracker& LLMRef = FLowLevelMemTracker::Get();
-	// We have to check bIsDisabled again after calling Get, because the constructor is called from Get, and will set bIsDisabled=false if the platform doesn't support it
+	// We have to check bIsDisabled again after calling Get, because the constructor is called from Get, and will set
+	// bIsDisabled=false if the platform doesn't support it.
 	if (FLowLevelMemTracker::bIsDisabled)
 	{
 		bEnabled = false;
@@ -2232,7 +2320,8 @@ void FLLMScope::Destruct()
 	}
 }
 
-FLLMPauseScope::FLLMPauseScope(FName TagName, bool bIsStatTag, uint64 Amount, ELLMTracker TrackerToPause, ELLMAllocType InAllocType)
+FLLMPauseScope::FLLMPauseScope(FName TagName, bool bIsStatTag, uint64 Amount, ELLMTracker TrackerToPause,
+	ELLMAllocType InAllocType)
 {
 	if (FLowLevelMemTracker::bIsDisabled)
 	{
@@ -2242,7 +2331,8 @@ FLLMPauseScope::FLLMPauseScope(FName TagName, bool bIsStatTag, uint64 Amount, EL
 	Init(TagName, ELLMTag::Untagged, false, bIsStatTag, Amount, TrackerToPause, InAllocType);
 }
 
-FLLMPauseScope::FLLMPauseScope(ELLMTag TagEnum, bool bIsStatTag, uint64 Amount, ELLMTracker TrackerToPause, ELLMAllocType InAllocType)
+FLLMPauseScope::FLLMPauseScope(ELLMTag TagEnum, bool bIsStatTag, uint64 Amount, ELLMTracker TrackerToPause,
+	ELLMAllocType InAllocType)
 {
 	if (FLowLevelMemTracker::bIsDisabled)
 	{
@@ -2253,10 +2343,12 @@ FLLMPauseScope::FLLMPauseScope(ELLMTag TagEnum, bool bIsStatTag, uint64 Amount,
 	Init(NAME_None, TagEnum, true, false, Amount, TrackerToPause, InAllocType);
 }
 
-void FLLMPauseScope::Init(FName TagName, ELLMTag EnumTag, bool bIsEnumTag, bool bIsStatTag, uint64 Amount, ELLMTracker TrackerToPause, ELLMAllocType InAllocType)
+void FLLMPauseScope::Init(FName TagName, ELLMTag EnumTag, bool bIsEnumTag, bool bIsStatTag, uint64 Amount,
+	ELLMTracker TrackerToPause, ELLMAllocType InAllocType)
 {
 	FLowLevelMemTracker& LLMRef = FLowLevelMemTracker::Get();
-	// We have to check bIsDisabled again after calling Get, because the constructor is called from Get, and will set bIsDisabled=false if the platform doesn't support it
+	// We have to check bIsDisabled again after calling Get, because the constructor is called from Get, and will set 
+	// bIsDisabled=false if the platform doesn't support it.
 	if (FLowLevelMemTracker::bIsDisabled)
 	{
 		bEnabled = false;
@@ -2291,7 +2383,8 @@ void FLLMPauseScope::Init(FName TagName, ELLMTag EnumTag, bool bIsEnumTag, bool
 				}
 				else
 				{
-					LLMRef.GetTracker(Tracker)->PauseAndTrackMemory(TagName, bIsStatTag, static_cast<int64>(Amount), InAllocType);
+					LLMRef.GetTracker(Tracker)->PauseAndTrackMemory(TagName, bIsStatTag, static_cast<int64>(Amount),
+						InAllocType);
 				}
 			}
 		}
@@ -2333,7 +2426,8 @@ FLLMScopeFromPtr::FLLMScopeFromPtr(void* Ptr, ELLMTracker InTracker )
 	}
 
 	FLowLevelMemTracker& LLMRef = FLowLevelMemTracker::Get();
-	// We have to check bIsDisabled again after calling Get, because the constructor is called from Get, and will set bIsDisabled=false if the platform doesn't support it
+	// We have to check bIsDisabled again after calling Get, because the constructor is called from Get, and will set
+	// bIsDisabled=false if the platform doesn't support it.
 	if (FLowLevelMemTracker::bIsDisabled)
 	{
 		bEnabled = false;
@@ -2366,8 +2460,10 @@ FLLMScopeFromPtr::~FLLMScopeFromPtr()
 	LLMRef.GetTracker(Tracker)->PopTag();
 }
 
-FLLMTagDeclaration::FLLMTagDeclaration(const TCHAR* InCPPName, FName InDisplayName, FName InParentTagName, FName InStatName, FName InSummaryStatName)
-	:CPPName(InCPPName), UniqueName(NAME_None), DisplayName(InDisplayName), ParentTagName(InParentTagName), StatName(InStatName), SummaryStatName(InSummaryStatName)
+FLLMTagDeclaration::FLLMTagDeclaration(const TCHAR* InCPPName, FName InDisplayName, FName InParentTagName,
+	FName InStatName, FName InSummaryStatName)
+	:CPPName(InCPPName), UniqueName(NAME_None), DisplayName(InDisplayName), ParentTagName(InParentTagName),
+	StatName(InStatName), SummaryStatName(InSummaryStatName)
 {
 	Register();
 }
@@ -2408,7 +2504,7 @@ void FLLMTagDeclaration::AddCreationCallback(FCreationCallback InCallback)
 		UE::LLMPrivate::LLMTagDeclarationInternal::GetCreationCallbacks();
 	if (CreationCallbacks.Num() >= CreationCallbacks.Max())
 	{
-		check(false); // We are not allowed to allocate memory here; you need to increase the allocation size
+		check(false); // We are not allowed to allocate memory here; you need to increase the allocation size.
 	}
 	else
 	{
@@ -2444,9 +2540,7 @@ void FLLMTagDeclaration::Register()
 	List = this;
 }
 
-namespace UE
-{
-namespace LLMPrivate
+namespace UE::LLMPrivate
 {
 #if DO_CHECK
 
@@ -2455,7 +2549,8 @@ namespace LLMPrivate
 		if (bLog)
 		{
 			TCHAR DescriptionString[4096];
-			GET_VARARGS(DescriptionString, UE_ARRAY_COUNT(DescriptionString), UE_ARRAY_COUNT(DescriptionString) - 1, Format, Format);
+		GET_VARARGS(DescriptionString, UE_ARRAY_COUNT(DescriptionString), UE_ARRAY_COUNT(DescriptionString) - 1,
+			Format, Format);
 
 			FPlatformMisc::LowLevelOutputDebugString(DescriptionString);
 
@@ -2492,22 +2587,26 @@ namespace LLMPrivate
 
 	void ValidateUniqueName(FStringView UniqueName)
 	{
-		// Characters that are invalid for c++ names are invalid (other than /), since we use uniquenames (with / replaced by _) as part of the name of the auto-constructed FLLMTagDeclaration variables
-		// _ is invalid since we use an _ to indicate a / in FLLMTagDeclaration
-		// So only Alnum characters or / are allowed, and the first character can not be a number
+	// Characters that are invalid for c++ names are invalid (other than /), since we use uniquenames
+	// (with / replaced by _) as part of the name of the auto-constructed FLLMTagDeclaration variables
+	// _ is invalid since we use an _ to indicate a / in FLLMTagDeclaration.
+	// So only Alnum characters or / are allowed, and the first character can not be a number.
 		if (UniqueName.Len() == 0)
 		{
 			LLMCheckf(false, TEXT("Invalid length-zero Tag Unique Name"));
 		}
 		else
 		{
-			LLMCheckf(!TChar<TCHAR>::IsDigit(UniqueName[0]), TEXT("Invalid first character is digit in Tag Unique Name '%.*s'"), UniqueName.Len(), UniqueName.GetData());
+		LLMCheckf(!TChar<TCHAR>::IsDigit(UniqueName[0]),
+			TEXT("Invalid first character is digit in Tag Unique Name '%.*s'"),
+			UniqueName.Len(), UniqueName.GetData());
 		}
 		for (TCHAR c : UniqueName)
 		{
 			if (!TChar<TCHAR>::IsAlnum(c) && c != TEXT('/'))
 			{
-				LLMCheckf(false, TEXT("Invalid character %c in Tag Unique Name '%.*s'"), c, UniqueName.Len(), UniqueName.GetData());
+			LLMCheckf(false, TEXT("Invalid character %c in Tag Unique Name '%.*s'"), c,
+				UniqueName.Len(), UniqueName.GetData());
 			}
 		}
 	}
@@ -2515,10 +2614,13 @@ namespace LLMPrivate
 
 	namespace AllocatorPrivate
 	{
+
 		/**
-		 * When a Page is allocated, it splits the memory of the page up into blocks, and creates an FAlloc at the start of each block. All the FAllocs are joined together in a FreeList linked list.
+ * When a Page is allocated, it splits the memory of the page up into blocks, and creates an FAlloc at the start of
+ * each block. All the FAllocs are joined together in a FreeList linked list.
 		 * When a Page allocates memory, it takes an FAlloc from the freelist and gives it to the caller, and forgets about it.
-		 * When the caller returns a pointer, the Page restores the FAlloc at the beginning of the block and puts it back on the FreeList.
+ * When the caller returns a pointer, the Page restores the FAlloc at the beginning of the block and puts it back on
+ * the FreeList.
 		 */
 		struct FAlloc
 		{
@@ -2528,8 +2630,9 @@ namespace LLMPrivate
 		/**
 		 * An FPage holds a single page of memory received from the OS; all pages are of the same size.
 		 * FPages are owned by FBins, and the FPages for an FBin divide the page up into blocks of the FBin's size.
-		 * An FPage keeps track of the blocks it has not yet given out so it can allocate, and keeps track of how many blocks it has given out, so that it can be freed when no longer used.
-		 * Pages that are neither free nor empty (and thus are available for allocating from) are kept in a doubly-linked list on the FBin.
+ * An FPage keeps track of the blocks it has not yet given out so it can allocate, and keeps track of how many blocks
+ * it has given out, so that it can be freed when no longer used.
+ * Pages that are not free or empty are available for allocating from and are kept in a doubly-linked list on the FBin.
 		 */
 		struct FPage
 		{
@@ -2548,10 +2651,12 @@ namespace LLMPrivate
 		};
 
 		/**
-		 * An FBin handles all allocations that fit into its size range. It's size is set to the power of two at the top of its size range.
-		 * The FBin allocates one FPage at a time from the OS; the FPage gets split up into blocks and handles providing a block for callers requesting a pointer.
-		 * The FBin has a doubly-linked list of pages that are in use but are not yet full. It provides allocations from these pages.
-		 * When an FPage gets full, the FBin forgets about it, counting on the caller to give the pointer to the page back when it frees the pointer and the page becomes non-full again.
+ * An FBin handles all allocations that fit in its size range. Its size is the power of two at the top of that range.
+ * The FBin allocates one FPage at a time from the OS; the FPage gets split up into blocks and handles providing a
+ * block for callers requesting a pointer.
+ * The FBin has a doubly-linked list of pages in use but not yet full. It provides allocations from these pages.
+ * When an FPage gets full, the FBin forgets about it, counting on the caller to give the pointer to the page back
+ * when it frees the pointer and the page becomes non-full again.
 		 * When an FBin has no more non-full pages and needs to satisfy an alloc, it allocates a new page.
 		 * When a page becomes unused due to a free, the FBin frees the page, returning it to the OS.
 		 */
@@ -2565,7 +2670,8 @@ namespace LLMPrivate
 			int32 UsedCount;
 			int32 BinSize;
 		};
-	}
+
+} // AllocatorPrivate
 
 	FLLMAllocator*& FLLMAllocator::Get()
 	{
@@ -2601,8 +2707,10 @@ namespace LLMPrivate
 			constexpr int32 MinBinSizeForAlignment = 16;
 			constexpr int32 MinBinSizeForAllocStorage = static_cast<int32>(sizeof(FAlloc));
 			constexpr int32 MultiplierBetweenBins = 2;
-			// Setting MultiplierAfterLastBin=2 would be useless because the PageSize/2 bin would only get a single allocation out of each page due to the FPage data taking up the first half
-			// TODO: For bins >= 4*FPage size, allocate FPages in a separate list rather than embedding them. This will require allocating extra space in each allocation to store its page pointer.
+		// Setting MultiplierAfterLastBin=2 would be useless because the PageSize/2 bin would only get a single
+		// allocation out of each page due to the FPage data taking up the first half
+		// TODO: For bins >= 4*FPage size, allocate FPages in a separate list rather than embedding them.
+		// This will require allocating extra space in each allocation to store its page pointer.
 			constexpr int32 MultiplierAfterLastBin = 4;
 
 			int32 MinBinSize = FMath::Max(MinBinSizeForAllocStorage, MinBinSizeForAlignment);
@@ -2748,7 +2856,8 @@ namespace LLMPrivate
 
 			// Divide the rest of the page after this header into FAllocs, and add all the FAllocs into the free list
 			FreeList = reinterpret_cast<FAlloc*>(reinterpret_cast<intptr_t>(this) + NumHeaderBins*BinSize);
-			FAlloc* EndAlloc = reinterpret_cast<FAlloc*>(reinterpret_cast<intptr_t>(FreeList) + (FreeCount-1) * BinSize);
+		FAlloc* EndAlloc = reinterpret_cast<FAlloc*>(
+			reinterpret_cast<intptr_t>(FreeList) + (FreeCount-1) * BinSize);
 			FAlloc* Alloc = FreeList;
 			while (Alloc != EndAlloc)
 			{
@@ -2832,7 +2941,8 @@ namespace LLMPrivate
 				FPage* Page = reinterpret_cast<FPage*>(Allocator.AllocPages(Allocator.PageSize));
 				++UsedCount;
 				LLMCheck(Page);
-				new (Page) FPage(Allocator.PageSize, BinSize); // The FPage is at the beginning of the array of PageSize bytes
+			// The FPage is at the beginning of the array of PageSize bytes.
+			new (Page) FPage(Allocator.PageSize, BinSize);
 				Page->AddToList(FreePages);
 			}
 
@@ -2846,7 +2956,8 @@ namespace LLMPrivate
 
 		void FBin::Free(void* Ptr, FLLMAllocator& Allocator)
 		{
-			FPage* Page = reinterpret_cast<FPage*>(reinterpret_cast<intptr_t>(Ptr) & ~(static_cast<intptr_t>(Allocator.PageSize) - 1));
+		FPage* Page = reinterpret_cast<FPage*>(
+			reinterpret_cast<intptr_t>(Ptr) & ~(static_cast<intptr_t>(Allocator.PageSize) - 1));
 			if (Page->IsFull())
 			{
 				Page->AddToList(FreePages);
@@ -2861,21 +2972,25 @@ namespace LLMPrivate
 		}
 	}
 
-	FTagData::FTagData(FName InName, FName InDisplayName, FName InParentName, FName InStatName, FName InSummaryStatName, bool bInHasEnumTag, ELLMTag InEnumTag, ETagReferenceSource InReferenceSource)
+FTagData::FTagData(FName InName, FName InDisplayName, FName InParentName, FName InStatName,
+	FName InSummaryStatName, bool bInHasEnumTag, ELLMTag InEnumTag, ETagReferenceSource InReferenceSource)
 		: Name(InName), DisplayName(InDisplayName), ParentName(InParentName), StatName(InStatName)
 		, SummaryStatName(InSummaryStatName), EnumTag(InEnumTag), ReferenceSource(InReferenceSource)
 		, bIsFinishConstructed(false), bParentIsName(true), bHasEnumTag(bInHasEnumTag), bIsReportable(true)
 	{
 	}
 
-	FTagData::FTagData(FName InName, FName InDisplayName, const FTagData* InParent, FName InStatName, FName InSummaryStatName, bool bInHasEnumTag, ELLMTag InEnumTag, ETagReferenceSource InReferenceSource)
-		: FTagData(InName, InDisplayName, NAME_None, InStatName, InSummaryStatName, bInHasEnumTag, InEnumTag, InReferenceSource)
+FTagData::FTagData(FName InName, FName InDisplayName, const FTagData* InParent, FName InStatName,
+	FName InSummaryStatName, bool bInHasEnumTag, ELLMTag InEnumTag, ETagReferenceSource InReferenceSource)
+	: FTagData(InName, InDisplayName, NAME_None, InStatName, InSummaryStatName, bInHasEnumTag, InEnumTag,
+		InReferenceSource)
 	{
 		SetParent(InParent);
 	}
 
 	FTagData::FTagData(ELLMTag InEnumTag)
-		: FTagData(NAME_None, NAME_None, NAME_None, NAME_None, NAME_None, true, InEnumTag, ETagReferenceSource::EnumTag)
+	: FTagData(NAME_None, NAME_None, NAME_None, NAME_None, NAME_None, true, InEnumTag,
+		ETagReferenceSource::EnumTag)
 	{
 	}
 
@@ -2927,13 +3042,15 @@ namespace LLMPrivate
 
 	const FTagData* FTagData::GetParent() const
 	{
-		LLMCheckf(!bParentIsName, TEXT("GetParent called on TagData %s before SetParent was called"), *WriteToString<FName::StringBufferSize>(Name));
+	LLMCheckf(!bParentIsName, TEXT("GetParent called on TagData %s before SetParent was called"),
+		*WriteToString<FName::StringBufferSize>(Name));
 		return Parent;
 	}
 
 	FName FTagData::GetParentName() const
 	{
-		LLMCheckf(bParentIsName, TEXT("GetParentName called on TagData %s after SetParent was called"), *WriteToString<FName::StringBufferSize>(Name));
+	LLMCheckf(bParentIsName, TEXT("GetParentName called on TagData %s after SetParent was called"),
+		*WriteToString<FName::StringBufferSize>(Name));
 		return ParentName;
 	}
 
@@ -2968,7 +3085,8 @@ namespace LLMPrivate
 			}
 			TagData = TagData->GetParent();
 		} while (TagData);
-		LLMCheckf(false, TEXT("TagData is not a descendent of an ELLMTag TagData. All TagDatas must be descendents of ELLMTag::CustomName if they are not descendets of any other ELLMTag"));
+	LLMCheckf(false, TEXT("TagData is not a descendent of an ELLMTag TagData. ")
+		TEXT("All TagDatas must be descendents of ELLMTag::CustomName if they are not descendets of any other ELLMTag"));
 		return this;
 	}
 
@@ -3046,9 +3164,10 @@ namespace LLMPrivate
 
 	bool FTagData::IsUsedAsDisplayParent() const
 	{
-		// All Tags but one are UsedAsDisplayParent - their name is prepended during GetDisplayPath
-		// ELLMTag::CustomName is the exception. It is set for FName tags that do not have a real parent to provide a containing ELLMTag for them to provide to systems that do not support FName tags.
-		// When FName tags without a real parent are displayed, their path should display as parentless despite having the CustomName tag as their parent.
+	// All Tags but one are UsedAsDisplayParent - their name is prepended during GetDisplayPath.
+	// ELLMTag::CustomName is the exception. It is set for FName tags that do not have a real parent to provide a
+	// containing ELLMTag for them to provide to systems that do not support FName tags. When FName tags without a real
+	// parent are displayed, their path should display as parentless despite having the CustomName tag as their parent.
 		return !(bHasEnumTag && EnumTag == ELLMTag::CustomName);
 	}
 
@@ -3169,7 +3288,8 @@ namespace LLMPrivate
 	}
 #endif
 
-	void FLLMTracker::TrackAllocation(const void* Ptr, int64 Size, ELLMTag DefaultEnumTag, ELLMTracker Tracker, ELLMAllocType AllocType, bool bTrackInMemPro)
+void FLLMTracker::TrackAllocation(const void* Ptr, int64 Size, ELLMTag DefaultEnumTag, ELLMTracker Tracker,
+	ELLMAllocType AllocType, bool bTrackInMemPro)
 	{
 		FLLMThreadState* State = GetOrCreateState();
 		const FTagData* TagData = State->GetTopTag();
@@ -3180,7 +3300,8 @@ namespace LLMPrivate
 		TrackAllocation(Ptr, Size, TagData, Tracker, AllocType, State, bTrackInMemPro);
 	}
 
-	void FLLMTracker::TrackAllocation(const void* Ptr, int64 Size, FName DefaultTag, ELLMTracker Tracker, ELLMAllocType AllocType, bool bTrackInMemPro)
+void FLLMTracker::TrackAllocation(const void* Ptr, int64 Size, FName DefaultTag, ELLMTracker Tracker,
+	ELLMAllocType AllocType, bool bTrackInMemPro)
 	{
 		FLLMThreadState* State = GetOrCreateState();
 		const FTagData* TagData = State->GetTopTag();
@@ -3191,12 +3312,14 @@ namespace LLMPrivate
 		TrackAllocation(Ptr, Size, TagData, Tracker, AllocType, State, bTrackInMemPro);
 	}
 
-	void FLLMTracker::TrackAllocation(const void* Ptr, int64 Size, const FTagData* ActiveTagData, ELLMTracker Tracker, ELLMAllocType AllocType, FLLMThreadState* State, bool bTrackInMemPro)
+void FLLMTracker::TrackAllocation(const void* Ptr, int64 Size, const FTagData* ActiveTagData, ELLMTracker Tracker,
+	ELLMAllocType AllocType, FLLMThreadState* State, bool bTrackInMemPro)
 	{
 		if (IsPaused(AllocType))
 		{
-			// When Paused, we do not track any new allocations and we we do not update the counters for the memory they use; the code that triggered the pause is responsible for updating those counters
-			// Since we do not track the allocations, TrackFree will likewise not update the counters when those allocations are freed
+		// When Paused, we do not track any new allocations and we we do not update the counters for the memory
+		// they use; the code that triggered the pause is responsible for updating those counters. Since we do not
+		// track the allocations, TrackFree will likewise not update the counters when those allocations are freed.
 			return;
 		}
 
@@ -3204,7 +3327,7 @@ namespace LLMPrivate
 		FPlatformAtomics::InterlockedAdd(&TrackedTotal, Size);
 	
 #if !LLM_ENABLED_FULL_TAGS
-		// When full tags are disabled, we instead store the top-level enumtag parent of the tag used by each allocation
+	// When full tags are disabled, we instead store the top-level enumtag parent of the allocation's tag.
 		ActiveTagData = ActiveTagData->GetContainingEnumTagData();
 #endif
 
@@ -3217,7 +3340,7 @@ namespace LLMPrivate
 		// track on the thread state
 		State->TrackAllocation(Ptr, Size, Tracker, AllocType, ActiveTagData, AssetTagData, bTrackInMemPro);
 
-		// tracking a nullptr with a Size is allowed, but we don't need to remember it, since we can't free it ever
+	// tracking a nullptr with a Size is allowed, but we don't need to remember it, since we can't free it ever.
 		if (Ptr != nullptr)
 		{
 			// remember the size and tag info
@@ -3247,7 +3370,8 @@ namespace LLMPrivate
 
 		if (IsPaused(AllocType))
 		{
-			// When Paused, we remove our data for any freed allocations, but we do not update the counters for the memory they used; the code that triggered the pause is responsible for updating those counters
+		// When Paused, we remove our data for any freed allocations, but we do not update the counters for the
+		// memory they used; the code that triggered the pause is responsible for updating those counters.
 			return;
 		}
 
@@ -3285,7 +3409,7 @@ namespace LLMPrivate
 
 		if (IsPaused(AllocType))
 		{
-			// When Paused, don't update counters in case any of the external tracking systems are not available
+		// When Paused, don't update counters in case any of the external tracking systems are not available.
 			return;
 		}
 
@@ -3330,7 +3454,7 @@ namespace LLMPrivate
 		PauseAndTrackMemory(TagData, Amount, AllocType);
 	}
 
-	// This will pause/unpause tracking, and also manually increment a given tag
+// This will pause/unpause tracking, and also manually increment a given tag.
 	void FLLMTracker::PauseAndTrackMemory(const FTagData* TagData, int64 Amount, ELLMAllocType AllocType)
 	{
 		FLLMThreadState* State = GetOrCreateState();
@@ -3357,8 +3481,10 @@ namespace LLMPrivate
 	bool FLLMTracker::IsPaused(ELLMAllocType AllocType)
 	{
 		FLLMThreadState* State = GetState();
-		// pause during shutdown, as the external trackers might not be able to robustly handle tracking once we start shutting down
-		return IsEngineExitRequested() || (State == nullptr ? false : (State->PausedCounter[static_cast<int32>(ELLMAllocType::None)]>0) || (State->PausedCounter[static_cast<int32>(AllocType)])>0);
+	// pause during shutdown, as the external trackers might not be able to robustly handle tracking.
+	return IsEngineExitRequested() || (State == nullptr ? false :
+		((State->PausedCounter[static_cast<int32>(ELLMAllocType::None)]>0) ||
+			(State->PausedCounter[static_cast<int32>(AllocType)])>0));
 	}
 
 	void FLLMTracker::Clear()
@@ -3386,7 +3512,8 @@ namespace LLMPrivate
 		CsvWriter.OnPreFork();
 	}
 
-	void FLLMTracker::SetTotalTags(const FTagData* InOverrideUntaggedTagData, const FTagData* InOverrideTrackedTotalTagData)
+void FLLMTracker::SetTotalTags(const FTagData* InOverrideUntaggedTagData,
+	const FTagData* InOverrideTrackedTotalTagData)
 	{
 		OverrideUntaggedTagData = InOverrideUntaggedTagData;
 		OverrideTrackedTotalTagData = InOverrideTrackedTotalTagData;
@@ -3406,7 +3533,7 @@ namespace LLMPrivate
 			}
 		}
 
-		// Add the values from each thread to the central repository
+	// Add the values from each thread to the central repository.
 		bool bTrimThreads = bTrimAllocations;
 		for (FLLMThreadState* ThreadState : ThreadStates)
 		{
@@ -3414,7 +3541,7 @@ namespace LLMPrivate
 			ThreadState->FetchAndClearTagSizes(TagSizes, AllocTypeAmounts, bTrimThreads);
 		}
 
-		// Update peak sizes and external sizes in the central repository
+	// Update peak sizes and external sizes in the central repository.
 		for (TPair<const FTagData*, FTrackerTagSizeData>& It : TagSizes)
 		{
 			FTrackerTagSizeData& AllocationData = It.Value;
@@ -3424,16 +3551,17 @@ namespace LLMPrivate
 			{
 				if (AllocationData.bExternalAddToTotal)
 				{
-					FPlatformAtomics::InterlockedAdd(&TrackedTotal, AllocationData.ExternalAmount - AllocationData.Size);
+				FPlatformAtomics::InterlockedAdd(&TrackedTotal,
+					AllocationData.ExternalAmount - AllocationData.Size);
 				}
 				AllocationData.Size = AllocationData.ExternalAmount;
 				AllocationData.bExternalValid = false;
 			}
 
-			// Calculate peaks
+		// Calculate peaks.
 #if LLM_ENABLED_TRACK_PEAK_MEMORY
-			// @todo we should be keeping track of the intra-frame memory peak for the total tracked memory.
-			// For now we will just use the memory at the time the update happens since there are threading implications to being accurate.
+		// TODO: we should keep track of the intra-frame memory peak for the total tracked memory. For now we will
+		// use the memory at the time the update happens since there are threading implications to being accurate.
 			AllocationData.PeakSize = FMath::Max(AllocationData.PeakSize, AllocationData.Size);
 #endif
 		}
@@ -3441,8 +3569,9 @@ namespace LLMPrivate
 
 	void FLLMTracker::UpdateThreads()
 	{
-		// Consume pending thread states
-		// We must be careful to do all allocations outside of the PendingThreadStatesGuard guard as that can lead to a deadlock due to contention with PendingThreadStatesGuard & Locks inside the underlying allocator (i.e. MallocBinned2 -> Mutex)
+	// Consume pending thread states. We must be careful to do all allocations outside of the
+	// PendingThreadStatesGuard guard as that can lead to a deadlock due to contention with
+	// PendingThreadStatesGuard & Locks inside the underlying allocator (i.e. MallocBinned2 -> Mutex).
 		{
 			PendingThreadStatesGuard.Lock();
 			const int32 NumPendingThreadStatesToConsume = PendingThreadStates.Num();
@@ -3473,8 +3602,10 @@ namespace LLMPrivate
 		{
 			const FTagData* TagData = LLMRef.FindTagData(TagName_Untagged);
 			const FTrackerTagSizeData* AllocationData = TagData ? TagSizes.Find(TagData) : nullptr;
-			SetMemoryStatByFName(OverrideUntaggedTagData->GetStatName(), AllocationData ? AllocationData->GetSize(bTrackPeaks) : 0);
-			SetMemoryStatByFName(OverrideUntaggedTagData->GetSummaryStatName(), AllocationData ? AllocationData->GetSize(bTrackPeaks) : 0);
+		SetMemoryStatByFName(OverrideUntaggedTagData->GetStatName(),
+			AllocationData ? AllocationData->GetSize(bTrackPeaks) : 0);
+		SetMemoryStatByFName(OverrideUntaggedTagData->GetSummaryStatName(),
+			AllocationData ? AllocationData->GetSize(bTrackPeaks) : 0);
 		}
 
 		for (const TPair<const FTagData*, FTrackerTagSizeData>& It : TagSizes)
@@ -3486,7 +3617,7 @@ namespace LLMPrivate
 			}
 			if (OverrideUntaggedTagData && TagData->GetName() == TagName_Untagged)
 			{
-				// Handled separately by OverrideUntaggedTagData
+			// Handled separately by OverrideUntaggedTagData.
 				continue;
 			}
 			const FTrackerTagSizeData& AllocationData = It.Value;
@@ -3499,17 +3630,20 @@ namespace LLMPrivate
 
 	void FLLMTracker::PublishCsv(bool bTrackPeaks)
 	{
-		CsvWriter.Publish(LLMRef, TagSizes, OverrideTrackedTotalTagData, OverrideUntaggedTagData, TrackedTotal, bTrackPeaks);
+	CsvWriter.Publish(LLMRef, TagSizes, OverrideTrackedTotalTagData, OverrideUntaggedTagData, TrackedTotal,
+		bTrackPeaks);
 	}
 
 	void FLLMTracker::PublishTrace(bool bTrackPeaks)
 	{
-		TraceWriter.Publish(LLMRef, TagSizes, OverrideTrackedTotalTagData, OverrideUntaggedTagData, TrackedTotal, bTrackPeaks);
+	TraceWriter.Publish(LLMRef, TagSizes, OverrideTrackedTotalTagData, OverrideUntaggedTagData, TrackedTotal,
+		bTrackPeaks);
 	}
 
 	void FLLMTracker::PublishCsvProfiler(bool bTrackPeaks)
 	{
-		CsvProfilerWriter.Publish(LLMRef, TagSizes, OverrideTrackedTotalTagData, OverrideUntaggedTagData, TrackedTotal, bTrackPeaks);
+	CsvProfilerWriter.Publish(LLMRef, TagSizes, OverrideTrackedTotalTagData, OverrideUntaggedTagData, TrackedTotal,
+		bTrackPeaks);
 	}
 
 	void FLLMTracker::OnTagsResorted(FTagDataArray& OldTagDatas)
@@ -3517,7 +3651,7 @@ namespace LLMPrivate
 #if LLM_ENABLED_FULL_TAGS
 		{
 			// Each allocation references the tag by its index, which we have just remapped.
-			// Remap each allocation's tag index to the new index for the tag
+		// Remap each allocation's tag index to the new index for the tag.
 			AllocationMap.LockAll();
 			for (const FLLMAllocMap::FTuple& Tuple : AllocationMap)
 			{
@@ -3526,10 +3660,10 @@ namespace LLMPrivate
 			AllocationMap.UnlockAll();
 		}
 #else
-		// Values in AllocationMap are ELLMTags, and don't depend on the Index of the tagdatas
+	// Values in AllocationMap are ELLMTags, and don't depend on the Index of the tagdatas.
 #endif
 
-		// Update the uses of Index in the ThreadStates
+	// Update the uses of Index in the ThreadStates.
 		for (FLLMThreadState* ThreadState : ThreadStates)
 		{
 			ThreadState->OnTagsResorted(OldTagDatas);
@@ -3648,7 +3782,7 @@ namespace LLMPrivate
 	{
 		FScopeLock Lock(&TagSection);
 
-		// push a tag
+	// Push a tag.
 		TagStack.Add(TagData);
 	}
 
@@ -3656,7 +3790,8 @@ namespace LLMPrivate
 	{
 		FScopeLock Lock(&TagSection);
 
-		LLMCheckf(TagStack.Num() > 0, TEXT("Called FLLMThreadState::PopTag without a matching Push (stack was empty on pop)"));
+	LLMCheckf(TagStack.Num() > 0,
+		TEXT("Called FLLMThreadState::PopTag without a matching Push (stack was empty on pop)"));
 		TagStack.Pop(false /* bAllowShrinking */);
 	}
 
@@ -3670,7 +3805,7 @@ namespace LLMPrivate
 	{
 		FScopeLock Lock(&TagSection);
 
-		// push a tag
+	// Push a tag.
 		AssetTagStack.Add(TagData);
 	}
 
@@ -3678,7 +3813,8 @@ namespace LLMPrivate
 	{
 		FScopeLock Lock(&TagSection);
 
-		LLMCheckf(AssetTagStack.Num() > 0, TEXT("Called FLLMThreadState::PopTag without a matching Push (stack was empty on pop)"));
+	LLMCheckf(AssetTagStack.Num() > 0,
+		TEXT("Called FLLMThreadState::PopTag without a matching Push (stack was empty on pop)"));
 		AssetTagStack.Pop(false /* bAllowShrinking */);
 	}
 
@@ -3691,13 +3827,14 @@ namespace LLMPrivate
 
 	void FLLMThreadState::IncrTag(const FTagData* TagData, int64 Amount)
 	{
-		// Caller is responsible for holding a lock on TagSection
+	// Caller is responsible for holding a lock on TagSection.
 		FThreadTagSizeData& AllocationSize = Allocations.FindOrAdd(TagData->GetIndex());
 		AllocationSize.TagData = TagData;
 		AllocationSize.Size += Amount;
 	}
 
-	void FLLMThreadState::TrackAllocation(const void* Ptr, int64 Size, ELLMTracker Tracker, ELLMAllocType AllocType, const FTagData* TagData, const FTagData* AssetTagData, bool bTrackInMemPro)
+void FLLMThreadState::TrackAllocation(const void* Ptr, int64 Size, ELLMTracker Tracker, ELLMAllocType AllocType,
+	const FTagData* TagData, const FTagData* AssetTagData, bool bTrackInMemPro)
 	{
 		FScopeLock Lock(&TagSection);
 
@@ -3725,7 +3862,8 @@ namespace LLMPrivate
 #endif
 	}
 
-	void FLLMThreadState::TrackFree(const void* Ptr, int64 Size, ELLMTracker Tracker, ELLMAllocType AllocType, const FTagData* TagData, const FTagData* AssetTagData, bool bTrackInMemPro)
+void FLLMThreadState::TrackFree(const void* Ptr, int64 Size, ELLMTracker Tracker, ELLMAllocType AllocType,
+	const FTagData* TagData, const FTagData* AssetTagData, bool bTrackInMemPro)
 	{
 		FScopeLock Lock(&TagSection);
 
@@ -3752,9 +3890,10 @@ namespace LLMPrivate
 #endif
 	}
 
-	void FLLMThreadState::TrackMoved(const void* Dest, const void* Source, int64 Size, ELLMTracker Tracker, const FTagData* TagData)
+void FLLMThreadState::TrackMoved(const void* Dest, const void* Source, int64 Size, ELLMTracker Tracker,
+	const FTagData* TagData)
 	{
-		// update external memory trackers (ideally would want a proper 'move' option on these)
+	// Update external memory trackers (ideally would want a proper 'move' option on these).
 		ELLMTag EnumTag = TagData->GetContainingEnum();
 		if (Tracker == ELLMTracker::Default)
 		{
@@ -3775,7 +3914,7 @@ namespace LLMPrivate
 	{
 		FScopeLock Lock(&TagSection);
 
-		// Make sure all parents of any TagDatas in the Allocations are also present
+	// Make sure all parents of any TagDatas in the Allocations are also present.
 		FConstTagDataArray ParentsToAdd;
 		for (const TPair<int32, FThreadTagSizeData>& It : Allocations)
 		{
@@ -3793,7 +3932,8 @@ namespace LLMPrivate
 			Info.TagData = TagData;
 		}
 
-		// Tags are sorted topologically from parent to child, so we can accumulate children into parents recursively by reverse iterating the map
+	// Tags are sorted topologically from parent to child, so we can accumulate children into parents recursively
+	// by reverse iterating the map.
 		for (FThreadTagSizeMap::TConstReverseIterator It(Allocations); It; ++It)
 		{
 			const FThreadTagSizeData& Info = It->Value;
@@ -3833,7 +3973,8 @@ namespace LLMPrivate
 		}
 	}
 
-	void FLLMThreadState::FetchAndClearTagSizes(FTrackerTagSizeMap& TagSizes, int64* InAllocTypeAmounts, bool bTrimAllocations)
+void FLLMThreadState::FetchAndClearTagSizes(FTrackerTagSizeMap& TagSizes, int64* InAllocTypeAmounts,
+	bool bTrimAllocations)
 	{
 		FScopeLock Lock(&TagSection);
 		for (TPair<int32, FThreadTagSizeData>& Item : Allocations)
@@ -3875,10 +4016,7 @@ namespace LLMPrivate
 		}
 	}
 
-	/*
-	 * FLLMCsvWriter implementation
-	*/
-
+// FLLMCsvWriter implementation.
 	FLLMCsvWriter::FLLMCsvWriter()
 		: Archive(nullptr)
 		, LastWriteTime(FPlatformTime::Seconds())
@@ -3912,7 +4050,9 @@ namespace LLMPrivate
 		Tracker = InTracker;
 	}
 
-	void FLLMCsvWriter::Publish(FLowLevelMemTracker& LLMRef, const FTrackerTagSizeMap& TagSizes, const FTagData* OverrideTrackedTotalTagData, const FTagData* OverrideUntaggedTagData, int64 TrackedTotal, bool bTrackPeaks)
+void FLLMCsvWriter::Publish(FLowLevelMemTracker& LLMRef, const FTrackerTagSizeMap& TagSizes,
+	const FTagData* OverrideTrackedTotalTagData, const FTagData* OverrideUntaggedTagData, int64 TrackedTotal,
+	bool bTrackPeaks)
 	{
 		const double Now = FPlatformTime::Seconds();
 
@@ -3929,7 +4069,8 @@ namespace LLMPrivate
 
 		if (bCreatedArchive || bColumnsUpdated)
 		{
-			// The column names are written at the start of the archive; when they change we seek back to the start of the file and rewrite the column names.
+		// The column names are written at the start of the archive; when they change we seek back to the start of
+		// the file and rewrite the column names.
 			WriteHeader(OverrideTrackedTotalTagData, OverrideUntaggedTagData);
 		}
 
@@ -3946,12 +4087,11 @@ namespace LLMPrivate
 		}
 	}
 
-	/*
-	 * Archive is a binary stream, so we can't just serialise an FString using <<
-	*/
+	// Archive is a binary stream, so we can't just serialise an FString using <<.
 	void FLLMCsvWriter::Write(FStringView Text)
 	{
-		Archive->Serialize(const_cast<ANSICHAR*>(StringCast<ANSICHAR>(Text.GetData(), Text.Len()).Get()), Text.Len() * sizeof(ANSICHAR));
+		Archive->Serialize(const_cast<ANSICHAR*>(
+			StringCast<ANSICHAR>(Text.GetData(), Text.Len()).Get()), Text.Len() * sizeof(ANSICHAR));
 	}
 
 	bool FLLMCsvWriter::CreateArchive()
@@ -3961,7 +4101,7 @@ namespace LLMPrivate
 			return false;
 		}
 
-		// create the csv file
+	// Create the csv file.
 		FString Directory = FPaths::ProfilingDir() + TEXT("LLM/");
 		IFileManager::Get().MakeDirectory(*Directory, true);
 
@@ -3969,20 +4109,22 @@ namespace LLMPrivate
 		const FDateTime FileDate = FDateTime::Now();
 #if PLATFORM_DESKTOP
 		FString PlatformName = FPlatformProperties::PlatformName();
-#else // Use the CPU for consoles so we can differentiate things like different SKUs of a console generation
+#else // Use the CPU for consoles so we can differentiate things like different SKUs of a console generation.
 		FString PlatformName = FPlatformMisc::GetCPUBrand().TrimStartAndEnd();
 #endif
 		PlatformName.ReplaceCharInline(' ', '_');
 		PlatformName = FPaths::MakeValidFileName(PlatformName);
 #if WITH_SERVER_CODE
-		FString Filename = FString::Printf(TEXT("%s/%s_Pid%d_%s_%s.csv"), *Directory, TrackerName, FPlatformProcess::GetCurrentProcessId(), *FileDate.ToString(), *PlatformName);
+	FString Filename = FString::Printf(TEXT("%s/%s_Pid%d_%s_%s.csv"), *Directory, TrackerName,
+		FPlatformProcess::GetCurrentProcessId(), *FileDate.ToString(), *PlatformName);
 #else
-		FString Filename = FString::Printf(TEXT("%s/%s_%s_%s.csv"), *Directory, TrackerName, *FileDate.ToString(), *PlatformName);
+	FString Filename = FString::Printf(TEXT("%s/%s_%s_%s.csv"), *Directory, TrackerName, *FileDate.ToString(),
+		*PlatformName);
 #endif
 		Archive = IFileManager::Get().CreateFileWriter(*Filename, FILEWRITE_AllowRead | FILEWRITE_NoFail);
 		LLMCheck(Archive);
 
-		// create space for column titles that are filled in as we get them
+	// Create space for column titles that are filled in as we get them.
 		Write(FString::ChrN(CVarLLMHeaderMaxSize.GetValueOnAnyThread(), ' '));
 		Write(TEXT("\n"));
 
@@ -4002,7 +4144,7 @@ namespace LLMPrivate
 			}
 			if (TagData->GetName() == TagName_Untagged)
 			{
-				continue; // Handled by OverrideUntaggedName
+			continue; // Handled by OverrideUntaggedName.
 			}
 			if (ExistingColumns.Contains(TagData))
 			{
@@ -4016,7 +4158,8 @@ namespace LLMPrivate
 		return bUpdated;
 	}
 
-	void FLLMCsvWriter::WriteHeader(const FTagData* OverrideTrackedTotalTagData, const FTagData* OverrideUntaggedTagData)
+void FLLMCsvWriter::WriteHeader(const FTagData* OverrideTrackedTotalTagData,
+	const FTagData* OverrideUntaggedTagData)
 	{
 		int64 OriginalOffset = Archive->Tell();
 		Archive->Seek(0);
@@ -4045,13 +4188,17 @@ namespace LLMPrivate
 		int64 ColumnTitleTotalSize = Archive->Tell();
 		if (ColumnTitleTotalSize >= CVarLLMHeaderMaxSize.GetValueOnAnyThread())
 		{
-			UE_LOG(LogHAL, Error, TEXT("LLM column titles have overflowed, LLM CSM data will be corrupted. Increase CVarLLMHeaderMaxSize > %d"), ColumnTitleTotalSize);
+		UE_LOG(LogHAL, Error,
+			TEXT("LLM column titles have overflowed, LLM CSM data will be corrupted. Increase CVarLLMHeaderMaxSize > %d"),
+			ColumnTitleTotalSize);
 		}
 
 		Archive->Seek(OriginalOffset);
 	}
 
-	void FLLMCsvWriter::AddRow(FLowLevelMemTracker& LLMRef, const FTrackerTagSizeMap& TagSizes, const FTagData* OverrideTrackedTotalTagData, const FTagData* OverrideUntaggedTagData, int64 TrackedTotal, bool bTrackPeaks)
+void FLLMCsvWriter::AddRow(FLowLevelMemTracker& LLMRef, const FTrackerTagSizeMap& TagSizes,
+	const FTagData* OverrideTrackedTotalTagData, const FTagData* OverrideUntaggedTagData, int64 TrackedTotal,
+	bool bTrackPeaks)
 	{
 		TStringBuilder<256> TextBuffer;
 		auto WriteValue = [this, &TextBuffer](int64 Value)
@@ -4105,11 +4252,7 @@ namespace LLMPrivate
 		Archive->Flush();
 	}
 
-	/*
-	 * FLLMTraceWriter implementation
-	*/
-
-
+// FLLMTraceWriter implementation.
 	FLLMTraceWriter::FLLMTraceWriter()
 	{
 	}
@@ -4133,7 +4276,9 @@ namespace LLMPrivate
 		return reinterpret_cast<const void*>(TagData);
 	}
 
-	void FLLMTraceWriter::Publish(FLowLevelMemTracker& LLMRef, const FTrackerTagSizeMap& TagSizes, const FTagData* OverrideTrackedTotalTagData, const FTagData* OverrideUntaggedTagData, int64 TrackedTotal, bool bTrackPeaks)
+void FLLMTraceWriter::Publish(FLowLevelMemTracker& LLMRef, const FTrackerTagSizeMap& TagSizes,
+	const FTagData* OverrideTrackedTotalTagData, const FTagData* OverrideUntaggedTagData, int64 TrackedTotal,
+	bool bTrackPeaks)
 	{
 		if (!UE_TRACE_CHANNELEXPR_IS_ENABLED(MemTagChannel))
 		{
@@ -4185,7 +4330,7 @@ namespace LLMPrivate
 			}
 			if (OverrideUntaggedTagData != nullptr && TagData->GetName() == TagName_Untagged)
 			{
-				continue; // Handled by OverrideUntaggedTagData
+			continue; // Handled by OverrideUntaggedTagData.
 			}
 			SendTagDeclaration(TagData);
 		}
@@ -4228,7 +4373,7 @@ namespace LLMPrivate
 			}
 			if (OverrideUntaggedTagData && TagData->GetName() == TagName_Untagged)
 			{
-				continue; // Handled by OverrideUntaggedTagData
+			continue; // Handled by OverrideUntaggedTagData.
 			}
 			AddValue(TagData, It.Value.GetSize(bTrackPeaks));
 		}
@@ -4256,9 +4401,7 @@ namespace LLMPrivate
 		}
 	}
 
-	/*
-	 * FLLMCsvProfilerWriter implementation
-	*/
+// FLLMCsvProfilerWriter implementation.
 	FLLMCsvProfilerWriter::FLLMCsvProfilerWriter()
 	{
 	}
@@ -4278,10 +4421,13 @@ namespace LLMPrivate
 #endif
 	}
 
-	void FLLMCsvProfilerWriter::Publish(FLowLevelMemTracker& LLMRef, const FTrackerTagSizeMap& TagSizes, const FTagData* OverrideTrackedTotalTagData, const FTagData* OverrideUntaggedTagData, int64 TrackedTotal, bool bTrackPeaks)
+void FLLMCsvProfilerWriter::Publish(FLowLevelMemTracker& LLMRef, const FTrackerTagSizeMap& TagSizes,
+	const FTagData* OverrideTrackedTotalTagData, const FTagData* OverrideUntaggedTagData, int64 TrackedTotal,
+	bool bTrackPeaks)
 	{
 #if LLM_CSV_PROFILER_WRITER_ENABLED
-		int32 CsvCategoryIndex = (Tracker == ELLMTracker::Platform) ? CSV_CATEGORY_INDEX(LLMPlatform) : CSV_CATEGORY_INDEX(LLM);
+	int32 CsvCategoryIndex = (Tracker == ELLMTracker::Platform) ?
+		CSV_CATEGORY_INDEX(LLMPlatform) : CSV_CATEGORY_INDEX(LLM);
 
 		if (OverrideTrackedTotalTagData)
 		{
@@ -4292,7 +4438,8 @@ namespace LLMPrivate
 		{
 			const FTagData* TagData = LLMRef.FindTagData(TagName_Untagged);
 			const FTrackerTagSizeData* AllocationData = TagData ? TagSizes.Find(TagData) : nullptr;
-			RecordTagToCsv(CsvCategoryIndex, OverrideUntaggedTagData, AllocationData ? AllocationData->GetSize(bTrackPeaks) : 0);
+		RecordTagToCsv(CsvCategoryIndex, OverrideUntaggedTagData,
+			AllocationData ? AllocationData->GetSize(bTrackPeaks) : 0);
 		}
 
 		for (const TPair<const FTagData*, FTrackerTagSizeData>& It : TagSizes)
@@ -4304,7 +4451,7 @@ namespace LLMPrivate
 			}
 			if (OverrideUntaggedTagData && TagData->GetName() == TagName_Untagged)
 			{
-				// Handled separately by OverrideUntaggedTagData
+			// Handled separately by OverrideUntaggedTagData.
 				continue;
 			}
 			const FTrackerTagSizeData* AllocationData = TagSizes.Find(TagData);
@@ -4324,12 +4471,11 @@ namespace LLMPrivate
 			TagDataToCsvStatName.Add(TagData, NewCsvStatName);
 			CsvStatNamePtr = &NewCsvStatName;
 		}
-		FCsvProfiler::RecordCustomStat(*CsvStatNamePtr, CsvCategoryIndex, (float)((double)TagSize / (1024.0 * 1024.0)), ECsvCustomStatOp::Set);
+	FCsvProfiler::RecordCustomStat(*CsvStatNamePtr, CsvCategoryIndex, (float)((double)TagSize / (1024.0 * 1024.0)),
+		ECsvCustomStatOp::Set);
 #endif
 	}
 
-
-}
-}
+} // UE::LLMPrivate
 
 #endif		// #if ENABLE_LOW_LEVEL_MEM_TRACKER
diff --git a/Engine/Source/Runtime/Core/Public/HAL/LowLevelMemTracker.h b/Engine/Source/Runtime/Core/Public/HAL/LowLevelMemTracker.h
--- a/Engine/Source/Runtime/Core/Public/HAL/LowLevelMemTracker.h
+++ b/Engine/Source/Runtime/Core/Public/HAL/LowLevelMemTracker.h
@@ -10,34 +10,42 @@
 #endif
 
 // LLM is currently incompatible with PLATFORM_USES_FIXED_GMalloc_CLASS, because LLM is activated way too early
-// Inability to use LLM with PLATFORM_USES_FIXED_GMalloc_CLASS is not a problem, because fixed GMalloc is only used in Test/Shipping builds
+// This is not a problem, because fixed GMalloc is only used in Test/Shipping builds
 #define LLM_ENABLED_ON_PLATFORM (PLATFORM_SUPPORTS_LLM && !PLATFORM_USES_FIXED_GMalloc_CLASS)
 
 // *** enable/disable LLM here ***
 #if !defined(ENABLE_LOW_LEVEL_MEM_TRACKER) || !LLM_ENABLED_ON_PLATFORM 
 	#undef ENABLE_LOW_LEVEL_MEM_TRACKER
-	#define ENABLE_LOW_LEVEL_MEM_TRACKER (LLM_ENABLED_ON_PLATFORM && !UE_BUILD_SHIPPING && (!UE_BUILD_TEST || ALLOW_LOW_LEVEL_MEM_TRACKER_IN_TEST) && WITH_ENGINE && 1)
+	#define ENABLE_LOW_LEVEL_MEM_TRACKER ( \
+		LLM_ENABLED_ON_PLATFORM && \
+		!UE_BUILD_SHIPPING && (!UE_BUILD_TEST || ALLOW_LOW_LEVEL_MEM_TRACKER_IN_TEST) && \
+		WITH_ENGINE \
+	)
 #endif
 
 #if ENABLE_LOW_LEVEL_MEM_TRACKER
 
 // Public defines configuring LLM; see also private defines in LowLevelMemTracker.cpp
 
-// LLM_ALLOW_ASSETS_TAGS: Set to 1 to report the asset in scope for each allocation in addition to the callstack, at the cost of more LLM memory usage per allocation. Set to 0 to omit this information
-// using asset tagging requires a significantly higher number of per-thread tags, so make it optional
-// even if this is on, we still need to run with -llmtagsets=assets because of the sheer number of stat ids it makes
-// LLM Assets can be viewed in game using 'Stat LLMAssets'
+// LLM_ALLOW_ASSETS_TAGS: Set to 1 to enable run-time toggling of AssetTags reporting, 0 to disable.
+// Enabling the define causes extra cputime costs to track costs even when AssetTags are toggled off.
+// When defined on, the feature can be toggled on at runtime with commandline -llmtagsets=assets.
+// Toggling the feature on causes a huge number of stat ids to be created and has a high cputime cost.
+// When defined on and runtime-toggled on, AssetTags report the asset that is in scope for each allocation
+// LLM Assets can be viewed in game using 'Stat LLMAssets'.
 #ifndef LLM_ALLOW_ASSETS_TAGS
 	#define LLM_ALLOW_ASSETS_TAGS 0
 #endif
 
-// LLM_ALLOW_STATS: Set to 1 to allow stats to be used as tags (LLM_SCOPED_TAG_WITH_STAT et al), at the cost of more LLM memory usage per allocation. Set to 0 to disable these stat-specific scope macros.
-// Turning this on uses the same amount of memory per allocation as LLM_ALLOW_NAMES_TAGS. Turning both of them on has no extra cost.
+// LLM_ALLOW_STATS: Set to 1 to allow stats to be used as tags, 0 to disable.
+// When enabled LLM_SCOPED_TAG_WITH_STAT macros are enabled and create an LLM tag per stat at the cost of more LLM
+// memory usage per allocation. Turning this on uses the same amount of memory per allocation as LLM_ALLOW_NAMES_TAGS.
+// Turning both of them on has no extra cost.
 #ifndef LLM_ALLOW_STATS
 	#define LLM_ALLOW_STATS 0
 #endif
 
-// Enable stat tags if: (1) Stats are allowed or (2) Asset tags are allowed (asset tags use the stat macros to record asset scopes)
+// Enable stat tags if: (1) Stats or (2) Asset tags are allowed (asset tags use the stat macros to record asset scopes)
 #define LLM_ENABLED_STAT_TAGS LLM_ALLOW_STATS || LLM_ALLOW_ASSETS_TAGS
 
 #include "Containers/Array.h"
@@ -54,24 +62,22 @@
 
 #if DO_CHECK
 
-	namespace UE
-	{
-	namespace LLMPrivate
-	{
-		bool HandleAssert(bool bLog, const TCHAR* Format, ...);
+namespace UE::LLMPrivate
+{
 
-		// This is used by ensure to generate a bool per instance
-		// by passing a lambda which will uniquely instantiate the template.
-		template <typename Type>
-		bool TrueOnFirstCallOnly(const Type&)
-		{
-			static bool bValue = true;
-			bool Result = bValue;
-			bValue = false;
-			return Result;
-		}
-	}
-	}
+bool HandleAssert(bool bLog, const TCHAR* Format, ...);
+
+// LLMEnsure's use of this generates a bool per callsite by passing a lambda which uniquely instantiates the template.
+template <typename Type>
+bool TrueOnFirstCallOnly(const Type&)
+{
+	static bool bValue = true;
+	bool Result = bValue;
+	bValue = false;
+	return Result;
+}
+
+} // UE::LLMPrivate
 
 #if !USING_CODE_ANALYSIS
 	#define LLMTrueOnFirstCallOnly			UE::LLMPrivate::TrueOnFirstCallOnly([]{})
@@ -85,13 +91,13 @@
 
 #define LLMCheck(expr)					do { if (UNLIKELY(!(expr))) { UE::LLMPrivate::HandleAssert(true, LLMCheckMessage(expr));                         FPlatformMisc::RaiseException(1); } } while(false)
 #define LLMCheckf(expr,format,...)		do { if (UNLIKELY(!(expr))) { UE::LLMPrivate::HandleAssert(true, LLMCheckfMessage(expr, format), ##__VA_ARGS__); FPlatformMisc::RaiseException(1); } } while(false)
-#define LLMEnsure(expr)					(LIKELY(!!(expr)) || UE::LLMPrivate::HandleAssert(LLMTrueOnFirstCallOnly, LLMEnsureMessage(expr)))
+#define LLMEnsure(expr) (LIKELY(!!(expr)) || UE::LLMPrivate::HandleAssert(LLMTrueOnFirstCallOnly, LLMEnsureMessage(expr)))
 
 #else
 
 #define LLMCheck(expr)
 #define LLMCheckf(expr,...)
-#define LLMEnsure(expr)			(!!(expr))
+#define LLMEnsure(expr)	(!!(expr))
 
 #endif
 
@@ -122,7 +128,7 @@ enum class ELLMTagSet : uint8
 	Assets,
 	AssetClasses,
 	
-	Max,	// note: check out FLowLevelMemTracker::ShouldReduceThreads and IsAssetTagForAssets if you add any asset-style tagsets
+	Max, // note: see FLowLevelMemTracker::ShouldReduceThreads and IsAssetTagForAssets if you add any asset-style tagsets
 };
 
 // Do not add to these macros. Please use the LLM_DECLARE_TAG family of macros below to create new tags.
@@ -266,14 +272,14 @@ enum class ELLMTag : LLM_TAG_TYPE
 
 	// anything above this value is treated as an FName for a stat section
 };
-static_assert( ELLMTag::GenericTagCount <= ELLMTag::PlatformTagStart, "too many LLM tags defined -- Instead of adding a new tag and updating the limits, please use the LLM_DECLARE_TAG macros below"); 
+static_assert( ELLMTag::GenericTagCount <= ELLMTag::PlatformTagStart,
+	"too many LLM tags defined -- Instead of adding a new tag and updating the limits, please use the LLM_DECLARE_TAG macros below"); 
 
 static constexpr uint32 LLM_TAG_COUNT = 256;
 static constexpr uint32 LLM_CUSTOM_TAG_START = (int32)ELLMTag::PlatformTagStart;
 static constexpr uint32 LLM_CUSTOM_TAG_END = (int32)ELLMTag::ProjectTagEnd;
 static constexpr uint32 LLM_CUSTOM_TAG_COUNT = LLM_CUSTOM_TAG_END + 1 - LLM_CUSTOM_TAG_START;
 
-
 /**
  * Passed in to OnLowLevelAlloc to specify the type of allocation. Used to track FMalloc total
  * and pausing for a specific allocation type.
@@ -341,24 +347,31 @@ extern FName LLMGetTagStat(ELLMTag Tag);
 
 /**
  * Define a tag which can be used in LLM_SCOPE_BYTAG or referenced by name in other LLM_SCOPEs.
- * @param UniqueNameWithUnderscores - Modified version of the name of the tag. Used for looking up by name, must be unique across all tags passed to LLM_DEFINE_TAG, LLM_SCOPE, or ELLMTag.
- *                                    The modification: the usual separator / for parents must be replaced with _ in LLM_DEFINE_TAGs.
- * @param DisplayName - (Optional) - The name to display when tracing the tag; joined with "/" to the name of its parent if it has a parent, or NAME_None to use the UniqueName.
+ * @param UniqueNameWithUnderscores - Modified version of the name of the tag. Used for looking up by name,
+ *        must be unique across all tags passed to LLM_DEFINE_TAG, LLM_SCOPE, or ELLMTag.
+ *        The modification: the usual separator / for parents must be replaced with _ in LLM_DEFINE_TAGs.
+ * @param DisplayName - (Optional) - The name to display when tracing the tag; joined with "/" to the name of its
+ *        parent if it has a parent, or NAME_None to use the UniqueName.
  * @param ParentTagName - (Optional) - The unique name of the parent tag, or NAME_None if it has no parent.
- * @param StatName - (Optional) - The name of the stat to populate with this tag's amount when publishing LLM data each frame, or NAME_None if no stat should be populated.
- * @param SummaryStatName - (Optional) - The name of the stat group to add on this tag's amount when publishing LLM data each frame, or NAME_None if no stat group should be added to.
+ * @param StatName - (Optional) - The name of the stat to populate with this tag's amount when publishing LLM data each
+ *        frame, or NAME_None if no stat should be populated.
+ * @param SummaryStatName - (Optional) - The name of the stat group to add on this tag's amount when publishing LLM
+ *        data each frame, or NAME_None if no stat group should be added to.
  */
 #define LLM_DEFINE_TAG(UniqueNameWithUnderscores, ...) FLLMTagDeclaration PREPROCESSOR_JOIN(LLMTagDeclaration_, UniqueNameWithUnderscores)(TEXT(#UniqueNameWithUnderscores), ##__VA_ARGS__)
 
  /**
-  * Declare a tag which is defined elsewhere which can used in LLM_SCOPE_BYTAG or referenced by name in other LLM_SCOPEs.
-  * @param UniqueName - The name of the tag for looking up by name, must be unique across all tags passed to LLM_DEFINE_TAG, LLM_SCOPE, or ELLMTag.
-  * @param ModuleName - (Optional, only in LLM_DECLARE_TAG_API) - The MODULENAME_API symbol to use to declare module linkage, aka ENGINE_API. If omitted, no module linkage will be used and the tag will create link errors if used from another module.
+  * Declare a tag defined by LLM_DEFINE_TAG. It is used in LLM_SCOPE_BYTAG or referenced by name in other LLM_SCOPEs.
+  * @param UniqueName - The name of the tag for looking up by name, must be unique across all tags passed to
+  *        LLM_DEFINE_TAG, LLM_SCOPE, or ELLMTag.
+  * @param ModuleName - (Optional, only in LLM_DECLARE_TAG_API) - The MODULENAME_API symbol to use to declare module
+  *        linkage, aka ENGINE_API. If omitted, no module linkage will be used and the tag will create link errors if
+  *        used from another module.
   */
 #define LLM_DECLARE_TAG(UniqueNameWithUnderscores) extern FLLMTagDeclaration PREPROCESSOR_JOIN(LLMTagDeclaration_, UniqueNameWithUnderscores)
 #define LLM_DECLARE_TAG_API(UniqueNameWithUnderscores, ModuleAPI) extern ModuleAPI FLLMTagDeclaration PREPROCESSOR_JOIN(LLMTagDeclaration_, UniqueNameWithUnderscores)
 
-/** Get the unique Name of a Tag, suitable for passing to LLM functions such as OnLowLevelAlloc that take a Tag UniqueName. */
+/** Get the unique Name of a Tag. It can be passed to functions such as OnLowLevelAlloc that take a Tag UniqueName. */
 #define LLM_TAG_NAME(UniqueNameWithUnderscores) (PREPROCESSOR_JOIN(LLMTagDeclaration_, UniqueNameWithUnderscores).GetUniqueName())
 
 /**
@@ -382,100 +395,98 @@ typedef void(*LLMFreeFunction)(void*, size_t);
 
 class FLLMTagDeclaration;
 
-namespace UE
+namespace UE::LLMPrivate
+{
+
+class FLLMCsvProfilerWriter;
+class FLLMCsvWriter;
+class FLLMThreadState;
+class FLLMTraceWriter;
+class FLLMTracker;
+class FTagData;
+class FTagDataArray;
+class FTagDataNameMap;
+
+namespace AllocatorPrivate
 {
-namespace LLMPrivate
+	struct FBin;
+	struct FPage;
+}
+/**
+ * The allocator LLM uses to allocate internal memory. Uses platform defined
+ * allocation functions to grab memory directly from the OS.
+ */
+class FLLMAllocator
 {
+public:
+	FLLMAllocator();
+	~FLLMAllocator();
 
-	class FLLMCsvProfilerWriter;
-	class FLLMCsvWriter;
-	class FLLMThreadState;
-	class FLLMTraceWriter;
-	class FLLMTracker;
-	class FTagData;
-	class FTagDataArray;
-	class FTagDataNameMap;
+	static FLLMAllocator*& Get();
 
-	namespace AllocatorPrivate
+	void Initialise(LLMAllocFunction InAlloc, LLMFreeFunction InFree, int32 InPageSize);
+	void Clear();
+	void* Alloc(size_t Size);
+	void* Malloc(size_t Size);
+	void Free(void* Ptr, size_t Size);
+	void* Realloc(void* Ptr, size_t OldSize, size_t NewSize);
+	int64 GetTotal() const;
+
+	template <typename T, typename... ArgsType>
+	T* New(ArgsType&&... Args)
 	{
-		struct FBin;
-		struct FPage;
+		T* Ptr = reinterpret_cast<T*>(Alloc(sizeof(T)));
+		new (Ptr) T(Forward<ArgsType>(Args)...);
+		return Ptr;
 	}
-	/**
-	 * The allocator LLM uses to allocate internal memory. Uses platform defined
-	 * allocation functions to grab memory directly from the OS.
-	 */
-	class FLLMAllocator
+
+	template <typename T>
+	void Delete(T* Ptr)
 	{
-	public:
-		FLLMAllocator();
-		~FLLMAllocator();
-
-		static FLLMAllocator*& Get();
-
-		void Initialise(LLMAllocFunction InAlloc, LLMFreeFunction InFree, int32 InPageSize);
-		void Clear();
-		void* Alloc(size_t Size);
-		void* Malloc(size_t Size);
-		void Free(void* Ptr, size_t Size);
-		void* Realloc(void* Ptr, size_t OldSize, size_t NewSize);
-		int64 GetTotal() const;
-
-		template <typename T, typename... ArgsType>
-		T* New(ArgsType&&... Args)
+		if (Ptr)
 		{
-			T* Ptr = reinterpret_cast<T*>(Alloc(sizeof(T)));
-			new (Ptr) T(Forward<ArgsType>(Args)...);
-			return Ptr;
+			Ptr->~T();
+			Free(Ptr, sizeof(T));
 		}
+	}
 
-		template <typename T>
-		void Delete(T* Ptr)
-		{
-			if (Ptr)
-			{
-				Ptr->~T();
-				Free(Ptr, sizeof(T));
-			}
-		}
+	LLMAllocFunction GetPlatformAlloc()
+	{
+		return PlatformAlloc;
+	}
+	LLMFreeFunction GetPlatformFree()
+	{
+		return PlatformFree;
+	}
 
-		LLMAllocFunction GetPlatformAlloc()
-		{
-			return PlatformAlloc;
-		}
-		LLMFreeFunction GetPlatformFree()
-		{
-			return PlatformFree;
-		}
+private:
+	void* AllocPages(size_t Size);
+	void FreePages(void* Ptr, size_t Size);
+	int32 GetBinIndex(size_t Size) const;
+
+	FCriticalSection CriticalSection;
+	LLMAllocFunction PlatformAlloc;
+	LLMFreeFunction PlatformFree;
+	AllocatorPrivate::FBin* Bins;
+	int64 Total;
+	int32 PageSize;
+	int32 NumBins;
+
+	friend struct UE::LLMPrivate::AllocatorPrivate::FPage;
+	friend struct UE::LLMPrivate::AllocatorPrivate::FBin;
+};
 
-	private:
-		void* AllocPages(size_t Size);
-		void FreePages(void* Ptr, size_t Size);
-		int32 GetBinIndex(size_t Size) const;
-
-		FCriticalSection CriticalSection;
-		LLMAllocFunction PlatformAlloc;
-		LLMFreeFunction PlatformFree;
-		AllocatorPrivate::FBin* Bins;
-		int64 Total;
-		int32 PageSize;
-		int32 NumBins;
-
-		friend struct UE::LLMPrivate::AllocatorPrivate::FPage;
-		friend struct UE::LLMPrivate::AllocatorPrivate::FBin;
-	};
+enum class ETagReferenceSource
+{
+	Scope,
+	Declare,
+	EnumTag,
+	CustomEnumTag,
+	FunctionAPI,
+	ImplicitParent
+};
 
-	enum class ETagReferenceSource
-	{
-		Scope,
-		Declare,
-		EnumTag,
-		CustomEnumTag,
-		FunctionAPI,
-		ImplicitParent
-	};
-}
-}
+} // UE::LLMPrivate
 
 struct UE_DEPRECATED(4.27, "FLLMCustomTag was an implementation detail that has been modified, switch to FLLMTagInfo or to your own local struct") FLLMCustomTag
 {
@@ -485,7 +496,7 @@ struct UE_DEPRECATED(4.27, "FLLMCustomTag was an implementation detail that has
 	FName SummaryStatName;
 };
 
-/* A convenient struct for gathering the fields needed to report in RegisterProjectTag */
+/** A convenient struct for gathering the fields needed to report in RegisterProjectTag */
 struct FLLMTagInfo
 {
 	const TCHAR* Name;
@@ -494,14 +505,12 @@ struct FLLMTagInfo
 	int32 ParentTag = -1;
 };
 
-/*
- * The main LLM tracker class
- */
+/** The main LLM tracker class. */
 class CORE_API FLowLevelMemTracker
 {
 public:
 
-	// get the singleton, which makes sure that we always have a valid object
+	/** Get the singleton, which makes sure that we always have a valid object */
 	inline static FLowLevelMemTracker& Get()
 	{
 		if (TrackerInstance)
@@ -514,86 +523,112 @@ public:
 
 	static bool IsEnabled();
 
-	// we always start up running, but if the commandline disables us, we will do it later after main
-	// (can't get the commandline early enough in a cross-platform way)
+	/**
+	 * We always start up running, but if the commandline disables us, we will do it later after main
+	 * (can't get the commandline early enough in a cross-platform way).
+	 */
 	void ProcessCommandLine(const TCHAR* CmdLine);
 
-	// Return the total amount of memory being tracked
+	/** Return the total amount of memory being tracked. */
 	uint64 GetTotalTrackedMemory(ELLMTracker Tracker);
 
-	// this is the main entry point for the class - used to track any pointer that was allocated or freed 
-	void OnLowLevelAlloc(ELLMTracker Tracker, const void* Ptr, uint64 Size, ELLMTag DefaultTag = ELLMTag::Untagged, ELLMAllocType AllocType = ELLMAllocType::None, bool bTrackInMemPro = true);		// DefaultTag is used if no other tag is set
-	void OnLowLevelAlloc(ELLMTracker Tracker, const void* Ptr, uint64 Size, FName DefaultTag, ELLMAllocType AllocType = ELLMAllocType::None, bool bTrackInMemPro = true);		// DefaultTag is used if no other tag is set
-	void OnLowLevelFree(ELLMTracker Tracker, const void* Ptr, ELLMAllocType AllocType = ELLMAllocType::None, bool bTrackInMemPro = true);
+	/**
+	 * Records the use of memory that was allocated for a pointer.
+	 *
+	 * @param Tracker - Which tracker to use, the high-level default tracker for regular engine allocations, or the
+	 *                  low-level platform tracker for memory management systems.
+	 * @param Ptr - The pointer that was allocated.
+	 * @param Size - The size of the memory that was allocated for the pointer.
+	 * @param DefaultTag - The tag to use if there is no tag already in scope on the callstack.
+	 * @param AllocType - Type of allocation, FMalloc for regular allocations done through GMalloc, System for others.
+	 * @param bTrackInMemPro - Whether to pass the allocation/free information on to FMemProProfiler
+	 */
+	void OnLowLevelAlloc(ELLMTracker Tracker, const void* Ptr, uint64 Size, ELLMTag DefaultTag = ELLMTag::Untagged,
+		ELLMAllocType AllocType = ELLMAllocType::None, bool bTrackInMemPro = true);
+	void OnLowLevelAlloc(ELLMTracker Tracker, const void* Ptr, uint64 Size, FName DefaultTag,
+		ELLMAllocType AllocType = ELLMAllocType::None, bool bTrackInMemPro = true);
+
+	/**
+	 * Records the release of memory that was allocated for a pointer.
+	 *
+	 * @param Tracker - Which tracker to use, the high-level default tracker for regular engine allocations, or the
+	 *                  low-level platform tracker for memory management systems.
+	 * @param Ptr - The pointer that was freed.
+	 * @param AllocType - Type of allocation, FMalloc for regular allocations done through GMalloc, System for others.
+	 * @param bTrackInMemPro - Whether to pass the allocation/free information on to FMemProProfiler
+	 */
+	void OnLowLevelFree(ELLMTracker Tracker, const void* Ptr,
+		ELLMAllocType AllocType = ELLMAllocType::None, bool bTrackInMemPro = true);
 
-	// call if an allocation is moved in memory, such as in a defragger
-	void OnLowLevelAllocMoved(ELLMTracker Tracker, const void* Dest, const void* Source, ELLMAllocType AllocType = ELLMAllocType::None);
+	/** Call if an allocation is moved in memory, such as in a defragger. */
+	void OnLowLevelAllocMoved(ELLMTracker Tracker, const void* Dest, const void* Source,
+		ELLMAllocType AllocType = ELLMAllocType::None);
 
-	// expected to be called once a frame, from game thread or similar - updates memory stats 
+	/** Updates memory stats and optionally publishes them to observers. */
 	void UpdateStatsPerFrame(const TCHAR* LogName=nullptr);
-	/** A tick function that can be called as frequently as necessary rather than once per frame; this is sometimes necessary when tracking large amounts of tags that have a superlinear update cost */
+	/** Updates memory stats. */
 	void Tick();
 
-	// Optionally set the amount of memory taken up before the game starts for executable and data segments
+	/** Optionally set the amount of memory taken up before the game starts for executable and data segments .*/
 	void SetProgramSize(uint64 InProgramSize);
 
-	// console command handler
+	/** console command handler */
 	bool Exec(const TCHAR* Cmd, FOutputDevice& Ar);
 
-	// are we in the more intensive asset tracking mode, and is it active
+	/** Report whether the given TagSet is active, e.g. AssetTracking. */
 	bool IsTagSetActive(ELLMTagSet Set);
 
-	// for some tag sets, it's really useful to reduce threads, to attribute allocations to assets, for instance
+	/** For some tag sets, it's really useful to reduce threads, to attribute allocations to assets, for instance. */
 	bool ShouldReduceThreads();
 
-	// get the top active tag for the given tracker
+	/** Get the top active tag for the given tracker */
 	UE_DEPRECATED(4.27, "Tags have been changed to FNames and the old ELLMTag is now only the top-level coarse tag. Use GetActivateTagData instead to get the current Tag instead of its toplevel parent.")
 	int64 GetActiveTag(ELLMTracker Tracker);
 
-	/** Get an opaque identifier for the top active tag for the given tracker */
+	/** Get an opaque identifier for the top active tag for the given tracker. */
 	const UE::LLMPrivate::FTagData* GetActiveTagData(ELLMTracker Tracker);
 
-	// register custom ELLMTags
+	/** Register custom ELLMTags. */
 	void RegisterPlatformTag(int32 Tag, const TCHAR* Name, FName StatName, FName SummaryStatName, int32 ParentTag = -1);
 	void RegisterProjectTag(int32 Tag, const TCHAR* Name, FName StatName, FName SummaryStatName, int32 ParentTag = -1);
     
-	// Get all tags being tracked
+	/** Get all tags being tracked. */
 	TArray<const UE::LLMPrivate::FTagData*> GetTrackedTags();
 
-	// Get all tags being tracked by the given tracker
+	/** Get all tags being tracked by the given tracker. */
 	TArray<const UE::LLMPrivate::FTagData*> GetTrackedTags(ELLMTracker Tracker);
 
-	// look up the ELLMTag associated with the given display name
+	/** Look up the ELLMTag associated with the given display name. */
 	bool FindTagByName(const TCHAR* Name, uint64& OutTag) const;
 
 	UE_DEPRECATED(4.27, "Use FindTagDisplayName instead")
 	const TCHAR* FindTagName(uint64 Tag) const;
 
-	// get the display name for the given ELLMTag
+	/** Get the display name for the given ELLMTag. */
 	FName FindTagDisplayName(uint64 Tag) const;
 
-	// Get the display name for the given FTagData
+	/** Get the display name for the given FTagData. */
 	FName GetTagDisplayName(const UE::LLMPrivate::FTagData* TagData) const;
 
-	// Get the path name for the given FTagData from a chain of its parents' display names
+	/** Get the path name for the given FTagData from a chain of its parents' display names. */
 	FString GetTagDisplayPathName(const UE::LLMPrivate::FTagData* TagData) const;
 
-	// Get the unique identifier name for the given FTagData
+	/** Get the unique identifier name for the given FTagData. */
 	FName GetTagUniqueName(const UE::LLMPrivate::FTagData* TagData) const;
 
-	// Get the amount of memory for an ELLMTag from the given tracker
+	/** Get the amount of memory for an ELLMTag from the given tracker. */
 	int64 GetTagAmountForTracker(ELLMTracker Tracker, ELLMTag Tag, bool bPeakAmount = false);
 
-	// Get the amount of memory for a FTagData from the given tracker
+	/** Get the amount of memory for a FTagData from the given tracker. */
 	int64 GetTagAmountForTracker(ELLMTracker Tracker, const UE::LLMPrivate::FTagData* TagData, bool bPeakAmount = false);
 
-	// Set the amount of memory for an ELLMTag for a given tracker, optionally updating the total tracked memory too
+	/** Set the amount of memory for an ELLMTag for a given tracker and optionally update the total tracked memory. */
 	void SetTagAmountForTracker(ELLMTracker Tracker, ELLMTag Tag, int64 Amount, bool bAddToTotal);
 
-	// Dump the display name of the current TagData for the given tracker to the output
+	/** Dump the display name of the current TagData for the given tracker to the output. */
 	uint64 DumpTag(ELLMTracker Tracker, const char* FileName, int LineNumber);
 
-	// Publishes the active LLM stats in the active frame, useful for single targeted LLM snapshots
+	/** Publishes the active LLM stats in the active frame, useful for single targeted LLM snapshots. */
 	void PublishDataSingleFrame();
 
 	enum class EDumpFormat
@@ -610,13 +645,15 @@ private:
 
 	~FLowLevelMemTracker();
 
-	/** Allocation and Setup of the data required for tracking allocations is done as late as possible to prevent exercising allocation code too early
-	 * Note that as late as possible for tracking allocations is still earlier than ParseCommandLine, so we do not complete all initialisation in this function
-	 * (e.g. features required only for Update are omitted), and whatever we do initialise here will be torn down later in ParseCommandLine if LLM is disabled
+	/**
+	 * Allocation and Setup of the data required for tracking allocations is done as late as possible to prevent
+	 * exercising allocation code too early. Note that as late as possible for tracking allocations is still earlier
+	 * than ParseCommandLine, so we do not complete all initialization in this function (e.g. features required only
+	 * for Update are omitted). Initialization done here will be torn down in ParseCommandLine if LLM is disabled.
 	 */
 	void BootstrapInitialise();
 
-	/** Free all memory. This will put the tracker into a permanently disabled state */
+	/** Free all memory. This will put the tracker into a permanently disabled state. */
 	void Clear();
 	void InitialiseProgramSize();
 
@@ -696,9 +733,7 @@ public: // really internal but needs to be visible for LLM_IF_ENABLED macro
 	static bool bIsDisabled;
 };
 
-/*
- * LLM scope for tracking memory
- */
+/** LLM scope for tracking memory. */
 class CORE_API FLLMScope
 {
 public:
@@ -747,9 +782,7 @@ protected:
 #endif
 };
 
-/*
-* LLM scope for pausing LLM (disables the allocation hooks)
-*/
+/** LLM scope for pausing LLM (disables the allocation hooks). */
 class CORE_API FLLMPauseScope
 {
 public:
@@ -757,15 +790,14 @@ public:
 	FLLMPauseScope(ELLMTag TagEnum, bool bIsStatTag, uint64 Amount, ELLMTracker TrackerToPause, ELLMAllocType InAllocType);
 	~FLLMPauseScope();
 protected:
-	void Init(FName TagName, ELLMTag EnumTag, bool bIsEnumTag, bool bIsStatTag, uint64 Amount, ELLMTracker TrackerToPause, ELLMAllocType InAllocType);
+	void Init(FName TagName, ELLMTag EnumTag, bool bIsEnumTag, bool bIsStatTag, uint64 Amount,
+		ELLMTracker TrackerToPause, ELLMAllocType InAllocType);
 	ELLMTracker PausedTracker;
 	ELLMAllocType AllocType;
 	bool bEnabled;
 };
 
-/*
-* LLM scope for inheriting tag from the given address
-*/
+/** LLM scope for inheriting tag from the given address. */
 class CORE_API FLLMScopeFromPtr
 {
 public:
@@ -776,9 +808,7 @@ protected:
 	bool bEnabled;
 };
 
-/**
- * Global instances to provide information about a tag to LLM
- */
+/** Global instances to provide information about a tag to LLM. */
 class CORE_API FLLMTagDeclaration
 {
 public:
@@ -813,27 +843,27 @@ protected:
 
 #else
 
-	#define LLM(...)
-	#define LLM_IF_ENABLED(...)
-	#define LLM_SCOPE(...)
-	#define LLM_SCOPE_BYNAME(...)
-	#define LLM_SCOPE_BYTAG(...)
-	#define LLM_PLATFORM_SCOPE(...)
-	#define LLM_PLATFORM_SCOPE_BYNAME(...)
-	#define LLM_PLATFORM_SCOPE_BYTAG(...)
-	#define LLM_REALLOC_SCOPE(...)
-	#define LLM_REALLOC_PLATFORM_SCOPE(...)
-	#define LLM_SCOPED_PAUSE_TRACKING(...)
-	#define LLM_SCOPED_PAUSE_TRACKING_FOR_TRACKER(...)
-	#define LLM_SCOPED_PAUSE_TRACKING_WITH_ENUM_AND_AMOUNT(...)
-	#define LLM_DUMP_TAG()
-	#define LLM_DUMP_PLATFORM_TAG()
-	#define LLM_DEFINE_TAG(...)
-	#define LLM_DECLARE_TAG(...)
-	#define LLM_DECLARE_TAG_API(...)
-	#define LLM_DEFINE_BOOTSTRAP_TAG(...)
-	#define LLM_SCOPE_BY_BOOTSTRAP_TAG(...)
-	#define LLM_DECLARE_BOOTSTRAP_TAG(...) 
-	#define LLM_DECLARE_BOOTSTRAP_TAG_API(...)
+#define LLM(...)
+#define LLM_IF_ENABLED(...)
+#define LLM_SCOPE(...)
+#define LLM_SCOPE_BYNAME(...)
+#define LLM_SCOPE_BYTAG(...)
+#define LLM_PLATFORM_SCOPE(...)
+#define LLM_PLATFORM_SCOPE_BYNAME(...)
+#define LLM_PLATFORM_SCOPE_BYTAG(...)
+#define LLM_REALLOC_SCOPE(...)
+#define LLM_REALLOC_PLATFORM_SCOPE(...)
+#define LLM_SCOPED_PAUSE_TRACKING(...)
+#define LLM_SCOPED_PAUSE_TRACKING_FOR_TRACKER(...)
+#define LLM_SCOPED_PAUSE_TRACKING_WITH_ENUM_AND_AMOUNT(...)
+#define LLM_DUMP_TAG()
+#define LLM_DUMP_PLATFORM_TAG()
+#define LLM_DEFINE_TAG(...)
+#define LLM_DECLARE_TAG(...)
+#define LLM_DECLARE_TAG_API(...)
+#define LLM_DEFINE_BOOTSTRAP_TAG(...)
+#define LLM_SCOPE_BY_BOOTSTRAP_TAG(...)
+#define LLM_DECLARE_BOOTSTRAP_TAG(...) 
+#define LLM_DECLARE_BOOTSTRAP_TAG_API(...)
 
 #endif		// #if ENABLE_LOW_LEVEL_MEM_TRACKER
