LowLevelMemTracker: Fix corrupted whitespace from bad resolve.
#rnx

diff --git a/Engine/Source/Runtime/Core/Private/HAL/LowLevelMemTracker.cpp b/Engine/Source/Runtime/Core/Private/HAL/LowLevelMemTracker.cpp
--- a/Engine/Source/Runtime/Core/Private/HAL/LowLevelMemTracker.cpp
+++ b/Engine/Source/Runtime/Core/Private/HAL/LowLevelMemTracker.cpp
@@ -298,284 +298,282 @@ namespace UE::LLMPrivate
 {
 
 /** FLLMCsvWriter: class for writing out the LLM tag sizes to a csv file every few seconds. */
-	class FLLMCsvWriter
-	{
-	public:
-		FLLMCsvWriter();
-		~FLLMCsvWriter();
+class FLLMCsvWriter
+{
+public:
+	FLLMCsvWriter();
+	~FLLMCsvWriter();
 
-		void SetTracker(ELLMTracker InTracker);
-		void Clear();
+	void SetTracker(ELLMTracker InTracker);
+	void Clear();
 
 	void Publish(FLowLevelMemTracker& LLMRef, const FTrackerTagSizeMap& TagSizes,
 		const FTagData* OverrideTrackedTotalTagData, const FTagData* OverrideUntaggedTagData, int64 TrackedTotal,
 		bool bTrackPeaks);
 
-		void OnPreFork();
+	void OnPreFork();
 
-	private:
-		void Write(FStringView Text);
-		static const TCHAR* GetTrackerCsvName(ELLMTracker InTracker);
+private:
+	void Write(FStringView Text);
+	static const TCHAR* GetTrackerCsvName(ELLMTracker InTracker);
 
-		bool CreateArchive();
-		bool UpdateColumns(const FTrackerTagSizeMap& TagSizes);
-		void WriteHeader(const FTagData* OverrideTrackedTotalTagData, const FTagData* OverrideUntaggedTagData);
+	bool CreateArchive();
+	bool UpdateColumns(const FTrackerTagSizeMap& TagSizes);
+	void WriteHeader(const FTagData* OverrideTrackedTotalTagData, const FTagData* OverrideUntaggedTagData);
 	void AddRow(FLowLevelMemTracker& LLMRef, const FTrackerTagSizeMap& TagSizes,
 		const FTagData* OverrideTrackedTotalName, const FTagData* OverrideUntaggedName, int64 TrackedTotal,
 		bool bTrackPeaks);
 
-		FConstTagDataArray Columns;
-		TFastPointerLLMSet<const FTagData*> ExistingColumns;
-		FArchive* Archive;
-		double LastWriteTime;
-		int32 WriteCount;
-		ELLMTracker Tracker;
-	};
+	FConstTagDataArray Columns;
+	TFastPointerLLMSet<const FTagData*> ExistingColumns;
+	FArchive* Archive;
+	double LastWriteTime;
+	int32 WriteCount;
+	ELLMTracker Tracker;
+};
 
 /** Outputs the LLM tags and sizes to TraceLog events. */
-	class FLLMTraceWriter
-	{
-	public:
-		FLLMTraceWriter();
-		void SetTracker(ELLMTracker InTracker);
-		void Clear();
+class FLLMTraceWriter
+{
+public:
+	FLLMTraceWriter();
+	void SetTracker(ELLMTracker InTracker);
+	void Clear();
 	void Publish(FLowLevelMemTracker& LLMRef, const FTrackerTagSizeMap& TagSizes,
 		const FTagData* OverrideTrackedTotalName, const FTagData* OverrideUntaggedName, int64 TrackedTotal,
 		bool bTrackPeaks);
-		static void TraceGenericTags(FLowLevelMemTracker& LLMRef);
+	static void TraceGenericTags(FLowLevelMemTracker& LLMRef);
 
-	private:
-		static const void* GetTagId(const FTagData* TagData);
+private:
+	static const void* GetTagId(const FTagData* TagData);
 
-		ELLMTracker				Tracker;
-		TFastPointerLLMSet<const FTagData*> DeclaredTags;
-		bool					bTrackerSpecSent = false;
-	};
+	ELLMTracker				Tracker;
+	TFastPointerLLMSet<const FTagData*> DeclaredTags;
+	bool					bTrackerSpecSent = false;
+};
 
 /** FLLMCsvWriter: class for writing out LLM stats to the Csv Profiler. */
-	class FLLMCsvProfilerWriter
-	{
-	public:
-		FLLMCsvProfilerWriter();
-		void Clear();
-		void SetTracker(ELLMTracker InTracker);
+class FLLMCsvProfilerWriter
+{
+public:
+	FLLMCsvProfilerWriter();
+	void Clear();
+	void SetTracker(ELLMTracker InTracker);
 	void Publish(FLowLevelMemTracker& LLMRef, const FTrackerTagSizeMap& TagSizes,
 		const FTagData* OverrideTrackedTotalName, const FTagData* OverrideUntaggedName, int64 TrackedTotal,
 		bool bTrackPeaks);
-	protected:
-		void RecordTagToCsv(int32 CsvCategoryIndex, const FTagData* TagData, int64 TagSize);
+protected:
+	void RecordTagToCsv(int32 CsvCategoryIndex, const FTagData* TagData, int64 TagSize);
 #if LLM_CSV_PROFILER_WRITER_ENABLED 
-		ELLMTracker Tracker;
-		TFastPointerLLMMap<const FTagData*, FName> TagDataToCsvStatName;
+	ELLMTracker Tracker;
+	TFastPointerLLMMap<const FTagData*, FName> TagDataToCsvStatName;
 #endif
-	};
-
+};
 
 /** Per-thread state in an LLMTracker. */
-	class FLLMThreadState
-	{
-	public:
-		FLLMThreadState();
+class FLLMThreadState
+{
+public:
+	FLLMThreadState();
 
-		void Clear();
+	void Clear();
 
-		void PushTag(const FTagData* TagData);
-		void PopTag();
-		const FTagData* GetTopTag();
+	void PushTag(const FTagData* TagData);
+	void PopTag();
+	const FTagData* GetTopTag();
 #if LLM_ALLOW_ASSETS_TAGS
-		void PushAssetTag(const FTagData* TagData);
-		void PopAssetTag();
-		const FTagData* GetTopAssetTag();
+	void PushAssetTag(const FTagData* TagData);
+	void PopAssetTag();
+	const FTagData* GetTopAssetTag();
 #endif
 	void TrackAllocation(const void* Ptr, int64 Size, ELLMTracker Tracker, ELLMAllocType AllocType,
 		const FTagData* TagData, const FTagData* AssetTagData, bool bTrackInMemPro);
 	void TrackFree(const void* Ptr, int64 Size, ELLMTracker Tracker, ELLMAllocType AllocType,
-		const FTagData* TagData, const FTagData* AssetTagData, bool bTrackInMemPro);
-		void TrackMoved(const void* Dest, const void* Source, int64 Size, ELLMTracker Tracker, const FTagData* TagData);
-		void IncrTag(const FTagData* Tag, int64 Amount);
+	const FTagData* TagData, const FTagData* AssetTagData, bool bTrackInMemPro);
+	void TrackMoved(const void* Dest, const void* Source, int64 Size, ELLMTracker Tracker, const FTagData* TagData);
+	void IncrTag(const FTagData* Tag, int64 Amount);
 
-		void PropagateChildSizesToParents();
-		void OnTagsResorted(FTagDataArray& OldTagDatas);
-		void LockTags(bool bLock);
-		void FetchAndClearTagSizes(FTrackerTagSizeMap& TagSizes, int64* InAllocTypeAmounts, bool bTrimAllocations);
+	void PropagateChildSizesToParents();
+	void OnTagsResorted(FTagDataArray& OldTagDatas);
+	void LockTags(bool bLock);
+	void FetchAndClearTagSizes(FTrackerTagSizeMap& TagSizes, int64* InAllocTypeAmounts, bool bTrimAllocations);
 
-		void ClearAllocTypeAmounts();
+	void ClearAllocTypeAmounts();
 
-		FConstTagDataArray TagStack;
+	FConstTagDataArray TagStack;
 #if LLM_ALLOW_ASSETS_TAGS
-		FConstTagDataArray AssetTagStack;
+	FConstTagDataArray AssetTagStack;
 #endif
 
-		FThreadTagSizeMap Allocations;
-		FCriticalSection TagSection;
+	FThreadTagSizeMap Allocations;
+	FCriticalSection TagSection;
 
-		int8 PausedCounter[(int32)ELLMAllocType::Count];
-		int64 AllocTypeAmounts[(int32)ELLMAllocType::Count];
-	};
+	int8 PausedCounter[(int32)ELLMAllocType::Count];
+	int64 AllocTypeAmounts[(int32)ELLMAllocType::Count];
+};
 
 /** The main LLM implementation class. It owns the thread state objects. */
-	class FLLMTracker
-	{
-	public:
+class FLLMTracker
+{
+public:
 
-		FLLMTracker(FLowLevelMemTracker& InLLM);
-		~FLLMTracker();
+	FLLMTracker(FLowLevelMemTracker& InLLM);
+	~FLLMTracker();
 
-		void Initialise(ELLMTracker Tracker, FLLMAllocator* InAllocator);
+	void Initialise(ELLMTracker Tracker, FLLMAllocator* InAllocator);
 
-		void PushTag(ELLMTag EnumTag);
-		void PushTag(FName Tag, bool bInIsStatTag);
-		void PushTag(const FTagData* TagData);
-		void PopTag();
+	void PushTag(ELLMTag EnumTag);
+	void PushTag(FName Tag, bool bInIsStatTag);
+	void PushTag(const FTagData* TagData);
+	void PopTag();
 #if LLM_ALLOW_ASSETS_TAGS
-		void PushAssetTag(FName Tag);
-		void PushAssetTag(const FTagData* Tag);
-		void PopAssetTag();
+	void PushAssetTag(FName Tag);
+	void PushAssetTag(const FTagData* Tag);
+	void PopAssetTag();
 #endif
 	void TrackAllocation(const void* Ptr, int64 Size, FName DefaultTag, ELLMTracker Tracker, ELLMAllocType AllocType,
 		bool bTrackInMemPro);
 	void TrackAllocation(const void* Ptr, int64 Size, ELLMTag DefaultTag, ELLMTracker Tracker, ELLMAllocType AllocType,
 		bool bTrackInMemPro);
-		void TrackFree(const void* Ptr, ELLMTracker Tracker, ELLMAllocType AllocType, bool bTrackInMemPro);
-		void OnAllocMoved(const void* Dest, const void* Source, ELLMTracker Tracker, ELLMAllocType AllocType);
+	void TrackFree(const void* Ptr, ELLMTracker Tracker, ELLMAllocType AllocType, bool bTrackInMemPro);
+	void OnAllocMoved(const void* Dest, const void* Source, ELLMTracker Tracker, ELLMAllocType AllocType);
 
-		void TrackMemory(ELLMTag EnumTag, int64 Amount, ELLMAllocType AllocType);
-		void TrackMemory(FName TagName, int64 Amount, ELLMAllocType AllocType);
-		void TrackMemory(const FTagData* TagData, int64 Amount, ELLMAllocType AllocType);
+	void TrackMemory(ELLMTag EnumTag, int64 Amount, ELLMAllocType AllocType);
+	void TrackMemory(FName TagName, int64 Amount, ELLMAllocType AllocType);
+	void TrackMemory(const FTagData* TagData, int64 Amount, ELLMAllocType AllocType);
 
-		// This will pause/unpause tracking, and also manually increment a given tag
-		void PauseAndTrackMemory(FName TagName, bool bInIsStatTag, int64 Amount, ELLMAllocType AllocType);
-		void PauseAndTrackMemory(ELLMTag EnumTag, int64 Amount, ELLMAllocType AllocType);
-		void PauseAndTrackMemory(const FTagData* TagData, int64 Amount, ELLMAllocType AllocType);
-		void Pause(ELLMAllocType AllocType);
-		void Unpause(ELLMAllocType AllocType);
-		bool IsPaused(ELLMAllocType AllocType);
+	// This will pause/unpause tracking, and also manually increment a given tag
+	void PauseAndTrackMemory(FName TagName, bool bInIsStatTag, int64 Amount, ELLMAllocType AllocType);
+	void PauseAndTrackMemory(ELLMTag EnumTag, int64 Amount, ELLMAllocType AllocType);
+	void PauseAndTrackMemory(const FTagData* TagData, int64 Amount, ELLMAllocType AllocType);
+	void Pause(ELLMAllocType AllocType);
+	void Unpause(ELLMAllocType AllocType);
+	bool IsPaused(ELLMAllocType AllocType);
 
-		void Clear();
+	void Clear();
 
-		void PublishStats(bool bTrackPeaks);
-		void PublishCsv(bool bTrackPeaks);
-		void PublishTrace(bool bTrackPeaks);
-		void PublishCsvProfiler(bool bTrackPeaks);
+	void PublishStats(bool bTrackPeaks);
+	void PublishCsv(bool bTrackPeaks);
+	void PublishTrace(bool bTrackPeaks);
+	void PublishCsvProfiler(bool bTrackPeaks);
 
-		void OnPreFork();
+	void OnPreFork();
 
-		struct FLowLevelAllocInfo
+	struct FLowLevelAllocInfo
+	{
+	public:
+		void SetTag(const FTagData* InTag, FLowLevelMemTracker& InLLMRef)
 		{
-		public:
-			void SetTag(const FTagData* InTag, FLowLevelMemTracker& InLLMRef)
-			{
 #if LLM_ENABLED_FULL_TAGS
-				Tag = InTag->GetIndex();
+			Tag = InTag->GetIndex();
 #else
-				Tag = InTag->GetEnumTag();
+			Tag = InTag->GetEnumTag();
 #endif
-			}
+		}
 
-			const FTagData* GetTag(FLowLevelMemTracker& InLLMRef) const
-			{
+		const FTagData* GetTag(FLowLevelMemTracker& InLLMRef) const
+		{
 #if LLM_ENABLED_FULL_TAGS
-				FReadScopeLock TagDataScopeLock(InLLMRef.TagDataLock);
-				return (*InLLMRef.TagDatas)[Tag];
+			FReadScopeLock TagDataScopeLock(InLLMRef.TagDataLock);
+			return (*InLLMRef.TagDatas)[Tag];
 #else
-				return InLLMRef.FindTagData(Tag);
+			return InLLMRef.FindTagData(Tag);
 #endif
-			}
+		}
 
 #if LLM_ALLOW_ASSETS_TAGS
-			void SetAssetTag(const FTagData* InTag, FLowLevelMemTracker& InLLMRef)
-			{
-				AssetTag = InTag;
-			}
+		void SetAssetTag(const FTagData* InTag, FLowLevelMemTracker& InLLMRef)
+		{
+			AssetTag = InTag;
+		}
 
-			const FTagData* GetAssetTag(FLowLevelMemTracker& InLLMRef) const
-			{
-				return AssetTag;
-			}
+		const FTagData* GetAssetTag(FLowLevelMemTracker& InLLMRef) const
+		{
+			return AssetTag;
+		}
 #endif
 
 #if LLM_ENABLED_FULL_TAGS
-			int32 GetCompressedTag() { return Tag; }
-			void SetCompressedTag(int32 InTag) { Tag = InTag; }
+		int32 GetCompressedTag() { return Tag; }
+		void SetCompressedTag(int32 InTag) { Tag = InTag; }
 #else
-			ELLMTag GetCompressedTag() { return Tag; }
-			void SetCompressedTag(ELLMTag InTag) { Tag = InTag; }
+		ELLMTag GetCompressedTag() { return Tag; }
+		void SetCompressedTag(ELLMTag InTag) { Tag = InTag; }
 #endif
 
-		private:
+	private:
 #if LLM_ALLOW_ASSETS_TAGS
-			const FTagData* AssetTag = nullptr;
+		const FTagData* AssetTag = nullptr;
 #endif
 #if LLM_ENABLED_FULL_TAGS
 		// Even with arbitrary tags we are still partially compressed - the allocation records the tag's index
 		// (4 bytes) rather than the full tag pointer (8 bytes) or the tag's name (12 bytes).
-			int32 Tag = 0;
+		int32 Tag = 0;
 #else
-			ELLMTag Tag = ELLMTag::Untagged;
+		ELLMTag Tag = ELLMTag::Untagged;
 #endif
-		};
-
-		typedef LLMMap<PointerKey, uint32, FLowLevelAllocInfo, LLMNumAllocsType> FLLMAllocMap;	// pointer, size, info, Capacity SizeType
+	};
 
-		void SetTotalTags(const FTagData* InOverrideUntaggedTagData, const FTagData* InOverrideTrackedTotalTagData);
-		void Update();
-		void UpdateThreads();
-		void OnTagsResorted(FTagDataArray& OldTagDatas);
-		void LockAllThreadTags(bool bLock);
+	typedef LLMMap<PointerKey, uint32, FLowLevelAllocInfo, LLMNumAllocsType> FLLMAllocMap;	// pointer, size, info, Capacity SizeType
 
-		int64 GetTagAmount(const FTagData* TagData, bool bPeakAmount = false) const;
-		void SetTagAmountExternal(const FTagData* TagData, int64 Amount, bool bAddToTotal);
-		void SetTagAmountInUpdate(const FTagData* TagData, int64 Amount, bool bAddToTotal);
-		const FTagData* GetActiveTagData();
-		const FTagData* FindTagForPtr(void* Ptr);
-		TArray<const FTagData*> GetTagDatas();
+	void SetTotalTags(const FTagData* InOverrideUntaggedTagData, const FTagData* InOverrideTrackedTotalTagData);
+	void Update();
+	void UpdateThreads();
+	void OnTagsResorted(FTagDataArray& OldTagDatas);
+	void LockAllThreadTags(bool bLock);
 
-		int64 GetAllocTypeAmount(ELLMAllocType AllocType);
+	int64 GetTagAmount(const FTagData* TagData, bool bPeakAmount = false) const;
+	void SetTagAmountExternal(const FTagData* TagData, int64 Amount, bool bAddToTotal);
+	void SetTagAmountInUpdate(const FTagData* TagData, int64 Amount, bool bAddToTotal);
+	const FTagData* GetActiveTagData();
+	const FTagData* FindTagForPtr(void* Ptr);
+	TArray<const FTagData*> GetTagDatas();
 
-		int64 GetTrackedTotal() const
-		{
-			return TrackedTotal;
-		}
+	int64 GetAllocTypeAmount(ELLMAllocType AllocType);
 
-	protected:
+	int64 GetTrackedTotal() const
+	{
+		return TrackedTotal;
+	}
 
-		FLLMThreadState* GetOrCreateState();
-		FLLMThreadState* GetState();
+protected:
+	FLLMThreadState* GetOrCreateState();
+	FLLMThreadState* GetState();
 	void TrackAllocation(const void* Ptr, int64 Size, const FTagData* ActiveTagData, ELLMTracker Tracker,
 		ELLMAllocType AllocType, FLLMThreadState* State, bool bTrackInMemPro);
 
-		FLowLevelMemTracker& LLMRef;
+	FLowLevelMemTracker& LLMRef;
 
-		uint32 TlsSlot;
+	uint32 TlsSlot;
 
-		TArray<FLLMThreadState*, FDefaultLLMAllocator> ThreadStates;
+	TArray<FLLMThreadState*, FDefaultLLMAllocator> ThreadStates;
 
-		FCriticalSection PendingThreadStatesGuard;
-		TArray<FLLMThreadState*, FDefaultLLMAllocator> PendingThreadStates;
+	FCriticalSection PendingThreadStatesGuard;
+	TArray<FLLMThreadState*, FDefaultLLMAllocator> PendingThreadStates;
 
-		int64 TrackedTotal GCC_ALIGN(8);
+	int64 TrackedTotal GCC_ALIGN(8);
 
-		FLLMAllocMap AllocationMap;
+	FLLMAllocMap AllocationMap;
 
-		FTrackerTagSizeMap TagSizes;
+	FTrackerTagSizeMap TagSizes;
 
-		const FTagData* OverrideUntaggedTagData;
-		const FTagData* OverrideTrackedTotalTagData;
+	const FTagData* OverrideUntaggedTagData;
+	const FTagData* OverrideTrackedTotalTagData;
 
-		FLLMCsvWriter CsvWriter;
-		FLLMTraceWriter TraceWriter;
-		FLLMCsvProfilerWriter CsvProfilerWriter;
+	FLLMCsvWriter CsvWriter;
+	FLLMTraceWriter TraceWriter;
+	FLLMCsvProfilerWriter CsvProfilerWriter;
 
-		double LastTrimTime;
+	double LastTrimTime;
 
-		int64 AllocTypeAmounts[(int32)ELLMAllocType::Count];
-	};
+	int64 AllocTypeAmounts[(int32)ELLMAllocType::Count];
+};
 
-	const TCHAR* ToString(ETagReferenceSource ReferenceSource);
-	void SetMemoryStatByFName(FName Name, int64 Amount);
-	void ValidateUniqueName(FStringView UniqueName);
+const TCHAR* ToString(ETagReferenceSource ReferenceSource);
+void SetMemoryStatByFName(FName Name, int64 Amount);
+void ValidateUniqueName(FStringView UniqueName);
 
-} // UE::LLMPrivate
+} // namespace UE::LLMPrivate
 
 static FName GetTagName_CustomName()
 {
@@ -2475,11 +2473,7 @@ void FLLMTagDeclaration::ConstructUniqueName()
 	UniqueName = FName(*NameBuffer);
 }
 
-namespace UE
-{
-namespace LLMPrivate
-{
-namespace LLMTagDeclarationInternal
+namespace UE::LLMPrivate::LLMTagDeclarationInternal
 {
 
 TArray<FLLMTagDeclaration::FCreationCallback, TInlineAllocator<2>>& GetCreationCallbacks()
@@ -2494,9 +2488,7 @@ FLLMTagDeclaration*& GetList()
 	return List;
 }
 
-}
-}
-}
+} // namespace UE::LLMPrivate::LLMTagDeclarationInternal
 
 void FLLMTagDeclaration::AddCreationCallback(FCreationCallback InCallback)
 {
@@ -2544,1576 +2536,1582 @@ namespace UE::LLMPrivate
 {
 #if DO_CHECK
 
-	bool HandleAssert(bool bLog, const TCHAR* Format, ...)
+bool HandleAssert(bool bLog, const TCHAR* Format, ...)
+{
+	if (bLog)
 	{
-		if (bLog)
-		{
-			TCHAR DescriptionString[4096];
-		GET_VARARGS(DescriptionString, UE_ARRAY_COUNT(DescriptionString), UE_ARRAY_COUNT(DescriptionString) - 1,
-			Format, Format);
+		TCHAR DescriptionString[4096];
+	GET_VARARGS(DescriptionString, UE_ARRAY_COUNT(DescriptionString), UE_ARRAY_COUNT(DescriptionString) - 1,
+		Format, Format);
 
-			FPlatformMisc::LowLevelOutputDebugString(DescriptionString);
+		FPlatformMisc::LowLevelOutputDebugString(DescriptionString);
 
-			if (FPlatformMisc::IsDebuggerPresent())
-				FPlatformMisc::PromptForRemoteDebugging(true);
+		if (FPlatformMisc::IsDebuggerPresent())
+			FPlatformMisc::PromptForRemoteDebugging(true);
 
-			UE_DEBUG_BREAK();
-		}
-		return false;
+		UE_DEBUG_BREAK();
 	}
+	return false;
+}
 
 #endif
 
-	const TCHAR* ToString(UE::LLMPrivate::ETagReferenceSource ReferenceSource)
-	{
-		switch (ReferenceSource)
-		{
-		case ETagReferenceSource::Scope:
-			return TEXT("LLM_SCOPE");
-		case ETagReferenceSource::Declare:
-			return TEXT("LLM_DEFINE_TAG");
-		case ETagReferenceSource::EnumTag:
-			return TEXT("LLM_ENUM_GENERIC_TAGS");
-		case ETagReferenceSource::CustomEnumTag:
-			return TEXT("RegisterPlatformTag/RegisterProjectTag");
-		case ETagReferenceSource::FunctionAPI:
-			return TEXT("DefaultName/InternalCall");
-		case ETagReferenceSource::ImplicitParent:
-			return TEXT("ImplicitParent");
-		default:
-			return TEXT("Invalid");
-		}
+const TCHAR* ToString(UE::LLMPrivate::ETagReferenceSource ReferenceSource)
+{
+	switch (ReferenceSource)
+	{
+	case ETagReferenceSource::Scope:
+		return TEXT("LLM_SCOPE");
+	case ETagReferenceSource::Declare:
+		return TEXT("LLM_DEFINE_TAG");
+	case ETagReferenceSource::EnumTag:
+		return TEXT("LLM_ENUM_GENERIC_TAGS");
+	case ETagReferenceSource::CustomEnumTag:
+		return TEXT("RegisterPlatformTag/RegisterProjectTag");
+	case ETagReferenceSource::FunctionAPI:
+		return TEXT("DefaultName/InternalCall");
+	case ETagReferenceSource::ImplicitParent:
+		return TEXT("ImplicitParent");
+	default:
+		return TEXT("Invalid");
 	}
+}
 
-	void ValidateUniqueName(FStringView UniqueName)
+void ValidateUniqueName(FStringView UniqueName)
+{
+// Characters that are invalid for c++ names are invalid (other than /), since we use uniquenames
+// (with / replaced by _) as part of the name of the auto-constructed FLLMTagDeclaration variables
+// _ is invalid since we use an _ to indicate a / in FLLMTagDeclaration.
+// So only Alnum characters or / are allowed, and the first character can not be a number.
+	if (UniqueName.Len() == 0)
 	{
-	// Characters that are invalid for c++ names are invalid (other than /), since we use uniquenames
-	// (with / replaced by _) as part of the name of the auto-constructed FLLMTagDeclaration variables
-	// _ is invalid since we use an _ to indicate a / in FLLMTagDeclaration.
-	// So only Alnum characters or / are allowed, and the first character can not be a number.
-		if (UniqueName.Len() == 0)
-		{
-			LLMCheckf(false, TEXT("Invalid length-zero Tag Unique Name"));
-		}
-		else
+		LLMCheckf(false, TEXT("Invalid length-zero Tag Unique Name"));
+	}
+	else
+	{
+	LLMCheckf(!TChar<TCHAR>::IsDigit(UniqueName[0]),
+		TEXT("Invalid first character is digit in Tag Unique Name '%.*s'"),
+		UniqueName.Len(), UniqueName.GetData());
+	}
+	for (TCHAR c : UniqueName)
+	{
+		if (!TChar<TCHAR>::IsAlnum(c) && c != TEXT('/'))
 		{
-		LLMCheckf(!TChar<TCHAR>::IsDigit(UniqueName[0]),
-			TEXT("Invalid first character is digit in Tag Unique Name '%.*s'"),
+		LLMCheckf(false, TEXT("Invalid character %c in Tag Unique Name '%.*s'"), c,
 			UniqueName.Len(), UniqueName.GetData());
 		}
-		for (TCHAR c : UniqueName)
-		{
-			if (!TChar<TCHAR>::IsAlnum(c) && c != TEXT('/'))
-			{
-			LLMCheckf(false, TEXT("Invalid character %c in Tag Unique Name '%.*s'"), c,
-				UniqueName.Len(), UniqueName.GetData());
-			}
-		}
 	}
+}
 
 
-	namespace AllocatorPrivate
-	{
+namespace AllocatorPrivate
+{
 
-		/**
+/**
  * When a Page is allocated, it splits the memory of the page up into blocks, and creates an FAlloc at the start of
  * each block. All the FAllocs are joined together in a FreeList linked list.
-		 * When a Page allocates memory, it takes an FAlloc from the freelist and gives it to the caller, and forgets about it.
+ * When a Page allocates memory, it takes an FAlloc from the freelist and gives it to the caller, and forgets about it.
  * When the caller returns a pointer, the Page restores the FAlloc at the beginning of the block and puts it back on
  * the FreeList.
-		 */
-		struct FAlloc
-		{
-			FAlloc* Next;
-		};
+ */
+struct FAlloc
+{
+	FAlloc* Next;
+};
 
-		/**
-		 * An FPage holds a single page of memory received from the OS; all pages are of the same size.
-		 * FPages are owned by FBins, and the FPages for an FBin divide the page up into blocks of the FBin's size.
+/**
+ * An FPage holds a single page of memory received from the OS; all pages are of the same size.
+ * FPages are owned by FBins, and the FPages for an FBin divide the page up into blocks of the FBin's size.
  * An FPage keeps track of the blocks it has not yet given out so it can allocate, and keeps track of how many blocks
  * it has given out, so that it can be freed when no longer used.
  * Pages that are not free or empty are available for allocating from and are kept in a doubly-linked list on the FBin.
-		 */
-		struct FPage
-		{
-			FPage(int32 PageSize, int32 BinSize);
-			void* Allocate();
-			void Free(void* Ptr);
-			bool IsFull() const;
-			bool IsEmpty() const;
-			void AddToList(FPage*& Head);
-			void RemoveFromList(FPage*& Head);
-
-			FAlloc* FreeList;
-			FPage* Prev;
-			FPage* Next;
-			int32 UsedCount;
-		};
-
-		/**
+ */
+struct FPage
+{
+	FPage(int32 PageSize, int32 BinSize);
+	void* Allocate();
+	void Free(void* Ptr);
+	bool IsFull() const;
+	bool IsEmpty() const;
+	void AddToList(FPage*& Head);
+	void RemoveFromList(FPage*& Head);
+
+	FAlloc* FreeList;
+	FPage* Prev;
+	FPage* Next;
+	int32 UsedCount;
+};
+
+/**
  * An FBin handles all allocations that fit in its size range. Its size is the power of two at the top of that range.
  * The FBin allocates one FPage at a time from the OS; the FPage gets split up into blocks and handles providing a
  * block for callers requesting a pointer.
  * The FBin has a doubly-linked list of pages in use but not yet full. It provides allocations from these pages.
  * When an FPage gets full, the FBin forgets about it, counting on the caller to give the pointer to the page back
  * when it frees the pointer and the page becomes non-full again.
-		 * When an FBin has no more non-full pages and needs to satisfy an alloc, it allocates a new page.
-		 * When a page becomes unused due to a free, the FBin frees the page, returning it to the OS.
-		 */
-		struct FBin
-		{
-			FBin(int32 InBinSize);
-			void* Allocate(FLLMAllocator& Allocator);
-			void Free(void* Ptr, FLLMAllocator& Allocator);
+ * When an FBin has no more non-full pages and needs to satisfy an alloc, it allocates a new page.
+ * When a page becomes unused due to a free, the FBin frees the page, returning it to the OS.
+ */
+struct FBin
+{
+	FBin(int32 InBinSize);
+	void* Allocate(FLLMAllocator& Allocator);
+	void Free(void* Ptr, FLLMAllocator& Allocator);
 
-			FPage* FreePages;
-			int32 UsedCount;
-			int32 BinSize;
-		};
+	FPage* FreePages;
+	int32 UsedCount;
+	int32 BinSize;
+};
 
-} // AllocatorPrivate
+} // namespace AllocatorPrivate
 
-	FLLMAllocator*& FLLMAllocator::Get()
-	{
-		static FLLMAllocator* Allocator = nullptr;
-		return Allocator;
-	}
+FLLMAllocator*& FLLMAllocator::Get()
+{
+	static FLLMAllocator* Allocator = nullptr;
+	return Allocator;
+}
 
-	FLLMAllocator::FLLMAllocator()
-		: PlatformAlloc(nullptr)
-		, PlatformFree(nullptr)
-		, Bins(nullptr)
-		, Total(0)
-		, PageSize(0)
-		, NumBins(0)
-	{
-	}
+FLLMAllocator::FLLMAllocator()
+	: PlatformAlloc(nullptr)
+	, PlatformFree(nullptr)
+	, Bins(nullptr)
+	, Total(0)
+	, PageSize(0)
+	, NumBins(0)
+{
+}
 
-	FLLMAllocator::~FLLMAllocator()
-	{
-		Clear();
-	}
+FLLMAllocator::~FLLMAllocator()
+{
+	Clear();
+}
 
-	void FLLMAllocator::Initialise(LLMAllocFunction InAlloc, LLMFreeFunction InFree, int32 InPageSize)
-	{
-		using namespace UE::LLMPrivate::AllocatorPrivate;
+void FLLMAllocator::Initialise(LLMAllocFunction InAlloc, LLMFreeFunction InFree, int32 InPageSize)
+{
+	using namespace UE::LLMPrivate::AllocatorPrivate;
 
-		PlatformAlloc = InAlloc;
-		PlatformFree = InFree;
-		PageSize = InPageSize;
+	PlatformAlloc = InAlloc;
+	PlatformFree = InFree;
+	PageSize = InPageSize;
 
-		if (PlatformAlloc)
-		{
-			constexpr int32 MinBinSizeForAlignment = 16;
-			constexpr int32 MinBinSizeForAllocStorage = static_cast<int32>(sizeof(FAlloc));
-			constexpr int32 MultiplierBetweenBins = 2;
+	if (PlatformAlloc)
+	{
+		constexpr int32 MinBinSizeForAlignment = 16;
+		constexpr int32 MinBinSizeForAllocStorage = static_cast<int32>(sizeof(FAlloc));
+		constexpr int32 MultiplierBetweenBins = 2;
 		// Setting MultiplierAfterLastBin=2 would be useless because the PageSize/2 bin would only get a single
 		// allocation out of each page due to the FPage data taking up the first half
 		// TODO: For bins >= 4*FPage size, allocate FPages in a separate list rather than embedding them.
 		// This will require allocating extra space in each allocation to store its page pointer.
-			constexpr int32 MultiplierAfterLastBin = 4;
+		constexpr int32 MultiplierAfterLastBin = 4;
 
-			int32 MinBinSize = FMath::Max(MinBinSizeForAllocStorage, MinBinSizeForAlignment);
-			int32 MaxBinSize = InPageSize / MultiplierAfterLastBin;
-			int32 BinSize = MinBinSize;
-			while (BinSize <= MaxBinSize)
+		int32 MinBinSize = FMath::Max(MinBinSizeForAllocStorage, MinBinSizeForAlignment);
+		int32 MaxBinSize = InPageSize / MultiplierAfterLastBin;
+		int32 BinSize = MinBinSize;
+		while (BinSize <= MaxBinSize)
+		{
+			BinSize *= MultiplierBetweenBins;
+			++NumBins;
+		}
+		if (NumBins > 0)
+		{
+			Bins = reinterpret_cast<FBin*>(AllocPages(NumBins * sizeof(FBin)));
+			BinSize = MinBinSize;
+			for (int32 BinIndex = 0; BinIndex < NumBins; ++BinIndex)
 			{
+				new (&Bins[BinIndex]) FBin(BinSize);
 				BinSize *= MultiplierBetweenBins;
-				++NumBins;
-			}
-			if (NumBins > 0)
-			{
-				Bins = reinterpret_cast<FBin*>(AllocPages(NumBins * sizeof(FBin)));
-				BinSize = MinBinSize;
-				for (int32 BinIndex = 0; BinIndex < NumBins; ++BinIndex)
-				{
-					new (&Bins[BinIndex]) FBin(BinSize);
-					BinSize *= MultiplierBetweenBins;
-				}
 			}
 		}
 	}
+}
 
-	void FLLMAllocator::Clear()
-	{
-		using namespace UE::LLMPrivate::AllocatorPrivate;
+void FLLMAllocator::Clear()
+{
+	using namespace UE::LLMPrivate::AllocatorPrivate;
 
-		if (NumBins)
+	if (NumBins)
+	{
+		for (int32 BinIndex = 0; BinIndex < NumBins; ++BinIndex)
 		{
-			for (int32 BinIndex = 0; BinIndex < NumBins; ++BinIndex)
-			{
-				LLMCheck(Bins[BinIndex].UsedCount == 0);
-				Bins[BinIndex].~FBin();
-			}
-			FreePages(Bins, NumBins * sizeof(FBin));
-			Bins = nullptr;
-			NumBins = 0;
+			LLMCheck(Bins[BinIndex].UsedCount == 0);
+			Bins[BinIndex].~FBin();
 		}
+		FreePages(Bins, NumBins * sizeof(FBin));
+		Bins = nullptr;
+		NumBins = 0;
 	}
+}
+
+void* FLLMAllocator::Malloc(size_t Size)
+{
+	return Alloc(Size);
+}
 
-	void* FLLMAllocator::Malloc(size_t Size)
+void* FLLMAllocator::Alloc(size_t Size)
+{
+	using namespace UE::LLMPrivate::AllocatorPrivate;
+
+	if (Size == 0)
 	{
-		return Alloc(Size);
+		return nullptr;
 	}
-
-	void* FLLMAllocator::Alloc(size_t Size)
+	int32 BinIndex = GetBinIndex(Size);
+	FScopeLock Lock(&CriticalSection);
+	if (BinIndex == NumBins)
 	{
-		using namespace UE::LLMPrivate::AllocatorPrivate;
+		return AllocPages(Size);
+	}
+	return Bins[BinIndex].Allocate(*this);
+}
 
-		if (Size == 0)
-		{
-			return nullptr;
-		}
+void FLLMAllocator::Free(void* Ptr, size_t Size)
+{
+	using namespace UE::LLMPrivate::AllocatorPrivate;
+
+	if (Ptr != nullptr)
+	{
 		int32 BinIndex = GetBinIndex(Size);
 		FScopeLock Lock(&CriticalSection);
 		if (BinIndex == NumBins)
 		{
-			return AllocPages(Size);
+			FreePages(Ptr, Size);
 		}
-		return Bins[BinIndex].Allocate(*this);
-	}
-
-	void FLLMAllocator::Free(void* Ptr, size_t Size)
-	{
-		using namespace UE::LLMPrivate::AllocatorPrivate;
-
-		if (Ptr != nullptr)
+		else
 		{
-			int32 BinIndex = GetBinIndex(Size);
-			FScopeLock Lock(&CriticalSection);
-			if (BinIndex == NumBins)
-			{
-				FreePages(Ptr, Size);
-			}
-			else
-			{
-				Bins[BinIndex].Free(Ptr, *this);
-			}
+			Bins[BinIndex].Free(Ptr, *this);
 		}
 	}
+}
 
-	void* FLLMAllocator::Realloc(void* Ptr, size_t OldSize, size_t NewSize)
+void* FLLMAllocator::Realloc(void* Ptr, size_t OldSize, size_t NewSize)
+{
+	void* NewPtr;
+	if (NewSize)
 	{
-		void* NewPtr;
-		if (NewSize)
-		{
-			NewPtr = Alloc(NewSize);
-			if (OldSize)
-			{
-				size_t CopySize = FMath::Min(OldSize, NewSize);
-				FMemory::Memcpy(NewPtr, Ptr, CopySize);
-			}
-		}
-		else
+		NewPtr = Alloc(NewSize);
+		if (OldSize)
 		{
-			NewPtr = nullptr;
+			size_t CopySize = FMath::Min(OldSize, NewSize);
+			FMemory::Memcpy(NewPtr, Ptr, CopySize);
 		}
-		Free(Ptr, OldSize);
-		return NewPtr;
 	}
-
-	int64 FLLMAllocator::GetTotal() const
+	else
 	{
-		FScopeLock Lock((FCriticalSection*)&CriticalSection);
-		return Total;
+		NewPtr = nullptr;
 	}
+	Free(Ptr, OldSize);
+	return NewPtr;
+}
 
-	void* FLLMAllocator::AllocPages(size_t Size)
-	{
-		Size = Align(Size, PageSize);
-		void* Ptr = PlatformAlloc(Size);
-		LLMCheck(Ptr);
-		LLMCheck((reinterpret_cast<intptr_t>(Ptr) & (PageSize - 1)) == 0);
-		Total += Size;
-		return Ptr;
-	}
+int64 FLLMAllocator::GetTotal() const
+{
+	FScopeLock Lock((FCriticalSection*)&CriticalSection);
+	return Total;
+}
 
-	void FLLMAllocator::FreePages(void* Ptr, size_t Size)
-	{
-		Size = Align(Size, PageSize);
-		PlatformFree(Ptr, Size);
-		Total -= Size;
-	}
+void* FLLMAllocator::AllocPages(size_t Size)
+{
+	Size = Align(Size, PageSize);
+	void* Ptr = PlatformAlloc(Size);
+	LLMCheck(Ptr);
+	LLMCheck((reinterpret_cast<intptr_t>(Ptr) & (PageSize - 1)) == 0);
+	Total += Size;
+	return Ptr;
+}
 
-	int32 FLLMAllocator::GetBinIndex(size_t Size) const
-	{
-		int BinIndex = 0;
-		while (BinIndex < NumBins && static_cast<size_t>(Bins[BinIndex].BinSize) < Size)
-		{
-			++BinIndex;
-		}
-		return BinIndex;
-	}
+void FLLMAllocator::FreePages(void* Ptr, size_t Size)
+{
+	Size = Align(Size, PageSize);
+	PlatformFree(Ptr, Size);
+	Total -= Size;
+}
 
-	namespace AllocatorPrivate
+int32 FLLMAllocator::GetBinIndex(size_t Size) const
+{
+	int BinIndex = 0;
+	while (BinIndex < NumBins && static_cast<size_t>(Bins[BinIndex].BinSize) < Size)
 	{
-		FPage::FPage(int32 PageSize, int32 BinSize)
-		{
-			Next = Prev = nullptr;
-			UsedCount = 0;
-			int32 NumHeaderBins = (FMath::Max(static_cast<int32>(sizeof(FPage)), BinSize) + BinSize - 1) / BinSize;
-			int32 FreeCount = PageSize / BinSize - NumHeaderBins;
-
-			// Divide the rest of the page after this header into FAllocs, and add all the FAllocs into the free list
-			FreeList = reinterpret_cast<FAlloc*>(reinterpret_cast<intptr_t>(this) + NumHeaderBins*BinSize);
-		FAlloc* EndAlloc = reinterpret_cast<FAlloc*>(
-			reinterpret_cast<intptr_t>(FreeList) + (FreeCount-1) * BinSize);
-			FAlloc* Alloc = FreeList;
-			while (Alloc != EndAlloc)
-			{
-				Alloc->Next = reinterpret_cast<FAlloc*>(reinterpret_cast<intptr_t>(Alloc) + BinSize);
-				Alloc = Alloc->Next;
-			}
-			EndAlloc->Next = nullptr;
-		}
+		++BinIndex;
+	}
+	return BinIndex;
+}
 
-		void* FPage::Allocate()
-		{
-			LLMCheck(FreeList);
-			FAlloc* Alloc = FreeList;
-			FreeList = Alloc->Next;
-			++UsedCount;
-			return Alloc;
-		}
+namespace AllocatorPrivate
+{
 
-		void FPage::Free(void* Ptr)
-		{
-			LLMCheck(UsedCount > 0);
-			FAlloc* Alloc = reinterpret_cast<FAlloc*>(Ptr);
-			Alloc->Next = FreeList;
-			FreeList = Alloc;
-			--UsedCount;
-		}
+FPage::FPage(int32 PageSize, int32 BinSize)
+{
+	Next = Prev = nullptr;
+	UsedCount = 0;
+	int32 NumHeaderBins = (FMath::Max(static_cast<int32>(sizeof(FPage)), BinSize) + BinSize - 1) / BinSize;
+	int32 FreeCount = PageSize / BinSize - NumHeaderBins;
+
+	// Divide the rest of the page after this header into FAllocs, and add all the FAllocs into the free list.
+	FreeList = reinterpret_cast<FAlloc*>(reinterpret_cast<intptr_t>(this) + NumHeaderBins*BinSize);
+	FAlloc* EndAlloc = reinterpret_cast<FAlloc*>(
+	reinterpret_cast<intptr_t>(FreeList) + (FreeCount-1) * BinSize);
+	FAlloc* Alloc = FreeList;
+	while (Alloc != EndAlloc)
+	{
+		Alloc->Next = reinterpret_cast<FAlloc*>(reinterpret_cast<intptr_t>(Alloc) + BinSize);
+		Alloc = Alloc->Next;
+	}
+	EndAlloc->Next = nullptr;
+}
 
-		bool FPage::IsFull() const
-		{
-			return FreeList == nullptr;
-		}
+void* FPage::Allocate()
+{
+	LLMCheck(FreeList);
+	FAlloc* Alloc = FreeList;
+	FreeList = Alloc->Next;
+	++UsedCount;
+	return Alloc;
+}
 
-		bool FPage::IsEmpty() const
-		{
-			return UsedCount == 0;
-		}
+void FPage::Free(void* Ptr)
+{
+	LLMCheck(UsedCount > 0);
+	FAlloc* Alloc = reinterpret_cast<FAlloc*>(Ptr);
+	Alloc->Next = FreeList;
+	FreeList = Alloc;
+	--UsedCount;
+}
 
-		void FPage::AddToList(FPage*& Head)
-		{
-			Next = Head;
-			Prev = nullptr;
-			Head = this;
-			if (Next)
-			{
-				Next->Prev = this;
-			}
-		}
+bool FPage::IsFull() const
+{
+	return FreeList == nullptr;
+}
 
-		void FPage::RemoveFromList(FPage*& Head)
-		{
-			if (Prev)
-			{
-				Prev->Next = Next;
-				if (Next)
-				{
-					Next->Prev = Prev;
-				}
-			}
-			else
-			{
-				Head = Next;
-				if (Next)
-				{
-					Next->Prev = nullptr;
-				}
-			}
-			Next = Prev = nullptr;
-		}
+bool FPage::IsEmpty() const
+{
+	return UsedCount == 0;
+}
 
-		FBin::FBin(int32 InBinSize)
-		{
-			FreePages = nullptr;
-			UsedCount = 0;
-			BinSize = InBinSize;
-		}
+void FPage::AddToList(FPage*& Head)
+{
+	Next = Head;
+	Prev = nullptr;
+	Head = this;
+	if (Next)
+	{
+		Next->Prev = this;
+	}
+}
 
-		void* FBin::Allocate(FLLMAllocator& Allocator)
+void FPage::RemoveFromList(FPage*& Head)
+{
+	if (Prev)
+	{
+		Prev->Next = Next;
+		if (Next)
 		{
-			if (!FreePages)
-			{
-				FPage* Page = reinterpret_cast<FPage*>(Allocator.AllocPages(Allocator.PageSize));
-				++UsedCount;
-				LLMCheck(Page);
-			// The FPage is at the beginning of the array of PageSize bytes.
-			new (Page) FPage(Allocator.PageSize, BinSize);
-				Page->AddToList(FreePages);
-			}
-
-			void* Result = FreePages->Allocate();
-			if (FreePages->IsFull())
-			{
-				FreePages->RemoveFromList(FreePages); //-V678
-			}
-			return Result;
+			Next->Prev = Prev;
 		}
-
-		void FBin::Free(void* Ptr, FLLMAllocator& Allocator)
+	}
+	else
+	{
+		Head = Next;
+		if (Next)
 		{
-		FPage* Page = reinterpret_cast<FPage*>(
-			reinterpret_cast<intptr_t>(Ptr) & ~(static_cast<intptr_t>(Allocator.PageSize) - 1));
-			if (Page->IsFull())
-			{
-				Page->AddToList(FreePages);
-			}
-			Page->Free(Ptr);
-			if (Page->IsEmpty())
-			{
-				Page->RemoveFromList(FreePages);
-				--UsedCount;
-				Allocator.FreePages(Page, Allocator.PageSize);
-			}
+			Next->Prev = nullptr;
 		}
 	}
+	Next = Prev = nullptr;
+}
 
-FTagData::FTagData(FName InName, FName InDisplayName, FName InParentName, FName InStatName,
-	FName InSummaryStatName, bool bInHasEnumTag, ELLMTag InEnumTag, ETagReferenceSource InReferenceSource)
-		: Name(InName), DisplayName(InDisplayName), ParentName(InParentName), StatName(InStatName)
-		, SummaryStatName(InSummaryStatName), EnumTag(InEnumTag), ReferenceSource(InReferenceSource)
-		, bIsFinishConstructed(false), bParentIsName(true), bHasEnumTag(bInHasEnumTag), bIsReportable(true)
-	{
-	}
+FBin::FBin(int32 InBinSize)
+{
+	FreePages = nullptr;
+	UsedCount = 0;
+	BinSize = InBinSize;
+}
 
-FTagData::FTagData(FName InName, FName InDisplayName, const FTagData* InParent, FName InStatName,
-	FName InSummaryStatName, bool bInHasEnumTag, ELLMTag InEnumTag, ETagReferenceSource InReferenceSource)
-	: FTagData(InName, InDisplayName, NAME_None, InStatName, InSummaryStatName, bInHasEnumTag, InEnumTag,
-		InReferenceSource)
+void* FBin::Allocate(FLLMAllocator& Allocator)
+{
+	if (!FreePages)
 	{
-		SetParent(InParent);
+		FPage* Page = reinterpret_cast<FPage*>(Allocator.AllocPages(Allocator.PageSize));
+		++UsedCount;
+		LLMCheck(Page);
+		// The FPage is at the beginning of the array of PageSize bytes.
+		new (Page) FPage(Allocator.PageSize, BinSize);
+		Page->AddToList(FreePages);
 	}
 
-	FTagData::FTagData(ELLMTag InEnumTag)
-	: FTagData(NAME_None, NAME_None, NAME_None, NAME_None, NAME_None, true, InEnumTag,
-		ETagReferenceSource::EnumTag)
+	void* Result = FreePages->Allocate();
+	if (FreePages->IsFull())
 	{
+		FreePages->RemoveFromList(FreePages); //-V678
 	}
+	return Result;
+}
 
-	FTagData::~FTagData()
+void FBin::Free(void* Ptr, FLLMAllocator& Allocator)
+{
+	FPage* Page = reinterpret_cast<FPage*>(
+		reinterpret_cast<intptr_t>(Ptr) & ~(static_cast<intptr_t>(Allocator.PageSize) - 1));
+	if (Page->IsFull())
 	{
-		if (bParentIsName)
-		{
-			ParentName.~FName();
-			bParentIsName = false;
-		}
+		Page->AddToList(FreePages);
 	}
-
-	bool FTagData::IsFinishConstructed() const
+	Page->Free(Ptr);
+	if (Page->IsEmpty())
 	{
-		return bIsFinishConstructed;
+		Page->RemoveFromList(FreePages);
+		--UsedCount;
+		Allocator.FreePages(Page, Allocator.PageSize);
 	}
+}
 
-	bool FTagData::IsParentConstructed() const
-	{
-		return !bParentIsName;
-	}
+} // namespace AllocatorPrivate
 
-	FName FTagData::GetName() const
-	{
-		return Name;
-	}
+FTagData::FTagData(FName InName, FName InDisplayName, FName InParentName, FName InStatName,
+	FName InSummaryStatName, bool bInHasEnumTag, ELLMTag InEnumTag, ETagReferenceSource InReferenceSource)
+	: Name(InName), DisplayName(InDisplayName), ParentName(InParentName), StatName(InStatName)
+	, SummaryStatName(InSummaryStatName), EnumTag(InEnumTag), ReferenceSource(InReferenceSource)
+	, bIsFinishConstructed(false), bParentIsName(true), bHasEnumTag(bInHasEnumTag), bIsReportable(true)
+{
+}
 
-	FName FTagData::GetDisplayName() const
-	{
-		return DisplayName;
-	}
+FTagData::FTagData(FName InName, FName InDisplayName, const FTagData* InParent, FName InStatName,
+	FName InSummaryStatName, bool bInHasEnumTag, ELLMTag InEnumTag, ETagReferenceSource InReferenceSource)
+	: FTagData(InName, InDisplayName, NAME_None, InStatName, InSummaryStatName, bInHasEnumTag, InEnumTag,
+		InReferenceSource)
+{
+	SetParent(InParent);
+}
+
+FTagData::FTagData(ELLMTag InEnumTag)
+	: FTagData(NAME_None, NAME_None, NAME_None, NAME_None, NAME_None, true, InEnumTag,
+		ETagReferenceSource::EnumTag)
+{
+}
 
-	FString FTagData::GetDisplayPath() const
+FTagData::~FTagData()
+{
+	if (bParentIsName)
 	{
-		TStringBuilder<256> NameBuffer;
-		AppendDisplayPath(NameBuffer);
-		return FString(NameBuffer);
+		ParentName.~FName();
+		bParentIsName = false;
 	}
+}
+
+bool FTagData::IsFinishConstructed() const
+{
+	return bIsFinishConstructed;
+}
+
+bool FTagData::IsParentConstructed() const
+{
+	return !bParentIsName;
+}
+
+FName FTagData::GetName() const
+{
+	return Name;
+}
+
+FName FTagData::GetDisplayName() const
+{
+	return DisplayName;
+}
 
-	void FTagData::AppendDisplayPath(FStringBuilderBase& Result) const
+FString FTagData::GetDisplayPath() const
+{
+	TStringBuilder<256> NameBuffer;
+	AppendDisplayPath(NameBuffer);
+	return FString(NameBuffer);
+}
+
+void FTagData::AppendDisplayPath(FStringBuilderBase& Result) const
+{
+	if (Parent && Parent->IsUsedAsDisplayParent())
 	{
-		if (Parent && Parent->IsUsedAsDisplayParent())
-		{
-			Parent->AppendDisplayPath(Result);
-			Result << TEXT("/");
-		}
-		DisplayName.AppendString(Result);
+		Parent->AppendDisplayPath(Result);
+		Result << TEXT("/");
 	}
+	DisplayName.AppendString(Result);
+}
 
-	const FTagData* FTagData::GetParent() const
-	{
+const FTagData* FTagData::GetParent() const
+{
 	LLMCheckf(!bParentIsName, TEXT("GetParent called on TagData %s before SetParent was called"),
 		*WriteToString<FName::StringBufferSize>(Name));
-		return Parent;
-	}
+	return Parent;
+}
 
-	FName FTagData::GetParentName() const
-	{
+FName FTagData::GetParentName() const
+{
 	LLMCheckf(bParentIsName, TEXT("GetParentName called on TagData %s after SetParent was called"),
 		*WriteToString<FName::StringBufferSize>(Name));
-		return ParentName;
-	}
+	return ParentName;
+}
 
-	FName FTagData::GetStatName() const
-	{
-		return StatName;
-	}
+FName FTagData::GetStatName() const
+{
+	return StatName;
+}
 
-	FName FTagData::GetSummaryStatName() const
-	{
-		return SummaryStatName;
-	}
+FName FTagData::GetSummaryStatName() const
+{
+	return SummaryStatName;
+}
 
-	ELLMTag FTagData::GetEnumTag() const
-	{
-		return EnumTag;
-	}
+ELLMTag FTagData::GetEnumTag() const
+{
+	return EnumTag;
+}
 
-	bool FTagData::HasEnumTag() const
-	{
-		return bHasEnumTag;
-	}
+bool FTagData::HasEnumTag() const
+{
+	return bHasEnumTag;
+}
 
-	const FTagData* FTagData::GetContainingEnumTagData() const
+const FTagData* FTagData::GetContainingEnumTagData() const
+{
+	const FTagData* TagData = this;
+	do
 	{
-		const FTagData* TagData = this;
-		do
+		if (TagData->bHasEnumTag)
 		{
-			if (TagData->bHasEnumTag)
-			{
-				return TagData;
-			}
-			TagData = TagData->GetParent();
-		} while (TagData);
+			return TagData;
+		}
+		TagData = TagData->GetParent();
+	} while (TagData);
 	LLMCheckf(false, TEXT("TagData is not a descendent of an ELLMTag TagData. ")
 		TEXT("All TagDatas must be descendents of ELLMTag::CustomName if they are not descendets of any other ELLMTag"));
 		return this;
-	}
+}
 
-	ELLMTag FTagData::GetContainingEnum() const
-	{
-		const FTagData* TagData = GetContainingEnumTagData();
-		return TagData->EnumTag;
-	}
+ELLMTag FTagData::GetContainingEnum() const
+{
+	const FTagData* TagData = GetContainingEnumTagData();
+	return TagData->EnumTag;
+}
 
-	ETagReferenceSource FTagData::GetReferenceSource() const
-	{
-		return ReferenceSource;
-	}
+ETagReferenceSource FTagData::GetReferenceSource() const
+{
+	return ReferenceSource;
+}
 
-	int32 FTagData::GetIndex() const
-	{
-		return Index;
-	}
+int32 FTagData::GetIndex() const
+{
+	return Index;
+}
 
-	bool FTagData::IsReportable() const
-	{
-		return bIsReportable;
-	}
+bool FTagData::IsReportable() const
+{
+	return bIsReportable;
+}
 
-	void FTagData::SetParent(const FTagData* InParent)
+void FTagData::SetParent(const FTagData* InParent)
+{
+	if (bParentIsName)
 	{
-		if (bParentIsName)
-		{
-			ParentName.~FName();
-			bParentIsName = false;
-		}
-		Parent = InParent;
+		ParentName.~FName();
+		bParentIsName = false;
 	}
+	Parent = InParent;
+}
 
-	void FTagData::SetName(FName InName)
-	{
-		Name = InName;
-	}
+void FTagData::SetName(FName InName)
+{
+	Name = InName;
+}
 
-	void FTagData::SetDisplayName(FName InDisplayName)
-	{
-		DisplayName = InDisplayName;
-	}
+void FTagData::SetDisplayName(FName InDisplayName)
+{
+	DisplayName = InDisplayName;
+}
 
-	void FTagData::SetStatName(FName InStatName)
-	{
-		StatName = InStatName;
-	}
+void FTagData::SetStatName(FName InStatName)
+{
+	StatName = InStatName;
+}
 
-	void FTagData::SetSummaryStatName(FName InSummaryStatName)
-	{
-		SummaryStatName = InSummaryStatName;
-	}
+void FTagData::SetSummaryStatName(FName InSummaryStatName)
+{
+	SummaryStatName = InSummaryStatName;
+}
 
-	void FTagData::SetParentName(FName InParentName)
-	{
-		LLMCheck(bParentIsName);
-		ParentName = InParentName;
-	}
+void FTagData::SetParentName(FName InParentName)
+{
+	LLMCheck(bParentIsName);
+	ParentName = InParentName;
+}
 
-	void FTagData::SetFinishConstructed()
-	{
-		bIsFinishConstructed = true;
-	}
+void FTagData::SetFinishConstructed()
+{
+	bIsFinishConstructed = true;
+}
 
-	void FTagData::SetIndex(int32 InIndex)
-	{
-		Index = InIndex;
-	}
+void FTagData::SetIndex(int32 InIndex)
+{
+	Index = InIndex;
+}
 
-	void FTagData::SetIsReportable(bool bInIsReportable)
-	{
-		bIsReportable = bInIsReportable;
-	}
+void FTagData::SetIsReportable(bool bInIsReportable)
+{
+	bIsReportable = bInIsReportable;
+}
 
-	bool FTagData::IsUsedAsDisplayParent() const
-	{
+bool FTagData::IsUsedAsDisplayParent() const
+{
 	// All Tags but one are UsedAsDisplayParent - their name is prepended during GetDisplayPath.
 	// ELLMTag::CustomName is the exception. It is set for FName tags that do not have a real parent to provide a
 	// containing ELLMTag for them to provide to systems that do not support FName tags. When FName tags without a real
 	// parent are displayed, their path should display as parentless despite having the CustomName tag as their parent.
-		return !(bHasEnumTag && EnumTag == ELLMTag::CustomName);
-	}
+	return !(bHasEnumTag && EnumTag == ELLMTag::CustomName);
+}
 
-	FLLMTracker::FLLMTracker(FLowLevelMemTracker& InLLM)
-		: LLMRef(InLLM)
-		, TrackedTotal(0)
-		, OverrideUntaggedTagData(nullptr)
-		, OverrideTrackedTotalTagData(nullptr)
-		, LastTrimTime(0.0)
-	{
-		TlsSlot = FPlatformTLS::AllocTlsSlot();
+FLLMTracker::FLLMTracker(FLowLevelMemTracker& InLLM)
+	: LLMRef(InLLM)
+	, TrackedTotal(0)
+	, OverrideUntaggedTagData(nullptr)
+	, OverrideTrackedTotalTagData(nullptr)
+	, LastTrimTime(0.0)
+{
+	TlsSlot = FPlatformTLS::AllocTlsSlot();
 
-		for (int32 Index = 0; Index < static_cast<int32>(ELLMAllocType::Count); ++Index)
-		{
-			AllocTypeAmounts[Index] = 0;
-		}
+	for (int32 Index = 0; Index < static_cast<int32>(ELLMAllocType::Count); ++Index)
+	{
+		AllocTypeAmounts[Index] = 0;
 	}
+}
 
-	FLLMTracker::~FLLMTracker()
-	{
-		Clear();
+FLLMTracker::~FLLMTracker()
+{
+	Clear();
 
-		FPlatformTLS::FreeTlsSlot(TlsSlot);
-	}
+	FPlatformTLS::FreeTlsSlot(TlsSlot);
+}
 
-	void FLLMTracker::Initialise(
-		ELLMTracker Tracker,
-		FLLMAllocator* InAllocator)
-	{
-		CsvWriter.SetTracker(Tracker);
-		TraceWriter.SetTracker(Tracker);
-		CsvProfilerWriter.SetTracker(Tracker);
+void FLLMTracker::Initialise(
+	ELLMTracker Tracker,
+	FLLMAllocator* InAllocator)
+{
+	CsvWriter.SetTracker(Tracker);
+	TraceWriter.SetTracker(Tracker);
+	CsvProfilerWriter.SetTracker(Tracker);
 
-		AllocationMap.SetAllocator(InAllocator);
-	}
+	AllocationMap.SetAllocator(InAllocator);
+}
 
-	FLLMThreadState* FLLMTracker::GetOrCreateState()
+FLLMThreadState* FLLMTracker::GetOrCreateState()
+{
+	// look for already allocated thread state
+	FLLMThreadState* State = (FLLMThreadState*)FPlatformTLS::GetTlsValue(TlsSlot);
+	// get one if needed
+	if (State == nullptr)
 	{
-		// look for already allocated thread state
-		FLLMThreadState* State = (FLLMThreadState*)FPlatformTLS::GetTlsValue(TlsSlot);
-		// get one if needed
-		if (State == nullptr)
-		{
-			State = LLMRef.Allocator.New<FLLMThreadState>();
-
-			// Add to pending thread states, these will be consumed on the GT
-			{
-				FScopeLock Lock(&PendingThreadStatesGuard);
-				PendingThreadStates.Add(State);
-			}
+		State = LLMRef.Allocator.New<FLLMThreadState>();
 
-			// push to Tls
-			FPlatformTLS::SetTlsValue(TlsSlot, State);
+		// Add to pending thread states, these will be consumed on the GT
+		{
+			FScopeLock Lock(&PendingThreadStatesGuard);
+			PendingThreadStates.Add(State);
 		}
-		return State;
-	}
 
-	FLLMThreadState* FLLMTracker::GetState()
-	{
-		return (FLLMThreadState*)FPlatformTLS::GetTlsValue(TlsSlot);
+		// push to Tls
+		FPlatformTLS::SetTlsValue(TlsSlot, State);
 	}
+	return State;
+}
 
-	void FLLMTracker::PushTag(ELLMTag EnumTag)
-	{
-		const FTagData* TagData = LLMRef.FindOrAddTagData(EnumTag, ETagReferenceSource::Scope);
+FLLMThreadState* FLLMTracker::GetState()
+{
+	return (FLLMThreadState*)FPlatformTLS::GetTlsValue(TlsSlot);
+}
 
-		// pass along to the state object
-		GetOrCreateState()->PushTag(TagData);
-	}
+void FLLMTracker::PushTag(ELLMTag EnumTag)
+{
+	const FTagData* TagData = LLMRef.FindOrAddTagData(EnumTag, ETagReferenceSource::Scope);
 
-	void FLLMTracker::PushTag(FName Tag, bool bInIsStatData)
-	{
-		const FTagData* TagData = LLMRef.FindOrAddTagData(Tag, bInIsStatData, ETagReferenceSource::Scope);
+	// pass along to the state object
+	GetOrCreateState()->PushTag(TagData);
+}
 
-		// pass along to the state object
-		GetOrCreateState()->PushTag(TagData);
-	}
+void FLLMTracker::PushTag(FName Tag, bool bInIsStatData)
+{
+	const FTagData* TagData = LLMRef.FindOrAddTagData(Tag, bInIsStatData, ETagReferenceSource::Scope);
 
-	void FLLMTracker::PushTag(const FTagData* TagData)
-	{
-		// pass along to the state object
-		GetOrCreateState()->PushTag(TagData);
-	}
+	// pass along to the state object
+	GetOrCreateState()->PushTag(TagData);
+}
 
-	void FLLMTracker::PopTag()
-	{
-		// look for already allocated thread state
-		FLLMThreadState* State = GetState();
+void FLLMTracker::PushTag(const FTagData* TagData)
+{
+	// pass along to the state object
+	GetOrCreateState()->PushTag(TagData);
+}
 
-		LLMCheckf(State != nullptr, TEXT("Called PopTag but PushTag was never called!"));
+void FLLMTracker::PopTag()
+{
+	// look for already allocated thread state
+	FLLMThreadState* State = GetState();
 
-		State->PopTag();
-	}
+	LLMCheckf(State != nullptr, TEXT("Called PopTag but PushTag was never called!"));
+
+	State->PopTag();
+}
 
 #if LLM_ALLOW_ASSETS_TAGS
-	void FLLMTracker::PushAssetTag(FName Tag)
-	{
-		const FTagData* TagData = LLMRef.FindOrAddTagData(Tag, true /* bIsStatTag */, ETagReferenceSource::Scope);
+void FLLMTracker::PushAssetTag(FName Tag)
+{
+	const FTagData* TagData = LLMRef.FindOrAddTagData(Tag, true /* bIsStatTag */, ETagReferenceSource::Scope);
 
-		// pass along to the state object
-		GetOrCreateState()->PushAssetTag(TagData);
-	}
+	// pass along to the state object
+	GetOrCreateState()->PushAssetTag(TagData);
+}
 
-	void FLLMTracker::PushAssetTag(const FTagData* TagData)
-	{
-		// pass along to the state object
-		GetOrCreateState()->PushAssetTag(TagData);
-	}
+void FLLMTracker::PushAssetTag(const FTagData* TagData)
+{
+	// pass along to the state object
+	GetOrCreateState()->PushAssetTag(TagData);
+}
 
-	void FLLMTracker::PopAssetTag()
-	{
-		// look for already allocated thread state
-		FLLMThreadState* State = GetState();
+void FLLMTracker::PopAssetTag()
+{
+	// look for already allocated thread state
+	FLLMThreadState* State = GetState();
 
-		LLMCheckf(State != nullptr, TEXT("Called PopTag but PushTag was never called!"));
+	LLMCheckf(State != nullptr, TEXT("Called PopTag but PushTag was never called!"));
 
-		State->PopAssetTag();
-	}
+	State->PopAssetTag();
+}
 #endif
 
 void FLLMTracker::TrackAllocation(const void* Ptr, int64 Size, ELLMTag DefaultEnumTag, ELLMTracker Tracker,
 	ELLMAllocType AllocType, bool bTrackInMemPro)
+{
+	FLLMThreadState* State = GetOrCreateState();
+	const FTagData* TagData = State->GetTopTag();
+	if (!TagData)
 	{
-		FLLMThreadState* State = GetOrCreateState();
-		const FTagData* TagData = State->GetTopTag();
-		if (!TagData)
-		{
-			TagData = LLMRef.FindOrAddTagData(DefaultEnumTag);
-		}
-		TrackAllocation(Ptr, Size, TagData, Tracker, AllocType, State, bTrackInMemPro);
+		TagData = LLMRef.FindOrAddTagData(DefaultEnumTag);
 	}
+	TrackAllocation(Ptr, Size, TagData, Tracker, AllocType, State, bTrackInMemPro);
+}
 
 void FLLMTracker::TrackAllocation(const void* Ptr, int64 Size, FName DefaultTag, ELLMTracker Tracker,
 	ELLMAllocType AllocType, bool bTrackInMemPro)
+{
+	FLLMThreadState* State = GetOrCreateState();
+	const FTagData* TagData = State->GetTopTag();
+	if (!TagData)
 	{
-		FLLMThreadState* State = GetOrCreateState();
-		const FTagData* TagData = State->GetTopTag();
-		if (!TagData)
-		{
-			TagData = LLMRef.FindOrAddTagData(DefaultTag);
-		}
-		TrackAllocation(Ptr, Size, TagData, Tracker, AllocType, State, bTrackInMemPro);
+		TagData = LLMRef.FindOrAddTagData(DefaultTag);
 	}
+	TrackAllocation(Ptr, Size, TagData, Tracker, AllocType, State, bTrackInMemPro);
+}
 
 void FLLMTracker::TrackAllocation(const void* Ptr, int64 Size, const FTagData* ActiveTagData, ELLMTracker Tracker,
 	ELLMAllocType AllocType, FLLMThreadState* State, bool bTrackInMemPro)
+{
+	if (IsPaused(AllocType))
 	{
-		if (IsPaused(AllocType))
-		{
 		// When Paused, we do not track any new allocations and we we do not update the counters for the memory
 		// they use; the code that triggered the pause is responsible for updating those counters. Since we do not
 		// track the allocations, TrackFree will likewise not update the counters when those allocations are freed.
-			return;
-		}
+		return;
+	}
 
-		// track the total quickly
-		FPlatformAtomics::InterlockedAdd(&TrackedTotal, Size);
+	// track the total quickly
+	FPlatformAtomics::InterlockedAdd(&TrackedTotal, Size);
 	
 #if !LLM_ENABLED_FULL_TAGS
 	// When full tags are disabled, we instead store the top-level enumtag parent of the allocation's tag.
-		ActiveTagData = ActiveTagData->GetContainingEnumTagData();
+	ActiveTagData = ActiveTagData->GetContainingEnumTagData();
 #endif
 
 #if LLM_ALLOW_ASSETS_TAGS
-		const FTagData* AssetTagData = State->GetTopAssetTag();
+	const FTagData* AssetTagData = State->GetTopAssetTag();
 #else
-		const FTagData* AssetTagData = nullptr;
+	const FTagData* AssetTagData = nullptr;
 #endif
 
-		// track on the thread state
-		State->TrackAllocation(Ptr, Size, Tracker, AllocType, ActiveTagData, AssetTagData, bTrackInMemPro);
+	// track on the thread state
+	State->TrackAllocation(Ptr, Size, Tracker, AllocType, ActiveTagData, AssetTagData, bTrackInMemPro);
 
 	// tracking a nullptr with a Size is allowed, but we don't need to remember it, since we can't free it ever.
-		if (Ptr != nullptr)
-		{
-			// remember the size and tag info
-			FLLMTracker::FLowLevelAllocInfo AllocInfo;
-			AllocInfo.SetTag(ActiveTagData, LLMRef);
+	if (Ptr != nullptr)
+	{
+		// remember the size and tag info
+		FLLMTracker::FLowLevelAllocInfo AllocInfo;
+		AllocInfo.SetTag(ActiveTagData, LLMRef);
 #if LLM_ALLOW_ASSETS_TAGS
-			AllocInfo.SetAssetTag(AssetTagData, LLMRef);
+		AllocInfo.SetAssetTag(AssetTagData, LLMRef);
 #endif
-			LLMCheck(Size <= 0x0000'ffff'ffff'ffff);
-			uint32 SizeLow = uint32(Size);
-			uint16 SizeHigh = uint16(Size >> 32ull);
-			PointerKey Key(Ptr, SizeHigh);
-			AllocationMap.Add(Key, SizeLow, AllocInfo);
-		}
+		LLMCheck(Size <= 0x0000'ffff'ffff'ffff);
+		uint32 SizeLow = uint32(Size);
+		uint16 SizeHigh = uint16(Size >> 32ull);
+		PointerKey Key(Ptr, SizeHigh);
+		AllocationMap.Add(Key, SizeLow, AllocInfo);
 	}
+}
 
-	void FLLMTracker::TrackFree(const void* Ptr, ELLMTracker Tracker, ELLMAllocType AllocType, bool bTrackInMemPro)
+void FLLMTracker::TrackFree(const void* Ptr, ELLMTracker Tracker, ELLMAllocType AllocType, bool bTrackInMemPro)
+{
+	// look up the pointer in the tracking map
+	FLLMAllocMap::Values Values;
 	{
-		// look up the pointer in the tracking map
-		FLLMAllocMap::Values Values;
+		if (!AllocationMap.Remove(PointerKey(Ptr), Values))
 		{
-			if (!AllocationMap.Remove(PointerKey(Ptr), Values))
-			{
-				return;
-			}
+			return;
 		}
+	}
 
-		if (IsPaused(AllocType))
-		{
+	if (IsPaused(AllocType))
+	{
 		// When Paused, we remove our data for any freed allocations, but we do not update the counters for the
 		// memory they used; the code that triggered the pause is responsible for updating those counters.
-			return;
-		}
+		return;
+	}
 
-		int64 SizeLow = static_cast<int64>(Values.Value1);
-		int64 SizeHigh = Values.Key.GetExtraData();
-		int64 Size = (SizeHigh << 32ull) | SizeLow;
-		FLLMTracker::FLowLevelAllocInfo& AllocInfo = Values.Value2;
+	int64 SizeLow = static_cast<int64>(Values.Value1);
+	int64 SizeHigh = Values.Key.GetExtraData();
+	int64 Size = (SizeHigh << 32ull) | SizeLow;
+	FLLMTracker::FLowLevelAllocInfo& AllocInfo = Values.Value2;
 
-		// track the total quickly
-		FPlatformAtomics::InterlockedAdd(&TrackedTotal, -Size);
+	// track the total quickly
+	FPlatformAtomics::InterlockedAdd(&TrackedTotal, -Size);
 
-		FLLMThreadState* State = GetOrCreateState();
-		const FTagData* TagData = AllocInfo.GetTag(LLMRef);
+	FLLMThreadState* State = GetOrCreateState();
+	const FTagData* TagData = AllocInfo.GetTag(LLMRef);
 #if LLM_ALLOW_ASSETS_TAGS
-		const FTagData* AssetTagData = AllocInfo.GetAssetTag(LLMRef);
+	const FTagData* AssetTagData = AllocInfo.GetAssetTag(LLMRef);
 #else
-		const FTagData* AssetTagData = nullptr;
+	const FTagData* AssetTagData = nullptr;
 #endif
 
-		State->TrackFree(Ptr, Size, Tracker, AllocType, TagData, AssetTagData, bTrackInMemPro);
-	}
+	State->TrackFree(Ptr, Size, Tracker, AllocType, TagData, AssetTagData, bTrackInMemPro);
+}
 
-	void FLLMTracker::OnAllocMoved(const void* Dest, const void* Source, ELLMTracker Tracker, ELLMAllocType AllocType)
+void FLLMTracker::OnAllocMoved(const void* Dest, const void* Source, ELLMTracker Tracker, ELLMAllocType AllocType)
+{
+	FLLMAllocMap::Values Values;
 	{
-		FLLMAllocMap::Values Values;
+		if (!AllocationMap.Remove(PointerKey(Source), Values))
 		{
-			if (!AllocationMap.Remove(PointerKey(Source), Values))
-			{
-				return;
-			}
-
-			PointerKey Key(Dest, uint16(Values.Key.GetExtraData()));
-			AllocationMap.Add(Key, Values.Value1, Values.Value2);
-		}
-
-		if (IsPaused(AllocType))
-		{
-		// When Paused, don't update counters in case any of the external tracking systems are not available.
 			return;
 		}
 
-		int64 SizeLow = static_cast<int64>(Values.Value1);
-		int64 SizeHigh = Values.Key.GetExtraData();
-		int64 Size = (SizeHigh << 32ull) | SizeLow;
-		const FLLMTracker::FLowLevelAllocInfo& AllocInfo = Values.Value2;
-		const FTagData* TagData = AllocInfo.GetTag(LLMRef);
-
-		FLLMThreadState* State = GetOrCreateState();
-		State->TrackMoved(Dest, Source, Size, Tracker, TagData);
+		PointerKey Key(Dest, uint16(Values.Key.GetExtraData()));
+		AllocationMap.Add(Key, Values.Value1, Values.Value2);
 	}
 
-	void FLLMTracker::TrackMemory(ELLMTag Tag, int64 Amount, ELLMAllocType AllocType)
+	if (IsPaused(AllocType))
 	{
-		TrackMemory(LLMRef.FindOrAddTagData(Tag), Amount, AllocType);
+		// When Paused, don't update counters in case any of the external tracking systems are not available.
+		return;
 	}
 
-	void FLLMTracker::TrackMemory(FName Tag, int64 Amount, ELLMAllocType AllocType)
-	{
-		TrackMemory(LLMRef.FindOrAddTagData(Tag), Amount, AllocType);
-	}
+	int64 SizeLow = static_cast<int64>(Values.Value1);
+	int64 SizeHigh = Values.Key.GetExtraData();
+	int64 Size = (SizeHigh << 32ull) | SizeLow;
+	const FLLMTracker::FLowLevelAllocInfo& AllocInfo = Values.Value2;
+	const FTagData* TagData = AllocInfo.GetTag(LLMRef);
 
-	void FLLMTracker::TrackMemory(const FTagData* TagData, int64 Amount, ELLMAllocType AllocType)
-	{
-		FLLMThreadState* State = GetOrCreateState();
-		FScopeLock Lock(&State->TagSection);
-		State->IncrTag(TagData, Amount);
-		State->AllocTypeAmounts[static_cast<int32>(AllocType)] += Amount;
-		FPlatformAtomics::InterlockedAdd(&TrackedTotal, Amount);
-	}
+	FLLMThreadState* State = GetOrCreateState();
+	State->TrackMoved(Dest, Source, Size, Tracker, TagData);
+}
 
-	void FLLMTracker::PauseAndTrackMemory(FName TagName, bool bInIsStatTag, int64 Amount, ELLMAllocType AllocType)
-	{
-		const FTagData* TagData = LLMRef.FindOrAddTagData(TagName, bInIsStatTag);
-		PauseAndTrackMemory(TagData, Amount, AllocType);
-	}
+void FLLMTracker::TrackMemory(ELLMTag Tag, int64 Amount, ELLMAllocType AllocType)
+{
+	TrackMemory(LLMRef.FindOrAddTagData(Tag), Amount, AllocType);
+}
 
-	void FLLMTracker::PauseAndTrackMemory(ELLMTag EnumTag, int64 Amount, ELLMAllocType AllocType)
-	{
-		const FTagData* TagData = LLMRef.FindOrAddTagData(EnumTag);
-		PauseAndTrackMemory(TagData, Amount, AllocType);
-	}
+void FLLMTracker::TrackMemory(FName Tag, int64 Amount, ELLMAllocType AllocType)
+{
+	TrackMemory(LLMRef.FindOrAddTagData(Tag), Amount, AllocType);
+}
+
+void FLLMTracker::TrackMemory(const FTagData* TagData, int64 Amount, ELLMAllocType AllocType)
+{
+	FLLMThreadState* State = GetOrCreateState();
+	FScopeLock Lock(&State->TagSection);
+	State->IncrTag(TagData, Amount);
+	State->AllocTypeAmounts[static_cast<int32>(AllocType)] += Amount;
+	FPlatformAtomics::InterlockedAdd(&TrackedTotal, Amount);
+}
+
+void FLLMTracker::PauseAndTrackMemory(FName TagName, bool bInIsStatTag, int64 Amount, ELLMAllocType AllocType)
+{
+	const FTagData* TagData = LLMRef.FindOrAddTagData(TagName, bInIsStatTag);
+	PauseAndTrackMemory(TagData, Amount, AllocType);
+}
+
+void FLLMTracker::PauseAndTrackMemory(ELLMTag EnumTag, int64 Amount, ELLMAllocType AllocType)
+{
+	const FTagData* TagData = LLMRef.FindOrAddTagData(EnumTag);
+	PauseAndTrackMemory(TagData, Amount, AllocType);
+}
 
 // This will pause/unpause tracking, and also manually increment a given tag.
-	void FLLMTracker::PauseAndTrackMemory(const FTagData* TagData, int64 Amount, ELLMAllocType AllocType)
-	{
-		FLLMThreadState* State = GetOrCreateState();
-		FScopeLock Lock(&State->TagSection);
-		State->PausedCounter[static_cast<int32>(AllocType)]++;
-		State->IncrTag(TagData, Amount);
-		State->AllocTypeAmounts[static_cast<int32>(AllocType)] += Amount;
-		FPlatformAtomics::InterlockedAdd(&TrackedTotal, Amount);
-	}
+void FLLMTracker::PauseAndTrackMemory(const FTagData* TagData, int64 Amount, ELLMAllocType AllocType)
+{
+	FLLMThreadState* State = GetOrCreateState();
+	FScopeLock Lock(&State->TagSection);
+	State->PausedCounter[static_cast<int32>(AllocType)]++;
+	State->IncrTag(TagData, Amount);
+	State->AllocTypeAmounts[static_cast<int32>(AllocType)] += Amount;
+	FPlatformAtomics::InterlockedAdd(&TrackedTotal, Amount);
+}
 
-	void FLLMTracker::Pause(ELLMAllocType AllocType)
-	{
-		FLLMThreadState* State = GetOrCreateState();
-		State->PausedCounter[static_cast<int32>(AllocType)]++;
-	}
+void FLLMTracker::Pause(ELLMAllocType AllocType)
+{
+	FLLMThreadState* State = GetOrCreateState();
+	State->PausedCounter[static_cast<int32>(AllocType)]++;
+}
 
-	void FLLMTracker::Unpause(ELLMAllocType AllocType)
-	{
-		FLLMThreadState* State = GetOrCreateState();
-		State->PausedCounter[static_cast<int32>(AllocType)]--;
-		LLMCheck( State->PausedCounter[static_cast<int32>(AllocType)] >= 0 );
-	}
+void FLLMTracker::Unpause(ELLMAllocType AllocType)
+{
+	FLLMThreadState* State = GetOrCreateState();
+	State->PausedCounter[static_cast<int32>(AllocType)]--;
+	LLMCheck( State->PausedCounter[static_cast<int32>(AllocType)] >= 0 );
+}
 
-	bool FLLMTracker::IsPaused(ELLMAllocType AllocType)
-	{
-		FLLMThreadState* State = GetState();
+bool FLLMTracker::IsPaused(ELLMAllocType AllocType)
+{
+	FLLMThreadState* State = GetState();
 	// pause during shutdown, as the external trackers might not be able to robustly handle tracking.
 	return IsEngineExitRequested() || (State == nullptr ? false :
 		((State->PausedCounter[static_cast<int32>(ELLMAllocType::None)]>0) ||
 			(State->PausedCounter[static_cast<int32>(AllocType)])>0));
-	}
+}
 
-	void FLLMTracker::Clear()
+void FLLMTracker::Clear()
+{
 	{
+		FScopeLock Lock(&PendingThreadStatesGuard);
+		for (FLLMThreadState* ThreadState : PendingThreadStates)
 		{
-			FScopeLock Lock(&PendingThreadStatesGuard);
-			for (FLLMThreadState* ThreadState : PendingThreadStates)
-				LLMRef.Allocator.Delete(ThreadState);
-			PendingThreadStates.Empty();
-		}
-
-		for (FLLMThreadState* ThreadState : ThreadStates)
 			LLMRef.Allocator.Delete(ThreadState);
-		ThreadStates.Empty();
-
-		AllocationMap.Clear();
-
-		CsvWriter.Clear();
-		TraceWriter.Clear();
-		CsvProfilerWriter.Clear();
+		}
+		PendingThreadStates.Empty();
 	}
 
-	void FLLMTracker::OnPreFork()
+	for (FLLMThreadState* ThreadState : ThreadStates)
 	{
-		CsvWriter.OnPreFork();
+		LLMRef.Allocator.Delete(ThreadState);
 	}
+	ThreadStates.Empty();
+
+	AllocationMap.Clear();
+
+	CsvWriter.Clear();
+	TraceWriter.Clear();
+	CsvProfilerWriter.Clear();
+}
+
+void FLLMTracker::OnPreFork()
+{
+	CsvWriter.OnPreFork();
+}
 
 void FLLMTracker::SetTotalTags(const FTagData* InOverrideUntaggedTagData,
 	const FTagData* InOverrideTrackedTotalTagData)
-	{
-		OverrideUntaggedTagData = InOverrideUntaggedTagData;
-		OverrideTrackedTotalTagData = InOverrideTrackedTotalTagData;
-	}
+{
+	OverrideUntaggedTagData = InOverrideUntaggedTagData;
+	OverrideTrackedTotalTagData = InOverrideTrackedTotalTagData;
+}
 
-	void FLLMTracker::Update()
+void FLLMTracker::Update()
+{
+	UpdateThreads();
+	double CurrentTime = FPlatformTime::Seconds();
+	constexpr double UpdateTrimPeriod = 10.0;
+	bool bTrimAllocations = CurrentTime - LastTrimTime > UpdateTrimPeriod;
+	if (bTrimAllocations)
 	{
-		UpdateThreads();
-		double CurrentTime = FPlatformTime::Seconds();
-		constexpr double UpdateTrimPeriod = 10.0;
-		bool bTrimAllocations = CurrentTime - LastTrimTime > UpdateTrimPeriod;
-		if (bTrimAllocations)
+		LastTrimTime = CurrentTime;
 		{
-			LastTrimTime = CurrentTime;
-			{
-				AllocationMap.Trim();
-			}
+			AllocationMap.Trim();
 		}
+	}
 
 	// Add the values from each thread to the central repository.
-		bool bTrimThreads = bTrimAllocations;
-		for (FLLMThreadState* ThreadState : ThreadStates)
-		{
-			ThreadState->PropagateChildSizesToParents();
-			ThreadState->FetchAndClearTagSizes(TagSizes, AllocTypeAmounts, bTrimThreads);
-		}
+	bool bTrimThreads = bTrimAllocations;
+	for (FLLMThreadState* ThreadState : ThreadStates)
+	{
+		ThreadState->PropagateChildSizesToParents();
+		ThreadState->FetchAndClearTagSizes(TagSizes, AllocTypeAmounts, bTrimThreads);
+	}
 
 	// Update peak sizes and external sizes in the central repository.
-		for (TPair<const FTagData*, FTrackerTagSizeData>& It : TagSizes)
-		{
-			FTrackerTagSizeData& AllocationData = It.Value;
+	for (TPair<const FTagData*, FTrackerTagSizeData>& It : TagSizes)
+	{
+		FTrackerTagSizeData& AllocationData = It.Value;
 
-			// Update external amount
-			if (AllocationData.bExternalValid)
+		// Update external amount
+		if (AllocationData.bExternalValid)
+		{
+			if (AllocationData.bExternalAddToTotal)
 			{
-				if (AllocationData.bExternalAddToTotal)
-				{
-				FPlatformAtomics::InterlockedAdd(&TrackedTotal,
-					AllocationData.ExternalAmount - AllocationData.Size);
-				}
-				AllocationData.Size = AllocationData.ExternalAmount;
-				AllocationData.bExternalValid = false;
+			FPlatformAtomics::InterlockedAdd(&TrackedTotal,
+				AllocationData.ExternalAmount - AllocationData.Size);
 			}
+			AllocationData.Size = AllocationData.ExternalAmount;
+			AllocationData.bExternalValid = false;
+		}
 
 		// Calculate peaks.
 #if LLM_ENABLED_TRACK_PEAK_MEMORY
 		// TODO: we should keep track of the intra-frame memory peak for the total tracked memory. For now we will
 		// use the memory at the time the update happens since there are threading implications to being accurate.
-			AllocationData.PeakSize = FMath::Max(AllocationData.PeakSize, AllocationData.Size);
+		AllocationData.PeakSize = FMath::Max(AllocationData.PeakSize, AllocationData.Size);
 #endif
-		}
 	}
+}
 
-	void FLLMTracker::UpdateThreads()
-	{
+void FLLMTracker::UpdateThreads()
+{
 	// Consume pending thread states. We must be careful to do all allocations outside of the
 	// PendingThreadStatesGuard guard as that can lead to a deadlock due to contention with
 	// PendingThreadStatesGuard & Locks inside the underlying allocator (i.e. MallocBinned2 -> Mutex).
+	{
+		PendingThreadStatesGuard.Lock();
+		const int32 NumPendingThreadStatesToConsume = PendingThreadStates.Num();
+		if (NumPendingThreadStatesToConsume > 0)
 		{
+			PendingThreadStatesGuard.Unlock();
+			ThreadStates.Reserve(ThreadStates.Num() + NumPendingThreadStatesToConsume);
 			PendingThreadStatesGuard.Lock();
-			const int32 NumPendingThreadStatesToConsume = PendingThreadStates.Num();
-			if (NumPendingThreadStatesToConsume > 0)
-			{
-				PendingThreadStatesGuard.Unlock();
-				ThreadStates.Reserve(ThreadStates.Num() + NumPendingThreadStatesToConsume);
-				PendingThreadStatesGuard.Lock();
 
-				for (int32 i = 0; i < NumPendingThreadStatesToConsume; ++i)
-				{
-					ThreadStates.Add(PendingThreadStates.Pop(false /*bAllowShrinking*/));
-				}
+			for (int32 i = 0; i < NumPendingThreadStatesToConsume; ++i)
+			{
+				ThreadStates.Add(PendingThreadStates.Pop(false /*bAllowShrinking*/));
 			}
-			PendingThreadStatesGuard.Unlock();
 		}
+		PendingThreadStatesGuard.Unlock();
 	}
+}
 
-	void FLLMTracker::PublishStats(bool bTrackPeaks)
+void FLLMTracker::PublishStats(bool bTrackPeaks)
+{
+	if (OverrideTrackedTotalTagData)
 	{
-		if (OverrideTrackedTotalTagData)
-		{
-			SetMemoryStatByFName(OverrideTrackedTotalTagData->GetStatName(), TrackedTotal);
-			SetMemoryStatByFName(OverrideTrackedTotalTagData->GetSummaryStatName(), TrackedTotal);
-		}
+		SetMemoryStatByFName(OverrideTrackedTotalTagData->GetStatName(), TrackedTotal);
+		SetMemoryStatByFName(OverrideTrackedTotalTagData->GetSummaryStatName(), TrackedTotal);
+	}
 
-		if (OverrideUntaggedTagData)
-		{
-			const FTagData* TagData = LLMRef.FindTagData(TagName_Untagged);
-			const FTrackerTagSizeData* AllocationData = TagData ? TagSizes.Find(TagData) : nullptr;
+	if (OverrideUntaggedTagData)
+	{
+		const FTagData* TagData = LLMRef.FindTagData(TagName_Untagged);
+		const FTrackerTagSizeData* AllocationData = TagData ? TagSizes.Find(TagData) : nullptr;
 		SetMemoryStatByFName(OverrideUntaggedTagData->GetStatName(),
 			AllocationData ? AllocationData->GetSize(bTrackPeaks) : 0);
 		SetMemoryStatByFName(OverrideUntaggedTagData->GetSummaryStatName(),
 			AllocationData ? AllocationData->GetSize(bTrackPeaks) : 0);
-		}
+	}
 
-		for (const TPair<const FTagData*, FTrackerTagSizeData>& It : TagSizes)
+	for (const TPair<const FTagData*, FTrackerTagSizeData>& It : TagSizes)
+	{
+		const FTagData* TagData = It.Key;
+		if (!TagData->IsReportable())
+		{
+			continue;
+		}
+		if (OverrideUntaggedTagData && TagData->GetName() == TagName_Untagged)
 		{
-			const FTagData* TagData = It.Key;
-			if (!TagData->IsReportable())
-			{
-				continue;
-			}
-			if (OverrideUntaggedTagData && TagData->GetName() == TagName_Untagged)
-			{
 			// Handled separately by OverrideUntaggedTagData.
-				continue;
-			}
-			const FTrackerTagSizeData& AllocationData = It.Value;
-			int64 Amount = AllocationData.GetSize(bTrackPeaks);
-
-			SetMemoryStatByFName(TagData->GetStatName(), Amount);
-			SetMemoryStatByFName(TagData->GetSummaryStatName(), Amount);
+			continue;
 		}
+		const FTrackerTagSizeData& AllocationData = It.Value;
+		int64 Amount = AllocationData.GetSize(bTrackPeaks);
+
+		SetMemoryStatByFName(TagData->GetStatName(), Amount);
+		SetMemoryStatByFName(TagData->GetSummaryStatName(), Amount);
 	}
+}
 
-	void FLLMTracker::PublishCsv(bool bTrackPeaks)
-	{
+void FLLMTracker::PublishCsv(bool bTrackPeaks)
+{
 	CsvWriter.Publish(LLMRef, TagSizes, OverrideTrackedTotalTagData, OverrideUntaggedTagData, TrackedTotal,
 		bTrackPeaks);
-	}
+}
 
-	void FLLMTracker::PublishTrace(bool bTrackPeaks)
-	{
+void FLLMTracker::PublishTrace(bool bTrackPeaks)
+{
 	TraceWriter.Publish(LLMRef, TagSizes, OverrideTrackedTotalTagData, OverrideUntaggedTagData, TrackedTotal,
 		bTrackPeaks);
-	}
+}
 
-	void FLLMTracker::PublishCsvProfiler(bool bTrackPeaks)
-	{
+void FLLMTracker::PublishCsvProfiler(bool bTrackPeaks)
+{
 	CsvProfilerWriter.Publish(LLMRef, TagSizes, OverrideTrackedTotalTagData, OverrideUntaggedTagData, TrackedTotal,
 		bTrackPeaks);
-	}
+}
 
-	void FLLMTracker::OnTagsResorted(FTagDataArray& OldTagDatas)
-	{
+void FLLMTracker::OnTagsResorted(FTagDataArray& OldTagDatas)
+{
 #if LLM_ENABLED_FULL_TAGS
-		{
-			// Each allocation references the tag by its index, which we have just remapped.
+	{
+		// Each allocation references the tag by its index, which we have just remapped.
 		// Remap each allocation's tag index to the new index for the tag.
-			AllocationMap.LockAll();
-			for (const FLLMAllocMap::FTuple& Tuple : AllocationMap)
-			{
-				Tuple.Value2.SetCompressedTag(OldTagDatas[Tuple.Value2.GetCompressedTag()]->GetIndex());
-			}
-			AllocationMap.UnlockAll();
+		AllocationMap.LockAll();
+		for (const FLLMAllocMap::FTuple& Tuple : AllocationMap)
+		{
+			Tuple.Value2.SetCompressedTag(OldTagDatas[Tuple.Value2.GetCompressedTag()]->GetIndex());
 		}
+		AllocationMap.UnlockAll();
+	}
 #else
 	// Values in AllocationMap are ELLMTags, and don't depend on the Index of the tagdatas.
 #endif
 
 	// Update the uses of Index in the ThreadStates.
-		for (FLLMThreadState* ThreadState : ThreadStates)
-		{
-			ThreadState->OnTagsResorted(OldTagDatas);
-		}
-	}
-
-	void FLLMTracker::LockAllThreadTags(bool bLock)
+	for (FLLMThreadState* ThreadState : ThreadStates)
 	{
-		if (bLock)
-		{
-			UpdateThreads();
-			PendingThreadStatesGuard.Lock();
-		}
-
-		for (FLLMThreadState* ThreadState : ThreadStates)
-		{
-			ThreadState->LockTags(bLock);
-		}
-
-		if (!bLock)
-		{
-			PendingThreadStatesGuard.Unlock();
-		}
+		ThreadState->OnTagsResorted(OldTagDatas);
 	}
+}
 
-	const FTagData* FLLMTracker::GetActiveTagData()
+void FLLMTracker::LockAllThreadTags(bool bLock)
+{
+	if (bLock)
 	{
-		FLLMThreadState* State = GetOrCreateState();
-		return State->GetTopTag();
+		UpdateThreads();
+		PendingThreadStatesGuard.Lock();
 	}
 
-	const FTagData* FLLMTracker::FindTagForPtr( void* Ptr )
+	for (FLLMThreadState* ThreadState : ThreadStates)
 	{
-		FLLMThreadState* State = GetOrCreateState();
-		FLowLevelAllocInfo AllocInfo;
-		{
-			uint32* Size;
-			FLowLevelAllocInfo* AllocInfoPtr;
-			AllocationMap.Find(PointerKey(Ptr), Size, AllocInfoPtr);
-			if (!AllocInfoPtr)
-			{
-				return nullptr;
-			}
-			AllocInfo = *AllocInfoPtr;
-		}
-		return AllocInfo.GetTag(LLMRef);
+		ThreadState->LockTags(bLock);
 	}
 
-	TArray<const FTagData*> FLLMTracker::GetTagDatas()
+	if (!bLock)
 	{
-		TArray<const FTagData*> FoundTagDatas;
-		TagSizes.GetKeys(FoundTagDatas);
-		return FoundTagDatas;
+		PendingThreadStatesGuard.Unlock();
 	}
+}
+
+const FTagData* FLLMTracker::GetActiveTagData()
+{
+	FLLMThreadState* State = GetOrCreateState();
+	return State->GetTopTag();
+}
 
-	int64 FLLMTracker::GetTagAmount(const FTagData* TagData, bool bPeakAmount/* = false*/) const
+const FTagData* FLLMTracker::FindTagForPtr( void* Ptr )
+{
+	FLLMThreadState* State = GetOrCreateState();
+	FLowLevelAllocInfo AllocInfo;
 	{
-		const FTrackerTagSizeData* AllocationData = TagSizes.Find(TagData);
-		if (AllocationData)
-		{
-			return AllocationData->GetSize(bPeakAmount);
-		}
-		else
+		uint32* Size;
+		FLowLevelAllocInfo* AllocInfoPtr;
+		AllocationMap.Find(PointerKey(Ptr), Size, AllocInfoPtr);
+		if (!AllocInfoPtr)
 		{
-			return 0;
+			return nullptr;
 		}
+		AllocInfo = *AllocInfoPtr;
 	}
+	return AllocInfo.GetTag(LLMRef);
+}
 
-	void FLLMTracker::SetTagAmountExternal(const FTagData* TagData, int64 Amount, bool bAddToTotal)
+TArray<const FTagData*> FLLMTracker::GetTagDatas()
+{
+	TArray<const FTagData*> FoundTagDatas;
+	TagSizes.GetKeys(FoundTagDatas);
+	return FoundTagDatas;
+}
+
+int64 FLLMTracker::GetTagAmount(const FTagData* TagData, bool bPeakAmount/* = false*/) const
+{
+	const FTrackerTagSizeData* AllocationData = TagSizes.Find(TagData);
+	if (AllocationData)
 	{
-		FTrackerTagSizeData& AllocationData = TagSizes.FindOrAdd(TagData);
-		AllocationData.bExternalValid = true;
-		AllocationData.bExternalAddToTotal = bAddToTotal;
-		AllocationData.ExternalAmount = Amount;
+		return AllocationData->GetSize(bPeakAmount);
 	}
-
-	void FLLMTracker::SetTagAmountInUpdate(const FTagData* TagData, int64 Amount, bool bAddToTotal)
+	else
 	{
-		FTrackerTagSizeData& AllocationData = TagSizes.FindOrAdd(TagData);
-		if (bAddToTotal)
-		{
-			FPlatformAtomics::InterlockedAdd(&TrackedTotal, Amount - AllocationData.Size);
-		}
-		AllocationData.Size = Amount;
-#if LLM_ENABLED_TRACK_PEAK_MEMORY
-		AllocationData.PeakSize = FMath::Max(AllocationData.PeakSize, AllocationData.Size);
-#endif
+		return 0;
 	}
+}
+
+void FLLMTracker::SetTagAmountExternal(const FTagData* TagData, int64 Amount, bool bAddToTotal)
+{
+	FTrackerTagSizeData& AllocationData = TagSizes.FindOrAdd(TagData);
+	AllocationData.bExternalValid = true;
+	AllocationData.bExternalAddToTotal = bAddToTotal;
+	AllocationData.ExternalAmount = Amount;
+}
 
-	int64 FLLMTracker::GetAllocTypeAmount(ELLMAllocType AllocType)
+void FLLMTracker::SetTagAmountInUpdate(const FTagData* TagData, int64 Amount, bool bAddToTotal)
+{
+	FTrackerTagSizeData& AllocationData = TagSizes.FindOrAdd(TagData);
+	if (bAddToTotal)
 	{
-		return AllocTypeAmounts[static_cast<int32>(AllocType)];
+		FPlatformAtomics::InterlockedAdd(&TrackedTotal, Amount - AllocationData.Size);
 	}
+	AllocationData.Size = Amount;
+#if LLM_ENABLED_TRACK_PEAK_MEMORY
+	AllocationData.PeakSize = FMath::Max(AllocationData.PeakSize, AllocationData.Size);
+#endif
+}
 
-	FLLMThreadState::FLLMThreadState()
-	{
-		for (int32 Index = 0; Index < static_cast<int32>(ELLMAllocType::Count); ++Index)
-		{
-			PausedCounter[Index] = 0;
-		}
+int64 FLLMTracker::GetAllocTypeAmount(ELLMAllocType AllocType)
+{
+	return AllocTypeAmounts[static_cast<int32>(AllocType)];
+}
 
-		ClearAllocTypeAmounts();
+FLLMThreadState::FLLMThreadState()
+{
+	for (int32 Index = 0; Index < static_cast<int32>(ELLMAllocType::Count); ++Index)
+	{
+		PausedCounter[Index] = 0;
 	}
 
-	void FLLMThreadState::Clear()
-	{
-		TagStack.Empty();
+	ClearAllocTypeAmounts();
+}
+
+void FLLMThreadState::Clear()
+{
+	TagStack.Empty();
 #if LLM_ALLOW_ASSETS_TAGS
-		AssetTagStack.Empty();
+	AssetTagStack.Empty();
 #endif
-		Allocations.Empty();
-		ClearAllocTypeAmounts();
-	}
+	Allocations.Empty();
+	ClearAllocTypeAmounts();
+}
 
-	void FLLMThreadState::PushTag(const FTagData* TagData)
-	{
-		FScopeLock Lock(&TagSection);
+void FLLMThreadState::PushTag(const FTagData* TagData)
+{
+	FScopeLock Lock(&TagSection);
 
 	// Push a tag.
-		TagStack.Add(TagData);
-	}
+	TagStack.Add(TagData);
+}
 
-	void FLLMThreadState::PopTag()
-	{
-		FScopeLock Lock(&TagSection);
+void FLLMThreadState::PopTag()
+{
+	FScopeLock Lock(&TagSection);
 
 	LLMCheckf(TagStack.Num() > 0,
 		TEXT("Called FLLMThreadState::PopTag without a matching Push (stack was empty on pop)"));
-		TagStack.Pop(false /* bAllowShrinking */);
-	}
+	TagStack.Pop(false /* bAllowShrinking */);
+}
 
-	const FTagData* FLLMThreadState::GetTopTag()
-	{
-		return TagStack.Num() ? TagStack.Last() : nullptr;
-	}
+const FTagData* FLLMThreadState::GetTopTag()
+{
+	return TagStack.Num() ? TagStack.Last() : nullptr;
+}
 
 #if LLM_ALLOW_ASSETS_TAGS
-	void FLLMThreadState::PushAssetTag(const FTagData* TagData)
-	{
-		FScopeLock Lock(&TagSection);
+void FLLMThreadState::PushAssetTag(const FTagData* TagData)
+{
+	FScopeLock Lock(&TagSection);
 
 	// Push a tag.
-		AssetTagStack.Add(TagData);
-	}
+	AssetTagStack.Add(TagData);
+}
 
-	void FLLMThreadState::PopAssetTag()
-	{
-		FScopeLock Lock(&TagSection);
+void FLLMThreadState::PopAssetTag()
+{
+	FScopeLock Lock(&TagSection);
 
 	LLMCheckf(AssetTagStack.Num() > 0,
 		TEXT("Called FLLMThreadState::PopTag without a matching Push (stack was empty on pop)"));
-		AssetTagStack.Pop(false /* bAllowShrinking */);
-	}
+	AssetTagStack.Pop(false /* bAllowShrinking */);
+}
 
-	const FTagData* FLLMThreadState::GetTopAssetTag()
-	{
-		FScopeLock Lock(&TagSection);
-		return AssetTagStack.Num() ? AssetTagStack.Last() : nullptr;
-	}
+const FTagData* FLLMThreadState::GetTopAssetTag()
+{
+	FScopeLock Lock(&TagSection);
+	return AssetTagStack.Num() ? AssetTagStack.Last() : nullptr;
+}
 #endif
 
-	void FLLMThreadState::IncrTag(const FTagData* TagData, int64 Amount)
-	{
+void FLLMThreadState::IncrTag(const FTagData* TagData, int64 Amount)
+{
 	// Caller is responsible for holding a lock on TagSection.
-		FThreadTagSizeData& AllocationSize = Allocations.FindOrAdd(TagData->GetIndex());
-		AllocationSize.TagData = TagData;
-		AllocationSize.Size += Amount;
-	}
+	FThreadTagSizeData& AllocationSize = Allocations.FindOrAdd(TagData->GetIndex());
+	AllocationSize.TagData = TagData;
+	AllocationSize.Size += Amount;
+}
 
 void FLLMThreadState::TrackAllocation(const void* Ptr, int64 Size, ELLMTracker Tracker, ELLMAllocType AllocType,
 	const FTagData* TagData, const FTagData* AssetTagData, bool bTrackInMemPro)
-	{
-		FScopeLock Lock(&TagSection);
+{
+	FScopeLock Lock(&TagSection);
 
-		AllocTypeAmounts[static_cast<int32>(AllocType)] += Size;
+	AllocTypeAmounts[static_cast<int32>(AllocType)] += Size;
 
-		IncrTag(TagData, Size);
+	IncrTag(TagData, Size);
 #if LLM_ALLOW_ASSETS_TAGS
-		if (AssetTagData)
-		{
-			IncrTag(AssetTagData, Size);
-		}
+	if (AssetTagData)
+	{
+		IncrTag(AssetTagData, Size);
+	}
 #endif
 
-		ELLMTag EnumTag = TagData->GetContainingEnum();
-		if (Tracker == ELLMTracker::Default)
-		{
-			FPlatformMemory::OnLowLevelMemory_Alloc(Ptr, static_cast<uint64>(Size), static_cast<uint64>(EnumTag));
-		}
+	ELLMTag EnumTag = TagData->GetContainingEnum();
+	if (Tracker == ELLMTracker::Default)
+	{
+		FPlatformMemory::OnLowLevelMemory_Alloc(Ptr, static_cast<uint64>(Size), static_cast<uint64>(EnumTag));
+	}
 
 #if MEMPRO_ENABLED
-		if (FMemProProfiler::IsTrackingTag(EnumTag) && bTrackInMemPro)
-		{
-			MEMPRO_TRACK_ALLOC(const_cast<void*>(Ptr), static_cast<size_t>(Size));
-		}
-#endif
+	if (FMemProProfiler::IsTrackingTag(EnumTag) && bTrackInMemPro)
+	{
+		MEMPRO_TRACK_ALLOC(const_cast<void*>(Ptr), static_cast<size_t>(Size));
 	}
+#endif
+}
 
 void FLLMThreadState::TrackFree(const void* Ptr, int64 Size, ELLMTracker Tracker, ELLMAllocType AllocType,
 	const FTagData* TagData, const FTagData* AssetTagData, bool bTrackInMemPro)
-	{
-		FScopeLock Lock(&TagSection);
+{
+	FScopeLock Lock(&TagSection);
 
-		AllocTypeAmounts[static_cast<int32>(AllocType)] -= Size;
+	AllocTypeAmounts[static_cast<int32>(AllocType)] -= Size;
 
-		IncrTag(TagData, -Size);
+	IncrTag(TagData, -Size);
 #if LLM_ALLOW_ASSETS_TAGS
-		if (AssetTagData)
-		{
-			IncrTag(AssetTagData, -Size);
-		}
+	if (AssetTagData)
+	{
+		IncrTag(AssetTagData, -Size);
+	}
 #endif
-		ELLMTag EnumTag = TagData->GetContainingEnum();
-		if (Tracker == ELLMTracker::Default)
-		{
-			FPlatformMemory::OnLowLevelMemory_Free(Ptr, static_cast<uint64>(Size), static_cast<uint64>(EnumTag));
-		}
+	ELLMTag EnumTag = TagData->GetContainingEnum();
+	if (Tracker == ELLMTracker::Default)
+	{
+		FPlatformMemory::OnLowLevelMemory_Free(Ptr, static_cast<uint64>(Size), static_cast<uint64>(EnumTag));
+	}
 
 #if MEMPRO_ENABLED
-		if (FMemProProfiler::IsTrackingTag(EnumTag) && bTrackInMemPro)
-		{
-			MEMPRO_TRACK_FREE(const_cast<void*>(Ptr));
-		}
-#endif
+	if (FMemProProfiler::IsTrackingTag(EnumTag) && bTrackInMemPro)
+	{
+		MEMPRO_TRACK_FREE(const_cast<void*>(Ptr));
 	}
+#endif
+}
 
 void FLLMThreadState::TrackMoved(const void* Dest, const void* Source, int64 Size, ELLMTracker Tracker,
 	const FTagData* TagData)
-	{
+{
 	// Update external memory trackers (ideally would want a proper 'move' option on these).
-		ELLMTag EnumTag = TagData->GetContainingEnum();
-		if (Tracker == ELLMTracker::Default)
-		{
-			FPlatformMemory::OnLowLevelMemory_Free(Source, static_cast<uint64>(Size), static_cast<uint64>(EnumTag));
-			FPlatformMemory::OnLowLevelMemory_Alloc(Dest, static_cast<uint64>(Size), static_cast<uint64>(EnumTag));
-		}
+	ELLMTag EnumTag = TagData->GetContainingEnum();
+	if (Tracker == ELLMTracker::Default)
+	{
+		FPlatformMemory::OnLowLevelMemory_Free(Source, static_cast<uint64>(Size), static_cast<uint64>(EnumTag));
+		FPlatformMemory::OnLowLevelMemory_Alloc(Dest, static_cast<uint64>(Size), static_cast<uint64>(EnumTag));
+	}
 
 #if MEMPRO_ENABLED
-		if (FMemProProfiler::IsTrackingTag(EnumTag))
-		{
-			MEMPRO_TRACK_FREE(const_cast<void*>(Source));
-			MEMPRO_TRACK_ALLOC(const_cast<void*>(Dest), static_cast<size_t>(Size));
-		}
-#endif
+	if (FMemProProfiler::IsTrackingTag(EnumTag))
+	{
+		MEMPRO_TRACK_FREE(const_cast<void*>(Source));
+		MEMPRO_TRACK_ALLOC(const_cast<void*>(Dest), static_cast<size_t>(Size));
 	}
+#endif
+}
 
-	void FLLMThreadState::PropagateChildSizesToParents()
-	{
-		FScopeLock Lock(&TagSection);
+void FLLMThreadState::PropagateChildSizesToParents()
+{
+	FScopeLock Lock(&TagSection);
 
 	// Make sure all parents of any TagDatas in the Allocations are also present.
-		FConstTagDataArray ParentsToAdd;
-		for (const TPair<int32, FThreadTagSizeData>& It : Allocations)
-		{
-			const FTagData* TagData = It.Value.TagData;
-			const FTagData* ParentData = TagData->GetParent();
-			while (ParentData && !Allocations.Contains(ParentData->GetIndex()))
-			{
-				ParentsToAdd.Add(ParentData);
-				ParentData = ParentData->GetParent();
-			}
-		}
-		for (const FTagData* TagData : ParentsToAdd)
+	FConstTagDataArray ParentsToAdd;
+	for (const TPair<int32, FThreadTagSizeData>& It : Allocations)
+	{
+		const FTagData* TagData = It.Value.TagData;
+		const FTagData* ParentData = TagData->GetParent();
+		while (ParentData && !Allocations.Contains(ParentData->GetIndex()))
 		{
-			FThreadTagSizeData& Info = Allocations.FindOrAdd(TagData->GetIndex());
-			Info.TagData = TagData;
+			ParentsToAdd.Add(ParentData);
+			ParentData = ParentData->GetParent();
 		}
+	}
+	for (const FTagData* TagData : ParentsToAdd)
+	{
+		FThreadTagSizeData& Info = Allocations.FindOrAdd(TagData->GetIndex());
+		Info.TagData = TagData;
+	}
 
 	// Tags are sorted topologically from parent to child, so we can accumulate children into parents recursively
 	// by reverse iterating the map.
-		for (FThreadTagSizeMap::TConstReverseIterator It(Allocations); It; ++It)
+	for (FThreadTagSizeMap::TConstReverseIterator It(Allocations); It; ++It)
+	{
+		const FThreadTagSizeData& Info = It->Value;
+		const FTagData* ParentData = Info.TagData->GetParent();
+		if (Info.Size && ParentData)
 		{
-			const FThreadTagSizeData& Info = It->Value;
-			const FTagData* ParentData = Info.TagData->GetParent();
-			if (Info.Size && ParentData)
-			{
-				Allocations.FindChecked(ParentData->GetIndex()).Size += Info.Size;
-			}
+			Allocations.FindChecked(ParentData->GetIndex()).Size += Info.Size;
 		}
 	}
+}
 
-	void FLLMThreadState::OnTagsResorted(FTagDataArray& OldTagDatas)
+void FLLMThreadState::OnTagsResorted(FTagDataArray& OldTagDatas)
+{
+	FScopeLock Lock(&TagSection);
+	TArray<FThreadTagSizeData, FDefaultLLMAllocator> AllocationDatas;
+	AllocationDatas.Reserve(Allocations.Num());
+	for (const TPair<int32, FThreadTagSizeData>& It : Allocations)
 	{
-		FScopeLock Lock(&TagSection);
-		TArray<FThreadTagSizeData, FDefaultLLMAllocator> AllocationDatas;
-		AllocationDatas.Reserve(Allocations.Num());
-		for (const TPair<int32, FThreadTagSizeData>& It : Allocations)
-		{
-			AllocationDatas.Add(It.Value);
-		}
-		Allocations.Reset();
-		for (const FThreadTagSizeData& AllocationData : AllocationDatas)
-		{
-			Allocations.Add(AllocationData.TagData->GetIndex(), AllocationData);
-		}
+		AllocationDatas.Add(It.Value);
 	}
-
-	void FLLMThreadState::LockTags(bool bLock)
+	Allocations.Reset();
+	for (const FThreadTagSizeData& AllocationData : AllocationDatas)
 	{
-		if (bLock)
-		{
-			TagSection.Lock();
-		}
-		else
-		{
-			TagSection.Unlock();
-		}
+		Allocations.Add(AllocationData.TagData->GetIndex(), AllocationData);
 	}
+}
 
-void FLLMThreadState::FetchAndClearTagSizes(FTrackerTagSizeMap& TagSizes, int64* InAllocTypeAmounts,
-	bool bTrimAllocations)
+void FLLMThreadState::LockTags(bool bLock)
+{
+	if (bLock)
 	{
-		FScopeLock Lock(&TagSection);
-		for (TPair<int32, FThreadTagSizeData>& Item : Allocations)
-		{
-			FThreadTagSizeData& ThreadInfo = Item.Value;
-			if (ThreadInfo.Size)
-			{
-				const FTagData* TagData = ThreadInfo.TagData;
-				FTrackerTagSizeData& TrackerInfo = TagSizes.FindOrAdd(TagData);
-				TrackerInfo.Size += ThreadInfo.Size;
-				ThreadInfo.Size = 0;
-			}
-		}
-		if (bTrimAllocations)
-		{
-			Allocations.Empty();
-		}
-
-		for (int32 Index = 0; Index < static_cast<int32>(ELLMAllocType::Count); ++Index)
-		{
-			InAllocTypeAmounts[Index] += AllocTypeAmounts[Index];
-			AllocTypeAmounts[Index] = 0;
-		}
+		TagSection.Lock();
 	}
-
-	void SetMemoryStatByFName(FName Name, int64 Amount)
+	else
 	{
-		if (Name != NAME_None)
-		{
-			SET_MEMORY_STAT_FName(Name, Amount);
-		}
+		TagSection.Unlock();
 	}
+}
 
-	void FLLMThreadState::ClearAllocTypeAmounts()
+void FLLMThreadState::FetchAndClearTagSizes(FTrackerTagSizeMap& TagSizes, int64* InAllocTypeAmounts,
+	bool bTrimAllocations)
+{
+	FScopeLock Lock(&TagSection);
+	for (TPair<int32, FThreadTagSizeData>& Item : Allocations)
 	{
-		for (int32 Index = 0; Index < static_cast<int32>(ELLMAllocType::Count); ++Index)
+		FThreadTagSizeData& ThreadInfo = Item.Value;
+		if (ThreadInfo.Size)
 		{
-			AllocTypeAmounts[Index] = 0;
+			const FTagData* TagData = ThreadInfo.TagData;
+			FTrackerTagSizeData& TrackerInfo = TagSizes.FindOrAdd(TagData);
+			TrackerInfo.Size += ThreadInfo.Size;
+			ThreadInfo.Size = 0;
 		}
 	}
-
-// FLLMCsvWriter implementation.
-	FLLMCsvWriter::FLLMCsvWriter()
-		: Archive(nullptr)
-		, LastWriteTime(FPlatformTime::Seconds())
-		, WriteCount(0)
+	if (bTrimAllocations)
 	{
+		Allocations.Empty();
 	}
 
-	FLLMCsvWriter::~FLLMCsvWriter()
+	for (int32 Index = 0; Index < static_cast<int32>(ELLMAllocType::Count); ++Index)
 	{
-		delete Archive;
+		InAllocTypeAmounts[Index] += AllocTypeAmounts[Index];
+		AllocTypeAmounts[Index] = 0;
 	}
+}
 
-	void FLLMCsvWriter::Clear()
+void SetMemoryStatByFName(FName Name, int64 Amount)
+{
+	if (Name != NAME_None)
 	{
-		Columns.Empty();
-		ExistingColumns.Empty();
+		SET_MEMORY_STAT_FName(Name, Amount);
 	}
+}
 
-	void FLLMCsvWriter::OnPreFork()
+void FLLMThreadState::ClearAllocTypeAmounts()
+{
+	for (int32 Index = 0; Index < static_cast<int32>(ELLMAllocType::Count); ++Index)
 	{
-		if (Archive)
-		{
-			Archive->Flush();
-			delete Archive;
-			Archive = nullptr;
-		}
+		AllocTypeAmounts[Index] = 0;
 	}
+}
+
+// FLLMCsvWriter implementation.
+FLLMCsvWriter::FLLMCsvWriter()
+	: Archive(nullptr)
+	, LastWriteTime(FPlatformTime::Seconds())
+	, WriteCount(0)
+{
+}
+
+FLLMCsvWriter::~FLLMCsvWriter()
+{
+	delete Archive;
+}
+
+void FLLMCsvWriter::Clear()
+{
+	Columns.Empty();
+	ExistingColumns.Empty();
+}
 
-	void FLLMCsvWriter::SetTracker(ELLMTracker InTracker)
+void FLLMCsvWriter::OnPreFork()
+{
+	if (Archive)
 	{
-		Tracker = InTracker;
+		Archive->Flush();
+		delete Archive;
+		Archive = nullptr;
 	}
+}
+
+void FLLMCsvWriter::SetTracker(ELLMTracker InTracker)
+{
+	Tracker = InTracker;
+}
 
 void FLLMCsvWriter::Publish(FLowLevelMemTracker& LLMRef, const FTrackerTagSizeMap& TagSizes,
 	const FTagData* OverrideTrackedTotalTagData, const FTagData* OverrideUntaggedTagData, int64 TrackedTotal,
 	bool bTrackPeaks)
-	{
-		const double Now = FPlatformTime::Seconds();
+{
+	const double Now = FPlatformTime::Seconds();
 
-		if ((FLowLevelMemTracker::Get().bPublishSingleFrame == false) && 
-			(Now - LastWriteTime < (double)CVarLLMWriteInterval.GetValueOnAnyThread()))
-		{
-			return;
-		}
+	if ((FLowLevelMemTracker::Get().bPublishSingleFrame == false) && 
+		(Now - LastWriteTime < (double)CVarLLMWriteInterval.GetValueOnAnyThread()))
+	{
+		return;
+	}
 
-		LastWriteTime = Now;
+	LastWriteTime = Now;
 
-		const bool bCreatedArchive = CreateArchive();
-		const bool bColumnsUpdated = UpdateColumns(TagSizes);
+	const bool bCreatedArchive = CreateArchive();
+	const bool bColumnsUpdated = UpdateColumns(TagSizes);
 
-		if (bCreatedArchive || bColumnsUpdated)
-		{
+	if (bCreatedArchive || bColumnsUpdated)
+	{
 		// The column names are written at the start of the archive; when they change we seek back to the start of
 		// the file and rewrite the column names.
-			WriteHeader(OverrideTrackedTotalTagData, OverrideUntaggedTagData);
-		}
-
-		AddRow(LLMRef, TagSizes, OverrideTrackedTotalTagData, OverrideUntaggedTagData, TrackedTotal, bTrackPeaks);
+		WriteHeader(OverrideTrackedTotalTagData, OverrideUntaggedTagData);
 	}
 
-	const TCHAR* FLLMCsvWriter::GetTrackerCsvName(ELLMTracker InTracker)
-	{
-		switch (InTracker)
-		{
-			case ELLMTracker::Default: return TEXT("LLM");
-			case ELLMTracker::Platform: return TEXT("LLMPlatform");
-			default: LLMCheck(false); return TEXT("");
-		}
-	}
+	AddRow(LLMRef, TagSizes, OverrideTrackedTotalTagData, OverrideUntaggedTagData, TrackedTotal, bTrackPeaks);
+}
 
-	// Archive is a binary stream, so we can't just serialise an FString using <<.
-	void FLLMCsvWriter::Write(FStringView Text)
+const TCHAR* FLLMCsvWriter::GetTrackerCsvName(ELLMTracker InTracker)
+{
+	switch (InTracker)
 	{
-		Archive->Serialize(const_cast<ANSICHAR*>(
-			StringCast<ANSICHAR>(Text.GetData(), Text.Len()).Get()), Text.Len() * sizeof(ANSICHAR));
+		case ELLMTracker::Default: return TEXT("LLM");
+		case ELLMTracker::Platform: return TEXT("LLMPlatform");
+		default: LLMCheck(false); return TEXT("");
 	}
+}
 
-	bool FLLMCsvWriter::CreateArchive()
+// Archive is a binary stream, so we can't just serialise an FString using <<.
+void FLLMCsvWriter::Write(FStringView Text)
+{
+	Archive->Serialize(const_cast<ANSICHAR*>(
+		StringCast<ANSICHAR>(Text.GetData(), Text.Len()).Get()), Text.Len() * sizeof(ANSICHAR));
+}
+
+bool FLLMCsvWriter::CreateArchive()
+{
+	if (Archive)
 	{
-		if (Archive)
-		{
-			return false;
-		}
+		return false;
+	}
 
 	// Create the csv file.
-		FString Directory = FPaths::ProfilingDir() + TEXT("LLM/");
-		IFileManager::Get().MakeDirectory(*Directory, true);
+	FString Directory = FPaths::ProfilingDir() + TEXT("LLM/");
+	IFileManager::Get().MakeDirectory(*Directory, true);
 
-		const TCHAR* TrackerName = GetTrackerCsvName(Tracker);
-		const FDateTime FileDate = FDateTime::Now();
+	const TCHAR* TrackerName = GetTrackerCsvName(Tracker);
+	const FDateTime FileDate = FDateTime::Now();
 #if PLATFORM_DESKTOP
-		FString PlatformName = FPlatformProperties::PlatformName();
+	FString PlatformName = FPlatformProperties::PlatformName();
 #else // Use the CPU for consoles so we can differentiate things like different SKUs of a console generation.
-		FString PlatformName = FPlatformMisc::GetCPUBrand().TrimStartAndEnd();
+	FString PlatformName = FPlatformMisc::GetCPUBrand().TrimStartAndEnd();
 #endif
-		PlatformName.ReplaceCharInline(' ', '_');
-		PlatformName = FPaths::MakeValidFileName(PlatformName);
+	PlatformName.ReplaceCharInline(' ', '_');
+	PlatformName = FPaths::MakeValidFileName(PlatformName);
 #if WITH_SERVER_CODE
 	FString Filename = FString::Printf(TEXT("%s/%s_Pid%d_%s_%s.csv"), *Directory, TrackerName,
 		FPlatformProcess::GetCurrentProcessId(), *FileDate.ToString(), *PlatformName);
@@ -4121,361 +4119,360 @@ void FLLMCsvWriter::Publish(FLowLevelMemTracker& LLMRef, const FTrackerTagSizeMa
 	FString Filename = FString::Printf(TEXT("%s/%s_%s_%s.csv"), *Directory, TrackerName, *FileDate.ToString(),
 		*PlatformName);
 #endif
-		Archive = IFileManager::Get().CreateFileWriter(*Filename, FILEWRITE_AllowRead | FILEWRITE_NoFail);
-		LLMCheck(Archive);
+	Archive = IFileManager::Get().CreateFileWriter(*Filename, FILEWRITE_AllowRead | FILEWRITE_NoFail);
+	LLMCheck(Archive);
 
 	// Create space for column titles that are filled in as we get them.
-		Write(FString::ChrN(CVarLLMHeaderMaxSize.GetValueOnAnyThread(), ' '));
-		Write(TEXT("\n"));
+	Write(FString::ChrN(CVarLLMHeaderMaxSize.GetValueOnAnyThread(), ' '));
+	Write(TEXT("\n"));
 
-		return true;
-	}
+	return true;
+}
 
-	bool FLLMCsvWriter::UpdateColumns(const FTrackerTagSizeMap& TagSizes)
-	{
-		bool bUpdated = false;
+bool FLLMCsvWriter::UpdateColumns(const FTrackerTagSizeMap& TagSizes)
+{
+	bool bUpdated = false;
 
-		for (const TPair<const FTagData*, FTrackerTagSizeData>& It : TagSizes)
+	for (const TPair<const FTagData*, FTrackerTagSizeData>& It : TagSizes)
+	{
+		const FTagData* TagData = It.Key;
+		if (!TagData->IsReportable())
+		{
+			continue;
+		}
+		if (TagData->GetName() == TagName_Untagged)
 		{
-			const FTagData* TagData = It.Key;
-			if (!TagData->IsReportable())
-			{
-				continue;
-			}
-			if (TagData->GetName() == TagName_Untagged)
-			{
 			continue; // Handled by OverrideUntaggedName.
-			}
-			if (ExistingColumns.Contains(TagData))
-			{
-				continue;
-			}
-
-			ExistingColumns.Add(TagData);
-			Columns.Add(TagData);
-			bUpdated = true;
 		}
-		return bUpdated;
+		if (ExistingColumns.Contains(TagData))
+		{
+			continue;
+		}
+
+		ExistingColumns.Add(TagData);
+		Columns.Add(TagData);
+		bUpdated = true;
 	}
+	return bUpdated;
+}
 
-void FLLMCsvWriter::WriteHeader(const FTagData* OverrideTrackedTotalTagData,
-	const FTagData* OverrideUntaggedTagData)
-	{
-		int64 OriginalOffset = Archive->Tell();
-		Archive->Seek(0);
+void FLLMCsvWriter::WriteHeader(const FTagData* OverrideTrackedTotalTagData, const FTagData* OverrideUntaggedTagData)
+{
+	int64 OriginalOffset = Archive->Tell();
+	Archive->Seek(0);
 
-		TStringBuilder<256> NameBuffer;
-		auto WriteTagData = [this, &NameBuffer](const FTagData* TagData)
+	TStringBuilder<256> NameBuffer;
+	auto WriteTagData = [this, &NameBuffer](const FTagData* TagData)
+	{
+		if (!TagData)
 		{
-			if (!TagData)
-			{
-				return;
-			}
+			return;
+		}
 
-			NameBuffer.Reset();
-			TagData->AppendDisplayPath(NameBuffer);
-			NameBuffer << TEXT(",");
-			Write(NameBuffer);
-		};
+		NameBuffer.Reset();
+		TagData->AppendDisplayPath(NameBuffer);
+		NameBuffer << TEXT(",");
+		Write(NameBuffer);
+	};
 
-		WriteTagData(OverrideTrackedTotalTagData);
-		WriteTagData(OverrideUntaggedTagData);
-		for (const FTagData* TagData : Columns)
-		{
-			WriteTagData(TagData);
-		}
+	WriteTagData(OverrideTrackedTotalTagData);
+	WriteTagData(OverrideUntaggedTagData);
+	for (const FTagData* TagData : Columns)
+	{
+		WriteTagData(TagData);
+	}
 
-		int64 ColumnTitleTotalSize = Archive->Tell();
-		if (ColumnTitleTotalSize >= CVarLLMHeaderMaxSize.GetValueOnAnyThread())
-		{
+	int64 ColumnTitleTotalSize = Archive->Tell();
+	if (ColumnTitleTotalSize >= CVarLLMHeaderMaxSize.GetValueOnAnyThread())
+	{
 		UE_LOG(LogHAL, Error,
 			TEXT("LLM column titles have overflowed, LLM CSM data will be corrupted. Increase CVarLLMHeaderMaxSize > %d"),
 			ColumnTitleTotalSize);
-		}
-
-		Archive->Seek(OriginalOffset);
 	}
 
+	Archive->Seek(OriginalOffset);
+}
+
 void FLLMCsvWriter::AddRow(FLowLevelMemTracker& LLMRef, const FTrackerTagSizeMap& TagSizes,
 	const FTagData* OverrideTrackedTotalTagData, const FTagData* OverrideUntaggedTagData, int64 TrackedTotal,
 	bool bTrackPeaks)
+{
+	TStringBuilder<256> TextBuffer;
+	auto WriteValue = [this, &TextBuffer](int64 Value)
 	{
-		TStringBuilder<256> TextBuffer;
-		auto WriteValue = [this, &TextBuffer](int64 Value)
+		TextBuffer.Reset();
+		TextBuffer.Appendf(TEXT("%0.2f,"), (float)Value / 1024.0f / 1024.0f);
+		Write(TextBuffer);
+	};
+	auto WriteTag = [&WriteValue, &TagSizes, bTrackPeaks](const FTagData* TagData)
+	{
+		if (!TagData)
 		{
-			TextBuffer.Reset();
-			TextBuffer.Appendf(TEXT("%0.2f,"), (float)Value / 1024.0f / 1024.0f);
-			Write(TextBuffer);
-		};
-		auto WriteTag = [&WriteValue, &TagSizes, bTrackPeaks](const FTagData* TagData)
+			WriteValue(0);
+		}
+		else
 		{
-			if (!TagData)
+			const FTrackerTagSizeData* AllocationData = TagSizes.Find(TagData);
+			if (!AllocationData)
 			{
 				WriteValue(0);
 			}
 			else
 			{
-				const FTrackerTagSizeData* AllocationData = TagSizes.Find(TagData);
-				if (!AllocationData)
-				{
-					WriteValue(0);
-				}
-				else
-				{
-					WriteValue(AllocationData->GetSize(bTrackPeaks));
-				}
+				WriteValue(AllocationData->GetSize(bTrackPeaks));
 			}
-		};
-
-		if (OverrideTrackedTotalTagData)
-		{
-			WriteValue(TrackedTotal);
-		}
-		if (OverrideUntaggedTagData)
-		{
-			WriteTag(LLMRef.FindTagData(TagName_Untagged));
-		}
-
-		for (const FTagData* TagData : Columns)
-		{
-			WriteTag(TagData);
-		}
-		Write(TEXTVIEW("\n"));
-
-		WriteCount++;
-
-		if (CVarLLMWriteInterval.GetValueOnAnyThread())
-		{
-			UE_LOG(LogHAL, Log, TEXT("Wrote LLM csv line %d"), WriteCount);
 		}
+	};
 
-		Archive->Flush();
+	if (OverrideTrackedTotalTagData)
+	{
+		WriteValue(TrackedTotal);
 	}
-
-// FLLMTraceWriter implementation.
-	FLLMTraceWriter::FLLMTraceWriter()
+	if (OverrideUntaggedTagData)
 	{
+		WriteTag(LLMRef.FindTagData(TagName_Untagged));
 	}
 
-	inline void FLLMTraceWriter::SetTracker(ELLMTracker InTracker)
-	{ 
-		Tracker = InTracker;
+	for (const FTagData* TagData : Columns)
+	{
+		WriteTag(TagData);
 	}
+	Write(TEXTVIEW("\n"));
+
+	WriteCount++;
 
-	void FLLMTraceWriter::Clear()
+	if (CVarLLMWriteInterval.GetValueOnAnyThread())
 	{
-		DeclaredTags.Empty();
+		UE_LOG(LogHAL, Log, TEXT("Wrote LLM csv line %d"), WriteCount);
 	}
 
-	const void* FLLMTraceWriter::GetTagId(const FTagData* TagData)
+	Archive->Flush();
+}
+
+// FLLMTraceWriter implementation.
+FLLMTraceWriter::FLLMTraceWriter()
+{
+}
+
+inline void FLLMTraceWriter::SetTracker(ELLMTracker InTracker)
+{ 
+	Tracker = InTracker;
+}
+
+void FLLMTraceWriter::Clear()
+{
+	DeclaredTags.Empty();
+}
+
+const void* FLLMTraceWriter::GetTagId(const FTagData* TagData)
+{
+	if (!TagData)
 	{
-		if (!TagData)
-		{
-			return nullptr;
-		}
-		return reinterpret_cast<const void*>(TagData);
+		return nullptr;
 	}
+	return reinterpret_cast<const void*>(TagData);
+}
 
 void FLLMTraceWriter::Publish(FLowLevelMemTracker& LLMRef, const FTrackerTagSizeMap& TagSizes,
 	const FTagData* OverrideTrackedTotalTagData, const FTagData* OverrideUntaggedTagData, int64 TrackedTotal,
 	bool bTrackPeaks)
+{
+	if (!UE_TRACE_CHANNELEXPR_IS_ENABLED(MemTagChannel))
+	{
+		return;
+	}
+
+	if (!bTrackerSpecSent)
+	{
+		bTrackerSpecSent = true;
+		static struct {
+			const ANSICHAR* Name;
+			uint32 Len;
+		} TrackerNames[] = {
+			{ "Platform", 8 },
+			{ "Default", 7 },
+		};
+		static_assert(UE_ARRAY_COUNT(TrackerNames) == int(ELLMTracker::Max), "");
+		uint32 NameLen = TrackerNames[(uint8)Tracker].Len;
+		UE_TRACE_LOG(LLM, TrackerSpec, MemTagChannel, NameLen * sizeof(ANSICHAR))
+			<< TrackerSpec.TrackerId((uint8)Tracker)
+			<< TrackerSpec.Name(TrackerNames[(uint8)Tracker].Name, NameLen);
+	}
+
+	TStringBuilder<1024> NameBuffer;
+	auto SendTagDeclaration = [this, &NameBuffer](const FTagData* TagData)
 	{
-		if (!UE_TRACE_CHANNELEXPR_IS_ENABLED(MemTagChannel))
+		if (!TagData || DeclaredTags.Contains(TagData))
 		{
 			return;
 		}
+		DeclaredTags.Add(TagData);
 
-		if (!bTrackerSpecSent)
+		const FTagData* Parent = TagData->GetParent();
+		NameBuffer.Reset();
+		TagData->AppendDisplayPath(NameBuffer);
+		UE_TRACE_LOG(LLM, TagsSpec, MemTagChannel, NameBuffer.Len() * sizeof(ANSICHAR))
+			<< TagsSpec.TagId(GetTagId(TagData))
+			<< TagsSpec.ParentId(GetTagId(Parent))
+			<< TagsSpec.Name(*NameBuffer, NameBuffer.Len());
+	};
+	SendTagDeclaration(OverrideTrackedTotalTagData);
+	SendTagDeclaration(OverrideUntaggedTagData);
+	for (const TPair<const FTagData*, FTrackerTagSizeData>& It : TagSizes)
+	{
+		const FTagData* TagData = It.Key;
+		if (!TagData->IsReportable())
 		{
-			bTrackerSpecSent = true;
-			static struct {
-				const ANSICHAR* Name;
-				uint32 Len;
-			} TrackerNames[] = {
-				{ "Platform", 8 },
-				{ "Default", 7 },
-			};
-			static_assert(UE_ARRAY_COUNT(TrackerNames) == int(ELLMTracker::Max), "");
-			uint32 NameLen = TrackerNames[(uint8)Tracker].Len;
-			UE_TRACE_LOG(LLM, TrackerSpec, MemTagChannel, NameLen * sizeof(ANSICHAR))
-				<< TrackerSpec.TrackerId((uint8)Tracker)
-				<< TrackerSpec.Name(TrackerNames[(uint8)Tracker].Name, NameLen);
+			continue;
 		}
-
-		TStringBuilder<1024> NameBuffer;
-		auto SendTagDeclaration = [this, &NameBuffer](const FTagData* TagData)
-		{
-			if (!TagData || DeclaredTags.Contains(TagData))
-			{
-				return;
-			}
-			DeclaredTags.Add(TagData);
-
-			const FTagData* Parent = TagData->GetParent();
-			NameBuffer.Reset();
-			TagData->AppendDisplayPath(NameBuffer);
-			UE_TRACE_LOG(LLM, TagsSpec, MemTagChannel, NameBuffer.Len() * sizeof(ANSICHAR))
-				<< TagsSpec.TagId(GetTagId(TagData))
-				<< TagsSpec.ParentId(GetTagId(Parent))
-				<< TagsSpec.Name(*NameBuffer, NameBuffer.Len());
-		};
-		SendTagDeclaration(OverrideTrackedTotalTagData);
-		SendTagDeclaration(OverrideUntaggedTagData);
-		for (const TPair<const FTagData*, FTrackerTagSizeData>& It : TagSizes)
+		if (OverrideUntaggedTagData != nullptr && TagData->GetName() == TagName_Untagged)
 		{
-			const FTagData* TagData = It.Key;
-			if (!TagData->IsReportable())
-			{
-				continue;
-			}
-			if (OverrideUntaggedTagData != nullptr && TagData->GetName() == TagName_Untagged)
-			{
 			continue; // Handled by OverrideUntaggedTagData.
-			}
-			SendTagDeclaration(TagData);
 		}
+		SendTagDeclaration(TagData);
+	}
 
-		TArray<const void*, FDefaultLLMAllocator> TagIds;
-		TArray<int64, FDefaultLLMAllocator> TagValues;
-		TagIds.Reserve(TagSizes.Num() + 2);
-		TagValues.Reserve(TagSizes.Num() + 2);
-		auto AddValue = [&TagIds, &TagValues](const FTagData* TagData, int64 Value)
+	TArray<const void*, FDefaultLLMAllocator> TagIds;
+	TArray<int64, FDefaultLLMAllocator> TagValues;
+	TagIds.Reserve(TagSizes.Num() + 2);
+	TagValues.Reserve(TagSizes.Num() + 2);
+	auto AddValue = [&TagIds, &TagValues](const FTagData* TagData, int64 Value)
+	{
+		if (!TagData)
 		{
-			if (!TagData)
-			{
-				return;
-			}
-			TagIds.Add(GetTagId(TagData));
-			TagValues.Add(Value);
-		};
+			return;
+		}
+		TagIds.Add(GetTagId(TagData));
+		TagValues.Add(Value);
+	};
 
-		AddValue(OverrideTrackedTotalTagData, TrackedTotal);
-		if (OverrideUntaggedTagData)
+	AddValue(OverrideTrackedTotalTagData, TrackedTotal);
+	if (OverrideUntaggedTagData)
+	{
+		const FTagData* TagData = LLMRef.FindTagData(TagName_Untagged);
+		if (!TagData)
 		{
-			const FTagData* TagData = LLMRef.FindTagData(TagName_Untagged);
-			if (!TagData)
-			{
-				AddValue(OverrideUntaggedTagData, 0);
-			}
-			else
-			{
-				const FTrackerTagSizeData* AllocationData = TagSizes.Find(TagData);
-				AddValue(OverrideUntaggedTagData, AllocationData ? AllocationData->GetSize(bTrackPeaks) : 0);
-			}
+			AddValue(OverrideUntaggedTagData, 0);
 		}
-
-		for (const TPair<const FTagData*, FTrackerTagSizeData>& It : TagSizes)
+		else
 		{
-			const FTagData* TagData = It.Key;
-			if (!TagData->IsReportable())
-			{
-				continue;
-			}
-			if (OverrideUntaggedTagData && TagData->GetName() == TagName_Untagged)
-			{
-			continue; // Handled by OverrideUntaggedTagData.
-			}
-			AddValue(TagData, It.Value.GetSize(bTrackPeaks));
+			const FTrackerTagSizeData* AllocationData = TagSizes.Find(TagData);
+			AddValue(OverrideUntaggedTagData, AllocationData ? AllocationData->GetSize(bTrackPeaks) : 0);
 		}
-
-		int TagCount = TagIds.Num();
-		LLMCheck(TagCount == TagValues.Num());
-		const uint64 Cycle = FPlatformTime::Cycles64();
-		UE_TRACE_LOG(LLM, TagValue, MemTagChannel)
-			<< TagValue.TrackerId((uint8)Tracker)
-			<< TagValue.Cycle(Cycle)
-			<< TagValue.Tags(TagIds.GetData(), TagCount)
-			<< TagValue.Values(TagValues.GetData(), TagCount);
 	}
 
-	void FLLMTraceWriter::TraceGenericTags(FLowLevelMemTracker& LLMRef)
+	for (const TPair<const FTagData*, FTrackerTagSizeData>& It : TagSizes)
 	{
-		for (int32 GenericTagIndex = 0; GenericTagIndex < static_cast<int32>(ELLMTag::GenericTagCount); GenericTagIndex++)
+		const FTagData* TagData = It.Key;
+		if (!TagData->IsReportable())
+		{
+			continue;
+		}
+		if (OverrideUntaggedTagData && TagData->GetName() == TagName_Untagged)
 		{
-			const FTagData* TagData = LLMRef.FindTagData(static_cast<ELLMTag>(GenericTagIndex));
-			FString TagName = TagData->GetDisplayPath();
-			UE_TRACE_LOG(LLM, TagsSpec, MemTagChannel, TagName.Len() * sizeof(ANSICHAR))
-				<< TagsSpec.TagId(GetTagId(TagData))
-				<< TagsSpec.ParentId(GetTagId(TagData->GetParent()))
-				<< TagsSpec.Name(*TagName, TagName.Len());
+			continue; // Handled by OverrideUntaggedTagData.
 		}
+		AddValue(TagData, It.Value.GetSize(bTrackPeaks));
 	}
 
-// FLLMCsvProfilerWriter implementation.
-	FLLMCsvProfilerWriter::FLLMCsvProfilerWriter()
+	int TagCount = TagIds.Num();
+	LLMCheck(TagCount == TagValues.Num());
+	const uint64 Cycle = FPlatformTime::Cycles64();
+	UE_TRACE_LOG(LLM, TagValue, MemTagChannel)
+		<< TagValue.TrackerId((uint8)Tracker)
+		<< TagValue.Cycle(Cycle)
+		<< TagValue.Tags(TagIds.GetData(), TagCount)
+		<< TagValue.Values(TagValues.GetData(), TagCount);
+}
+
+void FLLMTraceWriter::TraceGenericTags(FLowLevelMemTracker& LLMRef)
+{
+	for (int32 GenericTagIndex = 0; GenericTagIndex < static_cast<int32>(ELLMTag::GenericTagCount); GenericTagIndex++)
 	{
+		const FTagData* TagData = LLMRef.FindTagData(static_cast<ELLMTag>(GenericTagIndex));
+		FString TagName = TagData->GetDisplayPath();
+		UE_TRACE_LOG(LLM, TagsSpec, MemTagChannel, TagName.Len() * sizeof(ANSICHAR))
+			<< TagsSpec.TagId(GetTagId(TagData))
+			<< TagsSpec.ParentId(GetTagId(TagData->GetParent()))
+			<< TagsSpec.Name(*TagName, TagName.Len());
 	}
+}
 
-	void FLLMCsvProfilerWriter::Clear()
-	{
+// FLLMCsvProfilerWriter implementation.
+FLLMCsvProfilerWriter::FLLMCsvProfilerWriter()
+{
+}
+
+void FLLMCsvProfilerWriter::Clear()
+{
 #if LLM_CSV_PROFILER_WRITER_ENABLED
-		TagDataToCsvStatName.Empty();
+	TagDataToCsvStatName.Empty();
 #endif
-	}
+}
 
-	void FLLMCsvProfilerWriter::SetTracker(ELLMTracker InTracker)
-	{
+void FLLMCsvProfilerWriter::SetTracker(ELLMTracker InTracker)
+{
 #if LLM_CSV_PROFILER_WRITER_ENABLED
-		check(InTracker == ELLMTracker::Platform || InTracker == ELLMTracker::Default);
-		Tracker = InTracker;
+	check(InTracker == ELLMTracker::Platform || InTracker == ELLMTracker::Default);
+	Tracker = InTracker;
 #endif
-	}
+}
 
 void FLLMCsvProfilerWriter::Publish(FLowLevelMemTracker& LLMRef, const FTrackerTagSizeMap& TagSizes,
 	const FTagData* OverrideTrackedTotalTagData, const FTagData* OverrideUntaggedTagData, int64 TrackedTotal,
 	bool bTrackPeaks)
-	{
+{
 #if LLM_CSV_PROFILER_WRITER_ENABLED
 	int32 CsvCategoryIndex = (Tracker == ELLMTracker::Platform) ?
 		CSV_CATEGORY_INDEX(LLMPlatform) : CSV_CATEGORY_INDEX(LLM);
 
-		if (OverrideTrackedTotalTagData)
-		{
-			RecordTagToCsv(CsvCategoryIndex, OverrideTrackedTotalTagData, TrackedTotal);
-		}
+	if (OverrideTrackedTotalTagData)
+	{
+		RecordTagToCsv(CsvCategoryIndex, OverrideTrackedTotalTagData, TrackedTotal);
+	}
 
-		if (OverrideUntaggedTagData)
-		{
-			const FTagData* TagData = LLMRef.FindTagData(TagName_Untagged);
-			const FTrackerTagSizeData* AllocationData = TagData ? TagSizes.Find(TagData) : nullptr;
+	if (OverrideUntaggedTagData)
+	{
+		const FTagData* TagData = LLMRef.FindTagData(TagName_Untagged);
+		const FTrackerTagSizeData* AllocationData = TagData ? TagSizes.Find(TagData) : nullptr;
 		RecordTagToCsv(CsvCategoryIndex, OverrideUntaggedTagData,
 			AllocationData ? AllocationData->GetSize(bTrackPeaks) : 0);
-		}
+	}
 
-		for (const TPair<const FTagData*, FTrackerTagSizeData>& It : TagSizes)
+	for (const TPair<const FTagData*, FTrackerTagSizeData>& It : TagSizes)
+	{
+		const FTagData* TagData = It.Key;
+		if (!TagData->IsReportable())
+		{
+			continue;
+		}
+		if (OverrideUntaggedTagData && TagData->GetName() == TagName_Untagged)
 		{
-			const FTagData* TagData = It.Key;
-			if (!TagData->IsReportable())
-			{
-				continue;
-			}
-			if (OverrideUntaggedTagData && TagData->GetName() == TagName_Untagged)
-			{
 			// Handled separately by OverrideUntaggedTagData.
-				continue;
-			}
-			const FTrackerTagSizeData* AllocationData = TagSizes.Find(TagData);
-			RecordTagToCsv(CsvCategoryIndex, TagData, AllocationData ? AllocationData->GetSize(bTrackPeaks) : 0);
+			continue;
 		}
-#endif
+		const FTrackerTagSizeData* AllocationData = TagSizes.Find(TagData);
+		RecordTagToCsv(CsvCategoryIndex, TagData, AllocationData ? AllocationData->GetSize(bTrackPeaks) : 0);
 	}
+#endif
+}
 
-	void FLLMCsvProfilerWriter::RecordTagToCsv(int32 CsvCategoryIndex, const FTagData* TagData, int64 TagSize)
-	{
+void FLLMCsvProfilerWriter::RecordTagToCsv(int32 CsvCategoryIndex, const FTagData* TagData, int64 TagSize)
+{
 #if LLM_CSV_PROFILER_WRITER_ENABLED
-		FName NewCsvStatName;
-		FName* CsvStatNamePtr = TagDataToCsvStatName.Find(TagData);
-		if (CsvStatNamePtr == nullptr)
-		{
-			NewCsvStatName = FName(TagData->GetDisplayPath());
-			TagDataToCsvStatName.Add(TagData, NewCsvStatName);
-			CsvStatNamePtr = &NewCsvStatName;
-		}
+	FName NewCsvStatName;
+	FName* CsvStatNamePtr = TagDataToCsvStatName.Find(TagData);
+	if (CsvStatNamePtr == nullptr)
+	{
+		NewCsvStatName = FName(TagData->GetDisplayPath());
+		TagDataToCsvStatName.Add(TagData, NewCsvStatName);
+		CsvStatNamePtr = &NewCsvStatName;
+	}
 	FCsvProfiler::RecordCustomStat(*CsvStatNamePtr, CsvCategoryIndex, (float)((double)TagSize / (1024.0 * 1024.0)),
 		ECsvCustomStatOp::Set);
 #endif
-	}
+}
 
-} // UE::LLMPrivate
+} // namespace UE::LLMPrivate
 
-#endif		// #if ENABLE_LOW_LEVEL_MEM_TRACKER
+#endif // #if ENABLE_LOW_LEVEL_MEM_TRACKER
