Fix MB3 Realloc reading out of bounds

Reported by [REDACTED_PERSON], detected via UBsan

diff --git a/Engine/Source/Runtime/Core/Private/HAL/MallocBinned3.cpp b/Engine/Source/Runtime/Core/Private/HAL/MallocBinned3.cpp
--- a/Engine/Source/Runtime/Core/Private/HAL/MallocBinned3.cpp
+++ b/Engine/Source/Runtime/Core/Private/HAL/MallocBinned3.cpp
@@ -108,7 +108,7 @@ TAtomic<int32> MemoryRangeFreeTotalCount(0);
 #define BINNED3_LARGE_POOL_CANARIES 1 // need to repad the data structure so that the pagesize divides by this to disable
 
 MS_ALIGN(PLATFORM_CACHE_LINE_SIZE) static uint8 Binned3UnusedAlignPadding[PLATFORM_CACHE_LINE_SIZE] GCC_ALIGN(PLATFORM_CACHE_LINE_SIZE) = { 0 };
-uint16 FMallocBinned3::SmallBlockSizesReversedShifted[BINNED3_SMALL_POOL_COUNT] = { 0 };
+uint16 FMallocBinned3::SmallBlockSizesReversedShifted[BINNED3_SMALL_POOL_COUNT + 1] = { 0 };
 uint32 FMallocBinned3::Binned3TlsSlot = 0;
 uint32 FMallocBinned3::OsAllocationGranularity = 0;
 
diff --git a/Engine/Source/Runtime/Core/Public/HAL/MallocBinned3.h b/Engine/Source/Runtime/Core/Public/HAL/MallocBinned3.h
--- a/Engine/Source/Runtime/Core/Public/HAL/MallocBinned3.h
+++ b/Engine/Source/Runtime/Core/Public/HAL/MallocBinned3.h
@@ -580,7 +580,7 @@ public:
 				{
 					// Reallocate to a smaller/bigger pool if necessary
 					BlockSize = PoolIndexToBlockSize(PoolIndex);
-					if ((!!NewSize) & (NewSize <= BlockSize) & ((!PoolIndex) | (NewSize > PoolIndexToBlockSize(PoolIndex - 1))))
+					if ((!!NewSize) & (NewSize <= BlockSize) & ((!PoolIndex) | (NewSize > PoolIndexToBlockSize(static_cast<uint32>(PoolIndex - 1)))))
 					{
 #if BINNED3_ALLOCATOR_STATS
 						SmallPoolTables[PoolIndex].HeadEndAlloc(NewSize);
@@ -697,7 +697,9 @@ public:
 	/** Dumps current allocator stats to the log. */
 	virtual void DumpAllocatorStats(class FOutputDevice& Ar) override;
 	
-	static uint16 SmallBlockSizesReversedShifted[BINNED3_SMALL_POOL_COUNT]; // this is reversed to get the smallest elements on our main cache line
+	// +1 enables PoolIndexToBlockSize(~0u / -1) dummy access that helps avoid PoolIndex == 0 branching in Realloc(),
+	// see ((!PoolIndex) | (NewSize > PoolIndexToBlockSize(static_cast<uint32>(PoolIndex - 1)))
+	static uint16 SmallBlockSizesReversedShifted[BINNED3_SMALL_POOL_COUNT + 1]; // this is reversed to get the smallest elements on our main cache line
 	static FMallocBinned3* MallocBinned3;
 	static uint32 Binned3TlsSlot;
 	static uint32 OsAllocationGranularity;
