Fix merge error
#rnx

diff --git a/Engine/Source/Runtime/Core/Public/Misc/Build.h b/Engine/Source/Runtime/Core/Public/Misc/Build.h
--- a/Engine/Source/Runtime/Core/Public/Misc/Build.h
+++ b/Engine/Source/Runtime/Core/Public/Misc/Build.h
@@ -478,11 +478,15 @@
 
 // Controls if iostore will be forced on
 #ifndef UE_FORCE_USE_IOSTORE
-#define UE_FORCE_USE_IOSTORE 0??// Controls if paks will be forced on unless -NoPaks argument is passed
+#define UE_FORCE_USE_IOSTORE 0
+#endif
+
+// Controls if paks will be forced on unless -NoPaks argument is passed
 #ifndef UE_FORCE_USE_PAKS
 #define UE_FORCE_USE_PAKS 0
 #endif
 
+
 // Controls whether Iris networking code is compiled in or not; should normally be defined by UBT
 #ifndef UE_WITH_IRIS
 	#define UE_WITH_IRIS 0
