Added missing FApplePlatformSymbolDatabaseKeyFuncs::KeyType.

diff --git a/Engine/Source/Runtime/Core/Public/Apple/ApplePlatformSymbolication.h b/Engine/Source/Runtime/Core/Public/Apple/ApplePlatformSymbolication.h
--- a/Engine/Source/Runtime/Core/Public/Apple/ApplePlatformSymbolication.h
+++ b/Engine/Source/Runtime/Core/Public/Apple/ApplePlatformSymbolication.h
@@ -41,6 +41,8 @@ struct FApplePlatformSymbolDatabaseKeyFuncs
 {
 	enum { bAllowDuplicateKeys = 0 };
 
+	typedef FString KeyType;
+
 private:
 	typedef typename TCallTraits<FString>::ParamType KeyInitType;
 	typedef typename TCallTraits<FApplePlatformSymbolDatabase>::ParamType ElementInitType;
