LowLevelMemTracker: Add OnLowLevelChangeInMemoryUse for systems like MiMalloc that can allocate and free subsets of a single reservation.
#rnx

diff --git a/Engine/Source/Runtime/Core/Private/HAL/LowLevelMemTracker.cpp b/Engine/Source/Runtime/Core/Private/HAL/LowLevelMemTracker.cpp
--- a/Engine/Source/Runtime/Core/Private/HAL/LowLevelMemTracker.cpp
+++ b/Engine/Source/Runtime/Core/Private/HAL/LowLevelMemTracker.cpp
@@ -389,7 +389,9 @@ public:
 	void TrackAllocation(const void* Ptr, int64 Size, ELLMTracker Tracker, ELLMAllocType AllocType,
 		const FTagData* TagData, const FTagData* AssetTagData, bool bTrackInMemPro);
 	void TrackFree(const void* Ptr, int64 Size, ELLMTracker Tracker, ELLMAllocType AllocType,
-	const FTagData* TagData, const FTagData* AssetTagData, bool bTrackInMemPro);
+		const FTagData* TagData, const FTagData* AssetTagData, bool bTrackInMemPro);
+	void TrackMemory(int64 Amount, ELLMTracker Tracker, ELLMAllocType AllocType,
+		const FTagData* TagData, const FTagData* AssetTagData);
 	void TrackMoved(const void* Dest, const void* Source, int64 Size, ELLMTracker Tracker, const FTagData* TagData);
 	void IncrTag(const FTagData* Tag, int64 Amount);
 
@@ -436,16 +438,18 @@ public:
 	void TrackAllocation(const void* Ptr, int64 Size, ELLMTag DefaultTag, ELLMTracker Tracker, ELLMAllocType AllocType,
 		bool bTrackInMemPro);
 	void TrackFree(const void* Ptr, ELLMTracker Tracker, ELLMAllocType AllocType, bool bTrackInMemPro);
+	void TrackMemoryOfActiveTag(int64 Amount, FName DefaultTag, ELLMTracker Tracker, ELLMAllocType AllocType);
+	void TrackMemoryOfActiveTag(int64 Amount, ELLMTag DefaultTag, ELLMTracker Tracker, ELLMAllocType AllocType);
 	void OnAllocMoved(const void* Dest, const void* Source, ELLMTracker Tracker, ELLMAllocType AllocType);
 
-	void TrackMemory(ELLMTag EnumTag, int64 Amount, ELLMAllocType AllocType);
-	void TrackMemory(FName TagName, int64 Amount, ELLMAllocType AllocType);
-	void TrackMemory(const FTagData* TagData, int64 Amount, ELLMAllocType AllocType);
+	void TrackMemory(ELLMTag EnumTag, ELLMTracker Tracker, int64 Amount, ELLMAllocType AllocType);
+	void TrackMemory(FName TagName, ELLMTracker Tracker, int64 Amount, ELLMAllocType AllocType);
+	void TrackMemory(const FTagData* TagData, ELLMTracker Tracker, int64 Amount, ELLMAllocType AllocType);
 
 	// This will pause/unpause tracking, and also manually increment a given tag
-	void PauseAndTrackMemory(FName TagName, bool bInIsStatTag, int64 Amount, ELLMAllocType AllocType);
-	void PauseAndTrackMemory(ELLMTag EnumTag, int64 Amount, ELLMAllocType AllocType);
-	void PauseAndTrackMemory(const FTagData* TagData, int64 Amount, ELLMAllocType AllocType);
+	void PauseAndTrackMemory(FName TagName, ELLMTracker Tracker, bool bInIsStatTag, int64 Amount, ELLMAllocType AllocType);
+	void PauseAndTrackMemory(ELLMTag EnumTag, ELLMTracker Tracker, int64 Amount, ELLMAllocType AllocType);
+	void PauseAndTrackMemory(const FTagData* TagData, ELLMTracker Tracker, int64 Amount, ELLMAllocType AllocType);
 	void Pause(ELLMAllocType AllocType);
 	void Unpause(ELLMAllocType AllocType);
 	bool IsPaused(ELLMAllocType AllocType);
@@ -541,6 +545,8 @@ protected:
 	FLLMThreadState* GetState();
 	void TrackAllocation(const void* Ptr, int64 Size, const FTagData* ActiveTagData, ELLMTracker Tracker,
 		ELLMAllocType AllocType, FLLMThreadState* State, bool bTrackInMemPro);
+	void TrackMemoryOfActiveTag(int64 Amount, const FTagData* TagData, ELLMTracker Tracker,
+		ELLMAllocType AllocType, FLLMThreadState* State);
 
 	FLowLevelMemTracker& LLMRef;
 
@@ -1104,8 +1110,10 @@ void FLowLevelMemTracker::InitialiseProgramSize()
 		FPlatformMemoryStats Stats = FPlatformMemory::GetStats();
 		ProgramSize = Stats.TotalPhysical - Stats.AvailablePhysical;
 
-		GetTracker(ELLMTracker::Platform)->TrackMemory(ELLMTag::ProgramSizePlatform, ProgramSize, ELLMAllocType::System);
-		GetTracker(ELLMTracker::Default)->TrackMemory(ELLMTag::ProgramSize, ProgramSize, ELLMAllocType::System);
+		GetTracker(ELLMTracker::Platform)->TrackMemory(ELLMTag::ProgramSizePlatform, ELLMTracker::Platform,
+			 ProgramSize, ELLMAllocType::System);
+		GetTracker(ELLMTracker::Default)->TrackMemory(ELLMTag::ProgramSize, ELLMTracker::Default,
+			 ProgramSize, ELLMAllocType::System);
 	}
 }
 
@@ -1121,8 +1129,10 @@ void FLowLevelMemTracker::SetProgramSize(uint64 InProgramSize)
 
 	ProgramSize = static_cast<int64>(InProgramSize);
 
-	GetTracker(ELLMTracker::Platform)->TrackMemory(ELLMTag::ProgramSizePlatform, ProgramSizeDiff, ELLMAllocType::System);
-	GetTracker(ELLMTracker::Default)->TrackMemory(ELLMTag::ProgramSize, ProgramSizeDiff, ELLMAllocType::System);
+	GetTracker(ELLMTracker::Platform)->TrackMemory(ELLMTag::ProgramSizePlatform, ELLMTracker::Platform,
+		ProgramSizeDiff, ELLMAllocType::System);
+	GetTracker(ELLMTracker::Default)->TrackMemory(ELLMTag::ProgramSize, ELLMTracker::Default,
+		ProgramSizeDiff, ELLMAllocType::System);
 }
 
 void FLowLevelMemTracker::ProcessCommandLine(const TCHAR* CmdLine)
@@ -1280,6 +1290,28 @@ void FLowLevelMemTracker::OnLowLevelFree(ELLMTracker Tracker, const void* Ptr,
 	}
 }
 
+void FLowLevelMemTracker::OnLowLevelChangeInMemoryUse(ELLMTracker Tracker, int64 DeltaMemory, ELLMTag DefaultTag,
+	ELLMAllocType AllocType)
+{
+	if (bIsDisabled)
+	{
+		return;
+	}
+	BootstrapInitialise();
+	GetTracker(Tracker)->TrackMemoryOfActiveTag(DeltaMemory, DefaultTag, Tracker, AllocType);
+}
+
+void FLowLevelMemTracker::OnLowLevelChangeInMemoryUse(ELLMTracker Tracker, int64 DeltaMemory, FName DefaultTag,
+	ELLMAllocType AllocType)
+{
+	if (bIsDisabled)
+	{
+		return;
+	}
+	BootstrapInitialise();
+	GetTracker(Tracker)->TrackMemoryOfActiveTag(DeltaMemory, DefaultTag, Tracker, AllocType);
+}
+
 void FLowLevelMemTracker::OnLowLevelAllocMoved(ELLMTracker Tracker, const void* Dest, const void* Source,
 	ELLMAllocType AllocType)
 {
@@ -2377,11 +2409,11 @@ void FLLMPauseScope::Init(FName TagName, ELLMTag EnumTag, bool bIsEnumTag, bool
 			{
 				if (bIsEnumTag)
 				{
-					LLMRef.GetTracker(Tracker)->PauseAndTrackMemory(EnumTag, static_cast<int64>(Amount), InAllocType);
+					LLMRef.GetTracker(Tracker)->PauseAndTrackMemory(EnumTag, Tracker, static_cast<int64>(Amount), InAllocType);
 				}
 				else
 				{
-					LLMRef.GetTracker(Tracker)->PauseAndTrackMemory(TagName, bIsStatTag, static_cast<int64>(Amount),
+					LLMRef.GetTracker(Tracker)->PauseAndTrackMemory(TagName, Tracker, bIsStatTag, static_cast<int64>(Amount),
 						InAllocType);
 				}
 			}
@@ -3282,14 +3314,14 @@ void FLLMTracker::PopAssetTag()
 }
 #endif
 
-void FLLMTracker::TrackAllocation(const void* Ptr, int64 Size, ELLMTag DefaultEnumTag, ELLMTracker Tracker,
+void FLLMTracker::TrackAllocation(const void* Ptr, int64 Size, ELLMTag DefaultTag, ELLMTracker Tracker,
 	ELLMAllocType AllocType, bool bTrackInMemPro)
 {
 	FLLMThreadState* State = GetOrCreateState();
 	const FTagData* TagData = State->GetTopTag();
 	if (!TagData)
 	{
-		TagData = LLMRef.FindOrAddTagData(DefaultEnumTag);
+		TagData = LLMRef.FindOrAddTagData(DefaultTag);
 	}
 	TrackAllocation(Ptr, Size, TagData, Tracker, AllocType, State, bTrackInMemPro);
 }
@@ -3417,46 +3449,92 @@ void FLLMTracker::OnAllocMoved(const void* Dest, const void* Source, ELLMTracker
 	State->TrackMoved(Dest, Source, Size, Tracker, TagData);
 }
 
-void FLLMTracker::TrackMemory(ELLMTag Tag, int64 Amount, ELLMAllocType AllocType)
+void FLLMTracker::TrackMemoryOfActiveTag(int64 Amount, FName DefaultTag, ELLMTracker Tracker, ELLMAllocType AllocType)
 {
-	TrackMemory(LLMRef.FindOrAddTagData(Tag), Amount, AllocType);
+	FLLMThreadState* State = GetOrCreateState();
+	const FTagData* TagData = State->GetTopTag();
+	if (!TagData)
+	{
+		TagData = LLMRef.FindOrAddTagData(DefaultTag);
+	}
+	TrackMemoryOfActiveTag(Amount, TagData, Tracker, AllocType, State);
 }
 
-void FLLMTracker::TrackMemory(FName Tag, int64 Amount, ELLMAllocType AllocType)
+void FLLMTracker::TrackMemoryOfActiveTag(int64 Amount, ELLMTag DefaultTag, ELLMTracker Tracker, ELLMAllocType AllocType)
 {
-	TrackMemory(LLMRef.FindOrAddTagData(Tag), Amount, AllocType);
+	FLLMThreadState* State = GetOrCreateState();
+	const FTagData* TagData = State->GetTopTag();
+	if (!TagData)
+	{
+		TagData = LLMRef.FindOrAddTagData(DefaultTag);
+	}
+	TrackMemoryOfActiveTag(Amount, TagData, Tracker, AllocType, State);
 }
 
-void FLLMTracker::TrackMemory(const FTagData* TagData, int64 Amount, ELLMAllocType AllocType)
+void FLLMTracker::TrackMemoryOfActiveTag(int64 Amount, const FTagData* TagData, ELLMTracker Tracker,
+	ELLMAllocType AllocType, FLLMThreadState* State)
+{
+	if (IsPaused(AllocType))
+	{
+		// When Paused, we do not track any delta memory; the code that triggered the pause is responsible for updating the delta memory
+		return;
+	}
+
+	// track the total quickly
+	FPlatformAtomics::InterlockedAdd(&TrackedTotal, Amount);
+
+#if !LLM_ENABLED_FULL_TAGS
+	// When full tags are disabled, we instead store the top-level enumtag parent of the tag used by each allocation
+	TagData = TagData->GetContainingEnumTagData();
+#endif
+
+#if LLM_ALLOW_ASSETS_TAGS
+	const FTagData* AssetTagData = State->GetTopAssetTag();
+#else
+	const FTagData* AssetTagData = nullptr;
+#endif
+
+	// track on the thread state
+	State->TrackMemory(Amount, Tracker, AllocType, TagData, AssetTagData);
+}
+
+void FLLMTracker::TrackMemory(ELLMTag Tag, ELLMTracker Tracker, int64 Amount, ELLMAllocType AllocType)
+{
+	TrackMemory(LLMRef.FindOrAddTagData(Tag), Tracker, Amount, AllocType);
+}
+
+void FLLMTracker::TrackMemory(FName Tag, ELLMTracker Tracker, int64 Amount, ELLMAllocType AllocType)
+{
+	TrackMemory(LLMRef.FindOrAddTagData(Tag), Tracker, Amount, AllocType);
+}
+
+void FLLMTracker::TrackMemory(const FTagData* TagData, ELLMTracker Tracker, int64 Amount, ELLMAllocType AllocType)
 {
 	FLLMThreadState* State = GetOrCreateState();
-	FScopeLock Lock(&State->TagSection);
-	State->IncrTag(TagData, Amount);
-	State->AllocTypeAmounts[static_cast<int32>(AllocType)] += Amount;
 	FPlatformAtomics::InterlockedAdd(&TrackedTotal, Amount);
+	State->TrackMemory(Amount, Tracker, AllocType, TagData, nullptr /* AssetTagData */);
 }
 
-void FLLMTracker::PauseAndTrackMemory(FName TagName, bool bInIsStatTag, int64 Amount, ELLMAllocType AllocType)
+void FLLMTracker::PauseAndTrackMemory(FName TagName, ELLMTracker Tracker, bool bInIsStatTag, int64 Amount, ELLMAllocType AllocType)
 {
 	const FTagData* TagData = LLMRef.FindOrAddTagData(TagName, bInIsStatTag);
-	PauseAndTrackMemory(TagData, Amount, AllocType);
+	PauseAndTrackMemory(TagData, Tracker, Amount, AllocType);
 }
 
-void FLLMTracker::PauseAndTrackMemory(ELLMTag EnumTag, int64 Amount, ELLMAllocType AllocType)
+void FLLMTracker::PauseAndTrackMemory(ELLMTag EnumTag, ELLMTracker Tracker, int64 Amount, ELLMAllocType AllocType)
 {
 	const FTagData* TagData = LLMRef.FindOrAddTagData(EnumTag);
-	PauseAndTrackMemory(TagData, Amount, AllocType);
+	PauseAndTrackMemory(TagData, Tracker, Amount, AllocType);
 }
 
 // This will pause/unpause tracking, and also manually increment a given tag.
-void FLLMTracker::PauseAndTrackMemory(const FTagData* TagData, int64 Amount, ELLMAllocType AllocType)
+void FLLMTracker::PauseAndTrackMemory(const FTagData* TagData, ELLMTracker Tracker, int64 Amount, ELLMAllocType AllocType)
 {
 	FLLMThreadState* State = GetOrCreateState();
+	FPlatformAtomics::InterlockedAdd(&TrackedTotal, Amount);
+	State->TrackMemory(Amount, Tracker, AllocType, TagData, nullptr /* AssetTagData */);
 	FScopeLock Lock(&State->TagSection);
 	State->PausedCounter[static_cast<int32>(AllocType)]++;
-	State->IncrTag(TagData, Amount);
-	State->AllocTypeAmounts[static_cast<int32>(AllocType)] += Amount;
-	FPlatformAtomics::InterlockedAdd(&TrackedTotal, Amount);
 }
 
 void FLLMTracker::Pause(ELLMAllocType AllocType)
@@ -3888,6 +3966,21 @@ void FLLMThreadState::TrackFree(const void* Ptr, int64 Size, ELLMTracker Tracker
 #endif
 }
 
+void FLLMThreadState::TrackMemory(int64 Amount, ELLMTracker Tracker, ELLMAllocType AllocType, const FTagData* TagData, const FTagData* AssetTagData)
+{
+	FScopeLock Lock(&TagSection);
+	AllocTypeAmounts[static_cast<int32>(AllocType)] += Amount;
+	IncrTag(TagData, Amount);
+#if LLM_ALLOW_ASSETS_TAGS
+	if (AssetTagData)
+	{
+		IncrTag(AssetTagData, Size);
+	}
+#endif
+
+	// TODO: Need to expose TrackMemory to Platform-specific trackers and FMemPropProfiler
+}
+
 void FLLMThreadState::TrackMoved(const void* Dest, const void* Source, int64 Size, ELLMTracker Tracker,
 	const FTagData* TagData)
 {
diff --git a/Engine/Source/Runtime/Core/Public/HAL/LowLevelMemTracker.h b/Engine/Source/Runtime/Core/Public/HAL/LowLevelMemTracker.h
--- a/Engine/Source/Runtime/Core/Public/HAL/LowLevelMemTracker.h
+++ b/Engine/Source/Runtime/Core/Public/HAL/LowLevelMemTracker.h
@@ -535,13 +535,13 @@ public:
 	/**
 	 * Records the use of memory that was allocated for a pointer.
 	 *
-	 * @param Tracker - Which tracker to use, the high-level default tracker for regular engine allocations, or the
-	 *                  low-level platform tracker for memory management systems.
-	 * @param Ptr - The pointer that was allocated.
-	 * @param Size - The size of the memory that was allocated for the pointer.
-	 * @param DefaultTag - The tag to use if there is no tag already in scope on the callstack.
-	 * @param AllocType - Type of allocation, FMalloc for regular allocations done through GMalloc, System for others.
-	 * @param bTrackInMemPro - Whether to pass the allocation/free information on to FMemProProfiler
+	 * @param Tracker Which tracker to use, the high-level default tracker for regular engine allocations, or the
+	 *        low-level platform tracker for memory management systems.
+	 * @param Ptr The pointer that was allocated.
+	 * @param Size The size of the memory that was allocated for the pointer.
+	 * @param DefaultTag The tag to use if there is no tag already in scope on the callstack.
+	 * @param AllocType Type of allocation, FMalloc for regular allocations done through GMalloc, System for others.
+	 * @param bTrackInMemPro Whether to pass the allocation/free information on to FMemProProfiler.
 	 */
 	void OnLowLevelAlloc(ELLMTracker Tracker, const void* Ptr, uint64 Size, ELLMTag DefaultTag = ELLMTag::Untagged,
 		ELLMAllocType AllocType = ELLMAllocType::None, bool bTrackInMemPro = true);
@@ -551,15 +551,27 @@ public:
 	/**
 	 * Records the release of memory that was allocated for a pointer.
 	 *
-	 * @param Tracker - Which tracker to use, the high-level default tracker for regular engine allocations, or the
-	 *                  low-level platform tracker for memory management systems.
-	 * @param Ptr - The pointer that was freed.
-	 * @param AllocType - Type of allocation, FMalloc for regular allocations done through GMalloc, System for others.
-	 * @param bTrackInMemPro - Whether to pass the allocation/free information on to FMemProProfiler
+	 * @param Tracker Which tracker to use, the high-level default tracker for regular engine allocations, or the
+	 *        low-level platform tracker for memory management systems.
+	 * @param Ptr The pointer that was freed.
+	 * @param AllocType Type of allocation, FMalloc for regular allocations done through GMalloc, System for others.
+	 * @param bTrackInMemPro Whether to pass the allocation/free information on to FMemProProfiler.
 	 */
 	void OnLowLevelFree(ELLMTracker Tracker, const void* Ptr,
 		ELLMAllocType AllocType = ELLMAllocType::None, bool bTrackInMemPro = true);
 
+	/**
+	 * Records an allocation or release of memory by a system that is not directly associated with a single pointer.
+	 *
+	 * @param Tracker Which tracker to use, the high-level default tracker for regular engine allocations, or the
+	 *        low-level platform tracker for memory management systems.
+	 * @param DeltaMemory The amount of memory that was allocated (positive) or freed (negative)
+	 * @param DefaultTag The tag to use if there is no tag already in scope on the callstack.
+	 * @param AllocType Type of allocation, FMalloc for regular allocations done through GMalloc, System for others.
+	 */
+	void OnLowLevelChangeInMemoryUse(ELLMTracker Tracker, int64 DeltaMemory, ELLMTag DefaultTag = ELLMTag::Untagged, ELLMAllocType AllocType = ELLMAllocType::None);
+	void OnLowLevelChangeInMemoryUse(ELLMTracker Tracker, int64 DeltaMemory, FName DefaultTag, ELLMAllocType AllocType = ELLMAllocType::None);
+
 	/** Call if an allocation is moved in memory, such as in a defragger. */
 	void OnLowLevelAllocMoved(ELLMTracker Tracker, const void* Dest, const void* Source,
 		ELLMAllocType AllocType = ELLMAllocType::None);
