Fixed leaks in LockFreeList implementation.

diff --git a/Engine/Source/Runtime/Core/Public/Containers/LockFreeList.h b/Engine/Source/Runtime/Core/Public/Containers/LockFreeList.h
--- a/Engine/Source/Runtime/Core/Public/Containers/LockFreeList.h
+++ b/Engine/Source/Runtime/Core/Public/Containers/LockFreeList.h
@@ -419,8 +419,14 @@ class FLockFreePointerListLIFOBase
 public:
 	FLockFreePointerListLIFOBase() = default;
 
+	~FLockFreePointerListLIFOBase()
+	{
+		while (Pop()) {};
+	}
+
 	void Reset()
 	{
+		while (Pop()) {};
 		RootList.Reset();
 	}
 
@@ -538,6 +544,12 @@ public:
 		Tail.SetPtr(Stub);
 	}
 
+	~FLockFreePointerFIFOBase()
+	{
+		while (Pop()) {};
+		FLockFreeLinkPolicy::FreeLockFreeLink(Head.GetPtr());
+	}
+
 	void Push(T* InPayload) TSAN_SAFE
 	{
 		TLinkPtr Item = FLockFreeLinkPolicy::AllocLockFreeLink();
