Fix build error

diff --git a/Engine/Source/Runtime/Core/Private/IO/IoDispatcher.cpp b/Engine/Source/Runtime/Core/Private/IO/IoDispatcher.cpp
--- a/Engine/Source/Runtime/Core/Private/IO/IoDispatcher.cpp
+++ b/Engine/Source/Runtime/Core/Private/IO/IoDispatcher.cpp
@@ -913,8 +913,14 @@ HasScriptObjectsChunk(FIoDispatcher& Dispatcher)
 static bool
 HasUseIoStoreParam()
 {
-	static bool bForceIoStore = UE_FORCE_USE_IOSTORE || (WITH_IOSTORE_IN_EDITOR && FParse::Param(FCommandLine::Get(), TEXT("UseIoStore")));
-	return bForceIoStore;
+#if UE_FORCE_USE_IOSTORE
+    return true;
+#elif WITH_IOSTORE_IN_EDITOR
+    static bool bForceIoStore = FParse::Param(FCommandLine::Get(), TEXT("UseIoStore"));
+    return bForceIoStore;
+#else
+    return false;
+#endif
 }
 
 bool
