Fix TSAN warning in TEventStack

#rnx

diff --git a/Engine/Source/Runtime/Core/Public/Async/Fundamental/ReserveScheduler.h b/Engine/Source/Runtime/Core/Public/Async/Fundamental/ReserveScheduler.h
--- a/Engine/Source/Runtime/Core/Public/Async/Fundamental/ReserveScheduler.h
+++ b/Engine/Source/Runtime/Core/Public/Async/Fundamental/ReserveScheduler.h
@@ -34,7 +34,7 @@ namespace LowLevelTasks
 		{
 			FEventRef		SleepEvent;
 			FConditional	CompletedDelegate = []() { return true; };
-			FYieldedWork*	Next = nullptr;
+			std::atomic<FYieldedWork*> Next {nullptr};
 			bool			bPermitBackgroundWork = false;
 		};
 	
diff --git a/Engine/Source/Runtime/Core/Public/Async/Fundamental/Scheduler.h b/Engine/Source/Runtime/Core/Public/Async/Fundamental/Scheduler.h
--- a/Engine/Source/Runtime/Core/Public/Async/Fundamental/Scheduler.h
+++ b/Engine/Source/Runtime/Core/Public/Async/Fundamental/Scheduler.h
@@ -57,10 +57,10 @@ namespace LowLevelTasks
 #if DO_CHECK
 				int64 LastRevision = int64(LocalTop.Revision); 
 #endif
-				NodeType* Item = reinterpret_cast<NodeType*>(LocalTop.Address << 3);																 //acquire on failure because we read Item->Next next itteration
-				if (Top.compare_exchange_weak(LocalTop, FTopNode { reinterpret_cast<uintptr_t>(Item->Next) >> 3, uintptr_t(LocalTop.Revision + 1) }, std::memory_order_relaxed, std::memory_order_acquire)) 
+				NodeType* Item = reinterpret_cast<NodeType*>(LocalTop.Address << 3);
+				if (Top.compare_exchange_weak(LocalTop, FTopNode { reinterpret_cast<uintptr_t>(Item->Next.load(std::memory_order_relaxed)) >> 3, uintptr_t(LocalTop.Revision + 1) }, std::memory_order_acq_rel))
 				{
-					Item->Next = nullptr;
+					Item->Next.store(nullptr, std::memory_order_relaxed);
 					return Item;
 				}
 #if DO_CHECK
@@ -77,7 +77,7 @@ namespace LowLevelTasks
 			checkSlow(reinterpret_cast<uintptr_t>(Item) < (1ull << 48));
 			checkSlow((reinterpret_cast<uintptr_t>(Item) & 0x7) == 0);
 #endif
-			checkSlow(Item->Next == nullptr);
+			checkSlow(Item->Next.load(std::memory_order_relaxed) == nullptr);
 
 			FTopNode LocalTop = Top.load(std::memory_order_relaxed);
 			while (true) 
@@ -85,8 +85,8 @@ namespace LowLevelTasks
 #if DO_CHECK
 				int64 LastRevision = int64(LocalTop.Revision); 
 #endif
-				Item->Next = reinterpret_cast<NodeType*>(LocalTop.Address << 3);
-				if (Top.compare_exchange_weak(LocalTop, FTopNode { reinterpret_cast<uintptr_t>(Item) >> 3, uintptr_t(LocalTop.Revision + 1) }, std::memory_order_release, std::memory_order_acquire))  
+				Item->Next.store(reinterpret_cast<NodeType*>(LocalTop.Address << 3), std::memory_order_relaxed);
+				if (Top.compare_exchange_weak(LocalTop, FTopNode { reinterpret_cast<uintptr_t>(Item) >> 3, uintptr_t(LocalTop.Revision + 1) }, std::memory_order_acq_rel))
 				{
 					return;
 				}
diff --git a/Engine/Source/Runtime/Core/Public/Async/Fundamental/Task.h b/Engine/Source/Runtime/Core/Public/Async/Fundamental/Task.h
--- a/Engine/Source/Runtime/Core/Public/Async/Fundamental/Task.h
+++ b/Engine/Source/Runtime/Core/Public/Async/Fundamental/Task.h
@@ -710,6 +710,6 @@ namespace LowLevelTasks
 	{
 		FEventRef SleepEvent;
 		std::atomic<ESleepState> SleepState { ESleepState::Running };
-		FSleepEvent* Next = nullptr;
+		std::atomic<FSleepEvent*> Next { nullptr };
 	};
 }
\ No newline at end of file
