Fix TSAN warning in LocalQueue.h

#rnx

diff --git a/Engine/Source/Runtime/Core/Public/Async/Fundamental/LocalQueue.h b/Engine/Source/Runtime/Core/Public/Async/Fundamental/LocalQueue.h
--- a/Engine/Source/Runtime/Core/Public/Async/Fundamental/LocalQueue.h
+++ b/Engine/Source/Runtime/Core/Public/Async/Fundamental/LocalQueue.h
@@ -156,8 +156,6 @@ public:
 		bool bCpuBeginEventEmitted = false;
 #endif
 
-		static uint32 WorkerLookingForWorkTraceId;
-
 		inline FOutOfWork(std::atomic_int& InNumWorkersLookingForWork) : NumWorkersLookingForWork(InNumWorkersLookingForWork)
 		{
 		}
@@ -175,10 +173,7 @@ public:
 #if CPUPROFILERTRACE_ENABLED
 				if (CpuChannel)
 				{
-					if (WorkerLookingForWorkTraceId == 0)
-					{
-						WorkerLookingForWorkTraceId = FCpuProfilerTrace::OutputEventType("TaskWorkerIsLookingForWork");
-					}
+					static uint32 WorkerLookingForWorkTraceId = FCpuProfilerTrace::OutputEventType("TaskWorkerIsLookingForWork");
 					FCpuProfilerTrace::OutputBeginEvent(WorkerLookingForWorkTraceId);
 					bCpuBeginEventEmitted = true;
 				}
@@ -544,6 +539,4 @@ private:
 	std::atomic_int NumActiveWorkers[2] = { {0}, {0} };
 };
 
-template<uint32 NumLocalItems>
-uint32 TLocalQueueRegistry<NumLocalItems>::FOutOfWork::WorkerLookingForWorkTraceId = 0;
 }
\ No newline at end of file
