[Core] fix for warning C6250 Calling 'VirtualFree' without the MEM_RELEASE flag may free memory but not address descriptors (VADs); results in address space leaks

added pragma to suppress the warning copied from MallocStomp

diff --git a/Engine/Source/Runtime/Core/Public/HAL/MallocStomp2.h b/Engine/Source/Runtime/Core/Public/HAL/MallocStomp2.h
--- a/Engine/Source/Runtime/Core/Public/HAL/MallocStomp2.h
+++ b/Engine/Source/Runtime/Core/Public/HAL/MallocStomp2.h
@@ -116,7 +116,16 @@ public:
 		
 		void Free(void* InPtr, SIZE_T InSize)
 		{
+#if PLATFORM_WINDOWS && USING_CODE_ANALYSIS
+	MSVC_PRAGMA(warning(push))
+	MSVC_PRAGMA(warning(disable : 6250)) // Suppress C6250, as virtual address space is "leaked" by design.
+#endif
+			// Unmap physical memory, but keep virtual address range reserved to catch use-after-free errors.
 			::VirtualFree(InPtr, InSize, MEM_DECOMMIT);
+
+#if PLATFORM_WINDOWS && USING_CODE_ANALYSIS
+	MSVC_PRAGMA(warning(pop))
+#endif
 		}
 	};
 
