Fixes to make code compliant with header units

diff --git a/Engine/Source/Runtime/Core/Public/Containers/StringConv.h b/Engine/Source/Runtime/Core/Public/Containers/StringConv.h
--- a/Engine/Source/Runtime/Core/Public/Containers/StringConv.h
+++ b/Engine/Source/Runtime/Core/Public/Containers/StringConv.h
@@ -54,7 +54,7 @@ public:
 namespace StringConv
 {
 	/** Is the provided Codepoint within the range of valid codepoints? */
-	static FORCEINLINE bool IsValidCodepoint(const uint32 Codepoint)
+	FORCEINLINE bool IsValidCodepoint(const uint32 Codepoint)
 	{
 		if ((Codepoint > 0x10FFFF) ||						// No Unicode codepoints above 10FFFFh, (for now!)
 			(Codepoint == 0xFFFE) || (Codepoint == 0xFFFF)) // illegal values.
@@ -65,23 +65,23 @@ namespace StringConv
 	}
 
 	/** Is the provided Codepoint within the range of the high-surrogates? */
-	static FORCEINLINE bool IsHighSurrogate(const uint32 Codepoint)
+	FORCEINLINE bool IsHighSurrogate(const uint32 Codepoint)
 	{
 		return Codepoint >= HIGH_SURROGATE_START_CODEPOINT && Codepoint <= HIGH_SURROGATE_END_CODEPOINT;
 	}
 
 	/** Is the provided Codepoint within the range of the low-surrogates? */
-	static FORCEINLINE bool IsLowSurrogate(const uint32 Codepoint)
+	FORCEINLINE bool IsLowSurrogate(const uint32 Codepoint)
 	{
 		return Codepoint >= LOW_SURROGATE_START_CODEPOINT && Codepoint <= LOW_SURROGATE_END_CODEPOINT;
 	}
 
-	static FORCEINLINE uint32 EncodeSurrogate(const uint16 HighSurrogate, const uint16 LowSurrogate)
+	FORCEINLINE uint32 EncodeSurrogate(const uint16 HighSurrogate, const uint16 LowSurrogate)
 	{
 		return ((HighSurrogate - HIGH_SURROGATE_START_CODEPOINT) << 10) + (LowSurrogate - LOW_SURROGATE_START_CODEPOINT) + 0x10000;
 	}
 
-	static FORCEINLINE void DecodeSurrogate(const uint32 Codepoint, uint16& OutHighSurrogate, uint16& OutLowSurrogate)
+	FORCEINLINE void DecodeSurrogate(const uint32 Codepoint, uint16& OutHighSurrogate, uint16& OutLowSurrogate)
 	{
 		const uint32 TmpCodepoint = Codepoint - 0x10000;
 		OutHighSurrogate = (uint16)((TmpCodepoint >> 10) + HIGH_SURROGATE_START_CODEPOINT);
@@ -89,14 +89,14 @@ namespace StringConv
 	}
 
 	/** Is the provided Codepoint outside of the range of the basic multilingual plane, but within the valid range of UTF8/16? */
-	static FORCEINLINE bool IsEncodedSurrogate(const uint32 Codepoint)
+	FORCEINLINE bool IsEncodedSurrogate(const uint32 Codepoint)
 	{
 		return Codepoint >= ENCODED_SURROGATE_START_CODEPOINT && Codepoint <= ENCODED_SURROGATE_END_CODEPOINT;
 	}
 
 	/** Inline combine any UTF-16 surrogate pairs in the given null-terminated character buffer, returning the updated length */
 	template<typename CharType>
-	static FORCEINLINE int32 InlineCombineSurrogates_Buffer(CharType* StrBuffer, int32 StrLen)
+	FORCEINLINE int32 InlineCombineSurrogates_Buffer(CharType* StrBuffer, int32 StrLen)
 	{
 		static_assert(sizeof(CharType) == 4, "CharType must be 4-bytes!");
 
@@ -155,7 +155,7 @@ namespace StringConv
 
 	/** Inline combine any UTF-16 surrogate pairs in the given null-terminated TCHAR array */
 	template<typename AllocatorType>
-	static FORCEINLINE void InlineCombineSurrogates_Array(TArray<TCHAR, AllocatorType>& StrBuffer)
+	FORCEINLINE void InlineCombineSurrogates_Array(TArray<TCHAR, AllocatorType>& StrBuffer)
 	{
 #if PLATFORM_TCHAR_IS_4_BYTES
 		const int32 NewStrLen = InlineCombineSurrogates_Buffer(StrBuffer.GetData(), StrBuffer.Num() - 1);
diff --git a/Engine/Source/Runtime/Core/Public/Math/UnrealPlatformMathSSE.h b/Engine/Source/Runtime/Core/Public/Math/UnrealPlatformMathSSE.h
--- a/Engine/Source/Runtime/Core/Public/Math/UnrealPlatformMathSSE.h
+++ b/Engine/Source/Runtime/Core/Public/Math/UnrealPlatformMathSSE.h
@@ -19,7 +19,7 @@ namespace UE4
 {
 namespace SSE
 {
-	static FORCEINLINE float InvSqrt(float InValue)
+	FORCEINLINE float InvSqrt(float InValue)
 	{
 		const __m128 One = _mm_set_ss(1.0f);
 		const __m128 Y0 = _mm_set_ss(InValue);
@@ -30,7 +30,7 @@ namespace SSE
 		return temp;
 	}
 
-	static FORCEINLINE double InvSqrt(double InValue)
+	FORCEINLINE double InvSqrt(double InValue)
 	{
 		const __m128d One = _mm_set_sd(1.0);
 		const __m128d Y0 = _mm_set_sd(InValue);
@@ -41,7 +41,7 @@ namespace SSE
 		return temp;
 	}
 
-	static FORCEINLINE float InvSqrtEst(float F)
+	FORCEINLINE float InvSqrtEst(float F)
 	{
 		// Performs one pass of Newton-Raphson iteration on the hardware estimate
 		const __m128 fOneHalf = _mm_set_ss(0.5f);
@@ -61,44 +61,44 @@ namespace SSE
 		return temp;
 	}
 
-	static FORCEINLINE double InvSqrtEst(double InValue)
+	FORCEINLINE double InvSqrtEst(double InValue)
 	{		
 		return InvSqrt(InValue);
 	}
 
-	static FORCEINLINE int32 TruncToInt32(float F)
+	FORCEINLINE int32 TruncToInt32(float F)
 	{
 		return _mm_cvtt_ss2si(_mm_set_ss(F));
 	}
 
-	static FORCEINLINE int32 TruncToInt32(double InValue)
+	FORCEINLINE int32 TruncToInt32(double InValue)
 	{
 		return _mm_cvttsd_si32(_mm_set_sd(InValue));
 	}
 
-	static FORCEINLINE int64 TruncToInt64(double InValue)
+	FORCEINLINE int64 TruncToInt64(double InValue)
 	{
 		return _mm_cvttsd_si64(_mm_set_sd(InValue));
 	}
 
-	static FORCEINLINE int32 FloorToInt32(float F)
+	FORCEINLINE int32 FloorToInt32(float F)
 	{
 		// Note: unlike the Generic solution and the SSE4 float solution, we implement FloorToInt using a rounding instruction, rather than implementing RoundToInt using a floor instruction.  
 		// We therefore need to do the same times-2 transform (with a slighly different formula) that RoundToInt does; see the note on RoundToInt
 		return _mm_cvt_ss2si(_mm_set_ss(F + F - 0.5f)) >> 1;
 	}
 
-	static FORCEINLINE int32 FloorToInt32(double InValue)
+	FORCEINLINE int32 FloorToInt32(double InValue)
 	{
 		return _mm_cvtsd_si32(_mm_set_sd(InValue + InValue - 0.5)) >> 1;
 	}
 
-	static FORCEINLINE int64 FloorToInt64(double InValue)
+	FORCEINLINE int64 FloorToInt64(double InValue)
 	{
 		return _mm_cvtsd_si64(_mm_set_sd(InValue + InValue - 0.5)) >> 1;
 	}
 
-	static FORCEINLINE int32 RoundToInt32(float F)
+	FORCEINLINE int32 RoundToInt32(float F)
 	{
 		// Note: the times-2 is to remove the rounding-to-nearest-even-number behavior that mm_cvt_ss2si uses when the fraction is .5
 		// The formula we uses causes the round instruction to always be applied to a an odd integer when the original value was 0.5, and eliminates the rounding-to-nearest-even-number behavior
@@ -108,36 +108,36 @@ namespace SSE
 		return _mm_cvt_ss2si(_mm_set_ss(F + F + 0.5f)) >> 1;
 	}
 
-	static FORCEINLINE int32 RoundToInt32(double InValue)
+	FORCEINLINE int32 RoundToInt32(double InValue)
 	{
 		return _mm_cvtsd_si32(_mm_set_sd(InValue + InValue + 0.5)) >> 1;
 	}
 
-	static FORCEINLINE int64 RoundToInt64(double InValue)
+	FORCEINLINE int64 RoundToInt64(double InValue)
 	{
 		return _mm_cvtsd_si64(_mm_set_sd(InValue + InValue + 0.5)) >> 1;
 	}
 
-	static FORCEINLINE int32 CeilToInt32(float F)
+	FORCEINLINE int32 CeilToInt32(float F)
 	{
 		// Note: unlike the Generic solution and the SSE4 float solution, we implement CeilToInt using a rounding instruction, rather than a dedicated ceil instruction
 		// We therefore need to do the same times-2 transform (with a slighly different formula) that RoundToInt does; see the note on RoundToInt
 		return -(_mm_cvt_ss2si(_mm_set_ss(-0.5f - (F + F))) >> 1);
 	}
 
-	static FORCEINLINE int32 CeilToInt32(double InValue)
+	FORCEINLINE int32 CeilToInt32(double InValue)
 	{
 		return -(_mm_cvtsd_si32(_mm_set_sd(-0.5 - (InValue + InValue))) >> 1);
 	}
 
-	static FORCEINLINE int64 CeilToInt64(double InValue)
+	FORCEINLINE int64 CeilToInt64(double InValue)
 	{
 		return -(_mm_cvtsd_si64(_mm_set_sd(-0.5 - (InValue + InValue))) >> 1);
 	}
 
 	// https://gist.github.com/rygorous/2156668
 	// float_to_half_rtne_SSE2
-	static __m128i FloatToHalf(__m128 f)
+	inline __m128i FloatToHalf(__m128 f)
 	{
 		const __m128 mask_sign		= _mm_set1_ps(-0.0f);
 		const __m128i c_f16max			= _mm_set1_epi32((127 + 16) << 23); // all FP32 values >=this round to +inf
@@ -192,7 +192,7 @@ namespace SSE
 	}
 	
 	// four halfs should be in the u64 part of input
-	static __m128 HalfToFloat(__m128i four_halfs_u64)
+	inline __m128 HalfToFloat(__m128i four_halfs_u64)
 	{
 		__m128i rgba_half32 = _mm_unpacklo_epi16(four_halfs_u64, _mm_setzero_si128());
 
diff --git a/Engine/Source/Runtime/Core/Public/Traits/IsCharEncodingCompatibleWith.h b/Engine/Source/Runtime/Core/Public/Traits/IsCharEncodingCompatibleWith.h
--- a/Engine/Source/Runtime/Core/Public/Traits/IsCharEncodingCompatibleWith.h
+++ b/Engine/Source/Runtime/Core/Public/Traits/IsCharEncodingCompatibleWith.h
@@ -12,7 +12,7 @@
 namespace UE::Core::Private
 {
 	template <typename SrcEncoding, typename DestEncoding>
-	static constexpr bool IsCharEncodingCompatibleWithImpl()
+	constexpr bool IsCharEncodingCompatibleWithImpl()
 	{
 		if constexpr (std::is_same_v<SrcEncoding, DestEncoding>)
 		{
diff --git a/Engine/Source/Runtime/Core/Public/Traits/IsCharEncodingSimplyConvertibleTo.h b/Engine/Source/Runtime/Core/Public/Traits/IsCharEncodingSimplyConvertibleTo.h
--- a/Engine/Source/Runtime/Core/Public/Traits/IsCharEncodingSimplyConvertibleTo.h
+++ b/Engine/Source/Runtime/Core/Public/Traits/IsCharEncodingSimplyConvertibleTo.h
@@ -10,7 +10,7 @@
 namespace UE::Core::Private
 {
 	template <typename SrcEncoding, typename DestEncoding>
-	static constexpr bool IsCharEncodingSimplyConvertibleToImpl()
+	constexpr bool IsCharEncodingSimplyConvertibleToImpl()
 	{
 		if constexpr (TIsCharEncodingCompatibleWith_V<SrcEncoding, DestEncoding>)
 		{
diff --git a/Engine/Source/Runtime/Core/Public/Traits/IsFixedWidthCharEncoding.h b/Engine/Source/Runtime/Core/Public/Traits/IsFixedWidthCharEncoding.h
--- a/Engine/Source/Runtime/Core/Public/Traits/IsFixedWidthCharEncoding.h
+++ b/Engine/Source/Runtime/Core/Public/Traits/IsFixedWidthCharEncoding.h
@@ -15,7 +15,7 @@ namespace UE::Core::Private
 	 * Tests whether an encoding has fixed-width characters
 	 */
 	template <typename Encoding>
-	static constexpr bool IsFixedWidthEncodingImpl()
+	constexpr bool IsFixedWidthEncodingImpl()
 	{
 		return
 			std::is_same_v<Encoding, ANSICHAR> ||
diff --git a/Engine/Source/Runtime/Engine/Classes/Engine/NetConnection.h b/Engine/Source/Runtime/Engine/Classes/Engine/NetConnection.h
--- a/Engine/Source/Runtime/Engine/Classes/Engine/NetConnection.h
+++ b/Engine/Source/Runtime/Engine/Classes/Engine/NetConnection.h
@@ -65,7 +65,7 @@ typedef TMap<TWeakObjectPtr<AActor>, UActorChannel*, FDefaultSetAllocator, TWeak
 namespace NetConnectionHelper
 {
 	/** Number of bits to use in the packet header for sending the milliseconds on the clock when the packet is sent */
-	constexpr int32 NumBitsForJitterClockTimeInHeader = 10;
+	inline constexpr int32 NumBitsForJitterClockTimeInHeader = 10;
 }
 
 extern ENGINE_API TAutoConsoleVariable<int32> CVarNetEnableCongestionControl;
diff --git a/Engine/Source/Runtime/RHI/Private/RHICommandList.cpp b/Engine/Source/Runtime/RHI/Private/RHICommandList.cpp
--- a/Engine/Source/Runtime/RHI/Private/RHICommandList.cpp
+++ b/Engine/Source/Runtime/RHI/Private/RHICommandList.cpp
@@ -1913,4 +1913,9 @@ RHI_API void RHISetComputeShaderBackwardsCompatible(IRHIComputeContext* InContex
 {
 	TRefCountPtr<FRHIComputePipelineState> ComputePipelineState = RHICreateComputePipelineState(InShader);
 	InContext->RHISetComputePipelineState(ComputePipelineState);
-}
\ No newline at end of file
+}
+
+void IRHIComputeContext::RHISetComputeShader(FRHIComputeShader* ComputeShader)
+{
+	RHISetComputeShaderBackwardsCompatible(this, ComputeShader);
+}
diff --git a/Engine/Source/Runtime/RHI/Public/RHIContext.h b/Engine/Source/Runtime/RHI/Public/RHIContext.h
--- a/Engine/Source/Runtime/RHI/Public/RHIContext.h
+++ b/Engine/Source/Runtime/RHI/Public/RHIContext.h
@@ -227,11 +227,7 @@ public:
 	}
 
 	UE_DEPRECATED(5.1, "ComputePipelineStates should be used instead of direct ComputeShaders.")
-	virtual void RHISetComputeShader(FRHIComputeShader* ComputeShader)
-	{
-		RHI_API void RHISetComputeShaderBackwardsCompatible(IRHIComputeContext*, FRHIComputeShader*);
-		RHISetComputeShaderBackwardsCompatible(this, ComputeShader);
-	}
+	RHI_API virtual void RHISetComputeShader(FRHIComputeShader* ComputeShader);
 
 	virtual void RHISetComputePipelineState(FRHIComputePipelineState* ComputePipelineState) = 0;
 
