Improved comments for unusual backwards-searching behavior in FString::Find.

diff --git a/Engine/Source/Runtime/Core/Public/Containers/UnrealString.h b/Engine/Source/Runtime/Core/Public/Containers/UnrealString.h
--- a/Engine/Source/Runtime/Core/Public/Containers/UnrealString.h
+++ b/Engine/Source/Runtime/Core/Public/Containers/UnrealString.h
@@ -1021,9 +1021,16 @@ public:
 	 * of the first found instance. Can search from beginning or end, and ignore case or not.
 	 *
 	 * @param SubStr			The string array of TCHAR to search for
-	 * @param StartPosition		The start character position to search from
+	 * @param StartPosition		The start character position to search from.  See note below.
 	 * @param SearchCase		Indicates whether the search is case sensitive or not
 	 * @param SearchDir			Indicates whether the search starts at the beginning or at the end.
+	 *
+	 * @note  When searching backwards where a StartPosition is provided, searching will actually begin from
+	 *        StartPosition - SubStr.Len(), therefore:
+	 *
+	 *        FString("X").Find("X", ESearchCase::CaseSensitive, ESearchDir::FromEnd, 0) == INDEX_NONE
+	 *
+	 *        Consider using UE::String::FindLast() as an alternative.
 	 */
 	UE_NODISCARD int32 Find( const TCHAR* SubStr, ESearchCase::Type SearchCase = ESearchCase::IgnoreCase, 
 				ESearchDir::Type SearchDir = ESearchDir::FromStart, int32 StartPosition=INDEX_NONE ) const;
@@ -1033,9 +1040,16 @@ public:
 	 * of the first found instance. Can search from beginning or end, and ignore case or not.
 	 *
 	 * @param SubStr			The string to search for
-	 * @param StartPosition		The start character position to search from
+	 * @param StartPosition		The start character position to search from  See note below.
 	 * @param SearchCase		Indicates whether the search is case sensitive or not ( defaults to ESearchCase::IgnoreCase )
 	 * @param SearchDir			Indicates whether the search starts at the beginning or at the end ( defaults to ESearchDir::FromStart )
+	 *
+	 * @note  When searching backwards where a StartPosition is provided, searching will actually begin from
+	 *        StartPosition - SubStr.Len(), therefore:
+	 *
+	 *        FString("X").Find("X", ESearchCase::CaseSensitive, ESearchDir::FromEnd, 0) == INDEX_NONE
+	 *
+	 *        Consider using UE::String::FindLast() as an alternative.
 	 */
 	UE_NODISCARD FORCEINLINE int32 Find( const FString& SubStr, ESearchCase::Type SearchCase = ESearchCase::IgnoreCase, 
 							ESearchDir::Type SearchDir = ESearchDir::FromStart, int32 StartPosition=INDEX_NONE ) const
