[Core]
* Moved IoChunkId to its own file

diff --git a/Engine/Source/Runtime/Core/Private/IO/IoStore.cpp b/Engine/Source/Runtime/Core/Private/IO/IoStore.cpp
--- a/Engine/Source/Runtime/Core/Private/IO/IoStore.cpp
+++ b/Engine/Source/Runtime/Core/Private/IO/IoStore.cpp
@@ -67,6 +67,12 @@ FString LexToString(const EIoChunkType Type)
 	}
 }
 
+FArchive& operator<<(FArchive& Ar, FIoChunkId& ChunkId)
+{
+	Ar.Serialize(&ChunkId.Id, sizeof FIoChunkId::Id);
+	return Ar;
+}
+
 FString LexToString(const FIoChunkId& Id)
 {
 	FString Output;
diff --git a/Engine/Source/Runtime/Core/Public/IO/IoChunkId.h b/Engine/Source/Runtime/Core/Public/IO/IoChunkId.h
new file mode 100644
--- /dev/null
+++ b/Engine/Source/Runtime/Core/Public/IO/IoChunkId.h
@@ -0,0 +1,89 @@
+// Copyright Epic Games, Inc. All Rights Reserved.
+
+#pragma once
+
+#include "Memory/MemoryView.h"
+#include "Misc/StringBuilder.h"
+#include "String/BytesToHex.h"
+
+class FArchive;
+enum class EIoChunkType : uint8;
+
+/**
+ * Identifier to a chunk of data.
+ */
+class FIoChunkId
+{
+public:
+	CORE_API static const FIoChunkId InvalidChunkId;
+
+	friend uint32 GetTypeHash(FIoChunkId InId)
+	{
+		uint32 Hash = 5381;
+		for (int i = 0; i < sizeof Id; ++i)
+		{
+			Hash = Hash * 33 + InId.Id[i];
+		}
+		return Hash;
+	}
+
+	friend CORE_API FArchive& operator<<(FArchive& Ar, FIoChunkId& ChunkId);
+
+	template <typename CharType>
+	friend TStringBuilderBase<CharType>& operator<<(TStringBuilderBase<CharType>& Builder, const FIoChunkId& ChunkId)
+	{
+		UE::String::BytesToHexLower(ChunkId.Id, Builder);
+		return Builder;
+	}
+
+	friend CORE_API FString LexToString(const FIoChunkId& Id);
+
+	inline bool operator ==(const FIoChunkId& Rhs) const
+	{
+		return 0 == FMemory::Memcmp(Id, Rhs.Id, sizeof Id);
+	}
+
+	inline bool operator !=(const FIoChunkId& Rhs) const
+	{
+		return !(*this == Rhs);
+	}
+
+	void Set(const void* InIdPtr, SIZE_T InSize)
+	{
+		check(InSize == sizeof Id);
+		FMemory::Memcpy(Id, InIdPtr, sizeof Id);
+	}
+
+	void Set(FMemoryView InView)
+	{
+		check(InView.GetSize() == sizeof Id);
+		FMemory::Memcpy(Id, InView.GetData(), sizeof Id);
+	}
+
+	inline bool IsValid() const
+	{
+		return *this != InvalidChunkId;
+	}
+
+	inline const uint8* GetData() const { return Id; }
+	inline uint32		GetSize() const { return sizeof Id; }
+
+	EIoChunkType GetChunkType() const
+	{
+		return static_cast<EIoChunkType>(Id[11]);
+	}
+
+	friend class FIoStoreReaderImpl;
+
+private:
+	static inline FIoChunkId CreateEmptyId()
+	{
+		FIoChunkId ChunkId;
+		uint8 Data[12] = { 0 };
+		ChunkId.Set(Data, sizeof Data);
+
+		return ChunkId;
+	}
+
+	uint8	Id[12];
+};
diff --git a/Engine/Source/Runtime/Core/Public/IO/IoDispatcher.h b/Engine/Source/Runtime/Core/Public/IO/IoDispatcher.h
--- a/Engine/Source/Runtime/Core/Public/IO/IoDispatcher.h
+++ b/Engine/Source/Runtime/Core/Public/IO/IoDispatcher.h
@@ -17,6 +17,7 @@
 #include "HAL/PlatformAtomics.h"
 #include "HAL/PlatformFile.h"
 #include "HAL/UnrealMemory.h"
+#include "IO/IoChunkId.h"
 #include "IO/IoContainerId.h"
 #include "IO/IoHash.h"
 #include "Logging/LogMacros.h"
@@ -687,89 +688,6 @@ enum class EIoChunkType : uint8
 
 CORE_API FString LexToString(const EIoChunkType Type);
 
-/**
- * Identifier to a chunk of data.
- */
-class FIoChunkId
-{
-public:
-	CORE_API static const FIoChunkId InvalidChunkId;
-
-	friend uint32 GetTypeHash(FIoChunkId InId)
-	{
-		uint32 Hash = 5381;
-		for (int i = 0; i < sizeof Id; ++i)
-		{
-			Hash = Hash * 33 + InId.Id[i];
-		}
-		return Hash;
-	}
-
-	friend FArchive& operator<<(FArchive& Ar, FIoChunkId& ChunkId)
-	{
-		Ar.Serialize(&ChunkId.Id, sizeof Id);
-		return Ar;
-	}
-
-	template <typename CharType>
-	friend TStringBuilderBase<CharType>& operator<<(TStringBuilderBase<CharType>& Builder, const FIoChunkId& ChunkId)
-	{
-		UE::String::BytesToHexLower(ChunkId.Id, Builder);
-		return Builder;
-	}
-
-	friend CORE_API FString LexToString(const FIoChunkId& Id);
-
-	inline bool operator ==(const FIoChunkId& Rhs) const
-	{
-		return 0 == FMemory::Memcmp(Id, Rhs.Id, sizeof Id);
-	}
-
-	inline bool operator !=(const FIoChunkId& Rhs) const
-	{
-		return !(*this == Rhs);
-	}
-
-	void Set(const void* InIdPtr, SIZE_T InSize)
-	{
-		check(InSize == sizeof Id);
-		FMemory::Memcpy(Id, InIdPtr, sizeof Id);
-	}
-
-	void Set(FMemoryView InView)
-	{
-		check(InView.GetSize() == sizeof Id);
-		FMemory::Memcpy(Id, InView.GetData(), sizeof Id);
-	}
-
-	inline bool IsValid() const
-	{
-		return *this != InvalidChunkId;
-	}
-
-	inline const uint8* GetData() const { return Id; }
-	inline uint32		GetSize() const { return sizeof Id; }
-
-	EIoChunkType GetChunkType() const
-	{
-		return static_cast<EIoChunkType>(Id[11]);
-	}
-
-	friend class FIoStoreReaderImpl;
-	
-private:
-	static inline FIoChunkId CreateEmptyId()
-	{
-		FIoChunkId ChunkId;
-		uint8 Data[12] = { 0 };
-		ChunkId.Set(Data, sizeof Data);
-
-		return ChunkId;
-	}
-
-	uint8	Id[12];
-};
-
 /**
  * Creates a chunk identifier (generic -- prefer specialized versions where possible)
  */
