Moved operator== to be hidden friend instead of put directly in global namespace
Moved GetTypeHash function to be hidden friend instead of put directly in global namespace.

Note that the function/operator needs to be fully inlined in the type or placed in the cpp. If the function is added as friend but then implemented outside the type then hidden friend optimization won't work.

This should improve compile time somewhat according to msvc devs.

diff --git a/Engine/Plugins/Enterprise/DataprepEditor/Source/DataprepEditor/Private/DataprepSnapshot.cpp b/Engine/Plugins/Enterprise/DataprepEditor/Source/DataprepEditor/Private/DataprepSnapshot.cpp
--- a/Engine/Plugins/Enterprise/DataprepEditor/Source/DataprepEditor/Private/DataprepSnapshot.cpp
+++ b/Engine/Plugins/Enterprise/DataprepEditor/Source/DataprepEditor/Private/DataprepSnapshot.cpp
@@ -211,7 +211,7 @@ namespace DataprepSnapshotUtil
 	{
 		static FString FileNamePrefix( TEXT("stream_") );
 
-		FString PackageFileName = FileNamePrefix + FString::Printf( TEXT("%08x"), ::GetTypeHash( AssetPath ) );
+		FString PackageFileName = FileNamePrefix + FString::Printf( TEXT("%08x"), GetTypeHash( AssetPath ) );
 		return FPaths::ConvertRelativePathToFull( FPaths::Combine( RootPath, PackageFileName ) + SnapshotExtension );
 	}
 
diff --git a/Engine/Plugins/Enterprise/DatasmithCADImporter/Source/CADTools/Private/CADData.cpp b/Engine/Plugins/Enterprise/DatasmithCADImporter/Source/CADTools/Private/CADData.cpp
--- a/Engine/Plugins/Enterprise/DatasmithCADImporter/Source/CADTools/Private/CADData.cpp
+++ b/Engine/Plugins/Enterprise/DatasmithCADImporter/Source/CADTools/Private/CADData.cpp
@@ -24,7 +24,7 @@ FMaterialUId BuildColorFastUId(uint32 ColorId, uint8 Alpha)
 
 FMaterialUId BuildColorUId(const FColor& Color)
 {
-	return FMath::Abs((int32) ::GetTypeHash(Color));
+	return FMath::Abs((int32) GetTypeHash(Color));
 }
 
 FMaterialUId BuildMaterialUId(const FCADMaterial& Material)
diff --git a/Engine/Plugins/FX/Niagara/Source/Niagara/Public/NiagaraTypes.h b/Engine/Plugins/FX/Niagara/Source/Niagara/Public/NiagaraTypes.h
--- a/Engine/Plugins/FX/Niagara/Source/Niagara/Public/NiagaraTypes.h
+++ b/Engine/Plugins/FX/Niagara/Source/Niagara/Public/NiagaraTypes.h
@@ -207,6 +207,11 @@ struct NIAGARA_API FNiagaraAssetVersion
 	bool operator<=(const FNiagaraAssetVersion& Other) const { return *this < Other || (MajorVersion == Other.MajorVersion && MinorVersion == Other.MinorVersion); }
 
 	static FGuid CreateStableVersionGuid(UObject* Object);
+
+	friend FORCEINLINE uint32 GetTypeHash(const FNiagaraAssetVersion& Version)
+	{
+		return HashCombine(GetTypeHash(Version.MajorVersion), GetTypeHash(Version.MinorVersion));
+	}
 };
 
 struct NIAGARA_API FNiagaraLWCConverter
@@ -279,11 +284,6 @@ private:
 	TArray<FNiagaraStructConversionStep> ConversionSteps;
 };
 
-FORCEINLINE uint32 GetTypeHash(const FNiagaraAssetVersion& Version)
-{
-	return HashCombine(GetTypeHash(Version.MajorVersion), GetTypeHash(Version.MinorVersion));
-}
-
 /** Data controlling the spawning of particles */
 USTRUCT(BlueprintType, meta = (DisplayName = "Spawn Info", NiagaraClearEachFrame = "true"))
 struct FNiagaraSpawnInfo
@@ -333,6 +333,11 @@ struct FNiagaraID
 	bool operator==(const FNiagaraID& Other)const { return Index == Other.Index && AcquireTag == Other.AcquireTag; }
 	bool operator!=(const FNiagaraID& Other)const { return !(*this == Other); }
 	bool operator<(const FNiagaraID& Other)const { return Index < Other.Index || (Index == Other.Index && AcquireTag < Other.AcquireTag); }
+
+	friend FORCEINLINE uint32 GetTypeHash(const FNiagaraID& ID)
+	{
+		return HashCombine(GetTypeHash(ID.Index), GetTypeHash(ID.AcquireTag));
+	}
 };
 
 USTRUCT()
@@ -352,11 +357,6 @@ struct FNiagaraRandInfo
 
 #define NIAGARA_INVALID_ID (FNiagaraID({(INDEX_NONE), (INDEX_NONE)}))
 
-FORCEINLINE uint32 GetTypeHash(const FNiagaraID& ID)
-{
-	return HashCombine(GetTypeHash(ID.Index), GetTypeHash(ID.AcquireTag));
-}
-
 enum class ENiagaraStructConversion : uint8
 {
 	/** Do not modify struct data members, even if they are not compatible with the Niagara VM, as the struct is user facing (or from an external api). */
@@ -1206,6 +1206,11 @@ private:
 	static TSet<UStruct*> BoolStructs;
 
 	static FNiagaraTypeDefinition CollisionEventDef;
+
+	friend FORCEINLINE uint32 GetTypeHash(const FNiagaraTypeDefinition& Type)
+	{
+		return HashCombine(HashCombine(GetTypeHash(Type.ClassStructOrEnum), GetTypeHash(Type.UnderlyingType)), GetTypeHash(Type.GetFlags()));
+	}
 };
 
 template<>
@@ -1236,11 +1241,6 @@ const FNiagaraTypeDefinition& FNiagaraTypeDefinition::Get()
 	if (TIsSame<T, FNiagaraRandInfo>::Value) { return GetRandInfoDef(); }
 }
 
-FORCEINLINE uint32 GetTypeHash(const FNiagaraTypeDefinition& Type)
-{
-	return HashCombine(HashCombine(GetTypeHash(Type.ClassStructOrEnum), GetTypeHash(Type.UnderlyingType)), GetTypeHash(Type.GetFlags()));
-}
-
 //////////////////////////////////////////////////////////////////////////
 
 enum class ENiagaraTypeRegistryFlags : uint32
@@ -1704,6 +1704,11 @@ protected:
 	UPROPERTY(meta = (DeprecatedProperty))
 	FNiagaraTypeDefinition TypeDef_DEPRECATED;
 #endif
+
+	friend FORCEINLINE uint32 GetTypeHash(const FNiagaraVariableBase& Var)
+	{
+		return HashCombine(GetTypeHash(Var.GetType()), GetTypeHash(Var.GetName()));
+	}
 };
 
 template<>
@@ -1718,11 +1723,6 @@ struct TStructOpsTypeTraits<FNiagaraVariableBase> : public TStructOpsTypeTraitsB
 	};
 };
 
-FORCEINLINE uint32 GetTypeHash(const FNiagaraVariableBase& Var)
-{
-	return HashCombine(GetTypeHash(Var.GetType()), GetTypeHash(Var.GetName()));
-}
-
 USTRUCT()
 struct FNiagaraVariable : public FNiagaraVariableBase
 {
@@ -2081,6 +2081,11 @@ struct NIAGARA_API FVersionedNiagaraEmitter
 	{
 		return !(*this == Other);
 	}
+
+	friend inline uint32 GetTypeHash(const FVersionedNiagaraEmitter& Item)
+	{
+		return GetTypeHash(Item.Emitter) + GetTypeHash(Item.Version);
+	}
 };
 
 struct NIAGARA_API FVersionedNiagaraEmitterWeakPtr
@@ -2103,16 +2108,9 @@ struct NIAGARA_API FVersionedNiagaraEmitterWeakPtr
 
 	TWeakObjectPtr<UNiagaraEmitter> Emitter;
 	FGuid Version;
-};
-
-inline uint32 GetTypeHash(const FVersionedNiagaraEmitterWeakPtr& Item)
-{
-	return GetTypeHash(Item.Emitter);
-}
-
-inline uint32 GetTypeHash(const FVersionedNiagaraEmitter& Item)
-{
-	return GetTypeHash(Item.Emitter) + GetTypeHash(Item.Version);
-}
-
 
+	friend inline uint32 GetTypeHash(const FVersionedNiagaraEmitterWeakPtr& Item)
+	{
+		return GetTypeHash(Item.Emitter);
+	}
+};
diff --git a/Engine/Plugins/Online/OnlineSubsystem/Source/Public/Interfaces/OnlineGameActivityInterface.h b/Engine/Plugins/Online/OnlineSubsystem/Source/Public/Interfaces/OnlineGameActivityInterface.h
--- a/Engine/Plugins/Online/OnlineSubsystem/Source/Public/Interfaces/OnlineGameActivityInterface.h
+++ b/Engine/Plugins/Online/OnlineSubsystem/Source/Public/Interfaces/OnlineGameActivityInterface.h
@@ -19,15 +19,15 @@ struct FOnlineActivityPlayerLocation
 	FString ZoneId;
 	/** Coordinates */
 	FVector Coordinates;
-};
 
-/**
- * Equals operator for comparing FOnlineActivityPlayerLocation objects
- */
-inline bool operator==(const FOnlineActivityPlayerLocation& Location1, const FOnlineActivityPlayerLocation& Location2)
-{
-	return Location1.ZoneId == Location2.ZoneId && Location1.Coordinates == Location2.Coordinates;
-}
+	/**
+	 * Equals operator for comparing FOnlineActivityPlayerLocation objects
+	 */
+	friend inline bool operator==(const FOnlineActivityPlayerLocation& Location1, const FOnlineActivityPlayerLocation& Location2)
+	{
+		return Location1.ZoneId == Location2.ZoneId && Location1.Coordinates == Location2.Coordinates;
+	}
+};
 
 /**
  * Outcome representation of ending an activity.
diff --git a/Engine/Plugins/Online/OnlineSubsystem/Source/Public/OnlineSubsystemTypes.h b/Engine/Plugins/Online/OnlineSubsystem/Source/Public/OnlineSubsystemTypes.h
--- a/Engine/Plugins/Online/OnlineSubsystem/Source/Public/OnlineSubsystemTypes.h
+++ b/Engine/Plugins/Online/OnlineSubsystem/Source/Public/OnlineSubsystemTypes.h
@@ -878,7 +878,7 @@ public:
 
 	virtual uint32 GetTypeHash() const override
 	{
-		return ::GetTypeHash(UniqueNetIdStr);
+		return GetTypeHashHelper(UniqueNetIdStr);
 	}
 
 public:
@@ -978,7 +978,7 @@ public: \
 	} \
 	friend uint32 GetTypeHash(const SUBCLASSNAME& A) \
 	{ \
-		return ::GetTypeHash(A.UniqueNetIdStr); \
+		return GetTypeHashHelper(A.UniqueNetIdStr); \
 	} \
 	static const SUBCLASSNAME##Ref& EmptyId() \
 	{ \
@@ -1580,15 +1580,15 @@ public:
 
 protected:
 	TInternalType Value;
+
+	friend inline uint32 GetTypeHash(const FOnlinePartyTypeId Id)
+	{
+		return Id.GetValue();
+	}
 };
 
 inline bool IsValid(const FOnlinePartyTypeId Id) { return Id.GetValue() != 0; }
 
-inline uint32 GetTypeHash(const FOnlinePartyTypeId Id)
-{
-	return Id.GetValue();
-}
-
 class FOnlineFriendSettingsSourceData
 {
 public:
diff --git a/Engine/Plugins/Online/OnlineSubsystemSteam/Source/Private/IPAddressSteam.h b/Engine/Plugins/Online/OnlineSubsystemSteam/Source/Private/IPAddressSteam.h
--- a/Engine/Plugins/Online/OnlineSubsystemSteam/Source/Private/IPAddressSteam.h
+++ b/Engine/Plugins/Online/OnlineSubsystemSteam/Source/Private/IPAddressSteam.h
@@ -176,7 +176,7 @@ public:
 
 	virtual uint32 GetTypeHash() const override
 	{
-		return ::GetTypeHash(ToString(true));
+		return GetTypeHashHelper(ToString(true));
 	}
 
 	friend uint32 GetTypeHash(const FInternetAddrSteam& A)
diff --git a/Engine/Plugins/Online/SocketSubsystemEOS/Source/SocketSubsystemEOS/Private/InternetAddrEOS.cpp b/Engine/Plugins/Online/SocketSubsystemEOS/Source/SocketSubsystemEOS/Private/InternetAddrEOS.cpp
--- a/Engine/Plugins/Online/SocketSubsystemEOS/Source/SocketSubsystemEOS/Private/InternetAddrEOS.cpp
+++ b/Engine/Plugins/Online/SocketSubsystemEOS/Source/SocketSubsystemEOS/Private/InternetAddrEOS.cpp
@@ -166,7 +166,7 @@ FString FInternetAddrEOS::ToString(bool bAppendPort) const
 
 uint32 FInternetAddrEOS::GetTypeHash() const
 {
-	return HashCombine(HashCombine(HashCombine(::GetTypeHash((void*)LocalUserId), ::GetTypeHash((void*)RemoteUserId)), ::GetTypeHash(FAnsiStringView(SocketName, EOS_SOCKET_NAME_SIZE))), Channel);
+	return HashCombine(HashCombine(HashCombine(GetTypeHashHelper((void*)LocalUserId), GetTypeHashHelper((void*)RemoteUserId)), GetTypeHashHelper(FAnsiStringView(SocketName, EOS_SOCKET_NAME_SIZE))), Channel);
 }
 
 bool FInternetAddrEOS::IsValid() const
diff --git a/Engine/Plugins/Runtime/DataRegistry/Source/DataRegistry/Public/DataRegistryTypes.h b/Engine/Plugins/Runtime/DataRegistry/Source/DataRegistry/Public/DataRegistryTypes.h
--- a/Engine/Plugins/Runtime/DataRegistry/Source/DataRegistry/Public/DataRegistryTypes.h
+++ b/Engine/Plugins/Runtime/DataRegistry/Source/DataRegistry/Public/DataRegistryTypes.h
@@ -132,6 +132,25 @@ struct DATAREGISTRY_API FDataRegistryLookup
 	/** Generate a debug string */
 	FString ToString() const;
 
+	friend FORCEINLINE uint32 GetTypeHash(const FDataRegistryLookup& Lookup)
+	{
+		uint32 Hash = GetTypeHash(Lookup.SourceLookups.Num());
+		for (int32 i = 0; i < Lookup.SourceLookups.Num(); i++)
+		{
+			Hash = HashCombine(Hash, GetTypeHash(Lookup.SourceLookups[i]));
+		}
+		return Hash;
+	}
+
+	friend FORCEINLINE bool operator==(const FDataRegistryLookup& X, const FDataRegistryLookup& Y)
+	{
+		return X.SourceLookups == Y.SourceLookups;
+	}
+
+	friend FORCEINLINE bool operator!=(const FDataRegistryLookup& X, const FDataRegistryLookup& Y)
+	{
+		return X.SourceLookups != Y.SourceLookups;
+	}
 };
 
 template<>
@@ -143,26 +162,6 @@ struct TStructOpsTypeTraits<FDataRegistryLookup> : public TStructOpsTypeTraitsBa
 	};
 };
 
-FORCEINLINE uint32 GetTypeHash(const FDataRegistryLookup& Lookup)
-{
-	uint32 Hash = GetTypeHash(Lookup.SourceLookups.Num());
-	for (int32 i = 0; i < Lookup.SourceLookups.Num(); i++)
-	{
-		Hash = HashCombine(Hash, GetTypeHash(Lookup.SourceLookups[i]));
-	}
-	return Hash;
-}
-
-FORCEINLINE bool operator==(const FDataRegistryLookup& X, const FDataRegistryLookup& Y)
-{
-	return X.SourceLookups == Y.SourceLookups;
-}
-
-FORCEINLINE bool operator!=(const FDataRegistryLookup& X, const FDataRegistryLookup& Y)
-{
-	return X.SourceLookups != Y.SourceLookups;
-}
-
 
 /** A debugging/editor struct used to describe the source for a single data registry item */
 USTRUCT()
@@ -184,25 +183,25 @@ struct DATAREGISTRY_API FDataRegistrySourceItemId
 
 	/** Gets a unique ID key for this item, or None for invalid cached data */
 	FString GetDebugString() const;
-};
 
-FORCEINLINE uint32 GetTypeHash(const FDataRegistrySourceItemId& SourceItem)
-{
-	uint32 Hash = GetTypeHash(SourceItem.ItemId);
-	Hash = HashCombine(Hash, GetTypeHash(SourceItem.CacheLookup));
-	Hash = HashCombine(Hash, GetTypeHash(SourceItem.SourceResolvedName));
-	return Hash;
-}
+	friend FORCEINLINE uint32 GetTypeHash(const FDataRegistrySourceItemId& SourceItem)
+	{
+		uint32 Hash = GetTypeHash(SourceItem.ItemId);
+		Hash = HashCombine(Hash, GetTypeHash(SourceItem.CacheLookup));
+		Hash = HashCombine(Hash, GetTypeHash(SourceItem.SourceResolvedName));
+		return Hash;
+	}
 
-FORCEINLINE bool operator==(const FDataRegistrySourceItemId& X, const FDataRegistrySourceItemId& Y)
-{
-	return X.ItemId == Y.ItemId && X.CacheLookup == Y.CacheLookup && X.SourceResolvedName == Y.SourceResolvedName;
-}
+	friend FORCEINLINE bool operator==(const FDataRegistrySourceItemId& X, const FDataRegistrySourceItemId& Y)
+	{
+		return X.ItemId == Y.ItemId && X.CacheLookup == Y.CacheLookup && X.SourceResolvedName == Y.SourceResolvedName;
+	}
 
-FORCEINLINE bool operator!=(const FDataRegistrySourceItemId& X, const FDataRegistrySourceItemId& Y)
-{
-	return !(X.ItemId == Y.ItemId && X.CacheLookup == Y.CacheLookup && X.SourceResolvedName == Y.SourceResolvedName);
-}
+	friend FORCEINLINE bool operator!=(const FDataRegistrySourceItemId& X, const FDataRegistrySourceItemId& Y)
+	{
+		return !(X.ItemId == Y.ItemId && X.CacheLookup == Y.CacheLookup && X.SourceResolvedName == Y.SourceResolvedName);
+	}
+};
 
 
 /** Information about the cache status of an item */
diff --git a/Engine/Plugins/Runtime/Metasound/Source/MetasoundFrontend/Private/Analysis/MetasoundFrontendAnalyzerAddress.cpp b/Engine/Plugins/Runtime/Metasound/Source/MetasoundFrontend/Private/Analysis/MetasoundFrontendAnalyzerAddress.cpp
--- a/Engine/Plugins/Runtime/Metasound/Source/MetasoundFrontend/Private/Analysis/MetasoundFrontendAnalyzerAddress.cpp
+++ b/Engine/Plugins/Runtime/Metasound/Source/MetasoundFrontend/Private/Analysis/MetasoundFrontendAnalyzerAddress.cpp
@@ -43,12 +43,12 @@ namespace Metasound
 
 		uint32 FAnalyzerAddress::GetHash() const
 		{
-			uint32 AddressHash = HashCombineFast(AnalyzerInstanceID.A, ::GetTypeHash(AnalyzerMemberName));
-			AddressHash = HashCombineFast(AddressHash, ::GetTypeHash(AnalyzerName));
-			AddressHash = HashCombineFast(AddressHash, ::GetTypeHash(DataType));
-			AddressHash = HashCombineFast(AddressHash, ::GetTypeHash(InstanceID));
+			uint32 AddressHash = HashCombineFast(AnalyzerInstanceID.A, GetTypeHashHelper(AnalyzerMemberName));
+			AddressHash = HashCombineFast(AddressHash, GetTypeHashHelper(AnalyzerName));
+			AddressHash = HashCombineFast(AddressHash, GetTypeHashHelper(DataType));
+			AddressHash = HashCombineFast(AddressHash, GetTypeHashHelper(InstanceID));
 			AddressHash = HashCombineFast(AddressHash, NodeID.A);
-			AddressHash = HashCombineFast(AddressHash, ::GetTypeHash(OutputName));
+			AddressHash = HashCombineFast(AddressHash, GetTypeHashHelper(OutputName));
 
 			return AddressHash;
 		}
diff --git a/Engine/Plugins/Runtime/Metasound/Source/MetasoundFrontend/Private/MetasoundFrontendQuery.cpp b/Engine/Plugins/Runtime/Metasound/Source/MetasoundFrontend/Private/MetasoundFrontendQuery.cpp
--- a/Engine/Plugins/Runtime/Metasound/Source/MetasoundFrontend/Private/MetasoundFrontendQuery.cpp
+++ b/Engine/Plugins/Runtime/Metasound/Source/MetasoundFrontend/Private/MetasoundFrontendQuery.cpp
@@ -662,13 +662,13 @@ namespace Metasound
 
 	FFrontendQueryKey::FFrontendQueryKey(const FString& InKey)
 	: Key(TInPlaceType<FString>(), InKey)
-	, Hash(::GetTypeHash(InKey))
+	, Hash(GetTypeHash(InKey))
 	{
 	}
 
 	FFrontendQueryKey::FFrontendQueryKey(const FName& InKey)
 	: Key(TInPlaceType<FName>(), InKey)
-	, Hash(::GetTypeHash(InKey))
+	, Hash(GetTypeHash(InKey))
 	{
 	}
 
diff --git a/Engine/Plugins/Runtime/Metasound/Source/MetasoundFrontend/Private/MetasoundRouter.cpp b/Engine/Plugins/Runtime/Metasound/Source/MetasoundFrontend/Private/MetasoundRouter.cpp
--- a/Engine/Plugins/Runtime/Metasound/Source/MetasoundFrontend/Private/MetasoundRouter.cpp
+++ b/Engine/Plugins/Runtime/Metasound/Source/MetasoundFrontend/Private/MetasoundRouter.cpp
@@ -68,8 +68,8 @@ namespace Metasound
 
 	uint32 FSendAddress::GetHash() const
 	{
-		uint32 HashedChannel = HashCombineFast(::GetTypeHash(DataType), ::GetTypeHash(ChannelName));
-		HashedChannel = HashCombineFast(HashedChannel, ::GetTypeHash(InstanceID));
+		uint32 HashedChannel = HashCombineFast(GetTypeHashHelper(DataType), GetTypeHashHelper(ChannelName));
+		HashedChannel = HashCombineFast(HashedChannel, GetTypeHashHelper(InstanceID));
 		return HashedChannel;
 	}
 
diff --git a/Engine/Plugins/Runtime/Metasound/Source/MetasoundFrontend/Public/MetasoundAssetManager.h b/Engine/Plugins/Runtime/Metasound/Source/MetasoundFrontend/Public/MetasoundAssetManager.h
--- a/Engine/Plugins/Runtime/Metasound/Source/MetasoundFrontend/Public/MetasoundAssetManager.h
+++ b/Engine/Plugins/Runtime/Metasound/Source/MetasoundFrontend/Public/MetasoundAssetManager.h
@@ -98,7 +98,7 @@ namespace Metasound
 
 				FORCEINLINE friend uint32 GetTypeHash(const IMetaSoundAssetManager::FAssetInfo& InInfo)
 				{
-					return HashCombineFast(::GetTypeHash(InInfo.RegistryKey), GetTypeHash(InInfo.AssetPath));
+					return HashCombineFast(GetTypeHash(InInfo.RegistryKey), GetTypeHash(InInfo.AssetPath));
 				}
 			};
 
diff --git a/Engine/Plugins/Runtime/Metasound/Source/MetasoundFrontend/Public/MetasoundFrontendDocument.h b/Engine/Plugins/Runtime/Metasound/Source/MetasoundFrontend/Public/MetasoundFrontendDocument.h
--- a/Engine/Plugins/Runtime/Metasound/Source/MetasoundFrontend/Public/MetasoundFrontendDocument.h
+++ b/Engine/Plugins/Runtime/Metasound/Source/MetasoundFrontend/Public/MetasoundFrontendDocument.h
@@ -236,12 +236,12 @@ struct METASOUNDFRONTEND_API FMetasoundFrontendVersionNumber
 	{
 		return FString::Format(TEXT("v{0}.{1}"), { Major, Minor });
 	}
-};
 
-FORCEINLINE uint32 GetTypeHash(const FMetasoundFrontendVersionNumber& InNumber)
-{
-	return HashCombineFast(GetTypeHash(InNumber.Major), GetTypeHash(InNumber.Minor));
-}
+	friend FORCEINLINE uint32 GetTypeHash(const FMetasoundFrontendVersionNumber& InNumber)
+	{
+		return HashCombineFast(GetTypeHash(InNumber.Major), GetTypeHash(InNumber.Minor));
+	}
+};
 
 // General purpose version info for Metasound Frontend objects.
 USTRUCT()
@@ -312,12 +312,12 @@ struct METASOUNDFRONTEND_API FMetasoundFrontendVersion
 	{
 		return InLHS == InRHS || InLHS < InRHS;
 	}
-};
 
-FORCEINLINE uint32 GetTypeHash(const FMetasoundFrontendVersion& InVersion)
-{
-	return HashCombineFast(GetTypeHash(InVersion.Name), GetTypeHash(InVersion.Number));
-}
+	friend FORCEINLINE uint32 GetTypeHash(const FMetasoundFrontendVersion& InVersion)
+	{
+		return HashCombineFast(GetTypeHash(InVersion.Name), GetTypeHash(InVersion.Number));
+	}
+};
 
 // An FMetasoundFrontendVertex provides a named connection point of a node.
 USTRUCT() 
diff --git a/Engine/Plugins/Runtime/OpenXR/Source/OpenXRInput/Private/OpenXRInput.h b/Engine/Plugins/Runtime/OpenXR/Source/OpenXRInput/Private/OpenXRInput.h
--- a/Engine/Plugins/Runtime/OpenXR/Source/OpenXRInput/Private/OpenXRInput.h
+++ b/Engine/Plugins/Runtime/OpenXR/Source/OpenXRInput/Private/OpenXRInput.h
@@ -21,12 +21,6 @@ struct FInputActionKeyMapping;
 struct FInputAxisKeyMapping;
 struct FKey;
 
-// On some platforms the XrPath type becomes ambiguous for overloading
-FORCEINLINE uint32 GetTypeHash(const TPair<XrPath, XrPath>& Pair)
-{
-	return HashCombine(GetTypeHash((uint64)Pair.Key), GetTypeHash((uint64)Pair.Value));
-}
-
 class FOpenXRInputPlugin : public IOpenXRInputPlugin
 {
 public:
diff --git a/Engine/Plugins/Runtime/Steam/SteamSockets/Source/SteamSockets/Private/IPAddressSteamSockets.cpp b/Engine/Plugins/Runtime/Steam/SteamSockets/Source/SteamSockets/Private/IPAddressSteamSockets.cpp
--- a/Engine/Plugins/Runtime/Steam/SteamSockets/Source/SteamSockets/Private/IPAddressSteamSockets.cpp
+++ b/Engine/Plugins/Runtime/Steam/SteamSockets/Source/SteamSockets/Private/IPAddressSteamSockets.cpp
@@ -221,15 +221,15 @@ uint32 FInternetAddrSteamSockets::GetTypeHash() const
 		else
 		{
 			// Otherwise, hash the virtual port.
-			return ::GetTypeHash(P2PVirtualPort);
+			return GetTypeHashHelper(P2PVirtualPort);
 		}
 	}
 
 	if (GetProtocolType() == FNetworkProtocolTypes::SteamSocketsP2P)
 	{
-		return ::GetTypeHash(Addr.GetSteamID64());
+		return GetTypeHashHelper(Addr.GetSteamID64());
 	}
-	return ::GetTypeHash(ToString(true));
+	return GetTypeHashHelper(ToString(true));
 }
 
 FName FInternetAddrSteamSockets::GetProtocolType() const
diff --git a/Engine/Source/Developer/DerivedDataCache/Public/DerivedDataCacheKey.h b/Engine/Source/Developer/DerivedDataCache/Public/DerivedDataCacheKey.h
--- a/Engine/Source/Developer/DerivedDataCache/Public/DerivedDataCacheKey.h
+++ b/Engine/Source/Developer/DerivedDataCache/Public/DerivedDataCacheKey.h
@@ -90,6 +90,28 @@ struct FCacheKey
 
 	/** A key with a null bucket and a zero hash. */
 	static const FCacheKey Empty;
+
+	friend inline bool operator==(const FCacheKey& A, const FCacheKey& B)
+	{
+		return A.Bucket == B.Bucket && A.Hash == B.Hash;
+	}
+
+	friend inline bool operator!=(const FCacheKey& A, const FCacheKey& B)
+	{
+		return A.Bucket != B.Bucket || A.Hash != B.Hash;
+	}
+
+	friend inline bool operator<(const FCacheKey& A, const FCacheKey& B)
+	{
+		const FCacheBucket& BucketA = A.Bucket;
+		const FCacheBucket& BucketB = B.Bucket;
+		return BucketA == BucketB ? A.Hash < B.Hash : BucketA < BucketB;
+	}
+
+	friend inline uint32 GetTypeHash(const FCacheKey& Key)
+	{
+		return HashCombine(GetTypeHash(Key.Bucket), GetTypeHash(Key.Hash));
+	}
 };
 
 inline const FCacheKey FCacheKey::Empty;
@@ -112,28 +134,6 @@ UE_API bool LoadFromCompactBinary(FCbFieldView Field, FCacheBucket& OutBucket);
 
 ///////////////////////////////////////////////////////////////////////////////////////////////////////////////////////
 
-inline bool operator==(const FCacheKey& A, const FCacheKey& B)
-{
-	return A.Bucket == B.Bucket && A.Hash == B.Hash;
-}
-
-inline bool operator!=(const FCacheKey& A, const FCacheKey& B)
-{
-	return A.Bucket != B.Bucket || A.Hash != B.Hash;
-}
-
-inline bool operator<(const FCacheKey& A, const FCacheKey& B)
-{
-	const FCacheBucket& BucketA = A.Bucket;
-	const FCacheBucket& BucketB = B.Bucket;
-	return BucketA == BucketB ? A.Hash < B.Hash : BucketA < BucketB;
-}
-
-inline uint32 GetTypeHash(const FCacheKey& Key)
-{
-	return HashCombine(GetTypeHash(Key.Bucket), GetTypeHash(Key.Hash));
-}
-
 template <typename CharType>
 inline TStringBuilderBase<CharType>& operator<<(TStringBuilderBase<CharType>& Builder, const FCacheKey& Key)
 {
diff --git a/Engine/Source/Developer/TargetPlatform/Public/Interfaces/ITargetPlatform.h b/Engine/Source/Developer/TargetPlatform/Public/Interfaces/ITargetPlatform.h
--- a/Engine/Source/Developer/TargetPlatform/Public/Interfaces/ITargetPlatform.h
+++ b/Engine/Source/Developer/TargetPlatform/Public/Interfaces/ITargetPlatform.h
@@ -789,12 +789,12 @@ struct FTargetPlatform
 		return Ordinal == Other.Ordinal;
 	}
 
+	friend inline uint32 GetTypeHash(const FTargetPlatform& Key) { return Key.GetOrdinal(); }
+
 private:
 	uint32 Ordinal = 0;
 };
 
-inline uint32 GetTypeHash(const FTargetPlatform& Key) { return Key.GetOrdinal(); }
-
 /**
  * Target platform set implementation using bitmask for compactness
  */
diff --git a/Engine/Source/Runtime/AIModule/Classes/AITypes.h b/Engine/Source/Runtime/AIModule/Classes/AITypes.h
--- a/Engine/Source/Runtime/AIModule/Classes/AITypes.h
+++ b/Engine/Source/Runtime/AIModule/Classes/AITypes.h
@@ -254,6 +254,11 @@ public:
 		static const FAINamedID<TCounter> InvalidIDInstance;
 		return InvalidIDInstance;
 	}
+
+	friend FORCEINLINE uint32 GetTypeHash(const FAINamedID& ID)
+	{
+		return GetTypeHash(ID.Index);
+	}
 };
 
 template<typename TCounter>
diff --git a/Engine/Source/Runtime/AIModule/Classes/GenericTeamAgentInterface.h b/Engine/Source/Runtime/AIModule/Classes/GenericTeamAgentInterface.h
--- a/Engine/Source/Runtime/AIModule/Classes/GenericTeamAgentInterface.h
+++ b/Engine/Source/Runtime/AIModule/Classes/GenericTeamAgentInterface.h
@@ -65,12 +65,12 @@ protected:
 
 public:
 	static const FGenericTeamId NoTeam;
-};
 
-FORCEINLINE uint32 GetTypeHash(const FGenericTeamId TeamID)
-{
-	return TeamID.GetId();
-}
+	friend FORCEINLINE uint32 GetTypeHash(const FGenericTeamId Value)
+	{
+		return Value.GetId();
+	}
+};
 
 UINTERFACE()
 class AIMODULE_API UGenericTeamAgentInterface : public UInterface
diff --git a/Engine/Source/Runtime/AudioExtensions/Public/ISoundfieldFormat.h b/Engine/Source/Runtime/AudioExtensions/Public/ISoundfieldFormat.h
--- a/Engine/Source/Runtime/AudioExtensions/Public/ISoundfieldFormat.h
+++ b/Engine/Source/Runtime/AudioExtensions/Public/ISoundfieldFormat.h
@@ -446,9 +446,9 @@ struct FSoundfieldEncodingKey
 	{
 		return (SoundfieldFormat == Other.SoundfieldFormat) && (EncodingSettingsID == Other.EncodingSettingsID);
 	} 
-};
 
-inline uint32 GetTypeHash(const FSoundfieldEncodingKey& Value)
-{
-	return HashCombine(Value.EncodingSettingsID, Value.SoundfieldFormat.GetNumber());
-}
+	friend inline uint32 GetTypeHash(const FSoundfieldEncodingKey& Value)
+	{
+		return HashCombine(Value.EncodingSettingsID, Value.SoundfieldFormat.GetNumber());
+	}
+};
diff --git a/Engine/Source/Runtime/Core/Private/Misc/ConfigUtilities.cpp b/Engine/Source/Runtime/Core/Private/Misc/ConfigUtilities.cpp
--- a/Engine/Source/Runtime/Core/Private/Misc/ConfigUtilities.cpp
+++ b/Engine/Source/Runtime/Core/Private/Misc/ConfigUtilities.cpp
@@ -300,10 +300,10 @@ private:
 	};
 	friend uint32 GetTypeHash(const FConfigHistory& CH)
 	{
-		uint32 Hash = ::GetTypeHash(CH.Type);
-		Hash = HashCombine(Hash, ::GetTypeHash(CH.FileName));
-		Hash = HashCombine(Hash, ::GetTypeHash(CH.SectionName));
-		Hash = HashCombine(Hash, ::GetTypeHash(CH.Key));
+		uint32 Hash = GetTypeHash(CH.Type);
+		Hash = HashCombine(Hash, GetTypeHash(CH.FileName));
+		Hash = HashCombine(Hash, GetTypeHash(CH.SectionName));
+		Hash = HashCombine(Hash, GetTypeHash(CH.Key));
 		return Hash;
 	}
 	friend bool operator==(const FConfigHistory& A, const FConfigHistory& B)
diff --git a/Engine/Source/Runtime/Core/Public/Containers/BitArray.h b/Engine/Source/Runtime/Core/Public/Containers/BitArray.h
--- a/Engine/Source/Runtime/Core/Public/Containers/BitArray.h
+++ b/Engine/Source/Runtime/Core/Public/Containers/BitArray.h
@@ -1730,6 +1730,19 @@ public:
 		checkf(!Writer.Is32BitTarget(), TEXT("TBitArray does not currently support freezing for 32bits"));
 		TSupportsFreezeMemoryImageHelper<bSupportsFreezeMemoryImage>::WriteMemoryImage(Writer, *this);
 	}
+
+	friend FORCEINLINE uint32 GetTypeHash(const TBitArray& BitArray)
+	{
+		uint32 NumWords = FBitSet::CalculateNumWords(BitArray.Num());
+		uint32 Hash = NumWords;
+		const uint32* Data = BitArray.GetData();
+		for (uint32 i = 0; i < NumWords; i++)
+		{
+			Hash ^= Data[i];
+		}
+		return Hash;
+	}
+
 };
 
 namespace Freeze
@@ -1743,19 +1756,6 @@ namespace Freeze
 
 DECLARE_TEMPLATE_INTRINSIC_TYPE_LAYOUT(template<typename Allocator>, TBitArray<Allocator>);
 
-template<typename Allocator>
-FORCEINLINE uint32 GetTypeHash(const TBitArray<Allocator>& BitArray)
-{
-	uint32 NumWords = FBitSet::CalculateNumWords(BitArray.Num());
-	uint32 Hash = NumWords;
-	const uint32* Data = BitArray.GetData();
-	for (uint32 i = 0; i < NumWords; i++)
-	{
-		Hash ^= Data[i];
-	}
-	return Hash;
-}
-
 template<typename Allocator>
 struct TContainerTraits<TBitArray<Allocator> > : public TContainerTraitsBase<TBitArray<Allocator> >
 {
diff --git a/Engine/Source/Runtime/Core/Public/Containers/Deque.h b/Engine/Source/Runtime/Core/Public/Containers/Deque.h
--- a/Engine/Source/Runtime/Core/Public/Containers/Deque.h
+++ b/Engine/Source/Runtime/Core/Public/Containers/Deque.h
@@ -576,38 +576,37 @@ private:
 			checkf((Index >= 0) & (Index < Count), TEXT("Parameter index %d exceeds container size %d"), Index, Count);
 		}
 	}
-};
 
-//---------------------------------------------------------------------------------------------------------------------
-// TDeque comparison operators
-//---------------------------------------------------------------------------------------------------------------------
+	//---------------------------------------------------------------------------------------------------------------------
+	// TDeque comparison operators
+	//---------------------------------------------------------------------------------------------------------------------
 
-template <typename InElementType, typename InAllocatorType>
-inline bool operator==(
-	const TDeque<InElementType, InAllocatorType>& Left,
-	const TDeque<InElementType, InAllocatorType>& Right)
-{
-	if (Left.Num() == Right.Num())
+	friend inline bool operator==(
+		const TDeque<InElementType, InAllocatorType>& Left,
+		const TDeque<InElementType, InAllocatorType>& Right)
 	{
-		auto EndIt = Left.end();
-		auto LeftIt = Left.begin();
-		auto RightIt = Right.begin();
-		while (LeftIt != EndIt)
+		if (Left.Num() == Right.Num())
 		{
-			if (*LeftIt++ != *RightIt++)
+			auto EndIt = Left.end();
+			auto LeftIt = Left.begin();
+			auto RightIt = Right.begin();
+			while (LeftIt != EndIt)
 			{
-				return false;
+				if (*LeftIt++ != *RightIt++)
+				{
+					return false;
+				}
 			}
+			return true;
 		}
-		return true;
+		return false;
 	}
-	return false;
-}
 
-template <typename InElementType, typename InAllocatorType>
-inline bool operator!=(
-	const TDeque<InElementType, InAllocatorType>& Left,
-	const TDeque<InElementType, InAllocatorType>& Right)
-{
-	return !(Left == Right);
-}
+	friend inline bool operator!=(
+		const TDeque<InElementType, InAllocatorType>& Left,
+		const TDeque<InElementType, InAllocatorType>& Right)
+	{
+		return !(Left == Right);
+	}
+
+};
diff --git a/Engine/Source/Runtime/Core/Public/Containers/UnrealString.h b/Engine/Source/Runtime/Core/Public/Containers/UnrealString.h
--- a/Engine/Source/Runtime/Core/Public/Containers/UnrealString.h
+++ b/Engine/Source/Runtime/Core/Public/Containers/UnrealString.h
@@ -1849,6 +1849,13 @@ public:
 	{
 		Data.CountBytes(Ar);
 	}
+
+	/** Case insensitive string hash function. */
+	friend FORCEINLINE uint32 GetTypeHash(const FString& S)
+	{
+		// This must match the GetTypeHash behavior of FStringView
+		return FCrc::Strihash_DEPRECATED(S.Len(), *S);
+	}
 };
 
 template<>
@@ -1875,13 +1882,6 @@ inline int32 GetNum(const FString& String)
 	return String.Len();
 }
 
-/** Case insensitive string hash function. */
-FORCEINLINE uint32 GetTypeHash(const FString& S)
-{
-	// This must match the GetTypeHash behavior of FStringView
-	return FCrc::Strihash_DEPRECATED(S.Len(), *S);
-}
-
 /** 
  * Convert an array of bytes to a TCHAR
  * @param In byte array values to convert
diff --git a/Engine/Source/Runtime/Core/Public/Hash/Blake3.h b/Engine/Source/Runtime/Core/Public/Hash/Blake3.h
--- a/Engine/Source/Runtime/Core/Public/Hash/Blake3.h
+++ b/Engine/Source/Runtime/Core/Public/Hash/Blake3.h
@@ -57,6 +57,32 @@ public:
 
 private:
 	alignas(uint32) ByteArray Hash{};
+
+	friend inline bool operator==(const FBlake3Hash& A, const FBlake3Hash& B)
+	{
+		return FMemory::Memcmp(A.GetBytes(), B.GetBytes(), sizeof(decltype(A.GetBytes()))) == 0;
+	}
+
+	friend inline bool operator!=(const FBlake3Hash& A, const FBlake3Hash& B)
+	{
+		return FMemory::Memcmp(A.GetBytes(), B.GetBytes(), sizeof(decltype(A.GetBytes()))) != 0;
+	}
+
+	friend inline bool operator<(const FBlake3Hash& A, const FBlake3Hash& B)
+	{
+		return FMemory::Memcmp(A.GetBytes(), B.GetBytes(), sizeof(decltype(A.GetBytes()))) < 0;
+	}
+
+	friend inline FArchive& operator<<(FArchive& Ar, FBlake3Hash& Value)
+	{
+		Ar.Serialize(Value.GetBytes(), sizeof(decltype(Value.GetBytes())));
+		return Ar;
+	}
+
+	friend inline uint32 GetTypeHash(const FBlake3Hash& Value)
+	{
+		return *reinterpret_cast<const uint32*>(Value.GetBytes());
+	}
 };
 
 inline const FBlake3Hash FBlake3Hash::Zero;
@@ -144,32 +170,6 @@ inline bool FBlake3Hash::IsZero() const
 	return true;
 }
 
-inline bool operator==(const FBlake3Hash& A, const FBlake3Hash& B)
-{
-	return FMemory::Memcmp(A.GetBytes(), B.GetBytes(), sizeof(decltype(A.GetBytes()))) == 0;
-}
-
-inline bool operator!=(const FBlake3Hash& A, const FBlake3Hash& B)
-{
-	return FMemory::Memcmp(A.GetBytes(), B.GetBytes(), sizeof(decltype(A.GetBytes()))) != 0;
-}
-
-inline bool operator<(const FBlake3Hash& A, const FBlake3Hash& B)
-{
-	return FMemory::Memcmp(A.GetBytes(), B.GetBytes(), sizeof(decltype(A.GetBytes()))) < 0;
-}
-
-inline FArchive& operator<<(FArchive& Ar, FBlake3Hash& Hash)
-{
-	Ar.Serialize(Hash.GetBytes(), sizeof(decltype(Hash.GetBytes())));
-	return Ar;
-}
-
-inline uint32 GetTypeHash(const FBlake3Hash& Hash)
-{
-	return *reinterpret_cast<const uint32*>(Hash.GetBytes());
-}
-
 template <typename CharType>
 inline TStringBuilderBase<CharType>& operator<<(TStringBuilderBase<CharType>& Builder, const FBlake3Hash& Hash)
 {
diff --git a/Engine/Source/Runtime/Core/Public/IO/IoHash.h b/Engine/Source/Runtime/Core/Public/IO/IoHash.h
--- a/Engine/Source/Runtime/Core/Public/IO/IoHash.h
+++ b/Engine/Source/Runtime/Core/Public/IO/IoHash.h
@@ -71,6 +71,26 @@ public:
 
 private:
 	alignas(uint32) ByteArray Hash{};
+
+	friend inline bool operator==(const FIoHash& A, const FIoHash& B)
+	{
+		return FMemory::Memcmp(A.GetBytes(), B.GetBytes(), sizeof(decltype(A.GetBytes()))) == 0;
+	}
+
+	friend inline bool operator!=(const FIoHash& A, const FIoHash& B)
+	{
+		return FMemory::Memcmp(A.GetBytes(), B.GetBytes(), sizeof(decltype(A.GetBytes()))) != 0;
+	}
+
+	friend inline bool operator<(const FIoHash& A, const FIoHash& B)
+	{
+		return FMemory::Memcmp(A.GetBytes(), B.GetBytes(), sizeof(decltype(A.GetBytes()))) < 0;
+	}
+
+	friend inline uint32 GetTypeHash(const FIoHash& Value)
+	{
+		return *reinterpret_cast<const uint32*>(Value.GetBytes());
+	}
 };
 
 inline const FIoHash FIoHash::Zero;
@@ -142,32 +162,12 @@ inline FIoHash FIoHash::HashBuffer(const FCompositeBuffer& Buffer)
 	return FBlake3::HashBuffer(Buffer);
 }
 
-inline bool operator==(const FIoHash& A, const FIoHash& B)
-{
-	return FMemory::Memcmp(A.GetBytes(), B.GetBytes(), sizeof(decltype(A.GetBytes()))) == 0;
-}
-
-inline bool operator!=(const FIoHash& A, const FIoHash& B)
-{
-	return FMemory::Memcmp(A.GetBytes(), B.GetBytes(), sizeof(decltype(A.GetBytes()))) != 0;
-}
-
-inline bool operator<(const FIoHash& A, const FIoHash& B)
-{
-	return FMemory::Memcmp(A.GetBytes(), B.GetBytes(), sizeof(decltype(A.GetBytes()))) < 0;
-}
-
 inline FArchive& operator<<(FArchive& Ar, FIoHash& Hash)
 {
 	Ar.Serialize(Hash.GetBytes(), sizeof(decltype(Hash.GetBytes())));
 	return Ar;
 }
 
-inline uint32 GetTypeHash(const FIoHash& Hash)
-{
-	return *reinterpret_cast<const uint32*>(Hash.GetBytes());
-}
-
 template <typename CharType>
 inline TStringBuilderBase<CharType>& operator<<(TStringBuilderBase<CharType>& Builder, const FIoHash& Hash)
 {
diff --git a/Engine/Source/Runtime/Core/Public/Math/Color.h b/Engine/Source/Runtime/Core/Public/Math/Color.h
--- a/Engine/Source/Runtime/Core/Public/Math/Color.h
+++ b/Engine/Source/Runtime/Core/Public/Math/Color.h
@@ -459,6 +459,12 @@ struct FLinearColor
 	static CORE_API const FLinearColor Green;
 	static CORE_API const FLinearColor Blue;
 	static CORE_API const FLinearColor Yellow;
+
+	friend FORCEINLINE uint32 GetTypeHash( const FLinearColor& LinearColor )
+	{
+		// Note: this assumes there's no padding in FLinearColor that could contain uncompared data.
+		return FCrc::MemCrc_DEPRECATED(&LinearColor, sizeof(FLinearColor));
+	}
 };
 DECLARE_INTRINSIC_TYPE_LAYOUT(FLinearColor);
 
@@ -779,6 +785,11 @@ public:
 	static CORE_API const FColor Silver;
 	static CORE_API const FColor Emerald;
 
+	friend FORCEINLINE uint32 GetTypeHash( const FColor& Color )
+	{
+		return Color.DWColor();
+	}
+
 private:
 	/**
 	 * Please use .ToFColor(true) on FLinearColor if you wish to convert from FLinearColor to FColor,
@@ -837,18 +848,6 @@ FORCEINLINE FColor FLinearColor::ToFColor(const bool bSRGB) const
 	}
 }
 
-FORCEINLINE uint32 GetTypeHash( const FColor& Color )
-{
-	return Color.DWColor();
-}
-
-
-FORCEINLINE uint32 GetTypeHash( const FLinearColor& LinearColor )
-{
-	// Note: this assumes there's no padding in FLinearColor that could contain uncompared data.
-	return FCrc::MemCrc_DEPRECATED(&LinearColor, sizeof(FLinearColor));
-}
-
 
 /** Computes a brightness and a fixed point color from a floating point color. */
 extern CORE_API void ComputeAndFixedColorAndIntensity(const FLinearColor& InLinearColor,FColor& OutColor,float& OutIntensity);
diff --git a/Engine/Source/Runtime/Core/Public/Math/UnitConversion.h b/Engine/Source/Runtime/Core/Public/Math/UnitConversion.h
--- a/Engine/Source/Runtime/Core/Public/Math/UnitConversion.h
+++ b/Engine/Source/Runtime/Core/Public/Math/UnitConversion.h
@@ -188,14 +188,39 @@ private:
 
 	/** Given a string, skip past whitespace, then any numeric characters. Set End pointer to the end of the last numeric character. */
 	static bool ExtractNumberBoundary(const TCHAR* Start, const TCHAR*& End);
-};
 
-/** Global arithmetic operators for number types. Deals with conversion from related units correctly. */
-template<typename NumericType, typename OtherType>
-bool operator==(const FNumericUnit<NumericType>& LHS, const FNumericUnit<OtherType>& RHS);
+	/** Global arithmetic operators for number types. Deals with conversion from related units correctly. Note must be inlined for hidden friend optimization to work */
+	template<typename OtherType>
+	friend inline bool operator==(const FNumericUnit<NumericType>& LHS, const FNumericUnit<OtherType>& RHS)
+	{
+		if (LHS.Units != EUnit::Unspecified && RHS.Units != EUnit::Unspecified)
+		{
+			if (LHS.Units == RHS.Units)
+			{
+				return LHS.Value == RHS.Value;
+			}
+			else if (FUnitConversion::AreUnitsCompatible(LHS.Units, RHS.Units))
+			{
+				return LHS.Value == FUnitConversion::Convert(RHS.Value, RHS.Units, LHS.Units);
+			}
+			else
+			{
+				// Invalid conversion
+				return false;
+			}
+		}
+		else
+		{
+			return LHS.Value == RHS.Value;
+		}
+	}
 
-template<typename NumericType, typename OtherType>
-bool operator!=(const FNumericUnit<NumericType>& LHS, const FNumericUnit<OtherType>& RHS);
+	template<typename OtherType>
+	friend inline bool operator!=(const FNumericUnit<NumericType>& LHS, const FNumericUnit<OtherType>& RHS)
+	{
+		return !(LHS == RHS);
+	}
+};
 
 template<typename T>
 FString LexToString(const FNumericUnit<T>& NumericUnit);
diff --git a/Engine/Source/Runtime/Core/Public/Math/UnitConversion.inl b/Engine/Source/Runtime/Core/Public/Math/UnitConversion.inl
--- a/Engine/Source/Runtime/Core/Public/Math/UnitConversion.inl
+++ b/Engine/Source/Runtime/Core/Public/Math/UnitConversion.inl
@@ -380,39 +380,6 @@ bool FNumericUnit<NumericType>::ExtractNumberBoundary(const TCHAR* Start, const
 	return true;
 }
 
-/** Global arithmetic operators for number types. Deals with conversion from related units correctly. */
-template<typename NumericType, typename OtherType>
-bool operator==(const FNumericUnit<NumericType>& LHS, const FNumericUnit<OtherType>& RHS)
-{
-	if (LHS.Units != EUnit::Unspecified && RHS.Units != EUnit::Unspecified)
-	{
-		if (LHS.Units == RHS.Units)
-		{
-			return LHS.Value == RHS.Value;
-		}
-		else if (FUnitConversion::AreUnitsCompatible(LHS.Units, RHS.Units))
-		{
-			return LHS.Value == FUnitConversion::Convert(RHS.Value, RHS.Units, LHS.Units);
-		}
-		else
-		{
-			// Invalid conversion
-			return false;
-		}
-	}
-	else
-	{
-		return LHS.Value == RHS.Value;
-	}
-}
-
-template<typename NumericType, typename OtherType>
-bool operator!=(const FNumericUnit<NumericType>& LHS, const FNumericUnit<OtherType>& RHS)
-{
-	return !(LHS == RHS);
-}
-
-
 template <typename NumericType>
 struct TNumericLimits<FNumericUnit<NumericType>> : public TNumericLimits<NumericType>
 { };
diff --git a/Engine/Source/Runtime/Core/Public/Math/Vector4.h b/Engine/Source/Runtime/Core/Public/Math/Vector4.h
--- a/Engine/Source/Runtime/Core/Public/Math/Vector4.h
+++ b/Engine/Source/Runtime/Core/Public/Math/Vector4.h
@@ -521,6 +521,18 @@ public:
 	// Conversion from other type: float->double
 	template<typename FArg, TEMPLATE_REQUIRES(TAnd<TIsSame<FArg, float>, TIsSame<T, double>>::Value)>
 	explicit TVector4(const TVector4<FArg>& From) : TVector4<T>((T)From.X, (T)From.Y, (T)From.Z, (T)From.W) {}
+
+	/**
+	 * Creates a hash value from a TVector4.
+	 *
+	 * @param Vector the vector to create a hash value for
+	 *
+	 * @return The hash value from the components
+	 */
+	FORCEINLINE friend uint32 GetTypeHash(const UE::Math::TVector4<T>& Vector)
+	{
+		return FCrc::MemCrc_DEPRECATED(&Vector, sizeof(Vector));
+	}
 };
 
 /**
@@ -924,21 +936,6 @@ inline bool FVector4d::SerializeFromMismatchedTag(FName StructTag, FArchive& Ar)
 }
 
 
-/**
- * Creates a hash value from a TVector4.
- *
- * @param Vector the vector to create a hash value for
- *
- * @return The hash value from the components
- */
-template<typename T>
-FORCEINLINE uint32 GetTypeHash(const UE::Math::TVector4<T>& Vector)
-{
-	// Note: this assumes there's no padding in FVector that could contain uncompared data.
-	return FCrc::MemCrc_DEPRECATED(&Vector, sizeof(Vector));
-}
-
-
 /**
  * Calculates 3D Dot product of two 4D vectors.
  *
diff --git a/Engine/Source/Runtime/Core/Public/Memory/SharedBuffer.h b/Engine/Source/Runtime/Core/Public/Memory/SharedBuffer.h
--- a/Engine/Source/Runtime/Core/Public/Memory/SharedBuffer.h
+++ b/Engine/Source/Runtime/Core/Public/Memory/SharedBuffer.h
@@ -318,6 +318,11 @@ private:
 	inline friend FOwnerPtrType ToPrivateOwnerPtr(FUniqueBuffer&& Buffer) { return MoveTemp(Buffer.Owner); }
 
 	FOwnerPtrType Owner;
+
+	friend inline uint32 GetTypeHash(const FUniqueBuffer& Buffer) { return PointerHash(ToPrivateOwnerPtr(Buffer).Get()); }
+
+	friend inline bool operator==(const FUniqueBuffer& BufferA, const FUniqueBuffer& BufferB) { return ToPrivateOwnerPtr(BufferA) == ToPrivateOwnerPtr(BufferB); }
+	friend inline bool operator!=(const FUniqueBuffer& BufferA, const FUniqueBuffer& BufferB) { return ToPrivateOwnerPtr(BufferA) != ToPrivateOwnerPtr(BufferB); }
 };
 
 ///////////////////////////////////////////////////////////////////////////////////////////////////////////////////////
@@ -430,6 +435,15 @@ private:
 	inline friend FOwnerPtrType ToPrivateOwnerPtr(FSharedBuffer&& Buffer) { return MoveTemp(Buffer.Owner); }
 
 	FOwnerPtrType Owner;
+
+	friend inline uint32 GetTypeHash(const FSharedBuffer& Buffer) { return PointerHash(ToPrivateOwnerPtr(Buffer).Get()); }
+	friend inline bool operator==(const FSharedBuffer& BufferA, const FSharedBuffer& BufferB) { return ToPrivateOwnerPtr(BufferA) == ToPrivateOwnerPtr(BufferB); }
+	friend inline bool operator!=(const FSharedBuffer& BufferA, const FSharedBuffer& BufferB) { return ToPrivateOwnerPtr(BufferA) != ToPrivateOwnerPtr(BufferB); }
+
+	friend inline bool operator==(const FSharedBuffer& BufferA, const FUniqueBuffer& BufferB) { return ToPrivateOwnerPtr(BufferA) == ToPrivateOwnerPtr(BufferB); }
+	friend inline bool operator!=(const FSharedBuffer& BufferA, const FUniqueBuffer& BufferB) { return ToPrivateOwnerPtr(BufferA) != ToPrivateOwnerPtr(BufferB); }
+	friend inline bool operator==(const FUniqueBuffer& BufferA, const FSharedBuffer& BufferB) { return ToPrivateOwnerPtr(BufferA) == ToPrivateOwnerPtr(BufferB); }
+	friend inline bool operator!=(const FUniqueBuffer& BufferA, const FSharedBuffer& BufferB) { return ToPrivateOwnerPtr(BufferA) != ToPrivateOwnerPtr(BufferB); }
 };
 
 ///////////////////////////////////////////////////////////////////////////////////////////////////////////////////////
@@ -463,27 +477,22 @@ private:
 	inline friend const FWeakOwnerPtrType& ToPrivateOwnerPtr(const FWeakSharedBuffer& Buffer) { return Buffer.Owner; }
 
 	FWeakOwnerPtrType Owner;
-};
 
-///////////////////////////////////////////////////////////////////////////////////////////////////////////////////////
+	friend inline uint32 GetTypeHash(const FWeakSharedBuffer& Buffer) { return PointerHash(ToPrivateOwnerPtr(Buffer).Get()); }
 
-template <typename TypeA, typename TypeB>
-inline auto operator==(const TypeA& BufferA, const TypeB& BufferB)
-	-> decltype(ToPrivateOwnerPtr(BufferA) == ToPrivateOwnerPtr(BufferB))
-{
-	return ToPrivateOwnerPtr(BufferA) == ToPrivateOwnerPtr(BufferB);
-}
+	friend inline bool operator==(const FWeakSharedBuffer& BufferA, const FWeakSharedBuffer& BufferB) { return ToPrivateOwnerPtr(BufferA) == ToPrivateOwnerPtr(BufferB); }
+	friend inline bool operator!=(const FWeakSharedBuffer& BufferA, const FWeakSharedBuffer& BufferB) { return ToPrivateOwnerPtr(BufferA) != ToPrivateOwnerPtr(BufferB); }
 
-template <typename TypeA, typename TypeB>
-inline auto operator!=(const TypeA& BufferA, const TypeB& BufferB)
-	-> decltype(ToPrivateOwnerPtr(BufferA) != ToPrivateOwnerPtr(BufferB))
-{
-	return ToPrivateOwnerPtr(BufferA) != ToPrivateOwnerPtr(BufferB);
-}
+	friend inline bool operator==(const FWeakSharedBuffer& BufferA, const FUniqueBuffer& BufferB) { return ToPrivateOwnerPtr(BufferA) == ToPrivateOwnerPtr(BufferB); }
+	friend inline bool operator!=(const FWeakSharedBuffer& BufferA, const FUniqueBuffer& BufferB) { return ToPrivateOwnerPtr(BufferA) != ToPrivateOwnerPtr(BufferB); }
+	friend inline bool operator==(const FUniqueBuffer& BufferA, const FWeakSharedBuffer& BufferB) { return ToPrivateOwnerPtr(BufferA) == ToPrivateOwnerPtr(BufferB); }
+	friend inline bool operator!=(const FUniqueBuffer& BufferA, const FWeakSharedBuffer& BufferB) { return ToPrivateOwnerPtr(BufferA) != ToPrivateOwnerPtr(BufferB); }
 
-inline uint32 GetTypeHash(const FUniqueBuffer& Buffer) { return PointerHash(ToPrivateOwnerPtr(Buffer).Get()); }
-inline uint32 GetTypeHash(const FSharedBuffer& Buffer) { return PointerHash(ToPrivateOwnerPtr(Buffer).Get()); }
-inline uint32 GetTypeHash(const FWeakSharedBuffer& Buffer) { return PointerHash(ToPrivateOwnerPtr(Buffer).Get()); }
+	friend inline bool operator==(const FWeakSharedBuffer& BufferA, const FSharedBuffer& BufferB) { return ToPrivateOwnerPtr(BufferA) == ToPrivateOwnerPtr(BufferB); }
+	friend inline bool operator!=(const FWeakSharedBuffer& BufferA, const FSharedBuffer& BufferB) { return ToPrivateOwnerPtr(BufferA) != ToPrivateOwnerPtr(BufferB); }
+	friend inline bool operator==(const FSharedBuffer& BufferA, const FWeakSharedBuffer& BufferB) { return ToPrivateOwnerPtr(BufferA) == ToPrivateOwnerPtr(BufferB); }
+	friend inline bool operator!=(const FSharedBuffer& BufferA, const FWeakSharedBuffer& BufferB) { return ToPrivateOwnerPtr(BufferA) != ToPrivateOwnerPtr(BufferB); }
+};
 
 ///////////////////////////////////////////////////////////////////////////////////////////////////////////////////////
 
diff --git a/Engine/Source/Runtime/Core/Public/Microsoft/COMPointer.h b/Engine/Source/Runtime/Core/Public/Microsoft/COMPointer.h
--- a/Engine/Source/Runtime/Core/Public/Microsoft/COMPointer.h
+++ b/Engine/Source/Runtime/Core/Public/Microsoft/COMPointer.h
@@ -264,4 +264,9 @@ private:
 
 	/** Pointer to the actual object, if any. */
 	PointerType* RawPointer;
+
+	friend FORCEINLINE uint32 GetTypeHash(const TComPtr<T>& InObjectPtr)
+	{
+		return GetTypeHash(InObjectPtr.RawPointer);
+	}
 };
diff --git a/Engine/Source/Runtime/Core/Public/Misc/FrameNumber.h b/Engine/Source/Runtime/Core/Public/Misc/FrameNumber.h
--- a/Engine/Source/Runtime/Core/Public/Misc/FrameNumber.h
+++ b/Engine/Source/Runtime/Core/Public/Misc/FrameNumber.h
@@ -66,17 +66,17 @@ struct FFrameNumber
 	friend FFrameNumber operator*(FFrameNumber A, float Scalar)   { return FFrameNumber(static_cast<int32>(FMath::Clamp(FMath::FloorToDouble(double(A.Value) * Scalar), (double)TNumericLimits<int32>::Min(), (double)TNumericLimits<int32>::Max()))); }
 	friend FFrameNumber operator/(FFrameNumber A, float Scalar)   { return FFrameNumber(static_cast<int32>(FMath::Clamp(FMath::FloorToDouble(double(A.Value) / Scalar), (double)TNumericLimits<int32>::Min(), (double)TNumericLimits<int32>::Max()))); }
 
+	friend inline uint32 GetTypeHash(FFrameNumber A)
+	{
+		return A.Value;
+	}
+
 	/**
 	 * The value of the frame number
 	 */
 	int32 Value;
 };
 
-inline uint32 GetTypeHash(FFrameNumber A)
-{
-	return A.Value;
-}
-
 template<>
 struct TNumericLimits<FFrameNumber>
 {
diff --git a/Engine/Source/Runtime/Core/Public/Misc/FrameRate.h b/Engine/Source/Runtime/Core/Public/Misc/FrameRate.h
--- a/Engine/Source/Runtime/Core/Public/Misc/FrameRate.h
+++ b/Engine/Source/Runtime/Core/Public/Misc/FrameRate.h
@@ -31,12 +31,6 @@ struct FFrameRate
 		: Numerator(InNumerator), Denominator(InDenominator)
 	{}
 
-	friend bool operator==(const FFrameRate& A, const FFrameRate& B);
-	friend bool operator!=(const FFrameRate& A, const FFrameRate& B);
-
-	friend FFrameRate operator*(FFrameRate A, FFrameRate B);
-	friend FFrameRate operator/(FFrameRate A, FFrameRate B);
-
 	/**
 	 * The numerator of the framerate represented as a number of frames per second (e.g. 60 for 60 fps)
 	 */
@@ -153,6 +147,66 @@ struct FFrameRate
 	{
 		return FFrameRate(Denominator, Numerator);
 	}
+
+	friend inline bool operator==(const FFrameRate& A, const FFrameRate& B)
+	{
+		return A.Numerator == B.Numerator && A.Denominator == B.Denominator;
+	}
+
+	friend inline bool operator!=(const FFrameRate& A, const FFrameRate& B)
+	{
+		return A.Numerator != B.Numerator || A.Denominator != B.Denominator;
+	}
+
+	friend inline FFrameRate operator*(FFrameRate A, FFrameRate B)
+	{
+		return FFrameRate(A.Numerator * B.Numerator, A.Denominator * B.Denominator);
+	}
+
+	friend inline FFrameRate operator/(FFrameRate A, FFrameRate B)
+	{
+		return FFrameRate(A.Numerator * B.Denominator, A.Denominator * B.Numerator);
+	}
+
+	friend inline double operator/(FFrameNumber Frame, FFrameRate Rate)
+	{
+		return Rate.AsSeconds(FFrameTime(Frame));
+	}
+
+	friend inline TRange<double> operator/(const TRange<FFrameNumber>& FrameRange, FFrameRate Rate)
+	{
+		TRangeBound<FFrameNumber> LowerBound = FrameRange.GetLowerBound();
+		TRangeBound<FFrameNumber> UpperBound = FrameRange.GetUpperBound();
+
+		return TRange<double>(
+			LowerBound.IsOpen()
+				? TRangeBound<double>::Open()
+				: LowerBound.IsInclusive()
+					? TRangeBound<double>::Inclusive(Rate.AsSeconds(LowerBound.GetValue()))
+					: TRangeBound<double>::Inclusive(Rate.AsSeconds(LowerBound.GetValue()+1)),
+
+			UpperBound.IsOpen()
+				? TRangeBound<double>::Open()
+				: UpperBound.IsInclusive()
+					? TRangeBound<double>::Exclusive(Rate.AsSeconds(UpperBound.GetValue()+1))
+					: TRangeBound<double>::Exclusive(Rate.AsSeconds(UpperBound.GetValue()))
+			);
+	}
+
+	friend inline double operator/(FFrameTime FrameTime, FFrameRate Rate)
+	{
+		return Rate.AsSeconds(FrameTime);
+	}
+
+	friend inline FFrameTime operator*(double TimeInSeconds, FFrameRate Rate)
+	{
+		return Rate.AsFrameTime(TimeInSeconds);
+	}
+
+	friend inline FFrameTime operator*(float TimeInSeconds, FFrameRate Rate)
+	{
+		return Rate.AsFrameTime(TimeInSeconds);
+	}
 };
 
 inline double FFrameRate::AsInterval() const
@@ -198,66 +252,6 @@ inline FFrameNumber FFrameRate::AsFrameNumber(double TimeInSeconds) const
 	return static_cast<int32>(FMath::Clamp(FMath::FloorToDouble(TimeAsFrame + static_cast<double>(FMath::TruncToInt(static_cast<float>(TimeAsFrame - FMath::FloorToDouble(TimeAsFrame))))), static_cast<double>(TNumericLimits<int32>::Min()), static_cast<double>(TNumericLimits<int32>::Max())));
 }
 
-inline bool operator==(const FFrameRate& A, const FFrameRate& B)
-{
-	return A.Numerator == B.Numerator && A.Denominator == B.Denominator;
-}
-
-inline bool operator!=(const FFrameRate& A, const FFrameRate& B)
-{
-	return A.Numerator != B.Numerator || A.Denominator != B.Denominator;
-}
-
-inline FFrameRate operator*(FFrameRate A, FFrameRate B)
-{
-	return FFrameRate(A.Numerator * B.Numerator, A.Denominator * B.Denominator);
-}
-
-inline FFrameRate operator/(FFrameRate A, FFrameRate B)
-{
-	return FFrameRate(A.Numerator * B.Denominator, A.Denominator * B.Numerator);
-}
-
-inline double operator/(FFrameNumber Frame, FFrameRate Rate)
-{
-	return Rate.AsSeconds(FFrameTime(Frame));
-}
-
-inline TRange<double> operator/(const TRange<FFrameNumber>& FrameRange, FFrameRate Rate)
-{
-	TRangeBound<FFrameNumber> LowerBound = FrameRange.GetLowerBound();
-	TRangeBound<FFrameNumber> UpperBound = FrameRange.GetUpperBound();
-
-	return TRange<double>(
-		LowerBound.IsOpen()
-			? TRangeBound<double>::Open()
-			: LowerBound.IsInclusive()
-				? TRangeBound<double>::Inclusive(Rate.AsSeconds(LowerBound.GetValue()))
-				: TRangeBound<double>::Inclusive(Rate.AsSeconds(LowerBound.GetValue()+1)),
-
-		UpperBound.IsOpen()
-			? TRangeBound<double>::Open()
-			: UpperBound.IsInclusive()
-				? TRangeBound<double>::Exclusive(Rate.AsSeconds(UpperBound.GetValue()+1))
-				: TRangeBound<double>::Exclusive(Rate.AsSeconds(UpperBound.GetValue()))
-		);
-}
-
-inline double operator/(FFrameTime FrameTime, FFrameRate Rate)
-{
-	return Rate.AsSeconds(FrameTime);
-}
-
-inline FFrameTime operator*(double TimeInSeconds, FFrameRate Rate)
-{
-	return Rate.AsFrameTime(TimeInSeconds);
-}
-
-inline FFrameTime operator*(float TimeInSeconds, FFrameRate Rate)
-{
-	return Rate.AsFrameTime(TimeInSeconds);
-}
-
 inline FFrameTime ConvertFrameTime(FFrameTime SourceTime, FFrameRate SourceRate, FFrameRate DestinationRate)
 {
 	if (SourceRate == DestinationRate)
diff --git a/Engine/Source/Runtime/Core/Public/Misc/FrameTime.h b/Engine/Source/Runtime/Core/Public/Misc/FrameTime.h
--- a/Engine/Source/Runtime/Core/Public/Misc/FrameTime.h
+++ b/Engine/Source/Runtime/Core/Public/Misc/FrameTime.h
@@ -88,197 +88,173 @@ public:
 	 */
 	static FFrameTime FromDecimal(double InDecimalFrame);
 
-
-private:
-
-	/*~ Built in operators */
-	friend bool operator==(FFrameTime A, FFrameTime B);
-	friend bool operator!=(FFrameTime A, FFrameTime B);
-	friend bool operator> (FFrameTime A, FFrameTime B);
-	friend bool operator>=(FFrameTime A, FFrameTime B);
-	friend bool operator< (FFrameTime A, FFrameTime B);
-	friend bool operator<=(FFrameTime A, FFrameTime B);
-
-	friend FFrameTime& operator+=(FFrameTime& LHS, FFrameTime RHS);
-	friend FFrameTime& operator-=(FFrameTime& LHS, FFrameTime RHS);
-	friend FFrameTime  operator+(FFrameTime A, FFrameTime B);
-	friend FFrameTime  operator-(FFrameTime A, FFrameTime B);
-	friend FFrameTime  operator%(FFrameTime A, FFrameTime B);
-
-	friend FFrameTime  operator-(FFrameTime A);
-
-	friend FFrameTime  operator*(FFrameTime A, float Scalar);
-	friend FFrameTime  operator/(FFrameTime A, float Scalar);
-
-public:
 	FFrameNumber FrameNumber;
 
 private:
 
 	/** Must be 0.f <= SubFrame < 1.f */
 	float SubFrame;
-};
 
 
-inline FFrameTime::FFrameTime()
-	: FrameNumber(0), SubFrame(0.f)
-{}
+	friend FORCEINLINE_DEBUGGABLE bool operator==(FFrameTime A, FFrameTime B)
+	{
+		return A.FrameNumber == B.FrameNumber && A.SubFrame == B.SubFrame;
+	}
 
 
-template<typename T, typename>
-FFrameTime::FFrameTime(T InFrameNumber)
-	: FrameNumber(InFrameNumber), SubFrame(0.f)
-{}
+	friend FORCEINLINE_DEBUGGABLE bool operator!=(FFrameTime A, FFrameTime B)
+	{
+		return A.FrameNumber != B.FrameNumber || A.SubFrame != B.SubFrame;
+	}
 
 
-inline FFrameTime::FFrameTime(FFrameNumber InFrameNumber)
-	: FrameNumber(InFrameNumber), SubFrame(0.f)
-{}
+	friend FORCEINLINE_DEBUGGABLE bool operator> (FFrameTime A, FFrameTime B)
+	{
+		return A.FrameNumber >  B.FrameNumber || ( A.FrameNumber == B.FrameNumber && A.SubFrame > B.SubFrame );
+	}
 
 
-inline FFrameTime::FFrameTime(FFrameNumber InFrameNumber, float InSubFrame)
-	: FrameNumber(InFrameNumber), SubFrame(InSubFrame)
-{
-	// Hack to ensure that SubFrames are in a sensible range of precision to work around
-	// problems with FloorToXYZ returning the wrong thing for very small negative numbers
-	SubFrame = FMath::Clamp(SubFrame + 0.5f - 0.5f, 0.f, MaxSubframe);
-	checkSlow(InSubFrame >= 0.f && InSubFrame < 1.f);
-}
+	friend FORCEINLINE_DEBUGGABLE bool operator>=(FFrameTime A, FFrameTime B)
+	{
+		return A.FrameNumber > B.FrameNumber || ( A.FrameNumber == B.FrameNumber && A.SubFrame >= B.SubFrame );
+	}
 
 
-inline FFrameTime& FFrameTime::operator=(FFrameNumber InFrameNumber)
-{
-	FrameNumber = InFrameNumber;
-	SubFrame    = 0.f;
-	return *this;
-}
+	friend FORCEINLINE_DEBUGGABLE bool operator< (FFrameTime A, FFrameTime B)
+	{
+		return A.FrameNumber <  B.FrameNumber || ( A.FrameNumber == B.FrameNumber && A.SubFrame < B.SubFrame );
+	}
 
 
-FORCEINLINE_DEBUGGABLE bool operator==(FFrameTime A, FFrameTime B)
-{
-	return A.FrameNumber == B.FrameNumber && A.SubFrame == B.SubFrame;
-}
+	friend FORCEINLINE_DEBUGGABLE bool operator<=(FFrameTime A, FFrameTime B)
+	{
+		return A.FrameNumber < B.FrameNumber || ( A.FrameNumber == B.FrameNumber && A.SubFrame <= B.SubFrame );
+	}
 
 
-FORCEINLINE_DEBUGGABLE bool operator!=(FFrameTime A, FFrameTime B)
-{
-	return A.FrameNumber != B.FrameNumber || A.SubFrame != B.SubFrame;
-}
+	friend FORCEINLINE_DEBUGGABLE FFrameTime& operator+=(FFrameTime& LHS, FFrameTime RHS)
+	{
+		float NewSubFrame = LHS.SubFrame + RHS.SubFrame;
 
+		LHS.FrameNumber = LHS.FrameNumber + RHS.FrameNumber + FFrameNumber(FMath::FloorToInt(NewSubFrame));
+		LHS.SubFrame    = FMath::Frac(NewSubFrame);
 
-FORCEINLINE_DEBUGGABLE bool operator> (FFrameTime A, FFrameTime B)
-{
-	return A.FrameNumber >  B.FrameNumber || ( A.FrameNumber == B.FrameNumber && A.SubFrame > B.SubFrame );
-}
+		return LHS;
+	}
 
 
-FORCEINLINE_DEBUGGABLE bool operator>=(FFrameTime A, FFrameTime B)
-{
-	return A.FrameNumber > B.FrameNumber || ( A.FrameNumber == B.FrameNumber && A.SubFrame >= B.SubFrame );
-}
+	friend FORCEINLINE_DEBUGGABLE FFrameTime operator+(FFrameTime A, FFrameTime B)
+	{
+		const float        NewSubFrame    = A.SubFrame + B.SubFrame;
+		const FFrameNumber NewFrameNumber = A.FrameNumber + B.FrameNumber + FFrameNumber(FMath::FloorToInt(NewSubFrame));
 
+		return FFrameTime(NewFrameNumber, FMath::Frac(NewSubFrame));
+	}
 
-FORCEINLINE_DEBUGGABLE bool operator< (FFrameTime A, FFrameTime B)
-{
-	return A.FrameNumber <  B.FrameNumber || ( A.FrameNumber == B.FrameNumber && A.SubFrame < B.SubFrame );
-}
 
+	friend FORCEINLINE_DEBUGGABLE FFrameTime& operator-=(FFrameTime& LHS, FFrameTime RHS)
+	{
+		// Ensure SubFrame is always between 0 and 1
+		// Note that the difference between frame -1.5 and 1.5 is 2, not 3, since sub frame positions are always positive
+		const float        NewSubFrame     = LHS.SubFrame - RHS.SubFrame;
+		const float        FlooredSubFrame = FMath::FloorToFloat(NewSubFrame);
+		LHS.FrameNumber  = LHS.FrameNumber - RHS.FrameNumber + FFrameNumber(FMath::TruncToInt(FlooredSubFrame));
+		LHS.SubFrame = NewSubFrame - FlooredSubFrame;
+
+		return LHS;
+	}
 
-FORCEINLINE_DEBUGGABLE bool operator<=(FFrameTime A, FFrameTime B)
-{
-	return A.FrameNumber < B.FrameNumber || ( A.FrameNumber == B.FrameNumber && A.SubFrame <= B.SubFrame );
-}
 
+	friend FORCEINLINE_DEBUGGABLE FFrameTime operator-(FFrameTime A, FFrameTime B)
+	{
+		// Ensure SubFrame is always between 0 and 1
+		// Note that the difference between frame -1.5 and 1.5 is 2, not 3, since sub frame positions are always positive
+		const float        NewSubFrame     = A.SubFrame - B.SubFrame;
+		const float        FlooredSubFrame = FMath::FloorToFloat(NewSubFrame);
+		const FFrameNumber NewFrameNumber  = A.FrameNumber - B.FrameNumber + FFrameNumber(FMath::TruncToInt(FlooredSubFrame));
 
-FORCEINLINE_DEBUGGABLE FFrameTime& operator+=(FFrameTime& LHS, FFrameTime RHS)
-{
-	float NewSubFrame = LHS.SubFrame + RHS.SubFrame;
+		return FFrameTime(NewFrameNumber, NewSubFrame - FlooredSubFrame);
+	}
 
-	LHS.FrameNumber = LHS.FrameNumber + RHS.FrameNumber + FFrameNumber(FMath::FloorToInt(NewSubFrame));
-	LHS.SubFrame    = FMath::Frac(NewSubFrame);
 
-	return LHS;
-}
+	friend FORCEINLINE_DEBUGGABLE FFrameTime operator%(FFrameTime A, FFrameTime B)
+	{
+		check(B.FrameNumber.Value != 0 || B.GetSubFrame() != 0.f);
 
+		if (A.SubFrame == 0.f && B.SubFrame == 0.f)
+		{
+			return FFrameTime(A.FrameNumber % B.FrameNumber);
+		}
+		else
+		{
+			FFrameTime Result = A;
+			while (Result >= B)
+			{
+				Result = Result - B;
+			}
+			return Result;
+		}
+	}
 
-FORCEINLINE_DEBUGGABLE FFrameTime operator+(FFrameTime A, FFrameTime B)
-{
-	const float        NewSubFrame    = A.SubFrame + B.SubFrame;
-	const FFrameNumber NewFrameNumber = A.FrameNumber + B.FrameNumber + FFrameNumber(FMath::FloorToInt(NewSubFrame));
 
-	return FFrameTime(NewFrameNumber, FMath::Frac(NewSubFrame));
-}
+	friend FORCEINLINE_DEBUGGABLE FFrameTime operator-(FFrameTime A)
+	{
+		return A.GetSubFrame() == 0.f
+			? FFrameTime(-A.FrameNumber)
+			: FFrameTime(-A.FrameNumber - 1, 1.f-A.GetSubFrame());
+	}
 
 
-FORCEINLINE_DEBUGGABLE FFrameTime& operator-=(FFrameTime& LHS, FFrameTime RHS)
-{
-	// Ensure SubFrame is always between 0 and 1
-	// Note that the difference between frame -1.5 and 1.5 is 2, not 3, since sub frame positions are always positive
-	const float        NewSubFrame     = LHS.SubFrame - RHS.SubFrame;
-	const float        FlooredSubFrame = FMath::FloorToFloat(NewSubFrame);
-	LHS.FrameNumber  = LHS.FrameNumber - RHS.FrameNumber + FFrameNumber(FMath::TruncToInt(FlooredSubFrame));
-	LHS.SubFrame = NewSubFrame - FlooredSubFrame;
-
-	return LHS;
-}
+	friend FORCEINLINE FFrameTime operator*(FFrameTime A, float Scalar)
+	{
+		return FFrameTime::FromDecimal(A.AsDecimal() * Scalar);
+	}
 
+	friend FORCEINLINE FFrameTime operator*(float Scalar, FFrameTime A)
+	{
+		return FFrameTime::FromDecimal(A.AsDecimal() * Scalar);
+	}
 
-FORCEINLINE_DEBUGGABLE FFrameTime operator-(FFrameTime A, FFrameTime B)
-{
-	// Ensure SubFrame is always between 0 and 1
-	// Note that the difference between frame -1.5 and 1.5 is 2, not 3, since sub frame positions are always positive
-	const float        NewSubFrame     = A.SubFrame - B.SubFrame;
-	const float        FlooredSubFrame = FMath::FloorToFloat(NewSubFrame);
-	const FFrameNumber NewFrameNumber  = A.FrameNumber - B.FrameNumber + FFrameNumber(FMath::TruncToInt(FlooredSubFrame));
+	friend FORCEINLINE FFrameTime operator/(FFrameTime A, float Scalar)
+	{
+		return FFrameTime::FromDecimal(A.AsDecimal() / Scalar);
+	}
+};
 
-	return FFrameTime(NewFrameNumber, NewSubFrame - FlooredSubFrame);
-}
 
+inline FFrameTime::FFrameTime()
+	: FrameNumber(0), SubFrame(0.f)
+{}
 
-FORCEINLINE_DEBUGGABLE FFrameTime operator%(FFrameTime A, FFrameTime B)
-{
-	check(B.FrameNumber.Value != 0 || B.GetSubFrame() != 0.f);
 
-	if (A.SubFrame == 0.f && B.SubFrame == 0.f)
-	{
-		return FFrameTime(A.FrameNumber % B.FrameNumber);
-	}
-	else
-	{
-		FFrameTime Result = A;
-		while (Result >= B)
-		{
-			Result = Result - B;
-		}
-		return Result;
-	}
-}
+template<typename T, typename>
+FFrameTime::FFrameTime(T InFrameNumber)
+	: FrameNumber(InFrameNumber), SubFrame(0.f)
+{}
 
 
-FORCEINLINE_DEBUGGABLE FFrameTime operator-(FFrameTime A)
-{
-	return A.GetSubFrame() == 0.f
-		? FFrameTime(-A.FrameNumber)
-		: FFrameTime(-A.FrameNumber - 1, 1.f-A.GetSubFrame());
-}
+inline FFrameTime::FFrameTime(FFrameNumber InFrameNumber)
+	: FrameNumber(InFrameNumber), SubFrame(0.f)
+{}
 
 
-FORCEINLINE FFrameTime operator*(FFrameTime A, float Scalar)
+inline FFrameTime::FFrameTime(FFrameNumber InFrameNumber, float InSubFrame)
+	: FrameNumber(InFrameNumber), SubFrame(InSubFrame)
 {
-	return FFrameTime::FromDecimal(A.AsDecimal() * Scalar);
+	// Hack to ensure that SubFrames are in a sensible range of precision to work around
+	// problems with FloorToXYZ returning the wrong thing for very small negative numbers
+	SubFrame = FMath::Clamp(SubFrame + 0.5f - 0.5f, 0.f, MaxSubframe);
+	checkSlow(InSubFrame >= 0.f && InSubFrame < 1.f);
 }
 
-FORCEINLINE FFrameTime operator*(float Scalar, FFrameTime A)
-{
-	return FFrameTime::FromDecimal(A.AsDecimal() * Scalar);
-}
 
-FORCEINLINE FFrameTime operator/(FFrameTime A, float Scalar)
+inline FFrameTime& FFrameTime::operator=(FFrameNumber InFrameNumber)
 {
-	return FFrameTime::FromDecimal(A.AsDecimal() / Scalar);
+	FrameNumber = InFrameNumber;
+	SubFrame    = 0.f;
+	return *this;
 }
 
-
 FORCEINLINE_DEBUGGABLE FFrameNumber FFrameTime::FloorToFrame() const
 {
 	return FrameNumber;
diff --git a/Engine/Source/Runtime/Core/Public/Misc/Timecode.h b/Engine/Source/Runtime/Core/Public/Misc/Timecode.h
--- a/Engine/Source/Runtime/Core/Public/Misc/Timecode.h
+++ b/Engine/Source/Runtime/Core/Public/Misc/Timecode.h
@@ -345,14 +345,16 @@ public:
 
 	/** If true, this Timecode represents a Drop Frame timecode used to account for fractional frame rates in NTSC play rates. */
 	bool bDropFrameFormat;
-};
 
-inline bool operator==(const FTimecode& A, const FTimecode& B)
-{
-	return A.Hours == B.Hours && A.Minutes == B.Minutes && A.Seconds == B.Seconds && A.Frames == B.Frames;
-}
 
-inline bool operator!=(const FTimecode& A, const FTimecode& B)
-{
-	return !(A == B);
-}
+	friend inline bool operator==(const FTimecode& A, const FTimecode& B)
+	{
+		return A.Hours == B.Hours && A.Minutes == B.Minutes && A.Seconds == B.Seconds && A.Frames == B.Frames;
+	}
+
+	friend inline bool operator!=(const FTimecode& A, const FTimecode& B)
+	{
+		return !(A == B);
+	}
+
+};
diff --git a/Engine/Source/Runtime/Core/Public/Serialization/CompactBinary.h b/Engine/Source/Runtime/Core/Public/Serialization/CompactBinary.h
--- a/Engine/Source/Runtime/Core/Public/Serialization/CompactBinary.h
+++ b/Engine/Source/Runtime/Core/Public/Serialization/CompactBinary.h
@@ -500,39 +500,39 @@ public:
 
 private:
 	alignas(uint32) ByteArray Bytes{};
-};
 
-inline FCbObjectId::FCbObjectId(const ByteArray& ObjectId)
-{
-	FMemory::Memcpy(Bytes, ObjectId, sizeof(ByteArray));
-}
+	friend inline bool operator==(const FCbObjectId& A, const FCbObjectId& B)
+	{
+		return FMemory::Memcmp(&A, &B, sizeof(FCbObjectId)) == 0;
+	}
 
-inline bool operator==(const FCbObjectId& A, const FCbObjectId& B)
-{
-	return FMemory::Memcmp(&A, &B, sizeof(FCbObjectId)) == 0;
-}
+	friend inline bool operator!=(const FCbObjectId& A, const FCbObjectId& B)
+	{
+		return FMemory::Memcmp(&A, &B, sizeof(FCbObjectId)) != 0;
+	}
 
-inline bool operator!=(const FCbObjectId& A, const FCbObjectId& B)
-{
-	return FMemory::Memcmp(&A, &B, sizeof(FCbObjectId)) != 0;
-}
+	friend inline bool operator<(const FCbObjectId& A, const FCbObjectId& B)
+	{
+		return FMemory::Memcmp(&A, &B, sizeof(FCbObjectId)) <= 0;
+	}
 
-inline bool operator<(const FCbObjectId& A, const FCbObjectId& B)
-{
-	return FMemory::Memcmp(&A, &B, sizeof(FCbObjectId)) <= 0;
-}
+	friend inline uint32 GetTypeHash(const FCbObjectId& Id)
+	{
+		return *reinterpret_cast<const uint32*>(&Id);
+	}
 
-inline uint32 GetTypeHash(const FCbObjectId& Id)
-{
-	return *reinterpret_cast<const uint32*>(&Id);
-}
+	/** Convert the ObjectId to a 24-character hex string. */
+	template <typename CharType>
+	friend inline TStringBuilderBase<CharType>& operator<<(TStringBuilderBase<CharType>& Builder, const FCbObjectId& Id)
+	{
+		UE::String::BytesToHexLower(Id.GetBytes(), Builder);
+		return Builder;
+	}
+};
 
-/** Convert the ObjectId to a 24-character hex string. */
-template <typename CharType>
-inline TStringBuilderBase<CharType>& operator<<(TStringBuilderBase<CharType>& Builder, const FCbObjectId& Id)
+inline FCbObjectId::FCbObjectId(const ByteArray& ObjectId)
 {
-	UE::String::BytesToHexLower(Id.GetBytes(), Builder);
-	return Builder;
+	FMemory::Memcpy(Bytes, ObjectId, sizeof(ByteArray));
 }
 
 ///////////////////////////////////////////////////////////////////////////////////////////////////////////////////////
diff --git a/Engine/Source/Runtime/Core/Public/Serialization/MemoryImage.h b/Engine/Source/Runtime/Core/Public/Serialization/MemoryImage.h
--- a/Engine/Source/Runtime/Core/Public/Serialization/MemoryImage.h
+++ b/Engine/Source/Runtime/Core/Public/Serialization/MemoryImage.h
@@ -156,43 +156,45 @@ struct FMemoryImageVTablePointer
 	uint64 TypeNameHash;
 	uint32 VTableOffset;
 	uint32 Offset;
+
+	friend inline bool operator==(const FMemoryImageVTablePointer& Lhs, const FMemoryImageVTablePointer& Rhs)
+	{
+		return Lhs.TypeNameHash == Rhs.TypeNameHash && Lhs.VTableOffset == Rhs.VTableOffset && Lhs.Offset == Rhs.Offset;
+	}
+	friend inline bool operator!=(const FMemoryImageVTablePointer& Lhs, const FMemoryImageVTablePointer& Rhs)
+	{
+		return !operator==(Lhs, Rhs);
+	}
+	friend inline bool operator<(const FMemoryImageVTablePointer& Lhs, const FMemoryImageVTablePointer& Rhs)
+	{
+		if (Lhs.TypeNameHash != Rhs.TypeNameHash) return Lhs.TypeNameHash < Rhs.TypeNameHash;
+		if (Lhs.VTableOffset != Rhs.VTableOffset) return Lhs.VTableOffset < Rhs.VTableOffset;
+		return Lhs.Offset < Rhs.Offset;
+	}
 };
-inline bool operator==(const FMemoryImageVTablePointer& Lhs, const FMemoryImageVTablePointer& Rhs)
-{
-	return Lhs.TypeNameHash == Rhs.TypeNameHash && Lhs.VTableOffset == Rhs.VTableOffset && Lhs.Offset == Rhs.Offset;
-}
-inline bool operator!=(const FMemoryImageVTablePointer& Lhs, const FMemoryImageVTablePointer& Rhs)
-{
-	return !operator==(Lhs, Rhs);
-}
-inline bool operator<(const FMemoryImageVTablePointer& Lhs, const FMemoryImageVTablePointer& Rhs)
-{
-	if (Lhs.TypeNameHash != Rhs.TypeNameHash) return Lhs.TypeNameHash < Rhs.TypeNameHash;
-	if (Lhs.VTableOffset != Rhs.VTableOffset) return Lhs.VTableOffset < Rhs.VTableOffset;
-	return Lhs.Offset < Rhs.Offset;
-}
 
 struct FMemoryImageNamePointer
 {
 	FName Name;
 	uint32 Offset;
-};
-inline bool operator==(const FMemoryImageNamePointer& Lhs, const FMemoryImageNamePointer& Rhs)
-{
-	return Lhs.Offset == Rhs.Offset && Lhs.Name.Compare(Rhs.Name) == 0;
-}
-inline bool operator!=(const FMemoryImageNamePointer& Lhs, const FMemoryImageNamePointer& Rhs)
-{
-	return !operator==(Lhs, Rhs);
-}
-inline bool operator<(const FMemoryImageNamePointer& Lhs, const FMemoryImageNamePointer& Rhs)
-{
-	if (Lhs.Name != Rhs.Name)
+
+	friend inline bool operator==(const FMemoryImageNamePointer& Lhs, const FMemoryImageNamePointer& Rhs)
 	{
-		return Lhs.Name.LexicalLess(Rhs.Name);
+		return Lhs.Offset == Rhs.Offset && Lhs.Name.Compare(Rhs.Name) == 0;
 	}
-	return Lhs.Offset < Rhs.Offset;
-}
+	friend inline bool operator!=(const FMemoryImageNamePointer& Lhs, const FMemoryImageNamePointer& Rhs)
+	{
+		return !operator==(Lhs, Rhs);
+	}
+	friend inline bool operator<(const FMemoryImageNamePointer& Lhs, const FMemoryImageNamePointer& Rhs)
+	{
+		if (Lhs.Name != Rhs.Name)
+		{
+			return Lhs.Name.LexicalLess(Rhs.Name);
+		}
+		return Lhs.Offset < Rhs.Offset;
+	}
+};
 
 struct FMemoryImageResult
 {
@@ -822,13 +824,13 @@ public:
 	}
 
 	inline DataType::ElementAllocatorType& GetAllocatorInstance() { return Data.GetAllocatorInstance(); }
-};
 
-/** Case insensitive string hash function. */
-FORCEINLINE uint32 GetTypeHash(const FMemoryImageString& S)
-{
-	return FCrc::Strihash_DEPRECATED(*S);
-}
+	/** Case insensitive string hash function. */
+	friend FORCEINLINE uint32 GetTypeHash(const FMemoryImageString& S)
+	{
+		return FCrc::Strihash_DEPRECATED(*S);
+	}
+};
 
 #if WITH_EDITORONLY_DATA
 struct FHashedNameDebugString
diff --git a/Engine/Source/Runtime/Core/Public/Serialization/MemoryLayout.h b/Engine/Source/Runtime/Core/Public/Serialization/MemoryLayout.h
--- a/Engine/Source/Runtime/Core/Public/Serialization/MemoryLayout.h
+++ b/Engine/Source/Runtime/Core/Public/Serialization/MemoryLayout.h
@@ -139,9 +139,10 @@ struct FTypeLayoutDesc
 	uint8 NumVirtualBases;
 	uint8 IsIntrinsic : 1;
 	uint8 IsInitialized : 1;
+
+	friend inline bool operator==(const FTypeLayoutDesc& Lhs, const FTypeLayoutDesc& Rhs) { return &Lhs == &Rhs; }
+	friend inline bool operator!=(const FTypeLayoutDesc& Lhs, const FTypeLayoutDesc& Rhs) { return &Lhs != &Rhs; }
 };
-inline bool operator==(const FTypeLayoutDesc& Lhs, const FTypeLayoutDesc& Rhs) { return &Lhs == &Rhs; }
-inline bool operator!=(const FTypeLayoutDesc& Lhs, const FTypeLayoutDesc& Rhs) { return &Lhs != &Rhs; }
 
 struct FRegisterTypeLayoutDesc
 {
diff --git a/Engine/Source/Runtime/Core/Public/Templates/Function.h b/Engine/Source/Runtime/Core/Public/Templates/Function.h
--- a/Engine/Source/Runtime/Core/Public/Templates/Function.h
+++ b/Engine/Source/Runtime/Core/Public/Templates/Function.h
@@ -916,6 +916,38 @@ public:
 	{
 		return Super::IsSet();
 	}
+
+	/**
+	 * Nullptr equality operator.
+	 */
+	friend FORCEINLINE bool operator==(TYPE_OF_NULLPTR, const TFunction& Func)
+	{
+		return !Func;
+	}
+
+	/**
+	 * Nullptr equality operator.
+	 */
+	friend FORCEINLINE bool operator==(const TFunction& Func, TYPE_OF_NULLPTR)
+	{
+		return !Func;
+	}
+
+	/**
+	 * Nullptr inequality operator.
+	 */
+	friend FORCEINLINE bool operator!=(TYPE_OF_NULLPTR, const TFunction& Func)
+	{
+		return (bool)Func;
+	}
+
+	/**
+	 * Nullptr inequality operator.
+	 */
+	friend FORCEINLINE bool operator!=(const TFunction& Func, TYPE_OF_NULLPTR)
+	{
+		return (bool)Func;
+	}
 };
 
 /**
@@ -1021,39 +1053,3 @@ public:
 		return Super::IsSet();
 	}
 };
-
-/**
- * Nullptr equality operator.
- */
-template <typename FuncType>
-FORCEINLINE bool operator==(TYPE_OF_NULLPTR, const TFunction<FuncType>& Func)
-{
-	return !Func;
-}
-
-/**
- * Nullptr equality operator.
- */
-template <typename FuncType>
-FORCEINLINE bool operator==(const TFunction<FuncType>& Func, TYPE_OF_NULLPTR)
-{
-	return !Func;
-}
-
-/**
- * Nullptr inequality operator.
- */
-template <typename FuncType>
-FORCEINLINE bool operator!=(TYPE_OF_NULLPTR, const TFunction<FuncType>& Func)
-{
-	return (bool)Func;
-}
-
-/**
- * Nullptr inequality operator.
- */
-template <typename FuncType>
-FORCEINLINE bool operator!=(const TFunction<FuncType>& Func, TYPE_OF_NULLPTR)
-{
-	return (bool)Func;
-}
diff --git a/Engine/Source/Runtime/Core/Public/Templates/PimplPtr.h b/Engine/Source/Runtime/Core/Public/Templates/PimplPtr.h
--- a/Engine/Source/Runtime/Core/Public/Templates/PimplPtr.h
+++ b/Engine/Source/Runtime/Core/Public/Templates/PimplPtr.h
@@ -52,7 +52,7 @@ template<typename T, EPimplPtrMode Mode = EPimplPtrMode::NoCopy> struct TPimplPt
 
 namespace UE::Core::Private::PimplPtr
 {
-	constexpr SIZE_T RequiredAlignment = 16;
+	inline constexpr SIZE_T RequiredAlignment = 16;
 
 	template <typename T>
 	struct TPimplHeapObjectImpl;
@@ -235,6 +235,11 @@ public:
 
 private:
 	T* Ptr = nullptr;
+
+	template <EPimplPtrMode Mode> friend FORCEINLINE bool operator==(const TPimplPtr<T, Mode>& PimplPtr, TYPE_OF_NULLPTR) { return !PimplPtr.IsValid(); }
+	template <EPimplPtrMode Mode> friend FORCEINLINE bool operator==(TYPE_OF_NULLPTR, const TPimplPtr<T, Mode>& PimplPtr) { return !PimplPtr.IsValid(); }
+	template <EPimplPtrMode Mode> friend FORCEINLINE bool operator!=(const TPimplPtr<T, Mode>& PimplPtr, TYPE_OF_NULLPTR) { return  PimplPtr.IsValid(); }
+	template <EPimplPtrMode Mode> friend FORCEINLINE bool operator!=(TYPE_OF_NULLPTR, const TPimplPtr<T, Mode>& PimplPtr) { return  PimplPtr.IsValid(); }
 };
 
 template <typename T>
@@ -301,12 +306,6 @@ public:
 };
 
 
-template <typename T, EPimplPtrMode Mode> FORCEINLINE bool operator==(const TPimplPtr<T, Mode>& Ptr, TYPE_OF_NULLPTR) { return !Ptr.IsValid(); }
-template <typename T, EPimplPtrMode Mode> FORCEINLINE bool operator==(TYPE_OF_NULLPTR, const TPimplPtr<T, Mode>& Ptr) { return !Ptr.IsValid(); }
-template <typename T, EPimplPtrMode Mode> FORCEINLINE bool operator!=(const TPimplPtr<T, Mode>& Ptr, TYPE_OF_NULLPTR) { return  Ptr.IsValid(); }
-template <typename T, EPimplPtrMode Mode> FORCEINLINE bool operator!=(TYPE_OF_NULLPTR, const TPimplPtr<T, Mode>& Ptr) { return  Ptr.IsValid(); }
-
-
 /**
  * Heap-allocates an instance of T with the given arguments and returns it as a TPimplPtr.
  *
diff --git a/Engine/Source/Runtime/Core/Public/Templates/RefCounting.h b/Engine/Source/Runtime/Core/Public/Templates/RefCounting.h
--- a/Engine/Source/Runtime/Core/Public/Templates/RefCounting.h
+++ b/Engine/Source/Runtime/Core/Public/Templates/RefCounting.h
@@ -334,24 +334,21 @@ private:
 
 	template <typename OtherType>
 	friend class TRefCountPtr;
-};
 
-ALIAS_TEMPLATE_TYPE_LAYOUT(template<typename T>, TRefCountPtr<T>, void*);
+	friend FORCEINLINE bool operator==(const TRefCountPtr& A, const TRefCountPtr& B)
+	{
+		return A.GetReference() == B.GetReference();
+	}
 
-template<typename ReferencedType>
-FORCEINLINE bool operator==(const TRefCountPtr<ReferencedType>& A, const TRefCountPtr<ReferencedType>& B)
-{
-	return A.GetReference() == B.GetReference();
-}
+	friend FORCEINLINE bool operator==(const TRefCountPtr& A, ReferencedType* B)
+	{
+		return A.GetReference() == B;
+	}
 
-template<typename ReferencedType>
-FORCEINLINE bool operator==(const TRefCountPtr<ReferencedType>& A, ReferencedType* B)
-{
-	return A.GetReference() == B;
-}
+	friend FORCEINLINE bool operator==(ReferencedType* A, const TRefCountPtr& B)
+	{
+		return A == B.GetReference();
+	}
+};
 
-template<typename ReferencedType>
-FORCEINLINE bool operator==(ReferencedType* A, const TRefCountPtr<ReferencedType>& B)
-{
-	return A == B.GetReference();
-}
+ALIAS_TEMPLATE_TYPE_LAYOUT(template<typename T>, TRefCountPtr<T>, void*);
diff --git a/Engine/Source/Runtime/Core/Public/Templates/SharedPointer.h b/Engine/Source/Runtime/Core/Public/Templates/SharedPointer.h
--- a/Engine/Source/Runtime/Core/Public/Templates/SharedPointer.h
+++ b/Engine/Source/Runtime/Core/Public/Templates/SharedPointer.h
@@ -633,6 +633,33 @@ private:
 
 		Init(InObject);
 	}
+
+public:
+
+	/**
+	 * Global equality operator for TSharedRef
+	 *
+	 * @return  True if the two shared references are equal
+	 */
+	template< class ObjectTypeB >
+	[[nodiscard]] FORCEINLINE bool operator==( TSharedRef< ObjectTypeB, InMode > const& InSharedRefB ) const
+	{
+		return &( Get() ) == &( InSharedRefB.Get() );
+	}
+
+
+	/**
+	 * Global inequality operator for TSharedRef
+	 *
+	 * @return  True if the two shared references are not equal
+	 */
+	template< class ObjectTypeB >
+	[[nodiscard]] FORCEINLINE bool operator!=( TSharedRef< ObjectTypeB, InMode > const& InSharedRefB ) const
+	{
+		return &( Get() ) != &( InSharedRefB.Get() );
+	}
+
+
 };
 
 
@@ -1196,6 +1223,123 @@ private:
 	/** Interface to the reference counter for this object.  Note that the actual reference
 		controller object is shared by all shared and weak pointers that refer to the object */
 	SharedPointerInternals::FSharedReferencer< Mode > SharedReferenceCount;
+
+public:
+
+	/**
+	 * Global equality operator for TSharedPtr
+	 *
+	 * @return  True if the two shared pointers are equal
+	 */
+	template< class ObjectTypeB >
+	[[nodiscard]] FORCEINLINE bool operator==( TSharedPtr< ObjectTypeB, InMode > const& InSharedPtrB ) const
+	{
+		return Get() == InSharedPtrB.Get();
+	}
+
+
+	/**
+	 * Global inequality operator for TSharedPtr
+	 *
+	 * @return  True if the two shared pointers are not equal
+	 */
+	template< class ObjectTypeB >
+	[[nodiscard]] FORCEINLINE bool operator!=( TSharedPtr< ObjectTypeB, InMode > const& InSharedPtrB ) const
+	{
+		return Get() != InSharedPtrB.Get();
+	}
+
+
+	/**
+	 * Global equality operator for TSharedPtr
+	 *
+	 * @return  True if the shared pointer is null
+	 */
+	[[nodiscard]] friend FORCEINLINE bool operator==(TSharedPtr const& InSharedPtr, TYPE_OF_NULLPTR )
+	{
+		return !InSharedPtr.IsValid();
+	}
+
+
+	/**
+	 * Global inequality operator for TSharedPtr
+	 *
+	 * @return  True if the shared pointer is not null
+	 */
+	[[nodiscard]] friend FORCEINLINE bool operator!=( TSharedPtr const& InSharedPtr, TYPE_OF_NULLPTR )
+	{
+		return InSharedPtr.IsValid();
+	}
+
+
+	/**
+	 * Tests to see if a TSharedRef is "equal" to a TSharedPtr (both are valid and refer to the same object) (reverse)
+	 *
+	 * @return  True if the shared reference and shared pointer are "equal"
+	 */
+	template< class ObjectTypeB >
+	[[nodiscard]] friend FORCEINLINE bool operator==(TSharedPtr const& InSharedPtr, TSharedRef< ObjectTypeB, InMode > const& InSharedRef )
+	{
+		return InSharedPtr.IsValid() && InSharedPtr.Get() == &( InSharedRef.Get() );
+	}
+
+
+	/**
+	 * Tests to see if a TSharedRef is not "equal" to a TSharedPtr (shared pointer is invalid, or both refer to different objects) (reverse)
+	 *
+	 * @return  True if the shared reference and shared pointer are not "equal"
+	 */
+	template< class ObjectTypeB >
+	[[nodiscard]] friend FORCEINLINE bool operator!=(TSharedPtr const& InSharedPtr, TSharedRef< ObjectTypeB, InMode > const& InSharedRef )
+	{
+		return !InSharedPtr.IsValid() || ( InSharedPtr.Get() != &( InSharedRef.Get() ) );
+	}
+
+
+
+	/**
+	 * Global equality operator for TSharedPtr
+	 *
+	 * @return  True if the shared pointer is null
+	 */
+	[[nodiscard]] friend FORCEINLINE bool operator==( TYPE_OF_NULLPTR, TSharedPtr const& InSharedPtrB )
+	{
+		return !InSharedPtrB.IsValid();
+	}
+
+
+	/**
+	 * Global inequality operator for TSharedPtr
+	 *
+	 * @return  True if the shared pointer is not null
+	 */
+	[[nodiscard]] friend FORCEINLINE bool operator!=( TYPE_OF_NULLPTR, TSharedPtr const& InSharedPtrB )
+	{
+		return InSharedPtrB.IsValid();
+	}
+
+	/**
+	 * Tests to see if a TSharedRef is "equal" to a TSharedPtr (both are valid and refer to the same object)
+	 *
+	 * @return  True if the shared reference and shared pointer are "equal"
+	 */
+	template< class ObjectTypeA >
+	[[nodiscard]] friend FORCEINLINE bool operator==( TSharedRef< ObjectTypeA, InMode > const& InSharedRef, TSharedPtr const& InSharedPtr)
+	{
+		return InSharedPtr.IsValid() && InSharedPtr.Get() == &( InSharedRef.Get() );
+	}
+
+
+	/**
+	 * Tests to see if a TSharedRef is not "equal" to a TSharedPtr (shared pointer is invalid, or both refer to different objects)
+	 *
+	 * @return  True if the shared reference and shared pointer are not "equal"
+	 */
+	template< class ObjectTypeA >
+	[[nodiscard]] friend FORCEINLINE bool operator!=( TSharedRef< ObjectTypeA, InMode > const& InSharedRef, TSharedPtr const& InSharedPtr )
+	{
+		return !InSharedPtr.IsValid() || ( InSharedPtr.Get() != &( InSharedRef.Get() ) );
+	}
 };
 
 namespace Freeze
@@ -1510,6 +1654,171 @@ private:
 	/** Interface to the reference counter for this object.  Note that the actual reference
 		controller object is shared by all shared and weak pointers that refer to the object */
 	SharedPointerInternals::FWeakReferencer< Mode > WeakReferenceCount;
+
+
+public:
+	/**
+	 * Global equality operator for TWeakPtr
+	 *
+	 * @return  True if the two weak pointers are equal
+	 */
+	template< class ObjectTypeB >
+	[[nodiscard]] FORCEINLINE bool operator==( TWeakPtr< ObjectTypeB, InMode > const& InWeakPtrB ) const
+	{
+		return Pin().Get() == InWeakPtrB.Pin().Get();
+	}
+
+
+	/**
+	 * Global inequality operator for TWeakPtr
+	 *
+	 * @return  True if the two weak pointers are not equal
+	 */
+	template< class ObjectTypeB>
+	[[nodiscard]] FORCEINLINE bool operator!=( TWeakPtr< ObjectTypeB, InMode > const& InWeakPtrB ) const
+	{
+		return Pin().Get() != InWeakPtrB.Pin().Get();
+	}
+
+
+	/**
+	 * Global equality operator for TWeakPtr
+	 *
+	 * @return  True if the weak pointer and the shared ref are equal
+	 */
+	template< class ObjectTypeB >
+	[[nodiscard]] friend FORCEINLINE bool operator==( TWeakPtr const& InWeakPtr, TSharedRef< ObjectTypeB, InMode > const& InSharedRefB )
+	{
+		return InWeakPtr.Pin().Get() == &InSharedRefB.Get();
+	}
+
+
+	/**
+	 * Global equality operator for TWeakPtr
+	 *
+	 * @return  True if the weak pointer and the shared ptr are equal
+	 */
+	template< class ObjectTypeB >
+	[[nodiscard]] friend FORCEINLINE bool operator==( TWeakPtr const& InWeakPtr, TSharedPtr< ObjectTypeB, InMode > const& InSharedPtrB )
+	{
+		return InWeakPtr.Pin().Get() == InSharedPtrB.Get();
+	}
+
+
+	/**
+	 * Global equality operator for TWeakPtr
+	 *
+	 * @return  True if the weak pointer and the shared ref are equal
+	 */
+	template< class ObjectTypeA >
+	[[nodiscard]] friend FORCEINLINE bool operator==( TSharedRef< ObjectTypeA, InMode > const& InSharedRefA, TWeakPtr const& InWeakPtrB )
+	{
+		return &InSharedRefA.Get() == InWeakPtrB.Pin().Get();
+	}
+
+
+	/**
+	 * Global equality operator for TWeakPtr
+	 *
+	 * @return  True if the weak pointer and the shared ptr are equal
+	 */
+	template< class ObjectTypeA >
+	[[nodiscard]] friend FORCEINLINE bool operator==( TSharedPtr< ObjectTypeA, InMode > const& InSharedPtrA, TWeakPtr const& InWeakPtrB )
+	{
+		return InSharedPtrA.Get() == InWeakPtrB.Pin().Get();
+	}
+
+
+	/**
+	 * Global equality operator for TWeakPtr
+	 *
+	 * @return  True if the weak pointer is null
+	 */
+	[[nodiscard]] friend FORCEINLINE bool operator==( TWeakPtr const& InWeakPtr, TYPE_OF_NULLPTR )
+	{
+		return !InWeakPtr.IsValid();
+	}
+
+
+	/**
+	 * Global equality operator for TWeakPtr
+	 *
+	 * @return  True if the weak pointer is null
+	 */
+	[[nodiscard]] friend FORCEINLINE bool operator==( TYPE_OF_NULLPTR, TWeakPtr const& InWeakPtr )
+	{
+		return !InWeakPtr.IsValid();
+	}
+
+
+	/**
+	 * Global equality operator for TWeakPtr
+	 *
+	 * @return  True if the weak pointer and the shared ref are not equal
+	 */
+	template< class ObjectTypeB>
+	[[nodiscard]] friend FORCEINLINE bool operator!=( TWeakPtr const& InWeakPtr, TSharedRef< ObjectTypeB, InMode > const& InSharedRefB )
+	{
+		return InWeakPtr.Pin().Get() != &InSharedRefB.Get();
+	}
+
+
+	/**
+	 * Global equality operator for TWeakPtr
+	 *
+	 * @return  True if the weak pointer and the shared ptr are not equal
+	 */
+	template< class ObjectTypeB>
+	[[nodiscard]] friend FORCEINLINE bool operator!=( TWeakPtr const& InWeakPtr, TSharedPtr< ObjectTypeB, InMode > const& InSharedPtrB )
+	{
+		return InWeakPtr.Pin().Get() != InSharedPtrB.Get();
+	}
+
+
+	/**
+	 * Global equality operator for TWeakPtr
+	 *
+	 * @return  True if the weak pointer and the shared ref are not equal
+	 */
+	template< class ObjectTypeA>
+	[[nodiscard]] friend FORCEINLINE bool operator!=( TSharedRef< ObjectTypeA, InMode > const& InSharedRefA, TWeakPtr const& InWeakPtrB )
+	{
+		return &InSharedRefA.Get() != InWeakPtrB.Pin().Get();
+	}
+
+
+	/**
+	 * Global equality operator for TWeakPtr
+	 *
+	 * @return  True if the weak pointer and the shared ptr are not equal
+	 */
+	template< class ObjectTypeA >
+	[[nodiscard]] friend FORCEINLINE bool operator!=( TSharedPtr< ObjectTypeA, InMode > const& InSharedPtrA, TWeakPtr const& InWeakPtrB )
+	{
+		return InSharedPtrA.Get() != InWeakPtrB.Pin().Get();
+	}
+
+
+	/**
+	 * Global inequality operator for TWeakPtr
+	 *
+	 * @return  True if the weak pointer is not null
+	 */
+	[[nodiscard]] friend FORCEINLINE bool operator!=( TWeakPtr const& InWeakPtr, TYPE_OF_NULLPTR )
+	{
+		return InWeakPtr.IsValid();
+	}
+
+
+	/**
+	 * Global inequality operator for TWeakPtr
+	 *
+	 * @return  True if the weak pointer is not null
+	 */
+	[[nodiscard]] friend FORCEINLINE bool operator!=( TYPE_OF_NULLPTR, TWeakPtr const& InWeakPtr )
+	{
+		return InWeakPtr.IsValid();
+	}
 };
 
 
@@ -1721,318 +2030,6 @@ private:
 };
 
 
-/**
- * Global equality operator for TSharedRef
- *
- * @return  True if the two shared references are equal
- */
-template< class ObjectTypeA, class ObjectTypeB, ESPMode Mode >
-[[nodiscard]] FORCEINLINE bool operator==( TSharedRef< ObjectTypeA, Mode > const& InSharedRefA, TSharedRef< ObjectTypeB, Mode > const& InSharedRefB )
-{
-	return &( InSharedRefA.Get() ) == &( InSharedRefB.Get() );
-}
-
-
-/**
- * Global inequality operator for TSharedRef
- *
- * @return  True if the two shared references are not equal
- */
-template< class ObjectTypeA, class ObjectTypeB, ESPMode Mode >
-[[nodiscard]] FORCEINLINE bool operator!=( TSharedRef< ObjectTypeA, Mode > const& InSharedRefA, TSharedRef< ObjectTypeB, Mode > const& InSharedRefB )
-{
-	return &( InSharedRefA.Get() ) != &( InSharedRefB.Get() );
-}
-
-
-/**
- * Global equality operator for TSharedPtr
- *
- * @return  True if the two shared pointers are equal
- */
-template< class ObjectTypeA, class ObjectTypeB, ESPMode Mode >
-[[nodiscard]] FORCEINLINE bool operator==( TSharedPtr< ObjectTypeA, Mode > const& InSharedPtrA, TSharedPtr< ObjectTypeB, Mode > const& InSharedPtrB )
-{
-	return InSharedPtrA.Get() == InSharedPtrB.Get();
-}
-
-
-/**
- * Global equality operator for TSharedPtr
- *
- * @return  True if the shared pointer is null
- */
-template< class ObjectTypeA, ESPMode Mode >
-[[nodiscard]] FORCEINLINE bool operator==( TSharedPtr< ObjectTypeA, Mode > const& InSharedPtrA, TYPE_OF_NULLPTR )
-{
-	return !InSharedPtrA.IsValid();
-}
-
-
-/**
- * Global equality operator for TSharedPtr
- *
- * @return  True if the shared pointer is null
- */
-template< class ObjectTypeB, ESPMode Mode >
-[[nodiscard]] FORCEINLINE bool operator==( TYPE_OF_NULLPTR, TSharedPtr< ObjectTypeB, Mode > const& InSharedPtrB )
-{
-	return !InSharedPtrB.IsValid();
-}
-
-
-/**
- * Global inequality operator for TSharedPtr
- *
- * @return  True if the two shared pointers are not equal
- */
-template< class ObjectTypeA, class ObjectTypeB, ESPMode Mode >
-[[nodiscard]] FORCEINLINE bool operator!=( TSharedPtr< ObjectTypeA, Mode > const& InSharedPtrA, TSharedPtr< ObjectTypeB, Mode > const& InSharedPtrB )
-{
-	return InSharedPtrA.Get() != InSharedPtrB.Get();
-}
-
-
-/**
- * Global inequality operator for TSharedPtr
- *
- * @return  True if the shared pointer is not null
- */
-template< class ObjectTypeA, ESPMode Mode >
-[[nodiscard]] FORCEINLINE bool operator!=( TSharedPtr< ObjectTypeA, Mode > const& InSharedPtrA, TYPE_OF_NULLPTR )
-{
-	return InSharedPtrA.IsValid();
-}
-
-
-/**
- * Global inequality operator for TSharedPtr
- *
- * @return  True if the shared pointer is not null
- */
-template< class ObjectTypeB, ESPMode Mode >
-[[nodiscard]] FORCEINLINE bool operator!=( TYPE_OF_NULLPTR, TSharedPtr< ObjectTypeB, Mode > const& InSharedPtrB )
-{
-	return InSharedPtrB.IsValid();
-}
-
-
-/**
- * Tests to see if a TSharedRef is "equal" to a TSharedPtr (both are valid and refer to the same object)
- *
- * @return  True if the shared reference and shared pointer are "equal"
- */
-template< class ObjectTypeA, class ObjectTypeB, ESPMode Mode >
-[[nodiscard]] FORCEINLINE bool operator==( TSharedRef< ObjectTypeA, Mode > const& InSharedRef, TSharedPtr< ObjectTypeB, Mode > const& InSharedPtr )
-{
-	return InSharedPtr.IsValid() && InSharedPtr.Get() == &( InSharedRef.Get() );
-}
-
-
-/**
- * Tests to see if a TSharedRef is not "equal" to a TSharedPtr (shared pointer is invalid, or both refer to different objects)
- *
- * @return  True if the shared reference and shared pointer are not "equal"
- */
-template< class ObjectTypeA, class ObjectTypeB, ESPMode Mode >
-[[nodiscard]] FORCEINLINE bool operator!=( TSharedRef< ObjectTypeA, Mode > const& InSharedRef, TSharedPtr< ObjectTypeB, Mode > const& InSharedPtr )
-{
-	return !InSharedPtr.IsValid() || ( InSharedPtr.Get() != &( InSharedRef.Get() ) );
-}
-
-
-/**
- * Tests to see if a TSharedRef is "equal" to a TSharedPtr (both are valid and refer to the same object) (reverse)
- *
- * @return  True if the shared reference and shared pointer are "equal"
- */
-template< class ObjectTypeA, class ObjectTypeB, ESPMode Mode >
-[[nodiscard]] FORCEINLINE bool operator==( TSharedPtr< ObjectTypeB, Mode > const& InSharedPtr, TSharedRef< ObjectTypeA, Mode > const& InSharedRef )
-{
-	return InSharedRef == InSharedPtr;
-}
-
-
-/**
- * Tests to see if a TSharedRef is not "equal" to a TSharedPtr (shared pointer is invalid, or both refer to different objects) (reverse)
- *
- * @return  True if the shared reference and shared pointer are not "equal"
- */
-template< class ObjectTypeA, class ObjectTypeB, ESPMode Mode >
-[[nodiscard]] FORCEINLINE bool operator!=( TSharedPtr< ObjectTypeB, Mode > const& InSharedPtr, TSharedRef< ObjectTypeA, Mode > const& InSharedRef )
-{
-	return InSharedRef != InSharedPtr;
-}
-
-
-/**
- * Global equality operator for TWeakPtr
- *
- * @return  True if the two weak pointers are equal
- */
-template< class ObjectTypeA, class ObjectTypeB, ESPMode Mode >
-[[nodiscard]] FORCEINLINE bool operator==( TWeakPtr< ObjectTypeA, Mode > const& InWeakPtrA, TWeakPtr< ObjectTypeB, Mode > const& InWeakPtrB )
-{
-	return InWeakPtrA.Pin().Get() == InWeakPtrB.Pin().Get();
-}
-
-
-/**
- * Global equality operator for TWeakPtr
- *
- * @return  True if the weak pointer and the shared ref are equal
- */
-template< class ObjectTypeA, class ObjectTypeB, ESPMode Mode >
-[[nodiscard]] FORCEINLINE bool operator==( TWeakPtr< ObjectTypeA, Mode > const& InWeakPtrA, TSharedRef< ObjectTypeB, Mode > const& InSharedRefB )
-{
-	return InWeakPtrA.Pin().Get() == &InSharedRefB.Get();
-}
-
-
-/**
- * Global equality operator for TWeakPtr
- *
- * @return  True if the weak pointer and the shared ptr are equal
- */
-template< class ObjectTypeA, class ObjectTypeB, ESPMode Mode >
-[[nodiscard]] FORCEINLINE bool operator==( TWeakPtr< ObjectTypeA, Mode > const& InWeakPtrA, TSharedPtr< ObjectTypeB, Mode > const& InSharedPtrB )
-{
-	return InWeakPtrA.Pin().Get() == InSharedPtrB.Get();
-}
-
-
-/**
- * Global equality operator for TWeakPtr
- *
- * @return  True if the weak pointer and the shared ref are equal
- */
-template< class ObjectTypeA, class ObjectTypeB, ESPMode Mode >
-[[nodiscard]] FORCEINLINE bool operator==( TSharedRef< ObjectTypeA, Mode > const& InSharedRefA, TWeakPtr< ObjectTypeB, Mode > const& InWeakPtrB )
-{
-	return &InSharedRefA.Get() == InWeakPtrB.Pin().Get();
-}
-
-
-/**
- * Global equality operator for TWeakPtr
- *
- * @return  True if the weak pointer and the shared ptr are equal
- */
-template< class ObjectTypeA, class ObjectTypeB, ESPMode Mode >
-[[nodiscard]] FORCEINLINE bool operator==( TSharedPtr< ObjectTypeA, Mode > const& InSharedPtrA, TWeakPtr< ObjectTypeB, Mode > const& InWeakPtrB )
-{
-	return InSharedPtrA.Get() == InWeakPtrB.Pin().Get();
-}
-
-
-/**
- * Global equality operator for TWeakPtr
- *
- * @return  True if the weak pointer is null
- */
-template< class ObjectTypeA, ESPMode Mode >
-[[nodiscard]] FORCEINLINE bool operator==( TWeakPtr< ObjectTypeA, Mode > const& InWeakPtrA, TYPE_OF_NULLPTR )
-{
-	return !InWeakPtrA.IsValid();
-}
-
-
-/**
- * Global equality operator for TWeakPtr
- *
- * @return  True if the weak pointer is null
- */
-template< class ObjectTypeB, ESPMode Mode >
-[[nodiscard]] FORCEINLINE bool operator==( TYPE_OF_NULLPTR, TWeakPtr< ObjectTypeB, Mode > const& InWeakPtrB )
-{
-	return !InWeakPtrB.IsValid();
-}
-
-
-/**
- * Global inequality operator for TWeakPtr
- *
- * @return  True if the two weak pointers are not equal
- */
-template< class ObjectTypeA, class ObjectTypeB, ESPMode Mode >
-[[nodiscard]] FORCEINLINE bool operator!=( TWeakPtr< ObjectTypeA, Mode > const& InWeakPtrA, TWeakPtr< ObjectTypeB, Mode > const& InWeakPtrB )
-{
-	return InWeakPtrA.Pin().Get() != InWeakPtrB.Pin().Get();
-}
-
-
-/**
- * Global equality operator for TWeakPtr
- *
- * @return  True if the weak pointer and the shared ref are not equal
- */
-template< class ObjectTypeA, class ObjectTypeB, ESPMode Mode >
-[[nodiscard]] FORCEINLINE bool operator!=( TWeakPtr< ObjectTypeA, Mode > const& InWeakPtrA, TSharedRef< ObjectTypeB, Mode > const& InSharedRefB )
-{
-	return InWeakPtrA.Pin().Get() != &InSharedRefB.Get();
-}
-
-
-/**
- * Global equality operator for TWeakPtr
- *
- * @return  True if the weak pointer and the shared ptr are not equal
- */
-template< class ObjectTypeA, class ObjectTypeB, ESPMode Mode >
-[[nodiscard]] FORCEINLINE bool operator!=( TWeakPtr< ObjectTypeA, Mode > const& InWeakPtrA, TSharedPtr< ObjectTypeB, Mode > const& InSharedPtrB )
-{
-	return InWeakPtrA.Pin().Get() != InSharedPtrB.Get();
-}
-
-
-/**
- * Global equality operator for TWeakPtr
- *
- * @return  True if the weak pointer and the shared ref are not equal
- */
-template< class ObjectTypeA, class ObjectTypeB, ESPMode Mode >
-[[nodiscard]] FORCEINLINE bool operator!=( TSharedRef< ObjectTypeA, Mode > const& InSharedRefA, TWeakPtr< ObjectTypeB, Mode > const& InWeakPtrB )
-{
-	return &InSharedRefA.Get() != InWeakPtrB.Pin().Get();
-}
-
-
-/**
- * Global equality operator for TWeakPtr
- *
- * @return  True if the weak pointer and the shared ptr are not equal
- */
-template< class ObjectTypeA, class ObjectTypeB, ESPMode Mode >
-[[nodiscard]] FORCEINLINE bool operator!=( TSharedPtr< ObjectTypeA, Mode > const& InSharedPtrA, TWeakPtr< ObjectTypeB, Mode > const& InWeakPtrB )
-{
-	return InSharedPtrA.Get() != InWeakPtrB.Pin().Get();
-}
-
-
-/**
- * Global inequality operator for TWeakPtr
- *
- * @return  True if the weak pointer is not null
- */
-template< class ObjectTypeA, ESPMode Mode >
-[[nodiscard]] FORCEINLINE bool operator!=( TWeakPtr< ObjectTypeA, Mode > const& InWeakPtrA, TYPE_OF_NULLPTR )
-{
-	return InWeakPtrA.IsValid();
-}
-
-
-/**
- * Global inequality operator for TWeakPtr
- *
- * @return  True if the weak pointer is not null
- */
-template< class ObjectTypeB, ESPMode Mode >
-[[nodiscard]] FORCEINLINE bool operator!=( TYPE_OF_NULLPTR, TWeakPtr< ObjectTypeB, Mode > const& InWeakPtrB )
-{
-	return InWeakPtrB.IsValid();
-}
-
-
 /**
  * Casts a shared pointer of one type to another type. (static_cast)  Useful for down-casting.
  *
diff --git a/Engine/Source/Runtime/Core/Public/Templates/Tuple.h b/Engine/Source/Runtime/Core/Public/Templates/Tuple.h
--- a/Engine/Source/Runtime/Core/Public/Templates/Tuple.h
+++ b/Engine/Source/Runtime/Core/Public/Templates/Tuple.h
@@ -764,18 +764,12 @@ public:
 	template <int N> friend decltype(auto) get(      volatile TTuple&& val) { return static_cast<      volatile TTuple&&>(val).template Get<N>(); }
 	template <int N> friend decltype(auto) get(const volatile TTuple&& val) { return static_cast<const volatile TTuple&&>(val).template Get<N>(); }
 #endif
-};
-
-template <typename... Types>
-FORCEINLINE uint32 GetTypeHash(const TTuple<Types...>& Tuple)
-{
-	return UE::Core::Private::Tuple::TGetTupleHashHelper<1u, sizeof...(Types)>::Do(GetTypeHash(Tuple.template Get<0>()), Tuple);
-}
 
-FORCEINLINE uint32 GetTypeHash(const TTuple<>& Tuple)
-{
-	return 0;
-}
+	FORCEINLINE friend uint32 GetTypeHash(const TTuple<Types...>& Tuple)
+	{
+		return UE::Core::Private::Tuple::TGetTupleHashHelper<1u, sizeof...(Types)>::Do(GetTypeHash(Tuple.template Get<0>()), Tuple);
+	}
+};
 
 namespace Freeze
 {
diff --git a/Engine/Source/Runtime/Core/Public/Templates/TypeHash.h b/Engine/Source/Runtime/Core/Public/Templates/TypeHash.h
--- a/Engine/Source/Runtime/Core/Public/Templates/TypeHash.h
+++ b/Engine/Source/Runtime/Core/Public/Templates/TypeHash.h
@@ -7,6 +7,7 @@
 #include "Misc/Crc.h"
 
 #include <stdint.h>
+#include <type_traits>
 
 
 namespace UE
@@ -88,94 +89,69 @@ inline uint32 PointerHash(const void* Key, uint32 C)
 //           expect them to be consistent across multiple runs.
 //
 
-inline uint32 GetTypeHash( const uint8 A )
+template <
+	typename ScalarType,
+	std::enable_if_t<std::is_scalar_v<ScalarType>>* = nullptr
+>
+inline uint32 GetTypeHash(ScalarType Value)
 {
-	return A;
-}
-
-inline uint32 GetTypeHash( const int8 A )
-{
-	return A;
-}
-
-inline uint32 GetTypeHash( const uint16 A )
-{
-	return A;
-}
-
-inline uint32 GetTypeHash( const int16 A )
-{
-	return A;
-}
-
-inline uint32 GetTypeHash( const int32 A )
-{
-	return A;
-}
-
-inline uint32 GetTypeHash( const uint32 A )
-{
-	return A;
-}
-
-inline uint32 GetTypeHash( const uint64 A )
-{
-	return (uint32)A+((uint32)(A>>32) * 23);
-}
-
-inline uint32 GetTypeHash( const int64 A )
-{
-	return (uint32)A+((uint32)(A>>32) * 23);
-}
-
-#if PLATFORM_MAC
-inline uint32 GetTypeHash( const __uint128_t A )
-{
-	uint64 Low = (uint64)A;
-	uint64 High = (uint64)(A >> 64);
-	return GetTypeHash(Low) ^ GetTypeHash(High);
-}
-#endif
-
-#if defined(INT64_T_TYPES_NOT_LONG_LONG)
-// int64_t and uint64_t are long types, not long long (aka int64/uint64). These types can't be automatically converted.
-inline uint32 GetTypeHash(uint64_t A)
-{
-	return GetTypeHash((uint64)A);
-}
-inline uint32 GetTypeHash(int64_t A)
-{
-	return GetTypeHash((int64)A);
-}
-#endif
-
-inline uint32 GetTypeHash( float Value )
-{
-	return *(uint32*)&Value;
-}
-
-inline uint32 GetTypeHash( double Value )
-{
-	return GetTypeHash(*(uint64*)&Value);
-}
-
-inline uint32 GetTypeHash( const TCHAR* S )
-{
-	return FCrc::Strihash_DEPRECATED(S);
-}
-
-inline uint32 GetTypeHash( const void* A )
-{
-	return PointerHash(A);
-}
-
-inline uint32 GetTypeHash( void* A )
-{
-	return PointerHash(A);
+	if constexpr (std::is_integral_v<ScalarType>)
+	{
+		if constexpr (sizeof(ScalarType) <= 4)
+		{
+			return Value;
+		}
+		else if constexpr (sizeof(ScalarType) == 8)
+		{
+			return (uint32)Value + ((uint32)(Value >> 32) * 23);
+		}
+		else if constexpr (sizeof(ScalarType) == 16)
+		{
+			const uint64 Low = (uint64)Value;
+			const uint64 High = (uint64)(Value >> 64);
+			return GetTypeHash(Low) ^ GetTypeHash(High);
+		}
+		else
+		{
+			static_assert(sizeof(ScalarType) == 0, "Unsupported integral type");
+		}
+	}
+	else if constexpr (std::is_floating_point_v<ScalarType>)
+	{
+		if constexpr (std::is_same_v<ScalarType, float>)
+		{
+			return *(uint32*)&Value;
+		}
+		else if constexpr (std::is_same_v<ScalarType, double>)
+		{
+			return GetTypeHash(*(uint64*)&Value);
+		}
+		else
+		{
+			static_assert(sizeof(ScalarType) == 0, "Unsupported floating point type");
+		}
+	}
+	else if constexpr (std::is_enum_v<ScalarType>)
+	{
+		return GetTypeHash((__underlying_type(ScalarType))Value);
+	}
+	else if constexpr (std::is_pointer_v<ScalarType>)
+	{
+		if constexpr (std::is_same_v<std::remove_cv_t<std::remove_pointer_t<ScalarType>>, TCHAR>)
+		{
+			return FCrc::Strihash_DEPRECATED(Value);
+		}
+		else
+		{
+			return PointerHash(Value);
+		}
+	}
+	else
+	{
+		static_assert(sizeof(ScalarType) == 0, "Unsupported scalar type");
+	}
 }
 
-template <typename EnumType>
-FORCEINLINE  typename TEnableIf<TIsEnum<EnumType>::Value, uint32>::Type GetTypeHash(EnumType E)
-{
-	return GetTypeHash((__underlying_type(EnumType))E);
-}
+// Use this when inside type that has GetTypeHash() (no in-parameters) implemented. It makes GetTypeHash dispatch in global namespace
+template <typename T>
+FORCEINLINE uint32 GetTypeHashHelper(const T& V) { return GetTypeHash(V); }
diff --git a/Engine/Source/Runtime/Core/Public/Templates/UniquePtr.h b/Engine/Source/Runtime/Core/Public/Templates/UniquePtr.h
--- a/Engine/Source/Runtime/Core/Public/Templates/UniquePtr.h
+++ b/Engine/Source/Runtime/Core/Public/Templates/UniquePtr.h
@@ -377,6 +377,66 @@ public:
 private:
 	using PtrType = T*;
 	LAYOUT_FIELD(PtrType, Ptr);
+
+	/**
+	 * Equality comparison operator
+	 *
+	 * @param Lhs The first TUniquePtr to compare.
+	 * @param Rhs The second TUniquePtr to compare.
+	 *
+	 * @return true if the two TUniquePtrs are logically substitutable for each other, false otherwise.
+	 */
+	template <typename RhsT>
+	friend FORCEINLINE bool operator==(const TUniquePtr& Lhs, const TUniquePtr<RhsT>& Rhs)
+	{
+		return Lhs.Get() == Rhs.Get();
+	}
+
+	/**
+	 * Inequality comparison operator
+	 *
+	 * @param Lhs The first TUniquePtr to compare.
+	 * @param Rhs The second TUniquePtr to compare.
+	 *
+	 * @return false if the two TUniquePtrs are logically substitutable for each other, true otherwise.
+	 */
+	template <typename RhsT>
+	friend FORCEINLINE bool operator!=(const TUniquePtr& Lhs, const TUniquePtr<RhsT>& Rhs)
+	{
+		return Lhs.Get() != Rhs.Get();
+	}
+
+	/**
+	 * Equality comparison operator against nullptr.
+	 *
+	 * @param Lhs The TUniquePtr to compare.
+	 *
+	 * @return true if the TUniquePtr is null, false otherwise.
+	 */
+	friend FORCEINLINE bool operator==(const TUniquePtr& Lhs, TYPE_OF_NULLPTR)
+	{
+		return !Lhs.IsValid();
+	}
+	friend FORCEINLINE bool operator==(TYPE_OF_NULLPTR, const TUniquePtr& Rhs)
+	{
+		return !Rhs.IsValid();
+	}
+
+	/**
+	 * Inequality comparison operator against nullptr.
+	 *
+	 * @param Rhs The TUniquePtr to compare.
+	 *
+	 * @return true if the TUniquePtr is not null, false otherwise.
+	 */
+	friend FORCEINLINE bool operator!=(const TUniquePtr& Lhs, TYPE_OF_NULLPTR)
+	{
+		return Lhs.IsValid();
+	}
+	friend FORCEINLINE bool operator!=(TYPE_OF_NULLPTR, const TUniquePtr& Rhs)
+	{
+		return Rhs.IsValid();
+	}
 };
 
 template <typename T, typename Deleter>
@@ -650,81 +710,67 @@ public:
 
 private:
 	T* Ptr;
-};
 
-/**
- * Equality comparison operator
- *
- * @param Lhs The first TUniquePtr to compare.
- * @param Rhs The second TUniquePtr to compare.
- *
- * @return true if the two TUniquePtrs are logically substitutable for each other, false otherwise.
- */
-template <typename LhsT, typename RhsT>
-FORCEINLINE bool operator==(const TUniquePtr<LhsT>& Lhs, const TUniquePtr<RhsT>& Rhs)
-{
-	return Lhs.Get() == Rhs.Get();
-}
-template <typename T>
-FORCEINLINE bool operator==(const TUniquePtr<T>& Lhs, const TUniquePtr<T>& Rhs)
-{
-	return Lhs.Get() == Rhs.Get();
-}
+	/**
+	 * Equality comparison operator
+	 *
+	 * @param Lhs The first TUniquePtr to compare.
+	 * @param Rhs The second TUniquePtr to compare.
+	 *
+	 * @return true if the two TUniquePtrs are logically substitutable for each other, false otherwise.
+	 */
+	template <typename RhsT>
+	friend FORCEINLINE bool operator==(const TUniquePtr& Lhs, const TUniquePtr<RhsT>& Rhs)
+	{
+		return Lhs.Get() == Rhs.Get();
+	}
 
-/**
- * Inequality comparison operator
- *
- * @param Lhs The first TUniquePtr to compare.
- * @param Rhs The second TUniquePtr to compare.
- *
- * @return false if the two TUniquePtrs are logically substitutable for each other, true otherwise.
- */
-template <typename LhsT, typename RhsT>
-FORCEINLINE bool operator!=(const TUniquePtr<LhsT>& Lhs, const TUniquePtr<RhsT>& Rhs)
-{
-	return Lhs.Get() != Rhs.Get();
-}
-template <typename T>
-FORCEINLINE bool operator!=(const TUniquePtr<T>& Lhs, const TUniquePtr<T>& Rhs)
-{
-	return Lhs.Get() != Rhs.Get();
-}
+	/**
+	 * Inequality comparison operator
+	 *
+	 * @param Lhs The first TUniquePtr to compare.
+	 * @param Rhs The second TUniquePtr to compare.
+	 *
+	 * @return false if the two TUniquePtrs are logically substitutable for each other, true otherwise.
+	 */
+	template <typename RhsT>
+	friend FORCEINLINE bool operator!=(const TUniquePtr& Lhs, const TUniquePtr<RhsT>& Rhs)
+	{
+		return Lhs.Get() != Rhs.Get();
+	}
 
-/**
- * Equality comparison operator against nullptr.
- *
- * @param Lhs The TUniquePtr to compare.
- *
- * @return true if the TUniquePtr is null, false otherwise.
- */
-template <typename T>
-FORCEINLINE bool operator==(const TUniquePtr<T>& Lhs, TYPE_OF_NULLPTR)
-{
-	return !Lhs.IsValid();
-}
-template <typename T>
-FORCEINLINE bool operator==(TYPE_OF_NULLPTR, const TUniquePtr<T>& Rhs)
-{
-	return !Rhs.IsValid();
-}
+	/**
+	 * Equality comparison operator against nullptr.
+	 *
+	 * @param Lhs The TUniquePtr to compare.
+	 *
+	 * @return true if the TUniquePtr is null, false otherwise.
+	 */
+	friend FORCEINLINE bool operator==(const TUniquePtr& Lhs, TYPE_OF_NULLPTR)
+	{
+		return !Lhs.IsValid();
+	}
+	friend FORCEINLINE bool operator==(TYPE_OF_NULLPTR, const TUniquePtr& Rhs)
+	{
+		return !Rhs.IsValid();
+	}
 
-/**
- * Inequality comparison operator against nullptr.
- *
- * @param Rhs The TUniquePtr to compare.
- *
- * @return true if the TUniquePtr is not null, false otherwise.
- */
-template <typename T>
-FORCEINLINE bool operator!=(const TUniquePtr<T>& Lhs, TYPE_OF_NULLPTR)
-{
-	return Lhs.IsValid();
-}
-template <typename T>
-FORCEINLINE bool operator!=(TYPE_OF_NULLPTR, const TUniquePtr<T>& Rhs)
-{
-	return Rhs.IsValid();
-}
+	/**
+	 * Inequality comparison operator against nullptr.
+	 *
+	 * @param Rhs The TUniquePtr to compare.
+	 *
+	 * @return true if the TUniquePtr is not null, false otherwise.
+	 */
+	friend FORCEINLINE bool operator!=(const TUniquePtr& Lhs, TYPE_OF_NULLPTR)
+	{
+		return Lhs.IsValid();
+	}
+	friend FORCEINLINE bool operator!=(TYPE_OF_NULLPTR, const TUniquePtr& Rhs)
+	{
+		return Rhs.IsValid();
+	}
+};
 
 // Trait which allows TUniquePtr to be default constructed by memsetting to zero.
 template <typename T>
diff --git a/Engine/Source/Runtime/Core/Public/UObject/NameTypes.h b/Engine/Source/Runtime/Core/Public/UObject/NameTypes.h
--- a/Engine/Source/Runtime/Core/Public/UObject/NameTypes.h
+++ b/Engine/Source/Runtime/Core/Public/UObject/NameTypes.h
@@ -108,12 +108,13 @@ private:
 
 	CORE_API static FNameEntryId FromValidEName(EName Ename);
 	CORE_API static FNameEntryId FromValidENamePostInit(EName Ename);
-};
 
-CORE_API uint32 GetTypeHash(FNameEntryId Id);
-inline bool operator==(EName Ename, FNameEntryId Id) { return Id == Ename; }
-inline bool operator!=(EName Ename, FNameEntryId Id) { return !(Id == Ename); }
-inline bool operator!=(FNameEntryId Id, EName Ename) { return !(Id == Ename); }
+public:
+	friend inline bool operator==(EName Ename, FNameEntryId Id) { return Id == Ename; }
+	friend inline bool operator!=(EName Ename, FNameEntryId Id) { return !(Id == Ename); }
+	friend inline bool operator!=(FNameEntryId Id, EName Ename) { return !(Id == Ename); }
+	friend CORE_API uint32 GetTypeHash(FNameEntryId Id);
+};
 
 /** Serialize as process specific unstable int */
 CORE_API FArchive& operator<<(FArchive& Ar, FNameEntryId& InId);
@@ -393,9 +394,6 @@ struct FNameEntrySerialized
  */
 struct FMinimalName
 {
-	friend uint32 GetTypeHash(FMinimalName Name);
-	friend bool operator==(FMinimalName Lhs, FMinimalName Rhs);
-	friend bool operator==(FName Lhs, FMinimalName Rhs);
 	friend FName;
 
 	FMinimalName() {}
@@ -416,6 +414,31 @@ private:
 	/** Number portion of the string/number pair (stored internally as 1 more than actual, so zero'd memory will be the default, no-instance case) */
 	int32			Number = NAME_NO_NUMBER_INTERNAL;
 #endif // UE_FNAME_OUTLINE_NUMBE
+
+#if UE_FNAME_OUTLINE_NUMBER
+	friend FORCEINLINE bool operator==(FMinimalName Lhs, FMinimalName Rhs)
+	{
+		return Lhs.Index == Rhs.Index;
+	}
+	friend FORCEINLINE uint32 GetTypeHash(FMinimalName Name)
+	{
+		return GetTypeHash(Name.Index);
+	}
+#else
+	friend FORCEINLINE bool operator==(FMinimalName Lhs, FMinimalName Rhs)
+	{
+		return Lhs.Index == Rhs.Index && Lhs.Number == Rhs.Number;
+	}
+	friend FORCEINLINE uint32 GetTypeHash(FMinimalName Name)
+	{
+		return GetTypeHash(Name.Index) + Name.Number;
+	}
+#endif
+	friend FORCEINLINE bool operator!=(FMinimalName Lhs, FMinimalName Rhs)
+	{
+		return !operator==(Lhs, Rhs);
+	}
+	friend FORCEINLINE bool operator==(FName Lhs, FMinimalName Rhs);
 };
 
 /**
@@ -427,9 +450,6 @@ private:
  */
 struct FScriptName
 {
-	friend uint32 GetTypeHash(FScriptName Name);
-	friend bool operator==(FScriptName Lhs, FScriptName Rhs);
-	friend bool operator==(FName Lhs, FScriptName Rhs);
 	friend FName;
 
 	FScriptName() {}
@@ -458,6 +478,32 @@ private:
 	/** Number portion of the string/number pair (stored internally as 1 more than actual, so zero'd memory will be the default, no-instance case) */
 	uint32			Number = NAME_NO_NUMBER_INTERNAL;
 #endif // UE_FNAME_OUTLINE_NUMBER
+
+
+#if UE_FNAME_OUTLINE_NUMBER
+	friend FORCEINLINE bool operator==(FScriptName Lhs, FScriptName Rhs)
+	{
+		return Lhs.ComparisonIndex == Rhs.ComparisonIndex;
+	}
+	friend FORCEINLINE uint32 GetTypeHash(FScriptName Name)
+	{
+		return GetTypeHash(Name.ComparisonIndex);
+	}
+#else
+	friend FORCEINLINE bool operator==(FScriptName Lhs, FScriptName Rhs)
+	{
+		return Lhs.ComparisonIndex == Rhs.ComparisonIndex && Lhs.Number == Rhs.Number;
+	}
+	friend FORCEINLINE uint32 GetTypeHash(FScriptName Name)
+	{
+		return GetTypeHash(Name.ComparisonIndex) + Name.Number;
+	}
+#endif
+	friend FORCEINLINE bool operator!=(FScriptName Lhs, FScriptName Rhs)
+	{
+		return !(Lhs == Rhs);
+	}
+	friend FORCEINLINE bool operator==(FName Lhs, FScriptName Rhs);
 };
 
 struct FMemoryImageName;
@@ -475,9 +521,6 @@ namespace Freeze
  */
 struct FMemoryImageName
 {
-	friend uint32 GetTypeHash(FMemoryImageName Name);
-	friend bool operator==(FMemoryImageName Lhs, FMemoryImageName Rhs);
-	friend bool operator==(FName Lhs, FMemoryImageName Rhs);
 	friend FName;
 
 	friend CORE_API void Freeze::ApplyMemoryImageNamePatch(void* NameDst, const FMemoryImageName& Name, const FPlatformTypeLayoutParameters& LayoutParams);
@@ -507,6 +550,31 @@ private:
 	/** Encoded address of name entry (used to find String portion of the string/number pair used for display) */
 	FNameEntryId	DisplayIndex;
 #endif
+
+#if UE_FNAME_OUTLINE_NUMBER
+	friend FORCEINLINE bool operator==(FMemoryImageName Lhs, FMemoryImageName Rhs)
+	{
+		return Lhs.ComparisonIndex == Rhs.ComparisonIndex;
+	}
+	friend FORCEINLINE uint32 GetTypeHash(FMemoryImageName Name)
+	{
+		return GetTypeHash(Name.ComparisonIndex);
+	}
+#else
+	friend FORCEINLINE bool operator==(FMemoryImageName Lhs, FMemoryImageName Rhs)
+	{
+		return Lhs.ComparisonIndex == Rhs.ComparisonIndex && Lhs.Number == Rhs.Number;
+	}
+	friend FORCEINLINE uint32 GetTypeHash(FMemoryImageName Name)
+	{
+		return GetTypeHash(Name.ComparisonIndex) + Name.Number;
+	}
+#endif
+	friend FORCEINLINE bool operator!=(FMemoryImageName Lhs, FMemoryImageName Rhs)
+	{
+		return !(Lhs == Rhs);
+	}
+	friend FORCEINLINE bool operator==(FName Lhs, FMemoryImageName Rhs);
 };
 
 /**
@@ -1122,15 +1190,10 @@ private:
 	friend struct FNameHelper;
 	friend FScriptName NameToScriptName(FName InName);
 	friend FMinimalName NameToMinimalName(FName InName);
-	friend uint32 GetTypeHash(FName Name);
 	friend FMinimalName::FMinimalName(const FName& Name);
 	friend FScriptName::FScriptName(const FName& Name);
 	friend FMemoryImageName::FMemoryImageName(const FName& Name);
 
-	friend bool operator==(FName Lhs, FMinimalName Rhs);
-	friend bool operator==(FName Lhs, FScriptName Rhs);
-	friend bool operator==(FName Lhs, FMemoryImageName Rhs);
-
 	template <typename StringBufferType>
 	FORCEINLINE void AppendStringInternal(StringBufferType& Out) const;
 
@@ -1196,6 +1259,106 @@ private:
 	{
 		return CreateNumberedNameIfNecessary(ComparisonId, ComparisonId, Number);
 	}
+
+	/**
+	 * Equality operator with CharType* on left hand side and FName on right hand side
+	 * 
+	 * @param	LHS		CharType to compare to FName
+	 * @param	RHS		FName to compare to CharType
+	 * @return True if strings match, false otherwise.
+	 */
+	template <typename CharType>
+	friend inline bool operator==(const CharType *LHS, const FName &RHS)
+	{
+		return RHS == LHS;
+	}
+
+	/**
+	 * Inequality operator with CharType* on left hand side and FName on right hand side
+	 *
+	 * @param	LHS		CharType to compare to FName
+	 * @param	RHS		FName to compare to CharType
+	 * @return True if strings don't match, false otherwise.
+	 */
+	template <typename CharType>
+	friend inline bool operator!=(const CharType *LHS, const FName &RHS)
+	{
+		return RHS != LHS;
+	}
+
+#if UE_FNAME_OUTLINE_NUMBER
+	friend FORCEINLINE bool operator==(FName Lhs, FMinimalName Rhs)
+	{
+		return Lhs.GetComparisonIndexInternal() == Rhs.Index;
+	}
+	friend FORCEINLINE bool operator==(FName Lhs, FScriptName Rhs)
+	{
+		return Lhs.GetComparisonIndexInternal() == Rhs.ComparisonIndex;
+	}
+	friend FORCEINLINE bool operator==(FName Lhs, FMemoryImageName Rhs)
+	{
+		return Lhs.GetComparisonIndexInternal() == Rhs.ComparisonIndex;
+	}
+	friend FORCEINLINE uint32 GetTypeHash(FName Name)
+	{
+		return GetTypeHash(Name.GetComparisonIndexInternal());
+	}
+#else
+	friend FORCEINLINE bool operator==(FName Lhs, FMinimalName Rhs)
+	{
+		return Lhs.GetComparisonIndex() == Rhs.Index && Lhs.GetNumber() == Rhs.Number;
+	}
+	friend FORCEINLINE bool operator==(FName Lhs, FScriptName Rhs)
+	{
+		return Lhs.GetComparisonIndex() == Rhs.ComparisonIndex && Lhs.GetNumber() == Rhs.Number;
+	}
+	friend FORCEINLINE bool operator==(FName Lhs, FMemoryImageName Rhs)
+	{
+		return Lhs.GetComparisonIndex() == Rhs.ComparisonIndex && Lhs.GetNumber() == Rhs.Number;
+	}
+	friend FORCEINLINE uint32 GetTypeHash(FName Name)
+	{
+		return GetTypeHash(Name.GetComparisonIndex()) + Name.GetNumber();
+	}
+#endif
+	friend FORCEINLINE bool operator==(FMinimalName Lhs, FName Rhs)
+	{
+		return Rhs == Lhs;
+	}
+	friend FORCEINLINE bool operator!=(FMinimalName Lhs, FName Rhs)
+	{
+		return !(Lhs == Rhs);
+	}
+
+	friend FORCEINLINE bool operator!=(FName Lhs, FMinimalName Rhs)
+	{
+		return !(Lhs == Rhs);
+	}
+
+	friend FORCEINLINE bool operator==(FScriptName Lhs, FName Rhs)
+	{
+		return Rhs == Lhs;
+	}
+	friend FORCEINLINE bool operator!=(FScriptName Lhs, FName Rhs)
+	{
+		return !(Lhs == Rhs);
+	}
+	friend FORCEINLINE bool operator!=(FName Lhs, FScriptName Rhs)
+	{
+		return !(Lhs == Rhs);
+	}
+	friend FORCEINLINE bool operator==(FMemoryImageName Lhs, FName Rhs)
+	{
+		return Rhs == Lhs;
+	}
+	friend FORCEINLINE bool operator!=(FName Lhs, FMemoryImageName Rhs)
+	{
+		return !(Lhs == Rhs);
+	}
+	friend FORCEINLINE bool operator!=(FMemoryImageName Lhs, FName Rhs)
+	{
+		return !(Lhs == Rhs);
+	}
 };
 
 template<> struct TIsZeroConstructType<class FName> { enum { Value = true }; };
@@ -1254,25 +1417,6 @@ DECLARE_INTRINSIC_TYPE_LAYOUT(FMemoryImageName);
 DECLARE_INTRINSIC_TYPE_LAYOUT(FScriptName);
 
 #if UE_FNAME_OUTLINE_NUMBER
-FORCEINLINE uint32 GetTypeHash(FName Name)
-{
-	return ::GetTypeHash(Name.GetComparisonIndexInternal());
-}
-
-FORCEINLINE uint32 GetTypeHash(FMinimalName Name)
-{
-	return ::GetTypeHash(Name.Index);
-}
-
-FORCEINLINE uint32 GetTypeHash(FScriptName Name)
-{
-	return ::GetTypeHash(Name.ComparisonIndex);
-}
-
-FORCEINLINE uint32 GetTypeHash(FMemoryImageName Name)
-{
-	return ::GetTypeHash(Name.ComparisonIndex);
-}
 
 FORCEINLINE FName::FName(FMinimalName InName)
 	: ComparisonIndex(InName.Index)
@@ -1345,36 +1489,6 @@ FORCEINLINE bool FName::operator==(EName Ename) const
 	return GetComparisonIndex() == Ename && GetNumber() == NAME_NO_NUMBER_INTERNAL;
 }
 
-FORCEINLINE bool operator==(FMinimalName Lhs, FMinimalName Rhs)
-{
-	return Lhs.Index == Rhs.Index;
-}
-
-FORCEINLINE bool operator==(FScriptName Lhs, FScriptName Rhs)
-{
-	return Lhs.ComparisonIndex == Rhs.ComparisonIndex;
-}
-
-FORCEINLINE bool operator==(FMemoryImageName Lhs, FMemoryImageName Rhs)
-{
-	return Lhs.ComparisonIndex == Rhs.ComparisonIndex;
-}
-
-FORCEINLINE bool operator==(FName Lhs, FMinimalName Rhs)
-{
-	return Lhs.GetComparisonIndexInternal() == Rhs.Index;
-}
-
-FORCEINLINE bool operator==(FName Lhs, FScriptName Rhs)
-{
-	return Lhs.GetComparisonIndexInternal() == Rhs.ComparisonIndex;
-}
-
-FORCEINLINE bool operator==(FName Lhs, FMemoryImageName Rhs)
-{
-	return Lhs.GetComparisonIndexInternal() == Rhs.ComparisonIndex;
-}
-
 FORCEINLINE bool FName::IsEqual(const FName& Rhs, const ENameCase CompareMethod /*= ENameCase::IgnoreCase*/, const bool bCompareNumber /*= true*/) const
 {
 	return bCompareNumber ?
@@ -1383,25 +1497,6 @@ FORCEINLINE bool FName::IsEqual(const FName& Rhs, const ENameCase CompareMethod
 }
 
 #else // UE_FNAME_OUTLINE_NUMBER
-FORCEINLINE uint32 GetTypeHash(FName Name)
-{
-	return ::GetTypeHash(Name.GetComparisonIndex()) + Name.GetNumber();
-}
-
-FORCEINLINE uint32 GetTypeHash(FMinimalName Name)
-{
-	return ::GetTypeHash(Name.Index) + Name.Number;
-}
-
-FORCEINLINE uint32 GetTypeHash(FScriptName Name)
-{
-	return ::GetTypeHash(Name.ComparisonIndex) + Name.Number;
-}
-
-FORCEINLINE uint32 GetTypeHash(FMemoryImageName Name)
-{
-	return ::GetTypeHash(Name.ComparisonIndex) + Name.Number;
-}
 
 FORCEINLINE FName::FName(FMinimalName InName)
 	: ComparisonIndex(InName.Index)
@@ -1479,36 +1574,6 @@ FORCEINLINE bool FName::operator==(EName Ename) const
 	return GetComparisonIndex() == Ename && GetNumber() == NAME_NO_NUMBER_INTERNAL;
 }
 
-FORCEINLINE bool operator==(FMinimalName Lhs, FMinimalName Rhs)
-{
-	return Lhs.Index == Rhs.Index && Lhs.Number == Rhs.Number;
-}
-
-FORCEINLINE bool operator==(FScriptName Lhs, FScriptName Rhs)
-{
-	return Lhs.ComparisonIndex == Rhs.ComparisonIndex && Lhs.Number == Rhs.Number;
-}
-
-FORCEINLINE bool operator==(FMemoryImageName Lhs, FMemoryImageName Rhs)
-{
-	return Lhs.ComparisonIndex == Rhs.ComparisonIndex && Lhs.Number == Rhs.Number;
-}
-
-FORCEINLINE bool operator==(FName Lhs, FMinimalName Rhs)
-{
-	return Lhs.GetComparisonIndex() == Rhs.Index && Lhs.GetNumber() == Rhs.Number;
-}
-
-FORCEINLINE bool operator==(FName Lhs, FScriptName Rhs)
-{
-	return Lhs.GetComparisonIndex() == Rhs.ComparisonIndex && Lhs.GetNumber() == Rhs.Number;
-}
-
-FORCEINLINE bool operator==(FName Lhs, FMemoryImageName Rhs)
-{
-	return Lhs.GetComparisonIndex() == Rhs.ComparisonIndex && Lhs.GetNumber() == Rhs.Number;
-}
-
 FORCEINLINE bool FName::IsEqual(const FName& Rhs, const ENameCase CompareMethod /*= ENameCase::IgnoreCase*/, const bool bCompareNumber /*= true*/) const
 {
 	return ((CompareMethod == ENameCase::IgnoreCase) ? GetComparisonIndex() == Rhs.GetComparisonIndex() : GetDisplayIndexFast() == Rhs.GetDisplayIndexFast())
@@ -1536,66 +1601,6 @@ FORCEINLINE FScriptName NameToScriptName(FName InName)
 	return FScriptName(InName);
 }
 
-FORCEINLINE bool operator==(FMinimalName Lhs, FName Rhs)
-{
-	return operator==(Rhs, Lhs);
-}
-
-FORCEINLINE bool operator==(FScriptName Lhs, FName Rhs)
-{
-	return operator==(Rhs, Lhs);
-}
-
-FORCEINLINE bool operator==(FMemoryImageName Lhs, FName Rhs)
-{
-	return operator==(Rhs, Lhs);
-}
-
-FORCEINLINE bool operator!=(FMinimalName Lhs, FMinimalName Rhs)
-{
-	return !operator==(Lhs, Rhs);
-}
-
-FORCEINLINE bool operator!=(FMinimalName Lhs, FName Rhs)
-{
-	return !operator==(Lhs, Rhs);
-}
-
-FORCEINLINE bool operator!=(FScriptName Lhs, FScriptName Rhs)
-{
-	return !operator==(Lhs, Rhs);
-}
-
-FORCEINLINE bool operator!=(FScriptName Lhs, FName Rhs)
-{
-	return !operator==(Lhs, Rhs);
-}
-
-FORCEINLINE bool operator!=(FMemoryImageName Lhs, FMemoryImageName Rhs)
-{
-	return !operator==(Lhs, Rhs);
-}
-
-FORCEINLINE bool operator!=(FMemoryImageName Lhs, FName Rhs)
-{
-	return !operator==(Lhs, Rhs);
-}
-
-FORCEINLINE bool operator!=(FName Lhs, FMinimalName Rhs)
-{
-	return !operator==(Lhs, Rhs);
-}
-
-FORCEINLINE bool operator!=(FName Lhs, FScriptName Rhs)
-{
-	return !operator==(Lhs, Rhs);
-}
-
-FORCEINLINE bool operator!=(FName Lhs, FMemoryImageName Rhs)
-{
-	return !operator==(Lhs, Rhs);
-}
-
 FORCEINLINE bool FName::operator!=(EName Ename) const
 {
 	return !(*this == Ename);
@@ -1626,32 +1631,6 @@ inline FUtf8StringBuilderBase& operator<<(FUtf8StringBuilderBase& Builder, const
 CORE_API FWideStringBuilderBase& operator<<(FWideStringBuilderBase& Builder, FNameEntryId Id);
 CORE_API FUtf8StringBuilderBase& operator<<(FUtf8StringBuilderBase& Builder, FNameEntryId Id);
 
-/**
- * Equality operator with CharType* on left hand side and FName on right hand side
- * 
- * @param	LHS		CharType to compare to FName
- * @param	RHS		FName to compare to CharType
- * @return True if strings match, false otherwise.
- */
-template <typename CharType>
-inline bool operator==(const CharType *LHS, const FName &RHS)
-{
-	return RHS == LHS;
-}
-
-/**
- * Inequality operator with CharType* on left hand side and FName on right hand side
- *
- * @param	LHS		CharType to compare to FName
- * @param	RHS		FName to compare to CharType
- * @return True if strings don't match, false otherwise.
- */
-template <typename CharType>
-inline bool operator!=(const CharType *LHS, const FName &RHS)
-{
-	return RHS != LHS;
-}
-
 /** FNames act like PODs. */
 template <> struct TIsPODType<FName> { enum { Value = true }; };
 
diff --git a/Engine/Source/Runtime/Core/Public/UObject/WeakObjectPtrTemplates.h b/Engine/Source/Runtime/Core/Public/UObject/WeakObjectPtrTemplates.h
--- a/Engine/Source/Runtime/Core/Public/UObject/WeakObjectPtrTemplates.h
+++ b/Engine/Source/Runtime/Core/Public/UObject/WeakObjectPtrTemplates.h
@@ -229,89 +229,85 @@ public:
 		Ar << static_cast<TWeakObjectPtrBase&>(WeakObjectPtr);
 		return Ar;
 	}
-};
 
-// Helper function which deduces the type of the initializer
-template <typename T>
-FORCEINLINE TWeakObjectPtr<T> MakeWeakObjectPtr(T* Ptr)
-{
-	return TWeakObjectPtr<T>(Ptr);
-}
+	/**
+	 * Compare weak pointers for equality.
+	 * If both pointers would return nullptr from Get() they count as equal even if they were not initialized to the same object.
+	 * @param Other weak pointer to compare to
+	 */
+	template <typename RhsT, typename = decltype((T*)nullptr == (RhsT*)nullptr)>
+	friend FORCENOINLINE bool operator==(const TWeakObjectPtr& Lhs, const TWeakObjectPtr<RhsT, TWeakObjectPtrBase>& Rhs)
+	{
+		return (const TWeakObjectPtrBase&)Lhs == (const TWeakObjectPtrBase&)Rhs;
+	}
 
-/**
- * Compare weak pointers for equality.
- * If both pointers would return nullptr from Get() they count as equal even if they were not initialized to the same object.
- * @param Other weak pointer to compare to
- */
-template <typename LhsT, typename RhsT, typename OtherTWeakObjectPtrBase, typename = decltype((LhsT*)nullptr == (RhsT*)nullptr)>
-FORCENOINLINE bool operator==(const TWeakObjectPtr<LhsT, OtherTWeakObjectPtrBase>& Lhs, const TWeakObjectPtr<RhsT, OtherTWeakObjectPtrBase>& Rhs)
-{
-	return (const OtherTWeakObjectPtrBase&)Lhs == (const OtherTWeakObjectPtrBase&)Rhs;
-}
+	template <typename RhsT, typename = decltype((T*)nullptr == (RhsT*)nullptr)>
+	friend FORCENOINLINE bool operator==(const TWeakObjectPtr& Lhs, const RhsT* Rhs)
+	{
+		// NOTE: this constructs a TWeakObjectPtrBase, which has some amount of overhead, so this may not be an efficient operation
+		return (const TWeakObjectPtrBase&)Lhs == TWeakObjectPtrBase(Rhs);
+	}
 
-template <typename LhsT, typename RhsT, typename OtherTWeakObjectPtrBase, typename = decltype((LhsT*)nullptr == (RhsT*)nullptr)>
-FORCENOINLINE bool operator==(const TWeakObjectPtr<LhsT, OtherTWeakObjectPtrBase>& Lhs, const RhsT* Rhs)
-{
-	// NOTE: this constructs a TWeakObjectPtrBase, which has some amount of overhead, so this may not be an efficient operation
-	return (const OtherTWeakObjectPtrBase&)Lhs == OtherTWeakObjectPtrBase(Rhs);
-}
+	template <typename LhsT, typename = decltype((LhsT*)nullptr == (T*)nullptr)>
+	friend FORCENOINLINE bool operator==(const LhsT* Lhs, const TWeakObjectPtr& Rhs)
+	{
+		// NOTE: this constructs a TWeakObjectPtrBase, which has some amount of overhead, so this may not be an efficient operation
+		return TWeakObjectPtrBase(Lhs) == (const TWeakObjectPtrBase&)Rhs;
+	}
 
-template <typename LhsT, typename RhsT, typename OtherTWeakObjectPtrBase, typename = decltype((LhsT*)nullptr == (RhsT*)nullptr)>
-FORCENOINLINE bool operator==(const LhsT* Lhs, const TWeakObjectPtr<RhsT, OtherTWeakObjectPtrBase>& Rhs)
-{
-	// NOTE: this constructs a TWeakObjectPtrBase, which has some amount of overhead, so this may not be an efficient operation
-	return OtherTWeakObjectPtrBase(Lhs) == (const OtherTWeakObjectPtrBase&)Rhs;
-}
+	friend FORCENOINLINE bool operator==(const TWeakObjectPtr& Lhs, TYPE_OF_NULLPTR)
+	{
+		return !Lhs.IsValid();
+	}
 
-template <typename LhsT, typename OtherTWeakObjectPtrBase>
-FORCENOINLINE bool operator==(const TWeakObjectPtr<LhsT, OtherTWeakObjectPtrBase>& Lhs, TYPE_OF_NULLPTR)
-{
-	return !Lhs.IsValid();
-}
+	friend FORCENOINLINE bool operator==(TYPE_OF_NULLPTR, const TWeakObjectPtr& Rhs)
+	{
+		return !Rhs.IsValid();
+	}
 
-template <typename RhsT, typename OtherTWeakObjectPtrBase>
-FORCENOINLINE bool operator==(TYPE_OF_NULLPTR, const TWeakObjectPtr<RhsT, OtherTWeakObjectPtrBase>& Rhs)
-{
-	return !Rhs.IsValid();
-}
+	/**
+	 * Compare weak pointers for inequality
+	 * @param Other weak pointer to compare to
+	 */
+	template <typename RhsT, typename = decltype((T*)nullptr != (RhsT*)nullptr)>
+	friend FORCENOINLINE bool operator!=(const TWeakObjectPtr& Lhs, const TWeakObjectPtr<RhsT, TWeakObjectPtrBase>& Rhs)
+	{
+		return (const TWeakObjectPtrBase&)Lhs != (const TWeakObjectPtrBase&)Rhs;
+	}
 
-/**
- * Compare weak pointers for inequality
- * @param Other weak pointer to compare to
- */
-template <typename LhsT, typename RhsT, typename OtherTWeakObjectPtrBase, typename = decltype((LhsT*)nullptr != (RhsT*)nullptr)>
-FORCENOINLINE bool operator!=(const TWeakObjectPtr<LhsT, OtherTWeakObjectPtrBase>& Lhs, const TWeakObjectPtr<RhsT, OtherTWeakObjectPtrBase>& Rhs)
-{
-	return (const OtherTWeakObjectPtrBase&)Lhs != (const OtherTWeakObjectPtrBase&)Rhs;
-}
+	template <typename RhsT, typename = decltype((T*)nullptr != (RhsT*)nullptr)>
+	friend FORCENOINLINE bool operator!=(const TWeakObjectPtr& Lhs, const RhsT* Rhs)
+	{
+		// NOTE: this constructs a TWeakObjectPtrBase, which has some amount of overhead, so this may not be an efficient operation
+		return (const TWeakObjectPtrBase&)Lhs != TWeakObjectPtrBase(Rhs);
+	}
 
-template <typename LhsT, typename RhsT, typename OtherTWeakObjectPtrBase, typename = decltype((LhsT*)nullptr != (RhsT*)nullptr)>
-FORCENOINLINE bool operator!=(const TWeakObjectPtr<LhsT, OtherTWeakObjectPtrBase>& Lhs, const RhsT* Rhs)
-{
-	// NOTE: this constructs a TWeakObjectPtrBase, which has some amount of overhead, so this may not be an efficient operation
-	return (const OtherTWeakObjectPtrBase&)Lhs != OtherTWeakObjectPtrBase(Rhs);
-}
+	template <typename LhsT, typename = decltype((LhsT*)nullptr != (T*)nullptr)>
+	friend FORCENOINLINE bool operator!=(const LhsT* Lhs, const TWeakObjectPtr& Rhs)
+	{
+		// NOTE: this constructs a TWeakObjectPtrBase, which has some amount of overhead, so this may not be an efficient operation
+		return TWeakObjectPtrBase(Lhs) != (const TWeakObjectPtrBase&)Rhs;
+	}
 
-template <typename LhsT, typename RhsT, typename OtherTWeakObjectPtrBase, typename = decltype((LhsT*)nullptr != (RhsT*)nullptr)>
-FORCENOINLINE bool operator!=(const LhsT* Lhs, const TWeakObjectPtr<RhsT, OtherTWeakObjectPtrBase>& Rhs)
-{
-	// NOTE: this constructs a TWeakObjectPtrBase, which has some amount of overhead, so this may not be an efficient operation
-	return OtherTWeakObjectPtrBase(Lhs) != (const OtherTWeakObjectPtrBase&)Rhs;
-}
+	friend FORCENOINLINE bool operator!=(const TWeakObjectPtr& Lhs, TYPE_OF_NULLPTR)
+	{
+		return Lhs.IsValid();
+	}
 
-template <typename LhsT, typename OtherTWeakObjectPtrBase>
-FORCENOINLINE bool operator!=(const TWeakObjectPtr<LhsT, OtherTWeakObjectPtrBase>& Lhs, TYPE_OF_NULLPTR)
-{
-	return Lhs.IsValid();
-}
+	friend FORCENOINLINE bool operator!=(TYPE_OF_NULLPTR, const TWeakObjectPtr& Rhs)
+	{
+		return Rhs.IsValid();
+	}
 
-template <typename RhsT, typename OtherTWeakObjectPtrBase>
-FORCENOINLINE bool operator!=(TYPE_OF_NULLPTR, const TWeakObjectPtr<RhsT, OtherTWeakObjectPtrBase>& Rhs)
+};
+
+// Helper function which deduces the type of the initializer
+template <typename T>
+FORCEINLINE TWeakObjectPtr<T> MakeWeakObjectPtr(T* Ptr)
 {
-	return Rhs.IsValid();
+	return TWeakObjectPtr<T>(Ptr);
 }
 
-
 template<class T> struct TIsPODType<TWeakObjectPtr<T> > { enum { Value = true }; };
 template<class T> struct TIsZeroConstructType<TWeakObjectPtr<T> > { enum { Value = true }; };
 template<class T> struct TIsWeakPointerType<TWeakObjectPtr<T> > { enum { Value = true }; };
diff --git a/Engine/Source/Runtime/CoreUObject/Public/AssetRegistry/AssetData.h b/Engine/Source/Runtime/CoreUObject/Public/AssetRegistry/AssetData.h
--- a/Engine/Source/Runtime/CoreUObject/Public/AssetRegistry/AssetData.h
+++ b/Engine/Source/Runtime/CoreUObject/Public/AssetRegistry/AssetData.h
@@ -733,15 +733,15 @@ private:
 public:
 	/** Helper function that tries to convert short class name to path name */
 	COREUOBJECT_API static FTopLevelAssetPath TryConvertShortClassNameToPathName(FName InClassName, ELogVerbosity::Type FailureMessageVerbosity = ELogVerbosity::Warning);
+
+	friend FORCEINLINE uint32 GetTypeHash(const FAssetData& AssetData)
+	{
+		return HashCombine(GetTypeHash(AssetData.PackageName), GetTypeHash(AssetData.AssetName));
+	}
 };
 
 ENUM_CLASS_FLAGS(FAssetData::ECreationFlags);
 
-FORCEINLINE uint32 GetTypeHash(const FAssetData& AssetData)
-{
-	return HashCombine(GetTypeHash(AssetData.PackageName), GetTypeHash(AssetData.AssetName));
-}
-
 
 template<>
 struct TStructOpsTypeTraits<FAssetData> : public TStructOpsTypeTraitsBase2<FAssetData>
diff --git a/Engine/Source/Runtime/CoreUObject/Public/AssetRegistry/AssetDataTagMap.h b/Engine/Source/Runtime/CoreUObject/Public/AssetRegistry/AssetDataTagMap.h
--- a/Engine/Source/Runtime/CoreUObject/Public/AssetRegistry/AssetDataTagMap.h
+++ b/Engine/Source/Runtime/CoreUObject/Public/AssetRegistry/AssetDataTagMap.h
@@ -54,10 +54,10 @@ PRAGMA_ENABLE_DEPRECATION_WARNINGS
 	}
 	bool IsEmpty() const { return ClassPath.IsNull() & Package.IsNone() & Object.IsNone(); } //-V792
 	explicit operator bool() const { return !IsEmpty(); }
-};
 
-bool operator==(const FAssetRegistryExportPath& A, const FAssetRegistryExportPath& B);
-uint32 GetTypeHash(const FAssetRegistryExportPath& Export);
+	friend bool operator==(const FAssetRegistryExportPath& A, const FAssetRegistryExportPath& B);
+	friend uint32 GetTypeHash(const FAssetRegistryExportPath& Export);
+};
 
 namespace FixedTagPrivate
 {
@@ -266,18 +266,18 @@ private:
 
 	friend class FixedTagPrivate::FStoreBuilder;
 	friend FAssetRegistryState;
-};
 
-inline bool operator==(FAssetTagValueRef A, FStringView B) { return  A.Equals(B); }
-inline bool operator!=(FAssetTagValueRef A, FStringView B) { return !A.Equals(B); }
-inline bool operator==(FStringView A, FAssetTagValueRef B) { return  B.Equals(A); }
-inline bool operator!=(FStringView A, FAssetTagValueRef B) { return !B.Equals(A); }
+	friend inline bool operator==(FAssetTagValueRef A, FStringView B) { return  A.Equals(B); }
+	friend inline bool operator!=(FAssetTagValueRef A, FStringView B) { return !A.Equals(B); }
+	friend inline bool operator==(FStringView A, FAssetTagValueRef B) { return  B.Equals(A); }
+	friend inline bool operator!=(FStringView A, FAssetTagValueRef B) { return !B.Equals(A); }
 
-// These overloads can be removed when the deprecated implicit operator FString is removed
-inline bool operator==(FAssetTagValueRef A, const FString& B) { return  A.Equals(B); }
-inline bool operator!=(FAssetTagValueRef A, const FString& B) { return !A.Equals(B); }
-inline bool operator==(const FString& A, FAssetTagValueRef B) { return  B.Equals(A); }
-inline bool operator!=(const FString& A, FAssetTagValueRef B) { return !B.Equals(A); }
+	// These overloads can be removed when the deprecated implicit operator FString is removed
+	friend inline bool operator==(FAssetTagValueRef A, const FString& B) { return  A.Equals(B); }
+	friend inline bool operator!=(FAssetTagValueRef A, const FString& B) { return !A.Equals(B); }
+	friend inline bool operator==(const FString& A, FAssetTagValueRef B) { return  B.Equals(A); }
+	friend inline bool operator!=(const FString& A, FAssetTagValueRef B) { return !B.Equals(A); }
+};
 
 using FAssetDataTagMapBase = TSortedMap<FName, FString, FDefaultAllocator, FNameFastLess>;
 
@@ -499,9 +499,9 @@ public:
 		SIZE_T GetLooseSize() const { return LooseBytes; }
 		SIZE_T GetFixedSize() const;		
 	};
-};
 
-inline bool operator==(const FAssetDataTagMap& A, const FAssetDataTagMapSharedView& B)				{ return B == A; }
-inline bool operator!=(const FAssetDataTagMap& A, const FAssetDataTagMapSharedView& B)				{ return !(B == A); }
-inline bool operator!=(const FAssetDataTagMapSharedView& A, const FAssetDataTagMap& B)				{ return !(A == B); }
-inline bool operator!=(const FAssetDataTagMapSharedView& A, const FAssetDataTagMapSharedView& B)	{ return !(A == B); }
\ No newline at end of file
+	friend inline bool operator==(const FAssetDataTagMap& A, const FAssetDataTagMapSharedView& B)			{ return B == A; }
+	friend inline bool operator!=(const FAssetDataTagMap& A, const FAssetDataTagMapSharedView& B)			{ return !(B == A); }
+	friend inline bool operator!=(const FAssetDataTagMapSharedView& A, const FAssetDataTagMap& B)			{ return !(A == B); }
+	friend inline bool operator!=(const FAssetDataTagMapSharedView& A, const FAssetDataTagMapSharedView& B)	{ return !(A == B); }
+};
diff --git a/Engine/Source/Runtime/CoreUObject/Public/Serialization/DerivedData.h b/Engine/Source/Runtime/CoreUObject/Public/Serialization/DerivedData.h
--- a/Engine/Source/Runtime/CoreUObject/Public/Serialization/DerivedData.h
+++ b/Engine/Source/Runtime/CoreUObject/Public/Serialization/DerivedData.h
@@ -420,15 +420,15 @@ public:
 
 private:
 	TPimplPtr<DerivedData::Private::FIoResponse> Response;
+
+	constexpr uint32 GetTypeHash(const FDerivedDataIoPriority& Priority)
+	{
+		return uint32(Priority.Value);
+	}
 };
 
 } // UE
 
-constexpr uint32 GetTypeHash(const UE::FDerivedDataIoPriority& Priority)
-{
-	return uint32(Priority.Value);
-}
-
 namespace UE::DerivedData::IoStore
 {
 
diff --git a/Engine/Source/Runtime/CoreUObject/Public/UObject/Class.h b/Engine/Source/Runtime/CoreUObject/Public/UObject/Class.h
--- a/Engine/Source/Runtime/CoreUObject/Public/UObject/Class.h
+++ b/Engine/Source/Runtime/CoreUObject/Public/UObject/Class.h
@@ -1433,7 +1433,7 @@ public:
 
 			if constexpr (TModels<CGetTypeHashable, CPPSTRUCT>::Value)
 			{
-				return GetTypeHash(*(const CPPSTRUCT*)Src);
+				return GetTypeHashHelper(*(const CPPSTRUCT*)Src);
 			}
 			else
 			{
diff --git a/Engine/Source/Runtime/CoreUObject/Public/UObject/FieldPath.h b/Engine/Source/Runtime/CoreUObject/Public/UObject/FieldPath.h
--- a/Engine/Source/Runtime/CoreUObject/Public/UObject/FieldPath.h
+++ b/Engine/Source/Runtime/CoreUObject/Public/UObject/FieldPath.h
@@ -475,61 +475,59 @@ public:
 
 		return Get() != Other;
 	}
-};
 
-// Helper function which deduces the type of the initializer
-template <typename PropertyType>
-FORCEINLINE TFieldPath<PropertyType> MakePropertyPath(PropertyType* Ptr)
-{
-	return TFieldPath<PropertyType>(Ptr);
-}
 
-template <typename LhsT, typename RhsT>
-FORCENOINLINE bool operator==(const LhsT* Lhs, const TFieldPath<RhsT>& Rhs)
-{
-	// It's also possible that these static_asserts may fail for valid conversions because
-	// one or both of the types have only been forward-declared.
-	static_assert(TPointerIsConvertibleFromTo<LhsT, const FField>::Value, "TFieldPath can only be compared with FField types");
-	static_assert(TPointerIsConvertibleFromTo<LhsT, RhsT>::Value || TPointerIsConvertibleFromTo<RhsT, LhsT>::Value, "Unable to compare TFieldPath with raw pointer - types are incompatible");
+	template <typename LhsT>
+	friend FORCENOINLINE bool operator==(const LhsT* Lhs, const TFieldPath<PropertyType>& Rhs)
+	{
+		// It's also possible that these static_asserts may fail for valid conversions because
+		// one or both of the types have only been forward-declared.
+		static_assert(TPointerIsConvertibleFromTo<LhsT, const FField>::Value, "TFieldPath can only be compared with FField types");
+		static_assert(TPointerIsConvertibleFromTo<LhsT, PropertyType>::Value || TPointerIsConvertibleFromTo<PropertyType, LhsT>::Value, "Unable to compare TFieldPath with raw pointer - types are incompatible");
 
-	return Rhs == Lhs;
-}
+		return Rhs == Lhs;
+	}
 
-template <typename LhsT>
-FORCENOINLINE bool operator==(const TFieldPath<LhsT>& Lhs, TYPE_OF_NULLPTR)
-{
-	return !Lhs.Get();
-}
+	friend FORCENOINLINE bool operator==(const TFieldPath<PropertyType>& Lhs, TYPE_OF_NULLPTR)
+	{
+		return !Lhs.Get();
+	}
 
-template <typename RhsT>
-FORCENOINLINE bool operator==(TYPE_OF_NULLPTR, const TFieldPath<RhsT>& Rhs)
-{
-	return !Rhs.Get();
-}
+	friend FORCENOINLINE bool operator==(TYPE_OF_NULLPTR, const TFieldPath<PropertyType>& Rhs)
+	{
+		return !Rhs.Get();
+	}
 
-template <typename LhsT, typename RhsT>
-FORCENOINLINE bool operator!=(const LhsT* Lhs, const TFieldPath<RhsT>& Rhs)
-{
-	// It's also possible that these static_asserts may fail for valid conversions because
-	// one or both of the types have only been forward-declared.
-	static_assert(TPointerIsConvertibleFromTo<LhsT, const FField>::Value, "TFieldPath can only be compared with FField types");
-	static_assert(TPointerIsConvertibleFromTo<LhsT, RhsT>::Value || TPointerIsConvertibleFromTo<RhsT, LhsT>::Value, "Unable to compare TFieldPath with raw pointer - types are incompatible");
+	template <typename LhsT>
+	friend FORCENOINLINE bool operator!=(const LhsT* Lhs, const TFieldPath<PropertyType>& Rhs)
+	{
+		// It's also possible that these static_asserts may fail for valid conversions because
+		// one or both of the types have only been forward-declared.
+		static_assert(TPointerIsConvertibleFromTo<LhsT, const FField>::Value, "TFieldPath can only be compared with FField types");
+		static_assert(TPointerIsConvertibleFromTo<LhsT, PropertyType>::Value || TPointerIsConvertibleFromTo<PropertyType, LhsT>::Value, "Unable to compare TFieldPath with raw pointer - types are incompatible");
 
-	return Rhs != Lhs;
-}
+		return Rhs != Lhs;
+	}
 
-template <typename LhsT>
-FORCENOINLINE bool operator!=(const TFieldPath<LhsT>& Lhs, TYPE_OF_NULLPTR)
-{
-	return !!Lhs.Get();
-}
+	friend FORCENOINLINE bool operator!=(const TFieldPath<PropertyType>& Lhs, TYPE_OF_NULLPTR)
+	{
+		return !!Lhs.Get();
+	}
 
-template <typename RhsT>
-FORCENOINLINE bool operator!=(TYPE_OF_NULLPTR, const TFieldPath<RhsT>& Rhs)
+	friend FORCENOINLINE bool operator!=(TYPE_OF_NULLPTR, const TFieldPath<PropertyType>& Rhs)
+	{
+		return !!Rhs.Get();
+	}
+};
+
+// Helper function which deduces the type of the initializer
+template <typename PropertyType>
+FORCEINLINE TFieldPath<PropertyType> MakePropertyPath(PropertyType* Ptr)
 {
-	return !!Rhs.Get();
+	return TFieldPath<PropertyType>(Ptr);
 }
 
+
 template<class T> struct TIsPODType<TFieldPath<T> > { enum { Value = true }; };
 template<class T> struct TIsZeroConstructType<TFieldPath<T> > { enum { Value = true }; };
 template<class T> struct TIsWeakPointerType<TFieldPath<T> > { enum { Value = true }; };
diff --git a/Engine/Source/Runtime/CoreUObject/Public/UObject/LazyObjectPtr.h b/Engine/Source/Runtime/CoreUObject/Public/UObject/LazyObjectPtr.h
--- a/Engine/Source/Runtime/CoreUObject/Public/UObject/LazyObjectPtr.h
+++ b/Engine/Source/Runtime/CoreUObject/Public/UObject/LazyObjectPtr.h
@@ -322,69 +322,63 @@ public:
 		Ar << static_cast<FLazyObjectPtr&>(LazyObjectPtr);
 		return Ar;
 	}
-};
-
-// The reason these aren't inside the class (above) is because Visual Studio 2012-2013 crashes when compiling them :D
 
-/** Compare with another TLazyObjectPtr of related type */
-template<typename T, typename U, typename = decltype((T*)nullptr == (U*)nullptr)>
-FORCEINLINE bool operator==(const TLazyObjectPtr<T>& Lhs, const TLazyObjectPtr<U>& Rhs)
-{
-	return (const FLazyObjectPtr&)Lhs == (const FLazyObjectPtr&)Rhs;
-}
-template<typename T, typename U, typename = decltype((T*)nullptr != (U*)nullptr)>
-FORCEINLINE bool operator!=(const TLazyObjectPtr<T>& Lhs, const TLazyObjectPtr<U>& Rhs)
-{
-	return (const FLazyObjectPtr&)Lhs != (const FLazyObjectPtr&)Rhs;
-}
+	/** Compare with another TLazyObjectPtr of related type */
+	template<typename U, typename = decltype((T*)nullptr == (U*)nullptr)>
+	friend FORCEINLINE bool operator==(const TLazyObjectPtr<T>& Lhs, const TLazyObjectPtr<U>& Rhs)
+	{
+		return (const FLazyObjectPtr&)Lhs == (const FLazyObjectPtr&)Rhs;
+	}
+	template<typename U, typename = decltype((T*)nullptr != (U*)nullptr)>
+	friend FORCEINLINE bool operator!=(const TLazyObjectPtr<T>& Lhs, const TLazyObjectPtr<U>& Rhs)
+	{
+		return (const FLazyObjectPtr&)Lhs != (const FLazyObjectPtr&)Rhs;
+	}
 
-/** Compare for equality with a raw pointer **/
-template<typename T, typename U, typename = decltype((T*)nullptr == (U*)nullptr)>
-FORCEINLINE bool operator==(const TLazyObjectPtr<T>& Lhs, const U* Rhs)
-{
-	return Lhs.Get() == Rhs;
-}
-template<typename T, typename U, typename = decltype((T*)nullptr == (U*)nullptr)>
-FORCEINLINE bool operator==(const U* Lhs, const TLazyObjectPtr<T>& Rhs)
-{
-	return Lhs == Rhs.Get();
-}
+	/** Compare for equality with a raw pointer **/
+	template<typename U, typename = decltype((T*)nullptr == (U*)nullptr)>
+	friend FORCEINLINE bool operator==(const TLazyObjectPtr<T>& Lhs, const U* Rhs)
+	{
+		return Lhs.Get() == Rhs;
+	}
+	template<typename U, typename = decltype((T*)nullptr == (U*)nullptr)>
+	friend FORCEINLINE bool operator==(const U* Lhs, const TLazyObjectPtr<T>& Rhs)
+	{
+		return Lhs == Rhs.Get();
+	}
 
-/** Compare to null */
-template<typename T>
-FORCEINLINE bool operator==(const TLazyObjectPtr<T>& Lhs, TYPE_OF_NULLPTR)
-{
-	return !Lhs.IsValid();
-}
-template<typename T>
-FORCEINLINE bool operator==(TYPE_OF_NULLPTR, const TLazyObjectPtr<T>& Rhs)
-{
-	return !Rhs.IsValid();
-}
+	/** Compare to null */
+	friend FORCEINLINE bool operator==(const TLazyObjectPtr<T>& Lhs, TYPE_OF_NULLPTR)
+	{
+		return !Lhs.IsValid();
+	}
+	friend FORCEINLINE bool operator==(TYPE_OF_NULLPTR, const TLazyObjectPtr<T>& Rhs)
+	{
+		return !Rhs.IsValid();
+	}
 
-/** Compare for inequality with a raw pointer	**/
-template<typename T, typename U, typename = decltype((T*)nullptr != (U*)nullptr)>
-FORCEINLINE bool operator!=(const TLazyObjectPtr<T>& Lhs, const U* Rhs)
-{
-	return Lhs.Get() != Rhs;
-}
-template<typename T, typename U, typename = decltype((T*)nullptr != (U*)nullptr)>
-FORCEINLINE bool operator!=(const U* Lhs, const TLazyObjectPtr<T>& Rhs)
-{
-	return Lhs != Rhs.Get();
-}
+	/** Compare for inequality with a raw pointer	**/
+	template<typename U, typename = decltype((T*)nullptr != (U*)nullptr)>
+	friend FORCEINLINE bool operator!=(const TLazyObjectPtr<T>& Lhs, const U* Rhs)
+	{
+		return Lhs.Get() != Rhs;
+	}
+	template<typename U, typename = decltype((T*)nullptr != (U*)nullptr)>
+	friend FORCEINLINE bool operator!=(const U* Lhs, const TLazyObjectPtr<T>& Rhs)
+	{
+		return Lhs != Rhs.Get();
+	}
 
-/** Compare for inequality with null **/
-template<typename T>
-FORCEINLINE bool operator!=(const TLazyObjectPtr<T>& Lhs, TYPE_OF_NULLPTR)
-{
-	return Lhs.IsValid();
-}
-template<typename T>
-FORCEINLINE bool operator!=(TYPE_OF_NULLPTR, const TLazyObjectPtr<T>& Rhs)
-{
-	return Rhs.IsValid();
-}
+	/** Compare for inequality with null **/
+	friend FORCEINLINE bool operator!=(const TLazyObjectPtr<T>& Lhs, TYPE_OF_NULLPTR)
+	{
+		return Lhs.IsValid();
+	}
+	friend FORCEINLINE bool operator!=(TYPE_OF_NULLPTR, const TLazyObjectPtr<T>& Rhs)
+	{
+		return Rhs.IsValid();
+	}
+};
 
 template<class T> struct TIsPODType<TLazyObjectPtr<T> > { enum { Value = TIsPODType<FLazyObjectPtr>::Value }; };
 template<class T> struct TIsWeakPointerType<TLazyObjectPtr<T> > { enum { Value = TIsWeakPointerType<FLazyObjectPtr>::Value }; };
\ No newline at end of file
diff --git a/Engine/Source/Runtime/CoreUObject/Public/UObject/ObjectHandle.h b/Engine/Source/Runtime/CoreUObject/Public/UObject/ObjectHandle.h
--- a/Engine/Source/Runtime/CoreUObject/Public/UObject/ObjectHandle.h
+++ b/Engine/Source/Runtime/CoreUObject/Public/UObject/ObjectHandle.h
@@ -119,6 +119,10 @@ struct FPackedObjectRef
 	// Must be 0 for a reference to null.
 	// The least significant bit must always be 1 in a non-null reference.
 	UPTRINT EncodedRef;
+
+	friend inline bool operator==(FPackedObjectRef LHS, FPackedObjectRef RHS) { return LHS.EncodedRef == RHS.EncodedRef; }
+	friend inline bool operator!=(FPackedObjectRef LHS, FPackedObjectRef RHS) { return LHS.EncodedRef != RHS.EncodedRef; }
+	friend inline uint32 GetTypeHash(FPackedObjectRef ObjectRef) { return GetTypeHash(ObjectRef.EncodedRef); }
 };
 
 inline bool IsPackedObjectRefNull(FPackedObjectRef ObjectRef) { return !ObjectRef.EncodedRef; }
@@ -128,26 +132,15 @@ COREUOBJECT_API FPackedObjectRef MakePackedObjectRef(const FObjectRef& ObjectRef
 COREUOBJECT_API UObject* ResolvePackedObjectRef(FPackedObjectRef ObjectRef, uint32 LoadFlags = LOAD_None);
 COREUOBJECT_API UClass* ResolvePackedObjectRefClass(FPackedObjectRef ObjectRef, uint32 LoadFlags = LOAD_None);
 
-inline bool operator==(FPackedObjectRef LHS, FPackedObjectRef RHS) { return LHS.EncodedRef == RHS.EncodedRef; }
-inline bool operator!=(FPackedObjectRef LHS, FPackedObjectRef RHS) { return LHS.EncodedRef != RHS.EncodedRef; }
-inline uint32 GetTypeHash(FPackedObjectRef ObjectRef) { return GetTypeHash(ObjectRef.EncodedRef); }
-
 /**
  * FObjectHandle is either a packed object ref or the resolved pointer to an object.  Depending on configuration
  * when you create a handle, it may immediately be resolved to a pointer.
  */
 #if UE_WITH_OBJECT_HANDLE_LATE_RESOLVE
 
-struct FObjectHandleInternal
-{
-	UPTRINT PointerOrRef;
-};
+struct FObjectHandleInternal;
 using FObjectHandle = FObjectHandleInternal;
 
-inline bool operator==(FObjectHandle LHS, FObjectHandle RHS);
-inline bool operator!=(FObjectHandle LHS, FObjectHandle RHS);
-inline uint32 GetTypeHash(FObjectHandle Handle);
-
 #else
 
 using FObjectHandle = UObject*;
@@ -182,6 +175,32 @@ inline FPackedObjectRef ReadObjectHandlePackedObjectRef(FObjectHandle Handle);
 #if UE_WITH_OBJECT_HANDLE_LATE_RESOLVE
 inline FObjectRef MakeObjectRef(FObjectHandle Handle);
 inline FPackedObjectRef MakePackedObjectRef(FObjectHandle Handle);
+
+struct FObjectHandleInternal
+{
+	UPTRINT PointerOrRef;
+
+	friend inline bool operator==(FObjectHandle LHS, FObjectHandle RHS)
+	{
+		if (IsObjectHandleResolved(LHS) == IsObjectHandleResolved(RHS))
+		{
+			return LHS.PointerOrRef == RHS.PointerOrRef;
+		}
+		else
+		{
+			return MakePackedObjectRef(LHS) == MakePackedObjectRef(RHS);
+		}
+	}
+	friend inline bool operator!=(FObjectHandle LHS, FObjectHandle RHS)
+	{
+		return !(LHS == RHS);
+	}
+	friend inline uint32 GetTypeHash(FObjectHandleInternal Handle)
+	{
+		checkf(IsObjectHandleResolved(Handle), TEXT("Cannot hash an unresolved handle."));
+		return GetTypeHash(ReadObjectHandlePointerNoCheck(Handle));
+	}
+};
 #endif
 
 
@@ -288,14 +307,14 @@ struct FObjectHandleDataClassDescriptor
 
 namespace ObjectHandle_Private
 {
-	constexpr uint32 ObjectPathIdShift = 1;
-	constexpr uint32 ObjectPathIdMask = 0x00FF'FFFF;
+	inline constexpr uint32 ObjectPathIdShift = 1;
+	inline constexpr uint32 ObjectPathIdMask = 0x00FF'FFFF;
 
-	constexpr uint32 DataClassDescriptorIdShift = 25;
-	constexpr uint32 DataClassDescriptorIdMask = 0x0000'00FF;
+	inline constexpr uint32 DataClassDescriptorIdShift = 25;
+	inline constexpr uint32 DataClassDescriptorIdMask = 0x0000'00FF;
 
-	constexpr uint32 PackageIdShift = 33;
-	constexpr uint32 PackageIdMask = 0x7FFF'FFFF;
+	inline constexpr uint32 PackageIdShift = 33;
+	inline constexpr uint32 PackageIdMask = 0x7FFF'FFFF;
 
 #if UE_WITH_OBJECT_HANDLE_TRACKING
 	DECLARE_MULTICAST_DELEGATE_OneParam(FObjectHandleReadEvent, UObject* Object);
@@ -397,28 +416,6 @@ inline FObjectHandle MakeObjectHandle(FPackedObjectRef ObjectRef) { return {Obje
 inline FObjectHandle MakeObjectHandle(const FObjectRef& ObjectRef) { return MakeObjectHandle(MakePackedObjectRef(ObjectRef)); }
 inline FObjectHandle MakeObjectHandle(UObject* Object) { return {UPTRINT(Object)}; }
 
-inline bool operator==(FObjectHandle LHS, FObjectHandle RHS)
-{
-	if (IsObjectHandleResolved(LHS) == IsObjectHandleResolved(RHS))
-	{
-		return LHS.PointerOrRef == RHS.PointerOrRef;
-	}
-	else
-	{
-		return MakePackedObjectRef(LHS) == MakePackedObjectRef(RHS);
-	}
-}
-inline bool operator!=(FObjectHandle LHS, FObjectHandle RHS)
-{
-	return !(LHS == RHS);
-}
-
-inline uint32 GetTypeHash(FObjectHandle Handle)
-{
-	checkf(IsObjectHandleResolved(Handle), TEXT("Cannot hash an unresolved handle."));
-	return GetTypeHash(ReadObjectHandlePointerNoCheck(Handle));
-}
-
 #else
 
 inline bool IsObjectHandleNull(FObjectHandle Handle) { return !Handle; }
diff --git a/Engine/Source/Runtime/CoreUObject/Public/UObject/ObjectPathId.h b/Engine/Source/Runtime/CoreUObject/Public/UObject/ObjectPathId.h
--- a/Engine/Source/Runtime/CoreUObject/Public/UObject/ObjectPathId.h
+++ b/Engine/Source/Runtime/CoreUObject/Public/UObject/ObjectPathId.h
@@ -94,10 +94,8 @@ private:
 
 	uint64 PathId = static_cast<uint64>(EPathId::None);
 
-	friend uint32 GetTypeHash(FObjectPathId PathId);
+	friend FORCEINLINE uint32 GetTypeHash(FObjectPathId Value)
+	{
+		return GetTypeHash(Value.PathId);
+	}
 };
-
-FORCEINLINE uint32 GetTypeHash(FObjectPathId PathId)
-{
-	return GetTypeHash(PathId.PathId);
-}
diff --git a/Engine/Source/Runtime/CoreUObject/Public/UObject/ObjectPtr.h b/Engine/Source/Runtime/CoreUObject/Public/UObject/ObjectPtr.h
--- a/Engine/Source/Runtime/CoreUObject/Public/UObject/ObjectPtr.h
+++ b/Engine/Source/Runtime/CoreUObject/Public/UObject/ObjectPtr.h
@@ -239,8 +239,86 @@
 
 namespace ObjectPtr_Private
 {
+	/** Coerce to pointer through implicit conversion to const T* (overload through less specific "const T*" parameter to avoid ambiguity with other coercion options that may also exist. */
+	template <typename T>
+	FORCEINLINE const T* CoerceToPointer(const T* Other)
+	{
+		return Other;
+	}
+
+#if UE_OBJECT_PTR_NONCONFORMANCE_SUPPORT
+	/**
+	 * Force acceptance of explicitly TYPE_OF_NULL to avoid ambiguity issues triggered by the presence of a type conversion operator to bool.
+	 * This has the negative consequence of making it possible to coerce an arbitrary integer for the purpose of use in the comparison operators.
+	 */
+	template <typename T>
+	UE_OBJPTR_DEPRECATED(5.0, "Coercing a NULL for operations with a TObjectPtr is deprecated when running in a non-standards conforming compiler mode.")
+	constexpr const T* CoerceToPointer(TYPE_OF_NULL Other)
+	{
+		checkfSlow(Other == 0, TEXT("TObjectPtr cannot be compared to a non-zero NULL type value."));
+		return nullptr;
+	}
+#endif
+
+	/** Coerce to pointer through implicit conversion to CommonPointerType where CommonPointerType is deduced, and must be a C++ pointer, not a wrapper type. */
+	template <
+		typename T,
+		typename U,
+		typename CommonPointerType =  decltype(std::declval<bool>() ? std::declval<const T*>() : std::declval<U>()),
+		std::enable_if_t<
+			std::is_pointer<CommonPointerType>::value
+		>* = nullptr
+	>
+	FORCEINLINE auto CoerceToPointer(U&& Other) -> CommonPointerType
+	{
+		return Other;
+	}
+
+	/** Coerce to pointer through the use of a ".Get()" member, which is the convention within Unreal smart pointer types. */
+	template <
+		typename T,
+		typename U,
+		std::enable_if_t<
+			!TIsTObjectPtr<std::decay_t<U>>::Value,
+			decltype(std::declval<U>().Get())
+		>* = nullptr
+	>
+	FORCEINLINE auto CoerceToPointer(U&& Other) -> decltype(std::declval<U>().Get())
+	{
+		return Other.Get();
+	}
+
 	template <typename T, typename U>
-	FORCEINLINE bool IsObjectPtrEqualToRawPtrOfRelatedType(const TObjectPtr<T>& Ptr, const U* Other);
+	bool IsObjectPtrEqualToRawPtrOfRelatedType(const TObjectPtr<T>& Ptr, const U* Other)
+	{
+#if UE_WITH_OBJECT_HANDLE_LATE_RESOLVE
+		if (Ptr.IsResolved())
+		{
+			return Ptr.GetNoResolveNoCheck() == Other;
+		}
+		else if (!Other) //avoids resolving if Other is null
+		{
+			return !Ptr;
+		}
+#endif
+		return Ptr.GetNoReadNoCheck() == Other;
+	}
+
+	/** Perform shallow equality check between a TObjectPtr and another (non TObjectPtr) type that we can coerce to a pointer. */
+	template <
+		typename T,
+		typename U,
+		std::enable_if_t<
+			!TIsTObjectPtr<std::decay_t<U>>::Value,
+			decltype(CoerceToPointer<T>(std::declval<U>()) == std::declval<const T*>())
+		>* = nullptr
+	>
+	bool IsObjectPtrEqual(const TObjectPtr<T>& Ptr, U&& Other)
+	{
+		// This function deliberately avoids the tracking code path as we are only doing
+		// a shallow pointer comparison.
+		return IsObjectPtrEqualToRawPtrOfRelatedType<T>(Ptr, ObjectPtr_Private::CoerceToPointer<T>(Other));
+	}
 };
 
 /**
@@ -343,6 +421,7 @@
 		return *this;
 	}
 
+	// Equality/Inequality comparisons against other TObjectPtr
 	template <
 		typename U,
 		typename Base = std::decay_t<decltype(false ? std::declval<std::decay_t<T*>>() : std::declval<std::decay_t<U*>>())>
@@ -351,17 +430,6 @@
 	{
 		return ObjectPtr == Other.ObjectPtr;
 	}
-
-	bool operator==(TYPE_OF_NULLPTR) const
-	{
-		return !ObjectPtr.operator bool();
-	}
-
-	bool operator!=(TYPE_OF_NULLPTR) const
-	{
-		return ObjectPtr.operator bool();
-	}
-
 	template <
 		typename U,
 		typename Base = std::decay_t<decltype(false ? std::declval<std::decay_t<T*>>() : std::declval<std::decay_t<U*>>())>
@@ -371,4 +439,61 @@
 		return ObjectPtr != Other.ObjectPtr;
 	}
 
+
+	// Equality/Inequality comparisons against nullptr
+	FORCEINLINE bool operator==(TYPE_OF_NULLPTR) const
+	{
+		return !ObjectPtr.operator bool();
+	}
+
+	FORCEINLINE bool operator!=(TYPE_OF_NULLPTR) const
+	{
+		return ObjectPtr.operator bool();
+	}
+
+	friend FORCEINLINE bool operator==(TYPE_OF_NULLPTR, const TObjectPtr& Rhs)
+	{
+		return Rhs == nullptr;
+	}
+
+	friend FORCEINLINE bool operator!=(TYPE_OF_NULLPTR, const TObjectPtr& Rhs)
+	{
+		return Rhs != nullptr;
+	}
+
+	// Equality/Inequality comparisons against another type that can be implicitly converted to the pointer type kept in a TObjectPtr
+	template <
+		typename U,
+		typename = decltype(ObjectPtr_Private::IsObjectPtrEqual(std::declval<const TObjectPtr<T>&>(), std::declval<U&&>()))
+	>
+	friend FORCEINLINE bool operator==(const TObjectPtr<T>& Ptr, U&& Other)
+	{
+		return ObjectPtr_Private::IsObjectPtrEqual(Ptr, Other);
+	}
+	template <
+		typename U,
+		typename = decltype(ObjectPtr_Private::IsObjectPtrEqual(std::declval<const TObjectPtr<T>&>(), std::declval<U&&>()))
+	>
+	friend FORCEINLINE bool operator==(U&& Other, const TObjectPtr<T>& Ptr)
+	{
+		return ObjectPtr_Private::IsObjectPtrEqual(Ptr, Other);
+	}
+
+	template <
+		typename U,
+		typename = decltype(ObjectPtr_Private::IsObjectPtrEqual(std::declval<const TObjectPtr<T>&>(), std::declval<U&&>()))
+	>
+	friend FORCEINLINE bool operator!=(const TObjectPtr<T>& Ptr, U&& Other)
+	{
+		return !ObjectPtr_Private::IsObjectPtrEqual(Ptr, Other);
+	}
+	template <
+		typename U,
+		typename = decltype(ObjectPtr_Private::IsObjectPtrEqual(std::declval<const TObjectPtr<T>&>(), std::declval<U&&>()))
+	>
+	friend FORCEINLINE bool operator!=(U&& Other, const TObjectPtr<T>& Ptr)
+	{
+		return !ObjectPtr_Private::IsObjectPtrEqual(Ptr, Other);
+	}
+
 	// @TODO: OBJPTR: There is a risk that the FObjectPtr is storing a reference to the wrong type.  This could
@@ -452,17 +577,6 @@
 };
 
 // Equals against nullptr to optimize comparing against nullptr as it avoids resolving
-template <typename U>
-bool operator==(TYPE_OF_NULLPTR, const TObjectPtr<U>& Rhs)
-{
-	return Rhs == nullptr;
-}
 
-template <typename U>
-bool operator!=(TYPE_OF_NULLPTR, const TObjectPtr<U>& Rhs)
-{
-	return Rhs != nullptr;
-}
-
 template <typename T> struct TIsTObjectPtr<               TObjectPtr<T>> { enum { Value = true }; };
 template <typename T> struct TIsTObjectPtr<const          TObjectPtr<T>> { enum { Value = true }; };
@@ -487,125 +601,6 @@
 	template <typename T> struct TRawPointerType<const          TObjectPtr<T>> { using Type = T*; };
 	template <typename T> struct TRawPointerType<      volatile TObjectPtr<T>> { using Type = T*; };
 	template <typename T> struct TRawPointerType<const volatile TObjectPtr<T>> { using Type = T*; };
-
-	/** Coerce to pointer through implicit conversion to const T* (overload through less specific "const T*" parameter to avoid ambiguity with other coercion options that may also exist. */
-	template <typename T>
-	FORCEINLINE const T* CoerceToPointer(const T* Other)
-	{
-		return Other;
-	}
-
-#if UE_OBJECT_PTR_NONCONFORMANCE_SUPPORT
-	/**
-	 * Force acceptance of explicitly TYPE_OF_NULL to avoid ambiguity issues triggered by the presence of a type conversion operator to bool.
-	 * This has the negative consequence of making it possible to coerce an arbitrary integer for the purpose of use in the comparison operators.
-	 */
-	template <typename T>
-	UE_OBJPTR_DEPRECATED(5.0, "Coercing a NULL for operations with a TObjectPtr is deprecated when running in a non-standards conforming compiler mode.")
-	constexpr const T* CoerceToPointer(TYPE_OF_NULL Other)
-	{
-		checkfSlow(Other == 0, TEXT("TObjectPtr cannot be compared to a non-zero NULL type value."));
-		return nullptr;
-	}
-#endif
-
-	/** Coerce to pointer through implicit conversion to CommonPointerType where CommonPointerType is deduced, and must be a C++ pointer, not a wrapper type. */
-	template <
-		typename T,
-		typename U,
-		typename CommonPointerType =  decltype(std::declval<bool>() ? std::declval<const T*>() : std::declval<U>()),
-		std::enable_if_t<
-			std::is_pointer<CommonPointerType>::value
-		>* = nullptr
-	>
-	FORCEINLINE auto CoerceToPointer(U&& Other) -> CommonPointerType
-	{
-		return Other;
-	}
-
-	/** Coerce to pointer through the use of a ".Get()" member, which is the convention within Unreal smart pointer types. */
-	template <
-		typename T,
-		typename U,
-		std::enable_if_t<
-			!TIsTObjectPtr<std::decay_t<U>>::Value,
-			decltype(std::declval<U>().Get())
-		>* = nullptr
-	>
-	FORCEINLINE auto CoerceToPointer(U&& Other) -> decltype(std::declval<U>().Get())
-	{
-		return Other.Get();
-	}
-
-	template <typename T, typename U>
-	bool IsObjectPtrEqualToRawPtrOfRelatedType(const TObjectPtr<T>& Ptr, const U* Other)
-	{
-#if UE_WITH_OBJECT_HANDLE_LATE_RESOLVE
-		if (Ptr.IsResolved())
-		{
-			return Ptr.GetNoResolveNoCheck() == Other;
-		}
-		else if (!Other) //avoids resolving if Other is null
-		{
-			return !Ptr;
-		}
-#endif
-		return Ptr.GetNoReadNoCheck() == Other;
-	}
-
-	/** Perform shallow equality check between a TObjectPtr and another (non TObjectPtr) type that we can coerce to a pointer. */
-	template <
-		typename T,
-		typename U,
-		std::enable_if_t<
-			!TIsTObjectPtr<std::decay_t<U>>::Value,
-			decltype(CoerceToPointer<T>(std::declval<U>()) == std::declval<const T*>())
-		>* = nullptr
-	>
-	bool IsObjectPtrEqual(const TObjectPtr<T>& Ptr, U&& Other)
-	{
-		// This function deliberately avoids the tracking code path as we are only doing
-		// a shallow pointer comparison.
-		return IsObjectPtrEqualToRawPtrOfRelatedType<T>(Ptr, ObjectPtr_Private::CoerceToPointer<T>(Other));
-	}
 }
 
-// Equality/Inequality comparisons against another type that can be implicitly converted to the pointer type kept in a TObjectPtr
-template <
-	typename T,
-	typename U,
-	decltype(ObjectPtr_Private::IsObjectPtrEqual(std::declval<const TObjectPtr<T>&>(), std::declval<U&&>()))* = nullptr
->
-FORCEINLINE bool operator==(const TObjectPtr<T>& Ptr, U&& Other)
-{
-	return ObjectPtr_Private::IsObjectPtrEqual(Ptr, Other);
-}
-template <
-	typename T,
-	typename U,
-	decltype(ObjectPtr_Private::IsObjectPtrEqual(std::declval<const TObjectPtr<T>&>(), std::declval<U&&>()))* = nullptr
->
-FORCEINLINE bool operator==(U&& Other, const TObjectPtr<T>& Ptr)
-{
-	return ObjectPtr_Private::IsObjectPtrEqual(Ptr, Other);
-}
-template <
-	typename T,
-	typename U,
-	decltype(ObjectPtr_Private::IsObjectPtrEqual(std::declval<const TObjectPtr<T>&>(), std::declval<U&&>()))* = nullptr
->
-FORCEINLINE bool operator!=(const TObjectPtr<T>& Ptr, U&& Other)
-{
-	return !ObjectPtr_Private::IsObjectPtrEqual(Ptr, Other);
-}
-template <
-	typename T,
-	typename U,
-	decltype(ObjectPtr_Private::IsObjectPtrEqual(std::declval<const TObjectPtr<T>&>(), std::declval<U&&>()))* = nullptr
->
-FORCEINLINE bool operator!=(U&& Other, const TObjectPtr<T>& Ptr)
-{
-	return !ObjectPtr_Private::IsObjectPtrEqual(Ptr, Other);
-}
-
 template <typename T>
diff --git a/Engine/Source/Runtime/CoreUObject/Public/UObject/StrongObjectPtr.h b/Engine/Source/Runtime/CoreUObject/Public/UObject/StrongObjectPtr.h
--- a/Engine/Source/Runtime/CoreUObject/Public/UObject/StrongObjectPtr.h
+++ b/Engine/Source/Runtime/CoreUObject/Public/UObject/StrongObjectPtr.h
@@ -178,16 +178,17 @@ public:
 
 private:
 	TUniquePtr< UEStrongObjectPtr_Private::TInternalReferenceCollector<ReferencerNameProvider> > ReferenceCollector;
-};
 
-template <typename LHSObjectType, typename LHSReferencerNameProvider, typename RHSObjectType, typename RHSReferencerNameProvider>
-FORCEINLINE bool operator==(const TStrongObjectPtr<LHSObjectType, LHSReferencerNameProvider>& InLHS, const TStrongObjectPtr<RHSObjectType, RHSReferencerNameProvider>& InRHS)
-{
-	return InLHS.Get() == InRHS.Get();
-}
+	template <typename RHSObjectType, typename RHSReferencerNameProvider>
+	friend FORCEINLINE bool operator==(const TStrongObjectPtr<ObjectType, ReferencerNameProvider>& InLHS, const TStrongObjectPtr<RHSObjectType, RHSReferencerNameProvider>& InRHS)
+	{
+		return InLHS.Get() == InRHS.Get();
+	}
 
-template <typename LHSObjectType, typename LHSReferencerNameProvider, typename RHSObjectType, typename RHSReferencerNameProvider>
-FORCEINLINE bool operator!=(const TStrongObjectPtr<LHSObjectType, LHSReferencerNameProvider>& InLHS, const TStrongObjectPtr<RHSObjectType, RHSReferencerNameProvider>& InRHS)
-{
-	return InLHS.Get() != InRHS.Get();
-}
+	template <typename RHSObjectType, typename RHSReferencerNameProvider>
+	friend FORCEINLINE bool operator!=(const TStrongObjectPtr<ObjectType, ReferencerNameProvider>& InLHS, const TStrongObjectPtr<RHSObjectType, RHSReferencerNameProvider>& InRHS)
+	{
+		return InLHS.Get() != InRHS.Get();
+	}
+
+};
diff --git a/Engine/Source/Runtime/Engine/Classes/AI/Navigation/NavigationTypes.h b/Engine/Source/Runtime/Engine/Classes/AI/Navigation/NavigationTypes.h
--- a/Engine/Source/Runtime/Engine/Classes/AI/Navigation/NavigationTypes.h
+++ b/Engine/Source/Runtime/Engine/Classes/AI/Navigation/NavigationTypes.h
@@ -514,12 +514,12 @@ struct ENGINE_API FNavAgentProperties : public FMovementProperties
 	void SetPreferredNavData(TSubclassOf<AActor> NavDataClass);
 
 	static const FNavAgentProperties DefaultProperties;
-};
 
-inline uint32 GetTypeHash(const FNavAgentProperties& A)
-{
-	return ((int16(A.AgentRadius) << 16) | int16(A.AgentHeight)) ^ int32(A.AgentStepHeight);
-}
+	friend inline uint32 GetTypeHash(const FNavAgentProperties& A)
+	{
+		return ((int16(A.AgentRadius) << 16) | int16(A.AgentHeight)) ^ int32(A.AgentStepHeight);
+	}
+};
 
 PRAGMA_DISABLE_DEPRECATION_WARNINGS
 USTRUCT(BlueprintType)
diff --git a/Engine/Source/Runtime/Engine/Classes/Curves/RichCurve.h b/Engine/Source/Runtime/Engine/Classes/Curves/RichCurve.h
--- a/Engine/Source/Runtime/Engine/Classes/Curves/RichCurve.h
+++ b/Engine/Source/Runtime/Engine/Classes/Curves/RichCurve.h
@@ -463,16 +463,15 @@ struct FRichCurveEditInfoTemplate
 
 	uint32 GetTypeHash() const
 	{
-		return HashCombine(::GetTypeHash(CurveName), PointerHash(CurveToEdit));
+		return HashCombine(GetTypeHashHelper(CurveName), PointerHash(CurveToEdit));
 	}
-};
 
+	friend inline uint32 GetTypeHash(const FRichCurveEditInfoTemplate<T>& RichCurveEditInfo)
+	{
+		return RichCurveEditInfo.GetTypeHash();
+	}
+};
 
-template<class T>
-inline uint32 GetTypeHash(const FRichCurveEditInfoTemplate<T>& RichCurveEditInfo)
-{
-	return RichCurveEditInfo.GetTypeHash();
-}
 
 // Rename and deprecate
 typedef FRichCurveEditInfoTemplate<FRealCurve*>			FRichCurveEditInfo;
diff --git a/Engine/Source/Runtime/Engine/Classes/EdGraph/EdGraphNode.h b/Engine/Source/Runtime/Engine/Classes/EdGraph/EdGraphNode.h
--- a/Engine/Source/Runtime/Engine/Classes/EdGraph/EdGraphNode.h
+++ b/Engine/Source/Runtime/Engine/Classes/EdGraph/EdGraphNode.h
@@ -75,6 +75,20 @@ struct ENGINE_API FEdGraphTerminalType
 	static FEdGraphTerminalType FromPinType(const FEdGraphPinType& PinType);
 
 	friend FArchive& operator<<(FArchive& Ar, FEdGraphTerminalType& P);
+
+	friend inline bool operator!= (const FEdGraphTerminalType& A, const FEdGraphTerminalType& B)
+	{
+		return A.TerminalCategory != B.TerminalCategory
+			|| A.TerminalSubCategory != B.TerminalSubCategory
+			|| A.TerminalSubCategoryObject != B.TerminalSubCategoryObject
+			|| A.bTerminalIsConst != B.bTerminalIsConst
+			|| A.bTerminalIsWeakPointer != B.bTerminalIsWeakPointer;
+	}
+
+	friend inline bool operator==(const FEdGraphTerminalType& A, const FEdGraphTerminalType& B)
+	{
+		return !(A != B);
+	}
 };
 
 /** Enum used to define which way data flows into or out of this pin. */
diff --git a/Engine/Source/Runtime/Engine/Classes/EdGraph/EdGraphPin.h b/Engine/Source/Runtime/Engine/Classes/EdGraph/EdGraphPin.h
--- a/Engine/Source/Runtime/Engine/Classes/EdGraph/EdGraphPin.h
+++ b/Engine/Source/Runtime/Engine/Classes/EdGraph/EdGraphPin.h
@@ -71,20 +71,6 @@ FORCEINLINE FArchive& operator<<(FArchive& Ar, FSimpleMemberReference& Data)
 	return Ar;
 }
 
-inline bool operator!= (const FEdGraphTerminalType& A, const FEdGraphTerminalType& B)
-{
-	return A.TerminalCategory != B.TerminalCategory
-		|| A.TerminalSubCategory != B.TerminalSubCategory
-		|| A.TerminalSubCategoryObject != B.TerminalSubCategoryObject
-		|| A.bTerminalIsConst != B.bTerminalIsConst
-		|| A.bTerminalIsWeakPointer != B.bTerminalIsWeakPointer;
-}
-
-inline bool operator==(const FEdGraphTerminalType& A, const FEdGraphTerminalType& B)
-{
-	return !(A != B);
-}
-
 /** Struct used to define the type of information carried on this pin */
 USTRUCT(BlueprintType)
 struct FEdGraphPinType
diff --git a/Engine/Source/Runtime/Engine/Classes/Engine/CurveTable.h b/Engine/Source/Runtime/Engine/Classes/Engine/CurveTable.h
--- a/Engine/Source/Runtime/Engine/Classes/Engine/CurveTable.h
+++ b/Engine/Source/Runtime/Engine/Classes/Engine/CurveTable.h
@@ -317,7 +317,7 @@ struct ENGINE_API FCurveTableRowHandle
 	/** Used so we can have a TMap of this struct */
 	FORCEINLINE friend uint32 GetTypeHash(const FCurveTableRowHandle& Handle)
 	{
-		return HashCombine(::GetTypeHash(Handle.RowName), PointerHash(Handle.CurveTable));
+		return HashCombine(GetTypeHash(Handle.RowName), PointerHash(Handle.CurveTable));
 	}
 };
 
diff --git a/Engine/Source/Runtime/Engine/Classes/Engine/LODActor.h b/Engine/Source/Runtime/Engine/Classes/Engine/LODActor.h
--- a/Engine/Source/Runtime/Engine/Classes/Engine/LODActor.h
+++ b/Engine/Source/Runtime/Engine/Classes/Engine/LODActor.h
@@ -53,7 +53,7 @@ struct FHLODInstancingKey
 
 	friend uint32 GetTypeHash(const FHLODInstancingKey& Key)
 	{
-		return HashCombine(::GetTypeHash(Key.StaticMesh), ::GetTypeHash(Key.Material));
+		return HashCombine(GetTypeHash(Key.StaticMesh), GetTypeHash(Key.Material));
 	}
 };
 
diff --git a/Engine/Source/Runtime/Engine/Classes/Engine/StaticMesh.h b/Engine/Source/Runtime/Engine/Classes/Engine/StaticMesh.h
--- a/Engine/Source/Runtime/Engine/Classes/Engine/StaticMesh.h
+++ b/Engine/Source/Runtime/Engine/Classes/Engine/StaticMesh.h
@@ -2,6 +2,7 @@
 
 #pragma once
 
+#include "Async/AsyncWork.h"
 #include "CoreMinimal.h"
 #include "UObject/ObjectMacros.h"
 #include "UObject/UObjectGlobals.h"
@@ -386,11 +387,11 @@ struct FMeshSectionInfo
 		, bForceOpaque(false)
 	{
 	}
-};
 
-/** Comparison for mesh section info. */
-bool operator==(const FMeshSectionInfo& A, const FMeshSectionInfo& B);
-bool operator!=(const FMeshSectionInfo& A, const FMeshSectionInfo& B);
+	/** Comparison for mesh section info. */
+	friend bool operator==(const FMeshSectionInfo& A, const FMeshSectionInfo& B);
+	friend bool operator!=(const FMeshSectionInfo& A, const FMeshSectionInfo& B);
+};
 
 /**
  * Map containing per-section settings for each section of each LOD.
diff --git a/Engine/Source/Runtime/Engine/Classes/Materials/MaterialLayersFunctions.h b/Engine/Source/Runtime/Engine/Classes/Materials/MaterialLayersFunctions.h
--- a/Engine/Source/Runtime/Engine/Classes/Materials/MaterialLayersFunctions.h
+++ b/Engine/Source/Runtime/Engine/Classes/Materials/MaterialLayersFunctions.h
@@ -380,13 +380,14 @@ private:
 
 	UPROPERTY()
 	TArray<FGuid> DeletedParentLayerGuids_DEPRECATED;
-};
 
-// Don't allowing comparing a full FMaterialLayersFunctions against partial RuntimeData
-bool operator==(const FMaterialLayersFunctions&, const FMaterialLayersFunctionsRuntimeData&) = delete;
-bool operator==(const FMaterialLayersFunctionsRuntimeData&, const FMaterialLayersFunctions&) = delete;
-bool operator!=(const FMaterialLayersFunctions&, const FMaterialLayersFunctionsRuntimeData&) = delete;
-bool operator!=(const FMaterialLayersFunctionsRuntimeData&, const FMaterialLayersFunctions&) = delete;
+
+	// Don't allowing comparing a full FMaterialLayersFunctions against partial RuntimeData
+	friend bool operator==(const FMaterialLayersFunctions&, const FMaterialLayersFunctionsRuntimeData&) = delete;
+	friend bool operator==(const FMaterialLayersFunctionsRuntimeData&, const FMaterialLayersFunctions&) = delete;
+	friend bool operator!=(const FMaterialLayersFunctions&, const FMaterialLayersFunctionsRuntimeData&) = delete;
+	friend bool operator!=(const FMaterialLayersFunctionsRuntimeData&, const FMaterialLayersFunctions&) = delete;
+};
 
 template<>
 struct TStructOpsTypeTraits<FMaterialLayersFunctions> : TStructOpsTypeTraitsBase2<FMaterialLayersFunctions>
diff --git a/Engine/Source/Runtime/Engine/Classes/Sound/DialogueTypes.h b/Engine/Source/Runtime/Engine/Classes/Sound/DialogueTypes.h
--- a/Engine/Source/Runtime/Engine/Classes/Sound/DialogueTypes.h
+++ b/Engine/Source/Runtime/Engine/Classes/Sound/DialogueTypes.h
@@ -56,10 +56,10 @@ struct ENGINE_API FDialogueContext
 
 	/** Gets a generated hash created from the source and targets. */
 	FString GetContextHash() const;
-};
 
-ENGINE_API bool operator==(const FDialogueContext& LHS, const FDialogueContext& RHS);
-ENGINE_API bool operator!=(const FDialogueContext& LHS, const FDialogueContext& RHS);
+	friend ENGINE_API bool operator==(const FDialogueContext& LHS, const FDialogueContext& RHS);
+	friend ENGINE_API bool operator!=(const FDialogueContext& LHS, const FDialogueContext& RHS);
+};
 
 USTRUCT()
 struct ENGINE_API FDialogueWaveParameter
diff --git a/Engine/Source/Runtime/Engine/Private/DebugRenderSceneProxy.cpp b/Engine/Source/Runtime/Engine/Private/DebugRenderSceneProxy.cpp
--- a/Engine/Source/Runtime/Engine/Private/DebugRenderSceneProxy.cpp
+++ b/Engine/Source/Runtime/Engine/Private/DebugRenderSceneProxy.cpp
@@ -137,7 +137,7 @@ void FDebugRenderSceneProxy::GetDynamicMeshElements(const TArray<const FSceneVie
 		FMaterialRenderProxy* operator[](FLinearColor Color)
 		{
 			FMaterialRenderProxy* MeshColor = NULL;
-			const uint32 HashKey = ::GetTypeHash(Color);
+			const uint32 HashKey = GetTypeHashHelper(Color);
 			if (MeshColorInstances.Contains(HashKey))
 			{
 				MeshColor = *MeshColorInstances.Find(HashKey);
diff --git a/Engine/Source/Runtime/Engine/Public/Audio.h b/Engine/Source/Runtime/Engine/Public/Audio.h
--- a/Engine/Source/Runtime/Engine/Public/Audio.h
+++ b/Engine/Source/Runtime/Engine/Public/Audio.h
@@ -497,9 +497,9 @@ public:
 
 	/** Whether spatialization is an external send */
 	bool GetSpatializationIsExternalSend() const {	return bSpatializationIsExternalSend;}
-};
 
-inline uint32 GetTypeHash(FWaveInstance* A) { return A->TypeHash; }
+	friend inline uint32 GetTypeHash(FWaveInstance* A) { return A->TypeHash; }
+};
 
 /*-----------------------------------------------------------------------------
 	FSoundBuffer.
diff --git a/Engine/Source/Runtime/Engine/Public/Folder.h b/Engine/Source/Runtime/Engine/Public/Folder.h
--- a/Engine/Source/Runtime/Engine/Public/Folder.h
+++ b/Engine/Source/Runtime/Engine/Public/Folder.h
@@ -179,12 +179,10 @@ private:
 	mutable FName Path;
 	FRootObject RootObject;
 	mutable FGuid ActorFolderGuid; // Optional : Used to find Level's ActorFolder faster than by using the path
+
+	friend FORCEINLINE uint32 GetTypeHash(const FFolder& InFolder)
+	{
+		return HashCombine(GetTypeHash(InFolder.GetPath()), GetTypeHash(InFolder.GetRootObject()));
+	}
 #endif
 };
-
-#if WITH_EDITOR
-FORCEINLINE uint32 GetTypeHash(const FFolder& InFolder)
-{
-	return HashCombine(GetTypeHash(InFolder.GetPath()), GetTypeHash(InFolder.GetRootObject()));
-}
-#endif
\ No newline at end of file
diff --git a/Engine/Source/Runtime/Engine/Public/MaterialShared.h b/Engine/Source/Runtime/Engine/Public/MaterialShared.h
--- a/Engine/Source/Runtime/Engine/Public/MaterialShared.h
+++ b/Engine/Source/Runtime/Engine/Public/MaterialShared.h
@@ -366,6 +366,11 @@ public:
 		return true;
 	}
 
+	friend inline bool operator!=(const FMaterialVirtualTextureStack& Lhs, const FMaterialVirtualTextureStack& Rhs)
+	{
+		return !operator==(Lhs, Rhs);
+	}
+
 	/** Number of layers that have been allocated in this stack. */
 	LAYOUT_FIELD(uint32, NumLayers);
 	/** Indices of the expressions that were set to layers in this stack. */
@@ -374,11 +379,6 @@ public:
 	LAYOUT_FIELD(int32, PreallocatedStackTextureIndex);
 };
 
-inline bool operator!=(const FMaterialVirtualTextureStack& Lhs, const FMaterialVirtualTextureStack& Rhs)
-{
-	return !operator==(Lhs, Rhs);
-}
-
 class FMaterialUniformPreshaderField
 {
 	DECLARE_TYPE_LAYOUT(FMaterialUniformPreshaderField, NonVirtual);
diff --git a/Engine/Source/Runtime/Engine/Public/MaterialTypes.h b/Engine/Source/Runtime/Engine/Public/MaterialTypes.h
--- a/Engine/Source/Runtime/Engine/Public/MaterialTypes.h
+++ b/Engine/Source/Runtime/Engine/Public/MaterialTypes.h
@@ -76,6 +76,21 @@ struct ENGINE_API FMaterialParameterInfo
 	}
 
 	bool RemapLayerIndex(TArrayView<const int32> IndexRemap, FMaterialParameterInfo& OutResult) const;
+
+	friend FORCEINLINE bool operator==(const FMaterialParameterInfo& Lhs, const FMaterialParameterInfo& Rhs)
+	{
+		return Lhs.Name.IsEqual(Rhs.Name) && Lhs.Association == Rhs.Association && Lhs.Index == Rhs.Index;
+	}
+
+	friend FORCEINLINE bool operator!=(const FMaterialParameterInfo& Lhs, const FMaterialParameterInfo& Rhs)
+	{
+		return !operator==(Lhs, Rhs);
+	}
+
+	friend FORCEINLINE uint32 GetTypeHash(const FMaterialParameterInfo& Value)
+	{
+		return HashCombine(HashCombine(GetTypeHash(NameToScriptName(Value.Name)), Value.Index), (uint32)Value.Association);
+	}
 };
 
 struct FMemoryImageMaterialParameterInfo
@@ -123,58 +138,48 @@ public:
 	LAYOUT_FIELD(FScriptName, Name);
 	LAYOUT_FIELD(int32, Index);
 	LAYOUT_FIELD(TEnumAsByte<EMaterialParameterAssociation>, Association);
-};
-
-FORCEINLINE FMaterialParameterInfo::FMaterialParameterInfo(const struct FMemoryImageMaterialParameterInfo& Rhs)
-	: Name(ScriptNameToName(Rhs.Name))
-	, Association(Rhs.Association)
-	, Index(Rhs.Index)
-{
-}
-
-FORCEINLINE bool operator==(const FMaterialParameterInfo& Lhs, const FMaterialParameterInfo& Rhs)
-{
-	return Lhs.Name.IsEqual(Rhs.Name) && Lhs.Association == Rhs.Association && Lhs.Index == Rhs.Index;
-}
 
-FORCEINLINE bool operator!=(const FMaterialParameterInfo& Lhs, const FMaterialParameterInfo& Rhs)
-{
-	return !operator==(Lhs, Rhs);
-}
+	friend FORCEINLINE bool operator==(const FMemoryImageMaterialParameterInfo& Lhs, const FMemoryImageMaterialParameterInfo& Rhs)
+	{
+		return Lhs.Name == Rhs.Name && Lhs.Association == Rhs.Association && Lhs.Index == Rhs.Index;
+	}
 
-FORCEINLINE bool operator==(const FMemoryImageMaterialParameterInfo& Lhs, const FMemoryImageMaterialParameterInfo& Rhs)
-{
-	return Lhs.Name == Rhs.Name && Lhs.Association == Rhs.Association && Lhs.Index == Rhs.Index;
-}
+	friend FORCEINLINE bool operator!=(const FMemoryImageMaterialParameterInfo& Lhs, const FMemoryImageMaterialParameterInfo& Rhs)
+	{
+		return !operator==(Lhs, Rhs);
+	}
 
-FORCEINLINE bool operator!=(const FMemoryImageMaterialParameterInfo& Lhs, const FMemoryImageMaterialParameterInfo& Rhs)
-{
-	return !operator==(Lhs, Rhs);
-}
+	friend FORCEINLINE bool operator==(const FMaterialParameterInfo& Lhs, const FMemoryImageMaterialParameterInfo& Rhs)
+	{
+		return Lhs.Name == Rhs.Name && Lhs.Index == Rhs.Index && Lhs.Association == Rhs.Association;
+	}
 
-FORCEINLINE bool operator==(const FMaterialParameterInfo& Lhs, const FMemoryImageMaterialParameterInfo& Rhs)
-{
-	return Lhs.Name == Rhs.Name && Lhs.Index == Rhs.Index && Lhs.Association == Rhs.Association;
-}
+	friend FORCEINLINE bool operator!=(const FMaterialParameterInfo& Lhs, const FMemoryImageMaterialParameterInfo& Rhs)
+	{
+		return !operator==(Lhs, Rhs);
+	}
 
-FORCEINLINE bool operator!=(const FMaterialParameterInfo& Lhs, const FMemoryImageMaterialParameterInfo& Rhs)
-{
-	return !operator==(Lhs, Rhs);
-}
+	friend FORCEINLINE bool operator==(const FMemoryImageMaterialParameterInfo& Lhs, const FMaterialParameterInfo& Rhs)
+	{
+		return Lhs.Name == Rhs.Name && Lhs.Index == Rhs.Index && Lhs.Association == Rhs.Association;
+	}
 
-FORCEINLINE bool operator==(const FMemoryImageMaterialParameterInfo& Lhs, const FMaterialParameterInfo& Rhs)
-{
-	return Lhs.Name == Rhs.Name && Lhs.Index == Rhs.Index && Lhs.Association == Rhs.Association;
-}
+	friend FORCEINLINE bool operator!=(const FMemoryImageMaterialParameterInfo& Lhs, const FMaterialParameterInfo& Rhs)
+	{
+		return !operator==(Lhs, Rhs);
+	}
 
-FORCEINLINE bool operator!=(const FMemoryImageMaterialParameterInfo& Lhs, const FMaterialParameterInfo& Rhs)
-{
-	return !operator==(Lhs, Rhs);
-}
+	friend FORCEINLINE uint32 GetTypeHash(const FMemoryImageMaterialParameterInfo& Value)
+	{
+		return HashCombine(HashCombine(GetTypeHash(Value.Name), Value.Index), (uint32)Value.Association);
+	}
+};
 
-FORCEINLINE uint32 GetTypeHash(const FMemoryImageMaterialParameterInfo& Value)
+FORCEINLINE FMaterialParameterInfo::FMaterialParameterInfo(const struct FMemoryImageMaterialParameterInfo& Rhs)
+	: Name(ScriptNameToName(Rhs.Name))
+	, Association(Rhs.Association)
+	, Index(Rhs.Index)
 {
-	return HashCombine(HashCombine(GetTypeHash(Value.Name), Value.Index), (uint32)Value.Association);
 }
 
 // Backwards compat
@@ -355,47 +360,47 @@ struct FMaterialParameterValue
 		} Font;
 	};
 	EMaterialParameterType Type;
-};
 
-inline bool operator==(const FMaterialParameterValue& Lhs, const FMaterialParameterValue& Rhs)
-{
-	const EMaterialParameterType Type = Lhs.Type;
-	if (Type != Rhs.Type)
+	friend inline bool operator==(const FMaterialParameterValue& Lhs, const FMaterialParameterValue& Rhs)
 	{
-		return false;
+		const EMaterialParameterType ParameterType = Lhs.Type;
+		if (ParameterType != Rhs.Type)
+		{
+			return false;
+		}
+		switch (ParameterType)
+		{
+		case EMaterialParameterType::None: return true;
+		case EMaterialParameterType::Scalar: return Lhs.Float[0] == Rhs.Float[0];
+		case EMaterialParameterType::Vector: return
+			Lhs.Float[0] == Rhs.Float[0] &&
+			Lhs.Float[1] == Rhs.Float[1] &&
+			Lhs.Float[2] == Rhs.Float[2] &&
+			Lhs.Float[3] == Rhs.Float[3];
+		case EMaterialParameterType::DoubleVector: return
+			Lhs.Double[0] == Rhs.Double[0] &&
+			Lhs.Double[1] == Rhs.Double[1] &&
+			Lhs.Double[2] == Rhs.Double[2] &&
+			Lhs.Double[3] == Rhs.Double[3];
+		case EMaterialParameterType::Texture: return Lhs.Texture == Rhs.Texture;
+		case EMaterialParameterType::Font: return Lhs.Font.Value == Rhs.Font.Value && Lhs.Font.Page == Rhs.Font.Page;
+		case EMaterialParameterType::RuntimeVirtualTexture: return Lhs.RuntimeVirtualTexture == Rhs.RuntimeVirtualTexture;
+		case EMaterialParameterType::SparseVolumeTexture: return Lhs.SparseVolumeTexture == Rhs.SparseVolumeTexture;
+		case EMaterialParameterType::StaticSwitch: return Lhs.Bool[0] == Rhs.Bool[0];
+		case EMaterialParameterType::StaticComponentMask: return
+			Lhs.Bool[0] == Rhs.Bool[0] &&
+			Lhs.Bool[1] == Rhs.Bool[1] &&
+			Lhs.Bool[2] == Rhs.Bool[2] &&
+			Lhs.Bool[3] == Rhs.Bool[3];
+		default: checkNoEntry(); return false;
+		}
 	}
-	switch (Type)
+
+	friend inline bool operator!=(const FMaterialParameterValue& Lhs, const FMaterialParameterValue& Rhs)
 	{
-	case EMaterialParameterType::None: return true;
-	case EMaterialParameterType::Scalar: return Lhs.Float[0] == Rhs.Float[0];
-	case EMaterialParameterType::Vector: return
-		Lhs.Float[0] == Rhs.Float[0] &&
-		Lhs.Float[1] == Rhs.Float[1] &&
-		Lhs.Float[2] == Rhs.Float[2] &&
-		Lhs.Float[3] == Rhs.Float[3];
-	case EMaterialParameterType::DoubleVector: return
-		Lhs.Double[0] == Rhs.Double[0] &&
-		Lhs.Double[1] == Rhs.Double[1] &&
-		Lhs.Double[2] == Rhs.Double[2] &&
-		Lhs.Double[3] == Rhs.Double[3];
-	case EMaterialParameterType::Texture: return Lhs.Texture == Rhs.Texture;
-	case EMaterialParameterType::Font: return Lhs.Font.Value == Rhs.Font.Value && Lhs.Font.Page == Rhs.Font.Page;
-	case EMaterialParameterType::RuntimeVirtualTexture: return Lhs.RuntimeVirtualTexture == Rhs.RuntimeVirtualTexture;
-	case EMaterialParameterType::SparseVolumeTexture: return Lhs.SparseVolumeTexture == Rhs.SparseVolumeTexture;
-	case EMaterialParameterType::StaticSwitch: return Lhs.Bool[0] == Rhs.Bool[0];
-	case EMaterialParameterType::StaticComponentMask: return
-		Lhs.Bool[0] == Rhs.Bool[0] &&
-		Lhs.Bool[1] == Rhs.Bool[1] &&
-		Lhs.Bool[2] == Rhs.Bool[2] &&
-		Lhs.Bool[3] == Rhs.Bool[3];
-	default: checkNoEntry(); return false;
+		return !operator==(Lhs, Rhs);
 	}
-}
-
-inline bool operator!=(const FMaterialParameterValue& Lhs, const FMaterialParameterValue& Rhs)
-{
-	return !operator==(Lhs, Rhs);
-}
+};
 
 /** Holds a value, along with editor-only metadata that describes that value */
 struct FMaterialParameterMetadata
diff --git a/Engine/Source/Runtime/Engine/Public/Net/Util/SequenceNumber.h b/Engine/Source/Runtime/Engine/Public/Net/Util/SequenceNumber.h
--- a/Engine/Source/Runtime/Engine/Public/Net/Util/SequenceNumber.h
+++ b/Engine/Source/Runtime/Engine/Public/Net/Util/SequenceNumber.h
@@ -51,6 +51,47 @@ public:
 private:
 	void Increment(SequenceT InValue) { *this = TSequenceNumber(Value + InValue); }
 	SequenceT Value;
+
+	friend bool operator<(const TSequenceNumber& Lhs, const TSequenceNumber& Rhs)
+	{
+		return !(Lhs >= Rhs);
+	}
+
+	friend bool operator<=(const TSequenceNumber& Lhs, const TSequenceNumber& Rhs)
+	{
+		return !(Lhs > Rhs);
+	}
+
+	/** Equals, NOTE that sequence numbers wrap around so 0 == 0 + SequenceNumberCount */
+	friend bool operator==(const TSequenceNumber& Lhs, const TSequenceNumber& Rhs)
+	{ 
+		return Lhs.Get() == Rhs.Get(); 
+	}
+
+	friend bool operator!=(const TSequenceNumber& Lhs, const TSequenceNumber& Rhs)
+	{
+		return Lhs.Get() != Rhs.Get();
+	}
+
+	friend const TSequenceNumber operator+(const TSequenceNumber& Lhs, const TSequenceNumber& Rhs)
+	{
+		return TSequenceNumber(Lhs.Get() + Rhs.Get());
+	}
+
+	friend const TSequenceNumber operator-(const TSequenceNumber& Lhs, const TSequenceNumber& Rhs)
+	{
+		return TSequenceNumber(Lhs.Get() - Rhs.Get());
+	}
+
+	friend const TSequenceNumber operator+(const TSequenceNumber& Lhs, SequenceType Rhs)
+	{
+		return TSequenceNumber(Lhs.Get() + Rhs);
+	}
+
+	friend const TSequenceNumber operator-(const TSequenceNumber& Lhs, SequenceType Rhs)
+	{
+		return TSequenceNumber(Lhs.Get() - Rhs);
+	}
 };
 
 template <SIZE_T NumBits, typename SequenceType>
@@ -63,64 +104,3 @@ typename TSequenceNumber<NumBits, SequenceType>::DifferenceT TSequenceNumber<Num
 
 	return (DifferenceT)((ValueA - ValueB) << ShiftValue) >> ShiftValue;
 };
-
-template <SIZE_T NumBits, typename SequenceType>
-bool operator>(const TSequenceNumber<NumBits, SequenceType>& Lhs, const TSequenceNumber<NumBits, SequenceType>& Rhs)
-{ 
-	return Lhs > Rhs; 
-}
-
-template <SIZE_T NumBits, typename SequenceType>
-bool operator>=(const TSequenceNumber<NumBits, SequenceType>& Lhs, const TSequenceNumber<NumBits, SequenceType>& Rhs)
-{
-	return Lhs >= Rhs;
-}
-
-template <SIZE_T NumBits, typename SequenceType>
-bool operator<(const TSequenceNumber<NumBits, SequenceType>& Lhs, const TSequenceNumber<NumBits, SequenceType>& Rhs)
-{
-	return !(Lhs >= Rhs);
-}
-
-template <SIZE_T NumBits, typename SequenceType>
-bool operator<=(const TSequenceNumber<NumBits, SequenceType>& Lhs, const TSequenceNumber<NumBits, SequenceType>& Rhs)
-{
-	return !(Lhs > Rhs);
-}
-
-/** Equals, NOTE that sequence numbers wrap around so 0 == 0 + SequenceNumberCount */
-template <SIZE_T NumBits, typename SequenceType>
-bool operator==(const TSequenceNumber<NumBits, SequenceType>& Lhs, const TSequenceNumber<NumBits, SequenceType>& Rhs)
-{ 
-	return Lhs.Get() == Rhs.Get(); 
-}
-
-template <SIZE_T NumBits, typename SequenceType>
-bool operator!=(const TSequenceNumber<NumBits, SequenceType>& Lhs, const TSequenceNumber<NumBits, SequenceType>& Rhs)
-{
-	return Lhs.Get() != Rhs.Get();
-}
-
-template <SIZE_T NumBits, typename SequenceType>
-const TSequenceNumber<NumBits, SequenceType> operator+(const TSequenceNumber<NumBits, SequenceType>& Lhs, const TSequenceNumber<NumBits, SequenceType>& Rhs)
-{
-	return TSequenceNumber<NumBits, SequenceType>(Lhs.Get() + Rhs.Get());
-}
-
-template <SIZE_T NumBits, typename SequenceType>
-const TSequenceNumber<NumBits, SequenceType> operator-(const TSequenceNumber<NumBits, SequenceType>& Lhs, const TSequenceNumber<NumBits, SequenceType>& Rhs)
-{
-	return TSequenceNumber<NumBits, SequenceType>(Lhs.Get() - Rhs.Get());
-}
-
-template <SIZE_T NumBits, typename SequenceType>
-const TSequenceNumber<NumBits, SequenceType> operator+(const TSequenceNumber<NumBits, SequenceType>& Lhs, SequenceType Rhs)
-{
-	return TSequenceNumber<NumBits, SequenceType>(Lhs.Get() + Rhs);
-}
-
-template <SIZE_T NumBits, typename SequenceType>
-const TSequenceNumber<NumBits, SequenceType> operator-(const TSequenceNumber<NumBits, SequenceType>& Lhs, SequenceType Rhs)
-{
-	return TSequenceNumber<NumBits, SequenceType>(Lhs.Get() - Rhs);
-}
diff --git a/Engine/Source/Runtime/Engine/Public/ReplayTypes.h b/Engine/Source/Runtime/Engine/Public/ReplayTypes.h
--- a/Engine/Source/Runtime/Engine/Public/ReplayTypes.h
+++ b/Engine/Source/Runtime/Engine/Public/ReplayTypes.h
@@ -416,7 +416,7 @@ public:
 
 	uint32 GetConstTypeHash() const
 	{
-		return ::GetTypeHash(ToString(true));
+		return GetTypeHashHelper(ToString(true));
 	}
 
 	friend uint32 GetTypeHash(const FInternetAddrDemo& A)
diff --git a/Engine/Source/Runtime/Engine/Public/ShaderCompiler.h b/Engine/Source/Runtime/Engine/Public/ShaderCompiler.h
--- a/Engine/Source/Runtime/Engine/Public/ShaderCompiler.h
+++ b/Engine/Source/Runtime/Engine/Public/ShaderCompiler.h
@@ -196,15 +196,16 @@ struct FShaderCompileJobKey
 	const FShaderType* ShaderType;
 	const FVertexFactoryType* VFType;
 	int32 PermutationId;
+
+	friend inline bool operator==(const FShaderCompileJobKey& Lhs, const FShaderCompileJobKey& Rhs)
+	{
+		return Lhs.VFType == Rhs.VFType && Lhs.ShaderType == Rhs.ShaderType && Lhs.PermutationId == Rhs.PermutationId;
+	}
+	friend inline bool operator!=(const FShaderCompileJobKey& Lhs, const FShaderCompileJobKey& Rhs)
+	{
+		return !operator==(Lhs, Rhs);
+	}
 };
-inline bool operator==(const FShaderCompileJobKey& Lhs, const FShaderCompileJobKey& Rhs)
-{
-	return Lhs.VFType == Rhs.VFType && Lhs.ShaderType == Rhs.ShaderType && Lhs.PermutationId == Rhs.PermutationId;
-}
-inline bool operator!=(const FShaderCompileJobKey& Lhs, const FShaderCompileJobKey& Rhs)
-{
-	return !operator==(Lhs, Rhs);
-}
 
 /** Stores all of the input and output information used to compile a single shader. */
 class FShaderCompileJob : public FShaderCommonCompileJob
@@ -247,15 +248,16 @@ struct FShaderPipelineCompileJobKey
 	const FShaderPipelineType* ShaderPipeline;
 	const FVertexFactoryType* VFType;
 	int32 PermutationId;
+
+	friend inline bool operator==(const FShaderPipelineCompileJobKey& Lhs, const FShaderPipelineCompileJobKey& Rhs)
+	{
+		return Lhs.ShaderPipeline == Rhs.ShaderPipeline && Lhs.VFType == Rhs.VFType && Lhs.PermutationId == Rhs.PermutationId;
+	}
+	friend inline bool operator!=(const FShaderPipelineCompileJobKey& Lhs, const FShaderPipelineCompileJobKey& Rhs)
+	{
+		return !operator==(Lhs, Rhs);
+	}
 };
-inline bool operator==(const FShaderPipelineCompileJobKey& Lhs, const FShaderPipelineCompileJobKey& Rhs)
-{
-	return Lhs.ShaderPipeline == Rhs.ShaderPipeline && Lhs.VFType == Rhs.VFType && Lhs.PermutationId == Rhs.PermutationId;
-}
-inline bool operator!=(const FShaderPipelineCompileJobKey& Lhs, const FShaderPipelineCompileJobKey& Rhs)
-{
-	return !operator==(Lhs, Rhs);
-}
 
 class FShaderPipelineCompileJob : public FShaderCommonCompileJob
 {
diff --git a/Engine/Source/Runtime/Engine/Public/StaticParameterSet.h b/Engine/Source/Runtime/Engine/Public/StaticParameterSet.h
--- a/Engine/Source/Runtime/Engine/Public/StaticParameterSet.h
+++ b/Engine/Source/Runtime/Engine/Public/StaticParameterSet.h
@@ -549,12 +549,12 @@ private:
 	void SetStaticSwitchParameterValue(const FMaterialParameterInfo& ParameterInfo, const FGuid& ExpressionGuid, bool Value);
 	void SetStaticComponentMaskParameterValue(const FMaterialParameterInfo& ParameterInfo, const FGuid& ExpressionGuid, bool R, bool G, bool B, bool A);
 #endif // WITH_EDITORONLY_DATA
-};
 
-bool operator==(const FStaticParameterSet&, const FStaticParameterSetRuntimeData&) = delete;
-bool operator==(const FStaticParameterSetRuntimeData&, const FStaticParameterSet&) = delete;
-bool operator!=(const FStaticParameterSet&, const FStaticParameterSetRuntimeData&) = delete;
-bool operator!=(const FStaticParameterSetRuntimeData&, const FStaticParameterSet&) = delete;
+	friend bool operator==(const FStaticParameterSet&, const FStaticParameterSetRuntimeData&) = delete;
+	friend bool operator==(const FStaticParameterSetRuntimeData&, const FStaticParameterSet&) = delete;
+	friend bool operator!=(const FStaticParameterSet&, const FStaticParameterSetRuntimeData&) = delete;
+	friend bool operator!=(const FStaticParameterSetRuntimeData&, const FStaticParameterSet&) = delete;
+};
 
 template<>
 struct TStructOpsTypeTraits<FStaticParameterSet> : TStructOpsTypeTraitsBase2<FStaticParameterSet>
diff --git a/Engine/Source/Runtime/Engine/Public/VisualLogger/VisualLoggerTypes.h b/Engine/Source/Runtime/Engine/Public/VisualLogger/VisualLoggerTypes.h
--- a/Engine/Source/Runtime/Engine/Public/VisualLogger/VisualLoggerTypes.h
+++ b/Engine/Source/Runtime/Engine/Public/VisualLogger/VisualLoggerTypes.h
@@ -80,6 +80,11 @@ struct ENGINE_API FVisualLogEvent
 	FVisualLogEvent() : Counter(1) { /* Empty */ }
 	FVisualLogEvent(const FVisualLogEventBase& Event);
 	FVisualLogEvent& operator=(const FVisualLogEventBase& Event);
+
+	friend inline bool operator==(const FVisualLogEvent& Left, const FVisualLogEvent& Right) 
+	{ 
+		return Left.Name == Right.Name; 
+	}
 };
 
 struct ENGINE_API FVisualLogLine
@@ -283,13 +288,13 @@ struct ENGINE_API FVisualLoggerCategoryVerbosityPair
 
 	FName CategoryName;
 	ELogVerbosity::Type Verbosity;
-};
 
-inline bool operator==(const FVisualLoggerCategoryVerbosityPair& A, const FVisualLoggerCategoryVerbosityPair& B)
-{
-	return A.CategoryName == B.CategoryName
-		&& A.Verbosity == B.Verbosity;
-}
+	friend inline bool operator==(const FVisualLoggerCategoryVerbosityPair& A, const FVisualLoggerCategoryVerbosityPair& B)
+	{
+		return A.CategoryName == B.CategoryName
+			&& A.Verbosity == B.Verbosity;
+	}
+};
 
 struct ENGINE_API FVisualLoggerHelpers
 {
@@ -339,12 +344,6 @@ ENGINE_API  FArchive& operator<<(FArchive& Ar, FVisualLogLine& LogLine);
 ENGINE_API  FArchive& operator<<(FArchive& Ar, FVisualLogStatusCategory& Status);
 ENGINE_API  FArchive& operator<<(FArchive& Ar, FVisualLogEntry& LogEntry);
 
-inline
-bool operator==(const FVisualLogEvent& Left, const FVisualLogEvent& Right) 
-{ 
-	return Left.Name == Right.Name; 
-}
-
 inline
 FVisualLogEvent::FVisualLogEvent(const FVisualLogEventBase& Event)
 : Counter(1)
diff --git a/Engine/Source/Runtime/Engine/Public/WorldPartition/DataLayer/ActorDataLayer.h b/Engine/Source/Runtime/Engine/Public/WorldPartition/DataLayer/ActorDataLayer.h
--- a/Engine/Source/Runtime/Engine/Public/WorldPartition/DataLayer/ActorDataLayer.h
+++ b/Engine/Source/Runtime/Engine/Public/WorldPartition/DataLayer/ActorDataLayer.h
@@ -34,4 +34,9 @@ struct ENGINE_API FActorDataLayer
 	/** The name of this layer */
 	UPROPERTY(BlueprintReadOnly, VisibleAnywhere, Category = DataLayer)
 	FName Name;
+
+	friend uint32 GetTypeHash(const FActorDataLayer& Value)
+	{
+		return GetTypeHash(Value.Name);
+	}
 };
\ No newline at end of file
diff --git a/Engine/Source/Runtime/Experimental/Chaos/Public/Chaos/CollisionFilterData.h b/Engine/Source/Runtime/Experimental/Chaos/Public/Chaos/CollisionFilterData.h
--- a/Engine/Source/Runtime/Experimental/Chaos/Public/Chaos/CollisionFilterData.h
+++ b/Engine/Source/Runtime/Experimental/Chaos/Public/Chaos/CollisionFilterData.h
@@ -39,6 +39,11 @@ struct CHAOS_API FCollisionFilterData
 		const uint32 FilterFlags = (Word3 & 0xFFFFFF);
 		return FilterFlags & static_cast<uint32>(InFlag);
 	}
+
+	friend inline bool operator!=(const FCollisionFilterData& A, const FCollisionFilterData& B)
+	{
+		return A.Word0!=B.Word0 || A.Word1!=B.Word1 || A.Word2!=B.Word2 || A.Word3!=B.Word3;
+	}
 };
 
 inline Chaos::FChaosArchive& operator<<(Chaos::FChaosArchive& Ar, FCollisionFilterData& Filter)
@@ -46,8 +51,3 @@ inline Chaos::FChaosArchive& operator<<(Chaos::FChaosArchive& Ar, FCollisionFilt
 	Ar << Filter.Word0 << Filter.Word1 << Filter.Word2 << Filter.Word3;
 	return Ar;
 }
-
-inline bool operator!=(const FCollisionFilterData& A, const FCollisionFilterData& B)
-{
-	return A.Word0!=B.Word0 || A.Word1!=B.Word1 || A.Word2!=B.Word2 || A.Word3!=B.Word3;
-}
\ No newline at end of file
diff --git a/Engine/Source/Runtime/Experimental/Chaos/Public/Chaos/ImplicitObjectTransformed.h b/Engine/Source/Runtime/Experimental/Chaos/Public/Chaos/ImplicitObjectTransformed.h
--- a/Engine/Source/Runtime/Experimental/Chaos/Public/Chaos/ImplicitObjectTransformed.h
+++ b/Engine/Source/Runtime/Experimental/Chaos/Public/Chaos/ImplicitObjectTransformed.h
@@ -345,7 +345,7 @@ public:
 	virtual uint32 GetTypeHash() const override
 	{
 		// Combine the hash from the inner, non transformed object with our transform
-		return HashCombine(MObject->GetTypeHash(), ::GetTypeHash(MTransform));
+		return HashCombine(MObject->GetTypeHash(), GetTypeHashHelper(MTransform));
 	}
 
 	virtual uint16 GetMaterialIndex(uint32 HintIndex) const override
diff --git a/Engine/Source/Runtime/Experimental/Chaos/Public/Chaos/Transform.h b/Engine/Source/Runtime/Experimental/Chaos/Public/Chaos/Transform.h
--- a/Engine/Source/Runtime/Experimental/Chaos/Public/Chaos/Transform.h
+++ b/Engine/Source/Runtime/Experimental/Chaos/Public/Chaos/Transform.h
@@ -246,6 +246,11 @@ namespace Chaos
 
 			return Result;
 		}
+
+		friend inline uint32 GetTypeHash(const TRigidTransform<FRealSingle, 3>& InTransform)
+		{
+			return HashCombine(GetTypeHash(InTransform.GetTranslation()), HashCombine(GetTypeHash(InTransform.GetRotation().Euler()), GetTypeHash(InTransform.GetScale3D())));
+		}
 	};
 
 	template<>
@@ -350,16 +355,10 @@ namespace Chaos
 
 			return Result;
 		}
-	};
-}
 
-inline uint32 GetTypeHash(const Chaos::TRigidTransform<Chaos::FRealSingle, 3>& InTransform)
-{
-	return HashCombine(GetTypeHash(InTransform.GetTranslation()), HashCombine(GetTypeHash(InTransform.GetRotation().Euler()), GetTypeHash(InTransform.GetScale3D())));
-}
-
-inline uint32 GetTypeHash(const Chaos::TRigidTransform<Chaos::FRealDouble, 3>& InTransform)
-{
-	return HashCombine(GetTypeHash(InTransform.GetTranslation()), HashCombine(GetTypeHash(InTransform.GetRotation().Euler()), GetTypeHash(InTransform.GetScale3D())));
+		friend inline uint32 GetTypeHash(const TRigidTransform<FRealDouble, 3>& InTransform)
+		{
+			return HashCombine(GetTypeHash(InTransform.GetTranslation()), HashCombine(GetTypeHash(InTransform.GetRotation().Euler()), GetTypeHash(InTransform.GetScale3D())));
+		}
+	};
 }
-
diff --git a/Engine/Source/Runtime/Experimental/ChaosCore/Public/Chaos/Array.h b/Engine/Source/Runtime/Experimental/ChaosCore/Public/Chaos/Array.h
--- a/Engine/Source/Runtime/Experimental/ChaosCore/Public/Chaos/Array.h
+++ b/Engine/Source/Runtime/Experimental/ChaosCore/Public/Chaos/Array.h
@@ -44,6 +44,14 @@ class TArray : public std::vector<T>
 	{
 		std::sort(begin(), end());
 	}
+
+	friend uint32 GetTypeHash(const TArray<int32>& Array)
+	{
+		uint32 Seed = 0;
+		for (const auto& Elem : Array)
+			Seed ^= GetTypeHash(Elem) + 0x9e3779b9 + (Seed << 6) + (Seed >> 2);
+		return Seed;
+	}
 };
 
 namespace std
@@ -61,11 +69,3 @@ struct hash<TArray<int32>>
 };
 }
 #endif
-
-static uint32 GetTypeHash(const TArray<int32>& Array)
-{
-	uint32 Seed = 0;
-	for (const auto& Elem : Array)
-		Seed ^= GetTypeHash(Elem) + 0x9e3779b9 + (Seed << 6) + (Seed >> 2);
-	return Seed;
-}
\ No newline at end of file
diff --git a/Engine/Source/Runtime/GameplayTags/Classes/GameplayTagContainer.h b/Engine/Source/Runtime/GameplayTags/Classes/GameplayTagContainer.h
--- a/Engine/Source/Runtime/GameplayTags/Classes/GameplayTagContainer.h
+++ b/Engine/Source/Runtime/GameplayTags/Classes/GameplayTagContainer.h
@@ -161,7 +161,7 @@ struct GAMEPLAYTAGS_API FGameplayTag
 	/** Used so we can have a TMap of this struct */
 	FORCEINLINE friend uint32 GetTypeHash(const FGameplayTag& Tag)
 	{
-		return ::GetTypeHash(Tag.TagName);
+		return GetTypeHash(Tag.TagName);
 	}
 
 	/** Displays gameplay tag as a string for blueprint graph usage */
diff --git a/Engine/Source/Runtime/Json/Public/Dom/JsonValue.h b/Engine/Source/Runtime/Json/Public/Dom/JsonValue.h
--- a/Engine/Source/Runtime/Json/Public/Dom/JsonValue.h
+++ b/Engine/Source/Runtime/Json/Public/Dom/JsonValue.h
@@ -107,17 +107,17 @@ protected:
 	virtual FString GetType() const = 0;
 
 	void ErrorMessage(const FString& InType) const;
-};
 
-inline bool operator==(const FJsonValue& Lhs, const FJsonValue& Rhs)
-{
-	return FJsonValue::CompareEqual(Lhs, Rhs);
-}
+	friend inline bool operator==(const FJsonValue& Lhs, const FJsonValue& Rhs)
+	{
+		return FJsonValue::CompareEqual(Lhs, Rhs);
+	}
 
-inline bool operator!=(const FJsonValue& Lhs, const FJsonValue& Rhs)
-{
-	return !FJsonValue::CompareEqual(Lhs, Rhs);
-}
+	friend inline bool operator!=(const FJsonValue& Lhs, const FJsonValue& Rhs)
+	{
+		return !FJsonValue::CompareEqual(Lhs, Rhs);
+	}
+};
 
 
 /** A Json String Value. */
diff --git a/Engine/Source/Runtime/Launch/Private/LaunchEngineLoop.cpp b/Engine/Source/Runtime/Launch/Private/LaunchEngineLoop.cpp
--- a/Engine/Source/Runtime/Launch/Private/LaunchEngineLoop.cpp
+++ b/Engine/Source/Runtime/Launch/Private/LaunchEngineLoop.cpp
@@ -1214,8 +1214,8 @@ private:
 	};
 	friend uint32 GetTypeHash(const FFileInPakFileHistory& H)
 	{
-		uint32 Hash = ::GetTypeHash(H.PakFileName);
-		Hash = HashCombine(Hash, ::GetTypeHash(H.FileName));
+		uint32 Hash = GetTypeHash(H.PakFileName);
+		Hash = HashCombine(Hash, GetTypeHash(H.FileName));
 		return Hash;
 	}
 	friend bool operator==(const FFileInPakFileHistory& A, const FFileInPakFileHistory& B)
diff --git a/Engine/Source/Runtime/LiveLinkInterface/Public/LiveLinkTypes.h b/Engine/Source/Runtime/LiveLinkInterface/Public/LiveLinkTypes.h
--- a/Engine/Source/Runtime/LiveLinkInterface/Public/LiveLinkTypes.h
+++ b/Engine/Source/Runtime/LiveLinkInterface/Public/LiveLinkTypes.h
@@ -64,8 +64,12 @@ public:
 		Ar << InSubjectName.Name;
 		return Ar;
 	}
+	
+	friend inline uint32 GetTypeHash(const FLiveLinkSubjectName& Value)
+	{
+		return GetTypeHash(Value.Name);
+	}
 };
-template <> struct TModels<CGetTypeHashable, FLiveLinkSubjectName> { enum { Value = TModels<CGetTypeHashable, FName>::Value }; };
 
 
 // Structure that identifies an individual subject
diff --git a/Engine/Source/Runtime/PhysicsCore/Public/Chaos/ChaosEngineInterface.h b/Engine/Source/Runtime/PhysicsCore/Public/Chaos/ChaosEngineInterface.h
--- a/Engine/Source/Runtime/PhysicsCore/Public/Chaos/ChaosEngineInterface.h
+++ b/Engine/Source/Runtime/PhysicsCore/Public/Chaos/ChaosEngineInterface.h
@@ -270,6 +270,11 @@ public:
 
 	Chaos::FPerShapeData* Shape;
     FPhysicsActorHandle ActorRef;
+
+	friend FORCEINLINE uint32 GetTypeHash(const FPhysicsShapeReference_Chaos& InShapeReference)
+	{
+		return PointerHash(InShapeReference.Shape);
+	}
 };
 
 class PHYSICSCORE_API FPhysicsShapeAdapter_Chaos
@@ -287,11 +292,6 @@ private:
 	FQuat GeometryRotation;
 };
 
-FORCEINLINE uint32 GetTypeHash(const FPhysicsShapeReference_Chaos& InShapeReference)
-{
-    return PointerHash(InShapeReference.Shape);
-}
-
 /**
  Wrapper around geometry. This is really just needed to make the physx chaos abstraction easier
  */
diff --git a/Engine/Source/Runtime/RHI/Public/RHICommandList.h b/Engine/Source/Runtime/RHI/Public/RHICommandList.h
--- a/Engine/Source/Runtime/RHI/Public/RHICommandList.h
+++ b/Engine/Source/Runtime/RHI/Public/RHICommandList.h
@@ -2263,6 +2263,10 @@ template<> RHI_API void FRHICommandSetShaderResourceViewParameter<FRHIComputeSha
 template<> RHI_API void FRHICommandSetShaderSampler              <FRHIComputeShader>::Execute(FRHICommandListBase& CmdList);
 template<> RHI_API void FRHICommandSetUAVParameter               <FRHIComputeShader>::Execute(FRHICommandListBase& CmdList);
 
+extern RHI_API FRHIComputePipelineState* ExecuteSetComputePipelineState(FComputePipelineState* ComputePipelineState);
+extern RHI_API FRHIGraphicsPipelineState* ExecuteSetGraphicsPipelineState(class FGraphicsPipelineState* GraphicsPipelineState);
+extern RHI_API FRHIRayTracingPipelineState* GetRHIRayTracingPipelineState(FRayTracingPipelineState*);
+
 class RHI_API FRHIComputeCommandList : public FRHICommandListBase
 {
 public:
@@ -2450,7 +2454,6 @@ PRAGMA_ENABLE_DEPRECATION_WARNINGS
 		PersistentState.BoundComputeShaderRHI = ComputeShader;
 		if (Bypass())
 		{
-			extern RHI_API FRHIComputePipelineState* ExecuteSetComputePipelineState(FComputePipelineState* ComputePipelineState);
 			FRHIComputePipelineState* RHIComputePipelineState = ExecuteSetComputePipelineState(ComputePipelineState);
 			GetComputeContext().RHISetComputePipelineState(RHIComputePipelineState);
 			return;
@@ -3306,7 +3309,6 @@ public:
 		PersistentState.BoundShaderInput = ShaderInput;
 		if (Bypass())
 		{
-			extern RHI_API FRHIGraphicsPipelineState* ExecuteSetGraphicsPipelineState(class FGraphicsPipelineState* GraphicsPipelineState);
 			FRHIGraphicsPipelineState* RHIGraphicsPipelineState = ExecuteSetGraphicsPipelineState(GraphicsPipelineState);
 			GetContext().RHISetGraphicsPipelineState(RHIGraphicsPipelineState, StencilRef, bApplyAdditionalState);
 			return;
@@ -3616,7 +3618,6 @@ public:
 	{
 		if (Bypass())
 		{
-			extern RHI_API FRHIRayTracingPipelineState* GetRHIRayTracingPipelineState(FRayTracingPipelineState*);
 			GetContext().RHIRayTraceDispatch(GetRHIRayTracingPipelineState(Pipeline), RayGenShader, Scene, GlobalResourceBindings, Width, Height);
 		}
 		else
@@ -3634,7 +3635,6 @@ public:
 		checkf(GRHISupportsRayTracingDispatchIndirect, TEXT("Indirect ray tracing is not supported on this machine."));
 		if (Bypass())
 		{
-			extern RHI_API FRHIRayTracingPipelineState* GetRHIRayTracingPipelineState(FRayTracingPipelineState*);
 			GetContext().RHIRayTraceDispatchIndirect(GetRHIRayTracingPipelineState(Pipeline), RayGenShader, Scene, GlobalResourceBindings, ArgumentBuffer, ArgumentOffset);
 		}
 		else
@@ -3651,7 +3651,6 @@ public:
 	{
 		if (Bypass())
 		{
-			extern RHI_API FRHIRayTracingPipelineState* GetRHIRayTracingPipelineState(FRayTracingPipelineState*);
 			GetContext().RHISetRayTracingBindings(Scene, GetRHIRayTracingPipelineState(Pipeline), NumBindings, Bindings, BindingType);
 		}
 		else
@@ -3729,7 +3728,6 @@ public:
 	{
 		if (Bypass())
 		{
-			extern RHI_API FRHIRayTracingPipelineState* GetRHIRayTracingPipelineState(FRayTracingPipelineState*);
 			GetContext().RHISetRayTracingHitGroup(Scene, InstanceIndex, SegmentIndex, ShaderSlot, GetRHIRayTracingPipelineState(Pipeline), HitGroupIndex, 
 				NumUniformBuffers, UniformBuffers,
 				LooseParameterDataSize, LooseParameterData,
@@ -3769,7 +3767,6 @@ public:
 	{
 		if (Bypass())
 		{
-			extern RHI_API FRHIRayTracingPipelineState* GetRHIRayTracingPipelineState(FRayTracingPipelineState*);
 			GetContext().RHISetRayTracingCallableShader(Scene, ShaderSlotInScene, GetRHIRayTracingPipelineState(Pipeline), ShaderIndexInPipeline, NumUniformBuffers, UniformBuffers, UserData);
 		}
 		else
@@ -3796,7 +3793,6 @@ public:
 	{
 		if (Bypass())
 		{
-			extern RHI_API FRHIRayTracingPipelineState* GetRHIRayTracingPipelineState(FRayTracingPipelineState*);
 			GetContext().RHISetRayTracingMissShader(Scene, ShaderSlotInScene, GetRHIRayTracingPipelineState(Pipeline), ShaderIndexInPipeline, NumUniformBuffers, UniformBuffers, UserData);
 		}
 		else
diff --git a/Engine/Source/Runtime/RHI/Public/RHIResources.h b/Engine/Source/Runtime/RHI/Public/RHIResources.h
--- a/Engine/Source/Runtime/RHI/Public/RHIResources.h
+++ b/Engine/Source/Runtime/RHI/Public/RHIResources.h
@@ -747,6 +747,13 @@ struct FRHIUniformBufferResource
 
 	/** Type of the member that allow (). */
 	LAYOUT_FIELD(EUniformBufferBaseType, MemberType);
+
+	/** Compare two uniform buffer layout resources. */
+	friend inline bool operator==(const FRHIUniformBufferResource& A, const FRHIUniformBufferResource& B)
+	{
+		return A.MemberOffset == B.MemberOffset
+			&& A.MemberType == B.MemberType;
+	}
 };
 
 inline FArchive& operator<<(FArchive& Ar, FRHIUniformBufferResource& Ref)
@@ -758,14 +765,7 @@ inline FArchive& operator<<(FArchive& Ar, FRHIUniformBufferResource& Ref)
 	return Ar;
 }
 
-/** Compare two uniform buffer layout resources. */
-inline bool operator==(const FRHIUniformBufferResource& A, const FRHIUniformBufferResource& B)
-{
-	return A.MemberOffset == B.MemberOffset
-		&& A.MemberType == B.MemberType;
-}
-
-static constexpr uint16 kUniformBufferInvalidOffset = TNumericLimits<uint16>::Max();
+inline constexpr uint16 kUniformBufferInvalidOffset = TNumericLimits<uint16>::Max();
 
 /** Initializer for the layout of a uniform buffer in memory. */
 struct FRHIUniformBufferLayoutInitializer
@@ -915,16 +915,16 @@ public:
 
 	/** Used for platforms which use emulated ub's, forces a real uniform buffer instead */
 	LAYOUT_FIELD_INITIALIZED(bool, bNoEmulatedUniformBuffer, false);
-};
 
-/** Compare two uniform buffer layout initializers. */
-inline bool operator==(const FRHIUniformBufferLayoutInitializer& A, const FRHIUniformBufferLayoutInitializer& B)
-{
-	return A.ConstantBufferSize == B.ConstantBufferSize
-		&& A.StaticSlot == B.StaticSlot
-		&& A.BindingFlags == B.BindingFlags
-		&& A.Resources == B.Resources;
-}
+	/** Compare two uniform buffer layout initializers. */
+	friend inline bool operator==(const FRHIUniformBufferLayoutInitializer& A, const FRHIUniformBufferLayoutInitializer& B)
+	{
+		return A.ConstantBufferSize == B.ConstantBufferSize
+			&& A.StaticSlot == B.StaticSlot
+			&& A.BindingFlags == B.BindingFlags
+			&& A.Resources == B.Resources;
+	}
+};
 
 /** The layout of a uniform buffer in memory. */
 struct FRHIUniformBufferLayout : public FRHIResource
@@ -1014,16 +1014,16 @@ struct FRHIUniformBufferLayout : public FRHIResource
 
 	/** Used for platforms which use emulated ub's, forces a real uniform buffer instead */
 	const bool bNoEmulatedUniformBuffer;
-};
 
-/** Compare two uniform buffer layouts. */
-inline bool operator==(const FRHIUniformBufferLayout& A, const FRHIUniformBufferLayout& B)
-{
-	return A.ConstantBufferSize == B.ConstantBufferSize
-		&& A.StaticSlot == B.StaticSlot
-		&& A.BindingFlags == B.BindingFlags
-		&& A.Resources == B.Resources;
-}
+	/** Compare two uniform buffer layouts. */
+	friend inline bool operator==(const FRHIUniformBufferLayout& A, const FRHIUniformBufferLayout& B)
+	{
+		return A.ConstantBufferSize == B.ConstantBufferSize
+			&& A.StaticSlot == B.StaticSlot
+			&& A.BindingFlags == B.BindingFlags
+			&& A.Resources == B.Resources;
+	}
+};
 
 class FRHIUniformBuffer : public FRHIResource
 #if ENABLE_RHI_VALIDATION
diff --git a/Engine/Source/Runtime/RenderCore/Public/RenderGraphDefinitions.h b/Engine/Source/Runtime/RenderCore/Public/RenderGraphDefinitions.h
--- a/Engine/Source/Runtime/RenderCore/Public/RenderGraphDefinitions.h
+++ b/Engine/Source/Runtime/RenderCore/Public/RenderGraphDefinitions.h
@@ -399,13 +399,12 @@ public:
 private:
 	static const IndexType kNullIndex = TNumericLimits<IndexType>::Max();
 	IndexType Index = kNullIndex;
-};
 
-template <typename ObjectType, typename IndexType>
-FORCEINLINE uint32 GetTypeHash(TRDGHandle<ObjectType, IndexType> Handle)
-{
-	return Handle.GetIndex();
-}
+	friend FORCEINLINE uint32 GetTypeHash(TRDGHandle Handle)
+	{
+		return Handle.GetIndex();
+	}
+};
 
 enum class ERDGHandleRegistryDestructPolicy
 {
diff --git a/Engine/Source/Runtime/RenderCore/Public/Shader.h b/Engine/Source/Runtime/RenderCore/Public/Shader.h
--- a/Engine/Source/Runtime/RenderCore/Public/Shader.h
+++ b/Engine/Source/Runtime/RenderCore/Public/Shader.h
@@ -100,17 +100,16 @@ struct TShaderTypePermutation
 	{
 		return !(*this == Other);
 	}
+
+	friend FORCEINLINE uint32 GetTypeHash(const TShaderTypePermutation& Var)
+	{
+		return HashCombine(GetTypeHash(Var.Type), (uint32)Var.PermutationId);
+	}
 };
 
 using FShaderPermutation = TShaderTypePermutation<FShaderType>;
 
-const int32 kUniqueShaderPermutationId = 0;
-
-template<typename MetaShaderType>
-FORCEINLINE uint32 GetTypeHash(const TShaderTypePermutation<MetaShaderType>& Var)
-{
-	return HashCombine(GetTypeHash(Var.Type), (uint32)Var.PermutationId);
-}
+inline const int32 kUniqueShaderPermutationId = 0;
 
 /** Used to compare order shader types permutation deterministically. */
 template<typename MetaShaderType>
@@ -1140,24 +1139,25 @@ private:
 
 	ShaderType* ShaderContent;
 	const FShaderMapBase* ShaderMap;
-};
 
-template<typename ShaderType0, typename ShaderType1, typename PointerTableType>
-inline bool operator==(const TShaderRefBase<ShaderType0, PointerTableType>& Lhs, const TShaderRefBase<ShaderType1, PointerTableType>& Rhs)
-{
-	if (Lhs.GetShader() == Rhs.GetShader())
+	template<typename ShaderType1>
+	friend inline bool operator==(const TShaderRefBase& Lhs, const TShaderRefBase<ShaderType1, PointerTableType>& Rhs)
 	{
-		check(Lhs.GetShaderMap() == Rhs.GetShaderMap());
-		return true;
+		if (Lhs.GetShader() == Rhs.GetShader())
+		{
+			check(Lhs.GetShaderMap() == Rhs.GetShaderMap());
+			return true;
+		}
+		return false;
 	}
-	return false;
-}
 
-template<typename ShaderType0, typename ShaderType1, typename PointerTableType>
-inline bool operator!=(const TShaderRefBase<ShaderType0, PointerTableType>& Lhs, const TShaderRefBase<ShaderType1, PointerTableType>& Rhs)
-{
-	return !operator==(Lhs, Rhs);
-}
+	template<typename ShaderType1>
+	friend inline bool operator!=(const TShaderRefBase& Lhs, const TShaderRefBase<ShaderType1, PointerTableType>& Rhs)
+	{
+		return !operator==(Lhs, Rhs);
+	}
+
+};
 
 template<typename ShaderType>
 using TShaderRef = TShaderRefBase<ShaderType, FShaderMapPointerTable>;
diff --git a/Engine/Source/Runtime/RenderCore/Public/ShaderCore.h b/Engine/Source/Runtime/RenderCore/Public/ShaderCore.h
--- a/Engine/Source/Runtime/RenderCore/Public/ShaderCore.h
+++ b/Engine/Source/Runtime/RenderCore/Public/ShaderCore.h
@@ -242,18 +242,18 @@ struct FParameterAllocation
 		Ar << Allocation.Type;
 		return Ar;
 	}
-};
 
-inline bool operator==(const FParameterAllocation& A, const FParameterAllocation& B)
-{
-	return
-		A.BufferIndex == B.BufferIndex && A.BaseIndex == B.BaseIndex && A.Size == B.Size && A.Type == B.Type && A.bBound == B.bBound;
-}
+	friend inline bool operator==(const FParameterAllocation& A, const FParameterAllocation& B)
+	{
+		return
+			A.BufferIndex == B.BufferIndex && A.BaseIndex == B.BaseIndex && A.Size == B.Size && A.Type == B.Type && A.bBound == B.bBound;
+	}
 
-inline bool operator!=(const FParameterAllocation& A, const FParameterAllocation& B)
-{
-	return !(A == B);
-}
+	friend inline bool operator!=(const FParameterAllocation& A, const FParameterAllocation& B)
+	{
+		return !(A == B);
+	}
+};
 
 /**
  * A map of shader parameter names to registers allocated to that parameter.
@@ -716,17 +716,18 @@ struct FShaderCodeVendorExtension
 	{
 		return Ar << Extension.VendorId << Extension.Parameter;
 	}
-};
 
-inline bool operator==(const FShaderCodeVendorExtension& A, const FShaderCodeVendorExtension& B)
-{
-	return A.VendorId == B.VendorId && A.Parameter == B.Parameter;
-}
+	friend inline bool operator==(const FShaderCodeVendorExtension& A, const FShaderCodeVendorExtension& B)
+	{
+		return A.VendorId == B.VendorId && A.Parameter == B.Parameter;
+	}
 
-inline bool operator!=(const FShaderCodeVendorExtension& A, const FShaderCodeVendorExtension& B)
-{
-	return !(A == B);
-}
+	friend inline bool operator!=(const FShaderCodeVendorExtension& A, const FShaderCodeVendorExtension& B)
+	{
+		return !(A == B);
+	}
+
+};
 
 #ifndef RENDERCORE_ATTRIBUTE_UNALIGNED
 // TODO find out if using GCC_ALIGN(1) instead of this new #define break on all kinds of platforms...
diff --git a/Engine/Source/Runtime/RenderCore/Public/VirtualTexturing.h b/Engine/Source/Runtime/RenderCore/Public/VirtualTexturing.h
--- a/Engine/Source/Runtime/RenderCore/Public/VirtualTexturing.h
+++ b/Engine/Source/Runtime/RenderCore/Public/VirtualTexturing.h
@@ -40,12 +40,12 @@ union FVirtualTextureProducerHandle
 		uint32 Index : 22;
 		uint32 Magic : 10;
 	};
+
+	friend inline bool operator==(const FVirtualTextureProducerHandle& Lhs, const FVirtualTextureProducerHandle& Rhs) { return Lhs.PackedValue == Rhs.PackedValue; }
+	friend inline bool operator!=(const FVirtualTextureProducerHandle& Lhs, const FVirtualTextureProducerHandle& Rhs) { return Lhs.PackedValue != Rhs.PackedValue; }
 };
 static_assert(sizeof(FVirtualTextureProducerHandle) == sizeof(uint32), "Bad packing");
 
-inline bool operator==(const FVirtualTextureProducerHandle& Lhs, const FVirtualTextureProducerHandle& Rhs) { return Lhs.PackedValue == Rhs.PackedValue; }
-inline bool operator!=(const FVirtualTextureProducerHandle& Lhs, const FVirtualTextureProducerHandle& Rhs) { return Lhs.PackedValue != Rhs.PackedValue; }
-
 /** Maximum number of layers that can be allocated in a single VT page table */
 #define VIRTUALTEXTURE_SPACE_MAXLAYERS 8
 
@@ -95,47 +95,47 @@ struct FAllocatedVTDescription
 			uint8 bShareDuplicateLayers : 1;
 		};
 	};
-};
 
-inline bool operator==(const FAllocatedVTDescription& Lhs, const FAllocatedVTDescription& Rhs)
-{
-	if (Lhs.TileSize != Rhs.TileSize ||
-		Lhs.TileBorderSize != Rhs.TileBorderSize ||
-		Lhs.Dimensions != Rhs.Dimensions ||
-		Lhs.NumTextureLayers != Rhs.NumTextureLayers ||
-		Lhs.PackedFlags != Rhs.PackedFlags)
+	friend inline bool operator==(const FAllocatedVTDescription& Lhs, const FAllocatedVTDescription& Rhs)
 	{
-		return false;
-	}
-	for (uint32 LayerIndex = 0u; LayerIndex < Lhs.NumTextureLayers; ++LayerIndex)
-	{
-		if (Lhs.ProducerHandle[LayerIndex] != Rhs.ProducerHandle[LayerIndex] ||
-			Lhs.ProducerLayerIndex[LayerIndex] != Rhs.ProducerLayerIndex[LayerIndex])
+		if (Lhs.TileSize != Rhs.TileSize ||
+			Lhs.TileBorderSize != Rhs.TileBorderSize ||
+			Lhs.Dimensions != Rhs.Dimensions ||
+			Lhs.NumTextureLayers != Rhs.NumTextureLayers ||
+			Lhs.PackedFlags != Rhs.PackedFlags)
 		{
 			return false;
 		}
+		for (uint32 LayerIndex = 0u; LayerIndex < Lhs.NumTextureLayers; ++LayerIndex)
+		{
+			if (Lhs.ProducerHandle[LayerIndex] != Rhs.ProducerHandle[LayerIndex] ||
+				Lhs.ProducerLayerIndex[LayerIndex] != Rhs.ProducerLayerIndex[LayerIndex])
+			{
+				return false;
+			}
+		}
+		return true;
+	}
+	friend inline bool operator!=(const FAllocatedVTDescription& Lhs, const FAllocatedVTDescription& Rhs)
+	{
+		return !operator==(Lhs, Rhs);
 	}
-	return true;
-}
-inline bool operator!=(const FAllocatedVTDescription& Lhs, const FAllocatedVTDescription& Rhs)
-{
-	return !operator==(Lhs, Rhs);
-}
 
-inline uint32 GetTypeHash(const FAllocatedVTDescription& Description)
-{
-	uint32 Hash = GetTypeHash(Description.TileSize);
-	Hash = HashCombine(Hash, GetTypeHash(Description.TileBorderSize));
-	Hash = HashCombine(Hash, GetTypeHash(Description.Dimensions));
-	Hash = HashCombine(Hash, GetTypeHash(Description.NumTextureLayers));
-	Hash = HashCombine(Hash, GetTypeHash(Description.PackedFlags));
-	for (uint32 LayerIndex = 0u; LayerIndex < Description.NumTextureLayers; ++LayerIndex)
+	friend inline uint32 GetTypeHash(const FAllocatedVTDescription& Description)
 	{
-		Hash = HashCombine(Hash, GetTypeHash(Description.ProducerHandle[LayerIndex].PackedValue));
-		Hash = HashCombine(Hash, GetTypeHash(Description.ProducerLayerIndex[LayerIndex]));
+		uint32 Hash = GetTypeHash(Description.TileSize);
+		Hash = HashCombine(Hash, GetTypeHash(Description.TileBorderSize));
+		Hash = HashCombine(Hash, GetTypeHash(Description.Dimensions));
+		Hash = HashCombine(Hash, GetTypeHash(Description.NumTextureLayers));
+		Hash = HashCombine(Hash, GetTypeHash(Description.PackedFlags));
+		for (uint32 LayerIndex = 0u; LayerIndex < Description.NumTextureLayers; ++LayerIndex)
+		{
+			Hash = HashCombine(Hash, GetTypeHash(Description.ProducerHandle[LayerIndex].PackedValue));
+			Hash = HashCombine(Hash, GetTypeHash(Description.ProducerLayerIndex[LayerIndex]));
+		}
+		return Hash;
 	}
-	return Hash;
-}
+};
 
 struct FVTProducerDescription
 {
@@ -481,13 +481,13 @@ union FVirtualTextureLocalTile
 		uint32 Local_vLevel : 4;
 		uint32 Pad : 4;
 	};
+
+	friend inline bool operator==(const FVirtualTextureLocalTile& Lhs, const FVirtualTextureLocalTile& Rhs) { return Lhs.PackedValue == Rhs.PackedValue; }
+	friend inline bool operator!=(const FVirtualTextureLocalTile& Lhs, const FVirtualTextureLocalTile& Rhs) { return Lhs.PackedValue != Rhs.PackedValue; }
+	friend inline uint64 GetTypeHash(const FVirtualTextureLocalTile& T) { return T.PackedValue; }
 };
 static_assert(sizeof(FVirtualTextureLocalTile) == sizeof(uint64), "Bad packing");
 
-inline uint64 GetTypeHash(const FVirtualTextureLocalTile& T) { return T.PackedValue; }
-inline bool operator==(const FVirtualTextureLocalTile& Lhs, const FVirtualTextureLocalTile& Rhs) { return Lhs.PackedValue == Rhs.PackedValue; }
-inline bool operator!=(const FVirtualTextureLocalTile& Lhs, const FVirtualTextureLocalTile& Rhs) { return Lhs.PackedValue != Rhs.PackedValue; }
-
 RENDERCORE_API DECLARE_LOG_CATEGORY_EXTERN(LogVirtualTexturing, Log, All);
 
 DECLARE_STATS_GROUP(TEXT("Virtual Texturing"), STATGROUP_VirtualTexturing, STATCAT_Advanced);
diff --git a/Engine/Source/Runtime/SlateCore/Public/Layout/Clipping.h b/Engine/Source/Runtime/SlateCore/Public/Layout/Clipping.h
--- a/Engine/Source/Runtime/SlateCore/Public/Layout/Clipping.h
+++ b/Engine/Source/Runtime/SlateCore/Public/Layout/Clipping.h
@@ -170,12 +170,12 @@ private:
 	uint8 bIntersect : 1;
 	/** Should this clipping zone always clip, even if another zone wants to ignore intersection? */
 	uint8 bAlwaysClip : 1;
-};
 
-FORCEINLINE uint32 GetTypeHash(const FSlateClippingZone& Zone)
-{
-	return Zone.ComputeHash();
-}
+	friend FORCEINLINE uint32 GetTypeHash(const FSlateClippingZone& Zone)
+	{
+		return Zone.ComputeHash();
+	}
+};
 
 template<> struct TIsPODType<FSlateClippingZone> { enum { Value = true }; };
 
diff --git a/Engine/Source/Runtime/Sockets/Public/IPAddress.h b/Engine/Source/Runtime/Sockets/Public/IPAddress.h
--- a/Engine/Source/Runtime/Sockets/Public/IPAddress.h
+++ b/Engine/Source/Runtime/Sockets/Public/IPAddress.h
@@ -187,6 +187,11 @@ public:
 	}
 
 	virtual void DumpAddrData() const;
+
+	friend inline uint32 GetTypeHash(const FInternetAddr& InAddr)
+	{
+		return InAddr.GetTypeHash();
+	}
 };
 
 /**
@@ -434,8 +439,3 @@ struct FInternetAddrConstKeyMapFuncs : public BaseKeyFuncs<ValueType, TSharedRef
 		return Key->GetTypeHash();
 	}
 };
-
-inline uint32 GetTypeHash(const FInternetAddr& InAddr)
-{
-	return InAddr.GetTypeHash();
-}
