[Core] changed FCustomVersionContainer::ToString to use a string builder to be more efficient

diff --git a/Engine/Source/Runtime/Core/Private/Serialization/CustomVersion.cpp b/Engine/Source/Runtime/Core/Private/Serialization/CustomVersion.cpp
--- a/Engine/Source/Runtime/Core/Private/Serialization/CustomVersion.cpp
+++ b/Engine/Source/Runtime/Core/Private/Serialization/CustomVersion.cpp
@@ -284,14 +284,17 @@ void FCustomVersionContainer::SortByKey()
 
 FString FCustomVersionContainer::ToString(const FString& Indent) const
 {
-	FString VersionsAsString;
+	TStringBuilder<2048> VersionsAsString;
 	for (const FCustomVersion& SomeVersion : Versions)
 	{
-		VersionsAsString += Indent;
-		VersionsAsString += FString::Printf(TEXT("Key=%s  Version=%d  Friendly Name=%s \n"), *SomeVersion.Key.ToString(), SomeVersion.Version, *SomeVersion.GetFriendlyName().ToString() );
+		VersionsAsString << Indent
+			<< TEXTVIEW("Key=") << SomeVersion.Key
+			<< TEXTVIEW(" Version = ") << SomeVersion.Version
+			<< TEXTVIEW(" Friendly Name = ") << SomeVersion.GetFriendlyName()
+			<< TEXTVIEW("\n");
 	}
 
-	return VersionsAsString;
+	return FString(VersionsAsString);
 }
 
 FArchive& operator<<(FArchive& Ar, FCustomVersion& Version)
