Fix driver version checks not working for Nvidia GPUs.

The new driver detection method (using the setup API) did not set the UserDriverVersion in the same way as the other methods, so the function which extracted the number used for comparison didn't work correctly.

Also removed a bogus cvar from BaseEngine.ini while we're at it (r.DetectAndWarnOfBadDrivers has been removed from the code years ago).

#rnx

diff --git a/Engine/Config/BaseEngine.ini b/Engine/Config/BaseEngine.ini
--- a/Engine/Config/BaseEngine.ini
+++ b/Engine/Config/BaseEngine.ini
@@ -1967,8 +1967,6 @@ con.DebugEarlyCheat = True
 con.DebugLateDefault = True
 con.DebugLateCheat = True
 
-r.DetectAndWarnOfBadDrivers = 0
-
 !LogNamedEventFilters=ClearArray
 +LogNamedEventFilters="Frame *"
 
diff --git a/Engine/Source/Runtime/Core/Private/Windows/WindowsPlatformMisc.cpp b/Engine/Source/Runtime/Core/Private/Windows/WindowsPlatformMisc.cpp
--- a/Engine/Source/Runtime/Core/Private/Windows/WindowsPlatformMisc.cpp
+++ b/Engine/Source/Runtime/Core/Private/Windows/WindowsPlatformMisc.cpp
@@ -2608,7 +2608,7 @@ static void GetVideoDriverDetailsFromSetup(const FString& DeviceName, FGPUDriver
 
 	if(Out.IsNVIDIA())
 	{
-		Out.UserDriverVersion = Out.GetUnifiedDriverVersion();
+		Out.UserDriverVersion = Out.TrimNVIDIAInternalVersion(Out.InternalDriverVersion);
 	}
 	else if(Out.IsAMD() && !RegistryKey.IsEmpty())
 	{
