Fix TSAN warning in FThreadSafeCounter

#rnx

diff --git a/Engine/Source/Runtime/Core/Public/HAL/ThreadSafeCounter.h b/Engine/Source/Runtime/Core/Public/HAL/ThreadSafeCounter.h
--- a/Engine/Source/Runtime/Core/Public/HAL/ThreadSafeCounter.h
+++ b/Engine/Source/Runtime/Core/Public/HAL/ThreadSafeCounter.h
@@ -22,7 +22,11 @@ public:
 	 */
 	FThreadSafeCounter()
 	{
+#if USING_THREAD_SANITISER
+		FPlatformAtomics::AtomicStore(&Counter, 0);
+#else
 		Counter = 0;
+#endif
 	}
 
 	/**
@@ -36,7 +40,11 @@ public:
 	 */
 	FThreadSafeCounter( const FThreadSafeCounter& Other )
 	{
+#if USING_THREAD_SANITISER
+		FPlatformAtomics::AtomicStore(&Counter, Other.GetValue());
+#else
 		Counter = Other.GetValue();
+#endif
 	}
 
 	/**
@@ -46,7 +54,11 @@ public:
 	 */
 	FThreadSafeCounter( int32 Value )
 	{
+#if USING_THREAD_SANITISER
+		FPlatformAtomics::AtomicStore(&Counter, Value);
+#else
 		Counter = Value;
+#endif
 	}
 
 	/**
