On average saved 500ms(e.g.746ms to 250ms) when calling CombineFromBuffer by caching the value hash.

diff --git a/Engine/Source/Runtime/Core/Public/Misc/ConfigCacheIni.h b/Engine/Source/Runtime/Core/Public/Misc/ConfigCacheIni.h
--- a/Engine/Source/Runtime/Core/Public/Misc/ConfigCacheIni.h
+++ b/Engine/Source/Runtime/Core/Public/Misc/ConfigCacheIni.h
@@ -18,6 +18,7 @@
 #include "Delegates/Delegate.h"
 #include "HAL/PlatformCrt.h"
 #include "Internationalization/Text.h"
+#include "Internationalization/TextLocalizationResource.h"
 #include "Logging/LogMacros.h"
 #include "Math/Color.h"
 #include "Math/MathFwd.h"
@@ -109,6 +110,7 @@ public:
 		, bRead(false)
 #endif
 	{
+		SavedValueHash = FTextLocalizationResource::HashString(SavedValue);
 		ExpandValueInternal();
 	}
 
@@ -118,6 +120,7 @@ public:
 		, bRead(false)
 #endif
 	{
+		SavedValueHash = FTextLocalizationResource::HashString(SavedValue);
 		ExpandValueInternal();
 	}
 
@@ -127,12 +130,14 @@ public:
 		, bRead(false)
 #endif
 	{
+		SavedValueHash = FTextLocalizationResource::HashString(SavedValue);
 		ExpandValueInternal();
 	}
 
 	FConfigValue(const FConfigValue& InConfigValue)
 		: SavedValue(InConfigValue.SavedValue)
 		, ExpandedValue(InConfigValue.ExpandedValue)
+		, SavedValueHash(InConfigValue.SavedValueHash)
 #if CONFIG_REMEMBER_ACCESS_PATTERN 
 		, bRead(InConfigValue.bRead)
 #endif
@@ -143,6 +148,7 @@ public:
 	FConfigValue(FConfigValue&& InConfigValue)
 		: SavedValue(MoveTemp(InConfigValue.SavedValue))
 		, ExpandedValue(MoveTemp(InConfigValue.ExpandedValue))
+		, SavedValueHash(InConfigValue.SavedValueHash)
 #if CONFIG_REMEMBER_ACCESS_PATTERN 
 		, bRead(InConfigValue.bRead)
 #endif
@@ -154,6 +160,7 @@ public:
 	{
 		SavedValue = MoveTemp(RHS.SavedValue);
 		ExpandedValue = MoveTemp(RHS.ExpandedValue);
+		SavedValueHash = RHS.SavedValueHash;
 #if CONFIG_REMEMBER_ACCESS_PATTERN 
 		bRead = RHS.bRead;
 #endif
@@ -165,6 +172,7 @@ public:
 	{
 		SavedValue = RHS.SavedValue;
 		ExpandedValue = RHS.ExpandedValue;
+		SavedValueHash = RHS.SavedValueHash;
 #if CONFIG_REMEMBER_ACCESS_PATTERN 
 		bRead = RHS.bRead;
 #endif
@@ -202,7 +210,7 @@ public:
 	}
 #endif
 
-	bool operator==(const FConfigValue& Other) const { return (SavedValue.Compare(Other.SavedValue) == 0); }
+	bool operator==(const FConfigValue& Other) const { return SavedValueHash == Other.SavedValueHash; }
 	bool operator!=(const FConfigValue& Other) const { return !(FConfigValue::operator==(Other)); }
 
 	friend FArchive& operator<<(FArchive& Ar, FConfigValue& ConfigSection)
@@ -275,6 +283,7 @@ private:
 
 	FString SavedValue;
 	FString ExpandedValue;
+	uint32 SavedValueHash;
 #if CONFIG_REMEMBER_ACCESS_PATTERN 
 	mutable bool bRead; // has this value been read since the config system started
 #endif
