Add base64URL encoding support.

diff --git a/Engine/Source/Runtime/Core/Private/Misc/Base64.cpp b/Engine/Source/Runtime/Core/Private/Misc/Base64.cpp
--- a/Engine/Source/Runtime/Core/Private/Misc/Base64.cpp
+++ b/Engine/Source/Runtime/Core/Private/Misc/Base64.cpp
@@ -1,67 +1,85 @@
 // Copyright Epic Games, Inc. All Rights Reserved.
 
 #include "Misc/Base64.h"
-
 #include "Containers/StringConv.h"
 #include "Misc/AssertionMacros.h"
 
-/** The table used to encode a 6 bit value as an ascii character */
-static const uint8 EncodingAlphabet[64] = 
-{ 
-	'A', 'B', 'C', 'D', 'E', 'F', 'G', 'H', 'I', 'J', 'K', 'L', 'M', 'N', 'O', 'P', 
-	'Q', 'R', 'S', 'T', 'U', 'V', 'W', 'X', 'Y', 'Z', 'a', 'b', 'c', 'd', 'e', 'f', 
-	'g', 'h', 'i', 'j', 'k', 'l', 'm', 'n', 'o', 'p', 'q', 'r', 's', 't', 'u', 'v', 
-	'w', 'x', 'y', 'z', '0', '1', '2', '3', '4', '5', '6', '7', '8', '9', '+', '/' 
-};
-
-/** The table used to convert an ascii character into a 6 bit value */
-static const uint8 DecodingAlphabet[256] = 
-{ 
-	0xFF, 0xFF, 0xFF, 0xFF, 0xFF, 0xFF, 0xFF, 0xFF, 0xFF, 0xFF, 0xFF, 0xFF, 0xFF, 0xFF, 0xFF, 0xFF, // 0x00-0x0f
-	0xFF, 0xFF, 0xFF, 0xFF, 0xFF, 0xFF, 0xFF, 0xFF, 0xFF, 0xFF, 0xFF, 0xFF, 0xFF, 0xFF, 0xFF, 0xFF, // 0x10-0x1f
-	0xFF, 0xFF, 0xFF, 0xFF, 0xFF, 0xFF, 0xFF, 0xFF, 0xFF, 0xFF, 0xFF, 0x3E, 0xFF, 0xFF, 0xFF, 0x3F, // 0x20-0x2f
-	0x34, 0x35, 0x36, 0x37, 0x38, 0x39, 0x3A, 0x3B, 0x3C, 0x3D, 0xFF, 0xFF, 0xFF, 0xFF, 0xFF, 0xFF, // 0x30-0x3f
-	0xFF, 0x00, 0x01, 0x02, 0x03, 0x04, 0x05, 0x06, 0x07, 0x08, 0x09, 0x0A, 0x0B, 0x0C, 0x0D, 0x0E, // 0x40-0x4f
-	0x0F, 0x10, 0x11, 0x12, 0x13, 0x14, 0x15, 0x16, 0x17, 0x18, 0x19, 0xFF, 0xFF, 0xFF, 0xFF, 0xFF, // 0x50-0x5f
-	0xFF, 0x1A, 0x1B, 0x1C, 0x1D, 0x1E, 0x1F, 0x20, 0x21, 0x22, 0x23, 0x24, 0x25, 0x26, 0x27, 0x28, // 0x60-0x6f
-	0x29, 0x2A, 0x2B, 0x2C, 0x2D, 0x2E, 0x2F, 0x30, 0x31, 0x32, 0x33, 0xFF, 0xFF, 0xFF, 0xFF, 0xFF, // 0x70-0x7f
-	0xFF, 0xFF, 0xFF, 0xFF, 0xFF, 0xFF, 0xFF, 0xFF, 0xFF, 0xFF, 0xFF, 0xFF, 0xFF, 0xFF, 0xFF, 0xFF, // 0x80-0x8f
-	0xFF, 0xFF, 0xFF, 0xFF, 0xFF, 0xFF, 0xFF, 0xFF, 0xFF, 0xFF, 0xFF, 0xFF, 0xFF, 0xFF, 0xFF, 0xFF, // 0x90-0x9f
-	0xFF, 0xFF, 0xFF, 0xFF, 0xFF, 0xFF, 0xFF, 0xFF, 0xFF, 0xFF, 0xFF, 0xFF, 0xFF, 0xFF, 0xFF, 0xFF, // 0xa0-0xaf
-	0xFF, 0xFF, 0xFF, 0xFF, 0xFF, 0xFF, 0xFF, 0xFF, 0xFF, 0xFF, 0xFF, 0xFF, 0xFF, 0xFF, 0xFF, 0xFF, // 0xb0-0xbf
-	0xFF, 0xFF, 0xFF, 0xFF, 0xFF, 0xFF, 0xFF, 0xFF, 0xFF, 0xFF, 0xFF, 0xFF, 0xFF, 0xFF, 0xFF, 0xFF, // 0xc0-0xcf
-	0xFF, 0xFF, 0xFF, 0xFF, 0xFF, 0xFF, 0xFF, 0xFF, 0xFF, 0xFF, 0xFF, 0xFF, 0xFF, 0xFF, 0xFF, 0xFF, // 0xd0-0xdf
-	0xFF, 0xFF, 0xFF, 0xFF, 0xFF, 0xFF, 0xFF, 0xFF, 0xFF, 0xFF, 0xFF, 0xFF, 0xFF, 0xFF, 0xFF, 0xFF, // 0xe0-0xef
-	0xFF, 0xFF, 0xFF, 0xFF, 0xFF, 0xFF, 0xFF, 0xFF, 0xFF, 0xFF, 0xFF, 0xFF, 0xFF, 0xFF, 0xFF, 0xFF  // 0xf0-0xff
-};
-
-FString FBase64::Encode(const FString& Source)
+namespace
+{
+	/** The table used to encode a 6 bit value as an ascii character using Table 1 from RFC 4648 ("base64") */
+	static constexpr uint8 GBase64EncodingAlphabet[64] =
+	{
+		'A', 'B', 'C', 'D', 'E', 'F', 'G', 'H', 'I', 'J', 'K', 'L', 'M', 'N', 'O', 'P',
+		'Q', 'R', 'S', 'T', 'U', 'V', 'W', 'X', 'Y', 'Z', 'a', 'b', 'c', 'd', 'e', 'f',
+		'g', 'h', 'i', 'j', 'k', 'l', 'm', 'n', 'o', 'p', 'q', 'r', 's', 't', 'u', 'v',
+		'w', 'x', 'y', 'z', '0', '1', '2', '3', '4', '5', '6', '7', '8', '9', '+', '/'
+	};
+	/** The table used to encode a 6 bit value as an ascii character using Table 2 from RFC 4648 ("base64url") */
+	static constexpr uint8 GBase64UrlSafeEncodingAlphabet[64] =
+	{
+		'A', 'B', 'C', 'D', 'E', 'F', 'G', 'H', 'I', 'J', 'K', 'L', 'M', 'N', 'O', 'P',
+		'Q', 'R', 'S', 'T', 'U', 'V', 'W', 'X', 'Y', 'Z', 'a', 'b', 'c', 'd', 'e', 'f',
+		'g', 'h', 'i', 'j', 'k', 'l', 'm', 'n', 'o', 'p', 'q', 'r', 's', 't', 'u', 'v',
+		'w', 'x', 'y', 'z', '0', '1', '2', '3', '4', '5', '6', '7', '8', '9', '-', '_'
+	};
+
+	/** The table used to convert an ascii character into a 6 bit value using Table 1 from RFC 4648 ("base64") */
+	static constexpr uint8 GBase64DecodingAlphabet[128] =
+	{
+		0xFF, 0xFF, 0xFF, 0xFF, 0xFF, 0xFF, 0xFF, 0xFF, 0xFF, 0xFF, 0xFF, 0xFF, 0xFF, 0xFF, 0xFF, 0xFF, // 0x00-0x0f
+		0xFF, 0xFF, 0xFF, 0xFF, 0xFF, 0xFF, 0xFF, 0xFF, 0xFF, 0xFF, 0xFF, 0xFF, 0xFF, 0xFF, 0xFF, 0xFF, // 0x10-0x1f
+		0xFF, 0xFF, 0xFF, 0xFF, 0xFF, 0xFF, 0xFF, 0xFF, 0xFF, 0xFF, 0xFF, 0x3E, 0xFF, 0xFF, 0xFF, 0x3F, // 0x20-0x2f
+		0x34, 0x35, 0x36, 0x37, 0x38, 0x39, 0x3A, 0x3B, 0x3C, 0x3D, 0xFF, 0xFF, 0xFF, 0xFF, 0xFF, 0xFF, // 0x30-0x3f
+		0xFF, 0x00, 0x01, 0x02, 0x03, 0x04, 0x05, 0x06, 0x07, 0x08, 0x09, 0x0A, 0x0B, 0x0C, 0x0D, 0x0E, // 0x40-0x4f
+		0x0F, 0x10, 0x11, 0x12, 0x13, 0x14, 0x15, 0x16, 0x17, 0x18, 0x19, 0xFF, 0xFF, 0xFF, 0xFF, 0xFF, // 0x50-0x5f
+		0xFF, 0x1A, 0x1B, 0x1C, 0x1D, 0x1E, 0x1F, 0x20, 0x21, 0x22, 0x23, 0x24, 0x25, 0x26, 0x27, 0x28, // 0x60-0x6f
+		0x29, 0x2A, 0x2B, 0x2C, 0x2D, 0x2E, 0x2F, 0x30, 0x31, 0x32, 0x33, 0xFF, 0xFF, 0xFF, 0xFF, 0xFF  // 0x70-0x7f
+	};
+
+	/** The table used to convert an ascii character into a 6 bit value using Table 2 from RFC 4648 ("base64url") */
+	static constexpr uint8 GBase64UrlSafeDecodingAlphabet[128] =
+	{
+		0xFF, 0xFF, 0xFF, 0xFF, 0xFF, 0xFF, 0xFF, 0xFF, 0xFF, 0xFF, 0xFF, 0xFF, 0xFF, 0xFF, 0xFF, 0xFF, // 0x00-0x0f
+		0xFF, 0xFF, 0xFF, 0xFF, 0xFF, 0xFF, 0xFF, 0xFF, 0xFF, 0xFF, 0xFF, 0xFF, 0xFF, 0xFF, 0xFF, 0xFF, // 0x10-0x1f
+		0xFF, 0xFF, 0xFF, 0xFF, 0xFF, 0xFF, 0xFF, 0xFF, 0xFF, 0xFF, 0xFF, 0xFF, 0xFF, 0x3E, 0xFF, 0xFF, // 0x20-0x2f
+		0x34, 0x35, 0x36, 0x37, 0x38, 0x39, 0x3A, 0x3B, 0x3C, 0x3D, 0xFF, 0xFF, 0xFF, 0xFF, 0xFF, 0xFF, // 0x30-0x3f
+		0xFF, 0x00, 0x01, 0x02, 0x03, 0x04, 0x05, 0x06, 0x07, 0x08, 0x09, 0x0A, 0x0B, 0x0C, 0x0D, 0x0E, // 0x40-0x4f
+		0x0F, 0x10, 0x11, 0x12, 0x13, 0x14, 0x15, 0x16, 0x17, 0x18, 0x19, 0xFF, 0xFF, 0xFF, 0xFF, 0x3F, // 0x50-0x5f
+		0xFF, 0x1A, 0x1B, 0x1C, 0x1D, 0x1E, 0x1F, 0x20, 0x21, 0x22, 0x23, 0x24, 0x25, 0x26, 0x27, 0x28, // 0x60-0x6f
+		0x29, 0x2A, 0x2B, 0x2C, 0x2D, 0x2E, 0x2F, 0x30, 0x31, 0x32, 0x33, 0xFF, 0xFF, 0xFF, 0xFF, 0xFF  // 0x70-0x7f
+	};
+}
+
+FString FBase64::Encode(const FString& Source, EBase64Mode Mode)
 {
-	return Encode((uint8*)TCHAR_TO_ANSI(*Source), Source.Len());
+	return Encode((uint8*)TCHAR_TO_ANSI(*Source), Source.Len(), Mode);
 }
 
-FString FBase64::Encode(const TArray<uint8>& Source)
+FString FBase64::Encode(const TArray<uint8>& Source, EBase64Mode Mode)
 {
-	return Encode((uint8*)Source.GetData(), Source.Num());
+	return Encode((uint8*)Source.GetData(), Source.Num(), Mode);
 }
 
-FString FBase64::Encode(const uint8* Source, uint32 Length)
+FString FBase64::Encode(const uint8* Source, uint32 Length, EBase64Mode Mode)
 {
 	uint32 ExpectedLength = GetEncodedDataSize(Length);
-	
+
 	FString OutBuffer;
 
 	TArray<TCHAR, FString::AllocatorType>& OutCharArray = OutBuffer.GetCharArray();
 	OutCharArray.SetNum(ExpectedLength + 1);
-	int64 EncodedLength = Encode(Source, Length, OutCharArray.GetData());
+	int64 EncodedLength = Encode(Source, Length, OutCharArray.GetData(), Mode);
 	verify(EncodedLength == OutBuffer.Len());
 
 	return OutBuffer;
 }
 
-template<typename CharType> uint32 FBase64::Encode(const uint8* Source, uint32 Length, CharType* Dest)
+template<typename CharType> uint32 FBase64::Encode(const uint8* Source, uint32 Length, CharType* Dest, EBase64Mode Mode)
 {
+	check(Mode == EBase64Mode::Standard || Mode == EBase64Mode::UrlSafe);
+
 	CharType* EncodedBytes = Dest;
+	const uint8* const EncodingAlphabet = (Mode == EBase64Mode::UrlSafe ? GBase64UrlSafeEncodingAlphabet : GBase64EncodingAlphabet);
 
 	// Loop through the buffer converting 3 bytes of binary data at a time
 	while (Length >= 3)
@@ -126,16 +144,16 @@ template<typename CharType> uint32 FBase64::Encode(const uint8* Source, uint32 L
 	return UE_PTRDIFF_TO_UINT32(EncodedBytes - Dest);
 }
 
-template CORE_API uint32 FBase64::Encode<ANSICHAR>(const uint8* Source, uint32 Length, ANSICHAR* Dest);
-template CORE_API uint32 FBase64::Encode<WIDECHAR>(const uint8* Source, uint32 Length, WIDECHAR* Dest);
+template CORE_API uint32 FBase64::Encode<ANSICHAR>(const uint8* Source, uint32 Length, ANSICHAR* Dest, EBase64Mode Mode);
+template CORE_API uint32 FBase64::Encode<WIDECHAR>(const uint8* Source, uint32 Length, WIDECHAR* Dest, EBase64Mode Mode);
 
-bool FBase64::Decode(const FString& Source, FString& OutDest)
+bool FBase64::Decode(const FString& Source, FString& OutDest, EBase64Mode Mode)
 {
 	uint32 ExpectedLength = GetDecodedDataSize(Source);
 
 	TArray<ANSICHAR> TempDest;
 	TempDest.AddZeroed(ExpectedLength + 1);
-	if(!Decode(*Source, Source.Len(), (uint8*)TempDest.GetData()))
+	if (!Decode(*Source, Source.Len(), (uint8*)TempDest.GetData(), Mode))
 	{
 		return false;
 	}
@@ -143,35 +161,39 @@ bool FBase64::Decode(const FString& Source, FString& OutDest)
 	return true;
 }
 
-bool FBase64::Decode(const FString& Source, TArray<uint8>& OutDest)
+bool FBase64::Decode(const FString& Source, TArray<uint8>& OutDest, EBase64Mode Mode)
 {
 	OutDest.SetNum(GetDecodedDataSize(Source));
-	return Decode(*Source, Source.Len(), OutDest.GetData());
+	return Decode(*Source, Source.Len(), OutDest.GetData(), Mode);
 }
 
-template<typename CharType> bool FBase64::Decode(const CharType* Source, uint32 Length, uint8* Dest)
+template<typename CharType> bool FBase64::Decode(const CharType* Source, uint32 Length, uint8* Dest, EBase64Mode Mode)
 {
+	check(Mode == EBase64Mode::Standard || Mode == EBase64Mode::UrlSafe);
+
 	// Remove the trailing '=' characters, so we can handle padded and non-padded input the same
-	while(Length > 0 && Source[Length - 1] == '=')
+	while (Length > 0 && Source[Length - 1] == '=')
 	{
 		Length--;
 	}
 
 	// Make sure the length is valid. Only lengths modulo 4 of 0, 2, and 3 are valid.
-	if((Length & 3) == 1)
+	if ((Length & 3) == 1)
 	{
 		return false;
 	}
 
+	const uint8* const DecodingAlphabet = (Mode == EBase64Mode::UrlSafe ? GBase64UrlSafeDecodingAlphabet : GBase64DecodingAlphabet);
+
 	// Convert all the full chunks of data
-	for(; Length >= 4; Length -= 4)
+	for (; Length >= 4; Length -= 4)
 	{
 		// Decode the next 4 BYTEs
 		uint32 OriginalTriplet = 0;
 		for (int32 Index = 0; Index < 4; Index++)
 		{
 			uint32 SourceChar = (uint32)*Source++;
-			if(SourceChar >= 256)
+			if (SourceChar >= 128)
 			{
 				return false;
 			}
@@ -195,14 +217,14 @@ template<typename CharType> bool FBase64::Decode(const CharType* Source, uint32
 	}
 
 	// Convert the last chunk of data
-	if(Length > 0)
+	if (Length > 0)
 	{
 		// Decode the next 4 BYTEs, or up to the end of the input buffer
 		uint32 OriginalTriplet = 0;
 		for (uint32 Index = 0; Index < Length; Index++)
 		{
 			uint32 SourceChar = (uint32)*Source++;
-			if(SourceChar >= 256)
+			if (SourceChar >= 128)
 			{
 				return false;
 			}
@@ -213,14 +235,14 @@ template<typename CharType> bool FBase64::Decode(const CharType* Source, uint32
 			}
 			OriginalTriplet = (OriginalTriplet << 6) | DecodedValue;
 		}
-		for(int32 Index = Length; Index < 4; Index++)
+		for (int32 Index = Length; Index < 4; Index++)
 		{
 			OriginalTriplet = (OriginalTriplet << 6);
 		}
 
 		// Now we can tear the uint32 into bytes
 		OriginalTriplet >>= 8;
-		if(Length >= 3)
+		if (Length >= 3)
 		{
 			Dest[1] = OriginalTriplet & 0xFF;
 		}
@@ -230,8 +252,8 @@ template<typename CharType> bool FBase64::Decode(const CharType* Source, uint32
 	return true;
 }
 
-template CORE_API bool FBase64::Decode<ANSICHAR>(const ANSICHAR* Source, uint32 Length, uint8* Dest);
-template CORE_API bool FBase64::Decode<WIDECHAR>(const WIDECHAR* Source, uint32 Length, uint8* Dest);
+template CORE_API bool FBase64::Decode<ANSICHAR>(const ANSICHAR* Source, uint32 Length, uint8* Dest, EBase64Mode Mode);
+template CORE_API bool FBase64::Decode<WIDECHAR>(const WIDECHAR* Source, uint32 Length, uint8* Dest, EBase64Mode Mode);
 
 uint32 FBase64::GetDecodedDataSize(const FString& Source)
 {
@@ -241,21 +263,21 @@ uint32 FBase64::GetDecodedDataSize(const FString& Source)
 template<typename CharType> uint32 FBase64::GetDecodedDataSize(const CharType* Source, uint32 Length)
 {
 	uint32 NumBytes = 0;
-	if(Length > 0)
+	if (Length > 0)
 	{
 		// Get the source length without the trailing padding characters
-		while(Length > 0 && Source[Length - 1] == '=')
+		while (Length > 0 && Source[Length - 1] == '=')
 		{
 			Length--;
 		}
 
 		// Get the lower bound for number of bytes, excluding the last partial chunk
 		NumBytes += (Length / 4) * 3;
-		if((Length & 3) == 3)
+		if ((Length & 3) == 3)
 		{
 			NumBytes += 2;
 		}
-		if((Length & 3) == 2)
+		if ((Length & 3) == 2)
 		{
 			NumBytes++;
 		}
diff --git a/Engine/Source/Runtime/Core/Public/Misc/Base64.h b/Engine/Source/Runtime/Core/Public/Misc/Base64.h
--- a/Engine/Source/Runtime/Core/Public/Misc/Base64.h
+++ b/Engine/Source/Runtime/Core/Public/Misc/Base64.h
@@ -1,4 +1,4 @@
-﻿// Copyright Epic Games, Inc. All Rights Reserved.
+// Copyright Epic Games, Inc. All Rights Reserved.
 
 #pragma once
 
@@ -7,6 +7,15 @@
 #include "CoreTypes.h"
 #include "Misc/Timespan.h"
 
+/** Mode to use for base64 encoding and decoding */
+enum class EBase64Mode : uint8
+{
+	/** Use the standard set of character mappings. (Table 1 from RFC 4648, known as "base64") */
+	Standard,
+	/** Use the "URL and Filename safe" set of character mappings. (Table 2 from RFC 4648, known as "base64url") */
+	UrlSafe
+};
+
 /**
  * Class for encoding/decoding Base64 data (RFC 4648)
  */
@@ -16,29 +25,32 @@ struct CORE_API FBase64
 	 * Encodes a FString into a Base64 string
 	 *
 	 * @param Source The string data to convert
+	 * @param Mode The mode to use for encoding. Default is EBase64Mode::Standard
 	 *
 	 * @return A string that encodes the binary data in a way that can be safely transmitted via various Internet protocols
 	 */
-	static FString Encode(const FString& Source);
+	static FString Encode(const FString& Source, EBase64Mode Mode = EBase64Mode::Standard);
 
 	/**
 	 * Encodes a binary uint8 array into a Base64 string
 	 *
 	 * @param Source The binary data to convert
+	 * @param Mode The mode to use for encoding. Default is EBase64Mode::Standard
 	 *
 	 * @return A string that encodes the binary data in a way that can be safely transmitted via various Internet protocols
 	 */
-	static FString Encode(const TArray<uint8>& Source);
+	static FString Encode(const TArray<uint8>& Source, EBase64Mode Mode = EBase64Mode::Standard);
 
 	/**
 	 * Encodes the source into a Base64 string
 	 *
 	 * @param Source The binary data to encode
 	 * @param Length Length of the binary data to be encoded
+	 * @param Mode The mode to use for encoding. Default is EBase64Mode::Standard
 	 *
 	 * @return Base64 encoded string containing the binary data.
 	 */
-	static FString Encode(const uint8* Source, uint32 Length);
+	static FString Encode(const uint8* Source, uint32 Length, EBase64Mode Mode = EBase64Mode::Standard);
 
 	/**
 	 * Encodes the source into a Base64 string, storing it in a preallocated buffer.
@@ -46,10 +58,11 @@ struct CORE_API FBase64
 	 * @param Source The binary data to encode
 	 * @param Length Length of the binary data to be encoded
 	 * @param Dest Buffer to receive the encoded data. Must be large enough to contain the entire output data (see GetEncodedDataSize()).
+	 * @param Mode The mode to use for encoding. Default is EBase64Mode::Standard
 	 *
 	 * @return The length of the encoded data
 	 */
-	template<typename CharType> static uint32 Encode(const uint8* Source, uint32 Length, CharType* Dest);
+	template<typename CharType> static uint32 Encode(const uint8* Source, uint32 Length, CharType* Dest, EBase64Mode Mode = EBase64Mode::Standard);
 
 	/**
 	* Get the encoded data size for the given number of bytes.
@@ -68,16 +81,22 @@ struct CORE_API FBase64
 	 *
 	 * @param Source The Base64 encoded string
 	 * @param OutDest Receives the decoded string data
+	 * @param Mode The mode to use for decoding. Default is EBase64Mode::Standard
+	 *
+	 * @return true if the buffer was decoded, false if it was invalid.
 	 */
-	static bool Decode(const FString& Source, FString& OutDest);
+	static bool Decode(const FString& Source, FString& OutDest, EBase64Mode Mode = EBase64Mode::Standard);
 
 	/**
 	 * Decodes a Base64 string into an array of bytes
 	 *
 	 * @param Source The Base64 encoded string
 	 * @param Dest Array to receive the decoded data
+	 * @param Mode The mode to use for decoding. Default is EBase64Mode::Standard
+	 *
+	 * @return true if the buffer was decoded, false if it was invalid.
 	 */
-	static bool Decode(const FString& Source, TArray<uint8>& Dest);
+	static bool Decode(const FString& Source, TArray<uint8>& Dest, EBase64Mode Mode = EBase64Mode::Standard);
 
 	/**
 	 * Decodes a Base64 string into a preallocated buffer
@@ -85,10 +104,11 @@ struct CORE_API FBase64
 	 * @param Source The Base64 encoded string
 	 * @param Length Length of the Base64 encoded string
 	 * @param Dest Buffer to receive the decoded data
+	 * @param Mode The mode to use for decoding. Default is EBase64Mode::Standard
 	 *
 	 * @return true if the buffer was decoded, false if it was invalid.
 	 */
-	template<typename CharType> static bool Decode(const CharType* Source, uint32 Length, uint8* Dest);
+	template<typename CharType> static bool Decode(const CharType* Source, uint32 Length, uint8* Dest, EBase64Mode Mode = EBase64Mode::Standard);
 
 	/**
 	* Determine the decoded data size for the incoming base64 encoded string
@@ -120,4 +140,4 @@ struct CORE_API FBase64
 	{
 		return ((NumChars * 3) + 3) / 4;
 	}
-};
+};
\ No newline at end of file
diff --git a/Engine/Source/Runtime/Core/Tests/Serialization/Base64Test.cpp b/Engine/Source/Runtime/Core/Tests/Serialization/Base64Test.cpp
new file mode 100644
--- /dev/null
+++ b/Engine/Source/Runtime/Core/Tests/Serialization/Base64Test.cpp
@@ -0,0 +1,162 @@
+// Copyright Epic Games, Inc. All Rights Reserved.
+
+#if WITH_LOW_LEVEL_TESTS
+
+#include "Misc/Base64.h"
+#include "TestHarness.h"
+
+template <size_t InputSize>
+struct TBase64Test
+{
+	const uint8 EncodeInput[InputSize]; // aka DecodeOutput
+	const FString EncodeOutput; // aka DecodeInput
+
+	TArray<uint8> GetInputArray() const
+	{
+		return TArray<uint8>(EncodeInput, InputSize);
+	}
+};
+
+struct FBase64StringTest
+{
+	const FString EncodeInput;
+	const FString EncodeOutput;
+
+	TArray<uint8> GetInputArray() const
+	{
+		TArray<uint8> Out;
+		Out.AddUninitialized(EncodeInput.Len());
+
+		for (int Ix = 0; Ix < EncodeInput.Len(); ++Ix)
+		{
+			CHECK((EncodeInput[Ix] & 0xff) == EncodeInput[Ix]);
+			Out[Ix] = (uint8)EncodeInput[Ix];
+		}
+		return Out;
+	}
+};
+
+template <size_t NumVectors, typename TestType>
+void RunRoundTripSuccessfulBase64Tests(const TestType (&TestVectors)[NumVectors], EBase64Mode Mode)
+{
+	for (size_t Ix = 0; Ix < NumVectors; ++Ix)
+	{
+		const auto& Test = TestVectors[Ix];
+		const TArray<uint8> InputArray = Test.GetInputArray();
+
+		const FString Encoded = FBase64::Encode(InputArray, Mode);
+		CHECK(Encoded == Test.EncodeOutput);
+
+		TArray<uint8> DecodeOut;
+		CHECK(FBase64::Decode(Test.EncodeOutput, DecodeOut, Mode));
+		CHECK(DecodeOut == InputArray);
+	}
+}
+
+// Base64 encode/decode works for the test vectors provided in the RFC
+TEST_CASE("Core::Serialization::Base64::encode/decode RFC test vectors", "[Core][base64][Smoke]")
+{
+	// This test does a good job of verifying the correct padding of various lengths by going through 2 cycles of the padding modulus.
+	const FBase64StringTest TestVectors[] =
+	{
+		{ FString(TEXT("")), FString(TEXT("")) },
+		{ FString(TEXT("f")), FString(TEXT("Zg==")) },
+		{ FString(TEXT("fo")), FString(TEXT("Zm8=")) },
+		{ FString(TEXT("foo")), FString(TEXT("Zm9v")) },
+		{ FString(TEXT("foob")), FString(TEXT("Zm9vYg==")) },
+		{ FString(TEXT("fooba")), FString(TEXT("Zm9vYmE=")) },
+		{ FString(TEXT("foobar")), FString(TEXT("Zm9vYmFy")) },
+	};
+	RunRoundTripSuccessfulBase64Tests(TestVectors, EBase64Mode::Standard);
+
+	// The results are the same for URL safe as these vectors never produce 0x1f or 0x3f
+	RunRoundTripSuccessfulBase64Tests(TestVectors, EBase64Mode::UrlSafe);
+}
+
+// Standard base64 encode/decode every possible 6 bit input
+TEST_CASE("Core::Serialization::Base64::Standard encode/decode every 6 bit input", "[Core][base64][Smoke]")
+{
+	const TBase64Test<48> TestVectors[] =
+	{
+		{ { 0,16,131,16,81,135,32,146,139,48,211,143,65,20,147,81,85,151,97,150,155,113,215,159,130,24,163,146,89,167,162,154,171,178,219,175,195,28,179,211,93,183,227,158,187,243,223,191 }, FString(TEXT("ABCDEFGHIJKLMNOPQRSTUVWXYZabcdefghijklmnopqrstuvwxyz0123456789+/")) }
+	};
+
+	RunRoundTripSuccessfulBase64Tests(TestVectors, EBase64Mode::Standard);
+}
+
+// UrlSafe base64 encode/decode every possible 6 bit input
+TEST_CASE("Core::Serialization::Base64::URLSafe encode/decode every 6 bit input", "[Core][base64][Smoke]")
+{
+	const TBase64Test<48> TestVectors[] = 
+	{
+		{ { 0,16,131,16,81,135,32,146,139,48,211,143,65,20,147,81,85,151,97,150,155,113,215,159,130,24,163,146,89,167,162,154,171,178,219,175,195,28,179,211,93,183,227,158,187,243,223,191 }, FString(TEXT("ABCDEFGHIJKLMNOPQRSTUVWXYZabcdefghijklmnopqrstuvwxyz0123456789-_")) }
+	};
+
+	RunRoundTripSuccessfulBase64Tests(TestVectors, EBase64Mode::UrlSafe);
+}
+
+void RunNegativeDecodeBase64Test(EBase64Mode Mode, TFunction<bool(int)>&& IsInDecodeAlphabet)
+{
+	TCHAR BadDecodeBuffer[]
+	{
+		TEXT('A'), TEXT('B'), TEXT('C'), TEXT('D'), TEXT('E'), TEXT('F'), TEXT('G'), TEXT('H'), TEXT('I'), TEXT('J'), TEXT('K'), TEXT('L'), TEXT('M'), TEXT('N'), TEXT('O'), TEXT('P'),
+		TEXT('Q'), TEXT('R'), TEXT('S'), TEXT('T'), TEXT('U'), TEXT('V'), TEXT('W'), TEXT('X'), TEXT('Y'), TEXT('Z'), TEXT('a'), TEXT('b'), TEXT('c'), TEXT('d'), TEXT('e'), TEXT('f'),
+		TEXT('g'), TEXT('h'), TEXT('i'), TEXT('j'), TEXT('k'), TEXT('l'), TEXT('m'), TEXT('n'), TEXT('o'), TEXT('p'), TEXT('q'), TEXT('r'), TEXT('s'), TEXT('t'), TEXT('u'), TEXT('v'),
+		TEXT('w'), TEXT('x'), TEXT('y'), TEXT('z'), TEXT('0'), TEXT('1'), TEXT('2'), TEXT('3'), TEXT('4'), TEXT('5'), TEXT('6'), TEXT('7'), TEXT('8'), TEXT('9'), TEXT('+'), TEXT('/'),
+		TEXT('A'), TEXT('=') /* bad bytes will go here */, TEXT('='), TEXT('=')
+	};
+
+	const uint32 NumChars = UE_ARRAY_COUNT(BadDecodeBuffer);
+	TArray<uint8> DecodeOut;
+	DecodeOut.AddZeroed(FBase64::GetMaxDecodedDataSize(NumChars));
+
+	// Test the end of the buffer so we're in an area that is not decoding 4 bytes at a time.
+	for (int Ix = 0; Ix <= 256; ++Ix)
+	{
+		if (IsInDecodeAlphabet(Ix))
+		{
+			continue;
+		}
+		BadDecodeBuffer[NumChars - 3] = TCHAR(Ix);
+		CHECK_FALSE(FBase64::Decode(BadDecodeBuffer, NumChars, DecodeOut.GetData(), Mode));
+	}
+
+	// Now test the beginning of the buffer so we're in an area that *is* decoding 4 bytes at a time.
+	for (int Ix = 0; Ix <= 256; ++Ix)
+	{
+		if (IsInDecodeAlphabet(Ix))
+		{
+			continue;
+		}
+		BadDecodeBuffer[0] = TCHAR(Ix);
+		// We use NumChars-4 as the length to ignore the end of buffer sequence.
+		CHECK_FALSE(FBase64::Decode(BadDecodeBuffer, NumChars - 4, DecodeOut.GetData(), Mode));
+	}
+}
+
+// Base64 decode should fail on input that is not in the decoding alphabet
+TEST_CASE("Core::Serialization::Base64::Negative", "[Core][base64][Smoke]")
+{
+	auto IsInStandardAlphabet = [](int Ix)
+	{
+		return (Ix >= 'A' && Ix <= 'Z')
+			|| (Ix >= 'a' && Ix <= 'z')
+			|| (Ix >= '0' && Ix <= '9')
+			|| Ix == '+'
+			|| Ix == '/';
+	};
+
+	auto IsInUrlSafeAlphabet = [](int Ix)
+	{
+		return (Ix >= 'A' && Ix <= 'Z')
+			|| (Ix >= 'a' && Ix <= 'z')
+			|| (Ix >= '0' && Ix <= '9')
+			|| Ix == '-'
+			|| Ix == '_';
+	};
+
+	RunNegativeDecodeBase64Test(EBase64Mode::Standard, MoveTemp(IsInStandardAlphabet));
+	RunNegativeDecodeBase64Test(EBase64Mode::UrlSafe, MoveTemp(IsInUrlSafeAlphabet));
+}
+
+#endif
