[Runtime]
* ConfigUtilities.cpp - Added explicit type to enum class to see if it can solve the GetTypeHash error seen in this file on certain configs

diff --git a/Engine/Source/Runtime/Core/Private/Misc/ConfigUtilities.cpp b/Engine/Source/Runtime/Core/Private/Misc/ConfigUtilities.cpp
--- a/Engine/Source/Runtime/Core/Private/Misc/ConfigUtilities.cpp
+++ b/Engine/Source/Runtime/Core/Private/Misc/ConfigUtilities.cpp
@@ -271,7 +271,7 @@ TUniquePtr<FCVarIniHistoryHelper> IniHistoryHelper;
 class FConfigHistoryHelper
 {
 private:
-	enum class HistoryType
+	enum class HistoryType : int
 	{
 		Value,
 		Section,
