LLM: Add missing pragma once.
#rnx

diff --git a/Engine/Source/Runtime/Core/Private/Windows/VirtualAllocPageStatus.h b/Engine/Source/Runtime/Core/Private/Windows/VirtualAllocPageStatus.h
--- a/Engine/Source/Runtime/Core/Private/Windows/VirtualAllocPageStatus.h
+++ b/Engine/Source/Runtime/Core/Private/Windows/VirtualAllocPageStatus.h
@@ -1,4 +1,5 @@
 // Copyright Epic Games, Inc. All Rights Reserved.
+#pragma once
 
 #include "HAL/LowLevelMemTracker.h"
 
