Large page allocator can preallocate a specific amount of memory and keep it for its usage (to avoid physical page defragmentation on certain platforms). If title might need memory, the unused memory of large page allocator will be trimmed

#rnx

diff --git a/Engine/Source/Runtime/Core/Private/HAL/Allocators/CachedOSVeryLargePageAllocator.cpp b/Engine/Source/Runtime/Core/Private/HAL/Allocators/CachedOSVeryLargePageAllocator.cpp
--- a/Engine/Source/Runtime/Core/Private/HAL/Allocators/CachedOSVeryLargePageAllocator.cpp
+++ b/Engine/Source/Runtime/Core/Private/HAL/Allocators/CachedOSVeryLargePageAllocator.cpp
@@ -8,11 +8,23 @@
 #include "ProfilingDebugging/CsvProfiler.h"
 
 #if UE_USE_VERYLARGEPAGEALLOCATOR
-static int32 GVeryLargePageAllocatorMaxEmptyBackstore = 10;
-static FAutoConsoleVariableRef CVarVeryLargePageAllocatorMaxEmptyBackstore(
-	TEXT("VeryLargePageAllocator.MaxEmptyBackstore"),
-	GVeryLargePageAllocatorMaxEmptyBackstore,
-	TEXT(""));
+
+// Free pages that are cached can be preallocated at the initialization of the VeryLargePageAllocator
+// Having them preallocated and kept during the life of the title help to avoid defragmentation of physical memory
+#ifndef UE_VERYLARGEPAGEALLOCATOR_PREALLOCATE_EMPTY_BACKSTORE
+	#define UE_VERYLARGEPAGEALLOCATOR_PREALLOCATE_EMPTY_BACKSTORE 0
+#endif
+
+// Number of free pages (2MB each) to be kept (not decommited) in cache so they can be reused:
+// Number of free pages to cache for SmallPool allocation hint. 
+#ifndef UE_VERYLARGEPAGEALLOCATOR_MAX_EMPTY_BACKSTORE_SMALLPOOL
+	#define UE_VERYLARGEPAGEALLOCATOR_MAX_EMPTY_BACKSTORE_SMALLPOOL 10
+#endif
+
+// Number of free pages to cache for Default allocation hint.
+#ifndef UE_VERYLARGEPAGEALLOCATOR_MAX_EMPTY_BACKSTORE_DEFAULT
+	#define UE_VERYLARGEPAGEALLOCATOR_MAX_EMPTY_BACKSTORE_DEFAULT 10
+#endif
 
 CORE_API bool GEnableVeryLargePageAllocator = true;
 
@@ -21,7 +33,7 @@ CSV_DECLARE_CATEGORY_MODULE_EXTERN(CORE_API, FMemory);
 
 static volatile int32 GLargePageAllocatorCommitCount = 0;
 static volatile int32 GLargePageAllocatorDecommitCount = 0;
-static volatile int32 GLargePageAllocatorLargePageCount = 0;
+static volatile int32 GLargePageAllocatorLargePageCount[FMemory::AllocationHints::Max];
 #endif
 
 void FCachedOSVeryLargePageAllocator::Init()
@@ -43,7 +55,14 @@ void FCachedOSVeryLargePageAllocator::Init()
 		UsedLargePagesHead[i] = nullptr;
 		EmptyButAvailableLargePagesHead[i] = nullptr;
 		EmptyBackStoreCount[i] = 0;
+		MaxEmptyBackStoreCount[i] = 0;
+		
+		GLargePageAllocatorLargePageCount[i] = 0;
 	}
+
+	MaxEmptyBackStoreCount[FMemory::AllocationHints::SmallPool] = UE_VERYLARGEPAGEALLOCATOR_MAX_EMPTY_BACKSTORE_SMALLPOOL;
+	MaxEmptyBackStoreCount[FMemory::AllocationHints::Default] = UE_VERYLARGEPAGEALLOCATOR_MAX_EMPTY_BACKSTORE_DEFAULT;
+
 #if UE_VERYLARGEPAGEALLOCATOR_TAKEONALL64KBALLOCATIONS
 	for (int i = 0; i < NumberOfLargePages / 2; i++)
 #else
@@ -62,6 +81,25 @@ void FCachedOSVeryLargePageAllocator::Init()
 	}
 #endif
 
+#if UE_VERYLARGEPAGEALLOCATOR_PREALLOCATE_EMPTY_BACKSTORE
+	for (int32 i = 0; i < FMemory::AllocationHints::Max; i++)
+	{
+		FMemory::AllocationHints AllocationHint = FMemory::AllocationHints(i);
+		
+		for (int32 BackStoreIndex = 0; BackStoreIndex < MaxEmptyBackStoreCount[AllocationHint]; BackStoreIndex++)
+		{
+			FLargePage* LargePage = FreeLargePagesHead[AllocationHint];
+
+			check(LargePage != nullptr); // Can't happen
+			LargePage->AllocationHint = AllocationHint;
+			LargePage->Unlink();
+			Block.Commit(LargePage->BaseAddress - AddressSpaceReserved, SizeOfLargePage, false);
+			LargePage->LinkHead(EmptyButAvailableLargePagesHead[AllocationHint]);
+			EmptyBackStoreCount[AllocationHint] += 1;
+		}
+	}
+#endif
+
 	if (!GEnableVeryLargePageAllocator)
 	{
 		bEnabled = false;
@@ -129,7 +167,7 @@ void* FCachedOSVeryLargePageAllocator::Allocate(SIZE_T Size, uint32 AllocationHi
 							else
 							{
 								FPlatformAtomics::InterlockedIncrement(&GLargePageAllocatorCommitCount);
-								FPlatformAtomics::InterlockedIncrement(&GLargePageAllocatorLargePageCount);
+								FPlatformAtomics::InterlockedIncrement(&GLargePageAllocatorLargePageCount[AllocationHint]);
 							}
 #endif
 						}
@@ -196,7 +234,7 @@ void FCachedOSVeryLargePageAllocator::Free(void* Ptr, SIZE_T Size, FCriticalSect
 			LargePage->Unlink();
 
 			// move it to EmptyButAvailableLargePagesHead if that pool of backstore is not full yet
-			if (EmptyBackStoreCount[LargePage->AllocationHint] < GVeryLargePageAllocatorMaxEmptyBackstore)
+			if (EmptyBackStoreCount[LargePage->AllocationHint] < MaxEmptyBackStoreCount[LargePage->AllocationHint])
 			{
 				LargePage->LinkHead(EmptyButAvailableLargePagesHead[LargePage->AllocationHint]);
 				EmptyBackStoreCount[LargePage->AllocationHint] += 1;
@@ -213,7 +251,7 @@ void FCachedOSVeryLargePageAllocator::Free(void* Ptr, SIZE_T Size, FCriticalSect
 
 #if CSV_PROFILER
 				FPlatformAtomics::InterlockedIncrement(&GLargePageAllocatorDecommitCount);
-				FPlatformAtomics::InterlockedDecrement(&GLargePageAllocatorLargePageCount);
+				FPlatformAtomics::InterlockedDecrement(&GLargePageAllocatorLargePageCount[LargePage->AllocationHint]);
 #endif
 
 				LargePage->LinkHead(FreeLargePagesHead[LargePage->AllocationHint]);
@@ -283,22 +321,45 @@ void FCachedOSVeryLargePageAllocator::Free(void* Ptr, SIZE_T Size, FCriticalSect
 
 void FCachedOSVeryLargePageAllocator::FreeAll(FCriticalSection* Mutex)
 {
+	for (int i = 0; i < FMemory::AllocationHints::Max; i++)
+	{
+		FMemory::AllocationHints AllocationHint = FMemory::AllocationHints(i);
+		
+		for (;;)
+		{
+			FLargePage* LargePage = EmptyButAvailableLargePagesHead[AllocationHint];
+
+			if (LargePage == nullptr)
+			{
+				break;
+			}
+			LargePage->Unlink();
+			EmptyBackStoreCount[AllocationHint] -= 1;
+			Block.Decommit(LargePage->BaseAddress - AddressSpaceReserved, SizeOfLargePage);
+
+			LargePage->LinkHead(FreeLargePagesHead[LargePage->AllocationHint]);
+			CachedFree -= SizeOfLargePage;
+
+#if CSV_PROFILER
+			FPlatformAtomics::InterlockedDecrement(&GLargePageAllocatorLargePageCount[LargePage->AllocationHint]);
+#endif
+		}
+	}
+
+
+	// Free empty cached pages of CachedOSPageAllocator
 	CachedOSPageAllocator.FreeAll(Mutex);
 }
 
 void FCachedOSVeryLargePageAllocator::UpdateStats()
 {
 #if CSV_PROFILER
-	int32 BackingStoreCount = 0;
-	for (int i = 0; i < FMemory::AllocationHints::Max; i++)
-	{
-		BackingStoreCount += EmptyBackStoreCount[i];
-	}
-
 	CSV_CUSTOM_STAT(FMemory, LargeAllocatorCommitCount, GLargePageAllocatorCommitCount, ECsvCustomStatOp::Set);
 	CSV_CUSTOM_STAT(FMemory, LargeAllocatorDecommitCount, GLargePageAllocatorDecommitCount, ECsvCustomStatOp::Set);
-	CSV_CUSTOM_STAT(FMemory, LargeAllocatorBackingStoreCount, BackingStoreCount, ECsvCustomStatOp::Set);
-	CSV_CUSTOM_STAT(FMemory, LargeAllocatorPageCount, GLargePageAllocatorLargePageCount, ECsvCustomStatOp::Set);
+	CSV_CUSTOM_STAT(FMemory, LargeAllocatorBackingStoreCountSmall, EmptyBackStoreCount[FMemory::AllocationHints::SmallPool], ECsvCustomStatOp::Set);
+	CSV_CUSTOM_STAT(FMemory, LargeAllocatorBackingStoreCountDefault, EmptyBackStoreCount[FMemory::AllocationHints::Default], ECsvCustomStatOp::Set);
+	CSV_CUSTOM_STAT(FMemory, LargeAllocatorPageCountSmall, GLargePageAllocatorLargePageCount[FMemory::AllocationHints::SmallPool], ECsvCustomStatOp::Set);
+	CSV_CUSTOM_STAT(FMemory, LargeAllocatorPageCountDefault, GLargePageAllocatorLargePageCount[FMemory::AllocationHints::Default], ECsvCustomStatOp::Set);
 
 	GLargePageAllocatorCommitCount = 0;
 	GLargePageAllocatorDecommitCount = 0;
diff --git a/Engine/Source/Runtime/Core/Public/HAL/Allocators/CachedOSVeryLargePageAllocator.h b/Engine/Source/Runtime/Core/Public/HAL/Allocators/CachedOSVeryLargePageAllocator.h
--- a/Engine/Source/Runtime/Core/Public/HAL/Allocators/CachedOSVeryLargePageAllocator.h
+++ b/Engine/Source/Runtime/Core/Public/HAL/Allocators/CachedOSVeryLargePageAllocator.h
@@ -96,6 +96,7 @@ private:
 	uintptr_t	AddressSpaceReservedEnd;
 	uint64		CachedFree;
 	int32		EmptyBackStoreCount[FMemory::AllocationHints::Max];
+	int32		MaxEmptyBackStoreCount[FMemory::AllocationHints::Max];
 
 	FPlatformMemory::FPlatformVirtualMemoryBlock Block;
 
