Fix Tpause feature on certain mainboards where WaitPKG is returning true even though some virtualization setting in the Bios can disable Tpause.

diff --git a/Engine/Source/Runtime/Core/Private/Windows/WindowsPlatformMisc.cpp b/Engine/Source/Runtime/Core/Private/Windows/WindowsPlatformMisc.cpp
--- a/Engine/Source/Runtime/Core/Private/Windows/WindowsPlatformMisc.cpp
+++ b/Engine/Source/Runtime/Core/Private/Windows/WindowsPlatformMisc.cpp
@@ -2154,9 +2154,7 @@ public:
 			CPUInfo = Info[0];
 			CPUInfo2 = Info[2];
 			CacheLineSize = QueryCacheLineSize();
-			int ExtendedFeatures[4];
-			QueryCPUExtendedFeatures(ExtendedFeatures);
-			CPUExtendedFeatures2 = ExtendedFeatures[2];
+			bHasTimedPauseInstruction = CheckForTimedPauseInstruction();
 		}
 	}
 
@@ -2170,6 +2168,16 @@ public:
 		return CPUIDStaticCache.bHasCPUIDInstruction;
 	}
 
+	/**
+	 * Checks if this CPU supports tpause instruction.
+	 *
+	 * @returns True if this CPU supports tpause instruction. False otherwise.
+	 */
+	static bool HasTimedPauseInstruction()
+	{
+		return CPUIDStaticCache.bHasTimedPauseInstruction;
+	}
+
 	/**
 	 * Gets pre-cached CPU vendor name.
 	 *
@@ -2210,16 +2218,6 @@ public:
 		return CPUIDStaticCache.CPUInfo2;
 	}
 
-	/**
-	* Gets __cpuidex CPU features.
-	*
-	* @returns CPU info unsigned int queried using __cpuidex.
-	*/
-	static uint32 GetCPUExtendedFeatures2()
-	{
-		return CPUIDStaticCache.CPUExtendedFeatures2;
-	}
-
 	/**
 	 * Gets cache line size.
 	 *
@@ -2259,6 +2257,58 @@ private:
 #endif
 	}
 
+	static int FilterInvalidOpcode(DWORD ExceptionCode, struct _EXCEPTION_POINTERS* ExceptionInformation)
+	{
+		if (ExceptionCode == STATUS_ILLEGAL_INSTRUCTION || ExceptionCode == STATUS_PRIVILEGED_INSTRUCTION)
+		{
+			return EXCEPTION_EXECUTE_HANDLER;
+		}
+		else
+		{
+			return EXCEPTION_CONTINUE_SEARCH;
+		}
+	}
+
+	/**
+	 * Checks if tpause instruction is present on current machine.
+	 *
+	 * @returns True if this CPU supports tpause instruction. False otherwise.
+	 */
+	static bool CheckForTimedPauseInstruction()
+	{
+#if PLATFORM_SEH_EXCEPTIONS_DISABLED
+		return false;
+#else
+		bool bSupportsTpause = false;
+		int CPUInfo[4];
+		__cpuid(CPUInfo, 0);
+
+		if (CPUInfo[0] >= 7)
+		{
+			int CPUExtendedInfo[4];
+			__cpuidex(CPUExtendedInfo, 7, 0);
+
+			if ((CPUExtendedInfo[2] & (1 << 5)) != 0)
+			{
+				// WAITPKG is supported
+				__try
+				{
+					unsigned long long tsc = __rdtsc();
+					_tpause(0, tsc + 1024);
+					// TPAUSE is supported
+					bSupportsTpause = true;
+				}
+				__except (FilterInvalidOpcode(GetExceptionCode(), GetExceptionInformation()))
+				{
+					bSupportsTpause = false;
+				}
+			}
+		}
+
+		return bSupportsTpause;
+#endif
+	}
+
 	/**
 	 * Queries Vendor name using __cpuid instruction.
 	 *
@@ -2327,16 +2377,6 @@ private:
 		__cpuid(Args, 1);
 	}
 
-	/**
-	 * Queries CPU info using __cpuid instruction.
-	 *
-	 * @returns CPU info unsigned int queried using __cpuidex.
-	 */
-	static void QueryCPUExtendedFeatures(int Args[4])
-	{
-		__cpuidex(Args, 7, 0);
-	}
-
 	/**
 	 * Queries cache line size using __cpuid instruction.
 	 *
@@ -2361,6 +2401,9 @@ private:
 	/** If machine has CPUID instruction. */
 	bool bHasCPUIDInstruction;
 
+	/** If machine has timed pause instruction. */
+	bool bHasTimedPauseInstruction;
+
 	/** Vendor of the CPU. */
 	ANSICHAR Vendor[12 + 1];
 
@@ -2370,7 +2413,6 @@ private:
 	/** CPU info from __cpuid. */
 	uint32 CPUInfo;
 	uint32 CPUInfo2;
-	uint32 CPUExtendedFeatures2;
 
 	/** CPU cache line size. */
 	int32 CacheLineSize;
@@ -3079,12 +3121,7 @@ bool FWindowsPlatformMisc::NeedsNonoptionalCPUFeaturesCheck()
 
 bool FWindowsPlatformMisc::HasTimedPauseCPUFeature()
 {
-	// Awaiting investigation from Intel on why _tpause instructions are failing on certain supported architectures
-	// e.g. a consistent crash on a 12700 i7 CPU
-	return false;
-
-	//// Check for waitpkg is bit 5
-	//return (FCPUIDQueriedData::GetCPUExtendedFeatures2() & (1 << 5)) != 0;
+	return FCPUIDQueriedData::HasTimedPauseInstruction();
 }
 
 int32 FWindowsPlatformMisc::GetCacheLineSize()
