If a console variable changes compilation module during live coding, the engine will no longer assert.

diff --git a/Engine/Source/Runtime/Core/Private/HAL/ConsoleManager.cpp b/Engine/Source/Runtime/Core/Private/HAL/ConsoleManager.cpp
--- a/Engine/Source/Runtime/Core/Private/HAL/ConsoleManager.cpp
+++ b/Engine/Source/Runtime/Core/Private/HAL/ConsoleManager.cpp
@@ -328,7 +328,15 @@ void OnCVarChange(T& Dst, const T& Src, EConsoleVariableFlags Flags, EConsoleVar
 {
 	FConsoleManager& ConsoleManager = (FConsoleManager&)IConsoleManager::Get();
 
+#if WITH_RELOAD
+	// Unlike HotReload, Live Coding does global initialization outside of the main thread.  During global initialization,
+	// Live Coding does have the main thread stalled so there is a "reduced" chance of threading issues.  During global
+	// initialization with live coding, this code only gets called when a new instance of an existing CVar is created.
+	// Normally this shouldn't happen unless a source file shuffles between two different unity files.
+	if (IsInGameThread() || IsReloadActive())
+#else
 	if(IsInGameThread())
+#endif
 	{
 		if((Flags & ECVF_RenderThreadSafe) && ConsoleManager.GetThreadPropagationCallback())
 		{
