Fix implicit conversion warnings in VirtualAllocPageStatus.cpp.
#rnx

diff --git a/Engine/Source/Runtime/Core/Private/Windows/VirtualAllocPageStatus.cpp b/Engine/Source/Runtime/Core/Private/Windows/VirtualAllocPageStatus.cpp
--- a/Engine/Source/Runtime/Core/Private/Windows/VirtualAllocPageStatus.cpp
+++ b/Engine/Source/Runtime/Core/Private/Windows/VirtualAllocPageStatus.cpp
@@ -143,10 +143,10 @@ uint64& FHashMapLinearProbingVAlloc::FindOrAdd(uint64 Key, uint64 ValueIfMissing
 {
 	LLMCheck(Key <= MaxKey);
 
-	if (NumActive > static_cast<int64>(NumBuckets * DesiredPopulation))
+	if (NumActive > static_cast<int64>(static_cast<float>(NumBuckets) * DesiredPopulation))
 	{
 		Realloc(NumBuckets * 2);
-		LLMCheck(NumActive < static_cast<int64>(NumBuckets * DesiredPopulation));
+		LLMCheck(NumActive < static_cast<int64>(static_cast<float>(NumBuckets) * DesiredPopulation));
 	}
 
 	int64 FirstBucketIndex = static_cast<int64>(Key & (NumBuckets - 1));
