Added TSharedRef::ToWeakPtr(), TSharedPtr::ToWeakPtr() and TSharedRef::ToSharedPtr().

diff --git a/Engine/Source/Runtime/Core/Public/Templates/SharedPointer.h b/Engine/Source/Runtime/Core/Public/Templates/SharedPointer.h
--- a/Engine/Source/Runtime/Core/Public/Templates/SharedPointer.h
+++ b/Engine/Source/Runtime/Core/Public/Templates/SharedPointer.h
@@ -464,6 +464,26 @@ public:
 		return *this;
 	}
 
+	/**
+	 * Converts a shared reference to a shared pointer.
+	 *
+	 * @return  Shared pointer to the object
+	 */
+	[[nodiscard]] FORCEINLINE TSharedPtr<ObjectType, Mode> ToSharedPtr() const
+	{
+		return TSharedPtr<ObjectType, Mode>(*this);
+	}
+
+	/**
+	 * Converts a shared reference to a weak ptr.
+	 *
+	 * @return  Weak pointer to the object
+	 */
+	[[nodiscard]] FORCEINLINE TWeakPtr<ObjectType, Mode> ToWeakPtr() const
+	{
+		return TWeakPtr<ObjectType, Mode>(*this);
+	}
+
 	/**
 	 * Returns a C++ reference to the object this shared reference is referencing
 	 *
@@ -1058,6 +1078,16 @@ public:
 		return TSharedRef< ObjectType, Mode >( MoveTemp( *this ) );
 	}
 
+	/**
+	 * Converts a shared pointer to a weak ptr.
+	 *
+	 * @return  Weak pointer to the object
+	 */
+	[[nodiscard]] FORCEINLINE TWeakPtr<ObjectType, Mode> ToWeakPtr() const
+	{
+		return TWeakPtr<ObjectType, Mode>(*this);
+	}
+
 	/**
 	 * Returns the object referenced by this pointer, or nullptr if no object is reference
 	 *
