LargePageAllocator cache size bucket is now driven via cvar

#rnx

diff --git a/Engine/Source/Runtime/Core/Private/HAL/Allocators/CachedOSVeryLargePageAllocator.cpp b/Engine/Source/Runtime/Core/Private/HAL/Allocators/CachedOSVeryLargePageAllocator.cpp
--- a/Engine/Source/Runtime/Core/Private/HAL/Allocators/CachedOSVeryLargePageAllocator.cpp
+++ b/Engine/Source/Runtime/Core/Private/HAL/Allocators/CachedOSVeryLargePageAllocator.cpp
@@ -27,6 +27,23 @@
 #endif
 
 CORE_API bool GEnableVeryLargePageAllocator = true;
+static int32 GVeryLargePageAllocatorMaxEmptyBackStoreCount[FMemory::AllocationHints::Max] = {
+	UE_VERYLARGEPAGEALLOCATOR_MAX_EMPTY_BACKSTORE_DEFAULT,	// FMemory::AllocationHints::Default
+	0,														// FMemory::AllocationHints::Temporary
+	UE_VERYLARGEPAGEALLOCATOR_MAX_EMPTY_BACKSTORE_SMALLPOOL // FMemory::AllocationHints::SmallPool
+};
+
+static_assert(int32(FMemory::AllocationHints::Max) == 3); // ensure FMemory::AllocationHints has 3 types of hint so GVeryLargePageAllocatorMaxEmptyBackStoreCount has needed values
+
+static FAutoConsoleVariableRef CVarVeryLargePageAllocatorMaxEmptyBackstoreDefault(
+	TEXT("VeryLargePageAllocator.MaxEmptyBackstoreDefault"),
+	GVeryLargePageAllocatorMaxEmptyBackStoreCount[FMemory::AllocationHints::Default],
+	TEXT(""));
+
+static FAutoConsoleVariableRef CVarVeryLargePageAllocatorMaxEmptyBackstoreSmallPool(
+	TEXT("VeryLargePageAllocator.MaxEmptyBackstoreSmallPool"),
+	GVeryLargePageAllocatorMaxEmptyBackStoreCount[FMemory::AllocationHints::SmallPool],
+	TEXT(""));
 
 #if CSV_PROFILER
 CSV_DECLARE_CATEGORY_MODULE_EXTERN(CORE_API, FMemory);
@@ -55,14 +72,10 @@ void FCachedOSVeryLargePageAllocator::Init()
 		UsedLargePagesHead[i] = nullptr;
 		EmptyButAvailableLargePagesHead[i] = nullptr;
 		EmptyBackStoreCount[i] = 0;
-		MaxEmptyBackStoreCount[i] = 0;
 		
 		GLargePageAllocatorLargePageCount[i] = 0;
 	}
 
-	MaxEmptyBackStoreCount[FMemory::AllocationHints::SmallPool] = UE_VERYLARGEPAGEALLOCATOR_MAX_EMPTY_BACKSTORE_SMALLPOOL;
-	MaxEmptyBackStoreCount[FMemory::AllocationHints::Default] = UE_VERYLARGEPAGEALLOCATOR_MAX_EMPTY_BACKSTORE_DEFAULT;
-
 #if UE_VERYLARGEPAGEALLOCATOR_TAKEONALL64KBALLOCATIONS
 	for (int i = 0; i < NumberOfLargePages / 2; i++)
 #else
@@ -86,7 +99,7 @@ void FCachedOSVeryLargePageAllocator::Init()
 	{
 		FMemory::AllocationHints AllocationHint = FMemory::AllocationHints(i);
 		
-		for (int32 BackStoreIndex = 0; BackStoreIndex < MaxEmptyBackStoreCount[AllocationHint]; BackStoreIndex++)
+		for (int32 BackStoreIndex = 0; BackStoreIndex < GVeryLargePageAllocatorMaxEmptyBackStoreCount[AllocationHint]; BackStoreIndex++)
 		{
 			FLargePage* LargePage = FreeLargePagesHead[AllocationHint];
 
@@ -96,6 +109,7 @@ void FCachedOSVeryLargePageAllocator::Init()
 			Block.Commit(LargePage->BaseAddress - AddressSpaceReserved, SizeOfLargePage, false);
 			LargePage->LinkHead(EmptyButAvailableLargePagesHead[AllocationHint]);
 			EmptyBackStoreCount[AllocationHint] += 1;
+			CachedFree += SizeOfLargePage;
 		}
 	}
 #endif
@@ -166,6 +180,8 @@ void* FCachedOSVeryLargePageAllocator::Allocate(SIZE_T Size, uint32 AllocationHi
 #if CSV_PROFILER
 							else
 							{
+								// A new large page has been created. Add it to CachedFree counter
+								CachedFree += SizeOfLargePage;
 								FPlatformAtomics::InterlockedIncrement(&GLargePageAllocatorCommitCount);
 								FPlatformAtomics::InterlockedIncrement(&GLargePageAllocatorLargePageCount[AllocationHint]);
 							}
@@ -184,7 +200,6 @@ void* FCachedOSVeryLargePageAllocator::Allocate(SIZE_T Size, uint32 AllocationHi
 				if (bLinkToUsedLargePagesWithSpaceHead)
 				{
 					LargePage->LinkHead(UsedLargePagesWithSpaceHead[AllocationHint]);
-					CachedFree += SizeOfLargePage;
 				}
 
 				ret = LargePage->Allocate();
@@ -234,7 +249,7 @@ void FCachedOSVeryLargePageAllocator::Free(void* Ptr, SIZE_T Size, FCriticalSect
 			LargePage->Unlink();
 
 			// move it to EmptyButAvailableLargePagesHead if that pool of backstore is not full yet
-			if (EmptyBackStoreCount[LargePage->AllocationHint] < MaxEmptyBackStoreCount[LargePage->AllocationHint])
+			if (EmptyBackStoreCount[LargePage->AllocationHint] < GVeryLargePageAllocatorMaxEmptyBackStoreCount[LargePage->AllocationHint])
 			{
 				LargePage->LinkHead(EmptyButAvailableLargePagesHead[LargePage->AllocationHint]);
 				EmptyBackStoreCount[LargePage->AllocationHint] += 1;
diff --git a/Engine/Source/Runtime/Core/Public/HAL/Allocators/CachedOSVeryLargePageAllocator.h b/Engine/Source/Runtime/Core/Public/HAL/Allocators/CachedOSVeryLargePageAllocator.h
--- a/Engine/Source/Runtime/Core/Public/HAL/Allocators/CachedOSVeryLargePageAllocator.h
+++ b/Engine/Source/Runtime/Core/Public/HAL/Allocators/CachedOSVeryLargePageAllocator.h
@@ -96,7 +96,6 @@ private:
 	uintptr_t	AddressSpaceReservedEnd;
 	uint64		CachedFree;
 	int32		EmptyBackStoreCount[FMemory::AllocationHints::Max];
-	int32		MaxEmptyBackStoreCount[FMemory::AllocationHints::Max];
 
 	FPlatformMemory::FPlatformVirtualMemoryBlock Block;
 
