Take cached slack memory into account in Memory Free and Memory Used csv stats if  PLATFORM_USE_CACHED_SLACK_MEMORY_IN_MEMORY_STATS is set

#rnx

diff --git a/Engine/Source/Runtime/Core/Private/ProfilingDebugging/CsvProfiler.cpp b/Engine/Source/Runtime/Core/Private/ProfilingDebugging/CsvProfiler.cpp
--- a/Engine/Source/Runtime/Core/Private/ProfilingDebugging/CsvProfiler.cpp
+++ b/Engine/Source/Runtime/Core/Private/ProfilingDebugging/CsvProfiler.cpp
@@ -15,6 +15,7 @@
 #include "HAL/PlatformProcess.h"
 #include "HAL/PlatformTime.h"
 #include "HAL/IConsoleManager.h"
+#include "HAL/MemoryMisc.h"
 #include "Misc/Paths.h"
 #include "Misc/CommandLine.h"
 #include "Misc/ScopeLock.h"
@@ -2968,6 +2969,27 @@ void FCsvProfiler::EndFrame()
 		CSV_CUSTOM_STAT_DEFINED(FrameTime, ElapsedMs, ECsvCustomStatOp::Set);
 
 		FPlatformMemoryStats MemoryStats = FPlatformMemory::GetStats();
+
+#if PLATFORM_USE_CACHED_SLACK_MEMORY_IN_MEMORY_STATS
+		// Take into account cached slacked memory to reflect what is really available
+		if (GMalloc != nullptr)
+		{
+			uint64 CachedFreeSize = 0;
+
+			FGenericMemoryStats MallocStats;
+			GMalloc->GetAllocatorStats(MallocStats);
+
+			SIZE_T* PageAllocatorFreeCacheSize = MallocStats.Data.Find("PageAllocatorFreeCacheSize");
+			if (PageAllocatorFreeCacheSize != nullptr)
+			{
+				CachedFreeSize = *PageAllocatorFreeCacheSize;
+			}
+
+			MemoryStats.UsedPhysical -= CachedFreeSize;
+			MemoryStats.AvailablePhysical += CachedFreeSize;
+		}
+#endif
+
 		float PhysicalMBFree = float(MemoryStats.AvailablePhysical) / (1024.0f * 1024.0f);
 		float UsedExtendedMB = 0;
 		float PhysicalMBUsed = float(MemoryStats.UsedPhysical) / (1024.0f * 1024.0f);
diff --git a/Engine/Source/Runtime/Core/Public/HAL/Platform.h b/Engine/Source/Runtime/Core/Public/HAL/Platform.h
--- a/Engine/Source/Runtime/Core/Public/HAL/Platform.h
+++ b/Engine/Source/Runtime/Core/Public/HAL/Platform.h
@@ -538,6 +538,10 @@
 	#define	PLATFORM_USE_SHOWFLAGS_ALWAYS_BITFIELD				1
 #endif
 
+#ifndef PLATFORM_USE_CACHED_SLACK_MEMORY_IN_MEMORY_STATS
+	#define	PLATFORM_USE_CACHED_SLACK_MEMORY_IN_MEMORY_STATS	0
+#endif
+
 #ifndef PLATFORM_SUPPORTS_LANDSCAPE_VISUAL_MESH_LOD_STREAMING
 	#define PLATFORM_SUPPORTS_LANDSCAPE_VISUAL_MESH_LOD_STREAMING 0 UE_DEPRECATED_MACRO(5.1, "PLATFORM_SUPPORTS_LANDSCAPE_VISUAL_MESH_LOD_STREAMING has been deprecated and should be replaced with 0.")
 #endif
