[Runtime]
* TaskGraphFwd - Changed typedef to using to try to get around some custom tool fwd header rules

diff --git a/Engine/Source/Runtime/Core/Public/Async/TaskGraphFwd.h b/Engine/Source/Runtime/Core/Public/Async/TaskGraphFwd.h
--- a/Engine/Source/Runtime/Core/Public/Async/TaskGraphFwd.h
+++ b/Engine/Source/Runtime/Core/Public/Async/TaskGraphFwd.h
@@ -22,6 +22,6 @@ using FGraphEventRef = TRefCountPtr<FBaseGraphTask>;
 #else
 
 /** Convenience typedef for a reference counted pointer to a graph event **/
-typedef TRefCountPtr<class FGraphEvent> FGraphEventRef;
+using FGraphEventRef = TRefCountPtr<class FGraphEvent>;
 
 #endif
