[Core]
* CachedOSVeryLargePageAllocator - Attempt to fix compile error

diff --git a/Engine/Source/Runtime/Core/Private/HAL/Allocators/CachedOSVeryLargePageAllocator.cpp b/Engine/Source/Runtime/Core/Private/HAL/Allocators/CachedOSVeryLargePageAllocator.cpp
--- a/Engine/Source/Runtime/Core/Private/HAL/Allocators/CachedOSVeryLargePageAllocator.cpp
+++ b/Engine/Source/Runtime/Core/Private/HAL/Allocators/CachedOSVeryLargePageAllocator.cpp
@@ -73,7 +73,9 @@ void FCachedOSVeryLargePageAllocator::Init()
 		EmptyButAvailableLargePagesHead[i] = nullptr;
 		EmptyBackStoreCount[i] = 0;
 		
+#if CSV_PROFILER
 		GLargePageAllocatorLargePageCount[i] = 0;
+#endif
 	}
 
 #if UE_VERYLARGEPAGEALLOCATOR_TAKEONALL64KBALLOCATIONS
