Removed unnecessary delegate instance Create functions.
Forward payload parameters through the delegate API rather than copying.

diff --git a/Engine/Source/Runtime/Core/Public/Delegates/DelegateInstanceInterface.h b/Engine/Source/Runtime/Core/Public/Delegates/DelegateInstanceInterface.h
--- a/Engine/Source/Runtime/Core/Public/Delegates/DelegateInstanceInterface.h
+++ b/Engine/Source/Runtime/Core/Public/Delegates/DelegateInstanceInterface.h
@@ -18,7 +18,7 @@ struct IBaseDelegateInstance<RetType(ArgTypes...), UserPolicy> : public UserPoli
 	/**
 	 * Emplaces a copy of the delegate instance into the given base delegate.
 	 */
-	virtual void CreateCopy(typename UserPolicy::FDelegateExtras& Base) = 0;
+	virtual void CreateCopy(typename UserPolicy::FDelegateExtras& Base) const = 0;
 
 	/**
 	 * Execute the delegate.  If the function pointer is not valid, an error will occur.
diff --git a/Engine/Source/Runtime/Core/Public/Delegates/DelegateInstancesImpl.h b/Engine/Source/Runtime/Core/Public/Delegates/DelegateInstancesImpl.h
--- a/Engine/Source/Runtime/Core/Public/Delegates/DelegateInstancesImpl.h
+++ b/Engine/Source/Runtime/Core/Public/Delegates/DelegateInstancesImpl.h
@@ -33,8 +33,9 @@ template <typename FuncType, typename UserPolicy, typename... VarTypes>
 class TCommonDelegateInstanceState : IBaseDelegateInstance<FuncType, UserPolicy>
 {
 public:
-	explicit TCommonDelegateInstanceState(VarTypes... Vars)
-		: Payload(Vars...)
+	template <typename... InVarTypes>
+	explicit TCommonDelegateInstanceState(InVarTypes&&... Vars)
+		: Payload(Forward<InVarTypes>(Vars)...)
 		, Handle (FDelegateHandle::GenerateNewHandle)
 	{
 	}
@@ -70,8 +71,9 @@ private:
 	static_assert(UE::Delegates::Private::IsUObjectPtr((UserClass*)nullptr), "You cannot use UFunction delegates with non UObject classes.");
 
 public:
-	TBaseUFunctionDelegateInstance(UserClass* InUserObject, const FName& InFunctionName, VarTypes... Vars)
-		: Super        (Vars...)
+	template <typename... InVarTypes>
+	explicit TBaseUFunctionDelegateInstance(UserClass* InUserObject, const FName& InFunctionName, InVarTypes&&... Vars)
+		: Super        (Forward<InVarTypes>(Vars)...)
 		, FunctionName (InFunctionName)
 		, UserObjectPtr(InUserObject)
 	{
@@ -129,7 +131,7 @@ public:
 
 	// IBaseDelegateInstance interface
 
-	void CreateCopy(DelegateBaseType& Base) final
+	void CreateCopy(DelegateBaseType& Base) const final
 	{
 		new (Base) TBaseUFunctionDelegateInstance(*this);
 	}
@@ -161,20 +163,6 @@ public:
 		return false;
 	}
 
-public:
-
-	/**
-	 * Creates a new UFunction delegate binding for the given user object and function name.
-	 *
-	 * @param InObject The user object to call the function on.
-	 * @param InFunctionName The name of the function call.
-	 * @return The new delegate.
-	 */
-	FORCEINLINE static void Create(DelegateBaseType& Base, UserClass* InUserObject, const FName& InFunctionName, VarTypes... Vars)
-	{
-		new (Base) TBaseUFunctionDelegateInstance(InUserObject, InFunctionName, Vars...);
-	}
-
 public:
 
 	// Holds the cached UFunction to call.
@@ -207,8 +195,9 @@ private:
 public:
 	using FMethodPtr = typename TMemFunPtrType<bConst, UserClass, RetValType(ParamTypes..., VarTypes...)>::Type;
 
-	TBaseSPMethodDelegateInstance(const TSharedPtr<UserClass, SPMode>& InUserObject, FMethodPtr InMethodPtr, VarTypes... Vars)
-		: Super     (Vars...)
+	template <typename... InVarTypes>
+	explicit TBaseSPMethodDelegateInstance(const TSharedPtr<UserClass, SPMode>& InUserObject, FMethodPtr InMethodPtr, InVarTypes&&... Vars)
+		: Super     (Forward<InVarTypes>(Vars)...)
 		, UserObject(InUserObject)
 		, MethodPtr (InMethodPtr)
 	{
@@ -262,7 +251,7 @@ public:
 
 	// IBaseDelegateInstance interface
 
-	void CreateCopy(DelegateBaseType& Base) final
+	void CreateCopy(DelegateBaseType& Base) const final
 	{
 		new (Base) TBaseSPMethodDelegateInstance(*this);
 	}
@@ -307,36 +296,6 @@ public:
 		return false;
 	}
 
-public:
-
-	/**
-	 * Creates a new shared pointer delegate binding for the given user object and method pointer.
-	 *
-	 * @param InUserObjectRef Shared reference to the user's object that contains the class method.
-	 * @param InFunc Member function pointer to your class method.
-	 * @return The new delegate.
-	 */
-	FORCEINLINE static void Create(DelegateBaseType& Base, const TSharedPtr<UserClass, SPMode>& InUserObjectRef, FMethodPtr InFunc, VarTypes... Vars)
-	{
-		new (Base) TBaseSPMethodDelegateInstance(InUserObjectRef, InFunc, Vars...);
-	}
-
-	/**
-	 * Creates a new shared pointer delegate binding for the given user object and method pointer.
-	 *
-	 * This overload requires that the supplied object derives from TSharedFromThis.
-	 *
-	 * @param InUserObject  The user's object that contains the class method.  Must derive from TSharedFromThis.
-	 * @param InFunc  Member function pointer to your class method.
-	 * @return The new delegate.
-	 */
-	FORCEINLINE static void Create(DelegateBaseType& Base, UserClass* InUserObject, FMethodPtr InFunc, VarTypes... Vars)
-	{
-		// We expect the incoming InUserObject to derived from TSharedFromThis.
-		TSharedRef<UserClass, SPMode> UserObjectRef = StaticCastSharedRef<UserClass>(InUserObject->AsShared());
-		Create(Base, UserObjectRef, InFunc, Vars...);
-	}
-
 protected:
 
 	// Weak reference to an instance of the user's class which contains a method we would like to call.
@@ -371,8 +330,9 @@ public:
 	 * @param InUserObject An arbitrary object (templated) that hosts the member function.
 	 * @param InMethodPtr C++ member function pointer for the method to bind.
 	 */
-	TBaseRawMethodDelegateInstance(UserClass* InUserObject, FMethodPtr InMethodPtr, VarTypes... Vars)
-		: Super     (Vars...)
+	template <typename... InVarTypes>
+	explicit TBaseRawMethodDelegateInstance(UserClass* InUserObject, FMethodPtr InMethodPtr, InVarTypes&&... Vars)
+		: Super     (Forward<InVarTypes>(Vars)...)
 		, UserObject(InUserObject)
 		, MethodPtr (InMethodPtr)
 	{
@@ -427,7 +387,7 @@ public:
 
 	// IBaseDelegateInstance interface
 
-	void CreateCopy(DelegateBaseType& Base) final
+	void CreateCopy(DelegateBaseType& Base) const final
 	{
 		new (Base) TBaseRawMethodDelegateInstance(*this);
 	}
@@ -463,20 +423,6 @@ public:
 		return true;
 	}
 
-public:
-
-	/**
-	 * Creates a new raw method delegate binding for the given user object and function pointer.
-	 *
-	 * @param InUserObject User's object that contains the class method.
-	 * @param InFunc Member function pointer to your class method.
-	 * @return The new delegate.
-	 */
-	FORCEINLINE static void Create(DelegateBaseType& Base, UserClass* InUserObject, FMethodPtr InFunc, VarTypes... Vars)
-	{
-		new (Base) TBaseRawMethodDelegateInstance(InUserObject, InFunc, Vars...);
-	}
-
 protected:
 
 	// Pointer to the user's class which contains a method we would like to call.
@@ -504,8 +450,9 @@ private:
 public:
 	using FMethodPtr = typename TMemFunPtrType<bConst, UserClass, RetValType(ParamTypes..., VarTypes...)>::Type;
 
-	TBaseUObjectMethodDelegateInstance(UserClass* InUserObject, FMethodPtr InMethodPtr, VarTypes... Vars)
-		: Super     (Vars...)
+	template <typename... InVarTypes>
+	explicit TBaseUObjectMethodDelegateInstance(UserClass* InUserObject, FMethodPtr InMethodPtr, InVarTypes&&... Vars)
+		: Super     (Forward<InVarTypes>(Vars)...)
 		, UserObject(InUserObject)
 		, MethodPtr (InMethodPtr)
 	{
@@ -564,7 +511,7 @@ public:
 
 	// IBaseDelegateInstance interface
 
-	void CreateCopy(DelegateBaseType& Base) final
+	void CreateCopy(DelegateBaseType& Base) const final
 	{
 		new (Base) TBaseUObjectMethodDelegateInstance(*this);
 	}
@@ -606,20 +553,6 @@ public:
 		return false;
 	}
 
-public:
-
-	/**
-	 * Creates a new UObject delegate binding for the given user object and method pointer.
-	 *
-	 * @param InUserObject User's object that contains the class method.
-	 * @param InFunc Member function pointer to your class method.
-	 * @return The new delegate.
-	 */
-	FORCEINLINE static void Create(DelegateBaseType& Base, UserClass* InUserObject, FMethodPtr InFunc, VarTypes... Vars)
-	{
-		new (Base) TBaseUObjectMethodDelegateInstance(InUserObject, InFunc, Vars...);
-	}
-
 protected:
 
 	// Pointer to the user's class which contains a method we would like to call.
@@ -646,8 +579,9 @@ private:
 public:
 	using FFuncPtr = RetValType(*)(ParamTypes..., VarTypes...);
 
-	TBaseStaticDelegateInstance(FFuncPtr InStaticFuncPtr, VarTypes... Vars)
-		: Super        (Vars...)
+	template <typename... InVarTypes>
+	explicit TBaseStaticDelegateInstance(FFuncPtr InStaticFuncPtr, InVarTypes&&... Vars)
+		: Super        (Forward<InVarTypes>(Vars)...)
 		, StaticFuncPtr(InStaticFuncPtr)
 	{
 		check(StaticFuncPtr != nullptr);
@@ -700,7 +634,7 @@ public:
 
 	// IBaseDelegateInstance interface
 
-	void CreateCopy(DelegateBaseType& Base) final
+	void CreateCopy(DelegateBaseType& Base) const final
 	{
 		new (Base) TBaseStaticDelegateInstance(*this);
 	}
@@ -723,19 +657,6 @@ public:
 		return true;
 	}
 
-public:
-
-	/**
-	 * Creates a new static function delegate binding for the given function pointer.
-	 *
-	 * @param InFunc Static function pointer.
-	 * @return The new delegate.
-	 */
-	FORCEINLINE static void Create(DelegateBaseType& Base, FFuncPtr InFunc, VarTypes... Vars)
-	{
-		new (Base) TBaseStaticDelegateInstance(InFunc, Vars...);
-	}
-
 private:
 
 	// C++ function pointer.
@@ -758,15 +679,10 @@ private:
 	using DelegateBaseType = typename UserPolicy::FDelegateExtras;
 
 public:
-	TBaseFunctorDelegateInstance(const FunctorType& InFunctor, VarTypes... Vars)
-		: Super  (Vars...)
-		, Functor(InFunctor)
-	{
-	}
-
-	TBaseFunctorDelegateInstance(FunctorType&& InFunctor, VarTypes... Vars)
-		: Super  (Vars...)
-		, Functor(MoveTemp(InFunctor))
+	template <typename InFunctorType, typename... InVarTypes>
+	explicit TBaseFunctorDelegateInstance(InFunctorType&& InFunctor, InVarTypes&&... Vars)
+		: Super  (Forward<InVarTypes>(Vars)...)
+		, Functor(Forward<InFunctorType>(InFunctor))
 	{
 	}
 
@@ -811,7 +727,7 @@ public:
 
 public:
 	// IBaseDelegateInstance interface
-	void CreateCopy(DelegateBaseType& Base) final
+	void CreateCopy(DelegateBaseType& Base) const final
 	{
 		new (Base) TBaseFunctorDelegateInstance(*this);
 	}
@@ -829,22 +745,6 @@ public:
 		return true;
 	}
 
-public:
-	/**
-	 * Creates a new static function delegate binding for the given function pointer.
-	 *
-	 * @param InFunctor C++ functor
-	 * @return The new delegate.
-	 */
-	FORCEINLINE static void Create(DelegateBaseType& Base, const FunctorType& InFunctor, VarTypes... Vars)
-	{
-		new (Base) TBaseFunctorDelegateInstance(InFunctor, Vars...);
-	}
-	FORCEINLINE static void Create(DelegateBaseType& Base, FunctorType&& InFunctor, VarTypes... Vars)
-	{
-		new (Base) TBaseFunctorDelegateInstance(MoveTemp(InFunctor), Vars...);
-	}
-
 private:
 	// C++ functor
 	// We make this mutable to allow mutable lambdas to be bound and executed.  We don't really want to
@@ -869,17 +769,11 @@ private:
 	using DelegateBaseType = typename UserPolicy::FDelegateExtras;
 
 public:
-	TWeakBaseFunctorDelegateInstance(UserClass* InContextObject, const FunctorType& InFunctor, VarTypes... Vars)
-		: Super        (Vars...)
-		, ContextObject(InContextObject)
-		, Functor      (InFunctor)
-	{
-	}
-
-	TWeakBaseFunctorDelegateInstance(UserClass* InContextObject, FunctorType&& InFunctor, VarTypes... Vars)
-		: Super        (Vars...)
+	template <typename InFunctorType, typename... InVarTypes>
+	explicit TWeakBaseFunctorDelegateInstance(UserClass* InContextObject, InFunctorType&& InFunctor, InVarTypes&&... Vars)
+		: Super        (Forward<InVarTypes>(Vars)...)
 		, ContextObject(InContextObject)
-		, Functor      (MoveTemp(InFunctor))
+		, Functor      (Forward< InFunctorType>(InFunctor))
 	{
 	}
 
@@ -927,7 +821,7 @@ public:
 
 public:
 	// IBaseDelegateInstance interface
-	void CreateCopy(DelegateBaseType& Base) final
+	void CreateCopy(DelegateBaseType& Base) const final
 	{
 		new (Base) TWeakBaseFunctorDelegateInstance(*this);
 	}
@@ -948,22 +842,6 @@ public:
 		return false;
 	}
 
-public:
-	/**
-	 * Creates a new static function delegate binding for the given function pointer.
-	 *
-	 * @param InFunctor C++ functor
-	 * @return The new delegate.
-	 */
-	FORCEINLINE static void Create(DelegateBaseType& Base, UserClass* InContextObject, const FunctorType& InFunctor, VarTypes... Vars)
-	{
-		new (Base) TWeakBaseFunctorDelegateInstance(InContextObject, InFunctor, Vars...);
-	}
-	FORCEINLINE static void Create(DelegateBaseType& Base, UserClass* InContextObject, FunctorType&& InFunctor, VarTypes... Vars)
-	{
-		new (Base) TWeakBaseFunctorDelegateInstance(InContextObject, MoveTemp(InFunctor), Vars...);
-	}
-
 private:
 	// Context object - the validity of this object controls the validity of the lambda
 	TWeakObjectPtr<UserClass> ContextObject;
diff --git a/Engine/Source/Runtime/Core/Public/Delegates/DelegateSignatureImpl.inl b/Engine/Source/Runtime/Core/Public/Delegates/DelegateSignatureImpl.inl
--- a/Engine/Source/Runtime/Core/Public/Delegates/DelegateSignatureImpl.inl
+++ b/Engine/Source/Runtime/Core/Public/Delegates/DelegateSignatureImpl.inl
@@ -150,10 +150,10 @@ public:
 	 * Static: Creates a raw C++ pointer global function delegate
 	 */
 	template <typename... VarTypes>
-	UE_NODISCARD inline static TDelegate<RetValType(ParamTypes...), UserPolicy> CreateStatic(typename TIdentity<RetValType (*)(ParamTypes..., VarTypes...)>::Type InFunc, VarTypes... Vars)
+	UE_NODISCARD inline static TDelegate<RetValType(ParamTypes...), UserPolicy> CreateStatic(typename TIdentity<RetValType (*)(ParamTypes..., std::decay_t<VarTypes>...)>::Type InFunc, VarTypes&&... Vars)
 	{
 		TDelegate<RetValType(ParamTypes...), UserPolicy> Result;
-		TBaseStaticDelegateInstance<FuncType, UserPolicy, VarTypes...>::Create(Result, InFunc, Vars...);
+		new (Result) TBaseStaticDelegateInstance<FuncType, UserPolicy, std::decay_t<VarTypes>...>(InFunc, Forward<VarTypes>(Vars)...);
 		return Result;
 	}
 
@@ -162,10 +162,10 @@ public:
 	 * technically this works for any functor types, but lambdas are the primary use case
 	 */
 	template<typename FunctorType, typename... VarTypes>
-	UE_NODISCARD inline static TDelegate<RetValType(ParamTypes...), UserPolicy> CreateLambda(FunctorType&& InFunctor, VarTypes... Vars)
+	UE_NODISCARD inline static TDelegate<RetValType(ParamTypes...), UserPolicy> CreateLambda(FunctorType&& InFunctor, VarTypes&&... Vars)
 	{
 		TDelegate<RetValType(ParamTypes...), UserPolicy> Result;
-		TBaseFunctorDelegateInstance<FuncType, UserPolicy, typename TRemoveReference<FunctorType>::Type, VarTypes...>::Create(Result, Forward<FunctorType>(InFunctor), Vars...);
+		new (Result) TBaseFunctorDelegateInstance<FuncType, UserPolicy, typename TRemoveReference<FunctorType>::Type, std::decay_t<VarTypes>...>(Forward<FunctorType>(InFunctor), Forward<VarTypes>(Vars)...);
 		return Result;
 	}
 
@@ -174,10 +174,10 @@ public:
 	 * technically this works for any functor types, but lambdas are the primary use case
 	 */
 	template<typename UserClass, typename FunctorType, typename... VarTypes>
-	UE_NODISCARD inline static TDelegate<RetValType(ParamTypes...), UserPolicy> CreateWeakLambda(UserClass* InUserObject, FunctorType&& InFunctor, VarTypes... Vars)
+	UE_NODISCARD inline static TDelegate<RetValType(ParamTypes...), UserPolicy> CreateWeakLambda(UserClass* InUserObject, FunctorType&& InFunctor, VarTypes&&... Vars)
 	{
 		TDelegate<RetValType(ParamTypes...), UserPolicy> Result;
-		TWeakBaseFunctorDelegateInstance<UserClass, FuncType, UserPolicy, typename TRemoveReference<FunctorType>::Type, VarTypes...>::Create(Result, InUserObject, Forward<FunctorType>(InFunctor), Vars...);
+		new (Result) TWeakBaseFunctorDelegateInstance<UserClass, FuncType, UserPolicy, typename TRemoveReference<FunctorType>::Type, std::decay_t<VarTypes>...>(InUserObject, Forward<FunctorType>(InFunctor), Forward<VarTypes>(Vars)...);
 		return Result;
 	}
 
@@ -188,19 +188,19 @@ public:
 	 * deleted out from underneath your delegate. Be careful when calling Execute()!
 	 */
 	template <typename UserClass, typename... VarTypes>
-	UE_NODISCARD inline static TDelegate<RetValType(ParamTypes...), UserPolicy> CreateRaw(UserClass* InUserObject, typename TMemFunPtrType<false, UserClass, RetValType (ParamTypes..., VarTypes...)>::Type InFunc, VarTypes... Vars)
+	UE_NODISCARD inline static TDelegate<RetValType(ParamTypes...), UserPolicy> CreateRaw(UserClass* InUserObject, typename TMemFunPtrType<false, UserClass, RetValType (ParamTypes..., std::decay_t<VarTypes>...)>::Type InFunc, VarTypes&&... Vars)
 	{
 		static_assert(!TIsConst<UserClass>::Value, "Attempting to bind a delegate with a const object pointer and non-const member function.");
 
 		TDelegate<RetValType(ParamTypes...), UserPolicy> Result;
-		TBaseRawMethodDelegateInstance<false, UserClass, FuncType, UserPolicy, VarTypes...>::Create(Result, InUserObject, InFunc, Vars...);
+		new (Result) TBaseRawMethodDelegateInstance<false, UserClass, FuncType, UserPolicy, std::decay_t<VarTypes>...>(InUserObject, InFunc, Forward<VarTypes>(Vars)...);
 		return Result;
 	}
 	template <typename UserClass, typename... VarTypes>
-	UE_NODISCARD inline static TDelegate<RetValType(ParamTypes...), UserPolicy> CreateRaw(const UserClass* InUserObject, typename TMemFunPtrType<true, UserClass, RetValType (ParamTypes..., VarTypes...)>::Type InFunc, VarTypes... Vars)
+	UE_NODISCARD inline static TDelegate<RetValType(ParamTypes...), UserPolicy> CreateRaw(const UserClass* InUserObject, typename TMemFunPtrType<true, UserClass, RetValType (ParamTypes..., std::decay_t<VarTypes>...)>::Type InFunc, VarTypes&&... Vars)
 	{
 		TDelegate<RetValType(ParamTypes...), UserPolicy> Result;
-		TBaseRawMethodDelegateInstance<true, const UserClass, FuncType, UserPolicy, VarTypes...>::Create(Result, InUserObject, InFunc, Vars...);
+		new (Result) TBaseRawMethodDelegateInstance<true, const UserClass, FuncType, UserPolicy, std::decay_t<VarTypes>...>(InUserObject, InFunc, Forward<VarTypes>(Vars)...);
 		return Result;
 	}
 
@@ -211,19 +211,19 @@ public:
 	 * You can use ExecuteIfBound() to call them.
 	 */
 	template <typename UserClass, ESPMode Mode, typename... VarTypes>
-	UE_NODISCARD inline static TDelegate<RetValType(ParamTypes...), UserPolicy> CreateSP(const TSharedRef<UserClass, Mode>& InUserObjectRef, typename TMemFunPtrType<false, UserClass, RetValType (ParamTypes..., VarTypes...)>::Type InFunc, VarTypes... Vars)
+	UE_NODISCARD inline static TDelegate<RetValType(ParamTypes...), UserPolicy> CreateSP(const TSharedRef<UserClass, Mode>& InUserObjectRef, typename TMemFunPtrType<false, UserClass, RetValType (ParamTypes..., std::decay_t<VarTypes>...)>::Type InFunc, VarTypes&&... Vars)
 	{
 		static_assert(!TIsConst<UserClass>::Value, "Attempting to bind a delegate with a const object pointer and non-const member function.");
 
 		TDelegate<RetValType(ParamTypes...), UserPolicy> Result;
-		TBaseSPMethodDelegateInstance<false, UserClass, Mode, FuncType, UserPolicy, VarTypes...>::Create(Result, InUserObjectRef, InFunc, Vars...);
+		new (Result) TBaseSPMethodDelegateInstance<false, UserClass, Mode, FuncType, UserPolicy, std::decay_t<VarTypes>...>(InUserObjectRef, InFunc, Forward<VarTypes>(Vars)...);
 		return Result;
 	}
 	template <typename UserClass, ESPMode Mode, typename... VarTypes>
-	UE_NODISCARD inline static TDelegate<RetValType(ParamTypes...), UserPolicy> CreateSP(const TSharedRef<UserClass, Mode>& InUserObjectRef, typename TMemFunPtrType<true, UserClass, RetValType (ParamTypes..., VarTypes...)>::Type InFunc, VarTypes... Vars)
+	UE_NODISCARD inline static TDelegate<RetValType(ParamTypes...), UserPolicy> CreateSP(const TSharedRef<UserClass, Mode>& InUserObjectRef, typename TMemFunPtrType<true, UserClass, RetValType (ParamTypes..., std::decay_t<VarTypes>...)>::Type InFunc, VarTypes&&... Vars)
 	{
 		TDelegate<RetValType(ParamTypes...), UserPolicy> Result;
-		TBaseSPMethodDelegateInstance<true, const UserClass, Mode, FuncType, UserPolicy, VarTypes...>::Create(Result, InUserObjectRef, InFunc, Vars...);
+		new (Result) TBaseSPMethodDelegateInstance<true, const UserClass, Mode, FuncType, UserPolicy, std::decay_t<VarTypes>...>(InUserObjectRef, InFunc, Forward<VarTypes>(Vars)...);
 		return Result;
 	}
 
@@ -234,16 +234,20 @@ public:
 	 * You can use ExecuteIfBound() to call them.
 	 */
 	template <typename UserClass, typename... VarTypes>
-	UE_NODISCARD inline static TDelegate<RetValType(ParamTypes...), UserPolicy> CreateSP(UserClass* InUserObject, typename TMemFunPtrType<false, UserClass, RetValType (ParamTypes..., VarTypes...)>::Type InFunc, VarTypes... Vars)
+	UE_NODISCARD inline static TDelegate<RetValType(ParamTypes...), UserPolicy> CreateSP(UserClass* InUserObject, typename TMemFunPtrType<false, UserClass, RetValType (ParamTypes..., std::decay_t<VarTypes>...)>::Type InFunc, VarTypes&&... Vars)
 	{
 		static_assert(!TIsConst<UserClass>::Value, "Attempting to bind a delegate with a const object pointer and non-const member function.");
 
-		return CreateSP(StaticCastSharedRef<UserClass>(InUserObject->AsShared()), InFunc, Vars...);
+		TDelegate<RetValType(ParamTypes...), UserPolicy> Result;
+		new (Result) TBaseSPMethodDelegateInstance<false, UserClass, decltype(InUserObject->AsShared())::Mode, FuncType, UserPolicy, std::decay_t<VarTypes>...>(StaticCastSharedRef<UserClass>(InUserObject->AsShared()), InFunc, Forward<VarTypes>(Vars)...);
+		return Result;
 	}
 	template <typename UserClass, typename... VarTypes>
-	UE_NODISCARD inline static TDelegate<RetValType(ParamTypes...), UserPolicy> CreateSP(const UserClass* InUserObject, typename TMemFunPtrType<true, UserClass, RetValType (ParamTypes..., VarTypes...)>::Type InFunc, VarTypes... Vars)
+	UE_NODISCARD inline static TDelegate<RetValType(ParamTypes...), UserPolicy> CreateSP(const UserClass* InUserObject, typename TMemFunPtrType<true, UserClass, RetValType (ParamTypes..., std::decay_t<VarTypes>...)>::Type InFunc, VarTypes&&... Vars)
 	{
-		return CreateSP(StaticCastSharedRef<const UserClass>(InUserObject->AsShared()), InFunc, Vars...);
+		TDelegate<RetValType(ParamTypes...), UserPolicy> Result;
+		new (Result) TBaseSPMethodDelegateInstance<true, const UserClass, decltype(InUserObject->AsShared())::Mode, FuncType, UserPolicy, std::decay_t<VarTypes>...>(StaticCastSharedRef<const UserClass>(InUserObject->AsShared()), InFunc, Forward<VarTypes>(Vars)...);
+		return Result;
 	}
 
 	/**
@@ -255,19 +259,19 @@ public:
 	 * Note: This function is redundant, but is retained for backwards compatibility.  CreateSP() works in both thread-safe and not-thread-safe modes and should be preferred.
 	 */
 	template <typename UserClass, typename... VarTypes>
-	UE_NODISCARD inline static TDelegate<RetValType(ParamTypes...), UserPolicy> CreateThreadSafeSP(const TSharedRef<UserClass, ESPMode::ThreadSafe>& InUserObjectRef, typename TMemFunPtrType<false, UserClass, RetValType (ParamTypes..., VarTypes...)>::Type InFunc, VarTypes... Vars)
+	UE_NODISCARD inline static TDelegate<RetValType(ParamTypes...), UserPolicy> CreateThreadSafeSP(const TSharedRef<UserClass, ESPMode::ThreadSafe>& InUserObjectRef, typename TMemFunPtrType<false, UserClass, RetValType (ParamTypes..., std::decay_t<VarTypes>...)>::Type InFunc, VarTypes&&... Vars)
 	{
 		static_assert(!TIsConst<UserClass>::Value, "Attempting to bind a delegate with a const object pointer and non-const member function.");
 
 		TDelegate<RetValType(ParamTypes...), UserPolicy> Result;
-		TBaseSPMethodDelegateInstance<false, UserClass, ESPMode::ThreadSafe, FuncType, UserPolicy, VarTypes...>::Create(Result, InUserObjectRef, InFunc, Vars...);
+		new (Result) TBaseSPMethodDelegateInstance<false, UserClass, ESPMode::ThreadSafe, FuncType, UserPolicy, std::decay_t<VarTypes>...>(InUserObjectRef, InFunc, Forward<VarTypes>(Vars)...);
 		return Result;
 	}
 	template <typename UserClass, typename... VarTypes>
-	UE_NODISCARD inline static TDelegate<RetValType(ParamTypes...), UserPolicy> CreateThreadSafeSP(const TSharedRef<UserClass, ESPMode::ThreadSafe>& InUserObjectRef, typename TMemFunPtrType<true, UserClass, RetValType (ParamTypes..., VarTypes...)>::Type InFunc, VarTypes... Vars)
+	UE_NODISCARD inline static TDelegate<RetValType(ParamTypes...), UserPolicy> CreateThreadSafeSP(const TSharedRef<UserClass, ESPMode::ThreadSafe>& InUserObjectRef, typename TMemFunPtrType<true, UserClass, RetValType (ParamTypes..., std::decay_t<VarTypes>...)>::Type InFunc, VarTypes&&... Vars)
 	{
 		TDelegate<RetValType(ParamTypes...), UserPolicy> Result;
-		TBaseSPMethodDelegateInstance<true, const UserClass, ESPMode::ThreadSafe, FuncType, UserPolicy, VarTypes...>::Create(Result, InUserObjectRef, InFunc, Vars...);
+		new (Result) TBaseSPMethodDelegateInstance<true, const UserClass, ESPMode::ThreadSafe, FuncType, UserPolicy, std::decay_t<VarTypes>...>(InUserObjectRef, InFunc, Forward<VarTypes>(Vars)...);
 		return Result;
 	}
 
@@ -280,16 +284,20 @@ public:
 	 * Note: This function is redundant, but is retained for backwards compatibility.  CreateSP() works in both thread-safe and not-thread-safe modes and should be preferred.
 	 */
 	template <typename UserClass, typename... VarTypes>
-	UE_NODISCARD inline static TDelegate<RetValType(ParamTypes...), UserPolicy> CreateThreadSafeSP(UserClass* InUserObject, typename TMemFunPtrType<false, UserClass, RetValType (ParamTypes..., VarTypes...)>::Type InFunc, VarTypes... Vars)
+	UE_NODISCARD inline static TDelegate<RetValType(ParamTypes...), UserPolicy> CreateThreadSafeSP(UserClass* InUserObject, typename TMemFunPtrType<false, UserClass, RetValType (ParamTypes..., std::decay_t<VarTypes>...)>::Type InFunc, VarTypes&&... Vars)
 	{
 		static_assert(!TIsConst<UserClass>::Value, "Attempting to bind a delegate with a const object pointer and non-const member function.");
 
-		return CreateThreadSafeSP(StaticCastSharedRef<UserClass>(InUserObject->AsShared()), InFunc, Vars...);
+		TDelegate<RetValType(ParamTypes...), UserPolicy> Result;
+		new (Result) TBaseSPMethodDelegateInstance<false, UserClass, ESPMode::ThreadSafe, FuncType, UserPolicy, std::decay_t<VarTypes>...>(StaticCastSharedRef<UserClass>(InUserObject->AsShared()), InFunc, Forward<VarTypes>(Vars)...);
+		return Result;
 	}
 	template <typename UserClass, typename... VarTypes>
-	UE_NODISCARD inline static TDelegate<RetValType(ParamTypes...), UserPolicy> CreateThreadSafeSP(const UserClass* InUserObject, typename TMemFunPtrType<true, UserClass, RetValType (ParamTypes..., VarTypes...)>::Type InFunc, VarTypes... Vars)
+	UE_NODISCARD inline static TDelegate<RetValType(ParamTypes...), UserPolicy> CreateThreadSafeSP(const UserClass* InUserObject, typename TMemFunPtrType<true, UserClass, RetValType (ParamTypes..., std::decay_t<VarTypes>...)>::Type InFunc, VarTypes&&... Vars)
 	{
-		return CreateThreadSafeSP(StaticCastSharedRef<const UserClass>(InUserObject->AsShared()), InFunc, Vars...);
+		TDelegate<RetValType(ParamTypes...), UserPolicy> Result;
+		new (Result) TBaseSPMethodDelegateInstance<true, const UserClass, ESPMode::ThreadSafe, FuncType, UserPolicy, std::decay_t<VarTypes>...>(StaticCastSharedRef<const UserClass>(InUserObject->AsShared()), InFunc, Forward<VarTypes>(Vars)...);
+		return Result;
 	}
 
 	/**
@@ -299,17 +307,17 @@ public:
 	 * You can use ExecuteIfBound() to call them.
 	 */
 	template <typename UObjectTemplate, typename... VarTypes>
-	UE_NODISCARD inline static TDelegate<RetValType(ParamTypes...), UserPolicy> CreateUFunction(UObjectTemplate* InUserObject, const FName& InFunctionName, VarTypes... Vars)
+	UE_NODISCARD inline static TDelegate<RetValType(ParamTypes...), UserPolicy> CreateUFunction(UObjectTemplate* InUserObject, const FName& InFunctionName, VarTypes&&... Vars)
 	{
 		TDelegate<RetValType(ParamTypes...), UserPolicy> Result;
-		TBaseUFunctionDelegateInstance<UObjectTemplate, FuncType, UserPolicy, VarTypes...>::Create(Result, InUserObject, InFunctionName, Vars...);
+		new (Result) TBaseUFunctionDelegateInstance<UObjectTemplate, FuncType, UserPolicy, std::decay_t<VarTypes>...>(InUserObject, InFunctionName, Forward<VarTypes>(Vars)...);
 		return Result;
 	}
 	template <typename UObjectTemplate, typename... VarTypes>
-	UE_NODISCARD inline static TDelegate<RetValType(ParamTypes...), UserPolicy> CreateUFunction(TObjectPtr<UObjectTemplate> InUserObject, const FName& InFunctionName, VarTypes... Vars)
+	UE_NODISCARD inline static TDelegate<RetValType(ParamTypes...), UserPolicy> CreateUFunction(TObjectPtr<UObjectTemplate> InUserObject, const FName& InFunctionName, VarTypes&&... Vars)
 	{
 		TDelegate<RetValType(ParamTypes...), UserPolicy> Result;
-		TBaseUFunctionDelegateInstance<UObjectTemplate, FuncType, UserPolicy, VarTypes...>::Create(Result, ToRawPtr(InUserObject), InFunctionName, Vars...);
+		new (Result) TBaseUFunctionDelegateInstance<UObjectTemplate, FuncType, UserPolicy, std::decay_t<VarTypes>...>(ToRawPtr(InUserObject), InFunctionName, Forward<VarTypes>(Vars)...);
 		return Result;
 	}
 
@@ -320,35 +328,35 @@ public:
 	 * You can use ExecuteIfBound() to call them.
 	 */
 	template <typename UserClass, typename... VarTypes>
-	UE_NODISCARD inline static TDelegate<RetValType(ParamTypes...), UserPolicy> CreateUObject(UserClass* InUserObject, typename TMemFunPtrType<false, UserClass, RetValType (ParamTypes..., VarTypes...)>::Type InFunc, VarTypes... Vars)
+	UE_NODISCARD inline static TDelegate<RetValType(ParamTypes...), UserPolicy> CreateUObject(UserClass* InUserObject, typename TMemFunPtrType<false, UserClass, RetValType (ParamTypes..., std::decay_t<VarTypes>...)>::Type InFunc, VarTypes&&... Vars)
 	{
 		static_assert(!TIsConst<UserClass>::Value, "Attempting to bind a delegate with a const object pointer and non-const member function.");
 
 		TDelegate<RetValType(ParamTypes...), UserPolicy> Result;
-		TBaseUObjectMethodDelegateInstance<false, UserClass, FuncType, UserPolicy, VarTypes...>::Create(Result, InUserObject, InFunc, Vars...);
+		new (Result) TBaseUObjectMethodDelegateInstance<false, UserClass, FuncType, UserPolicy, std::decay_t<VarTypes>...>(InUserObject, InFunc, Forward<VarTypes>(Vars)...);
 		return Result;
 	}
 	template <typename UserClass, typename... VarTypes>
-	UE_NODISCARD inline static TDelegate<RetValType(ParamTypes...), UserPolicy> CreateUObject(const UserClass* InUserObject, typename TMemFunPtrType<true, UserClass, RetValType (ParamTypes..., VarTypes...)>::Type InFunc, VarTypes... Vars)
+	UE_NODISCARD inline static TDelegate<RetValType(ParamTypes...), UserPolicy> CreateUObject(const UserClass* InUserObject, typename TMemFunPtrType<true, UserClass, RetValType (ParamTypes..., std::decay_t<VarTypes>...)>::Type InFunc, VarTypes&&... Vars)
 	{
 		TDelegate<RetValType(ParamTypes...), UserPolicy> Result;
-		TBaseUObjectMethodDelegateInstance<true, const UserClass, FuncType, UserPolicy, VarTypes...>::Create(Result, InUserObject, InFunc, Vars...);
+		new (Result) TBaseUObjectMethodDelegateInstance<true, const UserClass, FuncType, UserPolicy, std::decay_t<VarTypes>...>(InUserObject, InFunc, Forward<VarTypes>(Vars)...);
 		return Result;
 	}
 	template <typename UserClass, typename... VarTypes>
-	UE_NODISCARD inline static TDelegate<RetValType(ParamTypes...), UserPolicy> CreateUObject(TObjectPtr<UserClass> InUserObject, typename TMemFunPtrType<false, UserClass, RetValType (ParamTypes..., VarTypes...)>::Type InFunc, VarTypes... Vars)
+	UE_NODISCARD inline static TDelegate<RetValType(ParamTypes...), UserPolicy> CreateUObject(TObjectPtr<UserClass> InUserObject, typename TMemFunPtrType<false, UserClass, RetValType (ParamTypes..., std::decay_t<VarTypes>...)>::Type InFunc, VarTypes&&... Vars)
 	{
 		static_assert(!TIsConst<UserClass>::Value, "Attempting to bind a delegate with a const object pointer and non-const member function.");
 
 		TDelegate<RetValType(ParamTypes...), UserPolicy> Result;
-		TBaseUObjectMethodDelegateInstance<false, UserClass, FuncType, UserPolicy, VarTypes...>::Create(Result, ToRawPtr(InUserObject), InFunc, Vars...);
+		new (Result) TBaseUObjectMethodDelegateInstance<false, UserClass, FuncType, UserPolicy, std::decay_t<VarTypes>...>(ToRawPtr(InUserObject), InFunc, Forward<VarTypes>(Vars)...);
 		return Result;
 	}
 	template <typename UserClass, typename... VarTypes>
-	UE_NODISCARD inline static TDelegate<RetValType(ParamTypes...), UserPolicy> CreateUObject(TObjectPtr<UserClass> InUserObject, typename TMemFunPtrType<true, UserClass, RetValType (ParamTypes..., VarTypes...)>::Type InFunc, VarTypes... Vars)
+	UE_NODISCARD inline static TDelegate<RetValType(ParamTypes...), UserPolicy> CreateUObject(TObjectPtr<UserClass> InUserObject, typename TMemFunPtrType<true, UserClass, RetValType (ParamTypes..., std::decay_t<VarTypes>...)>::Type InFunc, VarTypes&&... Vars)
 	{
 		TDelegate<RetValType(ParamTypes...), UserPolicy> Result;
-		TBaseUObjectMethodDelegateInstance<true, const UserClass, FuncType, UserPolicy, VarTypes...>::Create(Result, ToRawPtr(InUserObject), InFunc, Vars...);
+		new (Result) TBaseUObjectMethodDelegateInstance<true, const UserClass, FuncType, UserPolicy, std::decay_t<VarTypes>...>(ToRawPtr(InUserObject), InFunc, Forward<VarTypes>(Vars)...);
 		return Result;
 	}
 
@@ -431,9 +439,9 @@ public:
 	 * Binds a raw C++ pointer global function delegate
 	 */
 	template <typename... VarTypes>
-	inline void BindStatic(typename TBaseStaticDelegateInstance<FuncType, UserPolicy, VarTypes...>::FFuncPtr InFunc, VarTypes... Vars)
+	inline void BindStatic(typename TBaseStaticDelegateInstance<FuncType, UserPolicy, std::decay_t<VarTypes>...>::FFuncPtr InFunc, VarTypes&&... Vars)
 	{
-		*this = CreateStatic(InFunc, Vars...);
+		new (*this) TBaseStaticDelegateInstance<FuncType, UserPolicy, std::decay_t<VarTypes>...>(InFunc, Forward<VarTypes>(Vars)...);
 	}
 	
 	/**
@@ -441,9 +449,9 @@ public:
 	 * technically this works for any functor types, but lambdas are the primary use case
 	 */
 	template<typename FunctorType, typename... VarTypes>
-	inline void BindLambda(FunctorType&& InFunctor, VarTypes... Vars)
+	inline void BindLambda(FunctorType&& InFunctor, VarTypes&&... Vars)
 	{
-		*this = CreateLambda(Forward<FunctorType>(InFunctor), Vars...);
+		new (*this) TBaseFunctorDelegateInstance<FuncType, UserPolicy, typename TRemoveReference<FunctorType>::Type, std::decay_t<VarTypes>...>(Forward<FunctorType>(InFunctor), Forward<VarTypes>(Vars)...);
 	}
 
 	/**
@@ -451,9 +459,9 @@ public:
 	 * technically this works for any functor types, but lambdas are the primary use case
 	 */
 	template<typename UserClass, typename FunctorType, typename... VarTypes>
-	inline void BindWeakLambda(UserClass* InUserObject, FunctorType&& InFunctor, VarTypes... Vars)
+	inline void BindWeakLambda(UserClass* InUserObject, FunctorType&& InFunctor, VarTypes&&... Vars)
 	{
-		*this = CreateWeakLambda(InUserObject, Forward<FunctorType>(InFunctor), Vars...);
+		new (*this) TWeakBaseFunctorDelegateInstance<UserClass, FuncType, UserPolicy, typename TRemoveReference<FunctorType>::Type, std::decay_t<VarTypes>...>(InUserObject, Forward<FunctorType>(InFunctor), Forward<VarTypes>(Vars)...);
 	}
 
 	/**
@@ -463,32 +471,32 @@ public:
 	 * deleted out from underneath your delegate. Be careful when calling Execute()!
 	 */
 	template <typename UserClass, typename... VarTypes>
-	inline void BindRaw(UserClass* InUserObject, typename TMemFunPtrType<false, UserClass, RetValType (ParamTypes..., VarTypes...)>::Type InFunc, VarTypes... Vars)
+	inline void BindRaw(UserClass* InUserObject, typename TMemFunPtrType<false, UserClass, RetValType (ParamTypes..., std::decay_t<VarTypes>...)>::Type InFunc, VarTypes&&... Vars)
 	{
 		static_assert(!TIsConst<UserClass>::Value, "Attempting to bind a delegate with a const object pointer and non-const member function.");
 
-		*this = CreateRaw(InUserObject, InFunc, Vars...);
+		new (*this) TBaseRawMethodDelegateInstance<false, UserClass, FuncType, UserPolicy, std::decay_t<VarTypes>...>(InUserObject, InFunc, Forward<VarTypes>(Vars)...);
 	}
 	template <typename UserClass, typename... VarTypes>
-	inline void BindRaw(const UserClass* InUserObject, typename TMemFunPtrType<true, UserClass, RetValType (ParamTypes..., VarTypes...)>::Type InFunc, VarTypes... Vars)
+	inline void BindRaw(const UserClass* InUserObject, typename TMemFunPtrType<true, UserClass, RetValType (ParamTypes..., std::decay_t<VarTypes>...)>::Type InFunc, VarTypes&&... Vars)
 	{
-		*this = CreateRaw(InUserObject, InFunc, Vars...);
+		new (*this) TBaseRawMethodDelegateInstance<true, const UserClass, FuncType, UserPolicy, std::decay_t<VarTypes>...>(InUserObject, InFunc, Forward<VarTypes>(Vars)...);
 	}
 
 	/**
 	 * Binds a shared pointer-based member function delegate.  Shared pointer delegates keep a weak reference to your object.  You can use ExecuteIfBound() to call them.
 	 */
 	template <typename UserClass, ESPMode Mode, typename... VarTypes>
-	inline void BindSP(const TSharedRef<UserClass, Mode>& InUserObjectRef, typename TMemFunPtrType<false, UserClass, RetValType (ParamTypes..., VarTypes...)>::Type InFunc, VarTypes... Vars)
+	inline void BindSP(const TSharedRef<UserClass, Mode>& InUserObjectRef, typename TMemFunPtrType<false, UserClass, RetValType (ParamTypes..., std::decay_t<VarTypes>...)>::Type InFunc, VarTypes&&... Vars)
 	{
 		static_assert(!TIsConst<UserClass>::Value, "Attempting to bind a delegate with a const object pointer and non-const member function.");
 
-		*this = CreateSP(InUserObjectRef, InFunc, Vars...);
+		new (*this) TBaseSPMethodDelegateInstance<false, UserClass, Mode, FuncType, UserPolicy, std::decay_t<VarTypes>...>(InUserObjectRef, InFunc, Forward<VarTypes>(Vars)...);
 	}
 	template <typename UserClass, ESPMode Mode, typename... VarTypes>
-	inline void BindSP(const TSharedRef<UserClass, Mode>& InUserObjectRef, typename TMemFunPtrType<true, UserClass, RetValType (ParamTypes..., VarTypes...)>::Type InFunc, VarTypes... Vars)
+	inline void BindSP(const TSharedRef<UserClass, Mode>& InUserObjectRef, typename TMemFunPtrType<true, UserClass, RetValType (ParamTypes..., std::decay_t<VarTypes>...)>::Type InFunc, VarTypes&&... Vars)
 	{
-		*this = CreateSP(InUserObjectRef, InFunc, Vars...);
+		new (*this) TBaseSPMethodDelegateInstance<true, const UserClass, Mode, FuncType, UserPolicy, std::decay_t<VarTypes>...>(InUserObjectRef, InFunc, Forward<VarTypes>(Vars)...);
 	}
 
 	/**
@@ -498,16 +506,16 @@ public:
 	 * You can use ExecuteIfBound() to call them.
 	 */
 	template <typename UserClass, typename... VarTypes>
-	inline void BindSP(UserClass* InUserObject, typename TMemFunPtrType<false, UserClass, RetValType (ParamTypes..., VarTypes...)>::Type InFunc, VarTypes... Vars)
+	inline void BindSP(UserClass* InUserObject, typename TMemFunPtrType<false, UserClass, RetValType (ParamTypes..., std::decay_t<VarTypes>...)>::Type InFunc, VarTypes&&... Vars)
 	{
 		static_assert(!TIsConst<UserClass>::Value, "Attempting to bind a delegate with a const object pointer and non-const member function.");
 
-		*this = CreateSP(InUserObject, InFunc, Vars...);
+		new (*this) TBaseSPMethodDelegateInstance<false, UserClass, decltype(InUserObject->AsShared())::Mode, FuncType, UserPolicy, std::decay_t<VarTypes>...>(StaticCastSharedRef<UserClass>(InUserObject->AsShared()), InFunc, Forward<VarTypes>(Vars)...);
 	}
 	template <typename UserClass, typename... VarTypes>
-	inline void BindSP(const UserClass* InUserObject, typename TMemFunPtrType<true, UserClass, RetValType (ParamTypes..., VarTypes...)>::Type InFunc, VarTypes... Vars)
+	inline void BindSP(const UserClass* InUserObject, typename TMemFunPtrType<true, UserClass, RetValType (ParamTypes..., std::decay_t<VarTypes>...)>::Type InFunc, VarTypes&&... Vars)
 	{
-		*this = CreateSP(InUserObject, InFunc, Vars...);
+		new (*this) TBaseSPMethodDelegateInstance<true, const UserClass, decltype(InUserObject->AsShared())::Mode, FuncType, UserPolicy, std::decay_t<VarTypes>...>(StaticCastSharedRef<const UserClass>(InUserObject->AsShared()), InFunc, Forward<VarTypes>(Vars)...);
 	}
 
 	/**
@@ -519,16 +527,16 @@ public:
 	 * Note: This function is redundant, but is retained for backwards compatibility.  BindSP() works in both thread-safe and not-thread-safe modes and should be preferred.
 	 */
 	template <typename UserClass, typename... VarTypes>
-	inline void BindThreadSafeSP(const TSharedRef<UserClass, ESPMode::ThreadSafe>& InUserObjectRef, typename TMemFunPtrType<false, UserClass, RetValType (ParamTypes..., VarTypes...)>::Type InFunc, VarTypes... Vars)
+	inline void BindThreadSafeSP(const TSharedRef<UserClass, ESPMode::ThreadSafe>& InUserObjectRef, typename TMemFunPtrType<false, UserClass, RetValType (ParamTypes..., std::decay_t<VarTypes>...)>::Type InFunc, VarTypes&&... Vars)
 	{
 		static_assert(!TIsConst<UserClass>::Value, "Attempting to bind a delegate with a const object pointer and non-const member function.");
 
-		*this = CreateThreadSafeSP(InUserObjectRef, InFunc, Vars...);
+		new (*this) TBaseSPMethodDelegateInstance<false, UserClass, ESPMode::ThreadSafe, FuncType, UserPolicy, std::decay_t<VarTypes>...>(InUserObjectRef, InFunc, Forward<VarTypes>(Vars)...);
 	}
 	template <typename UserClass, typename... VarTypes>
-	inline void BindThreadSafeSP(const TSharedRef<UserClass, ESPMode::ThreadSafe>& InUserObjectRef, typename TMemFunPtrType<true, UserClass, RetValType (ParamTypes..., VarTypes...)>::Type InFunc, VarTypes... Vars)
+	inline void BindThreadSafeSP(const TSharedRef<UserClass, ESPMode::ThreadSafe>& InUserObjectRef, typename TMemFunPtrType<true, UserClass, RetValType (ParamTypes..., std::decay_t<VarTypes>...)>::Type InFunc, VarTypes&&... Vars)
 	{
-		*this = CreateThreadSafeSP(InUserObjectRef, InFunc, Vars...);
+		new (*this) TBaseSPMethodDelegateInstance<true, const UserClass, ESPMode::ThreadSafe, FuncType, UserPolicy, std::decay_t<VarTypes>...>(InUserObjectRef, InFunc, Forward<VarTypes>(Vars)...);
 	}
 
 	/**
@@ -540,16 +548,16 @@ public:
 	 * Note: This function is redundant, but is retained for backwards compatibility.  BindSP() works in both thread-safe and not-thread-safe modes and should be preferred.
 	 */
 	template <typename UserClass, typename... VarTypes>
-	inline void BindThreadSafeSP(UserClass* InUserObject, typename TMemFunPtrType<false, UserClass, RetValType (ParamTypes..., VarTypes...)>::Type InFunc, VarTypes... Vars)
+	inline void BindThreadSafeSP(UserClass* InUserObject, typename TMemFunPtrType<false, UserClass, RetValType (ParamTypes..., std::decay_t<VarTypes>...)>::Type InFunc, VarTypes&&... Vars)
 	{
 		static_assert(!TIsConst<UserClass>::Value, "Attempting to bind a delegate with a const object pointer and non-const member function.");
 
-		*this = CreateThreadSafeSP(InUserObject, InFunc, Vars...);
+		new (*this) TBaseSPMethodDelegateInstance<false, UserClass, ESPMode::ThreadSafe, FuncType, UserPolicy, std::decay_t<VarTypes>...>(StaticCastSharedRef<UserClass>(InUserObject->AsShared()), InFunc, Forward<VarTypes>(Vars)...);
 	}
 	template <typename UserClass, typename... VarTypes>
-	inline void BindThreadSafeSP(const UserClass* InUserObject, typename TMemFunPtrType<true, UserClass, RetValType (ParamTypes..., VarTypes...)>::Type InFunc, VarTypes... Vars)
+	inline void BindThreadSafeSP(const UserClass* InUserObject, typename TMemFunPtrType<true, UserClass, RetValType (ParamTypes..., std::decay_t<VarTypes>...)>::Type InFunc, VarTypes&&... Vars)
 	{
-		*this = CreateThreadSafeSP(InUserObject, InFunc, Vars...);
+		new (*this) TBaseSPMethodDelegateInstance<true, const UserClass, ESPMode::ThreadSafe, FuncType, UserPolicy, std::decay_t<VarTypes>...>(StaticCastSharedRef<const UserClass>(InUserObject->AsShared()), InFunc, Forward<VarTypes>(Vars)...);
 	}
 
 	/**
@@ -559,14 +567,14 @@ public:
 	 * You can use ExecuteIfBound() to call them.
 	 */
 	template <typename UObjectTemplate, typename... VarTypes>
-	inline void BindUFunction(UObjectTemplate* InUserObject, const FName& InFunctionName, VarTypes... Vars)
+	inline void BindUFunction(UObjectTemplate* InUserObject, const FName& InFunctionName, VarTypes&&... Vars)
 	{
-		*this = CreateUFunction(InUserObject, InFunctionName, Vars...);
+		new (*this) TBaseUFunctionDelegateInstance<UObjectTemplate, FuncType, UserPolicy, std::decay_t<VarTypes>...>(InUserObject, InFunctionName, Forward<VarTypes>(Vars)...);
 	}
 	template <typename UObjectTemplate, typename... VarTypes>
-	inline void BindUFunction(TObjectPtr<UObjectTemplate> InUserObject, const FName& InFunctionName, VarTypes... Vars)
+	inline void BindUFunction(TObjectPtr<UObjectTemplate> InUserObject, const FName& InFunctionName, VarTypes&&... Vars)
 	{
-		*this = CreateUFunction(InUserObject, InFunctionName, Vars...);
+		new (*this) TBaseUFunctionDelegateInstance<UObjectTemplate, FuncType, UserPolicy, std::decay_t<VarTypes>...>(ToRawPtr(InUserObject), InFunctionName, Forward<VarTypes>(Vars)...);
 	}
 
 	/**
@@ -576,28 +584,28 @@ public:
 	 * You can use ExecuteIfBound() to call them.
 	 */
 	template <typename UserClass, typename... VarTypes>
-	inline void BindUObject(UserClass* InUserObject, typename TMemFunPtrType<false, UserClass, RetValType (ParamTypes..., VarTypes...)>::Type InFunc, VarTypes... Vars)
+	inline void BindUObject(UserClass* InUserObject, typename TMemFunPtrType<false, UserClass, RetValType (ParamTypes..., std::decay_t<VarTypes>...)>::Type InFunc, VarTypes&&... Vars)
 	{
 		static_assert(!TIsConst<UserClass>::Value, "Attempting to bind a delegate with a const object pointer and non-const member function.");
 
-		*this = CreateUObject(InUserObject, InFunc, Vars...);
+		new (*this) TBaseUObjectMethodDelegateInstance<false, UserClass, FuncType, UserPolicy, std::decay_t<VarTypes>...>(InUserObject, InFunc, Forward<VarTypes>(Vars)...);
 	}
 	template <typename UserClass, typename... VarTypes>
-	inline void BindUObject(const UserClass* InUserObject, typename TMemFunPtrType<true, UserClass, RetValType (ParamTypes..., VarTypes...)>::Type InFunc, VarTypes... Vars)
+	inline void BindUObject(const UserClass* InUserObject, typename TMemFunPtrType<true, UserClass, RetValType (ParamTypes..., std::decay_t<VarTypes>...)>::Type InFunc, VarTypes&&... Vars)
 	{
-		*this = CreateUObject(InUserObject, InFunc, Vars...);
+		new (*this) TBaseUObjectMethodDelegateInstance<true, const UserClass, FuncType, UserPolicy, std::decay_t<VarTypes>...>(InUserObject, InFunc, Forward<VarTypes>(Vars)...);
 	}
 	template <typename UserClass, typename... VarTypes>
-	inline void BindUObject(TObjectPtr<UserClass> InUserObject, typename TMemFunPtrType<false, UserClass, RetValType (ParamTypes..., VarTypes...)>::Type InFunc, VarTypes... Vars)
+	inline void BindUObject(TObjectPtr<UserClass> InUserObject, typename TMemFunPtrType<false, UserClass, RetValType (ParamTypes..., std::decay_t<VarTypes>...)>::Type InFunc, VarTypes&&... Vars)
 	{
 		static_assert(!TIsConst<UserClass>::Value, "Attempting to bind a delegate with a const object pointer and non-const member function.");
 
-		*this = CreateUObject(InUserObject, InFunc, Vars...);
+		new (*this) TBaseUObjectMethodDelegateInstance<false, UserClass, FuncType, UserPolicy, std::decay_t<VarTypes>...>(ToRawPtr(InUserObject), InFunc, Forward<VarTypes>(Vars)...);
 	}
 	template <typename UserClass, typename... VarTypes>
-	inline void BindUObject(TObjectPtr<UserClass> InUserObject, typename TMemFunPtrType<true, UserClass, RetValType (ParamTypes..., VarTypes...)>::Type InFunc, VarTypes... Vars)
+	inline void BindUObject(TObjectPtr<UserClass> InUserObject, typename TMemFunPtrType<true, UserClass, RetValType (ParamTypes..., std::decay_t<VarTypes>...)>::Type InFunc, VarTypes&&... Vars)
 	{
-		*this = CreateUObject(InUserObject, InFunc, Vars...);
+		new (*this) TBaseUObjectMethodDelegateInstance<true, const UserClass, FuncType, UserPolicy, std::decay_t<VarTypes>...>(ToRawPtr(InUserObject), InFunc, Forward<VarTypes>(Vars)...);
 	}
 
 public:
@@ -725,9 +733,9 @@ public:
 	 * @param	InFunc	Function pointer
 	 */
 	template <typename... VarTypes>
-	inline FDelegateHandle AddStatic(typename TBaseStaticDelegateInstance<void (ParamTypes...), UserPolicy, VarTypes...>::FFuncPtr InFunc, VarTypes... Vars)
+	inline FDelegateHandle AddStatic(typename TBaseStaticDelegateInstance<void (ParamTypes...), UserPolicy, std::decay_t<VarTypes>...>::FFuncPtr InFunc, VarTypes&&... Vars)
 	{
-		return Add(FDelegate::CreateStatic(InFunc, Vars...));
+		return Add(FDelegate::CreateStatic(InFunc, Forward<VarTypes>(Vars)...));
 	}
 
 	/**
@@ -737,9 +745,9 @@ public:
 	 * @param	InFunctor	Functor (e.g. Lambda)
 	 */
 	template<typename FunctorType, typename... VarTypes>
-	inline FDelegateHandle AddLambda(FunctorType&& InFunctor, VarTypes... Vars)
+	inline FDelegateHandle AddLambda(FunctorType&& InFunctor, VarTypes&&... Vars)
 	{
-		return Add(FDelegate::CreateLambda(Forward<FunctorType>(InFunctor), Vars...));
+		return Add(FDelegate::CreateLambda(Forward<FunctorType>(InFunctor), Forward<VarTypes>(Vars)...));
 	}
 
 	/**
@@ -750,9 +758,9 @@ public:
 	 * @param	InFunctor		Functor (e.g. Lambda)
 	 */
 	template<typename UserClass, typename FunctorType, typename... VarTypes>
-	inline FDelegateHandle AddWeakLambda(UserClass* InUserObject, FunctorType&& InFunctor, VarTypes... Vars)
+	inline FDelegateHandle AddWeakLambda(UserClass* InUserObject, FunctorType&& InFunctor, VarTypes&&... Vars)
 	{
-		return Add(FDelegate::CreateWeakLambda(InUserObject, Forward<FunctorType>(InFunctor), Vars...));
+		return Add(FDelegate::CreateWeakLambda(InUserObject, Forward<FunctorType>(InFunctor), Forward<VarTypes>(Vars)...));
 	}
 
 	/**
@@ -765,16 +773,16 @@ public:
 	 * @param	InFunc			Class method function address
 	 */
 	template <typename UserClass, typename... VarTypes>
-	inline FDelegateHandle AddRaw(UserClass* InUserObject, typename TMemFunPtrType<false, UserClass, void (ParamTypes..., VarTypes...)>::Type InFunc, VarTypes... Vars)
+	inline FDelegateHandle AddRaw(UserClass* InUserObject, typename TMemFunPtrType<false, UserClass, void (ParamTypes..., std::decay_t<VarTypes>...)>::Type InFunc, VarTypes&&... Vars)
 	{
 		static_assert(!TIsConst<UserClass>::Value, "Attempting to bind a delegate with a const object pointer and non-const member function.");
 
-		return Add(FDelegate::CreateRaw(InUserObject, InFunc, Vars...));
+		return Add(FDelegate::CreateRaw(InUserObject, InFunc, Forward<VarTypes>(Vars)...));
 	}
 	template <typename UserClass, typename... VarTypes>
-	inline FDelegateHandle AddRaw(const UserClass* InUserObject, typename TMemFunPtrType<true, UserClass, void (ParamTypes..., VarTypes...)>::Type InFunc, VarTypes... Vars)
+	inline FDelegateHandle AddRaw(const UserClass* InUserObject, typename TMemFunPtrType<true, UserClass, void (ParamTypes..., std::decay_t<VarTypes>...)>::Type InFunc, VarTypes&&... Vars)
 	{
-		return Add(FDelegate::CreateRaw(InUserObject, InFunc, Vars...));
+		return Add(FDelegate::CreateRaw(InUserObject, InFunc, Forward<VarTypes>(Vars)...));
 	}
 
 	/**
@@ -786,16 +794,16 @@ public:
 	 * @param	InFunc			Class method function address
 	 */
 	template <typename UserClass, ESPMode Mode, typename... VarTypes>
-	inline FDelegateHandle AddSP(const TSharedRef<UserClass, Mode>& InUserObjectRef, typename TMemFunPtrType<false, UserClass, void (ParamTypes..., VarTypes...)>::Type InFunc, VarTypes... Vars)
+	inline FDelegateHandle AddSP(const TSharedRef<UserClass, Mode>& InUserObjectRef, typename TMemFunPtrType<false, UserClass, void (ParamTypes..., std::decay_t<VarTypes>...)>::Type InFunc, VarTypes&&... Vars)
 	{
 		static_assert(!TIsConst<UserClass>::Value, "Attempting to bind a delegate with a const object pointer and non-const member function.");
 
-		return Add(FDelegate::CreateSP(InUserObjectRef, InFunc, Vars...));
+		return Add(FDelegate::CreateSP(InUserObjectRef, InFunc, Forward<VarTypes>(Vars)...));
 	}
 	template <typename UserClass, ESPMode Mode, typename... VarTypes>
-	inline FDelegateHandle AddSP(const TSharedRef<UserClass, Mode>& InUserObjectRef, typename TMemFunPtrType<true, UserClass, void (ParamTypes..., VarTypes...)>::Type InFunc, VarTypes... Vars)
+	inline FDelegateHandle AddSP(const TSharedRef<UserClass, Mode>& InUserObjectRef, typename TMemFunPtrType<true, UserClass, void (ParamTypes..., std::decay_t<VarTypes>...)>::Type InFunc, VarTypes&&... Vars)
 	{
-		return Add(FDelegate::CreateSP(InUserObjectRef, InFunc, Vars...));
+		return Add(FDelegate::CreateSP(InUserObjectRef, InFunc, Forward<VarTypes>(Vars)...));
 	}
 
 	/**
@@ -807,16 +815,16 @@ public:
 	 * @param	InFunc			Class method function address
 	 */
 	template <typename UserClass, typename... VarTypes>
-	inline FDelegateHandle AddSP(UserClass* InUserObject, typename TMemFunPtrType<false, UserClass, void (ParamTypes..., VarTypes...)>::Type InFunc, VarTypes... Vars)
+	inline FDelegateHandle AddSP(UserClass* InUserObject, typename TMemFunPtrType<false, UserClass, void (ParamTypes..., std::decay_t<VarTypes>...)>::Type InFunc, VarTypes&&... Vars)
 	{
 		static_assert(!TIsConst<UserClass>::Value, "Attempting to bind a delegate with a const object pointer and non-const member function.");
 
-		return Add(FDelegate::CreateSP(InUserObject, InFunc, Vars...));
+		return Add(FDelegate::CreateSP(InUserObject, InFunc, Forward<VarTypes>(Vars)...));
 	}
 	template <typename UserClass, typename... VarTypes>
-	inline FDelegateHandle AddSP(const UserClass* InUserObject, typename TMemFunPtrType<true, UserClass, void (ParamTypes..., VarTypes...)>::Type InFunc, VarTypes... Vars)
+	inline FDelegateHandle AddSP(const UserClass* InUserObject, typename TMemFunPtrType<true, UserClass, void (ParamTypes..., std::decay_t<VarTypes>...)>::Type InFunc, VarTypes&&... Vars)
 	{
-		return Add(FDelegate::CreateSP(InUserObject, InFunc, Vars...));
+		return Add(FDelegate::CreateSP(InUserObject, InFunc, Forward<VarTypes>(Vars)...));
 	}
 
 	/**
@@ -828,16 +836,16 @@ public:
 	 * Note: This function is redundant, but is retained for backwards compatibility.  AddSP() works in both thread-safe and not-thread-safe modes and should be preferred.
 	 */
 	template <typename UserClass, typename... VarTypes>
-	inline FDelegateHandle AddThreadSafeSP(const TSharedRef<UserClass, ESPMode::ThreadSafe>& InUserObjectRef, typename TMemFunPtrType<false, UserClass, void (ParamTypes..., VarTypes...)>::Type InFunc, VarTypes... Vars)
+	inline FDelegateHandle AddThreadSafeSP(const TSharedRef<UserClass, ESPMode::ThreadSafe>& InUserObjectRef, typename TMemFunPtrType<false, UserClass, void (ParamTypes..., std::decay_t<VarTypes>...)>::Type InFunc, VarTypes&&... Vars)
 	{
 		static_assert(!TIsConst<UserClass>::Value, "Attempting to bind a delegate with a const object pointer and non-const member function.");
 
-		return Add(FDelegate::CreateThreadSafeSP(InUserObjectRef, InFunc, Vars...));
+		return Add(FDelegate::CreateThreadSafeSP(InUserObjectRef, InFunc, Forward<VarTypes>(Vars)...));
 	}
 	template <typename UserClass, typename... VarTypes>
-	inline FDelegateHandle AddThreadSafeSP(const TSharedRef<UserClass, ESPMode::ThreadSafe>& InUserObjectRef, typename TMemFunPtrType<true, UserClass, void (ParamTypes..., VarTypes...)>::Type InFunc, VarTypes... Vars)
+	inline FDelegateHandle AddThreadSafeSP(const TSharedRef<UserClass, ESPMode::ThreadSafe>& InUserObjectRef, typename TMemFunPtrType<true, UserClass, void (ParamTypes..., std::decay_t<VarTypes>...)>::Type InFunc, VarTypes&&... Vars)
 	{
-		return Add(FDelegate::CreateThreadSafeSP(InUserObjectRef, InFunc, Vars...));
+		return Add(FDelegate::CreateThreadSafeSP(InUserObjectRef, InFunc, Forward<VarTypes>(Vars)...));
 	}
 
 	/**
@@ -851,16 +859,16 @@ public:
 	 * Note: This function is redundant, but is retained for backwards compatibility.  AddSP() works in both thread-safe and not-thread-safe modes and should be preferred.
 	 */
 	template <typename UserClass, typename... VarTypes>
-	inline FDelegateHandle AddThreadSafeSP(UserClass* InUserObject, typename TMemFunPtrType<false, UserClass, void (ParamTypes..., VarTypes...)>::Type InFunc, VarTypes... Vars)
+	inline FDelegateHandle AddThreadSafeSP(UserClass* InUserObject, typename TMemFunPtrType<false, UserClass, void (ParamTypes..., std::decay_t<VarTypes>...)>::Type InFunc, VarTypes&&... Vars)
 	{
 		static_assert(!TIsConst<UserClass>::Value, "Attempting to bind a delegate with a const object pointer and non-const member function.");
 
-		return Add(FDelegate::CreateThreadSafeSP(InUserObject, InFunc, Vars...));
+		return Add(FDelegate::CreateThreadSafeSP(InUserObject, InFunc, Forward<VarTypes>(Vars)...));
 	}
 	template <typename UserClass, typename... VarTypes>
-	inline FDelegateHandle AddThreadSafeSP(const UserClass* InUserObject, typename TMemFunPtrType<true, UserClass, void (ParamTypes..., VarTypes...)>::Type InFunc, VarTypes... Vars)
+	inline FDelegateHandle AddThreadSafeSP(const UserClass* InUserObject, typename TMemFunPtrType<true, UserClass, void (ParamTypes..., std::decay_t<VarTypes>...)>::Type InFunc, VarTypes&&... Vars)
 	{
-		return Add(FDelegate::CreateThreadSafeSP(InUserObject, InFunc, Vars...));
+		return Add(FDelegate::CreateThreadSafeSP(InUserObject, InFunc, Forward<VarTypes>(Vars)...));
 	}
 
 	/**
@@ -872,14 +880,14 @@ public:
 	 * @param	InFunctionName			Class method function address
 	 */
 	template <typename UObjectTemplate, typename... VarTypes>
-	inline FDelegateHandle AddUFunction(UObjectTemplate* InUserObject, const FName& InFunctionName, VarTypes... Vars)
+	inline FDelegateHandle AddUFunction(UObjectTemplate* InUserObject, const FName& InFunctionName, VarTypes&&... Vars)
 	{
-		return Add(FDelegate::CreateUFunction(InUserObject, InFunctionName, Vars...));
+		return Add(FDelegate::CreateUFunction(InUserObject, InFunctionName, Forward<VarTypes>(Vars)...));
 	}
 	template <typename UObjectTemplate, typename... VarTypes>
-	inline FDelegateHandle AddUFunction(TObjectPtr<UObjectTemplate> InUserObject, const FName& InFunctionName, VarTypes... Vars)
+	inline FDelegateHandle AddUFunction(TObjectPtr<UObjectTemplate> InUserObject, const FName& InFunctionName, VarTypes&&... Vars)
 	{
-		return Add(FDelegate::CreateUFunction(InUserObject, InFunctionName, Vars...));
+		return Add(FDelegate::CreateUFunction(InUserObject, InFunctionName, Forward<VarTypes>(Vars)...));
 	}
 
 	/**
@@ -891,28 +899,28 @@ public:
 	 * @param	InFunc			Class method function address
 	 */
 	template <typename UserClass, typename... VarTypes>
-	inline FDelegateHandle AddUObject(UserClass* InUserObject, typename TMemFunPtrType<false, UserClass, void (ParamTypes..., VarTypes...)>::Type InFunc, VarTypes... Vars)
+	inline FDelegateHandle AddUObject(UserClass* InUserObject, typename TMemFunPtrType<false, UserClass, void (ParamTypes..., std::decay_t<VarTypes>...)>::Type InFunc, VarTypes&&... Vars)
 	{
 		static_assert(!TIsConst<UserClass>::Value, "Attempting to bind a delegate with a const object pointer and non-const member function.");
 
-		return Add(FDelegate::CreateUObject(InUserObject, InFunc, Vars...));
+		return Add(FDelegate::CreateUObject(InUserObject, InFunc, Forward<VarTypes>(Vars)...));
 	}
 	template <typename UserClass, typename... VarTypes>
-	inline FDelegateHandle AddUObject(const UserClass* InUserObject, typename TMemFunPtrType<true, UserClass, void (ParamTypes..., VarTypes...)>::Type InFunc, VarTypes... Vars)
+	inline FDelegateHandle AddUObject(const UserClass* InUserObject, typename TMemFunPtrType<true, UserClass, void (ParamTypes..., std::decay_t<VarTypes>...)>::Type InFunc, VarTypes&&... Vars)
 	{
-		return Add(FDelegate::CreateUObject(InUserObject, InFunc, Vars...));
+		return Add(FDelegate::CreateUObject(InUserObject, InFunc, Forward<VarTypes>(Vars)...));
 	}
 	template <typename UserClass, typename... VarTypes>
-	inline FDelegateHandle AddUObject(TObjectPtr<UserClass> InUserObject, typename TMemFunPtrType<false, UserClass, void (ParamTypes..., VarTypes...)>::Type InFunc, VarTypes... Vars)
+	inline FDelegateHandle AddUObject(TObjectPtr<UserClass> InUserObject, typename TMemFunPtrType<false, UserClass, void (ParamTypes..., std::decay_t<VarTypes>...)>::Type InFunc, VarTypes&&... Vars)
 	{
 		static_assert(!TIsConst<UserClass>::Value, "Attempting to bind a delegate with a const object pointer and non-const member function.");
 
-		return Add(FDelegate::CreateUObject(InUserObject, InFunc, Vars...));
+		return Add(FDelegate::CreateUObject(InUserObject, InFunc, Forward<VarTypes>(Vars)...));
 	}
 	template <typename UserClass, typename... VarTypes>
-	inline FDelegateHandle AddUObject(TObjectPtr<UserClass> InUserObject, typename TMemFunPtrType<true, UserClass, void (ParamTypes..., VarTypes...)>::Type InFunc, VarTypes... Vars)
+	inline FDelegateHandle AddUObject(TObjectPtr<UserClass> InUserObject, typename TMemFunPtrType<true, UserClass, void (ParamTypes..., std::decay_t<VarTypes>...)>::Type InFunc, VarTypes&&... Vars)
 	{
-		return Add(FDelegate::CreateUObject(InUserObject, InFunc, Vars...));
+		return Add(FDelegate::CreateUObject(InUserObject, InFunc, Forward<VarTypes>(Vars)...));
 	}
 
 public:
