Backout changes to Core.build.cs as it is causing issues with the Mimallo.c compile

#rnx

diff --git a/Engine/Source/Runtime/Core/Core.Build.cs b/Engine/Source/Runtime/Core/Core.Build.cs
--- a/Engine/Source/Runtime/Core/Core.Build.cs
+++ b/Engine/Source/Runtime/Core/Core.Build.cs
@@ -23,11 +23,17 @@ public class Core : ModuleRules
 		PrivateDependencyModuleNames.Add("OodleDataCompression");
 
 		PublicDependencyModuleNames.Add("TraceLog");
+		PublicIncludePaths.Add("Runtime/TraceLog/Public");
+
+		PrivateIncludePaths.AddRange(
+			new string[] {
+				"Runtime/SynthBenchmark/Public",
+				"Runtime/Engine/Public",
+			}
+			);
 
 		PrivateIncludePathModuleNames.AddRange(
 			new string[] {
-				"SynthBenchmark",
-				"Engine",
 				"TargetPlatform",
 				"DerivedDataCache",
 				"InputDevice",
