Make Windows dialog boxes resizable, auto-fit to screen, implement scrollbar, make text selectable, use proper layout with DPI scale.

diff --git a/Engine/Source/Runtime/Core/Private/Windows/WindowsPlatformMisc.cpp b/Engine/Source/Runtime/Core/Private/Windows/WindowsPlatformMisc.cpp
--- a/Engine/Source/Runtime/Core/Private/Windows/WindowsPlatformMisc.cpp
+++ b/Engine/Source/Runtime/Core/Private/Windows/WindowsPlatformMisc.cpp
@@ -1011,282 +1011,549 @@ void FWindowsPlatformMisc::CreateGuid(FGuid& Result)
 	verify( CoCreateGuid( (GUID*)&Result )==S_OK );
 }
 
+class FWindowsDialog
+{
+private:
+	/**
+	 * Calculates button position and size, localize button text.
+	 * @param HandleWnd handle to dialog window
+	 * @param Text button text to localize
+	 * @param DlgItemId dialog item id
+	 * @param PositionX current button position (x coord)
+	 * @param PositionY current button position (y coord)
+	 * @return true if succeeded
+	 */
+	static bool SetDlgItem(HWND HandleWnd, const TCHAR* Text, int32 DlgItemId, float DPIScale, int32* PositionX, int32* PositionY)
+	{
+		SIZE SizeButton;
 
-#define HOTKEY_YES			100
-#define HOTKEY_NO			101
-#define HOTKEY_CANCEL		102
+		HDC DC = CreateCompatibleDC(NULL);
+		GetTextExtentPoint32(DC, Text, wcslen(Text), &SizeButton);
+		DeleteDC(DC);
+		DC = NULL;
 
-/**
- * Helper global variables, used in MessageBoxDlgProc for set message text.
- */
-static TCHAR* GMessageBoxText = NULL;
-static TCHAR* GMessageBoxCaption = NULL;
-/**
- * Used by MessageBoxDlgProc to indicate whether a 'Cancel' button is present and
- * thus 'Esc should be accepted as a hotkey.
- */
-static bool GCancelButtonEnabled = false;
+		SizeButton.cx += (int)(14 * DPIScale);
+		SizeButton.cy += (int)(8 * DPIScale);
+		SizeButton.cx = FMath::Max((int)(73 * DPIScale), (int)SizeButton.cx);
+		SizeButton.cy = FMath::Max((int)(21 * DPIScale), (int)SizeButton.cy);
 
-/**
- * Calculates button position and size, localize button text.
- * @param HandleWnd handle to dialog window
- * @param Text button text to localize
- * @param DlgItemId dialog item id
- * @param PositionX current button position (x coord)
- * @param PositionY current button position (y coord)
- * @return true if succeeded
- */
-static bool SetDlgItem( HWND HandleWnd, const TCHAR* Text, int32 DlgItemId, int32* PositionX, int32* PositionY )
-{
-	SIZE SizeButton;
-		
-	HDC DC = CreateCompatibleDC( NULL );
-	GetTextExtentPoint32( DC, Text, wcslen(Text), &SizeButton );
-	DeleteDC(DC);
-	DC = NULL;
+		HWND Handle = GetDlgItem(HandleWnd, DlgItemId);
+		if (Handle)
+		{
+			*PositionX -= (SizeButton.cx + (int)(7 * DPIScale));
+			SetWindowPos(Handle, HWND_TOP, *PositionX, *PositionY - SizeButton.cy, SizeButton.cx, SizeButton.cy, 0);
+			SetDlgItemText(HandleWnd, DlgItemId, Text);
 
-	SizeButton.cx += 14;
-	SizeButton.cy += 8;
+			return true;
+		}
 
-	HWND Handle = GetDlgItem( HandleWnd, DlgItemId );
-	if( Handle )
+		return false;
+	}
+
+	static float MessageBoxDlgGetDPI(HWND HandleWnd)
 	{
-		*PositionX -= ( SizeButton.cx + 5 );
-		SetWindowPos( Handle, HWND_TOP, *PositionX, *PositionY - SizeButton.cy, SizeButton.cx, SizeButton.cy, 0 );
-		SetDlgItemText( HandleWnd, DlgItemId, Text );
-		
-		return true;
+		float DPIScale = 1.0f;
+
+		typedef UINT(WINAPI* LPGetDpiForWindow)(HWND Hwnd);
+		LPGetDpiForWindow GetDpiForWindow = (LPGetDpiForWindow)(void*)GetProcAddress(GetModuleHandle(L"user32.dll"), "GetDpiForWindow");
+		if (GetDpiForWindow)
+		{
+			DPIScale = GetDpiForWindow(HandleWnd) / 96.0f;
+		}
+
+		return DPIScale;
 	}
 
-	return false;
-}
+	static void SetWindowStyleFlags(HWND HandleWnd, LONG Flags, bool bEnabled)
+	{
+		LONG Style = GetWindowLong(HandleWnd, GWL_STYLE);
+		if (bEnabled)
+		{
+			Style |= Flags;
+		}
+		else
+		{
+			Style &= ~Flags;
+		}
+		SetWindowLong(HandleWnd, GWL_STYLE, Style);
+	}
 
-/**
- * Callback for MessageBoxExt dialog (allowing for Yes to all / No to all )
- * @return		One of EAppReturnType::Yes, EAppReturnType::YesAll, EAppReturnType::No, EAppReturnType::NoAll, EAppReturnType::Cancel.
- */
-PTRINT CALLBACK MessageBoxDlgProc( HWND HandleWnd, uint32 Message, WPARAM WParam, LPARAM LParam )
-{
-	switch(Message)
+	static SIZE MeasureText(HWND HandleWnd, LPCWSTR Text, int TextLength)
 	{
-		case WM_INITDIALOG:
-			{
-				// Sets most bottom and most right position to begin button placement
-				RECT Rect;
-				POINT Point;
-				
-				GetWindowRect( HandleWnd, &Rect );
-				Point.x = Rect.right;
-				Point.y = Rect.bottom;
-				ScreenToClient( HandleWnd, &Point );
-				
-				int32 PositionX = Point.x - 8;
-				int32 PositionY = Point.y - 10;
-
-				// Localize dialog buttons, sets position and size.
-				FString CancelString;
-				FString NoToAllString;
-				FString NoString;
-				FString YesToAllString;
-				FString YesString;
-
-				// The Localize* functions will return the Key if a dialog is presented before the config system is initialized.
-				// Instead, we use hard-coded strings if config is not yet initialized.
-				if( !GConfig )
-				{
-					CancelString = TEXT("Cancel");
-					NoToAllString = TEXT("No to All");
-					NoString = TEXT("No");
-					YesToAllString = TEXT("Yes to All");
-					YesString = TEXT("Yes");
-				}
-				else
-				{
-					CancelString = NSLOCTEXT("UnrealEd", "Cancel", "Cancel").ToString();
-					NoToAllString = NSLOCTEXT("UnrealEd", "NoToAll", "No to All").ToString();
-					NoString = NSLOCTEXT("UnrealEd", "No", "No").ToString();
-					YesToAllString = NSLOCTEXT("UnrealEd", "YesToAll", "Yes to All").ToString();
-					YesString = NSLOCTEXT("UnrealEd", "Yes", "Yes").ToString();
-				}
-				SetDlgItem( HandleWnd, *CancelString, IDC_CANCEL, &PositionX, &PositionY );
-				SetDlgItem( HandleWnd, *NoToAllString, IDC_NOTOALL, &PositionX, &PositionY );
-				SetDlgItem( HandleWnd, *NoString, IDC_NO_B, &PositionX, &PositionY );
-				SetDlgItem( HandleWnd, *YesToAllString, IDC_YESTOALL, &PositionX, &PositionY );
-				SetDlgItem( HandleWnd, *YesString, IDC_YES, &PositionX, &PositionY );
-
-				SetDlgItemText( HandleWnd, IDC_MESSAGE, GMessageBoxText );
-				SetWindowText( HandleWnd, GMessageBoxCaption );
-
-				// If parent window exist, get it handle and make it foreground.
-				HWND ParentWindow = GetTopWindow( HandleWnd );
-				if( ParentWindow )
-				{
-					SetWindowPos( ParentWindow, HWND_TOPMOST, 0, 0, 0, 0, SWP_NOMOVE | SWP_NOSIZE );
-				}
+		HDC DC = CreateCompatibleDC(NULL);
+		HFONT Font = (HFONT)SendMessageW(HandleWnd, WM_GETFONT, 0, 0);
+		SelectObject(DC, Font);
+		RECT TextRect{};
+		DrawTextW(DC, Text, TextLength, &TextRect, DT_CALCRECT | DT_EDITCONTROL | DT_LEFT | DT_EXPANDTABS);
+		DeleteDC(DC);
 
-				SetForegroundWindow( HandleWnd );
-				SetWindowPos( HandleWnd, HWND_TOPMOST, 0, 0, 0, 0, SWP_NOMOVE | SWP_NOSIZE );
+		SIZE TextSize;
+		TextSize.cx = TextRect.right - TextRect.left;
+		TextSize.cy = TextRect.bottom - TextRect.top;
+		return TextSize;
+	}
 
-				RegisterHotKey( HandleWnd, HOTKEY_YES, 0, 'Y' );
-				RegisterHotKey( HandleWnd, HOTKEY_NO, 0, 'N' );
-				if ( GCancelButtonEnabled )
-				{
-					RegisterHotKey( HandleWnd, HOTKEY_CANCEL, 0, VK_ESCAPE );
-				}
+	void UpdateEditTextScrollbar(HWND HandleWnd, float DPIScale)
+	{
+		RECT MessageRect;
+		GetWindowRect(HandleWnd, &MessageRect);
+		SIZE MessageSize;
+		MessageSize.cx = MessageRect.right - MessageRect.left - (int)(20 * DPIScale);
+		MessageSize.cy = MessageRect.bottom - MessageRect.top - (int)(16 * DPIScale);
 
-				// Windows are foreground, make them not top most.
-				SetWindowPos( HandleWnd, HWND_NOTOPMOST, 0, 0, 0, 0, SWP_NOMOVE | SWP_NOSIZE );
-				if( ParentWindow )
-				{
-					SetWindowPos( ParentWindow, HWND_NOTOPMOST, 0, 0, 0, 0, SWP_NOMOVE | SWP_NOSIZE );
-				}
+		SIZE TextSize = MeasureText(HandleWnd, *Text, Text.Len());
+
+		bool bNeedsHScroll = TextSize.cx > MessageSize.cx;
+		SetWindowStyleFlags(HandleWnd, WS_HSCROLL, bNeedsHScroll);
+		bool bNeedsVScroll = TextSize.cy > MessageSize.cy;
+		SetWindowStyleFlags(HandleWnd, WS_VSCROLL, bNeedsVScroll);
+	}
 
+	void Close(EAppReturnType::Type NewResult)
+	{
+		WasClosed = true;
+		Result = NewResult;
+		DestroyWindow(DialogHwnd);
+	}
+
+	bool OnInitDialog()
+	{
+		// Sets most bottom and most right position to begin button placement
+		POINT Point;
+
+		GetWindowRect(DialogHwnd, &DefaultWindowRect);
+		GetClientRect(DialogHwnd, &ClientRect);
+		WasClosed = false;
+
+		Point.x = ClientRect.right;
+		Point.y = ClientRect.bottom;
+
+		float DPIScale = MessageBoxDlgGetDPI(DialogHwnd);
+
+		int32 PositionX = Point.x - (int)(5 * DPIScale);
+		int32 PositionY = Point.y - (int)(10 * DPIScale);
+
+		// Localize dialog buttons, sets position and size.
+		FString CancelString;
+		FString RetryString;
+		FString ContinueString;
+		FString NoToAllString;
+		FString NoString;
+		FString YesToAllString;
+		FString YesString;
+		FString OKString;
+
+		// The Localize* functions will return the Key if a dialog is presented before the config system is initialized.
+		// Instead, we use hard-coded strings if config is not yet initialized.
+		if (!GConfig)
+		{
+			CancelString = TEXT("Cancel");
+			RetryString = TEXT("Retry");
+			ContinueString = TEXT("Continue");
+			NoToAllString = TEXT("No to All");
+			NoString = TEXT("No");
+			YesToAllString = TEXT("Yes to All");
+			YesString = TEXT("Yes");
+			OKString = TEXT("OK");
+		}
+		else
+		{
+			CancelString = NSLOCTEXT("UnrealEd", "Cancel", "Cancel").ToString();
+			RetryString = NSLOCTEXT("UnrealEd", "Retry", "Retry").ToString();
+			ContinueString = NSLOCTEXT("UnrealEd", "Continue", "Continue").ToString();
+			NoToAllString = NSLOCTEXT("UnrealEd", "NoToAll", "No to All").ToString();
+			NoString = NSLOCTEXT("UnrealEd", "No", "No").ToString();
+			YesToAllString = NSLOCTEXT("UnrealEd", "YesToAll", "Yes to All").ToString();
+			YesString = NSLOCTEXT("UnrealEd", "Yes", "Yes").ToString();
+			OKString = NSLOCTEXT("UnrealEd", "OK", "OK").ToString();
+		}
+		SetDlgItem(DialogHwnd, *ContinueString, IDC_CONTINUE, DPIScale, &PositionX, &PositionY);
+		SetDlgItem(DialogHwnd, *RetryString, IDC_RETRY, DPIScale, &PositionX, &PositionY);
+		SetDlgItem(DialogHwnd, *CancelString, IDC_CANCEL, DPIScale, &PositionX, &PositionY);
+		SetDlgItem(DialogHwnd, *NoToAllString, IDC_NOTOALL, DPIScale, &PositionX, &PositionY);
+		SetDlgItem(DialogHwnd, *NoString, IDC_NO_B, DPIScale, &PositionX, &PositionY);
+		SetDlgItem(DialogHwnd, *YesToAllString, IDC_YESTOALL, DPIScale, &PositionX, &PositionY);
+		SetDlgItem(DialogHwnd, *YesString, IDC_YES, DPIScale, &PositionX, &PositionY);
+		SetDlgItem(DialogHwnd, *OKString, IDC_OK, DPIScale, &PositionX, &PositionY);
+
+		SetDlgItemText(DialogHwnd, IDC_MESSAGE, *Text);
+		SetWindowText(DialogHwnd, *Caption);
+
+		HWND MessageHandle = GetDlgItem(DialogHwnd, IDC_MESSAGE);
+		UpdateEditTextScrollbar(MessageHandle, DPIScale);
+
+		// If parent window exist, get it handle and make it foreground.
+		HWND ParentWindow = GetTopWindow(DialogHwnd);
+		if (ParentWindow)
+		{
+			SetWindowPos(ParentWindow, HWND_TOPMOST, 0, 0, 0, 0, SWP_NOMOVE | SWP_NOSIZE);
+		}
+
+		SetForegroundWindow(DialogHwnd);
+		SetWindowPos(DialogHwnd, HWND_TOPMOST, 0, 0, 0, 0, SWP_NOMOVE | SWP_NOSIZE);
+
+
+		// Windows are foreground, make them not top most.
+		SetWindowPos(DialogHwnd, HWND_NOTOPMOST, 0, 0, 0, 0, SWP_NOMOVE | SWP_NOSIZE);
+		if (ParentWindow)
+		{
+			SetWindowPos(ParentWindow, HWND_NOTOPMOST, 0, 0, 0, 0, SWP_NOMOVE | SWP_NOSIZE);
+		}
+
+		// Resize to fit text
+		int Width = DefaultWindowRect.right - DefaultWindowRect.left;
+		int Height = DefaultWindowRect.bottom - DefaultWindowRect.top;
+		SIZE TextSize = MeasureText(MessageHandle, *Text, Text.Len());
+
+		HMONITOR Monitor = MonitorFromWindow(DialogHwnd, MONITOR_DEFAULTTONEAREST);
+		MONITORINFO MonInfo;
+		MonInfo.cbSize = sizeof(MONITORINFO);
+		GetMonitorInfo(Monitor, &MonInfo);
+		int MonitorWidth = MonInfo.rcMonitor.right - MonInfo.rcMonitor.left;
+		int MonitorHeight = MonInfo.rcMonitor.bottom - MonInfo.rcMonitor.top;
+
+		int NewWidth = FMath::Clamp((int)(TextSize.cx + 100 * DPIScale), Width, (int)(MonitorWidth * 0.8f));
+		int NewHeight = FMath::Clamp((int)(TextSize.cy + 165 * DPIScale), Height, (int)(MonitorHeight * 0.8f));
+		int NewLeft = DefaultWindowRect.left - (NewWidth - Width) / 2;
+		int NewTop = DefaultWindowRect.top - (NewHeight - Height) / 2;
+		SetWindowPos(DialogHwnd, HWND_NOTOPMOST, NewLeft, NewTop, NewWidth, NewHeight, SWP_NOZORDER);
+
+		return true;
+	}
+
+	/**
+	 * Callback for MessageBoxExt dialog (allowing for Yes to all / No to all )
+	 * @return		One of EAppReturnType::Yes, EAppReturnType::YesAll, EAppReturnType::No, EAppReturnType::NoAll, EAppReturnType::Cancel.
+	 */
+	PTRINT MessageBoxDlgProc(HWND HandleWnd, uint32 Message, WPARAM WParam, LPARAM LParam)
+	{
+		switch (Message)
+		{
+			case WM_INITDIALOG:
+			{
+				return OnInitDialog();
+			}
+			case WM_DESTROY:
+			{
 				return true;
 			}
-		case WM_DESTROY:
+			case WM_CLOSE:
 			{
-				UnregisterHotKey( HandleWnd, HOTKEY_YES );
-				UnregisterHotKey( HandleWnd, HOTKEY_NO );
-				if ( GCancelButtonEnabled )
-				{
-					UnregisterHotKey( HandleWnd, HOTKEY_CANCEL );
-				}
+				Close(Result);
 				return true;
 			}
-		case WM_COMMAND:
-			switch( LOWORD( WParam ) )
+			case WM_COMMAND:
 			{
-				case IDC_YES:
-					EndDialog( HandleWnd, EAppReturnType::Yes );
-					break;
-				case IDC_YESTOALL:
-					EndDialog( HandleWnd, EAppReturnType::YesAll );
-					break;
-				case IDC_NO_B:
-					EndDialog( HandleWnd, EAppReturnType::No );
-					break;
-				case IDC_NOTOALL:
-					EndDialog( HandleWnd, EAppReturnType::NoAll );
-					break;
-				case IDC_CANCEL:
-					if ( GCancelButtonEnabled )
+				switch (LOWORD(WParam))
+				{
+					case IDC_OK:
+						Close(EAppReturnType::Ok);
+						break;
+					case IDC_YES:
+						Close(EAppReturnType::Yes);
+						break;
+					case IDC_YESTOALL:
+						Close(EAppReturnType::YesAll);
+						break;
+					case IDC_NO_B:
+						Close(EAppReturnType::No);
+						break;
+					case IDC_NOTOALL:
+						Close(EAppReturnType::NoAll);
+						break;
+					case IDC_RETRY:
+						Close(EAppReturnType::Retry);
+						break;
+					case IDC_CONTINUE:
+						Close(EAppReturnType::Continue);
+						break;
+					case IDC_CANCEL:
+						if (CancelButtonEnabled)
+						{
+							Close(EAppReturnType::Cancel);
+							break;
+						}
+						break;
+					case IDC_COPY:
 					{
-						EndDialog( HandleWnd, EAppReturnType::Cancel );
+						size_t TextByteCount = Text.GetAllocatedSize();
+						HGLOBAL StrClipboardMemory = GlobalAlloc(GMEM_MOVEABLE, TextByteCount);
+
+						void* StrTarget = GlobalLock(StrClipboardMemory);
+						memcpy(StrTarget, *Text, TextByteCount);
+						((char*)StrTarget)[TextByteCount] = 0;
+						GlobalUnlock(StrClipboardMemory);
+
+						OpenClipboard(HandleWnd);
+						EmptyClipboard();
+						SetClipboardData(CF_UNICODETEXT, StrClipboardMemory);
+						CloseClipboard();
 					}
 					break;
+					default:
+						return false;
+				}
+				return true;
 			}
-			break;
-		case WM_HOTKEY:
-			switch( WParam )
+			case WM_CTLCOLORSTATIC:
 			{
-			case HOTKEY_YES:
-				EndDialog( HandleWnd, EAppReturnType::Yes );
-				break;
-			case HOTKEY_NO:
-				EndDialog( HandleWnd, EAppReturnType::No );
-				break;
-			case HOTKEY_CANCEL:
-				if ( GCancelButtonEnabled )
+				if ((HWND)LParam == GetDlgItem(HandleWnd, IDC_MESSAGE))
 				{
-					EndDialog( HandleWnd, EAppReturnType::Cancel );
+					SetBkMode((HDC)WParam, TRANSPARENT);
+					return (LRESULT)(GetSysColorBrush(COLOR_WINDOW));
+				}
+				else
+				{
+					return false;
 				}
-				break;
 			}
 			break;
-		default:
-			return false;
-	}
-	return true;
-}
+			case WM_PAINT:
+			{
+				PAINTSTRUCT Paint;
+				HDC Hdc = BeginPaint(HandleWnd, &Paint);
+
+				float DPIScale = MessageBoxDlgGetDPI(HandleWnd);
+
+				RECT MainAreaRect;
+				MainAreaRect = ClientRect;
+				const int ActionBarHeight = (int)(45 * DPIScale);
+				MainAreaRect.bottom -= ActionBarHeight;
+				FillRect(Hdc, &MainAreaRect, GetSysColorBrush(COLOR_WINDOW));
+
+				/*HWND MessageHandle = GetDlgItem(HandleWnd, IDC_MESSAGE);
+				RECT MessageRect;
+				GetWindowRect(MessageHandle, &MessageRect);
+				MessageRect.right = MessageRect.right - MessageRect.left;
+				MessageRect.bottom = MessageRect.bottom - MessageRect.top;
+				MessageRect.left = MessageRect.top = 15;
+				HFONT Font = (HFONT)SendMessageW(MessageHandle, WM_GETFONT, 0, 0);
+				SelectObject(Hdc, Font);
+				DrawTextW(Hdc, GDialogState.Text, wcslen(GDialogState.Text), &MainAreaRect, DT_EDITCONTROL | DT_LEFT | DT_EXPANDTABS);*/
+
+				EndPaint(HandleWnd, &Paint);
+				return false;
+			}
+			case WM_SIZE:
+			{
+				RECT PrevDialogRect = ClientRect;
+				GetClientRect(HandleWnd, &ClientRect);
 
-/**
- * Displays extended message box allowing for YesAll/NoAll
- * @return 3 - YesAll, 4 - NoAll, -1 for Fail
- */
-int MessageBoxExtInternal( EAppMsgType::Type MsgType, HWND HandleWnd, const TCHAR* Text, const TCHAR* Caption )
-{
-	GMessageBoxText = (TCHAR *) Text;
-	GMessageBoxCaption = (TCHAR *) Caption;
+				int ControlsStretch[] =
+				{
+					IDC_MESSAGE
+				};
 
-	switch (MsgType)
-	{
-		case EAppMsgType::YesNoYesAllNoAll:
-		{
-			GCancelButtonEnabled = false;
-			return (int)DialogBox(GetModuleHandle(NULL), MAKEINTRESOURCE(IDD_YESNO2ALL), HandleWnd, MessageBoxDlgProc);
-		}
-		case EAppMsgType::YesNoYesAllNoAllCancel:
-		{
-			GCancelButtonEnabled = true;
-			return (int)DialogBox(GetModuleHandle(NULL), MAKEINTRESOURCE(IDD_YESNO2ALLCANCEL), HandleWnd, MessageBoxDlgProc);
-		}
-		case EAppMsgType::YesNoYesAll:
-		{
-			GCancelButtonEnabled = false;
-			return (int)DialogBox(GetModuleHandle(NULL), MAKEINTRESOURCE(IDD_YESNOYESTOALL), HandleWnd, MessageBoxDlgProc);
+				for (const int ControlId : ControlsStretch)
+				{
+					HWND ControlHandle = GetDlgItem(HandleWnd, ControlId);
+					if (ControlHandle == nullptr) { continue; }
+
+					RECT PrevMessageRect;
+					GetWindowRect(ControlHandle, &PrevMessageRect);
+					ScreenToClient(HandleWnd, (POINT*)&PrevMessageRect.left);
+					ScreenToClient(HandleWnd, (POINT*)&PrevMessageRect.right);
+
+					SetWindowPos(ControlHandle, 0,
+						PrevMessageRect.left,
+						PrevMessageRect.top,
+						(PrevMessageRect.right - PrevMessageRect.left) + (ClientRect.right - PrevDialogRect.right),
+						(PrevMessageRect.bottom - PrevMessageRect.top) + (ClientRect.bottom - PrevDialogRect.bottom), SWP_NOZORDER);
+				}
+
+				int ControlsBottomRight[] =
+				{
+					IDC_OK, IDC_YES, IDC_YESTOALL, IDC_NO_B, IDC_NOTOALL, IDC_CANCEL, IDC_RETRY, IDC_CONTINUE
+				};
+
+				for (const int ControlId : ControlsBottomRight)
+				{
+					HWND ControlHandle = GetDlgItem(HandleWnd, ControlId);
+					if (ControlHandle == nullptr) { continue; }
+
+					RECT PrevMessageRect;
+					GetWindowRect(ControlHandle, &PrevMessageRect);
+					ScreenToClient(HandleWnd, (POINT*)&PrevMessageRect.left);
+					ScreenToClient(HandleWnd, (POINT*)&PrevMessageRect.right);
+
+					SetWindowPos(ControlHandle, 0,
+						PrevMessageRect.left + (ClientRect.right - PrevDialogRect.right),
+						PrevMessageRect.top + (ClientRect.bottom - PrevDialogRect.bottom),
+						(PrevMessageRect.right - PrevMessageRect.left),
+						(PrevMessageRect.bottom - PrevMessageRect.top), SWP_NOZORDER);
+				}
+
+				float DPIScale = MessageBoxDlgGetDPI(HandleWnd);
+
+				HWND MessageHandle = GetDlgItem(HandleWnd, IDC_MESSAGE);
+				UpdateEditTextScrollbar(MessageHandle, DPIScale);
+
+				InvalidateRect(HandleWnd, nullptr, true);
+			}
+			return true;
+			case WM_GETMINMAXINFO:
+			{
+				MINMAXINFO* Info = (MINMAXINFO*)LParam;
+
+				Info->ptMinTrackSize.x = DefaultWindowRect.right - DefaultWindowRect.left;
+				Info->ptMinTrackSize.y = DefaultWindowRect.bottom - DefaultWindowRect.top;
+			}
+			return true;
+			case WM_DPICHANGED:
+				/*{
+					float NewScale = LOWORD(WParam) / 96.0f;
+					RECT* SuggestedRect = (RECT*)LParam;
+					SetWindowPos(HandleWnd,
+						HWND_NOTOPMOST,
+						SuggestedRect->left,
+						SuggestedRect->top,
+						SuggestedRect->right - SuggestedRect->left,
+						SuggestedRect->bottom - SuggestedRect->top,
+						SWP_NOZORDER);
+
+					//[REDACTED_PROFILE_URL]/Windows-classic-samples/blob/main/Samples/DPIAwarenessPerWindow/client/DpiAwarenessContext.cpp
+
+				}*/
+				return true;
+			default:
+				return false;
 		}
 	}
 
-	return -1;
-}
+	FWindowsDialog() {}
 
 
+	FString Text;
+	FString Caption;
 
+	/**
+	 * Used to indicate whether a 'Cancel' button is present and
+	 * thus 'Esc should be accepted as a hotkey.
+	 */
+	bool CancelButtonEnabled = false;
 
-EAppReturnType::Type FWindowsPlatformMisc::MessageBoxExt( EAppMsgType::Type MsgType, const TCHAR* Text, const TCHAR* Caption )
-{
-	FSlowHeartBeatScope SuspendHeartBeat;
 
-	HWND ParentWindow = (HWND)NULL;
-	switch( MsgType )
+	HWND DialogHwnd = NULL;
+	RECT ClientRect = {};
+	RECT DefaultWindowRect = {};
+	bool WasClosed = false;
+	EAppReturnType::Type Result = EAppReturnType::Cancel;
+
+	EAppReturnType::Type Show(HWND ParentWindowHandle, int Template)
 	{
-	case EAppMsgType::Ok:
-		{
-			MessageBox(ParentWindow, Text, Caption, MB_OK|MB_SYSTEMMODAL);
-			return EAppReturnType::Ok;
-		}
-	case EAppMsgType::YesNo:
-		{
-			int32 Return = MessageBox( ParentWindow, Text, Caption, MB_YESNO|MB_SYSTEMMODAL );
-			return Return == IDYES ? EAppReturnType::Yes : EAppReturnType::No;
-		}
-	case EAppMsgType::OkCancel:
-		{
-			int32 Return = MessageBox( ParentWindow, Text, Caption, MB_OKCANCEL|MB_SYSTEMMODAL );
-			return Return == IDOK ? EAppReturnType::Ok : EAppReturnType::Cancel;
-		}
-	case EAppMsgType::YesNoCancel:
+		HACCEL AcceleratorHandle = LoadAccelerators(hInstance, MAKEINTRESOURCE(IDR_ACCEL1));
+
+		DialogHwnd = CreateDialogParam(GetModuleHandle(NULL), MAKEINTRESOURCE(Template), ParentWindowHandle, 
+			[](HWND HandleWnd, uint32 Message, WPARAM WParam, LPARAM LParam) {
+				FWindowsDialog* Instance;
+				if (Message == WM_INITDIALOG)
+				{
+					Instance = (FWindowsDialog*)LParam;
+					Instance->DialogHwnd = HandleWnd;
+					SetWindowLongPtr(HandleWnd, DWLP_USER, (LPARAM)Instance);
+				}
+				else
+				{
+					Instance = (FWindowsDialog*)GetWindowLongPtr(HandleWnd, DWLP_USER);
+				}
+
+				return Instance->MessageBoxDlgProc(HandleWnd, Message, WParam, LParam);
+			}, (LPARAM)this);
+
+		ShowWindow(DialogHwnd, SW_SHOW);
+		MSG Msg;
+		while (!WasClosed && GetMessageW(&Msg, NULL, 0, 0))
 		{
-			int32 Return = MessageBox(ParentWindow, Text, Caption, MB_YESNOCANCEL | MB_ICONQUESTION | MB_SYSTEMMODAL);
-			return Return == IDYES ? EAppReturnType::Yes : (Return == IDNO ? EAppReturnType::No : EAppReturnType::Cancel);
+			bool bIsTextEditMsg = Msg.hwnd == GetDlgItem(DialogHwnd, IDC_MESSAGE);
+			bool bCheckAccelerators = DialogHwnd == Msg.hwnd || (IsChild(DialogHwnd, Msg.hwnd) && !bIsTextEditMsg);
+			if (!(bCheckAccelerators && TranslateAccelerator(DialogHwnd, AcceleratorHandle, &Msg)) &&
+				!IsDialogMessage(DialogHwnd, &Msg))
+			{
+				TranslateMessage(&Msg);
+				DispatchMessage(&Msg);
+			}
 		}
-	case EAppMsgType::CancelRetryContinue:
+
+		DestroyAcceleratorTable(AcceleratorHandle);
+
+		return Result;
+	}
+
+public:
+	static EAppReturnType::Type Show(HWND ParentWindowHandle, EAppMsgType::Type MsgType, const FString& Text, const FString& Caption)
+	{
+		FWindowsDialog Instance {};
+		Instance.Text = Text;
+		Instance.Caption = Caption;
+
+		switch (MsgType)
 		{
-			int32 Return = MessageBox(ParentWindow, Text, Caption, MB_CANCELTRYCONTINUE | MB_ICONQUESTION | MB_DEFBUTTON2 | MB_SYSTEMMODAL);
-			return Return == IDCANCEL ? EAppReturnType::Cancel : (Return == IDTRYAGAIN ? EAppReturnType::Retry : EAppReturnType::Continue);
+			case EAppMsgType::Ok:
+			{
+				Instance.CancelButtonEnabled = false;
+				Instance.Result = EAppReturnType::Ok; // default option used when dialog is closed
+				return Instance.Show(ParentWindowHandle, IDD_OK);
+			}
+			case EAppMsgType::YesNo:
+			{
+				Instance.CancelButtonEnabled = false;
+				Instance.Result = EAppReturnType::No;
+				return Instance.Show(ParentWindowHandle, IDD_YESNO);
+			}
+			case EAppMsgType::OkCancel:
+			{
+				Instance.CancelButtonEnabled = true;
+				Instance.Result = EAppReturnType::Cancel;
+				return Instance.Show(ParentWindowHandle, IDD_OKCANCEL);
+			}
+			case EAppMsgType::YesNoCancel:
+			{
+				Instance.CancelButtonEnabled = true;
+				Instance.Result = EAppReturnType::Cancel;
+				return Instance.Show(ParentWindowHandle, IDD_YESNOCANCEL);
+			}
+			case EAppMsgType::CancelRetryContinue:
+			{
+				Instance.CancelButtonEnabled = true;
+				Instance.Result = EAppReturnType::Cancel;
+				return Instance.Show(ParentWindowHandle, IDD_CANCELRETRYCONTINUE);
+			}
+			case EAppMsgType::YesNoYesAllNoAll:
+			{
+				Instance.CancelButtonEnabled = false;
+				Instance.Result = EAppReturnType::No;
+				return Instance.Show(ParentWindowHandle, IDD_YESNO2ALL);
+			}
+			case EAppMsgType::YesNoYesAllNoAllCancel:
+			{
+				Instance.CancelButtonEnabled = true;
+				Instance.Result = EAppReturnType::Cancel;
+				return Instance.Show(ParentWindowHandle, IDD_YESNO2ALLCANCEL);
+			}
+			case EAppMsgType::YesNoYesAll:
+			{
+				Instance.CancelButtonEnabled = false;
+				Instance.Result = EAppReturnType::No;
+				return Instance.Show(ParentWindowHandle, IDD_YESNOYESTOALL);
+			}
+			default:
+				return EAppReturnType::Cancel;
 		}
-		break;
-	case EAppMsgType::YesNoYesAllNoAll:
-		return (EAppReturnType::Type)MessageBoxExtInternal( EAppMsgType::YesNoYesAllNoAll, ParentWindow, Text, Caption );
-		//These return codes just happen to match up with ours.
-		// return 0 for No, 1 for Yes, 2 for YesToAll, 3 for NoToAll
-		break;
-	case EAppMsgType::YesNoYesAllNoAllCancel:
-		return (EAppReturnType::Type)MessageBoxExtInternal( EAppMsgType::YesNoYesAllNoAllCancel, ParentWindow, Text, Caption );
-		//These return codes just happen to match up with ours.
-		// return 0 for No, 1 for Yes, 2 for YesToAll, 3 for NoToAll, 4 for Cancel
-		break;
+	}
+};
 
-	case EAppMsgType::YesNoYesAll:
-		return (EAppReturnType::Type)MessageBoxExtInternal(EAppMsgType::YesNoYesAll, ParentWindow, Text, Caption);
-		//These return codes just happen to match up with ours.
-		// return 0 for No, 1 for Yes, 2 for YesToAll
-		break;
 
-	default:
-		break;
-	}
-	return EAppReturnType::Cancel;
+EAppReturnType::Type FWindowsPlatformMisc::MessageBoxExt( EAppMsgType::Type MsgType, const TCHAR* Text, const TCHAR* Caption )
+{
+	FSlowHeartBeatScope SuspendHeartBeat;
+
+	HWND ParentWindow = (HWND)NULL;
+	
+	FString PlatformText = FString(Text);
+	PlatformText.ReplaceInline(TEXT("\r"), TEXT(""));
+	PlatformText.ReplaceInline(TEXT("\n"), TEXT("\r\n"));
+
+	return FWindowsDialog::Show(ParentWindow, MsgType, PlatformText, FString(Caption));
 }
 
 static bool HandleGameExplorerIntegration()
diff --git a/Engine/Source/Runtime/Launch/Resources/Windows/PCLaunch.rc b/Engine/Source/Runtime/Launch/Resources/Windows/PCLaunch.rc
--- a/Engine/Source/Runtime/Launch/Resources/Windows/PCLaunch.rc
+++ b/Engine/Source/Runtime/Launch/Resources/Windows/PCLaunch.rc
@@ -24,42 +24,108 @@ LANGUAGE LANG_ENGLISH, SUBLANG_ENGLISH_US
 // Dialog
 //
 
-IDD_YESNO2ALL DIALOGEX 0, 0, 275, 118
-STYLE DS_SETFONT | DS_MODALFRAME | DS_FIXEDSYS | DS_CENTER | WS_POPUP | WS_CAPTION | WS_SYSMENU
+IDD_OK DIALOGEX 0, 0, 200, 80
+STYLE DS_ABSALIGN | DS_SYSMODAL | DS_SETFONT | DS_MODALFRAME | DS_FIXEDSYS | DS_CENTER | WS_POPUP | WS_CLIPCHILDREN | WS_CAPTION | WS_SYSMENU | WS_THICKFRAME
 CAPTION "Dialog"
 FONT 8, "MS Shell Dlg", 400, 0, 0x1
 BEGIN
-	DEFPUSHBUTTON   "Yes",IDC_YES,25,97,50,14
-	PUSHBUTTON      "No",IDC_NO_B,141,97,50,14
-	PUSHBUTTON      "Yes to all",IDC_YESTOALL,83,97,50,14
-	PUSHBUTTON      "No to all",IDC_NOTOALL,199,97,50,14
-    LTEXT           "Static",IDC_MESSAGE,17,19,246,66
+    EDITTEXT        IDC_MESSAGE,15,15,170,20,ES_MULTILINE | ES_AUTOVSCROLL | ES_AUTOHSCROLL | ES_READONLY | NOT WS_BORDER | WS_VSCROLL | WS_HSCROLL | 0x200
+    DEFPUSHBUTTON   "OK",IDC_OK,150,50,50,14
 END
 
-IDD_YESNO2ALLCANCEL DIALOGEX 0, 0, 297, 126
-STYLE DS_SETFONT | DS_MODALFRAME | DS_FIXEDSYS | DS_CENTER | WS_POPUP | WS_CAPTION | WS_SYSMENU
+IDD_YESNO DIALOGEX 0, 0, 200, 80
+STYLE DS_ABSALIGN | DS_SYSMODAL | DS_SETFONT | DS_MODALFRAME | DS_FIXEDSYS | DS_CENTER | WS_POPUP | WS_CLIPCHILDREN | WS_CAPTION | WS_SYSMENU | WS_THICKFRAME
 CAPTION "Dialog"
 FONT 8, "MS Shell Dlg", 400, 0, 0x1
 BEGIN
-    PUSHBUTTON      "Yes",IDC_YES,15,98,50,14
-    PUSHBUTTON      "Yes to all",IDC_YESTOALL,70,98,50,14
-    PUSHBUTTON      "No",IDC_NO_B,124,98,50,14
-    PUSHBUTTON      "No to all",IDC_NOTOALL,177,98,50,14
-    PUSHBUTTON      "Cancel",IDC_CANCEL,232,98,50,14
-    LTEXT           "Static",IDC_MESSAGE,16,14,265,66
+    EDITTEXT        IDC_MESSAGE,15,15,170,20,ES_MULTILINE | ES_AUTOVSCROLL | ES_AUTOHSCROLL | ES_READONLY | NOT WS_BORDER | WS_VSCROLL | WS_HSCROLL | 0x200
+    DEFPUSHBUTTON   "Yes",IDC_YES,100,50,50,14
+    PUSHBUTTON      "No",IDC_NO_B,150,50,50,14
 END
 
-IDD_YESNOYESTOALL DIALOGEX 0, 0, 297, 126
-STYLE DS_SETFONT | DS_MODALFRAME | DS_FIXEDSYS | DS_CENTER | WS_POPUP | WS_CAPTION | WS_SYSMENU
+IDD_OKCANCEL DIALOGEX 0, 0, 200, 80
+STYLE DS_ABSALIGN | DS_SYSMODAL | DS_SETFONT | DS_MODALFRAME | DS_FIXEDSYS | DS_CENTER | WS_POPUP | WS_CLIPCHILDREN | WS_CAPTION | WS_SYSMENU | WS_THICKFRAME
 CAPTION "Dialog"
 FONT 8, "MS Shell Dlg", 400, 0, 0x1
 BEGIN
-    PUSHBUTTON      "Yes",IDC_YES,69,98,50,14
-    PUSHBUTTON      "No",IDC_NO_B,124,98,50,14
-    PUSHBUTTON      "Yes to all",IDC_YESTOALL,178,98,50,14
-    LTEXT           "Static",IDC_MESSAGE,16,14,265,66
+    EDITTEXT        IDC_MESSAGE,15,15,170,20,ES_MULTILINE | ES_AUTOVSCROLL | ES_AUTOHSCROLL | ES_READONLY | NOT WS_BORDER | WS_VSCROLL | WS_HSCROLL | 0x200
+    DEFPUSHBUTTON   "OK",IDC_OK,100,50,50,14
+    PUSHBUTTON      "Cancel",IDC_CANCEL,150,50,50,14
 END
 
+IDD_YESNOCANCEL DIALOGEX 0, 0, 275, 80
+STYLE DS_ABSALIGN | DS_SYSMODAL | DS_SETFONT | DS_MODALFRAME | DS_FIXEDSYS | DS_CENTER | WS_POPUP | WS_CLIPCHILDREN | WS_CAPTION | WS_SYSMENU | WS_THICKFRAME
+CAPTION "Dialog"
+FONT 8, "MS Shell Dlg", 400, 0, 0x1
+BEGIN
+    EDITTEXT        IDC_MESSAGE,15,15,245,20,ES_MULTILINE | ES_AUTOVSCROLL | ES_AUTOHSCROLL | ES_READONLY | NOT WS_BORDER | WS_VSCROLL | WS_HSCROLL | 0x200
+    DEFPUSHBUTTON   "Yes",IDC_YES,100,50,50,14
+    PUSHBUTTON      "No",IDC_NO_B,150,50,50,14
+    PUSHBUTTON      "Cancel",IDC_CANCEL,200,50,50,14
+END
+
+IDD_CANCELRETRYCONTINUE DIALOGEX 0, 0, 275, 80
+STYLE DS_ABSALIGN | DS_SYSMODAL | DS_SETFONT | DS_MODALFRAME | DS_FIXEDSYS | DS_CENTER | WS_POPUP | WS_CLIPCHILDREN | WS_CAPTION | WS_SYSMENU
+CAPTION "Dialog"
+FONT 8, "MS Shell Dlg", 400, 0, 0x1
+BEGIN
+    EDITTEXT        IDC_MESSAGE,15,15,245,20,ES_MULTILINE | ES_AUTOVSCROLL | ES_AUTOHSCROLL | ES_READONLY | NOT WS_BORDER | WS_VSCROLL | WS_HSCROLL | 0x200
+    PUSHBUTTON      "Cancel",IDC_CANCEL,100,50,50,14
+    DEFPUSHBUTTON   "Retry",IDC_RETRY,150,50,50,14
+    PUSHBUTTON      "Continue",IDC_CONTINUE,200,50,50,14
+END
+
+IDD_YESNO2ALL DIALOGEX 0, 0, 275, 80
+STYLE DS_ABSALIGN | DS_SYSMODAL | DS_SETFONT | DS_MODALFRAME | DS_FIXEDSYS | DS_CENTER | WS_POPUP | WS_CLIPCHILDREN | WS_CAPTION | WS_SYSMENU | WS_THICKFRAME
+CAPTION "Dialog"
+FONT 8, "MS Shell Dlg", 400, 0, 0x1
+BEGIN
+    EDITTEXT        IDC_MESSAGE,15,15,245,20,ES_MULTILINE | ES_AUTOVSCROLL | ES_AUTOHSCROLL | ES_READONLY | NOT WS_BORDER | WS_VSCROLL | WS_HSCROLL | 0x200
+    DEFPUSHBUTTON   "Yes",IDC_YES,50,50,50,14
+    PUSHBUTTON      "No",IDC_NO_B,100,50,50,14
+    PUSHBUTTON      "Yes to all",IDC_YESTOALL,150,50,50,14
+    PUSHBUTTON      "No to all",IDC_NOTOALL,200,50,50,14
+END
+
+IDD_YESNO2ALLCANCEL DIALOGEX 0, 0, 300, 80
+STYLE DS_ABSALIGN | DS_SYSMODAL | DS_SETFONT | DS_MODALFRAME | DS_FIXEDSYS | DS_CENTER | WS_POPUP | WS_CLIPCHILDREN | WS_CAPTION | WS_SYSMENU | WS_THICKFRAME
+CAPTION "Dialog"
+FONT 8, "MS Shell Dlg", 400, 0, 0x1
+BEGIN
+    EDITTEXT        IDC_MESSAGE,15,15,270,20,ES_MULTILINE | ES_AUTOVSCROLL | ES_AUTOHSCROLL | ES_READONLY | NOT WS_BORDER | WS_VSCROLL | WS_HSCROLL | 0x200
+    PUSHBUTTON      "Yes",IDC_YES,0,50,50,14
+    PUSHBUTTON      "Yes to all",IDC_YESTOALL,50,50,50,14
+    PUSHBUTTON      "No",IDC_NO_B,100,50,50,14
+    PUSHBUTTON      "No to all",IDC_NOTOALL,150,50,50,14
+    PUSHBUTTON      "Cancel",IDC_CANCEL,200,50,50,14
+END
+
+IDD_YESNOYESTOALL DIALOGEX 0, 0, 275, 80
+STYLE DS_ABSALIGN | DS_SYSMODAL | DS_SETFONT | DS_MODALFRAME | DS_FIXEDSYS | DS_CENTER | WS_POPUP | WS_CLIPCHILDREN | WS_CAPTION | WS_SYSMENU | WS_THICKFRAME
+CAPTION "Dialog"
+FONT 8, "MS Shell Dlg", 400, 0, 0x1
+BEGIN
+    EDITTEXT        IDC_MESSAGE,15,15,245,20,ES_MULTILINE | ES_AUTOVSCROLL | ES_AUTOHSCROLL | ES_READONLY | NOT WS_BORDER | WS_VSCROLL | WS_HSCROLL | 0x200
+    PUSHBUTTON      "Yes",IDC_YES,100,50,50,14
+    PUSHBUTTON      "No",IDC_NO_B,150,50,50,14
+    PUSHBUTTON      "Yes to all",IDC_YESTOALL,200,50,50,14
+END
+
+
+/////////////////////////////////////////////////////////////////////////////
+//
+// Accelerator
+//
+
+IDR_ACCEL1 ACCELERATORS
+BEGIN
+    "Y",            IDC_YES,                VIRTKEY 
+    "N",            IDC_NO_B,               VIRTKEY 
+    VK_ESCAPE,      IDC_CANCEL,             VIRTKEY 
+    "C",            IDC_COPY,               VIRTKEY, CONTROL
+END
+
+
 /////////////////////////////////////////////////////////////////////////////
 //
 // DESIGNINFO
@@ -68,12 +134,20 @@ END
 #ifdef APSTUDIO_INVOKED
 GUIDELINES DESIGNINFO
 BEGIN
+    IDD_OK, DIALOG
+    BEGIN
+    END
+
+    IDD_CANCELRETRYCONTINUE, DIALOG
+    BEGIN
+    END
+
     IDD_YESNO2ALL, DIALOG
     BEGIN
         LEFTMARGIN, 7
-		RIGHTMARGIN, 268
+        RIGHTMARGIN, 268
         TOPMARGIN, 7
-		BOTTOMMARGIN, 111
+        BOTTOMMARGIN, 111
     END
 
     IDD_YESNO2ALLCANCEL, DIALOG
@@ -84,9 +158,9 @@ BEGIN
         BOTTOMMARGIN, 119
     END
 
-	IDD_YESNOYESTOALL, DIALOG
-	BEGIN
-	END
+    IDD_YESNOYESTOALL, DIALOG
+    BEGIN
+    END
 END
 #endif    // APSTUDIO_INVOKED
 
@@ -111,10 +185,28 @@ END
 #endif    // APSTUDIO_INVOKED
 
 
+/////////////////////////////////////////////////////////////////////////////
+//
+// AFX_DIALOG_LAYOUT
+//
 
+IDD_CANCELRETRYCONTINUE AFX_DIALOG_LAYOUT
+BEGIN
+    0
+END
 
 #endif    // English (United States) resources
 /////////////////////////////////////////////////////////////////////////////
 
 
 
+#ifndef APSTUDIO_INVOKED
+/////////////////////////////////////////////////////////////////////////////
+//
+// Generated from the TEXTINCLUDE 3 resource.
+//
+
+
+/////////////////////////////////////////////////////////////////////////////
+#endif    // not APSTUDIO_INVOKED
+
diff --git a/Engine/Source/Runtime/Launch/Resources/Windows/resource.h b/Engine/Source/Runtime/Launch/Resources/Windows/resource.h
--- a/Engine/Source/Runtime/Launch/Resources/Windows/resource.h
+++ b/Engine/Source/Runtime/Launch/Resources/Windows/resource.h
@@ -4,23 +4,35 @@
 // Microsoft Visual C++ generated include file.
 // Used by PCLaunch.rc
 //
-#define IDD_YESNO2ALL                   101
-#define IDD_YESNO2ALLCANCEL             102
-#define IDD_YESNOYESTOALL               103
+#define IDD_OK                          101
+#define IDD_YESNO                       102
+#define IDD_OKCANCEL                    103
+#define IDD_YESNOCANCEL                 104
+#define IDD_CANCELRETRYCONTINUE         105
+#define IDD_YESNO2ALL                   106
+#define IDD_YESNO2ALLCANCEL             107
+#define IDD_YESNOYESTOALL               108
 #define IDICON_UEGame                   123
-#define IDC_YESTOALL                    1000
-#define IDC_NOTOALL                     1001
-#define IDC_YES                         1002
-#define IDC_NO_B                        1003
-#define IDC_MESSAGE                     1004
-#define IDC_CANCEL                      1009
+#define IDC_OK                          1000
+#define IDC_YESTOALL                    1001
+#define IDC_NOTOALL                     1002
+#define IDC_YES                         1003
+#define IDC_NO_B                        1004
+#define IDC_MESSAGE                     1005
+#define IDC_CANCEL                      1006
+#define IDC_RETRY                       1007
+#define IDC_CONTINUE                    1008
+
+#define IDR_ACCEL1                      101
+#define IDC_COPY                        1009
+
 #define ID_HASHFILE                     1010
 
 // Next default values for new objects
 // 
 #ifdef APSTUDIO_INVOKED
 #ifndef APSTUDIO_READONLY_SYMBOLS
-#define _APS_NEXT_RESOURCE_VALUE        104
+#define _APS_NEXT_RESOURCE_VALUE        109
 #define _APS_NEXT_COMMAND_VALUE         40001
 #define _APS_NEXT_CONTROL_VALUE         1010
 #define _APS_NEXT_SYMED_VALUE           101
