Fixed compile error when CPUPROFILERTRACE_ENABLED is 0.
Also sorted include headers in the two affected files.

diff --git a/Engine/Source/Runtime/Core/Private/Modules/ModuleManager.cpp b/Engine/Source/Runtime/Core/Private/Modules/ModuleManager.cpp
--- a/Engine/Source/Runtime/Core/Private/Modules/ModuleManager.cpp
+++ b/Engine/Source/Runtime/Core/Private/Modules/ModuleManager.cpp
@@ -1,21 +1,23 @@
 // Copyright Epic Games, Inc. All Rights Reserved.
 
 #include "Modules/ModuleManager.h"
-#include "Misc/DateTime.h"
+
 #include "HAL/FileManager.h"
+#include "Internationalization/StringTableCore.h"
+#include "Misc/App.h"
+#include "Misc/DataDrivenPlatformInfoRegistry.h"
+#include "Misc/DateTime.h"
+#include "Misc/FileHelper.h"
 #include "Misc/Parse.h"
 #include "Misc/Paths.h"
-#include "Stats/Stats.h"
-#include "Misc/App.h"
 #include "Misc/ScopeExit.h"
-#include "Modules/ModuleManifest.h"
 #include "Misc/ScopeLock.h"
-#include "Misc/DataDrivenPlatformInfoRegistry.h"
+#include "Modules/ModuleManifest.h"
+#include "ProfilingDebugging/CpuProfilerTrace.h"
 #include "Serialization/LoadTimeTrace.h"
-#include "Internationalization/StringTableCore.h"
-#include "Misc/FileHelper.h"
 #include "Serialization/MemoryReader.h"
 #include "Serialization/MemoryWriter.h"
+#include "Stats/Stats.h"
 #include "Trace/Trace.h"
 #include "Trace/Trace.inl"
 
@@ -397,6 +399,8 @@ IModuleInterface& FModuleManager::LoadModuleChecked( const FName InModuleName )
 	return *Module;
 }
 
+#if CPUPROFILERTRACE_ENABLED
+
 UE_TRACE_EVENT_BEGIN(Cpu, LoadModule, NoSync)
 	UE_TRACE_EVENT_FIELD(UE::Trace::WideString, Name)
 UE_TRACE_EVENT_END()
@@ -409,6 +413,8 @@ UE_TRACE_EVENT_BEGIN(Cpu, StartupModule, NoSync)
 	UE_TRACE_EVENT_FIELD(UE::Trace::WideString, Name)
 UE_TRACE_EVENT_END()
 
+#endif // CPUPROFILERTRACE_ENABLED
+
 IModuleInterface* FModuleManager::LoadModuleWithFailureReason(const FName InModuleName, EModuleLoadResult& OutFailureReason, ELoadModuleFlags InLoadModuleFlags)
 {
 	IModuleInterface* LoadedModule = nullptr;
@@ -445,8 +451,10 @@ IModuleInterface* FModuleManager::LoadModuleWithFailureReason(const FName InModu
 	UE_SCOPED_ENGINE_ACTIVITY(TEXT("Loading Module %s"), *InModuleName.ToString());
 	FScopedBootTiming BootTimingScope("LoadModule");
 	DECLARE_SCOPE_CYCLE_COUNTER(TEXT("Module Load"), STAT_ModuleLoad, STATGROUP_LoadTime);
+#if CPUPROFILERTRACE_ENABLED
 	UE_TRACE_LOG_SCOPED_T(Cpu, LoadModule, CpuChannel)
 		<< LoadModule.Name(*InModuleName.ToString());
+#endif // CPUPROFILERTRACE_ENABLED
 
 #if	STATS
 	// This is fine here, we only load a handful of modules.
@@ -509,8 +517,10 @@ IModuleInterface* FModuleManager::LoadModuleWithFailureReason(const FName InModu
 
 			// Startup the module
 			{
+#if CPUPROFILERTRACE_ENABLED
 				UE_TRACE_LOG_SCOPED_T(Cpu, StartupModule, CpuChannel)
 					<< StartupModule.Name(*InModuleName.ToString());
+#endif // CPUPROFILERTRACE_ENABLED
 				ModuleInfo->Module->StartupModule();
 			}
 
@@ -590,8 +600,10 @@ IModuleInterface* FModuleManager::LoadModuleWithFailureReason(const FName InModu
 		if (FPaths::FileExists(ModuleFileToLoad))
 		{
 			{
+#if CPUPROFILERTRACE_ENABLED
 				UE_TRACE_LOG_SCOPED_T(Cpu, FPlatformProcess_GetDllHandle, CpuChannel)
 					<< FPlatformProcess_GetDllHandle.Name(*ModuleFileToLoad);
+#endif // CPUPROFILERTRACE_ENABLED
 				ModuleInfo->Handle = FPlatformProcess::GetDllHandle(*ModuleFileToLoad);
 			}
 			
@@ -631,8 +643,10 @@ IModuleInterface* FModuleManager::LoadModuleWithFailureReason(const FName InModu
 						{
 							// Startup the module
 							{
+#if CPUPROFILERTRACE_ENABLED
 								UE_TRACE_LOG_SCOPED_T(Cpu, StartupModule, CpuChannel)
 									<< StartupModule.Name(*InModuleName.ToString());
+#endif // CPUPROFILERTRACE_ENABLED
 								ModuleInfo->Module->StartupModule();
 							}
 
diff --git a/Engine/Source/Runtime/Core/Private/Windows/WindowsPlatformProcess.cpp b/Engine/Source/Runtime/Core/Private/Windows/WindowsPlatformProcess.cpp
--- a/Engine/Source/Runtime/Core/Private/Windows/WindowsPlatformProcess.cpp
+++ b/Engine/Source/Runtime/Core/Private/Windows/WindowsPlatformProcess.cpp
@@ -1,33 +1,35 @@
 // Copyright Epic Games, Inc. All Rights Reserved.
 
 #include "Windows/WindowsPlatformProcess.h"
-#include "HAL/PlatformMisc.h"
-#include "Misc/AssertionMacros.h"
-#include "Logging/LogMacros.h"
-#include "HAL/PlatformAffinity.h"
-#include "HAL/UnrealMemory.h"
-#include "Templates/UnrealTemplate.h"
-#include "CoreGlobals.h"
-#include "HAL/FileManager.h"
-#include "Misc/Parse.h"
+
+#include "Containers/Set.h"
 #include "Containers/StringConv.h"
 #include "Containers/UnrealString.h"
-#include "Containers/Set.h"
-#include "Misc/SingleThreadEvent.h"
-#include "Misc/CommandLine.h"
-#include "Misc/Paths.h"
-#include "Misc/TrackedActivity.h"
-#include "Internationalization/Internationalization.h"
 #include "CoreGlobals.h"
-#include "Stats/Stats.h"
-#include "Misc/CoreStats.h"
-#include "Windows/WindowsHWrapper.h"
+#include "CoreGlobals.h"
+#include "HAL/FileManager.h"
+#include "HAL/PlatformAffinity.h"
+#include "HAL/PlatformMisc.h"
+#include "HAL/UnrealMemory.h"
+#include "Internationalization/Internationalization.h"
+#include "Logging/LogMacros.h"
+#include "Misc/AssertionMacros.h"
+#include "Misc/CommandLine.h"
 #include "Misc/ConfigCacheIni.h"
 #include "Misc/CoreDelegates.h"
+#include "Misc/CoreStats.h"
 #include "Misc/Fork.h"
+#include "Misc/Parse.h"
+#include "Misc/Paths.h"
+#include "Misc/SingleThreadEvent.h"
+#include "Misc/TrackedActivity.h"
+#include "ProfilingDebugging/CpuProfilerTrace.h"
 #include "ProfilingDebugging/CsvProfiler.h"
+#include "Stats/Stats.h"
+#include "Templates/UnrealTemplate.h"
 #include "Trace/Trace.h"
 #include "Trace/Trace.inl"
+#include "Windows/WindowsHWrapper.h"
 
 #include "Windows/AllowWindowsPlatformTypes.h"
 	#include <shellapi.h>
@@ -1916,10 +1918,14 @@ static bool ResolveImport(const FString& Name, const TArray<FString>& SearchPath
 	return false;
 }
 
+#if CPUPROFILERTRACE_ENABLED
+
 UE_TRACE_EVENT_BEGIN(Cpu, ResolveMissingImports, NoSync)
 	UE_TRACE_EVENT_FIELD(UE::Trace::WideString, Name)
 UE_TRACE_EVENT_END()
 
+#endif // CPUPROFILERTRACE_ENABLED
+
 /**
  * Resolve all the imports for the given library, searching through a set of directories.
  *
@@ -1930,8 +1936,10 @@ UE_TRACE_EVENT_END()
  */
 static void ResolveMissingImportsRecursive(const FString& FileName, const TArray<FString>& SearchPaths, TArray<FString>& ImportFileNames, TSet<FString>& VisitedImportNames)
 {
+#if CPUPROFILERTRACE_ENABLED
 	UE_TRACE_LOG_SCOPED_T(Cpu, ResolveMissingImports, CpuChannel)
 		<< ResolveMissingImports.Name(*FileName);
+#endif // CPUPROFILERTRACE_ENABLED
 
 	// Read the imports for this library
 	TArray<FString> ImportNames;
@@ -1991,10 +1999,14 @@ static void LogImportDiagnostics(const FString& FileName, const TArray<FString>&
 	}
 }
 
+#if CPUPROFILERTRACE_ENABLED
+
 UE_TRACE_EVENT_BEGIN(Cpu, Windows_LoadLibrary, NoSync)
 	UE_TRACE_EVENT_FIELD(UE::Trace::WideString, Name)
 UE_TRACE_EVENT_END()
 
+#endif // CPUPROFILERTRACE_ENABLED
+
 void *FWindowsPlatformProcess::LoadLibraryWithSearchPaths(const FString& FileName, const TArray<FString>& SearchPaths)
 {
 	UE_SCOPED_IO_ACTIVITY(*WriteToString<256>("Loading Dll ", FileName));
@@ -2023,8 +2035,10 @@ void *FWindowsPlatformProcess::LoadLibraryWithSearchPaths(const FString& FileNam
 			{
 				const void* DependencyHandle = [&ImportFileName]() 
 				{
+#if CPUPROFILERTRACE_ENABLED
 					UE_TRACE_LOG_SCOPED_T(Cpu, Windows_LoadLibrary, CpuChannel)
 						<< Windows_LoadLibrary.Name(*ImportFileName);
+#endif // CPUPROFILERTRACE_ENABLED
 					return LoadLibrary(*ImportFileName);
 				}();
 				
@@ -2044,8 +2058,10 @@ void *FWindowsPlatformProcess::LoadLibraryWithSearchPaths(const FString& FileNam
 	// Try to load the actual library
 	void* Handle = [FullFileName]() 
 	{
+#if CPUPROFILERTRACE_ENABLED
 		UE_TRACE_LOG_SCOPED_T(Cpu, Windows_LoadLibrary, CpuChannel)
 			<< Windows_LoadLibrary.Name(*FullFileName);
+#endif // CPUPROFILERTRACE_ENABLED
 		return LoadLibrary(*FullFileName);
 	}();
 	
