- Add ExecutableName to startup log messages

diff --git a/Engine/Source/Runtime/Core/Private/Misc/App.cpp b/Engine/Source/Runtime/Core/Private/Misc/App.cpp
--- a/Engine/Source/Runtime/Core/Private/Misc/App.cpp
+++ b/Engine/Source/Runtime/Core/Private/Misc/App.cpp
@@ -327,6 +327,7 @@ bool FApp::HasFocus()
 
 void FApp::PrintStartupLogMessages()
 {
+	UE_LOG(LogInit, Log, TEXT("ExecutableName: %s"), FPlatformProcess::ExecutableName(false));
 	UE_LOG(LogInit, Log, TEXT("Build: %s"), FApp::GetBuildVersion());
 	UE_LOG(LogInit, Log, TEXT("Engine Version: %s"), *FEngineVersion::Current().ToString());
 	UE_LOG(LogInit, Log, TEXT("Compatible Engine Version: %s"), *FEngineVersion::CompatibleWith().ToString());
