Fix static analyzer warnings for clipboard copy and float casts

diff --git a/Engine/Source/Runtime/Core/Private/Windows/WindowsPlatformMisc.cpp b/Engine/Source/Runtime/Core/Private/Windows/WindowsPlatformMisc.cpp
--- a/Engine/Source/Runtime/Core/Private/Windows/WindowsPlatformMisc.cpp
+++ b/Engine/Source/Runtime/Core/Private/Windows/WindowsPlatformMisc.cpp
@@ -1052,16 +1052,14 @@ private:
 
 	static float MessageBoxDlgGetDPI(HWND HandleWnd)
 	{
-		float DPIScale = 1.0f;
+		HMODULE User32Module = GetModuleHandle(L"user32.dll");
+		if (User32Module == nullptr) { return 1.0f; }
 
 		typedef UINT(WINAPI* LPGetDpiForWindow)(HWND Hwnd);
-		LPGetDpiForWindow GetDpiForWindow = (LPGetDpiForWindow)(void*)GetProcAddress(GetModuleHandle(L"user32.dll"), "GetDpiForWindow");
-		if (GetDpiForWindow)
-		{
-			DPIScale = GetDpiForWindow(HandleWnd) / 96.0f;
-		}
+		LPGetDpiForWindow GetDpiForWindow = (LPGetDpiForWindow)(void*)GetProcAddress(User32Module, "GetDpiForWindow");
+		if (GetDpiForWindow == nullptr) { return 1.0f; }
 
-		return DPIScale;
+		return static_cast<float>(GetDpiForWindow(HandleWnd)) / 96.0f;
 	}
 
 	static void SetWindowStyleFlags(HWND HandleWnd, LONG Flags, bool bEnabled)
@@ -1212,8 +1210,8 @@ private:
 		int MonitorWidth = MonInfo.rcMonitor.right - MonInfo.rcMonitor.left;
 		int MonitorHeight = MonInfo.rcMonitor.bottom - MonInfo.rcMonitor.top;
 
-		int NewWidth = FMath::Clamp((int)(TextSize.cx + 100 * DPIScale), Width, (int)(MonitorWidth * 0.8f));
-		int NewHeight = FMath::Clamp((int)(TextSize.cy + 165 * DPIScale), Height, (int)(MonitorHeight * 0.8f));
+		int NewWidth = FMath::Clamp((int)((float)TextSize.cx + 100.0f * DPIScale), Width, (int)((float)MonitorWidth * 0.8f));
+		int NewHeight = FMath::Clamp((int)((float)TextSize.cy + 165.0f * DPIScale), Height, (int)((float)MonitorHeight * 0.8f));
 		int NewLeft = DefaultWindowRect.left - (NewWidth - Width) / 2;
 		int NewTop = DefaultWindowRect.top - (NewHeight - Height) / 2;
 		SetWindowPos(DialogHwnd, HWND_NOTOPMOST, NewLeft, NewTop, NewWidth, NewHeight, SWP_NOZORDER);
@@ -1221,6 +1219,24 @@ private:
 		return true;
 	}
 
+	void SetDialogTextInClipboard(HWND HandleWnd)
+	{
+		size_t TextByteCount = Text.GetAllocatedSize();
+		HGLOBAL StrClipboardMemory = GlobalAlloc(GMEM_MOVEABLE, TextByteCount);
+		if (StrClipboardMemory == nullptr) { return; }
+
+		void* StrTarget = GlobalLock(StrClipboardMemory);
+		if (StrTarget == nullptr) { return; }
+		memcpy(StrTarget, *Text, TextByteCount);
+		((char*)StrTarget)[TextByteCount] = 0;
+		GlobalUnlock(StrClipboardMemory);
+
+		if (!OpenClipboard(HandleWnd)) { return; }
+		if (!EmptyClipboard()) { return; }
+		SetClipboardData(CF_UNICODETEXT, StrClipboardMemory);
+		CloseClipboard();
+	}
+
 	/**
 	 * Callback for MessageBoxExt dialog (allowing for Yes to all / No to all )
 	 * @return		One of EAppReturnType::Yes, EAppReturnType::YesAll, EAppReturnType::No, EAppReturnType::NoAll, EAppReturnType::Cancel.
@@ -1276,18 +1292,8 @@ private:
 						break;
 					case IDC_COPY:
 					{
-						size_t TextByteCount = Text.GetAllocatedSize();
-						HGLOBAL StrClipboardMemory = GlobalAlloc(GMEM_MOVEABLE, TextByteCount);
-
-						void* StrTarget = GlobalLock(StrClipboardMemory);
-						memcpy(StrTarget, *Text, TextByteCount);
-						((char*)StrTarget)[TextByteCount] = 0;
-						GlobalUnlock(StrClipboardMemory);
-
-						OpenClipboard(HandleWnd);
-						EmptyClipboard();
-						SetClipboardData(CF_UNICODETEXT, StrClipboardMemory);
-						CloseClipboard();
+						SetDialogTextInClipboard(HandleWnd);
+						break;
 					}
 					break;
 					default:
