Fixed NonUnity build (missing include)

#rnx

diff --git a/Engine/Source/Runtime/Core/Private/HAL/Allocators/CachedOSVeryLargePageAllocator.cpp b/Engine/Source/Runtime/Core/Private/HAL/Allocators/CachedOSVeryLargePageAllocator.cpp
--- a/Engine/Source/Runtime/Core/Private/HAL/Allocators/CachedOSVeryLargePageAllocator.cpp
+++ b/Engine/Source/Runtime/Core/Private/HAL/Allocators/CachedOSVeryLargePageAllocator.cpp
@@ -6,6 +6,7 @@
 #include "Logging/LogMacros.h"
 #include "CoreGlobals.h"
 #include "HAL/LowLevelMemTracker.h"
+#include "HAL/IConsoleManager.h"
 #include "ProfilingDebugging/CsvProfiler.h"
 
 #if UE_USE_VERYLARGEPAGEALLOCATOR
