StringView: Moved functions and operators to be hidden friends where possible

#rnx

diff --git a/Engine/Source/Runtime/Core/Public/Containers/StringView.h b/Engine/Source/Runtime/Core/Public/Containers/StringView.h
--- a/Engine/Source/Runtime/Core/Public/Containers/StringView.h
+++ b/Engine/Source/Runtime/Core/Public/Containers/StringView.h
@@ -342,6 +342,99 @@ public:
 		return Index >= 0 && Index < Len();
 	}
 
+private:
+	static bool PrivateEquals(TStringView Lhs, const CharType* Rhs);
+	static bool PrivateEquals(TStringView Lhs, TStringView Rhs);
+	static bool PrivateLess(TStringView Lhs, TStringView Rhs);
+
+public:
+	friend constexpr inline auto GetNum(TStringView String)
+	{
+		return String.Len();
+	}
+
+	/** Case insensitive string hash function. */
+	friend inline uint32 GetTypeHash(TStringView View)
+	{
+		// This must match the GetTypeHash behavior of FString
+		return FCrc::Strihash_DEPRECATED(View.Len(), View.GetData());
+	}
+
+	friend inline bool operator==(TStringView Lhs, TStringView Rhs)
+	{
+		return Lhs.Equals(Rhs, ESearchCase::IgnoreCase);
+	}
+
+	friend inline bool operator!=(TStringView Lhs, TStringView Rhs)
+	{
+		return !(Lhs == Rhs);
+	}
+
+	friend inline bool operator<(TStringView Lhs, TStringView Rhs)
+	{
+		return Lhs.Compare(Rhs, ESearchCase::IgnoreCase) < 0;
+	}
+
+	template <typename CharRangeType>
+	friend inline auto operator==(TStringView Lhs, CharRangeType&& Rhs)
+		-> decltype(TStringView::PrivateEquals(Lhs, ImplicitConv<TStringView>(Forward<CharRangeType>(Rhs))))
+	{
+		return TStringView::PrivateEquals(Lhs, ImplicitConv<TStringView>(Forward<CharRangeType>(Rhs)));
+	}
+
+	template <typename CharRangeType>
+	friend inline auto operator==(CharRangeType&& Lhs, TStringView Rhs)
+		-> decltype(TStringView::PrivateEquals(ImplicitConv<TStringView>(Forward<CharRangeType>(Lhs)), Rhs))
+	{
+		return TStringView::PrivateEquals(ImplicitConv<TStringView>(Forward<CharRangeType>(Lhs)), Rhs);
+	}
+
+	template <typename CharRangeType>
+	friend inline auto operator!=(TStringView Lhs, CharRangeType&& Rhs) -> decltype(!(Lhs == Forward<CharRangeType>(Rhs)))
+	{
+		return !(Lhs == Forward<CharRangeType>(Rhs));
+	}
+
+	template <typename CharRangeType>
+	friend inline auto operator!=(CharRangeType&& Lhs, TStringView Rhs) -> decltype(!(Rhs == Forward<CharRangeType>(Lhs)))
+	{
+		return !(Rhs == Forward<CharRangeType>(Lhs));
+	}
+
+	template <typename CharRangeType>
+	friend inline auto operator<(TStringView Lhs, CharRangeType&& Rhs)
+		-> decltype(TStringView::PrivateLess(Lhs, ImplicitConv<TStringView>(Forward<CharRangeType>(Rhs))))
+	{
+		return TStringView::PrivateLess(Lhs, ImplicitConv<TStringView>(Forward<CharRangeType>(Rhs)));
+	}
+
+	template <typename CharRangeType>
+	friend inline auto operator<(CharRangeType&& Lhs, TStringView Rhs)
+		-> decltype(TStringView::PrivateLess(ImplicitConv<TStringView>(Forward<CharRangeType>(Lhs)), Rhs))
+	{
+		return TStringView::PrivateLess(ImplicitConv<TStringView>(Forward<CharRangeType>(Lhs)), Rhs);
+	}
+
+	friend inline bool operator==(TStringView Lhs, const CharType* Rhs)
+	{
+		return TStringView::PrivateEquals(Lhs, Rhs);
+	}
+
+	friend inline bool operator==(const CharType* Lhs, TStringView Rhs)
+	{
+		return Rhs == Lhs;
+	}
+
+	friend inline bool operator!=(TStringView Lhs, const CharType* Rhs)
+	{
+		return !(Lhs == Rhs);
+	}
+
+	friend inline bool operator!=(const CharType* Lhs, TStringView Rhs)
+	{
+		return !(Lhs == Rhs);
+	}
+
 public:
 	/**
 	 * DO NOT USE DIRECTLY
@@ -371,12 +464,6 @@ constexpr inline auto MakeStringView(CharPtrType&& CharPtr, int32 Size) -> declt
 	return TStringView(Forward<CharPtrType>(CharPtr), Size);
 }
 
-template <typename CharType>
-constexpr inline auto GetNum(TStringView<CharType> String)
-{
-	return String.Len();
-}
-
 //////////////////////////////////////////////////////////////////////////
 
 #if PLATFORM_TCHAR_IS_UTF8CHAR
@@ -453,14 +540,6 @@ constexpr inline FWideStringView operator "" _PrivateWSV(const WIDECHAR* String,
 
 //////////////////////////////////////////////////////////////////////////
 
-/** Case insensitive string hash function. */
-template <typename CharType>
-FORCEINLINE uint32 GetTypeHash(TStringView<CharType> View)
-{
-	// This must match the GetTypeHash behavior of FString
-	return FCrc::Strihash_DEPRECATED(View.Len(), View.GetData());
-}
-
 template <typename CharType>
 inline const CharType& TStringView<CharType>::operator[](int32 Index) const
 {
@@ -505,7 +584,7 @@ inline TStringView<CharType> TStringView<CharType>::RightChop(int32 CharCount) c
 template <typename CharType>
 inline TStringView<CharType> TStringView<CharType>::Mid(int32 Position, int32 CharCount) const
 {
-	const CharType* CurrentStart  = GetData();
+	const CharType* CurrentStart = GetData();
 	const int32 CurrentLength = Len();
 
 	// Clamp minimum position at the start of the string, adjusting the length down if necessary
@@ -649,19 +728,18 @@ inline bool TStringView<CharType>::EndsWith(ViewType Suffix, ESearchCase::Type S
 
 template <typename CharType>
 inline int32 TStringView<CharType>::Find(const ViewType Search, const int32 StartPosition) const
-			{
+{
 	const int32 Index = UE::String::FindFirst(RightChop(StartPosition), Search);
 	return Index == INDEX_NONE ? INDEX_NONE : Index + StartPosition;
 }
 
 template <typename CharType>
 inline bool TStringView<CharType>::FindChar(const CharType Search, int32& OutIndex) const
-	{
+{
 	OutIndex = UE::String::FindFirstChar(*this, Search);
 	return OutIndex != INDEX_NONE;
 }
 
-
 template <typename CharType>
 inline bool TStringView<CharType>::FindLastChar(const CharType Search, int32& OutIndex) const
 {
@@ -669,90 +747,20 @@ inline bool TStringView<CharType>::FindLastChar(const CharType Search, int32& Ou
 	return OutIndex != INDEX_NONE;
 }
 
-// Case-insensitive string view comparison operators
-
 template <typename CharType>
-inline bool operator==(TStringView<CharType> Lhs, TStringView<CharType> Rhs)
+inline bool TStringView<CharType>::PrivateEquals(TStringView Lhs, const CharType* Rhs)
 {
 	return Lhs.Equals(Rhs, ESearchCase::IgnoreCase);
 }
 
 template <typename CharType>
-inline bool operator!=(TStringView<CharType> Lhs, TStringView<CharType> Rhs)
-{
-	return !(Lhs == Rhs);
-}
-
-template <typename CharType>
-inline bool operator<(TStringView<CharType> Lhs, TStringView<CharType> Rhs)
-{
-	return Lhs.Compare(Rhs, ESearchCase::IgnoreCase) < 0;
-}
-
-// Case-insensitive character range comparison operators
-
-template <typename CharType, typename CharRangeType>
-inline auto operator==(TStringView<CharType> Lhs, CharRangeType&& Rhs)
-	-> decltype(Lhs.Equals(ImplicitConv<TStringView<CharType>>(Forward<CharRangeType>(Rhs)), ESearchCase::IgnoreCase))
-{
-	return Lhs.Equals(ImplicitConv<TStringView<CharType>>(Forward<CharRangeType>(Rhs)), ESearchCase::IgnoreCase);
-}
-
-template <typename CharType, typename CharRangeType>
-inline auto operator==(CharRangeType&& Lhs, TStringView<CharType> Rhs)
-	-> decltype(ImplicitConv<TStringView<CharType>>(Forward<CharRangeType>(Lhs)).Equals(Rhs, ESearchCase::IgnoreCase))
-{
-	return ImplicitConv<TStringView<CharType>>(Forward<CharRangeType>(Lhs)).Equals(Rhs, ESearchCase::IgnoreCase);
-}
-
-template <typename CharType, typename CharRangeType>
-inline auto operator!=(TStringView<CharType> Lhs, CharRangeType&& Rhs) -> decltype(!(Lhs == Forward<CharRangeType>(Rhs)))
-{
-	return !(Lhs == Forward<CharRangeType>(Rhs));
-}
-
-template <typename CharType, typename CharRangeType>
-inline auto operator!=(CharRangeType&& Lhs, TStringView<CharType> Rhs) -> decltype(!(Rhs == Forward<CharRangeType>(Lhs)))
-{
-	return !(Rhs == Forward<CharRangeType>(Lhs));
-}
-
-template <typename CharType, typename CharRangeType>
-inline auto operator<(TStringView<CharType> Lhs, CharRangeType&& Rhs)
-	-> decltype(Lhs.Compare(ImplicitConv<TStringView<CharType>>(Forward<CharRangeType>(Rhs)), ESearchCase::IgnoreCase) < 0)
-{
-	return Lhs.Compare(ImplicitConv<TStringView<CharType>>(Forward<CharRangeType>(Rhs)), ESearchCase::IgnoreCase) < 0;
-}
-
-template <typename CharType, typename CharRangeType>
-inline auto operator<(CharRangeType&& Lhs, TStringView<CharType> Rhs)
-	-> decltype(ImplicitConv<TStringView<CharType>>(Forward<CharRangeType>(Lhs)).Compare(Rhs, ESearchCase::IgnoreCase) < 0)
-{
-	return ImplicitConv<TStringView<CharType>>(Forward<CharRangeType>(Lhs)).Compare(Rhs, ESearchCase::IgnoreCase) < 0;
-}
-
-// Case-insensitive C-string comparison operators
-
-template <typename CharType>
-inline bool operator==(TStringView<CharType> Lhs, const CharType* Rhs)
+inline bool TStringView<CharType>::PrivateEquals(TStringView Lhs, TStringView Rhs)
 {
 	return Lhs.Equals(Rhs, ESearchCase::IgnoreCase);
 }
 
 template <typename CharType>
-inline bool operator==(const CharType* Lhs, TStringView<CharType> Rhs)
-{
-	return Rhs == Lhs;
-}
-
-template <typename CharType>
-inline bool operator!=(TStringView<CharType> Lhs, const CharType* Rhs)
-{
-	return !(Lhs == Rhs);
-}
-
-template <typename CharType>
-inline bool operator!=(const CharType* Lhs, TStringView<CharType> Rhs)
+inline bool TStringView<CharType>::PrivateLess(TStringView Lhs, TStringView Rhs)
 {
-	return !(Lhs == Rhs);
+	return Lhs.Compare(Rhs, ESearchCase::IgnoreCase) < 0;
 }
