[CSV Profiler] Removed personal data and free text from shipping builds for privacy rules

diff --git a/Engine/Source/Runtime/Core/Private/ProfilingDebugging/CsvProfiler.cpp b/Engine/Source/Runtime/Core/Private/ProfilingDebugging/CsvProfiler.cpp
--- a/Engine/Source/Runtime/Core/Private/ProfilingDebugging/CsvProfiler.cpp
+++ b/Engine/Source/Runtime/Core/Private/ProfilingDebugging/CsvProfiler.cpp
@@ -2751,19 +2751,23 @@ FCsvProfiler::FCsvProfiler()
 	SetMetadataInternal(TEXT("Config"), *BuildConfigurationStr);
 	SetMetadataInternal(TEXT("BuildVersion"), *BuildVersionString);
 	SetMetadataInternal(TEXT("EngineVersion"), *EngineVersionString);
-	SetMetadataInternal(TEXT("Commandline"), *CommandlineStr, false);
 	SetMetadataInternal(TEXT("OS"), *OSString);
 	SetMetadataInternal(TEXT("CPU"), *FPlatformMisc::GetDeviceMakeAndModel());
 	SetMetadataInternal(TEXT("PGOEnabled"), FPlatformMisc::IsPGOEnabled() ? TEXT("1") : TEXT("0"));
-	SetMetadataInternal(TEXT("LoginID"), *FPlatformMisc::GetLoginId());
 	SetMetadataInternal(TEXT("ASan"), USING_ADDRESS_SANITISER ? TEXT("1") : TEXT("0"));
 
+#if !UE_BUILD_SHIPPING
+	// for privacy, personal and free text fields are not allowed in shipping
+	SetMetadataInternal(TEXT("Commandline"), *CommandlineStr, false);
+	SetMetadataInternal(TEXT("LoginID"), *FPlatformMisc::GetLoginId());
+
 	// Set the device ID if the platform supports it
 	FString DeviceID = FPlatformMisc::GetDeviceId();
 	if (!DeviceID.IsEmpty())
 	{
 		SetMetadataInternal(TEXT("DeviceID"), *DeviceID);
 	}
+#endif
 }
 
 FCsvProfiler::~FCsvProfiler()
