Add extra debug info for Task destruction

diff --git a/Engine/Source/Runtime/Core/Public/Async/TaskGraphInterfaces.h b/Engine/Source/Runtime/Core/Public/Async/TaskGraphInterfaces.h
--- a/Engine/Source/Runtime/Core/Public/Async/TaskGraphInterfaces.h
+++ b/Engine/Source/Runtime/Core/Public/Async/TaskGraphInterfaces.h
@@ -851,7 +851,10 @@ protected:
 	/** destructor, just checks the life stage **/
 	virtual ~FBaseGraphTask()
 	{
-		checkThreadGraph(LifeStage.Increment() == int32(LS_Deconstucted));
+#if DO_GUARD_SLOW
+		int32 Stage = LifeStage.Increment();
+		checkf(Stage == int32(LS_Deconstucted), TEXT("LifeStage was %d"), Stage);
+#endif
 	}
 
 #if !(UE_BUILD_SHIPPING || UE_BUILD_TEST)
