Fix incorrect LWC error tolerance assert in TLargeWorldRenderScalar.

diff --git a/Engine/Source/Runtime/Core/Private/Misc/LargeWorldRenderPosition.cpp b/Engine/Source/Runtime/Core/Private/Misc/LargeWorldRenderPosition.cpp
--- a/Engine/Source/Runtime/Core/Private/Misc/LargeWorldRenderPosition.cpp
+++ b/Engine/Source/Runtime/Core/Private/Misc/LargeWorldRenderPosition.cpp
@@ -113,14 +113,33 @@ FMatrix TLargeWorldRenderScalar<TScalar>::MakeClampedToRelativeWorldMatrixDouble
 	return ClampedToRelativeWorld;
 }
 
-template<typename TScalar>
-void TLargeWorldRenderScalar<TScalar>::Validate(double InAbsolute)
+// TILE_SIZE = 2097152 = (1 << 21)
+//  => Offset < (1 << 21)
+// f32 has 23+1 mantissa bits, leaving 3 bits for fractions (24 - 21)
+// f64 has 52+1 mantissa bits, leaving 32 bits for fractions (53 - 21)
+// This defines a theoretical bound on the error.
+
+template<>
+void TLargeWorldRenderScalar<float>::Validate(double InAbsolute)
+{
+	constexpr int FractionBitCount = 3;
+	constexpr double Tolerance = 1.0 / (1ULL << FractionBitCount);
+	const double CheckAbsolute = GetAbsolute();
+	const double Delta = FMath::Abs(CheckAbsolute - InAbsolute);
+
+	ensureMsgf(Delta <= Tolerance, TEXT("Bad TLargeWorldRenderScalar<float> (%.15f) vs (%.15f)"),
+		InAbsolute, CheckAbsolute);
+}
+
+template<>
+void TLargeWorldRenderScalar<double>::Validate(double InAbsolute)
 {
-	const double Tolerance = 0.01; // TODO LWC - How precise do we need to be?
+	constexpr int FractionBitCount = 32;
+	constexpr double Tolerance = 1.0 / (1ULL << FractionBitCount);
 	const double CheckAbsolute = GetAbsolute();
 	const double Delta = FMath::Abs(CheckAbsolute - InAbsolute);
 
-	ensureMsgf(Delta < Tolerance, TEXT("Bad TLargeWorldRenderScalar<TScalar> (%g) vs (%g)"),
+	ensureMsgf(Delta <= Tolerance, TEXT("Bad TLargeWorldRenderScalar<double> (%.15f) vs (%.15f)"),
 		InAbsolute, CheckAbsolute);
 }
 
