Allow a project to redefine default TSetAllocator's number of elements per bucket

diff --git a/Engine/Source/Runtime/Core/Public/Containers/ContainerAllocationPolicies.h b/Engine/Source/Runtime/Core/Public/Containers/ContainerAllocationPolicies.h
--- a/Engine/Source/Runtime/Core/Public/Containers/ContainerAllocationPolicies.h
+++ b/Engine/Source/Runtime/Core/Public/Containers/ContainerAllocationPolicies.h
@@ -1119,9 +1119,11 @@ public:
 // Set allocation definitions.
 //
 
-#define DEFAULT_NUMBER_OF_ELEMENTS_PER_HASH_BUCKET	2
-#define DEFAULT_BASE_NUMBER_OF_HASH_BUCKETS			8
-#define DEFAULT_MIN_NUMBER_OF_HASHED_ELEMENTS		4
+#if !defined(DEFAULT_NUMBER_OF_ELEMENTS_PER_HASH_BUCKET)
+#	define DEFAULT_NUMBER_OF_ELEMENTS_PER_HASH_BUCKET	2
+#endif
+#define DEFAULT_BASE_NUMBER_OF_HASH_BUCKETS				8
+#define DEFAULT_MIN_NUMBER_OF_HASHED_ELEMENTS			4
 
 /** Encapsulates the allocators used by a set in a single type. */
 template<
