LLM: Change ENABLE_LOW_LEVEL_MEM_TRACKER to not be separately definable, and make LLM_ENABLED_IN_CONFIG separately definable instead. This is necessary because ENABLE_LOW_LEVEL_MEM_TRACKER can only be true if LLM_ENABLED_IN_CONFIG is also true, because c-language files include the definition of LLM_ENABLED_IN_CONFIG but cannot include the platform-specific definition of ENABLE_LOW_LEVEL_MEM_TRACKER.
#rnx

diff --git a/Engine/Source/Programs/ShaderCompileWorker/ShaderCompileWorker.Target.cs b/Engine/Source/Programs/ShaderCompileWorker/ShaderCompileWorker.Target.cs
--- a/Engine/Source/Programs/ShaderCompileWorker/ShaderCompileWorker.Target.cs
+++ b/Engine/Source/Programs/ShaderCompileWorker/ShaderCompileWorker.Target.cs
@@ -65,7 +65,7 @@ public class ShaderCompileWorkerTarget : TargetRules
 
 		if (bShaderCompilerWorkerTrace)
         {
-			GlobalDefinitions.Add("ENABLE_LOW_LEVEL_MEM_TRACKER=1");
+			GlobalDefinitions.Add("LLM_ENABLED_IN_CONFIG=1");
 			GlobalDefinitions.Add("UE_MEMORY_TAGS_TRACE_ENABLED=1");
 			GlobalDefinitions.Add("UE_TRACE_ENABLED=1");
 		}
diff --git a/Engine/Source/Runtime/Core/Public/HAL/LowLevelMemTracker.h b/Engine/Source/Runtime/Core/Public/HAL/LowLevelMemTracker.h
--- a/Engine/Source/Runtime/Core/Public/HAL/LowLevelMemTracker.h
+++ b/Engine/Source/Runtime/Core/Public/HAL/LowLevelMemTracker.h
@@ -14,10 +14,10 @@
 // This is not a problem, because fixed GMalloc is only used in Test/Shipping builds
 #define LLM_ENABLED_ON_PLATFORM (PLATFORM_SUPPORTS_LLM && !PLATFORM_USES_FIXED_GMalloc_CLASS)
 
-#if !defined(ENABLE_LOW_LEVEL_MEM_TRACKER) || !LLM_ENABLED_ON_PLATFORM 
-	#undef ENABLE_LOW_LEVEL_MEM_TRACKER
-	#define ENABLE_LOW_LEVEL_MEM_TRACKER (LLM_ENABLED_IN_CONFIG && LLM_ENABLED_ON_PLATFORM)
+#ifdef ENABLE_LOW_LEVEL_MEM_TRACKER
+#error ENABLE_LOW_LEVEL_MEM_TRACKER is now a derived define that should not be defined separately. Define LLM_ENABLED_IN_CONFIG (build environment only) or LLM_ENABLED_ON_PLATFORM (build environment or c++ header) instead.
 #endif
+#define ENABLE_LOW_LEVEL_MEM_TRACKER (LLM_ENABLED_IN_CONFIG && LLM_ENABLED_ON_PLATFORM)
 
 #if ENABLE_LOW_LEVEL_MEM_TRACKER
 
diff --git a/Engine/Source/Runtime/Core/Public/HAL/LowLevelMemTrackerDefines.h b/Engine/Source/Runtime/Core/Public/HAL/LowLevelMemTrackerDefines.h
--- a/Engine/Source/Runtime/Core/Public/HAL/LowLevelMemTrackerDefines.h
+++ b/Engine/Source/Runtime/Core/Public/HAL/LowLevelMemTrackerDefines.h
@@ -11,8 +11,14 @@
 	#define ALLOW_LOW_LEVEL_MEM_TRACKER_IN_TEST 0
 #endif
 
-// *** When locally instrumenting, override this definition of LLM_ENABLED_IN_CONFIG
-#define LLM_ENABLED_IN_CONFIG ( \
-	!UE_BUILD_SHIPPING && (!UE_BUILD_TEST || ALLOW_LOW_LEVEL_MEM_TRACKER_IN_TEST) && \
-	WITH_ENGINE \
-)
+// LLM_ENABLED_IN_CONFIG can be defined here or in the build environment only.
+// It cannot be defined in platform header files because it is included in c-language compilation units that
+// cannot include those headers.
+// When locally instrumenting, add your own definition of LLM_ENABLED_IN_CONFIG in build environment or this header.
+
+#ifndef LLM_ENABLED_IN_CONFIG 
+	#define LLM_ENABLED_IN_CONFIG ( \
+		!UE_BUILD_SHIPPING && (!UE_BUILD_TEST || ALLOW_LOW_LEVEL_MEM_TRACKER_IN_TEST) && \
+		WITH_ENGINE \
+	)
+#endif
