Compile time optimizations Managed to reduce unnamed game's biggest file compiler frontend cost from 113 seconds to 60 seconds. This also impact other things such as pch size and memory footprint when compiling which is great.

Takeaways. Try to use "friend" on functions/operators that are frequently overloaded EXCEPT if they are in a templated type that is frequently instantiated. So do not put friends in TMap, TSet, TObjectPtr etc, this will slow down compile times. There is a break-even somewhere and hard to tell where it is but taking a class that is templatized on character type probably don't matter either way and then it is nicer to use a friend since that simplies error messages when compiler can't resolve functions/operators.

If it is possible to use member functions instead of friend that is the best option in terms of compile time performance. With c++20 you only have to write operator==(Foo, Bar) and the compiler will automatically provide operator==(Bar, Foo), operator!=(Foo, Bar) and operator!=(Bar, Foo).

Changes in this changelist involes
* Making operator<< friends in non-template types and not friends in template types
* Making operator==/!= members where possible and if not possible moved out if type is a frequently instantiated templated type.

diff --git a/Engine/Source/Developer/DerivedDataCache/Public/DerivedDataBuildVersion.h b/Engine/Source/Developer/DerivedDataCache/Public/DerivedDataBuildVersion.h
--- a/Engine/Source/Developer/DerivedDataCache/Public/DerivedDataBuildVersion.h
+++ b/Engine/Source/Developer/DerivedDataCache/Public/DerivedDataBuildVersion.h
@@ -37,6 +37,14 @@ public:
 	using FArchive::operator<<;
 	FArchive& operator<<(FName& Name) final;
 
+	template <typename ArgType>
+	inline FBuildVersionBuilder& operator<<(ArgType&& Arg)
+	{
+		FArchive& Ar = *this;
+		Ar << const_cast<std::decay_t<ArgType>&>(Arg);
+		return *this;
+	}
+
 private:
 	FXxHash128Builder HashBuilder;
 };
@@ -76,12 +84,4 @@ inline FArchive& FBuildVersionBuilder::operator<<(FName& Name)
 	return *this << NameString;
 }
 
-template <typename ArgType>
-inline FBuildVersionBuilder& operator<<(FBuildVersionBuilder& Builder, ArgType&& Arg)
-{
-	FArchive& Ar = Builder;
-	Ar << const_cast<std::decay_t<ArgType>&>(Arg);
-	return Builder;
-}
-
 } // UE::DerivedData
diff --git a/Engine/Source/Runtime/Core/Private/Serialization/CompactBinaryWriter.cpp b/Engine/Source/Runtime/Core/Private/Serialization/CompactBinaryWriter.cpp
--- a/Engine/Source/Runtime/Core/Private/Serialization/CompactBinaryWriter.cpp
+++ b/Engine/Source/Runtime/Core/Private/Serialization/CompactBinaryWriter.cpp
@@ -617,22 +617,22 @@ void FCbWriter::AddCustom(const FUtf8StringView TypeName, const FMemoryView Valu
 
 ///////////////////////////////////////////////////////////////////////////////////////////////////////////////////////
 
-FCbWriter& operator<<(FCbWriter& Writer, const FDateTime Value)
+FCbWriter& FCbWriter::operator<<(const FDateTime Value)
 {
-	Writer.AddDateTime(Value);
-	return Writer;
+	AddDateTime(Value);
+	return *this;
 }
 
-FCbWriter& operator<<(FCbWriter& Writer, const FTimespan Value)
+FCbWriter& FCbWriter::operator<<(const FTimespan Value)
 {
-	Writer.AddTimeSpan(Value);
-	return Writer;
+	AddTimeSpan(Value);
+	return *this;
 }
 
-FCbWriter& operator<<(FCbWriter& Writer, FName Value)
+FCbWriter& FCbWriter::operator<<(FName Value)
 {
-	Writer << WriteToUtf8String<FName::StringBufferSize>(Value).ToView();
-	return Writer;
+	*this << WriteToUtf8String<FName::StringBufferSize>(Value).ToView();
+	return *this;
 }
 
 ///////////////////////////////////////////////////////////////////////////////////////////////////////////////////////
diff --git a/Engine/Source/Runtime/Core/Public/Containers/Array.h b/Engine/Source/Runtime/Core/Public/Containers/Array.h
--- a/Engine/Source/Runtime/Core/Public/Containers/Array.h
+++ b/Engine/Source/Runtime/Core/Public/Containers/Array.h
@@ -167,8 +167,8 @@
 		Index--;
 	}
 
-	FORCEINLINE friend bool operator==(const TIndexedContainerIterator& Lhs, const TIndexedContainerIterator& Rhs) { return &Lhs.Container == &Rhs.Container && Lhs.Index == Rhs.Index; }
-	FORCEINLINE friend bool operator!=(const TIndexedContainerIterator& Lhs, const TIndexedContainerIterator& Rhs) { return &Lhs.Container != &Rhs.Container || Lhs.Index != Rhs.Index; }
+	FORCEINLINE bool operator==(const TIndexedContainerIterator& Rhs) const { return &Container == &Rhs.Container && Index == Rhs.Index; }
+	FORCEINLINE bool operator!=(const TIndexedContainerIterator& Rhs) const { return &Container != &Rhs.Container || Index != Rhs.Index; }
 
 private:
 
@@ -222,21 +222,21 @@
 			return *this;
 		}
 
-	private:
-		ElementType*    Ptr;
-		const SizeType& CurrentNum;
-		SizeType        InitialNum;
-
-		FORCEINLINE friend bool operator!=(const TCheckedPointerIterator& Lhs, const TCheckedPointerIterator& Rhs)
+		FORCEINLINE bool operator!=(const TCheckedPointerIterator& Rhs) const
 		{
 			// We only need to do the check in this operator, because no other operator will be
 			// called until after this one returns.
 			//
 			// Also, we should only need to check one side of this comparison - if the other iterator isn't
 			// even from the same array then the compiler has generated bad code.
-			ensureMsgf(Lhs.CurrentNum == Lhs.InitialNum, TEXT("Array has changed during ranged-for iteration!"));
-			return Lhs.Ptr != Rhs.Ptr;
+			ensureMsgf(CurrentNum == InitialNum, TEXT("Array has changed during ranged-for iteration!"));
+			return Ptr != Rhs.Ptr;
 		}
+
+	private:
+		ElementType*    Ptr;
+		const SizeType& CurrentNum;
+		SizeType        InitialNum;
 	};
 #endif
 
@@ -260,13 +260,13 @@
 		return *this;
 	}
 
+	FORCEINLINE bool operator!=(const TDereferencingIterator& Rhs) const
+	{
+		return Iter != Rhs.Iter;
+	}
+
 private:
 	IteratorType Iter;
-
-	FORCEINLINE friend bool operator!=(const TDereferencingIterator& Lhs, const TDereferencingIterator& Rhs)
-	{
-		return Lhs.Iter != Rhs.Iter;
-	}
 };
 
 namespace UE4Array_Private
@@ -1218,94 +1218,14 @@
 	 * @param OtherArray Array to compare.
 	 * @returns True if this array is NOT the same as OtherArray. False otherwise.
 	 */
+#if !PLATFORM_COMPILER_HAS_GENERATED_COMPARISON_OPERATORS
 	FORCEINLINE bool operator!=(const TArray& OtherArray) const
 	{
 		return !(*this == OtherArray);
 	}
+#endif
+
+
 
 	/**
-	 * Serialization operator.
-	 *
-	 * @param Ar Archive to serialize the array with.
-	 * @param A Array to serialize.
-	 * @returns Passing the given archive.
-	 */
-	friend FArchive& operator<<(FArchive& Ar, TArray& A)
-	{
-		A.CountBytes(Ar);
-
-		// For net archives, limit serialization to 16MB, to protect against excessive allocation
-		constexpr SizeType MaxNetArraySerialize = (16 * 1024 * 1024) / sizeof(ElementType);
-		SizeType SerializeNum = Ar.IsLoading() ? 0 : A.ArrayNum;
-
-		Ar << SerializeNum;
-
-		if (SerializeNum == 0)
-		{
-			// if we are loading, then we have to reset the size to 0, in case it isn't currently 0
-			if (Ar.IsLoading())
-			{
-				A.Empty();
-			}
-			return Ar;
-		}
-
-		check(SerializeNum >= 0);
-
-		if (!Ar.IsError() && SerializeNum > 0 && ensure(!Ar.IsNetArchive() || SerializeNum <= MaxNetArraySerialize))
-		{
-			// if we don't need to perform per-item serialization, just read it in bulk
-			if constexpr (sizeof(ElementType) == 1 || TCanBulkSerialize<ElementType>::Value)
-			{
-				A.ArrayNum = SerializeNum;
-
-				// Serialize simple bytes which require no construction or destruction.
-				if ((A.ArrayNum || A.ArrayMax) && Ar.IsLoading())
-				{
-					A.ResizeForCopy(A.ArrayNum, A.ArrayMax);
-				}
-
-				if(TIsUECoreVariant<ElementType, double>::Value && Ar.IsLoading() && Ar.UEVer() < EUnrealEngineObjectUE5Version::LARGE_WORLD_COORDINATES)
-				{
-					// Per item serialization is required for core variant types loaded from pre LWC archives, to enable conversion from float to double.
-					A.Empty(SerializeNum);
-					for (SizeType i=0; i<SerializeNum; i++)
-					{
-						Ar << *::new(A) ElementType;
-					}		
-				}
-				else
-				{
-					Ar.Serialize(A.GetData(), A.Num() * sizeof(ElementType));
-				}
-			}
-			else if (Ar.IsLoading())
-			{
-				// Required for resetting ArrayNum
-				A.Empty(SerializeNum);
-
-				for (SizeType i=0; i<SerializeNum; i++)
-				{
-					Ar << *::new(A) ElementType;
-				}
-			}
-			else
-			{
-				A.ArrayNum = SerializeNum;
-
-				for (SizeType i=0; i<A.ArrayNum; i++)
-				{
-					Ar << A[i];
-				}
-			}
-		}
-		else
-		{
-			Ar.SetError();
-		}
-
-		return Ar;
-	}
-
-	/**
 	 * Bulk serialize array as a single memory blob when loading. Uses regular serialization code for saving
@@ -3479,6 +3399,8 @@
 
 	const ElementAllocatorType& GetAllocatorInstance() const { return AllocatorInstance; }
 	ElementAllocatorType& GetAllocatorInstance() { return AllocatorInstance; }
+
+	friend struct TArrayPrivateFriend;
 };
 
 
@@ -3564,3 +3486,99 @@
 	Array.InsertUninitialized(Index);
 	return &Array[Index];
 }
+
+struct TArrayPrivateFriend
+{
+	/**
+	 * Serialization operator.
+	 *
+	 * @param Ar Archive to serialize the array with.
+	 * @param A Array to serialize.
+	 * @returns Passing the given archive.
+	 */
+	template<typename ElementType, typename AllocatorType>
+	static FArchive& Serialize(FArchive& Ar, TArray<ElementType, AllocatorType>& A)
+	{
+		A.CountBytes(Ar);
+
+		// For net archives, limit serialization to 16MB, to protect against excessive allocation
+		typedef typename AllocatorType::SizeType SizeType;
+		constexpr SizeType MaxNetArraySerialize = (16 * 1024 * 1024) / sizeof(ElementType);
+		SizeType SerializeNum = Ar.IsLoading() ? 0 : A.ArrayNum;
+
+		Ar << SerializeNum;
+
+		if (SerializeNum == 0)
+		{
+			// if we are loading, then we have to reset the size to 0, in case it isn't currently 0
+			if (Ar.IsLoading())
+			{
+				A.Empty();
+			}
+			return Ar;
+		}
+
+		check(SerializeNum >= 0);
+
+		if (!Ar.IsError() && SerializeNum > 0 && ensure(!Ar.IsNetArchive() || SerializeNum <= MaxNetArraySerialize))
+		{
+			// if we don't need to perform per-item serialization, just read it in bulk
+			if constexpr (sizeof(ElementType) == 1 || TCanBulkSerialize<ElementType>::Value)
+			{
+				A.ArrayNum = SerializeNum;
+
+				// Serialize simple bytes which require no construction or destruction.
+				if ((A.ArrayNum || A.ArrayMax) && Ar.IsLoading())
+				{
+					A.ResizeForCopy(A.ArrayNum, A.ArrayMax);
+				}
+
+				if(TIsUECoreVariant<ElementType, double>::Value && Ar.IsLoading() && Ar.UEVer() < EUnrealEngineObjectUE5Version::LARGE_WORLD_COORDINATES)
+				{
+					// Per item serialization is required for core variant types loaded from pre LWC archives, to enable conversion from float to double.
+					A.Empty(SerializeNum);
+					for (SizeType i=0; i<SerializeNum; i++)
+					{
+						Ar << *::new(A) ElementType;
+					}		
+				}
+				else
+				{
+					Ar.Serialize(A.GetData(), A.Num() * sizeof(ElementType));
+				}
+			}
+			else if (Ar.IsLoading())
+			{
+				// Required for resetting ArrayNum
+				A.Empty(SerializeNum);
+
+				for (SizeType i=0; i<SerializeNum; i++)
+				{
+					Ar << *::new(A) ElementType;
+				}
+			}
+			else
+			{
+				A.ArrayNum = SerializeNum;
+
+				for (SizeType i=0; i<A.ArrayNum; i++)
+				{
+					Ar << A[i];
+				}
+			}
+		}
+		else
+		{
+			Ar.SetError();
+		}
+
+		return Ar;
+	}
+};
+
+
+template<typename ElementType, typename AllocatorType>
+FArchive& operator<<(FArchive& Ar, TArray<ElementType, AllocatorType>& A)
+{
+	return TArrayPrivateFriend::Serialize(Ar, A);
+}
diff --git a/Engine/Source/Runtime/Core/Public/Containers/ArrayView.h b/Engine/Source/Runtime/Core/Public/Containers/ArrayView.h
--- a/Engine/Source/Runtime/Core/Public/Containers/ArrayView.h
+++ b/Engine/Source/Runtime/Core/Public/Containers/ArrayView.h
@@ -773,6 +773,7 @@ bool operator==(TArrayView<ElementType> Lhs, RangeType&& Rhs)
 	return (Rhs == Lhs);
 }
 
+#if !PLATFORM_COMPILER_HAS_GENERATED_COMPARISON_OPERATORS
 /**
  * Inequality operator.
  *
@@ -798,6 +799,7 @@ bool operator!=(TArrayView<ElementType> Lhs, RangeType&& Rhs)
 {
 	return !(Rhs == Lhs);
 }
+#endif
 
 template<typename InElementType, typename InAllocatorType>
 template<typename OtherElementType, typename OtherSizeType>
diff --git a/Engine/Source/Runtime/Core/Public/Containers/BitArray.h b/Engine/Source/Runtime/Core/Public/Containers/BitArray.h
--- a/Engine/Source/Runtime/Core/Public/Containers/BitArray.h
+++ b/Engine/Source/Runtime/Core/Public/Containers/BitArray.h
@@ -70,12 +70,6 @@ enum class EBitwiseOperatorFlags
 };
 ENUM_CLASS_FLAGS(EBitwiseOperatorFlags)
 
-/**
- * Serializer (predefined for no friend injection in gcc 411)
- */
-template<typename Allocator>
-FArchive& operator<<(FArchive& Ar, TBitArray<Allocator>& BitArray);
-
 /** Used to read/write a bit in the array as a bool. */
 class FBitReference
 {
@@ -374,10 +368,12 @@ public:
 		return false;
 	}
 
+#if !PLATFORM_COMPILER_HAS_GENERATED_COMPARISON_OPERATORS
 	FORCEINLINE bool operator!=(const TBitArray<Allocator>& Other) const
 	{
 		return !(*this == Other);
 	}
+#endif
 
 private:
 	FORCEINLINE uint32 GetNumWords() const
@@ -472,29 +468,28 @@ public:
 	/**
 	 * Serializer
 	 */
-	friend FArchive& operator<<(FArchive& Ar, TBitArray& BitArray)
+	void Serialize(FArchive& Ar)
 	{
 		// serialize number of bits
-		Ar << BitArray.NumBits;
+		Ar << NumBits;
 
 		if (Ar.IsLoading())
 		{
 			// no need for slop when reading; set MaxBits to the smallest legal value that is >= NumBits
-			BitArray.MaxBits = NumBitsPerDWORD * FMath::Max(FBitSet::CalculateNumWords(BitArray.NumBits), (uint32)BitArray.AllocatorInstance.GetInitialCapacity());
+			MaxBits = NumBitsPerDWORD * FMath::Max(FBitSet::CalculateNumWords(NumBits), (uint32)AllocatorInstance.GetInitialCapacity());
 
 			// allocate room for new bits
-			BitArray.Realloc(0);
+			Realloc(0);
 		}
 
 		// serialize the data as one big chunk
-		Ar.Serialize(BitArray.GetData(), BitArray.GetNumWords() * sizeof(uint32));
+		Ar.Serialize(GetData(), GetNumWords() * sizeof(uint32));
 
 		if (Ar.IsLoading() && !Ar.IsObjectReferenceCollector() && !Ar.IsCountingMemory())
 		{
 			// Clear slack bits incase they were serialized non-null
-			BitArray.ClearPartialSlackBits();
+			ClearPartialSlackBits();
 		}
-		return Ar;
 	}
 
 	/**
@@ -1730,19 +1725,6 @@ public:
 		checkf(!Writer.Is32BitTarget(), TEXT("TBitArray does not currently support freezing for 32bits"));
 		TSupportsFreezeMemoryImageHelper<bSupportsFreezeMemoryImage>::WriteMemoryImage(Writer, *this);
 	}
-
-	friend FORCEINLINE uint32 GetTypeHash(const TBitArray& BitArray)
-	{
-		uint32 NumWords = FBitSet::CalculateNumWords(BitArray.Num());
-		uint32 Hash = NumWords;
-		const uint32* Data = BitArray.GetData();
-		for (uint32 i = 0; i < NumWords; i++)
-		{
-			Hash ^= Data[i];
-		}
-		return Hash;
-	}
-
 };
 
 namespace Freeze
@@ -1756,6 +1738,19 @@ namespace Freeze
 
 DECLARE_TEMPLATE_INTRINSIC_TYPE_LAYOUT(template<typename Allocator>, TBitArray<Allocator>);
 
+template<typename Allocator>
+FORCEINLINE uint32 GetTypeHash(const TBitArray<Allocator>& BitArray)
+{
+	uint32 NumWords = FBitSet::CalculateNumWords(BitArray.Num());
+	uint32 Hash = NumWords;
+	const uint32* Data = BitArray.GetData();
+	for (uint32 i = 0; i < NumWords; i++)
+	{
+		Hash ^= Data[i];
+	}
+	return Hash;
+}
+
 template<typename Allocator>
 struct TContainerTraits<TBitArray<Allocator> > : public TContainerTraitsBase<TBitArray<Allocator> >
 {
@@ -1797,16 +1792,18 @@ public:
 		return *this;
 	}
 
-	FORCEINLINE friend bool operator==(const TConstSetBitIterator& Lhs, const TConstSetBitIterator& Rhs) 
+	FORCEINLINE bool operator==(const TConstSetBitIterator& Rhs) const
 	{
 		// We only need to compare the bit index and the array... all the rest of the state is unobservable.
-		return Lhs.CurrentBitIndex == Rhs.CurrentBitIndex && &Lhs.Array == &Rhs.Array;
+		return CurrentBitIndex == Rhs.CurrentBitIndex && &Array == &Rhs.Array;
 	}
 
-	FORCEINLINE friend bool operator!=(const TConstSetBitIterator& Lhs, const TConstSetBitIterator& Rhs)
+#if !PLATFORM_COMPILER_HAS_GENERATED_COMPARISON_OPERATORS
+	FORCEINLINE bool operator!=(const TConstSetBitIterator& Rhs) const
 	{ 
-		return !(Lhs == Rhs);
+		return !(*this == Rhs);
 	}
+#endif
 
 	/** conversion to "bool" returning true if the iterator is valid. */
 	FORCEINLINE explicit operator bool() const
@@ -2176,3 +2173,10 @@ class FScriptBitArray : public TScriptBitArray<FDefaultBitArrayAllocator, FScrip
 public:
 	using Super::Super;
 };
+
+template<typename Allocator>
+FArchive& operator<<(FArchive& Ar, TBitArray<Allocator>& BitArray)
+{
+	BitArray.Serialize(Ar);
+	return Ar;
+}
diff --git a/Engine/Source/Runtime/Core/Public/Containers/ChunkedArray.h b/Engine/Source/Runtime/Core/Public/Containers/ChunkedArray.h
--- a/Engine/Source/Runtime/Core/Public/Containers/ChunkedArray.h
+++ b/Engine/Source/Runtime/Core/Public/Containers/ChunkedArray.h
@@ -39,9 +39,9 @@ namespace UE4ChunkedArray_Private
 			}
 		}
 
-		friend bool operator!=(const TChunkedArrayIterator& Lhs, const TChunkedArrayIterator& Rhs)
+		bool operator!=(const TChunkedArrayIterator& Rhs) const
 		{
-			return Lhs.Elem != Rhs.Elem;
+			return Elem != Rhs.Elem;
 		}
 	};
 }
diff --git a/Engine/Source/Runtime/Core/Public/Containers/Deque.h b/Engine/Source/Runtime/Core/Public/Containers/Deque.h
--- a/Engine/Source/Runtime/Core/Public/Containers/Deque.h
+++ b/Engine/Source/Runtime/Core/Public/Containers/Deque.h
@@ -580,15 +580,13 @@ private:
 	//---------------------------------------------------------------------------------------------------------------------
 	// TDeque comparison operators
 	//---------------------------------------------------------------------------------------------------------------------
-
-	friend inline bool operator==(
-		const TDeque<InElementType, InAllocatorType>& Left,
-		const TDeque<InElementType, InAllocatorType>& Right)
+public:
+	inline bool operator==(const TDeque& Right) const
 	{
-		if (Left.Num() == Right.Num())
+		if (Num() == Right.Num())
 		{
-			auto EndIt = Left.end();
-			auto LeftIt = Left.begin();
+			auto EndIt = end();
+			auto LeftIt = begin();
 			auto RightIt = Right.begin();
 			while (LeftIt != EndIt)
 			{
@@ -602,11 +600,10 @@ private:
 		return false;
 	}
 
-	friend inline bool operator!=(
-		const TDeque<InElementType, InAllocatorType>& Left,
-		const TDeque<InElementType, InAllocatorType>& Right)
+#if !PLATFORM_COMPILER_HAS_GENERATED_COMPARISON_OPERATORS
+	inline bool operator!=(const TDeque& Right) const
 	{
-		return !(Left == Right);
+		return !(*this == Right);
 	}
-
+#endif
 };
diff --git a/Engine/Source/Runtime/Core/Public/Containers/DynamicRHIResourceArray.h b/Engine/Source/Runtime/Core/Public/Containers/DynamicRHIResourceArray.h
--- a/Engine/Source/Runtime/Core/Public/Containers/DynamicRHIResourceArray.h
+++ b/Engine/Source/Runtime/Core/Public/Containers/DynamicRHIResourceArray.h
@@ -126,16 +126,6 @@ public:
 		Super::BulkSerialize(Ar, bForcePerElementSerialization);
 	}
 
-	/**
-	* Serializer for this class
-	* @param Ar - archive to serialize to
-	* @param ResourceArray - resource array data to serialize
-	*/
-	friend FArchive& operator<<(FArchive& Ar,TResourceArray<ElementType,Alignment>& ResourceArray)
-	{
-		return Ar << *(Super*)&ResourceArray;
-	}	
-
 private:
 	/** 
 	* true if this array needs to be accessed by the CPU.  
@@ -149,3 +139,14 @@ struct TContainerTraits<TResourceArray<ElementType, Alignment>> : public TContai
 {
 	enum { MoveWillEmptyContainer = TContainerTraits<typename TResourceArray<ElementType, Alignment>::Super>::MoveWillEmptyContainer };
 };
+
+/**
+* Serializer for this class
+* @param Ar - archive to serialize to
+* @param ResourceArray - resource array data to serialize
+*/
+template<typename ElementType, uint32 Alignment>
+FArchive& operator<<(FArchive& Ar,TResourceArray<ElementType,Alignment>& ResourceArray)
+{
+	return Ar << *(typename TResourceArray<ElementType,Alignment>::Super*)&ResourceArray;
+}	
diff --git a/Engine/Source/Runtime/Core/Public/Containers/EnumAsByte.h b/Engine/Source/Runtime/Core/Public/Containers/EnumAsByte.h
--- a/Engine/Source/Runtime/Core/Public/Containers/EnumAsByte.h
+++ b/Engine/Source/Runtime/Core/Public/Containers/EnumAsByte.h
@@ -108,13 +108,13 @@ private:
 
 	/** Holds the value as a byte. **/
 	uint8 Value;
-
-
-	FORCEINLINE friend uint32 GetTypeHash(const TEnumAsByte& Enum)
-	{
-		return GetTypeHash(Enum.Value);
-	}
 };
 
+template<class T>
+FORCEINLINE uint32 GetTypeHash(const TEnumAsByte<T>& Enum)
+{
+	return GetTypeHash((uint8)Enum.GetValue());
+}
+
 
 template<class T> struct TIsPODType<TEnumAsByte<T>> { enum { Value = true }; };
diff --git a/Engine/Source/Runtime/Core/Public/Containers/IndirectArray.h b/Engine/Source/Runtime/Core/Public/Containers/IndirectArray.h
--- a/Engine/Source/Runtime/Core/Public/Containers/IndirectArray.h
+++ b/Engine/Source/Runtime/Core/Public/Containers/IndirectArray.h
@@ -232,42 +232,6 @@ public:
 		}
 	}
 
-	/**
-	 * Serialization operator for TIndirectArray.
-	 *
-	 * @param Ar Archive to serialize with.
-	 * @param A Array to serialize.
-	 * @returns Passing down serializing archive.
-	 */
-	friend FArchive& operator<<(FArchive& Ar, TIndirectArray& A)
-	{
-		A.CountBytes(Ar);
-		if (Ar.IsLoading())
-		{
-			// Load array.
-			int32 NewNum;
-			Ar << NewNum;
-			A.Empty(NewNum);
-			for (int32 Index = 0; Index < NewNum; Index++)
-			{
-				T* NewElement = new T;
-				Ar << *NewElement;
-				A.Add(NewElement);
-			}
-		}
-		else
-		{
-			// Save array.
-			int32 Num = A.Num();
-			Ar << Num;
-			for (int32 Index = 0; Index < Num; Index++)
-			{
-				Ar << A[Index];
-			}
-		}
-		return Ar;
-	}
-
 	/**
 	 * Count bytes needed to serialize this array.
 	 *
@@ -487,3 +451,40 @@ void* operator new( size_t Size, TIndirectArray<T,Allocator>& Array, int32 Index
 	Array.Insert((T*)FMemory::Malloc(Size), Index);
 	return &Array[Index];
 }
+
+/**
+* Serialization operator for TIndirectArray.
+*
+* @param Ar Archive to serialize with.
+* @param A Array to serialize.
+* @returns Passing down serializing archive.
+*/
+template<typename T,typename Allocator>
+FArchive& operator<<(FArchive& Ar, TIndirectArray<T, Allocator>& A)
+{
+	A.CountBytes(Ar);
+	if (Ar.IsLoading())
+	{
+		// Load array.
+		int32 NewNum;
+		Ar << NewNum;
+		A.Empty(NewNum);
+		for (int32 Index = 0; Index < NewNum; Index++)
+		{
+			T* NewElement = new T;
+			Ar << *NewElement;
+			A.Add(NewElement);
+		}
+	}
+	else
+	{
+		// Save array.
+		int32 Num = A.Num();
+		Ar << Num;
+		for (int32 Index = 0; Index < Num; Index++)
+		{
+			Ar << A[Index];
+		}
+	}
+	return Ar;
+}
diff --git a/Engine/Source/Runtime/Core/Public/Containers/IteratorAdapter.h b/Engine/Source/Runtime/Core/Public/Containers/IteratorAdapter.h
--- a/Engine/Source/Runtime/Core/Public/Containers/IteratorAdapter.h
+++ b/Engine/Source/Runtime/Core/Public/Containers/IteratorAdapter.h
@@ -113,13 +113,14 @@ protected:
 	template <typename AnyElementType>
 	friend class TIteratorAdapter;
 
-	friend bool operator==(const TIteratorAdapter& Left, const TIteratorAdapter& Right)
+public:
+	bool operator==(const TIteratorAdapter& Right) const
 	{
-		return Left.Equals(Right);
+		return Equals(Right);
 	}
 
-	friend bool operator!=(const TIteratorAdapter& Left, const TIteratorAdapter& Right)
+	bool operator!=(const TIteratorAdapter& Right) const
 	{
-		return !(Left == Right);
+		return !Equals(Right);
 	}
 };
diff --git a/Engine/Source/Runtime/Core/Public/Containers/List.h b/Engine/Source/Runtime/Core/Public/Containers/List.h
--- a/Engine/Source/Runtime/Core/Public/Containers/List.h
+++ b/Engine/Source/Runtime/Core/Public/Containers/List.h
@@ -41,12 +41,11 @@ public:
 		return CurrentLink != nullptr;
 	}
 
+	FORCEINLINE bool operator==(const TLinkedListIteratorBase& Rhs) const { return CurrentLink == Rhs.CurrentLink; }
+	FORCEINLINE bool operator!=(const TLinkedListIteratorBase& Rhs) const { return CurrentLink != Rhs.CurrentLink; }
 protected:
 
 	ContainerType* CurrentLink;
-
-	FORCEINLINE friend bool operator==(const TLinkedListIteratorBase& Lhs, const TLinkedListIteratorBase& Rhs) { return Lhs.CurrentLink == Rhs.CurrentLink; }
-	FORCEINLINE friend bool operator!=(const TLinkedListIteratorBase& Lhs, const TLinkedListIteratorBase& Rhs) { return Lhs.CurrentLink != Rhs.CurrentLink; }
 };
 
 template <class ContainerType, class ElementType>
@@ -419,11 +418,11 @@ public:
 		return CurrentNode;
 	}
 
+	bool operator==(const TDoubleLinkedListIterator& Rhs) const { return CurrentNode == Rhs.CurrentNode; }
+	bool operator!=(const TDoubleLinkedListIterator& Rhs) const { return CurrentNode != Rhs.CurrentNode; }
+
 private:
 	NodeType* CurrentNode;
-
-	friend bool operator==(const TDoubleLinkedListIterator& Lhs, const TDoubleLinkedListIterator& Rhs) { return Lhs.CurrentNode == Rhs.CurrentNode; }
-	friend bool operator!=(const TDoubleLinkedListIterator& Lhs, const TDoubleLinkedListIterator& Rhs) { return Lhs.CurrentNode != Rhs.CurrentNode; }
 };
 
 
diff --git a/Engine/Source/Runtime/Core/Public/Containers/LruCache.h b/Engine/Source/Runtime/Core/Public/Containers/LruCache.h
--- a/Engine/Source/Runtime/Core/Public/Containers/LruCache.h
+++ b/Engine/Source/Runtime/Core/Public/Containers/LruCache.h
@@ -528,14 +528,14 @@ public:
 			return *this;
 		}
 
-		FORCEINLINE friend bool operator==(const TBaseIterator& Lhs, const TBaseIterator& Rhs)
+		FORCEINLINE bool operator==(const TBaseIterator& Rhs) const
 		{
-			return Lhs.CurrentEntry == Rhs.CurrentEntry;
+			return CurrentEntry == Rhs.CurrentEntry;
 		}
 
-		FORCEINLINE friend bool operator!=(const TBaseIterator& Lhs, const TBaseIterator& Rhs)
+		FORCEINLINE bool operator!=(const TBaseIterator& Rhs) const
 		{
-			return Lhs.CurrentEntry != Rhs.CurrentEntry;
+			return CurrentEntry != Rhs.CurrentEntry;
 		}
 
 		ValueType& operator->() const
diff --git a/Engine/Source/Runtime/Core/Public/Containers/Map.h b/Engine/Source/Runtime/Core/Public/Containers/Map.h
--- a/Engine/Source/Runtime/Core/Public/Containers/Map.h
+++ b/Engine/Source/Runtime/Core/Public/Containers/Map.h
@@ -181,15 +181,6 @@ protected:
 	}
 
 public:
-	// Legacy comparison operators.  Note that these also test whether the map's key-value pairs were added in the same order!
-	friend bool LegacyCompareEqual(const TMapBase& A, const TMapBase& B)
-	{
-		return LegacyCompareEqual(A.Pairs, B.Pairs);
-	}
-	friend bool LegacyCompareNotEqual(const TMapBase& A, const TMapBase& B)
-	{
-		return LegacyCompareNotEqual(A.Pairs, B.Pairs);
-	}
 
 	/**
 	 * Compare this map with another for equality. Does not make any assumptions about Key order.
@@ -725,51 +716,6 @@ public:
 		}
 	}
 
-	/** Serializer. */
-	FORCEINLINE friend FArchive& operator<<(FArchive& Ar, TMapBase& Map)
-	{
-		return Ar << Map.Pairs;
-	}
-
-	/** Structured archive serializer. */
-	FORCEINLINE friend void operator<<(FStructuredArchive::FSlot Slot, TMapBase& InMap)
-	{
-		/*
-		if (Slot.GetUnderlyingArchive().IsTextFormat())
-		{
-			int32 Num = InMap.Num();
-			FStructuredArchive::FMap Map = Slot.EnterMap(Num);
-
-			if (Slot.GetUnderlyingArchive().IsLoading())
-			{
-				FString KeyString;
-				KeyType Key;
-
-				for (int32 Index = 0; Index < Num; ++Index)
-				{
-					FStructuredArchive::FSlot ValueSlot = Map.EnterElement(KeyString);
-					LexFromString(Key, *KeyString);
-					ValueSlot << InMap.Add(Key);
-				}
-			}
-			else
-			{
-				FString StringK;
-				for (TMapBase::TIterator It(InMap); It; ++It)
-				{
-					StringK = LexToString(It->Key);
-					FStructuredArchive::FSlot ValueSlot = Map.EnterElement(StringK);
-					ValueSlot << It->Value;
-				}
-			}
-		}
-		else
-		*/
-		{
-			Slot << InMap.Pairs;
-		}
-	}
-
 	/**
 	 * Describes the map's contents through an output device.
 	 *
@@ -822,8 +768,8 @@ protected:
 			return !(bool)*this;
 		}
 
-		FORCEINLINE friend bool operator==(const TBaseIterator& Lhs, const TBaseIterator& Rhs) { return Lhs.PairIt == Rhs.PairIt; }
-		FORCEINLINE friend bool operator!=(const TBaseIterator& Lhs, const TBaseIterator& Rhs) { return Lhs.PairIt != Rhs.PairIt; }
+		FORCEINLINE bool operator==(const TBaseIterator& Rhs) const { return PairIt == Rhs.PairIt; }
+		FORCEINLINE bool operator!=(const TBaseIterator& Rhs) const { return PairIt != Rhs.PairIt; }
 
 		FORCEINLINE ItKeyType&   Key()   const { return PairIt->Key; }
 		FORCEINLINE ItValueType& Value() const { return PairIt->Value; }
@@ -1005,6 +951,8 @@ public:
 		return TConstKeyIterator(*this, InKey);
 	}
 
+	friend struct TMapPrivateFriend;
+
 public:
 	/**
 	 * DO NOT USE DIRECTLY
@@ -1934,3 +1882,84 @@ struct TContainerTraits<TMultiMap<KeyType, ValueType, SetAllocator, KeyFuncs>> :
 {
 	enum { MoveWillEmptyContainer = TContainerTraits<typename TMultiMap<KeyType, ValueType, SetAllocator, KeyFuncs>::ElementSetType>::MoveWillEmptyContainer };
 };
+
+struct TMapPrivateFriend
+{
+	template <typename KeyType, typename ValueType, typename SetAllocator, typename KeyFuncs>
+	FORCEINLINE static FArchive& Serialize(FArchive& Ar, TMapBase<KeyType, ValueType, SetAllocator, KeyFuncs>& Map)
+	{
+		Ar << Map.Pairs;
+		return Ar;
+	}
+
+	template <typename KeyType, typename ValueType, typename SetAllocator, typename KeyFuncs>
+	FORCEINLINE static void SerializeStructured(FStructuredArchive::FSlot Slot, TMapBase<KeyType, ValueType, SetAllocator, KeyFuncs>& InMap)
+	{
+		/*
+		if (Slot.GetUnderlyingArchive().IsTextFormat())
+		{
+			int32 Num = InMap.Num();
+			FStructuredArchive::FMap Map = Slot.EnterMap(Num);
+
+			if (Slot.GetUnderlyingArchive().IsLoading())
+			{
+				FString KeyString;
+				KeyType Key;
+
+				for (int32 Index = 0; Index < Num; ++Index)
+				{
+					FStructuredArchive::FSlot ValueSlot = Map.EnterElement(KeyString);
+					LexFromString(Key, *KeyString);
+					ValueSlot << InMap.Add(Key);
+				}
+			}
+			else
+			{
+				FString StringK;
+				for (TMapBase::TIterator It(InMap); It; ++It)
+				{
+					StringK = LexToString(It->Key);
+					FStructuredArchive::FSlot ValueSlot = Map.EnterElement(StringK);
+					ValueSlot << It->Value;
+				}
+			}
+		}
+		else
+		*/
+		{
+			Slot << InMap.Pairs;
+		}
+	}
+
+	template <typename KeyType, typename ValueType, typename SetAllocator, typename KeyFuncs>
+	static bool LegacyCompareEqual(const TMapBase<KeyType, ValueType, SetAllocator, KeyFuncs>& A, const TMapBase<KeyType, ValueType, SetAllocator, KeyFuncs>& B)
+	{
+		return TSetPrivateFriend::LegacyCompareEqual(A.Pairs, B.Pairs);
+	}
+};
+
+/** Serializer. */
+template <typename KeyType, typename ValueType, typename SetAllocator, typename KeyFuncs>
+FORCEINLINE FArchive& operator<<(FArchive& Ar, TMapBase<KeyType, ValueType, SetAllocator, KeyFuncs>& Map)
+{
+	return TMapPrivateFriend::Serialize(Ar, Map);
+}
+
+/** Structured archive serializer. */
+template <typename KeyType, typename ValueType, typename SetAllocator, typename KeyFuncs>
+FORCEINLINE void operator<<(FStructuredArchive::FSlot Slot, TMapBase<KeyType, ValueType, SetAllocator, KeyFuncs>& InMap)
+{
+	TMapPrivateFriend::SerializeStructured(Slot, InMap);
+}
+
+// Legacy comparison operators.  Note that these also test whether the map's key-value pairs were added in the same order!
+template <typename KeyType, typename ValueType, typename SetAllocator, typename KeyFuncs>
+bool LegacyCompareEqual(const TMapBase<KeyType, ValueType, SetAllocator, KeyFuncs>& A, const TMapBase<KeyType, ValueType, SetAllocator, KeyFuncs>& B)
+{
+	return TMapPrivateFriend::LegacyCompareEqual(A, B);
+}
+template <typename KeyType, typename ValueType, typename SetAllocator, typename KeyFuncs>
+bool LegacyCompareNotEqual(const TMapBase<KeyType, ValueType, SetAllocator, KeyFuncs>& A, const TMapBase<KeyType, ValueType, SetAllocator, KeyFuncs>& B)
+{
+	return !TMapPrivateFriend::LegacyCompareEqual(A, B);
+}
diff --git a/Engine/Source/Runtime/Core/Public/Containers/PagedArray.h b/Engine/Source/Runtime/Core/Public/Containers/PagedArray.h
--- a/Engine/Source/Runtime/Core/Public/Containers/PagedArray.h
+++ b/Engine/Source/Runtime/Core/Public/Containers/PagedArray.h
@@ -526,16 +526,16 @@ public:
 	}
 
 	// Friend operators
-	friend bool operator==(const TPagedArray& Left, const TPagedArray& Right)
+	bool operator==(const TPagedArray& Right) const
 	{
-		if (Left.Num() != Right.Num())
+		if (Num() != Right.Num())
 		{
 			return false;
 		}
-		const SizeType MinPageCount = NumRequiredPages(Left.Count);
+		const SizeType MinPageCount = NumRequiredPages(Count);
 		for (SizeType PageIndex = 0; PageIndex < MinPageCount; ++PageIndex)
 		{
-			if (Left.Pages[PageIndex] != Right.Pages[PageIndex])
+			if (Pages[PageIndex] != Right.Pages[PageIndex])
 			{
 				return false;
 			}
@@ -543,9 +543,9 @@ public:
 		return true;
 	}
 
-	FORCEINLINE friend bool operator!=(const TPagedArray& Left, const TPagedArray& Right)
+	FORCEINLINE bool operator!=(const TPagedArray& Right) const
 	{
-		return !(Left == Right);
+		return !(*this == Right);
 	}
 
 private:
diff --git a/Engine/Source/Runtime/Core/Public/Containers/RingBuffer.h b/Engine/Source/Runtime/Core/Public/Containers/RingBuffer.h
--- a/Engine/Source/Runtime/Core/Public/Containers/RingBuffer.h
+++ b/Engine/Source/Runtime/Core/Public/Containers/RingBuffer.h
@@ -113,8 +113,8 @@ public:
 		Index--;
 	}
 
-	FORCEINLINE friend bool operator==(const TRingBufferIterator& Lhs, const TRingBufferIterator& Rhs) { return &Lhs.Container == &Rhs.Container && Lhs.Index == Rhs.Index; }
-	FORCEINLINE friend bool operator!=(const TRingBufferIterator& Lhs, const TRingBufferIterator& Rhs) { return &Lhs.Container != &Rhs.Container || Lhs.Index != Rhs.Index; }
+	FORCEINLINE bool operator==(const TRingBufferIterator& Rhs) const { return &Container == &Rhs.Container && Index == Rhs.Index; }
+	FORCEINLINE bool operator!=(const TRingBufferIterator& Rhs) const { return &Container != &Rhs.Container || Index != Rhs.Index; }
 
 private:
 
diff --git a/Engine/Source/Runtime/Core/Public/Containers/Set.h b/Engine/Source/Runtime/Core/Public/Containers/Set.h
--- a/Engine/Source/Runtime/Core/Public/Containers/Set.h
+++ b/Engine/Source/Runtime/Core/Public/Containers/Set.h
@@ -240,18 +240,6 @@
 	TSetElement& operator=(TSetElement&&) = default;
 	TSetElement& operator=(const TSetElement&) = default;
 
-	/** Serializer. */
-	FORCEINLINE friend FArchive& operator<<(FArchive& Ar,TSetElement& Element)
-	{
-		return Ar << Element.Value;
-	}
-
-	/** Structured archive serializer. */
- 	FORCEINLINE friend void operator<<(FStructuredArchive::FSlot Slot, TSetElement& Element)
- 	{
- 		Slot << Element.Value;
- 	}
-
 	// Comparison operators
 	FORCEINLINE bool operator==(const TSetElement& Other) const
 	{
@@ -1077,41 +1065,6 @@
 		Rehash();
 	}
 
-	/** Serializer. */
-	friend FArchive& operator<<(FArchive& Ar,TSet& Set)
-	{
-		// Load the set's new elements.
-		Ar << Set.Elements;
-
-		if(Ar.IsLoading())
-		{
-			// Free the old hash.
-			Set.Hash.ResizeAllocation(0,0,sizeof(FSetElementId));
-			Set.HashSize = 0;
-
-			// Hash the newly loaded elements.
-			Set.ConditionalRehash(Set.Elements.Num());
-		}
-
-		return Ar;
-	}
-
-	/** Structured archive serializer. */
- 	friend void operator<<(FStructuredArchive::FSlot Slot, TSet& Set)
- 	{
-		Slot << Set.Elements;
-
-		if (Slot.GetUnderlyingArchive().IsLoading())
-		{
-			// Free the old hash.
-			Set.Hash.ResizeAllocation(0, 0, sizeof(FSetElementId));
-			Set.HashSize = 0;
-
-			// Hash the newly loaded elements.
-			Set.ConditionalRehash(Set.Elements.Num());
-		}
- 	}
-
 	/**
 	 * Describes the set's contents through an output device.
 	 * @param Ar - The output device to describe the set's contents through.
@@ -1179,16 +1132,6 @@
 		}
 	}
 
-	// Legacy comparison operators.  Note that these also test whether the set's elements were added in the same order!
-	friend bool LegacyCompareEqual(const TSet& A,const TSet& B)
-	{
-		return A.Elements == B.Elements;
-	}
-	friend bool LegacyCompareNotEqual(const TSet& A,const TSet& B)
-	{
-		return A.Elements != B.Elements;
-	}
-
 	/** @return the intersection of two sets. (A AND B)*/
 	TSet Intersect(const TSet& OtherSet) const
 	{
@@ -1630,8 +1573,8 @@
 			return ElementIt->Value;
 		}
 
-		FORCEINLINE friend bool operator==(const TBaseIterator& Lhs, const TBaseIterator& Rhs) { return Lhs.ElementIt == Rhs.ElementIt; }
-		FORCEINLINE friend bool operator!=(const TBaseIterator& Lhs, const TBaseIterator& Rhs) { return Lhs.ElementIt != Rhs.ElementIt; }
+		FORCEINLINE bool operator==(const TBaseIterator& Rhs) const { return ElementIt == Rhs.ElementIt; }
+		FORCEINLINE bool operator!=(const TBaseIterator& Rhs) const { return ElementIt != Rhs.ElementIt; }
 
 		ElementItType ElementIt;
 	};
@@ -1803,6 +1746,8 @@
 		return TConstIterator(*this);
 	}
 
+	friend struct TSetPrivateFriend;
+
 public:
 	/**
 	 * DO NOT USE DIRECTLY
@@ -2198,3 +2143,90 @@
 public:
 	using Super::Super;
 };
+
+struct TSetPrivateFriend
+{
+	/** Serializer. */
+	template<typename ElementType, typename KeyFuncs,typename Allocator>
+	static FArchive& Serialize(FArchive& Ar,TSet<ElementType, KeyFuncs, Allocator>& Set)
+	{
+		// Load the set's new elements.
+		Ar << Set.Elements;
+
+		if(Ar.IsLoading())
+		{
+			// Free the old hash.
+			Set.Hash.ResizeAllocation(0,0,sizeof(FSetElementId));
+			Set.HashSize = 0;
+
+			// Hash the newly loaded elements.
+			Set.ConditionalRehash(Set.Elements.Num());
+		}
+
+		return Ar;
+	}
+
+	/** Structured archive serializer. */
+	template<typename ElementType, typename KeyFuncs,typename Allocator>
+ 	static void SerializeStructured(FStructuredArchive::FSlot Slot, TSet<ElementType, KeyFuncs, Allocator>& Set)
+ 	{
+		Slot << Set.Elements;
+
+		if (Slot.GetUnderlyingArchive().IsLoading())
+		{
+			// Free the old hash.
+			Set.Hash.ResizeAllocation(0, 0, sizeof(FSetElementId));
+			Set.HashSize = 0;
+
+			// Hash the newly loaded elements.
+			Set.ConditionalRehash(Set.Elements.Num());
+		}
+ 	}
+
+	// Legacy comparison operators.  Note that these also test whether the set's elements were added in the same order!
+	template<typename ElementType, typename KeyFuncs,typename Allocator>
+	static bool LegacyCompareEqual(const TSet<ElementType, KeyFuncs, Allocator>& A, const TSet<ElementType, KeyFuncs, Allocator>& B)
+	{
+		return A.Elements == B.Elements;
+	}
+};
+
+/** Serializer. */
+template <typename ElementType>
+FORCEINLINE FArchive& operator<<(FArchive& Ar, TSetElement<ElementType>& Element)
+{
+	return Ar << Element.Value;
+}
+
+/** Structured archive serializer. */
+template <typename ElementType>
+FORCEINLINE void operator<<(FStructuredArchive::FSlot& Ar, TSetElement<ElementType>& Element)
+{
+	Ar << Element.Value;
+}
+
+/** Serializer. */
+template<typename ElementType, typename KeyFuncs,typename Allocator>
+FArchive& operator<<(FArchive& Ar, TSet<ElementType, KeyFuncs, Allocator>& Set)
+{
+	return TSetPrivateFriend::Serialize(Ar, Set);
+}
+
+/** Structured archive serializer. */
+template<typename ElementType, typename KeyFuncs,typename Allocator>
+void operator<<(FStructuredArchive::FSlot& Ar, TSet<ElementType, KeyFuncs, Allocator>& Set)
+{
+	TSetPrivateFriend::SerializeStructured(Ar, Set);
+}
+
+// Legacy comparison operators.  Note that these also test whether the set's elements were added in the same order!
+template<typename ElementType, typename KeyFuncs,typename Allocator>
+bool LegacyCompareEqual(const TSet<ElementType, KeyFuncs, Allocator>& A,const TSet<ElementType, KeyFuncs, Allocator>& B)
+{
+	return TSetPrivateFriend::LegacyCompareEqual(A, B);
+}
+template<typename ElementType, typename KeyFuncs,typename Allocator>
+bool LegacyCompareNotEqual(const TSet<ElementType, KeyFuncs, Allocator>& A,const TSet<ElementType, KeyFuncs, Allocator>& B)
+{
+	return !TSetPrivateFriend::LegacyCompareEqual(A, B);
+}
diff --git a/Engine/Source/Runtime/Core/Public/Containers/SortedMap.h b/Engine/Source/Runtime/Core/Public/Containers/SortedMap.h
--- a/Engine/Source/Runtime/Core/Public/Containers/SortedMap.h
+++ b/Engine/Source/Runtime/Core/Public/Containers/SortedMap.h
@@ -375,19 +375,6 @@ public:
 		}
 	}
 
-	/** Serializer. */
-	FORCEINLINE friend FArchive& operator<<(FArchive& Ar, TSortedMap& Map)
-	{
-		Ar << Map.Pairs;
-
-		if (Ar.IsLoading())
-		{
-			// We need to resort, in case the sorting is not consistent with what it was before
-			Algo::SortBy(Map.Pairs, FKeyForward(), SortPredicate());
-		}
-		return Ar;
-	}
-
 	/**
 	 * Describes the map's contents through an output device.
 	 *
@@ -556,8 +543,8 @@ private:
 			return !!PairIt; 
 		}
 
-		FORCEINLINE friend bool operator==(const TBaseIterator& Lhs, const TBaseIterator& Rhs) { return Lhs.PairIt == Rhs.PairIt; }
-		FORCEINLINE friend bool operator!=(const TBaseIterator& Lhs, const TBaseIterator& Rhs) { return Lhs.PairIt != Rhs.PairIt; }
+		FORCEINLINE bool operator==(const TBaseIterator& Rhs) const { return PairIt == Rhs.PairIt; }
+		FORCEINLINE bool operator!=(const TBaseIterator& Rhs) const { return PairIt != Rhs.PairIt; }
 
 		FORCEINLINE ItKeyType&   Key()   const { return PairIt->Key; }
 		FORCEINLINE ItValueType& Value() const { return PairIt->Value; }
@@ -603,8 +590,8 @@ private:
 			return Index != static_cast<SizeType>(-1);
 		}
 
-		FORCEINLINE friend bool operator==(const TBaseReverseIterator& Lhs, const TBaseReverseIterator& Rhs) { return Lhs.Index == Rhs.Index; }
-		FORCEINLINE friend bool operator!=(const TBaseReverseIterator& Lhs, const TBaseReverseIterator& Rhs) { return Lhs.Index != Rhs.Index; }
+		FORCEINLINE bool operator==(const TBaseReverseIterator& Rhs) const { return Index == Rhs.Index; }
+		FORCEINLINE bool operator!=(const TBaseReverseIterator& Rhs) const { return Index != Rhs.Index; }
 
 		FORCEINLINE ItKeyType& Key()   const { return Data[Index].Key; }
 		FORCEINLINE ItValueType& Value() const { return Data[Index].Value; }
@@ -779,6 +766,8 @@ public:
 	FORCEINLINE RangedForConstIteratorType	begin() const { return Pairs.begin(); }
 	FORCEINLINE RangedForIteratorType		end()         { return Pairs.end(); }
 	FORCEINLINE RangedForConstIteratorType	end() const   { return Pairs.end(); }
+
+	friend struct TSortedMapPrivateFriend;
 };
 
 DECLARE_TEMPLATE_INTRINSIC_TYPE_LAYOUT((template <typename KeyType, typename ValueType, typename ArrayAllocator, typename SortPredicate>), (TSortedMap<KeyType, ValueType, ArrayAllocator, SortPredicate>));
@@ -788,3 +777,27 @@ struct TContainerTraits<TSortedMap<KeyType, ValueType, ArrayAllocator, SortPredi
 {
 	enum { MoveWillEmptyContainer = TContainerTraits<typename TSortedMap<KeyType, ValueType, ArrayAllocator, SortPredicate>::ElementArrayType>::MoveWillEmptyContainer };
 };
+
+struct TSortedMapPrivateFriend
+{
+	template <typename KeyType, typename ValueType, typename ArrayAllocator, typename SortPredicate>
+	static void Serialize(FArchive& Ar, TSortedMap<KeyType, ValueType, ArrayAllocator, SortPredicate>& Map)
+	{
+		Ar << Map.Pairs;
+
+		if (Ar.IsLoading())
+		{
+			// We need to resort, in case the sorting is not consistent with what it was before
+			Algo::SortBy(Map.Pairs, typename TSortedMap<KeyType, ValueType, ArrayAllocator, SortPredicate>::FKeyForward(), SortPredicate());
+		}
+	}
+};
+
+/** Serializer. */
+template <typename KeyType, typename ValueType, typename ArrayAllocator, typename SortPredicate>
+FORCEINLINE FArchive& operator<<(FArchive& Ar, TSortedMap<KeyType, ValueType, ArrayAllocator, SortPredicate>& Map)
+{
+	TSortedMapPrivateFriend::Serialize(Ar, Map);
+	return Ar;
+}
+
diff --git a/Engine/Source/Runtime/Core/Public/Containers/SparseArray.h b/Engine/Source/Runtime/Core/Public/Containers/SparseArray.h
--- a/Engine/Source/Runtime/Core/Public/Containers/SparseArray.h
+++ b/Engine/Source/Runtime/Core/Public/Containers/SparseArray.h
@@ -621,73 +621,20 @@ public:
 		return NumFreeIndices == 0;
 	}
 
-	/** Serializer. */
-	friend FArchive& operator<<(FArchive& Ar,TSparseArray& Array)
-	{
-		Array.CountBytes(Ar);
-		if( Ar.IsLoading() )
-		{
-			// Load array.
-			int32 NewNumElements = 0;
-			Ar << NewNumElements;
-			Array.Empty( NewNumElements );
-			for(int32 ElementIndex = 0;ElementIndex < NewNumElements;ElementIndex++)
-			{
-				Ar << *::new(Array.AddUninitialized())ElementType;
-			}
-		}
-		else
-		{
-			// Save array.
-			int32 NewNumElements = Array.Num();
-			Ar << NewNumElements;
-			for(TIterator It(Array);It;++It)
-			{
-				Ar << *It;
-			}
-		}
-		return Ar;
-	}
-
-	/** Structured archive serializer. */
-	friend void operator<<(FStructuredArchive::FSlot Slot, TSparseArray& InArray)
-	{
-		int32 NumElements = InArray.Num();
-		FStructuredArchive::FArray Array = Slot.EnterArray(NumElements);
-		if (Slot.GetUnderlyingArchive().IsLoading())
-		{
-			InArray.Empty(NumElements);
-
-			for (int32 Index = 0; Index < NumElements; ++Index)
-			{
-				FStructuredArchive::FSlot ElementSlot = Array.EnterElement();
-				ElementSlot << *::new(InArray.AddUninitialized())ElementType;
-			}
-		}
-		else
-		{
-			for (TIterator It(InArray); It; ++It)
-			{
-				FStructuredArchive::FSlot ElementSlot = Array.EnterElement();
-				ElementSlot << *It;
-			}
-		}
-	}
-
 	/**
 	 * Equality comparison operator.
 	 * Checks that both arrays have the same elements and element indices; that means that unallocated elements are signifigant!
 	 */
-	friend bool operator==(const TSparseArray& A,const TSparseArray& B)
+	bool operator==(const TSparseArray& B) const
 	{
-		if(A.GetMaxIndex() != B.GetMaxIndex())
+		if(GetMaxIndex() != B.GetMaxIndex())
 		{
 			return false;
 		}
 
-		for(int32 ElementIndex = 0;ElementIndex < A.GetMaxIndex();ElementIndex++)
+		for(int32 ElementIndex = 0;ElementIndex < GetMaxIndex();ElementIndex++)
 		{
-			const bool bIsAllocatedA = A.IsAllocated(ElementIndex);
+			const bool bIsAllocatedA = IsAllocated(ElementIndex);
 			const bool bIsAllocatedB = B.IsAllocated(ElementIndex);
 			if(bIsAllocatedA != bIsAllocatedB)
 			{
@@ -695,7 +642,7 @@ public:
 			}
 			else if(bIsAllocatedA)
 			{
-				if(A[ElementIndex] != B[ElementIndex])
+				if((*this)[ElementIndex] != B[ElementIndex])
 				{
 					return false;
 				}
@@ -709,9 +656,9 @@ public:
 	 * Inequality comparison operator.
 	 * Checks that both arrays have the same elements and element indices; that means that unallocated elements are signifigant!
 	 */
-	friend bool operator!=(const TSparseArray& A,const TSparseArray& B)
+	bool operator!=(const TSparseArray& B) const
 	{
-		return !(A == B);
+		return !(*this == B);
 	}
 
 	/** Default constructor. */
@@ -889,8 +836,8 @@ private:
 
 		FORCEINLINE int32 GetIndex() const { return BitArrayIt.GetIndex(); }
 
-		FORCEINLINE friend bool operator==(const TBaseIterator& Lhs, const TBaseIterator& Rhs) { return Lhs.BitArrayIt == Rhs.BitArrayIt && &Lhs.Array == &Rhs.Array; }
-		FORCEINLINE friend bool operator!=(const TBaseIterator& Lhs, const TBaseIterator& Rhs) { return Lhs.BitArrayIt != Rhs.BitArrayIt || &Lhs.Array != &Rhs.Array; }
+		FORCEINLINE bool operator==(const TBaseIterator& Rhs) const { return BitArrayIt == Rhs.BitArrayIt && &Array == &Rhs.Array; }
+		FORCEINLINE bool operator!=(const TBaseIterator& Rhs) const { return BitArrayIt != Rhs.BitArrayIt || &Array != &Rhs.Array; }
 
 		/** conversion to "bool" returning true if the iterator is valid. */
 		FORCEINLINE explicit operator bool() const
@@ -961,19 +908,18 @@ public:
 			{
 			}
 
-		private:
-			int32 InitialNum;
-
-			friend FORCEINLINE bool operator!=(const TRangedForIterator& Lhs, const TRangedForIterator& Rhs)
+			FORCEINLINE bool operator!=(const TRangedForIterator& Rhs) const
 			{
 				// We only need to do the check in this operator, because no other operator will be
 				// called until after this one returns.
 				//
 				// Also, we should only need to check one side of this comparison - if the other iterator isn't
 				// even from the same array then the compiler has generated bad code.
-				ensureMsgf(Lhs.Array.Num() == Lhs.InitialNum, TEXT("Container has changed during ranged-for iteration!"));
-				return *(TIterator*)&Lhs != *(TIterator*)&Rhs;
+				ensureMsgf(this->Array.Num() == InitialNum, TEXT("Container has changed during ranged-for iteration!"));
+				return *(TIterator*)this != *(TIterator*)&Rhs;
 			}
+		private:
+			int32 InitialNum;
 		};
 
 		class TRangedForConstIterator : public TConstIterator
@@ -985,19 +931,18 @@ public:
 			{
 			}
 
-		private:
-			int32 InitialNum;
-
-			friend FORCEINLINE bool operator!=(const TRangedForConstIterator& Lhs, const TRangedForConstIterator& Rhs)
+			FORCEINLINE bool operator!=(const TRangedForConstIterator& Rhs) const
 			{
 				// We only need to do the check in this operator, because no other operator will be
 				// called until after this one returns.
 				//
 				// Also, we should only need to check one side of this comparison - if the other iterator isn't
 				// even from the same array then the compiler has generated bad code.
-				ensureMsgf(Lhs.Array.Num() == Lhs.InitialNum, TEXT("Container has changed during ranged-for iteration!"));
-				return *(TIterator*)&Lhs != *(TIterator*)&Rhs;
+				ensureMsgf(this->Array.Num() == InitialNum, TEXT("Container has changed during ranged-for iteration!"));
+				return *(TIterator*)this != *(TIterator*)&Rhs;
 			}
+		private:
+			int32 InitialNum;
 		};
 	#else
 		using TRangedForIterator      = TIterator;
@@ -1494,3 +1439,58 @@ inline void* operator new(size_t Size,const FSparseArrayAllocationInfo& Allocati
 	UE_ASSUME(Allocation.Pointer);
 	return Allocation.Pointer;
 }
+
+/** Serializer. */
+template<typename ElementType,typename Allocator>
+FArchive& operator<<(FArchive& Ar,TSparseArray<ElementType, Allocator>& Array)
+{
+	Array.CountBytes(Ar);
+	if( Ar.IsLoading() )
+	{
+		// Load array.
+		int32 NewNumElements = 0;
+		Ar << NewNumElements;
+		Array.Empty( NewNumElements );
+		for(int32 ElementIndex = 0;ElementIndex < NewNumElements;ElementIndex++)
+		{
+			Ar << *::new(Array.AddUninitialized())ElementType;
+		}
+	}
+	else
+	{
+		// Save array.
+		int32 NewNumElements = Array.Num();
+		Ar << NewNumElements;
+		for(typename TSparseArray<ElementType, Allocator>::TIterator It(Array);It;++It)
+		{
+			Ar << *It;
+		}
+	}
+	return Ar;
+}
+
+/** Structured archive serializer. */
+template<typename ElementType,typename Allocator>
+void operator<<(FStructuredArchive::FSlot Slot, TSparseArray<ElementType, Allocator>& InArray)
+{
+	int32 NumElements = InArray.Num();
+	FStructuredArchive::FArray Array = Slot.EnterArray(NumElements);
+	if (Slot.GetUnderlyingArchive().IsLoading())
+	{
+		InArray.Empty(NumElements);
+
+		for (int32 Index = 0; Index < NumElements; ++Index)
+		{
+			FStructuredArchive::FSlot ElementSlot = Array.EnterElement();
+			ElementSlot << *::new(InArray.AddUninitialized())ElementType;
+		}
+	}
+	else
+	{
+		for (typename TSparseArray<ElementType, Allocator>::TIterator It(InArray); It; ++It)
+		{
+			FStructuredArchive::FSlot ElementSlot = Array.EnterElement();
+			ElementSlot << *It;
+		}
+	}
+}
diff --git a/Engine/Source/Runtime/Core/Public/Containers/StaticArray.h b/Engine/Source/Runtime/Core/Public/Containers/StaticArray.h
--- a/Engine/Source/Runtime/Core/Public/Containers/StaticArray.h
+++ b/Engine/Source/Runtime/Core/Public/Containers/StaticArray.h
@@ -37,16 +37,6 @@ public:
 	TStaticArray& operator=(TStaticArray&& Other) = default;
 	TStaticArray& operator=(const TStaticArray& Other) = default;
 
-	/** Serializer. */
-	friend FArchive& operator<<(FArchive& Ar,TStaticArray& StaticArray)
-	{
-		for(uint32 Index = 0;Index < NumElements;++Index)
-		{
-			Ar << StaticArray[Index];
-		}
-		return Ar;
-	}
-
 	// Accessors.
 	FORCEINLINE_DEBUGGABLE InElementType& operator[](uint32 Index)
 	{
@@ -73,11 +63,11 @@ public:
 		return true;
 	}
 
-	friend bool operator!=(const TStaticArray& A,const TStaticArray& B)
+	bool operator!=(const TStaticArray& B) const
 	{
 		for(uint32 ElementIndex = 0;ElementIndex < NumElements;++ElementIndex)
 		{
-			if(!(A[ElementIndex] == B[ElementIndex]))
+			if(!((*this)[ElementIndex] == B[ElementIndex]))
 			{
 				return true;
 			}
@@ -103,18 +93,6 @@ public:
 	FORCEINLINE_DEBUGGABLE       InElementType* GetData()       { static_assert((alignof(ElementType) % Alignment) == 0, "GetData() cannot be called on a TStaticArray with non-standard alignment"); return &Storage.Elements[0].Element; }
 	FORCEINLINE_DEBUGGABLE const InElementType* GetData() const { static_assert((alignof(ElementType) % Alignment) == 0, "GetData() cannot be called on a TStaticArray with non-standard alignment"); return &Storage.Elements[0].Element; }
 
-	/** Hash function. */
-	friend uint32 GetTypeHash(const TStaticArray& Array)
-	{
-		uint32 Result = 0;
-		for(uint32 ElementIndex = 0;ElementIndex < NumElements;++ElementIndex)
-		{
-			Result ^= GetTypeHash(Array[ElementIndex]);
-		}
-		return Result;
-	}
-
-
 private:
 
 	struct alignas(Alignment) TArrayStorageElementAligned
@@ -177,9 +155,9 @@ public:
 			return *this;
 		}
 
-		friend bool operator!=(const FRangedForIterator& A, const FRangedForIterator& B)
+		bool operator!=(const FRangedForIterator& B) const
 		{
-			return A.Ptr != B.Ptr;
+			return Ptr != B.Ptr;
 		}
 
 	private:
@@ -214,3 +192,26 @@ struct TIsContiguousContainer<TStaticArray<ElementType, NumElements, Alignment>>
 {
 	enum { Value = (alignof(ElementType) % Alignment) == 0 };
 };
+
+/** Serializer. */
+template <typename ElementType, uint32 NumElements, uint32 Alignment>
+FArchive& operator<<(FArchive& Ar,TStaticArray<ElementType, NumElements, Alignment>& StaticArray)
+{
+	for(uint32 Index = 0;Index < NumElements;++Index)
+	{
+		Ar << StaticArray[Index];
+	}
+	return Ar;
+}
+
+/** Hash function. */
+template <typename ElementType, uint32 NumElements, uint32 Alignment>
+uint32 GetTypeHash(const TStaticArray<ElementType, NumElements, Alignment>& Array)
+{
+	uint32 Result = 0;
+	for(uint32 ElementIndex = 0;ElementIndex < NumElements;++ElementIndex)
+	{
+		Result ^= GetTypeHash(Array[ElementIndex]);
+	}
+	return Result;
+}
diff --git a/Engine/Source/Runtime/Core/Public/Containers/StaticBitArray.h b/Engine/Source/Runtime/Core/Public/Containers/StaticBitArray.h
--- a/Engine/Source/Runtime/Core/Public/Containers/StaticBitArray.h
+++ b/Engine/Source/Runtime/Core/Public/Containers/StaticBitArray.h
@@ -247,11 +247,11 @@ public:
 		Results ^= B;
 		return Results;
 	}
-	friend FORCEINLINE bool operator==(const TStaticBitArray<NumBits>& A, const TStaticBitArray<NumBits>& B)
+	FORCEINLINE bool operator==(const TStaticBitArray<NumBits>& B)
 	{
-		for(int32 Index = 0; Index < A.NumWords; ++Index)
+		for(int32 Index = 0; Index < NumWords; ++Index)
 		{
-			if(A.Words[Index] != B.Words[Index])
+			if(Words[Index] != B.Words[Index])
 			{
 				return false;
 			}
@@ -264,9 +264,9 @@ public:
 		return (UnspecifiedBoolType)A == Value;
 	}
 	/** != simple maps to == */
-	friend FORCEINLINE bool operator!=(const TStaticBitArray<NumBits>& A, const TStaticBitArray<NumBits>& B)
+	FORCEINLINE bool operator!=(const TStaticBitArray<NumBits>& B) const
 	{
-		return !(A == B);
+		return !(*this == B);
 	}
 	/** != simple maps to == */
 	friend FORCEINLINE bool operator!=(const TStaticBitArray<NumBits>& A, UnspecifiedBoolType Value)
@@ -335,24 +335,6 @@ public:
 		return INDEX_NONE;
 	}
 
-	/**
-	 * Serializer.
-	 */
-	friend FArchive& operator<<(FArchive& Ar, TStaticBitArray& BitArray)
-	{
-		uint32 ArchivedNumWords = BitArray.NumWords;
-		Ar << ArchivedNumWords;
-
-		if(Ar.IsLoading())
-		{
-			FMemory::Memset(BitArray.Words, 0, sizeof(BitArray.Words));
-			ArchivedNumWords = FMath::Min(BitArray.NumWords, ArchivedNumWords);
-		}
-
-		Ar.Serialize(BitArray.Words, ArchivedNumWords * sizeof(BitArray.Words[0]));
-		return Ar;
-	}
-
 	/**
 	 * Converts the bitarray to a string representing the binary representation of the array
 	 */
@@ -415,3 +397,22 @@ private:
 		LowLevelFatalError(TEXT("invalid TStaticBitArray<NumBits> character"));
 	}
 };
+
+/**
+	* Serializer.
+	*/
+template<uint32 NumBits>
+FArchive& operator<<(FArchive& Ar, TStaticBitArray<NumBits>& BitArray)
+{
+	uint32 ArchivedNumWords = BitArray.NumWords;
+	Ar << ArchivedNumWords;
+
+	if(Ar.IsLoading())
+	{
+		FMemory::Memset(BitArray.Words, 0, sizeof(BitArray.Words));
+		ArchivedNumWords = FMath::Min(BitArray.NumWords, ArchivedNumWords);
+	}
+
+	Ar.Serialize(BitArray.Words, ArchivedNumWords * sizeof(BitArray.Words[0]));
+	return Ar;
+}
diff --git a/Engine/Source/Runtime/Core/Public/Containers/Union.h b/Engine/Source/Runtime/Core/Public/Containers/Union.h
--- a/Engine/Source/Runtime/Core/Public/Containers/Union.h
+++ b/Engine/Source/Runtime/Core/Public/Containers/Union.h
@@ -286,41 +286,40 @@ public:
 		return false;
 	}
 
-	friend FArchive& operator<<(FArchive& Ar,TUnion& Union)
+	void Serialize(FArchive& Ar)
 	{
 		if(Ar.IsLoading())
 		{
-			Union.Reset();
+			Reset();
 
-			Ar << Union.CurrentSubtypeIndex;
+			Ar << CurrentSubtypeIndex;
 
-			switch(Union.CurrentSubtypeIndex)
+			switch(CurrentSubtypeIndex)
 			{
-			case 0: Ar << Union.InitSubtype<TypeA>(); break;
-			case 1: Ar << Union.InitSubtype<TypeB>(); break;
-			case 2: Ar << Union.InitSubtype<TypeC>(); break;
-			case 3: Ar << Union.InitSubtype<TypeD>(); break;
-			case 4: Ar << Union.InitSubtype<TypeE>(); break;
-			case 5: Ar << Union.InitSubtype<TypeF>(); break;
+			case 0: Ar << InitSubtype<TypeA>(); break;
+			case 1: Ar << InitSubtype<TypeB>(); break;
+			case 2: Ar << InitSubtype<TypeC>(); break;
+			case 3: Ar << InitSubtype<TypeD>(); break;
+			case 4: Ar << InitSubtype<TypeE>(); break;
+			case 5: Ar << InitSubtype<TypeF>(); break;
 			default: FatalErrorUndefinedSubtype(); break;
 			};
 		}
 		else
 		{
-			Ar << Union.CurrentSubtypeIndex;
+			Ar << CurrentSubtypeIndex;
 
-			switch(Union.CurrentSubtypeIndex)
+			switch(CurrentSubtypeIndex)
 			{
-			case 0: Ar << Union.GetSubtype<TypeA>(); break;
-			case 1: Ar << Union.GetSubtype<TypeB>(); break;
-			case 2: Ar << Union.GetSubtype<TypeC>(); break;
-			case 3: Ar << Union.GetSubtype<TypeD>(); break;
-			case 4: Ar << Union.GetSubtype<TypeE>(); break;
-			case 5: Ar << Union.GetSubtype<TypeF>(); break;
+			case 0: Ar << GetSubtype<TypeA>(); break;
+			case 1: Ar << GetSubtype<TypeB>(); break;
+			case 2: Ar << GetSubtype<TypeC>(); break;
+			case 3: Ar << GetSubtype<TypeD>(); break;
+			case 4: Ar << GetSubtype<TypeE>(); break;
+			case 5: Ar << GetSubtype<TypeF>(); break;
 			default: FatalErrorUndefinedSubtype(); break;
 			};
 		}
-		return Ar;
 	}
 
 private:
@@ -405,3 +404,10 @@ private:
 		UE_LOG(LogUnion, Fatal, TEXT("Unrecognized TUnion subtype"));
 	}
 };
+
+template<typename TypeA,typename TypeB,typename TypeC,typename TypeD,typename TypeE,typename TypeF>
+FArchive& operator<<(FArchive& Ar, TUnion<TypeA, TypeB, TypeC, TypeD, TypeE, TypeF>& Union)
+{
+	Union.Serialize(Ar);
+	return Ar;
+}
diff --git a/Engine/Source/Runtime/Core/Public/Containers/UnrealString.h b/Engine/Source/Runtime/Core/Public/Containers/UnrealString.h
--- a/Engine/Source/Runtime/Core/Public/Containers/UnrealString.h
+++ b/Engine/Source/Runtime/Core/Public/Containers/UnrealString.h
@@ -843,9 +843,9 @@ public:
 	 * @return true if the left string is lexicographically == the right string, otherwise false
 	 * @note case insensitive
 	 */
-	UE_NODISCARD FORCEINLINE friend bool operator==(const FString& Lhs, const FString& Rhs)
+	UE_NODISCARD FORCEINLINE bool operator==(const FString& Rhs) const
 	{
-		return Lhs.Equals(Rhs, ESearchCase::IgnoreCase);
+		return Equals(Rhs, ESearchCase::IgnoreCase);
 	}
 
 	/**
@@ -857,9 +857,9 @@ public:
 	 * @note case insensitive
 	 */
 	template <typename CharType>
-	UE_NODISCARD FORCEINLINE friend bool operator==(const FString& Lhs, const CharType* Rhs)
+	UE_NODISCARD FORCEINLINE bool operator==(const CharType* Rhs) const
 	{
-		return FPlatformString::Stricmp(*Lhs, Rhs) == 0;
+		return FPlatformString::Stricmp(**this, Rhs) == 0;
 	}
 
 	/**
@@ -884,9 +884,9 @@ public:
 	 * @return true if the left string is lexicographically != the right string, otherwise false
 	 * @note case insensitive
 	 */
-	UE_NODISCARD FORCEINLINE friend bool operator!=(const FString& Lhs, const FString& Rhs)
+	UE_NODISCARD FORCEINLINE bool operator!=(const FString& Rhs) const
 	{
-		return !(Lhs == Rhs);
+		return !Equals(Rhs, ESearchCase::IgnoreCase);
 	}
 
 	/**
@@ -898,9 +898,9 @@ public:
 	 * @note case insensitive
 	 */
 	template <typename CharType>
-	UE_NODISCARD FORCEINLINE friend bool operator!=(const FString& Lhs, const CharType* Rhs)
+	UE_NODISCARD FORCEINLINE bool operator!=(const CharType* Rhs) const
 	{
-		return FPlatformString::Stricmp(*Lhs, Rhs) != 0;
+		return FPlatformString::Stricmp(**this, Rhs) != 0;
 	}
 
 	/**
diff --git a/Engine/Source/Runtime/Core/Public/Core.h b/Engine/Source/Runtime/Core/Public/Core.h
--- a/Engine/Source/Runtime/Core/Public/Core.h
+++ b/Engine/Source/Runtime/Core/Public/Core.h
@@ -127,7 +127,6 @@ MONOLITHIC_HEADER_BOILERPLATE()
 #include "Containers/Array.h"
 #include "Containers/ScriptArray.h"
 #include "Containers/MRUArray.h"
-#include "Containers/TransArray.h"
 #include "Containers/IndirectArray.h"
 #include "Misc/ITransaction.h"
 #include "Containers/ArrayBuilder.h"
diff --git a/Engine/Source/Runtime/Core/Public/Delegates/IDelegateInstance.h b/Engine/Source/Runtime/Core/Public/Delegates/IDelegateInstance.h
--- a/Engine/Source/Runtime/Core/Public/Delegates/IDelegateInstance.h
+++ b/Engine/Source/Runtime/Core/Public/Delegates/IDelegateInstance.h
@@ -42,15 +42,14 @@ public:
 		ID = 0;
 	}
 
-private:
-	friend bool operator==(const FDelegateHandle& Lhs, const FDelegateHandle& Rhs)
+	bool operator==(const FDelegateHandle& Rhs) const
 	{
-		return Lhs.ID == Rhs.ID;
+		return ID == Rhs.ID;
 	}
 
-	friend bool operator!=(const FDelegateHandle& Lhs, const FDelegateHandle& Rhs)
+	bool operator!=(const FDelegateHandle& Rhs) const
 	{
-		return Lhs.ID != Rhs.ID;
+		return ID != Rhs.ID;
 	}
 
 	friend FORCEINLINE uint32 GetTypeHash(const FDelegateHandle& Key)
@@ -58,6 +57,8 @@ private:
 		return GetTypeHash(Key.ID);
 	}
 
+private:
+
 	/**
 	 * Generates a new ID for use the delegate handle.
 	 *
diff --git a/Engine/Source/Runtime/Core/Public/HAL/Platform.h b/Engine/Source/Runtime/Core/Public/HAL/Platform.h
--- a/Engine/Source/Runtime/Core/Public/HAL/Platform.h
+++ b/Engine/Source/Runtime/Core/Public/HAL/Platform.h
@@ -238,6 +238,9 @@
 #ifndef PLATFORM_COMPILER_HAS_FOLD_EXPRESSIONS
 	#define PLATFORM_COMPILER_HAS_FOLD_EXPRESSIONS 0
 #endif
+#ifndef PLATFORM_COMPILER_HAS_GENERATED_COMPARISON_OPERATORS
+	#define PLATFORM_COMPILER_HAS_GENERATED_COMPARISON_OPERATORS (__cplusplus >= 202002L)
+#endif
 #ifndef PLATFORM_TCHAR_IS_4_BYTES
 	#define PLATFORM_TCHAR_IS_4_BYTES			0
 #endif
diff --git a/Engine/Source/Runtime/Core/Public/Hash/Blake3.h b/Engine/Source/Runtime/Core/Public/Hash/Blake3.h
--- a/Engine/Source/Runtime/Core/Public/Hash/Blake3.h
+++ b/Engine/Source/Runtime/Core/Public/Hash/Blake3.h
@@ -55,22 +55,19 @@ public:
 	/** A zero hash. */
 	static const FBlake3Hash Zero;
 
-private:
-	alignas(uint32) ByteArray Hash{};
-
-	friend inline bool operator==(const FBlake3Hash& A, const FBlake3Hash& B)
+	inline bool operator==(const FBlake3Hash& B) const
 	{
-		return FMemory::Memcmp(A.GetBytes(), B.GetBytes(), sizeof(decltype(A.GetBytes()))) == 0;
+		return FMemory::Memcmp(GetBytes(), B.GetBytes(), sizeof(decltype(GetBytes()))) == 0;
 	}
 
-	friend inline bool operator!=(const FBlake3Hash& A, const FBlake3Hash& B)
+	inline bool operator!=(const FBlake3Hash& B) const
 	{
-		return FMemory::Memcmp(A.GetBytes(), B.GetBytes(), sizeof(decltype(A.GetBytes()))) != 0;
+		return FMemory::Memcmp(GetBytes(), B.GetBytes(), sizeof(decltype(GetBytes()))) != 0;
 	}
 
-	friend inline bool operator<(const FBlake3Hash& A, const FBlake3Hash& B)
+	inline bool operator<(const FBlake3Hash& B) const
 	{
-		return FMemory::Memcmp(A.GetBytes(), B.GetBytes(), sizeof(decltype(A.GetBytes()))) < 0;
+		return FMemory::Memcmp(GetBytes(), B.GetBytes(), sizeof(decltype(GetBytes()))) < 0;
 	}
 
 	friend inline FArchive& operator<<(FArchive& Ar, FBlake3Hash& Value)
@@ -83,6 +80,9 @@ private:
 	{
 		return *reinterpret_cast<const uint32*>(Value.GetBytes());
 	}
+
+private:
+	alignas(uint32) ByteArray Hash{};
 };
 
 inline const FBlake3Hash FBlake3Hash::Zero;
diff --git a/Engine/Source/Runtime/Core/Public/Hash/xxhash.h b/Engine/Source/Runtime/Core/Public/Hash/xxhash.h
--- a/Engine/Source/Runtime/Core/Public/Hash/xxhash.h
+++ b/Engine/Source/Runtime/Core/Public/Hash/xxhash.h
@@ -36,6 +36,41 @@ struct FXxHash64
 		FMemory::Memcpy(Bytes, &HashBigEndian, sizeof(uint64));
 	}
 
+	inline bool operator==(const FXxHash64& B) const
+	{
+		return Hash == B.Hash;
+	}
+
+	inline bool operator!=(const FXxHash64& B) const
+	{
+		return Hash != B.Hash;
+	}
+
+	inline bool operator<(const FXxHash64& B) const
+	{
+		return Hash < B.Hash;
+	}
+
+	friend inline uint32 GetTypeHash(const FXxHash64& InHash)
+	{
+		return uint32(InHash.Hash);
+	}
+
+	friend inline FArchive& operator<<(FArchive& Ar, FXxHash64& InHash)
+	{
+		return Ar << InHash.Hash;
+	}
+
+	template <typename CharType>
+	friend inline TStringBuilderBase<CharType>& operator<<(TStringBuilderBase<CharType>& Builder, const FXxHash64& InHash)
+	{
+		uint8 Bytes[8];
+		InHash.ToByteArray(Bytes);
+		UE::String::BytesToHexLower(MakeArrayView(Bytes), Builder);
+		return Builder;
+	}
+
+
 	/**
 	 * The hash in its native representation.
 	 *
@@ -79,6 +114,40 @@ public:
 	 * Use the canonical representation from ToByteArray to serialize or display the hash.
 	 */
 	uint64 HashHigh{};
+
+	inline bool operator==(const FXxHash128& B) const
+	{
+		return HashLow == B.HashLow && HashHigh == B.HashHigh;
+	}
+
+	inline bool operator!=(const FXxHash128& B) const
+	{
+		return HashLow != B.HashLow || HashHigh != B.HashHigh;
+	}
+
+	inline bool operator<(const FXxHash128& B) const
+	{
+		return HashHigh != B.HashHigh ? HashHigh < B.HashHigh : HashLow < B.HashLow;
+	}
+
+	friend inline uint32 GetTypeHash(const FXxHash128& Hash)
+	{
+		return uint32(Hash.HashLow);
+	}
+
+	friend inline FArchive& operator<<(FArchive& Ar, FXxHash128& Hash)
+	{
+		return Ar << Hash.HashLow << Hash.HashHigh;
+	}
+
+	template <typename CharType>
+	friend inline TStringBuilderBase<CharType>& operator<<(TStringBuilderBase<CharType>& Builder, const FXxHash128& Hash)
+	{
+		uint8 Bytes[16];
+		Hash.ToByteArray(Bytes);
+		UE::String::BytesToHexLower(MakeArrayView(Bytes), Builder);
+		return Builder;
+	}
 };
 
 /** Calculates a 64-bit hash with XXH3. */
@@ -122,71 +191,3 @@ public:
 private:
 	alignas(64) char StateBytes[576];
 };
-
-inline bool operator==(const FXxHash64& A, const FXxHash64& B)
-{
-	return A.Hash == B.Hash;
-}
-
-inline bool operator!=(const FXxHash64& A, const FXxHash64& B)
-{
-	return A.Hash != B.Hash;
-}
-
-inline bool operator<(const FXxHash64& A, const FXxHash64& B)
-{
-	return A.Hash < B.Hash;
-}
-
-inline uint32 GetTypeHash(const FXxHash64& Hash)
-{
-	return uint32(Hash.Hash);
-}
-
-inline FArchive& operator<<(FArchive& Ar, FXxHash64& Hash)
-{
-	return Ar << Hash.Hash;
-}
-
-template <typename CharType>
-inline TStringBuilderBase<CharType>& operator<<(TStringBuilderBase<CharType>& Builder, const FXxHash64& Hash)
-{
-	uint8 Bytes[8];
-	Hash.ToByteArray(Bytes);
-	UE::String::BytesToHexLower(MakeArrayView(Bytes), Builder);
-	return Builder;
-}
-
-inline bool operator==(const FXxHash128& A, const FXxHash128& B)
-{
-	return A.HashLow == B.HashLow && A.HashHigh == B.HashHigh;
-}
-
-inline bool operator!=(const FXxHash128& A, const FXxHash128& B)
-{
-	return A.HashLow != B.HashLow || A.HashHigh != B.HashHigh;
-}
-
-inline bool operator<(const FXxHash128& A, const FXxHash128& B)
-{
-	return A.HashHigh != B.HashHigh ? A.HashHigh < B.HashHigh : A.HashLow < B.HashLow;
-}
-
-inline uint32 GetTypeHash(const FXxHash128& Hash)
-{
-	return uint32(Hash.HashLow);
-}
-
-inline FArchive& operator<<(FArchive& Ar, FXxHash128& Hash)
-{
-	return Ar << Hash.HashLow << Hash.HashHigh;
-}
-
-template <typename CharType>
-inline TStringBuilderBase<CharType>& operator<<(TStringBuilderBase<CharType>& Builder, const FXxHash128& Hash)
-{
-	uint8 Bytes[16];
-	Hash.ToByteArray(Bytes);
-	UE::String::BytesToHexLower(MakeArrayView(Bytes), Builder);
-	return Builder;
-}
diff --git a/Engine/Source/Runtime/Core/Public/IO/IoHash.h b/Engine/Source/Runtime/Core/Public/IO/IoHash.h
--- a/Engine/Source/Runtime/Core/Public/IO/IoHash.h
+++ b/Engine/Source/Runtime/Core/Public/IO/IoHash.h
@@ -91,6 +91,12 @@ private:
 	{
 		return *reinterpret_cast<const uint32*>(Value.GetBytes());
 	}
+
+	friend inline FArchive& operator<<(FArchive& Ar, FIoHash& InHash)
+	{
+		Ar.Serialize(InHash.GetBytes(), sizeof(decltype(InHash.GetBytes())));
+		return Ar;
+	}
 };
 
 inline const FIoHash FIoHash::Zero;
@@ -162,12 +168,6 @@ inline FIoHash FIoHash::HashBuffer(const FCompositeBuffer& Buffer)
 	return FBlake3::HashBuffer(Buffer);
 }
 
-inline FArchive& operator<<(FArchive& Ar, FIoHash& Hash)
-{
-	Ar.Serialize(Hash.GetBytes(), sizeof(decltype(Hash.GetBytes())));
-	return Ar;
-}
-
 template <typename CharType>
 inline TStringBuilderBase<CharType>& operator<<(TStringBuilderBase<CharType>& Builder, const FIoHash& Hash)
 {
diff --git a/Engine/Source/Runtime/Core/Public/Internationalization/GatherableTextData.h b/Engine/Source/Runtime/Core/Public/Internationalization/GatherableTextData.h
--- a/Engine/Source/Runtime/Core/Public/Internationalization/GatherableTextData.h
+++ b/Engine/Source/Runtime/Core/Public/Internationalization/GatherableTextData.h
@@ -23,10 +23,10 @@ struct FTextSourceSiteContext
 	bool IsOptional;
 	FLocMetadataObject InfoMetaData;
 	FLocMetadataObject KeyMetaData;
-};
 
-CORE_API FArchive& operator<<(FArchive& Archive, FTextSourceSiteContext& This);
-CORE_API void operator<<(FStructuredArchive::FSlot Slot, FTextSourceSiteContext& This);
+	CORE_API friend FArchive& operator<<(FArchive& Archive, FTextSourceSiteContext& This);
+	CORE_API friend void operator<<(FStructuredArchive::FSlot Slot, FTextSourceSiteContext& This);
+};
 
 struct FTextSourceData
 {
@@ -43,7 +43,7 @@ struct FGatherableTextData
 	FTextSourceData SourceData;
 
 	TArray<FTextSourceSiteContext> SourceSiteContexts;
-};
 
-CORE_API FArchive& operator<<(FArchive& Archive, FGatherableTextData& This);
-CORE_API void operator<<(FStructuredArchive::FSlot Slot, FGatherableTextData& This);
+	CORE_API friend FArchive& operator<<(FArchive& Archive, FGatherableTextData& This);
+	CORE_API friend void operator<<(FStructuredArchive::FSlot Slot, FGatherableTextData& This);
+};
diff --git a/Engine/Source/Runtime/Core/Public/Internationalization/InternationalizationMetadata.h b/Engine/Source/Runtime/Core/Public/Internationalization/InternationalizationMetadata.h
--- a/Engine/Source/Runtime/Core/Public/Internationalization/InternationalizationMetadata.h
+++ b/Engine/Source/Runtime/Core/Public/Internationalization/InternationalizationMetadata.h
@@ -165,6 +165,9 @@ public:
 
 	FString ToString() const;
 
+	CORE_API friend FArchive& operator<<(FArchive& Archive, FLocMetadataObject& Object);
+	CORE_API friend void operator<<(FStructuredArchive::FSlot Slot, FLocMetadataObject& Object);
+
 public:
 	/** Stores the name/value pairs for the metadata object */
 	TMap< FString, TSharedPtr<FLocMetadataValue> > Values;
@@ -173,9 +176,6 @@ public:
 	static const TCHAR* COMPARISON_MODIFIER_PREFIX;
 };
 
-CORE_API FArchive& operator<<(FArchive& Archive, FLocMetadataObject& Object);
-CORE_API void operator<<(FStructuredArchive::FSlot Slot, FLocMetadataObject& Object);
-
 /** A LocMetadata String Value. */
 class CORE_API FLocMetadataValueString : public FLocMetadataValue
 {
diff --git a/Engine/Source/Runtime/Core/Public/Memory/SharedBuffer.h b/Engine/Source/Runtime/Core/Public/Memory/SharedBuffer.h
--- a/Engine/Source/Runtime/Core/Public/Memory/SharedBuffer.h
+++ b/Engine/Source/Runtime/Core/Public/Memory/SharedBuffer.h
@@ -319,10 +319,13 @@ private:
 
 	FOwnerPtrType Owner;
 
+public:
 	friend inline uint32 GetTypeHash(const FUniqueBuffer& Buffer) { return PointerHash(ToPrivateOwnerPtr(Buffer).Get()); }
 
-	friend inline bool operator==(const FUniqueBuffer& BufferA, const FUniqueBuffer& BufferB) { return ToPrivateOwnerPtr(BufferA) == ToPrivateOwnerPtr(BufferB); }
-	friend inline bool operator!=(const FUniqueBuffer& BufferA, const FUniqueBuffer& BufferB) { return ToPrivateOwnerPtr(BufferA) != ToPrivateOwnerPtr(BufferB); }
+	inline bool operator==(const FUniqueBuffer& BufferB) const { return ToPrivateOwnerPtr(*this) == ToPrivateOwnerPtr(BufferB); }
+#if !PLATFORM_COMPILER_HAS_GENERATED_COMPARISON_OPERATORS
+	inline bool operator!=(const FUniqueBuffer& BufferB) const { return ToPrivateOwnerPtr(*this) != ToPrivateOwnerPtr(BufferB); }
+#endif
 };
 
 ///////////////////////////////////////////////////////////////////////////////////////////////////////////////////////
@@ -436,14 +439,17 @@ private:
 
 	FOwnerPtrType Owner;
 
+public:
 	friend inline uint32 GetTypeHash(const FSharedBuffer& Buffer) { return PointerHash(ToPrivateOwnerPtr(Buffer).Get()); }
-	friend inline bool operator==(const FSharedBuffer& BufferA, const FSharedBuffer& BufferB) { return ToPrivateOwnerPtr(BufferA) == ToPrivateOwnerPtr(BufferB); }
-	friend inline bool operator!=(const FSharedBuffer& BufferA, const FSharedBuffer& BufferB) { return ToPrivateOwnerPtr(BufferA) != ToPrivateOwnerPtr(BufferB); }
-
+	inline bool operator==(const FSharedBuffer& BufferB) const { return ToPrivateOwnerPtr(*this) == ToPrivateOwnerPtr(BufferB); }
 	friend inline bool operator==(const FSharedBuffer& BufferA, const FUniqueBuffer& BufferB) { return ToPrivateOwnerPtr(BufferA) == ToPrivateOwnerPtr(BufferB); }
+
+#if !PLATFORM_COMPILER_HAS_GENERATED_COMPARISON_OPERATORS
+	inline bool operator!=(const FSharedBuffer& BufferB) const { return ToPrivateOwnerPtr(*this) != ToPrivateOwnerPtr(BufferB); }
 	friend inline bool operator!=(const FSharedBuffer& BufferA, const FUniqueBuffer& BufferB) { return ToPrivateOwnerPtr(BufferA) != ToPrivateOwnerPtr(BufferB); }
 	friend inline bool operator==(const FUniqueBuffer& BufferA, const FSharedBuffer& BufferB) { return ToPrivateOwnerPtr(BufferA) == ToPrivateOwnerPtr(BufferB); }
 	friend inline bool operator!=(const FUniqueBuffer& BufferA, const FSharedBuffer& BufferB) { return ToPrivateOwnerPtr(BufferA) != ToPrivateOwnerPtr(BufferB); }
+#endif
 };
 
 ///////////////////////////////////////////////////////////////////////////////////////////////////////////////////////
@@ -478,20 +484,22 @@ private:
 
 	FWeakOwnerPtrType Owner;
 
+public:
 	friend inline uint32 GetTypeHash(const FWeakSharedBuffer& Buffer) { return PointerHash(ToPrivateOwnerPtr(Buffer).Get()); }
+	inline bool operator==(const FWeakSharedBuffer& BufferB) const { return ToPrivateOwnerPtr(*this) == ToPrivateOwnerPtr(BufferB); }
+	friend inline bool operator==(const FWeakSharedBuffer& BufferA, const FUniqueBuffer& BufferB) { return ToPrivateOwnerPtr(BufferA) == ToPrivateOwnerPtr(BufferB); }
+	friend inline bool operator==(const FWeakSharedBuffer& BufferA, const FSharedBuffer& BufferB) { return ToPrivateOwnerPtr(BufferA) == ToPrivateOwnerPtr(BufferB); }
 
-	friend inline bool operator==(const FWeakSharedBuffer& BufferA, const FWeakSharedBuffer& BufferB) { return ToPrivateOwnerPtr(BufferA) == ToPrivateOwnerPtr(BufferB); }
-	friend inline bool operator!=(const FWeakSharedBuffer& BufferA, const FWeakSharedBuffer& BufferB) { return ToPrivateOwnerPtr(BufferA) != ToPrivateOwnerPtr(BufferB); }
+#if !PLATFORM_COMPILER_HAS_GENERATED_COMPARISON_OPERATORS
+	friend inline bool operator==(const FUniqueBuffer& BufferA, const FWeakSharedBuffer& BufferB) { return ToPrivateOwnerPtr(BufferA) == ToPrivateOwnerPtr(BufferB); }
+	friend inline bool operator==(const FSharedBuffer& BufferA, const FWeakSharedBuffer& BufferB) { return ToPrivateOwnerPtr(BufferA) == ToPrivateOwnerPtr(BufferB); }
 
-	friend inline bool operator==(const FWeakSharedBuffer& BufferA, const FUniqueBuffer& BufferB) { return ToPrivateOwnerPtr(BufferA) == ToPrivateOwnerPtr(BufferB); }
+	inline bool operator!=(const FWeakSharedBuffer& BufferB) const { return ToPrivateOwnerPtr(*this) != ToPrivateOwnerPtr(BufferB); }
 	friend inline bool operator!=(const FWeakSharedBuffer& BufferA, const FUniqueBuffer& BufferB) { return ToPrivateOwnerPtr(BufferA) != ToPrivateOwnerPtr(BufferB); }
-	friend inline bool operator==(const FUniqueBuffer& BufferA, const FWeakSharedBuffer& BufferB) { return ToPrivateOwnerPtr(BufferA) == ToPrivateOwnerPtr(BufferB); }
 	friend inline bool operator!=(const FUniqueBuffer& BufferA, const FWeakSharedBuffer& BufferB) { return ToPrivateOwnerPtr(BufferA) != ToPrivateOwnerPtr(BufferB); }
-
-	friend inline bool operator==(const FWeakSharedBuffer& BufferA, const FSharedBuffer& BufferB) { return ToPrivateOwnerPtr(BufferA) == ToPrivateOwnerPtr(BufferB); }
 	friend inline bool operator!=(const FWeakSharedBuffer& BufferA, const FSharedBuffer& BufferB) { return ToPrivateOwnerPtr(BufferA) != ToPrivateOwnerPtr(BufferB); }
-	friend inline bool operator==(const FSharedBuffer& BufferA, const FWeakSharedBuffer& BufferB) { return ToPrivateOwnerPtr(BufferA) == ToPrivateOwnerPtr(BufferB); }
 	friend inline bool operator!=(const FSharedBuffer& BufferA, const FWeakSharedBuffer& BufferB) { return ToPrivateOwnerPtr(BufferA) != ToPrivateOwnerPtr(BufferB); }
+#endif
 };
 
 ///////////////////////////////////////////////////////////////////////////////////////////////////////////////////////
diff --git a/Engine/Source/Runtime/Core/Public/Misc/ConfigCacheIni.h b/Engine/Source/Runtime/Core/Public/Misc/ConfigCacheIni.h
--- a/Engine/Source/Runtime/Core/Public/Misc/ConfigCacheIni.h
+++ b/Engine/Source/Runtime/Core/Public/Misc/ConfigCacheIni.h
@@ -359,8 +359,6 @@ public:
 	friend FArchive& operator<<(FArchive& Ar, FConfigSection& ConfigSection);
 };
 
-FArchive& operator<<(FArchive& Ar, FConfigSection& ConfigSection);
-
 
 #if ALLOW_INI_OVERRIDE_FROM_COMMANDLINE
 // Options which stemmed from the commandline
@@ -623,8 +621,6 @@ private:
 	friend FConfigContext;
 };
 
-FArchive& operator<<(FArchive& Ar, FConfigFile& ConfigFile);
-
 /**
  * Declares a delegate type that's used by the config system to allow iteration of key value pairs.
  */
diff --git a/Engine/Source/Runtime/Core/Public/Misc/Optional.h b/Engine/Source/Runtime/Core/Public/Misc/Optional.h
--- a/Engine/Source/Runtime/Core/Public/Misc/Optional.h
+++ b/Engine/Source/Runtime/Core/Public/Misc/Optional.h
@@ -166,33 +166,32 @@ public:
 		return !(lhs == rhs);
 	}
 
-	friend FArchive& operator<<(FArchive& Ar, TOptional& Optional)
+	void Serialize(FArchive& Ar)
 	{
-		bool bOptionalIsSet = Optional.bIsSet;
+		bool bOptionalIsSet = bIsSet;
 		Ar << bOptionalIsSet;
 		if (Ar.IsLoading())
 		{
 			if (bOptionalIsSet)
 			{
-				if (!Optional.bIsSet)
+				if (!bIsSet)
 				{
-					Optional.Emplace();
+					Emplace();
 				}
-				Ar << Optional.GetValue();
+				Ar << GetValue();
 			}
 			else
 			{
-				Optional.Reset();
+				Reset();
 			}
 		}
 		else
 		{
 			if (bOptionalIsSet)
 			{
-				Ar << Optional.GetValue();
+				Ar << GetValue();
 			}
 		}
-		return Ar;
 	}
 
 	/** @return true when the value is meaningful; false if calling GetValue() is undefined. */
@@ -220,3 +219,10 @@ private:
 	TTypeCompatibleBytes<OptionalType> Value;
 	bool bIsSet;
 };
+
+template<typename OptionalType>
+FArchive& operator<<(FArchive& Ar, TOptional<OptionalType>& Optional)
+{
+	Optional.Serialize(Ar);
+	return Ar;
+}
diff --git a/Engine/Source/Runtime/Core/Public/Misc/SecureHash.h b/Engine/Source/Runtime/Core/Public/Misc/SecureHash.h
--- a/Engine/Source/Runtime/Core/Public/Misc/SecureHash.h
+++ b/Engine/Source/Runtime/Core/Public/Misc/SecureHash.h
@@ -274,6 +274,10 @@ public:
 
 	friend CORE_API FString LexToString(const FSHAHash&);
 	friend CORE_API void LexFromString(FSHAHash& Hash, const TCHAR*);
+	friend inline FStringBuilderBase& operator<<(FStringBuilderBase& Builder, const FSHAHash& InHash) { UE::String::BytesToHex(InHash.Hash, Builder); return Builder; }
+	friend inline FAnsiStringBuilderBase& operator<<(FAnsiStringBuilderBase& Builder, const FSHAHash& InHash) { UE::String::BytesToHex(InHash.Hash, Builder); return Builder; }
+	CORE_API friend FCbWriter& operator<<(FCbWriter& Writer, const FSHAHash& Hash);
+	CORE_API friend bool LoadFromCompactBinary(FCbFieldView Field, FSHAHash& OutHash);
 };
 
 namespace Freeze
@@ -283,11 +287,6 @@ namespace Freeze
 
 DECLARE_INTRINSIC_TYPE_LAYOUT(FSHAHash);
 
-inline FStringBuilderBase& operator<<(FStringBuilderBase& Builder, const FSHAHash& Hash) { UE::String::BytesToHex(Hash.Hash, Builder); return Builder; }
-inline FAnsiStringBuilderBase& operator<<(FAnsiStringBuilderBase& Builder, const FSHAHash& Hash) { UE::String::BytesToHex(Hash.Hash, Builder); return Builder; }
-CORE_API FCbWriter& operator<<(FCbWriter& Writer, const FSHAHash& Hash);
-CORE_API bool LoadFromCompactBinary(FCbFieldView Field, FSHAHash& OutHash);
-
 class CORE_API FSHA1
 {
 public:
diff --git a/Engine/Source/Runtime/Core/Public/Misc/TVariant.h b/Engine/Source/Runtime/Core/Public/Misc/TVariant.h
--- a/Engine/Source/Runtime/Core/Public/Misc/TVariant.h
+++ b/Engine/Source/Runtime/Core/Public/Misc/TVariant.h
@@ -16,13 +16,14 @@ struct TInPlaceType {};
 /**
  * A special tag that can be used as the first type in a TVariant parameter pack if none of the other types can be default-constructed.
  */
-struct FEmptyVariantState {};
-
-/** Allow FEmptyVariantState to be used with FArchive serialization */
-inline FArchive& operator<<(FArchive& Ar, FEmptyVariantState&)
+struct FEmptyVariantState
 {
-	return Ar;
-}
+	/** Allow FEmptyVariantState to be used with FArchive serialization */
+	friend inline FArchive& operator<<(FArchive& Ar, FEmptyVariantState&)
+	{
+		return Ar;
+	}
+};
 
 /**
  * A type-safe union based loosely on std::variant. This flavor of variant requires that all the types in the declaring template parameter pack be unique.
diff --git a/Engine/Source/Runtime/Core/Public/Modules/ModuleManager.h b/Engine/Source/Runtime/Core/Public/Modules/ModuleManager.h
--- a/Engine/Source/Runtime/Core/Public/Modules/ModuleManager.h
+++ b/Engine/Source/Runtime/Core/Public/Modules/ModuleManager.h
@@ -663,8 +663,6 @@ private:
 	mutable FCriticalSection ModulesCriticalSection;
 };
 
-FArchive& operator<<( FArchive& Ar, FModuleManager& ModuleManager );
-
 /**
  * Utility class for registering modules that are statically linked.
  */
diff --git a/Engine/Source/Runtime/Core/Public/Serialization/Archive.h b/Engine/Source/Runtime/Core/Public/Serialization/Archive.h
--- a/Engine/Source/Runtime/Core/Public/Serialization/Archive.h
+++ b/Engine/Source/Runtime/Core/Public/Serialization/Archive.h
@@ -1333,24 +1333,6 @@ public:
 		return Ar;
 	}
 
-	/**
-	 * Serializes an enumeration value from or into an archive.
-	 *
-	 * @param Ar The archive to serialize from or to.
-	 * @param Value The value to serialize.
-	 */
-	template<class TEnum>
-	FORCEINLINE friend FArchive& operator<<(FArchive& Ar, TEnumAsByte<TEnum>& Value)
-	{
-#if DEVIRTUALIZE_FLinkerLoad_Serialize
-		if (!Ar.FastPathLoad<sizeof(Value)>(&Value))
-#endif
-		{
-			Ar.Serialize(&Value, 1);
-		}
-		return Ar;
-	}
-
 	/**
 	 * Serializes a signed 8-bit integer value from or into an archive.
 	 *
@@ -1568,21 +1550,6 @@ public:
 		return Ar;
 	}
 
-	/**
-	 * Serializes enum classes as their underlying type.
-	 *
-	 * @param Ar The archive to serialize from or to.
-	 * @param Value The value to serialize.
-	 */
-	template <
-		typename EnumType,
-		typename = typename TEnableIf<TIsEnumClass<EnumType>::Value>::Type
-	>
-	FORCEINLINE friend FArchive& operator<<(FArchive& Ar, EnumType& Value)
-	{
-		return Ar << (__underlying_type(EnumType)&)Value;
-	}
-
 	/**
 	 * Serializes an FIntRect value from or into an archive.
 	 *
@@ -1955,8 +1922,8 @@ public:
 	/** Resets all of the base archive members. */
 	using FArchiveState::Reset;
 
+public:
 #if DEVIRTUALIZE_FLinkerLoad_Serialize
-private:
 	template<SIZE_T Size>
 	FORCEINLINE bool FastPathLoad(void* InDest)
 	{
@@ -1992,7 +1959,6 @@ private:
 		return false;
 	}
 
-public:
 	//@todoio FArchive is really a horrible class and the way it is proxied by FLinkerLoad is double terrible. It makes the fast path really hacky and slower than it would need to be.
 	using FArchiveState::ActiveFPLB;
 	using FArchiveState::InlineFPLB;
@@ -2206,3 +2172,36 @@ template<class T> T Arctor(FArchive& Ar)
 
 	return Tmp;
 }
+
+/**
+* Serializes an enumeration value from or into an archive.
+*
+* @param Ar The archive to serialize from or to.
+* @param Value The value to serialize.
+*/
+template<class TEnum>
+FORCEINLINE FArchive& operator<<(FArchive& Ar, TEnumAsByte<TEnum>& Value)
+{
+#if DEVIRTUALIZE_FLinkerLoad_Serialize
+	if (!Ar.FastPathLoad<sizeof(Value)>(&Value))
+#endif
+	{
+		Ar.Serialize(&Value, 1);
+	}
+	return Ar;
+}
+
+/**
+* Serializes enum classes as their underlying type.
+*
+* @param Ar The archive to serialize from or to.
+* @param Value The value to serialize.
+*/
+template <
+	typename EnumType,
+	typename = typename TEnableIf<TIsEnumClass<EnumType>::Value>::Type
+>
+FORCEINLINE FArchive& operator<<(FArchive& Ar, EnumType& Value)
+{
+	return Ar << (__underlying_type(EnumType)&)Value;
+}
diff --git a/Engine/Source/Runtime/Core/Public/Serialization/CompactBinary.h b/Engine/Source/Runtime/Core/Public/Serialization/CompactBinary.h
--- a/Engine/Source/Runtime/Core/Public/Serialization/CompactBinary.h
+++ b/Engine/Source/Runtime/Core/Public/Serialization/CompactBinary.h
@@ -498,22 +498,19 @@ public:
 
 	CORE_API static FCbObjectId NewObjectId();
 
-private:
-	alignas(uint32) ByteArray Bytes{};
-
-	friend inline bool operator==(const FCbObjectId& A, const FCbObjectId& B)
+	inline bool operator==(const FCbObjectId& B) const
 	{
-		return FMemory::Memcmp(&A, &B, sizeof(FCbObjectId)) == 0;
+		return FMemory::Memcmp(this, &B, sizeof(FCbObjectId)) == 0;
 	}
 
-	friend inline bool operator!=(const FCbObjectId& A, const FCbObjectId& B)
+	inline bool operator!=(const FCbObjectId& B) const
 	{
-		return FMemory::Memcmp(&A, &B, sizeof(FCbObjectId)) != 0;
+		return FMemory::Memcmp(this, &B, sizeof(FCbObjectId)) != 0;
 	}
 
-	friend inline bool operator<(const FCbObjectId& A, const FCbObjectId& B)
+	inline bool operator<(const FCbObjectId& B) const
 	{
-		return FMemory::Memcmp(&A, &B, sizeof(FCbObjectId)) <= 0;
+		return FMemory::Memcmp(this, &B, sizeof(FCbObjectId)) <= 0;
 	}
 
 	friend inline uint32 GetTypeHash(const FCbObjectId& Id)
@@ -528,6 +525,9 @@ private:
 		UE::String::BytesToHexLower(Id.GetBytes(), Builder);
 		return Builder;
 	}
+
+private:
+	alignas(uint32) ByteArray Bytes{};
 };
 
 inline FCbObjectId::FCbObjectId(const ByteArray& ObjectId)
diff --git a/Engine/Source/Runtime/Core/Public/Serialization/CompactBinaryWriter.h b/Engine/Source/Runtime/Core/Public/Serialization/CompactBinaryWriter.h
--- a/Engine/Source/Runtime/Core/Public/Serialization/CompactBinaryWriter.h
+++ b/Engine/Source/Runtime/Core/Public/Serialization/CompactBinaryWriter.h
@@ -253,7 +253,6 @@ protected:
 	CORE_API explicit FCbWriter(int64 InitialSize);
 
 private:
-	friend FCbWriter& operator<<(FCbWriter& Writer, FUtf8StringView NameOrValue);
 
 	/** Begin writing a field. May be called twice for named fields. */
 	void BeginField();
@@ -292,190 +291,189 @@ private:
 	// operations only when saving to a buffer.
 	TArray64<uint8> Data;
 	TArray<FState> States;
-};
-
-///////////////////////////////////////////////////////////////////////////////////////////////////////////////////////
 
-/**
- * A writer for compact binary object, arrays, and fields that uses a fixed-size stack buffer.
- *
- * @see FCbWriter
- */
-template <uint32 InlineBufferSize>
-class TCbWriter : public FCbWriter
-{
 public:
-	inline TCbWriter()
-		: FCbWriter(InlineBufferSize)
+	/** Write the field name if valid in this state, otherwise write the string value. */
+	inline FCbWriter& operator<<(FUtf8StringView NameOrValue)
 	{
+		SetNameOrAddString(NameOrValue);
+		return *this;
 	}
 
-	TCbWriter(const TCbWriter&) = delete;
-	TCbWriter& operator=(const TCbWriter&) = delete;
-
-private:
-	// Reserve the inline buffer now even though we are unable to use it. This will avoid causing
-	// new stack overflows when this functionality is properly implemented in the future.
-	uint8 Buffer[InlineBufferSize];
-};
-
-///////////////////////////////////////////////////////////////////////////////////////////////////////////////////////
+	/** Write the field name if valid in this state, otherwise write the string value. */
+	inline FCbWriter& operator<<(const ANSICHAR* NameOrValue)
+	{
+		return *this << FAnsiStringView(NameOrValue);
+	}
 
-/** Write the field name if valid in this state, otherwise write the string value. */
-inline FCbWriter& operator<<(FCbWriter& Writer, FUtf8StringView NameOrValue)
-{
-	Writer.SetNameOrAddString(NameOrValue);
-	return Writer;
-}
+	/** Write the field name if valid in this state, otherwise write the string value. */
+	inline FCbWriter& operator<<(const UTF8CHAR* NameOrValue)
+	{
+		return *this << FUtf8StringView(NameOrValue);
+	}
 
-/** Write the field name if valid in this state, otherwise write the string value. */
-inline FCbWriter& operator<<(FCbWriter& Writer, const ANSICHAR* NameOrValue)
-{
-	return Writer << FAnsiStringView(NameOrValue);
-}
+	inline FCbWriter& operator<<(const FCbFieldView& Value)
+	{
+		AddField(Value);
+		return *this;
+	}
 
-/** Write the field name if valid in this state, otherwise write the string value. */
-inline FCbWriter& operator<<(FCbWriter& Writer, const UTF8CHAR* NameOrValue)
-{
-	return Writer << FUtf8StringView(NameOrValue);
-}
+	inline FCbWriter& operator<<(const FCbField& Value)
+	{
+		AddField(Value);
+		return *this;
+	}
 
-inline FCbWriter& operator<<(FCbWriter& Writer, const FCbFieldView& Value)
-{
-	Writer.AddField(Value);
-	return Writer;
-}
+	inline FCbWriter& operator<<(const FCbObjectView& Value)
+	{
+		AddObject(Value);
+		return *this;
+	}
 
-inline FCbWriter& operator<<(FCbWriter& Writer, const FCbField& Value)
-{
-	Writer.AddField(Value);
-	return Writer;
-}
+	inline FCbWriter& operator<<(const FCbObject& Value)
+	{
+		AddObject(Value);
+		return *this;
+	}
 
-inline FCbWriter& operator<<(FCbWriter& Writer, const FCbObjectView& Value)
-{
-	Writer.AddObject(Value);
-	return Writer;
-}
+	inline FCbWriter& operator<<(const FCbArrayView& Value)
+	{
+		AddArray(Value);
+		return *this;
+	}
 
-inline FCbWriter& operator<<(FCbWriter& Writer, const FCbObject& Value)
-{
-	Writer.AddObject(Value);
-	return Writer;
-}
+	inline FCbWriter& operator<<(const FCbArray& Value)
+	{
+		AddArray(Value);
+		return *this;
+	}
 
-inline FCbWriter& operator<<(FCbWriter& Writer, const FCbArrayView& Value)
-{
-	Writer.AddArray(Value);
-	return Writer;
-}
+	inline FCbWriter& operator<<(nullptr_t)
+	{
+		AddNull();
+		return *this;
+	}
 
-inline FCbWriter& operator<<(FCbWriter& Writer, const FCbArray& Value)
-{
-	Writer.AddArray(Value);
-	return Writer;
-}
+	inline FCbWriter& operator<<(FWideStringView Value)
+	{
+		AddString(Value);
+		return *this;
+	}
 
-inline FCbWriter& operator<<(FCbWriter& Writer, nullptr_t)
-{
-	Writer.AddNull();
-	return Writer;
-}
+	inline FCbWriter& operator<<(const WIDECHAR* Value)
+	{
+		AddString(Value);
+		return *this;
+	}
 
-inline FCbWriter& operator<<(FCbWriter& Writer, FWideStringView Value)
-{
-	Writer.AddString(Value);
-	return Writer;
-}
+	inline FCbWriter& operator<<(int32 Value)
+	{
+		AddInteger(Value);
+		return *this;
+	}
 
-inline FCbWriter& operator<<(FCbWriter& Writer, const WIDECHAR* Value)
-{
-	Writer.AddString(Value);
-	return Writer;
-}
+	inline FCbWriter& operator<<(int64 Value)
+	{
+		AddInteger(Value);
+		return *this;
+	}
 
-inline FCbWriter& operator<<(FCbWriter& Writer, int32 Value)
-{
-	Writer.AddInteger(Value);
-	return Writer;
-}
+	inline FCbWriter& operator<<(uint32 Value)
+	{
+		AddInteger(Value);
+		return *this;
+	}
 
-inline FCbWriter& operator<<(FCbWriter& Writer, int64 Value)
-{
-	Writer.AddInteger(Value);
-	return Writer;
-}
+	inline FCbWriter& operator<<(uint64 Value)
+	{
+		AddInteger(Value);
+		return *this;
+	}
 
-inline FCbWriter& operator<<(FCbWriter& Writer, uint32 Value)
-{
-	Writer.AddInteger(Value);
-	return Writer;
-}
+	inline FCbWriter& operator<<(float Value)
+	{
+		AddFloat(Value);
+		return *this;
+	}
 
-inline FCbWriter& operator<<(FCbWriter& Writer, uint64 Value)
-{
-	Writer.AddInteger(Value);
-	return Writer;
-}
+	inline FCbWriter& operator<<(double Value)
+	{
+		AddFloat(Value);
+		return *this;
+	}
 
-inline FCbWriter& operator<<(FCbWriter& Writer, float Value)
-{
-	Writer.AddFloat(Value);
-	return Writer;
-}
+	inline FCbWriter& operator<<(bool Value)
+	{
+		AddBool(Value);
+		return *this;
+	}
 
-inline FCbWriter& operator<<(FCbWriter& Writer, double Value)
-{
-	Writer.AddFloat(Value);
-	return Writer;
-}
+	inline FCbWriter& operator<<(const FCbAttachment& Attachment)
+	{
+		AddAttachment(Attachment);
+		return *this;
+	}
 
-inline FCbWriter& operator<<(FCbWriter& Writer, bool Value)
-{
-	Writer.AddBool(Value);
-	return Writer;
-}
+	inline FCbWriter& operator<<(const FIoHash& Value)
+	{
+		AddHash(Value);
+		return *this;
+	}
 
-inline FCbWriter& operator<<(FCbWriter& Writer, const FCbAttachment& Attachment)
-{
-	Writer.AddAttachment(Attachment);
-	return Writer;
-}
+	inline FCbWriter& operator<<(const FGuid& Value)
+	{
+		AddUuid(Value);
+		return *this;
+	}
 
-inline FCbWriter& operator<<(FCbWriter& Writer, const FIoHash& Value)
-{
-	Writer.AddHash(Value);
-	return Writer;
-}
+	CORE_API FCbWriter& operator<<(FDateTime Value);
+	CORE_API FCbWriter& operator<<(FTimespan Value);
 
-inline FCbWriter& operator<<(FCbWriter& Writer, const FGuid& Value)
-{
-	Writer.AddUuid(Value);
-	return Writer;
-}
+	inline FCbWriter& operator<<(const FCbObjectId& Value)
+	{
+		AddObjectId(Value);
+		return *this;
+	}
 
-CORE_API FCbWriter& operator<<(FCbWriter& Writer, FDateTime Value);
-CORE_API FCbWriter& operator<<(FCbWriter& Writer, FTimespan Value);
+	CORE_API FCbWriter& operator<<(FName Value);
 
-inline FCbWriter& operator<<(FCbWriter& Writer, const FCbObjectId& Value)
-{
-	Writer.AddObjectId(Value);
-	return Writer;
-}
+	template <typename T, typename Allocator,
+		std::void_t<decltype(std::declval<FCbWriter&>() << std::declval<const T&>())>* = nullptr>
+	inline FCbWriter& operator<<(const TArray<T, Allocator>& Value)
+	{
+		BeginArray();
+		for (const T& Element : Value)
+		{
+			*this << Element;
+		}
+		EndArray();
+		return *this;
+	}
+};
 
-CORE_API FCbWriter& operator<<(FCbWriter& Writer, FName Value);
+///////////////////////////////////////////////////////////////////////////////////////////////////////////////////////
 
-template <typename T, typename Allocator,
-	std::void_t<decltype(std::declval<FCbWriter&>() << std::declval<const T&>())>* = nullptr>
-inline FCbWriter& operator<<(FCbWriter& Writer, const TArray<T, Allocator>& Value)
+/**
+ * A writer for compact binary object, arrays, and fields that uses a fixed-size stack buffer.
+ *
+ * @see FCbWriter
+ */
+template <uint32 InlineBufferSize>
+class TCbWriter : public FCbWriter
 {
-	Writer.BeginArray();
-	for (const T& Element : Value)
+public:
+	inline TCbWriter()
+		: FCbWriter(InlineBufferSize)
 	{
-		Writer << Element;
 	}
-	Writer.EndArray();
-	return Writer;
-}
+
+	TCbWriter(const TCbWriter&) = delete;
+	TCbWriter& operator=(const TCbWriter&) = delete;
+
+private:
+	// Reserve the inline buffer now even though we are unable to use it. This will avoid causing
+	// new stack overflows when this functionality is properly implemented in the future.
+	uint8 Buffer[InlineBufferSize];
+};
 
 ///////////////////////////////////////////////////////////////////////////////////////////////////////////////////////
diff --git a/Engine/Source/Runtime/Core/Public/Serialization/CompressedChunkInfo.h b/Engine/Source/Runtime/Core/Public/Serialization/CompressedChunkInfo.h
--- a/Engine/Source/Runtime/Core/Public/Serialization/CompressedChunkInfo.h
+++ b/Engine/Source/Runtime/Core/Public/Serialization/CompressedChunkInfo.h
@@ -19,12 +19,13 @@ struct FCompressedChunkInfo
 
 	/** Holds the data's uncompresses size. */
 	int64 UncompressedSize;
-};
 
-/**
- * Serializes an FCompressedChunkInfo value from or into an archive.
- *
- * @param Ar The archive to serialize from or to.
- * @param Value The value to serialize.
- */
-CORE_API FArchive& operator<<(FArchive& Ar, FCompressedChunkInfo& Value);
+	/**
+	 * Serializes an FCompressedChunkInfo value from or into an archive.
+	 *
+	 * @param Ar The archive to serialize from or to.
+	 * @param Value The value to serialize.
+	 */
+	friend CORE_API FArchive& operator<<(FArchive& Ar, FCompressedChunkInfo& Value);
+
+};
diff --git a/Engine/Source/Runtime/Core/Public/Serialization/MemoryImage.h b/Engine/Source/Runtime/Core/Public/Serialization/MemoryImage.h
--- a/Engine/Source/Runtime/Core/Public/Serialization/MemoryImage.h
+++ b/Engine/Source/Runtime/Core/Public/Serialization/MemoryImage.h
@@ -157,19 +157,19 @@ struct FMemoryImageVTablePointer
 	uint32 VTableOffset;
 	uint32 Offset;
 
-	friend inline bool operator==(const FMemoryImageVTablePointer& Lhs, const FMemoryImageVTablePointer& Rhs)
+	inline bool operator==(const FMemoryImageVTablePointer& Rhs) const
 	{
-		return Lhs.TypeNameHash == Rhs.TypeNameHash && Lhs.VTableOffset == Rhs.VTableOffset && Lhs.Offset == Rhs.Offset;
+		return TypeNameHash == Rhs.TypeNameHash && VTableOffset == Rhs.VTableOffset && Offset == Rhs.Offset;
 	}
-	friend inline bool operator!=(const FMemoryImageVTablePointer& Lhs, const FMemoryImageVTablePointer& Rhs)
+	inline bool operator!=(const FMemoryImageVTablePointer& Rhs) const
 	{
-		return !operator==(Lhs, Rhs);
+		return !(*this == Rhs);
 	}
-	friend inline bool operator<(const FMemoryImageVTablePointer& Lhs, const FMemoryImageVTablePointer& Rhs)
+	inline bool operator<(const FMemoryImageVTablePointer& Rhs) const
 	{
-		if (Lhs.TypeNameHash != Rhs.TypeNameHash) return Lhs.TypeNameHash < Rhs.TypeNameHash;
-		if (Lhs.VTableOffset != Rhs.VTableOffset) return Lhs.VTableOffset < Rhs.VTableOffset;
-		return Lhs.Offset < Rhs.Offset;
+		if (TypeNameHash != Rhs.TypeNameHash) return TypeNameHash < Rhs.TypeNameHash;
+		if (VTableOffset != Rhs.VTableOffset) return VTableOffset < Rhs.VTableOffset;
+		return Offset < Rhs.Offset;
 	}
 };
 
@@ -178,21 +178,21 @@ struct FMemoryImageNamePointer
 	FName Name;
 	uint32 Offset;
 
-	friend inline bool operator==(const FMemoryImageNamePointer& Lhs, const FMemoryImageNamePointer& Rhs)
+	inline bool operator==(const FMemoryImageNamePointer& Rhs) const
 	{
-		return Lhs.Offset == Rhs.Offset && Lhs.Name.Compare(Rhs.Name) == 0;
+		return Offset == Rhs.Offset && Name.Compare(Rhs.Name) == 0;
 	}
-	friend inline bool operator!=(const FMemoryImageNamePointer& Lhs, const FMemoryImageNamePointer& Rhs)
+	inline bool operator!=(const FMemoryImageNamePointer& Rhs) const
 	{
-		return !operator==(Lhs, Rhs);
+		return !(*this == Rhs);
 	}
-	friend inline bool operator<(const FMemoryImageNamePointer& Lhs, const FMemoryImageNamePointer& Rhs)
+	inline bool operator<(const FMemoryImageNamePointer& Rhs) const
 	{
-		if (Lhs.Name != Rhs.Name)
+		if (Name != Rhs.Name)
 		{
-			return Lhs.Name.LexicalLess(Rhs.Name);
+			return Name.LexicalLess(Rhs.Name);
 		}
-		return Lhs.Offset < Rhs.Offset;
+		return Offset < Rhs.Offset;
 	}
 };
 
@@ -803,24 +803,24 @@ public:
 		return Ar;
 	}
 
-	friend inline bool operator==(const FMemoryImageString& Lhs, const FMemoryImageString& Rhs)
+	inline bool operator==(const FMemoryImageString& Rhs) const
 	{
-		return FCString::Stricmp(*Lhs, *Rhs) == 0;
+		return FCString::Stricmp(**this, *Rhs) == 0;
 	}
 
-	friend inline bool operator!=(const FMemoryImageString& Lhs, const FMemoryImageString& Rhs)
+	inline bool operator!=(const FMemoryImageString& Rhs) const
 	{
-		return !operator==(Lhs, Rhs);
+		return FCString::Stricmp(**this, *Rhs) != 0;
 	}
 
-	friend inline bool operator==(const FMemoryImageString& Lhs, const FString& Rhs)
+	inline bool operator==(const FString& Rhs) const
 	{
-		return FCString::Stricmp(*Lhs, *Rhs) == 0;
+		return FCString::Stricmp(**this, *Rhs) == 0;
 	}
 
-	friend inline bool operator!=(const FMemoryImageString& Lhs, const FString& Rhs)
+	inline bool operator!=(const FString& Rhs) const
 	{
-		return !operator==(Lhs, Rhs);
+		return FCString::Stricmp(**this, *Rhs) != 0;
 	}
 
 	inline DataType::ElementAllocatorType& GetAllocatorInstance() { return Data.GetAllocatorInstance(); }
@@ -866,11 +866,11 @@ public:
 	const FHashedNameDebugString& GetDebugString() const { return DebugString; }
 #endif
 
-	friend inline bool operator==(const FHashedName& Lhs, const FHashedName& Rhs) { return Lhs.Hash == Rhs.Hash; }
-	friend inline bool operator!=(const FHashedName& Lhs, const FHashedName& Rhs) { return Lhs.Hash != Rhs.Hash; }
+	inline bool operator==(const FHashedName& Rhs) const { return Hash == Rhs.Hash; }
+	inline bool operator!=(const FHashedName& Rhs) const { return Hash != Rhs.Hash; }
 
 	/** For sorting by name */
-	friend inline bool operator<(const FHashedName& Lhs, const FHashedName& Rhs) { return Lhs.Hash < Rhs.Hash; }
+	inline bool operator<(const FHashedName& Rhs) const { return Hash < Rhs.Hash; }
 
 	friend inline FArchive& operator<<(FArchive& Ar, FHashedName& String)
 	{
diff --git a/Engine/Source/Runtime/Core/Public/Templates/Function.h b/Engine/Source/Runtime/Core/Public/Templates/Function.h
--- a/Engine/Source/Runtime/Core/Public/Templates/Function.h
+++ b/Engine/Source/Runtime/Core/Public/Templates/Function.h
@@ -920,33 +920,17 @@ public:
 	/**
 	 * Nullptr equality operator.
 	 */
-	friend FORCEINLINE bool operator==(TYPE_OF_NULLPTR, const TFunction& Func)
+	FORCEINLINE bool operator==(TYPE_OF_NULLPTR) const
 	{
-		return !Func;
-	}
-
-	/**
-	 * Nullptr equality operator.
-	 */
-	friend FORCEINLINE bool operator==(const TFunction& Func, TYPE_OF_NULLPTR)
-	{
-		return !Func;
+		return !*this;
 	}
 
 	/**
 	 * Nullptr inequality operator.
 	 */
-	friend FORCEINLINE bool operator!=(TYPE_OF_NULLPTR, const TFunction& Func)
+	FORCEINLINE bool operator!=(TYPE_OF_NULLPTR) const
 	{
-		return (bool)Func;
-	}
-
-	/**
-	 * Nullptr inequality operator.
-	 */
-	friend FORCEINLINE bool operator!=(const TFunction& Func, TYPE_OF_NULLPTR)
-	{
-		return (bool)Func;
+		return (bool)*this;
 	}
 };
 
@@ -1053,3 +1037,24 @@ public:
 		return Super::IsSet();
 	}
 };
+
+
+#if !PLATFORM_COMPILER_HAS_GENERATED_COMPARISON_OPERATORS
+/**
+ * Nullptr equality operator.
+ */
+template <typename FuncType>
+FORCEINLINE bool operator==(TYPE_OF_NULLPTR, const TFunction<FuncType>& Func)
+{
+	return !Func;
+}
+
+/**
+ * Nullptr inequality operator.
+ */
+template <typename FuncType>
+FORCEINLINE bool operator!=(TYPE_OF_NULLPTR, const TFunction<FuncType>& Func)
+{
+	return (bool)Func;
+}
+#endif
\ No newline at end of file
diff --git a/Engine/Source/Runtime/Core/Public/Templates/NonNullPointer.h b/Engine/Source/Runtime/Core/Public/Templates/NonNullPointer.h
--- a/Engine/Source/Runtime/Core/Public/Templates/NonNullPointer.h
+++ b/Engine/Source/Runtime/Core/Public/Templates/NonNullPointer.h
@@ -210,14 +210,14 @@ public:
 		return InPointer;
 	}
 
-	friend bool operator==(const TOptional& lhs, const TOptional& rhs)
+	bool operator==(const TOptional& rhs) const
 	{
-		return lhs.Pointer == rhs.Pointer;
+		return Pointer == rhs.Pointer;
 	}
 
-	friend bool operator!=(const TOptional& lhs, const TOptional& rhs)
+	bool operator!=(const TOptional& rhs) const
 	{
-		return !(lhs == rhs);
+		return Pointer != rhs.Pointer;
 	}
 
 	friend FArchive& operator<<(FArchive& Ar, TOptional& Optional)
diff --git a/Engine/Source/Runtime/Core/Public/Templates/PimplPtr.h b/Engine/Source/Runtime/Core/Public/Templates/PimplPtr.h
--- a/Engine/Source/Runtime/Core/Public/Templates/PimplPtr.h
+++ b/Engine/Source/Runtime/Core/Public/Templates/PimplPtr.h
@@ -233,13 +233,11 @@ public:
 		}
 	}
 
+	FORCEINLINE bool operator==(TYPE_OF_NULLPTR) { return !IsValid(); }
+	FORCEINLINE bool operator!=(TYPE_OF_NULLPTR) { return  IsValid(); }
+
 private:
 	T* Ptr = nullptr;
-
-	template <EPimplPtrMode Mode> friend FORCEINLINE bool operator==(const TPimplPtr<T, Mode>& PimplPtr, TYPE_OF_NULLPTR) { return !PimplPtr.IsValid(); }
-	template <EPimplPtrMode Mode> friend FORCEINLINE bool operator==(TYPE_OF_NULLPTR, const TPimplPtr<T, Mode>& PimplPtr) { return !PimplPtr.IsValid(); }
-	template <EPimplPtrMode Mode> friend FORCEINLINE bool operator!=(const TPimplPtr<T, Mode>& PimplPtr, TYPE_OF_NULLPTR) { return  PimplPtr.IsValid(); }
-	template <EPimplPtrMode Mode> friend FORCEINLINE bool operator!=(TYPE_OF_NULLPTR, const TPimplPtr<T, Mode>& PimplPtr) { return  PimplPtr.IsValid(); }
 };
 
 template <typename T>
@@ -305,6 +303,10 @@ public:
 	using Super::Reset;
 };
 
+#if !PLATFORM_COMPILER_HAS_GENERATED_COMPARISON_OPERATORS
+template <typename T, EPimplPtrMode Mode> FORCEINLINE bool operator==(TYPE_OF_NULLPTR, const TPimplPtr<T, Mode>& Ptr) { return !Ptr.IsValid(); }
+template <typename T, EPimplPtrMode Mode> FORCEINLINE bool operator!=(TYPE_OF_NULLPTR, const TPimplPtr<T, Mode>& Ptr) { return  Ptr.IsValid(); }
+#endif
 
 /**
  * Heap-allocates an instance of T with the given arguments and returns it as a TPimplPtr.
diff --git a/Engine/Source/Runtime/Core/Public/Templates/RefCounting.h b/Engine/Source/Runtime/Core/Public/Templates/RefCounting.h
--- a/Engine/Source/Runtime/Core/Public/Templates/RefCounting.h
+++ b/Engine/Source/Runtime/Core/Public/Templates/RefCounting.h
@@ -268,11 +268,6 @@ public:
 		return Reference;
 	}
 
-	FORCEINLINE friend uint32 GetTypeHash(const TRefCountPtr& InPtr)
-	{
-		return GetTypeHash(InPtr.Reference);
-	}
-
 	FORCEINLINE ReferencedType** GetInitReference()
 	{
 		*this = nullptr;
@@ -317,15 +312,14 @@ public:
 		InPtr.Reference = OldReference;
 	}
 
-	friend FArchive& operator<<(FArchive& Ar,TRefCountPtr& Ptr)
+	void Serialize(FArchive& Ar)
 	{
-		ReferenceType PtrReference = Ptr.Reference;
+		ReferenceType PtrReference = Reference;
 		Ar << PtrReference;
 		if(Ar.IsLoading())
 		{
-			Ptr = PtrReference;
+			*this = PtrReference;
 		}
-		return Ar;
 	}
 
 private:
@@ -335,20 +329,38 @@ private:
 	template <typename OtherType>
 	friend class TRefCountPtr;
 
-	friend FORCEINLINE bool operator==(const TRefCountPtr& A, const TRefCountPtr& B)
-	{
-		return A.GetReference() == B.GetReference();
-	}
-
-	friend FORCEINLINE bool operator==(const TRefCountPtr& A, ReferencedType* B)
+public:
+	FORCEINLINE bool operator==(const TRefCountPtr& B) const
 	{
-		return A.GetReference() == B;
+		return GetReference() == B.GetReference();
 	}
 
-	friend FORCEINLINE bool operator==(ReferencedType* A, const TRefCountPtr& B)
+	FORCEINLINE bool operator==(ReferencedType* B) const
 	{
-		return A == B.GetReference();
+		return GetReference() == B;
 	}
 };
 
 ALIAS_TEMPLATE_TYPE_LAYOUT(template<typename T>, TRefCountPtr<T>, void*);
+
+#if !PLATFORM_COMPILER_HAS_GENERATED_COMPARISON_OPERATORS
+template<typename ReferencedType>
+FORCEINLINE bool operator==(ReferencedType* A, const TRefCountPtr<ReferencedType>& B)
+{
+	return A == B.GetReference();
+}
+#endif
+
+template<typename ReferencedType>
+FORCEINLINE uint32 GetTypeHash(const TRefCountPtr<ReferencedType>& InPtr)
+{
+	return GetTypeHash(InPtr.GetReference());
+}
+
+
+template<typename ReferencedType>
+FArchive& operator<<(FArchive& Ar,TRefCountPtr<ReferencedType>& Ptr)
+{
+	Ptr.Serialize(Ar);
+	return Ar;
+}
diff --git a/Engine/Source/Runtime/Core/Public/Templates/SharedPointer.h b/Engine/Source/Runtime/Core/Public/Templates/SharedPointer.h
--- a/Engine/Source/Runtime/Core/Public/Templates/SharedPointer.h
+++ b/Engine/Source/Runtime/Core/Public/Templates/SharedPointer.h
@@ -605,18 +605,6 @@ private:
 		return Object != nullptr;
 	}
 
-	/**
-	 * Computes a hash code for this object
-	 *
-	 * @param  InSharedRef  Shared pointer to compute hash code for
-	 *
-	 * @return  Hash code value
-	 */
-	[[nodiscard]] friend uint32 GetTypeHash( const TSharedRef< ObjectType, Mode >& InSharedRef )
-	{
-		return ::PointerHash( InSharedRef.Object );
-	}
-
 	// We declare ourselves as a friend (templated using OtherType) so we can access members as needed
     template< class OtherType, ESPMode OtherMode > friend class TSharedRef;
 
@@ -653,33 +641,6 @@ private:
 
 		Init(InObject);
 	}
-
-public:
-
-	/**
-	 * Global equality operator for TSharedRef
-	 *
-	 * @return  True if the two shared references are equal
-	 */
-	template< class ObjectTypeB >
-	[[nodiscard]] FORCEINLINE bool operator==( TSharedRef< ObjectTypeB, InMode > const& InSharedRefB ) const
-	{
-		return &( Get() ) == &( InSharedRefB.Get() );
-	}
-
-
-	/**
-	 * Global inequality operator for TSharedRef
-	 *
-	 * @return  True if the two shared references are not equal
-	 */
-	template< class ObjectTypeB >
-	[[nodiscard]] FORCEINLINE bool operator!=( TSharedRef< ObjectTypeB, InMode > const& InSharedRefB ) const
-	{
-		return &( Get() ) != &( InSharedRefB.Get() );
-	}
-
-
 };
 
 
@@ -1225,18 +1186,6 @@ private:
 		}
 	}
 
-	/**
-	 * Computes a hash code for this object
-	 *
-	 * @param  InSharedPtr  Shared pointer to compute hash code for
-	 *
-	 * @return  Hash code value
-	 */
-	[[nodiscard]] friend uint32 GetTypeHash( const TSharedPtr< ObjectType, Mode >& InSharedPtr )
-	{
-		return ::PointerHash( InSharedPtr.Object );
-	}
-
 	// We declare ourselves as a friend (templated using OtherType) so we can access members as needed
 	template< class OtherType, ESPMode OtherMode > friend class TSharedPtr;
 
@@ -1253,123 +1202,6 @@ private:
 	/** Interface to the reference counter for this object.  Note that the actual reference
 		controller object is shared by all shared and weak pointers that refer to the object */
 	SharedPointerInternals::FSharedReferencer< Mode > SharedReferenceCount;
-
-public:
-
-	/**
-	 * Global equality operator for TSharedPtr
-	 *
-	 * @return  True if the two shared pointers are equal
-	 */
-	template< class ObjectTypeB >
-	[[nodiscard]] FORCEINLINE bool operator==( TSharedPtr< ObjectTypeB, InMode > const& InSharedPtrB ) const
-	{
-		return Get() == InSharedPtrB.Get();
-	}
-
-
-	/**
-	 * Global inequality operator for TSharedPtr
-	 *
-	 * @return  True if the two shared pointers are not equal
-	 */
-	template< class ObjectTypeB >
-	[[nodiscard]] FORCEINLINE bool operator!=( TSharedPtr< ObjectTypeB, InMode > const& InSharedPtrB ) const
-	{
-		return Get() != InSharedPtrB.Get();
-	}
-
-
-	/**
-	 * Global equality operator for TSharedPtr
-	 *
-	 * @return  True if the shared pointer is null
-	 */
-	[[nodiscard]] friend FORCEINLINE bool operator==(TSharedPtr const& InSharedPtr, TYPE_OF_NULLPTR )
-	{
-		return !InSharedPtr.IsValid();
-	}
-
-
-	/**
-	 * Global inequality operator for TSharedPtr
-	 *
-	 * @return  True if the shared pointer is not null
-	 */
-	[[nodiscard]] friend FORCEINLINE bool operator!=( TSharedPtr const& InSharedPtr, TYPE_OF_NULLPTR )
-	{
-		return InSharedPtr.IsValid();
-	}
-
-
-	/**
-	 * Tests to see if a TSharedRef is "equal" to a TSharedPtr (both are valid and refer to the same object) (reverse)
-	 *
-	 * @return  True if the shared reference and shared pointer are "equal"
-	 */
-	template< class ObjectTypeB >
-	[[nodiscard]] friend FORCEINLINE bool operator==(TSharedPtr const& InSharedPtr, TSharedRef< ObjectTypeB, InMode > const& InSharedRef )
-	{
-		return InSharedPtr.IsValid() && InSharedPtr.Get() == &( InSharedRef.Get() );
-	}
-
-
-	/**
-	 * Tests to see if a TSharedRef is not "equal" to a TSharedPtr (shared pointer is invalid, or both refer to different objects) (reverse)
-	 *
-	 * @return  True if the shared reference and shared pointer are not "equal"
-	 */
-	template< class ObjectTypeB >
-	[[nodiscard]] friend FORCEINLINE bool operator!=(TSharedPtr const& InSharedPtr, TSharedRef< ObjectTypeB, InMode > const& InSharedRef )
-	{
-		return !InSharedPtr.IsValid() || ( InSharedPtr.Get() != &( InSharedRef.Get() ) );
-	}
-
-
-
-	/**
-	 * Global equality operator for TSharedPtr
-	 *
-	 * @return  True if the shared pointer is null
-	 */
-	[[nodiscard]] friend FORCEINLINE bool operator==( TYPE_OF_NULLPTR, TSharedPtr const& InSharedPtrB )
-	{
-		return !InSharedPtrB.IsValid();
-	}
-
-
-	/**
-	 * Global inequality operator for TSharedPtr
-	 *
-	 * @return  True if the shared pointer is not null
-	 */
-	[[nodiscard]] friend FORCEINLINE bool operator!=( TYPE_OF_NULLPTR, TSharedPtr const& InSharedPtrB )
-	{
-		return InSharedPtrB.IsValid();
-	}
-
-	/**
-	 * Tests to see if a TSharedRef is "equal" to a TSharedPtr (both are valid and refer to the same object)
-	 *
-	 * @return  True if the shared reference and shared pointer are "equal"
-	 */
-	template< class ObjectTypeA >
-	[[nodiscard]] friend FORCEINLINE bool operator==( TSharedRef< ObjectTypeA, InMode > const& InSharedRef, TSharedPtr const& InSharedPtr)
-	{
-		return InSharedPtr.IsValid() && InSharedPtr.Get() == &( InSharedRef.Get() );
-	}
-
-
-	/**
-	 * Tests to see if a TSharedRef is not "equal" to a TSharedPtr (shared pointer is invalid, or both refer to different objects)
-	 *
-	 * @return  True if the shared reference and shared pointer are not "equal"
-	 */
-	template< class ObjectTypeA >
-	[[nodiscard]] friend FORCEINLINE bool operator!=( TSharedRef< ObjectTypeA, InMode > const& InSharedRef, TSharedPtr const& InSharedPtr )
-	{
-		return !InSharedPtr.IsValid() || ( InSharedPtr.Get() != &( InSharedRef.Get() ) );
-	}
 };
 
 namespace Freeze
@@ -1655,20 +1487,13 @@ public:
 		return Pin().Get() == InOtherPtr;
 	}
 
-private:
-	
-	/**
-	 * Computes a hash code for this object
-	 *
-	 * @param  InWeakPtr  Weak pointer to compute hash code for
-	 *
-	 * @return  Hash code value
-	 */
-	[[nodiscard]] friend uint32 GetTypeHash( const TWeakPtr< ObjectType, Mode >& InWeakPtr )
+	FORCEINLINE uint32 GetWeakPtrTypeHash() const
 	{
-		return ::PointerHash( InWeakPtr.Object );
+		return ::PointerHash( Object );
 	}
 
+private:
+	
 	// We declare ourselves as a friend (templated using OtherType) so we can access members as needed
     template< class OtherType, ESPMode OtherMode > friend class TWeakPtr;
 
@@ -1684,171 +1509,6 @@ private:
 	/** Interface to the reference counter for this object.  Note that the actual reference
 		controller object is shared by all shared and weak pointers that refer to the object */
 	SharedPointerInternals::FWeakReferencer< Mode > WeakReferenceCount;
-
-
-public:
-	/**
-	 * Global equality operator for TWeakPtr
-	 *
-	 * @return  True if the two weak pointers are equal
-	 */
-	template< class ObjectTypeB >
-	[[nodiscard]] FORCEINLINE bool operator==( TWeakPtr< ObjectTypeB, InMode > const& InWeakPtrB ) const
-	{
-		return Pin().Get() == InWeakPtrB.Pin().Get();
-	}
-
-
-	/**
-	 * Global inequality operator for TWeakPtr
-	 *
-	 * @return  True if the two weak pointers are not equal
-	 */
-	template< class ObjectTypeB>
-	[[nodiscard]] FORCEINLINE bool operator!=( TWeakPtr< ObjectTypeB, InMode > const& InWeakPtrB ) const
-	{
-		return Pin().Get() != InWeakPtrB.Pin().Get();
-	}
-
-
-	/**
-	 * Global equality operator for TWeakPtr
-	 *
-	 * @return  True if the weak pointer and the shared ref are equal
-	 */
-	template< class ObjectTypeB >
-	[[nodiscard]] friend FORCEINLINE bool operator==( TWeakPtr const& InWeakPtr, TSharedRef< ObjectTypeB, InMode > const& InSharedRefB )
-	{
-		return InWeakPtr.Pin().Get() == &InSharedRefB.Get();
-	}
-
-
-	/**
-	 * Global equality operator for TWeakPtr
-	 *
-	 * @return  True if the weak pointer and the shared ptr are equal
-	 */
-	template< class ObjectTypeB >
-	[[nodiscard]] friend FORCEINLINE bool operator==( TWeakPtr const& InWeakPtr, TSharedPtr< ObjectTypeB, InMode > const& InSharedPtrB )
-	{
-		return InWeakPtr.Pin().Get() == InSharedPtrB.Get();
-	}
-
-
-	/**
-	 * Global equality operator for TWeakPtr
-	 *
-	 * @return  True if the weak pointer and the shared ref are equal
-	 */
-	template< class ObjectTypeA >
-	[[nodiscard]] friend FORCEINLINE bool operator==( TSharedRef< ObjectTypeA, InMode > const& InSharedRefA, TWeakPtr const& InWeakPtrB )
-	{
-		return &InSharedRefA.Get() == InWeakPtrB.Pin().Get();
-	}
-
-
-	/**
-	 * Global equality operator for TWeakPtr
-	 *
-	 * @return  True if the weak pointer and the shared ptr are equal
-	 */
-	template< class ObjectTypeA >
-	[[nodiscard]] friend FORCEINLINE bool operator==( TSharedPtr< ObjectTypeA, InMode > const& InSharedPtrA, TWeakPtr const& InWeakPtrB )
-	{
-		return InSharedPtrA.Get() == InWeakPtrB.Pin().Get();
-	}
-
-
-	/**
-	 * Global equality operator for TWeakPtr
-	 *
-	 * @return  True if the weak pointer is null
-	 */
-	[[nodiscard]] friend FORCEINLINE bool operator==( TWeakPtr const& InWeakPtr, TYPE_OF_NULLPTR )
-	{
-		return !InWeakPtr.IsValid();
-	}
-
-
-	/**
-	 * Global equality operator for TWeakPtr
-	 *
-	 * @return  True if the weak pointer is null
-	 */
-	[[nodiscard]] friend FORCEINLINE bool operator==( TYPE_OF_NULLPTR, TWeakPtr const& InWeakPtr )
-	{
-		return !InWeakPtr.IsValid();
-	}
-
-
-	/**
-	 * Global equality operator for TWeakPtr
-	 *
-	 * @return  True if the weak pointer and the shared ref are not equal
-	 */
-	template< class ObjectTypeB>
-	[[nodiscard]] friend FORCEINLINE bool operator!=( TWeakPtr const& InWeakPtr, TSharedRef< ObjectTypeB, InMode > const& InSharedRefB )
-	{
-		return InWeakPtr.Pin().Get() != &InSharedRefB.Get();
-	}
-
-
-	/**
-	 * Global equality operator for TWeakPtr
-	 *
-	 * @return  True if the weak pointer and the shared ptr are not equal
-	 */
-	template< class ObjectTypeB>
-	[[nodiscard]] friend FORCEINLINE bool operator!=( TWeakPtr const& InWeakPtr, TSharedPtr< ObjectTypeB, InMode > const& InSharedPtrB )
-	{
-		return InWeakPtr.Pin().Get() != InSharedPtrB.Get();
-	}
-
-
-	/**
-	 * Global equality operator for TWeakPtr
-	 *
-	 * @return  True if the weak pointer and the shared ref are not equal
-	 */
-	template< class ObjectTypeA>
-	[[nodiscard]] friend FORCEINLINE bool operator!=( TSharedRef< ObjectTypeA, InMode > const& InSharedRefA, TWeakPtr const& InWeakPtrB )
-	{
-		return &InSharedRefA.Get() != InWeakPtrB.Pin().Get();
-	}
-
-
-	/**
-	 * Global equality operator for TWeakPtr
-	 *
-	 * @return  True if the weak pointer and the shared ptr are not equal
-	 */
-	template< class ObjectTypeA >
-	[[nodiscard]] friend FORCEINLINE bool operator!=( TSharedPtr< ObjectTypeA, InMode > const& InSharedPtrA, TWeakPtr const& InWeakPtrB )
-	{
-		return InSharedPtrA.Get() != InWeakPtrB.Pin().Get();
-	}
-
-
-	/**
-	 * Global inequality operator for TWeakPtr
-	 *
-	 * @return  True if the weak pointer is not null
-	 */
-	[[nodiscard]] friend FORCEINLINE bool operator!=( TWeakPtr const& InWeakPtr, TYPE_OF_NULLPTR )
-	{
-		return InWeakPtr.IsValid();
-	}
-
-
-	/**
-	 * Global inequality operator for TWeakPtr
-	 *
-	 * @return  True if the weak pointer is not null
-	 */
-	[[nodiscard]] friend FORCEINLINE bool operator!=( TYPE_OF_NULLPTR, TWeakPtr const& InWeakPtr )
-	{
-		return InWeakPtr.IsValid();
-	}
 };
 
 
@@ -2060,6 +1720,318 @@ private:
 };
 
 
+/**
+ * Global equality operator for TSharedRef
+ *
+ * @return  True if the two shared references are equal
+ */
+template< class ObjectTypeA, class ObjectTypeB, ESPMode Mode >
+[[nodiscard]] FORCEINLINE bool operator==( TSharedRef< ObjectTypeA, Mode > const& InSharedRefA, TSharedRef< ObjectTypeB, Mode > const& InSharedRefB )
+{
+	return &( InSharedRefA.Get() ) == &( InSharedRefB.Get() );
+}
+
+
+/**
+ * Global inequality operator for TSharedRef
+ *
+ * @return  True if the two shared references are not equal
+ */
+template< class ObjectTypeA, class ObjectTypeB, ESPMode Mode >
+[[nodiscard]] FORCEINLINE bool operator!=( TSharedRef< ObjectTypeA, Mode > const& InSharedRefA, TSharedRef< ObjectTypeB, Mode > const& InSharedRefB )
+{
+	return &( InSharedRefA.Get() ) != &( InSharedRefB.Get() );
+}
+
+
+/**
+ * Global equality operator for TSharedPtr
+ *
+ * @return  True if the two shared pointers are equal
+ */
+template< class ObjectTypeA, class ObjectTypeB, ESPMode Mode >
+[[nodiscard]] FORCEINLINE bool operator==( TSharedPtr< ObjectTypeA, Mode > const& InSharedPtrA, TSharedPtr< ObjectTypeB, Mode > const& InSharedPtrB )
+{
+	return InSharedPtrA.Get() == InSharedPtrB.Get();
+}
+
+
+/**
+ * Global equality operator for TSharedPtr
+ *
+ * @return  True if the shared pointer is null
+ */
+template< class ObjectTypeA, ESPMode Mode >
+[[nodiscard]] FORCEINLINE bool operator==( TSharedPtr< ObjectTypeA, Mode > const& InSharedPtrA, TYPE_OF_NULLPTR )
+{
+	return !InSharedPtrA.IsValid();
+}
+
+
+/**
+ * Global equality operator for TSharedPtr
+ *
+ * @return  True if the shared pointer is null
+ */
+template< class ObjectTypeB, ESPMode Mode >
+[[nodiscard]] FORCEINLINE bool operator==( TYPE_OF_NULLPTR, TSharedPtr< ObjectTypeB, Mode > const& InSharedPtrB )
+{
+	return !InSharedPtrB.IsValid();
+}
+
+
+/**
+ * Global inequality operator for TSharedPtr
+ *
+ * @return  True if the two shared pointers are not equal
+ */
+template< class ObjectTypeA, class ObjectTypeB, ESPMode Mode >
+[[nodiscard]] FORCEINLINE bool operator!=( TSharedPtr< ObjectTypeA, Mode > const& InSharedPtrA, TSharedPtr< ObjectTypeB, Mode > const& InSharedPtrB )
+{
+	return InSharedPtrA.Get() != InSharedPtrB.Get();
+}
+
+
+/**
+ * Global inequality operator for TSharedPtr
+ *
+ * @return  True if the shared pointer is not null
+ */
+template< class ObjectTypeA, ESPMode Mode >
+[[nodiscard]] FORCEINLINE bool operator!=( TSharedPtr< ObjectTypeA, Mode > const& InSharedPtrA, TYPE_OF_NULLPTR )
+{
+	return InSharedPtrA.IsValid();
+}
+
+
+/**
+ * Global inequality operator for TSharedPtr
+ *
+ * @return  True if the shared pointer is not null
+ */
+template< class ObjectTypeB, ESPMode Mode >
+[[nodiscard]] FORCEINLINE bool operator!=( TYPE_OF_NULLPTR, TSharedPtr< ObjectTypeB, Mode > const& InSharedPtrB )
+{
+	return InSharedPtrB.IsValid();
+}
+
+
+/**
+ * Tests to see if a TSharedRef is "equal" to a TSharedPtr (both are valid and refer to the same object)
+ *
+ * @return  True if the shared reference and shared pointer are "equal"
+ */
+template< class ObjectTypeA, class ObjectTypeB, ESPMode Mode >
+[[nodiscard]] FORCEINLINE bool operator==( TSharedRef< ObjectTypeA, Mode > const& InSharedRef, TSharedPtr< ObjectTypeB, Mode > const& InSharedPtr )
+{
+	return InSharedPtr.IsValid() && InSharedPtr.Get() == &( InSharedRef.Get() );
+}
+
+
+/**
+ * Tests to see if a TSharedRef is not "equal" to a TSharedPtr (shared pointer is invalid, or both refer to different objects)
+ *
+ * @return  True if the shared reference and shared pointer are not "equal"
+ */
+template< class ObjectTypeA, class ObjectTypeB, ESPMode Mode >
+[[nodiscard]] FORCEINLINE bool operator!=( TSharedRef< ObjectTypeA, Mode > const& InSharedRef, TSharedPtr< ObjectTypeB, Mode > const& InSharedPtr )
+{
+	return !InSharedPtr.IsValid() || ( InSharedPtr.Get() != &( InSharedRef.Get() ) );
+}
+
+
+/**
+ * Tests to see if a TSharedRef is "equal" to a TSharedPtr (both are valid and refer to the same object) (reverse)
+ *
+ * @return  True if the shared reference and shared pointer are "equal"
+ */
+template< class ObjectTypeA, class ObjectTypeB, ESPMode Mode >
+[[nodiscard]] FORCEINLINE bool operator==( TSharedPtr< ObjectTypeB, Mode > const& InSharedPtr, TSharedRef< ObjectTypeA, Mode > const& InSharedRef )
+{
+	return InSharedRef == InSharedPtr;
+}
+
+
+/**
+ * Tests to see if a TSharedRef is not "equal" to a TSharedPtr (shared pointer is invalid, or both refer to different objects) (reverse)
+ *
+ * @return  True if the shared reference and shared pointer are not "equal"
+ */
+template< class ObjectTypeA, class ObjectTypeB, ESPMode Mode >
+[[nodiscard]] FORCEINLINE bool operator!=( TSharedPtr< ObjectTypeB, Mode > const& InSharedPtr, TSharedRef< ObjectTypeA, Mode > const& InSharedRef )
+{
+	return InSharedRef != InSharedPtr;
+}
+
+
+/**
+ * Global equality operator for TWeakPtr
+ *
+ * @return  True if the two weak pointers are equal
+ */
+template< class ObjectTypeA, class ObjectTypeB, ESPMode Mode >
+[[nodiscard]] FORCEINLINE bool operator==( TWeakPtr< ObjectTypeA, Mode > const& InWeakPtrA, TWeakPtr< ObjectTypeB, Mode > const& InWeakPtrB )
+{
+	return InWeakPtrA.Pin().Get() == InWeakPtrB.Pin().Get();
+}
+
+
+/**
+ * Global equality operator for TWeakPtr
+ *
+ * @return  True if the weak pointer and the shared ref are equal
+ */
+template< class ObjectTypeA, class ObjectTypeB, ESPMode Mode >
+[[nodiscard]] FORCEINLINE bool operator==( TWeakPtr< ObjectTypeA, Mode > const& InWeakPtrA, TSharedRef< ObjectTypeB, Mode > const& InSharedRefB )
+{
+	return InWeakPtrA.Pin().Get() == &InSharedRefB.Get();
+}
+
+
+/**
+ * Global equality operator for TWeakPtr
+ *
+ * @return  True if the weak pointer and the shared ptr are equal
+ */
+template< class ObjectTypeA, class ObjectTypeB, ESPMode Mode >
+[[nodiscard]] FORCEINLINE bool operator==( TWeakPtr< ObjectTypeA, Mode > const& InWeakPtrA, TSharedPtr< ObjectTypeB, Mode > const& InSharedPtrB )
+{
+	return InWeakPtrA.Pin().Get() == InSharedPtrB.Get();
+}
+
+
+/**
+ * Global equality operator for TWeakPtr
+ *
+ * @return  True if the weak pointer and the shared ref are equal
+ */
+template< class ObjectTypeA, class ObjectTypeB, ESPMode Mode >
+[[nodiscard]] FORCEINLINE bool operator==( TSharedRef< ObjectTypeA, Mode > const& InSharedRefA, TWeakPtr< ObjectTypeB, Mode > const& InWeakPtrB )
+{
+	return &InSharedRefA.Get() == InWeakPtrB.Pin().Get();
+}
+
+
+/**
+ * Global equality operator for TWeakPtr
+ *
+ * @return  True if the weak pointer and the shared ptr are equal
+ */
+template< class ObjectTypeA, class ObjectTypeB, ESPMode Mode >
+[[nodiscard]] FORCEINLINE bool operator==( TSharedPtr< ObjectTypeA, Mode > const& InSharedPtrA, TWeakPtr< ObjectTypeB, Mode > const& InWeakPtrB )
+{
+	return InSharedPtrA.Get() == InWeakPtrB.Pin().Get();
+}
+
+
+/**
+ * Global equality operator for TWeakPtr
+ *
+ * @return  True if the weak pointer is null
+ */
+template< class ObjectTypeA, ESPMode Mode >
+[[nodiscard]] FORCEINLINE bool operator==( TWeakPtr< ObjectTypeA, Mode > const& InWeakPtrA, TYPE_OF_NULLPTR )
+{
+	return !InWeakPtrA.IsValid();
+}
+
+
+/**
+ * Global equality operator for TWeakPtr
+ *
+ * @return  True if the weak pointer is null
+ */
+template< class ObjectTypeB, ESPMode Mode >
+[[nodiscard]] FORCEINLINE bool operator==( TYPE_OF_NULLPTR, TWeakPtr< ObjectTypeB, Mode > const& InWeakPtrB )
+{
+	return !InWeakPtrB.IsValid();
+}
+
+
+/**
+ * Global inequality operator for TWeakPtr
+ *
+ * @return  True if the two weak pointers are not equal
+ */
+template< class ObjectTypeA, class ObjectTypeB, ESPMode Mode >
+[[nodiscard]] FORCEINLINE bool operator!=( TWeakPtr< ObjectTypeA, Mode > const& InWeakPtrA, TWeakPtr< ObjectTypeB, Mode > const& InWeakPtrB )
+{
+	return InWeakPtrA.Pin().Get() != InWeakPtrB.Pin().Get();
+}
+
+
+/**
+ * Global equality operator for TWeakPtr
+ *
+ * @return  True if the weak pointer and the shared ref are not equal
+ */
+template< class ObjectTypeA, class ObjectTypeB, ESPMode Mode >
+[[nodiscard]] FORCEINLINE bool operator!=( TWeakPtr< ObjectTypeA, Mode > const& InWeakPtrA, TSharedRef< ObjectTypeB, Mode > const& InSharedRefB )
+{
+	return InWeakPtrA.Pin().Get() != &InSharedRefB.Get();
+}
+
+
+/**
+ * Global equality operator for TWeakPtr
+ *
+ * @return  True if the weak pointer and the shared ptr are not equal
+ */
+template< class ObjectTypeA, class ObjectTypeB, ESPMode Mode >
+[[nodiscard]] FORCEINLINE bool operator!=( TWeakPtr< ObjectTypeA, Mode > const& InWeakPtrA, TSharedPtr< ObjectTypeB, Mode > const& InSharedPtrB )
+{
+	return InWeakPtrA.Pin().Get() != InSharedPtrB.Get();
+}
+
+
+/**
+ * Global equality operator for TWeakPtr
+ *
+ * @return  True if the weak pointer and the shared ref are not equal
+ */
+template< class ObjectTypeA, class ObjectTypeB, ESPMode Mode >
+[[nodiscard]] FORCEINLINE bool operator!=( TSharedRef< ObjectTypeA, Mode > const& InSharedRefA, TWeakPtr< ObjectTypeB, Mode > const& InWeakPtrB )
+{
+	return &InSharedRefA.Get() != InWeakPtrB.Pin().Get();
+}
+
+
+/**
+ * Global equality operator for TWeakPtr
+ *
+ * @return  True if the weak pointer and the shared ptr are not equal
+ */
+template< class ObjectTypeA, class ObjectTypeB, ESPMode Mode >
+[[nodiscard]] FORCEINLINE bool operator!=( TSharedPtr< ObjectTypeA, Mode > const& InSharedPtrA, TWeakPtr< ObjectTypeB, Mode > const& InWeakPtrB )
+{
+	return InSharedPtrA.Get() != InWeakPtrB.Pin().Get();
+}
+
+
+/**
+ * Global inequality operator for TWeakPtr
+ *
+ * @return  True if the weak pointer is not null
+ */
+template< class ObjectTypeA, ESPMode Mode >
+[[nodiscard]] FORCEINLINE bool operator!=( TWeakPtr< ObjectTypeA, Mode > const& InWeakPtrA, TYPE_OF_NULLPTR )
+{
+	return InWeakPtrA.IsValid();
+}
+
+
+/**
+ * Global inequality operator for TWeakPtr
+ *
+ * @return  True if the weak pointer is not null
+ */
+template< class ObjectTypeB, ESPMode Mode >
+[[nodiscard]] FORCEINLINE bool operator!=( TYPE_OF_NULLPTR, TWeakPtr< ObjectTypeB, Mode > const& InWeakPtrB )
+{
+	return InWeakPtrB.IsValid();
+}
+
+
 /**
  * Casts a shared pointer of one type to another type. (static_cast)  Useful for down-casting.
  *
@@ -2199,6 +2171,44 @@ FORCEINLINE void CleanupPointerMap(TMap< TWeakPtr<KeyType>, ValueType >& Pointer
 	PointerMap = NewMap;
 }
 
+/**
+* Computes a hash code for this object
+*
+* @param  InSharedRef  Shared pointer to compute hash code for
+*
+* @return  Hash code value
+*/
+template <typename ObjectType, ESPMode Mode>
+[[nodiscard]] uint32 GetTypeHash( const TSharedRef< ObjectType, Mode >& InSharedRef )
+{
+	return ::PointerHash( &InSharedRef.Get() );
+}
+
+/**
+* Computes a hash code for this object
+*
+* @param  InSharedPtr  Shared pointer to compute hash code for
+*
+* @return  Hash code value
+*/
+template <typename ObjectType, ESPMode Mode>
+[[nodiscard]] uint32 GetTypeHash( const TSharedPtr< ObjectType, Mode >& InSharedPtr )
+{
+	return ::PointerHash( InSharedPtr.Get() );
+}
+
+/**
+* Computes a hash code for this object
+*
+* @param  InWeakPtr  Weak pointer to compute hash code for
+*
+* @return  Hash code value
+*/
+template <typename ObjectType, ESPMode Mode>
+[[nodiscard]] uint32 GetTypeHash( const TWeakPtr< ObjectType, Mode >& InWeakPtr )
+{
+	return InWeakPtr.GetWeakPtrTypeHash();
+}
 
 // Shared pointer testing
 #include "Templates/SharedPointerTesting.inl" // IWYU pragma: export
diff --git a/Engine/Source/Runtime/Core/Public/Templates/SharedPointerTesting.inl b/Engine/Source/Runtime/Core/Public/Templates/SharedPointerTesting.inl
--- a/Engine/Source/Runtime/Core/Public/Templates/SharedPointerTesting.inl
+++ b/Engine/Source/Runtime/Core/Public/Templates/SharedPointerTesting.inl
@@ -6,12 +6,11 @@
 /** Toggle this define to enable shared pointer testing features */
 #define WITH_SHARED_POINTER_TESTS 0 && !(UE_BUILD_SHIPPING || UE_BUILD_TEST)
 
-enum class ESPMode : uint8;
+#if WITH_SHARED_POINTER_TESTS
 
+enum class ESPMode : uint8;
 template<typename KeyType,typename ValueType,typename SetAllocator ,typename KeyFuncs > class TMap;
 
-#if WITH_SHARED_POINTER_TESTS
-
 /**
  * Shared pointer testing suite
  */
diff --git a/Engine/Source/Runtime/Core/Public/Templates/Tuple.h b/Engine/Source/Runtime/Core/Public/Templates/Tuple.h
--- a/Engine/Source/Runtime/Core/Public/Templates/Tuple.h
+++ b/Engine/Source/Runtime/Core/Public/Templates/Tuple.h
@@ -341,15 +341,14 @@ namespace UE::Core::Private::Tuple
 		template <typename FuncType, typename... ArgTypes> decltype(auto) ApplyBefore(FuncType&& Func, ArgTypes&&... Args)       volatile&& { return ::Invoke(Func, static_cast<      volatile TTupleBase&&>(*this).template Get<Indices>()..., Forward<ArgTypes>(Args)...); }
 		template <typename FuncType, typename... ArgTypes> decltype(auto) ApplyBefore(FuncType&& Func, ArgTypes&&... Args) const volatile&& { return ::Invoke(Func, static_cast<const volatile TTupleBase&&>(*this).template Get<Indices>()..., Forward<ArgTypes>(Args)...); }
 
-		FORCEINLINE friend FArchive& operator<<(FArchive& Ar, TTupleBase& Tuple)
+		FORCEINLINE static void Serialize(FArchive& Ar, TTupleBase& Tuple)
 		{
 			// This should be implemented with a fold expression when our compilers support it
 			int Temp[] = { 0, (Ar << Tuple.template Get<Indices>(), 0)... };
 			(void)Temp;
-			return Ar;
 		}
 
-		FORCEINLINE friend void operator<<(FStructuredArchive::FSlot Slot, TTupleBase& Tuple)
+		FORCEINLINE static void SerializeStructured(FStructuredArchive::FSlot Slot, TTupleBase& Tuple)
 		{
 			// This should be implemented with a fold expression when our compilers support it
 			FStructuredArchive::FStream Stream = Slot.EnterStream();
@@ -357,35 +356,35 @@ namespace UE::Core::Private::Tuple
 			(void)Temp;
 		}
 
-		FORCEINLINE friend bool operator==(const TTupleBase& Lhs, const TTupleBase& Rhs)
+		FORCEINLINE bool operator==(const TTupleBase& Rhs) const
 		{
 			// This could be implemented with a fold expression when our compilers support it
-			return FEqualityHelper<sizeof...(Types), 0>::Compare(Lhs, Rhs);
+			return FEqualityHelper<sizeof...(Types), 0>::Compare(*this, Rhs);
 		}
 
-		FORCEINLINE friend bool operator!=(const TTupleBase& Lhs, const TTupleBase& Rhs)
+		FORCEINLINE bool operator!=(const TTupleBase& Rhs) const
 		{
-			return !(Lhs == Rhs);
+			return !(*this == Rhs);
 		}
 
-		FORCEINLINE friend bool operator<(const TTupleBase& Lhs, const TTupleBase& Rhs)
+		FORCEINLINE bool operator<(const TTupleBase& Rhs) const
 		{
-			return TLessThanHelper<sizeof...(Types)>::Do(Lhs, Rhs);
+			return TLessThanHelper<sizeof...(Types)>::Do(*this, Rhs);
 		}
 
-		FORCEINLINE friend bool operator<=(const TTupleBase& Lhs, const TTupleBase& Rhs)
+		FORCEINLINE bool operator<=(const TTupleBase& Rhs) const
 		{
-			return !(Rhs < Lhs);
+			return !(Rhs < *this);
 		}
 
-		FORCEINLINE friend bool operator>(const TTupleBase& Lhs, const TTupleBase& Rhs)
+		FORCEINLINE bool operator>(const TTupleBase& Rhs) const
 		{
-			return Rhs < Lhs;
+			return Rhs < *this;
 		}
 
-		FORCEINLINE friend bool operator>=(const TTupleBase& Lhs, const TTupleBase& Rhs)
+		FORCEINLINE bool operator>=(const TTupleBase& Rhs) const
 		{
-			return !(Lhs < Rhs);
+			return !(*this < Rhs);
 		}
 	};
 
@@ -765,12 +764,19 @@ public:
 	template <int N> friend decltype(auto) get(const volatile TTuple&& val) { return static_cast<const volatile TTuple&&>(val).template Get<N>(); }
 #endif
 
-	FORCEINLINE friend uint32 GetTypeHash(const TTuple<Types...>& Tuple)
-	{
-		return UE::Core::Private::Tuple::TGetTupleHashHelper<1u, sizeof...(Types)>::Do(GetTypeHash(Tuple.template Get<0>()), Tuple);
-	}
 };
 
+template <typename... Types>
+FORCEINLINE uint32 GetTypeHash(const TTuple<Types...>& Tuple)
+{
+	return UE::Core::Private::Tuple::TGetTupleHashHelper<1u, sizeof...(Types)>::Do(GetTypeHash(Tuple.template Get<0>()), Tuple);
+}
+
+FORCEINLINE uint32 GetTypeHash(const TTuple<>& Tuple)
+{
+	return 0;
+}
+
 namespace Freeze
 {
 	template<typename KeyType, typename ValueType>
@@ -995,4 +1001,17 @@ template <typename T>
 struct TIsTuple
 {
 	enum { Value = TIsTuple_V<T> };
-};
\ No newline at end of file
+};
+
+template <typename... Types>
+FORCEINLINE FArchive& operator<<(FArchive& Ar, TTuple<Types...>& Tuple)
+{
+	TTuple<Types...>::Serialize(Ar, Tuple);
+	return Ar;
+}
+
+template <typename... Types>
+FORCEINLINE void operator<<(FStructuredArchive::FSlot Slot, TTuple<Types...>& Tuple)
+{
+	TTuple<Types...>::SerializeStructured(Slot, Tuple);
+}
diff --git a/Engine/Source/Runtime/Core/Public/Templates/UniqueObj.h b/Engine/Source/Runtime/Core/Public/Templates/UniqueObj.h
--- a/Engine/Source/Runtime/Core/Public/Templates/UniqueObj.h
+++ b/Engine/Source/Runtime/Core/Public/Templates/UniqueObj.h
@@ -58,12 +58,18 @@ public:
 	      T& operator*()       { return *Obj; }
 	const T& operator*() const { return *Obj; }
 
-	friend FArchive& operator<<(FArchive& Ar, TUniqueObj& P)
+	void Serialize(FArchive& Ar)
 	{
-		Ar << *P.Obj;
-		return Ar;
+		Ar << *Obj;
 	}
 
 private:
 	TUniquePtr<T> Obj;
 };
+
+template <typename T>
+FArchive& operator<<(FArchive& Ar, TUniqueObj<T>& P)
+{
+	P.Serialize(Ar);
+	return Ar;
+}
diff --git a/Engine/Source/Runtime/Core/Public/Templates/UniquePtr.h b/Engine/Source/Runtime/Core/Public/Templates/UniquePtr.h
--- a/Engine/Source/Runtime/Core/Public/Templates/UniquePtr.h
+++ b/Engine/Source/Runtime/Core/Public/Templates/UniquePtr.h
@@ -378,6 +378,7 @@ private:
 	using PtrType = T*;
 	LAYOUT_FIELD(PtrType, Ptr);
 
+public:
 	/**
 	 * Equality comparison operator
 	 *
@@ -387,39 +388,36 @@ private:
 	 * @return true if the two TUniquePtrs are logically substitutable for each other, false otherwise.
 	 */
 	template <typename RhsT>
-	friend FORCEINLINE bool operator==(const TUniquePtr& Lhs, const TUniquePtr<RhsT>& Rhs)
+	FORCEINLINE bool operator==(const TUniquePtr<RhsT>& Rhs) const
 	{
-		return Lhs.Get() == Rhs.Get();
+		return Get() == Rhs.Get();
 	}
 
 	/**
-	 * Inequality comparison operator
+	 * Equality comparison operator against nullptr.
 	 *
-	 * @param Lhs The first TUniquePtr to compare.
-	 * @param Rhs The second TUniquePtr to compare.
+	 * @param Lhs The TUniquePtr to compare.
 	 *
-	 * @return false if the two TUniquePtrs are logically substitutable for each other, true otherwise.
+	 * @return true if the TUniquePtr is null, false otherwise.
 	 */
-	template <typename RhsT>
-	friend FORCEINLINE bool operator!=(const TUniquePtr& Lhs, const TUniquePtr<RhsT>& Rhs)
+	FORCEINLINE bool operator==(TYPE_OF_NULLPTR) const
 	{
-		return Lhs.Get() != Rhs.Get();
+		return !IsValid();
 	}
 
+#if !PLATFORM_COMPILER_HAS_GENERATED_COMPARISON_OPERATORS
 	/**
-	 * Equality comparison operator against nullptr.
+	 * Inequality comparison operator
 	 *
-	 * @param Lhs The TUniquePtr to compare.
+	 * @param Lhs The first TUniquePtr to compare.
+	 * @param Rhs The second TUniquePtr to compare.
 	 *
-	 * @return true if the TUniquePtr is null, false otherwise.
+	 * @return false if the two TUniquePtrs are logically substitutable for each other, true otherwise.
 	 */
-	friend FORCEINLINE bool operator==(const TUniquePtr& Lhs, TYPE_OF_NULLPTR)
-	{
-		return !Lhs.IsValid();
-	}
-	friend FORCEINLINE bool operator==(TYPE_OF_NULLPTR, const TUniquePtr& Rhs)
+	template <typename RhsT>
+	FORCEINLINE bool operator!=(const TUniquePtr<RhsT>& Rhs) const
 	{
-		return !Rhs.IsValid();
+		return Get() != Rhs.Get();
 	}
 
 	/**
@@ -429,14 +427,11 @@ private:
 	 *
 	 * @return true if the TUniquePtr is not null, false otherwise.
 	 */
-	friend FORCEINLINE bool operator!=(const TUniquePtr& Lhs, TYPE_OF_NULLPTR)
+	FORCEINLINE bool operator!=(TYPE_OF_NULLPTR) const
 	{
-		return Lhs.IsValid();
-	}
-	friend FORCEINLINE bool operator!=(TYPE_OF_NULLPTR, const TUniquePtr& Rhs)
-	{
-		return Rhs.IsValid();
+		return IsValid();
 	}
+#endif
 };
 
 template <typename T, typename Deleter>
@@ -711,6 +706,7 @@ public:
 private:
 	T* Ptr;
 
+public:
 	/**
 	 * Equality comparison operator
 	 *
@@ -720,39 +716,36 @@ private:
 	 * @return true if the two TUniquePtrs are logically substitutable for each other, false otherwise.
 	 */
 	template <typename RhsT>
-	friend FORCEINLINE bool operator==(const TUniquePtr& Lhs, const TUniquePtr<RhsT>& Rhs)
+	FORCEINLINE bool operator==(const TUniquePtr<RhsT>& Rhs) const
 	{
-		return Lhs.Get() == Rhs.Get();
+		return Get() == Rhs.Get();
 	}
 
 	/**
-	 * Inequality comparison operator
+	 * Equality comparison operator against nullptr.
 	 *
-	 * @param Lhs The first TUniquePtr to compare.
-	 * @param Rhs The second TUniquePtr to compare.
+	 * @param Lhs The TUniquePtr to compare.
 	 *
-	 * @return false if the two TUniquePtrs are logically substitutable for each other, true otherwise.
+	 * @return true if the TUniquePtr is null, false otherwise.
 	 */
-	template <typename RhsT>
-	friend FORCEINLINE bool operator!=(const TUniquePtr& Lhs, const TUniquePtr<RhsT>& Rhs)
+	FORCEINLINE bool operator==(TYPE_OF_NULLPTR) const
 	{
-		return Lhs.Get() != Rhs.Get();
+		return !IsValid();
 	}
 
+#if !PLATFORM_COMPILER_HAS_GENERATED_COMPARISON_OPERATORS
 	/**
-	 * Equality comparison operator against nullptr.
+	 * Inequality comparison operator
 	 *
-	 * @param Lhs The TUniquePtr to compare.
+	 * @param Lhs The first TUniquePtr to compare.
+	 * @param Rhs The second TUniquePtr to compare.
 	 *
-	 * @return true if the TUniquePtr is null, false otherwise.
+	 * @return false if the two TUniquePtrs are logically substitutable for each other, true otherwise.
 	 */
-	friend FORCEINLINE bool operator==(const TUniquePtr& Lhs, TYPE_OF_NULLPTR)
-	{
-		return !Lhs.IsValid();
-	}
-	friend FORCEINLINE bool operator==(TYPE_OF_NULLPTR, const TUniquePtr& Rhs)
+	template <typename RhsT>
+	FORCEINLINE bool operator!=(const TUniquePtr<RhsT>& Rhs) const
 	{
-		return !Rhs.IsValid();
+		return Get() != Rhs.Get();
 	}
 
 	/**
@@ -762,16 +755,27 @@ private:
 	 *
 	 * @return true if the TUniquePtr is not null, false otherwise.
 	 */
-	friend FORCEINLINE bool operator!=(const TUniquePtr& Lhs, TYPE_OF_NULLPTR)
-	{
-		return Lhs.IsValid();
-	}
-	friend FORCEINLINE bool operator!=(TYPE_OF_NULLPTR, const TUniquePtr& Rhs)
+	FORCEINLINE bool operator!=(TYPE_OF_NULLPTR) const
 	{
-		return Rhs.IsValid();
+		return IsValid();
 	}
+#endif
 };
 
+#if !PLATFORM_COMPILER_HAS_GENERATED_COMPARISON_OPERATORS
+template <typename T>
+FORCEINLINE bool operator==(TYPE_OF_NULLPTR, const TUniquePtr<T>& Rhs)
+{
+	return !Rhs.IsValid();
+}
+
+template <typename T>
+FORCEINLINE bool operator!=(TYPE_OF_NULLPTR, const TUniquePtr<T>& Rhs)
+{
+	return Rhs.IsValid();
+}
+#endif
+
 // Trait which allows TUniquePtr to be default constructed by memsetting to zero.
 template <typename T>
 struct TIsZeroConstructType<TUniquePtr<T>>
diff --git a/Engine/Source/Runtime/Core/Public/UObject/NameTypes.h b/Engine/Source/Runtime/Core/Public/UObject/NameTypes.h
--- a/Engine/Source/Runtime/Core/Public/UObject/NameTypes.h
+++ b/Engine/Source/Runtime/Core/Public/UObject/NameTypes.h
@@ -114,10 +114,10 @@ public:
 	friend inline bool operator!=(EName Ename, FNameEntryId Id) { return !(Id == Ename); }
 	friend inline bool operator!=(FNameEntryId Id, EName Ename) { return !(Id == Ename); }
 	friend CORE_API uint32 GetTypeHash(FNameEntryId Id);
-};
 
-/** Serialize as process specific unstable int */
-CORE_API FArchive& operator<<(FArchive& Ar, FNameEntryId& InId);
+	/** Serialize as process specific unstable int */
+	CORE_API friend FArchive& operator<<(FArchive& Ar, FNameEntryId& InId);
+};
 
 /**
  * Legacy typedef - this is no longer an index
diff --git a/Engine/Source/Runtime/Core/Public/UObject/WeakObjectPtrTemplates.h b/Engine/Source/Runtime/Core/Public/UObject/WeakObjectPtrTemplates.h
--- a/Engine/Source/Runtime/Core/Public/UObject/WeakObjectPtrTemplates.h
+++ b/Engine/Source/Runtime/Core/Public/UObject/WeakObjectPtrTemplates.h
@@ -215,19 +215,18 @@ public:
 		return static_cast<const TWeakObjectPtrBase&>(*this).HasSameIndexAndSerialNumber(static_cast<const TWeakObjectPtrBase&>(Other));
 	}
 
-	/** Hash function. */
-	FORCEINLINE friend uint32 GetTypeHash(const TWeakObjectPtr& WeakObjectPtr)
-	{
-		return GetTypeHash(static_cast<const TWeakObjectPtrBase&>(WeakObjectPtr));
-	}
-
 	/**
 	 * Weak object pointer serialization, this forwards to FArchive::operator<<(struct FWeakObjectPtr&) or an override
 	 */
-	friend FArchive& operator<<( FArchive& Ar, TWeakObjectPtr& WeakObjectPtr )
+	FORCEINLINE	void Serialize(FArchive& Ar)
+	{
+		Ar << static_cast<TWeakObjectPtrBase&>(*this);
+	}
+
+	/** Hash function. */
+	FORCEINLINE uint32 GetWeakPtrTypeHash() const
 	{
-		Ar << static_cast<TWeakObjectPtrBase&>(WeakObjectPtr);
-		return Ar;
+		return static_cast<const TWeakObjectPtrBase&>(*this).GetTypeHash();
 	}
 
 	/**
@@ -236,70 +235,82 @@ public:
 	 * @param Other weak pointer to compare to
 	 */
 	template <typename RhsT, typename = decltype((T*)nullptr == (RhsT*)nullptr)>
-	friend FORCENOINLINE bool operator==(const TWeakObjectPtr& Lhs, const TWeakObjectPtr<RhsT, TWeakObjectPtrBase>& Rhs)
+	FORCENOINLINE bool operator==(const TWeakObjectPtr<RhsT, TWeakObjectPtrBase>& Rhs) const
 	{
-		return (const TWeakObjectPtrBase&)Lhs == (const TWeakObjectPtrBase&)Rhs;
+		return (const TWeakObjectPtrBase&)*this == (const TWeakObjectPtrBase&)Rhs;
 	}
 
 	template <typename RhsT, typename = decltype((T*)nullptr == (RhsT*)nullptr)>
-	friend FORCENOINLINE bool operator==(const TWeakObjectPtr& Lhs, const RhsT* Rhs)
-	{
-		// NOTE: this constructs a TWeakObjectPtrBase, which has some amount of overhead, so this may not be an efficient operation
-		return (const TWeakObjectPtrBase&)Lhs == TWeakObjectPtrBase(Rhs);
-	}
-
-	template <typename LhsT, typename = decltype((LhsT*)nullptr == (T*)nullptr)>
-	friend FORCENOINLINE bool operator==(const LhsT* Lhs, const TWeakObjectPtr& Rhs)
+	FORCENOINLINE bool operator==(const RhsT* Rhs) const
 	{
 		// NOTE: this constructs a TWeakObjectPtrBase, which has some amount of overhead, so this may not be an efficient operation
-		return TWeakObjectPtrBase(Lhs) == (const TWeakObjectPtrBase&)Rhs;
+		return (const TWeakObjectPtrBase&)*this == TWeakObjectPtrBase(Rhs);
 	}
 
-	friend FORCENOINLINE bool operator==(const TWeakObjectPtr& Lhs, TYPE_OF_NULLPTR)
+	FORCENOINLINE bool operator==(TYPE_OF_NULLPTR) const
 	{
-		return !Lhs.IsValid();
-	}
-
-	friend FORCENOINLINE bool operator==(TYPE_OF_NULLPTR, const TWeakObjectPtr& Rhs)
-	{
-		return !Rhs.IsValid();
+		return !IsValid();
 	}
 
+#if !PLATFORM_COMPILER_HAS_GENERATED_COMPARISON_OPERATORS
 	/**
 	 * Compare weak pointers for inequality
 	 * @param Other weak pointer to compare to
 	 */
 	template <typename RhsT, typename = decltype((T*)nullptr != (RhsT*)nullptr)>
-	friend FORCENOINLINE bool operator!=(const TWeakObjectPtr& Lhs, const TWeakObjectPtr<RhsT, TWeakObjectPtrBase>& Rhs)
+	FORCENOINLINE bool operator!=(const TWeakObjectPtr<RhsT, TWeakObjectPtrBase>& Rhs) const
 	{
-		return (const TWeakObjectPtrBase&)Lhs != (const TWeakObjectPtrBase&)Rhs;
+		return (const TWeakObjectPtrBase&)*this != (const TWeakObjectPtrBase&)Rhs;
 	}
 
 	template <typename RhsT, typename = decltype((T*)nullptr != (RhsT*)nullptr)>
-	friend FORCENOINLINE bool operator!=(const TWeakObjectPtr& Lhs, const RhsT* Rhs)
+	FORCENOINLINE bool operator!=(const RhsT* Rhs) const
 	{
 		// NOTE: this constructs a TWeakObjectPtrBase, which has some amount of overhead, so this may not be an efficient operation
-		return (const TWeakObjectPtrBase&)Lhs != TWeakObjectPtrBase(Rhs);
+		return (const TWeakObjectPtrBase&)*this != TWeakObjectPtrBase(Rhs);
 	}
 
-	template <typename LhsT, typename = decltype((LhsT*)nullptr != (T*)nullptr)>
-	friend FORCENOINLINE bool operator!=(const LhsT* Lhs, const TWeakObjectPtr& Rhs)
+	FORCENOINLINE bool operator!=(TYPE_OF_NULLPTR) const
 	{
-		// NOTE: this constructs a TWeakObjectPtrBase, which has some amount of overhead, so this may not be an efficient operation
-		return TWeakObjectPtrBase(Lhs) != (const TWeakObjectPtrBase&)Rhs;
-	}
-
-	friend FORCENOINLINE bool operator!=(const TWeakObjectPtr& Lhs, TYPE_OF_NULLPTR)
-	{
-		return Lhs.IsValid();
+		return IsValid();
 	}
+#endif
+};
 
-	friend FORCENOINLINE bool operator!=(TYPE_OF_NULLPTR, const TWeakObjectPtr& Rhs)
-	{
-		return Rhs.IsValid();
-	}
+#if !PLATFORM_COMPILER_HAS_GENERATED_COMPARISON_OPERATORS
+/**
+ * Compare weak pointers for equality.
+ * If both pointers would return nullptr from Get() they count as equal even if they were not initialized to the same object.
+ * @param Other weak pointer to compare to
+ */
+template <typename LhsT, typename RhsT, typename OtherTWeakObjectPtrBase, typename = decltype((LhsT*)nullptr == (RhsT*)nullptr)>
+FORCENOINLINE bool operator==(const LhsT* Lhs, const TWeakObjectPtr<RhsT, OtherTWeakObjectPtrBase>& Rhs)
+{
+	// NOTE: this constructs a TWeakObjectPtrBase, which has some amount of overhead, so this may not be an efficient operation
+	return OtherTWeakObjectPtrBase(Lhs) == (const OtherTWeakObjectPtrBase&)Rhs;
+}
+template <typename RhsT, typename OtherTWeakObjectPtrBase>
+FORCENOINLINE bool operator==(TYPE_OF_NULLPTR, const TWeakObjectPtr<RhsT, OtherTWeakObjectPtrBase>& Rhs)
+{
+	return !Rhs.IsValid();
+}
 
-};
+/**
+ * Compare weak pointers for inequality
+ * @param Other weak pointer to compare to
+ */
+template <typename LhsT, typename RhsT, typename OtherTWeakObjectPtrBase, typename = decltype((LhsT*)nullptr != (RhsT*)nullptr)>
+FORCENOINLINE bool operator!=(const LhsT* Lhs, const TWeakObjectPtr<RhsT, OtherTWeakObjectPtrBase>& Rhs)
+{
+	// NOTE: this constructs a TWeakObjectPtrBase, which has some amount of overhead, so this may not be an efficient operation
+	return OtherTWeakObjectPtrBase(Lhs) != (const OtherTWeakObjectPtrBase&)Rhs;
+}
+template <typename RhsT, typename OtherTWeakObjectPtrBase>
+FORCENOINLINE bool operator!=(TYPE_OF_NULLPTR, const TWeakObjectPtr<RhsT, OtherTWeakObjectPtrBase>& Rhs)
+{
+	return Rhs.IsValid();
+}
+#endif
 
 // Helper function which deduces the type of the initializer
 template <typename T>
@@ -390,4 +401,22 @@ void CopyFromWeakArray(DestArrayType& Dest, const SourceArrayType& Src)
 			Dest.Add(Value);
 		}
 	}
-}
\ No newline at end of file
+}
+
+/** Hash function. */
+template <typename T>
+FORCEINLINE uint32 GetTypeHash(const TWeakObjectPtr<T>& WeakObjectPtr)
+{
+	return WeakObjectPtr.GetWeakPtrTypeHash();
+}
+
+
+/**
+* Weak object pointer serialization, this forwards to FArchive::operator<<(struct FWeakObjectPtr&) or an override
+*/
+template<class T, class TWeakObjectPtrBase>
+FArchive& operator<<( FArchive& Ar, TWeakObjectPtr<T, TWeakObjectPtrBase>& WeakObjectPtr )
+{
+	WeakObjectPtr.Serialize(Ar);
+	return Ar;
+}
diff --git a/Engine/Source/Runtime/CoreUObject/Public/Serialization/ArchiveCrc32.h b/Engine/Source/Runtime/CoreUObject/Public/Serialization/ArchiveCrc32.h
--- a/Engine/Source/Runtime/CoreUObject/Public/Serialization/ArchiveCrc32.h
+++ b/Engine/Source/Runtime/CoreUObject/Public/Serialization/ArchiveCrc32.h
@@ -37,21 +37,21 @@ public:
 	virtual FString GetArchiveName() const { return TEXT("FArchiveCrc32"); }
 	//~ End FArchive Interface
 
-	/**
-	 * Serializes a USTRUCT value from or into an archive.
-	 *
-	 * @param Ar The archive to serialize from or to.
-	 * @param Value The value to serialize.
-	 */
-	template <typename StructType>
-	FORCEINLINE friend typename TEnableIf<TModels<CStaticStructProvider, StructType>::Value, FArchive&>::Type operator <<(FArchive& Ar, const StructType& Value)
-	{
-		StructType* MutableValue = const_cast<StructType*>(&Value);
-		StructType::StaticStruct()->SerializeItem(Ar, MutableValue, nullptr);
-		return Ar;
-	}
-
 private:
 	uint32 CRC;
 	UObject* RootObject;
 };
+
+/**
+	* Serializes a USTRUCT value from or into an archive.
+	*
+	* @param Ar The archive to serialize from or to.
+	* @param Value The value to serialize.
+	*/
+template <typename StructType>
+FORCEINLINE typename TEnableIf<TModels<CStaticStructProvider, StructType>::Value, FArchive&>::Type operator <<(FArchive& Ar, const StructType& Value)
+{
+	StructType* MutableValue = const_cast<StructType*>(&Value);
+	StructType::StaticStruct()->SerializeItem(Ar, MutableValue, nullptr);
+	return Ar;
+}
diff --git a/Engine/Source/Runtime/CoreUObject/Public/Templates/SubclassOf.h b/Engine/Source/Runtime/CoreUObject/Public/Templates/SubclassOf.h
--- a/Engine/Source/Runtime/CoreUObject/Public/Templates/SubclassOf.h
+++ b/Engine/Source/Runtime/CoreUObject/Public/Templates/SubclassOf.h
@@ -137,10 +137,9 @@ public:
 		return (TClass*)Result;
 	}
 
-	friend FArchive& operator<<(FArchive& Ar, TSubclassOf& SubclassOf)
+	FORCEINLINE void SerializeTSubclassOf(FArchive& Ar)
 	{
-		Ar << SubclassOf.Class;
-		return Ar;
+		Ar << Class;
 	}
 
 	friend uint32 GetTypeHash(const TSubclassOf& SubclassOf)
@@ -166,3 +165,11 @@ struct TIsTSubclassOf<TSubclassOf<T>>
 {
 	enum { Value = true };
 };
+
+template <typename T>
+FArchive& operator<<(FArchive& Ar, TSubclassOf<T>& SubclassOf)
+{
+	SubclassOf.SerializeTSubclassOf(Ar);
+	return Ar;
+}
+
diff --git a/Engine/Source/Runtime/CoreUObject/Public/UObject/FieldIterator.h b/Engine/Source/Runtime/CoreUObject/Public/UObject/FieldIterator.h
--- a/Engine/Source/Runtime/CoreUObject/Public/UObject/FieldIterator.h
+++ b/Engine/Source/Runtime/CoreUObject/Public/UObject/FieldIterator.h
@@ -73,13 +73,13 @@ public:
 		return !(bool)*this;
 	}
 
-	inline friend bool operator==(const TAllFieldsIterator<T>& Lhs, const TAllFieldsIterator<T>& Rhs) 
+	inline bool operator==(const TAllFieldsIterator<T>& Rhs) const 
 	{ 
-		return *Lhs.FieldIterator == *Rhs.FieldIterator && Lhs.CurrentFieldIndex == Rhs.CurrentFieldIndex; 
+		return *FieldIterator == *Rhs.FieldIterator && CurrentFieldIndex == Rhs.CurrentFieldIndex; 
 	}
-	inline friend bool operator!=(const TAllFieldsIterator<T>& Lhs, const TAllFieldsIterator<T>& Rhs) 
+	inline bool operator!=(const TAllFieldsIterator<T>& Rhs) const
 	{ 
-		return *Lhs.FieldIterator != *Rhs.FieldIterator || Lhs.CurrentFieldIndex != Rhs.CurrentFieldIndex;
+		return *FieldIterator != *Rhs.FieldIterator || CurrentFieldIndex != Rhs.CurrentFieldIndex;
 	}
 
 	inline void operator++()
diff --git a/Engine/Source/Runtime/CoreUObject/Public/UObject/FieldPath.h b/Engine/Source/Runtime/CoreUObject/Public/UObject/FieldPath.h
--- a/Engine/Source/Runtime/CoreUObject/Public/UObject/FieldPath.h
+++ b/Engine/Source/Runtime/CoreUObject/Public/UObject/FieldPath.h
@@ -438,29 +438,36 @@ public:
 	}
 
 	/**
-	* Compare weak pointers for inequality
-	* @param Other weak pointer to compare to
+	* Compare weak pointers for equality
+	* @param Other pointer to compare to
 	**/
 	template <typename OtherPropertyType>
-	FORCEINLINE bool operator!=(const TFieldPath<OtherPropertyType> &Other) const
+	FORCEINLINE bool operator==(const OtherPropertyType* Other) const
 	{
 		static_assert(TPointerIsConvertibleFromTo<OtherPropertyType, const FField>::Value, "TFieldPath can only be compared with FField types");
 		static_assert(TPointerIsConvertibleFromTo<PropertyType, OtherPropertyType>::Value, "Unable to compare TFieldPath with raw pointer - types are incompatible");
 
-		return FFieldPath::operator!=(Other);
+		return Get() == Other;
+	}
+
+	FORCENOINLINE bool operator==(TYPE_OF_NULLPTR) const
+	{
+		return !Get();
 	}
 
+#if !PLATFORM_COMPILER_HAS_GENERATED_COMPARISON_OPERATORS
+
 	/**
-	* Compare weak pointers for equality
-	* @param Other pointer to compare to
+	* Compare weak pointers for inequality
+	* @param Other weak pointer to compare to
 	**/
 	template <typename OtherPropertyType>
-	FORCEINLINE bool operator==(const OtherPropertyType* Other) const
+	FORCEINLINE bool operator!=(const TFieldPath<OtherPropertyType> &Other) const
 	{
 		static_assert(TPointerIsConvertibleFromTo<OtherPropertyType, const FField>::Value, "TFieldPath can only be compared with FField types");
 		static_assert(TPointerIsConvertibleFromTo<PropertyType, OtherPropertyType>::Value, "Unable to compare TFieldPath with raw pointer - types are incompatible");
 
-		return Get() == Other;
+		return FFieldPath::operator!=(Other);
 	}
 
 	/**
@@ -476,7 +483,6 @@ public:
 		return Get() != Other;
 	}
 
-
 	template <typename LhsT>
 	friend FORCENOINLINE bool operator==(const LhsT* Lhs, const TFieldPath<PropertyType>& Rhs)
 	{
@@ -488,11 +494,6 @@ public:
 		return Rhs == Lhs;
 	}
 
-	friend FORCENOINLINE bool operator==(const TFieldPath<PropertyType>& Lhs, TYPE_OF_NULLPTR)
-	{
-		return !Lhs.Get();
-	}
-
 	friend FORCENOINLINE bool operator==(TYPE_OF_NULLPTR, const TFieldPath<PropertyType>& Rhs)
 	{
 		return !Rhs.Get();
@@ -509,15 +510,16 @@ public:
 		return Rhs != Lhs;
 	}
 
-	friend FORCENOINLINE bool operator!=(const TFieldPath<PropertyType>& Lhs, TYPE_OF_NULLPTR)
+	FORCENOINLINE bool operator!=(TYPE_OF_NULLPTR) const
 	{
-		return !!Lhs.Get();
+		return !!Get();
 	}
 
 	friend FORCENOINLINE bool operator!=(TYPE_OF_NULLPTR, const TFieldPath<PropertyType>& Rhs)
 	{
 		return !!Rhs.Get();
 	}
+#endif
 };
 
 // Helper function which deduces the type of the initializer
diff --git a/Engine/Source/Runtime/CoreUObject/Public/UObject/LazyObjectPtr.h b/Engine/Source/Runtime/CoreUObject/Public/UObject/LazyObjectPtr.h
--- a/Engine/Source/Runtime/CoreUObject/Public/UObject/LazyObjectPtr.h
+++ b/Engine/Source/Runtime/CoreUObject/Public/UObject/LazyObjectPtr.h
@@ -312,73 +312,95 @@ public:
 	}
 
 	/** Hash function. */
-	FORCEINLINE friend uint32 GetTypeHash(const TLazyObjectPtr<T>& LazyObjectPtr)
+	FORCEINLINE uint32 GetLazyObjecPtrTypeHash() const
 	{
-		return GetTypeHash(static_cast<const FLazyObjectPtr&>(LazyObjectPtr));
+		return GetTypeHash(static_cast<const FLazyObjectPtr&>(*this));
 	}
 
-	friend FArchive& operator<<(FArchive& Ar, TLazyObjectPtr<T>& LazyObjectPtr)
+	FORCEINLINE void SerializePtr(FArchive& Ar)
 	{
-		Ar << static_cast<FLazyObjectPtr&>(LazyObjectPtr);
-		return Ar;
+		Ar << static_cast<FLazyObjectPtr&>(*this);
 	}
 
 	/** Compare with another TLazyObjectPtr of related type */
 	template<typename U, typename = decltype((T*)nullptr == (U*)nullptr)>
-	friend FORCEINLINE bool operator==(const TLazyObjectPtr<T>& Lhs, const TLazyObjectPtr<U>& Rhs)
+	FORCEINLINE bool operator==(const TLazyObjectPtr<U>& Rhs) const
 	{
-		return (const FLazyObjectPtr&)Lhs == (const FLazyObjectPtr&)Rhs;
+		return (const FLazyObjectPtr&)*this == (const FLazyObjectPtr&)Rhs;
 	}
 	template<typename U, typename = decltype((T*)nullptr != (U*)nullptr)>
-	friend FORCEINLINE bool operator!=(const TLazyObjectPtr<T>& Lhs, const TLazyObjectPtr<U>& Rhs)
+	FORCEINLINE bool operator!=(const TLazyObjectPtr<U>& Rhs) const
 	{
-		return (const FLazyObjectPtr&)Lhs != (const FLazyObjectPtr&)Rhs;
+		return (const FLazyObjectPtr&)*this != (const FLazyObjectPtr&)Rhs;
 	}
 
 	/** Compare for equality with a raw pointer **/
 	template<typename U, typename = decltype((T*)nullptr == (U*)nullptr)>
-	friend FORCEINLINE bool operator==(const TLazyObjectPtr<T>& Lhs, const U* Rhs)
-	{
-		return Lhs.Get() == Rhs;
-	}
-	template<typename U, typename = decltype((T*)nullptr == (U*)nullptr)>
-	friend FORCEINLINE bool operator==(const U* Lhs, const TLazyObjectPtr<T>& Rhs)
+	FORCEINLINE bool operator==(const U* Rhs) const
 	{
-		return Lhs == Rhs.Get();
+		return Get() == Rhs;
 	}
 
 	/** Compare to null */
-	friend FORCEINLINE bool operator==(const TLazyObjectPtr<T>& Lhs, TYPE_OF_NULLPTR)
+	FORCEINLINE bool operator==(TYPE_OF_NULLPTR) const
 	{
-		return !Lhs.IsValid();
+		return !IsValid();
 	}
-	friend FORCEINLINE bool operator==(TYPE_OF_NULLPTR, const TLazyObjectPtr<T>& Rhs)
-	{
-		return !Rhs.IsValid();
-	}
-
 	/** Compare for inequality with a raw pointer	**/
 	template<typename U, typename = decltype((T*)nullptr != (U*)nullptr)>
-	friend FORCEINLINE bool operator!=(const TLazyObjectPtr<T>& Lhs, const U* Rhs)
+	FORCEINLINE bool operator!=(const U* Rhs) const
 	{
-		return Lhs.Get() != Rhs;
-	}
-	template<typename U, typename = decltype((T*)nullptr != (U*)nullptr)>
-	friend FORCEINLINE bool operator!=(const U* Lhs, const TLazyObjectPtr<T>& Rhs)
-	{
-		return Lhs != Rhs.Get();
+		return Get() != Rhs;
 	}
 
 	/** Compare for inequality with null **/
-	friend FORCEINLINE bool operator!=(const TLazyObjectPtr<T>& Lhs, TYPE_OF_NULLPTR)
+	FORCEINLINE bool operator!=(TYPE_OF_NULLPTR) const
 	{
-		return Lhs.IsValid();
-	}
-	friend FORCEINLINE bool operator!=(TYPE_OF_NULLPTR, const TLazyObjectPtr<T>& Rhs)
-	{
-		return Rhs.IsValid();
+		return IsValid();
 	}
 };
 
+/** Hash function. */
+template<typename T>
+FORCEINLINE uint32 GetTypeHash(const TLazyObjectPtr<T>& LazyObjectPtr)
+{
+	return LazyObjectPtr.GetLazyObjecPtrTypeHash();
+}
+
+template<typename T>
+FArchive& operator<<(FArchive& Ar, TLazyObjectPtr<T>& LazyObjectPtr)
+{
+	LazyObjectPtr.SerializePtr(Ar);
+	return Ar;
+}
+
+/** Compare for equality with a raw pointer **/
+template<typename T, typename U, typename = decltype((T*)nullptr == (U*)nullptr)>
+FORCEINLINE bool operator==(const U* Lhs, const TLazyObjectPtr<T>& Rhs)
+{
+	return Lhs == Rhs.Get();
+}
+
+/** Compare to null */
+template<typename T>
+FORCEINLINE bool operator==(TYPE_OF_NULLPTR, const TLazyObjectPtr<T>& Rhs)
+{
+	return !Rhs.IsValid();
+}
+
+/** Compare for inequality with a raw pointer	**/
+template<typename T, typename U, typename = decltype((T*)nullptr != (U*)nullptr)>
+FORCEINLINE bool operator!=(const U* Lhs, const TLazyObjectPtr<T>& Rhs)
+{
+	return Lhs != Rhs.Get();
+}
+
+/** Compare for inequality with null **/
+template<typename T>
+FORCEINLINE bool operator!=(TYPE_OF_NULLPTR, const TLazyObjectPtr<T>& Rhs)
+{
+	return Rhs.IsValid();
+}
+
 template<class T> struct TIsPODType<TLazyObjectPtr<T> > { enum { Value = TIsPODType<FLazyObjectPtr>::Value }; };
 template<class T> struct TIsWeakPointerType<TLazyObjectPtr<T> > { enum { Value = TIsWeakPointerType<FLazyObjectPtr>::Value }; };
\ No newline at end of file
diff --git a/Engine/Source/Runtime/CoreUObject/Public/UObject/ObjectHandle.h b/Engine/Source/Runtime/CoreUObject/Public/UObject/ObjectHandle.h
--- a/Engine/Source/Runtime/CoreUObject/Public/UObject/ObjectHandle.h
+++ b/Engine/Source/Runtime/CoreUObject/Public/UObject/ObjectHandle.h
@@ -120,8 +120,8 @@ struct FPackedObjectRef
 	// The least significant bit must always be 1 in a non-null reference.
 	UPTRINT EncodedRef;
 
-	friend inline bool operator==(FPackedObjectRef LHS, FPackedObjectRef RHS) { return LHS.EncodedRef == RHS.EncodedRef; }
-	friend inline bool operator!=(FPackedObjectRef LHS, FPackedObjectRef RHS) { return LHS.EncodedRef != RHS.EncodedRef; }
+	inline bool operator==(FPackedObjectRef RHS) const { return EncodedRef == RHS.EncodedRef; }
+	inline bool operator!=(FPackedObjectRef RHS) const { return EncodedRef != RHS.EncodedRef; }
 	friend inline uint32 GetTypeHash(FPackedObjectRef ObjectRef) { return GetTypeHash(ObjectRef.EncodedRef); }
 };
 
diff --git a/Engine/Source/Runtime/CoreUObject/Public/UObject/ObjectPtr.h b/Engine/Source/Runtime/CoreUObject/Public/UObject/ObjectPtr.h
--- a/Engine/Source/Runtime/CoreUObject/Public/UObject/ObjectPtr.h
+++ b/Engine/Source/Runtime/CoreUObject/Public/UObject/ObjectPtr.h
@@ -326,6 +326,8 @@ namespace ObjectPtr_Private
 
 	template <typename T>
 	char (&ResolveTypeIsComplete(...))[1];
+
+	struct Friend;
 };
 
 /**
@@ -444,15 +446,6 @@ public:
 	{
 		return ObjectPtr == Other.ObjectPtr;
 	}
-	template <
-		typename U,
-		typename Base = std::decay_t<decltype(false ? std::declval<std::decay_t<T*>>() : std::declval<std::decay_t<U*>>())>
-	>
-	FORCEINLINE bool operator!=(const TObjectPtr<U>& Other) const
-	{
-		return ObjectPtr != Other.ObjectPtr;
-	}
-
 
 	// Equality/Inequality comparisons against nullptr
 	FORCEINLINE bool operator==(TYPE_OF_NULLPTR) const
@@ -460,55 +453,40 @@ public:
 		return !ObjectPtr.operator bool();
 	}
 
-	FORCEINLINE bool operator!=(TYPE_OF_NULLPTR) const
-	{
-		return ObjectPtr.operator bool();
-	}
-
-	friend FORCEINLINE bool operator==(TYPE_OF_NULLPTR, const TObjectPtr& Rhs)
-	{
-		return Rhs == nullptr;
-	}
-
-	friend FORCEINLINE bool operator!=(TYPE_OF_NULLPTR, const TObjectPtr& Rhs)
-	{
-		return Rhs != nullptr;
-	}
-
 	// Equality/Inequality comparisons against another type that can be implicitly converted to the pointer type kept in a TObjectPtr
 	template <
 		typename U,
 		typename = decltype(ObjectPtr_Private::IsObjectPtrEqual(std::declval<const TObjectPtr<T>&>(), std::declval<U&&>()))
 	>
-	friend FORCEINLINE bool operator==(const TObjectPtr<T>& Ptr, U&& Other)
+	FORCEINLINE bool operator==(U&& Other) const
 	{
-		return ObjectPtr_Private::IsObjectPtrEqual(Ptr, Other);
+		return ObjectPtr_Private::IsObjectPtrEqual(*this, Other);
 	}
+
+#if __cplusplus < 202002L
 	template <
 		typename U,
-		typename = decltype(ObjectPtr_Private::IsObjectPtrEqual(std::declval<const TObjectPtr<T>&>(), std::declval<U&&>()))
+		typename Base = std::decay_t<decltype(false ? std::declval<std::decay_t<T*>>() : std::declval<std::decay_t<U*>>())>
 	>
-	friend FORCEINLINE bool operator==(U&& Other, const TObjectPtr<T>& Ptr)
+	FORCEINLINE bool operator!=(const TObjectPtr<U>& Other) const
 	{
-		return ObjectPtr_Private::IsObjectPtrEqual(Ptr, Other);
+		return ObjectPtr != Other.ObjectPtr;
 	}
 
-	template <
-		typename U,
-		typename = decltype(ObjectPtr_Private::IsObjectPtrEqual(std::declval<const TObjectPtr<T>&>(), std::declval<U&&>()))
-	>
-	friend FORCEINLINE bool operator!=(const TObjectPtr<T>& Ptr, U&& Other)
+	FORCEINLINE bool operator!=(TYPE_OF_NULLPTR) const
 	{
-		return !ObjectPtr_Private::IsObjectPtrEqual(Ptr, Other);
+		return ObjectPtr.operator bool();
 	}
+
 	template <
 		typename U,
 		typename = decltype(ObjectPtr_Private::IsObjectPtrEqual(std::declval<const TObjectPtr<T>&>(), std::declval<U&&>()))
 	>
-	friend FORCEINLINE bool operator!=(U&& Other, const TObjectPtr<T>& Ptr)
+	FORCEINLINE bool operator!=(U&& Other) const
 	{
-		return !ObjectPtr_Private::IsObjectPtrEqual(Ptr, Other);
+		return !ObjectPtr_Private::IsObjectPtrEqual(*this, Other);
 	}
+#endif
 
 	// @TODO: OBJPTR: There is a risk that the FObjectPtr is storing a reference to the wrong type.  This could
 	//			happen if data was serialized at a time when a pointer field was declared to be of type A, but then the declaration
@@ -548,22 +526,17 @@ public:
 	FORCEINLINE bool IsA(const UClass* SomeBase) const { return ObjectPtr.IsA(SomeBase); }
 	template <typename U> FORCEINLINE bool IsA() const { return ObjectPtr.IsA<U>(); }
 
-	friend FORCEINLINE uint32 GetTypeHash(const TObjectPtr<T>& InObjectPtr)
+	FORCEINLINE uint32 GetPtrTypeHash() const
 	{
-		return GetTypeHash(InObjectPtr.ObjectPtr);
+		return GetTypeHash(ObjectPtr);
 	}
 
-	friend FORCEINLINE FArchive& operator<<(FArchive& Ar, TObjectPtr<T>& InObjectPtr)
+	FORCEINLINE void SerializePtrStructured(FStructuredArchiveSlot Slot)
 	{
-		Ar << InObjectPtr.ObjectPtr;
-		return Ar;
-	}
-
-	friend FORCEINLINE void operator<<(FStructuredArchiveSlot Slot, TObjectPtr<T>& InObjectPtr)
-	{
-		Slot << InObjectPtr.ObjectPtr;
+		Slot << ObjectPtr;
 	}
 
+	friend ObjectPtr_Private::Friend;
 	friend struct FObjectPtr;
 	template <typename U> friend struct TObjectPtr;
 	template <typename U, typename V> friend bool ObjectPtr_Private::IsObjectPtrEqualToRawPtrOfRelatedType(const TObjectPtr<U>& Ptr, const V* Other);
@@ -615,8 +588,69 @@ namespace ObjectPtr_Private
 	template <typename T> struct TRawPointerType<const          TObjectPtr<T>> { using Type = T*; };
 	template <typename T> struct TRawPointerType<      volatile TObjectPtr<T>> { using Type = T*; };
 	template <typename T> struct TRawPointerType<const volatile TObjectPtr<T>> { using Type = T*; };
+	struct Friend
+	{
+		template <typename T>
+		FORCEINLINE static uint32 GetPtrTypeHash(const TObjectPtr<T>& InObjectPtr)
+		{
+			return GetTypeHash(InObjectPtr.ObjectPtr);
+		}
+
+		template <typename T>
+		FORCEINLINE static FArchive& Serialize(FArchive& Ar, TObjectPtr<T>& InObjectPtr)
+		{
+			Ar << InObjectPtr.ObjectPtr;
+			return Ar;
+		}
+
+		template <typename T>
+		FORCEINLINE static void SerializePtrStructured(FStructuredArchiveSlot Slot, TObjectPtr<T>& InObjectPtr)
+		{
+			Slot << InObjectPtr.ObjectPtr;
+		}
+	};
+}
+
+template <typename T>
+FORCEINLINE uint32 GetTypeHash(const TObjectPtr<T>& InObjectPtr)
+{
+	return ObjectPtr_Private::Friend::GetPtrTypeHash(InObjectPtr);
+}
+
+template <typename T>
+FORCEINLINE FArchive& operator<<(FArchive& Ar, TObjectPtr<T>& InObjectPtr)
+{
+	return ObjectPtr_Private::Friend::Serialize(Ar, InObjectPtr);
 }
 
+template <typename T>
+FORCEINLINE void operator<<(FStructuredArchiveSlot Slot, TObjectPtr<T>& InObjectPtr)
+{
+	ObjectPtr_Private::Friend::SerializePtrStructured(Slot, InObjectPtr);
+}
+
+#if !PLATFORM_COMPILER_HAS_GENERATED_COMPARISON_OPERATORS
+// Equality/Inequality comparisons against another type that can be implicitly converted to the pointer type kept in a TObjectPtr
+template <
+	typename T,
+	typename U,
+	decltype(ObjectPtr_Private::IsObjectPtrEqual(std::declval<const TObjectPtr<T>&>(), std::declval<U&&>()))* = nullptr
+>
+FORCEINLINE bool operator==(U&& Other, const TObjectPtr<T>& Ptr)
+{
+	return ObjectPtr_Private::IsObjectPtrEqual(Ptr, Other);
+}
+template <
+	typename T,
+	typename U,
+	decltype(ObjectPtr_Private::IsObjectPtrEqual(std::declval<const TObjectPtr<T>&>(), std::declval<U&&>()))* = nullptr
+>
+FORCEINLINE bool operator!=(U&& Other, const TObjectPtr<T>& Ptr)
+{
+	return !ObjectPtr_Private::IsObjectPtrEqual(Ptr, Other);
+}
+#endif
+
 template <typename T>
 TPrivateObjectPtr<T> MakeObjectPtrUnsafe(const UObject* Obj);
 
diff --git a/Engine/Source/Runtime/CoreUObject/Public/UObject/SoftObjectPtr.h b/Engine/Source/Runtime/CoreUObject/Public/UObject/SoftObjectPtr.h
--- a/Engine/Source/Runtime/CoreUObject/Public/UObject/SoftObjectPtr.h
+++ b/Engine/Source/Runtime/CoreUObject/Public/UObject/SoftObjectPtr.h
@@ -223,9 +223,9 @@ public:
 	 *
 	 * @param Other soft pointer to compare to
 	 */
-	FORCEINLINE friend bool operator==(const TSoftObjectPtr& Lhs, const TSoftObjectPtr& Rhs)
+	FORCEINLINE bool operator==(const TSoftObjectPtr& Rhs) const
 	{
-		return Lhs.SoftObjectPtr == Rhs.SoftObjectPtr;
+		return SoftObjectPtr == Rhs.SoftObjectPtr;
 	}
 
 	/**
@@ -234,42 +234,21 @@ public:
 	 *
 	 * @param Other soft pointer to compare to
 	 */
-	FORCEINLINE friend bool operator==(const TSoftObjectPtr& Lhs, TYPE_OF_NULLPTR)
+	FORCEINLINE bool operator==(TYPE_OF_NULLPTR) const
 	{
-		return Lhs.SoftObjectPtr == nullptr;
-	}
-
-	/**
-	 * Compare soft pointers for equality
-	 * Caution: Two soft pointers might not be equal to each other, but they both might return nullptr
-	 *
-	 * @param Other soft pointer to compare to
-	 */
-	FORCEINLINE friend bool operator==(TYPE_OF_NULLPTR, const TSoftObjectPtr& Rhs)
-	{
-		return nullptr == Rhs.SoftObjectPtr;
-	}
-
-	/**
-	 * Compare soft pointers for inequality
-	 * Caution: Two soft pointers might not be equal to each other, but they both might return nullptr
-	 *
-	 * @param Other soft pointer to compare to
-	 */
-	FORCEINLINE friend bool operator!=(const TSoftObjectPtr& Lhs, const TSoftObjectPtr& Rhs)
-	{
-		return Lhs.SoftObjectPtr != Rhs.SoftObjectPtr;
+		return SoftObjectPtr == nullptr;
 	}
 
+#if !PLATFORM_COMPILER_HAS_GENERATED_COMPARISON_OPERATORS
 	/**
 	 * Compare soft pointers for inequality
 	 * Caution: Two soft pointers might not be equal to each other, but they both might return nullptr
 	 *
 	 * @param Other soft pointer to compare to
 	 */
-	FORCEINLINE friend bool operator!=(const TSoftObjectPtr& Lhs, TYPE_OF_NULLPTR)
+	FORCEINLINE bool operator!=(const TSoftObjectPtr& Rhs) const
 	{
-		return Lhs.SoftObjectPtr != nullptr;
+		return SoftObjectPtr != Rhs.SoftObjectPtr;
 	}
 
 	/**
@@ -278,10 +257,11 @@ public:
 	 *
 	 * @param Other soft pointer to compare to
 	 */
-	FORCEINLINE friend bool operator!=(TYPE_OF_NULLPTR, const TSoftObjectPtr& Rhs)
+	FORCEINLINE bool operator!=(TYPE_OF_NULLPTR) const
 	{
-		return nullptr != Rhs.SoftObjectPtr;
+		return SoftObjectPtr != nullptr;
 	}
+#endif
 
 	/**
 	 * Dereference the soft pointer.
@@ -377,21 +357,61 @@ public:
 	}
 
 	/** Hash function */
-	FORCEINLINE friend uint32 GetTypeHash(const TSoftObjectPtr<T>& Other)
+	FORCEINLINE uint32 GetPtrTypeHash() const
 	{
-		return GetTypeHash(static_cast<const TPersistentObjectPtr<FSoftObjectPath>&>(Other.SoftObjectPtr));
+		return GetTypeHash(static_cast<const TPersistentObjectPtr<FSoftObjectPath>&>(SoftObjectPtr));
 	}
 
-	friend FArchive& operator<<(FArchive& Ar, TSoftObjectPtr<T>& Other)
+	FORCEINLINE void Serialize(FArchive& Ar)
 	{
-		Ar << Other.SoftObjectPtr;
-		return Ar;
+		Ar << SoftObjectPtr;
 	}
 
 private:
 	FSoftObjectPtr SoftObjectPtr;
 };
 
+#if !PLATFORM_COMPILER_HAS_GENERATED_COMPARISON_OPERATORS
+template<class T>
+FORCEINLINE bool operator==(TYPE_OF_NULLPTR, const TSoftObjectPtr<T>& Rhs)
+{
+	return Rhs == nullptr;
+}
+
+template<class T>
+FORCEINLINE bool operator!=(TYPE_OF_NULLPTR, const TSoftObjectPtr<T>& Rhs)
+{
+	return Rhs != nullptr;
+}
+
+template<class T>
+FORCEINLINE bool operator==(const T* Lhs, const TSoftObjectPtr<T>& Rhs)
+{
+	return Rhs == Lhs;
+}
+
+template<class T>
+FORCEINLINE bool operator!=(const T* Lhs, const TSoftObjectPtr<T>& Rhs)
+{
+	return Rhs != Lhs;
+}
+#endif
+
+/** Hash function */
+template<class T>
+FORCEINLINE uint32 GetTypeHash(const TSoftObjectPtr<T>& Ptr)
+{
+	return Ptr.GetPtrTypeHash();
+}
+
+template<class T>
+FArchive& operator<<(FArchive& Ar, TSoftObjectPtr<T>& Ptr)
+{
+	Ptr.Serialize(Ar);
+	return Ar;
+}
+
+
 template<class T> struct TIsPODType<TSoftObjectPtr<T> > { enum { Value = TIsPODType<FSoftObjectPtr>::Value }; };
 template<class T> struct TIsWeakPointerType<TSoftObjectPtr<T> > { enum { Value = TIsWeakPointerType<FSoftObjectPtr>::Value }; };
 
@@ -475,21 +495,22 @@ public:
 	 *
 	 * @param Other soft pointer to compare to 
 	 */
-	FORCEINLINE friend bool operator==(const TSoftClassPtr& Lhs, const TSoftClassPtr& Rhs)
+	FORCEINLINE bool operator==(const TSoftClassPtr& Rhs) const
 	{
-		return Lhs.SoftObjectPtr == Rhs.SoftObjectPtr;
+		return SoftObjectPtr == Rhs.SoftObjectPtr;
 	}
-
+#if !PLATFORM_COMPILER_HAS_GENERATED_COMPARISON_OPERATORS
 	/**  
 	 * Compare soft pointers for inequality
 	 * Caution: Two soft pointers might not be equal to each other, but they both might return nullptr
 	 *
 	 * @param Other soft pointer to compare to
 	 */
-	FORCEINLINE friend bool operator!=(const TSoftClassPtr& Lhs, const TSoftClassPtr& Rhs)
+	FORCEINLINE bool operator!=(const TSoftClassPtr& Rhs) const
 	{
-		return Lhs.SoftObjectPtr != Rhs.SoftObjectPtr;
+		return SoftObjectPtr != Rhs.SoftObjectPtr;
 	}
+#endif
 
 	/**  
 	 * Dereference the soft pointer
@@ -586,9 +607,9 @@ public:
 	}
 
 	/** Hash function */
-	FORCEINLINE friend uint32 GetTypeHash(const TSoftClassPtr<TClass>& Other)
+	FORCEINLINE uint32 GetPtrTypeHash() const
 	{
-		return GetTypeHash(static_cast<const TPersistentObjectPtr<FSoftObjectPath>&>(Other.SoftObjectPtr));
+		return GetTypeHash(static_cast<const TPersistentObjectPtr<FSoftObjectPath>&>(SoftObjectPtr));
 	}
 
 	/** Synchronously load (if necessary) and return the asset object represented by this asset ptr */
@@ -603,10 +624,9 @@ public:
 		return Class;
 	}
 
-	friend FArchive& operator<<(FArchive& Ar, TSoftClassPtr<TClass>& Other)
+	inline void Serialize(FArchive& Ar)
 	{
-		Ar << static_cast<FSoftObjectPtr&>(Other.SoftObjectPtr);
-		return Ar;
+		Ar << static_cast<FSoftObjectPtr&>(SoftObjectPtr);
 	}
 
 private:
@@ -648,3 +668,31 @@ T* TSoftObjectPtr<T>::Get() const
 {
 	return dynamic_cast<T*>(SoftObjectPtr.Get());
 }
+
+#if !PLATFORM_COMPILER_HAS_GENERATED_COMPARISON_OPERATORS
+template<class TClass>
+FORCEINLINE bool operator==(const UClass* Lhs, const TSoftClassPtr<TClass>& Rhs)
+{
+	return Rhs == Lhs;
+}
+
+template<class TClass>
+FORCEINLINE bool operator!=(const UClass* Lhs, const TSoftClassPtr<TClass>& Rhs)
+{
+	return Rhs != Lhs;
+}
+#endif
+
+/** Hash function */
+template<class TClass>
+FORCEINLINE uint32 GetTypeHash(const TSoftClassPtr<TClass>& Ptr)
+{
+	return Ptr.GetPtrTypeHash();
+}
+
+template<class TClass>
+FArchive& operator<<(FArchive& Ar, TSoftClassPtr<TClass>& Ptr)
+{
+	Ptr.Serialize(Ar);
+	return Ar;
+}
diff --git a/Engine/Source/Runtime/CoreUObject/Public/UObject/StructOnScope.h b/Engine/Source/Runtime/CoreUObject/Public/UObject/StructOnScope.h
--- a/Engine/Source/Runtime/CoreUObject/Public/UObject/StructOnScope.h
+++ b/Engine/Source/Runtime/CoreUObject/Public/UObject/StructOnScope.h
@@ -309,7 +309,7 @@ public:
 		return const_cast<TStructOnScope*>(this)->CastChecked<U>();
 	}
 
-	friend FArchive& operator<<(FArchive& Ar, TStructOnScope& InStruct)
+	void Serialize(FArchive& Ar)
 	{
 		if (Ar.IsLoading())
 		{
@@ -321,30 +321,28 @@ public:
 				if (ScriptStructPtr == nullptr || !ScriptStructPtr->IsChildOf(TBaseStructure<T>::Get()))
 				{
 					Ar.SetError();
-					return Ar;
+					return;
 				}
-				InStruct.ScriptStruct = ScriptStructPtr;
-				InStruct.Initialize();
-				ScriptStructPtr->SerializeItem(Ar, InStruct.SampleStructMemory, nullptr);
+				ScriptStruct = ScriptStructPtr;
+				Initialize();
+				ScriptStructPtr->SerializeItem(Ar, SampleStructMemory, nullptr);
 			}
 		}
 		// Saving
 		else
 		{
 			FString StructPath;
-			if (UScriptStruct* ScriptStructPtr = const_cast<UScriptStruct*>(::Cast<UScriptStruct>(InStruct.ScriptStruct.Get())))
+			if (UScriptStruct* ScriptStructPtr = const_cast<UScriptStruct*>(::Cast<UScriptStruct>(ScriptStruct.Get())))
 			{
 				StructPath = ScriptStructPtr->GetPathName();
 				Ar << StructPath;
-				ScriptStructPtr->SerializeItem(Ar, InStruct.SampleStructMemory, nullptr);
+				ScriptStructPtr->SerializeItem(Ar, SampleStructMemory, nullptr);
 			}
 			else
 			{
 				Ar << StructPath;
 			}
 		}
-		
-		return Ar;
 	}
 
 private:
@@ -370,3 +368,10 @@ FORCEINLINE TStructOnScope<T> MakeStructOnScope(TArgs&&... Args)
 	Struct.template InitializeAs<U>(Forward<TArgs>(Args)...);
 	return Struct;
 }
+
+template<typename T>
+FArchive& operator<<(FArchive& Ar, TStructOnScope<T>& Struct)
+{
+	Struct.Serialize(Ar);
+	return Ar;
+}
diff --git a/Engine/Source/Runtime/CoreUObject/Public/UObject/UObjectArray.h b/Engine/Source/Runtime/CoreUObject/Public/UObject/UObjectArray.h
--- a/Engine/Source/Runtime/CoreUObject/Public/UObject/UObjectArray.h
+++ b/Engine/Source/Runtime/CoreUObject/Public/UObject/UObjectArray.h
@@ -1006,8 +1006,8 @@ public:
 			Advance();
 		}
 
-		friend bool operator==(const TIterator& Lhs, const TIterator& Rhs) { return Lhs.Index == Rhs.Index; }
-		friend bool operator!=(const TIterator& Lhs, const TIterator& Rhs) { return Lhs.Index != Rhs.Index; }
+		bool operator==(const TIterator& Rhs) const { return Index == Rhs.Index; }
+		bool operator!=(const TIterator& Rhs) const { return Index != Rhs.Index; }
 
 		/** Conversion to "bool" returning true if the iterator is valid. */
 		FORCEINLINE explicit operator bool() const
diff --git a/Engine/Source/Runtime/CoreUObject/Public/UObject/UObjectIterator.h b/Engine/Source/Runtime/CoreUObject/Public/UObject/UObjectIterator.h
--- a/Engine/Source/Runtime/CoreUObject/Public/UObject/UObjectIterator.h
+++ b/Engine/Source/Runtime/CoreUObject/Public/UObject/UObjectIterator.h
@@ -315,8 +315,8 @@ public:
 		return (T*)GetObject();
 	}
 
-	FORCEINLINE friend bool operator==(const TObjectIterator& Lhs, const TObjectIterator& Rhs) { return Lhs.Index == Rhs.Index; }
-	FORCEINLINE friend bool operator!=(const TObjectIterator& Lhs, const TObjectIterator& Rhs) { return Lhs.Index != Rhs.Index; }
+	FORCEINLINE bool operator==(const TObjectIterator& Rhs) const { return Index == Rhs.Index; }
+	FORCEINLINE bool operator!=(const TObjectIterator& Rhs) const { return Index != Rhs.Index; }
 
 protected:
 	/**
diff --git a/Engine/Source/Runtime/CoreUObject/Public/UObject/UnrealType.h b/Engine/Source/Runtime/CoreUObject/Public/UObject/UnrealType.h
--- a/Engine/Source/Runtime/CoreUObject/Public/UObject/UnrealType.h
+++ b/Engine/Source/Runtime/CoreUObject/Public/UObject/UnrealType.h
@@ -6638,8 +6638,8 @@ public:
 		return !(bool)*this;
 	}
 
-	inline friend bool operator==(const TFieldIterator<T>& Lhs, const TFieldIterator<T>& Rhs) { return Lhs.Field == Rhs.Field; }
-	inline friend bool operator!=(const TFieldIterator<T>& Lhs, const TFieldIterator<T>& Rhs) { return Lhs.Field != Rhs.Field; }
+	inline bool operator==(const TFieldIterator<T>& Rhs) const { return Field == Rhs.Field; }
+	inline bool operator!=(const TFieldIterator<T>& Rhs) const { return Field != Rhs.Field; }
 
 	inline void operator++()
 	{
@@ -6979,14 +6979,14 @@ public:
 		return false;
 	}
 
-	FORCEINLINE friend bool operator==(const FPropertyValueIterator& Lhs, const FPropertyValueIterator& Rhs) 
+	FORCEINLINE bool operator==(const FPropertyValueIterator& Rhs) const
 	{
-		return Lhs.PropertyIteratorStack == Rhs.PropertyIteratorStack;
+		return PropertyIteratorStack == Rhs.PropertyIteratorStack;
 	}
 	
-	FORCEINLINE friend bool operator!=(const FPropertyValueIterator& Lhs, const FPropertyValueIterator& Rhs)
+	FORCEINLINE bool operator!=(const FPropertyValueIterator& Rhs) const
 	{
-		return !(Lhs.PropertyIteratorStack == Rhs.PropertyIteratorStack);
+		return !(PropertyIteratorStack == Rhs.PropertyIteratorStack);
 	}
 
 	/** Returns a TPair containing Property/Value currently being iterated */
@@ -7086,9 +7086,9 @@ private:
 			: Owner(InValue)
 		{}
 
-		FORCEINLINE friend bool operator==(const FPropertyValueStackEntry& Lhs, const FPropertyValueStackEntry& Rhs)
+		FORCEINLINE bool operator==(const FPropertyValueStackEntry& Rhs) const
 		{
-			return Lhs.Owner == Rhs.Owner && Lhs.ValueIndex == Rhs.ValueIndex;
+			return Owner == Rhs.Owner && ValueIndex == Rhs.ValueIndex;
 		}
 
 		FORCEINLINE const BasePairType& GetPropertyValue() const
diff --git a/Engine/Source/Runtime/CoreUObject/Public/UObject/WeakFieldPtr.h b/Engine/Source/Runtime/CoreUObject/Public/UObject/WeakFieldPtr.h
--- a/Engine/Source/Runtime/CoreUObject/Public/UObject/WeakFieldPtr.h
+++ b/Engine/Source/Runtime/CoreUObject/Public/UObject/WeakFieldPtr.h
@@ -212,11 +212,10 @@ public:
 		return GetTypeHash(WeakObjectPtr.Field);
 	}
 
-	friend FArchive& operator<<(FArchive& Ar, TWeakFieldPtr& InWeakFieldPtr)
+	FORCEINLINE void Serialize(FArchive& Ar)
 	{
-		Ar << InWeakFieldPtr.Owner;
-		Ar << InWeakFieldPtr.Field;
-		return Ar;
+		Ar << Owner;
+		Ar << Field;
 	}
 
 	/**
@@ -347,4 +346,11 @@ struct TWeakFieldPtrMapKeyFuncs : public TDefaultMapKeyFuncs<KeyType, ValueType,
 	{
 		return GetTypeHash(Key);
 	}
-};
\ No newline at end of file
+};
+
+template<class T>
+FArchive& operator<<(FArchive& Ar, TWeakFieldPtr<T>& WeakFieldPtr)
+{
+	WeakFieldPtr.Serialize(Ar);
+	return Ar;
+}
diff --git a/Engine/Source/Runtime/CoreUObject/Public/UObject/WeakObjectPtr.h b/Engine/Source/Runtime/CoreUObject/Public/UObject/WeakObjectPtr.h
--- a/Engine/Source/Runtime/CoreUObject/Public/UObject/WeakObjectPtr.h
+++ b/Engine/Source/Runtime/CoreUObject/Public/UObject/WeakObjectPtr.h
@@ -129,6 +129,7 @@ public:
 			(!IsValid() && !Other.IsValid());
 	}
 
+#if !PLATFORM_COMPILER_HAS_GENERATED_COMPARISON_OPERATORS
 	/**  
 	 * Compare weak pointers for inequality
 	 * @param Other weak pointer to compare to
@@ -139,6 +140,7 @@ public:
 			(ObjectIndex != Other.ObjectIndex || ObjectSerialNumber != Other.ObjectSerialNumber) &&
 			(IsValid() || Other.IsValid());
 	}
+#endif
 
 	/**
 	 * Returns true if two weak pointers were originally set to the same object, even if they are now stale
@@ -215,9 +217,9 @@ public:
 	}
 
 	/** Hash function. */
-	friend uint32 GetTypeHash(const FWeakObjectPtr& WeakObjectPtr)
+	FORCEINLINE uint32 GetTypeHash() const
 	{
-		return uint32(WeakObjectPtr.ObjectIndex ^ WeakObjectPtr.ObjectSerialNumber);
+		return uint32(ObjectIndex ^ ObjectSerialNumber);
 	}
 
 	/**
@@ -332,3 +334,10 @@ template<> struct TIsWeakPointerType<FWeakObjectPtr> { enum { Value = true }; };
 // Typedef script delegates for convenience.
 typedef TScriptDelegate<> FScriptDelegate;
 typedef TMulticastScriptDelegate<> FMulticastScriptDelegate;
+
+/** Hash function. */
+FORCEINLINE uint32 GetTypeHash(const FWeakObjectPtr& WeakObjectPtr)
+{
+	return WeakObjectPtr.GetTypeHash();
+}
+
diff --git a/Engine/Source/Runtime/D3D12RHI/Public/D3D12ShaderResources.h b/Engine/Source/Runtime/D3D12RHI/Public/D3D12ShaderResources.h
--- a/Engine/Source/Runtime/D3D12RHI/Public/D3D12ShaderResources.h
+++ b/Engine/Source/Runtime/D3D12RHI/Public/D3D12ShaderResources.h
@@ -23,12 +23,13 @@ struct FD3D12ShaderResourceTable : public FBaseShaderResourceTable
 		const FBaseShaderResourceTable& BaseB = B;
 		return BaseA == BaseB && (FMemory::Memcmp(A.TextureMap.GetData(), B.TextureMap.GetData(), A.TextureMap.GetTypeSize()*A.TextureMap.Num()) == 0);
 	}
-};
 
-inline FArchive& operator<<(FArchive& Ar, FD3D12ShaderResourceTable& SRT)
-{
-	FBaseShaderResourceTable& BaseSRT = SRT;
-	Ar << BaseSRT;
-	Ar << SRT.TextureMap;
-	return Ar;
-}
+	friend inline FArchive& operator<<(FArchive& Ar, FD3D12ShaderResourceTable& SRT)
+	{
+		FBaseShaderResourceTable& BaseSRT = SRT;
+		Ar << BaseSRT;
+		Ar << SRT.TextureMap;
+		return Ar;
+	}
+
+};
