- Added GetRefreshLowLevelAllocatorDelegate so low level allocator could readjust their internal pool / cache in runtime if needed
- LargePageAllocator does all its preallocation handling in its Refresh method (invoked via delegate) thanks to cvars. If preallocated pools are larger than the new requested size, the pool is shrinked (and physical memory freed), if it is smaller then the pool is grown (and more memory will be preallocated)
- First preallocation is done at PreInit once all the INI files have been read on platforms that use very large page allocator
- If a OOM happens, the preallocation code is then be disabled until the title is restarted
- With this system, pools can be tweaked per device profile
- Preallocated pools are tagged as FMalloc in LLM

diff --git a/Engine/Source/Runtime/Core/Private/HAL/Allocators/CachedOSVeryLargePageAllocator.cpp b/Engine/Source/Runtime/Core/Private/HAL/Allocators/CachedOSVeryLargePageAllocator.cpp
--- a/Engine/Source/Runtime/Core/Private/HAL/Allocators/CachedOSVeryLargePageAllocator.cpp
+++ b/Engine/Source/Runtime/Core/Private/HAL/Allocators/CachedOSVeryLargePageAllocator.cpp
@@ -11,48 +11,37 @@
 
 #if UE_USE_VERYLARGEPAGEALLOCATOR
 
-// Free pages that are cached can be preallocated at the initialization of the VeryLargePageAllocator
-// Having them preallocated and kept during the life of the title help to avoid defragmentation of physical memory
-#ifndef UE_VERYLARGEPAGEALLOCATOR_PREALLOCATE_EMPTY_BACKSTORE
-	#define UE_VERYLARGEPAGEALLOCATOR_PREALLOCATE_EMPTY_BACKSTORE 0
-#endif
-
-// Number of free pages (2MB each) to be kept (not decommited) in cache so they can be reused:
-// Number of free pages to cache for SmallPool allocation hint. 
-#ifndef UE_VERYLARGEPAGEALLOCATOR_MAX_EMPTY_BACKSTORE_SMALLPOOL
-	#define UE_VERYLARGEPAGEALLOCATOR_MAX_EMPTY_BACKSTORE_SMALLPOOL 10
-#endif
-
-// Number of free pages to cache for Default allocation hint.
-#ifndef UE_VERYLARGEPAGEALLOCATOR_MAX_EMPTY_BACKSTORE_DEFAULT
-	#define UE_VERYLARGEPAGEALLOCATOR_MAX_EMPTY_BACKSTORE_DEFAULT 10
-#endif
-
 CORE_API bool GEnableVeryLargePageAllocator = true;
 static int32 GVeryLargePageAllocatorMaxEmptyBackStoreCount[FMemory::AllocationHints::Max] = {
-	UE_VERYLARGEPAGEALLOCATOR_MAX_EMPTY_BACKSTORE_DEFAULT,	// FMemory::AllocationHints::Default
-	0,														// FMemory::AllocationHints::Temporary
-	UE_VERYLARGEPAGEALLOCATOR_MAX_EMPTY_BACKSTORE_SMALLPOOL // FMemory::AllocationHints::SmallPool
+	0,	// FMemory::AllocationHints::Default
+	0,	// FMemory::AllocationHints::Temporary
+	0	// FMemory::AllocationHints::SmallPool
 };
 
 static_assert(int32(FMemory::AllocationHints::Max) == 3); // ensure FMemory::AllocationHints has 3 types of hint so GVeryLargePageAllocatorMaxEmptyBackStoreCount has needed values
 
+bool GVeryLargePageAllocatorPreAllocatePools = true;
+static FAutoConsoleVariableRef CVarVeryLargePageAllocatorPreAllocatePools(
+	TEXT("VeryLargePageAllocator.PreAllocatePools"),
+	GVeryLargePageAllocatorPreAllocatePools,
+	TEXT("Having pages preallocated and cached during the life of the title help to avoid defragmentation of physical memory.\n")
+	TEXT("Preallocation is disabled when system reaches OOM."));
+
 static FAutoConsoleVariableRef CVarVeryLargePageAllocatorMaxEmptyBackstoreDefault(
 	TEXT("VeryLargePageAllocator.MaxEmptyBackstoreDefault"),
 	GVeryLargePageAllocatorMaxEmptyBackStoreCount[FMemory::AllocationHints::Default],
-	TEXT(""));
+	TEXT("Number of free pages (2MB each) to cache (not decommited) for allocation hint DEFAULT"));
 
 static FAutoConsoleVariableRef CVarVeryLargePageAllocatorMaxEmptyBackstoreSmallPool(
 	TEXT("VeryLargePageAllocator.MaxEmptyBackstoreSmallPool"),
 	GVeryLargePageAllocatorMaxEmptyBackStoreCount[FMemory::AllocationHints::SmallPool],
-	TEXT(""));
+	TEXT("Number of free pages (2MB each) to cache (not decommited) for allocation hint SMALL POOL"));
 
 #if CSV_PROFILER
 CSV_DECLARE_CATEGORY_MODULE_EXTERN(CORE_API, FMemory);
 
 static volatile int32 GLargePageAllocatorCommitCount = 0;
 static volatile int32 GLargePageAllocatorDecommitCount = 0;
-static volatile int32 GLargePageAllocatorLargePageCount[FMemory::AllocationHints::Max];
 #endif
 
 void FCachedOSVeryLargePageAllocator::Init()
@@ -74,10 +63,7 @@ void FCachedOSVeryLargePageAllocator::Init()
 		UsedLargePagesHead[i] = nullptr;
 		EmptyButAvailableLargePagesHead[i] = nullptr;
 		EmptyBackStoreCount[i] = 0;
-		
-#if CSV_PROFILER
-		GLargePageAllocatorLargePageCount[i] = 0;
-#endif
+		CommitedLargePagesCount[i] = 0;
 	}
 
 #if UE_VERYLARGEPAGEALLOCATOR_TAKEONALL64KBALLOCATIONS
@@ -98,29 +84,59 @@ void FCachedOSVeryLargePageAllocator::Init()
 	}
 #endif
 
-#if UE_VERYLARGEPAGEALLOCATOR_PREALLOCATE_EMPTY_BACKSTORE
+	if (!GEnableVeryLargePageAllocator)
+	{
+		bEnabled = false;
+	}
+}
+
+void FCachedOSVeryLargePageAllocator::Refresh()
+{
+	if (!bEnabled)
+	{
+		return;
+	}
+
+	// Function is not thread safe 
+	UE_LOG(LogMemory, Log, TEXT("Refreshing LargePageAllocator"));
+
 	for (int32 i = 0; i < FMemory::AllocationHints::Max; i++)
 	{
 		FMemory::AllocationHints AllocationHint = FMemory::AllocationHints(i);
-		
-		for (int32 BackStoreIndex = 0; BackStoreIndex < GVeryLargePageAllocatorMaxEmptyBackStoreCount[AllocationHint]; BackStoreIndex++)
+
+		int32 LargePageCount = CommitedLargePagesCount[AllocationHint] + EmptyBackStoreCount[AllocationHint];
+
+		if (LargePageCount < GVeryLargePageAllocatorMaxEmptyBackStoreCount[AllocationHint] && GVeryLargePageAllocatorPreAllocatePools)
 		{
-			FLargePage* LargePage = FreeLargePagesHead[AllocationHint];
+			// Preallocate large pages
+			for ( ; LargePageCount < GVeryLargePageAllocatorMaxEmptyBackStoreCount[AllocationHint]; LargePageCount++)
+			{
+				LLM_PLATFORM_SCOPE(ELLMTag::FMalloc);
 
-			check(LargePage != nullptr); // Can't happen
-			LargePage->AllocationHint = AllocationHint;
-			LargePage->Unlink();
-			Block.Commit(LargePage->BaseAddress - AddressSpaceReserved, SizeOfLargePage, false);
-			LargePage->LinkHead(EmptyButAvailableLargePagesHead[AllocationHint]);
-			EmptyBackStoreCount[AllocationHint] += 1;
-			CachedFree += SizeOfLargePage;
-		}
-	}
-#endif
+				FLargePage* LargePage = FreeLargePagesHead[AllocationHint];
 
-	if (!GEnableVeryLargePageAllocator)
-	{
-		bEnabled = false;
+				check(LargePage != nullptr); // Can't happen
+				LargePage->AllocationHint = AllocationHint;
+				LargePage->Unlink();
+
+				if (!Block.Commit(LargePage->BaseAddress - AddressSpaceReserved, SizeOfLargePage, false))
+				{
+					// Cant commit 2MB pages, stop preallocating and return page to FreeLargePagesHead list
+					LargePage->LinkHead(FreeLargePagesHead[AllocationHint]);
+					GVeryLargePageAllocatorPreAllocatePools = false;
+					break;
+				}
+
+				LargePage->LinkHead(EmptyButAvailableLargePagesHead[AllocationHint]);
+				CachedFree += SizeOfLargePage;
+				CommitedLargePagesCount[AllocationHint] += 1;
+				EmptyBackStoreCount[AllocationHint] += 1;
+			}
+		}
+		else
+		{
+			ShrinkEmptyBackStore(GVeryLargePageAllocatorMaxEmptyBackStoreCount[AllocationHint], AllocationHint);
+		}
 	}
 }
 
@@ -181,19 +197,21 @@ void* FCachedOSVeryLargePageAllocator::Allocate(SIZE_T Size, uint32 AllocationHi
 								// Fallback to regular allocator
 								LargePage = nullptr;
 							}
-#if CSV_PROFILER
 							else
 							{
 								// A new large page has been created. Add it to CachedFree counter
 								CachedFree += SizeOfLargePage;
+								CommitedLargePagesCount[AllocationHint] += 1;
+#if CSV_PROFILER
 								FPlatformAtomics::InterlockedIncrement(&GLargePageAllocatorCommitCount);
-								FPlatformAtomics::InterlockedIncrement(&GLargePageAllocatorLargePageCount[AllocationHint]);
-							}
 #endif
+							}
 						}
 					}
 #else 
 							Block.Commit(LargePage->BaseAddress - AddressSpaceReserved, SizeOfLargePage);
+							CachedFree += SizeOfLargePage;
+							CommitedLargePagesCount[AllocationHint] += 1;
 						}
 					}
 #endif // UE_USE_VERYLARGEPAGEALLOCATOR_FALLBACKPATH
@@ -268,9 +286,9 @@ void FCachedOSVeryLargePageAllocator::Free(void* Ptr, SIZE_T Size, FCriticalSect
 					Block.Decommit(LargePage->BaseAddress - AddressSpaceReserved, SizeOfLargePage);
 				}
 
+				CommitedLargePagesCount[LargePage->AllocationHint] -= 1;
 #if CSV_PROFILER
 				FPlatformAtomics::InterlockedIncrement(&GLargePageAllocatorDecommitCount);
-				FPlatformAtomics::InterlockedDecrement(&GLargePageAllocatorLargePageCount[LargePage->AllocationHint]);
 #endif
 
 				LargePage->LinkHead(FreeLargePagesHead[LargePage->AllocationHint]);
@@ -338,33 +356,37 @@ void FCachedOSVeryLargePageAllocator::Free(void* Ptr, SIZE_T Size, FCriticalSect
 	}
 }
 
-void FCachedOSVeryLargePageAllocator::FreeAll(FCriticalSection* Mutex)
+void FCachedOSVeryLargePageAllocator::ShrinkEmptyBackStore(int32 NewEmptyBackStoreSize, FMemory::AllocationHints AllocationHint)
 {
-	for (int i = 0; i < FMemory::AllocationHints::Max; i++)
+	while (EmptyBackStoreCount[AllocationHint] > NewEmptyBackStoreSize)
 	{
-		FMemory::AllocationHints AllocationHint = FMemory::AllocationHints(i);
-		
-		for (;;)
+		FLargePage* LargePage = EmptyButAvailableLargePagesHead[AllocationHint];
+
+		if (LargePage == nullptr)
 		{
-			FLargePage* LargePage = EmptyButAvailableLargePagesHead[AllocationHint];
+			break;
+		}
+		LargePage->Unlink();
+		Block.Decommit(LargePage->BaseAddress - AddressSpaceReserved, SizeOfLargePage);
 
-			if (LargePage == nullptr)
-			{
-				break;
-			}
-			LargePage->Unlink();
-			EmptyBackStoreCount[AllocationHint] -= 1;
-			Block.Decommit(LargePage->BaseAddress - AddressSpaceReserved, SizeOfLargePage);
+		LargePage->LinkHead(FreeLargePagesHead[LargePage->AllocationHint]);
+		CachedFree -= SizeOfLargePage;
+		EmptyBackStoreCount[AllocationHint] -= 1;
+		CommitedLargePagesCount[LargePage->AllocationHint] -= 1;
+	}
+}
 
-			LargePage->LinkHead(FreeLargePagesHead[LargePage->AllocationHint]);
-			CachedFree -= SizeOfLargePage;
 
-#if CSV_PROFILER
-			FPlatformAtomics::InterlockedDecrement(&GLargePageAllocatorLargePageCount[LargePage->AllocationHint]);
-#endif
-		}
+void FCachedOSVeryLargePageAllocator::FreeAll(FCriticalSection* Mutex)
+{
+	for (int i = 0; i < FMemory::AllocationHints::Max; i++)
+	{
+		FMemory::AllocationHints AllocationHint = FMemory::AllocationHints(i);
+		ShrinkEmptyBackStore(0, AllocationHint);
 	}
 
+	// Stop preallocating system since allocator reached a OOM
+	GVeryLargePageAllocatorPreAllocatePools = false;
 
 	// Free empty cached pages of CachedOSPageAllocator
 	CachedOSPageAllocator.FreeAll(Mutex);
@@ -377,8 +399,8 @@ void FCachedOSVeryLargePageAllocator::UpdateStats()
 	CSV_CUSTOM_STAT(FMemory, LargeAllocatorDecommitCount, GLargePageAllocatorDecommitCount, ECsvCustomStatOp::Set);
 	CSV_CUSTOM_STAT(FMemory, LargeAllocatorBackingStoreCountSmall, EmptyBackStoreCount[FMemory::AllocationHints::SmallPool], ECsvCustomStatOp::Set);
 	CSV_CUSTOM_STAT(FMemory, LargeAllocatorBackingStoreCountDefault, EmptyBackStoreCount[FMemory::AllocationHints::Default], ECsvCustomStatOp::Set);
-	CSV_CUSTOM_STAT(FMemory, LargeAllocatorPageCountSmall, GLargePageAllocatorLargePageCount[FMemory::AllocationHints::SmallPool], ECsvCustomStatOp::Set);
-	CSV_CUSTOM_STAT(FMemory, LargeAllocatorPageCountDefault, GLargePageAllocatorLargePageCount[FMemory::AllocationHints::Default], ECsvCustomStatOp::Set);
+	CSV_CUSTOM_STAT(FMemory, LargeAllocatorPageCountSmall, CommitedLargePagesCount[FMemory::AllocationHints::SmallPool], ECsvCustomStatOp::Set);
+	CSV_CUSTOM_STAT(FMemory, LargeAllocatorPageCountDefault, CommitedLargePagesCount[FMemory::AllocationHints::Default], ECsvCustomStatOp::Set);
 
 	GLargePageAllocatorCommitCount = 0;
 	GLargePageAllocatorDecommitCount = 0;
diff --git a/Engine/Source/Runtime/Core/Private/HAL/MallocBinned2.cpp b/Engine/Source/Runtime/Core/Private/HAL/MallocBinned2.cpp
--- a/Engine/Source/Runtime/Core/Private/HAL/MallocBinned2.cpp
+++ b/Engine/Source/Runtime/Core/Private/HAL/MallocBinned2.cpp
@@ -798,7 +798,7 @@ FMallocBinned2::~FMallocBinned2()
 {
 }
 
-void FMallocBinned2::OnMallocInitialized() 
+void FMallocBinned2::OnMallocInitialized()
 {
 #if UE_USE_VERYLARGEPAGEALLOCATOR
 	FCoreDelegates::GetLowLevelAllocatorMemoryTrimDelegate().AddLambda([this]()
@@ -807,6 +807,14 @@ void FMallocBinned2::OnMallocInitialized()
 			CachedOSPageAllocator.FreeAll(&Mutex);
 		}
 	);
+
+	FCoreDelegates::GetRefreshLowLevelAllocatorDelegate().AddLambda([this]()
+		{
+			FScopeLock Lock(&Mutex);
+			CachedOSPageAllocator.Refresh();
+		}
+	);
+
 #endif
 }
 
diff --git a/Engine/Source/Runtime/Core/Private/Misc/CoreDelegates.cpp b/Engine/Source/Runtime/Core/Private/Misc/CoreDelegates.cpp
--- a/Engine/Source/Runtime/Core/Private/Misc/CoreDelegates.cpp
+++ b/Engine/Source/Runtime/Core/Private/Misc/CoreDelegates.cpp
@@ -291,6 +291,12 @@ FSimpleMulticastDelegate& FCoreDelegates::GetLowLevelAllocatorMemoryTrimDelegate
 	return OnLowLevelAllocatorMemoryTrim;
 }
 
+FSimpleMulticastDelegate& FCoreDelegates::GetRefreshLowLevelAllocatorDelegate()
+{
+	static FSimpleMulticastDelegate OnRefreshLowLevelAllocator;
+	return OnRefreshLowLevelAllocator;
+}
+
 /**	 Implemented as a function to address global ctor issues */
 FSimpleMulticastDelegate& FCoreDelegates::GetOutOfMemoryDelegate()
 {
diff --git a/Engine/Source/Runtime/Core/Public/HAL/Allocators/CachedOSPageAllocator.h b/Engine/Source/Runtime/Core/Public/HAL/Allocators/CachedOSPageAllocator.h
--- a/Engine/Source/Runtime/Core/Public/HAL/Allocators/CachedOSPageAllocator.h
+++ b/Engine/Source/Runtime/Core/Public/HAL/Allocators/CachedOSPageAllocator.h
@@ -52,6 +52,11 @@ struct TCachedOSPageAllocator : private FCachedOSPageAllocator
 	void FreeAll(FCriticalSection* Mutex = nullptr)
 	{
 		return FreeAllImpl(FreedPageBlocks, FreedPageBlocksNum, CachedTotal, Mutex);
+	}
+	// Refresh cached os allocator if needed. Does nothing for this implementation
+	void Refresh()
+	{
+
 	}
 	void UpdateStats()
 	{
diff --git a/Engine/Source/Runtime/Core/Public/HAL/Allocators/CachedOSVeryLargePageAllocator.h b/Engine/Source/Runtime/Core/Public/HAL/Allocators/CachedOSVeryLargePageAllocator.h
--- a/Engine/Source/Runtime/Core/Public/HAL/Allocators/CachedOSVeryLargePageAllocator.h
+++ b/Engine/Source/Runtime/Core/Public/HAL/Allocators/CachedOSVeryLargePageAllocator.h
@@ -70,6 +70,9 @@ public:
 
 	void FreeAll(FCriticalSection* Mutex = nullptr);
 
+	// Refresh cached os allocator if needed. Will preallocate / reduce backstore if preallocation is enabled
+	void Refresh();
+
 	void UpdateStats();
 
 	uint64 GetCachedFreeTotal()
@@ -89,6 +92,7 @@ public:
 private:
 
 	void Init();
+	void ShrinkEmptyBackStore(int32 NewEmptyBackStoreSize, FMemory::AllocationHints AllocationHint);
 
 	bool bEnabled;
 	uintptr_t	AddressSpaceReserved;
@@ -96,6 +100,7 @@ private:
 	uintptr_t	AddressSpaceReservedEnd;
 	uint64		CachedFree;
 	int32		EmptyBackStoreCount[FMemory::AllocationHints::Max];
+	int32		CommitedLargePagesCount[FMemory::AllocationHints::Max];
 
 	FPlatformMemory::FPlatformVirtualMemoryBlock Block;
 
diff --git a/Engine/Source/Runtime/Core/Public/HAL/Allocators/PooledVirtualMemoryAllocator.h b/Engine/Source/Runtime/Core/Public/HAL/Allocators/PooledVirtualMemoryAllocator.h
--- a/Engine/Source/Runtime/Core/Public/HAL/Allocators/PooledVirtualMemoryAllocator.h
+++ b/Engine/Source/Runtime/Core/Public/HAL/Allocators/PooledVirtualMemoryAllocator.h
@@ -78,6 +78,9 @@ struct FPooledVirtualMemoryAllocator
 	/** Returns free memory in the pools */
 	uint64 GetCachedFreeTotal();
 
+	/** Refresh allocator if needed. (does nothing in that implementation) */
+	void Refresh() {}
+
 	/** Update memory stats (does nothing in that implementation) */
 	void UpdateStats() {}
 
diff --git a/Engine/Source/Runtime/Core/Public/Misc/CoreDelegates.h b/Engine/Source/Runtime/Core/Public/Misc/CoreDelegates.h
--- a/Engine/Source/Runtime/Core/Public/Misc/CoreDelegates.h
+++ b/Engine/Source/Runtime/Core/Public/Misc/CoreDelegates.h
@@ -464,6 +464,9 @@ public:
 	// Called to request that low level allocator free whatever memory they are able to. 
 	static FSimpleMulticastDelegate& GetLowLevelAllocatorMemoryTrimDelegate();
 
+	// Called to request that low level allocator must refreshed
+	static FSimpleMulticastDelegate& GetRefreshLowLevelAllocatorDelegate();
+
 	// Called when OOM event occurs, after backup memory has been freed, so there's some hope of being effective
 	static FSimpleMulticastDelegate& GetOutOfMemoryDelegate();
 
