Fixed CopyTemp returning a reference instead of a copy when passed an rvalue reference (can cause dangling references).

diff --git a/Engine/Source/Runtime/Analytics/AnalyticsET/Private/AnalyticsSessionSummarySender.cpp b/Engine/Source/Runtime/Analytics/AnalyticsET/Private/AnalyticsSessionSummarySender.cpp
--- a/Engine/Source/Runtime/Analytics/AnalyticsET/Private/AnalyticsSessionSummarySender.cpp
+++ b/Engine/Source/Runtime/Analytics/AnalyticsET/Private/AnalyticsSessionSummarySender.cpp
@@ -52,9 +52,9 @@ bool FAnalyticsSessionSummarySender::SendSessionSummary(const FString& UserId, c
 		// Reconfigure the analytics provider to sent the summary event 'as if' it was sent by the process that created it. This is required by the analytics backend.
 		FGuid SessionGuid;
 		FGuid::Parse(*SessionId, SessionGuid);
-		TempSummaryProvider->SetUserID(CopyTemp(*UserId));
-		TempSummaryProvider->SetAppID(CopyTemp(*AppId));
-		TempSummaryProvider->SetAppVersion(CopyTemp(*AppVersion));
+		TempSummaryProvider->SetUserID(*UserId);
+		TempSummaryProvider->SetAppID(*AppId);
+		TempSummaryProvider->SetAppVersion(*AppVersion);
 		TempSummaryProvider->SetSessionID(SessionGuid.ToString(EGuidFormats::DigitsWithHyphensInBraces)); // Ensure to put back the {} around the GUID.
 
 		// Send the summary.
diff --git a/Engine/Source/Runtime/Core/Public/Templates/UnrealTemplate.h b/Engine/Source/Runtime/Core/Public/Templates/UnrealTemplate.h
--- a/Engine/Source/Runtime/Core/Public/Templates/UnrealTemplate.h
+++ b/Engine/Source/Runtime/Core/Public/Templates/UnrealTemplate.h
@@ -15,6 +15,7 @@
 #include "Templates/TypeCompatibleBytes.h"
 #include "Templates/Identity.h"
 #include "Traits/IsContiguousContainer.h"
+#include <type_traits>
 
 /*-----------------------------------------------------------------------------
 	Standard templates.
@@ -473,7 +474,7 @@ FORCEINLINE typename TRemoveReference<T>::Type&& MoveTemp(T&& Obj)
 }
 
 /**
- * MoveTemp will cast a reference to an rvalue reference.
+ * MoveTempIfPossible will cast a reference to an rvalue reference.
  * This is UE's equivalent of std::move.  It doesn't static assert like MoveTemp, because it is useful in
  * templates or macros where it's not obvious what the argument is, but you want to take advantage of move semantics
  * where you can but not stop compilation.
@@ -486,9 +487,13 @@ FORCEINLINE typename TRemoveReference<T>::Type&& MoveTempIfPossible(T&& Obj)
 }
 
 /**
- * CopyTemp will enforce the creation of an rvalue which can bind to rvalue reference parameters.
- * Unlike MoveTemp, the source object will never be modifed. (i.e. a copy will be made)
- * There is no std:: equivalent.
+ * CopyTemp will enforce the creation of a prvalue which can bind to rvalue reference parameters.
+ * Unlike MoveTemp, a source lvalue will never be modified. (i.e. a copy will always be made)
+ * There is no std:: equivalent, though there is the exposition function std::decay-copy:
+ * https://eel.is/c++draft/expos.only.func
+ * CopyTemp(<rvalue>) is regarded as an error and will not compile, similarly to how MoveTemp(<rvalue>)
+ * does not compile, and CopyTempIfNecessary should be used instead when the nature of the
+ * argument is not known in advance.
  */
 template <typename T>
 FORCEINLINE T CopyTemp(T& Val)
@@ -503,12 +508,35 @@ FORCEINLINE T CopyTemp(const T& Val)
 }
 
 template <typename T>
+UE_DEPRECATED(5.2, "CopyTemp on an rvalue is deprecated and should be removed or replaced with CopyTempIfNecessary when the argument is unknown")
 FORCEINLINE T&& CopyTemp(T&& Val)
 {
-	// If we already have an rvalue, just return it unchanged, rather than needlessly creating yet another rvalue from it.
+// Compile this block back in after removing the deprecation, rather than deleting the function entirely.
+// Also change the return type to `T`.
+#if 0
+	// Comment this in rather than deleting the function when removing the deprecation
+	static_assert(sizeof(T) == 0, "CopyTemp called on an rvalue");
+#endif
+
+	// Create a prvalue by move-constructing from the xvalue - wasteful if Val
+	// already refers to a prvalue, but we can't differentiate those.
 	return MoveTemp(Val);
 }
 
+/**
+ * CopyTempIfNecessary will enforce the creation of a prvalue.
+ * This is UE's equivalent of the exposition std::decay-copy:
+ * https://eel.is/c++draft/expos.only.func
+ * It doesn't static assert like CopyTemp, because it is useful in
+ * templates or macros where it's not obvious what the argument is, but you want to
+ * create a PR value without stopping compilation.
+ */
+template <typename T>
+FORCEINLINE std::decay_t<T> CopyTempIfNecessary(T&& Val)
+{
+	return (T&&)Val;
+}
+
 /**
  * Forward will cast a reference to an rvalue reference.
  * This is UE's equivalent of std::forward.
