Removed legacy TAllocatorTraits::SupportsMove, which is assumed to be always be true everywhere it is used.

diff --git a/Engine/Source/Runtime/Core/Private/HAL/Allocators/AnsiAllocator.h b/Engine/Source/Runtime/Core/Private/HAL/Allocators/AnsiAllocator.h
--- a/Engine/Source/Runtime/Core/Private/HAL/Allocators/AnsiAllocator.h
+++ b/Engine/Source/Runtime/Core/Private/HAL/Allocators/AnsiAllocator.h
@@ -132,7 +132,6 @@ public:
 template <>
 struct TAllocatorTraits<FAnsiAllocator> : TAllocatorTraitsBase<FAnsiAllocator>
 {
-	enum { SupportsMove = true };
 	enum { IsZeroConstruct = true };
 };
 
diff --git a/Engine/Source/Runtime/Core/Public/Containers/Array.h b/Engine/Source/Runtime/Core/Public/Containers/Array.h
--- a/Engine/Source/Runtime/Core/Public/Containers/Array.h
+++ b/Engine/Source/Runtime/Core/Public/Containers/Array.h
@@ -3450,8 +3450,7 @@ struct TIsZeroConstructType<TArray<InElementType, AllocatorType>>
 template <typename InElementType, typename AllocatorType>
 struct TContainerTraits<TArray<InElementType, AllocatorType> > : public TContainerTraitsBase<TArray<InElementType, AllocatorType> >
 {
-	static_assert(TAllocatorTraits<AllocatorType>::SupportsMove, "TArray no longer supports move-unaware allocators");
-	enum { MoveWillEmptyContainer = TAllocatorTraits<AllocatorType>::SupportsMove };
+	enum { MoveWillEmptyContainer = true };
 };
 
 template <typename T, typename AllocatorType>
diff --git a/Engine/Source/Runtime/Core/Public/Containers/BitArray.h b/Engine/Source/Runtime/Core/Public/Containers/BitArray.h
--- a/Engine/Source/Runtime/Core/Public/Containers/BitArray.h
+++ b/Engine/Source/Runtime/Core/Public/Containers/BitArray.h
@@ -1754,8 +1754,7 @@ FORCEINLINE uint32 GetTypeHash(const TBitArray<Allocator>& BitArray)
 template<typename Allocator>
 struct TContainerTraits<TBitArray<Allocator> > : public TContainerTraitsBase<TBitArray<Allocator> >
 {
-	static_assert(TAllocatorTraits<Allocator>::SupportsMove, "TBitArray no longer supports move-unaware allocators");
-	enum { MoveWillEmptyContainer = TAllocatorTraits<Allocator>::SupportsMove };
+	enum { MoveWillEmptyContainer = true };
 };
 
 
diff --git a/Engine/Source/Runtime/Core/Public/Containers/ContainerAllocationPolicies.h b/Engine/Source/Runtime/Core/Public/Containers/ContainerAllocationPolicies.h
--- a/Engine/Source/Runtime/Core/Public/Containers/ContainerAllocationPolicies.h
+++ b/Engine/Source/Runtime/Core/Public/Containers/ContainerAllocationPolicies.h
@@ -144,7 +144,6 @@ struct FScriptContainerElement
 template <typename AllocatorType>
 struct TAllocatorTraitsBase
 {
-	enum { SupportsMove              = false };
 	enum { IsZeroConstruct           = false };
 	enum { SupportsFreezeMemoryImage = false };
 	enum { SupportsElementAlignment  = false };
@@ -453,7 +452,6 @@ public:
 template <uint32 Alignment>
 struct TAllocatorTraits<TAlignedHeapAllocator<Alignment>> : TAllocatorTraitsBase<TAlignedHeapAllocator<Alignment>>
 {
-	enum { SupportsMove    = true };
 	enum { IsZeroConstruct = true };
 };
 
@@ -620,7 +618,6 @@ public:
 template <uint8 IndexSize>
 struct TAllocatorTraits<TSizedHeapAllocator<IndexSize>> : TAllocatorTraitsBase<TSizedHeapAllocator<IndexSize>>
 {
-	enum { SupportsMove             = true };
 	enum { IsZeroConstruct          = true };
 	enum { SupportsElementAlignment = true };
 };
@@ -782,12 +779,6 @@ public:
 	typedef void ForAnyElementType;
 };
 
-template <uint32 NumInlineElements, int IndexSize, typename SecondaryAllocator>
-struct TAllocatorTraits<TSizedInlineAllocator<NumInlineElements, IndexSize, SecondaryAllocator>> : TAllocatorTraitsBase<TSizedInlineAllocator<NumInlineElements, IndexSize, SecondaryAllocator>>
-{
-	enum { SupportsMove = TAllocatorTraits<SecondaryAllocator>::SupportsMove };
-};
-
 template <uint32 NumInlineElements, typename SecondaryAllocator = FDefaultAllocator>
 using TInlineAllocator = TSizedInlineAllocator<NumInlineElements, 32, SecondaryAllocator>;
 
@@ -943,12 +934,6 @@ public:
 	typedef void ForAnyElementType;
 };
 
-template <uint32 NumInlineElements>
-struct TAllocatorTraits<TNonRelocatableInlineAllocator<NumInlineElements>> : TAllocatorTraitsBase<TNonRelocatableInlineAllocator<NumInlineElements>>
-{
-	enum { SupportsMove = true };
-};
-
 /**
  * The fixed allocation policy allocates up to a specified number of elements in the same allocation as the container.
  * It's like the inline allocator, except it doesn't provide secondary storage when the inline storage has been filled.
@@ -1048,12 +1033,6 @@ public:
 	typedef void ForAnyElementType;
 };
 
-template <uint32 NumInlineElements>
-struct TAllocatorTraits<TFixedAllocator<NumInlineElements>> : TAllocatorTraitsBase<TFixedAllocator<NumInlineElements>>
-{
-	enum { SupportsMove = true };
-};
-
 // We want these to be correctly typed as int32, but we don't want them to have linkage, so we make them macros
 #define NumBitsPerDWORD ((int32)32)
 #define NumBitsPerDWORDLogTwo ((int32)5)
diff --git a/Engine/Source/Runtime/Core/Public/Containers/Set.h b/Engine/Source/Runtime/Core/Public/Containers/Set.h
--- a/Engine/Source/Runtime/Core/Public/Containers/Set.h
+++ b/Engine/Source/Runtime/Core/Public/Containers/Set.h
@@ -1787,10 +1787,7 @@ DECLARE_TEMPLATE_INTRINSIC_TYPE_LAYOUT((template <typename ElementType, typename
 template<typename ElementType, typename KeyFuncs, typename Allocator>
 struct TContainerTraits<TSet<ElementType, KeyFuncs, Allocator> > : public TContainerTraitsBase<TSet<ElementType, KeyFuncs, Allocator> >
 {
-	static_assert(TAllocatorTraits<typename Allocator::HashAllocator>::SupportsMove, "TSet no longer supports move-unaware allocators");
-	enum { MoveWillEmptyContainer =
-		TContainerTraits<typename TSet<ElementType, KeyFuncs, Allocator>::ElementArrayType>::MoveWillEmptyContainer &&
-		TAllocatorTraits<typename Allocator::HashAllocator>::SupportsMove };
+	enum { MoveWillEmptyContainer = TContainerTraits<typename TSet<ElementType, KeyFuncs, Allocator>::ElementArrayType>::MoveWillEmptyContainer };
 };
 
 struct FScriptSetLayout
diff --git a/Engine/Source/Runtime/Core/Public/Experimental/ConcurrentLinearAllocator.h b/Engine/Source/Runtime/Core/Public/Experimental/ConcurrentLinearAllocator.h
--- a/Engine/Source/Runtime/Core/Public/Experimental/ConcurrentLinearAllocator.h
+++ b/Engine/Source/Runtime/Core/Public/Experimental/ConcurrentLinearAllocator.h
@@ -638,7 +638,6 @@ public:
 template <typename BlockAllocationTag>
 struct TAllocatorTraits<TConcurrentLinearArrayAllocator<BlockAllocationTag>> : TAllocatorTraitsBase<TConcurrentLinearArrayAllocator<BlockAllocationTag>>
 {
-	enum { SupportsMove = true };
 	enum { IsZeroConstruct = true };
 };
 
diff --git a/Engine/Source/Runtime/Core/Public/Misc/MemStack.h b/Engine/Source/Runtime/Core/Public/Misc/MemStack.h
--- a/Engine/Source/Runtime/Core/Public/Misc/MemStack.h
+++ b/Engine/Source/Runtime/Core/Public/Misc/MemStack.h
@@ -456,7 +456,6 @@ public:
 template <uint32 Alignment>
 struct TAllocatorTraits<TMemStackAllocator<Alignment>> : TAllocatorTraitsBase<TMemStackAllocator<Alignment>>
 {
-	enum { SupportsMove    = true };
 	enum { IsZeroConstruct = true };
 };
 
diff --git a/Engine/Source/Runtime/Core/Public/Serialization/MemoryImage.h b/Engine/Source/Runtime/Core/Public/Serialization/MemoryImage.h
--- a/Engine/Source/Runtime/Core/Public/Serialization/MemoryImage.h
+++ b/Engine/Source/Runtime/Core/Public/Serialization/MemoryImage.h
@@ -710,7 +710,6 @@ public:
 template <uint32 Alignment>
 struct TAllocatorTraits<TMemoryImageAllocator<Alignment>> : TAllocatorTraitsBase<TMemoryImageAllocator<Alignment>>
 {
-	enum { SupportsMove = true };
 	enum { IsZeroConstruct = true };
 	enum { SupportsFreezeMemoryImage = true };
 	enum { SupportsElementAlignment = true };
diff --git a/Engine/Source/Runtime/Engine/Public/Animation/AnimCompressionTypes.h b/Engine/Source/Runtime/Engine/Public/Animation/AnimCompressionTypes.h
--- a/Engine/Source/Runtime/Engine/Public/Animation/AnimCompressionTypes.h
+++ b/Engine/Source/Runtime/Engine/Public/Animation/AnimCompressionTypes.h
@@ -695,12 +695,6 @@ public:
 	}
 };
 
-template <uint32 Alignment>
-struct TAllocatorTraits<TMaybeMappedAllocator<Alignment>> : TAllocatorTraitsBase<TMaybeMappedAllocator<Alignment>>
-{
-	enum { SupportsMove = true };
-};
-
 template<typename T, uint32 Alignment>
 struct TIsContiguousContainer<TMaybeMappedArray<T, Alignment>>
 {
diff --git a/Engine/Source/Runtime/RenderCore/Public/RenderGraphAllocator.h b/Engine/Source/Runtime/RenderCore/Public/RenderGraphAllocator.h
--- a/Engine/Source/Runtime/RenderCore/Public/RenderGraphAllocator.h
+++ b/Engine/Source/Runtime/RenderCore/Public/RenderGraphAllocator.h
@@ -235,7 +235,6 @@ public:
 template <uint32 Alignment>
 struct TAllocatorTraits<TRDGArrayAllocator<Alignment>> : TAllocatorTraitsBase<TRDGArrayAllocator<Alignment>>
 {
-	enum { SupportsMove = true };
 	enum { IsZeroConstruct = true };
 };
 
diff --git a/Engine/Source/Runtime/SlateCore/Public/Rendering/RenderingCommon.h b/Engine/Source/Runtime/SlateCore/Public/Rendering/RenderingCommon.h
--- a/Engine/Source/Runtime/SlateCore/Public/Rendering/RenderingCommon.h
+++ b/Engine/Source/Runtime/SlateCore/Public/Rendering/RenderingCommon.h
@@ -470,7 +470,6 @@ public:
 template <typename T>
 struct TAllocatorTraits<FSlateStatTrackingMemoryAllocator<T>> : TAllocatorTraitsBase<FSlateStatTrackingMemoryAllocator<T>>
 {
-	enum { SupportsMove = TAllocatorTraits<FDefaultAllocator>::SupportsMove };
 	enum { IsZeroConstruct = TAllocatorTraits<FDefaultAllocator>::IsZeroConstruct };
 };
 
