[Engine]
* Attempt at fixing Tuple static analysis farm error

diff --git a/Engine/Source/Runtime/Core/Public/Templates/Tuple.h b/Engine/Source/Runtime/Core/Public/Templates/Tuple.h
--- a/Engine/Source/Runtime/Core/Public/Templates/Tuple.h
+++ b/Engine/Source/Runtime/Core/Public/Templates/Tuple.h
@@ -463,14 +463,13 @@ namespace UE::Core::Private::Tuple
 		template <typename FuncType, typename... ArgTypes> decltype(auto) ApplyBefore(FuncType&& Func, ArgTypes&&... Args)       volatile&& { return ::Invoke(Func, static_cast<      volatile TTupleBase&&>(*this).template Get<0>(), static_cast<      volatile TTupleBase&&>(*this).template Get<1>(), Forward<ArgTypes>(Args)...); }
 		template <typename FuncType, typename... ArgTypes> decltype(auto) ApplyBefore(FuncType&& Func, ArgTypes&&... Args) const volatile&& { return ::Invoke(Func, static_cast<const volatile TTupleBase&&>(*this).template Get<0>(), static_cast<const volatile TTupleBase&&>(*this).template Get<1>(), Forward<ArgTypes>(Args)...); }
 
-		FORCEINLINE friend FArchive& operator<<(FArchive& Ar, TTupleBase& Tuple)
+		FORCEINLINE void Serialize(FArchive& Ar, TTupleBase& Tuple)
 		{
 			Ar << Tuple.Key;
 			Ar << Tuple.Value;
-			return Ar;
 		}
 
-		FORCEINLINE friend void operator<<(FStructuredArchive::FSlot Slot, TTupleBase& Tuple)
+		FORCEINLINE void SerializeStructured(FStructuredArchive::FSlot Slot, TTupleBase& Tuple)
 		{
 			FStructuredArchive::FStream Stream = Slot.EnterStream();
 			Stream.EnterElement() << Tuple.Key;
