removed duplicated include

diff --git a/Engine/Source/Runtime/Core/Private/Windows/WindowsPlatformProcess.cpp b/Engine/Source/Runtime/Core/Private/Windows/WindowsPlatformProcess.cpp
--- a/Engine/Source/Runtime/Core/Private/Windows/WindowsPlatformProcess.cpp
+++ b/Engine/Source/Runtime/Core/Private/Windows/WindowsPlatformProcess.cpp
@@ -6,7 +6,6 @@
 #include "Containers/StringConv.h"
 #include "Containers/UnrealString.h"
 #include "CoreGlobals.h"
-#include "CoreGlobals.h"
 #include "HAL/FileManager.h"
 #include "HAL/PlatformAffinity.h"
 #include "HAL/PlatformMisc.h"
