Removed MoveWillEmptyContainer and TContainerTraits.

diff --git a/Engine/Source/Runtime/Core/Public/Containers/Array.h b/Engine/Source/Runtime/Core/Public/Containers/Array.h
--- a/Engine/Source/Runtime/Core/Public/Containers/Array.h
+++ b/Engine/Source/Runtime/Core/Public/Containers/Array.h
@@ -291,7 +291,6 @@ namespace UE4Array_Private
 		{
 			Value =
 				TOr<TAreTypesEqual<FromAllocatorType, ToAllocatorType>, TCanMoveBetweenAllocators<FromAllocatorType, ToAllocatorType>>::Value && // Allocators must be equal or move-compatible
-				TContainerTraits<FromArrayType>::MoveWillEmptyContainer &&   // A move must be allowed to leave the source array empty
 				(
 					TAreTypesEqual         <ToElementType, FromElementType>::Value || // The element type of the container must be the same, or...
 					TIsBitwiseConstructible<ToElementType, FromElementType>::Value    // ... the element type of the source container must be bitwise constructible from the element type in the destination container
@@ -515,101 +514,89 @@ private:
 	/**
 	 * Moves or copies array. Depends on the array type traits.
 	 *
-	 * This override moves.
-	 *
 	 * @param ToArray Array to move into.
 	 * @param FromArray Array to move from.
+	 * @param PrevMax The previous allocated size.
 	 */
 	template <typename FromArrayType, typename ToArrayType>
-	static FORCEINLINE std::enable_if_t<UE4Array_Private::TCanMoveTArrayPointersBetweenArrayTypes<FromArrayType, ToArrayType>::Value> MoveOrCopy(ToArrayType& ToArray, FromArrayType& FromArray, SizeType PrevMax)
+	static FORCEINLINE void MoveOrCopy(ToArrayType& ToArray, FromArrayType& FromArray, SizeType PrevMax)
 	{
-		static_assert(std::is_same_v<TArray, ToArrayType>, "MoveOrCopy is expected to be called with the current array type as the destination");
+		if constexpr (UE4Array_Private::TCanMoveTArrayPointersBetweenArrayTypes<FromArrayType, ToArrayType>::Value)
+		{
+			// Move
 
-		using FromAllocatorType = typename FromArrayType::AllocatorType;
-		using ToAllocatorType   = typename ToArrayType::AllocatorType;
+			static_assert(std::is_same_v<TArray, ToArrayType>, "MoveOrCopy is expected to be called with the current array type as the destination");
 
-		if constexpr (TCanMoveBetweenAllocators<FromAllocatorType, ToAllocatorType>::Value)
-		{
-			ToArray.AllocatorInstance.template MoveToEmptyFromOtherAllocator<FromAllocatorType>(FromArray.AllocatorInstance);
-		}
-		else
-		{
-			ToArray.AllocatorInstance.MoveToEmpty(FromArray.AllocatorInstance);
-		}
+			using FromAllocatorType = typename FromArrayType::AllocatorType;
+			using ToAllocatorType = typename ToArrayType::AllocatorType;
 
-		ToArray  .ArrayNum = (SizeType)FromArray.ArrayNum;
-		ToArray  .ArrayMax = (SizeType)FromArray.ArrayMax;
+			if constexpr (TCanMoveBetweenAllocators<FromAllocatorType, ToAllocatorType>::Value)
+			{
+				ToArray.AllocatorInstance.template MoveToEmptyFromOtherAllocator<FromAllocatorType>(FromArray.AllocatorInstance);
+			}
+			else
+			{
+				ToArray.AllocatorInstance.MoveToEmpty(FromArray.AllocatorInstance);
+			}
 
-		// Ensure the destination container could hold the source range (when the allocator size types shrink)
-		if constexpr (sizeof(USizeType) < sizeof(typename FromArrayType::USizeType))
-		{
-			if (ToArray.ArrayNum != FromArray.ArrayNum || ToArray.ArrayMax != FromArray.ArrayMax)
+			ToArray.ArrayNum = (SizeType)FromArray.ArrayNum;
+			ToArray.ArrayMax = (SizeType)FromArray.ArrayMax;
+
+			// Ensure the destination container could hold the source range (when the allocator size types shrink)
+			if constexpr (sizeof(USizeType) < sizeof(typename FromArrayType::USizeType))
 			{
-				OnInvalidNum((USizeType)ToArray.ArrayNum);
+				if (ToArray.ArrayNum != FromArray.ArrayNum || ToArray.ArrayMax != FromArray.ArrayMax)
+				{
+					OnInvalidNum((USizeType)ToArray.ArrayNum);
+				}
 			}
+
+			FromArray.ArrayNum = 0;
+			FromArray.ArrayMax = FromArray.AllocatorInstance.GetInitialCapacity();
 		}
+		else
+		{
+			// Copy
 
-		FromArray.ArrayNum = 0;
-		FromArray.ArrayMax = FromArray.AllocatorInstance.GetInitialCapacity();
+			ToArray.CopyToEmpty(FromArray.GetData(), FromArray.Num(), PrevMax);
+		}
 	}
 
 	/**
 	 * Moves or copies array. Depends on the array type traits.
 	 *
-	 * This override copies.
-	 *
 	 * @param ToArray Array to move into.
 	 * @param FromArray Array to move from.
+	 * @param PrevMax The previous allocated size.
 	 * @param ExtraSlack Tells how much extra memory should be preallocated
 	 *                   at the end of the array in the number of elements.
 	 */
 	template <typename FromArrayType, typename ToArrayType>
-	static FORCEINLINE std::enable_if_t<!UE4Array_Private::TCanMoveTArrayPointersBetweenArrayTypes<FromArrayType, ToArrayType>::Value> MoveOrCopy(ToArrayType& ToArray, FromArrayType& FromArray, SizeType PrevMax)
+	static FORCEINLINE void MoveOrCopyWithSlack(ToArrayType& ToArray, FromArrayType& FromArray, SizeType PrevMax, SizeType ExtraSlack)
 	{
-		ToArray.CopyToEmpty(FromArray.GetData(), FromArray.Num(), PrevMax);
-	}
+		if constexpr (UE4Array_Private::TCanMoveTArrayPointersBetweenArrayTypes<FromArrayType, ToArrayType>::Value)
+		{
+			// Move
 
-	/**
-	 * Moves or copies array. Depends on the array type traits.
-	 *
-	 * This override moves.
-	 *
-	 * @param ToArray Array to move into.
-	 * @param FromArray Array to move from.
-	 * @param ExtraSlack Tells how much extra memory should be preallocated
-	 *                   at the end of the array in the number of elements.
-	 */
-	template <typename FromArrayType, typename ToArrayType>
-	static FORCEINLINE std::enable_if_t<UE4Array_Private::TCanMoveTArrayPointersBetweenArrayTypes<FromArrayType, ToArrayType>::Value> MoveOrCopyWithSlack(ToArrayType& ToArray, FromArrayType& FromArray, SizeType PrevMax, SizeType ExtraSlack)
-	{
-		MoveOrCopy(ToArray, FromArray, PrevMax);
+			MoveOrCopy(ToArray, FromArray, PrevMax);
 
-		USizeType LocalArrayNum = (USizeType)ToArray.ArrayNum;
-		USizeType NewMax        = (USizeType)LocalArrayNum + (USizeType)ExtraSlack;
+			USizeType LocalArrayNum = (USizeType)ToArray.ArrayNum;
+			USizeType NewMax = (USizeType)LocalArrayNum + (USizeType)ExtraSlack;
 
-		// This should only happen when we've underflowed or overflowed SizeType
-		if ((SizeType)NewMax < LocalArrayNum)
-		{
-			OnInvalidNum((USizeType)ExtraSlack);
-		}
+			// This should only happen when we've underflowed or overflowed SizeType
+			if ((SizeType)NewMax < LocalArrayNum)
+			{
+				OnInvalidNum((USizeType)ExtraSlack);
+			}
 
-		ToArray.Reserve(NewMax);
-	}
+			ToArray.Reserve(NewMax);
+		}
+		else
+		{
+			// Copy
 
-	/**
-	 * Moves or copies array. Depends on the array type traits.
-	 *
-	 * This override copies.
-	 *
-	 * @param ToArray Array to move into.
-	 * @param FromArray Array to move from.
-	 * @param ExtraSlack Tells how much extra memory should be preallocated
-	 *                   at the end of the array in the number of elements.
-	 */
-	template <typename FromArrayType, typename ToArrayType>
-	static FORCEINLINE std::enable_if_t<!UE4Array_Private::TCanMoveTArrayPointersBetweenArrayTypes<FromArrayType, ToArrayType>::Value> MoveOrCopyWithSlack(ToArrayType& ToArray, FromArrayType& FromArray, SizeType PrevMax, SizeType ExtraSlack)
-	{
-		ToArray.CopyToEmptyWithSlack(FromArray.GetData(), FromArray.Num(), PrevMax, ExtraSlack);
+			ToArray.CopyToEmptyWithSlack(FromArray.GetData(), FromArray.Num(), PrevMax, ExtraSlack);
+		}
 	}
 
 public:
@@ -3447,12 +3434,6 @@ struct TIsZeroConstructType<TArray<InElementType, AllocatorType>>
 	enum { Value = TAllocatorTraits<AllocatorType>::IsZeroConstruct };
 };
 
-template <typename InElementType, typename AllocatorType>
-struct TContainerTraits<TArray<InElementType, AllocatorType> > : public TContainerTraitsBase<TArray<InElementType, AllocatorType> >
-{
-	enum { MoveWillEmptyContainer = true };
-};
-
 template <typename T, typename AllocatorType>
 struct TIsContiguousContainer<TArray<T, AllocatorType>>
 {
diff --git a/Engine/Source/Runtime/Core/Public/Containers/BitArray.h b/Engine/Source/Runtime/Core/Public/Containers/BitArray.h
--- a/Engine/Source/Runtime/Core/Public/Containers/BitArray.h
+++ b/Engine/Source/Runtime/Core/Public/Containers/BitArray.h
@@ -278,7 +278,7 @@ public:
 	 */
 	FORCEINLINE TBitArray(TBitArray&& Other)
 	{
-		MoveOrCopy(*this, Other);
+		this->Move(*this, Other);
 	}
 
 	/**
@@ -306,7 +306,7 @@ public:
 	{
 		if (this != &Other)
 		{
-			MoveOrCopy(*this, Other);
+			this->Move(*this, Other);
 		}
 
 		return *this;
@@ -409,7 +409,7 @@ private:
 	}
 
 	template <typename BitArrayType>
-	static FORCEINLINE typename TEnableIf<TContainerTraits<BitArrayType>::MoveWillEmptyContainer>::Type MoveOrCopy(BitArrayType& ToArray, BitArrayType& FromArray)
+	static FORCEINLINE void Move(BitArrayType& ToArray, BitArrayType& FromArray)
 	{
 		ToArray.AllocatorInstance.MoveToEmpty(FromArray.AllocatorInstance);
 
@@ -421,12 +421,6 @@ private:
 		// No need to call FromArray.ClearPartialSlackBits because NumBits == 0 automatically satisfies the invariant
 	}
 
-	template <typename BitArrayType>
-	static FORCEINLINE typename TEnableIf<!TContainerTraits<BitArrayType>::MoveWillEmptyContainer>::Type MoveOrCopy(BitArrayType& ToArray, BitArrayType& FromArray)
-	{
-		ToArray = FromArray;
-	}
-
 	template<typename OtherAllocator>
 	void Assign(const TBitArray<OtherAllocator>& Other)
 	{
@@ -1751,12 +1745,6 @@ FORCEINLINE uint32 GetTypeHash(const TBitArray<Allocator>& BitArray)
 	return Hash;
 }
 
-template<typename Allocator>
-struct TContainerTraits<TBitArray<Allocator> > : public TContainerTraitsBase<TBitArray<Allocator> >
-{
-	enum { MoveWillEmptyContainer = true };
-};
-
 
 /** An iterator which only iterates over set bits. */
 template<typename Allocator>
diff --git a/Engine/Source/Runtime/Core/Public/Containers/ChunkedArray.h b/Engine/Source/Runtime/Core/Public/Containers/ChunkedArray.h
--- a/Engine/Source/Runtime/Core/Public/Containers/ChunkedArray.h
+++ b/Engine/Source/Runtime/Core/Public/Containers/ChunkedArray.h
@@ -79,30 +79,24 @@ public:
 
 private:
 	template <typename ArrayType>
-	FORCEINLINE static typename TEnableIf<TContainerTraits<ArrayType>::MoveWillEmptyContainer>::Type MoveOrCopy(ArrayType& ToArray, ArrayType& FromArray)
+	FORCEINLINE static void Move(ArrayType& ToArray, ArrayType& FromArray)
 	{
 		ToArray.Chunks      = (ChunksType&&)FromArray.Chunks;
 		ToArray.NumElements = FromArray.NumElements;
 		FromArray.NumElements = 0;
 	}
 
-	template <typename ArrayType>
-	FORCEINLINE static typename TEnableIf<!TContainerTraits<ArrayType>::MoveWillEmptyContainer>::Type MoveOrCopy(ArrayType& ToArray, ArrayType& FromArray)
-	{
-		ToArray = FromArray;
-	}
-
 public:
 	TChunkedArray(TChunkedArray&& Other)
 	{
-		MoveOrCopy(*this, Other);
+		this->Move(*this, Other);
 	}
 
 	TChunkedArray& operator=(TChunkedArray&& Other)
 	{
 		if (this != &Other)
 		{
-			MoveOrCopy(*this, Other);
+			this->Move(*this, Other);
 		}
 
 		return *this;
@@ -278,8 +272,6 @@ public:
 
 protected:
 
-	friend struct TContainerTraits<TChunkedArray<ElementType, TargetBytesPerChunk, AllocatorType>>;
-
 	enum { NumElementsPerChunk = TargetBytesPerChunk / sizeof(ElementType) };
 
 	/** A chunk of the array's elements. */
@@ -337,12 +329,6 @@ public:
 };
 
 
-template <typename ElementType, uint32 TargetBytesPerChunk, typename AllocatorType>
-struct TContainerTraits<TChunkedArray<ElementType, TargetBytesPerChunk, AllocatorType> > : public TContainerTraitsBase<TChunkedArray<ElementType, TargetBytesPerChunk, AllocatorType> >
-{
-	enum { MoveWillEmptyContainer = TContainerTraits<typename TChunkedArray<ElementType, TargetBytesPerChunk, AllocatorType>::ChunksType>::MoveWillEmptyContainer };
-};
-
 
 template <typename T,uint32 TargetBytesPerChunk, typename AllocatorType> 
 void* operator new( size_t Size, TChunkedArray<T,TargetBytesPerChunk, AllocatorType>& ChunkedArray )
diff --git a/Engine/Source/Runtime/Core/Public/Containers/DynamicRHIResourceArray.h b/Engine/Source/Runtime/Core/Public/Containers/DynamicRHIResourceArray.h
--- a/Engine/Source/Runtime/Core/Public/Containers/DynamicRHIResourceArray.h
+++ b/Engine/Source/Runtime/Core/Public/Containers/DynamicRHIResourceArray.h
@@ -134,12 +134,6 @@ private:
 	LAYOUT_FIELD(bool, bNeedsCPUAccess);
 };
 
-template< typename ElementType, uint32 Alignment>
-struct TContainerTraits<TResourceArray<ElementType, Alignment>> : public TContainerTraitsBase<TResourceArray<ElementType, Alignment>>
-{
-	enum { MoveWillEmptyContainer = TContainerTraits<typename TResourceArray<ElementType, Alignment>::Super>::MoveWillEmptyContainer };
-};
-
 /**
 * Serializer for this class
 * @param Ar - archive to serialize to
diff --git a/Engine/Source/Runtime/Core/Public/Containers/IndirectArray.h b/Engine/Source/Runtime/Core/Public/Containers/IndirectArray.h
--- a/Engine/Source/Runtime/Core/Public/Containers/IndirectArray.h
+++ b/Engine/Source/Runtime/Core/Public/Containers/IndirectArray.h
@@ -427,31 +427,6 @@ private:
 };
 
 
-template<typename T, typename Allocator>
-struct TContainerTraits<TIndirectArray<T, Allocator> >
-	: public TContainerTraitsBase<TIndirectArray<T, Allocator> >
-{
-	enum { MoveWillEmptyContainer = TContainerTraitsBase<typename TIndirectArray<T, Allocator>::InternalArrayType>::MoveWillEmptyContainer };
-};
-
-template <typename T,typename Allocator>
-UE_DEPRECATED(4.22, "Placement new on TIndirectArray has been deprecated - users should call Add() passing a pointer to an object created with new.")
-void* operator new( size_t Size, TIndirectArray<T,Allocator>& Array )
-{
-	check(Size == sizeof(T));
-	const int32 Index = Array.Add((T*)FMemory::Malloc(Size));
-	return &Array[Index];
-}
-
-template <typename T,typename Allocator>
-UE_DEPRECATED(4.22, "Placement new on TIndirectArray has been deprecated - users should call Insert() passing a pointer to an object created with new.")
-void* operator new( size_t Size, TIndirectArray<T,Allocator>& Array, int32 Index )
-{
-	check(Size == sizeof(T));
-	Array.Insert((T*)FMemory::Malloc(Size), Index);
-	return &Array[Index];
-}
-
 /**
 * Serialization operator for TIndirectArray.
 *
diff --git a/Engine/Source/Runtime/Core/Public/Containers/MRUArray.h b/Engine/Source/Runtime/Core/Public/Containers/MRUArray.h
--- a/Engine/Source/Runtime/Core/Public/Containers/MRUArray.h
+++ b/Engine/Source/Runtime/Core/Public/Containers/MRUArray.h
@@ -109,10 +109,4 @@ public:
 };
 
 
-template<typename T, typename Allocator>
-struct TContainerTraits<TMRUArray<T, Allocator> > : public TContainerTraitsBase<TMRUArray<T, Allocator> >
-{
-	enum { MoveWillEmptyContainer = TContainerTraitsBase<typename TMRUArray<T, Allocator>::Super>::MoveWillEmptyContainer };
-};
-
 
diff --git a/Engine/Source/Runtime/Core/Public/Containers/Map.h b/Engine/Source/Runtime/Core/Public/Containers/Map.h
--- a/Engine/Source/Runtime/Core/Public/Containers/Map.h
+++ b/Engine/Source/Runtime/Core/Public/Containers/Map.h
@@ -135,8 +135,6 @@ class TMapBase
 	template <typename OtherKeyType, typename OtherValueType, typename OtherSetAllocator, typename OtherKeyFuncs>
 	friend class TMapBase;
 
-	friend struct TContainerTraits<TMapBase>;
-
 public:
 	static constexpr bool SupportsFreezeMemoryImage = TAllocatorTraits<SetAllocator>::SupportsFreezeMemoryImage;
 
@@ -968,8 +966,6 @@ public:
 template <typename KeyType, typename ValueType, typename SetAllocator, typename KeyFuncs>
 class TSortableMapBase : public TMapBase<KeyType, ValueType, SetAllocator, KeyFuncs>
 {
-	friend struct TContainerTraits<TSortableMapBase>;
-
 protected:
 	typedef TMapBase<KeyType, ValueType, SetAllocator, KeyFuncs> Super;
 
@@ -1097,8 +1093,6 @@ class TScriptMap;
 template<typename InKeyType, typename InValueType, typename SetAllocator /*= FDefaultSetAllocator*/, typename KeyFuncs /*= TDefaultMapHashableKeyFuncs<KeyType,ValueType,false>*/>
 class TMap : public TSortableMapBase<InKeyType, InValueType, SetAllocator, KeyFuncs>
 {
-	friend struct TContainerTraits<TMap>;
-
 	template <typename, typename>
 	friend class TScriptMap;
 
@@ -1291,8 +1285,6 @@ DECLARE_TEMPLATE_INTRINSIC_TYPE_LAYOUT((template <typename KeyType, typename Val
 template<typename KeyType, typename ValueType, typename SetAllocator /* = FDefaultSetAllocator */, typename KeyFuncs /*= TDefaultMapHashableKeyFuncs<KeyType,ValueType,true>*/>
 class TMultiMap : public TSortableMapBase<KeyType, ValueType, SetAllocator, KeyFuncs>
 {
-	friend struct TContainerTraits<TMultiMap>;
-
 	static_assert(KeyFuncs::bAllowDuplicateKeys, "TMultiMap cannot be instantiated with a KeyFuncs which disallows duplicate keys");
 
 public:
@@ -1870,19 +1862,6 @@ public:
 	using Super::Super;
 };
 
-template <typename KeyType, typename ValueType, typename SetAllocator, typename KeyFuncs>
-struct TContainerTraits<TMap<KeyType, ValueType, SetAllocator, KeyFuncs>> : public TContainerTraitsBase<TMap<KeyType, ValueType, SetAllocator, KeyFuncs>>
-{
-	enum { MoveWillEmptyContainer = TContainerTraits<typename TMap<KeyType, ValueType, SetAllocator, KeyFuncs>::ElementSetType>::MoveWillEmptyContainer };
-};
-
-
-template <typename KeyType, typename ValueType, typename SetAllocator, typename KeyFuncs>
-struct TContainerTraits<TMultiMap<KeyType, ValueType, SetAllocator, KeyFuncs>> : public TContainerTraitsBase<TMultiMap<KeyType, ValueType, SetAllocator, KeyFuncs>>
-{
-	enum { MoveWillEmptyContainer = TContainerTraits<typename TMultiMap<KeyType, ValueType, SetAllocator, KeyFuncs>::ElementSetType>::MoveWillEmptyContainer };
-};
-
 struct TMapPrivateFriend
 {
 	template <typename KeyType, typename ValueType, typename SetAllocator, typename KeyFuncs>
diff --git a/Engine/Source/Runtime/Core/Public/Containers/Set.h b/Engine/Source/Runtime/Core/Public/Containers/Set.h
--- a/Engine/Source/Runtime/Core/Public/Containers/Set.h
+++ b/Engine/Source/Runtime/Core/Public/Containers/Set.h
@@ -281,8 +281,6 @@ public:
 	typedef Allocator   AllocatorType;
 
 private:
-	friend struct TContainerTraits<TSet>;
-
 	template <typename, typename>
 	friend class TScriptSet;
 
@@ -341,7 +339,7 @@ public:
 
 private:
 	template <typename SetType>
-	static FORCEINLINE std::enable_if_t<TContainerTraits<SetType>::MoveWillEmptyContainer> MoveOrCopy(SetType& ToSet, SetType& FromSet)
+	static FORCEINLINE void Move(SetType& ToSet, SetType& FromSet)
 	{
 		ToSet.Elements = (ElementArrayType&&)FromSet.Elements;
 
@@ -351,12 +349,6 @@ private:
 		FromSet.HashSize = 0;
 	}
 
-	template <typename SetType>
-	static FORCEINLINE std::enable_if_t<!TContainerTraits<SetType>::MoveWillEmptyContainer> MoveOrCopy(SetType& ToSet, SetType& FromSet)
-	{
-		ToSet = FromSet;
-	}
-
 public:
 	/** Initializer list constructor. */
 	TSet(std::initializer_list<ElementType> InitList)
@@ -369,7 +361,7 @@ public:
 	TSet(TSet&& Other)
 		: HashSize(0)
 	{
-		MoveOrCopy(*this, Other);
+		this->Move(*this, Other);
 	}
 
 	/** Move assignment operator. */
@@ -377,7 +369,7 @@ public:
 	{
 		if (this != &Other)
 		{
-			MoveOrCopy(*this, Other);
+			this->Move(*this, Other);
 		}
 
 		return *this;
@@ -1784,12 +1776,6 @@ namespace Freeze
 
 DECLARE_TEMPLATE_INTRINSIC_TYPE_LAYOUT((template <typename ElementType, typename KeyFuncs, typename Allocator>), (TSet<ElementType, KeyFuncs, Allocator>));
 
-template<typename ElementType, typename KeyFuncs, typename Allocator>
-struct TContainerTraits<TSet<ElementType, KeyFuncs, Allocator> > : public TContainerTraitsBase<TSet<ElementType, KeyFuncs, Allocator> >
-{
-	enum { MoveWillEmptyContainer = TContainerTraits<typename TSet<ElementType, KeyFuncs, Allocator>::ElementArrayType>::MoveWillEmptyContainer };
-};
-
 struct FScriptSetLayout
 {
 	// int32 ElementOffset = 0; // always at zero offset from the TSetElement - not stored here
diff --git a/Engine/Source/Runtime/Core/Public/Containers/SortedMap.h b/Engine/Source/Runtime/Core/Public/Containers/SortedMap.h
--- a/Engine/Source/Runtime/Core/Public/Containers/SortedMap.h
+++ b/Engine/Source/Runtime/Core/Public/Containers/SortedMap.h
@@ -20,7 +20,6 @@ class TSortedMap
 {
 	template <typename OtherKeyType, typename OtherValueType, typename OtherArrayAllocator, typename OtherSortPredicate>
 	friend class TSortedMap;
-	friend struct TContainerTraits<TSortedMap>;
 
 public:
 	typedef typename TTypeTraits<KeyType  >::ConstPointerType KeyConstPointerType;
@@ -772,12 +771,6 @@ public:
 
 DECLARE_TEMPLATE_INTRINSIC_TYPE_LAYOUT((template <typename KeyType, typename ValueType, typename ArrayAllocator, typename SortPredicate>), (TSortedMap<KeyType, ValueType, ArrayAllocator, SortPredicate>));
 
-template <typename KeyType, typename ValueType, typename ArrayAllocator, typename SortPredicate>
-struct TContainerTraits<TSortedMap<KeyType, ValueType, ArrayAllocator, SortPredicate>> : public TContainerTraitsBase<TSortedMap<KeyType, ValueType, ArrayAllocator, SortPredicate>>
-{
-	enum { MoveWillEmptyContainer = TContainerTraits<typename TSortedMap<KeyType, ValueType, ArrayAllocator, SortPredicate>::ElementArrayType>::MoveWillEmptyContainer };
-};
-
 struct TSortedMapPrivateFriend
 {
 	template <typename KeyType, typename ValueType, typename ArrayAllocator, typename SortPredicate>
diff --git a/Engine/Source/Runtime/Core/Public/Containers/SparseArray.h b/Engine/Source/Runtime/Core/Public/Containers/SparseArray.h
--- a/Engine/Source/Runtime/Core/Public/Containers/SparseArray.h
+++ b/Engine/Source/Runtime/Core/Public/Containers/SparseArray.h
@@ -75,8 +75,6 @@ class TSparseArray
 {
 	using ElementType = InElementType;
 
-	friend struct TContainerTraits<TSparseArray>;
-
 	template <typename, typename>
 	friend class TScriptSparseArray;
 
@@ -670,7 +668,7 @@ public:
 	/** Move constructor. */
 	TSparseArray(TSparseArray&& InCopy)
 	{
-		MoveOrCopy(*this, InCopy);
+		this->Move(*this, InCopy);
 	}
 
 	/** Copy constructor. */
@@ -686,7 +684,7 @@ public:
 	{
 		if(this != &InCopy)
 		{
-			MoveOrCopy(*this, InCopy);
+			this->Move(*this, InCopy);
 		}
 		return *this;
 	}
@@ -740,7 +738,7 @@ public:
 
 private:
 	template <typename SparseArrayType>
-	FORCEINLINE static typename TEnableIf<TContainerTraits<SparseArrayType>::MoveWillEmptyContainer>::Type MoveOrCopy(SparseArrayType& ToArray, SparseArrayType& FromArray)
+	FORCEINLINE static void Move(SparseArrayType& ToArray, SparseArrayType& FromArray)
 	{
 		// Destruct the allocated elements.
 		if( !TIsTriviallyDestructible<ElementType>::Value )
@@ -760,12 +758,6 @@ private:
 		FromArray.NumFreeIndices = 0;
 	}
 
-	template <typename SparseArrayType>
-	FORCEINLINE static typename TEnableIf<!TContainerTraits<SparseArrayType>::MoveWillEmptyContainer>::Type MoveOrCopy(SparseArrayType& ToArray, SparseArrayType& FromArray)
-	{
-		ToArray = FromArray;
-	}
-
 public:
 	// Accessors.
 	ElementType& operator[](int32 Index)
@@ -1204,14 +1196,6 @@ namespace Freeze
 
 DECLARE_TEMPLATE_INTRINSIC_TYPE_LAYOUT((template <typename ElementType, typename Allocator>), (TSparseArray<ElementType, Allocator>));
 
-template<typename ElementType, typename Allocator>
-struct TContainerTraits<TSparseArray<ElementType, Allocator> > : public TContainerTraitsBase<TSparseArray<ElementType, Allocator> >
-{
-	enum { MoveWillEmptyContainer =
-		TContainerTraits<typename TSparseArray<ElementType, Allocator>::DataType>::MoveWillEmptyContainer &&
-		TContainerTraits<typename TSparseArray<ElementType, Allocator>::AllocationBitArrayType>::MoveWillEmptyContainer };
-};
-
 //
 // TSparseArray operator news.
 //
diff --git a/Engine/Source/Runtime/Core/Public/Containers/TransArray.h b/Engine/Source/Runtime/Core/Public/Containers/TransArray.h
--- a/Engine/Source/Runtime/Core/Public/Containers/TransArray.h
+++ b/Engine/Source/Runtime/Core/Public/Containers/TransArray.h
@@ -187,12 +187,6 @@ protected:
 };
 
 
-template<typename T>
-struct TContainerTraits<TTransArray<T> > : public TContainerTraitsBase<TTransArray<T> >
-{
-	enum { MoveWillEmptyContainer = TContainerTraitsBase<typename TTransArray<T>::Super>::MoveWillEmptyContainer };
-};
-
 
 //
 // Transactional array operator news.
diff --git a/Engine/Source/Runtime/Core/Public/Containers/UnrealString.h b/Engine/Source/Runtime/Core/Public/Containers/UnrealString.h
--- a/Engine/Source/Runtime/Core/Public/Containers/UnrealString.h
+++ b/Engine/Source/Runtime/Core/Public/Containers/UnrealString.h
@@ -55,8 +55,6 @@ public:
 	using AllocatorType = TSizedDefaultAllocator<32>;
 
 private:
-	friend struct TContainerTraits<FString>;
-
 	/** Array holding the character data */
 	typedef TArray<TCHAR, AllocatorType> DataType;
 	DataType Data;
@@ -1858,12 +1856,6 @@ public:
 	}
 };
 
-template<>
-struct TContainerTraits<FString> : public TContainerTraitsBase<FString>
-{
-	enum { MoveWillEmptyContainer = TContainerTraits<FString::DataType>::MoveWillEmptyContainer };
-};
-
 template<> struct TIsZeroConstructType<FString> { enum { Value = true }; };
 Expose_TNameOf(FString)
 
diff --git a/Engine/Source/Runtime/Core/Public/Templates/UnrealTypeTraits.h b/Engine/Source/Runtime/Core/Public/Templates/UnrealTypeTraits.h
--- a/Engine/Source/Runtime/Core/Public/Templates/UnrealTypeTraits.h
+++ b/Engine/Source/Runtime/Core/Public/Templates/UnrealTypeTraits.h
@@ -382,17 +382,6 @@ struct TTypeTraitsBase
 template<typename T> struct TTypeTraits : public TTypeTraitsBase<T> {};
 
 
-/**
- * Traits for containers.
- */
-template<typename T> struct TContainerTraitsBase
-{
-	// This should be overridden by every container that supports emptying its contents via a move operation.
-	enum { MoveWillEmptyContainer = false };
-};
-
-template<typename T> struct TContainerTraits : public TContainerTraitsBase<T> {};
-
 struct FVirtualDestructor
 {
 	virtual ~FVirtualDestructor() {}
diff --git a/Engine/Source/Runtime/Engine/Classes/AI/Navigation/NavigationTypes.h b/Engine/Source/Runtime/Engine/Classes/AI/Navigation/NavigationTypes.h
--- a/Engine/Source/Runtime/Engine/Classes/AI/Navigation/NavigationTypes.h
+++ b/Engine/Source/Runtime/Engine/Classes/AI/Navigation/NavigationTypes.h
@@ -745,12 +745,6 @@ public:
 	typedef TArray<InElementType, NavMeshMemory::FNavAllocator> Super;
 };
 
-template<typename InElementType>
-struct TContainerTraits<TNavStatArray<InElementType> > : public TContainerTraitsBase<TNavStatArray<InElementType> >
-{
-	enum { MoveWillEmptyContainer = TContainerTraits<typename TNavStatArray<InElementType>::Super>::MoveWillEmptyContainer };
-};
-
 
 //----------------------------------------------------------------------//
 // generic "landscape" support
diff --git a/Engine/Source/Runtime/Engine/Public/StaticMeshResources.h b/Engine/Source/Runtime/Engine/Public/StaticMeshResources.h
--- a/Engine/Source/Runtime/Engine/Public/StaticMeshResources.h
+++ b/Engine/Source/Runtime/Engine/Public/StaticMeshResources.h
@@ -343,10 +343,6 @@ struct TIsZeroConstructType<FStaticMeshSectionArray> : TIsZeroConstructType<TArr
 {
 };
 template <>
-struct TContainerTraits<FStaticMeshSectionArray> : TContainerTraits<TArray<FStaticMeshSection, TInlineAllocator<1>>>
-{
-};
-template <>
 struct TIsContiguousContainer<FStaticMeshSectionArray> : TIsContiguousContainer<TArray<FStaticMeshSection, TInlineAllocator<1>>>
 {
 };
