Fix clang win64 linker error due to missing CORE_API in TLargeWorldRenderScalar

diff --git a/Engine/Source/Runtime/Core/Private/Misc/LargeWorldRenderPosition.cpp b/Engine/Source/Runtime/Core/Private/Misc/LargeWorldRenderPosition.cpp
--- a/Engine/Source/Runtime/Core/Private/Misc/LargeWorldRenderPosition.cpp
+++ b/Engine/Source/Runtime/Core/Private/Misc/LargeWorldRenderPosition.cpp
@@ -120,7 +120,7 @@ FMatrix TLargeWorldRenderScalar<TScalar>::MakeClampedToRelativeWorldMatrixDouble
 // This defines a theoretical bound on the error.
 
 template<>
-void TLargeWorldRenderScalar<float>::Validate(double InAbsolute)
+CORE_API void TLargeWorldRenderScalar<float>::Validate(double InAbsolute)
 {
 	constexpr int FractionBitCount = 3;
 	constexpr double Tolerance = 1.0 / (1ULL << FractionBitCount);
@@ -132,7 +132,7 @@ void TLargeWorldRenderScalar<float>::Validate(double InAbsolute)
 }
 
 template<>
-void TLargeWorldRenderScalar<double>::Validate(double InAbsolute)
+CORE_API void TLargeWorldRenderScalar<double>::Validate(double InAbsolute)
 {
 	constexpr int FractionBitCount = 32;
 	constexpr double Tolerance = 1.0 / (1ULL << FractionBitCount);
