Add logging if tpause support is enabled

#localization none
#tests boot test

diff --git a/Engine/Source/Runtime/Core/Private/Windows/WindowsPlatformMisc.cpp b/Engine/Source/Runtime/Core/Private/Windows/WindowsPlatformMisc.cpp
--- a/Engine/Source/Runtime/Core/Private/Windows/WindowsPlatformMisc.cpp
+++ b/Engine/Source/Runtime/Core/Private/Windows/WindowsPlatformMisc.cpp
@@ -2418,6 +2418,7 @@ public:
 	FCPUIDQueriedData()
 		: bHasCPUIDInstruction(CheckForCPUIDInstruction()), Vendor(), CPUInfo(0), CacheLineSize(PLATFORM_CACHE_LINE_SIZE)
 	{
+		bHasTimedPauseInstruction = false;
 		if(bHasCPUIDInstruction)
 		{
 			GetCPUVendor(Vendor);
@@ -2428,6 +2429,7 @@ public:
 			CPUInfo2 = Info[2];
 			CacheLineSize = QueryCacheLineSize();
 			bHasTimedPauseInstruction = CheckForTimedPauseInstruction();
+			UE_CLOG(bHasTimedPauseInstruction, LogWindows, Log, TEXT("Enabling Tpause support"));
 		}
 	}
 
