Support remainder operator in IntVector.

diff --git a/Engine/Source/Runtime/Core/Public/Math/IntVector.h b/Engine/Source/Runtime/Core/Public/Math/IntVector.h
--- a/Engine/Source/Runtime/Core/Public/Math/IntVector.h
+++ b/Engine/Source/Runtime/Core/Public/Math/IntVector.h
@@ -226,6 +226,21 @@ PRAGMA_ENABLE_DEPRECATION_WARNINGS
 		return *this;
 	}
 
+	/**
+	 * Remainder of division of this point.
+	 *
+	 * @param Divisor What to divide the point by.
+	 * @return Reference to this point after remainder.
+	 */
+	TIntVector3& operator%=(IntType Divisor)
+	{
+		X %= Divisor;
+		Y %= Divisor;
+		Z %= Divisor;
+
+		return *this;
+	}
+
 	/**
 	 * Adds to this point.
 	 *
@@ -289,6 +304,17 @@ PRAGMA_ENABLE_DEPRECATION_WARNINGS
 		return TIntVector3(*this) /= Divisor;
 	}
 
+	/**
+	 * Gets the remainder of division on this point.
+	 *
+	 * @param Divisor What to divide the point by.
+	 * @return A new remainder int point.
+	 */
+	TIntVector3 operator%(IntType Divisor) const
+	{
+		return TIntVector3(*this) %= Divisor;
+	}
+
 	/**
 	 * Gets the result of addition on this point.
 	 *
@@ -809,6 +835,22 @@ PRAGMA_ENABLE_DEPRECATION_WARNINGS
 		return *this;
 	}
 
+	/**
+	 * Remainder of division of this point.
+	 *
+	 * @param Divisor What to divide the point by.
+	 * @return Reference to this point after remainder.
+	 */
+	TIntVector4& operator%=(IntType Divisor)
+	{
+		X %= Divisor;
+		Y %= Divisor;
+		Z %= Divisor;
+		W %= Divisor;
+
+		return *this;
+	}
+
 	/**
 	 * Adds to this point.
 	 *
@@ -890,6 +932,17 @@ PRAGMA_ENABLE_DEPRECATION_WARNINGS
 		return TIntVector4(*this) /= Divisor;
 	}
 
+	/**
+	 * Gets the remainder of division on this point.
+	 *
+	 * @param Divisor What to divide the point by.
+	 * @return A new remainder int point.
+	 */
+	TIntVector4 operator%(IntType Divisor) const
+	{
+		return TIntVector4(*this) %= Divisor;
+	}
+
 	/**
 	 * Gets the result of addition on this point.
 	 *
