Allow backtrace captures to vary in size on a per-platform basis.

Some platforms have threads with very small stacks that we don't appear to be able to control. This places a limit on how much stack is available to capture a callstack in situations where memory hooks are in place. Curl's host name resolving threads an the CurlCalloc() hooks are one such example.

#rnx

diff --git a/Engine/Source/Runtime/Core/Private/ProfilingDebugging/CallstackTracePrivate.h b/Engine/Source/Runtime/Core/Private/ProfilingDebugging/CallstackTracePrivate.h
--- a/Engine/Source/Runtime/Core/Private/ProfilingDebugging/CallstackTracePrivate.h
+++ b/Engine/Source/Runtime/Core/Private/ProfilingDebugging/CallstackTracePrivate.h
@@ -42,10 +42,9 @@ namespace {
 	public:
 		struct FBacktraceEntry
 		{
-			enum {MaxStackDepth = 256};
 			uint64	Hash = 0;
 			uint32	FrameCount = 0;
-			uint64	Frames[MaxStackDepth] = { 0 };
+			uint64* Frames;
 		};
 
 		FCallstackTracer(FMalloc* InMalloc)
diff --git a/Engine/Source/Runtime/Core/Private/ProfilingDebugging/Microsoft/WindowsCallstackTrace.cpp b/Engine/Source/Runtime/Core/Private/ProfilingDebugging/Microsoft/WindowsCallstackTrace.cpp
--- a/Engine/Source/Runtime/Core/Private/ProfilingDebugging/Microsoft/WindowsCallstackTrace.cpp
+++ b/Engine/Source/Runtime/Core/Private/ProfilingDebugging/Microsoft/WindowsCallstackTrace.cpp
@@ -567,7 +567,7 @@ const FBacktracer::FFunction* FBacktracer::LookupFunction(UPTRINT Address, FLook
 uint32 FBacktracer::GetBacktraceId(void* AddressOfReturnAddress) const
 {
 	FLookupState LookupState = {};
-	FCallstackTracer::FBacktraceEntry BacktraceEntry;
+	uint64 Frames[256];
 
 	UPTRINT* StackPointer = (UPTRINT*)AddressOfReturnAddress;
 
@@ -596,8 +596,7 @@ uint32 FBacktracer::GetBacktraceId(void* AddressOfReturnAddress) const
 	{
 		UPTRINT RetAddr = *StackPointer;
 		
-		BacktraceEntry.Frames[FrameIdx++] = RetAddr;
-		BacktraceEntry.FrameCount = FrameIdx;
+		Frames[FrameIdx++] = RetAddr;
 
 		// This is a simple order-dependent LCG. Should be sufficient enough
 		BacktraceHash += RetAddr;
@@ -685,10 +684,13 @@ uint32 FBacktracer::GetBacktraceId(void* AddressOfReturnAddress) const
 		StackPointer += Function->RspBias;
 	}
 	// Trunkate callstacks longer than MaxStackDepth
-	while (*StackPointer && FrameIdx < FCallstackTracer::FBacktraceEntry::MaxStackDepth);
+	while (*StackPointer && FrameIdx < UE_ARRAY_COUNT(Frames));
 
-	// Save the collected id
+	// Build the backtrace entry for submission
+	FCallstackTracer::FBacktraceEntry BacktraceEntry;
 	BacktraceEntry.Hash = BacktraceHash;
+	BacktraceEntry.FrameCount = FrameIdx;
+	BacktraceEntry.Frames = Frames;
 
 #if BACKTRACE_DBGLVL >= 3
 	for (uint32 i = 0; i < NumBacktrace; ++i)
