fix for PrintScriptCallStack crashing when called from the immedieate window
changed `FWindowsPlatformMisc::IsLocalPrintThreadSafe` to return true if attached in the debugger this allows log calls to immediately print to the debug window

a follow up change is necessary to allow for attaching after the process has started

diff --git a/Engine/Source/Runtime/Core/Public/Windows/WindowsPlatformMisc.h b/Engine/Source/Runtime/Core/Public/Windows/WindowsPlatformMisc.h
--- a/Engine/Source/Runtime/Core/Public/Windows/WindowsPlatformMisc.h
+++ b/Engine/Source/Runtime/Core/Public/Windows/WindowsPlatformMisc.h
@@ -82,6 +82,15 @@ struct CORE_API FWindowsPlatformMisc
 
 	static void SetUTF8Output();
 	static void LocalPrint(const TCHAR *Message);
+
+	static bool IsLocalPrintThreadSafe()
+	{ 
+		//returning true when the debugger is attached is to allow
+		//printing of log lines immediately to the output window.
+		//return false in not attached because OutputDebuString is slow.
+		return IsDebuggerPresent();
+	}
+	
 	static void RequestExitWithStatus(bool Force, uint8 ReturnCode);
 	static void RequestExit(bool Force);
 	static const TCHAR* GetSystemErrorMessage(TCHAR* OutBuffer, int32 BufferCount, int32 Error);
diff --git a/Engine/Source/Runtime/CoreUObject/Private/UObject/ScriptCore.cpp b/Engine/Source/Runtime/CoreUObject/Private/UObject/ScriptCore.cpp
--- a/Engine/Source/Runtime/CoreUObject/Private/UObject/ScriptCore.cpp
+++ b/Engine/Source/Runtime/CoreUObject/Private/UObject/ScriptCore.cpp
@@ -346,7 +346,6 @@ void PrintScriptCallStackImpl()
 			ScriptStack << TEXT("\n");
 		}
 		UE_LOG(LogOutputDevice, Warning, TEXT("%s"), *ScriptStack);
-		GLog->Flush(); // flush is needed because logging is buffered 
 	}
 }
 
