added new macro `UE_DISABLE_OPTIMIZATION_SHIP` to replace `PRAGMA_DISABLE_OPTIMIZATION`

will deprecate `PRAGMA_DISABLE_OPTIMIZATION `in a following changelist after converting code to use `UE_DISABLE_OPTIMIZATION_SHIP`

added `UE_CHECK_DISABLE_OPTIMIZATION` to check for
`UE_DISABLE_OPTIMIZATION` being left in by accident

will enable `UE_CHECK_DISABLE_OPTIMIZATION` in a following changelist

diff --git a/Engine/Source/Programs/AutomationTool/Scripts/CheckBalancedMacros.cs b/Engine/Source/Programs/AutomationTool/Scripts/CheckBalancedMacros.cs
--- a/Engine/Source/Programs/AutomationTool/Scripts/CheckBalancedMacros.cs
+++ b/Engine/Source/Programs/AutomationTool/Scripts/CheckBalancedMacros.cs
@@ -27,6 +27,10 @@ namespace AutomationTool
 				"PRAGMA_DISABLE_OPTIMIZATION",
 				"PRAGMA_ENABLE_OPTIMIZATION"
 			},
+			{
+				"UE_DISABLE_OPTIMIZATION_SHIP",
+				"UE_ENABLE_OPTIMIZATION_SHIP"
+			},
 			{
 				"PRAGMA_DISABLE_DEPRECATION_WARNINGS",
 				"PRAGMA_ENABLE_DEPRECATION_WARNINGS"
diff --git a/Engine/Source/Runtime/Core/Public/HAL/Platform.h b/Engine/Source/Runtime/Core/Public/HAL/Platform.h
--- a/Engine/Source/Runtime/Core/Public/HAL/Platform.h
+++ b/Engine/Source/Runtime/Core/Public/HAL/Platform.h
@@ -730,8 +730,8 @@
 #endif
 
 // Enable/disable optimizations for a specific function to improve build times
-#define BEGIN_FUNCTION_BUILD_OPTIMIZATION PRAGMA_DISABLE_OPTIMIZATION
-#define END_FUNCTION_BUILD_OPTIMIZATION   PRAGMA_ENABLE_OPTIMIZATION
+#define BEGIN_FUNCTION_BUILD_OPTIMIZATION UE_DISABLE_OPTIMIZATION_SHIP
+#define END_FUNCTION_BUILD_OPTIMIZATION   UE_ENABLE_OPTIMIZATION_SHIP
 
 #ifndef FORCEINLINE_DEBUGGABLE_ACTUAL
 	#define FORCEINLINE_DEBUGGABLE_ACTUAL inline
diff --git a/Engine/Source/Runtime/Core/Public/Misc/CoreMiscDefines.h b/Engine/Source/Runtime/Core/Public/Misc/CoreMiscDefines.h
--- a/Engine/Source/Runtime/Core/Public/Misc/CoreMiscDefines.h
+++ b/Engine/Source/Runtime/Core/Public/Misc/CoreMiscDefines.h
@@ -4,6 +4,8 @@
 
 #include "../HAL/PreprocessorHelpers.h"
 
+// When passed to pragma message will result in clickable warning in VS
+#define WARNING_LOCATION(Line) __FILE__ "(" PREPROCESSOR_TO_STRING(Line) ")"
 
 // This file is included in some resource files, which issue a warning:
 //
@@ -27,14 +29,38 @@
 /** This controls if metadata for compiled in classes is unpacked and setup at boot time. Meta data is not normally used except by the editor. **/
 #define WITH_METADATA (WITH_EDITORONLY_DATA && WITH_EDITOR)
 
+// Option to check for UE_DISABLE_OPTIMIZATION being submitted
+#ifndef UE_CHECK_DISABLE_OPTIMIZATION
+#define UE_CHECK_DISABLE_OPTIMIZATION 0
+#endif
+
 // Set up optimization control macros, now that we have both the build settings and the platform macros
-#define PRAGMA_DISABLE_OPTIMIZATION		PRAGMA_DISABLE_OPTIMIZATION_ACTUAL
+
+// Defines for submitting optimizations off
+#define UE_DISABLE_OPTIMIZATION_SHIP  PRAGMA_DISABLE_OPTIMIZATION_ACTUAL
+
+//in debug keep optimizations off for the enable macro otherwise code following enable will be optimized
 #if UE_BUILD_DEBUG
-	#define PRAGMA_ENABLE_OPTIMIZATION  PRAGMA_DISABLE_OPTIMIZATION_ACTUAL
+	#define UE_ENABLE_OPTIMIZATION_SHIP  PRAGMA_DISABLE_OPTIMIZATION_ACTUAL
+#else
+	#define UE_ENABLE_OPTIMIZATION_SHIP  PRAGMA_ENABLE_OPTIMIZATION_ACTUAL
+#endif
+
+// if running on a build machine assert on the dev optimizations macros to validate that code is not being submitted with optimizations off
+#if UE_CHECK_DISABLE_OPTIMIZATION
+	#define UE_DISABLE_OPTIMIZATION static_assert(false, "Error UE_DISABLE_OPTIMIZATION submitted. Use UE_DISABLE_OPTIMIZATION_SHIP to submit with optimizations off.");
+	#define UE_ENABLE_OPTIMIZATION static_assert(false, "Error UE_ENABLE_OPTIMIZATION submitted. Use UE_ENABLE_OPTIMIZATION_SHIP to submit with optimizations off.");
 #else
-	#define PRAGMA_ENABLE_OPTIMIZATION  PRAGMA_ENABLE_OPTIMIZATION_ACTUAL
+	#define UE_DISABLE_OPTIMIZATION  UE_DISABLE_OPTIMIZATION_SHIP
+	#define UE_ENABLE_OPTIMIZATION  UE_ENABLE_OPTIMIZATION_SHIP
 #endif
 
+#define PRAGMA_DISABLE_OPTIMIZATION \
+	UE_DISABLE_OPTIMIZATION_SHIP
+
+#define PRAGMA_ENABLE_OPTIMIZATION  \
+	UE_ENABLE_OPTIMIZATION_SHIP
+
 #if UE_BUILD_DEBUG
 	#define FORCEINLINE_DEBUGGABLE FORCEINLINE_DEBUGGABLE_ACTUAL
 #else
@@ -113,9 +139,6 @@ enum EInPlace {InPlace};
 
 #endif // RC_INVOKED
 
-// When passed to pragma message will result in clickable warning in VS
-#define WARNING_LOCATION(Line) __FILE__ "(" PREPROCESSOR_TO_STRING(Line) ")"
-
 // Push and pop macro definitions
 #ifdef __clang__
 	#define UE_PUSH_MACRO(name) _Pragma(PREPROCESSOR_TO_STRING(push_macro(name)))
