IoStoreReader: Fixed incorrect data being returned when reading partial chunks

diff --git a/Engine/Source/Runtime/Core/Private/IO/IoStore.cpp b/Engine/Source/Runtime/Core/Private/IO/IoStore.cpp
--- a/Engine/Source/Runtime/Core/Private/IO/IoStore.cpp
+++ b/Engine/Source/Runtime/Core/Private/IO/IoStore.cpp
@@ -2444,7 +2444,10 @@ public:
 			if (CompressionMethod.IsNone())
 			{
 				uint64 CopySize = FMath::Min<uint64>(UncompressedSize - OffsetInBlock, RemainingSize);
+				check(UncompressedDestination + CopySize <= UncompressedBuffer.Data() + UncompressedBuffer.DataSize());
 				FMemory::Memcpy(UncompressedDestination, CompressedBuffers[OurBufferIndex].GetData() + OffsetInBlock, CopySize);
+				UncompressedDestinationOffset += CopySize;
+				RemainingSize -= CopySize;
 			}
 			else
 			{
@@ -2456,21 +2459,23 @@ public:
 					TempBuffer.SetNumUninitialized(UncompressedSize);
 					bUncompressed = FCompression::UncompressMemory(CompressionMethod, TempBuffer.GetData(), UncompressedSize, CompressedBuffers[OurBufferIndex].GetData(), CompressionBlock.GetCompressedSize());
 					uint64 CopySize = FMath::Min<uint64>(UncompressedSize - OffsetInBlock, RemainingSize);
+					check(UncompressedDestination + CopySize <= UncompressedBuffer.Data() + UncompressedBuffer.DataSize());
 					FMemory::Memcpy(UncompressedDestination, TempBuffer.GetData() + OffsetInBlock, CopySize);
+					UncompressedDestinationOffset += CopySize;
+					RemainingSize -= CopySize;
 				}
 				else
 				{
 					check(UncompressedDestination + UncompressedSize <= UncompressedBuffer.Data() + UncompressedBuffer.DataSize());
 					bUncompressed = FCompression::UncompressMemory(CompressionMethod, UncompressedDestination, UncompressedSize, CompressedBuffers[OurBufferIndex].GetData(), CompressionBlock.GetCompressedSize());
+					UncompressedDestinationOffset += UncompressedSize;
+					RemainingSize -= UncompressedSize;
 				}
 				if (!bUncompressed)
 				{
 					return FIoStatus(EIoErrorCode::ReadError, TEXT("Failed uncompressing chunk"));
 				}
 			}
-
-			UncompressedDestinationOffset += UncompressedSize;
-			RemainingSize -= UncompressedSize;
 			OffsetInBlock = 0;
 		}
 		return UncompressedBuffer;
