deprecating `PRAGMA_DISABLE_OPTIMIZATION` and `PRAGMA_ENABLE_OPTIMIZATION` replaced with

`UE_DISABLE_OPTIMIZATION` for temporary development and `UE_ENABLE_OPTIMIZATION_SHIP` to submit

`UE_DISABLE_OPTIMIZATION` will be an error on the farm in a future changelist
#p4v-preflight-copy 23129238

diff --git a/Engine/Source/Runtime/Core/Public/Misc/CoreMiscDefines.h b/Engine/Source/Runtime/Core/Public/Misc/CoreMiscDefines.h
--- a/Engine/Source/Runtime/Core/Public/Misc/CoreMiscDefines.h
+++ b/Engine/Source/Runtime/Core/Public/Misc/CoreMiscDefines.h
@@ -58,9 +58,11 @@
 #endif
 
 #define PRAGMA_DISABLE_OPTIMIZATION \
+	UE_DEPRECATED_MACRO(5.2, "PRAGMA_DISABLE_OPTIMIZATION has been deprecated. Use UE_DISABLE_OPTIMIZATION for temporary development or UE_DISABLE_OPTIMIZATION_SHIP to submit") \
 	UE_DISABLE_OPTIMIZATION_SHIP
 
 #define PRAGMA_ENABLE_OPTIMIZATION  \
+	UE_DEPRECATED_MACRO(5.2, "PRAGMA_ENABLE_OPTIMIZATION has been deprecated. Use UE_ENABLE_OPTIMIZATION for temporary development or UE_ENABLE_OPTIMIZATION_SHIP to submit") \
 	UE_ENABLE_OPTIMIZATION_SHIP
 
 #if UE_BUILD_DEBUG
