Added ECVF_ExcludeFromPreview to exclude ECVF_Scalability Cvars from the preview

diff --git a/Engine/Source/Runtime/Core/Private/HAL/ConsoleManager.cpp b/Engine/Source/Runtime/Core/Private/HAL/ConsoleManager.cpp
--- a/Engine/Source/Runtime/Core/Private/HAL/ConsoleManager.cpp
+++ b/Engine/Source/Runtime/Core/Private/HAL/ConsoleManager.cpp
@@ -110,7 +110,8 @@ public:
 
 	void ApplyPreviewIfScalability()
 	{
-		if (((uint32)Flags & (uint32)ECVF_Scalability) != 0)
+		if (((uint32)Flags & (uint32)ECVF_Scalability) != 0
+			&& ((uint32)Flags & (uint32)ECVF_ExcludeFromPreview) == 0)
 		{
 			Flags = (EConsoleVariableFlags)((uint32)Flags | (uint32)ECVF_Preview);
 		}
@@ -281,7 +282,8 @@ public:
 
 	void ApplyPreviewIfScalability()
 	{
-		if (((uint32)Flags & (uint32)ECVF_Scalability) != 0)
+		if (((uint32)Flags & (uint32)ECVF_Scalability) != 0
+			&& ((uint32)Flags & (uint32)ECVF_ExcludeFromPreview) == 0)
 		{
 			Flags = (EConsoleVariableFlags)((uint32)Flags | (uint32)ECVF_Preview);
 		}
@@ -383,6 +385,14 @@ static void ExpandScalabilityCVar(FConfigCacheIni* ConfigSystem, const FString&
 	}
 }
 
+EConsoleVariableFlags GetPreviewFlagsOfCvar(const TCHAR* Name)
+{
+	// now look up the cvar, if it exists (it's okay if it doesn't, it may not exist on host platform, but then it's not previewable!)
+	IConsoleVariable* CVar = IConsoleManager::Get().FindConsoleVariable(Name, false /* bTrackFrequentCalls */);
+	EConsoleVariableFlags PreviewFlag = (CVar != nullptr) ? (EConsoleVariableFlags)(CVar->GetFlags() & ECVF_Preview) : ECVF_Default;
+	return PreviewFlag;
+}
+
 bool IConsoleManager::VisitPlatformCVarsForEmulation(FName PlatformName, const FString& DeviceProfileName, TFunctionRef<void(const FString& CVarName, const FString& CVarValue, EConsoleVariableFlags SetBy)> Visit)
 {
 	// we can't get to Scalability code in here (it's in Engine) but we still want to apply the default level... 
@@ -467,10 +477,7 @@ bool IConsoleManager::VisitPlatformCVarsForEmulation(FName PlatformName, const F
 
 		CVarSetByMap[Name] = SetByInt;
 
-		// now look up the cvar, if it exists (it's okay if it doesn't, it may not exist on host platform, but then it's not previewable!)
-		IConsoleVariable* CVar = IConsoleManager::Get().FindConsoleVariable(*Name, false /* bTrackFrequentCalls */);
-		EConsoleVariableFlags PreviewFlag = (CVar != nullptr) ? (EConsoleVariableFlags)(CVar->GetFlags() & ECVF_Preview) : ECVF_Default;
-		Visit(Name, Value, (EConsoleVariableFlags)(SetBy | PreviewFlag));
+		Visit(Name, Value, (EConsoleVariableFlags)(SetBy | GetPreviewFlagsOfCvar(*Name)));
 	};
 
 	// now walk up the stack getting current values
@@ -538,8 +545,7 @@ bool IConsoleManager::VisitPlatformCVarsForEmulation(FName PlatformName, const F
 				FString Value = Pair.Value.GetValue();
 
 				// don't bother tracking when looking up other platform cvars
-				IConsoleVariable* CVar = IConsoleManager::Get().FindConsoleVariable(*Key, false /* bTrackFrequentCalls */);
-				EConsoleVariableFlags PreviewFlag = (CVar != nullptr) ? (EConsoleVariableFlags)(CVar->GetFlags() & ECVF_Preview) : ECVF_Default;
+				EConsoleVariableFlags PreviewFlag = GetPreviewFlagsOfCvar(*Key);
 
 				if (Key.StartsWith(TEXT("sg.")))
 				{
@@ -552,7 +558,17 @@ bool IConsoleManager::VisitPlatformCVarsForEmulation(FName PlatformName, const F
 
 					for (const auto& ScalabilityPair : ScalabilityCVars)
 					{
-						VisitIfAllowed(ScalabilityPair.Key, ScalabilityPair.Value, (EConsoleVariableFlags)(ECVF_SetByScalability | PreviewFlag));
+						// See if the expanded scalability cvar is not allowed to preview and has ECVF_ExcludeFromPreview set
+						IConsoleVariable* CVar = IConsoleManager::Get().FindConsoleVariable(*ScalabilityPair.Key, false /* bTrackFrequentCalls */);
+						EConsoleVariableFlags ScalabilityCvarFlags = (CVar != nullptr) ? CVar->GetFlags() : ECVF_Default;
+						if (ScalabilityCvarFlags & ECVF_ExcludeFromPreview)
+						{
+							VisitIfAllowed(ScalabilityPair.Key, ScalabilityPair.Value, (EConsoleVariableFlags)(ECVF_SetByScalability));
+						}
+						else
+						{
+							VisitIfAllowed(ScalabilityPair.Key, ScalabilityPair.Value, (EConsoleVariableFlags)(ECVF_SetByScalability | PreviewFlag));
+						}
 					}
 				}
 				else
diff --git a/Engine/Source/Runtime/Core/Public/HAL/IConsoleManager.h b/Engine/Source/Runtime/Core/Public/HAL/IConsoleManager.h
--- a/Engine/Source/Runtime/Core/Public/HAL/IConsoleManager.h
+++ b/Engine/Source/Runtime/Core/Public/HAL/IConsoleManager.h
@@ -84,7 +84,9 @@ enum EConsoleVariableFlags
 	 */
 	ECVF_RenderThreadSafe = 0x20,
 
-	/* ApplyCVarSettingsGroupFromIni will complain if this wasn't set, should not be combined with ECVF_Cheat */
+	/* ApplyCVarSettingsGroupFromIni will complain if this wasn't set, should not be combined with ECVF_Cheat and ECVF_ExcludeFromPreview. 
+	 * They are automatically added as  ECVF_Preview unless ECVF_ExcludeFromPreview is used
+	 */
 	ECVF_Scalability = 0x40,
 
 	/* those cvars control other cvars with the flag ECVF_Scalability, names should start with "sg." */
@@ -102,6 +104,9 @@ enum EConsoleVariableFlags
 	/* Cvars with this flag will modify the Shader Keystring for Desktop Platforms*/
 	ECVF_DesktopShaderChange = 0x800,
 
+	/* CVars with this flag will be excluded from the device profile previews. */
+	ECVF_ExcludeFromPreview = 0x1000,
+
 	// ------------------------------------------------
 
 	/* Set flags */
diff --git a/Engine/Source/Runtime/Engine/Private/Streaming/StreamingManagerTexture.cpp b/Engine/Source/Runtime/Engine/Private/Streaming/StreamingManagerTexture.cpp
--- a/Engine/Source/Runtime/Engine/Private/Streaming/StreamingManagerTexture.cpp
+++ b/Engine/Source/Runtime/Engine/Private/Streaming/StreamingManagerTexture.cpp
@@ -1309,14 +1309,14 @@ static TAutoConsoleVariable<int32> CVarTextureStreamingAmortizeCPUToGPUCopy(
 	0,
 	TEXT("If set and r.Streaming.MaxNumTexturesToStreamPerFrame > 0, limit the number of 2D textures ")
 	TEXT("streamed from CPU memory to GPU memory each frame"),
-	ECVF_Scalability);
+	ECVF_Scalability | ECVF_ExcludeFromPreview);
 
 static TAutoConsoleVariable<int32> CVarTextureStreamingMaxNumTexturesToStreamPerFrame(
 	TEXT("r.Streaming.MaxNumTexturesToStreamPerFrame"),
 	0,
 	TEXT("Maximum number of 2D textures allowed to stream from CPU memory to GPU memory each frame. ")
 	TEXT("<= 0 means no limit. This has no effect if r.Streaming.AmortizeCPUToGPUCopy is not set"),
-	ECVF_Scalability);
+	ECVF_Scalability | ECVF_ExcludeFromPreview);
 
 static FORCEINLINE bool ShouldAmortizeMipCopies()
 {
diff --git a/Engine/Source/Runtime/Engine/Private/Streaming/TextureStreamingHelpers.cpp b/Engine/Source/Runtime/Engine/Private/Streaming/TextureStreamingHelpers.cpp
--- a/Engine/Source/Runtime/Engine/Private/Streaming/TextureStreamingHelpers.cpp
+++ b/Engine/Source/Runtime/Engine/Private/Streaming/TextureStreamingHelpers.cpp
@@ -72,7 +72,7 @@ TAutoConsoleVariable<float> CVarStreamingBoost(
 	TEXT("=1.0: normal\n")
 	TEXT("<1.0: decrease wanted mip levels\n")
 	TEXT(">1.0: increase wanted mip levels"),
-	ECVF_Scalability
+	ECVF_Scalability | ECVF_ExcludeFromPreview
 	);
 
 TAutoConsoleVariable<float> CVarStreamingMinBoost(
@@ -110,13 +110,13 @@ TAutoConsoleVariable<int32> CVarStreamingPoolSize(
 	TEXT("r.Streaming.PoolSize"),
 	-1,
 	TEXT("-1: Default texture pool size, otherwise the size in MB"),
-	ECVF_Scalability);
+	ECVF_Scalability | ECVF_ExcludeFromPreview);
 
 static TAutoConsoleVariable<int32> CVarStreamingPoolSizeForMeshes(
 	TEXT("r.Streaming.PoolSizeForMeshes"),
 	-1,
 	TEXT("< 0: Mesh and texture share the same pool, otherwise the size of pool dedicated to meshes."),
-	ECVF_Scalability);
+	ECVF_Scalability | ECVF_ExcludeFromPreview);
 
 TAutoConsoleVariable<int32> CVarStreamingMaxTempMemoryAllowed(
 	TEXT("r.Streaming.MaxTempMemoryAllowed"),
