Clang warning fixes

#rnx

diff --git a/Engine/Source/Runtime/Core/Private/ProfilingDebugging/Microsoft/EventTracingForWindows.cpp b/Engine/Source/Runtime/Core/Private/ProfilingDebugging/Microsoft/EventTracingForWindows.cpp
--- a/Engine/Source/Runtime/Core/Private/ProfilingDebugging/Microsoft/EventTracingForWindows.cpp
+++ b/Engine/Source/Runtime/Core/Private/ProfilingDebugging/Microsoft/EventTracingForWindows.cpp
@@ -432,11 +432,11 @@ bool FPlatformEvents::StartETW()
 			ULONG Source;
 			ULONG Interval;
 		}
-		TRACE_PROFILE_INTERVAL;
+		MOCK_TRACE_PROFILE_INTERVAL;
 
 		// interval is specified in 100 nanosecond units, and allowed range is limited
 		// see "-setprofint" argument in https://docs.microsoft.com/en-us/windows-hardware/test/wpt/start
-		TRACE_PROFILE_INTERVAL Interval;
+		MOCK_TRACE_PROFILE_INTERVAL Interval;
 		Interval.Source = 0;
 		Interval.Interval = FMath::Clamp<ULONG>(10 * SamplingIntervalUsec, 1221, 10000000);
 
