Moved all container-level instances of SupportsFreezeMemoryImage out of the class definition so that they can't require the full type definition.
Unified all container freezing function definitions to not use a helper class - neater and removes some PDB symbols.

diff --git a/Engine/Source/Runtime/Core/Public/Containers/HashTable.h b/Engine/Source/Runtime/Core/Public/Containers/HashTable.h
--- a/Engine/Source/Runtime/Core/Public/Containers/HashTable.h
+++ b/Engine/Source/Runtime/Core/Public/Containers/HashTable.h
@@ -431,46 +431,34 @@ private:
 	uint32					HashMask;
 	uint32					IndexSize;
 
-	template<bool bFreezeMemoryImage, typename Dummy = void>
-	struct TSupportsFreezeMemoryImageHelper
+public:
+	void WriteMemoryImage(FMemoryImageWriter& Writer) const
 	{
-		static void WriteMemoryImage(FMemoryImageWriter& Writer, const THashTable&)
+		if constexpr (TAllocatorTraits<Allocator>::SupportsFreezeMemoryImage)
+		{
+			this->Hash.WriteMemoryImage(Writer, StaticGetTypeLayoutDesc<uint32>(), this->HashMask + 1u);
+			this->NextIndex.WriteMemoryImage(Writer, StaticGetTypeLayoutDesc<uint32>(), this->IndexSize);
+			Writer.WriteBytes(this->HashMask);
+			Writer.WriteBytes(this->IndexSize);
+		}
+		else
 		{
 			check(false);
 		}
+	}
 
-		static void CopyUnfrozen(const FMemoryUnfreezeContent& Context, const THashTable&, void* Dst) { new(Dst) THashTable(); }
-	};
-
-	template<typename Dummy>
-	struct TSupportsFreezeMemoryImageHelper<true, Dummy>
+	void CopyUnfrozen(const FMemoryUnfreezeContent& Context, void* Dst) const
 	{
-		static void WriteMemoryImage(FMemoryImageWriter& Writer, const THashTable& Object)
+		if constexpr (TAllocatorTraits<Allocator>::SupportsFreezeMemoryImage)
 		{
-			Object.Hash.WriteMemoryImage(Writer, StaticGetTypeLayoutDesc<uint32>(), Object.HashMask + 1u);
-			Object.NextIndex.WriteMemoryImage(Writer, StaticGetTypeLayoutDesc<uint32>(), Object.IndexSize);
-			Writer.WriteBytes(Object.HashMask);
-			Writer.WriteBytes(Object.IndexSize);
+			THashTable* DstTable = new(Dst) THashTable(this->HashMask + 1u, this->IndexSize);
+			FMemory::Memcpy(DstTable->Hash.GetAllocation(), this->Hash.GetAllocation(), (this->HashMask + 1u) * 4);
+			FMemory::Memcpy(DstTable->NextIndex.GetAllocation(), this->NextIndex.GetAllocation(), this->IndexSize * 4);
 		}
-		static void CopyUnfrozen(const FMemoryUnfreezeContent& Context, const THashTable& Object, void* Dst)
+		else
 		{
-			THashTable* DstTable = new(Dst) THashTable(Object.HashMask + 1u, Object.IndexSize);
-			FMemory::Memcpy(DstTable->Hash.GetAllocation(), Object.Hash.GetAllocation(), (Object.HashMask + 1u) * 4);
-			FMemory::Memcpy(DstTable->NextIndex.GetAllocation(), Object.NextIndex.GetAllocation(), Object.IndexSize * 4);
+			new(Dst) THashTable();
 		}
-	};
-
-public:
-	void WriteMemoryImage(FMemoryImageWriter& Writer) const
-	{
-		static const bool bSupportsFreezeMemoryImage = TAllocatorTraits<Allocator>::SupportsFreezeMemoryImage;
-		TSupportsFreezeMemoryImageHelper<bSupportsFreezeMemoryImage>::WriteMemoryImage(Writer, *this);
-	}
-
-	void CopyUnfrozen(const FMemoryUnfreezeContent& Context, void* Dst) const
-	{
-		static const bool bSupportsFreezeMemoryImage = TAllocatorTraits<Allocator>::SupportsFreezeMemoryImage;
-		TSupportsFreezeMemoryImageHelper<bSupportsFreezeMemoryImage>::CopyUnfrozen(Context, *this, Dst);
 	}
 };
 
diff --git a/Engine/Source/Runtime/Core/Public/Containers/Map.h b/Engine/Source/Runtime/Core/Public/Containers/Map.h
--- a/Engine/Source/Runtime/Core/Public/Containers/Map.h
+++ b/Engine/Source/Runtime/Core/Public/Containers/Map.h
@@ -136,8 +136,6 @@ class TMapBase
 	friend class TMapBase;
 
 public:
-	static constexpr bool SupportsFreezeMemoryImage = TAllocatorTraits<SetAllocator>::SupportsFreezeMemoryImage;
-
 	typedef typename TTypeTraits<KeyType  >::ConstPointerType KeyConstPointerType;
 	typedef typename TTypeTraits<KeyType  >::ConstInitType    KeyInitType;
 	typedef typename TTypeTraits<ValueType>::ConstInitType    ValueInitType;
diff --git a/Engine/Source/Runtime/Core/Public/Containers/Set.h b/Engine/Source/Runtime/Core/Public/Containers/Set.h
--- a/Engine/Source/Runtime/Core/Public/Containers/Set.h
+++ b/Engine/Source/Runtime/Core/Public/Containers/Set.h
@@ -274,8 +274,6 @@ template<
 class TSet
 {
 public:
-	static constexpr bool SupportsFreezeMemoryImage = TAllocatorTraits<Allocator>::SupportsFreezeMemoryImage && THasTypeLayout<InElementType>::Value;
-
 	typedef InElementType ElementType;
 	typedef KeyFuncs    KeyFuncsType;
 	typedef Allocator   AllocatorType;
@@ -1337,45 +1335,38 @@ private:
 	mutable HashType Hash;
 	mutable int32	 HashSize;
 
-	template<bool bFreezeMemoryImage, typename Dummy=void>
-	struct TSupportsFreezeMemoryImageHelper
-	{
-		static void WriteMemoryImage(FMemoryImageWriter& Writer, const TSet&) { Writer.WriteBytes(TSet()); }
-		static void CopyUnfrozen(const FMemoryUnfreezeContent& Context, const TSet&, void* Dst) { new(Dst) TSet(); }
-	};
-
-	template<typename Dummy>
-	struct TSupportsFreezeMemoryImageHelper<true, Dummy>
+public:
+	void WriteMemoryImage(FMemoryImageWriter& Writer) const
 	{
-		static void WriteMemoryImage(FMemoryImageWriter& Writer, const TSet& Object)
+		checkf(!Writer.Is32BitTarget(), TEXT("TSet does not currently support freezing for 32bits"));
+		if constexpr (TAllocatorTraits<Allocator>::SupportsFreezeMemoryImage && THasTypeLayout<InElementType>::Value)
 		{
-			Object.Elements.WriteMemoryImage(Writer);
-			Object.Hash.WriteMemoryImage(Writer, StaticGetTypeLayoutDesc<FSetElementId>(), Object.HashSize);
-			Writer.WriteBytes(Object.HashSize);
+			this->Elements.WriteMemoryImage(Writer);
+			this->Hash.WriteMemoryImage(Writer, StaticGetTypeLayoutDesc<FSetElementId>(), this->HashSize);
+			Writer.WriteBytes(this->HashSize);
 		}
-
-		static void CopyUnfrozen(const FMemoryUnfreezeContent& Context, const TSet& Object, void* Dst)
+		else
 		{
-			TSet* DstObject = static_cast<TSet*>(Dst);
-			Object.Elements.CopyUnfrozen(Context, &DstObject->Elements);
-
-			new(&DstObject->Hash) HashType();
-			DstObject->Hash.ResizeAllocation(0, Object.HashSize, sizeof(FSetElementId));
-			FMemory::Memcpy(DstObject->Hash.GetAllocation(), Object.Hash.GetAllocation(), sizeof(FSetElementId) * Object.HashSize);
-			DstObject->HashSize = Object.HashSize;
+			Writer.WriteBytes(TSet());
 		}
-	};
-
-public:
-	void WriteMemoryImage(FMemoryImageWriter& Writer) const
-	{
-		checkf(!Writer.Is32BitTarget(), TEXT("TSet does not currently support freezing for 32bits"));
-		TSupportsFreezeMemoryImageHelper<SupportsFreezeMemoryImage>::WriteMemoryImage(Writer, *this);
 	}
 
 	void CopyUnfrozen(const FMemoryUnfreezeContent& Context, void* Dst) const
 	{
-		TSupportsFreezeMemoryImageHelper<SupportsFreezeMemoryImage>::CopyUnfrozen(Context, *this, Dst);
+		if constexpr (TAllocatorTraits<Allocator>::SupportsFreezeMemoryImage && THasTypeLayout<InElementType>::Value)
+		{
+			TSet* DstObject = static_cast<TSet*>(Dst);
+			this->Elements.CopyUnfrozen(Context, &DstObject->Elements);
+
+			new(&DstObject->Hash) HashType();
+			DstObject->Hash.ResizeAllocation(0, this->HashSize, sizeof(FSetElementId));
+			FMemory::Memcpy(DstObject->Hash.GetAllocation(), this->Hash.GetAllocation(), sizeof(FSetElementId) * this->HashSize);
+			DstObject->HashSize = this->HashSize;
+		}
+		else
+		{
+			new(Dst) TSet();
+		}
 	}
 
 	static void AppendHash(const FPlatformTypeLayoutParameters& LayoutParams, FSHA1& Hasher)
diff --git a/Engine/Source/Runtime/Core/Public/Containers/SparseArray.h b/Engine/Source/Runtime/Core/Public/Containers/SparseArray.h
--- a/Engine/Source/Runtime/Core/Public/Containers/SparseArray.h
+++ b/Engine/Source/Runtime/Core/Public/Containers/SparseArray.h
@@ -79,8 +79,6 @@ class TSparseArray
 	friend class TScriptSparseArray;
 
 public:
-	static constexpr bool SupportsFreezeMemoryImage = TAllocatorTraits<Allocator>::SupportsFreezeMemoryImage && THasTypeLayout<InElementType>::Value;
-
 	/** Destructor. */
 	~TSparseArray()
 	{
@@ -1071,30 +1069,23 @@ private:
 	/** The number of elements in the free list. */
 	int32 NumFreeIndices;
 
-	template<bool bFreezeMemoryImage, typename Dummy=void>
-	struct TSupportsFreezeMemoryImageHelper
-	{
-		static void WriteMemoryImage(FMemoryImageWriter& Writer, const TSparseArray&) { Writer.WriteBytes(TSparseArray()); }
-		static void CopyUnfrozen(const FMemoryUnfreezeContent& Context, const TSparseArray&, void* Dst) { new(Dst) TSparseArray(); }
-		static void AppendHash(const FPlatformTypeLayoutParameters& LayoutParams, FSHA1& Hasher) {}
-	};
-
-	template<typename Dummy>
-	struct TSupportsFreezeMemoryImageHelper<true, Dummy>
+public:
+	void WriteMemoryImage(FMemoryImageWriter& Writer) const
 	{
-		static void WriteMemoryImage(FMemoryImageWriter& Writer, const TSparseArray& Object)
+		checkf(!Writer.Is32BitTarget(), TEXT("TSparseArray does not currently support freezing for 32bits"));
+		if constexpr (TAllocatorTraits<Allocator>::SupportsFreezeMemoryImage && THasTypeLayout<ElementType>::Value)
 		{
 			// Write Data
-			const int32 NumElements = Object.Data.Num();
+			const int32 NumElements = this->Data.Num();
 			if (NumElements > 0)
 			{
 				const FTypeLayoutDesc& ElementTypeDesc = StaticGetTypeLayoutDesc<ElementType>();
 				FMemoryImageWriter ArrayWriter = Writer.WritePointer(ElementTypeDesc);
 				for (int32 i = 0; i < NumElements; ++i)
 				{
-					const FElementOrFreeListLink& Elem = Object.Data[i];
+					const FElementOrFreeListLink& Elem = this->Data[i];
 					const uint32 StartOffset = ArrayWriter.WriteAlignment<FElementOrFreeListLink>();
-					if (Object.AllocationFlags[i])
+					if (this->AllocationFlags[i])
 					{
 						ArrayWriter.WriteObject(&Elem.ElementData, ElementTypeDesc);
 					}
@@ -1114,23 +1105,30 @@ private:
 			Writer.WriteBytes(NumElements);
 
 			//
-			Object.AllocationFlags.WriteMemoryImage(Writer);
-			Writer.WriteBytes(Object.FirstFreeIndex);
-			Writer.WriteBytes(Object.NumFreeIndices);
+			this->AllocationFlags.WriteMemoryImage(Writer);
+			Writer.WriteBytes(this->FirstFreeIndex);
+			Writer.WriteBytes(this->NumFreeIndices);
 		}
+		else
+		{
+			Writer.WriteBytes(TSparseArray());
+		}
+	}
 
-		static void CopyUnfrozen(const FMemoryUnfreezeContent& Context, const TSparseArray& Object, void* Dst)
+	void CopyUnfrozen(const FMemoryUnfreezeContent& Context, void* Dst) const
+	{
+		if constexpr (TAllocatorTraits<Allocator>::SupportsFreezeMemoryImage && THasTypeLayout<ElementType>::Value)
 		{
 			const FTypeLayoutDesc& ElementTypeDesc = StaticGetTypeLayoutDesc<ElementType>();
 			TSparseArray* DstObject = (TSparseArray*)Dst;
 			{
 				new(&DstObject->Data) DataType();
-				DstObject->Data.SetNumUninitialized(Object.Data.Num());
-				for (int32 i = 0; i < Object.Data.Num(); ++i)
+				DstObject->Data.SetNumUninitialized(this->Data.Num());
+				for (int32 i = 0; i < this->Data.Num(); ++i)
 				{
-					const FElementOrFreeListLink& Elem = Object.Data[i];
+					const FElementOrFreeListLink& Elem = this->Data[i];
 					FElementOrFreeListLink& DstElem = DstObject->Data[i];
-					if (Object.AllocationFlags[i])
+					if (this->AllocationFlags[i])
 					{
 						Context.UnfreezeObject(&Elem.ElementData, ElementTypeDesc, &DstElem.ElementData);
 					}
@@ -1142,32 +1140,22 @@ private:
 				}
 			}
 
-			new(&DstObject->AllocationFlags) AllocationBitArrayType(Object.AllocationFlags);
-			DstObject->FirstFreeIndex = Object.FirstFreeIndex;
-			DstObject->NumFreeIndices = Object.NumFreeIndices;
+			new(&DstObject->AllocationFlags) AllocationBitArrayType(this->AllocationFlags);
+			DstObject->FirstFreeIndex = this->FirstFreeIndex;
+			DstObject->NumFreeIndices = this->NumFreeIndices;
 		}
-
-		static void AppendHash(const FPlatformTypeLayoutParameters& LayoutParams, FSHA1& Hasher)
+		else
 		{
-			Freeze::AppendHash(StaticGetTypeLayoutDesc<ElementType>(), LayoutParams, Hasher);
+			new(Dst) TSparseArray();
 		}
-	};
-
-public:
-	void WriteMemoryImage(FMemoryImageWriter& Writer) const
-	{
-		checkf(!Writer.Is32BitTarget(), TEXT("TSparseArray does not currently support freezing for 32bits"));
-		TSupportsFreezeMemoryImageHelper<SupportsFreezeMemoryImage>::WriteMemoryImage(Writer, *this);
-	}
-
-	void CopyUnfrozen(const FMemoryUnfreezeContent& Context, void* Dst) const
-	{
-		TSupportsFreezeMemoryImageHelper<SupportsFreezeMemoryImage>::CopyUnfrozen(Context, *this, Dst);
 	}
 
 	static void AppendHash(const FPlatformTypeLayoutParameters& LayoutParams, FSHA1& Hasher)
 	{
-		TSupportsFreezeMemoryImageHelper<SupportsFreezeMemoryImage>::AppendHash(LayoutParams, Hasher);
+		if constexpr (TAllocatorTraits<Allocator>::SupportsFreezeMemoryImage && THasTypeLayout<ElementType>::Value)
+		{
+			Freeze::AppendHash(StaticGetTypeLayoutDesc<ElementType>(), LayoutParams, Hasher);
+		}
 	}
 };
 
