Applied the same SupportsFreezeMemoryImage changes to TArray as to the other containers as in CL# 23153171.

diff --git a/Engine/Source/Runtime/Core/Public/Containers/Array.h b/Engine/Source/Runtime/Core/Public/Containers/Array.h
--- a/Engine/Source/Runtime/Core/Public/Containers/Array.h
+++ b/Engine/Source/Runtime/Core/Public/Containers/Array.h
@@ -525,41 +525,41 @@ private:
 		{
 			// Move
 
-			static_assert(std::is_same_v<TArray, ToArrayType>, "MoveOrCopy is expected to be called with the current array type as the destination");
+		static_assert(std::is_same_v<TArray, ToArrayType>, "MoveOrCopy is expected to be called with the current array type as the destination");
 
-			using FromAllocatorType = typename FromArrayType::AllocatorType;
-			using ToAllocatorType = typename ToArrayType::AllocatorType;
+		using FromAllocatorType = typename FromArrayType::AllocatorType;
+		using ToAllocatorType   = typename ToArrayType::AllocatorType;
 
-			if constexpr (TCanMoveBetweenAllocators<FromAllocatorType, ToAllocatorType>::Value)
-			{
-				ToArray.AllocatorInstance.template MoveToEmptyFromOtherAllocator<FromAllocatorType>(FromArray.AllocatorInstance);
-			}
-			else
-			{
-				ToArray.AllocatorInstance.MoveToEmpty(FromArray.AllocatorInstance);
-			}
+		if constexpr (TCanMoveBetweenAllocators<FromAllocatorType, ToAllocatorType>::Value)
+		{
+			ToArray.AllocatorInstance.template MoveToEmptyFromOtherAllocator<FromAllocatorType>(FromArray.AllocatorInstance);
+		}
+		else
+		{
+			ToArray.AllocatorInstance.MoveToEmpty(FromArray.AllocatorInstance);
+		}
 
-			ToArray.ArrayNum = (SizeType)FromArray.ArrayNum;
-			ToArray.ArrayMax = (SizeType)FromArray.ArrayMax;
+		ToArray  .ArrayNum = (SizeType)FromArray.ArrayNum;
+		ToArray  .ArrayMax = (SizeType)FromArray.ArrayMax;
 
-			// Ensure the destination container could hold the source range (when the allocator size types shrink)
-			if constexpr (sizeof(USizeType) < sizeof(typename FromArrayType::USizeType))
+		// Ensure the destination container could hold the source range (when the allocator size types shrink)
+		if constexpr (sizeof(USizeType) < sizeof(typename FromArrayType::USizeType))
+		{
+			if (ToArray.ArrayNum != FromArray.ArrayNum || ToArray.ArrayMax != FromArray.ArrayMax)
 			{
-				if (ToArray.ArrayNum != FromArray.ArrayNum || ToArray.ArrayMax != FromArray.ArrayMax)
-				{
-					OnInvalidNum((USizeType)ToArray.ArrayNum);
-				}
+				OnInvalidNum((USizeType)ToArray.ArrayNum);
 			}
-
-			FromArray.ArrayNum = 0;
-			FromArray.ArrayMax = FromArray.AllocatorInstance.GetInitialCapacity();
 		}
+
+		FromArray.ArrayNum = 0;
+		FromArray.ArrayMax = FromArray.AllocatorInstance.GetInitialCapacity();
+	}
 		else
 		{
 			// Copy
 
-			ToArray.CopyToEmpty(FromArray.GetData(), FromArray.Num(), PrevMax);
-		}
+		ToArray.CopyToEmpty(FromArray.GetData(), FromArray.Num(), PrevMax);
+	}
 	}
 
 	/**
@@ -578,25 +578,25 @@ private:
 		{
 			// Move
 
-			MoveOrCopy(ToArray, FromArray, PrevMax);
-
-			USizeType LocalArrayNum = (USizeType)ToArray.ArrayNum;
-			USizeType NewMax = (USizeType)LocalArrayNum + (USizeType)ExtraSlack;
+		MoveOrCopy(ToArray, FromArray, PrevMax);
 
-			// This should only happen when we've underflowed or overflowed SizeType
-			if ((SizeType)NewMax < LocalArrayNum)
-			{
-				OnInvalidNum((USizeType)ExtraSlack);
-			}
+		USizeType LocalArrayNum = (USizeType)ToArray.ArrayNum;
+		USizeType NewMax        = (USizeType)LocalArrayNum + (USizeType)ExtraSlack;
 
-			ToArray.Reserve(NewMax);
+		// This should only happen when we've underflowed or overflowed SizeType
+		if ((SizeType)NewMax < LocalArrayNum)
+		{
+			OnInvalidNum((USizeType)ExtraSlack);
 		}
+
+		ToArray.Reserve(NewMax);
+	}
 		else
 		{
 			// Copy
 
-			ToArray.CopyToEmptyWithSlack(FromArray.GetData(), FromArray.Num(), PrevMax, ExtraSlack);
-		}
+		ToArray.CopyToEmptyWithSlack(FromArray.GetData(), FromArray.Num(), PrevMax, ExtraSlack);
+	}
 	}
 
 public:
@@ -2986,12 +2986,12 @@ private:
 		ArrayNum = NewNum;
 		if (OtherNum || PrevMax)
 		{
-			ResizeForCopy(NewNum, PrevMax);
-			ConstructItems<ElementType>(GetData(), OtherData, OtherNum);
+		ResizeForCopy(NewNum, PrevMax);
+		ConstructItems<ElementType>(GetData(), OtherData, OtherNum);
 		}
 		else
 		{
-			ArrayMax = AllocatorInstance.GetInitialCapacity();
+		ArrayMax = AllocatorInstance.GetInitialCapacity();
 		}
 	}
 
@@ -3039,71 +3039,52 @@ protected:
 	SizeType             ArrayNum;
 	SizeType             ArrayMax;
 
-private:
-	template<bool bFreezeMemoryImage, typename Dummy=void>
-	struct TSupportsFreezeMemoryImageHelper
+public:
+	void WriteMemoryImage(FMemoryImageWriter& Writer) const
 	{
-		static void WriteMemoryImage(FMemoryImageWriter& Writer, const TArray&)
+		if constexpr (TAllocatorTraits<AllocatorType>::SupportsFreezeMemoryImage && THasTypeLayout<ElementType>::Value)
+		{
+			this->AllocatorInstance.WriteMemoryImage(Writer, StaticGetTypeLayoutDesc<ElementType>(), this->ArrayNum);
+			Writer.WriteBytes(this->ArrayNum);
+			Writer.WriteBytes(this->ArrayNum);
+		}
+		else
 		{
 			// Writing non-freezable TArray is only supported for 64-bit target for now
 			// Would need complete layout macros for all allocator types in order to properly write (empty) 32bit versions
 			check(Writer.Is64BitTarget());
 			Writer.WriteBytes(TArray());
 		}
+	}
 
-		static void CopyUnfrozen(const FMemoryUnfreezeContent& Context, const TArray&, void* Dst) { new(Dst) TArray(); }
-		static void AppendHash(const FPlatformTypeLayoutParameters& LayoutParams, FSHA1& Hasher) {}
-		static void ToString(const FPlatformTypeLayoutParameters& LayoutParams, FMemoryToStringContext& OutContext, const TArray& Object) {}
-	};
-
-	template<typename Dummy>
-	struct TSupportsFreezeMemoryImageHelper<true, Dummy>
+	void CopyUnfrozen(const FMemoryUnfreezeContent& Context, void* Dst) const
 	{
-		static void WriteMemoryImage(FMemoryImageWriter& Writer, const TArray& Object)
-		{
-			Object.AllocatorInstance.WriteMemoryImage(Writer, StaticGetTypeLayoutDesc<ElementType>(), Object.ArrayNum);
-			Writer.WriteBytes(Object.ArrayNum);
-			Writer.WriteBytes(Object.ArrayNum);
-		}
-		static void CopyUnfrozen(const FMemoryUnfreezeContent& Context, const TArray& Object, void* Dst)
+		if constexpr (TAllocatorTraits<AllocatorType>::SupportsFreezeMemoryImage && THasTypeLayout<ElementType>::Value)
 		{
 			TArray* DstArray = new(Dst) TArray();
-			DstArray->SetNumZeroed(Object.ArrayNum);
-			Object.AllocatorInstance.CopyUnfrozen(Context, StaticGetTypeLayoutDesc<ElementType>(), Object.ArrayNum, DstArray->GetData());
+			DstArray->SetNumZeroed(this->ArrayNum);
+			this->AllocatorInstance.CopyUnfrozen(Context, StaticGetTypeLayoutDesc<ElementType>(), this->ArrayNum, DstArray->GetData());
 		}
-		static void AppendHash(const FPlatformTypeLayoutParameters& LayoutParams, FSHA1& Hasher)
-		{
-			Freeze::AppendHash(StaticGetTypeLayoutDesc<ElementType>(), LayoutParams, Hasher);
-		}
-		static void ToString(const FPlatformTypeLayoutParameters& LayoutParams, FMemoryToStringContext& OutContext, const TArray& Object)
+		else
 		{
-			Object.AllocatorInstance.ToString(StaticGetTypeLayoutDesc<ElementType>(), Object.ArrayNum, Object.ArrayMax, LayoutParams, OutContext);
+			new(Dst) TArray();
 		}
-	};
-
-public:
-	void WriteMemoryImage(FMemoryImageWriter& Writer) const
-	{
-		static constexpr bool bSupportsFreezeMemoryImage = TAllocatorTraits<AllocatorType>::SupportsFreezeMemoryImage && THasTypeLayout<ElementType>::Value;
-		TSupportsFreezeMemoryImageHelper<bSupportsFreezeMemoryImage>::WriteMemoryImage(Writer, *this);
-	}
-
-	void CopyUnfrozen(const FMemoryUnfreezeContent& Context, void* Dst) const
-	{
-		static constexpr bool bSupportsFreezeMemoryImage = TAllocatorTraits<AllocatorType>::SupportsFreezeMemoryImage && THasTypeLayout<ElementType>::Value;
-		TSupportsFreezeMemoryImageHelper<bSupportsFreezeMemoryImage>::CopyUnfrozen(Context, *this, Dst);
 	}
 
 	static void AppendHash(const FPlatformTypeLayoutParameters& LayoutParams, FSHA1& Hasher)
 	{
-		static constexpr bool bSupportsFreezeMemoryImage = TAllocatorTraits<AllocatorType>::SupportsFreezeMemoryImage && THasTypeLayout<ElementType>::Value;
-		TSupportsFreezeMemoryImageHelper<bSupportsFreezeMemoryImage>::AppendHash(LayoutParams, Hasher);
+		if constexpr (TAllocatorTraits<AllocatorType>::SupportsFreezeMemoryImage && THasTypeLayout<ElementType>::Value)
+		{
+			Freeze::AppendHash(StaticGetTypeLayoutDesc<ElementType>(), LayoutParams, Hasher);
+		}
 	}
 
 	void ToString(const FPlatformTypeLayoutParameters& LayoutParams, FMemoryToStringContext& OutContext) const
 	{
-		static constexpr bool bSupportsFreezeMemoryImage = TAllocatorTraits<AllocatorType>::SupportsFreezeMemoryImage && THasTypeLayout<ElementType>::Value;
-		TSupportsFreezeMemoryImageHelper<bSupportsFreezeMemoryImage>::ToString(LayoutParams, OutContext, *this);
+		if constexpr (TAllocatorTraits<AllocatorType>::SupportsFreezeMemoryImage && THasTypeLayout<ElementType>::Value)
+		{
+			this->AllocatorInstance.ToString(StaticGetTypeLayoutDesc<ElementType>(), this->ArrayNum, this->ArrayMax, LayoutParams, OutContext);
+		}
 	}
 
 	/**
