Split the FText comparison API out of FText so that it can be used for FString

FTextComparison now exists to provide static access to the collated comparison functions of FText, without needing to actually create an FText to do the comparison.

This follows the pattern of things like FTextTransformer and FTextChronoFormatter.

diff --git a/Engine/Source/Runtime/Core/Private/Internationalization/ICUText.cpp b/Engine/Source/Runtime/Core/Private/Internationalization/ICUText.cpp
--- a/Engine/Source/Runtime/Core/Private/Internationalization/ICUText.cpp
+++ b/Engine/Source/Runtime/Core/Private/Internationalization/ICUText.cpp
@@ -96,39 +96,41 @@ bool FText::IsWhitespace(const TCHAR Char)
 	return u_isWhitespace(ICUChar) != 0;
 }
 
-int32 FText::CompareTo( const FText& Other, const ETextComparisonLevel::Type ComparisonLevel ) const
+int32 FTextComparison::CompareTo( const FString& A, const FString& B, const ETextComparisonLevel::Type ComparisonLevel )
 {
 	const TSharedRef<const icu::Collator, ESPMode::ThreadSafe> Collator( FInternationalization::Get().GetCurrentLanguage()->Implementation->GetCollator(ComparisonLevel) );
 
-	// Create an iterator for 'this' so that we can interface with ICU
-	UCharIterator DisplayStringICUIterator;
-	FICUTextCharacterIterator DisplayStringIterator(FStringView(this->ToString()));
-	uiter_setCharacterIterator(&DisplayStringICUIterator, &DisplayStringIterator);
+	// Create an iterator for 'A' so that we can interface with ICU
+	FStringView AView = A;
+	UCharIterator ADisplayStringICUIterator;
+	FICUTextCharacterIterator ADisplayStringIterator(AView);
+	uiter_setCharacterIterator(&ADisplayStringICUIterator, &ADisplayStringIterator);
 
-	// Create an iterator for 'Other' so that we can interface with ICU
-	UCharIterator OtherDisplayStringICUIterator;
-	FICUTextCharacterIterator OtherDisplayStringIterator(FStringView(Other.ToString()));
-	uiter_setCharacterIterator(&OtherDisplayStringICUIterator, &OtherDisplayStringIterator);
+	// Create an iterator for 'B' so that we can interface with ICU
+	FStringView BView = B;
+	UCharIterator BDisplayStringICUIterator;
+	FICUTextCharacterIterator BDisplayStringIterator(BView);
+	uiter_setCharacterIterator(&BDisplayStringICUIterator, &BDisplayStringIterator);
 
 	UErrorCode ICUStatus = U_ZERO_ERROR;
-	const UCollationResult Result = Collator->compare(DisplayStringICUIterator, OtherDisplayStringICUIterator, ICUStatus);
+	const UCollationResult Result = Collator->compare(ADisplayStringICUIterator, BDisplayStringICUIterator, ICUStatus);
 
 	return Result;
 }
 
-int32 FText::CompareToCaseIgnored( const FText& Other ) const
+int32 FTextComparison::CompareToCaseIgnored( const FString& A, const FString& B )
 {
-	return CompareTo(Other, ETextComparisonLevel::Secondary);
+	return CompareTo(A, B, ETextComparisonLevel::Secondary);
 }
 
-bool FText::EqualTo( const FText& Other, const ETextComparisonLevel::Type ComparisonLevel ) const
+bool FTextComparison::EqualTo( const FString& A, const FString& B, const ETextComparisonLevel::Type ComparisonLevel )
 {
-	return CompareTo(Other, ComparisonLevel) == 0;
+	return CompareTo(A, B, ComparisonLevel) == 0;
 }
 
-bool FText::EqualToCaseIgnored( const FText& Other ) const
+bool FTextComparison::EqualToCaseIgnored( const FString& A, const FString& B )
 {
-	return EqualTo(Other, ETextComparisonLevel::Secondary);
+	return EqualTo(A, B, ETextComparisonLevel::Secondary);
 }
 
 class FText::FSortPredicate::FSortPredicateImplementation
diff --git a/Engine/Source/Runtime/Core/Private/Internationalization/LegacyText.cpp b/Engine/Source/Runtime/Core/Private/Internationalization/LegacyText.cpp
--- a/Engine/Source/Runtime/Core/Private/Internationalization/LegacyText.cpp
+++ b/Engine/Source/Runtime/Core/Private/Internationalization/LegacyText.cpp
@@ -50,24 +50,24 @@ bool FText::IsWhitespace(const TCHAR Char)
 	return FChar::IsWhitespace(Char);
 }
 
-int32 FText::CompareTo( const FText& Other, const ETextComparisonLevel::Type ComparisonLevel ) const
+int32 FTextComparison::CompareTo( const FString& A, const FString& B, const ETextComparisonLevel::Type ComparisonLevel )
 {
-	return ToString().Compare(Other.ToString(), ESearchCase::CaseSensitive);
+	return A.Compare(B, ESearchCase::CaseSensitive);
 }
 
-int32 FText::CompareToCaseIgnored( const FText& Other ) const
+int32 FTextComparison::CompareToCaseIgnored( const FString& A, const FString& B )
 {
-	return ToString().Compare(Other.ToString(), ESearchCase::IgnoreCase);
+	return A.Compare(B, ESearchCase::IgnoreCase);
 }
 
-bool FText::EqualTo( const FText& Other, const ETextComparisonLevel::Type ComparisonLevel ) const
+bool FTextComparison::EqualTo( const FString& A, const FString& B, const ETextComparisonLevel::Type ComparisonLevel )
 {
-	return ToString().Equals(Other.ToString(), ESearchCase::CaseSensitive);
+	return A.Equals(B, ESearchCase::CaseSensitive);
 }
 
-bool FText::EqualToCaseIgnored( const FText& Other ) const
+bool FTextComparison::EqualToCaseIgnored( const FString& A, const FString& B )
 {
-	return ToString().Equals(Other.ToString(), ESearchCase::IgnoreCase);
+	return A.Equals(B, ESearchCase::IgnoreCase);
 }
 
 FText::FSortPredicate::FSortPredicate(const ETextComparisonLevel::Type ComparisonLevel)
diff --git a/Engine/Source/Runtime/Core/Private/Internationalization/Text.cpp b/Engine/Source/Runtime/Core/Private/Internationalization/Text.cpp
--- a/Engine/Source/Runtime/Core/Private/Internationalization/Text.cpp
+++ b/Engine/Source/Runtime/Core/Private/Internationalization/Text.cpp
@@ -1129,6 +1129,26 @@ bool FText::IsNumeric() const
 	return ToString().IsNumeric();
 }
 
+int32 FText::CompareTo(const FText& Other, const ETextComparisonLevel::Type ComparisonLevel) const
+{
+	return FTextComparison::CompareTo(ToString(), Other.ToString(), ComparisonLevel);
+}
+
+int32 FText::CompareToCaseIgnored(const FText& Other) const
+{
+	return FTextComparison::CompareToCaseIgnored(ToString(), Other.ToString());
+}
+
+bool FText::EqualTo(const FText& Other, const ETextComparisonLevel::Type ComparisonLevel) const
+{
+	return FTextComparison::EqualTo(ToString(), Other.ToString(), ComparisonLevel);
+}
+
+bool FText::EqualToCaseIgnored(const FText& Other) const
+{
+	return FTextComparison::EqualToCaseIgnored(ToString(), Other.ToString());
+}
+
 void FText::Rebuild() const
 {
 	TextData->GetMutableTextHistory().UpdateDisplayStringIfOutOfDate();
diff --git a/Engine/Source/Runtime/Core/Public/Internationalization/Culture.h b/Engine/Source/Runtime/Core/Public/Internationalization/Culture.h
--- a/Engine/Source/Runtime/Core/Public/Internationalization/Culture.h
+++ b/Engine/Source/Runtime/Core/Public/Internationalization/Culture.h
@@ -24,6 +24,7 @@ class CORE_API FCulture
 #if UE_ENABLE_ICU
 	friend class FText;
 	friend class FTextChronoFormatter;
+	friend class FTextComparison;
 	friend class FICUBreakIteratorManager;
 #endif
 
diff --git a/Engine/Source/Runtime/Core/Public/Internationalization/Text.h b/Engine/Source/Runtime/Core/Public/Internationalization/Text.h
--- a/Engine/Source/Runtime/Core/Public/Internationalization/Text.h
+++ b/Engine/Source/Runtime/Core/Public/Internationalization/Text.h
@@ -14,6 +14,7 @@
 #include "Internationalization/TextKey.h"
 #include "Internationalization/LocKeyFuncs.h"
 #include "Internationalization/CulturePointer.h"
+#include "Internationalization/TextComparison.h"
 #include "Internationalization/TextLocalizationManager.h"
 #include "Internationalization/StringTableCoreFwd.h"
 #include "Internationalization/ITextData.h"
@@ -45,19 +46,6 @@ namespace ETextFlag
 	};
 }
 
-namespace ETextComparisonLevel
-{
-	enum Type
-	{
-		Default,	// Locale-specific Default
-		Primary,	// Base
-		Secondary,	// Accent
-		Tertiary,	// Case
-		Quaternary,	// Punctuation
-		Quinary		// Identical
-	};
-}
-
 enum class ETextIdenticalModeFlags : uint8
 {
 	/** No special behavior */
diff --git a/Engine/Source/Runtime/Core/Public/Internationalization/TextComparison.h b/Engine/Source/Runtime/Core/Public/Internationalization/TextComparison.h
new file mode 100644
--- /dev/null
+++ b/Engine/Source/Runtime/Core/Public/Internationalization/TextComparison.h
@@ -0,0 +1,33 @@
+// Copyright Epic Games, Inc. All Rights Reserved.
+
+#pragma once
+
+#include "CoreTypes.h"
+#include "Containers/UnrealString.h"
+
+namespace ETextComparisonLevel
+{
+	enum Type
+	{
+		Default,	// Locale-specific Default
+		Primary,	// Base
+		Secondary,	// Accent
+		Tertiary,	// Case
+		Quaternary,	// Punctuation
+		Quinary		// Identical
+	};
+}
+
+/**
+ * Utility for performing low-level text comparison.
+ * The implementation can be found in LegacyText.cpp and ICUText.cpp.
+ */
+class CORE_API FTextComparison
+{
+public:
+	static int32 CompareTo(const FString& A, const FString& B, const ETextComparisonLevel::Type ComparisonLevel = ETextComparisonLevel::Default);
+	static int32 CompareToCaseIgnored(const FString& A, const FString& B);
+
+	static bool EqualTo(const FString& A, const FString& B, const ETextComparisonLevel::Type ComparisonLevel = ETextComparisonLevel::Default);
+	static bool EqualToCaseIgnored(const FString& A, const FString& B);
+};
