Fixed a crash on exit when init fails early

If PreInit fails before setting the command line, AppExit will fail in its last step because it attempts to access the command line.

diff --git a/Engine/Source/Runtime/Core/Private/Misc/Fork.cpp b/Engine/Source/Runtime/Core/Private/Misc/Fork.cpp
--- a/Engine/Source/Runtime/Core/Private/Misc/Fork.cpp
+++ b/Engine/Source/Runtime/Core/Private/Misc/Fork.cpp
@@ -35,7 +35,7 @@ namespace ForkInternal
 	bool IsRealForkRequested()
 	{
 		// We cache the value since only the master process might receive -WaitAndFork
-		static const bool bRealForkRequested = FParse::Param(FCommandLine::Get(), TEXT("WaitAndFork"));
+		static const bool bRealForkRequested = FParse::Param(FCommandLine::IsInitialized() ? FCommandLine::Get() : TEXT(""), TEXT("WaitAndFork"));
 		return bRealForkRequested;
 	}
 
