IoDispatcher: Fixed shutdown crash

diff --git a/Engine/Source/Runtime/Core/Private/IO/IoDispatcher.cpp b/Engine/Source/Runtime/Core/Private/IO/IoDispatcher.cpp
--- a/Engine/Source/Runtime/Core/Private/IO/IoDispatcher.cpp
+++ b/Engine/Source/Runtime/Core/Private/IO/IoDispatcher.cpp
@@ -338,7 +338,7 @@ public:
 		DispatcherEvent = FPlatformProcess::GetSynchEventFromPool(false);
 		BackendContext->WakeUpDispatcherThreadDelegate.BindRaw(this, &FIoDispatcherImpl::WakeUpDispatcherThread);
 		BackendContext->bIsMultiThreaded = bInIsMultithreaded;
-		FCoreDelegates::GetMemoryTrimDelegate().AddLambda([this]()
+		MemoryTrimDelegateHandle = FCoreDelegates::GetMemoryTrimDelegate().AddLambda([this]()
 		{
 			RequestAllocator->Trim();
 			BatchAllocator.Trim();
@@ -350,6 +350,12 @@ public:
 		delete Thread;
 		FPlatformProcess::ReturnSynchEventToPool(DispatcherEvent);
 		RequestAllocator->ReleaseRef();
+		for (const TSharedRef<IIoDispatcherBackend>& Backend : Backends)
+		{
+			Backend->Shutdown();
+		}
+		FCoreDelegates::GetMemoryTrimDelegate().Remove(MemoryTrimDelegateHandle);
+		BackendContext->WakeUpDispatcherThreadDelegate.Unbind();
 	}
 
 	void Initialize()
@@ -832,6 +838,7 @@ private:
 	using FBatchAllocator = TBlockAllocator<FIoBatchImpl, 4096>;
 
 	TSharedRef<FIoDispatcherBackendContext> BackendContext;
+	FDelegateHandle MemoryTrimDelegateHandle;
 	TArray<TSharedRef<IIoDispatcherBackend>> Backends;
 	FIoRequestAllocator* RequestAllocator = nullptr;
 	FBatchAllocator BatchAllocator;
diff --git a/Engine/Source/Runtime/Core/Public/IO/IoDispatcherBackend.h b/Engine/Source/Runtime/Core/Public/IO/IoDispatcherBackend.h
--- a/Engine/Source/Runtime/Core/Public/IO/IoDispatcherBackend.h
+++ b/Engine/Source/Runtime/Core/Public/IO/IoDispatcherBackend.h
@@ -101,6 +101,7 @@ struct FIoDispatcherBackendContext
 struct IIoDispatcherBackend
 {
 	virtual void Initialize(TSharedRef<const FIoDispatcherBackendContext> Context) = 0;
+	virtual void Shutdown() {};
 	virtual bool Resolve(FIoRequestImpl* Request) = 0;
 	virtual void CancelIoRequest(FIoRequestImpl* Request) = 0;
 	virtual void UpdatePriorityForIoRequest(FIoRequestImpl* Request) = 0;
diff --git a/Engine/Source/Runtime/PakFile/Private/IoDispatcherFileBackend.cpp b/Engine/Source/Runtime/PakFile/Private/IoDispatcherFileBackend.cpp
--- a/Engine/Source/Runtime/PakFile/Private/IoDispatcherFileBackend.cpp
+++ b/Engine/Source/Runtime/PakFile/Private/IoDispatcherFileBackend.cpp
@@ -1215,7 +1215,7 @@ FFileIoStore::FFileIoStore(TUniquePtr<IPlatformFileIoStore>&& InPlatformImpl)
 
 FFileIoStore::~FFileIoStore()
 {
-	delete Thread;
+	StopThread();
 }
 
 void FFileIoStore::Initialize(TSharedRef<const FIoDispatcherBackendContext> InContext)
@@ -1254,6 +1254,20 @@ void FFileIoStore::Initialize(TSharedRef<const FIoDispatcherBackendContext> InCo
 	Thread = FRunnableThread::Create(this, TEXT("IoService"), 0, TPri_AboveNormal);
 }
 
+void FFileIoStore::StopThread()
+{
+	if (Thread)
+	{
+		delete Thread;
+		Thread = nullptr;
+	}
+}
+
+void FFileIoStore::Shutdown()
+{
+	StopThread();
+}
+
 TIoStatusOr<FIoContainerHeader> FFileIoStore::Mount(const TCHAR* InTocPath, int32 Order, const FGuid& EncryptionKeyGuid, const FAES::FAESKey& EncryptionKey)
 {
 	TUniquePtr<FFileIoStoreReader> Reader(new FFileIoStoreReader(*PlatformImpl, Stats));
diff --git a/Engine/Source/Runtime/PakFile/Private/IoDispatcherFileBackend.h b/Engine/Source/Runtime/PakFile/Private/IoDispatcherFileBackend.h
--- a/Engine/Source/Runtime/PakFile/Private/IoDispatcherFileBackend.h
+++ b/Engine/Source/Runtime/PakFile/Private/IoDispatcherFileBackend.h
@@ -106,6 +106,7 @@ public:
 	FFileIoStore(TUniquePtr<IPlatformFileIoStore>&& PlatformImpl);
 	~FFileIoStore();
 	void Initialize(TSharedRef<const FIoDispatcherBackendContext> Context) override;
+	void Shutdown() override;
 	TIoStatusOr<FIoContainerHeader> Mount(const TCHAR* InTocPath, int32 Order, const FGuid& EncryptionKeyGuid, const FAES::FAESKey& EncryptionKey);
 	bool Unmount(const TCHAR* InTocPath);
 	bool Resolve(FIoRequestImpl* Request) override;
@@ -164,6 +165,7 @@ private:
 	void ScatterBlock(FFileIoStoreCompressedBlock* CompressedBlock, bool bIsAsync);
 	void CompleteDispatcherRequest(FFileIoStoreResolvedRequest* ResolvedRequest);
 	void FinalizeCompressedBlock(FFileIoStoreCompressedBlock* CompressedBlock);
+	void StopThread();
 
 	uint64 ReadBufferSize = 0;
 	TSharedPtr<const FIoDispatcherBackendContext> BackendContext;
