Fix memory leak in DepletableQueue as the last node was never deleted and make an accidental concurrent call to Deplete not catastrophic.
Adding a ConsumeAllQueue which is based on a Stack because the Queues have livelock issues and the Stack provides a robust alternative.

diff --git a/Engine/Source/Runtime/Core/Public/Containers/ClosableMpscQueue.h b/Engine/Source/Runtime/Core/Public/Containers/ClosableMpscQueue.h
--- a/Engine/Source/Runtime/Core/Public/Containers/ClosableMpscQueue.h
+++ b/Engine/Source/Runtime/Core/Public/Containers/ClosableMpscQueue.h
@@ -116,7 +116,7 @@ private:
 			do
 			{
 				Next = Node->Next.load(std::memory_order_relaxed);
-			} while (Next == nullptr);
+			} while (Next == nullptr); // <- This loop has the potential for live locking if enqueue was not completed (e.g was running at lower priority)
 
 			return Next;
 		};
diff --git a/Engine/Source/Runtime/Core/Public/Containers/ConsumeAllMpmcQueue.h b/Engine/Source/Runtime/Core/Public/Containers/ConsumeAllMpmcQueue.h
new file mode 100644
--- /dev/null
+++ b/Engine/Source/Runtime/Core/Public/Containers/ConsumeAllMpmcQueue.h
@@ -0,0 +1,121 @@
+// Copyright Epic Games, Inc. All Rights Reserved.
+
+#pragma once
+
+#include "CoreTypes.h"
+#include "HAL/UnrealMemory.h"
+#include "Templates/TypeCompatibleBytes.h"
+#include <atomic>
+#include <type_traits>
+
+namespace UE
+{
+	enum class EConsumeAllMpmcQueueResult
+	{
+		HadItems,
+		WasEmpty,
+	};
+
+	/**
+	 * Multi-producer/multi-consumer unbounded concurrent queue (implemented as a Stack) that is atomically consumed
+	 * and is reset to its default empty state. Validated and run though atomic race detector.
+	 */
+	template<typename T, typename AllocatorType = FMemory>
+	class TConsumeAllMpmcQueue final
+	{
+	private:
+		struct FNode
+		{
+			std::atomic<FNode*> Next{ nullptr };
+			TTypeCompatibleBytes<T> Item;
+		};
+
+		std::atomic<FNode*> Head{ nullptr };
+
+	public:
+		UE_NONCOPYABLE(TConsumeAllMpmcQueue);
+
+		TConsumeAllMpmcQueue() = default;
+
+		~TConsumeAllMpmcQueue()
+		{
+			static_assert(std::is_trivially_destructible_v<FNode>);
+			ConsumeAllLifo([](T&&){});
+		}
+
+		//Push an Item to the Queue and 
+		//returns EConsumeAllMpmcQueueResult::WasEmpty if the Queue was empty before
+		//or EConsumeAllMpmcQueueResult::HadItems if there were already items in it
+		template <typename... ArgTypes>
+		EConsumeAllMpmcQueueResult ProduceItem(ArgTypes&&... Args)
+		{
+			FNode* New = new(AllocatorType::Malloc(sizeof(FNode), alignof(FNode))) FNode;
+			new (&New->Item) T(Forward<ArgTypes>(Args)...);
+
+			//Atomically append to the top of the Queue
+			FNode* Prev = Head.load(std::memory_order_relaxed);
+			do
+			{
+				New->Next.store(Prev, std::memory_order_relaxed);
+			} while (!Head.compare_exchange_weak(Prev, New, std::memory_order_acq_rel, std::memory_order_relaxed));
+			
+			return Prev == nullptr ? EConsumeAllMpmcQueueResult::WasEmpty : EConsumeAllMpmcQueueResult::HadItems;
+		}
+
+		//Take all items off the Queue atomically and consumes them in LIFO order
+		//returns EConsumeAllMpmcQueueResult::WasEmpty if the Queue was empty
+		//or EConsumeAllMpmcQueueResult::HadItems if there were items to consume
+		template<typename F>
+		EConsumeAllMpmcQueueResult ConsumeAllLifo(const F& Consumer)
+		{
+			return ConsumeAll<false>(Consumer);
+		}
+
+		//Take all items off the Queue atomically and consumes them in FIFO order
+		//at the cost of the reversing the Links once
+		//returns EConsumeAllMpmcQueueResult::WasEmpty if the Queue was empty
+		//or EConsumeAllMpmcQueueResult::HadItems if there were items to consume
+		template<typename F>
+		EConsumeAllMpmcQueueResult ConsumeAllFifo(const F& Consumer)
+		{
+			return ConsumeAll<true>(Consumer);
+		}
+
+	private:
+		template<bool bReverse, typename F>
+		inline EConsumeAllMpmcQueueResult ConsumeAll(const F& Consumer)
+		{
+			//pop the entire Stack
+			FNode* Node = Head.exchange(nullptr, std::memory_order_acq_rel);
+
+			if (Node == nullptr)
+			{
+				return EConsumeAllMpmcQueueResult::WasEmpty;
+			}
+
+			if (bReverse) //reverse the links to FIFO Order if requested
+			{
+				FNode* Prev = nullptr;
+				while (Node)
+				{
+					FNode* Tmp = Node;
+					Node = Node->Next.exchange(Prev, std::memory_order_relaxed);
+					Prev = Tmp;
+				}
+				Node = Prev;
+			}
+
+			while (Node) //consume the nodes of the Queue
+			{
+				FNode* Next = Node->Next.load(std::memory_order_relaxed);
+				T* ValuePtr = Node->Item.GetTypedPtr();
+				Consumer(MoveTemp(*ValuePtr));
+				DestructItem(ValuePtr);
+				AllocatorType::Free(Node);
+				Node = Next;
+			}
+
+			return EConsumeAllMpmcQueueResult::HadItems;
+		}
+	};
+}
diff --git a/Engine/Source/Runtime/Core/Public/Containers/DepletableMpscQueue.h b/Engine/Source/Runtime/Core/Public/Containers/DepletableMpscQueue.h
--- a/Engine/Source/Runtime/Core/Public/Containers/DepletableMpscQueue.h
+++ b/Engine/Source/Runtime/Core/Public/Containers/DepletableMpscQueue.h
@@ -36,13 +36,14 @@ namespace UE
 
 		~TDepletableMpscQueue()
 		{
+			static_assert(std::is_trivially_destructible_v<FNode>);
+
 			// delete remaining elements
 			FNode* Node = Sentinel.Next.load(std::memory_order_relaxed);
 			while (Node != nullptr)
 			{
 				DestructItem(Node->Value.GetTypedPtr());
 				FNode* Next = Node->Next.load(std::memory_order_relaxed);
-				static_assert(std::is_trivially_destructible_v<FNode>);
 				AllocatorType::Free(Node);
 				Node = Next;
 			}
@@ -79,16 +80,13 @@ namespace UE
 		template<typename F>
 		void Deplete(const F& Consumer)
 		{
-			FNode* First = Sentinel.Next.load(std::memory_order_relaxed);
+			// reset the head so the next consumption can detect that the queue is empty
+			FNode* First = Sentinel.Next.exchange(nullptr, std::memory_order_relaxed);
 			if (First == nullptr)
 			{
 				return; // empty
 			}
 
-			// reset the head so the next consumption can detect that the queue is empty
-			// `Sentinel.Next` is not touched by producers right now because it's already not null
-			Sentinel.Next.store(nullptr, std::memory_order_relaxed);
-
 			// reset the queue to the empty state. this redirects producers to start from `Sentinel` again.
 			// take note of the tail on resetting it because the list can be still not fully linked and so `Node.Next == nullptr` can't be 
 			// used to detect the end of the list
@@ -98,8 +96,7 @@ namespace UE
 			// the previously queued items are detached from the instance (as a linked list, though potentially not fully linked yet)
 
 			check(Last != &Sentinel); // can't be empty because of `First != nullptr` above
-
-			Consume(First, Last, Consumer);
+			Deplete_Internal(First, Last, Consumer);
 		}
 
 		// the result can be relied upon only in special cases, as the state can change concurrently. use with caution 
@@ -110,17 +107,16 @@ namespace UE
 
 	private:
 		template<typename F>
-		static void Consume(FNode* First, FNode* Last, const F& Consumer)
+		static void Deplete_Internal(FNode* First, FNode* Last, const F& Consumer)
 		{
 			auto GetNext = [](FNode* Node)
 			{
-				FNode* Next;
+				FNode* Next = nullptr;
 				// producers can be still updating `Next`, wait until the link to the next element is established
-				do
+				while(Next == nullptr) // <- This loop has the potential for live locking if enqueue was not completed (e.g was running at lower priority)
 				{
 					Next = Node->Next.load(std::memory_order_relaxed);
-				} while (Next == nullptr);
-
+				};
 				return Next;
 			};
 
@@ -129,14 +125,13 @@ namespace UE
 				T* ValuePtr = (T*)&Node->Value;
 				Consumer(MoveTemp(*ValuePtr));
 				DestructItem(ValuePtr);
+				AllocatorType::Free(Node);
 			};
 
 			while (First != Last)
 			{
-				Consume(First);
 				FNode* Next = GetNext(First);
-				static_assert(std::is_trivially_destructible_v<FNode>);
-				AllocatorType::Free(First);
+				Consume(First);
 				First = Next;
 			}
 
