Suppress warnings in Windows SDK headers that weren't wrapped with THIRD_PARTY_INCLUDES_START

#rnx

diff --git a/Engine/Platforms/Hololens/Source/Developer/HoloLensPlatformEditor/Private/HoloLensTargetSettingsCustomization.cpp b/Engine/Platforms/Hololens/Source/Developer/HoloLensPlatformEditor/Private/HoloLensTargetSettingsCustomization.cpp
--- a/Engine/Platforms/Hololens/Source/Developer/HoloLensPlatformEditor/Private/HoloLensTargetSettingsCustomization.cpp
+++ b/Engine/Platforms/Hololens/Source/Developer/HoloLensPlatformEditor/Private/HoloLensTargetSettingsCustomization.cpp
@@ -28,7 +28,9 @@
 #include "HoloLensLocalizedResourcesCustomization.h"
 #include "Styling/AppStyle.h"
 
+THIRD_PARTY_INCLUDES_START
 #include <winreg.h>
+THIRD_PARTY_INCLUDES_END
 
 #pragma comment( lib, "Advapi32.lib" )
 
diff --git a/Engine/Plugins/Developer/RiderSourceCodeAccess/Source/RiderSourceCodeAccess/Private/RiderPathLocator/Win/RiderPathLocatorWin.cpp b/Engine/Plugins/Developer/RiderSourceCodeAccess/Source/RiderSourceCodeAccess/Private/RiderPathLocator/Win/RiderPathLocatorWin.cpp
--- a/Engine/Plugins/Developer/RiderSourceCodeAccess/Source/RiderSourceCodeAccess/Private/RiderPathLocator/Win/RiderPathLocatorWin.cpp
+++ b/Engine/Plugins/Developer/RiderSourceCodeAccess/Source/RiderSourceCodeAccess/Private/RiderPathLocator/Win/RiderPathLocatorWin.cpp
@@ -12,7 +12,9 @@
 
 #include "Windows/WindowsPlatformMisc.h"
 #include "Windows/AllowWindowsPlatformTypes.h"
+THIRD_PARTY_INCLUDES_START
 #include <winreg.h>
+THIRD_PARTY_INCLUDES_END
 #include "Windows/HideWindowsPlatformTypes.h"
 
 static FString GetToolboxPath()
diff --git a/Engine/Plugins/VirtualProduction/Switchboard/Source/SwitchboardEditor/Private/SwitchboardEditorModule.cpp b/Engine/Plugins/VirtualProduction/Switchboard/Source/SwitchboardEditor/Private/SwitchboardEditorModule.cpp
--- a/Engine/Plugins/VirtualProduction/Switchboard/Source/SwitchboardEditor/Private/SwitchboardEditorModule.cpp
+++ b/Engine/Plugins/VirtualProduction/Switchboard/Source/SwitchboardEditor/Private/SwitchboardEditorModule.cpp
@@ -28,7 +28,9 @@
 
 #if PLATFORM_WINDOWS
 #include "Windows/AllowWindowsPlatformTypes.h"
+THIRD_PARTY_INCLUDES_START
 #include <winreg.h>
+THIRD_PARTY_INCLUDES_END
 #include "Windows/HideWindowsPlatformTypes.h"
 
 static const FString WindowsRunRegKeyPath = "Software\\Microsoft\\Windows\\CurrentVersion\\Run";
diff --git a/Engine/Plugins/XGEController/Source/Private/XGEControllerModule.cpp b/Engine/Plugins/XGEController/Source/Private/XGEControllerModule.cpp
--- a/Engine/Plugins/XGEController/Source/Private/XGEControllerModule.cpp
+++ b/Engine/Plugins/XGEController/Source/Private/XGEControllerModule.cpp
@@ -21,7 +21,9 @@
 #if PLATFORM_WINDOWS
 
 	#include "Windows/AllowWindowsPlatformTypes.h"
+	THIRD_PARTY_INCLUDES_START
 	#include <winreg.h>
+	THIRD_PARTY_INCLUDES_END
 	#include "Windows/HideWindowsPlatformTypes.h"
 
 	#define XGE_CONTROL_WORKER_NAME		TEXT("XGEControlWorker")
diff --git a/Engine/Source/Developer/TraceServices/Private/Model/DbgHelpResolver.cpp b/Engine/Source/Developer/TraceServices/Private/Model/DbgHelpResolver.cpp
--- a/Engine/Source/Developer/TraceServices/Private/Model/DbgHelpResolver.cpp
+++ b/Engine/Source/Developer/TraceServices/Private/Model/DbgHelpResolver.cpp
@@ -21,7 +21,9 @@
 #include <atomic>
 
 #include "Windows/AllowWindowsPlatformTypes.h"
+THIRD_PARTY_INCLUDES_START
 #include <dbghelp.h>
+THIRD_PARTY_INCLUDES_END
 
 /////////////////////////////////////////////////////////////////////
 DEFINE_LOG_CATEGORY_STATIC(LogDbgHelp, Log, All);
diff --git a/Engine/Source/Runtime/Core/Private/Windows/WindowsPlatformStackWalk.cpp b/Engine/Source/Runtime/Core/Private/Windows/WindowsPlatformStackWalk.cpp
--- a/Engine/Source/Runtime/Core/Private/Windows/WindowsPlatformStackWalk.cpp
+++ b/Engine/Source/Runtime/Core/Private/Windows/WindowsPlatformStackWalk.cpp
@@ -19,10 +19,12 @@
 
 #include "Windows/WindowsHWrapper.h"
 #include "Windows/AllowWindowsPlatformTypes.h"
+THIRD_PARTY_INCLUDES_START
 	#include <DbgHelp.h>
 	#include <Shlwapi.h>
 	#include <TlHelp32.h>
 	#include <Psapi.h>
+THIRD_PARTY_INCLUDES_END
 #include "Windows/HideWindowsPlatformTypes.h"
 
 #include "Modules/ModuleManager.h"
@@ -761,7 +763,7 @@ void LoadSymbolsForModule(HMODULE ModuleHandle, const FString& RemoteStorage)
 	SymSetSearchPathW(ProcessHandle, *SearchPathList);
 
 	// Load module.
-	const DWORD64 BaseAddress = SymLoadModuleExW(ProcessHandle, ModuleHandle, ImageName, ModuleName, (DWORD64)ModuleInfo.lpBaseOfDll, (uint32)ModuleInfo.SizeOfImage, NULL, 0);
+	const DWORD64 BaseAddress = SymLoadModuleExW(ProcessHandle, ModuleHandle, ImageName, ModuleName, (DWORD64)ModuleInfo.lpBaseOfDll, ModuleInfo.SizeOfImage, NULL, 0);
 	if (!BaseAddress)
 	{
 		ErrorCode = GetLastError();
diff --git a/Engine/Source/Runtime/Core/Public/Windows/WindowsPlatformCompilerPreSetup.h b/Engine/Source/Runtime/Core/Public/Windows/WindowsPlatformCompilerPreSetup.h
--- a/Engine/Source/Runtime/Core/Public/Windows/WindowsPlatformCompilerPreSetup.h
+++ b/Engine/Source/Runtime/Core/Public/Windows/WindowsPlatformCompilerPreSetup.h
@@ -206,6 +206,7 @@
 			__pragma(warning(disable: 6340))  /* Mismatch on sign: Incorrect type passed as parameter in call to function. */ \
 			__pragma(warning(disable: 6385))  /* Reading invalid data from '<ptr>':  the readable size is '<num1>' bytes, but '<num2>' bytes may be read. */ \
 			__pragma(warning(disable: 6386))  /* Buffer overrun while writing to '<ptr>':  the writable size is '<num1>' bytes, but '<num2>' bytes might be written. */ \
+			__pragma(warning(disable: 6553))  /* The annotation for function '<func>' on _Param_(<num>) does not apply to a value type. */ \
 			__pragma(warning(disable: 28182)) /* Dereferencing NULL pointer. '<ptr1>' contains the same NULL value as '<ptr2>' did. */ \
 			__pragma(warning(disable: 28251)) /* Inconsistent annotation for '<func>': this instance has no annotations. */ \
 			__pragma(warning(disable: 28252)) /* Inconsistent annotation for '<func>': return/function has '<annotation>' on the prior instance. */ \
diff --git a/Engine/Source/Runtime/Online/SSL/Private/Windows/WindowsPlatformSslCertificateManager.cpp b/Engine/Source/Runtime/Online/SSL/Private/Windows/WindowsPlatformSslCertificateManager.cpp
--- a/Engine/Source/Runtime/Online/SSL/Private/Windows/WindowsPlatformSslCertificateManager.cpp
+++ b/Engine/Source/Runtime/Online/SSL/Private/Windows/WindowsPlatformSslCertificateManager.cpp
@@ -12,9 +12,11 @@
 #include "Windows/WindowsHWrapper.h"
 #include "Windows/AllowWindowsPlatformTypes.h"
 
+THIRD_PARTY_INCLUDES_START
 #include <wincrypt.h>
 #include <openssl/x509.h>
 #include <openssl/err.h>
+THIRD_PARTY_INCLUDES_END
 
 #include "Windows/HideWindowsPlatformTypes.h"
 
