Replaced some != loop boundary tests with < to workaround false positives in MSVC's static analyser.

diff --git a/Engine/Source/Runtime/Core/Private/Tests/Misc/PlatformStringConvertTests.cpp b/Engine/Source/Runtime/Core/Private/Tests/Misc/PlatformStringConvertTests.cpp
--- a/Engine/Source/Runtime/Core/Private/Tests/Misc/PlatformStringConvertTests.cpp
+++ b/Engine/Source/Runtime/Core/Private/Tests/Misc/PlatformStringConvertTests.cpp
@@ -50,7 +50,8 @@ bool FStringConvertTest::RunTest(const FString& Parameters)
 
 		TestTrue(FString::Printf(TEXT("%s converted length matches"), TestName), End == Ptr + 4 + ConvertedLength);
 
-		for (int i = 0; i != 4 * sizeof(ToType); ++i)
+		// != changed to < because of false positive in MSVC's static analyzer: warning C6295: Ill-defined for-loop.  Loop executes infinitely.
+		for (int i = 0; i < 4 * sizeof(ToType); ++i)
 		{
 			TestTrue(FString::Printf(TEXT("%s guard block is preserved"), TestName), ((unsigned char*)Ptr)[i] == 0xCD && ((unsigned char*)End)[i] == 0xCD);
 		}
diff --git a/Engine/Source/Runtime/CoreUObject/Public/UObject/UObjectBaseUtility.h b/Engine/Source/Runtime/CoreUObject/Public/UObject/UObjectBaseUtility.h
--- a/Engine/Source/Runtime/CoreUObject/Public/UObject/UObjectBaseUtility.h
+++ b/Engine/Source/Runtime/CoreUObject/Public/UObject/UObjectBaseUtility.h
@@ -922,7 +922,8 @@ FORCEINLINE bool IsPossiblyAllocatedUObjectPointer(UObject* Ptr)
 	{
 		int32 Result = 0;
 
-		for (int32 I = 0; I != sizeof(UPTRINT); ++I)
+		// != changed to < because of false positive in MSVC's static analyzer: warning C6295: Ill-defined for-loop.  Loop executes infinitely.
+		for (int32 I = 0; I < sizeof(UPTRINT); ++I)
 		{
 			if ((Val & 0xFF) == ByteVal)
 			{
