Fix GetModularFeatureImplementation asserting despite promising that it wouldn't in the header comment.

diff --git a/Engine/Source/Runtime/Core/Private/Features/ModularFeatures.cpp b/Engine/Source/Runtime/Core/Private/Features/ModularFeatures.cpp
--- a/Engine/Source/Runtime/Core/Private/Features/ModularFeatures.cpp
+++ b/Engine/Source/Runtime/Core/Private/Features/ModularFeatures.cpp
@@ -48,7 +48,6 @@ IModularFeature* FModularFeatures::GetModularFeatureImplementation( const FName
 		++CurrentIndex;
 	}
 
-	checkf( ModularFeature != nullptr, TEXT("ModularFeature not found: '%s', Index: %d"), *Type.ToString(), Index);
 	return ModularFeature;
 }
 
