Fix potential truncation when loading data from a compact binary to an array

#rnx

- The array view returned by FCbFieldView uses uint64 as it's size type, but TArray uses the significantly smaller type int32.
- It is very unlikely that we will ever see a truncation here but if it doesn't happen it could create some annoyingly difficult to diagnose bugs.
- We now test to make sure that the data can fit into a TArray before loading it, if not we print an error.
- The error logging is in a private function so that the code can be in the cpp and reduce include pollution in the header, it will also help hide the expanded code from the UE_LOGGING macro.

diff --git a/Engine/Source/Runtime/Core/Private/Serialization/CompactBinarySerialization.cpp b/Engine/Source/Runtime/Core/Private/Serialization/CompactBinarySerialization.cpp
--- a/Engine/Source/Runtime/Core/Private/Serialization/CompactBinarySerialization.cpp
+++ b/Engine/Source/Runtime/Core/Private/Serialization/CompactBinarySerialization.cpp
@@ -7,7 +7,9 @@
 #include "Containers/StringConv.h"
 #include "Containers/StringView.h"
 #include "Containers/UnrealString.h"
+#include "CoreGlobals.h"
 #include "HAL/PlatformString.h"
+#include "Logging/LogMacros.h"
 #include "Memory/MemoryView.h"
 #include "Misc/AsciiSet.h"
 #include "Misc/AssertionMacros.h"
@@ -29,6 +31,18 @@
 
 ///////////////////////////////////////////////////////////////////////////////////////////////////////////////////////
 
+namespace UE::Serialization::Private
+{
+
+void LogFieldTooLargeForArrayWarning(uint64 FieldLength)
+{
+	UE_LOG(LogSerialization, Warning, TEXT("Attempting to load a FCbFieldView that has too many entries (%" UINT64_FMT ") to fit in a TArray "), FieldLength);
+}
+
+} // namespace UE::Serialization::Private
+
+///////////////////////////////////////////////////////////////////////////////////////////////////////////////////////
+
 uint64 MeasureCompactBinary(FMemoryView View, ECbFieldType Type)
 {
 	uint64 Size;
diff --git a/Engine/Source/Runtime/Core/Public/Serialization/CompactBinarySerialization.h b/Engine/Source/Runtime/Core/Public/Serialization/CompactBinarySerialization.h
--- a/Engine/Source/Runtime/Core/Public/Serialization/CompactBinarySerialization.h
+++ b/Engine/Source/Runtime/Core/Public/Serialization/CompactBinarySerialization.h
@@ -18,6 +18,16 @@ struct FGuid;
 
 ///////////////////////////////////////////////////////////////////////////////////////////////////////////////////////
 
+namespace UE::Serialization::Private
+{
+
+/** Utility for logging problems with FCbField. For internal use only */
+CORE_API void LogFieldTooLargeForArrayWarning(uint64 FieldLength);
+
+} // namespace UE::Serialization::Private
+
+///////////////////////////////////////////////////////////////////////////////////////////////////////////////////////
+
 /**
  * Determine the size in bytes of the compact binary field at the start of the view.
  *
@@ -181,13 +191,22 @@ inline bool LoadFromCompactBinary(FCbFieldView Field, FCbObjectId& OutValue, con
 template <typename T, typename Allocator>
 inline bool LoadFromCompactBinary(FCbFieldView Field, TArray<T, Allocator>& OutValue)
 {
-	OutValue.Reset(Field.AsArrayView().Num());
-	bool bOk = !Field.HasError();
-	for (const FCbFieldView& ElementField : Field)
+	const uint64 Length = Field.AsArrayView().Num();
+	if (Length <= MAX_int32)
+	{
+		OutValue.Reset((int32)Length);
+		bool bOk = !Field.HasError();
+		for (const FCbFieldView& ElementField : Field)
+		{
+			bOk = LoadFromCompactBinary(ElementField, OutValue.Emplace_GetRef()) & bOk;
+		}
+		return bOk;
+	}
+	else
 	{
-		bOk = LoadFromCompactBinary(ElementField, OutValue.Emplace_GetRef()) & bOk;
+		UE::Serialization::Private::LogFieldTooLargeForArrayWarning(Length);
+		return false;
 	}
-	return bOk;
 }
 
 ///////////////////////////////////////////////////////////////////////////////////////////////////////////////////////
