Allow tracing to be started without sending the tail of historical events.

#rnx

diff --git a/Engine/Source/Runtime/Core/Private/ProfilingDebugging/TraceAuxiliary.cpp b/Engine/Source/Runtime/Core/Private/ProfilingDebugging/TraceAuxiliary.cpp
--- a/Engine/Source/Runtime/Core/Private/ProfilingDebugging/TraceAuxiliary.cpp
+++ b/Engine/Source/Runtime/Core/Private/ProfilingDebugging/TraceAuxiliary.cpp
@@ -89,7 +89,7 @@ public:
 	bool HasCommandlineChannels() const { return !CommandlineChannels.IsEmpty(); }
 	void EnableChannels(const TCHAR* ChannelList);
 	void DisableChannels(const TCHAR* ChannelList);
-	bool Connect(ETraceConnectType Type, const TCHAR* Parameter, const FTraceAuxiliary::FLogCategoryAlias& LogCategory);
+	bool Connect(ETraceConnectType Type, const TCHAR* Parameter, const FTraceAuxiliary::FLogCategoryAlias& LogCategory, uint16 SendFlags);
 	bool Stop();
 	void ResumeChannels();
 	void PauseChannels();
@@ -123,8 +123,8 @@ private:
 	static uint32 HashChannelName(const TCHAR* Name);
 	bool EnableChannel(const TCHAR* Channel);
 	void DisableChannel(const TCHAR* Channel);
-	bool SendToHost(const TCHAR* Host, const FTraceAuxiliary::FLogCategoryAlias& LogCategory);
-	bool WriteToFile(const TCHAR* Path, const FTraceAuxiliary::FLogCategoryAlias& LogCategory);
+	bool SendToHost(const TCHAR* Host, const FTraceAuxiliary::FLogCategoryAlias& LogCategory, uint16 SendFlags);
+	bool WriteToFile(const TCHAR* Path, const FTraceAuxiliary::FLogCategoryAlias& LogCategory, uint16 SendFlags);
 	bool FinalizeFilePath(const TCHAR* InPath, FString& OutPath, const FTraceAuxiliary::FLogCategoryAlias& LogCategory);
 
 	typedef TMap<uint32, FChannelEntry, TInlineSetAllocator<128>> ChannelSet;
@@ -274,7 +274,7 @@ void FTraceAuxiliaryImpl::RemoveChannel(const TCHAR* Name)
 }
 
 ////////////////////////////////////////////////////////////////////////////////
-bool FTraceAuxiliaryImpl::Connect(ETraceConnectType Type, const TCHAR* Parameter, const FTraceAuxiliary::FLogCategoryAlias& LogCategory)
+bool FTraceAuxiliaryImpl::Connect(ETraceConnectType Type, const TCHAR* Parameter, const FTraceAuxiliary::FLogCategoryAlias& LogCategory, uint16 SendFlags)
 {
 	// Connect/write to file, but only if we're not already sending/writing.
 	bool bConnected = UE::Trace::IsTracing();
@@ -282,7 +282,7 @@ bool FTraceAuxiliaryImpl::Connect(ETraceConnectType Type, const TCHAR* Parameter
 	{
 		if (Type == ETraceConnectType::Network)
 		{
-			bConnected = SendToHost(Parameter, LogCategory);
+			bConnected = SendToHost(Parameter, LogCategory, SendFlags);
 			if (bConnected)
 			{
 				UE_LOG_REF(LogCategory, Display, TEXT("Trace started (connected to trace server %s)."), GetDest());
@@ -297,7 +297,7 @@ bool FTraceAuxiliaryImpl::Connect(ETraceConnectType Type, const TCHAR* Parameter
 
 		else if (Type == ETraceConnectType::File)
 		{
-			bConnected = WriteToFile(Parameter, LogCategory);
+			bConnected = WriteToFile(Parameter, LogCategory, SendFlags);
 			if (bConnected)
 			{
 				UE_LOG_REF(LogCategory, Display, TEXT("Trace started (writing to file \"%s\")."), GetDest());
@@ -441,9 +441,10 @@ void FTraceAuxiliaryImpl::SetTruncateFile(bool bNewTruncateFileState)
 }
 
 ////////////////////////////////////////////////////////////////////////////////
-bool FTraceAuxiliaryImpl::SendToHost(const TCHAR* Host, const FTraceAuxiliary::FLogCategoryAlias& LogCategory)
+bool FTraceAuxiliaryImpl::SendToHost(const TCHAR* Host, const FTraceAuxiliary::FLogCategoryAlias& LogCategory, uint16 SendFlags)
 {
-	if (!UE::Trace::SendTo(Host))
+	uint32 Port = 0; // "0" indicates that the default should be used.
+	if (!UE::Trace::SendTo(Host, Port, SendFlags))
 	{
 		UE_LOG_REF(LogCategory, Warning, TEXT("Unable to trace to host '%s'"), Host);
 		return false;
@@ -523,7 +524,7 @@ bool FTraceAuxiliaryImpl::FinalizeFilePath(const TCHAR* InPath, FString& OutPath
 }
 
 ////////////////////////////////////////////////////////////////////////////////
-bool FTraceAuxiliaryImpl::WriteToFile(const TCHAR* Path, const FTraceAuxiliary::FLogCategoryAlias& LogCategory)
+bool FTraceAuxiliaryImpl::WriteToFile(const TCHAR* Path, const FTraceAuxiliary::FLogCategoryAlias& LogCategory, uint16 SendFlags)
 {
 	FString NativePath;
 	if (!FinalizeFilePath(Path, NativePath, LogCategory))
@@ -531,7 +532,7 @@ bool FTraceAuxiliaryImpl::WriteToFile(const TCHAR* Path, const FTraceAuxiliary::
 		return false;
 	}
 
-	if (!UE::Trace::WriteTo(*NativePath))
+	if (!UE::Trace::WriteTo(*NativePath, SendFlags))
 	{
 		if (FPathViews::Equals(NativePath, FStringView(Path)))
 		{
@@ -1316,13 +1317,15 @@ bool FTraceAuxiliary::Start(EConnectionType Type, const TCHAR* Target, const TCH
 		}
 	}
 
+	uint16 SendFlags = (Options && Options->bExcludeTail) ? UE::Trace::FSendFlags::ExcludeTail : 0;
+	
 	if (Type == EConnectionType::File)
 	{
-		return GTraceAuxiliary.Connect(ETraceConnectType::File, Target, LogCategory);
+		return GTraceAuxiliary.Connect(ETraceConnectType::File, Target, LogCategory, SendFlags);
 	}
 	else if(Type == EConnectionType::Network)
 	{
-		return GTraceAuxiliary.Connect(ETraceConnectType::Network, Target, LogCategory);
+		return GTraceAuxiliary.Connect(ETraceConnectType::Network, Target, LogCategory, SendFlags);
 	}
 #endif
 	return false;
diff --git a/Engine/Source/Runtime/Core/Public/ProfilingDebugging/TraceAuxiliary.h b/Engine/Source/Runtime/Core/Public/ProfilingDebugging/TraceAuxiliary.h
--- a/Engine/Source/Runtime/Core/Public/ProfilingDebugging/TraceAuxiliary.h
+++ b/Engine/Source/Runtime/Core/Public/ProfilingDebugging/TraceAuxiliary.h
@@ -33,11 +33,11 @@ public:
 		 * Write to a file. Target string is filename. Absolute or relative current working directory.
 		 * If target is null the current date and time is used.
 		 */
-		 File,
-		 /**
-		  * Don't connect, just start tracing to memory.
-		  */
-		  None,
+		File,
+		/**
+		 * Don't connect, just start tracing to memory.
+		 */
+		None,
 	};
 
 	/**
@@ -67,6 +67,8 @@ public:
 		bool bNoWorkerThread = false;
 		/** When set, the target file will be truncated if it already exists. */
 		bool bTruncateFile = false;
+		/** When set, trace data buffered before starting tracing will not be output to the trace file */
+		bool bExcludeTail = false;
 	};
 
 	/**
diff --git a/Engine/Source/Runtime/TraceLog/Private/Trace/Control.cpp b/Engine/Source/Runtime/TraceLog/Private/Trace/Control.cpp
--- a/Engine/Source/Runtime/TraceLog/Private/Trace/Control.cpp
+++ b/Engine/Source/Runtime/TraceLog/Private/Trace/Control.cpp
@@ -18,8 +18,8 @@ namespace Private {
 #if !defined(TRACE_PRIVATE_CONTROL_ENABLED) || TRACE_PRIVATE_CONTROL_ENABLED
 
 ////////////////////////////////////////////////////////////////////////////////
-bool	Writer_SendTo(const ANSICHAR*, uint32=0);
-bool	Writer_WriteTo(const ANSICHAR*);
+bool	Writer_SendTo(const ANSICHAR*, uint32=0, uint32=0);
+bool	Writer_WriteTo(const ANSICHAR*, uint32=0);
 bool	Writer_Stop();
 
 
diff --git a/Engine/Source/Runtime/TraceLog/Private/Trace/Trace.cpp b/Engine/Source/Runtime/TraceLog/Private/Trace/Trace.cpp
--- a/Engine/Source/Runtime/TraceLog/Private/Trace/Trace.cpp
+++ b/Engine/Source/Runtime/TraceLog/Private/Trace/Trace.cpp
@@ -20,8 +20,8 @@ void	Writer_Initialize(const FInitializeDesc&);
 void	Writer_WorkerCreate();
 void	Writer_Shutdown();
 void	Writer_Update();
-bool	Writer_SendTo(const ANSICHAR*, uint32);
-bool	Writer_WriteTo(const ANSICHAR*);
+bool	Writer_SendTo(const ANSICHAR*, uint32, uint32);
+bool	Writer_WriteTo(const ANSICHAR*, uint32);
 bool	Writer_WriteSnapshotTo(const ANSICHAR*);
 bool	Writer_IsTracing();
 bool	Writer_Stop();
@@ -82,19 +82,19 @@ void GetStatistics(FStatistics& Out)
 }
 
 ////////////////////////////////////////////////////////////////////////////////
-bool SendTo(const TCHAR* InHost, uint32 Port)
+bool SendTo(const TCHAR* InHost, uint32 Port, uint16 Flags)
 {
 	char Host[256];
 	ToAnsiCheap(Host, InHost);
-	return Private::Writer_SendTo(Host, Port);
+	return Private::Writer_SendTo(Host, Flags, Port);
 }
 
 ////////////////////////////////////////////////////////////////////////////////
-bool WriteTo(const TCHAR* InPath)
+bool WriteTo(const TCHAR* InPath, uint16 Flags)
 {
 	char Path[512];
 	ToAnsiCheap(Path, InPath);
-	return Private::Writer_WriteTo(Path);
+	return Private::Writer_WriteTo(Path, Flags);
 }
 
 ////////////////////////////////////////////////////////////////////////////////
diff --git a/Engine/Source/Runtime/TraceLog/Private/Trace/Writer.cpp b/Engine/Source/Runtime/TraceLog/Private/Trace/Writer.cpp
--- a/Engine/Source/Runtime/TraceLog/Private/Trace/Writer.cpp
+++ b/Engine/Source/Runtime/TraceLog/Private/Trace/Writer.cpp
@@ -493,6 +493,10 @@ static bool Writer_UpdateConnection()
 
 	AtomicStoreRelaxed(&GPendingDataHandle, UPTRINT(0));
 
+	// Extract send flags
+	uint32 SendFlags = uint32(PendingDataHandle >> 48ull);
+	PendingDataHandle &= 0x0000'ffff'ffff'ffffull;
+
 	// Reject the pending connection if we've already got a connection
 	if (GDataHandle)
 	{
@@ -522,7 +526,10 @@ static bool Writer_UpdateConnection()
 	Writer_CallbackOnConnect();
 
 	// Finally write the events in the tail buffer
-	Writer_TailOnConnect();
+	if ((SendFlags & FSendFlags::ExcludeTail) == 0)
+	{
+		Writer_TailOnConnect();
+	}
 
 	// See Writer_SendSync() for details.
 	GSyncPacketCountdown = GNumSyncPackets;
@@ -812,7 +819,26 @@ void Writer_Update()
 }
 
 ////////////////////////////////////////////////////////////////////////////////
-bool Writer_SendTo(const ANSICHAR* Host, uint32 Port)
+static UPTRINT Writer_PackSendFlags(UPTRINT DataHandle, uint32 Flags)
+{
+	// Passing ownership of IO to the worker thread via a single 64 bit value is
+	// convenient and saves a lot of machinery for something that mostly never
+	// happens, or perhaps just once or twice. Here we make the assumption that
+	// our supported platforms' handles are low integer file descriptor IDs or
+	// addresses and thus we have some most-significant bits to use for flags.
+
+	// Guard against the assumption being wrong.
+	if (DataHandle & 0xffff'0000'0000'0000ull)
+	{
+		IoClose(DataHandle);
+		return 0;
+	}
+
+	return DataHandle | (UPTRINT(Flags) << 48ull);
+}
+
+////////////////////////////////////////////////////////////////////////////////
+bool Writer_SendTo(const ANSICHAR* Host, uint32 Flags, uint32 Port)
 {
 	if (AtomicLoadRelaxed(&GPendingDataHandle))
 	{
@@ -828,12 +854,18 @@ bool Writer_SendTo(const ANSICHAR* Host, uint32 Port)
 		return false;
 	}
 
+	DataHandle = Writer_PackSendFlags(DataHandle, Flags);
+	if (!DataHandle)
+	{
+		return false;
+	}
+
 	AtomicStoreRelaxed(&GPendingDataHandle, DataHandle);
 	return true;
 }
 
 ////////////////////////////////////////////////////////////////////////////////
-bool Writer_WriteTo(const ANSICHAR* Path)
+bool Writer_WriteTo(const ANSICHAR* Path, uint32 Flags)
 {
 	if (AtomicLoadRelaxed(&GPendingDataHandle))
 	{
@@ -848,6 +880,12 @@ bool Writer_WriteTo(const ANSICHAR* Path)
 		return false;
 	}
 
+	DataHandle = Writer_PackSendFlags(DataHandle, Flags);
+	if (!DataHandle)
+	{
+		return false;
+	}
+
 	AtomicStoreRelaxed(&GPendingDataHandle, DataHandle);
 	return true;
 }
diff --git a/Engine/Source/Runtime/TraceLog/Public/Trace/Trace.h b/Engine/Source/Runtime/TraceLog/Public/Trace/Trace.h
--- a/Engine/Source/Runtime/TraceLog/Public/Trace/Trace.h
+++ b/Engine/Source/Runtime/TraceLog/Public/Trace/Trace.h
@@ -106,14 +106,21 @@ struct FStatistics
 	uint32		CacheWaste;		// Unused memory from retired cache buffers
 };
 
+struct FSendFlags
+{
+	static const uint16 None		= 0;
+	static const uint16 ExcludeTail	= 1 << 0;	// do not send the tail of historical events
+	static const uint16 _Reserved	= 1 << 15;	// this bit is used internally
+};
+
 UE_TRACE_API void	SetMemoryHooks(AllocFunc Alloc, FreeFunc Free) UE_TRACE_IMPL();
 UE_TRACE_API void	Initialize(const FInitializeDesc& Desc) UE_TRACE_IMPL();
 UE_TRACE_API void	StartWorkerThread() UE_TRACE_IMPL();	
 UE_TRACE_API void	Shutdown() UE_TRACE_IMPL();
 UE_TRACE_API void	Update() UE_TRACE_IMPL();
 UE_TRACE_API void	GetStatistics(FStatistics& Out) UE_TRACE_IMPL();
-UE_TRACE_API bool	SendTo(const TCHAR* Host, uint32 Port=0) UE_TRACE_IMPL(false);
-UE_TRACE_API bool	WriteTo(const TCHAR* Path) UE_TRACE_IMPL(false);
+UE_TRACE_API bool	SendTo(const TCHAR* Host, uint32 Port=0, uint16 Flags=FSendFlags::None) UE_TRACE_IMPL(false);
+UE_TRACE_API bool	WriteTo(const TCHAR* Path, uint16 Flags=FSendFlags::None) UE_TRACE_IMPL(false);
 UE_TRACE_API bool	WriteSnapshotTo(const TCHAR* Path) UE_TRACE_IMPL(false);
 UE_TRACE_API bool	IsTracing() UE_TRACE_IMPL(false);
 UE_TRACE_API bool	Stop() UE_TRACE_IMPL(false);
