Fixed deprecated intrinsic type traits.

diff --git a/Engine/Source/Runtime/Core/Public/Templates/IsTriviallyCopyAssignable.h b/Engine/Source/Runtime/Core/Public/Templates/IsTriviallyCopyAssignable.h
--- a/Engine/Source/Runtime/Core/Public/Templates/IsTriviallyCopyAssignable.h
+++ b/Engine/Source/Runtime/Core/Public/Templates/IsTriviallyCopyAssignable.h
@@ -3,6 +3,7 @@
 #pragma once
 
 #include "CoreTypes.h"
+#include <type_traits>
 
 /**
  * Traits class which tests if a type has a trivial copy assignment operator.
@@ -10,5 +11,5 @@
 template <typename T>
 struct TIsTriviallyCopyAssignable
 {
-	enum { Value = __has_trivial_assign(T) };
+	enum { Value = std::is_trivially_copy_assignable_v<T> };
 };
diff --git a/Engine/Source/Runtime/Core/Public/Templates/IsTriviallyCopyConstructible.h b/Engine/Source/Runtime/Core/Public/Templates/IsTriviallyCopyConstructible.h
--- a/Engine/Source/Runtime/Core/Public/Templates/IsTriviallyCopyConstructible.h
+++ b/Engine/Source/Runtime/Core/Public/Templates/IsTriviallyCopyConstructible.h
@@ -3,6 +3,7 @@
 #pragma once
 
 #include "CoreTypes.h"
+#include <type_traits>
 
 /**
  * Traits class which tests if a type has a trivial copy constructor.
@@ -10,5 +11,5 @@
 template <typename T>
 struct TIsTriviallyCopyConstructible
 {
-	enum { Value = __has_trivial_copy(T) };
+	enum { Value = std::is_trivially_copy_constructible_v<T> };
 };
diff --git a/Engine/Source/Runtime/Core/Public/Templates/IsTriviallyDestructible.h b/Engine/Source/Runtime/Core/Public/Templates/IsTriviallyDestructible.h
--- a/Engine/Source/Runtime/Core/Public/Templates/IsTriviallyDestructible.h
+++ b/Engine/Source/Runtime/Core/Public/Templates/IsTriviallyDestructible.h
@@ -1,23 +1,7 @@
 // Copyright Epic Games, Inc. All Rights Reserved.
 
 #pragma once
-
-namespace UE::Core::Private::IsTriviallyDestructible
-{
-	// We have this specialization for enums to avoid the need to have a full definition of
-	// the type.
-	template <typename T, bool bIsTriviallyTriviallyDestructible = __is_enum(T)>
-	struct TImpl
-	{
-		enum { Value = true };
-	};
-
-	template <typename T>
-	struct TImpl<T, false>
-	{
-		enum { Value = __has_trivial_destructor(T) };
-	};
-}
+#include <type_traits>
 
 /**
  * Traits class which tests if a type has a trivial destructor.
@@ -25,5 +9,5 @@ namespace UE::Core::Private::IsTriviallyDestructible
 template <typename T>
 struct TIsTriviallyDestructible
 {
-	enum { Value = UE::Core::Private::IsTriviallyDestructible::TImpl<T>::Value };
+	enum { Value = std::is_trivially_destructible_v<T> };
 };
