Fix rounding issue with FDateTime::ParseIso8601

The millisecond part could get rounded to 1000 which should not happen. This DataTime was having the problem: "2022-10-10T18:51:58.999500Z"

Problem reported on UDN 233271

#r devin.doucette

robomerge[FN_EngineMerge] -UE5-RES

diff --git a/Engine/Source/Runtime/Core/Private/Misc/DateTime.cpp b/Engine/Source/Runtime/Core/Private/Misc/DateTime.cpp
--- a/Engine/Source/Runtime/Core/Private/Misc/DateTime.cpp
+++ b/Engine/Source/Runtime/Core/Private/Misc/DateTime.cpp
@@ -820,7 +820,7 @@ bool FDateTime::ParseIso8601(const TCHAR* DateTimeString, FDateTime& OutDateTime
 				MillisecondTemp = (MillisecondTemp + (Divisor >> 1 )) / Divisor;
 			}
 			
-			Millisecond = static_cast<int32>(MillisecondTemp);
+			Millisecond = FMath::Min(static_cast<int32>(MillisecondTemp), 999);
 		}
 
 		// see if the timezone offset is included
