Fix ApplicationVerifier errors by not leaving FWindowsRWLock on destruction.

- We were checking if the SRW was locked or not in the destructor, but we did this by trying to acquire the lock, meaning that it would always be in the locked state when destroyed.
- Technically this is causing no issues but it does cause ApplicationVerifier to throw errors, which makes it hard to use the tool to find actual problems.
- Added a private method ::IsLocked which will do the acquire but also unacquire the lock before returning meaning we are in the correct state by the time that the destructor ends.

diff --git a/Engine/Source/Runtime/Core/Public/Windows/WindowsCriticalSection.h b/Engine/Source/Runtime/Core/Public/Windows/WindowsCriticalSection.h
--- a/Engine/Source/Runtime/Core/Public/Windows/WindowsCriticalSection.h
+++ b/Engine/Source/Runtime/Core/Public/Windows/WindowsCriticalSection.h
@@ -117,9 +117,9 @@ public:
 		Windows::InitializeSRWLock(&Mutex);
 	}
 
-	FORCEINLINE ~FWindowsRWLock()
+	~FWindowsRWLock()
 	{
-		checkf(Windows::TryAcquireSRWLockExclusive(&Mutex), TEXT("Destroying a lock that is still held!"));
+		checkf(!IsLocked(), TEXT("Destroying a lock that is still held!"));
 	}
 
 	FORCEINLINE void ReadLock()
@@ -153,6 +153,20 @@ public:
 	}
 
 private:
+
+	bool IsLocked()
+	{
+		if (Windows::TryAcquireSRWLockExclusive(&Mutex))
+		{
+			Windows::ReleaseSRWLockExclusive(&Mutex);
+			return false;
+		}
+		else
+		{
+			return true;
+		}
+	}
+
 	Windows::SRWLOCK Mutex;
 };
 
