disabled TaskGraph slow broadcast processing warnings on named threads as we don't wait for named threads broadcasts anymore

diff --git a/Engine/Source/Runtime/Core/Private/Async/TaskGraph.cpp b/Engine/Source/Runtime/Core/Private/Async/TaskGraph.cpp
--- a/Engine/Source/Runtime/Core/Private/Async/TaskGraph.cpp
+++ b/Engine/Source/Runtime/Core/Private/Async/TaskGraph.cpp
@@ -2622,33 +2622,34 @@ public:
 		return DesiredThread;
 	}
 
-	FORCEINLINE TStatId GetStatId() const
+	TStatId GetStatId() const
 	{
 		return GET_STATID(STAT_FBroadcastTask);
 	}
 
-	static FORCEINLINE ESubsequentsMode::Type GetSubsequentsMode() { return ESubsequentsMode::TrackSubsequents; }
-	void FORCEINLINE DoTask(ENamedThreads::Type CurrentThread, const FGraphEventRef& MyCompletionGraphEvent)
+	static ESubsequentsMode::Type GetSubsequentsMode() { return ESubsequentsMode::TrackSubsequents; }
+	void DoTask(ENamedThreads::Type CurrentThread, const FGraphEventRef& MyCompletionGraphEvent)
 	{
+		auto LogWarningIfSlow = [this](const TCHAR* Msg)
 		{
+			const bool bNamedThread = TaskEvent == nullptr; // we don't wait for named threads broadcasting, as they can be quite busy, 
+			// it may take longer to reach them. do not report slow processing in this case
 			const double ThisTime = FPlatformTime::Seconds() - StartTime;
-			if (ThisTime > 0.02)
+			if (!bNamedThread && ThisTime > 0.02)
 			{
-				UE_CLOG(GPrintBroadcastWarnings, LogTaskGraph, Warning, TEXT("Task graph took %6.2fms for %s to receive broadcast."), ThisTime * 1000.0, Name);
+				UE_CLOG(GPrintBroadcastWarnings, LogTaskGraph, Warning, TEXT("Task graph took %6.2fms for %s to %s"), ThisTime * 1000.0, Name, Msg);
 			}
-		}
+		};
+
+		LogWarningIfSlow(TEXT("receive broadcast."));
 
 		{
 			QUICK_SCOPE_CYCLE_COUNTER(STAT_Broadcast_PayloadFunction);
 			Function(CurrentThread);
 		}
-		{
-			const double ThisTime = FPlatformTime::Seconds() - StartTime;
-			if (ThisTime > 0.02)
-			{
-				UE_CLOG(GPrintBroadcastWarnings, LogTaskGraph, Warning, TEXT("Task graph took %6.2fms for %s to receive broadcast and do processing."), ThisTime * 1000.0, Name);
-			}
-		}
+
+		LogWarningIfSlow(TEXT("receive broadcast and do processing."));
+
 		if (StallForTaskThread)
 		{
 			if (StallForTaskThread->Decrement())
@@ -2657,25 +2658,13 @@ public:
 				{
 					QUICK_SCOPE_CYCLE_COUNTER(STAT_Broadcast_WaitForOthers);
 					TaskEvent->Wait();
-					{
-						const double ThisTime = FPlatformTime::Seconds() - StartTime;
-						if (ThisTime > 0.02)
-						{
-							UE_CLOG(GPrintBroadcastWarnings, LogTaskGraph, Warning, TEXT("Task graph took %6.2fms for %s to receive broadcast do processing and wait for other task threads."), ThisTime * 1000.0f, Name);
-						}
-					}
+					LogWarningIfSlow(TEXT("receive broadcast do processing and wait for other task threads."));
 				}
 			}
 			else
 			{
 				CallerEvent->Trigger();
-				{
-					const double ThisTime = FPlatformTime::Seconds() - StartTime;
-					if (ThisTime > 0.02)
-					{
-						UE_CLOG(GPrintBroadcastWarnings, LogTaskGraph, Warning, TEXT("Task graph took %6.2fms for %s to receive broadcast do processing and trigger other task threads."), ThisTime * 1000.0, Name);
-					}
-				}
+				LogWarningIfSlow(TEXT("to receive broadcast do processing and trigger other task threads."));
 			}
 		}
 	}
