Expose TRay/FRay to Blueprints.
Add some standard UE stuff to TRay (EForceInit constructor. Init() function, operator== and !=, operator<< and Serialize(), and SerializeFromMismatchedTag support.
Add TRay3StructOpsTypeTraits in Property.cpp
Register names in UnrealNames.inl
Add FRay/FRay3f/FRay3d in NoExportTypes.h. FRay3f and FRay3d are not exposed to Blueprints
Add UE_DECLARE_CORE_VARIANT_TYPE declaration in Class.h
Add forward declarations in MathFwd.h

diff --git a/Engine/Source/Runtime/Core/Public/Math/MathFwd.h b/Engine/Source/Runtime/Core/Public/Math/MathFwd.h
--- a/Engine/Source/Runtime/Core/Public/Math/MathFwd.h
+++ b/Engine/Source/Runtime/Core/Public/Math/MathFwd.h
@@ -28,6 +28,7 @@ namespace UE::Math
 	template <typename IntType> struct TIntVector2;
 	template <typename IntType> struct TIntVector3;
 	template <typename IntType> struct TIntVector4;
+	template<typename T> struct TRay;
 }
 
 // ISPC forward declarations
@@ -54,6 +55,7 @@ using FSphere 		= UE::Math::TSphere<double>;		// UE_DECLARE_LWC_TYPE(Sphere, 3);
 using FBox 			= UE::Math::TBox<double>;			// UE_DECLARE_LWC_TYPE(Box, 3);
 using FBox2D 		= UE::Math::TBox2<double>;			// UE_DECLARE_LWC_TYPE(Box2,, FBox2D);
 using FRotator 		= UE::Math::TRotator<double>;		// UE_DECLARE_LWC_TYPE(Rotator, 3);
+using FRay			= UE::Math::TRay<double>;			// UE_DECLARE_LWC_TYPE(Ray, 3);
 
 using FVector3d 	= UE::Math::TVector<double>;
 using FVector2d 	= UE::Math::TVector2<double>;
@@ -66,6 +68,7 @@ using FSphere3d 	= UE::Math::TSphere<double>;
 using FBox3d 		= UE::Math::TBox<double>;
 using FBox2d 		= UE::Math::TBox2<double>;
 using FRotator3d 	= UE::Math::TRotator<double>;
+using FRay3d		= UE::Math::TRay<double>;
 
 using FVector3f		= UE::Math::TVector<float>;
 using FVector2f		= UE::Math::TVector2<float>;
@@ -78,6 +81,7 @@ using FSphere3f		= UE::Math::TSphere<float>;
 using FBox3f		= UE::Math::TBox<float>;
 using FBox2f		= UE::Math::TBox2<float>;
 using FRotator3f	= UE::Math::TRotator<float>;
+using FRay3f		= UE::Math::TRay<float>;
 
 
 struct FColor;
diff --git a/Engine/Source/Runtime/Core/Public/Math/Ray.h b/Engine/Source/Runtime/Core/Public/Math/Ray.h
--- a/Engine/Source/Runtime/Core/Public/Math/Ray.h
+++ b/Engine/Source/Runtime/Core/Public/Math/Ray.h
@@ -30,10 +30,19 @@ public:
 	/** Default constructor initializes ray to Zero origin and Z-axis direction */
 	TRay()
 	{
-		Origin = TVector<T>::ZeroVector;
-		Direction = TVector<T>(0, 0, 1);
+		Init();
+	}
+
+	/**
+	 * Creates and initializes a new Ray to Zero Origin and Z-axis Direction.
+	 * @param EForceInit Force Init Enum.
+	 */
+	explicit TRay<T>( EForceInit )
+	{
+		Init();
 	}
 
+
 	/** 
 	  * Initialize Ray with origin and direction
 	  *
@@ -51,9 +60,38 @@ public:
 		}
 	}
 
-	// Conversion to other type.
-	template<typename FArg, TEMPLATE_REQUIRES(!TIsSame<T, FArg>::Value)>
-	explicit TRay(const TRay<FArg>& From) : TRay<T>(TVector<T>(From.Origin), TVector<T>(From.Direction), true) {}
+
+	/**
+	 * Set the initial values of the Ray to Zero Origin and Z-axis Direction.
+	 */
+	void Init()
+	{
+		Origin = TVector<T>::ZeroVector;
+		Direction = TVector<T>(0, 0, 1);
+	}
+
+	/**
+	 * Compares two Rays for equality.
+	 *
+	 * @param Other The other Ray to compare with.
+	 * @return true if the Rays are equal, false otherwise.
+	 */
+	bool operator==( const TRay<T>& Other ) const
+	{
+		return (Origin == Other.Origin) && (Direction == Other.Direction);
+	}
+
+	/**
+	 * Compares two Rays for inequality.
+	 *
+	 * @param Other The other Ray to compare with.
+	 * @return true if the Rays are not equal, false otherwise.
+	 */
+	bool operator!=( const TRay<T>& Other ) const
+	{
+		return !(*this == Other);
+	}
+
 
 public:
 
@@ -129,6 +167,44 @@ public:
 	}
 
 
+
+public:
+
+	/**
+	* Get a textual representation of the Ray.
+	*
+	* @return Text describing the Ray.
+	*/
+	FString ToString() const
+	{
+		return FString::Printf(TEXT("Origin=(%s), Direction=(%s)"), *Origin.ToString(), *Direction.ToString());
+	}
+
+	/**
+	 * Serializes the Ray.
+	 *
+	 * @param Ar The archive to serialize into.
+	 * @param Box The box to serialize.
+	 *
+	 * @return Reference to the Archive after serialization.
+	 */
+	friend FArchive& operator<<( FArchive& Ar, TRay<T>& Ray )
+	{
+		return Ar << Ray.Origin << Ray.Direction;
+	}
+
+	// Note: TRay is usually written via binary serialization. This function exists for SerializeFromMismatchedTag conversion usage. 
+	bool Serialize(FArchive& Ar)
+	{
+		Ar << *this;
+		return true;
+	}
+
+	bool SerializeFromMismatchedTag(FName StructTag, FArchive& Ar);
+
+	// Conversion to other type.
+	template<typename FArg, TEMPLATE_REQUIRES(!TIsSame<T, FArg>::Value)>
+	explicit TRay(const TRay<FArg>& From) : TRay<T>(TVector<T>(From.Origin), TVector<T>(From.Direction), true) {}
 };
 
 }	// namespace UE::Math
@@ -136,6 +212,21 @@ public:
 
 UE_DECLARE_LWC_TYPE(Ray, 3);
 
+template<> struct TIsPODType<FRay3f> { enum { Value = true }; };
+template<> struct TIsPODType<FRay3d> { enum { Value = true }; };
 template<> struct TIsUECoreVariant<FRay3f> { enum { Value = true }; };
 template<> struct TIsUECoreVariant<FRay3d> { enum { Value = true }; };
+template<> struct TCanBulkSerialize<FRay3f> { enum { Value = true }; };
+template<> struct TCanBulkSerialize<FRay3d> { enum { Value = true }; };
 
+template<>
+inline bool FRay3f::SerializeFromMismatchedTag(FName StructTag, FArchive& Ar)
+{
+	return UE_SERIALIZE_VARIANT_FROM_MISMATCHED_TAG(Ar, Ray, Ray3f, Ray3d);
+}
+
+template<>
+inline bool FRay3d::SerializeFromMismatchedTag(FName StructTag, FArchive& Ar)
+{
+	return UE_SERIALIZE_VARIANT_FROM_MISMATCHED_TAG(Ar, Ray, Ray3d, Ray3f);
+}
diff --git a/Engine/Source/Runtime/Core/Public/UObject/UnrealNames.inl b/Engine/Source/Runtime/Core/Public/UObject/UnrealNames.inl
--- a/Engine/Source/Runtime/Core/Public/UObject/UnrealNames.inl
+++ b/Engine/Source/Runtime/Core/Public/UObject/UnrealNames.inl
@@ -134,6 +134,9 @@ REGISTER_NAME(165, Int32Point)
 REGISTER_NAME(166, Int64Point)
 REGISTER_NAME(167, Uint32Point)
 REGISTER_NAME(168, Uint64Point)
+REGISTER_NAME(169, Ray)
+REGISTER_NAME(170, Ray3f)
+REGISTER_NAME(171, Ray3d)
 
 // Misc.
 REGISTER_NAME(200,State)
diff --git a/Engine/Source/Runtime/CoreUObject/Private/UObject/Property.cpp b/Engine/Source/Runtime/CoreUObject/Private/UObject/Property.cpp
--- a/Engine/Source/Runtime/CoreUObject/Private/UObject/Property.cpp
+++ b/Engine/Source/Runtime/CoreUObject/Private/UObject/Property.cpp
@@ -21,6 +21,7 @@
 #include "UObject/CoreRedirects.h"
 #include "UObject/SoftObjectPath.h"
 #include "Math/Box2D.h"
+#include "Math/Ray.h"
 #include "Math/InterpCurvePoint.h"
 #include "UObject/ReleaseObjectVersion.h"
 
@@ -331,6 +332,25 @@ struct TStructOpsTypeTraits<FTwoVectors> : public TStructOpsTypeTraitsBase2<FTwo
 };
 UE_IMPLEMENT_STRUCT("/Script/CoreUObject", TwoVectors);
 
+
+template<typename T>
+struct TRay3StructOpsTypeTraits : public TStructOpsTypeTraitsBase2<T>
+{
+	enum
+	{
+		WithIdenticalViaEquality = true,
+		WithNoInitConstructor = true,
+		WithZeroConstructor = true,
+		WithSerializeFromMismatchedTag = true,		
+	};
+};
+template<> struct TStructOpsTypeTraits<FRay3f> : public TRay3StructOpsTypeTraits<FRay3f> {};
+template<> struct TStructOpsTypeTraits<FRay3d> : public TRay3StructOpsTypeTraits<FRay3d> {};
+UE_IMPLEMENT_STRUCT("/Script/CoreUObject", Ray3f);
+UE_IMPLEMENT_STRUCT("/Script/CoreUObject", Ray3d);
+UE_IMPLEMENT_STRUCT("/Script/CoreUObject", Ray);
+
+
 template<>
 struct TStructOpsTypeTraits<FInterpCurvePointFloat> : public TStructOpsTypeTraitsBase2<FInterpCurvePointFloat>
 {
diff --git a/Engine/Source/Runtime/CoreUObject/Public/UObject/Class.h b/Engine/Source/Runtime/CoreUObject/Public/UObject/Class.h
--- a/Engine/Source/Runtime/CoreUObject/Public/UObject/Class.h
+++ b/Engine/Source/Runtime/CoreUObject/Public/UObject/Class.h
@@ -4120,5 +4120,6 @@ UE_DECLARE_CORE_VARIANT_TYPE(Rotator3,	Rotator);
 UE_DECLARE_CORE_VARIANT_TYPE(Transform3,Transform);
 UE_DECLARE_CORE_VARIANT_TYPE(Matrix44,	Matrix);
 UE_DECLARE_CORE_VARIANT_TYPE(Box2,		Box2D);
+UE_DECLARE_CORE_VARIANT_TYPE(Ray3,		Ray);
 
 #undef UE_DECLARE_CORE_VARIANT_TYPE
\ No newline at end of file
diff --git a/Engine/Source/Runtime/CoreUObject/Public/UObject/NoExportTypes.h b/Engine/Source/Runtime/CoreUObject/Public/UObject/NoExportTypes.h
--- a/Engine/Source/Runtime/CoreUObject/Public/UObject/NoExportTypes.h
+++ b/Engine/Source/Runtime/CoreUObject/Public/UObject/NoExportTypes.h
@@ -673,6 +673,52 @@ struct FPlane : public FVector
 };
 
 
+
+/**
+ * 3D Ray represented by Origin and (normalized) Direction.
+ * @note The full C++ class is located here: Engine\Source\Runtime\Core\Public\Math\Ray.h
+ * @note FRay3f is not currently exposed as a Blueprint type
+ */
+USTRUCT(immutable, noexport, IsAlwaysAccessible, HasDefaults, HasNoOpConstructor, IsCoreType)
+struct FRay3f
+{
+	UPROPERTY(EditAnywhere, Category = Ray, SaveGame)
+	FVector3f Origin;
+
+	UPROPERTY(EditAnywhere, Category = Ray, SaveGame)
+	FVector3f Direction;
+};
+
+/**
+ * 3D Ray represented by Origin and (normalized) Direction.
+ * @note The full C++ class is located here: Engine\Source\Runtime\Core\Public\Math\Ray.h
+ */
+USTRUCT(immutable, noexport, IsAlwaysAccessible, HasDefaults, HasNoOpConstructor, IsCoreType)
+struct FRay3d
+{
+	UPROPERTY(EditAnywhere, Category = Ray, SaveGame)
+	FVector3d Origin;
+
+	UPROPERTY(EditAnywhere, Category = Ray, SaveGame)
+	FVector3d Direction;
+};
+
+/**
+ * 3D Ray represented by Origin and (normalized) Direction.
+ * @note The full C++ class is located here: Engine\Source\Runtime\Core\Public\Math\Ray.h
+ */
+USTRUCT(immutable, noexport, BlueprintType, IsAlwaysAccessible, HasDefaults, HasNoOpConstructor, IsCoreType)
+struct FRay
+{
+	UPROPERTY(EditAnywhere, BlueprintReadWrite, Category = Ray, SaveGame)
+	FVector Origin;
+
+	UPROPERTY(EditAnywhere, BlueprintReadWrite, Category = Ray, SaveGame)
+	FVector Direction;
+};
+
+
+
 /**
  * An orthogonal rotation in 3d space.
  * @note The full C++ class is located here: Engine\Source\Runtime\Core\Public\Math\Rotator.h
