Fixed TStringConversion move constructor by deleting it.

diff --git a/Engine/Source/Runtime/Core/Public/Containers/StringConv.h b/Engine/Source/Runtime/Core/Public/Containers/StringConv.h
--- a/Engine/Source/Runtime/Core/Public/Containers/StringConv.h
+++ b/Engine/Source/Runtime/Core/Public/Containers/StringConv.h
@@ -673,6 +673,12 @@ class TStringConversion : private Converter, private TInlineAllocator<DefaultCon
 	}
 
 public:
+	// Non-copyable, non-movable
+	TStringConversion(TStringConversion&&) = delete;
+	TStringConversion(const TStringConversion&) = delete;
+	TStringConversion& operator=(TStringConversion&&) = delete;
+	TStringConversion& operator=(const TStringConversion&) = delete;
+
 	template <
 		typename SrcBufferType,
 		std::enable_if_t<TIsCharEncodingCompatibleWith_V<SrcBufferType, FromType>>* = nullptr
@@ -742,20 +748,6 @@ public:
 	{
 	}
 
-	/**
-	 * Move constructor
-	 */
-	TStringConversion(TStringConversion&& Other)
-		: Converter(MoveTemp(ImplicitConv<Converter&>(Other)))
-		, Ptr(Other.Ptr)
-		, StringLength(Other.StringLength)
-	{
-		AllocatorType::MoveToEmpty(Other);
-
-		Other.Ptr          = nullptr;
-		Other.StringLength = 0;
-	}
-
 	/**
 	 * Accessor for the converted string.
 	 * @note The string may not be null-terminated if constructed from an explicitly sized buffer that didn't include the null-terminator.
@@ -778,10 +770,6 @@ public:
 	}
 
 private:
-	// Non-copyable
-	TStringConversion(const TStringConversion&) = delete;
-	TStringConversion& operator=(const TStringConversion&) = delete;
-
 	ToType* Ptr;
 	int32   StringLength;
 };
