Add more info to help track down check triggering in Stats

diff --git a/Engine/Source/Runtime/Core/Private/Stats/StatsData.cpp b/Engine/Source/Runtime/Core/Private/Stats/StatsData.cpp
--- a/Engine/Source/Runtime/Core/Private/Stats/StatsData.cpp
+++ b/Engine/Source/Runtime/Core/Private/Stats/StatsData.cpp
@@ -1030,6 +1030,22 @@ void FStatsThreadState::AddToHistoryAndEmpty(FStatPacketArray& NewData)
 	check(History.Num() <= HistoryFrames * 2 + 5);
 	check(CondensedStackHistory.Num() <= HistoryFrames * 2 + 5);
 	check(GoodFrames.Num() <= HistoryFrames * 2 + 5);
+#if DO_CHECK
+	if (BadFrames.Num() > HistoryFrames * 2 + 5)
+	{
+		UE_LOG(LogStats, Display, TEXT("BadFrames.Num(): %d, History.Num(): %d, HistoryFrames: %d"), BadFrames.Num(), History.Num(), HistoryFrames);
+		UE_LOG(LogStats, Display, TEXT("CurrentGameFrame: %lld, CurrentRenderFrame: %lld, LastFullFrameMetaAndNonFrame: %lld"), CurrentGameFrame, CurrentRenderFrame, LastFullFrameMetaAndNonFrame);
+		for (int64 BadFrame : BadFrames)
+		{
+			UE_LOG(LogStats, Display, TEXT("Bad frame: %lld"), BadFrame);
+		}
+		for (auto It = History.CreateIterator(); It; ++It)
+		{
+			const int64 HistoryFrame = It.Key();
+			UE_LOG(LogStats, Display, TEXT("History frame: %lld"), HistoryFrame);
+		}
+	}
+#endif
 	check(BadFrames.Num() <= HistoryFrames * 2 + 5);
 }
 
