[core] fix copy-paste error in fixed-width-to-fixed-width string convert template

diff --git a/Engine/Source/Runtime/Core/Public/GenericPlatform/GenericPlatformString.h b/Engine/Source/Runtime/Core/Public/GenericPlatform/GenericPlatformString.h
--- a/Engine/Source/Runtime/Core/Public/GenericPlatform/GenericPlatformString.h
+++ b/Engine/Source/Runtime/Core/Public/GenericPlatform/GenericPlatformString.h
@@ -214,9 +214,9 @@ struct FGenericPlatformString : public FGenericPlatformStricmp
 		}
 		else if constexpr (TIsFixedWidthCharEncoding_V<SourceEncoding> && TIsFixedWidthCharEncoding_V<DestEncoding>)
 		{
-			DestEncoding* DestCopy     = Dest;
-			DestEncoding* SrcCopy      = Src;
-			int32         DestSizeCopy = DestSize;
+			DestEncoding*         DestCopy     = Dest;
+			const SourceEncoding* SrcCopy      = Src;
+			int32                 DestSizeCopy = DestSize;
 
 			bool bInvalidChars = false;
 			for (;;)
