Added a dll cache for the directories added using AddDllDirectory.

LoadLibraryWithSearchPaths:
Before:13s 615ms
After:9s 199

diff --git a/Engine/Source/Runtime/Core/Private/Windows/WindowsPlatformProcess.cpp b/Engine/Source/Runtime/Core/Private/Windows/WindowsPlatformProcess.cpp
--- a/Engine/Source/Runtime/Core/Private/Windows/WindowsPlatformProcess.cpp
+++ b/Engine/Source/Runtime/Core/Private/Windows/WindowsPlatformProcess.cpp
@@ -21,6 +21,7 @@
 #include "Misc/Parse.h"
 #include "Misc/Paths.h"
 #include "Misc/SingleThreadEvent.h"
+#include "Misc/PathViews.h"
 #include "Misc/TrackedActivity.h"
 #include "ProfilingDebugging/CpuProfilerTrace.h"
 #include "ProfilingDebugging/CsvProfiler.h"
@@ -46,6 +47,7 @@ PRAGMA_DISABLE_UNSAFE_TYPECAST_WARNINGS
 // static variables
 TArray<FString> FWindowsPlatformProcess::DllDirectoryStack;
 TArray<FString> FWindowsPlatformProcess::DllDirectories;
+TMap<FName, TArray<FString>> FWindowsPlatformProcess::SearchPathDllCache;
 bool IsJobObjectSet = false;
 HANDLE GhJob = NULL;
 
@@ -97,7 +99,24 @@ void FWindowsPlatformProcess::AddDllDirectory(const TCHAR* Directory)
 	FString NormalizedDirectory = FPaths::ConvertRelativePathToFull(Directory);
 	FPaths::NormalizeDirectoryName(NormalizedDirectory);
 	FPaths::MakePlatformFilename(NormalizedDirectory);
-	DllDirectories.AddUnique(NormalizedDirectory);
+
+	if (DllDirectories.Find(NormalizedDirectory) == INDEX_NONE)
+	{
+		DllDirectories.Add(NormalizedDirectory);
+
+		// enumerate the dir and cache all the dlls
+		{
+			TArray<FString> FoundDllFileNames;
+			IFileManager::Get().FindFiles(FoundDllFileNames, *NormalizedDirectory, TEXT("*.dll"));
+			for (const FString& DllFileName : FoundDllFileNames)
+			{
+				TArray<FString>& Paths = SearchPathDllCache.FindOrAdd(*DllFileName);
+				FString DllPath(NormalizedDirectory / DllFileName);
+				FPaths::NormalizeDirectoryName(DllPath);
+				Paths.Add(DllPath);
+			}
+		}
+	}
 }
 
 void FWindowsPlatformProcess::GetDllDirectories(TArray<FString>& OutDllDirectories)
@@ -110,7 +129,9 @@ void* FWindowsPlatformProcess::GetDllHandle( const TCHAR* FileName )
 	check(FileName);
 
 	// Combine the explicit DLL search directories with the contents of the directory stack 
+	// Note that the search path logic here needs to match the logic found in ResolveImport
 	TArray<FString> SearchPaths;
+	SearchPaths.Reserve(1 + ((DllDirectoryStack.Num() > 0) ? 1 : 0) + DllDirectories.Num());
 	SearchPaths.Add(FPlatformProcess::GetModulesDirectory());
 	if(DllDirectoryStack.Num() > 0)
 	{
@@ -1894,26 +1915,69 @@ static bool ReadLibraryImports(const TCHAR* FileName, TArray<FString>& ImportNam
 	return bResult;
 }
 
-/**
- * Resolve an individual import.
- *
- * @param ImportName Name of the imported module
- * @param SearchPaths Search directories to scan for imports
- * @param OutFileName On success, receives the path to the imported file
- * @return true if an import was found.
- */
-static bool ResolveImport(const FString& Name, const TArray<FString>& SearchPaths, FString& OutFileName)
+bool FWindowsPlatformProcess::ResolveImport(const FString& Name, const TArray<FString>& SearchPaths, FString& OutFileName)
 {
-	// Look for the named DLL on any of the search paths
-	for(int Idx = 0; Idx < SearchPaths.Num(); Idx++)
+	auto SearchPathsFunc = [&OutFileName, &SearchPaths, &Name](int StartIdx, int EndIdx)
 	{
-		FString FileName = SearchPaths[Idx] / Name;
-		if(FPaths::FileExists(FileName))
+		IFileManager& FileManager = IFileManager::Get();
+		for (int Idx = StartIdx; Idx < EndIdx; Idx++)
 		{
-			OutFileName = FPaths::ConvertRelativePathToFull(FileName);
-			return true;
+			TStringBuilder<MAX_PATH> FileName;
+			FPathViews::Append(FileName, SearchPaths[Idx], Name);
+			if (FileManager.FileExists(*FileName))
+			{
+				OutFileName = FPaths::ConvertRelativePathToFull(*FileName);
+				return true;
+			}
+		}
+		return false;
+	};
+
+	// Search the module and current dll directories found in the search path array first.
+	// Note that there is an assumption that the first slots in the array are the module and
+	// current dll directories.
+	const int FirstAddedSearchPathIdx = 1 + ((DllDirectoryStack.Num() > 0) ? 1 : 0);
+	if (SearchPathsFunc(0, FirstAddedSearchPathIdx))
+	{
+		return true;
+	}
+
+	// Search the dll cache that gets populated by AddDllDirectory
+	FName DllName(*Name, FNAME_Find);
+	if (DllName != NAME_None)
+	{
+		if (TArray<FString>* CachedPaths = SearchPathDllCache.Find(DllName))
+		{
+			for (auto Itr = CachedPaths->CreateIterator(); Itr; ++Itr)
+			{
+				const FString& FoundPath = *Itr;
+				// Double check the dll still exists
+				if (FPaths::FileExists(FoundPath))
+				{
+					OutFileName = FoundPath;
+					return true;
+				}
+				else
+				{
+					// The dll cache is out of date
+					Itr.RemoveCurrent();
+				}
+			}
+
+			// Remove invalid entry
+			if (CachedPaths->Num() == 0)
+			{
+				SearchPathDllCache.Remove(DllName);
+			}
 		}
 	}
+
+	// Fall back to going through the search paths
+	if (SearchPathsFunc(FirstAddedSearchPathIdx, SearchPaths.Num()))
+	{
+		return true;
+	}
+
 	return false;
 }
 
@@ -1925,15 +1989,7 @@ UE_TRACE_EVENT_END()
 
 #endif // CPUPROFILERTRACE_ENABLED
 
-/**
- * Resolve all the imports for the given library, searching through a set of directories.
- *
- * @param FileName Path to the library to load
- * @param SearchPaths Search directories to scan for imports
- * @param ImportFileNames Array which is filled with a list of the resolved imports found in the given search directories
- * @param VisitedImportNames Array which stores a list of imports which have been checked
- */
-static void ResolveMissingImportsRecursive(const FString& FileName, const TArray<FString>& SearchPaths, TArray<FString>& ImportFileNames, TSet<FString>& VisitedImportNames)
+void FWindowsPlatformProcess::ResolveMissingImportsRecursive(const FString& FileName, const TArray<FString>& SearchPaths, TArray<FString>& ImportFileNames, TSet<FString>& VisitedImportNames)
 {
 #if CPUPROFILERTRACE_ENABLED
 	UE_TRACE_LOG_SCOPED_T(Cpu, ResolveMissingImports, CpuChannel)
@@ -2017,6 +2073,12 @@ void *FWindowsPlatformProcess::LoadLibraryWithSearchPaths(const FString& FileNam
 		// Convert it to a full path, since LoadLibrary will try to resolve it against the executable directory (which may not be the same as the working dir)
 		FullFileName = FPaths::ConvertRelativePathToFull(FullFileName);
 
+		// If this library is already loaded then just return now with the handle
+		if (void* Handle = GetModuleHandle(*FullFileName))
+		{
+			return Handle;
+		}
+
 		// Create a list of files which we've already checked for imports. Don't add the initial file to this list to improve the resolution of dependencies for direct circular dependencies of this
 		// module; by allowing the module to be visited twice, any mutually depended on DLLs will be visited first.
 		TSet<FString> VisitedImportNames;
diff --git a/Engine/Source/Runtime/Core/Public/Windows/WindowsPlatformProcess.h b/Engine/Source/Runtime/Core/Public/Windows/WindowsPlatformProcess.h
--- a/Engine/Source/Runtime/Core/Public/Windows/WindowsPlatformProcess.h
+++ b/Engine/Source/Runtime/Core/Public/Windows/WindowsPlatformProcess.h
@@ -227,6 +227,11 @@ private:
 	 * All the DLL directories we want to load from. 
 	 */
 	static TArray<FString> DllDirectories;
+	
+	/**
+	 * A cache of the dlls found in each directory in DllDirectories. 
+	 */
+	static TMap<FName, TArray<FString>> SearchPathDllCache;
 
 	/**
 	 * Replacement implementation of the Win32 LoadLibrary function which searches the given list of directories for dependent imports, and attempts
@@ -236,6 +241,25 @@ private:
 	 * @param SearchPaths Search directories to scan for imports
 	 */
 	static void* LoadLibraryWithSearchPaths(const FString& FileName, const TArray<FString>& SearchPaths);
+
+	/**
+	 * Resolve an individual import.
+	 *
+	 * @param ImportName Name of the imported module
+	 * @param OutFileName On success, receives the path to the imported file
+	 * @return true if an import was found.
+	 */
+	static bool ResolveImport(const FString& Name, const TArray<FString>& SearchPaths, FString& OutFileName);
+
+	/**
+	 * Resolve all the imports for the given library, searching through a set of directories.
+	 *
+	 * @param FileName Path to the library to load
+	 * @param SearchPaths Search directories to scan for imports
+	 * @param ImportFileNames Array which is filled with a list of the resolved imports found in the given search directories
+	 * @param VisitedImportNames Array which stores a list of imports which have been checked
+	 */
+	static void ResolveMissingImportsRecursive(const FString& FileName, const TArray<FString>& SearchPaths, TArray<FString>& ImportFileNames, TSet<FString>& VisitedImportNames);
 };
 
 #if WINDOWS_USE_FEATURE_PLATFORMPROCESS_CLASS
