I/O Dispatcher - move chunk ID related types/functions to make IoChunkId.h independant of IoDispatcher.h

diff --git a/Engine/Source/Runtime/Core/Private/IO/IoChunkId.cpp b/Engine/Source/Runtime/Core/Private/IO/IoChunkId.cpp
new file mode 100644
--- /dev/null
+++ b/Engine/Source/Runtime/Core/Private/IO/IoChunkId.cpp
@@ -0,0 +1,88 @@
+// Copyright Epic Games, Inc. All Rights Reserved.
+
+#include "IO/IoChunkId.h"
+#include "IO/PackageId.h"
+#include "Hash/Blake3.h"
+#include "Serialization/Archive.h"
+
+FString LexToString(const EIoChunkType Type)
+{
+	const TCHAR* Strings[] = {
+		TEXT("Invalid"),
+		TEXT("ExportBundleData"),
+		TEXT("BulkData"),
+		TEXT("OptionalBulkData"),
+		TEXT("MemoryMappedBulkData"),
+		TEXT("ScriptObjects"),
+		TEXT("ContainerHeader"),
+		TEXT("ExternalFile"),
+		TEXT("ShaderCodeLibrary"),
+		TEXT("ShaderCode"),
+		TEXT("PackageStoreEntry"),
+		TEXT("DerivedData"),
+		TEXT("EditorDerivedData"),
+		TEXT("PackageResource"),
+	};
+	static_assert(UE_ARRAY_COUNT(Strings) == (SIZE_T)EIoChunkType::MAX);
+	
+	const uint8 Index = (uint8)Type;
+	if (Index < UE_ARRAY_COUNT(Strings))
+	{
+		return Strings[Index];
+	}
+	else
+	{
+		return Strings[0]; // return Invalid
+	}
+}
+
+FString LexToString(const FIoChunkId& Id)
+{
+	FString Output;
+	TArray<TCHAR, FString::AllocatorType>& CharArray = Output.GetCharArray();
+	CharArray.AddUninitialized(sizeof(FIoChunkId) * 2 + 1);
+	UE::String::BytesToHexLower(Id.Id, CharArray.GetData());
+	CharArray.Last() = TCHAR('\0');
+	return Output;
+}
+
+FArchive& operator<<(FArchive& Ar, FIoChunkId& ChunkId)
+{
+	Ar.Serialize(&ChunkId.Id, sizeof FIoChunkId::Id);
+	return Ar;
+}
+
+FIoChunkId CreatePackageDataChunkId(const FPackageId& PackageId)
+{
+	return CreateIoChunkId(PackageId.Value(), 0, EIoChunkType::ExportBundleData);
+}
+
+FIoChunkId CreateExternalFileChunkId(const FStringView Filename)
+{
+	check(Filename.Len() > 0);
+
+	TArray<TCHAR, TInlineAllocator<FName::StringBufferSize>> Buffer;
+	Buffer.SetNum(Filename.Len());
+
+	for (int32 Idx = 0, Len = Filename.Len(); Idx < Len; Idx++)
+	{
+		TCHAR Char = TChar<TCHAR>::ToLower(Filename[Idx]);
+		if (Char == TEXT('\\'))
+		{
+			Char = TEXT('/');
+		}
+
+		Buffer[Idx] = Char;
+	}
+
+	const FBlake3Hash Hash = FBlake3::HashBuffer(FMemoryView(Buffer.GetData(), Buffer.Num() * Buffer.GetTypeSize()));
+
+	uint8 Id[12] = {0};
+	FMemory::Memcpy(Id, Hash.GetBytes(), 11);
+	Id[11] = uint8(EIoChunkType::ExternalFile);
+
+	FIoChunkId ChunkId;
+	ChunkId.Set(Id, sizeof(Id));
+
+	return ChunkId;
+}
diff --git a/Engine/Source/Runtime/Core/Private/IO/IoDispatcher.cpp b/Engine/Source/Runtime/Core/Private/IO/IoDispatcher.cpp
--- a/Engine/Source/Runtime/Core/Private/IO/IoDispatcher.cpp
+++ b/Engine/Source/Runtime/Core/Private/IO/IoDispatcher.cpp
@@ -19,8 +19,6 @@
 #include "ProfilingDebugging/CsvProfiler.h"
 #include "Containers/Ticker.h"
 #include "IO/IoDispatcherBackend.h"
-#include "Hash/Blake3.h"
-#include "IO/PackageId.h"
 
 DEFINE_LOG_CATEGORY(LogIoDispatcher);
 
@@ -1233,38 +1231,3 @@ FIoRequest::Release()
 		Impl = nullptr;
 	}
 }
-
-FIoChunkId CreatePackageDataChunkId(const FPackageId& PackageId)
-{
-	return CreateIoChunkId(PackageId.Value(), 0, EIoChunkType::ExportBundleData);
-}
-
-FIoChunkId CreateExternalFileChunkId(const FStringView Filename)
-{
-	check(Filename.Len() > 0);
-
-	TArray<TCHAR, TInlineAllocator<FName::StringBufferSize>> Buffer;
-	Buffer.SetNum(Filename.Len());
-
-	for (int32 Idx = 0, Len = Filename.Len(); Idx < Len; Idx++)
-	{
-		TCHAR Char = TChar<TCHAR>::ToLower(Filename[Idx]);
-		if (Char == TEXT('\\'))
-		{
-			Char = TEXT('/');
-		}
-
-		Buffer[Idx] = Char;
-	}
-
-	const FBlake3Hash Hash = FBlake3::HashBuffer(FMemoryView(Buffer.GetData(), Buffer.Num() * Buffer.GetTypeSize()));
-
-	uint8 Id[12] = {0};
-	FMemory::Memcpy(Id, Hash.GetBytes(), 11);
-	Id[11] = uint8(EIoChunkType::ExternalFile);
-
-	FIoChunkId ChunkId;
-	ChunkId.Set(Id, sizeof(Id));
-
-	return ChunkId;
-}
diff --git a/Engine/Source/Runtime/Core/Private/IO/IoStore.cpp b/Engine/Source/Runtime/Core/Private/IO/IoStore.cpp
--- a/Engine/Source/Runtime/Core/Private/IO/IoStore.cpp
+++ b/Engine/Source/Runtime/Core/Private/IO/IoStore.cpp
@@ -37,52 +37,6 @@ constexpr char FIoStoreTocHeader::TocMagicImg[];
 
 //TRACE_DECLARE_INT_COUNTER(IoStoreAvailableCompressionBuffers, TEXT("IoStore/AvailableCompressionBuffers"));
 
-FString LexToString(const EIoChunkType Type)
-{
-	const TCHAR* Strings[] = {
-		TEXT("Invalid"),
-		TEXT("ExportBundleData"),
-		TEXT("BulkData"),
-		TEXT("OptionalBulkData"),
-		TEXT("MemoryMappedBulkData"),
-		TEXT("ScriptObjects"),
-		TEXT("ContainerHeader"),
-		TEXT("ExternalFile"),
-		TEXT("ShaderCodeLibrary"),
-		TEXT("ShaderCode"),
-		TEXT("PackageStoreEntry"),
-		TEXT("DerivedData"),
-		TEXT("EditorDerivedData"),
-		TEXT("PackageResource"),
-	};
-	static_assert(UE_ARRAY_COUNT(Strings) == (SIZE_T)EIoChunkType::MAX);
-	uint8 Index = (uint8)Type;
-	if (Index < UE_ARRAY_COUNT(Strings))
-	{
-		return Strings[Index];
-	}
-	else
-	{
-		return Strings[0]; // return Invalid
-	}
-}
-
-FArchive& operator<<(FArchive& Ar, FIoChunkId& ChunkId)
-{
-	Ar.Serialize(&ChunkId.Id, sizeof FIoChunkId::Id);
-	return Ar;
-}
-
-FString LexToString(const FIoChunkId& Id)
-{
-	FString Output;
-	TArray<TCHAR, FString::AllocatorType>& CharArray = Output.GetCharArray();
-	CharArray.AddUninitialized(sizeof(FIoChunkId) * 2 + 1);
-	UE::String::BytesToHexLower(Id.Id, CharArray.GetData());
-	CharArray.Last() = TCHAR('\0');
-	return Output;
-}
-
 template<typename ArrayType>
 bool WriteArray(IFileHandle* FileHandle, const ArrayType& Array)
 {
diff --git a/Engine/Source/Runtime/Core/Public/IO/IoChunkId.h b/Engine/Source/Runtime/Core/Public/IO/IoChunkId.h
--- a/Engine/Source/Runtime/Core/Public/IO/IoChunkId.h
+++ b/Engine/Source/Runtime/Core/Public/IO/IoChunkId.h
@@ -2,12 +2,45 @@
 
 #pragma once
 
+#include "CoreTypes.h"
 #include "Memory/MemoryView.h"
+#include "Misc/ByteSwap.h"
 #include "Misc/StringBuilder.h"
 #include "String/BytesToHex.h"
 
 class FArchive;
-enum class EIoChunkType : uint8;
+class FPackageId;
+
+/**
+ * Addressable chunk types.
+ * 
+ * The enumerators have explicitly defined values here to encourage backward/forward
+ * compatible updates. 
+ * 
+ * Also note that for certain discriminators, Zen Store will assume certain things
+ * about the structure of the chunk ID so changes must be made carefully.
+ * 
+ */
+enum class EIoChunkType : uint8
+{
+	Invalid = 0,
+	ExportBundleData = 1,
+	BulkData = 2,
+	OptionalBulkData = 3,
+	MemoryMappedBulkData = 4,
+	ScriptObjects = 5,
+	ContainerHeader = 6,
+	ExternalFile = 7,
+	ShaderCodeLibrary = 8,
+	ShaderCode = 9,
+	PackageStoreEntry = 10,
+	DerivedData = 11,
+	EditorDerivedData = 12,
+	PackageResource = 13,
+	MAX
+};
+
+CORE_API FString LexToString(const EIoChunkType Type);
 
 /**
  * Identifier to a chunk of data.
@@ -87,3 +120,26 @@ private:
 
 	uint8	Id[12];
 };
+
+/** Creates a chunk identifier (generic -- prefer specialized versions where possible). */
+static FIoChunkId CreateIoChunkId(uint64 ChunkId, uint16 ChunkIndex, EIoChunkType IoChunkType)
+{
+	checkSlow(IoChunkType != EIoChunkType::ExternalFile);	// Use CreateExternalFileChunkId() instead
+
+	uint8 Data[12] = {0};
+
+	*reinterpret_cast<uint64*>(&Data[0]) = ChunkId;
+	*reinterpret_cast<uint16*>(&Data[8]) = NETWORK_ORDER16(ChunkIndex);
+	*reinterpret_cast<uint8*>(&Data[11]) = static_cast<uint8>(IoChunkType);
+
+	FIoChunkId IoChunkId;
+	IoChunkId.Set(Data, 12);
+
+	return IoChunkId;
+}
+
+/** Returns a package data I/O chunk ID for the specified package ID. */
+CORE_API FIoChunkId CreatePackageDataChunkId(const FPackageId& PackageId);
+
+/** Returns a file data I/O chunk ID for the specified filename. */
+CORE_API FIoChunkId CreateExternalFileChunkId(const FStringView Filename);
diff --git a/Engine/Source/Runtime/Core/Public/IO/IoDispatcher.h b/Engine/Source/Runtime/Core/Public/IO/IoDispatcher.h
--- a/Engine/Source/Runtime/Core/Public/IO/IoDispatcher.h
+++ b/Engine/Source/Runtime/Core/Public/IO/IoDispatcher.h
@@ -658,59 +658,6 @@ private:
 	uint8	Hash[32];
 };
 
-/**
- * Addressable chunk types.
- * 
- * The enumerators have explicitly defined values here to encourage backward/forward
- * compatible updates. 
- * 
- * Also note that for certain discriminators, Zen Store will assume certain things
- * about the structure of the chunk ID so changes must be made carefully.
- * 
- */
-enum class EIoChunkType : uint8
-{
-	Invalid = 0,
-	ExportBundleData = 1,
-	BulkData = 2,
-	OptionalBulkData = 3,
-	MemoryMappedBulkData = 4,
-	ScriptObjects = 5,
-	ContainerHeader = 6,
-	ExternalFile = 7,
-	ShaderCodeLibrary = 8,
-	ShaderCode = 9,
-	PackageStoreEntry = 10,
-	DerivedData = 11,
-	EditorDerivedData = 12,
-	PackageResource = 13,
-	MAX
-};
-
-CORE_API FString LexToString(const EIoChunkType Type);
-
-/**
- * Creates a chunk identifier (generic -- prefer specialized versions where possible)
- */
-static FIoChunkId CreateIoChunkId(uint64 ChunkId, uint16 ChunkIndex, EIoChunkType IoChunkType)
-{
-	checkSlow(IoChunkType != EIoChunkType::ExternalFile);	// Use CreateExternalFileChunkId() instead
-
-	uint8 Data[12] = {0};
-
-	*reinterpret_cast<uint64*>(&Data[0]) = ChunkId;
-	*reinterpret_cast<uint16*>(&Data[8]) = NETWORK_ORDER16(ChunkIndex);
-	*reinterpret_cast<uint8*>(&Data[11]) = static_cast<uint8>(IoChunkType);
-
-	FIoChunkId IoChunkId;
-	IoChunkId.Set(Data, 12);
-
-	return IoChunkId;
-}
-
-CORE_API FIoChunkId CreatePackageDataChunkId(const FPackageId& PackageId);
-CORE_API FIoChunkId CreateExternalFileChunkId(const FStringView Filename);
-
 //////////////////////////////////////////////////////////////////////////
 
 class FIoReadOptions
