Added error message when FMessageDialog::Open or FPlatformMisc::MessageBoxExt fails to create the dialog (ex.: when missing windows resources).

diff --git a/Engine/Source/Runtime/Core/Private/Windows/WindowsPlatformMisc.cpp b/Engine/Source/Runtime/Core/Private/Windows/WindowsPlatformMisc.cpp
--- a/Engine/Source/Runtime/Core/Private/Windows/WindowsPlatformMisc.cpp
+++ b/Engine/Source/Runtime/Core/Private/Windows/WindowsPlatformMisc.cpp
@@ -1466,6 +1466,15 @@ private:
 				return Instance->MessageBoxDlgProc(HandleWnd, Message, WParam, LParam);
 			}, (LPARAM)this);
 
+		if (!DialogHwnd)
+		{
+			DWORD LastError = GetLastError();
+			TCHAR ErrorBuffer[1024];
+			FWindowsPlatformMisc::GetSystemErrorMessage(ErrorBuffer, UE_ARRAY_COUNT(ErrorBuffer), LastError);
+			UE_LOG(LogWindows, Error, TEXT("Failed to create dialog. %s Error: 0x%X (%u)"), ErrorBuffer, LastError, LastError);
+			return Result;
+		}
+
 		ShowWindow(DialogHwnd, SW_SHOW);
 		MSG Msg;
 		while (!WasClosed && GetMessageW(&Msg, NULL, 0, 0))
