FString: Add FStringView overloads for RemoveFromStart, RemoveFromEnd, Find, StartsWith, EndsWith, MatchesWildcard.
TCString: Add Strnstr and Strnistr.
#rnx

diff --git a/Engine/Source/Runtime/Core/Private/Containers/String.cpp b/Engine/Source/Runtime/Core/Private/Containers/String.cpp
--- a/Engine/Source/Runtime/Core/Private/Containers/String.cpp
+++ b/Engine/Source/Runtime/Core/Private/Containers/String.cpp
@@ -395,22 +395,23 @@ void FString::TrimToNullTerminator()
 }
 
 
-int32 FString::Find(const TCHAR* SubStr, ESearchCase::Type SearchCase, ESearchDir::Type SearchDir, int32 StartPosition) const
+int32 FString::Find(const TCHAR* SubStr, int32 SubStrLen, ESearchCase::Type SearchCase, ESearchDir::Type SearchDir, int32 StartPosition) const
 {
-	if (SubStr == nullptr)
-	{
-		return INDEX_NONE;
-	}
-	if( SearchDir == ESearchDir::FromStart)
+	checkf(SubStrLen >= 0, TEXT("Invalid SubStrLen: %d"), SubStrLen);
+
+	if (SearchDir == ESearchDir::FromStart)
 	{
 		const TCHAR* Start = **this;
-		if (StartPosition != INDEX_NONE && Len() > 0)
+		int32 RemainingLength = Len();
+		if (StartPosition != INDEX_NONE && RemainingLength > 0)
 		{
-			Start += FMath::Clamp(StartPosition, 0, Len() - 1);
+			const TCHAR* End = Start + RemainingLength;
+			Start += FMath::Clamp(StartPosition, 0, RemainingLength - 1);
+			RemainingLength = UE_PTRDIFF_TO_INT32(End - Start);
 		}
 		const TCHAR* Tmp = SearchCase == ESearchCase::IgnoreCase
-			? FCString::Stristr(Start, SubStr)
-			: FCString::Strstr(Start, SubStr);
+			? FCString::Strnistr(Start, RemainingLength, SubStr, SubStrLen)
+			: FCString::Strnstr(Start, RemainingLength, SubStr, SubStrLen);
 
 		return Tmp ? UE_PTRDIFF_TO_INT32(Tmp-**this) : INDEX_NONE;
 	}
@@ -420,29 +421,29 @@ int32 FString::Find(const TCHAR* SubStr, ESearchCase::Type SearchCase, ESearchDi
 		// times in the loop below
 		if ( SearchCase == ESearchCase::IgnoreCase)
 		{
-			return ToUpper().Find(FString(SubStr).ToUpper(), ESearchCase::CaseSensitive, SearchDir, StartPosition);
+			return ToUpper().Find(FString(SubStrLen, SubStr).ToUpper(), ESearchCase::CaseSensitive, SearchDir, StartPosition);
 		}
 		else
 		{
-			const int32 SearchStringLength=FMath::Max(1, FCString::Strlen(SubStr));
+			const int32 SearchStringLength=FMath::Max(1, SubStrLen);
 			
-			if ( StartPosition == INDEX_NONE || StartPosition >= Len() )
+			if (StartPosition == INDEX_NONE || StartPosition >= Len())
 			{
 				StartPosition = Len();
 			}
 			
-			for( int32 i = StartPosition - SearchStringLength; i >= 0; i-- )
+			for (int32 i = StartPosition - SearchStringLength; i >= 0; i--)
 			{
 				int32 j;
-				for( j=0; SubStr[j]; j++ )
+				for (j=0; j != SubStrLen; j++)
 				{
-					if( (*this)[i+j]!=SubStr[j] )
+					if ((*this)[i+j]!=SubStr[j])
 					{
 						break;
 					}
 				}
 				
-				if( !SubStr[j] )
+				if (j == SubStrLen)
 				{
 					return i;
 				}
@@ -561,68 +562,31 @@ void FString::RemoveSpacesInline()
 	}
 }
 
-bool FString::StartsWith(const TCHAR* InPrefix, ESearchCase::Type SearchCase) const
+bool FString::StartsWith(const TCHAR* InPrefix, int32 InPrefixLen, ESearchCase::Type SearchCase) const
 {
 	if (SearchCase == ESearchCase::IgnoreCase)
 	{
-		return InPrefix && *InPrefix && !FCString::Strnicmp(**this, InPrefix, FCString::Strlen(InPrefix));
+		return InPrefixLen > 0 && !FCString::Strnicmp(**this, InPrefix, InPrefixLen);
 	}
 	else
 	{
-		return InPrefix && *InPrefix && !FCString::Strncmp(**this, InPrefix, FCString::Strlen(InPrefix));
+		return InPrefixLen > 0 && !FCString::Strncmp(**this, InPrefix, InPrefixLen);
 	}
 }
 
-bool FString::StartsWith(const FString& InPrefix, ESearchCase::Type SearchCase ) const
+bool FString::EndsWith(const TCHAR* InSuffix, int32 InSuffixLen, ESearchCase::Type SearchCase ) const
 {
-	if( SearchCase == ESearchCase::IgnoreCase )
-	{
-		return InPrefix.Len() > 0 && !FCString::Strnicmp(**this, *InPrefix, InPrefix.Len());
-	}
-	else
-	{
-		return InPrefix.Len() > 0 && !FCString::Strncmp(**this, *InPrefix, InPrefix.Len());
-	}
-}
-
-bool FString::EndsWith(const TCHAR* InSuffix, ESearchCase::Type SearchCase) const
-{
-	if (!InSuffix || *InSuffix == TEXT('\0'))
-	{
-		return false;
-	}
-
-	int32 ThisLen   = this->Len();
-	int32 SuffixLen = FCString::Strlen(InSuffix);
-	if (SuffixLen > ThisLen)
-	{
-		return false;
-	}
-
-	const TCHAR* StrPtr = Data.GetData() + ThisLen - SuffixLen;
 	if (SearchCase == ESearchCase::IgnoreCase)
 	{
-		return !FCString::Stricmp(StrPtr, InSuffix);
+		return InSuffixLen > 0 &&
+			Len() >= InSuffixLen &&
+			!FCString::Stricmp( &(*this)[ Len() - InSuffixLen ], InSuffix );
 	}
 	else
 	{
-		return !FCString::Strcmp(StrPtr, InSuffix);
-	}
-}
-
-bool FString::EndsWith(const FString& InSuffix, ESearchCase::Type SearchCase ) const
-{
-	if( SearchCase == ESearchCase::IgnoreCase )
-	{
-		return InSuffix.Len() > 0 &&
-			Len() >= InSuffix.Len() &&
-			!FCString::Stricmp( &(*this)[ Len() - InSuffix.Len() ], *InSuffix );
-	}
-	else
-	{
-		return InSuffix.Len() > 0 &&
-			Len() >= InSuffix.Len() &&
-			!FCString::Strcmp( &(*this)[ Len() - InSuffix.Len() ], *InSuffix );
+		return InSuffixLen > 0 &&
+			Len() >= InSuffixLen &&
+			!FCString::Strcmp( &(*this)[ Len() - InSuffixLen ], InSuffix );
 	}
 }
 
@@ -661,65 +625,32 @@ void FString::RemoveAt(int32 Index, int32 Count, bool bAllowShrinking)
 	Data.RemoveAt(Index, FMath::Clamp(Count, 0, Len()-Index), bAllowShrinking);
 }
 
-bool FString::RemoveFromStart( const TCHAR* InPrefix, ESearchCase::Type SearchCase )
-{
-	if ( *InPrefix == 0 )
-	{
-		return false;
-	}
-
-	if ( StartsWith( InPrefix, SearchCase ) )
-	{
-		RemoveAt( 0, FCString::Strlen(InPrefix) );
-		return true;
-	}
-
-	return false;
-}
-
-bool FString::RemoveFromStart( const FString& InPrefix, ESearchCase::Type SearchCase )
-{
-	if ( InPrefix.IsEmpty() )
-	{
-		return false;
-	}
-
-	if ( StartsWith( InPrefix, SearchCase ) )
-	{
-		RemoveAt( 0, InPrefix.Len() );
-		return true;
-	}
-
-	return false;
-}
-
-bool FString::RemoveFromEnd( const TCHAR* InSuffix, ESearchCase::Type SearchCase )
+bool FString::RemoveFromStart(const TCHAR* InPrefix, int32 InPrefixLen, ESearchCase::Type SearchCase)
 {
-	if ( *InSuffix == 0 )
+	if (InPrefixLen == 0 )
 	{
 		return false;
 	}
 
-	if ( EndsWith( InSuffix, SearchCase ) )
+	if (StartsWith(InPrefix, InPrefixLen, SearchCase ))
 	{
-		int32 SuffixLen = FCString::Strlen(InSuffix);
-		RemoveAt( Len() - SuffixLen, SuffixLen );
+		RemoveAt(0, InPrefixLen);
 		return true;
 	}
 
 	return false;
 }
 
-bool FString::RemoveFromEnd( const FString& InSuffix, ESearchCase::Type SearchCase )
+bool FString::RemoveFromEnd(const TCHAR* InSuffix, int32 InSuffixLen, ESearchCase::Type SearchCase)
 {
-	if ( InSuffix.IsEmpty() )
+	if (InSuffixLen == 0)
 	{
 		return false;
 	}
 
-	if ( EndsWith( InSuffix, SearchCase ) )
+	if (EndsWith(InSuffix, InSuffixLen, SearchCase))
 	{
-		RemoveAt( Len() - InSuffix.Len(), InSuffix.Len() );
+		RemoveAt( Len() - InSuffixLen, InSuffixLen );
 		return true;
 	}
 
@@ -1381,19 +1312,18 @@ int32 FString::ParseIntoArray( TArray<FString>& OutArray, const TCHAR* pchDelim,
 	return OutArray.Num();
 }
 
-bool FString::MatchesWildcard(const TCHAR* InWildcard, ESearchCase::Type SearchCase) const
+bool FString::MatchesWildcard(const TCHAR* InWildcard, int32 InWildcardLen, ESearchCase::Type SearchCase) const
 {
 	const TCHAR* Target = **this;
 	int32        TargetLength = Len();
-	int32        WildcardLength = FCString::Strlen(InWildcard);
 
 	if (SearchCase == ESearchCase::CaseSensitive)
 	{
-		return UE::Core::String::Private::MatchesWildcardRecursive<UE::Core::String::Private::FCompareCharsCaseSensitive>(Target, TargetLength, InWildcard, WildcardLength);
+		return UE::Core::String::Private::MatchesWildcardRecursive<UE::Core::String::Private::FCompareCharsCaseSensitive>(Target, TargetLength, InWildcard, InWildcardLen);
 	}
 	else
 	{
-		return UE::Core::String::Private::MatchesWildcardRecursive<UE::Core::String::Private::FCompareCharsCaseInsensitive>(Target, TargetLength, InWildcard, WildcardLength);
+		return UE::Core::String::Private::MatchesWildcardRecursive<UE::Core::String::Private::FCompareCharsCaseInsensitive>(Target, TargetLength, InWildcard, InWildcardLen);
 	}
 }
 
diff --git a/Engine/Source/Runtime/Core/Private/Containers/StringTest.cpp b/Engine/Source/Runtime/Core/Private/Containers/StringTest.cpp
--- a/Engine/Source/Runtime/Core/Private/Containers/StringTest.cpp
+++ b/Engine/Source/Runtime/Core/Private/Containers/StringTest.cpp
@@ -610,4 +610,136 @@ bool FStringPathConcatCompoundOperatorTest::RunTest(const FString& Parameters)
 	return true;
 }
 
+IMPLEMENT_SIMPLE_AUTOMATION_TEST(FStringFindTest, "System.Core.String.Find", EAutomationTestFlags::ApplicationContextMask | EAutomationTestFlags::EngineFilter)
+bool FStringFindTest::RunTest(const FString& Parameters)
+{
+	auto RunTest = [this](FStringView Search, FStringView Find, ESearchCase::Type SearchCase,
+		ESearchDir::Type SearchDir, int32 StartPosition, int32 Expected)
+	{
+		FString SearchStr(Search);
+		FString FindStr(Find);
+		// TCHAR*, int32
+		{
+			int32 Actual = SearchStr.Find(Find.GetData(), Find.Len(), SearchCase, SearchDir, StartPosition);
+			if (Actual != Expected)
+			{
+				AddError(FString::Printf(TEXT("FString(\"%s\").Find(\"%s\", %d, %d, %d, %d) returned Actual %d not equal to Expected %d"),
+					*SearchStr, *FindStr, FindStr.Len(), (int32)SearchCase, (int32)SearchDir, StartPosition, Actual, Expected));
+			}
+		}
+		// FStringView
+		{
+			int32 Actual = SearchStr.Find(Find, SearchCase, SearchDir, StartPosition);
+			if (Actual != Expected)
+			{
+				AddError(FString::Printf(TEXT("FString(\"%s\").Find(FStringView(\"%s\", %d), %d, %d, %d) returned Actual %d not equal to Expected %d"),
+					*SearchStr, *FindStr, FindStr.Len(), (int32)SearchCase, (int32)SearchDir, StartPosition));
+			}
+		}
+
+		// TCHAR*, nullterminated
+		{
+			int32 Actual = SearchStr.Find(*FindStr, SearchCase, SearchDir, StartPosition);
+			if (Actual != Expected)
+			{
+				AddError(FString::Printf(TEXT("FString(\"%s\").Find(TEXT(\"%s\"), %d, %d, %d) returned Actual %d not equal to Expected %d"),
+					*SearchStr, *FindStr, (int32)SearchCase, (int32)SearchDir, StartPosition));
+			}
+		}
+
+		// FString
+		{
+			int32 Actual = SearchStr.Find(FindStr, SearchCase, SearchDir, StartPosition);
+			if (Actual != Expected)
+			{
+				AddError(FString::Printf(TEXT("FString(\"%s\").Find(FString(%d, \"%s\"), %d, %d, %d) returned Actual %d not equal to Expected %d"),
+					*SearchStr, FindStr.Len(), *FindStr, (int32)SearchCase, (int32)SearchDir, StartPosition));
+			}
+		}
+	};
+	FStringView ABACADAB(TEXTVIEW("ABACADAB"));
+	FStringView A(TEXTVIEW("A"));
+	FStringView B(TEXTVIEW("B"));
+	FStringView CAD(TEXTVIEW("CAD"));
+	FStringView a(TEXTVIEW("a"));
+	FStringView EmptyString;
+	RunTest(ABACADAB, A, ESearchCase::CaseSensitive, ESearchDir::FromStart, 0, 0);
+	RunTest(ABACADAB, A, ESearchCase::CaseSensitive, ESearchDir::FromStart, 1, 2);
+	RunTest(ABACADAB, A, ESearchCase::CaseSensitive, ESearchDir::FromStart, 7, INDEX_NONE);
+	RunTest(ABACADAB, A, ESearchCase::CaseSensitive, ESearchDir::FromStart, 8, INDEX_NONE);
+	RunTest(ABACADAB, A, ESearchCase::CaseSensitive, ESearchDir::FromStart, 80, INDEX_NONE);
+	RunTest(ABACADAB, A, ESearchCase::CaseSensitive, ESearchDir::FromEnd, 80, 6);
+	RunTest(ABACADAB, A, ESearchCase::CaseSensitive, ESearchDir::FromEnd, 8, 6);
+	RunTest(ABACADAB, A, ESearchCase::CaseSensitive, ESearchDir::FromEnd, 7, 6);
+	RunTest(ABACADAB, A, ESearchCase::CaseSensitive, ESearchDir::FromEnd, 6, 4);
+	RunTest(ABACADAB, A, ESearchCase::CaseSensitive, ESearchDir::FromEnd, 1, 0);
+	RunTest(ABACADAB, A, ESearchCase::CaseSensitive, ESearchDir::FromEnd, 0, INDEX_NONE);
+
+	RunTest(ABACADAB, B, ESearchCase::CaseSensitive, ESearchDir::FromStart, 0, 1);
+	RunTest(ABACADAB, B, ESearchCase::CaseSensitive, ESearchDir::FromStart, 1, 1);
+	RunTest(ABACADAB, B, ESearchCase::CaseSensitive, ESearchDir::FromStart, 2, 7);
+	RunTest(ABACADAB, B, ESearchCase::CaseSensitive, ESearchDir::FromStart, 7, 7);
+	RunTest(ABACADAB, B, ESearchCase::CaseSensitive, ESearchDir::FromStart, 8, 7); // StartPosition clamped to [0, Len-1]
+	RunTest(ABACADAB, B, ESearchCase::CaseSensitive, ESearchDir::FromEnd, 80, 7);
+	RunTest(ABACADAB, B, ESearchCase::CaseSensitive, ESearchDir::FromEnd, 8, 7);
+	RunTest(ABACADAB, B, ESearchCase::CaseSensitive, ESearchDir::FromEnd, 7, 1);
+	RunTest(ABACADAB, B, ESearchCase::CaseSensitive, ESearchDir::FromEnd, 1, INDEX_NONE);
+	RunTest(ABACADAB, B, ESearchCase::CaseSensitive, ESearchDir::FromEnd, 0, INDEX_NONE);
+
+	RunTest(ABACADAB, a, ESearchCase::IgnoreCase, ESearchDir::FromStart, 0, 0);
+	RunTest(ABACADAB, a, ESearchCase::IgnoreCase, ESearchDir::FromStart, 1, 2);
+	RunTest(ABACADAB, a, ESearchCase::IgnoreCase, ESearchDir::FromStart, 7, INDEX_NONE);
+	RunTest(ABACADAB, a, ESearchCase::IgnoreCase, ESearchDir::FromEnd, 8, 6);
+	RunTest(ABACADAB, a, ESearchCase::IgnoreCase, ESearchDir::FromEnd, 1, 0);
+	RunTest(ABACADAB, a, ESearchCase::IgnoreCase, ESearchDir::FromEnd, 0, INDEX_NONE);
+
+	RunTest(ABACADAB, a, ESearchCase::CaseSensitive, ESearchDir::FromStart, 0, INDEX_NONE);
+	RunTest(ABACADAB, a, ESearchCase::CaseSensitive, ESearchDir::FromStart, 1, INDEX_NONE);
+	RunTest(ABACADAB, a, ESearchCase::CaseSensitive, ESearchDir::FromStart, 7, INDEX_NONE);
+	RunTest(ABACADAB, a, ESearchCase::CaseSensitive, ESearchDir::FromEnd, 8, INDEX_NONE);
+	RunTest(ABACADAB, a, ESearchCase::CaseSensitive, ESearchDir::FromEnd, 1, INDEX_NONE);
+	RunTest(ABACADAB, a, ESearchCase::CaseSensitive, ESearchDir::FromEnd, 0, INDEX_NONE);
+
+	RunTest(ABACADAB, CAD, ESearchCase::CaseSensitive, ESearchDir::FromStart, 0, 3);
+	RunTest(ABACADAB, CAD, ESearchCase::CaseSensitive, ESearchDir::FromStart, 3, 3);
+	RunTest(ABACADAB, CAD, ESearchCase::CaseSensitive, ESearchDir::FromStart, 4, INDEX_NONE);
+
+	RunTest(ABACADAB, EmptyString, ESearchCase::CaseSensitive, ESearchDir::FromStart, 0, 0);
+	RunTest(ABACADAB, EmptyString, ESearchCase::CaseSensitive, ESearchDir::FromStart, 4, 4);
+	RunTest(ABACADAB, EmptyString, ESearchCase::CaseSensitive, ESearchDir::FromEnd, 8, 7);
+	RunTest(ABACADAB, EmptyString, ESearchCase::CaseSensitive, ESearchDir::FromEnd, 2, 1);
+	RunTest(ABACADAB, EmptyString, ESearchCase::CaseSensitive, ESearchDir::FromEnd, 0, INDEX_NONE);
+
+	// Find with a null char*
+	int32 Actual = FString(ABACADAB).Find(nullptr, ESearchCase::CaseSensitive, ESearchDir::FromStart, 0);
+	if (Actual != INDEX_NONE)
+	{
+		AddError(FString::Printf(TEXT("FString(\"ABACADAB\").Find(nullptr, 0, %d, %d, %d) returned Actual %d not equal to Expected -1"),
+			(int32)ESearchCase::CaseSensitive, (int32)ESearchDir::FromStart, 0, Actual));
+	}
+	Actual = FString(ABACADAB).Find(nullptr, ESearchCase::CaseSensitive, ESearchDir::FromEnd, 0);
+	if (Actual != INDEX_NONE)
+	{
+		AddError(FString::Printf(TEXT("FString(\"ABACADAB\").Find(nullptr, 0, %d, %d, %d) returned Actual %d not equal to Expected -1"),
+			(int32)ESearchCase::CaseSensitive, (int32)ESearchDir::FromEnd, 0, Actual));
+	}
+
+	// Find with a null char* and a length
+	Actual = FString(ABACADAB).Find(nullptr, 0, ESearchCase::CaseSensitive, ESearchDir::FromStart, 0);
+	if (Actual != 0)
+	{
+		AddError(FString::Printf(TEXT("FString(\"ABACADAB\").Find(nullptr, 0, %d, %d, %d) returned Actual %d not equal to Expected 0"),
+			(int32)ESearchCase::CaseSensitive, (int32)ESearchDir::FromStart, 0, Actual));
+	}
+	Actual = FString(ABACADAB).Find(nullptr, 0, ESearchCase::CaseSensitive, ESearchDir::FromEnd, 8);
+	if (Actual != 7)
+	{
+		AddError(FString::Printf(TEXT("FString(\"ABACADAB\").Find(nullptr, 0, %d, %d, %d) returned Actual %d not equal to Expected 7"),
+			(int32)ESearchCase::CaseSensitive, (int32)ESearchDir::FromEnd, 0, Actual));
+	}
+	// Negative SubStrLen are not allowed so we do not test them
+
+	return true;
+}
+
 #endif // WITH_DEV_AUTOMATION_TESTS
diff --git a/Engine/Source/Runtime/Core/Private/Misc/CStringTest.cpp b/Engine/Source/Runtime/Core/Private/Misc/CStringTest.cpp
--- a/Engine/Source/Runtime/Core/Private/Misc/CStringTest.cpp
+++ b/Engine/Source/Runtime/Core/Private/Misc/CStringTest.cpp
@@ -1,12 +1,16 @@
 // Copyright Epic Games, Inc. All Rights Reserved.
 
-#include "CoreTypes.h"
-#include "Misc/AutomationTest.h"
-#include "Misc/AssertionMacros.h"
 #include "Misc/CString.h"
 
+#include "Misc/AutomationTest.h"
+
 #if WITH_DEV_AUTOMATION_TESTS
 
+#include "Containers/StringView.h"
+#include "CoreTypes.h"
+#include "Misc/AssertionMacros.h"
+#include "Misc/StringBuilder.h"
+
 // This class is a workaround for clang compilers causing error "'va_start' used in function with fixed args" when using it in a lambda in RunTest()
 class FCStringGetVarArgsTestBase : public FAutomationTestBase
 {
@@ -93,4 +97,87 @@ bool FCStringGetVarArgsTest::RunTest(const FString& Parameters)
 	return true;
 }
 
+IMPLEMENT_SIMPLE_AUTOMATION_TEST(FCStringStrnstrTest, "System.Core.Misc.CString.Strnstr", EAutomationTestFlags::ApplicationContextMask | EAutomationTestFlags::EngineFilter)
+bool FCStringStrnstrTest::RunTest(const FString& Parameters)
+{
+	auto RunTest = [this](FStringView Search, FStringView Find, int32 ExpectedSensitiveIndex, int32 ExpectedInsensitiveIndex)
+	{
+		TStringBuilder<128> SearchWithoutNull;
+		TStringBuilder<128> FindWithoutNull;
+		// TODO: If we could add a way to assert the trailing characters are not read, that would be a better test
+		SearchWithoutNull << Search << TEXT("SearchTrailing");
+		FindWithoutNull << Find << TEXT("FindTrailing");
+		const TCHAR* ExpectedSensitive = ExpectedSensitiveIndex == INDEX_NONE ? nullptr :
+			(Search.GetData() + ExpectedSensitiveIndex);
+		const TCHAR* ExpectedInsensitive = ExpectedInsensitiveIndex == INDEX_NONE ? nullptr :
+			(Search.GetData() + ExpectedInsensitiveIndex);
+		const TCHAR* ExpectedSensitiveWithoutNull = ExpectedSensitiveIndex == INDEX_NONE ? nullptr :
+			(SearchWithoutNull.GetData() + ExpectedSensitiveIndex);
+		const TCHAR* ExpectedInsensitiveWithoutNull = ExpectedInsensitiveIndex == INDEX_NONE ? nullptr :
+			(SearchWithoutNull.GetData() + ExpectedInsensitiveIndex);
+		if (FCString::Strnstr(Search.GetData(), Search.Len(), Find.GetData(), Find.Len()) != ExpectedSensitive)
+		{
+			AddError(FString::Printf(TEXT("Strnstr(\"%.*s\", %d, \"%.*s\", %d)\" did not equal index \"%d\"."),
+				Search.Len(), Search.GetData(), Search.Len(), Find.Len(), Find.GetData(), Find.Len(), ExpectedSensitiveIndex));
+		}
+		if (FCString::Strnstr(SearchWithoutNull.GetData(), Search.Len(), FindWithoutNull.GetData(), Find.Len()) != ExpectedSensitiveWithoutNull)
+		{
+			AddError(FString::Printf(TEXT("Strnstr(\"%.*s\", %d, \"%.*s\", %d)\" did not equal index \"%d\", when embedded in a string without a nullterminator."),
+				Search.Len(), Search.GetData(), Search.Len(), Find.Len(), Find.GetData(), Find.Len(), ExpectedSensitiveIndex));
+		}
+		if (FCString::Strnistr(Search.GetData(), Search.Len(), Find.GetData(), Find.Len()) != ExpectedInsensitive)
+		{
+			AddError(FString::Printf(TEXT("Strnistr(\"%.*s\", %d, \"%.*s\", %d)\" did not equal index \"%d\"."),
+				Search.Len(), Search.GetData(), Search.Len(), Find.Len(), Find.GetData(), Find.Len(), ExpectedInsensitiveIndex));
+		}
+		if (FCString::Strnistr(SearchWithoutNull.GetData(), Search.Len(), FindWithoutNull.GetData(), Find.Len()) != ExpectedInsensitiveWithoutNull)
+		{
+			AddError(FString::Printf(TEXT("Strnistr(\"%.*s\", %d, \"%.*s\", %d)\" did not equal index \"%d\", when embedded in a string without a nullterminator."),
+				Search.Len(), Search.GetData(), Search.Len(), Find.Len(), Find.GetData(), Find.Len(), ExpectedInsensitiveIndex));
+		}
+	};
+	FStringView ABACADAB(TEXTVIEW("ABACADAB"));
+
+	RunTest(ABACADAB, TEXTVIEW("A"), 0, 0);
+	RunTest(ABACADAB, TEXTVIEW("a"), INDEX_NONE, 0);
+	RunTest(ABACADAB, TEXTVIEW("BAC"), 1, 1);
+	RunTest(ABACADAB, TEXTVIEW("BaC"), INDEX_NONE, 1);
+	RunTest(ABACADAB, TEXTVIEW("BAC"), 1, 1);
+	RunTest(ABACADAB, TEXTVIEW("BaC"), INDEX_NONE, 1);
+	RunTest(ABACADAB, TEXTVIEW("DAB"), 5, 5);
+	RunTest(ABACADAB, TEXTVIEW("dab"), INDEX_NONE, 5);
+	RunTest(ABACADAB, ABACADAB, 0, 0);
+	RunTest(ABACADAB, TEXTVIEW("abacadab"), INDEX_NONE, 0);
+	RunTest(ABACADAB, TEXTVIEW("F"), INDEX_NONE, INDEX_NONE);
+	RunTest(ABACADAB, TEXTVIEW("DABZ"), INDEX_NONE, INDEX_NONE);
+	RunTest(ABACADAB, TEXTVIEW("ABACADABA"), INDEX_NONE, INDEX_NONE);
+	RunTest(ABACADAB, TEXTVIEW("NoMatchLongerString"), INDEX_NONE, INDEX_NONE);
+	RunTest(TEXTVIEW(""), TEXTVIEW("FindText"), INDEX_NONE, INDEX_NONE);
+	RunTest(TEXTVIEW(""), TEXTVIEW(""), 0, 0);
+	RunTest(ABACADAB, TEXTVIEW(""), 0, 0);
+
+	// Tests that pass in nullptr
+	const TCHAR* NullString = nullptr;
+	const TCHAR* EmptyString = TEXT("");
+	if (FCString::Strnstr(NullString, 0, NullString, 0) != nullptr ||
+		FCString::Strnistr(NullString, 0, NullString, 0) != nullptr)
+	{
+		AddError(TEXT("Strnstr(nullptr, 0, nullptr, 0) did not equal nullptr."));
+	}
+	if (FCString::Strnstr(EmptyString, 0, NullString, 0) != EmptyString ||
+		FCString::Strnistr(EmptyString, 0, NullString, 0) != EmptyString)
+	{
+		AddError(TEXT("Strnstr(EmptyString, 0, nullptr, 0) did not equal EmptyString."));
+	}
+	if (FCString::Strnstr(NullString, 0, EmptyString, 0) != nullptr ||
+		FCString::Strnistr(NullString, 0, EmptyString, 0) != nullptr)
+	{
+		AddError(TEXT("Strnstr(nullptr, 0, EmptyString, 0) did not equal nullptr."));
+	}
+
+	// Negative lengths are not allowed so we do not test them.
+
+	return true;
+}
+
 #endif // WITH_DEV_AUTOMATION_TESTS
diff --git a/Engine/Source/Runtime/Core/Public/Containers/UnrealString.h b/Engine/Source/Runtime/Core/Public/Containers/UnrealString.h
--- a/Engine/Source/Runtime/Core/Public/Containers/UnrealString.h
+++ b/Engine/Source/Runtime/Core/Public/Containers/UnrealString.h
@@ -29,6 +29,8 @@
 #include "Traits/IsCharEncodingCompatibleWith.h"
 #include "Traits/IsCharEncodingSimplyConvertibleTo.h"
 
+#include <type_traits>
+
 struct FStringFormatArg;
 template<typename InKeyType,typename InValueType,typename SetAllocator ,typename KeyFuncs > class TMap;
 
@@ -432,15 +434,56 @@ public:
 	 * @param InPrefix the prefix to search for at the start of the string to remove.
 	 * @return true if the prefix was removed, otherwise false.
 	 */
-	bool RemoveFromStart( const TCHAR* InPrefix, ESearchCase::Type SearchCase = ESearchCase::IgnoreCase );
+	template <typename TCharRangeType, std::enable_if_t<TIsCharRangeNotCArrayNotFString<TCharRangeType>::Value>* = nullptr>
+	bool RemoveFromStart(TCharRangeType&& InPrefix, ESearchCase::Type SearchCase = ESearchCase::IgnoreCase)
+	{
+		static_assert(std::is_same_v<typename TElementType<TCharRangeType>::Type, TCHAR>, "Expected a range of TCHAR");
+		return RemoveFromStart(GetData(InPrefix), GetNum(InPrefix), SearchCase);
+	}
+
+	/**
+	 * Removes the text from the start of the string if it exists.
+	 *
+	 * @param InPrefix the prefix to search for at the start of the string to remove.
+	 * @return true if the prefix was removed, otherwise false.
+	 */
+	bool RemoveFromStart(const TCHAR* InPrefix, ESearchCase::Type SearchCase = ESearchCase::IgnoreCase)
+	{
+		return RemoveFromStart(InPrefix, InPrefix ? FCString::Strlen(InPrefix) : 0, SearchCase);
+	}
+
+	/**
+	 * Removes the text from the start of the string if it exists.
+	 *
+	 * @param InPrefix the prefix to search for at the start of the string to remove.
+	 * @return true if the prefix was removed, otherwise false.
+	 */
+	bool RemoveFromStart(const FString& InPrefix, ESearchCase::Type SearchCase = ESearchCase::IgnoreCase)
+	{
+		return RemoveFromStart(*InPrefix, InPrefix.Len(), SearchCase);
+	}
 
 	/**
 	 * Removes the text from the start of the string if it exists.
 	 *
 	 * @param InPrefix the prefix to search for at the start of the string to remove.
+	 * @param InPrefixLen length of InPrefix
 	 * @return true if the prefix was removed, otherwise false.
 	 */
-	bool RemoveFromStart( const FString& InPrefix, ESearchCase::Type SearchCase = ESearchCase::IgnoreCase );
+	bool RemoveFromStart(const TCHAR* InPrefix, int32 InPrefixLen, ESearchCase::Type SearchCase = ESearchCase::IgnoreCase);
+
+	/**
+	 * Removes the text from the end of the string if it exists.
+	 *
+	 * @param InSuffix the suffix to search for at the end of the string to remove.
+	 * @return true if the suffix was removed, otherwise false.
+	 */
+	template <typename TCharRangeType, std::enable_if_t<TIsCharRangeNotCArrayNotFString<TCharRangeType>::Value>* = nullptr>
+	bool RemoveFromEnd(TCharRangeType&& InSuffix, ESearchCase::Type SearchCase = ESearchCase::IgnoreCase)
+	{
+		static_assert(std::is_same_v<typename TElementType<TCharRangeType>::Type, TCHAR>, "Expected a range of TCHAR");
+		return RemoveFromEnd(GetData(InSuffix), GetNum(InSuffix), SearchCase);
+	}
 
 	/**
 	 * Removes the text from the end of the string if it exists.
@@ -448,7 +491,10 @@ public:
 	 * @param InSuffix the suffix to search for at the end of the string to remove.
 	 * @return true if the suffix was removed, otherwise false.
 	 */
-	bool RemoveFromEnd( const TCHAR* InSuffix, ESearchCase::Type SearchCase = ESearchCase::IgnoreCase );
+	bool RemoveFromEnd(const TCHAR* InSuffix, ESearchCase::Type SearchCase = ESearchCase::IgnoreCase)
+	{
+		return RemoveFromEnd(InSuffix, InSuffix ? FCString::Strlen(InSuffix) : 0, SearchCase);
+	}
 
 	/**
 	 * Removes the text from the end of the string if it exists.
@@ -456,7 +502,19 @@ public:
 	 * @param InSuffix the suffix to search for at the end of the string to remove.
 	 * @return true if the suffix was removed, otherwise false.
 	 */
-	bool RemoveFromEnd( const FString& InSuffix, ESearchCase::Type SearchCase = ESearchCase::IgnoreCase );
+	bool RemoveFromEnd(const FString& InSuffix, ESearchCase::Type SearchCase = ESearchCase::IgnoreCase)
+	{
+		return RemoveFromEnd(*InSuffix, InSuffix.Len(), SearchCase);
+	}
+
+	/**
+	 * Removes the text from the end of the string if it exists.
+	 *
+	 * @param InSuffix the suffix to search for at the end of the string to remove.
+	 * @param InSuffixLen length of InSuffix
+	 * @return true if the suffix was removed, otherwise false.
+	 */
+	bool RemoveFromEnd(const TCHAR* InSuffix, int32 InSuffixLen, ESearchCase::Type SearchCase = ESearchCase::IgnoreCase);
 
 	/**
 	 * Concatenate this path with given path ensuring the / character is used between them
@@ -1007,6 +1065,45 @@ public:
 		RightChopInline(Start, bAllowShrinking);
 	}
 
+	/**
+	 * Searches the string for a substring, and returns index into this string
+	 * of the first found instance. Can search from beginning or end, and ignore case or not.
+	 *
+	 * @param SubStr			The string to search for
+	 * @param StartPosition		The start character position to search from
+	 * @param SearchCase		Indicates whether the search is case sensitive or not ( defaults to ESearchCase::IgnoreCase )
+	 * @param SearchDir			Indicates whether the search starts at the beginning or at the end ( defaults to ESearchDir::FromStart )
+	 *
+	 * @note  When searching backwards where a StartPosition is provided, searching will actually begin from
+	 *        StartPosition - SubStr.Len(), therefore:
+	 *
+	 *        FString("X").Find("X", ESearchCase::CaseSensitive, ESearchDir::FromEnd, 0) == INDEX_NONE
+	 *
+	 *        Consider using UE::String::FindLast() as an alternative.
+	 */
+	template <typename TCharRangeType, std::enable_if_t<TIsCharRangeNotCArrayNotFString<TCharRangeType>::Value>* = nullptr>
+	UE_NODISCARD int32 Find(TCharRangeType&& SubStr, ESearchCase::Type SearchCase = ESearchCase::IgnoreCase,
+		ESearchDir::Type SearchDir = ESearchDir::FromStart, int32 StartPosition = INDEX_NONE) const
+	{
+		static_assert(std::is_same_v<typename TElementType<TCharRangeType>::Type, TCHAR>, "Expected a range of TCHAR");
+		return Find(GetData(SubStr), GetNum(SubStr), SearchCase, SearchDir, StartPosition);
+	}
+
+	/**
+	 * Searches the string for a substring, and returns index into this string
+	 * of the first found instance. Can search from beginning or end, and ignore case or not.
+	 *
+	 * @param SubStr			The string array of TCHAR to search for
+	 * @param StartPosition		The start character position to search from.  See note below.
+	 * @param SearchCase		Indicates whether the search is case sensitive or not
+	 * @param SearchDir			Indicates whether the search starts at the beginning or at the end.
+	 */
+	UE_NODISCARD int32 Find(const TCHAR* SubStr, ESearchCase::Type SearchCase = ESearchCase::IgnoreCase,
+		ESearchDir::Type SearchDir = ESearchDir::FromStart, int32 StartPosition = INDEX_NONE) const
+	{
+		return SubStr ? Find(SubStr, FCString::Strlen(SubStr), SearchCase, SearchDir, StartPosition) : INDEX_NONE;
+	}
+
 	/**
 	 * Searches the string for a substring, and returns index into this string
 	 * of the first found instance. Can search from beginning or end, and ignore case or not.
@@ -1023,17 +1120,21 @@ public:
 	 *
 	 *        Consider using UE::String::FindLast() as an alternative.
 	 */
-	UE_NODISCARD int32 Find( const TCHAR* SubStr, ESearchCase::Type SearchCase = ESearchCase::IgnoreCase, 
-				ESearchDir::Type SearchDir = ESearchDir::FromStart, int32 StartPosition=INDEX_NONE ) const;
+	UE_NODISCARD int32 Find(const FString& SubStr, ESearchCase::Type SearchCase = ESearchCase::IgnoreCase,
+		ESearchDir::Type SearchDir = ESearchDir::FromStart, int32 StartPosition = INDEX_NONE) const
+	{
+		return Find(*SubStr, SubStr.Len(), SearchCase, SearchDir, StartPosition);
+	}
 
 	/**
 	 * Searches the string for a substring, and returns index into this string
 	 * of the first found instance. Can search from beginning or end, and ignore case or not.
 	 *
-	 * @param SubStr			The string to search for
-	 * @param StartPosition		The start character position to search from  See note below.
-	 * @param SearchCase		Indicates whether the search is case sensitive or not ( defaults to ESearchCase::IgnoreCase )
-	 * @param SearchDir			Indicates whether the search starts at the beginning or at the end ( defaults to ESearchDir::FromStart )
+	 * @param SubStr			The string array of TCHAR to search for
+	 * @param SubStrLen			The length of the SubStr array
+	 * @param StartPosition		The start character position to search from.  See note below.
+	 * @param SearchCase		Indicates whether the search is case sensitive or not
+	 * @param SearchDir			Indicates whether the search starts at the beginning or at the end.
 	 *
 	 * @note  When searching backwards where a StartPosition is provided, searching will actually begin from
 	 *        StartPosition - SubStr.Len(), therefore:
@@ -1042,11 +1143,9 @@ public:
 	 *
 	 *        Consider using UE::String::FindLast() as an alternative.
 	 */
-	UE_NODISCARD FORCEINLINE int32 Find( const FString& SubStr, ESearchCase::Type SearchCase = ESearchCase::IgnoreCase, 
-							ESearchDir::Type SearchDir = ESearchDir::FromStart, int32 StartPosition=INDEX_NONE ) const
-	{
-		return Find( *SubStr, SearchCase, SearchDir, StartPosition );
-	}
+	UE_NODISCARD int32 Find(const TCHAR* SubStr, int32 InSubStrLen, ESearchCase::Type SearchCase = ESearchCase::IgnoreCase,
+		ESearchDir::Type SearchDir = ESearchDir::FromStart, int32 StartPosition = INDEX_NONE) const;
+
 
 	/** 
 	 * Returns whether this string contains the specified substring.
@@ -1309,36 +1408,90 @@ public:
 
 
 	/**
-	 * Test whether this string starts with given string.
+	 * Test whether this string starts with given prefix.
 	 *
 	 * @param SearchCase		Indicates whether the search is case sensitive or not ( defaults to ESearchCase::IgnoreCase )
 	 * @return true if this string begins with specified text, false otherwise
 	 */
-	UE_NODISCARD bool StartsWith(const TCHAR* InPrefix, ESearchCase::Type SearchCase = ESearchCase::IgnoreCase) const;
+	template <typename TCharRangeType, std::enable_if_t<TIsCharRangeNotCArrayNotFString<TCharRangeType>::Value>* = nullptr>
+	UE_NODISCARD bool StartsWith(TCharRangeType&& InPrefix, ESearchCase::Type SearchCase = ESearchCase::IgnoreCase) const
+	{
+		static_assert(std::is_same_v<typename TElementType<TCharRangeType>::Type, TCHAR>, "Expected a range of TCHAR");
+		return StartsWith(GetData(InPrefix), GetNum(InPrefix), SearchCase);
+	}
 
 	/**
-	 * Test whether this string starts with given string.
+	 * Test whether this string starts with given prefix.
 	 *
 	 * @param SearchCase		Indicates whether the search is case sensitive or not ( defaults to ESearchCase::IgnoreCase )
 	 * @return true if this string begins with specified text, false otherwise
 	 */
-	UE_NODISCARD bool StartsWith(const FString& InPrefix, ESearchCase::Type SearchCase = ESearchCase::IgnoreCase) const;
+	UE_NODISCARD bool StartsWith(const TCHAR* InPrefix, ESearchCase::Type SearchCase = ESearchCase::IgnoreCase) const
+	{
+		return StartsWith(InPrefix, InPrefix ? FCString::Strlen(InPrefix) : 0, SearchCase);
+	}
+
+	/**
+	 * Test whether this string starts with given prefix.
+	 *
+	 * @param SearchCase		Indicates whether the search is case sensitive or not ( defaults to ESearchCase::IgnoreCase )
+	 * @return true if this string begins with specified text, false otherwise
+	 */
+	UE_NODISCARD bool StartsWith(const FString& InPrefix, ESearchCase::Type SearchCase = ESearchCase::IgnoreCase) const
+	{
+		return StartsWith(*InPrefix, InPrefix.Len(), SearchCase);
+	}
+
+	/**
+	 * Test whether this string starts with given prefix.
+	 *
+	 * @param SearchCase		Indicates whether the search is case sensitive or not ( defaults to ESearchCase::IgnoreCase )
+	 * @return true if this string begins with specified text, false otherwise
+	 */
+	UE_NODISCARD bool StartsWith(const TCHAR* InPrefix, int32 InPrefixLen, ESearchCase::Type SearchCase = ESearchCase::IgnoreCase) const;
+
+	/**
+	 * Test whether this string ends with given suffix.
+	 *
+	 * @param SearchCase		Indicates whether the search is case sensitive or not ( defaults to ESearchCase::IgnoreCase )
+	 * @return true if this string ends with specified text, false otherwise
+	 */
+	template <typename TCharRangeType, std::enable_if_t<TIsCharRangeNotCArrayNotFString<TCharRangeType>::Value>* = nullptr>
+	UE_NODISCARD bool EndsWith(TCharRangeType&& InSuffix, ESearchCase::Type SearchCase = ESearchCase::IgnoreCase) const
+	{
+		static_assert(std::is_same_v<typename TElementType<TCharRangeType>::Type, TCHAR>, "Expected a range of TCHAR");
+		return EndsWith(GetData(InSuffix), GetNum(InSuffix), SearchCase);
+	}
+
+	/**
+	 * Test whether this string ends with given suffix.
+	 *
+	 * @param SearchCase		Indicates whether the search is case sensitive or not ( defaults to ESearchCase::IgnoreCase )
+	 * @return true if this string ends with specified text, false otherwise
+	 */
+	UE_NODISCARD bool EndsWith(const TCHAR* InSuffix, ESearchCase::Type SearchCase = ESearchCase::IgnoreCase) const
+	{
+		return EndsWith(InSuffix, InSuffix ? FCString::Strlen(InSuffix) : 0, SearchCase);
+	}
 
 	/**
-	 * Test whether this string ends with given string.
+	 * Test whether this string ends with given suffix.
 	 *
 	 * @param SearchCase		Indicates whether the search is case sensitive or not ( defaults to ESearchCase::IgnoreCase )
 	 * @return true if this string ends with specified text, false otherwise
 	 */
-	UE_NODISCARD bool EndsWith(const TCHAR* InSuffix, ESearchCase::Type SearchCase = ESearchCase::IgnoreCase) const;
+	UE_NODISCARD bool EndsWith(const FString& InSuffix, ESearchCase::Type SearchCase = ESearchCase::IgnoreCase) const
+	{
+		return EndsWith(*InSuffix, InSuffix.Len(), SearchCase);
+	}
 
 	/**
-	 * Test whether this string ends with given string.
+	 * Test whether this string ends with given suffix.
 	 *
 	 * @param SearchCase		Indicates whether the search is case sensitive or not ( defaults to ESearchCase::IgnoreCase )
 	 * @return true if this string ends with specified text, false otherwise
 	 */
-	UE_NODISCARD bool EndsWith(const FString& InSuffix, ESearchCase::Type SearchCase = ESearchCase::IgnoreCase ) const;
+	UE_NODISCARD bool EndsWith(const TCHAR* InSuffix, int32 InSuffixLen, ESearchCase::Type SearchCase = ESearchCase::IgnoreCase ) const;
 
 	/**
 	 * Searches this string for a given wild card
@@ -1348,7 +1501,12 @@ public:
 	 * @return true if this string matches the *?-type wildcard given. 
 	 * @warning This is a simple, SLOW routine. Use with caution
 	 */
-	UE_NODISCARD bool MatchesWildcard(const TCHAR* Wildcard, ESearchCase::Type SearchCase = ESearchCase::IgnoreCase) const;
+	template <typename TCharRangeType, std::enable_if_t<TIsCharRangeNotCArrayNotFString<TCharRangeType>::Value>* = nullptr>
+	UE_NODISCARD bool MatchesWildcard(TCharRangeType&& Wildcard, ESearchCase::Type SearchCase = ESearchCase::IgnoreCase) const
+	{
+		static_assert(std::is_same_v<typename TElementType<TCharRangeType>::Type, TCHAR>, "Expected a range of TCHAR");
+		return MatchesWildcard(GetData(Wildcard), GetNum(Wildcard), SearchCase);
+	}
 
 	/**
 	 * Searches this string for a given wild card
@@ -1358,11 +1516,34 @@ public:
 	 * @return true if this string matches the *?-type wildcard given.
 	 * @warning This is a simple, SLOW routine. Use with caution
 	 */
-	UE_NODISCARD FORCEINLINE bool MatchesWildcard(const FString& Wildcard, ESearchCase::Type SearchCase = ESearchCase::IgnoreCase) const
+	UE_NODISCARD bool MatchesWildcard(const TCHAR* Wildcard, ESearchCase::Type SearchCase = ESearchCase::IgnoreCase) const
 	{
-		return MatchesWildcard(*Wildcard, SearchCase);
+		return MatchesWildcard(Wildcard, Wildcard ? FCString::Strlen(Wildcard) : 0, SearchCase);
 	}
 
+	/**
+	 * Searches this string for a given wild card
+	 *
+	 * @param Wildcard		*?-type wildcard
+	 * @param SearchCase	Indicates whether the search is case sensitive or not ( defaults to ESearchCase::IgnoreCase )
+	 * @return true if this string matches the *?-type wildcard given.
+	 * @warning This is a simple, SLOW routine. Use with caution
+	 */
+	UE_NODISCARD bool MatchesWildcard(const FString& Wildcard, ESearchCase::Type SearchCase = ESearchCase::IgnoreCase) const
+	{
+		return MatchesWildcard(*Wildcard, Wildcard.Len(), SearchCase);
+	}
+
+	/**
+	 * Searches this string for a given wild card
+	 *
+	 * @param Wildcard		*?-type wildcard
+	 * @param SearchCase	Indicates whether the search is case sensitive or not ( defaults to ESearchCase::IgnoreCase )
+	 * @return true if this string matches the *?-type wildcard given.
+	 * @warning This is a simple, SLOW routine. Use with caution
+	 */
+	UE_NODISCARD bool MatchesWildcard(const TCHAR* Wildcard, int32 WildcardLen, ESearchCase::Type SearchCase = ESearchCase::IgnoreCase) const;
+
 	/**
 	 * Removes whitespace characters from the start and end of this string. Modifies the string in-place.
 	 */
diff --git a/Engine/Source/Runtime/Core/Public/Misc/CString.h b/Engine/Source/Runtime/Core/Public/Misc/CString.h
--- a/Engine/Source/Runtime/Core/Public/Misc/CString.h
+++ b/Engine/Source/Runtime/Core/Public/Misc/CString.h
@@ -259,7 +259,7 @@ struct TCString
 	 * Finds string in string, case insensitive 
 	 * @param Str The string to look through
 	 * @param Find The string to find inside Str
-	 * @return Position in Str if Find was found, otherwise, NULL
+	 * @return Position in Str if Find was found, otherwise, nullptr
 	 */
 	static const CharType* Stristr(const CharType* Str, const CharType* Find);
 
@@ -267,10 +267,56 @@ struct TCString
 	 * Finds string in string, case insensitive (non-const version)
 	 * @param Str The string to look through
 	 * @param Find The string to find inside Str
-	 * @return Position in Str if Find was found, otherwise, NULL
+	 * @return Position in Str if Find was found, otherwise, nullptr
 	 */
 	static CharType* Stristr(CharType* Str, const CharType* Find) { return (CharType*)Stristr((const CharType*)Str, Find); }
 
+	/**
+	 * Finds string in string, case insensitive
+	 * @param Str The character array to look through
+	 * @param InStrLen The length of the Str array
+	 * @param Find The character array to find inside Str
+	 * @param FindLen The length of the Find array
+	 * @return Position in Str if Find was found, otherwise, nullptr
+	 */
+	static const CharType* Strnistr(const CharType* Str, int32 InStrLen, const CharType* Find, int32 FindLen);
+
+	/**
+	 * Finds string in string, case insensitive (non-const version)
+	 * @param Str The character array to look through
+	 * @param InStrLen The length of the Str array
+	 * @param Find The character array to find inside Str
+	 * @param FindLen The length of the Find array
+	 * @return Position in Str if Find was found, otherwise, nullptr
+	 */
+	static CharType* Strnistr(CharType* Str, int32 InStrLen, const CharType* Find, int32 FindLen)
+	{
+		return (CharType*)Strnistr((const CharType*)Str, InStrLen, Find, FindLen);
+	}
+
+	/**
+	 * Finds string in string, case sensitive
+	 * @param Str The character array to look through
+	 * @param InStrLen The length of the Str array
+	 * @param Find The character array to find inside Str
+	 * @param FindLen The length of the Find array
+	 * @return Position in Str if Find was found, otherwise, nullptr
+	 */
+	static const CharType* Strnstr(const CharType* Str, int32 InStrLen, const CharType* Find, int32 FindLen);
+
+	/**
+	 * Finds string in string, case sensitive (non-const version)
+	 * @param Str The character array to look through
+	 * @param InStrLen The length of the Str array
+	 * @param Find The character array to find inside Str
+	 * @param FindLen The length of the Find array
+	 * @return Position in Str if Find was found, otherwise, nullptr
+	 */
+	static CharType* Strnstr(CharType* Str, int32 InStrLen, const CharType* Find, int32 FindLen)
+	{
+		return (CharType*)Strnstr((const CharType*)Str, InStrLen, Find, FindLen);
+	}
+
 	/**
 	 * strlen wrapper
 	 */
@@ -612,10 +658,13 @@ const typename TCString<T>::CharType* TCString<T>::StrfindDelim(const CharType*
 template <typename T>
 const typename TCString<T>::CharType* TCString<T>::Stristr(const CharType* Str, const CharType* Find)
 {
+	// Note this implementation is not reducible to Strnistr because we need to avoid calling strlen(Str) to
+	// avoid scanning it twice
+
 	// both strings must be valid
-	if( Find == NULL || Str == NULL )
+	if( Find == nullptr || Str == nullptr )
 	{
-		return NULL;
+		return nullptr;
 	}
 
 	// get upper-case first letter of the find string (to reduce the number of full strnicmps)
@@ -639,8 +688,86 @@ const typename TCString<T>::CharType* TCString<T>::Stristr(const CharType* Str,
 		StrChar = *Str++;
 	}
 
-	// if nothing was found, return NULL
-	return NULL;
+	// if nothing was found, return nullptr
+	return nullptr;
+}
+
+template <typename T>
+const typename TCString<T>::CharType* TCString<T>::Strnistr(const CharType* Str, int32 InStrLen, const CharType* Find, int32 FindLen)
+{
+	if (FindLen <= 0)
+	{
+		checkf(FindLen >= 0, TEXT("Invalid FindLen: %d"), FindLen);
+		return Str;
+	}
+	if (InStrLen < FindLen)
+	{
+		checkf(InStrLen >= 0, TEXT("Invalid InStrLen: %d"), InStrLen);
+		return nullptr;
+	}
+
+	// get upper-case first letter of the find string (to reduce the number of full strnicmps)
+	CharType FindInitial = TChar<CharType>::ToUpper(*Find);
+	// Set FindSuffix,FindSuffixLength to the characters of Find after the first letter
+	int32 FindSuffixLength = FindLen - 1;
+	const CharType* FindSuffix = Find + 1;
+
+	// while the length of the remaining string is >= FindLen
+	const CharType* StrLastChance = Str + InStrLen - FindLen;
+	while (Str <= StrLastChance)
+	{
+		CharType StrChar = *Str++;
+
+		// make sure it's upper-case
+		StrChar = TChar<CharType>::ToUpper(StrChar);
+		// if it matches the first letter of the find string, do a case-insensitive string compare for the length of the find string
+		if (StrChar == FindInitial && !Strnicmp(Str, FindSuffix, FindSuffixLength))
+		{
+			// if we found the string, then return a pointer to the beginning of it in the search string
+			return Str - 1;
+		}
+	}
+
+	// if nothing was found, return nullptr
+	return nullptr;
+}
+
+template <typename T>
+const typename TCString<T>::CharType* TCString<T>::Strnstr(const CharType* Str, int32 InStrLen, const CharType* Find, int32 FindLen)
+{
+	if (FindLen <= 0)
+	{
+		checkf(FindLen >= 0, TEXT("Invalid FindLen: %d"), FindLen);
+		return Str;
+	}
+	if (InStrLen < FindLen)
+	{
+		checkf(InStrLen >= 0, TEXT("Invalid InStrLen: %d"), InStrLen);
+		return nullptr;
+	}
+
+	// get first letter of the find string (to reduce the number of full strncmps)
+	CharType FindInitial = *Find;
+	// Set FindSuffix,FindSuffixLength to the characters of Find after the first letter
+	int32 FindSuffixLength = FindLen - 1;
+	const CharType* FindSuffix = Find + 1;
+
+	// while the length of the remaining string is >= FindLen
+	const CharType* StrLastChance = Str + InStrLen - FindLen;
+	while (Str <= StrLastChance)
+	{
+		CharType StrChar = *Str++;
+
+		// if it matches the first letter of the find string, do a string compare for the length of the find string
+		if (StrChar == FindInitial && !Strncmp(Str, FindSuffix, FindSuffixLength))
+		{
+			// if we found the string, then return a pointer to the beginning of it in the search string
+			return Str - 1;
+		}
+	}
+
+	// if nothing was found, return nullptr
+	return nullptr;
 }
 
 template <typename T> FORCEINLINE
