IoStore: Add functionality to sign a container as a separate step

diff --git a/Engine/Source/Developer/IoStoreUtilities/Private/IoStoreUtilities.cpp b/Engine/Source/Developer/IoStoreUtilities/Private/IoStoreUtilities.cpp
--- a/Engine/Source/Developer/IoStoreUtilities/Private/IoStoreUtilities.cpp
+++ b/Engine/Source/Developer/IoStoreUtilities/Private/IoStoreUtilities.cpp
@@ -66,6 +66,7 @@
 #include "ZenFileSystemManifest.h"
 #include "IPlatformFileSandboxWrapper.h"
 #include "Misc/PathViews.h"
+#include "HAL/FileManagerGeneric.h"
 
 //PRAGMA_DISABLE_OPTIMIZATION
 
@@ -5904,6 +5905,90 @@ bool ExtractFilesFromIoStoreContainer(
 	return true;
 }
 
+bool SignIoStoreContainer(const TCHAR* InContainerFilename, const FRSAKeyHandle InSigningKey)
+{
+	FString TocFilePath = FPaths::ChangeExtension(InContainerFilename, TEXT(".utoc"));
+	FString TempOutputPath = TocFilePath + ".tmp";
+	IPlatformFile& Ipf = FPlatformFileManager::Get().GetPlatformFile();
+	ON_SCOPE_EXIT
+	{
+		if (Ipf.FileExists(*TempOutputPath))
+		{
+			Ipf.DeleteFile(*TempOutputPath);
+		}
+	};
+
+	FIoStoreTocResource TocResource;
+	FIoStatus Status = FIoStoreTocResource::Read(*TocFilePath, EIoStoreTocReadOptions::ReadAll, TocResource);
+	if (!Status.IsOk())
+	{
+		UE_LOG(LogIoStore, Error, TEXT("Failed reading container file \"%s\"."), InContainerFilename);
+		return false;
+	}
+
+	if (TocResource.ChunkBlockSignatures.Num() != TocResource.CompressionBlocks.Num())
+	{
+		UE_LOG(LogIoStore, Display, TEXT("Container is not signed, calculating block hashes..."));
+		TocResource.ChunkBlockSignatures.Empty();
+		TUniquePtr<FArchive> ContainerFileReader;
+		int32 LastPartitionIndex = -1;
+		TArray<uint8> BlockBuffer;
+		BlockBuffer.SetNum(static_cast<int32>(TocResource.Header.CompressionBlockSize));
+		const int32 BlockCount = TocResource.CompressionBlocks.Num();
+		FString ContainerBasePath = FPaths::ChangeExtension(InContainerFilename, TEXT(""));
+		TStringBuilder<256> UcasFilePath;
+		for (int32 BlockIndex = 0; BlockIndex < BlockCount; ++BlockIndex)
+		{
+			const FIoStoreTocCompressedBlockEntry& CompressionBlockEntry = TocResource.CompressionBlocks[BlockIndex];
+			uint64 BlockRawSize = Align(CompressionBlockEntry.GetCompressedSize(), FAES::AESBlockSize);
+			check(BlockRawSize <= TocResource.Header.CompressionBlockSize);
+			const int32 PartitionIndex = int32(CompressionBlockEntry.GetOffset() / TocResource.Header.PartitionSize);
+			const uint64 PartitionRawOffset = CompressionBlockEntry.GetOffset() % TocResource.Header.PartitionSize;
+			if (PartitionIndex != LastPartitionIndex)
+			{
+				UcasFilePath.Reset();
+				UcasFilePath.Append(ContainerBasePath);
+				if (PartitionIndex > 0)
+				{
+					UcasFilePath.Append(FString::Printf(TEXT("_s%d"), PartitionIndex));
+				}
+				UcasFilePath.Append(TEXT(".ucas"));
+				IFileHandle* ContainerFileHandle = Ipf.OpenRead(*UcasFilePath, /* allowwrite */ false);
+				if (!ContainerFileHandle)
+				{
+					UE_LOG(LogIoStore, Error, TEXT("Failed opening container file \"%s\"."), *UcasFilePath);
+					return false;
+				}
+				ContainerFileReader.Reset(new FArchiveFileReaderGeneric(ContainerFileHandle, *UcasFilePath, ContainerFileHandle->Size(), 256 << 10));
+				LastPartitionIndex = PartitionIndex;
+			}
+			ContainerFileReader->Seek(PartitionRawOffset);
+			ContainerFileReader->Precache(PartitionRawOffset, 0); // Without this buffering won't work due to the first read after a seek always being uncached
+			ContainerFileReader->Serialize(BlockBuffer.GetData(), BlockRawSize);
+			FSHAHash& BlockHash = TocResource.ChunkBlockSignatures.AddDefaulted_GetRef();
+			FSHA1::HashBuffer(BlockBuffer.GetData(), BlockRawSize, BlockHash.Hash);
+		}
+	}
+
+	FIoContainerSettings ContainerSettings;
+	ContainerSettings.ContainerId = TocResource.Header.ContainerId;
+	ContainerSettings.ContainerFlags = TocResource.Header.ContainerFlags | EIoContainerFlags::Signed;
+	ContainerSettings.EncryptionKeyGuid = TocResource.Header.EncryptionKeyGuid;
+	ContainerSettings.SigningKey = InSigningKey;
+
+	TIoStatusOr<uint64> WriteStatus = FIoStoreTocResource::Write(*TempOutputPath, TocResource, TocResource.Header.CompressionBlockSize, TocResource.Header.PartitionSize, ContainerSettings);
+	if (!WriteStatus.IsOk())
+	{
+		UE_LOG(LogIoStore, Error, TEXT("Failed writing new utoc file file \"%s\"."), *TocFilePath);
+		return false;
+	}
+
+	Ipf.DeleteFile(*TocFilePath);
+	Ipf.MoveFile(*TocFilePath, *TempOutputPath);
+
+	return true;
+}
+
 static bool ParsePakResponseFile(const TCHAR* FilePath, TArray<FContainerSourceFile>& OutFiles)
 {
 	TArray<FString> ResponseFileContents;
diff --git a/Engine/Source/Developer/IoStoreUtilities/Public/IoStoreUtilities.h b/Engine/Source/Developer/IoStoreUtilities/Public/IoStoreUtilities.h
--- a/Engine/Source/Developer/IoStoreUtilities/Public/IoStoreUtilities.h
+++ b/Engine/Source/Developer/IoStoreUtilities/Public/IoStoreUtilities.h
@@ -7,6 +7,7 @@
 #include "Containers/UnrealString.h"
 #include "CoreMinimal.h"
 #include "HAL/Platform.h"
+#include "Misc/IEngineCrypto.h"
 
 struct FGuid;
 struct FKeyChain;
@@ -36,3 +37,5 @@ IOSTOREUTILITIES_API bool ExtractFilesFromIoStoreContainer(
 	TMap<FString, uint64>* OutOrderMap,
 	TArray<FGuid>* OutUsedEncryptionKeys,
 	bool* bOutIsSigned);
+
+IOSTOREUTILITIES_API bool SignIoStoreContainer(const TCHAR* InContainerFilename, const FRSAKeyHandle InSigningKey);
diff --git a/Engine/Source/Runtime/Core/Internal/IO/IoStore.h b/Engine/Source/Runtime/Core/Internal/IO/IoStore.h
--- a/Engine/Source/Runtime/Core/Internal/IO/IoStore.h
+++ b/Engine/Source/Runtime/Core/Internal/IO/IoStore.h
@@ -248,7 +248,7 @@ struct FIoStoreTocResource
 
 	UE_NODISCARD CORE_API static FIoStatus Read(const TCHAR* TocFilePath, EIoStoreTocReadOptions ReadOptions, FIoStoreTocResource& OutTocResource);
 
-	UE_NODISCARD static TIoStatusOr<uint64> Write(const TCHAR* TocFilePath, FIoStoreTocResource& TocResource, const FIoContainerSettings& ContainerSettings, const FIoStoreWriterSettings& WriterSettings);
+	UE_NODISCARD CORE_API static TIoStatusOr<uint64> Write(const TCHAR* TocFilePath, FIoStoreTocResource& TocResource, uint32 CompressionBlockSize, uint64 MaxPartitionSize, const FIoContainerSettings& ContainerSettings);
 
 	CORE_API static uint64 HashChunkIdWithSeed(int32 Seed, const FIoChunkId& ChunkId);
 };
diff --git a/Engine/Source/Runtime/Core/Private/IO/IoStore.cpp b/Engine/Source/Runtime/Core/Private/IO/IoStore.cpp
--- a/Engine/Source/Runtime/Core/Private/IO/IoStore.cpp
+++ b/Engine/Source/Runtime/Core/Private/IO/IoStore.cpp
@@ -1135,7 +1135,7 @@ public:
 				ContainerSettings.IsEncrypted() ? ContainerSettings.EncryptionKey : FAES::FAESKey());
 		}
 
-		TIoStatusOr<uint64> TocSize = FIoStoreTocResource::Write(*TocFilePath, TocResource, ContainerSettings, WriterContext->GetSettings());
+		TIoStatusOr<uint64> TocSize = FIoStoreTocResource::Write(*TocFilePath, TocResource, static_cast<uint32>(WriterSettings.CompressionBlockSize), WriterSettings.MaxPartitionSize, ContainerSettings);
 		check(TocSize.IsOk());
 		
 		Result.ContainerId = ContainerSettings.ContainerId;
@@ -1148,7 +1148,7 @@ public:
 		Result.CompressedContainerSize = CompressedContainerSize;
 		Result.DirectoryIndexSize = TocResource.Header.DirectoryIndexSize;
 		Result.CompressionMethod = EnumHasAnyFlags(ContainerSettings.ContainerFlags, EIoContainerFlags::Compressed)
-			? WriterContext->GetSettings().CompressionMethod
+			? WriterSettings.CompressionMethod
 			: NAME_None;
 		Result.ModifiedChunksCount = 0;
 		Result.AddedChunksCount = 0;
@@ -2892,8 +2892,9 @@ FIoStatus FIoStoreTocResource::Read(const TCHAR* TocFilePath, EIoStoreTocReadOpt
 TIoStatusOr<uint64> FIoStoreTocResource::Write(
 	const TCHAR* TocFilePath,
 	FIoStoreTocResource& TocResource,
-	const FIoContainerSettings& ContainerSettings,
-	const FIoStoreWriterSettings& WriterSettings)
+	uint32 CompressionBlockSize,
+	uint64 MaxPartitionSize,
+	const FIoContainerSettings& ContainerSettings)
 {
 	check(TocFilePath != nullptr);
 
@@ -2916,6 +2917,19 @@ TIoStatusOr<uint64> FIoStoreTocResource::Write(
 		return FIoStatus(EIoErrorCode::InvalidParameter, TEXT("Number of TOC chunk IDs doesn't match the number of chunk meta data"));
 	}
 
+	bool bHasExplicitCompressionMethodNone = false;
+	for (int32 CompressionMethodIndex = 0; CompressionMethodIndex < TocResource.CompressionMethods.Num(); ++CompressionMethodIndex)
+	{
+		if (TocResource.CompressionMethods[CompressionMethodIndex].IsNone())
+		{
+			if (CompressionMethodIndex != 0)
+			{
+				return FIoStatus(EIoErrorCode::InvalidParameter, TEXT("Compression method None must be the first compression method"));
+			}
+			bHasExplicitCompressionMethodNone = true;
+		}
+	}
+
 	FMemory::Memzero(&TocResource.Header, sizeof(FIoStoreTocHeader));
 
 	FIoStoreTocHeader& TocHeader = TocResource.Header;
@@ -2927,8 +2941,8 @@ TIoStatusOr<uint64> FIoStoreTocResource::Write(
 	TocHeader.TocChunksWithoutPerfectHashCount = TocResource.ChunkIndicesWithoutPerfectHash.Num();
 	TocHeader.TocCompressedBlockEntryCount = TocResource.CompressionBlocks.Num();
 	TocHeader.TocCompressedBlockEntrySize = sizeof(FIoStoreTocCompressedBlockEntry);
-	TocHeader.CompressionBlockSize = uint32(WriterSettings.CompressionBlockSize);
-	TocHeader.CompressionMethodNameCount = TocResource.CompressionMethods.Num();
+	TocHeader.CompressionBlockSize = CompressionBlockSize;
+	TocHeader.CompressionMethodNameCount = TocResource.CompressionMethods.Num() - (bHasExplicitCompressionMethodNone ? 1 : 0);
 	TocHeader.CompressionMethodNameLength = FIoStoreTocResource::CompressionMethodNameLen;
 	TocHeader.DirectoryIndexSize = TocResource.DirectoryIndexBuffer.Num();
 	TocHeader.ContainerId = ContainerSettings.ContainerId;
@@ -2939,13 +2953,13 @@ TIoStatusOr<uint64> FIoStoreTocResource::Write(
 		TocHeader.PartitionCount = 0;
 		TocHeader.PartitionSize = MAX_uint64;
 	}
-	else if (WriterSettings.MaxPartitionSize)
+	else if (MaxPartitionSize)
 	{
 		const FIoStoreTocCompressedBlockEntry& LastBlock = TocResource.CompressionBlocks.Last();
 		uint64 LastBlockEnd = LastBlock.GetOffset() + LastBlock.GetCompressedSize();
-		TocHeader.PartitionCount = uint32(Align(LastBlockEnd, WriterSettings.MaxPartitionSize) / WriterSettings.MaxPartitionSize);
+		TocHeader.PartitionCount = uint32(Align(LastBlockEnd, MaxPartitionSize) / MaxPartitionSize);
 		check(TocHeader.PartitionCount > 0);
-		TocHeader.PartitionSize = WriterSettings.MaxPartitionSize;
+		TocHeader.PartitionSize = MaxPartitionSize;
 	}
 	else
 	{
@@ -2994,6 +3008,10 @@ TIoStatusOr<uint64> FIoStoreTocResource::Write(
 
 	for (FName MethodName : TocResource.CompressionMethods)
 	{
+		if (MethodName.IsNone())
+		{
+			continue;
+		}
 		FMemory::Memzero(AnsiMethodName, FIoStoreTocResource::CompressionMethodNameLen);
 		FCStringAnsi::Strcpy(AnsiMethodName, FIoStoreTocResource::CompressionMethodNameLen, TCHAR_TO_ANSI(*MethodName.ToString()));
 
