- Removing NoSave flag for layered config section (in editor) to avoid breaking changes to project settings targeting that config

diff --git a/Engine/Source/Runtime/Core/Private/Misc/ConfigCacheIni.cpp b/Engine/Source/Runtime/Core/Private/Misc/ConfigCacheIni.cpp
--- a/Engine/Source/Runtime/Core/Private/Misc/ConfigCacheIni.cpp
+++ b/Engine/Source/Runtime/Core/Private/Misc/ConfigCacheIni.cpp
@@ -1316,8 +1316,12 @@ void FConfigFile::AddDynamicLayerToHierarchy(const FString& Filename)
 	SourceIniHierarchy.AddDynamicLayer(Filename);
 	CombineFromBuffer(ConfigContent, Filename);
 
-	// Disable saving since dynamic layers are only for runtime
-	NoSave = true;
+	// Removing NoSave flag that was being set to true when adding a dynamic layered to a config hierarchy 
+	// (usually from plugins), which was causing settings in the editor to not be saved out
+	if (!GIsEditor)
+	{
+		NoSave = true;
+	}
 }
 
 
