Added CVar to disable FBitReader overflow ensures.

diff --git a/Engine/Source/Runtime/Core/Private/Serialization/BitReader.cpp b/Engine/Source/Runtime/Core/Private/Serialization/BitReader.cpp
--- a/Engine/Source/Runtime/Core/Private/Serialization/BitReader.cpp
+++ b/Engine/Source/Runtime/Core/Private/Serialization/BitReader.cpp
@@ -17,10 +17,22 @@ PRAGMA_DISABLE_UNSAFE_TYPECAST_WARNINGS
 extern const uint8 GShift[8]={0x01,0x02,0x04,0x08,0x10,0x20,0x40,0x80};
 extern const uint8 GMask [8]={0x00,0x01,0x03,0x07,0x0f,0x1f,0x3f,0x7f};
 
-static TAutoConsoleVariable<bool> CVarLogFatalOnOverflow(
-	TEXT("BitReader.LogFatalOnOverflow"),
-	false,
-	TEXT("LogFatal if BitReader Overflows"));
+namespace UE::Net
+{
+	static bool GLogFatalOnOverflow = false;
+
+	static FAutoConsoleVariableRef CVarLogFatalOnOverflow(
+		TEXT("BitReader.LogFatalOnOverflow"),
+		GLogFatalOnOverflow,
+		TEXT("LogFatal if BitReader Overflows"));
+
+	static int32 GEnsureOnOverflow = 1;
+
+	static FAutoConsoleVariableRef CVarEnsureOnOverflow(
+		TEXT("net.BitReader.EnsureOnOverflow"),
+		GEnsureOnOverflow,
+		TEXT("Ensures if BitReader overflows. 'BitReader.LogFatalOnOverflow' takes precedence over this, if set."));
+}
 
 // Optimized arbitrary bit range memory copy routine.
 
@@ -252,15 +264,22 @@ void FBitReader::CountMemory(FArchive& Ar) const
 
 void FBitReader::SetOverflowed(int64 LengthBits)
 {
-	if (CVarLogFatalOnOverflow.GetValueOnAnyThread())
+	using namespace UE::Net;
+
+	if (GLogFatalOnOverflow)
 	{
 		UE_LOG(LogNetSerialization, Fatal, TEXT("FBitReader::SetOverflowed() called! (ReadLen: %i, Remaining: %i, Max: %i)"),
-			LengthBits, (Num - Pos), Num);
+				LengthBits, (Num - Pos), Num);
 	}
-	else
+	else if (GEnsureOnOverflow)
 	{
 		ensureMsgf(false, TEXT("FBitReader::SetOverflowed() called! (ReadLen: %i, Remaining: %i, Max: %i)"),
-			LengthBits, (Num - Pos), Num);
+					LengthBits, (Num - Pos), Num);
+	}
+	else
+	{
+		UE_LOG(LogNetSerialization, Error, TEXT("FBitReader::SetOverflowed() called! (ReadLen: %i, Remaining: %i, Max: %i)"),
+				LengthBits, (Num - Pos), Num);
 	}
 
 	SetError();
