Improved comments for FText::Format

#rnx

diff --git a/Engine/Source/Runtime/Core/Public/Internationalization/Text.h b/Engine/Source/Runtime/Core/Public/Internationalization/Text.h
--- a/Engine/Source/Runtime/Core/Public/Internationalization/Text.h
+++ b/Engine/Source/Runtime/Core/Public/Internationalization/Text.h
@@ -247,7 +247,8 @@ struct CORE_API FNumberParsingOptions
 
 /**
  * Cached compiled expression used by the text formatter.
- * The compiled expression will automatically update if the display string is changed.
+ * The compiled expression will automatically update if the display string is changed,
+ * and is safe to be used as a function-level static.
  * See TextFormatter.cpp for the definition.
  */
 class CORE_API FTextFormat
@@ -548,12 +549,44 @@ public:
 
 	static void GetFormatPatternParameters(const FTextFormat& Fmt, TArray<FString>& ParameterNames);
 
+	/**
+	 * Format the given map of key->value pairs as named arguments within the given format pattern
+	 *
+	 * @note You may want to pre-compile your FText pattern into a FTextFormat prior to performing formats within a loop or on a critical path,
+	 *       as this can save CPU cycles, memory, and mutex resources vs re-compiling the pattern for each format call. See FTextFormat for more info.
+	 *
+	 * @param Fmt The format pattern to use
+	 * @param InArguments The map of key->value pairs to inject into the format pattern
+	 * @return The formatted FText
+	 */
 	static FText Format(FTextFormat Fmt, const FFormatNamedArguments& InArguments);
 	static FText Format(FTextFormat Fmt, FFormatNamedArguments&& InArguments);
 
+	/**
+	 * Format the given list values as ordered arguments within the given format pattern
+	 *
+	 * @note You may want to pre-compile your FText pattern into a FTextFormat prior to performing formats within a loop or on a critical path,
+	 *       as this can save CPU cycles, memory, and mutex resources vs re-compiling the pattern for each format call. See FTextFormat for more info.
+	 *
+	 * @param Fmt The format pattern to use
+	 * @param InArguments The list of values to inject into the format pattern
+	 * @return The formatted FText
+	 */
 	static FText Format(FTextFormat Fmt, const FFormatOrderedArguments& InArguments);
 	static FText Format(FTextFormat Fmt, FFormatOrderedArguments&& InArguments);
 
+	/**
+	 * Format the given list of variadic values as ordered arguments within the given format pattern
+	 *
+	 * @note You may want to pre-compile your FText pattern into a FTextFormat prior to performing formats within a loop or on a critical path, 
+	 *       as this can save CPU cycles, memory, and mutex resources vs re-compiling the pattern for each format call. See FTextFormat for more info.
+	 * 
+	 * @usage FText::Format(LOCTEXT("PlayerNameFmt", "{0} is really cool"), FText::FromString(PlayerName));
+	 * 
+	 * @param Fmt The format pattern to use
+	 * @param Args A variadic list of values to inject into the format pattern
+	 * @return The formatted FText
+	 */
 	template <typename... ArgTypes>
 	static FORCEINLINE FText Format(FTextFormat Fmt, ArgTypes... Args)
 	{
@@ -568,23 +601,31 @@ public:
 	}
 
 	/**
-	 * FormatNamed allows you to pass name <-> value pairs to the function to format automatically
-	 *
-	 * @usage FText::FormatNamed( FText::FromString( TEXT( "{PlayerName} is really cool" ) ), TEXT( "PlayerName" ), FText::FromString( TEXT( "Awesomegirl" ) ) );
+	 * Format the given list of variadic key->value pairs as named arguments within the given format pattern
 	 *
-	 * @param Fmt the format to create from
-	 * @param Args a variadic list of FString to Value (must be even numbered)
-	 * @return a formatted FText
+	 * @note You may want to pre-compile your FText pattern into a FTextFormat prior to performing formats within a loop or on a critical path,
+	 *       as this can save CPU cycles, memory, and mutex resources vs re-compiling the pattern for each format call. See FTextFormat for more info.
+	 * 
+	 * @usage FText::FormatNamed(LOCTEXT("PlayerNameFmt", "{PlayerName} is really cool"), TEXT("PlayerName"), FText::FromString(PlayerName));
+	 * 
+	 * @param Fmt The format pattern to use
+	 * @param Args A variadic list of "key then value" pairs to inject into the format pattern (must be an even number)
+	 * @return The formatted FText
 	 */
 	template < typename... TArguments >
 	static FText FormatNamed( FTextFormat Fmt, TArguments&&... Args );
 
 	/**
-	 * FormatOrdered allows you to pass a variadic list of types to use for formatting in order desired
+	 * Format the given list of variadic values as ordered arguments within the given format pattern
 	 *
-	 * @param Fmt the format to create from
-	 * @param Args a variadic list of values in order of desired formatting
-	 * @return a formatted FText
+	 * @note You may want to pre-compile your FText pattern into a FTextFormat prior to performing formats within a loop or on a critical path,
+	 *       as this can save CPU cycles, memory, and mutex resources vs re-compiling the pattern for each format call. See FTextFormat for more info.
+	 * 
+	 * @usage FText::FormatOrdered(LOCTEXT("PlayerNameFmt", "{0} is really cool"), FText::FromString(PlayerName));
+	 * 
+	 * @param Fmt The format pattern to use
+	 * @param Args A variadic list of values to inject into the format pattern
+	 * @return The formatted FText
 	 */
 	template < typename... TArguments >
 	static FText FormatOrdered( FTextFormat Fmt, TArguments&&... Args );
