Fix unitialized FPackageStoreEntryResource::Flags member in default generated constructor.

#rnx

diff --git a/Engine/Source/Runtime/Core/Public/IO/PackageStore.h b/Engine/Source/Runtime/Core/Public/IO/PackageStore.h
--- a/Engine/Source/Runtime/Core/Public/IO/PackageStore.h
+++ b/Engine/Source/Runtime/Core/Public/IO/PackageStore.h
@@ -82,7 +82,7 @@ ENUM_CLASS_FLAGS(EPackageStoreEntryFlags);
 struct FPackageStoreEntryResource
 {
 	/** The package store entry flags. */
-	EPackageStoreEntryFlags Flags;
+	EPackageStoreEntryFlags Flags = EPackageStoreEntryFlags::None;
 	/** The package name. */
 	FName PackageName;
 	/** Used for localized and redirected packages. */
