Removed deprecated functionality from FScriptArray.

diff --git a/Engine/Source/Runtime/Core/Public/Containers/ScriptArray.h b/Engine/Source/Runtime/Core/Public/Containers/ScriptArray.h
--- a/Engine/Source/Runtime/Core/Public/Containers/ScriptArray.h
+++ b/Engine/Source/Runtime/Core/Public/Containers/ScriptArray.h
@@ -42,21 +42,11 @@ public:
 		checkSlow(ArrayMax>=ArrayNum);
 		return ArrayNum;
 	}
-	UE_DEPRECATED(5.0, "TScriptArray::InsertZeroed without an alignment parameter is deprecated.")
-	void InsertZeroed( int32 Index, int32 Count, int32 NumBytesPerElement )
-	{
-		InsertZeroed( Index, Count, NumBytesPerElement, __STDCPP_DEFAULT_NEW_ALIGNMENT__ );
-	}
 	void InsertZeroed( int32 Index, int32 Count, int32 NumBytesPerElement, uint32 AlignmentOfElement )
 	{
 		Insert( Index, Count, NumBytesPerElement, AlignmentOfElement );
 		FMemory::Memzero( (uint8*)this->GetAllocation()+Index*NumBytesPerElement, Count*NumBytesPerElement );
 	}
-	UE_DEPRECATED(5.0, "TScriptArray::Insert without an alignment parameter is deprecated.")
-	void Insert( int32 Index, int32 Count, int32 NumBytesPerElement )
-	{
-		Insert(Index, Count, NumBytesPerElement, __STDCPP_DEFAULT_NEW_ALIGNMENT__);
-	}
 	void Insert( int32 Index, int32 Count, int32 NumBytesPerElement, uint32 AlignmentOfElement )
 	{
 		check(Count>=0);
@@ -77,11 +67,6 @@ public:
 			                                               (OldNum-Index)*NumBytesPerElement
 		);
 	}
-	UE_DEPRECATED(5.0, "TScriptArray::Add without an alignment parameter is deprecated.")
-	int32 Add(int32 Count, int32 NumBytesPerElement)
-	{
-		return Add(Count, NumBytesPerElement, __STDCPP_DEFAULT_NEW_ALIGNMENT__);
-	}
 	int32 Add( int32 Count, int32 NumBytesPerElement, uint32 AlignmentOfElement )
 	{
 		check(Count>=0);
@@ -96,22 +81,12 @@ public:
 
 		return OldNum;
 	}
-	UE_DEPRECATED(5.0, "TScriptArray::AddZeroed without an alignment parameter is deprecated.")
-	int32 AddZeroed( int32 Count, int32 NumBytesPerElement )
-	{
-		return AddZeroed(Count, NumBytesPerElement, __STDCPP_DEFAULT_NEW_ALIGNMENT__);
-	}
 	int32 AddZeroed( int32 Count, int32 NumBytesPerElement, uint32 AlignmentOfElement )
 	{
 		const int32 Index = Add( Count, NumBytesPerElement, AlignmentOfElement );
 		FMemory::Memzero( (uint8*)this->GetAllocation()+Index*NumBytesPerElement, Count*NumBytesPerElement );
 		return Index;
 	}
-	UE_DEPRECATED(5.0, "TScriptArray::Shrink without an alignment parameter is deprecated.")
-	void Shrink( int32 NumBytesPerElement )
-	{
-		Shrink(NumBytesPerElement, __STDCPP_DEFAULT_NEW_ALIGNMENT__);
-	}
 	void Shrink( int32 NumBytesPerElement, uint32 AlignmentOfElement )
 	{
 		checkSlow(ArrayNum>=0);
@@ -121,11 +96,6 @@ public:
 			ResizeTo(ArrayNum, NumBytesPerElement, AlignmentOfElement);
 		}
 	}
-	UE_DEPRECATED(5.0, "TScriptArray::MoveAssign without an alignment parameter is deprecated.")
-	void MoveAssign(TScriptArray& Other, int32 NumBytesPerElement)
-	{
-		MoveAssign(Other, NumBytesPerElement, __STDCPP_DEFAULT_NEW_ALIGNMENT__);
-	}
 	void MoveAssign(TScriptArray& Other, int32 NumBytesPerElement, uint32 AlignmentOfElement)
 	{
 		checkSlow(this != &Other);
@@ -134,11 +104,6 @@ public:
 		ArrayNum = Other.ArrayNum; Other.ArrayNum = 0;
 		ArrayMax = Other.ArrayMax; Other.ArrayMax = 0;
 	}
-	UE_DEPRECATED(5.0, "TScriptArray::Empty without an alignment parameter is deprecated.")
-	void Empty( int32 Slack, int32 NumBytesPerElement )
-	{
-		Empty(Slack, NumBytesPerElement, __STDCPP_DEFAULT_NEW_ALIGNMENT__);
-	}
 	void Empty( int32 Slack, int32 NumBytesPerElement, uint32 AlignmentOfElement )
 	{
 		checkSlow(Slack>=0);
@@ -173,11 +138,6 @@ public:
 		return ArrayMax - ArrayNum;
 	}
 
-	UE_DEPRECATED(5.0, "TScriptArray::Remove without an alignment parameter is deprecated.")
-	void Remove( int32 Index, int32 Count, int32 NumBytesPerElement  )
-	{
-		Remove(Index, Count, NumBytesPerElement, __STDCPP_DEFAULT_NEW_ALIGNMENT__);
-	}
 	void Remove( int32 Index, int32 Count, int32 NumBytesPerElement, uint32 AlignmentOfElement )
 	{
 		if (Count)
@@ -208,11 +168,6 @@ public:
 
 protected:
 
-	UE_DEPRECATED(5.0, "FScriptArray::TScriptArray without an alignment parameter is deprecated.")
-	TScriptArray(int32 InNum, int32 NumBytesPerElement)
-		: TScriptArray(InNum, NumBytesPerElement, __STDCPP_DEFAULT_NEW_ALIGNMENT__)
-	{
-	}
 	TScriptArray( int32 InNum, int32 NumBytesPerElement, uint32 AlignmentOfElement )
 	:   ArrayNum( 0 )
 	,	ArrayMax( InNum )
@@ -274,22 +229,12 @@ class FScriptArray : public TScriptArray<FHeapAllocator>
 public:
 	FScriptArray() = default;
 
-	UE_DEPRECATED(5.0, "FScriptArray::MoveAssign without an alignment parameter is deprecated.")
-	void MoveAssign(FScriptArray& Other, int32 NumBytesPerElement)
-	{
-		MoveAssign(Other, NumBytesPerElement, __STDCPP_DEFAULT_NEW_ALIGNMENT__);
-	}
 	void MoveAssign(FScriptArray& Other, int32 NumBytesPerElement, uint32 AlignmentOfElement)
 	{
 		Super::MoveAssign(Other, NumBytesPerElement, AlignmentOfElement);
 	}
 
 protected:
-	UE_DEPRECATED(5.0, "FScriptArray::FScriptArray without an alignment parameter is deprecated.")
-	FScriptArray(int32 InNum, int32 NumBytesPerElement)
-		: TScriptArray<FHeapAllocator>(InNum, NumBytesPerElement, __STDCPP_DEFAULT_NEW_ALIGNMENT__)
-	{
-	}
 	FScriptArray(int32 InNum, int32 NumBytesPerElement, uint32 AlignmentOfElement)
 		: TScriptArray<FHeapAllocator>(InNum, NumBytesPerElement, AlignmentOfElement)
 	{
