Migrated some more functionality from TArray to FScriptArray.
Fixed constness of FScriptArray::CountBytes.
Added optional shrinking to FScriptArray::Remove.

diff --git a/Engine/Source/Runtime/Core/Public/Containers/ScriptArray.h b/Engine/Source/Runtime/Core/Public/Containers/ScriptArray.h
--- a/Engine/Source/Runtime/Core/Public/Containers/ScriptArray.h
+++ b/Engine/Source/Runtime/Core/Public/Containers/ScriptArray.h
@@ -96,6 +96,19 @@ public:
 			ResizeTo(ArrayNum, NumBytesPerElement, AlignmentOfElement);
 		}
 	}
+	void SetNumUninitialized(int32 NewNum, int32 NumBytesPerElement, uint32 AlignmentOfElement, bool bAllowShrinking = true)
+	{
+		checkSlow(NewNum >= 0);
+		int32 OldNum = Num();
+		if (NewNum > OldNum)
+		{
+			Add(NewNum - OldNum, NumBytesPerElement, AlignmentOfElement);
+		}
+		else if (NewNum < OldNum)
+		{
+			Remove(NewNum, OldNum - NewNum, NumBytesPerElement, AlignmentOfElement, bAllowShrinking);
+		}
+	}
 	void MoveAssign(TScriptArray& Other, int32 NumBytesPerElement, uint32 AlignmentOfElement)
 	{
 		checkSlow(this != &Other);
@@ -113,6 +126,17 @@ public:
 			ResizeTo(Slack, NumBytesPerElement, AlignmentOfElement);
 		}
 	}
+	void Reset(int32 NewSize, int32 NumBytesPerElement, uint32 AlignmentOfElement)
+	{
+		if (NewSize <= ArrayMax)
+		{
+			ArrayNum = 0;
+		}
+		else
+		{
+			Empty(NewSize, NumBytesPerElement, AlignmentOfElement);
+		}
+	}
 	void SwapMemory(int32 A, int32 B, int32 NumBytesPerElement )
 	{
 		FMemory::Memswap(
@@ -126,10 +150,14 @@ public:
 	,	ArrayMax( 0 )
 	{
 	}
-	void CountBytes( FArchive& Ar, int32 NumBytesPerElement  )
+	void CountBytes( FArchive& Ar, int32 NumBytesPerElement  ) const
 	{
 		Ar.CountBytes( ArrayNum*NumBytesPerElement, ArrayMax*NumBytesPerElement );
 	}
+	FORCEINLINE void CheckAddress(const void* Addr, int32 NumBytesPerElement) const
+	{
+		checkf((const char*)Addr < (const char*)GetData() || (const char*)Addr >= ((const char*)GetData() + ArrayMax * NumBytesPerElement), TEXT("Attempting to use a container element (%p) which already comes from the container being modified (%p, ArrayMax: %lld, ArrayNum: %lld, SizeofElement: %d)!"), Addr, GetData(), (long long)ArrayMax, (long long)ArrayNum, NumBytesPerElement);
+	}
 	/**
 	 * Returns the amount of slack in this array in elements.
 	 */
@@ -138,7 +166,7 @@ public:
 		return ArrayMax - ArrayNum;
 	}
 
-	void Remove( int32 Index, int32 Count, int32 NumBytesPerElement, uint32 AlignmentOfElement )
+	void Remove( int32 Index, int32 Count, int32 NumBytesPerElement, uint32 AlignmentOfElement, bool bAllowShrinking = true )
 	{
 		if (Count)
 		{
@@ -160,11 +188,18 @@ public:
 			}
 			ArrayNum -= Count;
 
-			ResizeShrink(NumBytesPerElement, AlignmentOfElement);
+			if (bAllowShrinking)
+			{
+				ResizeShrink(NumBytesPerElement, AlignmentOfElement);
+			}
 			checkSlow(ArrayNum >= 0);
 			checkSlow(ArrayMax >= ArrayNum);
 		}
 	}
+	SIZE_T GetAllocatedSize(int32 NumBytesPerElement) const
+	{
+		return ((const typename AllocatorType::ForAnyElementType*)this)->GetAllocatedSize(ArrayMax, NumBytesPerElement);
+	}
 
 protected:
 
