Fix change in behavior to FString::Find: when called with an empty string and IgnoreCase it should return INDEX_NONE and a recent change had changed this to return 0 instead.
#rnx

diff --git a/Engine/Source/Runtime/Core/Private/Containers/String.cpp b/Engine/Source/Runtime/Core/Private/Containers/String.cpp
--- a/Engine/Source/Runtime/Core/Private/Containers/String.cpp
+++ b/Engine/Source/Runtime/Core/Private/Containers/String.cpp
@@ -409,6 +409,11 @@ int32 FString::Find(const TCHAR* SubStr, int32 SubStrLen, ESearchCase::Type Sear
 			Start += FMath::Clamp(StartPosition, 0, RemainingLength - 1);
 			RemainingLength = UE_PTRDIFF_TO_INT32(End - Start);
 		}
+		if (SubStrLen == 0 && SearchCase == ESearchCase::IgnoreCase)
+		{
+			// Legacy behavior: when SubStr is empty string and SearchCase is IgnoreCase, return INDEX_NONE
+			return INDEX_NONE;
+		}
 		const TCHAR* Tmp = SearchCase == ESearchCase::IgnoreCase
 			? FCString::Strnistr(Start, RemainingLength, SubStr, SubStrLen)
 			: FCString::Strnstr(Start, RemainingLength, SubStr, SubStrLen);
diff --git a/Engine/Source/Runtime/Core/Private/Containers/StringTest.cpp b/Engine/Source/Runtime/Core/Private/Containers/StringTest.cpp
--- a/Engine/Source/Runtime/Core/Private/Containers/StringTest.cpp
+++ b/Engine/Source/Runtime/Core/Private/Containers/StringTest.cpp
@@ -710,6 +710,13 @@ bool FStringFindTest::RunTest(const FString& Parameters)
 	RunTest(ABACADAB, EmptyString, ESearchCase::CaseSensitive, ESearchDir::FromEnd, 2, 1);
 	RunTest(ABACADAB, EmptyString, ESearchCase::CaseSensitive, ESearchDir::FromEnd, 0, INDEX_NONE);
 
+	// Find with an empty search string has different behavior depending on SearchCase due to legacy implementation
+	RunTest(ABACADAB, EmptyString, ESearchCase::IgnoreCase, ESearchDir::FromStart, 0, INDEX_NONE);
+	RunTest(ABACADAB, EmptyString, ESearchCase::IgnoreCase, ESearchDir::FromStart, 4, INDEX_NONE);
+	RunTest(ABACADAB, EmptyString, ESearchCase::IgnoreCase, ESearchDir::FromEnd, 8, 7);
+	RunTest(ABACADAB, EmptyString, ESearchCase::IgnoreCase, ESearchDir::FromEnd, 2, 1);
+	RunTest(ABACADAB, EmptyString, ESearchCase::IgnoreCase, ESearchDir::FromEnd, 0, INDEX_NONE);
+
 	// Find with a null char*
 	int32 Actual = FString(ABACADAB).Find(nullptr, ESearchCase::CaseSensitive, ESearchDir::FromStart, 0);
 	if (Actual != INDEX_NONE)
