Expose TSphere/FSphere to Blueprints

Add some standard UE stuff to TSphere (operator== and != and SerializeFromMismatchedTag support.
Add TSphere3StructOpsTypeTraits in Property.cpp
Register names in UnrealNames.inl
Add FSphere/FSphere3f/FSphere3d in NoExportTypes.h. FSphere3f and FSphere3d are not exposed to Blueprints
Add UE_DECLARE_CORE_VARIANT_TYPE declaration in Class.h

diff --git a/Engine/Source/Runtime/Core/Public/Math/Sphere.h b/Engine/Source/Runtime/Core/Public/Math/Sphere.h
--- a/Engine/Source/Runtime/Core/Public/Math/Sphere.h
+++ b/Engine/Source/Runtime/Core/Public/Math/Sphere.h
@@ -92,7 +92,6 @@ public:
 	template<typename FArg, TEMPLATE_REQUIRES(!TIsSame<T, FArg>::Value)>
 	explicit TSphere(const TSphere<FArg>& From) : TSphere<T>(TVector<T>(From.Center), T(From.W)) {}
 
-
 	/**
 	 * Check whether two spheres are the same within specified tolerance.
 	 *
@@ -105,6 +104,47 @@ public:
 		return Center.Equals(Sphere.Center, Tolerance) && FMath::Abs(W - Sphere.W) <= Tolerance;
 	}
 
+	/**
+	* Compares two spheres for equality.
+	*
+	* @param Other The other sphere to compare with.
+	* @return true if the spheres are equal, false otherwise.
+	*/
+	bool operator==(const TSphere<T>& Other) const
+	{
+		return Center == Other.Center && W == Other.W;
+	}
+
+	/**
+	 * Compares two spheres for inequality.
+	 *
+	 * @param Other The other sphere to compare with.
+	 * @return true if the spheres are not equal, false otherwise.
+	 */
+	bool operator!=(const TSphere<T>& Other) const
+	{
+		return !(*this == Other);
+	}
+
+	/**
+	 * Gets the result of addition to this bounding volume.
+	 *
+	 * @param Other The other volume to add to this.
+	 * @return A new bounding volume.
+	 */
+	TSphere<T> operator+(const TSphere<T>& Other) const
+	{
+		return TSphere<T>(*this) += Other;
+	}
+
+	/**
+	 * Adds to this sphere to include a new bounding volume.
+	 *
+	 * @param Other the bounding volume to increase the bounding volume to.
+	 * @return Reference to this bounding volume after resizing to include the other bounding volume.
+	 */
+	TSphere<T>& operator+=(const TSphere<T>& Other);
+
 	/**
 	 * Check whether sphere is inside of another.
 	 *
@@ -172,25 +212,25 @@ public:
 	}
 
 	/**
-	 * Adds to this bounding box to include a new bounding volume.
-	 *
-	 * @param Other the bounding volume to increase the bounding volume to.
-	 * @return Reference to this bounding volume after resizing to include the other bounding volume.
-	 */
-	TSphere<T>& operator+=(const TSphere<T>& Other);
+	* Get a textual representation of the sphere.
+	*
+	* @return Text describing the sphere.
+	*/
+	FString ToString() const
+	{
+		return FString::Printf(TEXT("Center=(%s), Radius=(%s)"), *Center.ToString(), *W.ToString());
+	}
 
-	/**
-	 * Gets the result of addition to this bounding volume.
-	 *
-	 * @param Other The other volume to add to this.
-	 * @return A new bounding volume.
-	 */
-	TSphere operator+(const TSphere<T>& Other) const
+	// Note: TSphere is usually written via binary serialization. This function exists for SerializeFromMismatchedTag conversion usage. 
+	bool Serialize(FArchive& Ar)
 	{
-		return TSphere(*this) += Other;
+		Ar << *this;
+		return true;
 	}
-};
 
+	bool SerializeFromMismatchedTag(FName StructTag, FArchive& Ar);
+
+};
 
 /**
  * Serializes the given sphere from or into the specified archive.
@@ -318,14 +358,26 @@ template<> CORE_API TSphere<double>::TSphere(const TSphere<double>* Spheres, int
 
 UE_DECLARE_LWC_TYPE(Sphere, 3);
 
-//template<> struct TCanBulkSerialize<FSphere3f> { enum { Value = true }; };
-//template<> struct TIsPODType<FSphere3f> { enum { Value = true }; };
+template<> struct TCanBulkSerialize<FSphere3f> { enum { Value = true }; };
+template<> struct TIsPODType<FSphere3f> { enum { Value = true }; };
 template<> struct TIsUECoreVariant<FSphere3f> { enum { Value = true }; };
 
-//template<> struct TCanBulkSerialize<FSphere3d> { enum { Value = false }; };	// LWC_TODO: This can be done (via versioning) once LWC is fixed to on.
-//template<> struct TIsPODType<FSphere3d> { enum { Value = true }; };
+template<> struct TCanBulkSerialize<FSphere3d> { enum { Value = true }; };
+template<> struct TIsPODType<FSphere3d> { enum { Value = true }; };
 template<> struct TIsUECoreVariant<FSphere3d> { enum { Value = true }; };
 
+template<>
+inline bool FSphere3f::SerializeFromMismatchedTag(FName StructTag, FArchive& Ar)
+{
+	return UE_SERIALIZE_VARIANT_FROM_MISMATCHED_TAG(Ar, Sphere, Sphere3f, Sphere3d);
+}
+
+template<>
+inline bool FSphere3d::SerializeFromMismatchedTag(FName StructTag, FArchive& Ar)
+{
+	return UE_SERIALIZE_VARIANT_FROM_MISMATCHED_TAG(Ar, Sphere, Sphere3d, Sphere3f);
+}
+
 
 /* FMath inline functions
  *****************************************************************************/
diff --git a/Engine/Source/Runtime/Core/Public/UObject/UnrealNames.inl b/Engine/Source/Runtime/Core/Public/UObject/UnrealNames.inl
--- a/Engine/Source/Runtime/Core/Public/UObject/UnrealNames.inl
+++ b/Engine/Source/Runtime/Core/Public/UObject/UnrealNames.inl
@@ -137,6 +137,8 @@ REGISTER_NAME(168, Uint64Point)
 REGISTER_NAME(169, Ray)
 REGISTER_NAME(170, Ray3f)
 REGISTER_NAME(171, Ray3d)
+REGISTER_NAME(172, Sphere3f)
+REGISTER_NAME(173, Sphere3d)
 
 // Misc.
 REGISTER_NAME(200,State)
diff --git a/Engine/Source/Runtime/CoreUObject/Private/UObject/Property.cpp b/Engine/Source/Runtime/CoreUObject/Private/UObject/Property.cpp
--- a/Engine/Source/Runtime/CoreUObject/Private/UObject/Property.cpp
+++ b/Engine/Source/Runtime/CoreUObject/Private/UObject/Property.cpp
@@ -22,6 +22,7 @@
 #include "UObject/SoftObjectPath.h"
 #include "Math/Box2D.h"
 #include "Math/Ray.h"
+#include "Math/Sphere.h"
 #include "Math/InterpCurvePoint.h"
 #include "UObject/ReleaseObjectVersion.h"
 
@@ -351,6 +352,23 @@ UE_IMPLEMENT_STRUCT("/Script/CoreUObject", Ray3d);
 UE_IMPLEMENT_STRUCT("/Script/CoreUObject", Ray);
 
 
+template<typename T>
+struct TSphere3StructOpsTypeTraits : public TStructOpsTypeTraitsBase2<T>
+{
+	enum
+	{
+		WithIdenticalViaEquality = true,
+		WithNoInitConstructor = true,
+		WithZeroConstructor = true,
+		WithSerializeFromMismatchedTag = true,
+	};
+};
+template<> struct TStructOpsTypeTraits<FSphere3f> : public TSphere3StructOpsTypeTraits<FSphere3f> {};
+template<> struct TStructOpsTypeTraits<FSphere3d> : public TSphere3StructOpsTypeTraits<FSphere3d> {};
+UE_IMPLEMENT_STRUCT("/Script/CoreUObject", Sphere3f);
+UE_IMPLEMENT_STRUCT("/Script/CoreUObject", Sphere3d);
+UE_IMPLEMENT_STRUCT("/Script/CoreUObject", Sphere); 
+
 template<>
 struct TStructOpsTypeTraits<FInterpCurvePointFloat> : public TStructOpsTypeTraitsBase2<FInterpCurvePointFloat>
 {
diff --git a/Engine/Source/Runtime/CoreUObject/Public/UObject/Class.h b/Engine/Source/Runtime/CoreUObject/Public/UObject/Class.h
--- a/Engine/Source/Runtime/CoreUObject/Public/UObject/Class.h
+++ b/Engine/Source/Runtime/CoreUObject/Public/UObject/Class.h
@@ -4121,5 +4121,6 @@ UE_DECLARE_CORE_VARIANT_TYPE(Transform3,Transform);
 UE_DECLARE_CORE_VARIANT_TYPE(Matrix44,	Matrix);
 UE_DECLARE_CORE_VARIANT_TYPE(Box2,		Box2D);
 UE_DECLARE_CORE_VARIANT_TYPE(Ray3,		Ray);
+UE_DECLARE_CORE_VARIANT_TYPE(Sphere3, Sphere);
 
 #undef UE_DECLARE_CORE_VARIANT_TYPE
\ No newline at end of file
diff --git a/Engine/Source/Runtime/CoreUObject/Public/UObject/NoExportTypes.h b/Engine/Source/Runtime/CoreUObject/Public/UObject/NoExportTypes.h
--- a/Engine/Source/Runtime/CoreUObject/Public/UObject/NoExportTypes.h
+++ b/Engine/Source/Runtime/CoreUObject/Public/UObject/NoExportTypes.h
@@ -47,6 +47,7 @@
 #include "Math/RandomStream.h"
 #include "Math/RangeBound.h"
 #include "Math/Interval.h"
+#include "Math/Sphere.h"
 
 #include "Internationalization/PolyglotTextData.h"
 
@@ -779,6 +780,52 @@ struct FRotator
 	FLargeWorldCoordinatesReal Roll;
 };
 
+
+
+/**
+ * 3D Sphere represented by Center and Radius.
+ * @note The full C++ class is located here: Engine\Source\Runtime\Core\Public\Math\Sphere.h
+ * @note FSphere3f is not currently exposed as a Blueprint type
+ */
+USTRUCT(immutable, noexport, IsAlwaysAccessible, HasDefaults, HasNoOpConstructor, IsCoreType)
+struct FSphere3f
+{
+	UPROPERTY(EditAnywhere, Category = Sphere, SaveGame)
+	FVector3f Center;
+
+	UPROPERTY(EditAnywhere, Category = Sphere, SaveGame, meta = (DisplayName = "Radius"))
+	float W;
+};
+/**
+ * 3D Sphere represented by Center and Radius.
+ * @note The full C++ class is located here: Engine\Source\Runtime\Core\Public\Math\Sphere.h
+ * @note FSphere3d is not currently exposed as a Blueprint type
+ */
+USTRUCT(immutable, noexport, IsAlwaysAccessible, HasDefaults, HasNoOpConstructor, IsCoreType)
+struct FSphere3d
+{
+	UPROPERTY(EditAnywhere, Category = Sphere, SaveGame)
+	FVector3d Center;
+
+	UPROPERTY(EditAnywhere, Category = Sphere, SaveGame, meta = (DisplayName = "Radius"))
+	double W;
+};
+/**
+ * 3D Sphere represented by Center and Radius.
+ * @note The full C++ class is located here: Engine\Source\Runtime\Core\Public\Math\Sphere.h
+ */
+USTRUCT(immutable, noexport, BlueprintType, IsAlwaysAccessible, HasDefaults, HasNoOpConstructor, IsCoreType)
+struct FSphere
+{
+	UPROPERTY(EditAnywhere, BlueprintReadWrite, Category = Sphere, SaveGame)
+	FVector Center;
+
+	UPROPERTY(EditAnywhere, BlueprintReadWrite, Category = Sphere, SaveGame, meta = (DisplayName = "Radius"))
+	FLargeWorldCoordinatesReal W;
+};
+
+
+
 /**
  * Quaternion.
  * @note The full C++ class is located here: Engine\Source\Runtime\Core\Public\Math\Quat.h
