Fix retraction not working properly on assets async compilation causing many delays on the GT because of inability to reschedule blocking work

#rnx

diff --git a/Engine/Source/Runtime/Core/Public/Misc/QueuedThreadPoolWrapper.h b/Engine/Source/Runtime/Core/Public/Misc/QueuedThreadPoolWrapper.h
--- a/Engine/Source/Runtime/Core/Public/Misc/QueuedThreadPoolWrapper.h
+++ b/Engine/Source/Runtime/Core/Public/Misc/QueuedThreadPoolWrapper.h
@@ -459,7 +459,7 @@ private:
 
 		LowLevelTasks::ETaskPriority TaskPriorityMapper[int(EQueuedWorkPriority::Count)] = { LowLevelTasks::ETaskPriority::High, LowLevelTasks::ETaskPriority::High, LowLevelTasks::ETaskPriority::BackgroundHigh, LowLevelTasks::ETaskPriority::BackgroundNormal, LowLevelTasks::ETaskPriority::BackgroundLow, LowLevelTasks::ETaskPriority::BackgroundLow };
 		LowLevelTasks::ETaskPriority TaskPriority = TaskPriorityMapper[int(Priority)];
-		LowLevelTasks::ETaskFlags Flags = (InQueuedWork->GetQueuedWorkFlags() & EQueuedWorkFlags::DoNotRunInsideBusyWait) == EQueuedWorkFlags::None ? LowLevelTasks::ETaskFlags::DefaultFlags : (LowLevelTasks::ETaskFlags::DefaultFlags & ~LowLevelTasks::ETaskFlags::AllowCancellation);
+		LowLevelTasks::ETaskFlags Flags = (InQueuedWork->GetQueuedWorkFlags() & EQueuedWorkFlags::DoNotRunInsideBusyWait) == EQueuedWorkFlags::None ? LowLevelTasks::ETaskFlags::DefaultFlags : (LowLevelTasks::ETaskFlags::DefaultFlags & ~LowLevelTasks::ETaskFlags::AllowBusyWaiting);
 
 		QueuedWorkInternalData->Task.Init(TEXT("FQueuedLowLevelThreadPoolTask"), TaskPriority, [InQueuedWork, InternalData = InQueuedWork->InternalData, Deleter = LowLevelTasks::TDeleter<FQueuedLowLevelThreadPool, &FQueuedLowLevelThreadPool::FinalizeExecution>{ this }]
 		{
