renamed `TDepletableMpscQueue` to `TDepletableMpmcQueue` as it supports multiple concurrent consumers after the last change.
Deprecated the old name.
Updated existing use cases to use the new name or to switch to `TConsumeAllMpmcQueue` if the order of consumption doesn't matter.
Added `IsEmpty()` method to `TConsumeAllMpmcQueue`

diff --git a/Engine/Source/Developer/DerivedDataCache/Private/Http/CurlHttpClient.cpp b/Engine/Source/Developer/DerivedDataCache/Private/Http/CurlHttpClient.cpp
--- a/Engine/Source/Developer/DerivedDataCache/Private/Http/CurlHttpClient.cpp
+++ b/Engine/Source/Developer/DerivedDataCache/Private/Http/CurlHttpClient.cpp
@@ -3,7 +3,7 @@
 #include "HttpClient.h"
 
 #include "Async/InheritedContext.h"
-#include "Containers/DepletableMpscQueue.h"
+#include "Containers/ConsumeAllMpmcQueue.h"
 #include "Containers/LockFreeList.h"
 #include "Containers/StringView.h"
 #include "HAL/Event.h"
@@ -197,7 +197,7 @@ private:
 		EThreadCommandType Type;
 	};
 
-	TDepletableMpscQueue<FThreadCommand> ThreadCommands;
+	TConsumeAllMpmcQueue<FThreadCommand> ThreadCommands;
 	FThread Thread;
 	std::atomic<bool> bThreadStarting;
 	std::atomic<bool> bThreadStopping;
@@ -456,7 +456,7 @@ bool FCurlHttpConnectionPool::BeginAsyncRequest(FCurlHttpResponse* Response)
 	{
 		return false;
 	}
-	if (ThreadCommands.EnqueueAndReturnWasEmpty(FThreadCommand{Response, EThreadCommandType::Begin}))
+	if (ThreadCommands.ProduceItem(FThreadCommand{Response, EThreadCommandType::Begin}) == EConsumeAllMpmcQueueResult::WasEmpty)
 	{
 		AssertMultiCodeOk(curl_multi_wakeup(CurlMulti));
 	}
@@ -469,7 +469,7 @@ bool FCurlHttpConnectionPool::BeginAsyncRequest(FCurlHttpResponse* Response)
 
 void FCurlHttpConnectionPool::CancelAsyncRequest(FCurlHttpResponse* Response)
 {
-	if (ThreadCommands.EnqueueAndReturnWasEmpty(FThreadCommand{Response, EThreadCommandType::Cancel}))
+	if (ThreadCommands.ProduceItem(FThreadCommand{Response, EThreadCommandType::Cancel}) == EConsumeAllMpmcQueueResult::WasEmpty)
 	{
 		AssertMultiCodeOk(curl_multi_wakeup(CurlMulti));
 	}
@@ -479,7 +479,7 @@ void FCurlHttpConnectionPool::ThreadLoop()
 {
 	while (!ThreadCommands.IsEmpty() || !bThreadStopping.load(std::memory_order_relaxed))
 	{
-		ThreadCommands.Deplete([this](FThreadCommand Command)
+		ThreadCommands.ConsumeAllFifo([this](FThreadCommand Command)
 		{
 			if (CURL* Curl = Command.Response->GetCurl())
 			{
diff --git a/Engine/Source/Editor/UnrealEd/Private/PlayLevel.cpp b/Engine/Source/Editor/UnrealEd/Private/PlayLevel.cpp
--- a/Engine/Source/Editor/UnrealEd/Private/PlayLevel.cpp
+++ b/Engine/Source/Editor/UnrealEd/Private/PlayLevel.cpp
@@ -121,7 +121,7 @@
 #include "IAssetViewport.h"
 #include "IPIEAuthorizer.h"
 #include "Features/IModularFeatures.h"
-#include "Containers/DepletableMpscQueue.h"
+#include "Containers/DepletableMpmcQueue.h"
 #include "TickableEditorObject.h"
 
 DEFINE_LOG_CATEGORY(LogPlayLevel);
@@ -215,7 +215,7 @@ private:
 		}
 	}
 
-	UE::TDepletableMpscQueue<FLine> QueuedLines;
+	UE::TDepletableMpmcQueue<FLine> QueuedLines;
 };
 
 void UEditorEngine::EndPlayMap()
diff --git a/Engine/Source/Programs/UnrealPackageTool/UnrealPackageTool.cpp b/Engine/Source/Programs/UnrealPackageTool/UnrealPackageTool.cpp
--- a/Engine/Source/Programs/UnrealPackageTool/UnrealPackageTool.cpp
+++ b/Engine/Source/Programs/UnrealPackageTool/UnrealPackageTool.cpp
@@ -2,7 +2,7 @@
 
 #include "Async/ParallelFor.h"
 #include "Containers/Array.h"
-#include "Containers/DepletableMpscQueue.h"
+#include "Containers/ConsumeAllMpmcQueue.h"
 #include "Containers/StringView.h"
 #include "Containers/UnrealString.h"
 #include "HAL/PlatformFile.h"
@@ -88,13 +88,13 @@ public:
 			EPackageExtension Extension = FPackagePath::ParseExtension(Path);
 			if (Extension == EPackageExtension::Asset || Extension == EPackageExtension::Map)
 			{
-				PackagePaths.Enqueue(Path);
+				PackagePaths.ProduceItem(Path);
 			}
 		}
 		return true;
 	}
 
-	TDepletableMpscQueue<FString> PackagePaths;
+	TConsumeAllMpmcQueue<FString> PackagePaths;
 };
 
 void ScanPackage(const FParams& Params, const TCHAR* Path)
@@ -131,7 +131,7 @@ void Main(const FParams& Params)
 	}, EParallelForFlags::Unbalanced);
 
 	TArray<FString> PackagePaths;
-	Visitor.PackagePaths.Deplete([&PackagePaths](FString PackagePath)
+	Visitor.PackagePaths.ConsumeAllLifo([&PackagePaths](FString PackagePath)
 	{
 		PackagePaths.Emplace(MoveTemp(PackagePath));
 	});
diff --git a/Engine/Source/Runtime/Core/Private/Misc/OutputDeviceRedirector.cpp b/Engine/Source/Runtime/Core/Private/Misc/OutputDeviceRedirector.cpp
--- a/Engine/Source/Runtime/Core/Private/Misc/OutputDeviceRedirector.cpp
+++ b/Engine/Source/Runtime/Core/Private/Misc/OutputDeviceRedirector.cpp
@@ -3,7 +3,8 @@
 #include "Misc/OutputDeviceRedirector.h"
 
 #include "Containers/BitArray.h"
-#include "Containers/DepletableMpscQueue.h"
+#include "Containers/DepletableMpmcQueue.h"
+#include "Containers/ConsumeAllMpmcQueue.h"
 #include "Experimental/ConcurrentLinearAllocator.h"
 #include "HAL/Event.h"
 #include "HAL/PlatformProcess.h"
@@ -115,7 +116,7 @@ struct FOutputDeviceRedirectorState
 	uint8 OutputDevicesLockPadding[CalculateRedirectorCacheLinePadding(sizeof(OutputDevicesLock) + sizeof(OutputDevicesLockState))]{};
 
 	/** A queue of lines logged by non-primary threads. */
-	TDepletableMpscQueue<FOutputDeviceLine, FOutputDeviceLinearAllocator> BufferedLines;
+	TDepletableMpmcQueue<FOutputDeviceLine, FOutputDeviceLinearAllocator> BufferedLines;
 	uint8 BufferedLinesPadding[CalculateRedirectorCacheLinePadding(sizeof(BufferedLines))]{};
 
 	/** Array of output devices to redirect to from the primary thread. */
@@ -135,7 +136,7 @@ struct FOutputDeviceRedirectorState
 	FRWLock ThreadLock;
 
 	/** A queue of events to trigger when the dedicated primary thread is idle. */
-	TDepletableMpscQueue<FEvent*, FOutputDeviceLinearAllocator> ThreadIdleEvents;
+	TConsumeAllMpmcQueue<FEvent*, FOutputDeviceLinearAllocator> ThreadIdleEvents;
 
 	/** An event to wake the dedicated primary thread to process buffered lines. */
 	std::atomic<FEvent*> ThreadWakeEvent = nullptr;
@@ -436,7 +437,7 @@ void FOutputDeviceRedirectorState::ThreadLoop()
 				FlushBufferedLines();
 			}
 		}
-		ThreadIdleEvents.Deplete([](FEvent* Event) { Event->Trigger(); });
+		ThreadIdleEvents.ConsumeAllLifo([](FEvent* Event) { Event->Trigger(); });
 	}
 }
 
@@ -501,7 +502,7 @@ void FOutputDeviceRedirector::FlushThreadedLogs(EOutputDeviceRedirectorFlushOpti
 		{
 			TRACE_CPUPROFILER_EVENT_SCOPE(FOutputDeviceRedirector::FlushThreadedLogs);
 			FEventRef IdleEvent(EEventMode::ManualReset);
-			if (State->ThreadIdleEvents.EnqueueAndReturnWasEmpty(IdleEvent.Get()))
+			if (State->ThreadIdleEvents.ProduceItem(IdleEvent.Get()) == UE::EConsumeAllMpmcQueueResult::WasEmpty)
 			{
 				WakeEvent->Trigger();
 			}
diff --git a/Engine/Source/Runtime/Core/Private/Tests/Containers/ConcurrentQueuesTest.cpp b/Engine/Source/Runtime/Core/Private/Tests/Containers/ConcurrentQueuesTest.cpp
--- a/Engine/Source/Runtime/Core/Private/Tests/Containers/ConcurrentQueuesTest.cpp
+++ b/Engine/Source/Runtime/Core/Private/Tests/Containers/ConcurrentQueuesTest.cpp
@@ -9,7 +9,7 @@
 #include "Containers/SpscQueue.h"
 #include "Containers/MpscQueue.h"
 #include "Containers/ClosableMpscQueue.h"
-#include "Containers/DepletableMpscQueue.h"
+#include "Containers/DepletableMpmcQueue.h"
 #include "Tests/Benchmark.h"
 
 #include <atomic>
@@ -803,12 +803,12 @@ namespace DepletableMpscQueueTests
 {
 	// straightforward implementation over MPSC TQueue for comparison with TResettableMpscQueue
 	template<typename T>
-	class TDepletableMpscQueueOnTQueue
+	class TDepletableMpmcQueueOnTQueue
 	{
 		using FQueue = TQueue<T, EQueueMode::Mpsc>;
 
 	public:
-		~TDepletableMpscQueueOnTQueue()
+		~TDepletableMpmcQueueOnTQueue()
 		{
 			checkf(SubscribingThreadsNum.load(std::memory_order_relaxed) == 0, TEXT("No other threads can use the object during its destruction"));
 			delete Queue.load(std::memory_order_relaxed);
@@ -852,10 +852,10 @@ namespace DepletableMpscQueueTests
 
 	// straightforward implementation over MPSC TQueue for comparison with TResettableMpscQueue
 	template<typename T>
-	class TDepletableMpscQueueOnTMpscQueue
+	class TDepletableMpmcQueueOnTMpscQueue
 	{
 	public:
-		~TDepletableMpscQueueOnTMpscQueue()
+		~TDepletableMpmcQueueOnTMpscQueue()
 		{
 			checkf(SubscribingThreadsNum.load(std::memory_order_relaxed) == 0, TEXT("No other threads can use the object during its destration"));
 			delete Queue.load(std::memory_order_relaxed);
@@ -898,12 +898,12 @@ namespace DepletableMpscQueueTests
 
 	// wrapper over `TClosableLockFreePointerListUnorderedSingleConsumer`
 	template<typename T>
-	class TDepletableMpscQueueOnTClosableLockFreePointerListUnorderedSingleConsumer
+	class TDepletableMpmcQueueOnTClosableLockFreePointerListUnorderedSingleConsumer
 	{
 		using FQueue = TClosableLockFreePointerListUnorderedSingleConsumer<T, 0 /* as was used in FGraphEvent */>;
 
 	public:
-		~TDepletableMpscQueueOnTClosableLockFreePointerListUnorderedSingleConsumer()
+		~TDepletableMpmcQueueOnTClosableLockFreePointerListUnorderedSingleConsumer()
 		{
 			checkf(SubscribingThreadsNum.load(std::memory_order_relaxed) == 0, TEXT("No other threads can use the object during its destration"));
 			delete Queue.load(std::memory_order_relaxed);
@@ -948,12 +948,12 @@ namespace DepletableMpscQueueTests
 
 	// wrapper over `TClosableMpscQueue`
 	template<typename T>
-	class TDepletableMpscQueueOnTClosableMpscQueue
+	class TDepletableMpmcQueueOnTClosableMpscQueue
 	{
 		using FQueue = TClosableMpscQueue<T>;
 
 	public:
-		~TDepletableMpscQueueOnTClosableMpscQueue()
+		~TDepletableMpmcQueueOnTClosableMpscQueue()
 		{
 			checkf(SubscribingThreadsNum.load(std::memory_order_relaxed) == 0, TEXT("No other threads can use the object during its destration"));
 			delete Queue.load(std::memory_order_relaxed);
@@ -1057,7 +1057,7 @@ namespace DepletableMpscQueueTests
 	bool FDepletableMpscQueueTest::RunTest(const FString& Parameters)
 	{
 		{
-			TDepletableMpscQueue<int> Q;
+			TDepletableMpmcQueue<int> Q;
 			Q.Enqueue(1);
 			Q.Enqueue(2);
 			Q.Deplete([](int) {});
@@ -1068,18 +1068,18 @@ namespace DepletableMpscQueueTests
 		}
 
 		{
-			TDepletableMpscQueue<int> Q;
+			TDepletableMpmcQueue<int> Q;
 			Q.Deplete([](int) { check(false); });
 		}
 		{
-			TDepletableMpscQueue<int> Q;
+			TDepletableMpmcQueue<int> Q;
 			Q.Enqueue(1);
 			bool Done = false;
 			Q.Deplete([&Done](int Item) { check(!Done && Item == 1); Done = true; });
 		}
 
 		{
-			TDepletableMpscQueue<int> Q;
+			TDepletableMpmcQueue<int> Q;
 			Q.Enqueue(1);
 			Q.Enqueue(2);
 			int Step = 0;
@@ -1105,7 +1105,7 @@ namespace DepletableMpscQueueTests
 
 		//for (int i = 0; i != 1000000; ++i)
 		{
-			TDepletableMpscQueue<int> Q;
+			TDepletableMpmcQueue<int> Q;
 			Q.Deplete([](int) { check(false); });
 			Q.Deplete([](int) { check(false); });
 		}
@@ -1113,7 +1113,7 @@ namespace DepletableMpscQueueTests
 		//for (int i = 0; i != 1000000; ++i)
 		{
 			int Res = 0;
-			TDepletableMpscQueue<int> Q;
+			TDepletableMpmcQueue<int> Q;
 			Q.Enqueue(1);
 			Q.Deplete([&Res](int i) { Res += i; });
 			Q.Enqueue(2);
@@ -1122,7 +1122,7 @@ namespace DepletableMpscQueueTests
 		}
 
 		{	// EnqueueAndReturnWasEmpty
-			TDepletableMpscQueue<int> Q;
+			TDepletableMpmcQueue<int> Q;
 			verify(Q.EnqueueAndReturnWasEmpty(1));
 			verify(!Q.EnqueueAndReturnWasEmpty(2));
 			Q.Deplete([](int) {});
@@ -1133,17 +1133,17 @@ namespace DepletableMpscQueueTests
 			verify(!Q.EnqueueAndReturnWasEmpty(2));
 		}
 
-		UE_BENCHMARK(5, TestPerf<10'000, TDepletableMpscQueue<void*>>);
-		UE_BENCHMARK(5, TestPerf<10'000, TDepletableMpscQueueOnTQueue<void*>>);
-		UE_BENCHMARK(5, TestPerf<10'000, TDepletableMpscQueueOnTMpscQueue<void*>>);
-		UE_BENCHMARK(5, TestPerf<10'000, TDepletableMpscQueueOnTClosableLockFreePointerListUnorderedSingleConsumer<void>>);
-		UE_BENCHMARK(5, TestPerf<10'000, TDepletableMpscQueueOnTClosableMpscQueue<void*>>);
-
-		UE_BENCHMARK(5, TestConsumingEmptyQueue<10'000'000, TDepletableMpscQueue<void*>>);
-		UE_BENCHMARK(5, TestConsumingEmptyQueue<10'000'000, TDepletableMpscQueueOnTQueue<void*>>);
-		UE_BENCHMARK(5, TestConsumingEmptyQueue<10'000'000, TDepletableMpscQueueOnTMpscQueue<void*>>);
-		UE_BENCHMARK(5, TestConsumingEmptyQueue<10'000'000, TDepletableMpscQueueOnTClosableLockFreePointerListUnorderedSingleConsumer<void>>);
-		UE_BENCHMARK(5, TestConsumingEmptyQueue<10'000'000, TDepletableMpscQueueOnTClosableMpscQueue<void*>>);
+		UE_BENCHMARK(5, TestPerf<10'000, TDepletableMpmcQueue<void*>>);
+		UE_BENCHMARK(5, TestPerf<10'000, TDepletableMpmcQueueOnTQueue<void*>>);
+		UE_BENCHMARK(5, TestPerf<10'000, TDepletableMpmcQueueOnTMpscQueue<void*>>);
+		UE_BENCHMARK(5, TestPerf<10'000, TDepletableMpmcQueueOnTClosableLockFreePointerListUnorderedSingleConsumer<void>>);
+		UE_BENCHMARK(5, TestPerf<10'000, TDepletableMpmcQueueOnTClosableMpscQueue<void*>>);
+
+		UE_BENCHMARK(5, TestConsumingEmptyQueue<10'000'000, TDepletableMpmcQueue<void*>>);
+		UE_BENCHMARK(5, TestConsumingEmptyQueue<10'000'000, TDepletableMpmcQueueOnTQueue<void*>>);
+		UE_BENCHMARK(5, TestConsumingEmptyQueue<10'000'000, TDepletableMpmcQueueOnTMpscQueue<void*>>);
+		UE_BENCHMARK(5, TestConsumingEmptyQueue<10'000'000, TDepletableMpmcQueueOnTClosableLockFreePointerListUnorderedSingleConsumer<void>>);
+		UE_BENCHMARK(5, TestConsumingEmptyQueue<10'000'000, TDepletableMpmcQueueOnTClosableMpscQueue<void*>>);
 
 		return true;
 	}
diff --git a/Engine/Source/Runtime/Core/Public/Containers/ConsumeAllMpmcQueue.h b/Engine/Source/Runtime/Core/Public/Containers/ConsumeAllMpmcQueue.h
--- a/Engine/Source/Runtime/Core/Public/Containers/ConsumeAllMpmcQueue.h
+++ b/Engine/Source/Runtime/Core/Public/Containers/ConsumeAllMpmcQueue.h
@@ -81,6 +81,12 @@ namespace UE
 			return ConsumeAll<true>(Consumer);
 		}
 
+		// the result can be relied upon only in special cases (e.g. debug checks), as the state can change concurrently. use with caution 
+		bool IsEmpty() const
+		{
+			return Head.load(std::memory_order_relaxed) == nullptr;
+		}
+
 	private:
 		template<bool bReverse, typename F>
 		inline EConsumeAllMpmcQueueResult ConsumeAll(const F& Consumer)
diff --git a/Engine/Source/Runtime/Core/Public/Containers/DepletableMpmcQueue.h b/Engine/Source/Runtime/Core/Public/Containers/DepletableMpmcQueue.h
new file mode 100644
--- /dev/null
+++ b/Engine/Source/Runtime/Core/Public/Containers/DepletableMpmcQueue.h
@@ -0,0 +1,148 @@
+// Copyright Epic Games, Inc. All Rights Reserved.
+
+#pragma once
+
+#include "CoreTypes.h"
+#include "HAL/UnrealMemory.h"
+#include "Templates/TypeCompatibleBytes.h"
+#include <atomic>
+#include <type_traits>
+
+namespace UE
+{
+	/**
+	 * WARNING: this queue can cause priority inversion or a livelock due to spinlocking in	`Deplete()` method, though we haven't seen this
+	 * happened in practice. Prefer `TConsumeAllMpmcQueue` which is equally fast in LIFO mode, and only slightly slower in FIFO mode.
+	 *
+	 * Multi-producer/multi-consumer unbounded concurrent queue that is atomically consumed and is reset to its
+	 * default empty state.
+	 * A typical use case is when the consumer doesn't stop until the queue is depleted.
+	 * Is faster than traditional MPSC and MPMC queues, especially for consumer.
+	 * Consumes elements in FIFO order.
+	 */
+	template<typename T, typename AllocatorType = FMemory>
+	class TDepletableMpmcQueue final
+	{
+	private:
+		struct FNode
+		{
+			std::atomic<FNode*> Next{ nullptr };
+			TTypeCompatibleBytes<T> Value;
+		};
+
+		FNode Sentinel; // `Sentinel.Next` is the head of the queue
+		std::atomic<FNode*> Tail{ &Sentinel };
+
+	public:
+		UE_NONCOPYABLE(TDepletableMpmcQueue);
+
+		TDepletableMpmcQueue() = default;
+
+		~TDepletableMpmcQueue()
+		{
+			static_assert(std::is_trivially_destructible_v<FNode>);
+
+			// delete remaining elements
+			FNode* Node = Sentinel.Next.load(std::memory_order_relaxed);
+			while (Node != nullptr)
+			{
+				DestructItem(Node->Value.GetTypedPtr());
+				FNode* Next = Node->Next.load(std::memory_order_relaxed);
+				AllocatorType::Free(Node);
+				Node = Next;
+			}
+		}
+
+		template <typename... ArgTypes>
+		void Enqueue(ArgTypes&&... Args)
+		{
+			EnqueueAndReturnWasEmpty(Forward<ArgTypes>(Args)...);
+		}
+
+		template <typename... ArgTypes>
+		bool EnqueueAndReturnWasEmpty(ArgTypes&&... Args)
+		{
+			FNode* New = new(AllocatorType::Malloc(sizeof(FNode), alignof(FNode))) FNode;
+			new (&New->Value) T(Forward<ArgTypes>(Args)...);
+
+			// switch `Tail` to the new node and only then link the old tail to the new one. The list is not fully linked between these ops,
+			// this is explicitly handled by the consumer by waiting for the link
+
+			FNode* Prev = Tail.exchange(New, std::memory_order_acq_rel); // "acquire" to sync with `Prev->Next` initialisation from a concurrent calls,
+			// `release` to make sure the new node is fully constructed before it becomes visible to the consumer or a concurrent enqueueing
+			check(Prev->Next.load(std::memory_order_relaxed) == nullptr); // `Tail` is assigned before its Next
+
+			Prev->Next.store(New, std::memory_order_relaxed);
+
+			return Prev == &Sentinel;
+		}
+
+		/**
+		 * Takes all items from the queue atomically and then consumes them.
+		 * @param Consumer: a functor with signature `AnyReturnType (T Value)`
+		 */
+		template<typename F>
+		void Deplete(const F& Consumer)
+		{
+			// reset the head so the next consumption can detect that the queue is empty
+			FNode* First = Sentinel.Next.exchange(nullptr, std::memory_order_relaxed);
+			if (First == nullptr)
+			{
+				return; // empty
+			}
+
+			// reset the queue to the empty state. this redirects producers to start from `Sentinel` again.
+			// take note of the tail on resetting it because the list can be still not fully linked and so `Node.Next == nullptr` can't be 
+			// used to detect the end of the list
+			FNode* Last = Tail.exchange(&Sentinel, std::memory_order_acq_rel); // `acquire` to sync with producers' tail modifications, and
+			// "release" to force `Sentinel.Next = nullptr` happening before modifying `Tail
+			check(Last->Next.load(std::memory_order_relaxed) == nullptr); // `Tail` is assigned before its Next
+			// the previously queued items are detached from the instance (as a linked list, though potentially not fully linked yet)
+
+			check(Last != &Sentinel); // can't be empty because of `First != nullptr` above
+			Deplete_Internal(First, Last, Consumer);
+		}
+
+		// the result can be relied upon only in special cases (e.g. debug checks), as the state can change concurrently. use with caution 
+		bool IsEmpty() const
+		{
+			return Tail.load(std::memory_order_relaxed) == &Sentinel;
+		}
+
+	private:
+		template<typename F>
+		static void Deplete_Internal(FNode* First, FNode* Last, const F& Consumer)
+		{
+			auto GetNext = [](FNode* Node)
+			{
+				FNode* Next = nullptr;
+				// producers can be still updating `Next`, wait until the link to the next element is established
+				while (Next == nullptr) // <- This loop has the potential for live locking if enqueue was not completed (e.g was running at lower priority)
+				{
+					Next = Node->Next.load(std::memory_order_relaxed);
+				};
+				return Next;
+			};
+
+			auto Consume = [&Consumer](FNode* Node)
+			{
+				T* ValuePtr = (T*)&Node->Value;
+				Consumer(MoveTemp(*ValuePtr));
+				DestructItem(ValuePtr);
+				AllocatorType::Free(Node);
+			};
+
+			while (First != Last)
+			{
+				FNode* Next = GetNext(First);
+				Consume(First);
+				First = Next;
+			}
+
+			Consume(Last);
+		}
+	};
+
+	template<typename T, typename AllocatorType = FMemory>
+	using TDepletableMpscQueue UE_DEPRECATED(5.2, "The queue was renamed to TDepletableMpmcQueue as now it supports multiple concurrent consumers. Please consider switching to TConsumeAllMpmcQueue that is not vulnerable to live locking.") = TDepletableMpmcQueue<T, AllocatorType>;
+}
diff --git a/Engine/Source/Runtime/Core/Public/Containers/DepletableMpscQueue.h b/Engine/Source/Runtime/Core/Public/Containers/DepletableMpscQueue.h
--- a/Engine/Source/Runtime/Core/Public/Containers/DepletableMpscQueue.h
+++ b/Engine/Source/Runtime/Core/Public/Containers/DepletableMpscQueue.h
@@ -2,140 +2,10 @@
 
 #pragma once
 
-#include "CoreTypes.h"
-#include "HAL/UnrealMemory.h"
-#include "Templates/TypeCompatibleBytes.h"
-#include <atomic>
-#include <type_traits>
+#ifdef _MSC_VER
+#pragma message(__FILE__"(6): warning : #include Containers/DepletableMpmcQueue.h instead of Containers/DepletableMpscQueue.h (/s/Mpsc/Mpmc)")
+#else
+#pragma message("#include Containers/DepletableMpmcQueue.h instead of Containers/DepletableMpscQueue.h (/s/Mpsc/Mpmc)")
+#endif
 
-namespace UE
-{
-	/**
-	 * Multi-producer/single-consumer unbounded concurrent queue that is atomically consumed and is reset to its 
-	 * default empty state.
-	 * A typical use case is when the consumer doesn't stop until the queue is depleted.
-	 * Is faster than traditional MPSC queues, especially for consumer.
-	 */
-	template<typename T, typename AllocatorType = FMemory>
-	class TDepletableMpscQueue final
-	{
-	private:
-		struct FNode
-		{
-			std::atomic<FNode*> Next{ nullptr };
-			TTypeCompatibleBytes<T> Value;
-		};
-
-		FNode Sentinel; // `Sentinel.Next` is the head of the queue
-		std::atomic<FNode*> Tail{ &Sentinel };
-
-	public:
-		UE_NONCOPYABLE(TDepletableMpscQueue);
-
-		TDepletableMpscQueue() = default;
-
-		~TDepletableMpscQueue()
-		{
-			static_assert(std::is_trivially_destructible_v<FNode>);
-
-			// delete remaining elements
-			FNode* Node = Sentinel.Next.load(std::memory_order_relaxed);
-			while (Node != nullptr)
-			{
-				DestructItem(Node->Value.GetTypedPtr());
-				FNode* Next = Node->Next.load(std::memory_order_relaxed);
-				AllocatorType::Free(Node);
-				Node = Next;
-			}
-		}
-
-		template <typename... ArgTypes>
-		void Enqueue(ArgTypes&&... Args)
-		{
-			EnqueueAndReturnWasEmpty(Forward<ArgTypes>(Args)...);
-		}
-
-		template <typename... ArgTypes>
-		bool EnqueueAndReturnWasEmpty(ArgTypes&&... Args)
-		{
-			FNode* New = new(AllocatorType::Malloc(sizeof(FNode), alignof(FNode))) FNode;
-			new (&New->Value) T(Forward<ArgTypes>(Args)...);
-
-			// switch `Tail` to the new node and only then link the old tail to the new one. The list is not fully linked between these ops,
-			// this is explicitly handled by the consumer by waiting for the link
-
-			FNode* Prev = Tail.exchange(New, std::memory_order_acq_rel); // "acquire" to sync with `Prev->Next` initialisation from a concurrent calls,
-			// `release` to make sure the new node is fully constructed before it becomes visible to the consumer or a concurrent enqueueing
-			check(Prev->Next.load(std::memory_order_relaxed) == nullptr); // `Tail` is assigned before its Next
-
-			Prev->Next.store(New, std::memory_order_relaxed);
-
-			return Prev == &Sentinel;
-		}
-
-		/**
-		 * Takes all items from the queue atomically and then consumes them.
-		 * @param Consumer: a functor with signature `AnyReturnType (T Value)`
-		 */
-		template<typename F>
-		void Deplete(const F& Consumer)
-		{
-			// reset the head so the next consumption can detect that the queue is empty
-			FNode* First = Sentinel.Next.exchange(nullptr, std::memory_order_relaxed);
-			if (First == nullptr)
-			{
-				return; // empty
-			}
-
-			// reset the queue to the empty state. this redirects producers to start from `Sentinel` again.
-			// take note of the tail on resetting it because the list can be still not fully linked and so `Node.Next == nullptr` can't be 
-			// used to detect the end of the list
-			FNode* Last = Tail.exchange(&Sentinel, std::memory_order_acq_rel); // `acquire` to sync with producers' tail modifications, and
-			// "release" to force `Sentinel.Next = nullptr` happening before modifying `Tail
-			check(Last->Next.load(std::memory_order_relaxed) == nullptr); // `Tail` is assigned before its Next
-			// the previously queued items are detached from the instance (as a linked list, though potentially not fully linked yet)
-
-			check(Last != &Sentinel); // can't be empty because of `First != nullptr` above
-			Deplete_Internal(First, Last, Consumer);
-		}
-
-		// the result can be relied upon only in special cases, as the state can change concurrently. use with caution 
-		bool IsEmpty() const
-		{
-			return Tail.load(std::memory_order_relaxed) == &Sentinel;
-		}
-
-	private:
-		template<typename F>
-		static void Deplete_Internal(FNode* First, FNode* Last, const F& Consumer)
-		{
-			auto GetNext = [](FNode* Node)
-			{
-				FNode* Next = nullptr;
-				// producers can be still updating `Next`, wait until the link to the next element is established
-				while(Next == nullptr) // <- This loop has the potential for live locking if enqueue was not completed (e.g was running at lower priority)
-				{
-					Next = Node->Next.load(std::memory_order_relaxed);
-				};
-				return Next;
-			};
-
-			auto Consume = [&Consumer](FNode* Node)
-			{
-				T* ValuePtr = (T*)&Node->Value;
-				Consumer(MoveTemp(*ValuePtr));
-				DestructItem(ValuePtr);
-				AllocatorType::Free(Node);
-			};
-
-			while (First != Last)
-			{
-				FNode* Next = GetNext(First);
-				Consume(First);
-				First = Next;
-			}
-
-			Consume(Last);
-		}
-	};
-}
+#include "Containers/DepletableMpmcQueue.h"
