tuned concurrent queues unit tests to not take so long on windows

diff --git a/Engine/Source/Runtime/Core/Private/Tests/Containers/ConcurrentQueuesTest.cpp b/Engine/Source/Runtime/Core/Private/Tests/Containers/ConcurrentQueuesTest.cpp
--- a/Engine/Source/Runtime/Core/Private/Tests/Containers/ConcurrentQueuesTest.cpp
+++ b/Engine/Source/Runtime/Core/Private/Tests/Containers/ConcurrentQueuesTest.cpp
@@ -333,8 +333,8 @@ namespace UE { namespace ConcurrentQueuesTests
 			// SPSC
 
 			{	// destroy queue while it's holding one unconsumed item
-			TSpscQueue<FNonTrivial> Q;
-			Q.Enqueue(MakeUnique<int>(1));
+				TSpscQueue<FNonTrivial> Q;
+				Q.Enqueue(MakeUnique<int>(1));
 			}
 
 			{	// destroy queue while it's holding one cached consumed time
@@ -465,8 +465,8 @@ namespace UE { namespace ConcurrentQueuesTests
 		UE_BENCHMARK(5, TestQueueSingleThread<5'000'000, TMpscQueue<uint32>>);
 		UE_BENCHMARK(5, TestQueueSingleThread<5'000'000, TSpscQueue<uint32>>);
 
-		UE_BENCHMARK(5, TestSpscQueueCorrectness<5'000'000, TMpscQueue<uint32>>);
-		UE_BENCHMARK(5, TestSpscQueueCorrectness<5'000'000, TSpscQueue<uint32>>);
+		UE_BENCHMARK(5, TestSpscQueueCorrectness<1'000'000, TMpscQueue<uint32>>);
+		UE_BENCHMARK(5, TestSpscQueueCorrectness<1'000'000, TSpscQueue<uint32>>);
 
 		UE_BENCHMARK(5, TestTCircularQueue<5'000'000>);
 		UE_BENCHMARK(5, TestSpscQueue<5'000'000, TQueueAdapter<TQueue<uint32, EQueueMode::Spsc>>>);
@@ -477,8 +477,8 @@ namespace UE { namespace ConcurrentQueuesTests
 		UE_BENCHMARK(5, TestMpscQueueCorrectness<5'000'000, TQueueAdapter<TQueue<uint32, EQueueMode::Mpsc>>>);
 		UE_BENCHMARK(5, TestMpscQueueCorrectness<5'000'000, TMpscQueue<uint32>>);
 
-		UE_BENCHMARK(5, TestMpscQueue<1'000'000, TQueueAdapter<TQueue<uint32, EQueueMode::Mpsc>>>);
-		UE_BENCHMARK(5, TestMpscQueue<1'000'000, TMpscQueue<uint32>>);
+		UE_BENCHMARK(5, TestMpscQueue<100'000, TQueueAdapter<TQueue<uint32, EQueueMode::Mpsc>>>);
+		UE_BENCHMARK(5, TestMpscQueue<100'000, TMpscQueue<uint32>>);
 
 		return true;
 	}
@@ -785,10 +785,10 @@ namespace ClosableMpscQueueTests
 		UE_BENCHMARK(5, TestClosingEmptyQueue<10'000'000, TClosableMpscQueueOnTMpscQueue<void*>>);
 		UE_BENCHMARK(5, TestClosingEmptyQueue<10'000'000, TClosableLockFreePointerListUnorderedSingleConsumer_Adapter<void>>);
 
-		UE_BENCHMARK(5, TestCorrectness<300, TClosableMpscQueue<void*>>);
-		UE_BENCHMARK(5, TestCorrectness<300, TClosableMpscQueueOnTQueue<void*>>);
-		UE_BENCHMARK(5, TestCorrectness<300, TClosableMpscQueueOnTMpscQueue<void*>>);
-		UE_BENCHMARK(5, TestCorrectness<300, TClosableLockFreePointerListUnorderedSingleConsumer_Adapter<void>>);
+		UE_BENCHMARK(5, TestCorrectness<20, TClosableMpscQueue<void*>>);
+		UE_BENCHMARK(5, TestCorrectness<20, TClosableMpscQueueOnTQueue<void*>>);
+		UE_BENCHMARK(5, TestCorrectness<20, TClosableMpscQueueOnTMpscQueue<void*>>);
+		UE_BENCHMARK(5, TestCorrectness<20, TClosableLockFreePointerListUnorderedSingleConsumer_Adapter<void>>);
 
 		UE_BENCHMARK(5, TestPerf<5'000'000, TClosableMpscQueue<void*>>);
 		UE_BENCHMARK(5, TestPerf<5'000'000, TClosableMpscQueueOnTQueue<void*>>);
