Removed TSparseArray::GetData() in favor of operating on elements via explicitly-cast pointers - part of an upcoming change to remove FElementOrFreeListLink from the API of TSparseArray, which requires the full definition of the ElementType.  Also saves some PDB size.

diff --git a/Engine/Source/Runtime/Core/Public/Containers/SparseArray.h b/Engine/Source/Runtime/Core/Public/Containers/SparseArray.h
--- a/Engine/Source/Runtime/Core/Public/Containers/SparseArray.h
+++ b/Engine/Source/Runtime/Core/Public/Containers/SparseArray.h
@@ -99,7 +99,7 @@ public:
 		// Set the allocation info.
 		FSparseArrayAllocationInfo Result;
 		Result.Index = Index;
-		Result.Pointer = &GetData(Result.Index).ElementData;
+		Result.Pointer = &((FElementOrFreeListLink*)Data.GetData())[Result.Index].ElementData;
 
 		return Result;
 	}
@@ -113,13 +113,15 @@ public:
 		int32 Index;
 		if(NumFreeIndices)
 		{
+			FElementOrFreeListLink* DataPtr = (FElementOrFreeListLink*)Data.GetData();
+
 			// Remove and use the first index from the list of free elements.
 			Index = FirstFreeIndex;
-			FirstFreeIndex = GetData(FirstFreeIndex).NextFreeIndex;
+			FirstFreeIndex = DataPtr[FirstFreeIndex].NextFreeIndex;
 			--NumFreeIndices;
 			if(NumFreeIndices)
 			{
-				GetData(FirstFreeIndex).PrevFreeIndex = -1;
+				DataPtr[FirstFreeIndex].PrevFreeIndex = -1;
 			}
 		}
 		else
@@ -150,29 +152,31 @@ public:
 
 	FSparseArrayAllocationInfo AddUninitializedAtLowestFreeIndex(int32& LowestFreeIndexSearchStart)
 	{
+		FElementOrFreeListLink* DataPtr;
+
 		int32 Index;
 		if(NumFreeIndices)
 		{
 			Index = AllocationFlags.FindAndSetFirstZeroBit(LowestFreeIndexSearchStart);
 			LowestFreeIndexSearchStart = Index + 1;
 
-			auto& IndexData = GetData(Index);
+			DataPtr = (FElementOrFreeListLink*)Data.GetData();
 
 			// Update FirstFreeIndex
 			if (FirstFreeIndex == Index)
 			{
-				FirstFreeIndex = IndexData.NextFreeIndex;
+				FirstFreeIndex = DataPtr[Index].NextFreeIndex;
 			}
 
 			// Link our next and prev free nodes together
-			if (IndexData.NextFreeIndex >= 0)
+			if (DataPtr[Index].NextFreeIndex >= 0)
 			{
-				GetData(IndexData.NextFreeIndex).PrevFreeIndex = IndexData.PrevFreeIndex;
+				DataPtr[DataPtr[Index].NextFreeIndex].PrevFreeIndex = DataPtr[Index].PrevFreeIndex;
 			}
 
-			if (IndexData.PrevFreeIndex >= 0)
+			if (DataPtr[Index].PrevFreeIndex >= 0)
 			{
-				GetData(IndexData.PrevFreeIndex).NextFreeIndex = IndexData.NextFreeIndex;
+				DataPtr[DataPtr[Index].PrevFreeIndex].NextFreeIndex = DataPtr[Index].NextFreeIndex;
 			}
 
 			--NumFreeIndices;
@@ -182,11 +186,14 @@ public:
 			// Add a new element.
 			Index = Data.AddUninitialized(1);
 			AllocationFlags.Add(true);
+
+			// Defer getting the data pointer until after a possible reallocation
+			DataPtr = (FElementOrFreeListLink*)Data.GetData();
 		}
 
 		FSparseArrayAllocationInfo Result;
 		Result.Index = Index;
-		Result.Pointer = &GetData(Result.Index).ElementData;
+		Result.Pointer = &DataPtr[Result.Index].ElementData;
 		return Result;
 	}
 
@@ -250,9 +257,9 @@ public:
 		else
 		{
 			Allocation.Index = Index;
-			Allocation.Pointer = &GetData(Allocation.Index).ElementData;			
+			Allocation.Pointer = &((FElementOrFreeListLink*)Data.GetData())[Allocation.Index].ElementData;
 		}
-		
+
 		new(Allocation) ElementType(Forward<ArgsType>(Args)...);
 		return Allocation.Index;
 	}
@@ -263,35 +270,45 @@ public:
 	 */
 	FSparseArrayAllocationInfo InsertUninitialized(int32 Index)
 	{
+		FElementOrFreeListLink* DataPtr;
+
 		// Enlarge the array to include the given index.
 		if(Index >= Data.Num())
 		{
 			Data.AddUninitialized(Index + 1 - Data.Num());
+
+			// Defer getting the data pointer until after a possible reallocation
+			DataPtr = (FElementOrFreeListLink*)Data.GetData();
+
 			while(AllocationFlags.Num() < Data.Num())
 			{
 				const int32 FreeIndex = AllocationFlags.Num();
-				GetData(FreeIndex).PrevFreeIndex = -1;
-				GetData(FreeIndex).NextFreeIndex = FirstFreeIndex;
+				DataPtr[FreeIndex].PrevFreeIndex = -1;
+				DataPtr[FreeIndex].NextFreeIndex = FirstFreeIndex;
 				if(NumFreeIndices)
 				{
-					GetData(FirstFreeIndex).PrevFreeIndex = FreeIndex;
+					DataPtr[FirstFreeIndex].PrevFreeIndex = FreeIndex;
 				}
 				FirstFreeIndex = FreeIndex;
 				verify(AllocationFlags.Add(false) == FreeIndex);
 				++NumFreeIndices;
 			};
 		}
+		else
+		{
+			DataPtr = (FElementOrFreeListLink*)Data.GetData();
+		}
 
 		// Verify that the specified index is free.
 		check(!AllocationFlags[Index]);
 
 		// Remove the index from the list of free elements.
 		--NumFreeIndices;
-		const int32 PrevFreeIndex = GetData(Index).PrevFreeIndex;
-		const int32 NextFreeIndex = GetData(Index).NextFreeIndex;
+		const int32 PrevFreeIndex = DataPtr[Index].PrevFreeIndex;
+		const int32 NextFreeIndex = DataPtr[Index].NextFreeIndex;
 		if(PrevFreeIndex != -1)
 		{
-			GetData(PrevFreeIndex).NextFreeIndex = NextFreeIndex;
+			DataPtr[PrevFreeIndex].NextFreeIndex = NextFreeIndex;
 		}
 		else
 		{
@@ -299,7 +316,7 @@ public:
 		}
 		if(NextFreeIndex != -1)
 		{
-			GetData(NextFreeIndex).PrevFreeIndex = PrevFreeIndex;
+			DataPtr[NextFreeIndex].PrevFreeIndex = PrevFreeIndex;
 		}
 
 		return AllocateIndex(Index);
@@ -318,9 +335,10 @@ public:
 	{
 		if (!TIsTriviallyDestructible<ElementType>::Value)
 		{
+			FElementOrFreeListLink* DataPtr = (FElementOrFreeListLink*)Data.GetData();
 			for (int32 It = Index, ItCount = Count; ItCount; ++It, --ItCount)
 			{
-				((ElementType&)GetData(It).ElementData).~ElementType();
+				((ElementType&)DataPtr[It].ElementData).~ElementType();
 			}
 		}
 
@@ -330,6 +348,8 @@ public:
 	/** Removes Count elements from the array, starting from Index, without destructing them. */
 	void RemoveAtUninitialized(int32 Index,int32 Count = 1)
 	{
+		FElementOrFreeListLink* DataPtr = (FElementOrFreeListLink*)Data.GetData();
+
 		for (; Count; --Count)
 		{
 			check(AllocationFlags[Index]);
@@ -337,11 +357,10 @@ public:
 			// Mark the element as free and add it to the free element list.
 			if(NumFreeIndices)
 			{
-				GetData(FirstFreeIndex).PrevFreeIndex = Index;
+				DataPtr[FirstFreeIndex].PrevFreeIndex = Index;
 			}
-			auto& IndexData = GetData(Index);
-			IndexData.PrevFreeIndex = -1;
-			IndexData.NextFreeIndex = NumFreeIndices > 0 ? FirstFreeIndex : INDEX_NONE;
+			DataPtr[Index].PrevFreeIndex = -1;
+			DataPtr[Index].NextFreeIndex = NumFreeIndices > 0 ? FirstFreeIndex : INDEX_NONE;
 			FirstFreeIndex = Index;
 			++NumFreeIndices;
 			AllocationFlags[Index] = false;
@@ -407,15 +426,17 @@ public:
 			// allocate memory in the array itself
 			int32 ElementIndex = Data.AddUninitialized(ElementsToAdd);
 
+			FElementOrFreeListLink* DataPtr = (FElementOrFreeListLink*)Data.GetData();
+
 			// now mark the new elements as free
 			for ( int32 FreeIndex = ExpectedNumElements - 1; FreeIndex >= ElementIndex; --FreeIndex )
 			{
 				if(NumFreeIndices)
 				{
-					GetData(FirstFreeIndex).PrevFreeIndex = FreeIndex;
+					DataPtr[FirstFreeIndex].PrevFreeIndex = FreeIndex;
 				}
-				GetData(FreeIndex).PrevFreeIndex = -1;
-				GetData(FreeIndex).NextFreeIndex = NumFreeIndices > 0 ? FirstFreeIndex : INDEX_NONE;
+				DataPtr[FreeIndex].PrevFreeIndex = -1;
+				DataPtr[FreeIndex].NextFreeIndex = NumFreeIndices > 0 ? FirstFreeIndex : INDEX_NONE;
 				FirstFreeIndex = FreeIndex;
 				++NumFreeIndices;
 			}
@@ -442,21 +463,23 @@ public:
 		{
 			if(NumFreeIndices > 0)
 			{
+				FElementOrFreeListLink* DataPtr = (FElementOrFreeListLink*)Data.GetData();
+
 				// Look for elements in the free list that are in the memory to be freed.
 				int32 FreeIndex = FirstFreeIndex;
 				while(FreeIndex != INDEX_NONE)
 				{
 					if(FreeIndex >= FirstIndexToRemove)
 					{
-						const int32 PrevFreeIndex = GetData(FreeIndex).PrevFreeIndex;
-						const int32 NextFreeIndex = GetData(FreeIndex).NextFreeIndex;
+						const int32 PrevFreeIndex = DataPtr[FreeIndex].PrevFreeIndex;
+						const int32 NextFreeIndex = DataPtr[FreeIndex].NextFreeIndex;
 						if(NextFreeIndex != -1)
 						{
-							GetData(NextFreeIndex).PrevFreeIndex = PrevFreeIndex;
+							DataPtr[NextFreeIndex].PrevFreeIndex = PrevFreeIndex;
 						}
 						if(PrevFreeIndex != -1)
 						{
-							GetData(PrevFreeIndex).NextFreeIndex = NextFreeIndex;
+							DataPtr[PrevFreeIndex].NextFreeIndex = NextFreeIndex;
 						}
 						else
 						{
@@ -468,13 +491,13 @@ public:
 					}
 					else
 					{
-						FreeIndex = GetData(FreeIndex).NextFreeIndex;
+						FreeIndex = DataPtr[FreeIndex].NextFreeIndex;
 					}
 				}
 			}
 
 			// Truncate unallocated elements at the end of the data array.
-			Data.RemoveAt(FirstIndexToRemove,Data.Num() - FirstIndexToRemove, false);
+			Data.RemoveAt(FirstIndexToRemove, Data.Num() - FirstIndexToRemove, false);
 			AllocationFlags.RemoveAt(FirstIndexToRemove,AllocationFlags.Num() - FirstIndexToRemove);
 		}
 
@@ -494,14 +517,14 @@ public:
 
 		bool bResult = false;
 
-		FElementOrFreeListLink* ElementData = Data.GetData();
+		FElementOrFreeListLink* DataPtr = (FElementOrFreeListLink*)Data.GetData();
 
 		int32 EndIndex    = Data.Num();
 		int32 TargetIndex = EndIndex - NumFree;
 		int32 FreeIndex   = FirstFreeIndex;
 		while (FreeIndex != -1)
 		{
-			int32 NextFreeIndex = GetData(FreeIndex).NextFreeIndex;
+			int32 NextFreeIndex = DataPtr[FreeIndex].NextFreeIndex;
 			if (FreeIndex < TargetIndex)
 			{
 				// We need an element here
@@ -511,7 +534,7 @@ public:
 				}
 				while (!AllocationFlags[EndIndex]);
 
-				RelocateConstructItems<FElementOrFreeListLink>(ElementData + FreeIndex, ElementData + EndIndex, 1);
+				RelocateConstructItems<FElementOrFreeListLink>(DataPtr + FreeIndex, DataPtr + EndIndex, 1);
 				AllocationFlags[FreeIndex] = true;
 
 				bResult = true;
@@ -520,7 +543,7 @@ public:
 			FreeIndex = NextFreeIndex;
 		}
 
-		Data           .RemoveAt(TargetIndex, NumFree, false);
+		Data.RemoveAt(TargetIndex, NumFree, false);
 		AllocationFlags.RemoveAt(TargetIndex, NumFree);
 
 		NumFreeIndices = 0;
@@ -565,7 +588,7 @@ public:
 			Compact();
 
 			// Sort the elements according to the provided comparison class.
-			::Sort( &GetData(0), Num(), FElementCompareClass< PREDICATE_CLASS >( Predicate ) );
+			::Sort( (FElementOrFreeListLink*)Data.GetData(), Num(), FElementCompareClass< PREDICATE_CLASS >( Predicate ) );
 		}
 	}
 
@@ -585,7 +608,7 @@ public:
 			CompactStable();
 
 			// Sort the elements according to the provided comparison class.
-			::StableSort(&GetData(0), Num(), FElementCompareClass< PREDICATE_CLASS >(Predicate));
+			::StableSort((FElementOrFreeListLink*)Data.GetData(), Num(), FElementCompareClass< PREDICATE_CLASS >(Predicate));
 		}
 	}
 
@@ -703,12 +726,12 @@ public:
 			NumFreeIndices  = InCopy.NumFreeIndices;
 			AllocationFlags = InCopy.AllocationFlags;
 
+			      FElementOrFreeListLink* DestData = (      FElementOrFreeListLink*)Data.GetData();
+			const FElementOrFreeListLink* SrcData  = (const FElementOrFreeListLink*)InCopy.Data.GetData();
+
 			// Determine whether we need per element construction or bulk copy is fine
 			if (!TIsTriviallyCopyConstructible<ElementType>::Value)
 			{
-				      FElementOrFreeListLink* DestData = (FElementOrFreeListLink*)Data.GetData();
-				const FElementOrFreeListLink* SrcData  = (FElementOrFreeListLink*)InCopy.Data.GetData();
-
 				// Use the inplace new to copy the element to an array element
 				for (int32 Index = 0; Index < SrcMax; ++Index)
 				{
@@ -728,7 +751,7 @@ public:
 			else
 			{
 				// Use the much faster path for types that allow it
-				FMemory::Memcpy(Data.GetData(), InCopy.Data.GetData(), sizeof(FElementOrFreeListLink) * SrcMax);
+				FMemory::Memcpy(DestData, SrcData, sizeof(FElementOrFreeListLink) * SrcMax);
 			}
 		}
 		return *this;
@@ -762,18 +785,18 @@ public:
 	{
 		checkSlow(Index >= 0 && Index < Data.Num() && Index < AllocationFlags.Num());
 		//checkSlow(AllocationFlags[Index]); // Disabled to improve loading times -BZ
-		return *(ElementType*)&GetData(Index).ElementData;
+		return *(ElementType*)&((FElementOrFreeListLink*)Data.GetData())[Index].ElementData;
 	}
 	const ElementType& operator[](int32 Index) const
 	{
 		checkSlow(Index >= 0 && Index < Data.Num() && Index < AllocationFlags.Num());
 		//checkSlow(AllocationFlags[Index]); // Disabled to improve loading times -BZ
-		return *(ElementType*)&GetData(Index).ElementData;
+		return *(ElementType*)&((FElementOrFreeListLink*)Data.GetData())[Index].ElementData;
 	}
 	int32 PointerToIndex(const ElementType* Ptr) const
 	{
 		checkSlow(Data.Num());
-		int32 Index = (int32)((FElementOrFreeListLink*)Ptr - &GetData(0));
+		int32 Index = (int32)((FElementOrFreeListLink*)Ptr - (FElementOrFreeListLink*)Data.GetData());
 		checkSlow(Index >= 0 && Index < Data.Num() && Index < AllocationFlags.Num() && AllocationFlags[Index]);
 		return Index;
 	}
@@ -1045,18 +1068,6 @@ private:
 		}
 	};
 
-	/** Accessor for the element or free list data. */
-	FElementOrFreeListLink& GetData(int32 Index)
-	{
-		return ((FElementOrFreeListLink*)Data.GetData())[Index];
-	}
-
-	/** Accessor for the element or free list data. */
-	const FElementOrFreeListLink& GetData(int32 Index) const
-	{
-		return ((FElementOrFreeListLink*)Data.GetData())[Index];
-	}
-
 	typedef TArray<FElementOrFreeListLink,typename Allocator::ElementAllocator> DataType;
 	DataType Data;
 
@@ -1083,7 +1094,7 @@ public:
 				FMemoryImageWriter ArrayWriter = Writer.WritePointer(ElementTypeDesc);
 				for (int32 i = 0; i < NumElements; ++i)
 				{
-					const FElementOrFreeListLink& Elem = this->Data[i];
+					const FElementOrFreeListLink& Elem = ((const FElementOrFreeListLink*)this->Data.GetData())[i];
 					const uint32 StartOffset = ArrayWriter.WriteAlignment<FElementOrFreeListLink>();
 					if (this->AllocationFlags[i])
 					{
@@ -1126,8 +1137,8 @@ public:
 				DstObject->Data.SetNumUninitialized(this->Data.Num());
 				for (int32 i = 0; i < this->Data.Num(); ++i)
 				{
-					const FElementOrFreeListLink& Elem = this->Data[i];
-					FElementOrFreeListLink& DstElem = DstObject->Data[i];
+					const FElementOrFreeListLink& Elem    = ((const FElementOrFreeListLink*)this     ->Data.GetData())[i];
+					      FElementOrFreeListLink& DstElem = ((      FElementOrFreeListLink*)DstObject->Data.GetData())[i];
 					if (this->AllocationFlags[i])
 					{
 						Context.UnfreezeObject(&Elem.ElementData, ElementTypeDesc, &DstElem.ElementData);
