- Fixed a crash on shutdown in Mac programs, with TaskGraph shutting down earlier than some Stats code that uses it

diff --git a/Engine/Source/Runtime/Core/Private/Stats/StatsCommand.cpp b/Engine/Source/Runtime/Core/Private/Stats/StatsCommand.cpp
--- a/Engine/Source/Runtime/Core/Private/Stats/StatsCommand.cpp
+++ b/Engine/Source/Runtime/Core/Private/Stats/StatsCommand.cpp
@@ -2057,12 +2057,21 @@ static void StatCmd(FString InCmd, bool bStatCommand, FOutputDevice* Ar /*= null
 
 			Stats.ResetStatsForRawStats();
 
-			// Disable displaying the raw stats memory overhead.
-			FSimpleDelegateGraphTask::CreateAndDispatchWhenReady
-				(
-				FSimpleDelegateGraphTask::FDelegate::CreateRaw(&FLatestGameThreadStatsData::Get(), &FLatestGameThreadStatsData::NewData, (FGameThreadStatsData*)nullptr),
-				TStatId(), nullptr, ENamedThreads::GameThread
-				);
+			// stopfile command happens when some threads shutdown, and depending on order of shutdown operations, the taskgraph may
+			// shutdown before stopfile command is executed
+			if (FTaskGraphInterface::IsRunning())
+			{
+				// Disable displaying the raw stats memory overhead.
+				FSimpleDelegateGraphTask::CreateAndDispatchWhenReady
+					(
+					 FSimpleDelegateGraphTask::FDelegate::CreateRaw(&FLatestGameThreadStatsData::Get(), &FLatestGameThreadStatsData::NewData, (FGameThreadStatsData*)nullptr),
+					 TStatId(), nullptr, ENamedThreads::GameThread
+					 );
+			}
+			else
+			{
+				FLatestGameThreadStatsData::Get().NewData(nullptr);
+			}
 		}
 		else if (FParse::Command(&Cmd, TEXT("TESTFILE")))
 		{
