Probable fix for HoloLens compile error.

diff --git a/Engine/Source/Runtime/Core/Public/Templates/Tuple.h b/Engine/Source/Runtime/Core/Public/Templates/Tuple.h
--- a/Engine/Source/Runtime/Core/Public/Templates/Tuple.h
+++ b/Engine/Source/Runtime/Core/Public/Templates/Tuple.h
@@ -648,7 +648,7 @@ namespace UE::Core::Private::Tuple
 		template <typename TupleType>
 		FORCEINLINE static uint32 Do(uint32 Hash, const TupleType& Tuple)
 		{
-			return TGetTupleHashHelper<ArgToCombine + 1, ArgCount>::Do(HashCombineFast(Hash, GetTypeHash(Tuple.template Get<ArgToCombine>())), Tuple);
+			return TGetTupleHashHelper<ArgToCombine + 1, ArgCount>::Do(HashCombineFast(Hash, GetTypeHashHelper(Tuple.template Get<ArgToCombine>())), Tuple);
 		}
 	};
 
