NormalizePackageNames: Fix NormalizePackageNames handling of NORMALIZE_ExcludeDeveloperPackages to exclude/include Developers packages' external actor and external object package paths.
FPaths: Add DevelopersFolderName.
#rnx

diff --git a/Engine/Source/Editor/UnrealEd/Private/Commandlets/DiffAssetRegistriesCommandlet.cpp b/Engine/Source/Editor/UnrealEd/Private/Commandlets/DiffAssetRegistriesCommandlet.cpp
--- a/Engine/Source/Editor/UnrealEd/Private/Commandlets/DiffAssetRegistriesCommandlet.cpp
+++ b/Engine/Source/Editor/UnrealEd/Private/Commandlets/DiffAssetRegistriesCommandlet.cpp
@@ -302,7 +302,8 @@ void UDiffAssetRegistriesCommandlet::FillChangelists(FString Branch, FString CL,
 					FString PostContentPath;
 					if (DepotPathName.Split(BasePath, nullptr, &PostContentPath))
 					{
-						if (!PostContentPath.IsEmpty() && !PostContentPath.StartsWith(TEXT("Cinematics")) && !PostContentPath.StartsWith(TEXT("Developers")) && !PostContentPath.StartsWith(TEXT("Maps/Test_Maps")))
+						if (!PostContentPath.IsEmpty() && !PostContentPath.StartsWith(TEXT("Cinematics")) &&
+							!PostContentPath.StartsWith(FPaths::DevelopersFolderName()) && !PostContentPath.StartsWith(TEXT("Maps/Test_Maps")))
 						{
 							const FString PostContentPathWithoutExtension = FPaths::GetBaseFilename(PostContentPath, false);
 							const FString FullPackageName = AssetPath + PostContentPathWithoutExtension;
@@ -340,7 +341,8 @@ void UDiffAssetRegistriesCommandlet::FillChangelists(FString Branch, FString CL,
 					FString PostContentPath;
 					if (DepotPathName.Split(BasePath, nullptr, &PostContentPath))
 					{
-						if (!PostContentPath.IsEmpty() && !PostContentPath.StartsWith(TEXT("Cinematics")) && !PostContentPath.StartsWith(TEXT("Developers")) && !PostContentPath.StartsWith(TEXT("Maps/Test_Maps")))
+						if (!PostContentPath.IsEmpty() && !PostContentPath.StartsWith(TEXT("Cinematics")) &&
+							!PostContentPath.StartsWith(FPaths::DevelopersFolderName()) && !PostContentPath.StartsWith(TEXT("Maps/Test_Maps")))
 						{
 							const FString PostContentPathWithoutExtension = FPaths::GetBaseFilename(PostContentPath, false);
 							const FString FullPackageName = AssetPath + PostContentPathWithoutExtension;
diff --git a/Engine/Source/Editor/UnrealEd/Private/Commandlets/PackageUtilities.cpp b/Engine/Source/Editor/UnrealEd/Private/Commandlets/PackageUtilities.cpp
--- a/Engine/Source/Editor/UnrealEd/Private/Commandlets/PackageUtilities.cpp
+++ b/Engine/Source/Editor/UnrealEd/Private/Commandlets/PackageUtilities.cpp
@@ -105,7 +105,14 @@ bool NormalizePackageNames( TArray<FString> PackageNames, TArray<FString>& Packa
 		IFileManager::Get().FindFiles( PackageNames, *PackageWildcard, true, false );
 	}
 
-	const FString DeveloperFolder = IFileManager::Get().ConvertToAbsolutePathForExternalAppForRead(*FPaths::GameDevelopersDir());
+	FString ProjectContentDir = FPaths::ProjectContentDir();
+	FStringView DevelopersFolderName = FPaths::DevelopersFolderName();
+	const FString DeveloperFolder = IFileManager::Get().ConvertToAbsolutePathForExternalAppForRead(
+		*FPaths::Combine(ProjectContentDir, DevelopersFolderName));
+	const FString DeveloperExternalActorsFolder = IFileManager::Get().ConvertToAbsolutePathForExternalAppForRead(
+		*FPaths::Combine(ProjectContentDir, FPackagePath::GetExternalActorsFolderName(), DevelopersFolderName));
+	const FString DeveloperExternalObjectsFolder = IFileManager::Get().ConvertToAbsolutePathForExternalAppForRead(
+		*FPaths::Combine(ProjectContentDir, FPackagePath::GetExternalObjectsFolderName(), DevelopersFolderName));
 
 	if( PackageNames.Num() == 0 )
 	{
@@ -195,17 +202,16 @@ bool NormalizePackageNames( TArray<FString> PackageNames, TArray<FString>& Packa
 
 			FString Filename = IFileManager::Get().ConvertToAbsolutePathForExternalAppForRead(*PackagePathNames[PackageIndex]);
 			
-			if ( (PackageFilter&NORMALIZE_ExcludeDeveloperPackages) != 0 )
-			{
-				if (Filename.StartsWith(DeveloperFolder))
-				{
-					PackagePathNames.RemoveAt(PackageIndex);
-					continue;
-				}
-			}
-			else if ( (PackageFilter&NORMALIZE_ExcludeNonDeveloperPackages) != 0 )
-			{
-				if (!Filename.StartsWith(DeveloperFolder))
+			if ( (PackageFilter & NORMALIZE_ExcludeDeveloperPackages) != 0 || 
+				 (PackageFilter & NORMALIZE_ExcludeNonDeveloperPackages) != 0)
+			{
+				// Technically both flags present should mean exclude everything, but legacy behavior is to have
+				// excludedevelopers override excludenondevelopers.
+				bool bDevelopersFolderIncluded = !((PackageFilter & NORMALIZE_ExcludeDeveloperPackages) != 0);
+				bool bIsDeveloperFolder = Filename.StartsWith(DeveloperFolder) ||
+					Filename.StartsWith(DeveloperExternalActorsFolder) ||
+					Filename.StartsWith(DeveloperExternalObjectsFolder);
+				if (bDevelopersFolderIncluded != bIsDeveloperFolder)
 				{
 					PackagePathNames.RemoveAt(PackageIndex);
 					continue;
diff --git a/Engine/Source/Runtime/Core/Private/Misc/Paths.cpp b/Engine/Source/Runtime/Core/Private/Misc/Paths.cpp
--- a/Engine/Source/Runtime/Core/Private/Misc/Paths.cpp
+++ b/Engine/Source/Runtime/Core/Private/Misc/Paths.cpp
@@ -481,7 +481,15 @@ FString FPaths::CloudDir()
 
 FString FPaths::GameDevelopersDir()
 {
-	return FPaths::ProjectContentDir() + TEXT("Developers/");
+	FString ContentDir = FPaths::ProjectContentDir();
+	FStringView DevelopersFolder = DevelopersFolderName();
+	return FString::Printf(TEXT("%.*s%.*s/"), ContentDir.Len(), *ContentDir,
+		DevelopersFolder.Len(), DevelopersFolder.GetData());
+}
+
+FStringView FPaths::DevelopersFolderName()
+{
+	return TEXTVIEW("Developers");
 }
 
 FString FPaths::GameUserDeveloperFolderName()
@@ -1589,23 +1597,23 @@ const FString& FPaths::GetRelativePathToRoot()
 	return StaticData.RelativePathToRoot;
 }
 
-void FPaths::CombineInternal(FString& OutPath, const TCHAR** Pathes, int32 NumPathes)
+void FPaths::CombineInternal(FString& OutPath, const FStringView* Paths, int32 NumPaths)
 {
-	check(Pathes != NULL && NumPathes > 0);
+	check(Paths != NULL && NumPaths > 0);
 
 	int32 OutStringSize = 0;
 
-	for (int32 i=0; i < NumPathes; ++i)
+	for (int32 i=0; i < NumPaths; ++i)
 	{
-		OutStringSize += FCString::Strlen(Pathes[i]) + 1;
+		OutStringSize += Paths[i].Len() + 1;
 	}
 
 	OutPath.Empty(OutStringSize);
-	OutPath += Pathes[0];
+	OutPath += Paths[0];
 	
-	for (int32 i=1; i < NumPathes; ++i)
+	for (int32 i=1; i < NumPaths; ++i)
 	{
-		OutPath /= Pathes[i];
+		OutPath /= Paths[i];
 	}
 }
 
diff --git a/Engine/Source/Runtime/Core/Public/Misc/Paths.h b/Engine/Source/Runtime/Core/Public/Misc/Paths.h
--- a/Engine/Source/Runtime/Core/Public/Misc/Paths.h
+++ b/Engine/Source/Runtime/Core/Public/Misc/Paths.h
@@ -3,6 +3,7 @@
 #pragma once
 
 #include "Containers/Array.h"
+#include "Containers/StringView.h"
 #include "Containers/UnrealString.h"
 #include "CoreTypes.h"
 #include "HAL/CriticalSection.h"
@@ -16,6 +17,7 @@ namespace UE::Core::Private
 {
 	const TCHAR*   GetBindingType(const TCHAR* Ptr);
 	const FString& GetBindingType(const FString& Str);
+	const FStringView& GetBindingType(const FStringView& StringView);
 
 	// This is used to force the arguments of FString::Combine to implicitly convert (if necessary)
 	// to FString when calling CombineImpl(), allowing them to remain as temporaries on the stack
@@ -334,9 +336,18 @@ public:
 	/** Returns the directory for local files used in cloud emulation or support */
 	static FString CloudDir();
 
-	/** Returns the directory that contains subfolders for developer-specific content */
+	/**
+	 * Returns the directory that contains subfolders for developer-specific content
+	 * Example: "../../../ProjectName/Content/Developers
+	 */
 	static FString GameDevelopersDir();
 
+	/**
+	 * Returns the name of the subfolder for developer-specific content
+	 * Example: "Developers"
+	 */
+	static FStringView DevelopersFolderName();
+
 	/** Returns The folder name for the developer-specific directory for the current user */
 	static FString GameUserDeveloperFolderName();
 
@@ -637,6 +648,26 @@ public:
 		return CombineImpl<PathTypes...>(Forward<PathTypes>(InPaths)...);
 	}
 
+#if 0
+	FORCEINLINE static FString Combine(const FString& Path1, const FStringView& Path2)
+	{
+		const FStringView Paths[] = { InternalGetStringView(Path1), InternalGetStringView(Path2) };
+		FString Out;
+
+		CombineInternal(Out, Paths, UE_ARRAY_COUNT(Paths));
+		return Out;
+	}
+
+	FORCEINLINE static FString Combine(const FString& Path1, const TCHAR* PathMid, const FStringView& Path2)
+	{
+		const FStringView Paths[] = { InternalGetStringView(Path1), InternalGetStringView(PathMid), InternalGetStringView(Path2) };
+		FString Out;
+
+		CombineInternal(Out, Paths, UE_ARRAY_COUNT(Paths));
+		return Out;
+	}
+#endif
+
 	/**
 	 * Frees any memory retained by FPaths.
 	 */
@@ -644,7 +675,7 @@ public:
 
 protected:
 
-	static void CombineInternal(FString& OutPath, const TCHAR** Paths, int32 NumPaths);
+	static void CombineInternal(FString& OutPath, const FStringView* Paths, int32 NumPaths);
 
 private:
 	struct FStaticData;
@@ -652,23 +683,13 @@ private:
 	template <typename... PathTypes>
 	FORCEINLINE static FString CombineImpl(UE::Core::Private::TToStringType_T<std::decay_t<PathTypes>>... InPaths)
 	{
-		const TCHAR* Paths[] = { GetTCharPtr(InPaths)... };
+		const FStringView Paths[] = { FStringView(InPaths)... };
 		FString Out;
 		
 		CombineInternal(Out, Paths, UE_ARRAY_COUNT(Paths));
 		return Out;
 	}
 
-	FORCEINLINE static const TCHAR* GetTCharPtr(const TCHAR* Ptr)
-	{
-		return Ptr;
-	}
-
-	FORCEINLINE static const TCHAR* GetTCharPtr(const FString& Str)
-	{
-		return *Str;
-	}
-
 	/** Returns, if any, the value of the -userdir command line argument. This can be used to sandbox artifacts to a desired location */
 	static const FString& CustomUserDirArgument();
 
