Fixed %hh* formats in FGenericWidePlatform::GetVarArgs.
Added defensive checks against snprintf returning a negative result.

diff --git a/Engine/Source/Runtime/Core/Private/GenericPlatform/GenericWidePlatformString.cpp b/Engine/Source/Runtime/Core/Private/GenericPlatform/GenericWidePlatformString.cpp
--- a/Engine/Source/Runtime/Core/Private/GenericPlatform/GenericWidePlatformString.cpp
+++ b/Engine/Source/Runtime/Core/Private/GenericPlatform/GenericWidePlatformString.cpp
@@ -574,7 +574,7 @@ int32 FGenericWidePlatformString::GetVarArgs( WIDECHAR* Dest, SIZE_T DestSize, c
 				FmtBuf[CpyIdx] = 0;
 
 				int RetCnt = snprintf(AnsiNum, sizeof(AnsiNum), FmtBuf, Val);
-				if (!DestIter.Write(AnsiNum, RetCnt))
+				if (RetCnt < 0 || !DestIter.Write(AnsiNum, RetCnt))
 				{
 					return -1;
 				}
@@ -623,7 +623,7 @@ int32 FGenericWidePlatformString::GetVarArgs( WIDECHAR* Dest, SIZE_T DestSize, c
 				FmtBuf[CpyIdx] = 0;
 
 				int RetCnt = snprintf(AnsiNum, sizeof(AnsiNum), FmtBuf, Val);
-				if (!DestIter.Write(AnsiNum, RetCnt))
+				if (RetCnt < 0 || !DestIter.Write(AnsiNum, RetCnt))
 				{
 					return -1;
 				}
@@ -648,7 +648,7 @@ int32 FGenericWidePlatformString::GetVarArgs( WIDECHAR* Dest, SIZE_T DestSize, c
 				FmtBuf[CpyIdx] = 0;
 
 				int RetCnt = snprintf(AnsiNum, sizeof(AnsiNum), FmtBuf, Val);
-				if (!DestIter.Write(AnsiNum, RetCnt))
+				if (RetCnt < 0 || !DestIter.Write(AnsiNum, RetCnt))
 				{
 					return -1;
 				}
@@ -669,7 +669,7 @@ int32 FGenericWidePlatformString::GetVarArgs( WIDECHAR* Dest, SIZE_T DestSize, c
 				}
 
 				// treat %ld as %d. Also shorts for %h will be promoted to ints. This path also handles %li, %lu, %lx and %lX.
-				if ((Src[0] == 'l' && CharIsIntegerFormatSpecifier(Src[1])) || Src[0] == 'h')
+				if ((Src[0] == 'l' || Src[0] == 'h') && CharIsIntegerFormatSpecifier(Src[1]))
 				{
 					Src+=2;
 					long int Val = va_arg(ArgPtr, long int);
@@ -687,7 +687,7 @@ int32 FGenericWidePlatformString::GetVarArgs( WIDECHAR* Dest, SIZE_T DestSize, c
 					FmtBuf[CpyIdx] = 0;
 
 					int RetCnt = snprintf(AnsiNum, sizeof(AnsiNum), FmtBuf, Val);
-					if (!DestIter.Write(AnsiNum, RetCnt))
+					if (RetCnt < 0 || !DestIter.Write(AnsiNum, RetCnt))
 					{
 						return -1;
 					}
@@ -712,7 +712,7 @@ int32 FGenericWidePlatformString::GetVarArgs( WIDECHAR* Dest, SIZE_T DestSize, c
 					FmtBuf[CpyIdx] = 0;
 
 					int RetCnt = snprintf(AnsiNum, sizeof(AnsiNum), FmtBuf, Val);
-					if (!DestIter.Write(AnsiNum, RetCnt))
+					if (RetCnt < 0 || !DestIter.Write(AnsiNum, RetCnt))
 					{
 						return -1;
 					}
@@ -727,7 +727,38 @@ int32 FGenericWidePlatformString::GetVarArgs( WIDECHAR* Dest, SIZE_T DestSize, c
 					break;
 				}
 
-				if (Src[0] == 'l')
+				// Deal with half-half formats - %hhd, %hhi, %hhu, %hhx and %hhX.
+				if (Src[0] == 'h')
+				{
+					if (Src[1] != 'h' || !CharIsIntegerFormatSpecifier(Src[2]))
+					{
+						printf("Unknown percent [%lc%lc%lc] in FGenericWidePlatformString::GetVarArgs() [%s]\n.", Src[0], Src[1], Src[2], TCHAR_TO_ANSI(Fmt));
+						Src++;  // skip it, I guess.
+						break;
+					}
+					Src += 3;
+					int Val = va_arg(ArgPtr, int);
+					ANSICHAR AnsiNum[8];
+					ANSICHAR FmtBuf[30];
+
+					// Yes, this is lame.
+					int CpyIdx = 0;
+					while (Percent < Src && CpyIdx < UE_ARRAY_COUNT(FmtBuf))
+					{
+						FmtBuf[CpyIdx] = (ANSICHAR)*Percent;
+						Percent++;
+						CpyIdx++;
+					}
+					FmtBuf[CpyIdx] = 0;
+
+					int RetCnt = snprintf(AnsiNum, sizeof(AnsiNum), FmtBuf, Val);
+					if (RetCnt < 0 || !DestIter.Write(AnsiNum, RetCnt))
+					{
+						return -1;
+					}
+					break;
+				}
+				else if (Src[0] == 'l')
 				{
 					if (Src[1] != 'l' || !CharIsIntegerFormatSpecifier(Src[2]))
 					{
@@ -789,7 +820,7 @@ int32 FGenericWidePlatformString::GetVarArgs( WIDECHAR* Dest, SIZE_T DestSize, c
 				FmtBuf[CpyIdx] = 0;
 
 				int RetCnt = snprintf(AnsiNum, sizeof(AnsiNum), FmtBuf, Val);
-				if (!DestIter.Write(AnsiNum, RetCnt))
+				if (RetCnt < 0 || !DestIter.Write(AnsiNum, RetCnt))
 				{
 					return -1;
 				}
@@ -830,7 +861,7 @@ int32 FGenericWidePlatformString::GetVarArgs( WIDECHAR* Dest, SIZE_T DestSize, c
 					AnsiNum[UE_ARRAY_COUNT(AnsiNum) - 1] = '\0';
 					checkf(0, TEXT("Attempting to read past the size our buffer. Buffer Size: %d Size to read: %d. Current contents: '%s'\n"), UE_ARRAY_COUNT(AnsiNum), RetCnt, UTF8_TO_TCHAR(AnsiNum));
 				}
-				if (!DestIter.Write(AnsiNum, RetCnt))
+				if (RetCnt < 0 || !DestIter.Write(AnsiNum, RetCnt))
 				{
 					return -1;
 				}
