Fixed LLM methods FindTagDisplayName to lookup for platform/project specific tag names.

diff --git a/Engine/Source/Runtime/Core/Private/HAL/LowLevelMemTracker.cpp b/Engine/Source/Runtime/Core/Private/HAL/LowLevelMemTracker.cpp
--- a/Engine/Source/Runtime/Core/Private/HAL/LowLevelMemTracker.cpp
+++ b/Engine/Source/Runtime/Core/Private/HAL/LowLevelMemTracker.cpp
@@ -2000,7 +2000,7 @@ FName FLowLevelMemTracker::FindTagDisplayName(uint64 Tag) const
 	LLMCheck(bInitialisedTracking);
 
 	int32 Index = static_cast<int32>(Tag);
-	if (0 <= Index && Index < LLM_CUSTOM_TAG_START)
+	if (0 <= Index && Index < LLM_TAG_COUNT)
 	{
 		const FTagData* TagData = TagDataEnumMap[Index];
 		if (TagData)
