Leave large allocator page caching enabled when a OOM has happened and all unusued pages have been freed (so new allocated pages gets cached again over the time). Driven with CVar VeryLargePageAllocator.LeavePageCachingEnabledWhenOOMHappened enabled by default

#rnx
#tests repro the OOM with 20 runs of the same replay in demo loop to ensure stability was good

diff --git a/Engine/Source/Runtime/Core/Private/HAL/Allocators/CachedOSVeryLargePageAllocator.cpp b/Engine/Source/Runtime/Core/Private/HAL/Allocators/CachedOSVeryLargePageAllocator.cpp
--- a/Engine/Source/Runtime/Core/Private/HAL/Allocators/CachedOSVeryLargePageAllocator.cpp
+++ b/Engine/Source/Runtime/Core/Private/HAL/Allocators/CachedOSVeryLargePageAllocator.cpp
@@ -27,6 +27,12 @@ static FAutoConsoleVariableRef CVarVeryLargePageAllocatorPreAllocatePools(
 	TEXT("Having pages preallocated and cached during the life of the title help to avoid defragmentation of physical memory.\n")
 	TEXT("Preallocation is disabled when system reaches OOM."));
 
+bool GVeryLargePageAllocatorLeavePageCachingEnabledWhenOOMHappened = true;
+static FAutoConsoleVariableRef CVarVeryLargePageAllocatorLeavePageCachingEnabledWhenOOMHappened(
+	TEXT("VeryLargePageAllocator.LeavePageCachingEnabledWhenOOMHappened"),
+	GVeryLargePageAllocatorLeavePageCachingEnabledWhenOOMHappened,
+	TEXT("Leave page caching enabled when a OOM has happened and all unusued pages have been freed (so new allocated pages gets cached again)"));
+
 static FAutoConsoleVariableRef CVarVeryLargePageAllocatorMaxEmptyBackstoreDefault(
 	TEXT("VeryLargePageAllocator.MaxEmptyBackstoreDefault"),
 	GVeryLargePageAllocatorMaxEmptyBackStoreCount[FMemory::AllocationHints::Default],
@@ -386,7 +392,10 @@ void FCachedOSVeryLargePageAllocator::FreeAll(FCriticalSection* Mutex)
 	}
 
 	// Stop preallocating system since allocator reached a OOM
-	GVeryLargePageAllocatorPreAllocatePools = false;
+	if (!GVeryLargePageAllocatorLeavePageCachingEnabledWhenOOMHappened)
+	{
+		GVeryLargePageAllocatorPreAllocatePools = false;
+	}
 
 	// Free empty cached pages of CachedOSPageAllocator
 	CachedOSPageAllocator.FreeAll(Mutex);
