Fix TGraphTask::CreateTaskWithEvent documentation

#rnx

diff --git a/Engine/Source/Runtime/Core/Public/Async/TaskGraphInterfaces.h b/Engine/Source/Runtime/Core/Public/Async/TaskGraphInterfaces.h
--- a/Engine/Source/Runtime/Core/Public/Async/TaskGraphInterfaces.h
+++ b/Engine/Source/Runtime/Core/Public/Async/TaskGraphInterfaces.h
@@ -1286,6 +1286,7 @@ public:
 
 	/**
 	 *	Factory to create a task and return the helper object to construct the embedded task and set it up for execution.
+	 *	@param SubsequentEventToAssume; subsequent to "assume" from an existing task
 	 *	@param Prerequisites; the list of FGraphEvents that must be completed prior to this task executing.
 	 *	@param CurrentThreadIfKnown; provides the index of the thread we are running on. Can be ENamedThreads::AnyThread if the current thread is unknown.
 	 *	@return a temporary helper class which can be used to complete the process.
@@ -1293,6 +1294,7 @@ public:
 	static FConstructor CreateTaskWithEvent(FGraphEventRef SubsequentEventToAssume, const FGraphEventArray* Prerequisites = NULL, ENamedThreads::Type CurrentThreadIfKnown = ENamedThreads::AnyThread)
 	{
 		int32 NumPrereq = Prerequisites ? Prerequisites->Num() : 0;
+		// If an event is provided, make sure the task will call DispatchSubsequents on it in ExecuteTask
 		check(SubsequentEventToAssume == NULL || TTask::GetSubsequentsMode() == ESubsequentsMode::TrackSubsequents);
 		return FConstructor(new TGraphTask(SubsequentEventToAssume, NumPrereq), Prerequisites, CurrentThreadIfKnown);
 	}
