Implemented TIntPoint::DivideAndRoundDown with TIntPoint Divisor.

diff --git a/Engine/Source/Runtime/Core/Public/Math/IntPoint.h b/Engine/Source/Runtime/Core/Public/Math/IntPoint.h
--- a/Engine/Source/Runtime/Core/Public/Math/IntPoint.h
+++ b/Engine/Source/Runtime/Core/Public/Math/IntPoint.h
@@ -436,6 +436,11 @@ PRAGMA_ENABLE_DEPRECATION_WARNINGS
 		return TIntPoint(FMath::DivideAndRoundDown(lhs.X, Divisor), FMath::DivideAndRoundDown(lhs.Y, Divisor));
 	}
 
+	static TIntPoint DivideAndRoundDown(TIntPoint lhs, TIntPoint Divisor)
+	{
+		return TIntPoint(FMath::DivideAndRoundDown(lhs.X, Divisor.X), FMath::DivideAndRoundDown(lhs.Y, Divisor.Y));
+	}
+
 	/**
 	 * Get number of components point has.
 	 *
