FString/CString: Change FString::Contains, FString::Find, and TCString::Stristr to return success when the Find argument is EmptyString, when called with ESearchCase::IgnoreCase. This change makes the IgnoreCase behavior match the CaseSensitive behavior, and match the behavior of standard library strstr. But this will break callsites that were expecting Contains("") to return false, and those callsites will need to be changed as they are discovered.
FString: Add FStringView override for Contains.
#rn Minor, Core

diff --git a/Engine/Source/Runtime/AudioMixer/Private/AudioMixerSource.cpp b/Engine/Source/Runtime/AudioMixer/Private/AudioMixerSource.cpp
--- a/Engine/Source/Runtime/AudioMixer/Private/AudioMixerSource.cpp
+++ b/Engine/Source/Runtime/AudioMixer/Private/AudioMixerSource.cpp
@@ -784,7 +784,7 @@ namespace Audio
 			bool bIsDebug = false;
 			FString WaveInstanceName = WaveInstance->GetName(); //-V595
 			FString TestName = GEngine->GetAudioDeviceManager()->GetDebugger().GetAudioMixerDebugSoundName();
-			if (WaveInstanceName.Contains(TestName))
+			if (!TestName.IsEmpty() && WaveInstanceName.Contains(TestName))
 			{
 				bDebugMode = true;
 				InitParams.bIsDebugMode = bDebugMode;
diff --git a/Engine/Source/Runtime/Core/Private/Containers/String.cpp b/Engine/Source/Runtime/Core/Private/Containers/String.cpp
--- a/Engine/Source/Runtime/Core/Private/Containers/String.cpp
+++ b/Engine/Source/Runtime/Core/Private/Containers/String.cpp
@@ -409,11 +409,6 @@ int32 FString::Find(const TCHAR* SubStr, int32 SubStrLen, ESearchCase::Type Sear
 			Start += FMath::Clamp(StartPosition, 0, RemainingLength - 1);
 			RemainingLength = UE_PTRDIFF_TO_INT32(End - Start);
 		}
-		if (SubStrLen == 0 && SearchCase == ESearchCase::IgnoreCase)
-		{
-			// Legacy behavior: when SubStr is empty string and SearchCase is IgnoreCase, return INDEX_NONE
-			return INDEX_NONE;
-		}
 		const TCHAR* Tmp = SearchCase == ESearchCase::IgnoreCase
 			? FCString::Strnistr(Start, RemainingLength, SubStr, SubStrLen)
 			: FCString::Strnstr(Start, RemainingLength, SubStr, SubStrLen);
diff --git a/Engine/Source/Runtime/Core/Private/Containers/StringTest.cpp b/Engine/Source/Runtime/Core/Private/Containers/StringTest.cpp
--- a/Engine/Source/Runtime/Core/Private/Containers/StringTest.cpp
+++ b/Engine/Source/Runtime/Core/Private/Containers/StringTest.cpp
@@ -610,14 +610,16 @@ bool FStringPathConcatCompoundOperatorTest::RunTest(const FString& Parameters)
 	return true;
 }
 
-IMPLEMENT_SIMPLE_AUTOMATION_TEST(FStringFindTest, "System.Core.String.Find", EAutomationTestFlags::ApplicationContextMask | EAutomationTestFlags::EngineFilter)
-bool FStringFindTest::RunTest(const FString& Parameters)
+IMPLEMENT_SIMPLE_AUTOMATION_TEST(FStringFindAndContainsTest, "System.Core.String.FindAndContains", EAutomationTestFlags::ApplicationContextMask | EAutomationTestFlags::EngineFilter)
+bool FStringFindAndContainsTest::RunTest(const FString& Parameters)
 {
 	auto RunTest = [this](FStringView Search, FStringView Find, ESearchCase::Type SearchCase,
 		ESearchDir::Type SearchDir, int32 StartPosition, int32 Expected)
 	{
 		FString SearchStr(Search);
 		FString FindStr(Find);
+		bool bRunContainsTest = (SearchDir == ESearchDir::FromStart && StartPosition == 0) ||
+			(SearchDir == ESearchDir::FromEnd && StartPosition >= SearchStr.Len());
 		// TCHAR*, int32
 		{
 			int32 Actual = SearchStr.Find(Find.GetData(), Find.Len(), SearchCase, SearchDir, StartPosition);
@@ -626,6 +628,16 @@ bool FStringFindTest::RunTest(const FString& Parameters)
 				AddError(FString::Printf(TEXT("FString(\"%s\").Find(\"%s\", %d, %d, %d, %d) returned Actual %d not equal to Expected %d"),
 					*SearchStr, *FindStr, FindStr.Len(), (int32)SearchCase, (int32)SearchDir, StartPosition, Actual, Expected));
 			}
+			if (bRunContainsTest)
+			{
+				bool ContainsExpected = Expected != INDEX_NONE;
+				bool ContainsActual = SearchStr.Contains(Find.GetData(), Find.Len(), SearchCase, SearchDir);
+				if (ContainsActual != ContainsExpected)
+				{
+					AddError(FString::Printf(TEXT("FString(\"%s\").Contains(\"%s\", %d, %d, %d) returned Actual %s not equal to Expected %s"),
+						*SearchStr, *FindStr, FindStr.Len(), (int32)SearchCase, (int32)SearchDir,  *LexToString(ContainsActual), *LexToString(ContainsExpected)));
+				}
+			}
 		}
 		// FStringView
 		{
@@ -633,7 +645,17 @@ bool FStringFindTest::RunTest(const FString& Parameters)
 			if (Actual != Expected)
 			{
 				AddError(FString::Printf(TEXT("FString(\"%s\").Find(FStringView(\"%s\", %d), %d, %d, %d) returned Actual %d not equal to Expected %d"),
-					*SearchStr, *FindStr, FindStr.Len(), (int32)SearchCase, (int32)SearchDir, StartPosition));
+					*SearchStr, *FindStr, FindStr.Len(), (int32)SearchCase, (int32)SearchDir, StartPosition, Actual, Expected));
+			}
+			if (bRunContainsTest)
+			{
+				bool ContainsExpected = Expected != INDEX_NONE;
+				bool ContainsActual = ContainsExpected;
+				if (ContainsActual != ContainsExpected)
+				{
+					AddError(FString::Printf(TEXT("FString(\"%s\").Contains(FStringView(\"%s\", %d), %d, %d) returned Actual %s not equal to Expected %s"),
+						*SearchStr, *FindStr, FindStr.Len(), (int32)SearchCase, (int32)SearchDir, *LexToString(ContainsActual), *LexToString(ContainsExpected)));
+				}
 			}
 		}
 
@@ -643,7 +665,17 @@ bool FStringFindTest::RunTest(const FString& Parameters)
 			if (Actual != Expected)
 			{
 				AddError(FString::Printf(TEXT("FString(\"%s\").Find(TEXT(\"%s\"), %d, %d, %d) returned Actual %d not equal to Expected %d"),
-					*SearchStr, *FindStr, (int32)SearchCase, (int32)SearchDir, StartPosition));
+					*SearchStr, *FindStr, (int32)SearchCase, (int32)SearchDir, StartPosition, Actual, Expected));
+			}
+			if (bRunContainsTest)
+			{
+				bool ContainsExpected = Expected != INDEX_NONE;
+				bool ContainsActual = SearchStr.Contains(*FindStr, SearchCase, SearchDir);
+				if (ContainsActual != ContainsExpected)
+				{
+					AddError(FString::Printf(TEXT("FString(\"%s\").Contains(TEXT(\"%s\"), %d, %d) returned Actual %s not equal to Expected %s"),
+						*SearchStr, *FindStr, (int32)SearchCase, (int32)SearchDir, *LexToString(ContainsActual), *LexToString(ContainsExpected)));
+				}
 			}
 		}
 
@@ -653,7 +685,17 @@ bool FStringFindTest::RunTest(const FString& Parameters)
 			if (Actual != Expected)
 			{
 				AddError(FString::Printf(TEXT("FString(\"%s\").Find(FString(%d, \"%s\"), %d, %d, %d) returned Actual %d not equal to Expected %d"),
-					*SearchStr, FindStr.Len(), *FindStr, (int32)SearchCase, (int32)SearchDir, StartPosition));
+					*SearchStr, FindStr.Len(), *FindStr, (int32)SearchCase, (int32)SearchDir, StartPosition, Actual, Expected));
+			}
+			if (bRunContainsTest)
+			{
+				bool ContainsExpected = Expected != INDEX_NONE;
+				bool ContainsActual = SearchStr.Contains(FindStr, SearchCase, SearchDir);
+				if (ContainsActual != ContainsExpected)
+				{
+					AddError(FString::Printf(TEXT("FString(\"%s\").Contains(FString(%d, \"%s\"), %d, %d) returned Actual %s not equal to Expected %s"),
+						*SearchStr, FindStr.Len(), *FindStr, (int32)SearchCase, (int32)SearchDir, *LexToString(ContainsActual), *LexToString(ContainsExpected)));
+				}
 			}
 		}
 	};
@@ -710,9 +752,8 @@ bool FStringFindTest::RunTest(const FString& Parameters)
 	RunTest(ABACADAB, EmptyString, ESearchCase::CaseSensitive, ESearchDir::FromEnd, 2, 1);
 	RunTest(ABACADAB, EmptyString, ESearchCase::CaseSensitive, ESearchDir::FromEnd, 0, INDEX_NONE);
 
-	// Find with an empty search string has different behavior depending on SearchCase due to legacy implementation
-	RunTest(ABACADAB, EmptyString, ESearchCase::IgnoreCase, ESearchDir::FromStart, 0, INDEX_NONE);
-	RunTest(ABACADAB, EmptyString, ESearchCase::IgnoreCase, ESearchDir::FromStart, 4, INDEX_NONE);
+	RunTest(ABACADAB, EmptyString, ESearchCase::IgnoreCase, ESearchDir::FromStart, 0, 0);
+	RunTest(ABACADAB, EmptyString, ESearchCase::IgnoreCase, ESearchDir::FromStart, 4, 4);
 	RunTest(ABACADAB, EmptyString, ESearchCase::IgnoreCase, ESearchDir::FromEnd, 8, 7);
 	RunTest(ABACADAB, EmptyString, ESearchCase::IgnoreCase, ESearchDir::FromEnd, 2, 1);
 	RunTest(ABACADAB, EmptyString, ESearchCase::IgnoreCase, ESearchDir::FromEnd, 0, INDEX_NONE);
diff --git a/Engine/Source/Runtime/Core/Private/Misc/CStringTest.cpp b/Engine/Source/Runtime/Core/Private/Misc/CStringTest.cpp
--- a/Engine/Source/Runtime/Core/Private/Misc/CStringTest.cpp
+++ b/Engine/Source/Runtime/Core/Private/Misc/CStringTest.cpp
@@ -97,6 +97,51 @@ bool FCStringGetVarArgsTest::RunTest(const FString& Parameters)
 	return true;
 }
 
+IMPLEMENT_SIMPLE_AUTOMATION_TEST(FCStringStrstrTest, "System.Core.Misc.CString.Strstr", EAutomationTestFlags::ApplicationContextMask | EAutomationTestFlags::EngineFilter)
+bool FCStringStrstrTest::RunTest(const FString& Parameters)
+{
+	auto RunTest = [this](const TCHAR* Search, const TCHAR* Find, int32 ExpectedSensitiveIndex, int32 ExpectedInsensitiveIndex)
+	{
+		const TCHAR* ExpectedSensitive = ExpectedSensitiveIndex == INDEX_NONE ? nullptr :
+			(Search + ExpectedSensitiveIndex);
+		const TCHAR* ExpectedInsensitive = ExpectedInsensitiveIndex == INDEX_NONE ? nullptr :
+			(Search + ExpectedInsensitiveIndex);
+		if (FCString::Strstr(Search, Find) != ExpectedSensitive)
+		{
+			AddError(FString::Printf(TEXT("Strstr(\"%s\", \"%s\") did not equal index \"%d\"."),
+				Search, Find, ExpectedSensitiveIndex));
+		}
+		if (FCString::Stristr(Search, Find) != ExpectedInsensitive)
+		{
+			AddError(FString::Printf(TEXT("Stristr(\"%s\", \"%s\") did not equal index \"%d\"."),
+				Search, Find, ExpectedInsensitiveIndex));
+		}
+	};
+	const TCHAR* ABACADAB = TEXT("ABACADAB");
+
+	RunTest(ABACADAB, TEXT("A"), 0, 0);
+	RunTest(ABACADAB, TEXT("a"), INDEX_NONE, 0);
+	RunTest(ABACADAB, TEXT("BAC"), 1, 1);
+	RunTest(ABACADAB, TEXT("BaC"), INDEX_NONE, 1);
+	RunTest(ABACADAB, TEXT("BAC"), 1, 1);
+	RunTest(ABACADAB, TEXT("BaC"), INDEX_NONE, 1);
+	RunTest(ABACADAB, TEXT("DAB"), 5, 5);
+	RunTest(ABACADAB, TEXT("dab"), INDEX_NONE, 5);
+	RunTest(ABACADAB, ABACADAB, 0, 0);
+	RunTest(ABACADAB, TEXT("abacadab"), INDEX_NONE, 0);
+	RunTest(ABACADAB, TEXT("F"), INDEX_NONE, INDEX_NONE);
+	RunTest(ABACADAB, TEXT("DABZ"), INDEX_NONE, INDEX_NONE);
+	RunTest(ABACADAB, TEXT("ABACADABA"), INDEX_NONE, INDEX_NONE);
+	RunTest(ABACADAB, TEXT("NoMatchLongerString"), INDEX_NONE, INDEX_NONE);
+	RunTest(TEXT(""), TEXT("FindText"), INDEX_NONE, INDEX_NONE);
+	RunTest(TEXT(""), TEXT(""), 0, 0);
+	RunTest(ABACADAB, TEXT(""), 0, 0);
+
+	// Passing in nullpt r is not allowed by StrStr so we do not test it
+
+	return true;
+}
+
 IMPLEMENT_SIMPLE_AUTOMATION_TEST(FCStringStrnstrTest, "System.Core.Misc.CString.Strnstr", EAutomationTestFlags::ApplicationContextMask | EAutomationTestFlags::EngineFilter)
 bool FCStringStrnstrTest::RunTest(const FString& Parameters)
 {
diff --git a/Engine/Source/Runtime/Core/Public/Containers/UnrealString.h b/Engine/Source/Runtime/Core/Public/Containers/UnrealString.h
--- a/Engine/Source/Runtime/Core/Public/Containers/UnrealString.h
+++ b/Engine/Source/Runtime/Core/Public/Containers/UnrealString.h
@@ -1066,8 +1066,8 @@ public:
 	}
 
 	/**
-	 * Searches the string for a substring, and returns index into this string
-	 * of the first found instance. Can search from beginning or end, and ignore case or not.
+	 * Searches the string for a substring, and returns index into this string of the first found instance. Can search
+	 * from beginning or end, and ignore case or not. If substring is empty, returns clamped StartPosition.
 	 *
 	 * @param SubStr			The string to search for
 	 * @param StartPosition		The start character position to search from
@@ -1090,8 +1090,8 @@ public:
 	}
 
 	/**
-	 * Searches the string for a substring, and returns index into this string
-	 * of the first found instance. Can search from beginning or end, and ignore case or not.
+	 * Searches the string for a substring, and returns index into this string of the first found instance. Can search
+	 * from beginning or end, and ignore case or not. If substring is empty, returns clamped StartPosition.
 	 *
 	 * @param SubStr			The string array of TCHAR to search for
 	 * @param StartPosition		The start character position to search from.  See note below.
@@ -1105,8 +1105,8 @@ public:
 	}
 
 	/**
-	 * Searches the string for a substring, and returns index into this string
-	 * of the first found instance. Can search from beginning or end, and ignore case or not.
+	 * Searches the string for a substring, and returns index into this string of the first found instance. Can search
+	 * from beginning or end, and ignore case or not. If substring is empty, returns clamped StartPosition.
 	 *
 	 * @param SubStr			The string array of TCHAR to search for
 	 * @param StartPosition		The start character position to search from.  See note below.
@@ -1127,8 +1127,8 @@ public:
 	}
 
 	/**
-	 * Searches the string for a substring, and returns index into this string
-	 * of the first found instance. Can search from beginning or end, and ignore case or not.
+	 * Searches the string for a substring, and returns index into this string of the first found instance. Can search
+	 * from beginning or end, and ignore case or not. If substring is empty, returns clamped StartPosition.
 	 *
 	 * @param SubStr			The string array of TCHAR to search for
 	 * @param SubStrLen			The length of the SubStr array
@@ -1147,16 +1147,32 @@ public:
 		ESearchDir::Type SearchDir = ESearchDir::FromStart, int32 StartPosition = INDEX_NONE) const;
 
 
+	/**
+	 * Returns whether this string contains the specified substring.
+	 *
+	 * @param SubStr			Text to search for
+	 * @param SearchCase		Indicates whether the search is case sensitive or not ( defaults to ESearchCase::IgnoreCase )
+	 * @param SearchDir			Indicates whether the search starts at the beginning or at the end ( defaults to ESearchDir::FromStart )
+	 * @return					Returns whether the string contains the substring. If the substring is empty, returns true.
+	 **/
+	template <typename TCharRangeType, std::enable_if_t<TIsCharRangeNotCArrayNotFString<TCharRangeType>::Value>* = nullptr>
+	UE_NODISCARD FORCEINLINE bool Contains(TCharRangeType&& SubStr, ESearchCase::Type SearchCase = ESearchCase::IgnoreCase,
+		ESearchDir::Type SearchDir = ESearchDir::FromStart) const
+	{
+		static_assert(std::is_same_v<typename TElementType<TCharRangeType>::Type, TCHAR>, "Expected a range of TCHAR");
+		return Find(Forward<TCharRangeType>(SubStr), SearchCase, SearchDir) != INDEX_NONE;
+	}
+
 	/** 
 	 * Returns whether this string contains the specified substring.
 	 *
-	 * @param SubStr			Find to search for
+	 * @param SubStr			Text to search for
 	 * @param SearchCase		Indicates whether the search is case sensitive or not ( defaults to ESearchCase::IgnoreCase )
 	 * @param SearchDir			Indicates whether the search starts at the beginning or at the end ( defaults to ESearchDir::FromStart )
-	 * @return					Returns whether the string contains the substring
+	 * @return					Returns whether the string contains the substring. If the substring is empty, returns true.
 	 **/
-	UE_NODISCARD FORCEINLINE bool Contains(const TCHAR* SubStr, ESearchCase::Type SearchCase = ESearchCase::IgnoreCase, 
-							  ESearchDir::Type SearchDir = ESearchDir::FromStart ) const
+	UE_NODISCARD FORCEINLINE bool Contains(const TCHAR* SubStr, ESearchCase::Type SearchCase = ESearchCase::IgnoreCase,
+		ESearchDir::Type SearchDir = ESearchDir::FromStart) const
 	{
 		return Find(SubStr, SearchCase, SearchDir) != INDEX_NONE;
 	}
@@ -1164,10 +1180,10 @@ public:
 	/** 
 	 * Returns whether this string contains the specified substring.
 	 *
-	 * @param SubStr			Find to search for
+	 * @param SubStr			Text to search for
 	 * @param SearchCase		Indicates whether the search is case sensitive or not ( defaults to ESearchCase::IgnoreCase )
 	 * @param SearchDir			Indicates whether the search starts at the beginning or at the end ( defaults to ESearchDir::FromStart )
-	 * @return					Returns whether the string contains the substring
+	 * @return					Returns whether the string contains the substring. If the substring is empty, returns true.
 	 **/
 	UE_NODISCARD FORCEINLINE bool Contains(const FString& SubStr, ESearchCase::Type SearchCase = ESearchCase::IgnoreCase, 
 							  ESearchDir::Type SearchDir = ESearchDir::FromStart ) const
@@ -1175,6 +1191,21 @@ public:
 		return Find(*SubStr, SearchCase, SearchDir) != INDEX_NONE;
 	}
 
+	/**
+	 * Returns whether this string contains the specified substring.
+	 *
+	 * @param SubStr			Text to search for
+	 * @param SubStrLen			Length of the Text
+	 * @param SearchCase		Indicates whether the search is case sensitive or not ( defaults to ESearchCase::IgnoreCase )
+	 * @param SearchDir			Indicates whether the search starts at the beginning or at the end ( defaults to ESearchDir::FromStart )
+	 * @return					Returns whether the string contains the substring. If the substring is empty, returns true.
+	 **/
+	UE_NODISCARD FORCEINLINE bool Contains(const TCHAR* SubStr, int32 SubStrLen,
+		ESearchCase::Type SearchCase = ESearchCase::IgnoreCase, ESearchDir::Type SearchDir = ESearchDir::FromStart) const
+	{
+		return Find(SubStr, SubStrLen, SearchCase, SearchDir) != INDEX_NONE;
+	}
+
 	/**
 	 * Searches the string for a character
 	 *
diff --git a/Engine/Source/Runtime/Core/Public/Misc/CString.h b/Engine/Source/Runtime/Core/Public/Misc/CString.h
--- a/Engine/Source/Runtime/Core/Public/Misc/CString.h
+++ b/Engine/Source/Runtime/Core/Public/Misc/CString.h
@@ -259,7 +259,7 @@ struct TCString
 	 * Finds string in string, case insensitive 
 	 * @param Str The string to look through
 	 * @param Find The string to find inside Str
-	 * @return Position in Str if Find was found, otherwise, nullptr
+	 * @return Position in Str if Find was found, otherwise, nullptr. If Find is non-null but empty, returns Str.
 	 */
 	static const CharType* Stristr(const CharType* Str, const CharType* Find);
 
@@ -267,7 +267,7 @@ struct TCString
 	 * Finds string in string, case insensitive (non-const version)
 	 * @param Str The string to look through
 	 * @param Find The string to find inside Str
-	 * @return Position in Str if Find was found, otherwise, nullptr
+	 * @return Position in Str if Find was found, otherwise, nullptr. If Find is non-null but empty, returns Str.
 	 */
 	static CharType* Stristr(CharType* Str, const CharType* Find) { return (CharType*)Stristr((const CharType*)Str, Find); }
 
@@ -277,7 +277,7 @@ struct TCString
 	 * @param InStrLen The length of the Str array
 	 * @param Find The character array to find inside Str
 	 * @param FindLen The length of the Find array
-	 * @return Position in Str if Find was found, otherwise, nullptr
+	 * @return Position in Str if Find was found, otherwise, nullptr. If FindLen is 0, returns Str.
 	 */
 	static const CharType* Strnistr(const CharType* Str, int32 InStrLen, const CharType* Find, int32 FindLen);
 
@@ -287,7 +287,7 @@ struct TCString
 	 * @param InStrLen The length of the Str array
 	 * @param Find The character array to find inside Str
 	 * @param FindLen The length of the Find array
-	 * @return Position in Str if Find was found, otherwise, nullptr
+	 * @return Position in Str if Find was found, otherwise, nullptr. If FindLen is 0, returns Str.
 	 */
 	static CharType* Strnistr(CharType* Str, int32 InStrLen, const CharType* Find, int32 FindLen)
 	{
@@ -300,7 +300,7 @@ struct TCString
 	 * @param InStrLen The length of the Str array
 	 * @param Find The character array to find inside Str
 	 * @param FindLen The length of the Find array
-	 * @return Position in Str if Find was found, otherwise, nullptr
+	 * @return Position in Str if Find was found, otherwise, nullptr. If FindLen is 0, returns Str.
 	 */
 	static const CharType* Strnstr(const CharType* Str, int32 InStrLen, const CharType* Find, int32 FindLen);
 
@@ -310,7 +310,7 @@ struct TCString
 	 * @param InStrLen The length of the Str array
 	 * @param Find The character array to find inside Str
 	 * @param FindLen The length of the Find array
-	 * @return Position in Str if Find was found, otherwise, nullptr
+	 * @return Position in Str if Find was found, otherwise, nullptr. If FindLen is 0, returns Str.
 	 */
 	static CharType* Strnstr(CharType* Str, int32 InStrLen, const CharType* Find, int32 FindLen)
 	{
@@ -653,7 +653,7 @@ const typename TCString<T>::CharType* TCString<T>::StrfindDelim(const CharType*
  * Finds string in string, case insensitive 
  * @param Str The string to look through
  * @param Find The string to find inside Str
- * @return Position in Str if Find was found, otherwise, NULL
+ * @return Position in Str if Find was found, otherwise, NULL. If Find is non-null but empty, returns Str.
  */
 template <typename T>
 const typename TCString<T>::CharType* TCString<T>::Stristr(const CharType* Str, const CharType* Find)
@@ -669,6 +669,11 @@ const typename TCString<T>::CharType* TCString<T>::Stristr(const CharType* Str,
 
 	// get upper-case first letter of the find string (to reduce the number of full strnicmps)
 	CharType FindInitial = TChar<CharType>::ToUpper(*Find);
+	if (!FindInitial)
+	{
+		// When searching for the empty string, always return index of the first element of Str even if Str is empty.
+		return Str;
+	}
 	// get length of find string, and increment past first letter
 	int32   Length = Strlen(Find++) - 1;
 	// get the first letter of the search string, and increment past it
