CompactBinary: Added serialization to WIDECHAR JSON

diff --git a/Engine/Source/Runtime/Core/Private/Serialization/CompactBinarySerialization.cpp b/Engine/Source/Runtime/Core/Private/Serialization/CompactBinarySerialization.cpp
--- a/Engine/Source/Runtime/Core/Private/Serialization/CompactBinarySerialization.cpp
+++ b/Engine/Source/Runtime/Core/Private/Serialization/CompactBinarySerialization.cpp
@@ -4,7 +4,6 @@
 
 #include "Containers/Array.h"
 #include "Containers/ContainerAllocationPolicies.h"
-#include "Containers/StringConv.h"
 #include "Containers/StringView.h"
 #include "Containers/UnrealString.h"
 #include "CoreGlobals.h"
@@ -328,14 +327,76 @@ bool LoadFromCompactBinary(FCbFieldView Field, FGuid& OutValue, const FGuid& Def
 
 ///////////////////////////////////////////////////////////////////////////////////////////////////////////////////////
 
-class FCbJsonWriter
+template <typename CharType>
+struct TCbJsonWriterConstants;
+
+template <>
+struct TCbJsonWriterConstants<ANSICHAR>
+{
+	static inline const FAnsiStringView LineTerminator = LINE_TERMINATOR_ANSI;
+	static inline const FAnsiStringView Space = ANSITEXTVIEW(" ");
+	static inline const FAnsiStringView Tab = ANSITEXTVIEW("\t");
+
+	static inline const FAnsiStringView Null = ANSITEXTVIEW("null");
+	static inline const FAnsiStringView True = ANSITEXTVIEW("true");
+	static inline const FAnsiStringView False = ANSITEXTVIEW("false");
+
+	static inline const FAnsiStringView EscapeBackslash = ANSITEXTVIEW("\\\\");
+	static inline const FAnsiStringView EscapeDoubleQuote = ANSITEXTVIEW("\\\"");
+	static inline const FAnsiStringView EscapeBell = ANSITEXTVIEW("\\b");
+	static inline const FAnsiStringView EscapeFormFeed = ANSITEXTVIEW("\\f");
+	static inline const FAnsiStringView EscapeLineFeed = ANSITEXTVIEW("\\n");
+	static inline const FAnsiStringView EscapeCarriageReturn = ANSITEXTVIEW("\\r");
+	static inline const FAnsiStringView EscapeTab = ANSITEXTVIEW("\\t");
+
+	static inline const FAnsiStringView FieldId = ANSITEXTVIEW("\"Id\":");
+	static inline const FAnsiStringView FieldName = ANSITEXTVIEW("\"Name\":");
+	static inline const FAnsiStringView FieldData = ANSITEXTVIEW("\"Data\":");
+
+	static inline decltype(auto) FormatEscapeUnicode4 = "\\u%04x";
+	static inline decltype(auto) FormatFloat32 = "%.9g";
+	static inline decltype(auto) FormatFloat64 = "%.17g";
+};
+
+template <>
+struct TCbJsonWriterConstants<WIDECHAR>
 {
+	static inline const FWideStringView LineTerminator = LINE_TERMINATOR;
+	static inline const FWideStringView Space = WIDETEXTVIEW(" ");
+	static inline const FWideStringView Tab = WIDETEXTVIEW("\t");
+
+	static inline const FWideStringView Null = WIDETEXTVIEW("null");
+	static inline const FWideStringView True = WIDETEXTVIEW("true");
+	static inline const FWideStringView False = WIDETEXTVIEW("false");
+
+	static inline const FWideStringView EscapeBackslash = WIDETEXTVIEW("\\\\");
+	static inline const FWideStringView EscapeDoubleQuote = WIDETEXTVIEW("\\\"");
+	static inline const FWideStringView EscapeBell = WIDETEXTVIEW("\\b");
+	static inline const FWideStringView EscapeFormFeed = WIDETEXTVIEW("\\f");
+	static inline const FWideStringView EscapeLineFeed = WIDETEXTVIEW("\\n");
+	static inline const FWideStringView EscapeCarriageReturn = WIDETEXTVIEW("\\r");
+	static inline const FWideStringView EscapeTab = WIDETEXTVIEW("\\t");
+
+	static inline const FWideStringView FieldId = WIDETEXTVIEW("\"Id\":");
+	static inline const FWideStringView FieldName = WIDETEXTVIEW("\"Name\":");
+	static inline const FWideStringView FieldData = WIDETEXTVIEW("\"Data\":");
+
+	static inline decltype(auto) FormatEscapeUnicode4 = WIDETEXT("\\u%04x");
+	static inline decltype(auto) FormatFloat32 = WIDETEXT("%.9g");
+	static inline decltype(auto) FormatFloat64 = WIDETEXT("%.17g");
+};
+
+template <typename CharType, typename ConstantCharType = CharType>
+class TCbJsonWriter
+{
+	using FConstants = TCbJsonWriterConstants<ConstantCharType>;
+
 public:
-	explicit FCbJsonWriter(
-		FUtf8StringBuilderBase& InBuilder,
-		FUtf8StringView InLineTerminator = LINE_TERMINATOR_ANSI,
-		FUtf8StringView InIndent = "\t",
-		FUtf8StringView InSpace = " ")
+	explicit TCbJsonWriter(
+		TStringBuilderBase<CharType>& InBuilder,
+		TStringView<CharType> InLineTerminator = FConstants::LineTerminator,
+		TStringView<CharType> InIndent = FConstants::Tab,
+		TStringView<CharType> InSpace = FConstants::Space)
 		: Builder(InBuilder)
 		, Indent(InIndent)
 		, Space(InSpace)
@@ -351,13 +412,13 @@ public:
 		if (FUtf8StringView Name = Field.GetName(); !Name.IsEmpty())
 		{
 			AppendQuotedString(Name);
-			Builder << ANSITEXTVIEW(":") << Space;
+			Builder << ':' << Space;
 		}
 
 		switch (FCbValue Accessor = Field.GetValue(); Accessor.GetType())
 		{
 		case ECbFieldType::Null:
-			Builder << ANSITEXTVIEW("null");
+			Builder << FConstants::Null;
 			break;
 		case ECbFieldType::Object:
 		case ECbFieldType::UniformObject:
@@ -404,16 +465,16 @@ public:
 			Builder << Accessor.AsIntegerNegative();
 			break;
 		case ECbFieldType::Float32:
-			Builder.Appendf(UTF8TEXT("%.9g"), Accessor.AsFloat32());
+			Builder.Appendf(FConstants::FormatFloat32, Accessor.AsFloat32());
 			break;
 		case ECbFieldType::Float64:
-			Builder.Appendf(UTF8TEXT("%.17g"), Accessor.AsFloat64());
+			Builder.Appendf(FConstants::FormatFloat64, Accessor.AsFloat64());
 			break;
 		case ECbFieldType::BoolFalse:
-			Builder << ANSITEXTVIEW("false");
+			Builder << FConstants::False;
 			break;
 		case ECbFieldType::BoolTrue:
-			Builder << ANSITEXTVIEW("true");
+			Builder << FConstants::True;
 			break;
 		case ECbFieldType::ObjectAttachment:
 		case ECbFieldType::BinaryAttachment:
@@ -426,18 +487,18 @@ public:
 			Builder << '"' << Accessor.AsUuid() << '"';
 			break;
 		case ECbFieldType::DateTime:
-			Builder << '"' << FTCHARToUTF8(FDateTime(Accessor.AsDateTimeTicks()).ToIso8601()) << '"';
+			Builder << '"' << FDateTime(Accessor.AsDateTimeTicks()).ToIso8601() << '"';
 			break;
 		case ECbFieldType::TimeSpan:
 		{
 			const FTimespan Span(Accessor.AsTimeSpanTicks());
 			if (Span.GetDays() == 0)
 			{
-				Builder << '"' << FTCHARToUTF8(Span.ToString(TEXT("%h:%m:%s.%n"))) << '"';
+				Builder << '"' << Span.ToString(TEXT("%h:%m:%s.%n")) << '"';
 			}
 			else
 			{
-				Builder << '"' << FTCHARToUTF8(Span.ToString(TEXT("%d.%h:%m:%s.%n"))) << '"';
+				Builder << '"' << Span.ToString(TEXT("%d.%h:%m:%s.%n")) << '"';
 			}
 			break;
 		}
@@ -447,21 +508,21 @@ public:
 		case ECbFieldType::CustomById:
 		{
 			FCbCustomById Custom = Accessor.AsCustomById();
-			Builder << "{" << Space << "\"Id\":" << Space;
+			Builder << '{' << Space << FConstants::FieldId << Space;
 			Builder << Custom.Id;
-			Builder << "," << Space << "\"Data\":" << Space;
+			Builder << ',' << Space << FConstants::FieldData << Space;
 			AppendBase64String(Custom.Data);
-			Builder << Space << "}";
+			Builder << Space << '}';
 			break;
 		}
 		case ECbFieldType::CustomByName:
 		{
 			FCbCustomByName Custom = Accessor.AsCustomByName();
-			Builder << "{" << Space << "\"Name\":" << Space;
+			Builder << '{' << Space << FConstants::FieldName << Space;
 			AppendQuotedString(Custom.Name);
-			Builder << "," << Space << "\"Data\":" << Space;
+			Builder << ',' << Space << FConstants::FieldData << Space;
 			AppendBase64String(Custom.Data);
-			Builder << Space << "}";
+			Builder << Space << '}';
 			break;
 		}
 		default:
@@ -504,19 +565,19 @@ private:
 			Builder << Verbatim;
 			Value.RightChopInline(Verbatim.Len());
 			FUtf8StringView Escape = FAsciiSet::FindPrefixWith(Value, EscapeSet);
-			for (UTF8CHAR Char : Escape)
+			for (CharType Char : Escape)
 			{
 				switch (Char)
 				{
-				case '\\': Builder << ANSITEXTVIEW("\\\\"); break;
-				case '\"': Builder << ANSITEXTVIEW("\\\""); break;
-				case '\b': Builder << ANSITEXTVIEW("\\b"); break;
-				case '\f': Builder << ANSITEXTVIEW("\\f"); break;
-				case '\n': Builder << ANSITEXTVIEW("\\n"); break;
-				case '\r': Builder << ANSITEXTVIEW("\\r"); break;
-				case '\t': Builder << ANSITEXTVIEW("\\t"); break;
+				case '\\': Builder << FConstants::EscapeBackslash; break;
+				case '\"': Builder << FConstants::EscapeDoubleQuote; break;
+				case '\b': Builder << FConstants::EscapeBell; break;
+				case '\f': Builder << FConstants::EscapeFormFeed; break;
+				case '\n': Builder << FConstants::EscapeLineFeed; break;
+				case '\r': Builder << FConstants::EscapeCarriageReturn; break;
+				case '\t': Builder << FConstants::EscapeTab; break;
 				default:
-					Builder.Appendf(UTF8TEXT("\\u%04x"), uint32(Char));
+					Builder.Appendf(FConstants::FormatEscapeUnicode4, uint32(Char));
 					break;
 				}
 			}
@@ -533,28 +594,40 @@ private:
 		const uint32 EncodedSize = FBase64::GetEncodedDataSize(uint32(Value.GetSize()));
 		const int32 EncodedIndex = Builder.AddUninitialized(int32(EncodedSize));
 		FBase64::Encode(static_cast<const uint8*>(Value.GetData()), uint32(Value.GetSize()),
-			reinterpret_cast<ANSICHAR*>(Builder.GetData() + EncodedIndex));
+			reinterpret_cast<ConstantCharType*>(Builder.GetData() + EncodedIndex));
 		Builder << '"';
 	}
 
 private:
-	FUtf8StringBuilderBase& Builder;
-	TUtf8StringBuilder<32> NewLineAndIndent;
-	FUtf8StringView Indent;
-	FUtf8StringView Space;
+	TStringBuilderBase<CharType>& Builder;
+	TStringBuilderWithBuffer<CharType, 32> NewLineAndIndent;
+	TStringView<CharType> Indent;
+	TStringView<CharType> Space;
 	bool bNeedsComma{false};
 	bool bNeedsNewLine{false};
 };
 
 void CompactBinaryToJson(const FCbObjectView& Object, FUtf8StringBuilderBase& Builder)
 {
-	FCbJsonWriter Writer(Builder);
+	TCbJsonWriter<UTF8CHAR, ANSICHAR> Writer(Builder);
+	Writer.WriteField(Object.AsFieldView());
+}
+
+void CompactBinaryToJson(const FCbObjectView& Object, FWideStringBuilderBase& Builder)
+{
+	TCbJsonWriter<WIDECHAR> Writer(Builder);
 	Writer.WriteField(Object.AsFieldView());
 }
 
 void CompactBinaryToCompactJson(const FCbObjectView& Object, FUtf8StringBuilderBase& Builder)
 {
-	FCbJsonWriter Writer(Builder, "", "", "");
+	TCbJsonWriter<UTF8CHAR, ANSICHAR> Writer(Builder, ANSITEXTVIEW(""), ANSITEXTVIEW(""), ANSITEXTVIEW(""));
+	Writer.WriteField(Object.AsFieldView());
+}
+
+void CompactBinaryToCompactJson(const FCbObjectView& Object, FWideStringBuilderBase& Builder)
+{
+	TCbJsonWriter<WIDECHAR> Writer(Builder, WIDETEXTVIEW(""), WIDETEXTVIEW(""), WIDETEXTVIEW(""));
 	Writer.WriteField(Object.AsFieldView());
 }
 
diff --git a/Engine/Source/Runtime/Core/Public/Serialization/CompactBinarySerialization.h b/Engine/Source/Runtime/Core/Public/Serialization/CompactBinarySerialization.h
--- a/Engine/Source/Runtime/Core/Public/Serialization/CompactBinarySerialization.h
+++ b/Engine/Source/Runtime/Core/Public/Serialization/CompactBinarySerialization.h
@@ -213,8 +213,10 @@ inline bool LoadFromCompactBinary(FCbFieldView Field, TArray<T, Allocator>& OutV
 
 /** Convert the object to JSON in a multi-line format with indentation. */
 CORE_API void CompactBinaryToJson(const FCbObjectView& Object, FUtf8StringBuilderBase& Builder);
+CORE_API void CompactBinaryToJson(const FCbObjectView& Object, FWideStringBuilderBase& Builder);
 
 /** Convert the object to JSON in a compact format with no added whitespace. */
 CORE_API void CompactBinaryToCompactJson(const FCbObjectView& Object, FUtf8StringBuilderBase& Builder);
+CORE_API void CompactBinaryToCompactJson(const FCbObjectView& Object, FWideStringBuilderBase& Builder);
 
 ///////////////////////////////////////////////////////////////////////////////////////////////////////////////////////
