Fixes to get arm64 and arm64ec to partially compile.

diff --git a/Engine/Build/Windows/Resources/Default-Win64-arm64.manifest b/Engine/Build/Windows/Resources/Default-Win64-arm64.manifest
new file mode 100644
--- /dev/null
+++ b/Engine/Build/Windows/Resources/Default-Win64-arm64.manifest
@@ -0,0 +1,41 @@
+<?xml version="1.0" encoding="UTF-8" standalone="yes"?>
+<assembly xmlns="urn:schemas-microsoft-com:asm.v1" manifestVersion="1.0" xmlns:asmv3="urn:schemas-microsoft-com:asm.v3">
+
+  <!-- Tell Windows we don't need administrator level access-->
+  <ms_asmv2:trustInfo xmlns:ms_asmv2="urn:schemas-microsoft-com:asm.v2">
+    <ms_asmv2:security>
+      <ms_asmv2:requestedPrivileges>
+        <ms_asmv2:requestedExecutionLevel level="asInvoker" uiAccess="false" />
+      </ms_asmv2:requestedPrivileges>
+    </ms_asmv2:security>
+  </ms_asmv2:trustInfo>
+
+  <!-- Use themed common controls -->
+  <dependency>
+    <dependentAssembly>
+      <assemblyIdentity type='win32' name='Microsoft.Windows.Common-Controls' version='6.0.0.0' processorArchitecture='arm64' publicKeyToken='6595b64144ccf1df' />
+    </dependentAssembly>
+  </dependency>
+
+  <!-- OS Support -->
+  <compatibility xmlns="urn:schemas-microsoft-com:compatibility.v1">
+    <application>
+      <!-- Windows 10 -->
+      <supportedOS Id="{8e0f7a12-bfb3-4fe8-b9a5-48fd50a15a9a}"/>
+      <!-- Windows 8.1 -->
+      <supportedOS Id="{1f676c76-80e1-4239-95bb-83d0f6d0da78}"/>
+      <!-- Windows 7 -->
+      <supportedOS Id="{35138b9a-5d96-4fbd-8e2d-a2440225f93a}"/>
+      <!-- Windows 8 -->
+      <supportedOS Id="{4a2f28e3-53b9-4441-ba9c-d69d4a4a6e38}"/>
+    </application>
+  </compatibility>
+
+  <!-- Allow paths longer than MAX_PATH (also requires enabling per-system via group policy or registry entry) -->
+  <asmv3:application>
+    <windowsSettings xmlns:ws2="http://schemas.microsoft.com/SMI/2016/WindowsSettings">
+      <ws2:longPathAware>true</ws2:longPathAware>
+    </windowsSettings>
+  </asmv3:application>
+
+</assembly>
diff --git a/Engine/Source/Developer/Windows/LiveCoding/Private/External/xxhash.h b/Engine/Source/Developer/Windows/LiveCoding/Private/External/xxhash.h
--- a/Engine/Source/Developer/Windows/LiveCoding/Private/External/xxhash.h
+++ b/Engine/Source/Developer/Windows/LiveCoding/Private/External/xxhash.h
@@ -1458,7 +1458,7 @@ typedef XXH64_hash_t xxh_u64;
 #ifndef XXH_REROLL_XXH64
 #  if (defined(__ILP32__) || defined(_ILP32)) /* ILP32 is often defined on 32-bit GCC family */ \
    || !(defined(__x86_64__) || defined(_M_X64) || defined(_M_AMD64) /* x86-64 */ \
-     || defined(_M_ARM64) || defined(__aarch64__) || defined(__arm64__) /* aarch64 */ \
+     || defined(_M_ARM64) || defined(_M_ARM64EC) || defined(__aarch64__) || defined(__arm64__) /* aarch64 */ /* EPIC MOD: Added _M_ARM64EC support */ \
      || defined(__PPC64__) || defined(__PPC64LE__) || defined(__ppc64__) || defined(__powerpc64__) /* ppc64 */ \
      || defined(__mips64__) || defined(__mips64)) /* mips64 */ \
    || (!defined(SIZE_MAX) || SIZE_MAX < ULLONG_MAX) /* check limits */
diff --git a/Engine/Source/Programs/UnrealBuildTool/Configuration/TargetRules.cs b/Engine/Source/Programs/UnrealBuildTool/Configuration/TargetRules.cs
--- a/Engine/Source/Programs/UnrealBuildTool/Configuration/TargetRules.cs
+++ b/Engine/Source/Programs/UnrealBuildTool/Configuration/TargetRules.cs
@@ -1159,7 +1159,7 @@ namespace UnrealBuildTool
 		[RequiresUniqueBuildEnvironment]
 		public bool bWithLiveCoding
 		{
-			get { return bWithLiveCodingPrivate ?? (Platform == UnrealTargetPlatform.Win64 && Configuration != UnrealTargetConfiguration.Shipping && Configuration != UnrealTargetConfiguration.Test && Type != TargetType.Program); }
+			get { return bWithLiveCodingPrivate ?? (Platform == UnrealTargetPlatform.Win64 && Architecture != "arm64" && Architecture != "arm64ec" && Configuration != UnrealTargetConfiguration.Shipping && Configuration != UnrealTargetConfiguration.Test && Type != TargetType.Program); }
 			set { bWithLiveCodingPrivate = value; }
 		}
 		bool? bWithLiveCodingPrivate;
diff --git a/Engine/Source/Programs/UnrealBuildTool/Platform/Windows/MicrosoftPlatformSDK.cs b/Engine/Source/Programs/UnrealBuildTool/Platform/Windows/MicrosoftPlatformSDK.cs
--- a/Engine/Source/Programs/UnrealBuildTool/Platform/Windows/MicrosoftPlatformSDK.cs
+++ b/Engine/Source/Programs/UnrealBuildTool/Platform/Windows/MicrosoftPlatformSDK.cs
@@ -940,11 +940,16 @@ namespace UnrealBuildTool
 			Logger.LogDebug("Found Visual Studio toolchain: {ToolChainDir} (Family={Family}, FamilyRank={FamilyRank}, Version={Version}, Is64Bit={Is64Bit}, Preview={Preview}, Architecture={Arch}, Error={Error}, Redist={RedistDir})", ToolChainDir, Family, FamilyRank, Version, Is64Bit, bPreview, WindowsArchitecture.x64.ToString(), Error != null, RedistDir);
 			ToolChains.Add(new ToolChainInstallation(Family, FamilyRank, Version, Is64Bit, bPreview, WindowsArchitecture.x64, Error, ToolChainDir, RedistDir, IsAutoSdk));
 
-			bool HasArm64 = HasArm64ToolChain(ToolChainDir);
-			if (HasArm64)
+			if (HasArm64ToolChain(ToolChainDir))
 			{
 				Logger.LogDebug("Found Visual Studio toolchain: {ToolChainDir} (Family={Family}, FamilyRank={FamilyRank}, Version={Version}, Is64Bit={Is64Bit}, Preview={Preview}, Architecture={Arch}, Error={Error}, Redist={RedistDir})", ToolChainDir, Family, FamilyRank, Version, Is64Bit, bPreview, WindowsArchitecture.ARM64.ToString(), Error != null, RedistDir);
 				ToolChains.Add(new ToolChainInstallation(Family, FamilyRank, Version, Is64Bit, bPreview, WindowsArchitecture.ARM64, Error, ToolChainDir, RedistDir, IsAutoSdk));
+				
+				if (HasArm64ECToolChain(ToolChainDir))
+				{
+					Logger.LogDebug("Found Visual Studio toolchain: {ToolChainDir} (Family={Family}, FamilyRank={FamilyRank}, Version={Version}, Is64Bit={Is64Bit}, Preview={Preview}, Architecture={Arch}, Error={Error}, Redist={RedistDir})", ToolChainDir, Family, FamilyRank, Version, Is64Bit, bPreview, WindowsArchitecture.ARM64EC.ToString(), Error != null, RedistDir);
+					ToolChains.Add(new ToolChainInstallation(Family, FamilyRank, Version, Is64Bit, bPreview, WindowsArchitecture.ARM64EC, Error, ToolChainDir, RedistDir, IsAutoSdk));
+				}
 			}
 		}
 
@@ -1121,6 +1126,16 @@ namespace UnrealBuildTool
 			return FileReference.Exists(FileReference.Combine(ToolChainDir, "bin", "Hostx64", "arm64", "cl.exe"));
 		}
 
+		/// <summary>
+		/// Checks if the given directory contains a arm64ec toolchain.  Used to require arm64, which is an optional install item, when that is our target architecture.
+		/// </summary>
+		/// <param name="ToolChainDir">Directory to check</param>
+		/// <returns>True if the given directory contains the arm64ec toolchain</returns>
+		static bool HasArm64ECToolChain(DirectoryReference ToolChainDir)
+		{
+			return FileReference.Exists(FileReference.Combine(ToolChainDir, "lib", "arm64ec", "binmode.obj"));
+		}
+
 		/// <summary>
 		/// Determines if an IDE for the given compiler is installed.
 		/// </summary>
diff --git a/Engine/Source/Programs/UnrealBuildTool/Platform/Windows/UEBuildWindows.cs b/Engine/Source/Programs/UnrealBuildTool/Platform/Windows/UEBuildWindows.cs
--- a/Engine/Source/Programs/UnrealBuildTool/Platform/Windows/UEBuildWindows.cs
+++ b/Engine/Source/Programs/UnrealBuildTool/Platform/Windows/UEBuildWindows.cs
@@ -142,6 +142,10 @@ namespace UnrealBuildTool
 		/// ARM64
 		/// </summary>
 		ARM64,
+		/// <summary>
+		/// ARM64EC
+		/// </summary>
+		ARM64EC,
 	}
 
 	/// <summary>
@@ -538,7 +542,12 @@ namespace UnrealBuildTool
 		{
 			this.Target = Target;
 
-			ManifestFile = FileReference.Combine(Unreal.EngineDirectory, "Build", "Windows", "Resources", String.Format("Default-{0}.manifest", Target.Platform)).FullName;
+			string Platform = Target.Platform.ToString();
+			if (Target.Architecture == "arm64" || Target.Architecture == "arm64ec")
+			{
+				Platform += "-arm64";
+			}
+			ManifestFile = FileReference.Combine(Unreal.EngineDirectory, "Build", "Windows", "Resources", String.Format("Default-{0}.manifest", Platform)).FullName;
 		}
 	}
 
@@ -902,7 +911,28 @@ namespace UnrealBuildTool
 		{
 			if (Platform == UnrealTargetPlatform.Win64)
 			{
-				Target.WindowsPlatform.Architecture = WindowsArchitecture.x64;
+				switch(Target.Architecture)
+				{
+					case "arm64":
+						Target.WindowsPlatform.Architecture = WindowsArchitecture.ARM64;
+						break;
+					case "arm64ec":
+						Target.WindowsPlatform.Architecture = WindowsArchitecture.ARM64EC;
+						break;
+					default:
+						Target.WindowsPlatform.Architecture = WindowsArchitecture.x64;
+						break;
+				}
+
+				// Add names of plugins here that are incompatible with arm64
+				bool bCompilingForArm = Target.Architecture.IndexOf("arm", StringComparison.OrdinalIgnoreCase) >= 0;
+				if (bCompilingForArm && Target.Name != "UnrealHeaderTool")
+				{
+					Target.DisablePlugins.AddRange(new string[]
+					{
+						"OpenImageDenoise"
+					});
+				}
 			}
 
 			// Disable Simplygon support if compiling against the NULL RHI.
@@ -1145,13 +1175,30 @@ namespace UnrealBuildTool
 			{
 				archPath = "x64";
 			}
-			else if (arch == WindowsArchitecture.ARM64)
+			else if (arch == WindowsArchitecture.ARM64 || arch == WindowsArchitecture.ARM64EC)
 			{
 				archPath = "arm64";
 			}
 			return archPath;
 		}
 
+		public static string GetArchitectureName(WindowsArchitecture arch)
+		{
+			string archPath = "Unknown";
+			if (arch == WindowsArchitecture.x64)
+			{
+				archPath = "x64";
+			}
+			else if (arch == WindowsArchitecture.ARM64)
+			{
+				archPath = "arm64";
+			}
+			else if (arch == WindowsArchitecture.ARM64EC)
+			{
+				archPath = "arm64ec";
+			}
+			return archPath;
+		}
 
 		/// <summary>
 		/// Determines if a directory contains a valid DIA SDK
diff --git a/Engine/Source/Programs/UnrealBuildTool/Platform/Windows/VCEnvironment.cs b/Engine/Source/Programs/UnrealBuildTool/Platform/Windows/VCEnvironment.cs
--- a/Engine/Source/Programs/UnrealBuildTool/Platform/Windows/VCEnvironment.cs
+++ b/Engine/Source/Programs/UnrealBuildTool/Platform/Windows/VCEnvironment.cs
@@ -163,7 +163,14 @@ namespace UnrealBuildTool
 			// Add the Windows SDK directory to the path too, for mt.exe.
 			if (WindowsSdkVersion >= new VersionNumber(10))
 			{
-				AddDirectoryToPath(DirectoryReference.Combine(WindowsSdkDir, "bin", WindowsSdkVersion.ToString(), Architecture.ToString()));
+				string BuildHostArch = Architecture.ToString();
+				if (Architecture == WindowsArchitecture.ARM64 || Architecture == WindowsArchitecture.ARM64EC)
+				{
+					Debug.Assert(BuildHostPlatform.Current.Platform == UnrealTargetPlatform.Win64);
+					BuildHostArch = "x64";
+				}
+				AddDirectoryToPath(DirectoryReference.Combine(WindowsSdkDir, "bin", WindowsSdkVersion.ToString(), BuildHostArch));
+
 			}
 		}
 
diff --git a/Engine/Source/Programs/UnrealBuildTool/Platform/Windows/VCToolChain.cs b/Engine/Source/Programs/UnrealBuildTool/Platform/Windows/VCToolChain.cs
--- a/Engine/Source/Programs/UnrealBuildTool/Platform/Windows/VCToolChain.cs
+++ b/Engine/Source/Programs/UnrealBuildTool/Platform/Windows/VCToolChain.cs
@@ -828,6 +828,18 @@ namespace UnrealBuildTool
 				Arguments.Add($"-Xclang -x -Xclang \"{FileSpecifier}\"");
 			}
 
+			if (CompileEnvironment.Architecture == "arm64ec")
+			{
+				Arguments.Add("/arm64EC");
+				// The latest vc toolchain requires that these be manually set for arm64ec.
+				AddDefinition(Arguments, "_ARM64EC_");
+				AddDefinition(Arguments, "_ARM64EC_WORKAROUND_");
+				AddDefinition(Arguments, "ARM64EC");
+				AddDefinition(Arguments, "AMD64");
+				AddDefinition(Arguments, "_AMD64_");
+				AddDefinition(Arguments, "_WINDOWS");
+				AddDefinition(Arguments, "WIN32");
+			}
 
 			if (!CompileEnvironment.bEnableBufferSecurityChecks)
 			{
@@ -1121,8 +1133,7 @@ namespace UnrealBuildTool
 			//
 			if (UseWindowsArchitecture(LinkEnvironment.Platform))
 			{
-				Arguments.Add($"/MACHINE:{WindowsExports.GetArchitectureSubpath(Target.WindowsPlatform.Architecture)}");
-
+				Arguments.Add($"/MACHINE:{WindowsExports.GetArchitectureName(Target.WindowsPlatform.Architecture)}");
 				{
 					if (LinkEnvironment.bIsBuildingConsoleApplication)
 					{
@@ -1158,7 +1169,10 @@ namespace UnrealBuildTool
 			}
 
 			// Allow delay-loaded DLLs to be explicitly unloaded.
-			Arguments.Add("/DELAY:UNLOAD");
+			if (Target.WindowsPlatform.Architecture == WindowsArchitecture.x64)
+			{
+				Arguments.Add("/DELAY:UNLOAD");
+			}
 
 			if (LinkEnvironment.bIsBuildingDLL)
 			{
@@ -1244,7 +1258,8 @@ namespace UnrealBuildTool
 			if (LinkEnvironment.bUseIncrementalLinking && 
 				LinkEnvironment.Configuration != CppConfiguration.Shipping && 
 				!Target.WindowsPlatform.bMergeIdenticalCOMDATs &&
-				!LinkEnvironment.bAllowLTCG)
+				!LinkEnvironment.bAllowLTCG &&
+				Target.WindowsPlatform.Architecture == WindowsArchitecture.x64)
 			{
 				Arguments.Add("/INCREMENTAL");
 				Arguments.Add("/verbose:incr");
@@ -1286,8 +1301,7 @@ namespace UnrealBuildTool
 			//
 			if (UseWindowsArchitecture(LinkEnvironment.Platform))
 			{
-				Arguments.Add($"/MACHINE:{WindowsExports.GetArchitectureSubpath(Target.WindowsPlatform.Architecture)}");
-
+				Arguments.Add($"/MACHINE:{WindowsExports.GetArchitectureName(Target.WindowsPlatform.Architecture)}");
 				{
 					if (LinkEnvironment.bIsBuildingConsoleApplication)
 					{
@@ -1792,10 +1806,7 @@ namespace UnrealBuildTool
 
 				// If we're compiling for 64-bit Windows, also add the _WIN64 definition to the resource
 				// compiler so that we can switch on that in the .rc file using #ifdef.
-				if (Target.WindowsPlatform.Architecture == WindowsArchitecture.x64 || Target.WindowsPlatform.Architecture == WindowsArchitecture.ARM64)
-				{
-					AddDefinition(Arguments, "_WIN64");
-				}
+				AddDefinition(Arguments, "_WIN64");
 
 				// Language
 				Arguments.Add("/l 0x409");
diff --git a/Engine/Source/Programs/UnrealBuildTool/Platform/Windows/WindowsExports.cs b/Engine/Source/Programs/UnrealBuildTool/Platform/Windows/WindowsExports.cs
--- a/Engine/Source/Programs/UnrealBuildTool/Platform/Windows/WindowsExports.cs
+++ b/Engine/Source/Programs/UnrealBuildTool/Platform/Windows/WindowsExports.cs
@@ -47,6 +47,16 @@ namespace UnrealBuildTool
 			return WindowsPlatform.GetArchitectureSubpath(arch);
 		}
 
+		/// <summary>
+		/// Returns the architecture name of the current architecture.
+		/// </summary>
+		/// <param name="arch">The architecture enum</param>
+		/// <returns>String with the name</returns>
+		public static string GetArchitectureName(WindowsArchitecture arch)
+		{
+			return WindowsPlatform.GetArchitectureName(arch);
+		}
+
 		/// <summary>
 		/// Tries to get the directory for an installed Windows SDK
 		/// </summary>
diff --git a/Engine/Source/Runtime/Core/Public/HAL/Platform.h b/Engine/Source/Runtime/Core/Public/HAL/Platform.h
--- a/Engine/Source/Runtime/Core/Public/HAL/Platform.h
+++ b/Engine/Source/Runtime/Core/Public/HAL/Platform.h
@@ -72,7 +72,7 @@
 
 // Whether the CPU is x86/x64 (i.e. both 32 and 64-bit variants)
 #ifndef PLATFORM_CPU_X86_FAMILY
-	#if (defined(_M_IX86) || defined(__i386__) || defined(_M_X64) || defined(__amd64__) || defined(__x86_64__))
+	#if (defined(_M_IX86) || defined(__i386__) || defined(_M_X64) || defined(__amd64__) || defined(__x86_64__)) && !defined(_M_ARM64EC)
 		#define PLATFORM_CPU_X86_FAMILY	1
 	#else
 		#define PLATFORM_CPU_X86_FAMILY	0
@@ -81,7 +81,7 @@
 
 // Whether the CPU is AArch32/AArch64 (i.e. both 32 and 64-bit variants)
 #ifndef PLATFORM_CPU_ARM_FAMILY
-	#if (defined(__arm__) || defined(_M_ARM) || defined(__aarch64__) || defined(_M_ARM64))
+	#if (defined(__arm__) || defined(_M_ARM) || defined(__aarch64__) || defined(_M_ARM64) || defined(_M_ARM64EC))
 		#define PLATFORM_CPU_ARM_FAMILY	1
 	#else
 		#define PLATFORM_CPU_ARM_FAMILY	0
diff --git a/Engine/Source/Runtime/Core/Public/Windows/WindowsPlatform.h b/Engine/Source/Runtime/Core/Public/Windows/WindowsPlatform.h
--- a/Engine/Source/Runtime/Core/Public/Windows/WindowsPlatform.h
+++ b/Engine/Source/Runtime/Core/Public/Windows/WindowsPlatform.h
@@ -54,7 +54,13 @@ typedef FWindowsPlatformTypes FPlatformTypes;
 #define PLATFORM_SUPPORTS_UNALIGNED_LOADS					1
 
 #define PLATFORM_SUPPORTS_PRAGMA_PACK						1
-#define PLATFORM_ENABLE_VECTORINTRINSICS					1
+#if defined(_M_ARM) || defined(_M_ARM64) || defined(_M_ARM64EC)
+	#define PLATFORM_CPU_ARM_FAMILY							1
+	#define PLATFORM_ENABLE_VECTORINTRINSICS_NEON			1
+	#define PLATFORM_ENABLE_VECTORINTRINSICS				1
+#elif (defined(_M_IX86) || defined(_M_X64))
+	#define PLATFORM_CPU_X86_FAMILY							1
+	#define PLATFORM_ENABLE_VECTORINTRINSICS				1
 #ifndef PLATFORM_MAYBE_HAS_SSE4_1 // May be set from UnrealBuildTool
 	#define PLATFORM_MAYBE_HAS_SSE4_1						1
 #endif
@@ -64,6 +70,8 @@ typedef FWindowsPlatformTypes FPlatformTypes;
 #ifndef PLATFORM_ALWAYS_HAS_SSE4_1 // May be set from UnrealBuildTool
 	#define PLATFORM_ALWAYS_HAS_SSE4_1						0
 #endif
+
+#endif
 // FMA3 support was added starting from AMD Piledriver (excluding Jaguar) and Intel Haswell (excluding Pentium and Celeron)
 #ifndef PLATFORM_ALWAYS_HAS_FMA3
 	#define PLATFORM_ALWAYS_HAS_FMA3						0
@@ -81,7 +89,11 @@ typedef FWindowsPlatformTypes FPlatformTypes;
 #define PLATFORM_USES_MICROSOFT_LIBC_FUNCTIONS				1
 #define PLATFORM_IS_ANSI_MALLOC_THREADSAFE					1
 #define PLATFORM_SUPPORTS_TBB								1
+#if defined(PLATFORM_CPU_ARM_FAMILY)
+#define PLATFORM_SUPPORTS_MIMALLOC							0
+#else
 #define PLATFORM_SUPPORTS_MIMALLOC							PLATFORM_64BITS
+#endif
 #define PLATFORM_SUPPORTS_NAMED_PIPES						1
 #define PLATFORM_COMPILER_HAS_TCHAR_WMAIN					1
 #define PLATFORM_SUPPORTS_EARLY_MOVIE_PLAYBACK				(!WITH_EDITOR) // movies will start before engine is initalized
diff --git a/Engine/Source/Runtime/Core/Public/Windows/WindowsPlatformCompilerSetup.h b/Engine/Source/Runtime/Core/Public/Windows/WindowsPlatformCompilerSetup.h
--- a/Engine/Source/Runtime/Core/Public/Windows/WindowsPlatformCompilerSetup.h
+++ b/Engine/Source/Runtime/Core/Public/Windows/WindowsPlatformCompilerSetup.h
@@ -30,7 +30,7 @@ static_assert(_MSC_VER >= 1929, "Visual Studio 2019 v16.11 or greater is require
 #pragma warning (error: 4400 4401 4402 4403 4404 4405 4406 4407 4408 4409 4410 4411      4413 4414 4415 4416 4417 4418 4419 4420           4423 4424 4425 4426 4427      4429 4430 4431           4434      4436      4438 4439 4440 4441 4442 4443      4445 4446 4447 4448 4449 4450 4451 4452 4453 4454 4455                     4460 4461 4462      4464 4465 4466 4467 4468 4469 4470           4473 4474 4475 4476 4477 4478      4480      4482 4483 4484 4485 4486 4487 4488 4489 4490 4491 4492 4493 4494 4495 4496 4497 4498 4499)
 #pragma warning (error:           4502 4503      4505 4506      4508 4509      4511 4512 4513 4514 4515 4516 4517 4518 4519      4521 4522 4523           4526                4530 4531 4532 4533 4534 4535 4536 4537 4538      4540 4541 4542 4543 4544 4545 4546                4550 4551 4552 4553 4554      4556 4557 4558 4559      4561 4562      4564 4565 4566      4568 4569 4570      4572 4573      4575 4576      4578      4580 4581           4584 4585 4586 4587 4588 4589      4591      4593 4594 4595 4596 4597 4598     )
 #pragma warning (error: 4600      4602 4603 4604      4606           4609 4610 4611 4612 4613      4615 4616      4618      4620 4621 4622 4623 4624 4625 4626 4627 4628 4629 4630 4631 4632 4633 4634 4635 4636 4637 4638 4639 4640 4641 4642      4644 4645 4646      4648 4649 4650      4652 4653 4654 4655 4656 4657 4658 4659      4661 4662                     4667      4669 4670 4671 4672 4673 4674      4676 4677 4678 4679 4680 4681 4682 4683 4684 4685 4686 4687 4688 4689 4690 4691      4693 4694 4695 4696      4698     )
-#pragma warning (error: 4700      4702 4703           4706           4709 4710 4711           4714 4715 4716 4717 4718 4719 4720 4721 4722 4723 4724 4725 4726 4727 4728 4729      4731 4732 4733 4734                     4739 4740 4741 4742 4743 4744 4745 4746 4747      4749 4750 4751 4752 4753 4754 4755 4756 4757                4761           4764                               4771 4772           4775 4776 4777 4778                                    4786      4788 4789           4792 4793 4794                4798 4799)
+#pragma warning (error: 4700      4702 4703           4706           4709 4710 4711           4714 4715 4716 4717 4718 4719 4720 4721 4722 4723 4724 4725 4726 4727 4728 4729      4731 4732 4733 4734                     4739 4740 4741 4742 4743 4744 4745      4747      4749 4750 4751 4752 4753 4754 4755 4756 4757                4761           4764                               4771 4772           4775 4776 4777 4778                                    4786      4788 4789           4792 4793 4794                4798 4799)
 #pragma warning (error: 4800 4801      4803 4804 4805 4806 4807 4808 4809 4810 4811 4812 4813      4815 4816 4817                4821 4822 4823                4827      4829                     4834 4835                4839 4840 4841 4842 4843 4844 4845 4846 4847 4848 4849                                                                                      4867      4869           4872                                    4880 4881 4882 4883                                                                                )
 #pragma warning (error: 4900                     4905 4906                4910      4912 4913           4916 4917 4918      4920 4921                4925 4926 4927 4928 4929 4930 4931 4932      4934 4935 4936 4937 4938 4939                     4944 4945 4946 4947 4948 4949 4950 4951 4952 4953 4954 4955 4956 4957 4958 4959 4960 4961      4963 4964 4965 4966                4970 4971 4972 4973 4974                               4981                4985                4989 4990 4991 4992           4995      4997 4998 4999)
 #pragma warning (error:                                                                                                                                                                                               5038                                                                                                                                                                                                                                                                                                                 )
@@ -221,6 +221,7 @@ static_assert(_MSC_VER >= 1929, "Visual Studio 2019 v16.11 or greater is require
 #pragma warning(default: 4686) // 'user-defined type' : possible change in behavior, change in UDT return calling convention		https://docs.microsoft.com/en-us/cpp/error-messages/compiler-warnings/compiler-warning-level-3-c4686
 
 #pragma warning(disable: 4710) // 'function' : function not inlined / The given function was selected for inline expansion, but the compiler did not perform the inlining.
+#pragma warning(disable: 4746) // volatile access of '<expression>' is subject to /volatile											https://docs.microsoft.com/en-us/cpp/error-messages/compiler-warnings/compiler-warning-c4746?view=msvc-170
 #pragma warning(default: 4786) // 'identifier' : identifier was truncated to 'number' characters in the debug information			// No docs
 #pragma warning(default: 4793) // 'function' : function is compiled as native code: 'reason'										https://docs.microsoft.com/en-us/cpp/error-messages/compiler-warnings/compiler-warning-level-1-and-3-c4793
 
diff --git a/Engine/Source/Runtime/Core/Public/Windows/WindowsPlatformMath.h b/Engine/Source/Runtime/Core/Public/Windows/WindowsPlatformMath.h
--- a/Engine/Source/Runtime/Core/Public/Windows/WindowsPlatformMath.h
+++ b/Engine/Source/Runtime/Core/Public/Windows/WindowsPlatformMath.h
@@ -6,8 +6,11 @@
 	#error this code should only be included on Windows
 #endif
 
-#include <intrin.h>
-#include <smmintrin.h>
+#if PLATFORM_CPU_X86_FAMILY
+	#include <intrin.h>
+	#include <smmintrin.h>
+#endif
+
 #include "Microsoft/MicrosoftPlatformMath.h"
 
 typedef FMicrosoftPlatformMathBase FWindowsPlatformMath;
diff --git a/Engine/Source/Runtime/Core/Public/Windows/WindowsPlatformMemory.h b/Engine/Source/Runtime/Core/Public/Windows/WindowsPlatformMemory.h
--- a/Engine/Source/Runtime/Core/Public/Windows/WindowsPlatformMemory.h
+++ b/Engine/Source/Runtime/Core/Public/Windows/WindowsPlatformMemory.h
@@ -6,6 +6,8 @@
 #include "GenericPlatform/GenericPlatformMemory.h"
 #include "Windows/WindowsSystemIncludes.h"
 
+#include <malloc.h>
+
 class FString;
 class FMalloc;
 struct FGenericMemoryStats;
diff --git a/Engine/Source/Runtime/Core/Public/Windows/WindowsPlatformMisc.h b/Engine/Source/Runtime/Core/Public/Windows/WindowsPlatformMisc.h
--- a/Engine/Source/Runtime/Core/Public/Windows/WindowsPlatformMisc.h
+++ b/Engine/Source/Runtime/Core/Public/Windows/WindowsPlatformMisc.h
@@ -76,7 +76,14 @@ struct CORE_API FWindowsPlatformMisc
 	static void CustomNamedStat(const ANSICHAR* Text, float Value, const ANSICHAR* Graph, const ANSICHAR* Unit);
 #endif
 
-	FORCEINLINE static void MemoryBarrier() { _mm_sfence(); }
+	FORCEINLINE static void MemoryBarrier() 
+	{
+#if PLATFORM_CPU_X86_FAMILY
+		_mm_sfence();
+#elif PLATFORM_CPU_ARM_FAMILY
+		__dmb(_ARM64_BARRIER_SY);
+#endif
+	}
 
 	static bool IsRemoteSession();
 
@@ -162,14 +169,23 @@ struct CORE_API FWindowsPlatformMisc
 		const int32 CacheLineSize = GetCacheLineSize();
 		for (int32 LinesToPrefetch = (NumBytes + CacheLineSize - 1) / CacheLineSize; LinesToPrefetch; --LinesToPrefetch)
 		{
+#if PLATFORM_CPU_X86_FAMILY
 			_mm_prefetch( Ptr, _MM_HINT_T0 );
+#elif PLATFORM_CPU_ARM_FAMILY
+			__prefetch(Ptr);
+#endif
 			Ptr += CacheLineSize;
 		}
 	}
 
 	FORCEINLINE static void Prefetch(void const* x, int32 offset = 0)
 	{
-		 _mm_prefetch( (char const*)(x) + offset, _MM_HINT_T0 );
+#if PLATFORM_CPU_X86_FAMILY
+		_mm_prefetch((char const*)(x)+offset, _MM_HINT_T0);
+#elif PLATFORM_CPU_ARM_FAMILY
+		__prefetch((char const*)(x)+offset);
+#endif
+		
 	}
 
 	/** 
diff --git a/Engine/Source/Runtime/TraceLog/standalone_prologue.h b/Engine/Source/Runtime/TraceLog/standalone_prologue.h
--- a/Engine/Source/Runtime/TraceLog/standalone_prologue.h
+++ b/Engine/Source/Runtime/TraceLog/standalone_prologue.h
@@ -39,7 +39,7 @@
 #	define PLATFORM_CPU_ARM_FAMILY	0
 #	define PLATFORM_64BITS			1
 #	define PLATFORM_CACHE_LINE_SIZE	64
-#elif defined(__arm64__) || defined(_M_ARM64)
+#elif defined(__arm64__) || defined(_M_ARM64) || defined(_M_ARM64EC)
 #	define PLATFORM_CPU_X86_FAMILY	0
 #	define PLATFORM_CPU_ARM_FAMILY	1
 #	define PLATFORM_64BITS			1
