UE_LIFETIMEBOUND for wrapping Clang's lifetime checking attribute

Adds support to TFunctionRef, TStringView and TArrayView. More to come

diff --git a/Engine/Source/Runtime/Core/Public/Clang/ClangPlatform.h b/Engine/Source/Runtime/Core/Public/Clang/ClangPlatform.h
--- a/Engine/Source/Runtime/Core/Public/Clang/ClangPlatform.h
+++ b/Engine/Source/Runtime/Core/Public/Clang/ClangPlatform.h
@@ -20,3 +20,5 @@
 
 #define PLATFORM_RETURN_ADDRESS()			__builtin_return_address(0)
 #define PLATFORM_RETURN_ADDRESS_POINTER()	__builtin_frame_address(0)
+
+#define UE_LIFETIMEBOUND [[clang::lifetimebound]]
diff --git a/Engine/Source/Runtime/Core/Public/Containers/ArrayView.h b/Engine/Source/Runtime/Core/Public/Containers/ArrayView.h
--- a/Engine/Source/Runtime/Core/Public/Containers/ArrayView.h
+++ b/Engine/Source/Runtime/Core/Public/Containers/ArrayView.h
@@ -88,6 +88,12 @@ namespace ArrayViewPrivate
 	};
 }
 
+template <typename T>                                  struct TIsTArrayView                                                       { static constexpr bool Value = false; };
+template <typename InElementType, typename InSizeType> struct TIsTArrayView<               TArrayView<InElementType, InSizeType>> { static constexpr bool Value = true;  };
+template <typename InElementType, typename InSizeType> struct TIsTArrayView<      volatile TArrayView<InElementType, InSizeType>> { static constexpr bool Value = true;  };
+template <typename InElementType, typename InSizeType> struct TIsTArrayView<const          TArrayView<InElementType, InSizeType>> { static constexpr bool Value = true;  };
+template <typename InElementType, typename InSizeType> struct TIsTArrayView<const volatile TArrayView<InElementType, InSizeType>> { static constexpr bool Value = true;  };
+
 /**
  * Templated fixed-size view of another array
  *
@@ -163,7 +169,8 @@ public:
 					TIsReinterpretableRangeType<OtherRangeType>
 				>
 			>::Value
-		>::Type
+		>::Type,
+		std::enable_if_t<TIsTArrayView<std::decay_t<OtherRangeType>>::Value>* = nullptr
 	>
 	FORCEINLINE TArrayView(OtherRangeType&& Other)
 		: DataPtr(TChooseClass<
@@ -176,6 +183,31 @@ public:
 		check((InCount >= 0) && ((sizeof(InCount) < sizeof(SizeType)) || (InCount <= static_cast<decltype(InCount)>(TNumericLimits<SizeType>::Max()))));
 		ArrayNum = (SizeType)InCount;
 	}
+	template <
+		typename OtherRangeType,
+		typename CVUnqualifiedOtherRangeType = typename TRemoveCV<typename TRemoveReference<OtherRangeType>::Type>::Type,
+		typename = typename TEnableIf<
+		TAnd<
+		TIsContiguousContainer<CVUnqualifiedOtherRangeType>,
+		TOr<
+		TIsCompatibleRangeType<OtherRangeType>,
+		TIsReinterpretableRangeType<OtherRangeType>
+		>
+		>::Value
+		>::Type,
+		std::enable_if_t<!TIsTArrayView<std::decay_t<OtherRangeType>>::Value>* = nullptr
+	>
+	FORCEINLINE TArrayView(OtherRangeType&& Other UE_LIFETIMEBOUND)
+		: DataPtr(TChooseClass<
+			TIsCompatibleRangeType<OtherRangeType>::Value,
+			TIsCompatibleRangeType<OtherRangeType>,
+			TIsReinterpretableRangeType<OtherRangeType>
+		>::Result::GetData(Forward<OtherRangeType>(Other)))
+	{
+		const auto InCount = GetNum(Forward<OtherRangeType>(Other));
+		check((InCount >= 0) && ((sizeof(InCount) < sizeof(SizeType)) || (InCount <= static_cast<decltype(InCount)>(TNumericLimits<SizeType>::Max()))));
+		ArrayNum = (SizeType)InCount;
+	}
 
 	/**
 	 * Construct a view of an arbitrary pointer
@@ -197,10 +229,11 @@ public:
 	 *
 	 * The caller is responsible for ensuring that the view does not outlive the initializer list.
 	 */
-	FORCEINLINE TArrayView(std::initializer_list<ElementType> List)
+	FORCEINLINE TArrayView(std::initializer_list<ElementType> List UE_LIFETIMEBOUND)
 		: DataPtr(ArrayViewPrivate::GetDataHelper(List))
 		, ArrayNum(GetNum(List))
 	{
+		static_assert(std::is_const_v<ElementType>, "Only views of const elements can bind to initializer lists");
 	}
 
 public:
@@ -717,12 +750,23 @@ struct TIsContiguousContainer<TArrayView<T, SizeType>>
 template <
 	typename OtherRangeType,
 	typename CVUnqualifiedOtherRangeType = typename TRemoveCV<typename TRemoveReference<OtherRangeType>::Type>::Type,
-	typename = typename TEnableIf<TIsContiguousContainer<CVUnqualifiedOtherRangeType>::Value>::Type
+	typename = typename TEnableIf<TIsContiguousContainer<CVUnqualifiedOtherRangeType>::Value>::Type,
+	std::enable_if_t<TIsTArrayView<std::decay_t<OtherRangeType>>::Value>* = nullptr
 >
 auto MakeArrayView(OtherRangeType&& Other)
 {
 	return TArrayView<typename TRemovePointer<decltype(GetData(DeclVal<OtherRangeType&>()))>::Type>(Forward<OtherRangeType>(Other));
 }
+template <
+	typename OtherRangeType,
+	typename CVUnqualifiedOtherRangeType = typename TRemoveCV<typename TRemoveReference<OtherRangeType>::Type>::Type,
+	typename = typename TEnableIf<TIsContiguousContainer<CVUnqualifiedOtherRangeType>::Value>::Type,
+	std::enable_if_t<!TIsTArrayView<std::decay_t<OtherRangeType>>::Value>* = nullptr
+>
+auto MakeArrayView(OtherRangeType&& Other UE_LIFETIMEBOUND)
+{
+	return TArrayView<typename TRemovePointer<decltype(GetData(DeclVal<OtherRangeType&>()))>::Type>(Forward<OtherRangeType>(Other));
+}
 
 template<typename ElementType>
 auto MakeArrayView(ElementType* Pointer, int32 Size)
@@ -731,9 +775,9 @@ auto MakeArrayView(ElementType* Pointer, int32 Size)
 }
 
 template <typename T>
-TArrayView<const T> MakeArrayView(std::initializer_list<T> List)
+TArrayView<const T> MakeArrayView(std::initializer_list<T> List UE_LIFETIMEBOUND)
 {
-	return TArrayView<const T>(List);
+	return TArrayView<const T>(List.begin(), List.size());
 }
 
 //////////////////////////////////////////////////////////////////////////
diff --git a/Engine/Source/Runtime/Core/Public/Containers/StringView.h b/Engine/Source/Runtime/Core/Public/Containers/StringView.h
--- a/Engine/Source/Runtime/Core/Public/Containers/StringView.h
+++ b/Engine/Source/Runtime/Core/Public/Containers/StringView.h
@@ -142,7 +142,7 @@ public:
 			!std::is_array_v<std::remove_reference_t<CharRangeType>> &&
 			!std::is_same_v<CharRangeType, ViewType>
 		>* = nullptr>
-	constexpr inline TStringView(const CharRangeType& InRange)
+	constexpr inline TStringView(const CharRangeType& InRange UE_LIFETIMEBOUND)
 		: DataPtr((const CharType*)UE::Core::Private::StringViewGetData(InRange))
 		, Size(IntCastChecked<int32>(GetNum(InRange)))
 	{
@@ -453,7 +453,7 @@ TStringView(CharRangeType&& Range)
 	-> TStringView<TElementType_T<CharRangeType>>;
 
 template <typename CharPtrOrRangeType>
-constexpr inline auto MakeStringView(CharPtrOrRangeType&& CharPtrOrRange) -> decltype(TStringView(Forward<CharPtrOrRangeType>(CharPtrOrRange)))
+constexpr inline auto MakeStringView(CharPtrOrRangeType&& CharPtrOrRange UE_LIFETIMEBOUND) -> decltype(TStringView(Forward<CharPtrOrRangeType>(CharPtrOrRange)))
 {
 	return TStringView(Forward<CharPtrOrRangeType>(CharPtrOrRange));
 }
diff --git a/Engine/Source/Runtime/Core/Public/HAL/Platform.h b/Engine/Source/Runtime/Core/Public/HAL/Platform.h
--- a/Engine/Source/Runtime/Core/Public/HAL/Platform.h
+++ b/Engine/Source/Runtime/Core/Public/HAL/Platform.h
@@ -688,6 +688,11 @@
 	#define FUNCTION_NON_NULL_RETURN_END
 #endif
 
+/* Wrap lifetimebound annotations to indicate that a function argument must outlive a return value or constructed object */
+#ifndef UE_LIFETIMEBOUND
+	#define UE_LIFETIMEBOUND
+#endif
+
 /** Promise expression is true. Compiler can optimize accordingly with undefined behavior if wrong. Static analyzers understand this.  */
 #ifndef UE_ASSUME
 	#if defined(__clang__)
diff --git a/Engine/Source/Runtime/Core/Public/Templates/Function.h b/Engine/Source/Runtime/Core/Public/Templates/Function.h
--- a/Engine/Source/Runtime/Core/Public/Templates/Function.h
+++ b/Engine/Source/Runtime/Core/Public/Templates/Function.h
@@ -785,7 +785,7 @@ public:
 			>::Value
 		>::Type
 	>
-	TFunctionRef(FunctorType&& InFunc)
+	TFunctionRef(FunctorType&& InFunc UE_LIFETIMEBOUND)
 		: Super(Forward<FunctorType>(InFunc))
 	{
 		// This constructor is disabled for TFunctionRef types so it isn't incorrectly selected as copy/move constructors.
