removed StatsThread as it's not used anymore
cleaned recent removal of AudioThread
removed `StatsThread` and `AudioThread` from ENamedThreads
renamed FPlatformAffinity::GetAudioThreadMask() to GetAudioRenderThreadMask() as it's still used for AudioRenderThread

diff --git a/Engine/Plugins/Messaging/MessagingDebugger/Source/MessagingDebugger/Private/Widgets/MessageDetails/SMessagingDispatchStateTableRow.h b/Engine/Plugins/Messaging/MessagingDebugger/Source/MessagingDebugger/Private/Widgets/MessageDetails/SMessagingDispatchStateTableRow.h
--- a/Engine/Plugins/Messaging/MessagingDebugger/Source/MessagingDebugger/Private/Widgets/MessageDetails/SMessagingDispatchStateTableRow.h
+++ b/Engine/Plugins/Messaging/MessagingDebugger/Source/MessagingDebugger/Private/Widgets/MessageDetails/SMessagingDispatchStateTableRow.h
@@ -189,12 +189,6 @@ protected:
 			return LOCTEXT("ActualRenderingThread", "ActualRenderingThread");
 			break;
 
-#if STATS && !UE_STATS_THREAD_AS_PIPE
-		case ENamedThreads::StatsThread:
-			return LOCTEXT("StatsThread", "StatsThread");
-			break;
-#endif
-
 		default:
 			return LOCTEXT("UnknownThread", "Unknown");
 		}
diff --git a/Engine/Plugins/Messaging/MessagingDebugger/Source/MessagingDebugger/Private/Widgets/MessageDetails/SMessagingMessageDetails.cpp b/Engine/Plugins/Messaging/MessagingDebugger/Source/MessagingDebugger/Private/Widgets/MessageDetails/SMessagingMessageDetails.cpp
--- a/Engine/Plugins/Messaging/MessagingDebugger/Source/MessagingDebugger/Private/Widgets/MessageDetails/SMessagingMessageDetails.cpp
+++ b/Engine/Plugins/Messaging/MessagingDebugger/Source/MessagingDebugger/Private/Widgets/MessageDetails/SMessagingMessageDetails.cpp
@@ -233,16 +233,6 @@ FText SMessagingMessageDetails::HandleSenderThreadText() const
 			return LOCTEXT("ActualRenderingThread_Local", "ActualRenderingThread_Local");
 			break;
 
-#if STATS && !UE_STATS_THREAD_AS_PIPE
-		case ENamedThreads::StatsThread:
-			return LOCTEXT("StatsThread", "StatsThread");
-			break;
-
-		case ENamedThreads::StatsThread_Local:
-			return LOCTEXT("StatsThread_Local", "StatsThread_Local");
-			break;
-#endif
-
 		default:
 			return LOCTEXT("UnknownThread", "Unknown");
 		}
diff --git a/Engine/Plugins/Runtime/MicrosoftSpatialAudio/Source/MicrosoftSpatialSound/Private/MicrosoftSpatialSoundPlugin.cpp b/Engine/Plugins/Runtime/MicrosoftSpatialAudio/Source/MicrosoftSpatialSound/Private/MicrosoftSpatialSoundPlugin.cpp
--- a/Engine/Plugins/Runtime/MicrosoftSpatialAudio/Source/MicrosoftSpatialSound/Private/MicrosoftSpatialSoundPlugin.cpp
+++ b/Engine/Plugins/Runtime/MicrosoftSpatialAudio/Source/MicrosoftSpatialSound/Private/MicrosoftSpatialSoundPlugin.cpp
@@ -212,7 +212,7 @@ void FMicrosoftSpatialSound::Initialize(const FAudioPluginInitializationParams I
 
 			bWarnedMicrosoftSpatialSoundDynamicObjectCountIsZero = false;
 
-			SpatialAudioRenderThread = FRunnableThread::Create(this, TEXT("MicrosoftSpatialAudioThread"), 0, TPri_TimeCritical, FPlatformAffinity::GetAudioThreadMask());
+			SpatialAudioRenderThread = FRunnableThread::Create(this, TEXT("MicrosoftSpatialAudioThread"), 0, TPri_TimeCritical, FPlatformAffinity::GetAudioRenderThreadMask());
 
 			Audio::Analytics::RecordEvent_Usage(TEXT("MicrosoftSpatialSound.Initialized"));
 		}
diff --git a/Engine/Source/Developer/ProfilerService/Private/ProfilerServiceManager.cpp b/Engine/Source/Developer/ProfilerService/Private/ProfilerServiceManager.cpp
--- a/Engine/Source/Developer/ProfilerService/Private/ProfilerServiceManager.cpp
+++ b/Engine/Source/Developer/ProfilerService/Private/ProfilerServiceManager.cpp
@@ -164,13 +164,11 @@ void FProfilerServiceManager::SetPreviewState( const FMessageAddress& ClientAddr
 				// Enable stat capture.
 				if (PreviewClients.Num() == 0)
 				{
-#if UE_STATS_THREAD_AS_PIPE
 					if (FPlatformProcess::SupportsMultithreading())
 					{
 						GStatsPipe.Launch(UE_SOURCE_LOCATION, [this] { AddNewFrameHandleStatsPipe(); });
 					}
 					else
-#endif
 					{
 						FSimpleDelegateGraphTask::CreateAndDispatchWhenReady
 						(
@@ -192,13 +190,11 @@ void FProfilerServiceManager::SetPreviewState( const FMessageAddress& ClientAddr
 				// Disable stat capture.
 				if (PreviewClients.Num() == 0)
 				{
-#if UE_STATS_THREAD_AS_PIPE
 					if (FPlatformProcess::SupportsMultithreading())
 					{
 						GStatsPipe.Launch(UE_SOURCE_LOCATION, [this] { RemoveNewFrameHandleStatsPipe(); });
 					}
 					else
-#endif
 					{
 						FSimpleDelegateGraphTask::CreateAndDispatchWhenReady
 						(
diff --git a/Engine/Source/Runtime/AudioMixerCore/Private/AudioMixer.cpp b/Engine/Source/Runtime/AudioMixerCore/Private/AudioMixer.cpp
--- a/Engine/Source/Runtime/AudioMixerCore/Private/AudioMixer.cpp
+++ b/Engine/Source/Runtime/AudioMixerCore/Private/AudioMixer.cpp
@@ -652,7 +652,7 @@ namespace Audio
 		check(AudioFadeEvent != nullptr);
 
 		check(!AudioRenderThread.IsValid());
-		uint64 RenderThreadAffinityCVar = SetRenderThreadAffinityCVar > 0 ? uint64(SetRenderThreadAffinityCVar) : FPlatformAffinity::GetAudioThreadMask();
+		uint64 RenderThreadAffinityCVar = SetRenderThreadAffinityCVar > 0 ? uint64(SetRenderThreadAffinityCVar) : FPlatformAffinity::GetAudioRenderThreadMask();
 		AudioRenderThread.Reset(FRunnableThread::Create(this, *FString::Printf(TEXT("AudioMixerRenderThread(%d)"), AudioMixerTaskCounter.Increment()), 0, (EThreadPriority)SetRenderThreadPriorityCVar, RenderThreadAffinityCVar));
 		check(AudioRenderThread.IsValid());
 	}
diff --git a/Engine/Source/Runtime/AudioMixerCore/Private/AudioMixerNullDevice.cpp b/Engine/Source/Runtime/AudioMixerCore/Private/AudioMixerNullDevice.cpp
--- a/Engine/Source/Runtime/AudioMixerCore/Private/AudioMixerNullDevice.cpp
+++ b/Engine/Source/Runtime/AudioMixerCore/Private/AudioMixerNullDevice.cpp
@@ -101,7 +101,7 @@ namespace Audio
 			WakeupEvent->Trigger();
 		}
 		
-		CallbackThread.Reset(FRunnableThread::Create(this, TEXT("AudioMixerNullCallbackThread"), 0, ThreadPriority, FPlatformAffinity::GetAudioThreadMask()));
+		CallbackThread.Reset(FRunnableThread::Create(this, TEXT("AudioMixerNullCallbackThread"), 0, ThreadPriority, FPlatformAffinity::GetAudioRenderThreadMask()));
 	}
 		
 	void FMixerNullCallback::Stop()
diff --git a/Engine/Source/Runtime/Core/Private/Async/TaskGraph.cpp b/Engine/Source/Runtime/Core/Private/Async/TaskGraph.cpp
--- a/Engine/Source/Runtime/Core/Private/Async/TaskGraph.cpp
+++ b/Engine/Source/Runtime/Core/Private/Async/TaskGraph.cpp
@@ -698,9 +698,7 @@ public:
 			}
 			// else StatName = none, we need to let the scope empty so that the render thread submits tasks in a timely manner. 
 		}
-PRAGMA_DISABLE_DEPRECATION_WARNINGS
-		else if (ThreadId != ENamedThreads::StatsThread)
-PRAGMA_ENABLE_DEPRECATION_WARNINGS
+		else
 		{
 			StatName = GET_STATID(STAT_TaskGraph_OtherTasks);
 			StallStatId = GET_STATID(STAT_TaskGraph_OtherStalls);
@@ -851,20 +849,6 @@ private:
 		{
 			return TEXT("RHI Thread");
 		}
-PRAGMA_DISABLE_DEPRECATION_WARNINGS
-		else if (ThreadId == ENamedThreads::AudioThread)
-PRAGMA_ENABLE_DEPRECATION_WARNINGS
-		{
-			return TEXT("Audio Thread");
-		}
-#if STATS
-PRAGMA_DISABLE_DEPRECATION_WARNINGS
-		else if (ThreadId == ENamedThreads::StatsThread)
-PRAGMA_ENABLE_DEPRECATION_WARNINGS
-		{
-			return TEXT("Stats Thread");
-		}
-#endif
 		else
 		{
 			return TEXT("Unknown Named Thread");
@@ -1413,52 +1397,6 @@ public:
 		FPlatformTLS::FreeTlsSlot(PerThreadIDTLSSlot);
 	}
 
-	// StatsThread was removed and replaced by a pipe. While all engine code was converted where's a tiny possibilty that 
-	// an external code is using it directly. Redirect all stats tasks to the pipe
-	static void RedirectStatsTasksToPipe(FBaseGraphTask* Task)
-	{
-#if STATS
-PRAGMA_DISABLE_DEPRECATION_WARNINGS
-		if (ENamedThreads::GetThreadIndex(Task->GetThreadToExecuteOn()) == ENamedThreads::StatsThread)
-PRAGMA_ENABLE_DEPRECATION_WARNINGS
-		{
-			checkf(ENamedThreads::GetQueueIndex(Task->GetThreadToExecuteOn()) == ENamedThreads::MainQueue, TEXT("`StatsThread_Local` is not supported and unlikely to be used intentionally. Check that `Task->ThreadToExecuteOn` value is not rubbish: %d"), Task->GetThreadToExecuteOn());
-
-			extern CORE_API UE::Tasks::FPipe GStatsPipe;
-			GStatsPipe.Launch(UE_SOURCE_LOCATION,
-				[Task]
-				{
-					TRACE_CPUPROFILER_EVENT_SCOPE(StatsPipeWork);
-					TArray<FBaseGraphTask*> Dummy;
-					Task->Execute(Dummy, ENamedThreads::AnyThread, true);
-				}
-			);
-			return;
-		}
-#endif
-	}
-
-	// AudioThread was removed and replaced by a pipe. While all engine code was converted where's a tiny possibilty that 
-	// an external code is using it directly. Redirect all audio tasks to the pipe
-	static void RedirectAudioTasksToPipe(FBaseGraphTask* Task)
-	{
-PRAGMA_DISABLE_DEPRECATION_WARNINGS
-		if (ENamedThreads::GetThreadIndex(Task->GetThreadToExecuteOn()) == ENamedThreads::AudioThread)
-PRAGMA_ENABLE_DEPRECATION_WARNINGS
-		{
-			extern CORE_API UE::Tasks::FPipe GAudioPipe;
-			GAudioPipe.Launch(UE_SOURCE_LOCATION,
-				[Task]
-				{
-					TRACE_CPUPROFILER_EVENT_SCOPE(AudioPipeWork);
-					TArray<FBaseGraphTask*> Dummy;
-					Task->Execute(Dummy, ENamedThreads::AnyThread, true);
-				}
-			);
-			return;
-		}
-	}
-
 	// API inherited from FTaskGraphInterface
 
 	/** 
@@ -1471,9 +1409,6 @@ PRAGMA_ENABLE_DEPRECATION_WARNINGS
 	{
 		TASKGRAPH_SCOPE_CYCLE_COUNTER(2, STAT_TaskGraph_QueueTask);
 
-		RedirectStatsTasksToPipe(Task);
-		RedirectAudioTasksToPipe(Task);
-
 		if (ENamedThreads::GetThreadIndex(ThreadToExecuteOn) == ENamedThreads::AnyThread)
 		{
 			TASKGRAPH_SCOPE_CYCLE_COUNTER(3, STAT_TaskGraph_QueueTask_AnyThread);
@@ -2016,9 +1951,6 @@ public:
 private:
 	void QueueTask(class FBaseGraphTask* Task, bool bWakeUpWorker, ENamedThreads::Type InThreadToExecuteOn, ENamedThreads::Type InCurrentThreadIfKnown) override
 	{
-		FTaskGraphImplementation::RedirectStatsTasksToPipe(Task);
-		FTaskGraphImplementation::RedirectAudioTasksToPipe(Task);
-
 		if (ENamedThreads::GetThreadIndex(InThreadToExecuteOn) == ENamedThreads::AnyThread)
 		{
 #if TASKGRAPH_NEW_FRONTEND
@@ -2767,14 +2699,6 @@ void FTaskGraphInterface::BroadcastSlow_OnlyUseForSpecialPurposes(bool bDoTaskTh
 		}
 	}
 
-#if STATS
-PRAGMA_DISABLE_DEPRECATION_WARNINGS
-	if (FTaskGraphInterface::Get().IsThreadProcessingTasks(ENamedThreads::StatsThread))
-	{
-		TGraphTask<FBroadcastTask>::CreateTask().ConstructAndDispatchWhenReady(Callback, StartTime, TEXT("Stats"), ENamedThreads::SetTaskPriority(ENamedThreads::StatsThread, ENamedThreads::HighTaskPriority), nullptr, nullptr, nullptr);
-	}
-PRAGMA_ENABLE_DEPRECATION_WARNINGS
-#endif
 	if (IsRHIThreadRunning())
 	{
 		TGraphTask<FBroadcastTask>::CreateTask().ConstructAndDispatchWhenReady(Callback, StartTime, TEXT("RHIT"), ENamedThreads::SetTaskPriority(ENamedThreads::RHIThread, ENamedThreads::HighTaskPriority), nullptr, nullptr, nullptr);
@@ -2784,12 +2708,6 @@ PRAGMA_ENABLE_DEPRECATION_WARNINGS
 	{
 		TGraphTask<FBroadcastTask>::CreateTask().ConstructAndDispatchWhenReady(Callback, StartTime, TEXT("RT"), ENamedThreads::SetTaskPriority(RenderThread, ENamedThreads::HighTaskPriority), nullptr, nullptr, nullptr);
 	}
-PRAGMA_DISABLE_DEPRECATION_WARNINGS
-	if (FTaskGraphInterface::Get().IsThreadProcessingTasks(ENamedThreads::AudioThread))
-	{
-		TGraphTask<FBroadcastTask>::CreateTask().ConstructAndDispatchWhenReady(Callback, StartTime, TEXT("AudioT"), ENamedThreads::SetTaskPriority(ENamedThreads::AudioThread, ENamedThreads::HighTaskPriority), nullptr, nullptr, nullptr);
-	}
-PRAGMA_ENABLE_DEPRECATION_WARNINGS
 
 	for (FEvent* TaskEvent : TaskEvents)
 	{
diff --git a/Engine/Source/Runtime/Core/Private/Stats/Stats2.cpp b/Engine/Source/Runtime/Core/Private/Stats/Stats2.cpp
--- a/Engine/Source/Runtime/Core/Private/Stats/Stats2.cpp
+++ b/Engine/Source/Runtime/Core/Private/Stats/Stats2.cpp
@@ -884,8 +884,6 @@ bool FStatNameAndInfo::GetSortByNameFrom(FName InLongName)
 
 static TAutoConsoleVariable<int32> CVarDumpStatPackets(	TEXT("DumpStatPackets"),0,	TEXT("If true, dump stat packets."));
 
-#if UE_STATS_THREAD_AS_PIPE
-
 /** The rendering thread runnable object. */
 class FStatsThread
 {
@@ -1004,208 +1002,6 @@ private:
 	}
 };
 
-#else // UE_STATS_THREAD_AS_PIPE
-
-/** The rendering thread runnable object. */
-class FStatsThread : public FRunnable, FSingleThreadRunnable
-{
-	/** Array of stat packets, queued data to be processed on this thread. */
-	FStatPacketArray IncomingData;
-
-	/** Stats state. */
-	FStatsThreadState& State;
-
-	/** Whether we are ready to process the packets, sets by game or render packets. */
-	bool bReadyToProcess;
-public:
-
-	/** Default constructor. */
-	FStatsThread()
-		: State(FStatsThreadState::GetLocalState())
-		, bReadyToProcess(false)
-	{
-		check(IsInGameThread());
-	}
-
-	/**
-	 * Returns a pointer to the single threaded interface when multithreading is disabled.
-	 */
-	virtual FSingleThreadRunnable* GetSingleThreadInterface() override
-	{
-		return this;
-	}
-
-	/** Attaches to the task graph stats thread, all processing will be handled by the task graph. */
-	virtual uint32 Run() override
-	{
-		FThreadStats::GetThreadStats()->bIsStatsThread = true;
-		FMemory::SetupTLSCachesOnCurrentThread();
-		FTaskGraphInterface::Get().AttachToThread(ENamedThreads::StatsThread);
-		FTaskGraphInterface::Get().ProcessThreadUntilRequestReturn(ENamedThreads::StatsThread);
-		FMemory::ClearAndDisableTLSCachesOnCurrentThread();
-		return 0;
-	}
-
-	/** Tick function. */
-	virtual void Tick() override
-	{
-		LLM_SCOPE(ELLMTag::Stats);
-
-		static double LastTime = -1.0;
-		bool bShouldProcess = false;
-
-		const int32 MaxIncomingPackets = 16;
-		if (FThreadStats::bIsRawStatsActive)
-		{
-			// For raw stats we process every 24MB of packet data to minimize the stats messages memory usage.
-			//const bool bShouldProcessRawStats = IncomingData.Packets.Num() > 10;
-			const int32 MaxIncomingMessages = 24 * 1024 * 1024 / sizeof(FStatMessage);
-
-			int32 IncomingDataMessages = 0;
-			for (FStatPacket* Packet : IncomingData.Packets)
-			{
-				IncomingDataMessages += Packet->StatMessages.Num();
-			}
-
-			bShouldProcess = IncomingDataMessages > MaxIncomingMessages || IncomingData.Packets.Num() > MaxIncomingPackets;
-		}
-		else
-		{
-			// For regular stats we won't process more than every 5ms or every 16 packets.
-			// Commandlet stats are flushed as soon as.
-			bShouldProcess = bReadyToProcess && (FPlatformTime::Seconds() - LastTime > 0.005f || IncomingData.Packets.Num() > MaxIncomingPackets || FStats::EnabledForCommandlet());
-		}
-
-		if (bShouldProcess)
-		{
-			SCOPE_CYCLE_COUNTER(STAT_StatsNewTick);
-
-			IStatGroupEnableManager::Get().UpdateMemoryUsage();
-			State.UpdateStatMessagesMemoryUsage();
-
-			bReadyToProcess = false;
-			FStatPacketArray NowData;
-			Exchange(NowData.Packets, IncomingData.Packets);
-			INC_DWORD_STAT_BY(STAT_StatFramePacketsRecv, NowData.Packets.Num());
-			{
-				SCOPE_CYCLE_COUNTER(STAT_StatsNewParseMeta);
-				TArray<FStatMessage> MetaMessages;
-				{
-					FScopeLock Lock(&FStartupMessages::Get().CriticalSection);
-					Exchange(FStartupMessages::Get().DelayedMessages, MetaMessages);
-				}
-				if (MetaMessages.Num())
-				{
-					State.ProcessMetaDataOnly(MetaMessages);
-				}
-			}
-			{
-				SCOPE_CYCLE_COUNTER(STAT_ScanForAdvance);
-				State.ScanForAdvance(NowData);
-			}
-
-			if (FThreadStats::bIsRawStatsActive)
-			{
-				// Process raw stats.
-				State.ProcessRawStats(NowData);
-				State.ResetRegularStats();
-			}
-			else
-			{
-				// Process regular stats.
-				SCOPE_CYCLE_COUNTER(STAT_StatsNewAddToHistory);
-				State.ResetRawStats();
-				State.AddToHistoryAndEmpty(NowData);
-			}
-			check(!NowData.Packets.Num());
-			LastTime = FPlatformTime::Seconds();
-		}
-	}
-
-	/** Accesses singleton. */
-	static FStatsThread& Get()
-	{
-		static FStatsThread Singleton;
-		return Singleton;
-	}
-
-	/** Received a stat packet from other thread and add to the processing queue. */
-	void StatMessage(FStatPacket* Packet)
-	{
-		LLM_SCOPE(ELLMTag::Stats);
-
-		if (CVarDumpStatPackets.GetValueOnAnyThread())
-		{
-			UE_LOG(LogStats, Log, TEXT("Packet from %x with %d messages"), Packet->ThreadId, Packet->StatMessages.Num());
-		}
-
-		bReadyToProcess = Packet->ThreadType != EThreadType::Other;
-		IncomingData.Packets.Add(Packet);
-		State.NumStatMessages.Add(Packet->StatMessages.Num());
-
-		Tick();
-	}
-
-	void SelfStatMessage(FStatPacket* Packet)
-	{
-		if (CVarDumpStatPackets.GetValueOnAnyThread())
-		{
-			UE_LOG(LogStats, Log, TEXT("Self Packet from %x with %d messages"), Packet->ThreadId, Packet->StatMessages.Num());
-		}
-
-		IncomingData.Packets.Add(Packet);
-		State.NumStatMessages.Add(Packet->StatMessages.Num());
-	}
-
-	/** Start a stats runnable thread. */
-	void Start()
-	{
-		Thread = FRunnableThread::Create(this, TEXT("StatsThread"), 512 * 1024, TPri_BelowNormal, FPlatformAffinity::GetStatsThreadMask());
-		check(Thread);
-	}
-
-	/** Ends the stats runnable thread. */
-	void End()
-	{
-		delete Thread;
-		Thread = nullptr;
-	}
-
-private:
-	FRunnableThread* Thread;
-};
-
-/*-----------------------------------------------------------------------------
-	FStatMessagesTask
------------------------------------------------------------------------------*/
-
-// not using a delegate here to allow higher performance since we may end up sending a lot of small message arrays to the thread.
-class FStatMessagesTask
-{
-	FStatPacket* Packet;
-public:
-	FStatMessagesTask(FStatPacket* InPacket)
-		: Packet(InPacket)
-	{
-	}
-	FORCEINLINE TStatId GetStatId() const
-	{
-		return TStatId(); // we don't want to record this or it spams the stat system; we cover this time when we tick the stats system
-	}
-	static ENamedThreads::Type GetDesiredThread()
-	{
-		return FPlatformProcess::SupportsMultithreading() ? ENamedThreads::StatsThread : ENamedThreads::GameThread;
-	}
-	static ESubsequentsMode::Type GetSubsequentsMode() { return ESubsequentsMode::TrackSubsequents; }
-
-	void DoTask(ENamedThreads::Type CurrentThread, const FGraphEventRef& MyCompletionGraphEvent)
-	{
-		FStatsThread::Get().StatMessage(Packet);
-		Packet = NULL;
-	}
-};
-
-#endif // UE_STATS_THREAD_AS_PIPE
 /*-----------------------------------------------------------------------------
 	FThreadStatsPool
 -----------------------------------------------------------------------------*/
@@ -1264,9 +1060,6 @@ FThreadStats::FThreadStats():
 	MemoryMessageScope(0),
 	bReentranceGuard(false),
 	bSawExplicitFlush(false)
-#if !UE_STATS_THREAD_AS_PIPE
-	, bIsStatsThread(false)
-#endif
 {
 	Packet.SetThreadProperties();
 
@@ -1281,9 +1074,6 @@ FThreadStats::FThreadStats( EConstructor ):
 	MemoryMessageScope(0),
 	bReentranceGuard(false),
 	bSawExplicitFlush(false)
-#if !UE_STATS_THREAD_AS_PIPE
-	, bIsStatsThread(false)
-#endif
 {}
 
 void FThreadStats::CheckEnable()
@@ -1388,19 +1178,7 @@ void FThreadStats::FlushRegularStats( bool bHasBrokenCallstacks, bool bForceFlus
 			Packet.StatMessages.Empty(MaxPresize);
 		}
 
-#if UE_STATS_THREAD_AS_PIPE
 		SendMessage_Async(ToSend);
-#else
-		if (bIsStatsThread)
-		{
-			FStatsThread::Get().SelfStatMessage(ToSend);
-		}
-		else
-		{
-			TGraphTask<FStatMessagesTask>::CreateTask().ConstructAndDispatchWhenReady(ToSend);
-		}
-#endif
-
 		UpdateExplicitFlush();
 	}
 }
@@ -1438,19 +1216,7 @@ void FThreadStats::FlushRawStats( bool bHasBrokenCallstacks /*= false*/, bool bF
 
 		check(!Packet.StatMessages.Num());
 
-#if UE_STATS_THREAD_AS_PIPE
 		SendMessage_Async(ToSend);
-#else
-		if (bIsStatsThread)
-		{
-			FStatsThread::Get().SelfStatMessage(ToSend);
-		}
-		else
-		{
-			TGraphTask<FStatMessagesTask>::CreateTask().ConstructAndDispatchWhenReady(ToSend);
-		}
-#endif
-
 		UpdateExplicitFlush();
 
 		const float NumMessagesAsMB = float(NumMessages * sizeof(FStatMessage)) / 1024.0f / 1024.0f;
@@ -1560,11 +1326,6 @@ void FThreadStats::StartThread()
 		check(TlsSlot);
 	}
 
-#if !UE_STATS_THREAD_AS_PIPE
-	FStatsThread::Get();
-	FStatsThread::Get().Start();
-#endif
-
 	check(IsThreadingReady());
 	CheckEnable();
 
@@ -1581,9 +1342,7 @@ void FThreadStats::StartThread()
 
 static int32 CurrentEventIndex = 0;
 
-#if UE_STATS_THREAD_AS_PIPE
 static UE::Tasks::FTask LastFramesBarriers[MAX_STAT_LAG];
-#endif
 
 static FGraphEventRef LastFramesEvents[MAX_STAT_LAG];
 
@@ -1606,7 +1365,6 @@ void FThreadStats::StopThread()
 		WaitForStats();
 	}
 
-#if UE_STATS_THREAD_AS_PIPE
 	// wait for the pipe to complete all currently piped tasks
 	while (GStatsPipe.HasWork())
 	{
@@ -1616,16 +1374,6 @@ void FThreadStats::StopThread()
 	{
 		LastFramesBarriers[Index] = {}; // release memory before the allocator is destroyed
 	}
-#else
-	for (int32 Index = 0; Index < MAX_STAT_LAG; Index++)
-	{
-		LastFramesEvents[Index] = NULL;
-	}
-	FGraphEventRef QuitTask = TGraphTask<FReturnGraphTask>::CreateTask(NULL, ENamedThreads::GameThread).ConstructAndDispatchWhenReady(FPlatformProcess::SupportsMultithreading() ? ENamedThreads::StatsThread : ENamedThreads::GameThread);
-	FTaskGraphInterface::Get().WaitUntilTaskCompletes(QuitTask, ENamedThreads::GameThread_Local);
-
-	FStatsThread::Get().End();
-#endif
 }
 
 void FThreadStats::WaitForStats()
@@ -1642,7 +1390,6 @@ void FThreadStats::WaitForStats()
 
 		int32 EventIndex = (CurrentEventIndex + MAX_STAT_LAG - 1) % MAX_STAT_LAG;
 
-#if UE_STATS_THREAD_AS_PIPE
 		if (FPlatformProcess::SupportsMultithreading())
 		{
 			SCOPE_CYCLE_COUNTER(STAT_WaitForStats);
@@ -1650,7 +1397,6 @@ void FThreadStats::WaitForStats()
 			LastFramesBarriers[EventIndex] = GStatsPipe.Launch(UE_SOURCE_LOCATION, [] {});
 		}
 		else
-#endif
 		{
 			{
 				SCOPE_CYCLE_COUNTER(STAT_WaitForStats);
diff --git a/Engine/Source/Runtime/Core/Private/Stats/StatsCommand.cpp b/Engine/Source/Runtime/Core/Private/Stats/StatsCommand.cpp
--- a/Engine/Source/Runtime/Core/Private/Stats/StatsCommand.cpp
+++ b/Engine/Source/Runtime/Core/Private/Stats/StatsCommand.cpp
@@ -2370,30 +2370,12 @@ bool DirectStatsCommand(const TCHAR* Cmd, bool bBlockForCompletion /*= false*/,
 				STAT_FSimpleDelegateGraphTask_StatCmd,
 				STATGROUP_TaskGraphTasks);
 
-#if UE_STATS_THREAD_AS_PIPE
 			UE::Tasks::FTask Task = GStatsPipe.Launch(UE_SOURCE_LOCATION, [FullCmd, bStatCommand, Ar] { StatCmd(FullCmd, bStatCommand, Ar); });
 			if (bBlockForCompletion)
 			{
 				Task.Wait();
 				GLog->FlushThreadedLogs();
 			}
-#else // UE_STATS_THREAD_AS_PIPE
-			ENamedThreads::Type ThreadType = ENamedThreads::GameThread;
-			if (FPlatformProcess::SupportsMultithreading())
-			{
-				ThreadType = ENamedThreads::StatsThread;
-			}
-
-			FGraphEventRef CompleteHandle = FSimpleDelegateGraphTask::CreateAndDispatchWhenReady(
-				FSimpleDelegateGraphTask::FDelegate::CreateStatic(&StatCmd, FullCmd, bStatCommand, Ar),
-				GET_STATID(STAT_FSimpleDelegateGraphTask_StatCmd), NULL, ThreadType
-			);
-			if (bBlockForCompletion && FPlatformProcess::SupportsMultithreading())
-			{
-				FTaskGraphInterface::Get().WaitUntilTaskCompletes(CompleteHandle);
-				GLog->FlushThreadedLogs();
-			}
-#endif // UE_STATS_THREAD_AS_PIPE
 #else // STATS
 			// If stats aren't enabled, broadcast so engine stats can still be triggered
 			StatCmd(FullCmd, bStatCommand, Ar);
@@ -2423,18 +2405,9 @@ static void GetPermanentStats_StatsThread(TArray<FStatMessage>* OutStats)
 
 void GetPermanentStats(TArray<FStatMessage>& OutStats)
 {
-#if UE_STATS_THREAD_AS_PIPE
 	GStatsPipe
 		.Launch(UE_SOURCE_LOCATION, [&OutStats] { GetPermanentStats_StatsThread(&OutStats); })
 		.Wait();
-#else
-	FGraphEventRef CompleteHandle = FSimpleDelegateGraphTask::CreateAndDispatchWhenReady(
-		FSimpleDelegateGraphTask::FDelegate::CreateStatic(&GetPermanentStats_StatsThread, &OutStats),
-		TStatId{}, NULL,
-		FPlatformProcess::SupportsMultithreading() ? ENamedThreads::StatsThread : ENamedThreads::GameThread
-	);
-	FTaskGraphInterface::Get().WaitUntilTaskCompletes(CompleteHandle);
-#endif
 }
 
 #endif
diff --git a/Engine/Source/Runtime/Core/Private/Tests/Async/TaskGraphTest.cpp b/Engine/Source/Runtime/Core/Private/Tests/Async/TaskGraphTest.cpp
--- a/Engine/Source/Runtime/Core/Private/Tests/Async/TaskGraphTest.cpp
+++ b/Engine/Source/Runtime/Core/Private/Tests/Async/TaskGraphTest.cpp
@@ -897,16 +897,6 @@ namespace TaskGraphTests
 
 		// named threads and local queues
 
-#if STATS
-		{	// StatsThread
-			bool bExecuted = false;
-PRAGMA_DISABLE_DEPRECATION_WARNINGS
-			FFunctionGraphTask::CreateAndDispatchWhenReady([&bExecuted] { bExecuted = true; }, TStatId{}, nullptr, ENamedThreads::StatsThread)->Wait();
-PRAGMA_ENABLE_DEPRECATION_WARNINGS
-			check(bExecuted);
-		}
-#endif
-
 		if (IsRHIThreadRunning())
 		{	// RHIThread
 			bool bExecuted = false;
@@ -914,15 +904,6 @@ PRAGMA_ENABLE_DEPRECATION_WARNINGS
 			check(bExecuted);
 		}
 
-		if (IsAudioThreadRunning())
-		{	// AudioThread
-			bool bExecuted = false;
-PRAGMA_DISABLE_DEPRECATION_WARNINGS
-			FFunctionGraphTask::CreateAndDispatchWhenReady([&bExecuted] { bExecuted = true; }, TStatId{}, nullptr, ENamedThreads::AudioThread)->Wait();
-PRAGMA_ENABLE_DEPRECATION_WARNINGS
-			check(bExecuted);
-		}
-
 		{	// GameThread
 			bool bExecuted = false;
 			FFunctionGraphTask::CreateAndDispatchWhenReady([&bExecuted] { bExecuted = true; }, TStatId{}, nullptr, ENamedThreads::GameThread);
@@ -1337,29 +1318,6 @@ PRAGMA_ENABLE_DEPRECATION_WARNINGS
 		return true;
 	}
 
-	IMPLEMENT_SIMPLE_AUTOMATION_TEST(FTaskGraphStatsThreadRedirectionTest, "System.Core.Async.TaskGraph.StatsThreadRedirection", EAutomationTestFlags::ApplicationContextMask | EAutomationTestFlags::EngineFilter | EAutomationTestFlags::Disabled);
-
-	bool FTaskGraphStatsThreadRedirectionTest::RunTest(const FString& Parameters)
-	{
-		if (!FPlatformProcess::SupportsMultithreading())
-		{
-			// no StatsThread around
-			return true;
-		}
-
-#if STATS
-
-		bool bExecuted = false;
-PRAGMA_DISABLE_DEPRECATION_WARNINGS
-		FFunctionGraphTask::CreateAndDispatchWhenReady([&bExecuted] { bExecuted = true; }, TStatId{}, nullptr, ENamedThreads::StatsThread)->Wait();
-PRAGMA_ENABLE_DEPRECATION_WARNINGS
-		check(bExecuted);
-
-#endif
-
-		return true;
-	}
-
 	template<uint32 Num>
 	void OversubscriptionStressTest()
 	{
diff --git a/Engine/Source/Runtime/Core/Public/Android/AndroidPlatformAffinity.h b/Engine/Source/Runtime/Core/Public/Android/AndroidPlatformAffinity.h
--- a/Engine/Source/Runtime/Core/Public/Android/AndroidPlatformAffinity.h
+++ b/Engine/Source/Runtime/Core/Public/Android/AndroidPlatformAffinity.h
@@ -49,12 +49,7 @@ public:
 		return AllCores;
 	}
 
-	static const CORE_API uint64 GetStatsThreadMask()
-	{
-		return GetLittleCoreMask();
-	}
-
-	static const CORE_API uint64 GetAudioThreadMask()
+	static const CORE_API uint64 GetAudioRenderThreadMask()
 	{
 		return GetLittleCoreMask();
 	}
diff --git a/Engine/Source/Runtime/Core/Public/Async/TaskGraphInterfaces.h b/Engine/Source/Runtime/Core/Public/Async/TaskGraphInterfaces.h
--- a/Engine/Source/Runtime/Core/Public/Async/TaskGraphInterfaces.h
+++ b/Engine/Source/Runtime/Core/Public/Async/TaskGraphInterfaces.h
@@ -55,15 +55,7 @@ namespace ENamedThreads
 	{
 		UnusedAnchor = -1,
 		/** The always-present, named threads are listed next **/
-#if STATS
-#if UE_STATS_THREAD_AS_PIPE
-		StatsThread UE_DEPRECATED(5.0, "`StatsThread` has been removed. Stats system should be used by Stats public API"),
-#else
-		StatsThread,
-#endif
-#endif
 		RHIThread,
-		AudioThread UE_DEPRECATED(5.0, "`AudioThread` has been removed. Please use `FAudioThread` API"),
 		GameThread,
 		// The render thread is sometimes the game thread and is sometimes the actual rendering thread
 		ActualRenderingThread = GameThread + 1,
@@ -100,15 +92,6 @@ namespace ENamedThreads
 		ThreadPriorityShift = 10,
 
 		/** Combinations **/
-#if STATS
-#if UE_STATS_THREAD_AS_PIPE
-PRAGMA_DISABLE_DEPRECATION_WARNINGS
-		StatsThread_Local UE_DEPRECATED(5.0, "`StatsThread_Local` has been removed (and has never been supported). Stats system should be used by its public API in `Stats2.h`") = StatsThread | LocalQueue,
-PRAGMA_ENABLE_DEPRECATION_WARNINGS
-#else
-		StatsThread_Local = StatsThread | LocalQueue,
-#endif
-#endif
 		GameThread_Local = GameThread | LocalQueue,
 		ActualRenderingThread_Local = ActualRenderingThread | LocalQueue,
 
diff --git a/Engine/Source/Runtime/Core/Public/GenericPlatform/GenericPlatformAffinity.h b/Engine/Source/Runtime/Core/Public/GenericPlatform/GenericPlatformAffinity.h
--- a/Engine/Source/Runtime/Core/Public/GenericPlatform/GenericPlatformAffinity.h
+++ b/Engine/Source/Runtime/Core/Public/GenericPlatform/GenericPlatformAffinity.h
@@ -80,15 +80,7 @@ public:
 		return 0xFFFFFFFFFFFFFFFF;
 	}
 
-#if UE_STATS_THREAD_AS_PIPE
-	UE_DEPRECATED(5.0, "This parameter is obsolete as StatsThread was removed and the affinity mask has no effect anymore. Please remove the call.")
-#endif
-	static const CORE_API uint64 GetStatsThreadMask()
-	{
-		return 0xFFFFFFFFFFFFFFFF;
-	}
-
-	static const CORE_API uint64 GetAudioThreadMask()
+	static const CORE_API uint64 GetAudioRenderThreadMask()
 	{
 		return 0xFFFFFFFFFFFFFFFF;
 	}
diff --git a/Engine/Source/Runtime/Core/Public/Misc/CoreDefines.h b/Engine/Source/Runtime/Core/Public/Misc/CoreDefines.h
--- a/Engine/Source/Runtime/Core/Public/Misc/CoreDefines.h
+++ b/Engine/Source/Runtime/Core/Public/Misc/CoreDefines.h
@@ -65,8 +65,3 @@ Can be overridden by setting to 1 or 0 in the project's .Target.cs files.
 #ifndef NO_CVARS
 	#define NO_CVARS 0
 #endif
-
-#if !defined(UE_STATS_THREAD_AS_PIPE)
-	#define UE_STATS_THREAD_AS_PIPE 1
-#endif
-
diff --git a/Engine/Source/Runtime/Core/Public/Stats/Stats2.h b/Engine/Source/Runtime/Core/Public/Stats/Stats2.h
--- a/Engine/Source/Runtime/Core/Public/Stats/Stats2.h
+++ b/Engine/Source/Runtime/Core/Public/Stats/Stats2.h
@@ -1338,11 +1338,6 @@ class FThreadStats : FNoncopyable
 	/** Tracks current stack depth for cycle counters. **/
 	bool bSawExplicitFlush;
 
-#if !UE_STATS_THREAD_AS_PIPE
-	/** True if this is the stats thread, which needs special handling. **/
-	bool bIsStatsThread;
-#endif
-
 	/** Gathers information about the current thread and sets up the TLS value. **/
 	CORE_API FThreadStats();
 
diff --git a/Engine/Source/Runtime/Windows/AudioMixerXAudio2/Private/AudioMixerPlatformXAudio2.cpp b/Engine/Source/Runtime/Windows/AudioMixerXAudio2/Private/AudioMixerPlatformXAudio2.cpp
--- a/Engine/Source/Runtime/Windows/AudioMixerXAudio2/Private/AudioMixerPlatformXAudio2.cpp
+++ b/Engine/Source/Runtime/Windows/AudioMixerXAudio2/Private/AudioMixerPlatformXAudio2.cpp
@@ -77,7 +77,7 @@ THIRD_PARTY_INCLUDES_END
 
 static XAUDIO2_PROCESSOR GetXAudio2ProcessorsToUse()
 {
-	XAUDIO2_PROCESSOR ProcessorsToUse = (XAUDIO2_PROCESSOR)FPlatformAffinity::GetAudioThreadMask();
+	XAUDIO2_PROCESSOR ProcessorsToUse = (XAUDIO2_PROCESSOR)FPlatformAffinity::GetAudioRenderThreadMask();
 	// https://docs.microsoft.com/en-us/windows/win32/api/xaudio2/nf-xaudio2-xaudio2create
 	// Warning If you specify XAUDIO2_ANY_PROCESSOR, the system will use all of the device's processors and, as noted above, create a worker thread for each processor.
 	// We certainly don't want to use all available CPU. XAudio threads are time critical priority and wake up every 10 ms, they may cause lots of unwarranted context switches.
