system semaphore

diff --git a/Engine/Platforms/Hololens/Source/Runtime/Core/Public/HoloLensSemaphore.h b/Engine/Platforms/Hololens/Source/Runtime/Core/Public/HoloLensSemaphore.h
new file mode 100644
--- /dev/null
+++ b/Engine/Platforms/Hololens/Source/Runtime/Core/Public/HoloLensSemaphore.h
@@ -0,0 +1,5 @@
+// Copyright Epic Games, Inc. All Rights Reserved.
+
+#pragma once
+
+#include "Windows/WindowsSemaphore.h"
diff --git a/Engine/Source/Runtime/Core/Private/Tests/Async/SemaphoreTest.cpp b/Engine/Source/Runtime/Core/Private/Tests/Async/SemaphoreTest.cpp
new file mode 100644
--- /dev/null
+++ b/Engine/Source/Runtime/Core/Private/Tests/Async/SemaphoreTest.cpp
@@ -0,0 +1,172 @@
+// Copyright Epic Games, Inc. All Rights Reserved.
+
+#include "CoreTypes.h"
+#include "Misc/AutomationTest.h"
+#include "Tests/Benchmark.h"
+#include "HAL/Thread.h"
+#include "HAL/Semaphore.h"
+#include "HAL/Event.h"
+#include "HAL/PlatformProcess.h"
+#include "Containers/Array.h"
+#include "Templates/UniquePtr.h"
+
+#include <atomic>
+
+#if WITH_DEV_AUTOMATION_TESTS
+
+namespace UE::SemaphoreTests
+{
+	using namespace Tasks;
+
+	// benchmark waking up multiple threads by releasing a semaphore and compare with triggering multiple auto-reset events
+	IMPLEMENT_SIMPLE_AUTOMATION_TEST(FSemaphoreWakeUpPerfTest, "System.Core.Async.Semaphore.WakeUpPerf", EAutomationTestFlags::ApplicationContextMask | EAutomationTestFlags::EngineFilter);
+
+	template<uint32 NumThreads, uint32 NumLoops, bool ReleaseAllAtOnce>
+	void TestSemaphoreWakeUpPerf()
+	{
+		TArray<FThread> Threads;
+		Threads.Reserve(NumThreads);
+		std::atomic<bool> bQuit{ false };
+
+		FSemaphore Semaphore{ 0, NumThreads };
+		std::atomic<uint32> Sync1{ 0 };
+		std::atomic<uint32> Sync2{ 0 };
+
+		for (int i = 0; i != NumThreads; ++i)
+		{
+			Threads.Emplace(TEXT("SemaphoreWakeUpPerfTest"),
+				[&Semaphore, &Sync1, &Sync2, &bQuit]
+				{
+					while (!bQuit)
+					{
+						Semaphore.Acquire();
+						--Sync1;
+						while (Sync1 != 0)
+						{
+						}
+						--Sync2;
+					}
+				}
+				);
+		}
+
+		auto Iteration = [&Semaphore, &Sync1, &Sync2]
+		{
+			Sync2 = NumThreads;
+			Sync1 = NumThreads;
+
+			if constexpr (ReleaseAllAtOnce)
+			{
+				Semaphore.Release(NumThreads);
+			}
+			else
+			{
+				for (int i = 0; i != NumThreads; ++i)
+				{
+					Semaphore.Release();
+				}
+			}
+
+			while (Sync2 != 0)
+			{
+			}
+			FPlatformProcess::Sleep(0.0);
+		};
+
+		for (int i = 0; i != NumLoops; ++i)
+		{
+			Iteration();
+		}
+
+		Sync1 = NumThreads;
+		bQuit = true;
+
+		Semaphore.Release(NumThreads);
+		while (Sync1 != 0)
+		{
+		}
+
+		for (FThread& Thread : Threads)
+		{
+			Thread.Join();
+		}
+	}
+
+	template<uint32 NumThreads, uint32 NumLoops>
+	void TestMultiEventWakeUpPerf()
+	{
+		TArray<FThread> Threads;
+		Threads.Reserve(NumThreads);
+		std::atomic<bool> bQuit{ false };
+
+		FEventRef Events[NumThreads];
+
+		std::atomic<uint32> Sync1{ 0 };
+		std::atomic<uint32> Sync2{ 0 };
+
+		for (int i = 0; i != NumThreads; ++i)
+		{
+			Threads.Emplace(TEXT("MultiEventWakeUpPerfTest"),
+				[Event = &Events[i], &Sync1, &Sync2, &bQuit]()
+				{
+					while (!bQuit)
+					{
+						Event->Get()->Wait();
+						--Sync1;
+						while (Sync1 != 0)
+						{
+						}
+						--Sync2;
+					}
+				}
+			);
+		}
+
+		auto Iteration = [&Events, &Sync1, &Sync2]
+		{
+			Sync2 = NumThreads;
+			Sync1 = NumThreads;
+
+			for (int i = 0; i != NumThreads; ++i)
+			{
+				Events[i]->Trigger();
+			}
+
+			while (Sync2 != 0)
+			{
+			}
+			FPlatformProcess::Sleep(0.0);
+		};
+
+		for (uint32 i = 0; i != NumLoops; ++i)
+		{
+			Iteration();
+		}
+
+		Sync1 = NumThreads;
+		bQuit = true;
+
+		for (int i = 0; i != NumThreads; ++i)
+		{
+			Events[i]->Trigger();
+		}
+		while (Sync1 != 0)
+		{
+		}
+	
+		for (FThread& Thread : Threads)
+		{
+			Thread.Join();
+		}
+	}
+
+	bool FSemaphoreWakeUpPerfTest::RunTest(const FString& Parameters)
+	{
+		UE_BENCHMARK(5, TestSemaphoreWakeUpPerf<5, 100000, true>);
+		UE_BENCHMARK(5, TestSemaphoreWakeUpPerf<5, 100000, false>);
+		UE_BENCHMARK(5, TestMultiEventWakeUpPerf<5, 100000>);
+		return true;
+	}
+}
+
+#endif
diff --git a/Engine/Source/Runtime/Core/Public/Android/AndroidSemaphore.h b/Engine/Source/Runtime/Core/Public/Android/AndroidSemaphore.h
new file mode 100644
--- /dev/null
+++ b/Engine/Source/Runtime/Core/Public/Android/AndroidSemaphore.h
@@ -0,0 +1,5 @@
+// Copyright Epic Games, Inc. All Rights Reserved.
+
+#pragma once
+
+#include "HAL/PThreadSemaphore.h"
\ No newline at end of file
diff --git a/Engine/Source/Runtime/Core/Public/HAL/PThreadSemaphore.h b/Engine/Source/Runtime/Core/Public/HAL/PThreadSemaphore.h
new file mode 100644
--- /dev/null
+++ b/Engine/Source/Runtime/Core/Public/HAL/PThreadSemaphore.h
@@ -0,0 +1,62 @@
+// Copyright Epic Games, Inc. All Rights Reserved.
+
+#pragma once
+
+#include "CoreTypes.h"
+#include "Misc/Timespan.h"
+#include "Misc/AssertionMacros.h"
+#include <pthread.h>
+#include <errno.h>
+#include <semaphore.h>
+
+class FPThreadSemaphore
+{
+public:
+	UE_NONCOPYABLE(FPThreadSemaphore);
+
+	FPThreadSemaphore(int32 InitialCount, int32 MaxCount)
+	{
+		int Res = sem_init(&Semaphore, /*pshared = */ 0, InitialCount);
+		checkfSlow(Res == 0, TEXT("Failed to create semaphore (%d:%d): %d"), InitialCount, MaxCount, errno);
+	}
+
+	~FPThreadSemaphore()
+	{
+		int Res = sem_destroy(&Semaphore);
+		checkfSlow(Res == 0, TEXT("Error destroying semaphore: %d"), errno);
+	}
+
+	void Acquire()
+	{
+		int Res = sem_wait(&Semaphore);
+		checkfSlow(Res == 0, TEXT("Acquiring semaphore failed: %d"), errno);
+	}
+
+	bool TryAcquire(FTimespan Timeout = FTimespan::Zero())
+	{
+		timespec ts;
+		int Res = clock_gettime(CLOCK_REALTIME, &ts);
+		checkfSlow(Res == 0, TEXT("clock_gettime failed: %d"), errno);
+
+		ts.tv_sec += (uint64)Timeout.GetTotalSeconds();
+		ts.tv_nsec += Timeout.GetFractionNano();
+		Res = sem_timedwait(&Semaphore, &ts);
+		checkfSlow(Res == 0 || Res == ETIMEDOUT, TEXT("sem_timedwait failed: %d"), errno);
+		return Res == 0;
+	}
+
+	void Release(int32 Count = 1)
+	{
+		checkfSlow(Count > 0, TEXT("Releasing semaphore with Count = %d, that should be greater than 0"), Count);
+		for (int i = 0; i != Count; ++i)
+		{
+			int Res = sem_post(&Semaphore);
+			checkfSlow(Res == 0, TEXT("Releasing semaphore failed: %d"), errno);
+		}
+	}
+
+private:
+	sem_t Semaphore;
+};
+
+typedef FPThreadSemaphore FSemaphore;
\ No newline at end of file
diff --git a/Engine/Source/Runtime/Core/Public/HAL/Semaphore.h b/Engine/Source/Runtime/Core/Public/HAL/Semaphore.h
new file mode 100644
--- /dev/null
+++ b/Engine/Source/Runtime/Core/Public/HAL/Semaphore.h
@@ -0,0 +1,7 @@
+// Copyright Epic Games, Inc. All Rights Reserved.
+
+#pragma once
+
+#include "CoreTypes.h"
+
+#include COMPILED_PLATFORM_HEADER(Semaphore.h)
diff --git a/Engine/Source/Runtime/Core/Public/Linux/LinuxSemaphore.h b/Engine/Source/Runtime/Core/Public/Linux/LinuxSemaphore.h
new file mode 100644
--- /dev/null
+++ b/Engine/Source/Runtime/Core/Public/Linux/LinuxSemaphore.h
@@ -0,0 +1,5 @@
+// Copyright Epic Games, Inc. All Rights Reserved.
+
+#pragma once
+
+#include "HAL/PThreadSemaphore.h"
\ No newline at end of file
diff --git a/Engine/Source/Runtime/Core/Public/Tests/Benchmark.h b/Engine/Source/Runtime/Core/Public/Tests/Benchmark.h
--- a/Engine/Source/Runtime/Core/Public/Tests/Benchmark.h
+++ b/Engine/Source/Runtime/Core/Public/Tests/Benchmark.h
@@ -10,7 +10,7 @@
 template<uint32 NumRuns, typename TestT>
 void Benchmark(const TCHAR* TestName, TestT&& TestBody)
 {
-	UE_LOG(LogTemp, Display, TEXT("\n-------------------------------\n%s"), TestName);
+	UE_LOG(LogTemp, Log, TEXT("\n-------------------------------\n%s"), TestName);
 	double MinTime = TNumericLimits<double>::Max();
 	double TotalTime = 0;
 	for (uint32 RunNo = 0; RunNo != NumRuns; ++RunNo)
@@ -19,7 +19,7 @@ void Benchmark(const TCHAR* TestName, TestT&& TestBody)
 		TestBody();
 		Time = FPlatformTime::Seconds() - Time;
 
-		UE_LOG(LogTemp, Display, TEXT("#%d: %f secs"), RunNo, Time);
+		UE_LOG(LogTemp, Log, TEXT("#%d: %f secs"), RunNo, Time);
 
 		TotalTime += Time;
 		if (MinTime > Time)
@@ -27,7 +27,7 @@ void Benchmark(const TCHAR* TestName, TestT&& TestBody)
 			MinTime = Time;
 		}
 	}
-	UE_LOG(LogTemp, Display, TEXT("min: %f secs, avg: %f secs\n-------------------------------\n"), MinTime, TotalTime / NumRuns);
+	UE_LOG(LogTemp, Log, TEXT("min: %f secs, avg: %f secs\n-------------------------------\n"), MinTime, TotalTime / NumRuns);
 
 #if NO_LOGGING
 	printf("%s\nmin: %f secs, avg: %f secs\n-------------------------------\n\n", TCHAR_TO_ANSI(TestName), MinTime, TotalTime / NumRuns);
diff --git a/Engine/Source/Runtime/Core/Public/Windows/WindowsSemaphore.h b/Engine/Source/Runtime/Core/Public/Windows/WindowsSemaphore.h
new file mode 100644
--- /dev/null
+++ b/Engine/Source/Runtime/Core/Public/Windows/WindowsSemaphore.h
@@ -0,0 +1,51 @@
+#pragma once
+
+#include "CoreTypes.h"
+#include "Misc/Timespan.h"
+#include "Misc/AssertionMacros.h"
+#include "Windows/WindowsHWrapper.h"
+#include "Windows/AllowWindowsPlatformTypes.h"
+
+class FWindowsSemaphore
+{
+public:
+	UE_NONCOPYABLE(FWindowsSemaphore);
+
+	FWindowsSemaphore(int32 InitialCount, int32 MaxCount)
+		: Semaphore(CreateSemaphore(nullptr, InitialCount, MaxCount, nullptr))
+	{
+		checkfSlow(Semaphore, TEXT("CreateSemaphore failed: %d"), GetLastError());
+	}
+
+	~FWindowsSemaphore()
+	{
+		CloseHandle(Semaphore);
+	}
+
+	void Acquire()
+	{
+		HRESULT hRes = WaitForSingleObject(Semaphore, INFINITE);
+		checkfSlow(hRes == WAIT_OBJECT_0, TEXT("Acquiring semaphore failed: %d (%d)"), hRes, GetLastError());
+	}
+
+	bool TryAcquire(FTimespan Timeout = FTimespan::Zero())
+	{
+		HRESULT hRes = WaitForSingleObject(Semaphore, (DWORD)Timeout.GetTotalMilliseconds());
+		checkfSlow(hRes == WAIT_OBJECT_0 || hRes == WAIT_TIMEOUT, TEXT("Acquiring semaphore failed: %d (%d)"), hRes, GetLastError());
+		return hRes == WAIT_OBJECT_0;
+	}
+
+	void Release(int32 Count = 1)
+	{
+		checkfSlow(Count > 0, TEXT("Releasing semaphore with Count = %d, that should be greater than 0"), Count);
+		bool bRes = ReleaseSemaphore(Semaphore, Count, nullptr);
+		checkfSlow(bRes, TEXT("Releasing semaphore for %d failed: %d"), Count, GetLastError());
+	}
+
+private:
+	HANDLE Semaphore;
+};
+
+using FSemaphore = FWindowsSemaphore;
+
+#include "Windows/HideWindowsPlatformTypes.h"
