fixed copyright notice

diff --git a/Engine/Source/Runtime/Core/Public/Windows/WindowsSemaphore.h b/Engine/Source/Runtime/Core/Public/Windows/WindowsSemaphore.h
--- a/Engine/Source/Runtime/Core/Public/Windows/WindowsSemaphore.h
+++ b/Engine/Source/Runtime/Core/Public/Windows/WindowsSemaphore.h
@@ -1,3 +1,5 @@
+// Copyright Epic Games, Inc. All Rights Reserved.
+
 #pragma once
 
 #include "CoreTypes.h"
