added missing Mac version of semaphore. as POSIX semaphore is deprecated on Mac, rolled out a custom version based on Dispatch API

diff --git a/Engine/Source/Runtime/Core/Public/Mac/MacSemaphore.h b/Engine/Source/Runtime/Core/Public/Mac/MacSemaphore.h
new file mode 100644
--- /dev/null
+++ b/Engine/Source/Runtime/Core/Public/Mac/MacSemaphore.h
@@ -0,0 +1,48 @@
+// Copyright Epic Games, Inc. All Rights Reserved.
+
+#pragma once
+
+#include "CoreTypes.h"
+#include "Misc/Timespan.h"
+#include "Misc/AssertionMacros.h"
+
+#include <dispatch/dispatch.h>
+
+class FMacSemaphore
+{
+public:
+	UE_NONCOPYABLE(FMacSemaphore);
+
+	FMacSemaphore(int32 InitialCount, int32 /*MaxCount*/)
+		: Semaphore(dispatch_semaphore_create(InitialCount))
+	{
+		checkfSlow(InitialCount >= 0, TEXT("Semaphore's initial count must be non negative value: %d"), InitialCount);
+	}
+
+	void Acquire()
+	{
+		intptr_t Res = dispatch_semaphore_wait(Semaphore, DISPATCH_TIME_FOREVER);
+		checkfSlow(Res == 0, TEXT("Acquiring semaphore failed"));
+	}
+
+	bool TryAcquire(FTimespan Timeout = FTimespan::Zero())
+	{
+		dispatch_time_t TS = dispatch_time(DISPATCH_TIME_NOW, (int64)Timeout.GetTotalMicroseconds() * 1000);
+		intptr_t Res = dispatch_semaphore_wait(Semaphore, TS);
+		return Res == 0;
+	}
+
+	void Release(int32 Count = 1)
+	{
+		checkfSlow(Count > 0, TEXT("Releasing semaphore with Count = %d, that should be greater than 0"), Count);
+		for (int i = 0; i != Count; ++i)
+		{
+			dispatch_semaphore_signal(Semaphore);
+		}
+	}
+
+private:
+	dispatch_semaphore_t Semaphore;
+};
+
+typedef FMacSemaphore FSemaphore;
\ No newline at end of file
