Fix for breaking backward compatibility of replays.

diff --git a/Engine/Source/Runtime/Core/Public/Misc/NetworkVersion.h b/Engine/Source/Runtime/Core/Public/Misc/NetworkVersion.h
--- a/Engine/Source/Runtime/Core/Public/Misc/NetworkVersion.h
+++ b/Engine/Source/Runtime/Core/Public/Misc/NetworkVersion.h
@@ -43,35 +43,36 @@ ENUM_CLASS_FLAGS(EEngineNetworkRuntimeFeatures);
 enum EEngineNetworkVersionHistory
 {
 	HISTORY_INITIAL = 1,
-	HISTORY_REPLAY_BACKWARDS_COMPAT = 2,			// Bump version to get rid of older replays before backwards compat was turned on officially
-	HISTORY_MAX_ACTOR_CHANNELS_CUSTOMIZATION = 3,	// Bump version because serialization of the actor channels changed
-	HISTORY_REPCMD_CHECKSUM_REMOVE_PRINTF = 4,		// Bump version since the way FRepLayoutCmd::CompatibleChecksum was calculated changed due to an optimization
-	HISTORY_NEW_ACTOR_OVERRIDE_LEVEL = 5,			// Bump version since a level reference was added to the new actor information
-	HISTORY_CHANNEL_NAMES = 6,						// Bump version since channel type is now an fname
-	HISTORY_CHANNEL_CLOSE_REASON = 7,				// Bump version to serialize a channel close reason in bunches instead of bDormant
-	HISTORY_ACKS_INCLUDED_IN_HEADER = 8,			// Bump version since acks are now sent as part of the header
-	HISTORY_NETEXPORT_SERIALIZATION = 9,			// Bump version due to serialization change to FNetFieldExport
-	HISTORY_NETEXPORT_SERIALIZE_FIX = 10,			// Bump version to fix net field export name serialization 
-	HISTORY_FAST_ARRAY_DELTA_STRUCT = 11,			// Bump version to allow fast array serialization, delta struct serialization.
-	HISTORY_FIX_ENUM_SERIALIZATION = 12,			// Bump version to fix enum net serialization issues.
-	HISTORY_OPTIONALLY_QUANTIZE_SPAWN_INFO = 13,	// Bump version to conditionally disable quantization for Scale, Location, and Velocity when spawning network actors.
-	HISTORY_JITTER_IN_HEADER = 14,					// Bump version since we added jitter clock time to packet headers and removed remote saturation
-	HISTORY_CLASSNETCACHE_FULLNAME = 15,			// Bump version to use full paths in GetNetFieldExportGroupForClassNetCache
-	HISTORY_REPLAY_DORMANCY = 16,					// Bump version to support dormancy properly in replays
-	HISTORY_ENUM_SERIALIZATION_COMPAT = 17,			// Bump version to include enum bits required for serialization into compat checksums, as well as unify enum and byte property enum serialization
-	HISTORY_SUBOBJECT_OUTER_CHAIN = 18,				// Bump version to support subobject outer chains matching on client and server
-	HISTORY_HITRESULT_INSTANCEHANDLE = 19,			// Bump version to support FHitResult change of Actor to HitObjectHandle. This change was made in CL 14369221 but a net version wasn't added at the time.
-	HISTORY_INTERFACE_PROPERTY_SERIALIZATION = 20,	// Bump version to support net serialization of FInterfaceProperty
-	HISTORY_MONTAGE_PLAY_INST_ID_SERIALIZATION = 21,// Bump version to support net serialization of FGameplayAbilityRepAnimMontage, addition of PlayInstanceId and removal of bForcePlayBit
-	HISTORY_SERIALIZE_DOUBLE_VECTORS_AS_DOUBLES	= 22,// Bump version to support net serialization of double vector types
-	HISTORY_PACKED_VECTOR_LWC_SUPPORT = 23,			// Bump version to support quantized LWC FVector net serialization
-	HISTORY_PAWN_REMOTEVIEWPITCH = 24,				// Bump version to support serialization changes to RemoteViewPitch
-	HISTORY_REPMOVE_SERVERFRAME_AND_HANDLE = 25,	// Bump version to support serialization changes to RepMove so we can get the serverframe and physics handle associated with the object
-	HISTORY_21_AND_VIEWPITCH_ONLY_DO_NOT_USE = 26,	// Bump version to support up to history 21 + HISTORY_PAWN_REMOTEVIEWPITCH.  DO NOT USE!!!
-	HISTORY_PLACEHOLDER = 27,						// Bump version to a placeholder.  This version is the same as HISTORY_REPMOVE_SERVERFRAME_AND_HANDLE
-	HISTORY_RUNTIME_FEATURES_COMPATIBILITY = 28,	// Bump version to add network runtime feature compatibility test to handshake (hello/upgrade) control messages
-	HISTORY_SOFTOBJECTPTR_NETGUIDS = 29,			// Bump version to support replicating SoftObjectPtrs by NetGuid instead of raw strings.
-	HISTORY_SUBOBJECT_DESTROY_FLAG = 30,			// Bump version to support subobject destruction message flags
+	HISTORY_REPLAY_BACKWARDS_COMPAT = 2,				// Bump version to get rid of older replays before backwards compat was turned on officially
+	HISTORY_MAX_ACTOR_CHANNELS_CUSTOMIZATION = 3,		// Bump version because serialization of the actor channels changed
+	HISTORY_REPCMD_CHECKSUM_REMOVE_PRINTF = 4,			// Bump version since the way FRepLayoutCmd::CompatibleChecksum was calculated changed due to an optimization
+	HISTORY_NEW_ACTOR_OVERRIDE_LEVEL = 5,				// Bump version since a level reference was added to the new actor information
+	HISTORY_CHANNEL_NAMES = 6,							// Bump version since channel type is now an fname
+	HISTORY_CHANNEL_CLOSE_REASON = 7,					// Bump version to serialize a channel close reason in bunches instead of bDormant
+	HISTORY_ACKS_INCLUDED_IN_HEADER = 8,				// Bump version since acks are now sent as part of the header
+	HISTORY_NETEXPORT_SERIALIZATION = 9,				// Bump version due to serialization change to FNetFieldExport
+	HISTORY_NETEXPORT_SERIALIZE_FIX = 10,				// Bump version to fix net field export name serialization 
+	HISTORY_FAST_ARRAY_DELTA_STRUCT = 11,				// Bump version to allow fast array serialization, delta struct serialization.
+	HISTORY_FIX_ENUM_SERIALIZATION = 12,				// Bump version to fix enum net serialization issues.
+	HISTORY_OPTIONALLY_QUANTIZE_SPAWN_INFO = 13,		// Bump version to conditionally disable quantization for Scale, Location, and Velocity when spawning network actors.
+	HISTORY_JITTER_IN_HEADER = 14,						// Bump version since we added jitter clock time to packet headers and removed remote saturation
+	HISTORY_CLASSNETCACHE_FULLNAME = 15,				// Bump version to use full paths in GetNetFieldExportGroupForClassNetCache
+	HISTORY_REPLAY_DORMANCY = 16,						// Bump version to support dormancy properly in replays
+	HISTORY_ENUM_SERIALIZATION_COMPAT = 17,				// Bump version to include enum bits required for serialization into compat checksums, as well as unify enum and byte property enum serialization
+	HISTORY_SUBOBJECT_OUTER_CHAIN = 18,					// Bump version to support subobject outer chains matching on client and server
+	HISTORY_HITRESULT_INSTANCEHANDLE = 19,				// Bump version to support FHitResult change of Actor to HitObjectHandle. This change was made in CL 14369221 but a net version wasn't added at the time.
+	HISTORY_INTERFACE_PROPERTY_SERIALIZATION = 20,		// Bump version to support net serialization of FInterfaceProperty
+	HISTORY_MONTAGE_PLAY_INST_ID_SERIALIZATION = 21,	// Bump version to support net serialization of FGameplayAbilityRepAnimMontage, addition of PlayInstanceId and removal of bForcePlayBit
+	HISTORY_SERIALIZE_DOUBLE_VECTORS_AS_DOUBLES	= 22,	// Bump version to support net serialization of double vector types
+	HISTORY_PACKED_VECTOR_LWC_SUPPORT = 23,				// Bump version to support quantized LWC FVector net serialization
+	HISTORY_PAWN_REMOTEVIEWPITCH = 24,					// Bump version to support serialization changes to RemoteViewPitch
+	HISTORY_REPMOVE_SERVERFRAME_AND_HANDLE = 25,		// Bump version to support serialization changes to RepMove so we can get the serverframe and physics handle associated with the object
+	HISTORY_21_AND_VIEWPITCH_ONLY_DO_NOT_USE = 26,		// Bump version to support up to history 21 + HISTORY_PAWN_REMOTEVIEWPITCH.  DO NOT USE!!!
+	HISTORY_PLACEHOLDER = 27,							// Bump version to a placeholder.  This version is the same as HISTORY_REPMOVE_SERVERFRAME_AND_HANDLE
+	HISTORY_RUNTIME_FEATURES_COMPATIBILITY = 28,		// Bump version to add network runtime feature compatibility test to handshake (hello/upgrade) control messages
+	HISTORY_SOFTOBJECTPTR_NETGUIDS = 29,				// Bump version to support replicating SoftObjectPtrs by NetGuid instead of raw strings.
+	HISTORY_SUBOBJECT_DESTROY_FLAG = 30,				// Bump version to support subobject destruction message flags
+	HISTORY_GAMESTATE_REPLCIATED_TIME_AS_DOUBLE = 31,	// Bump version to support AGameStateBase::ReplicatedWorldTimeSeconds as double instead of float.
 	// New history items go above here.
 
 	HISTORY_ENGINENETVERSION_PLUS_ONE,
diff --git a/Engine/Source/Runtime/Engine/Classes/GameFramework/GameStateBase.h b/Engine/Source/Runtime/Engine/Classes/GameFramework/GameStateBase.h
--- a/Engine/Source/Runtime/Engine/Classes/GameFramework/GameStateBase.h
+++ b/Engine/Source/Runtime/Engine/Classes/GameFramework/GameStateBase.h
@@ -139,16 +139,25 @@ protected:
 	UPROPERTY(Transient, ReplicatedUsing = OnRep_ReplicatedHasBegunPlay)
 	bool bReplicatedHasBegunPlay;
 
-	/** Called periodically to update ReplicatedWorldTimeSeconds */
+	/** Called periodically to update ReplicatedWorldTimeSecondsDouble */
 	virtual void UpdateServerTimeSeconds();
 
 	/** Allows clients to calculate ServerWorldTimeSecondsDelta */
+	UE_DEPRECATED(5.2, "OnRep_ReplicatedWorldTimeSeconds() is deprecated. Use OnRep_ReplicatedWorldTimeSecondsDouble().")
 	UFUNCTION()
-	virtual void OnRep_ReplicatedWorldTimeSeconds();
+	virtual void OnRep_ReplicatedWorldTimeSeconds() final;
+
+	/** Allows clients to calculate ServerWorldTimeSecondsDelta */
+	UFUNCTION()
+	virtual void OnRep_ReplicatedWorldTimeSecondsDouble();
 
 	/** Server TimeSeconds. Useful for syncing up animation and gameplay. */
+	UE_DEPRECATED(5.2, "ReplicatedWorldTimeSeconds is deprecated. Use ReplicatedWorldTimeSecondsDouble.")
 	UPROPERTY(Transient, ReplicatedUsing=OnRep_ReplicatedWorldTimeSeconds)
-	double ReplicatedWorldTimeSeconds;
+	float ReplicatedWorldTimeSeconds;
+
+	UPROPERTY(Transient, ReplicatedUsing = OnRep_ReplicatedWorldTimeSecondsDouble)
+	double ReplicatedWorldTimeSecondsDouble;
 
 	/** The difference from the local world's TimeSeconds and the server world's TimeSeconds. */
 	UPROPERTY(Transient)
diff --git a/Engine/Source/Runtime/Engine/Private/DemoNetDriver.cpp b/Engine/Source/Runtime/Engine/Private/DemoNetDriver.cpp
--- a/Engine/Source/Runtime/Engine/Private/DemoNetDriver.cpp
+++ b/Engine/Source/Runtime/Engine/Private/DemoNetDriver.cpp
@@ -2973,11 +2973,11 @@ void UDemoNetDriver::FinalizeFastForward(const double StartTime)
 		if (bIsFastForwardingForCheckpoint)
 		{
 			const double PostCheckpointServerTime = SavedReplicatedWorldTimeSeconds + SavedSecondsToSkip;
-			GameState->ReplicatedWorldTimeSeconds = PostCheckpointServerTime;
+			GameState->ReplicatedWorldTimeSecondsDouble = PostCheckpointServerTime;
 		}
 
 		// Correct the ServerWorldTimeSecondsDelta
-		GameState->OnRep_ReplicatedWorldTimeSeconds();
+		GameState->OnRep_ReplicatedWorldTimeSecondsDouble();
 	}
 
 	if (ServerConnection != nullptr && bIsFastForwardingForCheckpoint)
@@ -4524,7 +4524,7 @@ bool UDemoNetDriver::LoadCheckpoint(const FGotoResult& GotoResult)
 		const AGameStateBase* const GameState = World->GetGameState();
 		if (GameState != nullptr)
 		{
-			SavedReplicatedWorldTimeSeconds = GameState->ReplicatedWorldTimeSeconds;
+			SavedReplicatedWorldTimeSeconds = GameState->ReplicatedWorldTimeSecondsDouble;
 		}
 	}
 
diff --git a/Engine/Source/Runtime/Engine/Private/GameStateBase.cpp b/Engine/Source/Runtime/Engine/Private/GameStateBase.cpp
--- a/Engine/Source/Runtime/Engine/Private/GameStateBase.cpp
+++ b/Engine/Source/Runtime/Engine/Private/GameStateBase.cpp
@@ -9,6 +9,7 @@
 #include "Net/UnrealNetwork.h"
 #include "GameFramework/PlayerState.h"
 #include "EngineUtils.h"
+#include "Engine/DemoNetDriver.h"
 
 #include UE_INLINE_GENERATED_CPP_BY_NAME(GameStateBase)
 
@@ -154,7 +155,7 @@ void AGameStateBase::UpdateServerTimeSeconds()
 	UWorld* World = GetWorld();
 	if (World)
 	{
-		ReplicatedWorldTimeSeconds = World->GetTimeSeconds();
+		ReplicatedWorldTimeSecondsDouble = World->GetTimeSeconds();
 	}
 }
 
@@ -163,7 +164,25 @@ void AGameStateBase::OnRep_ReplicatedWorldTimeSeconds()
 	UWorld* World = GetWorld();
 	if (World)
 	{
-		const double ServerWorldTimeDelta = ReplicatedWorldTimeSeconds - World->GetTimeSeconds();
+		const UDemoNetDriver* DemoNetDriver = World ? World->GetDemoNetDriver() : nullptr;
+
+		// Support old replays.
+		if (DemoNetDriver && DemoNetDriver->IsPlaying() && (DemoNetDriver->GetPlaybackEngineNetworkProtocolVersion() < EEngineNetworkVersionHistory::HISTORY_GAMESTATE_REPLCIATED_TIME_AS_DOUBLE))
+		{
+			PRAGMA_DISABLE_DEPRECATION_WARNINGS
+			ReplicatedWorldTimeSecondsDouble = ReplicatedWorldTimeSeconds;
+			PRAGMA_ENABLE_DEPRECATION_WARNINGS
+			OnRep_ReplicatedWorldTimeSecondsDouble();
+		}
+	}
+}
+
+void AGameStateBase::OnRep_ReplicatedWorldTimeSecondsDouble()
+{
+	UWorld* World = GetWorld();
+	if (World)
+	{
+		const double ServerWorldTimeDelta = ReplicatedWorldTimeSecondsDouble - World->GetTimeSeconds();
 
 		// Accumulate the computed server world delta
 		SumServerWorldTimeSecondsDelta += ServerWorldTimeDelta;
@@ -251,8 +270,15 @@ void AGameStateBase::GetLifetimeReplicatedProps(TArray<FLifetimeProperty>& OutLi
 	DOREPLIFETIME( AGameStateBase, SpectatorClass );
 
 	DOREPLIFETIME_CONDITION( AGameStateBase, GameModeClass,	COND_InitialOnly );
+	DOREPLIFETIME( AGameStateBase, ReplicatedWorldTimeSecondsDouble );
+
+	FDoRepLifetimeParams Params;
+	Params.bIsPushBased = true;
+	PRAGMA_DISABLE_DEPRECATION_WARNINGS
+	// ReplicatedWorldTimeSeconds is still 'replicated' to support old replays, however this variable does not change value on the server now and is never dirtyed.
+	DOREPLIFETIME_WITH_PARAMS_FAST(AGameStateBase, ReplicatedWorldTimeSeconds, Params);
+	PRAGMA_ENABLE_DEPRECATION_WARNINGS
 
-	DOREPLIFETIME( AGameStateBase, ReplicatedWorldTimeSeconds );
 	DOREPLIFETIME( AGameStateBase, bReplicatedHasBegunPlay );
 }
 
diff --git a/Engine/Source/Runtime/Engine/Private/ReplayHelper.cpp b/Engine/Source/Runtime/Engine/Private/ReplayHelper.cpp
--- a/Engine/Source/Runtime/Engine/Private/ReplayHelper.cpp
+++ b/Engine/Source/Runtime/Engine/Private/ReplayHelper.cpp
@@ -749,9 +749,9 @@ void FReplayHelper::ProcessCheckpointActors(UNetConnection* Connection, TArrayVi
 	// more often than the normal update rate.
 	AGameStateBase* const GameState = World != nullptr ? World->GetGameState() : nullptr;
 
-	const double SavedReplicatedServerTimeSeconds = GameState ? GameState->ReplicatedWorldTimeSeconds : -1.0;
+	const double SavedReplicatedServerTimeSeconds = GameState ? GameState->ReplicatedWorldTimeSecondsDouble : -1.0;
 
-	// Normally AGameStateBase::ReplicatedWorldTimeSeconds is only updated periodically,
+	// Normally AGameStateBase::ReplicatedWorldTimeSecondsDouble is only updated periodically,
 	// but we want to make sure it's accurate for the checkpoint.
 	if (GameState)
 	{
@@ -783,7 +783,7 @@ void FReplayHelper::ProcessCheckpointActors(UNetConnection* Connection, TArrayVi
 		if (GameState)
 		{
 			// Restore the game state's replicated world time
-			GameState->ReplicatedWorldTimeSeconds = SavedReplicatedServerTimeSeconds;
+			GameState->ReplicatedWorldTimeSecondsDouble = SavedReplicatedServerTimeSeconds;
 		}
 
 		FlushNetChecked(*Connection);
