ARM64 and ARM64EC fixes

diff --git a/Engine/Source/Runtime/Core/Private/Misc/SecureHash.cpp b/Engine/Source/Runtime/Core/Private/Misc/SecureHash.cpp
--- a/Engine/Source/Runtime/Core/Private/Misc/SecureHash.cpp
+++ b/Engine/Source/Runtime/Core/Private/Misc/SecureHash.cpp
@@ -8,7 +8,7 @@
 #include "Misc/StringBuilder.h"
 #include "Serialization/CompactBinaryWriter.h"
 
-#if defined(_M_AMD64) || defined(__x86_64__)
+#if (defined(_M_AMD64) || defined(__x86_64__)) && !defined(_M_ARM64EC)
 
 #include <immintrin.h>
 #if PLATFORM_COMPILER_CLANG
@@ -23,7 +23,7 @@
 #define UE_PLATFORM_SHA_FALLBACK 0
 #endif
 
-#elif defined(__aarch64__)
+#elif defined(__aarch64__) || defined(_M_ARM64) || defined(_M_ARM64EC)
 
 // on ARMv8 enable SHA instructions unconditionally only when they are enabled with compiler
 #ifdef __ARM_FEATURE_CRYPTO
diff --git a/Engine/Source/Runtime/Core/Private/Windows/WindowsPlatformCrashContext.cpp b/Engine/Source/Runtime/Core/Private/Windows/WindowsPlatformCrashContext.cpp
--- a/Engine/Source/Runtime/Core/Private/Windows/WindowsPlatformCrashContext.cpp
+++ b/Engine/Source/Runtime/Core/Private/Windows/WindowsPlatformCrashContext.cpp
@@ -1669,7 +1669,11 @@ static void ReportEvent(ECrashContextType InType, const TCHAR* ErrorMessage, uin
 
 					if (ProgramCounter == nullptr)
 					{
+#if PLATFORM_CPU_X86_FAMILY || defined(_M_ARM64EC)
 						ProgramCounter = (void*)(ThreadContext.Rip);
+#else
+						ProgramCounter = (void*)(ThreadContext.Pc);
+#endif
 					}
 				}
 			}
diff --git a/Engine/Source/Runtime/Core/Private/Windows/WindowsPlatformStackWalk.cpp b/Engine/Source/Runtime/Core/Private/Windows/WindowsPlatformStackWalk.cpp
--- a/Engine/Source/Runtime/Core/Private/Windows/WindowsPlatformStackWalk.cpp
+++ b/Engine/Source/Runtime/Core/Private/Windows/WindowsPlatformStackWalk.cpp
@@ -144,11 +144,18 @@ static int32 CaptureStackTraceHelper(uint64 *BackTrace, uint32 MaxDepth, FWindow
 		StackFrame64.AddrStack.Mode      = AddrModeFlat;
 		StackFrame64.AddrFrame.Mode      = AddrModeFlat;
 #if PLATFORM_64BITS
-		StackFrame64.AddrPC.Offset = ContextWapper->Context.Rip;
-		StackFrame64.AddrStack.Offset = ContextWapper->Context.Rsp;
-		StackFrame64.AddrFrame.Offset = ContextWapper->Context.Rbp;
+#if PLATFORM_CPU_X86_FAMILY || defined(_M_ARM64EC)
+		StackFrame64.AddrPC.Offset       = ContextWapper->Context.Rip;
+		StackFrame64.AddrStack.Offset    = ContextWapper->Context.Rsp;
+		StackFrame64.AddrFrame.Offset    = ContextWapper->Context.Rbp;
 		MachineType                      = IMAGE_FILE_MACHINE_AMD64;
-#else	//PLATFORM_64BITS
+#elif PLATFORM_CPU_ARM_FAMILY
+		StackFrame64.AddrPC.Offset       = ContextWapper->Context.Pc;
+		StackFrame64.AddrStack.Offset    = ContextWapper->Context.Sp;
+		StackFrame64.AddrFrame.Offset    = ContextWapper->Context.Fp;
+		MachineType                      = IMAGE_FILE_MACHINE_ARM64;
+#endif
+#else   //PLATFORM_64BITS
 		StackFrame64.AddrPC.Offset       = ContextWapper->Context.Eip;
 		StackFrame64.AddrStack.Offset    = ContextWapper->Context.Esp;
 		StackFrame64.AddrFrame.Offset    = ContextWapper->Context.Ebp;
diff --git a/Engine/Source/Runtime/OpenGLDrv/Private/OpenGLShaders.cpp b/Engine/Source/Runtime/OpenGLDrv/Private/OpenGLShaders.cpp
--- a/Engine/Source/Runtime/OpenGLDrv/Private/OpenGLShaders.cpp
+++ b/Engine/Source/Runtime/OpenGLDrv/Private/OpenGLShaders.cpp
@@ -20,7 +20,7 @@
 
 #define CHECK_FOR_GL_SHADERS_TO_REPLACE 0
 
-#if PLATFORM_WINDOWS
+#if PLATFORM_WINDOWS && PLATFORM_CPU_X86_FAMILY
 #include <mmintrin.h>
 #endif
 #include "SceneUtils.h"
