Changing eligible platforms to use ANSI for profiler strings to reduce the necessary additional memory required. They were previously using wide strings because using ANSI strings incurred a conversion to wide-char in order to use the wide API, but the API has since been fixed.

diff --git a/Engine/Source/Runtime/Core/Public/Windows/WindowsPlatform.h b/Engine/Source/Runtime/Core/Public/Windows/WindowsPlatform.h
--- a/Engine/Source/Runtime/Core/Public/Windows/WindowsPlatform.h
+++ b/Engine/Source/Runtime/Core/Public/Windows/WindowsPlatform.h
@@ -137,7 +137,6 @@ typedef FWindowsPlatformTypes FPlatformTypes;
 // Intrinsics for 128-bit atomics on Windows platform requires Windows 8 or higher (WINVER>=0x0602)
 // http://msdn.microsoft.com/en-us/library/windows/desktop/hh972640.aspx
 #define PLATFORM_HAS_128BIT_ATOMICS							(!HACK_HEADER_GENERATOR && PLATFORM_64BITS && (WINVER >= 0x602))
-#define PLATFORM_USES_ANSI_STRING_FOR_EXTERNAL_PROFILING	0
 
 #ifdef CDECL
 #undef CDECL
