Fix StringViewTest using TStringView over a temp object

diff --git a/Engine/Source/Runtime/Core/Tests/Containers/StringViewTest.cpp b/Engine/Source/Runtime/Core/Tests/Containers/StringViewTest.cpp
--- a/Engine/Source/Runtime/Core/Tests/Containers/StringViewTest.cpp
+++ b/Engine/Source/Runtime/Core/Tests/Containers/StringViewTest.cpp
@@ -122,9 +122,12 @@ TEST_CASE("Core::String::StringView::Constructor", "[Core][String][Smoke]")
 		STATIC_CHECK(std::is_same_v<decltype(ViewWideConstLiteral)::ElementType, WIDECHAR>);
 		STATIC_CHECK(std::is_same_v<decltype(ViewUtf8ConstLiteral)::ElementType, UTF8CHAR>);
 
-		TStringView ViewAnsiStringBuilder(WriteToAnsiString<16>("Test"));
-		TStringView ViewWideStringBuilder(WriteToWideString<16>(WIDETEXT("Test")));
-		TStringView ViewUtf8StringBuilder(WriteToUtf8String<16>(UTF8TEXT("Test")));
+		TWriteToString<ANSICHAR, 16> WriteAnsiString = WriteToAnsiString<16>("Test");
+ 		TWriteToString<WIDECHAR, 16> WriteWideString = WriteToWideString<16>(WIDETEXT("Test"));
+ 		TWriteToString<UTF8CHAR, 16> WriteUft8String = WriteToUtf8String<16>(UTF8TEXT("Test"));
+		TStringView ViewAnsiStringBuilder(WriteAnsiString);
+		TStringView ViewWideStringBuilder(WriteWideString);
+		TStringView ViewUtf8StringBuilder(WriteUft8String);
 		STATIC_CHECK(std::is_same_v<decltype(ViewAnsiStringBuilder)::ElementType, ANSICHAR>);
 		STATIC_CHECK(std::is_same_v<decltype(ViewWideStringBuilder)::ElementType, WIDECHAR>);
 		STATIC_CHECK(std::is_same_v<decltype(ViewUtf8StringBuilder)::ElementType, UTF8CHAR>);
