Fix general regression caused by mimalloc not being enabled by default for editor builds anymore

diff --git a/Engine/Source/Runtime/Core/Public/Windows/WindowsPlatform.h b/Engine/Source/Runtime/Core/Public/Windows/WindowsPlatform.h
--- a/Engine/Source/Runtime/Core/Public/Windows/WindowsPlatform.h
+++ b/Engine/Source/Runtime/Core/Public/Windows/WindowsPlatform.h
@@ -89,7 +89,7 @@ typedef FWindowsPlatformTypes FPlatformTypes;
 #define PLATFORM_USES_MICROSOFT_LIBC_FUNCTIONS				1
 #define PLATFORM_IS_ANSI_MALLOC_THREADSAFE					1
 #define PLATFORM_SUPPORTS_TBB								1
-#if defined(PLATFORM_CPU_ARM_FAMILY)
+#if PLATFORM_CPU_ARM_FAMILY
 #define PLATFORM_SUPPORTS_MIMALLOC							0
 #else
 #define PLATFORM_SUPPORTS_MIMALLOC							PLATFORM_64BITS
