[Core]
* NonNullPointer - Added includes and forward declares to fix a bunch of compile errors after running IWYU.

diff --git a/Engine/Source/Runtime/Core/Public/Templates/NonNullPointer.h b/Engine/Source/Runtime/Core/Public/Templates/NonNullPointer.h
--- a/Engine/Source/Runtime/Core/Public/Templates/NonNullPointer.h
+++ b/Engine/Source/Runtime/Core/Public/Templates/NonNullPointer.h
@@ -5,12 +5,14 @@
 // HEADER_UNIT_SKIP - This needs TOptional to compile properly. Don't know why it is not included?
 
 #include "CoreTypes.h"
+#include "Misc/AssertionMacros.h"
 #include "Misc/OptionalFwd.h"
 #include "Templates/EnableIf.h"
 #include "Templates/PointerIsConvertibleFromTo.h"
 
-// So we can construct TNonNullPtrs
-enum class EDefaultConstructNonNullPtr { UnsafeDoNotUse };
+class FArchive;
+enum class EDefaultConstructNonNullPtr { UnsafeDoNotUse }; // So we can construct TNonNullPtrs
+struct FNullOpt;
 
 namespace UE::Core::Private::NonNullPtr {
 template <typename...>
