change to CreateDirectoryTree as windows CreateDirectoryW will not make the inbetween directories that might exist in a path.
I'm seeing issues where CopyDirectoryTree can fail on windows if the full directory path doesn't exist. This can happen during DLC cooking

https://p4-swarm.epicgames.net/reviews/23481470/v1/

diff --git a/Engine/Source/Runtime/Core/Private/GenericPlatform/GenericPlatformFile.cpp b/Engine/Source/Runtime/Core/Private/GenericPlatform/GenericPlatformFile.cpp
--- a/Engine/Source/Runtime/Core/Private/GenericPlatform/GenericPlatformFile.cpp
+++ b/Engine/Source/Runtime/Core/Private/GenericPlatform/GenericPlatformFile.cpp
@@ -811,7 +811,7 @@ bool IPlatformFile::CopyDirectoryTree(const TCHAR* DestinationDirectory, const T
 
 	// Destination directory exists already or can be created ?
 	if (!DirectoryExists(*DestDir) &&
-		!CreateDirectory(*DestDir))
+		!CreateDirectoryTree(*DestDir))
 	{
 		return false;
 	}
