minor renaming of global vars to avoid name clashes in unity builds

diff --git a/Engine/Source/Runtime/Core/Private/ProfilingDebugging/StallDetector.cpp b/Engine/Source/Runtime/Core/Private/ProfilingDebugging/StallDetector.cpp
--- a/Engine/Source/Runtime/Core/Private/ProfilingDebugging/StallDetector.cpp
+++ b/Engine/Source/Runtime/Core/Private/ProfilingDebugging/StallDetector.cpp
@@ -80,8 +80,8 @@ namespace UE
 #endif
 	};
 
-	static FStallDetectorRunnable* Runnable = nullptr;
-	static FRunnableThread* Thread = nullptr;
+	static FStallDetectorRunnable* StallDetectorRunnable = nullptr;
+	static FRunnableThread* StallDetectorThread = nullptr;
 }
 
 UE::FStallDetectorRunnable::FStallDetectorRunnable()
@@ -554,7 +554,7 @@ double UE::FStallDetector::Seconds()
 	if (FStallDetector::IsRunning())
 	{
 #if STALL_DETECTOR_HEART_BEAT_CLOCK
-		Result = Runnable->GetClock().Seconds();
+		Result = StallDetectorRunnable->GetClock().Seconds();
 #else
 		Result = FPlatformTime::Seconds();
 #endif
@@ -589,17 +589,17 @@ void UE::FStallDetector::Startup()
 		check(FPlatformTime::GetSecondsPerCycle() > 0.0);
 
 		// Cannot be a global due to clock member
-		Runnable = new FStallDetectorRunnable();
+		StallDetectorRunnable = new FStallDetectorRunnable();
 
 		if (FPlatformProcess::SupportsMultithreading())
 		{
-			if (Thread == nullptr)
+			if (StallDetectorThread == nullptr)
 			{
-				Thread = FRunnableThread::Create(Runnable, TEXT("StallDetectorThread"));
-				check(Thread);
+				StallDetectorThread = FRunnableThread::Create(StallDetectorRunnable, TEXT("StallDetectorThread"));
+				check(StallDetectorThread);
 
 				// Poll until we have ticked the clock
-				while (!Runnable->GetStartedThread())
+				while (!StallDetectorRunnable->GetStartedThread())
 				{
 					FPlatformProcess::YieldThread();
 				}
@@ -616,11 +616,11 @@ void UE::FStallDetector::Shutdown()
 	{
 		UE_LOG(LogStall, Log, TEXT("Shutdown..."));
 
-		delete Thread;
-		Thread = nullptr;
+		delete StallDetectorThread;
+		StallDetectorThread = nullptr;
 
-		delete Runnable;
-		Runnable = nullptr;
+		delete StallDetectorRunnable;
+		StallDetectorRunnable = nullptr;
 
 		UE_LOG(LogStall, Log, TEXT("Shutdown complete."));
 	}
