fixed some of static analysis warnings reported here: https://horde.devtools.epicgames.com/job/6398595e2960b73220e958d3?step=1b3d

diff --git a/Engine/Source/Runtime/Core/Private/Tests/Async/SemaphoreTest.cpp b/Engine/Source/Runtime/Core/Private/Tests/Async/SemaphoreTest.cpp
--- a/Engine/Source/Runtime/Core/Private/Tests/Async/SemaphoreTest.cpp
+++ b/Engine/Source/Runtime/Core/Private/Tests/Async/SemaphoreTest.cpp
@@ -32,7 +32,7 @@ namespace UE::SemaphoreTests
 		std::atomic<uint32> Sync1{ 0 };
 		std::atomic<uint32> Sync2{ 0 };
 
-		for (int i = 0; i != NumThreads; ++i)
+		for (uint32 i = 0; i < NumThreads; ++i)
 		{
 			Threads.Emplace(TEXT("SemaphoreWakeUpPerfTest"),
 				[&Semaphore, &Sync1, &Sync2, &bQuit]
@@ -61,7 +61,7 @@ namespace UE::SemaphoreTests
 			}
 			else
 			{
-				for (int i = 0; i != NumThreads; ++i)
+				for (uint32 i = 0; i < NumThreads; ++i)
 				{
 					Semaphore.Release();
 				}
@@ -73,7 +73,7 @@ namespace UE::SemaphoreTests
 			FPlatformProcess::Sleep(0.0);
 		};
 
-		for (int i = 0; i != NumLoops; ++i)
+		for (uint32 i = 0; i < NumLoops; ++i)
 		{
 			Iteration();
 		}
@@ -104,7 +104,7 @@ namespace UE::SemaphoreTests
 		std::atomic<uint32> Sync1{ 0 };
 		std::atomic<uint32> Sync2{ 0 };
 
-		for (int i = 0; i != NumThreads; ++i)
+		for (uint32 i = 0; i < NumThreads; ++i)
 		{
 			Threads.Emplace(TEXT("MultiEventWakeUpPerfTest"),
 				[Event = &Events[i], &Sync1, &Sync2, &bQuit]()
@@ -127,7 +127,7 @@ namespace UE::SemaphoreTests
 			Sync2 = NumThreads;
 			Sync1 = NumThreads;
 
-			for (int i = 0; i != NumThreads; ++i)
+			for (uint32 i = 0; i < NumThreads; ++i)
 			{
 				Events[i]->Trigger();
 			}
diff --git a/Engine/Source/Runtime/Core/Public/Windows/WindowsSemaphore.h b/Engine/Source/Runtime/Core/Public/Windows/WindowsSemaphore.h
--- a/Engine/Source/Runtime/Core/Public/Windows/WindowsSemaphore.h
+++ b/Engine/Source/Runtime/Core/Public/Windows/WindowsSemaphore.h
@@ -33,7 +33,7 @@ public:
 	bool TryAcquire(FTimespan Timeout = FTimespan::Zero())
 	{
 		HRESULT hRes = WaitForSingleObject(Semaphore, (DWORD)Timeout.GetTotalMilliseconds());
-		checkfSlow(hRes == WAIT_OBJECT_0 || hRes == WAIT_TIMEOUT, TEXT("Acquiring semaphore failed: %d (%d)"), hRes, GetLastError());
+		checkfSlow(hRes == WAIT_OBJECT_0 || hRes == WAIT_TIMEOUT, TEXT("Acquiring semaphore failed: %d (%d)"), (int)hRes, GetLastError());
 		return hRes == WAIT_OBJECT_0;
 	}
 
