Serialization: Add analysis only initialization to suppress uninitialized warning when deserializing values

#rnx

diff --git a/Engine/Source/Runtime/Core/Public/Serialization/Archive.h b/Engine/Source/Runtime/Core/Public/Serialization/Archive.h
--- a/Engine/Source/Runtime/Core/Public/Serialization/Archive.h
+++ b/Engine/Source/Runtime/Core/Public/Serialization/Archive.h
@@ -1567,7 +1567,16 @@ public:
 	friend CORE_API FArchive& operator<<(FArchive& Ar, FString& Value);
 
 public:
-	virtual void Serialize(void* V, int64 Length) { }
+	virtual void Serialize(void* V, int64 Length)
+	{
+#if defined(__clang_analyzer__)
+		// Suppress core.uninitialized.Assign static analysis warning
+		// Base class FArchive::Serialize() does not modify the value being serialized, so the analysis
+		// assumes any data being serialized to could be uninitalized, but practically speaking any
+		// data being serialized will use a derived class that always modifies the value.
+		memset(V, 0, Length);
+#endif
+	}
 
 	virtual void SerializeBits(void* V, int64 LengthBits)
 	{
