[Core/CoreUObject]
* Added macro IWYU_MARKUP_IMPLICIT_CAST that can be used in extreme situations where we want to communicate to IWYU that we actually need to know about the <from> type and therefor include can't be removed

diff --git a/Engine/Source/Runtime/Core/Public/HAL/PlatformMisc.h b/Engine/Source/Runtime/Core/Public/HAL/PlatformMisc.h
--- a/Engine/Source/Runtime/Core/Public/HAL/PlatformMisc.h
+++ b/Engine/Source/Runtime/Core/Public/HAL/PlatformMisc.h
@@ -237,3 +237,18 @@ public:
 	FScopedEnterBackgroundEvent EnterBackgroundEvent_##Name##_(TEXT(#Name)); \
 	QUICK_SCOPE_CYCLE_COUNTER(Name);
 
+
+#ifdef PLATFORM_COMPILER_IWYU
+
+// There are limitations to what IWYU can read out of the ast. decltype() inside unused template parameter default initializer is one example
+// In those cases we need to give IWYU something it can use to know what is needed
+namespace UE::Core::Private
+{
+	template <typename T, bool IsSameTypes = true> struct IwyuTestSize { enum { Value = 0 }; };
+	template <typename T> struct IwyuTestSize<T, false> { enum { Value = sizeof(T) }; };
+}
+#define IWYU_MARKUP_IMPLICIT_CAST(From, To) UE::Core::Private::IwyuTestSize<From, std::is_same<typename std::remove_cv<To>::type, typename std::remove_cv<From>::type>::value>::Value
+
+#else
+#define IWYU_MARKUP_IMPLICIT_CAST(From, To)
+#endif
diff --git a/Engine/Source/Runtime/CoreUObject/Public/Templates/SubclassOf.h b/Engine/Source/Runtime/CoreUObject/Public/Templates/SubclassOf.h
--- a/Engine/Source/Runtime/CoreUObject/Public/Templates/SubclassOf.h
+++ b/Engine/Source/Runtime/CoreUObject/Public/Templates/SubclassOf.h
@@ -66,6 +66,7 @@ public:
 	FORCEINLINE TSubclassOf(const TSubclassOf<OtherT>& Other)
 		: Class(Other.Class)
 	{
+		IWYU_MARKUP_IMPLICIT_CAST(OtherT, T);
 	}
 
 	/** Assign from another TSubclassOf, only if types are compatible */
@@ -75,6 +76,7 @@ public:
 	>
 	FORCEINLINE TSubclassOf& operator=(const TSubclassOf<OtherT>& Other)
 	{
+		IWYU_MARKUP_IMPLICIT_CAST(OtherT, T);
 		Class = Other.Class;
 		return *this;
 	}
diff --git a/Engine/Source/Runtime/CoreUObject/Public/UObject/ObjectPtr.h b/Engine/Source/Runtime/CoreUObject/Public/UObject/ObjectPtr.h
--- a/Engine/Source/Runtime/CoreUObject/Public/UObject/ObjectPtr.h
+++ b/Engine/Source/Runtime/CoreUObject/Public/UObject/ObjectPtr.h
@@ -380,6 +380,7 @@ public:
 	FORCEINLINE TObjectPtr(const TObjectPtr<U>& Other)
 		: ObjectPtr(Other.ObjectPtr)
 	{
+		IWYU_MARKUP_IMPLICIT_CAST(U, T);
 	}
 
 	template <
@@ -414,9 +415,7 @@ public:
 	>
 	FORCEINLINE TObjectPtr<T>& operator=(const TObjectPtr<U>& Other)
 	{
-		#ifdef PLATFORM_COMPILER_IWYU
-		enum { IWYU = sizeof(U) }; // Unfortunately the decltype expression above is not visible to IWYU so it can't tell that the code needs to know about type U
-		#endif
+		IWYU_MARKUP_IMPLICIT_CAST(U, T);
 		ObjectPtr = Other.ObjectPtr;
 		return *this;
 	}
