CompactBinary: Added serialization of fields and arrays to JSON

diff --git a/Engine/Source/Runtime/Core/Private/Serialization/CompactBinarySerialization.cpp b/Engine/Source/Runtime/Core/Private/Serialization/CompactBinarySerialization.cpp
--- a/Engine/Source/Runtime/Core/Private/Serialization/CompactBinarySerialization.cpp
+++ b/Engine/Source/Runtime/Core/Private/Serialization/CompactBinarySerialization.cpp
@@ -607,28 +607,68 @@ private:
 	bool bNeedsNewLine{false};
 };
 
-void CompactBinaryToJson(const FCbObjectView& Object, FUtf8StringBuilderBase& Builder)
+void CompactBinaryToJson(const FCbFieldView& Field, FUtf8StringBuilderBase& Builder)
 {
 	TCbJsonWriter<UTF8CHAR, ANSICHAR> Writer(Builder);
-	Writer.WriteField(Object.AsFieldView());
+	Writer.WriteField(Field);
 }
 
-void CompactBinaryToJson(const FCbObjectView& Object, FWideStringBuilderBase& Builder)
+void CompactBinaryToJson(const FCbFieldView& Field, FWideStringBuilderBase& Builder)
 {
 	TCbJsonWriter<WIDECHAR> Writer(Builder);
-	Writer.WriteField(Object.AsFieldView());
+	Writer.WriteField(Field);
 }
 
-void CompactBinaryToCompactJson(const FCbObjectView& Object, FUtf8StringBuilderBase& Builder)
+void CompactBinaryToJson(const FCbArrayView& Array, FUtf8StringBuilderBase& Builder)
+{
+	CompactBinaryToJson(Array.AsFieldView(), Builder);
+}
+
+void CompactBinaryToJson(const FCbArrayView& Array, FWideStringBuilderBase& Builder)
+{
+	CompactBinaryToJson(Array.AsFieldView(), Builder);
+}
+
+void CompactBinaryToJson(const FCbObjectView& Object, FUtf8StringBuilderBase& Builder)
+{
+	CompactBinaryToJson(Object.AsFieldView(), Builder);
+}
+
+void CompactBinaryToJson(const FCbObjectView& Object, FWideStringBuilderBase& Builder)
+{
+	CompactBinaryToJson(Object.AsFieldView(), Builder);
+}
+
+void CompactBinaryToCompactJson(const FCbFieldView& Field, FUtf8StringBuilderBase& Builder)
 {
 	TCbJsonWriter<UTF8CHAR, ANSICHAR> Writer(Builder, ANSITEXTVIEW(""), ANSITEXTVIEW(""), ANSITEXTVIEW(""));
-	Writer.WriteField(Object.AsFieldView());
+	Writer.WriteField(Field);
 }
 
-void CompactBinaryToCompactJson(const FCbObjectView& Object, FWideStringBuilderBase& Builder)
+void CompactBinaryToCompactJson(const FCbFieldView& Field, FWideStringBuilderBase& Builder)
 {
 	TCbJsonWriter<WIDECHAR> Writer(Builder, WIDETEXTVIEW(""), WIDETEXTVIEW(""), WIDETEXTVIEW(""));
-	Writer.WriteField(Object.AsFieldView());
+	Writer.WriteField(Field);
+}
+
+void CompactBinaryToCompactJson(const FCbArrayView& Array, FUtf8StringBuilderBase& Builder)
+{
+	CompactBinaryToCompactJson(Array.AsFieldView(), Builder);
+}
+
+void CompactBinaryToCompactJson(const FCbArrayView& Array, FWideStringBuilderBase& Builder)
+{
+	CompactBinaryToCompactJson(Array.AsFieldView(), Builder);
+}
+
+void CompactBinaryToCompactJson(const FCbObjectView& Object, FUtf8StringBuilderBase& Builder)
+{
+	CompactBinaryToCompactJson(Object.AsFieldView(), Builder);
+}
+
+void CompactBinaryToCompactJson(const FCbObjectView& Object, FWideStringBuilderBase& Builder)
+{
+	CompactBinaryToCompactJson(Object.AsFieldView(), Builder);
 }
 
 ///////////////////////////////////////////////////////////////////////////////////////////////////////////////////////
diff --git a/Engine/Source/Runtime/Core/Public/Serialization/CompactBinarySerialization.h b/Engine/Source/Runtime/Core/Public/Serialization/CompactBinarySerialization.h
--- a/Engine/Source/Runtime/Core/Public/Serialization/CompactBinarySerialization.h
+++ b/Engine/Source/Runtime/Core/Public/Serialization/CompactBinarySerialization.h
@@ -211,11 +211,19 @@ inline bool LoadFromCompactBinary(FCbFieldView Field, TArray<T, Allocator>& OutV
 
 ///////////////////////////////////////////////////////////////////////////////////////////////////////////////////////
 
-/** Convert the object to JSON in a multi-line format with indentation. */
+/** Convert the compact binary to JSON in a multi-line format with indentation. */
+CORE_API void CompactBinaryToJson(const FCbFieldView& Field, FUtf8StringBuilderBase& Builder);
+CORE_API void CompactBinaryToJson(const FCbFieldView& Field, FWideStringBuilderBase& Builder);
+CORE_API void CompactBinaryToJson(const FCbArrayView& Array, FUtf8StringBuilderBase& Builder);
+CORE_API void CompactBinaryToJson(const FCbArrayView& Array, FWideStringBuilderBase& Builder);
 CORE_API void CompactBinaryToJson(const FCbObjectView& Object, FUtf8StringBuilderBase& Builder);
 CORE_API void CompactBinaryToJson(const FCbObjectView& Object, FWideStringBuilderBase& Builder);
 
-/** Convert the object to JSON in a compact format with no added whitespace. */
+/** Convert the compact binary to JSON in a compact format with no added whitespace. */
+CORE_API void CompactBinaryToCompactJson(const FCbFieldView& Field, FUtf8StringBuilderBase& Builder);
+CORE_API void CompactBinaryToCompactJson(const FCbFieldView& Field, FWideStringBuilderBase& Builder);
+CORE_API void CompactBinaryToCompactJson(const FCbArrayView& Array, FUtf8StringBuilderBase& Builder);
+CORE_API void CompactBinaryToCompactJson(const FCbArrayView& Array, FWideStringBuilderBase& Builder);
 CORE_API void CompactBinaryToCompactJson(const FCbObjectView& Object, FUtf8StringBuilderBase& Builder);
 CORE_API void CompactBinaryToCompactJson(const FCbObjectView& Object, FWideStringBuilderBase& Builder);
 
