Add more LIFETIME bounds for TArray, StringConv, StringView, and ArrayView

diff --git a/Engine/Source/Runtime/Core/Public/Containers/Array.h b/Engine/Source/Runtime/Core/Public/Containers/Array.h
--- a/Engine/Source/Runtime/Core/Public/Containers/Array.h
+++ b/Engine/Source/Runtime/Core/Public/Containers/Array.h
@@ -679,7 +679,7 @@ public:
 	 *
 	 * @returns Pointer to first array entry or nullptr if ArrayMax == 0.
 	 */
-	FORCEINLINE ElementType* GetData()
+	FORCEINLINE ElementType* GetData() UE_LIFETIMEBOUND
 	{
 		return (ElementType*)AllocatorInstance.GetAllocation();
 	}
@@ -689,7 +689,7 @@ public:
 	 *
 	 * @returns Pointer to first array entry or nullptr if ArrayMax == 0.
 	 */
-	FORCEINLINE const ElementType* GetData() const
+	FORCEINLINE const ElementType* GetData() const UE_LIFETIMEBOUND
 	{
 		return (const ElementType*)AllocatorInstance.GetAllocation();
 	}
@@ -800,7 +800,7 @@ public:
 	 *
 	 * @returns Reference to indexed element.
 	 */
-	FORCEINLINE ElementType& operator[](SizeType Index)
+	FORCEINLINE ElementType& operator[](SizeType Index) UE_LIFETIMEBOUND
 	{
 		RangeCheck(Index);
 		return GetData()[Index];
@@ -813,7 +813,7 @@ public:
 	 *
 	 * @returns Reference to indexed element.
 	 */
-	FORCEINLINE const ElementType& operator[](SizeType Index) const
+	FORCEINLINE const ElementType& operator[](SizeType Index) const UE_LIFETIMEBOUND
 	{
 		RangeCheck(Index);
 		return GetData()[Index];
@@ -3542,4 +3542,4 @@ template<typename ElementType, typename AllocatorType>
 FArchive& operator<<(FArchive& Ar, TArray<ElementType, AllocatorType>& A)
 {
 	return TArrayPrivateFriend::Serialize(Ar, A);
-}
\ No newline at end of file
+}
diff --git a/Engine/Source/Runtime/Core/Public/Containers/ArrayView.h b/Engine/Source/Runtime/Core/Public/Containers/ArrayView.h
--- a/Engine/Source/Runtime/Core/Public/Containers/ArrayView.h
+++ b/Engine/Source/Runtime/Core/Public/Containers/ArrayView.h
@@ -217,7 +217,7 @@ public:
 	 */
 	template <typename OtherElementType,
 		typename = typename TEnableIf<TIsCompatibleElementType<OtherElementType>::Value>::Type>
-	FORCEINLINE TArrayView(OtherElementType* InData, SizeType InCount)
+	FORCEINLINE TArrayView(OtherElementType* InData UE_LIFETIMEBOUND, SizeType InCount)
 		: DataPtr(InData)
 		, ArrayNum(InCount)
 	{
@@ -769,7 +769,7 @@ auto MakeArrayView(OtherRangeType&& Other UE_LIFETIMEBOUND)
 }
 
 template<typename ElementType>
-auto MakeArrayView(ElementType* Pointer, int32 Size)
+auto MakeArrayView(ElementType* Pointer UE_LIFETIMEBOUND, int32 Size)
 {
 	return TArrayView<ElementType>(Pointer, Size);
 }
diff --git a/Engine/Source/Runtime/Core/Public/Containers/StringConv.h b/Engine/Source/Runtime/Core/Public/Containers/StringConv.h
--- a/Engine/Source/Runtime/Core/Public/Containers/StringConv.h
+++ b/Engine/Source/Runtime/Core/Public/Containers/StringConv.h
@@ -754,7 +754,7 @@ public:
 	 *
 	 * @return A const pointer to the converted string.
 	 */
-	FORCEINLINE const ToType* Get() const
+	FORCEINLINE const ToType* Get() const UE_LIFETIMEBOUND
 	{
 		return Ptr;
 	}
@@ -799,6 +799,9 @@ struct TElementType<TStringConversion<Converter, DefaultConversionSize>>
 };
 
 
+namespace UE::Core::Private
+{
+
 /**
  * Class takes one type of string and and stores it as-is.
  * For API compatibility with TStringConversion when the To and From types are already in the same format.
@@ -882,10 +885,11 @@ public:
 	/**
 	 * Accessor for the string.
 	 * @note The string may not be null-terminated if constructed from an explicitly sized buffer that didn't include the null-terminator.
+	 * @note Setting to UE_LIFETIMEBOUND even though TStringPointer doesn't own the pointer, because we want StringCast to give consistent warnings no matter what implementation choice it makes
 	 *
 	 * @return A const pointer to the string.
 	 */
-	FORCEINLINE const ToType* Get() const
+	FORCEINLINE const ToType* Get() const UE_LIFETIMEBOUND
 	{
 		return Ptr;
 	}
@@ -909,31 +913,32 @@ private:
 	int32 StringLength;
 };
 
+} // namespace UE::Core::Private
+
 template <typename FromType, typename ToType>
-inline auto GetData(const TStringPointer<FromType, ToType>& Pointer) -> decltype(Pointer.Get())
+inline auto GetData(const UE::Core::Private::TStringPointer<FromType, ToType>& Pointer) -> decltype(Pointer.Get())
 {
 	return Pointer.Get();
 }
 
 template <typename FromType, typename ToType>
-inline auto GetNum(const TStringPointer<FromType, ToType>& Pointer) -> decltype(Pointer.Length())
+inline auto GetNum(const UE::Core::Private::TStringPointer<FromType, ToType>& Pointer) -> decltype(Pointer.Length())
 {
 	return Pointer.Length();
 }
 
 template <typename FromType, typename ToType>
-struct TIsContiguousContainer<TStringPointer<FromType, ToType>>
+struct TIsContiguousContainer<UE::Core::Private::TStringPointer<FromType, ToType>>
 {
 	static constexpr bool Value = true;
 };
 
 template <typename FromType, typename ToType>
-struct TElementType<TStringPointer<FromType, ToType>>
+struct TElementType<UE::Core::Private::TStringPointer<FromType, ToType>>
 {
 	using Type = ToType;
 };
 
-
 /**
  * NOTE: The objects these macros declare have very short lifetimes. They are
  * meant to be used as parameters to functions. You cannot assign a variable
@@ -970,8 +975,8 @@ typedef TStringConversion<TUTF16ToUTF32_Convert<UTF16CHAR, TCHAR>> FUTF16ToTCHAR
 #define UTF16_TO_TCHAR(str) (TCHAR*)FUTF16ToTCHAR((const UTF16CHAR*)str).Get()
 
 static_assert(sizeof(TCHAR) == sizeof(UTF32CHAR), "TCHAR and UTF32CHAR are expected to be the same size for inline conversion! PLATFORM_TCHAR_IS_4_BYTES is not configured correctly for this platform.");
-typedef TStringPointer<TCHAR, UTF32CHAR> FTCHARToUTF32;
-typedef TStringPointer<UTF32CHAR, TCHAR> FUTF32ToTCHAR;
+typedef UE::Core::Private::TStringPointer<TCHAR, UTF32CHAR> FTCHARToUTF32;
+typedef UE::Core::Private::TStringPointer<UTF32CHAR, TCHAR> FUTF32ToTCHAR;
 #define TCHAR_TO_UTF32(str) (UTF32CHAR*)(str)
 #define UTF32_TO_TCHAR(str) (TCHAR*)(str)
 
@@ -990,8 +995,8 @@ typedef TStringConversion<TStringConvert<UTF32CHAR, TCHAR>> FUTF32ToTCHAR;
 #else
 
 static_assert(sizeof(TCHAR) == sizeof(UTF16CHAR), "TCHAR and UTF16CHAR are expected to be the same size for inline conversion! PLATFORM_TCHAR_IS_4_BYTES is not configured correctly for this platform.");
-typedef TStringPointer<TCHAR, UTF16CHAR> FTCHARToUTF16;
-typedef TStringPointer<UTF16CHAR, TCHAR> FUTF16ToTCHAR;
+typedef UE::Core::Private::TStringPointer<TCHAR, UTF16CHAR> FTCHARToUTF16;
+typedef UE::Core::Private::TStringPointer<UTF16CHAR, TCHAR> FUTF16ToTCHAR;
 #define TCHAR_TO_UTF16(str) (UTF16CHAR*)(str)
 #define UTF16_TO_TCHAR(str) (TCHAR*)(str)
 
@@ -1019,8 +1024,8 @@ typedef TStringConversion<TStringConvert<wchar_t, TCHAR>> FWCharToTCHAR;
 #else
 
 static_assert(sizeof(TCHAR) == sizeof(wchar_t), "TCHAR and wchar_t are expected to be the same size for inline conversion! PLATFORM_WCHAR_IS_4_BYTES is not configured correctly for this platform.");
-typedef TStringPointer<TCHAR, wchar_t> FTCHARToWChar;
-typedef TStringPointer<wchar_t, TCHAR> FWCharToTCHAR;
+typedef UE::Core::Private::TStringPointer<TCHAR, wchar_t> FTCHARToWChar;
+typedef UE::Core::Private::TStringPointer<wchar_t, TCHAR> FWCharToTCHAR;
 #define TCHAR_TO_WCHAR(str) (wchar_t*)(str)
 #define WCHAR_TO_TCHAR(str) (TCHAR*)(str)
 
@@ -1054,7 +1059,7 @@ FORCEINLINE auto StringCast(const From* Str)
 {
 	if constexpr (TIsCharEncodingCompatibleWith_V<From, To>)
 	{
-		return TStringPointer<To>((const To*)Str);
+		return UE::Core::Private::TStringPointer<To>((const To*)Str);
 	}
 	else
 	{
@@ -1073,7 +1078,7 @@ FORCEINLINE auto StringCast(const From* Str, int32 Len)
 {
 	if constexpr (TIsCharEncodingCompatibleWith_V<From, To>)
 	{
-		return TStringPointer<To>((const To*)Str, Len);
+		return UE::Core::Private::TStringPointer<To>((const To*)Str, Len);
 	}
 	else
 	{
@@ -1096,6 +1101,9 @@ FORCEINLINE To CharCast(From Ch)
 	return Result;
 }
 
+namespace UE::Core::Private
+{
+
 /**
  * This class is returned by StringPassthru and is not intended to be used directly.
  */
@@ -1125,7 +1133,7 @@ public:
 		FPlatformString::Convert(Dest, DestLen, Source, SrcLen);
 	}
 
-	FORCEINLINE FromType* Get()
+	FORCEINLINE FromType* Get() UE_LIFETIMEBOUND
 	{
 		return (FromType*)AllocatorType::GetAllocation();
 	}
@@ -1151,7 +1159,8 @@ public:
 	{
 	}
 
-	FORCEINLINE T* Get() const
+	// @note Setting to UE_LIFETIMEBOUND even though TPassthruPointer doesn't own the pointer, because we want StringMemoryPassthru to give consistent warnings no matter what implementation choice it makes
+	FORCEINLINE T* Get() const UE_LIFETIMEBOUND
 	{
 		return Ptr;
 	}
@@ -1164,6 +1173,8 @@ private:
 	T* Ptr;
 };
 
+}
+
 /**
  * Allows the efficient conversion of strings by means of a temporary memory buffer only when necessary.  Intended to be used
  * when you have an API which populates a buffer with some string representation which is ultimately going to be stored in another
@@ -1202,11 +1213,11 @@ FORCEINLINE auto StringMemoryPassthru(To* Buffer, int32 BufferSize, int32 Source
 	if constexpr (TIsCharEncodingCompatibleWith_V<From, To>)
 	{
 		check(SourceLength <= BufferSize);
-		return TPassthruPointer<From>((From*)Buffer);
+		return UE::Core::Private::TPassthruPointer<From>((From*)Buffer);
 	}
 	else
 	{
-		return TStringPassthru<To, From, DefaultConversionSize>(Buffer, BufferSize, SourceLength);
+		return UE::Core::Private::TStringPassthru<To, From, DefaultConversionSize>(Buffer, BufferSize, SourceLength);
 	}
 }
 
diff --git a/Engine/Source/Runtime/Core/Public/Containers/StringView.h b/Engine/Source/Runtime/Core/Public/Containers/StringView.h
--- a/Engine/Source/Runtime/Core/Public/Containers/StringView.h
+++ b/Engine/Source/Runtime/Core/Public/Containers/StringView.h
@@ -86,14 +86,14 @@ public:
 	constexpr TStringView() = default;
 
 	/** Construct a view of the null-terminated string pointed to by InData. */
-	constexpr inline TStringView(const CharType* InData)
+	constexpr inline TStringView(const CharType* InData UE_LIFETIMEBOUND)
 		: DataPtr(InData)
 		, Size(InData ? TCString<CharType>::Strlen(InData) : 0)
 	{
 	}
 
 	/** Construct a view of InSize characters beginning at InData. */
-	constexpr inline TStringView(const CharType* InData, int32 InSize)
+	constexpr inline TStringView(const CharType* InData UE_LIFETIMEBOUND, int32 InSize)
 		: DataPtr(InData)
 		, Size(InSize)
 	{
@@ -109,7 +109,7 @@ public:
 			>::Value
 		>* = nullptr
 	>
-	constexpr inline TStringView(const OtherCharType* InData)
+	constexpr inline TStringView(const OtherCharType* InData UE_LIFETIMEBOUND)
 		: DataPtr((const CharType*)InData)
 		, Size(InData ? TCString<CharType>::Strlen((const CharType*)InData) : 0)
 	{
@@ -125,7 +125,7 @@ public:
 			>::Value
 		>* = nullptr
 	>
-	constexpr inline TStringView(const OtherCharType* InData, int32 InSize)
+	constexpr inline TStringView(const OtherCharType* InData UE_LIFETIMEBOUND, int32 InSize)
 		: DataPtr((const CharType*)InData)
 		, Size(InSize)
 	{
@@ -459,7 +459,7 @@ constexpr inline auto MakeStringView(CharPtrOrRangeType&& CharPtrOrRange UE_LIFE
 }
 
 template <typename CharPtrType>
-constexpr inline auto MakeStringView(CharPtrType&& CharPtr, int32 Size) -> decltype(TStringView(Forward<CharPtrType>(CharPtr), Size))
+constexpr inline auto MakeStringView(CharPtrType&& CharPtr UE_LIFETIMEBOUND, int32 Size) -> decltype(TStringView(Forward<CharPtrType>(CharPtr), Size))
 {
 	return TStringView(Forward<CharPtrType>(CharPtr), Size);
 }
