Revert change that disables uploading ensures if we have enabled an outside source for uploading crashes

#localization none
#tests reverting what the change did, which was tested to disable ensures

diff --git a/Engine/Source/Runtime/Core/Private/Unix/UnixPlatformCrashContext.cpp b/Engine/Source/Runtime/Core/Private/Unix/UnixPlatformCrashContext.cpp
--- a/Engine/Source/Runtime/Core/Private/Unix/UnixPlatformCrashContext.cpp
+++ b/Engine/Source/Runtime/Core/Private/Unix/UnixPlatformCrashContext.cpp
@@ -759,7 +759,7 @@ void FUnixCrashContext::GenerateCrashInfoAndLaunchReporter() const
 				GConfig->GetBool(TEXT("CrashReportClient"), TEXT("bStartCRCFromEngineHandler"), bStartCRCFromEngineHandler, GEngineIni);
 			}
 
-			if (bReportingNonCrash && bStartCRCFromEngineHandler)
+			if (bReportingNonCrash)
 			{
 				bool bFoundEmptySlot = false;
 
