fixed "tasks integration" as it's rotted while being disabled

diff --git a/Engine/Source/Runtime/Core/Public/Async/TaskGraphInterfaces.h b/Engine/Source/Runtime/Core/Public/Async/TaskGraphInterfaces.h
--- a/Engine/Source/Runtime/Core/Public/Async/TaskGraphInterfaces.h
+++ b/Engine/Source/Runtime/Core/Public/Async/TaskGraphInterfaces.h
@@ -555,7 +555,6 @@ public:
 			EExtendedTaskPriority ConversionMap[] =
 			{
 				EExtendedTaskPriority::RHIThreadNormalPri,
-				EExtendedTaskPriority::None, // invalid, maps from ENamedThreads::AudioThread
 				EExtendedTaskPriority::GameThreadNormalPri,
 				EExtendedTaskPriority::RenderThreadNormalPri
 			};
@@ -699,7 +698,7 @@ private:
 		FBaseGraphTask::Init(TEXT("GraphTask"), InPriority, InExtendedPriority);
 	}
 
-	virtual bool TryExecuteTask() override
+	virtual bool TryExecuteTaskVirtual() override
 	{
 		return TryExecute(
 			[](UE::Tasks::Private::FTaskBase& Task)
@@ -744,7 +743,7 @@ public:
 	static void operator delete(void* Ptr);
 
 private:
-	virtual bool TryExecuteTask() override
+	virtual bool TryExecuteTaskVirtual() override
 	{
 		checkNoEntry(); // graph events are never executed
 		return true;
