Add more LIFETIME bounds for TArray, StringConv, StringView, and ArrayView

diff --git a/Engine/Source/Runtime/Core/Public/Containers/Array.h b/Engine/Source/Runtime/Core/Public/Containers/Array.h
--- a/Engine/Source/Runtime/Core/Public/Containers/Array.h
+++ b/Engine/Source/Runtime/Core/Public/Containers/Array.h
@@ -679,7 +679,7 @@ public:
 	 *
 	 * @returns Pointer to first array entry or nullptr if ArrayMax == 0.
 	 */
-	FORCEINLINE ElementType* GetData()
+	FORCEINLINE ElementType* GetData() UE_LIFETIMEBOUND
 	{
 		return (ElementType*)AllocatorInstance.GetAllocation();
 	}
@@ -689,7 +689,7 @@ public:
 	 *
 	 * @returns Pointer to first array entry or nullptr if ArrayMax == 0.
 	 */
-	FORCEINLINE const ElementType* GetData() const
+	FORCEINLINE const ElementType* GetData() const UE_LIFETIMEBOUND
 	{
 		return (const ElementType*)AllocatorInstance.GetAllocation();
 	}
@@ -800,7 +800,7 @@ public:
 	 *
 	 * @returns Reference to indexed element.
 	 */
-	FORCEINLINE ElementType& operator[](SizeType Index)
+	FORCEINLINE ElementType& operator[](SizeType Index) UE_LIFETIMEBOUND
 	{
 		RangeCheck(Index);
 		return GetData()[Index];
@@ -813,7 +813,7 @@ public:
 	 *
 	 * @returns Reference to indexed element.
 	 */
-	FORCEINLINE const ElementType& operator[](SizeType Index) const
+	FORCEINLINE const ElementType& operator[](SizeType Index) const UE_LIFETIMEBOUND
 	{
 		RangeCheck(Index);
 		return GetData()[Index];
@@ -3542,4 +3542,4 @@ template<typename ElementType, typename AllocatorType>
 FArchive& operator<<(FArchive& Ar, TArray<ElementType, AllocatorType>& A)
 {
 	return TArrayPrivateFriend::Serialize(Ar, A);
-}
\ No newline at end of file
+}
diff --git a/Engine/Source/Runtime/Core/Public/Containers/ArrayView.h b/Engine/Source/Runtime/Core/Public/Containers/ArrayView.h
--- a/Engine/Source/Runtime/Core/Public/Containers/ArrayView.h
+++ b/Engine/Source/Runtime/Core/Public/Containers/ArrayView.h
@@ -217,7 +217,7 @@ public:
 	 */
 	template <typename OtherElementType,
 		typename = typename TEnableIf<TIsCompatibleElementType<OtherElementType>::Value>::Type>
-	FORCEINLINE TArrayView(OtherElementType* InData, SizeType InCount)
+	FORCEINLINE TArrayView(OtherElementType* InData UE_LIFETIMEBOUND, SizeType InCount)
 		: DataPtr(InData)
 		, ArrayNum(InCount)
 	{
@@ -769,7 +769,7 @@ auto MakeArrayView(OtherRangeType&& Other UE_LIFETIMEBOUND)
 }
 
 template<typename ElementType>
-auto MakeArrayView(ElementType* Pointer, int32 Size)
+auto MakeArrayView(ElementType* Pointer UE_LIFETIMEBOUND, int32 Size)
 {
 	return TArrayView<ElementType>(Pointer, Size);
 }
diff --git a/Engine/Source/Runtime/Core/Public/Containers/StringConv.h b/Engine/Source/Runtime/Core/Public/Containers/StringConv.h
--- a/Engine/Source/Runtime/Core/Public/Containers/StringConv.h
+++ b/Engine/Source/Runtime/Core/Public/Containers/StringConv.h
@@ -754,7 +754,7 @@ public:
 	 *
 	 * @return A const pointer to the converted string.
 	 */
-	FORCEINLINE const ToType* Get() const
+	FORCEINLINE const ToType* Get() const UE_LIFETIMEBOUND
 	{
 		return Ptr;
 	}
@@ -882,10 +882,11 @@ public:
 	/**
 	 * Accessor for the string.
 	 * @note The string may not be null-terminated if constructed from an explicitly sized buffer that didn't include the null-terminator.
+	 * @note Setting to UE_LIFETIMEBOUND even though TStringPointer doesn't own the pointer, because we want StringCast to give consistent warnings no matter what implementation choice it makes
 	 *
 	 * @return A const pointer to the string.
 	 */
-	FORCEINLINE const ToType* Get() const
+	FORCEINLINE const ToType* Get() const UE_LIFETIMEBOUND
 	{
 		return Ptr;
 	}
@@ -1125,7 +1126,7 @@ public:
 		FPlatformString::Convert(Dest, DestLen, Source, SrcLen);
 	}
 
-	FORCEINLINE FromType* Get()
+	FORCEINLINE FromType* Get() UE_LIFETIMEBOUND
 	{
 		return (FromType*)AllocatorType::GetAllocation();
 	}
@@ -1151,7 +1152,8 @@ public:
 	{
 	}
 
-	FORCEINLINE T* Get() const
+	// @note Setting to UE_LIFETIMEBOUND even though TPassthruPointer doesn't own the pointer, because we want StringMemoryPassthru to give consistent warnings no matter what implementation choice it makes
+	FORCEINLINE T* Get() const UE_LIFETIMEBOUND
 	{
 		return Ptr;
 	}
diff --git a/Engine/Source/Runtime/Core/Public/Containers/StringView.h b/Engine/Source/Runtime/Core/Public/Containers/StringView.h
--- a/Engine/Source/Runtime/Core/Public/Containers/StringView.h
+++ b/Engine/Source/Runtime/Core/Public/Containers/StringView.h
@@ -86,14 +86,14 @@ public:
 	constexpr TStringView() = default;
 
 	/** Construct a view of the null-terminated string pointed to by InData. */
-	constexpr inline TStringView(const CharType* InData)
+	constexpr inline TStringView(const CharType* InData UE_LIFETIMEBOUND)
 		: DataPtr(InData)
 		, Size(InData ? TCString<CharType>::Strlen(InData) : 0)
 	{
 	}
 
 	/** Construct a view of InSize characters beginning at InData. */
-	constexpr inline TStringView(const CharType* InData, int32 InSize)
+	constexpr inline TStringView(const CharType* InData UE_LIFETIMEBOUND, int32 InSize)
 		: DataPtr(InData)
 		, Size(InSize)
 	{
@@ -109,7 +109,7 @@ public:
 			>::Value
 		>* = nullptr
 	>
-	constexpr inline TStringView(const OtherCharType* InData)
+	constexpr inline TStringView(const OtherCharType* InData UE_LIFETIMEBOUND)
 		: DataPtr((const CharType*)InData)
 		, Size(InData ? TCString<CharType>::Strlen((const CharType*)InData) : 0)
 	{
@@ -125,7 +125,7 @@ public:
 			>::Value
 		>* = nullptr
 	>
-	constexpr inline TStringView(const OtherCharType* InData, int32 InSize)
+	constexpr inline TStringView(const OtherCharType* InData UE_LIFETIMEBOUND, int32 InSize)
 		: DataPtr((const CharType*)InData)
 		, Size(InSize)
 	{
@@ -459,7 +459,7 @@ constexpr inline auto MakeStringView(CharPtrOrRangeType&& CharPtrOrRange UE_LIFE
 }
 
 template <typename CharPtrType>
-constexpr inline auto MakeStringView(CharPtrType&& CharPtr, int32 Size) -> decltype(TStringView(Forward<CharPtrType>(CharPtr), Size))
+constexpr inline auto MakeStringView(CharPtrType&& CharPtr UE_LIFETIMEBOUND, int32 Size) -> decltype(TStringView(Forward<CharPtrType>(CharPtr), Size))
 {
 	return TStringView(Forward<CharPtrType>(CharPtr), Size);
 }
