Add more lifetime support for FString

diff --git a/Engine/Source/Runtime/Core/Public/Containers/UnrealString.h b/Engine/Source/Runtime/Core/Public/Containers/UnrealString.h
--- a/Engine/Source/Runtime/Core/Public/Containers/UnrealString.h
+++ b/Engine/Source/Runtime/Core/Public/Containers/UnrealString.h
@@ -212,7 +212,7 @@ public:
 	 * @param Index into string
 	 * @return Character at Index
 	 */
-	FORCEINLINE TCHAR& operator[]( int32 Index )
+	FORCEINLINE TCHAR& operator[]( int32 Index ) UE_LIFETIMEBOUND
 	{
 		checkf(IsValidIndex(Index), TEXT("String index out of bounds: Index %i from a string with a length of %i"), Index, Len());
 		return Data.GetData()[Index];
@@ -224,7 +224,7 @@ public:
 	 * @param Index into string
 	 * @return const Character at Index
 	 */
-	FORCEINLINE const TCHAR& operator[]( int32 Index ) const
+	FORCEINLINE const TCHAR& operator[]( int32 Index ) const UE_LIFETIMEBOUND
 	{
 		checkf(IsValidIndex(Index), TEXT("String index out of bounds: Index %i from a string with a length of %i"), Index, Len());
 		return Data.GetData()[Index];
@@ -320,7 +320,7 @@ public:
 	 *
 	 * @Return Pointer to Array of TCHAR if Num, otherwise the empty string
 	 */
-	UE_NODISCARD FORCEINLINE const TCHAR* operator*() const
+	UE_NODISCARD FORCEINLINE const TCHAR* operator*() const UE_LIFETIMEBOUND
 	{
 		return Data.Num() ? Data.GetData() : TEXT("");
 	}
@@ -331,13 +331,13 @@ public:
 	 * @warning: Operations on the TArray<*CHAR> can be unsafe, such as adding
 	 *		non-terminating 0's or removing the terminating zero.
 	 */
-	UE_NODISCARD FORCEINLINE DataType& GetCharArray()
+	UE_NODISCARD FORCEINLINE DataType& GetCharArray() UE_LIFETIMEBOUND
 	{
 		return Data;
 	}
 
 	/** Get string as const array of TCHARS */
-	UE_NODISCARD FORCEINLINE const DataType& GetCharArray() const
+	UE_NODISCARD FORCEINLINE const DataType& GetCharArray() const UE_LIFETIMEBOUND
 	{
 		return Data;
 	}
