Logging: Added support for structured logging

- Added UE_SLOG(Category, Verbosity, "Example with {Count} {Name}.", 2, "fields");
- Added UE_SLOG_NAMED(Category, Verbosity, "Example with {Count} {Name}.", ("Count", 2), ("Name", "fields"));
- Added FOutputDevice::SerializeRecord to receive structured logging. Note that FOutputDeviceRedirector does not forward these to its output devices yet.
- Added (currently) private FLogTemplate to parse, store, and apply structured log format strings.
- Trace receives structured logs pre-formatted until new trace events are added to support fields.
- Use of Fatal verbosity with structured logging is disabled until its implementation is complete.

Approximate cost of UE_SLOG[_NAMED] in Development Win64:

Generated Code by Field/Param Count
- 0: 4 instructions, 21 bytes (UE_LOG is 6i/33b & 32i/136b)
- 1: 5 instructions, 24 bytes inline & 11 instructions, 54 bytes non-inline (UE_LOG is 13i/61b & 38i/155b)
- 2: 8 instructions, 24 bytes inline & 16 instructions, 65 bytes non-inline (UE_LOG is 16i/76b & 44i/186b)
- 3: 10 instructions, 51 bytes inline & 26 instructions, 127 bytes non-inline (UE_LOG is 19i/91b & 51i/219b)

Build Impact (Average of 1,000 Logs with 0-9 Fields/Params)
- Link Time: ~0ms (UE_LOG is ~0ms)
- Compile Time: 0.83ms (UE_LOG is 1.28ms)
- EXE Size: 364 bytes (UE_LOG is 436 bytes)
- PDB Size: 2,163 bytes (UE_LOG is 2,449 bytes)
- OBJ Size: 4,469 bytes (UE_LOG is 4,745 bytes)

diff --git a/Engine/Source/Runtime/Core/Private/Logging/StructuredLog.cpp b/Engine/Source/Runtime/Core/Private/Logging/StructuredLog.cpp
new file mode 100644
--- /dev/null
+++ b/Engine/Source/Runtime/Core/Private/Logging/StructuredLog.cpp
@@ -0,0 +1,575 @@
+// Copyright Epic Games, Inc. All Rights Reserved.
+
+#include "Logging/StructuredLog.h"
+
+#include "Algo/Accumulate.h"
+#include "Algo/AllOf.h"
+#include "Algo/NoneOf.h"
+#include "Containers/StringConv.h"
+#include "Containers/StringView.h"
+#include "CoreGlobals.h"
+#include "HAL/PlatformTime.h"
+#include "Logging/LogTrace.h"
+#include "Misc/AsciiSet.h"
+#include "Misc/DateTime.h"
+#include "Misc/FeedbackContext.h"
+#include "Misc/OutputDevice.h"
+#include "Misc/OutputDeviceRedirector.h"
+#include "Misc/ScopeExit.h"
+#include "Misc/StringBuilder.h"
+#include "Serialization/CompactBinarySerialization.h"
+#include "Serialization/CompactBinaryValue.h"
+#include "Serialization/VarInt.h"
+
+namespace UE::Logging::Private
+{
+
+static constexpr FAsciiSet ValidLogFieldName("0123456789ABCDEFGHIJKLMNOPQRSTUVWXYZabcdefghijklmnopqrstuvwxyz_");
+
+///////////////////////////////////////////////////////////////////////////////////////////////////////////////////////
+
+struct FLogTemplateOp
+{
+	enum EOpCode : int32 { OpEnd, OpSkip, OpText, OpName, OpIndex, OpCount };
+
+	static constexpr int32 ValueShift = 3;
+	static_assert(OpCount <= (1 << ValueShift));
+
+	EOpCode Code = OpEnd;
+	int32 Value = 0;
+
+	inline int32 GetSkipSize() const { return Code == OpIndex ? 0 : Value; }
+
+	static inline FLogTemplateOp Load(const uint8*& Data);
+	static inline uint32 SaveSize(const FLogTemplateOp& Op) { return MeasureVarUInt(Encode(Op)); }
+	static inline void Save(const FLogTemplateOp& Op, uint8*& Data);
+	static inline uint64 Encode(const FLogTemplateOp& Op) { return uint64(Op.Code) | (uint64(Op.Value) << ValueShift); }
+	static inline FLogTemplateOp Decode(uint64 Value) { return {EOpCode(Value & ((1 << ValueShift) - 1)), int32(Value >> ValueShift)}; }
+};
+
+inline FLogTemplateOp FLogTemplateOp::Load(const uint8*& Data)
+{
+	uint32 ByteCount = 0;
+	ON_SCOPE_EXIT { Data += ByteCount; };
+	return Decode(ReadVarUInt(Data, ByteCount));
+}
+
+inline void FLogTemplateOp::Save(const FLogTemplateOp& Op, uint8*& Data)
+{
+	Data += WriteVarUInt(Encode(Op), Data);
+}
+
+///////////////////////////////////////////////////////////////////////////////////////////////////////////////////////
+
+static void LogFieldValue(FStringBuilderBase& Out, const FCbFieldView& Field)
+{
+	switch (FCbValue Accessor = Field.GetValue(); Accessor.GetType())
+	{
+	case ECbFieldType::Null:
+		Out.Append(TEXTVIEW("null"));
+		break;
+	case ECbFieldType::Object:
+	case ECbFieldType::UniformObject:
+	case ECbFieldType::Array:
+	case ECbFieldType::UniformArray:
+	case ECbFieldType::Binary:
+		CompactBinaryToCompactJson(Field, Out);
+		break;
+	case ECbFieldType::String:
+		Out.Append(Accessor.AsString());
+		break;
+	case ECbFieldType::IntegerPositive:
+		Out << Accessor.AsIntegerPositive();
+		break;
+	case ECbFieldType::IntegerNegative:
+		Out << Accessor.AsIntegerNegative();
+		break;
+	case ECbFieldType::Float32:
+	case ECbFieldType::Float64:
+		CompactBinaryToCompactJson(Field, Out);
+		break;
+	case ECbFieldType::BoolFalse:
+		Out.Append(TEXTVIEW("false"));
+		break;
+	case ECbFieldType::BoolTrue:
+		Out.Append(TEXTVIEW("true"));
+		break;
+	case ECbFieldType::ObjectAttachment:
+	case ECbFieldType::BinaryAttachment:
+		Out << Accessor.AsAttachment();
+		break;
+	case ECbFieldType::Hash:
+		Out << Accessor.AsHash();
+		break;
+	case ECbFieldType::Uuid:
+		Out << Accessor.AsUuid();
+		break;
+	case ECbFieldType::DateTime:
+		Out << FDateTime(Accessor.AsDateTimeTicks()).ToIso8601();
+		break;
+	case ECbFieldType::TimeSpan:
+	{
+		const FTimespan Span(Accessor.AsTimeSpanTicks());
+		if (Span.GetDays() == 0)
+		{
+			Out << Span.ToString(TEXT("%h:%m:%s.%n"));
+		}
+		else
+		{
+			Out << Span.ToString(TEXT("%d.%h:%m:%s.%n"));
+		}
+		break;
+	}
+	case ECbFieldType::ObjectId:
+		Out << Accessor.AsObjectId();
+		break;
+	case ECbFieldType::CustomById:
+	case ECbFieldType::CustomByName:
+		CompactBinaryToCompactJson(Field, Out);
+		break;
+	default:
+		checkNoEntry();
+		break;
+	}
+}
+
+} // UE::Logging::Private
+
+namespace UE
+{
+
+///////////////////////////////////////////////////////////////////////////////////////////////////////////////////////
+
+class FLogTemplate
+{
+	using FLogField = Logging::Private::FLogField;
+
+public:
+	static FLogTemplate* Create(const TCHAR* Format, const FLogField* Fields = nullptr, int32 FieldCount = 0);
+	static void Destroy(FLogTemplate* Template);
+
+	void FormatTo(FStringBuilderBase& Out, const TCHAR* Format, const FCbObjectView& Fields) const;
+
+private:
+	FLogTemplate() = default;
+	~FLogTemplate() = default;
+	FLogTemplate(const FLogTemplate&) = delete;
+	FLogTemplate& operator=(const FLogTemplate&) = delete;
+};
+
+FLogTemplate* FLogTemplate::Create(const TCHAR* Format, const FLogField* Fields, const int32 FieldCount)
+{
+	using namespace Logging::Private;
+
+	const TConstArrayView<FLogField> FieldsView(Fields, FieldCount);
+	const bool bFindFields = !!Fields;
+	const bool bPositional = !FieldCount || Algo::NoneOf(FieldsView, &FLogField::Name);
+	checkf(bPositional || Algo::AllOf(FieldsView, &FLogField::Name),
+		TEXT("Log fields must be entirely named or entirely anonymous. [[%s]]"), Format);
+	checkf(bPositional || Algo::AllOf(FieldsView,
+		[](const FLogField& Field) { return *Field.Name && FAsciiSet::HasOnly(Field.Name, ValidLogFieldName); }),
+		TEXT("Log field names must match \"[A-Za-z0-9_]+\" in [[%s]]."), Format);
+
+	TArray<FLogTemplateOp, TInlineAllocator<16>> Ops;
+
+	int32 FormatFieldCount = 0;
+	int32 BracketSearchOffset = 0;
+	for (const TCHAR* TextStart = Format;;)
+	{
+		constexpr FAsciiSet Brackets("{}");
+		const TCHAR* const TextEnd = FAsciiSet::FindFirstOrEnd(TextStart + BracketSearchOffset, Brackets);
+		BracketSearchOffset = 0;
+
+		// Escaped "{{" or "}}"
+		if ((TextEnd[0] == TEXT('{') && TextEnd[1] == TEXT('{')) ||
+			(TextEnd[0] == TEXT('}') && TextEnd[1] == TEXT('}')))
+		{
+			// Only "{{" or "}}"
+			if (TextStart == TextEnd)
+			{
+				Ops.Add({FLogTemplateOp::OpSkip, 1});
+				TextStart = TextEnd + 1;
+				BracketSearchOffset = 1;
+			}
+			// Text and "{{" or "}}"
+			else
+			{
+				Ops.Add({FLogTemplateOp::OpText, UE_PTRDIFF_TO_INT32(1 + TextEnd - TextStart)});
+				Ops.Add({FLogTemplateOp::OpSkip, 1});
+				TextStart = TextEnd + 2;
+			}
+			continue;
+		}
+
+		// Text
+		if (TextStart != TextEnd)
+		{
+			Ops.Add({FLogTemplateOp::OpText, UE_PTRDIFF_TO_INT32(TextEnd - TextStart)});
+			TextStart = TextEnd;
+		}
+
+		// End
+		if (!*TextEnd)
+		{
+			Ops.Add({FLogTemplateOp::OpEnd});
+			break;
+		}
+
+		// Invalid '}'
+		checkf(*TextStart == TEXT('{'), TEXT("Log format has an unexpected '}' character. Use '}}' to escape it. [[%s]]"), Format);
+
+		// Field
+		const TCHAR* const FieldStart = TextStart;
+		const TCHAR* const FieldNameEnd = FAsciiSet::Skip(FieldStart + 1, ValidLogFieldName);
+		checkf(*FieldNameEnd, TEXT("Log format has an unterminated field reference. Use '{{' to escape '{' if needed. [[%s]]"), Format);
+		checkf(*FieldNameEnd == TEXT('}'), TEXT("Log format has invalid character '%c' in field name. Use '{{' to escape '{' if needed. Names must match \"[A-Za-z0-9_]+\". [[%s]]"), *FieldNameEnd, Format);
+		const TCHAR* const FieldEnd = FieldNameEnd + 1;
+		const int32 FieldLen = UE_PTRDIFF_TO_INT32(FieldEnd - FieldStart);
+		const int32 FieldNameLen = FieldLen - 2;
+
+		if (bFindFields && !bPositional)
+		{
+			bool bFoundField = false;
+			for (int32 BaseFieldIndex = 0; BaseFieldIndex < FieldCount; ++BaseFieldIndex)
+			{
+				const int32 FieldIndex = (FormatFieldCount + BaseFieldIndex) % FieldCount;
+				const ANSICHAR* FieldName = Fields[FieldIndex].Name;
+				if (FPlatformString::Strncmp(FieldName, FieldStart + 1, FieldNameLen) == 0 && !FieldName[FieldNameLen])
+				{
+					Ops.Add({FLogTemplateOp::OpIndex, FieldIndex});
+					bFoundField = true;
+					break;
+				}
+			}
+			checkf(bFoundField, TEXT("Log format requires field '%.*s' which was not provided. [[%s]]"),
+				FieldNameLen, FieldStart + 1, Format);
+		}
+
+		Ops.Add({FLogTemplateOp::OpName, FieldLen});
+		++FormatFieldCount;
+
+		TextStart = FieldEnd;
+	}
+
+	checkf(!bFindFields || !bPositional || FormatFieldCount == FieldCount,
+		TEXT("Log format requires %d fields and %d were provided. [[%s]]"), FormatFieldCount, FieldCount);
+
+	const uint32 TotalSize = Algo::TransformAccumulate(Ops, FLogTemplateOp::SaveSize, 0);
+	FLogTemplate* const Template = (FLogTemplate*)FMemory::Malloc(TotalSize);
+	uint8* Data = (uint8*)Template;
+	for (const FLogTemplateOp& Op : Ops)
+	{
+		FLogTemplateOp::Save(Op, Data);
+	}
+	return Template;
+}
+
+void FLogTemplate::Destroy(FLogTemplate* Template)
+{
+	FMemory::Free(Template);
+}
+
+void FLogTemplate::FormatTo(FStringBuilderBase& Out, const TCHAR* Format, const FCbObjectView& Fields) const
+{
+	using namespace Logging::Private;
+
+	auto FindField = [&Fields, It = Fields.CreateViewIterator(), Index = 0, Format](FAnsiStringView Name, int32 IndexHint = -1) mutable -> FCbFieldView&
+	{
+		if (IndexHint >= 0)
+		{
+			for (; Index < IndexHint && It; ++Index, ++It);
+			if (IndexHint < Index)
+			{
+				It = Fields.CreateViewIterator();
+				for (Index = 0; Index < IndexHint && It; ++Index, ++It);
+			}
+			if (IndexHint == Index && Name.Equals(It.GetName()))
+			{
+				return It;
+			}
+		}
+		const int32 PrevIndex = Index;
+		for (; It; ++Index, ++It)
+		{
+			if (Name.Equals(It.GetName()))
+			{
+				return It;
+			}
+		}
+		It = Fields.CreateViewIterator();
+		for (Index = 0; Index < PrevIndex && It; ++Index, ++It)
+		{
+			if (Name.Equals(It.GetName()))
+			{
+				return It;
+			}
+		}
+		checkf(false, TEXT("Log format requires field '%.*hs' which was not provided. [[%s]]"), Name.Len(), Name.GetData(), Format);
+		return It;
+	};
+
+	int32 FieldIndexHint = -1;
+	const uint8* NextOp = (const uint8*)this;
+	const TCHAR* NextFormat = Format;
+	for (;;)
+	{
+		const FLogTemplateOp Op = FLogTemplateOp::Load(NextOp);
+		switch (Op.Code)
+		{
+		case FLogTemplateOp::OpEnd:
+			return;
+		case FLogTemplateOp::OpText:
+			Out.Append(NextFormat, Op.Value);
+			break;
+		case FLogTemplateOp::OpIndex:
+			FieldIndexHint = Op.Value;
+			break;
+		case FLogTemplateOp::OpName:
+			const auto Name = StringCast<ANSICHAR>(NextFormat + 1, Op.Value - 2);
+			LogFieldValue(Out, FindField(MakeStringView(Name.Get(), Name.Length()), FieldIndexHint));
+			FieldIndexHint = -1;
+			break;
+		}
+		NextFormat += Op.GetSkipSize();
+	}
+}
+
+///////////////////////////////////////////////////////////////////////////////////////////////////////////////////////
+
+FLogTime FLogTime::Now()
+{
+	FLogTime Time;
+	Time.UtcTicks = FDateTime::UtcNow().GetTicks();
+	Time.Cycles = FPlatformTime::Cycles64();
+	return Time;
+}
+
+FDateTime FLogTime::GetUtcTime() const
+{
+	return FDateTime(UtcTicks);
+}
+
+///////////////////////////////////////////////////////////////////////////////////////////////////////////////////////
+
+void FLogRecord::FormatMessageTo(FStringBuilderBase& Out) const
+{
+	if (UNLIKELY(!Format))
+	{
+		return;
+	}
+
+	if (LIKELY(Template))
+	{
+		return Template->FormatTo(Out, Format, Fields);
+	}
+
+	FLogTemplate* LocalTemplate = FLogTemplate::Create(Format);
+	LocalTemplate->FormatTo(Out, Format, Fields);
+	FLogTemplate::Destroy(LocalTemplate);
+}
+
+///////////////////////////////////////////////////////////////////////////////////////////////////////////////////////
+
+} // UE
+
+#if !NO_LOGGING
+
+namespace UE::Logging::Private
+{
+
+class FLogTemplateFieldIterator
+{
+public:
+	inline FLogTemplateFieldIterator(const FLogTemplate* Template, const TCHAR* Format)
+		: NextOp((const uint8*)Template)
+		, NextFormat(Format)
+	{
+		++*this;
+	}
+
+	FLogTemplateFieldIterator& operator++();
+	inline explicit operator bool() const { return !!NextOp; }
+	inline const FStringView& GetName() const { return Name; }
+
+private:
+	FStringView Name;
+	const uint8* NextOp = nullptr;
+	const TCHAR* NextFormat = nullptr;
+};
+
+FLogTemplateFieldIterator& FLogTemplateFieldIterator::operator++()
+{
+	using namespace Logging::Private;
+
+	while (NextOp)
+	{
+		FLogTemplateOp Op = FLogTemplateOp::Load(NextOp);
+		if (Op.Code == FLogTemplateOp::OpName)
+		{
+			Name = FStringView(NextFormat + 1, Op.Value - 2);
+			NextFormat += Op.GetSkipSize();
+			return *this;
+		}
+		if (Op.Code == FLogTemplateOp::OpEnd)
+		{
+			break;
+		}
+		NextFormat += Op.GetSkipSize();
+	}
+
+	NextOp = nullptr;
+	Name.Reset();
+	return *this;
+}
+
+///////////////////////////////////////////////////////////////////////////////////////////////////////////////////////
+
+struct FStaticLogDynamicDataManager
+{
+	std::atomic<FStaticLogDynamicData*> Head = nullptr;
+
+	~FStaticLogDynamicDataManager()
+	{
+		for (FStaticLogDynamicData* Data = Head.exchange(nullptr); Data; Data = Data->Next)
+		{
+			if (FLogTemplate* Template = Data->Template.exchange(nullptr))
+			{
+				FLogTemplate::Destroy(Template);
+			}
+		}
+	}
+};
+
+static FStaticLogDynamicDataManager GStaticLogDynamicDataManager;
+
+// Tracing the log happens in its own function because that allows stack space for the message to
+// be returned before calling into the output devices.
+FORCENOINLINE void LogToTrace(const FStaticLogRecord& Log, const FLogRecord& Record)
+{
+#if LOGTRACE_ENABLED
+	TStringBuilder<1024> Message;
+	Record.FormatMessageTo(Message);
+	FLogTrace::OutputLogMessageSimple(&Log, TEXT("%s"), *Message);
+#endif
+}
+
+// Serializing log fields to compact binary happens in its own function because that allows stack
+// space for the writer to be returned before calling into the output devices.
+FORCENOINLINE FCbObject SerializeLogFields(
+	const FStaticLogRecord& Log,
+	const FLogTemplate& Template,
+	const FLogField* Fields,
+	const int32 FieldCount)
+{
+	if (FieldCount == 0)
+	{
+		return FCbObject();
+	}
+
+	TCbWriter<1024> Writer;
+	Writer.BeginObject();
+
+	// Anonymous. Extract names from Template.
+	if (!Fields->Name)
+	{
+		FLogTemplateFieldIterator It(&Template, Log.Format);
+		for (const FLogField* FieldsEnd = Fields + FieldCount; Fields != FieldsEnd; ++Fields, ++It)
+		{
+			check(It);
+			const auto Name = StringCast<ANSICHAR>(It.GetName().GetData(), It.GetName().Len());
+			Fields->WriteValue(Writer.SetName(MakeStringView(Name.Get(), Name.Length())), Fields->Value);
+		}
+		check(!It);
+	}
+	// Named
+	else
+	{
+		for (const FLogField* FieldsEnd = Fields + FieldCount; Fields != FieldsEnd; ++Fields)
+		{
+			Fields->WriteValue(Writer.SetName(Fields->Name), Fields->Value);
+		}
+	}
+
+	Writer.EndObject();
+	return Writer.Save().AsObject();
+}
+
+FORCENOINLINE FLogTemplate& CreateLogTemplate(const FStaticLogRecord& Log, const FLogField* Fields, const int32 FieldCount)
+{
+#if LOGTRACE_ENABLED
+	FLogTrace::OutputLogMessageSpec(&Log, &Log.Category, Log.Verbosity, Log.File, Log.Line, TEXT("%s"));
+#endif
+
+	FLogTemplate* NewTemplate = FLogTemplate::Create(Log.Format, Fields, FieldCount);
+	if (FLogTemplate* ExistingTemplate = nullptr;
+		UNLIKELY(!Log.DynamicData.Template.compare_exchange_strong(ExistingTemplate, NewTemplate, std::memory_order_release, std::memory_order_acquire)))
+	{
+		FLogTemplate::Destroy(NewTemplate);
+		return *ExistingTemplate;
+	}
+
+	// Register the template to destroy on exit.
+	for (FStaticLogDynamicData* Head = GStaticLogDynamicDataManager.Head.load(std::memory_order_relaxed);;)
+	{
+		Log.DynamicData.Next = Head;
+		if (GStaticLogDynamicDataManager.Head.compare_exchange_weak(Head, &Log.DynamicData, std::memory_order_release, std::memory_order_relaxed))
+		{
+			break;
+		}
+	}
+
+	return *NewTemplate;
+}
+
+inline FLogTemplate& EnsureLogTemplate(const FStaticLogRecord& Log, const FLogField* Fields, const int32 FieldCount)
+{
+	if (FLogTemplate* Template = Log.DynamicData.Template.load(std::memory_order_acquire))
+	{
+		return *Template;
+	}
+	return CreateLogTemplate(Log, Fields, FieldCount);
+}
+
+void LogWithFieldArray(const FStaticLogRecord& Log, const FLogField* Fields, const int32 FieldCount)
+{
+	FLogTemplate& Template = EnsureLogTemplate(Log, Fields, FieldCount);
+
+	FLogRecord Record;
+	Record.SetFormat(Log.Format);
+	Record.SetTemplate(&Template);
+	Record.SetFields(SerializeLogFields(Log, Template, Fields, FieldCount));
+	Record.SetFile(Log.File);
+	Record.SetLine(Log.Line);
+	Record.SetCategory(Log.Category.GetCategoryName());
+	Record.SetVerbosity(Log.Verbosity);
+	Record.SetTime(FLogTime::Now());
+
+#if LOGTRACE_ENABLED
+	if (UE_TRACE_CHANNELEXPR_IS_ENABLED(LogChannel))
+	{
+		LogToTrace(Log, Record);
+	}
+#endif
+
+	FOutputDevice* LogOverride = nullptr;
+	switch (Log.Verbosity)
+	{
+	case ELogVerbosity::Error:
+	case ELogVerbosity::Warning:
+	case ELogVerbosity::Display:
+		return GWarn->SerializeRecord(Record);
+	default:
+		return GLog->SerializeRecord(Record);
+	}
+}
+
+void LogWithNoFields(const FStaticLogRecord& Log)
+{
+	// A non-null field pointer enables field validation in FLogTemplate::Create.
+	static constexpr FLogField EmptyField{};
+	LogWithFieldArray(Log, &EmptyField, 0);
+}
+
+} // UE::Logging::Private
+
+#endif // !NO_LOGGING
diff --git a/Engine/Source/Runtime/Core/Private/Misc/OutputDevice.cpp b/Engine/Source/Runtime/Core/Private/Misc/OutputDevice.cpp
--- a/Engine/Source/Runtime/Core/Private/Misc/OutputDevice.cpp
+++ b/Engine/Source/Runtime/Core/Private/Misc/OutputDevice.cpp
@@ -4,9 +4,12 @@
 
 #include "Containers/UnrealString.h"
 #include "CoreGlobals.h"
+#include "HAL/PlatformTime.h"
 #include "Internationalization/Text.h"
 #include "Logging/LogMacros.h"
 #include "Logging/LogScopedCategoryAndVerbosityOverride.h"
+#include "Logging/StructuredLog.h"
+#include "Misc/StringBuilder.h"
 #include "Misc/VarargsHelper.h"
 #include "Templates/UnrealTemplate.h"
 #include "UObject/NameTypes.h"
@@ -67,8 +70,12 @@ FORCENOINLINE void FOutputDevice::LogfImpl(const TCHAR* Fmt, ...)
 	GROWABLE_LOGF(Serialize(Buffer, TLS->Verbosity, TLS->Category))
 }
 
-
+void FOutputDevice::SerializeRecord(const UE::FLogRecord& Record)
+{
+	TStringBuilder<512> Text;
+	Record.FormatMessageTo(Text);
+	Serialize(*Text, Record.GetVerbosity(), Record.GetCategory(), FPlatformTime::ToSeconds64(Record.GetTime().GetCycles()));
+}
 
 /** Critical errors. */
 CORE_API FOutputDeviceError* GError = NULL;
-
diff --git a/Engine/Source/Runtime/Core/Public/Logging/StructuredLog.h b/Engine/Source/Runtime/Core/Public/Logging/StructuredLog.h
new file mode 100644
--- /dev/null
+++ b/Engine/Source/Runtime/Core/Public/Logging/StructuredLog.h
@@ -0,0 +1,342 @@
+// Copyright Epic Games, Inc. All Rights Reserved.
+
+#pragma once
+
+#include "Containers/StringFwd.h"
+#include "Logging/LogCategory.h"
+#include "Logging/LogMacros.h"
+#include "Logging/LogVerbosity.h"
+#include "Serialization/CompactBinary.h"
+#include "Serialization/CompactBinaryWriter.h"
+#include "Templates/IsArrayOrRefOfType.h"
+
+#include <atomic>
+#include <type_traits>
+
+#define UE_API CORE_API
+
+#if !NO_LOGGING
+
+/**
+ * Records a structured log event if this category is active at this level of verbosity.
+ *
+ * Example:
+ * UE_SLOG(LogCore, Warning, "Loading '{Name}' failed with error {Error}", Package->GetName(), ErrorCode);
+ *
+ * Field names must match "[A-Za-z0-9_]+" and must be unique within this log event.
+ * Field values will be serialized using operator<<(FCbWriter&, FieldType).
+ * The field values must exactly match the fields referenced by Format.
+ *
+ * @param CategoryName   Name of a log category declared by DECLARE_LOG_CATEGORY_*.
+ * @param Verbosity      Name of a log verbosity level from ELogVerbosity.
+ * @param Format         Format string in the style of FLogTemplate.
+ * @param FieldValues*   Zero or more field values that correspond to the names in Format.
+ */
+#define UE_SLOG(CategoryName, Verbosity, Format, ...) UE_PRIVATE_SLOG(CategoryName, Verbosity, Format, Anonymous, ##__VA_ARGS__)
+
+/**
+ * Records a structured log event if this category is active at this level of verbosity.
+ *
+ * Example:
+ * UE_SLOG_NAMED(LogCore, Warning, "Loading '{PackageName}' failed with error {Error}",
+ *     ("Name", Package->GetName()), ("Error", ErrorCode), ("Flags", LoadFlags));
+ *
+ * Field names must match "[A-Za-z0-9_]+" and must be unique within this log event.
+ * Field values will be serialized using operator<<(FCbWriter&, FieldType).
+ * The fields must contain every field referenced by Format. Order is irrelevant and extra fields are permitted.
+ *
+ * @param CategoryName   Name of a log category declared by DECLARE_LOG_CATEGORY_*.
+ * @param Verbosity      Name of a log verbosity level from ELogVerbosity.
+ * @param Format         Format string in the style of FLogTemplate.
+ * @param Fields[0-9]    Zero to nine fields in the format ("Name", Value).
+ */
+#define UE_SLOG_NAMED(CategoryName, Verbosity, Format, ...) UE_PRIVATE_SLOG_CALL(UE_SLOG_NAMED_EX, (CategoryName, Verbosity, Format UE_PRIVATE_SLOG_FIELDS(__VA_ARGS__)))
+
+/**
+ * Records a structured log event if this category is active at this level of verbosity.
+ *
+ * Example:
+ * UE_SLOG_NAMED_EX(LogCore, Warning, "Loading '{PackageName}' failed with error {Error}",
+ *     UE_SLOG_FIELD("Name", Package->GetName()), UE_SLOG_FIELD("Error", ErrorCode), UE_SLOG_FIELD("Flags", LoadFlags));
+ *
+ * Same as UE_SLOG_NAMED and works for any number of fields.
+ * Fields must be written as UE_SLOG_FIELD("Name", Value).
+ */
+#define UE_SLOG_NAMED_EX(CategoryName, Verbosity, Format, ...) UE_PRIVATE_SLOG(CategoryName, Verbosity, Format, Named, ##__VA_ARGS__)
+
+/**
+ * Expands to a named structured log field. UE_SLOG_FIELD("Name", Value)
+ */
+#define UE_SLOG_FIELD(Name, Value) ::UE::Logging::Private::CheckFieldName(Name), Value
+
+#else // #if !NO_LOGGING
+
+#define UE_SLOG(...)
+#define UE_SLOG_NAMED(...)
+#define UE_SLOG_NAMED_EX(...)
+#define UE_SLOG_FIELD(Name, Value)
+
+#endif
+
+struct FDateTime;
+
+namespace UE
+{
+
+/** Template format is "Text with {Fields} embedded {Like}{This}. {{Double to escape.}}" */
+class FLogTemplate;
+
+/**
+ * Time that a log event occurred.
+ *
+ * This abstraction exists to hide the multiple time representations.
+ * When cycles can be converted to FDateTime, this can stop storing ticks.
+ */
+class FLogTime
+{
+public:
+	UE_API static FLogTime Now();
+
+	constexpr FLogTime() = default;
+
+	/** Returns the cycle count. 0 when the time was not set. */
+	constexpr uint64 GetCycles() const { return Cycles; }
+
+	/** Returns the UTC time. 0 ticks when the time was not set. */
+	UE_API FDateTime GetUtcTime() const;
+
+private:
+	/** Cycles from FPlatformTime::Cycles64() */
+	uint64 Cycles = 0;
+	/** Ticks from FDateTime::UtcNow() */
+	int64 UtcTicks = 0;
+};
+
+/**
+ * Record of a log event.
+ */
+class FLogRecord
+{
+public:
+	/** The optional name of the category for the log record. None when omitted. */
+	const FName& GetCategory() const { return Category; }
+	void SetCategory(const FName& InCategory) { Category = InCategory; }
+
+	/** The verbosity level of the log record. Must be a valid level with no flags or special values. */
+	ELogVerbosity::Type GetVerbosity() const { return Verbosity; }
+	void SetVerbosity(ELogVerbosity::Type InVerbosity) { Verbosity = InVerbosity; }
+
+	/** The time at which the log record was created. */
+	const FLogTime& GetTime() const { return Time; }
+	void SetTime(const FLogTime& InTime) { Time = InTime; }
+
+	/** The format string that serves as the message for the log record. Example: TEXT("FieldName is {FieldName}") */
+	const TCHAR* GetFormat() const { return Format; }
+	void SetFormat(const TCHAR* InFormat) { Format = InFormat; }
+
+	/** The optional template for the format string. */
+	const FLogTemplate* GetTemplate() const { return Template; }
+	void SetTemplate(const FLogTemplate* InTemplate) { Template = InTemplate; }
+
+	/** The fields referenced by the format string, along with optional additional fields. */
+	const FCbObject& GetFields() const { return Fields; }
+	void SetFields(FCbObject&& InFields) { Fields = MoveTemp(InFields); }
+
+	/** The optional source file path for the code that created the log record. Null when omitted. */
+	const ANSICHAR* GetFile() const { return File; }
+	void SetFile(const ANSICHAR* InFile) { File = InFile; }
+
+	/** The optional source line number for the code that created the log record. 0 when omitted. */
+	int32 GetLine() const { return Line; }
+	void SetLine(int32 InLine) { Line = InLine; }
+
+	/** Formats the message using the format, template, and fields. */
+	UE_API void FormatMessageTo(FStringBuilderBase& Out) const;
+
+private:
+	const TCHAR* Format = nullptr;
+	const ANSICHAR* File = nullptr;
+	int32 Line = 0;
+	FName Category;
+	ELogVerbosity::Type Verbosity = ELogVerbosity::Log;
+	FLogTime Time;
+	FCbObject Fields;
+	const FLogTemplate* Template = nullptr;
+};
+
+} // UE
+
+namespace UE::Logging::Private
+{
+
+/** Data about a static log that is created on-demand. */
+struct FStaticLogDynamicData
+{
+	std::atomic<FLogTemplate*> Template;
+	FStaticLogDynamicData* Next = nullptr;
+};
+
+/** Data about a static log that is constant for every occurrence. */
+struct FStaticLogRecord
+{
+	const FLogCategoryBase& Category;
+	const TCHAR* Format = nullptr;
+	const ANSICHAR* File = nullptr;
+	int32 Line = 0;
+	ELogVerbosity::Type Verbosity = ELogVerbosity::Log;
+	FStaticLogDynamicData& DynamicData;
+};
+
+struct FLogField
+{
+	using FWriteFn = FCbWriter& (FCbWriter& Writer, const void* Value);
+
+	const ANSICHAR* Name;
+	const void* Value;
+	FWriteFn* WriteValue;
+
+	template <typename ValueType>
+	static FCbWriter& Write(FCbWriter& Writer, const void* Value)
+	{
+		return Writer << *(const ValueType*)Value;
+	}
+};
+
+} // UE::Logging::Private
+
+#if !NO_LOGGING
+
+namespace UE::Logging::Private
+{
+
+UE_API void LogWithNoFields(const FStaticLogRecord& Log);
+UE_API void LogWithFieldArray(const FStaticLogRecord& Log, const FLogField* Fields, int32 FieldCount);
+
+/** Creates log fields with no name from value arguments. */
+struct FLogFieldsAnonymous
+{
+	template <typename... FieldArgTypes>
+	static constexpr int32 GetCount()
+	{
+		return sizeof...(FieldArgTypes);
+	}
+
+	static inline void Create(FLogField* Fields)
+	{
+	}
+
+	template <typename ValueType, typename... FieldArgTypes>
+	static inline void Create(FLogField* Fields, const ValueType& Value, FieldArgTypes&&... FieldArgs)
+	{
+		new(Fields) FLogField{nullptr, &Value, FLogField::Write<std::remove_reference_t<ValueType>>};
+		Create(Fields + 1, Forward<FieldArgTypes>(FieldArgs)...);
+	}
+};
+
+/** Creates log fields from pairs of name/value arguments. */
+struct FLogFieldsNamed
+{
+	template <typename... FieldArgTypes>
+	static constexpr int32 GetCount()
+	{
+		constexpr int32 FieldArgsCount = sizeof...(FieldArgTypes);
+		static_assert(FieldArgsCount % 2 == 0);
+		return FieldArgsCount / 2;
+	}
+
+	static inline void Create(FLogField* Fields)
+	{
+	}
+
+	template <typename ValueType, typename... FieldArgTypes>
+	static inline void Create(FLogField* Fields, const ANSICHAR* Name, const ValueType& Value, FieldArgTypes&&... FieldArgs)
+	{
+		new(Fields) FLogField{Name, &Value, FLogField::Write<std::remove_reference_t<ValueType>>};
+		Create(Fields + 1, Forward<FieldArgTypes>(FieldArgs)...);
+	}
+};
+
+/** Log with fields created from the arguments, which may be values or pairs of name/value. */
+template <const FStaticLogRecord& Log, typename LogFieldsType, typename... FieldArgTypes>
+FORCENOINLINE UE_DEBUG_SECTION void LogWithFields(FieldArgTypes&&... FieldArgs)
+{
+	constexpr int32 FieldCount = LogFieldsType::template GetCount<FieldArgTypes...>();
+	static_assert(FieldCount > 0);
+	FLogField Fields[FieldCount];
+	LogFieldsType::Create(Fields, Forward<FieldArgTypes>(FieldArgs)...);
+	LogWithFieldArray(Log, Fields, FieldCount);
+}
+
+/** Log if the category is active at this level of verbosity. */
+template <const FStaticLogRecord& Log, typename LogCategoryType, ELogVerbosity::Type Verbosity, typename LogFieldsType, typename... FieldArgTypes>
+inline void LogIfActive(FieldArgTypes&&... FieldArgs)
+{
+	static_assert(Verbosity != ELogVerbosity::Fatal, "Fatal verbosity is not supported by this API at this time.");
+	static_assert((Verbosity & ELogVerbosity::VerbosityMask) < ELogVerbosity::NumVerbosity && Verbosity > 0, "Verbosity must be constant and in range.");
+
+	if constexpr (
+		(Verbosity & ELogVerbosity::VerbosityMask) <= ELogVerbosity::COMPILED_IN_MINIMUM_VERBOSITY &&
+		(Verbosity & ELogVerbosity::VerbosityMask) <= LogCategoryType::CompileTimeVerbosity)
+	{
+		if (Verbosity == ELogVerbosity::Fatal || !Log.Category.IsSuppressed(Verbosity))
+		{
+			if constexpr (sizeof...(FieldArgTypes))
+			{
+				LogWithFields<Log, LogFieldsType>(Forward<FieldArgTypes>(FieldArgs)...);
+			}
+			else
+			{
+				LogWithNoFields(Log);
+			}
+		}
+	}
+
+	if constexpr (Verbosity == ELogVerbosity::Fatal)
+	{
+		UE_DEBUG_BREAK_AND_PROMPT_FOR_REMOTE();
+		FDebug::ProcessFatalError(PLATFORM_RETURN_ADDRESS());
+		CA_ASSUME(false);
+	}
+}
+
+/** Verify that the name is likely a string literal and forward it on. */
+template <typename NameType>
+constexpr decltype(auto) CheckFieldName(NameType&& Name)
+{
+	static_assert(TIsArrayOrRefOfType<NameType, ANSICHAR>::Value, "Name must be an ANSICHAR string literal.");
+	return Forward<NameType>(Name);
+}
+
+} // UE::Logging::Private
+
+#define UE_PRIVATE_SLOG(CategoryName, Verbosity, Format, FieldsType, ...) \
+	do \
+	{ \
+		static ::UE::Logging::Private::FStaticLogDynamicData DynamicData; \
+		static constexpr ::UE::Logging::Private::FStaticLogRecord Log{CategoryName, TEXT(Format), __FILE__, __LINE__, ::ELogVerbosity::Verbosity, DynamicData}; \
+		::UE::Logging::Private::LogIfActive<Log, FLogCategory##CategoryName, ::ELogVerbosity::Verbosity, ::UE::Logging::Private::FLogFields##FieldsType>(__VA_ARGS__); \
+	} \
+	while (false)
+
+#define UE_PRIVATE_SLOG_FIELDS(...) UE_PRIVATE_SLOG_CALL(PREPROCESSOR_JOIN(UE_PRIVATE_SLOG_FIELDS_, UE_PRIVATE_SLOG_COUNT(__VA_ARGS__)), (__VA_ARGS__))
+
+#define UE_PRIVATE_SLOG_FIELDS_0()
+#define UE_PRIVATE_SLOG_FIELDS_1(A)                 , UE_SLOG_FIELD A
+#define UE_PRIVATE_SLOG_FIELDS_2(A,B)               , UE_SLOG_FIELD A, UE_SLOG_FIELD B
+#define UE_PRIVATE_SLOG_FIELDS_3(A,B,C)             , UE_SLOG_FIELD A, UE_SLOG_FIELD B, UE_SLOG_FIELD C
+#define UE_PRIVATE_SLOG_FIELDS_4(A,B,C,D)           , UE_SLOG_FIELD A, UE_SLOG_FIELD B, UE_SLOG_FIELD C, UE_SLOG_FIELD D
+#define UE_PRIVATE_SLOG_FIELDS_5(A,B,C,D,E)         , UE_SLOG_FIELD A, UE_SLOG_FIELD B, UE_SLOG_FIELD C, UE_SLOG_FIELD D, UE_SLOG_FIELD E
+#define UE_PRIVATE_SLOG_FIELDS_6(A,B,C,D,E,F)       , UE_SLOG_FIELD A, UE_SLOG_FIELD B, UE_SLOG_FIELD C, UE_SLOG_FIELD D, UE_SLOG_FIELD E, UE_SLOG_FIELD F
+#define UE_PRIVATE_SLOG_FIELDS_7(A,B,C,D,E,F,G)     , UE_SLOG_FIELD A, UE_SLOG_FIELD B, UE_SLOG_FIELD C, UE_SLOG_FIELD D, UE_SLOG_FIELD E, UE_SLOG_FIELD F, UE_SLOG_FIELD G
+#define UE_PRIVATE_SLOG_FIELDS_8(A,B,C,D,E,F,G,H)   , UE_SLOG_FIELD A, UE_SLOG_FIELD B, UE_SLOG_FIELD C, UE_SLOG_FIELD D, UE_SLOG_FIELD E, UE_SLOG_FIELD F, UE_SLOG_FIELD G, UE_SLOG_FIELD H
+#define UE_PRIVATE_SLOG_FIELDS_9(A,B,C,D,E,F,G,H,I) , UE_SLOG_FIELD A, UE_SLOG_FIELD B, UE_SLOG_FIELD C, UE_SLOG_FIELD D, UE_SLOG_FIELD E, UE_SLOG_FIELD F, UE_SLOG_FIELD G, UE_SLOG_FIELD H, UE_SLOG_FIELD I
+
+#define UE_PRIVATE_SLOG_COUNT(...) UE_PRIVATE_SLOG_CALL(UE_PRIVATE_SLOG_COUNT_IMPL, (_, ##__VA_ARGS__, 9,8,7,6,5,4,3,2,1,0))
+#define UE_PRIVATE_SLOG_COUNT_IMPL(_, A,B,C,D,E,F,G,H,I, Count, ...) Count
+
+#define UE_PRIVATE_SLOG_CALL(F, A) UE_PRIVATE_SLOG_EXPAND(F A)
+#define UE_PRIVATE_SLOG_EXPAND(X) X
+
+#endif // #if !NO_LOGGING
+
+#undef UE_API
diff --git a/Engine/Source/Runtime/Core/Public/Misc/OutputDevice.h b/Engine/Source/Runtime/Core/Public/Misc/OutputDevice.h
--- a/Engine/Source/Runtime/Core/Public/Misc/OutputDevice.h
+++ b/Engine/Source/Runtime/Core/Public/Misc/OutputDevice.h
@@ -14,6 +14,8 @@
 class FString;
 class FText;
 
+namespace UE { class FLogRecord; }
+
 #ifndef USE_DEBUG_LOGGING
 #define USE_DEBUG_LOGGING 1
 #endif
@@ -146,6 +148,8 @@ public:
 		Serialize( V, Verbosity, Category );
 	}
 
+	virtual void SerializeRecord(const UE::FLogRecord& Record);
+
 	virtual void Flush()
 	{
 	}
