I/O Chunk ID - using forward declared string builder

#rnx

diff --git a/Engine/Source/Runtime/Core/Public/IO/IoChunkId.h b/Engine/Source/Runtime/Core/Public/IO/IoChunkId.h
--- a/Engine/Source/Runtime/Core/Public/IO/IoChunkId.h
+++ b/Engine/Source/Runtime/Core/Public/IO/IoChunkId.h
@@ -3,9 +3,9 @@
 #pragma once
 
 #include "CoreTypes.h"
+#include "Containers/StringFwd.h"
 #include "Memory/MemoryView.h"
 #include "Misc/ByteSwap.h"
-#include "Misc/StringBuilder.h"
 #include "String/BytesToHex.h"
 
 class FArchive;
