Split TSet::Remove out into TSet::RemoveByIndex which doesn't deal with the FSetElementId.
Added some extra validity checks.

diff --git a/Engine/Source/Runtime/Core/Public/Containers/Set.h b/Engine/Source/Runtime/Core/Public/Containers/Set.h
--- a/Engine/Source/Runtime/Core/Public/Containers/Set.h
+++ b/Engine/Source/Runtime/Core/Public/Containers/Set.h
@@ -282,6 +282,10 @@ public:
 	typedef KeyFuncs    KeyFuncsType;
 	typedef Allocator   AllocatorType;
 
+	using SizeType = typename Allocator::SparseArrayAllocator::ElementAllocator::SizeType;
+
+	static_assert(std::is_same_v<SizeType, int32>, "TSet currently only supports 32-bit allocators");
+
 private:
 	template <typename, typename>
 	friend class TScriptSet;
@@ -815,31 +819,42 @@ public:
 		}
 	}
 
-	/**
-	 * Removes an element from the set.
-	 * @param Element - A pointer to the element in the set, as returned by Add or Find.
-	 */
-	void Remove(FSetElementId ElementId)
+private:
+	void RemoveByIndex(SizeType ElementIndex)
 	{
-		if (Elements.Num())
+		checkf(Elements.IsValidIndex(ElementIndex), TEXT("Invalid ElementIndex passed to TSet::RemoveByIndex"));
+
+		const SetElementType& ElementBeingRemoved = Elements[ElementIndex];
+
+		// Remove the element from the hash.
+		FSetElementId* HashPtr              = Hash.GetAllocation();
+		SizeType*      NextElementIndexIter = &HashPtr[ElementBeingRemoved.HashIndex].Index;
+		for (;;)
 		{
-			const auto& ElementBeingRemoved = Elements[ElementId];
+			SizeType NextElementIndex = *NextElementIndexIter;
+			checkf(NextElementIndex != INDEX_NONE, TEXT("Corrupt hash"));
 
-			// Remove the element from the hash.
-			for(FSetElementId* NextElementId = &GetTypedHash(ElementBeingRemoved.HashIndex);
-				NextElementId->IsValidId();
-				NextElementId = &Elements[*NextElementId].HashNextId)
+			if (NextElementIndex == ElementIndex)
 			{
-				if(*NextElementId == ElementId)
-				{
-					*NextElementId = ElementBeingRemoved.HashNextId;
-					break;
-				}
+				*NextElementIndexIter = ElementBeingRemoved.HashNextId.Index;
+				break;
 			}
+
+			NextElementIndexIter = &Elements[NextElementIndex].HashNextId.Index;
 		}
 
 		// Remove the element from the elements array.
-		Elements.RemoveAt(ElementId);
+		Elements.RemoveAt(ElementIndex);
+	}
+
+public:
+	/**
+	 * Removes an element from the set.
+	 * @param Element - A pointer to the element in the set, as returned by Add or Find.
+	 */
+	void Remove(FSetElementId ElementId)
+	{
+		RemoveByIndex(ElementId.AsInteger());
 	}
 
 	/**
@@ -958,7 +973,7 @@ private:
 			{
 				// This element matches the key, remove it from the set.  Note that Remove sets *NextElementId to point to the next
 				// element after the removed element in the hash bucket.
-				Remove(*NextElementId);
+				RemoveByIndex(NextElementId->AsInteger());
 				NumRemovedElements++;
 
 				if (!KeyFuncs::bAllowDuplicateKeys)
@@ -1674,7 +1689,7 @@ public:
 		/** Removes the current element from the set. */
 		FORCEINLINE void RemoveCurrent()
 		{
-			Set.Remove(TBaseIterator<false>::GetId());
+			Set.RemoveByIndex(TBaseIterator<false>::GetId().AsInteger());
 		}
 
 	private:
@@ -1716,7 +1731,7 @@ public:
 		/** Removes the current element from the set. */
 		FORCEINLINE void RemoveCurrent()
 		{
-			this->Set.Remove(TBaseKeyIterator<false>::Id);
+			this->Set.RemoveByIndex(TBaseKeyIterator<false>::Id.AsInteger());
 			TBaseKeyIterator<false>::Id = FSetElementId();
 		}
 	};
