Split TSet::FindIndexByHash out into TSet::FindIdByHash which doesn't deal with the FSetElementId.

diff --git a/Engine/Source/Runtime/Core/Public/Containers/Set.h b/Engine/Source/Runtime/Core/Public/Containers/Set.h
--- a/Engine/Source/Runtime/Core/Public/Containers/Set.h
+++ b/Engine/Source/Runtime/Core/Public/Containers/Set.h
@@ -636,15 +636,15 @@ public:
 	template <typename ElementReferenceType>
 	ElementType& FindOrAddByHash(uint32 KeyHash, ElementReferenceType&& InElement, bool* bIsAlreadyInSetPtr = nullptr)
 	{
-		FSetElementId ExistingId = FindIdByHash(KeyHash, KeyFuncs::GetSetKey(InElement));
-		bool bIsAlreadyInSet = ExistingId.IsValidId();
+		SizeType ExistingIndex = FindIndexByHash(KeyHash, KeyFuncs::GetSetKey(InElement));
+		bool bIsAlreadyInSet = ExistingIndex != INDEX_NONE;
 		if (bIsAlreadyInSetPtr)
 		{
 			*bIsAlreadyInSetPtr = bIsAlreadyInSet;
 		}
 		if (bIsAlreadyInSet)
 		{
-			return Elements[ExistingId].Value;
+			return Elements[ExistingIndex].Value;
 		}
 
 		// Create a new element.
@@ -665,18 +665,18 @@ private:
 			// Don't bother searching for a duplicate if this is the first element we're adding
 			if (Elements.Num() != 1)
 			{
-				FSetElementId ExistingId = FindIdByHash(KeyHash, KeyFuncs::GetSetKey(Element.Value));
-				bIsAlreadyInSet = ExistingId.IsValidId();
+				SizeType ExistingIndex = FindIndexByHash(KeyHash, KeyFuncs::GetSetKey(Element.Value));
+				bIsAlreadyInSet = ExistingIndex != INDEX_NONE;
 				if (bIsAlreadyInSet)
 				{
 					// If there's an existing element with the same key as the new element, replace the existing element with the new element.
-					MoveByRelocate(Elements[ExistingId].Value, Element.Value);
+					MoveByRelocate(Elements[ExistingIndex].Value, Element.Value);
 
 					// Then remove the new element.
-					Elements.RemoveAtUninitialized(InOutElementId);
+					Elements.RemoveAtUninitialized(InOutElementId.AsInteger());
 
 					// Then point the return value at the replaced element.
-					InOutElementId = ExistingId;
+					InOutElementId.Index = ExistingIndex;
 				}
 			}
 		}
@@ -857,27 +857,48 @@ public:
 		RemoveByIndex(ElementId.AsInteger());
 	}
 
+private:
 	/**
-	 * Finds an element with the given key in the set.
-	 * @param Key - The key to search for.
-	 * @return The id of the set element matching the given key, or the NULL id if none matches.
+	 * Finds an element with a pre-calculated hash and a key that can be compared to KeyType
+	 * @see	Class documentation section on ByHash() functions
+	 * @return The element id that matches the key and hash or an invalid element id
 	 */
-	FSetElementId FindId(KeyInitType Key) const
+	template <typename ComparableKey>
+	SizeType FindIndexByHash(uint32 KeyHash, const ComparableKey& Key) const
 	{
-		if (Elements.Num())
+		if (Elements.Num() == 0)
 		{
-			for(FSetElementId ElementId = GetTypedHash(KeyFuncs::GetKeyHash(Key));
-				ElementId.IsValidId();
-				ElementId = Elements[ElementId].HashNextId)
+			return INDEX_NONE;
+		}
+
+		FSetElementId* HashPtr      = Hash.GetAllocation();
+		SizeType       ElementIndex = HashPtr[KeyHash & (HashSize - 1)].AsInteger();
+		for (;;)
+		{
+			if (ElementIndex == INDEX_NONE)
 			{
-				if(KeyFuncs::Matches(KeyFuncs::GetSetKey(Elements[ElementId].Value),Key))
-				{
-					// Return the first match, regardless of whether the set has multiple matches for the key or not.
-					return ElementId;
-				}
+				return INDEX_NONE;
 			}
+
+			if (KeyFuncs::Matches(KeyFuncs::GetSetKey(Elements[ElementIndex].Value), Key))
+			{
+				// Return the first match, regardless of whether the set has multiple matches for the key or not.
+				return ElementIndex;
+			}
+
+			ElementIndex = Elements[ElementIndex].HashNextId.AsInteger();
 		}
-		return FSetElementId();
+	}
+
+public:
+	/**
+	 * Finds an element with the given key in the set.
+	 * @param Key - The key to search for.
+	 * @return The id of the set element matching the given key, or the NULL id if none matches.
+	 */
+	FSetElementId FindId(KeyInitType Key) const
+	{
+		return FSetElementId(FindIndexByHash(KeyFuncs::GetKeyHash(Key), Key));
 	}
 
 	/**
@@ -888,22 +909,9 @@ public:
 	template<typename ComparableKey>
 	FSetElementId FindIdByHash(uint32 KeyHash, const ComparableKey& Key) const
 	{
-		if (Elements.Num())
-		{
 			checkSlow(KeyHash == KeyFuncs::GetKeyHash(Key));
 
-			for (FSetElementId ElementId = GetTypedHash(KeyHash);
-				ElementId.IsValidId();
-				ElementId = Elements[ElementId].HashNextId)
-			{
-				if (KeyFuncs::Matches(KeyFuncs::GetSetKey(Elements[ElementId].Value), Key))
-				{
-					// Return the first match, regardless of whether the set has multiple matches for the key or not.
-					return ElementId;
-				}
-			}
-		}
-		return FSetElementId();
+		return FSetElementId(FindIndexByHash(KeyHash, Key));
 	}
 
 	/**
@@ -913,10 +921,10 @@ public:
 	 */
 	FORCEINLINE ElementType* Find(KeyInitType Key)
 	{
-		FSetElementId ElementId = FindId(Key);
-		if(ElementId.IsValidId())
+		SizeType ElementIndex = FindIndexByHash(KeyFuncs::GetKeyHash(Key), Key);
+		if (ElementIndex != INDEX_NONE)
 		{
-			return &Elements[ElementId].Value;
+			return &Elements[ElementIndex].Value;
 		}
 		else
 		{
@@ -942,10 +950,10 @@ public:
 	template<typename ComparableKey>
 	ElementType* FindByHash(uint32 KeyHash, const ComparableKey& Key)
 	{
-		FSetElementId ElementId = FindIdByHash(KeyHash, Key);
-		if (ElementId.IsValidId())
+		SizeType ElementIndex = FindIndexByHash(KeyHash, Key);
+		if (ElementIndex != INDEX_NONE)
 		{
-			return &Elements[ElementId].Value;
+			return &Elements[ElementIndex].Value;
 		}
 		else
 		{
@@ -1034,7 +1042,7 @@ public:
 	 */
 	FORCEINLINE bool Contains(KeyInitType Key) const
 	{
-		return FindId(Key).IsValidId();
+		return FindIndexByHash(KeyFuncs::GetKeyHash(Key), Key) != INDEX_NONE;
 	}
 
 	/**
@@ -1045,7 +1053,9 @@ public:
 	template<typename ComparableKey>
 	FORCEINLINE bool ContainsByHash(uint32 KeyHash, const ComparableKey& Key) const
 	{
-		return FindIdByHash(KeyHash, Key).IsValidId();
+		checkSlow(KeyHash == KeyFuncs::GetKeyHash(Key));
+
+		return FindIndexByHash(KeyHash, Key) != INDEX_NONE;
 	}
 
 	/**
