Moved hash resetting out of FSetElementId.
Fixed some style issues.

diff --git a/Engine/Source/Runtime/Core/Public/Containers/Set.h b/Engine/Source/Runtime/Core/Public/Containers/Set.h
--- a/Engine/Source/Runtime/Core/Public/Containers/Set.h
+++ b/Engine/Source/Runtime/Core/Public/Containers/Set.h
@@ -114,9 +114,7 @@ public:
 	friend class TScriptSet;
 
 	/** Default constructor. */
-	FORCEINLINE FSetElementId():
-		Index(INDEX_NONE)
-	{}
+	FORCEINLINE FSetElementId() = default;
 
 	/** @return a boolean value representing whether the id is NULL. */
 	FORCEINLINE bool IsValidId() const
@@ -145,23 +143,14 @@ public:
 	}
 
 private:
-	
-	/** Reset a range of FSetElementIds to invalid */
-	FORCEINLINE static void ResetRange(FSetElementId* Range, int32 Count)
-	{
-		for (int32 I = 0; I < Count; ++I)
-		{
-			Range[I] = FSetElementId();
-		}
-	}
-	
 	/** The index of the element in the set's element array. */
-	int32 Index;
+	int32 Index = INDEX_NONE;
 
 	/** Initialization constructor. */
-	FORCEINLINE FSetElementId(int32 InIndex):
-		Index(InIndex)
-	{}
+	FORCEINLINE FSetElementId(int32 InIndex)
+		: Index(InIndex)
+	{
+	}
 
 	/** Implicit conversion to the element index. */
 	FORCEINLINE operator int32() const
@@ -976,7 +965,8 @@ private:
 		FSetElementId* NextElementId = &GetTypedHash(KeyHash);
 		while (NextElementId->IsValidId())
 		{
-			auto& Element = Elements[*NextElementId];
+			SetElementType& Element = Elements[*NextElementId];
+
 			if (KeyFuncs::Matches(KeyFuncs::GetSetKey(Element.Value), Key))
 			{
 				// This element matches the key, remove it from the set.  Note that Remove sets *NextElementId to point to the next
@@ -1356,8 +1346,8 @@ private:
 		}
 	};
 
-	typedef TSparseArray<SetElementType,typename Allocator::SparseArrayAllocator>     ElementArrayType;
-	typedef typename Allocator::HashAllocator::template ForElementType<FSetElementId> HashType;
+	using ElementArrayType = TSparseArray<SetElementType, typename Allocator::SparseArrayAllocator>;
+	using HashType         = typename Allocator::HashAllocator::template ForElementType<FSetElementId>;
 
 	ElementArrayType Elements;
 
@@ -1449,26 +1439,27 @@ private:
 		LinkElement(ElementId, Element, KeyFuncs::GetKeyHash(KeyFuncs::GetSetKey(Element.Value)));
 	}
 
-	/** Returns if it should be faster to clear the hash by going through elements instead of reseting the whole bucket lists*/
-	FORCEINLINE bool ShouldClearByElements()
-	{
-		return Num() < (HashSize / 4);
-	}
-
-	/** Reset elements buckets of FSetElementIds to invalid */
+	/** Reset hash buckets to invalid */
 	void UnhashElements()
 	{
-		if (ShouldClearByElements())
+		FSetElementId* HashPtr = Hash.GetAllocation();
+
+		// Check if it should be faster to clear the hash by going through elements instead of resetting the whole hash
+		if (Num() < (HashSize / 4))
 		{
-			// Faster path: only reset hash buckets to FSetElementId for elements in the hash
+			// Faster path: only reset hash buckets to invalid for elements in the hash
 			for (const SetElementType& Element: Elements)
 			{
-				Hash.GetAllocation()[Element.HashIndex] = FSetElementId();
+				HashPtr[Element.HashIndex].Index = INDEX_NONE;
 			}
 		}
 		else
 		{
-			FSetElementId::ResetRange(Hash.GetAllocation(), HashSize);
+			/** Reset the whole hash to invalid */
+			for (int32 I = 0; I < HashSize; ++I)
+			{
+				HashPtr[I].Index = INDEX_NONE;
+			}
 		}
 	}
 
@@ -1892,7 +1883,11 @@ public:
 			Hash.ResizeAllocation(0, HashSize, sizeof(FSetElementId));
 		}
 
-		FSetElementId::ResetRange(Hash.GetAllocation(), HashSize);
+		FSetElementId* HashPtr = Hash.GetAllocation();
+		for (int32 I = 0; I < HashSize; ++I)
+		{
+			HashPtr[I].Index = INDEX_NONE;
+		}
 	}
 
 	void RemoveAt(int32 Index, const FScriptSetLayout& Layout)
