Removed implicit conversions from FSetElementId.
Unified on .Index instead of .AsInteger().

diff --git a/Engine/Source/Runtime/Core/Public/Containers/Set.h b/Engine/Source/Runtime/Core/Public/Containers/Set.h
--- a/Engine/Source/Runtime/Core/Public/Containers/Set.h
+++ b/Engine/Source/Runtime/Core/Public/Containers/Set.h
@@ -147,16 +147,10 @@ private:
 	int32 Index = INDEX_NONE;
 
 	/** Initialization constructor. */
-	FORCEINLINE FSetElementId(int32 InIndex)
+	FORCEINLINE explicit FSetElementId(int32 InIndex)
 		: Index(InIndex)
 	{
 	}
-
-	/** Implicit conversion to the element index. */
-	FORCEINLINE operator int32() const
-	{
-		return Index;
-	}
 };
 
 // This is just an int32
@@ -553,22 +547,23 @@ public:
 	 */
 	FORCEINLINE bool IsValidId(FSetElementId Id) const
 	{
-		return	Id.IsValidId() && 
-				Id >= 0 &&
-				Id < Elements.GetMaxIndex() &&
-				Elements.IsAllocated(Id);
+		SizeType Index = Id.Index;
+		return Index != INDEX_NONE &&
+			Index >= 0 &&
+			Index < Elements.GetMaxIndex() &&
+			Elements.IsAllocated(Index);
 	}
 
 	/** Accesses the identified element's value. */
 	FORCEINLINE ElementType& operator[](FSetElementId Id)
 	{
-		return Elements[Id].Value;
+		return Elements[Id.Index].Value;
 	}
 
 	/** Accesses the identified element's value. */
 	FORCEINLINE const ElementType& operator[](FSetElementId Id) const
 	{
-		return Elements[Id].Value;
+		return Elements[Id.Index].Value;
 	}
 
 	/**
@@ -639,7 +634,7 @@ public:
 		// Create a new element.
 		FSparseArrayAllocationInfo ElementAllocation = Elements.AddUninitialized();
 		SetElementType& Element = *new (ElementAllocation) SetElementType(Forward<ElementReferenceType>(InElement));
-		RehashOrLink(KeyHash, Element, ElementAllocation.Index);
+		RehashOrLink(KeyHash, Element, FSetElementId(ElementAllocation.Index));
 		return Element.Value;
 	}
 
@@ -662,7 +657,7 @@ private:
 					MoveByRelocate(Elements[ExistingIndex].Value, Element.Value);
 
 					// Then remove the new element.
-					Elements.RemoveAtUninitialized(InOutElementId.AsInteger());
+					Elements.RemoveAtUninitialized(InOutElementId.Index);
 
 					// Then point the return value at the replaced element.
 					InOutElementId.Index = ExistingIndex;
@@ -700,7 +695,7 @@ public:
 		// Create a new element.
 		FSparseArrayAllocationInfo ElementAllocation = Elements.AddUninitialized();
 		SetElementType& Element = *new (ElementAllocation) SetElementType(Forward<ArgsType>(Args));
-		FSetElementId ElementId = ElementAllocation.Index;
+		FSetElementId ElementId(ElementAllocation.Index);
 
 		uint32 KeyHash = KeyFuncs::GetKeyHash(KeyFuncs::GetSetKey(Element.Value));
 		if (!TryReplaceExisting(KeyHash, Element, ElementId, bIsAlreadyInSetPtr))
@@ -724,7 +719,7 @@ public:
 		// Create a new element.
 		FSparseArrayAllocationInfo ElementAllocation = Elements.AddUninitialized();
 		SetElementType& Element = *new (ElementAllocation) SetElementType(Forward<ArgsType>(Args));
-		FSetElementId ElementId = ElementAllocation.Index;
+		FSetElementId ElementId(ElementAllocation.Index);
 
 		if (!TryReplaceExisting(KeyHash, Element, ElementId, bIsAlreadyInSetPtr))
 		{
@@ -843,7 +838,7 @@ public:
 	 */
 	void Remove(FSetElementId ElementId)
 	{
-		RemoveByIndex(ElementId.AsInteger());
+		RemoveByIndex(ElementId.Index);
 	}
 
 private:
@@ -861,7 +856,7 @@ private:
 		}
 
 		FSetElementId* HashPtr      = Hash.GetAllocation();
-		SizeType       ElementIndex = HashPtr[KeyHash & (HashSize - 1)].AsInteger();
+		SizeType       ElementIndex = HashPtr[KeyHash & (HashSize - 1)].Index;
 		for (;;)
 		{
 			if (ElementIndex == INDEX_NONE)
@@ -875,7 +870,7 @@ private:
 				return ElementIndex;
 			}
 
-			ElementIndex = Elements[ElementIndex].HashNextId.AsInteger();
+			ElementIndex = Elements[ElementIndex].HashNextId.Index;
 		}
 	}
 
@@ -898,7 +893,7 @@ public:
 	template<typename ComparableKey>
 	FSetElementId FindIdByHash(uint32 KeyHash, const ComparableKey& Key) const
 	{
-			checkSlow(KeyHash == KeyFuncs::GetKeyHash(Key));
+		checkSlow(KeyHash == KeyFuncs::GetKeyHash(Key));
 
 		return FSetElementId(FindIndexByHash(KeyHash, Key));
 	}
@@ -965,13 +960,13 @@ private:
 		FSetElementId* NextElementId = &GetTypedHash(KeyHash);
 		while (NextElementId->IsValidId())
 		{
-			SetElementType& Element = Elements[*NextElementId];
+			SetElementType& Element = Elements[NextElementId->Index];
 
 			if (KeyFuncs::Matches(KeyFuncs::GetSetKey(Element.Value), Key))
 			{
 				// This element matches the key, remove it from the set.  Note that Remove sets *NextElementId to point to the next
 				// element after the removed element in the hash bucket.
-				RemoveByIndex(NextElementId->AsInteger());
+				RemoveByIndex(NextElementId->Index);
 				NumRemovedElements++;
 
 				if (!KeyFuncs::bAllowDuplicateKeys)
@@ -1087,7 +1082,7 @@ public:
 			int32 NumElementsInBucket = 0;
 			for(FSetElementId ElementId = GetTypedHash(HashIndex);
 				ElementId.IsValidId();
-				ElementId = Elements[ElementId].HashNextId)
+				ElementId = Elements[ElementId.Index].HashNextId)
 			{
 				NumElementsInBucket++;
 			}
@@ -1111,7 +1106,7 @@ public:
 					bResult=false;
 					break;
 				}
-				ElementId = Elements[ElementId].HashNextId;
+				ElementId = Elements[ElementId.Index].HashNextId;
 			}
 		}
 		return bResult;
@@ -1406,11 +1401,11 @@ private:
 	 */
 	FORCEINLINE const SetElementType& GetInternalElement(FSetElementId Id) const
 	{
-		return Elements[Id];
+		return Elements[Id.Index];
 	}
 	FORCEINLINE SetElementType& GetInternalElement(FSetElementId Id)
 	{
-		return Elements[Id];
+		return Elements[Id.Index];
 	}
 
 	/**
@@ -1690,7 +1685,7 @@ public:
 		/** Removes the current element from the set. */
 		FORCEINLINE void RemoveCurrent()
 		{
-			Set.RemoveByIndex(TBaseIterator<false>::GetId().AsInteger());
+			Set.RemoveByIndex(TBaseIterator<false>::GetId().Index);
 		}
 
 	private:
@@ -1732,7 +1727,7 @@ public:
 		/** Removes the current element from the set. */
 		FORCEINLINE void RemoveCurrent()
 		{
-			this->Set.RemoveByIndex(TBaseKeyIterator<false>::Id.AsInteger());
+			this->Set.RemoveByIndex(TBaseKeyIterator<false>::Id.Index);
 			TBaseKeyIterator<false>::Id = FSetElementId();
 		}
 	};
@@ -1897,9 +1892,9 @@ public:
 		void* ElementBeingRemoved = Elements.GetData(Index, Layout.SparseArrayLayout);
 
 		// Remove the element from the hash.
-		for (FSetElementId* NextElementId = &GetTypedHash(GetHashIndexRef(ElementBeingRemoved, Layout)); NextElementId->IsValidId(); NextElementId = &GetHashNextIdRef(Elements.GetData(NextElementId->AsInteger(), Layout.SparseArrayLayout), Layout))
+		for (FSetElementId* NextElementId = &GetTypedHash(GetHashIndexRef(ElementBeingRemoved, Layout)); NextElementId->IsValidId(); NextElementId = &GetHashNextIdRef(Elements.GetData(NextElementId->Index, Layout.SparseArrayLayout), Layout))
 		{
-			if (NextElementId->AsInteger() == Index)
+			if (NextElementId->Index == Index)
 			{
 				*NextElementId = GetHashNextIdRef(ElementBeingRemoved, Layout);
 				break;
@@ -1976,10 +1971,10 @@ private:
 			ElementId.IsValidId();
 			ElementId = GetHashNextIdRef(CurrentElement, Layout))
 		{
-			CurrentElement = (uint8*)Elements.GetData(ElementId, Layout.SparseArrayLayout);
+			CurrentElement = (uint8*)Elements.GetData(ElementId.Index, Layout.SparseArrayLayout);
 			if (EqualityFn(Element, CurrentElement))
 			{
-				return ElementId;
+				return ElementId.Index;
 			}
 		}
 	
