FSetElementId removed from internal TSet function signatures.

diff --git a/Engine/Source/Runtime/Core/Public/Containers/Set.h b/Engine/Source/Runtime/Core/Public/Containers/Set.h
--- a/Engine/Source/Runtime/Core/Public/Containers/Set.h
+++ b/Engine/Source/Runtime/Core/Public/Containers/Set.h
@@ -634,12 +634,12 @@ public:
 		// Create a new element.
 		FSparseArrayAllocationInfo ElementAllocation = Elements.AddUninitialized();
 		SetElementType& Element = *new (ElementAllocation) SetElementType(Forward<ElementReferenceType>(InElement));
-		RehashOrLink(KeyHash, Element, FSetElementId(ElementAllocation.Index));
+		RehashOrLink(KeyHash, Element, ElementAllocation.Index);
 		return Element.Value;
 	}
 
 private:
-	bool TryReplaceExisting(uint32 KeyHash, SetElementType& Element, FSetElementId& InOutElementId, bool* bIsAlreadyInSetPtr)
+	bool TryReplaceExisting(uint32 KeyHash, SetElementType& Element, SizeType& InOutElementIndex, bool* bIsAlreadyInSetPtr)
 	{
 		bool bIsAlreadyInSet = false;
 		if (!KeyFuncs::bAllowDuplicateKeys)
@@ -657,10 +657,10 @@ private:
 					MoveByRelocate(Elements[ExistingIndex].Value, Element.Value);
 
 					// Then remove the new element.
-					Elements.RemoveAtUninitialized(InOutElementId.Index);
+					Elements.RemoveAtUninitialized(InOutElementIndex);
 
 					// Then point the return value at the replaced element.
-					InOutElementId.Index = ExistingIndex;
+					InOutElementIndex = ExistingIndex;
 				}
 			}
 		}
@@ -671,13 +671,13 @@ private:
 		return bIsAlreadyInSet;
 	}
 
-	FORCEINLINE void RehashOrLink(uint32 KeyHash, SetElementType& Element, FSetElementId ElementId)
+	FORCEINLINE void RehashOrLink(uint32 KeyHash, SetElementType& Element, SizeType ElementIndex)
 	{
 		// Check if the hash needs to be resized.
 		if (!ConditionalRehash(Elements.Num()))
 		{
 			// If the rehash didn't add the new element to the hash, add it.
-			LinkElement(ElementId, Element, KeyHash);
+			LinkElement(ElementIndex, Element, KeyHash);
 		}
 	}
 
@@ -695,14 +695,15 @@ public:
 		// Create a new element.
 		FSparseArrayAllocationInfo ElementAllocation = Elements.AddUninitialized();
 		SetElementType& Element = *new (ElementAllocation) SetElementType(Forward<ArgsType>(Args));
-		FSetElementId ElementId(ElementAllocation.Index);
+
+		SizeType NewHashIndex = ElementAllocation.Index;
 
 		uint32 KeyHash = KeyFuncs::GetKeyHash(KeyFuncs::GetSetKey(Element.Value));
-		if (!TryReplaceExisting(KeyHash, Element, ElementId, bIsAlreadyInSetPtr))
+		if (!TryReplaceExisting(KeyHash, Element, NewHashIndex, bIsAlreadyInSetPtr))
 		{
-			RehashOrLink(KeyHash, Element, ElementId);
+			RehashOrLink(KeyHash, Element, NewHashIndex);
 		}
-		return ElementId;
+		return FSetElementId(NewHashIndex);
 	}
 	
 	/**
@@ -719,13 +720,14 @@ public:
 		// Create a new element.
 		FSparseArrayAllocationInfo ElementAllocation = Elements.AddUninitialized();
 		SetElementType& Element = *new (ElementAllocation) SetElementType(Forward<ArgsType>(Args));
-		FSetElementId ElementId(ElementAllocation.Index);
 
-		if (!TryReplaceExisting(KeyHash, Element, ElementId, bIsAlreadyInSetPtr))
+		SizeType NewHashIndex = ElementAllocation.Index;
+
+		if (!TryReplaceExisting(KeyHash, Element, ElementAllocation.Index, bIsAlreadyInSetPtr))
 		{
-			RehashOrLink(KeyHash, Element, ElementId);
+			RehashOrLink(KeyHash, Element, NewHashIndex);
 		}
-		return ElementId;
+		return FSetElementId(NewHashIndex);
 	}
 
 	template<typename ViewSizeType>
@@ -1418,20 +1420,20 @@ private:
 	}
 
 	/** Links an added element to the hash chain. */
-	FORCEINLINE void LinkElement(FSetElementId ElementId, const SetElementType& Element, uint32 KeyHash) const
+	FORCEINLINE void LinkElement(SizeType ElementIndex, const SetElementType& Element, uint32 KeyHash) const
 	{
 		// Compute the hash bucket the element goes in.
 		Element.HashIndex = KeyHash & (HashSize - 1);
 
 		// Link the element into the hash bucket.
-		Element.HashNextId = GetTypedHash(Element.HashIndex);
-		GetTypedHash(Element.HashIndex) = ElementId;
+		Element.HashNextId.Index = GetTypedHash(Element.HashIndex).Index;
+		GetTypedHash(Element.HashIndex).Index = ElementIndex;
 	}
 
 	/** Hashes and links an added element to the hash chain. */
-	FORCEINLINE void HashElement(FSetElementId ElementId, const SetElementType& Element) const
+	FORCEINLINE void HashElement(SizeType ElementIndex, const SetElementType& Element) const
 	{
-		LinkElement(ElementId, Element, KeyFuncs::GetKeyHash(KeyFuncs::GetSetKey(Element.Value)));
+		LinkElement(ElementIndex, Element, KeyFuncs::GetKeyHash(KeyFuncs::GetSetKey(Element.Value)));
 	}
 
 	/** Reset hash buckets to invalid */
@@ -1513,7 +1515,7 @@ private:
 			// Add the existing elements to the new hash.
 			for(typename ElementArrayType::TConstIterator ElementIt(Elements);ElementIt;++ElementIt)
 			{
-				HashElement(FSetElementId(ElementIt.GetIndex()),*ElementIt);
+				HashElement(ElementIt.GetIndex(), *ElementIt);
 			}
 		}
 	}
