Add new function FPlatformProcess::LaunchURLFiltered, and new class FURLRequestFilter, which the function uses to optionally filter URLs by scheme and domain.
UKismetSystemLibrary::LaunchURL now uses LaunchURLFiltered to allow filtering of URLs that may come from untrusted sources in Blueprints.
Added test cases for FURLRequestFilter to LowLevelTests.

diff --git a/Engine/Source/Runtime/Core/Private/GenericPlatform/GenericPlatformProcess.cpp b/Engine/Source/Runtime/Core/Private/GenericPlatform/GenericPlatformProcess.cpp
--- a/Engine/Source/Runtime/Core/Private/GenericPlatform/GenericPlatformProcess.cpp
+++ b/Engine/Source/Runtime/Core/Private/GenericPlatform/GenericPlatformProcess.cpp
@@ -18,6 +18,7 @@
 #include "Misc/EngineVersion.h"
 #include "Misc/LazySingleton.h"
 #include "Misc/Fork.h"
+#include "Misc/URLRequestFilter.h"
 #include "ProfilingDebugging/CsvProfiler.h"
 #include "Async/TaskGraphInterfaces.h"
 
@@ -227,6 +228,22 @@ void FGenericPlatformProcess::LaunchURL( const TCHAR* URL, const TCHAR* Parms, F
 	UE_LOG(LogHAL, Fatal, TEXT("FGenericPlatformProcess::LaunchURL not implemented on this platform"));
 }
 
+bool FGenericPlatformProcess::LaunchURLFiltered(const TCHAR* URL, const TCHAR* Parms, FString* Error, const UE::Core::FURLRequestFilter& Filter)
+{
+	const bool bAllowedByFilter = Filter.IsRequestAllowed(URL);
+
+	if (bAllowedByFilter)
+	{
+		FPlatformProcess::LaunchURL(URL, Parms, Error);
+	}
+	else if (Error)
+	{
+		*Error = TEXT("URL rejected by filter");
+	}
+
+	return bAllowedByFilter;
+}
+
 bool FGenericPlatformProcess::CanLaunchURL(const TCHAR* URL)
 {
 	UE_LOG(LogHAL, Warning, TEXT("FGenericPlatformProcess::CanLaunchURL not implemented on this platform"));
diff --git a/Engine/Source/Runtime/Core/Private/Misc/URLRequestFilter.cpp b/Engine/Source/Runtime/Core/Private/Misc/URLRequestFilter.cpp
new file mode 100644
--- /dev/null
+++ b/Engine/Source/Runtime/Core/Private/Misc/URLRequestFilter.cpp
@@ -0,0 +1,156 @@
+// Copyright Epic Games, Inc. All Rights Reserved.
+
+#include "Misc/URLRequestFilter.h"
+#include "Algo/AnyOf.h"
+#include "Algo/Find.h"
+#include "Algo/ForEach.h"
+#include "Containers/Array.h"
+#include "Containers/StringView.h"
+#include "Containers/UnrealString.h"
+#include "CoreGlobals.h"
+#include "Misc/CommandLine.h"
+#include "Misc/ConfigCacheIni.h"
+#include "Misc/Parse.h"
+#include "String/Find.h"
+
+namespace UE::Core
+{
+
+FURLRequestFilter::FURLRequestFilter(const TCHAR* ConfigSectionRootName, const FString& ConfigFileName)
+{
+	FConfigFile* File = GConfig->FindConfigFile(ConfigFileName);
+	if (File != nullptr)
+	{
+		Algo::ForEach(*File, [this, ConfigSectionRootName](const FConfigFile::ElementType& Entry)
+			{
+				const FString& SectionName = Entry.Key;
+
+				// Find config sections with the syntax [ConfigSectionRootName scheme]
+				const int32 DelimiterIndex = UE::String::FindLastChar(SectionName, ' ');
+				if (DelimiterIndex != INDEX_NONE)
+				{
+					FStringView RootName = SectionName;
+					RootName.LeftInline(DelimiterIndex);
+					RootName.TrimStartAndEndInline();
+
+					if (!RootName.Equals(ConfigSectionRootName, ESearchCase::CaseSensitive))
+					{
+						return;
+					}
+
+					FStringView Scheme = SectionName;
+					Scheme.RightChopInline(DelimiterIndex);
+					Scheme.TrimStartAndEndInline();
+
+					TArray<FString> Domains;
+					Entry.Value.MultiFind(TEXT("AllowedDomains"), Domains, true);
+
+#if !(UE_BUILD_SHIPPING || UE_BUILD_TEST)
+					TArray<FString> DevDomains;
+					Entry.Value.MultiFind(TEXT("DevAllowedDomains"), DevDomains, true);
+					Domains.Append(DevDomains);
+#endif
+
+					// Empty scheme or domain list is OK
+					AllowedRequests.Emplace(FString(Scheme), MoveTemp(Domains));
+				}
+				else if (SectionName.Equals(ConfigSectionRootName, ESearchCase::CaseSensitive))
+				{
+					AllowedRequests.Emplace(FString(), TArray<FString>());
+				}
+			});
+	}
+}
+
+FURLRequestFilter::FURLRequestFilter(const FRequestMap& InAllowedRequests)
+	: AllowedRequests(InAllowedRequests)
+{
+}
+
+FURLRequestFilter::FURLRequestFilter(FRequestMap&& InAllowedRequests)
+	: AllowedRequests(MoveTemp(InAllowedRequests))
+{
+}
+
+bool FURLRequestFilter::IsRequestAllowed(FStringView InURL) const
+{
+	if (AllowedRequests.IsEmpty())
+	{
+		// No filters configured
+		return true;
+	}
+
+	FStringView Scheme;
+	FStringView Domain;
+
+	const FStringView Delimiter = TEXT("://");
+	const int32 FoundIndex = UE::String::FindFirst(InURL, Delimiter);
+	if (FoundIndex != INDEX_NONE)
+	{
+		Scheme = InURL.Left(FoundIndex);
+		Domain = InURL.RightChop(FoundIndex + Delimiter.Len());
+	}
+	else
+	{
+		Domain = InURL;
+	}
+
+	// userinfo is unsupported (usernames in URLs, like "user" in http:[REDACTED_EMAIL]).
+
+	// Check for an IPv6 address delimited in '[' and ']'
+	if (Domain.StartsWith('['))
+	{
+		const int32 DomainEnd = UE::String::FindFirstChar(Domain, ']');
+		if (DomainEnd != INDEX_NONE)
+		{
+			Domain.LeftInline(DomainEnd + 1);
+		}
+	}
+	else
+	{
+		const int32 FoundDomainEndIndex = UE::String::FindFirstOfAnyChar(Domain, TEXTVIEW("/?:#"));
+		if (FoundDomainEndIndex != INDEX_NONE)
+		{
+			Domain.LeftInline(FoundDomainEndIndex);
+		}
+	}
+
+	// Look for the scheme, they are case-insensitive
+	const FRequestMap::ElementType* FoundScheme = Algo::FindByPredicate(AllowedRequests, [Scheme](const FRequestMap::ElementType& Item)
+	{
+		return !Item.Key.IsEmpty() && Scheme.Equals(Item.Key, ESearchCase::IgnoreCase);
+	});
+
+	if (FoundScheme == nullptr)
+	{
+		UE_LOG(LogCore, Log, TEXT("Scheme for URL %.*s not allowed."), InURL.Len(), InURL.GetData());
+		return false;
+	}
+
+	if (FoundScheme->Value.IsEmpty())
+	{
+		// No domain filters configured for this scheme
+		return true;
+	}
+
+	// Look for the domain in this scheme's filter, they are case-insensitive
+	const bool bFoundDomain = Algo::AnyOf(FoundScheme->Value, [Domain](const FString& Entry)
+	{
+		if (UE::String::FindFirstChar(Entry, '.') == 0)
+		{
+			// Entry starts with a period, check the end of the domain to allow for subdomains
+			return Domain.EndsWith(Entry, ESearchCase::IgnoreCase);
+		}
+		else
+		{
+			// Entry does not start with a period, match the whole domain.
+			return !Entry.IsEmpty() && Domain.Equals(Entry, ESearchCase::IgnoreCase);
+		}
+	});
+
+	UE_CLOG(!bFoundDomain, LogCore, Log, TEXT("Domain for URL %.*s not allowed."), InURL.Len(), InURL.GetData());
+
+	return bFoundDomain;
+}
+
+}
\ No newline at end of file
diff --git a/Engine/Source/Runtime/Core/Public/GenericPlatform/GenericPlatformProcess.h b/Engine/Source/Runtime/Core/Public/GenericPlatform/GenericPlatformProcess.h
--- a/Engine/Source/Runtime/Core/Public/GenericPlatform/GenericPlatformProcess.h
+++ b/Engine/Source/Runtime/Core/Public/GenericPlatform/GenericPlatformProcess.h
@@ -9,7 +9,6 @@
 #include "HAL/PlatformCrt.h"
 #include "HAL/PlatformMisc.h"
 #include "Misc/EnumClassFlags.h"
-#include "Misc/EnumClassFlags.h"
 #include "Templates/Function.h"
 
 class FEvent;
@@ -70,6 +69,10 @@ namespace ENamedThreads
 	enum Type : int32;
 }
 
+namespace UE::Core
+{
+	class FURLRequestFilter;
+}
 
 /** Generic implementation for the process handle. */
 template< typename T, T InvalidHandleValue >
@@ -388,6 +391,18 @@ struct CORE_API FGenericPlatformProcess
 	 */
 	static void LaunchURL( const TCHAR* URL, const TCHAR* Parms, FString* Error );
 
+	/**
+	 * Launch a uniform resource locator (i.e. http://www.epicgames.com/unreal).
+	 * This is expected to return immediately as the URL is launched by another
+	 * task. The URL param must already be a valid URL. The URL is passed through
+	 * the filter parameter for an added measure of security if the URL is from
+	 * and untrusted source. If you're looking for code to properly escape
+	 * a URL fragment, use FGenericPlatformHttp::UrlEncode.
+	 * 
+	 * @return true if URL passed the filter and was launched, false if it was rejected by the filter.
+	 */
+	static bool LaunchURLFiltered(const TCHAR* URL, const TCHAR* Parms, FString* Error, const UE::Core::FURLRequestFilter& Filter);
+
 	/**
 	 * Checks if the platform can launch a uniform resource locator (i.e. http://www.epicgames.com/unreal).
 	 **/
diff --git a/Engine/Source/Runtime/Core/Public/Misc/URLRequestFilter.h b/Engine/Source/Runtime/Core/Public/Misc/URLRequestFilter.h
new file mode 100644
--- /dev/null
+++ b/Engine/Source/Runtime/Core/Public/Misc/URLRequestFilter.h
@@ -0,0 +1,63 @@
+// Copyright Epic Games, Inc. All Rights Reserved.
+
+#pragma once
+
+#include "CoreTypes.h"
+#include "Containers/ContainersFwd.h"
+#include "Containers/StringFwd.h"
+#include "Containers/Map.h"
+
+namespace UE::Core
+{
+
+/**
+ * A filter that can used to restrict URL schemes and domains to known-safe values.
+ * Stores a map of schemes to domains: if the map is empty, no filtering is performed and any URL is allowed.
+ * Each key is a scheme (such as "https"). If a scheme is present, it will be allowed.
+ * Each value is an array of absolute domains or partial higher-level domains that are allowed for their scheme.
+ *
+ * Absolute domains do not start with a period and will be matched exactly with domains from the input URL,
+ * for example, "epicgames.com".
+ *
+ * Partial domains start with a period and will be matched with the ending of the input domain.
+ * Use these to allow all subdomains of the higher-level domain. For example, ".epicgames.com"
+ *
+ * To allow both an absolute domain and all of its subdomains, provide two entries in the array,
+ * one with the period prefix and one without. For example, "epicgames.com" and ".epicgames.com".
+ */
+class CORE_API FURLRequestFilter
+{
+public:
+	using FRequestMap = TMap<const FString, const TArray<FString>>;
+
+	FURLRequestFilter() = default;
+
+	/**
+	 * Initializes the allowlist map based on ini configs. Example syntax:
+	 * 
+	 * [ConfigSectionRootName https]
+	 * !AllowedDomains=ClearArray
+	 * +AllowedDomains=epicgames.com
+	 * 
+	 * [ConfigSectionRootName http]
+	 * !AllowedDomains=ClearArray
+	 * +AllowedDomains=epicgames.com
+	 */
+	FURLRequestFilter(const TCHAR* ConfigSectionRootName, const FString& ConfigFileName);
+
+	/**
+	 * Initializes the allowlist based on an explicit scheme/domain map
+	 */
+	FURLRequestFilter(const FRequestMap& InAllowedRequests);
+	FURLRequestFilter(FRequestMap&& InAllowedRequests);
+
+	/**
+	 * Matches InURL against the configued list of allowed URL schemes and domains. Returns true if the URL is allowed.
+	 */
+	bool IsRequestAllowed(FStringView InURL) const;
+
+private:
+	FRequestMap AllowedRequests;
+};
+
+}
\ No newline at end of file
diff --git a/Engine/Source/Runtime/Core/Tests/Misc/URLRequestFilterTest.cpp b/Engine/Source/Runtime/Core/Tests/Misc/URLRequestFilterTest.cpp
new file mode 100644
--- /dev/null
+++ b/Engine/Source/Runtime/Core/Tests/Misc/URLRequestFilterTest.cpp
@@ -0,0 +1,528 @@
+// Copyright Epic Games, Inc. All Rights Reserved.
+
+#if WITH_LOW_LEVEL_TESTS
+
+#include "Misc/URLRequestFilter.h"
+#include "Misc/ConfigCacheIni.h"
+#include "Misc/ScopeExit.h"
+#include "Containers/StringView.h"
+#include "Containers/Array.h"
+#include "Containers/Map.h"
+#include "Containers/UnrealString.h"
+
+#include <catch2/catch_test_macros.hpp>
+
+namespace UE::Core::Private
+{
+
+// Adds a config file to GConfig while object is in scope, for ease of testing with ini settings.
+// Doesn't add the file if one already exists with the same name, to avoid overwriting it.
+class FScopedConfigFile
+{
+public:
+	FScopedConfigFile(const TCHAR* BaseName, const FString& Contents)
+	{
+		Filename = GConfig->GetConfigFilename(BaseName);
+		FConfigFile* Existingfile = GConfig->FindConfigFile(Filename);
+
+		// If the file already exists don't overwrite it
+		if (Existingfile != nullptr)
+		{
+			bWasAdded = false;
+		}
+		else
+		{
+			FConfigFile TestConfigFile;
+			TestConfigFile.CombineFromBuffer(Contents, Filename);
+			GConfig->Add(Filename, TestConfigFile);
+			bWasAdded = true;
+		}
+	}
+
+	~FScopedConfigFile()
+	{
+		if (bWasAdded)
+		{
+			GConfig->Remove(Filename);
+		}
+	}
+
+	// Could be moveable, but not needed for these tests
+	FScopedConfigFile(const FScopedConfigFile&) = delete;
+	FScopedConfigFile(FScopedConfigFile&&) = delete;
+
+	FScopedConfigFile& operator==(const FScopedConfigFile&) = delete;
+	FScopedConfigFile& operator==(FScopedConfigFile&&) = delete;
+
+	FString GetFilename() const { return Filename; }
+	bool WasAdded() const { return bWasAdded; }
+
+private:
+	FString Filename;
+	bool bWasAdded = false;
+};
+
+
+TEST_CASE("Core::Misc::URLRequestFilter", "[Core]")
+{
+	SECTION("Empty filter")
+	{
+		const FURLRequestFilter Filter{FURLRequestFilter::FRequestMap()};
+
+		CHECK(Filter.IsRequestAllowed(TEXTVIEW("https://epicgames.com")));
+		CHECK(Filter.IsRequestAllowed(TEXTVIEW("https://epicgames.com/path/to/somewhere")));
+		CHECK(Filter.IsRequestAllowed(TEXTVIEW("http://epicgames.com")));
+		CHECK(Filter.IsRequestAllowed(TEXTVIEW("http://epicgames.com/path/to/somewhere")));
+		CHECK(Filter.IsRequestAllowed(TEXTVIEW("epicgames.com")));
+		CHECK(Filter.IsRequestAllowed(TEXTVIEW("epicgames.com/path/to/somewhere")));
+		CHECK(Filter.IsRequestAllowed(TEXTVIEW("http://subdomain.domain.com")));
+		CHECK(Filter.IsRequestAllowed(TEXTVIEW("http://subdomain.domain.com/path/to/somewhere")));
+		CHECK(Filter.IsRequestAllowed(TEXTVIEW("file://C:/somefile.txt")));
+		CHECK(Filter.IsRequestAllowed(TEXTVIEW("unreal://127.0.0.1")));
+		CHECK(Filter.IsRequestAllowed(TEXTVIEW("unreal://127.0.0.1:7777")));
+		CHECK(Filter.IsRequestAllowed(TEXTVIEW("unreal://127.0.0.1/path/to/somewhere")));
+		CHECK(Filter.IsRequestAllowed(TEXTVIEW("unreal://127.0.0.1:7777/path/to/somewhere")));
+	}
+
+	SECTION("Empty scheme")
+	{
+		FURLRequestFilter::FRequestMap SchemeMap;
+		SchemeMap.Add(TEXT(""), TArray<FString>{});
+
+		FURLRequestFilter Filter{SchemeMap};
+
+		// Reject everything
+		CHECK(!Filter.IsRequestAllowed(TEXTVIEW("://")));
+		CHECK(!Filter.IsRequestAllowed(TEXTVIEW(":///")));
+		CHECK(!Filter.IsRequestAllowed(TEXTVIEW("https://")));
+		CHECK(!Filter.IsRequestAllowed(TEXTVIEW("https:///")));
+		CHECK(!Filter.IsRequestAllowed(TEXTVIEW("://epicgames.com")));
+		CHECK(!Filter.IsRequestAllowed(TEXTVIEW("https://epicgames.com")));
+		CHECK(!Filter.IsRequestAllowed(TEXTVIEW("https://subdomain.epicgames.com")));
+		CHECK(!Filter.IsRequestAllowed(TEXTVIEW("epicgames.com")));
+		CHECK(!Filter.IsRequestAllowed(TEXTVIEW("http://epicgames.com")));
+		CHECK(!Filter.IsRequestAllowed(TEXTVIEW("http://subdomain.epicgames.com")));
+		CHECK(!Filter.IsRequestAllowed(TEXTVIEW("file://C:/somefile.txt")));
+		CHECK(!Filter.IsRequestAllowed(TEXTVIEW("unreal://127.0.0.1")));
+		CHECK(!Filter.IsRequestAllowed(TEXTVIEW("unreal://127.0.0.1:7777")));
+	}
+
+	SECTION("Empty domain")
+	{
+		FURLRequestFilter::FRequestMap SchemeMap;
+		SchemeMap.Add(TEXT("https"), TArray<FString>{TEXT("")});
+
+		FURLRequestFilter Filter{SchemeMap};
+
+		// Reject everything
+		CHECK(!Filter.IsRequestAllowed(TEXTVIEW("://")));
+		CHECK(!Filter.IsRequestAllowed(TEXTVIEW(":///")));
+		CHECK(!Filter.IsRequestAllowed(TEXTVIEW("https://")));
+		CHECK(!Filter.IsRequestAllowed(TEXTVIEW("https:///")));
+		CHECK(!Filter.IsRequestAllowed(TEXTVIEW("://epicgames.com")));
+		CHECK(!Filter.IsRequestAllowed(TEXTVIEW("https://epicgames.com")));
+		CHECK(!Filter.IsRequestAllowed(TEXTVIEW("https://epicgames.com/path/to/somewhere")));
+		CHECK(!Filter.IsRequestAllowed(TEXTVIEW("https://subdomain.epicgames.com")));
+		CHECK(!Filter.IsRequestAllowed(TEXTVIEW("https://subdomain.epicgames.com/path/to/somewhere")));
+		CHECK(!Filter.IsRequestAllowed(TEXTVIEW("epicgames.com")));
+		CHECK(!Filter.IsRequestAllowed(TEXTVIEW("httpsepicgames.com")));
+		CHECK(!Filter.IsRequestAllowed(TEXTVIEW("epicgameshttps")));
+		CHECK(!Filter.IsRequestAllowed(TEXTVIEW("epicgames.com/path/to/somewhere")));
+		CHECK(!Filter.IsRequestAllowed(TEXTVIEW("http://epicgames.com")));
+		CHECK(!Filter.IsRequestAllowed(TEXTVIEW("http://epicgames.com/path/to/somewhere")));
+		CHECK(!Filter.IsRequestAllowed(TEXTVIEW("http://subdomain.epicgames.com")));
+		CHECK(!Filter.IsRequestAllowed(TEXTVIEW("http://subdomain.epicgames.com/path/to/somewhere")));
+		CHECK(!Filter.IsRequestAllowed(TEXTVIEW("file://C:/somefile.txt")));
+		CHECK(!Filter.IsRequestAllowed(TEXTVIEW("unreal://127.0.0.1")));
+		CHECK(!Filter.IsRequestAllowed(TEXTVIEW("unreal://127.0.0.1:7777")));
+	}
+
+	SECTION("Single scheme filter")
+	{
+		FURLRequestFilter::FRequestMap SchemeMap;
+		SchemeMap.Add(TEXT("https"), TArray<FString>{});
+
+		FURLRequestFilter Filter{SchemeMap};
+
+		CHECK(Filter.IsRequestAllowed(TEXTVIEW("https://epicgames.com")));
+		CHECK(Filter.IsRequestAllowed(TEXTVIEW("https://epicgames.com/path/to/somewhere")));
+		CHECK(Filter.IsRequestAllowed(TEXTVIEW("https://subdomain.epicgames.com")));
+		CHECK(Filter.IsRequestAllowed(TEXTVIEW("https://subdomain.epicgames.com/path/to/somewhere")));
+		CHECK(!Filter.IsRequestAllowed(TEXTVIEW("epicgames.com")));
+		CHECK(!Filter.IsRequestAllowed(TEXTVIEW("httpsepicgames.com")));
+		CHECK(!Filter.IsRequestAllowed(TEXTVIEW("epicgameshttps")));
+		CHECK(!Filter.IsRequestAllowed(TEXTVIEW("epicgames.com/path/to/somewhere")));
+		CHECK(!Filter.IsRequestAllowed(TEXTVIEW("http://epicgames.com")));
+		CHECK(!Filter.IsRequestAllowed(TEXTVIEW("http://epicgames.com/path/to/somewhere")));
+		CHECK(!Filter.IsRequestAllowed(TEXTVIEW("http://subdomain.epicgames.com")));
+		CHECK(!Filter.IsRequestAllowed(TEXTVIEW("http://subdomain.epicgames.com/path/to/somewhere")));
+		CHECK(!Filter.IsRequestAllowed(TEXTVIEW("file://C:/somefile.txt")));
+		CHECK(!Filter.IsRequestAllowed(TEXTVIEW("unreal://127.0.0.1")));
+		CHECK(!Filter.IsRequestAllowed(TEXTVIEW("unreal://127.0.0.1:7777")));
+
+		// Various Windows file paths
+
+		// Drive absolute
+		CHECK(!Filter.IsRequestAllowed(TEXTVIEW("C:/somefile.txt")));
+		CHECK(!Filter.IsRequestAllowed(TEXTVIEW("C:\\somefile.txt")));
+		CHECK(!Filter.IsRequestAllowed(TEXTVIEW("C:/https/somefile.txt")));
+		CHECK(!Filter.IsRequestAllowed(TEXTVIEW("C:\\https\\somefile.txt")));
+		CHECK(!Filter.IsRequestAllowed(TEXTVIEW("C:/https/../somefile.txt")));
+		CHECK(!Filter.IsRequestAllowed(TEXTVIEW("C:\\https\\..\\somefile.txt")));
+
+		// Drive relative
+		CHECK(!Filter.IsRequestAllowed(TEXTVIEW("C:somefile.txt")));
+		CHECK(!Filter.IsRequestAllowed(TEXTVIEW("C:https/somefile.txt")));
+		CHECK(!Filter.IsRequestAllowed(TEXTVIEW("C:https\\somefile.txt")));
+		CHECK(!Filter.IsRequestAllowed(TEXTVIEW("C:https/../somefile.txt")));
+		CHECK(!Filter.IsRequestAllowed(TEXTVIEW("C:https\\..\\somefile.txt")));
+
+		// Rooted
+		CHECK(!Filter.IsRequestAllowed(TEXTVIEW("/somefile.txt")));
+		CHECK(!Filter.IsRequestAllowed(TEXTVIEW("\\somefile.txt")));
+		CHECK(!Filter.IsRequestAllowed(TEXTVIEW("/https/somefile.txt")));
+		CHECK(!Filter.IsRequestAllowed(TEXTVIEW("\\https\\somefile.txt")));
+		CHECK(!Filter.IsRequestAllowed(TEXTVIEW("/https/../somefile.txt")));
+		CHECK(!Filter.IsRequestAllowed(TEXTVIEW("\\https\\..\\somefile.txt")));
+
+		// Relative
+		CHECK(!Filter.IsRequestAllowed(TEXTVIEW("somefile.txt")));
+		CHECK(!Filter.IsRequestAllowed(TEXTVIEW("https/somefile.txt")));
+		CHECK(!Filter.IsRequestAllowed(TEXTVIEW("https\\somefile.txt")));
+		CHECK(!Filter.IsRequestAllowed(TEXTVIEW("../https/somefile.txt")));
+		CHECK(!Filter.IsRequestAllowed(TEXTVIEW("..\\https\\somefile.txt")));
+		
+		// UNC aboslute
+		CHECK(!Filter.IsRequestAllowed(TEXTVIEW("//someserver/share/somefile.txt")));
+		CHECK(!Filter.IsRequestAllowed(TEXTVIEW("\\\\someserver\\share\\somefile.txt")));
+		CHECK(!Filter.IsRequestAllowed(TEXTVIEW("//https/share/somefile.txt")));
+		CHECK(!Filter.IsRequestAllowed(TEXTVIEW("\\\\https\\share\\somefile.txt")));
+		CHECK(!Filter.IsRequestAllowed(TEXTVIEW("//https/https/somefile.txt")));
+		CHECK(!Filter.IsRequestAllowed(TEXTVIEW("\\\\https\\https\\somefile.txt")));
+		CHECK(!Filter.IsRequestAllowed(TEXTVIEW("//someserver/https/somefile.txt")));
+		CHECK(!Filter.IsRequestAllowed(TEXTVIEW("\\\\someserver\\https\\somefile.txt")));
+		CHECK(!Filter.IsRequestAllowed(TEXTVIEW("//someserver/https/../somefile.txt")));
+		CHECK(!Filter.IsRequestAllowed(TEXTVIEW("\\\\someserver\\https\\..\\somefile.txt")));
+
+		// Local device
+		CHECK(!Filter.IsRequestAllowed(TEXTVIEW("//./share/somefile.txt")));
+		CHECK(!Filter.IsRequestAllowed(TEXTVIEW("\\\\.\\share\\somefile.txt")));
+		CHECK(!Filter.IsRequestAllowed(TEXTVIEW("//./C:/https/somefile.txt")));
+		CHECK(!Filter.IsRequestAllowed(TEXTVIEW("\\\\.\\C:\\https\\somefile.txt")));
+		CHECK(!Filter.IsRequestAllowed(TEXTVIEW("//./https/somefile.txt")));
+		CHECK(!Filter.IsRequestAllowed(TEXTVIEW("\\\\.\\https\\somefile.txt")));
+		CHECK(!Filter.IsRequestAllowed(TEXTVIEW("//./https/../somefile.txt")));
+		CHECK(!Filter.IsRequestAllowed(TEXTVIEW("\\\\.\\https\\..\\somefile.txt")));
+
+		// Root local device
+		CHECK(!Filter.IsRequestAllowed(TEXTVIEW("//?/share/somefile.txt")));
+		CHECK(!Filter.IsRequestAllowed(TEXTVIEW("\\\\?\\share\\somefile.txt")));
+		CHECK(!Filter.IsRequestAllowed(TEXTVIEW("//?/C:/https/somefile.txt")));
+		CHECK(!Filter.IsRequestAllowed(TEXTVIEW("\\\\?\\C:\\https\\somefile.txt")));
+		CHECK(!Filter.IsRequestAllowed(TEXTVIEW("//?/https/somefile.txt")));
+		CHECK(!Filter.IsRequestAllowed(TEXTVIEW("\\\\?\\https\\somefile.txt")));
+		CHECK(!Filter.IsRequestAllowed(TEXTVIEW("//?/https/../somefile.txt")));
+		CHECK(!Filter.IsRequestAllowed(TEXTVIEW("\\\\?\\https\\..\\somefile.txt")));
+	}
+
+	SECTION("Case-insensitive scheme only")
+	{
+		FURLRequestFilter::FRequestMap SchemeMap;
+		SchemeMap.Add(TEXT("https"), TArray<FString>{});
+		SchemeMap.Add(TEXT("HtTpS"), TArray<FString>{});
+		SchemeMap.Add(TEXT("HTTPS"), TArray<FString>{});
+
+		FURLRequestFilter Filter{SchemeMap};
+
+		CHECK(Filter.IsRequestAllowed(TEXTVIEW("https://epicgames.com")));
+		CHECK(Filter.IsRequestAllowed(TEXTVIEW("https://subdomain.epicgames.com")));
+		CHECK(Filter.IsRequestAllowed(TEXTVIEW("HTTPS://epicgames.com/path/to/somewhere")));
+		CHECK(Filter.IsRequestAllowed(TEXTVIEW("HTTPS://subdomain.epicgames.com/path/to/somewhere")));
+		CHECK(Filter.IsRequestAllowed(TEXTVIEW("hTtPs://epicgames.com")));
+		CHECK(Filter.IsRequestAllowed(TEXTVIEW("hTtPs://subdomain.epicgames.com/path/to/somewhere")));
+		CHECK(!Filter.IsRequestAllowed(TEXTVIEW("epicgames.com")));
+		CHECK(!Filter.IsRequestAllowed(TEXTVIEW("HTTPSepicgames.com")));
+		CHECK(!Filter.IsRequestAllowed(TEXTVIEW("hTtPsepicgames.com")));
+		CHECK(!Filter.IsRequestAllowed(TEXTVIEW("epicgamesHTTPS")));
+		CHECK(!Filter.IsRequestAllowed(TEXTVIEW("epicgameshTtPs")));
+		CHECK(!Filter.IsRequestAllowed(TEXTVIEW("http://epicgames.com/path/to/somewhere")));
+		CHECK(!Filter.IsRequestAllowed(TEXTVIEW("http://subdomain.epicgames.com")));
+		CHECK(!Filter.IsRequestAllowed(TEXTVIEW("file://C:/somefile.txt")));
+		CHECK(!Filter.IsRequestAllowed(TEXTVIEW("unreal://127.0.0.1")));
+	}
+
+	SECTION("Case-insensitive scheme with domains")
+	{
+		FURLRequestFilter::FRequestMap SchemeMap;
+		SchemeMap.Add(TEXT("https"), TArray<FString>{TEXT("test1.epicgames.com")});
+		// Case-insensitive FString key overwrites previous entry
+		SchemeMap.Add(TEXT("HtTpS"), TArray<FString>{TEXT("test2.epicgames.com")});
+
+		FURLRequestFilter Filter{SchemeMap};
+
+		CHECK(!Filter.IsRequestAllowed(TEXTVIEW("https://test1.epicgames.com")));
+		CHECK(Filter.IsRequestAllowed(TEXTVIEW("https://test2.epicgames.com")));
+
+		CHECK(!Filter.IsRequestAllowed(TEXTVIEW("hTtPs://test1.epicgames.com")));
+		CHECK(Filter.IsRequestAllowed(TEXTVIEW("hTtPs://test2.epicgames.com")));
+
+		CHECK(!Filter.IsRequestAllowed(TEXTVIEW("HTTPS://test1.epicgames.com")));
+		CHECK(Filter.IsRequestAllowed(TEXTVIEW("HTTPS://test2.epicgames.com")));
+		
+	}
+
+	SECTION("Multiple scheme filter")
+	{
+		FURLRequestFilter::FRequestMap SchemeMap;
+		SchemeMap.Add(TEXT("https"), TArray<FString>{});
+		SchemeMap.Add(TEXT("http"), TArray<FString>{});
+
+		FURLRequestFilter Filter{SchemeMap};
+
+		CHECK(Filter.IsRequestAllowed(TEXTVIEW("https://epicgames.com")));
+		CHECK(Filter.IsRequestAllowed(TEXTVIEW("https://epicgames.com/path/to/somewhere")));
+		CHECK(Filter.IsRequestAllowed(TEXTVIEW("https://subdomain.epicgames.com")));
+		CHECK(Filter.IsRequestAllowed(TEXTVIEW("https://subdomain.epicgames.com/path/to/somewhere")));
+		CHECK(!Filter.IsRequestAllowed(TEXTVIEW("epicgames.com")));
+		CHECK(!Filter.IsRequestAllowed(TEXTVIEW("httpsepicgames.com")));
+		CHECK(!Filter.IsRequestAllowed(TEXTVIEW("epicgameshttps")));
+		CHECK(!Filter.IsRequestAllowed(TEXTVIEW("httpepicgames.com")));
+		CHECK(!Filter.IsRequestAllowed(TEXTVIEW("epicgameshttp")));
+		CHECK(Filter.IsRequestAllowed(TEXTVIEW("http://epicgames.com")));
+		CHECK(Filter.IsRequestAllowed(TEXTVIEW("http://epicgames.com/path/to/somewhere")));
+		CHECK(Filter.IsRequestAllowed(TEXTVIEW("http://subdomain.epicgames.com")));
+		CHECK(Filter.IsRequestAllowed(TEXTVIEW("http://subdomain.epicgames.com/path/to/somewhere")));
+		CHECK(!Filter.IsRequestAllowed(TEXTVIEW("file://C:/somefile.txt")));
+		CHECK(!Filter.IsRequestAllowed(TEXTVIEW("unreal://127.0.0.1")));
+		CHECK(!Filter.IsRequestAllowed(TEXTVIEW("unreal://127.0.0.1/path/to/somewhere")));
+	}
+
+	SECTION("Absolute domains")
+	{
+		FURLRequestFilter::FRequestMap SchemeMap;
+		SchemeMap.Add(TEXT("https"), TArray<FString>{TEXT("epicgames.com"), TEXT("unrealengine.com")});
+		SchemeMap.Add(TEXT("http"), TArray<FString>{TEXT("test1.epicgames.com"), TEXT("test2.epicgames.com")});
+
+		FURLRequestFilter Filter{SchemeMap};
+
+		CHECK(Filter.IsRequestAllowed(TEXTVIEW("https://epicgames.com")));
+		CHECK(Filter.IsRequestAllowed(TEXTVIEW("https://unrealengine.com/path/to/somewhere")));
+		CHECK(!Filter.IsRequestAllowed(TEXTVIEW("https://test1.epicgames.com")));
+		CHECK(!Filter.IsRequestAllowed(TEXTVIEW("https://test3.epicgames.com/path/to/somewhere")));
+		CHECK(!Filter.IsRequestAllowed(TEXTVIEW("epicgames.com")));
+		CHECK(!Filter.IsRequestAllowed(TEXTVIEW("httpsepicgames.com")));
+		CHECK(!Filter.IsRequestAllowed(TEXTVIEW("epicgameshttps")));
+		CHECK(!Filter.IsRequestAllowed(TEXTVIEW("httptest1.epicgames.com")));
+		CHECK(!Filter.IsRequestAllowed(TEXTVIEW("test2.epicgames.com.https")));
+		CHECK(!Filter.IsRequestAllowed(TEXTVIEW("unrealengine.com/path/to/somewhere")));
+
+		CHECK(Filter.IsRequestAllowed(TEXTVIEW("http://test1.epicgames.com/path/to/somewhere")));
+		CHECK(Filter.IsRequestAllowed(TEXTVIEW("http://test2.epicgames.com")));
+		CHECK(!Filter.IsRequestAllowed(TEXTVIEW("http://epicgames.com")));
+		CHECK(!Filter.IsRequestAllowed(TEXTVIEW("http://unrealengine.com/path/to/somewhere")));
+		CHECK(!Filter.IsRequestAllowed(TEXTVIEW("test1.epicgames.com")));
+		CHECK(!Filter.IsRequestAllowed(TEXTVIEW("test2.epicgames.com/path/to/somewhere")));
+
+		CHECK(!Filter.IsRequestAllowed(TEXTVIEW("wss://epicgames.com")));
+		CHECK(!Filter.IsRequestAllowed(TEXTVIEW("wss://unrealengine.com")));
+		CHECK(!Filter.IsRequestAllowed(TEXTVIEW("wss://unrealengine.com/path/to/somewhere")));
+		CHECK(!Filter.IsRequestAllowed(TEXTVIEW("wss://test1.epicgames.com")));
+		CHECK(!Filter.IsRequestAllowed(TEXTVIEW("wss://test2.epicgames.com/path/to/somewhere")));
+	}
+
+	SECTION("Subdomains")
+	{
+		FURLRequestFilter::FRequestMap SchemeMap;
+		SchemeMap.Add(TEXT("https"), TArray<FString>{TEXT(".epicgames.com"), TEXT(".unrealengine.com")});
+		SchemeMap.Add(TEXT("http"), TArray<FString>{TEXT(".test1.epicgames.com"), TEXT(".test2.epicgames.com")});
+
+		FURLRequestFilter Filter{SchemeMap};
+
+		CHECK(!Filter.IsRequestAllowed(TEXTVIEW("https://epicgames.com")));
+		CHECK(!Filter.IsRequestAllowed(TEXTVIEW("https://unrealengine.com/path/to/somewhere")));
+
+		// Technically allowed by the URI grammar in RFC 3986
+		CHECK(Filter.IsRequestAllowed(TEXTVIEW("https://.epicgames.com")));
+		CHECK(Filter.IsRequestAllowed(TEXTVIEW("https://.sub.epicgames.com/path/to/somewhere")));
+		CHECK(Filter.IsRequestAllowed(TEXTVIEW("https://.unrealengine.com")));
+		CHECK(Filter.IsRequestAllowed(TEXTVIEW("https://.sub.unrealengine.com")));
+
+		CHECK(Filter.IsRequestAllowed(TEXTVIEW("https://test1.epicgames.com/path/to/somewhere")));
+		CHECK(Filter.IsRequestAllowed(TEXTVIEW("https://test2.epicgames.com")));
+		CHECK(Filter.IsRequestAllowed(TEXTVIEW("https://test1.unrealengine.com/path/to/somewhere")));
+		CHECK(Filter.IsRequestAllowed(TEXTVIEW("https://test2.unrealengine.com")));
+		
+		CHECK(Filter.IsRequestAllowed(TEXTVIEW("https://sub.test1.epicgames.com")));
+		CHECK(Filter.IsRequestAllowed(TEXTVIEW("https://sub.test2.epicgames.com/path/to/somewhere")));
+		CHECK(Filter.IsRequestAllowed(TEXTVIEW("https://sub.test1.unrealengine.com")));
+		CHECK(Filter.IsRequestAllowed(TEXTVIEW("https://sub.test2.unrealengine.com/path/to/somewhere")));
+
+		CHECK(!Filter.IsRequestAllowed(TEXTVIEW("epicgames.com")));
+		CHECK(!Filter.IsRequestAllowed(TEXTVIEW("https.epicgames.com")));
+		CHECK(!Filter.IsRequestAllowed(TEXTVIEW("epicgameshttps")));
+		CHECK(!Filter.IsRequestAllowed(TEXTVIEW("http.test1.epicgames.com")));
+		CHECK(!Filter.IsRequestAllowed(TEXTVIEW(".test2.epicgames.com.http")));
+		CHECK(!Filter.IsRequestAllowed(TEXTVIEW("unrealengine.com/path/to/somewhere")));
+		CHECK(!Filter.IsRequestAllowed(TEXTVIEW(".epicgames.com/path/to/somewhere")));
+		CHECK(!Filter.IsRequestAllowed(TEXTVIEW(".unrealengine.com")));
+
+		CHECK(!Filter.IsRequestAllowed(TEXTVIEW("http://test1.epicgames.com")));
+		CHECK(!Filter.IsRequestAllowed(TEXTVIEW("http://test2.epicgames.com/path/to/somewhere")));
+
+		// Technically allowed by the URI grammar in RFC 3986
+		CHECK(Filter.IsRequestAllowed(TEXTVIEW("http://.test1.epicgames.com/path/to/somewhere")));
+		CHECK(Filter.IsRequestAllowed(TEXTVIEW("http://.sub.test1.epicgames.com")));
+		CHECK(Filter.IsRequestAllowed(TEXTVIEW("http://.test2.epicgames.com/path/to/somewhere")));
+		CHECK(Filter.IsRequestAllowed(TEXTVIEW("http://.sub.test2.epicgames.com")));
+
+		CHECK(Filter.IsRequestAllowed(TEXTVIEW("http://sub1.test1.epicgames.com")));
+		CHECK(Filter.IsRequestAllowed(TEXTVIEW("http://sub2.test1.epicgames.com/path/to/somewhere")));
+		CHECK(Filter.IsRequestAllowed(TEXTVIEW("http://sub1.test2.epicgames.com")));
+		CHECK(Filter.IsRequestAllowed(TEXTVIEW("http://sub2.test2.epicgames.com/path/to/somewhere")));
+		
+		CHECK(Filter.IsRequestAllowed(TEXTVIEW("http://sub1.sub2.test1.epicgames.com/path/to/somewhere")));
+		CHECK(Filter.IsRequestAllowed(TEXTVIEW("http://sub3.sub4.test1.epicgames.com")));
+		CHECK(Filter.IsRequestAllowed(TEXTVIEW("http://sub1.sub2.test2.epicgames.com/path/to/somewhere")));
+		CHECK(Filter.IsRequestAllowed(TEXTVIEW("http://sub3.sub4.test2.epicgames.com")));
+
+		CHECK(!Filter.IsRequestAllowed(TEXTVIEW("epicgames.com/path/to/somewhere")));
+		CHECK(!Filter.IsRequestAllowed(TEXTVIEW("unrealengine.com")));
+		CHECK(!Filter.IsRequestAllowed(TEXTVIEW(".epicgames.com")));
+		CHECK(!Filter.IsRequestAllowed(TEXTVIEW(".unrealengine.com/path/to/somewhere")));
+
+		CHECK(!Filter.IsRequestAllowed(TEXTVIEW("wss://.epicgames.com")));
+		CHECK(!Filter.IsRequestAllowed(TEXTVIEW("wss://.unrealengine.com/path/to/somewhere")));
+		CHECK(!Filter.IsRequestAllowed(TEXTVIEW("wss://.test1.epicgames.com/path/to/somewhere")));
+		CHECK(!Filter.IsRequestAllowed(TEXTVIEW("wss://.test2.epicgames.com")));
+		CHECK(!Filter.IsRequestAllowed(TEXTVIEW("wss://sub1.epicgames.com")));
+		CHECK(!Filter.IsRequestAllowed(TEXTVIEW("wss://sub1.unrealengine.com/path/to/somewhere")));
+		CHECK(!Filter.IsRequestAllowed(TEXTVIEW("wss://sub1.test1.epicgames.com/path/to/somewhere")));
+		CHECK(!Filter.IsRequestAllowed(TEXTVIEW("wss://sub1.test2.epicgames.com")));
+		CHECK(!Filter.IsRequestAllowed(TEXTVIEW("wss://sub1.sub2.epicgames.com/path/to/somewhere")));
+		CHECK(!Filter.IsRequestAllowed(TEXTVIEW("wss://sub1.sub2.unrealengine.com")));
+		CHECK(!Filter.IsRequestAllowed(TEXTVIEW("wss://sub1.sub2.test1.epicgames.com")));
+		CHECK(!Filter.IsRequestAllowed(TEXTVIEW("wss://sub1.sub2.test2.epicgames.com/path/to/somewhere")));
+	}
+
+	SECTION("IP addresses")
+	{
+		FURLRequestFilter::FRequestMap SchemeMap;
+		SchemeMap.Add(TEXT("https"), TArray<FString>{TEXT("127.0.0.1"), TEXT("[::1]")});
+
+		FURLRequestFilter Filter{SchemeMap};
+
+		CHECK(Filter.IsRequestAllowed(TEXTVIEW("https://127.0.0.1")));
+		CHECK(Filter.IsRequestAllowed(TEXTVIEW("https://127.0.0.1/path/to/somewhere")));
+		CHECK(Filter.IsRequestAllowed(TEXTVIEW("https://127.0.0.1:443")));
+		CHECK(Filter.IsRequestAllowed(TEXTVIEW("https://127.0.0.1:443/path/to/somewhere")));
+
+		CHECK(!Filter.IsRequestAllowed(TEXTVIEW("https://255.255.255.0")));
+		CHECK(!Filter.IsRequestAllowed(TEXTVIEW("https://255.255.255.0/path/to/somewhere")));
+		CHECK(!Filter.IsRequestAllowed(TEXTVIEW("https://255.255.255.0:443")));
+		CHECK(!Filter.IsRequestAllowed(TEXTVIEW("https://255.255.255.0:443/path/to/somewhere")));
+
+		CHECK(Filter.IsRequestAllowed(TEXTVIEW("https://[::1]")));
+		CHECK(Filter.IsRequestAllowed(TEXTVIEW("https://[::1]/path/to/somewhere")));
+		CHECK(Filter.IsRequestAllowed(TEXTVIEW("https://[::1]:443")));
+		CHECK(Filter.IsRequestAllowed(TEXTVIEW("https://[::1]:443/path/to/somewhere")));
+
+		CHECK(!Filter.IsRequestAllowed(TEXTVIEW("https://[2001:db8::]")));
+		CHECK(!Filter.IsRequestAllowed(TEXTVIEW("https://[2001:db8::]/path/to/somewhere")));
+		CHECK(!Filter.IsRequestAllowed(TEXTVIEW("https://[2001:db8::]:443")));
+		CHECK(!Filter.IsRequestAllowed(TEXTVIEW("https://[2001:db8::]:443/path/to/somewhere")));
+
+	}
+
+	SECTION("Config initialization")
+	{
+		const FString TestConfigFilename = TEXT("TestConfigFile");
+		const TCHAR* ConfigContents = TEXT(R"(
+[NoScheme]
+
+[NoSchemeSpace ]
+
+[NoSchemeSpaces   ]
+
+[SchemeFilter https]
+!AllowedDomains=ClearArray
++AllowedDomains=epicgames.com
++AllowedDomains=.epicgames.com
+!DevAllowedDomains=ClearArray
++DevAllowedDomains=docs.unrealengine.com
+
+[SchemeFilter http]
+!AllowedDomains=ClearArray
++AllowedDomains=unrealengine.com
+
+[SchemeNoDomain https]
+!AllowedDomains=ClearArray
+)");
+
+		Private::FScopedConfigFile ScopedFile(TEXT("TestConfigFile"), ConfigContents);
+		REQUIRE(ScopedFile.WasAdded());
+
+		// Reject everything
+		FURLRequestFilter NoScheme{TEXT("NoScheme"), ScopedFile.GetFilename()};
+		CHECK(!NoScheme.IsRequestAllowed(TEXTVIEW("https://epicgames.com")));
+		CHECK(!NoScheme.IsRequestAllowed(TEXTVIEW("https://unrealengine.com")));
+		CHECK(!NoScheme.IsRequestAllowed(TEXTVIEW("http://epicgames.com")));
+		CHECK(!NoScheme.IsRequestAllowed(TEXTVIEW("epicgames.com")));
+		CHECK(!NoScheme.IsRequestAllowed(TEXTVIEW("https://127.0.0.1")));
+		CHECK(!NoScheme.IsRequestAllowed(TEXTVIEW("127.0.0.1")));
+		CHECK(!NoScheme.IsRequestAllowed(TEXTVIEW("https://[::1]")));
+		CHECK(!NoScheme.IsRequestAllowed(TEXTVIEW("[::1]")));
+
+		FURLRequestFilter NoSchemeSpace{TEXT("NoSchemeSpace"), ScopedFile.GetFilename()};
+		CHECK(!NoSchemeSpace.IsRequestAllowed(TEXTVIEW("https://epicgames.com")));
+		CHECK(!NoSchemeSpace.IsRequestAllowed(TEXTVIEW("https://unrealengine.com")));
+		CHECK(!NoSchemeSpace.IsRequestAllowed(TEXTVIEW("http://epicgames.com")));
+		CHECK(!NoSchemeSpace.IsRequestAllowed(TEXTVIEW("epicgames.com")));
+		CHECK(!NoSchemeSpace.IsRequestAllowed(TEXTVIEW("https://127.0.0.1")));
+		CHECK(!NoSchemeSpace.IsRequestAllowed(TEXTVIEW("127.0.0.1")));
+		CHECK(!NoSchemeSpace.IsRequestAllowed(TEXTVIEW("https://[::1]")));
+		CHECK(!NoSchemeSpace.IsRequestAllowed(TEXTVIEW("[::1]")));
+
+		FURLRequestFilter NoSchemeSpaces{TEXT("NoSchemeSpaces"), ScopedFile.GetFilename()};
+		CHECK(!NoSchemeSpaces.IsRequestAllowed(TEXTVIEW("https://epicgames.com")));
+		CHECK(!NoSchemeSpaces.IsRequestAllowed(TEXTVIEW("https://unrealengine.com")));
+		CHECK(!NoSchemeSpaces.IsRequestAllowed(TEXTVIEW("http://epicgames.com")));
+		CHECK(!NoSchemeSpaces.IsRequestAllowed(TEXTVIEW("epicgames.com")));
+		CHECK(!NoSchemeSpaces.IsRequestAllowed(TEXTVIEW("https://127.0.0.1")));
+		CHECK(!NoSchemeSpaces.IsRequestAllowed(TEXTVIEW("127.0.0.1")));
+		CHECK(!NoSchemeSpaces.IsRequestAllowed(TEXTVIEW("https://[::1]")));
+		CHECK(!NoSchemeSpaces.IsRequestAllowed(TEXTVIEW("[::1]")));
+
+		// Normal filtering
+		FURLRequestFilter SchemeFilter{TEXT("SchemeFilter"), ScopedFile.GetFilename()};
+		CHECK(SchemeFilter.IsRequestAllowed(TEXTVIEW("https://epicgames.com")));
+		CHECK(!SchemeFilter.IsRequestAllowed(TEXTVIEW("https://unrealengine.com")));
+		CHECK(!SchemeFilter.IsRequestAllowed(TEXTVIEW("http://epicgames.com")));
+		CHECK(!SchemeFilter.IsRequestAllowed(TEXTVIEW("epicgames.com")));
+#if (UE_BUILD_SHIPPING || UE_BUILD_TEST)
+		CHECK(!SchemeFilter.IsRequestAllowed(TEXTVIEW("https://docs.unrealengine.com")));
+#else
+		CHECK(SchemeFilter.IsRequestAllowed(TEXTVIEW("https://docs.unrealengine.com")));
+#endif
+
+		FURLRequestFilter SchemeNoDomain{TEXT("SchemeNoDomain"), ScopedFile.GetFilename()};
+		CHECK(SchemeNoDomain.IsRequestAllowed(TEXTVIEW("https://epicgames.com")));
+		CHECK(SchemeNoDomain.IsRequestAllowed(TEXTVIEW("https://epicgames.com/path/to/somewhere")));
+		CHECK(SchemeNoDomain.IsRequestAllowed(TEXTVIEW("https://subdomain.epicgames.com")));
+		CHECK(SchemeNoDomain.IsRequestAllowed(TEXTVIEW("https://subdomain.epicgames.com/path/to/somewhere")));
+		CHECK(!SchemeNoDomain.IsRequestAllowed(TEXTVIEW("epicgames.com")));
+		CHECK(!SchemeNoDomain.IsRequestAllowed(TEXTVIEW("httpsepicgames.com")));
+		CHECK(!SchemeNoDomain.IsRequestAllowed(TEXTVIEW("epicgameshttps")));
+		CHECK(!SchemeNoDomain.IsRequestAllowed(TEXTVIEW("epicgames.com/path/to/somewhere")));
+		CHECK(!SchemeNoDomain.IsRequestAllowed(TEXTVIEW("http://epicgames.com")));
+		CHECK(!SchemeNoDomain.IsRequestAllowed(TEXTVIEW("http://epicgames.com/path/to/somewhere")));
+		CHECK(!SchemeNoDomain.IsRequestAllowed(TEXTVIEW("http://subdomain.epicgames.com")));
+		CHECK(!SchemeNoDomain.IsRequestAllowed(TEXTVIEW("http://subdomain.epicgames.com/path/to/somewhere")));
+		CHECK(!SchemeNoDomain.IsRequestAllowed(TEXTVIEW("file://C:/somefile.txt")));
+		CHECK(!SchemeNoDomain.IsRequestAllowed(TEXTVIEW("unreal://127.0.0.1")));
+		CHECK(!SchemeNoDomain.IsRequestAllowed(TEXTVIEW("unreal://127.0.0.1:7777")));
+	}
+}
+
+}
+
+#endif
diff --git a/Engine/Source/Runtime/Engine/Private/KismetSystemLibrary.cpp b/Engine/Source/Runtime/Engine/Private/KismetSystemLibrary.cpp
--- a/Engine/Source/Runtime/Engine/Private/KismetSystemLibrary.cpp
+++ b/Engine/Source/Runtime/Engine/Private/KismetSystemLibrary.cpp
@@ -11,6 +11,7 @@
 #include "Misc/PackageName.h"
 #include "Misc/Paths.h"
 #include "Misc/RuntimeErrors.h"
+#include "Misc/URLRequestFilter.h"
 #include "UObject/GCObject.h"
 #include "EngineGlobals.h"
 #include "Components/ActorComponent.h"
@@ -2653,7 +2654,9 @@ void UKismetSystemLibrary::LaunchURL(const FString& URL)
 {
 	if (!URL.IsEmpty())
 	{
-		FPlatformProcess::LaunchURL(*URL, nullptr, nullptr);
+		UE::Core::FURLRequestFilter Filter(TEXT("SystemLibrary.LaunchURLFilter"), GEngineIni);
+
+		FPlatformProcess::LaunchURLFiltered(*URL, nullptr, nullptr, Filter);
 	}
 }
 
