Adding a sharedpointer fwd file so that we don't always need to bring so much into a header.  This not only means we can remove the sharedpointer.h but the thing passed into its template from the header.

#test Tested as part of another upcoming change.

diff --git a/Engine/Source/Programs/IncludeTool/IncludeTool/Rules.cs b/Engine/Source/Programs/IncludeTool/IncludeTool/Rules.cs
--- a/Engine/Source/Programs/IncludeTool/IncludeTool/Rules.cs
+++ b/Engine/Source/Programs/IncludeTool/IncludeTool/Rules.cs
@@ -444,6 +444,8 @@ namespace IncludeTool
 			"/Engine/Source/Runtime/Online/HTTP/Public/HttpFwd.h", // error: invalid forward declaration - 'typedef TSharedPtr<IHttpRequest, ESPMode::ThreadSafe> FHttpRequestPtr;'
 			"/Engine/Restricted/NotForLicensees/Plugins/Online/OnlineSubsystemMcp/Source/Public/OnlineSubsystemMcpFwd.h", // expected only include directives and text in forward declaration header
 			"/Engine/Source/Runtime/RHI/Public/RHIFwd.h", // invalid forward declaration - 'namespace ERHIFeatureLevel { enum Type : int; }'
+			"/Engine/Source/Runtime/Core/Public/Templates/SharedPointerFwd.h", // Has an enum as well
+
 		};
 
 		/// <summary>
diff --git a/Engine/Source/Runtime/Core/Public/Templates/SharedPointerFwd.h b/Engine/Source/Runtime/Core/Public/Templates/SharedPointerFwd.h
new file mode 100644
--- /dev/null
+++ b/Engine/Source/Runtime/Core/Public/Templates/SharedPointerFwd.h
@@ -0,0 +1,24 @@
+// Copyright Epic Games, Inc. All Rights Reserved.
+
+#pragma once
+
+/**
+ * ESPMode is used select between either 'fast' or 'thread safe' shared pointer types.
+ * This is only used by templates at compile time to generate one code path or another.
+ */
+enum class ESPMode : uint8
+{
+	/** Forced to be not thread-safe. */
+	NotThreadSafe = 0,
+
+	/** Thread-safe, never spin locks, but slower */
+	ThreadSafe = 1
+};
+
+
+// Forward declarations.  By default, thread safety features are turned on. (Mode = ESPMode::ThreadSafe).
+// If you need more concerned with performance of ref-counting, you should use ESPMode::NotThreadSafe.
+template< class ObjectType, ESPMode Mode = ESPMode::ThreadSafe > class TSharedRef;
+template< class ObjectType, ESPMode Mode = ESPMode::ThreadSafe > class TSharedPtr;
+template< class ObjectType, ESPMode Mode = ESPMode::ThreadSafe > class TWeakPtr;
+template< class ObjectType, ESPMode Mode = ESPMode::ThreadSafe > class TSharedFromThis;
\ No newline at end of file
diff --git a/Engine/Source/Runtime/Core/Public/Templates/SharedPointerInternals.h b/Engine/Source/Runtime/Core/Public/Templates/SharedPointerInternals.h
--- a/Engine/Source/Runtime/Core/Public/Templates/SharedPointerInternals.h
+++ b/Engine/Source/Runtime/Core/Public/Templates/SharedPointerInternals.h
@@ -3,14 +3,14 @@
 #pragma once
 
 #include "CoreTypes.h"
-#include "Misc/AssertionMacros.h"
 #include "HAL/UnrealMemory.h"
+#include "Misc/AssertionMacros.h"
 #include "Templates/RemoveReference.h"
+#include "Templates/SharedPointerFwd.h"
 #include "Templates/TypeCompatibleBytes.h"
 #include <atomic>
 #include <type_traits>
 
-
 /** Default behavior. */
 #define THREAD_SANITISE_UNSAFEPTR 0
 
@@ -20,33 +20,6 @@
 	#define TSAN_SAFE_UNSAFEPTR TSAN_SAFE
 #endif
 
-/**
- * ESPMode is used select between either 'fast' or 'thread safe' shared pointer types.
- * This is only used by templates at compile time to generate one code path or another.
- */
-enum class ESPMode : uint8
-{
-	/** Forced to be not thread-safe. */
-	NotThreadSafe = 0,
-
-	/** Thread-safe, never spin locks, but slower */
-	ThreadSafe = 1,
-
-	/**
-	 * Fast was sometimes ThreadSafe and sometimes NotThreadSafe, but Fast is a misnomer.
-	 * Deprecated since ThreadSafe became the default.
-	 */
-	Fast UE_DEPRECATED(5.0, "ESPMode::Fast has been deprecated - please use ESPMode::ThreadSafe instead") = 1
-};
-
-
-// Forward declarations.  By default, thread safety features are turned on. (Mode = ESPMode::ThreadSafe).
-// If you need more concerned with performance of ref-counting, you should use ESPMode::NotThreadSafe.
-template< class ObjectType, ESPMode Mode = ESPMode::ThreadSafe > class TSharedRef;
-template< class ObjectType, ESPMode Mode = ESPMode::ThreadSafe > class TSharedPtr;
-template< class ObjectType, ESPMode Mode = ESPMode::ThreadSafe > class TWeakPtr;
-template< class ObjectType, ESPMode Mode = ESPMode::ThreadSafe > class TSharedFromThis;
-
 
 /**
  * SharedPointerInternals contains internal workings of shared and weak pointers.  You should
