Logging: Fixed a static analysis warning about an "odd semicolon"

#rnx

diff --git a/Engine/Source/Runtime/Core/Private/Logging/StructuredLog.cpp b/Engine/Source/Runtime/Core/Private/Logging/StructuredLog.cpp
--- a/Engine/Source/Runtime/Core/Private/Logging/StructuredLog.cpp
+++ b/Engine/Source/Runtime/Core/Private/Logging/StructuredLog.cpp
@@ -277,7 +277,9 @@ void FLogTemplate::FormatTo(FStringBuilderBase& Out, const TCHAR* Format, const
 	{
 		if (IndexHint >= 0)
 		{
-			for (; Index < IndexHint && It; ++Index, ++It);
+			for (; Index < IndexHint && It; ++Index, ++It)
+			{
+			}
 			if (IndexHint < Index)
 			{
 				It = Fields.CreateViewIterator();
