fixed MSVC static analysis warning

diff --git a/Engine/Source/Runtime/Core/Public/Windows/WindowsSemaphore.h b/Engine/Source/Runtime/Core/Public/Windows/WindowsSemaphore.h
--- a/Engine/Source/Runtime/Core/Public/Windows/WindowsSemaphore.h
+++ b/Engine/Source/Runtime/Core/Public/Windows/WindowsSemaphore.h
@@ -16,7 +16,7 @@ public:
 	FWindowsSemaphore(int32 InitialCount, int32 MaxCount)
 		: Semaphore(CreateSemaphore(nullptr, InitialCount, MaxCount, nullptr))
 	{
-		checkfSlow(Semaphore, TEXT("CreateSemaphore failed: %d"), GetLastError());
+		checkfSlow(Semaphore, TEXT("CreateSemaphore failed: %u"), GetLastError());
 	}
 
 	~FWindowsSemaphore()
@@ -26,22 +26,22 @@ public:
 
 	void Acquire()
 	{
-		HRESULT hRes = WaitForSingleObject(Semaphore, INFINITE);
-		checkfSlow(hRes == WAIT_OBJECT_0, TEXT("Acquiring semaphore failed: %d (%d)"), hRes, GetLastError());
+		DWORD Res = WaitForSingleObject(Semaphore, INFINITE);
+		checkfSlow(Res == WAIT_OBJECT_0, TEXT("Acquiring semaphore failed: %d (%u)"), Res, GetLastError());
 	}
 
 	bool TryAcquire(FTimespan Timeout = FTimespan::Zero())
 	{
-		HRESULT hRes = WaitForSingleObject(Semaphore, (DWORD)Timeout.GetTotalMilliseconds());
-		checkfSlow(hRes == WAIT_OBJECT_0 || hRes == WAIT_TIMEOUT, TEXT("Acquiring semaphore failed: %d (%d)"), (int)hRes, GetLastError());
-		return hRes == WAIT_OBJECT_0;
+		DWORD Res = WaitForSingleObject(Semaphore, (DWORD)Timeout.GetTotalMilliseconds());
+		checkfSlow(Res == WAIT_OBJECT_0 || Res == WAIT_TIMEOUT, TEXT("Acquiring semaphore failed: %d (%u)"), Res, GetLastError());
+		return Res == WAIT_OBJECT_0;
 	}
 
 	void Release(int32 Count = 1)
 	{
 		checkfSlow(Count > 0, TEXT("Releasing semaphore with Count = %d, that should be greater than 0"), Count);
 		bool bRes = ReleaseSemaphore(Semaphore, Count, nullptr);
-		checkfSlow(bRes, TEXT("Releasing semaphore for %d failed: %d"), Count, GetLastError());
+		checkfSlow(bRes, TEXT("Releasing semaphore for %d failed: %u"), Count, GetLastError());
 	}
 
 private:
