Disable log parsing while running assertion tests, to avoid Horde failing jobs.

diff --git a/Engine/Source/Runtime/Core/Tests/Misc/AssertionMacrosTest.cpp b/Engine/Source/Runtime/Core/Tests/Misc/AssertionMacrosTest.cpp
--- a/Engine/Source/Runtime/Core/Tests/Misc/AssertionMacrosTest.cpp
+++ b/Engine/Source/Runtime/Core/Tests/Misc/AssertionMacrosTest.cpp
@@ -114,12 +114,27 @@ TEST_CASE("Core::Misc::TestEnsure", "[Core][Misc][AssertionMacros][Ensure]")
 	}
 }
 
+// Disables log parsing by UAT for a scope, preventing assertion messages being tagged as errors.
+struct FSuspendLogParsing
+{
+	FSuspendLogParsing()
+	{
+		UE_LOG(LogCore, Display, TEXT("<-- Suspend Log Parsing -->"));
+	}
+
+	~FSuspendLogParsing()
+	{
+		UE_LOG(LogCore, Display, TEXT("<-- Resume Log Parsing -->"));
+	}
+};
+
 TEST_CASE("Core::Misc::TestCheck", "[Core][Misc][AssertionMacros][Check]")
 {
+	FSuspendLogParsing Scope;
+
 	REQUIRE_CHECK(check(1 == 2));
 	REQUIRE_CHECK_MSG("1 == 2", checkf(1 == 2, TEXT("Error Message")));
 	REQUIRE_CHECK_MSG("Error Message", checkf(1 == 2, TEXT("Error Message")));
 }
 
-
 #endif
