Fixed the Editor on Windows to prevent loading all debug symbols when resolving a first callstack before a load/unload module and a second one after the load/unload module.
  - This can happens with an early ensure or if some code prints out callstack in the initialization phass (like UClass::TryFindTypeSlow() does to warn user about deprecated short object names).
  - Refreshing the symbol through SymRefreshModuleList() (7s) takes much longer and use a lot more memory than just resetting DBGHelp (0.03s). The JIRA describe the tests performed and how the condition used to measure.

diff --git a/Engine/Source/Runtime/Core/Private/Windows/WindowsPlatformStackWalk.cpp b/Engine/Source/Runtime/Core/Private/Windows/WindowsPlatformStackWalk.cpp
--- a/Engine/Source/Runtime/Core/Private/Windows/WindowsPlatformStackWalk.cpp
+++ b/Engine/Source/Runtime/Core/Private/Windows/WindowsPlatformStackWalk.cpp
@@ -936,6 +936,18 @@ int32 FWindowsPlatformStackWalk::GetProcessModuleSignatures(FStackWalkModuleInfo
  */ 
 static void OnModulesChanged( FName ModuleThatChanged, EModuleChangeReason ReasonForChange )
 {
+	// PluginDirectoryChanged didn't change which modules are loaded or unloaded.
+	if (ReasonForChange == EModuleChangeReason::PluginDirectoryChanged)
+	{
+		return;
+	}
+#if ON_DEMAND_SYMBOL_LOADING
+	if (ReasonForChange == EModuleChangeReason::ModuleLoaded)
+	{
+		return; // If the module is needed to resolve a callstack, the module and its symbols will be loaded on demand.
+	}
+#endif
+
 	GNeedToRefreshSymbols = true;
 }
 
@@ -1095,7 +1107,7 @@ bool FWindowsPlatformStackWalk::InitStackWalkingInternal(void* Process, bool bFo
 		return false;
 	}
 
-	if (!GStackWalkingInitialized)
+	auto InitDbgHelp = []()
 	{
 		// Set up the symbol engine.
 		uint32 SymOpts = SymGetOptions();
@@ -1131,9 +1143,34 @@ bool FWindowsPlatformStackWalk::InitStackWalkingInternal(void* Process, bool bFo
 			LoadSymbolsForProcessModules(RemoteStorage);
 		}
 #endif
+	};
+
+	if (!GStackWalkingInitialized)
+	{
+		InitDbgHelp();
 	}
 	else if (GNeedToRefreshSymbols)
 	{
+#if ON_DEMAND_SYMBOL_LOADING
+		// Cleaning up and reinitializing is faster than calling SymRefreshModuleList() in a basic test conditions. The MS documentation claims that SymInitialize() with bInvadeProcess true (as done above)
+		// do the same thing as SymRefreshModuleList(), but that's likely not true. It was measured as following:
+		//
+		// FPlatformStackWalk::InitStackWalking(); // Must be the very first call (no prior ensure and anything that dumped the stack)
+		// FPlatformStackWalk::StackWalkAndDump(...);
+		// FModuleManager::Get().OnModulesChanged().Broadcast(NAME_None, EModuleChangeReason::ModuleUnloaded); // This set to GNeedToRefreshSymbols = true.
+		// FPlatformStackWalk::InitStackWalking(); // This will trigger the code that needs to be measured.
+		// FPlatformStackWalk::StackWalkAndDump(...);
+		//
+		// If we use: SymRefreshModuleList(), it takes 7 seconds and the memory usage of the Editor is doubled.
+		// If we use: SymCleanup() + SymInitializeW(..., bInvadeProcess=true), it takes 0.03s and doesn't use more memory to perform the dump.
+		//
+		// Why? Hard to tell. Maybe a bug. That's in the internal working of DBGHelp and this might change in the future. Cleaning up DBGHelp
+		// invalidates the cached states. This may affect negatively other use cases, but that's speculation. Also, module change events mostly
+		// happen at startup, usually before PlatformStackWalk::InitStackWalking() is called the first time unless an ensure() fires early or
+		// some code calls StackWalkAndDump() during the initialization phase, forcing initialization of DBGHelp.
+		SymCleanup(GProcessHandle);
+		InitDbgHelp();
+#else
 		// Refresh and reload symbols
 		SymRefreshModuleList(GProcessHandle);
 
@@ -1146,6 +1183,7 @@ bool FWindowsPlatformStackWalk::InitStackWalkingInternal(void* Process, bool bFo
 			// so load symbols for all modules the process has loaded.
 			LoadSymbolsForProcessModules( RemoteStorage );
 		}
+#endif
 	}
 
 	return GStackWalkingInitialized;
