Logging: Added structured log support to FOutputDeviceRedirector

diff --git a/Engine/Source/Runtime/Core/Private/Misc/OutputDeviceRedirector.cpp b/Engine/Source/Runtime/Core/Private/Misc/OutputDeviceRedirector.cpp
--- a/Engine/Source/Runtime/Core/Private/Misc/OutputDeviceRedirector.cpp
+++ b/Engine/Source/Runtime/Core/Private/Misc/OutputDeviceRedirector.cpp
@@ -11,8 +11,11 @@
 #include "HAL/PlatformTime.h"
 #include "HAL/PlatformTLS.h"
 #include "HAL/Thread.h"
+#include "Logging/StructuredLog.h"
 #include "Misc/App.h"
 #include "Misc/ScopeRWLock.h"
+#include "Misc/StringBuilder.h"
+#include "Misc/TVariant.h"
 #include "ProfilingDebugging/CpuProfilerTrace.h"
 #include <atomic>
 
@@ -81,7 +84,7 @@ struct FOutputDeviceLine
 	FOutputDeviceLine(const FBufferedLine&) = delete;
 	FOutputDeviceLine& operator=(const FBufferedLine&) = delete;
 
-	FORCEINLINE FOutputDeviceLine(const TCHAR* const InData, const FName InCategory, const ELogVerbosity::Type InVerbosity, const double InTime)
+	FORCEINLINE FOutputDeviceLine(const TCHAR* const InData, const FName& InCategory, const ELogVerbosity::Type InVerbosity, const double InTime)
 		: Time(InTime)
 		, Data(CopyData(InData))
 		, Category(InCategory)
@@ -103,6 +106,21 @@ private:
 	}
 };
 
+struct FOutputDeviceItem
+{
+	TVariant<FOutputDeviceLine, FLogRecord> Value;
+
+	FORCEINLINE FOutputDeviceItem(const TCHAR* const Data, const FName& Category, const ELogVerbosity::Type Verbosity, const double Time)
+		: Value(TInPlaceType<FOutputDeviceLine>(), Data, Category, Verbosity, Time)
+	{
+	}
+
+	FORCEINLINE explicit FOutputDeviceItem(const FLogRecord& Record)
+		: Value(TInPlaceType<FLogRecord>(), Record)
+	{
+	}
+};
+
 static constexpr uint64 CalculateRedirectorCacheLinePadding(const uint64 Size)
 {
 	return PLATFORM_CACHE_LINE_SIZE * FMath::DivideAndRoundUp<uint64>(Size, PLATFORM_CACHE_LINE_SIZE) - Size;
@@ -115,9 +133,9 @@ struct FOutputDeviceRedirectorState
 	std::atomic<uint32> OutputDevicesLockState = 0;
 	uint8 OutputDevicesLockPadding[CalculateRedirectorCacheLinePadding(sizeof(OutputDevicesLock) + sizeof(OutputDevicesLockState))]{};
 
-	/** A queue of lines logged by non-primary threads. */
-	TDepletableMpmcQueue<FOutputDeviceLine, FOutputDeviceLinearAllocator> BufferedLines;
-	uint8 BufferedLinesPadding[CalculateRedirectorCacheLinePadding(sizeof(BufferedLines))]{};
+	/** A queue of items logged by non-primary threads. */
+	TDepletableMpmcQueue<FOutputDeviceItem, FOutputDeviceLinearAllocator> BufferedItems;
+	uint8 BufferedItemsPadding[CalculateRedirectorCacheLinePadding(sizeof(BufferedItems))]{};
 
 	/** Array of output devices to redirect to from the primary thread. */
 	TArray<FOutputDevice*> BufferedOutputDevices;
@@ -138,7 +156,7 @@ struct FOutputDeviceRedirectorState
 	/** A queue of events to trigger when the dedicated primary thread is idle. */
 	TConsumeAllMpmcQueue<FEvent*, FOutputDeviceLinearAllocator> ThreadIdleEvents;
 
-	/** An event to wake the dedicated primary thread to process buffered lines. */
+	/** An event to wake the dedicated primary thread to process buffered items. */
 	std::atomic<FEvent*> ThreadWakeEvent = nullptr;
 
 	/** The ID of the thread holding the primary lock. */
@@ -157,6 +175,11 @@ struct FOutputDeviceRedirectorState
 	TBitArray<TInlineAllocator<1>> BufferedOutputDevicesCanBeUsedOnPanicThread;
 	TBitArray<TInlineAllocator<1>> UnbufferedOutputDevicesCanBeUsedOnPanicThread;
 
+	bool HasPanicThread() const
+	{
+		return PanicThreadId.load(std::memory_order_relaxed) != MAX_uint32;
+	}
+
 	bool IsPrimaryThread(const uint32 ThreadId) const
 	{
 		return ThreadId == PrimaryThreadId.load(std::memory_order_relaxed);
@@ -185,7 +208,7 @@ struct FOutputDeviceRedirectorState
 
 	void ThreadLoop();
 
-	void FlushBufferedLines();
+	void FlushBufferedItems();
 
 	template <typename OutputDevicesType, typename FunctionType, typename... ArgTypes>
 	FORCEINLINE void BroadcastTo(
@@ -205,6 +228,9 @@ struct FOutputDeviceRedirectorState
 			}
 		}
 	}
+
+	FORCENOINLINE void AddToBacklog(const FLogRecord& Record);
+	FORCENOINLINE void AddToBacklog(const TCHAR* Data, ELogVerbosity::Type Verbosity, const FName& Category, double Time);
 };
 
 /**
@@ -301,7 +327,7 @@ private:
  * - Must not be entered when the thread holds a write lock or primary lock.
  * - Must check IsLocked() before performing restricted operations.
  * - Must be locked to write to buffered output devices.
- * - Must be locked while calling FlushBufferedLines().
+ * - Must be locked while calling FlushBufferedItems().
  * - May be locked when the thread holds a read lock.
  * - When a panic thread is active, locking will only succeed from the panic thread.
  */
@@ -430,35 +456,69 @@ void FOutputDeviceRedirectorState::ThreadLoop()
 	while (FEvent* WakeEvent = ThreadWakeEvent.load(std::memory_order_acquire))
 	{
 		WakeEvent->Wait();
-		while (!BufferedLines.IsEmpty() && IsPrimaryThread(ThreadId))
+		while (!BufferedItems.IsEmpty() && IsPrimaryThread(ThreadId))
 		{
 			if (FOutputDevicesPrimaryScopeLock Lock(*this); Lock.IsLocked())
 			{
-				FlushBufferedLines();
+				FlushBufferedItems();
 			}
 		}
 		ThreadIdleEvents.ConsumeAllLifo([](FEvent* Event) { Event->Trigger(); });
 	}
 }
 
-void FOutputDeviceRedirectorState::FlushBufferedLines()
+void FOutputDeviceRedirectorState::FlushBufferedItems()
 {
-	if (BufferedLines.IsEmpty())
+	using namespace UE;
+	using namespace UE::Private;
+
+	if (BufferedItems.IsEmpty())
 	{
 		return;
 	}
 
-	TRACE_CPUPROFILER_EVENT_SCOPE(FOutputDeviceRedirector::FlushBufferedLines);
+	TRACE_CPUPROFILER_EVENT_SCOPE(FOutputDeviceRedirector::FlushBufferedItems);
 
 	const uint32 ThreadId = FPlatformTLS::GetCurrentThreadId();
-	BufferedLines.Deplete([this, ThreadId](UE::Private::FOutputDeviceLine&& Line)
+	BufferedItems.Deplete([this, ThreadId](FOutputDeviceItem&& Item)
 	{
-		BroadcastTo(ThreadId, BufferedOutputDevices, BufferedOutputDevicesCanBeUsedOnPanicThread,
-			UE_PROJECTION_MEMBER(FOutputDevice, Serialize),
-			Line.Data, Line.Verbosity, Line.Category, Line.Time);
+		Visit([this, ThreadId](auto&& Value)
+		{
+			using ValueType = std::decay_t<decltype(Value)>;
+			if constexpr (std::is_same_v<ValueType, FOutputDeviceLine>)
+			{
+				const FOutputDeviceLine& Line = Value;
+				BroadcastTo(ThreadId, BufferedOutputDevices, BufferedOutputDevicesCanBeUsedOnPanicThread,
+					UE_PROJECTION_MEMBER(FOutputDevice, Serialize),
+					Line.Data, Line.Verbosity, Line.Category, Line.Time);
+			}
+			else if constexpr (std::is_same_v<ValueType, FLogRecord>)
+			{
+				const FLogRecord& Record = Value;
+				BroadcastTo(ThreadId, BufferedOutputDevices, BufferedOutputDevicesCanBeUsedOnPanicThread,
+					UE_PROJECTION_MEMBER(FOutputDevice, SerializeRecord), Record);
+			}
+		}, Item.Value);
 	});
 }
 
+void FOutputDeviceRedirectorState::AddToBacklog(const FLogRecord& Record)
+{
+	TStringBuilder<512> Text;
+	Record.FormatMessageTo(Text);
+	AddToBacklog(*Text, Record.GetVerbosity(), Record.GetCategory(), FPlatformTime::ToSeconds64(Record.GetTime().GetCycles()));
+}
+
+void FOutputDeviceRedirectorState::AddToBacklog(
+	const TCHAR* const Data,
+	const ELogVerbosity::Type Verbosity,
+	const FName& Category,
+	const double Time)
+{
+	FWriteScopeLock ScopeLock(BacklogLock);
+	BacklogLines.Emplace(Data, Category, Verbosity, Time);
+}
+
 } // UE::Private
 
 FOutputDeviceRedirector::FOutputDeviceRedirector()
@@ -513,7 +573,7 @@ void FOutputDeviceRedirector::FlushThreadedLogs(EOutputDeviceRedirectorFlushOpti
 
 	if (UE::Private::FOutputDevicesPrimaryScopeLock Lock(*State); Lock.IsLocked())
 	{
-		State->FlushBufferedLines();
+		State->FlushBufferedItems();
 	}
 }
 
@@ -547,7 +607,7 @@ void FOutputDeviceRedirector::SetCurrentThreadAsPrimaryThread()
 	else
 	{
 		State->PrimaryThreadId.store(ThreadId, std::memory_order_relaxed);
-		State->FlushBufferedLines();
+		State->FlushBufferedItems();
 	}
 
 	State->TryStopThread();
@@ -558,11 +618,61 @@ bool FOutputDeviceRedirector::TryStartDedicatedPrimaryThread()
 	return FApp::ShouldUseThreadingForPerformance() && State->TryStartThread();
 }
 
+void FOutputDeviceRedirector::SerializeRecord(const UE::FLogRecord& Record)
+{
+	using namespace UE::Private;
+
+	FOutputDevicesReadScopeLock Lock(*State);
+
+	const uint32 ThreadId = FPlatformTLS::GetCurrentThreadId();
+
+	// Serialize directly to any output devices which don't require buffering
+	State->BroadcastTo(ThreadId, State->UnbufferedOutputDevices, State->UnbufferedOutputDevicesCanBeUsedOnPanicThread,
+		UE_PROJECTION_MEMBER(FOutputDevice, SerializeRecord), Record);
+
+	// Serialize to the backlog when not in panic mode. This will deadlock in panic mode when the
+	// FPlatformMallocCrash allocator has been enabled and logging occurs on a non-panic thread.
+	if (UNLIKELY(State->bEnableBacklog && !State->HasPanicThread()))
+	{
+		State->AddToBacklog(Record);
+	}
+
+	// Serialize to buffered output devices from the primary logging thread.
+	// Records are queued until buffered output devices are added to avoid missing early log records.
+	if (State->IsPrimaryThread(ThreadId) && !State->BufferedOutputDevices.IsEmpty())
+	{
+		// Verify that this is the primary thread again because another thread may have become
+		// the primary thread between the previous check and the lock.
+		if (FOutputDevicesPrimaryScopeLock PrimaryLock(*State); PrimaryLock.IsLocked() && State->IsPrimaryThread(ThreadId))
+		{
+			State->FlushBufferedItems();
+			State->BroadcastTo(ThreadId, State->BufferedOutputDevices, State->BufferedOutputDevicesCanBeUsedOnPanicThread,
+				UE_PROJECTION_MEMBER(FOutputDevice, SerializeRecord), Record);
+			if (UNLIKELY(State->IsPanicThread(ThreadId)))
+			{
+				Flush();
+			}
+			return;
+		}
+	}
+
+	// Queue the record to serialize to buffered output devices from the primary thread.
+	if (State->BufferedItems.EnqueueAndReturnWasEmpty(Record))
+	{
+		if (FEvent* WakeEvent = State->ThreadWakeEvent.load(std::memory_order_acquire))
+		{
+			WakeEvent->Trigger();
+		}
+	}
+}
+
 void FOutputDeviceRedirector::Serialize(const TCHAR* const Data, const ELogVerbosity::Type Verbosity, const FName& Category, const double Time)
 {
-	const double RealTime = Time == -1.0f ? FPlatformTime::Seconds() - GStartTime : Time;
+	using namespace UE::Private;
 
-	UE::Private::FOutputDevicesReadScopeLock Lock(*State);
+	const double RealTime = Time == -1.0 ? FPlatformTime::Seconds() - GStartTime : Time;
+
+	FOutputDevicesReadScopeLock Lock(*State);
 
 #if PLATFORM_DESKTOP
 	// Print anything that arrives after logging has shut down to at least have it in stdout.
@@ -585,10 +695,9 @@ void FOutputDeviceRedirector::Serialize(const TCHAR* const Data, const ELogVerbo
 
 	// Serialize to the backlog when not in panic mode. This will deadlock in panic mode when the
 	// FPlatformMallocCrash allocator has been enabled and logging occurs on a non-panic thread.
-	if (UNLIKELY(State->bEnableBacklog && !State->IsPanicThread(ThreadId)))
+	if (UNLIKELY(State->bEnableBacklog && !State->HasPanicThread()))
 	{
-		FWriteScopeLock ScopeLock(State->BacklogLock);
-		State->BacklogLines.Emplace(Data, Category, Verbosity, RealTime);
+		State->AddToBacklog(Data, Verbosity, Category, RealTime);
 	}
 
 	// Serialize to buffered output devices from the primary logging thread.
@@ -597,9 +706,9 @@ void FOutputDeviceRedirector::Serialize(const TCHAR* const Data, const ELogVerbo
 	{
 		// Verify that this is the primary thread again because another thread may have become
 		// the primary thread between the previous check and the lock.
-		if (UE::Private::FOutputDevicesPrimaryScopeLock PrimaryLock(*State); PrimaryLock.IsLocked() && State->IsPrimaryThread(ThreadId))
+		if (FOutputDevicesPrimaryScopeLock PrimaryLock(*State); PrimaryLock.IsLocked() && State->IsPrimaryThread(ThreadId))
 		{
-			State->FlushBufferedLines();
+			State->FlushBufferedItems();
 			State->BroadcastTo(ThreadId, State->BufferedOutputDevices, State->BufferedOutputDevicesCanBeUsedOnPanicThread,
 				UE_PROJECTION_MEMBER(FOutputDevice, Serialize),
 				Data, Verbosity, Category, RealTime);
@@ -612,7 +721,7 @@ void FOutputDeviceRedirector::Serialize(const TCHAR* const Data, const ELogVerbo
 	}
 
 	// Queue the line to serialize to buffered output devices from the primary thread.
-	if (State->BufferedLines.EnqueueAndReturnWasEmpty(Data, Category, Verbosity, RealTime))
+	if (State->BufferedItems.EnqueueAndReturnWasEmpty(Data, Category, Verbosity, RealTime))
 	{
 		if (FEvent* WakeEvent = State->ThreadWakeEvent.load(std::memory_order_acquire))
 		{
@@ -640,7 +749,7 @@ void FOutputDeviceRedirector::Flush()
 {
 	if (UE::Private::FOutputDevicesPrimaryScopeLock Lock(*State); Lock.IsLocked())
 	{
-		State->FlushBufferedLines();
+		State->FlushBufferedItems();
 		const uint32 ThreadId = FPlatformTLS::GetCurrentThreadId();
 		State->BroadcastTo(ThreadId, State->BufferedOutputDevices, State->BufferedOutputDevicesCanBeUsedOnPanicThread, &FOutputDevice::Flush);
 		State->BroadcastTo(ThreadId, State->UnbufferedOutputDevices, State->UnbufferedOutputDevicesCanBeUsedOnPanicThread, &FOutputDevice::Flush);
diff --git a/Engine/Source/Runtime/Core/Public/Misc/OutputDeviceRedirector.h b/Engine/Source/Runtime/Core/Public/Misc/OutputDeviceRedirector.h
--- a/Engine/Source/Runtime/Core/Public/Misc/OutputDeviceRedirector.h
+++ b/Engine/Source/Runtime/Core/Public/Misc/OutputDeviceRedirector.h
@@ -143,6 +143,13 @@ public:
 	 */
 	bool TryStartDedicatedPrimaryThread();
 
+	/**
+	 * Serializes the log record via all current output devices.
+	 *
+	 * The format string pointed to by the record must remain valid indefinitely.
+	 */
+	void SerializeRecord(const UE::FLogRecord& Record) final;
+
 	/**
 	 * Serializes the passed in data via all current output devices.
 	 *
