Logging: Added UTF-8 formatting of FLogRecord

#rnx

diff --git a/Engine/Source/Runtime/Core/Private/Logging/StructuredLog.cpp b/Engine/Source/Runtime/Core/Private/Logging/StructuredLog.cpp
--- a/Engine/Source/Runtime/Core/Private/Logging/StructuredLog.cpp
+++ b/Engine/Source/Runtime/Core/Private/Logging/StructuredLog.cpp
@@ -61,12 +61,33 @@ inline void FLogTemplateOp::Save(const FLogTemplateOp& Op, uint8*& Data)
 
 ///////////////////////////////////////////////////////////////////////////////////////////////////////////////////////
 
-static void LogFieldValue(FStringBuilderBase& Out, const FCbFieldView& Field)
+template <typename CharType>
+struct TLogFieldValueConstants;
+
+template <>
+struct TLogFieldValueConstants<UTF8CHAR>
+{
+	static inline const FAnsiStringView Null = ANSITEXTVIEW("null");
+	static inline const FAnsiStringView True = ANSITEXTVIEW("true");
+	static inline const FAnsiStringView False = ANSITEXTVIEW("false");
+};
+
+template <>
+struct TLogFieldValueConstants<WIDECHAR>
+{
+	static inline const FWideStringView Null = WIDETEXTVIEW("null");
+	static inline const FWideStringView True = WIDETEXTVIEW("true");
+	static inline const FWideStringView False = WIDETEXTVIEW("false");
+};
+
+template <typename CharType>
+static void LogFieldValue(TStringBuilderBase<CharType>& Out, const FCbFieldView& Field)
 {
+	using FConstants = TLogFieldValueConstants<CharType>;
 	switch (FCbValue Accessor = Field.GetValue(); Accessor.GetType())
 	{
 	case ECbFieldType::Null:
-		Out.Append(TEXTVIEW("null"));
+		Out.Append(FConstants::Null);
 		break;
 	case ECbFieldType::Object:
 	case ECbFieldType::UniformObject:
@@ -89,10 +110,10 @@ static void LogFieldValue(FStringBuilderBase& Out, const FCbFieldView& Field)
 		CompactBinaryToCompactJson(Field, Out);
 		break;
 	case ECbFieldType::BoolFalse:
-		Out.Append(TEXTVIEW("false"));
+		Out.Append(FConstants::False);
 		break;
 	case ECbFieldType::BoolTrue:
-		Out.Append(TEXTVIEW("true"));
+		Out.Append(FConstants::True);
 		break;
 	case ECbFieldType::ObjectAttachment:
 	case ECbFieldType::BinaryAttachment:
@@ -148,7 +169,8 @@ public:
 	static FLogTemplate* Create(const TCHAR* Format, const FLogField* Fields = nullptr, int32 FieldCount = 0);
 	static void Destroy(FLogTemplate* Template);
 
-	void FormatTo(FStringBuilderBase& Out, const TCHAR* Format, const FCbObjectView& Fields) const;
+	template <typename CharType>
+	void FormatTo(TStringBuilderBase<CharType>& Out, const TCHAR* Format, const FCbObjectView& Fields) const;
 
 private:
 	FLogTemplate() = default;
@@ -269,7 +291,8 @@ void FLogTemplate::Destroy(FLogTemplate* Template)
 	FMemory::Free(Template);
 }
 
-void FLogTemplate::FormatTo(FStringBuilderBase& Out, const TCHAR* Format, const FCbObjectView& Fields) const
+template <typename CharType>
+void FLogTemplate::FormatTo(TStringBuilderBase<CharType>& Out, const TCHAR* Format, const FCbObjectView& Fields) const
 {
 	using namespace Logging::Private;
 
@@ -353,23 +376,36 @@ FDateTime FLogTime::GetUtcTime() const
 
 ///////////////////////////////////////////////////////////////////////////////////////////////////////////////////////
 
-void FLogRecord::FormatMessageTo(FStringBuilderBase& Out) const
+template <typename CharType>
+static void FormatRecordMessageTo(TStringBuilderBase<CharType>& Out, const FLogRecord& Record)
 {
+	const TCHAR* Format = Record.GetFormat();
 	if (UNLIKELY(!Format))
 	{
 		return;
 	}
 
+	const FLogTemplate* Template = Record.GetTemplate();
 	if (LIKELY(Template))
 	{
-		return Template->FormatTo(Out, Format, Fields);
+		return Template->FormatTo(Out, Format, Record.GetFields());
 	}
 
 	FLogTemplate* LocalTemplate = FLogTemplate::Create(Format);
-	LocalTemplate->FormatTo(Out, Format, Fields);
+	LocalTemplate->FormatTo(Out, Format, Record.GetFields());
 	FLogTemplate::Destroy(LocalTemplate);
 }
 
+void FLogRecord::FormatMessageTo(FUtf8StringBuilderBase& Out) const
+{
+	FormatRecordMessageTo(Out, *this);
+}
+
+void FLogRecord::FormatMessageTo(FWideStringBuilderBase& Out) const
+{
+	FormatRecordMessageTo(Out, *this);
+}
+
 ///////////////////////////////////////////////////////////////////////////////////////////////////////////////////////
 
 } // UE
diff --git a/Engine/Source/Runtime/Core/Public/Logging/StructuredLog.h b/Engine/Source/Runtime/Core/Public/Logging/StructuredLog.h
--- a/Engine/Source/Runtime/Core/Public/Logging/StructuredLog.h
+++ b/Engine/Source/Runtime/Core/Public/Logging/StructuredLog.h
@@ -151,7 +151,8 @@ public:
 	void SetLine(int32 InLine) { Line = InLine; }
 
 	/** Formats the message using the format, template, and fields. */
-	UE_API void FormatMessageTo(FStringBuilderBase& Out) const;
+	UE_API void FormatMessageTo(FUtf8StringBuilderBase& Out) const;
+	UE_API void FormatMessageTo(FWideStringBuilderBase& Out) const;
 
 private:
 	const TCHAR* Format = nullptr;
