[Core]
* Added include to fix compile error with orphaned headers

diff --git a/Engine/Source/Runtime/Core/Public/Templates/SharedPointerFwd.h b/Engine/Source/Runtime/Core/Public/Templates/SharedPointerFwd.h
--- a/Engine/Source/Runtime/Core/Public/Templates/SharedPointerFwd.h
+++ b/Engine/Source/Runtime/Core/Public/Templates/SharedPointerFwd.h
@@ -2,6 +2,8 @@
 
 #pragma once
 
+#include "HAL/Platform.h" // For "uint8"
+
 /**
  * ESPMode is used select between either 'fast' or 'thread safe' shared pointer types.
  * This is only used by templates at compile time to generate one code path or another.
