fixed piping order to be FIFO in an edge case where multiple piped inline tasks depend on the same prerequisite, so when the prereq is completed they all get piped in one go

diff --git a/Engine/Source/Runtime/Core/Private/Tests/Tasks/TasksTest.cpp b/Engine/Source/Runtime/Core/Private/Tests/Tasks/TasksTest.cpp
--- a/Engine/Source/Runtime/Core/Private/Tests/Tasks/TasksTest.cpp
+++ b/Engine/Source/Runtime/Core/Private/Tests/Tasks/TasksTest.cpp
@@ -354,6 +354,14 @@ namespace UE { namespace TasksTests
 			Outer.Wait();
 		}
 
+		{
+			FTaskEvent T1{ UE_SOURCE_LOCATION };
+			FTaskEvent T2{ UE_SOURCE_LOCATION };
+			TArrayView<const FTask> Prereqs{ T1, T2 };
+			FTask Task = Launch(UE_SOURCE_LOCATION, [] {}, Prerequisites(Prereqs));
+			Task.Wait();
+		}
+
 #if TASKGRAPH_NEW_FRONTEND
 		{	// a basic test for a named thread task
 			FTask GTTask = Launch
@@ -571,6 +579,18 @@ namespace UE { namespace TasksTests
 			Task.Wait();
 		}
 
+		{	// multiple "inline" tasks piped in one go, check that pipe maintains FIFO order in this case
+			FPipe Pipe{ UE_SOURCE_LOCATION };
+			FTaskEvent Block{ UE_SOURCE_LOCATION };
+			bool bFirstDone = false;
+			bool bSecondDone = false;
+			FTask Task1 = Pipe.Launch(UE_SOURCE_LOCATION, [&] { check(!bSecondDone); bFirstDone = true; }, Prerequisites(Block), ETaskPriority::Normal, EExtendedTaskPriority::Inline);
+			FTask Task2 = Pipe.Launch(UE_SOURCE_LOCATION, [&] { check(bFirstDone); bSecondDone = true; }, Prerequisites(Block), ETaskPriority::Normal, EExtendedTaskPriority::Inline);
+			Block.Trigger();
+			Wait(TArray{ Task1, Task2 });
+			check(bFirstDone && bSecondDone);
+		}
+
 		UE_BENCHMARK(5, PipeStressTest<200, 100>);
 
 		return true;
diff --git a/Engine/Source/Runtime/Core/Public/Tasks/TaskPrivate.h b/Engine/Source/Runtime/Core/Public/Tasks/TaskPrivate.h
--- a/Engine/Source/Runtime/Core/Public/Tasks/TaskPrivate.h
+++ b/Engine/Source/Runtime/Core/Public/Tasks/TaskPrivate.h
@@ -466,7 +466,11 @@ namespace UE::Tasks
 				}
 
 				TArray<FTaskBase*> Subs;
-				Subsequents.PopAllAndClose(Subs);
+				Subsequents.PopAllAndClose(Subs); // gets `Subs` in LIFO order
+				// try to maintain FIFO order where it's possible. e.g. if multiple piped tasks (subsequents) depend on the same (this) task, 
+				// preserve the piping order, which is also the order in which subsequents were added as dependencies of this task
+				Algo::Reverse(Subs);
+
 				for (FTaskBase* Sub : Subs)
 				{
 					Sub->TryUnlock();
