minor improvement to handling FIFO order of task subsequents

diff --git a/Engine/Source/Runtime/Core/Private/Tests/Tasks/TasksTest.cpp b/Engine/Source/Runtime/Core/Private/Tests/Tasks/TasksTest.cpp
--- a/Engine/Source/Runtime/Core/Private/Tests/Tasks/TasksTest.cpp
+++ b/Engine/Source/Runtime/Core/Private/Tests/Tasks/TasksTest.cpp
@@ -354,14 +354,6 @@ namespace UE { namespace TasksTests
 			Outer.Wait();
 		}
 
-		{
-			FTaskEvent T1{ UE_SOURCE_LOCATION };
-			FTaskEvent T2{ UE_SOURCE_LOCATION };
-			TArrayView<const FTask> Prereqs{ T1, T2 };
-			FTask Task = Launch(UE_SOURCE_LOCATION, [] {}, Prerequisites(Prereqs));
-			Task.Wait();
-		}
-
 #if TASKGRAPH_NEW_FRONTEND
 		{	// a basic test for a named thread task
 			FTask GTTask = Launch
diff --git a/Engine/Source/Runtime/Core/Public/Tasks/TaskPrivate.h b/Engine/Source/Runtime/Core/Public/Tasks/TaskPrivate.h
--- a/Engine/Source/Runtime/Core/Public/Tasks/TaskPrivate.h
+++ b/Engine/Source/Runtime/Core/Public/Tasks/TaskPrivate.h
@@ -469,11 +469,9 @@ namespace UE::Tasks
 				Subsequents.PopAllAndClose(Subs); // gets `Subs` in LIFO order
 				// try to maintain FIFO order where it's possible. e.g. if multiple piped tasks (subsequents) depend on the same (this) task, 
 				// preserve the piping order, which is also the order in which subsequents were added as dependencies of this task
-				Algo::Reverse(Subs);
-
-				for (FTaskBase* Sub : Subs)
+				for (int32 i = Subs.Num() - 1; i >= 0; --i)
 				{
-					Sub->TryUnlock();
+					Subs[i]->TryUnlock();
 				}
 
 				// release nested tasks
