Flush the logs after writing <-- Suspend log parsing --> to avoid race condition with check output.

diff --git a/Engine/Source/Runtime/Core/Tests/Misc/AssertionMacrosTest.cpp b/Engine/Source/Runtime/Core/Tests/Misc/AssertionMacrosTest.cpp
--- a/Engine/Source/Runtime/Core/Tests/Misc/AssertionMacrosTest.cpp
+++ b/Engine/Source/Runtime/Core/Tests/Misc/AssertionMacrosTest.cpp
@@ -120,11 +120,13 @@ struct FSuspendLogParsing
 	FSuspendLogParsing()
 	{
 		UE_LOG(LogCore, Display, TEXT("<-- Suspend Log Parsing -->"));
+		GLog->Flush();
 	}
 
 	~FSuspendLogParsing()
 	{
 		UE_LOG(LogCore, Display, TEXT("<-- Resume Log Parsing -->"));
+		GLog->Flush();
 	}
 };
 
