Logging: Fixed a truncation warning in structured logging

#rnx

diff --git a/Engine/Source/Runtime/Core/Public/Logging/StructuredLog.h b/Engine/Source/Runtime/Core/Public/Logging/StructuredLog.h
--- a/Engine/Source/Runtime/Core/Public/Logging/StructuredLog.h
+++ b/Engine/Source/Runtime/Core/Public/Logging/StructuredLog.h
@@ -281,13 +281,13 @@ inline void LogIfActive(FieldArgTypes&&... FieldArgs)
 	{
 		if (Verbosity == ELogVerbosity::Fatal || !Log.Category.IsSuppressed(Verbosity))
 		{
-			if constexpr (sizeof...(FieldArgTypes))
+			if constexpr (sizeof...(FieldArgTypes) == 0)
 			{
-				LogWithFields<Log, LogFieldsType>(Forward<FieldArgTypes>(FieldArgs)...);
+				LogWithNoFields(Log);
 			}
 			else
 			{
-				LogWithNoFields(Log);
+				LogWithFields<Log, LogFieldsType>(Forward<FieldArgTypes>(FieldArgs)...);
 			}
 		}
 	}
