Fix for an error when creating an iostore toc if the max partition size is set to int64_max

Submitting on behalf of PJ Kack

diff --git a/Engine/Source/Runtime/Core/Private/IO/IoStore.cpp b/Engine/Source/Runtime/Core/Private/IO/IoStore.cpp
--- a/Engine/Source/Runtime/Core/Private/IO/IoStore.cpp
+++ b/Engine/Source/Runtime/Core/Private/IO/IoStore.cpp
@@ -2956,8 +2956,8 @@ TIoStatusOr<uint64> FIoStoreTocResource::Write(
 	else if (MaxPartitionSize)
 	{
 		const FIoStoreTocCompressedBlockEntry& LastBlock = TocResource.CompressionBlocks.Last();
-		uint64 LastBlockEnd = LastBlock.GetOffset() + LastBlock.GetCompressedSize();
-		TocHeader.PartitionCount = uint32(Align(LastBlockEnd, MaxPartitionSize) / MaxPartitionSize);
+		uint64 LastBlockEnd = LastBlock.GetOffset() + LastBlock.GetCompressedSize() - 1;
+		TocHeader.PartitionCount = IntCastChecked<uint32>(LastBlockEnd / MaxPartitionSize + 1);
 		check(TocHeader.PartitionCount > 0);
 		TocHeader.PartitionSize = MaxPartitionSize;
 	}
