Logging: Updated Windows and Linux feedback contexts to support structured logging

This sets up the base FFeedbackContext to do most of the work because there is very little behavioral difference between platforms, yet the majority of the code is duplicated for each platform.

#rnx

diff --git a/Engine/Source/Developer/DesktopPlatform/Private/Windows/WindowsNativeFeedbackContext.cpp b/Engine/Source/Developer/DesktopPlatform/Private/Windows/WindowsNativeFeedbackContext.cpp
--- a/Engine/Source/Developer/DesktopPlatform/Private/Windows/WindowsNativeFeedbackContext.cpp
+++ b/Engine/Source/Developer/DesktopPlatform/Private/Windows/WindowsNativeFeedbackContext.cpp
@@ -2,6 +2,7 @@
 
 #include "WindowsNativeFeedbackContext.h"
 #include "HAL/ThreadHeartBeat.h"
+#include "Logging/StructuredLog.h"
 #include "Misc/CoreMisc.h"
 #include "Misc/OutputDeviceHelper.h"
 #include "Misc/OutputDeviceConsole.h"
@@ -29,58 +30,42 @@ FWindowsNativeFeedbackContext::~FWindowsNativeFeedbackContext()
 	DestroySlowTaskWindow();
 }
 
-void FWindowsNativeFeedbackContext::Serialize( const TCHAR* V, ELogVerbosity::Type Verbosity, const class FName& Category )
+void FWindowsNativeFeedbackContext::Serialize(const TCHAR* V, ELogVerbosity::Type Verbosity, const FName& Category)
 {
-	// if we set the color for warnings or errors, then reset at the end of the function
-	// note, we have to set the colors directly without using the standard SET_WARN_COLOR macro
-	if( Verbosity==ELogVerbosity::Error || Verbosity==ELogVerbosity::Warning )
-	{
-		if( TreatWarningsAsErrors && Verbosity==ELogVerbosity::Warning )
-		{
-			Verbosity = ELogVerbosity::Error;
-		}
+	Serialize(V, Verbosity, Category, -1.0);
+}
 
-		FString Prefix;
-		if( Context )
-		{
-			Prefix = Context->GetContext() + TEXT(" : ");
-		}
-		FString Format = Prefix + FOutputDeviceHelper::FormatLogLine(Verbosity, Category, V);
+void FWindowsNativeFeedbackContext::Serialize(const TCHAR* V, ELogVerbosity::Type Verbosity, const FName& Category, double Time)
+{
+	FFeedbackContext::Serialize(V, Verbosity, Category, Time);
 
-		if(Verbosity == ELogVerbosity::Error)
-		{
-			// Only store off the message if running a commandlet.
-			if ( IsRunningCommandlet() )
-			{
-				AddError(Format);
-			}
-		}
-		else
+	// Buffer up the output during a slow task so that we can dump it all to the log console if the show log button is clicked
+	if (GIsSlowTask)
+	{
+		FScopeLock Lock(&CriticalSection);
+		if (hThread)
 		{
-			// Only store off the message if running a commandlet.
-			if ( IsRunningCommandlet() )
-			{
-				AddWarning(Format);
-			}
+			LogOutput += V;
+			LogOutput += TEXT("\r\n");
+			SetEvent(hUpdateEvent);
 		}
 	}
+}
 
-	if( GLogConsole && IsRunningCommandlet() && !GLog->IsRedirectingTo(GLogConsole) )
-	{
-		GLogConsole->Serialize( V, Verbosity, Category );
-	}
-	if( !GLog->IsRedirectingTo( this ) )
-	{
-		GLog->Serialize( V, Verbosity, Category );
-	}
+void FWindowsNativeFeedbackContext::SerializeRecord(const UE::FLogRecord& Record)
+{
+	FFeedbackContext::SerializeRecord(Record);
 
 	// Buffer up the output during a slow task so that we can dump it all to the log console if the show log button is clicked
-	if(GIsSlowTask)
+	if (GIsSlowTask)
 	{
+		TStringBuilder<512> Message;
+		Record.FormatMessageTo(Message);
+		Message.Append(TEXTVIEW("\r\n"));
 		FScopeLock Lock(&CriticalSection);
-		if(hThread != NULL)
+		if (hThread)
 		{
-			LogOutput += FString(V) + FString("\r\n");
+			LogOutput += Message.ToView();
 			SetEvent(hUpdateEvent);
 		}
 	}
diff --git a/Engine/Source/Developer/DesktopPlatform/Private/Windows/WindowsNativeFeedbackContext.h b/Engine/Source/Developer/DesktopPlatform/Private/Windows/WindowsNativeFeedbackContext.h
--- a/Engine/Source/Developer/DesktopPlatform/Private/Windows/WindowsNativeFeedbackContext.h
+++ b/Engine/Source/Developer/DesktopPlatform/Private/Windows/WindowsNativeFeedbackContext.h
@@ -19,7 +19,9 @@ public:
 	FWindowsNativeFeedbackContext();
 	virtual ~FWindowsNativeFeedbackContext();
 
-	virtual void Serialize( const TCHAR* V, ELogVerbosity::Type Verbosity, const class FName& Category ) override;
+	virtual void Serialize(const TCHAR* V, ELogVerbosity::Type Verbosity, const FName& Category) override;
+	virtual void Serialize(const TCHAR* V, ELogVerbosity::Type Verbosity, const FName& Category, double Time) override;
+	virtual void SerializeRecord(const UE::FLogRecord& Record) override;
 
 	virtual bool ReceivedUserCancel() override;
 	virtual void StartSlowTask( const FText& Task, bool bShouldShowCancelButton=false ) override;
diff --git a/Engine/Source/Runtime/ApplicationCore/Private/Unix/UnixFeedbackContext.cpp b/Engine/Source/Runtime/ApplicationCore/Private/Unix/UnixFeedbackContext.cpp
new file mode 100644
--- /dev/null
+++ b/Engine/Source/Runtime/ApplicationCore/Private/Unix/UnixFeedbackContext.cpp
@@ -0,0 +1,62 @@
+// Copyright Epic Games, Inc. All Rights Reserved.
+
+#include "Unix/UnixFeedbackContext.h"
+
+#include "Logging/StructuredLog.h"
+
+void FUnixFeedbackContext::Serialize(const TCHAR* V, ELogVerbosity::Type Verbosity, const FName& Category)
+{
+	Serialize(V, Verbosity, Category, -1.0);
+}
+
+void FUnixFeedbackContext::Serialize(const TCHAR* V, ELogVerbosity::Type Verbosity, const FName& Category, double Time)
+{
+	if (Verbosity == ELogVerbosity::Error || (Verbosity == ELogVerbosity::Warning && TreatWarningsAsErrors))
+	{
+		// send errors (warnings are too spammy) to syslog too (for zabbix etc)
+		LogErrorToSysLog(V, Category, Time);
+	}
+
+	FFeedbackContext::Serialize(V, Verbosity, Category, Time);
+}
+
+void FUnixFeedbackContext::SerializeRecord(const UE::FLogRecord& Record)
+{
+	const ELogVerbosity::Type Verbosity = Record.GetVerbosity();
+	if (Verbosity == ELogVerbosity::Error || (Verbosity == ELogVerbosity::Warning && TreatWarningsAsErrors))
+	{
+		// send errors (warnings are too spammy) to syslog too (for zabbix etc)
+		LogErrorRecordToSysLog(Record);
+	}
+
+	FFeedbackContext::SerializeRecord(Record);
+}
+
+/** Ask the user a binary question, returning their answer */
+bool FUnixFeedbackContext::YesNof(const FText& Question)
+{
+	if ((GIsClient || GIsEditor) && !GIsSilent && !FApp::IsUnattended())
+	{
+		//return( ::MessageBox( NULL, TempStr, *NSLOCTEXT("Core", "Question", "Question").ToString(), MB_YESNO|MB_TASKMODAL ) == IDYES);
+		STUBBED("+++++++++++++++ UNIXPLATFORMFEEDBACKCONTEXT.CPP DIALOG BOX PROMPT +++++++++++++++++++");
+		return true;
+	}
+	else
+	{
+		return false;
+	}
+}
+
+FORCENOINLINE void FUnixFeedbackContext::LogErrorToSysLog(const TCHAR* V, const FName& Category, double Time) const
+{
+	TStringBuilder<512> Line;
+	FormatLine(Line, V, ELogVerbosity::Error, Category, Time);
+	syslog(LOG_ERR | LOG_USER, "%s", (const char*)StringCast<UTF8CHAR>(*Line).Get());
+}
+
+FORCENOINLINE void FUnixFeedbackContext::LogErrorRecordToSysLog(const UE::FLogRecord& Record) const
+{
+	TStringBuilder<512> Line;
+	FormatRecordLine(Line, Record);
+	syslog(LOG_ERR | LOG_USER, "%s", (const char*)StringCast<UTF8CHAR>(*Line).Get());
+}
diff --git a/Engine/Source/Runtime/ApplicationCore/Private/Windows/WindowsFeedbackContext.cpp b/Engine/Source/Runtime/ApplicationCore/Private/Windows/WindowsFeedbackContext.cpp
--- a/Engine/Source/Runtime/ApplicationCore/Private/Windows/WindowsFeedbackContext.cpp
+++ b/Engine/Source/Runtime/ApplicationCore/Private/Windows/WindowsFeedbackContext.cpp
@@ -1,7 +1,7 @@
 // Copyright Epic Games, Inc. All Rights Reserved.
 
 #include "Windows/WindowsFeedbackContext.h"
-#include "Windows/WindowsHWrapper.h"
+
 #include "HAL/ThreadHeartBeat.h"
 #include "Internationalization/Internationalization.h"
 #include "Misc/App.h"
@@ -9,13 +9,10 @@
 
 bool FWindowsFeedbackContext::YesNof(const FText& Question)
 {
-	if ((GIsClient || GIsEditor) && ((GIsSilent != true) && (FApp::IsUnattended() != true)))
+	if ((GIsClient || GIsEditor) && !GIsSilent && !FApp::IsUnattended())
 	{
 		FSlowHeartBeatScope SuspendHeartBeat;
-		return(::MessageBox(NULL, Question.ToString().GetCharArray().GetData(), *NSLOCTEXT("Core", "Question", "Question").ToString(), MB_YESNO | MB_TASKMODAL) == IDYES);
-	}
-	else
-	{
-		return false;
+		return ::MessageBox(nullptr, Question.ToString().GetCharArray().GetData(), *NSLOCTEXT("Core", "Question", "Question").ToString(), MB_YESNO | MB_TASKMODAL) == IDYES;
 	}
+	return false;
 }
diff --git a/Engine/Source/Runtime/ApplicationCore/Public/Unix/UnixFeedbackContext.h b/Engine/Source/Runtime/ApplicationCore/Public/Unix/UnixFeedbackContext.h
--- a/Engine/Source/Runtime/ApplicationCore/Public/Unix/UnixFeedbackContext.h
+++ b/Engine/Source/Runtime/ApplicationCore/Public/Unix/UnixFeedbackContext.h
@@ -13,8 +13,6 @@
 #include "Misc/FeedbackContext.h"
 #include <syslog.h>
 
-class Error;
-
 /**
  * Feedback context implementation for Unix.
  */
@@ -34,88 +32,40 @@ public:
 	, SlowTaskCount( 0 )
 	{}
 
-	void Serialize( const TCHAR* V, ELogVerbosity::Type Verbosity, const class FName& Category )
-	{
-		// if we set the color for warnings or errors, then reset at the end of the function
-		// note, we have to set the colors directly without using the standard SET_WARN_COLOR macro
-		if (Verbosity == ELogVerbosity::Error || Verbosity == ELogVerbosity::Warning)
-		{
-			if (TreatWarningsAsErrors && Verbosity==ELogVerbosity::Warning)
-			{
-				Verbosity = ELogVerbosity::Error;
-			}
-
-			FString Prefix;
-			if (Context)
-			{
-				Prefix = Context->GetContext() + TEXT(" : ");
-			}
-			FString Format = Prefix + FOutputDeviceHelper::FormatLogLine(Verbosity, Category, V);
-
-			if (Verbosity == ELogVerbosity::Error)
-			{
-				// Only store off the message if running a commandlet.
-				if (IsRunningCommandlet())
-				{
-					AddError(Format);
-				}
+	void Serialize(const TCHAR* V, ELogVerbosity::Type Verbosity, const FName& Category) override;
+	void Serialize(const TCHAR* V, ELogVerbosity::Type Verbosity, const FName& Category, double Time) override;
+	void SerializeRecord(const UE::FLogRecord& Record) override;
 
-				// send errors (warnings are too spammy) to syslog too (for zabbix etc)
-				syslog(LOG_ERR | LOG_USER, "%s", StringCast< char >(*Format).Get());
-			}
-			else
-			{
-				// Only store off the message if running a commandlet.
-				if ( IsRunningCommandlet() )
-				{
-					AddWarning(Format);
-				}
-			}
-		}
-
-		if( GLogConsole && IsRunningCommandlet() && !GLog->IsRedirectingTo(GLogConsole) )
-		{
-			GLogConsole->Serialize( V, Verbosity, Category );
-		}
-		if( !GLog->IsRedirectingTo( this ) )
-		{
-			GLog->Serialize( V, Verbosity, Category );
-		}
-	}
 	/** Ask the user a binary question, returning their answer */
-	virtual bool YesNof( const FText& Question )
-	{
-		if( ( GIsClient || GIsEditor ) && ( ( GIsSilent != true ) && ( FApp::IsUnattended() != true ) ) )
-		{
-			//return( ::MessageBox( NULL, TempStr, *NSLOCTEXT("Core", "Question", "Question").ToString(), MB_YESNO|MB_TASKMODAL ) == IDYES);
-			STUBBED("+++++++++++++++ UNIXPLATFORMFEEDBACKCONTEXTPRIVATE.H DIALOG BOX PROMPT +++++++++++++++++++");
-			return true;
-		}
-		else
-		{
-			return false;
-		}
-	}
+	bool YesNof(const FText& Question) override;
 
-	void BeginSlowTask( const FText& Task, bool ShowProgressDialog, bool bShowCancelButton=false )
+	void BeginSlowTask(const FText& Task, bool ShowProgressDialog, bool bShowCancelButton=false)
 	{
 		GIsSlowTask = ++SlowTaskCount>0;
 	}
+
 	void EndSlowTask()
 	{
 		check(SlowTaskCount>0);
 		GIsSlowTask = --SlowTaskCount>0;
 	}
-	virtual bool StatusUpdate( int32 Numerator, int32 Denominator, const FText& StatusText )
+
+	bool StatusUpdate(int32 Numerator, int32 Denominator, const FText& StatusText)
 	{
 		return true;
 	}
-	FContextSupplier* GetContext() const
+
+	FContextSupplier* GetContext() const override
 	{
 		return Context;
 	}
-	void SetContext( FContextSupplier* InSupplier )
+
+	void SetContext(FContextSupplier* InSupplier) override
 	{
 		Context = InSupplier;
 	}
+
+private:
+	void LogErrorToSysLog(const TCHAR* V, const FName& Category, double Time) const;
+	void LogErrorRecordToSysLog(const UE::FLogRecord& Record) const;
 };
diff --git a/Engine/Source/Runtime/ApplicationCore/Public/Windows/WindowsFeedbackContext.h b/Engine/Source/Runtime/ApplicationCore/Public/Windows/WindowsFeedbackContext.h
--- a/Engine/Source/Runtime/ApplicationCore/Public/Windows/WindowsFeedbackContext.h
+++ b/Engine/Source/Runtime/ApplicationCore/Public/Windows/WindowsFeedbackContext.h
@@ -4,14 +4,10 @@
 
 #include "CoreTypes.h"
 #include "Misc/CoreMisc.h"
-#include "Misc/OutputDeviceRedirector.h"
-#include "Misc/OutputDeviceHelper.h"
 #include "Misc/FeedbackContext.h"
 #include "Misc/OutputDeviceConsole.h"
-
-class FContextSupplier;
-class FName;
-class FText;
+#include "Misc/OutputDeviceHelper.h"
+#include "Misc/OutputDeviceRedirector.h"
 
 /**
  * Feedback context implementation for windows.
@@ -19,71 +15,11 @@ class FText;
 class APPLICATIONCORE_API FWindowsFeedbackContext : public FFeedbackContext
 {
 	/** Context information for warning and error messages */
-	FContextSupplier*	Context;
+	FContextSupplier* Context = nullptr;
 
 public:
-	// Variables.
-
-	// Constructor.
-	FWindowsFeedbackContext()
-	: FFeedbackContext()
-	, Context( NULL )
-	{}
-	void Serialize( const TCHAR* V, ELogVerbosity::Type Verbosity, const class FName& Category )
-	{
-		// if we set the color for warnings or errors, then reset at the end of the function
-		// note, we have to set the colors directly without using the standard SET_WARN_COLOR macro
-		if( Verbosity==ELogVerbosity::Error || Verbosity==ELogVerbosity::Warning )
-		{
-			if( TreatWarningsAsErrors && Verbosity==ELogVerbosity::Warning )
-			{
-				Verbosity = ELogVerbosity::Error;
-			}
+	bool YesNof(const FText& Question) override;
 
-			FString Prefix;
-			if( Context )
-			{
-				Prefix = Context->GetContext() + TEXT(" : ");
-			}
-			FString Format = Prefix + FOutputDeviceHelper::FormatLogLine(Verbosity, Category, V);
-
-			if(Verbosity == ELogVerbosity::Error)
-			{
-				// Only store off the message if running a commandlet.
-				if ( IsRunningCommandlet() )
-				{
-					AddError(Format);
-				}
-			}
-			else
-			{
-				// Only store off the message if running a commandlet.
-				if ( IsRunningCommandlet() )
-				{
-					AddWarning(Format);
-				}
-			}
-		}
-
-		if( GLogConsole && IsRunningCommandlet() && !GLog->IsRedirectingTo(GLogConsole) )
-		{
-			GLogConsole->Serialize( V, Verbosity, Category );
-		}
-		if( !GLog->IsRedirectingTo( this ) )
-		{
-			GLog->Serialize( V, Verbosity, Category );
-		}
-	}
-
-	virtual bool YesNof(const FText& Question) override;
-
-	FContextSupplier* GetContext() const
-	{
-		return Context;
-	}
-	void SetContext( FContextSupplier* InSupplier )
-	{
-		Context = InSupplier;
-	}
+	FContextSupplier* GetContext() const { return Context; }
+	void SetContext(FContextSupplier* InContext) { Context = InContext; }
 };
-
diff --git a/Engine/Source/Runtime/Core/Private/HAL/FeedbackContextAnsi.cpp b/Engine/Source/Runtime/Core/Private/HAL/FeedbackContextAnsi.cpp
new file mode 100644
--- /dev/null
+++ b/Engine/Source/Runtime/Core/Private/HAL/FeedbackContextAnsi.cpp
@@ -0,0 +1,95 @@
+// Copyright Epic Games, Inc. All Rights Reserved.
+
+#include "HAL/FeedbackContextAnsi.h"
+
+#include "Logging/StructuredLog.h"
+
+void FFeedbackContextAnsi::LocalPrint(const TCHAR* Str)
+{
+#if PLATFORM_APPLE || PLATFORM_UNIX
+	printf("%s", (const char*)StringCast<UTF8CHAR>(Str).Get());
+#elif PLATFORM_MICROSOFT
+	wprintf(TEXT("%ls"), Str);
+#else
+	// If this function ever gets more complicated, we could make a PlatformMisc::Printf, and each platform can then 
+	// do the right thing. For instance, LocalPrint is OutputDebugString on Windows, which messes up a lot of stuff
+	FPlatformMisc::LocalPrint(Str);
+#endif
+	fflush(stdout);
+}
+
+bool FFeedbackContextAnsi::IsUsingLocalPrint() const
+{
+	// When -stdout is specified then FOutputDeviceStdOutput will be installed and pipe logging to stdout.
+	// Skip calling LocalPrint in that case or else duplicate messages will be written to stdout.
+	// A similar issue happens when a Console is shown.
+	static bool bIsUsingLocalPrint =
+		!FParse::Param(FCommandLine::Get(), TEXT("stdout")) &&
+		(!GLogConsole || !GLogConsole->IsShown());
+	return bIsUsingLocalPrint;
+}
+
+void FFeedbackContextAnsi::Serialize(const TCHAR* V, ELogVerbosity::Type Verbosity, const FName& Category)
+{
+	Serialize(V, Verbosity, Category, -1.0);
+}
+
+void FFeedbackContextAnsi::Serialize(const TCHAR* V, ELogVerbosity::Type Verbosity, const FName& Category, double Time)
+{
+	if (IsUsingLocalPrint() && (Verbosity == ELogVerbosity::Error || Verbosity == ELogVerbosity::Warning || Verbosity == ELogVerbosity::Display))
+	{
+		TStringBuilder<512> Line;
+		FormatLine(Line, V, Verbosity, Category, Time);
+		Line.AppendChar(TEXT('\n'));
+		LocalPrint(*Line);
+	}
+
+	FFeedbackContext::Serialize(V, Verbosity, Category, Time);
+
+	if (AuxOut)
+	{
+		AuxOut->Serialize(V, Verbosity, Category, Time);
+	}
+}
+
+void FFeedbackContextAnsi::SerializeRecord(const UE::FLogRecord& Record)
+{
+	const ELogVerbosity::Type Verbosity = Record.GetVerbosity();
+	if (IsUsingLocalPrint() && (Verbosity == ELogVerbosity::Error || Verbosity == ELogVerbosity::Warning || Verbosity == ELogVerbosity::Display))
+	{
+		TStringBuilder<512> Line;
+		FormatRecordLine(Line, Record);
+		Line.AppendChar(TEXT('\n'));
+		LocalPrint(*Line);
+	}
+
+	FFeedbackContext::SerializeRecord(Record);
+
+	if (AuxOut)
+	{
+		AuxOut->SerializeRecord(Record);
+	}
+}
+
+bool FFeedbackContextAnsi::YesNof(const FText& Question)
+{
+	if (GIsClient || GIsEditor)
+	{
+		LocalPrint(*Question.ToString());
+		LocalPrint(TEXT(" (Y/N): "));
+		if (GIsSilent || FApp::IsUnattended())
+		{
+			LocalPrint(TEXT("Y\n"));
+			return true;
+		}
+		else
+		{
+			char InputText[256];
+			if (fgets(InputText, sizeof(InputText), stdin) != nullptr)
+			{
+				return InputText[0] == 'Y' || InputText[0] == 'y';
+			}
+		}
+	}
+	return true;
+}
diff --git a/Engine/Source/Runtime/Core/Private/Misc/FeedbackContext.cpp b/Engine/Source/Runtime/Core/Private/Misc/FeedbackContext.cpp
--- a/Engine/Source/Runtime/Core/Private/Misc/FeedbackContext.cpp
+++ b/Engine/Source/Runtime/Core/Private/Misc/FeedbackContext.cpp
@@ -4,19 +4,155 @@
 
 #include "CoreTypes.h"
 #include "HAL/PlatformTime.h"
+#include "Logging/StructuredLog.h"
+#include "Misc/App.h"
 #include "Misc/AssertionMacros.h"
+#include "Misc/CoreMisc.h"
+#include "Misc/OutputDeviceConsole.h"
+#include "Misc/OutputDeviceHelper.h"
+#include "Misc/OutputDeviceRedirector.h"
 #include "Misc/SlowTask.h"
+#include "Misc/StringBuilder.h"
 
-FFeedbackContext::FFeedbackContext()
-	: TreatWarningsAsErrors(0)
-{
-}
+FFeedbackContext::FFeedbackContext() = default;
 
 FFeedbackContext::~FFeedbackContext()
 {
 	ensureMsgf(LegacyAPIScopes.Num() == 0, TEXT("EndSlowTask has not been called for %d outstanding tasks"), LegacyAPIScopes.Num());
 }
 
+void FFeedbackContext::Serialize(const TCHAR* V, ELogVerbosity::Type Verbosity, const FName& Category)
+{
+	if (IsRunningCommandlet())
+	{
+		if (Verbosity == ELogVerbosity::Error || Verbosity == ELogVerbosity::Warning)
+		{
+			AddToHistory(V, Verbosity, Category, -1.0);
+		}
+		if (GLogConsole && !GLog->IsRedirectingTo(GLogConsole))
+		{
+			GLogConsole->Serialize(V, Verbosity, Category);
+		}
+	}
+	if (!GLog->IsRedirectingTo(this))
+	{
+		GLog->Serialize(V, Verbosity, Category);
+	}
+}
+
+void FFeedbackContext::Serialize(const TCHAR* V, ELogVerbosity::Type Verbosity, const FName& Category, double Time)
+{
+	if (IsRunningCommandlet())
+	{
+		if (Verbosity == ELogVerbosity::Error || Verbosity == ELogVerbosity::Warning)
+		{
+			AddToHistory(V, Verbosity, Category, Time);
+		}
+		if (GLogConsole && !GLog->IsRedirectingTo(GLogConsole))
+		{
+			GLogConsole->Serialize(V, Verbosity, Category, Time);
+		}
+	}
+	if (!GLog->IsRedirectingTo(this))
+	{
+		GLog->Serialize(V, Verbosity, Category, Time);
+	}
+}
+
+void FFeedbackContext::SerializeRecord(const UE::FLogRecord& Record)
+{
+	if (IsRunningCommandlet())
+	{
+		const ELogVerbosity::Type Verbosity = Record.GetVerbosity();
+		if (Verbosity == ELogVerbosity::Error || Verbosity == ELogVerbosity::Warning)
+		{
+			AddRecordToHistory(Record);
+		}
+		if (GLogConsole && !GLog->IsRedirectingTo(GLogConsole))
+		{
+			GLogConsole->SerializeRecord(Record);
+		}
+	}
+	if (!GLog->IsRedirectingTo(this))
+	{
+		GLog->SerializeRecord(Record);
+	}
+}
+
+void FFeedbackContext::FormatLine(FStringBuilderBase& Out, const TCHAR* V, ELogVerbosity::Type Verbosity, const FName& Category, double Time, ELogVerbosity::Type* OutVerbosity) const
+{
+	if (TreatWarningsAsErrors && Verbosity == ELogVerbosity::Warning)
+	{
+		Verbosity = ELogVerbosity::Error;
+	}
+	if (OutVerbosity)
+	{
+		*OutVerbosity = Verbosity;
+	}
+	if (FContextSupplier* Context = GetContext())
+	{
+		Out.Append(Context->GetContext()).Append(TEXTVIEW(" : "));
+	}
+	FOutputDeviceHelper::AppendFormatLogLine(Out, Verbosity, Category, V);
+}
+
+void FFeedbackContext::FormatRecordLine(FStringBuilderBase& Out, const UE::FLogRecord& Record, ELogVerbosity::Type* OutVerbosity) const
+{
+	ELogVerbosity::Type Verbosity = Record.GetVerbosity();
+	if (TreatWarningsAsErrors && Verbosity == ELogVerbosity::Warning)
+	{
+		Verbosity = ELogVerbosity::Error;
+	}
+	if (OutVerbosity)
+	{
+		*OutVerbosity = Verbosity;
+	}
+	if (FContextSupplier* Context = GetContext())
+	{
+		Out.Append(Context->GetContext()).Append(TEXTVIEW(" : "));
+	}
+	FOutputDeviceHelper::AppendFormatLogLine(Out, Verbosity, Record.GetCategory());
+	Record.FormatMessageTo(Out);
+}
+
+void FFeedbackContext::AddToHistory(const TCHAR* V, ELogVerbosity::Type Verbosity, const FName& Category, double Time)
+{
+	TStringBuilder<512> Line;
+	FormatLine(Line, V, Verbosity, Category, Time, &Verbosity);
+	if (Verbosity == ELogVerbosity::Error)
+	{
+		AddError(FString(Line));
+	}
+	else
+	{
+		AddWarning(FString(Line));
+	}
+}
+
+void FFeedbackContext::AddRecordToHistory(const UE::FLogRecord& Record)
+{
+	TStringBuilder<512> Line;
+	ELogVerbosity::Type Verbosity;
+	FormatRecordLine(Line, Record, &Verbosity);
+	if (Verbosity == ELogVerbosity::Warning)
+	{
+		AddWarning(FString(Line));
+	}
+	else if (Verbosity == ELogVerbosity::Error)
+	{
+		AddError(FString(Line));
+	}
+}
+
+bool FFeedbackContext::YesNof(const FText& Question)
+{
+	if (!GIsSilent && !FApp::IsUnattended())
+	{
+		FPlatformMisc::LowLevelOutputDebugString(*Question.ToString());
+	}
+	return false;
+}
+
 void FFeedbackContext::RequestUpdateUI(bool bForceUpdate)
 {
 	// Only update a maximum of 5 times a second
diff --git a/Engine/Source/Runtime/Core/Public/HAL/FeedbackContextAnsi.h b/Engine/Source/Runtime/Core/Public/HAL/FeedbackContextAnsi.h
--- a/Engine/Source/Runtime/Core/Public/HAL/FeedbackContextAnsi.h
+++ b/Engine/Source/Runtime/Core/Public/HAL/FeedbackContextAnsi.h
@@ -3,125 +3,32 @@
 #pragma once
 
 #include "CoreTypes.h"
-#include "Containers/UnrealString.h"
 #include "Containers/StringConv.h"
+#include "Containers/UnrealString.h"
 #include "CoreGlobals.h"
-#include "Misc/OutputDeviceRedirector.h"
-#include "Misc/OutputDeviceConsole.h"
 #include "Misc/App.h"
+#include "Misc/OutputDeviceConsole.h"
 #include "Misc/OutputDeviceHelper.h"
+#include "Misc/OutputDeviceRedirector.h"
 #include "Misc/FeedbackContext.h"
 
-class Error;
-
-/*-----------------------------------------------------------------------------
-	FFeedbackContextAnsi.
------------------------------------------------------------------------------*/
-
-//
-// Feedback context.
-//
 class FFeedbackContextAnsi : public FFeedbackContext
 {
 public:
-	// Variables.
-	FContextSupplier*	Context;
-	FOutputDevice*		AuxOut;
-
-	// Local functions.
-	void LocalPrint( const TCHAR* Str )
-	{
-#if PLATFORM_APPLE || PLATFORM_UNIX
-		printf("%s", TCHAR_TO_UTF8(Str));
-#elif PLATFORM_MICROSOFT
-		wprintf(TEXT("%ls"), Str);
-#else
-		// If this function ever gets more complicated, we could make a PlatformMisc::Printf, and each platform can then 
-		// do the right thing. For instance, LocalPrint is OutputDebugString on Windows, which messes up a lot of stuff
-		FPlatformMisc::LocalPrint(Str);
-#endif
-	}
+	FContextSupplier*	Context = nullptr;
+	FOutputDevice*		AuxOut = nullptr;
 
-	// Constructor.
-	FFeedbackContextAnsi()
-	: FFeedbackContext()
-	, Context(nullptr)
-	, AuxOut(nullptr)
-	{}
-	void Serialize( const TCHAR* V, ELogVerbosity::Type Verbosity, const class FName& Category ) override
-	{
-		// When -stdout is specified then FOutputDeviceStdOutput will be installed and pipe logging to stdout. 
-		// If so don't use LocalPrint or else duplicate messages will be written to stdout.
-		// A similar issue happens when a Console is shown.
-		static bool bUsingStdOut = FParse::Param(FCommandLine::Get(), TEXT("stdout")) ||
-			(GLogConsole != nullptr && GLogConsole->IsShown());
+	CORE_API static void LocalPrint(const TCHAR* Str);
 
-		if (bUsingStdOut == false &&
-			(Verbosity == ELogVerbosity::Error || Verbosity == ELogVerbosity::Warning || Verbosity == ELogVerbosity::Display))
-		{
-			if( TreatWarningsAsErrors && Verbosity==ELogVerbosity::Warning )
-			{
-				Verbosity = ELogVerbosity::Error;
-			}
+	CORE_API void Serialize(const TCHAR* V, ELogVerbosity::Type Verbosity, const FName& Category) override;
+	CORE_API void Serialize(const TCHAR* V, ELogVerbosity::Type Verbosity, const FName& Category, double Time) override;
+	CORE_API void SerializeRecord(const UE::FLogRecord& Record) override;
 
-			FString Prefix;
-			if( Context )
-			{
-				Prefix = Context->GetContext() + TEXT(" : ");
-			}
-			FString Format = Prefix + FOutputDeviceHelper::FormatLogLine(Verbosity, Category, V);
+	CORE_API bool YesNof(const FText& Question) override;
 
-			// Only store off the message if running a commandlet.
-			if ( IsRunningCommandlet() )
-			{
-				if (Verbosity == ELogVerbosity::Error)
-				{
-					AddError(Format);
-				}
-				else if (Verbosity == ELogVerbosity::Warning)
-				{
-					AddWarning(Format);
-				}
-			}
-			Format.AppendChar(TEXT('\n'));
-			LocalPrint(*Format);
-		}
-		else if (Verbosity == ELogVerbosity::SetColor)
-		{
-		}
-		if( !GLog->IsRedirectingTo( this ) )
-			GLog->Serialize( V, Verbosity, Category );
-		if( AuxOut )
-			AuxOut->Serialize( V, Verbosity, Category );
-		fflush( stdout );
-	}
+	FContextSupplier* GetContext() const override { return Context; }
+	void SetContext(FContextSupplier* InContext) override { Context = InContext; }
 
-	virtual bool YesNof( const FText& Question ) override
-	{
-		if( (GIsClient || GIsEditor) )
-		{
-			LocalPrint( *Question.ToString() );
-			LocalPrint( TEXT(" (Y/N): ") );
-			if( ( GIsSilent == true ) || ( FApp::IsUnattended() == true ) )
-			{
-				LocalPrint( TEXT("Y") );
-				return true;
-			}
-			else
-			{
-				char InputText[256];
-				if (fgets( InputText, sizeof(InputText), stdin ) != nullptr)
-				{
-					return (InputText[0]=='Y' || InputText[0]=='y');
-				}
-			}
-		}
-		return true;
-	}
-
-	void SetContext( FContextSupplier* InSupplier ) override
-	{
-		Context = InSupplier;
-	}
+private:
+	bool IsUsingLocalPrint() const;
 };
-
diff --git a/Engine/Source/Runtime/Core/Public/Misc/FeedbackContext.h b/Engine/Source/Runtime/Core/Public/Misc/FeedbackContext.h
--- a/Engine/Source/Runtime/Core/Public/Misc/FeedbackContext.h
+++ b/Engine/Source/Runtime/Core/Public/Misc/FeedbackContext.h
@@ -3,6 +3,7 @@
 #pragma once
 
 #include "Containers/Array.h"
+#include "Containers/StringFwd.h"
 #include "Containers/UnrealString.h"
 #include "CoreGlobals.h"
 #include "CoreTypes.h"
@@ -26,16 +27,19 @@ class CORE_API FFeedbackContext
 	: public FOutputDevice
 {
 public:
+	virtual void Serialize(const TCHAR* V, ELogVerbosity::Type Verbosity, const FName& Category) override;
+	virtual void Serialize(const TCHAR* V, ELogVerbosity::Type Verbosity, const FName& Category, double Time) override;
+	virtual void SerializeRecord(const UE::FLogRecord& Record) override;
 
 	/** Ask the user a binary question, returning their answer */
-	virtual bool YesNof( const FText& Question ) { return false; }
+	virtual bool YesNof( const FText& Question );
 	
 	/**
 	 * Whether or not the user has canceled out of the progress dialog
 	 * (i.e. the ongoing slow task or the last one that ran).
 	 * The user cancel flag is reset when starting a new root slow task.
 	 */
-	virtual bool ReceivedUserCancel() { return false; };
+	virtual bool ReceivedUserCancel() { return false; }
 	
 	/** Public const access to the current state of the scope stack */
 	FORCEINLINE const FSlowTaskStack& GetScopeStack() const
@@ -79,18 +83,20 @@ protected:
 	/** Called to check whether we are playing in editor when starting a slow task */
 	virtual bool IsPlayingInEditor() const;
 
+	void FormatLine(FStringBuilderBase& Out, const TCHAR* V, ELogVerbosity::Type Verbosity, const FName& Category, double Time, ELogVerbosity::Type* OutVerbosity = nullptr) const;
+	void FormatRecordLine(FStringBuilderBase& Out, const UE::FLogRecord& Record, ELogVerbosity::Type* OutVerbosity = nullptr) const;
+
 public:
-	virtual FContextSupplier* GetContext() const { return NULL; }
-	virtual void SetContext( FContextSupplier* InSupplier ) {}
+	virtual FContextSupplier* GetContext() const { return nullptr; }
+	virtual void SetContext(FContextSupplier* InContext) {}
 
 	/** Shows/Closes Special Build Progress dialogs */
 	virtual TWeakPtr<class SBuildProgressWidget> ShowBuildProgressWindow() {return TWeakPtr<class SBuildProgressWidget>();}
 	virtual void CloseBuildProgressWindow() {}
 
-	bool	TreatWarningsAsErrors;
+	bool	TreatWarningsAsErrors = false;
 
 	FFeedbackContext();
-
 	virtual ~FFeedbackContext();
 
 	/** Gets warnings history */
@@ -131,6 +137,9 @@ public:
 	}
 
 private:
+	void AddToHistory(const TCHAR* V, ELogVerbosity::Type Verbosity, const FName& Category, double Time);
+	void AddRecordToHistory(const UE::FLogRecord& Record);
+
 	FFeedbackContext(const FFeedbackContext&);
 	FFeedbackContext& operator=(const FFeedbackContext&);
 
@@ -180,6 +189,11 @@ protected:
 		FScopeLock WarningsAndErrorsLock(&WarningsAndErrorsCritical);
 		Warnings.Add(InWarning);
 	}
+	void AddWarning(FString&& InWarning)
+	{
+		FScopeLock WarningsAndErrorsLock(&WarningsAndErrorsCritical);
+		Warnings.Add(MoveTemp(InWarning));
+	}
 
 	/**
 	* Adds a new error message to errors history.
@@ -190,4 +204,9 @@ protected:
 		FScopeLock WarningsAndErrorsLock(&WarningsAndErrorsCritical);
 		Errors.Add(InError);
 	}
+	void AddError(FString&& InError)
+	{
+		FScopeLock WarningsAndErrorsLock(&WarningsAndErrorsCritical);
+		Errors.Add(MoveTemp(InError));
+	}
 };
