Collapse set of defines to enable TBB/Mimalloc to avoid wrong allocator fallback being mistakenly choosen

#rnx

diff --git a/Engine/Platforms/Hololens/Source/Runtime/Core/Private/HoloLensPlatformMemory.cpp b/Engine/Platforms/Hololens/Source/Runtime/Core/Private/HoloLensPlatformMemory.cpp
--- a/Engine/Platforms/Hololens/Source/Runtime/Core/Private/HoloLensPlatformMemory.cpp
+++ b/Engine/Platforms/Hololens/Source/Runtime/Core/Private/HoloLensPlatformMemory.cpp
@@ -69,7 +69,7 @@ FMalloc* FHoloLensPlatformMemory::BaseAllocator()
 {
 #if FORCE_ANSI_ALLOCATOR
 	return new FMallocAnsi();
-#elif (WITH_EDITORONLY_DATA || IS_PROGRAM) && TBB_ALLOCATOR_ALLOWED
+#elif (WITH_EDITORONLY_DATA || IS_PROGRAM) && TBBMALLOC_ENABLED
 	return new FMallocTBB();
 #else
 	return new FMallocBinned3();
diff --git a/Engine/Source/Runtime/Core/Private/HAL/MallocTBB.cpp b/Engine/Source/Runtime/Core/Private/HAL/MallocTBB.cpp
--- a/Engine/Source/Runtime/Core/Private/HAL/MallocTBB.cpp
+++ b/Engine/Source/Runtime/Core/Private/HAL/MallocTBB.cpp
@@ -12,7 +12,7 @@
 #endif
 
 // Only use for supported platforms
-#if PLATFORM_SUPPORTS_TBB && TBB_ALLOCATOR_ALLOWED
+#if TBBMALLOC_ENABLED
 
 /** Value we fill a memory block with after it is free, in UE_BUILD_DEBUG **/
 #define DEBUG_FILL_FREED (0xdd)
@@ -176,4 +176,4 @@ void FMallocTBB::Trim(bool bTrimThreadCaches)
 #undef DEBUG_FILL_FREED
 #undef DEBUG_FILL_NEW
 
-#endif // PLATFORM_SUPPORTS_TBB && TBB_ALLOCATOR_ALLOWED
+#endif // TBBMALLOC_ENABLED
diff --git a/Engine/Source/Runtime/Core/Private/Mac/MacPlatformMemory.cpp b/Engine/Source/Runtime/Core/Private/Mac/MacPlatformMemory.cpp
--- a/Engine/Source/Runtime/Core/Private/Mac/MacPlatformMemory.cpp
+++ b/Engine/Source/Runtime/Core/Private/Mac/MacPlatformMemory.cpp
@@ -68,7 +68,7 @@ FMalloc* FMacPlatformMemory::BaseAllocator()
 	{
 		AllocatorToUse = EMemoryAllocatorToUse::Ansi;
 	}
-	else if ((WITH_EDITORONLY_DATA || IS_PROGRAM) && TBB_ALLOCATOR_ALLOWED)
+	else if ((WITH_EDITORONLY_DATA || IS_PROGRAM) && TBBMALLOC_ENABLED)
 	{
 		AllocatorToUse = EMemoryAllocatorToUse::TBB;
 	}
@@ -92,12 +92,12 @@ FMalloc* FMacPlatformMemory::BaseAllocator()
 		AllocatorToUse = EMemoryAllocatorToUse::Stomp;
 	}
 #endif // WITH_MALLOC_STOMP
-#if PLATFORM_SUPPORTS_MIMALLOC && MIMALLOC_ALLOCATOR_ALLOWED
+#if MIMALLOC_ENABLED
 	else if (HasArg("-mimalloc"))
 	{
 		AllocatorToUse = EMemoryAllocatorToUse::Mimalloc;
 	}
-#endif // PLATFORM_SUPPORTS_MIMALLOC && MIMALLOC_ALLOCATOR_ALLOWED
+#endif // MIMALLOC_ENABLED
 
 	// Force ANSI malloc with TSAN
 #if defined(__has_feature)
@@ -114,11 +114,11 @@ FMalloc* FMacPlatformMemory::BaseAllocator()
 	case EMemoryAllocatorToUse::Stomp:
 		return new FMallocStomp();
 #endif
-#if TBB_ALLOCATOR_ALLOWED
+#if TBBMALLOC_ENABLED
 	case EMemoryAllocatorToUse::TBB:
 		return new FMallocTBB();
 #endif
-#if PLATFORM_SUPPORTS_MIMALLOC && MIMALLOC_ALLOCATOR_ALLOWED && PLATFORM_BUILDS_MIMALLOC
+#if MIMALLOC_ENABLED
 	case EMemoryAllocatorToUse::Mimalloc:
 		return new FMallocMimalloc();
 #endif
diff --git a/Engine/Source/Runtime/Core/Private/Unix/UnixPlatformMemory.cpp b/Engine/Source/Runtime/Core/Private/Unix/UnixPlatformMemory.cpp
--- a/Engine/Source/Runtime/Core/Private/Unix/UnixPlatformMemory.cpp
+++ b/Engine/Source/Runtime/Core/Private/Unix/UnixPlatformMemory.cpp
@@ -160,7 +160,7 @@ class FMalloc* FUnixPlatformMemory::BaseAllocator()
 	// Mimalloc is now the default allocator for editor and programs because it has shown
 	// both great performance and as much as half the memory usage of TBB after
 	// heavy editor workloads. See CL 15887498 description for benchmarks.
-#if (WITH_EDITORONLY_DATA || IS_PROGRAM) && PLATFORM_SUPPORTS_MIMALLOC && MIMALLOC_ALLOCATOR_ALLOWED
+#if (WITH_EDITORONLY_DATA || IS_PROGRAM) && MIMALLOC_ENABLED
 	AllocatorToUse = EMemoryAllocatorToUse::Mimalloc;
 #endif
 
@@ -192,7 +192,7 @@ class FMalloc* FUnixPlatformMemory::BaseAllocator()
 				break;
 			}
 
-#if PLATFORM_SUPPORTS_MIMALLOC && MIMALLOC_ALLOCATOR_ALLOWED
+#if MIMALLOC_ENABLED
 			if (FCStringAnsi::Stricmp(Arg, "-mimalloc") == 0)
 			{
 				AllocatorToUse = EMemoryAllocatorToUse::Mimalloc;
@@ -324,7 +324,7 @@ class FMalloc* FUnixPlatformMemory::BaseAllocator()
 		break;
 #endif // PLATFORM_SUPPORTS_JEMALLOC
 
-#if PLATFORM_SUPPORTS_MIMALLOC && MIMALLOC_ALLOCATOR_ALLOWED && PLATFORM_BUILDS_MIMALLOC
+#if MIMALLOC_ENABLED
 	case EMemoryAllocatorToUse::Mimalloc:
 		Allocator = new FMallocMimalloc();
 		break;
diff --git a/Engine/Source/Runtime/Core/Private/Windows/WindowsPlatformMemory.cpp b/Engine/Source/Runtime/Core/Private/Windows/WindowsPlatformMemory.cpp
--- a/Engine/Source/Runtime/Core/Private/Windows/WindowsPlatformMemory.cpp
+++ b/Engine/Source/Runtime/Core/Private/Windows/WindowsPlatformMemory.cpp
@@ -129,12 +129,12 @@ FMalloc* FWindowsPlatformMemory::BaseAllocator()
 	// Mimalloc is now the default allocator for editor and programs because it has shown
 	// both great performance and as much as half the memory usage of TBB after
 	// heavy editor workloads. See CL 15887498 description for benchmarks.
-	else if ((WITH_EDITORONLY_DATA || IS_PROGRAM) && MIMALLOC_ALLOCATOR_ALLOWED) //-V517
+	else if ((WITH_EDITORONLY_DATA || IS_PROGRAM) && MIMALLOC_ENABLED) //-V517
 	{
 		AllocatorToUse = EMemoryAllocatorToUse::Mimalloc;
 	}
 #endif
-	else if ((WITH_EDITORONLY_DATA || IS_PROGRAM) && TBB_ALLOCATOR_ALLOWED) //-V517
+	else if ((WITH_EDITORONLY_DATA || IS_PROGRAM) && TBBMALLOC_ENABLED) //-V517
 	{
 		AllocatorToUse = EMemoryAllocatorToUse::TBB;
 	}
@@ -162,13 +162,13 @@ FMalloc* FWindowsPlatformMemory::BaseAllocator()
 		// see FPlatformMisc::GetProcessDiagnostics()
 		AllocatorToUse = EMemoryAllocatorToUse::Ansi;
 	}
-#if TBB_ALLOCATOR_ALLOWED
+#if TBBMALLOC_ENABLED
 	else if (FCString::Stristr(CommandLine, TEXT("-tbbmalloc")))
 	{
 		AllocatorToUse = EMemoryAllocatorToUse::TBB;
 	}
 #endif
-#if MIMALLOC_ALLOCATOR_ALLOWED
+#if MIMALLOC_ENABLED
 	else if (FCString::Stristr(CommandLine, TEXT("-mimalloc")))
 	{
 		AllocatorToUse = EMemoryAllocatorToUse::Mimalloc;
@@ -217,7 +217,7 @@ FMalloc* FWindowsPlatformMemory::BaseAllocator()
 	case EMemoryAllocatorToUse::Stomp:
 		return new FMallocStomp();
 #endif
-#if TBB_ALLOCATOR_ALLOWED
+#if TBBMALLOC_ENABLED
 	case EMemoryAllocatorToUse::TBB:
 		return new FMallocTBB();
 #endif
diff --git a/Engine/Source/Runtime/Core/Public/HAL/MallocTBB.h b/Engine/Source/Runtime/Core/Public/HAL/MallocTBB.h
--- a/Engine/Source/Runtime/Core/Public/HAL/MallocTBB.h
+++ b/Engine/Source/Runtime/Core/Public/HAL/MallocTBB.h
@@ -6,7 +6,9 @@
 #include "HAL/MemoryBase.h"
 #include "HAL/PlatformMemory.h"
 
-#if PLATFORM_SUPPORTS_TBB && TBB_ALLOCATOR_ALLOWED
+#define TBBMALLOC_ENABLED PLATFORM_SUPPORTS_TBB && TBB_ALLOCATOR_ALLOWED
+
+#if TBBMALLOC_ENABLED
 
 /**
  * TBB 64-bit scalable memory allocator.
@@ -44,4 +46,4 @@ protected:
 	}
 };
 
-#endif // PLATFORM_SUPPORTS_TBB && TBB_ALLOCATOR_ALLOWED
\ No newline at end of file
+#endif // TBBMALLOC_ENABLED
\ No newline at end of file
