Let WindowsPlatformProcess::LoadLibraryWithSearchPaths query FileExists directly on the lowest level platform file system to e.g. avoid hitting Pak locks when loading modules or doing stack walking.

#rnx

diff --git a/Engine/Source/Runtime/Core/Private/Windows/WindowsPlatformProcess.cpp b/Engine/Source/Runtime/Core/Private/Windows/WindowsPlatformProcess.cpp
--- a/Engine/Source/Runtime/Core/Private/Windows/WindowsPlatformProcess.cpp
+++ b/Engine/Source/Runtime/Core/Private/Windows/WindowsPlatformProcess.cpp
@@ -6,7 +6,7 @@
 #include "Containers/StringConv.h"
 #include "Containers/UnrealString.h"
 #include "CoreGlobals.h"
-#include "HAL/FileManager.h"
+#include "GenericPlatform/GenericPlatformFile.h"
 #include "HAL/PlatformAffinity.h"
 #include "HAL/PlatformMisc.h"
 #include "HAL/UnrealMemory.h"
@@ -107,7 +107,7 @@ void FWindowsPlatformProcess::AddDllDirectory(const TCHAR* Directory)
 		// enumerate the dir and cache all the dlls
 		{
 			TArray<FString> FoundDllFileNames;
-			IFileManager::Get().FindFiles(FoundDllFileNames, *NormalizedDirectory, TEXT("*.dll"));
+			IPlatformFile::GetPlatformPhysical().FindFiles(FoundDllFileNames, *NormalizedDirectory, TEXT("*.dll"));
 			for (const FString& DllFileName : FoundDllFileNames)
 			{
 				TArray<FString>& Paths = SearchPathDllCache.FindOrAdd(*DllFileName);
@@ -1379,7 +1379,7 @@ bool FWindowsPlatformProcess::LaunchFileInDefaultExternalApplication( const TCHA
 
 void FWindowsPlatformProcess::ExploreFolder( const TCHAR* FilePath )
 {
-	if (IFileManager::Get().DirectoryExists( FilePath ))
+	if (IPlatformFile::GetPlatformPhysical().DirectoryExists( FilePath ))
 	{
 		// Explore the folder
 		::ShellExecuteW( NULL, TEXT("explore"), FilePath, NULL, NULL, SW_SHOWNORMAL );
@@ -1917,14 +1917,14 @@ static bool ReadLibraryImports(const TCHAR* FileName, TArray<FString>& ImportNam
 
 bool FWindowsPlatformProcess::ResolveImport(const FString& Name, const TArray<FString>& SearchPaths, FString& OutFileName)
 {
-	auto SearchPathsFunc = [&OutFileName, &SearchPaths, &Name](int StartIdx, int EndIdx)
+	IPlatformFile& PlatformFile = IPlatformFile::GetPlatformPhysical();
+	auto SearchPathsFunc = [&PlatformFile, &OutFileName, &SearchPaths, &Name](int StartIdx, int EndIdx)
 	{
-		IFileManager& FileManager = IFileManager::Get();
 		for (int Idx = StartIdx; Idx < EndIdx; Idx++)
 		{
 			TStringBuilder<MAX_PATH> FileName;
 			FPathViews::Append(FileName, SearchPaths[Idx], Name);
-			if (FileManager.FileExists(*FileName))
+			if (PlatformFile.FileExists(*FileName))
 			{
 				OutFileName = FPaths::ConvertRelativePathToFull(*FileName);
 				return true;
@@ -1952,7 +1952,7 @@ bool FWindowsPlatformProcess::ResolveImport(const FString& Name, const TArray<FS
 			{
 				const FString& FoundPath = *Itr;
 				// Double check the dll still exists
-				if (FPaths::FileExists(FoundPath))
+				if (PlatformFile.FileExists(*FoundPath))
 				{
 					OutFileName = FoundPath;
 					return true;
@@ -2068,7 +2068,8 @@ void *FWindowsPlatformProcess::LoadLibraryWithSearchPaths(const FString& FileNam
 
 	// Make sure the initial module exists. If we can't find it from the path we're given, it's probably a system dll.
 	FString FullFileName = FileName;
-	if (FPaths::FileExists(FullFileName))
+	IPlatformFile& PlatformFile = IPlatformFile::GetPlatformPhysical();
+	if (PlatformFile.FileExists(*FullFileName))
 	{
 		// Convert it to a full path, since LoadLibrary will try to resolve it against the executable directory (which may not be the same as the working dir)
 		FullFileName = FPaths::ConvertRelativePathToFull(FullFileName);
@@ -2117,7 +2118,7 @@ void *FWindowsPlatformProcess::LoadLibraryWithSearchPaths(const FString& FileNam
 	}
 
 	// Try to load the actual library
-	void* Handle = [FullFileName]() 
+	void* Handle = [&FullFileName]() 
 	{
 #if CPUPROFILERTRACE_ENABLED
 		UE_TRACE_LOG_SCOPED_T(Cpu, Windows_LoadLibrary, CpuChannel)
@@ -2133,7 +2134,7 @@ void *FWindowsPlatformProcess::LoadLibraryWithSearchPaths(const FString& FileNam
 	else
 	{
 		UE_LOG(LogWindows, Log, TEXT("Failed to load '%s' (GetLastError=%d)"), *FileName, ::GetLastError());
-		if(IFileManager::Get().FileExists(*FileName))
+		if (PlatformFile.FileExists(*FileName))
 		{
 			LogImportDiagnostics(FileName, SearchPaths);
 		}
