Removed FSetElementId usage from iterators.

diff --git a/Engine/Source/Runtime/Core/Public/Containers/Set.h b/Engine/Source/Runtime/Core/Public/Containers/Set.h
--- a/Engine/Source/Runtime/Core/Public/Containers/Set.h
+++ b/Engine/Source/Runtime/Core/Public/Containers/Set.h
@@ -1397,28 +1397,6 @@ private:
 		return ((FSetElementId*)Hash.GetAllocation())[HashIndex & (HashSize - 1)];
 	}
 
-	/**
-	 * Accesses an element in the set.
-	 * This is needed because the iterator classes aren't friends of FSetElementId and so can't access the element index.
-	 */
-	FORCEINLINE const SetElementType& GetInternalElement(FSetElementId Id) const
-	{
-		return Elements[Id.Index];
-	}
-	FORCEINLINE SetElementType& GetInternalElement(FSetElementId Id)
-	{
-		return Elements[Id.Index];
-	}
-
-	/**
-	 * Translates an element index into an element ID.
-	 * This is needed because the iterator classes aren't friends of FSetElementId and so can't access the FSetElementId private constructor.
-	 */
-	static FORCEINLINE FSetElementId IndexToId(int32 Index)
-	{
-		return FSetElementId(Index);
-	}
-
 	/** Links an added element to the hash chain. */
 	FORCEINLINE void LinkElement(SizeType ElementIndex, const SetElementType& Element, uint32 KeyHash) const
 	{
@@ -1562,7 +1540,7 @@ private:
 		// Accessors.
 		FORCEINLINE FSetElementId GetId() const
 		{
-			return TSet::IndexToId(ElementIt.GetIndex());
+			return FSetElementId(ElementIt.GetIndex());
 		}
 		FORCEINLINE ItElementType* operator->() const
 		{
@@ -1598,34 +1576,39 @@ private:
 
 		/** Initialization constructor. */
 		FORCEINLINE TBaseKeyIterator(SetType& InSet, KeyArgumentType InKey)
-		:	Set(InSet)
-		,	Key(InKey) //-V1041
+			: Set  (InSet)
+			, Key  (InKey) //-V1041
+			, Index(INDEX_NONE)
 		{
 			// The set's hash needs to be initialized to find the elements with the specified key.
 			Set.ConditionalRehash(Set.Elements.Num());
-			if(Set.HashSize)
+			if (Set.HashSize)
 			{
-				NextId = Set.GetTypedHash(KeyFuncs::GetKeyHash(Key));
+				NextIndex = Set.GetTypedHash(KeyFuncs::GetKeyHash(Key)).Index;
 				++(*this);
 			}
+			else
+			{
+				NextIndex = INDEX_NONE;
+			}
 		}
 
 		/** Advances the iterator to the next element. */
 		FORCEINLINE TBaseKeyIterator& operator++()
 		{
-			Id = NextId;
+			Index = NextIndex;
 
-			while(Id.IsValidId())
+			while (Index != INDEX_NONE)
 			{
-				NextId = Set.GetInternalElement(Id).HashNextId;
-				checkSlow(Id != NextId);
+				NextIndex = Set.Elements[Index].HashNextId.Index;
+				checkSlow(Index != NextIndex);
 
-				if(KeyFuncs::Matches(KeyFuncs::GetSetKey(Set[Id]),Key))
+				if (KeyFuncs::Matches(KeyFuncs::GetSetKey(Set.Elements[Index].Value),Key))
 				{
 					break;
 				}
 
-				Id = NextId;
+				Index = NextIndex;
 			}
 			return *this;
 		}
@@ -1633,7 +1616,7 @@ private:
 		/** conversion to "bool" returning true if the iterator is valid. */
 		FORCEINLINE explicit operator bool() const
 		{ 
-			return Id.IsValidId(); 
+			return Index != INDEX_NONE;
 		}
 		/** inverse of the "bool" operator */
 		FORCEINLINE bool operator !() const 
@@ -1644,18 +1627,18 @@ private:
 		// Accessors.
 		FORCEINLINE ItElementType* operator->() const
 		{
-			return &Set[Id];
+			return &Set.Elements[Index].Value;
 		}
 		FORCEINLINE ItElementType& operator*() const
 		{
-			return Set[Id];
+			return Set.Elements[Index].Value;
 		}
 
 	protected:
 		SetType& Set;
 		ReferenceOrValueType Key;
-		FSetElementId Id;
-		FSetElementId NextId;
+		SizeType Index;
+		SizeType NextIndex;
 	};
 
 public:
@@ -1687,7 +1670,7 @@ public:
 		/** Removes the current element from the set. */
 		FORCEINLINE void RemoveCurrent()
 		{
-			Set.RemoveByIndex(TBaseIterator<false>::GetId().Index);
+			Set.RemoveByIndex(TBaseIterator<false>::ElementIt.GetIndex());
 		}
 
 	private:
@@ -1729,8 +1712,8 @@ public:
 		/** Removes the current element from the set. */
 		FORCEINLINE void RemoveCurrent()
 		{
-			this->Set.RemoveByIndex(TBaseKeyIterator<false>::Id.Index);
-			TBaseKeyIterator<false>::Id = FSetElementId();
+			this->Set.RemoveByIndex(TBaseKeyIterator<false>::Index);
+			TBaseKeyIterator<false>::Index = INDEX_NONE;
 		}
 	};
 
