Providing a way to relax package version checks by omitting changelist checks.

diff --git a/Engine/Source/Runtime/Core/Private/Misc/EngineVersion.cpp b/Engine/Source/Runtime/Core/Private/Misc/EngineVersion.cpp
--- a/Engine/Source/Runtime/Core/Private/Misc/EngineVersion.cpp
+++ b/Engine/Source/Runtime/Core/Private/Misc/EngineVersion.cpp
@@ -15,6 +15,13 @@
 #include "Serialization/StructuredArchiveSlots.h"
 #include "UObject/ReleaseObjectVersion.h"
 
+int32 GSkipChangelistCompatibilityVersionCheck = 0;
+static FAutoConsoleVariableRef CSkipChangelistCompatibilityVersionCheck(
+	TEXT("s.SkipChangelistCompatibilityVersionCheck"),
+	GSkipChangelistCompatibilityVersionCheck,
+	TEXT("If true, engine version compatibility checks will not use changelist numbers to determine compatibility. NOTE: Changelist compatibility checks are automatically skipped in licensee builds."),
+	ECVF_Default);
+
 FEngineVersionBase::FEngineVersionBase(uint16 InMajor, uint16 InMinor, uint16 InPatch, uint32 InChangelist)
 : Major(InMajor)
 , Minor(InMinor)
@@ -121,7 +128,11 @@ bool FEngineVersion::IsCompatibleWith(const FEngineVersionBase &Other) const
 	}
 	else
 	{
-		return FEngineVersion::GetNewest(*this, Other, nullptr) != EVersionComparison::Second;
+		EVersionComponent ConflictingComponent = EVersionComponent::Minor;
+		EVersionComparison CompareResult = FEngineVersion::GetNewest(*this, Other, &ConflictingComponent);
+		
+		const bool bIsCompatible = (CompareResult != EVersionComparison::Second) || (GSkipChangelistCompatibilityVersionCheck && ConflictingComponent == EVersionComponent::Changelist);
+		return bIsCompatible;
 	}
 }
 
