Logging: Added SerializeForLog to allow overriding serialization of structured log fields to compact binary

#rnx

diff --git a/Engine/Source/Runtime/Core/Public/Logging/StructuredLog.h b/Engine/Source/Runtime/Core/Public/Logging/StructuredLog.h
--- a/Engine/Source/Runtime/Core/Public/Logging/StructuredLog.h
+++ b/Engine/Source/Runtime/Core/Public/Logging/StructuredLog.h
@@ -165,6 +165,12 @@ private:
 	const FLogTemplate* Template = nullptr;
 };
 
+template <typename ValueType>
+inline void SerializeForLog(FCbWriter& Writer, ValueType&& Value)
+{
+	Writer << Forward<ValueType>(Value);
+}
+
 } // UE
 
 namespace UE::Logging::Private
@@ -190,16 +196,16 @@ struct FStaticLogRecord
 
 struct FLogField
 {
-	using FWriteFn = FCbWriter& (FCbWriter& Writer, const void* Value);
+	using FWriteFn = void (FCbWriter& Writer, const void* Value);
 
 	const ANSICHAR* Name;
 	const void* Value;
 	FWriteFn* WriteValue;
 
 	template <typename ValueType>
-	static FCbWriter& Write(FCbWriter& Writer, const void* Value)
+	static void Write(FCbWriter& Writer, const void* Value)
 	{
-		return Writer << *(const ValueType*)Value;
+		SerializeForLog(Writer, *(const ValueType*)Value);
 	}
 };
 
