Avoid unnecessary copies, use of auto
#rnx

diff --git a/Engine/Source/Runtime/Core/Private/Misc/ConfigCacheIni.cpp b/Engine/Source/Runtime/Core/Private/Misc/ConfigCacheIni.cpp
--- a/Engine/Source/Runtime/Core/Private/Misc/ConfigCacheIni.cpp
+++ b/Engine/Source/Runtime/Core/Private/Misc/ConfigCacheIni.cpp
@@ -4377,13 +4377,14 @@ void FConfigCacheIni::AsyncInitializeConfigForPlatforms()
 	FPlatformProcess::ApplicationSettingsDir();
 
 	// pre-create all platforms so that the loop below doesn't reallocate anything in the map
-	for (auto Pair : FDataDrivenPlatformInfoRegistry::GetAllPlatformInfos())
+	const TMap<FName, FDataDrivenPlatformInfo>& AllPlatformInfos = FDataDrivenPlatformInfoRegistry::GetAllPlatformInfos();
+	for (const TPair<FName, FDataDrivenPlatformInfo>& Pair : AllPlatformInfos)
 	{
 		GetPlatformConfigFutures().Emplace(Pair.Key);
 		GConfigForPlatform.Add(Pair.Key, new FConfigCacheIni(EConfigCacheType::Temporary));
 	}
 
-	for (auto Pair : FDataDrivenPlatformInfoRegistry::GetAllPlatformInfos())
+	for (const TPair<FName, FDataDrivenPlatformInfo>& Pair : AllPlatformInfos)
 	{
 		FName PlatformName = Pair.Key;
 		GetPlatformConfigFutures()[PlatformName] = Async(EAsyncExecution::ThreadPool, [PlatformName]
