Adding missing include to resolve CIS non-unity error.

diff --git a/Engine/Source/Runtime/Core/Private/Misc/EngineVersion.cpp b/Engine/Source/Runtime/Core/Private/Misc/EngineVersion.cpp
--- a/Engine/Source/Runtime/Core/Private/Misc/EngineVersion.cpp
+++ b/Engine/Source/Runtime/Core/Private/Misc/EngineVersion.cpp
@@ -14,6 +14,7 @@
 #include "Serialization/StructuredArchiveNameHelpers.h"
 #include "Serialization/StructuredArchiveSlots.h"
 #include "UObject/ReleaseObjectVersion.h"
+#include "HAL/IConsoleManager.h" // for FAutoConsoleVariableRef
 
 int32 GSkipChangelistCompatibilityVersionCheck = 0;
 static FAutoConsoleVariableRef CSkipChangelistCompatibilityVersionCheck(
