Added MD5HashToGuid that constructs a Guid from a MD5Hash.
This means that calling ToString on the resulting Guid will not result in the expected MD5 hash string, due to how Guid outputs the string from memory; LexToString should be used in that case.

#rnx

diff --git a/Engine/Source/Runtime/Core/Private/Tests/Misc/GuidTest.cpp b/Engine/Source/Runtime/Core/Private/Tests/Misc/GuidTest.cpp
--- a/Engine/Source/Runtime/Core/Private/Tests/Misc/GuidTest.cpp
+++ b/Engine/Source/Runtime/Core/Private/Tests/Misc/GuidTest.cpp
@@ -4,6 +4,7 @@
 #include "Containers/UnrealString.h"
 #include "Misc/Guid.h"
 #include "Misc/AutomationTest.h"
+#include "Misc/SecureHash.h"
 
 #if WITH_DEV_AUTOMATION_TESTS
 
@@ -95,6 +96,17 @@ bool FGuidTest::RunTest( const FString& Parameters )
 
 	TestFalse(TEXT("Invalidated GUIDs must be invalid"), g4_1.IsValid());
 
+	// GUID creation from MD5 hash
+	FMD5 MD5;
+	ANSICHAR String[] = "Construct Guid from MD5 hash";
+	MD5.Update((uint8*)String, sizeof(String) - 1);
+	FMD5Hash Hash;
+	Hash.Set(MD5);
+	TestTrue(TEXT("Expected hash must match"), LexToString(Hash) == TEXT("e500d989a5db77bc696e79ee785f4a69"));
+
+	FGuid g5_1(MD5HashToGuid(Hash));
+	TestTrue(TEXT("Constructing from MD5 hash must be valid"), g5_1.ToString() == TEXT("89d900e5bc77dba5ee796e69694a5f78"));
+
 	return true;
 }
 
diff --git a/Engine/Source/Runtime/Core/Public/Misc/SecureHash.h b/Engine/Source/Runtime/Core/Public/Misc/SecureHash.h
--- a/Engine/Source/Runtime/Core/Public/Misc/SecureHash.h
+++ b/Engine/Source/Runtime/Core/Public/Misc/SecureHash.h
@@ -14,6 +14,7 @@
 #include "HAL/UnrealMemory.h"
 #include "Misc/AssertionMacros.h"
 #include "Misc/CString.h"
+#include "Misc/Guid.h"
 #include "Serialization/Archive.h"
 #include "Serialization/BufferReader.h"
 #include "Serialization/MemoryLayout.h"
@@ -191,6 +192,17 @@ private:
 	friend CORE_API void LexFromString(FMD5Hash& Hash, const TCHAR*);
 };
 
+/** 
+  * Construct a FGuid from a MD5Hash. This means that calling ToString on the resulting FGuid will not result in the 
+  * expected MD5 hash string, due to how FGuid outputs the string; LexToString should be used in that case.
+  */
+inline FGuid MD5HashToGuid(const FMD5Hash& Hash)
+{
+	FGuid Result;
+	FMemory::Memcpy(&Result, Hash.GetBytes(), sizeof(FGuid));
+	return Result;
+}
+
 /*-----------------------------------------------------------------------------
 	SHA-1 functions.
 -----------------------------------------------------------------------------*/
