Disabling check() test. Assertion messages are output outside the regular log channels and get flagged as legitimate errors on Horde due to non-deterministic order of output between the suspend/resume log parsing messages.

diff --git a/Engine/Source/Runtime/Core/Tests/Misc/AssertionMacrosTest.cpp b/Engine/Source/Runtime/Core/Tests/Misc/AssertionMacrosTest.cpp
--- a/Engine/Source/Runtime/Core/Tests/Misc/AssertionMacrosTest.cpp
+++ b/Engine/Source/Runtime/Core/Tests/Misc/AssertionMacrosTest.cpp
@@ -114,6 +114,9 @@ TEST_CASE("Core::Misc::TestEnsure", "[Core][Misc][AssertionMacros][Ensure]")
 	}
 }
 
+// Disabling check() tests for now. Assertion messages are output outside the regular log flow, and get flagged as legitimate errors on Horde 
+// due to non-deterministic order of output between the suspend/resume log parsing messages.
+#if 0
 // Disables log parsing by UAT for a scope, preventing assertion messages being tagged as errors.
 struct FSuspendLogParsing
 {
@@ -138,5 +141,6 @@ TEST_CASE("Core::Misc::TestCheck", "[Core][Misc][AssertionMacros][Check]")
 	REQUIRE_CHECK_MSG("1 == 2", checkf(1 == 2, TEXT("Error Message")));
 	REQUIRE_CHECK_MSG("Error Message", checkf(1 == 2, TEXT("Error Message")));
 }
+#endif
 
 #endif
