MemQuery integrating into 23.30
MemQuery plugin. LLM tracking of Assets and AssetClasses if enabled.
Memory usage queries to display memory usage for assets and classes with group breakdowns amongst other uses.
Modification to LowLevelMemTracker to support asset and assetclass tracking.

diff --git a/Engine/Plugins/MemoryUsageQueries/MemoryUsageQueries.uplugin b/Engine/Plugins/MemoryUsageQueries/MemoryUsageQueries.uplugin
new file mode 100644
--- /dev/null
+++ b/Engine/Plugins/MemoryUsageQueries/MemoryUsageQueries.uplugin
@@ -0,0 +1,26 @@
+{
+	"FileVersion": 3,
+	"Version": 1,
+	"VersionName": "1.0",
+	"FriendlyName": "MemoryUsageQueries",
+	"Description": "Memory Usage Queries, original contribution from The Coalition (Microsoft) https://thecoalitionstudio.com",
+	"Category": "Engine",
+	"CreatedBy": "Epic Games, Inc.",
+	"CreatedByURL": "http://epicgames.com",
+	"DocsURL": "",
+	"MarketplaceURL": "",
+	"SupportURL": "",
+	"EnabledByDefault": false,
+	"CanContainContent": false,
+	"IsBetaVersion": false,
+	"Installed": false,
+	"Plugins": [
+	],
+	"Modules": [
+		{
+			"Name": "MemoryUsageQueries",
+			"Type": "Runtime",
+			"LoadingPhase": "EarliestPossible"
+		}
+	]
+}
\ No newline at end of file
diff --git a/Engine/Plugins/MemoryUsageQueries/Source/MemoryUsageQueries/MemoryUsageQueries.Build.cs b/Engine/Plugins/MemoryUsageQueries/Source/MemoryUsageQueries/MemoryUsageQueries.Build.cs
new file mode 100644
--- /dev/null
+++ b/Engine/Plugins/MemoryUsageQueries/Source/MemoryUsageQueries/MemoryUsageQueries.Build.cs
@@ -0,0 +1,36 @@
+// Copyright Epic Games, Inc. All Rights Reserved.
+
+namespace UnrealBuildTool.Rules
+{
+	public class MemoryUsageQueries : ModuleRules
+	{
+		public MemoryUsageQueries(ReadOnlyTargetRules Target) : base(Target)
+		{
+			PCHUsage = PCHUsageMode.UseExplicitOrSharedPCHs;
+			bAllowConfidentialPlatformDefines = true;
+
+			PublicIncludePaths.AddRange(
+				new string[] {
+					"$(ModuleDir)/Public",
+				}
+			);
+
+			PublicDependencyModuleNames.AddRange(
+				new string[]
+				{
+					"Core",
+					"CoreUObject"
+				}
+			);
+
+			PrivateDependencyModuleNames.AddRange(
+				new string[]
+				{
+					"AssetRegistry",
+					"Engine",
+					"EngineSettings"
+				}
+			);
+		}
+	}
+}
\ No newline at end of file
diff --git a/Engine/Plugins/MemoryUsageQueries/Source/MemoryUsageQueries/Private/MemoryUsageInfoProvider.cpp b/Engine/Plugins/MemoryUsageQueries/Source/MemoryUsageQueries/Private/MemoryUsageInfoProvider.cpp
new file mode 100644
--- /dev/null
+++ b/Engine/Plugins/MemoryUsageQueries/Source/MemoryUsageQueries/Private/MemoryUsageInfoProvider.cpp
@@ -0,0 +1,59 @@
+// Copyright Epic Games, Inc. All Rights Reserved.
+
+#include "MemoryUsageInfoProvider.h"
+#include "HAL/LowLevelMemTracker.h"
+
+uint64 FMemoryUsageInfoProviderLLM::GetAssetMemoryUsage(FName Asset, FOutputDevice* ErrorOutput /* = GLog */) const
+{
+#if ENABLE_LOW_LEVEL_MEM_TRACKER
+	if (FLowLevelMemTracker::IsEnabled())
+	{
+		return FLowLevelMemTracker::Get().GetTagAmountForTracker(ELLMTracker::Default, Asset, ELLMTagSet::Assets, false);
+	}
+#endif
+	ErrorOutput->Logf(TEXT("MemoryUsageInfoProvider Error: LLM is disabled. Please run with -LLM"));
+	return 0U;
+}
+
+uint64 FMemoryUsageInfoProviderLLM::GetAssetsMemoryUsage(const TSet<FName>& Assets, FOutputDevice* ErrorOutput /* = GLog */) const
+{
+#if ENABLE_LOW_LEVEL_MEM_TRACKER
+	if (FLowLevelMemTracker::IsEnabled())
+	{
+		uint64 TotalSize = 0;
+		for (const FName& Asset : Assets)
+		{
+			TotalSize += FLowLevelMemTracker::Get().GetTagAmountForTracker(ELLMTracker::Default, Asset, ELLMTagSet::Assets, false);
+		}
+		return TotalSize;
+	}
+#endif
+	ErrorOutput->Logf(TEXT("MemoryUsageInfoProvider Error: LLM is disabled. Please run with -LLM"));
+	return 0U;
+}
+
+void FMemoryUsageInfoProviderLLM::GetAllAssetsWithSize(TMap<FName, uint64>& OutAssets, FOutputDevice* ErrorOutput /* = GLog */) const
+{
+	OutAssets.Reset();
+#if ENABLE_LOW_LEVEL_MEM_TRACKER
+	if (FLowLevelMemTracker::IsEnabled())
+	{
+		FLowLevelMemTracker::Get().GetTrackedTagsNamesWithAmount(OutAssets, ELLMTracker::Default, ELLMTagSet::Assets);
+		return;
+	}
+#endif
+	ErrorOutput->Logf(TEXT("MemoryUsageInfoProvider Error: LLM is disabled. Please run with -LLM"));
+}
+
+#if ENABLE_LOW_LEVEL_MEM_TRACKER
+void FMemoryUsageInfoProviderLLM::GetFilteredTagsWithSize(TMap<FName, uint64>& OutTags, ELLMTracker Tracker, ELLMTagSet TagSet, TArray<FLLMTagSetAllocationFilter>& Filters, FOutputDevice* ErrorOutput /* = GLog */) const
+{
+	OutTags.Reset();
+	if (FLowLevelMemTracker::IsEnabled())
+	{
+		FLowLevelMemTracker::Get().GetTrackedTagsNamesWithAmountFiltered(OutTags, Tracker, TagSet, Filters);
+		return;
+	}
+	ErrorOutput->Logf(TEXT("MemoryUsageInfoProvider Error: LLM is disabled. Please run with -LLM"));
+}
+#endif
\ No newline at end of file
diff --git a/Engine/Plugins/MemoryUsageQueries/Source/MemoryUsageQueries/Private/MemoryUsageQueries.cpp b/Engine/Plugins/MemoryUsageQueries/Source/MemoryUsageQueries/Private/MemoryUsageQueries.cpp
new file mode 100644
--- /dev/null
+++ b/Engine/Plugins/MemoryUsageQueries/Source/MemoryUsageQueries/Private/MemoryUsageQueries.cpp
@@ -0,0 +1,1131 @@
+// Copyright Epic Games, Inc. All Rights Reserved.
+
+#include "MemoryUsageQueries.h"
+#include "MemoryUsageQueriesConfig.h"
+#include "MemoryUsageQueriesPrivate.h"
+
+#include "Engine/BlueprintGeneratedClass.h"
+#include "Misc/CoreMisc.h"
+#include "UObject/GCObject.h"
+#include "UObject/UObjectAllocator.h"
+#include "UObject/UObjectIterator.h"
+
+namespace MemoryUsageQueries
+{
+
+FMemoryUsageInfoProviderLLM MemoryUsageInfoProviderLLM;
+IMemoryUsageInfoProvider* CurrentMemoryUsageInfoProvider = &MemoryUsageInfoProviderLLM;
+
+} // namespace MemoryUsageQueries
+
+struct FMemoryUsageQueriesExec : public FSelfRegisteringExec
+{
+	FMemoryUsageQueriesExec() {}
+
+	virtual bool Exec(UWorld* Inworld, const TCHAR* Cmd, FOutputDevice& Ar) override;
+};
+
+static FMemoryUsageQueriesExec DebugSettinsSubsystemExecInstance;
+static int32 DefaultResultLimit = 15;
+
+bool FMemoryUsageQueriesExec::Exec(UWorld* InWorld, const TCHAR* Cmd, FOutputDevice& Ar)
+{
+	if (FParse::Command(&Cmd, TEXT("MemQuery")))
+	{
+		const bool bTruncate = !FParse::Param(Cmd, TEXT("notrunc"));
+
+		// Parse some common options
+		FString Name;
+		FParse::Value(Cmd, TEXT("Name="), Name);
+
+		FString Names;
+		FParse::Value(Cmd, TEXT("Names="), Names);
+
+		int32 Limit = -1;
+		FParse::Value(Cmd, TEXT("Limit="), Limit);
+
+		if (FParse::Command(&Cmd, TEXT("Usage")))
+		{
+			if (!Name.IsEmpty())
+			{
+				uint64 ExclusiveSize = 0U;
+				uint64 InclusiveSize = 0U;
+
+				if (MemoryUsageQueries::GetMemoryUsage(MemoryUsageQueries::CurrentMemoryUsageInfoProvider, Name, ExclusiveSize, InclusiveSize, &Ar))
+				{
+					Ar.Logf(TEXT("MemoryUsage: ExclusiveSize: %.2f MiB (%.2f KiB); InclusiveSize: %.2f MiB (%.2f KiB)"),
+							ExclusiveSize / (1024.f * 1024.f),
+							ExclusiveSize / 1024.f,
+							InclusiveSize / (1024.f * 1024.f),
+							InclusiveSize / 1024.f);
+				}
+
+				return true;
+			}
+		}
+		else if (FParse::Command(&Cmd, TEXT("CombinedUsage")))
+		{
+			if (!Names.IsEmpty())
+			{
+				TArray<FString> Packages;
+				Names.ParseIntoArrayWS(Packages);
+				uint64 TotalSize;
+
+				if (MemoryUsageQueries::GetMemoryUsageCombined(MemoryUsageQueries::CurrentMemoryUsageInfoProvider, Packages, TotalSize, &Ar))
+				{
+					Ar.Logf(TEXT("MemoryUsageCombined: TotalSize: %.2f MiB (%.2f KiB)"), TotalSize / (1024.f * 1024.f), TotalSize / 1024.f);
+				}
+
+				return true;
+			}
+		}
+		else if (FParse::Command(&Cmd, TEXT("SharedUsage")))
+		{
+			if (!Names.IsEmpty())
+			{
+				TArray<FString> Packages;
+				Names.ParseIntoArrayWS(Packages);
+				uint64 SharedSize;
+
+				if (MemoryUsageQueries::GetMemoryUsageShared(MemoryUsageQueries::CurrentMemoryUsageInfoProvider, Packages, SharedSize, &Ar))
+				{
+					Ar.Logf(TEXT("MemoryUsageShared: SharedSize: %.2f MiB (%.2f KiB)"), SharedSize / (1024.f * 1024.f), SharedSize / 1024.f);
+				}
+
+				return true;
+			}
+		}
+		else if (FParse::Command(&Cmd, TEXT("UniqueUsage")))
+		{
+			if (!Names.IsEmpty())
+			{
+				TArray<FString> Packages;
+				Names.ParseIntoArrayWS(Packages);
+				uint64 UniqueSize = 0U;
+
+				if (MemoryUsageQueries::GetMemoryUsageUnique(MemoryUsageQueries::CurrentMemoryUsageInfoProvider, Packages, UniqueSize, &Ar))
+				{
+					Ar.Logf(TEXT("MemoryUsageUnique: UniqueSize: %.2f MiB (%.2f KiB)"), UniqueSize / (1024.f * 1024.f), UniqueSize / 1024.f);
+				}
+
+				return true;
+			}
+		}
+		else if (FParse::Command(&Cmd, TEXT("CommonUsage")))
+		{
+			if (!Names.IsEmpty())
+			{
+				TArray<FString> Packages;
+				Names.ParseIntoArrayWS(Packages);
+				uint64 CommonSize = 0U;
+
+				if (MemoryUsageQueries::GetMemoryUsageCommon(MemoryUsageQueries::CurrentMemoryUsageInfoProvider, Packages, CommonSize, &Ar))
+				{
+					Ar.Logf(TEXT("MemoryUsageCommon: CommonSize: %.2f MiB (%.2f KiB)"), CommonSize / (1024.f * 1024.f), CommonSize / 1024.f);
+				}
+
+				return true;
+			}
+		}
+		else if (FParse::Command(&Cmd, TEXT("Dependencies")))
+		{
+			if (!Name.IsEmpty())
+			{
+				TMap<FName, uint64> DependenciesWithSize;
+
+				if (MemoryUsageQueries::GetDependenciesWithSize(MemoryUsageQueries::CurrentMemoryUsageInfoProvider, Name, DependenciesWithSize, &Ar))
+				{
+					Ar.Logf(TEXT("Dependencies (%s):"), *Name);
+					MemoryUsageQueries::Internal::PrintTagsWithSize(Ar, DependenciesWithSize, TEXT("dependencies"), bTruncate, Limit);
+				}
+
+				return true;
+			}
+		}
+		else if (FParse::Command(&Cmd, TEXT("CombinedDependencies")))
+		{
+			if (!Names.IsEmpty())
+			{
+				TArray<FString> Packages;
+				Names.ParseIntoArrayWS(Packages);
+				TMap<FName, uint64> CombinedDependenciesWithSize;
+
+				if (MemoryUsageQueries::GetDependenciesWithSizeCombined(MemoryUsageQueries::CurrentMemoryUsageInfoProvider, Packages, CombinedDependenciesWithSize, &Ar))
+				{
+					Ar.Logf(TEXT("Combined dependencies (%s):"), *Names);
+					MemoryUsageQueries::Internal::PrintTagsWithSize(Ar, CombinedDependenciesWithSize, TEXT("combined dependencies"), bTruncate, Limit);
+				}
+
+				return true;
+			}
+		}
+		else if (FParse::Command(&Cmd, TEXT("SharedDependencies")))
+		{
+			if (!Names.IsEmpty())
+			{
+				TArray<FString> Packages;
+				Names.ParseIntoArrayWS(Packages);
+				TMap<FName, uint64> SharedDependenciesWithSize;
+
+				if (MemoryUsageQueries::GetDependenciesWithSizeShared(MemoryUsageQueries::CurrentMemoryUsageInfoProvider, Packages, SharedDependenciesWithSize, &Ar))
+				{
+					Ar.Logf(TEXT("Shared dependencies (%s):"), *Names);
+					MemoryUsageQueries::Internal::PrintTagsWithSize(Ar, SharedDependenciesWithSize, TEXT("shared dependencies"), bTruncate, Limit);
+				}
+
+				return true;
+			}
+		}
+		else if (FParse::Command(&Cmd, TEXT("UniqueDependencies")))
+		{
+			if (!Names.IsEmpty())
+			{
+				TArray<FString> Packages;
+				Names.ParseIntoArrayWS(Packages);
+				TMap<FName, uint64> UniqueDependenciesWithSize;
+
+				if (MemoryUsageQueries::GetDependenciesWithSizeUnique(MemoryUsageQueries::CurrentMemoryUsageInfoProvider, Packages, UniqueDependenciesWithSize, &Ar))
+				{
+					Ar.Logf(TEXT("Unique dependencies (%s):"), *Names);
+					MemoryUsageQueries::Internal::PrintTagsWithSize(Ar, UniqueDependenciesWithSize, TEXT("unique dependencies"), bTruncate, Limit);
+				}
+
+				return true;
+			}
+		}
+		else if (FParse::Command(&Cmd, TEXT("CommonDependencies")))
+		{
+			if (!Names.IsEmpty())
+			{
+				TArray<FString> Packages;
+				Names.ParseIntoArrayWS(Packages);
+				TMap<FName, uint64> CommonDependenciesWithSize;
+
+				if (MemoryUsageQueries::GetDependenciesWithSizeCommon(MemoryUsageQueries::CurrentMemoryUsageInfoProvider, Packages, CommonDependenciesWithSize, &Ar))
+				{
+					Ar.Logf(TEXT("Common dependencies (%s):"), *Names);
+					MemoryUsageQueries::Internal::PrintTagsWithSize(Ar, CommonDependenciesWithSize, TEXT("common dependencies"), bTruncate, Limit);
+				}
+
+				return true;
+			}
+		}
+#if ENABLE_LOW_LEVEL_MEM_TRACKER
+		else if (FParse::Command(&Cmd, TEXT("ListAssets")))
+		{
+			FString AssetName;
+			FParse::Value(Cmd, TEXT("NAME="), AssetName);
+
+			FName Group = NAME_None;
+			FString GroupName;
+			if (FParse::Value(Cmd, TEXT("GROUP="), GroupName))
+			{
+				Group = FName(*GroupName);
+			}
+
+			FName Class = NAME_None;
+			FString ClassName;
+			if (FParse::Value(Cmd, TEXT("CLASS="), ClassName))
+			{
+				Class = FName(*ClassName);
+			}
+
+			bool bSuccess;
+			TMap<FName, uint64> AssetsWithSize;
+
+			if (Group != NAME_None || Class != NAME_None)
+			{
+				bSuccess = MemoryUsageQueries::GetFilteredPackagesWithSize(AssetsWithSize, Group, AssetName, Class, &Ar);
+			}
+			else
+			{
+				// TODO - Implement using faster path if there are no group / class filters
+				bSuccess = MemoryUsageQueries::GetFilteredPackagesWithSize(AssetsWithSize, Group, AssetName, Class, &Ar);
+			}
+
+			if (bSuccess)
+			{
+				Ar.Logf(TEXT("%d largest assets:"), (Limit != -1 ? FMath::Min(Limit, AssetsWithSize.Num()) : AssetsWithSize.Num()));
+				MemoryUsageQueries::Internal::PrintTagsWithSize(Ar, AssetsWithSize, TEXT("largest assets"), bTruncate, Limit);
+			}
+
+			return true;
+		}
+		else if (FParse::Command(&Cmd, TEXT("ListClasses")))
+		{
+			FName Group = NAME_None;
+			FString GroupName;
+			if (FParse::Value(Cmd, TEXT("GROUP="), GroupName))
+			{
+				Group = FName(*GroupName);
+			}
+
+			FString AssetName;
+			FParse::Value(Cmd, TEXT("ASSET="), AssetName);
+			
+			TMap<FName, uint64> ClassesWithSize;
+
+			if (MemoryUsageQueries::GetFilteredClassesWithSize(ClassesWithSize, Group, AssetName, &Ar))
+			{
+				Ar.Logf(TEXT("%d largest classes"), (Limit != -1 ? FMath::Min(Limit, ClassesWithSize.Num()) : ClassesWithSize.Num()));
+				MemoryUsageQueries::Internal::PrintTagsWithSize(Ar, ClassesWithSize, *FString::Printf(TEXT("largest classes")), bTruncate, Limit);
+			}
+
+			return true;
+		}
+		else if (FParse::Command(&Cmd, TEXT("ListGroups")))
+		{
+			FString AssetName = "";
+			FParse::Value(Cmd, TEXT("ASSET="), AssetName);
+
+			FName Class = NAME_None;
+			FString ClassName;
+			if (FParse::Value(Cmd, TEXT("CLASS="), ClassName))
+			{
+				Class = FName(*ClassName);
+			}
+
+			TMap<FName, uint64> GroupsWithSize;
+
+			if (MemoryUsageQueries::GetFilteredGroupsWithSize(GroupsWithSize, AssetName, Class, &Ar))
+			{
+				Ar.Logf(TEXT("%d largest groups"), (Limit != -1 ? FMath::Min(Limit, GroupsWithSize.Num()) : GroupsWithSize.Num()));
+				MemoryUsageQueries::Internal::PrintTagsWithSize(Ar, GroupsWithSize, *FString::Printf(TEXT("largest groups")), bTruncate, Limit);
+			}
+
+			return true;
+		}
+#endif
+		else if (FParse::Command(&Cmd, TEXT("Savings")))
+		{
+			const UMemoryUsageQueriesConfig* Config = GetDefault<UMemoryUsageQueriesConfig>();
+
+			for (auto It = Config->SavingsPresets.CreateConstIterator(); It; ++It)
+			{
+				if (!FParse::Command(&Cmd, *It.Key()))
+				{
+					continue;
+				}
+
+				TMap<FName, uint64> PresetSavings;
+				TSet<FName> Packages;
+
+				UClass* SavingsClass = FindObject<UClass>(nullptr, *It.Value());
+				if (SavingsClass != nullptr)
+				{
+					TArray<UClass*> DerivedClasses;
+					TArray<UClass*> DerivedResults;
+
+					GetDerivedClasses(SavingsClass, DerivedClasses, true);
+
+					for (UClass* DerivedClass : DerivedClasses)
+					{
+						UBlueprintGeneratedClass* BPClass = Cast<UBlueprintGeneratedClass>(DerivedClass);
+						if (BPClass != nullptr)
+						{
+							DerivedResults.Reset();
+							GetDerivedClasses(BPClass, DerivedResults, false);
+							if (DerivedResults.Num() == 0)
+							{
+								Packages.Add(DerivedClass->GetPackage()->GetFName());
+							}
+						}
+					}
+				}
+
+				for (const auto& Package : Packages)
+				{
+					uint64 Size = 0;
+					MemoryUsageQueries::GetMemoryUsageUnique(MemoryUsageQueries::CurrentMemoryUsageInfoProvider, TArray<FString>({Package.ToString()}), Size, &Ar);
+					PresetSavings.Add(Package, Size);
+				}
+
+				PresetSavings.ValueSort(TGreater<uint64>());
+				MemoryUsageQueries::Internal::PrintTagsWithSize(Ar, PresetSavings, *FString::Printf(TEXT("possible savings")), bTruncate);
+
+				return true;
+			}
+		}
+	}
+
+	return false;
+}
+
+namespace MemoryUsageQueries
+{
+
+void RegisterConsoleAutoCompleteEntries(TArray<FAutoCompleteCommand>& AutoCompleteList)
+{
+	const UConsoleSettings* ConsoleSettings = GetDefault<UConsoleSettings>();
+
+	{
+		FAutoCompleteCommand AutoCompleteCommand;
+		AutoCompleteCommand.Command = TEXT("MemQuery Usage");
+		AutoCompleteCommand.Desc = TEXT("Name=<AssetName> Prints memory usage of the specified asset.");
+		AutoCompleteCommand.Color = ConsoleSettings->AutoCompleteCommandColor;
+		AutoCompleteList.Add(MoveTemp(AutoCompleteCommand));
+	}
+
+	{
+		FAutoCompleteCommand AutoCompleteCommand;
+		AutoCompleteCommand.Command = TEXT("MemQuery CombinedUsage");
+		AutoCompleteCommand.Desc = TEXT("Names=\"<AssetName1> <AssetName2> ...\" Prints combined memory usage of the specified assets (including all dependencies).");
+		AutoCompleteCommand.Color = ConsoleSettings->AutoCompleteCommandColor;
+		AutoCompleteList.Add(MoveTemp(AutoCompleteCommand));
+	}
+
+	{
+		FAutoCompleteCommand AutoCompleteCommand;
+		AutoCompleteCommand.Command = TEXT("MemQuery SharedUsage");
+		AutoCompleteCommand.Desc = TEXT("Names=\"<AssetName1> <AssetName2> ...\" Prints shared memory usage of the specified assets (including only dependencies shared by the specified assets).");
+		AutoCompleteCommand.Color = ConsoleSettings->AutoCompleteCommandColor;
+		AutoCompleteList.Add(MoveTemp(AutoCompleteCommand));
+	}
+
+	{
+		FAutoCompleteCommand AutoCompleteCommand;
+		AutoCompleteCommand.Command = TEXT("MemQuery UniqueUsage");
+		AutoCompleteCommand.Desc = TEXT("Names=\"<AssetName1> <AssetName2> ...\" Prints unique memory usage of the specified assets (including only dependencies unique to the specified assets).");
+		AutoCompleteCommand.Color = ConsoleSettings->AutoCompleteCommandColor;
+		AutoCompleteList.Add(MoveTemp(AutoCompleteCommand));
+	}
+
+	{
+		FAutoCompleteCommand AutoCompleteCommand;
+		AutoCompleteCommand.Command = TEXT("MemQuery CommonUsage");
+		AutoCompleteCommand.Desc = TEXT("Names=\"<AssetName1> <AssetName2> ...\" Prints common memory usage of the specified assets (including only dependencies that are not unique to the specified assets).");
+		AutoCompleteCommand.Color = ConsoleSettings->AutoCompleteCommandColor;
+		AutoCompleteList.Add(MoveTemp(AutoCompleteCommand));
+	}
+
+	{
+		FAutoCompleteCommand AutoCompleteCommand;
+		AutoCompleteCommand.Command = TEXT("MemQuery Dependencies");
+		AutoCompleteCommand.Desc = TEXT("Name=<AssetName> Limit=<n> Lists dependencies of the specified asset, sorted by size.");
+		AutoCompleteCommand.Color = ConsoleSettings->AutoCompleteCommandColor;
+		AutoCompleteList.Add(MoveTemp(AutoCompleteCommand));
+	}
+
+	{
+		FAutoCompleteCommand AutoCompleteCommand;
+		AutoCompleteCommand.Command = TEXT("MemQuery CombinedDependencies");
+		AutoCompleteCommand.Desc = TEXT("Names=\"<AssetName1> <AssetName2> ...\" Limit=<n> Lists n largest dependencies of the specified assets, sorted by size.");
+		AutoCompleteCommand.Color = ConsoleSettings->AutoCompleteCommandColor;
+		AutoCompleteList.Add(MoveTemp(AutoCompleteCommand));
+	}
+
+	{
+		FAutoCompleteCommand AutoCompleteCommand;
+		AutoCompleteCommand.Command = TEXT("MemQuery SharedDependencies");
+		AutoCompleteCommand.Desc = TEXT("Names=\"<AssetName1> <AssetName2> ...\" Limit=<n> Lists n largest dependencies that are shared by the specified assets, sorted by size.");
+		AutoCompleteCommand.Color = ConsoleSettings->AutoCompleteCommandColor;
+		AutoCompleteList.Add(MoveTemp(AutoCompleteCommand));
+	}
+		
+	{
+		FAutoCompleteCommand AutoCompleteCommand;
+		AutoCompleteCommand.Command = TEXT("MemQuery UniqueDependencies");
+		AutoCompleteCommand.Desc = TEXT("Names=\"<AssetName1> <AssetName2> ...\" Limit=<n> Lists n largest dependencies that are unique to the specified assets, sorted by size.");
+		AutoCompleteCommand.Color = ConsoleSettings->AutoCompleteCommandColor;
+		AutoCompleteList.Add(MoveTemp(AutoCompleteCommand));
+	}
+
+	{
+		FAutoCompleteCommand AutoCompleteCommand;
+		AutoCompleteCommand.Command = TEXT("MemQuery CommonDependencies");
+		AutoCompleteCommand.Desc = TEXT("Names=\"<AssetName1> <AssetName2> ...\" Limit=<n> Lists n largest dependencies that are NOT unique to the specified assets, sorted by size.");
+		AutoCompleteCommand.Color = ConsoleSettings->AutoCompleteCommandColor;
+		AutoCompleteList.Add(MoveTemp(AutoCompleteCommand));
+	}
+
+	{
+		FAutoCompleteCommand AutoCompleteCommand;
+		AutoCompleteCommand.Command = TEXT("MemQuery ListAssets");
+		AutoCompleteCommand.Desc = TEXT("Name=<AssetNameSubstring> Group=<GroupName> Class=<ClassName> Limit=<n> Lists n largest assets.");
+		AutoCompleteCommand.Color = ConsoleSettings->AutoCompleteCommandColor;
+		AutoCompleteList.Add(MoveTemp(AutoCompleteCommand));
+	}
+
+	{
+		FAutoCompleteCommand AutoCompleteCommand;
+		AutoCompleteCommand.Command = TEXT("MemQuery ListClasses");
+		AutoCompleteCommand.Desc = TEXT("Group=<GroupName> Asset=<AssetName> Limit=<n> Lists n largest classes.");
+		AutoCompleteCommand.Color = ConsoleSettings->AutoCompleteCommandColor;
+		AutoCompleteList.Add(MoveTemp(AutoCompleteCommand));
+	}
+
+	{
+		FAutoCompleteCommand AutoCompleteCommand;
+		AutoCompleteCommand.Command = TEXT("MemQuery ListGroups");
+		AutoCompleteCommand.Desc = TEXT("Asset=<AssetName> Class=<ClassName> Limit=<n> Lists n largest groups.");
+		AutoCompleteCommand.Color = ConsoleSettings->AutoCompleteCommandColor;
+		AutoCompleteList.Add(MoveTemp(AutoCompleteCommand));
+	}
+
+	const UMemoryUsageQueriesConfig* MemQueryConfig = GetDefault<UMemoryUsageQueriesConfig>();
+
+	for (auto It = MemQueryConfig->SavingsPresets.CreateConstIterator(); It; ++It)
+	{
+		FAutoCompleteCommand AutoCompleteCommand;
+		AutoCompleteCommand.Command = FString::Printf(TEXT("MemQuery Savings %s"), *It.Key());
+		AutoCompleteCommand.Desc = FString::Printf(TEXT("Limit=<n> Lists potential savings among %s. How much memory can be saved it we delete certain object?"), *It.Key());
+		AutoCompleteCommand.Color = ConsoleSettings->AutoCompleteCommandColor;
+		AutoCompleteList.Add(MoveTemp(AutoCompleteCommand));
+	}
+}
+
+bool GetMemoryUsage(const IMemoryUsageInfoProvider* MemoryUsageInfoProvider, const FString& PackageName, uint64& OutExclusiveSize, uint64& OutInclusiveSize, FOutputDevice* ErrorOutput /* = GLog */)
+{
+	FName LongPackageName;
+	TSet<FName> Dependencies;
+	if (!Internal::GetLongNameAndDependencies(PackageName, LongPackageName, Dependencies, ErrorOutput))
+	{
+		return false;
+	}
+
+	OutExclusiveSize = MemoryUsageInfoProvider->GetAssetMemoryUsage(LongPackageName, ErrorOutput);
+	OutInclusiveSize = MemoryUsageInfoProvider->GetAssetsMemoryUsage(Dependencies, ErrorOutput);
+
+	return true;
+}
+
+bool GetMemoryUsageCombined(const IMemoryUsageInfoProvider* MemoryUsageInfoProvider, const TArray<FString>& PackageNames, uint64& OutTotalSize, FOutputDevice* ErrorOutput /* = GLog */)
+{
+	TSet<FName> Dependencies;
+	if (!Internal::GetDependenciesCombined(PackageNames, Dependencies, ErrorOutput))
+	{
+		return false;
+	}
+
+	OutTotalSize =  MemoryUsageInfoProvider->GetAssetsMemoryUsage(Dependencies, ErrorOutput);
+
+	return true;
+}
+
+bool GetMemoryUsageShared(const IMemoryUsageInfoProvider* MemoryUsageInfoProvider, const TArray<FString>& PackageNames, uint64& OutTotalSize, FOutputDevice* ErrorOutput /* = GLog */)
+{
+	TSet<FName> Dependencies;
+	if (!Internal::GetDependenciesShared(PackageNames, Dependencies, ErrorOutput))
+	{
+		return false;
+	}
+
+	OutTotalSize = MemoryUsageInfoProvider->GetAssetsMemoryUsage(Dependencies, ErrorOutput);
+
+	return true;
+}
+
+bool GetMemoryUsageUnique(const IMemoryUsageInfoProvider* MemoryUsageInfoProvider, const TArray<FString>& PackageNames, uint64& OutUniqueSize, FOutputDevice* ErrorOutput /* = GLog */)
+{
+	TSet<FName> RemovablePackages;
+	if (!Internal::GetRemovablePackages(PackageNames, RemovablePackages, ErrorOutput))
+	{
+		return false;
+	}
+
+	OutUniqueSize = MemoryUsageInfoProvider->GetAssetsMemoryUsage(RemovablePackages, ErrorOutput);
+
+	return true;
+}
+
+bool GetMemoryUsageCommon(const IMemoryUsageInfoProvider* MemoryUsageInfoProvider, const TArray<FString>& PackageNames, uint64& OutCommonSize, FOutputDevice* ErrorOutput /* = GLog */)
+{
+	TSet<FName> UnremovablePackages;
+	if (!Internal::GetUnremovablePackages(PackageNames, UnremovablePackages, ErrorOutput))
+	{
+		return false;
+	}
+
+	OutCommonSize = MemoryUsageInfoProvider->GetAssetsMemoryUsage(UnremovablePackages, ErrorOutput);
+
+	return true;
+}
+
+bool GetDependenciesWithSize(const IMemoryUsageInfoProvider* MemoryUsageInfoProvider, const FString& PackageName, TMap<FName, uint64>& OutDependenciesWithSize, FOutputDevice* ErrorOutput /* = GLog */)
+{
+	FName LongPackageName;
+	TSet<FName> Dependencies;
+
+	if (!Internal::GetLongNameAndDependencies(PackageName, LongPackageName, Dependencies, ErrorOutput))
+	{
+		return false;
+	}
+
+	OutDependenciesWithSize.Empty();
+	Internal::SortPackagesBySize(MemoryUsageInfoProvider, Dependencies, OutDependenciesWithSize, ErrorOutput);
+
+	return true;
+}
+
+bool GetDependenciesWithSizeCombined(const IMemoryUsageInfoProvider* MemoryUsageInfoProvider, const TArray<FString>& PackageNames, TMap<FName, uint64>& OutDependenciesWithSize, FOutputDevice* ErrorOutput /* = GLog */)
+{
+	TSet<FName> Dependencies;
+
+	if (!Internal::GetDependenciesCombined(PackageNames, Dependencies, ErrorOutput))
+	{
+		return false;
+	}
+
+	OutDependenciesWithSize.Empty();
+	Internal::SortPackagesBySize(MemoryUsageInfoProvider, Dependencies, OutDependenciesWithSize, ErrorOutput);
+
+	return true;
+}
+
+bool GetDependenciesWithSizeShared(const IMemoryUsageInfoProvider* MemoryUsageInfoProvider, const TArray<FString>& PackageNames, TMap<FName, uint64>& OutDependenciesWithSize, FOutputDevice* ErrorOutput /* = GLog */)
+{
+	TSet<FName> Dependencies;
+	if (!Internal::GetDependenciesShared(PackageNames, Dependencies, ErrorOutput))
+	{
+		return false;
+	}
+
+	OutDependenciesWithSize.Empty();
+	Internal::SortPackagesBySize(MemoryUsageInfoProvider, Dependencies, OutDependenciesWithSize, ErrorOutput);
+
+	return true;
+}
+
+bool GetDependenciesWithSizeUnique(const IMemoryUsageInfoProvider* MemoryUsageInfoProvider, const TArray<FString>& PackageNames, TMap<FName, uint64>& OutDependenciesWithSize, FOutputDevice* ErrorOutput /* = GLog */)
+{
+	TSet<FName> RemovablePackages;
+	if (!Internal::GetRemovablePackages(PackageNames, RemovablePackages, ErrorOutput))
+	{
+		return false;
+	}
+
+	OutDependenciesWithSize.Empty();
+	Internal::SortPackagesBySize(MemoryUsageInfoProvider, RemovablePackages, OutDependenciesWithSize, ErrorOutput);
+
+	return true;
+}
+
+bool GetDependenciesWithSizeCommon(const IMemoryUsageInfoProvider* MemoryUsageInfoProvider, const TArray<FString>& PackageNames, TMap<FName, uint64>& OutDependenciesWithSize, FOutputDevice* ErrorOutput /* = GLog */)
+{
+	TSet<FName> UnremovablePackages;
+	if (!Internal::GetUnremovablePackages(PackageNames, UnremovablePackages, ErrorOutput))
+	{
+		return false;
+	}
+
+	OutDependenciesWithSize.Empty();
+	Internal::SortPackagesBySize(MemoryUsageInfoProvider, UnremovablePackages, OutDependenciesWithSize, ErrorOutput);
+
+	return true;
+}
+
+#if ENABLE_LOW_LEVEL_MEM_TRACKER
+bool GetFilteredPackagesWithSize(TMap<FName, uint64>& OutPackagesWithSize, FName GroupName /* = NAME_None */, FString AssetSubstring /* = FString() */, FName ClassName /* = NAME_None */, FOutputDevice* ErrorOutput /* = GLog */)
+{
+	TArray<FLLMTagSetAllocationFilter> Filters;
+	if (GroupName != NAME_None)
+	{
+		Filters.Add({GroupName, ELLMTagSet::None});
+	}
+
+	if (ClassName != NAME_None)
+	{
+		Filters.Add({ClassName, ELLMTagSet::AssetClasses});
+	}
+
+	MemoryUsageInfoProviderLLM.GetFilteredTagsWithSize(OutPackagesWithSize, ELLMTracker::Default, ELLMTagSet::Assets, Filters, ErrorOutput);
+
+	if (AssetSubstring.Len() > 0)
+	{
+		Internal::RemoveFilteredPackages(OutPackagesWithSize, AssetSubstring);
+	}
+
+	OutPackagesWithSize.ValueSort(TGreater<uint64>());
+
+	return true;
+}
+
+bool GetFilteredClassesWithSize(TMap<FName, uint64>& OutClassesWithSize, FName GroupName /* = NAME_None */, FString AssetName /* = FString() */, FOutputDevice* ErrorOutput /* = GLog */)
+{
+	TArray<FLLMTagSetAllocationFilter> Filters;
+
+	FName LongName = NAME_None;
+	if (!AssetName.IsEmpty() && !Internal::GetLongName(AssetName, LongName, ErrorOutput))
+	{
+		return false;
+	}
+
+	if (LongName != NAME_None)
+	{
+		Filters.Add({LongName, ELLMTagSet::Assets});
+	}
+
+	if (GroupName != NAME_None)
+	{
+		Filters.Add({GroupName, ELLMTagSet::None});
+	}
+
+	MemoryUsageInfoProviderLLM.GetFilteredTagsWithSize(OutClassesWithSize, ELLMTracker::Default, ELLMTagSet::AssetClasses, Filters, ErrorOutput);
+
+	OutClassesWithSize.ValueSort(TGreater<uint64>());
+
+	return true;
+}
+
+bool GetFilteredGroupsWithSize(TMap<FName, uint64>& OutGroupsWithSize, FString AssetName /* = FString() */, FName ClassName /* = NAME_None */, FOutputDevice* ErrorOutput /* = GLog */)
+{
+	TArray<FLLMTagSetAllocationFilter> Filters;
+
+	FName LongName = NAME_None;
+	if (!AssetName.IsEmpty() && !Internal::GetLongName(AssetName, LongName, ErrorOutput))
+	{
+		return false;
+	}
+
+	if (LongName != NAME_None)
+	{
+		Filters.Add({LongName, ELLMTagSet::Assets});
+	}
+
+	if (ClassName != NAME_None)
+	{
+		Filters.Add({ClassName, ELLMTagSet::AssetClasses});
+	}
+
+	MemoryUsageInfoProviderLLM.GetFilteredTagsWithSize(OutGroupsWithSize, ELLMTracker::Default, ELLMTagSet::None, Filters, ErrorOutput);
+
+	OutGroupsWithSize.ValueSort(TGreater<uint64>());
+
+	return true;
+}
+#endif
+
+namespace Internal
+{
+FMemoryUsageReferenceProcessor::FMemoryUsageReferenceProcessor()
+{
+	int32 Num = GUObjectArray.GetObjectArrayNum();
+
+	Excluded.Init(false, Num);
+	ReachableFull.Init(false, Num);
+	ReachableExcluded.Init(false, Num);
+}
+
+bool FMemoryUsageReferenceProcessor::Init(const TArray<FString>& PackageNames, FOutputDevice* ErrorOutput /* = GLog */)
+{
+	for (FRawObjectIterator It(true); It; ++It)
+	{
+		FUObjectItem* ObjectItem = *It;
+		UObject* Object = (UObject*)ObjectItem->Object;
+
+		if (ObjectItem->IsUnreachable())
+		{
+			continue;
+		}
+
+		if (ObjectItem->IsRootSet())
+		{
+			RootSetPackages.Add(Object);
+		}
+
+		if (UClass* Class = dynamic_cast<UClass*>(Object))
+		{
+			if (!Class->HasAnyClassFlags(CLASS_TokenStreamAssembled))
+			{
+				Class->AssembleReferenceTokenStream();
+				check(Class->HasAnyClassFlags(CLASS_TokenStreamAssembled));
+			}
+		}
+	}
+
+	if (FPlatformProperties::RequiresCookedData() && FGCObject::GGCObjectReferencer && GUObjectArray.IsDisregardForGC(FGCObject::GGCObjectReferencer))
+	{
+		RootSetPackages.Add(FGCObject::GGCObjectReferencer);
+	}
+
+	TSet<FName> LongPackageNames;
+	if (!GetLongNames(PackageNames, LongPackageNames, ErrorOutput))
+	{
+		return false;
+	}
+
+	for (const auto& PackageName : LongPackageNames)
+	{
+		UPackage* Package = Cast<UPackage>(StaticFindObjectFast(UPackage::StaticClass(), nullptr, PackageName, /*bExactClass*/ true));
+		if (Package)
+		{
+			auto ExcludeObject = [&](UObject* ObjectToExclude) {
+				if (ObjectToExclude && GUObjectArray.ObjectToIndex(ObjectToExclude) < Excluded.Num())
+				{
+					Excluded[GUObjectArray.ObjectToIndex(ObjectToExclude)] = true;
+				}
+			};
+
+			auto ExcludeObjectOfClass = [&](UObject* ObjectOfClass) {
+				if (ObjectOfClass)
+				{
+					ForEachObjectWithOuter(ObjectOfClass, ExcludeObject);
+				}
+				ExcludeObject(ObjectOfClass);
+			};
+
+			auto ExcludeObjectInPackage = [&](UObject* ObjectInPackage) {
+				if (ObjectInPackage->IsA<UClass>())
+				{
+					UClass* Class = static_cast<UClass*>(ObjectInPackage);
+					ForEachObjectOfClass(Class, ExcludeObjectOfClass);
+				}
+				ExcludeObject(ObjectInPackage);
+			};
+
+			ForEachObjectWithOuter(Package, ExcludeObjectInPackage);
+			ExcludeObject(Package);
+		}
+	}
+
+	return true;
+}
+
+TArray<UObject*>& FMemoryUsageReferenceProcessor::GetRootSet()
+{
+	return RootSetPackages;
+}
+
+void FMemoryUsageReferenceProcessor::HandleTokenStreamObjectReference(UE::GC::FWorkerContext& Context, const UObject* ReferencingObject, UObject*& Object, UE::GC::FTokenId TokenIndex, const EGCTokenType TokenType, bool bAllowReferenceElimination)
+{
+	if (Object == nullptr || GUObjectArray.ObjectToIndex(Object) >= ReachableFull.Num() || GUObjectAllocator.ResidesInPermanentPool(Object) || GUObjectArray.IsDisregardForGC(Object))
+	{
+		return;
+	}
+
+	int32 ObjectIndex = GUObjectArray.ObjectToIndex(Object);
+
+	if (Mode == EMode::Full)
+	{
+		if (!ReachableFull[ObjectIndex])
+		{
+			ReachableFull[ObjectIndex] = true;
+			Context.ObjectsToSerialize.Add<Options>(Object);
+		}
+	}
+	else if (Mode == EMode::Excluding)
+	{
+		if (!ReachableExcluded[ObjectIndex] && !Excluded[ObjectIndex])
+		{
+			ReachableExcluded[ObjectIndex] = true;
+			Context.ObjectsToSerialize.Add<Options>(Object);
+		}
+	}
+}
+
+bool FMemoryUsageReferenceProcessor::GetUnreachablePackages(TSet<FName>& OutUnreachablePackages)
+{
+	for (int i = 0; i < ReachableFull.Num(); i++)
+	{
+		if (ReachableFull[i] && !ReachableExcluded[i])
+		{
+			UObject* Obj = static_cast<UObject*>(GUObjectArray.IndexToObjectUnsafeForGC(i)->Object);
+			if (Obj && Obj->IsA<UPackage>())
+			{
+				OutUnreachablePackages.Add(Obj->GetFName());
+			}
+		}
+	}
+
+	return true;
+}
+
+bool GetLongName(const FString& ShortPackageName, FName& OutLongPackageName, FOutputDevice* ErrorOutput /* = GLog */)
+{
+	FAssetRegistryModule& AssetRegistryModule = GetAssetRegistryModule();
+
+	if (FPackageName::IsValidLongPackageName(ShortPackageName))
+	{
+		OutLongPackageName = FName(*ShortPackageName);
+	}
+	else
+	{
+		OutLongPackageName = AssetRegistryModule.Get().GetFirstPackageByName(ShortPackageName);
+		if (OutLongPackageName == NAME_None)
+		{
+			ErrorOutput->Logf(TEXT("MemQuery Error: Package not found: %s"), *ShortPackageName);
+			return false;
+		}
+	}
+
+	return true;
+}
+
+bool GetLongNames(const TArray<FString>& PackageNames, TSet<FName>& OutLongPackageNames, FOutputDevice* ErrorOutput /* = GLog */)
+{
+	for (const auto& Package : PackageNames)
+	{
+		FName LongName;
+		if (!GetLongName(Package, LongName, ErrorOutput))
+		{
+			return false;
+		}
+		OutLongPackageNames.Add(LongName);
+	}
+
+	return true;
+}
+
+bool GetLongNameAndDependencies(const FString& PackageName, FName& OutLongPackageName, TSet<FName>& OutDependencies, FOutputDevice* ErrorOutput /* = GLog */)
+{
+	if (!GetLongName(PackageName, OutLongPackageName, ErrorOutput))
+	{
+		return false;
+	}
+
+	FAssetRegistryModule& AssetRegistryModule = GetAssetRegistryModule();
+
+	GetTransitiveDependencies(OutLongPackageName, AssetRegistryModule, OutDependencies);
+	OutDependencies.Add(OutLongPackageName);
+
+	return true;
+}
+
+bool GetDependenciesCombined(const TArray<FString>& PackageNames, TSet<FName>& OutDependencies, FOutputDevice* ErrorOutput /* = GLog */)
+{
+	FName LongPackageName;
+	TSet<FName> Dependencies;
+
+	for (const auto& Package : PackageNames)
+	{
+		Dependencies.Empty();
+		if (!Internal::GetLongNameAndDependencies(Package, LongPackageName, Dependencies, ErrorOutput))
+		{
+			return false;
+		}
+
+		OutDependencies.Append(Dependencies);
+	}
+
+	return true;
+}
+
+bool GetDependenciesShared(const TArray<FString>& PackageNames, TSet<FName>& OutDependencies, FOutputDevice* ErrorOutput /* = GLog */)
+{
+	FName LongPackageName;
+	TSet<FName> Dependencies;
+
+	for (int i = 0; i < PackageNames.Num(); i++)
+	{
+		Dependencies.Empty();
+		if (!Internal::GetLongNameAndDependencies(PackageNames[i], LongPackageName, Dependencies, ErrorOutput))
+		{
+			return false;
+		}
+
+		if (i == 0)
+		{
+			OutDependencies.Append(Dependencies);
+			continue;
+		}
+
+		OutDependencies = OutDependencies.Intersect(Dependencies);
+	}
+
+	return true;
+}
+
+bool PerformReachabilityAnalysis(FMemoryUsageReferenceProcessor& ReferenceProcessor, FOutputDevice* ErrorOutput /* = GLog */)
+{
+	{
+		FGCArrayStruct ArrayStruct;
+		ArrayStruct.SetInitialObjectsUnpadded(ReferenceProcessor.GetRootSet());
+		ReferenceProcessor.SetMode(FMemoryUsageReferenceProcessor::Full);
+		CollectReferences<FMemoryUsageReferenceCollector, FMemoryUsageReferenceProcessor>(ReferenceProcessor, ArrayStruct);
+	}
+
+	{
+		FGCArrayStruct ArrayStruct;
+		ArrayStruct.SetInitialObjectsUnpadded(ReferenceProcessor.GetRootSet());
+		ReferenceProcessor.SetMode(FMemoryUsageReferenceProcessor::Excluding);
+		CollectReferences<FMemoryUsageReferenceCollector, FMemoryUsageReferenceProcessor>(ReferenceProcessor, ArrayStruct);
+	}
+
+	return true;
+}
+
+bool GetRemovablePackages(const TArray<FString>& PackagesToUnload, TSet<FName>& OutRemovablePackages, FOutputDevice* ErrorOutput /* = GLog */)
+{
+	FMemoryUsageReferenceProcessor ReferenceProcessor;
+
+	if (!ReferenceProcessor.Init(PackagesToUnload, ErrorOutput))
+	{
+		return false;
+	}
+
+	if (!PerformReachabilityAnalysis(ReferenceProcessor, ErrorOutput))
+	{
+		return false;
+	}
+
+	return ReferenceProcessor.GetUnreachablePackages(OutRemovablePackages);
+}
+
+bool GetUnremovablePackages(const TArray<FString>& PackagesToUnload, TSet<FName>& OutUnremovablePackages, FOutputDevice* ErrorOutput /* = GLog */)
+{
+	FMemoryUsageReferenceProcessor ReferenceProcessor;
+
+	if (!ReferenceProcessor.Init(PackagesToUnload, ErrorOutput))
+	{
+		return false;
+	}
+
+	if (!PerformReachabilityAnalysis(ReferenceProcessor, ErrorOutput))
+	{
+		return false;
+	}
+
+	TSet<FName> UnreachablePackages;
+	if (!ReferenceProcessor.GetUnreachablePackages(UnreachablePackages))
+	{
+		return false;
+	}
+
+	TSet<FName> Dependencies;
+	if (!Internal::GetDependenciesCombined(PackagesToUnload, Dependencies, ErrorOutput))
+	{
+		return false;
+	}
+
+	for (const auto& Package : Dependencies)
+	{
+		if (!UnreachablePackages.Contains(Package) && StaticFindObjectFast(UPackage::StaticClass(), nullptr, Package, true) != nullptr)
+		{
+			OutUnremovablePackages.Add(Package);
+		}
+	}
+
+	return true;
+}
+
+FAssetRegistryModule& GetAssetRegistryModule()
+{
+	FAssetRegistryModule& AssetRegistryModule = FModuleManager::Get().LoadModuleChecked<FAssetRegistryModule>(TEXT("AssetRegistry"));
+
+	if (IsInGameThread())
+	{
+		AssetRegistryModule.Get().WaitForCompletion();
+	}
+
+	return AssetRegistryModule;
+}
+
+void GetTransitiveDependencies(FName PackageName, FAssetRegistryModule& AssetRegistryModule, TSet<FName>& OutDependencies)
+{
+	TArray<FName> PackageQueue;
+	TSet<FName> ExaminedPackages;
+	TArray<FName> PackageDependencies;
+	OutDependencies.Empty();
+	
+	PackageQueue.Push(PackageName);
+
+	while (PackageQueue.Num() > 0)
+	{
+		const FName& CurrentPackage = PackageQueue.Pop();
+		if (ExaminedPackages.Contains(CurrentPackage))
+		{
+			continue;
+		}
+
+		ExaminedPackages.Add(CurrentPackage);
+
+		if (CurrentPackage != PackageName && !OutDependencies.Contains(CurrentPackage))
+		{
+			OutDependencies.Add(CurrentPackage);
+		}
+
+		PackageDependencies.Empty();
+		AssetRegistryModule.Get().GetDependencies(CurrentPackage, PackageDependencies, UE::AssetRegistry::EDependencyCategory::Package, UE::AssetRegistry::EDependencyQuery::Hard);
+
+		for (const auto& Package : PackageDependencies)
+		{
+			if (!ExaminedPackages.Contains(Package))
+			{
+				PackageQueue.Push(Package);
+			}
+		}
+	}
+}
+
+void SortPackagesBySize(const IMemoryUsageInfoProvider* MemoryUsageInfoProvider, const TSet<FName>& Packages, TMap<FName, uint64>& OutPackagesWithSize, FOutputDevice* ErrorOutput /* = GLog */)
+{
+	GetPackagesSize(MemoryUsageInfoProvider, Packages, OutPackagesWithSize, ErrorOutput);
+	OutPackagesWithSize.ValueSort([&](const uint64& A, const uint64& B) { return A > B; });
+}
+
+void GetPackagesSize(const IMemoryUsageInfoProvider* MemoryUsageInfoProvider, const TSet<FName>& Packages, TMap<FName, uint64>& OutPackagesWithSize, FOutputDevice* ErrorOutput /* = GLog */)
+{
+	for (const auto& Package : Packages)
+	{
+		uint64 Size = MemoryUsageInfoProvider->GetAssetMemoryUsage(Package, ErrorOutput);
+		OutPackagesWithSize.Add(Package, Size);
+	}
+}
+
+void RemoveNonExistentPackages(TMap<FName, uint64>& OutPackagesWithSize)
+{
+	FAssetRegistryModule& AssetRegistryModule = GetAssetRegistryModule();
+
+	for (auto It = OutPackagesWithSize.CreateIterator(); It; ++It)
+	{
+		if (!AssetRegistryModule.Get().DoesPackageExistOnDisk(It.Key()))
+		{
+			It.RemoveCurrent();
+		}
+	}
+}
+
+void RemoveFilteredPackages(TMap<FName, uint64>& OutPackagesWithSize, const FString& AssetSubstring)
+{
+	for (auto It = OutPackagesWithSize.CreateIterator(); It; ++It)
+	{
+		FString KeyString = It.Key().ToString();
+		if (!KeyString.Contains(AssetSubstring))
+		{
+			It.RemoveCurrent();
+		}
+	}
+}
+
+void PrintTagsWithSize(FOutputDevice& Ar, const TMap<FName, uint64>& TagsWithSize, const TCHAR* Name, bool bTruncate /* = false */, int32 Limit /* = -1 */)
+{
+	uint64 TotalSize = 0U;
+	static const FString NoScopeString(TEXT("No scope"));
+
+	Ar.Logf(TEXT("----------------------------------------------------------"));
+	if (Limit < 0)
+	{
+		Limit = DefaultResultLimit;
+	}
+
+	int32 Num = TagsWithSize.Num();
+	int32 TagsToDisplay = bTruncate ? FMath::Min(Num, Limit) : Num;
+	int It = 0;
+
+	for (auto& Elem : TagsWithSize)
+	{
+		if (It++ >= TagsToDisplay)
+		{
+			break;
+		}
+
+		TotalSize += Elem.Value;
+		const FString KeyName = Elem.Key.IsValid() ? Elem.Key.ToString() : NoScopeString;
+
+		Ar.Logf(TEXT("%s - %.2f MB (%.2f KB)"), *KeyName, Elem.Value / (1024.f * 1024.f), Elem.Value / 1024.f);
+	}
+
+	Ar.Logf(TEXT("----------------------------------------------------------"));
+
+	if (TagsToDisplay < Num)
+	{
+		Ar.Logf(TEXT("<<truncated>> - displayed %d out of %d %s."), TagsToDisplay, Num, Name);
+	}
+
+	Ar.Logf(TEXT("TOTAL: %.2f MB (%.2f KB)"), TotalSize / (1024.f * 1024.f), TotalSize / 1024.f);
+	Ar.Logf(TEXT("----------------------------------------------------------"));
+}
+
+} // namespace Internal
+
+} // namespace MemoryUsageQueries
diff --git a/Engine/Plugins/MemoryUsageQueries/Source/MemoryUsageQueries/Private/MemoryUsageQueriesModule.cpp b/Engine/Plugins/MemoryUsageQueries/Source/MemoryUsageQueries/Private/MemoryUsageQueriesModule.cpp
new file mode 100644
--- /dev/null
+++ b/Engine/Plugins/MemoryUsageQueries/Source/MemoryUsageQueries/Private/MemoryUsageQueriesModule.cpp
@@ -0,0 +1,18 @@
+// Copyright Epic Games, Inc. All Rights Reserved.
+
+#include "MemoryUsageQueriesModule.h"
+#include "MemoryUsageQueries.h"
+
+#include "Engine/Console.h"
+
+void FMemoryUsageQueriesModule::StartupModule()
+{
+	Handle_RegisterConsoleAutoCompleteEntries = UConsole::RegisterConsoleAutoCompleteEntries.AddStatic(MemoryUsageQueries::RegisterConsoleAutoCompleteEntries);
+}
+
+void FMemoryUsageQueriesModule::ShutdownModule()
+{
+	UConsole::RegisterConsoleAutoCompleteEntries.Remove(Handle_RegisterConsoleAutoCompleteEntries);
+}
+
+IMPLEMENT_MODULE(FMemoryUsageQueriesModule, MemoryUsageQueries);
diff --git a/Engine/Plugins/MemoryUsageQueries/Source/MemoryUsageQueries/Private/MemoryUsageQueriesPrivate.h b/Engine/Plugins/MemoryUsageQueries/Source/MemoryUsageQueries/Private/MemoryUsageQueriesPrivate.h
new file mode 100644
--- /dev/null
+++ b/Engine/Plugins/MemoryUsageQueries/Source/MemoryUsageQueries/Private/MemoryUsageQueriesPrivate.h
@@ -0,0 +1,71 @@
+// Copyright Epic Games, Inc. All Rights Reserved.
+#pragma once
+
+#include "AssetRegistry/AssetRegistryModule.h"
+#include "Containers/UnrealString.h"
+#include "CoreTypes.h"
+#include "Engine/EngineTypes.h"
+#include "UObject/FastReferenceCollector.h"
+
+namespace MemoryUsageQueries::Internal
+{
+
+class FMemoryUsageReferenceProcessor : public FSimpleReferenceProcessorBase
+{
+public:
+	enum EMode
+	{
+		Full,
+		Excluding
+	};
+
+private:
+	TBitArray<> Excluded;
+
+	TBitArray<> ReachableFull;
+	TBitArray<> ReachableExcluded;
+
+	TArray<UObject*> RootSetPackages;
+
+	EMode Mode;
+
+public:
+	FMemoryUsageReferenceProcessor();
+
+	bool Init(const TArray<FString>& PackageNames, FOutputDevice* ErrorOutput = GLog);
+	TArray<UObject*>& GetRootSet();
+	void HandleTokenStreamObjectReference(UE::GC::FWorkerContext& Context, const UObject* ReferencingObject, UObject*& Object, UE::GC::FTokenId TokenIndex, const EGCTokenType TokenType, bool bAllowReferenceElimination);
+	bool GetUnreachablePackages(TSet<FName>& OutUnreachablePackages);
+
+	void SetMode(EMode InMode) { Mode = InMode; }
+};
+
+typedef UE::GC::TDefaultCollector<FMemoryUsageReferenceProcessor> FMemoryUsageReferenceCollector;
+
+bool GetLongName(const FString& ShortPackageName, FName& OutLongPackageName, FOutputDevice* ErrorOutput = GLog);
+bool GetLongNames(const TArray<FString>& PackageNames, TSet<FName>& OutLongPackageNames, FOutputDevice* ErrorOutput = GLog);
+bool GetLongNameAndDependencies(const FString& PackageName, FName& OutLongPackageName, TSet<FName>& OutDependencies, FOutputDevice* ErrorOutput = GLog);
+
+bool GetDependenciesCombined(const TArray<FString>& PackageNames, TSet<FName>& OutDependencies, FOutputDevice* ErrorOutput = GLog);
+bool GetDependenciesShared(const TArray<FString>& PackageNames, TSet<FName>& OutDependencies, FOutputDevice* ErrorOutput = GLog);
+
+bool PerformReachabilityAnalysis(FMemoryUsageReferenceProcessor& ReferenceProcessor, FOutputDevice* ErrorOutput = GLog);
+
+// Get Packages that would be GC'd if PackagesToUnload were unloaded
+bool GetRemovablePackages(const TArray<FString>& PackagesToUnload, TSet<FName>& OutRemovablePackages, FOutputDevice* ErrorOutput = GLog);
+
+// Get Packages that would not be GC'd if PackagesToUnload were unloaded
+bool GetUnremovablePackages(const TArray<FString>& PackagesToUnload, TSet<FName>& OutUnremovablePackages, FOutputDevice* ErrorOutput = GLog);
+
+FAssetRegistryModule& GetAssetRegistryModule();
+void GetTransitiveDependencies(FName PackageName, FAssetRegistryModule& AssetRegistryModule, TSet<FName>& OutDependencies);
+
+void SortPackagesBySize(const IMemoryUsageInfoProvider* MemoryUsageInfoProvider, const TSet<FName>& Packages, TMap<FName, uint64>& OutPackagesWithSize, FOutputDevice* ErrorOutput = GLog);
+void GetPackagesSize(const IMemoryUsageInfoProvider* MemoryUsageInfoProvider, const TSet<FName>& Packages, TMap<FName, uint64>& OutPackagesWithSize, FOutputDevice* ErrorOutput = GLog);
+
+void RemoveNonExistentPackages(TMap<FName, uint64>& OutPackagesWithSize);
+void RemoveFilteredPackages(TMap<FName, uint64>& OutPackagesWithSize, const FString& AssetSubstring);
+
+void PrintTagsWithSize(FOutputDevice& Ar, const TMap<FName, uint64>& TagsWithSize, const TCHAR* Name, bool bTruncate = false, int32 Limit = -1);
+
+}
diff --git a/Engine/Plugins/MemoryUsageQueries/Source/MemoryUsageQueries/Public/MemoryUsageInfoProvider.h b/Engine/Plugins/MemoryUsageQueries/Source/MemoryUsageQueries/Public/MemoryUsageInfoProvider.h
new file mode 100644
--- /dev/null
+++ b/Engine/Plugins/MemoryUsageQueries/Source/MemoryUsageQueries/Public/MemoryUsageInfoProvider.h
@@ -0,0 +1,28 @@
+// Copyright Epic Games, Inc. All Rights Reserved.
+
+#pragma once
+
+#include "CoreMinimal.h"
+#include "HAL/LowLevelMemTracker.h"
+#include "Misc/OutputDeviceRedirector.h"
+
+
+class IMemoryUsageInfoProvider
+{
+public:
+	virtual uint64 GetAssetMemoryUsage(FName Asset, FOutputDevice* ErrorOutput) const = 0;
+	virtual uint64 GetAssetsMemoryUsage(const TSet<FName>& Assets, FOutputDevice* ErrorOutput) const = 0;
+	virtual void GetAllAssetsWithSize(TMap<FName, uint64>& OutAssets, FOutputDevice* ErrorOutput) const = 0;
+};
+
+class FMemoryUsageInfoProviderLLM : public IMemoryUsageInfoProvider
+{
+public:
+	virtual uint64 GetAssetMemoryUsage(FName Asset, FOutputDevice* ErrorOutput = GLog) const override;
+	virtual uint64 GetAssetsMemoryUsage(const TSet<FName>& Assets, FOutputDevice* ErrorOutput = GLog) const override;
+	virtual void GetAllAssetsWithSize(TMap<FName, uint64>& OutAssets, FOutputDevice* ErrorOutput = GLog) const override;
+
+#if ENABLE_LOW_LEVEL_MEM_TRACKER
+	virtual void GetFilteredTagsWithSize(TMap<FName, uint64>& OutTags, ELLMTracker Tracker, ELLMTagSet TagSet, TArray<FLLMTagSetAllocationFilter>& Filters, FOutputDevice* ErrorOutput = GLog) const;
+#endif
+};
\ No newline at end of file
diff --git a/Engine/Plugins/MemoryUsageQueries/Source/MemoryUsageQueries/Public/MemoryUsageQueries.h b/Engine/Plugins/MemoryUsageQueries/Source/MemoryUsageQueries/Public/MemoryUsageQueries.h
new file mode 100644
--- /dev/null
+++ b/Engine/Plugins/MemoryUsageQueries/Source/MemoryUsageQueries/Public/MemoryUsageQueries.h
@@ -0,0 +1,36 @@
+// Copyright Epic Games, Inc. All Rights Reserved.
+
+#pragma once
+
+#include "Engine/Console.h"
+#include "HAL/Platform.h"
+#include "MemoryUsageInfoProvider.h"
+
+class FOutputDevice;
+
+namespace MemoryUsageQueries
+{
+
+extern MEMORYUSAGEQUERIES_API FMemoryUsageInfoProviderLLM MemoryUsageInfoProviderLLM;
+
+void RegisterConsoleAutoCompleteEntries(TArray<FAutoCompleteCommand>& AutoCompleteList);
+
+bool MEMORYUSAGEQUERIES_API GetMemoryUsage(const IMemoryUsageInfoProvider* MemoryUsageInfoProvider, const FString& PackageName, uint64& OutExclusiveSize, uint64& OutInclusiveSize, FOutputDevice* ErrorOutput = GLog);
+bool MEMORYUSAGEQUERIES_API GetMemoryUsageCombined(const IMemoryUsageInfoProvider* MemoryUsageInfoProvider, const TArray<FString>& PackageNames, uint64& OutTotalSize, FOutputDevice* ErrorOutput = GLog);
+bool MEMORYUSAGEQUERIES_API GetMemoryUsageShared(const IMemoryUsageInfoProvider* MemoryUsageInfoProvider, const TArray<FString>& PackageNames, uint64& OutTotalSize, FOutputDevice* ErrorOutput = GLog);
+bool MEMORYUSAGEQUERIES_API GetMemoryUsageUnique(const IMemoryUsageInfoProvider* MemoryUsageInfoProvider, const TArray<FString>& PackageNames, uint64& OutUniqueSize, FOutputDevice* ErrorOutput = GLog);
+bool MEMORYUSAGEQUERIES_API GetMemoryUsageCommon(const IMemoryUsageInfoProvider* MemoryUsageInfoProvider, const TArray<FString>& PackageNames, uint64& OutCommonSize, FOutputDevice* ErrorOutput = GLog);
+
+bool MEMORYUSAGEQUERIES_API GetDependenciesWithSize(const IMemoryUsageInfoProvider* MemoryUsageInfoProvider, const FString& PackageName, TMap<FName, uint64>& OutDependenciesWithSize, FOutputDevice* ErrorOutput = GLog);
+bool MEMORYUSAGEQUERIES_API GetDependenciesWithSizeCombined(const IMemoryUsageInfoProvider* MemoryUsageInfoProvider, const TArray<FString>& PackageNames, TMap<FName, uint64>& OutDependenciesWithSize, FOutputDevice* ErrorOutput = GLog);
+bool MEMORYUSAGEQUERIES_API GetDependenciesWithSizeShared(const IMemoryUsageInfoProvider* MemoryUsageInfoProvider, const TArray<FString>& PackageNames, TMap<FName, uint64>& OutDependenciesWithSize, FOutputDevice* ErrorOutput = GLog);
+bool MEMORYUSAGEQUERIES_API GetDependenciesWithSizeUnique(const IMemoryUsageInfoProvider* MemoryUsageInfoProvider, const TArray<FString>& PackageNames, TMap<FName, uint64>& OutDependenciesWithSize, FOutputDevice* ErrorOutput = GLog);
+bool MEMORYUSAGEQUERIES_API GetDependenciesWithSizeCommon(const IMemoryUsageInfoProvider* MemoryUsageInfoProvider, const TArray<FString>& PackageNames, TMap<FName, uint64>& OutDependenciesWithSize, FOutputDevice* ErrorOutput = GLog);
+
+#if ENABLE_LOW_LEVEL_MEM_TRACKER
+bool MEMORYUSAGEQUERIES_API GetFilteredPackagesWithSize(TMap<FName, uint64>& OutPackagesWithSize, FName GroupName = NAME_None, FString AssetSubstring = FString(), FName ClassName = NAME_None, FOutputDevice* ErrorOutput = GLog);
+bool MEMORYUSAGEQUERIES_API GetFilteredClassesWithSize(TMap<FName, uint64>& OutClassesWithSize, FName GroupName = NAME_None, FString AssetName = FString(), FOutputDevice* ErrorOutput = GLog);
+bool MEMORYUSAGEQUERIES_API GetFilteredGroupsWithSize(TMap<FName, uint64>& OutGroupsWithSize, FString AssetName = FString(), FName ClassName = NAME_None, FOutputDevice* ErrorOutput = GLog);
+#endif
+
+}
\ No newline at end of file
diff --git a/Engine/Plugins/MemoryUsageQueries/Source/MemoryUsageQueries/Public/MemoryUsageQueriesConfig.h b/Engine/Plugins/MemoryUsageQueries/Source/MemoryUsageQueries/Public/MemoryUsageQueriesConfig.h
new file mode 100644
--- /dev/null
+++ b/Engine/Plugins/MemoryUsageQueries/Source/MemoryUsageQueries/Public/MemoryUsageQueriesConfig.h
@@ -0,0 +1,23 @@
+// Copyright Epic Games, Inc. All Rights Reserved.
+
+#pragma once
+
+#include "MemoryUsageQueriesConfig.generated.h"
+
+// “MemQuery Savings” command is supposed to list potential memory savings for assets based on predefined presets. 
+// You can define them in DefaultMemoryUsageQueries.ini in your project – you have to provide preset name and UClass.
+// "MemQuery Savings" command should list all derived classes with potential savings for the specified preset.
+// 
+// e.g.
+// [/Script/MemoryUsageQueries.MemoryUsageQueriesConfig]
+// +SavingsPresets = (("Pawns", "/Script/OurProject.OurCharacter"))
+
+UCLASS(config = MemoryUsageQueries)
+class UMemoryUsageQueriesConfig : public UObject
+{
+	GENERATED_BODY()
+
+public:
+	UPROPERTY(config)
+	TMap<FString, FString> SavingsPresets;
+};
\ No newline at end of file
diff --git a/Engine/Plugins/MemoryUsageQueries/Source/MemoryUsageQueries/Public/MemoryUsageQueriesModule.h b/Engine/Plugins/MemoryUsageQueries/Source/MemoryUsageQueries/Public/MemoryUsageQueriesModule.h
new file mode 100644
--- /dev/null
+++ b/Engine/Plugins/MemoryUsageQueries/Source/MemoryUsageQueries/Public/MemoryUsageQueriesModule.h
@@ -0,0 +1,18 @@
+// Copyright Epic Games, Inc. All Rights Reserved.
+
+#pragma once
+
+#include "CoreMinimal.h"
+#include "Modules/ModuleInterface.h"
+#include "Modules/ModuleManager.h"
+ 
+class FMemoryUsageQueriesModule : public IModuleInterface
+{
+public:
+	/** IModuleInterface implementation */
+	virtual void StartupModule() override;
+	virtual void ShutdownModule() override;
+
+private:
+	FDelegateHandle Handle_RegisterConsoleAutoCompleteEntries;
+};
\ No newline at end of file
diff --git a/Engine/Source/Runtime/AssetRegistry/Private/AssetRegistryState.cpp b/Engine/Source/Runtime/AssetRegistry/Private/AssetRegistryState.cpp
--- a/Engine/Source/Runtime/AssetRegistry/Private/AssetRegistryState.cpp
+++ b/Engine/Source/Runtime/AssetRegistry/Private/AssetRegistryState.cpp
@@ -27,6 +27,12 @@
 #include "UObject/UObjectHash.h"
 #include "UObject/UObjectIterator.h"
 
+// Even if bSerializeDependencies is enabled, this can bypass serializing at runtime.
+// Update your <project>.Target.cs build scripts to define as needed.
+#ifndef ASSET_REGISTRY_ALLOW_DEPENDENCY_SERIALIZATION
+#define ASSET_REGISTRY_ALLOW_DEPENDENCY_SERIALIZATION 1
+#endif
+
 FAssetRegistryState& FAssetRegistryState::operator=(FAssetRegistryState&& Rhs)
 {
 	Reset();
@@ -1262,6 +1268,7 @@ void FAssetRegistryState::Load(Archive&& Ar, const FAssetRegistryHeader& Header,
 		Ar << DependencySectionSize;
 		int64 DependencySectionEnd = Ar.Tell() + DependencySectionSize;
 
+#if ASSET_REGISTRY_ALLOW_DEPENDENCY_SERIALIZATION
 		if (Options.bLoadDependencies)
 		{
 			LoadDependencies(Ar);
@@ -1271,6 +1278,9 @@ void FAssetRegistryState::Load(Archive&& Ar, const FAssetRegistryHeader& Header,
 		{
 			Ar.Seek(DependencySectionEnd);
 		}
+#else
+		Ar.Seek(DependencySectionEnd);
+#endif
 	}
 
 	int32 LocalNumPackageData = 0;
diff --git a/Engine/Source/Runtime/Core/Private/Async/InheritedContext.cpp b/Engine/Source/Runtime/Core/Private/Async/InheritedContext.cpp
--- a/Engine/Source/Runtime/Core/Private/Async/InheritedContext.cpp
+++ b/Engine/Source/Runtime/Core/Private/Async/InheritedContext.cpp
@@ -11,7 +11,7 @@ namespace UE
 
 		return FInheritedContextScope{
 #if ENABLE_LOW_LEVEL_MEM_TRACKER
-				InheritedLLMTag
+				InheritedLLMTags
 	#if UE_MEMORY_TAGS_TRACE_ENABLED && UE_TRACE_ENABLED
 				,
 	#endif
diff --git a/Engine/Source/Runtime/Core/Private/HAL/LowLevelMemTracker.cpp b/Engine/Source/Runtime/Core/Private/HAL/LowLevelMemTracker.cpp
--- a/Engine/Source/Runtime/Core/Private/HAL/LowLevelMemTracker.cpp
+++ b/Engine/Source/Runtime/Core/Private/HAL/LowLevelMemTracker.cpp
@@ -378,20 +378,15 @@ public:
 
 	void Clear();
 
-	void PushTag(const FTagData* TagData);
-	void PopTag();
-	const FTagData* GetTopTag();
-#if LLM_ALLOW_ASSETS_TAGS
-	void PushAssetTag(const FTagData* TagData);
-	void PopAssetTag();
-	const FTagData* GetTopAssetTag();
-#endif
-	void TrackAllocation(const void* Ptr, int64 Size, ELLMTracker Tracker, ELLMAllocType AllocType,
-		const FTagData* TagData, const FTagData* AssetTagData, bool bTrackInMemPro);
-	void TrackFree(const void* Ptr, int64 Size, ELLMTracker Tracker, ELLMAllocType AllocType,
-		const FTagData* TagData, const FTagData* AssetTagData, bool bTrackInMemPro);
+	void PushTag(const FTagData* TagData, ELLMTagSet TagSet);
+	void PopTag(ELLMTagSet TagSet);
+	const FTagData* GetTopTag(ELLMTagSet TagSet);
+	void TrackAllocation(const void* Ptr, int64 Size, ELLMTracker Tracker, ELLMAllocType AllocType, 
+		const FTagData* TagData, const FTagData* AssetTagData, const FTagData* AssetClassTagData, bool bTrackInMemPro);
+	void TrackFree(const void* Ptr, int64 Size, ELLMTracker Tracker, ELLMAllocType AllocType, 
+		const FTagData* TagData, const FTagData* AssetTagData, const FTagData* AssetClassTagData, bool bTrackInMemPro);
 	void TrackMemory(int64 Amount, ELLMTracker Tracker, ELLMAllocType AllocType,
-		const FTagData* TagData, const FTagData* AssetTagData);
+		const FTagData* TagData);
 	void TrackMoved(const void* Dest, const void* Source, int64 Size, ELLMTracker Tracker, const FTagData* TagData);
 	void IncrTag(const FTagData* Tag, int64 Amount);
 
@@ -402,11 +397,7 @@ public:
 
 	void ClearAllocTypeAmounts();
 
-	FConstTagDataArray TagStack;
-#if LLM_ALLOW_ASSETS_TAGS
-	FConstTagDataArray AssetTagStack;
-#endif
-
+	FConstTagDataArray TagStack[static_cast<int32>(ELLMTagSet::Max)];
 	FThreadTagSizeMap Allocations;
 	FCriticalSection TagSection;
 
@@ -424,15 +415,10 @@ public:
 
 	void Initialise(ELLMTracker Tracker, FLLMAllocator* InAllocator);
 
-	void PushTag(ELLMTag EnumTag);
-	void PushTag(FName Tag, bool bInIsStatTag);
-	void PushTag(const FTagData* TagData);
-	void PopTag();
-#if LLM_ALLOW_ASSETS_TAGS
-	void PushAssetTag(FName Tag);
-	void PushAssetTag(const FTagData* Tag);
-	void PopAssetTag();
-#endif
+	void PushTag(ELLMTag EnumTag, ELLMTagSet TagSet);
+	void PushTag(FName Tag, bool bInIsStatTag, ELLMTagSet TagSet);
+	void PushTag(const FTagData* TagData, ELLMTagSet TagSet);
+	void PopTag(ELLMTagSet TagSet);
 	void TrackAllocation(const void* Ptr, int64 Size, FName DefaultTag, ELLMTracker Tracker, ELLMAllocType AllocType,
 		bool bTrackInMemPro);
 	void TrackAllocation(const void* Ptr, int64 Size, ELLMTag DefaultTag, ELLMTracker Tracker, ELLMAllocType AllocType,
@@ -443,11 +429,12 @@ public:
 	void OnAllocMoved(const void* Dest, const void* Source, ELLMTracker Tracker, ELLMAllocType AllocType);
 
 	void TrackMemory(ELLMTag EnumTag, ELLMTracker Tracker, int64 Amount, ELLMAllocType AllocType);
-	void TrackMemory(FName TagName, ELLMTracker Tracker, int64 Amount, ELLMAllocType AllocType);
+	void TrackMemory(FName TagName, ELLMTagSet TagSet, ELLMTracker Tracker, int64 Amount, ELLMAllocType AllocType);
 	void TrackMemory(const FTagData* TagData, ELLMTracker Tracker, int64 Amount, ELLMAllocType AllocType);
 
 	// This will pause/unpause tracking, and also manually increment a given tag
-	void PauseAndTrackMemory(FName TagName, ELLMTracker Tracker, bool bInIsStatTag, int64 Amount, ELLMAllocType AllocType);
+	void PauseAndTrackMemory(FName TagName, ELLMTagSet TagSet, ELLMTracker Tracker, bool bInIsStatTag, int64 Amount,
+		ELLMAllocType AllocType);
 	void PauseAndTrackMemory(ELLMTag EnumTag, ELLMTracker Tracker, int64 Amount, ELLMAllocType AllocType);
 	void PauseAndTrackMemory(const FTagData* TagData, ELLMTracker Tracker, int64 Amount, ELLMAllocType AllocType);
 	void Pause(ELLMAllocType AllocType);
@@ -466,56 +453,88 @@ public:
 	struct FLowLevelAllocInfo
 	{
 	public:
-		void SetTag(const FTagData* InTag, FLowLevelMemTracker& InLLMRef)
+		FLowLevelAllocInfo()
+		{
+#if LLM_ENABLED_FULL_TAGS && LLM_ALLOW_ASSETS_TAGS
+			FMemory::Memset(Tag, InvalidCompressedTagValue);
+#endif
+		}
+
+		void SetTag(const FTagData* InTag, FLowLevelMemTracker& InLLMRef
+#if LLM_ALLOW_ASSETS_TAGS
+		, ELLMTagSet InTagSet = ELLMTagSet::None
+#endif
+		)
 		{
 #if LLM_ENABLED_FULL_TAGS
+#if !LLM_ALLOW_ASSETS_TAGS
 			Tag = InTag->GetIndex();
+#else
+			Tag[static_cast<int32>(InTagSet)] = InTag ? InTag->GetIndex() : InvalidCompressedTagValue;
+#endif //!LLM_ALLOW_ASSETS_TAGS
 #else
 			Tag = InTag->GetEnumTag();
 #endif
 		}
 
-		const FTagData* GetTag(FLowLevelMemTracker& InLLMRef) const
+		const FTagData* GetTag(FLowLevelMemTracker& InLLMRef
+#if LLM_ALLOW_ASSETS_TAGS
+		, ELLMTagSet InTagSet = ELLMTagSet::None
+#endif
+		) const
 		{
 #if LLM_ENABLED_FULL_TAGS
 			FReadScopeLock TagDataScopeLock(InLLMRef.TagDataLock);
+#if !LLM_ALLOW_ASSETS_TAGS
 			return (*InLLMRef.TagDatas)[Tag];
+#else
+			return Tag[static_cast<int32>(InTagSet)] >= 0 ? (*InLLMRef.TagDatas)[Tag[static_cast<int32>(InTagSet)]] : nullptr;
+#endif
 #else
 			return InLLMRef.FindTagData(Tag);
 #endif
 		}
 
-#if LLM_ALLOW_ASSETS_TAGS
-		void SetAssetTag(const FTagData* InTag, FLowLevelMemTracker& InLLMRef)
+#if LLM_ENABLED_FULL_TAGS
+		int32 GetCompressedTag(ELLMTagSet InTagSet)
 		{
-			AssetTag = InTag;
+#if !LLM_ALLOW_ASSETS_TAGS
+			return (InTagSet == ELLMTagSet::None) ? Tag : InvalidCompressedTagValue;
+#else
+			return Tag[static_cast<int32>(InTagSet)];
+#endif
 		}
 
-		const FTagData* GetAssetTag(FLowLevelMemTracker& InLLMRef) const
+		void SetCompressedTag(int32 InTag, ELLMTagSet InTagSet)
 		{
-			return AssetTag;
-		}
+#if !LLM_ALLOW_ASSETS_TAGS
+			if (InTagSet == ELLMTagSet::None)
+			{
+				Tag = InTag;
+			}
+#else
+			Tag[static_cast<int32>(InTagSet)] = InTag;
 #endif
-
-#if LLM_ENABLED_FULL_TAGS
-		int32 GetCompressedTag() { return Tag; }
-		void SetCompressedTag(int32 InTag) { Tag = InTag; }
+		}
 #else
 		ELLMTag GetCompressedTag() { return Tag; }
 		void SetCompressedTag(ELLMTag InTag) { Tag = InTag; }
 #endif
 
 	private:
-#if LLM_ALLOW_ASSETS_TAGS
-		const FTagData* AssetTag = nullptr;
-#endif
 #if LLM_ENABLED_FULL_TAGS
 		// Even with arbitrary tags we are still partially compressed - the allocation records the tag's index
 		// (4 bytes) rather than the full tag pointer (8 bytes) or the tag's name (12 bytes).
+#if !LLM_ALLOW_ASSETS_TAGS
 		int32 Tag = 0;
+#else
+		int32 Tag[static_cast<int32>(ELLMTagSet::Max)];
+		static_assert(static_cast<uint32>(ELLMTagSet::Max) <= 3, "Every new TagSet enum adds 4 bytes of LLM overhead memory for every active allocation. Do not add them without measuring the cost of the overhead in the cook and runtime of large projects.");
+#endif
 #else
 		ELLMTag Tag = ELLMTag::Untagged;
 #endif
+		static const int32 InvalidCompressedTagValue = -1;
 	};
 
 	typedef LLMMap<PointerKey, uint32, FLowLevelAllocInfo, LLMNumAllocsType> FLLMAllocMap;	// pointer, size, info, Capacity SizeType
@@ -529,9 +548,13 @@ public:
 	int64 GetTagAmount(const FTagData* TagData, bool bPeakAmount = false) const;
 	void SetTagAmountExternal(const FTagData* TagData, int64 Amount, bool bAddToTotal);
 	void SetTagAmountInUpdate(const FTagData* TagData, int64 Amount, bool bAddToTotal);
-	const FTagData* GetActiveTagData();
-	const FTagData* FindTagForPtr(void* Ptr);
-	TArray<const FTagData*> GetTagDatas();
+	const FTagData* GetActiveTagData(ELLMTagSet TagSet = ELLMTagSet::None);
+	TArray<const FTagData*> GetTagDatas(ELLMTagSet TagSet = ELLMTagSet::None);
+
+	void GetTagsNamesWithAmount(TMap<FName, uint64>& OutTagsNamesWithAmount, ELLMTagSet TagSet = ELLMTagSet::None);
+	void GetTagsNamesWithAmountFiltered(TMap<FName, uint64>& OutTagsNamesWithAmount, ELLMTagSet TagSet, TArray<FLLMTagSetAllocationFilter>& Filters);
+
+	bool FindTagsForPtr(void* InPtr, TArray<const FTagData*, TInlineAllocator<static_cast<int32>(ELLMTagSet::Max)>>& OutTags);
 
 	int64 GetAllocTypeAmount(ELLMAllocType AllocType);
 
@@ -589,6 +612,8 @@ static FName GetTagName_CustomName()
 	return TagName_CustomName;
 }
 static FName TagName_Untagged(TEXT("Untagged"));
+static FName TagName_UntaggedAsset(TEXT("UntaggedAsset"));
+static FName TagName_UntaggedAssetClass(TEXT("UntaggedAssetClass"));
 
 FName LLMGetTagUniqueName(ELLMTag Tag)
 {
@@ -1412,7 +1437,7 @@ static bool IsAssetTagForAssets(ELLMTagSet Set)
 	return Set == ELLMTagSet::Assets || Set == ELLMTagSet::AssetClasses;
 }
 
-void FLowLevelMemTracker::RegisterCustomTagInternal(int32 Tag, const TCHAR* InDisplayName, FName StatName,
+void FLowLevelMemTracker::RegisterCustomTagInternal(int32 Tag, ELLMTagSet TagSet, const TCHAR* InDisplayName, FName StatName,
 	FName SummaryStatName, int32 ParentTag)
 {
 	using namespace UE::LLMPrivate;
@@ -1426,7 +1451,7 @@ void FLowLevelMemTracker::RegisterCustomTagInternal(int32 Tag, const TCHAR* InDi
 	FName ParentName = ParentTag >= 0 ? LLMGetTagUniqueName(static_cast<ELLMTag>(ParentTag)) : NAME_None;
 
 	RegisterTagData(LLMGetTagUniqueName(EnumTag), DisplayName, ParentName, StatName, SummaryStatName, true,
-		EnumTag, false, ETagReferenceSource::CustomEnumTag);
+		EnumTag, false, ETagReferenceSource::CustomEnumTag, TagSet);
 }
 
 void FLowLevelMemTracker::RegisterPlatformTag(int32 Tag, const TCHAR* Name, FName StatName, FName SummaryStatName,
@@ -1442,7 +1467,7 @@ void FLowLevelMemTracker::RegisterPlatformTag(int32 Tag, const TCHAR* Name, FNam
 
 	LLMCheck(Tag >= static_cast<int32>(ELLMTag::PlatformTagStart) &&
 		Tag <= static_cast<int32>(ELLMTag::PlatformTagEnd));
-	RegisterCustomTagInternal(Tag, Name, StatName, SummaryStatName, ParentTag);
+	RegisterCustomTagInternal(Tag, ELLMTagSet::None, Name, StatName, SummaryStatName, ParentTag);
 }
 
 void FLowLevelMemTracker::RegisterProjectTag(int32 Tag, const TCHAR* Name, FName StatName, FName SummaryStatName,
@@ -1457,7 +1482,7 @@ void FLowLevelMemTracker::RegisterProjectTag(int32 Tag, const TCHAR* Name, FName
 	BootstrapInitialise();
 
 	LLMCheck(Tag >= static_cast<int32>(ELLMTag::ProjectTagStart) && Tag <= static_cast<int32>(ELLMTag::ProjectTagEnd));
-	RegisterCustomTagInternal(Tag, Name, StatName, SummaryStatName, ParentTag);
+	RegisterCustomTagInternal(Tag, ELLMTagSet::None, Name, StatName, SummaryStatName, ParentTag);
 }
 
 void GlobalRegisterTagDeclaration(FLLMTagDeclaration& TagDeclaration)
@@ -1558,7 +1583,7 @@ void FLowLevelMemTracker::InitialiseTagDatas_FinishRegister()
 			ELLMTag EnumTag = ELLMTag::Enum; \
 			int32 Index = static_cast<int32>(EnumTag); \
 			FTagData* TagData = TagDataEnumMap[Index]; \
-			FTagData*& ExistingTagData = TagDataNameMap->FindOrAdd(TagData->GetName(), nullptr); \
+			FTagData*& ExistingTagData = TagDataNameMap->FindOrAdd(FTagDataNameKey(TagData->GetName(), TagData->GetTagSet()), nullptr); \
 			if (ExistingTagData != nullptr) \
 			{ \
 				ReportDuplicateTagName(ExistingTagData, ETagReferenceSource::EnumTag); \
@@ -1624,14 +1649,14 @@ void FLowLevelMemTracker::ClearTagDatas()
 void FLowLevelMemTracker::RegisterTagDeclaration(FLLMTagDeclaration& TagDeclaration)
 {
 	TagDeclaration.ConstructUniqueName();
-	RegisterTagData(TagDeclaration.UniqueName, TagDeclaration.DisplayName, TagDeclaration.ParentTagName,
-		TagDeclaration.StatName, TagDeclaration.SummaryStatName, false, ELLMTag::CustomName, false,
-		UE::LLMPrivate::ETagReferenceSource::Declare);
+	RegisterTagData(TagDeclaration.UniqueName, TagDeclaration.DisplayName, TagDeclaration.ParentTagName, 
+		TagDeclaration.StatName, TagDeclaration.SummaryStatName, false, ELLMTag::CustomName, false, 
+		UE::LLMPrivate::ETagReferenceSource::Declare, TagDeclaration.TagSet);
 }
 
-UE::LLMPrivate::FTagData& FLowLevelMemTracker::RegisterTagData(FName Name, FName DisplayName, FName ParentName,
-	FName StatName, FName SummaryStatName, bool bHasEnumTag, ELLMTag EnumTag, bool bIsStatTag,
-	UE::LLMPrivate::ETagReferenceSource ReferenceSource)
+UE::LLMPrivate::FTagData& FLowLevelMemTracker::RegisterTagData(FName Name, FName DisplayName, FName ParentName, 
+	FName StatName, FName SummaryStatName, bool bHasEnumTag, ELLMTag EnumTag, bool bIsStatTag, 
+	UE::LLMPrivate::ETagReferenceSource ReferenceSource, ELLMTagSet TagSet/* = ELLMTagSet::None*/)
 {
 	LLMCheckf(!bIsBootstrapping,
 		TEXT("A tag outside of LLM_ENUM_GENERIC_TAGS was requested from LLM_SCOPE or allocation while bootstrapping the names for LLM_ENUM_GENERIC_TAGS, this is not supported."));
@@ -1654,27 +1679,19 @@ UE::LLMPrivate::FTagData& FLowLevelMemTracker::RegisterTagData(FName Name, FName
 			DisplayName = Name;
 		}
 	}
+	else if (TagSet != ELLMTagSet::None)
+	{
+		// Tags not in the None TagSet do not have parents and do not require name validation
+		DisplayName = Name;
+	}
 	else if (bIsStatTag)
 	{
-		// Stat tag unique names do not have to be validated, because they are never used as parent tags.
-		// Stat tag unique names are of the form //GroupName//StatUniqueName///StatDisplayName///<OtherData>. 
 		// We set LLM UniqueName = <TheEntireString> and LLM DisplayName = StatDisplayName.
 		// Stat tags do not specify their parent, and their parent is set to the CustomName aggregator.
 		LLMCheck(DisplayName.IsNone());
 		LLMCheck(ParentName.IsNone());
 		DisplayName = Name;
 		ParentName = GetTagName_CustomName();
-
-		const TCHAR* Start = FCString::Strstr(NameBuffer.ToString(), TEXT("///"));
-		if (Start)
-		{
-			Start += 3;
-			const TCHAR* End = FCString::Strstr(Start, TEXT("///"));
-			if (End)
-			{
-				DisplayName = FName(FStringView(Start, UE_PTRDIFF_TO_INT32(End - Start)));
-			}
-		}
 	}
 	else
 	{
@@ -1724,7 +1741,7 @@ UE::LLMPrivate::FTagData& FLowLevelMemTracker::RegisterTagData(FName Name, FName
 	}
 
 	FWriteScopeLock ScopeLock(TagDataLock);
-	FTagData*& TagDataForName = TagDataNameMap->FindOrAdd(Name, nullptr);
+	FTagData*& TagDataForName = TagDataNameMap->FindOrAdd(FTagDataNameKey(Name, TagSet), nullptr);
 	if (TagDataForName)
 	{
 		// The tag already exists; this can happen because two formal registrations have tried to add the same name,
@@ -1755,7 +1772,7 @@ UE::LLMPrivate::FTagData& FLowLevelMemTracker::RegisterTagData(FName Name, FName
 	FTagData* ParentData = nullptr;
 	if (!ParentName.IsNone())
 	{
-		FTagData** ParentPtr = TagDataNameMap->Find(ParentName);
+		FTagData** ParentPtr = TagDataNameMap->Find(FTagDataNameKey(ParentName, TagSet));
 		if (ParentPtr)
 		{
 			ParentData = *ParentPtr;
@@ -1765,12 +1782,12 @@ UE::LLMPrivate::FTagData& FLowLevelMemTracker::RegisterTagData(FName Name, FName
 	FTagData* TagData;
 	if (ParentName.IsNone() || ParentData)
 	{
-		TagData = Allocator.New<FTagData>(Name, DisplayName, ParentData, StatName, SummaryStatName, bHasEnumTag,
+		TagData = Allocator.New<FTagData>(Name, TagSet, DisplayName, ParentData, StatName, SummaryStatName, bHasEnumTag,
 			EnumTag, ReferenceSource);
 	}
 	else
 	{
-		TagData = Allocator.New<FTagData>(Name, DisplayName, ParentName, StatName, SummaryStatName, bHasEnumTag,
+		TagData = Allocator.New<FTagData>(Name, TagSet, DisplayName, ParentName, StatName, SummaryStatName, bHasEnumTag,
 			EnumTag, ReferenceSource);
 	}
 	TagData->SetIndex(TagDatas->Num());
@@ -1856,20 +1873,20 @@ void FLowLevelMemTracker::FinishConstruct(UE::LLMPrivate::FTagData* TagData,
 		}
 		else
 		{
-			FTagData** ParentDataPtr = TagDataNameMap->Find(ParentName);
+			FTagData** ParentDataPtr = TagDataNameMap->Find(FTagDataNameKey(ParentName, TagData->GetTagSet()));
 			if (!ParentDataPtr)
 			{
 				// We have to drop the ReadLock to call RegisterTagData, which takes a WriteLock.
 				TagDataLock.ReadUnlock();
 				FTagData* ParentTagData = &RegisterTagData(ParentName, NAME_None, NAME_None, NAME_None, NAME_None,
-					false, ELLMTag::CustomName, false, ETagReferenceSource::ImplicitParent);
+					false, ELLMTag::CustomName,false, ETagReferenceSource::ImplicitParent, TagData->GetTagSet());
 				TagDataLock.ReadLock();
 				if (TagData->IsFinishConstructed())
 				{
 					// Another thread got in and finished construction while we were outside of the lock.
 					return;
 				}
-				ParentDataPtr = TagDataNameMap->Find(ParentName);
+				ParentDataPtr = TagDataNameMap->Find(FTagDataNameKey(ParentName, TagData->GetTagSet()));
 				LLMCheck(ParentDataPtr);
 			}
 			FTagData* ParentData = *ParentDataPtr;
@@ -1887,7 +1904,7 @@ void FLowLevelMemTracker::FinishConstruct(UE::LLMPrivate::FTagData* TagData,
 	}
 }
 
-TArray<const UE::LLMPrivate::FTagData*> FLowLevelMemTracker::GetTrackedTags()
+TArray<const UE::LLMPrivate::FTagData*> FLowLevelMemTracker::GetTrackedTags(ELLMTagSet TagSet)
 {
 	using namespace UE::LLMPrivate;
 
@@ -1899,10 +1916,15 @@ TArray<const UE::LLMPrivate::FTagData*> FLowLevelMemTracker::GetTrackedTags()
 	BootstrapInitialise();
 
 	FReadScopeLock TagDataScopeLock(TagDataLock);
-	return TArray<const FTagData*>(*TagDatas);
+	
+	TArray<const FTagData*> FoundResults(TagDatas->FilterByPredicate([TagSet](const FTagData* InTagData){
+		return InTagData != nullptr && InTagData->GetTagSet() == TagSet;
+	}));
+
+	return FoundResults;
 }
 
-TArray<const UE::LLMPrivate::FTagData*> FLowLevelMemTracker::GetTrackedTags(ELLMTracker Tracker)
+TArray<const UE::LLMPrivate::FTagData*> FLowLevelMemTracker::GetTrackedTags(ELLMTracker Tracker, ELLMTagSet TagSet /* = ELLMTagSet::None */)
 {
 	using namespace UE::LLMPrivate;
 
@@ -1914,10 +1936,40 @@ TArray<const UE::LLMPrivate::FTagData*> FLowLevelMemTracker::GetTrackedTags(ELLM
 	BootstrapInitialise();
 
 	FScopeLock UpdateScopeLock(&UpdateLock); // uses of TagSizes are guarded by the UpdateLock
-	return GetTracker(Tracker)->GetTagDatas();
+	return GetTracker(Tracker)->GetTagDatas(TagSet);
+}
+
+void FLowLevelMemTracker::GetTrackedTagsNamesWithAmount(TMap<FName, uint64>& TagsNamesWithAmount, ELLMTracker Tracker, ELLMTagSet TagSet)
+{
+	using namespace UE::LLMPrivate;
+
+	if (bIsDisabled)
+	{
+		return;
+	}
+
+	BootstrapInitialise();
+
+	FScopeLock UpdateScopeLock(&UpdateLock);
+	GetTracker(Tracker)->GetTagsNamesWithAmount(TagsNamesWithAmount, TagSet);
 }
 
-bool FLowLevelMemTracker::FindTagByName( const TCHAR* Name, uint64& OutTag ) const
+void FLowLevelMemTracker::GetTrackedTagsNamesWithAmountFiltered(TMap<FName, uint64>& TagsNamesWithAmount, ELLMTracker Tracker, ELLMTagSet TagSet, TArray<FLLMTagSetAllocationFilter>& Filters)
+{
+	using namespace UE::LLMPrivate;
+
+	if (bIsDisabled)
+	{
+		return;
+	}
+
+	BootstrapInitialise();
+
+	FScopeLock UpdateScopeLock(&UpdateLock);
+	GetTracker(Tracker)->GetTagsNamesWithAmountFiltered(TagsNamesWithAmount, TagSet, Filters);
+}
+
+bool FLowLevelMemTracker::FindTagByName( const TCHAR* Name, uint64& OutTag, ELLMTagSet InTagSet /*= ELLMTagSet::None*/ ) const
 {
 	using namespace UE::LLMPrivate;
 
@@ -1934,7 +1986,7 @@ bool FLowLevelMemTracker::FindTagByName( const TCHAR* Name, uint64& OutTag ) con
 
 		// Search by Name
 		FName SearchName(Name);
-		FTagData** TagDataPtr = TagDataNameMap->Find(SearchName);
+		FTagData** TagDataPtr = TagDataNameMap->Find(FTagDataNameKey(SearchName, InTagSet));
 		if (TagDataPtr)
 		{
 			const FTagData* TagData = *TagDataPtr;
@@ -2059,6 +2111,19 @@ int64 FLowLevelMemTracker::GetTagAmountForTracker(ELLMTracker Tracker, const UE:
 	return GetTracker(Tracker)->GetTagAmount(TagData, bPeakAmount);
 }
 
+int64 FLowLevelMemTracker::GetTagAmountForTracker(ELLMTracker Tracker, FName Tag, ELLMTagSet TagSet,
+	bool bPeakAmount /* = false */)
+{
+	if (bIsDisabled)
+	{
+		return 0;
+	}
+
+	BootstrapInitialise();
+
+	return GetTagAmountForTracker(Tracker, FindTagData(Tag, TagSet, UE::LLMPrivate::ETagReferenceSource::FunctionAPI), bPeakAmount);
+}
+
 void FLowLevelMemTracker::SetTagAmountForTracker(ELLMTracker Tracker, ELLMTag Tag, int64 Amount, bool bAddToTotal)
 {
 	using namespace UE::LLMPrivate;
@@ -2095,7 +2160,7 @@ int64 FLowLevelMemTracker::GetActiveTag(ELLMTracker Tracker)
 	}
 }
 
-const UE::LLMPrivate::FTagData* FLowLevelMemTracker::GetActiveTagData(ELLMTracker Tracker)
+const UE::LLMPrivate::FTagData* FLowLevelMemTracker::GetActiveTagData(ELLMTracker Tracker, ELLMTagSet TagSet /*= ELLMTagSet::None*/)
 {
 	if (bIsDisabled)
 	{
@@ -2103,7 +2168,7 @@ const UE::LLMPrivate::FTagData* FLowLevelMemTracker::GetActiveTagData(ELLMTracke
 	}
 	BootstrapInitialise();
 
-	return GetTracker(Tracker)->GetActiveTagData();
+	return GetTracker(Tracker)->GetActiveTagData(TagSet);
 }
 
 uint64 FLowLevelMemTracker::DumpTag( ELLMTracker Tracker, const char* FileName, int LineNumber )
@@ -2181,14 +2246,14 @@ const UE::LLMPrivate::FTagData* FLowLevelMemTracker::FindOrAddTagData(ELLMTag En
 	}
 }
 
-const UE::LLMPrivate::FTagData* FLowLevelMemTracker::FindOrAddTagData(FName TagName, bool bIsStatTag,
+const UE::LLMPrivate::FTagData* FLowLevelMemTracker::FindOrAddTagData(FName TagName, ELLMTagSet TagSet, bool bIsStatTag,
 	UE::LLMPrivate::ETagReferenceSource ReferenceSource)
 {
 	using namespace UE::LLMPrivate;
 
 	{
 		FReadScopeLock ScopeLock(TagDataLock);
-		FTagData** TagDataPtr = TagDataNameMap->Find(TagName);
+		FTagData** TagDataPtr = TagDataNameMap->Find(FTagDataNameKey(TagName, TagSet));
 		if (TagDataPtr)
 		{
 			FTagData* TagData = *TagDataPtr;
@@ -2203,14 +2268,14 @@ const UE::LLMPrivate::FTagData* FLowLevelMemTracker::FindOrAddTagData(FName TagN
 		FinishInitialise();
 		// Reeneter this function so that we retry the find above; note we avoid infinite recursion because
 		// bFullyInitialised is now true.
-		return FindOrAddTagData(TagName, bIsStatTag, ReferenceSource);
+		return FindOrAddTagData(TagName, TagSet, bIsStatTag, ReferenceSource);
 	}
 	LLMCheckf(!bIsBootstrapping, TEXT("LLM Error: Invalid use of FName tag when initialising tags."));
 
 	// Add the new Tag
 	FName StatName = bIsStatTag ? TagName : NAME_None;
 	FTagData* TagData = &RegisterTagData(TagName, NAME_None, NAME_None, StatName, NAME_None, false,
-		ELLMTag::CustomName, bIsStatTag, ReferenceSource);
+		ELLMTag::CustomName, bIsStatTag, ReferenceSource, TagSet);
 	{
 		FReadScopeLock ScopeLock(TagDataLock);
 		FinishConstruct(TagData, ReferenceSource);
@@ -2239,14 +2304,14 @@ const UE::LLMPrivate::FTagData* FLowLevelMemTracker::FindTagData(ELLMTag EnumTag
 	}
 }
 
-const UE::LLMPrivate::FTagData* FLowLevelMemTracker::FindTagData(FName TagName,
+const UE::LLMPrivate::FTagData* FLowLevelMemTracker::FindTagData(FName TagName, ELLMTagSet TagSet,
 	UE::LLMPrivate::ETagReferenceSource ReferenceSource)
 {
 	using namespace UE::LLMPrivate;
 
 	FReadScopeLock ScopeLock(TagDataLock);
 
-	FTagData** TagDataPtr = TagDataNameMap->Find(TagName);
+	FTagData** TagDataPtr = TagDataNameMap->Find(FTagDataNameKey(TagName, TagSet));
 	if (TagDataPtr)
 	{
 		FTagData* TagData = *TagDataPtr;
@@ -2259,7 +2324,7 @@ const UE::LLMPrivate::FTagData* FLowLevelMemTracker::FindTagData(FName TagName,
 	}
 }
 
-void FLLMScope::Init(ELLMTag TagEnum, bool bInIsStatTag, ELLMTagSet InTagSet, ELLMTracker InTracker)
+void FLLMScope::Init(ELLMTag TagEnum, bool bInIsStatTag, ELLMTagSet InTagSet, ELLMTracker InTracker, bool bOverride)
 {
 	LLMCheck(!bInIsStatTag && InTagSet == ELLMTagSet::None);
 	FLowLevelMemTracker& LLMRef = FLowLevelMemTracker::Get();
@@ -2272,15 +2337,19 @@ void FLLMScope::Init(ELLMTag TagEnum, bool bInIsStatTag, ELLMTagSet InTagSet, EL
 	}
 	LLMRef.BootstrapInitialise();
 
+	if (!bOverride && LLMRef.GetTracker(InTracker)->GetActiveTagData(InTagSet) != nullptr)
+	{
+		bEnabled = false;
+		return;
+	}
+
 	bEnabled = true;
 	Tracker = InTracker;
-#if LLM_ALLOW_ASSETS_TAGS
-	bIsAssetTag = false;
-#endif
-	LLMRef.GetTracker(Tracker)->PushTag(TagEnum);
+	TagSet = InTagSet;
+	LLMRef.GetTracker(Tracker)->PushTag(TagEnum, InTagSet);
 }
 
-void FLLMScope::Init(FName TagName, bool bInIsStatTag, ELLMTagSet InTagSet, ELLMTracker InTracker)
+void FLLMScope::Init(FName TagName, bool bInIsStatTag, ELLMTagSet InTagSet, ELLMTracker InTracker, bool bOverride)
 {
 	FLowLevelMemTracker& LLMRef = FLowLevelMemTracker::Get();
 	// We have to check bIsDisabled again after calling Get, because the constructor is called from Get, and will set
@@ -2297,26 +2366,22 @@ void FLLMScope::Init(FName TagName, bool bInIsStatTag, ELLMTagSet InTagSet, ELLM
 		return;
 	}
 
+	if (!bOverride && LLMRef.GetTracker(InTracker)->GetActiveTagData(InTagSet) != nullptr)
+	{
+		bEnabled = false;
+		return;
+	}
+
 	bEnabled = true;
 	Tracker = InTracker;
+	TagSet = InTagSet;
 
-#if LLM_ALLOW_ASSETS_TAGS
-	bIsAssetTag = bInIsStatTag && IsAssetTagForAssets(InTagSet);
-	if (bIsAssetTag)
-	{
-		LLMRef.GetTracker(Tracker)->PushAssetTag(TagName);
-	}
-	else
-#endif
-	{
-		LLMRef.GetTracker(Tracker)->PushTag(TagName, bInIsStatTag);
-	}
+	LLMRef.GetTracker(Tracker)->PushTag(TagName, bInIsStatTag, InTagSet);
 }
 
 void FLLMScope::Init(const UE::LLMPrivate::FTagData* TagData, bool bInIsStatTag, ELLMTagSet InTagSet,
-	ELLMTracker InTracker)
+	ELLMTracker InTracker, bool bOverride)
 {
-	LLMCheck(!bInIsStatTag && InTagSet == ELLMTagSet::None);
 	FLowLevelMemTracker& LLMRef = FLowLevelMemTracker::Get();
 	// We have to check bIsDisabled again after calling Get, because the constructor is called from Get, and will set
 	// bIsDisabled=false if the platform doesn't support it.
@@ -2327,27 +2392,22 @@ void FLLMScope::Init(const UE::LLMPrivate::FTagData* TagData, bool bInIsStatTag,
 	}
 	LLMRef.BootstrapInitialise();
 
+	if (!bOverride && LLMRef.GetTracker(InTracker)->GetActiveTagData(InTagSet) != nullptr)
+	{
+		bEnabled = false;
+		return;
+	}
+
 	bEnabled = true;
 	Tracker = InTracker;
-#if LLM_ALLOW_ASSETS_TAGS
-	bIsAssetTag = false;
-#endif
-	LLMRef.GetTracker(Tracker)->PushTag(TagData);
+	TagSet = InTagSet;
+	LLMRef.GetTracker(Tracker)->PushTag(TagData, InTagSet);
 }
 
 void FLLMScope::Destruct()
 {
 	FLowLevelMemTracker& LLMRef = FLowLevelMemTracker::Get();
-#if LLM_ALLOW_ASSETS_TAGS
-	if (bIsAssetTag)
-	{
-		LLMRef.GetTracker(Tracker)->PopAssetTag();
-	}
-	else
-#endif
-	{
-		LLMRef.GetTracker(Tracker)->PopTag();
-	}
+	LLMRef.GetTracker(Tracker)->PopTag(TagSet);
 }
 
 FLLMPauseScope::FLLMPauseScope(FName TagName, bool bIsStatTag, uint64 Amount, ELLMTracker TrackerToPause,
@@ -2413,8 +2473,8 @@ void FLLMPauseScope::Init(FName TagName, ELLMTag EnumTag, bool bIsEnumTag, bool
 				}
 				else
 				{
-					LLMRef.GetTracker(Tracker)->PauseAndTrackMemory(TagName, Tracker, bIsStatTag, static_cast<int64>(Amount),
-						InAllocType);
+					LLMRef.GetTracker(Tracker)->PauseAndTrackMemory(TagName, ELLMTagSet::None, Tracker, bIsStatTag, 
+						static_cast<int64>(Amount), InAllocType);
 				}
 			}
 		}
@@ -2441,17 +2501,17 @@ FLLMPauseScope::~FLLMPauseScope()
 }
 
 
-FLLMScopeFromPtr::FLLMScopeFromPtr(void* Ptr, ELLMTracker InTracker )
+FLLMScopeFromPtr::FLLMScopeFromPtr(void* Ptr, ELLMTracker InTracker)
 {
 	using namespace UE::LLMPrivate;
 	if (FLowLevelMemTracker::bIsDisabled)
 	{
-		bEnabled = false;
+		DisableAll();
 		return;
 	}
 	if(Ptr == nullptr)
 	{
-		bEnabled = false;
+		DisableAll();
 		return;
 	}
 
@@ -2460,40 +2520,51 @@ FLLMScopeFromPtr::FLLMScopeFromPtr(void* Ptr, ELLMTracker InTracker )
 	// bIsDisabled=false if the platform doesn't support it.
 	if (FLowLevelMemTracker::bIsDisabled)
 	{
-		bEnabled = false;
+		DisableAll();
 		return;
 	}
 	LLMRef.BootstrapInitialise();
 
 	FLLMTracker* TrackerData = LLMRef.GetTracker(InTracker);
-	const FTagData* TagData = TrackerData->FindTagForPtr(Ptr);
-	if (!TagData)
+
+	TArray<const FTagData*, TInlineAllocator<static_cast<int32>(ELLMTagSet::Max)>> Tags;
+	if (!TrackerData->FindTagsForPtr(Ptr, Tags))
 	{
-		bEnabled = false;
+		DisableAll();
 		return;
 	}
 
 	Tracker = InTracker;
-	bEnabled = true;
 
-	TrackerData->PushTag(TagData);
+	for (int32 TagSetIndex = 0; TagSetIndex < static_cast<int32>(ELLMTagSet::Max); TagSetIndex++)
+	{
+		if (Tags[TagSetIndex])
+		{
+			bEnabled[TagSetIndex] = true;
+			TrackerData->PushTag(Tags[TagSetIndex], static_cast<ELLMTagSet>(TagSetIndex));
+		}
+		else
+		{
+			bEnabled[TagSetIndex] = false;
+		}
+	}
 }
 
 FLLMScopeFromPtr::~FLLMScopeFromPtr()
 {
-	if (!bEnabled)
+	FLowLevelMemTracker& LLMRef = FLowLevelMemTracker::Get();
+	for (int32 TagSetIndex = 0; TagSetIndex < static_cast<int32>(ELLMTagSet::Max); TagSetIndex++)
 	{
-		return;
+		if (bEnabled[TagSetIndex])
+		{
+			LLMRef.GetTracker(Tracker)->PopTag(static_cast<ELLMTagSet>(TagSetIndex));
+		}
 	}
-
-	FLowLevelMemTracker& LLMRef = FLowLevelMemTracker::Get();
-	LLMRef.GetTracker(Tracker)->PopTag();
 }
 
-FLLMTagDeclaration::FLLMTagDeclaration(const TCHAR* InCPPName, FName InDisplayName, FName InParentTagName,
-	FName InStatName, FName InSummaryStatName)
-	:CPPName(InCPPName), UniqueName(NAME_None), DisplayName(InDisplayName), ParentTagName(InParentTagName),
-	StatName(InStatName), SummaryStatName(InSummaryStatName)
+FLLMTagDeclaration::FLLMTagDeclaration(const TCHAR* InCPPName, FName InDisplayName, FName InParentTagName, FName InStatName, FName InSummaryStatName, ELLMTagSet InTagSet)
+	:CPPName(InCPPName), UniqueName(NAME_None), DisplayName(InDisplayName), ParentTagName(InParentTagName), StatName(InStatName), SummaryStatName(InSummaryStatName)
+	, TagSet(InTagSet)
 {
 	Register();
 }
@@ -2998,24 +3069,25 @@ void FBin::Free(void* Ptr, FLLMAllocator& Allocator)
 
 } // namespace AllocatorPrivate
 
-FTagData::FTagData(FName InName, FName InDisplayName, FName InParentName, FName InStatName,
+FTagData::FTagData(FName InName, ELLMTagSet InTagSet, FName InDisplayName, FName InParentName, FName InStatName,
 	FName InSummaryStatName, bool bInHasEnumTag, ELLMTag InEnumTag, ETagReferenceSource InReferenceSource)
 	: Name(InName), DisplayName(InDisplayName), ParentName(InParentName), StatName(InStatName)
 	, SummaryStatName(InSummaryStatName), EnumTag(InEnumTag), ReferenceSource(InReferenceSource)
-	, bIsFinishConstructed(false), bParentIsName(true), bHasEnumTag(bInHasEnumTag), bIsReportable(true)
+	, TagSet(InTagSet), bIsFinishConstructed(false), bParentIsName(true), bHasEnumTag(bInHasEnumTag)
+	, bIsReportable(InTagSet == ELLMTagSet::None)
 {
 }
 
-FTagData::FTagData(FName InName, FName InDisplayName, const FTagData* InParent, FName InStatName,
+FTagData::FTagData(FName InName, ELLMTagSet InTagSet, FName InDisplayName, const FTagData* InParent, FName InStatName,
 	FName InSummaryStatName, bool bInHasEnumTag, ELLMTag InEnumTag, ETagReferenceSource InReferenceSource)
-	: FTagData(InName, InDisplayName, NAME_None, InStatName, InSummaryStatName, bInHasEnumTag, InEnumTag,
+	: FTagData(InName, InTagSet, InDisplayName, NAME_None, InStatName, InSummaryStatName, bInHasEnumTag, InEnumTag,
 		InReferenceSource)
 {
 	SetParent(InParent);
 }
 
 FTagData::FTagData(ELLMTag InEnumTag)
-	: FTagData(NAME_None, NAME_None, NAME_None, NAME_None, NAME_None, true, InEnumTag,
+	: FTagData(NAME_None, ELLMTagSet::None, NAME_None, NAME_None, NAME_None, NAME_None, true, InEnumTag,
 		ETagReferenceSource::EnumTag)
 {
 }
@@ -3095,6 +3167,11 @@ ELLMTag FTagData::GetEnumTag() const
 	return EnumTag;
 }
 
+ELLMTagSet FTagData::GetTagSet() const
+{
+	return TagSet;
+}
+
 bool FTagData::HasEnumTag() const
 {
 	return bHasEnumTag;
@@ -3256,69 +3333,43 @@ FLLMThreadState* FLLMTracker::GetState()
 	return (FLLMThreadState*)FPlatformTLS::GetTlsValue(TlsSlot);
 }
 
-void FLLMTracker::PushTag(ELLMTag EnumTag)
+void FLLMTracker::PushTag(ELLMTag EnumTag, ELLMTagSet TagSet)
 {
 	const FTagData* TagData = LLMRef.FindOrAddTagData(EnumTag, ETagReferenceSource::Scope);
 
 	// pass along to the state object
-	GetOrCreateState()->PushTag(TagData);
+	GetOrCreateState()->PushTag(TagData, TagSet);
 }
 
-void FLLMTracker::PushTag(FName Tag, bool bInIsStatData)
+void FLLMTracker::PushTag(FName Tag, bool bInIsStatData, ELLMTagSet TagSet)
 {
-	const FTagData* TagData = LLMRef.FindOrAddTagData(Tag, bInIsStatData, ETagReferenceSource::Scope);
+	const FTagData* TagData = LLMRef.FindOrAddTagData(Tag, TagSet, bInIsStatData, ETagReferenceSource::Scope);
 
 	// pass along to the state object
-	GetOrCreateState()->PushTag(TagData);
+	GetOrCreateState()->PushTag(TagData, TagSet);
 }
 
-void FLLMTracker::PushTag(const FTagData* TagData)
+void FLLMTracker::PushTag(const FTagData* TagData, ELLMTagSet TagSet)
 {
 	// pass along to the state object
-	GetOrCreateState()->PushTag(TagData);
+	GetOrCreateState()->PushTag(TagData, TagSet);
 }
 
-void FLLMTracker::PopTag()
+void FLLMTracker::PopTag(ELLMTagSet TagSet)
 {
 	// look for already allocated thread state
 	FLLMThreadState* State = GetState();
 
 	LLMCheckf(State != nullptr, TEXT("Called PopTag but PushTag was never called!"));
 
-	State->PopTag();
-}
-
-#if LLM_ALLOW_ASSETS_TAGS
-void FLLMTracker::PushAssetTag(FName Tag)
-{
-	const FTagData* TagData = LLMRef.FindOrAddTagData(Tag, true /* bIsStatTag */, ETagReferenceSource::Scope);
-
-	// pass along to the state object
-	GetOrCreateState()->PushAssetTag(TagData);
-}
-
-void FLLMTracker::PushAssetTag(const FTagData* TagData)
-{
-	// pass along to the state object
-	GetOrCreateState()->PushAssetTag(TagData);
+	State->PopTag(TagSet);
 }
 
-void FLLMTracker::PopAssetTag()
-{
-	// look for already allocated thread state
-	FLLMThreadState* State = GetState();
-
-	LLMCheckf(State != nullptr, TEXT("Called PopTag but PushTag was never called!"));
-
-	State->PopAssetTag();
-}
-#endif
-
 void FLLMTracker::TrackAllocation(const void* Ptr, int64 Size, ELLMTag DefaultTag, ELLMTracker Tracker,
 	ELLMAllocType AllocType, bool bTrackInMemPro)
 {
 	FLLMThreadState* State = GetOrCreateState();
-	const FTagData* TagData = State->GetTopTag();
+	const FTagData* TagData = State->GetTopTag(ELLMTagSet::None);
 	if (!TagData)
 	{
 		TagData = LLMRef.FindOrAddTagData(DefaultTag);
@@ -3330,10 +3381,10 @@ void FLLMTracker::TrackAllocation(const void* Ptr, int64 Size, FName DefaultTag,
 	ELLMAllocType AllocType, bool bTrackInMemPro)
 {
 	FLLMThreadState* State = GetOrCreateState();
-	const FTagData* TagData = State->GetTopTag();
+	const FTagData* TagData = State->GetTopTag(ELLMTagSet::None);
 	if (!TagData)
 	{
-		TagData = LLMRef.FindOrAddTagData(DefaultTag);
+		TagData = LLMRef.FindOrAddTagData(DefaultTag, ELLMTagSet::None);
 	}
 	TrackAllocation(Ptr, Size, TagData, Tracker, AllocType, State, bTrackInMemPro);
 }
@@ -3358,13 +3409,23 @@ void FLLMTracker::TrackAllocation(const void* Ptr, int64 Size, const FTagData* A
 #endif
 
 #if LLM_ALLOW_ASSETS_TAGS
-	const FTagData* AssetTagData = State->GetTopAssetTag();
+	const FTagData* AssetTagData = State->GetTopTag(ELLMTagSet::Assets);
+	const FTagData* AssetClassTagData = State->GetTopTag(ELLMTagSet::AssetClasses);
+
+	if (AssetTagData == nullptr && !LLMRef.IsBootstrapping() && LLMRef.IsTagSetActive(ELLMTagSet::Assets))
+	{
+		AssetTagData = LLMRef.FindOrAddTagData(TagName_UntaggedAsset, ELLMTagSet::Assets, false, UE::LLMPrivate::ETagReferenceSource::FunctionAPI);
+	}
+	if (AssetClassTagData == nullptr && !LLMRef.IsBootstrapping() && LLMRef.IsTagSetActive(ELLMTagSet::AssetClasses))
+	{
+		AssetClassTagData = LLMRef.FindOrAddTagData(TagName_UntaggedAssetClass, ELLMTagSet::AssetClasses, false, UE::LLMPrivate::ETagReferenceSource::FunctionAPI);
+	}
 #else
 	const FTagData* AssetTagData = nullptr;
+	const FTagData* AssetClassTagData = nullptr;
 #endif
-
 	// track on the thread state
-	State->TrackAllocation(Ptr, Size, Tracker, AllocType, ActiveTagData, AssetTagData, bTrackInMemPro);
+	State->TrackAllocation(Ptr, Size, Tracker, AllocType, ActiveTagData, AssetTagData, AssetClassTagData, bTrackInMemPro);
 
 	// tracking a nullptr with a Size is allowed, but we don't need to remember it, since we can't free it ever.
 	if (Ptr != nullptr)
@@ -3373,7 +3434,8 @@ void FLLMTracker::TrackAllocation(const void* Ptr, int64 Size, const FTagData* A
 		FLLMTracker::FLowLevelAllocInfo AllocInfo;
 		AllocInfo.SetTag(ActiveTagData, LLMRef);
 #if LLM_ALLOW_ASSETS_TAGS
-		AllocInfo.SetAssetTag(AssetTagData, LLMRef);
+		AllocInfo.SetTag(AssetTagData, LLMRef, ELLMTagSet::Assets);
+		AllocInfo.SetTag(AssetClassTagData, LLMRef, ELLMTagSet::AssetClasses);
 #endif
 		LLMCheck(Size <= 0x0000'ffff'ffff'ffff);
 		uint32 SizeLow = uint32(Size);
@@ -3412,12 +3474,14 @@ void FLLMTracker::TrackFree(const void* Ptr, ELLMTracker Tracker, ELLMAllocType
 	FLLMThreadState* State = GetOrCreateState();
 	const FTagData* TagData = AllocInfo.GetTag(LLMRef);
 #if LLM_ALLOW_ASSETS_TAGS
-	const FTagData* AssetTagData = AllocInfo.GetAssetTag(LLMRef);
+	const FTagData* AssetTagData = AllocInfo.GetTag(LLMRef, ELLMTagSet::Assets);
+	const FTagData* AssetClassTagData = AllocInfo.GetTag(LLMRef, ELLMTagSet::AssetClasses);
 #else
 	const FTagData* AssetTagData = nullptr;
+	const FTagData* AssetClassTagData = nullptr;
 #endif
 
-	State->TrackFree(Ptr, Size, Tracker, AllocType, TagData, AssetTagData, bTrackInMemPro);
+	State->TrackFree(Ptr, Size, Tracker, AllocType, TagData, AssetTagData, AssetClassTagData, bTrackInMemPro);
 }
 
 void FLLMTracker::OnAllocMoved(const void* Dest, const void* Source, ELLMTracker Tracker, ELLMAllocType AllocType)
@@ -3452,10 +3516,10 @@ void FLLMTracker::OnAllocMoved(const void* Dest, const void* Source, ELLMTracker
 void FLLMTracker::TrackMemoryOfActiveTag(int64 Amount, FName DefaultTag, ELLMTracker Tracker, ELLMAllocType AllocType)
 {
 	FLLMThreadState* State = GetOrCreateState();
-	const FTagData* TagData = State->GetTopTag();
+	const FTagData* TagData = State->GetTopTag(ELLMTagSet::None);
 	if (!TagData)
 	{
-		TagData = LLMRef.FindOrAddTagData(DefaultTag);
+		TagData = LLMRef.FindOrAddTagData(DefaultTag, ELLMTagSet::None);
 	}
 	TrackMemoryOfActiveTag(Amount, TagData, Tracker, AllocType, State);
 }
@@ -3463,7 +3527,7 @@ void FLLMTracker::TrackMemoryOfActiveTag(int64 Amount, FName DefaultTag, ELLMTra
 void FLLMTracker::TrackMemoryOfActiveTag(int64 Amount, ELLMTag DefaultTag, ELLMTracker Tracker, ELLMAllocType AllocType)
 {
 	FLLMThreadState* State = GetOrCreateState();
-	const FTagData* TagData = State->GetTopTag();
+	const FTagData* TagData = State->GetTopTag(ELLMTagSet::None);
 	if (!TagData)
 	{
 		TagData = LLMRef.FindOrAddTagData(DefaultTag);
@@ -3488,14 +3552,8 @@ void FLLMTracker::TrackMemoryOfActiveTag(int64 Amount, const FTagData* TagData,
 	TagData = TagData->GetContainingEnumTagData();
 #endif
 
-#if LLM_ALLOW_ASSETS_TAGS
-	const FTagData* AssetTagData = State->GetTopAssetTag();
-#else
-	const FTagData* AssetTagData = nullptr;
-#endif
-
 	// track on the thread state
-	State->TrackMemory(Amount, Tracker, AllocType, TagData, AssetTagData);
+	State->TrackMemory(Amount, Tracker, AllocType, TagData);
 }
 
 void FLLMTracker::TrackMemory(ELLMTag Tag, ELLMTracker Tracker, int64 Amount, ELLMAllocType AllocType)
@@ -3503,21 +3561,22 @@ void FLLMTracker::TrackMemory(ELLMTag Tag, ELLMTracker Tracker, int64 Amount, EL
 	TrackMemory(LLMRef.FindOrAddTagData(Tag), Tracker, Amount, AllocType);
 }
 
-void FLLMTracker::TrackMemory(FName Tag, ELLMTracker Tracker, int64 Amount, ELLMAllocType AllocType)
+void FLLMTracker::TrackMemory(FName Tag, ELLMTagSet TagSet, ELLMTracker Tracker, int64 Amount, ELLMAllocType AllocType)
 {
-	TrackMemory(LLMRef.FindOrAddTagData(Tag), Tracker, Amount, AllocType);
+	TrackMemory(LLMRef.FindOrAddTagData(Tag, TagSet, false, UE::LLMPrivate::ETagReferenceSource::FunctionAPI), Tracker, Amount, AllocType);
 }
 
 void FLLMTracker::TrackMemory(const FTagData* TagData, ELLMTracker Tracker, int64 Amount, ELLMAllocType AllocType)
 {
 	FLLMThreadState* State = GetOrCreateState();
 	FPlatformAtomics::InterlockedAdd(&TrackedTotal, Amount);
-	State->TrackMemory(Amount, Tracker, AllocType, TagData, nullptr /* AssetTagData */);
+	State->TrackMemory(Amount, Tracker, AllocType, TagData);
 }
 
-void FLLMTracker::PauseAndTrackMemory(FName TagName, ELLMTracker Tracker, bool bInIsStatTag, int64 Amount, ELLMAllocType AllocType)
+void FLLMTracker::PauseAndTrackMemory(FName TagName, ELLMTagSet TagSet, ELLMTracker Tracker, bool bInIsStatTag, 
+	int64 Amount, ELLMAllocType AllocType)
 {
-	const FTagData* TagData = LLMRef.FindOrAddTagData(TagName, bInIsStatTag);
+	const FTagData* TagData = LLMRef.FindOrAddTagData(TagName, TagSet, bInIsStatTag, UE::LLMPrivate::ETagReferenceSource::FunctionAPI);
 	PauseAndTrackMemory(TagData, Tracker, Amount, AllocType);
 }
 
@@ -3532,7 +3591,7 @@ void FLLMTracker::PauseAndTrackMemory(const FTagData* TagData, ELLMTracker Track
 {
 	FLLMThreadState* State = GetOrCreateState();
 	FPlatformAtomics::InterlockedAdd(&TrackedTotal, Amount);
-	State->TrackMemory(Amount, Tracker, AllocType, TagData, nullptr /* AssetTagData */);
+	State->TrackMemory(Amount, Tracker, AllocType, TagData);
 	FScopeLock Lock(&State->TagSection);
 	State->PausedCounter[static_cast<int32>(AllocType)]++;
 }
@@ -3627,8 +3686,7 @@ void FLLMTracker::Update()
 		{
 			if (AllocationData.bExternalAddToTotal)
 			{
-			FPlatformAtomics::InterlockedAdd(&TrackedTotal,
-				AllocationData.ExternalAmount - AllocationData.Size);
+				FPlatformAtomics::InterlockedAdd(&TrackedTotal, AllocationData.ExternalAmount - AllocationData.Size);
 			}
 			AllocationData.Size = AllocationData.ExternalAmount;
 			AllocationData.bExternalValid = false;
@@ -3676,7 +3734,7 @@ void FLLMTracker::PublishStats(bool bTrackPeaks)
 
 	if (OverrideUntaggedTagData)
 	{
-		const FTagData* TagData = LLMRef.FindTagData(TagName_Untagged);
+		const FTagData* TagData = LLMRef.FindTagData(TagName_Untagged, ELLMTagSet::None);
 		const FTrackerTagSizeData* AllocationData = TagData ? TagSizes.Find(TagData) : nullptr;
 		SetMemoryStatByFName(OverrideUntaggedTagData->GetStatName(),
 			AllocationData ? AllocationData->GetSize(bTrackPeaks) : 0);
@@ -3731,7 +3789,14 @@ void FLLMTracker::OnTagsResorted(FTagDataArray& OldTagDatas)
 		AllocationMap.LockAll();
 		for (const FLLMAllocMap::FTuple& Tuple : AllocationMap)
 		{
-			Tuple.Value2.SetCompressedTag(OldTagDatas[Tuple.Value2.GetCompressedTag()]->GetIndex());
+			for (int32 TagSetIndex = 0; TagSetIndex < static_cast<int32>(ELLMTagSet::Max); TagSetIndex++)
+			{
+				const int32 OldTagIndex = Tuple.Value2.GetCompressedTag(static_cast<ELLMTagSet>(TagSetIndex));
+				if (OldTagIndex >= 0)
+				{
+					Tuple.Value2.SetCompressedTag(OldTagDatas[OldTagIndex]->GetIndex(), static_cast<ELLMTagSet>(TagSetIndex));
+				}
+			}
 		}
 		AllocationMap.UnlockAll();
 	}
@@ -3765,34 +3830,94 @@ void FLLMTracker::LockAllThreadTags(bool bLock)
 	}
 }
 
-const FTagData* FLLMTracker::GetActiveTagData()
+const FTagData* FLLMTracker::GetActiveTagData(ELLMTagSet TagSet)
 {
 	FLLMThreadState* State = GetOrCreateState();
-	return State->GetTopTag();
+	return State->GetTopTag(TagSet);
 }
 
-const FTagData* FLLMTracker::FindTagForPtr( void* Ptr )
+TArray<const FTagData*> FLLMTracker::GetTagDatas(ELLMTagSet TagSet)
 {
-	FLLMThreadState* State = GetOrCreateState();
-	FLowLevelAllocInfo AllocInfo;
+	TArray<const FTagData*> FoundTagDatas;
+	TagSizes.GetKeys(FoundTagDatas);
+
+	return FoundTagDatas.FilterByPredicate([TagSet](const FTagData* InTagData) {
+		return InTagData != nullptr && InTagData->GetTagSet() == TagSet;
+	});
+}
+
+void FLLMTracker::GetTagsNamesWithAmount(TMap<FName, uint64>& OutTagsNamesWithAmount, ELLMTagSet TagSet /* = ELLMTagSet::None */)
+{
+	for (const TPair<const FTagData*, FTrackerTagSizeData>& It : TagSizes)
 	{
-		uint32* Size;
-		FLowLevelAllocInfo* AllocInfoPtr;
-		AllocationMap.Find(PointerKey(Ptr), Size, AllocInfoPtr);
-		if (!AllocInfoPtr)
+		const FTagData* TagData = It.Key;
+		if (TagData->GetTagSet() == TagSet)
 		{
-			return nullptr;
+			OutTagsNamesWithAmount.Add(TagData->GetName(), It.Value.GetSize(false));
 		}
-		AllocInfo = *AllocInfoPtr;
 	}
-	return AllocInfo.GetTag(LLMRef);
 }
 
-TArray<const FTagData*> FLLMTracker::GetTagDatas()
+void FLLMTracker::GetTagsNamesWithAmountFiltered(TMap<FName, uint64>& OutTagsNamesWithAmount, ELLMTagSet TagSet /* = ELLMTagSet::None */, TArray<FLLMTagSetAllocationFilter>& Filters)
 {
-	TArray<const FTagData*> FoundTagDatas;
-	TagSizes.GetKeys(FoundTagDatas);
-	return FoundTagDatas;
+	AllocationMap.LockAll();
+	for (const FLLMAllocMap::FTuple& Tuple : AllocationMap)
+	{
+		bool bIncludeAllocation = true;
+		for (const FLLMTagSetAllocationFilter& Filter : Filters)
+		{
+#if LLM_ALLOW_ASSETS_TAGS
+			const FTagData* TagData = Tuple.Value2.GetTag(LLMRef, Filter.TagSet);
+#else
+			const FTagData* TagData = Tuple.Value2.GetTag(LLMRef);
+#endif
+			if (TagData == nullptr || TagData->GetName() != Filter.Name)
+			{
+				bIncludeAllocation = false;
+				break;
+			}
+		}
+			
+		if (bIncludeAllocation)
+		{
+#if LLM_ALLOW_ASSETS_TAGS
+			const FTagData* TagData = Tuple.Value2.GetTag(LLMRef, TagSet);
+#else
+			const FTagData* TagData = Tuple.Value2.GetTag(LLMRef);
+#endif
+			if (TagData != nullptr)
+			{
+				uint64& Size = OutTagsNamesWithAmount.FindOrAdd(TagData->GetName());
+				Size += Tuple.Value1;
+			}
+		}
+	}
+	AllocationMap.UnlockAll();
+}
+
+bool FLLMTracker::FindTagsForPtr(void* InPtr, TArray<const FTagData *, TInlineAllocator<static_cast<int32>(ELLMTagSet::Max)>>& OutTags)
+{
+	FLLMThreadState* State = GetOrCreateState();
+		
+	uint32* Size;
+	FLowLevelAllocInfo* AllocInfoPtr;
+	AllocationMap.Find(PointerKey(InPtr), Size, AllocInfoPtr);
+	if (!AllocInfoPtr)
+	{
+		return false;
+	}
+
+	OutTags.SetNumUnsafeInternal(static_cast<int32>(ELLMTagSet::Max));
+#if LLM_ALLOW_ASSETS_TAGS
+	for (int32 TagSetIndex = 0; TagSetIndex < static_cast<int32>(ELLMTagSet::Max); TagSetIndex++)
+	{
+		OutTags[TagSetIndex] = AllocInfoPtr->GetTag(LLMRef, static_cast<ELLMTagSet>(TagSetIndex));
+	}
+#else
+	OutTags[static_cast<int32>(ELLMTagSet::None)] = AllocInfoPtr->GetTag(LLMRef);
+#endif
+
+	return true;
 }
 
 int64 FLLMTracker::GetTagAmount(const FTagData* TagData, bool bPeakAmount/* = false*/) const
@@ -3846,61 +3971,36 @@ FLLMThreadState::FLLMThreadState()
 
 void FLLMThreadState::Clear()
 {
-	TagStack.Empty();
-#if LLM_ALLOW_ASSETS_TAGS
-	AssetTagStack.Empty();
-#endif
+	for (int32 TagSetIndex = 0; TagSetIndex < static_cast<int32>(ELLMTagSet::Max); TagSetIndex++)
+	{
+		TagStack[TagSetIndex].Empty();
+	}
 	Allocations.Empty();
 	ClearAllocTypeAmounts();
 }
 
-void FLLMThreadState::PushTag(const FTagData* TagData)
+void FLLMThreadState::PushTag(const FTagData* TagData, ELLMTagSet TagSet)
 {
 	FScopeLock Lock(&TagSection);
 
 	// Push a tag.
-	TagStack.Add(TagData);
+	TagStack[static_cast<int32>(TagSet)].Add(TagData);
 }
 
-void FLLMThreadState::PopTag()
+void FLLMThreadState::PopTag(ELLMTagSet TagSet)
 {
 	FScopeLock Lock(&TagSection);
 
-	LLMCheckf(TagStack.Num() > 0,
+	LLMCheckf(TagStack[static_cast<int32>(TagSet)].Num() > 0,
 		TEXT("Called FLLMThreadState::PopTag without a matching Push (stack was empty on pop)"));
-	TagStack.Pop(false /* bAllowShrinking */);
-}
-
-const FTagData* FLLMThreadState::GetTopTag()
-{
-	return TagStack.Num() ? TagStack.Last() : nullptr;
-}
-
-#if LLM_ALLOW_ASSETS_TAGS
-void FLLMThreadState::PushAssetTag(const FTagData* TagData)
-{
-	FScopeLock Lock(&TagSection);
-
-	// Push a tag.
-	AssetTagStack.Add(TagData);
+	TagStack[static_cast<int32>(TagSet)].Pop(false /* bAllowShrinking */);
 }
 
-void FLLMThreadState::PopAssetTag()
+const FTagData* FLLMThreadState::GetTopTag(ELLMTagSet TagSet)
 {
-	FScopeLock Lock(&TagSection);
-
-	LLMCheckf(AssetTagStack.Num() > 0,
-		TEXT("Called FLLMThreadState::PopTag without a matching Push (stack was empty on pop)"));
-	AssetTagStack.Pop(false /* bAllowShrinking */);
+	return TagStack[static_cast<int32>(TagSet)].Num() ? TagStack[static_cast<int32>(TagSet)].Last() : nullptr;
 }
 
-const FTagData* FLLMThreadState::GetTopAssetTag()
-{
-	FScopeLock Lock(&TagSection);
-	return AssetTagStack.Num() ? AssetTagStack.Last() : nullptr;
-}
-#endif
-
 void FLLMThreadState::IncrTag(const FTagData* TagData, int64 Amount)
 {
 	// Caller is responsible for holding a lock on TagSection.
@@ -3909,8 +4009,8 @@ void FLLMThreadState::IncrTag(const FTagData* TagData, int64 Amount)
 	AllocationSize.Size += Amount;
 }
 
-void FLLMThreadState::TrackAllocation(const void* Ptr, int64 Size, ELLMTracker Tracker, ELLMAllocType AllocType,
-	const FTagData* TagData, const FTagData* AssetTagData, bool bTrackInMemPro)
+void FLLMThreadState::TrackAllocation(const void* Ptr, int64 Size, ELLMTracker Tracker, ELLMAllocType AllocType, 
+	const FTagData* TagData, const FTagData* AssetTagData, const FTagData* AssetClassTagData, bool bTrackInMemPro)
 {
 	FScopeLock Lock(&TagSection);
 
@@ -3922,6 +4022,10 @@ void FLLMThreadState::TrackAllocation(const void* Ptr, int64 Size, ELLMTracker T
 	{
 		IncrTag(AssetTagData, Size);
 	}
+	if (AssetClassTagData)
+	{
+		IncrTag(AssetClassTagData, Size);
+	}
 #endif
 
 	ELLMTag EnumTag = TagData->GetContainingEnum();
@@ -3938,8 +4042,8 @@ void FLLMThreadState::TrackAllocation(const void* Ptr, int64 Size, ELLMTracker T
 #endif
 }
 
-void FLLMThreadState::TrackFree(const void* Ptr, int64 Size, ELLMTracker Tracker, ELLMAllocType AllocType,
-	const FTagData* TagData, const FTagData* AssetTagData, bool bTrackInMemPro)
+void FLLMThreadState::TrackFree(const void* Ptr, int64 Size, ELLMTracker Tracker, ELLMAllocType AllocType, 
+	const FTagData* TagData, const FTagData* AssetTagData, const FTagData* AssetClassTagData, bool bTrackInMemPro)
 {
 	FScopeLock Lock(&TagSection);
 
@@ -3951,6 +4055,10 @@ void FLLMThreadState::TrackFree(const void* Ptr, int64 Size, ELLMTracker Tracker
 	{
 		IncrTag(AssetTagData, -Size);
 	}
+	if (AssetClassTagData)
+	{
+		IncrTag(AssetClassTagData, -Size);
+	}
 #endif
 	ELLMTag EnumTag = TagData->GetContainingEnum();
 	if (Tracker == ELLMTracker::Default)
@@ -3966,17 +4074,11 @@ void FLLMThreadState::TrackFree(const void* Ptr, int64 Size, ELLMTracker Tracker
 #endif
 }
 
-void FLLMThreadState::TrackMemory(int64 Amount, ELLMTracker Tracker, ELLMAllocType AllocType, const FTagData* TagData, const FTagData* AssetTagData)
+void FLLMThreadState::TrackMemory(int64 Amount, ELLMTracker Tracker, ELLMAllocType AllocType, const FTagData* TagData)
 {
 	FScopeLock Lock(&TagSection);
 	AllocTypeAmounts[static_cast<int32>(AllocType)] += Amount;
 	IncrTag(TagData, Amount);
-#if LLM_ALLOW_ASSETS_TAGS
-	if (AssetTagData)
-	{
-		IncrTag(AssetTagData, Size);
-	}
-#endif
 
 	// TODO: Need to expose TrackMemory to Platform-specific trackers and FMemPropProfiler
 }
@@ -4064,7 +4166,7 @@ void FLLMThreadState::LockTags(bool bLock)
 	}
 }
 
-void FLLMThreadState::FetchAndClearTagSizes(FTrackerTagSizeMap& TagSizes, int64* InAllocTypeAmounts,
+void FLLMThreadState::FetchAndClearTagSizes(FTrackerTagSizeMap& TagSizes, int64* InAllocTypeAmounts, 
 	bool bTrimAllocations)
 {
 	FScopeLock Lock(&TagSection);
@@ -4323,7 +4425,7 @@ void FLLMCsvWriter::AddRow(FLowLevelMemTracker& LLMRef, const FTrackerTagSizeMap
 	}
 	if (OverrideUntaggedTagData)
 	{
-		WriteTag(LLMRef.FindTagData(TagName_Untagged));
+		WriteTag(LLMRef.FindTagData(TagName_Untagged, ELLMTagSet::None));
 	}
 
 	for (const FTagData* TagData : Columns)
@@ -4442,7 +4544,7 @@ void FLLMTraceWriter::Publish(FLowLevelMemTracker& LLMRef, const FTrackerTagSize
 	AddValue(OverrideTrackedTotalTagData, TrackedTotal);
 	if (OverrideUntaggedTagData)
 	{
-		const FTagData* TagData = LLMRef.FindTagData(TagName_Untagged);
+		const FTagData* TagData = LLMRef.FindTagData(TagName_Untagged, ELLMTagSet::None);
 		if (!TagData)
 		{
 			AddValue(OverrideUntaggedTagData, 0);
@@ -4526,7 +4628,7 @@ void FLLMCsvProfilerWriter::Publish(FLowLevelMemTracker& LLMRef, const FTrackerT
 
 	if (OverrideUntaggedTagData)
 	{
-		const FTagData* TagData = LLMRef.FindTagData(TagName_Untagged);
+		const FTagData* TagData = LLMRef.FindTagData(TagName_Untagged, ELLMTagSet::None);
 		const FTrackerTagSizeData* AllocationData = TagData ? TagSizes.Find(TagData) : nullptr;
 		RecordTagToCsv(CsvCategoryIndex, OverrideUntaggedTagData,
 			AllocationData ? AllocationData->GetSize(bTrackPeaks) : 0);
diff --git a/Engine/Source/Runtime/Core/Private/HAL/LowLevelMemTrackerPrivate.h b/Engine/Source/Runtime/Core/Private/HAL/LowLevelMemTrackerPrivate.h
--- a/Engine/Source/Runtime/Core/Private/HAL/LowLevelMemTrackerPrivate.h
+++ b/Engine/Source/Runtime/Core/Private/HAL/LowLevelMemTrackerPrivate.h
@@ -22,7 +22,7 @@
 #endif
 
 // Enable storing the full tag for each alloc if (1) arbitrary names tags are allowed or (2) Stats are allowed (stats use names tags) or (3) Asset tags are allowed (assets use names tags)
-#define LLM_ENABLED_FULL_TAGS LLM_ALLOW_NAMES_TAGS || LLM_ALLOW_STATS || LLM_ALLOW_ASSETS_TAGS
+#define LLM_ENABLED_FULL_TAGS (LLM_ALLOW_NAMES_TAGS || LLM_ALLOW_STATS || LLM_ALLOW_ASSETS_TAGS)
 
 // Whether to enable running with reduced threads. This is currently enabled because the engine crashes with -norenderthread
 #define LLM_ENABLED_REDUCE_THREADS 0
@@ -50,6 +50,36 @@ namespace UE
 {
 namespace LLMPrivate
 {
+	struct FTagDataNameKey
+	{
+		FName Name;
+		ELLMTagSet TagSet;
+		
+		FTagDataNameKey(FName InName, ELLMTagSet InTagSet) :
+			Name(InName),
+			TagSet(InTagSet)
+		{
+		}
+
+		friend uint32 GetTypeHash(const FTagDataNameKey& Key)
+		{
+			constexpr uint32 HashPrime = 101;
+			return GetTypeHash(Key.Name) + HashPrime * static_cast<uint32>(Key.TagSet);
+		}
+
+		friend bool operator==(const FTagDataNameKey& A, const FTagDataNameKey& B)
+		{
+			return A.Name == B.Name && A.TagSet == B.TagSet;
+		}
+
+		friend bool operator!=(const FTagDataNameKey& A, const FTagDataNameKey& B)
+		{
+			return (A.Name != B.Name || A.TagSet != B.TagSet);
+		}
+
+		FTagDataNameKey() = delete;
+	};
+
 	/**
 	 * FTagData: Description of the properties of a Tag that can be used in LLM_SCOPE
 	 */
@@ -57,8 +87,8 @@ namespace LLMPrivate
 	{
 	public:
 
-		FTagData(FName InName, FName InDisplayName, FName InParentName, FName InStatName, FName InSummaryStatName, bool bInHasEnumTag, ELLMTag InEnumTag, ETagReferenceSource InReferenceSource);
-		FTagData(FName InName, FName InDisplayName, const FTagData* InParent, FName InStatName, FName InSummaryStatName, bool bInHasEnumTag, ELLMTag InEnumTag, ETagReferenceSource InReferenceSource);
+		FTagData(FName InName, ELLMTagSet InTagSet, FName InDisplayName, FName InParentName, FName InStatName, FName InSummaryStatName, bool bInHasEnumTag, ELLMTag InEnumTag, ETagReferenceSource InReferenceSource);
+		FTagData(FName InName, ELLMTagSet InTagSet, FName InDisplayName, const FTagData* InParent, FName InStatName, FName InSummaryStatName, bool bInHasEnumTag, ELLMTag InEnumTag, ETagReferenceSource InReferenceSource);
 		~FTagData();
 
 		bool IsParentConstructed() const;
@@ -72,6 +102,7 @@ namespace LLMPrivate
 		FName GetStatName() const;
 		FName GetSummaryStatName() const;
 		ELLMTag GetEnumTag() const;
+		ELLMTagSet GetTagSet() const;
 		bool HasEnumTag() const;
 		const FTagData* GetContainingEnumTagData() const;
 		ELLMTag GetContainingEnum() const;
@@ -107,6 +138,7 @@ namespace LLMPrivate
 		int32 Index;
 		ELLMTag EnumTag;
 		ETagReferenceSource ReferenceSource;
+		ELLMTagSet TagSet;
 		bool bIsFinishConstructed;
 		bool bParentIsName;
 		bool bHasEnumTag;
@@ -114,9 +146,9 @@ namespace LLMPrivate
 	};
 
 	// TagData container types that use FLLMAllocator
-	class FTagDataNameMap : public TMap<FName, FTagData*, FDefaultSetLLMAllocator>
+	class FTagDataNameMap : public TMap<FTagDataNameKey, FTagData*, FDefaultSetLLMAllocator>
 	{
-		using TMap<FName, FTagData*, FDefaultSetLLMAllocator>::TMap;
+		using TMap<FTagDataNameKey, FTagData*, FDefaultSetLLMAllocator>::TMap;
 	};
 	class FConstTagDataArray : public TArray<const FTagData*, FDefaultLLMAllocator>
 	{
diff --git a/Engine/Source/Runtime/Core/Public/Async/InheritedContext.h b/Engine/Source/Runtime/Core/Public/Async/InheritedContext.h
--- a/Engine/Source/Runtime/Core/Public/Async/InheritedContext.h
+++ b/Engine/Source/Runtime/Core/Public/Async/InheritedContext.h
@@ -8,6 +8,44 @@
 
 namespace UE
 {
+#if ENABLE_LOW_LEVEL_MEM_TRACKER
+	/** Structure representing the captured LLM Tags */
+	struct FLLMActiveTagsCapture
+	{
+		/** Fixed array of LLMTagSets captured */
+		const UE::LLMPrivate::FTagData* LLMTags[static_cast<uint32>(ELLMTagSet::Max)];
+
+		void CaptureActiveTagData()
+		{
+			for (int32 TagSetIndex = 0; TagSetIndex < static_cast<int32>(ELLMTagSet::Max); ++TagSetIndex)
+			{
+				LLMTags[TagSetIndex] = FLowLevelMemTracker::bIsDisabled ? nullptr : FLowLevelMemTracker::Get().GetActiveTagData(ELLMTracker::Default, static_cast<ELLMTagSet>(TagSetIndex));
+			}
+		}
+	};
+
+	/** Structure holding the captured LLM scopes */
+	struct FLLMActiveTagsScope
+	{
+		/** Fixed size array which holds a copy of the LLM scopes for later recall */
+		/** Order doesn't matter so we can use a FixedAllocator and not require a default constructor */
+		TArray<FLLMScope, TFixedAllocator<static_cast<uint32>(ELLMTagSet::Max)>> LLMScopes;
+
+		explicit FLLMActiveTagsScope(const FLLMActiveTagsCapture& InActiveTagsCapture)
+		{
+			CaptureLLMScopes(InActiveTagsCapture);
+		}
+
+		void CaptureLLMScopes(const FLLMActiveTagsCapture& InActiveTagsCapture)
+		{
+			for (int32 TagSetIndex = 0; TagSetIndex < static_cast<int32>(ELLMTagSet::Max); ++TagSetIndex)
+			{
+				LLMScopes.Emplace(InActiveTagsCapture.LLMTags[TagSetIndex], false /* bIsStatTag */, static_cast<ELLMTagSet>(TagSetIndex), ELLMTracker::Default);
+			}
+		}
+	};
+#endif
+
 	// Restores an inherited contex for the current scope.
 	// An instance must be obtained by calling `FInheritedContextBase::RestoreInheritedContext()`
 	class FInheritedContextScope
@@ -19,7 +57,7 @@ namespace UE
 
 		FInheritedContextScope(
 #if ENABLE_LOW_LEVEL_MEM_TRACKER
-			const UE::LLMPrivate::FTagData* InInheritedLLMTag
+			const FLLMActiveTagsCapture& InInheritedLLMTag
 	#if UE_MEMORY_TAGS_TRACE_ENABLED && UE_TRACE_ENABLED
 			,
 	#endif
@@ -29,7 +67,7 @@ namespace UE
 #endif
 		)
 #if ENABLE_LOW_LEVEL_MEM_TRACKER
-			: LLMScope(InInheritedLLMTag, false /* bIsStatTag */, ELLMTagSet::None, ELLMTracker::Default)
+			: LLMScopes(InInheritedLLMTag)
 	#if UE_MEMORY_TAGS_TRACE_ENABLED && UE_TRACE_ENABLED
 			,
 	#endif
@@ -45,7 +83,7 @@ namespace UE
 
 	private:
 #if ENABLE_LOW_LEVEL_MEM_TRACKER
-		FLLMScope LLMScope;
+		FLLMActiveTagsScope LLMScopes;
 #endif
 
 #if UE_MEMORY_TAGS_TRACE_ENABLED && UE_TRACE_ENABLED
@@ -62,7 +100,7 @@ namespace UE
 		void CaptureInheritedContext()
 		{
 #if ENABLE_LOW_LEVEL_MEM_TRACKER
-			InheritedLLMTag = FLowLevelMemTracker::bIsDisabled ? nullptr : FLowLevelMemTracker::Get().GetActiveTagData(ELLMTracker::Default);
+			InheritedLLMTags.CaptureActiveTagData();
 #endif
 
 #if UE_MEMORY_TAGS_TRACE_ENABLED && UE_TRACE_ENABLED
@@ -79,7 +117,7 @@ namespace UE
 
 	private:
 #if ENABLE_LOW_LEVEL_MEM_TRACKER
-		const UE::LLMPrivate::FTagData* InheritedLLMTag;
+		FLLMActiveTagsCapture InheritedLLMTags;
 #endif
 
 #if UE_MEMORY_TAGS_TRACE_ENABLED && UE_TRACE_ENABLED
diff --git a/Engine/Source/Runtime/Core/Public/HAL/LowLevelMemStats.h b/Engine/Source/Runtime/Core/Public/HAL/LowLevelMemStats.h
--- a/Engine/Source/Runtime/Core/Public/HAL/LowLevelMemStats.h
+++ b/Engine/Source/Runtime/Core/Public/HAL/LowLevelMemStats.h
@@ -64,7 +64,7 @@
 #define LLM_SCOPED_SINGLE_STAT_TAG(Stat) DECLARE_LLM_MEMORY_STAT(TEXT(#Stat), Stat, STATGROUP_LLMFULL); LLM_SCOPED_TAG_WITH_STAT(Stat, ELLMTracker::Default);
 #define LLM_SCOPED_SINGLE_STAT_TAG_IN_SET(Stat, Set) DECLARE_LLM_MEMORY_STAT(TEXT(#Stat), Stat, STATGROUP_LLMFULL); LLM_SCOPED_TAG_WITH_STAT_IN_SET(Stat, Set, ELLMTracker::Default);
 #define LLM_SCOPED_PAUSE_TRACKING_WITH_STAT_AND_AMOUNT(Stat, Amount, Tracker) FLLMPauseScope SCOPE_NAME(GET_STATFNAME(Stat), true /* bIsStatTag */, Amount, Tracker, ELLMAllocType::None);
-#define LLM_SCOPED_TAG_WITH_OBJECT_IN_SET(Object, Set) LLM_SCOPED_TAG_WITH_STAT_NAME_IN_SET(FLowLevelMemTracker::Get().IsTagSetActive(Set) ? (FDynamicStats::CreateMemoryStatId<FStatGroup_STATGROUP_LLMAssets>(FName(*(Object)->GetFullName())).GetName()) : NAME_None, Set, ELLMTracker::Default);
+#define LLM_SCOPED_TAG_WITH_OBJECT_IN_SET(Object, Set) FLLMScope SCOPE_NAME(FLowLevelMemTracker::Get().IsTagSetActive(Set) ? Object->GetOutermost()->GetFName() : NAME_None, (Set == ELLMTagSet::None) /* bIsStatTag */, Set, ELLMTracker::Default);
 
 // special stat pushing to update threads after each asset when tracking assets
 // Currently this is unused, but we may use it for optimizations later
diff --git a/Engine/Source/Runtime/Core/Public/HAL/LowLevelMemTracker.h b/Engine/Source/Runtime/Core/Public/HAL/LowLevelMemTracker.h
--- a/Engine/Source/Runtime/Core/Public/HAL/LowLevelMemTracker.h
+++ b/Engine/Source/Runtime/Core/Public/HAL/LowLevelMemTracker.h
@@ -312,11 +312,15 @@ extern FName LLMGetTagStat(ELLMTag Tag);
  */
 #define LLM_SCOPE(Tag)												FLLMScope SCOPE_NAME(Tag, false /* bIsStatTag */, ELLMTagSet::None, ELLMTracker::Default);\
 																	UE_MEMSCOPE(Tag) 
+#define LLM_TAGSET_SCOPE(Tag, TagSet)								FLLMScope SCOPE_NAME(Tag, false /* bIsStaTag */, TagSet, ELLMTracker::Default);
 #define LLM_SCOPE_BYNAME(Tag) 										static FName PREPROCESSOR_JOIN(LLMScope_Name,__LINE__)(Tag);\
 																	FLLMScope SCOPE_NAME(PREPROCESSOR_JOIN(LLMScope_Name,__LINE__), false /* bIsStatTag */, ELLMTagSet::None, ELLMTracker::Default);\
 																	UE_MEMSCOPE(PREPROCESSOR_JOIN(LLMScope_Name,__LINE__));
 #define LLM_SCOPE_BYTAG(TagDeclName)								FLLMScope SCOPE_NAME(PREPROCESSOR_JOIN(LLMTagDeclaration_, TagDeclName).GetUniqueName(), false /* bIsStatTag */, ELLMTagSet::None, ELLMTracker::Default);\
 																	UE_MEMSCOPE(PREPROCESSOR_JOIN(LLMTagDeclaration_,TagDeclName).GetUniqueName());
+#define LLM_SCOPE_RENDER_RESOURCE(Tag)								static const FString PREPROCESSOR_JOIN(LLMScope_NamePrefix,__LINE__)(TEXT("RenderResources.")); \
+																	FName PREPROCESSOR_JOIN(LLMScope_Name,__LINE__)(PREPROCESSOR_JOIN(LLMScope_NamePrefix, __LINE__) + (Tag ? Tag : TEXT("Unknown"))); \
+																	FLLMScope SCOPE_NAME(PREPROCESSOR_JOIN(LLMScope_Name,__LINE__), false /* bIsStatTag */, ELLMTagSet::Assets, ELLMTracker::Default, false /* bOverride */);
 #define LLM_PLATFORM_SCOPE(Tag)										FLLMScope SCOPE_NAME(Tag, false /* bIsStatTag */, ELLMTagSet::None, ELLMTracker::Platform);
 #define LLM_PLATFORM_SCOPE_BYNAME(Tag) 								static FName PREPROCESSOR_JOIN(LLMScope_Name,__LINE__)(Tag);\
 																	FLLMScope SCOPE_NAME(PREPROCESSOR_JOIN(LLMLLMScope_NameScope,__LINE__), false /* bIsStatTag */, ELLMTagSet::None, ELLMTracker::Platform);
@@ -391,6 +395,8 @@ typedef void(*LLMFreeFunction)(void*, size_t);
 
 class FLLMTagDeclaration;
 
+struct FLLMTagSetAllocationFilter;
+
 namespace UE::LLMPrivate
 {
 
@@ -594,20 +600,23 @@ public:
 	int64 GetActiveTag(ELLMTracker Tracker);
 
 	/** Get an opaque identifier for the top active tag for the given tracker. */
-	const UE::LLMPrivate::FTagData* GetActiveTagData(ELLMTracker Tracker);
+	const UE::LLMPrivate::FTagData* GetActiveTagData(ELLMTracker Tracker, ELLMTagSet TagSet = ELLMTagSet::None);
 
 	/** Register custom ELLMTags. */
 	void RegisterPlatformTag(int32 Tag, const TCHAR* Name, FName StatName, FName SummaryStatName, int32 ParentTag = -1);
 	void RegisterProjectTag(int32 Tag, const TCHAR* Name, FName StatName, FName SummaryStatName, int32 ParentTag = -1);
     
 	/** Get all tags being tracked. */
-	TArray<const UE::LLMPrivate::FTagData*> GetTrackedTags();
+	TArray<const UE::LLMPrivate::FTagData*> GetTrackedTags(ELLMTagSet TagSet = ELLMTagSet::None);
 
 	/** Get all tags being tracked by the given tracker. */
-	TArray<const UE::LLMPrivate::FTagData*> GetTrackedTags(ELLMTracker Tracker);
+	TArray<const UE::LLMPrivate::FTagData*> GetTrackedTags(ELLMTracker Tracker, ELLMTagSet TagSet = ELLMTagSet::None);
+
+	void GetTrackedTagsNamesWithAmount(TMap<FName, uint64>& TagsNamesWithAmount, ELLMTracker Tracker, ELLMTagSet TagSet);
+	void GetTrackedTagsNamesWithAmountFiltered(TMap<FName, uint64>& TagsNamesWithAmount, ELLMTracker Tracker, ELLMTagSet TagSet, TArray<FLLMTagSetAllocationFilter>& Filters);
 
 	/** Look up the ELLMTag associated with the given display name. */
-	bool FindTagByName(const TCHAR* Name, uint64& OutTag) const;
+	bool FindTagByName(const TCHAR* Name, uint64& OutTag, ELLMTagSet InTagSet = ELLMTagSet::None) const;
 
 	UE_DEPRECATED(4.27, "Use FindTagDisplayName instead")
 	const TCHAR* FindTagName(uint64 Tag) const;
@@ -630,6 +639,8 @@ public:
 	/** Get the amount of memory for a FTagData from the given tracker. */
 	int64 GetTagAmountForTracker(ELLMTracker Tracker, const UE::LLMPrivate::FTagData* TagData, bool bPeakAmount = false);
 
+    int64 GetTagAmountForTracker(ELLMTracker Tracker, FName Tag, ELLMTagSet TagSet, bool bPeakAmount = false);
+
 	/** Set the amount of memory for an ELLMTag for a given tracker and optionally update the total tracked memory. */
 	void SetTagAmountForTracker(ELLMTracker Tracker, ELLMTag Tag, int64 Amount, bool bAddToTotal);
 
@@ -661,6 +672,8 @@ private:
 	 */
 	void BootstrapInitialise();
 
+	bool IsBootstrapping() const { return bIsBootstrapping; }
+
 	/** Free all memory. This will put the tracker into a permanently disabled state. */
 	void Clear();
 	void InitialiseProgramSize();
@@ -672,7 +685,7 @@ private:
 	void SortTags(UE::LLMPrivate::FTagDataArray*& OutOldTagDatas);
 	void PublishDataPerFrame(const TCHAR* LogName);
 
-	void RegisterCustomTagInternal(int32 Tag, const TCHAR* Name, FName StatName, FName SummaryStatName, int32 ParentTag = -1);
+	void RegisterCustomTagInternal(int32 Tag, ELLMTagSet TagSet, const TCHAR* Name, FName StatName, FName SummaryStatName, int32 ParentTag = -1);
 	/**
 	* Called during C++ global static initialization when GMalloc is available (and hence FNames are unavailable)
 	* Creates the subset of tags necessary to record allocations during GMalloc and FName construction
@@ -685,15 +698,15 @@ private:
 	void InitialiseTagDatas();
 	void ClearTagDatas();
 	void RegisterTagDeclaration(FLLMTagDeclaration& TagDeclaration);
-	UE::LLMPrivate::FTagData& RegisterTagData(FName Name, FName DisplayName, FName ParentName, FName StatName, FName SummaryStatName, bool bHasEnumTag, ELLMTag EnumTag, bool bIsStatTag, UE::LLMPrivate::ETagReferenceSource ReferenceSource);
+	UE::LLMPrivate::FTagData& RegisterTagData(FName Name, FName DisplayName, FName ParentName, FName StatName, FName SummaryStatName, bool bHasEnumTag, ELLMTag EnumTag, bool bIsStatTag, UE::LLMPrivate::ETagReferenceSource ReferenceSource, ELLMTagSet TagSet = ELLMTagSet::None);
 	/** Construct if not yet done the data on the given FTagData that relies on the presence of other TagDatas */
 	void FinishConstruct(UE::LLMPrivate::FTagData* TagData, UE::LLMPrivate::ETagReferenceSource ReferenceSource);
 	void ReportDuplicateTagName(UE::LLMPrivate::FTagData* TagData, UE::LLMPrivate::ETagReferenceSource ReferenceSource);
 
 	const UE::LLMPrivate::FTagData* FindOrAddTagData(ELLMTag EnumTag, UE::LLMPrivate::ETagReferenceSource ReferenceSource = UE::LLMPrivate::ETagReferenceSource::FunctionAPI);
-	const UE::LLMPrivate::FTagData* FindOrAddTagData(FName Name, bool bIsStatData=false, UE::LLMPrivate::ETagReferenceSource ReferenceSource = UE::LLMPrivate::ETagReferenceSource::FunctionAPI);
+	const UE::LLMPrivate::FTagData* FindOrAddTagData(FName Name, ELLMTagSet TagSet, bool bIsStatData=false, UE::LLMPrivate::ETagReferenceSource ReferenceSource = UE::LLMPrivate::ETagReferenceSource::FunctionAPI);
 	const UE::LLMPrivate::FTagData* FindTagData(ELLMTag EnumTag, UE::LLMPrivate::ETagReferenceSource ReferenceSource = UE::LLMPrivate::ETagReferenceSource::FunctionAPI);
-	const UE::LLMPrivate::FTagData* FindTagData(FName Name, UE::LLMPrivate::ETagReferenceSource ReferenceSource = UE::LLMPrivate::ETagReferenceSource::FunctionAPI);
+	const UE::LLMPrivate::FTagData* FindTagData(FName Name, ELLMTagSet TagSet, UE::LLMPrivate::ETagReferenceSource ReferenceSource = UE::LLMPrivate::ETagReferenceSource::FunctionAPI);
 
 	friend class FLLMPauseScope;
 	friend class FLLMScope;
@@ -745,27 +758,27 @@ public: // really internal but needs to be visible for LLM_IF_ENABLED macro
 class CORE_API FLLMScope
 {
 public:
-	FLLMScope(FName TagName, bool bIsStatTag, ELLMTagSet InTagSet, ELLMTracker InTracker)
+	FLLMScope(FName TagName, bool bIsStatTag, ELLMTagSet InTagSet, ELLMTracker InTracker, bool bOverride = true)
 	{
 		if (!FLowLevelMemTracker::bIsDisabled)
 		{
-			Init(TagName, bIsStatTag, InTagSet, InTracker);
+			Init(TagName, bIsStatTag, InTagSet, InTracker, bOverride);
 		}
 	}
-	FLLMScope(ELLMTag TagEnum, bool bIsStatTag, ELLMTagSet InTagSet, ELLMTracker InTracker)
+	FLLMScope(ELLMTag TagEnum, bool bIsStatTag, ELLMTagSet InTagSet, ELLMTracker InTracker, bool bOverride = true)
 	{
 		if (!FLowLevelMemTracker::bIsDisabled)
 		{
-			Init(TagEnum, bIsStatTag, InTagSet, InTracker);
+			Init(TagEnum, bIsStatTag, InTagSet, InTracker, bOverride);
 		}
 
 	}
 
-	FLLMScope(const UE::LLMPrivate::FTagData* TagData, bool bIsStatTag, ELLMTagSet Set, ELLMTracker Tracker)
+	FLLMScope(const UE::LLMPrivate::FTagData* TagData, bool bIsStatTag, ELLMTagSet Set, ELLMTracker Tracker, bool bOverride = true)
 	{
 		if (!FLowLevelMemTracker::bIsDisabled)
 		{
-			Init(TagData, bIsStatTag, Set, Tracker);
+			Init(TagData, bIsStatTag, Set, Tracker, bOverride);
 		}
 	}
 
@@ -778,16 +791,14 @@ public:
 	}
 
 protected:
-	void Init(FName TagName, bool bIsStatTag, ELLMTagSet InTagSet, ELLMTracker InTracker);
-	void Init(ELLMTag TagEnum, bool bIsStatTag, ELLMTagSet InTagSet, ELLMTracker InTracker);
-	void Init(const UE::LLMPrivate::FTagData* TagData, bool bIsStatTag, ELLMTagSet InTagSet, ELLMTracker InTracker);
+	void Init(FName TagName, bool bIsStatTag, ELLMTagSet InTagSet, ELLMTracker InTracker, bool bOverride = true);
+	void Init(ELLMTag TagEnum, bool bIsStatTag, ELLMTagSet InTagSet, ELLMTracker InTracker, bool bOverride = true);
+	void Init(const UE::LLMPrivate::FTagData* TagData, bool bIsStatTag, ELLMTagSet InTagSet, ELLMTracker InTracker, bool bOverride = true);
 	void Destruct();
 
 	ELLMTracker Tracker{};
 	bool bEnabled = false;
-#if LLM_ALLOW_ASSETS_TAGS
-	bool bIsAssetTag = false;
-#endif
+	ELLMTagSet TagSet;
 };
 
 /** LLM scope for pausing LLM (disables the allocation hooks). */
@@ -813,14 +824,20 @@ public:
 	~FLLMScopeFromPtr();
 protected:
 	ELLMTracker Tracker;
-	bool bEnabled;
+	bool bEnabled[static_cast<int32>(ELLMTagSet::Max)];
+
+private:
+	inline void DisableAll()
+	{
+		FMemory::Memzero(bEnabled);
+	}
 };
 
 /** Global instances to provide information about a tag to LLM. */
 class CORE_API FLLMTagDeclaration
 {
 public:
-	FLLMTagDeclaration(const TCHAR* InCPPName, const FName InDisplayName=NAME_None, FName InParentTagName = NAME_None, FName InStatName = NAME_None, FName InSummaryStatName = NAME_None);
+	FLLMTagDeclaration(const TCHAR* InCPPName, const FName InDisplayName=NAME_None, FName InParentTagName = NAME_None, FName InStatName = NAME_None, FName InSummaryStatName = NAME_None, ELLMTagSet TagSet = ELLMTagSet::None);
 	FName GetUniqueName() const { return UniqueName; }
 
 	typedef void (*FCreationCallback)(FLLMTagDeclaration&);
@@ -842,20 +859,31 @@ protected:
 	FName ParentTagName;
 	FName StatName;
 	FName SummaryStatName;
+	ELLMTagSet TagSet;
 	FLLMTagDeclaration* Next = nullptr;
 
 	friend class FLowLevelMemTracker;
 	friend class FTagTrace;
 };
 
+/** Used to filter Tag allocations specifying a Name that matches in a TagSet, or TagSet alone or just by Name. */
+struct CORE_API FLLMTagSetAllocationFilter
+{
+	/** Tag name to match */
+	FName Name;
+	/** Tag set to match */
+	ELLMTagSet TagSet;
+};
 
 #else
 
 #define LLM(...)
 #define LLM_IF_ENABLED(...)
 #define LLM_SCOPE(...)
+#define LLM_TAGSET_SCOPE(...)
 #define LLM_SCOPE_BYNAME(...)
 #define LLM_SCOPE_BYTAG(...)
+#define LLM_SCOPE_RENDER_RESOURCE(...)
 #define LLM_PLATFORM_SCOPE(...)
 #define LLM_PLATFORM_SCOPE_BYNAME(...)
 #define LLM_PLATFORM_SCOPE_BYTAG(...)
diff --git a/Engine/Source/Runtime/CoreUObject/Private/UObject/UObjectGlobals.cpp b/Engine/Source/Runtime/CoreUObject/Private/UObject/UObjectGlobals.cpp
--- a/Engine/Source/Runtime/CoreUObject/Private/UObject/UObjectGlobals.cpp
+++ b/Engine/Source/Runtime/CoreUObject/Private/UObject/UObjectGlobals.cpp
@@ -1939,7 +1939,7 @@ UPackage* LoadPackage(UPackage* InOuter, const FPackagePath& PackagePath, uint32
 
 	// since we are faking the object name, this is basically a duplicate of LLM_SCOPED_TAG_WITH_OBJECT_IN_SET
 	LLM_SCOPED_TAG_WITH_STAT_NAME_IN_SET(FLowLevelMemTracker::Get().IsTagSetActive(ELLMTagSet::Assets) ?
-											FDynamicStats::CreateMemoryStatId<FStatGroup_STATGROUP_LLMAssets>(FName(*(FString(TEXT("Package ")) + PackagePath.GetPackageNameOrFallback()))).GetName() :
+											PackagePath.GetPackageFName() :
 											NAME_None,
 										 ELLMTagSet::Assets, ELLMTracker::Default);
 	TRACE_LOADTIME_REQUEST_GROUP_SCOPE(TEXT("SyncLoad - %s"), *PackagePath.GetDebugName());
diff --git a/Engine/Source/Runtime/Engine/Private/Components/ActorComponent.cpp b/Engine/Source/Runtime/Engine/Private/Components/ActorComponent.cpp
--- a/Engine/Source/Runtime/Engine/Private/Components/ActorComponent.cpp
+++ b/Engine/Source/Runtime/Engine/Private/Components/ActorComponent.cpp
@@ -1597,6 +1597,7 @@ void UActorComponent::OnDestroyPhysicsState()
 void UActorComponent::CreatePhysicsState(bool bAllowDeferral)
 {
 	LLM_SCOPE(ELLMTag::Chaos);
+	LLM_SCOPED_TAG_WITH_OBJECT_IN_SET(GetOutermost(), ELLMTagSet::Assets);
 
 	SCOPE_CYCLE_COUNTER(STAT_ComponentCreatePhysicsState);
 
@@ -1682,6 +1683,7 @@ void UActorComponent::ExecuteRegisterEvents(FRegisterComponentContext* Context)
 	{
 		SCOPE_CYCLE_COUNTER(STAT_ComponentCreateRenderState);
 		LLM_SCOPE(ELLMTag::SceneRender);
+		LLM_SCOPED_TAG_WITH_OBJECT_IN_SET(GetOutermost(), ELLMTagSet::Assets);
 		CreateRenderState_Concurrent(Context);
 		checkf(bRenderStateCreated, TEXT("Failed to route CreateRenderState_Concurrent (%s)"), *GetFullName());
 	}
@@ -1791,6 +1793,7 @@ void UActorComponent::RemoveTickPrerequisiteComponent(UActorComponent* Prerequis
 void UActorComponent::DoDeferredRenderUpdates_Concurrent()
 {
 	LLM_SCOPE(ELLMTag::SceneRender);
+	LLM_SCOPED_TAG_WITH_OBJECT_IN_SET(GetOutermost(), ELLMTagSet::Assets);
 
 	checkf(!IsUnreachable(), TEXT("%s"), *GetFullName());
 	checkf(!IsTemplate(), TEXT("%s"), *GetFullName());
diff --git a/Engine/Source/Runtime/Engine/Private/Components/SkinnedMeshComponent.cpp b/Engine/Source/Runtime/Engine/Private/Components/SkinnedMeshComponent.cpp
--- a/Engine/Source/Runtime/Engine/Private/Components/SkinnedMeshComponent.cpp
+++ b/Engine/Source/Runtime/Engine/Private/Components/SkinnedMeshComponent.cpp
@@ -847,6 +847,8 @@ void USkinnedMeshComponent::CreateRenderState_Concurrent(FRegisterComponentConte
 
 	if( GetSkinnedAsset() )
 	{
+		LLM_SCOPED_TAG_WITH_OBJECT_IN_SET(GetSkinnedAsset()->GetOutermost(), ELLMTagSet::Assets);
+
 		// Attempting to track down UE-45505, where it looks as if somehow a skeletal mesh component's mesh has only been partially loaded, causing a mismatch in the LOD arrays
 		checkf(!GetSkinnedAsset()->HasAnyFlags(RF_NeedLoad | RF_NeedPostLoad | RF_NeedPostLoadSubobjects | RF_WillBeLoaded), TEXT("Attempting to create render state for a skeletal mesh that is is not fully loaded. Mesh: %s"), *GetSkinnedAsset()->GetName());
 
@@ -2190,6 +2192,8 @@ bool USkinnedMeshComponent::AllocateTransformData()
 	// Allocate transforms if not present.
 	if (GetSkinnedAsset() != nullptr && LeaderPoseComponent == nullptr )
 	{
+		LLM_SCOPED_TAG_WITH_OBJECT_IN_SET(GetSkinnedAsset()->GetOutermost(), ELLMTagSet::Assets);
+
 		const int32 NumBones = GetSkinnedAsset()->GetRefSkeleton().GetNum();
 
 		if(GetNumComponentSpaceTransforms() != NumBones)
diff --git a/Engine/Source/Runtime/Engine/Private/ShaderCompiler/ShaderCompiler.cpp b/Engine/Source/Runtime/Engine/Private/ShaderCompiler/ShaderCompiler.cpp
--- a/Engine/Source/Runtime/Engine/Private/ShaderCompiler/ShaderCompiler.cpp
+++ b/Engine/Source/Runtime/Engine/Private/ShaderCompiler/ShaderCompiler.cpp
@@ -7635,6 +7635,8 @@ static bool TryLoadCookedGlobalShaderMap(EShaderPlatform Platform, FScopedSlowTa
 
 void CompileGlobalShaderMap(EShaderPlatform Platform, const ITargetPlatform* TargetPlatform, bool bRefreshShaderMap)
 {
+	LLM_SCOPE_RENDER_RESOURCE(TEXT("GlobalShaderMap"));
+
 	// No global shaders needed on dedicated server or clients that use NullRHI. Note that cook commandlet needs to have them, even if it is not allowed to render otherwise.
 	if (FPlatformProperties::IsServerOnly() || (!IsRunningCommandlet() && !FApp::CanEverRender()))
 	{
diff --git a/Engine/Source/Runtime/Engine/Private/Texture2D.cpp b/Engine/Source/Runtime/Engine/Private/Texture2D.cpp
--- a/Engine/Source/Runtime/Engine/Private/Texture2D.cpp
+++ b/Engine/Source/Runtime/Engine/Private/Texture2D.cpp
@@ -1273,9 +1273,10 @@ FVirtualTexture2DResource::~FVirtualTexture2DResource()
 
 void FVirtualTexture2DResource::InitRHI()
 {
-	TRACE_CPUPROFILER_EVENT_SCOPE(FVirtualTexture2DResource::InitRHI);
-
 	check(TextureOwner);
+	
+	TRACE_CPUPROFILER_EVENT_SCOPE(FVirtualTexture2DResource::InitRHI);
+	LLM_SCOPED_TAG_WITH_OBJECT_IN_SET(TextureOwner->GetOutermost(), ELLMTagSet::Assets);
 
 	uint32 MaxAnisotropy = 0;
 	if (VirtualTextureScalability::IsAnisotropicFilteringEnabled())
diff --git a/Engine/Source/Runtime/Engine/Private/TextureRenderTarget2D.cpp b/Engine/Source/Runtime/Engine/Private/TextureRenderTarget2D.cpp
--- a/Engine/Source/Runtime/Engine/Private/TextureRenderTarget2D.cpp
+++ b/Engine/Source/Runtime/Engine/Private/TextureRenderTarget2D.cpp
@@ -614,6 +614,8 @@ ETextureCreateFlags FTextureRenderTarget2DResource::GetCreateFlags()
  */
 void FTextureRenderTarget2DResource::InitDynamicRHI()
 {
+	LLM_SCOPED_TAG_WITH_OBJECT_IN_SET(Owner->GetOutermost(), ELLMTagSet::Assets);
+
 	if( TargetSizeX > 0 && TargetSizeY > 0 )
 	{
 		FString ResourceName = Owner->GetName();
@@ -687,6 +689,8 @@ void FTextureRenderTarget2DResource::ReleaseDynamicRHI()
  */
 void FTextureRenderTarget2DResource::UpdateDeferredResource( FRHICommandListImmediate& RHICmdList, bool bClearRenderTarget/*=true*/ )
 {
+	LLM_SCOPED_TAG_WITH_OBJECT_IN_SET(Owner->GetOutermost(), ELLMTagSet::Assets);
+
 	SCOPED_DRAW_EVENT(RHICmdList, GPUResourceUpdate)
 	RemoveFromDeferredUpdateList();
 
diff --git a/Engine/Source/Runtime/Engine/Private/TextureRenderTarget2DArray.cpp b/Engine/Source/Runtime/Engine/Private/TextureRenderTarget2DArray.cpp
--- a/Engine/Source/Runtime/Engine/Private/TextureRenderTarget2DArray.cpp
+++ b/Engine/Source/Runtime/Engine/Private/TextureRenderTarget2DArray.cpp
@@ -220,6 +220,8 @@ UTexture2DArray* UTextureRenderTarget2DArray::ConstructTexture2DArray(UObject* O
  */
 void FTextureRenderTarget2DArrayResource::InitDynamicRHI()
 {
+	LLM_SCOPED_TAG_WITH_OBJECT_IN_SET(Owner->GetOutermost(), ELLMTagSet::Assets);
+
 	if((Owner->SizeX > 0) && (Owner->SizeY > 0) && (Owner->Slices > 0))
 	{
 		bool bIsSRGB = true;
@@ -297,6 +299,8 @@ void FTextureRenderTarget2DArrayResource::ReleaseDynamicRHI()
  */
 void FTextureRenderTarget2DArrayResource::UpdateDeferredResource(FRHICommandListImmediate& RHICmdList, bool bClearRenderTarget/*=true*/)
 {
+	LLM_SCOPED_TAG_WITH_OBJECT_IN_SET(Owner->GetOutermost(), ELLMTagSet::Assets);
+
 	if (!bClearRenderTarget)
 	{
 		return;
diff --git a/Engine/Source/Runtime/Engine/Private/TextureRenderTargetCube.cpp b/Engine/Source/Runtime/Engine/Private/TextureRenderTargetCube.cpp
--- a/Engine/Source/Runtime/Engine/Private/TextureRenderTargetCube.cpp
+++ b/Engine/Source/Runtime/Engine/Private/TextureRenderTargetCube.cpp
@@ -248,6 +248,8 @@ UTextureCube* UTextureRenderTargetCube::ConstructTextureCube(
  */
 void FTextureRenderTargetCubeResource::InitDynamicRHI()
 {
+	LLM_SCOPED_TAG_WITH_OBJECT_IN_SET(Owner->GetOutermost(), ELLMTagSet::Assets);
+
 	if(Owner->SizeX > 0)
 	{
 		bool bIsSRGB = true;
@@ -338,6 +340,8 @@ void FTextureRenderTargetCubeResource::ReleaseDynamicRHI()
  */
 void FTextureRenderTargetCubeResource::UpdateDeferredResource(FRHICommandListImmediate& RHICmdList, bool bClearRenderTarget/*=true*/)
 {
+	LLM_SCOPED_TAG_WITH_OBJECT_IN_SET(Owner->GetOutermost(), ELLMTagSet::Assets);
+
 	if (!bClearRenderTarget)
 	{
 		return;
diff --git a/Engine/Source/Runtime/Engine/Private/TextureRenderTargetVolume.cpp b/Engine/Source/Runtime/Engine/Private/TextureRenderTargetVolume.cpp
--- a/Engine/Source/Runtime/Engine/Private/TextureRenderTargetVolume.cpp
+++ b/Engine/Source/Runtime/Engine/Private/TextureRenderTargetVolume.cpp
@@ -226,6 +226,8 @@ UVolumeTexture* UTextureRenderTargetVolume::ConstructTextureVolume(UObject* ObjO
  */
 void FTextureRenderTargetVolumeResource::InitDynamicRHI()
 {
+	LLM_SCOPED_TAG_WITH_OBJECT_IN_SET(Owner->GetOutermost(), ELLMTagSet::Assets);
+
 	if((Owner->SizeX > 0) && (Owner->SizeY > 0) && (Owner->SizeZ > 0))
 	{
 		bool bIsSRGB = true;
@@ -305,6 +307,8 @@ void FTextureRenderTargetVolumeResource::ReleaseDynamicRHI()
  */
 void FTextureRenderTargetVolumeResource::UpdateDeferredResource(FRHICommandListImmediate& RHICmdList, bool bClearRenderTarget/*=true*/)
 {
+	LLM_SCOPED_TAG_WITH_OBJECT_IN_SET(Owner->GetOutermost(), ELLMTagSet::Assets);
+
 	if (bClearRenderTarget)
 	{
 		RHICmdList.Transition(FRHITransitionInfo(TextureRHI, ERHIAccess::Unknown, ERHIAccess::RTV));
diff --git a/Engine/Source/Runtime/Renderer/Private/RendererScene.cpp b/Engine/Source/Runtime/Renderer/Private/RendererScene.cpp
--- a/Engine/Source/Runtime/Renderer/Private/RendererScene.cpp
+++ b/Engine/Source/Runtime/Renderer/Private/RendererScene.cpp
@@ -2013,6 +2013,8 @@ FScene::~FScene()
 
 void FScene::AddPrimitive(UPrimitiveComponent* Primitive)
 {
+	LLM_SCOPED_TAG_WITH_OBJECT_IN_SET(Primitive->GetOutermost(), ELLMTagSet::Assets);
+
 	SCOPE_CYCLE_COUNTER(STAT_AddScenePrimitiveGT);
 	SCOPED_NAMED_EVENT(FScene_AddPrimitive, FColor::Green);
 
