Core - LinearColor(FColor) construction is now a constexpr

diff --git a/Engine/Source/Runtime/Core/Public/Math/Color.h b/Engine/Source/Runtime/Core/Public/Math/Color.h
--- a/Engine/Source/Runtime/Core/Public/Math/Color.h
+++ b/Engine/Source/Runtime/Core/Public/Math/Color.h
@@ -77,7 +77,7 @@ struct FLinearColor
 	 * @param Color The sRGB color that needs to be converted into linear space.
 	 * to get direct conversion use ReinterpretAsLinear
 	 */
-	FORCEINLINE FLinearColor(const FColor& Color);
+	constexpr FORCEINLINE FLinearColor(const FColor& Color);
 
 	CORE_API FLinearColor(const FVector3f& Vector);
 	CORE_API explicit FLinearColor(const FVector3d& Vector); // Warning: keep this explicit, or FVector4f will be implicitly created from FVector3d via FLinearColor
@@ -802,12 +802,12 @@ private:
 };
 DECLARE_INTRINSIC_TYPE_LAYOUT(FColor);
 
-FORCEINLINE FLinearColor::FLinearColor(const FColor& Color)
+constexpr FORCEINLINE FLinearColor::FLinearColor(const FColor& Color)
+	: R(sRGBToLinearTable[Color.R])
+    , G(sRGBToLinearTable[Color.G])
+    , B(sRGBToLinearTable[Color.B])
+    , A(static_cast<float>(Color.A) * (1.0f / 255.0f))
 {
-	R = sRGBToLinearTable[Color.R];
-	G = sRGBToLinearTable[Color.G];
-	B = sRGBToLinearTable[Color.B];
-	A = float(Color.A) * (1.0f / 255.0f);
 }
 
 FORCEINLINE FColor FLinearColor::QuantizeRound() const
