Fixed some templated types whose constructors accept anything rather than 'any other variation of the template', which were resulting in some incorrect overload resolution.

diff --git a/Engine/Source/Runtime/Core/Public/Math/TransformCalculus2D.h b/Engine/Source/Runtime/Core/Public/Math/TransformCalculus2D.h
--- a/Engine/Source/Runtime/Core/Public/Math/TransformCalculus2D.h
+++ b/Engine/Source/Runtime/Core/Public/Math/TransformCalculus2D.h
@@ -651,43 +651,43 @@ public:
 	using Matrix2Type = TMatrix2x2<T>;
 
 	/** Initialize the transform using an identity matrix and a translation. */
-	template<typename VType=float, typename VArg=UE::Math::TVector2<VType>>
-	TTransform2(const VArg& Translation = VArg(0.f, 0.f))
+	template<typename VType = float>
+	TTransform2(const UE::Math::TVector2<VType>& Translation = UE::Math::TVector2<VType>(0.f, 0.f))
 		: Trans((Vector2Type)Translation)
 	{
 	}
 
 	/** Initialize the transform using a uniform scale and a translation. */
-	template<typename VType=float, typename VArg=UE::Math::TVector2<VType>>
-	explicit TTransform2(T UniformScale, const VArg& Translation = VArg(0.f, 0.f))
+	template<typename VType = float>
+	explicit TTransform2(T UniformScale, const UE::Math::TVector2<VType>& Translation = UE::Math::TVector2<VType>(0.f, 0.f))
 		: M(TScale2<T>(UniformScale)), Trans((Vector2Type)Translation)
 	{
 	}
 
 	/** Initialize the transform using a 2D scale and a translation. */
-	template<typename VType=float, typename VArg=UE::Math::TVector2<VType>>
-	explicit TTransform2(const TScale2<T>& Scale, const VArg& Translation = VArg(0.f, 0.f))
+	template<typename VType = float>
+	explicit TTransform2(const TScale2<T>& Scale, const UE::Math::TVector2<VType>& Translation = UE::Math::TVector2<VType>(0.f, 0.f))
 		: M(Scale), Trans((Vector2Type)Translation)
 	{
 	}
 
 	/** Initialize the transform using a 2D shear and a translation. */
-	template<typename VType=float, typename VArg=UE::Math::TVector2<VType>>
-	explicit TTransform2(const TShear2<T>& Shear, const VArg& Translation = VArg(0.f, 0.f))
+	template<typename VType = float>
+	explicit TTransform2(const TShear2<T>& Shear, const UE::Math::TVector2<VType>& Translation = UE::Math::TVector2<VType>(0.f, 0.f))
 		: M(Shear), Trans((Vector2Type)Translation)
 	{
 	}
 
 	/** Initialize the transform using a 2D rotation and a translation. */
-	template<typename VType=float, typename VArg=UE::Math::TVector2<VType>>
-	explicit TTransform2(const TQuat2<T>& Rot, const VArg& Translation = VArg(0.f, 0.f))
+	template<typename VType = float>
+	explicit TTransform2(const TQuat2<T>& Rot, const UE::Math::TVector2<VType>& Translation = UE::Math::TVector2<VType>(0.f, 0.f))
 		: M(Rot), Trans((Vector2Type)Translation)
 	{
 	}
 
 	/** Initialize the transform using a general 2x2 transform and a translation. */
-	template<typename VType=float, typename VArg=UE::Math::TVector2<VType>>
-	explicit TTransform2(const TMatrix2x2<T>& Transform, const VArg& Translation = VArg(0.f, 0.f))
+	template<typename VType = float>
+	explicit TTransform2(const TMatrix2x2<T>& Transform, const UE::Math::TVector2<VType>& Translation = UE::Math::TVector2<VType>(0.f, 0.f))
 		: M(Transform), Trans((Vector2Type)Translation)
 	{
 	}
