Fix Debug assert in LLM tracker. Resize the tag array upfront correctly.

diff --git a/Engine/Source/Runtime/Core/Private/HAL/LowLevelMemTracker.cpp b/Engine/Source/Runtime/Core/Private/HAL/LowLevelMemTracker.cpp
--- a/Engine/Source/Runtime/Core/Private/HAL/LowLevelMemTracker.cpp
+++ b/Engine/Source/Runtime/Core/Private/HAL/LowLevelMemTracker.cpp
@@ -3907,7 +3907,7 @@ bool FLLMTracker::FindTagsForPtr(void* InPtr, TArray<const FTagData *, TInlineAl
 		return false;
 	}
 
-	OutTags.SetNumUnsafeInternal(static_cast<int32>(ELLMTagSet::Max));
+	OutTags.SetNumUninitialized(static_cast<int32>(ELLMTagSet::Max));
 #if LLM_ALLOW_ASSETS_TAGS
 	for (int32 TagSetIndex = 0; TagSetIndex < static_cast<int32>(ELLMTagSet::Max); TagSetIndex++)
 	{
