Introduced the keys debug culture. When users switch to this culture, all the localization text that is displayed will instead display the localization key of the UI test.
This allows users to easily see the key of the localization text that is localized incorrectly or not localized at all, allowing them to track down the issue faster.
- Users ccan change the culture using the command line argument -culture=keys. They can also do the same in the in game console.
- Introduced an automation test for the keys culture
- Removed unnecessary headers for some of the files.
#test Introduced an automation test to switch from the english culture to the keys culture, back to the english culture and then to the japanese culture to ensure that all display text is correct in each culture. Also had assistance checking that all text in the editor changed to display the localization key when the keys culture is active.

diff --git a/Engine/Source/Runtime/Core/Private/Internationalization/Cultures/KeysCulture.cpp b/Engine/Source/Runtime/Core/Private/Internationalization/Cultures/KeysCulture.cpp
new file mode 100644
--- /dev/null
+++ b/Engine/Source/Runtime/Core/Private/Internationalization/Cultures/KeysCulture.cpp
@@ -0,0 +1,18 @@
+// Copyright Epic Games, Inc. All Rights Reserved.
+
+#include "Internationalization/Cultures/KeysCulture.h"
+
+#if ENABLE_LOC_TESTING
+
+FKeysCulture::FKeysCulture(const FCultureRef& InInvariantCulture)
+	: InvariantCulture(InInvariantCulture)
+{
+}
+
+const FString& FKeysCulture::StaticGetName()
+{
+	static const FString KeysCultureName = TEXT("keys");
+	return KeysCultureName;
+}
+
+#endif
diff --git a/Engine/Source/Runtime/Core/Private/Internationalization/Cultures/KeysCulture.h b/Engine/Source/Runtime/Core/Private/Internationalization/Cultures/KeysCulture.h
new file mode 100644
--- /dev/null
+++ b/Engine/Source/Runtime/Core/Private/Internationalization/Cultures/KeysCulture.h
@@ -0,0 +1,45 @@
+// Copyright Epic Games, Inc. All Rights Reserved.
+
+#pragma once
+
+#include "Internationalization/LocTesting.h"
+#include "Internationalization/ICustomCulture.h"
+
+#if ENABLE_LOC_TESTING
+
+/** 
+* A debug culture that when set will display the localization key of an FText instead of its associated text.
+* This is used to help diagnose issues with wrong or missing translations in Editor/Game UI text.
+* Use the command line or in game console to set this culture and it will allow you to find the localization key for an FText where you can look in the corresponding .po files to find the problematic text entry.
+* The logic for performing the conversion is in FTextLocalizationManager.
+* 
+* @see FTextLocalizationManager
+*/
+class FKeysCulture : public ICustomCulture
+{
+public:
+	explicit FKeysCulture(const FCultureRef& InInvariantCulture);
+
+	static const FString& StaticGetName();
+
+	//~ ICustomCulture interface
+	virtual FCultureRef GetBaseCulture() const override { return InvariantCulture; }
+	virtual FString GetDisplayName() const override { return FKeysCulture::StaticGetName(); }
+	virtual FString GetEnglishName() const override { return FKeysCulture::StaticGetName(); }
+	virtual FString GetName() const override { return FKeysCulture::StaticGetName(); }
+	virtual FString GetNativeName() const override { return FKeysCulture::StaticGetName(); }
+	virtual FString GetUnrealLegacyThreeLetterISOLanguageName() const override { return TEXT("INT"); }
+	virtual FString GetThreeLetterISOLanguageName() const override { return FKeysCulture::StaticGetName(); }
+	virtual FString GetTwoLetterISOLanguageName() const override { return FKeysCulture::StaticGetName(); }
+	virtual FString GetNativeLanguage() const override { return FKeysCulture::StaticGetName(); }
+	virtual FString GetNativeRegion() const override { return FString(); }
+	virtual FString GetRegion() const override { return FString(); }
+	virtual FString GetScript() const override { return FString(); }
+	virtual FString GetVariant() const override { return FString(); }
+	virtual bool IsRightToLeft() const override { return false; }
+
+private:
+	FCultureRef InvariantCulture;
+};
+
+#endif
diff --git a/Engine/Source/Runtime/Core/Private/Internationalization/ICUCulture.cpp b/Engine/Source/Runtime/Core/Private/Internationalization/ICUCulture.cpp
--- a/Engine/Source/Runtime/Core/Private/Internationalization/ICUCulture.cpp
+++ b/Engine/Source/Runtime/Core/Private/Internationalization/ICUCulture.cpp
@@ -1,7 +1,6 @@
 // Copyright Epic Games, Inc. All Rights Reserved.
 
 #include "Internationalization/ICUCulture.h"
-#include "Internationalization/Cultures/LeetCulture.h"
 #include "Internationalization/Internationalization.h"
 #include "Misc/ScopeLock.h"
 #include "Containers/SortedMap.h"
diff --git a/Engine/Source/Runtime/Core/Private/Internationalization/ICUInternationalization.cpp b/Engine/Source/Runtime/Core/Private/Internationalization/ICUInternationalization.cpp
--- a/Engine/Source/Runtime/Core/Private/Internationalization/ICUInternationalization.cpp
+++ b/Engine/Source/Runtime/Core/Private/Internationalization/ICUInternationalization.cpp
@@ -7,6 +7,7 @@
 #include "Internationalization/Culture.h"
 #include "Internationalization/Internationalization.h"
 #include "Internationalization/Cultures/LeetCulture.h"
+#include "Internationalization/Cultures/KeysCulture.h"
 #include "Stats/Stats.h"
 #include "Misc/CoreStats.h"
 #include "Misc/ConfigCacheIni.h"
@@ -183,6 +184,7 @@ bool FICUInternationalization::Initialize()
 
 #if ENABLE_LOC_TESTING
 	I18N->AddCustomCulture(MakeShared<FLeetCulture>(I18N->InvariantCulture.ToSharedRef()));
+	I18N->AddCustomCulture(MakeShared<FKeysCulture>(I18N->InvariantCulture.ToSharedRef()));
 #endif
 
 	InitializeTimeZone();
diff --git a/Engine/Source/Runtime/Core/Private/Internationalization/InternationalizationUtilities.cpp b/Engine/Source/Runtime/Core/Private/Internationalization/InternationalizationUtilities.cpp
--- a/Engine/Source/Runtime/Core/Private/Internationalization/InternationalizationUtilities.cpp
+++ b/Engine/Source/Runtime/Core/Private/Internationalization/InternationalizationUtilities.cpp
@@ -3,6 +3,7 @@
 #include "Internationalization/InternationalizationUtilities.h"
 #include "Internationalization/Internationalization.h"
 #include "Internationalization/Cultures/LeetCulture.h"
+#include "Internationalization/Cultures/KeysCulture.h"
 
 namespace InternationalizationUtilities
 {
@@ -89,11 +90,12 @@ namespace InternationalizationUtilities
 	{
 		auto IsLanguageCode = [](const FString& InCode)
 		{
-			// Language codes must be 2 or 3 letters, or our special "LEET" language
+			// Language codes must be 2 or 3 letters, or our special "LEET"/"keys" languages
 			return InCode.Len() == 2
 				|| InCode.Len() == 3
 #if ENABLE_LOC_TESTING
 				|| InCode == FLeetCulture::StaticGetName()
+				|| InCode == FKeysCulture::StaticGetName()
 #endif
 				;
 		};
diff --git a/Engine/Source/Runtime/Core/Private/Internationalization/LegacyInternationalization.cpp b/Engine/Source/Runtime/Core/Private/Internationalization/LegacyInternationalization.cpp
--- a/Engine/Source/Runtime/Core/Private/Internationalization/LegacyInternationalization.cpp
+++ b/Engine/Source/Runtime/Core/Private/Internationalization/LegacyInternationalization.cpp
@@ -2,6 +2,7 @@
 
 #include "Internationalization/LegacyInternationalization.h"
 #include "Internationalization/Cultures/LeetCulture.h"
+#include "Internationalization/Cultures/KeysCulture.h"
 
 #if !UE_ENABLE_ICU
 
@@ -23,6 +24,7 @@ bool FLegacyInternationalization::Initialize()
 
 #if ENABLE_LOC_TESTING
 	I18N->AddCustomCulture(MakeShared<FLeetCulture>(I18N->InvariantCulture.ToSharedRef()));
+	I18N->AddCustomCulture(MakeShared<FKeysCulture>(I18N->InvariantCulture.ToSharedRef()));
 #endif
 
 	return true;
diff --git a/Engine/Source/Runtime/Core/Private/Internationalization/TextLocalizationManager.cpp b/Engine/Source/Runtime/Core/Private/Internationalization/TextLocalizationManager.cpp
--- a/Engine/Source/Runtime/Core/Private/Internationalization/TextLocalizationManager.cpp
+++ b/Engine/Source/Runtime/Core/Private/Internationalization/TextLocalizationManager.cpp
@@ -7,6 +7,7 @@
 #include "Internationalization/PolyglotTextSource.h"
 #include "Internationalization/StringTableRegistry.h"
 #include "Internationalization/Cultures/LeetCulture.h"
+#include "Internationalization/Cultures/KeysCulture.h"
 #include "GenericPlatform/GenericPlatformFile.h"
 #include "HAL/FileManager.h"
 #include "HAL/LowLevelMemTracker.h"
@@ -61,6 +62,11 @@ FString GetRequestedCulture(const TCHAR* InCommandLineKey, const TCHAR* InConfig
 			RequestedCulture = FLeetCulture::StaticGetName();
 			OutOverrideLevel = ERequestedCultureOverrideLevel::CommandLine;
 		}
+		else if (RequestedCulture.IsEmpty() && FParse::Param(FCommandLine::Get(), *FKeysCulture::StaticGetName()))
+		{
+			RequestedCulture = FKeysCulture::StaticGetName();
+			OutOverrideLevel = ERequestedCultureOverrideLevel::CommandLine;
+		}
 #endif
 
 		if (RequestedCulture.IsEmpty() && FParse::Value(FCommandLine::Get(), TEXT("CULTUREFORCOOKING="), RequestedCulture))
@@ -220,7 +226,8 @@ void ApplyDefaultCultureSettings(const ELocalizationLoadFlags LocLoadFlags)
 		FString TargetCultureName = InRequestedCulture;
 
 #if ENABLE_LOC_TESTING
-		if (TargetCultureName != FLeetCulture::StaticGetName())
+		bool bIsTargetCultureDebugCulture = (TargetCultureName == FLeetCulture::StaticGetName()) || (TargetCultureName == FKeysCulture::StaticGetName());
+		if (!bIsTargetCultureDebugCulture)
 #endif
 		{
 			// Validate the locale has data or fallback to one that does.
@@ -1110,9 +1117,11 @@ void FTextLocalizationManager::LoadLocalizationResourcesForPrioritizedCultures(T
 	}
 
 	// Leet-ify always needs the native text to operate on, so force native data if we're loading for LEET
+	// The keys culture also uses native so that we have a good way to restore the orginal text from the keys state 
 	ELocalizationLoadFlags FinalLocLoadFlags = LocLoadFlags;
 #if ENABLE_LOC_TESTING
-	if (PrioritizedCultureNames[0] == FLeetCulture::StaticGetName())
+	bool bShouldForceLoadNative = (PrioritizedCultureNames[0] == FLeetCulture::StaticGetName()) || (PrioritizedCultureNames[0] == FKeysCulture::StaticGetName());
+	if (bShouldForceLoadNative)
 	{
 		FinalLocLoadFlags |= ELocalizationLoadFlags::Native;
 	}
@@ -1138,23 +1147,11 @@ void FTextLocalizationManager::LoadLocalizationResourcesForPrioritizedCultures(T
 	// The leet culture is fake. Just leet-ify existing strings.
 	if (PrioritizedCultureNames[0] == FLeetCulture::StaticGetName())
 	{
-		// Lock while updating the tables
-		{
-			FScopeLock ScopeLock(&DisplayStringLookupTableCS);
-
-			for (auto& DisplayStringPair : DisplayStringLookupTable)
-			{
-				FDisplayStringEntry& LiveEntry = DisplayStringPair.Value;
-				LiveEntry.NativeStringBackup = LiveEntry.DisplayString;
-				
-				if (!LiveEntry.DisplayString->IsEmpty())
-				{
-					FTextDisplayStringRef TmpDisplayString = MakeTextDisplayString(CopyTemp(*LiveEntry.DisplayString));
-					FInternationalization::Leetify(*TmpDisplayString);
-					LiveEntry.DisplayString = TmpDisplayString;
-				}
-			}
-		}
+		LeetifyAllDisplayStrings();
+	}
+	else if (PrioritizedCultureNames[0] == FKeysCulture::StaticGetName())
+	{
+		KeyifyAllDisplayStrings();
 	}
 	else
 #endif
@@ -1508,3 +1505,40 @@ bool FTextLocalizationManager::IsLocalizationLocked() const
 	return bIsLocalizationLocked;
 }
 #endif
+
+#if ENABLE_LOC_TESTING
+void FTextLocalizationManager::LeetifyAllDisplayStrings()
+{
+	// Lock while updating the tables
+	FScopeLock ScopeLock(&DisplayStringLookupTableCS);
+
+	for (auto& DisplayStringPair : DisplayStringLookupTable)
+	{
+		FDisplayStringEntry& LiveEntry = DisplayStringPair.Value;
+		LiveEntry.NativeStringBackup = LiveEntry.DisplayString;
+
+		if (!LiveEntry.DisplayString->IsEmpty())
+		{
+			FTextDisplayStringRef TmpDisplayString = MakeTextDisplayString(CopyTemp(*LiveEntry.DisplayString));
+			FInternationalization::Leetify(*TmpDisplayString);
+			LiveEntry.DisplayString = TmpDisplayString;
+		}
+	}
+}
+
+void FTextLocalizationManager::KeyifyAllDisplayStrings()
+{
+	// Lock while updating the tables
+	FScopeLock ScopeLock(&DisplayStringLookupTableCS);
+
+	for (auto& DisplayStringPair : DisplayStringLookupTable)
+	{
+		FDisplayStringEntry& LiveEntry = DisplayStringPair.Value;
+		LiveEntry.NativeStringBackup = LiveEntry.DisplayString;
+
+		FString Key = DisplayStringPair.Key.GetKey().GetChars();
+		FTextDisplayStringRef TmpDisplayString = MakeTextDisplayString(MoveTemp(Key));
+		LiveEntry.DisplayString = TmpDisplayString;
+	}
+}
+#endif 
\ No newline at end of file
diff --git a/Engine/Source/Runtime/Core/Private/Tests/Internationalization/KeysCultureTest.cpp b/Engine/Source/Runtime/Core/Private/Tests/Internationalization/KeysCultureTest.cpp
new file mode 100644
--- /dev/null
+++ b/Engine/Source/Runtime/Core/Private/Tests/Internationalization/KeysCultureTest.cpp
@@ -0,0 +1,153 @@
+// Copyright Epic Games, Inc. All Rights Reserved.
+
+#include "CoreTypes.h"
+#include "Containers/UnrealString.h"
+#include "CoreGlobals.h"
+#include "Internationalization/Internationalization.h"
+#include "Internationalization/LocTesting.h"
+#include "Internationalization/PolyglotTextData.h"
+#include "Internationalization/Culture.h"
+#include "Internationalization/Cultures/KeysCulture.h"
+#include "Misc/AutomationTest.h"
+
+#if WITH_DEV_AUTOMATION_TESTS && ENABLE_LOC_TESTING
+
+#define LOCTEXT_NAMESPACE "Core.Tests.KeysCultureTests"
+
+/** 
+* A utility struct to ensure that the internationalization system is reverted
+* to its original state by the end of a test. Declare this at the start of tests that involve culture changes.
+*/
+struct FRestoreCultureStateGuard
+{
+	FRestoreCultureStateGuard()
+	{
+		FInternationalization::Get().BackupCultureState(OriginalCultureState);
+	}
+	~FRestoreCultureStateGuard()
+	{
+		FInternationalization::Get().RestoreCultureState(OriginalCultureState);
+	}
+private:
+	FInternationalization::FCultureStateSnapshot OriginalCultureState;
+};
+
+IMPLEMENT_SIMPLE_AUTOMATION_TEST(FKeysCultureTest, "System.Core.Misc.KeysCulture", EAutomationTestFlags::EditorContext | EAutomationTestFlags::ClientContext | EAutomationTestFlags::EngineFilter)
+bool FKeysCultureTest::RunTest (const FString& Parameters)
+{
+	FRestoreCultureStateGuard RestoreCultureStateGuard;
+	
+	static const FString NativeCulture = TEXT("en");
+	static const FString TargetCulture = TEXT("ja");
+
+	// Test Overview
+	// 1. We set up the test environment and change the culture to the native culture
+	// 2. We register text data that as a native string and a target culture translation.
+	// 3. We register text data that as a native string and no target culture translation. 
+	// 4. We change to the target culture. The text with translation should have its display string the same as the localized string in the target culture. The display string for the text with no translation should match the native string.
+	// 5. We change to the keys culture. Both text with and without translations in a target culture should have their display strings show the localization key of the text.
+	// 6. We revert back to the native culture. Both text with and without translations should have their display strings the same as their native strings with no change.
+	
+	FInternationalization& I18N = FInternationalization::Get();
+	// We initialize the test environment to the native culture to start from a clean slate.
+	I18N.SetCurrentCulture(NativeCulture);
+	if (I18N.GetCurrentCulture()->GetName() != NativeCulture)
+	{
+		// somehow failed to initialize the test environment to the right culture. Early out.
+		return false;
+	}
+#define CREATE_TEST_POLYGLOTDATA(InKey, InSourceString) FPolyglotTextData(ELocalizedTextSourceCategory::Editor, TEXT(LOCTEXT_NAMESPACE), TEXT(InKey), TEXT(InSourceString), NativeCulture)
+	// We use polyglot text data here to patch in our test strings so we don't have to create a separte loclization resource file 
+	// Test1 is for text with a native string and a non-empty translation in a target language
+	// Simulates text with a valid translation editor/game
+	FPolyglotTextData Test1 = CREATE_TEST_POLYGLOTDATA("Test1Key", "Test1 Native");
+	Test1.AddLocalizedString(TargetCulture, TEXT("Test1 Target"));
+
+	// Test2 is for text with a valid native string but a missing translation in a target culture
+	// Simulates text with missing translation in editor/game
+	FPolyglotTextData Test2 = CREATE_TEST_POLYGLOTDATA("Test2Key", "Test2 Native");
+#undef CREATE_TEST_POLYGLOTDATA
+
+	// GetText() already handles registration of the polyglot text data with the FTextLocalizationManager internally and adds these strings to the display string table.
+	FText Test1Text = Test1.GetText();
+	FText Test2Text = Test2.GetText();
+	
+	FString Test1NativeCultureDisplayString = Test1Text.ToString();
+	// Test1 display string should have the same value as the native string. Otherwise it could be a problem with registering the text. The set up of the test failed. Early out.
+	if (Test1NativeCultureDisplayString != Test1.GetNativeString())
+	{
+		return false;
+	}
+
+	FString Test2NativeCultureDisplayString = Test2Text.ToString();
+	// Same as above.
+	if (Test2NativeCultureDisplayString != Test2.GetNativeString())
+	{
+		return false;
+	}
+
+		// Change to target culture. 
+	I18N.SetCurrentCulture(TargetCulture);
+	// Test 1 - We should be able to change from the native culture to the target culture 
+	if (!TestTrue(TEXT("Changing from native culture to target culture."), I18N.GetCurrentCulture()->GetName() == TargetCulture))
+	{
+		// just early out. Doesn't make sense to keep testing.
+		return false;
+	}
+
+	// Test 2 - The display string for the text with translation should match the localized string in the target culture 
+	FString Test1TargetCultureDisplayString = Test1Text.ToString();
+	FString Test1TargetCultureString;
+	Test1.GetLocalizedString(TargetCulture, Test1TargetCultureString);
+	TestEqual(TEXT("Text with translation display is the same as the localized string in target culture."), Test1TargetCultureDisplayString, Test1TargetCultureString);
+
+	// Test 3 - The text with no translation should have no retrievable text from it spolyglot data in the target culture.
+	// Test 4 - The display string for the text with no translation should match the native string while in the target culture. Text with no translations for a culture should fall back to the native string. 
+	FString Test2TargetCultureDisplayString = Test2Text.ToString();
+	FString Test2TargetCultureString;
+	Test2.GetLocalizedString(TargetCulture, Test2TargetCultureString);
+	TestTrue(TEXT("Text without translation should not have a localized string in the target culture."), Test2TargetCultureString.IsEmpty());
+	TestEqual(TEXT("Text without translation display string is the same as the native string in target culture."), Test2TargetCultureDisplayString, Test2.GetNativeString());
+	
+	// Change to the keys culture
+	// This culture should output all the display strings as the localization key of the text.
+	// There will be logs that say the localization resource files associated with the keys culture can't be read. That is normal and can be ignored.
+	I18N.SetCurrentCulture(FKeysCulture::StaticGetName());
+	// Test 5 - Can we successfully change from the native culture to the keys culture 
+	if (!TestTrue(TEXT("Changing from native culture to keys culture."), I18N.GetCurrentCulture()->GetName() == FKeysCulture::StaticGetName()))
+	{
+		// just early out. Doesn't make sense to keep testing.
+		return false;
+	}
+	
+	// Test 6 - The display string for the text with translation should be the localization key for the text in the keys culture 
+	FString Test1KeysCultureDisplayString = Test1Text.ToString();
+	TestEqual(TEXT("Text with translation now has display string that matches its localization key in keys culture."), Test1KeysCultureDisplayString, Test1.GetKey());
+	
+	// Test 7 - The display string for the text with no translation should be its localization key in the keys culture.
+	FString Test2KeysCultureDisplayString = Test2Text.ToString();
+	TestEqual(TEXT("Text with blank translation has display string that matches its localization key in keys culture."), Test2KeysCultureDisplayString, Test2.GetKey());
+
+	// Revert to native culture. All of the text should still show up as the originals.
+	I18N.SetCurrentCulture(NativeCulture);
+	// Test 8 - Can we change from the keys culture back to the native culture 
+	if (!TestTrue(TEXT("Able to revert from keys culture to native culture.."), I18N.GetCurrentCulture()->GetName() == NativeCulture))
+	{
+		// early out same as above 
+		return false;
+	}
+
+	// Test 9 - The display string for the text with translation should be exactly the same as its original native string before changing to the keys culture
+	FString Test1NativeCultureRevertedDisplayString = Test1Text.ToString();
+	TestEqual(TEXT("Text with translation has display string that is the same as the original native string in native culture."), Test1NativeCultureRevertedDisplayString, Test1.GetNativeString());
+
+// Test 10 - The display string for the text with no translation should be exactly the same as its original native string before changing to the keys culture
+	FString Test2NativeCultureRevertedDisplayString = Test2Text.ToString();
+	TestEqual(TEXT("Text without translation has display string the same as its native string in native culture."), Test2NativeCultureRevertedDisplayString, Test2.GetNativeString());
+
+// @TODOLocalization: Unregister the polyglot data once the functionality is implemented.
+	return true;
+}
+
+#undef LOCTEXT_NAMESPACE 
+#endif // WITH_DEV_AUTOMATION_TESTS && ENABLE_LOC_TESTING
diff --git a/Engine/Source/Runtime/Core/Public/Internationalization/TextLocalizationManager.h b/Engine/Source/Runtime/Core/Public/Internationalization/TextLocalizationManager.h
--- a/Engine/Source/Runtime/Core/Public/Internationalization/TextLocalizationManager.h
+++ b/Engine/Source/Runtime/Core/Public/Internationalization/TextLocalizationManager.h
@@ -292,6 +292,13 @@ private:
 
 	/** Dirties the text revision counter by incrementing it, causing a revision mismatch for any information cached before this happens.  */
 	void DirtyTextRevision();
+#if ENABLE_LOC_TESTING
+	/** A helper function that leetifies all of the display strings when the LEET culture is active. */
+	void LeetifyAllDisplayStrings();
+	/** A helper function that converts all of the display strings to show the localization key associated with the string when the keys culture is active. */
+	void KeyifyAllDisplayStrings();
+#endif 
+
 
 	/** Array of registered localized text sources, sorted by priority (@see RegisterTextSource) */
 	TArray<TSharedPtr<ILocalizedTextSource>> LocalizedTextSources;
