[Core]
* ScriptDelegates - Fixed non unity/pch compile error

diff --git a/Engine/Source/Runtime/Core/Public/UObject/ScriptDelegates.h b/Engine/Source/Runtime/Core/Public/UObject/ScriptDelegates.h
--- a/Engine/Source/Runtime/Core/Public/UObject/ScriptDelegates.h
+++ b/Engine/Source/Runtime/Core/Public/UObject/ScriptDelegates.h
@@ -3,6 +3,7 @@
 #pragma once
 
 #include "CoreTypes.h"
+#include "Delegates/Delegate.h"
 #include "Misc/AssertionMacros.h"
 #include "Templates/AreTypesEqual.h"
 #include "Templates/UnrealTypeTraits.h"
