[Core]
* Moved delegate macros from Delegate.h to DelegateMacros.h in order to prevent circular include dependencies
* Made Delegate.h and ScriptDelegates.h include DelegateMacros.h

diff --git a/Engine/Source/Runtime/Core/Public/Delegates/Delegate.h b/Engine/Source/Runtime/Core/Public/Delegates/Delegate.h
--- a/Engine/Source/Runtime/Core/Public/Delegates/Delegate.h
+++ b/Engine/Source/Runtime/Core/Public/Delegates/Delegate.h
@@ -2,24 +2,8 @@
 
 #pragma once
 
-//#define UE_DETECT_DELEGATES_RACE_CONDITIONS 1
-
-#if !defined(UE_DETECT_DELEGATES_RACE_CONDITIONS)
-	#define UE_DETECT_DELEGATES_RACE_CONDITIONS 0
-#endif
-
-#if UE_DETECT_DELEGATES_RACE_CONDITIONS
-	#include "Misc/MTAccessDetector.h"
-	#define UE_DELEGATES_MT_ACCESS_DETECTOR UE_MT_DECLARE_RW_FULLY_RECURSIVE_ACCESS_DETECTOR
-	#define UE_DELEGATES_MT_SCOPED_WRITE_ACCESS UE_MT_SCOPED_WRITE_ACCESS
-	#define UE_DELEGATES_MT_SCOPED_READ_ACCESS UE_MT_SCOPED_READ_ACCESS
-#else
-	#define UE_DELEGATES_MT_ACCESS_DETECTOR(...)
-	#define UE_DELEGATES_MT_SCOPED_WRITE_ACCESS(...)
-	#define UE_DELEGATES_MT_SCOPED_READ_ACCESS(...)
-#endif
-
 #include "CoreTypes.h"
+#include "Delegates/DelegateMacros.h"
 #include "Misc/AssertionMacros.h"
 #include "UObject/NameTypes.h"
 #include "Templates/SharedPointer.h"
diff --git a/Engine/Source/Runtime/Core/Public/Delegates/DelegateMacros.h b/Engine/Source/Runtime/Core/Public/Delegates/DelegateMacros.h
new file mode 100644
--- /dev/null
+++ b/Engine/Source/Runtime/Core/Public/Delegates/DelegateMacros.h
@@ -0,0 +1,20 @@
+// Copyright Epic Games, Inc. All Rights Reserved.
+
+#pragma once
+
+//#define UE_DETECT_DELEGATES_RACE_CONDITIONS 1
+
+#if !defined(UE_DETECT_DELEGATES_RACE_CONDITIONS)
+	#define UE_DETECT_DELEGATES_RACE_CONDITIONS 0
+#endif
+
+#if UE_DETECT_DELEGATES_RACE_CONDITIONS
+	#include "Misc/MTAccessDetector.h"
+	#define UE_DELEGATES_MT_ACCESS_DETECTOR UE_MT_DECLARE_RW_FULLY_RECURSIVE_ACCESS_DETECTOR
+	#define UE_DELEGATES_MT_SCOPED_WRITE_ACCESS UE_MT_SCOPED_WRITE_ACCESS
+	#define UE_DELEGATES_MT_SCOPED_READ_ACCESS UE_MT_SCOPED_READ_ACCESS
+#else
+	#define UE_DELEGATES_MT_ACCESS_DETECTOR(...)
+	#define UE_DELEGATES_MT_SCOPED_WRITE_ACCESS(...)
+	#define UE_DELEGATES_MT_SCOPED_READ_ACCESS(...)
+#endif
diff --git a/Engine/Source/Runtime/Core/Public/UObject/ScriptDelegates.h b/Engine/Source/Runtime/Core/Public/UObject/ScriptDelegates.h
--- a/Engine/Source/Runtime/Core/Public/UObject/ScriptDelegates.h
+++ b/Engine/Source/Runtime/Core/Public/UObject/ScriptDelegates.h
@@ -3,7 +3,7 @@
 #pragma once
 
 #include "CoreTypes.h"
-#include "Delegates/Delegate.h"
+#include "Delegates/DelegateMacros.h"
 #include "Misc/AssertionMacros.h"
 #include "Templates/AreTypesEqual.h"
 #include "Templates/UnrealTypeTraits.h"
