Added GetAllocatedSize functions to delegate classes to aid with memory tracking

diff --git a/Engine/Source/Runtime/Core/Public/Delegates/DelegateBase.h b/Engine/Source/Runtime/Core/Public/Delegates/DelegateBase.h
--- a/Engine/Source/Runtime/Core/Public/Delegates/DelegateBase.h
+++ b/Engine/Source/Runtime/Core/Public/Delegates/DelegateBase.h
@@ -116,6 +116,14 @@ protected:
 		return DelegateSize ? (IDelegateInstance*)DelegateAllocator.GetAllocation() : nullptr;
 	}
 
+	/**
+	 * Returns the amount of memory allocated by this delegate, not including sizeof(*this).
+	 */
+	SIZE_T GetAllocatedSize() const
+	{
+		return DelegateAllocator.GetAllocatedSize(DelegateSize, sizeof(FAlignedInlineDelegateType));
+	}
+
 private:
 	friend void* operator new(size_t Size, FDelegateBase& Base);
 
diff --git a/Engine/Source/Runtime/Core/Public/Delegates/DelegateSignatureImpl.inl b/Engine/Source/Runtime/Core/Public/Delegates/DelegateSignatureImpl.inl
--- a/Engine/Source/Runtime/Core/Public/Delegates/DelegateSignatureImpl.inl
+++ b/Engine/Source/Runtime/Core/Public/Delegates/DelegateSignatureImpl.inl
@@ -74,6 +74,7 @@ class TDelegate<InRetValType(ParamTypes...), UserPolicy> : public TDelegateBase<
 public:
 	// Make sure FDelegateBase's public functions are publicly exposed through the TDelegate API
 	using Super::Unbind;
+	using Super::GetAllocatedSize;
 
 private:
 	// Make sure FDelegateBase's protected functions are not accidentally exposed through the TDelegate API
@@ -702,6 +703,7 @@ public:
 	using Super::IsBound;
 	using Super::IsBoundToObject;
 	using Super::RemoveAll;
+	using Super::GetAllocatedSize;
 
 private:
 	// Make sure TMulticastDelegateBase's protected functions are not accidentally exposed through the TMulticastDelegate API
diff --git a/Engine/Source/Runtime/Core/Public/Delegates/MulticastDelegateBase.h b/Engine/Source/Runtime/Core/Public/Delegates/MulticastDelegateBase.h
--- a/Engine/Source/Runtime/Core/Public/Delegates/MulticastDelegateBase.h
+++ b/Engine/Source/Runtime/Core/Public/Delegates/MulticastDelegateBase.h
@@ -140,6 +140,20 @@ public:
 		return Result;
 	}
 
+	/**
+	 * Returns the amount of memory allocated by this delegate's invocation list and the delegates stored within it, not including sizeof(*this).
+	 */
+	SIZE_T GetAllocatedSize() const
+	{
+		SIZE_T Size = 0;
+		Size += InvocationList.GetAllocatedSize();
+		for (const TDelegateBase<UserPolicy>& DelegateBaseRef : InvocationList)
+		{
+			Size += DelegateBaseRef.GetAllocatedSize();
+		}
+		return Size;
+	}
+
 protected:
 
 	/** Hidden default constructor. */
diff --git a/Engine/Source/Runtime/Core/Public/UObject/ScriptDelegates.h b/Engine/Source/Runtime/Core/Public/UObject/ScriptDelegates.h
--- a/Engine/Source/Runtime/Core/Public/UObject/ScriptDelegates.h
+++ b/Engine/Source/Runtime/Core/Public/UObject/ScriptDelegates.h
@@ -621,6 +621,14 @@ public:
 		return OutputList;
 	}
 
+	/**
+	 * Returns the amount of memory allocated by this delegate's invocation list, not including sizeof(*this).
+	 */
+	SIZE_T GetAllocatedSize() const
+	{
+		return InvocationList.GetAllocatedSize();
+	}
+
 protected:
 
 	/**
