Added missing ElementType aliases to various smart pointers and containers

diff --git a/Engine/Source/Runtime/Core/Public/Containers/CircularBuffer.h b/Engine/Source/Runtime/Core/Public/Containers/CircularBuffer.h
--- a/Engine/Source/Runtime/Core/Public/Containers/CircularBuffer.h
+++ b/Engine/Source/Runtime/Core/Public/Containers/CircularBuffer.h
@@ -14,9 +14,10 @@
  * operations using a simple bit mask instead of the commonly used modulus operator that may
  * be slow on some platforms.
  */
-template<typename ElementType> class TCircularBuffer
+template<typename InElementType> class TCircularBuffer
 {
 public:
+	using ElementType = InElementType;
 
 	/**
 	 * Creates and initializes a new instance of the TCircularBuffer class.
diff --git a/Engine/Source/Runtime/Core/Public/UObject/WeakObjectPtrTemplates.h b/Engine/Source/Runtime/Core/Public/UObject/WeakObjectPtrTemplates.h
--- a/Engine/Source/Runtime/Core/Public/UObject/WeakObjectPtrTemplates.h
+++ b/Engine/Source/Runtime/Core/Public/UObject/WeakObjectPtrTemplates.h
@@ -38,6 +38,8 @@ struct TWeakObjectPtr : private TWeakObjectPtrBase
 	static_assert(TAreTypesEqual<TWeakObjectPtrBase, FWeakObjectPtr>::Value, "TWeakObjectPtrBase should not be overridden");
 
 public:
+	using ElementType = T;
+	
 	TWeakObjectPtr() = default;
 	TWeakObjectPtr(const TWeakObjectPtr&) = default;
 	TWeakObjectPtr& operator=(const TWeakObjectPtr&) = default;
diff --git a/Engine/Source/Runtime/CoreUObject/Public/UObject/LazyObjectPtr.h b/Engine/Source/Runtime/CoreUObject/Public/UObject/LazyObjectPtr.h
--- a/Engine/Source/Runtime/CoreUObject/Public/UObject/LazyObjectPtr.h
+++ b/Engine/Source/Runtime/CoreUObject/Public/UObject/LazyObjectPtr.h
@@ -181,6 +181,8 @@ template<class T=UObject>
 struct TLazyObjectPtr : private FLazyObjectPtr
 {
 public:
+	using ElementType = T;
+	
 	TLazyObjectPtr() = default;
 
 	TLazyObjectPtr(TLazyObjectPtr<T>&&) = default;
diff --git a/Engine/Source/Runtime/CoreUObject/Public/UObject/PersistentObjectPtr.h b/Engine/Source/Runtime/CoreUObject/Public/UObject/PersistentObjectPtr.h
--- a/Engine/Source/Runtime/CoreUObject/Public/UObject/PersistentObjectPtr.h
+++ b/Engine/Source/Runtime/CoreUObject/Public/UObject/PersistentObjectPtr.h
@@ -15,7 +15,10 @@ template<class TObjectID>
 struct TPersistentObjectPtr
 {
 	friend struct FGCInternals;
-public:	
+
+public:
+	using ElementType = TObjectID;
+
 	/** Default constructor, will be null */
 	FORCEINLINE TPersistentObjectPtr()
 	{
diff --git a/Engine/Source/Runtime/CoreUObject/Public/UObject/SoftObjectPtr.h b/Engine/Source/Runtime/CoreUObject/Public/UObject/SoftObjectPtr.h
--- a/Engine/Source/Runtime/CoreUObject/Public/UObject/SoftObjectPtr.h
+++ b/Engine/Source/Runtime/CoreUObject/Public/UObject/SoftObjectPtr.h
@@ -113,6 +113,8 @@ struct TSoftObjectPtr
 	friend struct TSoftObjectPtr;
 
 public:
+	using ElementType = T;
+	
 	/** Default constructor, will be null */
 	FORCEINLINE TSoftObjectPtr()
 	{
@@ -439,6 +441,8 @@ class TSoftClassPtr
 	friend class TSoftClassPtr;
 
 public:
+	using ElementType = TClass;
+	
 	/** Default constructor, will be null */
 	FORCEINLINE TSoftClassPtr()
 	{
diff --git a/Engine/Source/Runtime/CoreUObject/Public/UObject/StrongObjectPtr.h b/Engine/Source/Runtime/CoreUObject/Public/UObject/StrongObjectPtr.h
--- a/Engine/Source/Runtime/CoreUObject/Public/UObject/StrongObjectPtr.h
+++ b/Engine/Source/Runtime/CoreUObject/Public/UObject/StrongObjectPtr.h
@@ -72,6 +72,8 @@ template <typename ObjectType, typename ReferencerNameProvider = UEStrongObjectP
 class TStrongObjectPtr
 {
 public:
+	using ElementType = ObjectType;
+	
 	TStrongObjectPtr(TStrongObjectPtr&& InOther) = default;
 	TStrongObjectPtr& operator=(TStrongObjectPtr&& InOther) = default;
 	~TStrongObjectPtr() = default;
diff --git a/Engine/Source/Runtime/CoreUObject/Public/UObject/WeakFieldPtr.h b/Engine/Source/Runtime/CoreUObject/Public/UObject/WeakFieldPtr.h
--- a/Engine/Source/Runtime/CoreUObject/Public/UObject/WeakFieldPtr.h
+++ b/Engine/Source/Runtime/CoreUObject/Public/UObject/WeakFieldPtr.h
@@ -22,6 +22,8 @@ private:
 	mutable TFieldPath<T> Field;
 
 public:
+	using ElementType = T;
+	
 	TWeakFieldPtr() = default;
 	TWeakFieldPtr(const TWeakFieldPtr&) = default;
 	TWeakFieldPtr& operator=(const TWeakFieldPtr&) = default;
diff --git a/Engine/Source/Runtime/CoreUObject/Public/UObject/WeakInterfacePtr.h b/Engine/Source/Runtime/CoreUObject/Public/UObject/WeakInterfacePtr.h
--- a/Engine/Source/Runtime/CoreUObject/Public/UObject/WeakInterfacePtr.h
+++ b/Engine/Source/Runtime/CoreUObject/Public/UObject/WeakInterfacePtr.h
@@ -16,6 +16,8 @@
 template<class T>
 struct TWeakInterfacePtr
 {
+	using ElementType = T;
+	
 	/**
 	 * Construct a new default weak pointer, pointing to null object.
 	 */
