Clean up Unix strlen to align with other implementations

diff --git a/Engine/Source/Runtime/Core/Public/Unix/UnixPlatformString.h b/Engine/Source/Runtime/Core/Public/Unix/UnixPlatformString.h
--- a/Engine/Source/Runtime/Core/Public/Unix/UnixPlatformString.h
+++ b/Engine/Source/Runtime/Core/Public/Unix/UnixPlatformString.h
@@ -26,7 +26,7 @@ struct FUnixPlatformString : public
 			return 0;
 
 		int Len = 0;
-		while(String[Len])
+		while(*String++)
 		{
 			++Len;
 		}
