Do not lower priority on tasks that are canceled or were expedited

diff --git a/Engine/Source/Runtime/Core/Private/Async/Fundamental/Scheduler.cpp b/Engine/Source/Runtime/Core/Private/Async/Fundamental/Scheduler.cpp
--- a/Engine/Source/Runtime/Core/Private/Async/Fundamental/Scheduler.cpp
+++ b/Engine/Source/Runtime/Core/Private/Async/Fundamental/Scheduler.cpp
@@ -164,7 +164,7 @@ namespace LowLevelTasks
 		else
 		{
 			// Dynamic priority only enables for root task when we're not inside a named thread (i.e. GT, RT)
-			const bool bSkipPriorityChange = ParentTask || !GTaskGraphUseDynamicPrioritization || !FSchedulerTls::IsWorkerThread();
+			const bool bSkipPriorityChange = ParentTask || !GTaskGraphUseDynamicPrioritization || !FSchedulerTls::IsWorkerThread() || InOutTask->WasCanceledOrIsExpediting();
 
 			FRunnableThread* RunnableThread = nullptr;
 			if (!bSkipPriorityChange)
@@ -228,6 +228,7 @@ namespace LowLevelTasks
 			}
 		}
 	}
+
 	void FScheduler::RestartWorkers(uint32 NumForegroundWorkers, uint32 NumBackgroundWorkers, FThread::EForkable IsForkable, EThreadPriority InWorkerPriority, EThreadPriority InBackgroundPriority, uint64 InWorkerAffinity, uint64 InBackgroundAffinity)
 	{
 		FScopeLock Lock(&WorkerThreadsCS);
@@ -322,7 +323,7 @@ namespace LowLevelTasks
 			if constexpr (bIsBusyWaiting)
 			{
 				FTask::FInitData InitData = Task->GetInitData();
-				bool bAllowBusyWaiting = EnumHasAnyFlags(InitData.Flags, ETaskFlags::AllowBusyWaiting);
+				bool bAllowBusyWaiting = EnumHasAnyFlags(InitData.Flags, ETaskFlags::AllowBusyWaiting) || Task->WasCanceledOrIsExpediting();
 				
 				if (!bAllowBusyWaiting || AnyExecuted) 
 				{
diff --git a/Engine/Source/Runtime/Core/Public/Async/Fundamental/Task.h b/Engine/Source/Runtime/Core/Public/Async/Fundamental/Task.h
--- a/Engine/Source/Runtime/Core/Public/Async/Fundamental/Task.h
+++ b/Engine/Source/Runtime/Core/Public/Async/Fundamental/Task.h
@@ -348,6 +348,15 @@ namespace LowLevelTasks
 			return EnumHasAnyFlags(State, ETaskState::ExpeditedFlag | ETaskState::CompletedFlag);
 		}
 
+	private:
+		//Scheduler internal interface to speed things up
+		inline bool WasCanceledOrIsExpediting() const
+		{
+			ETaskState State = PackedData.load(std::memory_order_relaxed).GetState();
+			return EnumHasAnyFlags(State, ETaskState::CanceledFlag | ETaskState::RunningFlag);
+		}
+
+	public:
 		/*
 		* means the task is ready to be launched but might already been canceled 
 		*/
