UE5 CPU minspec updated to SSE 4.2

Tried to run automation tests but it's basically a red board right now. I ran two, one with these changes and one without, and they seem to be basically the same errors, however there are so many and I can't really cross compare easily.

base automation tests: 63b8a5dcaf3ebedd990b7dfe
preflight automation tests: 63b87801221aab21e08d5973

diff --git a/Engine/Source/Programs/UnrealBuildTool/Platform/Android/AndroidToolChain.cs b/Engine/Source/Programs/UnrealBuildTool/Platform/Android/AndroidToolChain.cs
--- a/Engine/Source/Programs/UnrealBuildTool/Platform/Android/AndroidToolChain.cs
+++ b/Engine/Source/Programs/UnrealBuildTool/Platform/Android/AndroidToolChain.cs
@@ -772,6 +772,7 @@ namespace UnrealBuildTool
 				Result += " -fno-strict-aliasing";
 				Result += " -fno-short-enums";
 				Result += " -march=atom";
+				Result += " -msse4.2";
 			}
 
 			ClangSanitizer Sanitizer = BuildWithSanitizer();
diff --git a/Engine/Source/Programs/UnrealBuildTool/Platform/Windows/VCToolChain.cs b/Engine/Source/Programs/UnrealBuildTool/Platform/Windows/VCToolChain.cs
--- a/Engine/Source/Programs/UnrealBuildTool/Platform/Windows/VCToolChain.cs
+++ b/Engine/Source/Programs/UnrealBuildTool/Platform/Windows/VCToolChain.cs
@@ -391,8 +391,8 @@ namespace UnrealBuildTool
 				// Tell the Clang compiler to generate 64-bit code
 				Arguments.Add("--target=x86_64-pc-windows-msvc");
 
-				// This matches Microsoft's default support floor for SSE.
-				Arguments.Add("-mssse3");
+				// UE5 minspec is 4.2
+				Arguments.Add("-msse4.2");
 
 				// Use tpause on supported processors.
 				Arguments.Add("-mwaitpkg");
@@ -641,8 +641,6 @@ namespace UnrealBuildTool
 				// Inform Unreal code that we have sse2, sse4, and AVX, both available to compile and available to run
 				// By setting the ALWAYS_HAS defines, we we direct Unreal code to skip cpuid checks to verify that the running hardware supports sse/avx.
 				AddDefinition(Arguments, "PLATFORM_ENABLE_VECTORINTRINSICS=1");
-				AddDefinition(Arguments, "PLATFORM_MAYBE_HAS_SSE4_1=1");
-				AddDefinition(Arguments, "PLATFORM_ALWAYS_HAS_SSE4_1=1");
 				AddDefinition(Arguments, "PLATFORM_MAYBE_HAS_AVX=1");
 				AddDefinition(Arguments, "PLATFORM_ALWAYS_HAS_AVX=1");
 			}
diff --git a/Engine/Source/Programs/UnrealBuildTool/ToolChain/ClangToolChain.cs b/Engine/Source/Programs/UnrealBuildTool/ToolChain/ClangToolChain.cs
--- a/Engine/Source/Programs/UnrealBuildTool/ToolChain/ClangToolChain.cs
+++ b/Engine/Source/Programs/UnrealBuildTool/ToolChain/ClangToolChain.cs
@@ -831,6 +831,12 @@ namespace UnrealBuildTool
 			Arguments.Add("-c");
 			Arguments.Add("-pipe");
 
+			if (UEBuildPlatform.GetBuildPlatform(CompileEnvironment.Platform).IsX86Architecture(CompileEnvironment.Architecture))
+			{		
+				// UE5 minspec is 4.2
+				Arguments.Add("-msse4.2");
+			}
+
 			// Add include paths to the argument list.
 			GetCompileArguments_IncludePaths(CompileEnvironment, Arguments);
 
diff --git a/Engine/Source/Programs/UnrealBuildTool/ToolChain/ISPCToolChain.cs b/Engine/Source/Programs/UnrealBuildTool/ToolChain/ISPCToolChain.cs
--- a/Engine/Source/Programs/UnrealBuildTool/ToolChain/ISPCToolChain.cs
+++ b/Engine/Source/Programs/UnrealBuildTool/ToolChain/ISPCToolChain.cs
@@ -31,7 +31,7 @@ namespace UnrealBuildTool
 				(UEBuildPlatform.IsPlatformInGroup(Platform, UnrealPlatformGroup.Unix) && Platform != UnrealTargetPlatform.LinuxArm64) ||
 				Platform == UnrealTargetPlatform.Mac)
 			{
-				ISPCTargets.AddRange(new string[] { "avx512skx-i32x8", "avx2", "avx", "sse4", "sse2" });
+				ISPCTargets.AddRange(new string[] { "avx512skx-i32x8", "avx2", "avx", "sse4" });
 			}
 			else if (Platform == UnrealTargetPlatform.LinuxArm64)
 			{
@@ -43,8 +43,8 @@ namespace UnrealBuildTool
 				{
 					case "-armv7": ISPCTargets.Add("neon"); break; // Assumes NEON is in use
 					case "-arm64": ISPCTargets.Add("neon"); break;
-					case "-x86": ISPCTargets.AddRange(new string[] { "sse4", "sse2" }); break;
-					case "-x64": ISPCTargets.AddRange(new string[] { "sse4", "sse2" }); break;
+					case "-x86": ISPCTargets.Add("sse4"); break;
+					case "-x64": ISPCTargets.Add("sse4"); break;
 					default: Logger.LogWarning("Invalid Android architecture for ISPC. At least one architecture (armv7, x86, etc) needs to be selected in the project settings to build"); break;
 				}
 			}
diff --git a/Engine/Source/Runtime/Core/Public/HAL/Platform.h b/Engine/Source/Runtime/Core/Public/HAL/Platform.h
--- a/Engine/Source/Runtime/Core/Public/HAL/Platform.h
+++ b/Engine/Source/Runtime/Core/Public/HAL/Platform.h
@@ -162,12 +162,18 @@
 // If PLATFORM_MAYBE_HAS_### is 1, then ### intrinsics are compilable.
 // This does not guarantee that the intrinsics are runnable on all instances of the platform however; a runtime check such as cpuid may be required to confirm availability.
 // If PLATFORM_ALWAYS_HAS_### is 1, then ## intrinsics will compile and run on all instances of the platform.  PLATFORM_ALWAYS_HAS_### == 1 implies PLATFORM_MAYBE_HAS_### == 1.
+
+// UE5.2+ requires SSE4.2.
 #ifndef PLATFORM_MAYBE_HAS_SSE4_1
-	#define PLATFORM_MAYBE_HAS_SSE4_1			0
+	#define PLATFORM_MAYBE_HAS_SSE4_1			PLATFORM_CPU_X86_FAMILY
 #endif
 #ifndef PLATFORM_ALWAYS_HAS_SSE4_1
-	#define PLATFORM_ALWAYS_HAS_SSE4_1			0
+	#define PLATFORM_ALWAYS_HAS_SSE4_1			PLATFORM_CPU_X86_FAMILY
+#endif
+#ifndef PLATFORM_ALWAYS_HAS_SSE4_2
+	#define PLATFORM_ALWAYS_HAS_SSE4_2			PLATFORM_CPU_X86_FAMILY
 #endif
+
 #ifndef PLATFORM_MAYBE_HAS_AVX
 	#define PLATFORM_MAYBE_HAS_AVX				0
 #endif
@@ -175,7 +181,7 @@
 	#define PLATFORM_ALWAYS_HAS_AVX				0
 #endif
 #ifndef PLATFORM_ALWAYS_HAS_AVX_2
-#define PLATFORM_ALWAYS_HAS_AVX_2				0
+	#define PLATFORM_ALWAYS_HAS_AVX_2			0
 #endif
 #ifndef PLATFORM_ALWAYS_HAS_FMA3
 	#define PLATFORM_ALWAYS_HAS_FMA3			0
diff --git a/Engine/Source/Runtime/Core/Public/Mac/MacPlatform.h b/Engine/Source/Runtime/Core/Public/Mac/MacPlatform.h
--- a/Engine/Source/Runtime/Core/Public/Mac/MacPlatform.h
+++ b/Engine/Source/Runtime/Core/Public/Mac/MacPlatform.h
@@ -52,15 +52,6 @@ typedef FMacPlatformTypes FPlatformTypes;
 #define PLATFORM_SUPPORTS_PRAGMA_PACK				1
 #define PLATFORM_ENABLE_VECTORINTRINSICS			1
 #define PLATFORM_ENABLE_VECTORINTRINSICS_NEON       PLATFORM_MAC_ARM64
-#ifndef PLATFORM_MAYBE_HAS_SSE4_1  // May be set from UnrealBuildTool
-	#define PLATFORM_MAYBE_HAS_SSE4_1				PLATFORM_MAC_X86
-#endif
-// Current unreal minspec is sse2, not sse4, so on mac any calling code must check _cpuid before calling SSE4 instructions
-// If called on a platform for which _cpuid for SSE4 returns false, attempting to call SSE4 intrinsics will crash
-// If your title has raised the minspec to sse4, you can define PLATFORM_ALWAYS_HAS_SSE4_1 to 1
-#ifndef PLATFORM_ALWAYS_HAS_SSE4_1 // May be set from UnrealBuildTool
-	#define PLATFORM_ALWAYS_HAS_SSE4_1				0
-#endif
 // FMA3 support was added starting from Intel Haswell
 #ifndef PLATFORM_ALWAYS_HAS_FMA3
 	#define PLATFORM_ALWAYS_HAS_FMA3				0
diff --git a/Engine/Source/Runtime/Core/Public/Math/UnrealPlatformMathSSE4.h b/Engine/Source/Runtime/Core/Public/Math/UnrealPlatformMathSSE4.h
--- a/Engine/Source/Runtime/Core/Public/Math/UnrealPlatformMathSSE4.h
+++ b/Engine/Source/Runtime/Core/Public/Math/UnrealPlatformMathSSE4.h
@@ -6,13 +6,15 @@
 
 #include "HAL/Platform.h"
 #include "Math/UnrealPlatformMathSSE.h"
-// Code including this header is responsible for including the correct platform-specific header for SSE intrinsics.
 
+// UE5.2+ requires SSE4.2
+
+// We have to retain this #if because it's pulled in via the linux header chain
+// for all platforms at the moment and we rely on the parent class to implement
+// the functions
 #if PLATFORM_MAYBE_HAS_SSE4_1
+#include <smmintrin.h>
 
-// SSE4_1 instrinsics are only available for compile if PLATFORM_MAYBE_HAS_SSE4_1 is 1, and even if compiled, should not
-// be called at runtime unless PLATFORM_ALWAYS_HAS_SSE4_1 is 1 or cpuid has been checked to verify that the current hardware
-// instance supports it.
 namespace UE4
 {
 namespace SSE4
@@ -61,8 +63,6 @@ namespace SSE4
 
 #endif // PLATFORM_MAYBE_HAS_SSE4_1
 
-// We currently don't have any runtime checks in FPlatforMath for whether the SSE4_1 intrinsics are available,
-// so we have to not call them from FPlatformMath classes unless PLATFORM_ALWAYS_HAS_SSE4_1
 #define UNREALPLATFORMMATH_SSE4_1_ENABLED PLATFORM_ALWAYS_HAS_SSE4_1
 
 template<class Base>
diff --git a/Engine/Source/Runtime/Core/Public/Unix/UnixPlatform.h b/Engine/Source/Runtime/Core/Public/Unix/UnixPlatform.h
--- a/Engine/Source/Runtime/Core/Public/Unix/UnixPlatform.h
+++ b/Engine/Source/Runtime/Core/Public/Unix/UnixPlatform.h
@@ -105,13 +105,6 @@ typedef FUnixPlatformTypes FPlatformTypes;
 	#define PLATFORM_ENABLE_VECTORINTRINSICS		1
 #endif
 
-// We do not currently compile with -msse4 or higher on Unix or Linux
-#ifndef PLATFORM_MAYBE_HAS_SSE4_1 // May be set from UnrealBuildTool
-	#define PLATFORM_MAYBE_HAS_SSE4_1							0
-#endif
-#ifndef PLATFORM_ALWAYS_HAS_SSE4_1 // May be set from UnrealBuildTool
-	#define PLATFORM_ALWAYS_HAS_SSE4_1							0
-#endif
 #ifndef PLATFORM_ALWAYS_HAS_FMA3
 	#define PLATFORM_ALWAYS_HAS_FMA3							0
 #endif
diff --git a/Engine/Source/Runtime/Core/Public/Windows/WindowsPlatform.h b/Engine/Source/Runtime/Core/Public/Windows/WindowsPlatform.h
--- a/Engine/Source/Runtime/Core/Public/Windows/WindowsPlatform.h
+++ b/Engine/Source/Runtime/Core/Public/Windows/WindowsPlatform.h
@@ -61,15 +61,6 @@ typedef FWindowsPlatformTypes FPlatformTypes;
 #elif (defined(_M_IX86) || defined(_M_X64))
 	#define PLATFORM_CPU_X86_FAMILY							1
 	#define PLATFORM_ENABLE_VECTORINTRINSICS				1
-#ifndef PLATFORM_MAYBE_HAS_SSE4_1 // May be set from UnrealBuildTool
-	#define PLATFORM_MAYBE_HAS_SSE4_1						1
-#endif
-// Current unreal minspec is sse2, not sse4, so on windows any calling code must check _cpuid before calling SSE4 instructions;
-// If called on a platform for which _cpuid for SSE4 returns false, attempting to call SSE4 intrinsics will crash
-// If your title has raised the minspec to sse4, you can define PLATFORM_ALWAYS_HAS_SSE4_1 to 1
-#ifndef PLATFORM_ALWAYS_HAS_SSE4_1 // May be set from UnrealBuildTool
-	#define PLATFORM_ALWAYS_HAS_SSE4_1						0
-#endif
 
 #endif
 // FMA3 support was added starting from AMD Piledriver (excluding Jaguar) and Intel Haswell (excluding Pentium and Celeron)
