Fixed cancellation in ParallelFor to avoid scheduling an additional task

Canceling a task that had not executed yet would cause its destructor to launch a new task. Explicitly track the need to reschedule to avoid unintentional rescheduling. Also changes the TryCancel flags because the task has always been launched by the time we attempt to cancel it.

diff --git a/Engine/Source/Runtime/Core/Public/Async/ParallelFor.h b/Engine/Source/Runtime/Core/Public/Async/ParallelFor.h
--- a/Engine/Source/Runtime/Core/Public/Async/ParallelFor.h
+++ b/Engine/Source/Runtime/Core/Public/Async/ParallelFor.h
@@ -237,8 +237,9 @@ namespace ParallelForImpl
 		class FParallelExecutor 
 		{
 			FDataHandle Data;
-			mutable int32 WorkerIndex;
+			int32 WorkerIndex;
 			LowLevelTasks::ETaskPriority Priority;
+			mutable bool bReschedule = false;
 
 		public:
 			inline FParallelExecutor(FDataHandle&& InData, int32 InWorkerIndex, LowLevelTasks::ETaskPriority InPriority) 
@@ -258,7 +259,7 @@ namespace ParallelForImpl
 
 			~FParallelExecutor()
 			{
-				if (Data.IsValid() && WorkerIndex >= 0)
+				if (Data.IsValid() && bReschedule)
 				{
 					FParallelExecutor::LaunchTask(nullptr, MoveTemp(Data), WorkerIndex, Priority);
 				}
@@ -323,7 +324,6 @@ namespace ParallelForImpl
 					{
 						if (!bIsMaster)
 						{
-							WorkerIndex = -1;
 							return false;
 						}
 						BatchIndex = (NumBatches - 1);
@@ -344,12 +344,10 @@ namespace ParallelForImpl
 						{
 							Data->FinishedSignal->Trigger();
 						}
-						WorkerIndex = -1;
 						return true;
 					}
 					else if (EndIndex >= Num)
 					{
-						WorkerIndex = -1;
 						return false;
 					}
 					else if (!bIsBackgroundPriority)
@@ -360,7 +358,8 @@ namespace ParallelForImpl
 					auto PassedTime = [Start, &Now]() { return Now() - Start; };
 					if (PassedTime() > YieldingThreshold)
 					{
-						//abort and reschedule (in the destructor as WorkerIndex is larger_eq Zero) to give higher priority tasks a chance to run
+						//abort and reschedule (in the destructor) to give higher priority tasks a chance to run
+						bReschedule = true;
 						return false;
 					}
 				}
@@ -407,7 +406,7 @@ namespace ParallelForImpl
 			TRACE_CPUPROFILER_EVENT_SCOPE(ParallelFor.Cancel);
 			for (FTracedTask& TracedTask : LocalExecutor.GetData()->Tasks)
 			{
-				TracedTask.Task.TryCancel(LowLevelTasks::ECancellationFlags::PrelaunchCancellation);
+				TracedTask.Task.TryCancel(LowLevelTasks::ECancellationFlags::None);
 			}
 		}
 
