Change TryVirtualizePackages/TryRehydratePackages to return the corresponding results structure rather than an enum value.

#rnx

- Since the API was forcing the caller to pass in a results structure to be filled in, we might as well make it the return value.
- Added a ::WasSuccessful method to the results structures that can be used instead of checking if the result had errors or not.
- Remove the reset method from FVirtualizationResult/FRehydrationResult as they no longer need it.
- The older deprecated methods still use the results enum, so we cannot easily deprecate those enums yet.

diff --git a/Engine/Source/Developer/Virtualization/Private/VirtualizationManager.cpp b/Engine/Source/Developer/Virtualization/Private/VirtualizationManager.cpp
--- a/Engine/Source/Developer/Virtualization/Private/VirtualizationManager.cpp
+++ b/Engine/Source/Developer/Virtualization/Private/VirtualizationManager.cpp
@@ -839,29 +839,30 @@ EQueryResult FVirtualizationManager::QueryPayloadStatuses(TArrayView<const FIoHa
 	return EQueryResult::Success;
 }
 
-EVirtualizationResult FVirtualizationManager::TryVirtualizePackages(TConstArrayView<FString> PackagePaths, EVirtualizationOptions Options, FVirtualizationResult& OutResult)
+FVirtualizationResult FVirtualizationManager::TryVirtualizePackages(TConstArrayView<FString> PackagePaths, EVirtualizationOptions Options)
 {
+	FVirtualizationResult Result;
+
 	if (IsEnabled() && IsPushingEnabled(EStorageType::Persistent))
 	{
-		UE::Virtualization::VirtualizePackages(PackagePaths, Options, OutResult);
+		UE::Virtualization::VirtualizePackages(PackagePaths, Options, Result);
 
-		if (OutResult.Errors.IsEmpty() && !VirtualizationProcessTag.IsEmpty())
+		if (Result.WasSuccessful() && !VirtualizationProcessTag.IsEmpty())
 		{
-			FText Tag = FText::FromString(VirtualizationProcessTag);
-			OutResult.DescriptionTags.Add(Tag);
+			Result.DescriptionTags.Add(FText::FromString(VirtualizationProcessTag));
 		}
 	}
 
-	return OutResult.Errors.IsEmpty() ? EVirtualizationResult::Success : EVirtualizationResult::Failed;
+	return Result;
 }
 
-ERehydrationResult FVirtualizationManager::TryRehydratePackages(TConstArrayView<FString> PackagePaths, ERehydrationOptions Options, FRehydrationResult& OutResult)
+FRehydrationResult FVirtualizationManager::TryRehydratePackages(TConstArrayView<FString> PackagePaths, ERehydrationOptions Options)
 {
-	OutResult.Reset();
+	FRehydrationResult Result;
 
-	UE::Virtualization::RehydratePackages(PackagePaths, Options, OutResult);
+	UE::Virtualization::RehydratePackages(PackagePaths, Options, Result);
 
-	return OutResult.Errors.IsEmpty() ? ERehydrationResult::Success : ERehydrationResult::Failed;
+	return Result;
 }
 
 ERehydrationResult FVirtualizationManager::TryRehydratePackages(TConstArrayView<FString> PackagePaths, uint64 PaddingAlignment, TArray<FText>& OutErrors, TArray<FSharedBuffer>& OutPackages, TArray<FRehydrationInfo>* OutInfo)
diff --git a/Engine/Source/Developer/Virtualization/Private/VirtualizationManager.h b/Engine/Source/Developer/Virtualization/Private/VirtualizationManager.h
--- a/Engine/Source/Developer/Virtualization/Private/VirtualizationManager.h
+++ b/Engine/Source/Developer/Virtualization/Private/VirtualizationManager.h
@@ -125,9 +125,9 @@ private:
 
 	virtual EQueryResult QueryPayloadStatuses(TArrayView<const FIoHash> Ids, EStorageType StorageType, TArray<EPayloadStatus>& OutStatuses) override;
 
-	virtual EVirtualizationResult TryVirtualizePackages(TConstArrayView<FString> PackagePaths, EVirtualizationOptions Options, FVirtualizationResult& OutResult) override;
+	virtual FVirtualizationResult TryVirtualizePackages(TConstArrayView<FString> PackagePaths, EVirtualizationOptions Options) override;
 
-	virtual ERehydrationResult TryRehydratePackages(TConstArrayView<FString> PackagePaths, ERehydrationOptions Options, FRehydrationResult& OutResult) override;
+	virtual FRehydrationResult TryRehydratePackages(TConstArrayView<FString> PackagePaths, ERehydrationOptions Options) override;
 	virtual ERehydrationResult TryRehydratePackages(TConstArrayView<FString> PackagePaths, uint64 PaddingAlignment, TArray<FText>& OutErrors, TArray<FSharedBuffer>& OutPackages, TArray<FRehydrationInfo>* OutInfo) override;
 
 	virtual void DumpStats() const override;
diff --git a/Engine/Source/Editor/SourceControlWindows/Private/SSourceControlSubmit.cpp b/Engine/Source/Editor/SourceControlWindows/Private/SSourceControlSubmit.cpp
--- a/Engine/Source/Editor/SourceControlWindows/Private/SSourceControlSubmit.cpp
+++ b/Engine/Source/Editor/SourceControlWindows/Private/SSourceControlSubmit.cpp
@@ -55,16 +55,15 @@ PRAGMA_ENABLE_DEPRECATION_WARNINGS
 		return true;
 	}
 
-	FVirtualizationResult VirtualizationResults;
 	EVirtualizationOptions VirtualizationOptions = EVirtualizationOptions::None;
 
-	EVirtualizationResult Result = System.TryVirtualizePackages(FilesToSubmit, VirtualizationOptions, VirtualizationResults);
-	if (Result == EVirtualizationResult::Success)
+	FVirtualizationResult Result = System.TryVirtualizePackages(FilesToSubmit, VirtualizationOptions);
+	if (Result.WasSuccessful())
 	{
 		FTextBuilder NewDescription;
 		NewDescription.AppendLine(Description);
 
-		for (const FText& Line : VirtualizationResults.DescriptionTags)
+		for (const FText& Line : Result.DescriptionTags)
 		{
 			NewDescription.AppendLine(Line);
 		}
@@ -75,7 +74,7 @@ PRAGMA_ENABLE_DEPRECATION_WARNINGS
 	}
 	else if (System.AllowSubmitIfVirtualizationFailed())
 	{
-		for (const FText& Error : VirtualizationResults.Errors)
+		for (const FText& Error : Result.Errors)
 		{
 			FMessageLog("SourceControl").Warning(Error);
 		}
@@ -85,7 +84,7 @@ PRAGMA_ENABLE_DEPRECATION_WARNINGS
 	}
 	else
 	{
-		for (const FText& Error : VirtualizationResults.Errors)
+		for (const FText& Error : Result.Errors)
 		{
 			FMessageLog("SourceControl").Error(Error);
 		}
diff --git a/Engine/Source/Editor/UnrealEd/Private/Commandlets/ContentCommandlets.cpp b/Engine/Source/Editor/UnrealEd/Private/Commandlets/ContentCommandlets.cpp
--- a/Engine/Source/Editor/UnrealEd/Private/Commandlets/ContentCommandlets.cpp
+++ b/Engine/Source/Editor/UnrealEd/Private/Commandlets/ContentCommandlets.cpp
@@ -1788,16 +1788,15 @@ bool UResavePackagesCommandlet::TryVirtualization(const TArray<FString>& FilesTo
 		return true;
 	}
 
-	FVirtualizationResult VirtualizationResults;
 	EVirtualizationOptions VirtualizationOptions = EVirtualizationOptions::None;
 
-	EVirtualizationResult Result = System.TryVirtualizePackages(FilesToSubmit, VirtualizationOptions, VirtualizationResults);
-	if (Result == EVirtualizationResult::Success)
+	FVirtualizationResult Result = System.TryVirtualizePackages(FilesToSubmit, VirtualizationOptions);
+	if (Result.WasSuccessful())
 	{
 		FTextBuilder NewDescription;
 		NewDescription.AppendLine(InOutDescription);
 
-		for (const FText& Line : VirtualizationResults.DescriptionTags)
+		for (const FText& Line : Result.DescriptionTags)
 		{
 			NewDescription.AppendLine(Line);
 		}
@@ -1808,7 +1807,7 @@ bool UResavePackagesCommandlet::TryVirtualization(const TArray<FString>& FilesTo
 	}
 	else if (System.AllowSubmitIfVirtualizationFailed())
 	{
-		for (const FText& Error : VirtualizationResults.Errors)
+		for (const FText& Error : Result.Errors)
 		{
 			UE_LOG(LogContentCommandlet, Warning, TEXT("%s"), *Error.ToString());
 		}
@@ -1818,7 +1817,7 @@ bool UResavePackagesCommandlet::TryVirtualization(const TArray<FString>& FilesTo
 	}
 	else
 	{
-		for (const FText& Error : VirtualizationResults.Errors)
+		for (const FText& Error : Result.Errors)
 		{
 			UE_LOG(LogContentCommandlet, Error, TEXT("%s"), *Error.ToString());
 		}
diff --git a/Engine/Source/Editor/VirtualizationEditor/Private/RehydrationMenuEntry.cpp b/Engine/Source/Editor/VirtualizationEditor/Private/RehydrationMenuEntry.cpp
--- a/Engine/Source/Editor/VirtualizationEditor/Private/RehydrationMenuEntry.cpp
+++ b/Engine/Source/Editor/VirtualizationEditor/Private/RehydrationMenuEntry.cpp
@@ -27,8 +27,8 @@ void RehydratePackages(TArray<FString> SelectedFiles)
 {
 	TArray<FText> Errors;
 
-	FRehydrationResult ResultInfo;
-	if (IVirtualizationSystem::Get().TryRehydratePackages(SelectedFiles, ERehydrationOptions::Checkout, ResultInfo) == ERehydrationResult::Success)
+	FRehydrationResult Result = IVirtualizationSystem::Get().TryRehydratePackages(SelectedFiles, ERehydrationOptions::Checkout);
+	if (Result.WasSuccessful())
 	{
 		// TODO: At some point ::TryRehydratePackages will return more detail info about the process
 		// when it does we should make a better job at logging it.
diff --git a/Engine/Source/Programs/UnrealVirtualizationTool/Private/Commands/RehydrateCommand.cpp b/Engine/Source/Programs/UnrealVirtualizationTool/Private/Commands/RehydrateCommand.cpp
--- a/Engine/Source/Programs/UnrealVirtualizationTool/Private/Commands/RehydrateCommand.cpp
+++ b/Engine/Source/Programs/UnrealVirtualizationTool/Private/Commands/RehydrateCommand.cpp
@@ -115,13 +115,11 @@ bool FRehydrateCommand::Run(const TArray<FProject>& Projects)
 			Options |= ERehydrationOptions::Checkout;
 		}
 
-		FRehydrationResult Info;
-		UE::Virtualization::IVirtualizationSystem::Get().TryRehydratePackages(ProjectPackages, Options, Info);
-
-		if (!Info.Errors.IsEmpty())
+		FRehydrationResult Result = UE::Virtualization::IVirtualizationSystem::Get().TryRehydratePackages(ProjectPackages, Options);
+		if (!Result.WasSuccessful())
 		{
 			UE_LOG(LogVirtualizationTool, Error, TEXT("The rehydration process failed with the following errors:"));
-			for (const FText& Error : Info.Errors)
+			for (const FText& Error : Result.Errors)
 			{
 				UE_LOG(LogVirtualizationTool, Error, TEXT("\t%s"), *Error.ToString());
 			}
@@ -130,7 +128,7 @@ bool FRehydrateCommand::Run(const TArray<FProject>& Projects)
 
 		if (bShouldCheckout)
 		{
-			UE_LOG(LogVirtualizationTool, Display, TEXT("\t\t%d packages were checked out of revision control"), Info.CheckedOutPackages.Num());
+			UE_LOG(LogVirtualizationTool, Display, TEXT("\t\t%d packages were checked out of revision control"), Result.CheckedOutPackages.Num());
 		}
 
 		Project.UnRegisterMountPoints();
diff --git a/Engine/Source/Programs/UnrealVirtualizationTool/Private/Commands/VirtualizeCommand.cpp b/Engine/Source/Programs/UnrealVirtualizationTool/Private/Commands/VirtualizeCommand.cpp
--- a/Engine/Source/Programs/UnrealVirtualizationTool/Private/Commands/VirtualizeCommand.cpp
+++ b/Engine/Source/Programs/UnrealVirtualizationTool/Private/Commands/VirtualizeCommand.cpp
@@ -166,13 +166,12 @@ bool FVirtualizeCommand::Run(const TArray<FProject>& Projects)
 			Options |= EVirtualizationOptions::Checkout;
 		}
 
-		FVirtualizationResult ResultInfo;
-		UE::Virtualization::IVirtualizationSystem::Get().TryVirtualizePackages(ProjectPackages, Options, ResultInfo);
+		FVirtualizationResult Result = UE::Virtualization::IVirtualizationSystem::Get().TryVirtualizePackages(ProjectPackages, Options);
 
-		if (!ResultInfo.Errors.IsEmpty())
+		if (!Result.WasSuccessful())
 		{
 			UE_LOG(LogVirtualizationTool, Error, TEXT("The virtualization process failed with the following errors:"));
-			for (const FText& Error : ResultInfo.Errors)
+			for (const FText& Error : Result.Errors)
 			{
 				UE_LOG(LogVirtualizationTool, Error, TEXT("\t%s"), *Error.ToString());
 			}
@@ -181,7 +180,7 @@ bool FVirtualizeCommand::Run(const TArray<FProject>& Projects)
 
 		if (bShouldCheckout)
 		{
-			UE_LOG(LogVirtualizationTool, Display, TEXT("\t\t%d packages were checked out of revision control"), ResultInfo.CheckedOutPackages.Num());
+			UE_LOG(LogVirtualizationTool, Display, TEXT("\t\t%d packages were checked out of revision control"), Result.CheckedOutPackages.Num());
 		}
 
 		UE_LOG(LogVirtualizationTool, Display, TEXT("\t\tVirtualization of project packages complete!"), ProjectName.ToString());
diff --git a/Engine/Source/Runtime/Core/Private/Virtualization/VirtualizationSystem.cpp b/Engine/Source/Runtime/Core/Private/Virtualization/VirtualizationSystem.cpp
--- a/Engine/Source/Runtime/Core/Private/Virtualization/VirtualizationSystem.cpp
+++ b/Engine/Source/Runtime/Core/Private/Virtualization/VirtualizationSystem.cpp
@@ -111,20 +111,20 @@ public:
 		return EQueryResult::Failure_NotImplemented;
 	}
 
-	virtual EVirtualizationResult TryVirtualizePackages(TConstArrayView<FString> PackagePaths, EVirtualizationOptions Options, FVirtualizationResult& OutResultInfo) override
+	virtual FVirtualizationResult TryVirtualizePackages(TConstArrayView<FString> PackagePaths, EVirtualizationOptions Options) override
 	{
-		OutResultInfo.Reset();
-		OutResultInfo.Errors.Add(FText::FromString(TEXT("Calling ::TryVirtualizePackages on FNullVirtualizationSystem")));
+		FVirtualizationResult Result;
+		Result.AddError(FText::FromString(TEXT("Calling ::TryVirtualizePackages on FNullVirtualizationSystem")));
 
-		return EVirtualizationResult::Failed;
+		return Result;
 	}
 
-	virtual ERehydrationResult TryRehydratePackages(TConstArrayView<FString> PackagePaths, ERehydrationOptions Options, FRehydrationResult& OutResultInfo) override
+	virtual FRehydrationResult TryRehydratePackages(TConstArrayView<FString> PackagePaths, ERehydrationOptions Options) override
 	{
-		OutResultInfo.Reset();
-		OutResultInfo.Errors.Add(FText::FromString(TEXT("Calling ::TryRehydratePackages on FNullVirtualizationSystem")));
+		FRehydrationResult Result;
+		Result.AddError(FText::FromString(TEXT("Calling ::TryRehydratePackages on FNullVirtualizationSystem")));
 
-		return ERehydrationResult::Failed;
+		return Result;
 	}
 
 	virtual ERehydrationResult TryRehydratePackages(TConstArrayView<FString> PackagePaths, uint64 PaddingAlignment, TArray<FText>& OutErrors, TArray<FSharedBuffer>& OutPackages, TArray<FRehydrationInfo>* OutInfo) override
diff --git a/Engine/Source/Runtime/Core/Public/Virtualization/VirtualizationSystem.h b/Engine/Source/Runtime/Core/Public/Virtualization/VirtualizationSystem.h
--- a/Engine/Source/Runtime/Core/Public/Virtualization/VirtualizationSystem.h
+++ b/Engine/Source/Runtime/Core/Public/Virtualization/VirtualizationSystem.h
@@ -460,13 +460,10 @@ struct FVirtualizationResult
 	/** The length of time that the process took in seconds */
 	double TimeTaken = 0.0;
 
-	/** Returns members to the default state */
-	void Reset()
+	/* Returns if the virtualization process succeeded or not */
+	bool WasSuccessful() const
 	{
-		Errors.Reset();
-		DescriptionTags.Reset();
-		VirtualizedPackages.Reset();
-		CheckedOutPackages.Reset();
+		return Errors.IsEmpty();
 	}
 
 	/** Add an error message to the result */
@@ -513,10 +510,10 @@ struct FRehydrationResult
 	/** The length of time that the process took in seconds */
 	double TimeTaken = 0.0;
 
-	/** Returns members to the default state */
-	void Reset()
+	/* Returns if the rehydration process succeeded or not */
+	bool WasSuccessful() const
 	{
-		Errors.Reset();
+		return Errors.IsEmpty();
 	}
 
 	/** Add an error message to the result */
@@ -799,13 +796,12 @@ public:
 	UE_DEPRECATED(5.2, "Use The other overload that takes a EVirtualizationOptions parameter")
 	virtual EVirtualizationResult TryVirtualizePackages(TConstArrayView<FString> PackagePaths, TArray<FText>& OutDescriptionTags, TArray<FText>& OutErrors)
 	{
-		FVirtualizationResult Info;
-		EVirtualizationResult Result = TryVirtualizePackages(PackagePaths, EVirtualizationOptions::None, Info);
+		FVirtualizationResult Result = TryVirtualizePackages(PackagePaths, EVirtualizationOptions::None);
 
-		OutDescriptionTags = MoveTemp(Info.DescriptionTags);
-		OutErrors = MoveTemp(Info.Errors);
+		OutDescriptionTags = MoveTemp(Result.DescriptionTags);
+		OutErrors = MoveTemp(Result.Errors);
 
-		return Result;
+		return Result.WasSuccessful() ? EVirtualizationResult::Success : EVirtualizationResult::Failed;
 	}
 	
 	/**
@@ -820,17 +816,16 @@ public:
 	 * @return						A EVirtualizationResult enum with the status of the process. If the status is not EVirtualizationResult::Success
 	 *								then the parameter OutErrors should contain at least one entry.
 	 */
-	virtual EVirtualizationResult TryVirtualizePackages(TConstArrayView<FString> PackagePaths, EVirtualizationOptions Options, FVirtualizationResult& OutResultInfo) = 0;
+	virtual FVirtualizationResult TryVirtualizePackages(TConstArrayView<FString> PackagePaths, EVirtualizationOptions Options) = 0;
 
 	UE_DEPRECATED(5.2, "Use The other overload that takes a EVirtualizationOptions parameter")
 	virtual ERehydrationResult TryRehydratePackages(TConstArrayView<FString> PackagePaths, TArray<FText>& OutErrors)
 	{
-		FRehydrationResult Info;
-		ERehydrationResult Result = TryRehydratePackages(PackagePaths, ERehydrationOptions::None, Info);
+		FRehydrationResult Result = TryRehydratePackages(PackagePaths, ERehydrationOptions::None);
 
-		OutErrors = MoveTemp(Info.Errors);
+		OutErrors = MoveTemp(Result.Errors);
 
-		return Result;
+		return Result.WasSuccessful() ? ERehydrationResult::Success : ERehydrationResult::Failed;
 	}
 
 	/**
@@ -844,7 +839,7 @@ public:
 	 * @return	A ERehydrationResult enum with the status of the process. If the status indicates any sort of failure then OutResultInfo.Errors should
 	 *			contain at least one entry.
 	 */
-	virtual ERehydrationResult TryRehydratePackages(TConstArrayView<FString> PackagePaths, ERehydrationOptions Options, FRehydrationResult& OutResultInfo) = 0;
+	virtual FRehydrationResult TryRehydratePackages(TConstArrayView<FString> PackagePaths, ERehydrationOptions Options) = 0;
 
 	/**
 	 * Rehydrates a number of packages into memory buffers.
