Clean up some Seek/SeekFromEnd/ReadInternals calls

diff --git a/Engine/Source/Runtime/Core/Private/Unix/UnixPlatformFile.cpp b/Engine/Source/Runtime/Core/Private/Unix/UnixPlatformFile.cpp
--- a/Engine/Source/Runtime/Core/Private/Unix/UnixPlatformFile.cpp
+++ b/Engine/Source/Runtime/Core/Private/Unix/UnixPlatformFile.cpp
@@ -146,6 +146,12 @@ public:
 	{
 		check(NewPosition >= 0);
 
+		// Avoid allowing for a negative NewPosition as this will set FileOffset which is returned in Tell blindly
+		if (NewPosition < 0)
+		{
+			return false;
+		}
+
 		if (!FileOpenAsWrite)
 		{
 			FileOffset = NewPosition >= FileSize ? FileSize - 1 : NewPosition;
@@ -162,6 +168,13 @@ public:
 	{
 		check(NewPositionRelativeToEnd <= 0);
 
+		// Avoid allowing a relative position to set less then the size of the file
+		// lseek handles this negative but we return FileOffset blindly which could be used incorrectly
+		if (NewPositionRelativeToEnd < (-FileSize))
+		{
+			return false;
+		}
+
 		if (!FileOpenAsWrite)
 		{
 			FileOffset = (NewPositionRelativeToEnd >= FileSize) ? 0 : (FileSize + NewPositionRelativeToEnd - 1);
@@ -265,6 +278,12 @@ private:
 	{
 		check(IsValid());
 		int64 BytesRead = 0;
+
+		if (BytesToRead < 0)
+		{
+			return 0;
+		}
+
 		while (BytesToRead)
 		{
 			check(BytesToRead >= 0);
