Update iOS UIScreen.mainScreen access due to deprecation in iOS16. Also, dispatch to UI thread to avoid iOS runtime warning. Reported by licensee.

diff --git a/Engine/Source/Runtime/Core/Private/IOS/IOSPlatformMisc.cpp b/Engine/Source/Runtime/Core/Private/IOS/IOSPlatformMisc.cpp
--- a/Engine/Source/Runtime/Core/Private/IOS/IOSPlatformMisc.cpp
+++ b/Engine/Source/Runtime/Core/Private/IOS/IOSPlatformMisc.cpp
@@ -237,7 +237,7 @@ int FIOSPlatformMisc::GetBatteryLevel()
 float FIOSPlatformMisc::GetBrightness()
 {
 #if !PLATFORM_TVOS
-	return (float)[UIScreen mainScreen].brightness;
+	return (float)[[[IOSAppDelegate GetDelegate] window] screen].brightness;
 #else
 	return 1.0f;
 #endif // !PLATFORM_TVOS
@@ -246,7 +246,9 @@ float FIOSPlatformMisc::GetBrightness()
 void FIOSPlatformMisc::SetBrightness(float Brightness)
 {
 #if !PLATFORM_TVOS
-	[UIScreen mainScreen].brightness = Brightness;
+	dispatch_async(dispatch_get_main_queue(), ^{
+		[[[IOSAppDelegate GetDelegate] window] screen].brightness = Brightness;
+	});
 #endif // !PLATFORM_TVOS
 }
 
