Replaced FModuleManager's friendship of TOptional with a private token approach.

diff --git a/Engine/Source/Runtime/Core/Private/Modules/ModuleManager.cpp b/Engine/Source/Runtime/Core/Private/Modules/ModuleManager.cpp
--- a/Engine/Source/Runtime/Core/Private/Modules/ModuleManager.cpp
+++ b/Engine/Source/Runtime/Core/Private/Modules/ModuleManager.cpp
@@ -67,9 +67,9 @@ FModuleManager::ModuleInfoRef FModuleManager::FindModuleChecked(FName InModuleNa
 }
 
 // Function level static to allow lazy construction during static initialization
-static TOptional<FModuleManager>& GetModuleManagerSingleton()
+TOptional<FModuleManager>& UE::Core::Private::GetModuleManagerSingleton()
 {
-	static TOptional<FModuleManager> Singleton(InPlace);
+	static TOptional<FModuleManager> Singleton(InPlace, FModuleManager::FPrivateToken{});
 	return Singleton;
 }
 
@@ -77,15 +77,15 @@ static TOptional<FModuleManager>& GetModuleManagerSingleton()
 void FModuleManager::TearDown()
 {
 	check(IsInGameThread());
-	GetModuleManagerSingleton().Reset();
+	UE::Core::Private::GetModuleManagerSingleton().Reset();
 }
 
 FModuleManager& FModuleManager::Get()
 {
-	return GetModuleManagerSingleton().GetValue();
+	return UE::Core::Private::GetModuleManagerSingleton().GetValue();
 }
 
-FModuleManager::FModuleManager()
+FModuleManager::FModuleManager(FPrivateToken)
 	: bCanProcessNewlyLoadedObjects(false)
 	, bExtraBinarySearchPathsAdded(false)
 {
diff --git a/Engine/Source/Runtime/Core/Public/Modules/ModuleManager.h b/Engine/Source/Runtime/Core/Public/Modules/ModuleManager.h
--- a/Engine/Source/Runtime/Core/Public/Modules/ModuleManager.h
+++ b/Engine/Source/Runtime/Core/Public/Modules/ModuleManager.h
@@ -26,6 +26,7 @@
 #include "UObject/NameTypes.h"
 
 class FArchive;
+class FModuleManager;
 class FOutputDevice;
 class UClass;
 class UWorld;
@@ -127,6 +128,11 @@ struct FModuleStatus
 	bool bIsGameModule;
 };
 
+namespace UE::Core::Private
+{
+	TOptional<FModuleManager>& GetModuleManagerSingleton();
+}
+
 /**
  * Implements the module manager.
  *
@@ -136,14 +142,14 @@ struct FModuleStatus
 class CORE_API FModuleManager
 	: private FSelfRegisteringExec
 {
+public:
+
 	/**
 	 * Destructor.
 	 */
 	~FModuleManager();
 
 
-public:
-
 	/**
 	 * Gets the singleton instance of the module manager.
 	 *
@@ -507,17 +513,20 @@ public:
 
 
 private:
-	friend struct TOptional<FModuleManager>;
+	friend TOptional<FModuleManager>& UE::Core::Private::GetModuleManagerSingleton();
+	struct FPrivateToken { explicit FPrivateToken() = default; };
 
+public:
 	/**
 	 * Hidden constructor.
 	 *
 	 * Use the static Get function to return the singleton instance.
 	 */
-	FModuleManager();
+	FModuleManager(FPrivateToken);
 	FModuleManager(const FModuleManager&) = delete;
 	FModuleManager& operator=(const FModuleManager&) = delete;
 
+private:
 	/**
 	 * Information about a single module (may or may not be loaded.)
 	 */
