Set Mac ARM cache line size to 128 to cater for Apple M1/M2

diff --git a/Engine/Source/Runtime/Core/Public/Mac/MacPlatform.h b/Engine/Source/Runtime/Core/Public/Mac/MacPlatform.h
--- a/Engine/Source/Runtime/Core/Public/Mac/MacPlatform.h
+++ b/Engine/Source/Runtime/Core/Public/Mac/MacPlatform.h
@@ -127,6 +127,12 @@ typedef FMacPlatformTypes FPlatformTypes;
 #define OPERATOR_NEW_NOTHROW_SPEC  _NOEXCEPT
 #define OPERATOR_DELETE_NOTHROW_SPEC  _NOEXCEPT
 
+#ifdef PLATFORM_MAC_ARM64
+#define PLATFORM_CACHE_LINE_SIZE	128
+#else
+#define PLATFORM_CACHE_LINE_SIZE	64
+#endif
+
 // DLL export and import definitions
 #define DLLEXPORT			__attribute__((visibility("default")))
 #define DLLIMPORT			__attribute__((visibility("default")))
