Ensure a page in TPagedArray never shrinks below PageTraits::Capacity in TPagedArray::SetNum

diff --git a/Engine/Source/Runtime/Core/Public/Containers/PagedArray.h b/Engine/Source/Runtime/Core/Public/Containers/PagedArray.h
--- a/Engine/Source/Runtime/Core/Public/Containers/PagedArray.h
+++ b/Engine/Source/Runtime/Core/Public/Containers/PagedArray.h
@@ -248,7 +248,7 @@ public:
 			if (PendingCount)
 			{
 				Pages[PageIndex].Reserve(PageTraits::Capacity);
-				Pages[PageIndex].SetNum(PendingCount);
+				Pages[PageIndex].SetNum(PendingCount, false);
 			}
 		}
 		else if (NewNum < Num())
@@ -257,7 +257,7 @@ public:
 			Pages.SetNum(RequiredPageCount, bAllowShrinking);
 			if (const SizeType Mod = NewNum % PageTraits::Capacity)
 			{
-				Pages.Last().SetNum(Mod);
+				Pages.Last().SetNum(Mod, false);
 			}
 		}
 		Count = NewNum;
