Added GetAllocatedSize functions to TDeque and TPagedArray to aid with memory tracking

diff --git a/Engine/Source/Runtime/Core/Public/Containers/Deque.h b/Engine/Source/Runtime/Core/Public/Containers/Deque.h
--- a/Engine/Source/Runtime/Core/Public/Containers/Deque.h
+++ b/Engine/Source/Runtime/Core/Public/Containers/Deque.h
@@ -206,6 +206,17 @@ public:
 		return Count;
 	}
 
+	/**
+	 * Helper function to return the amount of memory allocated by this container.
+	 * Only returns the size of allocations made directly by the container, not the elements themselves.
+	 *
+	 * @returns Number of bytes allocated by this container.
+	 */
+	FORCEINLINE SIZE_T GetAllocatedSize() const
+	{
+		return Storage.GetAllocatedSize(Capacity, sizeof(ElementType));
+	}
+
 	/*
 	 * Constructs an element in place using the parameter arguments and adds it at the back of the queue.
 	 * This method returns a reference to the constructed element.
diff --git a/Engine/Source/Runtime/Core/Public/Containers/PagedArray.h b/Engine/Source/Runtime/Core/Public/Containers/PagedArray.h
--- a/Engine/Source/Runtime/Core/Public/Containers/PagedArray.h
+++ b/Engine/Source/Runtime/Core/Public/Containers/PagedArray.h
@@ -194,6 +194,23 @@ public:
 		return Pages[GetPageIndex(Index)][GetPageOffset(Index)];
 	}
 
+	/**
+	 * Helper function to return the amount of memory allocated by this container.
+	 * Only returns the size of allocations made directly by the container, not the elements themselves.
+	 *
+	 * @returns Number of bytes allocated by this container.
+	 */
+	FORCEINLINE SIZE_T GetAllocatedSize() const
+	{
+		SIZE_T Size = 0;
+		Size += Pages.GetAllocatedSize();
+		for (const PageType& Page : Pages)
+		{
+			Size += Page.GetAllocatedSize();
+		}
+		return Size;
+	}
+
 	FORCEINLINE SizeType Max() const
 	{
 		return Pages.Num() * PageTraits::Capacity;
