deprecated a couple of stats methods that take TArray arg and provided a version with TArray64 to handle big stats files

diff --git a/Engine/Source/Developer/Profiler/Private/ProfilerRawStatsForThreadView.cpp b/Engine/Source/Developer/Profiler/Private/ProfilerRawStatsForThreadView.cpp
--- a/Engine/Source/Developer/Profiler/Private/ProfilerRawStatsForThreadView.cpp
+++ b/Engine/Source/Developer/Profiler/Private/ProfilerRawStatsForThreadView.cpp
@@ -168,7 +168,7 @@ void FRawProfilerSession::PrepareLoading()
 	if( Stream.Header.bRawStatsFile )
 	{
 		// Read metadata.
-		TArray<FStatMessage> MetadataMessages;
+		TArray64<FStatMessage> MetadataMessages;
 		Stream.ReadFNamesAndMetadataMessages( *FileReader, MetadataMessages );
 		StatsThreadStats.ProcessMetaDataOnly( MetadataMessages );
 
diff --git a/Engine/Source/Runtime/Core/Private/Stats/Stats2.cpp b/Engine/Source/Runtime/Core/Private/Stats/Stats2.cpp
--- a/Engine/Source/Runtime/Core/Private/Stats/Stats2.cpp
+++ b/Engine/Source/Runtime/Core/Private/Stats/Stats2.cpp
@@ -968,7 +968,7 @@ private:
 			INC_DWORD_STAT_BY(STAT_StatFramePacketsRecv, NowData.Packets.Num());
 			{
 				SCOPE_CYCLE_COUNTER(STAT_StatsNewParseMeta);
-				TArray<FStatMessage> MetaMessages;
+				TArray64<FStatMessage> MetaMessages;
 				{
 					FScopeLock Lock(&FStartupMessages::Get().CriticalSection);
 					Exchange(FStartupMessages::Get().DelayedMessages, MetaMessages);
diff --git a/Engine/Source/Runtime/Core/Private/Stats/StatsData.cpp b/Engine/Source/Runtime/Core/Private/Stats/StatsData.cpp
--- a/Engine/Source/Runtime/Core/Private/Stats/StatsData.cpp
+++ b/Engine/Source/Runtime/Core/Private/Stats/StatsData.cpp
@@ -807,6 +807,16 @@ void FStatsThreadState::ProcessMetaDataOnly(TArray<FStatMessage>& Data)
 	}
 }
 
+void FStatsThreadState::ProcessMetaDataOnly(TArray64<FStatMessage>& Data)
+{
+	for (int64 Index = 0; Index < Data.Num(); Index++)
+	{
+		FStatMessage& Item = Data[Index];
+		EStatOperation::Type Op = Item.NameAndInfo.GetField<EStatOperation>();
+		check(Op == EStatOperation::SetLongName);
+		FindOrAddMetaData(Item);
+	}
+}
 
 void FStatsThreadState::ToggleFindMemoryExtensiveStats()
 {
diff --git a/Engine/Source/Runtime/Core/Private/Stats/StatsFile.cpp b/Engine/Source/Runtime/Core/Private/Stats/StatsFile.cpp
--- a/Engine/Source/Runtime/Core/Private/Stats/StatsFile.cpp
+++ b/Engine/Source/Runtime/Core/Private/Stats/StatsFile.cpp
@@ -611,7 +611,7 @@ bool FStatsReadFile::PrepareLoading()
 	}
 
 	// Read metadata.
-	TArray<FStatMessage> MetadataMessages;
+	TArray64<FStatMessage> MetadataMessages;
 	Stream.ReadFNamesAndMetadataMessages( *Reader, MetadataMessages );
 	State.ProcessMetaDataOnly( MetadataMessages );
 
diff --git a/Engine/Source/Runtime/Core/Public/Stats/Stats2.h b/Engine/Source/Runtime/Core/Public/Stats/Stats2.h
--- a/Engine/Source/Runtime/Core/Public/Stats/Stats2.h
+++ b/Engine/Source/Runtime/Core/Public/Stats/Stats2.h
@@ -1775,7 +1775,7 @@ class FStartupMessages
 {
 	friend class FStatsThread;
 
-	TArray<FStatMessage> DelayedMessages;
+	TArray64<FStatMessage> DelayedMessages;
 	FCriticalSection CriticalSection;
 
 public:
diff --git a/Engine/Source/Runtime/Core/Public/Stats/StatsData.h b/Engine/Source/Runtime/Core/Public/Stats/StatsData.h
--- a/Engine/Source/Runtime/Core/Public/Stats/StatsData.h
+++ b/Engine/Source/Runtime/Core/Public/Stats/StatsData.h
@@ -513,8 +513,11 @@ class CORE_API FStatsThreadState
 
 public:
 	/** Internal method to update the internal metadata. **/
+	UE_DEPRECATED(5.2, "Please use the overload that takes TArray64.")
 	void ProcessMetaDataOnly(TArray<FStatMessage>& Data);
 
+	void ProcessMetaDataOnly(TArray64<FStatMessage>& Data);
+
 	/** Toggles tracking the most memory expensive stats. */
 	void ToggleFindMemoryExtensiveStats();
 
diff --git a/Engine/Source/Runtime/Core/Public/Stats/StatsFile.h b/Engine/Source/Runtime/Core/Public/Stats/StatsFile.h
--- a/Engine/Source/Runtime/Core/Public/Stats/StatsFile.h
+++ b/Engine/Source/Runtime/Core/Public/Stats/StatsFile.h
@@ -806,6 +806,7 @@ public:
 	 *	Reads FNames and metadata messages from the specified archive, only valid for finalized stats files.
 	 *	Allow unordered file access.
 	 */
+	UE_DEPRECATED(5.2, "Please use the overload that takes TArray64.")
 	void ReadFNamesAndMetadataMessages( FArchive& Ar, TArray<FStatMessage>& out_MetadataMessages )
 	{
 		// Read FNames.
@@ -824,6 +825,25 @@ public:
 			new(out_MetadataMessages)FStatMessage( ReadMessage( Ar, false ) );
 		}
 	}
+
+	void ReadFNamesAndMetadataMessages(FArchive& Ar, TArray64<FStatMessage>& out_MetadataMessages)
+	{
+		// Read FNames.
+		Ar.Seek(Header.FNameTableOffset);
+		out_MetadataMessages.Reserve(Header.NumFNames);
+		for (uint64 Index = 0; Index < Header.NumFNames; Index++)
+		{
+			ReadFName(Ar, false);
+		}
+
+		// Read metadata messages.
+		Ar.Seek(Header.MetadataMessagesOffset);
+		out_MetadataMessages.Reserve(Header.NumMetadataMessages);
+		for (uint64 Index = 0; Index < Header.NumMetadataMessages; Index++)
+		{
+			new(out_MetadataMessages)FStatMessage(ReadMessage(Ar, false));
+		}
+	}
 };
 
 /** Raw stats information. */
