Temporarily backing out the fix to ParallelFor task cancellation because it had the unexpected side effect of degrading performance for some workloads

#rnx

diff --git a/Engine/Source/Runtime/Core/Public/Async/ParallelFor.h b/Engine/Source/Runtime/Core/Public/Async/ParallelFor.h
--- a/Engine/Source/Runtime/Core/Public/Async/ParallelFor.h
+++ b/Engine/Source/Runtime/Core/Public/Async/ParallelFor.h
@@ -237,9 +237,8 @@ namespace ParallelForImpl
 		class FParallelExecutor 
 		{
 			FDataHandle Data;
-			int32 WorkerIndex;
+			mutable int32 WorkerIndex;
 			LowLevelTasks::ETaskPriority Priority;
-			mutable bool bReschedule = false;
 
 		public:
 			inline FParallelExecutor(FDataHandle&& InData, int32 InWorkerIndex, LowLevelTasks::ETaskPriority InPriority) 
@@ -259,7 +258,7 @@ namespace ParallelForImpl
 
 			~FParallelExecutor()
 			{
-				if (Data.IsValid() && bReschedule)
+				if (Data.IsValid() && WorkerIndex >= 0)
 				{
 					FParallelExecutor::LaunchTask(nullptr, MoveTemp(Data), WorkerIndex, Priority);
 				}
@@ -324,6 +323,7 @@ namespace ParallelForImpl
 					{
 						if (!bIsMaster)
 						{
+							WorkerIndex = -1;
 							return false;
 						}
 						BatchIndex = (NumBatches - 1);
@@ -344,10 +344,12 @@ namespace ParallelForImpl
 						{
 							Data->FinishedSignal->Trigger();
 						}
+						WorkerIndex = -1;
 						return true;
 					}
 					else if (EndIndex >= Num)
 					{
+						WorkerIndex = -1;
 						return false;
 					}
 					else if (!bIsBackgroundPriority)
@@ -358,8 +360,7 @@ namespace ParallelForImpl
 					auto PassedTime = [Start, &Now]() { return Now() - Start; };
 					if (PassedTime() > YieldingThreshold)
 					{
-						//abort and reschedule (in the destructor) to give higher priority tasks a chance to run
-						bReschedule = true;
+						//abort and reschedule (in the destructor as WorkerIndex is larger_eq Zero) to give higher priority tasks a chance to run
 						return false;
 					}
 				}
@@ -406,7 +407,7 @@ namespace ParallelForImpl
 			TRACE_CPUPROFILER_EVENT_SCOPE(ParallelFor.Cancel);
 			for (FTracedTask& TracedTask : LocalExecutor.GetData()->Tasks)
 			{
-				TracedTask.Task.TryCancel(LowLevelTasks::ECancellationFlags::None);
+				TracedTask.Task.TryCancel(LowLevelTasks::ECancellationFlags::PrelaunchCancellation);
 			}
 		}
 
