Deprecated TAreTypesEqual and ARE_TYPES_EQUAL.

diff --git a/Engine/Plugins/AI/MassAI/Source/MassZoneGraphNavigation/Private/MassZoneGraphNavigationFragments.cpp b/Engine/Plugins/AI/MassAI/Source/MassZoneGraphNavigation/Private/MassZoneGraphNavigationFragments.cpp
--- a/Engine/Plugins/AI/MassAI/Source/MassZoneGraphNavigation/Private/MassZoneGraphNavigationFragments.cpp
+++ b/Engine/Plugins/AI/MassAI/Source/MassZoneGraphNavigation/Private/MassZoneGraphNavigationFragments.cpp
@@ -166,7 +166,7 @@ bool FMassZoneGraphShortPathFragment::RequestPath(const FMassZoneGraphCachedLane
 
 
 	// The current distance can come from a quantized lane distance. Check against quantized bounds, but clamp it to the actual path length when calculating the path.
-	static_assert(TAreTypesEqual<decltype(FMassZoneGraphPathPoint::Distance), FMassInt16Real>::Value, "Assuming FMassZoneGraphPathPoint::Distance is quantized to 10 units.");
+	static_assert(std::is_same_v<decltype(FMassZoneGraphPathPoint::Distance), FMassInt16Real>, "Assuming FMassZoneGraphPathPoint::Distance is quantized to 10 units.");
 	const float LaneLengthQuantized = FMath::CeilToFloat(CachedLane.LaneLength / 10.0f) * 10.0f;
 
 	static constexpr float Epsilon = 0.1f;
diff --git a/Engine/Plugins/Tests/RHITests/Source/RHITests/Public/RHIBufferTests.h b/Engine/Plugins/Tests/RHITests/Source/RHITests/Public/RHIBufferTests.h
--- a/Engine/Plugins/Tests/RHITests/Source/RHITests/Public/RHIBufferTests.h
+++ b/Engine/Plugins/Tests/RHITests/Source/RHITests/Public/RHIBufferTests.h
@@ -80,7 +80,7 @@ class FRHIBufferTests
 		});
 
 		FString ClearValueStr;
-		if (TAreTypesEqual<ValueType, FVector4f>::Value)
+		if constexpr (std::is_same_v<ValueType, FVector4f>)
 		{
 			ClearValueStr = FString::Printf(TEXT("%f %f %f %f"), ClearValue.X, ClearValue.Y, ClearValue.Z, ClearValue.W);
 		}
diff --git a/Engine/Plugins/Tests/RHITests/Source/RHITests/Public/RHITestsCommon.h b/Engine/Plugins/Tests/RHITests/Source/RHITests/Public/RHITestsCommon.h
--- a/Engine/Plugins/Tests/RHITests/Source/RHITests/Public/RHITestsCommon.h
+++ b/Engine/Plugins/Tests/RHITests/Source/RHITests/Public/RHITestsCommon.h
@@ -17,7 +17,7 @@ bool RunOnRenderThreadSynchronous(TFunctionRef<bool(FRHICommandListImmediate&)>
 template <typename ValueType>
 static inline FString ClearValueToString(const ValueType & ClearValue)
 {
-	if (TAreTypesEqual<ValueType, FVector4>::Value)
+	if constexpr (std::is_same_v<ValueType, FVector4>)
 	{
 		return FString::Printf(TEXT("%f %f %f %f"), ClearValue.X, ClearValue.Y, ClearValue.Z, ClearValue.W);
 	}
diff --git a/Engine/Source/Editor/UnrealEd/Private/Commandlets/AssetRegistryGenerator.cpp b/Engine/Source/Editor/UnrealEd/Private/Commandlets/AssetRegistryGenerator.cpp
--- a/Engine/Source/Editor/UnrealEd/Private/Commandlets/AssetRegistryGenerator.cpp
+++ b/Engine/Source/Editor/UnrealEd/Private/Commandlets/AssetRegistryGenerator.cpp
@@ -41,7 +41,6 @@
 #include "Settings/ProjectPackagingSettings.h"
 #include "Stats/StatsMisc.h"
 #include "String/ParseTokens.h"
-#include "Templates/AreTypesEqual.h"
 #include "Templates/UniquePtr.h"
 #include "UObject/SoftObjectPath.h"
 
@@ -2675,7 +2674,7 @@ void FAssetRegistryPackageMessage::Write(FCbWriter& Writer, const FPackageData&
 	// We cast TMap<FSoftObjectPath,                     ValueType>
 	//      to TMap<FSoftObjectPathSerializationWrapper, ValueType>
 	// to workaround FSoftObjectPath's implicit constructor. See comment in CompactBinaryTCP.h
-	static_assert(ARE_TYPES_EQUAL(decltype(CookTags)::KeyType, FSoftObjectPath), "Expected KeyType of CookTags to be FSoftObjectPath");
+	static_assert(std::is_same_v<decltype(CookTags)::KeyType, FSoftObjectPath>, "Expected KeyType of CookTags to be FSoftObjectPath");
 	Writer << "T" << reinterpret_cast<const TMap<FSoftObjectPathSerializationWrapper, decltype(CookTags)::ValueType>&>(CookTags);
 	Writer << "F" << PackageFlags;
 	Writer << "S" << DiskSize;
@@ -2700,7 +2699,7 @@ bool FAssetRegistryPackageMessage::TryRead(FCbObjectView Object, FPackageData& P
 		}
 	}
 
-	static_assert(ARE_TYPES_EQUAL(decltype(CookTags)::KeyType, FSoftObjectPath), "Expected KeyType of CookTags to be FSoftObjectPath");
+	static_assert(std::is_same_v<decltype(CookTags)::KeyType, FSoftObjectPath>, "Expected KeyType of CookTags to be FSoftObjectPath");
 	if (!LoadFromCompactBinary(Object["T"], reinterpret_cast<TMap<FSoftObjectPathSerializationWrapper, decltype(CookTags)::ValueType>&>(CookTags)))
 	{
 		return false;
diff --git a/Engine/Source/Editor/UnrealEd/Private/UnrealEdPrivatePCH.h b/Engine/Source/Editor/UnrealEd/Private/UnrealEdPrivatePCH.h
--- a/Engine/Source/Editor/UnrealEd/Private/UnrealEdPrivatePCH.h
+++ b/Engine/Source/Editor/UnrealEd/Private/UnrealEdPrivatePCH.h
@@ -23,7 +23,6 @@
 #include "Misc/VarArgs.h"
 #include "HAL/PlatformAtomics.h"
 #include "GenericPlatform/GenericPlatformAtomics.h"
-#include "Templates/AreTypesEqual.h"
 #include "Templates/UnrealTypeTraits.h"
 #include "Templates/AndOrNot.h"
 #include "Templates/IsArithmetic.h"
diff --git a/Engine/Source/Editor/UnrealEd/Public/UnrealEdSharedPCH.h b/Engine/Source/Editor/UnrealEd/Public/UnrealEdSharedPCH.h
--- a/Engine/Source/Editor/UnrealEd/Public/UnrealEdSharedPCH.h
+++ b/Engine/Source/Editor/UnrealEd/Public/UnrealEdSharedPCH.h
@@ -23,7 +23,6 @@
 #include "Misc/VarArgs.h"
 #include "HAL/PlatformAtomics.h"
 #include "GenericPlatform/GenericPlatformAtomics.h"
-#include "Templates/AreTypesEqual.h"
 #include "Templates/UnrealTypeTraits.h"
 #include "Templates/AndOrNot.h"
 #include "Templates/IsArithmetic.h"
diff --git a/Engine/Source/Runtime/ClothingSystemRuntimeCommon/Public/Utils/ClothingMeshUtils.h b/Engine/Source/Runtime/ClothingSystemRuntimeCommon/Public/Utils/ClothingMeshUtils.h
--- a/Engine/Source/Runtime/ClothingSystemRuntimeCommon/Public/Utils/ClothingMeshUtils.h
+++ b/Engine/Source/Runtime/ClothingSystemRuntimeCommon/Public/Utils/ClothingMeshUtils.h
@@ -277,7 +277,7 @@ namespace ClothingMeshUtils
 		{
 			// Enforce the interp func signature (returns T and takes a bary and 3 Ts)
 			// If you hit this then either the return type isn't T or your arguments aren't convertible to T
-			static_assert(TAreTypesEqual<T, typename TDecay<decltype(Func(DeclVal<FVector3f>(), DeclVal<T>(), DeclVal<T>(), DeclVal<T>()))>::Type>::Value, "Invalid Lambda signature passed to Map");
+			static_assert(std::is_same_v<T, typename TDecay<decltype(Func(DeclVal<FVector3f>(), DeclVal<T>(), DeclVal<T>(), DeclVal<T>()))>::Type>, "Invalid Lambda signature passed to Map");
 
 			const int32 NumMesh0Positions = Mesh0Positions.Num();
 			const int32 NumMesh0Normals = Mesh0Normals.Num();
diff --git a/Engine/Source/Runtime/Core/Public/Containers/Array.h b/Engine/Source/Runtime/Core/Public/Containers/Array.h
--- a/Engine/Source/Runtime/Core/Public/Containers/Array.h
+++ b/Engine/Source/Runtime/Core/Public/Containers/Array.h
@@ -5,7 +5,6 @@
 #include "CoreTypes.h"
 #include "Misc/AssertionMacros.h"
 #include "HAL/UnrealMemory.h"
-#include "Templates/AreTypesEqual.h"
 #include "Templates/IsSigned.h"
 #include "Templates/UnrealTypeTraits.h"
 #include "Templates/UnrealTemplate.h"
@@ -280,23 +279,25 @@ namespace UE4Array_Private
 	}
 
 	template <typename FromArrayType, typename ToArrayType>
-	struct TCanMoveTArrayPointersBetweenArrayTypes
+	constexpr bool CanMoveTArrayPointersBetweenArrayTypes()
 	{
 		typedef typename FromArrayType::AllocatorType FromAllocatorType;
-		typedef typename ToArrayType  ::AllocatorType ToAllocatorType;
+		typedef typename ToArrayType::AllocatorType ToAllocatorType;
 		typedef typename FromArrayType::ElementType   FromElementType;
-		typedef typename ToArrayType  ::ElementType   ToElementType;
+		typedef typename ToArrayType::ElementType   ToElementType;
 
-		enum
+		// Allocators must be equal or move-compatible...
+		if constexpr (std::is_same_v<FromAllocatorType, ToAllocatorType> || TCanMoveBetweenAllocators<FromAllocatorType, ToAllocatorType>::Value)
 		{
-			Value =
-				TOr<TAreTypesEqual<FromAllocatorType, ToAllocatorType>, TCanMoveBetweenAllocators<FromAllocatorType, ToAllocatorType>>::Value && // Allocators must be equal or move-compatible
-				(
-					TAreTypesEqual         <ToElementType, FromElementType>::Value || // The element type of the container must be the same, or...
-					TIsBitwiseConstructible<ToElementType, FromElementType>::Value    // ... the element type of the source container must be bitwise constructible from the element type in the destination container
-				)
-		};
-	};
+			return
+				std::is_same_v         <ToElementType, FromElementType> ||      // The element type of the container must be the same, or...
+				TIsBitwiseConstructible<ToElementType, FromElementType>::Value; // ... the element type of the source container must be bitwise constructible from the element type in the destination container
+		}
+		else
+		{
+			return false;
+		}
+	}
 
 	// Assume elements are compatible with themselves - avoids problems with generated copy
 	// constructors of arrays of forwarded types, e.g.:
@@ -521,45 +522,45 @@ private:
 	template <typename FromArrayType, typename ToArrayType>
 	static FORCEINLINE void MoveOrCopy(ToArrayType& ToArray, FromArrayType& FromArray, SizeType PrevMax)
 	{
-		if constexpr (UE4Array_Private::TCanMoveTArrayPointersBetweenArrayTypes<FromArrayType, ToArrayType>::Value)
+		if constexpr (UE4Array_Private::CanMoveTArrayPointersBetweenArrayTypes<FromArrayType, ToArrayType>())
 		{
 			// Move
 
-		static_assert(std::is_same_v<TArray, ToArrayType>, "MoveOrCopy is expected to be called with the current array type as the destination");
+			static_assert(std::is_same_v<TArray, ToArrayType>, "MoveOrCopy is expected to be called with the current array type as the destination");
 
-		using FromAllocatorType = typename FromArrayType::AllocatorType;
-		using ToAllocatorType   = typename ToArrayType::AllocatorType;
+			using FromAllocatorType = typename FromArrayType::AllocatorType;
+			using ToAllocatorType   = typename ToArrayType::AllocatorType;
 
-		if constexpr (TCanMoveBetweenAllocators<FromAllocatorType, ToAllocatorType>::Value)
-		{
-			ToArray.AllocatorInstance.template MoveToEmptyFromOtherAllocator<FromAllocatorType>(FromArray.AllocatorInstance);
-		}
-		else
-		{
-			ToArray.AllocatorInstance.MoveToEmpty(FromArray.AllocatorInstance);
-		}
+			if constexpr (TCanMoveBetweenAllocators<FromAllocatorType, ToAllocatorType>::Value)
+			{
+				ToArray.AllocatorInstance.template MoveToEmptyFromOtherAllocator<FromAllocatorType>(FromArray.AllocatorInstance);
+			}
+			else
+			{
+				ToArray.AllocatorInstance.MoveToEmpty(FromArray.AllocatorInstance);
+			}
 
-		ToArray  .ArrayNum = (SizeType)FromArray.ArrayNum;
-		ToArray  .ArrayMax = (SizeType)FromArray.ArrayMax;
+			ToArray  .ArrayNum = (SizeType)FromArray.ArrayNum;
+			ToArray  .ArrayMax = (SizeType)FromArray.ArrayMax;
 
-		// Ensure the destination container could hold the source range (when the allocator size types shrink)
-		if constexpr (sizeof(USizeType) < sizeof(typename FromArrayType::USizeType))
-		{
-			if (ToArray.ArrayNum != FromArray.ArrayNum || ToArray.ArrayMax != FromArray.ArrayMax)
+			// Ensure the destination container could hold the source range (when the allocator size types shrink)
+			if constexpr (sizeof(USizeType) < sizeof(typename FromArrayType::USizeType))
 			{
-				OnInvalidNum((USizeType)ToArray.ArrayNum);
+				if (ToArray.ArrayNum != FromArray.ArrayNum || ToArray.ArrayMax != FromArray.ArrayMax)
+				{
+					OnInvalidNum((USizeType)ToArray.ArrayNum);
+				}
 			}
-		}
 
-		FromArray.ArrayNum = 0;
-		FromArray.ArrayMax = FromArray.AllocatorInstance.GetInitialCapacity();
-	}
+			FromArray.ArrayNum = 0;
+			FromArray.ArrayMax = FromArray.AllocatorInstance.GetInitialCapacity();
+		}
 		else
 		{
 			// Copy
 
-		ToArray.CopyToEmpty(FromArray.GetData(), FromArray.Num(), PrevMax);
-	}
+			ToArray.CopyToEmpty(FromArray.GetData(), FromArray.Num(), PrevMax);
+		}
 	}
 
 	/**
@@ -574,29 +575,29 @@ private:
 	template <typename FromArrayType, typename ToArrayType>
 	static FORCEINLINE void MoveOrCopyWithSlack(ToArrayType& ToArray, FromArrayType& FromArray, SizeType PrevMax, SizeType ExtraSlack)
 	{
-		if constexpr (UE4Array_Private::TCanMoveTArrayPointersBetweenArrayTypes<FromArrayType, ToArrayType>::Value)
+		if constexpr (UE4Array_Private::CanMoveTArrayPointersBetweenArrayTypes<FromArrayType, ToArrayType>())
 		{
 			// Move
 
-		MoveOrCopy(ToArray, FromArray, PrevMax);
+			MoveOrCopy(ToArray, FromArray, PrevMax);
 
-		USizeType LocalArrayNum = (USizeType)ToArray.ArrayNum;
-		USizeType NewMax        = (USizeType)LocalArrayNum + (USizeType)ExtraSlack;
+			USizeType LocalArrayNum = (USizeType)ToArray.ArrayNum;
+			USizeType NewMax        = (USizeType)LocalArrayNum + (USizeType)ExtraSlack;
 
-		// This should only happen when we've underflowed or overflowed SizeType
-		if ((SizeType)NewMax < LocalArrayNum)
-		{
-			OnInvalidNum((USizeType)ExtraSlack);
-		}
+			// This should only happen when we've underflowed or overflowed SizeType
+			if ((SizeType)NewMax < LocalArrayNum)
+			{
+				OnInvalidNum((USizeType)ExtraSlack);
+			}
 
-		ToArray.Reserve(NewMax);
-	}
+			ToArray.Reserve(NewMax);
+		}
 		else
 		{
 			// Copy
 
-		ToArray.CopyToEmptyWithSlack(FromArray.GetData(), FromArray.Num(), PrevMax, ExtraSlack);
-	}
+			ToArray.CopyToEmptyWithSlack(FromArray.GetData(), FromArray.Num(), PrevMax, ExtraSlack);
+		}
 	}
 
 public:
@@ -1700,7 +1701,7 @@ public:
 	template <typename CountType>
 	FORCEINLINE void RemoveAt(SizeType Index, CountType Count, bool bAllowShrinking = true)
 	{
-		static_assert(!TAreTypesEqual<CountType, bool>::Value, "TArray::RemoveAt: unexpected bool passed as the Count argument");
+		static_assert(!std::is_same_v<CountType, bool>, "TArray::RemoveAt: unexpected bool passed as the Count argument");
 		RemoveAtImpl(Index, (SizeType)Count, bAllowShrinking);
 	}
 
@@ -1765,7 +1766,7 @@ public:
 	template <typename CountType>
 	FORCEINLINE void RemoveAtSwap(SizeType Index, CountType Count, bool bAllowShrinking = true)
 	{
-		static_assert(!TAreTypesEqual<CountType, bool>::Value, "TArray::RemoveAtSwap: unexpected bool passed as the Count argument");
+		static_assert(!std::is_same_v<CountType, bool>, "TArray::RemoveAtSwap: unexpected bool passed as the Count argument");
 		RemoveAtSwapImpl(Index, Count, bAllowShrinking);
 	}
 
diff --git a/Engine/Source/Runtime/Core/Public/Containers/Set.h b/Engine/Source/Runtime/Core/Public/Containers/Set.h
--- a/Engine/Source/Runtime/Core/Public/Containers/Set.h
+++ b/Engine/Source/Runtime/Core/Public/Containers/Set.h
@@ -16,7 +16,6 @@
 #include <initializer_list>
 #include "Templates/TypeHash.h"
 #include "Containers/SparseArray.h"
-#include "Templates/AreTypesEqual.h"
 #include "Templates/Decay.h"
 #include "Serialization/StructuredArchive.h"
 #include "Serialization/MemoryImageWriter.h"
@@ -165,7 +164,14 @@ public:
 	FORCEINLINE TSetElementBase() {}
 
 	/** Initialization constructor. */
-	template <typename InitType, typename = std::enable_if_t<!TAreTypesEqual<TSetElementBase, typename TDecay<InitType>::Type>::Value>> explicit FORCEINLINE TSetElementBase(InitType&& InValue) : Value(Forward<InitType>(InValue)) {}
+	template <
+		typename InitType,
+		typename = std::enable_if_t<!std::is_same_v<TSetElementBase, typename TDecay<InitType>::Type>>
+	>
+	explicit FORCEINLINE TSetElementBase(InitType&& InValue)
+		: Value(Forward<InitType>(InValue))
+	{
+	}
 
 	TSetElementBase(TSetElementBase&&) = default;
 	TSetElementBase(const TSetElementBase&) = default;
@@ -192,7 +198,14 @@ public:
 	FORCEINLINE TSetElementBase() {}
 
 	/** Initialization constructor. */
-	template <typename InitType, typename = std::enable_if_t<!TAreTypesEqual<TSetElementBase, typename TDecay<InitType>::Type>::Value>> explicit FORCEINLINE TSetElementBase(InitType&& InValue) : Value(Forward<InitType>(InValue)) {}
+	template <
+		typename InitType,
+		typename = std::enable_if_t<!std::is_same_v<TSetElementBase, typename TDecay<InitType>::Type>>
+	>
+	explicit FORCEINLINE TSetElementBase(InitType&& InValue)
+		: Value(Forward<InitType>(InValue))
+	{
+	}
 
 	TSetElementBase(TSetElementBase&&) = default;
 	TSetElementBase(const TSetElementBase&) = default;
@@ -220,7 +233,14 @@ public:
 	{}
 
 	/** Initialization constructor. */
-	template <typename InitType, typename = std::enable_if_t<!TAreTypesEqual<TSetElement, typename TDecay<InitType>::Type>::Value>> explicit FORCEINLINE TSetElement(InitType&& InValue) : Super(Forward<InitType>(InValue)) {}
+	template <
+		typename InitType,
+		typename = std::enable_if_t<!std::is_same_v<TSetElement, typename TDecay<InitType>::Type>>
+	>
+	explicit FORCEINLINE TSetElement(InitType&& InValue)
+		: Super(Forward<InitType>(InValue))
+	{
+	}
 
 	TSetElement(TSetElement&&) = default;
 	TSetElement(const TSetElement&) = default;
diff --git a/Engine/Source/Runtime/Core/Public/Containers/Union.h b/Engine/Source/Runtime/Core/Public/Containers/Union.h
--- a/Engine/Source/Runtime/Core/Public/Containers/Union.h
+++ b/Engine/Source/Runtime/Core/Public/Containers/Union.h
@@ -8,7 +8,6 @@
 #include "Logging/LogMacros.h"
 #include "Misc/AssertionMacros.h"
 #include "Serialization/Archive.h"
-#include "Templates/AreTypesEqual.h"
 #include "Templates/TypeCompatibleBytes.h"
 #include "Templates/UnrealTemplate.h"
 #include "Templates/UnrealTypeTraits.h"
@@ -354,46 +353,39 @@ private:
 		PointerType& OutValuePointer
 		)
 	{
-		if(TAreTypesEqual<TypeA,Subtype>::Value)
+		if constexpr (std::is_same_v<TypeA,Subtype>)
 		{
 			OutIndex = 0;
 			OutValuePointer = (PointerType)&Union.Values.A;
 		}
-		else if(TAreTypesEqual<TypeB,Subtype>::Value)
+		else if constexpr (std::is_same_v<TypeB,Subtype>)
 		{
 			OutIndex = 1;
 			OutValuePointer = (PointerType)&Union.Values.B;
 		}
-		else if(TAreTypesEqual<TypeC,Subtype>::Value)
+		else if constexpr (std::is_same_v<TypeC,Subtype>)
 		{
 			OutIndex = 2;
 			OutValuePointer = (PointerType)&Union.Values.C;
 		}
-		else if(TAreTypesEqual<TypeD,Subtype>::Value)
+		else if constexpr (std::is_same_v<TypeD,Subtype>)
 		{
 			OutIndex = 3;
 			OutValuePointer = (PointerType)&Union.Values.D;
 		}
-		else if(TAreTypesEqual<TypeE,Subtype>::Value)
+		else if constexpr (std::is_same_v<TypeE,Subtype>)
 		{
 			OutIndex = 4;
 			OutValuePointer = (PointerType)&Union.Values.E;
 		}
-		else if(TAreTypesEqual<TypeF,Subtype>::Value)
+		else if constexpr (std::is_same_v<TypeF,Subtype>)
 		{
 			OutIndex = 5;
 			OutValuePointer = (PointerType)&Union.Values.F;
 		}
 		else
 		{
-			static_assert(
-				TAreTypesEqual<TEMPLATE_PARAMETERS2(TypeA, Subtype)>::Value ||
-				TAreTypesEqual<TEMPLATE_PARAMETERS2(TypeB, Subtype)>::Value ||
-				TAreTypesEqual<TEMPLATE_PARAMETERS2(TypeC, Subtype)>::Value ||
-				TAreTypesEqual<TEMPLATE_PARAMETERS2(TypeD, Subtype)>::Value ||
-				TAreTypesEqual<TEMPLATE_PARAMETERS2(TypeE, Subtype)>::Value ||
-				TAreTypesEqual<TEMPLATE_PARAMETERS2(TypeF, Subtype)>::Value,
-				"Type is not subtype of union.");
+			static_assert(sizeof(TypeA) == 0, "Type is not subtype of union.");
 			OutIndex = (uint8)-1;
 			OutValuePointer = NULL;
 		}
diff --git a/Engine/Source/Runtime/Core/Public/Core.h b/Engine/Source/Runtime/Core/Public/Core.h
--- a/Engine/Source/Runtime/Core/Public/Core.h
+++ b/Engine/Source/Runtime/Core/Public/Core.h
@@ -19,7 +19,6 @@ MONOLITHIC_HEADER_BOILERPLATE()
 #include "Misc/Char.h"
 #include "Templates/EnableIf.h"
 #include "Templates/AndOrNot.h"
-#include "Templates/AreTypesEqual.h"
 #include "Templates/RemoveCV.h"
 #include "Templates/IsArithmetic.h"
 #include "Templates/IsFloatingPoint.h"
diff --git a/Engine/Source/Runtime/Core/Public/CoreMinimal.h b/Engine/Source/Runtime/Core/Public/CoreMinimal.h
--- a/Engine/Source/Runtime/Core/Public/CoreMinimal.h
+++ b/Engine/Source/Runtime/Core/Public/CoreMinimal.h
@@ -34,7 +34,6 @@
 #include "HAL/UnrealMemory.h"
 #include "Templates/IsArithmetic.h"
 #include "Templates/AndOrNot.h"
-#include "Templates/AreTypesEqual.h"
 #include "Templates/RemoveCV.h"
 #include "Templates/IsPODType.h"
 #include "Templates/IsUECoreType.h"
diff --git a/Engine/Source/Runtime/Core/Public/CoreSharedPCH.h b/Engine/Source/Runtime/Core/Public/CoreSharedPCH.h
--- a/Engine/Source/Runtime/Core/Public/CoreSharedPCH.h
+++ b/Engine/Source/Runtime/Core/Public/CoreSharedPCH.h
@@ -171,7 +171,6 @@
 #include "Stats/Stats.h"
 #include "Templates/AlignmentTemplates.h"
 #include "Templates/AndOrNot.h"
-#include "Templates/AreTypesEqual.h"
 #include "Templates/ChooseClass.h"
 #include "Templates/Decay.h"
 #include "Templates/EnableIf.h"
diff --git a/Engine/Source/Runtime/Core/Public/Delegates/DelegateInstancesImpl.h b/Engine/Source/Runtime/Core/Public/Delegates/DelegateInstancesImpl.h
--- a/Engine/Source/Runtime/Core/Public/Delegates/DelegateInstancesImpl.h
+++ b/Engine/Source/Runtime/Core/Public/Delegates/DelegateInstancesImpl.h
@@ -13,7 +13,6 @@
 #include "CoreTypes.h"
 #include "Delegates/DelegateInstanceInterface.h"
 #include "Misc/AssertionMacros.h"
-#include "Templates/AreTypesEqual.h"
 #include "Templates/RemoveReference.h"
 #include "Templates/SharedPointer.h"
 #include "Templates/Tuple.h"
@@ -675,7 +674,7 @@ template <typename RetValType, typename... ParamTypes, typename UserPolicy, type
 class TBaseFunctorDelegateInstance<RetValType(ParamTypes...), UserPolicy, FunctorType, VarTypes...> : public TCommonDelegateInstanceState<RetValType(ParamTypes...), UserPolicy, VarTypes...>
 {
 private:
-	static_assert(TAreTypesEqual<FunctorType, typename TRemoveReference<FunctorType>::Type>::Value, "FunctorType cannot be a reference");
+	static_assert(std::is_same_v<FunctorType, typename TRemoveReference<FunctorType>::Type>, "FunctorType cannot be a reference");
 
 	using Super            = TCommonDelegateInstanceState<RetValType(ParamTypes...), UserPolicy, VarTypes...>;
 	using DelegateBaseType = typename UserPolicy::FDelegateExtras;
@@ -765,7 +764,7 @@ template <typename UserClass, typename RetValType, typename... ParamTypes, typen
 class TWeakBaseFunctorDelegateInstance<UserClass, RetValType(ParamTypes...), UserPolicy, FunctorType, VarTypes...> : public TCommonDelegateInstanceState<RetValType(ParamTypes...), UserPolicy, VarTypes...>
 {
 private:
-	static_assert(TAreTypesEqual<FunctorType, typename TRemoveReference<FunctorType>::Type>::Value, "FunctorType cannot be a reference");
+	static_assert(std::is_same_v<FunctorType, typename TRemoveReference<FunctorType>::Type>, "FunctorType cannot be a reference");
 
 	using Super            = TCommonDelegateInstanceState<RetValType(ParamTypes...), UserPolicy, VarTypes...>;
 	using DelegateBaseType = typename UserPolicy::FDelegateExtras;
diff --git a/Engine/Source/Runtime/Core/Public/Experimental/Containers/SherwoodHashTable.h b/Engine/Source/Runtime/Core/Public/Experimental/Containers/SherwoodHashTable.h
--- a/Engine/Source/Runtime/Core/Public/Experimental/Containers/SherwoodHashTable.h
+++ b/Engine/Source/Runtime/Core/Public/Experimental/Containers/SherwoodHashTable.h
@@ -41,7 +41,7 @@ struct TSherwoodHashTable
 
 	// TSherwoodHashTable can be used to implement a set or a map.
 	// In map mode we allocate memory for keys and values, but in set mode we only allocate keys.
-	static constexpr bool bIsMap = !TAreTypesEqual<ValueType, FNoopStruct>::Value;
+	static constexpr bool bIsMap = !std::is_same_v<ValueType, FNoopStruct>;
 
 	// Minimum probing distance when searching for an entry slot.
 	static constexpr uint32 MinNumLookups = 4;
diff --git a/Engine/Source/Runtime/Core/Public/HAL/Platform.h b/Engine/Source/Runtime/Core/Public/HAL/Platform.h
--- a/Engine/Source/Runtime/Core/Public/HAL/Platform.h
+++ b/Engine/Source/Runtime/Core/Public/HAL/Platform.h
@@ -894,8 +894,6 @@ int32 main(int32 ArgC, ANSICHAR* Utf8ArgV[]) \
 int32 tchar_main(int32 ArgC, TCHAR* ArgV[])
 #endif
 
-template<typename,typename> struct TAreTypesEqual;
-
 //--------------------------------------------------------------------------------------------
 // POD types refactor for porting old code:
 //
@@ -994,16 +992,10 @@ typedef FPlatformTypes::TYPE_OF_NULLPTR	TYPE_OF_NULLPTR;
 namespace TypeTests
 {
 	template <typename A, typename B>
-	struct TAreTypesEqual
-	{
-		static constexpr bool Value = false;
-	};
+	inline constexpr bool TAreTypesEqual_V = false;
 
 	template <typename T>
-	struct TAreTypesEqual<T, T>
-	{
-		static constexpr bool Value = true;
-	};
+	inline constexpr bool TAreTypesEqual_V<T, T> = true;
 
 #if PLATFORM_TCHAR_IS_4_BYTES
 	static_assert(sizeof(TCHAR) == 4, "TCHAR size must be 4 bytes.");
@@ -1023,24 +1015,24 @@ namespace TypeTests
 
 	static_assert(char(-1) < char(0), "Unsigned char type test failed.");
 
-	static_assert((!TAreTypesEqual<ANSICHAR, WIDECHAR>::Value),  "ANSICHAR and WIDECHAR should be different types.");
-	static_assert((!TAreTypesEqual<ANSICHAR, UTF8CHAR>::Value),  "ANSICHAR and UTF8CHAR should be different types.");
+	static_assert((!TAreTypesEqual_V<ANSICHAR, WIDECHAR>),  "ANSICHAR and WIDECHAR should be different types.");
+	static_assert((!TAreTypesEqual_V<ANSICHAR, UTF8CHAR>),  "ANSICHAR and UTF8CHAR should be different types.");
 #if !PLATFORM_UCS2CHAR_IS_UTF16CHAR
 	// We want these types to be different, because we want to be able to determine whether an encoding
 	// is fixed-width (UCS2CHAR) or variable-width (UTF16CHAR) at compile time.
-	static_assert((!TAreTypesEqual<UCS2CHAR, UTF16CHAR>::Value), "UCS2CHAR and UTF16CHAR should be different types.");
+	static_assert((!TAreTypesEqual_V<UCS2CHAR, UTF16CHAR>), "UCS2CHAR and UTF16CHAR should be different types.");
 #else
 	// We don't want these types to be equal, but while this macro exists, they ought to be equal
-	static_assert(TAreTypesEqual<UCS2CHAR, UTF16CHAR>::Value, "UCS2CHAR and UTF16CHAR are expected to be the same type.");
+	static_assert(TAreTypesEqual_V<UCS2CHAR, UTF16CHAR>, "UCS2CHAR and UTF16CHAR are expected to be the same type.");
 #endif
-	static_assert((!TAreTypesEqual<ANSICHAR, UCS2CHAR>::Value),  "ANSICHAR and UCS2CHAR should be different types.");
-	static_assert((!TAreTypesEqual<WIDECHAR, UCS2CHAR>::Value),  "WIDECHAR and UCS2CHAR should be different types.");
-	static_assert(TAreTypesEqual<TCHAR, ANSICHAR>::Value || TAreTypesEqual<TCHAR, WIDECHAR>::Value || TAreTypesEqual<TCHAR, UTF8CHAR>::Value, "TCHAR should either be ANSICHAR, WIDECHAR or UTF8CHAR.");
+	static_assert((!TAreTypesEqual_V<ANSICHAR, UCS2CHAR>),  "ANSICHAR and UCS2CHAR should be different types.");
+	static_assert((!TAreTypesEqual_V<WIDECHAR, UCS2CHAR>),  "WIDECHAR and UCS2CHAR should be different types.");
+	static_assert(TAreTypesEqual_V<TCHAR, ANSICHAR> || TAreTypesEqual_V<TCHAR, WIDECHAR> || TAreTypesEqual_V<TCHAR, UTF8CHAR>, "TCHAR should either be ANSICHAR, WIDECHAR or UTF8CHAR.");
 
 #if PLATFORM_WIDECHAR_IS_CHAR16
-	static_assert(TAreTypesEqual<WIDECHAR, char16_t>::Value, "WIDECHAR should be char16_t");
+	static_assert(TAreTypesEqual_V<WIDECHAR, char16_t>, "WIDECHAR should be char16_t");
 #else
-	static_assert(TAreTypesEqual<WIDECHAR, wchar_t>::Value, "WIDECHAR should be wchar_t");
+	static_assert(TAreTypesEqual_V<WIDECHAR, wchar_t>, "WIDECHAR should be wchar_t");
 #endif
 
 	static_assert(sizeof(uint8) == 1, "uint8 type size test failed.");
diff --git a/Engine/Source/Runtime/Core/Public/Math/TransformCalculus2D.h b/Engine/Source/Runtime/Core/Public/Math/TransformCalculus2D.h
--- a/Engine/Source/Runtime/Core/Public/Math/TransformCalculus2D.h
+++ b/Engine/Source/Runtime/Core/Public/Math/TransformCalculus2D.h
@@ -4,7 +4,6 @@
 
 #include "CoreTypes.h"
 #include "Misc/AssertionMacros.h"
-#include "Templates/AreTypesEqual.h"
 #include "Math/UnrealMathUtility.h"
 #include "Math/Vector2D.h"
 #include "Math/TransformCalculus.h"
diff --git a/Engine/Source/Runtime/Core/Public/Misc/TypeContainer.h b/Engine/Source/Runtime/Core/Public/Misc/TypeContainer.h
--- a/Engine/Source/Runtime/Core/Public/Misc/TypeContainer.h
+++ b/Engine/Source/Runtime/Core/Public/Misc/TypeContainer.h
@@ -5,7 +5,6 @@
 #include "CoreTypes.h"
 #include "Templates/PointerIsConvertibleFromTo.h"
 #include "Misc/AssertionMacros.h"
-#include "Templates/AreTypesEqual.h"
 #include "Templates/UnrealTypeTraits.h"
 #include "Containers/UnrealString.h"
 #include "Templates/Function.h"
@@ -282,7 +281,7 @@ public:
 	template<class R, class D, typename... P>
 	void RegisterDelegate(D Delegate)
 	{
-		static_assert(TAreTypesEqual<TSharedRef<R, Mode>, typename D::RetValType>::Value, "Delegate return type must be TSharedPtr<R>");
+		static_assert(std::is_same_v<TSharedRef<R, Mode>, typename D::RetValType>, "Delegate return type must be TSharedPtr<R>");
 
 		TSharedPtr<IInstanceProvider> Provider = MakeShareable(
 			new TFunctionInstanceProvider<R>(
diff --git a/Engine/Source/Runtime/Core/Public/Templates/AreTypesEqual.h b/Engine/Source/Runtime/Core/Public/Templates/AreTypesEqual.h
--- a/Engine/Source/Runtime/Core/Public/Templates/AreTypesEqual.h
+++ b/Engine/Source/Runtime/Core/Public/Templates/AreTypesEqual.h
@@ -2,20 +2,13 @@
 
 #pragma once
 
+#include <type_traits>
+
 /** Tests whether two typenames refer to the same type. */
 template<typename A,typename B>
-struct TAreTypesEqual;
-
-template<typename,typename>
-struct TAreTypesEqual
-{
-	enum { Value = false };
-};
-
-template<typename A>
-struct TAreTypesEqual<A,A>
+struct UE_DEPRECATED(5.2, "TAreTypesEqual has been deprecated, please use std::is_same instead.") TAreTypesEqual
 {
-	enum { Value = true };
+	enum { Value = std::is_same_v<A, B> };
 };
 
-#define ARE_TYPES_EQUAL(A,B) TAreTypesEqual<A,B>::Value
+#define ARE_TYPES_EQUAL(A,B) std::is_same_v<A, B>
diff --git a/Engine/Source/Runtime/Core/Public/Templates/LosesQualifiersFromTo.h b/Engine/Source/Runtime/Core/Public/Templates/LosesQualifiersFromTo.h
--- a/Engine/Source/Runtime/Core/Public/Templates/LosesQualifiersFromTo.h
+++ b/Engine/Source/Runtime/Core/Public/Templates/LosesQualifiersFromTo.h
@@ -2,8 +2,8 @@
 
 #pragma once
 
-#include "Templates/AreTypesEqual.h"
 #include "Templates/CopyQualifiersFromTo.h"
+#include <type_traits>
 
 /**
  * Tests if qualifiers are lost between one type and another, e.g.:
@@ -14,5 +14,5 @@
 template <typename From, typename To>
 struct TLosesQualifiersFromTo
 {
-	enum { Value = !TAreTypesEqual<typename TCopyQualifiersFromTo<From, To>::Type, To>::Value };
+	enum { Value = !std::is_same_v<typename TCopyQualifiersFromTo<From, To>::Type, To> };
 };
diff --git a/Engine/Source/Runtime/Core/Public/Templates/MemoryOps.h b/Engine/Source/Runtime/Core/Public/Templates/MemoryOps.h
--- a/Engine/Source/Runtime/Core/Public/Templates/MemoryOps.h
+++ b/Engine/Source/Runtime/Core/Public/Templates/MemoryOps.h
@@ -4,7 +4,6 @@
 
 #include "CoreTypes.h"
 #include "HAL/UnrealMemory.h"
-#include "Templates/AreTypesEqual.h"
 #include "Templates/IsTriviallyCopyConstructible.h"
 #include "Templates/UnrealTypeTraits.h"
 #include <new> // IWYU pragma: export
@@ -21,12 +20,10 @@ namespace UE::Core::Private::MemoryOps
 		enum
 		{
 			Value =
-				TOr<
-					TAreTypesEqual<DestinationElementType, SourceElementType>,
-					TAnd<
-						TIsBitwiseConstructible<DestinationElementType, SourceElementType>,
-						TIsTriviallyDestructible<SourceElementType>
-					>
+				std::is_same_v<DestinationElementType, SourceElementType> ||
+				TAnd<
+					TIsBitwiseConstructible<DestinationElementType, SourceElementType>,
+					TIsTriviallyDestructible<SourceElementType>
 				>::Value
 		};
 	};
diff --git a/Engine/Source/Runtime/Core/Public/Templates/Tuple.h b/Engine/Source/Runtime/Core/Public/Templates/Tuple.h
--- a/Engine/Source/Runtime/Core/Public/Templates/Tuple.h
+++ b/Engine/Source/Runtime/Core/Public/Templates/Tuple.h
@@ -3,7 +3,6 @@
 #pragma once
 
 #include "CoreTypes.h"
-#include "Templates/AreTypesEqual.h"
 #include "Templates/UnrealTemplate.h"
 #include "Templates/Decay.h"
 #include "Delegates/IntegerSequence.h"
diff --git a/Engine/Source/Runtime/Core/Public/Templates/UnrealTemplate.h b/Engine/Source/Runtime/Core/Public/Templates/UnrealTemplate.h
--- a/Engine/Source/Runtime/Core/Public/Templates/UnrealTemplate.h
+++ b/Engine/Source/Runtime/Core/Public/Templates/UnrealTemplate.h
@@ -7,7 +7,6 @@
 #include "HAL/UnrealMemory.h"
 #include "Templates/EnableIf.h"
 #include "Templates/AndOrNot.h"
-#include "Templates/AreTypesEqual.h"
 #include "Templates/CopyQualifiersAndRefsFromTo.h"
 #include "Templates/IsArithmetic.h"
 #include "Templates/UnrealTypeTraits.h"
@@ -468,7 +467,7 @@ FORCEINLINE typename TRemoveReference<T>::Type&& MoveTemp(T&& Obj)
 
 	// Validate that we're not being passed an rvalue or a const object - the former is redundant, the latter is almost certainly a mistake
 	static_assert(TIsLValueReferenceType<T>::Value, "MoveTemp called on an rvalue");
-	static_assert(!TAreTypesEqual<CastType&, const CastType&>::Value, "MoveTemp called on a const object");
+	static_assert(!std::is_same_v<CastType&, const CastType&>, "MoveTemp called on a const object");
 
 	return (CastType&&)Obj;
 }
@@ -618,7 +617,7 @@ template <typename T> struct TRValueToLValueReference<T&&> { typedef T& Type; };
  * @return The bit-swapped value.
  */
 template <typename T>
-FORCEINLINE typename TEnableIf<TAreTypesEqual<T, uint32>::Value, T>::Type ReverseBits( T Bits )
+FORCEINLINE typename TEnableIf<std::is_same_v<T, uint32>, T>::Type ReverseBits( T Bits )
 {
 	Bits = ( Bits << 16) | ( Bits >> 16);
 	Bits = ( (Bits & 0x00ff00ff) << 8 ) | ( (Bits & 0xff00ff00) >> 8 );
diff --git a/Engine/Source/Runtime/Core/Public/Templates/UnrealTypeTraits.h b/Engine/Source/Runtime/Core/Public/Templates/UnrealTypeTraits.h
--- a/Engine/Source/Runtime/Core/Public/Templates/UnrealTypeTraits.h
+++ b/Engine/Source/Runtime/Core/Public/Templates/UnrealTypeTraits.h
@@ -12,7 +12,6 @@
 #include "Templates/IsPointer.h"
 #include "Misc/AssertionMacros.h"
 #include "Templates/AndOrNot.h"
-#include "Templates/AreTypesEqual.h"
 #include "Templates/EnableIf.h"
 #include "Templates/IsArithmetic.h"
 #include "Templates/IsEnum.h"
@@ -448,8 +447,8 @@ struct TIsBitwiseConstructible
 		"TIsBitwiseConstructible is not designed to accept reference types");
 
 	static_assert(
-		TAreTypesEqual<T,   typename TRemoveCV<T  >::Type>::Value &&
-		TAreTypesEqual<Arg, typename TRemoveCV<Arg>::Type>::Value,
+		std::is_same_v<T,   typename TRemoveCV<T  >::Type> &&
+		std::is_same_v<Arg, typename TRemoveCV<Arg>::Type>,
 		"TIsBitwiseConstructible is not designed to accept qualified types");
 
 	// Assume no bitwise construction in general
diff --git a/Engine/Source/Runtime/Core/Public/UObject/ScriptDelegates.h b/Engine/Source/Runtime/Core/Public/UObject/ScriptDelegates.h
--- a/Engine/Source/Runtime/Core/Public/UObject/ScriptDelegates.h
+++ b/Engine/Source/Runtime/Core/Public/UObject/ScriptDelegates.h
@@ -5,7 +5,6 @@
 #include "CoreTypes.h"
 #include "Delegates/DelegateMacros.h"
 #include "Misc/AssertionMacros.h"
-#include "Templates/AreTypesEqual.h"
 #include "Templates/UnrealTypeTraits.h"
 #include "Containers/ContainerAllocationPolicies.h"
 #include "Containers/Array.h"
@@ -23,7 +22,7 @@ class TScriptDelegate
 {
 	// Although templated, the parameter is not intended to be anything other than the default,
 	// and is only a template for module organization reasons.
-	static_assert(TAreTypesEqual<TWeakPtr, FWeakObjectPtr>::Value, "TWeakPtr should not be overridden");
+	static_assert(std::is_same_v<TWeakPtr, FWeakObjectPtr>, "TWeakPtr should not be overridden");
 
 public:
 
diff --git a/Engine/Source/Runtime/Core/Public/UObject/WeakObjectPtrTemplates.h b/Engine/Source/Runtime/Core/Public/UObject/WeakObjectPtrTemplates.h
--- a/Engine/Source/Runtime/Core/Public/UObject/WeakObjectPtrTemplates.h
+++ b/Engine/Source/Runtime/Core/Public/UObject/WeakObjectPtrTemplates.h
@@ -5,7 +5,6 @@
 #include "CoreTypes.h"
 #include "Templates/IsPointer.h"
 #include "Templates/PointerIsConvertibleFromTo.h"
-#include "Templates/AreTypesEqual.h"
 #include "Templates/AndOrNot.h"
 #include "Templates/LosesQualifiersFromTo.h"
 #include "Containers/Map.h"
@@ -35,7 +34,7 @@ struct TWeakObjectPtr : private TWeakObjectPtrBase
 
 	// Although templated, these parameters are not intended to be anything other than the default,
 	// and are only templates for module organization reasons.
-	static_assert(TAreTypesEqual<TWeakObjectPtrBase, FWeakObjectPtr>::Value, "TWeakObjectPtrBase should not be overridden");
+	static_assert(std::is_same_v<TWeakObjectPtrBase, FWeakObjectPtr>, "TWeakObjectPtrBase should not be overridden");
 
 public:
 	using ElementType = T;
diff --git a/Engine/Source/Runtime/Engine/Private/EnginePrivatePCH.h b/Engine/Source/Runtime/Engine/Private/EnginePrivatePCH.h
--- a/Engine/Source/Runtime/Engine/Private/EnginePrivatePCH.h
+++ b/Engine/Source/Runtime/Engine/Private/EnginePrivatePCH.h
@@ -23,7 +23,6 @@
 #include "Misc/VarArgs.h"
 #include "HAL/PlatformAtomics.h"
 #include "GenericPlatform/GenericPlatformAtomics.h"
-#include "Templates/AreTypesEqual.h"
 #include "Templates/UnrealTypeTraits.h"
 #include "Templates/AndOrNot.h"
 #include "Templates/IsArithmetic.h"
diff --git a/Engine/Source/Runtime/Engine/Private/RepLayout.cpp b/Engine/Source/Runtime/Engine/Private/RepLayout.cpp
--- a/Engine/Source/Runtime/Engine/Private/RepLayout.cpp
+++ b/Engine/Source/Runtime/Engine/Private/RepLayout.cpp
@@ -135,7 +135,7 @@ namespace UE_RepLayout_Private
 		using BaseBufferType = typename TRemovePointer<typename TDecay<BufferType>::Type>::Type;
 
 		static_assert(!TIsPointer<OutputType>::Value, "GetTypedProperty invalid OutputType!  Don't specify output as a pointer.");
-		static_assert(TOr<TAreTypesEqual<uint8, BaseBufferType>, TAreTypesEqual<void, BaseBufferType>>::Value, "GetTypedProperty invalid BufferType! Only TRepDataBufferBase, void*, and uint8* are supported!");
+		static_assert(std::is_same_v<uint8, BaseBufferType> || std::is_same_v<void, BaseBufferType>, "GetTypedProperty invalid BufferType! Only TRepDataBufferBase, void*, and uint8* are supported!");
 
 		// TODO: Conditionally compilable runtime type validation.
 		return reinterpret_cast<ConstOrNotOutputType *>(Buffer);
diff --git a/Engine/Source/Runtime/Engine/Public/EngineSharedPCH.h b/Engine/Source/Runtime/Engine/Public/EngineSharedPCH.h
--- a/Engine/Source/Runtime/Engine/Public/EngineSharedPCH.h
+++ b/Engine/Source/Runtime/Engine/Public/EngineSharedPCH.h
@@ -23,7 +23,6 @@
 #include "Misc/VarArgs.h"
 #include "HAL/PlatformAtomics.h"
 #include "GenericPlatform/GenericPlatformAtomics.h"
-#include "Templates/AreTypesEqual.h"
 #include "Templates/UnrealTypeTraits.h"
 #include "Templates/AndOrNot.h"
 #include "Templates/IsArithmetic.h"
diff --git a/Engine/Source/Runtime/Experimental/Chaos/Private/Chaos/TriangleMesh.cpp b/Engine/Source/Runtime/Experimental/Chaos/Private/Chaos/TriangleMesh.cpp
--- a/Engine/Source/Runtime/Experimental/Chaos/Private/Chaos/TriangleMesh.cpp
+++ b/Engine/Source/Runtime/Experimental/Chaos/Private/Chaos/TriangleMesh.cpp
@@ -311,7 +311,7 @@ void FTriangleMesh::GetFaceNormals(TArray<TVec3<T>>& Normals, const TConstArrayV
 	else
 	{
 #if INTEL_ISPC
-		if (bChaos_TriangleMesh_ISPC_Enabled && TAreTypesEqual<T, FRealSingle>::Value)
+		if (bChaos_TriangleMesh_ISPC_Enabled && std::is_same_v<T, FRealSingle>)
 		{
 			Normals.SetNumUninitialized(MElements.Num());
 			ispc::GetFaceNormals(
diff --git a/Engine/Source/Runtime/Experimental/Chaos/Public/Chaos/Box.h b/Engine/Source/Runtime/Experimental/Chaos/Public/Chaos/Box.h
--- a/Engine/Source/Runtime/Experimental/Chaos/Public/Chaos/Box.h
+++ b/Engine/Source/Runtime/Experimental/Chaos/Public/Chaos/Box.h
@@ -499,7 +499,7 @@ namespace Chaos
 					AABBs.Emplace(TAABB<OtherType, d>(Box.AABB));
 				}
 			}
-			else if (TAreTypesEqual<T, OtherType>::Value)
+			else if (std::is_same_v<T, OtherType>)
 			{
 				Ar << AABBs;
 			}
diff --git a/Engine/Source/Runtime/Experimental/ChaosCore/Public/Chaos/Real.h b/Engine/Source/Runtime/Experimental/ChaosCore/Public/Chaos/Real.h
--- a/Engine/Source/Runtime/Experimental/ChaosCore/Public/Chaos/Real.h
+++ b/Engine/Source/Runtime/Experimental/ChaosCore/Public/Chaos/Real.h
@@ -2,7 +2,7 @@
 
 #pragma once
 
-#include "Templates/AreTypesEqual.h"
+#include <type_traits>
 
 namespace Chaos
 {
@@ -25,5 +25,5 @@ namespace Chaos
 	* ISPC optimization supports float and double, this allows classes that uses ISPC to branch to the right implementation 
 	* without having to check the actual underlying type of FReal
 	*/
-	inline constexpr bool bRealTypeCompatibleWithISPC = (TAreTypesEqual<FReal, float>::Value == true) || (TAreTypesEqual<FReal, double>::Value == true);
+	inline constexpr bool bRealTypeCompatibleWithISPC = std::is_same_v<FReal, float> || std::is_same_v<FReal, double>;
 }
diff --git a/Engine/Source/Runtime/RenderCore/Public/RenderMath.h b/Engine/Source/Runtime/RenderCore/Public/RenderMath.h
--- a/Engine/Source/Runtime/RenderCore/Public/RenderMath.h
+++ b/Engine/Source/Runtime/RenderCore/Public/RenderMath.h
@@ -46,8 +46,8 @@ FORCEINLINE int8 GetBasisDeterminantSignByte(const FPackedNormal& XAxis, const F
 template<typename VectorType>
 FORCEINLINE FVector GenerateYAxis(const VectorType& XAxis, const VectorType& ZAxis)
 {
-	static_assert(ARE_TYPES_EQUAL(VectorType, FPackedNormal) ||
-		ARE_TYPES_EQUAL(VectorType, FPackedRGBA16N), "ERROR: Must be FPackedNormal or FPackedRGBA16N");
+	static_assert(std::is_same_v<VectorType, FPackedNormal> ||
+		std::is_same_v<VectorType, FPackedRGBA16N>, "ERROR: Must be FPackedNormal or FPackedRGBA16N");
 	FVector  x = XAxis.ToFVector();
 	FVector4 z = ZAxis.ToFVector4();
 	return (FVector(z) ^ x) * z.W;
diff --git a/Engine/Source/Runtime/Slate/Public/SlateSharedPCH.h b/Engine/Source/Runtime/Slate/Public/SlateSharedPCH.h
--- a/Engine/Source/Runtime/Slate/Public/SlateSharedPCH.h
+++ b/Engine/Source/Runtime/Slate/Public/SlateSharedPCH.h
@@ -23,7 +23,6 @@
 #include "Misc/VarArgs.h"
 #include "HAL/PlatformAtomics.h"
 #include "GenericPlatform/GenericPlatformAtomics.h"
-#include "Templates/AreTypesEqual.h"
 #include "Templates/UnrealTypeTraits.h"
 #include "Templates/AndOrNot.h"
 #include "Templates/IsArithmetic.h"
