Prevent crash if GetStatData called with a directory on Android

#android

diff --git a/Engine/Source/Runtime/Core/Private/Android/AndroidPlatformFile.cpp b/Engine/Source/Runtime/Core/Private/Android/AndroidPlatformFile.cpp
--- a/Engine/Source/Runtime/Core/Private/Android/AndroidPlatformFile.cpp
+++ b/Engine/Source/Runtime/Core/Private/Android/AndroidPlatformFile.cpp
@@ -881,7 +881,8 @@ public:
 
 	int64 GetEntryLength(const FString & Path)
 	{
-		return Entries[Path]->File->Size();
+		TSharedPtr<FFileHandleAndroid> File = Entries[Path]->File;
+		return File != nullptr ? File->Size() : 0;
 	}
 
 	int64 GetEntryModTime(const FString & Path)
