Fix FColor related clang static analyzer warning.

#rn Fix Previously the clang static analyzer was having trouble realizing that all members of the FColor struct are initialized when casting the this pointer to a uint32. This uses a union instead of the alignment attribute to force 32bit alignment, this appears to make it easier for the static analyzer to read.

diff --git a/Engine/Source/Runtime/Core/Public/Math/Color.h b/Engine/Source/Runtime/Core/Public/Math/Color.h
--- a/Engine/Source/Runtime/Core/Public/Math/Color.h
+++ b/Engine/Source/Runtime/Core/Public/Math/Color.h
@@ -485,20 +485,13 @@ struct FColor
 public:
 	// Variables.
 #if PLATFORM_LITTLE_ENDIAN
-	#ifdef _MSC_VER
-		// Win32 x86
-		union { struct{ uint8 B,G,R,A; }; uint32 AlignmentDummy; };
-	#else
-		// Linux x86, etc
-		uint8 B GCC_ALIGN(4);
-		uint8 G,R,A;
-	#endif
+	union { struct{ uint8 B,G,R,A; }; uint32 Bits; };
 #else // PLATFORM_LITTLE_ENDIAN
-	union { struct{ uint8 A,R,G,B; }; uint32 AlignmentDummy; };
+	union { struct{ uint8 A,R,G,B; }; uint32 Bits; };
 #endif
 
-	uint32& DWColor(void) {return *((uint32*)this);}
-	const uint32& DWColor(void) const {return *((uint32*)this);}
+	uint32& DWColor(void) {return Bits;}
+	const uint32& DWColor(void) const {return Bits;}
 
 	// Constructors.
 	FORCEINLINE FColor() {}
