Archive.cpp SerializedCompressed fix security/overflow issues

diff --git a/Engine/Source/Runtime/Core/Private/Serialization/Archive.cpp b/Engine/Source/Runtime/Core/Private/Serialization/Archive.cpp
--- a/Engine/Source/Runtime/Core/Private/Serialization/Archive.cpp
+++ b/Engine/Source/Runtime/Core/Private/Serialization/Archive.cpp
@@ -699,6 +699,8 @@ void FArchive::SerializeCompressedNew(void* V, int64 Length, FName CompressionFo
 
 	// CompressionFormatToEncode can be changed freely without breaking loading of old files
 	// CompressionFormatToDecodeOldV1Files must match what was used to encode old files, cannot change
+	
+	UE_CLOG(Length < 0, LogSerialization, Fatal, TEXT(" Archive SerializedCompressed Length (%lld) < 0"), Length);
 
 	if( IsLoading() )
 	{
@@ -803,7 +805,9 @@ void FArchive::SerializeCompressedNew(void* V, int64 Length, FName CompressionFo
 			Summary.UncompressedSize = BYTESWAP_ORDER64(Summary.UncompressedSize);
 			PackageFileTag.UncompressedSize = BYTESWAP_ORDER64(PackageFileTag.UncompressedSize);
 		}
-
+		
+		UE_CLOG(Summary.CompressedSize < 0 || Summary.CompressedSize > (INT64_MAX/2) , LogSerialization, Fatal, TEXT(" Archive SerializedCompressed CompressedSize (%lld) invalid"), (int64)Summary.CompressedSize);
+		UE_CLOG(Summary.UncompressedSize < 0 || Summary.UncompressedSize > (INT64_MAX/2) , LogSerialization, Fatal, TEXT(" Archive SerializedCompressed UncompressedSize (%lld) invalid"), (int64)Summary.UncompressedSize);
 
 		// Handle change in compression chunk size in backward compatible way.
 		int64 LoadingCompressionChunkSize = PackageFileTag.UncompressedSize;
@@ -812,32 +816,33 @@ void FArchive::SerializeCompressedNew(void* V, int64 Length, FName CompressionFo
 			LoadingCompressionChunkSize = LOADING_COMPRESSION_CHUNK_SIZE;
 		}
 
-		check( LoadingCompressionChunkSize > 0 );
+		UE_CLOG(LoadingCompressionChunkSize <= 0, LogSerialization, Fatal, TEXT(" Archive SerializedCompressed LoadingCompressionChunkSize (%lld) <= 0"), LoadingCompressionChunkSize);
+		UE_CLOG(LoadingCompressionChunkSize >= INT32_MAX, LogSerialization, Fatal, TEXT(" Archive SerializedCompressed LoadingCompressionChunkSize (%lld) >= INT32_MAX"), LoadingCompressionChunkSize);
 
 		// check Summary.UncompressedSize vs [V,Length] passed in
 		if ( OutPartialReadLength == nullptr )
 		{
-			check( Summary.UncompressedSize == Length );
+			// UncompressedSize must == Length
 			UE_CLOG( Summary.UncompressedSize != Length, LogSerialization, Fatal, TEXT(" Archive SerializedCompressed UncompressedSize (%lld) != Length (%lld)"), (int64)Summary.UncompressedSize, (int64) Length );
 		}
 		else
 		{
-			// UncompressedSize smaller than length is okay
-			check( Summary.UncompressedSize <= Length );
-			check( Summary.UncompressedSize >= 0 );
-			UE_CLOG( Summary.UncompressedSize > Length, LogSerialization, Fatal, TEXT(" Archive SerializedCompressed UncompressedSize (%lld) > Length (%lld)"), (int64)Summary.UncompressedSize, (int64) Length );
+			// UncompressedSize must be <= Length and >= 0
+			UE_CLOG( Summary.UncompressedSize > Length || Summary.UncompressedSize < 0, LogSerialization, Fatal, TEXT(" Archive SerializedCompressed UncompressedSize (%lld) > Length (%lld) or < 0"), (int64)Summary.UncompressedSize, (int64) Length );
 			*OutPartialReadLength = Summary.UncompressedSize;
 		}
 
 		// Figure out how many chunks there are going to be based on uncompressed size and compression chunk size.
-		int64	TotalChunkCount	= (Summary.UncompressedSize + LoadingCompressionChunkSize - 1) / LoadingCompressionChunkSize;
-		
+		//  divide and round up, safe without overflow due to previous range checks :
+		int64	TotalChunkCount	= FMath::DivideAndRoundUp( Summary.UncompressedSize , LoadingCompressionChunkSize );
+
 		// Allocate compression chunk infos and serialize them, keeping track of max size of compression chunks used.
-		FCompressedChunkInfo* CompressionChunks = new FCompressedChunkInfo[TotalChunkCount];
+		TArray64<FCompressedChunkInfo> CompressionChunks;
+		CompressionChunks.SetNum(TotalChunkCount);
 		int64 MaxCompressedSize	= 0;
 		int64 TotalChunkCompressedSize = 0;
 		int64 TotalChunkUncompressedSize = 0;
-		for( int32 ChunkIndex=0; ChunkIndex<TotalChunkCount; ChunkIndex++ )
+		for( int64 ChunkIndex=0; ChunkIndex<TotalChunkCount; ChunkIndex++ )
 		{
 			*this << CompressionChunks[ChunkIndex];
 			if (bWasByteSwapped)
@@ -845,6 +850,10 @@ void FArchive::SerializeCompressedNew(void* V, int64 Length, FName CompressionFo
 				CompressionChunks[ChunkIndex].CompressedSize	= BYTESWAP_ORDER64( CompressionChunks[ChunkIndex].CompressedSize );
 				CompressionChunks[ChunkIndex].UncompressedSize	= BYTESWAP_ORDER64( CompressionChunks[ChunkIndex].UncompressedSize );
 			}
+			
+			UE_CLOG( CompressionChunks[ChunkIndex].CompressedSize < 0 || CompressionChunks[ChunkIndex].UncompressedSize < 0, 
+				LogSerialization, Fatal, TEXT(" Archive SerializedCompressed CompressionChunks[ChunkIndex].CompressedSize (%lld) < 0 || CompressionChunks[ChunkIndex].UncompressedSize (%lld) < 0"), CompressionChunks[ChunkIndex].CompressedSize, CompressionChunks[ChunkIndex].UncompressedSize );
+
 			MaxCompressedSize = FMath::Max( CompressionChunks[ChunkIndex].CompressedSize, MaxCompressedSize );
 
 			TotalChunkCompressedSize += CompressionChunks[ChunkIndex].CompressedSize;
@@ -890,10 +899,9 @@ void FArchive::SerializeCompressedNew(void* V, int64 Length, FName CompressionFo
 
 		// Free up allocated memory.
 		FMemory::Free( CompressedBuffer );
-		delete [] CompressionChunks;
 	}
 	else if( IsSaving() )
-	{	
+	{
 		SCOPE_SECONDS_COUNTER(GArchiveSerializedCompressedSavingTime);
 		check( Length > 0 );
 
@@ -915,6 +923,13 @@ void FArchive::SerializeCompressedNew(void* V, int64 Length, FName CompressionFo
 			// with v2 headers, the modified CompressionFormatToEncode will be written in the archive
 		}
 
+		// compression chunk sizes must fit in int32 for old FCompression API
+		//	(GSavingCompressionChunkSize is an int32 so this is automatic)
+		check( GSavingCompressionChunkSize > 0 );
+		check( GSavingCompressionChunkSize < INT32_MAX );
+		// limit on maximum length we can serialize :
+		check( Length <= (INT64_MAX/2) );
+
 		// Serialize package file tag used to determine endianness in LoadCompressedData.
 		FCompressedChunkInfo PackageFileTag;
 		//PackageFileTag.CompressedSize	= PACKAGE_FILE_TAG;
@@ -926,14 +941,18 @@ void FArchive::SerializeCompressedNew(void* V, int64 Length, FName CompressionFo
 		FCompressionUtil::SerializeCompressorName(*this,CompressionFormatToEncode);
 
 		// Figure out how many chunks there are going to be based on uncompressed size and compression chunk size.
-		//  +1 for Summary chunk
-		int64	TotalChunkCount	= (Length + GSavingCompressionChunkSize - 1) / GSavingCompressionChunkSize + 1;
-		
+		//  divide and round up, overflow safe due to previous range checks
+		int64	TotalChunkCount	= FMath::DivideAndRoundUp( Length, (int64)GSavingCompressionChunkSize );
+
+		//  +1 for Summary chunk	
+		TotalChunkCount += 1;
+
 		// Keep track of current position so we can later seek back and overwrite stub compression chunk infos.
 		int64 StartPosition = Tell();
 
-		// Allocate compression chunk infos and serialize them so we can later overwrite the data.
-		FCompressedChunkInfo* CompressionChunks	= new FCompressedChunkInfo[TotalChunkCount];
+		// Allocate compression chunk infos and serialize them with default fields so we can later overwrite the data.
+		TArray64<FCompressedChunkInfo> CompressionChunks;
+		CompressionChunks.SetNum(TotalChunkCount);
 		for( int64 ChunkIndex=0; ChunkIndex<TotalChunkCount; ChunkIndex++ )
 		{
 			*this << CompressionChunks[ChunkIndex];
@@ -951,7 +970,7 @@ void FArchive::SerializeCompressedNew(void* V, int64 Length, FName CompressionFo
 		FAsyncTask<FAsyncCompressionChunk> AsyncChunks[MAX_COMPRESSION_JOBS];
 
 		// used to keep track of which job is the next one we need to retire
-		int32 AsyncChunkIndex[MAX_COMPRESSION_JOBS]={0};
+		int64 AsyncChunkIndex[MAX_COMPRESSION_JOBS]={0};
 
 		static uint32 GNumUnusedThreads_SerializeCompressed = -1;
 		if (GNumUnusedThreads_SerializeCompressed == (uint32)-1)
@@ -975,7 +994,7 @@ void FArchive::SerializeCompressedNew(void* V, int64 Length, FName CompressionFo
 		}
 
 		// Number of chunks left to finalize.
-		int64 NumChunksLeftToFinalize	= (Length + GSavingCompressionChunkSize - 1) / GSavingCompressionChunkSize;
+		int64 NumChunksLeftToFinalize	= TotalChunkCount -1; // -1 for summary chunk
 		// Number of chunks left to kick off
 		int64 NumChunksLeftToKickOff	= NumChunksLeftToFinalize;
 		// Start at index 1 as first chunk info is summary.
@@ -1149,7 +1168,7 @@ void FArchive::SerializeCompressedNew(void* V, int64 Length, FName CompressionFo
 		}
 		int64		BytesRemaining			= Length;
 		// Start at index 1 as first chunk info is summary.
-		int32		CurrentChunkIndex		= 1;
+		int64		CurrentChunkIndex		= 1;
 
 		int64		CompressedBufferSize	= FCompression::CompressMemoryBound(CompressionFormatToEncode, GSavingCompressionChunkSize);
 		void*	CompressedBuffer		= FMemory::Malloc( CompressedBufferSize );
@@ -1204,15 +1223,12 @@ void FArchive::SerializeCompressedNew(void* V, int64 Length, FName CompressionFo
 		// Seek to the beginning.
 		Seek( StartPosition );
 		// Serialize chunk infos.
-		for( int32 ChunkIndex=0; ChunkIndex<TotalChunkCount; ChunkIndex++ )
+		for( int64 ChunkIndex=0; ChunkIndex<TotalChunkCount; ChunkIndex++ )
 		{
 			*this << CompressionChunks[ChunkIndex];
 		}
 		// Seek back to end.
 		Seek( EndPosition );
-
-		// Free intermediate data.
-		delete [] CompressionChunks;
 	}
 }
 
diff --git a/Engine/Source/Runtime/Core/Public/Serialization/CompressedChunkInfo.h b/Engine/Source/Runtime/Core/Public/Serialization/CompressedChunkInfo.h
--- a/Engine/Source/Runtime/Core/Public/Serialization/CompressedChunkInfo.h
+++ b/Engine/Source/Runtime/Core/Public/Serialization/CompressedChunkInfo.h
@@ -15,10 +15,10 @@ class FArchive;
 struct FCompressedChunkInfo
 {
 	/** Holds the data's compressed size. */
-	int64 CompressedSize;
+	int64 CompressedSize = 0;
 
 	/** Holds the data's uncompresses size. */
-	int64 UncompressedSize;
+	int64 UncompressedSize = 0;
 
 	/**
 	 * Serializes an FCompressedChunkInfo value from or into an archive.
