Add prefetch support for mobile platforms and fuse all platforms into a single implementation

Switch Unix from runtime to compile time cache line constant to avoid loading a value for every PrefetchBlock and allow the compiler to unroll the loop.

diff --git a/Engine/Source/Runtime/Core/Private/Unix/UnixPlatformMisc.cpp b/Engine/Source/Runtime/Core/Private/Unix/UnixPlatformMisc.cpp
--- a/Engine/Source/Runtime/Core/Private/Unix/UnixPlatformMisc.cpp
+++ b/Engine/Source/Runtime/Core/Private/Unix/UnixPlatformMisc.cpp
@@ -107,26 +107,6 @@ void FUnixPlatformMisc::NormalizePath(FStringBuilderBase& InPath)
 	}
 }
 
-size_t CORE_API GCacheLineSize = PLATFORM_CACHE_LINE_SIZE;
-
-void UnixPlatform_UpdateCacheLineSize()
-{
-	// sysfs "API", as usual ;/
-	FILE * SysFsFile = fopen("/sys/devices/system/cpu/cpu0/cache/index0/coherency_line_size", "r");
-	if (SysFsFile)
-	{
-		int SystemLineSize = 0;
-		if (1 == fscanf(SysFsFile, "%d", &SystemLineSize))
-		{
-			if (SystemLineSize > 0)
-			{
-				GCacheLineSize = SystemLineSize;
-			}
-		}
-		fclose(SysFsFile);
-	}
-}
-
 // Defined in UnixPlatformMemory
 extern bool GUseKSM;
 extern bool GKSMMergeAllPages;
@@ -223,8 +203,6 @@ void FUnixPlatformMisc::PlatformInit()
 	UE_LOG(LogInit, Log, TEXT(" - CPU: %s '%s' (signature: 0x%X)"), *FPlatformMisc::GetCPUVendor(), *FPlatformMisc::GetCPUBrand(), FPlatformMisc::GetCPUInfo());
 	UE_LOG(LogInit, Log, TEXT(" - Number of physical cores available for the process: %d"), FPlatformMisc::NumberOfCores());
 	UE_LOG(LogInit, Log, TEXT(" - Number of logical cores available for the process: %d"), FPlatformMisc::NumberOfCoresIncludingHyperthreads());
-	UnixPlatform_UpdateCacheLineSize();
-	UE_LOG(LogInit, Log, TEXT(" - Cache line size: %zu"), GCacheLineSize);
 
 	if (!GPUBrandInfo.IsEmpty())
 	{
diff --git a/Engine/Source/Runtime/Core/Public/GenericPlatform/GenericPlatformMisc.h b/Engine/Source/Runtime/Core/Public/GenericPlatform/GenericPlatformMisc.h
--- a/Engine/Source/Runtime/Core/Public/GenericPlatform/GenericPlatformMisc.h
+++ b/Engine/Source/Runtime/Core/Public/GenericPlatform/GenericPlatformMisc.h
@@ -10,6 +10,10 @@
 #include "Misc/CompressionFlags.h"
 #include "Misc/EnumClassFlags.h"
 
+#if PLATFORM_CPU_X86_FAMILY
+#include <xmmintrin.h>
+#endif
+
 class Error;
 class FOutputDevice;
 class FString;
@@ -1373,13 +1377,35 @@ public:
 	{
 	}
 
-	FORCEINLINE static void PrefetchBlock(const void* InPtr, int32 NumBytes = 1)
+	FORCEINLINE static void Prefetch(const void* Ptr)
+	{
+#if PLATFORM_CPU_X86_FAMILY
+		_mm_prefetch(static_cast<const char*>(Ptr), _MM_HINT_T0);
+#elif PLATFORM_CPU_ARM_FAMILY
+		__asm__ __volatile__("prfm pldl1keep, [%[ptr]]\n" ::[ptr] "r"(Ptr) : );
+#else
+#error Unknown architecture
+#endif
+	}
+
+	FORCEINLINE static void Prefetch(const void* Ptr, int32 Offset)
+	{
+		Prefetch(reinterpret_cast<const void*>(reinterpret_cast<UPTRINT>(Ptr) + Offset));
+	}
+
+	UE_DEPRECATED(5.2, "Must supply size when prefetching a block of data")
+	FORCEINLINE static void PrefetchBlock(const void* Ptr)
 	{
+		Prefetch(Ptr);
 	}
 
-	/** Platform-specific instruction prefetch */
-	FORCEINLINE static void Prefetch(void const* x, int32 offset = 0)
+	FORCEINLINE static void PrefetchBlock(const void* Ptr, int32 NumBytes)
 	{
+		// Use compile-time PLATFORM_CACHE_LINE_SIZE to avoid having to load a value before issuing prefetch instructions
+		for (int32 Line = 0, NumLines = (NumBytes + PLATFORM_CACHE_LINE_SIZE - 1) / PLATFORM_CACHE_LINE_SIZE; Line < NumLines; ++Line)
+		{
+			Prefetch(Ptr, Line * PLATFORM_CACHE_LINE_SIZE);
+		}
 	}
 
 	/**
diff --git a/Engine/Source/Runtime/Core/Public/Unix/UnixPlatformMisc.h b/Engine/Source/Runtime/Core/Public/Unix/UnixPlatformMisc.h
--- a/Engine/Source/Runtime/Core/Public/Unix/UnixPlatformMisc.h
+++ b/Engine/Source/Runtime/Core/Public/Unix/UnixPlatformMisc.h
@@ -61,24 +61,6 @@ struct CORE_API FUnixPlatformMisc : public FGenericPlatformMisc
 		__sync_synchronize();
 	}
 
-	FORCEINLINE static void PrefetchBlock(const void* InPtr, int32 NumBytes = 1)
-	{
-		extern size_t GCacheLineSize;
-
-		const char* Ptr = static_cast<const char*>(InPtr);
-		const size_t CacheLineSize = GCacheLineSize;
-		for (size_t BytesPrefetched = 0; BytesPrefetched < NumBytes; BytesPrefetched += CacheLineSize)
-		{
-			__builtin_prefetch(Ptr);
-			Ptr += CacheLineSize;
-		}
-	}
-
-	FORCEINLINE static void Prefetch(void const* Ptr, int32 Offset = 0)
-	{
-		__builtin_prefetch(static_cast<char const*>(Ptr) + Offset);
-	}
-
 	static int32 NumberOfCores();
 	static int32 NumberOfCoresIncludingHyperthreads();
 	static FString GetOperatingSystemId();
diff --git a/Engine/Source/Runtime/Core/Public/Windows/WindowsPlatformMisc.h b/Engine/Source/Runtime/Core/Public/Windows/WindowsPlatformMisc.h
--- a/Engine/Source/Runtime/Core/Public/Windows/WindowsPlatformMisc.h
+++ b/Engine/Source/Runtime/Core/Public/Windows/WindowsPlatformMisc.h
@@ -163,30 +163,6 @@ struct CORE_API FWindowsPlatformMisc
 	static void PromptForRemoteDebugging(bool bIsEnsure);
 #endif	//#if !UE_BUILD_SHIPPING
 
-	FORCEINLINE static void PrefetchBlock(const void* InPtr, int32 NumBytes = 1)
-	{
-		const char* Ptr           = (const char*)InPtr;
-		const int32 CacheLineSize = GetCacheLineSize();
-		for (int32 LinesToPrefetch = (NumBytes + CacheLineSize - 1) / CacheLineSize; LinesToPrefetch; --LinesToPrefetch)
-		{
-#if PLATFORM_CPU_X86_FAMILY
-			_mm_prefetch( Ptr, _MM_HINT_T0 );
-#elif PLATFORM_CPU_ARM_FAMILY
-			__prefetch(Ptr);
-#endif
-			Ptr += CacheLineSize;
-		}
-	}
-
-	FORCEINLINE static void Prefetch(void const* x, int32 offset = 0)
-	{
-#if PLATFORM_CPU_X86_FAMILY
-		_mm_prefetch((char const*)(x)+offset, _MM_HINT_T0);
-#elif PLATFORM_CPU_ARM_FAMILY
-		__prefetch((char const*)(x)+offset);
-#endif
-		
-	}
 
 	/** 
 	 * Determines if the cpuid instruction is supported on this processor
