Made TArray::GetTypeSize() and TIndirectArray::GetTypeSize() constexpr and static.

diff --git a/Engine/Source/Runtime/Core/Public/Containers/Array.h b/Engine/Source/Runtime/Core/Public/Containers/Array.h
--- a/Engine/Source/Runtime/Core/Public/Containers/Array.h
+++ b/Engine/Source/Runtime/Core/Public/Containers/Array.h
@@ -700,7 +700,7 @@ public:
 	 *
 	 * @returns Size in bytes of array type.
 	 */
-	FORCEINLINE uint32 GetTypeSize() const
+	FORCEINLINE static constexpr uint32 GetTypeSize()
 	{
 		return sizeof(ElementType);
 	}
diff --git a/Engine/Source/Runtime/Core/Public/Containers/IndirectArray.h b/Engine/Source/Runtime/Core/Public/Containers/IndirectArray.h
--- a/Engine/Source/Runtime/Core/Public/Containers/IndirectArray.h
+++ b/Engine/Source/Runtime/Core/Public/Containers/IndirectArray.h
@@ -120,7 +120,7 @@ public:
 	 *
 	 * @returns Size in bytes of array type.
 	 */
-	uint32 GetTypeSize() const
+	static constexpr uint32 GetTypeSize()
 	{
 		return sizeof(T*);
 	}
