Add Android memory mapped file support

diff --git a/Engine/Source/Runtime/Core/Private/Android/AndroidPlatformFile.cpp b/Engine/Source/Runtime/Core/Private/Android/AndroidPlatformFile.cpp
--- a/Engine/Source/Runtime/Core/Private/Android/AndroidPlatformFile.cpp
+++ b/Engine/Source/Runtime/Core/Private/Android/AndroidPlatformFile.cpp
@@ -20,7 +20,9 @@
 #include "Android/AndroidJava.h"
 #include "Containers/Map.h"
 #include <limits>
+#include <sys/mman.h>
 #include "Misc/ScopeLock.h"
+#include <Async/MappedFileHandle.h>
 
 #if PLATFORM_USE_PLATFORM_FILE_MANAGED_STORAGE_WRAPPER
 #include "HAL/IPlatformFileManagedStorageWrapper.h"
@@ -962,6 +964,106 @@ private:
 	};
 };
 
+
+class FAndroidMappedFileRegion final : public IMappedFileRegion
+{
+public:
+	class FAndroidMappedFileHandle* Parent;
+	const uint8* AlignedPtr;
+	uint64 AlignedSize;
+	FAndroidMappedFileRegion(const uint8* InMappedPtr, const uint8* InAlignedPtr, size_t InMappedSize, uint64 InAlignedSize, const FString& InDebugFilename, size_t InDebugOffsetIntoFile, FAndroidMappedFileHandle* InParent)
+		: IMappedFileRegion(InMappedPtr, InMappedSize, InDebugFilename, InDebugOffsetIntoFile)
+		, Parent(InParent)
+		, AlignedPtr(InAlignedPtr)
+		, AlignedSize(InAlignedSize)
+	{
+	}
+
+	virtual ~FAndroidMappedFileRegion();
+};
+
+class FAndroidMappedFileHandle final : public IMappedFileHandle
+{
+	inline static SIZE_T FileMappingAlignment = FPlatformMemory::GetConstants().PageSize;
+
+public:
+	FAndroidMappedFileHandle(int InFileHandle, int64 FileSize, const FString& InFilename)
+		: IMappedFileHandle(FileSize)
+		, MappedPtr(nullptr)
+		, Filename(InFilename)
+		, NumOutstandingRegions(0)
+		, FileHandle(InFileHandle)
+	{
+	}
+
+	virtual ~FAndroidMappedFileHandle() override
+	{
+		check(!NumOutstandingRegions); // can't delete the file before you delete all outstanding regions
+		close(FileHandle);
+	}
+
+	virtual IMappedFileRegion* MapRegion(int64 Offset = 0, int64 BytesToMap = MAX_int64, bool bPreloadHint = false) override
+	{
+		LLM_PLATFORM_SCOPE(ELLMTag::PlatformMMIO);
+		check(Offset < GetFileSize()); // don't map zero bytes and don't map off the end of the file
+		BytesToMap = FMath::Min<int64>(BytesToMap, GetFileSize() - Offset);
+		check(BytesToMap > 0); // don't map zero bytes
+
+		const int64 AlignedOffset = AlignDown(Offset, FileMappingAlignment);
+		//File mapping can extend beyond file size. It's OK, kernel will just fill any leftover page data with zeros
+		const int64 AlignedSize = Align(BytesToMap + Offset - AlignedOffset, FileMappingAlignment);
+
+		int Flags = MAP_PRIVATE;
+		if (bPreloadHint)
+		{
+			Flags |= MAP_POPULATE;
+		}
+
+		const uint8* AlignedMapPtr = static_cast<const uint8*>(mmap(nullptr, AlignedSize, PROT_READ, Flags, FileHandle, AlignedOffset));
+		if (AlignedMapPtr == MAP_FAILED || AlignedMapPtr == nullptr)
+		{
+#if LOG_ANDROID_FILE
+			FPlatformMisc::LowLevelOutputDebugStringf(TEXT("Failed to mmap region from %s, errno=%s"), *Filename, UTF8_TO_TCHAR(strerror(errno)));
+#endif
+			UE_LOG(LogAndroidFile, Warning, TEXT("Failed to map memory %s, error is %d"), *Filename, errno);
+			return nullptr;
+		}
+		LLM(FLowLevelMemTracker::Get().OnLowLevelAlloc(ELLMTracker::Platform, AlignedMapPtr, AlignedSize));
+
+		// create a mapping for this range
+		const uint8* MapPtr = AlignedMapPtr + Offset - AlignedOffset;
+		FAndroidMappedFileRegion* Result = new FAndroidMappedFileRegion(MapPtr, AlignedMapPtr, BytesToMap, AlignedSize, Filename, Offset, this);
+		NumOutstandingRegions++;
+		return Result;
+	}
+
+	void UnMap(const FAndroidMappedFileRegion* Region)
+	{
+		LLM_PLATFORM_SCOPE(ELLMTag::PlatformMMIO);
+		check(NumOutstandingRegions > 0);
+		NumOutstandingRegions--;
+
+		LLM(FLowLevelMemTracker::Get().OnLowLevelFree(ELLMTracker::Platform, Region->AlignedPtr));
+		const int Res = munmap(const_cast<uint8*>(Region->AlignedPtr), Region->AlignedSize);
+#if LOG_ANDROID_FILE
+		FPlatformMisc::LowLevelOutputDebugStringf(TEXT("Failed to unmap region from %s, errno=%s"), *Filename, UTF8_TO_TCHAR(strerror(errno)));
+#endif
+		checkf(Res == 0, TEXT("Failed to unmap, error is %d, errno is %d [params: %x, %d]"), Res, errno, MappedPtr, GetFileSize());
+	}
+
+private:
+	const uint8* MappedPtr;
+	FString Filename;
+	int32 NumOutstandingRegions;
+	int FileHandle;
+};
+
+FAndroidMappedFileRegion::~FAndroidMappedFileRegion()
+{
+	Parent->UnMap(this);
+}
+
+
 // NOTE: Files are stored either loosely in the deployment directory
 // or packed in an OBB archive. We don't know which one unless we try
 // and get the files. We always first check if the files are local,
@@ -1211,6 +1313,43 @@ public:
 		return result;
 	}
 
+	virtual IMappedFileHandle* OpenMapped(const TCHAR* Filename) override
+	{
+		FString LocalPath;
+		FString AssetPath;
+		PathToAndroidPaths(LocalPath, AssetPath, Filename, false);
+
+		const FString NormalizedFilename = LocalPath;
+
+		constexpr int Flags = O_RDONLY;
+		const int32 Handle = open(TCHAR_TO_UTF8(*NormalizedFilename), Flags);
+		if (Handle == -1)
+		{
+			const int ErrNo = errno;
+#if LOG_ANDROID_FILE
+			FPlatformMisc::LowLevelOutputDebugStringf(TEXT("FAndroidPlatformFile::OpenMapped('%s', Flags=0x%08X) failed: errno=%d (%s)"), *NormalizedFilename, Flags, ErrNo, UTF8_TO_TCHAR(strerror(ErrNo)));
+#endif
+
+			return nullptr;
+		}
+
+		struct stat FileInfo;
+		FileInfo.st_size = -1;
+		const int StatResult = fstat(Handle, &FileInfo);
+		if (StatResult == -1)
+		{
+			const int ErrNo = errno;
+
+#if LOG_ANDROID_FILE
+			FPlatformMisc::LowLevelOutputDebugStringf(TEXT("FAndroidPlatformFile::OpenMapped fstat failed: ('%s', Flags=0x%08X) failed: errno=%d (%s)"), *NormalizedFilename, Flags, ErrNo, UTF8_TO_TCHAR(strerror(ErrNo)));
+#endif
+
+			return nullptr;
+		}
+
+		return new FAndroidMappedFileHandle(Handle, FileInfo.st_size, NormalizedFilename);
+	}
+
 	virtual int64 FileSize(const TCHAR* Filename) override
 	{
 		return FileSize(Filename, false);
diff --git a/Engine/Source/Runtime/Core/Public/Android/AndroidPlatformProperties.h b/Engine/Source/Runtime/Core/Public/Android/AndroidPlatformProperties.h
--- a/Engine/Source/Runtime/Core/Public/Android/AndroidPlatformProperties.h
+++ b/Engine/Source/Runtime/Core/Public/Android/AndroidPlatformProperties.h
@@ -119,6 +119,11 @@ struct FAndroidPlatformProperties
 		return true;
 	}
 
+	static FORCEINLINE bool SupportsMemoryMappedFiles()
+	{
+		return true;
+	}
+
 	static FORCEINLINE bool SupportsVirtualTextureStreaming()
 	{
 		return true;
