Fix ARM on Windows

#rnx

diff --git a/Engine/Source/Runtime/Core/Public/GenericPlatform/GenericPlatformMisc.h b/Engine/Source/Runtime/Core/Public/GenericPlatform/GenericPlatformMisc.h
--- a/Engine/Source/Runtime/Core/Public/GenericPlatform/GenericPlatformMisc.h
+++ b/Engine/Source/Runtime/Core/Public/GenericPlatform/GenericPlatformMisc.h
@@ -11,7 +11,9 @@
 #include "Misc/EnumClassFlags.h"
 
 #if PLATFORM_CPU_X86_FAMILY
-#include <xmmintrin.h>
+#include <xmmintrin.h> // _mm_prefetch
+#elif PLATFORM_CPU_ARM_FAMILY && defined(_MSC_VER)
+#include <intrin.h> // __prefetch
 #endif
 
 class Error;
@@ -1382,9 +1384,13 @@ public:
 #if PLATFORM_CPU_X86_FAMILY
 		_mm_prefetch(static_cast<const char*>(Ptr), _MM_HINT_T0);
 #elif PLATFORM_CPU_ARM_FAMILY
+#	if defined(_MSC_VER)
+		__prefetch(Ptr);
+#	else
 		__asm__ __volatile__("prfm pldl1keep, [%[ptr]]\n" ::[ptr] "r"(Ptr) : );
+#	endif
 #else
-#error Unknown architecture
+#	error Unknown architecture
 #endif
 	}
 
