Reduce binary bloat

diff --git a/Engine/Source/Runtime/Core/Public/Math/Box.h b/Engine/Source/Runtime/Core/Public/Math/Box.h
--- a/Engine/Source/Runtime/Core/Public/Math/Box.h
+++ b/Engine/Source/Runtime/Core/Public/Math/Box.h
@@ -19,6 +19,8 @@
 namespace UE {
 namespace Math {
 
+struct TBoxConstInit {};
+
 template<typename T>
 struct TBox
 {
@@ -49,6 +51,11 @@ public:
 		Init();
 	}
 
+	constexpr TBox(EForceInit, TBoxConstInit)
+		: Min(0, TVectorConstInit{}), Max(0, TVectorConstInit{}), IsValid(0)
+	{
+	}
+
 	/**
 	 * Creates and initializes a new box from the specified extents.
 	 *
diff --git a/Engine/Source/Runtime/Core/Public/Math/Vector.h b/Engine/Source/Runtime/Core/Public/Math/Vector.h
--- a/Engine/Source/Runtime/Core/Public/Math/Vector.h
+++ b/Engine/Source/Runtime/Core/Public/Math/Vector.h
@@ -41,6 +41,8 @@ namespace UE
 namespace Math
 {
 
+struct TVectorConstInit {};
+
 /**
  * A vector in 3-D space composed of components (X, Y, Z) with floating point precision.
  */
@@ -153,6 +155,8 @@ public:
      */
     explicit FORCEINLINE TVector(T InF);
 
+	FORCEINLINE constexpr TVector(T InF, TVectorConstInit);
+
     /**
      * Constructor using initial values for each component.
      *
@@ -1390,6 +1394,12 @@ FORCEINLINE TVector<T>::TVector(T InF)
     DiagnosticCheckNaN();
 }
 
+template<typename T>
+FORCEINLINE constexpr TVector<T>::TVector(T InF, TVectorConstInit)
+	: X(InF), Y(InF), Z(InF)
+{
+}
+
 template<typename T>
 FORCEINLINE TVector<T>::TVector(T InX, T InY, T InZ)
     : X(InX), Y(InY), Z(InZ)
diff --git a/Engine/Source/Runtime/Engine/Classes/AI/Navigation/NavigationTypes.h b/Engine/Source/Runtime/Engine/Classes/AI/Navigation/NavigationTypes.h
--- a/Engine/Source/Runtime/Engine/Classes/AI/Navigation/NavigationTypes.h
+++ b/Engine/Source/Runtime/Engine/Classes/AI/Navigation/NavigationTypes.h
@@ -48,9 +48,9 @@ namespace FNavigationSystem
 	/** used as a fallback value for navigation agent height, when none specified via UNavigationSystemV1::SupportedAgents */
 	extern ENGINE_API const float FallbackAgentHeight;
 
-	static const FBox InvalidBoundingBox(ForceInit);
+	inline constexpr FBox InvalidBoundingBox(ForceInit, UE::Math::TBoxConstInit{});
 
-	static const FVector InvalidLocation = FVector(TNumericLimits<FVector::FReal>::Max());
+	inline constexpr FVector InvalidLocation = FVector(TNumericLimits<FVector::FReal>::Max(), UE::Math::TVectorConstInit{});
 
 	FORCEINLINE bool IsValidLocation(const FVector& TestLocation)
 	{
