Added support for moderation artifact generation automated testing.

diff --git a/Engine/Source/Runtime/Core/Public/Misc/MonitoredProcess.h b/Engine/Source/Runtime/Core/Public/Misc/MonitoredProcess.h
--- a/Engine/Source/Runtime/Core/Public/Misc/MonitoredProcess.h
+++ b/Engine/Source/Runtime/Core/Public/Misc/MonitoredProcess.h
@@ -91,6 +91,14 @@ public:
 		return ProcessHandle;
 	}
 
+	/**
+	 * Returns the commandline of the process which will be executed if Launch is called
+	 */
+	FString GetCommandline() const
+	{
+		return FString::Printf(TEXT("%s %s"), *URL, *Params);
+	}
+
 	/**
 	* Checks whether the process is still running. In single threaded mode, this will tick the thread processing
 	*
