Modify AES num bytes type to prevent truncation fixes and warnings with its usage

diff --git a/Engine/Source/Runtime/Core/Private/Misc/AES.cpp b/Engine/Source/Runtime/Core/Private/Misc/AES.cpp
--- a/Engine/Source/Runtime/Core/Private/Misc/AES.cpp
+++ b/Engine/Source/Runtime/Core/Private/Misc/AES.cpp
@@ -457,7 +457,7 @@ static inline void AesDecryptExpand(FAesExpandedKey* DecryptKey, const uint8* Ke
 	DKey[3] = EKey[3];
 }
 
-static inline void AesEncrypt(const uint8* Key, uint8* Contents, uint32 NumBytes)
+static inline void AesEncrypt(const uint8* Key, uint8* Contents, uint64 NumBytes)
 {
 #define AES_ENC(a, b, c, d) (      \
 	AesTbox[0][(uint8)(a >>  0)] ^ \
@@ -478,7 +478,7 @@ static inline void AesEncrypt(const uint8* Key, uint8* Contents, uint32 NumBytes
 
 	const uint32* EKey = EncryptKey.Key;
 
-	for (uint32 Offset = 0; Offset < NumBytes; Offset += 16)
+	for (uint64 Offset = 0; Offset < NumBytes; Offset += 16)
 	{
 		uint8* Block = Contents + Offset;
 
@@ -520,7 +520,7 @@ static inline void AesEncrypt(const uint8* Key, uint8* Contents, uint32 NumBytes
 #undef AES_ENC_LAST
 }
 
-static inline void AesDecrypt(const uint8* Key, uint8* Contents, uint32 NumBytes)
+static inline void AesDecrypt(const uint8* Key, uint8* Contents, uint64 NumBytes)
 {
 #define AES_DEC(a, b, c, d) (       \
 	AesITbox[0][(uint8)(a >>  0)] ^ \
@@ -541,7 +541,7 @@ static inline void AesDecrypt(const uint8* Key, uint8* Contents, uint32 NumBytes
 
 	const uint32* DKey = DecryptKey.Key;
 
-	for (uint32 Offset = 0; Offset < NumBytes; Offset += 16)
+	for (uint64 Offset = 0; Offset < NumBytes; Offset += 16)
 	{
 		uint8* Block = Contents + Offset;
 
@@ -646,7 +646,7 @@ static bool CanUseAesInstructions()
 #pragma clang attribute push (__attribute__((target("aes"))), apply_to=function)
 #endif
 
-static inline void AesEncryptX86(const void* Key, void* Contents, uint32 NumBytes)
+static inline void AesEncryptX86(const void* Key, void* Contents, uint64 NumBytes)
 {
 	const __m128i* Key128 = reinterpret_cast<const __m128i*>(Key);
 	__m128i x = _mm_loadu_si128(Key128 + 0);
@@ -738,7 +738,7 @@ static inline void AesEncryptX86(const void* Key, void* Contents, uint32 NumByte
 	}
 }
 
-static inline void AesDecryptX86(const void* Key, void* Contents, uint32 NumBytes)
+static inline void AesDecryptX86(const void* Key, void* Contents, uint64 NumBytes)
 {
 	const __m128i* Key128 = reinterpret_cast<const __m128i*>(Key);
 	__m128i x = _mm_loadu_si128(Key128 + 0);
@@ -899,14 +899,14 @@ static inline void AesDecryptExpandARMV8(FAesExpandedKeyARMV8* DecryptKey, const
 	DecryptKey->Key[AES256_ROUND_COUNT] = EncryptKey.Key[0];
 }
 
-static inline void AesEncryptARMV8(const uint8* Key, uint8* Contents, uint32 NumBytes)
+static inline void AesEncryptARMV8(const uint8* Key, uint8* Contents, uint64 NumBytes)
 {
 	FAesExpandedKeyARMV8 EncryptKey;
 	AesEncryptExpandARMV8(&EncryptKey, Key);
 
 	const uint8x16_t* EKey = EncryptKey.Key;
 
-	for (uint32 Offset = 0; Offset < NumBytes; Offset += 16)
+	for (uint64 Offset = 0; Offset < NumBytes; Offset += 16)
 	{
 		uint8* Data = Contents + Offset;
 
@@ -930,14 +930,14 @@ static inline void AesEncryptARMV8(const uint8* Key, uint8* Contents, uint32 Num
 	}
 }
 
-static inline void AesDecryptARMV8(const uint8* Key, uint8* Contents, uint32 NumBytes)
+static inline void AesDecryptARMV8(const uint8* Key, uint8* Contents, uint64 NumBytes)
 {
 	FAesExpandedKeyARMV8 DecryptKey;
 	AesDecryptExpandARMV8(&DecryptKey, Key);
 
 	const uint8x16_t* DKey = DecryptKey.Key;
 
-	for (uint32 Offset = 0; Offset < NumBytes; Offset += 16)
+	for (uint64 Offset = 0; Offset < NumBytes; Offset += 16)
 	{
 		uint8* Data = Contents + Offset;
 
@@ -964,19 +964,19 @@ static inline void AesDecryptARMV8(const uint8* Key, uint8* Contents, uint32 Num
 #endif
 
 
-void FAES::EncryptData(uint8 *Contents, uint32 NumBytes, const FAESKey& Key)
+void FAES::EncryptData(uint8 *Contents, uint64 NumBytes, const FAESKey& Key)
 {
 	checkf(Key.IsValid(), TEXT("No valid encryption key specified"));
 	EncryptData(Contents, NumBytes, (const uint8*)Key.Key, sizeof(Key.Key));
 }
 
-void FAES::EncryptData(uint8 *Contents, uint32 NumBytes, const ANSICHAR* Key)
+void FAES::EncryptData(uint8 *Contents, uint64 NumBytes, const ANSICHAR* Key)
 {
 	checkf(Key, TEXT("No encryption key specified"));
 	EncryptData(Contents, NumBytes, (const uint8*)Key, TCString<ANSICHAR>::Strlen(Key));
 }
 
-void FAES::EncryptData(uint8* Contents, uint32 NumBytes, const uint8* KeyBytes, uint32 NumKeyBytes)
+void FAES::EncryptData(uint8* Contents, uint64 NumBytes, const uint8* KeyBytes, uint32 NumKeyBytes)
 {
 	checkf(NumBytes % AESBlockSize == 0, TEXT("NumBytes needs to be a multiple of 16 bytes"));
 	checkf(NumKeyBytes == FAESKey::KeySize, TEXT("AES-256 key needs to be 32 bytes long"));
@@ -999,19 +999,19 @@ void FAES::EncryptData(uint8* Contents, uint32 NumBytes, const uint8* KeyBytes,
 #endif
 }
 
-void FAES::DecryptData(uint8* Contents, uint32 NumBytes, const FAESKey& Key)
+void FAES::DecryptData(uint8* Contents, uint64 NumBytes, const FAESKey& Key)
 {
 	checkf(Key.IsValid(), TEXT("No valid decryption key specified"));
 	DecryptData(Contents, NumBytes, Key.Key, sizeof(Key.Key));
 }
 
-void FAES::DecryptData(uint8 *Contents, uint32 NumBytes, const ANSICHAR* Key)
+void FAES::DecryptData(uint8 *Contents, uint64 NumBytes, const ANSICHAR* Key)
 {
 	checkf(Key, TEXT("No valid decryption key specified"));
 	DecryptData(Contents, NumBytes, (const uint8*)Key, TCString<ANSICHAR>::Strlen(Key));
 }
 
-void FAES::DecryptData(uint8* Contents, uint32 NumBytes, const uint8* KeyBytes, uint32 NumKeyBytes)
+void FAES::DecryptData(uint8* Contents, uint64 NumBytes, const uint8* KeyBytes, uint32 NumKeyBytes)
 {
 	checkf(NumBytes % AESBlockSize == 0, TEXT("NumBytes needs to be a multiple of 16 bytes"));
 	checkf(NumKeyBytes == FAESKey::KeySize, TEXT("AES-256 key needs to be 32 bytes long"));
diff --git a/Engine/Source/Runtime/Core/Public/Misc/AES.h b/Engine/Source/Runtime/Core/Public/Misc/AES.h
--- a/Engine/Source/Runtime/Core/Public/Misc/AES.h
+++ b/Engine/Source/Runtime/Core/Public/Misc/AES.h
@@ -65,7 +65,7 @@ struct CORE_API FAES
 	* @param NumBytes the size of the buffer
 	* @param Key An FAESKey object containing the encryption key
 	*/
-	static void EncryptData(uint8* Contents, uint32 NumBytes, const FAESKey& Key);
+	static void EncryptData(uint8* Contents, uint64 NumBytes, const FAESKey& Key);
 
 	/**
 	 * Encrypts a chunk of data using a specific key
@@ -74,7 +74,7 @@ struct CORE_API FAES
 	 * @param NumBytes the size of the buffer
 	 * @param Key a null terminated string that is a 32 bytes long
 	 */
-	static void EncryptData(uint8* Contents, uint32 NumBytes, const ANSICHAR* Key);
+	static void EncryptData(uint8* Contents, uint64 NumBytes, const ANSICHAR* Key);
 
 	/**
 	* Encrypts a chunk of data using a specific key
@@ -84,7 +84,7 @@ struct CORE_API FAES
 	* @param Key a byte array that is a 32 byte length
 	* @param NumKeyBytes length of Key byte array, must be 32
 	*/
-	static void EncryptData(uint8* Contents, uint32 NumBytes, const uint8* KeyBytes, uint32 NumKeyBytes);
+	static void EncryptData(uint8* Contents, uint64 NumBytes, const uint8* KeyBytes, uint32 NumKeyBytes);
 
 	/**
 	* Decrypts a chunk of data using a specific key
@@ -93,7 +93,7 @@ struct CORE_API FAES
 	* @param NumBytes the size of the buffer
 	* @param Key An FAESKey object containing the decryption key
 	*/
-	static void DecryptData(uint8* Contents, uint32 NumBytes, const FAESKey& Key);
+	static void DecryptData(uint8* Contents, uint64 NumBytes, const FAESKey& Key);
 
 	/**
 	 * Decrypts a chunk of data using a specific key
@@ -102,7 +102,7 @@ struct CORE_API FAES
 	 * @param NumBytes the size of the buffer
 	 * @param Key a null terminated string that is a 32 bytes long
 	 */
-	static void DecryptData(uint8* Contents, uint32 NumBytes, const ANSICHAR* Key);
+	static void DecryptData(uint8* Contents, uint64 NumBytes, const ANSICHAR* Key);
 
 	/**
 	* Decrypts a chunk of data using a specific key
@@ -112,5 +112,5 @@ struct CORE_API FAES
 	* @param Key a byte array that is a 32 byte length
 	* @param NumKeyBytes length of Key byte array, must be 32
 	*/
-	static void DecryptData(uint8* Contents, uint32 NumBytes, const uint8* KeyBytes, uint32 NumKeyBytes);
+	static void DecryptData(uint8* Contents, uint64 NumBytes, const uint8* KeyBytes, uint32 NumKeyBytes);
 };
