Added "AddExpectedMessage" to Automation Tests, so that any log entry can be intecepted and not just warnings/errors

diff --git a/Engine/Plugins/Tests/RuntimeTests/Source/RuntimeTests/Private/AutomationExpectedError.spec.cpp b/Engine/Plugins/Tests/RuntimeTests/Source/RuntimeTests/Private/AutomationExpectedError.spec.cpp
deleted file mode 100644
diff --git a/Engine/Plugins/Tests/RuntimeTests/Source/RuntimeTests/Private/AutomationExpectedMessage.spec.cpp b/Engine/Plugins/Tests/RuntimeTests/Source/RuntimeTests/Private/AutomationExpectedMessage.spec.cpp
new file mode 100644
--- /dev/null
+++ b/Engine/Plugins/Tests/RuntimeTests/Source/RuntimeTests/Private/AutomationExpectedMessage.spec.cpp
@@ -0,0 +1,214 @@
+// Copyright Epic Games, Inc. All Rights Reserved.
+
+#include "Misc/AutomationTest.h"
+
+BEGIN_DEFINE_SPEC(FAutomationExpectedErrorTest, "System.Automation.ExpectedError", EAutomationTestFlags::EngineFilter | EAutomationTestFlags::ApplicationContextMask)
+END_DEFINE_SPEC(FAutomationExpectedErrorTest)
+void FAutomationExpectedErrorTest::Define()
+{
+	Describe("A defined expected error in a test", [this]()
+	{
+
+		It("will not add an error with a number of occurrences less than zero", [this]()
+		{
+			// Suppress error logged when adding entry with invalid occurrence count
+			AddExpectedError(TEXT("number of expected occurrences must be >= 0"), EAutomationExpectedErrorFlags::Contains, 1);
+
+			AddExpectedError(TEXT("The two values are not equal"), EAutomationExpectedErrorFlags::Contains,  -1);
+
+			TArray<FAutomationExpectedMessage> Errors;
+			GetExpectedMessages(Errors, ELogVerbosity::Warning);
+
+			// Only the first expected error should exist in the list.
+			TestEqual("Expected Errors Count", Errors.Num(), 1);
+		});
+
+		It("will add an error with a number of occurrences equal to zero", [this]()
+		{
+			AddExpectedError(TEXT("Expected Error"), EAutomationExpectedErrorFlags::Contains, 0);
+
+			TArray<FAutomationExpectedMessage> Errors;
+			GetExpectedMessages(Errors, ELogVerbosity::Warning);
+			TestEqual("Expected Errors Count", Errors.Num(), 1);
+
+			// Add the expected error to ensure all test conditions pass
+			AddError(TEXT("Expected Error"));
+		});
+
+		It("will not duplicate an existing expected error using the same matcher", [this]()
+		{
+			// Suppress warning logged when adding duplicate value
+			AddExpectedError(TEXT("cannot add duplicate entries"), EAutomationExpectedErrorFlags::Contains, 1);
+			
+			AddExpectedError(TEXT("Expected Error"), EAutomationExpectedErrorFlags::Contains, 1);
+			AddExpectedError(TEXT("Expected Error"), EAutomationExpectedErrorFlags::Contains, 1);
+
+			TArray<FAutomationExpectedMessage> Errors;
+			GetExpectedMessages(Errors, ELogVerbosity::Warning);
+			TestEqual("Expected Errors Count", Errors.Num(), 2);
+
+			// Add the expected error to ensure all test conditions pass
+			AddError(TEXT("Expected Error"));
+		});
+
+		It("will not duplicate an expected error using a different matcher", [this]()
+		{
+			// Suppress warning logged when adding duplicate value
+			AddExpectedError(TEXT("cannot add duplicate entries"), EAutomationExpectedErrorFlags::Contains, 2);
+
+			AddExpectedError(TEXT("Expected Exact Error"), EAutomationExpectedErrorFlags::Exact, 1);
+			AddExpectedError(TEXT("Expected Exact Error"), EAutomationExpectedErrorFlags::Contains, 1);
+
+			AddExpectedError(TEXT("Expected Contains Error"), EAutomationExpectedErrorFlags::Contains, 1);
+			AddExpectedError(TEXT("Expected Contains Error"), EAutomationExpectedErrorFlags::Exact, 1);
+
+			TArray<FAutomationExpectedMessage> Errors;
+			GetExpectedMessages(Errors, ELogVerbosity::Warning);
+			TestEqual("Expected Errors Count", Errors.Num(), 3);
+
+			// Add the expected errors to ensure all test conditions pass
+			AddError(TEXT("Expected Exact Error"));
+			AddError(TEXT("Expected Contains Error"));
+		});
+
+		// Disabled until fix for UE-44340 (crash creating invalid regex) is merged
+		xIt("will not add an error with an invalid regex pattern", [this]()
+		{
+			AddExpectedError(TEXT("invalid regex }])([{"), EAutomationExpectedErrorFlags::Contains, 0);
+
+			TArray<FAutomationExpectedMessage> Errors;
+			GetExpectedMessages(Errors, ELogVerbosity::Warning);
+
+			TestEqual("Expected Errors Count", Errors.Num(), 0);
+		});
+
+		It("will match both Error and Warning messages", [this]()
+		{
+			AddExpectedError(TEXT("Expected Message"), EAutomationExpectedErrorFlags::Contains, 0);
+			AddError(TEXT("Expected Message"));
+			AddWarning(TEXT("Expected Message"));
+		});
+
+		It("will not fail or warn if encountered", [this]()
+		{
+			AddExpectedError(TEXT("Expected Error"), EAutomationExpectedErrorFlags::Contains, 1);
+			AddExpectedError(TEXT("Expected Warning"), EAutomationExpectedErrorFlags::Contains, 1);
+			AddError(TEXT("Expected Error"));
+			AddWarning(TEXT("Expected Warning"));
+		});
+
+		It("will not match multiple occurrences in the same message when using Contains matcher", [this]()
+		{
+			AddExpectedError(TEXT("Expected"), EAutomationExpectedErrorFlags::Contains, 1);
+			AddError(TEXT("ExpectedExpectedExpected ExpectedExpectedExpected"));
+		});
+
+		It("will match different messages that fit the regex pattern", [this]()
+		{
+			AddExpectedError(TEXT("Response \\(-?\\d+\\)"), EAutomationExpectedErrorFlags::Contains, 4);
+			AddError(TEXT("Response (0)"));
+			AddError(TEXT("Response (1)"));
+			AddError(FString::Printf(TEXT("Response (%d)"), MIN_int64));
+			AddError(FString::Printf(TEXT("Response (%d)"), MAX_uint64));
+		});
+
+	});
+
+}
+
+// Tests for cases where expected errors will fail a test.
+// IMPORTANT: The pass condition for these tests is that they FAIL. To prevent
+// the expected failures from interfering with regular test runs, these tests
+// must be run manually.
+BEGIN_DEFINE_SPEC(FAutomationExpectedErrorFailureTest, "System.Automation.ExpectedError", EAutomationTestFlags::NegativeFilter | EAutomationTestFlags::ApplicationContextMask | EAutomationTestFlags::RequiresUser)
+END_DEFINE_SPEC(FAutomationExpectedErrorFailureTest)
+void FAutomationExpectedErrorFailureTest::Define()
+{
+	Describe("An expected error failure", [this]()
+	{
+		It("will occur if expected a specific number of times and NOT encountered.", [this]()
+		{
+			AddExpectedError(TEXT("Expected Error"), EAutomationExpectedErrorFlags::Exact, 1);
+		});
+
+		It("will occur if expected a specific number of times and is encountered too few times.", [this]()
+		{
+			AddExpectedError(TEXT("Expected Error"), EAutomationExpectedErrorFlags::Exact, 2);
+			AddError(TEXT("Expected Error"));
+		});
+
+		It("will occur if expected a specific number of times and is encountered too many times.", [this]()
+		{
+			AddExpectedError(TEXT("Expected Error"), EAutomationExpectedErrorFlags::Exact, 1);
+			AddError(TEXT("Expected Error"));
+			AddError(TEXT("Expected Error"));
+		});
+
+		It("will occur if expected any number of times and is not encountered.", [this]()
+		{
+			AddExpectedError(TEXT("Expected Error"), EAutomationExpectedErrorFlags::Exact, 0);
+		});
+
+		It("will occur if multiple expected errors are NOT encountered.", [this]()
+		{
+			AddExpectedError(TEXT("Expected Error 1"), EAutomationExpectedErrorFlags::Exact, 1);
+			AddExpectedError(TEXT("Expected Error 2"), EAutomationExpectedErrorFlags::Contains, 1);
+		});
+
+		It("will occur if not all expected errors are encountered.", [this]()
+		{
+			AddExpectedError(TEXT("Expected error 1"), EAutomationExpectedErrorFlags::Exact, 1);
+			AddExpectedError(TEXT("Expected error 2"), EAutomationExpectedErrorFlags::Contains, 1);
+			AddError(TEXT("Expected error 1"));
+		});
+
+		It("will occur if only partial matches are encountered when using Exact matcher", [this]()
+		{
+			AddExpectedError(TEXT("Expected"), EAutomationExpectedErrorFlags::Exact, 1);
+			AddError(TEXT("Expected error"));
+		});
+	});
+}
+
+BEGIN_DEFINE_SPEC(FAutomationExpectedMessageTest, "System.Automation.ExpectedMessage", EAutomationTestFlags::EngineFilter | EAutomationTestFlags::ApplicationContextMask)
+END_DEFINE_SPEC(FAutomationExpectedMessageTest)
+void FAutomationExpectedMessageTest::Define()
+{
+	Describe("A defined expected message with a specified verbosity ", [this]()
+	{
+
+		It("will include error and fatal categories when checking for warning messages", [this]()
+		{
+			AddExpectedMessage(TEXT("suppress this error message"), ELogVerbosity::Error, EAutomationExpectedMessageFlags::Contains, 1);
+			AddExpectedMessage(TEXT("suppress this fatal message"), ELogVerbosity::Fatal, EAutomationExpectedMessageFlags::Contains, 1);
+
+			TArray<FAutomationExpectedMessage> Messages;
+			GetExpectedMessages(Messages, ELogVerbosity::Warning);
+
+			TestEqual(TEXT("Both the error and fatal messages are returned when Warning is specified as the maximum verbosity"), Messages.Num(), 2);
+
+			TestTrue(TEXT("An Error message also counts as an error"), LogCategoryMatchesSeverityInclusive(Messages[0].Verbosity, ELogVerbosity::Warning));
+			TestTrue(TEXT("A Fatal message also counts as an error"), LogCategoryMatchesSeverityInclusive(Messages[1].Verbosity, ELogVerbosity::Warning));
+
+			AddError(TEXT("suppress this error message please"));
+			AddError(TEXT("also suppress this fatal message"));
+		});
+
+
+		// An expected message with a verbosity of "all" behaves as "any"
+		It("will be included when the message is \"all\" but the queried verbosity is NOT \"all\"", [this]()
+		{
+			AddExpectedMessage(TEXT("suppress this message"), ELogVerbosity::All, EAutomationExpectedMessageFlags::Contains, 1);
+
+			// Message with verbosity "all" should always be returned regardless of requested verbosity here
+			TArray<FAutomationExpectedMessage> Messages;
+			GetExpectedMessages(Messages, ELogVerbosity::Warning);
+
+			TestEqual(TEXT("The message with the intention of being any category is returned when something other than \"all\" is specified as the maximum verbosity to query"), Messages.Num(), 1);
+
+			AddInfo(TEXT("suppress this message please"));
+		});
+
+	});
+
+}
diff --git a/Engine/Source/Runtime/Core/Private/Misc/AutomationTest.cpp b/Engine/Source/Runtime/Core/Private/Misc/AutomationTest.cpp
--- a/Engine/Source/Runtime/Core/Private/Misc/AutomationTest.cpp
+++ b/Engine/Source/Runtime/Core/Private/Misc/AutomationTest.cpp
@@ -1,21 +1,24 @@
 // Copyright Epic Games, Inc. All Rights Reserved.
 
 #include "Misc/AutomationTest.h"
-#include "HAL/PlatformStackWalk.h"
+
+#include <inttypes.h>
+
+#include "Algo/Copy.h"
 #include "HAL/FileManager.h"
-#include "Misc/CommandLine.h"
-#include "Misc/Paths.h"
+#include "HAL/PlatformStackWalk.h"
+#include "HAL/ThreadHeartBeat.h"
 #include "Internationalization/Internationalization.h"
+#include "Internationalization/Regex.h"
 #include "Logging/StructuredLog.h"
+#include "Logging/TokenizedMessage.h"
+#include "Misc/App.h"
+#include "Misc/CommandLine.h"
 #include "Misc/ConfigCacheIni.h"
+#include "Misc/OutputDeviceRedirector.h"
+#include "Misc/Paths.h"
 #include "Misc/ScopedSlowTask.h"
-#include "Misc/App.h"
 #include "Modules/ModuleManager.h"
-#include "Misc/OutputDeviceRedirector.h"
-#include "Internationalization/Regex.h"
-#include <inttypes.h>
-#include "Misc/App.h"
-#include "HAL/ThreadHeartBeat.h"
 
 DEFINE_LOG_CATEGORY_STATIC(LogAutomationTest, Warning, All);
 
@@ -212,12 +215,9 @@ void FAutomationTestFramework::FAutomationTestMessageFilter::Serialize(const TCH
 	if (LocalDestinationContext)
 	{
 		FScopeLock Lock(&ActionCS);
-		if ((Verbosity == ELogVerbosity::Warning) || (Verbosity == ELogVerbosity::Error))
+		if (LocalCurTest && LocalCurTest->IsExpectedMessage(FString(V), Verbosity))
 		{
-			if (LocalCurTest && LocalCurTest->IsExpectedError(FString(V)))
-			{
-				Verbosity = ELogVerbosity::Verbose;
-			}
+			Verbosity = ELogVerbosity::Verbose;
 		}
 		LocalDestinationContext->Serialize(V, Verbosity, Category, Time);
 	}
@@ -237,7 +237,7 @@ void FAutomationTestFramework::FAutomationTestMessageFilter::SerializeRecord(con
 		{
 			TStringBuilder<512> Line;
 			Record.FormatMessageTo(Line);
-			if (LocalCurTest->IsExpectedError(FString(Line)))
+			if (LocalCurTest->IsExpectedMessage(FString(Line), ELogVerbosity::Warning))
 			{
 				LocalRecord.SetVerbosity(ELogVerbosity::Verbose);
 			}
@@ -862,11 +862,11 @@ bool FAutomationTestFramework::InternalStopTest(FAutomationTestExecutionInfo& Ou
 	// 1) Did the test itself report success?
 	// 2) Did any errors occur and were logged by the feedback context during execution?++----
 	// 3) Did we meet any errors that were expected with this test
-	bTestSuccessful = bTestSuccessful && !CurrentTest->HasAnyErrors() && CurrentTest->HasMetExpectedErrors();
+	bTestSuccessful = bTestSuccessful && !CurrentTest->HasAnyErrors() && CurrentTest->HasMetExpectedMessages();
 
 	{
 		FScopeLock Lock(&CurrentTest->ActionCS);
-		CurrentTest->ExpectedErrors.Empty();
+		CurrentTest->ExpectedMessages.Empty();
 	}
 
 	// Set the success state of the test based on the above criteria
@@ -1111,7 +1111,7 @@ void FAutomationTestBase::ClearExecutionInfo()
 
 void FAutomationTestBase::AddError(const FString& InError, int32 StackOffset)
 {
-	if( !IsExpectedError(InError))
+	if( !IsExpectedMessage(InError, ELogVerbosity::Warning))
 	{
 		FScopeLock Lock(&ActionCS);
 		ExecutionInfo.AddEvent(FAutomationEvent(EAutomationEventType::Error, InError), StackOffset + 1);
@@ -1128,7 +1128,7 @@ void FAutomationTestBase::AddErrorIfFalse(bool bCondition, const FString& InErro
 
 void FAutomationTestBase::AddErrorS(const FString& InError, const FString& InFilename, int32 InLineNumber)
 {
-	if ( !IsExpectedError(InError))
+	if ( !IsExpectedMessage(InError, ELogVerbosity::Warning))
 	{
 		FScopeLock Lock(&ActionCS);
 		//ExecutionInfo.AddEvent(FAutomationEvent(EAutomationEventType::Error, InError, ExecutionInfo.GetContext(), InFilename, InLineNumber));
@@ -1137,7 +1137,7 @@ void FAutomationTestBase::AddErrorS(const FString& InError, const FString& InFil
 
 void FAutomationTestBase::AddWarningS(const FString& InWarning, const FString& InFilename, int32 InLineNumber)
 {
-	if ( !IsExpectedError(InWarning))
+	if ( !IsExpectedMessage(InWarning, ELogVerbosity::Warning))
 	{
 		FScopeLock Lock(&ActionCS);
 		//ExecutionInfo.AddEvent(FAutomationEvent(EAutomationEventType::Warning, InWarning, ExecutionInfo.GetContext(), InFilename, InLineNumber));
@@ -1146,7 +1146,7 @@ void FAutomationTestBase::AddWarningS(const FString& InWarning, const FString& I
 
 void FAutomationTestBase::AddWarning( const FString& InWarning, int32 StackOffset )
 {
-	if ( !IsExpectedError(InWarning))
+	if ( !IsExpectedMessage(InWarning, ELogVerbosity::Warning))
 	{
 		FScopeLock Lock(&ActionCS);
 		ExecutionInfo.AddEvent(FAutomationEvent(EAutomationEventType::Warning, InWarning), StackOffset + 1);
@@ -1155,8 +1155,11 @@ void FAutomationTestBase::AddWarning( const FString& InWarning, int32 StackOffse
 
 void FAutomationTestBase::AddInfo( const FString& InLogItem, int32 StackOffset, bool bCaptureStack )
 {
-	FScopeLock Lock(&ActionCS);
-	ExecutionInfo.AddEvent(FAutomationEvent(EAutomationEventType::Info, InLogItem), StackOffset + 1, bCaptureStack);
+	if ( !IsExpectedMessage(InLogItem, ELogVerbosity::Display))
+	{
+		FScopeLock Lock(&ActionCS);
+		ExecutionInfo.AddEvent(FAutomationEvent(EAutomationEventType::Info, InLogItem), StackOffset + 1, bCaptureStack);
+	}
 }
 
 void FAutomationTestBase::AddAnalyticsItem(const FString& InAnalyticsItem)
@@ -1196,47 +1199,65 @@ bool FAutomationTestBase::HasAnyErrors() const
 	return ExecutionInfo.GetErrorTotal() > 0;
 }
 
-bool FAutomationTestBase::HasMetExpectedErrors()
+bool FAutomationTestBase::HasMetExpectedMessages(ELogVerbosity::Type VerbosityType)
 {
 	FScopeLock Lock(&ActionCS);
-	bool HasMetAllExpectedErrors = true;
+	bool bHasMetAllExpectedMessages = true;
 
-	for (auto& EError : ExpectedErrors)
+	for (FAutomationExpectedMessage& ExpectedMessage : ExpectedMessages)
 	{
-		if ((EError.ExpectedNumberOfOccurrences > 0) && (EError.ExpectedNumberOfOccurrences != EError.ActualNumberOfOccurrences))
+		if (!LogCategoryMatchesSeverityInclusive(ExpectedMessage.Verbosity, VerbosityType))
+		{
+			continue;
+		}
+
+		// Avoid ambiguity of the messages below when the verbosity is "All"
+		const TCHAR* LogVerbosityString = ExpectedMessage.Verbosity == ELogVerbosity::All ? TEXT("Any") : ToString(ExpectedMessage.Verbosity);
+
+		const bool bExpectsOneOrMore = ExpectedMessage.ExpectedNumberOfOccurrences == 0;
+		if (!bExpectsOneOrMore && (ExpectedMessage.ExpectedNumberOfOccurrences != ExpectedMessage.ActualNumberOfOccurrences))
 		{
-			HasMetAllExpectedErrors = false;
+			bHasMetAllExpectedMessages = false;
 
 			ExecutionInfo.AddEvent(FAutomationEvent(EAutomationEventType::Error,
-				FString::Printf(TEXT("Expected Error or Warning matching '%s' to occur %d times with %s match type, but it was found %d time(s).")
-					, *EError.ErrorPatternString
-					, EError.ExpectedNumberOfOccurrences
-					, EAutomationExpectedErrorFlags::ToString(EError.CompareType)
-					, EError.ActualNumberOfOccurrences)
-				, ExecutionInfo.GetContext()));			
+				FString::Printf(TEXT("Expected ('%s') log message matching '%s' to occur %d times with %s match type, but it was found %d time(s).")
+					, LogVerbosityString
+					, *ExpectedMessage.MessagePatternString
+					, ExpectedMessage.ExpectedNumberOfOccurrences
+					, EAutomationExpectedMessageFlags::ToString(ExpectedMessage.CompareType)
+					, ExpectedMessage.ActualNumberOfOccurrences)
+				, ExecutionInfo.GetContext()));
 		}
-		else if (EError.ExpectedNumberOfOccurrences == 0)
+		else if (bExpectsOneOrMore)
 		{
-			if (EError.ActualNumberOfOccurrences == 0)
+			if (ExpectedMessage.ActualNumberOfOccurrences == 0)
 			{
-				HasMetAllExpectedErrors = false;
+				bHasMetAllExpectedMessages = false;
 
 				ExecutionInfo.AddEvent(FAutomationEvent(EAutomationEventType::Error,
-					FString::Printf(TEXT("Expected suppressed Error or Warning matching '%s' did not occur."), *EError.ErrorPatternString),
-					ExecutionInfo.GetContext()));
+					FString::Printf(TEXT("Expected suppressed ('%s') log message matching '%s' did not occur.")
+						, LogVerbosityString
+						, *ExpectedMessage.MessagePatternString)
+					, ExecutionInfo.GetContext()));
 			}
 			else
 			{
 				ExecutionInfo.AddEvent(FAutomationEvent(EAutomationEventType::Info,
-					FString::Printf(TEXT("Suppressed expected Error or Warning matching '%s' %d times.")
-						, *EError.ErrorPatternString
-						, EError.ActualNumberOfOccurrences)
+					FString::Printf(TEXT("Suppressed expected ('%s') log message matching '%s' %d times.")
+						, LogVerbosityString
+						, *ExpectedMessage.MessagePatternString
+						, ExpectedMessage.ActualNumberOfOccurrences)
 					, ExecutionInfo.GetContext()));
 			}
 		}
 	}
 
-	return HasMetAllExpectedErrors;
+	return bHasMetAllExpectedMessages;
+}
+
+bool FAutomationTestBase::HasMetExpectedErrors()
+{
+	return HasMetExpectedMessages(ELogVerbosity::Warning);
 }
 
 void FAutomationTestBase::InternalSetSuccessState( bool bSuccessful )
@@ -1254,35 +1275,71 @@ void FAutomationTestBase::GetExecutionInfo( FAutomationTestExecutionInfo& OutInf
 	OutInfo = ExecutionInfo;
 }
 
-void FAutomationTestBase::AddExpectedError(FString ExpectedErrorPattern, EAutomationExpectedErrorFlags::MatchType InCompareType, int32 Occurrences)
+void FAutomationTestBase::AddExpectedMessage(
+	FString ExpectedPatternString,
+	ELogVerbosity::Type ExpectedVerbosity,
+	EAutomationExpectedMessageFlags::MatchType CompareType,
+	int32 Occurrences)
 {
 	if (Occurrences >= 0)
 	{
 		FScopeLock Lock(&ActionCS);
-		// If we already have an error matching string in our list, let's not add it again.
-		FAutomationExpectedError* FoundEntry = ExpectedErrors.FindByPredicate(
-			[ExpectedErrorPattern](const FAutomationExpectedError& InItem) 
+		// If we already have a message matching string in our list, let's not add it again.
+		FAutomationExpectedMessage* FoundEntry = ExpectedMessages.FindByPredicate(
+			[ExpectedPatternString](const FAutomationExpectedMessage& InItem) 
 				{
-					return InItem.ErrorPatternString == ExpectedErrorPattern; 
+					return InItem.MessagePatternString == ExpectedPatternString; 
 				}
 		);
 
 		if (FoundEntry)
 		{
-			UE_LOG(LogAutomationTest, Warning, TEXT("Adding expected error matching '%s' failed: cannot add duplicate entries"), *ExpectedErrorPattern)
+			UE_LOG(LogAutomationTest, Warning, TEXT("Adding expected log message matching '%s' failed: cannot add duplicate entries"), *ExpectedPatternString)
 		}
 		else
 		{
-			// ToDo: After UE-44340 is resolved, create FAutomationExpectedError and check that its ErrorPattern is valid before adding
-			ExpectedErrors.Add(FAutomationExpectedError(ExpectedErrorPattern, InCompareType, Occurrences));
+			// ToDo: Check that ExpectedPatternString is valid before adding
+			ExpectedMessages.Add(FAutomationExpectedMessage(ExpectedPatternString, ExpectedVerbosity, CompareType, Occurrences));
 		}
 	}
 	else
 	{
-		UE_LOG(LogAutomationTest, Error, TEXT("Adding expected error matching '%s' failed: number of expected occurrences must be >= 0"), *ExpectedErrorPattern);
+		UE_LOG(LogAutomationTest, Error, TEXT("Adding expected log message matching '%s' failed: number of expected occurrences must be >= 0"), *ExpectedPatternString);
+	}
+}
+
+void FAutomationTestBase::AddExpectedMessage(
+	FString ExpectedPatternString,
+	EAutomationExpectedMessageFlags::MatchType CompareType,
+	int32 Occurrences)
+{
+	AddExpectedMessage(MoveTemp(ExpectedPatternString), ELogVerbosity::All, CompareType, Occurrences);	
+}
+
+void FAutomationTestBase::GetExpectedMessages(
+	TArray<FAutomationExpectedMessage>& OutInfo,
+	ELogVerbosity::Type Verbosity) const
+{
+	if (Verbosity == ELogVerbosity::All)
+	{
+		OutInfo = ExpectedMessages;
+	}
+	else
+	{
+		OutInfo.Reserve(ExpectedMessages.Num());
+		Algo::CopyIf(ExpectedMessages, OutInfo, [Verbosity](const FAutomationExpectedMessage& Message)
+		{
+			return FAutomationTestBase::LogCategoryMatchesSeverityInclusive(Message.Verbosity, Verbosity);
+		});
 	}
 }
 
+void FAutomationTestBase::AddExpectedError(FString ExpectedErrorPattern, EAutomationExpectedErrorFlags::MatchType InCompareType, int32 Occurrences)
+{
+	// Set verbosity to Warning as it's inclusive, and so checks for both Warnings and Errors
+	AddExpectedMessage(MoveTemp(ExpectedErrorPattern), ELogVerbosity::Warning, static_cast<EAutomationExpectedMessageFlags::MatchType>(InCompareType), Occurrences);
+}
+
 uint32 FAutomationTestBase::ExtractAutomationTestFlags(FString InTagNotation)
 {
 	uint32 Result = 0;
@@ -1302,11 +1359,6 @@ uint32 FAutomationTestBase::ExtractAutomationTestFlags(FString InTagNotation)
 	return Result;
 }
 
-void FAutomationTestBase::GetExpectedErrors(TArray<FAutomationExpectedError>& OutInfo) const
-{
-	OutInfo = ExpectedErrors;
-}
-
 void FAutomationTestBase::GenerateTestNames(TArray<FAutomationTestInfo>& TestInfo) const
 {
 	// This can take a while, particularly as spec tests walk the callstack, so suspend the heartbeat watchdog and hitch detector
@@ -1348,6 +1400,14 @@ void FAutomationTestBase::GenerateTestNames(TArray<FAutomationTestInfo>& TestInf
 	}
 }
 
+bool FAutomationTestBase::LogCategoryMatchesSeverityInclusive(
+	ELogVerbosity::Type Actual,
+	ELogVerbosity::Type MaximumVerbosity)
+{
+	// Special case for "all", which should always match
+	return Actual == ELogVerbosity::All || MaximumVerbosity == ELogVerbosity::All || Actual <= MaximumVerbosity;
+}
+
 // --------------------------------------------------------------------------------------
 
 bool FAutomationTestBase::TestEqual(const TCHAR* What, const int32 Actual, const int32 Expected)
@@ -1502,16 +1562,26 @@ bool FAutomationTestBase::TestNull(const TCHAR* What, const void* Pointer)
 	return true;
 }
 
-bool FAutomationTestBase::IsExpectedError(const FString& Error)
+bool FAutomationTestBase::IsExpectedMessage(
+	const FString& Message,
+	const ELogVerbosity::Type& Verbosity)
 {
 	FScopeLock Lock(&ActionCS);
-	for (auto& EError : ExpectedErrors)
+	for (FAutomationExpectedMessage& ExpectedMessage : ExpectedMessages)
 	{
-		FRegexMatcher ErrorMatcher(EError.ErrorPattern, Error);
+		FRegexMatcher MessageMatcher(ExpectedMessage.MessagePattern, Message);
+
+		// Maintains previous behavior: Adjust so that error and fatal messages are tested against when the input verbosity is "Warning"
+		// Similarly, any message above warning should be considered an "info" message
+		const ELogVerbosity::Type AdjustedMessageVerbosity =
+			ExpectedMessage.Verbosity < ELogVerbosity::Warning 
+			? ELogVerbosity::Warning
+			: ELogVerbosity::VeryVerbose;
 
-		if (ErrorMatcher.FindNext())
+		// Compare the incoming message verbosity with the expected verbosity,
+		if (LogCategoryMatchesSeverityInclusive(Verbosity, AdjustedMessageVerbosity) && MessageMatcher.FindNext())
 		{
-			EError.ActualNumberOfOccurrences++;
+			ExpectedMessage.ActualNumberOfOccurrences++;
 			return true;
 		}
 	}
diff --git a/Engine/Source/Runtime/Core/Public/Misc/AutomationTest.h b/Engine/Source/Runtime/Core/Public/Misc/AutomationTest.h
--- a/Engine/Source/Runtime/Core/Public/Misc/AutomationTest.h
+++ b/Engine/Source/Runtime/Core/Public/Misc/AutomationTest.h
@@ -2,6 +2,7 @@
 
 #pragma once
 
+#include "CoreTypes.h"
 #include "Async/Async.h"
 #include "Async/Future.h"
 #include "Containers/Array.h"
@@ -9,7 +10,6 @@
 #include "Containers/Queue.h"
 #include "Containers/Set.h"
 #include "Containers/UnrealString.h"
-#include "CoreTypes.h"
 #include "Delegates/Delegate.h"
 #include "Delegates/DelegateBase.h"
 #include "Delegates/DelegateInstancesImpl.h"
@@ -32,8 +32,8 @@
 #include "Misc/AssertionMacros.h"
 #include "Misc/AutomationEvent.h"
 #include "Misc/Build.h"
-#include "Misc/CString.h"
 #include "Misc/Char.h"
+#include "Misc/CString.h"
 #include "Misc/DateTime.h"
 #include "Misc/FeedbackContext.h"
 #include "Misc/Guid.h"
@@ -140,6 +140,31 @@ struct EAutomationTestFlags
 	}
 };
 
+/** Flags for indicating the matching type to use for an expected message */
+namespace EAutomationExpectedMessageFlags
+{
+	enum MatchType
+	{
+		// When matching expected messages, do so exactly.
+		Exact,
+		// When matching expected messages, just see if the message string is contained in the string to be evaluated.
+		Contains,
+	};
+
+	inline const TCHAR* ToString(EAutomationExpectedMessageFlags::MatchType ThisType)
+	{
+		switch (ThisType)
+		{
+		case Contains:
+			return TEXT("Contains");
+		case Exact:
+			return TEXT("Exact");
+		}
+		return TEXT("Unknown");
+	}
+}
+
+// @todo: replace/refactor to use EAutomationExpectedMessageFlags - touches many, many files
 /** Flags for indicating the matching type to use for an expected error */
 namespace EAutomationExpectedErrorFlags
 {
@@ -593,40 +618,44 @@ public:
 	virtual void Run() = 0;
 };
 
-struct FAutomationExpectedError
+struct FAutomationExpectedMessage
 {
-	// Original regular expression pattern string matching expected error message.
+	// Original regular expression pattern string matching expected log message.
 	// NOTE: using the Exact comparison type wraps the pattern string with ^ and $ tokens,
 	// but the base pattern string is preserved to allow checks for duplicate entries.
-	FString ErrorPatternString;
-	// Regular expression pattern for ErrorPatternString
-	FRegexPattern ErrorPattern;
-	// Type of comparison to perform on error message using ErrorPattern.
-	EAutomationExpectedErrorFlags::MatchType CompareType;
+	FString MessagePatternString;
+	// Regular expression pattern for MessagePatternString
+	FRegexPattern MessagePattern;
+	// Type of comparison to perform on error log using MessagePattern.
+	EAutomationExpectedMessageFlags::MatchType CompareType;
 	/** 
-	 * Number of occurrences expected for error. If set greater than 0, it will cause the test to fail if the
+	 * Number of occurrences expected for message. If set greater than 0, it will cause the test to fail if the
 	 * exact number of occurrences expected is not matched. If set to 0, it will suppress all matching messages. 
 	 */
 	int32 ExpectedNumberOfOccurrences;
 	int32 ActualNumberOfOccurrences;
+	// Log message Verbosity
+	ELogVerbosity::Type Verbosity;
 
 	/**
 	* Constructor
 	*/
-	
-	FAutomationExpectedError(FString& InErrorPattern, EAutomationExpectedErrorFlags::MatchType InCompareType, int32 InExpectedNumberOfOccurrences = 1)
-		: ErrorPatternString(InErrorPattern)
-		, ErrorPattern((InCompareType == EAutomationExpectedErrorFlags::Exact) ? FString::Printf(TEXT("^%s$"), *InErrorPattern) : InErrorPattern)
+	FAutomationExpectedMessage(FString& InMessagePattern, ELogVerbosity::Type InVerbosity, EAutomationExpectedMessageFlags::MatchType InCompareType, int32 InExpectedNumberOfOccurrences = 1)
+		: MessagePatternString(InMessagePattern)
+		, MessagePattern((InCompareType == EAutomationExpectedMessageFlags::Exact) ? FString::Printf(TEXT("^%s$"), *InMessagePattern) : InMessagePattern)
 		, CompareType(InCompareType)
 		, ExpectedNumberOfOccurrences(InExpectedNumberOfOccurrences)
 		, ActualNumberOfOccurrences(0)
+		, Verbosity(InVerbosity)
 	{}
 
-	FAutomationExpectedError(FString& InErrorPattern, int32 InExpectedNumberOfOccurrences)
-		: ErrorPatternString(InErrorPattern)
-		, ErrorPattern(InErrorPattern)
-		, CompareType(EAutomationExpectedErrorFlags::Contains)
+	FAutomationExpectedMessage(FString& InMessagePattern, ELogVerbosity::Type InVerbosity, int32 InExpectedNumberOfOccurrences)
+		: MessagePatternString(InMessagePattern)
+		, MessagePattern(InMessagePattern)
+		, CompareType(EAutomationExpectedMessageFlags::Contains)
 		, ExpectedNumberOfOccurrences(InExpectedNumberOfOccurrences)
+		, ActualNumberOfOccurrences(0)
+		, Verbosity(InVerbosity)
 	{}
 };
 
@@ -1380,6 +1409,13 @@ public:
 	 */
 	bool HasAnyErrors() const;
 
+	/**
+	* Returns whether this test has encountered all expected log messages defined for it
+	* @param VerbosityType Optionally specify to check by log level. Defaults to all.
+	* @return true if this test has encountered all expected messages; false if not
+	*/
+	bool HasMetExpectedMessages(ELogVerbosity::Type VerbosityType = ELogVerbosity::All);
+
 	/**
 	* Returns whether this test has encountered all expected errors defined for it
 	*
@@ -1417,6 +1453,42 @@ public:
 	 */
 	void GenerateTestNames( TArray<FAutomationTestInfo>& TestInfo ) const;
 
+	/**
+	 * Helper function that determines if the given log category matches the expected category, inclusively (so an Error counts as a Warning)
+	*/
+	static bool LogCategoryMatchesSeverityInclusive(ELogVerbosity::Type Actual, ELogVerbosity::Type MaximumVerbosity);
+
+	/**
+	* Adds a regex pattern to an internal list that this test will expect to encounter in logs (of the specified verbosity) during its execution. If an expected pattern
+	* is not encountered, it will cause this test to fail.
+	*
+	* @param ExpectedPatternString - The expected message string. Supports basic regex patterns.
+	* @param ExpectedVerbosity - The expected message verbosity. This is treated as a minimum requirement, so for example the Warning level will intercept Warnings, Errors and Fatal.
+	* @param CompareType - How to match this string with an encountered message, should it match exactly or simply just contain the string.
+	* @param Occurrences - How many times to expect this message string to be seen. If > 0, the message must be seen the exact number of times
+	* specified or the test will fail. If == 0, the message must be seen one or more times (with no upper limit) or the test will fail.
+	*/
+	void AddExpectedMessage(FString ExpectedPatternString, ELogVerbosity::Type ExpectedVerbosity = ELogVerbosity::All, EAutomationExpectedMessageFlags::MatchType CompareType = EAutomationExpectedMessageFlags::Contains, int32 Occurrences = 1);
+	
+	/**
+	* Adds a regex pattern to an internal list that this test will expect to encounter in logs (of all severities) during its execution. If an expected pattern
+	* is not encountered, it will cause this test to fail.
+	*
+	* @param ExpectedPatternString - The expected message string. Supports basic regex patterns.
+	* @param CompareType - How to match this string with an encountered message, should it match exactly or simply just contain the string.
+	* @param Occurrences - How many times to expect this message string to be seen. If > 0, the message must be seen the exact number of times
+	* specified or the test will fail. If == 0, the message must be seen one or more times (with no upper limit) or the test will fail.
+	*/
+	void AddExpectedMessage(FString ExpectedPatternString, EAutomationExpectedMessageFlags::MatchType CompareType = EAutomationExpectedMessageFlags::Contains, int32 Occurrences = 1);
+
+	/**
+	* Populate the provided expected log messages object with the expected messages contained within the test. Not particularly efficient,
+	* but providing direct access to the test's private execution messages list could result in errors.
+	* @param Verbosity - Optionally filter the returned messages by verbosity. This is inclusive, so Warning will return Warnings, Errors, etc.
+	* @param OutInfo - Array of Expected Messages to be populated with the same data contained within this test's expected messages list
+	*/
+	void GetExpectedMessages(TArray<FAutomationExpectedMessage>& OutInfo, ELogVerbosity::Type Verbosity = ELogVerbosity::All) const;
+
 	/**
 	* Adds a regex pattern to an internal list that this test will expect to encounter in error or warning logs during its execution. If an expected pattern
 	* is not encountered, it will cause this test to fail.
@@ -1428,14 +1500,6 @@ public:
 	*/
 	void AddExpectedError(FString ExpectedPatternString, EAutomationExpectedErrorFlags::MatchType CompareType = EAutomationExpectedErrorFlags::Contains, int32 Occurrences = 1);
 
-	/**
-	* Populate the provided expected errors object with the expected errors contained within the test. Not particularly efficient,
-	* but providing direct access to the test's private execution errors list could result in errors.
-	*
-	* @param OutInfo - Array of Expected Errors to be populated with the same data contained within this test's expected errors list
-	*/
-	void GetExpectedErrors(TArray<FAutomationExpectedError>& OutInfo) const;
-
 	/**
 	 * Is this a complex tast - if so it will be a stress test.
 	 *
@@ -1879,12 +1943,12 @@ protected:
 
 private:
 	/**
-	* Returns whether this test has defined any expected errors matching the given message.
-	* If a match is found, the expected error definition increments it actual occurrence count.
+	* Returns whether this test has defined any expected log messages matching the given message.
+	* If a match is found, the expected message definition increments it actual occurrence count.
 	*
-	* @return true if this message matches any of the expected errors
+	* @return true if this message matches any of the expected messages
 	*/
-	bool IsExpectedError(const FString& Error);
+	bool IsExpectedMessage(const FString& Message, const ELogVerbosity::Type& Verbosity = ELogVerbosity::All);
 
 	/**
 	 * Sets whether the test has succeeded or not
@@ -1893,8 +1957,8 @@ private:
 	 */
 	void InternalSetSuccessState(bool bSuccessful);
 
-	/* Errors to be expected while processing this test.*/
-	TArray< FAutomationExpectedError> ExpectedErrors;
+	/* Log messages to be expected while processing this test.*/
+	TArray<FAutomationExpectedMessage> ExpectedMessages;
 
 	/** Critical section lock */
 	FCriticalSection ActionCS;
