Remove artificial 64B max alignment limit from TAlignedBytes / TTypeCompatibleBytes

diff --git a/Engine/Source/Runtime/Core/Public/Templates/TypeCompatibleBytes.h b/Engine/Source/Runtime/Core/Public/Templates/TypeCompatibleBytes.h
--- a/Engine/Source/Runtime/Core/Public/Templates/TypeCompatibleBytes.h
+++ b/Engine/Source/Runtime/Core/Public/Templates/TypeCompatibleBytes.h
@@ -11,62 +11,21 @@
  * Used to declare an untyped array of data with compile-time alignment.
  * It needs to use template specialization as the MS_ALIGN and GCC_ALIGN macros require literal parameters.
  */
-template<int32 Size,uint32 Alignment>
-struct TAlignedBytes; // this intentionally won't compile, we don't support the requested alignment
-
-/** Unaligned storage. */
-template<int32 Size>
-struct TAlignedBytes<Size,1>
+template<int32 Size, uint32 Alignment>
+struct TAlignedBytes
 {
-	uint8 Pad[Size];
+	alignas(Alignment) uint8 Pad[Size];
 };
 
-
-// C++/CLI doesn't support alignment of native types in managed code, so we enforce that the element
-// size is a multiple of the desired alignment
-#ifdef __cplusplus_cli
-	#define IMPLEMENT_ALIGNED_STORAGE(Align) \
-		template<int32 Size>        \
-		struct TAlignedBytes<Size,Align> \
-		{ \
-			uint8 Pad[Size]; \
-			static_assert(Size % Align == 0, "CLR interop types must not be aligned."); \
-		};
-#else
-/** A macro that implements TAlignedBytes for a specific alignment. */
-#define IMPLEMENT_ALIGNED_STORAGE(Align) \
-	template<int32 Size>        \
-	struct TAlignedBytes<Size,Align> \
-	{ \
-		struct MS_ALIGN(Align) TPadding \
-		{ \
-			uint8 Pad[Size]; \
-		} GCC_ALIGN(Align); \
-		TPadding Padding; \
-	};
-#endif
-
-// Implement TAlignedBytes for these alignments.
-IMPLEMENT_ALIGNED_STORAGE(64);
-IMPLEMENT_ALIGNED_STORAGE(32);
-IMPLEMENT_ALIGNED_STORAGE(16);
-IMPLEMENT_ALIGNED_STORAGE(8);
-IMPLEMENT_ALIGNED_STORAGE(4);
-IMPLEMENT_ALIGNED_STORAGE(2);
-
-#undef IMPLEMENT_ALIGNED_STORAGE
-
 /** An untyped array of data with compile-time alignment and size derived from another type. */
 template<typename ElementType>
-struct TTypeCompatibleBytes :
-	public TAlignedBytes<
-		sizeof(ElementType),
-		alignof(ElementType)
-		>
+struct TTypeCompatibleBytes
 {
 	using ElementTypeAlias_NatVisHelper = ElementType;
 	ElementType*		GetTypedPtr()		{ return (ElementType*)this;  }
 	const ElementType*	GetTypedPtr() const	{ return (const ElementType*)this; }
+
+	alignas(ElementType) uint8 Pad[sizeof(ElementType)];
 };
 
 template <
