PR #9759: Fix SelectRandomWeighted algos not working with floating point value weights by initialising the accumulation with the appropriate zero type

diff --git a/Engine/Source/Runtime/Core/Public/Algo/SelectRandomWeighted.h b/Engine/Source/Runtime/Core/Public/Algo/SelectRandomWeighted.h
--- a/Engine/Source/Runtime/Core/Public/Algo/SelectRandomWeighted.h
+++ b/Engine/Source/Runtime/Core/Public/Algo/SelectRandomWeighted.h
@@ -6,6 +6,7 @@
 #include "Algo/Impl/RangePointerType.h"
 #include "Math/RandomStream.h"
 #include "Templates/UnrealTemplate.h"
+#include "Traits/ElementType.h"
 #include "GenericPlatform/GenericPlatformMath.h"
 
 namespace AlgoImpl
@@ -13,14 +14,15 @@ namespace AlgoImpl
 	template <typename RangeType, typename ProjectionType>
 	typename TRangePointerType<typename TRemoveReference<RangeType>::Type>::Type SelectRandomWeightedBy(RangeType&& Range, ProjectionType Proj)
 	{
-		const auto SumOfAllDesires = Algo::Accumulate(Range, 0, [&Proj](auto Acc, auto&& Elem) {
+		using WeightType = std::decay_t<TInvokeResult_T<ProjectionType, TElementType_T<RangeType>>>;
+
+		const auto SumOfAllDesires = Algo::Accumulate(Range, (WeightType)0, [&Proj](auto Acc, auto&& Elem)
+		{
 			const auto Weight = Invoke(Proj, Elem);
 			// Negative values are invalid and should be ignored
-			using WeightType = decltype(Weight);
 			return Acc + (Weight < (WeightType)0 ? (WeightType)0 : Weight);
 		});
-
-		using WeightType = decltype(SumOfAllDesires);
+		
 		auto RandomWeightedAvg = (WeightType)(FMath::FRand() * SumOfAllDesires);
 
 		for (auto&& Elem : Forward<RangeType>(Range))
