Removed XOM check from Android stack walking as new libunwind no longer triggers this issue

diff --git a/Engine/Source/Runtime/Core/Private/Android/AndroidPlatformStackWalk.cpp b/Engine/Source/Runtime/Core/Private/Android/AndroidPlatformStackWalk.cpp
--- a/Engine/Source/Runtime/Core/Private/Android/AndroidPlatformStackWalk.cpp
+++ b/Engine/Source/Runtime/Core/Private/Android/AndroidPlatformStackWalk.cpp
@@ -11,8 +11,6 @@
 #include <dlfcn.h>
 #include <cxxabi.h>
 #include <stdio.h>
-#include <signal.h>
-#include <setjmp.h>
 #include <android/log.h>
 #include "Android/AndroidSignals.h"
 
@@ -34,63 +32,21 @@
 #include "HAL/PlatformProcess.h"
 #include "Misc/ScopeExit.h"
 
-// Some devices with Android 10 have XOM security feature and walking stack might crash
-// We need to verify first that it's safe to read callstack in a call toInitStackWalking
-// Otherwise stack walking will be disabled
-static sigjmp_buf XomJmp;
-static bool DisableStackBacktracing = true;
-static bool StackWalkingInitialized = false;
-static void XomSignalHandler(int Sig)
-{
-	siglongjmp(XomJmp, Sig);
-}
-
 bool FAndroidPlatformStackWalk::InitStackWalking()
 {
-	if (StackWalkingInitialized)
+	static bool StackWalkingInitialized = false;
+	if (!StackWalkingInitialized)
 	{
-		return true;
-	}
-
 #if HAS_LIBUNWIND
-	// Without this stack walk might touch executable memory and ASan will terminate the app on that.
-	// Xom protection presents the same issue and is enabled on android 10 devices
-	// see https://source.android.com/devices/tech/debug/execute-only-memory
-	// Please note that XOM is enabled on some devices even when building with TargetSDK < 29 (Oculus Quest 2)
-	// and not enabled on some other devices at all, like Pixel 4
-	sigset_t SignalSet;
-	sigemptyset(&SignalSet);
-	sigaddset(&SignalSet, SIGSEGV);
-
-	struct sigaction SigAction;
-	struct sigaction OldSigAction;
-	sigset_t OldSignalSet;
-	memset(&SigAction, 0, sizeof(SigAction));
-	SigAction.sa_handler = XomSignalHandler;
-	SigAction.sa_mask = SignalSet;
-
-	sigprocmask(SIG_SETMASK, &SigAction.sa_mask, &OldSignalSet);
-	sigaction(SIGSEGV, &SigAction, &OldSigAction);
-
-	if (sigsetjmp(XomJmp, 1) == 0)
-	{
 		// first call to unw_backtrace will trigger some initial large allocations and if it happens during stack capturing on an exception we might get another out of memory exception
 		const uint32 Depth = 16;
 		void* Stack[Depth];
 		unw_backtrace((void**)Stack, Depth);
-		DisableStackBacktracing = false;
-	}
-	else
-	{
-		//unw_disable_signal_frame_test(1);
-		__android_log_print(ANDROID_LOG_DEBUG, "UE", "XOM has been detected");
-	}
+#endif
 
-	sigaction(SIGSEGV, &OldSigAction, nullptr);
-	sigprocmask(SIG_SETMASK, &OldSignalSet, nullptr);
+		StackWalkingInitialized = true;
+	}
 
-	StackWalkingInitialized = true;
-#endif
 	return true;
 }
 
@@ -213,11 +169,6 @@ extern int32 unwind_backtrace_signal(void* sigcontext, uint64* Backtrace, int32
 
 uint32 FAndroidPlatformStackWalk::CaptureStackBackTrace(uint64* BackTrace, uint32 MaxDepth, void* Context)
 {
-	if (DisableStackBacktracing)
-	{
-		return 0;
-	}
-
 	// Make sure we have place to store the information
 	if (BackTrace == NULL || MaxDepth == 0)
 	{
