Fixed documentation for IMappedFileHandle::MapRegion to reflect that it can actually fail and return nullptr

diff --git a/Engine/Source/Runtime/Core/Public/Async/MappedFileHandle.h b/Engine/Source/Runtime/Core/Public/Async/MappedFileHandle.h
--- a/Engine/Source/Runtime/Core/Public/Async/MappedFileHandle.h
+++ b/Engine/Source/Runtime/Core/Public/Async/MappedFileHandle.h
@@ -110,7 +110,7 @@ public:
 	* @param Offset				Offset into the file to start mapping.
 	* @param BytesToMap			Number of bytes to map. Clamped to the size of the file.
 	* @param bPreloadHint		If true, preload the data. This is only a hint and might be ignored, see IMappedFileRegion::PreloadHint
-	* @return the mapped region interface. This method cannot fail.
+	* @return the mapped region interface. Returns nullptr on failure.
 	**/
 	virtual IMappedFileRegion* MapRegion(int64 Offset = 0, int64 BytesToMap = MAX_int64, bool bPreloadHint = false) = 0;
 
