Remove compilation ambiguity from tests by renaming Engine\Source\Runtime\Core\Public\Tests\TestHarness.h - which is also present in LowLevelTestsRunner module - to TestHarnessAdapter.h

#rnx

diff --git a/Engine/Source/Runtime/Core/Public/Tests/TestHarness.h b/Engine/Source/Runtime/Core/Public/Tests/TestHarness.h
deleted file mode 100644
diff --git a/Engine/Source/Runtime/Core/Public/Tests/TestHarnessAdapter.h b/Engine/Source/Runtime/Core/Public/Tests/TestHarnessAdapter.h
new file mode 100644
--- /dev/null
+++ b/Engine/Source/Runtime/Core/Public/Tests/TestHarnessAdapter.h
@@ -0,0 +1,17 @@
+// Copyright Epic Games, Inc. All Rights Reserved.
+
+#pragma once
+
+/**
+ * Common test harness include file. Use this with tests that can be run as both low level tests or functional tests.
+ * The two types of frameworks are as follows:
+ * - Low level tests using Catch2, generates a monolithic test application, built separate
+ * - UE legacy test automation framework, tests are built into the target executable
+ */
+
+#if WITH_LOW_LEVEL_TESTS
+#include "TestHarness.h"
+#elif defined(WITH_AUTOMATION_TESTS) || (WITH_DEV_AUTOMATION_TESTS || WITH_PERF_AUTOMATION_TESTS)
+#include "Misc/AutomationTest.h"
+#include "Misc/LowLevelTestAdapter.h"
+#endif
diff --git a/Engine/Source/Runtime/Core/Tests/Containers/DequeTest.cpp b/Engine/Source/Runtime/Core/Tests/Containers/DequeTest.cpp
--- a/Engine/Source/Runtime/Core/Tests/Containers/DequeTest.cpp
+++ b/Engine/Source/Runtime/Core/Tests/Containers/DequeTest.cpp
@@ -7,7 +7,7 @@
 #include "TokenTest.h"
 #include "TestUtils.h"
 
-#include "Tests/TestHarness.h"
+#include "Tests/TestHarnessAdapter.h"
 
 namespace Deque
 { 
diff --git a/Engine/Source/Runtime/Core/Tests/Containers/PagedArrayTest.cpp b/Engine/Source/Runtime/Core/Tests/Containers/PagedArrayTest.cpp
--- a/Engine/Source/Runtime/Core/Tests/Containers/PagedArrayTest.cpp
+++ b/Engine/Source/Runtime/Core/Tests/Containers/PagedArrayTest.cpp
@@ -7,7 +7,7 @@
 #include "TestUtils.h"
 #include "TokenTest.h"
 
-#include "Tests/TestHarness.h"
+#include "Tests/TestHarnessAdapter.h"
 
 namespace PagedArray
 {
