use MSVC alignment for all Microsoft platforms when compiling with Clang, not just Windows.

#rnx

diff --git a/Engine/Source/Runtime/Core/Private/Serialization/MemoryImage.cpp b/Engine/Source/Runtime/Core/Private/Serialization/MemoryImage.cpp
--- a/Engine/Source/Runtime/Core/Private/Serialization/MemoryImage.cpp
+++ b/Engine/Source/Runtime/Core/Private/Serialization/MemoryImage.cpp
@@ -79,8 +79,8 @@ void FPlatformTypeLayoutParameters::InitializeForCurrent()
 
 	check(GetRawPointerSize() == sizeof(void*));
 
-	// clang for Windows matches the MSVC ABI
-#if defined(__clang__) && !PLATFORM_WINDOWS
+	// clang for MS matches the MSVC ABI
+#if defined(__clang__) && !PLATFORM_MICROSOFT
 	InitializeForClang();
 #else
 	InitializeForMSVC();
