[Core]
* UnixCriticalSection - Removed UnrealString include to match windows critical section. Since IWYU is using linux path we don't want unix files to include more than windows files

diff --git a/Engine/Source/Runtime/Core/Public/Unix/UnixCriticalSection.h b/Engine/Source/Runtime/Core/Public/Unix/UnixCriticalSection.h
--- a/Engine/Source/Runtime/Core/Public/Unix/UnixCriticalSection.h
+++ b/Engine/Source/Runtime/Core/Public/Unix/UnixCriticalSection.h
@@ -2,12 +2,13 @@
 
 #pragma once
 
-#include "Containers/UnrealString.h"
 #include "CoreTypes.h"
 #include "HAL/PThreadCriticalSection.h"
 #include "HAL/PThreadRWLock.h"
 #include "Misc/Timespan.h"
 
+class FString;
+
 /**
 	* Unix implementation of the FSystemWideCriticalSection. Uses exclusive file locking.
 	**/
