[Core]
* Optional - Moved FNullOpt from Optional.h to OptionalFwd.h to be able to compile NonNullPointer.h without also including Optional.h (which seems to be the intention in that file)

diff --git a/Engine/Source/Runtime/Core/Public/Misc/Optional.h b/Engine/Source/Runtime/Core/Public/Misc/Optional.h
--- a/Engine/Source/Runtime/Core/Public/Misc/Optional.h
+++ b/Engine/Source/Runtime/Core/Public/Misc/Optional.h
@@ -4,14 +4,11 @@
 
 #include "CoreTypes.h"
 #include "Misc/AssertionMacros.h"
+#include "Misc/OptionalFwd.h"
 #include "Templates/TypeCompatibleBytes.h"
 #include "Templates/UnrealTemplate.h"
 #include "Serialization/Archive.h"
 
-struct FNullOpt
-{
-	explicit constexpr FNullOpt(int) {}
-};
 inline constexpr FNullOpt NullOpt{0};
 
 /**
diff --git a/Engine/Source/Runtime/Core/Public/Misc/OptionalFwd.h b/Engine/Source/Runtime/Core/Public/Misc/OptionalFwd.h
--- a/Engine/Source/Runtime/Core/Public/Misc/OptionalFwd.h
+++ b/Engine/Source/Runtime/Core/Public/Misc/OptionalFwd.h
@@ -8,3 +8,8 @@
  */
 template<typename OptionalType>
 struct TOptional;
+
+struct FNullOpt
+{
+	explicit constexpr FNullOpt(int) {}
+};
diff --git a/Engine/Source/Runtime/Core/Public/Templates/NonNullPointer.h b/Engine/Source/Runtime/Core/Public/Templates/NonNullPointer.h
--- a/Engine/Source/Runtime/Core/Public/Templates/NonNullPointer.h
+++ b/Engine/Source/Runtime/Core/Public/Templates/NonNullPointer.h
@@ -2,8 +2,6 @@
 
 #pragma once
 
-// HEADER_UNIT_SKIP - This needs TOptional to compile properly. Don't know why it is not included?
-
 #include "CoreTypes.h"
 #include "Misc/AssertionMacros.h"
 #include "Misc/OptionalFwd.h"
@@ -12,7 +10,6 @@
 
 class FArchive;
 enum class EDefaultConstructNonNullPtr { UnsafeDoNotUse }; // So we can construct TNonNullPtrs
-struct FNullOpt;
 
 namespace UE::Core::Private::NonNullPtr {
 template <typename...>
