Remove Dynamic LocalQueue Management as it was not really used and we have alternatives if needed. This does not impact performance but should reduce power consumption by 1%.

diff --git a/Engine/Source/Runtime/Core/Private/Async/Fundamental/ReserveScheduler.cpp b/Engine/Source/Runtime/Core/Private/Async/Fundamental/ReserveScheduler.cpp
--- a/Engine/Source/Runtime/Core/Private/Async/Fundamental/ReserveScheduler.cpp
+++ b/Engine/Source/Runtime/Core/Private/Async/Fundamental/ReserveScheduler.cpp
@@ -17,7 +17,7 @@ namespace LowLevelTasks
 
 FReserveScheduler FReserveScheduler::Singleton;
 
-TUniquePtr<FThread> FReserveScheduler::CreateWorker(FThread::EForkable IsForkable, FSchedulerTls::FLocalQueueType* WorkerLocalQueue, EThreadPriority Priority)
+TUniquePtr<FThread> FReserveScheduler::CreateWorker(FThread::EForkable IsForkable, EThreadPriority Priority)
 {
 	FYieldedWork* ReserveEvent = new FYieldedWork;
 	ReserveEvents.Emplace(ReserveEvent);
@@ -26,10 +26,10 @@ TUniquePtr<FThread> FReserveScheduler::CreateWorker(FThread::EForkable IsForkabl
 	return MakeUnique<FThread>
 	(
 		*FString::Printf(TEXT("Reserve Worker #%d"), WorkerId),
-		[this, WorkerLocalQueue, ReserveEvent]
+		[this, ReserveEvent]
 		{
 			FSchedulerTls::ActiveScheduler = this;
-			FSchedulerTls::LocalQueue = WorkerLocalQueue;
+			FSchedulerTls::LocalQueue = nullptr;
 
 			while (true)
 			{
@@ -49,7 +49,6 @@ TUniquePtr<FThread> FReserveScheduler::CreateWorker(FThread::EForkable IsForkabl
 			}
 			FSchedulerTls::WorkerType = EWorkerType::None;
 			FSchedulerTls::ActiveScheduler = nullptr;
-			FSchedulerTls::LocalQueue = nullptr;
 		}, 0, Priority, FThreadAffinity{ FPlatformAffinity::GetTaskGraphThreadMask(), 0 }, IsForkable
 	);
 }
@@ -67,14 +66,14 @@ bool FReserveScheduler::DoReserveWorkUntil(FConditional&& Condition)
 	return false;
 }
 
-void FReserveScheduler::StartWorkers(FScheduler& MainScheduler, uint32 NumWorkers, FThread::EForkable IsForkable, EThreadPriority WorkerPriority)
+void FReserveScheduler::StartWorkers(uint32 NumWorkers, FThread::EForkable IsForkable, EThreadPriority WorkerPriority)
 {
 	if (NumWorkers == 0)
 	{
 		NumWorkers = FMath::Min(FPlatformMisc::NumberOfWorkerThreadsToSpawn(), 64);
 	}
 
-	WorkerPriority = GTaskGraphUseDynamicPrioritization ? MainScheduler.GetWorkerPriority() : WorkerPriority;
+	WorkerPriority = GTaskGraphUseDynamicPrioritization ? LowLevelTasks::FScheduler::Get().GetWorkerPriority() : WorkerPriority;
 
 	uint32 OldActiveWorkers = ActiveWorkers.load(std::memory_order_relaxed);
 	if(OldActiveWorkers == 0 && FPlatformProcess::SupportsMultithreading() && ActiveWorkers.compare_exchange_strong(OldActiveWorkers, NumWorkers, std::memory_order_relaxed))
@@ -83,15 +82,12 @@ void FReserveScheduler::StartWorkers(FScheduler& MainScheduler, uint32 NumWorker
 
 		FScopeLock Lock(&WorkerThreadsCS);
 		check(!WorkerThreads.Num());
-		check(!WorkerLocalQueues.Num());
 		check(NextWorkerId == 0);
 
-		WorkerLocalQueues.Reserve(NumWorkers);
 		UE::Trace::ThreadGroupBegin(TEXT("Reserve Workers"));
 		for (uint32 WorkerId = 0; WorkerId < NumWorkers; ++WorkerId)
 		{
-			WorkerLocalQueues.Emplace(MainScheduler.GetQueueRegistry(), ELocalQueueType::EBusyWait, nullptr);
-			WorkerThreads.Add(CreateWorker(IsForkable, &WorkerLocalQueues.Last(), WorkerPriority));
+			WorkerThreads.Add(CreateWorker(IsForkable, WorkerPriority));
 		}
 		UE::Trace::ThreadGroupEnd();
 	}
@@ -119,7 +115,6 @@ void FReserveScheduler::StopWorkers()
 		}
 		NextWorkerId = 0;
 		WorkerThreads.Reset();
-		WorkerLocalQueues.Reset();
 		ReserveEvents.Reset();
 	}
 }
diff --git a/Engine/Source/Runtime/Core/Private/Async/Fundamental/Scheduler.cpp b/Engine/Source/Runtime/Core/Private/Async/Fundamental/Scheduler.cpp
--- a/Engine/Source/Runtime/Core/Private/Async/Fundamental/Scheduler.cpp
+++ b/Engine/Source/Runtime/Core/Private/Async/Fundamental/Scheduler.cpp
@@ -24,26 +24,6 @@ namespace LowLevelTasks
 
 	FScheduler FScheduler::Singleton;
 
-	FScheduler::FLocalQueueInstaller::FLocalQueueInstaller(FScheduler& Scheduler)
-	{
-		RegisteredLocalQueue = FSchedulerTls::LocalQueue == nullptr;
-		if (RegisteredLocalQueue)
-		{
-			bool bPermitBackgroundWork = FTask::PermitBackgroundWork();
-			FSchedulerTls::LocalQueue = FSchedulerTls::FLocalQueueType::AllocateLocalQueue(Scheduler.QueueRegistry, bPermitBackgroundWork ? ELocalQueueType::EBackground : ELocalQueueType::EForeground);
-		}
-	}
-
-	FScheduler::FLocalQueueInstaller::~FLocalQueueInstaller()
-	{
-		if (RegisteredLocalQueue)
-		{
-			bool bPermitBackgroundWork = FTask::PermitBackgroundWork();
-			FSchedulerTls::FLocalQueueType::DeleteLocalQueue(FSchedulerTls::LocalQueue, bPermitBackgroundWork ? ELocalQueueType::EBackground : ELocalQueueType::EForeground);
-			FSchedulerTls::LocalQueue = nullptr;
-		}
-	}
-
 	TUniquePtr<FThread> FScheduler::CreateWorker(bool bPermitBackgroundWork, FThread::EForkable IsForkable, FSleepEvent* ExternalWorkerEvent, FSchedulerTls::FLocalQueueType* ExternalWorkerLocalQueue, EThreadPriority Priority, uint64 InAffinity)
 	{
 		uint32 WorkerId = NextWorkerId++;
@@ -129,23 +109,31 @@ namespace LowLevelTasks
 			check(!WorkerEvents.Num());		
 			check(NextWorkerId == 0);
 
-			WorkerThreads.Reserve(NumForegroundWorkers + NumBackgroundWorkers);
-			WorkerLocalQueues.Reserve(NumForegroundWorkers + NumBackgroundWorkers);
 			WorkerEvents.Reserve(NumForegroundWorkers + NumBackgroundWorkers);
-			UE::Trace::ThreadGroupBegin(TEXT("Foreground Workers"));
+			WorkerLocalQueues.Reserve(NumForegroundWorkers + NumBackgroundWorkers);
 			for (uint32 WorkerId = 0; WorkerId < NumForegroundWorkers; ++WorkerId)
 			{
 				WorkerEvents.Emplace();
-				WorkerLocalQueues.Emplace(QueueRegistry, ELocalQueueType::EForeground, &WorkerEvents.Last());
-				WorkerThreads.Add(CreateWorker(false, IsForkable, &WorkerEvents.Last(), &WorkerLocalQueues.Last(), WorkerPriority, WorkerAffinity));
+				WorkerLocalQueues.Emplace(QueueRegistry, ELocalQueueType::EForeground, &WorkerEvents[WorkerId]);
+			}
+
+			for (uint32 WorkerId = 0; WorkerId < NumBackgroundWorkers; ++WorkerId)
+			{
+				WorkerEvents.Emplace();
+				WorkerLocalQueues.Emplace(QueueRegistry, ELocalQueueType::EBackground, &WorkerEvents[NumForegroundWorkers + WorkerId]);
+			}
+
+			WorkerThreads.Reserve(NumForegroundWorkers + NumBackgroundWorkers);
+			UE::Trace::ThreadGroupBegin(TEXT("Foreground Workers"));
+			for (uint32 WorkerId = 0; WorkerId < NumForegroundWorkers; ++WorkerId)
+			{
+				WorkerThreads.Add(CreateWorker(false, IsForkable, &WorkerEvents[WorkerId], &WorkerLocalQueues[WorkerId], WorkerPriority, WorkerAffinity));
 			}
 			UE::Trace::ThreadGroupEnd();
 			UE::Trace::ThreadGroupBegin(TEXT("Background Workers"));
 			for (uint32 WorkerId = 0; WorkerId < NumBackgroundWorkers; ++WorkerId)
 			{
-				WorkerEvents.Emplace();
-				WorkerLocalQueues.Emplace(QueueRegistry, ELocalQueueType::EBackground, &WorkerEvents.Last());
-				WorkerThreads.Add(CreateWorker(true, IsForkable, &WorkerEvents.Last(), &WorkerLocalQueues.Last(), GTaskGraphUseDynamicPrioritization ? WorkerPriority : BackgroundPriority, BackgroundAffinity));
+				WorkerThreads.Add(CreateWorker(true, IsForkable, &WorkerEvents[NumForegroundWorkers + WorkerId], &WorkerLocalQueues[NumForegroundWorkers + WorkerId], BackgroundPriority, BackgroundAffinity));
 			}
 			UE::Trace::ThreadGroupEnd();
 		}
@@ -214,7 +202,7 @@ namespace LowLevelTasks
 
 			if (DrainGlobalQueue)
 			{
-				for (FTask* Task = QueueRegistry.Dequeue(); Task != nullptr; Task = QueueRegistry.Dequeue())
+				for (FTask* Task = QueueRegistry.DequeueGlobal(); Task != nullptr; Task = QueueRegistry.DequeueGlobal())
 				{
 					while(Task)
 					{
@@ -312,8 +300,8 @@ namespace LowLevelTasks
 		return ~0;
 	}
 
-	template<FTask* (FSchedulerTls::FLocalQueueType::*DequeueFunction)(bool, bool), bool bIsBusyWaiting>
-	bool FScheduler::TryExecuteTaskFrom(FSchedulerTls::FLocalQueueType* Queue, FSchedulerTls::FQueueRegistry::FOutOfWork& OutOfWork, bool bPermitBackgroundWork, bool bDisableThrottleStealing)
+	template<typename QueueType, FTask* (QueueType::*DequeueFunction)(bool, bool), bool bIsBusyWaiting>
+	bool FScheduler::TryExecuteTaskFrom(QueueType* Queue, FSchedulerTls::FQueueRegistry::FOutOfWork& OutOfWork, bool bPermitBackgroundWork, bool bDisableThrottleStealing)
 	{
 		bool AnyExecuted = false;
 
@@ -368,53 +356,40 @@ namespace LowLevelTasks
 		return AnyExecuted;
 	}
 
-	void FScheduler::WorkerMain(FSleepEvent* WorkerEvent, FSchedulerTls::FLocalQueueType* ExternalWorkerLocalQueue, uint32 WaitCycles, bool bPermitBackgroundWork)
+	void FScheduler::WorkerMain(FSleepEvent* WorkerEvent, FSchedulerTls::FLocalQueueType* WorkerLocalQueue, uint32 WaitCycles, bool bPermitBackgroundWork)
 	{
+		checkSlow(FSchedulerTls::LocalQueue == nullptr);
+		checkSlow(WorkerLocalQueue != nullptr);
+		checkSlow(WorkerEvent != nullptr);
+
 		FTaskTagScope WorkerScope(ETaskTag::EWorkerThread);
 		FSchedulerTls::ActiveScheduler = this;
 
 		FMemory::SetupTLSCachesOnCurrentThread();
 		FSchedulerTls::WorkerType = bPermitBackgroundWork ? FSchedulerTls::EWorkerType::Background : FSchedulerTls::EWorkerType::Foreground;
-
-		FSleepEvent LocalWorkerEvent;
-		if (!WorkerEvent)
-		{
-			WorkerEvent = &LocalWorkerEvent;
-		}
-
-		checkSlow(FSchedulerTls::LocalQueue == nullptr);
-		if(ExternalWorkerLocalQueue)
-		{
-			FSchedulerTls::LocalQueue = ExternalWorkerLocalQueue;
-		}
-		else
-		{
-			FSchedulerTls::LocalQueue = FSchedulerTls::FLocalQueueType::AllocateLocalQueue(QueueRegistry, bPermitBackgroundWork ? ELocalQueueType::EBackground : ELocalQueueType::EForeground, WorkerEvent);
-		}
-
-		FSchedulerTls::FLocalQueueType* WorkerLocalQueue = FSchedulerTls::LocalQueue;
+		FSchedulerTls::LocalQueue = WorkerLocalQueue;
 
 		bool bDrowsing = false;
 		uint32 WaitCount = 0;
 		FSchedulerTls::FQueueRegistry::FOutOfWork OutOfWork = QueueRegistry.GetOutOfWorkScope(bPermitBackgroundWork ? ELocalQueueType::EBackground : ELocalQueueType::EForeground);
 		while (true)
 		{
-			while(TryExecuteTaskFrom<&FSchedulerTls::FLocalQueueType::DequeueLocal,  false>(WorkerLocalQueue, OutOfWork, bPermitBackgroundWork, bDrowsing)
-			   || TryExecuteTaskFrom<&FSchedulerTls::FLocalQueueType::DequeueGlobal, false>(WorkerLocalQueue, OutOfWork, bPermitBackgroundWork, bDrowsing))
+			while(TryExecuteTaskFrom<FSchedulerTls::FLocalQueueType, &FSchedulerTls::FLocalQueueType::DequeueLocal,  false>(WorkerLocalQueue, OutOfWork, bPermitBackgroundWork, bDrowsing)
+			   || TryExecuteTaskFrom<FSchedulerTls::FLocalQueueType, &FSchedulerTls::FLocalQueueType::DequeueGlobal, false>(WorkerLocalQueue, OutOfWork, bPermitBackgroundWork, bDrowsing))
 			{		
 				bDrowsing = false;
 				WaitCount = 0;
 			}
 
-			while(TryExecuteTaskFrom<&FSchedulerTls::FLocalQueueType::DequeueLocal, false>(WorkerLocalQueue, OutOfWork, bPermitBackgroundWork, bDrowsing)
-			   || TryExecuteTaskFrom<&FSchedulerTls::FLocalQueueType::DequeueSteal, false>(WorkerLocalQueue, OutOfWork, bPermitBackgroundWork, bDrowsing))
+			while(TryExecuteTaskFrom<FSchedulerTls::FLocalQueueType, &FSchedulerTls::FLocalQueueType::DequeueLocal, false>(WorkerLocalQueue, OutOfWork, bPermitBackgroundWork, bDrowsing)
+			   || TryExecuteTaskFrom<FSchedulerTls::FLocalQueueType, &FSchedulerTls::FLocalQueueType::DequeueSteal, false>(WorkerLocalQueue, OutOfWork, bPermitBackgroundWork, bDrowsing))
 			{
 				bDrowsing = false;
 				WaitCount = 0;
 			}
 
-			while(TryExecuteTaskFrom<&FSchedulerTls::FLocalQueueType::DequeueLocal, false>(WorkerLocalQueue, OutOfWork, bPermitBackgroundWork, bDrowsing)
-				|| TryExecuteTaskFrom<&FSchedulerTls::FLocalQueueType::DequeueAffinity, false>(WorkerLocalQueue, OutOfWork, bPermitBackgroundWork, true))
+			while(TryExecuteTaskFrom<FSchedulerTls::FLocalQueueType, &FSchedulerTls::FLocalQueueType::DequeueLocal, false>(WorkerLocalQueue, OutOfWork, bPermitBackgroundWork, bDrowsing)
+			   || TryExecuteTaskFrom<FSchedulerTls::FLocalQueueType, &FSchedulerTls::FLocalQueueType::DequeueAffinity, false>(WorkerLocalQueue, OutOfWork, bPermitBackgroundWork, true))
 			{
 				bDrowsing = false;
 				WaitCount = 0;
@@ -443,9 +418,7 @@ namespace LowLevelTasks
 
 		while (WakeUpWorker(bPermitBackgroundWork)) {}
 
-		FSchedulerTls::FLocalQueueType::DeleteLocalQueue(WorkerLocalQueue, bPermitBackgroundWork ? ELocalQueueType::EBackground : ELocalQueueType::EForeground, ExternalWorkerLocalQueue != nullptr);
 		FSchedulerTls::LocalQueue = nullptr;
-
 		FSchedulerTls::ActiveScheduler = nullptr;
 		FSchedulerTls::WorkerType = FSchedulerTls::EWorkerType::None;
 		FMemory::ClearAndDisableTLSCachesOnCurrentThread();
@@ -459,7 +432,6 @@ namespace LowLevelTasks
 		++FSchedulerTls::BusyWaitingDepth;
 		ON_SCOPE_EXIT{ --FSchedulerTls::BusyWaitingDepth; };
 
-		checkSlow(FSchedulerTls::LocalQueue != nullptr);
 		check(ActiveWorkers.load(std::memory_order_relaxed));
 		FSchedulerTls::FLocalQueueType* WorkerLocalQueue = FSchedulerTls::LocalQueue;
 
@@ -470,8 +442,8 @@ namespace LowLevelTasks
 		FSchedulerTls::FQueueRegistry::FOutOfWork OutOfWork = QueueRegistry.GetOutOfWorkScope(bPermitBackgroundWork ? ELocalQueueType::EBackground : ELocalQueueType::EForeground);
 		while (true)
 		{
-			while(TryExecuteTaskFrom<&FSchedulerTls::FLocalQueueType::DequeueLocal, true>(WorkerLocalQueue, OutOfWork, bPermitBackgroundWork, true)
-			   || TryExecuteTaskFrom<&FSchedulerTls::FLocalQueueType::DequeueGlobal, true>(WorkerLocalQueue, OutOfWork, bPermitBackgroundWork, true))
+			while((WorkerLocalQueue && TryExecuteTaskFrom<FSchedulerTls::FLocalQueueType, &FSchedulerTls::FLocalQueueType::DequeueLocal, true>(WorkerLocalQueue, OutOfWork, bPermitBackgroundWork, true))
+									|| TryExecuteTaskFrom<FSchedulerTls::FQueueRegistry, &FSchedulerTls::FQueueRegistry::DequeueGlobal, true>(&QueueRegistry, OutOfWork, bPermitBackgroundWork, true))
 			{
 				if (Conditional())
 				{
@@ -480,8 +452,8 @@ namespace LowLevelTasks
 				WaitCount = 0;
 			}
 
-			while(TryExecuteTaskFrom<&FSchedulerTls::FLocalQueueType::DequeueLocal, true>(WorkerLocalQueue, OutOfWork, bPermitBackgroundWork, true)
-			   || TryExecuteTaskFrom<&FSchedulerTls::FLocalQueueType::DequeueSteal, true>(WorkerLocalQueue, OutOfWork, bPermitBackgroundWork, true))
+			while((WorkerLocalQueue && TryExecuteTaskFrom<FSchedulerTls::FLocalQueueType, &FSchedulerTls::FLocalQueueType::DequeueLocal, true>(WorkerLocalQueue, OutOfWork, bPermitBackgroundWork, true))
+									|| TryExecuteTaskFrom<FSchedulerTls::FQueueRegistry, &FSchedulerTls::FQueueRegistry::DequeueSteal, true>(&QueueRegistry, OutOfWork, bPermitBackgroundWork, true))
 			{
 				if (Conditional())
 				{
@@ -490,8 +462,8 @@ namespace LowLevelTasks
 				WaitCount = 0;
 			}
 
-			while(TryExecuteTaskFrom<&FSchedulerTls::FLocalQueueType::DequeueLocal, true>(WorkerLocalQueue, OutOfWork, bPermitBackgroundWork, true)
-				|| TryExecuteTaskFrom<&FSchedulerTls::FLocalQueueType::DequeueAffinity, true>(WorkerLocalQueue, OutOfWork, bPermitBackgroundWork, true))
+			while((WorkerLocalQueue && TryExecuteTaskFrom<FSchedulerTls::FLocalQueueType, &FSchedulerTls::FLocalQueueType::DequeueLocal, true>(WorkerLocalQueue, OutOfWork, bPermitBackgroundWork, true))
+			   || (WorkerLocalQueue && TryExecuteTaskFrom<FSchedulerTls::FLocalQueueType, &FSchedulerTls::FLocalQueueType::DequeueAffinity, true>(WorkerLocalQueue, OutOfWork, bPermitBackgroundWork, true)))
 			{
 				if (Conditional())
 				{
diff --git a/Engine/Source/Runtime/Core/Private/Async/TaskGraph.cpp b/Engine/Source/Runtime/Core/Private/Async/TaskGraph.cpp
--- a/Engine/Source/Runtime/Core/Private/Async/TaskGraph.cpp
+++ b/Engine/Source/Runtime/Core/Private/Async/TaskGraph.cpp
@@ -1943,7 +1943,7 @@ public:
 			if (bReserveWorkersEnabled)
 			{
 				LowLevelTasks::FReserveScheduler::Get().StopWorkers();
-				LowLevelTasks::FReserveScheduler::Get().StartWorkers(LowLevelTasks::FScheduler::Get(), NumForegroundWorkers + NumBackgroundWorkers, FForkProcessHelper::IsForkedMultithreadInstance() ? FThread::Forkable : FThread::NonForkable, FPlatformAffinity::GetTaskBPThreadPriority());
+				LowLevelTasks::FReserveScheduler::Get().StartWorkers(NumForegroundWorkers + NumBackgroundWorkers, FForkProcessHelper::IsForkedMultithreadInstance() ? FThread::Forkable : FThread::NonForkable, FPlatformAffinity::GetTaskBPThreadPriority());
 			}
 		}
 	}
@@ -2326,7 +2326,7 @@ private:
 		if (bEnableReserveWorkers)
 		{
 			bReserveWorkersEnabled = true;
-			LowLevelTasks::FReserveScheduler::Get().StartWorkers(LowLevelTasks::FScheduler::Get(), NumForegroundWorkers + NumBackgroundWorkers, FForkProcessHelper::IsForkedMultithreadInstance() ? FThread::Forkable : FThread::NonForkable, FPlatformAffinity::GetTaskBPThreadPriority());
+			LowLevelTasks::FReserveScheduler::Get().StartWorkers(NumForegroundWorkers + NumBackgroundWorkers, FForkProcessHelper::IsForkedMultithreadInstance() ? FThread::Forkable : FThread::NonForkable, FPlatformAffinity::GetTaskBPThreadPriority());
 		}
 	}
 };
diff --git a/Engine/Source/Runtime/Core/Public/Async/Fundamental/LocalQueue.h b/Engine/Source/Runtime/Core/Public/Async/Fundamental/LocalQueue.h
--- a/Engine/Source/Runtime/Core/Public/Async/Fundamental/LocalQueue.h
+++ b/Engine/Source/Runtime/Core/Public/Async/Fundamental/LocalQueue.h
@@ -124,7 +124,6 @@ enum class ELocalQueueType
 {
 	EBackground,
 	EForeground,
-	EBusyWait,
 };
 
 /********************************************************************************************************************************************
@@ -138,6 +137,14 @@ enum class ELocalQueueType
 template<uint32 NumLocalItems = 1024>
 class TLocalQueueRegistry
 {
+	static uint32 Rand()
+	{
+		uint32 State = FPlatformTime::Cycles();
+		State = State * 747796405u + 2891336453u;
+		State = ((State >> ((State >> 28u) + 4u)) ^ State) * 277803737u;
+		return (State >> 22u) ^ State;
+	}
+
 public:
 	class TLocalQueue;
 
@@ -209,27 +216,6 @@ private:
 	using FOverflowQueueType = FAAArrayQueue<FTask>;
 	using DequeueHazard		 = typename FOverflowQueueType::DequeueHazard;
 
-private:
-	// FLocalQueueCollection is a read only collection of LocalQueues registerd with this Registry
-	struct FLocalQueueCollection
-	{
-		TArray<TLocalQueue*, TInlineAllocator<32>> LocalQueues;
-		TLocalQueue* RemovedQueue = nullptr;
-
-		FLocalQueueCollection() = default;
-
-		~FLocalQueueCollection()
-		{
-			//if the registry also request deletion of a queue (removal case)
-			if (RemovedQueue)
-			{
-				RemovedQueue->~TLocalQueue();
-				FMemory::Free(RemovedQueue);
-			}
-		}
-	};
-	using FStealHazard = THazardPointer<FLocalQueueCollection, true>;
-
 public:
 	class TLocalQueue
 	{
@@ -237,11 +223,9 @@ public:
 		friend class TLocalQueueRegistry;
 
 	public:
-		TLocalQueue(TLocalQueueRegistry& InRegistry, ELocalQueueType QueueType, FSleepEvent* InSleepEvent) : Registry(&InRegistry), SleepEvent(InSleepEvent)
+		TLocalQueue(TLocalQueueRegistry& InRegistry, ELocalQueueType InQueueType, FSleepEvent* InSleepEvent) : Registry(&InRegistry), SleepEvent(InSleepEvent), QueueType(InQueueType)
 		{
-			checkSlow(Registry);
-			StealHazard = FStealHazard(Registry->QueueCollection, Registry->HazardsCollection);
-			AffinityIndex = Registry->AddLocalQueue(StealHazard, this, QueueType);
+			AffinityIndex = Registry->AddLocalQueue(this, QueueType);
 			for (int32 PriorityIndex = 0; PriorityIndex < int32(ETaskPriority::Count); PriorityIndex++)
 			{
 				DequeueHazards[PriorityIndex] = Registry->OverflowQueues[PriorityIndex].getHeadHazard();
@@ -249,33 +233,21 @@ public:
 			AffinityHazard = AffinityQueue.getHeadHazard();
 		}
 
-		static TLocalQueue* AllocateLocalQueue(TLocalQueueRegistry& InRegistry, ELocalQueueType QueueType, FSleepEvent* InSleepEvent = nullptr)
-		{
-			void* Memory = FMemory::Malloc(sizeof(TLocalQueue), 128u);
-			return new (Memory) TLocalQueue(InRegistry, QueueType, InSleepEvent);
-		}
-
-		//delete a queue
-		//WorkeOwned means that the queue will not be automatically deleted when removal succeded.
-		//It is a special case where the memory for the local queues is alloated linearly by the Scheduler for improved stealing performance.
-		static void DeleteLocalQueue(TLocalQueue* Queue, ELocalQueueType QueueType, bool WorkerOwned = false)
+		~TLocalQueue()
 		{
-			TLocalQueueRegistry* Registry = Queue->Registry;
-			Queue->Registry = nullptr;
-
 			for (int32 PriorityIndex = 0; PriorityIndex < int32(ETaskPriority::Count); PriorityIndex++)
 			{
 				while (true)
 				{
 					FTask* Item;
-					if (!Queue->LocalQueues[PriorityIndex].Get(Item))
+					if (!LocalQueues[PriorityIndex].Get(Item))
 					{
 						break;
 					}
 					Registry->OverflowQueues[PriorityIndex].enqueue(Item);
 				}
 			}
-			Registry->DeleteLocalQueue(Queue->StealHazard, Queue, QueueType, WorkerOwned);
+			Registry->DeleteLocalQueue(this, QueueType);
 		}
 
 		//add an item to the local queue and overflow into the global queue if full
@@ -310,7 +282,7 @@ public:
 
 		inline FTask* DequeueGlobal(bool GetBackGroundTasks, bool bDisableThrottleStealing)
 		{
-			if (bDisableThrottleStealing || (Registry->NumActiveWorkers[GetBackGroundTasks].load(std::memory_order_relaxed) >= (2 * Registry->NumWorkersLookingForWork[GetBackGroundTasks].load(std::memory_order_relaxed) - 1)) || ((Random.GetUnsignedInt() % 4) == 0))
+			if (bDisableThrottleStealing || (Registry->NumActiveWorkers[GetBackGroundTasks] >= (2 * Registry->NumWorkersLookingForWork[GetBackGroundTasks].load(std::memory_order_relaxed) - 1)) || ((Rand() % 4) == 0))
 			{
 				int32 MaxPriority = GetBackGroundTasks ? int32(ETaskPriority::Count) : int32(ETaskPriority::ForegroundCount);
 				for (int32 PriorityIndex = 0; PriorityIndex < MaxPriority; PriorityIndex++)
@@ -327,14 +299,14 @@ public:
 
 		inline FTask* DequeueSteal(bool GetBackGroundTasks, bool bDisableThrottleStealing)
 		{
-			if (bDisableThrottleStealing || (Registry->NumActiveWorkers[GetBackGroundTasks].load(std::memory_order_relaxed) >= (2 * Registry->NumWorkersLookingForWork[GetBackGroundTasks].load(std::memory_order_relaxed) - 1)) || ((Random.GetUnsignedInt() % 4) == 0))
+			if (bDisableThrottleStealing || (Registry->NumActiveWorkers[GetBackGroundTasks] >= (2 * Registry->NumWorkersLookingForWork[GetBackGroundTasks].load(std::memory_order_relaxed) - 1)) || ((Rand() % 4) == 0))
 			{
 				if (CachedRandomIndex == InvalidIndex)
 				{
-					CachedRandomIndex = Random.GetUnsignedInt();
+					CachedRandomIndex = Rand();
 				}
 
-				FTask* Result = Registry->StealItem(StealHazard, CachedRandomIndex, CachedPriorityIndex, GetBackGroundTasks);
+				FTask* Result = Registry->StealItem(CachedRandomIndex, CachedPriorityIndex, GetBackGroundTasks);
 				if (Result)
 				{
 					return Result;
@@ -376,91 +348,47 @@ public:
 		DequeueHazard			DequeueHazards[uint32(ETaskPriority::Count)];
 		FOverflowQueueType		AffinityQueue;
 		DequeueHazard			AffinityHazard;
-		FStealHazard			StealHazard;
 		TLocalQueueRegistry*	Registry;
 		FSleepEvent*			SleepEvent;
-		FRandomStream			Random;
 		uint32					CachedRandomIndex = InvalidIndex;
 		uint32					CachedPriorityIndex = 0;
 		uint32					AffinityIndex = ~0;
+		ELocalQueueType			QueueType;
 	};
 
-	TLocalQueueRegistry()
-	{
-		QueueCollection.store(new FLocalQueueCollection(), std::memory_order_relaxed);
-	}
+	TLocalQueueRegistry() = default;
 
 private:
 	// add a queue to the Registry
-	uint32 AddLocalQueue(FStealHazard& Hazard, TLocalQueue* QueueToAdd, ELocalQueueType QueueType)
+	uint32 AddLocalQueue(TLocalQueue* QueueToAdd, ELocalQueueType QueueType)
 	{
-		if (QueueType != ELocalQueueType::EBusyWait)
-		{
-			NumActiveWorkers[QueueType == ELocalQueueType::EBackground].fetch_add(1, std::memory_order_relaxed);
-		}
-
-		FLocalQueueCollection* Copy = new FLocalQueueCollection();
-		while(true)
-		{
-			FLocalQueueCollection* Previous = Hazard.Get();
-			Copy->LocalQueues = Previous->LocalQueues;	
-			Copy->LocalQueues.Add(QueueToAdd);
-			if (!QueueCollection.compare_exchange_strong(Previous, Copy, std::memory_order_acq_rel, std::memory_order_relaxed))
-			{
-				continue;
-			}
-			HazardsCollection.Delete(Previous);
-			Hazard.Retire();
-			return Previous->LocalQueues.Num();
-		}
+		NumActiveWorkers[QueueType == ELocalQueueType::EBackground]++;
+		LocalQueues.Add(QueueToAdd);
+		return LocalQueues.Num() - 1;
 	}
 
 	// remove a queue from the Registry
-	void DeleteLocalQueue(FStealHazard& Hazard, TLocalQueue* QueueToRemove, ELocalQueueType QueueType, bool WorkerOwned)
+	void DeleteLocalQueue(TLocalQueue* QueueToRemove, ELocalQueueType QueueType)
 	{
-		if (QueueType != ELocalQueueType::EBusyWait)
-		{
-			NumActiveWorkers[QueueType == ELocalQueueType::EBackground].fetch_sub(1, std::memory_order_relaxed);
-		}
-
-		FLocalQueueCollection* Copy = new FLocalQueueCollection();
-		while(true)
-		{
-			FLocalQueueCollection* Previous = Hazard.Get();
-			Copy->LocalQueues = Previous->LocalQueues;	
-			verifySlow(Copy->LocalQueues.Remove(QueueToRemove) == 1);
-			if (!QueueCollection.compare_exchange_strong(Previous, Copy, std::memory_order_acq_rel, std::memory_order_relaxed))
-			{
-				continue;
-			}
-			if (!WorkerOwned)
-			{
-				checkSlow(Previous->RemovedQueue == nullptr);
-				Previous->RemovedQueue = QueueToRemove;
-			}		
-			HazardsCollection.Delete(Previous);
-			Hazard.Retire();
-			return;
-		}
+		NumActiveWorkers[QueueType == ELocalQueueType::EBackground]--;
+		verify(LocalQueues.Remove(QueueToRemove) == 1);
 	}
 
 	// StealItem tries to steal an Item from a Registered LocalQueue
-	FTask* StealItem(FStealHazard& Hazard, uint32& CachedRandomIndex, uint32& CachedPriorityIndex, bool GetBackGroundTasks)
+	FTask* StealItem(uint32& CachedRandomIndex, uint32& CachedPriorityIndex, bool GetBackGroundTasks)
 	{
-		FLocalQueueCollection* Queues = Hazard.Get();
-		uint32 NumQueues = Queues->LocalQueues.Num();
+		uint32 NumQueues = LocalQueues.Num();
 		uint32 MaxPriority = GetBackGroundTasks ? int32(ETaskPriority::Count) : int32(ETaskPriority::ForegroundCount);
 		CachedRandomIndex = CachedRandomIndex % NumQueues;
 
 		for(uint32 i = 0; i < NumQueues; i++)
 		{
-			TLocalQueue* LocalQueue = Queues->LocalQueues[CachedRandomIndex];
+			TLocalQueue* LocalQueue = LocalQueues[CachedRandomIndex];
 			for(uint32 PriorityIndex = 0; PriorityIndex < MaxPriority; PriorityIndex++)
 			{	
 				FTask* Item;
 				if (LocalQueue->LocalQueues[CachedPriorityIndex].Steal(Item))
 				{
-					Hazard.Retire();
 					return Item;
 				}
 				CachedPriorityIndex = ++CachedPriorityIndex < MaxPriority ? CachedPriorityIndex : 0;
@@ -469,7 +397,6 @@ private:
 		}
 		CachedPriorityIndex = 0;
 		CachedRandomIndex = TLocalQueue::InvalidIndex;
-		Hazard.Retire();
 		return nullptr;
 	}
 
@@ -489,23 +416,39 @@ public:
 
 	inline bool EnqueueAffinity(FTask* Item, uint32 AffinityIndex)
 	{
-		FStealHazard Hazard(QueueCollection, HazardsCollection);
-		FLocalQueueCollection* Queues = Hazard.Get();
-
-		if (AffinityIndex < uint32(Queues->LocalQueues.Num()) && Queues->LocalQueues[AffinityIndex]->SleepEvent)
+		if (AffinityIndex < uint32(LocalQueues.Num()) && LocalQueues[AffinityIndex]->SleepEvent)
 		{
-			Queues->LocalQueues[AffinityIndex]->EnqueueAffinity(Item);
+			LocalQueues[AffinityIndex]->EnqueueAffinity(Item);
 			return true;
 		}
 		return false;
 	}
 
 	// grab an Item directy from the Global OverflowQueue
-	FTask* Dequeue()
+	FTask* DequeueGlobal(bool GetBackGroundTasks = true, bool bDisableThrottleStealing = true)
+	{
+		if (bDisableThrottleStealing || (NumActiveWorkers[GetBackGroundTasks] >= (2 * NumWorkersLookingForWork[GetBackGroundTasks].load(std::memory_order_relaxed) - 1)) || ((Rand() % 4) == 0))
+		{
+			int32 MaxPriority = GetBackGroundTasks ? int32(ETaskPriority::Count) : int32(ETaskPriority::ForegroundCount);
+			for (int32 PriorityIndex = 0; PriorityIndex < MaxPriority; PriorityIndex++)
+			{
+				FTask* Item = OverflowQueues[PriorityIndex].dequeue();
+				if (Item)
+				{
+					return Item;
+				}
+			}
+		}
+		return nullptr;
+	}
+
+	inline FTask* DequeueSteal(bool GetBackGroundTasks = true, bool bDisableThrottleStealing = true)
 	{
-		for (int32 PriorityIndex = 0; PriorityIndex < int32(ETaskPriority::Count); PriorityIndex++)
+		if (bDisableThrottleStealing || (NumActiveWorkers[GetBackGroundTasks] >= (2 * NumWorkersLookingForWork[GetBackGroundTasks].load(std::memory_order_relaxed) - 1)) || ((Rand() % 4) == 0))
 		{
-			FTask* Result = OverflowQueues[PriorityIndex].dequeue();
+			uint32 CachedRandomIndex = Rand();
+			uint32 CachedPriorityIndex = 0;
+			FTask* Result = StealItem(CachedRandomIndex, CachedPriorityIndex, GetBackGroundTasks);
 			if (Result)
 			{
 				return Result;
@@ -524,19 +467,18 @@ private:
 	{
 		if (bBackgroundTask)
 		{
-			return NumWorkersLookingForWork[true].load(std::memory_order_acquire) * 2 <= NumActiveWorkers[true].load(std::memory_order_acquire);
+			return NumWorkersLookingForWork[true].load(std::memory_order_acquire) * 2 <= NumActiveWorkers[true];
 		}
 		else
 		{
-			return (NumWorkersLookingForWork[false].load(std::memory_order_acquire) + NumWorkersLookingForWork[true].load(std::memory_order_acquire)) * 2 <= (NumActiveWorkers[false].load(std::memory_order_acquire) + NumActiveWorkers[true].load(std::memory_order_acquire));
+			return (NumWorkersLookingForWork[false].load(std::memory_order_acquire) + NumWorkersLookingForWork[true].load(std::memory_order_acquire)) * 2 <= (NumActiveWorkers[false] + NumActiveWorkers[true]);
 		}
 	}
 
-	FOverflowQueueType	  OverflowQueues[uint32(ETaskPriority::Count)];
-	FHazardPointerCollection		  HazardsCollection;
-	std::atomic<FLocalQueueCollection*>	QueueCollection;
-	std::atomic_int NumWorkersLookingForWork[2] = { {0}, {0} };
-	std::atomic_int NumActiveWorkers[2] = { {0}, {0} };
+	FOverflowQueueType		OverflowQueues[uint32(ETaskPriority::Count)];
+	TArray<TLocalQueue*>	LocalQueues;
+	std::atomic_int			NumWorkersLookingForWork[2] = { {0}, {0} };
+	int						NumActiveWorkers[2] = { 0, 0 };
 };
 
 }
\ No newline at end of file
diff --git a/Engine/Source/Runtime/Core/Public/Async/Fundamental/ReserveScheduler.h b/Engine/Source/Runtime/Core/Public/Async/Fundamental/ReserveScheduler.h
--- a/Engine/Source/Runtime/Core/Public/Async/Fundamental/ReserveScheduler.h
+++ b/Engine/Source/Runtime/Core/Public/Async/Fundamental/ReserveScheduler.h
@@ -42,22 +42,19 @@ namespace LowLevelTasks
 		FORCEINLINE_DEBUGGABLE static FReserveScheduler& Get();
 
 		//start number of reserve workers where 0 is the system default
-		CORE_API void StartWorkers(FScheduler& MainScheduler, uint32 ReserveWorkers = 0, FThread::EForkable IsForkable = FThread::NonForkable, EThreadPriority WorkerPriority = EThreadPriority::TPri_Normal);
+		CORE_API void StartWorkers(uint32 ReserveWorkers = 0, FThread::EForkable IsForkable = FThread::NonForkable, EThreadPriority WorkerPriority = EThreadPriority::TPri_Normal);
 		CORE_API void StopWorkers();
 
 		//tries to yield this thread using the YieldEvent and do busywork on a reserve worker.
 		CORE_API bool DoReserveWorkUntil(FConditional&& Condition);
 
 	private: 
-		TUniquePtr<FThread> CreateWorker(FThread::EForkable IsForkable = FThread::NonForkable, FSchedulerTls::FLocalQueueType* WorkerLocalQueue = nullptr, EThreadPriority Priority = EThreadPriority::TPri_Normal);
+		TUniquePtr<FThread> CreateWorker(FThread::EForkable IsForkable = FThread::NonForkable, EThreadPriority Priority = EThreadPriority::TPri_Normal);
 
 	private:
-		template<typename ElementType>
-		using TAlignedArray = TArray<ElementType, TAlignedHeapAllocator<alignof(ElementType)>>;
 		TEventStack<FYieldedWork> 						EventStack;
 		TArray<TUniquePtr<FYieldedWork>>				ReserveEvents;
 		FCriticalSection 								WorkerThreadsCS;
-		TAlignedArray<FSchedulerTls::FLocalQueueType>	WorkerLocalQueues;
 		TArray<TUniquePtr<FThread>>						WorkerThreads;
 		std::atomic_uint								ActiveWorkers { 0 };
 		std::atomic_uint								NextWorkerId { 0 };
diff --git a/Engine/Source/Runtime/Core/Public/Async/Fundamental/Scheduler.h b/Engine/Source/Runtime/Core/Public/Async/Fundamental/Scheduler.h
--- a/Engine/Source/Runtime/Core/Public/Async/Fundamental/Scheduler.h
+++ b/Engine/Source/Runtime/Core/Public/Async/Fundamental/Scheduler.h
@@ -143,15 +143,6 @@ namespace LowLevelTasks
 		// using 16 bytes here because it fits the vtable and one additional pointer
 		using FConditional = TTaskDelegate<bool(), 16>;
 
-	private: 
-		//the FLocalQueueInstaller installs a LocalQueue into the current thread
-		struct FLocalQueueInstaller
-		{
-			bool RegisteredLocalQueue = false;
-			CORE_API FLocalQueueInstaller(FScheduler& Scheduler);
-			CORE_API ~FLocalQueueInstaller();
-		};
-
 	public: // Public Interface of the Scheduler
 		FORCEINLINE_DEBUGGABLE static FScheduler& Get();
 
@@ -182,9 +173,6 @@ namespace LowLevelTasks
 		//number of instantiated workers
 		inline uint32 GetNumWorkers() const;
 
-		//get the Queue registry to register additonal WorkerQueues for this Scheduler
-		inline FSchedulerTls::FQueueRegistry& GetQueueRegistry();
-
 		//get the worker priority set when workers were started
 		inline EThreadPriority GetWorkerPriority() const { return WorkerPriority; }
 
@@ -203,8 +191,8 @@ namespace LowLevelTasks
 		FORCENOINLINE bool TrySleeping(FSleepEvent* WorkerEvent, bool bStopOutOfWorkScope, bool Drowsing, bool bBackgroundWorker);
 		inline bool WakeUpWorker(bool bBackgroundWorker);
 
-		template<FTask* (FSchedulerTls::FLocalQueueType::*DequeueFunction)(bool, bool), bool bIsBusyWaiting>
-		bool TryExecuteTaskFrom(FSchedulerTls::FLocalQueueType* Queue, FSchedulerTls::FQueueRegistry::FOutOfWork& OutOfWork, bool bPermitBackgroundWork, bool bDisableThrottleStealing);
+		template<typename QueueType, FTask* (QueueType::*DequeueFunction)(bool, bool), bool bIsBusyWaiting>
+		bool TryExecuteTaskFrom(QueueType* Queue, FSchedulerTls::FQueueRegistry::FOutOfWork& OutOfWork, bool bPermitBackgroundWork, bool bDisableThrottleStealing);
 
 	private:
 		template<typename ElementType>
@@ -283,7 +271,6 @@ namespace LowLevelTasks
 	{
 		if(!Task.IsCompleted())
 		{
-			FLocalQueueInstaller Installer(*this);
 			FScheduler::BusyWaitInternal([&Task](){ return Task.IsCompleted(); }, ForceAllowBackgroundWork);
 		}
 	}
@@ -296,7 +283,6 @@ namespace LowLevelTasks
 		
 		if(!Cond())
 		{
-			FLocalQueueInstaller Installer(*this);
 			FScheduler::BusyWaitInternal(Forward<Conditional>(Cond), ForceAllowBackgroundWork);
 		}
 	}
@@ -319,7 +305,6 @@ namespace LowLevelTasks
 
 		if (!AllTasksCompleted())
 		{
-			FLocalQueueInstaller Installer(*this);
 			FScheduler::BusyWaitInternal([&AllTasksCompleted](){ return AllTasksCompleted(); }, ForceAllowBackgroundWork);
 		}
 	}
@@ -341,6 +326,7 @@ namespace LowLevelTasks
 			verifySlow(!bStopOutOfWorkScope);
 			bDrowsing = false;
 			WorkerEvent->SleepEvent->Wait(); // State two: ((Running -> Drowsing) -> Sleeping)
+			checkf(WorkerEvent->SleepState.load(std::memory_order_relaxed) == ESleepState::Running, TEXT("Worker was supposed to be running or drowsing: %d"), WorkerEvent->SleepState.load(std::memory_order_relaxed));
 		}
 		else if(WorkerEvent->SleepState.compare_exchange_strong(AffinityState, ESleepState::Drowsing, std::memory_order_release)) //continue drowsing
 		{
@@ -377,11 +363,6 @@ namespace LowLevelTasks
 		return false;
 	}
 
-	inline FSchedulerTls::FQueueRegistry& FScheduler::GetQueueRegistry()
-	{
-		return QueueRegistry;
-	}
-
 	inline FScheduler& FScheduler::Get()
 	{
 		return Singleton;
