Remove deprecated functions

diff --git a/Engine/Source/Runtime/Core/Public/Async/Fundamental/Task.h b/Engine/Source/Runtime/Core/Public/Async/Fundamental/Task.h
--- a/Engine/Source/Runtime/Core/Public/Async/Fundamental/Task.h
+++ b/Engine/Source/Runtime/Core/Public/Async/Fundamental/Task.h
@@ -164,12 +164,6 @@ namespace LowLevelTasks
 			Other.Value = nullptr;
 		}
 
-		UE_DEPRECATED(5.1, "TDeleter::GetValue() has been deprecated. Please update to TDeleter::operator->().")
-		inline Type* GetValue() const
-		{
-			return Value;
-		}
-
 		inline Type* operator->() const
 		{
 			return Value;
@@ -411,14 +405,6 @@ namespace LowLevelTasks
 		template<typename TRunnable>
 		inline void Init(const TCHAR* InDebugName, TRunnable&& InRunnable, ETaskFlags Flags = ETaskFlags::DefaultFlags);
 
-		template<typename TRunnable, typename TContinuation>
-		UE_DEPRECATED(5.1, "FTask::Init() has been deprecated. Please update to using a TDeleter.")
-		inline void Init(const TCHAR* InDebugName, ETaskPriority InPriority, TRunnable&& InRunnable, TContinuation&& InContinuation, ETaskFlags Flags = ETaskFlags::DefaultFlags);
-
-		template<typename TRunnable, typename TContinuation>
-		UE_DEPRECATED(5.1, "FTask::Init() has been deprecated. Please update to using a TDeleter.")
-		inline void Init(const TCHAR* InDebugName, TRunnable&& InRunnable, TContinuation&& InContinuation, ETaskFlags Flags = ETaskFlags::DefaultFlags);
-
 		inline const TCHAR* GetDebugName() const;
 		inline ETaskPriority GetPriority() const;
 		inline bool IsBackgroundTask() const;
@@ -519,30 +505,6 @@ namespace LowLevelTasks
 		PackedData.store(FPackedData(InDebugName, InPriority, ETaskState::Ready, Flags), std::memory_order_release);
 	}
 
-	template<typename TRunnable, typename TContinuation>
-	inline void FTask::Init(const TCHAR* InDebugName, ETaskPriority InPriority, TRunnable&& InRunnable, TContinuation&& InContinuation, ETaskFlags Flags)
-	{
-		checkf(IsCompleted(), TEXT("State: %d"), PackedData.load(std::memory_order_relaxed).GetState());
-		checkSlow(!Runnable.IsSet());
-		Runnable = [LocalRunnable = Forward<TRunnable>(InRunnable), LocalContinuation = Forward<TContinuation>(InContinuation)](const bool NotCanceled)
-		{
-			if (NotCanceled)
-			{
-				LocalRunnable();
-			}
-			LocalContinuation();
-			return nullptr;
-		};
-		InheritParentData(InPriority);
-		PackedData.store(FPackedData(InDebugName, InPriority, ETaskState::Ready, Flags), std::memory_order_release);
-	}
-
-	template<typename TRunnable, typename TContinuation>
-	inline void FTask::Init(const TCHAR* InDebugName, TRunnable&& InRunnable, TContinuation&& InContinuation, ETaskFlags Flags)
-	{
-		Init(InDebugName, ETaskPriority::Default, Forward<TRunnable>(InRunnable), Forward<TContinuation>(InContinuation), Flags);
-	}
-
 	template<typename TRunnable>
 	inline void FTask::Init(const TCHAR* InDebugName, TRunnable&& InRunnable, ETaskFlags Flags)
 	{
