[Core]
* VirtualStackAllocator - Fixed unix non unity/pch compile error (intptr_t was not defined)

diff --git a/Engine/Source/Runtime/Core/Private/Memory/VirtualStackAllocator.cpp b/Engine/Source/Runtime/Core/Private/Memory/VirtualStackAllocator.cpp
--- a/Engine/Source/Runtime/Core/Private/Memory/VirtualStackAllocator.cpp
+++ b/Engine/Source/Runtime/Core/Private/Memory/VirtualStackAllocator.cpp
@@ -4,6 +4,7 @@
 
 #include "HAL/LowLevelMemTracker.h"
 #include "HAL/PlatformMemory.h"
+#include "HAL/PlatformProcess.h"
 #include "Templates/AlignmentTemplates.h"
 
 LLM_DEFINE_TAG(VirtualStackAllocator);
