Don't assert if retagging a thread with the same tag.

diff --git a/Engine/Source/Runtime/Core/Private/HAL/ThreadingBase.cpp b/Engine/Source/Runtime/Core/Private/HAL/ThreadingBase.cpp
--- a/Engine/Source/Runtime/Core/Private/HAL/ThreadingBase.cpp
+++ b/Engine/Source/Runtime/Core/Private/HAL/ThreadingBase.cpp
@@ -121,7 +121,7 @@ FTaskTagScope::FTaskTagScope(bool InTagOnlyIfNone, ETaskTag InTag) : Tag(InTag),
 	{
 		ActiveTaskTag = Tag;
 	}
-	else if (TagOnlyIfNone)
+	else if (TagOnlyIfNone && ActiveTaskTag != Tag)
 	{
 		if (EnumHasAllFlags(Tag, ETaskTag::EParallelRenderingThread))
 		{
