Deprecate half-implemented and unused FPlatformMisc::GetCacheLineSize to get runtime cache line size

We currently have no use case for a runtime-queried cache line size and the function only existed on a subset of our supported platforms

diff --git a/Engine/Source/Runtime/Core/Public/Windows/WindowsPlatformMisc.h b/Engine/Source/Runtime/Core/Public/Windows/WindowsPlatformMisc.h
--- a/Engine/Source/Runtime/Core/Public/Windows/WindowsPlatformMisc.h
+++ b/Engine/Source/Runtime/Core/Public/Windows/WindowsPlatformMisc.h
@@ -233,13 +233,8 @@ struct CORE_API FWindowsPlatformMisc
 	 */
 	static bool GetVSComnTools(int32 Version, FString& OutData);
 
-	/**
-	 * Returns the size of the cache line in bytes.
-	 *
-	 * @return The cache line size.
-	 */
+	UE_DEPRECATED(5.2, "Please use PLATFORM_CACHE_LINE_SIZE instead, runtime query of cache line size not supported")
 	static int32 GetCacheLineSize();
-
 	/**
 	* @return Windows path separator.
 	*/
