Removed TIsSame usage from most of Core (part of a larger change to deprecate TIsSame).

diff --git a/Engine/Source/Runtime/Core/Private/GenericPlatform/GenericWidePlatformString.cpp b/Engine/Source/Runtime/Core/Private/GenericPlatform/GenericWidePlatformString.cpp
--- a/Engine/Source/Runtime/Core/Private/GenericPlatform/GenericWidePlatformString.cpp
+++ b/Engine/Source/Runtime/Core/Private/GenericPlatform/GenericWidePlatformString.cpp
@@ -888,7 +888,7 @@ int32 FGenericWidePlatformString::GetVarArgs( WIDECHAR* Dest, SIZE_T DestSize, c
 			case 'S':
 			{
 				// The %S format represents a string which is the opposite of %s - wide if TCHAR is narrow, or narrow if TCHAR is wide
-				using OtherCharType = TChooseClass<TIsSame<TCHAR, ANSICHAR>::Value, WIDECHAR, ANSICHAR>::Result;
+				using OtherCharType = TChooseClass<std::is_same_v<TCHAR, ANSICHAR>, WIDECHAR, ANSICHAR>::Result;
 
 				ProcessStringArg<OtherCharType>(DestIter, Src, FieldLen, PrecisionLen, ArgPtr);
 				if (!DestIter)
diff --git a/Engine/Source/Runtime/Core/Private/UObject/UnrealNames.cpp b/Engine/Source/Runtime/Core/Private/UObject/UnrealNames.cpp
--- a/Engine/Source/Runtime/Core/Private/UObject/UnrealNames.cpp
+++ b/Engine/Source/Runtime/Core/Private/UObject/UnrealNames.cpp
@@ -138,7 +138,7 @@ static bool operator==(FNameEntryHeader A, FNameEntryHeader B)
 template<typename FromCharType, typename ToCharType>
 ToCharType* ConvertInPlace(FromCharType* Str, uint32 Len)
 {
-	static_assert(TIsSame<FromCharType, ToCharType>::Value, "Unsupported conversion");
+	static_assert(std::is_same_v<FromCharType, ToCharType>, "Unsupported conversion");
 	return Str;
 }
 
diff --git a/Engine/Source/Runtime/Core/Public/Async/Fundamental/Scheduler.h b/Engine/Source/Runtime/Core/Public/Async/Fundamental/Scheduler.h
--- a/Engine/Source/Runtime/Core/Public/Async/Fundamental/Scheduler.h
+++ b/Engine/Source/Runtime/Core/Public/Async/Fundamental/Scheduler.h
@@ -279,7 +279,7 @@ namespace LowLevelTasks
 	inline void FScheduler::BusyWaitUntil(Conditional&& Cond, bool ForceAllowBackgroundWork)
 	{
 		static_assert(TIsInvocable<Conditional>::Value, "Conditional is not invocable");
-		static_assert(TIsSame<decltype(Cond()), bool>::Value, "Conditional must return a boolean");
+		static_assert(std::is_same_v<decltype(Cond()), bool>, "Conditional must return a boolean");
 		
 		if(!Cond())
 		{
diff --git a/Engine/Source/Runtime/Core/Public/Async/Fundamental/Task.h b/Engine/Source/Runtime/Core/Public/Async/Fundamental/Task.h
--- a/Engine/Source/Runtime/Core/Public/Async/Fundamental/Task.h
+++ b/Engine/Source/Runtime/Core/Public/Async/Fundamental/Task.h
@@ -478,7 +478,7 @@ namespace LowLevelTasks
 		checkSlow(!Runnable.IsSet());
 		
 		//if the Runnable returns an FTask* than enable symetric switching
-		if constexpr (TIsSame<FTask*, decltype(UE::Core::Private::IsInvocable::DeclVal<TRunnable>()())>::Value)
+		if constexpr (std::is_same_v<FTask*, decltype(UE::Core::Private::IsInvocable::DeclVal<TRunnable>()())>)
 		{
 			Runnable = [LocalRunnable = Forward<TRunnable>(InRunnable)](const bool bNotCanceled) -> FTask*
 			{
diff --git a/Engine/Source/Runtime/Core/Public/Async/Fundamental/TaskDelegate.h b/Engine/Source/Runtime/Core/Public/Async/Fundamental/TaskDelegate.h
--- a/Engine/Source/Runtime/Core/Public/Async/Fundamental/TaskDelegate.h
+++ b/Engine/Source/Runtime/Core/Public/Async/Fundamental/TaskDelegate.h
@@ -137,7 +137,7 @@ namespace LowLevelTasks
 			inline TTaskDelegateImpl(CallableT&& Callable, void* InlineData)
 			{
 				static_assert(TIsInvocable<TCallableType, ParamTypes...>::Value, "TCallableType is not invocable");
-				static_assert(TIsSame<ReturnType, decltype(Callable(UE::Core::Private::IsInvocable::DeclVal<ParamTypes>()...))>::Value, "TCallableType return type does not match");
+				static_assert(std::is_same_v<ReturnType, decltype(Callable(UE::Core::Private::IsInvocable::DeclVal<ParamTypes>()...))>, "TCallableType return type does not match");
 				static_assert(sizeof(TTaskDelegateImpl<TCallableType, false>) == sizeof(TTaskDelegateBase), "Size must match the Baseclass");
 				new (InlineData) TCallableType(Forward<CallableT>(Callable));
 			}
@@ -207,7 +207,7 @@ namespace LowLevelTasks
 			inline TTaskDelegateImpl(CallableT&& Callable, void* InlineData)
 			{
 				static_assert(TIsInvocable<TCallableType, ParamTypes...>::Value, "TCallableType is not invocable");
-				static_assert(TIsSame<ReturnType, decltype(Callable(UE::Core::Private::IsInvocable::DeclVal<ParamTypes>()...))>::Value, "TCallableType return type does not match");
+				static_assert(std::is_same_v<ReturnType, decltype(Callable(UE::Core::Private::IsInvocable::DeclVal<ParamTypes>()...))>, "TCallableType return type does not match");
 				static_assert(sizeof(TTaskDelegateImpl<TCallableType, true>) == sizeof(TTaskDelegateBase), "Size must match the Baseclass");
 				TCallableType** HeapPtr = reinterpret_cast<TCallableType**>(InlineData);
 				*HeapPtr = reinterpret_cast<TCallableType*>(TConcurrentLinearAllocator<FLowLevelTasksBlockAllocationTag>::template Malloc<alignof(TCallableType)>(sizeof(TCallableType)));
diff --git a/Engine/Source/Runtime/Core/Public/Containers/ArrayView.h b/Engine/Source/Runtime/Core/Public/Containers/ArrayView.h
--- a/Engine/Source/Runtime/Core/Public/Containers/ArrayView.h
+++ b/Engine/Source/Runtime/Core/Public/Containers/ArrayView.h
@@ -76,7 +76,7 @@ namespace ArrayViewPrivate
 
 	public:
 		static constexpr bool Value = 
-			!TIsSame<typename TContainerElementTypeCompatibility<NaturalElementType>::ReinterpretType, NaturalElementType>::Value
+			!std::is_same_v<typename TContainerElementTypeCompatibility<NaturalElementType>::ReinterpretType, NaturalElementType>
 			&&
 			TIsCompatibleElementType<typename TContainerElementTypeCompatibility<NaturalElementType>::ReinterpretType, ElementType>::Value;
 
diff --git a/Engine/Source/Runtime/Core/Public/Containers/ContainerElementTypeCompatibility.h b/Engine/Source/Runtime/Core/Public/Containers/ContainerElementTypeCompatibility.h
--- a/Engine/Source/Runtime/Core/Public/Containers/ContainerElementTypeCompatibility.h
+++ b/Engine/Source/Runtime/Core/Public/Containers/ContainerElementTypeCompatibility.h
@@ -32,7 +32,7 @@ struct TContainerElementTypeCompatibility
 template <typename ElementType>
 struct TIsContainerElementTypeReinterpretable
 {
-	enum { Value = !TIsSame<typename TContainerElementTypeCompatibility<ElementType>::ReinterpretType, ElementType>::Value };
+	enum { Value = !std::is_same_v<typename TContainerElementTypeCompatibility<ElementType>::ReinterpretType, ElementType> };
 };
 
 /**
@@ -42,5 +42,5 @@ struct TIsContainerElementTypeReinterpretable
 template <typename ElementType>
 struct TIsContainerElementTypeCopyable
 {
-	enum { Value = !TIsSame<typename TContainerElementTypeCompatibility<ElementType>::CopyFromOtherType, ElementType>::Value };
+	enum { Value = !std::is_same_v<typename TContainerElementTypeCompatibility<ElementType>::CopyFromOtherType, ElementType> };
 };
diff --git a/Engine/Source/Runtime/Core/Public/Containers/UnrealString.h b/Engine/Source/Runtime/Core/Public/Containers/UnrealString.h
--- a/Engine/Source/Runtime/Core/Public/Containers/UnrealString.h
+++ b/Engine/Source/Runtime/Core/Public/Containers/UnrealString.h
@@ -70,8 +70,9 @@ private:
 	};
 
 	template <typename CharRangeType>
-	struct TIsRangeOfTCHAR : TIsSame<TCHAR, TRangeElementType<CharRangeType>>
+	struct TIsRangeOfTCHAR
 	{
+		enum { Value = std::is_same_v<TCHAR, TRangeElementType<CharRangeType>> };
 	};
 
 	/** Trait testing whether a type is a contiguous range of characters, and not CharType[]. */
diff --git a/Engine/Source/Runtime/Core/Public/HAL/IConsoleManager.h b/Engine/Source/Runtime/Core/Public/HAL/IConsoleManager.h
--- a/Engine/Source/Runtime/Core/Public/HAL/IConsoleManager.h
+++ b/Engine/Source/Runtime/Core/Public/HAL/IConsoleManager.h
@@ -1628,7 +1628,7 @@ public:
 	virtual class TConsoleVariableData<float>*		AsVariableFloat()	override { return AsImpl<float>(); }
 	virtual class TConsoleVariableData<FString>*	AsVariableString()	override { return AsImpl<FString>(); }
 
-	virtual bool		IsVariableInt() const override	{ return TIsSame<int32, T>::Value; }
+	virtual bool		IsVariableInt() const override	{ return std::is_same_v<int32, T>; }
 	virtual int32		GetInt()		const override	{ return GetImpl<int32>(); }
 	virtual float		GetFloat()		const override	{ return GetImpl<float>(); }
 	virtual FString		GetString()		const override	{ return GetImpl<FString>(); }
@@ -1682,27 +1682,27 @@ private:
 	EConsoleVariableFlags Flags = EConsoleVariableFlags::ECVF_Default;
 
 	template<class Y>
-	typename TEnableIf<!TIsSame<T, Y>::Value, Y>::Type GetImpl() const
+	typename TEnableIf<!std::is_same_v<T, Y>, Y>::Type GetImpl() const
 	{
 		check(false);
 		return Y();
 	}
 
 	template<class Y>
-	typename TEnableIf<TIsSame<T, Y>::Value, Y>::Type GetImpl() const
+	typename TEnableIf<std::is_same_v<T, Y>, Y>::Type GetImpl() const
 	{
 		return GetValueOnAnyThread();
 	}
 
 	template<class Y>
-	typename TEnableIf<!TIsSame<T, Y>::Value, TConsoleVariableData<Y>*>::Type AsImpl()
+	typename TEnableIf<!std::is_same_v<T, Y>, TConsoleVariableData<Y>*>::Type AsImpl()
 	{
 		check(false);
 		return nullptr;
 	}
 
 	template<class Y>
-	typename TEnableIf<TIsSame<T, Y>::Value, TConsoleVariableData<T>*>::Type AsImpl()
+	typename TEnableIf<std::is_same_v<T, Y>, TConsoleVariableData<T>*>::Type AsImpl()
 	{
 		return &Value;
 	}
diff --git a/Engine/Source/Runtime/Core/Public/Misc/FrameNumber.h b/Engine/Source/Runtime/Core/Public/Misc/FrameNumber.h
--- a/Engine/Source/Runtime/Core/Public/Misc/FrameNumber.h
+++ b/Engine/Source/Runtime/Core/Public/Misc/FrameNumber.h
@@ -23,7 +23,7 @@ struct FFrameNumber
 	/**
 	 * Implicit construction from a signed integer frame number, whilst disallowing any construction from other types.
 	 */
-	template<typename T, typename U = typename TEnableIf<TIsSame<T, int32>::Value>::Type>
+	template<typename T, typename U = typename TEnableIf<std::is_same_v<T, int32>>::Type>
 	constexpr FFrameNumber(T /*int32*/ InValue)
 		: Value(InValue)
 	{}
diff --git a/Engine/Source/Runtime/Core/Public/Misc/FrameTime.h b/Engine/Source/Runtime/Core/Public/Misc/FrameTime.h
--- a/Engine/Source/Runtime/Core/Public/Misc/FrameTime.h
+++ b/Engine/Source/Runtime/Core/Public/Misc/FrameTime.h
@@ -25,7 +25,7 @@ struct FFrameTime
 	/**
 	 * Implicit construction from a single integer, while disallowing implicit conversion from any other numeric type
 	 */
-	template<typename T, typename = typename TEnableIf<TIsSame<T, int32>::Value>::Type>
+	template<typename T, typename = typename TEnableIf<std::is_same_v<T, int32>>::Type>
 	FFrameTime(T /* int32 */ InFrameNumber);
 
 	/**
diff --git a/Engine/Source/Runtime/Core/Public/Misc/StringBuilder.h b/Engine/Source/Runtime/Core/Public/Misc/StringBuilder.h
--- a/Engine/Source/Runtime/Core/Public/Misc/StringBuilder.h
+++ b/Engine/Source/Runtime/Core/Public/Misc/StringBuilder.h
@@ -68,11 +68,11 @@ public:
 
 	/** Whether the given type can be appended to this builder using the append operator. */
 	template <typename AppendType>
-	using TCanAppend = TIsSame<BuilderType&, decltype(DeclVal<BuilderType&>() << DeclVal<AppendType>())>;
+	inline constexpr static bool TCanAppend_V = std::is_same_v<BuilderType&, decltype(DeclVal<BuilderType&>() << DeclVal<AppendType>())>;
 
 	/** Whether the given range type can have its elements appended to the builder using the append operator. */
 	template <typename RangeType>
-	using TCanAppendRange = TAnd<TIsContiguousContainer<RangeType>, TCanAppend<decltype(*::GetData(DeclVal<RangeType>()))>>;
+	inline constexpr static bool TCanAppendRange_V = TIsContiguousContainer<RangeType>::Value && TCanAppend_V<decltype(*::GetData(DeclVal<RangeType>()))>;
 
 				TStringBuilderBase() = default;
 	CORE_API	~TStringBuilderBase();
@@ -266,7 +266,7 @@ public:
 	 * @return The builder, to allow additional operations to be composed with this one.
 	 */
 	template <typename RangeType, typename DelimiterType,
-		std::enable_if_t<TAnd<TCanAppendRange<RangeType&&>, TCanAppend<DelimiterType&&>>::Value>* = nullptr>
+		std::enable_if_t<TCanAppendRange_V<RangeType&&> && TCanAppend_V<DelimiterType&&>>* = nullptr>
 	inline BuilderType& Join(RangeType&& InRange, DelimiterType&& InDelimiter)
 	{
 		bool bFirst = true;
@@ -299,7 +299,7 @@ public:
 	 * @return The builder, to allow additional operations to be composed with this one.
 	 */
 	template <typename RangeType, typename DelimiterType, typename QuoteType,
-		std::enable_if_t<TAnd<TCanAppendRange<RangeType>, TCanAppend<DelimiterType&&>, TCanAppend<QuoteType&&>>::Value>* = nullptr>
+		std::enable_if_t<TCanAppendRange_V<RangeType> && TCanAppend_V<DelimiterType&&> && TCanAppend_V<QuoteType&&>>* = nullptr>
 	inline BuilderType& JoinQuoted(RangeType&& InRange, DelimiterType&& InDelimiter, QuoteType&& InQuote)
 	{
 		bool bFirst = true;
diff --git a/Engine/Source/Runtime/Core/Public/Misc/TVariantMeta.h b/Engine/Source/Runtime/Core/Public/Misc/TVariantMeta.h
--- a/Engine/Source/Runtime/Core/Public/Misc/TVariantMeta.h
+++ b/Engine/Source/Runtime/Core/Public/Misc/TVariantMeta.h
@@ -264,7 +264,7 @@ namespace Private
 		/** Check if the type at the provided index is the lookup type */
 		static bool IsSame(SIZE_T TypeIndex)
 		{
-			static constexpr bool bIsSameType[] = { TIsSame<Ts, LookupType>::Value... };
+			static constexpr bool bIsSameType[] = { std::is_same_v<Ts, LookupType>... };
 			check(TypeIndex < UE_ARRAY_COUNT(bIsSameType));
 			return bIsSameType[TypeIndex];
 		}
diff --git a/Engine/Source/Runtime/Core/Public/Tasks/TaskPrivate.h b/Engine/Source/Runtime/Core/Public/Tasks/TaskPrivate.h
--- a/Engine/Source/Runtime/Core/Public/Tasks/TaskPrivate.h
+++ b/Engine/Source/Runtime/Core/Public/Tasks/TaskPrivate.h
@@ -702,7 +702,7 @@ namespace UE::Tasks
 
 		// a specialization for tasks that don't return results
 		template<typename TaskBodyType>
-		class TExecutableTaskBase<TaskBodyType, typename TEnableIf<TIsSame<TInvokeResult_T<TaskBodyType>, void>::Value>::Type> : public FTaskBase
+		class TExecutableTaskBase<TaskBodyType, typename TEnableIf<std::is_same_v<TInvokeResult_T<TaskBodyType>, void>>::Type> : public FTaskBase
 		{
 			UE_NONCOPYABLE(TExecutableTaskBase);
 
diff --git a/Engine/Source/Runtime/Core/Public/Templates/Function.h b/Engine/Source/Runtime/Core/Public/Templates/Function.h
--- a/Engine/Source/Runtime/Core/Public/Templates/Function.h
+++ b/Engine/Source/Runtime/Core/Public/Templates/Function.h
@@ -563,9 +563,7 @@ namespace UE::Core::Private::Function
 		template <
 			typename FunctorType,
 			typename = typename TEnableIf<
-				TNot<
-					TIsSame<TFunctionRefBase, typename TDecay<FunctorType>::Type>
-				>::Value
+				!std::is_same_v<TFunctionRefBase, typename TDecay<FunctorType>::Type>
 			>::Type
 		>
 		TFunctionRefBase(FunctorType&& InFunc)
diff --git a/Engine/Source/Runtime/Core/Public/Templates/ResolveTypeAmbiguity.h b/Engine/Source/Runtime/Core/Public/Templates/ResolveTypeAmbiguity.h
--- a/Engine/Source/Runtime/Core/Public/Templates/ResolveTypeAmbiguity.h
+++ b/Engine/Source/Runtime/Core/Public/Templates/ResolveTypeAmbiguity.h
@@ -44,11 +44,10 @@
 
 // Mixing any signed integral types with any other signed integral type results in same type as "X * Y", which is promoted to the type of the result of mixed arithmetic between the types
 #define MIX_SIGNED_INTS_2_ARGS_ACTUAL(Func, OptionalMarkup) \
-	template<typename Arg1, typename Arg2, TEMPLATE_REQUIRES(TAnd< \
-																TNot<TIsSame<Arg1, Arg2>>, \
-																TIsSigned<Arg1>, TIsIntegral<Arg1>, \
-																TIsSigned<Arg2>, TIsIntegral<Arg2> \
-																>::Value)> \
+	template<typename Arg1, typename Arg2, TEMPLATE_REQUIRES(!std::is_same_v<Arg1, Arg2> && \
+																std::is_signed_v<Arg1> && std::is_integral_v<Arg1> && \
+																std::is_signed_v<Arg2> && std::is_integral_v<Arg2> \
+																)> \
 	static OptionalMarkup FORCEINLINE auto Func(Arg1 X, Arg2 Y) -> decltype(X * Y) \
 	{ \
 		using ArgType = decltype(X * Y); \
@@ -61,15 +60,15 @@
 // Mixing any signed integral types with any other signed integral type results in same type as "X * Y", which is promoted to the type of the result of mixed arithmetic between the types
 #define MIX_SIGNED_TYPES_3_ARGS_ACTUAL(Func, OptionalMarkup) \
 	template<typename Arg1, typename Arg2, typename Arg3, TEMPLATE_REQUIRES(TAnd< \
-																				TOr< \
-																					TNot<TIsSame<Arg1, Arg2>>, \
-																					TNot<TIsSame<Arg2, Arg3>>, \
-																					TNot<TIsSame<Arg1, Arg3>> \
-																					>, \
-																				TIsSigned<Arg1>, TIsIntegral<Arg1>, \
-																				TIsSigned<Arg2>, TIsIntegral<Arg2>, \
-																				TIsSigned<Arg3>, TIsIntegral<Arg3> \
-																				>::Value)> \
+																				( \
+																					!std::is_same_v<Arg1, Arg2>> || \
+																					!std::is_same_v<Arg2, Arg3>> || \
+																					!std::is_same_v<Arg1, Arg3>> \
+																				) && \
+																				std::is_signed_v<Arg1> && std::is_integral_v<Arg1> && \
+																				std::is_signed_v<Arg2> && std::is_integral_v<Arg2> && \
+																				std::is_signed_v<Arg3> && std::is_integral_v<Arg3> \
+																			)> \
 	static OptionalMarkup FORCEINLINE auto Func(Arg1 X, Arg2 Y, Arg3 Z) -> decltype(X * Y * Z) \
 	{ \
 		using ArgType = decltype(X * Y * Z); \
@@ -93,10 +92,8 @@
 
 // Mixing float and double types with any other type results in same type as "X * Y", which is promoted to the floating point type with the highest precision of the argument types.
 #define MIX_FLOATS_2_ARGS(Func) \
-	template<typename Arg1, typename Arg2, TEMPLATE_REQUIRES(TAnd< \
-																TOr<TIsFloatingPoint<Arg1>, TIsFloatingPoint<Arg2>>, \
-																TNot<TIsSame<Arg1, Arg2>> \
-																>::Value \
+	template<typename Arg1, typename Arg2, TEMPLATE_REQUIRES( (std::is_floating_point_v<Arg1> || std::is_floating_point_v<Arg2>) && \
+																!std::is_same_v<Arg1, Arg2> \
 															)> \
 	static FORCEINLINE auto Func(Arg1 X, Arg2 Y) -> decltype(X * Y) \
 	{ \
@@ -119,18 +116,17 @@
 
 // Mixing float and double types with any other type results in same type as "X * Y", which is promoted to the floating point type with the highest precision of the argument types.
 #define MIX_FLOATS_3_ARGS(Func) \
-	template<typename Arg1, typename Arg2, typename Arg3, TEMPLATE_REQUIRES(TAnd< \
-																				TOr< \
-																					TIsFloatingPoint<Arg1>, \
-																					TIsFloatingPoint<Arg2>, \
-																					TIsFloatingPoint<Arg3> \
-																					>, \
-																				TOr< \
-																					TNot<TIsSame<Arg1, Arg2>>, \
-																					TNot<TIsSame<Arg2, Arg3>>, \
-																					TNot<TIsSame<Arg1, Arg3>> \
-																					> \
-																				>::Value \
+	template<typename Arg1, typename Arg2, typename Arg3, TEMPLATE_REQUIRES( \
+																				( \
+																					std::is_floating_point_v<Arg1> || \
+																					std::is_floating_point_v<Arg2> || \
+																					std::is_floating_point_v<Arg3> \
+																				) && \
+																				( \
+																					!std::is_same_v<Arg1, Arg2> || \
+																					!std::is_same_v<Arg2, Arg3> || \
+																					!std::is_same_v<Arg1, Arg3> \
+																				) \
 																			)> \
 	static FORCEINLINE auto Func(Arg1 X, Arg2 Y, Arg3 Z) -> decltype(X * Y * Z) \
 	{ \
@@ -170,10 +166,9 @@
 
 // Mixing float and double types with any other type results in same type as "X * Y", which is promoted to the floating point type with the highest precision of the argument types.
 #define MIX_FLOATS_TO_TYPE_2_ARGS(Func, ReturnType) \
-	template<typename Arg1, typename Arg2, TEMPLATE_REQUIRES(TAnd< \
-																TOr<TIsFloatingPoint<Arg1>, TIsFloatingPoint<Arg2>>, \
-																TNot<TIsSame<Arg1, Arg2>> \
-																>::Value \
+	template<typename Arg1, typename Arg2, TEMPLATE_REQUIRES( \
+																(std::is_floating_point_v<Arg1> || std::is_floating_point_v<Arg2>) && \
+																!std::is_same_v<Arg1, Arg2> \
 															)> \
 	static FORCEINLINE ReturnType Func(Arg1 X, Arg2 Y) \
 	{ \
@@ -196,18 +191,17 @@
 
 // Mixing float and double types with any other type results in same type as "X * Y", which is promoted to the floating point type with the highest precision of the argument types.
 #define MIX_FLOATS_TO_TYPE_3_ARGS(Func, ReturnType) \
-	template<typename Arg1, typename Arg2, typename Arg3, TEMPLATE_REQUIRES(TAnd< \
-																				TOr< \
-																					TIsFloatingPoint<Arg1>, \
-																					TIsFloatingPoint<Arg2>, \
-																					TIsFloatingPoint<Arg3> \
-																					>, \
-																				TOr< \
-																					TNot<TIsSame<Arg1, Arg2>>, \
-																					TNot<TIsSame<Arg1, Arg3>>, \
-																					TNot<TIsSame<Arg2, Arg3>> \
-																					> \
-																				>::Value \
+	template<typename Arg1, typename Arg2, typename Arg3, TEMPLATE_REQUIRES( \
+																				( \
+																					std::is_floating_point_v<Arg1> || \
+																					std::is_floating_point_v<Arg2> || \
+																					std::is_floating_point_v<Arg3> \
+																				) && \
+																				( \
+																					!std::is_same_v<Arg1, Arg2> || \
+																					!std::is_same_v<Arg1, Arg3> || \
+																					!std::is_same_v<Arg2, Arg3> \
+																				) \
 																			)> \
 	static FORCEINLINE ReturnType Func(Arg1 X, Arg2 Y, Arg3 Z) \
 	{ \
diff --git a/Engine/Source/Runtime/Core/Tests/Containers/StringViewTest.cpp b/Engine/Source/Runtime/Core/Tests/Containers/StringViewTest.cpp
--- a/Engine/Source/Runtime/Core/Tests/Containers/StringViewTest.cpp
+++ b/Engine/Source/Runtime/Core/Tests/Containers/StringViewTest.cpp
@@ -8,13 +8,13 @@
 #include "Misc/StringBuilder.h"
 #include "TestHarness.h"
 
-static_assert(TIsSame<typename FStringView::ElementType, TCHAR>::Value, "FStringView must use TCHAR.");
-static_assert(TIsSame<typename FAnsiStringView::ElementType, ANSICHAR>::Value, "FAnsiStringView must use ANSICHAR.");
-static_assert(TIsSame<typename FWideStringView::ElementType, WIDECHAR>::Value, "FWideStringView must use WIDECHAR.");
+static_assert(std::is_same_v<typename FStringView::ElementType, TCHAR>, "FStringView must use TCHAR.");
+static_assert(std::is_same_v<typename FAnsiStringView::ElementType, ANSICHAR>, "FAnsiStringView must use ANSICHAR.");
+static_assert(std::is_same_v<typename FWideStringView::ElementType, WIDECHAR>, "FWideStringView must use WIDECHAR.");
 
-static_assert(TIsSame<FStringView, TStringView<TCHAR>>::Value, "FStringView must be the same as TStringView<TCHAR>.");
-static_assert(TIsSame<FAnsiStringView, TStringView<ANSICHAR>>::Value, "FAnsiStringView must be the same as TStringView<ANSICHAR>.");
-static_assert(TIsSame<FWideStringView, TStringView<WIDECHAR>>::Value, "FWideStringView must be the same as TStringView<WIDECHAR>.");
+static_assert(std::is_same_v<FStringView, TStringView<TCHAR>>, "FStringView must be the same as TStringView<TCHAR>.");
+static_assert(std::is_same_v<FAnsiStringView, TStringView<ANSICHAR>>, "FAnsiStringView must be the same as TStringView<ANSICHAR>.");
+static_assert(std::is_same_v<FWideStringView, TStringView<WIDECHAR>>, "FWideStringView must be the same as TStringView<WIDECHAR>.");
 
 static_assert(TIsContiguousContainer<FStringView>::Value, "FStringView must be a contiguous container.");
 static_assert(TIsContiguousContainer<FAnsiStringView>::Value, "FAnsiStringView must be a contiguous container.");
diff --git a/Engine/Source/Runtime/Core/Tests/Misc/StringBuilderTest.cpp b/Engine/Source/Runtime/Core/Tests/Misc/StringBuilderTest.cpp
--- a/Engine/Source/Runtime/Core/Tests/Misc/StringBuilderTest.cpp
+++ b/Engine/Source/Runtime/Core/Tests/Misc/StringBuilderTest.cpp
@@ -14,13 +14,13 @@ TEST_CASE("Core::String::StringBuilder", "[Core][String][Smoke]")
 {
 	SECTION("Static")
 	{
-		STATIC_REQUIRE(TIsSame<typename FStringBuilderBase::ElementType, TCHAR>::Value);
-		STATIC_REQUIRE(TIsSame<typename FAnsiStringBuilderBase::ElementType, ANSICHAR>::Value);
-		STATIC_REQUIRE(TIsSame<typename FWideStringBuilderBase::ElementType, WIDECHAR>::Value);
+		STATIC_REQUIRE(std::is_same_v<typename FStringBuilderBase::ElementType, TCHAR>);
+		STATIC_REQUIRE(std::is_same_v<typename FAnsiStringBuilderBase::ElementType, ANSICHAR>);
+		STATIC_REQUIRE(std::is_same_v<typename FWideStringBuilderBase::ElementType, WIDECHAR>);
 
-		STATIC_REQUIRE(TIsSame<FStringBuilderBase, TStringBuilderBase<TCHAR>>::Value);
-		STATIC_REQUIRE(TIsSame<FAnsiStringBuilderBase, TStringBuilderBase<ANSICHAR>>::Value);
-		STATIC_REQUIRE(TIsSame<FWideStringBuilderBase, TStringBuilderBase<WIDECHAR>>::Value);
+		STATIC_REQUIRE(std::is_same_v<FStringBuilderBase, TStringBuilderBase<TCHAR>>);
+		STATIC_REQUIRE(std::is_same_v<FAnsiStringBuilderBase, TStringBuilderBase<ANSICHAR>>);
+		STATIC_REQUIRE(std::is_same_v<FWideStringBuilderBase, TStringBuilderBase<WIDECHAR>>);
 
 		STATIC_REQUIRE(TIsContiguousContainer<FStringBuilderBase>::Value);
 		STATIC_REQUIRE(TIsContiguousContainer<FAnsiStringBuilderBase>::Value);
