Removed TIsSame from Core/Math (part of larger change to deprecate TIsSame).

diff --git a/Engine/Source/Runtime/Core/Private/Tests/Math/UnrealMathTest.cpp b/Engine/Source/Runtime/Core/Private/Tests/Math/UnrealMathTest.cpp
--- a/Engine/Source/Runtime/Core/Private/Tests/Math/UnrealMathTest.cpp
+++ b/Engine/Source/Runtime/Core/Private/Tests/Math/UnrealMathTest.cpp
@@ -3246,13 +3246,13 @@ bool FVectorRegisterAbstractionTest::RunTest(const FString& Parameters)
 		D = FMath::Min(D1, I1);
 
 		// Test compilation mixing int32/int64 types
-		static_assert(TIsSame< decltype(FMath::Max(Big1, I1)), decltype(Big1) >::Value);
-		static_assert(TIsSame< decltype(FMath::Min(I1, Big1)), decltype(Big1) >::Value);
+		static_assert(std::is_same_v< decltype(FMath::Max(Big1, I1)), decltype(Big1) >);
+		static_assert(std::is_same_v< decltype(FMath::Min(I1, Big1)), decltype(Big1) >);
 		Big3 = FMath::Max(Big1, I1);
 		Big3 = FMath::Min(Big1, I1);
 
-		static_assert(TIsSame< decltype(FMath::Max(I1, Big2)), decltype(Big1) >::Value);
-		static_assert(TIsSame< decltype(FMath::Min(Big2, I1)), decltype(Big1) >::Value);
+		static_assert(std::is_same_v< decltype(FMath::Max(I1, Big2)), decltype(Big1) >);
+		static_assert(std::is_same_v< decltype(FMath::Min(Big2, I1)), decltype(Big1) >);
 		Big3 = FMath::Max(I1, Big2);
 		Big3 = FMath::Min(I1, Big2);
 
diff --git a/Engine/Source/Runtime/Core/Public/GenericPlatform/GenericPlatformMath.h b/Engine/Source/Runtime/Core/Public/GenericPlatform/GenericPlatformMath.h
--- a/Engine/Source/Runtime/Core/Public/GenericPlatform/GenericPlatformMath.h
+++ b/Engine/Source/Runtime/Core/Public/GenericPlatform/GenericPlatformMath.h
@@ -1039,17 +1039,6 @@ struct FGenericPlatformMath
 	static void AutoTest();
 #endif
 
-
-	// Temporary support for ambiguities to simplify UE4 -> UE5 upgrades - falls back to float variant.
-	// A call is considered ambiguous if it passes no float types, any long double, or multiple mismatched float/double types.
-
-	template<typename T1, typename T2>
-	struct TIsSameOrNotFP { static constexpr bool Value = TIsSame<T1, T2>::Value || !TIsFloatingPoint<T2>::Value; };
-	template<typename... Ts>
-	static constexpr bool TIsAmbiguous =	!TOr<TIsFloatingPoint<typename TDecay<Ts>::Type>...>::Value ||
-											TOr<TIsSame<long double, typename TDecay<Ts>::Type>...>::Value ||
-											(TOr<TIsSame<float, typename TDecay<Ts>::Type>...>::Value && !TAnd<TIsSameOrNotFP<float, typename TDecay<Ts>::Type>...>::Value);
-
 private:
 
 	/** Error reporting for Fmod. Not inlined to avoid compilation issues and avoid all the checks and error reporting at all callsites. */
diff --git a/Engine/Source/Runtime/Core/Public/Math/Box.h b/Engine/Source/Runtime/Core/Public/Math/Box.h
--- a/Engine/Source/Runtime/Core/Public/Math/Box.h
+++ b/Engine/Source/Runtime/Core/Public/Math/Box.h
@@ -56,20 +56,8 @@ public:
 	{
 	}
 
-	/**
-	 * Creates and initializes a new box from the specified extents.
-	 *
-	 * @param InMin The box's minimum point.
-	 * @param InMax The box's maximum point.
-	 */
-	TBox( const TVector<T>& InMin, const TVector<T>& InMax )
-		: Min(InMin)
-		, Max(InMax)
-		, IsValid(1)
-	{ }
-
 	// double box taking float params
-	template<typename FArg, TEMPLATE_REQUIRES(TAnd<TIsSame<T, double>, TIsSame<FArg, float>>::Value)>
+	template<typename FArg>
 	TBox(const TVector<FArg>& InMin, const TVector<FArg>& InMax)
 		: Min(InMin)
 		, Max(InMax)
@@ -104,7 +92,7 @@ public:
 	TBox(const TArray<TVector<T>>& Points) : TBox<T>(&Points[0], Points.Num()) {};
 
 	// Conversion from other type.
-	template<typename FArg, TEMPLATE_REQUIRES(!TIsSame<T, FArg>::Value)>
+	template<typename FArg, TEMPLATE_REQUIRES(!std::is_same_v<T, FArg>)>
 	explicit TBox(const TBox<FArg>& From) : TBox<T>(TVector<T>(From.Min), TVector<T>(From.Max)) {}
 
 public:
diff --git a/Engine/Source/Runtime/Core/Public/Math/Box2D.h b/Engine/Source/Runtime/Core/Public/Math/Box2D.h
--- a/Engine/Source/Runtime/Core/Public/Math/Box2D.h
+++ b/Engine/Source/Runtime/Core/Public/Math/Box2D.h
@@ -416,7 +416,7 @@ public:
 	bool SerializeFromMismatchedTag(FName StructTag, FArchive& Ar);
 
 	// Conversion from other type. 
-	template<typename FArg, TEMPLATE_REQUIRES(!TIsSame<T, FArg>::Value)>
+	template<typename FArg, TEMPLATE_REQUIRES(!std::is_same_v<T, FArg>)>
 	explicit TBox2(const TBox2<FArg>& From) : Min((TVector2<T>)From.Min), Max((TVector2<T>)From.Max), bIsValid(From.bIsValid) {}
 };
 
diff --git a/Engine/Source/Runtime/Core/Public/Math/BoxSphereBounds.h b/Engine/Source/Runtime/Core/Public/Math/BoxSphereBounds.h
--- a/Engine/Source/Runtime/Core/Public/Math/BoxSphereBounds.h
+++ b/Engine/Source/Runtime/Core/Public/Math/BoxSphereBounds.h
@@ -124,7 +124,7 @@ public:
 	TBoxSphereBounds( const TVector<T>* Points, uint32 NumPoints );
 
 	// Conversion to other type.
-	template<typename TFrom, typename TExtentFrom, TEMPLATE_REQUIRES(!TAnd<TIsSame<T, TFrom>, TIsSame<TExtent, TExtentFrom>>::Value)>
+	template<typename TFrom, typename TExtentFrom, TEMPLATE_REQUIRES(!(std::is_same_v<T, TFrom> && std::is_same_v<TExtent, TExtentFrom>))>
 	explicit TBoxSphereBounds(const TBoxSphereBounds<TFrom, TExtentFrom>& From) : TBoxSphereBounds<T, TExtent>(TVector<T>(From.Origin), TVector<TExtent>(From.BoxExtent), (TExtent)From.SphereRadius) {}
 
 public:
diff --git a/Engine/Source/Runtime/Core/Public/Math/CapsuleShape.h b/Engine/Source/Runtime/Core/Public/Math/CapsuleShape.h
--- a/Engine/Source/Runtime/Core/Public/Math/CapsuleShape.h
+++ b/Engine/Source/Runtime/Core/Public/Math/CapsuleShape.h
@@ -53,7 +53,7 @@ public:
 	{ }
 
 	// Conversion to other type.
-	template<typename FArg, TEMPLATE_REQUIRES(!TIsSame<T, FArg>::Value)>
+	template<typename FArg, TEMPLATE_REQUIRES(!std::is_same_v<T, FArg>)>
 	explicit TCapsuleShape(const TCapsuleShape<FArg>& From) : TCapsuleShape<T>(TVector<T>(From.Center), (T)From.Radius, TVector<T>(From.Orientation), (T)From.Length) {}
 };
 
diff --git a/Engine/Source/Runtime/Core/Public/Math/ClipProjectionMatrix.h b/Engine/Source/Runtime/Core/Public/Math/ClipProjectionMatrix.h
--- a/Engine/Source/Runtime/Core/Public/Math/ClipProjectionMatrix.h
+++ b/Engine/Source/Runtime/Core/Public/Math/ClipProjectionMatrix.h
@@ -29,7 +29,7 @@ public:
 	TClipProjectionMatrix( const TMatrix<T>& SrcProjMat, const TPlane<T>& Plane );
 
 	// Conversion to other type.
-	template<typename FArg, TEMPLATE_REQUIRES(!TIsSame<T, FArg>::Value)>
+	template<typename FArg, TEMPLATE_REQUIRES(!std::is_same_v<T, FArg>)>
 	explicit TClipProjectionMatrix(const TClipProjectionMatrix<FArg>& From) : TMatrix<T>(From) {}
 	
 private:
diff --git a/Engine/Source/Runtime/Core/Public/Math/DualQuat.h b/Engine/Source/Runtime/Core/Public/Math/DualQuat.h
--- a/Engine/Source/Runtime/Core/Public/Math/DualQuat.h
+++ b/Engine/Source/Runtime/Core/Public/Math/DualQuat.h
@@ -38,7 +38,7 @@ public:
 	}
 
 	// Conversion to other type.
-	template<typename FArg, TEMPLATE_REQUIRES(!TIsSame<T, FArg>::Value)>
+	template<typename FArg, TEMPLATE_REQUIRES(!std::is_same_v<T, FArg>)>
 	explicit TDualQuat(const TDualQuat<FArg>& From) : TDualQuat<T>(TQuat<T>(From.R), TQuat<T>(From.D)) {}
 
 	/** Dual quat addition */
diff --git a/Engine/Source/Runtime/Core/Public/Math/InverseRotationMatrix.h b/Engine/Source/Runtime/Core/Public/Math/InverseRotationMatrix.h
--- a/Engine/Source/Runtime/Core/Public/Math/InverseRotationMatrix.h
+++ b/Engine/Source/Runtime/Core/Public/Math/InverseRotationMatrix.h
@@ -24,7 +24,7 @@ public:
 	TInverseRotationMatrix(const TRotator<T>& Rot);
 	
 	// Conversion to other type.
-    template<typename FArg, TEMPLATE_REQUIRES(!TIsSame<T, FArg>::Value)>
+    template<typename FArg, TEMPLATE_REQUIRES(!std::is_same_v<T, FArg>)>
     explicit TInverseRotationMatrix(const TInverseRotationMatrix<FArg>& From) : TMatrix<T>(From) {}	
 };
 
diff --git a/Engine/Source/Runtime/Core/Public/Math/Matrix.h b/Engine/Source/Runtime/Core/Public/Math/Matrix.h
--- a/Engine/Source/Runtime/Core/Public/Math/Matrix.h
+++ b/Engine/Source/Runtime/Core/Public/Math/Matrix.h
@@ -427,7 +427,7 @@ public:
 	}
 
 	// Conversion to other type.
-	template<typename FArg, TEMPLATE_REQUIRES(!TIsSame<T, FArg>::Value)>
+	template<typename FArg, TEMPLATE_REQUIRES(!std::is_same_v<T, FArg>)>
 	explicit TMatrix(const TMatrix<FArg>& From)
 	{
 		// TODO: SIMD this?
@@ -516,7 +516,7 @@ struct TBasisVectorMatrix : public TMatrix<T>
 	TBasisVectorMatrix(const TVector<T>& XAxis,const TVector<T>& YAxis,const TVector<T>& ZAxis,const TVector<T>& Origin);
 
 	// Conversion to other type.
-	template<typename FArg, TEMPLATE_REQUIRES(!TIsSame<T, FArg>::Value)>
+	template<typename FArg, TEMPLATE_REQUIRES(!std::is_same_v<T, FArg>)>
 	explicit TBasisVectorMatrix(const TBasisVectorMatrix<FArg>& From) : TMatrix<T>(From) {}
 };
 
@@ -535,7 +535,7 @@ struct TLookFromMatrix : public TMatrix<T>
 	TLookFromMatrix(const TVector<T>& EyePosition, const TVector<T>& LookDirection, const TVector<T>& UpVector);
 
 	// Conversion to other type.
-	template<typename FArg, TEMPLATE_REQUIRES(!TIsSame<T, FArg>::Value)>
+	template<typename FArg, TEMPLATE_REQUIRES(!std::is_same_v<T, FArg>)>
 	explicit TLookFromMatrix(const TLookFromMatrix<FArg>& From) : TMatrix<T>(From) {}
 };
 
@@ -554,7 +554,7 @@ struct TLookAtMatrix : public TLookFromMatrix<T>
 	TLookAtMatrix(const TVector<T>& EyePosition, const TVector<T>& LookAtPosition, const TVector<T>& UpVector);
 
 	// Conversion to other type.
-	template<typename FArg, TEMPLATE_REQUIRES(!TIsSame<T, FArg>::Value)>
+	template<typename FArg, TEMPLATE_REQUIRES(!std::is_same_v<T, FArg>)>
 	explicit TLookAtMatrix(const TLookAtMatrix<FArg>& From) : TLookFromMatrix<T>(From) {}
 };
 
diff --git a/Engine/Source/Runtime/Core/Public/Math/MirrorMatrix.h b/Engine/Source/Runtime/Core/Public/Math/MirrorMatrix.h
--- a/Engine/Source/Runtime/Core/Public/Math/MirrorMatrix.h
+++ b/Engine/Source/Runtime/Core/Public/Math/MirrorMatrix.h
@@ -26,7 +26,7 @@ public:
 	TMirrorMatrix( const TPlane<T>& Plane );
 
 	// Conversion to other type.
-	template<typename FArg, TEMPLATE_REQUIRES(!TIsSame<T, FArg>::Value)>
+	template<typename FArg, TEMPLATE_REQUIRES(!std::is_same_v<T, FArg>)>
 	explicit TMirrorMatrix(const TMirrorMatrix<FArg>& From) : TMatrix<T>(From) {}
 };
 
diff --git a/Engine/Source/Runtime/Core/Public/Math/OrthoMatrix.h b/Engine/Source/Runtime/Core/Public/Math/OrthoMatrix.h
--- a/Engine/Source/Runtime/Core/Public/Math/OrthoMatrix.h
+++ b/Engine/Source/Runtime/Core/Public/Math/OrthoMatrix.h
@@ -27,7 +27,7 @@ public:
 	TOrthoMatrix(T Width,T Height,T ZScale,T ZOffset);
 	
 	// Conversion to other type.
-	template<typename FArg, TEMPLATE_REQUIRES(!TIsSame<T, FArg>::Value)>
+	template<typename FArg, TEMPLATE_REQUIRES(!std::is_same_v<T, FArg>)>
 	explicit TOrthoMatrix(const TOrthoMatrix<FArg>& From) : TMatrix<T>(From) {}
 };
 
@@ -40,7 +40,7 @@ public:
 	TReversedZOrthoMatrix(T Left, T Right, T Bottom, T Top, T ZScale, T ZOffset);
 
 	// Conversion to other type.
-	template<typename FArg, TEMPLATE_REQUIRES(!TIsSame<T, FArg>::Value)>
+	template<typename FArg, TEMPLATE_REQUIRES(!std::is_same_v<T, FArg>)>
 	explicit TReversedZOrthoMatrix(const TReversedZOrthoMatrix<FArg>& From) : TMatrix<T>(From) {}
 };
 
diff --git a/Engine/Source/Runtime/Core/Public/Math/PerspectiveMatrix.h b/Engine/Source/Runtime/Core/Public/Math/PerspectiveMatrix.h
--- a/Engine/Source/Runtime/Core/Public/Math/PerspectiveMatrix.h
+++ b/Engine/Source/Runtime/Core/Public/Math/PerspectiveMatrix.h
@@ -55,7 +55,7 @@ public:
 	TPerspectiveMatrix(T HalfFOV, T Width, T Height, T MinZ);
 
 	// Conversion to other type.
-	template<typename FArg, TEMPLATE_REQUIRES(!TIsSame<T, FArg>::Value)>
+	template<typename FArg, TEMPLATE_REQUIRES(!std::is_same_v<T, FArg>)>
 	explicit TPerspectiveMatrix(const TPerspectiveMatrix<FArg>& From) : TMatrix<T>(From) {}
 };
 
@@ -69,7 +69,7 @@ public:
 	TReversedZPerspectiveMatrix(T HalfFOV, T Width, T Height, T MinZ);
 	
 	// Conversion to other type.
-	template<typename FArg, TEMPLATE_REQUIRES(!TIsSame<T, FArg>::Value)>
+	template<typename FArg, TEMPLATE_REQUIRES(!std::is_same_v<T, FArg>)>
 	explicit TReversedZPerspectiveMatrix(const TReversedZPerspectiveMatrix<FArg>& From) : TMatrix<T>(From) {}
 };
 
diff --git a/Engine/Source/Runtime/Core/Public/Math/Plane.h b/Engine/Source/Runtime/Core/Public/Math/Plane.h
--- a/Engine/Source/Runtime/Core/Public/Math/Plane.h
+++ b/Engine/Source/Runtime/Core/Public/Math/Plane.h
@@ -336,7 +336,7 @@ public:
 	}
 
 	// Conversion to other type.
-	template<typename FArg, TEMPLATE_REQUIRES(!TIsSame<T, FArg>::Value)>
+	template<typename FArg, TEMPLATE_REQUIRES(!std::is_same_v<T, FArg>)>
 	explicit TPlane(const TPlane<FArg>& From) : TPlane<T>((T)From.X, (T)From.Y, (T)From.Z, (T)From.W) {}
 };
 
diff --git a/Engine/Source/Runtime/Core/Public/Math/Quat.h b/Engine/Source/Runtime/Core/Public/Math/Quat.h
--- a/Engine/Source/Runtime/Core/Public/Math/Quat.h
+++ b/Engine/Source/Runtime/Core/Public/Math/Quat.h
@@ -688,7 +688,7 @@ public:
 	bool SerializeFromMismatchedTag(FName StructTag, FArchive& Ar);
 
 	// Conversion to other type.
-	template<typename FArg, TEMPLATE_REQUIRES(!TIsSame<T, FArg>::Value)>
+	template<typename FArg, TEMPLATE_REQUIRES(!std::is_same_v<T, FArg>)>
 	explicit TQuat(const TQuat<FArg>& From) : TQuat<T>((T)From.X, (T)From.Y, (T)From.Z, (T)From.W) {}
 };
 
diff --git a/Engine/Source/Runtime/Core/Public/Math/QuatRotationTranslationMatrix.h b/Engine/Source/Runtime/Core/Public/Math/QuatRotationTranslationMatrix.h
--- a/Engine/Source/Runtime/Core/Public/Math/QuatRotationTranslationMatrix.h
+++ b/Engine/Source/Runtime/Core/Public/Math/QuatRotationTranslationMatrix.h
@@ -27,7 +27,7 @@ public:
 	TQuatRotationTranslationMatrix(const TQuat<T>& Q, const TVector<T>& Origin);
 
 	// Conversion to other type.
-	template<typename FArg, TEMPLATE_REQUIRES(!TIsSame<T, FArg>::Value)>
+	template<typename FArg, TEMPLATE_REQUIRES(!std::is_same_v<T, FArg>)>
 	explicit TQuatRotationTranslationMatrix(const TQuatRotationTranslationMatrix<FArg>& From) : TMatrix<T>(From) {}
 
 	/** Matrix factory. Return an FMatrix so we don't have type conversion issues in expressions. */
@@ -56,7 +56,7 @@ public:
 	}
 
 	// Conversion to other type.
-	template<typename FArg, TEMPLATE_REQUIRES(!TIsSame<T, FArg>::Value)>
+	template<typename FArg, TEMPLATE_REQUIRES(!std::is_same_v<T, FArg>)>
 	explicit TQuatRotationMatrix(const TQuatRotationMatrix<FArg>& From) : TQuatRotationTranslationMatrix<T>(From) {}
 	
 	/** Matrix factory. Return an FMatrix so we don't have type conversion issues in expressions. */
diff --git a/Engine/Source/Runtime/Core/Public/Math/Ray.h b/Engine/Source/Runtime/Core/Public/Math/Ray.h
--- a/Engine/Source/Runtime/Core/Public/Math/Ray.h
+++ b/Engine/Source/Runtime/Core/Public/Math/Ray.h
@@ -203,7 +203,7 @@ public:
 	bool SerializeFromMismatchedTag(FName StructTag, FArchive& Ar);
 
 	// Conversion to other type.
-	template<typename FArg, TEMPLATE_REQUIRES(!TIsSame<T, FArg>::Value)>
+	template<typename FArg, TEMPLATE_REQUIRES(!std::is_same_v<T, FArg>)>
 	explicit TRay(const TRay<FArg>& From) : TRay<T>(TVector<T>(From.Origin), TVector<T>(From.Direction), true) {}
 };
 
diff --git a/Engine/Source/Runtime/Core/Public/Math/RotationAboutPointMatrix.h b/Engine/Source/Runtime/Core/Public/Math/RotationAboutPointMatrix.h
--- a/Engine/Source/Runtime/Core/Public/Math/RotationAboutPointMatrix.h
+++ b/Engine/Source/Runtime/Core/Public/Math/RotationAboutPointMatrix.h
@@ -28,7 +28,7 @@ public:
 	TRotationAboutPointMatrix(const TRotator<T>& Rot, const TVector<T>& Origin);
 
 	// Conversion to other type.
-	template<typename FArg, TEMPLATE_REQUIRES(!TIsSame<T, FArg>::Value)>
+	template<typename FArg, TEMPLATE_REQUIRES(!std::is_same_v<T, FArg>)>
 	explicit TRotationAboutPointMatrix(const TRotationAboutPointMatrix<FArg>& From) : TRotationTranslationMatrix<T>(From) {}	
 	
 	/** Matrix factory. Return an TMatrix<T> so we don't have type conversion issues in expressions. */
diff --git a/Engine/Source/Runtime/Core/Public/Math/RotationTranslationMatrix.h b/Engine/Source/Runtime/Core/Public/Math/RotationTranslationMatrix.h
--- a/Engine/Source/Runtime/Core/Public/Math/RotationTranslationMatrix.h
+++ b/Engine/Source/Runtime/Core/Public/Math/RotationTranslationMatrix.h
@@ -27,7 +27,7 @@ public:
 	TRotationTranslationMatrix(const TRotator<T>& Rot, const TVector<T>& Origin);
 
 	// Conversion to other type.
-	template<typename FArg, TEMPLATE_REQUIRES(!TIsSame<T, FArg>::Value)>
+	template<typename FArg, TEMPLATE_REQUIRES(!std::is_same_v<T, FArg>)>
 	explicit TRotationTranslationMatrix(const TRotationTranslationMatrix<FArg>& From) : TMatrix<T>(From) {}
 	
 	/** Matrix factory. Return an TMatrix<T> so we don't have type conversion issues in expressions. */
diff --git a/Engine/Source/Runtime/Core/Public/Math/Rotator.h b/Engine/Source/Runtime/Core/Public/Math/Rotator.h
--- a/Engine/Source/Runtime/Core/Public/Math/Rotator.h
+++ b/Engine/Source/Runtime/Core/Public/Math/Rotator.h
@@ -467,7 +467,7 @@ public:
 	bool SerializeFromMismatchedTag(FName StructTag, FArchive& Ar);
 
 	// Conversion from other type.
-	template<typename FArg, TEMPLATE_REQUIRES(!TIsSame<T, FArg>::Value)>
+	template<typename FArg, TEMPLATE_REQUIRES(!std::is_same_v<T, FArg>)>
 	explicit TRotator(const TRotator<FArg>& From) : TRotator<T>((T)From.Pitch, (T)From.Yaw, (T)From.Roll) {}
 };
 
diff --git a/Engine/Source/Runtime/Core/Public/Math/ScaleMatrix.h b/Engine/Source/Runtime/Core/Public/Math/ScaleMatrix.h
--- a/Engine/Source/Runtime/Core/Public/Math/ScaleMatrix.h
+++ b/Engine/Source/Runtime/Core/Public/Math/ScaleMatrix.h
@@ -29,7 +29,7 @@ public:
 	TScaleMatrix( const TVector<T>& Scale );
 
 	// Conversion to other type.
-	template<typename FArg, TEMPLATE_REQUIRES(!TIsSame<T, FArg>::Value)>
+	template<typename FArg, TEMPLATE_REQUIRES(!std::is_same_v<T, FArg>)>
 	explicit TScaleMatrix(const TScaleMatrix<FArg>& From) : TMatrix<T>(From) {}
 	
 	/** Matrix factory. Return an TMatrix<T> so we don't have type conversion issues in expressions. */
diff --git a/Engine/Source/Runtime/Core/Public/Math/ScaleRotationTranslationMatrix.h b/Engine/Source/Runtime/Core/Public/Math/ScaleRotationTranslationMatrix.h
--- a/Engine/Source/Runtime/Core/Public/Math/ScaleRotationTranslationMatrix.h
+++ b/Engine/Source/Runtime/Core/Public/Math/ScaleRotationTranslationMatrix.h
@@ -28,7 +28,7 @@ public:
 	TScaleRotationTranslationMatrix(const TVector<T>& Scale, const TRotator<T>& Rot, const TVector<T>& Origin);
 	
 	// Conversion to other type.
-	template<typename FArg, TEMPLATE_REQUIRES(!TIsSame<T, FArg>::Value)>
+	template<typename FArg, TEMPLATE_REQUIRES(!std::is_same_v<T, FArg>)>
 	explicit TScaleRotationTranslationMatrix(const TScaleRotationTranslationMatrix<FArg>& From) : TMatrix<T>(From) {}
 };
 
diff --git a/Engine/Source/Runtime/Core/Public/Math/Sphere.h b/Engine/Source/Runtime/Core/Public/Math/Sphere.h
--- a/Engine/Source/Runtime/Core/Public/Math/Sphere.h
+++ b/Engine/Source/Runtime/Core/Public/Math/Sphere.h
@@ -89,7 +89,7 @@ public:
 	CORE_API TSphere(const TSphere<T>* Spheres, int32 Count);
 
 	// Conversion from other variant type.
-	template<typename FArg, TEMPLATE_REQUIRES(!TIsSame<T, FArg>::Value)>
+	template<typename FArg, TEMPLATE_REQUIRES(!std::is_same_v<T, FArg>)>
 	explicit TSphere(const TSphere<FArg>& From) : TSphere<T>(TVector<T>(From.Center), T(From.W)) {}
 
 	/**
diff --git a/Engine/Source/Runtime/Core/Public/Math/TransformNonVectorized.h b/Engine/Source/Runtime/Core/Public/Math/TransformNonVectorized.h
--- a/Engine/Source/Runtime/Core/Public/Math/TransformNonVectorized.h
+++ b/Engine/Source/Runtime/Core/Public/Math/TransformNonVectorized.h
@@ -1211,7 +1211,7 @@ private:
 
 public:
 	// Conversion to other type.
-	template<typename FArg, TEMPLATE_REQUIRES(!TIsSame<T, FArg>::Value)>
+	template<typename FArg, TEMPLATE_REQUIRES(!std::is_same_v<T, FArg>)>
 	explicit TTransform(const TTransform<FArg>& From) : TTransform<T>((TQuat<T>)From.GetRotation(), (TVector<T>)From.GetTranslation(), (TVector<T>)From.GetScale3D()) {}
 };
 
diff --git a/Engine/Source/Runtime/Core/Public/Math/TransformVectorized.h b/Engine/Source/Runtime/Core/Public/Math/TransformVectorized.h
--- a/Engine/Source/Runtime/Core/Public/Math/TransformVectorized.h
+++ b/Engine/Source/Runtime/Core/Public/Math/TransformVectorized.h
@@ -1475,11 +1475,11 @@ public:
 
 	// Conversion to other type.
 	friend struct TTransform<double>;
-	template<typename FArg, TEMPLATE_REQUIRES(!TIsSame<T, FArg>::Value && TIsSame<T, float>::Value)>
+	template<typename FArg, TEMPLATE_REQUIRES(!std::is_same_v<T, FArg> && std::is_same_v<T, float>)>
 	explicit TTransform(const TTransform<FArg>& From) : TTransform(MakeVectorRegisterFloatFromDouble(From.Rotation), MakeVectorRegisterFloatFromDouble(From.Translation), MakeVectorRegisterFloatFromDouble(From.Scale3D)) {}
 
 	friend struct TTransform<float>;
-	template<typename FArg, TEMPLATE_REQUIRES(!TIsSame<T, FArg>::Value && TIsSame<T, double>::Value)>
+	template<typename FArg, TEMPLATE_REQUIRES(!std::is_same_v<T, FArg> && std::is_same_v<T, double>)>
 	explicit TTransform(const TTransform<FArg>& From) : TTransform(MakeVectorRegisterDouble(From.Rotation), MakeVectorRegisterDouble(From.Translation), MakeVectorRegisterDouble(From.Scale3D)) {}
 };
 
diff --git a/Engine/Source/Runtime/Core/Public/Math/TranslationMatrix.h b/Engine/Source/Runtime/Core/Public/Math/TranslationMatrix.h
--- a/Engine/Source/Runtime/Core/Public/Math/TranslationMatrix.h
+++ b/Engine/Source/Runtime/Core/Public/Math/TranslationMatrix.h
@@ -20,7 +20,7 @@ public:
 	TTranslationMatrix(const TVector<T>& Delta);
 
 	// Conversion to other type.
-	template<typename FArg, TEMPLATE_REQUIRES(!TIsSame<T, FArg>::Value)>
+	template<typename FArg, TEMPLATE_REQUIRES(!std::is_same_v<T, FArg>)>
 	explicit TTranslationMatrix(const TTranslationMatrix<FArg>& From) : TMatrix<T>(From) {}
 	
 	/** Matrix factory. Return an FMatrix so we don't have type conversion issues in expressions. */
diff --git a/Engine/Source/Runtime/Core/Public/Math/UnrealMathUtility.h b/Engine/Source/Runtime/Core/Public/Math/UnrealMathUtility.h
--- a/Engine/Source/Runtime/Core/Public/Math/UnrealMathUtility.h
+++ b/Engine/Source/Runtime/Core/Public/Math/UnrealMathUtility.h
@@ -741,7 +741,7 @@ public:
 		*ScalarCos = FMath::Cos(Value);
 	}
 
-	template<typename T, typename U, TEMPLATE_REQUIRES(!TIsSame<T, U>::Value)>
+	template<typename T, typename U, TEMPLATE_REQUIRES(!std::is_same_v<T, U>)>
 	static FORCEINLINE void SinCos(T* ScalarSin, T* ScalarCos, U Value)
 	{
 		SinCos(ScalarSin, ScalarCos, T(Value));
@@ -1061,10 +1061,7 @@ public:
 	}
 
 	/** Performs a linear interpolation between two values, Alpha ranges from 0-1 */
-	template< class T, class U, TEMPLATE_REQUIRES(TAnd<
-													TNot<TCustomLerp<T>>,
-													TOr<TIsFloatingPoint<U>, TIsSame<T, U>>
-													>::Value)>
+	template< class T, class U, TEMPLATE_REQUIRES(!TCustomLerp<T>::Value && (std::is_floating_point_v<U> || std::is_same_v<T, U>))>
 	UE_NODISCARD static constexpr FORCEINLINE_DEBUGGABLE T Lerp( const T& A, const T& B, const U& Alpha )
 	{
 		return (T)(A + Alpha * (B-A));
@@ -1078,11 +1075,7 @@ public:
 	}
 
 	// Allow passing of differing A/B types.
-	template<typename T1, typename T2, typename T3, TEMPLATE_REQUIRES(TAnd<
-																		TNot<TIsSame<T1, T2>>,
-																		TNot<TCustomLerp<T1>>,
-																		TNot<TCustomLerp<T2>>
-																		>::Value)>
+	template<typename T1, typename T2, typename T3, TEMPLATE_REQUIRES(!std::is_same_v<T1, T2> && !TCustomLerp<T1>::Value && !TCustomLerp<T2>::Value)>
 	UE_NODISCARD static auto Lerp( const T1& A, const T2& B, const T3& Alpha ) -> decltype(A * B)
 	{
 		using ABType = decltype(A * B);
@@ -1104,7 +1097,7 @@ public:
 	}
 
 	// Allow passing of differing A/B types.
-	template<typename T1, typename T2, typename T3, TEMPLATE_REQUIRES(!TIsSame<T1, T2>::Value)>
+	template<typename T1, typename T2, typename T3, TEMPLATE_REQUIRES(!std::is_same_v<T1, T2>)>
 	UE_NODISCARD static auto LerpStable( const T1& A, const T2& B, const T3& Alpha ) -> decltype(A * B)
 	{
 		using ABType = decltype(A * B);
@@ -1112,10 +1105,7 @@ public:
 	}
 	
 	/** Performs a 2D linear interpolation between four values values, FracX, FracY ranges from 0-1 */
-	template< class T, class U, TEMPLATE_REQUIRES(TAnd<
-													TNot<TCustomLerp<T>>,
-													TOr<TIsFloatingPoint<U>, TIsSame<T, U>>
-													>::Value)>
+	template< class T, class U, TEMPLATE_REQUIRES(!TCustomLerp<T>::Value && (std::is_floating_point_v<U> || std::is_same_v<T, U>))>
 	UE_NODISCARD static constexpr FORCEINLINE_DEBUGGABLE T BiLerp(const T& P00,const T& P10,const T& P01,const T& P11, const U& FracX, const U& FracY)
 	{
 		return Lerp(
@@ -1141,10 +1131,7 @@ public:
 	 *
 	 * @return  Interpolated value
 	 */
-	template< class T, class U, TEMPLATE_REQUIRES(TAnd<
-													TNot<TCustomLerp<T>>,
-													TOr<TIsFloatingPoint<U>, TIsSame<T, U>>
-													>::Value)>
+	template< class T, class U, TEMPLATE_REQUIRES(!TCustomLerp<T>::Value && (std::is_floating_point_v<U> || std::is_same_v<T, U>))>
 	UE_NODISCARD static constexpr FORCEINLINE_DEBUGGABLE T CubicInterp( const T& P0, const T& T0, const T& P1, const T& T1, const U& A )
 	{
 		const U A2 = A * A;
@@ -2427,22 +2414,6 @@ public:
 		}
 		return WeightedMovingAverage(CurrentSample, PreviousSample, Weight);
 	}
-
-	// Temporary support for ambiguities to simplify UE4 -> UE5 upgrades - falls back to float variant.
-	// A call is considered ambiguous if it passes no float types, any long double, or multiple mismatched float/double types.
-	template<typename T> UE_DEPRECATED(5.0, "Arguments cause function resolution ambiguity.")
-	UE_NODISCARD static FORCEINLINE float Log2(T&& Value) { return Log2((float)Value); }
-	template<typename T1, typename T2, typename T3, TEMPLATE_REQUIRES(TIsAmbiguous<T1, T2, T3>)> UE_DEPRECATED(5.0, "Arguments cause function resolution ambiguity.")
-	static FORCEINLINE void PolarToCartesian(const T1 Rad, const T2 Ang, T3& OutX, T3& OutY) { float OX, OY; PolarToCartesian((float)Rad, (float)Ang, OX, OY); OutX = OX; OutY = OY; }
-	template<typename T1, typename T2, typename T3, TEMPLATE_REQUIRES(TIsAmbiguous<T1, T2, T3>)> UE_DEPRECATED(5.0, "Arguments cause function resolution ambiguity.")
-	static FORCEINLINE void CartesianToPolar(const T1 X, const T2 Y, T3& OutRad, T3& OutAng) { float ORad, OAng; CartesianToPolar((float)X, (float)Y, ORad, OAng); OutRad = ORad; OutAng = OAng; }
-
-	template<typename T1, typename T2, typename T3, TEMPLATE_REQUIRES(TIsAmbiguous<T1, T2, T3>)> UE_DEPRECATED(5.0, "Arguments cause function resolution ambiguity.")
-	UE_NODISCARD static FORCEINLINE float SmoothStep(T1&& A, T2&& B, T3&& C) { return SmoothStep((float)A, (float)B, (float)C); }
-	template<typename T1, typename T2, typename T3, TEMPLATE_REQUIRES(TIsAmbiguous<T1, T2, T3>)> UE_DEPRECATED(5.0, "Arguments cause function resolution ambiguity.")
-	UE_NODISCARD static FORCEINLINE float WeightedMovingAverage(T1&& CurrentSample, T2&& PreviousSample, T3&& Weight) { return WeightedMovingAverage((float)CurrentSample, (float)PreviousSample, (float)Weight); }
-	template<typename T1, typename T2, typename T3, typename T4, typename T5, TEMPLATE_REQUIRES(TIsAmbiguous<T1, T2, T3, T4, T5>)> UE_DEPRECATED(5.0, "Arguments cause function resolution ambiguity.")
-	UE_NODISCARD static FORCEINLINE float DynamicWeightedMovingAverage(T1&& CurrentSample, T2&& PreviousSample, T3&& MaxDistance, T4&& MinWeight, T5&& MaxWeight) { return DynamicWeightedMovingAverage((float)CurrentSample, (float)PreviousSample, (float)MaxDistance, (float)MinWeight, (float)MaxWeight); }
 };
 
 // LWC Conversion helpers
diff --git a/Engine/Source/Runtime/Core/Public/Math/Vector.h b/Engine/Source/Runtime/Core/Public/Math/Vector.h
--- a/Engine/Source/Runtime/Core/Public/Math/Vector.h
+++ b/Engine/Source/Runtime/Core/Public/Math/Vector.h
@@ -1178,7 +1178,7 @@ public:
 	}
 
 	// Conversion from other type.
-	template<typename FArg, TEMPLATE_REQUIRES(!TIsSame<T, FArg>::Value)>
+	template<typename FArg, TEMPLATE_REQUIRES(!std::is_same_v<T, FArg>)>
 	explicit TVector(const TVector<FArg>& From) : TVector<T>((T)From.X, (T)From.Y, (T)From.Z) {}
 };
 
diff --git a/Engine/Source/Runtime/Core/Public/Math/Vector2D.h b/Engine/Source/Runtime/Core/Public/Math/Vector2D.h
--- a/Engine/Source/Runtime/Core/Public/Math/Vector2D.h
+++ b/Engine/Source/Runtime/Core/Public/Math/Vector2D.h
@@ -700,7 +700,7 @@ public:
 	
 	
 	// Conversion from other type.
-	template<typename FArg, TEMPLATE_REQUIRES(!TIsSame<T, FArg>::Value)>
+	template<typename FArg, TEMPLATE_REQUIRES(!std::is_same_v<T, FArg>)>
 	explicit TVector2(const TVector2<FArg>& From) : TVector2<T>((T)From.X, (T)From.Y) {}
 };
 
diff --git a/Engine/Source/Runtime/Core/Public/Math/Vector4.h b/Engine/Source/Runtime/Core/Public/Math/Vector4.h
--- a/Engine/Source/Runtime/Core/Public/Math/Vector4.h
+++ b/Engine/Source/Runtime/Core/Public/Math/Vector4.h
@@ -515,11 +515,11 @@ public:
 	bool SerializeFromMismatchedTag(FName StructTag, FArchive& Ar);
 	
 	// Conversion from other type: double->float
-	template<typename FArg, TEMPLATE_REQUIRES(TAnd<TIsSame<FArg, double>, TIsSame<T, float>>::Value)>
+	template<typename FArg, TEMPLATE_REQUIRES(std::is_same_v<FArg, double> && std::is_same_v<T, float>)>
 	explicit TVector4(const TVector4<FArg>& From) : TVector4<T>((T)From.X, (T)From.Y, (T)From.Z, (T)From.W) {}
 
 	// Conversion from other type: float->double
-	template<typename FArg, TEMPLATE_REQUIRES(TAnd<TIsSame<FArg, float>, TIsSame<T, double>>::Value)>
+	template<typename FArg, TEMPLATE_REQUIRES(std::is_same_v<FArg, float> && std::is_same_v<T, double>)>
 	explicit TVector4(const TVector4<FArg>& From) : TVector4<T>((T)From.X, (T)From.Y, (T)From.Z, (T)From.W) {}
 
 	/**
