Comment fix.

diff --git a/Engine/Source/Runtime/Core/Public/Math/Box.h b/Engine/Source/Runtime/Core/Public/Math/Box.h
--- a/Engine/Source/Runtime/Core/Public/Math/Box.h
+++ b/Engine/Source/Runtime/Core/Public/Math/Box.h
@@ -56,7 +56,12 @@ public:
 	{
 	}
 
-	// double box taking float params
+	/**
+	 * Creates and initializes a new box from the specified extents.
+	 *
+	 * @param InMin The box's minimum point.
+	 * @param InMax The box's maximum point.
+	 */
 	template<typename FArg>
 	TBox(const TVector<FArg>& InMin, const TVector<FArg>& InMax)
 		: Min(InMin)
