Fix check in debug

diff --git a/Engine/Source/Runtime/Core/Private/Async/Fundamental/Scheduler.cpp b/Engine/Source/Runtime/Core/Private/Async/Fundamental/Scheduler.cpp
--- a/Engine/Source/Runtime/Core/Private/Async/Fundamental/Scheduler.cpp
+++ b/Engine/Source/Runtime/Core/Private/Async/Fundamental/Scheduler.cpp
@@ -133,7 +133,7 @@ namespace LowLevelTasks
 			UE::Trace::ThreadGroupBegin(TEXT("Background Workers"));
 			for (uint32 WorkerId = 0; WorkerId < NumBackgroundWorkers; ++WorkerId)
 			{
-				WorkerThreads.Add(CreateWorker(true, IsForkable, &WorkerEvents[NumForegroundWorkers + WorkerId], &WorkerLocalQueues[NumForegroundWorkers + WorkerId], BackgroundPriority, BackgroundAffinity));
+				WorkerThreads.Add(CreateWorker(true, IsForkable, &WorkerEvents[NumForegroundWorkers + WorkerId], &WorkerLocalQueues[NumForegroundWorkers + WorkerId], GTaskGraphUseDynamicPrioritization ? WorkerPriority : BackgroundPriority, BackgroundAffinity));
 			}
 			UE::Trace::ThreadGroupEnd();
 		}
