Added static_assert to stop data loss when constructing a TBox<float> from TVector<double> extents, inadvertently reverted in CL# 23736521.

diff --git a/Engine/Source/Runtime/Core/Public/Math/Box.h b/Engine/Source/Runtime/Core/Public/Math/Box.h
--- a/Engine/Source/Runtime/Core/Public/Math/Box.h
+++ b/Engine/Source/Runtime/Core/Public/Math/Box.h
@@ -67,7 +67,10 @@ public:
 		: Min(InMin)
 		, Max(InMax)
 		, IsValid(1)
-	{ }
+	{
+		// Intended to catch TBox<float>(TVector<double>(), TVector<double>())
+		static_assert(sizeof(FArg) <= sizeof(T), "Losing precision when constructing a box of floats from vectors of doubles");
+	}
 
 	TBox(const TVector4<T>& InMin, const TVector4<T>& InMax)
 		: Min(InMin)
