Deprecated TIsSame.

diff --git a/Engine/Source/Runtime/Core/Public/Templates/UnrealTypeTraits.h b/Engine/Source/Runtime/Core/Public/Templates/UnrealTypeTraits.h
--- a/Engine/Source/Runtime/Core/Public/Templates/UnrealTypeTraits.h
+++ b/Engine/Source/Runtime/Core/Public/Templates/UnrealTypeTraits.h
@@ -73,8 +73,11 @@ struct TIsDerivedFrom
  *
  * Unreal implementation of std::is_same trait.
  */
-template<typename A, typename B>	struct TIsSame			{ enum { Value = false	}; };
-template<typename T>				struct TIsSame<T, T>	{ enum { Value = true	}; };
+template<typename A, typename B>
+struct UE_DEPRECATED(5.2, "TIsSame has been deprecated, please use std::is_same instead.") TIsSame
+{
+	enum { Value = std::is_same_v<A, B>	};
+};
 
 /** Gets the Nth type in a template parameter pack. N must be less than sizeof...(Types) */
 template <int32 N, typename... Types>
