Ensure that TArray's bulk-serialized data sizes match expected sizes.

diff --git a/Engine/Source/Runtime/Core/Public/Containers/Array.h b/Engine/Source/Runtime/Core/Public/Containers/Array.h
--- a/Engine/Source/Runtime/Core/Public/Containers/Array.h
+++ b/Engine/Source/Runtime/Core/Public/Containers/Array.h
@@ -28,6 +28,7 @@
 #include "Templates/MakeUnsigned.h"
 #include "Traits/ElementType.h"
 
+#include <limits>
 #include <type_traits>
 
 #if UE_BUILD_SHIPPING || UE_BUILD_TEST
@@ -1241,7 +1242,7 @@ public:
 	 */
 	void BulkSerialize(FArchive& Ar, bool bForcePerElementSerialization = false)
 	{
-		int32 ElementSize = sizeof(ElementType);
+		constexpr int32 ElementSize = sizeof(ElementType);
 		// Serialize element size to detect mismatch across platforms.
 		int32 SerializedElementSize = ElementSize;
 		Ar << SerializedElementSize;
@@ -1261,21 +1262,31 @@ public:
 			if (Ar.IsLoading())
 			{
 				// Basic sanity checking to ensure that sizes match.
-				checkf(SerializedElementSize == 0 || SerializedElementSize == ElementSize, TEXT("Unexpected array element size. Expected %i, Got: %i. Package can be corrupt or the array template type changed."), ElementSize, SerializedElementSize);
+				if (!ensure(SerializedElementSize == ElementSize))
+				{
+					Ar.SetError();
+					return;
+				}
+
 				// Serialize the number of elements, block allocate the right amount of memory and deserialize
 				// the data as a giant memory blob in a single call to Serialize. Please see the function header
 				// for detailed documentation on limitations and implications.
 				SizeType NewArrayNum = 0;
 				Ar << NewArrayNum;
+				if (!ensure(NewArrayNum >= 0 && std::numeric_limits<SizeType>::max() / (SizeType)ElementSize >= NewArrayNum))
+				{
+					Ar.SetError();
+					return;
+				}
 				Empty(NewArrayNum);
 				AddUninitialized(NewArrayNum);
-				Ar.Serialize(GetData(), (int64)NewArrayNum * (int64)SerializedElementSize);
+				Ar.Serialize(GetData(), (int64)NewArrayNum * (int64)ElementSize);
 			}
 			else if (Ar.IsSaving())
 			{
 				SizeType ArrayCount = Num();
 				Ar << ArrayCount;
-				Ar.Serialize(GetData(), (int64)ArrayCount * (int64)SerializedElementSize);
+				Ar.Serialize(GetData(), (int64)ArrayCount * (int64)ElementSize);
 			}
 		}
 	}
