OodleDataCompressionUtil : take dest buff size to ensure no overrun

diff --git a/Engine/Source/Runtime/Core/Private/Compression/OodleDataCompressionUtil.cpp b/Engine/Source/Runtime/Core/Private/Compression/OodleDataCompressionUtil.cpp
--- a/Engine/Source/Runtime/Core/Private/Compression/OodleDataCompressionUtil.cpp
+++ b/Engine/Source/Runtime/Core/Private/Compression/OodleDataCompressionUtil.cpp
@@ -109,7 +109,7 @@ namespace FOodleCompressedArray
 		return true;
 	}
 
-	bool CORE_API DecompressToExistingBuffer(void* InDestinationBuffer, TArray<uint8> const& InCompressed)
+	bool CORE_API DecompressToExistingBuffer(void* InDestinationBuffer, int64 InDestinationBufferSize, TArray<uint8> const& InCompressed)
 	{
 		int32 DecompressedSize, CompressedSize;
 		int32 OffsetToCompressedData = PeekSizes(InCompressed, CompressedSize, DecompressedSize);
@@ -117,13 +117,21 @@ namespace FOodleCompressedArray
 		{
 			return false;
 		}
+		// DestinationBuffer provided is not large enough to hold decompressed data :
+		if ( DecompressedSize > InDestinationBufferSize )
+		{
+			return false;
+		}
 
 		// If we have a valid header, then if we don't have the actual data, it's corrupted data.
-		check(OffsetToCompressedData + CompressedSize <= InCompressed.Num());
+		if ( OffsetToCompressedData + CompressedSize > InCompressed.Num() )
+		{
+			return false;
+		}
 
 		return FOodleDataCompression::Decompress(InDestinationBuffer, DecompressedSize, InCompressed.GetData() + OffsetToCompressedData, CompressedSize);
 	}
-	bool CORE_API DecompressToExistingBuffer64(void* InDestinationBuffer, TArray64<uint8> const& InCompressed)
+	bool CORE_API DecompressToExistingBuffer64(void* InDestinationBuffer, int64 InDestinationBufferSize, TArray64<uint8> const& InCompressed)
 	{
 		int64 DecompressedSize, CompressedSize;
 		int32 OffsetToCompressedData = PeekSizes64(InCompressed, CompressedSize, DecompressedSize);
@@ -131,9 +139,17 @@ namespace FOodleCompressedArray
 		{
 			return false;
 		}
-
+		// DestinationBuffer provided is not large enough to hold decompressed data :
+		if ( DecompressedSize > InDestinationBufferSize )
+		{
+			return false;
+		}
+		
 		// If we have a valid header, then if we don't have the actual data, it's corrupted data.
-		check(OffsetToCompressedData + CompressedSize <= InCompressed.Num());
+		if ( OffsetToCompressedData + CompressedSize > InCompressed.Num() )
+		{
+			return false;
+		}
 
 		return FOodleDataCompression::Decompress(InDestinationBuffer, DecompressedSize, InCompressed.GetData() + OffsetToCompressedData, CompressedSize);
 	}
diff --git a/Engine/Source/Runtime/Core/Public/Compression/OodleDataCompressionUtil.h b/Engine/Source/Runtime/Core/Public/Compression/OodleDataCompressionUtil.h
--- a/Engine/Source/Runtime/Core/Public/Compression/OodleDataCompressionUtil.h
+++ b/Engine/Source/Runtime/Core/Public/Compression/OodleDataCompressionUtil.h
@@ -142,15 +142,15 @@ namespace FOodleCompressedArray
 	/**
 	* Decompresses a compressed TArray to a buffer that has already been allocated.
 	* 
-	* @param InDestinationBUffer		The buffer to contain the decompressed data. This buffer must be at
-	*									least as large as the decompressed size specified by PeekSizes. If
-	*									it's not, memory corruption will occur.
+	* @param InDestinationBuffer		The buffer to contain the decompressed data.
+	* @param InDestinationBufferSize	Allocated size of InDestinationBuffer; if it is not large enough to hold the decompressed data
+	*									(as reported by PeekSizes), this function will return false.
 	* @param InCompressed				A TArray created by one of the FOodleCompressedArray compression functions.
 	* @return							Success or failure. Failure is usually because the FCompressedArray doesn't
 	*									actually contain any data, but might fail because 64 bit is needed.
 	*/
-	bool CORE_API DecompressToExistingBuffer(void* InDestinationBuffer, TArray<uint8> const& InCompressed);
-	bool CORE_API DecompressToExistingBuffer64(void* InDestinationBuffer, TArray64<uint8> const& InCompressed);
+	bool CORE_API DecompressToExistingBuffer(void* InDestinationBuffer, int64 InDestinationBufferSize,  TArray<uint8> const& InCompressed);
+	bool CORE_API DecompressToExistingBuffer64(void* InDestinationBuffer, int64 InDestinationBufferSize, TArray64<uint8> const& InCompressed);
 
 	/**
 	* Decompresses a compressed TArray to a buffer that will be allocated by this function.
@@ -210,7 +210,7 @@ namespace FOodleCompressedArray
 		}
 
 		OutDecompressed.SetNum(DecompressedSize / sizeof(T));
-		return DecompressToExistingBuffer(OutDecompressed.GetData(), InCompressed);
+		return DecompressToExistingBuffer(OutDecompressed.GetData(),DecompressedSize, InCompressed);
 	}
 	template <class T>
 	static bool DecompressToTArray64(TArray64<T>& OutDecompressed, TArray64<uint8> const& InCompressed)
@@ -228,7 +228,7 @@ namespace FOodleCompressedArray
 		}
 
 		OutDecompressed.SetNum(DecompressedSize / sizeof(T));
-		return DecompressToExistingBuffer64(OutDecompressed.GetData(), InCompressed);
+		return DecompressToExistingBuffer64(OutDecompressed.GetData(),DecompressedSize, InCompressed);
 	}
 };
 
