Limited TVariant to 256 types, which reduces the size needed for the index.

diff --git a/Engine/Source/Runtime/Core/Public/Misc/TVariant.h b/Engine/Source/Runtime/Core/Public/Misc/TVariant.h
--- a/Engine/Source/Runtime/Core/Public/Misc/TVariant.h
+++ b/Engine/Source/Runtime/Core/Public/Misc/TVariant.h
@@ -35,6 +35,9 @@ class TVariant final : private UE::Core::Private::TVariantStorage<T, Ts...>
 	static_assert(!UE::Core::Private::TTypePackContainsDuplicates<T, Ts...>::Value, "All the types used in TVariant should be unique");
 	static_assert(!UE::Core::Private::TContainsReferenceType<T, Ts...>::Value, "TVariant cannot hold reference types");
 
+	// Test for 255 here, because the parameter pack doesn't include the initial T
+	static_assert(sizeof...(Ts) <= 255, "TVariant cannot hold more than 256 types");
+
 public:
 	/** Default initialize the TVariant to the first type in the parameter pack */
 	TVariant()
@@ -52,7 +55,7 @@ public:
 		static_assert(Index != (SIZE_T)-1, "The TVariant is not declared to hold the type being constructed");
 
 		new(&UE::Core::Private::CastToStorage(*this).Storage) U(Forward<TArgs>(Args)...);
-		TypeIndex = Index;
+		TypeIndex = (uint8)Index;
 	}
 
 	/** Copy construct the variant from another variant of the same type */
@@ -134,7 +137,7 @@ public:
 		static_assert(Index != (SIZE_T)-1, "The TVariant is not declared to hold the type passed to TryGet<>");
 		// The intermediate step of casting to void* is used to avoid warnings due to use of reinterpret_cast between related types if U and the storage class are related
 		// This was specifically encountered when U derives from TAlignedBytes
-		return Index == TypeIndex ? reinterpret_cast<U*>(reinterpret_cast<void*>(&UE::Core::Private::CastToStorage(*this).Storage)) : nullptr;
+		return Index == (SIZE_T)TypeIndex ? reinterpret_cast<U*>(reinterpret_cast<void*>(&UE::Core::Private::CastToStorage(*this).Storage)) : nullptr;
 	}
 
 	/** Get a pointer to the held value if the held type is the same as the one specified */
@@ -168,7 +171,7 @@ public:
 
 		UE::Core::Private::TDestructorLookup<T, Ts...>::Destruct(TypeIndex, &UE::Core::Private::CastToStorage(*this).Storage);
 		new(&UE::Core::Private::CastToStorage(*this).Storage) U(Forward<TArgs>(Args)...);
-		TypeIndex = Index;
+		TypeIndex = (uint8)Index;
 	}
 
 	/** Lookup the index of a type in the template parameter pack at compile time. */
@@ -183,12 +186,12 @@ public:
 	/** Returns the currently held type's index into the template parameter pack */
 	SIZE_T GetIndex() const
 	{
-		return TypeIndex;
+		return (SIZE_T)TypeIndex;
 	}
 
 private:
 	/** Index into the template parameter pack for the type held. */
-	SIZE_T TypeIndex;
+	uint8 TypeIndex;
 };
 
 /** Determine if a type is a variant */
@@ -255,7 +258,6 @@ decltype(auto) Visit(Func&& Callable, Variants&&... Args)
 template <typename... Ts>
 inline FArchive& operator<<(typename UE::Core::Private::TAlwaysFArchive<TVariant<Ts...>>::Type& Ar, TVariant<Ts...>& Variant)
 {
-	static_assert(sizeof...(Ts) < 256, "TVariant serialization assumes that the stored index can fit in 8 bits");
 	if (Ar.IsLoading())
 	{
 		uint8 Index;
