Logging: Added localized structured logging

#rnx

diff --git a/Engine/Source/Runtime/Core/Private/Logging/StructuredLog.cpp b/Engine/Source/Runtime/Core/Private/Logging/StructuredLog.cpp
--- a/Engine/Source/Runtime/Core/Private/Logging/StructuredLog.cpp
+++ b/Engine/Source/Runtime/Core/Private/Logging/StructuredLog.cpp
@@ -9,6 +9,7 @@
 #include "Containers/StringView.h"
 #include "CoreGlobals.h"
 #include "HAL/PlatformTime.h"
+#include "Internationalization/Internationalization.h"
 #include "Logging/LogTrace.h"
 #include "Misc/AsciiSet.h"
 #include "Misc/DateTime.h"
@@ -30,7 +31,7 @@ static constexpr FAsciiSet ValidLogFieldName("0123456789ABCDEFGHIJKLMNOPQRSTUVWX
 
 struct FLogTemplateOp
 {
-	enum EOpCode : int32 { OpEnd, OpSkip, OpText, OpName, OpIndex, OpCount };
+	enum EOpCode : int32 { OpEnd, OpSkip, OpText, OpName, OpIndex, OpLocText, OpCount };
 
 	static constexpr int32 ValueShift = 3;
 	static_assert(OpCount <= (1 << ValueShift));
@@ -154,6 +155,56 @@ static void LogFieldValue(TStringBuilderBase<CharType>& Out, const FCbFieldView&
 	}
 }
 
+static void AddFieldValue(FFormatNamedArguments& Out, const FCbFieldView& Field)
+{
+	const FString FieldName(Field.GetName());
+	switch (FCbValue Accessor = Field.GetValue(); Accessor.GetType())
+	{
+	case ECbFieldType::Null:
+		break;
+	case ECbFieldType::Object:
+	case ECbFieldType::UniformObject:
+		break;
+	case ECbFieldType::Array:
+	case ECbFieldType::UniformArray:
+	case ECbFieldType::Binary:
+	case ECbFieldType::String:
+		break;
+	case ECbFieldType::IntegerPositive:
+		Out.Emplace(FieldName, Accessor.AsIntegerPositive());
+		return;
+	case ECbFieldType::IntegerNegative:
+		Out.Emplace(FieldName, Accessor.AsIntegerNegative());
+		return;
+	case ECbFieldType::Float32:
+		Out.Emplace(FieldName, Accessor.AsFloat32());
+		return;
+	case ECbFieldType::Float64:
+		Out.Emplace(FieldName, Accessor.AsFloat64());
+		return;
+	case ECbFieldType::BoolFalse:
+	case ECbFieldType::BoolTrue:
+	case ECbFieldType::ObjectAttachment:
+	case ECbFieldType::BinaryAttachment:
+	case ECbFieldType::Hash:
+	case ECbFieldType::Uuid:
+	case ECbFieldType::DateTime:
+	case ECbFieldType::TimeSpan:
+	case ECbFieldType::ObjectId:
+	case ECbFieldType::CustomById:
+	case ECbFieldType::CustomByName:
+		break;
+	default:
+		checkNoEntry();
+		break;
+	}
+
+	// Handle anything that falls through as text.
+	TStringBuilder<128> Text;
+	LogFieldValue(Text, Field);
+	Out.Emplace(FieldName, FText::FromString(FString(Text)));
+}
+
 } // UE::Logging::Private
 
 namespace UE
@@ -166,13 +217,34 @@ class FLogTemplate
 	using FLogField = Logging::Private::FLogField;
 
 public:
+	static FLogTemplate* CreateStatic(const Logging::Private::FStaticLogRecord& Log, const FLogField* Fields, int32 FieldCount)
+	{
+		return Create(Log.Format, Fields, FieldCount);
+	}
+
+	static FLogTemplate* CreateStatic(const Logging::Private::FStaticLocalizedLogRecord& Log, const FLogField* Fields, int32 FieldCount)
+	{
+		return CreateLocalized(Log.TextNamespace, Log.TextKey, Log.Format, Fields, FieldCount);
+	}
+
 	static FLogTemplate* Create(const TCHAR* Format, const FLogField* Fields = nullptr, int32 FieldCount = 0);
+
+	static FLogTemplate* CreateLocalized(
+		const TCHAR* TextNamespace,
+		const TCHAR* TextKey,
+		const TCHAR* Format,
+		const FLogField* Fields = nullptr,
+		int32 FieldCount = 0);
+
 	static void Destroy(FLogTemplate* Template);
 
 	template <typename CharType>
 	void FormatTo(TStringBuilderBase<CharType>& Out, const TCHAR* Format, const FCbObjectView& Fields) const;
 
 private:
+	template <typename CharType>
+	void FormatLocalizedTo(TStringBuilderBase<CharType>& Out, const FCbObjectView& Fields) const;
+
 	FLogTemplate() = default;
 	~FLogTemplate() = default;
 	FLogTemplate(const FLogTemplate&) = delete;
@@ -277,7 +349,7 @@ FLogTemplate* FLogTemplate::Create(const TCHAR* Format, const FLogField* Fields,
 	}
 
 	checkf(!bFindFields || !bPositional || FormatFieldCount == FieldCount,
-		TEXT("Log format requires %d fields and %d were provided. [[%s]]"), FormatFieldCount, FieldCount);
+		TEXT("Log format requires %d fields and %d were provided. [[%s]]"), FormatFieldCount, FieldCount, Format);
 
 	const uint32 TotalSize = Algo::TransformAccumulate(Ops, FLogTemplateOp::SaveSize, 0);
 	FLogTemplate* const Template = (FLogTemplate*)FMemory::Malloc(TotalSize);
@@ -289,9 +361,67 @@ FLogTemplate* FLogTemplate::Create(const TCHAR* Format, const FLogField* Fields,
 	return Template;
 }
 
+FLogTemplate* FLogTemplate::CreateLocalized(const TCHAR* TextNamespace, const TCHAR* TextKey, const TCHAR* Format, const FLogField* Fields, const int32 FieldCount)
+{
+	using namespace Logging::Private;
+
+	// Disallow escape sequences and argument modifiers.
+	// These are inconsistent with non-localized format strings and are disallowed for now.
+	constexpr FAsciiSet Invalid("`|");
+	checkf(FAsciiSet::HasNone(Format, Invalid),
+		TEXT("Log format does not currently allow escapes (`) or argument modifiers (|). [[%s]]"), Format);
+
+	FTextFormat TextFormat(FInternationalization::ForUseOnlyByLocMacroAndGraphNodeTextLiterals_CreateText(Format, TextNamespace, TextKey));
+
+	int32 FormatFieldCount = 0;
+	TArray<FString> TextFormatFieldNames;
+	TextFormat.GetFormatArgumentNames(TextFormatFieldNames);
+	for (const FString& FormatName : TextFormatFieldNames)
+	{
+		bool bFoundField = false;
+		const TCHAR* FormatNameStr = *FormatName;
+		const int32 FormatNameLen = FormatName.Len();
+		for (int32 BaseFieldIndex = 0; BaseFieldIndex < FieldCount; ++BaseFieldIndex)
+		{
+			const int32 FieldIndex = (FormatFieldCount + BaseFieldIndex) % FieldCount;
+			const ANSICHAR* FieldName = Fields[FieldIndex].Name;
+			if (FPlatformString::Strncmp(FieldName, FormatNameStr, FormatNameLen) == 0 && !FieldName[FormatNameLen])
+			{
+				bFoundField = true;
+				break;
+			}
+		}
+		checkf(bFoundField, TEXT("Log format requires field '%s' which was not provided. [[%s]]"), FormatNameStr, Format);
+		++FormatFieldCount;
+	}
+
+	const FLogTemplateOp Ops[]{{FLogTemplateOp::OpLocText}, {FLogTemplateOp::OpEnd}};
+	const uint32 TotalSize = sizeof(FTextFormat) + Algo::TransformAccumulate(Ops, FLogTemplateOp::SaveSize, 0);
+	FTextFormat* NewTextFormat = new(FMemory::Malloc(TotalSize, alignof(FTextFormat))) FTextFormat(MoveTemp(TextFormat));
+	FLogTemplate* const Template = (FLogTemplate*)(NewTextFormat + 1);
+	uint8* Data = (uint8*)Template;
+	for (const FLogTemplateOp& Op : Ops)
+	{
+		FLogTemplateOp::Save(Op, Data);
+	}
+	return Template;
+}
+
 void FLogTemplate::Destroy(FLogTemplate* Template)
 {
-	FMemory::Free(Template);
+	using namespace Logging::Private;
+
+	const uint8* NextOp = (const uint8*)Template;
+	if (FLogTemplateOp::Load(NextOp).Code == FLogTemplateOp::OpLocText)
+	{
+		FTextFormat* TextFormat = (FTextFormat*)Template - 1;
+		TextFormat->~FTextFormat();
+		FMemory::Free(TextFormat);
+	}
+	else
+	{
+		FMemory::Free(Template);
+	}
 }
 
 template <typename CharType>
@@ -344,6 +474,8 @@ void FLogTemplate::FormatTo(TStringBuilderBase<CharType>& Out, const TCHAR* Form
 		const FLogTemplateOp Op = FLogTemplateOp::Load(NextOp);
 		switch (Op.Code)
 		{
+		case FLogTemplateOp::OpLocText:
+			return FormatLocalizedTo(Out, Fields);
 		case FLogTemplateOp::OpEnd:
 			return;
 		case FLogTemplateOp::OpText:
@@ -362,6 +494,22 @@ void FLogTemplate::FormatTo(TStringBuilderBase<CharType>& Out, const TCHAR* Form
 	}
 }
 
+template <typename CharType>
+FORCENOINLINE void FLogTemplate::FormatLocalizedTo(TStringBuilderBase<CharType>& Out, const FCbObjectView& Fields) const
+{
+	using namespace Logging::Private;
+
+	FFormatNamedArguments TextFormatArguments;
+	for (const FCbFieldView& Field : Fields)
+	{
+		AddFieldValue(TextFormatArguments, Field);
+	}
+
+	const FTextFormat* TextFormat = (const FTextFormat*)this - 1;
+	const FText Text = FText::Format(*TextFormat, TextFormatArguments);
+	Out.Append(Text.ToString());
+}
+
 ///////////////////////////////////////////////////////////////////////////////////////////////////////////////////////
 
 FLogTime FLogTime::Now()
@@ -394,7 +542,12 @@ static void FormatRecordMessageTo(TStringBuilderBase<CharType>& Out, const FLogR
 		return Template->FormatTo(Out, Format, Record.GetFields());
 	}
 
-	FLogTemplate* LocalTemplate = FLogTemplate::Create(Format);
+	const TCHAR* TextNamespace = Record.GetTextNamespace();
+	const TCHAR* TextKey = Record.GetTextKey();
+	checkf(!TextNamespace == !TextKey,
+		TEXT("Log record must have both or neither of the text namespace and text key. [[%s]]"), Format);
+
+	FLogTemplate* LocalTemplate = TextKey ? FLogTemplate::CreateLocalized(TextNamespace, TextKey, Format) : FLogTemplate::Create(Format);
 	LocalTemplate->FormatTo(Out, Format, Record.GetFields());
 	FLogTemplate::Destroy(LocalTemplate);
 }
@@ -485,18 +638,18 @@ static FStaticLogDynamicDataManager GStaticLogDynamicDataManager;
 
 // Tracing the log happens in its own function because that allows stack space for the message to
 // be returned before calling into the output devices.
-FORCENOINLINE void LogToTrace(const FStaticLogRecord& Log, const FLogRecord& Record)
+FORCENOINLINE static void LogToTrace(const void* LogPoint, const FLogRecord& Record)
 {
 #if LOGTRACE_ENABLED
 	TStringBuilder<1024> Message;
 	Record.FormatMessageTo(Message);
-	FLogTrace::OutputLogMessageSimple(&Log, TEXT("%s"), *Message);
+	FLogTrace::OutputLogMessageSimple(LogPoint, TEXT("%s"), *Message);
 #endif
 }
 
 // Serializing log fields to compact binary happens in its own function because that allows stack
 // space for the writer to be returned before calling into the output devices.
-FORCENOINLINE FCbObject SerializeLogFields(
+FORCENOINLINE static FCbObject SerializeLogFields(
 	const FStaticLogRecord& Log,
 	const FLogTemplate& Template,
 	const FLogField* Fields,
@@ -535,13 +688,37 @@ FORCENOINLINE FCbObject SerializeLogFields(
 	return Writer.Save().AsObject();
 }
 
-FORCENOINLINE FLogTemplate& CreateLogTemplate(const FStaticLogRecord& Log, const FLogField* Fields, const int32 FieldCount)
+// Serializing log fields to compact binary happens in its own function because that allows stack
+// space for the writer to be returned before calling into the output devices.
+FORCENOINLINE static FCbObject SerializeLogFields(
+	const FStaticLocalizedLogRecord& Log,
+	const FLogTemplate& Template,
+	const FLogField* Fields,
+	const int32 FieldCount)
+{
+	if (FieldCount == 0)
+	{
+		return FCbObject();
+	}
+
+	TCbWriter<1024> Writer;
+	Writer.BeginObject();
+	for (const FLogField* FieldsEnd = Fields + FieldCount; Fields != FieldsEnd; ++Fields)
+	{
+		Fields->WriteValue(Writer.SetName(Fields->Name), Fields->Value);
+	}
+	Writer.EndObject();
+	return Writer.Save().AsObject();
+}
+
+template <typename StaticLogRecordType>
+FORCENOINLINE static FLogTemplate& CreateLogTemplate(const StaticLogRecordType& Log, const FLogField* Fields, const int32 FieldCount)
 {
 #if LOGTRACE_ENABLED
 	FLogTrace::OutputLogMessageSpec(&Log, &Log.Category, Log.Verbosity, Log.File, Log.Line, TEXT("%s"));
 #endif
 
-	FLogTemplate* NewTemplate = FLogTemplate::Create(Log.Format, Fields, FieldCount);
+	FLogTemplate* NewTemplate = FLogTemplate::CreateStatic(Log, Fields, FieldCount);
 	if (FLogTemplate* ExistingTemplate = nullptr;
 		UNLIKELY(!Log.DynamicData.Template.compare_exchange_strong(ExistingTemplate, NewTemplate, std::memory_order_release, std::memory_order_acquire)))
 	{
@@ -562,7 +739,8 @@ FORCENOINLINE FLogTemplate& CreateLogTemplate(const FStaticLogRecord& Log, const
 	return *NewTemplate;
 }
 
-inline FLogTemplate& EnsureLogTemplate(const FStaticLogRecord& Log, const FLogField* Fields, const int32 FieldCount)
+template <typename StaticLogRecordType>
+inline static FLogTemplate& EnsureLogTemplate(const StaticLogRecordType& Log, const FLogField* Fields, const int32 FieldCount)
 {
 	if (FLogTemplate* Template = Log.DynamicData.Template.load(std::memory_order_acquire))
 	{
@@ -571,7 +749,8 @@ inline FLogTemplate& EnsureLogTemplate(const FStaticLogRecord& Log, const FLogFi
 	return CreateLogTemplate(Log, Fields, FieldCount);
 }
 
-void LogWithFieldArray(const FStaticLogRecord& Log, const FLogField* Fields, const int32 FieldCount)
+template <typename StaticLogRecordType>
+inline static FLogRecord CreateLogRecord(const StaticLogRecordType& Log, const FLogField* Fields, const int32 FieldCount)
 {
 	FLogTemplate& Template = EnsureLogTemplate(Log, Fields, FieldCount);
 
@@ -584,11 +763,16 @@ void LogWithFieldArray(const FStaticLogRecord& Log, const FLogField* Fields, con
 	Record.SetCategory(Log.Category.GetCategoryName());
 	Record.SetVerbosity(Log.Verbosity);
 	Record.SetTime(FLogTime::Now());
+	return Record;
+}
 
+template <typename StaticLogRecordType>
+inline static void DispatchLogRecord(const StaticLogRecordType& Log, const FLogRecord& Record)
+{
 #if LOGTRACE_ENABLED
 	if (UE_TRACE_CHANNELEXPR_IS_ENABLED(LogChannel))
 	{
-		LogToTrace(Log, Record);
+		LogToTrace(&Log, Record);
 	}
 #endif
 
@@ -603,6 +787,11 @@ void LogWithFieldArray(const FStaticLogRecord& Log, const FLogField* Fields, con
 	}
 }
 
+void LogWithFieldArray(const FStaticLogRecord& Log, const FLogField* Fields, const int32 FieldCount)
+{
+	DispatchLogRecord(Log, CreateLogRecord(Log, Fields, FieldCount));
+}
+
 void LogWithNoFields(const FStaticLogRecord& Log)
 {
 	// A non-null field pointer enables field validation in FLogTemplate::Create.
@@ -610,6 +799,19 @@ void LogWithNoFields(const FStaticLogRecord& Log)
 	LogWithFieldArray(Log, &EmptyField, 0);
 }
 
+void LogWithFieldArray(const FStaticLocalizedLogRecord& Log, const FLogField* Fields, const int32 FieldCount)
+{
+	FLogRecord Record = CreateLogRecord(Log, Fields, FieldCount);
+	Record.SetTextNamespace(Log.TextNamespace);
+	Record.SetTextKey(Log.TextKey);
+	DispatchLogRecord(Log, Record);
+}
+
+void LogWithNoFields(const FStaticLocalizedLogRecord& Log)
+{
+	LogWithFieldArray(Log, nullptr, 0);
+}
+
 } // UE::Logging::Private
 
 #endif // !NO_LOGGING
diff --git a/Engine/Source/Runtime/Core/Private/Misc/OutputDeviceRedirector.cpp b/Engine/Source/Runtime/Core/Private/Misc/OutputDeviceRedirector.cpp
--- a/Engine/Source/Runtime/Core/Private/Misc/OutputDeviceRedirector.cpp
+++ b/Engine/Source/Runtime/Core/Private/Misc/OutputDeviceRedirector.cpp
@@ -135,7 +135,6 @@ struct FOutputDeviceRedirectorState
 
 	/** A queue of items logged by non-primary threads. */
 	TDepletableMpmcQueue<FOutputDeviceItem, FOutputDeviceLinearAllocator> BufferedItems;
-	uint8 BufferedItemsPadding[CalculateRedirectorCacheLinePadding(sizeof(BufferedItems))]{};
 
 	/** Array of output devices to redirect to from the primary thread. */
 	TArray<FOutputDevice*> BufferedOutputDevices;
diff --git a/Engine/Source/Runtime/Core/Public/Logging/StructuredLog.h b/Engine/Source/Runtime/Core/Public/Logging/StructuredLog.h
--- a/Engine/Source/Runtime/Core/Public/Logging/StructuredLog.h
+++ b/Engine/Source/Runtime/Core/Public/Logging/StructuredLog.h
@@ -13,6 +13,8 @@
 #include <atomic>
 #include <type_traits>
 
+class FTextFormat;
+
 #define UE_API CORE_API
 
 #if !NO_LOGGING
@@ -24,7 +26,7 @@
  * UE_SLOG(LogCore, Warning, "Loading '{Name}' failed with error {Error}", Package->GetName(), ErrorCode);
  *
  * Field names must match "[A-Za-z0-9_]+" and must be unique within this log event.
- * Field values will be serialized using operator<<(FCbWriter&, FieldType).
+ * Field values will be serialized using SerializeForLog or operator<<(FCbWriter&, FieldType).
  * The field values must exactly match the fields referenced by Format.
  *
  * @param CategoryName   Name of a log category declared by DECLARE_LOG_CATEGORY_*.
@@ -42,7 +44,7 @@
  *     ("Name", Package->GetName()), ("Error", ErrorCode), ("Flags", LoadFlags));
  *
  * Field names must match "[A-Za-z0-9_]+" and must be unique within this log event.
- * Field values will be serialized using operator<<(FCbWriter&, FieldType).
+ * Field values will be serialized using SerializeForLog or operator<<(FCbWriter&, FieldType).
  * The fields must contain every field referenced by Format. Order is irrelevant and extra fields are permitted.
  *
  * @param CategoryName   Name of a log category declared by DECLARE_LOG_CATEGORY_*.
@@ -69,6 +71,44 @@
  */
 #define UE_SLOG_FIELD(Name, Value) ::UE::Logging::Private::CheckFieldName(Name), Value
 
+/**
+ * Records a localized structured log event if this category is active at this level of verbosity.
+ *
+ * Example:
+ * UE_SLOG_LOC(LogCore, Warning, "LoadingFailed", "Loading '{Name}' failed with error {Error}",
+ *     ("Name", Package->GetName()), ("Error", ErrorCode), ("Flags", LoadFlags));
+ *
+ * Field names must match "[A-Za-z0-9_]+" and must be unique within this log event.
+ * Field values will be serialized using SerializeForLog or operator<<(FCbWriter&, FieldType).
+ * The fields must contain every field referenced by Format. Order is irrelevant and extra fields are permitted.
+ *
+ * @param CategoryName   Name of a log category declared by DECLARE_LOG_CATEGORY_*.
+ * @param Verbosity      Name of a log verbosity level from ELogVerbosity.
+ * @param Namespace      Namespace for the format FText, or LOCTEXT_NAMESPACE for the non-NS macro.
+ * @param Key            Key for the format FText that is unique within the namespace.
+ * @param Format         Format string in the style of FTextFormat.
+ * @param Fields[0-9]    Zero to nine fields in the format ("Name", Value).
+ */
+#define UE_SLOG_LOC(CategoryName, Verbosity, Key, Format, ...) \
+	UE_SLOG_NSLOC(CategoryName, Verbosity, LOCTEXT_NAMESPACE, Key, Format, ##__VA_ARGS__)
+#define UE_SLOG_NSLOC(CategoryName, Verbosity, Namespace, Key, Format, ...) \
+	UE_PRIVATE_SLOG_CALL(UE_SLOG_NSLOC_EX, (CategoryName, Verbosity, Namespace, Key, Format UE_PRIVATE_SLOG_FIELDS(__VA_ARGS__)))
+
+/**
+ * Records a localized structured log event if this category is active at this level of verbosity.
+ *
+ * Example:
+ * UE_SLOG_LOC_EX(LogCore, Warning, "LoadingFailed", "Loading '{PackageName}' failed with error {Error}",
+ *     UE_SLOG_FIELD("Name", Package->GetName()), UE_SLOG_FIELD("Error", ErrorCode), UE_SLOG_FIELD("Flags", LoadFlags));
+ *
+ * Same as UE_SLOG_LOC and works for any number of fields.
+ * Fields must be written as UE_SLOG_FIELD("Name", Value).
+ */
+#define UE_SLOG_LOC_EX(CategoryName, Verbosity, Key, Format, ...) \
+	UE_SLOG_NSLOC_EX(CategoryName, Verbosity, LOCTEXT_NAMESPACE, Key, Format, ##__VA_ARGS__)
+#define UE_SLOG_NSLOC_EX(CategoryName, Verbosity, Namespace, Key, Format, ...) \
+	UE_PRIVATE_SLOG_LOC(CategoryName, Verbosity, Namespace, Key, Format, ##__VA_ARGS__)
+
 #else // #if !NO_LOGGING
 
 #define UE_SLOG(...)
@@ -76,6 +116,11 @@
 #define UE_SLOG_NAMED_EX(...)
 #define UE_SLOG_FIELD(Name, Value)
 
+#define UE_SLOG_LOC(...)
+#define UE_SLOG_NSLOC(...)
+#define UE_SLOG_LOC_EX(...)
+#define UE_SLOG_NSLOC_EX(...)
+
 #endif
 
 struct FDateTime;
@@ -150,6 +195,14 @@ public:
 	int32 GetLine() const { return Line; }
 	void SetLine(int32 InLine) { Line = InLine; }
 
+	/** The namespace of the localized text. Null when non-localized. */
+	const TCHAR* GetTextNamespace() const { return TextNamespace; }
+	void SetTextNamespace(const TCHAR* InTextNamespace) { TextNamespace = InTextNamespace; }
+
+	/** The key of the localized text. Null when non-localized. */
+	const TCHAR* GetTextKey() const { return TextKey; }
+	void SetTextKey(const TCHAR* InTextKey) { TextKey = InTextKey; }
+
 	/** Formats the message using the format, template, and fields. */
 	UE_API void FormatMessageTo(FUtf8StringBuilderBase& Out) const;
 	UE_API void FormatMessageTo(FWideStringBuilderBase& Out) const;
@@ -163,6 +216,8 @@ private:
 	FLogTime Time;
 	FCbObject Fields;
 	const FLogTemplate* Template = nullptr;
+	const TCHAR* TextNamespace = nullptr;
+	const TCHAR* TextKey = nullptr;
 };
 
 template <typename ValueType>
@@ -194,6 +249,19 @@ struct FStaticLogRecord
 	FStaticLogDynamicData& DynamicData;
 };
 
+/** Data about a static localized log that is constant for every occurrence. */
+struct FStaticLocalizedLogRecord
+{
+	const FLogCategoryBase& Category;
+	const TCHAR* TextNamespace = nullptr;
+	const TCHAR* TextKey = nullptr;
+	const TCHAR* Format = nullptr;
+	const ANSICHAR* File = nullptr;
+	int32 Line = 0;
+	ELogVerbosity::Type Verbosity = ELogVerbosity::Log;
+	FStaticLogDynamicData& DynamicData;
+};
+
 struct FLogField
 {
 	using FWriteFn = void (FCbWriter& Writer, const void* Value);
@@ -219,6 +287,9 @@ namespace UE::Logging::Private
 UE_API void LogWithNoFields(const FStaticLogRecord& Log);
 UE_API void LogWithFieldArray(const FStaticLogRecord& Log, const FLogField* Fields, int32 FieldCount);
 
+UE_API void LogWithNoFields(const FStaticLocalizedLogRecord& Log);
+UE_API void LogWithFieldArray(const FStaticLocalizedLogRecord& Log, const FLogField* Fields, int32 FieldCount);
+
 /** Creates log fields with no name from value arguments. */
 struct FLogFieldsAnonymous
 {
@@ -264,7 +335,7 @@ struct FLogFieldsNamed
 };
 
 /** Log with fields created from the arguments, which may be values or pairs of name/value. */
-template <const FStaticLogRecord& Log, typename LogFieldsType, typename... FieldArgTypes>
+template <const auto& Log, typename LogFieldsType, typename... FieldArgTypes>
 FORCENOINLINE UE_DEBUG_SECTION void LogWithFields(FieldArgTypes&&... FieldArgs)
 {
 	constexpr int32 FieldCount = LogFieldsType::template GetCount<FieldArgTypes...>();
@@ -275,7 +346,7 @@ FORCENOINLINE UE_DEBUG_SECTION void LogWithFields(FieldArgTypes&&... FieldArgs)
 }
 
 /** Log if the category is active at this level of verbosity. */
-template <const FStaticLogRecord& Log, typename LogCategoryType, ELogVerbosity::Type Verbosity, typename LogFieldsType, typename... FieldArgTypes>
+template <const auto& Log, typename LogCategoryType, ELogVerbosity::Type Verbosity, typename LogFieldsType, typename... FieldArgTypes>
 inline void LogIfActive(FieldArgTypes&&... FieldArgs)
 {
 	static_assert(Verbosity != ELogVerbosity::Fatal, "Fatal verbosity is not supported by this API at this time.");
@@ -325,6 +396,15 @@ constexpr decltype(auto) CheckFieldName(NameType&& Name)
 	} \
 	while (false)
 
+#define UE_PRIVATE_SLOG_LOC(CategoryName, Verbosity, Namespace, Key, Format, ...) \
+	do \
+	{ \
+		static ::UE::Logging::Private::FStaticLogDynamicData DynamicData; \
+		static constexpr ::UE::Logging::Private::FStaticLocalizedLogRecord Log{CategoryName, TEXT(Namespace), TEXT(Key), TEXT(Format), __FILE__, __LINE__, ::ELogVerbosity::Verbosity, DynamicData}; \
+		::UE::Logging::Private::LogIfActive<Log, FLogCategory##CategoryName, ::ELogVerbosity::Verbosity, ::UE::Logging::Private::FLogFieldsNamed>(__VA_ARGS__); \
+	} \
+	while (false)
+
 #define UE_PRIVATE_SLOG_FIELDS(...) UE_PRIVATE_SLOG_CALL(PREPROCESSOR_JOIN(UE_PRIVATE_SLOG_FIELDS_, UE_PRIVATE_SLOG_COUNT(__VA_ARGS__)), (__VA_ARGS__))
 
 #define UE_PRIVATE_SLOG_FIELDS_0()
diff --git a/Engine/Source/Runtime/Launch/Private/LaunchEngineLoop.cpp b/Engine/Source/Runtime/Launch/Private/LaunchEngineLoop.cpp
--- a/Engine/Source/Runtime/Launch/Private/LaunchEngineLoop.cpp
+++ b/Engine/Source/Runtime/Launch/Private/LaunchEngineLoop.cpp
@@ -559,6 +559,14 @@ private:
 				Writer.AddString(ANSITEXTVIEW("$text"), ToString(Record.GetVerbosity()));
 				Writer.EndObject();
 			}
+			if (const TCHAR* TextNamespace = Record.GetTextNamespace())
+			{
+				Writer.AddString(ANSITEXTVIEW("_ns"), TextNamespace);
+			}
+			if (const TCHAR* TextKey = Record.GetTextKey())
+			{
+				Writer.AddString(ANSITEXTVIEW("_key"), TextKey);
+			}
 			for (const FCbField& Field : Record.GetFields())
 			{
 				Writer.AddField(Field.GetName(), Field);
