enabling REQUIRE_CHECK tests
supressed outputing the log line that Horde was flagging as an error

diff --git a/Engine/Source/Developer/LowLevelTestsRunner/Private/CheckScope.cpp b/Engine/Source/Developer/LowLevelTestsRunner/Private/CheckScope.cpp
--- a/Engine/Source/Developer/LowLevelTestsRunner/Private/CheckScope.cpp
+++ b/Engine/Source/Developer/LowLevelTestsRunner/Private/CheckScope.cpp
@@ -68,11 +68,13 @@ FCheckScope::FCheckScope(const ANSICHAR* Msg)
 #else
 	, bIgnoreDebugger(false)
 #endif
+	, bCriticalError(GIsCriticalError)
 {
 #if !UE_BUILD_SHIPPING
 	GIgnoreDebugger = true;
 #endif
 	GError = DeviceError;
+	GIsCriticalError = true; //set to true to disable printf of error message which causes Horde to flag the lines as an error
 }
 
 FCheckScope::FCheckScope()
@@ -86,6 +88,7 @@ FCheckScope::~FCheckScope()
 	GIgnoreDebugger = bIgnoreDebugger;
 #endif
 	GError = DeviceError->GetDeviceError();
+	GIsCriticalError = bCriticalError;
 	delete DeviceError;
 }
 
diff --git a/Engine/Source/Developer/LowLevelTestsRunner/Public/LowLevelTestsRunner/CheckScope.h b/Engine/Source/Developer/LowLevelTestsRunner/Public/LowLevelTestsRunner/CheckScope.h
--- a/Engine/Source/Developer/LowLevelTestsRunner/Public/LowLevelTestsRunner/CheckScope.h
+++ b/Engine/Source/Developer/LowLevelTestsRunner/Public/LowLevelTestsRunner/CheckScope.h
@@ -22,6 +22,7 @@ struct FCheckScope
 private:
 	FCheckScopeOutputDeviceError* DeviceError;
 	bool bIgnoreDebugger;
+	bool bCriticalError;
 }; 
 
 }
\ No newline at end of file
diff --git a/Engine/Source/Runtime/Core/Tests/Misc/AssertionMacrosTest.cpp b/Engine/Source/Runtime/Core/Tests/Misc/AssertionMacrosTest.cpp
--- a/Engine/Source/Runtime/Core/Tests/Misc/AssertionMacrosTest.cpp
+++ b/Engine/Source/Runtime/Core/Tests/Misc/AssertionMacrosTest.cpp
@@ -114,33 +114,11 @@ TEST_CASE("Core::Misc::TestEnsure", "[Core][Misc][AssertionMacros][Ensure]")
 	}
 }
 
-// Disabling check() tests for now. Assertion messages are output outside the regular log flow, and get flagged as legitimate errors on Horde 
-// due to non-deterministic order of output between the suspend/resume log parsing messages.
-#if 0
-// Disables log parsing by UAT for a scope, preventing assertion messages being tagged as errors.
-struct FSuspendLogParsing
-{
-	FSuspendLogParsing()
-	{
-		UE_LOG(LogCore, Display, TEXT("<-- Suspend Log Parsing -->"));
-		GLog->Flush();
-	}
-
-	~FSuspendLogParsing()
-	{
-		UE_LOG(LogCore, Display, TEXT("<-- Resume Log Parsing -->"));
-		GLog->Flush();
-	}
-};
-
 TEST_CASE("Core::Misc::TestCheck", "[Core][Misc][AssertionMacros][Check]")
 {
-	FSuspendLogParsing Scope;
-
 	REQUIRE_CHECK(check(1 == 2));
 	REQUIRE_CHECK_MSG("1 == 2", checkf(1 == 2, TEXT("Error Message")));
 	REQUIRE_CHECK_MSG("Error Message", checkf(1 == 2, TEXT("Error Message")));
 }
-#endif
 
 #endif
