Replaced complex metaprogramming from FName with TStringView overloads - reduces compiler impact.

diff --git a/Engine/Source/Runtime/Core/Public/UObject/NameTypes.h b/Engine/Source/Runtime/Core/Public/UObject/NameTypes.h
--- a/Engine/Source/Runtime/Core/Public/UObject/NameTypes.h
+++ b/Engine/Source/Runtime/Core/Public/UObject/NameTypes.h
@@ -977,16 +977,23 @@ public:
 	FName(int32 Len, const ANSICHAR* Name, EFindName FindType=FNAME_Add);
 	FName(int32 Len, const UTF8CHAR* Name, EFindName FindType=FNAME_Add);
 
-	template <typename CharRangeType,
-		typename CharType = typename TRemoveCV<typename TRemovePointer<decltype(GetData(DeclVal<CharRangeType>()))>::Type>::Type,
-		typename = decltype(ImplicitConv<TStringView<CharType>>(DeclVal<CharRangeType>()))>
-	inline explicit FName(CharRangeType&& Name, EFindName FindType = FNAME_Add)
+	inline explicit FName(TStringView<ANSICHAR> View, EFindName FindType = FNAME_Add)
+		: FName(NoInit)
+	{
+		*this = FName(View.Len(), View.GetData(), FindType);
+	}
+	inline explicit FName(TStringView<WIDECHAR> View, EFindName FindType = FNAME_Add)
+		: FName(NoInit)
+	{
+		*this = FName(View.Len(), View.GetData(), FindType);
+	}
+	inline explicit FName(TStringView<UTF8CHAR> View, EFindName FindType = FNAME_Add)
 		: FName(NoInit)
 	{
-		TStringView<CharType> View = Forward<CharRangeType>(Name);
 		*this = FName(View.Len(), View.GetData(), FindType);
 	}
 
+
 	/**
 	 * Create an FName. Will add the string to the name table if it does not exist.
 	 * When UE_FNAME_OUTLINE_NUMBER is set, will also add the combination of base string and number to the name table if it doesn't exist.
@@ -1013,26 +1020,44 @@ public:
 	FName(int32 Len, const ANSICHAR* Name, int32 Number);
 	FName(int32 Len, const UTF8CHAR* Name, int32 Number);
 
-	template <typename CharRangeType,
-		typename CharType = typename TRemoveCV<typename TRemovePointer<decltype(GetData(DeclVal<CharRangeType>()))>::Type>::Type,
-		typename = decltype(ImplicitConv<TStringView<CharType>>(DeclVal<CharRangeType>()))>
 	UE_DEPRECATED(5.1, "EFindName has been removed from constructors taking a Number argument to add clarity around UE_FNAME_OUTLINE_NUMBER.")
-	inline FName(CharRangeType&& Name, int32 InNumber, EFindName FindType)
+	inline FName(TStringView<ANSICHAR> View, int32 InNumber, EFindName FindType)
+		: FName(NoInit)
+	{
+PRAGMA_DISABLE_DEPRECATION_WARNINGS
+		*this = FName(View.Len(), View.GetData(), InNumber, FindType);
+PRAGMA_ENABLE_DEPRECATION_WARNINGS
+	}
+	UE_DEPRECATED(5.1, "EFindName has been removed from constructors taking a Number argument to add clarity around UE_FNAME_OUTLINE_NUMBER.")
+	inline FName(TStringView<WIDECHAR> View, int32 InNumber, EFindName FindType)
+		: FName(NoInit)
+	{
+PRAGMA_DISABLE_DEPRECATION_WARNINGS
+		*this = FName(View.Len(), View.GetData(), InNumber, FindType);
+PRAGMA_ENABLE_DEPRECATION_WARNINGS
+	}
+	UE_DEPRECATED(5.1, "EFindName has been removed from constructors taking a Number argument to add clarity around UE_FNAME_OUTLINE_NUMBER.")
+	inline FName(TStringView<UTF8CHAR> View, int32 InNumber, EFindName FindType)
 		: FName(NoInit)
 	{
 PRAGMA_DISABLE_DEPRECATION_WARNINGS
-		TStringView<CharType> View = Forward<CharRangeType>(Name);
 		*this = FName(View.Len(), View.GetData(), InNumber, FindType);
 PRAGMA_ENABLE_DEPRECATION_WARNINGS
 	}
 
-	template <typename CharRangeType,
-		typename CharType = typename TRemoveCV<typename TRemovePointer<decltype(GetData(DeclVal<CharRangeType>()))>::Type>::Type,
-		typename = decltype(ImplicitConv<TStringView<CharType>>(DeclVal<CharRangeType>()))>
-	inline FName(CharRangeType&& Name, int32 InNumber)
+	inline FName(TStringView<ANSICHAR> View, int32 InNumber)
+		: FName(NoInit)
+	{
+		*this = FName(View.Len(), View.GetData(), InNumber);
+	}
+	inline FName(TStringView<WIDECHAR> View, int32 InNumber)
+		: FName(NoInit)
+	{
+		*this = FName(View.Len(), View.GetData(), InNumber);
+	}
+	inline FName(TStringView<UTF8CHAR> View, int32 InNumber)
 		: FName(NoInit)
 	{
-		TStringView<CharType> View = Forward<CharRangeType>(Name);
 		*this = FName(View.Len(), View.GetData(), InNumber);
 	}
 
