Added traits for determining if a type is a TOptional, TUniquePtr, TUniqueObj, TSharedRef, TSharedPtr, and TWeakPtr.

diff --git a/Engine/Source/Runtime/Core/Public/Misc/Optional.h b/Engine/Source/Runtime/Core/Public/Misc/Optional.h
--- a/Engine/Source/Runtime/Core/Public/Misc/Optional.h
+++ b/Engine/Source/Runtime/Core/Public/Misc/Optional.h
@@ -223,3 +223,12 @@ FArchive& operator<<(FArchive& Ar, TOptional<OptionalType>& Optional)
 	Optional.Serialize(Ar);
 	return Ar;
 }
+
+/**
+ * Trait which determines whether or not a type is a TOptional.
+ */
+template <typename T> static constexpr bool TIsTOptional_V                              = false;
+template <typename T> static constexpr bool TIsTOptional_V<               TOptional<T>> = true;
+template <typename T> static constexpr bool TIsTOptional_V<const          TOptional<T>> = true;
+template <typename T> static constexpr bool TIsTOptional_V<      volatile TOptional<T>> = true;
+template <typename T> static constexpr bool TIsTOptional_V<const volatile TOptional<T>> = true;
diff --git a/Engine/Source/Runtime/Core/Public/Templates/SharedPointer.h b/Engine/Source/Runtime/Core/Public/Templates/SharedPointer.h
--- a/Engine/Source/Runtime/Core/Public/Templates/SharedPointer.h
+++ b/Engine/Source/Runtime/Core/Public/Templates/SharedPointer.h
@@ -643,6 +643,15 @@ private:
 	}
 };
 
+/**
+ * Trait which determines whether or not a type is a TSharedRef.
+ */
+template <typename T>                       constexpr bool TIsTSharedRef_V                                                = false;
+template <class ObjectType, ESPMode InMode> constexpr bool TIsTSharedRef_V<               TSharedRef<ObjectType, InMode>> = true;
+template <class ObjectType, ESPMode InMode> constexpr bool TIsTSharedRef_V<const          TSharedRef<ObjectType, InMode>> = true;
+template <class ObjectType, ESPMode InMode> constexpr bool TIsTSharedRef_V<      volatile TSharedRef<ObjectType, InMode>> = true;
+template <class ObjectType, ESPMode InMode> constexpr bool TIsTSharedRef_V<const volatile TSharedRef<ObjectType, InMode>> = true;
+
 
 /**
  * Wrapper for a type that yields a reference to that type.
@@ -1204,6 +1213,15 @@ private:
 	SharedPointerInternals::FSharedReferencer< Mode > SharedReferenceCount;
 };
 
+/**
+ * Trait which determines whether or not a type is a TSharedPtr.
+ */
+template <typename T>                       constexpr bool TIsTSharedPtr_V                                                = false;
+template <class ObjectType, ESPMode InMode> constexpr bool TIsTSharedPtr_V<               TSharedPtr<ObjectType, InMode>> = true;
+template <class ObjectType, ESPMode InMode> constexpr bool TIsTSharedPtr_V<const          TSharedPtr<ObjectType, InMode>> = true;
+template <class ObjectType, ESPMode InMode> constexpr bool TIsTSharedPtr_V<      volatile TSharedPtr<ObjectType, InMode>> = true;
+template <class ObjectType, ESPMode InMode> constexpr bool TIsTSharedPtr_V<const volatile TSharedPtr<ObjectType, InMode>> = true;
+
 namespace Freeze
 {
 	template<class ObjectType, ESPMode Mode>
@@ -1511,6 +1529,15 @@ private:
 	SharedPointerInternals::FWeakReferencer< Mode > WeakReferenceCount;
 };
 
+/**
+ * Trait which determines whether or not a type is a TWeakPtr.
+ */
+template <typename T>                       constexpr bool TIsTWeakPtr_V                                              = false;
+template <class ObjectType, ESPMode InMode> constexpr bool TIsTWeakPtr_V<               TWeakPtr<ObjectType, InMode>> = true;
+template <class ObjectType, ESPMode InMode> constexpr bool TIsTWeakPtr_V<const          TWeakPtr<ObjectType, InMode>> = true;
+template <class ObjectType, ESPMode InMode> constexpr bool TIsTWeakPtr_V<      volatile TWeakPtr<ObjectType, InMode>> = true;
+template <class ObjectType, ESPMode InMode> constexpr bool TIsTWeakPtr_V<const volatile TWeakPtr<ObjectType, InMode>> = true;
+
 
 template<class T, ESPMode Mode> struct TIsWeakPointerType<TWeakPtr<T, Mode> > { enum { Value = true }; };
 template<class T, ESPMode Mode> struct TIsZeroConstructType<TWeakPtr<T, Mode> > { enum { Value = true }; };
diff --git a/Engine/Source/Runtime/Core/Public/Templates/UniqueObj.h b/Engine/Source/Runtime/Core/Public/Templates/UniqueObj.h
--- a/Engine/Source/Runtime/Core/Public/Templates/UniqueObj.h
+++ b/Engine/Source/Runtime/Core/Public/Templates/UniqueObj.h
@@ -73,3 +73,12 @@ FArchive& operator<<(FArchive& Ar, TUniqueObj<T>& P)
 	P.Serialize(Ar);
 	return Ar;
 }
+
+/**
+ * Trait which determines whether or not a type is a TUniqueObj.
+ */
+template <typename T> constexpr bool TIsTUniqueObj_V                               = false;
+template <typename T> constexpr bool TIsTUniqueObj_V<               TUniqueObj<T>> = true;
+template <typename T> constexpr bool TIsTUniqueObj_V<const          TUniqueObj<T>> = true;
+template <typename T> constexpr bool TIsTUniqueObj_V<      volatile TUniqueObj<T>> = true;
+template <typename T> constexpr bool TIsTUniqueObj_V<const volatile TUniqueObj<T>> = true;
diff --git a/Engine/Source/Runtime/Core/Public/Templates/UniquePtr.h b/Engine/Source/Runtime/Core/Public/Templates/UniquePtr.h
--- a/Engine/Source/Runtime/Core/Public/Templates/UniquePtr.h
+++ b/Engine/Source/Runtime/Core/Public/Templates/UniquePtr.h
@@ -853,3 +853,12 @@ FORCEINLINE typename TEnableIf<TIsUnboundedArray<T>::Value, TUniquePtr<T>>::Type
  */
 template <typename T, typename... TArgs>
 typename TEnableIf<TIsBoundedArray<T>::Value, TUniquePtr<T>>::Type MakeUnique(TArgs&&... Args) = delete;
+
+/**
+ * Trait which determines whether or not a type is a TUniquePtr.
+ */
+template <typename T>                   constexpr bool TIsTUniquePtr_V                                        = false;
+template <typename T, typename Deleter> constexpr bool TIsTUniquePtr_V<               TUniquePtr<T, Deleter>> = true;
+template <typename T, typename Deleter> constexpr bool TIsTUniquePtr_V<const          TUniquePtr<T, Deleter>> = true;
+template <typename T, typename Deleter> constexpr bool TIsTUniquePtr_V<      volatile TUniquePtr<T, Deleter>> = true;
+template <typename T, typename Deleter> constexpr bool TIsTUniquePtr_V<const volatile TUniquePtr<T, Deleter>> = true;
