Logging: Added $text field handling to structured log formatting

If a field value is an object containing a $text field, the string value of that field is used when formatting the field.

#rnx

diff --git a/Engine/Source/Runtime/Core/Private/Logging/StructuredLog.cpp b/Engine/Source/Runtime/Core/Private/Logging/StructuredLog.cpp
--- a/Engine/Source/Runtime/Core/Private/Logging/StructuredLog.cpp
+++ b/Engine/Source/Runtime/Core/Private/Logging/StructuredLog.cpp
@@ -92,6 +92,12 @@ static void LogFieldValue(TStringBuilderBase<CharType>& Out, const FCbFieldView&
 		break;
 	case ECbFieldType::Object:
 	case ECbFieldType::UniformObject:
+		if (FCbFieldView TextField = Accessor.AsObjectView().FindViewIgnoreCase(ANSITEXTVIEW("$text")))
+		{
+			Out.Append(TextField.AsString());
+			break;
+		}
+		[[fallthrough]];
 	case ECbFieldType::Array:
 	case ECbFieldType::UniformArray:
 	case ECbFieldType::Binary:
@@ -158,18 +164,9 @@ static void LogFieldValue(TStringBuilderBase<CharType>& Out, const FCbFieldView&
 static void AddFieldValue(FFormatNamedArguments& Out, const FCbFieldView& Field)
 {
 	const FString FieldName(Field.GetName());
+
 	switch (FCbValue Accessor = Field.GetValue(); Accessor.GetType())
 	{
-	case ECbFieldType::Null:
-		break;
-	case ECbFieldType::Object:
-	case ECbFieldType::UniformObject:
-		break;
-	case ECbFieldType::Array:
-	case ECbFieldType::UniformArray:
-	case ECbFieldType::Binary:
-	case ECbFieldType::String:
-		break;
 	case ECbFieldType::IntegerPositive:
 		Out.Emplace(FieldName, Accessor.AsIntegerPositive());
 		return;
@@ -182,20 +179,7 @@ static void AddFieldValue(FFormatNamedArguments& Out, const FCbFieldView& Field)
 	case ECbFieldType::Float64:
 		Out.Emplace(FieldName, Accessor.AsFloat64());
 		return;
-	case ECbFieldType::BoolFalse:
-	case ECbFieldType::BoolTrue:
-	case ECbFieldType::ObjectAttachment:
-	case ECbFieldType::BinaryAttachment:
-	case ECbFieldType::Hash:
-	case ECbFieldType::Uuid:
-	case ECbFieldType::DateTime:
-	case ECbFieldType::TimeSpan:
-	case ECbFieldType::ObjectId:
-	case ECbFieldType::CustomById:
-	case ECbFieldType::CustomByName:
-		break;
 	default:
-		checkNoEntry();
 		break;
 	}
 
