[Insights]
  - MemoryTrace: Added CallstackId for all free events and also for HeapMarkAlloc, HeapUnmarkAlloc events.
  - MemoryTrace: Increased trace version to 2.

diff --git a/Engine/Source/Runtime/Core/Private/ProfilingDebugging/MemoryTrace.cpp b/Engine/Source/Runtime/Core/Private/ProfilingDebugging/MemoryTrace.cpp
--- a/Engine/Source/Runtime/Core/Private/ProfilingDebugging/MemoryTrace.cpp
+++ b/Engine/Source/Runtime/Core/Private/ProfilingDebugging/MemoryTrace.cpp
@@ -25,10 +25,10 @@ void MemoryTrace_EnableTracePump();
 ////////////////////////////////////////////////////////////////////////////////
 namespace
 {
-	// Controls how often time markers are emitted (default every 4095 allocation)
+	// Controls how often time markers are emitted (default: every 4095 allocations).
 	constexpr uint32 MarkerSamplePeriod = (4 << 10) - 1;
 
-	// Number of bits shifted bits to SizeLower
+	// Number of shifted bits to SizeLower
 	constexpr uint32 SizeShift = 3;
 
 	// Counter to track when time marker is emitted
@@ -40,10 +40,10 @@ namespace
 
 	// Temporarily disables any internal operation that causes allocations. Used to
 	// avoid recursive behaviour when memory tracing needs to allocate memory through
-	// TraceMalloc
+	// TraceMalloc.
 	thread_local bool GDoNotAllocateInTrace;
 
-	// Set on initialization, on some platforms we hook allocator functions very early
+	// Set on initialization; on some platforms we hook allocator functions very early
 	// before Trace has the ability to allocate memory.
 	bool GTraceAllowed;
 }
@@ -56,13 +56,12 @@ namespace Trace {
 } // namespace UE
 
 ////////////////////////////////////////////////////////////////////////////////
-
 UE_TRACE_EVENT_BEGIN(Memory, Init, NoSync|Important)
 	UE_TRACE_EVENT_FIELD(uint32, MarkerPeriod)
 	UE_TRACE_EVENT_FIELD(uint8, Version)
 	UE_TRACE_EVENT_FIELD(uint8, MinAlignment)
 	UE_TRACE_EVENT_FIELD(uint8, SizeShift)
-	UE_TRACE_EVENT_FIELD(uint8, Mode)
+	UE_TRACE_EVENT_FIELD(uint8, Mode) // unused
 UE_TRACE_EVENT_END()
 
 UE_TRACE_EVENT_BEGIN(Memory, Marker)
@@ -124,11 +123,13 @@ UE_TRACE_EVENT_END()
 
 UE_TRACE_EVENT_BEGIN(Memory, ReallocFree)
 	UE_TRACE_EVENT_FIELD(uint64, Address)
+	UE_TRACE_EVENT_FIELD(uint32, CallstackId)
 	UE_TRACE_EVENT_FIELD(uint8, RootHeap)
 UE_TRACE_EVENT_END()
 
 UE_TRACE_EVENT_BEGIN(Memory, ReallocFreeSystem)
 	UE_TRACE_EVENT_FIELD(uint64, Address)
+	UE_TRACE_EVENT_FIELD(uint32, CallstackId)
 UE_TRACE_EVENT_END()
 
 UE_TRACE_EVENT_BEGIN(Memory, HeapSpec, NoSync|Important)
@@ -140,17 +141,21 @@ UE_TRACE_EVENT_END()
 
 UE_TRACE_EVENT_BEGIN(Memory, HeapMarkAlloc)
 	UE_TRACE_EVENT_FIELD(uint64, Address)
+	UE_TRACE_EVENT_FIELD(uint32, CallstackId)
 	UE_TRACE_EVENT_FIELD(uint16, Flags)
 	UE_TRACE_EVENT_FIELD(HeapId, Heap)
 UE_TRACE_EVENT_END()
 
 UE_TRACE_EVENT_BEGIN(Memory, HeapUnmarkAlloc)
 	UE_TRACE_EVENT_FIELD(uint64, Address)
+	UE_TRACE_EVENT_FIELD(uint32, CallstackId)
 	UE_TRACE_EVENT_FIELD(HeapId, Heap)
 UE_TRACE_EVENT_END()
 
-// If layout of the above events are changed, bump this version number
-constexpr uint8 MemoryTraceVersion = 1;
+// If the layout of the above events is changed, bump this version number.
+// version 1: Initial version (UE 5.0, UE 5.1)
+// version 2: Added CallstackId for Free events and also for HeapMarkAlloc, HeapUnmarkAlloc events (UE 5.2).
+constexpr uint8 MemoryTraceVersion = 2;
 
 ////////////////////////////////////////////////////////////////////////////////
 class FMallocWrapper
@@ -437,8 +442,8 @@ void MemoryTrace_Alloc(uint64 Address, uint64 Size, uint32 Alignment, HeapId Roo
 		case EMemoryTraceRootHeap::SystemMemory:
 		{
 			UE_TRACE_LOG(Memory, AllocSystem, MemAllocChannel)
-				<< AllocSystem.CallstackId(CallstackId)
 				<< AllocSystem.Address(uint64(Address))
+				<< AllocSystem.CallstackId(CallstackId)
 				<< AllocSystem.Size(uint32(Size >> SizeShift))
 				<< AllocSystem.AlignmentPow2_SizeLower(uint8(Alignment_SizeLower));
 			break;
@@ -447,8 +452,8 @@ void MemoryTrace_Alloc(uint64 Address, uint64 Size, uint32 Alignment, HeapId Roo
 		case EMemoryTraceRootHeap::VideoMemory:
 		{
 			UE_TRACE_LOG(Memory, AllocVideo, MemAllocChannel)
-				<< AllocVideo.CallstackId(CallstackId)
 				<< AllocVideo.Address(uint64(Address))
+				<< AllocVideo.CallstackId(CallstackId)
 				<< AllocVideo.Size(uint32(Size >> SizeShift))
 				<< AllocVideo.AlignmentPow2_SizeLower(uint8(Alignment_SizeLower));
 			break;
@@ -457,11 +462,11 @@ void MemoryTrace_Alloc(uint64 Address, uint64 Size, uint32 Alignment, HeapId Roo
 		default:
 		{
 			UE_TRACE_LOG(Memory, Alloc, MemAllocChannel)
-				<< Alloc.CallstackId(CallstackId)
 				<< Alloc.Address(uint64(Address))
-				<< Alloc.RootHeap(uint8(RootHeap))
+				<< Alloc.CallstackId(CallstackId)
 				<< Alloc.Size(uint32(Size >> SizeShift))
-				<< Alloc.AlignmentPow2_SizeLower(uint8(Alignment_SizeLower));
+				<< Alloc.AlignmentPow2_SizeLower(uint8(Alignment_SizeLower))
+				<< Alloc.RootHeap(uint8(RootHeap));
 			break;
 		}
 	}
@@ -529,8 +534,8 @@ void MemoryTrace_ReallocAlloc(uint64 Address, uint64 Size, uint32 Alignment, Hea
 		case EMemoryTraceRootHeap::SystemMemory:
 		{
 			UE_TRACE_LOG(Memory, ReallocAllocSystem, MemAllocChannel)
-				<< ReallocAllocSystem.CallstackId(CallstackId)
 				<< ReallocAllocSystem.Address(uint64(Address))
+				<< ReallocAllocSystem.CallstackId(CallstackId)
 				<< ReallocAllocSystem.Size(uint32(Size >> SizeShift))
 				<< ReallocAllocSystem.AlignmentPow2_SizeLower(uint8(Alignment_SizeLower));
 			break;
@@ -539,11 +544,11 @@ void MemoryTrace_ReallocAlloc(uint64 Address, uint64 Size, uint32 Alignment, Hea
 		default:
 		{
 			UE_TRACE_LOG(Memory, ReallocAlloc, MemAllocChannel)
-				<< ReallocAlloc.CallstackId(CallstackId)
 				<< ReallocAlloc.Address(uint64(Address))
-				<< ReallocAlloc.RootHeap(uint8(RootHeap))
+				<< ReallocAlloc.CallstackId(CallstackId)
 				<< ReallocAlloc.Size(uint32(Size >> SizeShift))
-				<< ReallocAlloc.AlignmentPow2_SizeLower(uint8(Alignment_SizeLower));
+				<< ReallocAlloc.AlignmentPow2_SizeLower(uint8(Alignment_SizeLower))
+				<< ReallocAlloc.RootHeap(uint8(RootHeap));
 			break;
 		}
 	}
@@ -561,12 +566,15 @@ void MemoryTrace_ReallocFree(uint64 Address, HeapId RootHeap)
 
 	check(RootHeap < 16);
 
+	const uint32 CallstackId = GDoNotAllocateInTrace ? 0 : CallstackTrace_GetCurrentId();
+
 	switch (RootHeap)
 	{
 		case EMemoryTraceRootHeap::SystemMemory:
 		{
 			UE_TRACE_LOG(Memory, ReallocFreeSystem, MemAllocChannel)
-				<< ReallocFreeSystem.Address(uint64(Address));
+				<< ReallocFreeSystem.Address(uint64(Address))
+				<< ReallocFreeSystem.CallstackId(CallstackId);
 			break;
 		}
 
@@ -574,6 +582,7 @@ void MemoryTrace_ReallocFree(uint64 Address, HeapId RootHeap)
 		{
 			UE_TRACE_LOG(Memory, ReallocFree, MemAllocChannel)
 				<< ReallocFree.Address(uint64(Address))
+				<< ReallocFree.CallstackId(CallstackId)
 				<< ReallocFree.RootHeap(uint8(RootHeap));
 			break;
 		}
@@ -637,10 +646,13 @@ void MemoryTrace_MarkAllocAsHeap(uint64 Address, HeapId Heap, EMemoryTraceHeapAl
 		return;
 	}
 
+	const uint32 CallstackId = GDoNotAllocateInTrace ? 0 : CallstackTrace_GetCurrentId();
+
 	UE_TRACE_LOG(Memory, HeapMarkAlloc, MemAllocChannel)
 		<< HeapMarkAlloc.Address(uint64(Address))
-		<< HeapMarkAlloc.Heap(Heap)
-		<< HeapMarkAlloc.Flags(uint16(EMemoryTraceHeapAllocationFlags::Heap | Flags));
+		<< HeapMarkAlloc.CallstackId(CallstackId)
+		<< HeapMarkAlloc.Flags(uint16(EMemoryTraceHeapAllocationFlags::Heap | Flags))
+		<< HeapMarkAlloc.Heap(Heap);
 }
 
 ////////////////////////////////////////////////////////////////////////////////
@@ -651,9 +663,12 @@ void MemoryTrace_UnmarkAllocAsHeap(uint64 Address, HeapId Heap)
 		return;
 	}
 
+	const uint32 CallstackId = GDoNotAllocateInTrace ? 0 : CallstackTrace_GetCurrentId();
+
 	// Sets all flags to zero
 	UE_TRACE_LOG(Memory, HeapUnmarkAlloc, MemAllocChannel)
 		<< HeapUnmarkAlloc.Address(uint64(Address))
+		<< HeapUnmarkAlloc.CallstackId(CallstackId)
 		<< HeapUnmarkAlloc.Heap(Heap);
 }
 
@@ -667,7 +682,6 @@ bool MemoryTrace_IsActive()
 
 #endif // UE_MEMORY_TRACE_ENABLED
 
-
 /////////////////////////////////////////////////////////////////////////////
 FTraceMalloc::FTraceMalloc(FMalloc* InMalloc)
 {
@@ -697,11 +711,11 @@ void* FTraceMalloc::Malloc(SIZE_T Count, uint32 Alignment)
 	UE_MEMSCOPE(TRACE_TAG);
 
 	UE_TRACE_LOG(Memory, Alloc, MemAllocChannel)
-		<< Alloc.CallstackId(0)
 		<< Alloc.Address(uint64(NewPtr))
-		<< Alloc.RootHeap(uint8(EMemoryTraceRootHeap::SystemMemory))
+		<< Alloc.CallstackId(0)
 		<< Alloc.Size(uint32(Size >> SizeShift))
-		<< Alloc.AlignmentPow2_SizeLower(uint8(Alignment_SizeLower));
+		<< Alloc.AlignmentPow2_SizeLower(uint8(Alignment_SizeLower))
+		<< Alloc.RootHeap(uint8(EMemoryTraceRootHeap::SystemMemory));
 
 	return NewPtr;
 #else
@@ -731,11 +745,11 @@ void* FTraceMalloc::Realloc(void* Original, SIZE_T Count, uint32 Alignment)
 	}
 
 	UE_TRACE_LOG(Memory, ReallocAlloc, MemAllocChannel)
-		<< ReallocAlloc.CallstackId(0)
 		<< ReallocAlloc.Address(uint64(NewPtr))
-		<< ReallocAlloc.RootHeap(uint8(EMemoryTraceRootHeap::SystemMemory))
+		<< ReallocAlloc.CallstackId(0)
 		<< ReallocAlloc.Size(uint32(Size >> SizeShift))
-		<< ReallocAlloc.AlignmentPow2_SizeLower(uint8(Alignment_SizeLower));
+		<< ReallocAlloc.AlignmentPow2_SizeLower(uint8(Alignment_SizeLower))
+		<< ReallocAlloc.RootHeap(uint8(EMemoryTraceRootHeap::SystemMemory));
 
 	return NewPtr;
 #else
diff --git a/Engine/Source/Runtime/Core/Public/ProfilingDebugging/MemoryTrace.h b/Engine/Source/Runtime/Core/Public/ProfilingDebugging/MemoryTrace.h
--- a/Engine/Source/Runtime/Core/Public/ProfilingDebugging/MemoryTrace.h
+++ b/Engine/Source/Runtime/Core/Public/ProfilingDebugging/MemoryTrace.h
@@ -7,7 +7,7 @@
 #include "Trace/Trace.h"
 
 #if !defined(UE_MEMORY_TRACE_AVAILABLE)
-#   define UE_MEMORY_TRACE_AVAILABLE 0
+#	define UE_MEMORY_TRACE_AVAILABLE 0
 #endif
 
 #if !defined(UE_MEMORY_TRACE_LATE_INIT)
@@ -59,9 +59,11 @@ enum class EMemoryTraceHeapAllocationFlags : uint8
 };
 ENUM_CLASS_FLAGS(EMemoryTraceHeapAllocationFlags);
 
+////////////////////////////////////////////////////////////////////////////////
 #if UE_MEMORY_TRACE_ENABLED
 
-////////////////////////////////////////////////////////////////////////////////
+#define UE_MEMORY_TRACE(x) x
+
 CORE_API UE_TRACE_CHANNEL_EXTERN(MemAllocChannel);
 
 ////////////////////////////////////////////////////////////////////////////////
@@ -132,8 +134,10 @@ CORE_API void MemoryTrace_ReallocFree(uint64 Address, HeapId RootHeap = EMemoryT
  */
 CORE_API void MemoryTrace_ReallocAlloc(uint64 Address, uint64 NewSize, uint32 Alignment, HeapId RootHeap = EMemoryTraceRootHeap::SystemMemory);
 
-#else
+////////////////////////////////////////////////////////////////////////////////
+#else // UE_MEMORY_TRACE_ENABLED
 
+#define UE_MEMORY_TRACE(x)
 inline HeapId MemoryTrace_RootHeapSpec(const TCHAR* Name, EMemoryTraceHeapFlags Flags = EMemoryTraceHeapFlags::None) { return ~0; };
 inline HeapId MemoryTrace_HeapSpec(HeapId ParentId, const TCHAR* Name, EMemoryTraceHeapFlags Flags = EMemoryTraceHeapFlags::None) { return ~0; }
 inline void MemoryTrace_MarkAllocAsHeap(uint64 Address, HeapId Heap) {}
