#ue5 - Remove extra NaN checks in VectorMatrixMultiply() and VectorMatrixInverse(), as they are already covered by DiagnosticCheckNaN in Matrix operators in a way that affects fewer build configurations by default.

diff --git a/Engine/Source/Runtime/Core/Private/Math/UnrealMathSSE.cpp b/Engine/Source/Runtime/Core/Private/Math/UnrealMathSSE.cpp
--- a/Engine/Source/Runtime/Core/Private/Math/UnrealMathSSE.cpp
+++ b/Engine/Source/Runtime/Core/Private/Math/UnrealMathSSE.cpp
@@ -95,11 +95,6 @@ void VectorMatrixMultiply(FMatrix44f* Result, const FMatrix44f* Matrix1, const F
 	R[1] = R1;
 	R[2] = R2;
 	R[3] = Temp;
-
-	// TODO: Should change to DiagnosticCheckNaN call, though that's only in Debug.
-#if !(UE_BUILD_SHIPPING || UE_BUILD_TEST)
-	ensure(!Result->ContainsNaN());
-#endif
 }
 
 
@@ -159,11 +154,6 @@ void VectorMatrixMultiply(FMatrix44d* Result, const FMatrix44d* Matrix1, const F
 	VectorStore(R1, ResultDst[1]);
 	VectorStore(R2, ResultDst[2]);
 	VectorStore(R3, ResultDst[3]);
-
-	// TODO: Should change to DiagnosticCheckNaN call, though that's only in Debug.
-#if !(UE_BUILD_SHIPPING || UE_BUILD_TEST)
-	ensure(!Result->ContainsNaN());
-#endif
 }
 
 
@@ -255,11 +245,6 @@ void VectorMatrixInverse(FMatrix44f* DstMatrix, const FMatrix44f* SrcMatrix)
 		);
 
 	memcpy(DstMatrix, &Result, sizeof(Result));
-
-	// TODO: Should change to DiagnosticCheckNaN call, though that's only in Debug.
-#if !(UE_BUILD_SHIPPING || UE_BUILD_TEST)
-	ensure(!DstMatrix->ContainsNaN());
-#endif
 }
 
 void VectorMatrixInverse(FMatrix44d* DstMatrix, const FMatrix44d* SrcMatrix)
@@ -344,11 +329,6 @@ void VectorMatrixInverse(FMatrix44d* DstMatrix, const FMatrix44d* SrcMatrix)
 		);
 
 	memcpy(DstMatrix, &Result, sizeof(Result));
-
-	// TODO: Should change to DiagnosticCheckNaN call, though that's only in Debug.
-#if !(UE_BUILD_SHIPPING || UE_BUILD_TEST)
-	ensure(!DstMatrix->ContainsNaN());
-#endif
 }
 
 #endif // UE_USING_UNREALMATH_SSE
