Add ApplicationSystemUIOverlayAppearedDelegate and ApplicationSystemUIOverlayDisappearedDelegate to FCoreDelegates

diff --git a/Engine/Source/Runtime/Core/Private/Misc/CoreDelegates.cpp b/Engine/Source/Runtime/Core/Private/Misc/CoreDelegates.cpp
--- a/Engine/Source/Runtime/Core/Private/Misc/CoreDelegates.cpp
+++ b/Engine/Source/Runtime/Core/Private/Misc/CoreDelegates.cpp
@@ -110,6 +110,9 @@ TMulticastDelegate<void()> FCoreDelegates::ApplicationWillTerminateDelegate;
 TMulticastDelegate<void()> FCoreDelegates::ApplicationShouldUnloadResourcesDelegate;
 TMulticastDelegate<void(float)> FCoreDelegates::MobileBackgroundTickDelegate;
 
+TMulticastDelegate<void()> FCoreDelegates::ApplicationSystemUIOverlayAppearedDelegate;
+TMulticastDelegate<void()> FCoreDelegates::ApplicationSystemUIOverlayDisappearedDelegate;
+
 TMulticastDelegate<void(const TArray<FString>&)> FCoreDelegates::ApplicationReceivedStartupArgumentsDelegate;
 
 TMulticastDelegate<void(bool)> FCoreDelegates::UserMusicInterruptDelegate;
diff --git a/Engine/Source/Runtime/Core/Public/Misc/CoreDelegates.h b/Engine/Source/Runtime/Core/Public/Misc/CoreDelegates.h
--- a/Engine/Source/Runtime/Core/Public/Misc/CoreDelegates.h
+++ b/Engine/Source/Runtime/Core/Public/Misc/CoreDelegates.h
@@ -318,6 +318,13 @@ public:
 	// save state when ApplicationWillEnterBackgroundDelegate is called instead.
 	static TMulticastDelegate<void()> ApplicationWillTerminateDelegate;
 
+	// Some platform have a System UI Overlay that can draw on top of the application.
+	// The game might want to be notified so it can pause, etc...
+	static TMulticastDelegate<void()> ApplicationSystemUIOverlayAppearedDelegate;
+
+	// Called when the UI Overlay was hidden
+	static TMulticastDelegate<void()> ApplicationSystemUIOverlayDisappearedDelegate;
+
 	// Called when in the background, if the OS is giving CPU time to the device. It is very likely
 	// this will never be called due to mobile OS backgrounded CPU restrictions. But if, for instance,
 	// VOIP is active on iOS, the will be getting called
