Added UE_CONSTRAINT* for defining template constraint blocks in a forward-compatible-with-concepts way.

diff --git a/Engine/Source/Runtime/Core/Public/Templates/UnrealTypeTraits.h b/Engine/Source/Runtime/Core/Public/Templates/UnrealTypeTraits.h
--- a/Engine/Source/Runtime/Core/Public/Templates/UnrealTypeTraits.h
+++ b/Engine/Source/Runtime/Core/Public/Templates/UnrealTypeTraits.h
@@ -23,12 +23,37 @@
 #include "Templates/IsTriviallyCopyConstructible.h"
 
 /*-----------------------------------------------------------------------------
-	Readability macro for enable_if in template definitions. Usage:
-
-	template<typename T, TEMPLATE_REQUIRES(TIsFloatingPoint<T>)>
-	void FloatingPointOnlyPlease(T In) {}
+	Readability macro for enable_if in template definitions, future-proofed
+	for C++ 20 concepts. Usage:
+
+	template <
+		typename T,
+		typename U  // note - no trailing comma before the constraints block
+		UE_CONSTRAINTS_BEGIN
+			UE_CONSTRAINT(std::is_integral_v<T>)
+			UE_CONSTRAINT(sizeof(U) <= 4)
+		UE_CONSTRAINTS_END
+	>
+	void IntegralUpTo32Bit(T Lhs, U Rhs) {}
  -----------------------------------------------------------------------------*/
-#define TEMPLATE_REQUIRES(...) typename TEnableIf<__VA_ARGS__, int>::type = 0
+#if __cplusplus < 202000
+    #define UE_CONSTRAINTS_BEGIN , std::enable_if_t<
+    #define UE_CONSTRAINT(...) (__VA_ARGS__) &&
+    #define UE_CONSTRAINTS_END true, int> = 0
+#else
+    namespace UE::Core::Private
+    {
+        // Only needed for the UE_CONSTRAINT* macros to work
+        template <bool B>
+        concept BoolIdentityConcept = B;
+    }
+
+    #define UE_CONSTRAINTS_BEGIN > requires
+    #define UE_CONSTRAINT(...) (__VA_ARGS__) &&
+    #define UE_CONSTRAINTS_END UE::Core::Private::BoolIdentityConcept<true
+#endif
+
+#define TEMPLATE_REQUIRES(...) /*UE_DEPRECATED_MACRO(5.2, "TEMPLATE_REQUIRES has been deprecated, please use UE_CONSTRAINT instead.")*/ std::enable_if_t<__VA_ARGS__, int> = 0
 
 /*-----------------------------------------------------------------------------
  * Macros to abstract the presence of certain compiler intrinsic type traits 
