ParallelFor IoStoreWriter Finalize() calls. Previously there were very few iostore writers as there weren't many containers, but with content on demand there are 4k+, so we need to multi thread this. 130s -> 12s for pak with my testing.

diff --git a/Engine/Source/Runtime/Core/Private/IO/IoStore.cpp b/Engine/Source/Runtime/Core/Private/IO/IoStore.cpp
--- a/Engine/Source/Runtime/Core/Private/IO/IoStore.cpp
+++ b/Engine/Source/Runtime/Core/Private/IO/IoStore.cpp
@@ -1836,10 +1836,14 @@ void FIoStoreWriterContextImpl::Flush()
 	}
 
 	AllEntries.Empty();
-	for (TSharedPtr<FIoStoreWriter> IoStoreWriter : IoStoreWriters)
-	{
-		IoStoreWriter->Finalize();
-	}
+
+	// Classically there were so few writers that this didn't need to be multi threaded, but it
+	// involves writing files, and with content on demand this ends up being thousands of iterations. 
+	// 10x speedup in testing.
+	double FinalizeStart = FPlatformTime::Seconds();	
+	ParallelFor(TEXT("IoStoreWriter::Finalize.PF"), IoStoreWriters.Num(), 1, [this](int Index){ IoStoreWriters[Index]->Finalize(); });
+	double FinalizeEnd = FPlatformTime::Seconds();
+	UE_LOG(LogIoStore, Display, TEXT("Finalize took %.1f seconds for %d writers"), FinalizeEnd - FinalizeStart, IoStoreWriters.Num());
 }
 
 void FIoStoreWriterContextImpl::BeginCompressionThreadFunc()
