First pass on a TTuple simplification.

diff --git a/Engine/Source/Runtime/Core/Public/Templates/Tuple.h b/Engine/Source/Runtime/Core/Public/Templates/Tuple.h
--- a/Engine/Source/Runtime/Core/Public/Templates/Tuple.h
+++ b/Engine/Source/Runtime/Core/Public/Templates/Tuple.h
@@ -61,26 +61,19 @@ namespace UE::Core::Private::Tuple
 	void ConceptCheckingHelper(ArgTypes&&...);
 
 	template <typename T, typename... Types>
-	struct TTypeCountInParameterPack;
-
-	template <typename T>
-	struct TTypeCountInParameterPack<T>
-	{
-		constexpr static uint32 Value = 0;
-	};
+	constexpr uint32 TTypeCountInParameterPack_V = 0;
 
 	template <typename T, typename U, typename... Types>
-	struct TTypeCountInParameterPack<T, U, Types...>
-	{
-		constexpr static uint32 Value = TTypeCountInParameterPack<T, Types...>::Value + (std::is_same<T, U>::value ? 1 : 0);
-	};
+	constexpr uint32 TTypeCountInParameterPack_V<T, U, Types...> = TTypeCountInParameterPack_V<T, Types...> + (std::is_same_v<T, U> ? 1 : 0);
 
 	template <typename T, uint32 Index, uint32 TupleSize>
 	struct TTupleBaseElement
 	{
 		template <
-			typename ArgType,
-			typename TEnableIf<TIsConstructible<T, ArgType&&>::Value>::Type* = nullptr
+			typename ArgType
+			UE_CONSTRAINTS_BEGIN
+				UE_CONSTRAINT(std::is_constructible_v<T, ArgType&&>)
+			UE_CONSTRAINTS_END
 		>
 		explicit TTupleBaseElement(EForwardingConstructor, ArgType&& Arg)
 			: Value(Forward<ArgType>(Arg))
@@ -104,8 +97,10 @@ namespace UE::Core::Private::Tuple
 	struct TTupleBaseElement<T, 0, 2>
 	{
 		template <
-			typename ArgType,
-			typename TEnableIf<TIsConstructible<T, ArgType&&>::Value>::Type* = nullptr
+			typename ArgType
+			UE_CONSTRAINTS_BEGIN
+				UE_CONSTRAINT(std::is_constructible_v<T, ArgType&&>)
+			UE_CONSTRAINTS_END
 		>
 		explicit TTupleBaseElement(EForwardingConstructor, ArgType&& Arg)
 			: Key(Forward<ArgType>(Arg))
@@ -304,23 +299,23 @@ namespace UE::Core::Private::Tuple
 		TTupleBase& operator=(TTupleBase&& Other) = default;
 		TTupleBase& operator=(const TTupleBase& Other) = default;
 
-		template <uint32 Index, typename TEnableIf<(Index < sizeof...(Types))>::Type* = nullptr> FORCEINLINE decltype(auto) Get()               &  { return TTupleElementGetterByIndex<Index, sizeof...(Types)>::Get(static_cast<               TTupleBase& >(*this)); }
-		template <uint32 Index, typename TEnableIf<(Index < sizeof...(Types))>::Type* = nullptr> FORCEINLINE decltype(auto) Get() const         &  { return TTupleElementGetterByIndex<Index, sizeof...(Types)>::Get(static_cast<const          TTupleBase& >(*this)); }
-		template <uint32 Index, typename TEnableIf<(Index < sizeof...(Types))>::Type* = nullptr> FORCEINLINE decltype(auto) Get()       volatile&  { return TTupleElementGetterByIndex<Index, sizeof...(Types)>::Get(static_cast<      volatile TTupleBase& >(*this)); }
-		template <uint32 Index, typename TEnableIf<(Index < sizeof...(Types))>::Type* = nullptr> FORCEINLINE decltype(auto) Get() const volatile&  { return TTupleElementGetterByIndex<Index, sizeof...(Types)>::Get(static_cast<const volatile TTupleBase& >(*this)); }
-		template <uint32 Index, typename TEnableIf<(Index < sizeof...(Types))>::Type* = nullptr> FORCEINLINE decltype(auto) Get()               && { return TTupleElementGetterByIndex<Index, sizeof...(Types)>::Get(static_cast<               TTupleBase&&>(*this)); }
-		template <uint32 Index, typename TEnableIf<(Index < sizeof...(Types))>::Type* = nullptr> FORCEINLINE decltype(auto) Get() const         && { return TTupleElementGetterByIndex<Index, sizeof...(Types)>::Get(static_cast<const          TTupleBase&&>(*this)); }
-		template <uint32 Index, typename TEnableIf<(Index < sizeof...(Types))>::Type* = nullptr> FORCEINLINE decltype(auto) Get()       volatile&& { return TTupleElementGetterByIndex<Index, sizeof...(Types)>::Get(static_cast<      volatile TTupleBase&&>(*this)); }
-		template <uint32 Index, typename TEnableIf<(Index < sizeof...(Types))>::Type* = nullptr> FORCEINLINE decltype(auto) Get() const volatile&& { return TTupleElementGetterByIndex<Index, sizeof...(Types)>::Get(static_cast<const volatile TTupleBase&&>(*this)); }
-
-		template <typename T, typename TEnableIf<TTypeCountInParameterPack<T, Types...>::Value == 1>::Type* = nullptr> FORCEINLINE decltype(auto) Get()               &  { return TTupleElementGetterByType<T, sizeof...(Types)>::Get(static_cast<               TTupleBase& >(*this)); }
-		template <typename T, typename TEnableIf<TTypeCountInParameterPack<T, Types...>::Value == 1>::Type* = nullptr> FORCEINLINE decltype(auto) Get() const         &  { return TTupleElementGetterByType<T, sizeof...(Types)>::Get(static_cast<const          TTupleBase& >(*this)); }
-		template <typename T, typename TEnableIf<TTypeCountInParameterPack<T, Types...>::Value == 1>::Type* = nullptr> FORCEINLINE decltype(auto) Get()       volatile&  { return TTupleElementGetterByType<T, sizeof...(Types)>::Get(static_cast<      volatile TTupleBase& >(*this)); }
-		template <typename T, typename TEnableIf<TTypeCountInParameterPack<T, Types...>::Value == 1>::Type* = nullptr> FORCEINLINE decltype(auto) Get() const volatile&  { return TTupleElementGetterByType<T, sizeof...(Types)>::Get(static_cast<const volatile TTupleBase& >(*this)); }
-		template <typename T, typename TEnableIf<TTypeCountInParameterPack<T, Types...>::Value == 1>::Type* = nullptr> FORCEINLINE decltype(auto) Get()               && { return TTupleElementGetterByType<T, sizeof...(Types)>::Get(static_cast<               TTupleBase&&>(*this)); }
-		template <typename T, typename TEnableIf<TTypeCountInParameterPack<T, Types...>::Value == 1>::Type* = nullptr> FORCEINLINE decltype(auto) Get() const         && { return TTupleElementGetterByType<T, sizeof...(Types)>::Get(static_cast<const          TTupleBase&&>(*this)); }
-		template <typename T, typename TEnableIf<TTypeCountInParameterPack<T, Types...>::Value == 1>::Type* = nullptr> FORCEINLINE decltype(auto) Get()       volatile&& { return TTupleElementGetterByType<T, sizeof...(Types)>::Get(static_cast<      volatile TTupleBase&&>(*this)); }
-		template <typename T, typename TEnableIf<TTypeCountInParameterPack<T, Types...>::Value == 1>::Type* = nullptr> FORCEINLINE decltype(auto) Get() const volatile&& { return TTupleElementGetterByType<T, sizeof...(Types)>::Get(static_cast<const volatile TTupleBase&&>(*this)); }
+		template <uint32 Index> FORCEINLINE decltype(auto) Get()               &  { static_assert(Index < sizeof...(Types), "Invalid index passed to TTuple::Get"); return TTupleElementGetterByIndex<Index, sizeof...(Types)>::Get(static_cast<               TTupleBase& >(*this)); }
+		template <uint32 Index> FORCEINLINE decltype(auto) Get() const         &  { static_assert(Index < sizeof...(Types), "Invalid index passed to TTuple::Get"); return TTupleElementGetterByIndex<Index, sizeof...(Types)>::Get(static_cast<const          TTupleBase& >(*this)); }
+		template <uint32 Index> FORCEINLINE decltype(auto) Get()       volatile&  { static_assert(Index < sizeof...(Types), "Invalid index passed to TTuple::Get"); return TTupleElementGetterByIndex<Index, sizeof...(Types)>::Get(static_cast<      volatile TTupleBase& >(*this)); }
+		template <uint32 Index> FORCEINLINE decltype(auto) Get() const volatile&  { static_assert(Index < sizeof...(Types), "Invalid index passed to TTuple::Get"); return TTupleElementGetterByIndex<Index, sizeof...(Types)>::Get(static_cast<const volatile TTupleBase& >(*this)); }
+		template <uint32 Index> FORCEINLINE decltype(auto) Get()               && { static_assert(Index < sizeof...(Types), "Invalid index passed to TTuple::Get"); return TTupleElementGetterByIndex<Index, sizeof...(Types)>::Get(static_cast<               TTupleBase&&>(*this)); }
+		template <uint32 Index> FORCEINLINE decltype(auto) Get() const         && { static_assert(Index < sizeof...(Types), "Invalid index passed to TTuple::Get"); return TTupleElementGetterByIndex<Index, sizeof...(Types)>::Get(static_cast<const          TTupleBase&&>(*this)); }
+		template <uint32 Index> FORCEINLINE decltype(auto) Get()       volatile&& { static_assert(Index < sizeof...(Types), "Invalid index passed to TTuple::Get"); return TTupleElementGetterByIndex<Index, sizeof...(Types)>::Get(static_cast<      volatile TTupleBase&&>(*this)); }
+		template <uint32 Index> FORCEINLINE decltype(auto) Get() const volatile&& { static_assert(Index < sizeof...(Types), "Invalid index passed to TTuple::Get"); return TTupleElementGetterByIndex<Index, sizeof...(Types)>::Get(static_cast<const volatile TTupleBase&&>(*this)); }
+
+		template <typename T> FORCEINLINE decltype(auto) Get()               &  { static_assert(TTypeCountInParameterPack_V<T, Types...> == 1, "Invalid type passed to TTuple::Get"); return TTupleElementGetterByType<T, sizeof...(Types)>::Get(static_cast<               TTupleBase& >(*this)); }
+		template <typename T> FORCEINLINE decltype(auto) Get() const         &  { static_assert(TTypeCountInParameterPack_V<T, Types...> == 1, "Invalid type passed to TTuple::Get"); return TTupleElementGetterByType<T, sizeof...(Types)>::Get(static_cast<const          TTupleBase& >(*this)); }
+		template <typename T> FORCEINLINE decltype(auto) Get()       volatile&  { static_assert(TTypeCountInParameterPack_V<T, Types...> == 1, "Invalid type passed to TTuple::Get"); return TTupleElementGetterByType<T, sizeof...(Types)>::Get(static_cast<      volatile TTupleBase& >(*this)); }
+		template <typename T> FORCEINLINE decltype(auto) Get() const volatile&  { static_assert(TTypeCountInParameterPack_V<T, Types...> == 1, "Invalid type passed to TTuple::Get"); return TTupleElementGetterByType<T, sizeof...(Types)>::Get(static_cast<const volatile TTupleBase& >(*this)); }
+		template <typename T> FORCEINLINE decltype(auto) Get()               && { static_assert(TTypeCountInParameterPack_V<T, Types...> == 1, "Invalid type passed to TTuple::Get"); return TTupleElementGetterByType<T, sizeof...(Types)>::Get(static_cast<               TTupleBase&&>(*this)); }
+		template <typename T> FORCEINLINE decltype(auto) Get() const         && { static_assert(TTypeCountInParameterPack_V<T, Types...> == 1, "Invalid type passed to TTuple::Get"); return TTupleElementGetterByType<T, sizeof...(Types)>::Get(static_cast<const          TTupleBase&&>(*this)); }
+		template <typename T> FORCEINLINE decltype(auto) Get()       volatile&& { static_assert(TTypeCountInParameterPack_V<T, Types...> == 1, "Invalid type passed to TTuple::Get"); return TTupleElementGetterByType<T, sizeof...(Types)>::Get(static_cast<      volatile TTupleBase&&>(*this)); }
+		template <typename T> FORCEINLINE decltype(auto) Get() const volatile&& { static_assert(TTypeCountInParameterPack_V<T, Types...> == 1, "Invalid type passed to TTuple::Get"); return TTupleElementGetterByType<T, sizeof...(Types)>::Get(static_cast<const volatile TTupleBase&&>(*this)); }
 
 		template <typename FuncType, typename... ArgTypes> decltype(auto) ApplyAfter(FuncType&& Func, ArgTypes&&... Args)               &  { return ::Invoke(Func, Forward<ArgTypes>(Args)..., static_cast<               TTupleBase& >(*this).template Get<Indices>()...); }
 		template <typename FuncType, typename... ArgTypes> decltype(auto) ApplyAfter(FuncType&& Func, ArgTypes&&... Args) const         &  { return ::Invoke(Func, Forward<ArgTypes>(Args)..., static_cast<const          TTupleBase& >(*this).template Get<Indices>()...); }
@@ -398,9 +393,11 @@ namespace UE::Core::Private::Tuple
 
 		template <
 			typename KeyArgType,
-			typename ValueArgType,
-			typename TEnableIf<TIsConstructible<KeyType,   KeyArgType  &&>::Value>::Type* = nullptr,
-			typename TEnableIf<TIsConstructible<ValueType, ValueArgType&&>::Value>::Type* = nullptr
+			typename ValueArgType
+			UE_CONSTRAINTS_BEGIN
+				UE_CONSTRAINT(std::is_constructible_v<KeyType,   KeyArgType  &&>)
+				UE_CONSTRAINT(std::is_constructible_v<ValueType, ValueArgType&&>)
+			UE_CONSTRAINTS_END
 		>
 		explicit TTupleBase(EForwardingConstructor, KeyArgType&& KeyArg, ValueArgType&& ValueArg)
 			: Key  (Forward<KeyArgType  >(KeyArg  ))
@@ -410,9 +407,11 @@ namespace UE::Core::Private::Tuple
 
 		template <
 			typename TupleType,
-			typename std::decay_t<TupleType>::DummyPairIdentifier* = nullptr,
-			typename TEnableIf<TIsConstructible<KeyType,   decltype(DeclVal<TupleType&&>().Get<0>())>::Value>::Type* = nullptr,
-			typename TEnableIf<TIsConstructible<ValueType, decltype(DeclVal<TupleType&&>().Get<1>())>::Value>::Type* = nullptr
+			typename std::decay_t<TupleType>::DummyPairIdentifier* = nullptr
+			UE_CONSTRAINTS_BEGIN
+				UE_CONSTRAINT(std::is_constructible_v<KeyType,   decltype(DeclVal<TupleType&&>().Get<0>())>)
+				UE_CONSTRAINT(std::is_constructible_v<ValueType, decltype(DeclVal<TupleType&&>().Get<1>())>)
+			UE_CONSTRAINTS_END
 		>
 		explicit TTupleBase(EOtherTupleConstructor, TupleType&& Other)
 			: Key  (Forward<TupleType>(Other).Get<0>())
@@ -426,23 +425,23 @@ namespace UE::Core::Private::Tuple
 		TTupleBase& operator=(TTupleBase&& Other) = default;
 		TTupleBase& operator=(const TTupleBase& Other) = default;
 
-		template <uint32 Index, typename TEnableIf<(Index < 2)>::Type* = nullptr> FORCEINLINE decltype(auto) Get()               &  { return TTupleElementGetterByIndex<Index, 2>::Get(static_cast<               TTupleBase& >(*this)); }
-		template <uint32 Index, typename TEnableIf<(Index < 2)>::Type* = nullptr> FORCEINLINE decltype(auto) Get() const         &  { return TTupleElementGetterByIndex<Index, 2>::Get(static_cast<const          TTupleBase& >(*this)); }
-		template <uint32 Index, typename TEnableIf<(Index < 2)>::Type* = nullptr> FORCEINLINE decltype(auto) Get()       volatile&  { return TTupleElementGetterByIndex<Index, 2>::Get(static_cast<      volatile TTupleBase& >(*this)); }
-		template <uint32 Index, typename TEnableIf<(Index < 2)>::Type* = nullptr> FORCEINLINE decltype(auto) Get() const volatile&  { return TTupleElementGetterByIndex<Index, 2>::Get(static_cast<const volatile TTupleBase& >(*this)); }
-		template <uint32 Index, typename TEnableIf<(Index < 2)>::Type* = nullptr> FORCEINLINE decltype(auto) Get()               && { return TTupleElementGetterByIndex<Index, 2>::Get(static_cast<               TTupleBase&&>(*this)); }
-		template <uint32 Index, typename TEnableIf<(Index < 2)>::Type* = nullptr> FORCEINLINE decltype(auto) Get() const         && { return TTupleElementGetterByIndex<Index, 2>::Get(static_cast<const          TTupleBase&&>(*this)); }
-		template <uint32 Index, typename TEnableIf<(Index < 2)>::Type* = nullptr> FORCEINLINE decltype(auto) Get()       volatile&& { return TTupleElementGetterByIndex<Index, 2>::Get(static_cast<      volatile TTupleBase&&>(*this)); }
-		template <uint32 Index, typename TEnableIf<(Index < 2)>::Type* = nullptr> FORCEINLINE decltype(auto) Get() const volatile&& { return TTupleElementGetterByIndex<Index, 2>::Get(static_cast<const volatile TTupleBase&&>(*this)); }
-
-		template <typename T, typename TEnableIf<TTypeCountInParameterPack<T, KeyType, ValueType>::Value == 1>::Type* = nullptr> FORCEINLINE decltype(auto) Get()               &  { return TTupleElementGetterByType<T, 2>::Get(static_cast<               TTupleBase& >(*this)); }
-		template <typename T, typename TEnableIf<TTypeCountInParameterPack<T, KeyType, ValueType>::Value == 1>::Type* = nullptr> FORCEINLINE decltype(auto) Get() const         &  { return TTupleElementGetterByType<T, 2>::Get(static_cast<const          TTupleBase& >(*this)); }
-		template <typename T, typename TEnableIf<TTypeCountInParameterPack<T, KeyType, ValueType>::Value == 1>::Type* = nullptr> FORCEINLINE decltype(auto) Get()       volatile&  { return TTupleElementGetterByType<T, 2>::Get(static_cast<      volatile TTupleBase& >(*this)); }
-		template <typename T, typename TEnableIf<TTypeCountInParameterPack<T, KeyType, ValueType>::Value == 1>::Type* = nullptr> FORCEINLINE decltype(auto) Get() const volatile&  { return TTupleElementGetterByType<T, 2>::Get(static_cast<const volatile TTupleBase& >(*this)); }
-		template <typename T, typename TEnableIf<TTypeCountInParameterPack<T, KeyType, ValueType>::Value == 1>::Type* = nullptr> FORCEINLINE decltype(auto) Get()               && { return TTupleElementGetterByType<T, 2>::Get(static_cast<               TTupleBase&&>(*this)); }
-		template <typename T, typename TEnableIf<TTypeCountInParameterPack<T, KeyType, ValueType>::Value == 1>::Type* = nullptr> FORCEINLINE decltype(auto) Get() const         && { return TTupleElementGetterByType<T, 2>::Get(static_cast<const          TTupleBase&&>(*this)); }
-		template <typename T, typename TEnableIf<TTypeCountInParameterPack<T, KeyType, ValueType>::Value == 1>::Type* = nullptr> FORCEINLINE decltype(auto) Get()       volatile&& { return TTupleElementGetterByType<T, 2>::Get(static_cast<      volatile TTupleBase&&>(*this)); }
-		template <typename T, typename TEnableIf<TTypeCountInParameterPack<T, KeyType, ValueType>::Value == 1>::Type* = nullptr> FORCEINLINE decltype(auto) Get() const volatile&& { return TTupleElementGetterByType<T, 2>::Get(static_cast<const volatile TTupleBase&&>(*this)); }
+		template <uint32 Index> FORCEINLINE decltype(auto) Get()               &  { static_assert(Index < 2, "Invalid index passed to TTuple::Get"); return TTupleElementGetterByIndex<Index, 2>::Get(static_cast<               TTupleBase& >(*this)); }
+		template <uint32 Index> FORCEINLINE decltype(auto) Get() const         &  { static_assert(Index < 2, "Invalid index passed to TTuple::Get"); return TTupleElementGetterByIndex<Index, 2>::Get(static_cast<const          TTupleBase& >(*this)); }
+		template <uint32 Index> FORCEINLINE decltype(auto) Get()       volatile&  { static_assert(Index < 2, "Invalid index passed to TTuple::Get"); return TTupleElementGetterByIndex<Index, 2>::Get(static_cast<      volatile TTupleBase& >(*this)); }
+		template <uint32 Index> FORCEINLINE decltype(auto) Get() const volatile&  { static_assert(Index < 2, "Invalid index passed to TTuple::Get"); return TTupleElementGetterByIndex<Index, 2>::Get(static_cast<const volatile TTupleBase& >(*this)); }
+		template <uint32 Index> FORCEINLINE decltype(auto) Get()               && { static_assert(Index < 2, "Invalid index passed to TTuple::Get"); return TTupleElementGetterByIndex<Index, 2>::Get(static_cast<               TTupleBase&&>(*this)); }
+		template <uint32 Index> FORCEINLINE decltype(auto) Get() const         && { static_assert(Index < 2, "Invalid index passed to TTuple::Get"); return TTupleElementGetterByIndex<Index, 2>::Get(static_cast<const          TTupleBase&&>(*this)); }
+		template <uint32 Index> FORCEINLINE decltype(auto) Get()       volatile&& { static_assert(Index < 2, "Invalid index passed to TTuple::Get"); return TTupleElementGetterByIndex<Index, 2>::Get(static_cast<      volatile TTupleBase&&>(*this)); }
+		template <uint32 Index> FORCEINLINE decltype(auto) Get() const volatile&& { static_assert(Index < 2, "Invalid index passed to TTuple::Get"); return TTupleElementGetterByIndex<Index, 2>::Get(static_cast<const volatile TTupleBase&&>(*this)); }
+
+		template <typename T> FORCEINLINE decltype(auto) Get()               &  { static_assert(TTypeCountInParameterPack_V<T, KeyType, ValueType> == 1, "Invalid type passed to TTuple::Get"); return TTupleElementGetterByType<T, 2>::Get(static_cast<               TTupleBase& >(*this)); }
+		template <typename T> FORCEINLINE decltype(auto) Get() const         &  { static_assert(TTypeCountInParameterPack_V<T, KeyType, ValueType> == 1, "Invalid type passed to TTuple::Get"); return TTupleElementGetterByType<T, 2>::Get(static_cast<const          TTupleBase& >(*this)); }
+		template <typename T> FORCEINLINE decltype(auto) Get()       volatile&  { static_assert(TTypeCountInParameterPack_V<T, KeyType, ValueType> == 1, "Invalid type passed to TTuple::Get"); return TTupleElementGetterByType<T, 2>::Get(static_cast<      volatile TTupleBase& >(*this)); }
+		template <typename T> FORCEINLINE decltype(auto) Get() const volatile&  { static_assert(TTypeCountInParameterPack_V<T, KeyType, ValueType> == 1, "Invalid type passed to TTuple::Get"); return TTupleElementGetterByType<T, 2>::Get(static_cast<const volatile TTupleBase& >(*this)); }
+		template <typename T> FORCEINLINE decltype(auto) Get()               && { static_assert(TTypeCountInParameterPack_V<T, KeyType, ValueType> == 1, "Invalid type passed to TTuple::Get"); return TTupleElementGetterByType<T, 2>::Get(static_cast<               TTupleBase&&>(*this)); }
+		template <typename T> FORCEINLINE decltype(auto) Get() const         && { static_assert(TTypeCountInParameterPack_V<T, KeyType, ValueType> == 1, "Invalid type passed to TTuple::Get"); return TTupleElementGetterByType<T, 2>::Get(static_cast<const          TTupleBase&&>(*this)); }
+		template <typename T> FORCEINLINE decltype(auto) Get()       volatile&& { static_assert(TTypeCountInParameterPack_V<T, KeyType, ValueType> == 1, "Invalid type passed to TTuple::Get"); return TTupleElementGetterByType<T, 2>::Get(static_cast<      volatile TTupleBase&&>(*this)); }
+		template <typename T> FORCEINLINE decltype(auto) Get() const volatile&& { static_assert(TTypeCountInParameterPack_V<T, KeyType, ValueType> == 1, "Invalid type passed to TTuple::Get"); return TTupleElementGetterByType<T, 2>::Get(static_cast<const volatile TTupleBase&&>(*this)); }
 
 		template <typename FuncType, typename... ArgTypes> decltype(auto) ApplyAfter(FuncType&& Func, ArgTypes&&... Args)               &  { return ::Invoke(Func, Forward<ArgTypes>(Args)..., static_cast<               TTupleBase& >(*this).template Get<0>(), static_cast<               TTupleBase& >(*this).template Get<1>()); }
 		template <typename FuncType, typename... ArgTypes> decltype(auto) ApplyAfter(FuncType&& Func, ArgTypes&&... Args) const         &  { return ::Invoke(Func, Forward<ArgTypes>(Args)..., static_cast<const          TTupleBase& >(*this).template Get<0>(), static_cast<const          TTupleBase& >(*this).template Get<1>()); }
@@ -583,8 +582,8 @@ namespace UE::Core::Private::Tuple
 	template <typename Type, typename... TupleTypes>
 	struct TCVTupleIndex<Type, const volatile TTuple<TupleTypes...>>
 	{
-		static_assert(TTypeCountInParameterPack<Type, TupleTypes...>::Value >= 1, "TTupleIndex instantiated with a tuple which does not contain the type");
-		static_assert(TTypeCountInParameterPack<Type, TupleTypes...>::Value <= 1, "TTupleIndex instantiated with a tuple which contains multiple instances of the type");
+		static_assert(TTypeCountInParameterPack_V<Type, TupleTypes...> >= 1, "TTupleIndex instantiated with a tuple which does not contain the type");
+		static_assert(TTypeCountInParameterPack_V<Type, TupleTypes...> <= 1, "TTupleIndex instantiated with a tuple which contains multiple instances of the type");
 
 	private:
 		template <uint32 DeducedIndex>
@@ -599,10 +598,10 @@ namespace UE::Core::Private::Tuple
 	template <typename Type, typename KeyType, typename ValueType>
 	struct TCVTupleIndex<Type, const volatile TTuple<KeyType, ValueType>>
 	{
-		static_assert(TTypeCountInParameterPack<Type, KeyType, ValueType>::Value >= 1, "TTupleIndex instantiated with a tuple which does not contain the type");
-		static_assert(TTypeCountInParameterPack<Type, KeyType, ValueType>::Value <= 1, "TTupleIndex instantiated with a tuple which contains multiple instances of the type");
+		static_assert(TTypeCountInParameterPack_V<Type, KeyType, ValueType> >= 1, "TTupleIndex instantiated with a tuple which does not contain the type");
+		static_assert(TTypeCountInParameterPack_V<Type, KeyType, ValueType> <= 1, "TTupleIndex instantiated with a tuple which contains multiple instances of the type");
 
-		static constexpr uint32 Value = std::is_same<Type, ValueType>::value ? 1 : 0;
+		static constexpr uint32 Value = std::is_same_v<Type, ValueType> ? 1 : 0;
 	};
 #endif
 
@@ -687,8 +686,10 @@ public:
 #else
 	template <
 		typename... ArgTypes,
-		decltype(UE::Core::Private::Tuple::ConstructibleConceptCheck<Types...>(DeclVal<ArgTypes&&>()...))* = nullptr,
-		std::enable_if_t<std::conjunction_v<std::is_convertible<ArgTypes&&, Types>...>>* = nullptr
+		decltype(UE::Core::Private::Tuple::ConstructibleConceptCheck<Types...>(DeclVal<ArgTypes&&>()...))* = nullptr
+		UE_CONSTRAINTS_BEGIN
+			UE_CONSTRAINT((std::is_convertible_v<ArgTypes&&, Types> && ...))
+		UE_CONSTRAINTS_END
 	>
 	TTuple(ArgTypes&&... Args)
 		: Super(UE::Core::Private::Tuple::ForwardingConstructor, Forward<ArgTypes>(Args)...)
@@ -697,8 +698,10 @@ public:
 
 	template <
 		typename... ArgTypes,
-		decltype(UE::Core::Private::Tuple::ConstructibleConceptCheck<Types...>(DeclVal<ArgTypes&&>()...))* = nullptr,
-		std::enable_if_t<!std::conjunction_v<std::is_convertible<ArgTypes&&, Types>...>>* = nullptr
+		decltype(UE::Core::Private::Tuple::ConstructibleConceptCheck<Types...>(DeclVal<ArgTypes&&>()...))* = nullptr
+		UE_CONSTRAINTS_BEGIN
+			UE_CONSTRAINT(!(std::is_convertible_v<ArgTypes&&, Types> && ...))
+		UE_CONSTRAINTS_END
 	>
 	explicit TTuple(ArgTypes&&... Args)
 		: Super(UE::Core::Private::Tuple::ForwardingConstructor, Forward<ArgTypes>(Args)...)
@@ -943,7 +946,7 @@ FORCEINLINE decltype(auto) TransformTuple(const TTuple<Types...>& Tuple, FuncTyp
 template <typename FuncType, typename FirstTupleType, typename... TupleTypes>
 FORCEINLINE void VisitTupleElements(FuncType&& Func, FirstTupleType&& FirstTuple, TupleTypes&&... Tuples)
 {
-	UE::Core::Private::Tuple::TVisitTupleElements_Impl<TMakeIntegerSequence<uint32, TTupleArity<typename TDecay<FirstTupleType>::Type>::Value>>::Do(Forward<FuncType>(Func), Forward<FirstTupleType>(FirstTuple), Forward<TupleTypes>(Tuples)...);
+	UE::Core::Private::Tuple::TVisitTupleElements_Impl<TMakeIntegerSequence<uint32, TTupleArity<std::decay_t<FirstTupleType>>::Value>>::Do(Forward<FuncType>(Func), Forward<FirstTupleType>(FirstTuple), Forward<TupleTypes>(Tuples)...);
 }
 
 /**
