Replaced math constraints with new macros.

diff --git a/Engine/Source/Runtime/Core/Public/GenericPlatform/GenericPlatformMath.h b/Engine/Source/Runtime/Core/Public/GenericPlatform/GenericPlatformMath.h
--- a/Engine/Source/Runtime/Core/Public/GenericPlatform/GenericPlatformMath.h
+++ b/Engine/Source/Runtime/Core/Public/GenericPlatform/GenericPlatformMath.h
@@ -245,8 +245,6 @@ struct FGenericPlatformMath
 	{
 		return TruncToDouble(F);
 	}
-	
-	RESOLVE_FLOAT_AMBIGUITY(TruncToFloat);
 
 	/**
 	 * Converts a float to a nearest less or equal integer.
@@ -301,8 +299,6 @@ struct FGenericPlatformMath
 		return FloorToDouble(F);
 	}
 
-	RESOLVE_FLOAT_AMBIGUITY(FloorToFloat);
-
 	/**
 	 * Converts a float to the nearest integer. Rounds up when the fraction is .5
 	 * @param F		Floating point value to convert
@@ -349,8 +345,6 @@ struct FGenericPlatformMath
 		return RoundToDouble(F);
 	}
 
-	RESOLVE_FLOAT_AMBIGUITY(RoundToFloat);
-
 	/**
 	* Converts a float to the nearest greater or equal integer.
 	* @param F		Floating point value to convert
@@ -403,8 +397,6 @@ struct FGenericPlatformMath
 		return CeilToDouble(F);
 	}
 
-	RESOLVE_FLOAT_AMBIGUITY(CeilToFloat);
-
 	/**
 	 * Converts a double to nearest int64 with ties rounding to nearest even
 	 * May incur a performance penalty. Asserts on platforms that do not support this mode.
@@ -446,8 +438,6 @@ struct FGenericPlatformMath
 		return Value - TruncToDouble(Value);
 	}
 
-	RESOLVE_FLOAT_AMBIGUITY(Fractional);
-
 	/**
 	* Returns the fractional part of a float.
 	* @param Value	Floating point value to convert
@@ -464,8 +454,6 @@ struct FGenericPlatformMath
 		return Value - FloorToDouble(Value);
 	}
 
-	RESOLVE_FLOAT_AMBIGUITY(Frac);
-
 	/**
 	* Breaks the given value into an integral and a fractional part.
 	* @param InValue	Floating point value to convert
@@ -491,16 +479,13 @@ struct FGenericPlatformMath
 	// Returns e^Value
 	static FORCEINLINE float Exp( float Value ) { return expf(Value); }
 	static FORCEINLINE double Exp(double Value) { return exp(Value); }
-	RESOLVE_FLOAT_AMBIGUITY(Exp);
 
 	// Returns 2^Value
 	static FORCEINLINE float Exp2( float Value ) { return powf(2.f, Value); /*exp2f(Value);*/ }
 	static FORCEINLINE double Exp2(double Value) { return pow(2.0, Value); /*exp2(Value);*/ }
-	RESOLVE_FLOAT_AMBIGUITY(Exp2);
 
 	static FORCEINLINE float Loge( float Value ) {	return logf(Value); }
 	static FORCEINLINE double Loge(double Value) { return log(Value); }
-	RESOLVE_FLOAT_AMBIGUITY(Loge);
 
 	static FORCEINLINE float LogX( float Base, float Value ) { return Loge(Value) / Loge(Base); }
 	static FORCEINLINE double LogX(double Base, double Value) { return Loge(Value) / Loge(Base); }
@@ -510,7 +495,6 @@ struct FGenericPlatformMath
 	static FORCEINLINE float Log2( float Value ) { return Loge(Value) * 1.4426950f; }	
 	// 1.0 / Loge(2) = 1.442695040888963387
 	static FORCEINLINE double Log2(double Value) { return Loge(Value) * 1.442695040888963387; }
-	RESOLVE_FLOAT_AMBIGUITY(Log2);
 
 	/**
 	 * Returns the floating-point remainder of X / Y
@@ -528,31 +512,24 @@ struct FGenericPlatformMath
 
 	static FORCEINLINE float Sin( float Value ) { return sinf(Value); }
 	static FORCEINLINE double Sin( double Value ) { return sin(Value); }
-	RESOLVE_FLOAT_AMBIGUITY(Sin);
 
 	static FORCEINLINE float Asin( float Value ) { return asinf( (Value<-1.f) ? -1.f : ((Value<1.f) ? Value : 1.f) ); }
 	static FORCEINLINE double Asin( double Value ) { return asin( (Value<-1.0) ? -1.0 : ((Value<1.0) ? Value : 1.0) ); }
-	RESOLVE_FLOAT_AMBIGUITY(Asin);
 
 	static FORCEINLINE float Sinh(float Value) { return sinhf(Value); }
 	static FORCEINLINE double Sinh(double Value) { return sinh(Value); }
-	RESOLVE_FLOAT_AMBIGUITY(Sinh);
 
 	static FORCEINLINE float Cos( float Value ) { return cosf(Value); }
 	static FORCEINLINE double Cos( double Value ) { return cos(Value); }
-	RESOLVE_FLOAT_AMBIGUITY(Cos);
 
 	static FORCEINLINE float Acos( float Value ) { return acosf( (Value<-1.f) ? -1.f : ((Value<1.f) ? Value : 1.f) ); }
 	static FORCEINLINE double Acos( double Value ) { return acos( (Value<-1.0) ? -1.0 : ((Value<1.0) ? Value : 1.0) ); }
-	RESOLVE_FLOAT_AMBIGUITY(Acos);
 
 	static FORCEINLINE float Tan( float Value ) { return tanf(Value); }
 	static FORCEINLINE double Tan( double Value ) { return tan(Value); }
-	RESOLVE_FLOAT_AMBIGUITY(Tan);
 
 	static FORCEINLINE float Atan( float Value ) { return atanf(Value); }
 	static FORCEINLINE double Atan( double Value ) { return atan(Value); }
-	RESOLVE_FLOAT_AMBIGUITY(Atan);
 
 	static CORE_API float Atan2( float Y, float X );
 	static CORE_API double Atan2( double Y, double X );
@@ -560,7 +537,6 @@ struct FGenericPlatformMath
 
 	static FORCEINLINE float Sqrt( float Value ) { return sqrtf(Value); }
 	static FORCEINLINE double Sqrt( double Value ) { return sqrt(Value); }
-	RESOLVE_FLOAT_AMBIGUITY(Sqrt);
 
 	static FORCEINLINE float Pow( float A, float B ) { return powf(A,B); }
 	static FORCEINLINE double Pow( double A, double B ) { return pow(A,B); }
@@ -569,12 +545,10 @@ struct FGenericPlatformMath
 	/** Computes a fully accurate inverse square root */
 	static FORCEINLINE float InvSqrt( float F ) { return 1.0f / sqrtf( F ); }
 	static FORCEINLINE double InvSqrt( double F ) { return 1.0 / sqrt( F ); }
-	RESOLVE_FLOAT_AMBIGUITY(InvSqrt);
 
 	/** Computes a faster but less accurate inverse square root */
 	static FORCEINLINE float InvSqrtEst( float F ) { return InvSqrt( F ); }
 	static FORCEINLINE double InvSqrtEst( double F ) { return InvSqrt( F ); }
-	RESOLVE_FLOAT_AMBIGUITY(InvSqrtEst);
 
 	/** Return true if value is NaN (not a number). */
 	static FORCEINLINE bool IsNaN( float A ) 
diff --git a/Engine/Source/Runtime/Core/Public/Math/UnrealMathUtility.h b/Engine/Source/Runtime/Core/Public/Math/UnrealMathUtility.h
--- a/Engine/Source/Runtime/Core/Public/Math/UnrealMathUtility.h
+++ b/Engine/Source/Runtime/Core/Public/Math/UnrealMathUtility.h
@@ -411,7 +411,6 @@ struct FMath : public FPlatformMath
 		return Abs<double>( Value ) <= ErrorTolerance;
 	}
 
-	RESOLVE_FLOAT_PREDICATE_AMBIGUITY(IsNearlyZero);
 	RESOLVE_FLOAT_PREDICATE_AMBIGUITY_2_ARGS(IsNearlyZero);
 
 private:
@@ -530,8 +529,10 @@ public:
 	 * Unlike std::floor, it returns an IntegralType.
 	 */
 	template <
-	    typename IntegralType,
-	    typename TEnableIf<TIsIntegral<IntegralType>::Value>::Type* = nullptr
+		typename IntegralType
+		UE_CONSTRAINTS_BEGIN
+			UE_CONSTRAINT(std::is_integral_v<IntegralType>)
+		UE_CONSTRAINTS_END
 	>
 	UE_NODISCARD static constexpr FORCEINLINE IntegralType Floor(IntegralType I)
 	{
@@ -692,8 +693,13 @@ public:
 	* @param ScalarCos	Pointer to where the Cos result should be stored
 	* @param Value  input angles 
 	*/
-	template<typename T, TEMPLATE_REQUIRES(TIsFloatingPoint<T>::Value)>
-	static constexpr FORCEINLINE void SinCos(typename TDecay<T>::Type* ScalarSin, typename TDecay<T>::Type* ScalarCos, T  Value )
+	template <
+		typename T
+		UE_CONSTRAINTS_BEGIN
+			UE_CONSTRAINT(std::is_floating_point_v<T>)
+		UE_CONSTRAINTS_END
+	>
+	static constexpr FORCEINLINE void SinCos(std::decay_t<T>* ScalarSin, std::decay_t<T>* ScalarCos, T  Value )
 	{
 		// Map Value to y in [-pi,pi], x = 2*pi*quotient + remainder.
 		T quotient = (UE_INV_PI*0.5f)*Value;
@@ -741,7 +747,13 @@ public:
 		*ScalarCos = FMath::Cos(Value);
 	}
 
-	template<typename T, typename U, TEMPLATE_REQUIRES(!std::is_same_v<T, U>)>
+	template <
+		typename T,
+		typename U
+		UE_CONSTRAINTS_BEGIN
+			UE_CONSTRAINT(!std::is_same_v<T, U>)
+		UE_CONSTRAINTS_END
+	>
 	static FORCEINLINE void SinCos(T* ScalarSin, T* ScalarCos, U Value)
 	{
 		SinCos(ScalarSin, ScalarCos, T(Value));
@@ -828,7 +840,13 @@ public:
 	RESOLVE_FLOAT_AMBIGUITY_3_ARGS(ClampAngle);
 
 	/** Find the smallest angle between two headings (in degrees) */
-	template<typename T, typename T2, TEMPLATE_REQUIRES(TOr<TIsFloatingPoint<T>, TIsFloatingPoint<T2>>::Value)>
+	template <
+		typename T,
+		typename T2
+		UE_CONSTRAINTS_BEGIN
+			UE_CONSTRAINT(std::is_floating_point_v<T> || std::is_floating_point_v<T2>)
+		UE_CONSTRAINTS_END
+	>
 	UE_NODISCARD static constexpr auto FindDeltaAngleDegrees(T A1, T2 A2) -> decltype(A1 * A2)
 	{
 		// Find the difference
@@ -852,7 +870,13 @@ public:
 	}
 
 	/** Find the smallest angle between two headings (in radians) */
-	template<typename T, typename T2, TEMPLATE_REQUIRES(TOr<TIsFloatingPoint<T>, TIsFloatingPoint<T2>>::Value)>
+	template <
+		typename T,
+		typename T2
+		UE_CONSTRAINTS_BEGIN
+			UE_CONSTRAINT(std::is_floating_point_v<T> || std::is_floating_point_v<T2>)
+		UE_CONSTRAINTS_END
+	>
 	UE_NODISCARD static constexpr auto FindDeltaAngleRadians(T A1, T2 A2) -> decltype(A1 * A2)
 	{
 		// Find the difference
@@ -876,7 +900,12 @@ public:
 	}
 
 	/** Given a heading which may be outside the +/- PI range, 'unwind' it back into that range. */
-	template<typename T, TEMPLATE_REQUIRES(TIsFloatingPoint<T>::Value)>
+	template <
+		typename T
+		UE_CONSTRAINTS_BEGIN
+			UE_CONSTRAINT(std::is_floating_point_v<T>)
+		UE_CONSTRAINTS_END
+	>
 	UE_NODISCARD static constexpr T UnwindRadians(T A)
 	{
 		while(A > UE_PI)
@@ -893,7 +922,12 @@ public:
 	}
 
 	/** Utility to ensure angle is between +/- 180 degrees by unwinding. */
-	template<typename T, TEMPLATE_REQUIRES(TIsFloatingPoint<T>::Value)>
+	template <
+		typename T
+		UE_CONSTRAINTS_BEGIN
+			UE_CONSTRAINT(std::is_floating_point_v<T>)
+		UE_CONSTRAINTS_END
+	>
 	UE_NODISCARD static constexpr T UnwindDegrees(T A)
 	{
 		while(A > 180.f)
@@ -996,7 +1030,13 @@ public:
 	// Interpolation Functions
 
 	/** Calculates the percentage along a line from MinValue to MaxValue that Value is. */
-	template<typename T, typename T2, TEMPLATE_REQUIRES(TIsFloatingPoint<T>::Value)>
+	template <
+		typename T,
+		typename T2
+		UE_CONSTRAINTS_BEGIN
+			UE_CONSTRAINT(std::is_floating_point_v<T>)
+		UE_CONSTRAINTS_END
+	>
 	UE_NODISCARD static constexpr FORCEINLINE auto GetRangePct(T MinValue, T MaxValue, T2 Value)
 	{
 		// Avoid Divide by Zero.
@@ -1012,14 +1052,26 @@ public:
 	}
 
 	/** Same as above, but taking a 2d vector as the range. */
-	template<typename T, typename T2, TEMPLATE_REQUIRES(TIsFloatingPoint<T>::Value)>
+	template <
+		typename T,
+		typename T2
+		UE_CONSTRAINTS_BEGIN
+			UE_CONSTRAINT(std::is_floating_point_v<T>)
+		UE_CONSTRAINTS_END
+	>
 	UE_NODISCARD static auto GetRangePct(UE::Math::TVector2<T> const& Range, T2 Value)
 	{
 		return GetRangePct(Range.X, Range.Y, Value);
 	}
 	
 	/** Basically a Vector2d version of Lerp. */
-	template<typename T, typename T2, TEMPLATE_REQUIRES(TIsFloatingPoint<T>::Value)>
+	template <
+		typename T,
+		typename T2
+		UE_CONSTRAINTS_BEGIN
+			UE_CONSTRAINT(std::is_floating_point_v<T>)
+		UE_CONSTRAINTS_END
+	>
 	UE_NODISCARD static auto GetRangeValue(UE::Math::TVector2<T> const& Range, T2 Pct)
 	{
 		return Lerp(Range.X, Range.Y, Pct);
@@ -1061,21 +1113,40 @@ public:
 	}
 
 	/** Performs a linear interpolation between two values, Alpha ranges from 0-1 */
-	template< class T, class U, TEMPLATE_REQUIRES(!TCustomLerp<T>::Value && (std::is_floating_point_v<U> || std::is_same_v<T, U>))>
+	template <
+		typename T,
+		typename U
+		UE_CONSTRAINTS_BEGIN
+			UE_CONSTRAINT(!TCustomLerp<T>::Value && (std::is_floating_point_v<U> || std::is_same_v<T, U>))
+		UE_CONSTRAINTS_END
+	>
 	UE_NODISCARD static constexpr FORCEINLINE_DEBUGGABLE T Lerp( const T& A, const T& B, const U& Alpha )
 	{
 		return (T)(A + Alpha * (B-A));
 	}
 
 	/** Custom lerps defined for those classes not suited to the default implemenation. */
-	template< class T, class U, TEMPLATE_REQUIRES(TCustomLerp<T>::Value) >
+	template <
+		typename T,
+		typename U
+		UE_CONSTRAINTS_BEGIN
+			UE_CONSTRAINT(TCustomLerp<T>::Value)
+		UE_CONSTRAINTS_END
+	>
 	UE_NODISCARD static FORCEINLINE_DEBUGGABLE T Lerp(const T& A, const T& B, const U& Alpha)
 	{
 		return TCustomLerp<T>::Lerp(A, B, Alpha);
 	}
 
 	// Allow passing of differing A/B types.
-	template<typename T1, typename T2, typename T3, TEMPLATE_REQUIRES(!std::is_same_v<T1, T2> && !TCustomLerp<T1>::Value && !TCustomLerp<T2>::Value)>
+	template <
+		typename T1,
+		typename T2,
+		typename T3
+		UE_CONSTRAINTS_BEGIN
+			UE_CONSTRAINT(!std::is_same_v<T1, T2> && !TCustomLerp<T1>::Value && !TCustomLerp<T2>::Value)
+		UE_CONSTRAINTS_END
+	>
 	UE_NODISCARD static auto Lerp( const T1& A, const T2& B, const T3& Alpha ) -> decltype(A * B)
 	{
 		using ABType = decltype(A * B);
@@ -1097,7 +1168,14 @@ public:
 	}
 
 	// Allow passing of differing A/B types.
-	template<typename T1, typename T2, typename T3, TEMPLATE_REQUIRES(!std::is_same_v<T1, T2>)>
+	template <
+		typename T1,
+		typename T2,
+		typename T3
+		UE_CONSTRAINTS_BEGIN
+			UE_CONSTRAINT(!std::is_same_v<T1, T2>)
+		UE_CONSTRAINTS_END
+	>
 	UE_NODISCARD static auto LerpStable( const T1& A, const T2& B, const T3& Alpha ) -> decltype(A * B)
 	{
 		using ABType = decltype(A * B);
@@ -1105,7 +1183,13 @@ public:
 	}
 	
 	/** Performs a 2D linear interpolation between four values values, FracX, FracY ranges from 0-1 */
-	template< class T, class U, TEMPLATE_REQUIRES(!TCustomLerp<T>::Value && (std::is_floating_point_v<U> || std::is_same_v<T, U>))>
+	template <
+		typename T,
+		typename U
+		UE_CONSTRAINTS_BEGIN
+			UE_CONSTRAINT(!TCustomLerp<T>::Value && (std::is_floating_point_v<U> || std::is_same_v<T, U>))
+		UE_CONSTRAINTS_END
+	>
 	UE_NODISCARD static constexpr FORCEINLINE_DEBUGGABLE T BiLerp(const T& P00,const T& P10,const T& P01,const T& P11, const U& FracX, const U& FracY)
 	{
 		return Lerp(
@@ -1116,7 +1200,13 @@ public:
 	}
 
 	/** Custom lerps defined for those classes not suited to the default implemenation. */
-	template< class T, class U, TEMPLATE_REQUIRES(TCustomLerp<T>::Value) >
+	template <
+		typename T,
+		typename U
+		UE_CONSTRAINTS_BEGIN
+			UE_CONSTRAINT(TCustomLerp<T>::Value)
+		UE_CONSTRAINTS_END
+	>
 	UE_NODISCARD static FORCEINLINE_DEBUGGABLE T BiLerp(const T& P00, const T& P10, const T& P01, const T& P11, const U& FracX, const U& FracY)
 	{
 		return TCustomLerp<T>::BiLerp(P00, P10, P01, P11, FracX, FracY);
@@ -1131,7 +1221,13 @@ public:
 	 *
 	 * @return  Interpolated value
 	 */
-	template< class T, class U, TEMPLATE_REQUIRES(!TCustomLerp<T>::Value && (std::is_floating_point_v<U> || std::is_same_v<T, U>))>
+	template <
+		typename T,
+		typename U
+		UE_CONSTRAINTS_BEGIN
+			UE_CONSTRAINT(!TCustomLerp<T>::Value && (std::is_floating_point_v<U> || std::is_same_v<T, U>))
+		UE_CONSTRAINTS_END
+	>
 	UE_NODISCARD static constexpr FORCEINLINE_DEBUGGABLE T CubicInterp( const T& P0, const T& T0, const T& P1, const T& T1, const U& A )
 	{
 		const U A2 = A * A;
@@ -1141,7 +1237,13 @@ public:
 	}
 
 	/** Custom lerps defined for those classes not suited to the default implemenation. */
-	template< class T, class U, TEMPLATE_REQUIRES(TCustomLerp<T>::Value) >
+	template <
+		typename T,
+		typename U
+		UE_CONSTRAINTS_BEGIN
+			UE_CONSTRAINT(TCustomLerp<T>::Value)
+		UE_CONSTRAINTS_END
+	>
 	UE_NODISCARD static FORCEINLINE_DEBUGGABLE T CubicInterp(const T& P0, const T& T0, const T& P1, const T& T1, const U& A)
 	{
 		return TCustomLerp<T>::CubicInterp(P0, T0, P1, T1, A);
@@ -1156,7 +1258,13 @@ public:
 	 *
 	 * @return  Interpolated value
 	 */
-	template< class T, class U, TEMPLATE_REQUIRES(TIsFloatingPoint<U>::Value) >
+	template <
+		typename T,
+		typename U
+		UE_CONSTRAINTS_BEGIN
+			UE_CONSTRAINT(std::is_floating_point_v<U>)
+		UE_CONSTRAINTS_END
+	>
 	UE_NODISCARD static constexpr FORCEINLINE_DEBUGGABLE T CubicInterpDerivative( const T& P0, const T& T0, const T& P1, const T& T1, const U& A )
 	{
 		T a = 6.f*P0 + 3.f*T0 + 3.f*T1 - 6.f*P1;
@@ -1177,7 +1285,13 @@ public:
 	 *
 	 * @return  Interpolated value
 	 */
-	template< class T, class U, TEMPLATE_REQUIRES(TIsFloatingPoint<U>::Value) >
+	template <
+		typename T,
+		typename U
+		UE_CONSTRAINTS_BEGIN
+			UE_CONSTRAINT(std::is_floating_point_v<U>)
+		UE_CONSTRAINTS_END
+	>
 	UE_NODISCARD static constexpr FORCEINLINE_DEBUGGABLE T CubicInterpSecondDerivative( const T& P0, const T& T0, const T& P1, const T& T1, const U& A )
 	{
 		T a = 12.f*P0 + 6.f*T0 + 6.f*T1 - 12.f*P1;
@@ -2426,8 +2540,8 @@ namespace LWC
 	template<typename TDest, typename TSrc, typename InAllocatorType>
 	TArray<TDest, InAllocatorType> ConvertArrayType(const TArray<TSrc, InAllocatorType>& From)
 	{
-		//static_assert(!std::is_same<TDest, TSrc>::value, "Redundant call to ConvertArrayType");	// Unavoidable if supporting LWC toggle, but a useful check once LWC is locked to enabled.
-		if constexpr (std::is_same<TDest, TSrc>::value)
+		//static_assert(!std::is_same_v<TDest, TSrc>, "Redundant call to ConvertArrayType");	// Unavoidable if supporting LWC toggle, but a useful check once LWC is locked to enabled.
+		if constexpr (std::is_same_v<TDest, TSrc>)
 		{
 			return From;
 		}
@@ -2447,8 +2561,8 @@ namespace LWC
 	template<typename TDest, typename TSrc, typename InAllocatorType>
 	TArray<TDest, InAllocatorType> ConvertArrayTypeClampMax(const TArray<TSrc, InAllocatorType>& From)
 	{
-		//static_assert(!std::is_same<TDest, TSrc>::value, "Redundant call to ConvertArrayType");	// Unavoidable if supporting LWC toggle, but a useful check once LWC is locked to enabled.
-		if constexpr (std::is_same<TDest, TSrc>::value)
+		//static_assert(!std::is_same_v<TDest, TSrc>, "Redundant call to ConvertArrayType");	// Unavoidable if supporting LWC toggle, but a useful check once LWC is locked to enabled.
+		if constexpr (std::is_same_v<TDest, TSrc>)
 		{
 			return From;
 		}
diff --git a/Engine/Source/Runtime/Core/Public/Templates/ResolveTypeAmbiguity.h b/Engine/Source/Runtime/Core/Public/Templates/ResolveTypeAmbiguity.h
--- a/Engine/Source/Runtime/Core/Public/Templates/ResolveTypeAmbiguity.h
+++ b/Engine/Source/Runtime/Core/Public/Templates/ResolveTypeAmbiguity.h
@@ -25,8 +25,6 @@
  * 
  * 	static float Sin( float Value ) { return sinf(Value); }
  *	static double Sin( double Value ) { return sin(Value); }
- *  // Disallows any other type for the 'Value' argument, such as Sin(int32).
- *	RESOLVE_FLOAT_AMBIGUITY(Sin);
  *
  *	static float Fmod(float X, float Y);
  *	static double Fmod(double X, double Y);
@@ -36,18 +34,19 @@
  */
 
 
-// Assuming that a float and/or double version of a function are declared, this deprecates all other types as ambiguous and falls back to the float variant.
-#define RESOLVE_FLOAT_AMBIGUITY(Func) \
-	template<typename T> UE_DEPRECATED(5.0, "Arguments cause function resolution ambiguity.") \
-	static FORCEINLINE float Func(T&& Value) { return Func((float)Value); }
-
-
 // Mixing any signed integral types with any other signed integral type results in same type as "X * Y", which is promoted to the type of the result of mixed arithmetic between the types
 #define MIX_SIGNED_INTS_2_ARGS_ACTUAL(Func, OptionalMarkup) \
-	template<typename Arg1, typename Arg2, TEMPLATE_REQUIRES(!std::is_same_v<Arg1, Arg2> && \
-																std::is_signed_v<Arg1> && std::is_integral_v<Arg1> && \
-																std::is_signed_v<Arg2> && std::is_integral_v<Arg2> \
-																)> \
+	template < \
+		typename Arg1, \
+		typename Arg2 \
+		UE_CONSTRAINTS_BEGIN \
+			UE_CONSTRAINT( \
+				!std::is_same_v<Arg1, Arg2> && \
+				std::is_signed_v<Arg1> && std::is_integral_v<Arg1> && \
+				std::is_signed_v<Arg2> && std::is_integral_v<Arg2> \
+			) \
+		UE_CONSTRAINTS_END \
+	> \
 	static OptionalMarkup FORCEINLINE auto Func(Arg1 X, Arg2 Y) -> decltype(X * Y) \
 	{ \
 		using ArgType = decltype(X * Y); \
@@ -59,16 +58,23 @@
 
 // Mixing any signed integral types with any other signed integral type results in same type as "X * Y", which is promoted to the type of the result of mixed arithmetic between the types
 #define MIX_SIGNED_TYPES_3_ARGS_ACTUAL(Func, OptionalMarkup) \
-	template<typename Arg1, typename Arg2, typename Arg3, TEMPLATE_REQUIRES(TAnd< \
-																				( \
-																					!std::is_same_v<Arg1, Arg2>> || \
-																					!std::is_same_v<Arg2, Arg3>> || \
-																					!std::is_same_v<Arg1, Arg3>> \
-																				) && \
-																				std::is_signed_v<Arg1> && std::is_integral_v<Arg1> && \
-																				std::is_signed_v<Arg2> && std::is_integral_v<Arg2> && \
-																				std::is_signed_v<Arg3> && std::is_integral_v<Arg3> \
-																			)> \
+	template < \
+		typename Arg1, \
+		typename Arg2, \
+		typename Arg3 \
+		UE_CONSTRAINTS_BEGIN \
+			UE_CONSTRAINT( \
+				( \
+					!std::is_same_v<Arg1, Arg2>> || \
+					!std::is_same_v<Arg2, Arg3>> || \
+					!std::is_same_v<Arg1, Arg3>> \
+				) && \
+				std::is_signed_v<Arg1> && std::is_integral_v<Arg1> && \
+				std::is_signed_v<Arg2> && std::is_integral_v<Arg2> && \
+				std::is_signed_v<Arg3> && std::is_integral_v<Arg3> \
+			) \
+		UE_CONSTRAINTS_END \
+	> \
 	static OptionalMarkup FORCEINLINE auto Func(Arg1 X, Arg2 Y, Arg3 Z) -> decltype(X * Y * Z) \
 	{ \
 		using ArgType = decltype(X * Y * Z); \
@@ -79,55 +85,34 @@
 #define MIX_SIGNED_INTS_3_ARGS_CONSTEXPR(Func)		MIX_SIGNED_INTS_3_ARGS_ACTUAL(Func, CONSTEXPR)
 
 
-// Cannot know whether the type of e.g. Fmod(int, int) should be float or double.
-#define REMOVE_AMBIGUOUS_FLOAT_2_ARGS(Func) \
-	template<typename Arg1, typename Arg2, TEMPLATE_REQUIRES(TAnd< \
-																TNot<TIsFloatingPoint<Arg1>>,\
-																TNot<TIsFloatingPoint<Arg2>> \
-																>::Value)> UE_DEPRECATED(5.0, "Arguments cause function resolution ambiguity, expected a floating point type.") \
-	static FORCEINLINE float Func(Arg1 X, Arg2 Y) \
-	{ \
-		return Func((float)X, (float)Y); \
-	}
-
 // Mixing float and double types with any other type results in same type as "X * Y", which is promoted to the floating point type with the highest precision of the argument types.
 #define MIX_FLOATS_2_ARGS(Func) \
-	template<typename Arg1, typename Arg2, TEMPLATE_REQUIRES( (std::is_floating_point_v<Arg1> || std::is_floating_point_v<Arg2>) && \
-																!std::is_same_v<Arg1, Arg2> \
-															)> \
+	template < \
+		typename Arg1, \
+		typename Arg2 \
+		UE_CONSTRAINTS_BEGIN \
+			UE_CONSTRAINT((std::is_floating_point_v<Arg1> || std::is_floating_point_v<Arg2>) && !std::is_same_v<Arg1, Arg2>) \
+		UE_CONSTRAINTS_END \
+	> \
 	static FORCEINLINE auto Func(Arg1 X, Arg2 Y) -> decltype(X * Y) \
 	{ \
 		using ArgType = decltype(X * Y); \
 		return Func((ArgType)X, (ArgType)Y); \
 	}
 
-// Cannot know whether the type of e.g. Fmod(int, int) should be float or double.
-#define REMOVE_AMBIGUOUS_FLOAT_3_ARGS(Func) \
-	template<typename Arg1, typename Arg2, typename Arg3, TEMPLATE_REQUIRES(TAnd< \
-																				TNot<TIsFloatingPoint<Arg1>>, \
-																				TNot<TIsFloatingPoint<Arg2>>, \
-																				TNot<TIsFloatingPoint<Arg3>> \
-																				>::Value \
-																			)> UE_DEPRECATED(5.0, "Arguments cause function resolution ambiguity, expected a floating point type.") \
-	static FORCEINLINE float Func(Arg1 X, Arg2 Y, Arg3 Z) \
-	{ \
-		return Func((float)X, (float)Y, (float)Z); \
-	}
-
 // Mixing float and double types with any other type results in same type as "X * Y", which is promoted to the floating point type with the highest precision of the argument types.
 #define MIX_FLOATS_3_ARGS(Func) \
-	template<typename Arg1, typename Arg2, typename Arg3, TEMPLATE_REQUIRES( \
-																				( \
-																					std::is_floating_point_v<Arg1> || \
-																					std::is_floating_point_v<Arg2> || \
-																					std::is_floating_point_v<Arg3> \
-																				) && \
-																				( \
-																					!std::is_same_v<Arg1, Arg2> || \
-																					!std::is_same_v<Arg2, Arg3> || \
-																					!std::is_same_v<Arg1, Arg3> \
-																				) \
-																			)> \
+	template < \
+		typename Arg1, \
+		typename Arg2, \
+		typename Arg3 \
+		UE_CONSTRAINTS_BEGIN \
+			UE_CONSTRAINT( \
+				(std::is_floating_point_v<Arg1> || std::is_floating_point_v<Arg2> || std::is_floating_point_v<Arg3>) && \
+				(!std::is_same_v<Arg1, Arg2> || !std::is_same_v<Arg2, Arg3> || !std::is_same_v<Arg1, Arg3>) \
+			) \
+		UE_CONSTRAINTS_END \
+	> \
 	static FORCEINLINE auto Func(Arg1 X, Arg2 Y, Arg3 Z) -> decltype(X * Y * Z) \
 	{ \
 		using ArgType = decltype(X * Y * Z); \
@@ -139,70 +124,45 @@
 // Otherwise promotes to the highest precision floating point type of the arugments.
 
 #define RESOLVE_FLOAT_AMBIGUITY_2_ARGS(Func) \
-	REMOVE_AMBIGUOUS_FLOAT_2_ARGS(Func); \
 	MIX_FLOATS_2_ARGS(Func);
 
 #define RESOLVE_FLOAT_AMBIGUITY_3_ARGS(Func) \
-	REMOVE_AMBIGUOUS_FLOAT_3_ARGS(Func); \
 	MIX_FLOATS_3_ARGS(Func);
 
 
 ///////////////////////////////////////////////////////////////////////////////////////////////////////////////////////////////////////////////////////////////
 
-// Assuming that a float and/or double version of a function are declared, this deprecates all other types as ambiguous and falls back to the float variant (returning ReturnType type)
-#define RESOLVE_FLOAT_TO_TYPE_AMBIGUITY(Func, ReturnType) \
-	template<typename T> UE_DEPRECATED(5.0, "Arguments cause function resolution ambiguity.") \
-	static FORCEINLINE ReturnType Func(T&& Value) { return Func((float)Value); }
-
-#define REMOVE_AMBIGUOUS_FLOAT_TO_TYPE_2_ARGS(Func, ReturnType) \
-	template<typename Arg1, typename Arg2, TEMPLATE_REQUIRES(TAnd< \
-																TNot<TIsFloatingPoint<Arg1>>,\
-																TNot<TIsFloatingPoint<Arg2>> \
-																>::Value)> UE_DEPRECATED(5.0, "Arguments cause function resolution ambiguity, expected a floating point type.") \
-	static FORCEINLINE ReturnType Func(Arg1 X, Arg2 Y) \
-	{ \
-		return Func((float)X, (float)Y); \
-	}
-
 // Mixing float and double types with any other type results in same type as "X * Y", which is promoted to the floating point type with the highest precision of the argument types.
 #define MIX_FLOATS_TO_TYPE_2_ARGS(Func, ReturnType) \
-	template<typename Arg1, typename Arg2, TEMPLATE_REQUIRES( \
-																(std::is_floating_point_v<Arg1> || std::is_floating_point_v<Arg2>) && \
-																!std::is_same_v<Arg1, Arg2> \
-															)> \
+	template < \
+		typename Arg1, \
+		typename Arg2 \
+		UE_CONSTRAINTS_BEGIN \
+			UE_CONSTRAINT( \
+				(std::is_floating_point_v<Arg1> || std::is_floating_point_v<Arg2>) && \
+				!std::is_same_v<Arg1, Arg2> \
+			) \
+		UE_CONSTRAINTS_END \
+	> \
 	static FORCEINLINE ReturnType Func(Arg1 X, Arg2 Y) \
 	{ \
 		using ArgType = decltype(X * Y); \
 		return Func((ArgType)X, (ArgType)Y); \
 	}
 
-// Cannot know whether the type of e.g. Fmod(int, int) should be float or double.
-#define REMOVE_AMBIGUOUS_FLOAT_TO_TYPE_3_ARGS(Func, ReturnType) \
-	template<typename Arg1, typename Arg2, typename Arg3, TEMPLATE_REQUIRES(TAnd< \
-																				TNot<TIsFloatingPoint<Arg1>>, \
-																				TNot<TIsFloatingPoint<Arg2>>, \
-																				TNot<TIsFloatingPoint<Arg3>> \
-																				>::Value \
-																			)> UE_DEPRECATED(5.0, "Arguments cause function resolution ambiguity, expected a floating point type.") \
-	static FORCEINLINE ReturnType Func(Arg1 X, Arg2 Y) \
-	{ \
-		return Func((float)X, (float)Y); \
-	}
-
 // Mixing float and double types with any other type results in same type as "X * Y", which is promoted to the floating point type with the highest precision of the argument types.
 #define MIX_FLOATS_TO_TYPE_3_ARGS(Func, ReturnType) \
-	template<typename Arg1, typename Arg2, typename Arg3, TEMPLATE_REQUIRES( \
-																				( \
-																					std::is_floating_point_v<Arg1> || \
-																					std::is_floating_point_v<Arg2> || \
-																					std::is_floating_point_v<Arg3> \
-																				) && \
-																				( \
-																					!std::is_same_v<Arg1, Arg2> || \
-																					!std::is_same_v<Arg1, Arg3> || \
-																					!std::is_same_v<Arg2, Arg3> \
-																				) \
-																			)> \
+	template < \
+		typename Arg1, \
+		typename Arg2, \
+		typename Arg3 \
+		UE_CONSTRAINTS_BEGIN \
+			UE_CONSTRAINT( \
+				(std::is_floating_point_v<Arg1> || std::is_floating_point_v<Arg2> || std::is_floating_point_v<Arg3>) && \
+				(!std::is_same_v<Arg1, Arg2> || !std::is_same_v<Arg1, Arg3> || !std::is_same_v<Arg2, Arg3>) \
+			) \
+		UE_CONSTRAINTS_END \
+	> \
 	static FORCEINLINE ReturnType Func(Arg1 X, Arg2 Y, Arg3 Z) \
 	{ \
 		using ArgType = decltype(X * Y * Z); \
@@ -214,14 +174,11 @@
 // Otherwise promotes to the highest precision floating point type of the arugments.
 
 #define RESOLVE_FLOAT_TO_TYPE_AMBIGUITY_2_ARGS(Func, ReturnType) \
-	REMOVE_AMBIGUOUS_FLOAT_TO_TYPE_2_ARGS(Func, ReturnType); \
 	MIX_FLOATS_TO_TYPE_2_ARGS(Func, ReturnType);
 
 #define RESOLVE_FLOAT_TO_TYPE_AMBIGUITY_3_ARGS(Func, ReturnType) \
-	REMOVE_AMBIGUOUS_FLOAT_TO_TYPE_3_ARGS(Func, ReturnType); \
 	MIX_FLOATS_TO_TYPE_3_ARGS(Func, ReturnType);
 
 // nicer names for the bool version
-#define RESOLVE_FLOAT_PREDICATE_AMBIGUITY(Func) RESOLVE_FLOAT_TO_TYPE_AMBIGUITY(Func, bool)
 #define RESOLVE_FLOAT_PREDICATE_AMBIGUITY_2_ARGS(Func) RESOLVE_FLOAT_TO_TYPE_AMBIGUITY_2_ARGS(Func, bool)
 #define RESOLVE_FLOAT_PREDICATE_AMBIGUITY_3_ARGS(Func) RESOLVE_FLOAT_TO_TYPE_AMBIGUITY_3_ARGS(Func, bool)
