Temporary CIS fix for TEMPLATE_REQUIRES.

diff --git a/Engine/Source/Runtime/Core/Public/Templates/UnrealTypeTraits.h b/Engine/Source/Runtime/Core/Public/Templates/UnrealTypeTraits.h
--- a/Engine/Source/Runtime/Core/Public/Templates/UnrealTypeTraits.h
+++ b/Engine/Source/Runtime/Core/Public/Templates/UnrealTypeTraits.h
@@ -53,7 +53,7 @@
     #define UE_CONSTRAINTS_END UE::Core::Private::BoolIdentityConcept<true
 #endif
 
-#define TEMPLATE_REQUIRES(...) /*UE_DEPRECATED_MACRO(5.2, "TEMPLATE_REQUIRES has been deprecated, please use UE_CONSTRAINT instead.")*/ std::enable_if_t<__VA_ARGS__, int> = 0
+#define TEMPLATE_REQUIRES(...) typename TEnableIf<__VA_ARGS__, int>::type = 0
 
 /*-----------------------------------------------------------------------------
  * Macros to abstract the presence of certain compiler intrinsic type traits 
