SharedBuffer: Added a static assert that uint64 fits in SIZE_T

#rnx

diff --git a/Engine/Source/Runtime/Core/Private/Memory/SharedBuffer.cpp b/Engine/Source/Runtime/Core/Private/Memory/SharedBuffer.cpp
--- a/Engine/Source/Runtime/Core/Private/Memory/SharedBuffer.cpp
+++ b/Engine/Source/Runtime/Core/Private/Memory/SharedBuffer.cpp
@@ -32,6 +32,8 @@ inline void TBufferOwnerPtr<FOps>::Reset()
 class FBufferOwnerHeap final : public FBufferOwner
 {
 public:
+	static_assert(~uint64(0) <= ~SIZE_T(0));
+
 	inline explicit FBufferOwnerHeap(uint64 Size)
 		: FBufferOwner(FMemory::Malloc(Size), Size)
 	{
