CompressedBuffer: Added a check to avoid a buffer overflow

#rnx

diff --git a/Engine/Source/Runtime/Core/Private/Compression/CompressedBuffer.cpp b/Engine/Source/Runtime/Core/Private/Compression/CompressedBuffer.cpp
--- a/Engine/Source/Runtime/Core/Private/Compression/CompressedBuffer.cpp
+++ b/Engine/Source/Runtime/Core/Private/Compression/CompressedBuffer.cpp
@@ -930,13 +930,16 @@ FCompressedBuffer FCompressedBuffer::Load(FArchive& Ar)
 	Header.ByteSwap();
 
 	FCompressedBuffer Local;
-	if (Header.Magic == Header.ExpectedMagic)
+	if (Header.Magic == Header.ExpectedMagic && Header.TotalCompressedSize >= sizeof(FHeader))
 	{
 		FUniqueBuffer MutableBuffer = FUniqueBuffer::Alloc(Header.TotalCompressedSize);
 		Header.ByteSwap();
 		const FMutableMemoryView MutableView = MutableBuffer.GetView().CopyFrom(MakeMemoryView(&Header, &Header + 1));
 		Ar.Serialize(MutableView.GetData(), int64(MutableView.GetSize()));
-		Local.CompressedData = ValidBufferOrEmpty(MutableBuffer.MoveToShared());
+		if (!Ar.IsError())
+		{
+			Local.CompressedData = ValidBufferOrEmpty(MutableBuffer.MoveToShared());
+		}
 	}
 	if (Local.IsNull())
 	{
