[Core]
* Fixed c++20 compile errors caused by changes that enabled concepts

diff --git a/Engine/Source/Runtime/Core/Public/Math/Quat.h b/Engine/Source/Runtime/Core/Public/Math/Quat.h
--- a/Engine/Source/Runtime/Core/Public/Math/Quat.h
+++ b/Engine/Source/Runtime/Core/Public/Math/Quat.h
@@ -1405,7 +1405,7 @@ inline bool FQuat4d::SerializeFromMismatchedTag(FName StructTag, FArchive& Ar)
 template<typename T>
 struct TCustomLerp< UE::Math::TQuat<T> >
 {
-	enum { Value = true };
+	constexpr static bool Value = true;
 	using QuatType = UE::Math::TQuat<T>;
 
 	template<class U>
diff --git a/Engine/Source/Runtime/Core/Public/Math/Rotator.h b/Engine/Source/Runtime/Core/Public/Math/Rotator.h
--- a/Engine/Source/Runtime/Core/Public/Math/Rotator.h
+++ b/Engine/Source/Runtime/Core/Public/Math/Rotator.h
@@ -899,7 +899,7 @@ template<typename T>
 struct TCustomLerp<UE::Math::TRotator<T>>
 {
 	// Required to make FMath::Lerp<TRotator>() call our custom Lerp() implementation below.
-	enum { Value = true };
+	constexpr static bool Value = true;
 	using RotatorType = UE::Math::TRotator<T>;
 
 	template<class U>
diff --git a/Engine/Source/Runtime/Core/Public/Math/ScalarRegister.h b/Engine/Source/Runtime/Core/Public/Math/ScalarRegister.h
--- a/Engine/Source/Runtime/Core/Public/Math/ScalarRegister.h
+++ b/Engine/Source/Runtime/Core/Public/Math/ScalarRegister.h
@@ -194,7 +194,7 @@ FORCEINLINE ScalarRegister ScalarMax(const ScalarRegister& A, const ScalarRegist
 // TCustomLerp for FMath::Lerp()
 template<> struct TCustomLerp<ScalarRegister>
 {
-	enum { Value = true };
+	constexpr static bool Value = true;
 
 	// Specialization of Lerp template that works with scalar (float in vector) registers
 	FORCEINLINE ScalarRegister Lerp(const ScalarRegister& A, const ScalarRegister& B, const ScalarRegister& Alpha)
diff --git a/Engine/Source/Runtime/Core/Public/Math/UnrealMathUtility.h b/Engine/Source/Runtime/Core/Public/Math/UnrealMathUtility.h
--- a/Engine/Source/Runtime/Core/Public/Math/UnrealMathUtility.h
+++ b/Engine/Source/Runtime/Core/Public/Math/UnrealMathUtility.h
@@ -250,7 +250,7 @@ class TRange;
 template <typename T>
 struct TCustomLerp
 {
-	enum { Value = false };
+	constexpr static bool Value = false;
 };
 
 /**
diff --git a/Engine/Source/Runtime/Core/Public/Math/UnrealMathVectorCommon.h b/Engine/Source/Runtime/Core/Public/Math/UnrealMathVectorCommon.h
--- a/Engine/Source/Runtime/Core/Public/Math/UnrealMathVectorCommon.h
+++ b/Engine/Source/Runtime/Core/Public/Math/UnrealMathVectorCommon.h
@@ -545,7 +545,7 @@ FORCEINLINE VectorRegister4Double VectorLerp(const VectorRegister4Double& A, con
 template<>
 struct TCustomLerp<VectorRegister4Float>
 {
-	enum { Value = true };
+	constexpr static bool Value = true;
 	
 	// Specialization of Lerp function that works with vector registers
 	static FORCEINLINE VectorRegister4Float Lerp(const VectorRegister4Float& A, const VectorRegister4Float& B, const VectorRegister4Float& Alpha)
@@ -558,7 +558,7 @@ struct TCustomLerp<VectorRegister4Float>
 template<>
 struct TCustomLerp<VectorRegister4Double>
 {
-	enum { Value = true };
+	constexpr static bool Value = true;
 
 	// Specialization of Lerp function that works with vector registers
 	static FORCEINLINE VectorRegister4Double Lerp(const VectorRegister4Double& A, const VectorRegister4Double& B, const VectorRegister4Double& Alpha)
