AutomationTest AddExpectedMessage remove overload ambiguity

diff --git a/Engine/Source/Runtime/Core/Public/Misc/AutomationTest.h b/Engine/Source/Runtime/Core/Public/Misc/AutomationTest.h
--- a/Engine/Source/Runtime/Core/Public/Misc/AutomationTest.h
+++ b/Engine/Source/Runtime/Core/Public/Misc/AutomationTest.h
@@ -1468,7 +1468,7 @@ public:
 	* @param Occurrences - How many times to expect this message string to be seen. If > 0, the message must be seen the exact number of times
 	* specified or the test will fail. If == 0, the message must be seen one or more times (with no upper limit) or the test will fail.
 	*/
-	void AddExpectedMessage(FString ExpectedPatternString, ELogVerbosity::Type ExpectedVerbosity = ELogVerbosity::All, EAutomationExpectedMessageFlags::MatchType CompareType = EAutomationExpectedMessageFlags::Contains, int32 Occurrences = 1);
+	void AddExpectedMessage(FString ExpectedPatternString, ELogVerbosity::Type ExpectedVerbosity, EAutomationExpectedMessageFlags::MatchType CompareType = EAutomationExpectedMessageFlags::Contains, int32 Occurrences = 1);
 	
 	/**
 	* Adds a regex pattern to an internal list that this test will expect to encounter in logs (of all severities) during its execution. If an expected pattern
