AutomationTest remove redundant types after AddExpectedMessage introduction

diff --git a/Engine/Source/Runtime/Core/Public/Misc/AutomationTest.h b/Engine/Source/Runtime/Core/Public/Misc/AutomationTest.h
--- a/Engine/Source/Runtime/Core/Public/Misc/AutomationTest.h
+++ b/Engine/Source/Runtime/Core/Public/Misc/AutomationTest.h
@@ -164,30 +164,8 @@ namespace EAutomationExpectedMessageFlags
 	}
 }
 
-// @todo: replace/refactor to use EAutomationExpectedMessageFlags - touches many, many files
-/** Flags for indicating the matching type to use for an expected error */
-namespace EAutomationExpectedErrorFlags
-{
-	enum MatchType
-	{
-		// When matching expected errors, do so exactly.
-		Exact,
-		// When matching expected errors, just see if the error string is contained in the string to be evaluated.
-		Contains,
-	};
-
-	inline const TCHAR* ToString(EAutomationExpectedErrorFlags::MatchType ThisType)
-	{
-		switch (ThisType)
-		{
-		case Contains:
-			return TEXT("Contains");
-		case Exact:
-			return TEXT("Exact");
-		}
-		return TEXT("Unknown");
-	}
-}
+/** Flags for indicating the matching type to use for an expected error message. Aliased for backwards compatibility. */
+namespace EAutomationExpectedErrorFlags = EAutomationExpectedMessageFlags;
 
 struct FAutomationTelemetryData
 {
