Add send trace snapshot command

Adds the ability to send a trace snapshot to a server. A snapshot is the contents of any events in the tail buffer along with the important events needed to analyze them. This feature exists but was only supported writing to local files.

#UE-161986

diff --git a/Engine/Source/Runtime/Core/Private/ProfilingDebugging/TraceAuxiliary.cpp b/Engine/Source/Runtime/Core/Private/ProfilingDebugging/TraceAuxiliary.cpp
--- a/Engine/Source/Runtime/Core/Private/ProfilingDebugging/TraceAuxiliary.cpp
+++ b/Engine/Source/Runtime/Core/Private/ProfilingDebugging/TraceAuxiliary.cpp
@@ -100,6 +100,7 @@ public:
 	void StartWorkerThread();
 	void StartEndFramePump();
 	bool WriteSnapshot(const TCHAR* InFilePath, const FTraceAuxiliary::FLogCategoryAlias& LogCategory);
+	bool SendSnapshot(const TCHAR* InHost, uint32 InPort, const FTraceAuxiliary::FLogCategoryAlias& LogCategory);
 
 	// True if this is parent process with forking requested before forking.
 	bool					IsParentProcessAndPreFork();
@@ -577,6 +578,28 @@ bool FTraceAuxiliaryImpl::WriteSnapshot(const TCHAR* InFilePath, const FTraceAux
 	return bResult;
 }
 
+////////////////////////////////////////////////////////////////////////////////
+bool FTraceAuxiliaryImpl::SendSnapshot(const TCHAR* InHost, uint32 InPort, const FTraceAuxiliary::FLogCategoryAlias& LogCategory)
+{
+	double StartTime = FPlatformTime::Seconds();
+
+	UE_LOG_REF(LogCategory, Log, TEXT("Sending trace snapshot to '%s'..."), InHost);
+
+	const bool bResult = UE::Trace::SendSnapshotTo(InHost, InPort);
+
+	if (bResult)
+	{
+		FTraceAuxiliary::OnSnapshotSaved.Broadcast(InHost);
+		UE_LOG_REF(LogCategory, Display, TEXT("Trace snapshot generated in %.3f seconds to \"%s\"."), FPlatformTime::Seconds() - StartTime, InHost);
+	}
+	else
+	{
+		UE_LOG_REF(LogCategory, Error, TEXT("Failed to trace snapshot to \"%s\"."), InHost);
+	}
+
+	return bResult;
+}
+
 ////////////////////////////////////////////////////////////////////////////////
 const TCHAR* FTraceAuxiliaryImpl::GetDest() const
 {
@@ -928,6 +951,33 @@ static void TraceAuxiliarySnapshotFile(const TArray<FString>& Args)
 	GTraceAuxiliary.WriteSnapshot(FilePath, LogConsoleResponse);
 }
 
+////////////////////////////////////////////////////////////////////////////////
+static void TraceAuxiliarySnapshotSend(const TArray<FString>& Args)
+{
+	const TCHAR* Host = nullptr;
+	uint32 Port = 0;
+
+	if (Args.Num() >= 1)
+	{
+		Host = *Args[0];
+	}
+	if (Args.Num() >= 2)
+	{
+		LexFromString(Port, *Args[1]);
+	}
+	if (Args.Num() == 0)
+	{
+		Host = TEXT("localhost");
+	}
+	if (Args.Num() > 2) 
+	{
+		UE_LOG(LogConsoleResponse, Warning, TEXT("Invalid arguments. Usage: Trace.SnapshotFile <Host> <Port>"));
+		return;
+	}
+
+	GTraceAuxiliary.SendSnapshot(Host, Port, LogConsoleResponse);
+}
+
 ////////////////////////////////////////////////////////////////////////////////
 static void TraceBookmark(const TArray<FString>& Args)
 {
@@ -1018,6 +1068,14 @@ static FAutoConsoleCommand TraceAuxiliarySnapshotFileCmd(
 	FConsoleCommandWithArgsDelegate::CreateStatic(TraceAuxiliarySnapshotFile)
 );
 
+////////////////////////////////////////////////////////////////////////////////
+static FAutoConsoleCommand TraceAuxiliarySnapshotSendCmd(
+	TEXT("Trace.SnapshotSend"),
+	TEXT("<Host> <Port> - Sends a snapshot of the current in-memory trace buffer to a server. If no host is specified 'localhost' is used."
+	),
+	FConsoleCommandWithArgsDelegate::CreateStatic(TraceAuxiliarySnapshotSend)
+);
+
 ////////////////////////////////////////////////////////////////////////////////
 static FAutoConsoleCommand TraceBookmarkCmd(
 	TEXT("Trace.Bookmark"),
diff --git a/Engine/Source/Runtime/TraceLog/Private/Trace/Trace.cpp b/Engine/Source/Runtime/TraceLog/Private/Trace/Trace.cpp
--- a/Engine/Source/Runtime/TraceLog/Private/Trace/Trace.cpp
+++ b/Engine/Source/Runtime/TraceLog/Private/Trace/Trace.cpp
@@ -23,6 +23,7 @@ void	Writer_Update();
 bool	Writer_SendTo(const ANSICHAR*, uint32, uint32);
 bool	Writer_WriteTo(const ANSICHAR*, uint32);
 bool	Writer_WriteSnapshotTo(const ANSICHAR*);
+bool	Writer_SendSnapshotTo(const ANSICHAR*, uint32);	
 bool	Writer_IsTracing();
 bool	Writer_Stop();
 uint32	Writer_GetThreadId();
@@ -105,6 +106,14 @@ bool WriteSnapshotTo(const TCHAR* InPath)
 	return Private::Writer_WriteSnapshotTo(Path);
 }
 
+////////////////////////////////////////////////////////////////////////////////
+bool SendSnapshotTo(const TCHAR* InHost, uint32 InPort)
+{
+	char Host[512];
+	ToAnsiCheap(Host, InHost);
+	return Private::Writer_SendSnapshotTo(Host, InPort);
+}
+
 ////////////////////////////////////////////////////////////////////////////////
 bool IsTracing()
 {
diff --git a/Engine/Source/Runtime/TraceLog/Private/Trace/Writer.cpp b/Engine/Source/Runtime/TraceLog/Private/Trace/Writer.cpp
--- a/Engine/Source/Runtime/TraceLog/Private/Trace/Writer.cpp
+++ b/Engine/Source/Runtime/TraceLog/Private/Trace/Writer.cpp
@@ -959,7 +959,26 @@ private:
 };
 
 ////////////////////////////////////////////////////////////////////////////////
-bool Writer_WriteSnapshotTo(const ANSICHAR* Path)
+struct FSnapshotTarget
+{
+	enum EType { FileTarget, HostTarget };
+	EType Type;
+	union
+	{
+		 struct
+		 {
+			 const ANSICHAR* Path;
+		 } File;
+		 struct
+		 {
+			 const ANSICHAR* Host;
+			 uint32 Port;
+		 } Host;
+	};
+};
+
+////////////////////////////////////////////////////////////////////////////////
+bool Writer_WriteSnapshot(const FSnapshotTarget& Target)
 {
 	if (!Writer_IsTailing())
 	{
@@ -996,8 +1015,24 @@ bool Writer_WriteSnapshotTo(const ANSICHAR* Path)
 		TStashGlobal SyncPacketCountdown(GSyncPacketCountdown, GNumSyncPackets);
 		TStashGlobal TraceStatistics(GTraceStatistics);
 
-		// Open the snapshot file and write the file header
-		GDataHandle = FileOpen(Path);
+		if (Target.Type == FSnapshotTarget::EType::FileTarget)
+		{
+			// Open the snapshot file 
+			GDataHandle = FileOpen(Target.File.Path);
+		}
+		else
+		{
+			// Open the snapshot connection and write 
+			const uint32 Port = Target.Host.Port ? Target.Host.Port : 1981;
+			GDataHandle = TcpSocketConnect(Target.Host.Host, uint16(Port));
+			if (!GDataHandle)
+			{
+				return false;
+			}
+			GDataHandle = Writer_PackSendFlags(GDataHandle, 0);
+		}
+
+		// Write the file header
 		if (!GDataHandle || !Writer_SessionPrologue())
 		{
 			return false;
@@ -1028,7 +1063,26 @@ bool Writer_WriteSnapshotTo(const ANSICHAR* Path)
 
 	return true;
 }
+	
+////////////////////////////////////////////////////////////////////////////////
+bool Writer_WriteSnapshotTo(const ANSICHAR* Path)
+{
+	FSnapshotTarget Target;
+	Target.Type = FSnapshotTarget::EType::FileTarget;
+	Target.File.Path = Path;
+	return Writer_WriteSnapshot(Target);
+}
 
+////////////////////////////////////////////////////////////////////////////////
+bool Writer_SendSnapshotTo(const ANSICHAR* Host, uint32 Port)
+{
+	FSnapshotTarget Target;
+	Target.Type = FSnapshotTarget::EType::HostTarget;
+	Target.Host.Host = Host;
+	Target.Host.Port = Port;
+	return Writer_WriteSnapshot(Target);
+}
+	
 ////////////////////////////////////////////////////////////////////////////////
 bool Writer_IsTracing()
 {
diff --git a/Engine/Source/Runtime/TraceLog/Public/Trace/Trace.h b/Engine/Source/Runtime/TraceLog/Public/Trace/Trace.h
--- a/Engine/Source/Runtime/TraceLog/Public/Trace/Trace.h
+++ b/Engine/Source/Runtime/TraceLog/Public/Trace/Trace.h
@@ -127,6 +127,7 @@ UE_TRACE_API void	GetStatistics(FStatistics& Out) UE_TRACE_IMPL();
 UE_TRACE_API bool	SendTo(const TCHAR* Host, uint32 Port=0, uint16 Flags=FSendFlags::None) UE_TRACE_IMPL(false);
 UE_TRACE_API bool	WriteTo(const TCHAR* Path, uint16 Flags=FSendFlags::None) UE_TRACE_IMPL(false);
 UE_TRACE_API bool	WriteSnapshotTo(const TCHAR* Path) UE_TRACE_IMPL(false);
+UE_TRACE_API bool	SendSnapshotTo(const TCHAR* Host, uint32 Port) UE_TRACE_IMPL(false);	
 UE_TRACE_API bool	IsTracing() UE_TRACE_IMPL(false);
 UE_TRACE_API bool	Stop() UE_TRACE_IMPL(false);
 UE_TRACE_API bool	IsChannel(const TCHAR* ChanneName) UE_TRACE_IMPL(false);
