Fix potential race on FSelfRegisteringExec during ALT loading

#rnx

diff --git a/Engine/Source/Runtime/Core/Private/Misc/CoreMisc.cpp b/Engine/Source/Runtime/Core/Private/Misc/CoreMisc.cpp
--- a/Engine/Source/Runtime/Core/Private/Misc/CoreMisc.cpp
+++ b/Engine/Source/Runtime/Core/Private/Misc/CoreMisc.cpp
@@ -11,6 +11,7 @@
 #include "Misc/LazySingleton.h"
 #include "Misc/OutputDeviceError.h"
 #include "Misc/ScopeLock.h"
+#include "Misc/ScopeRWLock.h"
 #include "CoreGlobals.h"
 #include "Templates/RefCounting.h"
 
@@ -41,6 +42,7 @@ DEFINE_LOG_CATEGORY(LogCore);
 -----------------------------------------------------------------------------*/
 
 using FSelfRegisteredExecArray = TArray<FSelfRegisteringExec*, TInlineAllocator<8>>;
+static FRWLock ExecRegistryLock;
 
 FSelfRegisteredExecArray& GetExecRegistry()
 {
@@ -51,17 +53,20 @@ FSelfRegisteredExecArray& GetExecRegistry()
 /** Constructor, registering this instance. */
 FSelfRegisteringExec::FSelfRegisteringExec()
 {
+	FWriteScopeLock ScopeLock(ExecRegistryLock);
 	GetExecRegistry().Add( this );
 }
 
 /** Destructor, unregistering this instance. */
 FSelfRegisteringExec::~FSelfRegisteringExec()
 {
+	FWriteScopeLock ScopeLock(ExecRegistryLock);
 	verify(GetExecRegistry().Remove( this ) == 1 );
 }
 
 bool FSelfRegisteringExec::StaticExec( UWorld* InWorld, const TCHAR* Cmd, FOutputDevice& Ar )
 {
+	FReadScopeLock ScopeLock(ExecRegistryLock);
 	for (FSelfRegisteringExec* Exe : GetExecRegistry())
 	{
 		if (Exe->Exec( InWorld, Cmd,Ar ))
