Fix MacOS not liking pthread_rwlock being static initialized

#rnx

diff --git a/Engine/Source/Runtime/Core/Private/Misc/CoreMisc.cpp b/Engine/Source/Runtime/Core/Private/Misc/CoreMisc.cpp
--- a/Engine/Source/Runtime/Core/Private/Misc/CoreMisc.cpp
+++ b/Engine/Source/Runtime/Core/Private/Misc/CoreMisc.cpp
@@ -42,7 +42,13 @@ DEFINE_LOG_CATEGORY(LogCore);
 -----------------------------------------------------------------------------*/
 
 using FSelfRegisteredExecArray = TArray<FSelfRegisteringExec*, TInlineAllocator<8>>;
-static FRWLock ExecRegistryLock;
+
+// Lazy because pthread implementation doesn't like static initialization.
+FRWLock& GetExecRegistryLock()
+{
+	static FRWLock ExecRegistryLock;
+	return ExecRegistryLock;
+}
 
 FSelfRegisteredExecArray& GetExecRegistry()
 {
@@ -53,20 +59,20 @@ FSelfRegisteredExecArray& GetExecRegistry()
 /** Constructor, registering this instance. */
 FSelfRegisteringExec::FSelfRegisteringExec()
 {
-	FWriteScopeLock ScopeLock(ExecRegistryLock);
+	FWriteScopeLock ScopeLock(GetExecRegistryLock());
 	GetExecRegistry().Add( this );
 }
 
 /** Destructor, unregistering this instance. */
 FSelfRegisteringExec::~FSelfRegisteringExec()
 {
-	FWriteScopeLock ScopeLock(ExecRegistryLock);
+	FWriteScopeLock ScopeLock(GetExecRegistryLock());
 	verify(GetExecRegistry().Remove( this ) == 1 );
 }
 
 bool FSelfRegisteringExec::StaticExec( UWorld* InWorld, const TCHAR* Cmd, FOutputDevice& Ar )
 {
-	FReadScopeLock ScopeLock(ExecRegistryLock);
+	FReadScopeLock ScopeLock(GetExecRegistryLock());
 	for (FSelfRegisteringExec* Exe : GetExecRegistry())
 	{
 		if (Exe->Exec( InWorld, Cmd,Ar ))
@@ -601,4 +607,4 @@ static FAutoConsoleCommand GSetStatsVerboseNamedEventsCmd(
 			Ar.Logf(TEXT("Expected on/off argument"));
 		}
 	})
-);
+);
\ No newline at end of file
