IoDispatcher: Fix shutdown crash by shutting down the backends before releasing the DispatcherEvent

#rnx

diff --git a/Engine/Source/Runtime/Core/Private/IO/IoDispatcher.cpp b/Engine/Source/Runtime/Core/Private/IO/IoDispatcher.cpp
--- a/Engine/Source/Runtime/Core/Private/IO/IoDispatcher.cpp
+++ b/Engine/Source/Runtime/Core/Private/IO/IoDispatcher.cpp
@@ -345,15 +345,17 @@ public:
 
 	~FIoDispatcherImpl()
 	{
-		delete Thread;
-		FPlatformProcess::ReturnSynchEventToPool(DispatcherEvent);
-		RequestAllocator->ReleaseRef();
 		for (const TSharedRef<IIoDispatcherBackend>& Backend : Backends)
 		{
 			Backend->Shutdown();
 		}
 		FCoreDelegates::GetMemoryTrimDelegate().Remove(MemoryTrimDelegateHandle);
 		BackendContext->WakeUpDispatcherThreadDelegate.Unbind();
+		// when all mounted backends have been shutdown and the delegates have been cleared,
+		// then we can go ahead and delete the resolve thread
+		delete Thread;
+		FPlatformProcess::ReturnSynchEventToPool(DispatcherEvent);
+		RequestAllocator->ReleaseRef();
 	}
 
 	void Initialize()
