Condense delegates ApplicationSystemUIOverlayAppearedDelegate and ApplicationSystemUIOverlayDisappearedDelegate to ApplicationSystemUIOverlayStateChangedDelegate which shares the UI state via the bool parameter

diff --git a/Engine/Source/Runtime/Core/Private/Misc/CoreDelegates.cpp b/Engine/Source/Runtime/Core/Private/Misc/CoreDelegates.cpp
--- a/Engine/Source/Runtime/Core/Private/Misc/CoreDelegates.cpp
+++ b/Engine/Source/Runtime/Core/Private/Misc/CoreDelegates.cpp
@@ -110,8 +110,7 @@ TMulticastDelegate<void()> FCoreDelegates::ApplicationWillTerminateDelegate;
 TMulticastDelegate<void()> FCoreDelegates::ApplicationShouldUnloadResourcesDelegate;
 TMulticastDelegate<void(float)> FCoreDelegates::MobileBackgroundTickDelegate;
 
-TMulticastDelegate<void()> FCoreDelegates::ApplicationSystemUIOverlayAppearedDelegate;
-TMulticastDelegate<void()> FCoreDelegates::ApplicationSystemUIOverlayDisappearedDelegate;
+TMulticastDelegate<void(bool)> FCoreDelegates::ApplicationSystemUIOverlayStateChangedDelegate;
 
 TMulticastDelegate<void(const TArray<FString>&)> FCoreDelegates::ApplicationReceivedStartupArgumentsDelegate;
 
diff --git a/Engine/Source/Runtime/Core/Public/Misc/CoreDelegates.h b/Engine/Source/Runtime/Core/Public/Misc/CoreDelegates.h
--- a/Engine/Source/Runtime/Core/Public/Misc/CoreDelegates.h
+++ b/Engine/Source/Runtime/Core/Public/Misc/CoreDelegates.h
@@ -323,10 +323,8 @@ public:
 
 	// Some platform have a System UI Overlay that can draw on top of the application.
 	// The game might want to be notified so it can pause, etc...
-	static TMulticastDelegate<void()> ApplicationSystemUIOverlayAppearedDelegate;
-
-	// Called when the UI Overlay was hidden
-	static TMulticastDelegate<void()> ApplicationSystemUIOverlayDisappearedDelegate;
+	// Parameter (bool) should be true if the system UI is displayed, otherwise false should be passed
+	static TMulticastDelegate<void(bool)> ApplicationSystemUIOverlayStateChangedDelegate;
 
 	// Called when in the background, if the OS is giving CPU time to the device. It is very likely
 	// this will never be called due to mobile OS backgrounded CPU restrictions. But if, for instance,
