LruCache: Allow unlimited container size and add non-const access to value type.

Submitting this early as it's required for https://p4-swarm.epicgames.net/reviews/23788374, which is a large change that should make it into 5.2 before the branch later today. This should be pretty trivial, but if there is any issue let me know and I can revert this and make our own custom container that fits our particular needs within our plugin.

diff --git a/Engine/Source/Runtime/Core/Public/Containers/LruCache.h b/Engine/Source/Runtime/Core/Public/Containers/LruCache.h
--- a/Engine/Source/Runtime/Core/Public/Containers/LruCache.h
+++ b/Engine/Source/Runtime/Core/Public/Containers/LruCache.h
@@ -157,7 +157,7 @@ public:
 	 */
 	void Add(const KeyType& Key, const ValueType& Value)
 	{
-		check(MaxNumElements > 0 && "Cannot add values to zero size TLruCache");
+		check(MaxNumElements != 0 && "Cannot add values to zero size TLruCache");
 
 		FCacheEntry** EntryPtr = LookupSet.Find(Key);
 
@@ -232,15 +232,13 @@ public:
 	 */
 	void Empty(int32 InMaxNumElements = 0)
 	{
-		check(InMaxNumElements >= 0);
-
 		for (FCacheEntry* Entry : LookupSet)
 		{
 			delete Entry;
 		}
 
 		MaxNumElements = InMaxNumElements;
-		LookupSet.Empty(MaxNumElements);
+		LookupSet.Empty(FMath::Max(MaxNumElements, 0));
 
 		MostRecent = nullptr;
 		LeastRecent = nullptr;
@@ -288,6 +286,25 @@ public:
 		return nullptr;
 	}
 
+	/**
+	 * Find the value of the entry with the specified key.
+	 *
+	 * @param Key The key of the entry to get.
+	 * @return Pointer to the value, or nullptr if not found.
+	 * @see Add, Contains, Empty, FindAndTouch, GetKeys, Remove
+	 */
+	FORCEINLINE ValueType* Find(const KeyType& Key)
+	{
+		FCacheEntry* const* EntryPtr = LookupSet.Find(Key);
+
+		if (EntryPtr != nullptr)
+		{
+			return &(*EntryPtr)->Value;
+		}
+
+		return nullptr;
+	}
+
 	/**
 	 * Find the value of the entry with the specified key.
 	 *
@@ -297,7 +314,7 @@ public:
 	FORCEINLINE const ValueType& FindChecked(const KeyType& Key) const
 	{
 		FCacheEntry*const * EntryPtr = LookupSet.Find(Key);
-		
+
 		check(EntryPtr);
 
 		return (*EntryPtr)->Value;
@@ -351,7 +368,7 @@ public:
 	const ValueType& FindAndTouchChecked(const KeyType& Key)
 	{
 		FCacheEntry** EntryPtr = LookupSet.Find(Key);
-		
+
 		check(EntryPtr);
 
 		MarkAsRecent(**EntryPtr);
@@ -426,7 +443,7 @@ public:
 	}
 
 	/**
-	 * Returns true if the cache is empty and contains no elements. 
+	 * Returns true if the cache is empty and contains no elements.
 	 *
 	 * @returns True if the cache is empty.
 	 * @see Num
@@ -610,7 +627,7 @@ public:
 		{ }
 	};
 
-	
+
 	/**
 	 * Cache iterator.
 	 */
@@ -640,7 +657,7 @@ public:
 		}
 
 	private:
-		
+
 		TLruCache* Cache;
 	};
 
