compile fix for unit tests in shipping config

diff --git a/Engine/Source/Runtime/Core/Tests/Misc/AssertionMacrosTest.cpp b/Engine/Source/Runtime/Core/Tests/Misc/AssertionMacrosTest.cpp
--- a/Engine/Source/Runtime/Core/Tests/Misc/AssertionMacrosTest.cpp
+++ b/Engine/Source/Runtime/Core/Tests/Misc/AssertionMacrosTest.cpp
@@ -63,8 +63,10 @@ static int EnsureAlwaysMsg(int Value)
  */
 TEST_CASE("Core::Misc::TestEnsure", "[Core][Misc][AssertionMacros][Ensure]")
 {
+#if !UE_BUILD_SHIPPING
 	//ignore the debugger as all the breaks make for annoying debugging experience
 	TGuardValue<bool> IgnoreDebugger(GIgnoreDebugger, true);
+#endif
 	{
 		REQUIRE_ENSURE(EnsureFailed());
 		EnsureFailed(); //ensure already fired should not fire again
