[Insights] Added the "screenshot" trace channel to the "default" channel preset.

diff --git a/Engine/Source/Runtime/Core/Private/ProfilingDebugging/TraceAuxiliary.cpp b/Engine/Source/Runtime/Core/Private/ProfilingDebugging/TraceAuxiliary.cpp
--- a/Engine/Source/Runtime/Core/Private/ProfilingDebugging/TraceAuxiliary.cpp
+++ b/Engine/Source/Runtime/Core/Private/ProfilingDebugging/TraceAuxiliary.cpp
@@ -55,7 +55,7 @@
 #include "Trace/Trace.inl"
 
 ////////////////////////////////////////////////////////////////////////////////
-const TCHAR* GDefaultChannels = TEXT("cpu,gpu,frame,log,bookmark");
+const TCHAR* GDefaultChannels = TEXT("cpu,gpu,frame,log,bookmark,screenshot");
 const TCHAR* GMemoryChannels = TEXT("memtag,memalloc,callstack,module");
 const TCHAR* GTraceConfigSection = TEXT("Trace.Config");
 
